New styling and making scratchpads work similar to BP functions.
Move the scratchpad out of its own tab and make a more modular UI/UX experience based off a window with the active scripts that when double-clicked will open as a new tab sibling to system overview. The selection panel should update to match the active document selected.
The parameters panel should update to handle things from the graph. Selecting from the parameters panel will highlight active usages in the stack overview.

diff --git a/Engine/Plugins/FX/Niagara/Source/NiagaraEditor/Private/NiagaraEditorSettings.cpp b/Engine/Plugins/FX/Niagara/Source/NiagaraEditor/Private/NiagaraEditorSettings.cpp
--- a/Engine/Plugins/FX/Niagara/Source/NiagaraEditor/Private/NiagaraEditorSettings.cpp
+++ b/Engine/Plugins/FX/Niagara/Source/NiagaraEditor/Private/NiagaraEditorSettings.cpp
@@ -128,7 +128,6 @@ void UNiagaraEditorSettings::SetupNamespaceMetadata()
 			.SetDescription(LOCTEXT("ModuleDescription", "A value which exposes a module input to the system and emitter editor."))
 			.SetBackgroundColor(FLinearColor(FColor(136, 66, 65)))
 			.SetSortId(40)
-			.AddOption(ENiagaraNamespaceMetadataOptions::HideInSystem)
 			.AddOption(ENiagaraNamespaceMetadataOptions::HideInDefinitions)
 			.SetGuid(FNiagaraEditorGuids::ModuleNamespaceMetaDataGuid),
 		FNiagaraNamespaceMetadata({FNiagaraConstants::OutputNamespace}, FNiagaraConstants::ModuleNamespace)
@@ -149,7 +148,6 @@ void UNiagaraEditorSettings::SetupNamespaceMetadata()
 			.SetBackgroundColor(FLinearColor(FColor(191, 176, 84)))
 			.SetForegroundStyle("NiagaraEditor.ParameterName.NamespaceTextDark")
 			.SetSortId(50)
-			.AddOption(ENiagaraNamespaceMetadataOptions::HideInSystem)
 			.AddOption(ENiagaraNamespaceMetadataOptions::HideInDefinitions)
 			.AddOption(ENiagaraNamespaceMetadataOptions::PreventEditingNamespaceModifier)
 			.SetGuid(FNiagaraEditorGuids::ModuleLocalNamespaceMetaDataGuid),
@@ -221,7 +219,6 @@ void UNiagaraEditorSettings::SetupNamespaceMetadata()
 			.SetBackgroundColor(FLinearColor(FColor(170, 170, 170)))
 			.SetForegroundStyle("NiagaraEditor.ParameterName.NamespaceTextDark")
 			.SetSortId(100)
-			.AddOption(ENiagaraNamespaceMetadataOptions::HideInSystem)
 			.AddOption(ENiagaraNamespaceMetadataOptions::AdvancedInScript)
 			.AddOption(ENiagaraNamespaceMetadataOptions::AdvancedInSystem)
 			.AddOption(ENiagaraNamespaceMetadataOptions::HideInDefinitions)
@@ -233,7 +230,6 @@ void UNiagaraEditorSettings::SetupNamespaceMetadata()
 			.SetDisplayName(LOCTEXT("StatisSwitchDisplayName", "Static Switch Inputs"))
 			.SetDescription(LOCTEXT("StaticSwitchDescription", "Values which can only be set at edit time."))
 			.SetSortId(45)
-			.AddOption(ENiagaraNamespaceMetadataOptions::HideInSystem)
 			.AddOption(ENiagaraNamespaceMetadataOptions::HideInDefinitions)
 			.AddOption(ENiagaraNamespaceMetadataOptions::PreventEditingNamespace)
 			.AddOption(ENiagaraNamespaceMetadataOptions::PreventEditingNamespaceModifier)
@@ -264,8 +260,7 @@ void UNiagaraEditorSettings::SetupNamespaceMetadata()
 		FNiagaraNamespaceMetadata({FNiagaraConstants::ModuleNamespace})
 			.SetDisplayName(LOCTEXT("ModuleModifierDisplayName", "Module"))
 			.SetDescription(LOCTEXT("ModuleModifierDescription", "A namespace modifier which makes that attribute unique to the module\ninstance by appending the unique module name."))
-			.SetBackgroundColor(FLinearColor(FColor(102, 102, 152)))
-			.AddOption(ENiagaraNamespaceMetadataOptions::HideInSystem),
+			.SetBackgroundColor(FLinearColor(FColor(102, 102, 152))),
 		FNiagaraNamespaceMetadata({FNiagaraConstants::SystemNamespace})
 			.SetDisplayName(LOCTEXT("SystemModifierDisplayName", "System"))
 			.SetDescription(LOCTEXT("SystemModifierDescription", "A namespace modifier which specifies that an engine provided parameter comes from the system."))
diff --git a/Engine/Plugins/FX/Niagara/Source/NiagaraEditor/Private/NiagaraEditorStyle.cpp b/Engine/Plugins/FX/Niagara/Source/NiagaraEditor/Private/NiagaraEditorStyle.cpp
--- a/Engine/Plugins/FX/Niagara/Source/NiagaraEditor/Private/NiagaraEditorStyle.cpp
+++ b/Engine/Plugins/FX/Niagara/Source/NiagaraEditor/Private/NiagaraEditorStyle.cpp
@@ -678,10 +678,10 @@ void FNiagaraEditorStyle::InitScalabilityIcons()
 }
 
 FNiagaraEditorStyle::FNiagaraEditorStyle() : FSlateStyleSet("NiagaraEditorStyle")
-{	
+{
 	SetContentRoot(FPaths::EnginePluginsDir() / TEXT("FX/Niagara/Content/Slate"));
 	SetCoreContentRoot(FPaths::EngineContentDir() / TEXT("Editor/Slate"));
-	
+
 	InitStats();
 	InitAssetPicker();
 	InitActionMenu();
@@ -710,6 +710,21 @@ FNiagaraEditorStyle::FNiagaraEditorStyle() : FSlateStyleSet("NiagaraEditorStyle"
 	InitScalabilityColors();
 	InitScalabilityIcons();
 	InitViewportStyle();
+	InitScratchStyle();
+}
+void FNiagaraEditorStyle::InitScratchStyle()
+{
+	const FTextBlockStyle NormalText = FAppStyle::Get().GetWidgetStyle<FTextBlockStyle>("NormalText");
+
+	FSlateFontInfo ScratchPadEditorHeaderFont = DEFAULT_FONT("Bold", 11);
+	FTextBlockStyle ScratchPadEditorHeaderText = FTextBlockStyle(NormalText)
+		.SetFont(ScratchPadEditorHeaderFont);
+	Set("NiagaraEditor.ScratchPad.EditorHeaderText", ScratchPadEditorHeaderText);
+
+	FSlateFontInfo ScratchPadSubSectionHeaderFont = DEFAULT_FONT("Bold", 9);
+	FTextBlockStyle ScratchPadSubSectionHeaderText = FTextBlockStyle(NormalText)
+		.SetFont(ScratchPadSubSectionHeaderFont);
+	Set("NiagaraEditor.ScratchPad.SubSectionHeaderText", ScratchPadSubSectionHeaderText);
 }
 
 void FNiagaraEditorStyle::ReloadTextures()
diff --git a/Engine/Plugins/FX/Niagara/Source/NiagaraEditor/Private/Toolkits/NiagaraSystemToolkit.cpp b/Engine/Plugins/FX/Niagara/Source/NiagaraEditor/Private/Toolkits/NiagaraSystemToolkit.cpp
--- a/Engine/Plugins/FX/Niagara/Source/NiagaraEditor/Private/Toolkits/NiagaraSystemToolkit.cpp
+++ b/Engine/Plugins/FX/Niagara/Source/NiagaraEditor/Private/Toolkits/NiagaraSystemToolkit.cpp
@@ -28,6 +28,7 @@
 #include "NiagaraScriptStatsViewModel.h"
 #include "NiagaraBakerViewModel.h"
 #include "NiagaraToolkitCommon.h"
+#include "ViewModels/NiagaraSystemEditorDocumentsViewModel.h"
 
 #include "IContentBrowserSingleton.h"
 #include "ContentBrowserModule.h"
@@ -128,14 +129,15 @@ void FNiagaraSystemToolkit::SetSystemOverview(const TSharedPtr<SWidget>& InSyste
 	this->SystemOverview = InSystemOverview;
 }
 
-TSharedPtr<SWidget> FNiagaraSystemToolkit::GetScriptScratchpad() const
+TSharedPtr<SWidget> FNiagaraSystemToolkit::GetScriptScratchpadManager() const
 {
-	return ScriptScratchpad;
+	return ScriptScratchpadManager;
 }
 
-void FNiagaraSystemToolkit::SetScriptScratchpad(const TSharedPtr<SWidget>& InScriptScratchpad)
+
+void FNiagaraSystemToolkit::SetScriptScratchpadManager(const TSharedPtr<SWidget>& InScriptScratchpadManager)
 {
-	this->ScriptScratchpad = InScriptScratchpad;
+	this->ScriptScratchpadManager = InScriptScratchpadManager;
 }
 
 void FNiagaraSystemToolkit::InitializeWithSystem(const EToolkitMode::Type Mode, const TSharedPtr< class IToolkitHost >& InitToolkitHost, UNiagaraSystem& InSystem)
@@ -164,6 +166,7 @@ void FNiagaraSystemToolkit::InitializeWithSystem(const EToolkitMode::Type Mode,
 	ParameterDefinitionsPanelViewModel->Init(UIContext);
 	
 	SystemViewModel->SetToolkitCommands(GetToolkitCommands());
+	SystemViewModel->SetParameterPanelViewModel(ParameterPanelViewModel);
 	SystemToolkitMode = ESystemToolkitMode::System;
 
 	if (GbLogNiagaraSystemChanges > 0)
@@ -290,8 +293,8 @@ void FNiagaraSystemToolkit::InitializeInternal(const EToolkitMode::Type Mode, co
 	SystemViewModel->GetSelectionViewModel()->OnEmitterHandleIdSelectionChanged().AddSP(this, &FNiagaraSystemToolkit::OnSystemSelectionChanged);
 	SystemViewModel->GetOnPinnedEmittersChanged().AddSP(this, &FNiagaraSystemToolkit::RefreshParameters);
 	SystemViewModel->OnRequestFocusTab().AddSP(this, &FNiagaraSystemToolkit::OnViewModelRequestFocusTab);
-	
 	SystemViewModel->OnGetWorkflowMode().BindSP(this, &FNiagaraSystemToolkit::GetCurrentMode);
+	SystemViewModel->GetDocumentViewModel()->InitializePreTabManager(SharedThis(this));
 	
 	constexpr bool bCreateDefaultStandaloneMenu = true;
 	constexpr bool bCreateDefaultToolbar = true;
@@ -302,7 +305,9 @@ void FNiagaraSystemToolkit::InitializeInternal(const EToolkitMode::Type Mode, co
 	const TSharedRef<FTabManager::FLayout> DummyLayout = FTabManager::NewLayout("NullLayout")->AddArea(FTabManager::NewPrimaryArea());
 	FAssetEditorToolkit::InitAssetEditor(Mode, InitToolkitHost, FNiagaraEditorModule::NiagaraEditorAppIdentifier,
 		DummyLayout, bCreateDefaultStandaloneMenu, bCreateDefaultToolbar, ToolkitObject);
-	
+
+	SystemViewModel->GetDocumentViewModel()->InitializePostTabManager(SharedThis(this));
+
 	AddApplicationMode(DefaultModeName, MakeShared<FNiagaraSystemToolkitMode_Default>(SharedThis(this)));
 	AddApplicationMode(ScalabilityModeName, MakeShared<FNiagaraSystemToolkitMode_Scalability>(SharedThis(this)));
 
@@ -893,6 +898,7 @@ void FNiagaraSystemToolkit::OpenAttributeSpreadsheet()
 	InvokeTab(FNiagaraSystemToolkitModeBase::DebugSpreadsheetTabID);
 }
 
+
 void FNiagaraSystemToolkit::OnToggleBoundsSetFixedBounds_Emitters()
 {
 	FScopedTransaction Transaction(LOCTEXT("SetFixedBoundsEmitters", "Set Fixed Bounds (Emitters)"));
@@ -1390,4 +1396,5 @@ void FNiagaraSystemToolkit::OnViewModelRequestFocusTab(FName TabName)
 	GetTabManager()->TryInvokeTab(TabName);
 }
 
+
 #undef LOCTEXT_NAMESPACE
diff --git a/Engine/Plugins/FX/Niagara/Source/NiagaraEditor/Private/Toolkits/NiagaraSystemToolkit.h b/Engine/Plugins/FX/Niagara/Source/NiagaraEditor/Private/Toolkits/NiagaraSystemToolkit.h
--- a/Engine/Plugins/FX/Niagara/Source/NiagaraEditor/Private/Toolkits/NiagaraSystemToolkit.h
+++ b/Engine/Plugins/FX/Niagara/Source/NiagaraEditor/Private/Toolkits/NiagaraSystemToolkit.h
@@ -103,7 +103,7 @@ public:
 
 	TSharedPtr<SWidget> GetSystemOverview() const;
 	void SetSystemOverview(const TSharedPtr<SWidget>& SystemOverview);
-	TSharedPtr<SWidget> GetScriptScratchpad() const;
+	TSharedPtr<SWidget> GetScriptScratchpadManager() const;
 	void SetScriptScratchpad(const TSharedPtr<SWidget>& ScriptScratchpad);
 	TSharedPtr<SWidget> GetVersioningWidget() const { return VersionsWidget; }
 	FText GetVersionButtonLabel() const;
@@ -113,6 +113,8 @@ public:
 	bool IsVersionSelected(FNiagaraAssetVersion Version) const;
 	FText GetVersionMenuLabel(FNiagaraAssetVersion Version) const;
 
+	void SetScriptScratchpadManager(const TSharedPtr<SWidget>& ScriptScratchpad);
+	
 	/** Mode exposed functions */
 	TSharedRef<SWidget> GenerateCompileMenuContent();
 	TSharedRef<SWidget> GenerateBoundsMenuContent(TSharedRef<FUICommandList> InCommandList);
@@ -123,6 +125,8 @@ public:
 	bool HasSystem() const { return SystemToolkitMode == ESystemToolkitMode::System && System != nullptr; }
 
 	const TArray<UObject*>& GetObjectsBeingEdited() const;
+
+	
 public:
 	FRefreshItemSelectorDelegate RefreshItemSelector;
 
@@ -253,8 +257,8 @@ private:
 	bool bScratchPadChangesDiscarded;
 
 	TSharedPtr<SWidget> SystemOverview;
-	TSharedPtr<SWidget> ScriptScratchpad;
 	TSharedPtr<SWidget> VersionsWidget;
+	TSharedPtr<SWidget> ScriptScratchpadManager;
 	
 	static IConsoleVariable* VmStatEnabledVar;
 
diff --git a/Engine/Plugins/FX/Niagara/Source/NiagaraEditor/Private/Toolkits/SystemToolkitModes/NiagaraSystemToolkitModeBase.cpp b/Engine/Plugins/FX/Niagara/Source/NiagaraEditor/Private/Toolkits/SystemToolkitModes/NiagaraSystemToolkitModeBase.cpp
--- a/Engine/Plugins/FX/Niagara/Source/NiagaraEditor/Private/Toolkits/SystemToolkitModes/NiagaraSystemToolkitModeBase.cpp
+++ b/Engine/Plugins/FX/Niagara/Source/NiagaraEditor/Private/Toolkits/SystemToolkitModes/NiagaraSystemToolkitModeBase.cpp
@@ -25,6 +25,11 @@
 #include "Widgets/SNiagaraSelectedObjectsDetails.h"
 #include "Widgets/SNiagaraSpreadsheetView.h"
 #include "Widgets/SNiagaraGeneratedCodeView.h"
+#include "Widgets/Layout/SWidgetSwitcher.h"
+#include "ViewModels/NiagaraScratchPadScriptViewModel.h"
+#include "Widgets/SNiagaraSelectedObjectsDetails.h"
+#include "NiagaraObjectSelection.h"
+#include "ViewModels/NiagaraSystemEditorDocumentsViewModel.h"
 
 #define LOCTEXT_NAMESPACE "NiagaraSystemToolkitModeBase"
 
@@ -42,10 +47,160 @@ const FName FNiagaraSystemToolkitModeBase::PreviewSettingsTabId(TEXT("NiagaraSys
 const FName FNiagaraSystemToolkitModeBase::GeneratedCodeTabID(TEXT("NiagaraSystemEditor_GeneratedCode"));
 const FName FNiagaraSystemToolkitModeBase::MessageLogTabID(TEXT("NiagaraSystemEditor_MessageLog"));
 const FName FNiagaraSystemToolkitModeBase::SystemOverviewTabID(TEXT("NiagaraSystemEditor_SystemOverview"));
-const FName FNiagaraSystemToolkitModeBase::ScratchPadTabID(TEXT("NiagaraSystemEditor_ScratchPad"));
 const FName FNiagaraSystemToolkitModeBase::ScriptStatsTabID(TEXT("NiagaraSystemEditor_ScriptStats"));
 const FName FNiagaraSystemToolkitModeBase::BakerTabID(TEXT("NiagaraSystemEditor_Baker"));
 const FName FNiagaraSystemToolkitModeBase::VersioningTabID(TEXT("NiagaraSystemEditor_Versioning"));
+const FName FNiagaraSystemToolkitModeBase::ScratchPadScriptsTabID(TEXT("NiagaraSystemEditor_ScratchPadScripts"));
+
+
+FNiagaraSystemToolkitModeBase::FNiagaraSystemToolkitModeBase(FName InModeName, TWeakPtr<FNiagaraSystemToolkit> InSystemToolkit) : FApplicationMode(InModeName), SystemToolkit(InSystemToolkit), SwitcherIdx(0)
+{
+	DocChangedHandle = InSystemToolkit.Pin()->GetSystemViewModel()->GetDocumentViewModel()->OnActiveDocumentChanged().AddRaw(this, &FNiagaraSystemToolkitModeBase::OnActiveDocumentChanged);
+	ObjectSelection = MakeShared<FNiagaraObjectSelection>();
+}
+
+FNiagaraSystemToolkitModeBase::~FNiagaraSystemToolkitModeBase()
+{
+	if (SystemToolkit.Pin().IsValid())
+	{
+		SystemToolkit.Pin()->GetSystemViewModel()->GetDocumentViewModel()->OnActiveDocumentChanged().Remove(DocChangedHandle);
+		if (LastActiveDocumentModel.IsValid())
+		{
+			TSharedPtr<FNiagaraScratchPadScriptViewModel> OldScratchScriptVM = LastActiveDocumentModel.Pin();
+			OldScratchScriptVM->GetGraphViewModel()->GetNodeSelection()->OnSelectedObjectsChanged().Remove(LastSelectionUpdateDelegate);
+		}
+	}
+}
+
+void FNiagaraSystemToolkitModeBase::OnActiveDocumentChanged(TSharedPtr<SDockTab> NewActiveTab)
+{
+	TSharedPtr<FNiagaraSystemToolkit> Toolkit = StaticCastSharedPtr<FNiagaraSystemToolkit>(SystemToolkit.Pin());
+	if (LastActiveDocumentModel.IsValid())
+	{
+		TSharedPtr<FNiagaraScratchPadScriptViewModel> OldScratchScriptVM = LastActiveDocumentModel.Pin();
+		OldScratchScriptVM->GetGraphViewModel()->GetNodeSelection()->OnSelectedObjectsChanged().Remove(LastSelectionUpdateDelegate);
+
+		TArray<UNiagaraGraph*> EdGraphs = OldScratchScriptVM->GetEditableGraphs();
+		check(EdGraphs.Num() <= 1);
+		if (EdGraphs.Num() == 1)
+		{
+			EdGraphs[0]->OnSubObjectSelectionChanged().Remove(LastParamPanelSelectionUpdateDelegate);
+		}
+
+		Toolkit->GetSystemViewModel()->GetSelectionViewModel()->OnEntrySelectionChanged().Remove(LastSystemSelectionUpdateDelegate);
+		OldScratchScriptVM->GetGraphViewModel()->GetGraph()->RemoveOnGraphNeedsRecompileHandler(LastGraphEditDelegate);
+
+	}
+
+	SwitcherIdx = 0;
+	TSharedPtr<SDockTab>  ActiveTab = Toolkit->GetSystemViewModel()->GetDocumentViewModel()->GetActiveDocumentTab().Pin();
+	if (ActiveTab.IsValid())
+	{
+		if (ActiveTab->GetLayoutIdentifier().TabType == SystemOverviewTabID)
+		{
+			SwitcherIdx = 0;
+		}
+		else
+		{
+			SwitcherIdx = 1;
+		}
+		UpdateSelectionForActiveDocument();
+
+		TSharedPtr<FNiagaraScratchPadScriptViewModel> ScratchScriptVM = Toolkit->GetSystemViewModel()->GetDocumentViewModel()->GetActiveScratchPadViewModelIfSet();
+		if (ScratchScriptVM.IsValid())
+		{
+			LastSelectionUpdateDelegate = ScratchScriptVM->GetGraphViewModel()->GetNodeSelection()->OnSelectedObjectsChanged().AddRaw(this, &FNiagaraSystemToolkitModeBase::UpdateSelectionForActiveDocument);
+			LastParamPanelSelectionUpdateDelegate = Toolkit->GetSystemViewModel()->GetParameterPanelViewModel()->GetVariableObjectSelection()->OnSelectedObjectsChanged().AddRaw(this, &FNiagaraSystemToolkitModeBase::OnParameterPanelViewModelExternalSelectionChanged);
+			LastSystemSelectionUpdateDelegate = Toolkit->GetSystemViewModel()->GetSelectionViewModel()->OnEntrySelectionChanged().AddSP(this, &FNiagaraSystemToolkitModeBase::OnSystemSelectionChanged);
+			LastGraphEditDelegate = ScratchScriptVM->GetGraphViewModel()->GetGraph()->AddOnGraphNeedsRecompileHandler(
+				FOnGraphChanged::FDelegate::CreateRaw(this, &FNiagaraSystemToolkitModeBase::OnEditedScriptGraphChanged));
+		}
+	}
+}
+
+void FNiagaraSystemToolkitModeBase::OnEditedScriptGraphChanged(const FEdGraphEditAction& InAction)
+{
+	TSharedPtr<FNiagaraSystemToolkit> Toolkit = StaticCastSharedPtr<FNiagaraSystemToolkit>(SystemToolkit.Pin());
+	// We need to update the parameter panel view model with new parameters potentially
+	if (Toolkit->GetSystemViewModel()->GetParameterPanelViewModel())
+		Toolkit->GetSystemViewModel()->GetParameterPanelViewModel()->RefreshNextTick();
+}
+
+void FNiagaraSystemToolkitModeBase::OnParameterPanelViewModelExternalSelectionChanged()
+{
+	TSharedPtr<FNiagaraSystemToolkit> Toolkit = StaticCastSharedPtr<FNiagaraSystemToolkit>(SystemToolkit.Pin());
+	TSharedPtr<FNiagaraScratchPadScriptViewModel> ScratchScriptVM = Toolkit->GetSystemViewModel()->GetDocumentViewModel()->GetActiveScratchPadViewModelIfSet();
+	if (ScratchScriptVM.IsValid())
+	{
+		TSharedPtr<FNiagaraObjectSelection> ParamSel = Toolkit->GetSystemViewModel()->GetParameterPanelViewModel()->GetVariableObjectSelection();
+		if (ParamSel.IsValid())
+		{
+			TArray<UObject*> SelectedNodes = ParamSel->GetSelectedObjects().Array();
+			if (SelectedNodes.Num() != 0)
+			{
+				ObjectSelection->SetSelectedObject(SelectedNodes[0]);
+				ObjectSelectionSubHeaderText = LOCTEXT("ParamPanelSel", "Parameter Panel Selection");
+			}
+		}
+	}
+}
+
+void FNiagaraSystemToolkitModeBase::OnSystemSelectionChanged()
+{
+	TSharedPtr<FNiagaraSystemToolkit> Toolkit = StaticCastSharedPtr<FNiagaraSystemToolkit>(SystemToolkit.Pin());
+	TSharedPtr<SDockTab>  ActiveTab = Toolkit->GetSystemViewModel()->GetDocumentViewModel()->GetActiveDocumentTab().Pin();
+	if (ActiveTab.IsValid())
+	{
+		if (ActiveTab->GetLayoutIdentifier().TabType != SystemOverviewTabID)
+		{
+			TSharedPtr<SDockTab> FoundOverviewTab = Toolkit->GetTabManager()->FindExistingLiveTab(SystemOverviewTabID);
+			if (FoundOverviewTab.IsValid())
+			{
+				FoundOverviewTab->DrawAttention();
+				FoundOverviewTab->ActivateInParent(ETabActivationCause::SetDirectly);
+			}
+		}
+	}
+}
+
+void FNiagaraSystemToolkitModeBase::UpdateSelectionForActiveDocument()
+{
+	TSharedPtr<FNiagaraSystemToolkit> Toolkit = StaticCastSharedPtr<FNiagaraSystemToolkit>(SystemToolkit.Pin());
+	TSharedPtr<FNiagaraScratchPadScriptViewModel> ScratchScriptVM = Toolkit->GetSystemViewModel()->GetDocumentViewModel()->GetActiveScratchPadViewModelIfSet();
+	if (ScratchScriptVM.IsValid())
+	{
+		TArray<UObject*> SelectedNodes = ScratchScriptVM->GetGraphViewModel()->GetNodeSelection()->GetSelectedObjects().Array();
+		if (SelectedNodes.Num() == 0)
+		{
+			if (ScratchScriptVM)
+			{
+				ObjectSelection->SetSelectedObject(ScratchScriptVM->GetEditScript().Script);
+				ObjectSelectionSubHeaderText = ScratchScriptVM->GetDisplayName();
+			}
+			else
+			{
+				ObjectSelection->ClearSelectedObjects();
+				ObjectSelectionSubHeaderText = LOCTEXT("InvalidSelection","Missing selection");
+			}
+		}
+		else if (SelectedNodes.Num() == 1)
+		{
+			UEdGraphNode* SelectedGraphNode = CastChecked<UEdGraphNode>(SelectedNodes[0]);
+			ObjectSelectionSubHeaderText = SelectedGraphNode->GetNodeTitle(ENodeTitleType::ListView);
+			ObjectSelection->SetSelectedObject(SelectedGraphNode);
+		}
+		else
+		{
+			ObjectSelection->SetSelectedObjects(SelectedNodes);
+			ObjectSelectionSubHeaderText = LOCTEXT("MultipleNodesSelection", "Multiple items selected");
+		}
+	}
+	else
+	{
+		ObjectSelection->ClearSelectedObjects();
+		ObjectSelectionSubHeaderText = LOCTEXT("EmptySelection", "Nothing selected");
+	}
+}
 
 void FNiagaraSystemToolkitModeBase::RegisterTabFactories(TSharedPtr<FTabManager> InTabManager)
 {
@@ -118,8 +273,8 @@ void FNiagaraSystemToolkitModeBase::RegisterTabFactories(TSharedPtr<FTabManager>
 		.SetGroup(WorkspaceMenuCategory.ToSharedRef())
 		.SetIcon(FSlateIcon(FNiagaraEditorStyle::Get().GetStyleSetName(), "Tab.SystemOverview"));
 
-	InTabManager->RegisterTabSpawner(ScratchPadTabID, FOnSpawnTab::CreateSP(this, &FNiagaraSystemToolkitModeBase::SpawnTab_ScratchPad))
-		.SetDisplayName(LOCTEXT("ScratchPadTabName", "Scratch Pad"))
+	InTabManager->RegisterTabSpawner(ScratchPadScriptsTabID, FOnSpawnTab::CreateSP(this, &FNiagaraSystemToolkitModeBase::SpawnTab_ScratchPadScripts))
+		.SetDisplayName(LOCTEXT("ScratchPadScriptsTabName", "Local Scripts"))
 		.SetGroup(WorkspaceMenuCategory.ToSharedRef())
 		.SetIcon(FSlateIcon(FNiagaraEditorStyle::Get().GetStyleSetName(), "Tab.ScratchPad"));
 
@@ -139,6 +294,12 @@ void FNiagaraSystemToolkitModeBase::RegisterTabFactories(TSharedPtr<FTabManager>
 		.SetIcon(FSlateIcon(FAppStyle::GetAppStyleSetName(), "Versions"));
 }
 
+int FNiagaraSystemToolkitModeBase::GetActiveSelectionDetailsIndex() const
+{
+	return SwitcherIdx;
+}
+
+
 class SNiagaraSelectedEmitterGraph : public SCompoundWidget
 {
 public:
@@ -391,7 +552,38 @@ TSharedRef<SDockTab> FNiagaraSystemToolkitModeBase::SpawnTab_SelectedEmitterStac
 	TSharedRef<SDockTab> SpawnedTab = SNew(SDockTab)
 		.Label(LOCTEXT("SystemOverviewSelection", "Selection"))
 		[
-			NiagaraEditorModule.GetWidgetProvider()->CreateStackView(*SystemToolkit.Pin()->GetSystemViewModel()->GetSelectionViewModel()->GetSelectionStackViewModel())
+			SNew(SWidgetSwitcher)
+			.WidgetIndex(this, &FNiagaraSystemToolkitModeBase::GetActiveSelectionDetailsIndex)
+			+ SWidgetSwitcher::Slot()
+			[
+				NiagaraEditorModule.GetWidgetProvider()->CreateStackView(*SystemToolkit.Pin()->GetSystemViewModel()->GetSelectionViewModel()->GetSelectionStackViewModel())
+			]
+			+ SWidgetSwitcher::Slot()
+			[
+				SNew(SVerticalBox)
+				+ SVerticalBox::Slot()
+				.AutoHeight()
+				.Padding(2.0f, 2.0f, 2.0f, 5.0f)
+				[
+					SNew(STextBlock)
+					.TextStyle(FNiagaraEditorStyle::Get(), "NiagaraEditor.ScratchPad.SubSectionHeaderText")
+					.Visibility(this, &FNiagaraSystemToolkitModeBase::GetObjectSelectionSubHeaderTextVisibility)
+					.Text(this, &FNiagaraSystemToolkitModeBase::GetObjectSelectionSubHeaderText)
+				]
+				+ SVerticalBox::Slot()
+				.HAlign(HAlign_Center)
+				.AutoHeight()
+				.Padding(FMargin(0.0f, 10.0f, 0.0f, 0.0f))
+				[
+					SNew(STextBlock)
+					.Visibility(this, &FNiagaraSystemToolkitModeBase::GetObjectSelectionNoSelectionTextVisibility)
+					.Text(LOCTEXT("NoObjectSelection", "No object selected"))
+				]
+				+ SVerticalBox::Slot()
+				[
+					SNew(SNiagaraSelectedObjectsDetails, ObjectSelection.ToSharedRef())
+				]
+			]
 		];
 
 	SDockTab::FOnTabClosedCallback TabClosedCallback = SDockTab::FOnTabClosedCallback::CreateLambda([=](TSharedRef<SDockTab> DockTab)
@@ -403,6 +595,21 @@ TSharedRef<SDockTab> FNiagaraSystemToolkitModeBase::SpawnTab_SelectedEmitterStac
 	return SpawnedTab;
 }
 
+EVisibility FNiagaraSystemToolkitModeBase::GetObjectSelectionSubHeaderTextVisibility() const
+{
+	return ObjectSelection->GetSelectedObjects().Num() != 0 ? EVisibility::Visible : EVisibility::Collapsed;
+}
+
+FText FNiagaraSystemToolkitModeBase::GetObjectSelectionSubHeaderText() const
+{
+	return ObjectSelectionSubHeaderText;
+}
+
+EVisibility FNiagaraSystemToolkitModeBase::GetObjectSelectionNoSelectionTextVisibility() const
+{
+	return ObjectSelection->GetSelectedObjects().Num() == 0 ? EVisibility::Visible : EVisibility::Collapsed;
+}
+
 TSharedRef<SDockTab> FNiagaraSystemToolkitModeBase::SpawnTab_SelectedEmitterGraph(const FSpawnTabArgs& Args)
 {
 	check(Args.GetTabId().TabType == SelectedEmitterGraphTabID);
@@ -468,6 +675,12 @@ TSharedRef<SDockTab> FNiagaraSystemToolkitModeBase::SpawnTab_SystemOverview(cons
 			SystemToolkit.Pin()->GetSystemOverview().ToSharedRef()
 		];
 
+
+	SpawnedTab->SetOnTabActivated(SDockTab::FOnTabActivatedCallback::CreateLambda([this](TSharedRef<SDockTab> Input, ETabActivationCause)
+		{
+			SystemToolkit.Pin()->GetSystemViewModel()->GetDocumentViewModel()->SetActiveDocumentTab(Input);			
+		}));
+
 	SpawnedTab->SetOnTabClosed(SDockTab::FOnTabClosedCallback::CreateLambda([this](TSharedRef<SDockTab>)
 	{
 		SystemToolkit.Pin()->SetSystemOverview(nullptr);
@@ -476,23 +689,25 @@ TSharedRef<SDockTab> FNiagaraSystemToolkitModeBase::SpawnTab_SystemOverview(cons
 	return SpawnedTab;
 }
 
-TSharedRef<SDockTab> FNiagaraSystemToolkitModeBase::SpawnTab_ScratchPad(const FSpawnTabArgs& Args)
+
+
+TSharedRef<SDockTab> FNiagaraSystemToolkitModeBase::SpawnTab_ScratchPadScripts(const FSpawnTabArgs& Args)
 {
-	if(!SystemToolkit.Pin()->GetScriptScratchpad().IsValid())
+	if (!SystemToolkit.Pin()->GetScriptScratchpadManager().IsValid())
 	{
-		SystemToolkit.Pin()->SetScriptScratchpad(FNiagaraEditorModule::Get().GetWidgetProvider()->CreateScriptScratchPad(*SystemToolkit.Pin()->GetSystemViewModel()->GetScriptScratchPadViewModel()));
+		SystemToolkit.Pin()->SetScriptScratchpadManager(FNiagaraEditorModule::Get().GetWidgetProvider()->CreateScriptScratchPadManager(*SystemToolkit.Pin()->GetSystemViewModel()->GetScriptScratchPadViewModel()));
 	}
-	
+
 	TSharedRef<SDockTab> SpawnedTab = SNew(SDockTab)
-		.Label(LOCTEXT("ScratchPadTabLabel", "Scratch Pad"))
+		.Label(LOCTEXT("ScratchPadTabLabel", "Local Scripts"))
 		[
-			SystemToolkit.Pin()->GetScriptScratchpad().ToSharedRef()
+			SystemToolkit.Pin()->GetScriptScratchpadManager().ToSharedRef()
 		];
 
 	SpawnedTab->SetOnTabClosed(SDockTab::FOnTabClosedCallback::CreateLambda([this](TSharedRef<SDockTab>)
-	{
-		SystemToolkit.Pin()->SetScriptScratchpad(nullptr);
-	}));
+		{
+			SystemToolkit.Pin()->SetScriptScratchpadManager(nullptr);
+		}));
 
 	return SpawnedTab;
 }
diff --git a/Engine/Plugins/FX/Niagara/Source/NiagaraEditor/Private/Toolkits/SystemToolkitModes/NiagaraSystemToolkitModeBase.h b/Engine/Plugins/FX/Niagara/Source/NiagaraEditor/Private/Toolkits/SystemToolkitModes/NiagaraSystemToolkitModeBase.h
--- a/Engine/Plugins/FX/Niagara/Source/NiagaraEditor/Private/Toolkits/SystemToolkitModes/NiagaraSystemToolkitModeBase.h
+++ b/Engine/Plugins/FX/Niagara/Source/NiagaraEditor/Private/Toolkits/SystemToolkitModes/NiagaraSystemToolkitModeBase.h
@@ -5,17 +5,24 @@
 #include "NiagaraBakerViewModel.h"
 #include "WorkflowOrientedApp/ApplicationMode.h"
 #include "Widgets/Docking/SDockTab.h"
+#include "GraphEditAction.h"
+class FDocumentTracker;
 
 class FNiagaraSystemToolkit;
+class FNiagaraObjectSelection;
 
 class FNiagaraSystemToolkitModeBase : public FApplicationMode
 {
 public:
-	FNiagaraSystemToolkitModeBase(FName InModeName, TWeakPtr<FNiagaraSystemToolkit> InSystemToolkit) : FApplicationMode(InModeName), SystemToolkit(InSystemToolkit)
-	{ }
+	FNiagaraSystemToolkitModeBase(FName InModeName, TWeakPtr<FNiagaraSystemToolkit> InSystemToolkit);
+	~FNiagaraSystemToolkitModeBase();
 	
 	virtual void RegisterTabFactories(TSharedPtr<FTabManager> InTabManager) override;
 
+	int GetActiveSelectionDetailsIndex() const;
+
+	virtual void OnActiveDocumentChanged(TSharedPtr<class SDockTab> NewActiveTab);
+	virtual void OnSystemSelectionChanged();
 private:
 	TSharedRef<SDockTab> SpawnTab_Viewport(const FSpawnTabArgs& Args);
 	TSharedRef<SDockTab> SpawnTab_CurveEd(const FSpawnTabArgs& Args);
@@ -30,13 +37,30 @@ private:
 	TSharedRef<SDockTab> SpawnTab_GeneratedCode(const FSpawnTabArgs& Args);
 	TSharedRef<SDockTab> SpawnTab_MessageLog(const FSpawnTabArgs& Args);
 	TSharedRef<SDockTab> SpawnTab_SystemOverview(const FSpawnTabArgs& Args);
-	TSharedRef<SDockTab> SpawnTab_ScratchPad(const FSpawnTabArgs& Args);
+	TSharedRef<SDockTab> SpawnTab_ScratchPadScripts(const FSpawnTabArgs& Args);
 	TSharedRef<SDockTab> SpawnTab_ScriptStats(const FSpawnTabArgs& Args);
 	TSharedRef<SDockTab> SpawnTab_Baker(const FSpawnTabArgs& Args);
 	TSharedRef<SDockTab> SpawnTab_Versioning(const FSpawnTabArgs& Args);
 
 protected:
 	TWeakPtr<FNiagaraSystemToolkit> SystemToolkit;
+	FDelegateHandle DocChangedHandle;
+	FDelegateHandle LastSelectionUpdateDelegate;
+	FDelegateHandle LastParamPanelSelectionUpdateDelegate;
+	FDelegateHandle LastSystemSelectionUpdateDelegate;
+	FDelegateHandle LastGraphEditDelegate;
+	int32 SwitcherIdx;
+	TSharedPtr<FNiagaraObjectSelection> ObjectSelection;
+	FText ObjectSelectionSubHeaderText;
+
+	EVisibility GetObjectSelectionSubHeaderTextVisibility() const;
+	FText GetObjectSelectionSubHeaderText() const;
+	EVisibility GetObjectSelectionNoSelectionTextVisibility() const;
+	void UpdateSelectionForActiveDocument();
+	void OnParameterPanelViewModelExternalSelectionChanged();
+	void OnEditedScriptGraphChanged(const FEdGraphEditAction& InAction);
+
+	TWeakPtr< class FNiagaraScratchPadScriptViewModel > LastActiveDocumentModel;
 
 public:
 	static const FName ViewportTabID;
@@ -53,7 +77,7 @@ public:
 	static const FName GeneratedCodeTabID;
 	static const FName MessageLogTabID;
 	static const FName SystemOverviewTabID;
-	static const FName ScratchPadTabID;
+	static const FName ScratchPadScriptsTabID;
 	static const FName ScriptStatsTabID;
 	static const FName BakerTabID;
 	static const FName VersioningTabID;
diff --git a/Engine/Plugins/FX/Niagara/Source/NiagaraEditor/Private/Toolkits/SystemToolkitModes/NiagaraSystemToolkitMode_Default.cpp b/Engine/Plugins/FX/Niagara/Source/NiagaraEditor/Private/Toolkits/SystemToolkitModes/NiagaraSystemToolkitMode_Default.cpp
--- a/Engine/Plugins/FX/Niagara/Source/NiagaraEditor/Private/Toolkits/SystemToolkitModes/NiagaraSystemToolkitMode_Default.cpp
+++ b/Engine/Plugins/FX/Niagara/Source/NiagaraEditor/Private/Toolkits/SystemToolkitModes/NiagaraSystemToolkitMode_Default.cpp
@@ -16,7 +16,7 @@
 
 FNiagaraSystemToolkitMode_Default::FNiagaraSystemToolkitMode_Default(TWeakPtr<FNiagaraSystemToolkit> InSystemToolkit) : FNiagaraSystemToolkitModeBase(FNiagaraSystemToolkit::DefaultModeName, InSystemToolkit)
 {
-	TabLayout = FTabManager::NewLayout("Standalone_Niagara_System_Layout_v24")
+	TabLayout = FTabManager::NewLayout("Standalone_Niagara_System_Layout_v27")
 		->AddArea
 		(
 			FTabManager::NewPrimaryArea()->SetOrientation(Orient_Vertical)
@@ -25,52 +25,50 @@ FNiagaraSystemToolkitMode_Default::FNiagaraSystemToolkitMode_Default(TWeakPtr<FN
 				FTabManager::NewSplitter()->SetOrientation(Orient_Horizontal)
 				->Split
 				(
-					// Top Level Left
 					FTabManager::NewSplitter()->SetOrientation(Orient_Vertical)
-					->SetSizeCoefficient(.75f)
-					->Split
-					(
-						// Inner Left Top
-						FTabManager::NewSplitter()->SetOrientation(Orient_Horizontal)
-						->SetSizeCoefficient(0.75f)
-						->Split
-						(
-							FTabManager::NewStack()
-							->SetSizeCoefficient(.25f)
-							->AddTab(ViewportTabID, ETabState::OpenedTab)
-						)
-						->Split
-						(
-							FTabManager::NewStack()
-							->SetSizeCoefficient(0.15f)
-							->AddTab(SystemParametersTabID, ETabState::OpenedTab)
-						)
-						->Split
-						(
-							FTabManager::NewStack()
-							->SetSizeCoefficient(0.6f)
-							->AddTab(SystemOverviewTabID, ETabState::OpenedTab)
-							->AddTab(ScratchPadTabID, ETabState::OpenedTab)
-							->AddTab(BakerTabID, ETabState::ClosedTab)
-							->SetForegroundTab(SystemOverviewTabID)
-						)
+					->SetSizeCoefficient(0.33f)
+					->Split(
+						// Left top
+						FTabManager::NewStack()
+						->AddTab(ViewportTabID, ETabState::OpenedTab)
+					)
+					
+					->Split(
+						// Left sidebar
+						FTabManager::NewStack()
+						->AddTab(SystemParametersTabID, ETabState::OpenedTab)
+						->AddTab(ScratchPadScriptsTabID, ETabState::OpenedTab)
+						->SetForegroundTab(SystemParametersTabID)
+					)
+				)
+				->SetSizeCoefficient(0.33f)
+				->Split
+				(
+					FTabManager::NewSplitter()->SetOrientation(Orient_Vertical)
+					->Split(
+						// Center top
+						FTabManager::NewStack()
+						->AddTab(SystemOverviewTabID, ETabState::OpenedTab)
+						->AddTab("Document", ETabState::ClosedTab)
+						->AddTab(BakerTabID, ETabState::ClosedTab)
+						->SetForegroundTab(SystemOverviewTabID)
 					)
-					->Split
-					(
-						// Inner Left Bottom
+					->Split(
+						// center bottom
 						FTabManager::NewStack()
-						->SetSizeCoefficient(0.25f)
+						->SetSizeCoefficient(0.33f)
+						->AddTab(SequencerTabID, ETabState::OpenedTab)
 						->AddTab(CurveEditorTabID, ETabState::OpenedTab)
 						->AddTab(MessageLogTabID, ETabState::OpenedTab)
-						->AddTab(SequencerTabID, ETabState::OpenedTab)
-						->AddTab(ScriptStatsTabID, ETabState::ClosedTab)
+						->AddTab(ScriptStatsTabID, ETabState::OpenedTab)
+						->SetForegroundTab(SequencerTabID)
 					)
 				)
 				->Split
 				(
-					// Top Level Right
+					// Right
 					FTabManager::NewStack()
-					->SetSizeCoefficient(0.25f)
+					->SetSizeCoefficient(0.33f)
 					->AddTab(SelectedEmitterStackTabID, ETabState::OpenedTab)
 					->AddTab(SelectedEmitterGraphTabID, ETabState::ClosedTab)
 					->AddTab(SystemScriptTabID, ETabState::ClosedTab)
diff --git a/Engine/Plugins/FX/Niagara/Source/NiagaraEditor/Private/Toolkits/SystemToolkitModes/NiagaraSystemToolkitMode_Scalability.cpp b/Engine/Plugins/FX/Niagara/Source/NiagaraEditor/Private/Toolkits/SystemToolkitModes/NiagaraSystemToolkitMode_Scalability.cpp
--- a/Engine/Plugins/FX/Niagara/Source/NiagaraEditor/Private/Toolkits/SystemToolkitModes/NiagaraSystemToolkitMode_Scalability.cpp
+++ b/Engine/Plugins/FX/Niagara/Source/NiagaraEditor/Private/Toolkits/SystemToolkitModes/NiagaraSystemToolkitMode_Scalability.cpp
@@ -259,7 +259,7 @@ TSharedRef<FTabManager::FLayout> FNiagaraSystemToolkitMode_Scalability::ChooseTa
 {
 	if(SystemToolkit.Pin()->GetSystemViewModel()->GetEditMode() == ENiagaraSystemViewModelEditMode::SystemAsset)
 	{
-		return FTabManager::NewLayout("Standalone_Niagara_System_Layout_Scalability_v29")
+		return FTabManager::NewLayout("Standalone_Niagara_System_Layout_Scalability_v30")
 		->AddArea
 		(
 			FTabManager::NewPrimaryArea()->SetOrientation(Orient_Vertical)
@@ -294,7 +294,6 @@ TSharedRef<FTabManager::FLayout> FNiagaraSystemToolkitMode_Scalability::ChooseTa
 								FTabManager::NewStack()
 								->SetSizeCoefficient(0.8f)
 								->AddTab(SystemOverviewTabID, ETabState::OpenedTab)
-								->AddTab(ScratchPadTabID, ETabState::OpenedTab)
 								->SetForegroundTab(SystemOverviewTabID)
 							)
 						)
@@ -370,7 +369,6 @@ TSharedRef<FTabManager::FLayout> FNiagaraSystemToolkitMode_Scalability::ChooseTa
 						(							
 							FTabManager::NewStack()
 							->AddTab(SystemOverviewTabID, ETabState::OpenedTab)
-							->AddTab(ScratchPadTabID, ETabState::OpenedTab)
 							->SetForegroundTab(SystemOverviewTabID)							
 						)
 					)
diff --git a/Engine/Plugins/FX/Niagara/Source/NiagaraEditor/Private/ViewModels/NiagaraParameterPanelViewModel.cpp b/Engine/Plugins/FX/Niagara/Source/NiagaraEditor/Private/ViewModels/NiagaraParameterPanelViewModel.cpp
--- a/Engine/Plugins/FX/Niagara/Source/NiagaraEditor/Private/ViewModels/NiagaraParameterPanelViewModel.cpp
+++ b/Engine/Plugins/FX/Niagara/Source/NiagaraEditor/Private/ViewModels/NiagaraParameterPanelViewModel.cpp
@@ -31,6 +31,8 @@
 #include "ViewModels/Stack/NiagaraStackGraphUtilities.h"
 #include "Widgets/SNiagaraParameterMenu.h"
 #include "Widgets/SNiagaraParameterPanel.h"
+#include "ViewModels/Stack/NiagaraStackViewModel.h"
+#include "ViewModels/NiagaraSystemEditorDocumentsViewModel.h"
 
 #include "Framework/Commands/GenericCommands.h"
 #include "Framework/MultiBox/MultiBoxBuilder.h"
@@ -47,6 +49,7 @@ template<> TMap<UNiagaraSystem*, TArray<FNiagaraSystemToolkitParameterPanelViewM
 
 TArray<FNiagaraParameterPanelCategory> FNiagaraSystemToolkitParameterPanelViewModel::DefaultCategories;
 TArray<FNiagaraParameterPanelCategory> FNiagaraSystemToolkitParameterPanelViewModel::DefaultAdvancedCategories;
+TArray<FNiagaraParameterPanelCategory> FNiagaraSystemToolkitParameterPanelViewModel::UserCategories;
 TArray<FNiagaraParameterPanelCategory> FNiagaraScriptToolkitParameterPanelViewModel::DefaultCategories;
 TArray<FNiagaraParameterPanelCategory> FNiagaraScriptToolkitParameterPanelViewModel::DefaultAdvancedCategories;
 TArray<FNiagaraParameterPanelCategory> FNiagaraParameterDefinitionsToolkitParameterPanelViewModel::DefaultCategories;
@@ -98,6 +101,12 @@ TArray<UNiagaraGraph*> FNiagaraSystemToolkitParameterPanelUtilities::GetAllGraph
 			OutGraphs.Add(Cast<UNiagaraScriptSource>(EmitterHandle->GetEmitterData()->GraphSource)->NodeGraph);
 		}
 	}
+
+	TArray<UNiagaraGraph*> EdGraphs = SystemViewModel->GetDocumentViewModel()->GetEditableGraphsForActiveDocument();
+	for (UNiagaraGraph* Graph : EdGraphs)
+	{
+		OutGraphs.AddUnique(Graph);
+	}
 	return OutGraphs;
 }
 
@@ -128,9 +137,16 @@ TArray<UNiagaraGraph*> FNiagaraSystemToolkitParameterPanelUtilities::GetEditable
 	{
 		EditableGraphs.Add(Cast<UNiagaraScriptSource>(SystemViewModel->GetEmitterHandleViewModels()[0]->GetEmitterHandle()->GetEmitterData()->GraphSource)->NodeGraph);
 	}
+
+	TArray<UNiagaraGraph*> EdGraphs = SystemViewModel->GetDocumentViewModel()->GetEditableGraphsForActiveDocument();
+	for (UNiagaraGraph* Graph : EdGraphs)
+	{
+		EditableGraphs.AddUnique(Graph);
+	}
 	return EditableGraphs;
 }
 
+
 FReply FNiagaraSystemToolkitParameterPanelUtilities::CreateDragEventForParameterItem(const FNiagaraParameterPanelItemBase& DraggedItem, const FPointerEvent& MouseEvent, const TArray<FNiagaraGraphParameterReference>& GraphParameterReferencesForItem, const TSharedPtr<TArray<FName>>& ParametersWithNamespaceModifierRenamePending)
 {
 	const static FText TooltipFormat = LOCTEXT("Parameters", "Name: {0} \nType: {1}");
@@ -313,6 +329,11 @@ bool INiagaraImmutableParameterPanelViewModel::GetCanCopyParameterMetaDataAndToo
 ///	Parameter Panel View Model												///
 ///////////////////////////////////////////////////////////////////////////////
 
+INiagaraParameterPanelViewModel::INiagaraParameterPanelViewModel()
+{
+	ActiveSectionIndex = -1;
+}
+
 INiagaraParameterPanelViewModel::~INiagaraParameterPanelViewModel()
 {
 	// Clean up transient UNiagaraScriptVariables used as intermediate representations.
@@ -322,6 +343,7 @@ INiagaraParameterPanelViewModel::~INiagaraParameterPanelViewModel()
 		ScriptVar->RemoveFromRoot();
 		ScriptVar = nullptr;
 	}
+
 }
 
 bool INiagaraParameterPanelViewModel::GetCanDeleteParameterAndToolTip(const FNiagaraParameterPanelItem& ItemToDelete, FText& OutCanDeleteParameterToolTip) const
@@ -936,6 +958,17 @@ void FNiagaraSystemToolkitParameterPanelViewModel::Init(const FSystemToolkitUICo
 	UIContext = InUIContext;
 	UNiagaraSystem& System = SystemViewModel->GetSystem();
 
+	VariableObjectSelection = MakeShared< FNiagaraObjectSelection>();
+
+	ActiveSectionIndex = 0;
+	UserOnlyIdx = Sections.Add(LOCTEXT("SectionUser", "User"));
+	CommonOnlyIdx = Sections.Add(LOCTEXT("SectionCommon", "Common"));
+	ActiveIdx = Sections.Add(LOCTEXT("SectionActive", "Active"));
+
+	SectionDescriptions.Add(LOCTEXT("SectionUserDesc", "Displays only User exposed parameters"));
+	SectionDescriptions.Add(LOCTEXT("SectionCommonDesc", "Displays parameters defined in Parameter Libraries"));
+	SectionDescriptions.Add(LOCTEXT("SectionActiveDesc", "Displays parameters that are context-sensitive to the selected items in the active document tab."));
+
 	// Init bindings
 	// Bind OnChanged() and OnNeedsRecompile() callbacks for all script source graphs.
 	auto GetGraphFromScript = [](UNiagaraScript* Script)->UNiagaraGraph* {
@@ -1010,10 +1043,68 @@ void FNiagaraSystemToolkitParameterPanelViewModel::Init(const FSystemToolkitUICo
 			}
 		}
 	}
+	if (UserCategories.Num() == 0)
+	{
+		TArray<FNiagaraNamespaceMetadata> NamespaceMetadata = GetDefault<UNiagaraEditorSettings>()->GetAllNamespaceMetadata();
+		for (const FNiagaraNamespaceMetadata& NamespaceMetadatum : NamespaceMetadata)
+		{
+			if (NamespaceMetadatum.GetGuid() == FNiagaraEditorGuids::UserNamespaceMetaDataGuid)
+				UserCategories.Add(NamespaceMetadatum);
+		}
+	}
 
 	RegisteredHandle = RegisterViewModelWithMap(&System, this);
 }
 
+
+void FNiagaraSystemToolkitParameterPanelViewModel::OnParameterItemSelected(const FNiagaraParameterPanelItem& SelectedItem, ESelectInfo::Type SelectInfo) const
+{
+	bool bFound = false;
+	for (const UNiagaraGraph* Graph : GetEditableGraphsConst())
+	{
+		UNiagaraScriptVariable* ParameterScriptVariable = Graph->GetScriptVariable(SelectedItem.GetVariable().GetName());
+		if (ParameterScriptVariable != nullptr)
+		{
+			VariableObjectSelection->SetSelectedObject(ParameterScriptVariable);
+			bFound = true;
+			break;
+		}
+	}
+
+	if (bFound && SystemViewModel.IsValid())
+	{
+		if (SystemViewModel->GetSystemStackViewModel())
+			SystemViewModel->GetSystemStackViewModel()->InvalidateCachedParameterUsage();
+
+		for (TSharedRef<FNiagaraEmitterHandleViewModel> EmitterHandleVM : SystemViewModel->GetEmitterHandleViewModels())
+		{
+			UNiagaraStackViewModel* EmitterStackViewModel = EmitterHandleVM->GetEmitterStackViewModel();
+			if (EmitterStackViewModel)
+			{
+				EmitterStackViewModel->InvalidateCachedParameterUsage();
+			}
+		}
+	}
+}
+
+bool FNiagaraSystemToolkitParameterPanelViewModel::IsVariableSelected(FNiagaraVariableBase& InVar) const
+{
+	if (VariableObjectSelection.IsValid())
+	{
+		const TSet<UObject*>& Objects = VariableObjectSelection->GetSelectedObjects();
+		for (UObject* Obj : Objects)
+		{
+			UNiagaraScriptVariable* ScriptVar = Cast<UNiagaraScriptVariable>(Obj);
+			if (ScriptVar && ScriptVar->Variable.IsEquivalent(InVar, false))
+			{
+				return true;
+			}
+		}
+
+	}
+
+	return false;
+}
 const TArray<UNiagaraScriptVariable*> FNiagaraSystemToolkitParameterPanelViewModel::GetEditableScriptVariablesWithName(const FName ParameterName) const
 {
 	TArray<UNiagaraScriptVariable*> EditableScriptVariables;
@@ -1120,8 +1211,15 @@ void FNiagaraSystemToolkitParameterPanelViewModel::FindOrAddParameter(FNiagaraVa
 	SelectParameterItemByName(Variable.GetName(), false);
 }
 
+bool FNiagaraSystemToolkitParameterPanelViewModel::GetShowReferenceCounter() const
+{
+	return (ActiveSectionIndex != CommonOnlyIdx);
+}
+
 bool FNiagaraSystemToolkitParameterPanelViewModel::GetCanAddParametersToCategory(FNiagaraParameterPanelCategory Category) const
 {
+	if (ActiveSectionIndex == CommonOnlyIdx)
+		return false;
 	return GetEditableGraphsConst().Num() > 0 && Category.NamespaceMetaData.GetGuid() != FNiagaraEditorGuids::StaticSwitchNamespaceMetaDataGuid;
 }
 
@@ -1406,16 +1504,67 @@ void FNiagaraSystemToolkitParameterPanelViewModel::DebugParameters(const TArray<
 
 const TArray<FNiagaraParameterPanelCategory>& FNiagaraSystemToolkitParameterPanelViewModel::GetDefaultCategories() const
 {
-//	const UNiagaraEditorSettings* NiagaraEditorSettings = GetDefault<UNiagaraEditorSettings>();
-//	const bool bShowAdvanced = NiagaraEditorSettings->GetDisplayAdvancedParameterPanelCategories();
-// 	if (bShowAdvanced)
-// 	{
-// 		CachedCurrentCategories = FNiagaraSystemToolkitParameterPanelViewModel::DefaultAdvancedCategories;
-// 		return FNiagaraSystemToolkitParameterPanelViewModel::DefaultAdvancedCategories;
-// 	}
-	// For now always show advanced.
-	CachedCurrentCategories = FNiagaraSystemToolkitParameterPanelViewModel::DefaultAdvancedCategories;
-	return FNiagaraSystemToolkitParameterPanelViewModel::DefaultAdvancedCategories;
+	// Show advanced in Active/Common case. Only show User when in user mode.
+	bool bForceAll = false;
+	bool bUserOnly = false;
+	bool bCommonOnly = false;
+	if (ActiveSectionIndex == UserOnlyIdx && UserOnlyIdx != -1)
+		bUserOnly = true;
+	else if (ActiveSectionIndex == ActiveIdx && ActiveIdx != -1)
+		bForceAll = true;
+	else if (ActiveSectionIndex == CommonOnlyIdx && ActiveIdx != -1)
+		bCommonOnly = true;
+
+
+	if (bUserOnly)
+	{
+		CachedCurrentCategories = FNiagaraSystemToolkitParameterPanelViewModel::UserCategories;
+	}
+	else if (bCommonOnly)
+	{
+		// For common, dynamically create the cached array as it will depend on the parameter libraries that are live.
+		CachedCurrentCategories.Reset();
+
+		TArray<FNiagaraNamespaceMetadata> NamespaceMetadata = GetDefault<UNiagaraEditorSettings>()->GetAllNamespaceMetadata();
+		TArray<UNiagaraParameterDefinitions*> Defs = GetAvailableParameterDefinitions(false);
+		for (const FNiagaraNamespaceMetadata& NamespaceMetadatum : NamespaceMetadata)
+		{
+			for (int32 i = 0; i < NamespaceMetadatum.Namespaces.Num(); i++)
+			{
+				bool bFound = false;
+				for (UNiagaraParameterDefinitions* Def : Defs)
+				{
+					if (Def)
+					{
+						for (const UNiagaraScriptVariable* ScriptVar : Def->GetParametersConst())
+						{
+							if (ScriptVar && ScriptVar->Variable.IsInNameSpace(NamespaceMetadatum.Namespaces[i]))
+							{
+								bFound = true;
+								break;
+							}
+						}
+					}
+
+					if (bFound)
+						break;
+				}
+
+				if (bFound)
+				{
+					CachedCurrentCategories.Add(NamespaceMetadatum);
+				}
+			}
+		}
+	}
+	else
+	{
+		CachedCurrentCategories = FNiagaraSystemToolkitParameterPanelViewModel::DefaultAdvancedCategories;
+	}
+
+	
+
+	return CachedCurrentCategories;
 }
 
 FMenuAndSearchBoxWidgets FNiagaraSystemToolkitParameterPanelViewModel::GetParameterMenu(FNiagaraParameterPanelCategory Category) const
@@ -1721,8 +1870,40 @@ TArray<UNiagaraGraph*> FNiagaraSystemToolkitParameterPanelViewModel::GetEditable
 	return FNiagaraSystemToolkitParameterPanelUtilities::GetEditableGraphs(SystemViewModel, SystemGraphSelectionViewModelWeak);
 }
 
+void INiagaraParameterPanelViewModel::SetActiveSection(FText & InSection)
+{
+	ActiveSectionIndex = -1;
+	for (int32 Idx = 0; Idx <Sections.Num(); Idx++)
+	{
+		if (Sections[Idx].EqualTo(InSection))
+		{
+			ActiveSectionIndex = Idx;
+		}
+	}
+}
+
+FText INiagaraParameterPanelViewModel::GetTooltipForSection(FText& InSection) const
+{
+	for (int32 i = 0; i < Sections.Num(); i++)
+	{
+		if (Sections[i].EqualTo(InSection))
+			return SectionDescriptions[i];
+	}
+
+	return FText::GetEmpty();
+}
 TArray<FNiagaraParameterPanelItem> FNiagaraSystemToolkitParameterPanelViewModel::GetViewedParameterItems() const
 {
+	bool bForceAll = false;
+	bool bUserOnly = false;
+	bool bCommonOnly = false;
+	if (ActiveSectionIndex == UserOnlyIdx && UserOnlyIdx != -1)
+		bUserOnly = true;
+	else if (ActiveSectionIndex == CommonOnlyIdx && CommonOnlyIdx != -1)
+		bCommonOnly = true;
+	else if (ActiveSectionIndex == ActiveIdx && ActiveIdx != -1)
+		bForceAll = true;
+
 	// On the first time opening the parameter panel view model we are not guaranteed to call GetDefaultCategories() before GetViewedParameterItems(). 
 	// We require CachedCurrentCategories being set as this is used to filter out parameter items that are being viewed. If CachedCurrentCategories 
 	// is not set, call GetDefaultCategories() to initialize it. 
@@ -1731,10 +1912,11 @@ TArray<FNiagaraParameterPanelItem> FNiagaraSystemToolkitParameterPanelViewModel:
 		GetDefaultCategories();
 	}
 
+
 	TMap<FNiagaraVariable, FNiagaraParameterPanelItem> VisitedParameterToItemMap;
 	TMap<FNiagaraVariable, TObjectPtr<UNiagaraScriptVariable>> ParameterToScriptVariableMap;
-	const TArray<UNiagaraGraph*> EditableGraphs = GetEditableGraphsConst();
-	const TArray<UNiagaraGraph*> AllGraphs = GetAllGraphsConst();
+	TArray<UNiagaraGraph*> EditableGraphs = GetEditableGraphsConst();
+	TArray<UNiagaraGraph*> AllGraphs = GetAllGraphsConst();
 
 	// Collect all metadata to be packaged with the FNiagaraParameterPanelItems.
 	for (const UNiagaraGraph* EditableGraph : EditableGraphs)
@@ -1778,49 +1960,80 @@ TArray<FNiagaraParameterPanelItem> FNiagaraSystemToolkitParameterPanelViewModel:
 	};
 
 	// Collect user parameters from system.
-	CollectParamStore(&SystemViewModel->GetSystem().GetExposedParameters());
-
-	// Collect parameters added to the system asset.
-	for (UNiagaraScriptVariable* EditorOnlyScriptVar : SystemViewModel->GetEditorOnlyParametersAdapter()->GetParameters())
+	if (bUserOnly || bForceAll)
 	{
-		ParameterToScriptVariableMap.Add(EditorOnlyScriptVar->Variable, EditorOnlyScriptVar);
+		CollectParamStore(&SystemViewModel->GetSystem().GetExposedParameters());
+	}
 
-		const FNiagaraVariable& Var = EditorOnlyScriptVar->Variable;
-		FNiagaraParameterPanelItem Item = FNiagaraParameterPanelItem();
-		Item.ScriptVariable = EditorOnlyScriptVar;
-		Item.NamespaceMetaData = FNiagaraEditorUtilities::GetNamespaceMetaDataForVariableName(Var.GetName());
-		Item.bExternallyReferenced = false;
-		Item.bSourcedFromCustomStackContext = false;
-		Item.ReferenceCount = 0;
+	if (bCommonOnly)
+	{
+		TArray<UNiagaraParameterDefinitions*> Defs = GetAvailableParameterDefinitions(false);
+		for (UNiagaraParameterDefinitions* Def : Defs)
+		{
+			if (Def)
+			{
+				for (const UNiagaraScriptVariable* ScriptVar : Def->GetParametersConst())
+				{
+					FNiagaraVariableBase Var = ScriptVar->Variable;
+					FNiagaraParameterPanelItem Item = FNiagaraParameterPanelItem();
+					Item.ScriptVariable = ScriptVar;
+					Item.NamespaceMetaData = FNiagaraEditorUtilities::GetNamespaceMetaDataForVariableName(Var.GetName());
+					Item.bExternallyReferenced = false;
+					Item.bSourcedFromCustomStackContext = false;
+					Item.ReferenceCount = 0;
 
-		// Determine whether the item is name aliasing a parameter definition's parameter.
-		Item.DefinitionMatchState = FNiagaraParameterDefinitionsUtilities::GetDefinitionMatchStateForParameter(EditorOnlyScriptVar->Variable);
+					// Determine whether the item is name aliasing a parameter definition's parameter.
+					Item.DefinitionMatchState = FNiagaraParameterDefinitionsUtilities::GetDefinitionMatchStateForParameter(ScriptVar->Variable);
 
-		VisitedParameterToItemMap.Add(Var, Item);
+					VisitedParameterToItemMap.Add(Var, Item);
+				}
+			}
+		}
 	}
 
-	// Collect parameters for all emitters.
-	TArray<FNiagaraVariable> VisitedInvalidParameters;
-
-	// -For each selected graph perform a parameter map history traversal and record each unique visited parameter.
-	for (const UNiagaraGraph* EditableGraph : EditableGraphs)
+	if (bForceAll)
 	{
-		TArray<UNiagaraNodeOutput*> OutputNodes;
-		EditableGraph->GetNodesOfClass<UNiagaraNodeOutput>(OutputNodes);
-		for (UNiagaraNodeOutput* OutputNode : OutputNodes)
+		// Collect parameters added to the system asset.
+		for (UNiagaraScriptVariable* EditorOnlyScriptVar : SystemViewModel->GetEditorOnlyParametersAdapter()->GetParameters())
 		{
-			UNiagaraNode* NodeToTraverse = OutputNode;
-			if (OutputNode->GetUsage() == ENiagaraScriptUsage::SystemSpawnScript || OutputNode->GetUsage() == ENiagaraScriptUsage::SystemUpdateScript)
+			ParameterToScriptVariableMap.Add(EditorOnlyScriptVar->Variable, EditorOnlyScriptVar);
+
+			const FNiagaraVariable& Var = EditorOnlyScriptVar->Variable;
+			FNiagaraParameterPanelItem Item = FNiagaraParameterPanelItem();
+			Item.ScriptVariable = EditorOnlyScriptVar;
+			Item.NamespaceMetaData = FNiagaraEditorUtilities::GetNamespaceMetaDataForVariableName(Var.GetName());
+			Item.bExternallyReferenced = false;
+			Item.bSourcedFromCustomStackContext = false;
+			Item.ReferenceCount = 0;
+
+			// Determine whether the item is name aliasing a parameter definition's parameter.
+			Item.DefinitionMatchState = FNiagaraParameterDefinitionsUtilities::GetDefinitionMatchStateForParameter(EditorOnlyScriptVar->Variable);
+
+			VisitedParameterToItemMap.Add(Var, Item);
+		}
+
+		// Collect parameters for all emitters.
+		TArray<FNiagaraVariable> VisitedInvalidParameters;
+
+		// -For each selected graph perform a parameter map history traversal and record each unique visited parameter.
+		for (const UNiagaraGraph* EditableGraph : EditableGraphs)
+		{
+			TArray<UNiagaraNodeOutput*> OutputNodes;
+			EditableGraph->GetNodesOfClass<UNiagaraNodeOutput>(OutputNodes);
+			for (UNiagaraNodeOutput* OutputNode : OutputNodes)
 			{
-				// Traverse past the emitter nodes, otherwise the system scripts will pick up all of the emitter and particle script parameters.
-				UEdGraphPin* InputPin = FNiagaraStackGraphUtilities::GetParameterMapInputPin(*NodeToTraverse);
-				while (NodeToTraverse != nullptr && InputPin != nullptr && InputPin->LinkedTo.Num() == 1 &&
-					(NodeToTraverse->IsA<UNiagaraNodeOutput>() || NodeToTraverse->IsA<UNiagaraNodeEmitter>()))
+				UNiagaraNode* NodeToTraverse = OutputNode;
+				if (OutputNode->GetUsage() == ENiagaraScriptUsage::SystemSpawnScript || OutputNode->GetUsage() == ENiagaraScriptUsage::SystemUpdateScript)
 				{
-					NodeToTraverse = Cast<UNiagaraNode>(InputPin->LinkedTo[0]->GetOwningNode());
-					InputPin = NodeToTraverse != nullptr ? FNiagaraStackGraphUtilities::GetParameterMapInputPin(*NodeToTraverse) : nullptr;
+					// Traverse past the emitter nodes, otherwise the system scripts will pick up all of the emitter and particle script parameters.
+					UEdGraphPin* InputPin = FNiagaraStackGraphUtilities::GetParameterMapInputPin(*NodeToTraverse);
+					while (NodeToTraverse != nullptr && InputPin != nullptr && InputPin->LinkedTo.Num() == 1 &&
+						(NodeToTraverse->IsA<UNiagaraNodeOutput>() || NodeToTraverse->IsA<UNiagaraNodeEmitter>()))
+					{
+						NodeToTraverse = Cast<UNiagaraNode>(InputPin->LinkedTo[0]->GetOwningNode());
+						InputPin = NodeToTraverse != nullptr ? FNiagaraStackGraphUtilities::GetParameterMapInputPin(*NodeToTraverse) : nullptr;
+					}
 				}
-			}
 
 			if (NodeToTraverse == nullptr)
 			{
@@ -1850,107 +2063,108 @@ TArray<FNiagaraParameterPanelItem> FNiagaraSystemToolkitParameterPanelViewModel:
 			NodeToTraverse->BuildParameterMapHistory(Builder, true, true);
 			Builder.EndUsage();
 
-			if (Builder.Histories.Num() != 1)
-			{
-				// We should only have traversed one emitter (have not visited more than one NiagaraNodeEmitter.)
-				ensureMsgf(false, TEXT("Encountered more than one parameter map history when collecting parameters for system parameter panel view model!"));
-			}
-	
-			// Get all UNiagaraScriptVariables of visited graphs in the ParameterToScriptVariableMap so that generated items are in sync.
-			TSet<UNiagaraGraph*> VisitedExternalGraphs;
-			for (const UEdGraphPin* MapPin : Builder.Histories[0].MapPinHistory)
-			{
-				const UNiagaraNodeFunctionCall* MapPinOuterFunctionCallNode = Cast<UNiagaraNodeFunctionCall>(MapPin->GetOuter());
-				if (MapPinOuterFunctionCallNode != nullptr)
+				if (Builder.Histories.Num() != 1)
 				{
-					UNiagaraGraph* VisitedExternalGraph = MapPinOuterFunctionCallNode->GetCalledGraph();
-					if (VisitedExternalGraph != nullptr)
-					{
-						VisitedExternalGraphs.Add(VisitedExternalGraph);
-					}
+					// We should only have traversed one emitter (have not visited more than one NiagaraNodeEmitter.)
+					ensureMsgf(false, TEXT("Encountered more than one parameter map history when collecting parameters for system parameter panel view model!"));
 				}
-			}
-			for (const UNiagaraGraph* VisitedExternalGraph : VisitedExternalGraphs)
-			{
-				ParameterToScriptVariableMap.Append(VisitedExternalGraph->GetAllMetaData());
-			}
 
-			const TArray<FName>& CustomIterationSourceNamespaces = Builder.Histories[0].IterationNamespaceOverridesEncountered;
-			for (int32 VariableIndex = 0; VariableIndex < Builder.Histories[0].Variables.Num(); VariableIndex++)
-			{
-				const FNiagaraVariable& Var = Builder.Histories[0].Variables[VariableIndex];
-				// If this variable has already been visited and does not have a valid namespace then skip it.
-				if (VisitedInvalidParameters.Contains(Var))
+				// Get all UNiagaraScriptVariables of visited graphs in the ParameterToScriptVariableMap so that generated items are in sync.
+				TSet<UNiagaraGraph*> VisitedExternalGraphs;
+				for (const UEdGraphPin* MapPin : Builder.Histories[0].MapPinHistory)
 				{
-					continue;
+					const UNiagaraNodeFunctionCall* MapPinOuterFunctionCallNode = Cast<UNiagaraNodeFunctionCall>(MapPin->GetOuter());
+					if (MapPinOuterFunctionCallNode != nullptr)
+					{
+						UNiagaraGraph* VisitedExternalGraph = MapPinOuterFunctionCallNode->GetCalledGraph();
+						if (VisitedExternalGraph != nullptr)
+						{
+							VisitedExternalGraphs.Add(VisitedExternalGraph);
+						}
+					}
 				}
-
-				if(FNiagaraParameterPanelItem* ItemPtr = VisitedParameterToItemMap.Find(Var))
+				for (const UNiagaraGraph* VisitedExternalGraph : VisitedExternalGraphs)
 				{
-					// This variable has already been registered, increment the reference count.
-					ItemPtr->ReferenceCount += Builder.Histories[0].PerVariableReadHistory[VariableIndex].Num() + Builder.Histories[0].PerVariableWriteHistory[VariableIndex].Num();
+					ParameterToScriptVariableMap.Append(VisitedExternalGraph->GetAllMetaData());
 				}
-				else
+
+				const TArray<FName>& CustomIterationSourceNamespaces = Builder.Histories[0].IterationNamespaceOverridesEncountered;
+				for (int32 VariableIndex = 0; VariableIndex < Builder.Histories[0].Variables.Num(); VariableIndex++)
 				{
-					// This variable has not been registered, prepare the FNiagaraParameterPanelItem.
-					// -First make sure the variable namespace is in a valid category. If not, skip it.
-					FNiagaraNamespaceMetadata CandidateNamespaceMetaData = FNiagaraEditorUtilities::GetNamespaceMetaDataForVariableName(Var.GetName());
-					if (CachedCurrentCategories.Contains(FNiagaraParameterPanelCategory(CandidateNamespaceMetaData)) == false)
+					const FNiagaraVariable& Var = Builder.Histories[0].Variables[VariableIndex];
+					// If this variable has already been visited and does not have a valid namespace then skip it.
+					if (VisitedInvalidParameters.Contains(Var))
 					{
-						VisitedInvalidParameters.Add(Var);
 						continue;
 					}
 
-					// -Lookup the script variable.
-					TObjectPtr<UNiagaraScriptVariable> const* ScriptVarPtr = ParameterToScriptVariableMap.Find(Var);
-					UNiagaraScriptVariable* ScriptVar;
-					if (ScriptVarPtr != nullptr)
+					if (FNiagaraParameterPanelItem* ItemPtr = VisitedParameterToItemMap.Find(Var))
 					{
-						ScriptVar = *ScriptVarPtr;
+						// This variable has already been registered, increment the reference count.
+						ItemPtr->ReferenceCount += Builder.Histories[0].PerVariableReadHistory[VariableIndex].Num() + Builder.Histories[0].PerVariableWriteHistory[VariableIndex].Num();
 					}
 					else
 					{
-						// Create a new UNiagaraScriptVariable to represent this parameter for the lifetime of the ParameterPanelViewModel.
-						ScriptVar = NewObject<UNiagaraScriptVariable>(GetTransientPackage());
-						ScriptVar->AddToRoot();
-						ScriptVar->Init(Var, FNiagaraVariableMetaData());
-						TransientParameterToScriptVarMap.Add(Var, ScriptVar);
-					}
+						// This variable has not been registered, prepare the FNiagaraParameterPanelItem.
+						// -First make sure the variable namespace is in a valid category. If not, skip it.
+						FNiagaraNamespaceMetadata CandidateNamespaceMetaData = FNiagaraEditorUtilities::GetNamespaceMetaDataForVariableName(Var.GetName());
+						if (CachedCurrentCategories.Contains(FNiagaraParameterPanelCategory(CandidateNamespaceMetaData)) == false)
+						{
+							VisitedInvalidParameters.Add(Var);
+							continue;
+						}
 
-					bool bVarOnlyInTopLevelGraph = true;
-					for (FModuleScopedPin& WritePin : Builder.Histories[0].PerVariableWriteHistory[VariableIndex])
-					{
-						UEdGraphNode* VariableOwningNode = WritePin.Pin->GetOwningNode();
-						bVarOnlyInTopLevelGraph &= AllGraphs.Contains(static_cast<const UNiagaraGraph*>(VariableOwningNode->GetGraph()));
-					}
-					for (FNiagaraParameterMapHistory::FReadHistory& ReadPins : Builder.Histories[0].PerVariableReadHistory[VariableIndex])
-					{
-						bVarOnlyInTopLevelGraph &= !AllGraphs.Contains(static_cast<const UNiagaraGraph*>(ReadPins.ReadPin.Pin->GetOwningNode()->GetGraph()));
-					}
+						// -Lookup the script variable.
+						TObjectPtr<UNiagaraScriptVariable> const* ScriptVarPtr = ParameterToScriptVariableMap.Find(Var);
+						UNiagaraScriptVariable* ScriptVar;
+						if (ScriptVarPtr != nullptr)
+						{
+							ScriptVar = *ScriptVarPtr;
+						}
+						else
+						{
+							// Create a new UNiagaraScriptVariable to represent this parameter for the lifetime of the ParameterPanelViewModel.
+							ScriptVar = NewObject<UNiagaraScriptVariable>(GetTransientPackage());
+							ScriptVar->AddToRoot();
+							ScriptVar->Init(Var, FNiagaraVariableMetaData());
+							TransientParameterToScriptVarMap.Add(Var, ScriptVar);
+						}
 
-					FNiagaraParameterPanelItem Item = FNiagaraParameterPanelItem();
-					Item.ScriptVariable = ScriptVar;
-					Item.NamespaceMetaData = CandidateNamespaceMetaData;
-					Item.bExternallyReferenced = !bVarOnlyInTopLevelGraph;
-					
-					// -Determine whether the parameter is from a custom stack context.
-					Item.bSourcedFromCustomStackContext = false;
-					for (const FName CustomIterationNamespace : CustomIterationSourceNamespaces)
-					{
-						if (Var.IsInNameSpace(CustomIterationNamespace))
+						bool bVarOnlyInTopLevelGraph = true;
+						for (FModuleScopedPin& WritePin : Builder.Histories[0].PerVariableWriteHistory[VariableIndex])
 						{
-							Item.bSourcedFromCustomStackContext = true;
-							break;
+							UEdGraphNode* VariableOwningNode = WritePin.Pin->GetOwningNode();
+							bVarOnlyInTopLevelGraph &= AllGraphs.Contains(static_cast<const UNiagaraGraph*>(VariableOwningNode->GetGraph()));
+						}
+						for (FNiagaraParameterMapHistory::FReadHistory& ReadPins : Builder.Histories[0].PerVariableReadHistory[VariableIndex])
+						{
+							bVarOnlyInTopLevelGraph &= !AllGraphs.Contains(static_cast<const UNiagaraGraph*>(ReadPins.ReadPin.Pin->GetOwningNode()->GetGraph()));
 						}
-					}
 
-					// Determine whether the item is name aliasing a parameter definition's parameter.
-					Item.DefinitionMatchState = FNiagaraParameterDefinitionsUtilities::GetDefinitionMatchStateForParameter(Item.ScriptVariable->Variable);
+						FNiagaraParameterPanelItem Item = FNiagaraParameterPanelItem();
+						Item.ScriptVariable = ScriptVar;
+						Item.NamespaceMetaData = CandidateNamespaceMetaData;
+						Item.bExternallyReferenced = !bVarOnlyInTopLevelGraph;
 
-					// -Increment the reference count.
-					Item.ReferenceCount += Builder.Histories[0].PerVariableReadHistory[VariableIndex].Num() + Builder.Histories[0].PerVariableWriteHistory[VariableIndex].Num();
+						// -Determine whether the parameter is from a custom stack context.
+						Item.bSourcedFromCustomStackContext = false;
+						for (const FName CustomIterationNamespace : CustomIterationSourceNamespaces)
+						{
+							if (Var.IsInNameSpace(CustomIterationNamespace))
+							{
+								Item.bSourcedFromCustomStackContext = true;
+								break;
+							}
+						}
 
-					VisitedParameterToItemMap.Add(Var, Item);
+						// Determine whether the item is name aliasing a parameter definition's parameter.
+						Item.DefinitionMatchState = FNiagaraParameterDefinitionsUtilities::GetDefinitionMatchStateForParameter(Item.ScriptVariable->Variable);
+
+						// -Increment the reference count.
+						Item.ReferenceCount += Builder.Histories[0].PerVariableReadHistory[VariableIndex].Num() + Builder.Histories[0].PerVariableWriteHistory[VariableIndex].Num();
+
+						VisitedParameterToItemMap.Add(Var, Item);
+					}
 				}
 			}
 		}
@@ -2748,8 +2962,9 @@ const TArray<FNiagaraParameterPanelCategory>& FNiagaraScriptToolkitParameterPane
 // 		return FNiagaraScriptToolkitParameterPanelViewModel::DefaultAdvancedCategories;
 // 	}
 	// For now always show advanced.
+
 	CachedCurrentCategories = FNiagaraScriptToolkitParameterPanelViewModel::DefaultAdvancedCategories;
-	return FNiagaraScriptToolkitParameterPanelViewModel::DefaultAdvancedCategories;
+	return CachedCurrentCategories;
 }
 
 FMenuAndSearchBoxWidgets FNiagaraScriptToolkitParameterPanelViewModel::GetParameterMenu(FNiagaraParameterPanelCategory Category) const
@@ -2933,8 +3148,10 @@ TArray<FNiagaraParameterPanelItem> FNiagaraScriptToolkitParameterPanelViewModel:
 	TMap<FNiagaraVariable, FNiagaraParameterPanelItem> VisitedParameterToItemMap;
 	TArray<FNiagaraVariable> VisitedInvalidParameters;
 
+	TArray<UNiagaraGraph*> Graphs = GetEditableGraphsConst();
+
 	// For scripts we use the reference maps cached in the graph to collect parameters.
-	for (const UNiagaraGraph* Graph : GetEditableGraphsConst())
+	for (const UNiagaraGraph* Graph : Graphs)
 	{
 		TMap<FNiagaraVariable, TObjectPtr<UNiagaraScriptVariable>> ParameterToScriptVariableMap;
 		ParameterToScriptVariableMap.Append(Graph->GetAllMetaData());
diff --git a/Engine/Plugins/FX/Niagara/Source/NiagaraEditor/Private/ViewModels/NiagaraScratchPadScriptViewModel.cpp b/Engine/Plugins/FX/Niagara/Source/NiagaraEditor/Private/ViewModels/NiagaraScratchPadScriptViewModel.cpp
--- a/Engine/Plugins/FX/Niagara/Source/NiagaraEditor/Private/ViewModels/NiagaraScratchPadScriptViewModel.cpp
+++ b/Engine/Plugins/FX/Niagara/Source/NiagaraEditor/Private/ViewModels/NiagaraScratchPadScriptViewModel.cpp
@@ -12,6 +12,7 @@
 #include "NiagaraObjectSelection.h"
 #include "NiagaraParameterDefinitions.h"
 #include "ViewModels/NiagaraScratchPadUtilities.h"
+#include "ViewModels/NiagaraScriptGraphViewModel.h"
 
 #include "ScopedTransaction.h"
 #include "Framework/Commands/UICommandList.h"
@@ -89,6 +90,16 @@ const FVersionedNiagaraScript& FNiagaraScratchPadScriptViewModel::GetEditScript(
 	return EditScript;
 }
 
+TArray<UNiagaraGraph*> FNiagaraScratchPadScriptViewModel::GetEditableGraphs() 
+{
+	TArray<UNiagaraGraph*> EditableGraphs;
+	if (UNiagaraGraph* Graph = GetGraphViewModel()->GetGraph())
+	{
+		EditableGraphs.Add(Graph);
+	}
+	return EditableGraphs;
+}
+
 TSharedPtr<INiagaraParameterPanelViewModel> FNiagaraScratchPadScriptViewModel::GetParameterPanelViewModel() const
 {
 	return ParameterPaneViewModel;
diff --git a/Engine/Plugins/FX/Niagara/Source/NiagaraEditor/Private/ViewModels/NiagaraScratchPadViewModel.cpp b/Engine/Plugins/FX/Niagara/Source/NiagaraEditor/Private/ViewModels/NiagaraScratchPadViewModel.cpp
--- a/Engine/Plugins/FX/Niagara/Source/NiagaraEditor/Private/ViewModels/NiagaraScratchPadViewModel.cpp
+++ b/Engine/Plugins/FX/Niagara/Source/NiagaraEditor/Private/ViewModels/NiagaraScratchPadViewModel.cpp
@@ -6,6 +6,7 @@
 #include "ViewModels/NiagaraScratchPadScriptViewModel.h"
 #include "ViewModels/NiagaraScriptGraphViewModel.h"
 #include "ViewModels/NiagaraParameterPanelViewModel.h"
+#include "ViewModels/NiagaraSystemEditorDocumentsViewModel.h"
 #include "ViewModels/Stack/NiagaraStackModuleItem.h"
 #include "ViewModels/Stack/NiagaraStackFunctionInput.h"
 #include "Toolkits/SystemToolkitModes/NiagaraSystemToolkitModeBase.h"
@@ -262,7 +263,7 @@ void UNiagaraScratchPadViewModel::FocusScratchPadScriptViewModel(TSharedRef<FNia
 	if (ensureMsgf(ScriptViewModels.Contains(InScriptViewModel), TEXT("Can only focus a view model from this scratch pad view model.")))
 	{
 		SetActiveScriptViewModel(InScriptViewModel);
-		GetSystemViewModel()->FocusTab(FNiagaraSystemToolkitModeBase::ScratchPadTabID);
+		GetSystemViewModel()->GetDocumentViewModel()->OpenChildScript(InScriptViewModel->GetGraphViewModel()->GetGraph());
 	}
 }
 
@@ -374,6 +375,9 @@ void UNiagaraScratchPadViewModel::DeleteActiveScript()
 			}
 		}
 
+		// Close any open tabs
+		GetSystemViewModel()->GetDocumentViewModel()->CloseChildScript(ActiveScriptViewModel->GetGraphViewModel()->GetGraph());
+
 		UObject* ScriptOuter;
 		TArray<TObjectPtr<UNiagaraScript>>* TargetScripts;
 		GetOuterAndTargetScripts(GetSystemViewModel(), ScriptOuter, TargetScripts);
@@ -508,6 +512,19 @@ void UNiagaraScratchPadViewModel::CreateAssetFromActiveScript()
 	}
 }
 
+void UNiagaraScratchPadViewModel::OpenEditorForActive()
+{
+	if (ActiveScriptViewModel.IsValid())
+	{
+		UEdGraph* Graph = ActiveScriptViewModel->GetGraphViewModel()->GetGraph();
+		if (Graph)
+		{
+			GetSystemViewModel()->GetDocumentViewModel()->OpenChildScript(Graph);
+		}
+	}
+}
+
+
 bool UNiagaraScratchPadViewModel::CanSelectNextUsageForActiveScript()
 {
 	if (ActiveScriptViewModel.IsValid())
@@ -627,7 +644,7 @@ TSharedRef<FNiagaraScratchPadScriptViewModel> UNiagaraScratchPadViewModel::Creat
 	TSharedRef<FNiagaraScratchPadScriptViewModel> ScriptViewModel = MakeShared<FNiagaraScratchPadScriptViewModel>(GetSystemViewModel()->GetIsForDataProcessingOnly());
 	ScriptViewModel->Initialize(ScratchPadScript);
 	ScriptViewModel->GetGraphViewModel()->GetNodeSelection()->OnSelectedObjectsChanged().AddUObject(this, &UNiagaraScratchPadViewModel::ScriptGraphNodeSelectionChanged, TWeakPtr<FNiagaraScratchPadScriptViewModel>(ScriptViewModel));
-	ScriptViewModel->OnRenamed().AddUObject(this, &UNiagaraScratchPadViewModel::ScriptViewModelScriptRenamed);
+	ScriptViewModel->OnRenamed().AddUObject(this, &UNiagaraScratchPadViewModel::ScriptViewModelScriptRenamed, TWeakPtr<FNiagaraScratchPadScriptViewModel>(ScriptViewModel));
 	ScriptViewModel->OnPinnedChanged().AddUObject(this, &UNiagaraScratchPadViewModel::ScriptViewModelPinnedChanged, TWeakPtr<FNiagaraScratchPadScriptViewModel>(ScriptViewModel));
 	ScriptViewModel->OnHasUnappliedChangesChanged().AddUObject(this, &UNiagaraScratchPadViewModel::ScriptViewModelHasUnappliedChangesChanged);
 	ScriptViewModel->OnChangesApplied().AddUObject(this, &UNiagaraScratchPadViewModel::ScriptViewModelChangesApplied);
@@ -681,10 +698,21 @@ void UNiagaraScratchPadViewModel::ScriptGraphNodeSelectionChanged(TWeakPtr<FNiag
 	}
 }
 
-void UNiagaraScratchPadViewModel::ScriptViewModelScriptRenamed()
+void UNiagaraScratchPadViewModel::ScriptViewModelScriptRenamed(TWeakPtr<FNiagaraScratchPadScriptViewModel> ScriptViewModelWeak)
 {
+	TSharedPtr<FNiagaraScratchPadScriptViewModel> ScriptViewModel = ScriptViewModelWeak.Pin();
+	if (ScriptViewModel.IsValid())
+	{
+		GetSystemViewModel()->GetDocumentViewModel()->CloseChildScript(ScriptViewModel->GetGraphViewModel()->GetGraph());
+	}
+
 	UpdateChangeId(GetSystemViewModel());
 	OnScriptRenamed().Broadcast();
+
+	if (ScriptViewModel.IsValid())
+	{
+		GetSystemViewModel()->GetDocumentViewModel()->OpenChildScript(ScriptViewModel->GetGraphViewModel()->GetGraph());
+	}
 }
 
 void UNiagaraScratchPadViewModel::ScriptViewModelPinnedChanged(TWeakPtr<FNiagaraScratchPadScriptViewModel> ScriptViewModelWeak)
diff --git a/Engine/Plugins/FX/Niagara/Source/NiagaraEditor/Private/ViewModels/NiagaraSystemEditorDocumentsViewModel.cpp b/Engine/Plugins/FX/Niagara/Source/NiagaraEditor/Private/ViewModels/NiagaraSystemEditorDocumentsViewModel.cpp
new file mode 100644
--- /dev/null
+++ b/Engine/Plugins/FX/Niagara/Source/NiagaraEditor/Private/ViewModels/NiagaraSystemEditorDocumentsViewModel.cpp
@@ -0,0 +1,218 @@
+// Copyright Epic Games, Inc. All Rights Reserved.
+
+#include "ViewModels/NiagaraSystemEditorDocumentsViewModel.h"
+#include "ViewModels/NiagaraSystemViewModel.h"
+#include "ViewModels/NiagaraSystemSelectionViewModel.h"
+#include "ViewModels/NiagaraScratchPadViewModel.h"
+#include "ViewModels/NiagaraScratchPadScriptViewModel.h"
+#include "ViewModels/NiagaraScriptGraphViewModel.h"
+#include "ViewModels/NiagaraParameterPanelViewModel.h"
+#include "ViewModels/Stack/NiagaraStackModuleItem.h"
+#include "ViewModels/Stack/NiagaraStackFunctionInput.h"
+#include "Toolkits/SystemToolkitModes/NiagaraSystemToolkitModeBase.h"
+#include "NiagaraSystem.h"
+#include "NiagaraScriptSource.h"
+#include "NiagaraNodeOutput.h"
+#include "NiagaraObjectSelection.h"
+#include "NiagaraEditorSettings.h"
+#include "NiagaraNodeFunctionCall.h"
+#include "NiagaraEditorModule.h"
+#include "NiagaraClipboard.h"
+#include "Toolkits/NiagaraSystemToolkit.h"
+#include "NiagaraEditorUtilities.h"
+#include "NiagaraScriptFactoryNew.h"
+#include "NiagaraGraph.h"
+
+#include "ScopedTransaction.h"
+#include "Modules/ModuleManager.h"
+#include "AssetToolsModule.h"
+#include "AssetRegistry/IAssetRegistry.h"
+#include "Misc/PackageName.h"
+#include "Subsystems/AssetEditorSubsystem.h"
+#include "Misc/MessageDialog.h"
+#include "GenericPlatform/GenericPlatformMisc.h"
+#include "HAL/PlatformApplicationMisc.h"
+#include "WorkflowOrientedApp/WorkflowTabManager.h"
+#include "NiagaraScriptSource.h"
+#include "Widgets/NiagaraScratchScriptEditor.h"
+#include "Toolkits/NiagaraSystemToolkit.h"
+
+#define LOCTEXT_NAMESPACE "NiagaraScratchPadViewModel"
+
+void UNiagaraSystemEditorDocumentsViewModel::Initialize(TSharedRef<FNiagaraSystemViewModel> InSystemViewModel)
+{
+	SystemViewModelWeak = InSystemViewModel;
+}
+
+void UNiagaraSystemEditorDocumentsViewModel::Finalize()
+{
+	ActiveDocumentTabScriptViewModel.Reset();
+}
+
+TArray<UNiagaraGraph*> UNiagaraSystemEditorDocumentsViewModel::GetEditableGraphsForActiveDocument()
+{
+	if (ActiveDocumentTabScriptViewModel.IsValid())
+		return ActiveDocumentTabScriptViewModel->GetEditableGraphs();
+	else
+		return TArray<UNiagaraGraph*>();
+}
+
+void UNiagaraSystemEditorDocumentsViewModel::OpenChildScript(UEdGraph* InGraph)
+{
+	TArray< TSharedPtr<SDockTab> > Results;
+	if (!FindOpenTabsContainingDocument(InGraph, Results))
+	{
+		OpenDocument(InGraph, FDocumentTracker::EOpenDocumentCause::ForceOpenNewDocument);
+	}
+	else if (Results.Num() > 0)
+	{
+		TabManager->DrawAttention(Results[0].ToSharedRef());
+	}
+}
+
+void UNiagaraSystemEditorDocumentsViewModel::CloseChildScript(UEdGraph* InGraph)
+{
+	CloseDocumentTab(InGraph);
+}
+
+
+TSharedRef<FNiagaraSystemViewModel> UNiagaraSystemEditorDocumentsViewModel::GetSystemViewModel()
+{
+	TSharedPtr<FNiagaraSystemViewModel> SystemViewModel = SystemViewModelWeak.Pin();
+	checkf(SystemViewModel.IsValid(), TEXT("SystemViewModel destroyed before system editor document view model."));
+	return SystemViewModel.ToSharedRef();
+}
+
+void UNiagaraSystemEditorDocumentsViewModel::InitializePreTabManager(TSharedPtr<FNiagaraSystemToolkit> InToolkit)
+{
+	DocumentManager = MakeShareable(new FDocumentTracker);
+	DocumentManager->Initialize(InToolkit);
+
+
+	// Register the document factories
+	{
+
+		TSharedRef<FDocumentTabFactory> GraphEditorFactory = MakeShareable(new FGraphEditorSummoner(InToolkit,
+			FGraphEditorSummoner::FOnCreateGraphEditorWidget::CreateUObject(this, &UNiagaraSystemEditorDocumentsViewModel::CreateGraphEditorWidget)
+		));
+
+		// Also store off a reference to the grapheditor factory so we can find all the tabs spawned by it later.
+		GraphEditorTabFactoryPtr = GraphEditorFactory;
+		DocumentManager->RegisterDocumentFactory(GraphEditorFactory);
+	}
+}
+
+void UNiagaraSystemEditorDocumentsViewModel::InitializePostTabManager(TSharedPtr<FNiagaraSystemToolkit> InToolkit)
+{
+	TabManager = InToolkit->GetTabManager();
+	DocumentManager->SetTabManager(TabManager.ToSharedRef());
+}
+
+
+TSharedRef<SNiagaraScratchPadScriptEditor> UNiagaraSystemEditorDocumentsViewModel::CreateGraphEditorWidget(TSharedRef<FTabInfo> InTabInfo, UEdGraph* InGraph)
+{
+	TSharedRef<FNiagaraScratchPadScriptViewModel> GraphViewModel = MakeShared<FNiagaraScratchPadScriptViewModel>(false);
+
+	for (TSharedRef<FNiagaraScratchPadScriptViewModel> ScratchPadViewModel : GetSystemViewModel()->GetScriptScratchPadViewModel()->GetScriptViewModels())
+	{
+		if (ScratchPadViewModel->GetGraphViewModel()->GetGraph() == InGraph)
+		{
+			GraphViewModel = ScratchPadViewModel;
+			break;
+		}
+	}
+
+	TSharedRef<SNiagaraScratchPadScriptEditor> Editor = SNew(SNiagaraScratchPadScriptEditor, GraphViewModel);
+	return Editor;
+}
+
+
+
+TSharedPtr<SDockTab> UNiagaraSystemEditorDocumentsViewModel::OpenDocument(const UObject* DocumentID, FDocumentTracker::EOpenDocumentCause Cause)
+{
+	TSharedRef<FTabPayload_UObject> Payload = FTabPayload_UObject::Make(DocumentID);
+	return DocumentManager->OpenDocument(Payload, Cause);
+}
+
+void UNiagaraSystemEditorDocumentsViewModel::NavigateTab(FDocumentTracker::EOpenDocumentCause InCause)
+{
+	OpenDocument(nullptr, InCause);
+}
+
+void UNiagaraSystemEditorDocumentsViewModel::CloseDocumentTab(const UObject* DocumentID)
+{
+	TSharedRef<FTabPayload_UObject> Payload = FTabPayload_UObject::Make(DocumentID);
+	DocumentManager->CloseTab(Payload);
+	DocumentManager->CleanInvalidTabs();
+}
+
+// Finds any open tabs containing the specified document and adds them to the specified array; returns true if at least one is found
+bool UNiagaraSystemEditorDocumentsViewModel::FindOpenTabsContainingDocument(const UObject* DocumentID, /*inout*/ TArray< TSharedPtr<SDockTab> >& Results)
+{
+	int32 StartingCount = Results.Num();
+
+	TSharedRef<FTabPayload_UObject> Payload = FTabPayload_UObject::Make(DocumentID);
+
+	DocumentManager->FindMatchingTabs(Payload, /*inout*/ Results);
+
+	// Did we add anything new?
+	return (StartingCount != Results.Num());
+}
+
+void UNiagaraSystemEditorDocumentsViewModel::SetActiveDocumentTab(TSharedPtr<SDockTab> Tab)
+{
+	ActiveDocumentTab = Tab;
+	ActiveDocumentTabScriptViewModel = GetActiveScratchPadViewModelIfSet();
+
+	ActiveDocChangedDelegate.Broadcast(Tab);
+
+
+	// We need to update the parameter panel view model with new parameters potentially
+	if (GetSystemViewModel()->GetParameterPanelViewModel())
+		GetSystemViewModel()->GetParameterPanelViewModel()->RefreshNextTick();
+}
+
+TSharedPtr<FNiagaraScratchPadScriptViewModel> UNiagaraSystemEditorDocumentsViewModel::GetScratchPadViewModelFromGraph(FNiagaraSystemViewModel* InSysViewModel, UEdGraph* InTargetGraph)
+{
+	if (InSysViewModel && InTargetGraph)
+	{
+		UNiagaraScript* Script = InTargetGraph->GetTypedOuter<UNiagaraScript>();
+		UNiagaraScratchPadViewModel* ScratchViewModel = InSysViewModel->GetScriptScratchPadViewModel();
+		if (ScratchViewModel && Script)
+		{
+			for (TSharedRef<FNiagaraScratchPadScriptViewModel> ScratchScriptViewModel : ScratchViewModel->GetScriptViewModels())
+			{
+				UNiagaraScript* EditScript = ScratchScriptViewModel->GetGraphViewModel()->GetScriptSource()->GetTypedOuter<UNiagaraScript>();
+				if (EditScript == Script)
+				{
+					return ScratchScriptViewModel;
+				}
+				else if (ScratchScriptViewModel->GetOriginalScript() == Script)
+				{
+					return ScratchScriptViewModel;
+				}
+			}
+		}
+	}
+	return nullptr;
+}
+
+
+TSharedPtr<FNiagaraScratchPadScriptViewModel> UNiagaraSystemEditorDocumentsViewModel::GetActiveScratchPadViewModelIfSet()
+{
+	TSharedPtr<SDockTab> Tab = ActiveDocumentTab.Pin();
+	if (Tab.IsValid() && Tab->GetLayoutIdentifier().TabType == TEXT("Document"))
+	{
+		TSharedRef<SNiagaraScratchPadScriptEditor> GraphEditor = StaticCastSharedRef<SNiagaraScratchPadScriptEditor>(Tab->GetContent());
+		if (GraphEditor->GetGraphEditor())
+		{
+			TSharedPtr<FNiagaraSystemViewModel> SystemVM = SystemViewModelWeak.Pin();
+			return GetScratchPadViewModelFromGraph(SystemVM.Get(), GraphEditor->GetGraphEditor()->GetCurrentGraph());
+		}
+	}
+
+	return nullptr;
+}
+
+
+#undef LOCTEXT_NAMESPACE
\ No newline at end of file
diff --git a/Engine/Plugins/FX/Niagara/Source/NiagaraEditor/Private/ViewModels/NiagaraSystemViewModel.cpp b/Engine/Plugins/FX/Niagara/Source/NiagaraEditor/Private/ViewModels/NiagaraSystemViewModel.cpp
--- a/Engine/Plugins/FX/Niagara/Source/NiagaraEditor/Private/ViewModels/NiagaraSystemViewModel.cpp
+++ b/Engine/Plugins/FX/Niagara/Source/NiagaraEditor/Private/ViewModels/NiagaraSystemViewModel.cpp
@@ -19,6 +19,7 @@
 #include "ViewModels/Stack/NiagaraStackGraphUtilities.h"
 #include "ViewModels/Stack/NiagaraStackEntry.h"
 #include "ViewModels/NiagaraPlaceholderDataInterfaceManager.h"
+#include "ViewModels/NiagaraSystemEditorDocumentsViewModel.h"
 #include "NiagaraSystemScriptViewModel.h"
 #include "NiagaraScriptGraphViewModel.h"
 #include "NiagaraSequence.h"
@@ -85,6 +86,7 @@ FNiagaraSystemViewModel::FNiagaraSystemViewModel()
 	, bResetRequestPending(false)
 	, bCompilePendingCompletion(false)
 	, SystemStackViewModel(nullptr)
+	, EditorDocumentsViewModel(nullptr)
 	, SelectionViewModel(nullptr)
 	, ScriptScratchPadViewModel(nullptr)
 	, bPendingAssetMessagesChanged(true)
@@ -126,6 +128,9 @@ void FNiagaraSystemViewModel::Initialize(UNiagaraSystem& InSystem, FNiagaraSyste
 	SystemStackViewModel = NewObject<UNiagaraStackViewModel>(GetTransientPackage());
 	SystemStackViewModel->OnStructureChanged().AddSP(this, &FNiagaraSystemViewModel::StackViewModelStructureChanged);
 
+	EditorDocumentsViewModel = NewObject< UNiagaraSystemEditorDocumentsViewModel >(GetTransientPackage());
+	EditorDocumentsViewModel->Initialize(this->AsShared());
+
 	ScriptScratchPadViewModel = NewObject<UNiagaraScratchPadViewModel>(GetTransientPackage());
 	ScriptScratchPadViewModel->Initialize(this->AsShared());
 	ScriptScratchPadViewModel->OnScriptRenamed().AddSP(this, &FNiagaraSystemViewModel::ScratchPadScriptsChanged);
@@ -219,6 +224,12 @@ void FNiagaraSystemViewModel::Cleanup()
 		SystemStackViewModel = nullptr;
 	}
 
+	if (EditorDocumentsViewModel != nullptr)
+	{
+		EditorDocumentsViewModel->Finalize();
+		EditorDocumentsViewModel = nullptr;
+	}
+
 	if (SelectionViewModel != nullptr)
 	{
 		SelectionViewModel->OnEmitterHandleIdSelectionChanged().RemoveAll(this);
@@ -677,6 +688,10 @@ void FNiagaraSystemViewModel::AddReferencedObjects(FReferenceCollector& Collecto
 	{
 		Collector.AddReferencedObject(SystemStackViewModel);
 	}
+	if (EditorDocumentsViewModel != nullptr)
+	{
+		Collector.AddReferencedObject(EditorDocumentsViewModel);
+	}
 	if (SelectionViewModel != nullptr)
 	{
 		Collector.AddReferencedObject(SelectionViewModel);
diff --git a/Engine/Plugins/FX/Niagara/Source/NiagaraEditor/Private/ViewModels/Stack/NiagaraStackEntry.cpp b/Engine/Plugins/FX/Niagara/Source/NiagaraEditor/Private/ViewModels/Stack/NiagaraStackEntry.cpp
--- a/Engine/Plugins/FX/Niagara/Source/NiagaraEditor/Private/ViewModels/Stack/NiagaraStackEntry.cpp
+++ b/Engine/Plugins/FX/Niagara/Source/NiagaraEditor/Private/ViewModels/Stack/NiagaraStackEntry.cpp
@@ -567,6 +567,17 @@ bool UNiagaraStackEntry::HasIssuesOrAnyChildHasIssues() const
 	return GetCollectedIssueData().HasAnyIssues();
 }
 
+bool UNiagaraStackEntry::HasUsagesOrAnyChildHasUsages() const
+{
+	return GetCollectedUsageData().bHasReferencedParameterRead || GetCollectedUsageData().bHasReferencedParameterWrite;
+}
+
+void UNiagaraStackEntry::GetRecursiveUsages(bool& bRead, bool& bWrite) const
+{
+	bRead = GetCollectedUsageData().bHasReferencedParameterRead;
+	bWrite = GetCollectedUsageData().bHasReferencedParameterWrite;
+}
+
 int32 UNiagaraStackEntry::GetTotalNumberOfCustomNotes() const
 {
 	return GetCollectedIssueData().TotalNumberOfCustomNotes;
@@ -858,6 +869,7 @@ void UNiagaraStackEntry::InvalidateFilteredChildren()
 {
 	FilteredChildren.Empty();
 	CachedCollectedIssueData.Reset();
+	CachedCollectedUsageData.Reset();
 	bFilterChildrenPending = true;
 }
 
@@ -909,6 +921,42 @@ const UNiagaraStackEntry::FCollectedIssueData& UNiagaraStackEntry::GetCollectedI
 	return CachedCollectedIssueData.GetValue();
 }
 
+
+const UNiagaraStackEntry::FCollectedUsageData& UNiagaraStackEntry::GetCollectedUsageData() const
+{
+	if (CachedCollectedUsageData.IsSet() == false)
+	{
+		CachedCollectedUsageData = FCollectedUsageData();
+
+		TArray<UNiagaraStackEntry*> RefreshedChildren;
+		GetUnfilteredChildren(RefreshedChildren);
+
+		for (UNiagaraStackEntry* ChildStackEntry : RefreshedChildren)
+		{
+			if (ChildStackEntry->GetCollectedUsageData().bHasReferencedParameterRead)
+				CachedCollectedUsageData.GetValue().bHasReferencedParameterRead = true;
+			if (ChildStackEntry->GetCollectedUsageData().bHasReferencedParameterWrite)
+				CachedCollectedUsageData.GetValue().bHasReferencedParameterWrite = true;
+
+		}
+	}
+
+	return CachedCollectedUsageData.GetValue();
+}
+
+void UNiagaraStackEntry::InvalidateCollectedUsage()
+{
+	// Make sure to use the unfiltered children instead of filtered as the info we need is only found in certain, usually filtered children.
+	CachedCollectedUsageData.Reset();
+	TArray<UNiagaraStackEntry*> RefreshedFilteredChildren;
+	GetUnfilteredChildren(RefreshedFilteredChildren);
+
+	for (UNiagaraStackEntry* ChildStackEntry : RefreshedFilteredChildren)
+	{
+		ChildStackEntry->InvalidateCollectedUsage();
+	}
+}
+
 void UNiagaraStackEntry::BeginDestroy()
 {
 	ensureMsgf(HasAnyFlags(RF_ClassDefaultObject) || bIsFinalized, TEXT("Stack entry being destroyed but it was not finalized."));
diff --git a/Engine/Plugins/FX/Niagara/Source/NiagaraEditor/Private/ViewModels/Stack/NiagaraStackFunctionInput.cpp b/Engine/Plugins/FX/Niagara/Source/NiagaraEditor/Private/ViewModels/Stack/NiagaraStackFunctionInput.cpp
--- a/Engine/Plugins/FX/Niagara/Source/NiagaraEditor/Private/ViewModels/Stack/NiagaraStackFunctionInput.cpp
+++ b/Engine/Plugins/FX/Niagara/Source/NiagaraEditor/Private/ViewModels/Stack/NiagaraStackFunctionInput.cpp
@@ -50,6 +50,7 @@
 #include "ViewModels/Stack/NiagaraStackInputCategory.h"
 #include "ViewModels/Stack/NiagaraStackObject.h"
 #include "Widgets/Notifications/SNotificationList.h"
+#include "ViewModels/NiagaraParameterPanelViewModel.h"
 
 #define LOCTEXT_NAMESPACE "NiagaraStackViewModel"
 
@@ -2632,6 +2633,24 @@ TOptional<FGuid> UNiagaraStackFunctionInput::GetMetadataGuid() const
 	return InputMetaData.IsSet() ? InputMetaData->GetVariableGuid() : TOptional<FGuid>();
 }
 
+
+const UNiagaraStackFunctionInput::FCollectedUsageData& UNiagaraStackFunctionInput::GetCollectedUsageData() const
+{
+	if (CachedCollectedUsageData.IsSet() == false)
+	{
+		CachedCollectedUsageData = FCollectedUsageData();
+		TSharedRef<FNiagaraSystemViewModel> SystemVM = GetSystemViewModel();
+		INiagaraParameterPanelViewModel* ParamVM = SystemVM->GetParameterPanelViewModel();
+		if (ParamVM)
+		{
+			FNiagaraVariableBase Var(InputType, AliasedInputParameterHandle.GetParameterHandleString());
+			CachedCollectedUsageData.GetValue().bHasReferencedParameterRead = ParamVM->IsVariableSelected(Var);
+		}
+	}
+
+	return CachedCollectedUsageData.GetValue();
+}
+
 void UNiagaraStackFunctionInput::OnGraphChanged(const struct FEdGraphEditAction& InAction)
 {
 	if (bUpdatingGraphDirectly == false)
diff --git a/Engine/Plugins/FX/Niagara/Source/NiagaraEditor/Private/ViewModels/Stack/NiagaraStackModuleItem.cpp b/Engine/Plugins/FX/Niagara/Source/NiagaraEditor/Private/ViewModels/Stack/NiagaraStackModuleItem.cpp
--- a/Engine/Plugins/FX/Niagara/Source/NiagaraEditor/Private/ViewModels/Stack/NiagaraStackModuleItem.cpp
+++ b/Engine/Plugins/FX/Niagara/Source/NiagaraEditor/Private/ViewModels/Stack/NiagaraStackModuleItem.cpp
@@ -92,6 +92,8 @@ void UNiagaraStackModuleItem::Initialize(FRequiredEntryData InRequiredEntryData,
 	});
 }
 
+
+
 FText UNiagaraStackModuleItem::GetDisplayName() const
 {
 	if (DisplayNameCache.IsSet() == false)
@@ -214,6 +216,42 @@ void UNiagaraStackModuleItem::RefreshChildrenInternal(const TArray<UNiagaraStack
 	RefreshIssues(NewIssues);
 }
 
+
+const UNiagaraStackModuleItem::FCollectedUsageData& UNiagaraStackModuleItem::GetCollectedUsageData() const
+{
+	if (CachedCollectedUsageData.IsSet() == false)
+	{
+		CachedCollectedUsageData = FCollectedUsageData();
+		if (!LinkedInputCollection.IsNull())
+		{
+			if (LinkedInputCollection->GetCollectedUsageData().bHasReferencedParameterRead)
+				CachedCollectedUsageData.GetValue().bHasReferencedParameterRead = true;
+			if (LinkedInputCollection->GetCollectedUsageData().bHasReferencedParameterWrite)
+				CachedCollectedUsageData.GetValue().bHasReferencedParameterWrite = true;
+		}
+
+
+		if (!InputCollection.IsNull())
+		{
+			if (InputCollection->GetCollectedUsageData().bHasReferencedParameterRead)
+				CachedCollectedUsageData.GetValue().bHasReferencedParameterRead = true;
+			if (InputCollection->GetCollectedUsageData().bHasReferencedParameterWrite)
+				CachedCollectedUsageData.GetValue().bHasReferencedParameterWrite = true;
+		}
+
+		if (!OutputCollection.IsNull())
+		{
+			if (OutputCollection->GetCollectedUsageData().bHasReferencedParameterRead)
+				CachedCollectedUsageData.GetValue().bHasReferencedParameterRead = true;
+			if (OutputCollection->GetCollectedUsageData().bHasReferencedParameterWrite)
+				CachedCollectedUsageData.GetValue().bHasReferencedParameterWrite = true;
+		}
+		
+	}
+
+	return CachedCollectedUsageData.GetValue();
+}
+
 TOptional<UNiagaraStackEntry::FDropRequestResponse> UNiagaraStackModuleItem::CanDropInternal(const FDropRequest& DropRequest)
 {
 	if ((DropRequest.DropOptions != UNiagaraStackEntry::EDropOptions::Overview || DropRequest.DropZone == EItemDropZone::OntoItem) &&
diff --git a/Engine/Plugins/FX/Niagara/Source/NiagaraEditor/Private/ViewModels/Stack/NiagaraStackModuleItemOutput.cpp b/Engine/Plugins/FX/Niagara/Source/NiagaraEditor/Private/ViewModels/Stack/NiagaraStackModuleItemOutput.cpp
--- a/Engine/Plugins/FX/Niagara/Source/NiagaraEditor/Private/ViewModels/Stack/NiagaraStackModuleItemOutput.cpp
+++ b/Engine/Plugins/FX/Niagara/Source/NiagaraEditor/Private/ViewModels/Stack/NiagaraStackModuleItemOutput.cpp
@@ -6,6 +6,9 @@
 #include "NiagaraGraph.h"
 #include "NiagaraNodeFunctionCall.h"
 #include "NiagaraScriptSource.h"
+#include "ViewModels/NiagaraParameterPanelViewModel.h"
+#include "ViewModels/NiagaraSystemViewModel.h"
+#include "ViewModels/Stack/NiagaraStackModuleItemOutputCollection.h"
 
 
 #define LOCTEXT_NAMESPACE "NiagaraStackViewModel"
@@ -32,6 +35,27 @@ FText UNiagaraStackModuleItemOutput::GetDisplayName() const
 	return DisplayName;
 }
 
+const UNiagaraStackModuleItemOutput::FCollectedUsageData& UNiagaraStackModuleItemOutput::GetCollectedUsageData() const
+{
+	if (CachedCollectedUsageData.IsSet() == false)
+	{
+		CachedCollectedUsageData = FCollectedUsageData();
+		TSharedRef<FNiagaraSystemViewModel> SystemVM = GetSystemViewModel();
+		INiagaraParameterPanelViewModel* ParamVM = SystemVM->GetParameterPanelViewModel();
+		if (ParamVM)
+		{
+			FNiagaraVariableBase Var(OutputType, OutputParameterHandle.GetParameterHandleString());
+			bool bHandled = ParamVM->IsVariableSelected(Var);
+			if (GetOuter() && GetOuter()->IsA<UNiagaraStackModuleItemOutputCollection>())
+				CachedCollectedUsageData.GetValue().bHasReferencedParameterWrite = bHandled;
+			else
+				CachedCollectedUsageData.GetValue().bHasReferencedParameterRead = bHandled;
+		}
+	}
+
+	return CachedCollectedUsageData.GetValue();
+}
+
 FText UNiagaraStackModuleItemOutput::GetTooltipText() const
 {
 	FNiagaraVariable ValueVariable(OutputType, OutputParameterHandle.GetParameterHandleString());
diff --git a/Engine/Plugins/FX/Niagara/Source/NiagaraEditor/Private/ViewModels/Stack/NiagaraStackScriptItemGroup.cpp b/Engine/Plugins/FX/Niagara/Source/NiagaraEditor/Private/ViewModels/Stack/NiagaraStackScriptItemGroup.cpp
--- a/Engine/Plugins/FX/Niagara/Source/NiagaraEditor/Private/ViewModels/Stack/NiagaraStackScriptItemGroup.cpp
+++ b/Engine/Plugins/FX/Niagara/Source/NiagaraEditor/Private/ViewModels/Stack/NiagaraStackScriptItemGroup.cpp
@@ -792,6 +792,10 @@ TOptional<UNiagaraStackEntry::FDropRequestResponse> UNiagaraStackScriptItemGroup
 	{
 		return CanDropParameterOnTarget(TargetEntry, DropRequest);
 	}
+	else if (DropRequest.DragDropOperation->IsOfType<FNiagaraScriptDragOperation>())
+	{
+		return CanDropScriptsOnTarget(TargetEntry, DropRequest);
+	}
 	return TOptional<FDropRequestResponse>();
 }
 
@@ -978,6 +982,53 @@ TOptional<UNiagaraStackEntry::FDropRequestResponse> UNiagaraStackScriptItemGroup
 	return FDropRequestResponse(TargetDropZone, DropMessage);
 }
 
+
+TOptional<UNiagaraStackEntry::FDropRequestResponse> UNiagaraStackScriptItemGroup::CanDropScriptsOnTarget(const UNiagaraStackEntry& TargetEntry, const FDropRequest& DropRequest)
+{
+	TSharedRef<const FNiagaraScriptDragOperation> DragDropOp = StaticCastSharedRef<const FNiagaraScriptDragOperation>(DropRequest.DragDropOperation);
+	const UEnum* NiagaraScriptUsageEnum = FindObjectChecked<UEnum>(ANY_PACKAGE, TEXT("ENiagaraScriptUsage"), true);
+
+	if (&TargetEntry != this && TargetEntry.IsA<UNiagaraStackModuleItem>() == false)
+	{
+		// Only handle drops onto this script group, or child drop requests from module items.
+		return TOptional<FDropRequestResponse>();
+	}
+	if (DropRequest.DropOptions != UNiagaraStackEntry::EDropOptions::Overview)
+	{
+		// Only allow dropping in the overview stacks.
+		return FDropRequestResponse(TOptional<EItemDropZone>(), LOCTEXT("AssetCantDropOnStack", "Scripts can only be dropped into the overview."));
+	}
+	UNiagaraScript* Script = DragDropOp->Script.Get();
+	if (Script)
+	{
+		ENiagaraScriptUsage ScriptscriptUsage = Script->Usage;
+		if (ScriptscriptUsage != ENiagaraScriptUsage::Module)
+		{
+			return FDropRequestResponse(TOptional<EItemDropZone>(), FText::Format(LOCTEXT("CantDropNonModuleAssetFormat", "Can not drop asset {0} here because it is not a module script."), DragDropOp->FriendlyName));
+		}
+
+		int32 BitfieldValue = Script->GetScriptData(DragDropOp->Version)->ModuleUsageBitmask;;
+		TArray<ENiagaraScriptUsage> SupportedUsages = UNiagaraScript::GetSupportedUsageContextsForBitmask(BitfieldValue);
+		if (SupportedUsages.Contains(GetScriptUsage()) == false)
+		{
+			return FDropRequestResponse(TOptional<EItemDropZone>(), FText::Format(LOCTEXT("CantDropAssetByUsageFormat", "Can not drop asset {0} in this part of the stack\nbecause it's not valid for this usage context."), DragDropOp->FriendlyName));
+		}
+	}
+
+	TOptional<EItemDropZone> TargetDropZone = GetTargetDropZoneForTargetEntry(this, TargetEntry, DropRequest.DropZone);
+	if (TargetDropZone.IsSet() == false)
+	{
+		return TOptional<FDropRequestResponse>();
+	}
+
+	FText DropMessage;
+	if (Script)
+	{
+		DropMessage = LOCTEXT("DropAsset", "Insert a module for this asset here.");
+	}
+	return FDropRequestResponse(TargetDropZone, DropMessage);
+}
+
 TOptional<UNiagaraStackEntry::FDropRequestResponse> UNiagaraStackScriptItemGroup::CanDropParameterOnTarget(const UNiagaraStackEntry& TargetEntry, const FDropRequest& DropRequest)
 {
 	TSharedRef<const FNiagaraParameterDragOperation> ParameterDragDropOp = StaticCastSharedRef<const FNiagaraParameterDragOperation>(DropRequest.DragDropOperation);
@@ -1030,6 +1081,10 @@ TOptional<UNiagaraStackEntry::FDropRequestResponse> UNiagaraStackScriptItemGroup
 	{
 		return DropParameterOnTarget(TargetEntry, DropRequest);
 	}
+	else if (DropRequest.DragDropOperation->IsOfType<FNiagaraScriptDragOperation>())
+	{
+		return DropScriptsOnTarget(TargetEntry, DropRequest);
+	}
 	return TOptional<FDropRequestResponse>();
 }
 
@@ -1111,6 +1166,22 @@ TOptional<UNiagaraStackEntry::FDropRequestResponse> UNiagaraStackScriptItemGroup
 	return FDropRequestResponse(DropRequest.DropZone);
 }
 
+
+TOptional<UNiagaraStackEntry::FDropRequestResponse> UNiagaraStackScriptItemGroup::DropScriptsOnTarget(const UNiagaraStackEntry& TargetEntry, const FDropRequest& DropRequest)
+{
+	TSharedRef<const FNiagaraScriptDragOperation> DragDropOp = StaticCastSharedRef<const FNiagaraScriptDragOperation>(DropRequest.DragDropOperation);
+
+	int32 TargetIndex = GetTargetIndexForTargetEntry(this, TargetEntry, DropRequest.DropZone);
+
+	FScopedTransaction ScopedTransaction(LOCTEXT("DragAndDropScript", "Insert modules for scripts"));
+	if (UNiagaraScript* Script = DragDropOp->Script.Get())
+	{
+		TSharedRef<FScriptGroupAddAction> AddAction = FScriptGroupAddAction::CreateModuleActionFromScratchPadScript(Script);
+		AddUtilities->ExecuteAddAction(AddAction, TargetIndex);
+	}
+	return FDropRequestResponse(DropRequest.DropZone);
+}
+
 TOptional<UNiagaraStackEntry::FDropRequestResponse> UNiagaraStackScriptItemGroup::DropParameterOnTarget(const UNiagaraStackEntry& TargetEntry, const FDropRequest& DropRequest)
 {
 	TSharedRef<const FNiagaraParameterDragOperation> ParameterDragDropOp = StaticCastSharedRef<const FNiagaraParameterDragOperation>(DropRequest.DragDropOperation);
diff --git a/Engine/Plugins/FX/Niagara/Source/NiagaraEditor/Private/ViewModels/Stack/NiagaraStackViewModel.cpp b/Engine/Plugins/FX/Niagara/Source/NiagaraEditor/Private/ViewModels/Stack/NiagaraStackViewModel.cpp
--- a/Engine/Plugins/FX/Niagara/Source/NiagaraEditor/Private/ViewModels/Stack/NiagaraStackViewModel.cpp
+++ b/Engine/Plugins/FX/Niagara/Source/NiagaraEditor/Private/ViewModels/Stack/NiagaraStackViewModel.cpp
@@ -235,6 +235,13 @@ bool UNiagaraStackViewModel::HasIssues() const
 	return bHasIssues;
 }
 
+void UNiagaraStackViewModel::InvalidateCachedParameterUsage()
+{
+	if (RootEntry)
+		RootEntry->InvalidateCollectedUsage();
+
+}
+
 void UNiagaraStackViewModel::Finalize()
 {
 	Reset();
diff --git a/Engine/Plugins/FX/Niagara/Source/NiagaraEditor/Private/Widgets/NiagaraScratchScriptEditor.cpp b/Engine/Plugins/FX/Niagara/Source/NiagaraEditor/Private/Widgets/NiagaraScratchScriptEditor.cpp
new file mode 100644
--- /dev/null
+++ b/Engine/Plugins/FX/Niagara/Source/NiagaraEditor/Private/Widgets/NiagaraScratchScriptEditor.cpp
@@ -0,0 +1,251 @@
+// Copyright Epic Games, Inc. All Rights Reserved.
+
+//-TODO: Some platforms require an offline compile to get instruction counts, this needs to be handled
+
+#include "NiagaraScratchScriptEditor.h"
+#include "ViewModels/NiagaraSystemViewModel.h"
+#include "ViewModels/NiagaraScratchPadViewModel.h"
+#include "ViewModels/NiagaraScratchPadScriptViewModel.h"
+#include "ViewModels/NiagaraSystemEditorDocumentsViewModel.h"
+#include "NiagaraEditorStyle.h"
+#include "Toolkits/NiagaraSystemToolkit.h"
+#include "Widgets/Input/SButton.h"
+
+
+#define LOCTEXT_NAMESPACE "NiagaraScratchPadScriptEditor"
+
+void SNiagaraScratchPadScriptEditor::Construct(const FArguments& InArgs, TSharedRef<FNiagaraScratchPadScriptViewModel> InScriptViewModel)
+{
+	ScriptViewModel = InScriptViewModel;
+
+	ChildSlot
+		[
+			SNew(SVerticalBox)
+			+ SVerticalBox::Slot()
+			.AutoHeight()
+			[
+				SNew(SHorizontalBox)
+				+ SHorizontalBox::Slot()
+				.AutoWidth()
+				.Padding(2.0f, 3.0f, 5.0f, 3.0f)
+				[
+					SNew(SButton)
+					.OnClicked(this, &SNiagaraScratchPadScriptEditor::OnApplyButtonClicked)
+					.ToolTipText(LOCTEXT("ApplyButtonToolTip", "Apply the current changes to this script.  This will update the selection stack UI and compile neccessary scripts."))
+					.IsEnabled(this, &SNiagaraScratchPadScriptEditor::GetApplyButtonIsEnabled)
+					.ContentPadding(FMargin(3.0f, 0.0f))
+					.Content()
+					[
+						SNew(SHorizontalBox)
+						+ SHorizontalBox::Slot()
+						.AutoWidth()
+						.Padding(2.0f, 2.0f, 2.0f, 2.0f)
+						[
+							SNew(SImage)
+							.Image(FAppStyle::Get().GetBrush("AssetEditor.Apply"))
+						]
+						+ SHorizontalBox::Slot()
+						.AutoWidth()
+						.VAlign(VAlign_Center)
+						.Padding(2.0f, 2.0f, 2.0f, 3.0f)
+						[
+							SNew(STextBlock)
+							.TextStyle(FNiagaraEditorStyle::Get(), "NiagaraEditor.ScratchPad.EditorHeaderText")
+							.Text(LOCTEXT("ApplyButtonLabel", "Apply"))
+						]
+					]
+				]
+				+ SHorizontalBox::Slot()
+				.VAlign(VAlign_Center)
+				.Padding(5, 0, 2, 0)
+				.AutoWidth()
+				[
+					SNew(STextBlock)
+					.TextStyle(FNiagaraEditorStyle::Get(), "NiagaraEditor.ScratchPad.EditorHeaderText")
+					.Text(this, &SNiagaraScratchPadScriptEditor::GetNameText)
+					.ToolTipText(ScriptViewModel.ToSharedRef(), &FNiagaraScratchPadScriptViewModel::GetToolTip)
+				]
+				+ SHorizontalBox::Slot()
+				.VAlign(VAlign_Center)
+				.AutoWidth()
+				.Padding(0, 0, 2, 0)
+				[
+					SNew(STextBlock)
+					.Visibility(this, &SNiagaraScratchPadScriptEditor::GetUnappliedChangesVisibility)
+					.TextStyle(FNiagaraEditorStyle::Get(), "NiagaraEditor.ScratchPad.EditorHeaderText")
+					.Text(FText::FromString(TEXT("*")))
+					.ToolTipText(ScriptViewModel.ToSharedRef(), &FNiagaraScratchPadScriptViewModel::GetToolTip)
+				]
+			]
+			+ SVerticalBox::Slot()
+			[
+				SAssignNew(Graph, SNiagaraScriptGraph, ScriptViewModel->GetGraphViewModel())
+				.ZoomToFitOnLoad(true)
+				.ShowHeader(false)
+			]
+		];
+
+	if (ScriptViewModel)
+	{
+		NodeIDHandle = ScriptViewModel->OnNodeIDFocusRequested().AddLambda(
+			[this](FNiagaraScriptIDAndGraphFocusInfo* FocusInfo)
+			{
+				if (Graph.IsValid() && FocusInfo != nullptr)
+				{
+					Graph->FocusGraphElement(FocusInfo->GetScriptGraphFocusInfo().Get());
+				}
+			}
+		);
+
+		PinIDHandle = ScriptViewModel->OnPinIDFocusRequested().AddLambda(
+			[this](FNiagaraScriptIDAndGraphFocusInfo* FocusInfo)
+			{
+				if (Graph.IsValid() && FocusInfo != nullptr)
+				{
+					Graph->FocusGraphElement(FocusInfo->GetScriptGraphFocusInfo().Get());
+				}
+			}
+		);
+	}
+}
+
+void FGraphTabHistory::EvokeHistory(TSharedPtr<FTabInfo> InTabInfo, bool bPrevTabMatches)
+{
+	FWorkflowTabSpawnInfo SpawnInfo;
+	SpawnInfo.Payload = Payload;
+	SpawnInfo.TabInfo = InTabInfo;
+
+	if (bPrevTabMatches)
+	{
+		TSharedPtr<SDockTab> DockTab = InTabInfo->GetTab().Pin();
+		GraphEditor = StaticCastSharedRef<SNiagaraScratchPadScriptEditor>(DockTab->GetContent());
+	}
+	else
+	{
+		TSharedRef< SNiagaraScratchPadScriptEditor > GraphEditorRef = StaticCastSharedRef< SNiagaraScratchPadScriptEditor >(FactoryPtr.Pin()->CreateTabBody(SpawnInfo));
+		GraphEditor = GraphEditorRef;
+		TSharedPtr<SDockTab> TabPtr = InTabInfo->GetTab().Pin();
+		FactoryPtr.Pin()->UpdateTab(TabPtr, SpawnInfo, GraphEditorRef);
+	}
+}
+
+void FGraphTabHistory::SaveHistory()
+{
+	if (IsHistoryValid())
+	{
+		check(GraphEditor.IsValid());
+		TSharedPtr<SNiagaraScratchPadScriptEditor> Editor = GraphEditor.Pin();
+		if (Editor->GetGraphEditor())
+		{
+			Editor->GetGraphEditor()->GetViewLocation(SavedLocation, SavedZoomAmount);
+			Editor->GetGraphEditor()->GetViewBookmark(SavedBookmarkId);
+		}
+	}
+}
+
+void FGraphTabHistory::RestoreHistory()
+{
+	if (IsHistoryValid())
+	{
+		check(GraphEditor.IsValid());
+		TSharedPtr<SNiagaraScratchPadScriptEditor> Editor = GraphEditor.Pin();
+		if (Editor->GetGraphEditor())
+			Editor->GetGraphEditor()->SetViewLocation(SavedLocation, SavedZoomAmount, SavedBookmarkId);
+	}
+}
+void FGraphEditorSummoner::OnTabActivated(TSharedPtr<SDockTab> Tab) const
+{
+	TSharedRef<SNiagaraScratchPadScriptEditor> GraphEditor = StaticCastSharedRef<SNiagaraScratchPadScriptEditor>(Tab->GetContent());
+	EditorPtr.Pin()->GetSystemViewModel()->GetDocumentViewModel()->SetActiveDocumentTab(Tab);
+}
+
+void FGraphEditorSummoner::OnTabBackgrounded(TSharedPtr<SDockTab> Tab) const
+{
+	TSharedRef<SNiagaraScratchPadScriptEditor> GraphEditor = StaticCastSharedRef<SNiagaraScratchPadScriptEditor>(Tab->GetContent());
+	//EditorPtr.Pin()->OnGraphEditorBackgrounded(GraphEditor);
+}
+
+void FGraphEditorSummoner::OnTabRefreshed(TSharedPtr<SDockTab> Tab) const
+{
+	TSharedRef<SNiagaraScratchPadScriptEditor> GraphEditor = StaticCastSharedRef<SNiagaraScratchPadScriptEditor>(Tab->GetContent());
+	GraphEditor->GetGraphEditor()->NotifyGraphChanged();
+}
+
+void FGraphEditorSummoner::SaveState(TSharedPtr<SDockTab> Tab, TSharedPtr<FTabPayload> Payload) const
+{
+	// Do nothing here for now
+	/*
+	TSharedRef<SNiagaraScratchPadScriptEditor> GraphEditor = StaticCastSharedRef<SNiagaraScratchPadScriptEditor>(Tab->GetContent());
+
+	FVector2D ViewLocation;
+	float ZoomAmount;
+	GraphEditor->GetGraphEditor()->GetViewLocation(ViewLocation, ZoomAmount);
+
+	UEdGraph* Graph = Payload->IsValid() ? FTabPayload_UObject::CastChecked<UEdGraph>(Payload) : nullptr;
+	*/
+}
+
+FGraphEditorSummoner::FGraphEditorSummoner(TSharedPtr<class FNiagaraSystemToolkit> InToolkit, FOnCreateGraphEditorWidget CreateGraphEditorWidgetCallback) : FDocumentTabFactoryForObjects<UEdGraph>(TEXT("SystemToolkitGraph"), InToolkit->AsShared())
+, EditorPtr(InToolkit)
+, OnCreateGraphEditorWidget(CreateGraphEditorWidgetCallback)
+{
+
+}
+
+TSharedRef<SWidget> FGraphEditorSummoner::CreateTabBodyForObject(const FWorkflowTabSpawnInfo& Info, UEdGraph* DocumentID) const
+{
+	check(Info.TabInfo.IsValid());
+
+	// We need to register the tab being closed as we need to clear out the active selection logic when that happens.
+	// There may be better places to put this, but this works well in practice.
+	SDockTab::FOnTabClosedCallback TabClosedCallback = SDockTab::FOnTabClosedCallback::CreateLambda([=](TSharedRef<SDockTab> DockTab)
+		{
+			TSharedPtr<FNiagaraSystemToolkit> Toolkit = EditorPtr.Pin();
+			if (Toolkit.IsValid())
+				Toolkit->GetSystemViewModel()->GetDocumentViewModel()->SetActiveDocumentTab(nullptr);
+		});
+	TSharedPtr<SDockTab> TabPtr = Info.TabInfo->GetTab().Pin();
+	if (TabPtr.IsValid())
+		TabPtr->SetOnTabClosed(TabClosedCallback);
+
+	// Create the widget!
+	return OnCreateGraphEditorWidget.Execute(Info.TabInfo.ToSharedRef(), DocumentID);
+}
+
+const FSlateBrush* FGraphEditorSummoner::GetTabIconForObject(const FWorkflowTabSpawnInfo& Info, UEdGraph* DocumentID) const
+{
+	return FSlateIcon(FNiagaraEditorStyle::Get().GetStyleSetName(), "Tab.ScratchPad").GetSmallIcon();
+}
+
+TSharedRef<FGenericTabHistory> FGraphEditorSummoner::CreateTabHistoryNode(TSharedPtr<FTabPayload> Payload)
+{
+	return MakeShareable(new FGraphTabHistory(SharedThis(this), Payload));
+}
+
+TAttribute<FText> FGraphEditorSummoner::ConstructTabNameForObject(UEdGraph* DocumentID) const 
+{
+	UNiagaraScript* Script = DocumentID->GetTypedOuter<UNiagaraScript>();
+	TSharedPtr<FNiagaraSystemViewModel> SysViewModel = EditorPtr.Pin()->GetSystemViewModel();
+
+	if (Script)
+	{
+		if (SysViewModel.IsValid())
+		{
+			TSharedPtr<FNiagaraScratchPadScriptViewModel> ScratchScriptViewModel = UNiagaraSystemEditorDocumentsViewModel::GetScratchPadViewModelFromGraph(SysViewModel.Get(), DocumentID);
+			if (ScratchScriptViewModel)
+			{
+				return ScratchScriptViewModel->GetDisplayName(); /* FText::FromString(ScratchScriptViewModel->GetOriginalScript()->GetName());*/
+			}
+		}
+	}
+	return FText();
+}
+
+TSharedRef<SDockTab> FGraphEditorSummoner::OnSpawnTab(const FSpawnTabArgs& SpawnArgs, TWeakPtr<FTabManager> WeakTabManager) const
+{
+	TSharedRef<SDockTab> TabRef = FDocumentTabFactoryForObjects<UEdGraph>::OnSpawnTab(SpawnArgs, WeakTabManager);
+	return TabRef;
+}
+
+
+#undef LOCTEXT_NAMESPACE
diff --git a/Engine/Plugins/FX/Niagara/Source/NiagaraEditor/Private/Widgets/NiagaraScratchScriptEditor.h b/Engine/Plugins/FX/Niagara/Source/NiagaraEditor/Private/Widgets/NiagaraScratchScriptEditor.h
new file mode 100644
--- /dev/null
+++ b/Engine/Plugins/FX/Niagara/Source/NiagaraEditor/Private/Widgets/NiagaraScratchScriptEditor.h
@@ -0,0 +1,145 @@
+// Copyright Epic Games, Inc. All Rights Reserved.
+
+#pragma once
+
+#include "NiagaraGraph.h"
+#include "Widgets/SCompoundWidget.h"
+#include "ViewModels/NiagaraSystemSelectionViewModel.h"
+#include "ViewModels/NiagaraScratchPadScriptViewModel.h"
+#include "ViewModels/NiagaraScratchPadViewModel.h"
+#include "ViewModels/NiagaraScriptGraphViewModel.h"
+#include "WorkflowOrientedApp/WorkflowTabManager.h"
+#include "NiagaraScriptSource.h"
+#include "Widgets/SNiagaraScriptGraph.h"
+#include "Widgets/Docking/SDockTab.h"
+#include "WorkflowOrientedApp/WorkflowUObjectDocuments.h"
+
+class SNiagaraScratchPadScriptEditor : public SCompoundWidget
+{
+	SLATE_BEGIN_ARGS(SNiagaraScratchPadScriptEditor) {}
+	SLATE_END_ARGS()
+
+	void Construct(const FArguments& InArgs, TSharedRef<FNiagaraScratchPadScriptViewModel> InScriptViewModel);
+
+	~SNiagaraScratchPadScriptEditor()
+	{
+		if (ScriptViewModel)
+		{
+			ScriptViewModel->OnNodeIDFocusRequested().Remove(NodeIDHandle);
+			ScriptViewModel->OnPinIDFocusRequested().Remove(PinIDHandle);
+		}
+	}
+
+	TSharedPtr<SGraphEditor> GetGraphEditor()
+	{
+		if (Graph)
+			return Graph->GetGraphEditor();
+		else
+			return nullptr;
+	}
+
+	TSharedPtr<FNiagaraScratchPadScriptViewModel> GetViewModel()
+	{
+		return ScriptViewModel;
+	}
+
+
+private:
+	FText GetNameText() const
+	{
+		return ScriptViewModel->GetDisplayName();
+	}
+
+	EVisibility GetUnappliedChangesVisibility() const
+	{
+		return ScriptViewModel->HasUnappliedChanges() ? EVisibility::Visible : EVisibility::Collapsed;
+	}
+
+	FReply OnApplyButtonClicked()
+	{
+		ScriptViewModel->ApplyChanges();
+
+		return FReply::Handled();
+	}
+
+	bool GetApplyButtonIsEnabled() const
+	{
+		return ScriptViewModel->HasUnappliedChanges();
+	}
+
+private:
+	TSharedPtr<FNiagaraScratchPadScriptViewModel> ScriptViewModel;
+	TSharedPtr<SNiagaraScriptGraph> Graph;
+
+	FDelegateHandle NodeIDHandle;
+	FDelegateHandle PinIDHandle;
+};
+
+
+/////////////////////////////////////////////////////
+// FGraphTabHistory
+
+struct FGraphTabHistory : public FGenericTabHistory
+{
+public:
+	/**
+	 * @param InFactory		The factory used to regenerate the content
+	 * @param InPayload		The payload object used to regenerate the content
+	 */
+	FGraphTabHistory(TSharedPtr<FDocumentTabFactory> InFactory, TSharedPtr<FTabPayload> InPayload)
+		: FGenericTabHistory(InFactory, InPayload)
+		, SavedLocation(FVector2D::ZeroVector)
+		, SavedZoomAmount(INDEX_NONE)
+	{
+
+	}
+
+	virtual void EvokeHistory(TSharedPtr<FTabInfo> InTabInfo, bool bPrevTabMatches) override;
+
+	virtual void SaveHistory() override;
+
+	virtual void RestoreHistory() override;
+
+private:
+	/** The graph editor represented by this history node. While this node is inactive, the graph editor is invalid */
+	TWeakPtr< class SNiagaraScratchPadScriptEditor > GraphEditor;
+	/** Saved location the graph editor was at when this history node was last visited */
+	FVector2D SavedLocation;
+	/** Saved zoom the graph editor was at when this history node was last visited */
+	float SavedZoomAmount;
+	/** Saved bookmark ID the graph editor was at when this history node was last visited */
+	FGuid SavedBookmarkId;
+};
+
+
+struct FGraphEditorSummoner : public FDocumentTabFactoryForObjects<UEdGraph>
+{
+public:
+	DECLARE_DELEGATE_RetVal_TwoParams(TSharedRef<SNiagaraScratchPadScriptEditor>, FOnCreateGraphEditorWidget, TSharedRef<FTabInfo>, UEdGraph*);
+public:
+	FGraphEditorSummoner(TSharedPtr<class FNiagaraSystemToolkit> InToolkit, FOnCreateGraphEditorWidget CreateGraphEditorWidgetCallback);
+
+	virtual void OnTabActivated(TSharedPtr<SDockTab> Tab) const override;
+
+	virtual void OnTabBackgrounded(TSharedPtr<SDockTab> Tab) const override;
+
+	virtual void OnTabRefreshed(TSharedPtr<SDockTab> Tab) const override;
+
+
+	virtual void SaveState(TSharedPtr<SDockTab> Tab, TSharedPtr<FTabPayload> Payload) const override;
+
+	virtual TSharedRef<SDockTab> OnSpawnTab(const FSpawnTabArgs& SpawnArgs, TWeakPtr<FTabManager> WeakTabManager) const override;
+
+protected:
+	virtual TAttribute<FText> ConstructTabNameForObject(UEdGraph* DocumentID) const override;
+
+	virtual TSharedRef<SWidget> CreateTabBodyForObject(const FWorkflowTabSpawnInfo& Info, UEdGraph* DocumentID) const override;
+
+	virtual const FSlateBrush* GetTabIconForObject(const FWorkflowTabSpawnInfo& Info, UEdGraph* DocumentID) const override;
+
+	virtual TSharedRef<FGenericTabHistory> CreateTabHistoryNode(TSharedPtr<FTabPayload> Payload) override;
+
+protected:
+	TWeakPtr<FNiagaraSystemToolkit> EditorPtr;
+	FOnCreateGraphEditorWidget OnCreateGraphEditorWidget;
+};
diff --git a/Engine/Plugins/FX/Niagara/Source/NiagaraEditor/Private/Widgets/SNiagaraParameterPanel.cpp b/Engine/Plugins/FX/Niagara/Source/NiagaraEditor/Private/Widgets/SNiagaraParameterPanel.cpp
--- a/Engine/Plugins/FX/Niagara/Source/NiagaraEditor/Private/Widgets/SNiagaraParameterPanel.cpp
+++ b/Engine/Plugins/FX/Niagara/Source/NiagaraEditor/Private/Widgets/SNiagaraParameterPanel.cpp
@@ -24,6 +24,8 @@
 #include "Widgets/SNiagaraParameterName.h"
 #include "Widgets/SItemSelector.h"
 #include "Styling/AppStyle.h"
+#include "Widgets/Layout/SWrapBox.h"
+#include "Widgets/Input/SCheckBox.h"
 
 #define LOCTEXT_NAMESPACE "NiagaraParameterPanel"
 
@@ -85,7 +87,7 @@ void SNiagaraParameterPanel::Construct(const FArguments& InArgs, const TSharedPt
 	.OnItemSelected(this, &SNiagaraParameterPanel::OnParameterItemSelected)
 	.OnItemsDragged(this, &SNiagaraParameterPanel::OnParameterItemsDragged)
 	.OnItemActivated(this, &SNiagaraParameterPanel::OnParameterItemActived)
-	.AllowMultiselect(true)
+	.AllowMultiselect(ParameterPanelViewModel->GetAllowMultiSelect())
 	.ClearSelectionOnClick(true)
 	.CategoryRowStyle(FNiagaraEditorStyle::Get(), "NiagaraEditor.Parameters.TableRow")
 	.OnGetCategoryBackgroundImage(this, &SNiagaraParameterPanel::GetCategoryBackgroundImage)
@@ -97,9 +99,18 @@ void SNiagaraParameterPanel::Construct(const FArguments& InArgs, const TSharedPt
 
 	// Finalize the widget
 	ChildSlot
-	[
-		SNew(SBox)
-		.MinDesiredWidth(300)
+		[
+		SNew(SVerticalBox)
+		+ SVerticalBox::Slot()
+		.AutoHeight()
+		.Padding(4, 2)
+		[
+			SAssignNew(SectionSelectorBox, SWrapBox)
+			.UseAllottedSize(true)
+			.InnerSlotPadding(FVector2D(4, 4))
+		]
+		+ SVerticalBox::Slot()
+		.Padding(0, 2)
 		[
 			// Drop target
 			SNew(SDropTarget)
@@ -123,8 +134,48 @@ void SNiagaraParameterPanel::Construct(const FArguments& InArgs, const TSharedPt
 	ToolkitCommands->MapAction(FGenericCommands::Get().Copy,
 		FExecuteAction::CreateSP(this, &SNiagaraParameterPanel::CopyParameterReference),
 		FCanExecuteAction::CreateSP(this, &SNiagaraParameterPanel::CanCopyParameterReference));
+
+	ConstructSectionButtons();
+}
+
+
+void SNiagaraParameterPanel::ConstructSectionButtons()
+{
+	SectionSelectorBox->ClearChildren();
+
+	SectionSelectorBox->SetVisibility(ParameterPanelViewModel->GetShowSections() ? EVisibility::Visible : EVisibility::Collapsed);
+	for (FText Section : ParameterPanelViewModel->GetSections())
+	{
+		SectionSelectorBox->AddSlot()
+			[
+				SNew(SCheckBox)
+				.Style(FAppStyle::Get(), "DetailsView.SectionButton")
+				.OnCheckStateChanged(this, &SNiagaraParameterPanel::OnSectionChecked, Section)
+				.IsChecked(this, &SNiagaraParameterPanel::GetSectionCheckState, Section)
+				.ToolTipText(ParameterPanelViewModel->GetTooltipForSection(Section))
+				[
+					SNew(STextBlock)
+					.TextStyle(FAppStyle::Get(), "SmallText")
+					.Text(Section)
+				]
+			];
+	}
 }
 
+ECheckBoxState SNiagaraParameterPanel::GetSectionCheckState(FText Section) const
+{
+	return Section.IdenticalTo(ParameterPanelViewModel->GetActiveSection()) ? ECheckBoxState::Checked : ECheckBoxState::Unchecked;
+}
+
+
+void SNiagaraParameterPanel::OnSectionChecked(ECheckBoxState CheckState, FText Section)
+{
+	if (CheckState == ECheckBoxState::Checked)
+	{
+		ParameterPanelViewModel->SetActiveSection(Section);
+		Refresh();
+	}
+}
 
 TArray<FNiagaraParameterPanelCategory> SNiagaraParameterPanel::OnGetCategoriesForItem(const FNiagaraParameterPanelItem& Item)
 {
@@ -389,6 +440,7 @@ TSharedRef<SWidget> SNiagaraParameterPanel::OnGenerateWidgetForItem(const FNiaga
 			.ButtonStyle(FAppStyle::Get(), "RoundButton")
 			.ForegroundColor(FSlateColor::UseForeground())
 			.ContentPadding(FMargin(2.0f))
+			.Visibility(this, &SNiagaraParameterPanel::GetVisibilityForReferenceCounter)
 			.HAlign(HAlign_Right)
 			.VAlign(VAlign_Center)
 			.ButtonContent()
@@ -407,6 +459,18 @@ TSharedRef<SWidget> SNiagaraParameterPanel::OnGenerateWidgetForItem(const FNiaga
 	];
 }
 
+EVisibility SNiagaraParameterPanel::GetVisibilityForReferenceCounter() const
+{
+	if (ParameterPanelViewModel->GetShowReferenceCounter())
+	{
+		return EVisibility::Visible;
+	}
+	else
+	{
+		return EVisibility::Collapsed;
+	}
+}
+
 TSharedPtr<SWidget> SNiagaraParameterPanel::OnContextMenuOpening()
 {
 	const TArray<FNiagaraParameterPanelItem>& SelectedItems = ItemSelector->GetSelectedItems();
diff --git a/Engine/Plugins/FX/Niagara/Source/NiagaraEditor/Private/Widgets/SNiagaraScriptGraph.cpp b/Engine/Plugins/FX/Niagara/Source/NiagaraEditor/Private/Widgets/SNiagaraScriptGraph.cpp
--- a/Engine/Plugins/FX/Niagara/Source/NiagaraEditor/Private/Widgets/SNiagaraScriptGraph.cpp
+++ b/Engine/Plugins/FX/Niagara/Source/NiagaraEditor/Private/Widgets/SNiagaraScriptGraph.cpp
@@ -39,6 +39,7 @@ void SNiagaraScriptGraph::Construct(const FArguments& InArgs, TSharedRef<FNiagar
 
 	GraphTitle = InArgs._GraphTitle;
 	SetForegroundColor(InArgs._ForegroundColor);
+	bShowHeader = InArgs._ShowHeader;
 
 	GraphEditor = ConstructGraphEditor();
 	if (InArgs._ZoomToFitOnLoad)
@@ -112,6 +113,7 @@ TSharedRef<SGraphEditor> SNiagaraScriptGraph::ConstructGraphEditor()
 					.Text(ViewModel.ToSharedRef(), &FNiagaraScriptGraphViewModel::GetDisplayName)
 					.TextStyle(FAppStyle::Get(), TEXT("GraphBreadcrumbButtonText"))
 					.Justification(ETextJustify::Center)
+					.Visibility(bShowHeader ? EVisibility::Visible : EVisibility::Collapsed)
 				]
 			]
 			// Search Box
diff --git a/Engine/Plugins/FX/Niagara/Source/NiagaraEditor/Public/NiagaraActions.h b/Engine/Plugins/FX/Niagara/Source/NiagaraEditor/Public/NiagaraActions.h
--- a/Engine/Plugins/FX/Niagara/Source/NiagaraEditor/Public/NiagaraActions.h
+++ b/Engine/Plugins/FX/Niagara/Source/NiagaraEditor/Public/NiagaraActions.h
@@ -362,4 +362,19 @@ public:
 private:
 	TSharedPtr<FEdGraphSchemaAction> SourceAction;
 	TOptional<FNiagaraVariable> TargetParameter;
+};
+
+class NIAGARAEDITOR_API FNiagaraScriptDragOperation : public FDecoratedDragDropOp
+{
+public:
+	DRAG_DROP_OPERATOR_TYPE(FNiagaraScriptDragOperation, FDecoratedDragDropOp)
+
+		FNiagaraScriptDragOperation(UNiagaraScript* InPayloadScript, const FGuid& InVersion, const FText& InFriendlyName)
+		: Script(InPayloadScript), Version(InVersion), FriendlyName(InFriendlyName)
+	{
+	}
+
+	TWeakObjectPtr<UNiagaraScript> Script;
+	FGuid Version;
+	FText FriendlyName;
 };
\ No newline at end of file
diff --git a/Engine/Plugins/FX/Niagara/Source/NiagaraEditor/Public/NiagaraEditorModule.h b/Engine/Plugins/FX/Niagara/Source/NiagaraEditor/Public/NiagaraEditorModule.h
--- a/Engine/Plugins/FX/Niagara/Source/NiagaraEditor/Public/NiagaraEditorModule.h
+++ b/Engine/Plugins/FX/Niagara/Source/NiagaraEditor/Public/NiagaraEditorModule.h
@@ -56,7 +56,8 @@ public:
 	virtual TSharedRef<SWidget> CreateStackView(UNiagaraStackViewModel& StackViewModel) const = 0;
 	virtual TSharedRef<SWidget> CreateSystemOverview(TSharedRef<FNiagaraSystemViewModel> SystemViewModel) const = 0;
 	virtual TSharedRef<SWidget> CreateStackIssueIcon(UNiagaraStackViewModel& StackViewModel, UNiagaraStackEntry& StackEntry) const = 0;
-	virtual TSharedRef<SWidget> CreateScriptScratchPad(UNiagaraScratchPadViewModel& ScriptScratchPadViewModel) const = 0;
+
+	virtual TSharedRef<SWidget> CreateScriptScratchPadManager(UNiagaraScratchPadViewModel& ScriptScratchPadViewModel) const = 0;
 	virtual TSharedRef<SWidget> CreateCurveOverview(TSharedRef<FNiagaraSystemViewModel> SystemViewModel) const = 0;
 	virtual FLinearColor GetColorForExecutionCategory(FName ExecutionCategory) const = 0;
 };
diff --git a/Engine/Plugins/FX/Niagara/Source/NiagaraEditor/Public/NiagaraEditorStyle.h b/Engine/Plugins/FX/Niagara/Source/NiagaraEditor/Public/NiagaraEditorStyle.h
--- a/Engine/Plugins/FX/Niagara/Source/NiagaraEditor/Public/NiagaraEditorStyle.h
+++ b/Engine/Plugins/FX/Niagara/Source/NiagaraEditor/Public/NiagaraEditorStyle.h
@@ -50,6 +50,7 @@ private:
 	void InitOutlinerStyle();
 	void InitScalabilityColors();
 	void InitScalabilityIcons();
+	void InitScratchStyle();
 
 	static TSharedPtr<FNiagaraEditorStyle> NiagaraEditorStyle;
 };
diff --git a/Engine/Plugins/FX/Niagara/Source/NiagaraEditor/Public/ViewModels/NiagaraParameterPanelViewModel.h b/Engine/Plugins/FX/Niagara/Source/NiagaraEditor/Public/ViewModels/NiagaraParameterPanelViewModel.h
--- a/Engine/Plugins/FX/Niagara/Source/NiagaraEditor/Public/ViewModels/NiagaraParameterPanelViewModel.h
+++ b/Engine/Plugins/FX/Niagara/Source/NiagaraEditor/Public/ViewModels/NiagaraParameterPanelViewModel.h
@@ -140,7 +140,7 @@ public:
 
 	/** Delegate to get the names of all selected parameter items. */
 	DECLARE_DELEGATE_RetVal(TArray<FName>, FOnGetSelectedParameterNames);
-
+	INiagaraParameterPanelViewModel();
 	virtual ~INiagaraParameterPanelViewModel() override;
 
 	//~ Begin INiagaraImmutableParameterPanelViewModel interface
@@ -198,6 +198,9 @@ public:
 	virtual void PasteParameterMetaData(const TArray<FNiagaraParameterPanelItem> SelectedItems);
 	virtual bool GetCanPasteParameterMetaDataAndToolTip(FText& OutCanPasteToolTip);
 
+	virtual bool GetAllowMultiSelect() const { return false; }
+	virtual bool GetShowReferenceCounter() const { return true; }
+
 	virtual void DuplicateParameters(const TArray<FNiagaraParameterPanelItem> ItemToDuplicate) const;
 	virtual bool GetCanDuplicateParameterAndToolTip(const TArray<FNiagaraParameterPanelItem>& ItemsToDuplicate, FText& OutCanDuplicateParameterToolTip) const;
 	bool GetCanDebugParameters(const TArray<FNiagaraParameterPanelItem>& ItemsToDebug) const;
@@ -239,6 +242,20 @@ public:
 	FOnNotifyParameterPendingNamespaceModifierRename& GetOnNotifyParameterPendingNamespaceModifierRenameDelegate() { return OnNotifyParameterPendingNamespaceModifierRenameDelegate; };
 	FOnGetParametersWithNamespaceModifierRenamePending& GetParametersWithNamespaceModifierRenamePendingDelegate() { return OnGetParametersWithNamespaceModifierRenamePendingDelegate; };
 
+	FText GetActiveSection() const
+	{
+		return Sections.IsValidIndex(ActiveSectionIndex) ? Sections[ActiveSectionIndex] : FText::GetEmpty();
+	}
+	virtual bool GetShowSections() const { return false; }
+	virtual bool GetNamespaceActive(const FName& InNamespace) const { return true; }
+	virtual const TArray<FText>& GetSections() const { return Sections; }
+	virtual FText GetTooltipForSection(FText& InSection) const;
+	virtual void SetActiveSection(FText& InSection);
+
+	virtual bool IsVariableSelected(FNiagaraVariableBase& InVar) const { return false; };
+
+	virtual TSharedPtr<FNiagaraObjectSelection> GetVariableObjectSelection() { return TSharedPtr<FNiagaraObjectSelection>(); };
+
 protected:
 	static bool CanMakeNewParameterOfType(const FNiagaraTypeDefinition& InType);
 
@@ -262,6 +279,10 @@ protected:
 
 	/** Transient UNiagaraScriptVariables used to pass to new FNiagaraParameterPanelItems when the source FNiagaraVariable is not associated with a UNiagaraScriptVariable in a graph. */
 	mutable TMap<FNiagaraVariable, TObjectPtr<UNiagaraScriptVariable>> TransientParameterToScriptVarMap;
+
+	TArray<FText> Sections;
+	TArray<FText> SectionDescriptions;
+	int32 ActiveSectionIndex;
 };
 
 class FNiagaraSystemToolkitParameterPanelViewModel : public INiagaraParameterPanelViewModel, public TNiagaraViewModelManager<UNiagaraSystem, FNiagaraSystemToolkitParameterPanelViewModel>
@@ -326,6 +347,16 @@ public:
 	TSharedRef<SWidget> CreateAddParameterMenuForAssignmentNode(UNiagaraNodeAssignment* AssignmentNode, const TSharedPtr<SComboButton>& AddButton) const;
 
 	void DebugParameters(const TArray<FNiagaraParameterPanelItem> ItemToDuplicate) const;
+
+	virtual void OnParameterItemSelected(const FNiagaraParameterPanelItem& SelectedItem, ESelectInfo::Type SelectInfo) const override;
+	virtual bool IsVariableSelected(FNiagaraVariableBase& InVar) const override;
+
+
+	virtual bool GetShowSections() const override { return true; }
+	virtual bool GetNamespaceActive(const FName& InNamespace) const override { return true; }
+
+	virtual TSharedPtr<FNiagaraObjectSelection> GetVariableObjectSelection() override { return VariableObjectSelection; };
+	virtual bool GetShowReferenceCounter() const override;
 private:
 	const TArray<UNiagaraGraph*> GetAllGraphsConst() const;
 
@@ -363,8 +394,14 @@ private:
 	mutable TArray<FNiagaraParameterPanelCategory> CachedCurrentCategories;
 
 	static TArray<FNiagaraParameterPanelCategory> DefaultCategories;
+	static TArray<FNiagaraParameterPanelCategory> UserCategories;
 	static TArray<FNiagaraParameterPanelCategory> DefaultAdvancedCategories;
 
+	TSharedPtr<FNiagaraObjectSelection> VariableObjectSelection;
+	int32 UserOnlyIdx = -1;
+	int32 CommonOnlyIdx = -1;
+	int32 ActiveIdx = -1;
+
 	TNiagaraViewModelManager<UNiagaraSystem, FNiagaraSystemToolkitParameterPanelViewModel>::Handle RegisteredHandle;
 };
 
@@ -419,6 +456,9 @@ public:
 
 	virtual const TArray<FNiagaraParameterPanelCategory>& GetDefaultCategories() const override;
 
+	virtual bool GetAllowMultiSelect() const override { return true; }
+
+
 	virtual FMenuAndSearchBoxWidgets GetParameterMenu(FNiagaraParameterPanelCategory Category) const override;
 
 	virtual FReply HandleDragDropOperation(TSharedPtr<FDragDropOperation> DragDropOperation) const override;
diff --git a/Engine/Plugins/FX/Niagara/Source/NiagaraEditor/Public/ViewModels/NiagaraScratchPadScriptViewModel.h b/Engine/Plugins/FX/Niagara/Source/NiagaraEditor/Public/ViewModels/NiagaraScratchPadScriptViewModel.h
--- a/Engine/Plugins/FX/Niagara/Source/NiagaraEditor/Public/ViewModels/NiagaraScratchPadScriptViewModel.h
+++ b/Engine/Plugins/FX/Niagara/Source/NiagaraEditor/Public/ViewModels/NiagaraScratchPadScriptViewModel.h
@@ -32,6 +32,8 @@ public:
 	virtual void AddReferencedObjects(FReferenceCollector& Collector) override;
 	//~ End FGCObject
 
+	TArray<UNiagaraGraph*> GetEditableGraphs();
+
 	//~ Begin NiagaraParameterDefinitionsSubscriberViewModel Interface
 protected:
 	virtual INiagaraParameterDefinitionsSubscriber* GetParameterDefinitionsSubscriber() override { return &EditScript; };
diff --git a/Engine/Plugins/FX/Niagara/Source/NiagaraEditor/Public/ViewModels/NiagaraScratchPadViewModel.h b/Engine/Plugins/FX/Niagara/Source/NiagaraEditor/Public/ViewModels/NiagaraScratchPadViewModel.h
--- a/Engine/Plugins/FX/Niagara/Source/NiagaraEditor/Public/ViewModels/NiagaraScratchPadViewModel.h
+++ b/Engine/Plugins/FX/Niagara/Source/NiagaraEditor/Public/ViewModels/NiagaraScratchPadViewModel.h
@@ -83,6 +83,8 @@ public:
 
 	FOnScriptDeleted& OnScriptDeleted();
 
+	void OpenEditorForActive();
+
 private:
 	TSharedRef<FNiagaraSystemViewModel> GetSystemViewModel();
 
@@ -96,7 +98,7 @@ private:
 
 	void ScriptGraphNodeSelectionChanged(TWeakPtr<FNiagaraScratchPadScriptViewModel> InScriptViewModelWeak);
 
-	void ScriptViewModelScriptRenamed();
+	void ScriptViewModelScriptRenamed(TWeakPtr<FNiagaraScratchPadScriptViewModel> ScriptViewModelWeak);
 
 	void ScriptViewModelPinnedChanged(TWeakPtr<FNiagaraScratchPadScriptViewModel> ScriptViewModelWeak);
 
diff --git a/Engine/Plugins/FX/Niagara/Source/NiagaraEditor/Public/ViewModels/NiagaraSystemEditorDocumentsViewModel.h b/Engine/Plugins/FX/Niagara/Source/NiagaraEditor/Public/ViewModels/NiagaraSystemEditorDocumentsViewModel.h
new file mode 100644
--- /dev/null
+++ b/Engine/Plugins/FX/Niagara/Source/NiagaraEditor/Public/ViewModels/NiagaraSystemEditorDocumentsViewModel.h
@@ -0,0 +1,84 @@
+// Copyright Epic Games, Inc. All Rights Reserved.
+
+#pragma once
+
+#include "CoreMinimal.h"
+#include "NiagaraCommon.h"
+#include "WorkflowOrientedApp/WorkflowTabFactory.h"
+#include "WorkflowOrientedApp/WorkflowTabManager.h"
+#include "WorkflowOrientedApp/WorkflowUObjectDocuments.h"
+#include "Widgets/SNiagaraScriptGraph.h"
+#include "NiagaraSystemEditorDocumentsViewModel.generated.h"
+
+class FNiagaraSystemViewModel;
+class FNiagaraScratchPadScriptViewModel;
+class FNiagaraObjectSelection;
+class UNiagaraScript;
+class UEdGraph;
+
+UCLASS()
+class NIAGARAEDITOR_API UNiagaraSystemEditorDocumentsViewModel : public UObject
+{
+	GENERATED_BODY()
+
+public:
+	DECLARE_MULTICAST_DELEGATE_OneParam(FOnRequestOpenChildScript, UEdGraph* /* Graph */);
+	DECLARE_MULTICAST_DELEGATE_OneParam(FOnRequestCloseChildScript, UEdGraph* /* Graph */);
+
+	DECLARE_MULTICAST_DELEGATE_OneParam(FScriptToolkitsActiveDocumentChanged, TSharedPtr<SDockTab>);
+
+
+public:
+	void Initialize(TSharedRef<FNiagaraSystemViewModel> InSystemViewModel);
+	void Finalize();
+
+	void OpenChildScript(UEdGraph* Graph);
+	void CloseChildScript(UEdGraph* Graph);
+
+	TArray<class UNiagaraGraph*> GetEditableGraphsForActiveDocument();
+	
+	void InitializePreTabManager(TSharedPtr<class FNiagaraSystemToolkit> InToolkit);
+	void InitializePostTabManager(TSharedPtr<class FNiagaraSystemToolkit> InToolkit);
+
+	void SetActiveDocumentTab(TSharedPtr<SDockTab> Tab);
+	TWeakPtr<SDockTab> GetActiveDocumentTab() const { return ActiveDocumentTab; }
+	FScriptToolkitsActiveDocumentChanged& OnActiveDocumentChanged() {
+		return ActiveDocChangedDelegate;
+	}
+
+	TSharedPtr<class FNiagaraScratchPadScriptViewModel> GetActiveScratchPadViewModelIfSet();
+	static TSharedPtr < class FNiagaraScratchPadScriptViewModel> GetScratchPadViewModelFromGraph(FNiagaraSystemViewModel* InSysViewModel, UEdGraph* InTargetGraph);
+
+protected:
+	TSharedPtr<SDockTab> OpenDocument(const UObject* DocumentID, FDocumentTracker::EOpenDocumentCause Cause);
+
+	void NavigateTab(FDocumentTracker::EOpenDocumentCause InCause);
+
+	void CloseDocumentTab(const UObject* DocumentID);
+
+	// Finds any open tabs containing the specified document and adds them to the specified array; returns true if at least one is found
+	bool FindOpenTabsContainingDocument(const UObject* DocumentID, /*inout*/ TArray< TSharedPtr<SDockTab> >& Results);
+
+
+private:
+	TSharedRef<FNiagaraSystemViewModel> GetSystemViewModel();
+
+	FScriptToolkitsActiveDocumentChanged ActiveDocChangedDelegate;
+
+	TWeakPtr<FNiagaraSystemViewModel> SystemViewModelWeak;
+
+	TSharedRef<class SNiagaraScratchPadScriptEditor> CreateGraphEditorWidget(TSharedRef<FTabInfo> InTabInfo, UEdGraph* InGraph);
+
+private:
+	TSharedPtr<FDocumentTracker> DocumentManager;
+	/** Factory that spawns graph editors; used to look up all tabs spawned by it. */
+	TWeakPtr<FDocumentTabFactory> GraphEditorTabFactoryPtr;
+
+	TSharedPtr<class FTabManager> TabManager;
+
+	TWeakPtr<SDockTab> ActiveDocumentTab;
+
+	TSharedPtr<FNiagaraScratchPadScriptViewModel> ActiveDocumentTabScriptViewModel;
+
+
+};
\ No newline at end of file
diff --git a/Engine/Plugins/FX/Niagara/Source/NiagaraEditor/Public/ViewModels/NiagaraSystemViewModel.h b/Engine/Plugins/FX/Niagara/Source/NiagaraEditor/Public/ViewModels/NiagaraSystemViewModel.h
--- a/Engine/Plugins/FX/Niagara/Source/NiagaraEditor/Public/ViewModels/NiagaraSystemViewModel.h
+++ b/Engine/Plugins/FX/Niagara/Source/NiagaraEditor/Public/ViewModels/NiagaraSystemViewModel.h
@@ -18,6 +18,7 @@
 #include "NiagaraSystemScalabilityViewModel.h"
 #include "UObject/ObjectKey.h"
 
+
 class UNiagaraSystem;
 class UNiagaraComponent;
 class UNiagaraSequence;
@@ -44,10 +45,12 @@ class UNiagaraScratchPadViewModel;
 class UNiagaraCurveSelectionViewModel;
 class UNiagaraMessageData;
 class UNiagaraEditorParametersAdapter;
+class UNiagaraSystemEditorDocumentsViewModel;
 
 class INiagaraParameterDefinitionsSubscriber;
 
 class FNiagaraPlaceholderDataInterfaceManager;
+class INiagaraParameterPanelViewModel;
 
 /** Defines different editing modes for this system view model. */
 enum class NIAGARAEDITOR_API ENiagaraSystemViewModelEditMode
@@ -410,6 +413,13 @@ public:
 	/** Gets whether or not this view model is for data processing only and will not be displayed in the UI. */
 	bool GetIsForDataProcessingOnly() const;
 
+	INiagaraParameterPanelViewModel* GetParameterPanelViewModel() const {
+		return ParameterPanelViewModel.Pin().Get();
+	}
+	void SetParameterPanelViewModel(TSharedPtr<INiagaraParameterPanelViewModel> InVM) { ParameterPanelViewModel = InVM; }
+
+	UNiagaraSystemEditorDocumentsViewModel* GetDocumentViewModel() { return EditorDocumentsViewModel;};
+
 private:
 	/** Sends message jobs to FNiagaraMessageManager for all compile events from the last compile. */
 	void SendLastCompileMessageJobs() const;
@@ -662,6 +672,8 @@ private:
 	TSharedPtr<FNiagaraOverviewGraphViewModel> OverviewGraphViewModel;
 
 	UNiagaraStackViewModel* SystemStackViewModel;
+	
+	UNiagaraSystemEditorDocumentsViewModel* EditorDocumentsViewModel;
 
 	UNiagaraSystemSelectionViewModel* SelectionViewModel;
 
@@ -671,6 +683,8 @@ private:
 
 	UNiagaraSystemScalabilityViewModel* ScalabilityViewModel;
 
+	TWeakPtr<INiagaraParameterPanelViewModel> ParameterPanelViewModel;
+
 	TSharedPtr<FNiagaraPlaceholderDataInterfaceManager> PlaceholderDataInterfaceManager;
 
 	TArray<UNiagaraScript*> ScriptsToCheckForStatus;
diff --git a/Engine/Plugins/FX/Niagara/Source/NiagaraEditor/Public/ViewModels/Stack/NiagaraStackEntry.h b/Engine/Plugins/FX/Niagara/Source/NiagaraEditor/Public/ViewModels/Stack/NiagaraStackEntry.h
--- a/Engine/Plugins/FX/Niagara/Source/NiagaraEditor/Public/ViewModels/Stack/NiagaraStackEntry.h
+++ b/Engine/Plugins/FX/Niagara/Source/NiagaraEditor/Public/ViewModels/Stack/NiagaraStackEntry.h
@@ -455,6 +455,8 @@ public:
 	bool HasBaseEmitter() const;
 
 	bool HasIssuesOrAnyChildHasIssues() const;
+	bool HasUsagesOrAnyChildHasUsages() const;
+	void GetRecursiveUsages(bool& bRead, bool& bWrite) const;
 
 	int32 GetTotalNumberOfCustomNotes() const;
 	
@@ -529,6 +531,17 @@ public:
 
 	virtual FText GetInheritanceMessage() const { return FText(); }
 
+	virtual void InvalidateCollectedUsage();
+
+
+	struct FCollectedUsageData
+	{
+	public:
+		bool bHasReferencedParameterRead = false;
+		bool bHasReferencedParameterWrite = false;
+	};
+	virtual const FCollectedUsageData& GetCollectedUsageData() const;
+
 protected:
 	virtual void BeginDestroy() override;
 
@@ -552,6 +565,10 @@ protected:
 
 	virtual void FinalizeInternal();
 
+
+	mutable TOptional<FCollectedUsageData> CachedCollectedUsageData;
+
+
 private:
 	void ChildStructureChanged(ENiagaraStructureChangedFlags Info);
 
@@ -594,8 +611,9 @@ private:
 		TArray<UNiagaraStackEntry*> ChildrenWithIssues;
 	};
 
+
 	const FCollectedIssueData& GetCollectedIssueData() const;
-	
+
 	TWeakPtr<FNiagaraSystemViewModel> SystemViewModel;
 	TWeakPtr<FNiagaraEmitterViewModel> EmitterViewModel;
 
diff --git a/Engine/Plugins/FX/Niagara/Source/NiagaraEditor/Public/ViewModels/Stack/NiagaraStackFunctionInput.h b/Engine/Plugins/FX/Niagara/Source/NiagaraEditor/Public/ViewModels/Stack/NiagaraStackFunctionInput.h
--- a/Engine/Plugins/FX/Niagara/Source/NiagaraEditor/Public/ViewModels/Stack/NiagaraStackFunctionInput.h
+++ b/Engine/Plugins/FX/Niagara/Source/NiagaraEditor/Public/ViewModels/Stack/NiagaraStackFunctionInput.h
@@ -280,6 +280,8 @@ public:
 	TOptional<FNiagaraVariableMetaData> GetMetadata() const;
 	TOptional<FGuid> GetMetadataGuid() const;
 
+
+	virtual const FCollectedUsageData& GetCollectedUsageData() const override;
 protected:
 	//~ UNiagaraStackEntry interface
 	virtual void FinalizeInternal() override;
diff --git a/Engine/Plugins/FX/Niagara/Source/NiagaraEditor/Public/ViewModels/Stack/NiagaraStackModuleItem.h b/Engine/Plugins/FX/Niagara/Source/NiagaraEditor/Public/ViewModels/Stack/NiagaraStackModuleItem.h
--- a/Engine/Plugins/FX/Niagara/Source/NiagaraEditor/Public/ViewModels/Stack/NiagaraStackModuleItem.h
+++ b/Engine/Plugins/FX/Niagara/Source/NiagaraEditor/Public/ViewModels/Stack/NiagaraStackModuleItem.h
@@ -126,6 +126,7 @@ public:
 	void SetDebugDrawEnabled(bool bInEnabled);
 
 	FOnNoteModeSet& OnNoteModeSet() { return OnNoteModeSetDelegate;}
+
 	
 protected:
 	FOnRequestDeprecationRecommended DeprecationDelegate;
@@ -137,6 +138,8 @@ protected:
 	virtual TOptional<FDropRequestResponse> CanDropInternal(const FDropRequest& DropRequest) override;
 	virtual TOptional<FDropRequestResponse> DropInternal(const FDropRequest& DropRequest) override;
 
+	virtual const FCollectedUsageData& GetCollectedUsageData() const override;
+
 private:
 	bool FilterOutputCollection(const UNiagaraStackEntry& Child) const;
 	bool FilterOutputCollectionChild(const UNiagaraStackEntry& Child) const;
diff --git a/Engine/Plugins/FX/Niagara/Source/NiagaraEditor/Public/ViewModels/Stack/NiagaraStackModuleItemOutput.h b/Engine/Plugins/FX/Niagara/Source/NiagaraEditor/Public/ViewModels/Stack/NiagaraStackModuleItemOutput.h
--- a/Engine/Plugins/FX/Niagara/Source/NiagaraEditor/Public/ViewModels/Stack/NiagaraStackModuleItemOutput.h
+++ b/Engine/Plugins/FX/Niagara/Source/NiagaraEditor/Public/ViewModels/Stack/NiagaraStackModuleItemOutput.h
@@ -38,6 +38,8 @@ public:
 	/** Gets the assigned parameter handle as displayable text. */
 	FText GetOutputParameterHandleText() const;
 
+	virtual const FCollectedUsageData& GetCollectedUsageData() const override;
+
 private:
 	/** The function call node which represents the emitter which owns this Output */
 	TWeakObjectPtr<UNiagaraNodeFunctionCall> FunctionCallNode;
diff --git a/Engine/Plugins/FX/Niagara/Source/NiagaraEditor/Public/ViewModels/Stack/NiagaraStackScriptItemGroup.h b/Engine/Plugins/FX/Niagara/Source/NiagaraEditor/Public/ViewModels/Stack/NiagaraStackScriptItemGroup.h
--- a/Engine/Plugins/FX/Niagara/Source/NiagaraEditor/Public/ViewModels/Stack/NiagaraStackScriptItemGroup.h
+++ b/Engine/Plugins/FX/Niagara/Source/NiagaraEditor/Public/ViewModels/Stack/NiagaraStackScriptItemGroup.h
@@ -74,6 +74,8 @@ private:
 
 	TOptional<FDropRequestResponse> CanDropAssetsOnTarget(const UNiagaraStackEntry& TargetEntry, const FDropRequest& DropRequest);
 
+	TOptional<FDropRequestResponse> CanDropScriptsOnTarget(const UNiagaraStackEntry& TargetEntry, const FDropRequest& DropRequest);
+
 	TOptional<FDropRequestResponse> CanDropParameterOnTarget(const UNiagaraStackEntry& TargetEntry, const FDropRequest& DropRequest);
 
 	TOptional<FDropRequestResponse> DropOnTarget(const UNiagaraStackEntry& TargetEntry, const FDropRequest& DropRequest);
@@ -82,6 +84,8 @@ private:
 
 	TOptional<FDropRequestResponse> DropAssetsOnTarget(const UNiagaraStackEntry& TargetEntry, const FDropRequest& DropRequest);
 
+	TOptional<FDropRequestResponse> DropScriptsOnTarget(const UNiagaraStackEntry& TargetEntry, const FDropRequest& DropRequest);
+
 	TOptional<FDropRequestResponse> DropParameterOnTarget(const UNiagaraStackEntry& TargetEntry, const FDropRequest& DropRequest);
 
 protected:
diff --git a/Engine/Plugins/FX/Niagara/Source/NiagaraEditor/Public/ViewModels/Stack/NiagaraStackViewModel.h b/Engine/Plugins/FX/Niagara/Source/NiagaraEditor/Public/ViewModels/Stack/NiagaraStackViewModel.h
--- a/Engine/Plugins/FX/Niagara/Source/NiagaraEditor/Public/ViewModels/Stack/NiagaraStackViewModel.h
+++ b/Engine/Plugins/FX/Niagara/Source/NiagaraEditor/Public/ViewModels/Stack/NiagaraStackViewModel.h
@@ -148,6 +148,8 @@ public:
 
 	void RequestValidationUpdate() { bValidatorUpdatePending = true; }
 
+	void InvalidateCachedParameterUsage();
+
 private:
 	/** Recursively Expands all groups and collapses all items in the stack. */
 	void CollapseToHeadersRecursive(TArray<UNiagaraStackEntry*> Entries);
diff --git a/Engine/Plugins/FX/Niagara/Source/NiagaraEditor/Public/Widgets/SNiagaraParameterPanel.h b/Engine/Plugins/FX/Niagara/Source/NiagaraEditor/Public/Widgets/SNiagaraParameterPanel.h
--- a/Engine/Plugins/FX/Niagara/Source/NiagaraEditor/Public/Widgets/SNiagaraParameterPanel.h
+++ b/Engine/Plugins/FX/Niagara/Source/NiagaraEditor/Public/Widgets/SNiagaraParameterPanel.h
@@ -124,6 +124,16 @@ private:
 
 	static const FSlateBrush* GetViewOptionsBorderBrush();
 
+	void ConstructSectionButtons();
+
+	ECheckBoxState GetSectionCheckState(FText Section) const;
+
+	void OnSectionChecked(ECheckBoxState CheckState, FText Section);
+
+	FText GetTooltipForSection(FText Section);
+
+	EVisibility GetVisibilityForReferenceCounter() const;
+
 private:
 	mutable bool bPendingRefresh;
 	mutable bool bPendingSelectionRestore;
@@ -146,6 +156,7 @@ private:
 	TSharedPtr<INiagaraParameterPanelViewModel> ParameterPanelViewModel;
 	TSharedPtr<SNiagaraParameterPanelSelector> ItemSelector;
 	TSharedPtr<SSearchBox> ItemSelectorSearchBox;
+	TSharedPtr<class SWrapBox> SectionSelectorBox;
 
 	/** Whether or not to display icons signifying whether parameters in the panel are synchronizing with a subscribed parameter definition. */
 	bool bShowParameterSynchronizingWithLibraryIcon;
diff --git a/Engine/Plugins/FX/Niagara/Source/NiagaraEditor/Public/Widgets/SNiagaraScriptGraph.h b/Engine/Plugins/FX/Niagara/Source/NiagaraEditor/Public/Widgets/SNiagaraScriptGraph.h
--- a/Engine/Plugins/FX/Niagara/Source/NiagaraEditor/Public/Widgets/SNiagaraScriptGraph.h
+++ b/Engine/Plugins/FX/Niagara/Source/NiagaraEditor/Public/Widgets/SNiagaraScriptGraph.h
@@ -19,11 +19,13 @@ public:
 	SLATE_BEGIN_ARGS(SNiagaraScriptGraph)
 		: _ForegroundColor(FAppStyle::GetColor("Graph.ForegroundColor"))
 		, _ZoomToFitOnLoad(false)
+		, _ShowHeader(true)
 	{}
 		/** The text displayed in the title bar of the graph. */
 		SLATE_ATTRIBUTE(FText, GraphTitle)
 		SLATE_ATTRIBUTE(FSlateColor, ForegroundColor)
 		SLATE_ARGUMENT(bool, ZoomToFitOnLoad)
+		SLATE_ARGUMENT(bool, ShowHeader)
 	SLATE_END_ARGS();
 
 	NIAGARAEDITOR_API void Construct(const FArguments& InArgs, TSharedRef<FNiagaraScriptGraphViewModel> InViewModel);
@@ -129,6 +131,7 @@ private:
 	int CurrentFocusedSearchMatchIndex;
 	TArray<TSharedPtr<INiagaraScriptGraphFocusInfo>> CurrentSearchResults;
 	bool bGraphSearchBoxActive;
+	bool bShowHeader;
 
 private:
 	// action menu data
diff --git a/Engine/Plugins/FX/Niagara/Source/NiagaraEditorWidgets/Private/NiagaraEditorWidgetsModule.cpp b/Engine/Plugins/FX/Niagara/Source/NiagaraEditorWidgets/Private/NiagaraEditorWidgetsModule.cpp
--- a/Engine/Plugins/FX/Niagara/Source/NiagaraEditorWidgets/Private/NiagaraEditorWidgetsModule.cpp
+++ b/Engine/Plugins/FX/Niagara/Source/NiagaraEditorWidgets/Private/NiagaraEditorWidgetsModule.cpp
@@ -17,7 +17,7 @@
 #include "SNiagaraOverviewGraph.h"
 #include "NiagaraEditorWidgetsUtilities.h"
 #include "Stack/SNiagaraStackIssueIcon.h"
-#include "SNiagaraScratchPad.h"
+#include "SNiagaraScratchPadScriptManager.h"
 #include "SNiagaraCurveOverview.h"
 
 #include "Modules/ModuleManager.h"
@@ -207,9 +207,9 @@ TSharedRef<SWidget> FNiagaraEditorWidgetsModule::FNiagaraEditorWidgetProvider::C
 	return SNew(SNiagaraStackIssueIcon, &StackViewModel, &StackEntry);
 }
 
-TSharedRef<SWidget> FNiagaraEditorWidgetsModule::FNiagaraEditorWidgetProvider::CreateScriptScratchPad(UNiagaraScratchPadViewModel& ScriptScratchPadViewModel) const
+TSharedRef<SWidget> FNiagaraEditorWidgetsModule::FNiagaraEditorWidgetProvider::CreateScriptScratchPadManager(UNiagaraScratchPadViewModel& ScriptScratchPadViewModel) const
 {
-	return SNew(SNiagaraScratchPad, &ScriptScratchPadViewModel);
+	return SNew(SNiagaraScratchPadScriptManager, &ScriptScratchPadViewModel);
 }
 
 TSharedRef<SWidget> FNiagaraEditorWidgetsModule::FNiagaraEditorWidgetProvider::CreateCurveOverview(TSharedRef<FNiagaraSystemViewModel> SystemViewModel) const
diff --git a/Engine/Plugins/FX/Niagara/Source/NiagaraEditorWidgets/Private/NiagaraEditorWidgetsStyle.cpp b/Engine/Plugins/FX/Niagara/Source/NiagaraEditorWidgets/Private/NiagaraEditorWidgetsStyle.cpp
--- a/Engine/Plugins/FX/Niagara/Source/NiagaraEditorWidgets/Private/NiagaraEditorWidgetsStyle.cpp
+++ b/Engine/Plugins/FX/Niagara/Source/NiagaraEditorWidgets/Private/NiagaraEditorWidgetsStyle.cpp
@@ -231,6 +231,11 @@ FNiagaraEditorWidgetsStyle::FNiagaraEditorWidgetsStyle() : FSlateStyleSet("Niaga
 	Set("NiagaraEditor.Scalability.EffectType.CullReaction",  new IMAGE_BRUSH_SVG("Icons/Scalability/CullReaction_20", Icon20x20));
 	Set("NiagaraEditor.Scalability.EffectType.SignificanceHandler",  new IMAGE_BRUSH_SVG("Icons/Scalability/SignificanceHandler_20", Icon20x20));
 
+
+	Set("NiagaraEditor.Stack.ReadIcon", new IMAGE_BRUSH_SVG("Icons/Read_Icon", Icon16x16));
+	Set("NiagaraEditor.Stack.ReadWriteIcon", new IMAGE_BRUSH_SVG("Icons/Read_Write_Icon", Icon16x16));
+	Set("NiagaraEditor.Stack.WriteIcon", new IMAGE_BRUSH_SVG("Icons/Write_Icon", Icon16x16));
+
 	FTextBlockStyle EffectTypeScalabilityPropertyTextStyle = FTextBlockStyle(NormalText)
 		.SetFont(DEFAULT_FONT("Regular", 8))
 		.SetColorAndOpacity(FLinearColor(1.f, 1.f, 1.f, 1.f));
diff --git a/Engine/Plugins/FX/Niagara/Source/NiagaraEditorWidgets/Private/SNiagaraOverviewStack.cpp b/Engine/Plugins/FX/Niagara/Source/NiagaraEditorWidgets/Private/SNiagaraOverviewStack.cpp
--- a/Engine/Plugins/FX/Niagara/Source/NiagaraEditorWidgets/Private/SNiagaraOverviewStack.cpp
+++ b/Engine/Plugins/FX/Niagara/Source/NiagaraEditorWidgets/Private/SNiagaraOverviewStack.cpp
@@ -730,6 +730,87 @@ EVisibility SNiagaraOverviewStack::GetIssueIconVisibility(UNiagaraStackEntry* St
 	return bDisplayIcon ? EVisibility::Visible : EVisibility::Collapsed;
 }
 
+const FSlateBrush* SNiagaraOverviewStack::GetUsageIcon(UNiagaraStackEntry* StackEntry)
+{
+	bool bRead = false;
+	bool bWrite = false;
+
+	StackEntry->GetRecursiveUsages(bRead, bWrite);
+	if (bRead && bWrite)
+	{
+		return FNiagaraEditorWidgetsStyle::Get().GetBrush("NiagaraEditor.Stack.ReadWriteIcon");
+	}
+	else if (bRead)
+	{
+		return FNiagaraEditorWidgetsStyle::Get().GetBrush("NiagaraEditor.Stack.ReadIcon");
+	}
+	else if (bWrite)
+	{
+		return FNiagaraEditorWidgetsStyle::Get().GetBrush("NiagaraEditor.Stack.WriteIcon");
+	}
+	else
+	{
+		return FAppStyle::Get().GetBrush("WhiteBrush");
+	}
+}
+
+
+int32 SNiagaraOverviewStack::GetUsageIconWidth(UNiagaraStackEntry* StackEntry)
+{
+	const FSlateBrush* Brush = FNiagaraEditorWidgetsStyle::Get().GetBrush("NiagaraEditor.Stack.ReadWriteIcon");
+	if (Brush)
+		return Brush->GetImageSize().X;
+	return 20;
+}
+
+int32 SNiagaraOverviewStack::GetUsageIconHeight(UNiagaraStackEntry* StackEntry)
+{
+	const FSlateBrush* Brush = FNiagaraEditorWidgetsStyle::Get().GetBrush("NiagaraEditor.Stack.ReadWriteIcon");
+	if (Brush)
+		return Brush->GetImageSize().Y;
+	return 20;
+}
+
+EVisibility SNiagaraOverviewStack::GetUsageIconVisibility(UNiagaraStackEntry* StackEntry)
+{
+	bool bDisplayIcon;
+	if (StackEntry->IsA<UNiagaraStackItemGroup>() && StackEntry->GetCanExpandInOverview() && StackEntry->GetIsExpandedInOverview())
+	{
+		// If the entry is a group and it can expand and it is expanded, we only want to show the stack Usage icon if the group itself has Usages.
+		bDisplayIcon = StackEntry->HasUsagesOrAnyChildHasUsages();
+	}
+	else
+	{
+		bDisplayIcon = StackEntry->HasUsagesOrAnyChildHasUsages();
+	}
+	return bDisplayIcon ? EVisibility::Visible : EVisibility::Collapsed;
+}
+
+
+FText SNiagaraOverviewStack::GetUsageTooltip(UNiagaraStackEntry* StackEntry)
+{
+	bool bRead = false;
+	bool bWrite = false;
+
+	StackEntry->GetRecursiveUsages(bRead, bWrite);
+	if (bRead && bWrite)
+	{
+		return LOCTEXT("UsageTooltipRW", "Reads and writes selected parameter in Parameters panel.");
+	}
+	else if (bRead)
+	{
+		return LOCTEXT("UsageTooltipR", "Reads selected parameter in Parameters panel.");
+	}
+	else if (bWrite)
+	{
+		return LOCTEXT("UsageTooltipW", "Writes selected parameter in Parameters panel.");
+	}
+	else
+	{
+		return LOCTEXT("UsageTooltipNone", "Does not touch selected parameter in Parameters panel.");
+	}
+}
+
 void SNiagaraOverviewStack::AddEntriesRecursive(UNiagaraStackEntry& EntryToAdd, TArray<UNiagaraStackEntry*>& EntryList, const TArray<UClass*>& AcceptableClasses, TArray<UNiagaraStackEntry*> ParentChain)
 {
 	if (AcceptableClasses.ContainsByPredicate([&] (UClass* Class) { return EntryToAdd.IsA(Class); }) && EntryToAdd.GetShouldShowInOverview())
@@ -803,6 +884,33 @@ TSharedRef<ITableRow> SNiagaraOverviewStack::OnGenerateRowForEntry(UNiagaraStack
 	{
 		UNiagaraStackItem* StackItem = CastChecked<UNiagaraStackItem>(Item);
 		TSharedRef<SHorizontalBox> ContentBox = SNew(SHorizontalBox);
+
+		FSlateColor IconColor(FColor::White);
+
+		int32 WidthMax = SNiagaraOverviewStack::GetUsageIconWidth(Item);
+		int32 HeightMax = SNiagaraOverviewStack::GetUsageIconHeight(Item);
+
+		// Stack Usage Highlight
+		ContentBox->AddSlot()
+			.VAlign(VAlign_Center)
+			//.AutoWidth()
+			.MaxWidth(WidthMax)
+			.Padding(4, 0, 0, 0)
+			[
+				SNew(SBorder)
+				.BorderImage_Static(&SNiagaraOverviewStack::GetUsageIcon, Item)
+				.Visibility_Static(&SNiagaraOverviewStack::GetUsageIconVisibility, Item)
+				.ToolTipText_Static(& SNiagaraOverviewStack::GetUsageTooltip, Item)
+				.BorderBackgroundColor(IconColor)
+				.Padding(FMargin(0, 0, 0, 1))
+				[
+					SNew(SBox)
+					.WidthOverride(WidthMax)
+					.HeightOverride(HeightMax)
+				]
+			];
+
+
 		// Icon Brush
 		if (StackItem->GetSupportedIconMode() == UNiagaraStackEntry::EIconMode::Brush)
 		{
@@ -864,6 +972,8 @@ TSharedRef<ITableRow> SNiagaraOverviewStack::OnGenerateRowForEntry(UNiagaraStack
 			.IconMode(SNiagaraStackIssueIcon::EIconMode::Compact)
 			.Visibility_Static(&SNiagaraOverviewStack::GetIssueIconVisibility, Item)
 		];
+
+
 		// Perf widget
 		ContentBox->AddSlot()
         .AutoWidth()
diff --git a/Engine/Plugins/FX/Niagara/Source/NiagaraEditorWidgets/Private/SNiagaraOverviewStack.h b/Engine/Plugins/FX/Niagara/Source/NiagaraEditorWidgets/Private/SNiagaraOverviewStack.h
--- a/Engine/Plugins/FX/Niagara/Source/NiagaraEditorWidgets/Private/SNiagaraOverviewStack.h
+++ b/Engine/Plugins/FX/Niagara/Source/NiagaraEditorWidgets/Private/SNiagaraOverviewStack.h
@@ -34,6 +34,11 @@ public:
 	virtual void Tick(const FGeometry& AllottedGeometry, const double InCurrentTime, const float InDeltaTime) override;
 
 	static EVisibility GetIssueIconVisibility(UNiagaraStackEntry* StackEntry);
+	static EVisibility GetUsageIconVisibility(UNiagaraStackEntry* StackEntry);
+	static const FSlateBrush* GetUsageIcon(UNiagaraStackEntry* StackEntry);
+	static int32 GetUsageIconWidth(UNiagaraStackEntry* StackEntry);
+	static int32 GetUsageIconHeight(UNiagaraStackEntry* StackEntry);
+	static FText GetUsageTooltip(UNiagaraStackEntry* StackEntry);
 
 private:
 	void AddEntriesRecursive(UNiagaraStackEntry& EntryToAdd, TArray<UNiagaraStackEntry*>& EntryList, const TArray<UClass*>& AcceptableClasses, TArray<UNiagaraStackEntry*> ParentChain);
diff --git a/Engine/Plugins/FX/Niagara/Source/NiagaraEditorWidgets/Private/SNiagaraScratchPad.cpp b/Engine/Plugins/FX/Niagara/Source/NiagaraEditorWidgets/Private/SNiagaraScratchPad.cpp
deleted file mode 100644
diff --git a/Engine/Plugins/FX/Niagara/Source/NiagaraEditorWidgets/Private/SNiagaraScratchPad.h b/Engine/Plugins/FX/Niagara/Source/NiagaraEditorWidgets/Private/SNiagaraScratchPad.h
deleted file mode 100644
diff --git a/Engine/Plugins/FX/Niagara/Source/NiagaraEditorWidgets/Private/SNiagaraScratchPadScriptManager.cpp b/Engine/Plugins/FX/Niagara/Source/NiagaraEditorWidgets/Private/SNiagaraScratchPadScriptManager.cpp
new file mode 100644
--- /dev/null
+++ b/Engine/Plugins/FX/Niagara/Source/NiagaraEditorWidgets/Private/SNiagaraScratchPadScriptManager.cpp
@@ -0,0 +1,857 @@
+// Copyright Epic Games, Inc. All Rights Reserved.
+
+#include "SNiagaraScratchPadScriptManager.h"
+#include "ViewModels/NiagaraScratchPadViewModel.h"
+#include "ViewModels/NiagaraScratchPadScriptViewModel.h"
+#include "Widgets/SDynamicLayoutBox.h"
+#include "Widgets/SItemSelector.h"
+#include "Widgets/SNiagaraSelectedObjectsDetails.h"
+#include "Widgets/SVerticalResizeBox.h"
+#include "Widgets/SNiagaraScriptGraph.h"
+#include "NiagaraObjectSelection.h"
+#include "NiagaraNodeFunctionCall.h"
+#include "NiagaraEditorStyle.h"
+#include "NiagaraEditorWidgetsStyle.h"
+#include "NiagaraScratchPadCommandContext.h"
+#include "Widgets/SNiagaraParameterPanel.h"
+
+#include "Widgets/Text/SInlineEditableTextBlock.h"
+#include "Widgets/Input/SButton.h"
+#include "Widgets/Layout/SSpacer.h"
+#include "EditorStyleSet.h"
+#include "Framework/MultiBox/MultiBoxBuilder.h"
+#include "Framework/Application/SlateApplication.h"
+#include "Layout/WidgetPath.h"
+#include "EditorFontGlyphs.h"
+#include "NiagaraConstants.h"
+#include "DragAndDrop/AssetDragDropOp.h"
+
+#define LOCTEXT_NAMESPACE "NiagaraScratchPad"
+
+FName SNiagaraScratchPadScriptManager::ScriptSelectorName = "ScriptSelector";
+FName SNiagaraScratchPadScriptManager::ScriptParameterPanelName = "ScriptParameterPanel";
+FName SNiagaraScratchPadScriptManager::ScriptEditorName = "ScriptEditor";
+FName SNiagaraScratchPadScriptManager::SelectionEditorName = "SelectionEditor";
+FName SNiagaraScratchPadScriptManager::WideLayoutName = "Wide";
+FName SNiagaraScratchPadScriptManager::NarrowLayoutName = "Narrow";
+
+class SNiagaraScratchPadScriptRow : public SCompoundWidget
+{
+	SLATE_BEGIN_ARGS(SNiagaraScratchPadScriptRow) {}
+		SLATE_ATTRIBUTE(bool, IsSelected);
+	SLATE_END_ARGS()
+
+	void Construct(
+		const FArguments& InArgs,
+		UNiagaraScratchPadViewModel* InScratchPadViewModel,
+		TSharedRef<FNiagaraScratchPadScriptViewModel> InScriptViewModel,
+		TSharedPtr<FNiagaraScratchPadCommandContext> InCommandContext)
+	{
+		ScratchPadViewModel = InScratchPadViewModel;
+		ScriptViewModel = InScriptViewModel;
+		CommandContext = InCommandContext;
+		IsSelected = InArgs._IsSelected;
+
+		ChildSlot
+		[
+			SNew(SHorizontalBox)
+			+ SHorizontalBox::Slot()
+			.VAlign(VAlign_Center)
+			.AutoWidth()
+			.Padding(3, 0, 0, 0)
+			[
+				SAssignNew(NameEditableText, SInlineEditableTextBlock)
+				.Text(this, &SNiagaraScratchPadScriptRow::GetNameText)
+				.ToolTipText_Static(&SNiagaraScratchPadScriptRow::GetTooltipTextForSelection, ScriptViewModel.ToSharedRef())
+				.IsSelected(this, &SNiagaraScratchPadScriptRow::GetIsSelected)
+				.OnVerifyTextChanged(this, &SNiagaraScratchPadScriptRow::VerifyNameTextChange)
+				.OnTextCommitted(this, &SNiagaraScratchPadScriptRow::OnNameTextCommitted)
+			]
+			+ SHorizontalBox::Slot()
+			.VAlign(VAlign_Center)
+			.AutoWidth()
+			.Padding(2, 0, 0, 0)
+			[
+				SNew(STextBlock)
+				.Visibility(this, &SNiagaraScratchPadScriptRow::GetUnappliedChangesVisibility)
+				.TextStyle(FNiagaraEditorWidgetsStyle::Get(), "NiagaraEditor.ScratchPad.EditorHeaderText")
+				.Text(FText::FromString(TEXT("*")))
+				.ToolTipText(ScriptViewModel.ToSharedRef(), &FNiagaraScratchPadScriptViewModel::GetToolTip)
+			]
+			+ SHorizontalBox::Slot()
+			[
+				SNew(SSpacer)
+			]
+		];
+	}
+
+	static FText GetTooltipTextForSelection(TSharedRef<FNiagaraScratchPadScriptViewModel> InScriptViewModel)
+	{
+		return FText::Format(LOCTEXT("ScratchTooltip", "Script: {0}\r\n{1}\r\n\r\nDouble click to edit script or drag & drop onto the stack.\r\nAn \"*\" indicates that the script has been edited but not applied yet."), InScriptViewModel->GetDisplayName(), InScriptViewModel->GetToolTip());
+	}
+
+	virtual void Tick(const FGeometry& AllottedGeometry, const double InCurrentTime, const float InDeltaTime) override
+	{
+		if (ScriptViewModel->GetIsPendingRename())
+		{
+			ScriptViewModel->SetIsPendingRename(false);
+			NameEditableText->EnterEditingMode();
+		}
+	}
+
+	virtual FReply OnMouseButtonDown(const FGeometry& MyGeometry, const FPointerEvent& MouseEvent) override
+	{
+		if (MouseEvent.GetEffectingButton() == EKeys::RightMouseButton)
+		{
+			return FReply::Handled()
+				.CaptureMouse(SharedThis(this));
+		}
+		return SCompoundWidget::OnMouseButtonDown(MyGeometry, MouseEvent);
+	}
+
+	virtual FReply OnMouseButtonUp(const FGeometry& MyGeometry, const FPointerEvent& MouseEvent) override
+	{
+		if (MouseEvent.GetEffectingButton() == EKeys::RightMouseButton)
+		{
+			// Set this script to be the active one.
+			ScratchPadViewModel->SetActiveScriptViewModel(ScriptViewModel.ToSharedRef());
+
+			FMenuBuilder MenuBuilder(true, CommandContext->GetCommands());
+			CommandContext->AddMenuItems(MenuBuilder);
+
+			FWidgetPath WidgetPath = MouseEvent.GetEventPath() != nullptr ? *MouseEvent.GetEventPath() : FWidgetPath();
+			FSlateApplication::Get().PushMenu(AsShared(), WidgetPath, MenuBuilder.MakeWidget(), MouseEvent.GetScreenSpacePosition(), FPopupTransitionEffect(FPopupTransitionEffect::ContextMenu));
+			return FReply::Handled().ReleaseMouseCapture();
+		}
+		return SCompoundWidget::OnMouseButtonUp(MyGeometry, MouseEvent);
+	}
+
+private:
+	FText GetNameText() const
+	{
+		return ScriptViewModel->GetDisplayName();
+	}
+
+	bool GetIsSelected() const
+	{
+		return IsSelected.Get(false);
+	}
+
+	bool VerifyNameTextChange(const FText& InNewNameText, FText& OutErrorMessage)
+	{
+		if (InNewNameText.IsEmpty())
+		{
+			OutErrorMessage = NSLOCTEXT("NiagaraScratchPadScriptName", "EmptyNameErrorMessage", "Script name can not be empty.");
+			return false;
+		}
+		if (InNewNameText.ToString().Len() > FNiagaraConstants::MaxScriptNameLength)
+		{
+			OutErrorMessage = FText::Format(NSLOCTEXT("NiagaraScratchPadScriptName", "NameTooLongErrorFormat", "The name entered is too long.\nThe maximum script name length is {0}."), FText::AsNumber(FNiagaraConstants::MaxScriptNameLength));
+			return false;
+		}
+
+		return true;
+	}
+
+	void OnNameTextCommitted(const FText& InText, ETextCommit::Type InCommitType)
+	{
+		ScriptViewModel->SetScriptName(InText);
+	}
+
+	bool IsActive() const
+	{
+		return IsHovered();
+	}
+
+	EVisibility GetUnappliedChangesVisibility() const
+	{
+		return ScriptViewModel->HasUnappliedChanges() ? EVisibility::Visible : EVisibility::Collapsed;
+	}
+
+private:
+	UNiagaraScratchPadViewModel* ScratchPadViewModel;
+	TSharedPtr<FNiagaraScratchPadScriptViewModel> ScriptViewModel;
+	TSharedPtr<FNiagaraScratchPadCommandContext> CommandContext;
+	TAttribute<bool> IsSelected;
+	TSharedPtr<SInlineEditableTextBlock> NameEditableText;
+};
+
+class SNiagaraScratchPadScriptSelector : public SCompoundWidget
+{
+	SLATE_BEGIN_ARGS(SNiagaraScratchPadScriptSelector) {}
+	SLATE_END_ARGS()
+
+	void Construct(const FArguments& InArgs, UNiagaraScratchPadViewModel* InViewModel, TSharedPtr<FNiagaraScratchPadCommandContext> InCommandContext)
+	{
+		ViewModel = InViewModel;
+		CommandContext = InCommandContext;
+		ViewModel->OnScriptViewModelsChanged().AddSP(this, &SNiagaraScratchPadScriptSelector::ScriptViewModelsChanged);
+		ViewModel->OnActiveScriptChanged().AddSP(this, &SNiagaraScratchPadScriptSelector::ActiveScriptChanged);
+		bIsUpdatingSelection = false;
+
+		ChildSlot
+		[
+			SAssignNew(ScriptSelector, SNiagaraScriptViewModelSelector)
+			.ClickActivateMode(EItemSelectorClickActivateMode::DoubleClick)
+			.CategoryRowStyle(FNiagaraEditorWidgetsStyle::Get(), "NiagaraEditor.ScratchPad.CategoryRow")
+			.ClearSelectionOnClick(false)
+			.Items(ViewModel->GetScriptViewModels())
+			.DefaultCategories(ViewModel->GetAvailableUsages())
+			.OnGetCategoriesForItem(this, &SNiagaraScratchPadScriptSelector::OnGetCategoriesForItem)
+			.OnCompareCategoriesForEquality(this, &SNiagaraScratchPadScriptSelector::OnCompareCategoriesForEquality)
+			.OnCompareCategoriesForSorting(this, &SNiagaraScratchPadScriptSelector::OnCompareCategoriesForSorting)
+			.OnCompareItemsForEquality(this, &SNiagaraScratchPadScriptSelector::OnCompareItemsForEquality)
+			.OnCompareItemsForSorting(this, &SNiagaraScratchPadScriptSelector::OnCompareItemsForSorting)
+			.OnDoesItemMatchFilterText(this, &SNiagaraScratchPadScriptSelector::OnDoesItemMatchFilterText)
+			.OnGenerateWidgetForCategory(this, &SNiagaraScratchPadScriptSelector::OnGenerateWidgetForCategory)
+			.OnGenerateWidgetForItem(this, &SNiagaraScratchPadScriptSelector::OnGenerateWidgetForItem)
+			.OnItemActivated(this, &SNiagaraScratchPadScriptSelector::OnScriptActivated)
+			.OnItemsDragged(this, &SNiagaraScratchPadScriptSelector::OnScriptDragged)
+			.OnSelectionChanged(this, &SNiagaraScratchPadScriptSelector::OnSelectionChanged)
+
+		];
+
+		if (ViewModel->GetActiveScriptViewModel().IsValid())
+		{
+			TArray<TSharedRef<FNiagaraScratchPadScriptViewModel>> SelectedViewModels;
+			SelectedViewModels.Add(ViewModel->GetActiveScriptViewModel().ToSharedRef());
+			ScriptSelector->SetSelectedItems(SelectedViewModels);
+		}
+	}
+
+	virtual FReply OnKeyDown(const FGeometry& MyGeometry, const FKeyEvent& InKeyEvent) override
+	{
+		if (CommandContext->ProcessCommandBindings(InKeyEvent))
+		{
+			return FReply::Handled();
+		}
+		return SCompoundWidget::OnKeyDown(MyGeometry, InKeyEvent);
+	}
+
+	virtual FReply OnMouseButtonDown(const FGeometry& MyGeometry, const FPointerEvent& MouseEvent) override
+	{
+		if (MouseEvent.GetEffectingButton() == EKeys::RightMouseButton)
+		{
+			return FReply::Handled()
+				.CaptureMouse(SharedThis(this));
+		}
+		return SCompoundWidget::OnMouseButtonDown(MyGeometry, MouseEvent);
+	}
+
+	virtual FReply OnMouseButtonUp(const FGeometry& MyGeometry, const FPointerEvent& MouseEvent) override
+	{
+		if (MouseEvent.GetEffectingButton() == EKeys::RightMouseButton)
+		{
+			FMenuBuilder MenuBuilder(true, CommandContext->GetCommands());
+			CommandContext->AddMenuItems(MenuBuilder);
+
+			FWidgetPath WidgetPath = MouseEvent.GetEventPath() != nullptr ? *MouseEvent.GetEventPath() : FWidgetPath();
+			FSlateApplication::Get().PushMenu(AsShared(), WidgetPath, MenuBuilder.MakeWidget(), MouseEvent.GetScreenSpacePosition(), FPopupTransitionEffect(FPopupTransitionEffect::ContextMenu));
+			return FReply::Handled().ReleaseMouseCapture();
+		}
+		return SCompoundWidget::OnMouseButtonUp(MyGeometry, MouseEvent);
+	}
+
+private:
+	void ScriptViewModelsChanged()
+	{
+		if (ScriptSelector.IsValid())
+		{
+			ScriptSelector->RefreshItemsAndDefaultCategories(ViewModel->GetScriptViewModels(), ViewModel->GetAvailableUsages());
+		}
+	}
+
+	void ActiveScriptChanged()
+	{
+		if (bIsUpdatingSelection == false)
+		{
+			TGuardValue<bool> UpdateGuard(bIsUpdatingSelection, true);
+			TSharedPtr<FNiagaraScratchPadScriptViewModel> ActiveScriptViewModel = ViewModel->GetActiveScriptViewModel();
+			if (ActiveScriptViewModel.IsValid())
+			{
+				TArray<TSharedRef<FNiagaraScratchPadScriptViewModel>> SelectedItems;
+				SelectedItems.Add(ActiveScriptViewModel.ToSharedRef());
+				ScriptSelector->SetSelectedItems(SelectedItems);
+			}
+			else
+			{
+				ScriptSelector->ClearSelectedItems();
+			}
+		}
+	}
+
+	TArray<ENiagaraScriptUsage> OnGetCategoriesForItem(const TSharedRef<FNiagaraScratchPadScriptViewModel>& Item)
+	{
+		TArray<ENiagaraScriptUsage> Categories;
+		Categories.Add(Item->GetScripts()[0].Script->GetUsage());
+		return Categories;
+	}
+
+	bool OnCompareCategoriesForEquality(const ENiagaraScriptUsage& CategoryA, const ENiagaraScriptUsage& CategoryB) const
+	{
+		return CategoryA == CategoryB;
+	}
+
+	bool OnCompareCategoriesForSorting(const ENiagaraScriptUsage& CategoryA, const ENiagaraScriptUsage& CategoryB) const
+	{
+		return ((int32)CategoryA) < ((int32)CategoryB);
+	}
+
+	bool OnCompareItemsForEquality(const TSharedRef<FNiagaraScratchPadScriptViewModel>& ItemA, const TSharedRef<FNiagaraScratchPadScriptViewModel>& ItemB) const
+	{
+		return ItemA == ItemB;
+	}
+
+	bool OnCompareItemsForSorting(const TSharedRef<FNiagaraScratchPadScriptViewModel>& ItemA, const TSharedRef<FNiagaraScratchPadScriptViewModel>& ItemB) const
+	{
+		return ItemA->GetDisplayName().CompareTo(ItemB->GetDisplayName()) < 0;
+	}
+
+	bool OnDoesItemMatchFilterText(const FText& FilterText, const TSharedRef<FNiagaraScratchPadScriptViewModel>& Item)
+	{
+		return Item->GetDisplayName().ToString().Find(FilterText.ToString(), ESearchCase::IgnoreCase) != INDEX_NONE;
+	}
+
+	TSharedRef<SWidget> OnGenerateWidgetForCategory(const ENiagaraScriptUsage& Category)
+	{
+		return SNew(SHorizontalBox)
+		+ SHorizontalBox::Slot()
+		.VAlign(VAlign_Center)
+		.Padding(3, 0, 0, 0)
+		[
+			SNew(STextBlock)
+			.TextStyle(FNiagaraEditorWidgetsStyle::Get(), "NiagaraEditor.Stack.GroupText")
+			.Text(ViewModel->GetDisplayNameForUsage(Category))
+			.ColorAndOpacity(FSlateColor::UseForeground())
+		]
+		+ SHorizontalBox::Slot()
+		.AutoWidth()
+		.Padding(0.0f, 2.0f, 3.0f, 2.0f)
+		[
+			SNew(SButton)
+			.ButtonStyle(FAppStyle::Get(), "SimpleButton")
+			.OnClicked(this, &SNiagaraScratchPadScriptSelector::ScriptSelectorAddButtonClicked, Category)
+			.ContentPadding(FMargin(3.0f, 2.0f, 2.0f, 2.0f))
+			.Content()
+			[
+				SNew(SImage)
+				.Image(FAppStyle::Get().GetBrush("Icons.PlusCircle"))
+				.ColorAndOpacity(FSlateColor::UseForeground())
+			]
+		];
+	}
+
+	TSharedRef<SWidget> OnGenerateWidgetForItem(const TSharedRef<FNiagaraScratchPadScriptViewModel>& Item)
+	{
+		return SNew(SNiagaraScratchPadScriptRow, ViewModel, Item, CommandContext)
+			.IsSelected(this, &SNiagaraScratchPadScriptSelector::GetItemIsSelected, TWeakPtr<FNiagaraScratchPadScriptViewModel>(Item));
+	}
+
+	void OnScriptActivated(const TSharedRef<FNiagaraScratchPadScriptViewModel>& ActivatedScript)
+	{
+		if (bIsUpdatingSelection == false)
+		{
+			TGuardValue<bool> UpdateGuard(bIsUpdatingSelection, true);
+			ViewModel->SetActiveScriptViewModel(ActivatedScript);
+			ViewModel->OpenEditorForActive();
+		}
+	}
+
+	FReply OnScriptDragged(const TArray<TSharedRef<FNiagaraScratchPadScriptViewModel>>& DraggedItems, const FPointerEvent& MouseEvent) const
+	{
+		if (DraggedItems.Num() == 1)
+		{
+			TSharedPtr< FNiagaraScratchPadScriptViewModel> ScriptViewModel = DraggedItems[0];
+			if (ScriptViewModel.IsValid())
+			{
+				TSharedRef<FNiagaraScriptDragOperation> DragDropOp = MakeShared<FNiagaraScriptDragOperation>(ScriptViewModel->GetOriginalScript(), ScriptViewModel->GetOriginalScript()->VersionToOpenInEditor, ScriptViewModel->GetDisplayName());
+				DragDropOp->CurrentHoverText = ScriptViewModel->GetDisplayName();
+				DragDropOp->CurrentIconBrush = FNiagaraEditorStyle::Get().GetBrush("Tab.ScratchPad");
+				DragDropOp->SetupDefaults();
+				DragDropOp->Construct();
+				return FReply::Handled().BeginDragDrop(DragDropOp);
+			}
+		}
+		return FReply::Unhandled();
+	}
+
+
+	void OnSelectionChanged()
+	{
+		if (bIsUpdatingSelection == false)
+		{
+			TGuardValue<bool> UpdateGuard(bIsUpdatingSelection, true);
+			TArray<TSharedRef<FNiagaraScratchPadScriptViewModel>> SelectedScripts = ScriptSelector->GetSelectedItems();
+			if (SelectedScripts.Num() == 0)
+			{
+				ViewModel->ResetActiveScriptViewModel();
+			}
+			else if (SelectedScripts.Num())
+			{
+				ViewModel->SetActiveScriptViewModel(SelectedScripts[0]);
+			}
+		}
+	}
+
+	FReply ScriptSelectorAddButtonClicked(ENiagaraScriptUsage Usage)
+	{
+		TSharedPtr<FNiagaraScratchPadScriptViewModel> NewScriptViewModel = ViewModel->CreateNewScript(Usage, ENiagaraScriptUsage::ParticleUpdateScript, FNiagaraTypeDefinition());
+		if (NewScriptViewModel.IsValid())
+		{
+			ViewModel->SetActiveScriptViewModel(NewScriptViewModel.ToSharedRef());
+			NewScriptViewModel->SetIsPendingRename(true);
+		}
+		return FReply::Handled();
+	}
+
+	bool GetItemIsSelected(TWeakPtr<FNiagaraScratchPadScriptViewModel> ItemWeak) const
+	{
+		TSharedPtr<FNiagaraScratchPadScriptViewModel> Item = ItemWeak.Pin();
+		return Item.IsValid() && ViewModel->GetActiveScriptViewModel() == Item;
+	}
+
+private:
+	TSharedPtr<SNiagaraScriptViewModelSelector> ScriptSelector;
+	UNiagaraScratchPadViewModel* ViewModel;
+	TSharedPtr<FNiagaraScratchPadCommandContext> CommandContext;
+	bool bIsUpdatingSelection;
+};
+
+
+class SNiagaraScratchPadScriptEditorList : public SCompoundWidget
+{
+	SLATE_BEGIN_ARGS(SNiagaraScratchPadScriptEditorList) {}
+
+	SLATE_END_ARGS();
+
+	void Construct(const FArguments& InArgs, UNiagaraScratchPadViewModel* InViewModel)
+	{
+		ViewModel = InViewModel;
+		ViewModel->OnScriptViewModelsChanged().AddSP(this, &SNiagaraScratchPadScriptEditorList::ScriptViewModelsChanged);
+		ViewModel->OnEditScriptViewModelsChanged().AddSP(this, &SNiagaraScratchPadScriptEditorList::EditScriptViewModelsChanged);
+		ViewModel->OnActiveScriptChanged().AddSP(this, &SNiagaraScratchPadScriptEditorList::ActiveScriptChanged);
+		bIsUpdatingSelection = false;
+
+		ChildSlot
+		[
+			SAssignNew(ContentBorder, SBorder)
+			.BorderImage(FAppStyle::Get().GetBrush("ToolPanel.DarkGroupBorder"))
+		];
+		UpdateContentFromEditScriptViewModels();
+	}
+
+private:
+	void ScriptViewModelsChanged()
+	{
+		ScriptViewModelWidgetPairs.RemoveAll([](const FScriptViewModelWidgetPair& ScriptViewModelWidgetPair)
+		{
+			return ScriptViewModelWidgetPair.ViewModel.IsValid() == false || ScriptViewModelWidgetPair.Widget.IsValid() == false;
+		});
+	}
+
+	void EditScriptViewModelsChanged()
+	{
+		UpdateContentFromEditScriptViewModels();
+	}
+
+	TSharedRef<SWidget> FindScriptEditor(TSharedRef<FNiagaraScratchPadScriptViewModel> ScriptViewModel)
+	{
+		FScriptViewModelWidgetPair* ExistingPair = ScriptViewModelWidgetPairs.FindByPredicate([ScriptViewModel](FScriptViewModelWidgetPair& ScriptViewModelWidgetPair)
+		{ 
+			return ScriptViewModelWidgetPair.ViewModel == ScriptViewModel && ScriptViewModelWidgetPair.Widget.IsValid();
+		});
+
+		if (ExistingPair != nullptr)
+		{
+			return ExistingPair->Widget.ToSharedRef();
+		}
+		else
+		{
+			return SNullWidget::NullWidget;
+		}
+	}
+
+	void UpdateContentFromEditScriptViewModels()
+	{
+		TSharedPtr<SWidget> NewContent;
+		if (ViewModel->GetEditScriptViewModels().Num() == 0)
+		{
+			NewContent = SNew(SBox)
+				.HAlign(HAlign_Center)
+				.Padding(FMargin(0.0f, 10.0f, 0.0f, 0.0f))
+				[
+					SNew(STextBlock)
+					.Text(LOCTEXT("NoScriptToEdit", "No script selected"))
+				];
+			ScriptEditorList.Reset();
+		}
+		else if(ViewModel->GetEditScriptViewModels().Num() == 1)
+		{
+			NewContent = FindScriptEditor(ViewModel->GetEditScriptViewModels()[0]);
+			ScriptEditorList.Reset();
+		}
+		else 
+		{
+			if (ScriptEditorList.IsValid())
+			{
+				ScriptEditorList->RequestListRefresh();
+			}
+			else 
+			{
+				ScriptEditorList = SNew(SListView<TSharedRef<FNiagaraScratchPadScriptViewModel>>)
+					.ListItemsSource(&ViewModel->GetEditScriptViewModels())
+					.OnGenerateRow(this, &SNiagaraScratchPadScriptEditorList::OnGenerateScriptEditorRow)
+					.OnSelectionChanged(this, &SNiagaraScratchPadScriptEditorList::OnSelectionChanged);
+			}
+			NewContent = ScriptEditorList;
+		}
+
+		ContentBorder->SetContent(NewContent.ToSharedRef());
+	}
+
+	void ActiveScriptChanged()
+	{
+		if (ScriptEditorList.IsValid() && bIsUpdatingSelection == false)
+		{
+			TGuardValue<bool> UpdateGuard(bIsUpdatingSelection, true);
+			TSharedPtr<FNiagaraScratchPadScriptViewModel> ActiveScriptViewModel = ViewModel->GetActiveScriptViewModel();
+			if (ActiveScriptViewModel.IsValid())
+			{
+				ScriptEditorList->SetSelection(ActiveScriptViewModel.ToSharedRef());
+			}
+			else
+			{
+				ScriptEditorList->ClearSelection();
+			}
+		}
+	}
+
+	TSharedRef<ITableRow> OnGenerateScriptEditorRow(TSharedRef<FNiagaraScratchPadScriptViewModel> Item, const TSharedRef<STableViewBase>& OwnerTable)
+	{
+		return SNew(STableRow<TSharedRef<FNiagaraScratchPadScriptViewModel>>, OwnerTable)
+		[
+			SNew(SVerticalResizeBox)
+			.ContentHeight(Item, &FNiagaraScratchPadScriptViewModel::GetEditorHeight)
+			.ContentHeightChanged(Item, &FNiagaraScratchPadScriptViewModel::SetEditorHeight)
+			[
+				FindScriptEditor(Item)
+			]
+		];
+	}
+
+	void OnSelectionChanged(TSharedPtr<FNiagaraScratchPadScriptViewModel> InNewSelection, ESelectInfo::Type SelectInfo)
+	{
+		if (bIsUpdatingSelection == false)
+		{
+			TGuardValue<bool> UpdateGuard(bIsUpdatingSelection, true);
+			TArray<TSharedRef<FNiagaraScratchPadScriptViewModel>> SelectedScripts;
+			ScriptEditorList->GetSelectedItems(SelectedScripts);
+			if (SelectedScripts.Num() == 0)
+			{
+				ViewModel->ResetActiveScriptViewModel();
+			}
+			else if (SelectedScripts.Num())
+			{
+				ViewModel->SetActiveScriptViewModel(SelectedScripts[0]);
+			}
+		}
+	}
+
+private:
+	struct FScriptViewModelWidgetPair
+	{
+		TWeakPtr<FNiagaraScratchPadScriptViewModel> ViewModel;
+		TSharedPtr<SWidget> Widget;
+	};
+
+	UNiagaraScratchPadViewModel* ViewModel;
+	TSharedPtr<SBorder> ContentBorder;
+	TSharedPtr<SListView<TSharedRef<FNiagaraScratchPadScriptViewModel>>> ScriptEditorList;
+	TArray<FScriptViewModelWidgetPair> ScriptViewModelWidgetPairs;
+	bool bIsUpdatingSelection;
+};
+
+class SNiagaraScratchPadParameterPanel : public SCompoundWidget
+{
+	SLATE_BEGIN_ARGS(SNiagaraScratchPadParameterPanel) {}
+	SLATE_END_ARGS();
+
+	void Construct(const FArguments& InArgs, UNiagaraScratchPadViewModel* InViewModel)
+	{
+		ViewModel = InViewModel;
+		ViewModel->OnActiveScriptChanged().AddSP(this, &SNiagaraScratchPadParameterPanel::ActiveScriptChanged);
+		bool bForce = true;
+		UpdateContent(bForce);
+	}
+
+	virtual FReply OnKeyDown(const FGeometry& MyGeometry, const FKeyEvent& InKeyEvent) override
+	{
+		TSharedPtr<FNiagaraScratchPadScriptViewModel> ScriptViewModel = ScriptViewModelWeak.Pin();
+		if (ScriptViewModel.IsValid() && ScriptViewModel->GetParameterPanelCommands()->ProcessCommandBindings(InKeyEvent))
+		{
+			return FReply::Handled();
+		}
+		return SCompoundWidget::OnKeyDown(MyGeometry, InKeyEvent);
+	}
+
+private:
+	void UpdateContent(bool bForce)
+	{
+		TSharedPtr<FNiagaraScratchPadScriptViewModel> OldScriptViewModel = ScriptViewModelWeak.Pin();
+		TSharedPtr<FNiagaraScratchPadScriptViewModel> NewScriptViewModel = ViewModel->GetActiveScriptViewModel();
+		if (NewScriptViewModel != OldScriptViewModel || bForce)
+		{
+			ScriptViewModelWeak = NewScriptViewModel;
+			if (NewScriptViewModel.IsValid())
+			{
+				TSharedPtr<SWidget> ParameterPanelWidget;
+				if (NewScriptViewModel->GetParameterPanelViewModel().IsValid())
+				{
+					ParameterPanelWidget = SNew(SNiagaraParameterPanel, NewScriptViewModel->GetParameterPanelViewModel(), NewScriptViewModel->GetParameterPanelCommands());
+				}
+				else
+				{
+					TSharedRef<FNiagaraObjectSelection> ScriptSelection = MakeShared<FNiagaraObjectSelection>();
+					const FVersionedNiagaraScript& EditScript = NewScriptViewModel->GetEditScript();
+					ScriptSelection->SetSelectedObject(EditScript.Script, &EditScript.Version);
+					TArray<TSharedRef<FNiagaraObjectSelection>> ParameterSelections;
+					ParameterSelections.Add(ScriptSelection);
+					ParameterSelections.Add(NewScriptViewModel->GetVariableSelection());
+				}
+				ChildSlot
+				[
+					SNew(SVerticalBox)
+					+ SVerticalBox::Slot()
+					.AutoHeight()
+					.Padding(2.0f, 2.0f, 2.0f, 5.0f)
+					[
+						SNew(STextBlock)
+						.TextStyle(FNiagaraEditorWidgetsStyle::Get(), "NiagaraEditor.ScratchPad.SubSectionHeaderText")
+						.Text(this, &SNiagaraScratchPadParameterPanel::GetSubHeaderText)
+					]
+					+ SVerticalBox::Slot()
+					[
+						ParameterPanelWidget.ToSharedRef()
+					]
+				];
+			}
+			else
+			{
+				ChildSlot
+				[
+					SNew(SBox)
+					.HAlign(HAlign_Center)
+					.Padding(FMargin(0.0f, 10.0f, 0.0f, 0.0f))
+					[
+						SNew(STextBlock)
+						.Text(LOCTEXT("NoScriptForParameters", "No script selected"))
+					]
+				];
+			}
+		}
+	}
+
+	void ActiveScriptChanged()
+	{
+		bool bForce = false;
+		UpdateContent(bForce);
+	}
+
+	FText GetSubHeaderText() const
+	{
+		TSharedPtr<FNiagaraScratchPadScriptViewModel> ScriptViewModel = ScriptViewModelWeak.Pin();
+		return ScriptViewModel.IsValid() ? ScriptViewModel->GetDisplayName() : FText();
+	}
+
+private:
+	UNiagaraScratchPadViewModel* ViewModel;
+	TWeakPtr<FNiagaraScratchPadScriptViewModel> ScriptViewModelWeak;
+};
+
+class SNiagaraScratchPadSectionBox : public SCompoundWidget
+{
+	SLATE_BEGIN_ARGS(SNiagaraScratchPadSectionBox) {}
+		SLATE_ATTRIBUTE(FText, HeaderText)
+		SLATE_DEFAULT_SLOT(FArguments, Content);
+	SLATE_END_ARGS();
+
+	void Construct(const FArguments& InArgs)
+	{
+		ChildSlot
+		[
+			SNew(SVerticalBox)
+			+ SVerticalBox::Slot()
+			.AutoHeight()
+			.Padding(1.0f, 5.0f, 0.0f, 4.0f)
+			[
+				SNew(STextBlock)
+				.TextStyle(FNiagaraEditorWidgetsStyle::Get(), "NiagaraEditor.Stack.GroupText")
+				.Text(InArgs._HeaderText)
+				.ColorAndOpacity(FSlateColor::UseForeground())
+			]
+			+ SVerticalBox::Slot()
+			.Padding(5.0f, 0.0f, 3.0f, 5.0f)
+			[
+				InArgs._Content.Widget
+			]
+		];
+	}
+};
+
+void SNiagaraScratchPadScriptManager::Construct(const FArguments& InArgs, UNiagaraScratchPadViewModel* InViewModel)
+{
+	ViewModel = InViewModel;
+	ViewModel->GetObjectSelection()->OnSelectedObjectsChanged().AddSP(this, &SNiagaraScratchPadScriptManager::ObjectSelectionChanged);
+
+	CommandContext = MakeShared<FNiagaraScratchPadCommandContext>(InViewModel);
+
+	ChildSlot
+	[
+		SNew(SDynamicLayoutBox)
+		.GenerateNamedWidget_Lambda([this](FName InWidgetName)
+		{
+			if (InWidgetName == ScriptSelectorName)
+			{
+				return ConstructScriptSelector();
+			}
+			else if (InWidgetName == ScriptParameterPanelName)
+			{
+				return ConstructParameterPanel();
+			}
+			else if (InWidgetName == SelectionEditorName)
+			{
+				return ConstructSelectionEditor();
+			}
+			else
+			{
+				return SNullWidget::NullWidget;
+			}
+		})
+		.GenerateNamedLayout_Lambda([this](FName InLayoutName, const SDynamicLayoutBox::FNamedWidgetProvider& InNamedWidgetProvider)
+		{
+				TSharedPtr<SWidget> Layout =
+					SNew(SSplitter)
+					.Style(FAppStyle::Get(), "SplitterDark")
+					.Orientation(Orient_Vertical)
+					.PhysicalSplitterHandleSize(4.0f)
+					.HitDetectionSplitterHandleSize(6.0f)
+					+ SSplitter::Slot()
+					.Value(0.3f)
+					[
+						InNamedWidgetProvider.GetNamedWidget(ScriptSelectorName)
+					];
+			
+			return Layout.ToSharedRef();
+		})
+		.ChooseLayout_Lambda([this]() 
+		{ 
+				return NarrowLayoutName;
+		})
+	];
+}
+
+void SNiagaraScratchPadScriptManager::ObjectSelectionChanged()
+{
+	int32 SelectionCount = ViewModel->GetObjectSelection()->GetSelectedObjects().Num();
+	if (SelectionCount == 0)
+	{
+		ObjectSelectionSubHeaderText = FText();
+	}
+	else if (SelectionCount == 1)
+	{
+		UObject* SelectedObject = ViewModel->GetObjectSelection()->GetSelectedObjects().Array()[0];
+		if (SelectedObject->IsA<UEdGraphNode>())
+		{
+			UEdGraphNode* SelectedGraphNode = CastChecked<UEdGraphNode>(SelectedObject);
+			ObjectSelectionSubHeaderText = SelectedGraphNode->GetNodeTitle(ENodeTitleType::ListView);
+		}
+		else if (SelectedObject->IsA<UNiagaraScript>())
+		{
+			UNiagaraScript* SelectedScript = CastChecked<UNiagaraScript>(SelectedObject);
+			TSharedPtr<FNiagaraScratchPadScriptViewModel> ScratchPadScriptViewModel = ViewModel->GetViewModelForEditScript(SelectedScript);
+			if (ScratchPadScriptViewModel.IsValid())
+			{
+				ObjectSelectionSubHeaderText = ScratchPadScriptViewModel->GetDisplayName();
+			}
+			else
+			{
+				ObjectSelectionSubHeaderText = FText::FromString(SelectedScript->GetName());
+			}
+		}
+		else
+		{
+			ObjectSelectionSubHeaderText = FText::FromString(SelectedObject->GetName());
+		}
+	}
+	else
+	{
+		ObjectSelectionSubHeaderText = FText::Format(LOCTEXT("MultipleSelectionFormat", "{0} Objects Selected..."),
+			FText::AsNumber(SelectionCount));
+	}
+}
+
+TSharedRef<SWidget> SNiagaraScratchPadScriptManager::ConstructScriptSelector()
+{
+	return SNew(SNiagaraScratchPadSectionBox)
+	.HeaderText(LOCTEXT("ScriptManager", "Scratch Script Manager"))
+	[
+		SNew(SNiagaraScratchPadScriptSelector, ViewModel.Get(), CommandContext)
+	];
+}
+
+TSharedRef<SWidget> SNiagaraScratchPadScriptManager::ConstructParameterPanel()
+{
+	return SNew(SNiagaraScratchPadSectionBox)
+	.HeaderText(LOCTEXT("ScratchScriptParameters", "Scratch Script Parameters"))
+	[
+		SNew(SNiagaraScratchPadParameterPanel, ViewModel.Get())
+	];
+}
+
+TSharedRef<SWidget> SNiagaraScratchPadScriptManager::ConstructSelectionEditor()
+{
+	return SNew(SNiagaraScratchPadSectionBox)
+	.HeaderText(LOCTEXT("ScratchPadSelection", "Scratch Pad Selection"))
+	[
+		SNew(SVerticalBox)
+		+ SVerticalBox::Slot()
+		.AutoHeight()
+		.Padding(2.0f, 2.0f, 2.0f, 5.0f)
+		[
+			SNew(STextBlock)
+			.TextStyle(FNiagaraEditorWidgetsStyle::Get(), "NiagaraEditor.ScratchPad.SubSectionHeaderText")
+			.Visibility(this, &SNiagaraScratchPadScriptManager::GetObjectSelectionSubHeaderTextVisibility)
+			.Text(this, &SNiagaraScratchPadScriptManager::GetObjectSelectionSubHeaderText)
+		]
+		+ SVerticalBox::Slot()
+		.HAlign(HAlign_Center)
+		.AutoHeight()
+		.Padding(FMargin(0.0f, 10.0f, 0.0f, 0.0f))
+		[
+			SNew(STextBlock)
+			.Visibility(this, &SNiagaraScratchPadScriptManager::GetObjectSelectionNoSelectionTextVisibility)
+			.Text(LOCTEXT("NoObjectSelection", "No object selected"))
+		]
+		+ SVerticalBox::Slot()
+		[
+			SNew(SNiagaraSelectedObjectsDetails, ViewModel->GetObjectSelection())
+		]
+	];
+}
+
+EVisibility SNiagaraScratchPadScriptManager::GetObjectSelectionSubHeaderTextVisibility() const
+{
+	return ViewModel->GetObjectSelection()->GetSelectedObjects().Num() != 0 ? EVisibility::Visible : EVisibility::Collapsed;
+}
+
+FText SNiagaraScratchPadScriptManager::GetObjectSelectionSubHeaderText() const
+{
+	return ObjectSelectionSubHeaderText;
+}
+
+EVisibility SNiagaraScratchPadScriptManager::GetObjectSelectionNoSelectionTextVisibility() const
+{
+	return ViewModel->GetObjectSelection()->GetSelectedObjects().Num() == 0 ? EVisibility::Visible : EVisibility::Collapsed;
+}
+
+#undef LOCTEXT_NAMESPACE
diff --git a/Engine/Plugins/FX/Niagara/Source/NiagaraEditorWidgets/Private/SNiagaraScratchPadScriptManager.h b/Engine/Plugins/FX/Niagara/Source/NiagaraEditorWidgets/Private/SNiagaraScratchPadScriptManager.h
new file mode 100644
--- /dev/null
+++ b/Engine/Plugins/FX/Niagara/Source/NiagaraEditorWidgets/Private/SNiagaraScratchPadScriptManager.h
@@ -0,0 +1,51 @@
+// Copyright Epic Games, Inc. All Rights Reserved.
+
+#pragma once
+
+#include "NiagaraCommon.h"
+#include "Widgets/SItemSelector.h"
+#include "Widgets/SCompoundWidget.h"
+
+class UNiagaraScratchPadViewModel;
+class FNiagaraScratchPadScriptViewModel;
+class FNiagaraScratchPadCommandContext;
+
+typedef SItemSelector<ENiagaraScriptUsage, TSharedRef<FNiagaraScratchPadScriptViewModel>> SNiagaraScriptViewModelSelector;
+
+class SNiagaraScratchPadScriptManager : public SCompoundWidget
+{
+	SLATE_BEGIN_ARGS(SNiagaraScratchPadScriptManager) {}
+	SLATE_END_ARGS()
+
+	void Construct(const FArguments& InArgs, UNiagaraScratchPadViewModel* InViewModel);
+
+private:
+	void ObjectSelectionChanged();
+
+	TSharedRef<SWidget> ConstructScriptSelector();
+
+	TSharedRef<SWidget> ConstructParameterPanel();
+
+	TSharedRef<SWidget> ConstructScriptEditor();
+
+	TSharedRef<SWidget> ConstructSelectionEditor();
+
+	EVisibility GetObjectSelectionSubHeaderTextVisibility() const;
+
+	FText GetObjectSelectionSubHeaderText() const;
+
+	EVisibility GetObjectSelectionNoSelectionTextVisibility() const;
+	static FName ScriptSelectorName;
+	static FName ScriptParameterPanelName;
+	static FName ScriptEditorName;
+	static FName SelectionEditorName;
+	static FName WideLayoutName;
+	static FName NarrowLayoutName;
+private:
+	FText ObjectSelectionSubHeaderText;
+
+	TWeakObjectPtr<UNiagaraScratchPadViewModel> ViewModel;
+
+	TSharedPtr<FNiagaraScratchPadCommandContext> CommandContext;
+
+};
\ No newline at end of file
diff --git a/Engine/Plugins/FX/Niagara/Source/NiagaraEditorWidgets/Public/NiagaraEditorWidgetsModule.h b/Engine/Plugins/FX/Niagara/Source/NiagaraEditorWidgets/Public/NiagaraEditorWidgetsModule.h
--- a/Engine/Plugins/FX/Niagara/Source/NiagaraEditorWidgets/Public/NiagaraEditorWidgetsModule.h
+++ b/Engine/Plugins/FX/Niagara/Source/NiagaraEditorWidgets/Public/NiagaraEditorWidgetsModule.h
@@ -54,7 +54,7 @@ private:
 		virtual TSharedRef<SWidget> CreateStackView(UNiagaraStackViewModel& StackViewModel) const override;
 		virtual TSharedRef<SWidget> CreateSystemOverview(TSharedRef<FNiagaraSystemViewModel> SystemViewModel) const override;
 		virtual TSharedRef<SWidget> CreateStackIssueIcon(UNiagaraStackViewModel& StackViewModel, UNiagaraStackEntry& StackEntry) const override;
-		virtual TSharedRef<SWidget> CreateScriptScratchPad(UNiagaraScratchPadViewModel& ScriptScratchPadViewModel) const override;
+		virtual TSharedRef<SWidget> CreateScriptScratchPadManager(UNiagaraScratchPadViewModel& ScriptScratchPadViewModel) const override;
 		virtual TSharedRef<SWidget> CreateCurveOverview(TSharedRef<FNiagaraSystemViewModel> SystemViewModel) const override;
 		virtual FLinearColor GetColorForExecutionCategory(FName ExecutionCategory) const override;
 	};
