Fixed static analysis warnings

diff --git a/Engine/Plugins/FX/Niagara/Source/NiagaraEditor/Private/ViewModels/Stack/NiagaraStackGraphUtilities.cpp b/Engine/Plugins/FX/Niagara/Source/NiagaraEditor/Private/ViewModels/Stack/NiagaraStackGraphUtilities.cpp
--- a/Engine/Plugins/FX/Niagara/Source/NiagaraEditor/Private/ViewModels/Stack/NiagaraStackGraphUtilities.cpp
+++ b/Engine/Plugins/FX/Niagara/Source/NiagaraEditor/Private/ViewModels/Stack/NiagaraStackGraphUtilities.cpp
@@ -332,11 +332,12 @@ TArray<FName> FNiagaraStackGraphUtilities::StackContextResolution(FVersionedNiag
 			break;
 		case ENiagaraScriptUsage::ParticleSimulationStageScript:
 		{
-			if (OwningEmitter.Emitter)
+			if (FVersionedNiagaraEmitterData* EmitterData = OwningEmitter.GetEmitterData())
 			{
-				UNiagaraSimulationStageBase* Base = OwningEmitter.GetEmitterData()->GetSimulationStageById(OutputNodeInChain->GetUsageId());
-				if (Base)
+				if (UNiagaraSimulationStageBase* Base = EmitterData->GetSimulationStageById(OutputNodeInChain->GetUsageId()))
+				{
 					StageName = Base->GetStackContextReplacementName();
+				}
 			}
 			
 			if (StageName == NAME_None)
@@ -378,15 +379,17 @@ TArray<FName> FNiagaraStackGraphUtilities::StackContextResolution(FVersionedNiag
 void FNiagaraStackGraphUtilities::BuildParameterMapHistoryWithStackContextResolution(FVersionedNiagaraEmitter OwningEmitter, UNiagaraNodeOutput* OutputNodeInChain, UNiagaraNode* NodeToVisit, FNiagaraParameterMapHistoryBuilder& Builder, bool bRecursive /*= true*/, bool bFilterForCompilation /*= true*/)
 {
 	bool bSetUsage = false;
-	if (OwningEmitter.Emitter && OutputNodeInChain)
+	FVersionedNiagaraEmitterData* EmitterData = OwningEmitter.GetEmitterData();
+	if (EmitterData && OutputNodeInChain)
 	{
 		ENiagaraScriptUsage Usage = OutputNodeInChain->GetUsage();
 		FName StageName;
 		if (Usage == ENiagaraScriptUsage::ParticleSimulationStageScript)
 		{
-			UNiagaraSimulationStageBase* Base = OwningEmitter.GetEmitterData()->GetSimulationStageById(OutputNodeInChain->GetUsageId());
-			if (Base)
+			if (UNiagaraSimulationStageBase* Base = EmitterData->GetSimulationStageById(OutputNodeInChain->GetUsageId()))
+			{
 				StageName = Base->GetStackContextReplacementName();
+			}
 		}
 		Builder.BeginUsage(Usage, StageName);
 		bSetUsage = true;
