Fixed static analysis warning

diff --git a/Engine/Plugins/FX/Niagara/Source/Niagara/Private/NiagaraEmitterInstance.cpp b/Engine/Plugins/FX/Niagara/Source/Niagara/Private/NiagaraEmitterInstance.cpp
--- a/Engine/Plugins/FX/Niagara/Source/Niagara/Private/NiagaraEmitterInstance.cpp
+++ b/Engine/Plugins/FX/Niagara/Source/Niagara/Private/NiagaraEmitterInstance.cpp
@@ -1018,9 +1018,6 @@ void FNiagaraEmitterInstance::CalculateFixedBounds(const FTransform& ToWorldSpac
 void FNiagaraEmitterInstance::PostTick()
 {
 	SCOPE_CYCLE_COUNTER(STAT_NiagaraEmitterPostTick);
-	
-	FVersionedNiagaraEmitterData* EmitterData = CachedEmitter.GetEmitterData();
-	checkSlow(EmitterData);
 
 	if (EventInstanceData.IsValid())
 	{
@@ -1034,7 +1031,7 @@ void FNiagaraEmitterInstance::PostTick()
 	bCachedBoundsDynamic = false;
 	CachedBounds.Init();
 	
-	if (EmitterData)
+	if (FVersionedNiagaraEmitterData* EmitterData = CachedEmitter.GetEmitterData(); ensure(EmitterData))
 	{
 		// Read lock can be smaller in scope, but probably not necessary
 		FRWScopeLock ScopeLock(FixedBoundsGuard, SLT_ReadOnly);
