Prevent crashing the editor if a struct with unsupported type member is used in a module script

diff --git a/Engine/Plugins/FX/Niagara/Source/NiagaraEditor/Private/EdGraphSchema_Niagara.cpp b/Engine/Plugins/FX/Niagara/Source/NiagaraEditor/Private/EdGraphSchema_Niagara.cpp
--- a/Engine/Plugins/FX/Niagara/Source/NiagaraEditor/Private/EdGraphSchema_Niagara.cpp
+++ b/Engine/Plugins/FX/Niagara/Source/NiagaraEditor/Private/EdGraphSchema_Niagara.cpp
@@ -684,11 +684,14 @@ TArray<TSharedPtr<FNiagaraAction_NewNode>> UEdGraphSchema_Niagara::GetGraphActio
 						// we are checking here if the candidate type has a property that is the same as our pin
 						for (TFieldIterator<FProperty> CandidatePropertyIt(CandidateType.GetStruct(), EFieldIteratorFlags::IncludeSuper); CandidatePropertyIt; ++CandidatePropertyIt)
 						{
-							FNiagaraTypeDefinition CandidatePropertyType = GetTypeDefForProperty(*CandidatePropertyIt);
-							if(FNiagaraTypeDefinition::TypesAreAssignable(CandidatePropertyType, PinType))
+							if (IsValidNiagaraPropertyType(*CandidatePropertyIt))
 							{
-								bAddMake = true;
-								break;
+								FNiagaraTypeDefinition CandidatePropertyType = GetTypeDefForProperty(*CandidatePropertyIt);
+								if(FNiagaraTypeDefinition::TypesAreAssignable(CandidatePropertyType, PinType))
+								{
+									bAddMake = true;
+									break;
+								}
 							}
 						}
 					}
@@ -1888,6 +1891,18 @@ UNiagaraParameterCollection* UEdGraphSchema_Niagara::VariableIsFromParameterColl
 	return nullptr;
 }
 
+bool UEdGraphSchema_Niagara::IsValidNiagaraPropertyType(const FProperty* Property) const
+{
+	return Property != nullptr &&
+		(Property->IsA(FFloatProperty::StaticClass()) ||
+		Property->IsA(FDoubleProperty::StaticClass()) ||
+		Property->IsA(FIntProperty::StaticClass()) ||
+		Property->IsA(FBoolProperty::StaticClass()) ||
+		Property->IsA(FEnumProperty::StaticClass()) ||
+		Property->IsA(FStructProperty::StaticClass()) ||
+		Property->IsA(FUInt16Property::StaticClass()));
+}
+
 FNiagaraTypeDefinition UEdGraphSchema_Niagara::GetTypeDefForProperty(const FProperty* Property)const
 {
 	if (Property->IsA(FFloatProperty::StaticClass()) || Property->IsA(FDoubleProperty::StaticClass()))
diff --git a/Engine/Plugins/FX/Niagara/Source/NiagaraEditor/Private/NiagaraNodeConvert.cpp b/Engine/Plugins/FX/Niagara/Source/NiagaraEditor/Private/NiagaraNodeConvert.cpp
--- a/Engine/Plugins/FX/Niagara/Source/NiagaraEditor/Private/NiagaraNodeConvert.cpp
+++ b/Engine/Plugins/FX/Niagara/Source/NiagaraEditor/Private/NiagaraNodeConvert.cpp
@@ -102,15 +102,18 @@ struct FNiagaraConvertEntry
 		for (TFieldIterator<FProperty> PropertyIt(InStruct, EFieldIteratorFlags::IncludeSuper); PropertyIt; ++PropertyIt)
 		{
 			const FProperty* Property = *PropertyIt;
-			FNiagaraTypeDefinition PropType = Schema->GetTypeDefForProperty(Property);		
+			if (Schema->IsValidNiagaraPropertyType(Property))
+			{
+				FNiagaraTypeDefinition PropType = Schema->GetTypeDefForProperty(Property);
 
-			int32 Index = OutEntries.Emplace(InPinId, Property->GetFName(), PropType, InPin);
+				int32 Index = OutEntries.Emplace(InPinId, Property->GetFName(), PropType, InPin);
 
-			const FStructProperty* StructProperty = CastField<FStructProperty>(Property);
+				const FStructProperty* StructProperty = CastField<FStructProperty>(Property);
 
-			 if (StructProperty != nullptr)
-			{
-				CreateEntries(Schema, InPinId, InPin, FNiagaraTypeHelper::FindNiagaraFriendlyTopLevelStruct(StructProperty->Struct, ENiagaraStructConversion::UserFacing), OutEntries[Index].Children);
+				if (StructProperty != nullptr)
+				{
+					CreateEntries(Schema, InPinId, InPin, FNiagaraTypeHelper::FindNiagaraFriendlyTopLevelStruct(StructProperty->Struct, ENiagaraStructConversion::UserFacing), OutEntries[Index].Children);
+				}
 			}
 		}
 	}
@@ -566,6 +569,10 @@ void UNiagaraNodeConvert::AutowireNewNode(UEdGraphPin* FromPin)
 				SrcPath.Empty(SrcPath.Num());
 				DestPath.Empty(DestPath.Num());
 				const FProperty* Property = *PropertyIt;
+				if (!Schema->IsValidNiagaraPropertyType(Property))
+				{
+					continue;
+				}
 				FNiagaraTypeDefinition PropType = Schema->GetTypeDefForProperty(Property);
 				UEdGraphPin* NewPin = RequestNewTypedPin(Dir, PropType, *Property->GetDisplayNameText().ToString());
 
@@ -788,24 +795,27 @@ bool UNiagaraNodeConvert::InitConversion(UEdGraphPin* FromPin, UEdGraphPin* ToPi
 			++NextFromPropertyIt;
 		}
 
-		FNiagaraTypeDefinition FromPropType = Schema->GetTypeDefForProperty(FromProp);
-		FNiagaraTypeDefinition ToPropType = Schema->GetTypeDefForProperty(ToProp);
-		SrcPath.Empty();
-		DestPath.Empty();
-		if (FNiagaraTypeDefinition::TypesAreAssignable(FromPropType, ToPropType) || FNiagaraTypeDefinition::IsLossyConversion(FromPropType, ToPropType))
+		if (Schema->IsValidNiagaraPropertyType(FromProp) && Schema->IsValidNiagaraPropertyType(ToProp))
 		{
-			SrcPath.Add(*FromProp->GetName());
-			DestPath.Add(*ToProp->GetName());
-			Connections.Add(FNiagaraConvertConnection(ConnectFromPin->PinId, SrcPath, ConnectToPin->PinId, DestPath));
-
-			if (SrcPath.Num())
+			FNiagaraTypeDefinition FromPropType = Schema->GetTypeDefForProperty(FromProp);
+			FNiagaraTypeDefinition ToPropType = Schema->GetTypeDefForProperty(ToProp);
+			SrcPath.Empty();
+			DestPath.Empty();
+			if (FNiagaraTypeDefinition::TypesAreAssignable(FromPropType, ToPropType) || FNiagaraTypeDefinition::IsLossyConversion(FromPropType, ToPropType))
 			{
-				AddExpandedRecord(FNiagaraConvertPinRecord(ConnectFromPin->PinId, SrcPath).GetParent());
-			}
+				SrcPath.Add(*FromProp->GetName());
+				DestPath.Add(*ToProp->GetName());
+				Connections.Add(FNiagaraConvertConnection(ConnectFromPin->PinId, SrcPath, ConnectToPin->PinId, DestPath));
 
-			if (DestPath.Num())
-			{
-				AddExpandedRecord(FNiagaraConvertPinRecord(ConnectToPin->PinId, DestPath).GetParent());
+				if (SrcPath.Num())
+				{
+					AddExpandedRecord(FNiagaraConvertPinRecord(ConnectFromPin->PinId, SrcPath).GetParent());
+				}
+
+				if (DestPath.Num())
+				{
+					AddExpandedRecord(FNiagaraConvertPinRecord(ConnectToPin->PinId, DestPath).GetParent());
+				}
 			}
 		}
 			
diff --git a/Engine/Plugins/FX/Niagara/Source/NiagaraEditor/Private/ViewModels/NiagaraConvertPinViewModel.cpp b/Engine/Plugins/FX/Niagara/Source/NiagaraEditor/Private/ViewModels/NiagaraConvertPinViewModel.cpp
--- a/Engine/Plugins/FX/Niagara/Source/NiagaraEditor/Private/ViewModels/NiagaraConvertPinViewModel.cpp
+++ b/Engine/Plugins/FX/Niagara/Source/NiagaraEditor/Private/ViewModels/NiagaraConvertPinViewModel.cpp
@@ -80,7 +80,7 @@ void GenerateSocketViewModelsRecursive(const UEdGraphSchema_Niagara* Schema, TSh
 	for (TFieldIterator<FProperty> PropertyIterator(Struct); PropertyIterator; ++PropertyIterator)
 	{
 		FProperty* Property = *PropertyIterator;
-		if (Property != nullptr)
+		if (Schema->IsValidNiagaraPropertyType(Property))
 		{
 			FNiagaraTypeDefinition ChildTypeDef = Schema->GetTypeDefForProperty(Property);
 			TSharedRef<FNiagaraConvertPinSocketViewModel> SocketViewModel = MakeShareable(new FNiagaraConvertPinSocketViewModel(OwnerPinViewModel, OwnerPinSocketViewModel, Property->GetFName(), FName(*Property->GetDisplayNameText().ToString()), ChildTypeDef, Direction, TypeTraversalDepth));
diff --git a/Engine/Plugins/FX/Niagara/Source/NiagaraEditor/Public/EdGraphSchema_Niagara.h b/Engine/Plugins/FX/Niagara/Source/NiagaraEditor/Public/EdGraphSchema_Niagara.h
--- a/Engine/Plugins/FX/Niagara/Source/NiagaraEditor/Public/EdGraphSchema_Niagara.h
+++ b/Engine/Plugins/FX/Niagara/Source/NiagaraEditor/Public/EdGraphSchema_Niagara.h
@@ -159,7 +159,8 @@ class NIAGARAEDITOR_API UEdGraphSchema_Niagara : public UEdGraphSchema
 	class UNiagaraParameterCollection* VariableIsFromParameterCollection(const FNiagaraVariable& Var)const;
 	class UNiagaraParameterCollection* VariableIsFromParameterCollection(const FString& VarName, bool bAllowPartialMatch, FNiagaraVariable& OutVar)const;
 
-	FNiagaraTypeDefinition GetTypeDefForProperty(const FProperty* Property)const;
+	bool IsValidNiagaraPropertyType(const FProperty* Property) const;
+	FNiagaraTypeDefinition GetTypeDefForProperty(const FProperty* Property) const;
 
 	static const FLinearColor NodeTitleColor_Attribute;
 	static const FLinearColor NodeTitleColor_Constant;
