- Move recaching logic into update context and ensure we wait on the RT if we have GPU simulations to avoid reading from data that we are modifying

#rnx

diff --git a/Engine/Plugins/FX/Niagara/Source/Niagara/Private/NiagaraCommon.cpp b/Engine/Plugins/FX/Niagara/Source/Niagara/Private/NiagaraCommon.cpp
--- a/Engine/Plugins/FX/Niagara/Source/Niagara/Private/NiagaraCommon.cpp
+++ b/Engine/Plugins/FX/Niagara/Source/Niagara/Private/NiagaraCommon.cpp
@@ -74,14 +74,36 @@ FNiagaraSystemUpdateContext::~FNiagaraSystemUpdateContext()
 
 void FNiagaraSystemUpdateContext::CommitUpdate()
 {
-	for (UNiagaraSystem* Sys : SystemSimsToDestroy)
+	// Sims that are sent for destruction will also get a recache
+	if (SystemSimsToDestroy.Num() > 0)
 	{
-		if(Sys)
+		bool bAnyGpuSims = false;
+		for (UNiagaraSystem* NiagaraSystem : SystemSimsToDestroy)
 		{
-			FNiagaraWorldManager::DestroyAllSystemSimulations(Sys);
+			if (NiagaraSystem)
+			{
+				bAnyGpuSims |= NiagaraSystem->HasAnyGPUEmitters();
+				FNiagaraWorldManager::DestroyAllSystemSimulations(NiagaraSystem);
+			}
+		}
+
+		if (bAnyGpuSims)
+		{
+			FlushRenderingCommands();
 		}
+
+		for (UNiagaraSystem* NiagaraSystem : SystemSimsToDestroy)
+		{
+			if (NiagaraSystem)
+			{
+				NiagaraSystem->ComputeEmittersExecutionOrder();
+				NiagaraSystem->ComputeRenderersDrawOrder();
+				NiagaraSystem->CacheFromCompiledData();
+			}
+		}
+
+		SystemSimsToDestroy.Empty();
 	}
-	SystemSimsToDestroy.Empty();
 
 	for (UNiagaraComponent* Comp : ComponentsToReInit)
 	{
diff --git a/Engine/Plugins/FX/Niagara/Source/Niagara/Private/NiagaraEmitter.cpp b/Engine/Plugins/FX/Niagara/Source/Niagara/Private/NiagaraEmitter.cpp
--- a/Engine/Plugins/FX/Niagara/Source/Niagara/Private/NiagaraEmitter.cpp
+++ b/Engine/Plugins/FX/Niagara/Source/Niagara/Private/NiagaraEmitter.cpp
@@ -82,7 +82,6 @@ static FAutoConsoleVariableRef CVarNiagaraEmitterMaxGPUBufferElements(
 				if ( IsValid(System) )
 				{
 					UpdateCtx.Add(System, true);
-					System->CacheFromCompiledData();
 				}
 			}
 		}
@@ -2779,8 +2778,9 @@ void UNiagaraEmitter::AddRenderer(UNiagaraRendererProperties* Renderer, FGuid Em
 
 	if (UNiagaraSystem* Owner = GetTypedOuter<UNiagaraSystem>())
 	{
-		Owner->ComputeRenderersDrawOrder();
-		Owner->CacheFromCompiledData();
+		FNiagaraSystemUpdateContext UpdateContext;
+		UpdateContext.Add(Owner, true);
+		UpdateContext.CommitUpdate();
 	}
 }
 
diff --git a/Engine/Plugins/FX/Niagara/Source/NiagaraEditor/Private/ViewModels/NiagaraSystemViewModel.cpp b/Engine/Plugins/FX/Niagara/Source/NiagaraEditor/Private/ViewModels/NiagaraSystemViewModel.cpp
--- a/Engine/Plugins/FX/Niagara/Source/NiagaraEditor/Private/ViewModels/NiagaraSystemViewModel.cpp
+++ b/Engine/Plugins/FX/Niagara/Source/NiagaraEditor/Private/ViewModels/NiagaraSystemViewModel.cpp
@@ -1652,12 +1652,6 @@ void FNiagaraSystemViewModel::ResetSystem(ETimeResetMode TimeResetMode, EMultiRe
 	{
 		UNiagaraSystem* NiagaraSystem = &GetSystem();
 		UpdateContext.Add(NiagaraSystem, ReinitMode == EReinitMode::ReinitializeSystem);
-
-		// When reinitializing the whole system update all re-cache all information
-		// Note: This must be done after we add the system to the update context as this will ensure we waited for any async work to complete
-		NiagaraSystem->ComputeEmittersExecutionOrder();
-		NiagaraSystem->ComputeRenderersDrawOrder();
-		NiagaraSystem->CacheFromCompiledData();
 	}
 
 	UpdateContext.CommitUpdate();
