Replacing duplicate symbol FGraphEditorSummoner

diff --git a/Engine/Plugins/FX/Niagara/Source/NiagaraEditor/Private/ViewModels/NiagaraSystemEditorDocumentsViewModel.cpp b/Engine/Plugins/FX/Niagara/Source/NiagaraEditor/Private/ViewModels/NiagaraSystemEditorDocumentsViewModel.cpp
--- a/Engine/Plugins/FX/Niagara/Source/NiagaraEditor/Private/ViewModels/NiagaraSystemEditorDocumentsViewModel.cpp
+++ b/Engine/Plugins/FX/Niagara/Source/NiagaraEditor/Private/ViewModels/NiagaraSystemEditorDocumentsViewModel.cpp
@@ -92,8 +92,8 @@ void UNiagaraSystemEditorDocumentsViewModel::InitializePreTabManager(TSharedPtr<
 	// Register the document factories
 	{
 
-		TSharedRef<FDocumentTabFactory> GraphEditorFactory = MakeShareable(new FGraphEditorSummoner(InToolkit,
-			FGraphEditorSummoner::FOnCreateGraphEditorWidget::CreateUObject(this, &UNiagaraSystemEditorDocumentsViewModel::CreateGraphEditorWidget)
+		TSharedRef<FDocumentTabFactory> GraphEditorFactory = MakeShareable(new FNiagaraGraphEditorSummoner(InToolkit,
+			FNiagaraGraphEditorSummoner::FOnCreateGraphEditorWidget::CreateUObject(this, &UNiagaraSystemEditorDocumentsViewModel::CreateGraphEditorWidget)
 		));
 
 		// Also store off a reference to the grapheditor factory so we can find all the tabs spawned by it later.
diff --git a/Engine/Plugins/FX/Niagara/Source/NiagaraEditor/Private/Widgets/NiagaraScratchScriptEditor.cpp b/Engine/Plugins/FX/Niagara/Source/NiagaraEditor/Private/Widgets/NiagaraScratchScriptEditor.cpp
--- a/Engine/Plugins/FX/Niagara/Source/NiagaraEditor/Private/Widgets/NiagaraScratchScriptEditor.cpp
+++ b/Engine/Plugins/FX/Niagara/Source/NiagaraEditor/Private/Widgets/NiagaraScratchScriptEditor.cpp
@@ -109,7 +109,7 @@ void SNiagaraScratchPadScriptEditor::Construct(const FArguments& InArgs, TShared
 	}
 }
 
-void FGraphTabHistory::EvokeHistory(TSharedPtr<FTabInfo> InTabInfo, bool bPrevTabMatches)
+void FNiagaraGraphTabHistory::EvokeHistory(TSharedPtr<FTabInfo> InTabInfo, bool bPrevTabMatches)
 {
 	FWorkflowTabSpawnInfo SpawnInfo;
 	SpawnInfo.Payload = Payload;
@@ -129,7 +129,7 @@ void FGraphTabHistory::EvokeHistory(TSharedPtr<FTabInfo> InTabInfo, bool bPrevTa
 	}
 }
 
-void FGraphTabHistory::SaveHistory()
+void FNiagaraGraphTabHistory::SaveHistory()
 {
 	if (IsHistoryValid())
 	{
@@ -143,7 +143,7 @@ void FGraphTabHistory::SaveHistory()
 	}
 }
 
-void FGraphTabHistory::RestoreHistory()
+void FNiagaraGraphTabHistory::RestoreHistory()
 {
 	if (IsHistoryValid())
 	{
@@ -153,25 +153,25 @@ void FGraphTabHistory::RestoreHistory()
 			Editor->GetGraphEditor()->SetViewLocation(SavedLocation, SavedZoomAmount, SavedBookmarkId);
 	}
 }
-void FGraphEditorSummoner::OnTabActivated(TSharedPtr<SDockTab> Tab) const
+void FNiagaraGraphEditorSummoner::OnTabActivated(TSharedPtr<SDockTab> Tab) const
 {
 	TSharedRef<SNiagaraScratchPadScriptEditor> GraphEditor = StaticCastSharedRef<SNiagaraScratchPadScriptEditor>(Tab->GetContent());
 	EditorPtr.Pin()->GetSystemViewModel()->GetDocumentViewModel()->SetActiveDocumentTab(Tab);
 }
 
-void FGraphEditorSummoner::OnTabBackgrounded(TSharedPtr<SDockTab> Tab) const
+void FNiagaraGraphEditorSummoner::OnTabBackgrounded(TSharedPtr<SDockTab> Tab) const
 {
 	TSharedRef<SNiagaraScratchPadScriptEditor> GraphEditor = StaticCastSharedRef<SNiagaraScratchPadScriptEditor>(Tab->GetContent());
 	//EditorPtr.Pin()->OnGraphEditorBackgrounded(GraphEditor);
 }
 
-void FGraphEditorSummoner::OnTabRefreshed(TSharedPtr<SDockTab> Tab) const
+void FNiagaraGraphEditorSummoner::OnTabRefreshed(TSharedPtr<SDockTab> Tab) const
 {
 	TSharedRef<SNiagaraScratchPadScriptEditor> GraphEditor = StaticCastSharedRef<SNiagaraScratchPadScriptEditor>(Tab->GetContent());
 	GraphEditor->GetGraphEditor()->NotifyGraphChanged();
 }
 
-void FGraphEditorSummoner::SaveState(TSharedPtr<SDockTab> Tab, TSharedPtr<FTabPayload> Payload) const
+void FNiagaraGraphEditorSummoner::SaveState(TSharedPtr<SDockTab> Tab, TSharedPtr<FTabPayload> Payload) const
 {
 	// Do nothing here for now
 	/*
@@ -185,14 +185,14 @@ void FGraphEditorSummoner::SaveState(TSharedPtr<SDockTab> Tab, TSharedPtr<FTabPa
 	*/
 }
 
-FGraphEditorSummoner::FGraphEditorSummoner(TSharedPtr<class FNiagaraSystemToolkit> InToolkit, FOnCreateGraphEditorWidget CreateGraphEditorWidgetCallback) : FDocumentTabFactoryForObjects<UEdGraph>(TEXT("SystemToolkitGraph"), InToolkit->AsShared())
+FNiagaraGraphEditorSummoner::FNiagaraGraphEditorSummoner(TSharedPtr<class FNiagaraSystemToolkit> InToolkit, FOnCreateGraphEditorWidget CreateGraphEditorWidgetCallback) : FDocumentTabFactoryForObjects<UEdGraph>(TEXT("SystemToolkitGraph"), InToolkit->AsShared())
 , EditorPtr(InToolkit)
 , OnCreateGraphEditorWidget(CreateGraphEditorWidgetCallback)
 {
 
 }
 
-TSharedRef<SWidget> FGraphEditorSummoner::CreateTabBodyForObject(const FWorkflowTabSpawnInfo& Info, UEdGraph* DocumentID) const
+TSharedRef<SWidget> FNiagaraGraphEditorSummoner::CreateTabBodyForObject(const FWorkflowTabSpawnInfo& Info, UEdGraph* DocumentID) const
 {
 	check(Info.TabInfo.IsValid());
 
@@ -212,17 +212,17 @@ TSharedRef<SWidget> FGraphEditorSummoner::CreateTabBodyForObject(const FWorkflow
 	return OnCreateGraphEditorWidget.Execute(Info.TabInfo.ToSharedRef(), DocumentID);
 }
 
-const FSlateBrush* FGraphEditorSummoner::GetTabIconForObject(const FWorkflowTabSpawnInfo& Info, UEdGraph* DocumentID) const
+const FSlateBrush* FNiagaraGraphEditorSummoner::GetTabIconForObject(const FWorkflowTabSpawnInfo& Info, UEdGraph* DocumentID) const
 {
 	return FSlateIcon(FNiagaraEditorStyle::Get().GetStyleSetName(), "Tab.ScratchPad").GetSmallIcon();
 }
 
-TSharedRef<FGenericTabHistory> FGraphEditorSummoner::CreateTabHistoryNode(TSharedPtr<FTabPayload> Payload)
+TSharedRef<FGenericTabHistory> FNiagaraGraphEditorSummoner::CreateTabHistoryNode(TSharedPtr<FTabPayload> Payload)
 {
-	return MakeShareable(new FGraphTabHistory(SharedThis(this), Payload));
+	return MakeShareable(new FNiagaraGraphTabHistory(SharedThis(this), Payload));
 }
 
-TAttribute<FText> FGraphEditorSummoner::ConstructTabNameForObject(UEdGraph* DocumentID) const 
+TAttribute<FText> FNiagaraGraphEditorSummoner::ConstructTabNameForObject(UEdGraph* DocumentID) const 
 {
 	UNiagaraScript* Script = DocumentID->GetTypedOuter<UNiagaraScript>();
 	TSharedPtr<FNiagaraSystemViewModel> SysViewModel = EditorPtr.Pin()->GetSystemViewModel();
@@ -241,7 +241,7 @@ TAttribute<FText> FGraphEditorSummoner::ConstructTabNameForObject(UEdGraph* Docu
 	return FText();
 }
 
-TSharedRef<SDockTab> FGraphEditorSummoner::OnSpawnTab(const FSpawnTabArgs& SpawnArgs, TWeakPtr<FTabManager> WeakTabManager) const
+TSharedRef<SDockTab> FNiagaraGraphEditorSummoner::OnSpawnTab(const FSpawnTabArgs& SpawnArgs, TWeakPtr<FTabManager> WeakTabManager) const
 {
 	TSharedRef<SDockTab> TabRef = FDocumentTabFactoryForObjects<UEdGraph>::OnSpawnTab(SpawnArgs, WeakTabManager);
 	return TabRef;
diff --git a/Engine/Plugins/FX/Niagara/Source/NiagaraEditor/Private/Widgets/NiagaraScratchScriptEditor.h b/Engine/Plugins/FX/Niagara/Source/NiagaraEditor/Private/Widgets/NiagaraScratchScriptEditor.h
--- a/Engine/Plugins/FX/Niagara/Source/NiagaraEditor/Private/Widgets/NiagaraScratchScriptEditor.h
+++ b/Engine/Plugins/FX/Niagara/Source/NiagaraEditor/Private/Widgets/NiagaraScratchScriptEditor.h
@@ -77,16 +77,16 @@ private:
 
 
 /////////////////////////////////////////////////////
-// FGraphTabHistory
+// FNiagaraGraphTabHistory
 
-struct FGraphTabHistory : public FGenericTabHistory
+struct FNiagaraGraphTabHistory : public FGenericTabHistory
 {
 public:
 	/**
 	 * @param InFactory		The factory used to regenerate the content
 	 * @param InPayload		The payload object used to regenerate the content
 	 */
-	FGraphTabHistory(TSharedPtr<FDocumentTabFactory> InFactory, TSharedPtr<FTabPayload> InPayload)
+	FNiagaraGraphTabHistory(TSharedPtr<FDocumentTabFactory> InFactory, TSharedPtr<FTabPayload> InPayload)
 		: FGenericTabHistory(InFactory, InPayload)
 		, SavedLocation(FVector2D::ZeroVector)
 		, SavedZoomAmount(INDEX_NONE)
@@ -112,12 +112,12 @@ private:
 };
 
 
-struct FGraphEditorSummoner : public FDocumentTabFactoryForObjects<UEdGraph>
+struct FNiagaraGraphEditorSummoner : public FDocumentTabFactoryForObjects<UEdGraph>
 {
 public:
 	DECLARE_DELEGATE_RetVal_TwoParams(TSharedRef<SNiagaraScratchPadScriptEditor>, FOnCreateGraphEditorWidget, TSharedRef<FTabInfo>, UEdGraph*);
 public:
-	FGraphEditorSummoner(TSharedPtr<class FNiagaraSystemToolkit> InToolkit, FOnCreateGraphEditorWidget CreateGraphEditorWidgetCallback);
+	FNiagaraGraphEditorSummoner(TSharedPtr<class FNiagaraSystemToolkit> InToolkit, FOnCreateGraphEditorWidget CreateGraphEditorWidgetCallback);
 
 	virtual void OnTabActivated(TSharedPtr<SDockTab> Tab) const override;
 
