Null pointer check in UNiagaraGraph::GetOwningEmitter()

diff --git a/Engine/Plugins/FX/Niagara/Source/NiagaraEditor/Private/NiagaraGraph.cpp b/Engine/Plugins/FX/Niagara/Source/NiagaraEditor/Private/NiagaraGraph.cpp
--- a/Engine/Plugins/FX/Niagara/Source/NiagaraEditor/Private/NiagaraGraph.cpp
+++ b/Engine/Plugins/FX/Niagara/Source/NiagaraEditor/Private/NiagaraGraph.cpp
@@ -2282,10 +2282,13 @@ FVersionedNiagaraEmitter UNiagaraGraph::GetOwningEmitter() const
 		for (FNiagaraAssetVersion Version : OwningEmitter->GetAllAvailableVersions())
 		{
 			FVersionedNiagaraEmitterData* EmitterData = OwningEmitter->GetEmitterData(Version.VersionGuid);
-			UNiagaraGraph* EmitterGraph = Cast<UNiagaraScriptSource>(EmitterData->GraphSource)->NodeGraph;
-			if (this == EmitterGraph)
+			if (EmitterData->GraphSource)
 			{
-				return FVersionedNiagaraEmitter(OwningEmitter, Version.VersionGuid);
+				UNiagaraGraph* EmitterGraph = Cast<UNiagaraScriptSource>(EmitterData->GraphSource)->NodeGraph;
+				if (this == EmitterGraph)
+				{
+					return FVersionedNiagaraEmitter(OwningEmitter, Version.VersionGuid);
+				}
 			}
 		}
 	}
