- Conver LWC types to simulation types for function prototype generation

diff --git a/Engine/Plugins/FX/Niagara/Source/NiagaraEditor/Private/NiagaraHlslTranslator.cpp b/Engine/Plugins/FX/Niagara/Source/NiagaraEditor/Private/NiagaraHlslTranslator.cpp
--- a/Engine/Plugins/FX/Niagara/Source/NiagaraEditor/Private/NiagaraHlslTranslator.cpp
+++ b/Engine/Plugins/FX/Niagara/Source/NiagaraEditor/Private/NiagaraHlslTranslator.cpp
@@ -8309,9 +8309,9 @@ FString FHlslNiagaraTranslator::GenerateFunctionHlslPrototype(FStringView InVari
 			bool bNeedsComma = false;
 
 			// Inputs
-			for (int i = 1; i < FunctionSignature.Inputs.Num(); ++i)
+			for (int i=1; i < FunctionSignature.Inputs.Num(); ++i)
 			{
-				const FNiagaraVariable& InputVar = FunctionSignature.Inputs[i];
+				FNiagaraVariable InputVar = ConvertToSimulationVariable(FunctionSignature.Inputs[i]);
 				if (bNeedsComma)
 				{
 					StringBuilder.Append(TEXT(", "));
@@ -8320,19 +8320,20 @@ FString FHlslNiagaraTranslator::GenerateFunctionHlslPrototype(FStringView InVari
 
 				StringBuilder.Append(TEXT("in "));
 				StringBuilder.Append(FHlslNiagaraTranslator::GetStructHlslTypeName(InputVar.GetType()));
-
 				StringBuilder.Append(TEXT(" In_"));
 				StringBuilder.Append(FHlslNiagaraTranslator::GetSanitizedSymbolName(InputVar.GetName().ToString()));
 			}
 
 			// Outputs
-			for (const FNiagaraVariable& OutputVar : FunctionSignature.Outputs)
+			for (int i=0; i < FunctionSignature.Outputs.Num(); ++i)
 			{
+				FNiagaraVariable OutputVar = ConvertToSimulationVariable(FunctionSignature.Outputs[i]);
 				if (bNeedsComma)
 				{
 					StringBuilder.Append(TEXT(", "));
 				}
 				bNeedsComma = true;
+
 				StringBuilder.Append(TEXT("out "));
 				StringBuilder.Append(FHlslNiagaraTranslator::GetStructHlslTypeName(OutputVar.GetType()));
 				StringBuilder.Append(TEXT(" Out_"));
