Fixed a crash when opening scratch pads from the stack

diff --git a/Engine/Plugins/FX/Niagara/Source/NiagaraEditor/Private/Toolkits/SystemToolkitModes/NiagaraSystemToolkitModeBase.cpp b/Engine/Plugins/FX/Niagara/Source/NiagaraEditor/Private/Toolkits/SystemToolkitModes/NiagaraSystemToolkitModeBase.cpp
--- a/Engine/Plugins/FX/Niagara/Source/NiagaraEditor/Private/Toolkits/SystemToolkitModes/NiagaraSystemToolkitModeBase.cpp
+++ b/Engine/Plugins/FX/Niagara/Source/NiagaraEditor/Private/Toolkits/SystemToolkitModes/NiagaraSystemToolkitModeBase.cpp
@@ -110,7 +110,10 @@ void FNiagaraSystemToolkitModeBase::OnActiveDocumentChanged(TSharedPtr<SDockTab>
 		if (ScratchScriptVM.IsValid())
 		{
 			LastSelectionUpdateDelegate = ScratchScriptVM->GetGraphViewModel()->GetNodeSelection()->OnSelectedObjectsChanged().AddRaw(this, &FNiagaraSystemToolkitModeBase::UpdateSelectionForActiveDocument);
-			LastParamPanelSelectionUpdateDelegate = Toolkit->GetSystemViewModel()->GetParameterPanelViewModel()->GetVariableObjectSelection()->OnSelectedObjectsChanged().AddRaw(this, &FNiagaraSystemToolkitModeBase::OnParameterPanelViewModelExternalSelectionChanged);
+			if (INiagaraParameterPanelViewModel* ParameterPanelViewModel = Toolkit->GetSystemViewModel()->GetParameterPanelViewModel())
+			{
+				LastParamPanelSelectionUpdateDelegate = ParameterPanelViewModel->GetVariableObjectSelection()->OnSelectedObjectsChanged().AddRaw(this, &FNiagaraSystemToolkitModeBase::OnParameterPanelViewModelExternalSelectionChanged);
+			}
 			LastSystemSelectionUpdateDelegate = Toolkit->GetSystemViewModel()->GetSelectionViewModel()->OnEntrySelectionChanged().AddSP(this, &FNiagaraSystemToolkitModeBase::OnSystemSelectionChanged);
 			LastGraphEditDelegate = ScratchScriptVM->GetGraphViewModel()->GetGraph()->AddOnGraphNeedsRecompileHandler(
 				FOnGraphChanged::FDelegate::CreateRaw(this, &FNiagaraSystemToolkitModeBase::OnEditedScriptGraphChanged));
