Fix crash baking to sim cache with disabled GPU emitters.

diff --git a/Engine/Plugins/FX/Niagara/Source/Niagara/Private/NiagaraSimCache.cpp b/Engine/Plugins/FX/Niagara/Source/Niagara/Private/NiagaraSimCache.cpp
--- a/Engine/Plugins/FX/Niagara/Source/Niagara/Private/NiagaraSimCache.cpp
+++ b/Engine/Plugins/FX/Niagara/Source/Niagara/Private/NiagaraSimCache.cpp
@@ -293,8 +293,15 @@ struct FNiagaraSimCacheHelper
 
 	void ReadDataBufferGPU(FNiagaraEmitterInstance& EmitterInstance, const FNiagaraSimCacheDataBuffersLayout& InCacheLayout, const FNiagaraSimCacheDataBuffers& InCacheBuffer, FNiagaraDataSet& InDataSet, std::atomic<int32>& InPendingCommandsCounter)
 	{
+		if (EmitterInstance.IsDisabled())
+		{
+			return;
+		}
+
 		++InPendingCommandsCounter;
 
+		check(EmitterInstance.GetGPUContext());
+
 		FNiagaraGpuComputeDispatchInterface* DispathInterface = EmitterInstance.GetParentSystemInstance()->GetComputeDispatchInterface();
 		ENQUEUE_RENDER_COMMAND(NiagaraSimCacheGpuReadFrame)(
 			[DispathInterface, GPUExecContext=EmitterInstance.GetGPUContext(), CacheLayout=&InCacheLayout, CacheBuffer=&InCacheBuffer, DataSet=&InDataSet, PendingCommandsCounter=&InPendingCommandsCounter](FRHICommandListImmediate& RHICmdList)
