- Update the NiagaraVariant when calling array set functions
- This allows editor utilities to set the array data

diff --git a/Engine/Plugins/FX/Niagara/Source/Niagara/Classes/NiagaraDataInterfaceArrayFunctionLibrary.h b/Engine/Plugins/FX/Niagara/Source/Niagara/Classes/NiagaraDataInterfaceArrayFunctionLibrary.h
--- a/Engine/Plugins/FX/Niagara/Source/Niagara/Classes/NiagaraDataInterfaceArrayFunctionLibrary.h
+++ b/Engine/Plugins/FX/Niagara/Source/Niagara/Classes/NiagaraDataInterfaceArrayFunctionLibrary.h
@@ -5,6 +5,8 @@
 #include "Kismet/BlueprintFunctionLibrary.h"
 #include "NiagaraDataInterfaceArrayFunctionLibrary.generated.h"
 
+class UNiagaraComponent;
+
 /**
 * C++ and Blueprint library for accessing array types
 */
diff --git a/Engine/Plugins/FX/Niagara/Source/Niagara/Private/NiagaraDataInterfaceArrayFunctionLibrary.cpp b/Engine/Plugins/FX/Niagara/Source/Niagara/Private/NiagaraDataInterfaceArrayFunctionLibrary.cpp
--- a/Engine/Plugins/FX/Niagara/Source/Niagara/Private/NiagaraDataInterfaceArrayFunctionLibrary.cpp
+++ b/Engine/Plugins/FX/Niagara/Source/Niagara/Private/NiagaraDataInterfaceArrayFunctionLibrary.cpp
@@ -10,21 +10,33 @@
 #include "NiagaraSystemInstance.h"
 #include "NiagaraFunctionLibrary.h"
 
+template<typename TArrayType, typename TDataInterace>
+void SetNiagaraVariantArray(UNiagaraComponent* NiagaraComponent, FName OverrideName, TDataInterace* ExistingDataInterface)
+{
+#if WITH_EDITOR
+	TDataInterace* VariantDataInterface = CastChecked<TDataInterace>(DuplicateObject(ExistingDataInterface, NiagaraComponent));
+	auto* ExistingProxy = static_cast<FNDIArrayProxyImpl<TArrayType, TDataInterace>*>(ExistingDataInterface->GetProxy());
+	VariantDataInterface->GetArrayReference() = ExistingProxy->GetArrayData();
+	NiagaraComponent->SetParameterOverride(FNiagaraVariableBase(FNiagaraTypeDefinition(TDataInterace::StaticClass()), OverrideName), FNiagaraVariant(VariantDataInterface));
+#endif
+}
+
 // If / when we share user parameter UObjects we will need to make this per instance which introduces some tricky things about allocating before the instance is active
 template<typename TArrayType, typename TDataInterace>
-void SetNiagaraArray(UNiagaraComponent* NiagaraSystem, FName OverrideName, const TArray<TArrayType>& InArray)
+void SetNiagaraArray(UNiagaraComponent* NiagaraComponent, FName OverrideName, const TArray<TArrayType>& InArray)
 {
-	if (TDataInterace* ArrayDI = UNiagaraFunctionLibrary::GetDataInterface<TDataInterace>(NiagaraSystem, OverrideName))
+	if (TDataInterace* ArrayDI = UNiagaraFunctionLibrary::GetDataInterface<TDataInterace>(NiagaraComponent, OverrideName))
 	{
 		auto* ArrayProxy = static_cast<FNDIArrayProxyImpl<TArrayType, TDataInterace>*>(ArrayDI->GetProxy());
 		ArrayProxy->SetArrayData(InArray);
+		SetNiagaraVariantArray<TArrayType, TDataInterace>(NiagaraComponent, OverrideName, ArrayDI);
 	}
 }
 
 template<typename TArrayType, typename TDataInterace>
-TArray<TArrayType> GetNiagaraArray(UNiagaraComponent* NiagaraSystem, FName OverrideName)
+TArray<TArrayType> GetNiagaraArray(UNiagaraComponent* NiagaraComponent, FName OverrideName)
 {
-	if (TDataInterace* ArrayDI = UNiagaraFunctionLibrary::GetDataInterface<TDataInterace>(NiagaraSystem, OverrideName))
+	if (TDataInterace* ArrayDI = UNiagaraFunctionLibrary::GetDataInterface<TDataInterace>(NiagaraComponent, OverrideName))
 	{
 		auto* ArrayProxy = static_cast<FNDIArrayProxyImpl<TArrayType, TDataInterace>*>(ArrayDI->GetProxy());
 		return ArrayProxy->GetArrayData();
@@ -33,19 +45,20 @@ TArray<TArrayType> GetNiagaraArray(UNiagaraComponent* NiagaraSystem, FName Overr
 }
 
 template<typename TArrayType, typename TDataInterace>
-void SetNiagaraArrayValue(UNiagaraComponent* NiagaraSystem, FName OverrideName, int Index, const TArrayType& Value, bool bSizeToFit)
+void SetNiagaraArrayValue(UNiagaraComponent* NiagaraComponent, FName OverrideName, int Index, const TArrayType& Value, bool bSizeToFit)
 {
-	if (TDataInterace* ArrayDI = UNiagaraFunctionLibrary::GetDataInterface<TDataInterace>(NiagaraSystem, OverrideName))
+	if (TDataInterace* ArrayDI = UNiagaraFunctionLibrary::GetDataInterface<TDataInterace>(NiagaraComponent, OverrideName))
 	{
 		auto* ArrayProxy = static_cast<FNDIArrayProxyImpl<TArrayType, TDataInterace>*>(ArrayDI->GetProxy());
 		ArrayProxy->SetArrayValue(Index, Value, bSizeToFit);
+		SetNiagaraVariantArray<TArrayType, TDataInterace>(NiagaraComponent, OverrideName, ArrayDI);
 	}
 }
 
 template<typename TArrayType, typename TDataInterace>
-TArrayType GetNiagaraArrayValue(UNiagaraComponent* NiagaraSystem, FName OverrideName, int Index)
+TArrayType GetNiagaraArrayValue(UNiagaraComponent* NiagaraComponent, FName OverrideName, int Index)
 {
-	if (TDataInterace* ArrayDI = UNiagaraFunctionLibrary::GetDataInterface<TDataInterace>(NiagaraSystem, OverrideName))
+	if (TDataInterace* ArrayDI = UNiagaraFunctionLibrary::GetDataInterface<TDataInterace>(NiagaraComponent, OverrideName))
 	{
 		auto* ArrayProxy = static_cast<FNDIArrayProxyImpl<TArrayType, TDataInterace>*>(ArrayDI->GetProxy());
 		return ArrayProxy->GetArrayValue(Index);
@@ -54,9 +67,9 @@ TArrayType GetNiagaraArrayValue(UNiagaraComponent* NiagaraSystem, FName Override
 	return TArrayType();
 }
 
-FNiagaraLWCConverter GetLWCConverter(UNiagaraComponent* NiagaraSystem)
+FNiagaraLWCConverter GetLWCConverter(UNiagaraComponent* NiagaraComponent)
 {
-	FNiagaraSystemInstanceControllerPtr SystemInstanceController = NiagaraSystem->GetSystemInstanceController();
+	FNiagaraSystemInstanceControllerPtr SystemInstanceController = NiagaraComponent->GetSystemInstanceController();
 	if (!SystemInstanceController.IsValid())
 	{
 		// the instance controller can be invalid if there is no simulation, for example when fx.SuppressNiagaraSystems is set 
@@ -68,80 +81,80 @@ FNiagaraLWCConverter GetLWCConverter(UNiagaraComponent* NiagaraSystem)
 
 //////////////////////////////////////////////////////////////////////////
 
-void UNiagaraDataInterfaceArrayFunctionLibrary::SetNiagaraArrayFloat(UNiagaraComponent* NiagaraSystem, FName OverrideName, const TArray<float>& ArrayData)
+void UNiagaraDataInterfaceArrayFunctionLibrary::SetNiagaraArrayFloat(UNiagaraComponent* NiagaraComponent, FName OverrideName, const TArray<float>& ArrayData)
 {
-	SetNiagaraArray<float, UNiagaraDataInterfaceArrayFloat>(NiagaraSystem, OverrideName, ArrayData);
+	SetNiagaraArray<float, UNiagaraDataInterfaceArrayFloat>(NiagaraComponent, OverrideName, ArrayData);
 }
 
-void UNiagaraDataInterfaceArrayFunctionLibrary::SetNiagaraArrayVector2D(UNiagaraComponent* NiagaraSystem, FName OverrideName, const TArray<FVector2D>& ArrayData)
+void UNiagaraDataInterfaceArrayFunctionLibrary::SetNiagaraArrayVector2D(UNiagaraComponent* NiagaraComponent, FName OverrideName, const TArray<FVector2D>& ArrayData)
 {
-	SetNiagaraArray<FVector2D, UNiagaraDataInterfaceArrayFloat2>(NiagaraSystem, OverrideName, ArrayData);
+	SetNiagaraArray<FVector2D, UNiagaraDataInterfaceArrayFloat2>(NiagaraComponent, OverrideName, ArrayData);
 }
 
-void UNiagaraDataInterfaceArrayFunctionLibrary::SetNiagaraArrayVector(UNiagaraComponent* NiagaraSystem, FName OverrideName, const TArray<FVector>& ArrayData)
+void UNiagaraDataInterfaceArrayFunctionLibrary::SetNiagaraArrayVector(UNiagaraComponent* NiagaraComponent, FName OverrideName, const TArray<FVector>& ArrayData)
 {
-	SetNiagaraArray<FVector, UNiagaraDataInterfaceArrayFloat3>(NiagaraSystem, OverrideName, ArrayData);
+	SetNiagaraArray<FVector, UNiagaraDataInterfaceArrayFloat3>(NiagaraComponent, OverrideName, ArrayData);
 }
 
-void UNiagaraDataInterfaceArrayFunctionLibrary::SetNiagaraArrayPosition(UNiagaraComponent* NiagaraSystem, FName OverrideName, const TArray<FVector>& ArrayData)
+void UNiagaraDataInterfaceArrayFunctionLibrary::SetNiagaraArrayPosition(UNiagaraComponent* NiagaraComponent, FName OverrideName, const TArray<FVector>& ArrayData)
 {
-	FNiagaraLWCConverter LwcConverter = GetLWCConverter(NiagaraSystem);
+	FNiagaraLWCConverter LwcConverter = GetLWCConverter(NiagaraComponent);
 	TArray<FNiagaraPosition> ConvertedData;
 	ConvertedData.SetNumUninitialized(ArrayData.Num());
 	for (int i = 0; i < ArrayData.Num(); i++)
 	{
 		ConvertedData[i] = LwcConverter.ConvertWorldToSimulationPosition(ArrayData[i]);
 	}
-	SetNiagaraArray<FNiagaraPosition, UNiagaraDataInterfaceArrayPosition>(NiagaraSystem, OverrideName, ConvertedData);
+	SetNiagaraArray<FNiagaraPosition, UNiagaraDataInterfaceArrayPosition>(NiagaraComponent, OverrideName, ConvertedData);
 }
 
-void UNiagaraDataInterfaceArrayFunctionLibrary::SetNiagaraArrayVector4(UNiagaraComponent* NiagaraSystem, FName OverrideName, const TArray<FVector4>& ArrayData)
+void UNiagaraDataInterfaceArrayFunctionLibrary::SetNiagaraArrayVector4(UNiagaraComponent* NiagaraComponent, FName OverrideName, const TArray<FVector4>& ArrayData)
 {
-	SetNiagaraArray<FVector4, UNiagaraDataInterfaceArrayFloat4>(NiagaraSystem, OverrideName, ArrayData);
+	SetNiagaraArray<FVector4, UNiagaraDataInterfaceArrayFloat4>(NiagaraComponent, OverrideName, ArrayData);
 }
 
-void UNiagaraDataInterfaceArrayFunctionLibrary::SetNiagaraArrayColor(UNiagaraComponent* NiagaraSystem, FName OverrideName, const TArray<FLinearColor>& ArrayData)
+void UNiagaraDataInterfaceArrayFunctionLibrary::SetNiagaraArrayColor(UNiagaraComponent* NiagaraComponent, FName OverrideName, const TArray<FLinearColor>& ArrayData)
 {
-	SetNiagaraArray<FLinearColor, UNiagaraDataInterfaceArrayColor>(NiagaraSystem, OverrideName, ArrayData);
+	SetNiagaraArray<FLinearColor, UNiagaraDataInterfaceArrayColor>(NiagaraComponent, OverrideName, ArrayData);
 }
 
-void UNiagaraDataInterfaceArrayFunctionLibrary::SetNiagaraArrayQuat(UNiagaraComponent* NiagaraSystem, FName OverrideName, const TArray<FQuat>& ArrayData)
+void UNiagaraDataInterfaceArrayFunctionLibrary::SetNiagaraArrayQuat(UNiagaraComponent* NiagaraComponent, FName OverrideName, const TArray<FQuat>& ArrayData)
 {
-	SetNiagaraArray<FQuat, UNiagaraDataInterfaceArrayQuat>(NiagaraSystem, OverrideName, ArrayData);
+	SetNiagaraArray<FQuat, UNiagaraDataInterfaceArrayQuat>(NiagaraComponent, OverrideName, ArrayData);
 }
 
-void UNiagaraDataInterfaceArrayFunctionLibrary::SetNiagaraArrayInt32(UNiagaraComponent* NiagaraSystem, FName OverrideName, const TArray<int32>& ArrayData)
+void UNiagaraDataInterfaceArrayFunctionLibrary::SetNiagaraArrayInt32(UNiagaraComponent* NiagaraComponent, FName OverrideName, const TArray<int32>& ArrayData)
 {
-	SetNiagaraArray<int32, UNiagaraDataInterfaceArrayInt32>(NiagaraSystem, OverrideName, ArrayData);
+	SetNiagaraArray<int32, UNiagaraDataInterfaceArrayInt32>(NiagaraComponent, OverrideName, ArrayData);
 }
 
-void UNiagaraDataInterfaceArrayFunctionLibrary::SetNiagaraArrayBool(UNiagaraComponent* NiagaraSystem, FName OverrideName, const TArray<bool>& ArrayData)
+void UNiagaraDataInterfaceArrayFunctionLibrary::SetNiagaraArrayBool(UNiagaraComponent* NiagaraComponent, FName OverrideName, const TArray<bool>& ArrayData)
 {
-	SetNiagaraArray<bool, UNiagaraDataInterfaceArrayBool>(NiagaraSystem, OverrideName, ArrayData);
+	SetNiagaraArray<bool, UNiagaraDataInterfaceArrayBool>(NiagaraComponent, OverrideName, ArrayData);
 }
 
 //////////////////////////////////////////////////////////////////////////
 
-TArray<float> UNiagaraDataInterfaceArrayFunctionLibrary::GetNiagaraArrayFloat(UNiagaraComponent* NiagaraSystem, FName OverrideName)
+TArray<float> UNiagaraDataInterfaceArrayFunctionLibrary::GetNiagaraArrayFloat(UNiagaraComponent* NiagaraComponent, FName OverrideName)
 {
-	return GetNiagaraArray<float, UNiagaraDataInterfaceArrayFloat>(NiagaraSystem, OverrideName);
+	return GetNiagaraArray<float, UNiagaraDataInterfaceArrayFloat>(NiagaraComponent, OverrideName);
 }
 
-TArray<FVector2D> UNiagaraDataInterfaceArrayFunctionLibrary::GetNiagaraArrayVector2D(UNiagaraComponent* NiagaraSystem, FName OverrideName)
+TArray<FVector2D> UNiagaraDataInterfaceArrayFunctionLibrary::GetNiagaraArrayVector2D(UNiagaraComponent* NiagaraComponent, FName OverrideName)
 {
-	return GetNiagaraArray<FVector2D, UNiagaraDataInterfaceArrayFloat2>(NiagaraSystem, OverrideName);
+	return GetNiagaraArray<FVector2D, UNiagaraDataInterfaceArrayFloat2>(NiagaraComponent, OverrideName);
 }
 
-TArray<FVector> UNiagaraDataInterfaceArrayFunctionLibrary::GetNiagaraArrayVector(UNiagaraComponent* NiagaraSystem, FName OverrideName)
+TArray<FVector> UNiagaraDataInterfaceArrayFunctionLibrary::GetNiagaraArrayVector(UNiagaraComponent* NiagaraComponent, FName OverrideName)
 {
-	return GetNiagaraArray<FVector, UNiagaraDataInterfaceArrayFloat3>(NiagaraSystem, OverrideName);
+	return GetNiagaraArray<FVector, UNiagaraDataInterfaceArrayFloat3>(NiagaraComponent, OverrideName);
 }
 
-TArray<FVector> UNiagaraDataInterfaceArrayFunctionLibrary::GetNiagaraArrayPosition(UNiagaraComponent* NiagaraSystem, FName OverrideName)
+TArray<FVector> UNiagaraDataInterfaceArrayFunctionLibrary::GetNiagaraArrayPosition(UNiagaraComponent* NiagaraComponent, FName OverrideName)
 {
-	const TArray<FNiagaraPosition> SimData = GetNiagaraArray<FNiagaraPosition, UNiagaraDataInterfaceArrayPosition>(NiagaraSystem, OverrideName);
+	const TArray<FNiagaraPosition> SimData = GetNiagaraArray<FNiagaraPosition, UNiagaraDataInterfaceArrayPosition>(NiagaraComponent, OverrideName);
 	
-	FNiagaraLWCConverter LwcConverter = GetLWCConverter(NiagaraSystem);
+	FNiagaraLWCConverter LwcConverter = GetLWCConverter(NiagaraComponent);
 	TArray<FVector> ConvertedData;
 	ConvertedData.SetNumUninitialized(SimData.Num());
 	for (int i = 0; i < SimData.Num(); i++)
@@ -152,125 +165,125 @@ TArray<FVector> UNiagaraDataInterfaceArrayFunctionLibrary::GetNiagaraArrayPositi
 	return ConvertedData;
 }
 
-TArray<FVector4> UNiagaraDataInterfaceArrayFunctionLibrary::GetNiagaraArrayVector4(UNiagaraComponent* NiagaraSystem, FName OverrideName)
+TArray<FVector4> UNiagaraDataInterfaceArrayFunctionLibrary::GetNiagaraArrayVector4(UNiagaraComponent* NiagaraComponent, FName OverrideName)
 {
-	return GetNiagaraArray<FVector4, UNiagaraDataInterfaceArrayFloat4>(NiagaraSystem, OverrideName);
+	return GetNiagaraArray<FVector4, UNiagaraDataInterfaceArrayFloat4>(NiagaraComponent, OverrideName);
 }
 
-TArray<FLinearColor> UNiagaraDataInterfaceArrayFunctionLibrary::GetNiagaraArrayColor(UNiagaraComponent* NiagaraSystem, FName OverrideName)
+TArray<FLinearColor> UNiagaraDataInterfaceArrayFunctionLibrary::GetNiagaraArrayColor(UNiagaraComponent* NiagaraComponent, FName OverrideName)
 {
-	return GetNiagaraArray<FLinearColor, UNiagaraDataInterfaceArrayColor>(NiagaraSystem, OverrideName);
+	return GetNiagaraArray<FLinearColor, UNiagaraDataInterfaceArrayColor>(NiagaraComponent, OverrideName);
 }
 
-TArray<FQuat> UNiagaraDataInterfaceArrayFunctionLibrary::GetNiagaraArrayQuat(UNiagaraComponent* NiagaraSystem, FName OverrideName)
+TArray<FQuat> UNiagaraDataInterfaceArrayFunctionLibrary::GetNiagaraArrayQuat(UNiagaraComponent* NiagaraComponent, FName OverrideName)
 {
-	return GetNiagaraArray<FQuat, UNiagaraDataInterfaceArrayQuat>(NiagaraSystem, OverrideName);
+	return GetNiagaraArray<FQuat, UNiagaraDataInterfaceArrayQuat>(NiagaraComponent, OverrideName);
 }
 
-TArray<int32> UNiagaraDataInterfaceArrayFunctionLibrary::GetNiagaraArrayInt32(UNiagaraComponent* NiagaraSystem, FName OverrideName)
+TArray<int32> UNiagaraDataInterfaceArrayFunctionLibrary::GetNiagaraArrayInt32(UNiagaraComponent* NiagaraComponent, FName OverrideName)
 {
-	return GetNiagaraArray<int32, UNiagaraDataInterfaceArrayInt32>(NiagaraSystem, OverrideName);
+	return GetNiagaraArray<int32, UNiagaraDataInterfaceArrayInt32>(NiagaraComponent, OverrideName);
 }
 
-TArray<bool> UNiagaraDataInterfaceArrayFunctionLibrary::GetNiagaraArrayBool(UNiagaraComponent* NiagaraSystem, FName OverrideName)
+TArray<bool> UNiagaraDataInterfaceArrayFunctionLibrary::GetNiagaraArrayBool(UNiagaraComponent* NiagaraComponent, FName OverrideName)
 {
-	return GetNiagaraArray<bool, UNiagaraDataInterfaceArrayBool>(NiagaraSystem, OverrideName);
+	return GetNiagaraArray<bool, UNiagaraDataInterfaceArrayBool>(NiagaraComponent, OverrideName);
 }
 
 //////////////////////////////////////////////////////////////////////////
 
-void UNiagaraDataInterfaceArrayFunctionLibrary::SetNiagaraArrayFloatValue(UNiagaraComponent* NiagaraSystem, FName OverrideName, int Index, float Value, bool bSizeToFit)
+void UNiagaraDataInterfaceArrayFunctionLibrary::SetNiagaraArrayFloatValue(UNiagaraComponent* NiagaraComponent, FName OverrideName, int Index, float Value, bool bSizeToFit)
 {
-	SetNiagaraArrayValue<float, UNiagaraDataInterfaceArrayFloat>(NiagaraSystem, OverrideName, Index, Value, bSizeToFit);
+	SetNiagaraArrayValue<float, UNiagaraDataInterfaceArrayFloat>(NiagaraComponent, OverrideName, Index, Value, bSizeToFit);
 }
 
-void UNiagaraDataInterfaceArrayFunctionLibrary::SetNiagaraArrayVector2DValue(UNiagaraComponent* NiagaraSystem, FName OverrideName, int Index, const FVector2D& Value, bool bSizeToFit)
+void UNiagaraDataInterfaceArrayFunctionLibrary::SetNiagaraArrayVector2DValue(UNiagaraComponent* NiagaraComponent, FName OverrideName, int Index, const FVector2D& Value, bool bSizeToFit)
 {
-	SetNiagaraArrayValue<FVector2D, UNiagaraDataInterfaceArrayFloat2>(NiagaraSystem, OverrideName, Index, Value, bSizeToFit);
+	SetNiagaraArrayValue<FVector2D, UNiagaraDataInterfaceArrayFloat2>(NiagaraComponent, OverrideName, Index, Value, bSizeToFit);
 }
 
-void UNiagaraDataInterfaceArrayFunctionLibrary::SetNiagaraArrayVectorValue(UNiagaraComponent* NiagaraSystem, FName OverrideName, int Index, const FVector& Value, bool bSizeToFit)
+void UNiagaraDataInterfaceArrayFunctionLibrary::SetNiagaraArrayVectorValue(UNiagaraComponent* NiagaraComponent, FName OverrideName, int Index, const FVector& Value, bool bSizeToFit)
 {
-	SetNiagaraArrayValue<FVector, UNiagaraDataInterfaceArrayFloat3>(NiagaraSystem, OverrideName, Index, Value, bSizeToFit);
+	SetNiagaraArrayValue<FVector, UNiagaraDataInterfaceArrayFloat3>(NiagaraComponent, OverrideName, Index, Value, bSizeToFit);
 }
 
-void UNiagaraDataInterfaceArrayFunctionLibrary::SetNiagaraArrayPositionValue(UNiagaraComponent* NiagaraSystem, FName OverrideName, int Index, const FVector& Value, bool bSizeToFit)
+void UNiagaraDataInterfaceArrayFunctionLibrary::SetNiagaraArrayPositionValue(UNiagaraComponent* NiagaraComponent, FName OverrideName, int Index, const FVector& Value, bool bSizeToFit)
 {
-	FNiagaraLWCConverter LwcConverter = GetLWCConverter(NiagaraSystem);
+	FNiagaraLWCConverter LwcConverter = GetLWCConverter(NiagaraComponent);
 	FNiagaraPosition SimulationPosition = LwcConverter.ConvertWorldToSimulationPosition(Value);
-	SetNiagaraArrayValue<FNiagaraPosition, UNiagaraDataInterfaceArrayPosition>(NiagaraSystem, OverrideName, Index, SimulationPosition, bSizeToFit);
+	SetNiagaraArrayValue<FNiagaraPosition, UNiagaraDataInterfaceArrayPosition>(NiagaraComponent, OverrideName, Index, SimulationPosition, bSizeToFit);
 }
 
-void UNiagaraDataInterfaceArrayFunctionLibrary::SetNiagaraArrayVector4Value(UNiagaraComponent* NiagaraSystem, FName OverrideName, int Index, const FVector4& Value, bool bSizeToFit)
+void UNiagaraDataInterfaceArrayFunctionLibrary::SetNiagaraArrayVector4Value(UNiagaraComponent* NiagaraComponent, FName OverrideName, int Index, const FVector4& Value, bool bSizeToFit)
 {
-	SetNiagaraArrayValue<FVector4, UNiagaraDataInterfaceArrayFloat4>(NiagaraSystem, OverrideName, Index, Value, bSizeToFit);
+	SetNiagaraArrayValue<FVector4, UNiagaraDataInterfaceArrayFloat4>(NiagaraComponent, OverrideName, Index, Value, bSizeToFit);
 }
 
-void UNiagaraDataInterfaceArrayFunctionLibrary::SetNiagaraArrayColorValue(UNiagaraComponent* NiagaraSystem, FName OverrideName, int Index, const FLinearColor& Value, bool bSizeToFit)
+void UNiagaraDataInterfaceArrayFunctionLibrary::SetNiagaraArrayColorValue(UNiagaraComponent* NiagaraComponent, FName OverrideName, int Index, const FLinearColor& Value, bool bSizeToFit)
 {
-	SetNiagaraArrayValue<FLinearColor, UNiagaraDataInterfaceArrayColor>(NiagaraSystem, OverrideName, Index, Value, bSizeToFit);
+	SetNiagaraArrayValue<FLinearColor, UNiagaraDataInterfaceArrayColor>(NiagaraComponent, OverrideName, Index, Value, bSizeToFit);
 }
 
-void UNiagaraDataInterfaceArrayFunctionLibrary::SetNiagaraArrayQuatValue(UNiagaraComponent* NiagaraSystem, FName OverrideName, int Index, const FQuat& Value, bool bSizeToFit)
+void UNiagaraDataInterfaceArrayFunctionLibrary::SetNiagaraArrayQuatValue(UNiagaraComponent* NiagaraComponent, FName OverrideName, int Index, const FQuat& Value, bool bSizeToFit)
 {
-	SetNiagaraArrayValue<FQuat, UNiagaraDataInterfaceArrayQuat>(NiagaraSystem, OverrideName, Index, Value, bSizeToFit);
+	SetNiagaraArrayValue<FQuat, UNiagaraDataInterfaceArrayQuat>(NiagaraComponent, OverrideName, Index, Value, bSizeToFit);
 }
 
-void UNiagaraDataInterfaceArrayFunctionLibrary::SetNiagaraArrayInt32Value(UNiagaraComponent* NiagaraSystem, FName OverrideName, int Index, int32 Value, bool bSizeToFit)
+void UNiagaraDataInterfaceArrayFunctionLibrary::SetNiagaraArrayInt32Value(UNiagaraComponent* NiagaraComponent, FName OverrideName, int Index, int32 Value, bool bSizeToFit)
 {
-	SetNiagaraArrayValue<int32, UNiagaraDataInterfaceArrayInt32>(NiagaraSystem, OverrideName, Index, Value, bSizeToFit);
+	SetNiagaraArrayValue<int32, UNiagaraDataInterfaceArrayInt32>(NiagaraComponent, OverrideName, Index, Value, bSizeToFit);
 }
 
-void UNiagaraDataInterfaceArrayFunctionLibrary::SetNiagaraArrayBoolValue(UNiagaraComponent* NiagaraSystem, FName OverrideName, int Index, const bool& Value, bool bSizeToFit)
+void UNiagaraDataInterfaceArrayFunctionLibrary::SetNiagaraArrayBoolValue(UNiagaraComponent* NiagaraComponent, FName OverrideName, int Index, const bool& Value, bool bSizeToFit)
 {
-	SetNiagaraArrayValue<bool, UNiagaraDataInterfaceArrayBool>(NiagaraSystem, OverrideName, Index, Value, bSizeToFit);
+	SetNiagaraArrayValue<bool, UNiagaraDataInterfaceArrayBool>(NiagaraComponent, OverrideName, Index, Value, bSizeToFit);
 }
 
 //////////////////////////////////////////////////////////////////////////
 
-float UNiagaraDataInterfaceArrayFunctionLibrary::GetNiagaraArrayFloatValue(UNiagaraComponent* NiagaraSystem, FName OverrideName, int Index)
+float UNiagaraDataInterfaceArrayFunctionLibrary::GetNiagaraArrayFloatValue(UNiagaraComponent* NiagaraComponent, FName OverrideName, int Index)
 {
-	return GetNiagaraArrayValue<float, UNiagaraDataInterfaceArrayFloat>(NiagaraSystem, OverrideName, Index);
+	return GetNiagaraArrayValue<float, UNiagaraDataInterfaceArrayFloat>(NiagaraComponent, OverrideName, Index);
 }
 
-FVector2D UNiagaraDataInterfaceArrayFunctionLibrary::GetNiagaraArrayVector2DValue(UNiagaraComponent* NiagaraSystem, FName OverrideName, int Index)
+FVector2D UNiagaraDataInterfaceArrayFunctionLibrary::GetNiagaraArrayVector2DValue(UNiagaraComponent* NiagaraComponent, FName OverrideName, int Index)
 {
-	return GetNiagaraArrayValue<FVector2D, UNiagaraDataInterfaceArrayFloat2>(NiagaraSystem, OverrideName, Index);
+	return GetNiagaraArrayValue<FVector2D, UNiagaraDataInterfaceArrayFloat2>(NiagaraComponent, OverrideName, Index);
 }
 
-FVector UNiagaraDataInterfaceArrayFunctionLibrary::GetNiagaraArrayVectorValue(UNiagaraComponent* NiagaraSystem, FName OverrideName, int Index)
+FVector UNiagaraDataInterfaceArrayFunctionLibrary::GetNiagaraArrayVectorValue(UNiagaraComponent* NiagaraComponent, FName OverrideName, int Index)
 {
-	return GetNiagaraArrayValue<FVector, UNiagaraDataInterfaceArrayFloat3>(NiagaraSystem, OverrideName, Index);
+	return GetNiagaraArrayValue<FVector, UNiagaraDataInterfaceArrayFloat3>(NiagaraComponent, OverrideName, Index);
 }
 
-FVector UNiagaraDataInterfaceArrayFunctionLibrary::GetNiagaraArrayPositionValue(UNiagaraComponent* NiagaraSystem, FName OverrideName, int Index)
+FVector UNiagaraDataInterfaceArrayFunctionLibrary::GetNiagaraArrayPositionValue(UNiagaraComponent* NiagaraComponent, FName OverrideName, int Index)
 {
-	FNiagaraPosition SimPosition = GetNiagaraArrayValue<FNiagaraPosition, UNiagaraDataInterfaceArrayPosition>(NiagaraSystem, OverrideName, Index);
-	FNiagaraLWCConverter LwcConverter = GetLWCConverter(NiagaraSystem);
+	FNiagaraPosition SimPosition = GetNiagaraArrayValue<FNiagaraPosition, UNiagaraDataInterfaceArrayPosition>(NiagaraComponent, OverrideName, Index);
+	FNiagaraLWCConverter LwcConverter = GetLWCConverter(NiagaraComponent);
 	return LwcConverter.ConvertSimulationPositionToWorld(SimPosition);
 }
 
-FVector4 UNiagaraDataInterfaceArrayFunctionLibrary::GetNiagaraArrayVector4Value(UNiagaraComponent* NiagaraSystem, FName OverrideName, int Index)
+FVector4 UNiagaraDataInterfaceArrayFunctionLibrary::GetNiagaraArrayVector4Value(UNiagaraComponent* NiagaraComponent, FName OverrideName, int Index)
 {
-	return GetNiagaraArrayValue<FVector4, UNiagaraDataInterfaceArrayFloat4>(NiagaraSystem, OverrideName, Index);
+	return GetNiagaraArrayValue<FVector4, UNiagaraDataInterfaceArrayFloat4>(NiagaraComponent, OverrideName, Index);
 }
 
-FLinearColor UNiagaraDataInterfaceArrayFunctionLibrary::GetNiagaraArrayColorValue(UNiagaraComponent* NiagaraSystem, FName OverrideName, int Index)
+FLinearColor UNiagaraDataInterfaceArrayFunctionLibrary::GetNiagaraArrayColorValue(UNiagaraComponent* NiagaraComponent, FName OverrideName, int Index)
 {
-	return GetNiagaraArrayValue<FLinearColor, UNiagaraDataInterfaceArrayColor>(NiagaraSystem, OverrideName, Index);
+	return GetNiagaraArrayValue<FLinearColor, UNiagaraDataInterfaceArrayColor>(NiagaraComponent, OverrideName, Index);
 }
 
-FQuat UNiagaraDataInterfaceArrayFunctionLibrary::GetNiagaraArrayQuatValue(UNiagaraComponent* NiagaraSystem, FName OverrideName, int Index)
+FQuat UNiagaraDataInterfaceArrayFunctionLibrary::GetNiagaraArrayQuatValue(UNiagaraComponent* NiagaraComponent, FName OverrideName, int Index)
 {
-	return GetNiagaraArrayValue<FQuat, UNiagaraDataInterfaceArrayQuat>(NiagaraSystem, OverrideName, Index);
+	return GetNiagaraArrayValue<FQuat, UNiagaraDataInterfaceArrayQuat>(NiagaraComponent, OverrideName, Index);
 }
 
-int32 UNiagaraDataInterfaceArrayFunctionLibrary::GetNiagaraArrayInt32Value(UNiagaraComponent* NiagaraSystem, FName OverrideName, int Index)
+int32 UNiagaraDataInterfaceArrayFunctionLibrary::GetNiagaraArrayInt32Value(UNiagaraComponent* NiagaraComponent, FName OverrideName, int Index)
 {
-	return GetNiagaraArrayValue<int32, UNiagaraDataInterfaceArrayInt32>(NiagaraSystem, OverrideName, Index);
+	return GetNiagaraArrayValue<int32, UNiagaraDataInterfaceArrayInt32>(NiagaraComponent, OverrideName, Index);
 }
 
-bool UNiagaraDataInterfaceArrayFunctionLibrary::GetNiagaraArrayBoolValue(UNiagaraComponent* NiagaraSystem, FName OverrideName, int Index)
+bool UNiagaraDataInterfaceArrayFunctionLibrary::GetNiagaraArrayBoolValue(UNiagaraComponent* NiagaraComponent, FName OverrideName, int Index)
 {
-	return GetNiagaraArrayValue<bool, UNiagaraDataInterfaceArrayBool>(NiagaraSystem, OverrideName, Index);
+	return GetNiagaraArrayValue<bool, UNiagaraDataInterfaceArrayBool>(NiagaraComponent, OverrideName, Index);
 }
