- Ensure we recache systems that are set for reinit

#rnx

diff --git a/Engine/Plugins/FX/Niagara/Source/Niagara/Private/NiagaraCommon.cpp b/Engine/Plugins/FX/Niagara/Source/Niagara/Private/NiagaraCommon.cpp
--- a/Engine/Plugins/FX/Niagara/Source/Niagara/Private/NiagaraCommon.cpp
+++ b/Engine/Plugins/FX/Niagara/Source/Niagara/Private/NiagaraCommon.cpp
@@ -74,37 +74,31 @@ FNiagaraSystemUpdateContext::~FNiagaraSystemUpdateContext()
 
 void FNiagaraSystemUpdateContext::CommitUpdate()
 {
-	// Sims that are sent for destruction will also get a recache
-	if (SystemSimsToDestroy.Num() > 0)
+	for (UNiagaraSystem* NiagaraSystem : SystemSimsToDestroy)
 	{
-		bool bAnyGpuSims = false;
-		for (UNiagaraSystem* NiagaraSystem : SystemSimsToDestroy)
+		if (NiagaraSystem)
 		{
-			if (NiagaraSystem)
-			{
-				bAnyGpuSims |= NiagaraSystem->HasAnyGPUEmitters();
-				FNiagaraWorldManager::DestroyAllSystemSimulations(NiagaraSystem);
-			}
-		}
-
-		if (bAnyGpuSims)
-		{
-			FlushRenderingCommands();
+			FNiagaraWorldManager::DestroyAllSystemSimulations(NiagaraSystem);
 		}
+	}
 
-		for (UNiagaraSystem* NiagaraSystem : SystemSimsToDestroy)
+	bool bNeedsWaitOnGpu = true;
+	for (UNiagaraSystem* NiagaraSystem : SystemSimsToRecache)
+	{
+		if (NiagaraSystem)
 		{
-			if (NiagaraSystem)
+			if (bNeedsWaitOnGpu == true && NiagaraSystem->HasAnyGPUEmitters())
 			{
-				NiagaraSystem->ComputeEmittersExecutionOrder();
-				NiagaraSystem->ComputeRenderersDrawOrder();
-				NiagaraSystem->CacheFromCompiledData();
+				bNeedsWaitOnGpu = false;
+				FlushRenderingCommands();
 			}
-		}
 
-		SystemSimsToDestroy.Empty();
+			NiagaraSystem->ComputeEmittersExecutionOrder();
+			NiagaraSystem->ComputeRenderersDrawOrder();
+			NiagaraSystem->CacheFromCompiledData();
+		}
 	}
-
+	
 	for (UNiagaraComponent* Comp : ComponentsToReInit)
 	{
 		if (Comp)
@@ -232,6 +226,11 @@ void FNiagaraSystemUpdateContext::AddInternal(UNiagaraComponent* Comp, bool bReI
 		}
 	}
 
+	if (bReInit || bDestroySystemSim)
+	{
+		SystemSimsToRecache.AddUnique(Comp->GetAsset());
+	}
+
 	if (bReInit && bDestroySystemSim)
 	{
 		//Always destroy the system sims on a reinit, even if we're not reactivating the component.
diff --git a/Engine/Plugins/FX/Niagara/Source/Niagara/Public/NiagaraCommon.h b/Engine/Plugins/FX/Niagara/Source/Niagara/Public/NiagaraCommon.h
--- a/Engine/Plugins/FX/Niagara/Source/Niagara/Public/NiagaraCommon.h
+++ b/Engine/Plugins/FX/Niagara/Source/Niagara/Public/NiagaraCommon.h
@@ -773,6 +773,8 @@ private:
 	TArray<TObjectPtr<UNiagaraComponent>> ComponentsToNotifySimDestroy;
 	UPROPERTY(transient)
 	TArray<TObjectPtr<UNiagaraSystem>> SystemSimsToDestroy;
+	UPROPERTY(transient)
+	TArray<TObjectPtr<UNiagaraSystem>> SystemSimsToRecache;
 
 	bool bDestroyOnAdd;
 	bool bOnlyActive;
