- Replace CastFieldChecked and instead check with more details to diagnose the bad field type

#rnx

diff --git a/Engine/Plugins/FX/Niagara/Source/Niagara/Public/NiagaraTypes.h b/Engine/Plugins/FX/Niagara/Source/Niagara/Public/NiagaraTypes.h
--- a/Engine/Plugins/FX/Niagara/Source/Niagara/Public/NiagaraTypes.h
+++ b/Engine/Plugins/FX/Niagara/Source/Niagara/Public/NiagaraTypes.h
@@ -442,13 +442,13 @@ private:
 				Layout.Int32ComponentByteOffsets.Add(PropOffset);
 			}
 			//Should be able to support double easily enough
-			else if (FStructProperty* StructProp = CastFieldChecked<FStructProperty>(Property))
+			else if (FStructProperty* StructProp = CastField<FStructProperty>(Property))
 			{
 				GenerateLayoutInfoInternal(Layout, FNiagaraTypeHelper::FindNiagaraFriendlyTopLevelStruct(StructProp->Struct, ENiagaraStructConversion::Simulation), PropOffset);
 			}
 			else
 			{
-				check(false);
+				checkf(false, TEXT("Property(%s) Class(%s) is not a supported type"), *Property->GetName(), *Property->GetClass()->GetName());
 			}
 		}
 	}
