- Modify test for game / render thread to match what console variables do

#rnx

diff --git a/Engine/Plugins/FX/Niagara/Source/Niagara/Private/NiagaraSystemRenderData.cpp b/Engine/Plugins/FX/Niagara/Source/Niagara/Private/NiagaraSystemRenderData.cpp
--- a/Engine/Plugins/FX/Niagara/Source/Niagara/Private/NiagaraSystemRenderData.cpp
+++ b/Engine/Plugins/FX/Niagara/Source/Niagara/Private/NiagaraSystemRenderData.cpp
@@ -104,9 +104,19 @@ FPrimitiveViewRelevance FNiagaraSystemRenderData::GetViewRelevance(const FSceneV
 uint32 FNiagaraSystemRenderData::GetDynamicDataSize() const
 {
 	// only the render thread should use the RT data, the game thread and other parallel running tasks use the GT data
-	const TArray<FNiagaraRenderer*>& Emitter_Renderers = IsInParallelRenderingThread() ? EmitterRenderers_RT :EmitterRenderers_GT;
+	const TArray<FNiagaraRenderer*>* LocalEmitterRenderers;
+	if ( IsInParallelGameThread() || IsInGameThread() )
+	{
+		LocalEmitterRenderers = &EmitterRenderers_GT;
+	}
+	else
+	{
+		check(IsInParallelRenderingThread());
+		LocalEmitterRenderers = &EmitterRenderers_RT;
+	}
+
 	uint32 Size = 0;
-	for (const auto Renderer : Emitter_Renderers)
+	for (const auto Renderer : *LocalEmitterRenderers)
 	{
 		if (Renderer)
 		{
