Display overview nodes in a simplified manner when zoomed out

diff --git a/Engine/Plugins/FX/Niagara/Source/NiagaraEditorWidgets/NiagaraEditorWidgets.Build.cs b/Engine/Plugins/FX/Niagara/Source/NiagaraEditorWidgets/NiagaraEditorWidgets.Build.cs
--- a/Engine/Plugins/FX/Niagara/Source/NiagaraEditorWidgets/NiagaraEditorWidgets.Build.cs
+++ b/Engine/Plugins/FX/Niagara/Source/NiagaraEditorWidgets/NiagaraEditorWidgets.Build.cs
@@ -32,7 +32,8 @@ public class NiagaraEditorWidgets : ModuleRules
 			"ApplicationCore",
 			"CurveEditor",
 			"DesktopPlatform",
-			"ToolWidgets"
+			"ToolWidgets",
+			"KismetWidgets"
 		});
 
 		PrivateIncludePathModuleNames.AddRange(new string[] {
diff --git a/Engine/Plugins/FX/Niagara/Source/NiagaraEditorWidgets/Private/NiagaraEditorWidgetsStyle.cpp b/Engine/Plugins/FX/Niagara/Source/NiagaraEditorWidgets/Private/NiagaraEditorWidgetsStyle.cpp
--- a/Engine/Plugins/FX/Niagara/Source/NiagaraEditorWidgets/Private/NiagaraEditorWidgetsStyle.cpp
+++ b/Engine/Plugins/FX/Niagara/Source/NiagaraEditorWidgets/Private/NiagaraEditorWidgetsStyle.cpp
@@ -150,6 +150,11 @@ FNiagaraEditorWidgetsStyle::FNiagaraEditorWidgetsStyle() : FSlateStyleSet("Niaga
 		.SetFont(SystemOverviewItemFont);
 	Set("NiagaraEditor.SystemOverview.ItemText", SystemOverviewItemText);
 
+	FSlateFontInfo SystemOverviewZoomedOutNodeFont = DEFAULT_FONT("Regular", 45);
+	FTextBlockStyle SystemOverviewZoomedOutNodeText = FTextBlockStyle(NormalText)
+		.SetFont(SystemOverviewZoomedOutNodeFont);
+	Set("NiagaraEditor.SystemOverview.ZoomedOutNodeFont", SystemOverviewZoomedOutNodeText);
+
 	FSlateFontInfo SystemOverviewAlternateItemFont = DEFAULT_FONT("Italic", 9);
 	FTextBlockStyle SystemOverviewAlternateItemText = FTextBlockStyle(NormalText)
 		.SetFont(SystemOverviewAlternateItemFont);
diff --git a/Engine/Plugins/FX/Niagara/Source/NiagaraEditorWidgets/Private/SNiagaraOverviewStackNode.cpp b/Engine/Plugins/FX/Niagara/Source/NiagaraEditorWidgets/Private/SNiagaraOverviewStackNode.cpp
--- a/Engine/Plugins/FX/Niagara/Source/NiagaraEditorWidgets/Private/SNiagaraOverviewStackNode.cpp
+++ b/Engine/Plugins/FX/Niagara/Source/NiagaraEditorWidgets/Private/SNiagaraOverviewStackNode.cpp
@@ -33,8 +33,8 @@
 #include "Materials/Material.h"
 #include "Framework/MultiBox/MultiBoxBuilder.h"
 #include "Styling/StyleColors.h"
-#include "Widgets/Text/SInlineEditableTextBlock.h"
 #include "Styling/AppStyle.h"
+#include "KismetWidgets/Public/SLevelOfDetailBranchNode.h"
 
 #define LOCTEXT_NAMESPACE "NiagaraOverviewStackNode"
 
@@ -308,6 +308,11 @@ void SNiagaraOverviewStackNode::Tick(const FGeometry& AllottedGeometry, const do
 			FillTopContentBar();
 			bTopContentBarRefreshPending = false;
 		}
+
+		if (ScalabilityWrapper != nullptr && !UseLowDetailNodeContent())
+		{
+			LastHighDetailSize = ScalabilityWrapper->GetTickSpaceGeometry().Size;
+		}
 	}
 }
 
@@ -455,10 +460,31 @@ TSharedRef<SWidget> SNiagaraOverviewStackNode::CreateNodeContentArea()
 		];
 	}
 
-	TSharedPtr<SOverlay> ScalabilityWrapper = SNew(SOverlay)
+	TSharedRef<SWidget> DetailedContent = SNew(SLevelOfDetailBranchNode)
+		.UseLowDetailSlot(this, &SNiagaraOverviewStackNode::UseLowDetailNodeContent)
+		.LowDetail()
+		[
+			SNew(SBox)
+			.WidthOverride(this, &SNiagaraOverviewStackNode::GetLowDetailDesiredWidth)
+			.HeightOverride(this, &SNiagaraOverviewStackNode::GetLowDetailDesiredHeight)
+			[
+				SNew(STextBlock)
+				.Text(this, &SNiagaraOverviewStackNode::GetEditableNodeTitleAsText)
+				.TextStyle(&FNiagaraEditorWidgetsStyle::Get().GetWidgetStyle<FTextBlockStyle>("NiagaraEditor.SystemOverview.ZoomedOutNodeFont"))
+				.Visibility(EVisibility::HitTestInvisible)
+				.Clipping(EWidgetClipping::Inherit)
+				.RenderTransform(FSlateRenderTransform(Concatenate(FQuat2D(FMath::DegreesToRadians(90.0f)), FVector2D(100, 25))))
+			]
+		]
+		.HighDetail()
+		[
+			NodeWidget
+		];
+
+	ScalabilityWrapper = SNew(SOverlay)
 	+ SOverlay::Slot()
 	[
-		NodeWidget
+		DetailedContent
 	]
 	+ SOverlay::Slot()
 	.Padding(0, 1)
@@ -468,11 +494,40 @@ TSharedRef<SWidget> SNiagaraOverviewStackNode::CreateNodeContentArea()
 		.BorderImage(FNiagaraEditorStyle::Get().GetBrush("NiagaraEditor.SystemOverview.ExcludedFromScalability.NodeBody"))
 		.BorderBackgroundColor(TAttribute<FSlateColor>(this, &SNiagaraOverviewStackNode::GetScalabilityTintAlpha))
 	];
-	
+
 	return ScalabilityWrapper.ToSharedRef();
 
 }
 
+bool SNiagaraOverviewStackNode::UseLowDetailNodeContent() const
+{
+	if (LastHighDetailSize.IsNearlyZero())
+	{
+		return false;
+	}
+	
+	if (const SGraphPanel* MyOwnerPanel = GetOwnerPanel().Get())
+	{
+		return (MyOwnerPanel->GetCurrentLOD() <= EGraphRenderingLOD::LowDetail);
+	}
+	return false;
+}
+
+FVector2D SNiagaraOverviewStackNode::GetLowDetailDesiredSize() const
+{
+	return LastHighDetailSize;
+}
+
+FOptionalSize SNiagaraOverviewStackNode::GetLowDetailDesiredWidth() const
+{
+	return LastHighDetailSize.X;
+}
+
+FOptionalSize SNiagaraOverviewStackNode::GetLowDetailDesiredHeight() const
+{
+	return LastHighDetailSize.Y;
+}
+
 void SNiagaraOverviewStackNode::StackViewModelStructureChanged(ENiagaraStructureChangedFlags Flags)
 {
 	bTopContentBarRefreshPending = true;
diff --git a/Engine/Plugins/FX/Niagara/Source/NiagaraEditorWidgets/Private/SNiagaraOverviewStackNode.h b/Engine/Plugins/FX/Niagara/Source/NiagaraEditorWidgets/Private/SNiagaraOverviewStackNode.h
--- a/Engine/Plugins/FX/Niagara/Source/NiagaraEditorWidgets/Private/SNiagaraOverviewStackNode.h
+++ b/Engine/Plugins/FX/Niagara/Source/NiagaraEditorWidgets/Private/SNiagaraOverviewStackNode.h
@@ -30,10 +30,15 @@ protected:
 	virtual TSharedRef<SWidget> CreateTitleWidget(TSharedPtr<SNodeTitle> NodeTitle) override;
 	virtual TSharedRef<SWidget> CreateTitleRightWidget() override;
 	virtual TSharedRef<SWidget> CreateNodeContentArea() override;
+	virtual bool IsHidingPinWidgets() const override { return UseLowDetailNodeContent(); }
 	void StackViewModelStructureChanged(ENiagaraStructureChangedFlags Flags);
 	void StackViewModelDataObjectChanged(TArray<UObject*> ChangedObjects, ENiagaraDataObjectChange ChangeType);
 	void FillTopContentBar();
 	void OnMaterialCompiled(class UMaterialInterface* MaterialInterface);
+	bool UseLowDetailNodeContent() const;
+	FVector2D GetLowDetailDesiredSize() const;
+	FOptionalSize GetLowDetailDesiredWidth() const;
+	FOptionalSize GetLowDetailDesiredHeight() const;
 
 	void CreateBottomSummaryExpander();
 private:
@@ -86,10 +91,14 @@ private:
 	/** The top content bar houses the isolate button and the thumbnails */
 	TSharedPtr<SHorizontalBox> TopContentBar;
 	TSharedPtr<SWidget> BottomSummaryExpander;
+	TSharedPtr<SOverlay> ScalabilityWrapper;
 
 	TArray<UNiagaraStackEntry*> PreviewStackEntries;
 	bool bIsHoveringThumbnail = false;
 	bool bTopContentBarRefreshPending = false;
 	int32 CurrentIssueIndex = 0;
 	bool bScalabilityModeActive = false;
+	
+	/** Cached size from when we last drew at high detail */
+	FVector2D LastHighDetailSize = FVector2D::ZeroVector;
 };
\ No newline at end of file
