Added setting to zoom-fit system assets on open

diff --git a/Engine/Plugins/FX/Niagara/Source/Niagara/Public/NiagaraSettings.h b/Engine/Plugins/FX/Niagara/Source/Niagara/Public/NiagaraSettings.h
--- a/Engine/Plugins/FX/Niagara/Source/Niagara/Public/NiagaraSettings.h
+++ b/Engine/Plugins/FX/Niagara/Source/Niagara/Public/NiagaraSettings.h
@@ -102,6 +102,10 @@ class NIAGARA_API UNiagaraSettings : public UDeveloperSettings
 	UPROPERTY(config, EditAnywhere, Category = Viewport)
 	bool bSystemViewportInOrbitMode = true;
 
+	/** If true then the system editor will zoom to fit all emitters when opening an asset. */
+	UPROPERTY(config, EditAnywhere, Category = Niagara)
+	bool bAlwaysZoomToFitSystemGraph = true;
+
 	/** If true then the "link input" menu will also show variables of different types, as long as there is a conversion script for them. */
 	UPROPERTY(config, EditAnywhere, Category = Niagara)
 	bool bShowConvertibleInputsInStack = false;
diff --git a/Engine/Plugins/FX/Niagara/Source/NiagaraEditorWidgets/Private/SNiagaraOverviewGraph.cpp b/Engine/Plugins/FX/Niagara/Source/NiagaraEditorWidgets/Private/SNiagaraOverviewGraph.cpp
--- a/Engine/Plugins/FX/Niagara/Source/NiagaraEditorWidgets/Private/SNiagaraOverviewGraph.cpp
+++ b/Engine/Plugins/FX/Niagara/Source/NiagaraEditorWidgets/Private/SNiagaraOverviewGraph.cpp
@@ -18,6 +18,7 @@
 #include "NiagaraEditorModule.h"
 #include "GraphEditorActions.h"
 #include "NiagaraEditorWidgetsStyle.h"
+#include "NiagaraSettings.h"
 #include "Framework/MultiBox/MultiBoxBuilder.h"
 #include "ScopedTransaction.h"
 #include "Framework/Application/SlateApplication.h"
@@ -116,9 +117,10 @@ void SNiagaraOverviewGraph::Construct(const FArguments& InArgs, TSharedRef<FNiag
 	GraphEditor->SetNodeFactory(MakeShared<FNiagaraOverviewGraphNodeFactory>());
 
 	GraphEditor->GetCurrentGraph()->AddOnGraphChangedHandler(FOnGraphChanged::FDelegate::CreateSP(this, &SNiagaraOverviewGraph::OnNodesCreated));
-	
+
+	const UNiagaraSettings* NiagaraSettings = GetDefault<UNiagaraSettings>();
 	FNiagaraGraphViewSettings ViewSettings = ViewModel->GetViewSettings();
-	if (ViewSettings.IsValid())
+	if (NiagaraSettings->bAlwaysZoomToFitSystemGraph == false && ViewSettings.IsValid())
 	{
 		GraphEditor->SetViewLocation(ViewSettings.GetLocation(), ViewSettings.GetZoom());
 		ZoomToFitFrameDelay = 0;
