- Don't set variants on PIE / game instances
- There is not point and the parameter stores get confused about which DI to use with the GPU tick

#rnx

diff --git a/Engine/Plugins/FX/Niagara/Source/Niagara/Private/NiagaraDataInterfaceArrayFunctionLibrary.cpp b/Engine/Plugins/FX/Niagara/Source/Niagara/Private/NiagaraDataInterfaceArrayFunctionLibrary.cpp
--- a/Engine/Plugins/FX/Niagara/Source/Niagara/Private/NiagaraDataInterfaceArrayFunctionLibrary.cpp
+++ b/Engine/Plugins/FX/Niagara/Source/Niagara/Private/NiagaraDataInterfaceArrayFunctionLibrary.cpp
@@ -14,6 +14,14 @@ template<typename TArrayType, typename TDataInterace>
 void SetNiagaraVariantArray(UNiagaraComponent* NiagaraComponent, FName OverrideName, TDataInterace* ExistingDataInterface)
 {
 #if WITH_EDITOR
+	// We only need to do this for editor instances of the component as we are storing instance data on them
+	// For runtime instances they already have a unique copy of the data interface that we are modifying
+	UWorld* World = NiagaraComponent->GetWorld();
+	if (World == nullptr || World->IsGameWorld())
+	{
+		return;
+	}
+
 	TDataInterace* VariantDataInterface = CastChecked<TDataInterace>(DuplicateObject(ExistingDataInterface, NiagaraComponent));
 	auto* ExistingProxy = static_cast<FNDIArrayProxyImpl<TArrayType, TDataInterace>*>(ExistingDataInterface->GetProxy());
 	VariantDataInterface->GetArrayReference() = ExistingProxy->GetArrayData();
