- Use TSoftObjectPtr over TObjectPtr to allow for cross level references

diff --git a/Engine/Plugins/FX/Niagara/Source/Niagara/Classes/NiagaraDataInterfaceSkeletalMesh.h b/Engine/Plugins/FX/Niagara/Source/Niagara/Classes/NiagaraDataInterfaceSkeletalMesh.h
--- a/Engine/Plugins/FX/Niagara/Source/Niagara/Classes/NiagaraDataInterfaceSkeletalMesh.h
+++ b/Engine/Plugins/FX/Niagara/Source/Niagara/Classes/NiagaraDataInterfaceSkeletalMesh.h
@@ -730,8 +730,13 @@ public:
 #endif
 
 	/** The source actor from which to sample. Takes precedence over the direct mesh. Note that this can only be set when used as a user variable on a component in the world.*/
-	UPROPERTY(EditAnywhere, Category = "Mesh")
-	TObjectPtr<AActor> Source;
+	UPROPERTY(EditAnywhere, Category = "Mesh", meta = (DisplayName = "Source Actor"))
+	TSoftObjectPtr<AActor> SoftSourceActor;
+
+#if WITH_EDITORONLY_DATA
+	UPROPERTY()
+	TObjectPtr<AActor> Source_DEPRECATED;
+#endif
 
 	/** Reference to a user parameter if we're reading one. */
 	UPROPERTY(EditAnywhere, Category = "Mesh")
diff --git a/Engine/Plugins/FX/Niagara/Source/Niagara/Private/DataInterface/NiagaraDataInterfaceStaticMesh.cpp b/Engine/Plugins/FX/Niagara/Source/Niagara/Private/DataInterface/NiagaraDataInterfaceStaticMesh.cpp
--- a/Engine/Plugins/FX/Niagara/Source/Niagara/Private/DataInterface/NiagaraDataInterfaceStaticMesh.cpp
+++ b/Engine/Plugins/FX/Niagara/Source/Niagara/Private/DataInterface/NiagaraDataInterfaceStaticMesh.cpp
@@ -1407,6 +1407,18 @@ void UNiagaraDataInterfaceStaticMesh::PostInitProperties()
 	}
 }
 
+void UNiagaraDataInterfaceStaticMesh::PostLoad()
+{
+	Super::PostLoad();
+
+#if WITH_EDITORONLY_DATA
+	if ( Source_DEPRECATED != nullptr )
+	{
+		SoftSourceActor = Source_DEPRECATED;
+	}
+#endif
+}
+
 #if WITH_EDITOR
 void UNiagaraDataInterfaceStaticMesh::PostEditChangeProperty(struct FPropertyChangedEvent& PropertyChangedEvent)
 {
@@ -1419,7 +1431,7 @@ void UNiagaraDataInterfaceStaticMesh::PostEditChangeProperty(struct FPropertyCha
 		if (SourceMode != ENDIStaticMesh_SourceMode::Default && SourceMode != ENDIStaticMesh_SourceMode::Source)
 		{
 			// Ensure we don't have any reference to a source actor that we'll never use
-			Source = nullptr;
+			SoftSourceActor = nullptr;
 			SourceComponent = nullptr;
 		}
 	
@@ -1438,7 +1450,7 @@ bool UNiagaraDataInterfaceStaticMesh::CanEditChange(const FProperty* InProperty)
 		return false;
 	}
 
-	if (InProperty->GetFName() == GET_MEMBER_NAME_CHECKED(UNiagaraDataInterfaceStaticMesh, Source) &&
+	if (InProperty->GetFName() == GET_MEMBER_NAME_CHECKED(UNiagaraDataInterfaceStaticMesh, SoftSourceActor) &&
 		SourceMode != ENDIStaticMesh_SourceMode::Default &&
 		SourceMode != ENDIStaticMesh_SourceMode::Source)
 	{
@@ -2520,7 +2532,7 @@ bool UNiagaraDataInterfaceStaticMesh::Equals(const UNiagaraDataInterface* Other)
 	}
 	const UNiagaraDataInterfaceStaticMesh* OtherTyped = CastChecked<const UNiagaraDataInterfaceStaticMesh>(Other);
 	return OtherTyped->SourceMode == SourceMode &&
-		OtherTyped->Source == Source &&
+		OtherTyped->SoftSourceActor == SoftSourceActor &&
 		OtherTyped->SourceComponent == SourceComponent &&
 		OtherTyped->DefaultMesh == DefaultMesh &&
 		OtherTyped->SectionFilter.AllowedMaterialSlots == SectionFilter.AllowedMaterialSlots &&
@@ -2537,7 +2549,7 @@ bool UNiagaraDataInterfaceStaticMesh::CopyToInternal(UNiagaraDataInterface* Dest
 
 	UNiagaraDataInterfaceStaticMesh* OtherTyped = CastChecked<UNiagaraDataInterfaceStaticMesh>(Destination);
 	OtherTyped->SourceMode = SourceMode;
-	OtherTyped->Source = Source;
+	OtherTyped->SoftSourceActor = SoftSourceActor;
 	OtherTyped->SourceComponent = SourceComponent;
 	OtherTyped->DefaultMesh = DefaultMesh;
 #if WITH_EDITORONLY_DATA
@@ -2552,7 +2564,8 @@ bool UNiagaraDataInterfaceStaticMesh::CopyToInternal(UNiagaraDataInterface* Dest
 #if WITH_EDITOR
 void UNiagaraDataInterfaceStaticMesh::GetFeedback(UNiagaraSystem* Asset, UNiagaraComponent* Component, TArray<FNiagaraDataInterfaceError>& OutErrors, TArray<FNiagaraDataInterfaceFeedback>& OutWarnings, TArray<FNiagaraDataInterfaceFeedback>& OutInfo)
 {
-	if (Source == nullptr && DefaultMesh != nullptr && !DefaultMesh->bAllowCPUAccess)
+	AActor* SourceActor = SoftSourceActor.Get();
+	if (SourceActor == nullptr && DefaultMesh != nullptr && !DefaultMesh->bAllowCPUAccess)
 	{
 		FNiagaraDataInterfaceError CPUAccessNotAllowedError(FText::Format(LOCTEXT("CPUAccessNotAllowedError", "This mesh needs CPU access in order to be used properly.({0})"), FText::FromString(DefaultMesh->GetName())),
 			LOCTEXT("CPUAccessNotAllowedErrorSummary", "CPU access error"),
@@ -2566,7 +2579,7 @@ void UNiagaraDataInterfaceStaticMesh::GetFeedback(UNiagaraSystem* Asset, UNiagar
 		OutErrors.Add(CPUAccessNotAllowedError);
 	}
 
-	bool bHasNoMeshAssignedWarning = (Source == nullptr && DefaultMesh == nullptr);
+	bool bHasNoMeshAssignedWarning = (SourceActor == nullptr && DefaultMesh == nullptr);
 #if WITH_EDITORONLY_DATA
 	if (bHasNoMeshAssignedWarning)
 	{
@@ -2672,9 +2685,9 @@ UStaticMesh* UNiagaraDataInterfaceStaticMesh::GetStaticMesh(USceneComponent*& Ou
 	{
 		FoundMeshComponent = SourceComponent;
 	}
-	else if (bTrySource && Source)
+	else if (bTrySource && SoftSourceActor.Get())
 	{
-		FoundMeshComponent = FindActorMeshComponent(Source);
+		FoundMeshComponent = FindActorMeshComponent(SoftSourceActor.Get());
 	}
 	else if (bTryAttachParent && SystemInstance)
 	{
@@ -2736,7 +2749,7 @@ void UNiagaraDataInterfaceStaticMesh::SetSourceComponentFromBlueprints(UStaticMe
 	// NOTE: When ChangeId changes the next tick will be skipped and a reset of the per-instance data will be initiated. 
 	++ChangeId;
 	SourceComponent = ComponentToUse;
-	Source = ComponentToUse->GetOwner();
+	SoftSourceActor = ComponentToUse->GetOwner();
 }
 
 void UNiagaraDataInterfaceStaticMesh::SetDefaultMeshFromBlueprints(UStaticMesh* MeshToUse)
@@ -2744,7 +2757,7 @@ void UNiagaraDataInterfaceStaticMesh::SetDefaultMeshFromBlueprints(UStaticMesh*
 	// NOTE: When ChangeId changes the next tick will be skipped and a reset of the per-instance data will be initiated. 
 	++ChangeId;
 	SourceComponent = nullptr;
-	Source = nullptr;
+	SoftSourceActor = nullptr;
 	DefaultMesh = MeshToUse;
 }
 
diff --git a/Engine/Plugins/FX/Niagara/Source/Niagara/Private/DataInterface/NiagaraDataInterfaceStaticMesh.h b/Engine/Plugins/FX/Niagara/Source/Niagara/Private/DataInterface/NiagaraDataInterfaceStaticMesh.h
--- a/Engine/Plugins/FX/Niagara/Source/Niagara/Private/DataInterface/NiagaraDataInterfaceStaticMesh.h
+++ b/Engine/Plugins/FX/Niagara/Source/Niagara/Private/DataInterface/NiagaraDataInterfaceStaticMesh.h
@@ -73,8 +73,13 @@ public:
 	TObjectPtr<UStaticMesh> DefaultMesh;
 
 	/** The source actor from which to sample. Takes precedence over the direct mesh. Note that this can only be set when used as a user variable on a component in the world. */
-	UPROPERTY(EditAnywhere, Category = "Mesh")
-	TObjectPtr<AActor> Source;
+	UPROPERTY(EditAnywhere, Category = "Mesh", meta = (DisplayName = "Source Actor"))
+	TSoftObjectPtr<AActor> SoftSourceActor;
+
+#if WITH_EDITORONLY_DATA
+	UPROPERTY()
+	TObjectPtr<AActor> Source_DEPRECATED;
+#endif
 
 	/** The source component from which to sample. Takes precedence over the direct mesh. Not exposed to the user, only indirectly accessible from blueprints. */
 	UPROPERTY(Transient)
@@ -97,6 +102,7 @@ public:
 
 	//~ UObject interface
 	virtual void PostInitProperties() override;
+	virtual void PostLoad() override;
 #if WITH_EDITOR
 	virtual void PostEditChangeProperty(struct FPropertyChangedEvent& PropertyChangedEvent) override;
 	virtual bool CanEditChange(const FProperty* InProperty) const override;
diff --git a/Engine/Plugins/FX/Niagara/Source/Niagara/Private/NiagaraDataInterfaceSkeletalMesh.cpp b/Engine/Plugins/FX/Niagara/Source/Niagara/Private/NiagaraDataInterfaceSkeletalMesh.cpp
--- a/Engine/Plugins/FX/Niagara/Source/Niagara/Private/NiagaraDataInterfaceSkeletalMesh.cpp
+++ b/Engine/Plugins/FX/Niagara/Source/Niagara/Private/NiagaraDataInterfaceSkeletalMesh.cpp
@@ -1800,9 +1800,9 @@ USkeletalMesh* UNiagaraDataInterfaceSkeletalMesh::GetSkeletalMesh(FNiagaraSystem
 	{
 		FoundSkelComp = SourceComponent;
 	}
-	else if (bTrySource && Source)
+	else if (bTrySource && SoftSourceActor.Get())
 	{
-		FoundSkelComp = FindActorSkelMeshComponent(Source);
+		FoundSkelComp = FindActorSkelMeshComponent(SoftSourceActor.Get());
 	}
 	else if (bTryAttachParent && AttachComponent)
 	{
@@ -2468,7 +2468,7 @@ UNiagaraDataInterfaceSkeletalMesh::UNiagaraDataInterfaceSkeletalMesh(FObjectInit
 #if WITH_EDITORONLY_DATA
 	  , PreviewMesh(nullptr)
 #endif
-	  , Source(nullptr)
+	  , SoftSourceActor(nullptr)
 	  , SourceComponent(nullptr)
       , SkinningMode(ENDISkeletalMesh_SkinningMode::SkinOnTheFly)
 	  , WholeMeshLOD(INDEX_NONE)
@@ -2515,6 +2515,12 @@ void UNiagaraDataInterfaceSkeletalMesh::PostLoad()
 		LocalPreviewMesh->ConditionalPostLoad();
 	}
 #endif
+#if WITH_EDITORONLY_DATA
+	if (Source_DEPRECATED != nullptr)
+	{
+		SoftSourceActor = Source_DEPRECATED;
+	}
+#endif
 }
 
 #if WITH_EDITOR
@@ -2529,7 +2535,7 @@ void UNiagaraDataInterfaceSkeletalMesh::PostEditChangeProperty(FPropertyChangedE
 		SourceMode != ENDISkeletalMesh_SourceMode::Source)
 	{
 		// Clear out any source that is set to prevent unnecessary references, since we won't even consider them
-		Source = nullptr;
+		SoftSourceActor = nullptr;
 		SourceComponent = nullptr;
 	}
 }
@@ -2541,7 +2547,7 @@ bool UNiagaraDataInterfaceSkeletalMesh::CanEditChange(const FProperty* InPropert
 		return false;
 	}
 
-	if (InProperty->GetFName() == GET_MEMBER_NAME_CHECKED(UNiagaraDataInterfaceSkeletalMesh, Source) &&
+	if (InProperty->GetFName() == GET_MEMBER_NAME_CHECKED(UNiagaraDataInterfaceSkeletalMesh, SoftSourceActor) &&
 		SourceMode != ENDISkeletalMesh_SourceMode::Default &&
 		SourceMode != ENDISkeletalMesh_SourceMode::Source)
 	{
@@ -2622,7 +2628,7 @@ bool UNiagaraDataInterfaceSkeletalMesh::CopyToInternal(UNiagaraDataInterface* De
 
 	UNiagaraDataInterfaceSkeletalMesh* OtherTyped = CastChecked<UNiagaraDataInterfaceSkeletalMesh>(Destination);
 	OtherTyped->SourceMode = SourceMode;
-	OtherTyped->Source = Source;
+	OtherTyped->SoftSourceActor = SoftSourceActor;
 	OtherTyped->MeshUserParameter = MeshUserParameter;
 	OtherTyped->SourceComponent = SourceComponent;
 	OtherTyped->SkinningMode = SkinningMode;
@@ -2651,7 +2657,7 @@ bool UNiagaraDataInterfaceSkeletalMesh::Equals(const UNiagaraDataInterface* Othe
 #if WITH_EDITORONLY_DATA
 		OtherTyped->PreviewMesh == PreviewMesh &&
 #endif
-		OtherTyped->Source == Source &&
+		OtherTyped->SoftSourceActor == SoftSourceActor &&
 		OtherTyped->MeshUserParameter == MeshUserParameter &&
 		OtherTyped->SourceComponent == SourceComponent &&
 		OtherTyped->SkinningMode == SkinningMode &&
@@ -2846,7 +2852,7 @@ void UNiagaraDataInterfaceSkeletalMesh::GetFeedback(UNiagaraSystem* Asset, UNiag
 	}
 #endif
 
-	if (Source == nullptr && bHasNoMeshAssignedWarning)
+	if (SoftSourceActor.Get() == nullptr && bHasNoMeshAssignedWarning)
 	{
 		FNiagaraDataInterfaceFeedback NoMeshAssignedError(LOCTEXT("NoMeshAssignedError", "This Data Interface should be assigned a skeletal mesh to operate correctly."),
 			LOCTEXT("NoMeshAssignedErrorSummary", "No mesh assigned warning"),
@@ -3620,7 +3626,7 @@ void UNiagaraDataInterfaceSkeletalMesh::SetSourceComponentFromBlueprints(USkelet
 	// NOTE: When ChangeId changes the next tick will be skipped and a reset of the per-instance data will be initiated.
 	ChangeId++;
 	SourceComponent = ComponentToUse;
-	Source = ComponentToUse->GetOwner();
+	SoftSourceActor = ComponentToUse->GetOwner();
 }
 
 void UNiagaraDataInterfaceSkeletalMesh::SetSamplingRegionsFromBlueprints(const TArray<FName>& InSamplingRegions)
diff --git a/Engine/Plugins/Runtime/HairStrands/Source/HairStrandsCore/Private/Niagara/NiagaraDataInterfacePhysicsAsset.cpp b/Engine/Plugins/Runtime/HairStrands/Source/HairStrandsCore/Private/Niagara/NiagaraDataInterfacePhysicsAsset.cpp
--- a/Engine/Plugins/Runtime/HairStrands/Source/HairStrandsCore/Private/Niagara/NiagaraDataInterfacePhysicsAsset.cpp
+++ b/Engine/Plugins/Runtime/HairStrands/Source/HairStrandsCore/Private/Niagara/NiagaraDataInterfacePhysicsAsset.cpp
@@ -683,7 +683,7 @@ void FNDIPhysicsAssetProxy::ResetData(FRHICommandList& RHICmdList, const FNiagar
 UNiagaraDataInterfacePhysicsAsset::UNiagaraDataInterfacePhysicsAsset(FObjectInitializer const& ObjectInitializer)
 	: Super(ObjectInitializer)
 	, DefaultSource(nullptr)
-	, SourceActor(nullptr)
+	, SoftSourceActor(nullptr)
 	, SourceComponents()
 	, PhysicsAssets()
 {
@@ -770,7 +770,8 @@ void UNiagaraDataInterfacePhysicsAsset::ExtractSourceComponent(FNiagaraSystemIns
 		{
 			// The binding exists, but no object is bound. Not warning here in case the user knows what they're doing.
 		}
-	} else if (SourceActor)
+	}
+	else if (AActor* SourceActor = SoftSourceActor.Get())
 	{
 		ASkeletalMeshActor* SkeletalMeshActor = Cast<ASkeletalMeshActor>(SourceActor);
 		if (SkeletalMeshActor != nullptr)
@@ -911,7 +912,7 @@ bool UNiagaraDataInterfacePhysicsAsset::CopyToInternal(UNiagaraDataInterface* De
 
 	UNiagaraDataInterfacePhysicsAsset* OtherTyped = CastChecked<UNiagaraDataInterfacePhysicsAsset>(Destination);
 	OtherTyped->PhysicsAssets = PhysicsAssets;
-	OtherTyped->SourceActor = SourceActor;
+	OtherTyped->SoftSourceActor = SoftSourceActor;
 	OtherTyped->SourceComponents = SourceComponents;
 	OtherTyped->DefaultSource = DefaultSource;
 	OtherTyped->MeshUserParameter = MeshUserParameter;
@@ -928,7 +929,7 @@ bool UNiagaraDataInterfacePhysicsAsset::Equals(const UNiagaraDataInterface* Othe
 	const UNiagaraDataInterfacePhysicsAsset* OtherTyped = CastChecked<const UNiagaraDataInterfacePhysicsAsset>(Other);
 
 	return  (OtherTyped->PhysicsAssets == PhysicsAssets) && 
-		(OtherTyped->SourceActor == SourceActor) && 
+		(OtherTyped->SoftSourceActor == SoftSourceActor) &&
 		(OtherTyped->SourceComponents == SourceComponents) && 
 		(OtherTyped->DefaultSource == DefaultSource) &&
 		(OtherTyped->MeshUserParameter == MeshUserParameter);
@@ -945,6 +946,18 @@ void UNiagaraDataInterfacePhysicsAsset::PostInitProperties()
 	}
 }
 
+void UNiagaraDataInterfacePhysicsAsset::PostLoad()
+{
+	Super::PostLoad();
+
+#if WITH_EDITORONLY_DATA
+	if (SourceActor_DEPRECATED != nullptr)
+	{
+		SoftSourceActor = SourceActor_DEPRECATED;
+	}
+#endif
+}
+
 void UNiagaraDataInterfacePhysicsAsset::GetFunctions(TArray<FNiagaraFunctionSignature>& OutFunctions)
 {
 	{
diff --git a/Engine/Plugins/Runtime/HairStrands/Source/HairStrandsCore/Public/Niagara/NiagaraDataInterfacePhysicsAsset.h b/Engine/Plugins/Runtime/HairStrands/Source/HairStrandsCore/Public/Niagara/NiagaraDataInterfacePhysicsAsset.h
--- a/Engine/Plugins/Runtime/HairStrands/Source/HairStrandsCore/Public/Niagara/NiagaraDataInterfacePhysicsAsset.h
+++ b/Engine/Plugins/Runtime/HairStrands/Source/HairStrandsCore/Public/Niagara/NiagaraDataInterfacePhysicsAsset.h
@@ -108,11 +108,16 @@ public:
 
 	/** Skeletal Mesh from which the Physics Asset will be found. */
 	UPROPERTY(EditAnywhere, Category = "Source")
-		TObjectPtr<UPhysicsAsset> DefaultSource;
+	TObjectPtr<UPhysicsAsset> DefaultSource;
 
 	/** The source actor from which to sample */
-	UPROPERTY(EditAnywhere, Category = "Source")
-	TObjectPtr<AActor> SourceActor;
+	UPROPERTY(EditAnywhere, Category = "Source", meta = (DisplayName = "Source Actor"))
+	TSoftObjectPtr<AActor> SoftSourceActor;
+
+#if WITH_EDITORONLY_DATA
+	UPROPERTY()
+	TObjectPtr<AActor> SourceActor_DEPRECATED;
+#endif
 
 	/** Reference to a user parameter if we're reading one. */
 	UPROPERTY(EditAnywhere, Category = "Source")
@@ -126,6 +131,7 @@ public:
 
 	/** UObject Interface */
 	virtual void PostInitProperties() override;
+	virtual void PostLoad() override;
 
 	/** UNiagaraDataInterface Interface */
 	virtual void GetFunctions(TArray<FNiagaraFunctionSignature>& OutFunctions) override;
