- Move constant buffers into loose parameters, with the exception of the external cbuffer
- Parameter binding is now enabled when possible, i.e. no legacy data interfaces / external cbuffer / view uniform buffer

diff --git a/Engine/Plugins/FX/Niagara/Source/Niagara/Classes/NiagaraGPUSystemTick.h b/Engine/Plugins/FX/Niagara/Source/Niagara/Classes/NiagaraGPUSystemTick.h
--- a/Engine/Plugins/FX/Niagara/Source/Niagara/Classes/NiagaraGPUSystemTick.h
+++ b/Engine/Plugins/FX/Niagara/Source/Niagara/Classes/NiagaraGPUSystemTick.h
@@ -78,29 +78,42 @@ public:
 		return MakeArrayView(reinterpret_cast<FNiagaraComputeInstanceData*>(InstanceData_ParamData_Packed), InstanceCount);
 	};
 
-	enum EUniformBufferType
+	FORCEINLINE void GetGlobalParameters(const FNiagaraComputeInstanceData& InstanceData, void* OutputParameters) const
 	{
-		UBT_FirstSystemType = 0,
-		UBT_Global = UBT_FirstSystemType,
-		UBT_System,
-		UBT_Owner,
-		UBT_NumSystemTypes,
+		const uint32 ParamSize = sizeof(FNiagaraGlobalParameters) * (InstanceData.Context->HasInterpolationParameters ? 2 : 1);
+		FMemory::Memcpy(OutputParameters, GlobalParamData, ParamSize);
+	}
 
-		UBT_FirstInstanceType = UBT_NumSystemTypes,
-		UBT_Emitter = UBT_FirstInstanceType,
-		UBT_External,
+	FORCEINLINE void GetSystemParameters(const FNiagaraComputeInstanceData& InstanceData, void* OutputParameters) const
+	{
+		const uint32 ParamSize = sizeof(FNiagaraSystemParameters) * (InstanceData.Context->HasInterpolationParameters ? 2 : 1);
+		FMemory::Memcpy(OutputParameters, SystemParamData, ParamSize);
+	}
 
-		UBT_NumTypes,
+	FORCEINLINE void GetOwnerParameters(const FNiagaraComputeInstanceData& InstanceData, void* OutputParameters) const
+	{
+		const uint32 ParamSize = sizeof(FNiagaraOwnerParameters) * (InstanceData.Context->HasInterpolationParameters ? 2 : 1);
+		FMemory::Memcpy(OutputParameters, OwnerParamData, ParamSize);
+	}
 
-		UBT_NumInstanceTypes = UBT_NumTypes - UBT_NumSystemTypes,
-	};
+	FORCEINLINE void GetEmitterParameters(const FNiagaraComputeInstanceData& InstanceData, void* OutputParameters) const
+	{
+		const uint32 ParamSize = sizeof(FNiagaraEmitterParameters) * (InstanceData.Context->HasInterpolationParameters ? 2 : 1);
+		FMemory::Memcpy(OutputParameters, InstanceData.EmitterParamData, ParamSize);
+	}
+
+	void BuildUniformBuffers();
 
-	FUniformBufferRHIRef GetUniformBuffer(EUniformBufferType Type, const FNiagaraComputeInstanceData* InstanceData, bool Previous) const;
-	const uint8* GetUniformBufferSource(EUniformBufferType Type, const FNiagaraComputeInstanceData* InstanceData, bool Previous) const;
+	FORCEINLINE FRHIUniformBuffer* GetExternalUniformBuffer(const FNiagaraComputeInstanceData& InstanceData, bool bPrevious) const
+	{
+		const int32 InstanceIndex = &InstanceData - GetInstances().GetData();
+		const int32 BufferIndex = InstanceIndex + (bPrevious ? InstanceCount : 0);
+		return ExternalUnformBuffers_RT[BufferIndex];
+	}
 
 public:
 	// Transient data used by the RT
-	TArray<FUniformBufferRHIRef> UniformBuffers;
+	TArray<FUniformBufferRHIRef> ExternalUnformBuffers_RT;
 
 	// data assigned by GT
 	FNiagaraSystemInstanceID SystemInstanceID = 0LL;						//-TODO: Remove?
@@ -113,6 +126,8 @@ public:
 	uint32 InstanceCount = 0;
 	uint32 TotalDispatches = 0;
 	bool bIsFinalTick = false;
+	bool bHasMultipleStages = false;
+	bool bHasInterpolatedParameters = false;
 
 #if !UE_BUILD_SHIPPING && !UE_BUILD_TEST
 	// Debugging only
diff --git a/Engine/Plugins/FX/Niagara/Source/Niagara/Classes/NiagaraScript.h b/Engine/Plugins/FX/Niagara/Source/Niagara/Classes/NiagaraScript.h
--- a/Engine/Plugins/FX/Niagara/Source/Niagara/Classes/NiagaraScript.h
+++ b/Engine/Plugins/FX/Niagara/Source/Niagara/Classes/NiagaraScript.h
@@ -454,7 +454,7 @@ public:
 #endif
 
 	UPROPERTY()
-	TArray<FNiagaraDataInterfaceGPUParamInfo> DIParamInfo; //TODO: GPU Param info should not be in the "VM executable data"
+	FNiagaraShaderScriptParametersMetadata ShaderScriptParametersMetadata; //TODO: GPU Param info should not be in the "VM executable data"
 
 #if WITH_EDITORONLY_DATA
 	/** The parameter collections used by this script. */
@@ -1080,6 +1080,8 @@ public:
 	NIAGARA_API FORCEINLINE const FNiagaraVMExecutableData& GetVMExecutableData() const { return CachedScriptVM; }
 	NIAGARA_API FORCEINLINE const FNiagaraVMExecutableDataId& GetVMExecutableDataCompilationId() const { return CachedScriptVMId; }
 
+	NIAGARA_API TConstArrayView<FNiagaraDataInterfaceGPUParamInfo> GetDataInterfaceGPUParamInfos() const { return CachedScriptVM.ShaderScriptParametersMetadata.DataInterfaceParamInfo; }
+
 	NIAGARA_API TArray<UNiagaraParameterCollection*>& GetCachedParameterCollectionReferences();
 	TArray<FNiagaraScriptDataInterfaceInfo>& GetCachedDefaultDataInterfaces() { return CachedDefaultDataInterfaces; }
 	TConstArrayView<FNiagaraScriptDataInterfaceInfo> GetCachedDefaultDataInterfaces() const { return MakeArrayView(CachedDefaultDataInterfaces); }
diff --git a/Engine/Plugins/FX/Niagara/Source/Niagara/Private/NiagaraDataInterfaceGrid2DCollection.cpp b/Engine/Plugins/FX/Niagara/Source/Niagara/Private/NiagaraDataInterfaceGrid2DCollection.cpp
--- a/Engine/Plugins/FX/Niagara/Source/Niagara/Private/NiagaraDataInterfaceGrid2DCollection.cpp
+++ b/Engine/Plugins/FX/Niagara/Source/Niagara/Private/NiagaraDataInterfaceGrid2DCollection.cpp
@@ -1942,7 +1942,7 @@ void UNiagaraDataInterfaceGrid2DCollection::CollectAttributesForScript(UNiagaraS
 		const int32 * IndexOfDataInterface = ParameterStore->FindParameterOffset(DataInterfaceVariable);
 		if (IndexOfDataInterface != nullptr)
 		{
-			const TArray<FNiagaraDataInterfaceGPUParamInfo>& ParamInfoArray = Script->GetVMExecutableData().DIParamInfo;
+			TConstArrayView<FNiagaraDataInterfaceGPUParamInfo> ParamInfoArray = Script->GetDataInterfaceGPUParamInfos();
 			for (const FNiagaraDataInterfaceGeneratedFunction& Func : ParamInfoArray[*IndexOfDataInterface].GeneratedFunctions)
 			{
 				static const FName NAME_Attribute("Attribute");
diff --git a/Engine/Plugins/FX/Niagara/Source/Niagara/Private/NiagaraDataInterfaceGrid3DCollection.cpp b/Engine/Plugins/FX/Niagara/Source/Niagara/Private/NiagaraDataInterfaceGrid3DCollection.cpp
--- a/Engine/Plugins/FX/Niagara/Source/Niagara/Private/NiagaraDataInterfaceGrid3DCollection.cpp
+++ b/Engine/Plugins/FX/Niagara/Source/Niagara/Private/NiagaraDataInterfaceGrid3DCollection.cpp
@@ -2151,7 +2151,7 @@ void UNiagaraDataInterfaceGrid3DCollection::CollectAttributesForScript(UNiagaraS
 		const int32 * IndexOfDataInterface = ParameterStore->FindParameterOffset(DataInterfaceVariable);
 		if (IndexOfDataInterface != nullptr)
 		{
-			const TArray<FNiagaraDataInterfaceGPUParamInfo>& ParamInfoArray = Script->GetVMExecutableData().DIParamInfo;
+			TConstArrayView<FNiagaraDataInterfaceGPUParamInfo> ParamInfoArray = Script->GetDataInterfaceGPUParamInfos();
 			for (const FNiagaraDataInterfaceGeneratedFunction& Func : ParamInfoArray[*IndexOfDataInterface].GeneratedFunctions)
 			{
 				static const FName NAME_Attribute("Attribute");
diff --git a/Engine/Plugins/FX/Niagara/Source/Niagara/Private/NiagaraGPUSystemTick.cpp b/Engine/Plugins/FX/Niagara/Source/Niagara/Private/NiagaraGPUSystemTick.cpp
--- a/Engine/Plugins/FX/Niagara/Source/Niagara/Private/NiagaraGPUSystemTick.cpp
+++ b/Engine/Plugins/FX/Niagara/Source/Niagara/Private/NiagaraGPUSystemTick.cpp
@@ -166,6 +166,9 @@ void FNiagaraGPUSystemTick::Init(FNiagaraSystemInstance* InSystemInstance)
 				FMemory::Memcpy(InstanceData->EmitterParamData + sizeof(FNiagaraEmitterParameters), &InSystemInstance->GetEmitterParameters(EmitterIdx, true), sizeof(FNiagaraEmitterParameters));
 			}
 
+			bHasMultipleStages = InstanceData->bHasMultipleStages;
+			bHasInterpolatedParameters |= GPUContext->HasInterpolationParameters;
+
 			GPUContext->CombinedParamStore.CopyParameterDataToPaddedBuffer(InstanceData->ExternalParamData, InstanceData->ExternalParamDataSize);
 
 			// Calling PostTick will push current -> previous parameters this must be done after copying the parameter data
@@ -263,56 +266,34 @@ void FNiagaraGPUSystemTick::Destroy()
 	}
 }
 
-FUniformBufferRHIRef FNiagaraGPUSystemTick::GetUniformBuffer(EUniformBufferType Type, const FNiagaraComputeInstanceData* Instance, bool Current) const
+void FNiagaraGPUSystemTick::BuildUniformBuffers()
 {
-	const int32 InterpOffset = Current
-		? 0
-		: (UBT_NumSystemTypes + InstanceCount * UBT_NumInstanceTypes);
-
-	if (Instance)
-	{
-		check(Type >= UBT_FirstInstanceType);
-		check(Type < UBT_NumTypes);
+	check(ExternalUnformBuffers_RT.Num() == 0);
 
-		const int32 InstanceTypeIndex = Type - UBT_FirstInstanceType;
-
-		const int32 InstanceIndex = (Instance - GetInstances().GetData());
-		return UniformBuffers[InterpOffset + UBT_NumSystemTypes + InstanceCount * InstanceTypeIndex + InstanceIndex];
-	}
+	const int32 InterpCount = bHasInterpolatedParameters ? 2 : 1;
+	ExternalUnformBuffers_RT.AddDefaulted(InstanceCount * InterpCount);
 
-	check(Type >= UBT_FirstSystemType);
-	check(Type < UBT_FirstInstanceType);
-
-	return UniformBuffers[InterpOffset + Type];
-}
-
-const uint8* FNiagaraGPUSystemTick::GetUniformBufferSource(EUniformBufferType Type, const FNiagaraComputeInstanceData* Instance, bool Current) const
-{
-	check(Type >= UBT_FirstSystemType);
-	check(Type < UBT_NumTypes);
-
-	switch (Type)
+	TArrayView<FNiagaraComputeInstanceData> Instances = GetInstances();
+	for (uint32 iInstance=0; iInstance < InstanceCount; ++iInstance)
 	{
-		case UBT_Global:
-			return GlobalParamData + (Current ? 0 : sizeof(FNiagaraGlobalParameters));
-		case UBT_System:
-			return SystemParamData + (Current ? 0 : sizeof(FNiagaraSystemParameters));
-		case UBT_Owner:
-			return OwnerParamData + (Current ? 0 : sizeof(FNiagaraOwnerParameters));
-		case UBT_Emitter:
+		const FNiagaraComputeInstanceData& Instance = Instances[iInstance];
+
+		FNiagaraRHIUniformBufferLayout* ExternalCBufferLayout = Instance.Context->ExternalCBufferLayout;
+		const bool bExternalLayoutValid = ExternalCBufferLayout && (ExternalCBufferLayout->Resources.Num() || ExternalCBufferLayout->ConstantBufferSize > 0);
+		if ( Instance.Context->GPUScript_RT->IsExternalConstantBufferUsed_RenderThread(0) )
 		{
-			check(Instance);
-			return Instance->EmitterParamData + (Current ? 0 : sizeof(FNiagaraEmitterParameters));
+			if ( ensure(ExternalCBufferLayout && bExternalLayoutValid) )
+			{
+				ExternalUnformBuffers_RT[iInstance] = RHICreateUniformBuffer(Instance.ExternalParamData, ExternalCBufferLayout, bHasMultipleStages ? EUniformBufferUsage::UniformBuffer_SingleFrame : EUniformBufferUsage::UniformBuffer_SingleDraw);
+			}
 		}
-		case UBT_External:
+		if ( Instance.Context->GPUScript_RT->IsExternalConstantBufferUsed_RenderThread(1) )
 		{
-			// External parameters are pushed from the combined parameters store, split into two where first half is current second is previous
-			check(Instance && Instance->Context);
-			check((Current ? 0 : Instance->Context->ExternalCBufferLayout->ConstantBufferSize) + Instance->Context->ExternalCBufferLayout->ConstantBufferSize <= Instance->ExternalParamDataSize);
-
-			return Instance->ExternalParamData + (Current ? 0 : Instance->Context->ExternalCBufferLayout->ConstantBufferSize);
+			if (ensure(ExternalCBufferLayout && bExternalLayoutValid))
+			{
+				check(ExternalCBufferLayout->ConstantBufferSize + ExternalCBufferLayout->ConstantBufferSize <= Instance.ExternalParamDataSize);
+				ExternalUnformBuffers_RT[InstanceCount + iInstance] = RHICreateUniformBuffer(Instance.ExternalParamData + ExternalCBufferLayout->ConstantBufferSize, ExternalCBufferLayout, bHasMultipleStages ? EUniformBufferUsage::UniformBuffer_SingleFrame : EUniformBufferUsage::UniformBuffer_SingleDraw);
+			}
 		}
 	}
-
-	return nullptr;
 }
diff --git a/Engine/Plugins/FX/Niagara/Source/Niagara/Private/NiagaraGpuComputeDispatch.cpp b/Engine/Plugins/FX/Niagara/Source/Niagara/Private/NiagaraGpuComputeDispatch.cpp
--- a/Engine/Plugins/FX/Niagara/Source/Niagara/Private/NiagaraGpuComputeDispatch.cpp
+++ b/Engine/Plugins/FX/Niagara/Source/Niagara/Private/NiagaraGpuComputeDispatch.cpp
@@ -31,7 +31,7 @@
 #include "ShaderParameterUtils.h"
 #include "Renderer/Private/ScenePrivate.h"
 
-DECLARE_CYCLE_STAT(TEXT("Niagara Dispatch Setup"), STAT_NiagaraGPUDispatchSetup_RT, STATGROUP_Niagara);
+DECLARE_CYCLE_STAT(TEXT("GPU Dispatch Setup [RT]"), STAT_NiagaraGPUDispatchSetup_RT, STATGROUP_Niagara);
 DECLARE_CYCLE_STAT(TEXT("GPU Emitter Dispatch [RT]"), STAT_NiagaraGPUSimTick_RT, STATGROUP_Niagara);
 DECLARE_CYCLE_STAT(TEXT("GPU Data Readback [RT]"), STAT_NiagaraGPUReadback_RT, STATGROUP_Niagara);
 DECLARE_FLOAT_COUNTER_STAT(TEXT("Niagara GPU Sim"), STAT_GPU_NiagaraSim, STATGROUP_GPU);
@@ -204,12 +204,6 @@ FNiagaraGpuComputeDispatch::FNiagaraGpuComputeDispatch(ERHIFeatureLevel::Type In
 		}
 	}
 
-	GlobalCBufferLayout  = new FNiagaraRHIUniformBufferLayout(TEXT("Niagara GPU Global CBuffer"),  sizeof(FNiagaraGlobalParameters));
-	SystemCBufferLayout  = new FNiagaraRHIUniformBufferLayout(TEXT("Niagara GPU System CBuffer"),  sizeof(FNiagaraSystemParameters));
-	OwnerCBufferLayout   = new FNiagaraRHIUniformBufferLayout(TEXT("Niagara GPU Owner CBuffer"),   sizeof(FNiagaraOwnerParameters));
-	EmitterCBufferLayout = new FNiagaraRHIUniformBufferLayout(TEXT("Niagara GPU Emitter CBuffer"), sizeof(FNiagaraEmitterParameters));
-
-
 	AsyncGpuTraceHelper.Reset(new FNiagaraAsyncGpuTraceHelper(InShaderPlatform, InFeatureLevel, this));
 
 #if NIAGARA_COMPUTEDEBUG_ENABLED
@@ -228,11 +222,6 @@ FNiagaraGpuComputeDispatch::~FNiagaraGpuComputeDispatch()
 
 	AsyncGpuTraceHelper->Reset();
 
-	GlobalCBufferLayout = nullptr;
-	SystemCBufferLayout = nullptr;
-	OwnerCBufferLayout = nullptr;
-	EmitterCBufferLayout = nullptr;
-
 	FPrimitiveSceneInfo::OnGPUSceneInstancesAllocated.RemoveAll(this);
 	FPrimitiveSceneInfo::OnGPUSceneInstancesFreed.RemoveAll(this);
 }
@@ -293,107 +282,6 @@ void FNiagaraGpuComputeDispatch::RemoveGpuComputeProxy(FNiagaraSystemGpuComputeP
 #endif
 }
 
-void FNiagaraGpuComputeDispatch::BuildConstantBuffers(FNiagaraGPUSystemTick& Tick) const
-{
-	if (!Tick.InstanceCount)
-	{
-		return;
-	}
-
-	TArrayView<FNiagaraComputeInstanceData> InstanceData = Tick.GetInstances();
-
-	// first go through and figure out if we need to support interpolated spawning
-	bool HasInterpolationParameters = false;
-	bool HasMultipleStages = false;
-	for (const FNiagaraComputeInstanceData& Instance : InstanceData)
-	{
-		HasInterpolationParameters = HasInterpolationParameters || Instance.Context->HasInterpolationParameters;
-		HasMultipleStages = HasMultipleStages || Instance.bHasMultipleStages;
-	}
-
-	int32 BoundParameterCounts[FNiagaraGPUSystemTick::UBT_NumTypes][2];
-	FMemory::Memzero(BoundParameterCounts);
-
-	for (const FNiagaraComputeInstanceData& Instance : InstanceData)
-	{
-		for (int32 InterpIt = 0; InterpIt < (HasInterpolationParameters ? 2 : 1); ++InterpIt)
-		{
-			BoundParameterCounts[FNiagaraGPUSystemTick::UBT_Global][InterpIt] += Instance.Context->GPUScript_RT->IsGlobalConstantBufferUsed_RenderThread(InterpIt) ? 1 : 0;
-			BoundParameterCounts[FNiagaraGPUSystemTick::UBT_System][InterpIt] += Instance.Context->GPUScript_RT->IsSystemConstantBufferUsed_RenderThread(InterpIt) ? 1 : 0;
-			BoundParameterCounts[FNiagaraGPUSystemTick::UBT_Owner][InterpIt] += Instance.Context->GPUScript_RT->IsOwnerConstantBufferUsed_RenderThread(InterpIt) ? 1 : 0;
-		}
-	}
-
-	const int32 InterpScale = HasInterpolationParameters ? 2 : 1;
-	const int32 BufferCount = InterpScale * (FNiagaraGPUSystemTick::UBT_NumSystemTypes + FNiagaraGPUSystemTick::UBT_NumInstanceTypes * Tick.InstanceCount);
-
-	Tick.UniformBuffers.Empty(BufferCount);
-	Tick.UniformBuffers.AddDefaulted(BufferCount);
-
-	const FRHIUniformBufferLayout* SystemLayouts[FNiagaraGPUSystemTick::UBT_NumSystemTypes] =
-	{
-		GlobalCBufferLayout,
-		SystemCBufferLayout,
-		OwnerCBufferLayout
-	};
-
-	int32 CurrentBuffer = 0;
-	for (int32 InterpIt=0; InterpIt < InterpScale; ++InterpIt)
-	{
-		for (int32 SystemTypeIt = FNiagaraGPUSystemTick::UBT_FirstSystemType; SystemTypeIt < FNiagaraGPUSystemTick::UBT_NumSystemTypes; ++SystemTypeIt)
-		{
-			if (BoundParameterCounts[SystemTypeIt][InterpIt])
-			{
-				Tick.UniformBuffers[CurrentBuffer] = RHICreateUniformBuffer(
-					Tick.GetUniformBufferSource((FNiagaraGPUSystemTick::EUniformBufferType) SystemTypeIt, nullptr, !InterpIt),
-					SystemLayouts[SystemTypeIt],
-					((BoundParameterCounts[SystemTypeIt][InterpIt] > 1) || HasMultipleStages)
-						? EUniformBufferUsage::UniformBuffer_SingleFrame
-						: EUniformBufferUsage::UniformBuffer_SingleDraw);
-			}
-			++CurrentBuffer;
-		}
-
-		// Build emitter constant buffers
-		for (const FNiagaraComputeInstanceData& Instance : InstanceData)
-		{
-			if (InterpIt == 0 || Instance.Context->HasInterpolationParameters)
-			{
-				if (Instance.Context->GPUScript_RT->IsEmitterConstantBufferUsed_RenderThread(InterpIt))
-				{
-					Tick.UniformBuffers[CurrentBuffer] = RHICreateUniformBuffer(
-						Tick.GetUniformBufferSource(FNiagaraGPUSystemTick::UBT_Emitter, &Instance, !InterpIt),
-						EmitterCBufferLayout,
-						Instance.bHasMultipleStages ? EUniformBufferUsage::UniformBuffer_SingleFrame : EUniformBufferUsage::UniformBuffer_SingleDraw
-					);
-				}
-			}
-			++CurrentBuffer;
-		}
-
-		// Build external constant buffers
-		for (const FNiagaraComputeInstanceData& Instance : InstanceData)
-		{
-			if (InterpIt == 0 || Instance.Context->HasInterpolationParameters)
-			{
-				FNiagaraRHIUniformBufferLayout* ExternalCBufferLayout = Instance.Context->ExternalCBufferLayout;
-				if (Instance.Context->GPUScript_RT->IsExternalConstantBufferUsed_RenderThread(InterpIt))
-				{
-					if (ensure(ExternalCBufferLayout && (ExternalCBufferLayout->Resources.Num() > 0 || ExternalCBufferLayout->ConstantBufferSize > 0)))
-					{
-						Tick.UniformBuffers[CurrentBuffer] = RHICreateUniformBuffer(
-							Tick.GetUniformBufferSource(FNiagaraGPUSystemTick::UBT_External, &Instance, !InterpIt),
-							ExternalCBufferLayout,
-							Instance.bHasMultipleStages ? EUniformBufferUsage::UniformBuffer_SingleFrame : EUniformBufferUsage::UniformBuffer_SingleDraw
-						);
-					}
-				}
-			}
-			++CurrentBuffer;
-		}
-	}
-}
-
 void FNiagaraGpuComputeDispatch::Tick(UWorld* World, float DeltaTime)
 {
 	check(IsInGameThread());
@@ -1169,7 +1057,7 @@ void FNiagaraGpuComputeDispatch::PrepareTicksForProxy(FRHICommandListImmediate&
 		}
 
 		// Build constant buffers for tick
-		BuildConstantBuffers(Tick);
+		Tick.BuildUniformBuffers();
 	}
 
 	// Now that all ticks have been processed we can adjust our output buffers to the correct size
@@ -1657,47 +1545,23 @@ void FNiagaraGpuComputeDispatch::DispatchStage(FRHICommandList& RHICmdList, cons
 	// Set data interface parameters
 	SetDataInterfaceParameters(RHICmdList, Tick, InstanceData, ComputeShader, SimStageData, reinterpret_cast<uint8*>(DispatchParameters));
 
-	// Setup uniform buffers
-	//-TODO: This should be replaced with shader struct bindings instead but we currently have no way to include a untyped uniform buffer
-	auto SetConstantBuffer =
-		[](FRHICommandList& RHICmdList, FRHIComputeShader* ComputeShader, const FShaderUniformBufferParameter& BufferParam, const FUniformBufferRHIRef& UniformBuffer)
-		{
-			check(BufferParam.IsBound() == false || UniformBuffer.IsValid());
-			if (BufferParam.IsBound() && UniformBuffer.IsValid())
-			{
-				RHICmdList.SetShaderUniformBuffer(ComputeShader, BufferParam.GetBaseIndex(), UniformBuffer);
-			}
-		};
-
-	SetConstantBuffer(RHICmdList, RHIComputeShader, ComputeShader->GlobalConstantBufferParam[0], Tick.GetUniformBuffer(FNiagaraGPUSystemTick::UBT_Global, nullptr, true));
-	SetConstantBuffer(RHICmdList, RHIComputeShader, ComputeShader->SystemConstantBufferParam[0], Tick.GetUniformBuffer(FNiagaraGPUSystemTick::UBT_System, nullptr, true));
-	SetConstantBuffer(RHICmdList, RHIComputeShader, ComputeShader->OwnerConstantBufferParam[0], Tick.GetUniformBuffer(FNiagaraGPUSystemTick::UBT_Owner, nullptr, true));
-	SetConstantBuffer(RHICmdList, RHIComputeShader, ComputeShader->EmitterConstantBufferParam[0], Tick.GetUniformBuffer(FNiagaraGPUSystemTick::UBT_Emitter, &InstanceData, true));
-	SetConstantBuffer(RHICmdList, RHIComputeShader, ComputeShader->ExternalConstantBufferParam[0], Tick.GetUniformBuffer(FNiagaraGPUSystemTick::UBT_External, &InstanceData, true));
-
-	if (InstanceData.Context->HasInterpolationParameters)
-	{
-		SetConstantBuffer(RHICmdList, RHIComputeShader, ComputeShader->GlobalConstantBufferParam[1], Tick.GetUniformBuffer(FNiagaraGPUSystemTick::UBT_Global, nullptr, false));
-		SetConstantBuffer(RHICmdList, RHIComputeShader, ComputeShader->SystemConstantBufferParam[1], Tick.GetUniformBuffer(FNiagaraGPUSystemTick::UBT_System, nullptr, false));
-		SetConstantBuffer(RHICmdList, RHIComputeShader, ComputeShader->OwnerConstantBufferParam[1], Tick.GetUniformBuffer(FNiagaraGPUSystemTick::UBT_Owner, nullptr, false));
-		SetConstantBuffer(RHICmdList, RHIComputeShader, ComputeShader->EmitterConstantBufferParam[1], Tick.GetUniformBuffer(FNiagaraGPUSystemTick::UBT_Emitter, &InstanceData, false));
-		SetConstantBuffer(RHICmdList, RHIComputeShader, ComputeShader->ExternalConstantBufferParam[1], Tick.GetUniformBuffer(FNiagaraGPUSystemTick::UBT_External, &InstanceData, false));
-	}
-
-	//-TODO: Currently we don't have a way to bind raw UBs as shader parameters
-	//DispatchParameters.GlobalParameters			= Tick.GetUniformBuffer(FNiagaraGPUSystemTick::UBT_Global, nullptr, true);
-	//DispatchParameters.OwnerParameters			= Tick.GetUniformBuffer(FNiagaraGPUSystemTick::UBT_Owner, nullptr, true);
-	//DispatchParameters.SystemParameters			= Tick.GetUniformBuffer(FNiagaraGPUSystemTick::UBT_System, nullptr, true);
-	//DispatchParameters.EmitterParameters		= Tick.GetUniformBuffer(FNiagaraGPUSystemTick::UBT_Emitter, nullptr, true);
-	//DispatchParameters.ExternalParameters		= Tick.GetUniformBuffer(FNiagaraGPUSystemTick::UBT_External, nullptr, true);
-	//if (InstanceData.Context->HasInterpolationParameters)
-	//{
-	//	DispatchParameters.PrevGlobalParameters		= Tick.GetUniformBuffer(FNiagaraGPUSystemTick::UBT_Global, nullptr, false);
-	//	DispatchParameters.PrevOwnerParameters		= Tick.GetUniformBuffer(FNiagaraGPUSystemTick::UBT_Owner, nullptr, false);
-	//	DispatchParameters.PrevSystemParameters		= Tick.GetUniformBuffer(FNiagaraGPUSystemTick::UBT_System, nullptr, false);
-	//	DispatchParameters.PrevEmitterParameters	= Tick.GetUniformBuffer(FNiagaraGPUSystemTick::UBT_Emitter, nullptr, false);
-	//	DispatchParameters.PrevExternalParameters	= Tick.GetUniformBuffer(FNiagaraGPUSystemTick::UBT_External, nullptr, false);
-	//}
+	// Set tick parameters
+	Tick.GetGlobalParameters(InstanceData, &DispatchParameters->GlobalParameters);
+	Tick.GetSystemParameters(InstanceData, &DispatchParameters->SystemParameters);
+	Tick.GetOwnerParameters(InstanceData, &DispatchParameters->OwnerParameters);
+	Tick.GetEmitterParameters(InstanceData, &DispatchParameters->EmitterParameters);
+
+	// Set external constant buffer
+	//-TODO: This should be replace with parameters structure which is dynamically created
+	if (ComputeShader->ExternalConstantBufferParam[0].IsBound())
+	{
+		RHICmdList.SetShaderUniformBuffer(RHIComputeShader, ComputeShader->ExternalConstantBufferParam[0].GetBaseIndex(), Tick.GetExternalUniformBuffer(InstanceData, false));
+	}
+	if (ComputeShader->ExternalConstantBufferParam[1].IsBound())
+	{
+		check(InstanceData.Context->HasInterpolationParameters);
+		RHICmdList.SetShaderUniformBuffer(RHIComputeShader, ComputeShader->ExternalConstantBufferParam[1].GetBaseIndex(), Tick.GetExternalUniformBuffer(InstanceData, true));
+	}
 
 	// Execute the dispatch
 	{
@@ -1741,6 +1605,8 @@ void FNiagaraGpuComputeDispatch::DispatchStage(FRHICommandList& RHICmdList, cons
 
 void FNiagaraGpuComputeDispatch::PreInitViews(FRDGBuilder& GraphBuilder, bool bAllowGPUParticleUpdate)
 {
+	SCOPE_CYCLE_COUNTER(STAT_NiagaraGPUDispatchSetup_RT);
+
 	bRequiresReadback = false;
 	GNiagaraViewDataManager.ClearSceneTextureParameters();
 #if WITH_EDITOR
diff --git a/Engine/Plugins/FX/Niagara/Source/Niagara/Private/NiagaraGpuComputeDispatch.h b/Engine/Plugins/FX/Niagara/Source/Niagara/Private/NiagaraGpuComputeDispatch.h
--- a/Engine/Plugins/FX/Niagara/Source/Niagara/Private/NiagaraGpuComputeDispatch.h
+++ b/Engine/Plugins/FX/Niagara/Source/Niagara/Private/NiagaraGpuComputeDispatch.h
@@ -167,8 +167,6 @@ private:
 	void UpdateFreeIDsListSizesBuffer(FRHICommandList& RHICmdList, uint32 NumInstances);
 	void UpdateFreeIDBuffers(FRHICommandList& RHICmdList, TConstArrayView<FNiagaraComputeExecutionContext*> Instances);
 
-	void BuildConstantBuffers(FNiagaraGPUSystemTick& Tick) const;
-
 	/** The shared GPUSortManager, used to register GPU sort tasks in order to generate sorted particle indices per emitter. */
 	TRefCountPtr<FGPUSortManager> GPUSortManager;
 	/** All sort tasks registered in AddSortedGPUSimulation(). Holds all the data required in GenerateSortKeys(). */
@@ -180,12 +178,6 @@ private:
 	TUniquePtr<FNiagaraGPUProfiler> GPUProfilerPtr;
 #endif
 
-	// persistent layouts used to create the constant buffers for the compute sim shader
-	TRefCountPtr<FNiagaraRHIUniformBufferLayout> GlobalCBufferLayout;
-	TRefCountPtr<FNiagaraRHIUniformBufferLayout> SystemCBufferLayout;
-	TRefCountPtr<FNiagaraRHIUniformBufferLayout> OwnerCBufferLayout;
-	TRefCountPtr<FNiagaraRHIUniformBufferLayout> EmitterCBufferLayout;
-
 	uint32 FramesBeforeTickFlush = 0;
 
 	TConstArrayView<FViewInfo> CurrentPassViews;
diff --git a/Engine/Plugins/FX/Niagara/Source/Niagara/Private/NiagaraModule.cpp b/Engine/Plugins/FX/Niagara/Source/Niagara/Private/NiagaraModule.cpp
--- a/Engine/Plugins/FX/Niagara/Source/Niagara/Private/NiagaraModule.cpp
+++ b/Engine/Plugins/FX/Niagara/Source/Niagara/Private/NiagaraModule.cpp
@@ -1527,6 +1527,14 @@ const TArray<FNiagaraVariable>& FNiagaraSystemParameters::GetVariables()
 
 const TArray<FNiagaraVariable>& FNiagaraOwnerParameters::GetVariables()
 {
+	static const FName NAME_NiagaraStructPadding0 = "Engine.Owner.PaddingInt32_0";
+	static const FName NAME_NiagaraStructPadding1 = "Engine.Owner.PaddingInt32_1";
+	static const FName NAME_NiagaraStructPadding2 = "Engine.Owner.PaddingInt32_2";
+	static const FName NAME_NiagaraStructPadding3 = "Engine.Owner.PaddingInt32_3";
+	static const FName NAME_NiagaraStructPadding4 = "Engine.Owner.PaddingInt32_4";
+	static const FName NAME_NiagaraStructPadding5 = "Engine.Owner.PaddingInt32_5";
+	static const FName NAME_NiagaraStructPadding6 = "Engine.Owner.PaddingInt32_6";
+
 	static const TArray<FNiagaraVariable> Variables =
 	{
 		SYS_PARAM_ENGINE_LOCAL_TO_WORLD,
@@ -1537,11 +1545,17 @@ const TArray<FNiagaraVariable>& FNiagaraOwnerParameters::GetVariables()
 		SYS_PARAM_ENGINE_WORLD_TO_LOCAL_NO_SCALE,
 		SYS_PARAM_ENGINE_ROTATION,
 		SYS_PARAM_ENGINE_POSITION,
+		FNiagaraVariable(FNiagaraTypeDefinition::GetIntDef(), NAME_NiagaraStructPadding0),
 		SYS_PARAM_ENGINE_VELOCITY,
+		FNiagaraVariable(FNiagaraTypeDefinition::GetIntDef(), NAME_NiagaraStructPadding1),
 		SYS_PARAM_ENGINE_X_AXIS,
+		FNiagaraVariable(FNiagaraTypeDefinition::GetIntDef(), NAME_NiagaraStructPadding2),
 		SYS_PARAM_ENGINE_Y_AXIS,
+		FNiagaraVariable(FNiagaraTypeDefinition::GetIntDef(), NAME_NiagaraStructPadding3),
 		SYS_PARAM_ENGINE_Z_AXIS,
+		FNiagaraVariable(FNiagaraTypeDefinition::GetIntDef(), NAME_NiagaraStructPadding4),
 		SYS_PARAM_ENGINE_SCALE,
+		FNiagaraVariable(FNiagaraTypeDefinition::GetIntDef(), NAME_NiagaraStructPadding5),
 		SYS_PARAM_ENGINE_LWC_TILE,
 	};
 
diff --git a/Engine/Plugins/FX/Niagara/Source/Niagara/Private/NiagaraScript.cpp b/Engine/Plugins/FX/Niagara/Source/Niagara/Private/NiagaraScript.cpp
--- a/Engine/Plugins/FX/Niagara/Source/Niagara/Private/NiagaraScript.cpp
+++ b/Engine/Plugins/FX/Niagara/Source/Niagara/Private/NiagaraScript.cpp
@@ -10,8 +10,11 @@
 #include "NiagaraModule.h"
 #include "NiagaraCustomVersion.h"
 #include "NiagaraShaderCompilationManager.h"
+#include "NiagaraShaderModule.h"
+#include "NiagaraShaderParametersBuilder.h"
 #include "Serialization/MemoryReader.h"
 #include "Misc/SecureHash.h"
+#include "ShaderParameterMetadataBuilder.h"
 
 #include "ProfilingDebugging/CookStats.h"
 #include "Stats/Stats.h"
@@ -3121,7 +3124,7 @@ void UNiagaraScript::CacheResourceShadersForRendering(bool bRegenerateId, bool b
 				CachedScriptVMId.BaseScriptCompileHash, CachedScriptVMId.ReferencedCompileHashes,
 				CachedScriptVMId.bUsesRapidIterationParams, GetFriendlyName());
 
-			ScriptResource->BuildScriptParametersMetadata(CachedScriptVM.DIParamInfo);
+			ScriptResource->BuildScriptParametersMetadata(CachedScriptVM.ShaderScriptParametersMetadata);
 
 			if (FNiagaraUtilities::SupportsComputeShaders(ShaderPlatform))
 			{
@@ -3553,7 +3556,7 @@ void UNiagaraScript::ProcessSerializedShaderMaps()
 
 	if (HasScriptResource)
 	{
-		ScriptResource->BuildScriptParametersMetadata(CachedScriptVM.DIParamInfo);
+		ScriptResource->BuildScriptParametersMetadata(CachedScriptVM.ShaderScriptParametersMetadata);
 	}
 }
 
diff --git a/Engine/Plugins/FX/Niagara/Source/Niagara/Private/NiagaraScriptExecutionParameterStore.cpp b/Engine/Plugins/FX/Niagara/Source/Niagara/Private/NiagaraScriptExecutionParameterStore.cpp
--- a/Engine/Plugins/FX/Niagara/Source/Niagara/Private/NiagaraScriptExecutionParameterStore.cpp
+++ b/Engine/Plugins/FX/Niagara/Source/Niagara/Private/NiagaraScriptExecutionParameterStore.cpp
@@ -5,6 +5,19 @@
 #include "NiagaraDataInterface.h"
 #include "NiagaraSystemInstance.h"
 
+namespace ScriptParameterStoreInternal
+{
+	template<typename TType>
+	void AddLayoutParameter(TArray<FNiagaraScriptExecutionPaddingInfo>& Members, uint32& NextMemberOffset, uint32& SrcOffset)
+	{
+		const uint32 BytesRemain = 16 - (NextMemberOffset & 15);
+		NextMemberOffset += BytesRemain != 16 && BytesRemain < sizeof(TType) ? BytesRemain : 0;
+		Members.Emplace(SrcOffset, NextMemberOffset, sizeof(TType), sizeof(TType));
+		SrcOffset += sizeof(TType);
+		NextMemberOffset += sizeof(TType);
+	}
+}
+
 FNiagaraScriptExecutionParameterStore::FNiagaraScriptExecutionParameterStore()
 	: FNiagaraParameterStore()
 	, ParameterSize(0)
@@ -39,106 +52,42 @@ uint32 OffsetAlign(uint32 SrcOffset, uint32 Size)
 
 uint32 FNiagaraScriptExecutionParameterStore::GenerateLayoutInfoInternal(TArray<FNiagaraScriptExecutionPaddingInfo>& Members, uint32& NextMemberOffset, const UStruct* InSrcStruct, uint32 InSrcOffset)
 {
-	uint32 VectorPaddedSize = (TShaderParameterTypeInfo<FVector4f>::NumRows * TShaderParameterTypeInfo<FVector4f>::NumColumns) * sizeof(float);
-
 	// Now insert an appropriate data member into the mix...
+	FName SrcStructName = InSrcStruct ? InSrcStruct->GetFName() : NAME_None;
 	if (InSrcStruct == FNiagaraTypeDefinition::GetBoolStruct() || InSrcStruct == FNiagaraTypeDefinition::GetIntStruct())
 	{
-		uint32 IntSize = (TShaderParameterTypeInfo<uint32>::NumRows * TShaderParameterTypeInfo<uint32>::NumColumns) * sizeof(uint32);
-		Members.Emplace(InSrcOffset, Align(NextMemberOffset, TShaderParameterTypeInfo<uint32>::Alignment), IntSize, IntSize);
-		InSrcOffset += sizeof(uint32);
-		NextMemberOffset = Members[Members.Num() - 1].DestOffset + Members[Members.Num() - 1].DestSize;
+		ScriptParameterStoreInternal::AddLayoutParameter<int32>(Members, NextMemberOffset, InSrcOffset);
 	}
 	else if (InSrcStruct == FNiagaraTypeDefinition::GetFloatStruct())
 	{
-		uint32 FloatSize = (TShaderParameterTypeInfo<float>::NumRows * TShaderParameterTypeInfo<float>::NumColumns) * sizeof(float);
-		Members.Emplace(InSrcOffset, Align(NextMemberOffset, TShaderParameterTypeInfo<float>::Alignment), FloatSize, FloatSize);
-		InSrcOffset += sizeof(float);
-		NextMemberOffset = Members[Members.Num() - 1].DestOffset + Members[Members.Num() - 1].DestSize;
-	}
-	else if (InSrcStruct == FNiagaraTypeDefinition::GetVec2Struct())
-	{
-		uint32 StructFinalSize = (TShaderParameterTypeInfo<FVector2f>::NumRows * TShaderParameterTypeInfo<FVector2f>::NumColumns) * sizeof(float);
-		Members.Emplace(InSrcOffset, OffsetAlign(NextMemberOffset, VectorPaddedSize), StructFinalSize, VectorPaddedSize);
-		InSrcOffset += sizeof(FVector2f);
-		NextMemberOffset = Members[Members.Num() - 1].DestOffset + Members[Members.Num() - 1].DestSize;
+		ScriptParameterStoreInternal::AddLayoutParameter<float>(Members, NextMemberOffset, InSrcOffset);
 	}
-	else if (InSrcStruct && (InSrcStruct->GetFName() == NAME_Vector2D || InSrcStruct->GetFName() == NAME_Vector2d))
-	{
-		uint32 StructFinalSize = (TShaderParameterTypeInfo<FVector2f>::NumRows * TShaderParameterTypeInfo<FVector2f>::NumColumns) * sizeof(float);
-		Members.Emplace(InSrcOffset, OffsetAlign(NextMemberOffset, VectorPaddedSize), StructFinalSize, VectorPaddedSize);
-		InSrcOffset += sizeof(FVector2f);
-		NextMemberOffset = Members[Members.Num() - 1].DestOffset + Members[Members.Num() - 1].DestSize;
-	}	
-	else if (InSrcStruct == FNiagaraTypeDefinition::GetVec3Struct() || InSrcStruct == FNiagaraTypeDefinition::GetPositionStruct())
+	else if (InSrcStruct == FNiagaraTypeDefinition::GetVec2Struct() || SrcStructName == NAME_Vector2D || SrcStructName == NAME_Vector2d)
 	{
-		uint32 StructFinalSize = (TShaderParameterTypeInfo<FVector3f>::NumRows * TShaderParameterTypeInfo<FVector3f>::NumColumns) * sizeof(float);
-		Members.Emplace(InSrcOffset, OffsetAlign(NextMemberOffset, VectorPaddedSize), StructFinalSize, VectorPaddedSize);
-		InSrcOffset += sizeof(FVector3f);
-		NextMemberOffset = Members[Members.Num() - 1].DestOffset + Members[Members.Num() - 1].DestSize;
+		ScriptParameterStoreInternal::AddLayoutParameter<FVector2f>(Members, NextMemberOffset, InSrcOffset);
 	}
-	else if (InSrcStruct && (InSrcStruct->GetFName() == NAME_Vector || InSrcStruct->GetFName() == NAME_Vector3d))
+	else if (InSrcStruct == FNiagaraTypeDefinition::GetVec3Struct() || InSrcStruct == FNiagaraTypeDefinition::GetPositionStruct() || SrcStructName == NAME_Vector || SrcStructName == NAME_Vector3d)
 	{
-		uint32 StructFinalSize = (TShaderParameterTypeInfo<FVector3f>::NumRows * TShaderParameterTypeInfo<FVector3f>::NumColumns) * sizeof(float);
-		Members.Emplace(InSrcOffset, OffsetAlign(NextMemberOffset, VectorPaddedSize), StructFinalSize, VectorPaddedSize);
-		InSrcOffset += sizeof(FVector3f);
-		NextMemberOffset = Members[Members.Num() - 1].DestOffset + Members[Members.Num() - 1].DestSize;
+		ScriptParameterStoreInternal::AddLayoutParameter<FVector3f>(Members, NextMemberOffset, InSrcOffset);
 	}
-	else if (InSrcStruct && (InSrcStruct->GetFName() == NAME_Vector4 || InSrcStruct->GetFName() == NAME_Vector4d))
+	else if (SrcStructName == NAME_Vector4 || SrcStructName == NAME_Vector4d || SrcStructName == NAME_Quat || SrcStructName == NAME_Quat4d)
 	{
-		uint32 StructFinalSize = (TShaderParameterTypeInfo<FVector4f>::NumRows * TShaderParameterTypeInfo<FVector4f>::NumColumns) * sizeof(float);
-		Members.Emplace(InSrcOffset, OffsetAlign(NextMemberOffset, VectorPaddedSize), StructFinalSize, VectorPaddedSize);
-		InSrcOffset += sizeof(FVector4f);
-		NextMemberOffset = Members[Members.Num() - 1].DestOffset + Members[Members.Num() - 1].DestSize;
-	}
-	else if (InSrcStruct && (InSrcStruct->GetFName() == NAME_Quat || InSrcStruct->GetFName() == NAME_Quat4d))
-	{
-		uint32 StructFinalSize = (TShaderParameterTypeInfo<FQuat4f>::NumRows * TShaderParameterTypeInfo<FQuat4f>::NumColumns) * sizeof(float);
-		Members.Emplace(InSrcOffset, OffsetAlign(NextMemberOffset, VectorPaddedSize), StructFinalSize, VectorPaddedSize);
-		InSrcOffset += sizeof(FQuat4f);
-		NextMemberOffset = Members[Members.Num() - 1].DestOffset + Members[Members.Num() - 1].DestSize;
+		ScriptParameterStoreInternal::AddLayoutParameter<FVector4f>(Members, NextMemberOffset, InSrcOffset);
 	}
 	else if (InSrcStruct == FNiagaraTypeDefinition::GetVec4Struct() || InSrcStruct == FNiagaraTypeDefinition::GetColorStruct() || InSrcStruct == FNiagaraTypeDefinition::GetQuatStruct())
 	{
-		uint32 StructFinalSize = (TShaderParameterTypeInfo<FVector4f>::NumRows * TShaderParameterTypeInfo<FVector4f>::NumColumns) * sizeof(float);
-		Members.Emplace(InSrcOffset, Align(NextMemberOffset, TShaderParameterTypeInfo<FVector4f>::Alignment), StructFinalSize, StructFinalSize);
-		InSrcOffset += sizeof(FVector4f);
-		NextMemberOffset = Members[Members.Num() - 1].DestOffset + Members[Members.Num() - 1].DestSize;
+		ScriptParameterStoreInternal::AddLayoutParameter<FVector4f>(Members, NextMemberOffset, InSrcOffset);
 	}
 	else if (InSrcStruct == FNiagaraTypeDefinition::GetMatrix4Struct())
 	{
-		uint32 StructFinalSize = (TShaderParameterTypeInfo<FMatrix44f>::NumRows * TShaderParameterTypeInfo<FMatrix44f>::NumColumns) * sizeof(float);
-		Members.Emplace(InSrcOffset, Align(NextMemberOffset, TShaderParameterTypeInfo<FMatrix44f>::Alignment), StructFinalSize, StructFinalSize);
-		InSrcOffset += sizeof(FMatrix44f);
-		NextMemberOffset = Members[Members.Num() - 1].DestOffset + Members[Members.Num() - 1].DestSize;
-	}
-	else if (InSrcStruct == FNiagaraTypeDefinition::GetHalfStruct())
-	{
-		uint32 HalfSize = (/*NumRows*/1 * /*NumColumns*/ 1) * sizeof(FFloat16);
-		Members.Emplace(InSrcOffset, Align(NextMemberOffset, TShaderParameterTypeInfo<FFloat16>::Alignment), HalfSize, HalfSize);
-		InSrcOffset += sizeof(FFloat16);
-		NextMemberOffset = Members[Members.Num() - 1].DestOffset + Members[Members.Num() - 1].DestSize;
-	}
-	else if (InSrcStruct == FNiagaraTypeDefinition::GetHalfVec2Struct())
-	{
-		uint32 StructFinalSize = (/*NumRows*/2 * /*NumColumns*/1) * sizeof(FFloat16);
-		Members.Emplace(InSrcOffset, OffsetAlign(NextMemberOffset, VectorPaddedSize), StructFinalSize, VectorPaddedSize);
-		InSrcOffset += sizeof(FFloat16[2]);
-		NextMemberOffset = Members[Members.Num() - 1].DestOffset + Members[Members.Num() - 1].DestSize;
-	}
-	else if (InSrcStruct == FNiagaraTypeDefinition::GetHalfVec3Struct())
-	{
-		uint32 StructFinalSize = (/*NumRows*/3 * /*NumColumns*/1) * sizeof(FFloat16);
-		Members.Emplace(InSrcOffset, OffsetAlign(NextMemberOffset, VectorPaddedSize), StructFinalSize, VectorPaddedSize);
-		InSrcOffset += sizeof(FFloat16[3]);
-		NextMemberOffset = Members[Members.Num() - 1].DestOffset + Members[Members.Num() - 1].DestSize;
+		ScriptParameterStoreInternal::AddLayoutParameter<FMatrix44f>(Members, NextMemberOffset, InSrcOffset);
 	}
-	else if (InSrcStruct == FNiagaraTypeDefinition::GetHalfVec4Struct())
+	else if (InSrcStruct == FNiagaraTypeDefinition::GetHalfStruct() ||
+			 InSrcStruct == FNiagaraTypeDefinition::GetHalfVec2Struct() ||
+			 InSrcStruct == FNiagaraTypeDefinition::GetHalfVec3Struct() ||
+			 InSrcStruct == FNiagaraTypeDefinition::GetHalfVec4Struct())
 	{
-		uint32 StructFinalSize = (/*NumRows*/4 * /*NumColumns*/1) * sizeof(FFloat16);
-		Members.Emplace(InSrcOffset, OffsetAlign(NextMemberOffset, VectorPaddedSize), StructFinalSize, VectorPaddedSize);
-		InSrcOffset += sizeof(FFloat16[4]);
-		NextMemberOffset = Members[Members.Num() - 1].DestOffset + Members[Members.Num() - 1].DestSize;
+		UE_LOG(LogNiagara, Fatal, TEXT("Half is not supported inside parameter stores"));
 	}
 	else
 	{
diff --git a/Engine/Plugins/FX/Niagara/Source/Niagara/Private/NiagaraSystem.cpp b/Engine/Plugins/FX/Niagara/Source/Niagara/Private/NiagaraSystem.cpp
--- a/Engine/Plugins/FX/Niagara/Source/Niagara/Private/NiagaraSystem.cpp
+++ b/Engine/Plugins/FX/Niagara/Source/Niagara/Private/NiagaraSystem.cpp
@@ -3691,15 +3691,7 @@ void FNiagaraParameterDataSetBindingCollection::BuildInternal(const TArray<FNiag
 			}
 		}
 
-		// we need to take into account potential padding that is in the constant buffers based similar to what is done
-		// in the NiagaraHlslTranslator, where Vec2/Vec3/Position are treated as Vec4.
-		int32 ParameterSize = Var.GetSizeInBytes();
-		const FNiagaraTypeDefinition& Type = Var.GetType();
-		if (Type == FNiagaraTypeDefinition::GetVec2Def() || Type == FNiagaraTypeDefinition::GetVec3Def() || Type == FNiagaraTypeDefinition::GetPositionDef())
-		{
-			ParameterSize = Align(ParameterSize, FNiagaraTypeDefinition::GetVec4Def().GetSize());
-		}
-
+		const int32 ParameterSize = Var.GetSizeInBytes();
 		ParameterOffset += ParameterSize;
 	}
 
diff --git a/Engine/Plugins/FX/Niagara/Source/Niagara/Private/NiagaraSystemInstance.cpp b/Engine/Plugins/FX/Niagara/Source/Niagara/Private/NiagaraSystemInstance.cpp
--- a/Engine/Plugins/FX/Niagara/Source/Niagara/Private/NiagaraSystemInstance.cpp
+++ b/Engine/Plugins/FX/Niagara/Source/Niagara/Private/NiagaraSystemInstance.cpp
@@ -2775,11 +2775,10 @@ void FNiagaraSystemInstance::EvaluateBoundFunction(FName FunctionName, bool& Use
 		if (Script)
 		{
 			const bool IsGpuScript = UNiagaraScript::IsGPUScript(Script->Usage);
-			const auto& VMExecData = Script->GetVMExecutableData();
 
 			if (!UsedOnGpu && IsGpuScript)
 			{
-				for (const FNiagaraDataInterfaceGPUParamInfo& DIParamInfo : VMExecData.DIParamInfo)
+				for (const FNiagaraDataInterfaceGPUParamInfo& DIParamInfo : Script->GetDataInterfaceGPUParamInfos())
 				{
 					auto GpuFunctionPredicate = [&](const FNiagaraDataInterfaceGeneratedFunction& DIFunction)
 					{
@@ -2801,6 +2800,7 @@ void FNiagaraSystemInstance::EvaluateBoundFunction(FName FunctionName, bool& Use
 					return BindingInfo.Name == FunctionName;
 				};
 
+				const FNiagaraVMExecutableData& VMExecData =  Script->GetVMExecutableData();
 				if (VMExecData.CalledVMExternalFunctions.ContainsByPredicate(CpuFunctionPredicate))
 				{
 					UsedOnCpu = true;
diff --git a/Engine/Plugins/FX/Niagara/Source/Niagara/Public/NiagaraTypes.h b/Engine/Plugins/FX/Niagara/Source/Niagara/Public/NiagaraTypes.h
--- a/Engine/Plugins/FX/Niagara/Source/Niagara/Public/NiagaraTypes.h
+++ b/Engine/Plugins/FX/Niagara/Source/Niagara/Public/NiagaraTypes.h
@@ -2048,20 +2048,26 @@ struct alignas(16) FNiagaraOwnerParameters
 	NIAGARA_API static const TArray<FNiagaraVariable>& GetVariables();
 #endif
 
-	FMatrix44f EngineLocalToWorld = FMatrix44f::Identity;
-	FMatrix44f EngineWorldToLocal = FMatrix44f::Identity;
-	FMatrix44f EngineLocalToWorldTransposed = FMatrix44f::Identity;
-	FMatrix44f EngineWorldToLocalTransposed = FMatrix44f::Identity;
-	FMatrix44f EngineLocalToWorldNoScale = FMatrix44f::Identity;
-	FMatrix44f EngineWorldToLocalNoScale = FMatrix44f::Identity;
-	FQuat4f EngineRotation = FQuat4f::Identity;
-	FVector4f EnginePosition = FVector4f(ForceInitToZero);
-	FVector4f EngineVelocity = FVector4f(ForceInitToZero);
-	FVector4f EngineXAxis = FVector4f(1.0f, 0.0f, 0.0f, 0.0f);
-	FVector4f EngineYAxis = FVector4f(0.0f, 1.0f, 0.0f, 0.0f);
-	FVector4f EngineZAxis = FVector4f(0.0f, 0.0f, 1.0f, 0.0f);
-	FVector4f EngineScale = FVector4f(1.0f, 1.0f, 1.0f, 0.0f);
-	FVector4f EngineLWCTile = FVector4f(ForceInitToZero);
+	FMatrix44f	EngineLocalToWorld = FMatrix44f::Identity;
+	FMatrix44f	EngineWorldToLocal = FMatrix44f::Identity;
+	FMatrix44f	EngineLocalToWorldTransposed = FMatrix44f::Identity;
+	FMatrix44f	EngineWorldToLocalTransposed = FMatrix44f::Identity;
+	FMatrix44f	EngineLocalToWorldNoScale = FMatrix44f::Identity;
+	FMatrix44f	EngineWorldToLocalNoScale = FMatrix44f::Identity;
+	FQuat4f		EngineRotation = FQuat4f::Identity;
+	FVector3f	EnginePosition = FVector3f(ForceInitToZero);
+	float		_Pad0;
+	FVector3f	EngineVelocity = FVector3f(ForceInitToZero);
+	float		_Pad1;
+	FVector3f	EngineXAxis = FVector3f(1.0f, 0.0f, 0.0f);
+	float		_Pad2;
+	FVector3f	EngineYAxis = FVector3f(0.0f, 1.0f, 0.0f);
+	float		_Pad3;
+	FVector3f	EngineZAxis = FVector3f(0.0f, 0.0f, 1.0f);
+	float		_Pad4;
+	FVector3f	EngineScale = FVector3f(1.0f, 1.0f, 1.0f);
+	float		_Pad5;
+	FVector4f	EngineLWCTile = FVector4f(ForceInitToZero);
 };
 
 // Any change to this structure, or it's GetVariables implementation will require a bump in the CustomNiagaraVersion so that we
diff --git a/Engine/Plugins/FX/Niagara/Source/NiagaraEditor/Private/NiagaraHlslTranslator.cpp b/Engine/Plugins/FX/Niagara/Source/NiagaraEditor/Private/NiagaraHlslTranslator.cpp
--- a/Engine/Plugins/FX/Niagara/Source/NiagaraEditor/Private/NiagaraHlslTranslator.cpp
+++ b/Engine/Plugins/FX/Niagara/Source/NiagaraEditor/Private/NiagaraHlslTranslator.cpp
@@ -111,7 +111,7 @@ void FNiagaraShaderQueueTickable::ProcessQueue()
 
 		FNiagaraComputeShaderCompilationOutput NewCompilationOutput;
 
-		ShaderScript->BuildScriptParametersMetadata(CompilableScript->GetVMExecutableData().DIParamInfo);
+		ShaderScript->BuildScriptParametersMetadata(CompilableScript->GetVMExecutableData().ShaderScriptParametersMetadata);
 		ShaderScript->SetSourceName(TEXT("NiagaraComputeShader"));
 		UNiagaraEmitter* Emitter = Cast<UNiagaraEmitter>(CompilableScript->GetOuter());
 		if (Emitter && Emitter->GetUniqueEmitterName().Len() > 0)
@@ -996,7 +996,7 @@ void FHlslNiagaraTranslator::BuildConstantBuffer(ENiagaraCodeChunkMode ChunkMode
 	for (const FNiagaraVariable& Variable : T::GetVariables())
 	{
 		const FString SymbolName = GetSanitizedSymbolName(Variable.GetName().ToString(), true);
-		AddUniformChunk(SymbolName, Variable, ChunkMode, true);
+		AddUniformChunk(SymbolName, Variable, ChunkMode);
 	}
 }
 
@@ -1711,35 +1711,71 @@ const FNiagaraTranslateResults &FHlslNiagaraTranslator::Translate(const FNiagara
 			}
 		}
 
-		ENiagaraCodeChunkMode ChunkModes[] = { ENiagaraCodeChunkMode::GlobalConstant, ENiagaraCodeChunkMode::SystemConstant, ENiagaraCodeChunkMode::OwnerConstant, ENiagaraCodeChunkMode::EmitterConstant, ENiagaraCodeChunkMode::Uniform };
-		FString ConstantBufferNames[] = { TEXT("FNiagaraGlobalParameters"), TEXT("FNiagaraSystemParameters"), TEXT("FNiagaraOwnerParameters"), TEXT("FNiagaraEmitterParameters"), TEXT("FNiagaraExternalParameters") };
-
-		static_assert(UE_ARRAY_COUNT(ChunkModes) == UE_ARRAY_COUNT(ConstantBufferNames), "Mismatch between ChunkModes and ConstantBufferNames");
+		const TPair<ENiagaraCodeChunkMode, FString> ChunkModeInfos[] =
+		{
+			MakeTuple(ENiagaraCodeChunkMode::GlobalConstant,	TEXT("FNiagaraGlobalParameters")),
+			MakeTuple(ENiagaraCodeChunkMode::SystemConstant,	TEXT("FNiagaraSystemParameters")),
+			MakeTuple(ENiagaraCodeChunkMode::OwnerConstant,		TEXT("FNiagaraOwnerParameters")),
+			MakeTuple(ENiagaraCodeChunkMode::EmitterConstant,	TEXT("FNiagaraEmitterParameters")),
+			MakeTuple(ENiagaraCodeChunkMode::Uniform,			TEXT("FNiagaraExternalParameters")),
+		};
 
-		FString SymbolPrefix[] = { TEXT(""), INTERPOLATED_PARAMETER_PREFIX };
+		const TCHAR* InterpPrefix[] = { TEXT(""), INTERPOLATED_PARAMETER_PREFIX };
 
-		for (int32 PrevIt = 0; PrevIt < (bInterpolateParams ? 2 : 1); ++PrevIt)
+		// GPU simulation we prefer loose bindings, all but the external cbuffer are loose.  External contains structures which we
+		// don't understand when generating parameters so for the moment we can't pack it in easily.
+		// This is two separate loops as the VVM assumes cbuffer order is AllCurrent -> AllPrevious, whereas GPU we bind in order to minimize parameter copies
+		if (TranslationOptions.SimTarget == ENiagaraSimTarget::GPUComputeSim)
 		{
-			for (int32 ChunkModeIt = 0; ChunkModeIt < UE_ARRAY_COUNT(ChunkModes); ++ChunkModeIt)
+			for (const TPair<ENiagaraCodeChunkMode, FString>& ChunkModeInfo : ChunkModeInfos)
 			{
-				int32 ChunkMode = static_cast<int32>(ChunkModes[ChunkModeIt]);
-				const FString BufferName = SymbolPrefix[PrevIt] + ConstantBufferNames[ChunkModeIt];
+				for (int32 InterpIt = 0; InterpIt < (bInterpolateParams ? 2 : 1); ++InterpIt)
+				{
+					const bool bIsCBuffer = ChunkModeInfo.Key == ENiagaraCodeChunkMode::Uniform;
+					if (bIsCBuffer)
+					{
+						HlslOutput.Appendf(TEXT("cbuffer %s%s\n{\n"), InterpPrefix[InterpIt], *ChunkModeInfo.Value);
+					}
+
+					for (int32 ChunkOffset : ChunksByMode[int(ChunkModeInfo.Key)])
+					{
+						FNiagaraVariable BufferVariable(CodeChunks[ChunkOffset].Type, FName(*CodeChunks[ChunkOffset].SymbolName));
+						if (IsVariableInUniformBuffer(BufferVariable))
+						{
+							FNiagaraCodeChunk Chunk = CodeChunks[ChunkOffset];
+							Chunk.SymbolName = InterpPrefix[InterpIt] + Chunk.SymbolName;
 
-				HlslOutput += TEXT("cbuffer ") + BufferName + TEXT("\n{\n");
+							HlslOutput += TEXT("\t") + GetCode(Chunk);
+						}
+					}
 
-				for (int32 i = 0; i < ChunksByMode[ChunkMode].Num(); ++i)
+					if (bIsCBuffer)
+					{
+						HlslOutput += TEXT("}\n\n");
+					}
+				}
+			}
+		}
+		else
+		{
+			for (int32 InterpIt = 0; InterpIt < (bInterpolateParams ? 2 : 1); ++InterpIt)
+			{
+				for (const TPair<ENiagaraCodeChunkMode, FString>& ChunkModeInfo : ChunkModeInfos)
 				{
-					FNiagaraVariable BufferVariable(CodeChunks[ChunksByMode[ChunkMode][i]].Type, FName(*CodeChunks[ChunksByMode[ChunkMode][i]].SymbolName));
-					if (IsVariableInUniformBuffer(BufferVariable))
+					HlslOutput.Appendf(TEXT("cbuffer %s%s\n{\n"), InterpPrefix[InterpIt], *ChunkModeInfo.Value);
+					for (int32 ChunkOffset : ChunksByMode[int(ChunkModeInfo.Key)])
 					{
-						FNiagaraCodeChunk Chunk = CodeChunks[ChunksByMode[ChunkMode][i]];
-						Chunk.SymbolName = SymbolPrefix[PrevIt] + Chunk.SymbolName;
+						FNiagaraVariable BufferVariable(CodeChunks[ChunkOffset].Type, FName(*CodeChunks[ChunkOffset].SymbolName));
+						if (IsVariableInUniformBuffer(BufferVariable))
+						{
+							FNiagaraCodeChunk Chunk = CodeChunks[ChunkOffset];
+							Chunk.SymbolName = InterpPrefix[InterpIt] + Chunk.SymbolName;
 
-						HlslOutput += TEXT("\t") + GetCode(Chunk);
+							HlslOutput += TEXT("\t") + GetCode(Chunk);
+						}
 					}
+					HlslOutput += TEXT("}\n\n");
 				}
-
-				HlslOutput += TEXT("}\n\n");
 			}
 		}
 
@@ -2033,7 +2069,8 @@ const FNiagaraTranslateResults &FHlslNiagaraTranslator::Translate(const FNiagara
 		HlslOutput = Preamble + TEXT("\n\n") +  HlslOutput;
 
 		// We may have created some transient data interfaces. This cleans up the ones that we created.
-		CompilationOutput.ScriptData.DIParamInfo = DIParamInfo;
+		CompilationOutput.ScriptData.ShaderScriptParametersMetadata = ShaderScriptParametersMetadata;
+
 		if (InstanceRead.Variables.Num() == 1 && InstanceRead.Variables[0].GetName() == TEXT("Particles.UniqueID")) {
 			// NOTE(mv): Explicitly allow reading from Particles.UniqueID, as it is an engine managed variable and 
 			//           is written to before Simulate() in the SpawnScript...
@@ -2763,7 +2800,7 @@ void FHlslNiagaraTranslator::DefineDataInterfaceHLSL(FString& InHlslOutput)
 				InterfaceClasses.Add(Info.Type.GetFName());
 			}
 
-			FNiagaraDataInterfaceGPUParamInfo& DIInstanceInfo = DIParamInfo.AddDefaulted_GetRef();
+			FNiagaraDataInterfaceGPUParamInfo& DIInstanceInfo = ShaderScriptParametersMetadata.DataInterfaceParamInfo.AddDefaulted_GetRef();
 			ConvertCompileInfoToParamInfo(Info, DIInstanceInfo);
 
 			CDO->GetParameterDefinitionHLSL(DIInstanceInfo, InterfaceUniformHLSL);
@@ -4135,7 +4172,7 @@ FString FHlslNiagaraTranslator::GeneratedConstantString(FVector4 Constant)
 	return FString::Format(TEXT("float4({0}, {1}, {2}, {3})"), Args);
 }
 
-int32 FHlslNiagaraTranslator::AddUniformChunk(FString SymbolName, const FNiagaraVariable& InVariable, ENiagaraCodeChunkMode ChunkMode, bool AddPadding)
+int32 FHlslNiagaraTranslator::AddUniformChunk(FString SymbolName, const FNiagaraVariable& InVariable, ENiagaraCodeChunkMode ChunkMode)
 {
 	const FNiagaraTypeDefinition& Type = InVariable.GetType();
 
@@ -4155,20 +4192,6 @@ int32 FHlslNiagaraTranslator::AddUniformChunk(FString SymbolName, const FNiagara
 		Chunk.Type = Type;
 		Chunk.Original = InVariable;
 
-		if (AddPadding)
-		{
-			if (Type == FNiagaraTypeDefinition::GetVec2Def())
-			{
-				Chunk.Type = FNiagaraTypeDefinition::GetVec4Def();
-				Chunk.ComponentMask = TEXT(".xy");
-			}
-			else if (Type == FNiagaraTypeDefinition::GetVec3Def() || Type == FNiagaraTypeDefinition::GetPositionDef())
-			{
-				Chunk.Type = FNiagaraTypeDefinition::GetVec4Def();
-				Chunk.ComponentMask = TEXT(".xyz");
-			}
-		}
-
 		Chunk.Mode = ChunkMode;
 
 		ChunksByMode[static_cast<int32>(ChunkMode)].Add(Ret);
@@ -5503,7 +5526,7 @@ bool FHlslNiagaraTranslator::ParameterMapRegisterExternalConstantNamespaceVariab
 					CompilationOutput.ScriptData.Parameters.SetOrAdd(InVariable);
 				}
 
-				UniformChunk = AddUniformChunk(SymbolNameDefined, InVariable, ENiagaraCodeChunkMode::Uniform, UNiagaraScript::IsGPUScript(CompileOptions.TargetUsage));
+				UniformChunk = AddUniformChunk(SymbolNameDefined, InVariable, ENiagaraCodeChunkMode::Uniform);
 			}
 			else
 			{
diff --git a/Engine/Plugins/FX/Niagara/Source/NiagaraEditor/Private/NiagaraHlslTranslator.h b/Engine/Plugins/FX/Niagara/Source/NiagaraEditor/Private/NiagaraHlslTranslator.h
--- a/Engine/Plugins/FX/Niagara/Source/NiagaraEditor/Private/NiagaraHlslTranslator.h
+++ b/Engine/Plugins/FX/Niagara/Source/NiagaraEditor/Private/NiagaraHlslTranslator.h
@@ -539,7 +539,7 @@ protected:
 	int32 AddSourceChunk(FString SymbolName, const FNiagaraTypeDefinition& Type, bool bSanitize = true);
 
 	/** Add a chunk defining a uniform value. */
-	int32 AddUniformChunk(FString SymbolName, const FNiagaraVariable& InVariable, ENiagaraCodeChunkMode ChunkMode, bool AddPadding);
+	int32 AddUniformChunk(FString SymbolName, const FNiagaraVariable& InVariable, ENiagaraCodeChunkMode ChunkMode);
 
 	/* Add a chunk that is written to the body of the shader code. */
 	int32 AddBodyChunk(FString SymbolName, FString Definition, const FNiagaraTypeDefinition& Type, TArray<int32>& SourceChunks, bool bDecl = true, bool bIsTerminated = true);
@@ -606,7 +606,6 @@ public:
 	void DefineDataSetVariableWrites(FString &HlslOutput, const FNiagaraDataSetID& Id, int32 DataSetIndex, const TArray<FNiagaraVariable>& WriteVars);
 	void DefineDataInterfaceHLSL(FString &HlslOutput);
 	void DefineExternalFunctionsHLSL(FString &HlslOutput);
-	TArray<FNiagaraDataInterfaceGPUParamInfo>& GetDataInterfaceParameters() { return DIParamInfo; }
 
 	// Format string should have up to 5 entries, {{0} = Computed Variable Suffix, {1} = Float or Int, {2} = Data Set Index, {3} = Register Index, {4} Default value for that type.
 	void GatherVariableForDataSetAccess(const FNiagaraVariable& Variable, FString Format, int32& RegisterIdxInt, int32& RegisterIdxFloat, int32& RegisterIdxHalf, int32 DataSetIndex, FString InstanceIdxSymbol, FString &HlslOutput, bool bWriteHLSL = true);
@@ -846,7 +845,7 @@ private:
 	int32 WriteIdx;
 
 	// Parameter data per data interface.
-	TArray< FNiagaraDataInterfaceGPUParamInfo > DIParamInfo;
+	FNiagaraShaderScriptParametersMetadata ShaderScriptParametersMetadata;
 	
 	/** Stack of currently tracked stats scopes. */
 	TArray<int32> StatScopeStack;
diff --git a/Engine/Plugins/FX/Niagara/Source/NiagaraShader/Private/NiagaraShader.cpp b/Engine/Plugins/FX/Niagara/Source/NiagaraShader/Private/NiagaraShader.cpp
--- a/Engine/Plugins/FX/Niagara/Source/NiagaraShader/Private/NiagaraShader.cpp
+++ b/Engine/Plugins/FX/Niagara/Source/NiagaraShader/Private/NiagaraShader.cpp
@@ -1132,12 +1132,31 @@ FNiagaraShader::FNiagaraShader(const FNiagaraShaderType::CompiledShaderInitializ
 {
 	check(!DebugDescription.IsEmpty());
 
-	// Currently we can not bind all parameters due to the external constant buffers
-	// Once they are inside parameter structure and UseLegacyShaderBindings are all true we can completely bind with validation
-	const bool bShouldBindEverything = false || (GNiagaraShaderForceBindEverything != 0);
+	// Legacy bindings
+	bNeedsViewUniformBuffer = Initializer.ParameterMap.ContainsParameterAllocation(TEXT("View"));
+	ExternalConstantBufferParam[0].Bind(Initializer.ParameterMap, TEXT("FNiagaraExternalParameters"));
+	ExternalConstantBufferParam[1].Bind(Initializer.ParameterMap, TEXT("PREV_FNiagaraExternalParameters"));
 
-	// Bind parameters
+	// We can only use validation for missing parameter bindings if we don't use the external constant buffer and
+	// all DataInterfaces use the new parameter bindings.  Otherwise we can't bind in a single pass.
 	const FNiagaraShaderType::FParameters* InitializerParameters = static_cast<const FNiagaraShaderType::FParameters*>(Initializer.Parameters);
+	INiagaraShaderModule* ShaderModule = INiagaraShaderModule::Get();
+	bool bShouldBindEverything = true;
+	
+	if (GNiagaraShaderForceBindEverything == 0)
+	{
+		bShouldBindEverything &= bNeedsViewUniformBuffer == false;
+		bShouldBindEverything &= ExternalConstantBufferParam[0].IsBound() == false;
+		bShouldBindEverything &= ExternalConstantBufferParam[1].IsBound() == false;
+
+		for (const FNiagaraDataInterfaceGPUParamInfo& DataInterfaceParamInfo : InitializerParameters->ScriptParametersMetadata->DataInterfaceParamInfo)
+		{
+			UNiagaraDataInterfaceBase* CDODataInterface = ShaderModule->RequestDefaultDataInterface(*DataInterfaceParamInfo.DIClassName);
+			bShouldBindEverything &= CDODataInterface && CDODataInterface->UseLegacyShaderBindings() == false;
+		}
+	}
+
+	// Bind parameters
 	Bindings.BindForLegacyShaderParameters(
 		this,
 		Initializer.PermutationId,
@@ -1146,23 +1165,7 @@ FNiagaraShader::FNiagaraShader(const FNiagaraShaderType::CompiledShaderInitializ
 		bShouldBindEverything
 	);
 
-	bNeedsViewUniformBuffer = Initializer.ParameterMap.ContainsParameterAllocation(TEXT("View"));
-
-	// Legacy bindings
-	GlobalConstantBufferParam[0].Bind(Initializer.ParameterMap, TEXT("FNiagaraGlobalParameters"));
-	SystemConstantBufferParam[0].Bind(Initializer.ParameterMap, TEXT("FNiagaraSystemParameters"));
-	OwnerConstantBufferParam[0].Bind(Initializer.ParameterMap, TEXT("FNiagaraOwnerParameters"));
-	EmitterConstantBufferParam[0].Bind(Initializer.ParameterMap, TEXT("FNiagaraEmitterParameters"));
-	ExternalConstantBufferParam[0].Bind(Initializer.ParameterMap, TEXT("FNiagaraExternalParameters"));
-
-	GlobalConstantBufferParam[1].Bind(Initializer.ParameterMap, TEXT("PREV_FNiagaraGlobalParameters"));
-	SystemConstantBufferParam[1].Bind(Initializer.ParameterMap, TEXT("PREV_FNiagaraSystemParameters"));
-	OwnerConstantBufferParam[1].Bind(Initializer.ParameterMap, TEXT("PREV_FNiagaraOwnerParameters"));
-	EmitterConstantBufferParam[1].Bind(Initializer.ParameterMap, TEXT("PREV_FNiagaraEmitterParameters"));
-	ExternalConstantBufferParam[1].Bind(Initializer.ParameterMap, TEXT("PREV_FNiagaraExternalParameters"));
-
 	// Gather data interface bindings
-	INiagaraShaderModule* ShaderModule = INiagaraShaderModule::Get();
 	DataInterfaceParameters.Empty(InitializerParameters->ScriptParametersMetadata->DataInterfaceParamInfo.Num());
 	for (const FNiagaraDataInterfaceGPUParamInfo& DataInterfaceParamInfo : InitializerParameters->ScriptParametersMetadata->DataInterfaceParamInfo)
 	{
diff --git a/Engine/Plugins/FX/Niagara/Source/NiagaraShader/Private/NiagaraShared.cpp b/Engine/Plugins/FX/Niagara/Source/NiagaraShader/Private/NiagaraShared.cpp
--- a/Engine/Plugins/FX/Niagara/Source/NiagaraShader/Private/NiagaraShared.cpp
+++ b/Engine/Plugins/FX/Niagara/Source/NiagaraShader/Private/NiagaraShared.cpp
@@ -494,10 +494,6 @@ void FNiagaraShaderScript::UpdateCachedData_PostCompile(bool bCalledFromSerializ
 	FNiagaraShaderMapCachedData CachedData;
 	CachedData.NumPermutations = GetNumPermutations();
 	CachedData.bIsComplete = 1;
-	CachedData.bGlobalConstantBufferUsed = 0;
-	CachedData.bSystemConstantBufferUsed = 0;
-	CachedData.bOwnerConstantBufferUsed = 0;
-	CachedData.bEmitterConstantBufferUsed = 0;
 	CachedData.bExternalConstantBufferUsed = 0;
 	CachedData.bViewUniformBufferUsed = 0;
 
@@ -516,10 +512,6 @@ void FNiagaraShaderScript::UpdateCachedData_PostCompile(bool bCalledFromSerializ
 			for (int i = 0; i < 2; ++i)
 			{
 				const uint32 BitToSet = 1 << i;
-				CachedData.bGlobalConstantBufferUsed |= NiagaraShader->GlobalConstantBufferParam[i].IsBound() ? BitToSet : 0;
-				CachedData.bSystemConstantBufferUsed |= NiagaraShader->SystemConstantBufferParam[i].IsBound() ? BitToSet : 0;
-				CachedData.bOwnerConstantBufferUsed |= NiagaraShader->OwnerConstantBufferParam[i].IsBound() ? BitToSet : 0;
-				CachedData.bEmitterConstantBufferUsed |= NiagaraShader->EmitterConstantBufferParam[i].IsBound() ? BitToSet : 0;
 				CachedData.bExternalConstantBufferUsed |= NiagaraShader->ExternalConstantBufferParam[i].IsBound() ? BitToSet : 0;
 			}
 			CachedData.bViewUniformBufferUsed |= NiagaraShader->bNeedsViewUniformBuffer;
@@ -679,10 +671,10 @@ void FNiagaraShaderScript::FinishCompilation()
 
 #endif
 
-void FNiagaraShaderScript::BuildScriptParametersMetadata(const TArray<FNiagaraDataInterfaceGPUParamInfo>& InDIParamInfo)
+void FNiagaraShaderScript::BuildScriptParametersMetadata(const FNiagaraShaderScriptParametersMetadata& InScriptParametersMetadata)
 {
 	ScriptParametersMetadata = MakeShared<FNiagaraShaderScriptParametersMetadata>();
-	ScriptParametersMetadata->DataInterfaceParamInfo = InDIParamInfo;
+	ScriptParametersMetadata->DataInterfaceParamInfo = InScriptParametersMetadata.DataInterfaceParamInfo;
 
 	FShaderParametersMetadataBuilder ShaderMetadataBuilder(TShaderParameterStructTypeInfo<FNiagaraShader::FParameters>::GetStructMetadata());
 
@@ -694,13 +686,12 @@ void FNiagaraShaderScript::BuildScriptParametersMetadata(const TArray<FNiagaraDa
 		if (CDODataInterface == nullptr)
 		{
 			Invalidate();
-			return;
 		}
 
 		if (CDODataInterface->UseLegacyShaderBindings() == false)
 		{
 			const uint32 NextMemberOffset = ShaderMetadataBuilder.GetNextMemberOffset();
-			FNiagaraShaderParametersBuilder ShaderParametersBuilder(DataInterfaceParamInfo, ShaderMetadataBuilder);
+			FNiagaraShaderParametersBuilder ShaderParametersBuilder(DataInterfaceParamInfo, ScriptParametersMetadata->LooseMetadataNames, ShaderMetadataBuilder);
 			CDODataInterface->BuildShaderParameters(ShaderParametersBuilder);
 			if (NextMemberOffset != ShaderMetadataBuilder.GetNextMemberOffset())
 			{
@@ -709,7 +700,7 @@ void FNiagaraShaderScript::BuildScriptParametersMetadata(const TArray<FNiagaraDa
 		}
 	}
 
-	ScriptParametersMetadata->ShaderParametersMetadata.Reset(ShaderMetadataBuilder.Build(FShaderParametersMetadata::EUseCase::ShaderParameterStruct, TEXT("FNiagaraShaderScript")));
+	ScriptParametersMetadata->ShaderParametersMetadata = MakeShareable<FShaderParametersMetadata>(ShaderMetadataBuilder.Build(FShaderParametersMetadata::EUseCase::ShaderParameterStruct, TEXT("FNiagaraShaderScript")));
 }
 
 FNiagaraShaderRef FNiagaraShaderScript::GetShader(int32 PermutationId) const
diff --git a/Engine/Plugins/FX/Niagara/Source/NiagaraShader/Public/NiagaraShader.h b/Engine/Plugins/FX/Niagara/Source/NiagaraShader/Public/NiagaraShader.h
--- a/Engine/Plugins/FX/Niagara/Source/NiagaraShader/Public/NiagaraShader.h
+++ b/Engine/Plugins/FX/Niagara/Source/NiagaraShader/Public/NiagaraShader.h
@@ -28,7 +28,122 @@ class NIAGARASHADER_API FNiagaraShader : public FShader
 {
 public:
 	DECLARE_SHADER_TYPE(FNiagaraShader, Niagara);
-	//SHADER_USE_PARAMETER_STRUCT(FNiagaraShader, FShader);
+
+	// This structure is a replication of FNiagaraGlobalParameters with interpolated parameters includes
+	BEGIN_SHADER_PARAMETER_STRUCT(FGlobalParameters, )
+		SHADER_PARAMETER(float,		Engine_DeltaTime)
+		SHADER_PARAMETER(float,		Engine_InverseDeltaTime)
+		SHADER_PARAMETER(float,		Engine_Time)
+		SHADER_PARAMETER(float,		Engine_RealTime)
+		SHADER_PARAMETER(int32,		Engine_QualityLevel)
+		SHADER_PARAMETER(int32,		Engine_Pad0)
+		SHADER_PARAMETER(int32,		Engine_Pad1)
+		SHADER_PARAMETER(int32,		Engine_Pad2)
+
+		SHADER_PARAMETER(float,		PREV_Engine_DeltaTime)
+		SHADER_PARAMETER(float,		PREV_Engine_InverseDeltaTime)
+		SHADER_PARAMETER(float,		PREV_Engine_Time)
+		SHADER_PARAMETER(float,		PREV_Engine_RealTime)
+		SHADER_PARAMETER(int32,		PREV_Engine_QualityLevel)
+		SHADER_PARAMETER(int32,		PREV_Engine_Pad0)
+		SHADER_PARAMETER(int32,		PREV_Engine_Pad1)
+		SHADER_PARAMETER(int32,		PREV_Engine_Pad2)
+	END_SHADER_PARAMETER_STRUCT()
+
+	// This structure is a replication of FNiagaraSystemParameters with interpolated parameters includes
+	BEGIN_SHADER_PARAMETER_STRUCT(FSystemParameters, )
+		SHADER_PARAMETER(float,		Engine_Owner_TimeSinceRendered)
+		SHADER_PARAMETER(float,		Engine_Owner_LODDistance)
+		SHADER_PARAMETER(float,		Engine_Owner_LODDistanceFraction)
+		SHADER_PARAMETER(float,		Engine_System_Age)
+		SHADER_PARAMETER(uint32,	Engine_Owner_ExecutionState)
+		SHADER_PARAMETER(int32,		Engine_System_TickCount)
+		SHADER_PARAMETER(int32,		Engine_System_NumEmitters)
+		SHADER_PARAMETER(int32,		Engine_System_NumEmittersAlive)
+		SHADER_PARAMETER(int32,		Engine_System_SignificanceIndex)
+		SHADER_PARAMETER(int32,		Engine_System_RandomSeed)
+		SHADER_PARAMETER(int32,		System_Pad0)
+		SHADER_PARAMETER(int32,		System_Pad1)
+
+		SHADER_PARAMETER(float,		PREV_Engine_Owner_TimeSinceRendered)
+		SHADER_PARAMETER(float,		PREV_Engine_Owner_LODDistance)
+		SHADER_PARAMETER(float,		PREV_Engine_Owner_LODDistanceFraction)
+		SHADER_PARAMETER(float,		PREV_Engine_System_Age)
+		SHADER_PARAMETER(uint32,	PREV_Engine_Owner_ExecutionState)
+		SHADER_PARAMETER(int32,		PREV_Engine_System_TickCount)
+		SHADER_PARAMETER(int32,		PREV_Engine_System_NumEmitters)
+		SHADER_PARAMETER(int32,		PREV_Engine_System_NumEmittersAlive)
+		SHADER_PARAMETER(int32,		PREV_Engine_System_SignificanceIndex)
+		SHADER_PARAMETER(int32,		PREV_Engine_System_RandomSeed)
+		SHADER_PARAMETER(int32,		PREV_System_Pad0)
+		SHADER_PARAMETER(int32,		PREV_System_Pad1)
+	END_SHADER_PARAMETER_STRUCT()
+
+	// This structure is a replication of FNiagaraOwnerParameters with interpolated parameters includes
+	BEGIN_SHADER_PARAMETER_STRUCT(FOwnerParameters, )
+		SHADER_PARAMETER(FMatrix44f,	Engine_Owner_SystemLocalToWorld)
+		SHADER_PARAMETER(FMatrix44f,	Engine_Owner_SystemWorldToLocal)
+		SHADER_PARAMETER(FMatrix44f,	Engine_Owner_SystemLocalToWorldTransposed)
+		SHADER_PARAMETER(FMatrix44f,	Engine_Owner_SystemWorldToLocalTransposed)
+		SHADER_PARAMETER(FMatrix44f,	Engine_Owner_SystemLocalToWorldNoScale)
+		SHADER_PARAMETER(FMatrix44f,	Engine_Owner_SystemWorldToLocalNoScale)
+		SHADER_PARAMETER(FQuat4f,		Engine_Owner_Rotation)
+		SHADER_PARAMETER(FVector3f,		Engine_Owner_Position)
+		SHADER_PARAMETER(float,			Engine_Owner_Pad0)
+		SHADER_PARAMETER(FVector3f,		Engine_Owner_Velocity)
+		SHADER_PARAMETER(float,			Engine_Owner_Pad1)
+		SHADER_PARAMETER(FVector3f,		Engine_Owner_SystemXAxis)
+		SHADER_PARAMETER(float,			Engine_Owner_Pad2)
+		SHADER_PARAMETER(FVector3f,		Engine_Owner_SystemYAxis)
+		SHADER_PARAMETER(float,			Engine_Owner_Pad3)
+		SHADER_PARAMETER(FVector3f,		Engine_Owner_SystemZAxis)
+		SHADER_PARAMETER(float,			Engine_Owner_Pad4)
+		SHADER_PARAMETER(FVector3f,		Engine_Owner_Scale)
+		SHADER_PARAMETER(float,			Engine_Owner_Pad5)
+		SHADER_PARAMETER(FVector4f,		Engine_Owner_LWCTile)
+
+		SHADER_PARAMETER(FMatrix44f,	PREV_Engine_Owner_SystemLocalToWorld)
+		SHADER_PARAMETER(FMatrix44f,	PREV_Engine_Owner_SystemWorldToLocal)
+		SHADER_PARAMETER(FMatrix44f,	PREV_Engine_Owner_SystemLocalToWorldTransposed)
+		SHADER_PARAMETER(FMatrix44f,	PREV_Engine_Owner_SystemWorldToLocalTransposed)
+		SHADER_PARAMETER(FMatrix44f,	PREV_Engine_Owner_SystemLocalToWorldNoScale)
+		SHADER_PARAMETER(FMatrix44f,	PREV_Engine_Owner_SystemWorldToLocalNoScale)
+		SHADER_PARAMETER(FQuat4f,		PREV_Engine_Owner_Rotation)
+		SHADER_PARAMETER(FVector3f,		PREV_Engine_Owner_Position)
+		SHADER_PARAMETER(float,			PREV_Engine_Owner_Pad0)
+		SHADER_PARAMETER(FVector3f,		PREV_Engine_Owner_Velocity)
+		SHADER_PARAMETER(float,			PREV_Engine_Owner_Pad1)
+		SHADER_PARAMETER(FVector3f,		PREV_Engine_Owner_SystemXAxis)
+		SHADER_PARAMETER(float,			PREV_Engine_Owner_Pad2)
+		SHADER_PARAMETER(FVector3f,		PREV_Engine_Owner_SystemYAxis)
+		SHADER_PARAMETER(float,			PREV_Engine_Owner_Pad3)
+		SHADER_PARAMETER(FVector3f,		PREV_Engine_Owner_SystemZAxis)
+		SHADER_PARAMETER(float,			PREV_Engine_Owner_Pad4)
+		SHADER_PARAMETER(FVector3f,		PREV_Engine_Owner_Scale)
+		SHADER_PARAMETER(float,			PREV_Engine_Owner_Pad5)
+		SHADER_PARAMETER(FVector4f,		PREV_Engine_Owner_LWCTile)
+	END_SHADER_PARAMETER_STRUCT()
+
+	// This structure is a replication of FNiagaraEmitterParameters with interpolated parameters includes
+	BEGIN_SHADER_PARAMETER_STRUCT(FEmitterParameters, )
+		SHADER_PARAMETER(int32,		Engine_Emitter_NumParticle)
+		SHADER_PARAMETER(int32,		Engine_Emitter_TotalSpawnedParticles)
+		SHADER_PARAMETER(float,		Engine_Emitter_SpawnCountScale)
+		SHADER_PARAMETER(float,		Emitter_Age)
+		SHADER_PARAMETER(int32,		Emitter_RandomSeed)
+		SHADER_PARAMETER(int32,		Engine_Emitter_InstanceSeed)
+		SHADER_PARAMETER(int32,		Emitter_Pad0)
+		SHADER_PARAMETER(int32,		Emitter_Pad1)
+
+		SHADER_PARAMETER(int32,		PREV_Engine_Emitter_NumParticle)
+		SHADER_PARAMETER(int32,		PREV_Engine_Emitter_TotalSpawnedParticles)
+		SHADER_PARAMETER(float,		PREV_Engine_Emitter_SpawnCountScale)
+		SHADER_PARAMETER(float,		PREV_Emitter_Age)
+		SHADER_PARAMETER(int32,		PREV_Emitter_RandomSeed)
+		SHADER_PARAMETER(int32,		PREV_Engine_Emitter_InstanceSeed)
+		SHADER_PARAMETER(int32,		PREV_Emitter_Pad0)
+		SHADER_PARAMETER(int32,		PREV_Emitter_Pad1)
+	END_SHADER_PARAMETER_STRUCT()
 
 	BEGIN_SHADER_PARAMETER_STRUCT(FParameters, )
 		SHADER_PARAMETER(uint32,									ComponentBufferSizeRead)
@@ -63,6 +178,11 @@ public:
 
 		SHADER_PARAMETER(FUintVector3,								DispatchThreadIdBounds)
 		SHADER_PARAMETER(FUintVector3,								DispatchThreadIdToLinear)
+
+		SHADER_PARAMETER_STRUCT_INCLUDE(FGlobalParameters,			GlobalParameters)
+		SHADER_PARAMETER_STRUCT_INCLUDE(FSystemParameters,			SystemParameters)
+		SHADER_PARAMETER_STRUCT_INCLUDE(FOwnerParameters,			OwnerParameters)
+		SHADER_PARAMETER_STRUCT_INCLUDE(FEmitterParameters,			EmitterParameters)
 	END_SHADER_PARAMETER_STRUCT()
 
 	using FPermutationParameters = FNiagaraShaderPermutationParameters;
@@ -102,10 +222,6 @@ public:
 	}
 
 	LAYOUT_FIELD(bool, bNeedsViewUniformBuffer);
-	LAYOUT_ARRAY(FShaderUniformBufferParameter, GlobalConstantBufferParam, 2);
-	LAYOUT_ARRAY(FShaderUniformBufferParameter, SystemConstantBufferParam, 2);
-	LAYOUT_ARRAY(FShaderUniformBufferParameter, OwnerConstantBufferParam, 2);
-	LAYOUT_ARRAY(FShaderUniformBufferParameter, EmitterConstantBufferParam, 2);
 	LAYOUT_ARRAY(FShaderUniformBufferParameter, ExternalConstantBufferParam, 2);
 
 private:
diff --git a/Engine/Plugins/FX/Niagara/Source/NiagaraShader/Public/NiagaraShaderParametersBuilder.h b/Engine/Plugins/FX/Niagara/Source/NiagaraShader/Public/NiagaraShaderParametersBuilder.h
--- a/Engine/Plugins/FX/Niagara/Source/NiagaraShader/Public/NiagaraShaderParametersBuilder.h
+++ b/Engine/Plugins/FX/Niagara/Source/NiagaraShader/Public/NiagaraShaderParametersBuilder.h
@@ -8,30 +8,32 @@
 class NIAGARASHADER_API FNiagaraShaderParametersBuilder
 {
 public:
-	explicit FNiagaraShaderParametersBuilder(FNiagaraDataInterfaceGPUParamInfo& InParamInfo, FShaderParametersMetadataBuilder& InMetadataBuilder)
-		: ParamInfo(InParamInfo)
+	explicit FNiagaraShaderParametersBuilder(const FNiagaraDataInterfaceGPUParamInfo& InGPUParamInfo, TArray<FString>& InLooseNames, FShaderParametersMetadataBuilder& InMetadataBuilder)
+		: GPUParamInfo(InGPUParamInfo)
+		, LooseNames(InLooseNames)
 		, MetadataBuilder(InMetadataBuilder)
 	{
 	}
 
 	template<typename T> void AddLooseParam(const TCHAR* Name)
 	{
-		ParamInfo.LooseMetadataNames.Emplace(FString::Printf(TEXT("%s_%s"), *ParamInfo.DataInterfaceHLSLSymbol, Name));
-		MetadataBuilder.AddParam<T>(*ParamInfo.LooseMetadataNames.Last());
+		LooseNames.Emplace(FString::Printf(TEXT("%s_%s"), *GPUParamInfo.DataInterfaceHLSLSymbol, Name));
+		MetadataBuilder.AddParam<T>(*LooseNames.Last());
 	}
 	template<typename T> void AddLooseParamArray(const TCHAR* Name, int32 NumElements)
 	{
-		ParamInfo.LooseMetadataNames.Emplace(FString::Printf(TEXT("%s_%s"), *ParamInfo.DataInterfaceHLSLSymbol, Name));
-		MetadataBuilder.AddParamArray<T>(*ParamInfo.LooseMetadataNames.Last(), NumElements);
+		LooseNames.Emplace(FString::Printf(TEXT("%s_%s"), *GPUParamInfo.DataInterfaceHLSLSymbol, Name));
+		MetadataBuilder.AddParamArray<T>(*LooseNames.Last(), NumElements);
 	}
 	template<typename T> void AddNestedStruct()
 	{
-		MetadataBuilder.AddNestedStruct<T>(*ParamInfo.DataInterfaceHLSLSymbol);
+		MetadataBuilder.AddNestedStruct<T>(*GPUParamInfo.DataInterfaceHLSLSymbol);
 	}
 
-	TArrayView<FNiagaraDataInterfaceGeneratedFunction> GetGeneratedFunctions() const { return ParamInfo.GeneratedFunctions; }
+	TConstArrayView<FNiagaraDataInterfaceGeneratedFunction> GetGeneratedFunctions() const { return GPUParamInfo.GeneratedFunctions; }
 
 private:
-	FNiagaraDataInterfaceGPUParamInfo& ParamInfo;
+	const FNiagaraDataInterfaceGPUParamInfo& GPUParamInfo;
+	TArray<FString>& LooseNames;
 	FShaderParametersMetadataBuilder& MetadataBuilder;
 };
diff --git a/Engine/Plugins/FX/Niagara/Source/NiagaraShader/Public/NiagaraShared.h b/Engine/Plugins/FX/Niagara/Source/NiagaraShader/Public/NiagaraShared.h
--- a/Engine/Plugins/FX/Niagara/Source/NiagaraShader/Public/NiagaraShared.h
+++ b/Engine/Plugins/FX/Niagara/Source/NiagaraShader/Public/NiagaraShared.h
@@ -213,9 +213,6 @@ struct NIAGARASHADER_API FNiagaraDataInterfaceGPUParamInfo
 	UPROPERTY()
 	uint32 ShaderParametersOffset = INDEX_NONE;
 
-	UPROPERTY()
-	TArray<FString> LooseMetadataNames;
-
 	/** Information about all the functions generated by the translator for this data interface. */
 	UPROPERTY()
 	TArray<FNiagaraDataInterfaceGeneratedFunction> GeneratedFunctions;
@@ -233,11 +230,38 @@ template<> struct TStructOpsTypeTraits<FNiagaraDataInterfaceGPUParamInfo> : publ
 	};
 };
 
+USTRUCT()
+struct FNiagaraShaderScriptExternalConstant
+{
+	GENERATED_USTRUCT_BODY()
+		
+	UPROPERTY()
+	FName Type;
+
+	UPROPERTY()
+	FString	Name;
+};
+
+/* Data provided from the translator to track information requires to bind data interface or parameters at runtime. */
+USTRUCT()
 struct NIAGARASHADER_API FNiagaraShaderScriptParametersMetadata
 {
+	GENERATED_USTRUCT_BODY()
+
+	UPROPERTY()
 	TArray<FNiagaraDataInterfaceGPUParamInfo> DataInterfaceParamInfo;
 
-	TUniquePtr<FShaderParametersMetadata> ShaderParametersMetadata;
+	UPROPERTY()
+	TArray<FString> LooseMetadataNames;
+
+	UPROPERTY()
+	bool bExternalConstantsInterpolated = false;
+
+	UPROPERTY()
+	TArray<FNiagaraShaderScriptExternalConstant> ExternalConstants;
+
+	// Runtime generated bindings, not serialized
+	TSharedPtr<FShaderParametersMetadata> ShaderParametersMetadata;
 };
 
 
@@ -643,20 +667,12 @@ class FNiagaraShaderScript
 		{
 			NumPermutations = 0;
 			bIsComplete = 0;
-			bGlobalConstantBufferUsed = 0;
-			bSystemConstantBufferUsed = 0;
-			bOwnerConstantBufferUsed = 0;
-			bEmitterConstantBufferUsed = 0;
 			bExternalConstantBufferUsed = 0;
 			bViewUniformBufferUsed = 0;
 		}
 
 		int32 NumPermutations;
 		uint32 bIsComplete : 1;
-		uint32 bGlobalConstantBufferUsed : 2;
-		uint32 bSystemConstantBufferUsed : 2;
-		uint32 bOwnerConstantBufferUsed : 2;
-		uint32 bEmitterConstantBufferUsed : 2;
 		uint32 bExternalConstantBufferUsed : 2;
 		uint32 bViewUniformBufferUsed : 1;
 	};
@@ -837,8 +853,8 @@ public:
 	NIAGARASHADER_API  FNiagaraShaderRef GetShader(int32 PermutationId) const;
 	NIAGARASHADER_API  FNiagaraShaderRef GetShaderGameThread(int32 PermutationId) const;
 	
-	NIAGARASHADER_API void BuildScriptParametersMetadata(const TArray<FNiagaraDataInterfaceGPUParamInfo>& InDIParamInfo);
-	TSharedRef<FNiagaraShaderScriptParametersMetadata> GetScriptParametersMetadata() const { return ScriptParametersMetadata.ToSharedRef(); }
+	NIAGARASHADER_API void BuildScriptParametersMetadata(const FNiagaraShaderScriptParametersMetadata& ScriptParametersMetadata);
+	TSharedRef<FNiagaraShaderScriptParametersMetadata> GetScriptParametersMetadata() const { check(ScriptParametersMetadata.IsValid()); return ScriptParametersMetadata.ToSharedRef(); }
 
 	NIAGARASHADER_API FOnNiagaraScriptCompilationComplete& OnCompilationComplete()
 	{
@@ -853,10 +869,6 @@ public:
 
 	FORCEINLINE bool IsShaderMapComplete_RenderThread() const { check(IsInRenderingThread()); return CachedData_RenderThread.bIsComplete != 0; }
 	FORCEINLINE int32 GetNumPermutations_RenderThread() const { check(IsInRenderingThread()); return CachedData_RenderThread.NumPermutations; }
-	FORCEINLINE bool IsGlobalConstantBufferUsed_RenderThread(int32 Index) const { return (CachedData_RenderThread.bGlobalConstantBufferUsed & (1 << Index)) != 0; }
-	FORCEINLINE bool IsSystemConstantBufferUsed_RenderThread(int32 Index) const { return (CachedData_RenderThread.bSystemConstantBufferUsed & (1 << Index)) != 0; }
-	FORCEINLINE bool IsOwnerConstantBufferUsed_RenderThread(int32 Index) const { return (CachedData_RenderThread.bOwnerConstantBufferUsed & (1 << Index)) != 0; }
-	FORCEINLINE bool IsEmitterConstantBufferUsed_RenderThread(int32 Index) const { return (CachedData_RenderThread.bEmitterConstantBufferUsed & (1 << Index)) != 0; }
 	FORCEINLINE bool IsExternalConstantBufferUsed_RenderThread(int32 Index) const { return (CachedData_RenderThread.bExternalConstantBufferUsed & (1 << Index)) != 0; }
 	FORCEINLINE bool IsViewUniformBufferUsed_RenderThread() const { return CachedData_RenderThread.bViewUniformBufferUsed != 0; }
 
diff --git a/Engine/Source/Runtime/Core/Private/UObject/UE5MainStreamObjectVersion.cpp b/Engine/Source/Runtime/Core/Private/UObject/UE5MainStreamObjectVersion.cpp
--- a/Engine/Source/Runtime/Core/Private/UObject/UE5MainStreamObjectVersion.cpp
+++ b/Engine/Source/Runtime/Core/Private/UObject/UE5MainStreamObjectVersion.cpp
@@ -14,7 +14,7 @@ TMap<FGuid, FGuid> FUE5MainStreamObjectVersion::GetSystemGuids()
 	SystemGuids.Add(DevGuids.MATERIALSHADERMAP_DERIVEDDATA_VER, FGuid("65C01D817C9A4EEFAE9E988D41A1F3DD"));
 	SystemGuids.Add(DevGuids.NANITE_DERIVEDDATA_VER, FGuid("7A3EA93546624597BEBE2ADFDE3A7628"));
 	SystemGuids.Add(DevGuids.NIAGARASHADERMAP_DERIVEDDATA_VER, FGuid("8A37C45D24F2423CBE5F8F371DE33575"));
-	SystemGuids.Add(DevGuids.Niagara_LatestScriptCompileVersion, FGuid("D40FAE8F632C6B469CEE5F17C49C5BBE"));
+	SystemGuids.Add(DevGuids.Niagara_LatestScriptCompileVersion, FGuid("6BA472B9ADB03A4096668C11AA08DCD8"));
 	SystemGuids.Add(DevGuids.SkeletalMeshDerivedDataVersion, FGuid("FB7B2FB546DE49B2A2F88B7F69763CEE"));
 	SystemGuids.Add(DevGuids.STATICMESH_DERIVEDDATA_VER, FGuid("98F7D79A5811013825E75BBCC41ED3E9"));
 	SystemGuids.Add(DevGuids.POSESEARCHDB_DERIVEDDATA_VER, FGuid("4E595C2AC5E947D6BA9ABC874353E5BC"));
diff --git a/Engine/Source/Runtime/RenderCore/Public/ShaderParameterMetadataBuilder.h b/Engine/Source/Runtime/RenderCore/Public/ShaderParameterMetadataBuilder.h
--- a/Engine/Source/Runtime/RenderCore/Public/ShaderParameterMetadataBuilder.h
+++ b/Engine/Source/Runtime/RenderCore/Public/ShaderParameterMetadataBuilder.h
@@ -146,6 +146,11 @@ public:
 		EShaderPrecisionModifier::Type Precision = EShaderPrecisionModifier::Float
 		);
 
+	void AlignNextMemberToStruct()
+	{
+		NextMemberOffset = Align(NextMemberOffset, SHADER_PARAMETER_STRUCT_ALIGNMENT);
+	}
+
 	uint32 GetNextMemberOffset() const { return NextMemberOffset; }
 
 	FShaderParametersMetadata* Build(
