- Use EmitterName rather than the Found emitter as it makes it easier to track down culled emitters

#rnx

diff --git a/Engine/Plugins/FX/Niagara/Source/Niagara/Private/NiagaraDataInterfaceParticleRead.cpp b/Engine/Plugins/FX/Niagara/Source/Niagara/Private/NiagaraDataInterfaceParticleRead.cpp
--- a/Engine/Plugins/FX/Niagara/Source/Niagara/Private/NiagaraDataInterfaceParticleRead.cpp
+++ b/Engine/Plugins/FX/Niagara/Source/Niagara/Private/NiagaraDataInterfaceParticleRead.cpp
@@ -1417,7 +1417,7 @@ void UNiagaraDataInterfaceParticleRead::GetVMExternalFunction(const FVMExternalF
 	{
 		UNiagaraSystem* NiagaraSystem = PIData->SystemInstance ? PIData->SystemInstance->GetSystem() : nullptr;
 		UNiagaraEmitter* NiagaraEmitter = PIData->EmitterInstance ? PIData->EmitterInstance->GetCachedEmitter().Emitter : nullptr;
-		UE_LOG(LogNiagara, Warning, TEXT("ParticleRead: Failed to '%s' attribute '%s' for System '%s' Emitter '%s'! Check that the attribute is named correctly."), *BindingInfo.Name.ToString(), *AttributeToRead.ToString(), *GetNameSafe(NiagaraSystem), *GetNameSafe(NiagaraEmitter));
+		UE_LOG(LogNiagara, Warning, TEXT("ParticleRead: Failed to '%s' attribute '%s' System '%s' Emitter '%s'! Check that the attribute is named correctly."), *BindingInfo.Name.ToString(), *AttributeToRead.ToString(), *GetNameSafe(NiagaraSystem), *EmitterName);
 	}
 }
 
