Changed the default value of the component renderer's "only create on spawn" property

diff --git a/Engine/Plugins/FX/Niagara/Source/Niagara/Private/NiagaraComponentRendererProperties.cpp b/Engine/Plugins/FX/Niagara/Source/Niagara/Private/NiagaraComponentRendererProperties.cpp
--- a/Engine/Plugins/FX/Niagara/Source/Niagara/Private/NiagaraComponentRendererProperties.cpp
+++ b/Engine/Plugins/FX/Niagara/Source/Niagara/Private/NiagaraComponentRendererProperties.cpp
@@ -7,13 +7,13 @@
 #include "NiagaraComponent.h"
 #include "Modules/ModuleManager.h"
 #if WITH_EDITOR
-#include "Editor.h"
 #include "Widgets/Images/SImage.h"
 #include "Styling/SlateIconFinder.h"
 #include "Widgets/SWidget.h"
 #include "AssetThumbnail.h"
 #include "Widgets/Text/STextBlock.h"
 #endif
+#include "NiagaraCustomVersion.h"
 #include "NiagaraSettings.h"
 
 static float GNiagaraComponentRenderComponentCountWarning = 50;
@@ -171,9 +171,10 @@ FNiagaraTypeDefinition UNiagaraComponentRendererProperties::GetFQuatDef()
 TArray<TWeakObjectPtr<UNiagaraComponentRendererProperties>> UNiagaraComponentRendererProperties::ComponentRendererPropertiesToDeferredInit;
 
 UNiagaraComponentRendererProperties::UNiagaraComponentRendererProperties()
-	: ComponentCountLimit(15), bAssignComponentsOnParticleID(true), bOnlyCreateComponentsOnParticleSpawn(true), bOnlyActivateNewlyAquiredComponents(true)
+	: ComponentCountLimit(15), bAssignComponentsOnParticleID(true), bOnlyActivateNewlyAquiredComponents(true)
 #if WITH_EDITORONLY_DATA
 	, bVisualizeComponents(true)
+	, bOnlyCreateComponentsOnParticleSpawn_DEPRECATED(true)
 #endif
 	, TemplateComponent(nullptr)
 {
@@ -198,6 +199,14 @@ void UNiagaraComponentRendererProperties::PostLoad()
 {
 	Super::PostLoad();
 	ENiagaraRendererSourceDataMode InSourceMode = ENiagaraRendererSourceDataMode::Particles;
+
+#if WITH_EDITORONLY_DATA
+	const int32 NiagaraVer = GetLinkerCustomVersion(FNiagaraCustomVersion::GUID);
+	if (NiagaraVer < FNiagaraCustomVersion::ComponentRendererSpawnProperty)
+	{
+		bCreateComponentFirstParticleFrame = bOnlyCreateComponentsOnParticleSpawn_DEPRECATED;
+	}
+#endif
 	
 	const TArray<FNiagaraVariable>& OldTypes = FNiagaraConstants::GetOldPositionTypeVariables();
 	for (FNiagaraComponentPropertyBinding& Binding : PropertyBindings)
@@ -497,8 +506,7 @@ void UNiagaraComponentRendererProperties::OnObjectsReplacedCallback(const TMap<U
 	// When a custom component class is recompiled in the editor, we need to switch to the new template component object
 	if (TemplateComponent)
 	{
-		UObject* const* Replacement = ReplacementsMap.Find(TemplateComponent);
-		if (Replacement)
+		if (UObject* const* Replacement = ReplacementsMap.Find(TemplateComponent))
 		{
 			TemplateComponent = Cast<USceneComponent>(*Replacement);
 			UpdateSetterFunctions();
@@ -524,7 +532,7 @@ bool UNiagaraComponentRendererProperties::HasPropertyBinding(FName PropertyName)
 
 void UNiagaraComponentRendererProperties::PostEditChangeProperty(struct FPropertyChangedEvent& e)
 {
-	FName PropertyName = (e.Property != NULL) ? e.Property->GetFName() : NAME_None;
+	FName PropertyName = (e.Property != nullptr) ? e.Property->GetFName() : NAME_None;
 	if (PropertyName == GET_MEMBER_NAME_CHECKED(UNiagaraComponentRendererProperties, ComponentType))
 	{
 		PropertyBindings.Empty();
diff --git a/Engine/Plugins/FX/Niagara/Source/Niagara/Private/NiagaraRendererComponents.cpp b/Engine/Plugins/FX/Niagara/Source/Niagara/Private/NiagaraRendererComponents.cpp
--- a/Engine/Plugins/FX/Niagara/Source/Niagara/Private/NiagaraRendererComponents.cpp
+++ b/Engine/Plugins/FX/Niagara/Source/Niagara/Private/NiagaraRendererComponents.cpp
@@ -438,7 +438,7 @@ void FNiagaraRendererComponents::PostSystemTick_GameThread(const UNiagaraRendere
 			ParticleID = UniqueIDAccessor.GetSafe(ParticleIndex, -1);
 			ParticlesWithComponents.RemoveAndCopyValue(ParticleID, PoolIndex);
 
-			if (PoolIndex == -1 && Properties->bOnlyCreateComponentsOnParticleSpawn)
+			if (PoolIndex == -1 && Properties->bCreateComponentFirstParticleFrame)
 			{
 				// Don't allow this particle to acquire a component unless it was just spawned or had a component assigned to it previously
 				bool bIsNewlySpawnedParticle = Emitter->IsParticleComponentActive(ComponentKey, ParticleID) || ParticleIndex >= ParticleData.GetNumInstances() - ParticleData.GetNumSpawnedInstances();
@@ -537,7 +537,7 @@ void FNiagaraRendererComponents::PostSystemTick_GameThread(const UNiagaraRendere
 
 		PoolEntry.LastAssignedToParticleID = ParticleID;
 		PoolEntry.LastActiveTime = CurrentTime;
-		if (Properties->bOnlyCreateComponentsOnParticleSpawn)
+		if (Properties->bCreateComponentFirstParticleFrame)
 		{
 			Emitter->SetParticleComponentActive(ComponentKey, ParticleID);
 		}
diff --git a/Engine/Plugins/FX/Niagara/Source/Niagara/Public/NiagaraComponentRendererProperties.h b/Engine/Plugins/FX/Niagara/Source/Niagara/Public/NiagaraComponentRendererProperties.h
--- a/Engine/Plugins/FX/Niagara/Source/Niagara/Public/NiagaraComponentRendererProperties.h
+++ b/Engine/Plugins/FX/Niagara/Source/Niagara/Public/NiagaraComponentRendererProperties.h
@@ -52,7 +52,7 @@ public:
 	GENERATED_BODY()
 
 	UNiagaraComponentRendererProperties();
-	~UNiagaraComponentRendererProperties();
+	virtual ~UNiagaraComponentRendererProperties() override;
 
 	//UObject Interface
 	virtual void PostLoad() override;
@@ -68,7 +68,7 @@ public:
 	virtual bool IsSimTargetSupported(ENiagaraSimTarget InSimTarget) const override { return (InSimTarget == ENiagaraSimTarget::CPUSim); };
 	virtual void GetUsedMaterials(const FNiagaraEmitterInstance* InEmitter, TArray<UMaterialInterface*>& OutMaterials) const override {};
 #if WITH_EDITORONLY_DATA
-	void PostEditChangeProperty(struct FPropertyChangedEvent& e) override;
+	virtual void PostEditChangeProperty(struct FPropertyChangedEvent& e) override;
 
 	virtual void GetRendererWidgets(const FNiagaraEmitterInstance* InEmitter, TArray<TSharedPtr<SWidget>>& OutWidgets, TSharedPtr<FAssetThumbnailPool> InThumbnailPool) const override;
 	virtual void GetRendererTooltipWidgets(const FNiagaraEmitterInstance* InEmitter, TArray<TSharedPtr<SWidget>>& OutWidgets, TSharedPtr<FAssetThumbnailPool> InThumbnailPool) const override;
@@ -103,8 +103,8 @@ public:
 
 	/** If true then new components can only be created on newly spawned particles. If a particle is not able to create a component on it's first frame (e.g. because the component
 	 * limit was reached) then it will be blocked from spawning a component on subsequent frames. */
-	UPROPERTY(EditAnywhere, AdvancedDisplay, Category = "Component Rendering", meta = (EditCondition = "bAssignComponentsOnParticleID"))
-	bool bOnlyCreateComponentsOnParticleSpawn;
+	UPROPERTY(EditAnywhere, AdvancedDisplay, Category = "Component Rendering", DisplayName="Only Create Components on Particle Spawn", meta = (EditCondition = "bAssignComponentsOnParticleID"))
+	bool bCreateComponentFirstParticleFrame = false;
 	
 	/** 
 	If true then components will only be activated when newly acquired. e.g. on particle spawn or when the particle enables/disables the component.
@@ -122,6 +122,9 @@ public:
 	UPROPERTY(EditAnywhere, AdvancedDisplay, Category = "Component Rendering")
 	bool bVisualizeComponents;
 
+	UPROPERTY()
+	bool bOnlyCreateComponentsOnParticleSpawn_DEPRECATED;
+
 #endif
 
 	/** If a render visibility tag is present, particles whose tag matches this value will be visible in this renderer. */
diff --git a/Engine/Plugins/FX/Niagara/Source/NiagaraCore/Public/NiagaraCustomVersion.h b/Engine/Plugins/FX/Niagara/Source/NiagaraCore/Public/NiagaraCustomVersion.h
--- a/Engine/Plugins/FX/Niagara/Source/NiagaraCore/Public/NiagaraCustomVersion.h
+++ b/Engine/Plugins/FX/Niagara/Source/NiagaraCore/Public/NiagaraCustomVersion.h
@@ -176,6 +176,8 @@ struct FNiagaraCustomVersion
 
 		PopulateFunctionCallNodePinNameBindings, // Function call node refresh from external changes has been refactored so that they don't need to populate their name bindings every load.
 
+		ComponentRendererSpawnProperty, // Changed the default value for the component renderer's OnlyCreateComponentsOnParticleSpawn property
+
 		// DO NOT ADD A NEW VERSION UNLESS YOU HAVE TALKED TO THE NIAGARA LEAD. Mismanagement of these versions can lead to data loss if it is adjusted in multiple streams simultaneously.
 		// -----<new versions can be added above this line>  -------------------------------------------------
 		VersionPlusOne,
