Rephrased update frequency tooltips to clarify timing

diff --git a/Engine/Plugins/FX/Niagara/Source/Niagara/Classes/NiagaraEffectType.h b/Engine/Plugins/FX/Niagara/Source/Niagara/Classes/NiagaraEffectType.h
--- a/Engine/Plugins/FX/Niagara/Source/Niagara/Classes/NiagaraEffectType.h
+++ b/Engine/Plugins/FX/Niagara/Source/Niagara/Classes/NiagaraEffectType.h
@@ -33,11 +33,11 @@ enum class ENiagaraScalabilityUpdateFrequency
 {
 	/** Scalability will be checked only on spawn. */
 	SpawnOnly,
-	/** Scalability will be checked infrequently.*/
+	/** Scalability will be checked every 1.0s (default for fx.NiagaraScalabilityUpdateTime_Low). */
 	Low,
-	/** Scalability will be occasionally. */
+	/** Scalability will be checked every 0.5s (default for fx.NiagaraScalabilityUpdateTime_Medium). */
 	Medium,
-	/** Scalability will be checked regularly. */
+	/** Scalability will be checked every 0.25s (default for fx.NiagaraScalabilityUpdateTime_High). */
 	High,
 	/** Scalability will be checked every frame. */
 	Continuous,
