- Fix static analysis issue

#rnx

diff --git a/Engine/Plugins/FX/Niagara/Source/NiagaraShader/Private/NiagaraShared.cpp b/Engine/Plugins/FX/Niagara/Source/NiagaraShader/Private/NiagaraShared.cpp
--- a/Engine/Plugins/FX/Niagara/Source/NiagaraShader/Private/NiagaraShared.cpp
+++ b/Engine/Plugins/FX/Niagara/Source/NiagaraShader/Private/NiagaraShared.cpp
@@ -686,6 +686,7 @@ void FNiagaraShaderScript::BuildScriptParametersMetadata(const FNiagaraShaderScr
 		if (CDODataInterface == nullptr)
 		{
 			Invalidate();
+			continue;
 		}
 
 		if (CDODataInterface->UseLegacyShaderBindings() == false)
