- Clear arrays during update

#rnx

diff --git a/Engine/Plugins/FX/Niagara/Source/Niagara/Private/NiagaraCommon.cpp b/Engine/Plugins/FX/Niagara/Source/Niagara/Private/NiagaraCommon.cpp
--- a/Engine/Plugins/FX/Niagara/Source/Niagara/Private/NiagaraCommon.cpp
+++ b/Engine/Plugins/FX/Niagara/Source/Niagara/Private/NiagaraCommon.cpp
@@ -81,6 +81,7 @@ void FNiagaraSystemUpdateContext::CommitUpdate()
 			FNiagaraWorldManager::DestroyAllSystemSimulations(NiagaraSystem);
 		}
 	}
+	SystemSimsToDestroy.Empty();
 
 	bool bNeedsWaitOnGpu = true;
 	for (UNiagaraSystem* NiagaraSystem : SystemSimsToRecache)
@@ -98,6 +99,7 @@ void FNiagaraSystemUpdateContext::CommitUpdate()
 			NiagaraSystem->CacheFromCompiledData();
 		}
 	}
+	SystemSimsToRecache.Empty();
 	
 	for (UNiagaraComponent* Comp : ComponentsToReInit)
 	{
