- Expose option to render the component only or not

#rnx

diff --git a/Engine/Plugins/FX/Niagara/Source/Niagara/Classes/NiagaraBakerSettings.h b/Engine/Plugins/FX/Niagara/Source/Niagara/Classes/NiagaraBakerSettings.h
--- a/Engine/Plugins/FX/Niagara/Source/Niagara/Classes/NiagaraBakerSettings.h
+++ b/Engine/Plugins/FX/Niagara/Source/Niagara/Classes/NiagaraBakerSettings.h
@@ -167,7 +167,7 @@ public:
 	int32 CurrentCameraIndex = 0;
 
 	/** Should we render just the component or the whole scene. */
-	UPROPERTY(EditAnywhere, Category = "Settings", AdvancedDisplay)
+	UPROPERTY(EditAnywhere, Category = "Settings")
 	uint8 bRenderComponentOnly : 1;
 
 	bool Equals(const UNiagaraBakerSettings& Other) const;
diff --git a/Engine/Plugins/FX/Niagara/Source/NiagaraEditor/Private/NiagaraBakerRenderer.cpp b/Engine/Plugins/FX/Niagara/Source/NiagaraEditor/Private/NiagaraBakerRenderer.cpp
--- a/Engine/Plugins/FX/Niagara/Source/NiagaraEditor/Private/NiagaraBakerRenderer.cpp
+++ b/Engine/Plugins/FX/Niagara/Source/NiagaraEditor/Private/NiagaraBakerRenderer.cpp
@@ -356,9 +356,16 @@ void FNiagaraBakerRenderer::RenderSceneCapture(UTextureRenderTarget2D* RenderTar
 	SceneCaptureComponent->bUseCustomProjectionMatrix = true;
 	SceneCaptureComponent->CustomProjectionMatrix = BakerSettings->GetProjectionMatrix();
 
-	SceneCaptureComponent->PrimitiveRenderMode = ESceneCapturePrimitiveRenderMode::PRM_UseShowOnlyList;
-	SceneCaptureComponent->ShowOnlyComponents.Empty(1);
-	SceneCaptureComponent->ShowOnlyComponents.Add(NiagaraComponent);
+	if (BakerSettings->bRenderComponentOnly)
+	{
+		SceneCaptureComponent->PrimitiveRenderMode = ESceneCapturePrimitiveRenderMode::PRM_UseShowOnlyList;
+		SceneCaptureComponent->ShowOnlyComponents.Empty(1);
+		SceneCaptureComponent->ShowOnlyComponents.Add(NiagaraComponent);
+	}
+	else
+	{
+		SceneCaptureComponent->PrimitiveRenderMode = ESceneCapturePrimitiveRenderMode::PRM_RenderScenePrimitives;
+	}
 
 	SceneCaptureComponent->CaptureScene();
 
diff --git a/Engine/Plugins/FX/Niagara/Source/NiagaraEditor/Private/ViewModels/NiagaraBakerViewModel.cpp b/Engine/Plugins/FX/Niagara/Source/NiagaraEditor/Private/ViewModels/NiagaraBakerViewModel.cpp
--- a/Engine/Plugins/FX/Niagara/Source/NiagaraEditor/Private/ViewModels/NiagaraBakerViewModel.cpp
+++ b/Engine/Plugins/FX/Niagara/Source/NiagaraEditor/Private/ViewModels/NiagaraBakerViewModel.cpp
@@ -84,6 +84,22 @@ bool FNiagaraBakerViewModel::IsPlaybackLooping() const
 	return BakerSettings ? BakerSettings->bPreviewLooping : false;
 }
 
+bool FNiagaraBakerViewModel::ShowRenderComponentOnly() const
+{
+	UNiagaraBakerSettings* BakerSettings = GetBakerSettings();
+	return BakerSettings ? BakerSettings->bRenderComponentOnly : false;
+}
+
+void FNiagaraBakerViewModel::ToggleRenderComponentOnly()
+{
+	if ( UNiagaraBakerSettings* BakerSettings = GetBakerSettings() )
+	{
+		const FScopedTransaction Transaction(LOCTEXT("SetRenderComponentOnly", "Set Render Component Only"));
+		BakerSettings->Modify();
+		BakerSettings->bRenderComponentOnly = !BakerSettings->bRenderComponentOnly;
+	}
+}
+
 void FNiagaraBakerViewModel::SetCameraSettingsIndex(int CamerSettingsIndex)
 {
 	if (UNiagaraBakerSettings* BakerSettings = GetBakerSettings())
diff --git a/Engine/Plugins/FX/Niagara/Source/NiagaraEditor/Private/ViewModels/NiagaraBakerViewModel.h b/Engine/Plugins/FX/Niagara/Source/NiagaraEditor/Private/ViewModels/NiagaraBakerViewModel.h
--- a/Engine/Plugins/FX/Niagara/Source/NiagaraEditor/Private/ViewModels/NiagaraBakerViewModel.h
+++ b/Engine/Plugins/FX/Niagara/Source/NiagaraEditor/Private/ViewModels/NiagaraBakerViewModel.h
@@ -64,6 +64,9 @@ public:
 	bool ShowInfoText() const { return bShowInfoText; }
 	void ToggleInfoText() { bShowInfoText = !bShowInfoText; }
 
+	bool ShowRenderComponentOnly() const;
+	void ToggleRenderComponentOnly();
+
 	void SetCameraSettingsIndex(int CamerSettingsIndex);
 	bool IsCameraSettingIndex(int CamerSettingsIndex) const;
 
diff --git a/Engine/Plugins/FX/Niagara/Source/NiagaraEditor/Private/Widgets/SNiagaraBakerWidget.cpp b/Engine/Plugins/FX/Niagara/Source/NiagaraEditor/Private/Widgets/SNiagaraBakerWidget.cpp
--- a/Engine/Plugins/FX/Niagara/Source/NiagaraEditor/Private/Widgets/SNiagaraBakerWidget.cpp
+++ b/Engine/Plugins/FX/Niagara/Source/NiagaraEditor/Private/Widgets/SNiagaraBakerWidget.cpp
@@ -621,6 +621,19 @@ TSharedRef<SWidget> SNiagaraBakerWidget::MakeViewOptionsMenu()
 			EUserInterfaceActionType::ToggleButton
 		);
 
+		MenuBuilder.AddMenuEntry(
+			LOCTEXT("ShowNiagaraOnly", "Render Niagara Only"),
+			LOCTEXT("ShowNiagaraOnlyTooltip", "Renders only the Niagara System when enabled."),
+			FSlateIcon(),
+			FUIAction(
+				FExecuteAction::CreateSP(ViewModel, &FNiagaraBakerViewModel::ToggleRenderComponentOnly),
+				FCanExecuteAction(),
+				FIsActionChecked::CreateSP(ViewModel, &FNiagaraBakerViewModel::ShowRenderComponentOnly)
+			),
+			NAME_None,
+			EUserInterfaceActionType::ToggleButton
+		);
+
 		//DetailCategory.HideProperty(GET_MEMBER_NAME_CHECKED(UNiagaraBakerSettings, bPreviewLooping));
 		//DetailCategory.HideProperty(GET_MEMBER_NAME_CHECKED(UNiagaraBakerSettings, bRenderComponentOnly));
 	}
