- Add Systems with Simulation Stages enabled on them to audit

#rnx

diff --git a/Engine/Plugins/FX/Niagara/Source/NiagaraEditor/Private/Commandlets/NiagaraSystemAuditCommandlet.cpp b/Engine/Plugins/FX/Niagara/Source/NiagaraEditor/Private/Commandlets/NiagaraSystemAuditCommandlet.cpp
--- a/Engine/Plugins/FX/Niagara/Source/NiagaraEditor/Private/Commandlets/NiagaraSystemAuditCommandlet.cpp
+++ b/Engine/Plugins/FX/Niagara/Source/NiagaraEditor/Private/Commandlets/NiagaraSystemAuditCommandlet.cpp
@@ -17,6 +17,7 @@
 
 #include "NiagaraSettings.h"
 #include "NiagaraSystem.h"
+#include "NiagaraSimulationStageBase.h"
 #include "NiagaraRendererProperties.h"
 #include "NiagaraLightRendererProperties.h"
 #include "NiagaraRibbonRendererProperties.h"
@@ -201,6 +202,7 @@ bool UNiagaraSystemAuditCommandlet::ProcessNiagaraSystems()
 
 		// Iterate over all emitters
 		FString EmittersWithDynamicBounds;
+		FString EmittersWithSimulationStages;
 		bool bHasLights = false;
 		bool bHasEvents = false;
 
@@ -273,6 +275,22 @@ bool UNiagaraSystemAuditCommandlet::ProcessNiagaraSystems()
 				GpuEmitterBuilder.Append(TEXT(","));
 				GpuEmitterBuilder.Append(NiagaraSystem->GetPathName());
 				NiagaraSystemsWithGPUEmitters.Add(GpuEmitterBuilder.ToString());
+
+				// Do we have simulation stages enabled
+				const TArray<UNiagaraSimulationStageBase*>& SimulationStages = EmitterData->GetSimulationStages();
+				if (SimulationStages.Num() > 0 )
+				{
+					int32 TotalIterations = 0;
+					for (UNiagaraSimulationStageBase* SimStage : SimulationStages)
+					{
+						UNiagaraSimulationStageGeneric* SimStageGeneric = Cast<UNiagaraSimulationStageGeneric>(SimStage);
+						if (SimStage->bEnabled && SimStageGeneric)
+						{
+							TotalIterations += SimStageGeneric->Iterations;
+						}
+					}
+					EmittersWithSimulationStages.Appendf(TEXT("%s(%d Stages %d Iterations) "), EmitterData->GetDebugSimName(), SimulationStages.Num(), TotalIterations);
+				}
 			}
 
 			bHasEvents |= EmitterData->GetEventHandlers().Num() > 0;
@@ -335,6 +353,11 @@ bool UNiagaraSystemAuditCommandlet::ProcessNiagaraSystems()
 			NiagaraSystemsWithPrerequisites.Add(FString::Printf(TEXT("%s,%s"), *NiagaraSystem->GetPathName(), *DataInterfaceNames));
 		}
 
+		if ( !EmittersWithSimulationStages.IsEmpty() )
+		{
+			NiagaraSystemsWithSimulationStages.Add(FString::Printf(TEXT("%s,%s"), *NiagaraSystem->GetPathName(), *EmittersWithSimulationStages));
+		}
+
 		if ( !EmittersWithDynamicBounds.IsEmpty() )
 		{
 			NiagaraSystemsWithDynamicBounds.Add(EmittersWithDynamicBounds);
@@ -396,6 +419,7 @@ void UNiagaraSystemAuditCommandlet::DumpResults()
 		HeaderString.Append(TEXT(",CVar Conditions,System Path"));
 		DumpSimpleSet(NiagaraSystemsWithGPUEmitters, TEXT("NiagaraSystemsWithGPUEmitters"), HeaderString.ToString());
 	}
+	DumpSimpleSet(NiagaraSystemsWithSimulationStages, TEXT("NiagaraSystemsWithSimulationStages"), TEXT("System,Emitters"));
 
 	if (NiagaraDataInterfaceUsage.Num() > 0)
 	{
diff --git a/Engine/Plugins/FX/Niagara/Source/NiagaraEditor/Public/Commandlets/NiagaraSystemAuditCommandlet.h b/Engine/Plugins/FX/Niagara/Source/NiagaraEditor/Public/Commandlets/NiagaraSystemAuditCommandlet.h
--- a/Engine/Plugins/FX/Niagara/Source/NiagaraEditor/Public/Commandlets/NiagaraSystemAuditCommandlet.h
+++ b/Engine/Plugins/FX/Niagara/Source/NiagaraEditor/Public/Commandlets/NiagaraSystemAuditCommandlet.h
@@ -24,9 +24,10 @@ class UNiagaraSystemAuditCommandlet : public UCommandlet
 	TSet<FString> NiagaraSystemsWithPrerequisites;
 	/** All Niagara systems that contain emitters with dynamic bounds */
 	TSet<FString> NiagaraSystemsWithDynamicBounds;
-
 	/** All Niagara systems that use a user requested data interface */
 	TSet<FString> NiagaraSystemsWithUserDataInterface;
+	/** All Niagara Systems & Emitters with simulation stages enabled*/
+	TSet<FString> NiagaraSystemsWithSimulationStages;
 
 	/** Optional list of data interface usage */
 	struct FDataInterfaceUsage
