Volume Caching

- UVolumeCache
    UAsset for encapsulating a sequence of cached volume frames
    Can load a sequence of volumes into memory
    Generic interface for cached data
    Support for OpenVDB
    functions for filling 3d textures
- NiagaraDataInterefaceVolumeCache
    DI for interfacing with UVolumeCache
    Contains a 3d texture, which is filled from the UVolumeCache
    DI functions for setting the requested frame from the cache
    DI functions to query the volume texture for the cache at the current frame
- Changes to Baker
    Writing out a sequence of vdb files
    Creating (or updating) a UVolumeCache asset
- Changes to core Niagara
    Build support for OpenVDB
- Example content
    Test version of fluid emitter with read/write caching modes
    Test map and other assets

diff --git a/Engine/Plugins/FX/Niagara/Shaders/Private/NiagaraDataInterfaceVolumeCache.ush b/Engine/Plugins/FX/Niagara/Shaders/Private/NiagaraDataInterfaceVolumeCache.ush
new file mode 100644
--- /dev/null
+++ b/Engine/Plugins/FX/Niagara/Shaders/Private/NiagaraDataInterfaceVolumeCache.ush
@@ -0,0 +1,20 @@
+// Copyright Epic Games, Inc. All Rights Reserved.
+
+float3			{ParameterName}_TextureSize;
+Texture3D		{ParameterName}_Texture;
+SamplerState	{ParameterName}_TextureSampler;
+
+void SampleCurrentFrameValue_{ParameterName}(in float3 UV, in float MipLevel, out float4 OutValue)
+{
+	OutValue = {ParameterName}_Texture.SampleLevel({ParameterName}_TextureSampler, UV, MipLevel);
+}
+
+void GetCurrentFrameValue_{ParameterName}(in int x, in int y, in int z, in float MipLevel, out float4 OutValue)
+{
+	OutValue = {ParameterName}_Texture.Load(int4(x,y,z,MipLevel));
+}
+
+void GetCurrentFrameNumCells_{ParameterName}(out float3 OutValue)
+{
+	OutValue = {ParameterName}_TextureSize;
+}
diff --git a/Engine/Plugins/FX/Niagara/Source/Niagara/Classes/NiagaraDataInterfaceVolumeCache.h b/Engine/Plugins/FX/Niagara/Source/Niagara/Classes/NiagaraDataInterfaceVolumeCache.h
new file mode 100644
--- /dev/null
+++ b/Engine/Plugins/FX/Niagara/Source/Niagara/Classes/NiagaraDataInterfaceVolumeCache.h
@@ -0,0 +1,73 @@
+// Copyright Epic Games, Inc. All Rights Reserved.
+#pragma once
+
+#include "NiagaraDataInterface.h"
+#include "NiagaraCommon.h"
+#include "NiagaraShared.h"
+#include "VectorVM.h"
+#include "Engine/VolumeTexture.h"
+#include "VolumeCache.h"
+
+#include "NiagaraDataInterfaceVolumeCache.generated.h"
+
+struct FVolumeCacheInstanceData_GameThread;
+
+UCLASS(EditInlineNew, Category = "Grid",  meta = (DisplayName = "Volume Cache"))
+class NIAGARA_API UNiagaraDataInterfaceVolumeCache : public UNiagaraDataInterface
+{
+	GENERATED_UCLASS_BODY()
+
+	BEGIN_SHADER_PARAMETER_STRUCT(FShaderParameters, )
+		SHADER_PARAMETER(FVector3f, TextureSize)
+		SHADER_PARAMETER_TEXTURE(Texture3D, Texture)
+		SHADER_PARAMETER_SAMPLER(SamplerState, TextureSampler)
+	END_SHADER_PARAMETER_STRUCT()
+
+public:
+		
+	UPROPERTY(EditAnywhere, Category = "File")
+	UVolumeCache *VolumeCache;
+
+	virtual void PostInitProperties() override;
+
+	// UNiagaraDataInterface Interface Begin
+	virtual bool CanExecuteOnTarget(ENiagaraSimTarget Target) const override { return true; }
+
+	virtual bool Equals(const UNiagaraDataInterface* Other) const override;
+	virtual int32 PerInstanceDataSize()const override;
+	virtual bool InitPerInstanceData(void* PerInstanceData, FNiagaraSystemInstance* SystemInstance) override;	
+	virtual void GetFunctions(TArray<FNiagaraFunctionSignature>& OutFunctions) override;
+	virtual bool PerInstanceTick(void* PerInstanceData, FNiagaraSystemInstance* SystemInstance, float DeltaSeconds) override;
+	virtual bool HasPreSimulateTick() const override { return true; }
+	virtual void GetVMExternalFunction(const FVMExternalFunctionBindingInfo& BindingInfo, void* InstanceData, FVMExternalFunction& OutFunc) override;
+	virtual void DestroyPerInstanceData(void* PerInstanceData, FNiagaraSystemInstance* SystemInstance) override;
+
+#if WITH_EDITORONLY_DATA
+	virtual bool AppendCompileHash(FNiagaraCompileHashVisitor* InVisitor) const override;
+	virtual void GetParameterDefinitionHLSL(const FNiagaraDataInterfaceGPUParamInfo& ParamInfo, FString& OutHLSL) override;
+	virtual bool GetFunctionHLSL(const FNiagaraDataInterfaceGPUParamInfo& ParamInfo, const FNiagaraDataInterfaceGeneratedFunction& FunctionInfo, int FunctionInstanceIndex, FString& OutHLSL) override;
+#endif
+
+	virtual bool UseLegacyShaderBindings() const override { return false; }
+	virtual void BuildShaderParameters(FNiagaraShaderParametersBuilder& ShaderParametersBuilder) const override;
+	virtual void SetShaderParameters(const FNiagaraDataInterfaceSetShaderParametersContext& Context) const override;
+
+	void SetFrame(FVectorVMExternalFunctionContext& Context);
+	void ReadFile(FVectorVMExternalFunctionContext& Context);
+	
+	FString GetAssetPath(FString PathFormat, int32 FrameIndex) const;
+
+protected:
+	//~ UNiagaraDataInterface interface
+	virtual bool CopyToInternal(UNiagaraDataInterface* Destination) const override;
+	//~ UNiagaraDataInterface interface END
+
+	TMap<FNiagaraSystemInstanceID, FVolumeCacheInstanceData_GameThread*> SystemInstancesToProxyData_GT;	
+
+	static const FName SetFrameName;
+	static const FName ReadFileName;
+	static const FName SampleCurrentFrameValueName;
+	static const FName GetCurrentFrameValue;
+	static const FName GetCurrentFrameNumCells;
+	static const TCHAR* TemplateShaderFilePath;
+};
\ No newline at end of file
diff --git a/Engine/Plugins/FX/Niagara/Source/Niagara/Classes/NiagaraOpenVDB.h b/Engine/Plugins/FX/Niagara/Source/Niagara/Classes/NiagaraOpenVDB.h
new file mode 100644
--- /dev/null
+++ b/Engine/Plugins/FX/Niagara/Source/Niagara/Classes/NiagaraOpenVDB.h
@@ -0,0 +1,55 @@
+// Copyright Epic Games, Inc. All Rights Reserved.
+
+#pragma once
+
+#if PLATFORM_WINDOWS
+
+THIRD_PARTY_INCLUDES_START
+
+UE_PUSH_MACRO("check")
+#undef check
+
+#ifdef _MSC_VER
+#pragma warning (push)
+#pragma warning (disable : 4146) // unary minus operator applied to unsigned type, result still unsigned
+#endif
+
+#ifndef M_PI
+#define M_PI    3.14159265358979323846
+#define LOCAL_M_PI 1
+#endif
+
+#ifndef M_PI_2
+#define M_PI_2  1.57079632679489661923 // pi/2
+#define LOCAL_M_PI_2 1
+#endif
+
+#define BOOST_ALLOW_DEPRECATED_HEADERS
+
+#include <openvdb/openvdb.h>
+#include <openvdb/io/Stream.h>
+#include <openvdb/io/File.h>
+#include <openvdb/tools/Prune.h>
+#include <openvdb/tools/Dense.h>
+#include <openvdb/math/Vec4.h>
+
+using Vec4STree = openvdb::tree::Tree4<openvdb::Vec4s, 5, 4, 3>::Type;
+using Vec4SGrid = openvdb::Grid<Vec4STree>;
+
+#if LOCAL_M_PI
+#undef M_PI
+#endif
+
+#if LOCAL_M_PI_2
+#undef M_PI_2
+#endif
+
+#ifdef _MSC_VER
+#pragma warning (pop)
+#endif
+
+UE_POP_MACRO("check")
+
+THIRD_PARTY_INCLUDES_END
+
+#endif
\ No newline at end of file
diff --git a/Engine/Plugins/FX/Niagara/Source/Niagara/Classes/VolumeCache.h b/Engine/Plugins/FX/Niagara/Source/Niagara/Classes/VolumeCache.h
new file mode 100644
--- /dev/null
+++ b/Engine/Plugins/FX/Niagara/Source/Niagara/Classes/VolumeCache.h
@@ -0,0 +1,109 @@
+// Copyright Epic Games, Inc. All Rights Reserved.
+
+
+#pragma once
+
+#include "CoreMinimal.h"
+#include "UObject/ObjectMacros.h"
+
+// @todo: we need builds for OpenVDB for platforms other than windows
+#if PLATFORM_WINDOWS
+#include "NiagaraOpenVDB.h"
+#endif
+
+#include "VolumeCache.generated.h"
+
+class FVolumeCacheData;
+
+UENUM()
+enum class EVolumeCacheType : uint8
+{
+	OpenVDB
+};
+
+UCLASS(Experimental)
+class NIAGARA_API UVolumeCache : public UObject
+{
+	GENERATED_UCLASS_BODY()
+
+public:
+	/** File path to load */
+	UPROPERTY(EditAnywhere, Category=File, meta=(DisplayName="File Path"))
+	FString FilePath;
+	
+	UPROPERTY(EditAnywhere, Category = File, meta = (DisplayName = "Cache Type"))
+	EVolumeCacheType CacheType;
+
+	UPROPERTY(EditAnywhere, Category = File, meta = (DisplayName = "Resolution"))
+	FIntVector Resolution;
+
+	UPROPERTY(EditAnywhere, Category = File, meta = (DisplayName = "Frame Range Start"))
+	int32 FrameRangeStart;	
+
+	UPROPERTY(EditAnywhere, Category = File, meta = (DisplayName = "Frame Range End"))
+	int32 FrameRangeEnd;
+		
+	void InitData();
+
+	bool LoadFile(int frame);
+	bool UnloadFile(int frame);		
+	bool LoadRange();
+	void UnloadAll();	
+	
+	TSharedPtr<FVolumeCacheData> GetData() { return CachedVolumeFiles;  }
+
+	// @todo: high level fill volume texture method that works on GT calls
+	// bool Fill3DTexture(int frame, FTextureRHIRef TextureToFill, FRHICommandListImmediate& RHICmdList);
+	
+private:
+	TSharedPtr<FVolumeCacheData> CachedVolumeFiles;
+};
+
+class NIAGARA_API FVolumeCacheData
+{
+public:		
+	FVolumeCacheData() : DenseResolution(-1, -1, -1) {}
+	virtual ~FVolumeCacheData() {}
+
+	FString GetAssetPath(FString PathFormat, int32 FrameIndex) const;
+
+	FIntVector GetDenseResolution() { return DenseResolution;  }
+
+	virtual void Init(FIntVector Resolution) = 0;
+	virtual bool LoadFile(FString Path, int frame) = 0;
+	virtual bool UnloadFile(int frame) = 0;
+	virtual bool LoadRange(FString Path, int Start, int End) = 0;
+	virtual void UnloadAll() = 0;
+	virtual bool Fill3DTexture_RenderThread(int frame, FTextureRHIRef TextureToFill, FRHICommandListImmediate& RHICmdList) = 0;
+
+protected:
+	FIntVector DenseResolution;
+};
+
+// @todo: we need builds for OpenVDB for platforms other than windows
+#if PLATFORM_WINDOWS
+class NIAGARA_API FOpenVDBCacheData : public FVolumeCacheData
+{
+public:
+	FOpenVDBCacheData() {}
+	
+	virtual ~FOpenVDBCacheData() 
+	{
+		UnloadAll();
+		DenseGridPtr = nullptr;
+	}
+
+	virtual void Init(FIntVector Resolution);
+	virtual bool LoadFile(FString Path, int frame);
+	virtual bool UnloadFile(int frame);
+	virtual bool LoadRange(FString Path, int Start, int End);
+	virtual void UnloadAll();
+	virtual bool Fill3DTexture_RenderThread(int frame, FTextureRHIRef TextureToFill, FRHICommandListImmediate& RHICmdList);
+
+	static bool WriteImageDataToOpenVDBFile(FStringView FilePath, FIntVector ImageSize, TArrayView<FFloat16Color> ImageData, bool UseFloatGrids = false);
+
+private:	
+	TMap<int32, Vec4SGrid::Ptr> OpenVDBGrids;
+	openvdb::tools::Dense<openvdb::Vec4s, openvdb::tools::MemoryLayout::LayoutXYZ>::Ptr DenseGridPtr;	
+};
+#endif
\ No newline at end of file
diff --git a/Engine/Plugins/FX/Niagara/Source/Niagara/Niagara.Build.cs b/Engine/Plugins/FX/Niagara/Source/Niagara/Niagara.Build.cs
--- a/Engine/Plugins/FX/Niagara/Source/Niagara/Niagara.Build.cs
+++ b/Engine/Plugins/FX/Niagara/Source/Niagara/Niagara.Build.cs
@@ -8,7 +8,14 @@ public class Niagara : ModuleRules
     {
         PrivateIncludePaths.Add("../../../../Shaders/Shared");
 
-        PrivateDependencyModuleNames.AddRange(
+		// Specific to OpenVDB support
+		if (Target.IsInPlatformGroup(UnrealPlatformGroup.Windows))
+		{
+			bUseRTTI = true;
+			bEnableExceptions = true;
+		}
+		
+		PrivateDependencyModuleNames.AddRange(
             new string[] {
                 "ApplicationCore",
                 "AudioPlatformConfiguration",
@@ -67,7 +74,18 @@ public class Niagara : ModuleRules
             });
         }
 
-        PublicDefinitions.AddRange(
+		if (Target.IsInPlatformGroup(UnrealPlatformGroup.Windows))
+		{ 
+			AddEngineThirdPartyPrivateStaticDependencies(Target,
+				"IntelTBB",
+				"Blosc",
+				"zlib",
+				"Boost",				
+				"OpenVDB"
+			);
+		}
+
+		PublicDefinitions.AddRange(
             new string[]
             {
                 "VECTORVM_SUPPORTS_EXPERIMENTAL=0",
diff --git a/Engine/Plugins/FX/Niagara/Source/Niagara/Private/NiagaraDataInterfaceVolumeCache.cpp b/Engine/Plugins/FX/Niagara/Source/Niagara/Private/NiagaraDataInterfaceVolumeCache.cpp
new file mode 100644
--- /dev/null
+++ b/Engine/Plugins/FX/Niagara/Source/Niagara/Private/NiagaraDataInterfaceVolumeCache.cpp
@@ -0,0 +1,389 @@
+// Copyright Epic Games, Inc. All Rights Reserved.
+#include "NiagaraDataInterfaceVolumeCache.h"
+#include "NiagaraShader.h"
+#include "ShaderParameterUtils.h"
+#include "ClearQuad.h"
+#include "TextureResource.h"
+#include "NiagaraSystemInstance.h"
+#include "NiagaraRenderer.h"
+#include "Engine/VolumeTexture.h"
+#include "Engine/TextureRenderTargetVolume.h"
+#include "NiagaraSettings.h"
+#include "NiagaraConstants.h"
+#include "NiagaraComputeExecutionContext.h"
+#include "NiagaraShaderParametersBuilder.h"
+#if WITH_EDITOR
+#include "NiagaraGpuComputeDebug.h"
+#endif
+
+#define LOCTEXT_NAMESPACE "NiagaraDataInterfaceVolumeCache"
+
+const TCHAR* UNiagaraDataInterfaceVolumeCache::TemplateShaderFilePath = TEXT("/Plugin/FX/Niagara/Private/NiagaraDataInterfaceVolumeCache.ush");
+const FName UNiagaraDataInterfaceVolumeCache::SetFrameName("SetFrame");
+const FName UNiagaraDataInterfaceVolumeCache::ReadFileName("ReadFile");
+const FName UNiagaraDataInterfaceVolumeCache::SampleCurrentFrameValueName(TEXT("SampleCurrentFrameValue"));
+const FName UNiagaraDataInterfaceVolumeCache::GetCurrentFrameValue(TEXT("GetCurrentFrameValue"));
+const FName UNiagaraDataInterfaceVolumeCache::GetCurrentFrameNumCells(TEXT("GetCurrentFrameNumCells"));
+
+struct FVolumeCacheInstanceData_RenderThread
+{
+	int						CurrFrame = 0;
+
+	FSamplerStateRHIRef		SamplerStateRHI = nullptr;
+	FTextureReferenceRHIRef	TextureReferenceRHI = nullptr;
+	FTextureRHIRef			ResolvedTextureRHI = nullptr;
+	FVector3f				TextureSize = FVector3f::ZeroVector;
+};
+
+struct FVolumeCacheInstanceData_GameThread
+{
+	int CurrFrame = 0;
+	int PrevFrame = -1;
+
+	bool ReadFile = false;		
+};
+
+struct FNiagaraDataInterfaceVolumeCacheProxy : public FNiagaraDataInterfaceProxy
+{
+	virtual void ConsumePerInstanceDataFromGameThread(void* PerInstanceData, const FNiagaraSystemInstanceID& Instance) override { check(false); }
+	virtual int32 PerInstanceDataPassedToRenderThreadSize() const override { return 0; }
+
+	TMap<FNiagaraSystemInstanceID, FVolumeCacheInstanceData_RenderThread> InstanceData_RT;
+};
+
+
+
+UNiagaraDataInterfaceVolumeCache::UNiagaraDataInterfaceVolumeCache(FObjectInitializer const& ObjectInitializer)
+	: Super(ObjectInitializer)
+{
+	Proxy.Reset(new FNiagaraDataInterfaceVolumeCacheProxy());
+
+}
+
+void UNiagaraDataInterfaceVolumeCache::PostInitProperties()
+{
+	Super::PostInitProperties();
+
+	//Can we register data interfaces as regular types and fold them into the FNiagaraVariable framework for UI and function calls etc?
+	if (HasAnyFlags(RF_ClassDefaultObject))
+	{
+		ENiagaraTypeRegistryFlags Flags = ENiagaraTypeRegistryFlags::AllowAnyVariable | ENiagaraTypeRegistryFlags::AllowParameter;
+		FNiagaraTypeRegistry::Register(FNiagaraTypeDefinition(GetClass()), Flags);		
+	}
+}
+
+
+bool UNiagaraDataInterfaceVolumeCache::Equals(const UNiagaraDataInterface* Other) const
+{
+	if (!Super::Equals(Other))
+	{
+		return false;
+	}
+	const UNiagaraDataInterfaceVolumeCache* OtherTyped = CastChecked<const UNiagaraDataInterfaceVolumeCache>(Other);
+
+	return OtherTyped != nullptr && VolumeCache == OtherTyped->VolumeCache;
+}
+
+inline int32 UNiagaraDataInterfaceVolumeCache::PerInstanceDataSize() const 
+{ 
+	return sizeof(FVolumeCacheInstanceData_GameThread); 
+}
+
+bool UNiagaraDataInterfaceVolumeCache::CopyToInternal(UNiagaraDataInterface* Destination) const
+{
+	if (!Super::CopyToInternal(Destination))
+	{
+		return false;
+	}
+
+	UNiagaraDataInterfaceVolumeCache* OtherTyped = CastChecked<UNiagaraDataInterfaceVolumeCache>(Destination);
+
+	OtherTyped->VolumeCache = VolumeCache;
+	return true;
+}
+
+bool UNiagaraDataInterfaceVolumeCache::InitPerInstanceData(void* PerInstanceData, FNiagaraSystemInstance* SystemInstance)
+{
+	check(Proxy);
+	
+	FVolumeCacheInstanceData_GameThread* InstanceData = new (PerInstanceData) FVolumeCacheInstanceData_GameThread();
+	SystemInstancesToProxyData_GT.Emplace(SystemInstance->GetId(), InstanceData);	
+
+	// Push Updates to Proxy.
+	FNiagaraDataInterfaceVolumeCacheProxy* TheProxy = GetProxyAs<FNiagaraDataInterfaceVolumeCacheProxy>();
+	ENQUEUE_RENDER_COMMAND(FUpdateData)(
+		[TheProxy, InstanceID = SystemInstance->GetId()](FRHICommandListImmediate& RHICmdList)
+	{
+		check(!TheProxy->InstanceData_RT.Contains(InstanceID));
+		FVolumeCacheInstanceData_RenderThread* TargetData = &TheProxy->InstanceData_RT.Add(InstanceID);
+
+	});
+
+	return true;
+}
+
+void UNiagaraDataInterfaceVolumeCache::DestroyPerInstanceData(void* PerInstanceData, FNiagaraSystemInstance* SystemInstance)
+{
+	FVolumeCacheInstanceData_GameThread* InstanceData = SystemInstancesToProxyData_GT.FindRef(SystemInstance->GetId());	
+	
+	InstanceData->~FVolumeCacheInstanceData_GameThread();
+	SystemInstancesToProxyData_GT.Remove(SystemInstance->GetId());
+
+	ENQUEUE_RENDER_COMMAND(RemoveInstance)
+		(
+			[RT_Proxy = GetProxyAs<FNiagaraDataInterfaceVolumeCacheProxy>(), RT_InstanceID = SystemInstance->GetId()](FRHICommandListImmediate&)
+	{
+		RT_Proxy->InstanceData_RT.Remove(RT_InstanceID);
+	}
+	);
+}
+
+void UNiagaraDataInterfaceVolumeCache::GetFunctions(TArray<FNiagaraFunctionSignature>& OutFunctions)
+{	
+
+	{
+		FNiagaraFunctionSignature Sig;
+		Sig.Name = SetFrameName;
+		Sig.Inputs.Add(FNiagaraVariable(FNiagaraTypeDefinition(GetClass()), TEXT("Grid")));
+		Sig.Inputs.Add(FNiagaraVariable(FNiagaraTypeDefinition::GetIntDef(), TEXT("Frame")));
+		Sig.Outputs.Add(FNiagaraVariable(FNiagaraTypeDefinition::GetBoolDef(), TEXT("Success")));
+
+		Sig.ModuleUsageBitmask = ENiagaraScriptUsageMask::Emitter | ENiagaraScriptUsageMask::System;
+		Sig.bExperimental = true;
+		Sig.bMemberFunction = true;
+		Sig.bRequiresExecPin = true;
+		Sig.bRequiresContext = false;
+		Sig.bSupportsCPU = true;
+		Sig.bSupportsGPU = false;
+		OutFunctions.Add(Sig);
+	}
+
+	{
+		FNiagaraFunctionSignature Sig;
+		Sig.Name = ReadFileName;
+		Sig.Inputs.Add(FNiagaraVariable(FNiagaraTypeDefinition(GetClass()), TEXT("Grid")));
+		Sig.Inputs.Add(FNiagaraVariable(FNiagaraTypeDefinition::GetBoolDef(), TEXT("ReadFile")));
+		Sig.Outputs.Add(FNiagaraVariable(FNiagaraTypeDefinition::GetBoolDef(), TEXT("Success")));
+
+		Sig.ModuleUsageBitmask = ENiagaraScriptUsageMask::Emitter | ENiagaraScriptUsageMask::System;
+		Sig.bExperimental = true;
+		Sig.bMemberFunction = true;
+		Sig.bRequiresExecPin = true;
+		Sig.bRequiresContext = false;
+		Sig.bSupportsCPU = true;
+		Sig.bSupportsGPU = false;
+		OutFunctions.Add(Sig);
+	}
+
+	{
+		FNiagaraFunctionSignature Sig;
+		Sig.Name = SampleCurrentFrameValueName;
+		Sig.bMemberFunction = true;
+		Sig.bRequiresContext = false;
+		Sig.bSupportsCPU = false;
+		Sig.bSupportsGPU = true;
+		Sig.Inputs.Add(FNiagaraVariable(FNiagaraTypeDefinition(GetClass()), TEXT("Texture")));
+		Sig.Inputs.Add(FNiagaraVariable(FNiagaraTypeDefinition::GetVec3Def(), TEXT("UVW")));
+		Sig.Inputs.Add(FNiagaraVariable(FNiagaraTypeDefinition::GetFloatDef(), TEXT("MipLevel")));
+		Sig.SetDescription(LOCTEXT("TextureSampleVolumeTextureDesc", "Sample the specified mip level of the input 3d texture at the specified UVW coordinates. The UVW origin (0, 0, 0) is in the bottom left hand corner of the volume."));
+		Sig.Outputs.Add(FNiagaraVariable(FNiagaraTypeDefinition::GetVec4Def(), TEXT("Value")));
+		//Sig.Owner = *GetName();
+
+		OutFunctions.Add(Sig);
+	}
+
+	{
+		FNiagaraFunctionSignature Sig;
+		Sig.Name = GetCurrentFrameValue;
+		Sig.bMemberFunction = true;
+		Sig.bRequiresContext = false;
+		Sig.bSupportsCPU = false;
+		Sig.bSupportsGPU = true;
+		Sig.Inputs.Add(FNiagaraVariable(FNiagaraTypeDefinition(GetClass()), TEXT("Texture")));
+		Sig.Inputs.Add(FNiagaraVariable(FNiagaraTypeDefinition::GetIntDef(), TEXT("x")));
+		Sig.Inputs.Add(FNiagaraVariable(FNiagaraTypeDefinition::GetIntDef(), TEXT("y")));
+		Sig.Inputs.Add(FNiagaraVariable(FNiagaraTypeDefinition::GetIntDef(), TEXT("z")));
+		Sig.Inputs.Add(FNiagaraVariable(FNiagaraTypeDefinition::GetFloatDef(), TEXT("MipLevel")));
+		Sig.SetDescription(LOCTEXT("TextureLoadVolumeTextureDesc", "load the specified mip level of the input 3d texture at the specified x, y, z voxel coordinates."));
+		Sig.Outputs.Add(FNiagaraVariable(FNiagaraTypeDefinition::GetVec4Def(), TEXT("Value")));
+		//Sig.Owner = *GetName();
+
+		OutFunctions.Add(Sig);
+	}
+
+	{
+		FNiagaraFunctionSignature Sig;
+		Sig.Name = GetCurrentFrameNumCells;
+		Sig.bMemberFunction = true;
+		Sig.bRequiresContext = false;
+		Sig.bSupportsCPU = false;
+		Sig.bSupportsGPU = true;
+		Sig.Inputs.Add(FNiagaraVariable(FNiagaraTypeDefinition(GetClass()), TEXT("Texture")));
+		Sig.SetDescription(LOCTEXT("TextureDimsDesc", "Get the dimensions of mip 0 of the texture."));
+		Sig.Outputs.Add(FNiagaraVariable(FNiagaraTypeDefinition::GetVec3Def(), TEXT("Dimensions3D")));
+		//Sig.Owner = *GetName();
+
+		OutFunctions.Add(Sig);
+	}
+}
+
+void UNiagaraDataInterfaceVolumeCache::SetFrame(FVectorVMExternalFunctionContext& Context)
+{
+	// This should only be called from a system or emitter script due to a need for only setting up initially.
+	VectorVM::FUserPtrHandler<FVolumeCacheInstanceData_GameThread> InstData(Context);
+	VectorVM::FExternalFuncInputHandler<int> InFrame(Context);	
+	VectorVM::FExternalFuncRegisterHandler<FNiagaraBool> OutSuccess(Context);
+
+	InstData->CurrFrame = InFrame.GetAndAdvance();
+
+	*OutSuccess.GetDestAndAdvance() = true;
+}
+
+void UNiagaraDataInterfaceVolumeCache::ReadFile(FVectorVMExternalFunctionContext& Context)
+{
+	// This should only be called from a system or emitter script due to a need for only setting up initially.
+	VectorVM::FUserPtrHandler<FVolumeCacheInstanceData_GameThread> InstData(Context);
+	VectorVM::FExternalFuncInputHandler<bool> Read(Context);
+	VectorVM::FExternalFuncRegisterHandler<FNiagaraBool> OutSuccess(Context);
+
+	InstData->ReadFile = Read.GetAndAdvance();
+
+	*OutSuccess.GetDestAndAdvance() = true;
+}
+
+FString UNiagaraDataInterfaceVolumeCache::GetAssetPath(FString PathFormat, int32 FrameIndex) const
+{
+	UNiagaraSystem* NiagaraSystem = GetTypedOuter<UNiagaraSystem>();
+	check(NiagaraSystem);
+
+	const TMap<FString, FStringFormatArg> PathFormatArgs =
+	{		
+		{TEXT("SavedDir"),		FPaths::ProjectSavedDir()},
+		{TEXT("FrameIndex"),	FString::Printf(TEXT("%03d"), FrameIndex)},
+	};
+	FString AssetPath = FString::Format(*PathFormat, PathFormatArgs);
+	AssetPath.ReplaceInline(TEXT("//"), TEXT("/"));
+	return AssetPath;
+}
+
+DEFINE_NDI_DIRECT_FUNC_BINDER(UNiagaraDataInterfaceVolumeCache, SetFrame);
+DEFINE_NDI_DIRECT_FUNC_BINDER(UNiagaraDataInterfaceVolumeCache, ReadFile);
+void UNiagaraDataInterfaceVolumeCache::GetVMExternalFunction(const FVMExternalFunctionBindingInfo& BindingInfo, void* InstanceData, FVMExternalFunction& OutFunc)
+{
+	if (BindingInfo.Name == SetFrameName)
+	{
+		check(BindingInfo.GetNumInputs() == 2 && BindingInfo.GetNumOutputs() == 1);
+		NDI_FUNC_BINDER(UNiagaraDataInterfaceVolumeCache, SetFrame)::Bind(this, OutFunc);
+	}
+	if (BindingInfo.Name == ReadFileName)
+	{
+		check(BindingInfo.GetNumInputs() == 2 && BindingInfo.GetNumOutputs() == 1);
+		NDI_FUNC_BINDER(UNiagaraDataInterfaceVolumeCache, ReadFile)::Bind(this, OutFunc);
+	}
+}
+
+bool UNiagaraDataInterfaceVolumeCache::PerInstanceTick(void* PerInstanceData, FNiagaraSystemInstance* SystemInstance, float DeltaSeconds)
+{
+	FVolumeCacheInstanceData_GameThread* InstanceData = SystemInstancesToProxyData_GT.FindRef(SystemInstance->GetId());	
+
+	// we can run into the case where depending on the ordering of DI initialization, we might have not been able to grab the other grid's InstanceData
+	// in InitPerInstanceData.  If this is the case, we ensure it is correct here.
+	if (InstanceData && InstanceData->ReadFile && VolumeCache != nullptr)
+	{
+		FNiagaraDataInterfaceVolumeCacheProxy* TextureProxy = GetProxyAs<FNiagaraDataInterfaceVolumeCacheProxy>();
+
+		EPixelFormat Format = PF_A32B32G32R32F;
+		const int32 FormatSize = GPixelFormats[Format].BlockBytes;
+
+		// cannot read from cache...spew errors or let it go?
+		if (!VolumeCache->LoadFile(InstanceData->CurrFrame))
+		{
+			UE_LOG(LogNiagara, Warning, TEXT("Cache Read failed: %s"), *VolumeCache->GetName());
+			return false;
+		}
+
+		ENQUEUE_RENDER_COMMAND(FVolumeCacheFillTexture)(
+			[RT_Frame = InstanceData->CurrFrame, RT_VolumeCacheData = VolumeCache->GetData(), Format, TextureProxy, SystemID = SystemInstance->GetId()](FRHICommandListImmediate& RHICmdList)
+		{
+			FVolumeCacheInstanceData_RenderThread* InstanceData = TextureProxy->InstanceData_RT.Find(SystemID);
+			FIntVector Size = RT_VolumeCacheData->GetDenseResolution();
+
+			if (!InstanceData->ResolvedTextureRHI.IsValid())
+			{
+
+				const FRHITextureCreateDesc Desc =
+					FRHITextureCreateDesc::Create3D(TEXT("stuff"), Size.X, Size.Y, Size.Z, Format)
+					.SetFlags(ETextureCreateFlags::ShaderResource | ETextureCreateFlags::NoTiling);
+
+				InstanceData->ResolvedTextureRHI = RHICreateTexture(Desc);
+				InstanceData->TextureSize = FVector3f(Size.X, Size.Y, Size.Z);
+				InstanceData->SamplerStateRHI = TStaticSamplerState<SF_Bilinear, AM_Clamp, AM_Clamp, AM_Clamp>::GetRHI();
+			}
+				
+			RT_VolumeCacheData->Fill3DTexture_RenderThread(RT_Frame, InstanceData->ResolvedTextureRHI, RHICmdList);			
+		});
+	}
+
+	return false;
+
+}
+
+#if WITH_EDITORONLY_DATA
+bool UNiagaraDataInterfaceVolumeCache::AppendCompileHash(FNiagaraCompileHashVisitor* InVisitor) const
+{
+	bool bSuccess = Super::AppendCompileHash(InVisitor);
+	InVisitor->UpdateString(TEXT("UNiagaraDataInterfaceVolumeTextureHLSLSource"), GetShaderFileHash(TemplateShaderFilePath, EShaderPlatform::SP_PCD3D_SM5).ToString());
+	bSuccess &= InVisitor->UpdateShaderParameters<FShaderParameters>();
+	return bSuccess;
+}
+
+void UNiagaraDataInterfaceVolumeCache::GetParameterDefinitionHLSL(const FNiagaraDataInterfaceGPUParamInfo& ParamInfo, FString& OutHLSL)
+{
+	TMap<FString, FStringFormatArg> TemplateArgs =
+	{
+		{TEXT("ParameterName"),	ParamInfo.DataInterfaceHLSLSymbol},
+	};
+
+	FString TemplateFile;
+	LoadShaderSourceFile(TemplateShaderFilePath, EShaderPlatform::SP_PCD3D_SM5, &TemplateFile, nullptr);
+	OutHLSL += FString::Format(*TemplateFile, TemplateArgs);
+}
+
+bool UNiagaraDataInterfaceVolumeCache::GetFunctionHLSL(const FNiagaraDataInterfaceGPUParamInfo& ParamInfo, const FNiagaraDataInterfaceGeneratedFunction& FunctionInfo, int FunctionInstanceIndex, FString& OutHLSL)
+{
+	if ((FunctionInfo.DefinitionName == SampleCurrentFrameValueName) ||
+		(FunctionInfo.DefinitionName == GetCurrentFrameNumCells) ||
+		(FunctionInfo.DefinitionName == GetCurrentFrameValue))
+	{
+		return true;
+	}
+	return false;
+}
+#endif
+
+void UNiagaraDataInterfaceVolumeCache::BuildShaderParameters(FNiagaraShaderParametersBuilder& ShaderParametersBuilder) const
+{
+	ShaderParametersBuilder.AddNestedStruct<FShaderParameters>();
+}
+
+void UNiagaraDataInterfaceVolumeCache::SetShaderParameters(const FNiagaraDataInterfaceSetShaderParametersContext& Context) const
+{
+	const FNiagaraDataInterfaceVolumeCacheProxy& TextureProxy = Context.GetProxy<FNiagaraDataInterfaceVolumeCacheProxy>();
+	const FVolumeCacheInstanceData_RenderThread* InstanceData = TextureProxy.InstanceData_RT.Find(Context.GetSystemInstanceID());
+
+	FShaderParameters* Parameters = Context.GetParameterNestedStruct<FShaderParameters>();
+	if (InstanceData && InstanceData->ResolvedTextureRHI.IsValid())
+	{
+		Parameters->TextureSize = InstanceData->TextureSize;
+		Parameters->Texture = InstanceData->ResolvedTextureRHI;
+		Parameters->TextureSampler = InstanceData->SamplerStateRHI ? InstanceData->SamplerStateRHI : GBlackVolumeTexture->SamplerStateRHI;
+	}
+	else
+	{
+		Parameters->TextureSize = FVector3f::ZeroVector;
+		Parameters->Texture = GBlackVolumeTexture->TextureRHI;
+		Parameters->TextureSampler = GBlackVolumeTexture->SamplerStateRHI;
+	}
+}
+
+#undef LOCTEXT_NAMESPACE
diff --git a/Engine/Plugins/FX/Niagara/Source/Niagara/Private/NiagaraModule.cpp b/Engine/Plugins/FX/Niagara/Source/Niagara/Private/NiagaraModule.cpp
--- a/Engine/Plugins/FX/Niagara/Source/Niagara/Private/NiagaraModule.cpp
+++ b/Engine/Plugins/FX/Niagara/Source/Niagara/Private/NiagaraModule.cpp
@@ -29,6 +29,10 @@
 #include "Engine/StaticMesh.h"
 #include "UObject/UObjectGlobals.h"
 
+#if PLATFORM_WINDOWS
+#include "NiagaraOpenVDB.h"
+#endif
+
 IMPLEMENT_MODULE(INiagaraModule, Niagara);
 
 #define LOCTEXT_NAMESPACE "NiagaraModule"
@@ -197,6 +201,13 @@ void INiagaraModule::StartupModule()
 	FNiagaraTypeDefinition::Init();
 	FNiagaraRenderViewDataManager::Init();
 
+	
+#if PLATFORM_WINDOWS
+	// Global registration of  the vdb types.
+	openvdb::initialize();
+	Vec4SGrid::registerGrid();
+#endif
+
 	FNiagaraWorldManager::OnStartup();
 
 #if WITH_NIAGARA_DEBUGGER
diff --git a/Engine/Plugins/FX/Niagara/Source/Niagara/Private/VolumeCache.cpp b/Engine/Plugins/FX/Niagara/Source/Niagara/Private/VolumeCache.cpp
new file mode 100644
--- /dev/null
+++ b/Engine/Plugins/FX/Niagara/Source/Niagara/Private/VolumeCache.cpp
@@ -0,0 +1,357 @@
+// Copyright Epic Games, Inc. All Rights Reserved.
+
+/*=============================================================================
+	VolumeCache.cpp: UvolumeTexture implementation.
+=============================================================================*/
+
+#include "VolumeCache.h"
+
+
+DEFINE_LOG_CATEGORY_STATIC(LogVolumeCache, Log, All);
+
+//*****************************************************************************
+//***************************** UVolumeCache ********************************
+//*****************************************************************************
+
+UVolumeCache::UVolumeCache(const FObjectInitializer& ObjectInitializer)
+	: Super(ObjectInitializer), CachedVolumeFiles(nullptr)
+{	
+}
+
+void UVolumeCache::InitData()
+{	
+	if (CachedVolumeFiles != nullptr)
+	{
+		CachedVolumeFiles.Reset();
+	}
+	
+	switch (CacheType)
+	{
+		case EVolumeCacheType::OpenVDB :
+		{			
+#if PLATFORM_WINDOWS
+			CachedVolumeFiles = TSharedPtr<FVolumeCacheData>(new FOpenVDBCacheData());
+#else
+			UE_LOG(LogVolumeCache, Warning, TEXT("OpenVDB not supported on the platform"));
+#endif
+
+			break;
+		}
+		default:
+			UE_LOG(LogVolumeCache, Warning, TEXT("Invalid Volume Cache type"));
+	}
+		
+	CachedVolumeFiles->Init(Resolution);	
+}
+
+bool UVolumeCache::LoadRange()
+{
+	if (CachedVolumeFiles == nullptr)
+	{
+		InitData();
+	}
+
+	return CachedVolumeFiles->LoadRange(FilePath, FrameRangeStart, FrameRangeEnd);
+}
+
+void UVolumeCache::UnloadAll()
+{
+	if (CachedVolumeFiles != nullptr)
+	{
+		CachedVolumeFiles->UnloadAll();
+	}
+}
+
+bool UVolumeCache::LoadFile(int frame)
+{
+	if (CachedVolumeFiles == nullptr)
+	{
+		InitData();
+	}
+	
+	return CachedVolumeFiles->LoadFile(FilePath, frame);	
+}
+
+bool UVolumeCache::UnloadFile(int frame)
+{
+	if (CachedVolumeFiles == nullptr)
+	{
+		InitData();
+	}
+
+	return CachedVolumeFiles->UnloadFile(frame);
+}
+
+FString FVolumeCacheData::GetAssetPath(FString PathFormat, int32 FrameIndex) const
+{
+	const TMap<FString, FStringFormatArg> PathFormatArgs =
+	{
+		{TEXT("FrameIndex"),	FString::Printf(TEXT("%03d"), FrameIndex)},
+	};
+	FString AssetPath = FString::Format(*PathFormat, PathFormatArgs);
+	AssetPath.ReplaceInline(TEXT("//"), TEXT("/"));
+	return AssetPath;
+}
+
+// @todo: we need builds for OpenVDB for platforms other than windows
+#if PLATFORM_WINDOWS
+void FOpenVDBCacheData::Init(FIntVector Resolution)
+{
+	DenseGridPtr.reset(new openvdb::tools::Dense<openvdb::Vec4s, openvdb::tools::MemoryLayout::LayoutXYZ>(openvdb::CoordBBox(0, 0, 0, Resolution.X - 1, Resolution.Y - 1, Resolution.Z - 1), openvdb::Vec4s(0.0, 0.0, 0.0, 0.0)));
+}
+
+bool FOpenVDBCacheData::LoadFile(FString Path, int frame)
+{
+
+	// if the file is loaded at the current frame, return true
+	if (OpenVDBGrids.Contains(frame) && OpenVDBGrids[frame] != nullptr)
+	{
+		openvdb::tools::CopyToDense<Vec4STree, openvdb::tools::Dense<openvdb::Vec4s, openvdb::tools::MemoryLayout::LayoutXYZ>> Copier(OpenVDBGrids[frame]->tree(), *DenseGridPtr);
+		Copier.copy();
+
+		return true;
+	}
+	else
+	{
+		FString FullPath = GetAssetPath(Path, frame);
+		std::string FileNameStr(TCHAR_TO_ANSI(*FullPath));
+
+		// Create a VDB file object.	
+		openvdb::io::File file(FileNameStr);
+
+		// Open the file.  This reads the file header, but not any grids.		
+		try
+		{
+			file.open(false);
+		}
+		catch (openvdb::Exception e)
+		{
+			UE_LOG(LogVolumeCache, Warning, TEXT("Cache File not Found: %s"), *FullPath);
+			return false;
+		}
+
+		Vec4SGrid::Ptr ColorGrid;		
+		ColorGrid = openvdb::gridPtrCast<Vec4SGrid>(file.readGrid(openvdb::Name("Color")));		
+
+		if (ColorGrid == nullptr)
+		{
+			UE_LOG(LogVolumeCache, Warning, TEXT("Cache has invalid grids: %s"), *FullPath);
+			return false;
+		}
+
+		openvdb::MetaMap::Ptr Metadata = file.getMetadata();
+		FIntVector Size;
+		try
+		{
+			Size.X = Metadata->metaValue<int>("DenseResolutionX");
+			Size.Y = Metadata->metaValue<int>("DenseResolutionY");
+			Size.Z = Metadata->metaValue<int>("DenseResolutionZ");
+			if (DenseResolution == FIntVector(-1, -1, -1))
+			{
+				DenseResolution = Size;
+			} 
+			else if (DenseResolution != Size)
+			{
+				UE_LOG(LogVolumeCache, Warning, TEXT("Grids are not the same size: %s"), *FullPath);
+				return false;
+			}			
+
+			OpenVDBGrids.Add(frame, ColorGrid);
+
+			// load to dense buffer			
+			openvdb::tools::CopyToDense<Vec4STree, openvdb::tools::Dense<openvdb::Vec4s, openvdb::tools::MemoryLayout::LayoutXYZ>> Copier(ColorGrid->tree(), *DenseGridPtr);
+			Copier.copy();
+		}
+		catch (openvdb::Exception e)
+		{
+			UE_LOG(LogVolumeCache, Warning, TEXT("No cache metadata found: %s"), *FullPath);
+			return false;
+		}
+
+		file.close();
+	}
+	return true;
+}
+
+bool FOpenVDBCacheData::UnloadFile(int frame)
+{
+	if (OpenVDBGrids.Contains(frame) && OpenVDBGrids[frame] != nullptr)
+	{
+		OpenVDBGrids[frame] = nullptr;
+		OpenVDBGrids.Remove(frame);
+
+		return true;
+	}
+
+	return false;
+}
+
+bool FOpenVDBCacheData::LoadRange(FString Path, int Start, int End)
+{	
+	for (int i = Start; i <= End; ++i)
+	{
+		if (!LoadFile(Path, i))
+		{
+			return false;
+		}
+	}
+
+	return true;
+}
+
+void FOpenVDBCacheData::UnloadAll()
+{
+	OpenVDBGrids.Reset();
+}
+
+bool FOpenVDBCacheData::Fill3DTexture_RenderThread(int frame, FTextureRHIRef TextureToFill, FRHICommandListImmediate& RHICmdList)
+{
+	if (OpenVDBGrids.Contains(frame) && OpenVDBGrids[frame] != nullptr)
+	{
+		uint8* DataPtr = (uint8*)DenseGridPtr->data();
+
+		const int32 FormatSize = GPixelFormats[TextureToFill->GetFormat()].BlockBytes;
+
+		if (TextureToFill->GetSizeXYZ() != DenseResolution)
+		{
+			UE_LOG(LogVolumeCache, Warning, TEXT("Target texture resolution %s doesn't match OpenVDB grid resolution %s.  Regenerate your cache"), *TextureToFill->GetSizeXYZ().ToString(), *DenseResolution.ToString());
+		}
+		else
+		{
+			FUpdateTextureRegion3D UpdateRegion(0, 0, 0, 0, 0, 0, DenseResolution.X, DenseResolution.Y, DenseResolution.Z);			
+			RHICmdList.UpdateTexture3D(TextureToFill, 0, UpdateRegion, FormatSize * DenseResolution.X, FormatSize * DenseResolution.X * DenseResolution.Y, DataPtr);			
+		}
+
+		return true;
+	}
+
+	return false;
+}
+
+bool FOpenVDBCacheData::WriteImageDataToOpenVDBFile(FStringView FilePath, FIntVector ImageSize, TArrayView<FFloat16Color> ImageData, bool UseFloatGrids)
+{
+	openvdb::math::Transform::Ptr TransformPtr = openvdb::math::Transform::createLinearTransform(/*voxel size=*/1.0);
+
+	openvdb::FloatGrid::Ptr RGrid;
+	openvdb::FloatGrid::Ptr GGrid;
+	openvdb::FloatGrid::Ptr BGrid;
+	openvdb::FloatGrid::Ptr AGrid;
+
+	Vec4SGrid::Ptr ColorGrid;
+
+	if (UseFloatGrids)
+	{
+		RGrid = openvdb::FloatGrid::create(0.0);
+		GGrid = openvdb::FloatGrid::create(0.0);
+		BGrid = openvdb::FloatGrid::create(0.0);
+		AGrid = openvdb::FloatGrid::create(0.0);
+
+		RGrid->setTransform(TransformPtr);
+		GGrid->setTransform(TransformPtr);
+		BGrid->setTransform(TransformPtr);
+		AGrid->setTransform(TransformPtr);
+
+		RGrid->setGridClass(openvdb::GRID_FOG_VOLUME);
+		GGrid->setGridClass(openvdb::GRID_FOG_VOLUME);
+		BGrid->setGridClass(openvdb::GRID_FOG_VOLUME);
+		AGrid->setGridClass(openvdb::GRID_FOG_VOLUME);
+
+		RGrid->setName("R");
+		GGrid->setName("G");
+		BGrid->setName("B");
+		AGrid->setName("A");
+
+		openvdb::FloatGrid::Accessor Raccessor = RGrid->getAccessor();
+		openvdb::FloatGrid::Accessor Gaccessor = GGrid->getAccessor();
+		openvdb::FloatGrid::Accessor Baccessor = BGrid->getAccessor();
+		openvdb::FloatGrid::Accessor Aaccessor = AGrid->getAccessor();
+
+		for (int z = 0; z < ImageSize.Z; ++z) {
+			for (int y = 0; y < ImageSize.Y; ++y) {
+				for (int x = 0; x < ImageSize.X; ++x) {
+					uint32 LinearIndex = x + y * ImageSize.X + z * ImageSize.X * ImageSize.Y;
+					FFloat16Color CurrValue = ImageData[LinearIndex];
+
+					Raccessor.setValue(openvdb::Coord(x, y, z), CurrValue.R);
+					Gaccessor.setValue(openvdb::Coord(x, y, z), CurrValue.G);
+					Baccessor.setValue(openvdb::Coord(x, y, z), CurrValue.B);
+					Aaccessor.setValue(openvdb::Coord(x, y, z), CurrValue.A);
+				}	
+			}
+		}
+	}
+	else
+	{
+		ColorGrid = Vec4SGrid::create(openvdb::Vec4s(0.0));
+		ColorGrid->setTransform(TransformPtr);
+		ColorGrid->setGridClass(openvdb::GRID_FOG_VOLUME);
+		ColorGrid->setName("Color");
+		Vec4SGrid::Accessor ColorAccessor = ColorGrid->getAccessor();
+
+		for (int z = 0; z < ImageSize.Z; ++z) {
+			for (int y = 0; y < ImageSize.Y; ++y) {
+				for (int x = 0; x < ImageSize.X; ++x) {
+					uint32 LinearIndex = x + y * ImageSize.X + z * ImageSize.X * ImageSize.Y;
+					FFloat16Color CurrValue = ImageData[LinearIndex];
+
+					ColorAccessor.setValue(openvdb::Coord(x, y, z), openvdb::Vec4s(CurrValue.R, CurrValue.G, CurrValue.B, CurrValue.A));
+				}
+			}
+		}
+	}
+
+	// Create a VDB file object.		
+	FString Filename(FilePath);
+	std::string FileNameStr(TCHAR_TO_ANSI(*Filename));
+
+	openvdb::io::File file(FileNameStr);
+
+	// Add the grid pointer to a container.
+	openvdb::GridPtrVec grids;
+
+	if (UseFloatGrids)
+	{
+		openvdb::tools::pruneInactive(RGrid->tree());
+		openvdb::tools::pruneInactive(GGrid->tree());
+		openvdb::tools::pruneInactive(BGrid->tree());
+		openvdb::tools::pruneInactive(AGrid->tree());
+
+		RGrid->pruneGrid(0);
+		GGrid->pruneGrid(0);
+		BGrid->pruneGrid(0);
+		AGrid->pruneGrid(0);
+
+		grids.push_back(RGrid);
+		grids.push_back(GGrid);
+		grids.push_back(BGrid);
+		grids.push_back(AGrid);
+	}	
+	else
+	{
+		openvdb::tools::pruneInactive(ColorGrid->tree());
+		ColorGrid->pruneGrid(0);
+		grids.push_back(ColorGrid);
+	}
+	
+	// Add file-level metadata.
+	openvdb::MetaMap outMeta;
+	outMeta.insertMeta("creator",
+		openvdb::StringMetadata("Unreal Engine"));
+
+	outMeta.insertMeta("DenseResolutionX",
+		openvdb::Int32Metadata(ImageSize.X));
+	outMeta.insertMeta("DenseResolutionY",
+		openvdb::Int32Metadata(ImageSize.Y));
+	outMeta.insertMeta("DenseResolutionZ",
+		openvdb::Int32Metadata(ImageSize.Z));
+
+
+	uint32_t compressionFlags = openvdb::io::COMPRESS_BLOSC;
+	file.setCompression(compressionFlags);
+
+	file.write(grids, outMeta);
+	file.close();
+
+	return true;
+}
+#endif
\ No newline at end of file
diff --git a/Engine/Plugins/FX/Niagara/Source/NiagaraEditor/NiagaraEditor.Build.cs b/Engine/Plugins/FX/Niagara/Source/NiagaraEditor/NiagaraEditor.Build.cs
--- a/Engine/Plugins/FX/Niagara/Source/NiagaraEditor/NiagaraEditor.Build.cs
+++ b/Engine/Plugins/FX/Niagara/Source/NiagaraEditor/NiagaraEditor.Build.cs
@@ -6,6 +6,13 @@ public class NiagaraEditor : ModuleRules
 {
 	public NiagaraEditor(ReadOnlyTargetRules Target) : base(Target)
 	{
+		// Specific to OpenVDB support
+		if (Target.IsInPlatformGroup(UnrealPlatformGroup.Windows))
+		{
+			bUseRTTI = true;
+			bEnableExceptions = true;
+		}
+
 		PrivateIncludePaths.AddRange(new string[] {
 			"NiagaraEditor/Private",
 			"NiagaraEditor/Private/Toolkits",
@@ -92,7 +99,18 @@ public class NiagaraEditor : ModuleRules
             }
         );
 
-        PublicIncludePathModuleNames.AddRange(
+		if (Target.IsInPlatformGroup(UnrealPlatformGroup.Windows))
+		{
+			AddEngineThirdPartyPrivateStaticDependencies(Target,
+				"IntelTBB",
+				"Blosc",
+				"zlib",
+				"Boost",
+				"OpenVDB"
+			);
+		}
+
+		PublicIncludePathModuleNames.AddRange(
             new string[] {
 				"Engine",
 				"Niagara"
diff --git a/Engine/Plugins/FX/Niagara/Source/NiagaraEditor/Private/NiagaraBakerRenderer.cpp b/Engine/Plugins/FX/Niagara/Source/NiagaraEditor/Private/NiagaraBakerRenderer.cpp
--- a/Engine/Plugins/FX/Niagara/Source/NiagaraEditor/Private/NiagaraBakerRenderer.cpp
+++ b/Engine/Plugins/FX/Niagara/Source/NiagaraEditor/Private/NiagaraBakerRenderer.cpp
@@ -27,6 +27,7 @@
 #include "IImageWrapper.h"
 #include "ImageWrapperHelper.h"
 #include "LegacyScreenPercentageDriver.h"
+#include "VolumeCache.h"
 
 //////////////////////////////////////////////////////////////////////////
 
@@ -652,13 +653,7 @@ FNiagaraBakerOutputRenderer* FNiagaraBakerRenderer::GetOutputRenderer(UClass* Cl
 
 bool FNiagaraBakerRenderer::ExportImage(FStringView FilePath, FIntPoint ImageSize, TArrayView<FFloat16Color> ImageData)
 {
-	int32 LastDot = INDEX_NONE;
-	if ( FilePath.FindLastChar('.', LastDot) == false )
-	{
-		return false;
-	}
-
-	const FStringView FileExtension = FilePath.Mid(LastDot);
+	const FString FileExtension = FPaths::GetExtension(FilePath.GetData(), true);
 	const EImageFormat ImageFormat = ImageWrapperHelper::GetImageFormat(FileExtension);
 	if (ImageFormat == EImageFormat::Invalid)
 	{
@@ -738,3 +733,19 @@ void FNiagaraBakerRenderer::DestroyPreviewScene(UNiagaraComponent*& InOutCompone
 		InOutComponent = nullptr;
 	}
 }
+bool FNiagaraBakerRenderer::ExportVolume(FStringView FilePath, FIntVector ImageSize, TArrayView<FFloat16Color> ImageData)
+{
+	const FString FileExtension = FPaths::GetExtension(FilePath.GetData(), true);
+	if (FileExtension == TEXT(".vdb"))
+	{
+#if PLATFORM_WINDOWS
+		return FOpenVDBCacheData::WriteImageDataToOpenVDBFile(FilePath, ImageSize, ImageData, false);
+#else
+		return false;
+#endif
+	}
+	else
+	{
+		return ExportImage(FilePath, FIntPoint(ImageSize.X, ImageSize.Y * ImageSize.Z), ImageData);
+	}
+}
\ No newline at end of file
diff --git a/Engine/Plugins/FX/Niagara/Source/NiagaraEditor/Private/NiagaraBakerRenderer.h b/Engine/Plugins/FX/Niagara/Source/NiagaraEditor/Private/NiagaraBakerRenderer.h
--- a/Engine/Plugins/FX/Niagara/Source/NiagaraEditor/Private/NiagaraBakerRenderer.h
+++ b/Engine/Plugins/FX/Niagara/Source/NiagaraEditor/Private/NiagaraBakerRenderer.h
@@ -116,7 +116,7 @@ public:
 	static FNiagaraBakerOutputRenderer* GetOutputRenderer(UClass* Class);
 
 	static bool ExportImage(FStringView FilePath, FIntPoint ImageSize, TArrayView<FFloat16Color> ImageData);
-
+	static bool ExportVolume(FStringView FilePath, FIntVector ImageSize, TArrayView<FFloat16Color> ImageData);
 private:
 	static void CreatePreviewScene(UNiagaraSystem* NiagaraSystem, UNiagaraComponent*& OutComponent, TSharedPtr<FAdvancedPreviewScene>& OutPreviewScene);
 	static void DestroyPreviewScene(UNiagaraComponent*& InOutComponent, TSharedPtr<FAdvancedPreviewScene>& InOutPreviewScene);
diff --git a/Engine/Plugins/FX/Niagara/Source/NiagaraEditor/Private/NiagaraBakerRendererOutputVolumeTexture.cpp b/Engine/Plugins/FX/Niagara/Source/NiagaraEditor/Private/NiagaraBakerRendererOutputVolumeTexture.cpp
--- a/Engine/Plugins/FX/Niagara/Source/NiagaraEditor/Private/NiagaraBakerRendererOutputVolumeTexture.cpp
+++ b/Engine/Plugins/FX/Niagara/Source/NiagaraEditor/Private/NiagaraBakerRendererOutputVolumeTexture.cpp
@@ -17,6 +17,12 @@
 #include "Factories/VolumeTextureFactory.h"
 #include "IAssetTools.h"
 #include "AssetToolsModule.h"
+#include "UObject/UObjectGlobals.h"
+
+#include "VolumeCache.h"
+#include "VolumeCacheFactory.h"
+
+DEFINE_LOG_CATEGORY_STATIC(LogNiagaraBaker, Log, All);
 
 struct FVolumeDataInterfaceHelper
 {
@@ -386,6 +392,10 @@ void FNiagaraBakerRendererOutputVolumeTexture::RenderGenerated(UNiagaraBakerOutp
 bool FNiagaraBakerRendererOutputVolumeTexture::BeginBake(UNiagaraBakerOutput* InBakerOutput)
 {
 	UNiagaraBakerOutputVolumeTexture* OutputVolumeTexture = CastChecked<UNiagaraBakerOutputVolumeTexture>(InBakerOutput);
+
+	// @todo: 
+	BakedFrameRange = FIntVector2(TNumericLimits<int32>::Lowest(), TNumericLimits<int32>::Lowest());
+
 	return OutputVolumeTexture->bGenerateAtlas || OutputVolumeTexture->bGenerateFrames || OutputVolumeTexture->bExportFrames;
 }
 
@@ -570,8 +580,21 @@ void FNiagaraBakerRendererOutputVolumeTexture::BakeFrame(UNiagaraBakerOutput* In
 		// Export each frame
 		if (OutputVolumeTexture->bExportFrames)
 		{
+			BakedFrameSize = TextureSize;
+
+			if (BakedFrameRange.X == TNumericLimits<int32>::Lowest() && BakedFrameRange.Y == TNumericLimits<int32>::Lowest())
+			{
+				BakedFrameRange.X = FrameIndex;
+				BakedFrameRange.Y = FrameIndex;
+			}
+			else
+			{
+				BakedFrameRange.X = FMath::Min(BakedFrameRange.X, FrameIndex);
+				BakedFrameRange.Y = FMath::Max(BakedFrameRange.Y, FrameIndex);
+			}
+
 			const FString FilePath = OutputVolumeTexture->GetExportPath(OutputVolumeTexture->FramesExportPathFormat, FrameIndex);
-			FNiagaraBakerRenderer::ExportImage(FilePath, FIntPoint(TextureSize.X, TextureSize.Y * TextureSize.Z), TextureData);
+			FNiagaraBakerRenderer::ExportVolume(FilePath, TextureSize, TextureData);
 		}
 	}
 }
@@ -604,5 +627,66 @@ void FNiagaraBakerRendererOutputVolumeTexture::EndBake(UNiagaraBakerOutput* InBa
 		BakeAtlasTextureSize = FIntVector::ZeroValue;
 		BakeAtlasTextureData.Empty();
 	}
+
+	// @todo: move this to the right place
+	if (OutputVolumeTexture->bExportFrames)
+	{
+		const FString AssetFullName = OutputVolumeTexture->GetAssetPath(OutputVolumeTexture->AtlasAssetPathFormat, 0);
+		const FString AssetName = FString(FPathViews::GetCleanFilename(AssetFullName));
+		const FString PackagePath = FString(FPathViews::GetPath(AssetFullName));
+		IAssetTools& AssetTools = FModuleManager::Get().LoadModuleChecked<FAssetToolsModule>("AssetTools").Get();
+		
+		// see if the asset has been created
+		UPackage* Pkg = CreatePackage(*AssetFullName);		
+		UVolumeCache* OutputCache = Cast<UVolumeCache>(StaticFindObject(UObject::StaticClass(), Pkg, *AssetName));
+		
+		if (OutputCache == nullptr)
+		{
+			UVolumeCacheFactory* VolumeCacheFactory = NewObject<UVolumeCacheFactory>();
+			OutputCache = Cast<UVolumeCache>(AssetTools.CreateAsset(AssetName, PackagePath, UVolumeCache::StaticClass(), VolumeCacheFactory));
+		}
+
+		FString FullPath = OutputVolumeTexture->GetExportPath(OutputVolumeTexture->FramesExportPathFormat, 0);
+		FullPath = FString(FPathViews::GetPath(FullPath));
+
+		FString FullFileName = FString(FPathViews::GetCleanFilename(OutputVolumeTexture->FramesExportPathFormat));
+
+		FullPath = FullPath + "/" + FullFileName;
+
+		OutputCache->FilePath = FullPath;
+		
+		OutputCache->FrameRangeStart = BakedFrameRange.X;
+		OutputCache->FrameRangeEnd = BakedFrameRange.Y;
+		OutputCache->Resolution = BakedFrameSize;
+		OutputCache->CacheType = EVolumeCacheType::OpenVDB;
+
+		OutputCache->Modify();
+		OutputCache->PostEditChange();		
+
+		if (BakedFrameRange.Y - BakedFrameRange.X > 0)
+		{
+			if (BakedFrameSize.X + BakedFrameSize.Y + BakedFrameSize.Z > 0)
+			{
+
+				OutputCache->InitData();
+		
+				// @todo: could be cool to have an option where the whole cache is read into memory after writing it so you can preview results super quick
+				// bool LoadedCache = OutputCache->LoadRange();
+
+				// if (!LoadedCache)
+				// {
+				// 	UE_LOG(LogNiagaraBaker, Warning, TEXT("Could not load cache into memory: %s"), *FullPath);
+				// }
+			}
+			else
+			{
+				UE_LOG(LogNiagaraBaker, Warning, TEXT("Output cache has no voxels : %s"), *FullPath);
+			}
+		}
+		else
+		{
+			UE_LOG(LogNiagaraBaker, Warning, TEXT("No frames written : %s"), *FullPath);
+		}
+	}
 }
 
diff --git a/Engine/Plugins/FX/Niagara/Source/NiagaraEditor/Private/NiagaraBakerRendererOutputVolumeTexture.h b/Engine/Plugins/FX/Niagara/Source/NiagaraEditor/Private/NiagaraBakerRendererOutputVolumeTexture.h
--- a/Engine/Plugins/FX/Niagara/Source/NiagaraEditor/Private/NiagaraBakerRendererOutputVolumeTexture.h
+++ b/Engine/Plugins/FX/Niagara/Source/NiagaraEditor/Private/NiagaraBakerRendererOutputVolumeTexture.h
@@ -25,4 +25,7 @@ protected:
 	FIntVector				BakeAtlasFrameSize = FIntVector::ZeroValue;
 	FIntVector				BakeAtlasTextureSize = FIntVector::ZeroValue;
 	TArray<FFloat16Color>	BakeAtlasTextureData;
+
+	FIntVector				BakedFrameSize = FIntVector::ZeroValue;
+	FIntVector2				BakedFrameRange = FIntVector2(-1, -1);
 };
diff --git a/Engine/Plugins/FX/Niagara/Source/NiagaraEditor/Private/NiagaraEditorModule.cpp b/Engine/Plugins/FX/Niagara/Source/NiagaraEditor/Private/NiagaraEditorModule.cpp
--- a/Engine/Plugins/FX/Niagara/Source/NiagaraEditor/Private/NiagaraEditorModule.cpp
+++ b/Engine/Plugins/FX/Niagara/Source/NiagaraEditor/Private/NiagaraEditorModule.cpp
@@ -151,6 +151,10 @@
 #include "NiagaraGraphDataCache.h"
 #include "Misc/ScopedSlowTask.h"
 
+#if PLATFORM_WINDOWS
+#include "NiagaraOpenVDB.h"
+#endif
+
 IMPLEMENT_MODULE( FNiagaraEditorModule, NiagaraEditor );
 
 #define LOCTEXT_NAMESPACE "NiagaraEditorModule"
@@ -818,6 +822,12 @@ void FNiagaraEditorModule::PostGarbageCollect()
 
 void FNiagaraEditorModule::StartupModule()
 {
+	// Global registration of  the vdb types.
+#if PLATFORM_WINDOWS
+	openvdb::initialize();
+	Vec4SGrid::registerGrid();
+#endif	
+	
 	bThumbnailRenderersRegistered = false;
 
 	FHlslNiagaraTranslator::Init();
diff --git a/Engine/Plugins/FX/Niagara/Source/NiagaraEditor/Private/VolumeCacheFactory.cpp b/Engine/Plugins/FX/Niagara/Source/NiagaraEditor/Private/VolumeCacheFactory.cpp
new file mode 100644
--- /dev/null
+++ b/Engine/Plugins/FX/Niagara/Source/NiagaraEditor/Private/VolumeCacheFactory.cpp
@@ -0,0 +1,37 @@
+// Copyright Epic Games, Inc. All Rights Reserved.
+
+#include "VolumeCacheFactory.h"
+#include "Styling/SlateBrush.h"
+#include "Brushes/SlateDynamicImageBrush.h"
+#include "VolumeCache.h"
+
+#define LOCTEXT_NAMESPACE "VolumeCacheFactory"
+
+#if WITH_EDITOR
+UVolumeCacheFactory::UVolumeCacheFactory( const FObjectInitializer& ObjectInitializer )
+ : Super(ObjectInitializer)
+{
+	bCreateNew = true;
+	bEditAfterNew = true;
+	SupportedClass = UVolumeCache::StaticClass();
+}
+
+FText UVolumeCacheFactory::GetDisplayName() const
+{
+	return LOCTEXT("VolumeCacheFactoryDescription", "Volume Texture");
+}
+
+bool UVolumeCacheFactory::ConfigureProperties()
+{
+	return true;
+}
+
+UObject* UVolumeCacheFactory::FactoryCreateNew(UClass* Class,UObject* InParent,FName Name,EObjectFlags Flags,UObject* Context,FFeedbackContext* Warn)
+{
+	UVolumeCache* NewVolumeCache = NewObject<UVolumeCache>(InParent, Name, Flags);	
+
+	return NewVolumeCache;
+}
+#endif
+
+#undef LOCTEXT_NAMESPACE
diff --git a/Engine/Plugins/FX/Niagara/Source/NiagaraEditor/Public/VolumeCacheFactory.h b/Engine/Plugins/FX/Niagara/Source/NiagaraEditor/Public/VolumeCacheFactory.h
new file mode 100644
--- /dev/null
+++ b/Engine/Plugins/FX/Niagara/Source/NiagaraEditor/Public/VolumeCacheFactory.h
@@ -0,0 +1,26 @@
+// Copyright Epic Games, Inc. All Rights Reserved.
+
+#pragma once
+
+#include "CoreMinimal.h"
+#include "UObject/ObjectMacros.h"
+#include "Factories/Factory.h"
+#include "VolumeCacheFactory.generated.h"
+
+/** Factory for creating volume texture */
+#if WITH_EDITOR
+
+UCLASS(hidecategories=Object, MinimalAPI)
+class UVolumeCacheFactory : public UFactory
+{
+	GENERATED_UCLASS_BODY()
+
+	//~ Begin UFactory Interface
+	virtual FText GetDisplayName() const override;
+	virtual bool ConfigureProperties() override;
+	virtual UObject* FactoryCreateNew(UClass* Class,UObject* InParent,FName Name,EObjectFlags Flags,UObject* Context,FFeedbackContext* Warn) override;
+	//~ Begin UFactory Interface
+};
+#endif
+
+
