cis fix

diff --git a/Engine/Plugins/FX/Niagara/Source/Niagara/Classes/NiagaraDataInterfaceVolumeCache.h b/Engine/Plugins/FX/Niagara/Source/Niagara/Classes/NiagaraDataInterfaceVolumeCache.h
--- a/Engine/Plugins/FX/Niagara/Source/Niagara/Classes/NiagaraDataInterfaceVolumeCache.h
+++ b/Engine/Plugins/FX/Niagara/Source/Niagara/Classes/NiagaraDataInterfaceVolumeCache.h
@@ -26,7 +26,7 @@ class NIAGARA_API UNiagaraDataInterfaceVolumeCache : public UNiagaraDataInterfac
 public:
 		
 	UPROPERTY(EditAnywhere, Category = "File")
-	UVolumeCache *VolumeCache;
+	TObjectPtr<UVolumeCache> VolumeCache;
 
 	virtual void PostInitProperties() override;
 
