cis fix

diff --git a/Engine/Plugins/FX/Niagara/Source/Niagara/Classes/VolumeCache.h b/Engine/Plugins/FX/Niagara/Source/Niagara/Classes/VolumeCache.h
--- a/Engine/Plugins/FX/Niagara/Source/Niagara/Classes/VolumeCache.h
+++ b/Engine/Plugins/FX/Niagara/Source/Niagara/Classes/VolumeCache.h
@@ -5,6 +5,8 @@
 
 #include "CoreMinimal.h"
 #include "UObject/ObjectMacros.h"
+#include "RHICommandList.h"
+#include "RHIResources.h"
 
 // @todo: we need builds for OpenVDB for platforms other than windows
 #if PLATFORM_WINDOWS
