Niagara: suppress static analysis warning about arithmetic overflow coming from openvdb/tree/LeafNode.h

diff --git a/Engine/Plugins/FX/Niagara/Source/Niagara/Classes/NiagaraOpenVDB.h b/Engine/Plugins/FX/Niagara/Source/Niagara/Classes/NiagaraOpenVDB.h
--- a/Engine/Plugins/FX/Niagara/Source/Niagara/Classes/NiagaraOpenVDB.h
+++ b/Engine/Plugins/FX/Niagara/Source/Niagara/Classes/NiagaraOpenVDB.h
@@ -12,6 +12,7 @@ UE_PUSH_MACRO("check")
 #ifdef _MSC_VER
 #pragma warning (push)
 #pragma warning (disable : 4146) // unary minus operator applied to unsigned type, result still unsigned
+#pragma warning (disable : 6297) // Arithmetic overflow:  32-bit value is shifted, then cast to 64-bit value.  Results might not be an expected value.
 #endif
 
 #ifndef M_PI
