Clean-up copy and paste behaviors for Niagara Overview Graphs.
-Migrate checks to appropriate CanCopy and CanPaste funtions where we verify if the overivew is an emitter
-Copying from emitters is now no longer supported. The intended behavior here isn't well defined, so users are directed to the more supported path of using the add emitter option from the context menu.
-Creating comments is no longer supported in emitter graphs. The overview is transient so any comments would not persist.

diff --git a/Engine/Plugins/FX/Niagara/Source/NiagaraEditor/Private/ViewModels/NiagaraOverviewGraphViewModel.cpp b/Engine/Plugins/FX/Niagara/Source/NiagaraEditor/Private/ViewModels/NiagaraOverviewGraphViewModel.cpp
--- a/Engine/Plugins/FX/Niagara/Source/NiagaraEditor/Private/ViewModels/NiagaraOverviewGraphViewModel.cpp
+++ b/Engine/Plugins/FX/Niagara/Source/NiagaraEditor/Private/ViewModels/NiagaraOverviewGraphViewModel.cpp
@@ -288,6 +288,17 @@ void FNiagaraOverviewGraphViewModel::CopySelectedNodes()
 bool FNiagaraOverviewGraphViewModel::CanCopyNodes() const
 {
 	UEdGraph* Graph = GetGraph();
+
+	// Copying from emitters is unsupported. An emitter overview node only references an emitter handle ID and a transient system.
+	// Since the system is transient, the emitter handle isn't always valid for duplication. Moreover the duplication would circumvent emitter inheritance. 
+	// We don't store enough information on the node to find the emitter asset in question, so adding the emitter with the correct hierarchy would require a larger refactor to properly support.
+	if (GetSystemViewModel()->GetEditMode() == ENiagaraSystemViewModelEditMode::EmitterAsset)
+	{		
+		// Instead direct the user to add the emitter through the right-click menu, where the behavior is well defined.
+		FNiagaraEditorUtilities::InfoWithToastAndLog(LOCTEXT("CopyFromEmitterAsset", "Cannot copy emitters from emitter assets. Please use Add Emitter from the right click menu instead."));
+		return false;
+	}
+
 	if (Graph != nullptr)
 	{
 		for (UObject* SelectedNode : NodeSelection->GetSelectedObjects())
@@ -308,14 +319,6 @@ void FNiagaraOverviewGraphViewModel::PasteNodes()
 	UEdGraph* Graph = GetGraph();
 	if (Graph != nullptr)
 	{
-		// Check if we need to early out due to trying to paste nodes to an emitter asset edit mode overview graph.
-		TSharedPtr<FNiagaraSystemViewModel> WeakSystemViewModel = SystemViewModel.Pin();
-		ENiagaraSystemViewModelEditMode SystemEditMode = WeakSystemViewModel->GetEditMode();
-		if (SystemEditMode == ENiagaraSystemViewModelEditMode::EmitterAsset)
-		{
-			return;
-		}
-
 		const FScopedTransaction Transaction(FGenericCommands::Get().Paste->GetDescription());;
 		Graph->Modify();
 
@@ -342,7 +345,7 @@ void FNiagaraOverviewGraphViewModel::PasteNodes()
 					// Nodes pasted from emitters have no owning system, and will be invalid, so they are destroyed here instead.
 					FNiagaraEditorUtilities::InfoWithToastAndLog(LOCTEXT("PasteFromEmitterAsset", "Cannot paste emitters from emitter assets. Please use Add Emitter from the right click menu instead."));
 					OverviewNode->DestroyNode();
-				}					
+				}
 				else if	(OverviewNode->GetEmitterHandleGuid().IsValid())
 				{
 					FNiagaraSystemViewModel::FEmitterHandleToDuplicate EmitterHandleToDuplicate;
@@ -368,6 +371,13 @@ bool FNiagaraOverviewGraphViewModel::CanPasteNodes() const
 		return false;
 	}
 
+	// Cannot paste into emitter asset
+	if (GetSystemViewModel()->GetEditMode() == ENiagaraSystemViewModelEditMode::EmitterAsset)
+	{
+		FNiagaraEditorUtilities::InfoWithToastAndLog(LOCTEXT("PasteToEmitterAsset", "Cannot paste nodes into emitter assets."));
+		return false;
+	}
+
 	FString ClipboardContent;
 	FPlatformApplicationMisc::ClipboardPaste(ClipboardContent);
 
diff --git a/Engine/Plugins/FX/Niagara/Source/NiagaraEditorWidgets/Private/SNiagaraOverviewGraph.cpp b/Engine/Plugins/FX/Niagara/Source/NiagaraEditorWidgets/Private/SNiagaraOverviewGraph.cpp
--- a/Engine/Plugins/FX/Niagara/Source/NiagaraEditorWidgets/Private/SNiagaraOverviewGraph.cpp
+++ b/Engine/Plugins/FX/Niagara/Source/NiagaraEditorWidgets/Private/SNiagaraOverviewGraph.cpp
@@ -303,8 +303,12 @@ void SNiagaraOverviewGraph::OnCreateEmptyEmitter()
 
 void SNiagaraOverviewGraph::OnCreateComment()
 {
-	FNiagaraSchemaAction_NewComment CommentAction = FNiagaraSchemaAction_NewComment(GraphEditor);
-	CommentAction.PerformAction(ViewModel->GetGraph(), nullptr, GraphEditor->GetPasteLocation(), false);
+	// Emitter assets have a transient overview graph, so any created comments would also be transient. We skip creating these comments instead.
+	if (ViewModel->GetSystemViewModel()->GetEditMode() != ENiagaraSystemViewModelEditMode::EmitterAsset)
+	{
+		FNiagaraSchemaAction_NewComment CommentAction = FNiagaraSchemaAction_NewComment(GraphEditor);
+		CommentAction.PerformAction(ViewModel->GetGraph(), nullptr, GraphEditor->GetPasteLocation(), false);
+	}
 }
 
 void SNiagaraOverviewGraph::OnClearIsolated()
