- Simulation Stage Budget Validation Rule

diff --git a/Engine/Plugins/FX/Niagara/Source/Niagara/Private/NiagaraSimulationStageBase.cpp b/Engine/Plugins/FX/Niagara/Source/Niagara/Private/NiagaraSimulationStageBase.cpp
--- a/Engine/Plugins/FX/Niagara/Source/Niagara/Private/NiagaraSimulationStageBase.cpp
+++ b/Engine/Plugins/FX/Niagara/Source/Niagara/Private/NiagaraSimulationStageBase.cpp
@@ -215,6 +215,8 @@ void UNiagaraSimulationStageGeneric::PostEditChangeProperty(struct FPropertyChan
 	{
 		RequestRecompile();
 	}
+
+	OnChanged().Broadcast();
 }
 
 FVersionedNiagaraEmitterData* UNiagaraSimulationStageBase::GetEmitterData() const
diff --git a/Engine/Plugins/FX/Niagara/Source/NiagaraEditor/Private/NiagaraValidationRules.cpp b/Engine/Plugins/FX/Niagara/Source/NiagaraEditor/Private/NiagaraValidationRules.cpp
--- a/Engine/Plugins/FX/Niagara/Source/NiagaraEditor/Private/NiagaraValidationRules.cpp
+++ b/Engine/Plugins/FX/Niagara/Source/NiagaraEditor/Private/NiagaraValidationRules.cpp
@@ -7,6 +7,7 @@
 #include "NiagaraNodeFunctionCall.h"
 #include "NiagaraScriptSource.h"
 #include "NiagaraSettings.h"
+#include "NiagaraSimulationStageBase.h"
 #include "NiagaraSystem.h"
 #include "NiagaraSystemEditorData.h"
 #include "ViewModels/NiagaraEmitterHandleViewModel.h"
@@ -282,7 +283,6 @@ void UNiagaraValidationRule_InvalidEffectType::CheckValidity(TSharedPtr<FNiagara
 
 void UNiagaraValidationRule_LWC::CheckValidity(TSharedPtr<FNiagaraSystemViewModel> ViewModel, TArray<FNiagaraValidationResult>& Results)  const
 {
-
 	const UNiagaraSettings* Settings = GetDefault<UNiagaraSettings>();
 	UNiagaraSystem& System = ViewModel->GetSystem();
 	if (!System.SupportsLargeWorldCoordinates())
@@ -355,4 +355,70 @@ void UNiagaraValidationRule_LWC::CheckValidity(TSharedPtr<FNiagaraSystemViewMode
 	}
 }
 
-#undef LOCTEXT_NAMESPACE
\ No newline at end of file
+void UNiagaraValidationRule_SimulationStageBudget::CheckValidity(TSharedPtr<FNiagaraSystemViewModel> ViewModel, TArray<FNiagaraValidationResult>& OutResults) const
+{
+	UNiagaraSystem& System = ViewModel->GetSystem();
+	for (const TSharedRef<FNiagaraEmitterHandleViewModel>& EmitterHandleModel : ViewModel->GetEmitterHandleViewModels())
+	{
+		// Skip disabled
+		if ( EmitterHandleModel->GetIsEnabled() == false )
+		{	
+			continue;
+		}
+
+		// Simulation stages are GPU only currently
+		FVersionedNiagaraEmitterData* EmitterData = EmitterHandleModel.Get().GetEmitterHandle()->GetEmitterData();
+		if ( EmitterData->SimTarget != ENiagaraSimTarget::GPUComputeSim )
+		{
+			continue;
+		}
+
+		int32 TotalIterations = 0;
+		int32 TotalEnabledStages = 0;
+		for ( UNiagaraSimulationStageBase* SimStageBase : EmitterData->GetSimulationStages() )
+		{
+			UNiagaraSimulationStageGeneric* SimStage = Cast<UNiagaraSimulationStageGeneric>(SimStageBase);
+			if ( SimStage == nullptr || SimStage->bEnabled == false )
+			{
+				continue;
+			}
+
+			++TotalEnabledStages;
+			TotalIterations += SimStage->Iterations;
+			if ( bMaxIterationsPerStageEnabled && SimStage->Iterations > MaxIterationsPerStage )
+			{
+				UNiagaraStackEmitterPropertiesItem* EmitterProperties = GetStackEntry<UNiagaraStackEmitterPropertiesItem>(EmitterHandleModel.Get().GetEmitterStackViewModel());
+				OutResults.Emplace_GetRef(
+					ENiagaraValidationSeverity::Error,
+					FText::Format(LOCTEXT("SimStageTooManyIterationsFormat", "Simulation Stage '{0}' has too many iterations"), FText::FromName(SimStage->SimulationStageName)),
+					FText::Format(LOCTEXT("SimStageTooManyIterationsDetailedFormat", "Simulation Stage '{0}' has {1} iterations and we only allow {2}"), FText::FromName(SimStage->SimulationStageName), FText::AsNumber(SimStage->Iterations), FText::AsNumber(MaxIterationsPerStage)),
+					EmitterProperties
+				);
+			}
+		}
+
+		if ( bMaxTotalIterationsEnabled && TotalIterations > MaxTotalIterations )
+		{
+			UNiagaraStackEmitterPropertiesItem* EmitterProperties = GetStackEntry<UNiagaraStackEmitterPropertiesItem>(EmitterHandleModel.Get().GetEmitterStackViewModel());
+			OutResults.Emplace(
+				ENiagaraValidationSeverity::Error,
+				LOCTEXT("SimStageTooManyTotalIterationsFormat", "Emitter has too many total simulation stage iterations"),
+				FText::Format(LOCTEXT("SimStageTooManyTotalIterationsDetailedFormat", "Emitter has {0} total simulation stage iterations and we only allow {1}"), FText::AsNumber(TotalIterations), FText::AsNumber(MaxTotalIterations)),
+				EmitterProperties
+			);
+		}
+
+		if ( bMaxSimulationStagesEnabled && TotalEnabledStages > MaxSimulationStages )
+		{
+			UNiagaraStackEmitterPropertiesItem* EmitterProperties = GetStackEntry<UNiagaraStackEmitterPropertiesItem>(EmitterHandleModel.Get().GetEmitterStackViewModel());
+			OutResults.Emplace(
+				ENiagaraValidationSeverity::Error,
+				LOCTEXT("TooManySimStagesFormat", "Emitter has too many simulation stages"),
+				FText::Format(LOCTEXT("TooManySimStagesDetailedFormat", "Emitter has {0} simulation stages active and we only allow {1}"), FText::AsNumber(TotalEnabledStages), FText::AsNumber(MaxSimulationStages)),
+				EmitterProperties
+			);
+		}
+	}
+}
+
+#undef LOCTEXT_NAMESPACE
diff --git a/Engine/Plugins/FX/Niagara/Source/NiagaraEditor/Private/ViewModels/Stack/NiagaraStackSimulationStageGroup.cpp b/Engine/Plugins/FX/Niagara/Source/NiagaraEditor/Private/ViewModels/Stack/NiagaraStackSimulationStageGroup.cpp
--- a/Engine/Plugins/FX/Niagara/Source/NiagaraEditor/Private/ViewModels/Stack/NiagaraStackSimulationStageGroup.cpp
+++ b/Engine/Plugins/FX/Niagara/Source/NiagaraEditor/Private/ViewModels/Stack/NiagaraStackSimulationStageGroup.cpp
@@ -4,6 +4,7 @@
 #include "NiagaraEmitter.h"
 #include "NiagaraSimulationStageBase.h"
 #include "ViewModels/Stack/NiagaraStackObject.h"
+#include "ViewModels/Stack/NiagaraStackViewModel.h"
 #include "ViewModels/NiagaraScriptViewModel.h"
 #include "NiagaraScriptGraphViewModel.h"
 #include "ViewModels/NiagaraEmitterViewModel.h"
@@ -148,6 +149,9 @@ void UNiagaraStackSimulationStagePropertiesItem::RefreshChildrenInternal(const T
 
 void UNiagaraStackSimulationStagePropertiesItem::SimulationStagePropertiesChanged()
 {
+	FVersionedNiagaraEmitter VersionedEmitter = GetEmitterViewModel()->GetEmitter();
+	GetSystemViewModel()->GetEmitterHandleViewModelForEmitter(VersionedEmitter).Get()->GetEmitterStackViewModel()->RequestValidationUpdate();
+
 	bCanResetToBaseCache.Reset();
 }
 
diff --git a/Engine/Plugins/FX/Niagara/Source/NiagaraEditor/Public/NiagaraValidationRules.h b/Engine/Plugins/FX/Niagara/Source/NiagaraEditor/Public/NiagaraValidationRules.h
--- a/Engine/Plugins/FX/Niagara/Source/NiagaraEditor/Public/NiagaraValidationRules.h
+++ b/Engine/Plugins/FX/Niagara/Source/NiagaraEditor/Public/NiagaraValidationRules.h
@@ -80,4 +80,41 @@ class UNiagaraValidationRule_LWC : public UNiagaraValidationRule
 	GENERATED_BODY()
 public:
 	virtual void CheckValidity(TSharedPtr<FNiagaraSystemViewModel> ViewModel, TArray<FNiagaraValidationResult>& OutResults) const override;
-};
\ No newline at end of file
+};
+
+/** This validation rule can be used to enforce a budget on the number of simulation stages and the iterations that may execute. */
+UCLASS(Category = "Validation", DisplayName = "Simulation Stage Budget")
+class UNiagaraValidationRule_SimulationStageBudget : public UNiagaraValidationRule
+{
+	GENERATED_BODY()
+
+public:
+	virtual void CheckValidity(TSharedPtr<FNiagaraSystemViewModel> ViewModel, TArray<FNiagaraValidationResult>& OutResults) const override;
+
+	UPROPERTY(EditAnywhere, Category = Validation)
+	bool bMaxSimulationStagesEnabled = false;
+
+	UPROPERTY(EditAnywhere, Category = Validation)
+	bool bMaxIterationsPerStageEnabled = false;
+
+	UPROPERTY(EditAnywhere, Category = Validation)
+	bool bMaxTotalIterationsEnabled = false;
+
+	/** Maximum number of simulation stages allowed, where 0 means no simulation stages. */
+	UPROPERTY(EditAnywhere, Category = Validation, meta=(EditCondition="bMaxSimulationStagesEnabled"))
+	int32 MaxSimulationStages = 0;
+
+	/**
+	Maximum number of iterations a single stage is allowed to execute.
+	Note: Can only check across explicit counts, dynamic bindings will be ignored.
+	*/
+	UPROPERTY(EditAnywhere, Category = Validation, meta = (EditCondition = "bMaxIterationsPerStageEnabled"))
+	int32 MaxIterationsPerStage = 1;
+
+	/**
+	Maximum total iterations across all the enabled simulation stages.
+	Note: Can only check across explicit counts, dynamic bindings will be ignored.
+	*/
+	UPROPERTY(EditAnywhere, Category = Validation, meta = (EditCondition = "bMaxTotalIterationsEnabled"))
+	int32 MaxTotalIterations = 1;
+};
