- Fix for GetStackEntry being moved inside NiagaraValidation namespace

#rnx

diff --git a/Engine/Plugins/FX/Niagara/Source/NiagaraEditor/Private/NiagaraValidationRules.cpp b/Engine/Plugins/FX/Niagara/Source/NiagaraEditor/Private/NiagaraValidationRules.cpp
--- a/Engine/Plugins/FX/Niagara/Source/NiagaraEditor/Private/NiagaraValidationRules.cpp
+++ b/Engine/Plugins/FX/Niagara/Source/NiagaraEditor/Private/NiagaraValidationRules.cpp
@@ -387,7 +387,7 @@ void UNiagaraValidationRule_SimulationStageBudget::CheckValidity(TSharedPtr<FNia
 			TotalIterations += SimStage->Iterations;
 			if ( bMaxIterationsPerStageEnabled && SimStage->Iterations > MaxIterationsPerStage )
 			{
-				UNiagaraStackEmitterPropertiesItem* EmitterProperties = GetStackEntry<UNiagaraStackEmitterPropertiesItem>(EmitterHandleModel.Get().GetEmitterStackViewModel());
+				UNiagaraStackEmitterPropertiesItem* EmitterProperties = NiagaraValidation::GetStackEntry<UNiagaraStackEmitterPropertiesItem>(EmitterHandleModel.Get().GetEmitterStackViewModel());
 				OutResults.Emplace_GetRef(
 					ENiagaraValidationSeverity::Error,
 					FText::Format(LOCTEXT("SimStageTooManyIterationsFormat", "Simulation Stage '{0}' has too many iterations"), FText::FromName(SimStage->SimulationStageName)),
@@ -399,7 +399,7 @@ void UNiagaraValidationRule_SimulationStageBudget::CheckValidity(TSharedPtr<FNia
 
 		if ( bMaxTotalIterationsEnabled && TotalIterations > MaxTotalIterations )
 		{
-			UNiagaraStackEmitterPropertiesItem* EmitterProperties = GetStackEntry<UNiagaraStackEmitterPropertiesItem>(EmitterHandleModel.Get().GetEmitterStackViewModel());
+			UNiagaraStackEmitterPropertiesItem* EmitterProperties = NiagaraValidation::GetStackEntry<UNiagaraStackEmitterPropertiesItem>(EmitterHandleModel.Get().GetEmitterStackViewModel());
 			OutResults.Emplace(
 				ENiagaraValidationSeverity::Error,
 				LOCTEXT("SimStageTooManyTotalIterationsFormat", "Emitter has too many total simulation stage iterations"),
@@ -410,7 +410,7 @@ void UNiagaraValidationRule_SimulationStageBudget::CheckValidity(TSharedPtr<FNia
 
 		if ( bMaxSimulationStagesEnabled && TotalEnabledStages > MaxSimulationStages )
 		{
-			UNiagaraStackEmitterPropertiesItem* EmitterProperties = GetStackEntry<UNiagaraStackEmitterPropertiesItem>(EmitterHandleModel.Get().GetEmitterStackViewModel());
+			UNiagaraStackEmitterPropertiesItem* EmitterProperties = NiagaraValidation::GetStackEntry<UNiagaraStackEmitterPropertiesItem>(EmitterHandleModel.Get().GetEmitterStackViewModel());
 			OutResults.Emplace(
 				ENiagaraValidationSeverity::Error,
 				LOCTEXT("TooManySimStagesFormat", "Emitter has too many simulation stages"),
