- Add interpolated spawn information to GPU systems

#rnx

diff --git a/Engine/Plugins/FX/Niagara/Source/NiagaraEditor/Private/Commandlets/NiagaraSystemAuditCommandlet.cpp b/Engine/Plugins/FX/Niagara/Source/NiagaraEditor/Private/Commandlets/NiagaraSystemAuditCommandlet.cpp
--- a/Engine/Plugins/FX/Niagara/Source/NiagaraEditor/Private/Commandlets/NiagaraSystemAuditCommandlet.cpp
+++ b/Engine/Plugins/FX/Niagara/Source/NiagaraEditor/Private/Commandlets/NiagaraSystemAuditCommandlet.cpp
@@ -245,6 +245,9 @@ bool UNiagaraSystemAuditCommandlet::ProcessNiagaraSystems()
 				// Build information to write out
 				TStringBuilder<512> GpuEmitterBuilder;
 				GpuEmitterBuilder.Append(EmitterData->GetDebugSimName());
+
+				GpuEmitterBuilder.Append(EmitterData->bInterpolatedSpawning ? TEXT(",true") : TEXT(",false"));
+
 				for (int32 iQualityLevel=0; iQualityLevel < NiagaraSettings->QualityLevels.Num(); ++iQualityLevel)
 				{
 					const bool bEnabled = EmitterData->Platforms.IsEffectQualityEnabled(iQualityLevel);
@@ -410,7 +413,7 @@ void UNiagaraSystemAuditCommandlet::DumpResults()
 	if (NiagaraSystemsWithGPUEmitters.Num() > 0)
 	{
 		TStringBuilder<512> HeaderString;
-		HeaderString.Append(TEXT("Emitter Name"));
+		HeaderString.Append(TEXT("Emitter Name,Interpolated Spawn"));
 		for (const FText& QualityLevelName : GetDefault<UNiagaraSettings>()->QualityLevels)
 		{
 			HeaderString.Append(TEXT(","));
