- Fix ambiguous dot product
- Fix incorrect vector evaluations

#rnx

diff --git a/Engine/Plugins/FX/Niagara/Shaders/Private/NiagaraDataInterfaceSplineTemplate.ush b/Engine/Plugins/FX/Niagara/Shaders/Private/NiagaraDataInterfaceSplineTemplate.ush
--- a/Engine/Plugins/FX/Niagara/Shaders/Private/NiagaraDataInterfaceSplineTemplate.ush
+++ b/Engine/Plugins/FX/Niagara/Shaders/Private/NiagaraDataInterfaceSplineTemplate.ush
@@ -118,11 +118,11 @@ float3 EvaluateFindNearestPosition_{ParameterName}(float3 InPosition)
 		return 0.0f;
 	}
 	
-	float MinDistance = dot(InPosition, {ParameterName}_SplinePositionsLUT[0]);
+	float MinDistance = dot(InPosition, {ParameterName}_SplinePositionsLUT[0].xyz);
 	float KeyToNearest = 0.0f;
 	for (int Index = 1; Index <= {ParameterName}_MaxIndex; Index++)
 	{
-		const float Distance = dot(InPosition, {ParameterName}_SplinePositionsLUT[Index]);
+		const float Distance = dot(InPosition, {ParameterName}_SplinePositionsLUT[Index].xyz);
 		if (Distance < MinDistance)
 		{
 			MinDistance = Distance;
@@ -175,7 +175,7 @@ void SampleSplineUpVectorByUnitDistanceWS_{ParameterName}(float U, out float3 Up
 {
 	float4 Rotation = GetRotationAtDistanceAlongSpline_{ParameterName}(U * {ParameterName}_SplineLength);
 	UpVector = NiagaraQuatRotateVector(Rotation, float3(0,0,1));
-	UpVector = mul(float4(UpVector, 1.0f), {ParameterName}_SplineTransformRotationMat).xyz;
+	UpVector = mul(float4(UpVector, 0.0f), {ParameterName}_SplineTransformRotationMat).xyz;
 }
 
 void SampleSplineRightVectorByUnitDistance_{ParameterName}(float U, out float3 RightVector)
@@ -188,7 +188,7 @@ void SampleSplineRightVectorByUnitDistanceWS_{ParameterName}(float U, out float3
 {
 	float4 Rotation = GetRotationAtDistanceAlongSpline_{ParameterName}(U * {ParameterName}_SplineLength);
 	RightVector = NiagaraQuatRotateVector(Rotation, float3(0,1,0));
-	RightVector = mul(float4(RightVector, 1.0f), {ParameterName}_SplineTransformRotationMat).xyz;
+	RightVector = mul(float4(RightVector, 0.0f), {ParameterName}_SplineTransformRotationMat).xyz;
 }
 
 void SampleSplineTangentByUnitDistance_{ParameterName}(float U, out float3 Tangent)
@@ -199,7 +199,7 @@ void SampleSplineTangentByUnitDistance_{ParameterName}(float U, out float3 Tange
 void SampleSplineTangentByUnitDistanceWS_{ParameterName}(float U, out float3 Tangent)
 {
 	Tangent = EvaluateDerivativePosition_{ParameterName}(U * {ParameterName}_SplineLength);
-	Tangent = mul(Tangent, {ParameterName}_SplineTransformRotation);
+	Tangent = mul(float4(Tangent, 0.0f), {ParameterName}_SplineTransformRotationMat).xyz;
 }
 
 void GetSplineLocalToWorld_{ParameterName}(out float4x4 Transform)
