Added an info message to the emitter and script graphs when downstream assets are affected

diff --git a/Engine/Plugins/FX/Niagara/Source/Niagara/Public/NiagaraSettings.h b/Engine/Plugins/FX/Niagara/Source/Niagara/Public/NiagaraSettings.h
--- a/Engine/Plugins/FX/Niagara/Source/Niagara/Public/NiagaraSettings.h
+++ b/Engine/Plugins/FX/Niagara/Source/Niagara/Public/NiagaraSettings.h
@@ -102,6 +102,10 @@ class NIAGARA_API UNiagaraSettings : public UDeveloperSettings
 	UPROPERTY(config, EditAnywhere, Category = Viewport)
 	bool bSystemViewportInOrbitMode = true;
 
+	/** If true, then the emitter and script editors will show an info message how many downstream asset are affected by a change. Gathering this information for large asset graphs can delay the opening of the asset editors a bit. */
+	UPROPERTY(config, EditAnywhere, Category=Niagara)
+	bool bDisplayAffectedAssetStatsInEditor = true;
+
 	/** If true then the system editor will zoom to fit all emitters when opening an asset. */
 	UPROPERTY(config, EditAnywhere, Category = Niagara)
 	bool bAlwaysZoomToFitSystemGraph = true;
diff --git a/Engine/Plugins/FX/Niagara/Source/NiagaraEditor/Private/NiagaraEditorStyle.cpp b/Engine/Plugins/FX/Niagara/Source/NiagaraEditor/Private/NiagaraEditorStyle.cpp
--- a/Engine/Plugins/FX/Niagara/Source/NiagaraEditor/Private/NiagaraEditorStyle.cpp
+++ b/Engine/Plugins/FX/Niagara/Source/NiagaraEditor/Private/NiagaraEditorStyle.cpp
@@ -538,6 +538,7 @@ void FNiagaraEditorStyle::InitDropTarget()
 void FNiagaraEditorStyle::InitScriptGraph()
 {
 	Set("NiagaraEditor.ScriptGraph.SearchBorderColor", FLinearColor(.1f, .1f, .1f, 1.f));
+	Set("NiagaraEditor.ScriptGraph.AffectedAssetsWarningColor", FLinearColor(FColor(255, 184, 0)));
 }
 
 void FNiagaraEditorStyle::InitDebuggerStyle()
diff --git a/Engine/Plugins/FX/Niagara/Source/NiagaraEditor/Private/Toolkits/NiagaraScriptToolkit.cpp b/Engine/Plugins/FX/Niagara/Source/NiagaraEditor/Private/Toolkits/NiagaraScriptToolkit.cpp
--- a/Engine/Plugins/FX/Niagara/Source/NiagaraEditor/Private/Toolkits/NiagaraScriptToolkit.cpp
+++ b/Engine/Plugins/FX/Niagara/Source/NiagaraEditor/Private/Toolkits/NiagaraScriptToolkit.cpp
@@ -386,7 +386,7 @@ TSharedRef<SDockTab> FNiagaraScriptToolkit::SpawnTabNodeGraph( const FSpawnTabAr
 	return
 		SNew(SDockTab)
 		[
-			SAssignNew(NiagaraScriptGraphWidget, SNiagaraScriptGraph, ScriptViewModel->GetGraphViewModel())
+			SAssignNew(NiagaraScriptGraphWidget, SNiagaraScriptGraph, ScriptViewModel->GetGraphViewModel(), FAssetData(OriginalNiagaraScript.Script))
 			.GraphTitle(LOCTEXT("SpawnGraphTitle", "Script"))
 		];
 }
diff --git a/Engine/Plugins/FX/Niagara/Source/NiagaraEditor/Private/Toolkits/NiagaraSystemToolkit.cpp b/Engine/Plugins/FX/Niagara/Source/NiagaraEditor/Private/Toolkits/NiagaraSystemToolkit.cpp
--- a/Engine/Plugins/FX/Niagara/Source/NiagaraEditor/Private/Toolkits/NiagaraSystemToolkit.cpp
+++ b/Engine/Plugins/FX/Niagara/Source/NiagaraEditor/Private/Toolkits/NiagaraSystemToolkit.cpp
@@ -837,6 +837,19 @@ void FNiagaraSystemToolkit::RegisterToolbarTab(const TSharedRef<FTabManager>& In
 	FAssetEditorToolkit::RegisterTabSpawners(InTabManager);
 }
 
+FAssetData FNiagaraSystemToolkit::GetEditedAsset() const
+{
+	if (HasEmitter())
+	{
+		return FAssetData(Emitter);
+	}
+	if (HasSystem())
+	{
+		return FAssetData(System);
+	}
+	return FAssetData();
+}
+
 const TArray<UObject*>& FNiagaraSystemToolkit::GetObjectsBeingEdited() const
 {
 	return GetEditingObjects();
diff --git a/Engine/Plugins/FX/Niagara/Source/NiagaraEditor/Private/Toolkits/NiagaraSystemToolkit.h b/Engine/Plugins/FX/Niagara/Source/NiagaraEditor/Private/Toolkits/NiagaraSystemToolkit.h
--- a/Engine/Plugins/FX/Niagara/Source/NiagaraEditor/Private/Toolkits/NiagaraSystemToolkit.h
+++ b/Engine/Plugins/FX/Niagara/Source/NiagaraEditor/Private/Toolkits/NiagaraSystemToolkit.h
@@ -124,6 +124,8 @@ public:
 
 	bool HasSystem() const { return SystemToolkitMode == ESystemToolkitMode::System && System != nullptr; }
 
+	FAssetData GetEditedAsset() const; 
+
 	const TArray<UObject*>& GetObjectsBeingEdited() const;
 
 	
diff --git a/Engine/Plugins/FX/Niagara/Source/NiagaraEditor/Private/Toolkits/SystemToolkitModes/NiagaraSystemToolkitModeBase.cpp b/Engine/Plugins/FX/Niagara/Source/NiagaraEditor/Private/Toolkits/SystemToolkitModes/NiagaraSystemToolkitModeBase.cpp
--- a/Engine/Plugins/FX/Niagara/Source/NiagaraEditor/Private/Toolkits/SystemToolkitModes/NiagaraSystemToolkitModeBase.cpp
+++ b/Engine/Plugins/FX/Niagara/Source/NiagaraEditor/Private/Toolkits/SystemToolkitModes/NiagaraSystemToolkitModeBase.cpp
@@ -667,15 +667,16 @@ TSharedRef<SDockTab> FNiagaraSystemToolkitModeBase::SpawnTab_MessageLog(const FS
 
 TSharedRef<SDockTab> FNiagaraSystemToolkitModeBase::SpawnTab_SystemOverview(const FSpawnTabArgs& Args)
 {
-	if(!SystemToolkit.Pin()->GetSystemOverview().IsValid())
+	TSharedPtr<FNiagaraSystemToolkit, ESPMode::ThreadSafe> NiagaraSystemToolkit = SystemToolkit.Pin();
+	if(!NiagaraSystemToolkit->GetSystemOverview().IsValid())
 	{
-		SystemToolkit.Pin()->SetSystemOverview(FNiagaraEditorModule::Get().GetWidgetProvider()->CreateSystemOverview(SystemToolkit.Pin()->GetSystemViewModel().ToSharedRef()));
+		NiagaraSystemToolkit->SetSystemOverview(FNiagaraEditorModule::Get().GetWidgetProvider()->CreateSystemOverview(NiagaraSystemToolkit->GetSystemViewModel().ToSharedRef(), NiagaraSystemToolkit->GetEditedAsset()));
 	}
 	
 	TSharedRef<SDockTab> SpawnedTab = SNew(SDockTab)
 		.Label(LOCTEXT("SystemOverviewTabLabel", "System Overview"))
 		[
-			SystemToolkit.Pin()->GetSystemOverview().ToSharedRef()
+			NiagaraSystemToolkit->GetSystemOverview().ToSharedRef()
 		];
 
 
diff --git a/Engine/Plugins/FX/Niagara/Source/NiagaraEditor/Private/Widgets/SNiagaraScriptGraph.cpp b/Engine/Plugins/FX/Niagara/Source/NiagaraEditor/Private/Widgets/SNiagaraScriptGraph.cpp
--- a/Engine/Plugins/FX/Niagara/Source/NiagaraEditor/Private/Widgets/SNiagaraScriptGraph.cpp
+++ b/Engine/Plugins/FX/Niagara/Source/NiagaraEditor/Private/Widgets/SNiagaraScriptGraph.cpp
@@ -28,12 +28,16 @@
 #include "Framework/Application/SlateApplication.h"
 #include "GraphEditorActions.h"
 #include "EditorFontGlyphs.h"
+#include "NiagaraSettings.h"
+#include "AssetRegistry/AssetRegistryModule.h"
 #include "Widgets/SNiagaraGraphActionMenu.h"
 
 #define LOCTEXT_NAMESPACE "NiagaraScriptGraph"
 
-void SNiagaraScriptGraph::Construct(const FArguments& InArgs, TSharedRef<FNiagaraScriptGraphViewModel> InViewModel)
+void SNiagaraScriptGraph::Construct(const FArguments& InArgs, TSharedRef<FNiagaraScriptGraphViewModel> InViewModel, const FAssetData& InEditedAsset)
 {
+	EditedAsset = InEditedAsset;
+	
 	UpdateViewModel(InViewModel);
 	bUpdatingGraphSelectionFromViewModel = false;
 
@@ -53,6 +57,13 @@ void SNiagaraScriptGraph::Construct(const FArguments& InArgs, TSharedRef<FNiagar
 	];
 
 	CurrentFocusedSearchMatchIndex = 0;
+
+	AddAssetListeners();
+}
+
+SNiagaraScriptGraph::~SNiagaraScriptGraph()
+{
+	ClearListeners();
 }
 
 void SNiagaraScriptGraph::UpdateViewModel(TSharedRef<FNiagaraScriptGraphViewModel> InNewModel)
@@ -95,25 +106,46 @@ TSharedRef<SGraphEditor> SNiagaraScriptGraph::ConstructGraphEditor()
 			SNew(SOverlay)
 			+SOverlay::Slot()
 			[
-				SNew(SHorizontalBox)
-				+ SHorizontalBox::Slot()
-				.AutoWidth()
-				.Padding(0.0f, 0.0f, 3.0f, 0.0f)
+				SNew(SVerticalBox)
+				+ SVerticalBox::Slot()
+				.HAlign(HAlign_Fill)
+				.AutoHeight()
 				[
-					SNew(SErrorText)
-					.Visibility(ViewModel.ToSharedRef(), &FNiagaraScriptGraphViewModel::GetGraphErrorTextVisible)
-					.BackgroundColor(ViewModel.ToSharedRef(), &FNiagaraScriptGraphViewModel::GetGraphErrorColor)
-					.ToolTipText(ViewModel.ToSharedRef(), &FNiagaraScriptGraphViewModel::GetGraphErrorMsgToolTip)
-					.ErrorText(ViewModel->GetGraphErrorText())
+					SNew(SHorizontalBox)
+					+ SHorizontalBox::Slot()
+					.AutoWidth()
+					.Padding(0.0f, 0.0f, 3.0f, 0.0f)
+					[
+						SNew(SErrorText)
+						.Visibility(ViewModel.ToSharedRef(), &FNiagaraScriptGraphViewModel::GetGraphErrorTextVisible)
+						.BackgroundColor(ViewModel.ToSharedRef(), &FNiagaraScriptGraphViewModel::GetGraphErrorColor)
+						.ToolTipText(ViewModel.ToSharedRef(), &FNiagaraScriptGraphViewModel::GetGraphErrorMsgToolTip)
+						.ErrorText(ViewModel->GetGraphErrorText())
+					]
+					+SHorizontalBox::Slot()
+					.VAlign(VAlign_Center)
+					[
+						SNew(STextBlock)
+						.Text(ViewModel.ToSharedRef(), &FNiagaraScriptGraphViewModel::GetDisplayName)
+						.TextStyle(FAppStyle::Get(), TEXT("GraphBreadcrumbButtonText"))
+						.Justification(ETextJustify::Center)
+						.Visibility(bShowHeader ? EVisibility::Visible : EVisibility::Collapsed)
+					]
 				]
-				+SHorizontalBox::Slot()
-				.VAlign(VAlign_Center)
+				// optional info for downstream assets
+				+ SVerticalBox::Slot()
 				[
-					SNew(STextBlock)
-					.Text(ViewModel.ToSharedRef(), &FNiagaraScriptGraphViewModel::GetDisplayName)
-					.TextStyle(FAppStyle::Get(), TEXT("GraphBreadcrumbButtonText"))
-					.Justification(ETextJustify::Center)
-					.Visibility(bShowHeader ? EVisibility::Visible : EVisibility::Collapsed)
+					SNew(SBorder)
+					.BorderImage(FAppStyle::GetBrush(TEXT("Graph.TitleBackground")))
+					.ColorAndOpacity(this, &SNiagaraScriptGraph::GetScriptSubheaderColor)
+					.Visibility(this, &SNiagaraScriptGraph::GetScriptSubheaderVisibility)
+					.HAlign(HAlign_Fill)
+					[
+						SNew(STextBlock)
+						.Text(this, &SNiagaraScriptGraph::GetScriptSubheaderText)
+						.TextStyle(FAppStyle::Get(), TEXT("GraphBreadcrumbButtonText"))
+						.Justification(ETextJustify::Center)
+					]
 				]
 			]
 			// Search Box
@@ -646,6 +678,110 @@ TOptional<SSearchBox::FSearchResultData> SNiagaraScriptGraph::GetSearchResultDat
 	return TOptional<SSearchBox::FSearchResultData>({ CurrentSearchResults.Num(), CurrentFocusedSearchMatchIndex + 1 });
 }
 
+FText SNiagaraScriptGraph::GetScriptSubheaderText() const
+{
+	return FText::Format(LOCTEXT("EmitterSubheaderText", "Note: editing this script will affect {0} dependent assets (across all versions)!"), GetScriptAffectedAssets());
+}
+
+EVisibility SNiagaraScriptGraph::GetScriptSubheaderVisibility() const
+{
+	return bShowHeader && (GetScriptAffectedAssets() > 1) ? EVisibility::Visible : EVisibility::Collapsed;
+}
+
+FLinearColor SNiagaraScriptGraph::GetScriptSubheaderColor() const
+{
+	return GetScriptAffectedAssets() >= 5 ? FNiagaraEditorStyle::Get().GetColor("NiagaraEditor.ScriptGraph.AffectedAssetsWarningColor") : FLinearColor::White;
+}
+
+int32 SNiagaraScriptGraph::GetScriptAffectedAssets() const
+{
+	if (!EditedAsset.IsValid())
+	{
+		return 0;
+	}
+	const UNiagaraSettings* Settings = GetDefault<UNiagaraSettings>();
+	if (Settings->bDisplayAffectedAssetStatsInEditor == false)
+	{
+		return 0;
+	}
+	
+	if (!ScriptAffectedAssets.IsSet())
+	{
+		FAssetRegistryModule& AssetRegistryModule = FModuleManager::LoadModuleChecked<FAssetRegistryModule>(TEXT("AssetRegistry"));
+		IAssetRegistry& AssetRegistry = AssetRegistryModule.Get();
+		if (AssetRegistry.IsLoadingAssets())
+		{
+			// We are still discovering assets, wait a bit
+			return 0;
+		}
+
+		int32 Count = 0;
+		if (EditedAsset.IsValid())
+		{
+			TSet<FName> SeenObjects;
+			TArray<FAssetData> AssetsToCheck;
+			AssetsToCheck.Add(EditedAsset);
+
+			// search for assets that reference this script
+			while (AssetsToCheck.Num() > 0)
+			{
+				FAssetData AssetToCheck = AssetsToCheck[0];
+				AssetsToCheck.RemoveAtSwap(0);
+				if (SeenObjects.Contains(AssetToCheck.ObjectPath))
+				{
+					continue;
+				}
+				SeenObjects.Add(AssetToCheck.ObjectPath);
+
+				if (AssetToCheck.GetClass() == UNiagaraSystem::StaticClass())
+				{
+					Count++;
+				}
+				else if (AssetToCheck.GetClass() == UNiagaraEmitter::StaticClass() || AssetToCheck.GetClass() == UNiagaraScript::StaticClass())
+				{
+					Count++;
+					TArray<FName> Referencers;
+					AssetRegistry.GetReferencers(AssetToCheck.PackageName, Referencers);
+					for (FName& Referencer : Referencers)
+					{
+						AssetRegistry.GetAssetsByPackageName(Referencer, AssetsToCheck);
+					}
+				}
+			}
+			Count--; // remove one for our own asset
+		}
+		ScriptAffectedAssets = Count;
+	}
+	return ScriptAffectedAssets.Get(0);
+}
+
+void SNiagaraScriptGraph::ResetAssetCount(const FAssetData&)
+{
+	ScriptAffectedAssets.Reset();
+}
+
+void SNiagaraScriptGraph::AddAssetListeners()
+{
+	if (EditedAsset.IsValid())
+	{
+		IAssetRegistry& AssetRegistry = FModuleManager::LoadModuleChecked<FAssetRegistryModule>(TEXT("AssetRegistry")).Get();
+		AssetRegistry.OnAssetUpdated().AddRaw(this, &SNiagaraScriptGraph::ResetAssetCount);
+		AssetRegistry.OnAssetAdded().AddRaw(this, &SNiagaraScriptGraph::ResetAssetCount);
+		AssetRegistry.OnAssetRemoved().AddRaw(this, &SNiagaraScriptGraph::ResetAssetCount);
+	}
+}
+
+void SNiagaraScriptGraph::ClearListeners()
+{
+	if (EditedAsset.IsValid())
+	{
+		IAssetRegistry& AssetRegistry = FModuleManager::LoadModuleChecked<FAssetRegistryModule>(TEXT("AssetRegistry")).Get();
+		AssetRegistry.OnAssetUpdated().RemoveAll(this);
+		AssetRegistry.OnAssetAdded().RemoveAll(this);
+		AssetRegistry.OnAssetRemoved().RemoveAll(this);
+	}
+}
+
 FReply SNiagaraScriptGraph::CloseGraphSearchBoxPressed()
 {
 	bGraphSearchBoxActive = false;
diff --git a/Engine/Plugins/FX/Niagara/Source/NiagaraEditor/Public/NiagaraEditorModule.h b/Engine/Plugins/FX/Niagara/Source/NiagaraEditor/Public/NiagaraEditorModule.h
--- a/Engine/Plugins/FX/Niagara/Source/NiagaraEditor/Public/NiagaraEditorModule.h
+++ b/Engine/Plugins/FX/Niagara/Source/NiagaraEditor/Public/NiagaraEditorModule.h
@@ -54,7 +54,7 @@ class NIAGARAEDITOR_API INiagaraEditorWidgetProvider
 {
 public:
 	virtual TSharedRef<SWidget> CreateStackView(UNiagaraStackViewModel& StackViewModel) const = 0;
-	virtual TSharedRef<SWidget> CreateSystemOverview(TSharedRef<FNiagaraSystemViewModel> SystemViewModel) const = 0;
+	virtual TSharedRef<SWidget> CreateSystemOverview(TSharedRef<FNiagaraSystemViewModel> SystemViewModel, const FAssetData& EditedAsset) const = 0;
 	virtual TSharedRef<SWidget> CreateStackIssueIcon(UNiagaraStackViewModel& StackViewModel, UNiagaraStackEntry& StackEntry) const = 0;
 
 	virtual TSharedRef<SWidget> CreateScriptScratchPadManager(UNiagaraScratchPadViewModel& ScriptScratchPadViewModel) const = 0;
diff --git a/Engine/Plugins/FX/Niagara/Source/NiagaraEditor/Public/Widgets/SNiagaraScriptGraph.h b/Engine/Plugins/FX/Niagara/Source/NiagaraEditor/Public/Widgets/SNiagaraScriptGraph.h
--- a/Engine/Plugins/FX/Niagara/Source/NiagaraEditor/Public/Widgets/SNiagaraScriptGraph.h
+++ b/Engine/Plugins/FX/Niagara/Source/NiagaraEditor/Public/Widgets/SNiagaraScriptGraph.h
@@ -28,7 +28,9 @@ public:
 		SLATE_ARGUMENT(bool, ShowHeader)
 	SLATE_END_ARGS();
 
-	NIAGARAEDITOR_API void Construct(const FArguments& InArgs, TSharedRef<FNiagaraScriptGraphViewModel> InViewModel);
+	NIAGARAEDITOR_API void Construct(const FArguments& InArgs, TSharedRef<FNiagaraScriptGraphViewModel> InViewModel, const FAssetData& InEditedAsset = FAssetData());
+
+	NIAGARAEDITOR_API virtual ~SNiagaraScriptGraph() override;
 
 	const TSharedPtr<SGraphEditor> GetGraphEditor() { return GraphEditor; };
 
@@ -88,6 +90,14 @@ private:
 	FText GetCurrentSearchText() const { return CurrentSearchText; };
 
 	EVisibility GetGraphSearchBoxVisibility() const { return bGraphSearchBoxActive ? EVisibility::Visible : EVisibility::Collapsed; };
+	FText GetScriptSubheaderText() const;
+	EVisibility GetScriptSubheaderVisibility() const;
+	FLinearColor GetScriptSubheaderColor() const;
+	
+	int32 GetScriptAffectedAssets() const;
+	void ResetAssetCount(const FAssetData& InAssetData);
+	void AddAssetListeners();
+	void ClearListeners();
 
 	FReply CloseGraphSearchBoxPressed();
 
@@ -133,6 +143,9 @@ private:
 	bool bGraphSearchBoxActive;
 	bool bShowHeader;
 
+	mutable TOptional<int32> ScriptAffectedAssets;
+	FAssetData EditedAsset;
+
 private:
 	// action menu data
 	UEdGraph* GraphObj;
diff --git a/Engine/Plugins/FX/Niagara/Source/NiagaraEditorWidgets/Private/NiagaraEditorWidgetsModule.cpp b/Engine/Plugins/FX/Niagara/Source/NiagaraEditorWidgets/Private/NiagaraEditorWidgetsModule.cpp
--- a/Engine/Plugins/FX/Niagara/Source/NiagaraEditorWidgets/Private/NiagaraEditorWidgetsModule.cpp
+++ b/Engine/Plugins/FX/Niagara/Source/NiagaraEditorWidgets/Private/NiagaraEditorWidgetsModule.cpp
@@ -197,9 +197,9 @@ TSharedRef<SWidget> FNiagaraEditorWidgetsModule::FNiagaraEditorWidgetProvider::C
 	return SNew(SNiagaraStack, &StackViewModel);
 }
 
-TSharedRef<SWidget> FNiagaraEditorWidgetsModule::FNiagaraEditorWidgetProvider::CreateSystemOverview(TSharedRef<FNiagaraSystemViewModel> SystemViewModel) const
+TSharedRef<SWidget> FNiagaraEditorWidgetsModule::FNiagaraEditorWidgetProvider::CreateSystemOverview(TSharedRef<FNiagaraSystemViewModel> SystemViewModel, const FAssetData& EditedAsset) const
 {
-	return SNew(SNiagaraOverviewGraph, SystemViewModel->GetOverviewGraphViewModel().ToSharedRef());
+	return SNew(SNiagaraOverviewGraph, SystemViewModel->GetOverviewGraphViewModel().ToSharedRef(), EditedAsset);
 }
 
 TSharedRef<SWidget> FNiagaraEditorWidgetsModule::FNiagaraEditorWidgetProvider::CreateStackIssueIcon(UNiagaraStackViewModel& StackViewModel, UNiagaraStackEntry& StackEntry) const
diff --git a/Engine/Plugins/FX/Niagara/Source/NiagaraEditorWidgets/Private/NiagaraEditorWidgetsStyle.cpp b/Engine/Plugins/FX/Niagara/Source/NiagaraEditorWidgets/Private/NiagaraEditorWidgetsStyle.cpp
--- a/Engine/Plugins/FX/Niagara/Source/NiagaraEditorWidgets/Private/NiagaraEditorWidgetsStyle.cpp
+++ b/Engine/Plugins/FX/Niagara/Source/NiagaraEditorWidgets/Private/NiagaraEditorWidgetsStyle.cpp
@@ -230,6 +230,7 @@ FNiagaraEditorWidgetsStyle::FNiagaraEditorWidgetsStyle() : FSlateStyleSet("Niaga
 	Set("NiagaraEditor.SystemOverview.CheckBoxBorder", new CORE_BOX_BRUSH("Common/GroupBorderLight", FMargin(4.0f / 16.0f)));
 	Set("NiagaraEditor.SystemOverview.NodeBackgroundBorder", new BOX_BRUSH("Icons/SystemOverviewNodeBackground", FMargin(1.0f / 4.0f)));
 	Set("NiagaraEditor.SystemOverview.NodeBackgroundColor", FLinearColor(FColor(48, 48, 48)));
+	Set("NiagaraEditor.SystemOverview.AffectedAssetsWarningColor", FLinearColor(FColor(255, 184, 0)));
 
 	Set("NiagaraEditor.Scalability.EffectType.bAllowCullingForLocalPlayers", new IMAGE_BRUSH_SVG("Icons/Scalability/CullDistance_20", Icon20x20));
 	Set("NiagaraEditor.Scalability.EffectType.UpdateFrequency", new IMAGE_BRUSH_SVG("Icons/Scalability/Update", Icon20x20));
diff --git a/Engine/Plugins/FX/Niagara/Source/NiagaraEditorWidgets/Private/SNiagaraOverviewGraph.cpp b/Engine/Plugins/FX/Niagara/Source/NiagaraEditorWidgets/Private/SNiagaraOverviewGraph.cpp
--- a/Engine/Plugins/FX/Niagara/Source/NiagaraEditorWidgets/Private/SNiagaraOverviewGraph.cpp
+++ b/Engine/Plugins/FX/Niagara/Source/NiagaraEditorWidgets/Private/SNiagaraOverviewGraph.cpp
@@ -6,12 +6,10 @@
 #include "ViewModels/NiagaraSystemViewModel.h"
 #include "ViewModels/NiagaraEmitterHandleViewModel.h"
 #include "ViewModels/NiagaraOverviewGraphViewModel.h"
-#include "ViewModels/Stack/NiagaraStackViewModel.h"
 #include "NiagaraSystem.h"
 #include "NiagaraOverviewNode.h"
 #include "NiagaraSystemEditorData.h"
 #include "NiagaraObjectSelection.h"
-#include "ViewModels/NiagaraOverviewGraphViewModel.h"
 #include "EdGraphSchema_Niagara.h"
 #include "GraphEditAction.h"
 #include "NiagaraEditorCommands.h"
@@ -26,14 +24,13 @@
 #include "EdGraph/EdGraph.h"
 #include "EdGraph/EdGraphNode.h"
 #include "Framework/Commands/GenericCommands.h"
-#include "Widgets/Input/SCheckBox.h"
 #include "NiagaraEditor/Private/SNiagaraAssetPickerList.h"
 #include "Widgets/SItemSelector.h"
 #include "SNiagaraOverviewGraphTitleBar.h"
 
 #define LOCTEXT_NAMESPACE "NiagaraOverviewGraph"
 
-void SNiagaraOverviewGraph::Construct(const FArguments& InArgs, TSharedRef<FNiagaraOverviewGraphViewModel> InViewModel)
+void SNiagaraOverviewGraph::Construct(const FArguments& InArgs, TSharedRef<FNiagaraOverviewGraphViewModel> InViewModel, const FAssetData& InEditedAsset)
 {
 	ViewModel = InViewModel;
 	ViewModel->GetNodeSelection()->OnSelectedObjectsChanged().AddSP(this, &SNiagaraOverviewGraph::ViewModelSelectionChanged);
@@ -65,7 +62,7 @@ void SNiagaraOverviewGraph::Construct(const FArguments& InArgs, TSharedRef<FNiag
 		AppearanceInfo.CornerText = LOCTEXT("NiagaraOverview_AppearanceCornerTextGeneric", "NIAGARA");
 	}
 	
-	TSharedRef<SWidget> TitleBarWidget = SNew(SNiagaraOverviewGraphTitleBar, ViewModel->GetSystemViewModel()).Visibility(EVisibility::SelfHitTestInvisible);
+	TSharedRef<SWidget> TitleBarWidget = SNew(SNiagaraOverviewGraphTitleBar, ViewModel->GetSystemViewModel(), InEditedAsset).Visibility(EVisibility::SelfHitTestInvisible);
 
 	TSharedRef<FUICommandList> Commands = ViewModel->GetCommands();
 	Commands->MapAction(
diff --git a/Engine/Plugins/FX/Niagara/Source/NiagaraEditorWidgets/Private/SNiagaraOverviewGraph.h b/Engine/Plugins/FX/Niagara/Source/NiagaraEditorWidgets/Private/SNiagaraOverviewGraph.h
--- a/Engine/Plugins/FX/Niagara/Source/NiagaraEditorWidgets/Private/SNiagaraOverviewGraph.h
+++ b/Engine/Plugins/FX/Niagara/Source/NiagaraEditorWidgets/Private/SNiagaraOverviewGraph.h
@@ -18,7 +18,7 @@ public:
 	SLATE_BEGIN_ARGS(SNiagaraOverviewGraph) {}
 	SLATE_END_ARGS()
 
-	void Construct(const FArguments& InArgs, TSharedRef<FNiagaraOverviewGraphViewModel> InViewModel);
+	void Construct(const FArguments& InArgs, TSharedRef<FNiagaraOverviewGraphViewModel> InViewModel, const FAssetData& InEditedAsset);
 
 	virtual void Tick(const FGeometry& AllottedGeometry, const double InCurrentTime, const float InDeltaTime) override;
 
@@ -60,7 +60,6 @@ private:
 
 	void OnDistributeNodesH();
 	void OnDistributeNodesV();
-
 private:
 	TSharedPtr<FNiagaraOverviewGraphViewModel> ViewModel;
 	TSharedPtr<SGraphEditor> GraphEditor;
diff --git a/Engine/Plugins/FX/Niagara/Source/NiagaraEditorWidgets/Private/SNiagaraOverviewGraphTitleBar.cpp b/Engine/Plugins/FX/Niagara/Source/NiagaraEditorWidgets/Private/SNiagaraOverviewGraphTitleBar.cpp
--- a/Engine/Plugins/FX/Niagara/Source/NiagaraEditorWidgets/Private/SNiagaraOverviewGraphTitleBar.cpp
+++ b/Engine/Plugins/FX/Niagara/Source/NiagaraEditorWidgets/Private/SNiagaraOverviewGraphTitleBar.cpp
@@ -1,21 +1,31 @@
 // Copyright Epic Games, Inc. All Rights Reserved.
 
 #include "SNiagaraOverviewGraphTitleBar.h"
-#include "NiagaraEditorStyle.h"
+
+#include "NiagaraEditorWidgetsStyle.h"
 #include "NiagaraSettings.h"
 #include "ViewModels/NiagaraOverviewGraphViewModel.h"
 #include "SNiagaraScalabilityPreviewSettings.h"
 #include "SNiagaraSystemEffectTypeBar.h"
+#include "AssetRegistry/AssetRegistryModule.h"
+#include "ViewModels/NiagaraEmitterHandleViewModel.h"
 
 #define LOCTEXT_NAMESPACE "NiagaraOverviewGraphTitleBar"
 
-void SNiagaraOverviewGraphTitleBar::Construct(const FArguments& InArgs, TSharedRef<FNiagaraSystemViewModel> InSystemViewModel)
+SNiagaraOverviewGraphTitleBar::~SNiagaraOverviewGraphTitleBar()
+{
+	ClearListeners();
+}
+
+void SNiagaraOverviewGraphTitleBar::Construct(const FArguments& InArgs, TSharedRef<FNiagaraSystemViewModel> InSystemViewModel, const FAssetData& InEditedAsset)
 {
 	SystemViewModel = InSystemViewModel;
+	EditedAsset = InEditedAsset;
 
 	bScalabilityModeActive = SystemViewModel.Pin()->GetScalabilityViewModel()->IsActive();
 	SystemViewModel.Pin()->GetScalabilityViewModel()->OnScalabilityModeChanged().AddSP(this, &SNiagaraOverviewGraphTitleBar::OnUpdateScalabilityMode);
 
+	AddAssetListeners();
 	RebuildWidget();
 }
 
@@ -24,7 +34,7 @@ void SNiagaraOverviewGraphTitleBar::RebuildWidget()
 	ContainerWidget = SNew(SVerticalBox);
 
 	// we only allow the Niagara system effect type bar on system assets
-	if(bScalabilityModeActive && SystemViewModel.IsValid() && SystemViewModel.Pin()->GetEditMode() == ENiagaraSystemViewModelEditMode::SystemAsset)
+	if (bScalabilityModeActive && SystemViewModel.IsValid() && SystemViewModel.Pin()->GetEditMode() == ENiagaraSystemViewModelEditMode::SystemAsset)
 	{
 		ContainerWidget->AddSlot()
 		.HAlign(HAlign_Fill)
@@ -53,6 +63,26 @@ void SNiagaraOverviewGraphTitleBar::RebuildWidget()
 			.Justification(ETextJustify::Center)
 		]
 	];
+
+	if (SystemViewModel.IsValid() && SystemViewModel.Pin()->GetEditMode() == ENiagaraSystemViewModelEditMode::EmitterAsset)
+	{
+		ContainerWidget->AddSlot()
+		.HAlign(HAlign_Fill)
+		.AutoHeight()
+		[
+			SNew(SBorder)
+			.BorderImage(FAppStyle::GetBrush(TEXT("Graph.TitleBackground")))
+			.ColorAndOpacity(this, &SNiagaraOverviewGraphTitleBar::GetEmitterSubheaderColor)
+			.Visibility(this, &SNiagaraOverviewGraphTitleBar::GetEmitterSubheaderVisibility)
+			.HAlign(HAlign_Fill)
+			[
+				SNew(STextBlock)
+				.Text(this, &SNiagaraOverviewGraphTitleBar::GetEmitterSubheaderText)
+				.TextStyle(FAppStyle::Get(), TEXT("GraphBreadcrumbButtonText"))
+				.Justification(ETextJustify::Center)
+			]
+		];
+	}
 	
 	if(bScalabilityModeActive)
 	{
@@ -71,20 +101,115 @@ void SNiagaraOverviewGraphTitleBar::RebuildWidget()
 	];
 }
 
-// FReply SNiagaraOverviewGraphTitleBar::OnMouseButtonDown(const FGeometry& MyGeometry, const FPointerEvent& MouseEvent)
-// {
-// 	if(OverviewGraph.IsValid())
-// 	{
-// 		//OverviewGraph->g
-// 	}
-//
-// 	return FReply::Handled();
-// }
-
 void SNiagaraOverviewGraphTitleBar::OnUpdateScalabilityMode(bool bActive)
 {
 	bScalabilityModeActive = bActive;
 	RebuildWidget();
 }
 
+FText SNiagaraOverviewGraphTitleBar::GetEmitterSubheaderText() const
+{
+	return FText::Format(LOCTEXT("EmitterSubheaderText", "Note: editing this emitter will affect {0} child emitters and systems (across all versions)!"), GetEmitterAffectedAssets());
+}
+
+EVisibility SNiagaraOverviewGraphTitleBar::GetEmitterSubheaderVisibility() const
+{
+	return GetEmitterAffectedAssets() > 1 ? EVisibility::Visible : EVisibility::Collapsed;
+}
+
+FLinearColor SNiagaraOverviewGraphTitleBar::GetEmitterSubheaderColor() const
+{
+	return GetEmitterAffectedAssets() >= 5 ? FNiagaraEditorWidgetsStyle::Get().GetColor("NiagaraEditor.SystemOverview.AffectedAssetsWarningColor") : FLinearColor::White;
+}
+
+void SNiagaraOverviewGraphTitleBar::ResetAssetCount(const FAssetData&)
+{
+	EmitterAffectedAssets.Reset();
+}
+
+void SNiagaraOverviewGraphTitleBar::AddAssetListeners()
+{
+	if (EditedAsset.IsValid())
+	{
+		IAssetRegistry& AssetRegistry = FModuleManager::LoadModuleChecked<FAssetRegistryModule>(TEXT("AssetRegistry")).Get();
+		AssetRegistry.OnAssetUpdated().AddRaw(this, &SNiagaraOverviewGraphTitleBar::ResetAssetCount);
+		AssetRegistry.OnAssetAdded().AddRaw(this, &SNiagaraOverviewGraphTitleBar::ResetAssetCount);
+		AssetRegistry.OnAssetRemoved().AddRaw(this, &SNiagaraOverviewGraphTitleBar::ResetAssetCount);
+	}
+}
+
+void SNiagaraOverviewGraphTitleBar::ClearListeners()
+{
+	if (EditedAsset.IsValid())
+	{
+		IAssetRegistry& AssetRegistry = FModuleManager::LoadModuleChecked<FAssetRegistryModule>(TEXT("AssetRegistry")).Get();
+		AssetRegistry.OnAssetUpdated().RemoveAll(this);
+		AssetRegistry.OnAssetAdded().RemoveAll(this);
+		AssetRegistry.OnAssetRemoved().RemoveAll(this);
+	}
+}
+
+int32 SNiagaraOverviewGraphTitleBar::GetEmitterAffectedAssets() const
+{
+	if (!EditedAsset.IsValid())
+	{
+		return 0;
+	}
+	
+	const UNiagaraSettings* Settings = GetDefault<UNiagaraSettings>();
+	if (Settings->bDisplayAffectedAssetStatsInEditor == false)
+	{
+		return 0;
+	}
+	
+	if (!EmitterAffectedAssets.IsSet())
+	{
+		FAssetRegistryModule& AssetRegistryModule = FModuleManager::LoadModuleChecked<FAssetRegistryModule>(TEXT("AssetRegistry"));
+		IAssetRegistry& AssetRegistry = AssetRegistryModule.Get();
+		if (AssetRegistry.IsLoadingAssets())
+		{
+			// We are still discovering assets, wait a bit
+			return 0;
+		}
+
+		int32 Count = 0;
+		if (EditedAsset.IsValid())
+		{
+			TSet<FName> SeenObjects;
+			TArray<FAssetData> AssetsToCheck;
+			AssetsToCheck.Add(EditedAsset);
+
+			// search for assets that reference this emitter
+			while (AssetsToCheck.Num() > 0)
+			{
+				FAssetData AssetToCheck = AssetsToCheck[0];
+				AssetsToCheck.RemoveAtSwap(0);
+				if (SeenObjects.Contains(AssetToCheck.ObjectPath))
+				{
+					continue;
+				}
+				SeenObjects.Add(AssetToCheck.ObjectPath);
+
+				if (AssetToCheck.GetClass() == UNiagaraEmitter::StaticClass())
+				{
+					Count++;
+					TArray<FName> Referencers;
+					AssetRegistry.GetReferencers(AssetToCheck.PackageName, Referencers);
+					for (FName& Referencer : Referencers)
+					{
+						AssetRegistry.GetAssetsByPackageName(Referencer, AssetsToCheck);
+					}
+				}
+				else if (AssetToCheck.GetClass() == UNiagaraSystem::StaticClass())
+				{
+					Count++;
+				}
+			}
+			Count--; // remove one for our own asset
+		}
+		EmitterAffectedAssets = Count;
+	}
+	return EmitterAffectedAssets.Get(0);
+}
+
 #undef LOCTEXT_NAMESPACE
diff --git a/Engine/Plugins/FX/Niagara/Source/NiagaraEditorWidgets/Private/SNiagaraOverviewGraphTitleBar.h b/Engine/Plugins/FX/Niagara/Source/NiagaraEditorWidgets/Private/SNiagaraOverviewGraphTitleBar.h
--- a/Engine/Plugins/FX/Niagara/Source/NiagaraEditorWidgets/Private/SNiagaraOverviewGraphTitleBar.h
+++ b/Engine/Plugins/FX/Niagara/Source/NiagaraEditorWidgets/Private/SNiagaraOverviewGraphTitleBar.h
@@ -19,7 +19,8 @@ class SNiagaraOverviewGraphTitleBar : public SCompoundWidget
 	SLATE_END_ARGS()
 
 public:
-	void Construct(const FArguments& InArgs, TSharedRef<FNiagaraSystemViewModel> InSystemViewModel);
+	virtual ~SNiagaraOverviewGraphTitleBar() override;
+	void Construct(const FArguments& InArgs, TSharedRef<FNiagaraSystemViewModel> InSystemViewModel, const FAssetData& InEditedAsset);
 
 	void RebuildWidget();
 	//virtual FReply OnMouseButtonDown(const FGeometry& MyGeometry, const FPointerEvent& MouseEvent) override;
@@ -28,6 +29,18 @@ private:
 	bool bScalabilityModeActive = false;
 	TSharedPtr<SVerticalBox> ContainerWidget;
 	TWeakPtr<SNiagaraOverviewGraph> OverviewGraph;
-private:
+
 	void OnUpdateScalabilityMode(bool bActive);
+
+	FText GetEmitterSubheaderText() const;
+	EVisibility GetEmitterSubheaderVisibility() const;
+	FLinearColor GetEmitterSubheaderColor() const;
+	
+	int32 GetEmitterAffectedAssets() const;
+	void ResetAssetCount(const FAssetData& InAssetData);
+	void AddAssetListeners();
+	void ClearListeners();
+
+	mutable TOptional<int32> EmitterAffectedAssets;
+	FAssetData EditedAsset;
 };
diff --git a/Engine/Plugins/FX/Niagara/Source/NiagaraEditorWidgets/Public/NiagaraEditorWidgetsModule.h b/Engine/Plugins/FX/Niagara/Source/NiagaraEditorWidgets/Public/NiagaraEditorWidgetsModule.h
--- a/Engine/Plugins/FX/Niagara/Source/NiagaraEditorWidgets/Public/NiagaraEditorWidgetsModule.h
+++ b/Engine/Plugins/FX/Niagara/Source/NiagaraEditorWidgets/Public/NiagaraEditorWidgetsModule.h
@@ -52,7 +52,7 @@ private:
 	{
 	public:
 		virtual TSharedRef<SWidget> CreateStackView(UNiagaraStackViewModel& StackViewModel) const override;
-		virtual TSharedRef<SWidget> CreateSystemOverview(TSharedRef<FNiagaraSystemViewModel> SystemViewModel) const override;
+		virtual TSharedRef<SWidget> CreateSystemOverview(TSharedRef<FNiagaraSystemViewModel> SystemViewModel, const FAssetData& EditedAsset) const override;
 		virtual TSharedRef<SWidget> CreateStackIssueIcon(UNiagaraStackViewModel& StackViewModel, UNiagaraStackEntry& StackEntry) const override;
 		virtual TSharedRef<SWidget> CreateScriptScratchPadManager(UNiagaraScratchPadViewModel& ScriptScratchPadViewModel) const override;
 		virtual TSharedRef<SWidget> CreateCurveOverview(TSharedRef<FNiagaraSystemViewModel> SystemViewModel) const override;
