Fixed scalability mode button not appearing checked if activated
Also added missing icon

#rnx

diff --git a/Engine/Plugins/FX/Niagara/Source/NiagaraEditor/Private/NiagaraEditorStyle.cpp b/Engine/Plugins/FX/Niagara/Source/NiagaraEditor/Private/NiagaraEditorStyle.cpp
--- a/Engine/Plugins/FX/Niagara/Source/NiagaraEditor/Private/NiagaraEditorStyle.cpp
+++ b/Engine/Plugins/FX/Niagara/Source/NiagaraEditor/Private/NiagaraEditorStyle.cpp
@@ -674,6 +674,7 @@ void FNiagaraEditorStyle::InitScalabilityColors()
 
 void FNiagaraEditorStyle::InitScalabilityIcons()
 {
+	Set("NiagaraEditor.Scalability", new IMAGE_BRUSH_SVG("Icons/Scalability/Scalability", Icon20x20));
 	Set("NiagaraEditor.Scalability.System.Feature.Override", new IMAGE_BRUSH_SVG("Icons/Scalability/Override", Icon20x20));
 	Set("NiagaraEditor.Scalability.Preview.ResetPlatform", new IMAGE_BRUSH("Icons/Scalability/ResetPreviewPlatform", Icon20x20));
 }
diff --git a/Engine/Plugins/FX/Niagara/Source/NiagaraEditor/Private/Toolkits/SystemToolkitModes/NiagaraSystemToolkitMode_Default.cpp b/Engine/Plugins/FX/Niagara/Source/NiagaraEditor/Private/Toolkits/SystemToolkitModes/NiagaraSystemToolkitMode_Default.cpp
--- a/Engine/Plugins/FX/Niagara/Source/NiagaraEditor/Private/Toolkits/SystemToolkitModes/NiagaraSystemToolkitMode_Default.cpp
+++ b/Engine/Plugins/FX/Niagara/Source/NiagaraEditor/Private/Toolkits/SystemToolkitModes/NiagaraSystemToolkitMode_Default.cpp
@@ -259,7 +259,8 @@ void FNiagaraSystemToolkitMode_Default::ExtendToolbar()
 						NAME_None, 
 						LOCTEXT("ScalabilityLabel", "Scalability"),
 						LOCTEXT("ScalabilityTooltip", "Turn on scalability mode to optimize your effects for various platforms and quality settings."),
-						FSlateIcon(FNiagaraEditorStyle::Get().GetStyleSetName(), "NiagaraEditor.SimulationOptions")
+						FSlateIcon(FNiagaraEditorStyle::Get().GetStyleSetName(), "NiagaraEditor.Scalability"),
+						EUserInterfaceActionType::ToggleButton
 					);
 				}
 				ToolbarBuilder.EndSection();
diff --git a/Engine/Plugins/FX/Niagara/Source/NiagaraEditor/Private/Toolkits/SystemToolkitModes/NiagaraSystemToolkitMode_Scalability.cpp b/Engine/Plugins/FX/Niagara/Source/NiagaraEditor/Private/Toolkits/SystemToolkitModes/NiagaraSystemToolkitMode_Scalability.cpp
--- a/Engine/Plugins/FX/Niagara/Source/NiagaraEditor/Private/Toolkits/SystemToolkitModes/NiagaraSystemToolkitMode_Scalability.cpp
+++ b/Engine/Plugins/FX/Niagara/Source/NiagaraEditor/Private/Toolkits/SystemToolkitModes/NiagaraSystemToolkitMode_Scalability.cpp
@@ -203,19 +203,10 @@ void FNiagaraSystemToolkitMode_Scalability::ExtendToolbar()
 					ScalabilityToggleAction,
 					NAME_None, 
 					LOCTEXT("ScalabilityLabel", "Scalability"),
-					LOCTEXT("ScalabilityTooltip", "Turn on scalability mode to optimize your effects for various platforms and quality settings."),
-					FSlateIcon(FNiagaraEditorStyle::Get().GetStyleSetName(), "NiagaraEditor.SimulationOptions")
+					LOCTEXT("ScalabilityTooltip", "Turn off scalability mode."),
+					FSlateIcon(FNiagaraEditorStyle::Get().GetStyleSetName(), "NiagaraEditor.Scalability"),
+					EUserInterfaceActionType::ToggleButton
 				);
-
-				// @todo remove this once preview settings are finalized (tab vs. overlay in graph etc.)
-				// TSharedRef<SHorizontalBox> ViewModeContainer = SNew(SHorizontalBox)
-				// + SHorizontalBox::Slot()
-				// .AutoWidth()
-				// [
-				// 	SNew(SNiagaraScalabilityPreviewSettings, *Toolkit->GetSystemViewModel()->GetScalabilityViewModel())
-				// ];
-				//
-				// ToolbarBuilder.AddWidget(ViewModeContainer);
 			}
 			ToolbarBuilder.EndSection();		
 		}		
