Improved low detail node view by splitting up the emitter names and adding a setting to disable it

diff --git a/Engine/Plugins/FX/Niagara/Source/Niagara/Public/NiagaraSettings.h b/Engine/Plugins/FX/Niagara/Source/Niagara/Public/NiagaraSettings.h
--- a/Engine/Plugins/FX/Niagara/Source/Niagara/Public/NiagaraSettings.h
+++ b/Engine/Plugins/FX/Niagara/Source/Niagara/Public/NiagaraSettings.h
@@ -106,10 +106,6 @@ class NIAGARA_API UNiagaraSettings : public UDeveloperSettings
 	UPROPERTY(config, EditAnywhere, Category=Niagara)
 	bool bDisplayAffectedAssetStatsInEditor = true;
 
-	/** If true then the system editor will zoom to fit all emitters when opening an asset. */
-	UPROPERTY(config, EditAnywhere, Category = Niagara)
-	bool bAlwaysZoomToFitSystemGraph = true;
-
 	/** If true then the "link input" menu will also show variables of different types, as long as there is a conversion script for them. */
 	UPROPERTY(config, EditAnywhere, Category = Niagara)
 	bool bShowConvertibleInputsInStack = false;
diff --git a/Engine/Plugins/FX/Niagara/Source/NiagaraEditor/Public/NiagaraEditorSettings.h b/Engine/Plugins/FX/Niagara/Source/NiagaraEditor/Public/NiagaraEditorSettings.h
--- a/Engine/Plugins/FX/Niagara/Source/NiagaraEditor/Public/NiagaraEditorSettings.h
+++ b/Engine/Plugins/FX/Niagara/Source/NiagaraEditor/Public/NiagaraEditorSettings.h
@@ -247,6 +247,18 @@ public:
 	UPROPERTY(config, EditAnywhere, Category = Niagara)
 	TArray<FNiagaraSpawnShortcut> GraphCreationShortcuts;
 
+	/** If true then emitter and system nodes will show a simplified representation on low zoom levels. This improves performance and readablity when zoomed out of the system overview graph. */
+	UPROPERTY(config, EditAnywhere, Category = Niagara)
+	bool bSimplifyStackNodesAtLowResolution = true;
+
+	/** The max number of chars before names on the low resolution nodes are truncated. */
+	UPROPERTY(config, EditAnywhere, Category = Niagara, meta=(EditCondition="bSimplifyStackNodesAtLowResolution", UIMin=3))
+	int32 LowResolutionNodeMaxNameChars = 7;
+	
+	/** If true then the system editor will zoom to fit all emitters when opening an asset. */
+	UPROPERTY(config, EditAnywhere, Category = Niagara)
+	bool bAlwaysZoomToFitSystemGraph = true;
+
 	TArray<float> GetPlaybackSpeeds() const;
 
 	/** Gets whether or not auto-compile is enabled in the editors. */
diff --git a/Engine/Plugins/FX/Niagara/Source/NiagaraEditorWidgets/Private/SNiagaraOverviewGraph.cpp b/Engine/Plugins/FX/Niagara/Source/NiagaraEditorWidgets/Private/SNiagaraOverviewGraph.cpp
--- a/Engine/Plugins/FX/Niagara/Source/NiagaraEditorWidgets/Private/SNiagaraOverviewGraph.cpp
+++ b/Engine/Plugins/FX/Niagara/Source/NiagaraEditorWidgets/Private/SNiagaraOverviewGraph.cpp
@@ -16,7 +16,7 @@
 #include "NiagaraEditorModule.h"
 #include "GraphEditorActions.h"
 #include "NiagaraEditorWidgetsStyle.h"
-#include "NiagaraSettings.h"
+#include "NiagaraEditorSettings.h"
 #include "Framework/MultiBox/MultiBoxBuilder.h"
 #include "ScopedTransaction.h"
 #include "Framework/Application/SlateApplication.h"
@@ -115,7 +115,7 @@ void SNiagaraOverviewGraph::Construct(const FArguments& InArgs, TSharedRef<FNiag
 
 	GraphEditor->GetCurrentGraph()->AddOnGraphChangedHandler(FOnGraphChanged::FDelegate::CreateSP(this, &SNiagaraOverviewGraph::OnNodesCreated));
 
-	const UNiagaraSettings* NiagaraSettings = GetDefault<UNiagaraSettings>();
+	const UNiagaraEditorSettings* NiagaraSettings = GetDefault<UNiagaraEditorSettings>();
 	FNiagaraGraphViewSettings ViewSettings = ViewModel->GetViewSettings();
 	if (NiagaraSettings->bAlwaysZoomToFitSystemGraph == false && ViewSettings.IsValid())
 	{
diff --git a/Engine/Plugins/FX/Niagara/Source/NiagaraEditorWidgets/Private/SNiagaraOverviewStackNode.cpp b/Engine/Plugins/FX/Niagara/Source/NiagaraEditorWidgets/Private/SNiagaraOverviewStackNode.cpp
--- a/Engine/Plugins/FX/Niagara/Source/NiagaraEditorWidgets/Private/SNiagaraOverviewStackNode.cpp
+++ b/Engine/Plugins/FX/Niagara/Source/NiagaraEditorWidgets/Private/SNiagaraOverviewStackNode.cpp
@@ -26,6 +26,7 @@
 #include "Widgets/SToolTip.h"
 #include "Widgets/SBoxPanel.h"
 #include "NiagaraRendererProperties.h"
+#include "NiagaraEditorSettings.h"
 #include "SGraphPanel.h"
 #include "Widgets/SWidget.h"
 #include "ViewModels/Stack/NiagaraStackRendererItem.h"
@@ -312,6 +313,11 @@ void SNiagaraOverviewStackNode::Tick(const FGeometry& AllottedGeometry, const do
 		if (ScalabilityWrapper != nullptr && !UseLowDetailNodeContent())
 		{
 			LastHighDetailSize = ScalabilityWrapper->GetTickSpaceGeometry().Size;
+			if (GeometryTickForSize > 0)
+			{
+				// the stack needs a few tick to fully initialize, so we wait to bit to grab the right low detail geomtry size
+				GeometryTickForSize--;
+			}
 		}
 	}
 }
@@ -469,11 +475,10 @@ TSharedRef<SWidget> SNiagaraOverviewStackNode::CreateNodeContentArea()
 			.HeightOverride(this, &SNiagaraOverviewStackNode::GetLowDetailDesiredHeight)
 			[
 				SNew(STextBlock)
-				.Text(this, &SNiagaraOverviewStackNode::GetEditableNodeTitleAsText)
+				.Text(this, &SNiagaraOverviewStackNode::GetLowDetailNodeTitle)
 				.TextStyle(&FNiagaraEditorWidgetsStyle::Get().GetWidgetStyle<FTextBlockStyle>("NiagaraEditor.SystemOverview.ZoomedOutNodeFont"))
 				.Visibility(EVisibility::HitTestInvisible)
 				.Clipping(EWidgetClipping::Inherit)
-				.RenderTransform(FSlateRenderTransform(Concatenate(FQuat2D(FMath::DegreesToRadians(90.0f)), FVector2D(100, 25))))
 			]
 		]
 		.HighDetail()
@@ -501,7 +506,8 @@ TSharedRef<SWidget> SNiagaraOverviewStackNode::CreateNodeContentArea()
 
 bool SNiagaraOverviewStackNode::UseLowDetailNodeContent() const
 {
-	if (LastHighDetailSize.IsNearlyZero())
+	const UNiagaraEditorSettings* NiagaraSettings = GetDefault<UNiagaraEditorSettings>();
+	if (LastHighDetailSize.IsNearlyZero() || GeometryTickForSize > 0 || !NiagaraSettings->bSimplifyStackNodesAtLowResolution)
 	{
 		return false;
 	}
@@ -528,6 +534,100 @@ FOptionalSize SNiagaraOverviewStackNode::GetLowDetailDesiredHeight() const
 	return LastHighDetailSize.Y;
 }
 
+namespace NiagaraOverviewStackNode
+{
+	struct StringRange
+	{
+		int32 Start = 0;
+		int32 End = 0;
+	};
+	
+	bool IsUpperAlphaNumeric(const TCHAR& c)
+	{
+		return (c >= 'A' && c <= 'Z') || (c >= '0' && c <= '9');
+	}
+
+	// this method splits long emitter names for the low details view by inserting new lines
+	// and abbreviating large pieces when necessary
+	FString PrettySplitString(const FString& Input)
+	{
+		TArray<FStringView, TInlineAllocator<16>> Parts;
+
+		// split at convenient places
+		int32 Start = 0;
+		for (int i = 0; i < Input.Len() - 1; i++)
+		{
+			TCHAR currentChar = Input[i];
+			TCHAR nextChar = Input[i + 1];
+
+			bool bIsLowerAlpha = currentChar >= 'a' && currentChar <= 'z';
+
+			// skip '_' at the beginning or in sequence
+			if (currentChar == '_')
+			{
+				Start = i + 1;
+			}
+			// replace '_' in the middle of the name with a newline
+			else if (nextChar == '_')
+			{
+				Parts.Add(FStringView(&Input[Start], i - Start + 1));
+				i++;
+				Start = i + 1;
+			}
+			// split when changing from lowercase to uppercase or number
+			else if (bIsLowerAlpha && IsUpperAlphaNumeric(nextChar))
+			{
+				Parts.Add(FStringView(&Input[Start], i - Start + 1));
+				Start = i + 1;
+			}
+		}
+		// add the end piece
+		if (Start < Input.Len())
+		{
+			Parts.Add(FStringView(&Input[Start], Input.Len() - Start));
+		}
+
+		// assemble the pieces
+		const UNiagaraEditorSettings* NiagaraSettings = GetDefault<UNiagaraEditorSettings>();
+		int32 MaxLength = FMath::Max(3, NiagaraSettings->LowResolutionNodeMaxNameChars);
+		TStringBuilder<128> SplitNameBuilder;
+		int32 RunningLength = 0;
+		for (const FStringView& Split : Parts)
+		{
+			// merge small parts together, only insert a newline between big parts
+			if (Split.Len() + RunningLength > MaxLength && SplitNameBuilder.Len() > 0)
+			{
+				SplitNameBuilder.AppendChar('\n');
+				RunningLength = 0;
+			}
+			
+			// if the name is too long, abbreviate with ...
+			if (Split.Len() > MaxLength)
+			{
+				SplitNameBuilder.Append(Split.Left(MaxLength - 2));
+				SplitNameBuilder.Append(TEXT("…"));
+			}
+			else
+			{
+				SplitNameBuilder.Append(Split);
+			}
+			RunningLength += Split.Len();
+		}
+
+		return SplitNameBuilder.ToString();
+	}
+}
+
+FText SNiagaraOverviewStackNode::GetLowDetailNodeTitle() const
+{
+	if (FString Title = GetEditableNodeTitle(); LowDetailTitleCache.Key != Title)
+	{
+		LowDetailTitleCache.Key = Title;
+		LowDetailTitleCache.Value = FText::FromString(NiagaraOverviewStackNode::PrettySplitString(Title));
+	}
+	return LowDetailTitleCache.Value;
+}
+
 void SNiagaraOverviewStackNode::StackViewModelStructureChanged(ENiagaraStructureChangedFlags Flags)
 {
 	bTopContentBarRefreshPending = true;
diff --git a/Engine/Plugins/FX/Niagara/Source/NiagaraEditorWidgets/Private/SNiagaraOverviewStackNode.h b/Engine/Plugins/FX/Niagara/Source/NiagaraEditorWidgets/Private/SNiagaraOverviewStackNode.h
--- a/Engine/Plugins/FX/Niagara/Source/NiagaraEditorWidgets/Private/SNiagaraOverviewStackNode.h
+++ b/Engine/Plugins/FX/Niagara/Source/NiagaraEditorWidgets/Private/SNiagaraOverviewStackNode.h
@@ -35,10 +35,12 @@ protected:
 	void StackViewModelDataObjectChanged(TArray<UObject*> ChangedObjects, ENiagaraDataObjectChange ChangeType);
 	void FillTopContentBar();
 	void OnMaterialCompiled(class UMaterialInterface* MaterialInterface);
+	
 	bool UseLowDetailNodeContent() const;
 	FVector2D GetLowDetailDesiredSize() const;
 	FOptionalSize GetLowDetailDesiredWidth() const;
 	FOptionalSize GetLowDetailDesiredHeight() const;
+	FText GetLowDetailNodeTitle() const;
 
 	void CreateBottomSummaryExpander();
 private:
@@ -101,4 +103,6 @@ private:
 	
 	/** Cached size from when we last drew at high detail */
 	FVector2D LastHighDetailSize = FVector2D::ZeroVector;
+	int32 GeometryTickForSize = 3;
+	mutable TPair<FString, FText> LowDetailTitleCache;
 };
\ No newline at end of file
