Fixing NonUnity issue caused by including just RHIResources.h

diff --git a/Engine/Plugins/FX/Niagara/Source/Niagara/Classes/VolumeCache.h b/Engine/Plugins/FX/Niagara/Source/Niagara/Classes/VolumeCache.h
--- a/Engine/Plugins/FX/Niagara/Source/Niagara/Classes/VolumeCache.h
+++ b/Engine/Plugins/FX/Niagara/Source/Niagara/Classes/VolumeCache.h
@@ -5,7 +5,7 @@
 
 #include "CoreMinimal.h"
 #include "UObject/ObjectMacros.h"
-#include "RHIResources.h"
+#include "RHI.h"
 
 // @todo: we need builds for OpenVDB for platforms other than windows
 #if PLATFORM_WINDOWS
