Moved editor setting from runtime object to editor object

diff --git a/Engine/Plugins/FX/Niagara/Source/Niagara/Public/NiagaraSettings.h b/Engine/Plugins/FX/Niagara/Source/Niagara/Public/NiagaraSettings.h
--- a/Engine/Plugins/FX/Niagara/Source/Niagara/Public/NiagaraSettings.h
+++ b/Engine/Plugins/FX/Niagara/Source/Niagara/Public/NiagaraSettings.h
@@ -102,10 +102,6 @@ class NIAGARA_API UNiagaraSettings : public UDeveloperSettings
 	UPROPERTY(config, EditAnywhere, Category = Viewport)
 	bool bSystemViewportInOrbitMode = true;
 
-	/** If true, then the emitter and script editors will show an info message how many downstream asset are affected by a change. Gathering this information for large asset graphs can delay the opening of the asset editors a bit. */
-	UPROPERTY(config, EditAnywhere, Category=Niagara)
-	bool bDisplayAffectedAssetStatsInEditor = true;
-
 	/** If true then the "link input" menu will also show variables of different types, as long as there is a conversion script for them. */
 	UPROPERTY(config, EditAnywhere, Category = Niagara)
 	bool bShowConvertibleInputsInStack = false;
diff --git a/Engine/Plugins/FX/Niagara/Source/NiagaraEditor/Private/NiagaraEditorSettings.cpp b/Engine/Plugins/FX/Niagara/Source/NiagaraEditor/Private/NiagaraEditorSettings.cpp
--- a/Engine/Plugins/FX/Niagara/Source/NiagaraEditor/Private/NiagaraEditorSettings.cpp
+++ b/Engine/Plugins/FX/Niagara/Source/NiagaraEditor/Private/NiagaraEditorSettings.cpp
@@ -453,6 +453,11 @@ void UNiagaraEditorSettings::SetDisplayAdvancedParameterPanelCategories(bool bIn
 	}
 }
 
+bool UNiagaraEditorSettings::GetDisplayAffectedAssetStats() const
+{
+	return bDisplayAffectedAssetStats;
+}
+
 FNiagaraNewAssetDialogConfig UNiagaraEditorSettings::GetNewAssetDailogConfig(FName InDialogConfigKey) const
 {
 	const FNiagaraNewAssetDialogConfig* Config = NewAssetDialogConfigMap.Find(InDialogConfigKey);
diff --git a/Engine/Plugins/FX/Niagara/Source/NiagaraEditor/Private/Widgets/SNiagaraScriptGraph.cpp b/Engine/Plugins/FX/Niagara/Source/NiagaraEditor/Private/Widgets/SNiagaraScriptGraph.cpp
--- a/Engine/Plugins/FX/Niagara/Source/NiagaraEditor/Private/Widgets/SNiagaraScriptGraph.cpp
+++ b/Engine/Plugins/FX/Niagara/Source/NiagaraEditor/Private/Widgets/SNiagaraScriptGraph.cpp
@@ -28,7 +28,7 @@
 #include "Framework/Application/SlateApplication.h"
 #include "GraphEditorActions.h"
 #include "EditorFontGlyphs.h"
-#include "NiagaraSettings.h"
+#include "NiagaraEditorSettings.h"
 #include "AssetRegistry/AssetRegistryModule.h"
 #include "Widgets/SNiagaraGraphActionMenu.h"
 
@@ -699,8 +699,7 @@ int32 SNiagaraScriptGraph::GetScriptAffectedAssets() const
 	{
 		return 0;
 	}
-	const UNiagaraSettings* Settings = GetDefault<UNiagaraSettings>();
-	if (Settings->bDisplayAffectedAssetStatsInEditor == false)
+	if (GetDefault<UNiagaraEditorSettings>()->GetDisplayAffectedAssetStats() == false)
 	{
 		return 0;
 	}
diff --git a/Engine/Plugins/FX/Niagara/Source/NiagaraEditor/Public/NiagaraEditorSettings.h b/Engine/Plugins/FX/Niagara/Source/NiagaraEditor/Public/NiagaraEditorSettings.h
--- a/Engine/Plugins/FX/Niagara/Source/NiagaraEditor/Public/NiagaraEditorSettings.h
+++ b/Engine/Plugins/FX/Niagara/Source/NiagaraEditor/Public/NiagaraEditorSettings.h
@@ -297,6 +297,8 @@ public:
 	/** Sets whether or not to display advanced categories for the parameter panel. */
 	void SetDisplayAdvancedParameterPanelCategories(bool bInDisplayAdvancedParameterPanelCategories);
 
+	bool GetDisplayAffectedAssetStats() const;
+
 	bool IsShowGridInViewport() const;
 	void SetShowGridInViewport(bool bShowGridInViewport);
 	
@@ -364,6 +366,10 @@ private:
 	UPROPERTY(config, EditAnywhere, Category = Niagara)
 	bool bDisplayAdvancedParameterPanelCategories;
 
+	/** If true, then the emitter and script editors will show an info message how many downstream asset are affected by a change. Gathering this information for large asset graphs can delay the opening of the asset editors a bit. */
+	UPROPERTY(config, EditAnywhere, Category=Niagara)
+	bool bDisplayAffectedAssetStats = true;
+
 	/** Speeds used for slowing down and speeding up the playback speeds */
 	UPROPERTY(config, EditAnywhere, Category = Niagara)
 	TArray<float> PlaybackSpeeds;
diff --git a/Engine/Plugins/FX/Niagara/Source/NiagaraEditorWidgets/Private/SNiagaraOverviewGraphTitleBar.cpp b/Engine/Plugins/FX/Niagara/Source/NiagaraEditorWidgets/Private/SNiagaraOverviewGraphTitleBar.cpp
--- a/Engine/Plugins/FX/Niagara/Source/NiagaraEditorWidgets/Private/SNiagaraOverviewGraphTitleBar.cpp
+++ b/Engine/Plugins/FX/Niagara/Source/NiagaraEditorWidgets/Private/SNiagaraOverviewGraphTitleBar.cpp
@@ -3,7 +3,7 @@
 #include "SNiagaraOverviewGraphTitleBar.h"
 
 #include "NiagaraEditorWidgetsStyle.h"
-#include "NiagaraSettings.h"
+#include "NiagaraEditorSettings.h"
 #include "ViewModels/NiagaraOverviewGraphViewModel.h"
 #include "SNiagaraScalabilityPreviewSettings.h"
 #include "SNiagaraSystemEffectTypeBar.h"
@@ -156,8 +156,7 @@ int32 SNiagaraOverviewGraphTitleBar::GetEmitterAffectedAssets() const
 		return 0;
 	}
 	
-	const UNiagaraSettings* Settings = GetDefault<UNiagaraSettings>();
-	if (Settings->bDisplayAffectedAssetStatsInEditor == false)
+	if (GetDefault<UNiagaraEditorSettings>()->GetDisplayAffectedAssetStats() == false)
 	{
 		return 0;
 	}
