cis fix

diff --git a/Engine/Plugins/FX/Niagara/Source/Niagara/Classes/VolumeCache.h b/Engine/Plugins/FX/Niagara/Source/Niagara/Classes/VolumeCache.h
--- a/Engine/Plugins/FX/Niagara/Source/Niagara/Classes/VolumeCache.h
+++ b/Engine/Plugins/FX/Niagara/Source/Niagara/Classes/VolumeCache.h
@@ -89,8 +89,8 @@ public:
 	FOpenVDBCacheData() {}
 	
 	virtual ~FOpenVDBCacheData() 
-	{
-		UnloadAll();
+	{		
+		OpenVDBGrids.Reset();
 		DenseGridPtr = nullptr;
 	}
 
