#Niagara Add exception handling for call to RegisterGrid

diff --git a/Engine/Plugins/FX/Niagara/Source/NiagaraEditor/Private/NiagaraEditorModule.cpp b/Engine/Plugins/FX/Niagara/Source/NiagaraEditor/Private/NiagaraEditorModule.cpp
--- a/Engine/Plugins/FX/Niagara/Source/NiagaraEditor/Private/NiagaraEditorModule.cpp
+++ b/Engine/Plugins/FX/Niagara/Source/NiagaraEditor/Private/NiagaraEditorModule.cpp
@@ -825,7 +825,14 @@ void FNiagaraEditorModule::StartupModule()
 	// Global registration of  the vdb types.
 #if PLATFORM_WINDOWS
 	openvdb::initialize();
-	Vec4SGrid::registerGrid();
+	try
+	{
+		Vec4SGrid::registerGrid();
+	}
+	catch (openvdb::Exception e)
+	{
+		UE_LOG(LogNiagaraEditor, Warning, TEXT("OpenVDB exception: %s"), ANSI_TO_TCHAR(e.what()));
+	}
 #endif	
 	
 	bThumbnailRenderersRegistered = false;
