remove dependency on OpenVDB in NiagaraEditor
initial support for 16 bit floats

diff --git a/Engine/Plugins/FX/Niagara/Source/Niagara/Classes/NiagaraOpenVDB.h b/Engine/Plugins/FX/Niagara/Source/Niagara/Classes/NiagaraOpenVDB.h
--- a/Engine/Plugins/FX/Niagara/Source/Niagara/Classes/NiagaraOpenVDB.h
+++ b/Engine/Plugins/FX/Niagara/Source/Niagara/Classes/NiagaraOpenVDB.h
@@ -33,9 +33,27 @@ UE_PUSH_MACRO("check")
 #include <openvdb/tools/Prune.h>
 #include <openvdb/tools/Dense.h>
 #include <openvdb/math/Vec4.h>
+#include <openvdb//math/Half.h>
 
-using Vec4STree = openvdb::tree::Tree4<openvdb::Vec4s, 5, 4, 3>::Type;
+using Vec4s = openvdb::Vec4s;
+using Vec4STree = openvdb::tree::Tree4<Vec4s, 5, 4, 3>::Type;
 using Vec4SGrid = openvdb::Grid<Vec4STree>;
+using Vec4SDense = openvdb::tools::Dense<Vec4s, openvdb::tools::MemoryLayout::LayoutXYZ>;
+
+using Vec4h = openvdb::v8_1::math::Vec4<openvdb::math::internal::half>;
+OPENVDB_IS_POD(Vec4h)
+
+using Vec4HTree = openvdb::tree::Tree4<Vec4h, 5, 4, 3>::Type;
+using Vec4HGrid = openvdb::Grid<Vec4HTree>;
+using Vec4HDense = openvdb::tools::Dense<Vec4h, openvdb::tools::MemoryLayout::LayoutXYZ>;
+
+using Vec4 = Vec4s;
+using Vec4Tree = Vec4STree;
+using Vec4Grid = Vec4SGrid;
+using Vec4Dense = Vec4SDense;
+using Vec4GridPtr = Vec4Grid::Ptr;
+using Vec4DensePtr = Vec4Dense::Ptr;
+
 
 #if LOCAL_M_PI
 #undef M_PI
diff --git a/Engine/Plugins/FX/Niagara/Source/Niagara/Classes/VolumeCache.h b/Engine/Plugins/FX/Niagara/Source/Niagara/Classes/VolumeCache.h
--- a/Engine/Plugins/FX/Niagara/Source/Niagara/Classes/VolumeCache.h
+++ b/Engine/Plugins/FX/Niagara/Source/Niagara/Classes/VolumeCache.h
@@ -7,11 +7,6 @@
 #include "UObject/ObjectMacros.h"
 #include "RHI.h"
 
-// @todo: we need builds for OpenVDB for platforms other than windows
-#if PLATFORM_WINDOWS
-#include "NiagaraOpenVDB.h"
-#endif
-
 #include "VolumeCache.generated.h"
 
 class FVolumeCacheData;
@@ -81,30 +76,7 @@ protected:
 	FIntVector DenseResolution;
 };
 
-// @todo: we need builds for OpenVDB for platforms other than windows
-#if PLATFORM_WINDOWS
-class NIAGARA_API FOpenVDBCacheData : public FVolumeCacheData
+namespace OpenVDBTools
 {
-public:
-	FOpenVDBCacheData() {}
-	
-	virtual ~FOpenVDBCacheData() 
-	{		
-		OpenVDBGrids.Reset();
-		DenseGridPtr = nullptr;
-	}
-
-	virtual void Init(FIntVector Resolution);
-	virtual bool LoadFile(FString Path, int frame);
-	virtual bool UnloadFile(int frame);
-	virtual bool LoadRange(FString Path, int Start, int End);
-	virtual void UnloadAll();
-	virtual bool Fill3DTexture_RenderThread(int frame, FTextureRHIRef TextureToFill, FRHICommandListImmediate& RHICmdList);
-
-	static bool WriteImageDataToOpenVDBFile(FStringView FilePath, FIntVector ImageSize, TArrayView<FFloat16Color> ImageData, bool UseFloatGrids = false);
-
-private:	
-	TMap<int32, Vec4SGrid::Ptr> OpenVDBGrids;
-	openvdb::tools::Dense<openvdb::Vec4s, openvdb::tools::MemoryLayout::LayoutXYZ>::Ptr DenseGridPtr;	
-};
-#endif
\ No newline at end of file
+	NIAGARA_API bool WriteImageDataToOpenVDBFile(FStringView FilePath, FIntVector ImageSize, TArrayView<FFloat16Color> ImageData, bool UseFloatGrids = false);
+}
\ No newline at end of file
diff --git a/Engine/Plugins/FX/Niagara/Source/Niagara/Private/NiagaraModule.cpp b/Engine/Plugins/FX/Niagara/Source/Niagara/Private/NiagaraModule.cpp
--- a/Engine/Plugins/FX/Niagara/Source/Niagara/Private/NiagaraModule.cpp
+++ b/Engine/Plugins/FX/Niagara/Source/Niagara/Private/NiagaraModule.cpp
@@ -205,7 +205,15 @@ void INiagaraModule::StartupModule()
 #if PLATFORM_WINDOWS
 	// Global registration of  the vdb types.
 	openvdb::initialize();
-	Vec4SGrid::registerGrid();
+	if (!Vec4SGrid::isRegistered())
+	{
+		Vec4SGrid::registerGrid();
+	}
+	
+	if (!Vec4HGrid::isRegistered())
+	{
+		Vec4HGrid::registerGrid();
+	}
 #endif
 
 	FNiagaraWorldManager::OnStartup();
diff --git a/Engine/Plugins/FX/Niagara/Source/Niagara/Private/VolumeCache.cpp b/Engine/Plugins/FX/Niagara/Source/Niagara/Private/VolumeCache.cpp
--- a/Engine/Plugins/FX/Niagara/Source/Niagara/Private/VolumeCache.cpp
+++ b/Engine/Plugins/FX/Niagara/Source/Niagara/Private/VolumeCache.cpp
@@ -6,9 +6,42 @@
 
 #include "VolumeCache.h"
 
+// @todo: we need builds for OpenVDB for platforms other than windows
+#if PLATFORM_WINDOWS
+#include "NiagaraOpenVDB.h"
+#endif
+
 
 DEFINE_LOG_CATEGORY_STATIC(LogVolumeCache, Log, All);
 
+
+// @todo: we need builds for OpenVDB for platforms other than windows
+#if PLATFORM_WINDOWS
+class NIAGARA_API FOpenVDBCacheData : public FVolumeCacheData
+{
+public:
+	FOpenVDBCacheData() {}
+
+	virtual ~FOpenVDBCacheData()
+	{
+		OpenVDBGrids.Reset();
+		DenseGridPtr = nullptr;
+	}
+
+	virtual void Init(FIntVector Resolution);
+	virtual bool LoadFile(FString Path, int frame);
+	virtual bool UnloadFile(int frame);
+	virtual bool LoadRange(FString Path, int Start, int End);
+	virtual void UnloadAll();
+	virtual bool Fill3DTexture_RenderThread(int frame, FTextureRHIRef TextureToFill, FRHICommandListImmediate& RHICmdList);
+	
+private:
+	TMap<int32, Vec4SGrid::Ptr> OpenVDBGrids;
+	openvdb::tools::Dense<openvdb::Vec4s, openvdb::tools::MemoryLayout::LayoutXYZ>::Ptr DenseGridPtr;
+};
+#endif
+
+
 //*****************************************************************************
 //***************************** UVolumeCache ********************************
 //*****************************************************************************
@@ -97,7 +130,7 @@ FString FVolumeCacheData::GetAssetPath(FString PathFormat, int32 FrameIndex) con
 #if PLATFORM_WINDOWS
 void FOpenVDBCacheData::Init(FIntVector Resolution)
 {
-	DenseGridPtr.reset(new openvdb::tools::Dense<openvdb::Vec4s, openvdb::tools::MemoryLayout::LayoutXYZ>(openvdb::CoordBBox(0, 0, 0, Resolution.X - 1, Resolution.Y - 1, Resolution.Z - 1), openvdb::Vec4s(0.0, 0.0, 0.0, 0.0)));
+	DenseGridPtr.reset(new Vec4Dense(openvdb::CoordBBox(0, 0, 0, Resolution.X - 1, Resolution.Y - 1, Resolution.Z - 1), Vec4(0.0, 0.0, 0.0, 0.0)));
 }
 
 bool FOpenVDBCacheData::LoadFile(FString Path, int frame)
@@ -106,7 +139,7 @@ bool FOpenVDBCacheData::LoadFile(FString Path, int frame)
 	// if the file is loaded at the current frame, return true
 	if (OpenVDBGrids.Contains(frame) && OpenVDBGrids[frame] != nullptr)
 	{
-		openvdb::tools::CopyToDense<Vec4STree, openvdb::tools::Dense<openvdb::Vec4s, openvdb::tools::MemoryLayout::LayoutXYZ>> Copier(OpenVDBGrids[frame]->tree(), *DenseGridPtr);
+		openvdb::tools::CopyToDense<Vec4Tree, Vec4Dense> Copier(OpenVDBGrids[frame]->tree(), *DenseGridPtr);
 		Copier.copy();
 
 		return true;
@@ -130,8 +163,8 @@ bool FOpenVDBCacheData::LoadFile(FString Path, int frame)
 			return false;
 		}
 
-		Vec4SGrid::Ptr ColorGrid;		
-		ColorGrid = openvdb::gridPtrCast<Vec4SGrid>(file.readGrid(openvdb::Name("Color")));		
+		Vec4Grid::Ptr ColorGrid;		
+		ColorGrid = openvdb::gridPtrCast<Vec4Grid>(file.readGrid(openvdb::Name("Color")));		
 
 		if (ColorGrid == nullptr)
 		{
@@ -159,7 +192,7 @@ bool FOpenVDBCacheData::LoadFile(FString Path, int frame)
 			OpenVDBGrids.Add(frame, ColorGrid);
 
 			// load to dense buffer			
-			openvdb::tools::CopyToDense<Vec4STree, openvdb::tools::Dense<openvdb::Vec4s, openvdb::tools::MemoryLayout::LayoutXYZ>> Copier(ColorGrid->tree(), *DenseGridPtr);
+			openvdb::tools::CopyToDense<Vec4Tree, Vec4Dense> Copier(ColorGrid->tree(), *DenseGridPtr);
 			Copier.copy();
 		}
 		catch (openvdb::Exception e)
@@ -228,7 +261,7 @@ bool FOpenVDBCacheData::Fill3DTexture_RenderThread(int frame, FTextureRHIRef Tex
 	return false;
 }
 
-bool FOpenVDBCacheData::WriteImageDataToOpenVDBFile(FStringView FilePath, FIntVector ImageSize, TArrayView<FFloat16Color> ImageData, bool UseFloatGrids)
+bool OpenVDBTools::WriteImageDataToOpenVDBFile(FStringView FilePath, FIntVector ImageSize, TArrayView<FFloat16Color> ImageData, bool UseFloatGrids)
 {
 	openvdb::math::Transform::Ptr TransformPtr = openvdb::math::Transform::createLinearTransform(/*voxel size=*/1.0);
 
@@ -237,7 +270,7 @@ bool FOpenVDBCacheData::WriteImageDataToOpenVDBFile(FStringView FilePath, FIntVe
 	openvdb::FloatGrid::Ptr BGrid;
 	openvdb::FloatGrid::Ptr AGrid;
 
-	Vec4SGrid::Ptr ColorGrid;
+	Vec4Grid::Ptr ColorGrid;
 
 	if (UseFloatGrids)
 	{
@@ -272,21 +305,21 @@ bool FOpenVDBCacheData::WriteImageDataToOpenVDBFile(FStringView FilePath, FIntVe
 					uint32 LinearIndex = x + y * ImageSize.X + z * ImageSize.X * ImageSize.Y;
 					FFloat16Color CurrValue = ImageData[LinearIndex];
 
-					Raccessor.setValue(openvdb::Coord(x, y, z), CurrValue.R);
-					Gaccessor.setValue(openvdb::Coord(x, y, z), CurrValue.G);
-					Baccessor.setValue(openvdb::Coord(x, y, z), CurrValue.B);
-					Aaccessor.setValue(openvdb::Coord(x, y, z), CurrValue.A);
+					Raccessor.setValue(openvdb::Coord(x, y, z), CurrValue.R.GetFloat());
+					Gaccessor.setValue(openvdb::Coord(x, y, z), CurrValue.G.GetFloat());
+					Baccessor.setValue(openvdb::Coord(x, y, z), CurrValue.B.GetFloat());
+					Aaccessor.setValue(openvdb::Coord(x, y, z), CurrValue.A.GetFloat());
 				}	
 			}
 		}
 	}
 	else
 	{
-		ColorGrid = Vec4SGrid::create(openvdb::Vec4s(0.0));
+		ColorGrid = Vec4Grid::create(Vec4(0.0));
 		ColorGrid->setTransform(TransformPtr);
 		ColorGrid->setGridClass(openvdb::GRID_FOG_VOLUME);
 		ColorGrid->setName("Color");
-		Vec4SGrid::Accessor ColorAccessor = ColorGrid->getAccessor();
+		Vec4Grid::Accessor ColorAccessor = ColorGrid->getAccessor();
 
 		for (int z = 0; z < ImageSize.Z; ++z) {
 			for (int y = 0; y < ImageSize.Y; ++y) {
@@ -294,7 +327,7 @@ bool FOpenVDBCacheData::WriteImageDataToOpenVDBFile(FStringView FilePath, FIntVe
 					uint32 LinearIndex = x + y * ImageSize.X + z * ImageSize.X * ImageSize.Y;
 					FFloat16Color CurrValue = ImageData[LinearIndex];
 
-					ColorAccessor.setValue(openvdb::Coord(x, y, z), openvdb::Vec4s(CurrValue.R, CurrValue.G, CurrValue.B, CurrValue.A));
+					ColorAccessor.setValue(openvdb::Coord(x, y, z), Vec4(CurrValue.R.GetFloat(), CurrValue.G.GetFloat(), CurrValue.B.GetFloat(), CurrValue.A.GetFloat()));
 				}
 			}
 		}
diff --git a/Engine/Plugins/FX/Niagara/Source/NiagaraEditor/NiagaraEditor.Build.cs b/Engine/Plugins/FX/Niagara/Source/NiagaraEditor/NiagaraEditor.Build.cs
--- a/Engine/Plugins/FX/Niagara/Source/NiagaraEditor/NiagaraEditor.Build.cs
+++ b/Engine/Plugins/FX/Niagara/Source/NiagaraEditor/NiagaraEditor.Build.cs
@@ -6,13 +6,6 @@ public class NiagaraEditor : ModuleRules
 {
 	public NiagaraEditor(ReadOnlyTargetRules Target) : base(Target)
 	{
-		// Specific to OpenVDB support
-		if (Target.IsInPlatformGroup(UnrealPlatformGroup.Windows))
-		{
-			bUseRTTI = true;
-			bEnableExceptions = true;
-		}
-
 		PrivateIncludePaths.AddRange(new string[] {
 			"NiagaraEditor/Private",
 			"NiagaraEditor/Private/Toolkits",
@@ -99,17 +92,6 @@ public class NiagaraEditor : ModuleRules
             }
         );
 
-		if (Target.IsInPlatformGroup(UnrealPlatformGroup.Windows))
-		{
-			AddEngineThirdPartyPrivateStaticDependencies(Target,
-				"IntelTBB",
-				"Blosc",
-				"zlib",
-				"Boost",
-				"OpenVDB"
-			);
-		}
-
 		PublicIncludePathModuleNames.AddRange(
             new string[] {
 				"Engine",
diff --git a/Engine/Plugins/FX/Niagara/Source/NiagaraEditor/Private/NiagaraBakerRenderer.cpp b/Engine/Plugins/FX/Niagara/Source/NiagaraEditor/Private/NiagaraBakerRenderer.cpp
--- a/Engine/Plugins/FX/Niagara/Source/NiagaraEditor/Private/NiagaraBakerRenderer.cpp
+++ b/Engine/Plugins/FX/Niagara/Source/NiagaraEditor/Private/NiagaraBakerRenderer.cpp
@@ -739,7 +739,7 @@ bool FNiagaraBakerRenderer::ExportVolume(FStringView FilePath, FIntVector ImageS
 	if (FileExtension == TEXT(".vdb"))
 	{
 #if PLATFORM_WINDOWS
-		return FOpenVDBCacheData::WriteImageDataToOpenVDBFile(FilePath, ImageSize, ImageData, false);
+		return OpenVDBTools::WriteImageDataToOpenVDBFile(FilePath, ImageSize, ImageData, false);
 #else
 		return false;
 #endif
diff --git a/Engine/Plugins/FX/Niagara/Source/NiagaraEditor/Private/NiagaraEditorModule.cpp b/Engine/Plugins/FX/Niagara/Source/NiagaraEditor/Private/NiagaraEditorModule.cpp
--- a/Engine/Plugins/FX/Niagara/Source/NiagaraEditor/Private/NiagaraEditorModule.cpp
+++ b/Engine/Plugins/FX/Niagara/Source/NiagaraEditor/Private/NiagaraEditorModule.cpp
@@ -151,10 +151,6 @@
 #include "NiagaraGraphDataCache.h"
 #include "Misc/ScopedSlowTask.h"
 
-#if PLATFORM_WINDOWS
-#include "NiagaraOpenVDB.h"
-#endif
-
 IMPLEMENT_MODULE( FNiagaraEditorModule, NiagaraEditor );
 
 #define LOCTEXT_NAMESPACE "NiagaraEditorModule"
@@ -822,19 +818,6 @@ void FNiagaraEditorModule::PostGarbageCollect()
 
 void FNiagaraEditorModule::StartupModule()
 {
-	// Global registration of  the vdb types.
-#if PLATFORM_WINDOWS
-	openvdb::initialize();
-	try
-	{
-		Vec4SGrid::registerGrid();
-	}
-	catch (openvdb::Exception e)
-	{
-		UE_LOG(LogNiagaraEditor, Warning, TEXT("OpenVDB exception: %s"), ANSI_TO_TCHAR(e.what()));
-	}
-#endif	
-	
 	bThumbnailRenderersRegistered = false;
 
 	FHlslNiagaraTranslator::Init();
