Added an option to collapse disabled Niagara modules

diff --git a/Engine/Plugins/FX/Niagara/Source/NiagaraEditor/Private/NiagaraEditorCommands.cpp b/Engine/Plugins/FX/Niagara/Source/NiagaraEditor/Private/NiagaraEditorCommands.cpp
--- a/Engine/Plugins/FX/Niagara/Source/NiagaraEditor/Private/NiagaraEditorCommands.cpp
+++ b/Engine/Plugins/FX/Niagara/Source/NiagaraEditor/Private/NiagaraEditorCommands.cpp
@@ -48,6 +48,7 @@ void FNiagaraEditorCommands::RegisterCommands()
 
 	UI_COMMAND(IsolateSelectedEmitters, "Isolate selected emitters", "Isolate all currently selected emitters.", EUserInterfaceActionType::Button, FInputChord(EKeys::S));
 	UI_COMMAND(DisableSelectedEmitters, "Disable selected emitters", "Disables all currently selected emitters and recompiles the system.", EUserInterfaceActionType::Button, FInputChord(EKeys::D));
+	UI_COMMAND(HideDisabledModules, "Collapse disabled modules", "Collapses disabled modules in the stack, so they won't take up as much space.", EUserInterfaceActionType::Check, FInputChord(EModifierKey::Control, EKeys::M));
 	
 	UI_COMMAND(OpenDebugHUD, "Debug HUD", "Open the Niagara Debug HUD.", EUserInterfaceActionType::Button, FInputChord());
 	UI_COMMAND(OpenDebugOutliner, "FX Outliner", "Open the Niagara FX Outliner.", EUserInterfaceActionType::Button, FInputChord());
diff --git a/Engine/Plugins/FX/Niagara/Source/NiagaraEditor/Private/ViewModels/Stack/NiagaraStackViewModel.cpp b/Engine/Plugins/FX/Niagara/Source/NiagaraEditor/Private/ViewModels/Stack/NiagaraStackViewModel.cpp
--- a/Engine/Plugins/FX/Niagara/Source/NiagaraEditor/Private/ViewModels/Stack/NiagaraStackViewModel.cpp
+++ b/Engine/Plugins/FX/Niagara/Source/NiagaraEditor/Private/ViewModels/Stack/NiagaraStackViewModel.cpp
@@ -407,6 +407,16 @@ TSharedPtr<UNiagaraStackViewModel::FTopLevelViewModel> UNiagaraStackViewModel::G
 	return TSharedPtr<FTopLevelViewModel>();
 }
 
+bool UNiagaraStackViewModel::ShouldHideDisabledModules() const
+{
+	if (SystemViewModel.IsValid())
+	{
+		UNiagaraSystemEditorData& EditorData = SystemViewModel.Pin()->GetEditorData();
+		return EditorData.GetStackEditorData().bHideDisabledModules;
+	}
+	return false;
+}
+
 void UNiagaraStackViewModel::CollapseToHeadersRecursive(TArray<UNiagaraStackEntry*> Entries)
 {
 	for (UNiagaraStackEntry* Entry : Entries)
diff --git a/Engine/Plugins/FX/Niagara/Source/NiagaraEditor/Public/NiagaraEditorCommands.h b/Engine/Plugins/FX/Niagara/Source/NiagaraEditor/Public/NiagaraEditorCommands.h
--- a/Engine/Plugins/FX/Niagara/Source/NiagaraEditor/Public/NiagaraEditorCommands.h
+++ b/Engine/Plugins/FX/Niagara/Source/NiagaraEditor/Public/NiagaraEditorCommands.h
@@ -73,6 +73,7 @@ public:
 
 	TSharedPtr<FUICommandInfo> IsolateSelectedEmitters;
 	TSharedPtr<FUICommandInfo> DisableSelectedEmitters;
+	TSharedPtr<FUICommandInfo> HideDisabledModules;
 
 	TSharedPtr<FUICommandInfo> CollapseStackToHeaders;
 
diff --git a/Engine/Plugins/FX/Niagara/Source/NiagaraEditor/Public/NiagaraStackEditorData.h b/Engine/Plugins/FX/Niagara/Source/NiagaraEditor/Public/NiagaraStackEditorData.h
--- a/Engine/Plugins/FX/Niagara/Source/NiagaraEditor/Public/NiagaraStackEditorData.h
+++ b/Engine/Plugins/FX/Niagara/Source/NiagaraEditor/Public/NiagaraStackEditorData.h
@@ -156,6 +156,9 @@ public:
 	/* Gets a reference to the dismissed stack issue array */
 	NIAGARAEDITOR_API const TArray<FString>& GetDismissedStackIssueIds();
 
+	UPROPERTY()
+	bool bHideDisabledModules = false;
+
 private:
 	TMap<FString, bool> StackEntryKeyToRenamePendingMap;
 
diff --git a/Engine/Plugins/FX/Niagara/Source/NiagaraEditor/Public/ViewModels/Stack/NiagaraStackViewModel.h b/Engine/Plugins/FX/Niagara/Source/NiagaraEditor/Public/ViewModels/Stack/NiagaraStackViewModel.h
--- a/Engine/Plugins/FX/Niagara/Source/NiagaraEditor/Public/ViewModels/Stack/NiagaraStackViewModel.h
+++ b/Engine/Plugins/FX/Niagara/Source/NiagaraEditor/Public/ViewModels/Stack/NiagaraStackViewModel.h
@@ -133,7 +133,6 @@ public:
 	void CollapseToHeaders();
 	
 	void UndismissAllIssues();
-
 	bool HasDismissedStackIssues();
 
 	const TArray<TSharedRef<FTopLevelViewModel>>& GetTopLevelViewModels() const;
@@ -141,13 +140,12 @@ public:
 	TSharedPtr<FTopLevelViewModel> GetTopLevelViewModelForEntry(UNiagaraStackEntry& InEntry) const;
 
 	void Reset();
-
 	bool HasIssues() const;
-
 	void Refresh() { bRefreshPending = true; }
-
 	void RequestValidationUpdate() { bValidatorUpdatePending = true; }
 
+	bool ShouldHideDisabledModules() const;
+
 	void InvalidateCachedParameterUsage();
 
 private:
diff --git a/Engine/Plugins/FX/Niagara/Source/NiagaraEditorWidgets/Private/NiagaraEditorWidgetsUtilities.cpp b/Engine/Plugins/FX/Niagara/Source/NiagaraEditorWidgets/Private/NiagaraEditorWidgetsUtilities.cpp
--- a/Engine/Plugins/FX/Niagara/Source/NiagaraEditorWidgets/Private/NiagaraEditorWidgetsUtilities.cpp
+++ b/Engine/Plugins/FX/Niagara/Source/NiagaraEditorWidgets/Private/NiagaraEditorWidgetsUtilities.cpp
@@ -21,9 +21,11 @@
 #include "Subsystems/AssetEditorSubsystem.h"
 #include "Editor.h"
 #include "EditorFontGlyphs.h"
+#include "NiagaraEditorCommands.h"
 #include "NiagaraEmitter.h"
 #include "NiagaraEmitterEditorData.h"
 #include "NiagaraMessages.h"
+#include "NiagaraSystemEditorData.h"
 #include "ViewModels/NiagaraEmitterViewModel.h"
 #include "ViewModels/NiagaraSystemSelectionViewModel.h"
 #include "ViewModels/NiagaraSystemViewModel.h"
@@ -374,6 +376,11 @@ bool FNiagaraStackEditorWidgetsUtilities::AddStackItemContextMenuActions(FMenuBu
 			}
 
 			UNiagaraStackModuleItem* ModuleItem = Cast<UNiagaraStackModuleItem>(&StackItem);
+
+			if (ModuleItem && ModuleItem->GetIsEnabled() == false)
+			{
+				MenuBuilder.AddMenuEntry(FNiagaraEditorCommands::Get().HideDisabledModules);
+			}
 			
 			if (ModuleItem && ModuleItem->GetModuleNode().ContainsDebugSwitch())
 			{
diff --git a/Engine/Plugins/FX/Niagara/Source/NiagaraEditorWidgets/Private/NiagaraStackCommandContext.cpp b/Engine/Plugins/FX/Niagara/Source/NiagaraEditorWidgets/Private/NiagaraStackCommandContext.cpp
--- a/Engine/Plugins/FX/Niagara/Source/NiagaraEditorWidgets/Private/NiagaraStackCommandContext.cpp
+++ b/Engine/Plugins/FX/Niagara/Source/NiagaraEditorWidgets/Private/NiagaraStackCommandContext.cpp
@@ -1,8 +1,11 @@
 // Copyright Epic Games, Inc. All Rights Reserved.
 
 #include "NiagaraStackCommandContext.h"
+
+#include "NiagaraEditorCommands.h"
 #include "NiagaraEditorCommon.h"
 #include "NiagaraEditorUtilities.h"
+#include "NiagaraSystemEditorData.h"
 #include "ViewModels/Stack/NiagaraStackClipboardUtilities.h"
 #include "ViewModels/Stack/NiagaraStackEntry.h"
 
@@ -12,6 +15,7 @@
 #include "Framework/Commands/GenericCommands.h"
 #include "Framework/Commands/UICommandList.h"
 #include "Framework/MultiBox/MultiBoxBuilder.h"
+#include "ViewModels/NiagaraSystemViewModel.h"
 
 #define LOCTEXT_NAMESPACE "NiagaraStackCommandContext"
 
@@ -103,6 +107,27 @@ void FNiagaraStackCommandContext::SetupCommands()
 	Commands->MapAction(FGenericCommands::Get().Rename, FUIAction(
 		FExecuteAction::CreateSP(this, &FNiagaraStackCommandContext::RenameSelectedEntries),
 		FCanExecuteAction::CreateSP(this, &FNiagaraStackCommandContext::CanRenameSelectedEntries)));
+
+	Commands->MapAction(
+		FNiagaraEditorCommands::Get().HideDisabledModules,
+		FExecuteAction::CreateLambda([this]()
+		{
+			if (SelectedEntries.Num() > 0)
+			{
+				UNiagaraStackEditorData& EditorData = SelectedEntries[0]->GetSystemViewModel()->GetEditorData().GetStackEditorData();
+				EditorData.bHideDisabledModules = !EditorData.bHideDisabledModules;
+			}
+		}),
+		FCanExecuteAction(),
+		FIsActionChecked::CreateLambda([this]()
+		{
+			if (SelectedEntries.Num() > 0)
+			{
+				UNiagaraStackEditorData& EditorData = SelectedEntries[0]->GetSystemViewModel()->GetEditorData().GetStackEditorData();
+				return EditorData.bHideDisabledModules;
+			}
+			return false;
+		}));
 }
 
 bool FNiagaraStackCommandContext::CanCutSelectedEntries() const
diff --git a/Engine/Plugins/FX/Niagara/Source/NiagaraEditorWidgets/Private/SNiagaraOverviewGraph.cpp b/Engine/Plugins/FX/Niagara/Source/NiagaraEditorWidgets/Private/SNiagaraOverviewGraph.cpp
--- a/Engine/Plugins/FX/Niagara/Source/NiagaraEditorWidgets/Private/SNiagaraOverviewGraph.cpp
+++ b/Engine/Plugins/FX/Niagara/Source/NiagaraEditorWidgets/Private/SNiagaraOverviewGraph.cpp
@@ -102,6 +102,20 @@ void SNiagaraOverviewGraph::Construct(const FArguments& InArgs, TSharedRef<FNiag
 		FIsActionChecked(),
 		FIsActionButtonVisible::CreateSP(this, &SNiagaraOverviewGraph::CanAddEmitters)
 	);
+
+	Commands->MapAction(
+		FNiagaraEditorCommands::Get().HideDisabledModules,
+		FExecuteAction::CreateLambda([=]()
+		{
+			UNiagaraStackEditorData& EditorData = ViewModel->GetSystemViewModel()->GetEditorData().GetStackEditorData();
+			EditorData.bHideDisabledModules = !EditorData.bHideDisabledModules;
+		}),
+		FCanExecuteAction(),
+		FIsActionChecked::CreateLambda([=]()
+		{
+			UNiagaraStackEditorData& EditorData = ViewModel->GetSystemViewModel()->GetEditorData().GetStackEditorData();
+			return EditorData.bHideDisabledModules;
+		}));
 	
 	GraphEditor = SNew(SGraphEditor)
 		.AdditionalCommands(Commands)
diff --git a/Engine/Plugins/FX/Niagara/Source/NiagaraEditorWidgets/Private/SNiagaraOverviewStack.cpp b/Engine/Plugins/FX/Niagara/Source/NiagaraEditorWidgets/Private/SNiagaraOverviewStack.cpp
--- a/Engine/Plugins/FX/Niagara/Source/NiagaraEditorWidgets/Private/SNiagaraOverviewStack.cpp
+++ b/Engine/Plugins/FX/Niagara/Source/NiagaraEditorWidgets/Private/SNiagaraOverviewStack.cpp
@@ -1,6 +1,8 @@
 // Copyright Epic Games, Inc. All Rights Reserved.
 
 #include "SNiagaraOverviewStack.h"
+
+#include "IDocumentation.h"
 #include "NiagaraSystem.h"
 #include "NiagaraSystemEditorData.h"
 #include "ScopedTransaction.h"
@@ -43,19 +45,16 @@
 #include "Widgets/Images/SImage.h"
 #include "Widgets/Input/SButton.h"
 #include "Widgets/Input/SCheckBox.h"
-#include "Widgets/SToolTip.h"
 #include "Styling/AppStyle.h"
 #include "Framework/Application/SlateApplication.h"
 #include "Framework/MultiBox/MultiBoxBuilder.h"
 #include "Subsystems/AssetEditorSubsystem.h"
-#include "Widgets/Colors/SColorBlock.h"
 #include "NiagaraEditorStyle.h"
 #include "NiagaraEditorUtilities.h"
 #include "ScopedTransaction.h"
 #include "Styling/StyleColors.h"
-#include "ViewModels/NiagaraOverviewGraphViewModel.h"
 #include "ViewModels/Stack/NiagaraStackRendererItem.h"
-#include "Widgets/Input/SCheckBox.h"
+#include "Widgets/Layout/SWidgetSwitcher.h"
 #include "Widgets/Layout/SWrapBox.h"
 
 #define LOCTEXT_NAMESPACE "NiagaraOverviewStack"
@@ -289,9 +288,16 @@ class SNiagaraSystemOverviewEntryListRow : public STableRow<UNiagaraStackEntry*>
 
 	virtual FReply OnMouseButtonDoubleClick(const FGeometry& InMyGeometry, const FPointerEvent& InMouseEvent) override
 	{
-		UNiagaraStackModuleItem* ModuleItem = Cast<UNiagaraStackModuleItem>(StackEntry);
-		if (ModuleItem != nullptr)
+		if (UNiagaraStackModuleItem* ModuleItem = Cast<UNiagaraStackModuleItem>(StackEntry))
 		{
+			// if the user double-clicks a collapsed module we uncollapse it
+			UNiagaraStackEditorData& StackEditorData = ModuleItem->GetSystemViewModel()->GetEditorData().GetStackEditorData();
+			if (ModuleItem->GetIsEnabled() == false && StackEditorData.bHideDisabledModules)
+			{
+				StackEditorData.bHideDisabledModules = false;
+				return FReply::Handled();
+			}
+			
 			const UNiagaraNodeFunctionCall& ModuleFunctionCall = ModuleItem->GetModuleNode();
 			if (ModuleFunctionCall.FunctionScript != nullptr)
 			{
@@ -382,9 +388,9 @@ private:
 	}
 
 private:
-	UNiagaraStackViewModel* StackViewModel;
-	UNiagaraStackEntry* StackEntry;
-	UNiagaraSystemScalabilityViewModel* ScalabilityViewModel;
+	UNiagaraStackViewModel* StackViewModel = nullptr;
+	UNiagaraStackEntry* StackEntry = nullptr;
+	UNiagaraSystemScalabilityViewModel* ScalabilityViewModel = nullptr;
 	TSharedPtr<SOverlay> ScalabilityOverlay;
 	TSharedPtr<FNiagaraStackCommandContext> StackCommandContext;
 	FSlateColor ExcludedFromScalabilityColor;
@@ -392,7 +398,7 @@ private:
 	TWeakPtr<FNiagaraEmitterHandleViewModel> EmitterHandleViewModel;
 	const FSlateBrush* ExpandedImage = nullptr;
 	const FSlateBrush* CollapsedImage = nullptr;
-	bool bScalabilityModeActive;
+	bool bScalabilityModeActive = false;
 	mutable TOptional<uint32> ChildrenCount;
 	mutable TOptional<FSlateColor> IssueHighlightColor;
 };
@@ -1112,8 +1118,25 @@ TSharedRef<ITableRow> SNiagaraOverviewStack::OnGenerateRowForEntry(UNiagaraStack
 				OptionsBox
 			];
 		}
-		
-		Content = ContentBox;	
+
+		if (StackModuleItem)
+		{
+			SAssignNew(Content, SWidgetSwitcher)
+		   .WidgetIndex_Lambda([StackModuleItem, this]() { return (StackViewModel->ShouldHideDisabledModules() && StackModuleItem->GetIsEnabled() == false) ? 1 : 0; })
+		   + SWidgetSwitcher::Slot()
+		   [
+			   ContentBox
+		   ]
+		   + SWidgetSwitcher::Slot()
+		   [
+			   SNew(SBox)
+			   .HeightOverride(4)
+		   ];
+		}
+		else
+		{
+			Content = ContentBox;
+		}
 		RowStyle = &FNiagaraEditorWidgetsStyle::Get().GetWidgetStyle<FTableRowStyle>("NiagaraEditor.SystemOverview.TableViewRow.Item");
 	}
 	else if (Item->IsA<UNiagaraStackItemGroup>())
