- Added scalability indicator on emitter & system nodes
- Added ability to change workflow mode from the system view model (currently not used)
- removed fx.Niagara.ScalabilityModeEnabled cvar

diff --git a/Engine/Plugins/FX/Niagara/Source/Niagara/Classes/NiagaraPlatformSet.h b/Engine/Plugins/FX/Niagara/Source/Niagara/Classes/NiagaraPlatformSet.h
--- a/Engine/Plugins/FX/Niagara/Source/Niagara/Classes/NiagaraPlatformSet.h
+++ b/Engine/Plugins/FX/Niagara/Source/Niagara/Classes/NiagaraPlatformSet.h
@@ -320,6 +320,8 @@
 	static int32 GetMinQualityLevel();
 	static int32 GetMaxQualityLevel();
 	static int32 GetAvailableQualityLevelMask();
+	static int32 GetFullQualityLevelMask(int32 NumQualityLevels);
+	
 public:
 
 	//Runtime public API
diff --git a/Engine/Plugins/FX/Niagara/Source/Niagara/Private/NiagaraPlatformSet.cpp b/Engine/Plugins/FX/Niagara/Source/Niagara/Private/NiagaraPlatformSet.cpp
--- a/Engine/Plugins/FX/Niagara/Source/Niagara/Private/NiagaraPlatformSet.cpp
+++ b/Engine/Plugins/FX/Niagara/Source/Niagara/Private/NiagaraPlatformSet.cpp
@@ -238,6 +238,18 @@
 	return CachedAvailableQualityLevelMask;
 }
 
+int32 FNiagaraPlatformSet::GetFullQualityLevelMask(int32 NumQualityLevels)
+{
+	int32 QualityLevelMask = 0;
+	
+	for (int32 QL = 0; QL < NumQualityLevels; ++QL)
+	{
+		QualityLevelMask |= (1 << QL);
+	}
+
+	return QualityLevelMask;
+}
+
 void FNiagaraPlatformSet::OnQualityLevelChanged(IConsoleVariable* Variable)
 {
 	int32 NewQualityLevel = Variable->GetInt();
diff --git a/Engine/Plugins/FX/Niagara/Source/NiagaraEditor/Private/NiagaraEmitterDetailsCustomization.cpp b/Engine/Plugins/FX/Niagara/Source/NiagaraEditor/Private/NiagaraEmitterDetailsCustomization.cpp
--- a/Engine/Plugins/FX/Niagara/Source/NiagaraEditor/Private/NiagaraEmitterDetailsCustomization.cpp
+++ b/Engine/Plugins/FX/Niagara/Source/NiagaraEditor/Private/NiagaraEmitterDetailsCustomization.cpp
@@ -46,11 +46,9 @@
 			}
 		}
 	}
-	
-	if(GNiagaraScalabilityModeEnabled)
-	{
-		InDetailLayout.HideCategory("Scalability");
-	}
+
+	// we display the scalability category within scalability mode, which is why we hide it here
+	InDetailLayout.HideCategory("Scalability");
 }
 
 TSharedRef<IDetailCustomization> FNiagaraEmitterScalabilityDetails::MakeInstance()
diff --git a/Engine/Plugins/FX/Niagara/Source/NiagaraEditor/Private/NiagaraSystemDetailsCustomization.cpp b/Engine/Plugins/FX/Niagara/Source/NiagaraEditor/Private/NiagaraSystemDetailsCustomization.cpp
--- a/Engine/Plugins/FX/Niagara/Source/NiagaraEditor/Private/NiagaraSystemDetailsCustomization.cpp
+++ b/Engine/Plugins/FX/Niagara/Source/NiagaraEditor/Private/NiagaraSystemDetailsCustomization.cpp
@@ -1,21 +1,19 @@
-// Copyright Epic Games, Inc. All Rights Reserved.
-
-#include "NiagaraSystemDetailsCustomization.h"
-#include "PropertyHandle.h"
-#include "DetailLayoutBuilder.h"
-#include "DetailCategoryBuilder.h"
-#include "NiagaraSystem.h"
-#include "Toolkits/NiagaraSystemToolkit.h"
-
-TSharedRef<IDetailCustomization> FNiagaraSystemDetails::MakeInstance()
-{
-	return MakeShared<FNiagaraSystemDetails>();
-}
-
-void FNiagaraSystemDetails::CustomizeDetails(IDetailLayoutBuilder& InDetailLayout)
-{
-	if(GNiagaraScalabilityModeEnabled)
-	{
-		InDetailLayout.HideCategory("Scalability");
-	}
-}
+// Copyright Epic Games, Inc. All Rights Reserved.
+
+#include "NiagaraSystemDetailsCustomization.h"
+#include "PropertyHandle.h"
+#include "DetailLayoutBuilder.h"
+#include "DetailCategoryBuilder.h"
+#include "NiagaraSystem.h"
+#include "Toolkits/NiagaraSystemToolkit.h"
+
+TSharedRef<IDetailCustomization> FNiagaraSystemDetails::MakeInstance()
+{
+	return MakeShared<FNiagaraSystemDetails>();
+}
+
+void FNiagaraSystemDetails::CustomizeDetails(IDetailLayoutBuilder& InDetailLayout)
+{
+	// we display the scalability category within scalability mode, which is why we hide it here
+	InDetailLayout.HideCategory("Scalability");
+}
diff --git a/Engine/Plugins/FX/Niagara/Source/NiagaraEditor/Private/Toolkits/NiagaraSystemToolkit.cpp b/Engine/Plugins/FX/Niagara/Source/NiagaraEditor/Private/Toolkits/NiagaraSystemToolkit.cpp
--- a/Engine/Plugins/FX/Niagara/Source/NiagaraEditor/Private/Toolkits/NiagaraSystemToolkit.cpp
+++ b/Engine/Plugins/FX/Niagara/Source/NiagaraEditor/Private/Toolkits/NiagaraSystemToolkit.cpp
@@ -68,17 +68,6 @@
 const FName FNiagaraSystemToolkit::DefaultModeName(TEXT("Default"));
 const FName FNiagaraSystemToolkit::ScalabilityModeName(TEXT("Scalability"));
 
-int32 GNiagaraScalabilityModeEnabled = 1;
-
-// @TODO Remove cvar when scalability mode is going in properly
-static const TCHAR* NiagaraEnableScalabilityMode = TEXT("fx.Niagara.EnableScalabilityMode");
-static FAutoConsoleVariableRef CVarNiagaraScalabilityModeEnabled(
-	NiagaraEnableScalabilityMode,
-	GNiagaraScalabilityModeEnabled,
-	TEXT("Enable the scalability mode inside Niagara emitters and systems. \n"),
-	ECVF_Default
-);
-
 IConsoleVariable* FNiagaraSystemToolkit::VmStatEnabledVar = IConsoleManager::Get().FindConsoleVariable(TEXT("vm.DetailedVMScriptStats"));
 
 static int32 GbLogNiagaraSystemChanges = 0;
@@ -294,6 +283,7 @@
 	SystemViewModel->GetOnPinnedEmittersChanged().AddSP(this, &FNiagaraSystemToolkit::RefreshParameters);
 	SystemViewModel->OnRequestFocusTab().AddSP(this, &FNiagaraSystemToolkit::OnViewModelRequestFocusTab);
 	SystemViewModel->OnGetWorkflowMode().BindSP(this, &FNiagaraSystemToolkit::GetCurrentMode);
+	SystemViewModel->OnChangeWorkflowMode().BindSP(this, &FNiagaraSystemToolkit::SetCurrentMode);
 	SystemViewModel->GetDocumentViewModel()->InitializePreTabManager(SharedThis(this));
 	
 	constexpr bool bCreateDefaultStandaloneMenu = true;
diff --git a/Engine/Plugins/FX/Niagara/Source/NiagaraEditor/Private/Toolkits/NiagaraSystemToolkit.h b/Engine/Plugins/FX/Niagara/Source/NiagaraEditor/Private/Toolkits/NiagaraSystemToolkit.h
--- a/Engine/Plugins/FX/Niagara/Source/NiagaraEditor/Private/Toolkits/NiagaraSystemToolkit.h
+++ b/Engine/Plugins/FX/Niagara/Source/NiagaraEditor/Private/Toolkits/NiagaraSystemToolkit.h
@@ -41,8 +41,6 @@
 class FNiagaraScriptStatsViewModel;
 class FNiagaraBakerViewModel;
 
-extern int32 GNiagaraScalabilityModeEnabled;
-
 /** Viewer/editor for a NiagaraSystem
 */
 class FNiagaraSystemToolkit : public FWorkflowCentricApplication, public FGCObject, public FEditorUndoClient
diff --git a/Engine/Plugins/FX/Niagara/Source/NiagaraEditor/Private/Toolkits/SystemToolkitModes/NiagaraSystemToolkitMode_Default.cpp b/Engine/Plugins/FX/Niagara/Source/NiagaraEditor/Private/Toolkits/SystemToolkitModes/NiagaraSystemToolkitMode_Default.cpp
--- a/Engine/Plugins/FX/Niagara/Source/NiagaraEditor/Private/Toolkits/SystemToolkitModes/NiagaraSystemToolkitMode_Default.cpp
+++ b/Engine/Plugins/FX/Niagara/Source/NiagaraEditor/Private/Toolkits/SystemToolkitModes/NiagaraSystemToolkitMode_Default.cpp
@@ -243,29 +243,26 @@
 				);
 			}
 			ToolbarBuilder.EndSection();
+			
+			ToolbarBuilder.BeginSection("Scalability");
+			{
+				FUIAction ScalabilityToggleAction(FExecuteAction::CreateRaw(Toolkit, &FNiagaraSystemToolkit::SetCurrentMode, FNiagaraSystemToolkit::ScalabilityModeName));
+				ScalabilityToggleAction.GetActionCheckState = FGetActionCheckState::CreateLambda([Toolkit]()
+				{
+					return Toolkit->GetCurrentMode() == FNiagaraSystemToolkit::ScalabilityModeName ? ECheckBoxState::Checked : ECheckBoxState::Unchecked;
+				});
+				
+				ToolbarBuilder.AddToolBarButton(
+					ScalabilityToggleAction,
+					NAME_None, 
+					LOCTEXT("ScalabilityLabel", "Scalability"),
+					LOCTEXT("ScalabilityTooltip", "Turn on scalability mode to optimize your effects for various platforms and quality settings."),
+					FSlateIcon(FNiagaraEditorStyle::Get().GetStyleSetName(), "NiagaraEditor.Scalability"),
+					EUserInterfaceActionType::ToggleButton
+				);
+			}
+			ToolbarBuilder.EndSection(
 
-			if(GNiagaraScalabilityModeEnabled)
-			{
-				ToolbarBuilder.BeginSection("Scalability");
-				{
-					FUIAction ScalabilityToggleAction(FExecuteAction::CreateRaw(Toolkit, &FNiagaraSystemToolkit::SetCurrentMode, FNiagaraSystemToolkit::ScalabilityModeName));
-					ScalabilityToggleAction.GetActionCheckState = FGetActionCheckState::CreateLambda([Toolkit]()
-					{
-						return Toolkit->GetCurrentMode() == FNiagaraSystemToolkit::ScalabilityModeName ? ECheckBoxState::Checked : ECheckBoxState::Unchecked;
-					});
-					
-					ToolbarBuilder.AddToolBarButton(
-						ScalabilityToggleAction,
-						NAME_None, 
-						LOCTEXT("ScalabilityLabel", "Scalability"),
-						LOCTEXT("ScalabilityTooltip", "Turn on scalability mode to optimize your effects for various platforms and quality settings."),
-						FSlateIcon(FNiagaraEditorStyle::Get().GetStyleSetName(), "NiagaraEditor.Scalability"),
-						EUserInterfaceActionType::ToggleButton
-					);
-				}
-				ToolbarBuilder.EndSection();
-			}
-
 			if (Toolkit->HasEmitter())
 			{
 				ToolbarBuilder.BeginSection("Versioning");
diff --git a/Engine/Plugins/FX/Niagara/Source/NiagaraEditor/Private/ViewModels/NiagaraSystemViewModel.cpp b/Engine/Plugins/FX/Niagara/Source/NiagaraEditor/Private/ViewModels/NiagaraSystemViewModel.cpp
--- a/Engine/Plugins/FX/Niagara/Source/NiagaraEditor/Private/ViewModels/NiagaraSystemViewModel.cpp
+++ b/Engine/Plugins/FX/Niagara/Source/NiagaraEditor/Private/ViewModels/NiagaraSystemViewModel.cpp
@@ -398,6 +398,17 @@
 	return OnGetWorkflowModeDelegate;
 }
 
+void FNiagaraSystemViewModel::SetWorkflowMode(FName WorkflowMode)
+{
+	check(OnChangeWorkflowModeDelegate.IsBound());
+	OnChangeWorkflowModeDelegate.Execute(WorkflowMode);
+}
+
+FNiagaraSystemViewModel::FOnChangeWorkflowMode& FNiagaraSystemViewModel::OnChangeWorkflowMode()
+{
+	return OnChangeWorkflowModeDelegate;
+}
+
 TSharedPtr<FNiagaraEmitterHandleViewModel> FNiagaraSystemViewModel::AddEmitterFromAssetData(const FAssetData& AssetData)
 {
 	UNiagaraEmitter* Emitter = Cast<UNiagaraEmitter>(AssetData.GetAsset());
diff --git a/Engine/Plugins/FX/Niagara/Source/NiagaraEditor/Public/ViewModels/NiagaraSystemViewModel.h b/Engine/Plugins/FX/Niagara/Source/NiagaraEditor/Public/ViewModels/NiagaraSystemViewModel.h
--- a/Engine/Plugins/FX/Niagara/Source/NiagaraEditor/Public/ViewModels/NiagaraSystemViewModel.h
+++ b/Engine/Plugins/FX/Niagara/Source/NiagaraEditor/Public/ViewModels/NiagaraSystemViewModel.h
@@ -126,6 +126,7 @@
 	DECLARE_MULTICAST_DELEGATE_ThreeParams(FOnExternalRenameParameter, const FNiagaraVariableBase&, const FNiagaraVariableBase&, UNiagaraEmitter*);
 	DECLARE_MULTICAST_DELEGATE_TwoParams(FOnExternalRemoveParameter, const FNiagaraVariableBase&,  UNiagaraEmitter*);
 
+	DECLARE_DELEGATE_OneParam(FOnChangeWorkflowMode, FName /* WorkflowMode */)
 public:
 	struct FEmitterHandleToDuplicate
 	{
@@ -216,7 +217,12 @@
 
 	/** Gets the delegate which is called any time the array of emitter handle view models changes. */
 	NIAGARAEDITOR_API FOnGetWorkflowMode& OnGetWorkflowMode();
+
+	NIAGARAEDITOR_API void SetWorkflowMode(FName WorkflowMode);
 	
+	/** Gets the delegate that is bound to the owning toolkit's SetCurrentMode function. */
+	NIAGARAEDITOR_API FOnChangeWorkflowMode& OnChangeWorkflowMode();
+	
 	/** Adds a new emitter to the System from an emitter asset data. */
 	NIAGARAEDITOR_API TSharedPtr<FNiagaraEmitterHandleViewModel> AddEmitterFromAssetData(const FAssetData& AssetData);
 
@@ -594,6 +600,9 @@
 
 	/** A delegate to retrieve the current workflow mode from the toolkit. */
 	FOnGetWorkflowMode OnGetWorkflowModeDelegate;
+
+	/** A delegate that calls the toolkit's set current mode function */
+	FOnChangeWorkflowMode OnChangeWorkflowModeDelegate;
 	
 	/** A delegate which is used to generate the content for the add menu in sequencer. */
 	FOnGetAddMenuContent OnGetSequencerAddMenuContent;
diff --git a/Engine/Plugins/FX/Niagara/Source/NiagaraEditorWidgets/Private/SNiagaraOverviewStackNode.cpp b/Engine/Plugins/FX/Niagara/Source/NiagaraEditorWidgets/Private/SNiagaraOverviewStackNode.cpp
--- a/Engine/Plugins/FX/Niagara/Source/NiagaraEditorWidgets/Private/SNiagaraOverviewStackNode.cpp
+++ b/Engine/Plugins/FX/Niagara/Source/NiagaraEditorWidgets/Private/SNiagaraOverviewStackNode.cpp
@@ -26,6 +26,7 @@
 #include "Widgets/SToolTip.h"
 #include "Widgets/SBoxPanel.h"
 #include "NiagaraRendererProperties.h"
+#include "NiagaraSettings.h"
 #include "NiagaraEditorSettings.h"
 #include "SGraphPanel.h"
 #include "Widgets/SWidget.h"
@@ -194,6 +195,22 @@
 			]
 		]
 
+		// scalability indicator
+		+ SHorizontalBox::Slot()
+		.AutoWidth()
+		[
+			SNew(SBox)
+			.WidthOverride(16.f)
+			.HeightOverride(16.f)
+			[
+				SNew(SImage)
+				.Image(FNiagaraEditorStyle::Get().GetBrush("NiagaraEditor.Scalability"))
+				.Visibility(this, &SNiagaraOverviewStackNode::GetScalabilityIndicatorVisibility)
+				.ToolTipText(FText::FormatOrdered(LOCTEXT("ScalabilityIndicatorToolTip",
+					"This {0} has scalability set up. Inspecting and editing scalability is accessible by using scalability mode from the toolbar."), EmitterHandleViewModelWeak.IsValid() ? FText::FromString("emitter") : FText::FromString("system")))
+			]
+		]
+	
 		// issue/error icon
 		+ SHorizontalBox::Slot()
 		.AutoWidth()
@@ -380,6 +397,29 @@
 	];
 }
 
+EVisibility SNiagaraOverviewStackNode::GetScalabilityIndicatorVisibility() const
+{
+	TSharedPtr<FNiagaraEmitterHandleViewModel> EmitterHandleViewModel = EmitterHandleViewModelWeak.Pin();
+
+	if(EmitterHandleViewModel.IsValid())
+	{
+		if(FVersionedNiagaraEmitterData* EmitterData = EmitterHandleViewModel->GetEmitterHandle()->GetInstance().GetEmitterData())
+		{
+			bool bIsQualityLevelMaskSetup = EmitterData->Platforms.QualityLevelMask != INDEX_NONE;
+			bool bIsScalabilitySetup = EmitterData->ScalabilityOverrides.Overrides.Num() != 0 || (bIsQualityLevelMaskSetup && EmitterData->Platforms.QualityLevelMask != FNiagaraPlatformSet::GetFullQualityLevelMask(GetDefault<UNiagaraSettings>()->QualityLevels.Num())); 
+			return bIsScalabilitySetup ? EVisibility::Visible : EVisibility::Collapsed;
+		}
+	}
+
+	if(UNiagaraSystem* System = OverviewStackNode->GetOwningSystem())
+	{
+		bool bIsScalabilitySetup = System->GetOverrideScalabilitySettings();
+		return bIsScalabilitySetup ? EVisibility::Visible : EVisibility::Collapsed; 
+	}
+	
+	return EVisibility::Collapsed;
+}
+
 TSharedRef<SWidget> SNiagaraOverviewStackNode::CreateNodeContentArea()
 {
 	TSharedPtr<SWidget> ContentWidget;
@@ -813,11 +853,6 @@
 	bScalabilityModeActive = bActive;
 }
 
-EVisibility SNiagaraOverviewStackNode::GetScalabilityBarVisibility() const
-{
-	return bScalabilityModeActive ? EVisibility::Visible : EVisibility::Collapsed;
-}
-
 EVisibility SNiagaraOverviewStackNode::ShowExcludedOverlay() const
 {
 	// we only want actual results in scalability mode and for nodes representing emitters (not system nodes)
diff --git a/Engine/Plugins/FX/Niagara/Source/NiagaraEditorWidgets/Private/SNiagaraOverviewStackNode.h b/Engine/Plugins/FX/Niagara/Source/NiagaraEditorWidgets/Private/SNiagaraOverviewStackNode.h
--- a/Engine/Plugins/FX/Niagara/Source/NiagaraEditorWidgets/Private/SNiagaraOverviewStackNode.h
+++ b/Engine/Plugins/FX/Niagara/Source/NiagaraEditorWidgets/Private/SNiagaraOverviewStackNode.h
@@ -44,6 +44,7 @@
 
 	void CreateBottomSummaryExpander();
 private:
+	EVisibility GetScalabilityIndicatorVisibility() const;
 	EVisibility GetIssueIconVisibility() const;
 	EVisibility GetEnabledCheckBoxVisibility() const;
 	ECheckBoxState GetEnabledCheckState() const;
@@ -56,7 +57,6 @@
 	FSlateColor GetToggleIsolateImageColor() const;
 	FSlateColor GetScalabilityTintAlpha() const;
 	void OnScalabilityModeChanged(bool bActive);
-	EVisibility GetScalabilityBarVisibility() const;
 	EVisibility ShowExcludedOverlay() const;
 	float GetGraphZoomDistanceAlphaMultiplier() const;
 	void SetIsHoveringThumbnail(const FGeometry& InGeometry, const FPointerEvent& InEvent, const bool bInHoveringThumbnail)
@@ -71,7 +71,7 @@
 	{
 		return bIsHoveringThumbnail;
 	}
-
+	
 	FReply OnCycleThroughIssues();
 	FReply OpenParentEmitter();
 	FText OpenParentEmitterTooltip() const;
