Fix compile error

#rnx

diff --git a/Engine/Plugins/FX/Niagara/Source/NiagaraEditor/Private/Toolkits/SystemToolkitModes/NiagaraSystemToolkitMode_Default.cpp b/Engine/Plugins/FX/Niagara/Source/NiagaraEditor/Private/Toolkits/SystemToolkitModes/NiagaraSystemToolkitMode_Default.cpp
--- a/Engine/Plugins/FX/Niagara/Source/NiagaraEditor/Private/Toolkits/SystemToolkitModes/NiagaraSystemToolkitMode_Default.cpp
+++ b/Engine/Plugins/FX/Niagara/Source/NiagaraEditor/Private/Toolkits/SystemToolkitModes/NiagaraSystemToolkitMode_Default.cpp
@@ -261,7 +261,7 @@ void FNiagaraSystemToolkitMode_Default::ExtendToolbar()
 					EUserInterfaceActionType::ToggleButton
 				);
 			}
-			ToolbarBuilder.EndSection(
+			ToolbarBuilder.EndSection();
 
 			if (Toolkit->HasEmitter())
 			{
