- Add back in packing to external cbuffer as some platforms don't adhere to the same rules

#rnx

diff --git a/Engine/Plugins/FX/Niagara/Source/Niagara/Private/NiagaraScriptExecutionParameterStore.cpp b/Engine/Plugins/FX/Niagara/Source/Niagara/Private/NiagaraScriptExecutionParameterStore.cpp
--- a/Engine/Plugins/FX/Niagara/Source/Niagara/Private/NiagaraScriptExecutionParameterStore.cpp
+++ b/Engine/Plugins/FX/Niagara/Source/Niagara/Private/NiagaraScriptExecutionParameterStore.cpp
@@ -5,19 +5,6 @@
 #include "NiagaraDataInterface.h"
 #include "NiagaraSystemInstance.h"
 
-namespace ScriptParameterStoreInternal
-{
-	template<typename TType>
-	void AddLayoutParameter(TArray<FNiagaraScriptExecutionPaddingInfo>& Members, uint32& NextMemberOffset, uint32& SrcOffset)
-	{
-		const uint32 BytesRemain = 16 - (NextMemberOffset & 15);
-		NextMemberOffset += BytesRemain != 16 && BytesRemain < sizeof(TType) ? BytesRemain : 0;
-		Members.Emplace(SrcOffset, NextMemberOffset, sizeof(TType), sizeof(TType));
-		SrcOffset += sizeof(TType);
-		NextMemberOffset += sizeof(TType);
-	}
-}
-
 FNiagaraScriptExecutionParameterStore::FNiagaraScriptExecutionParameterStore()
 	: FNiagaraParameterStore()
 	, ParameterSize(0)
@@ -52,42 +39,106 @@ uint32 OffsetAlign(uint32 SrcOffset, uint32 Size)
 
 uint32 FNiagaraScriptExecutionParameterStore::GenerateLayoutInfoInternal(TArray<FNiagaraScriptExecutionPaddingInfo>& Members, uint32& NextMemberOffset, const UStruct* InSrcStruct, uint32 InSrcOffset)
 {
+	uint32 VectorPaddedSize = (TShaderParameterTypeInfo<FVector4f>::NumRows * TShaderParameterTypeInfo<FVector4f>::NumColumns) * sizeof(float);
+
 	// Now insert an appropriate data member into the mix...
-	FName SrcStructName = InSrcStruct ? InSrcStruct->GetFName() : NAME_None;
 	if (InSrcStruct == FNiagaraTypeDefinition::GetBoolStruct() || InSrcStruct == FNiagaraTypeDefinition::GetIntStruct())
 	{
-		ScriptParameterStoreInternal::AddLayoutParameter<int32>(Members, NextMemberOffset, InSrcOffset);
+		uint32 IntSize = (TShaderParameterTypeInfo<uint32>::NumRows * TShaderParameterTypeInfo<uint32>::NumColumns) * sizeof(uint32);
+		Members.Emplace(InSrcOffset, Align(NextMemberOffset, TShaderParameterTypeInfo<uint32>::Alignment), IntSize, IntSize);
+		InSrcOffset += sizeof(uint32);
+		NextMemberOffset = Members[Members.Num() - 1].DestOffset + Members[Members.Num() - 1].DestSize;
 	}
 	else if (InSrcStruct == FNiagaraTypeDefinition::GetFloatStruct())
 	{
-		ScriptParameterStoreInternal::AddLayoutParameter<float>(Members, NextMemberOffset, InSrcOffset);
+		uint32 FloatSize = (TShaderParameterTypeInfo<float>::NumRows * TShaderParameterTypeInfo<float>::NumColumns) * sizeof(float);
+		Members.Emplace(InSrcOffset, Align(NextMemberOffset, TShaderParameterTypeInfo<float>::Alignment), FloatSize, FloatSize);
+		InSrcOffset += sizeof(float);
+		NextMemberOffset = Members[Members.Num() - 1].DestOffset + Members[Members.Num() - 1].DestSize;
+	}
+	else if (InSrcStruct == FNiagaraTypeDefinition::GetVec2Struct())
+	{
+		uint32 StructFinalSize = (TShaderParameterTypeInfo<FVector2f>::NumRows * TShaderParameterTypeInfo<FVector2f>::NumColumns) * sizeof(float);
+		Members.Emplace(InSrcOffset, OffsetAlign(NextMemberOffset, VectorPaddedSize), StructFinalSize, VectorPaddedSize);
+		InSrcOffset += sizeof(FVector2f);
+		NextMemberOffset = Members[Members.Num() - 1].DestOffset + Members[Members.Num() - 1].DestSize;
 	}
-	else if (InSrcStruct == FNiagaraTypeDefinition::GetVec2Struct() || SrcStructName == NAME_Vector2D || SrcStructName == NAME_Vector2d)
+	else if (InSrcStruct && (InSrcStruct->GetFName() == NAME_Vector2D || InSrcStruct->GetFName() == NAME_Vector2d))
+	{
+		uint32 StructFinalSize = (TShaderParameterTypeInfo<FVector2f>::NumRows * TShaderParameterTypeInfo<FVector2f>::NumColumns) * sizeof(float);
+		Members.Emplace(InSrcOffset, OffsetAlign(NextMemberOffset, VectorPaddedSize), StructFinalSize, VectorPaddedSize);
+		InSrcOffset += sizeof(FVector2f);
+		NextMemberOffset = Members[Members.Num() - 1].DestOffset + Members[Members.Num() - 1].DestSize;
+	}	
+	else if (InSrcStruct == FNiagaraTypeDefinition::GetVec3Struct() || InSrcStruct == FNiagaraTypeDefinition::GetPositionStruct())
 	{
-		ScriptParameterStoreInternal::AddLayoutParameter<FVector2f>(Members, NextMemberOffset, InSrcOffset);
+		uint32 StructFinalSize = (TShaderParameterTypeInfo<FVector3f>::NumRows * TShaderParameterTypeInfo<FVector3f>::NumColumns) * sizeof(float);
+		Members.Emplace(InSrcOffset, OffsetAlign(NextMemberOffset, VectorPaddedSize), StructFinalSize, VectorPaddedSize);
+		InSrcOffset += sizeof(FVector3f);
+		NextMemberOffset = Members[Members.Num() - 1].DestOffset + Members[Members.Num() - 1].DestSize;
 	}
-	else if (InSrcStruct == FNiagaraTypeDefinition::GetVec3Struct() || InSrcStruct == FNiagaraTypeDefinition::GetPositionStruct() || SrcStructName == NAME_Vector || SrcStructName == NAME_Vector3d)
+	else if (InSrcStruct && (InSrcStruct->GetFName() == NAME_Vector || InSrcStruct->GetFName() == NAME_Vector3d))
 	{
-		ScriptParameterStoreInternal::AddLayoutParameter<FVector3f>(Members, NextMemberOffset, InSrcOffset);
+		uint32 StructFinalSize = (TShaderParameterTypeInfo<FVector3f>::NumRows * TShaderParameterTypeInfo<FVector3f>::NumColumns) * sizeof(float);
+		Members.Emplace(InSrcOffset, OffsetAlign(NextMemberOffset, VectorPaddedSize), StructFinalSize, VectorPaddedSize);
+		InSrcOffset += sizeof(FVector3f);
+		NextMemberOffset = Members[Members.Num() - 1].DestOffset + Members[Members.Num() - 1].DestSize;
 	}
-	else if (SrcStructName == NAME_Vector4 || SrcStructName == NAME_Vector4d || SrcStructName == NAME_Quat || SrcStructName == NAME_Quat4d)
+	else if (InSrcStruct && (InSrcStruct->GetFName() == NAME_Vector4 || InSrcStruct->GetFName() == NAME_Vector4d))
 	{
-		ScriptParameterStoreInternal::AddLayoutParameter<FVector4f>(Members, NextMemberOffset, InSrcOffset);
+		uint32 StructFinalSize = (TShaderParameterTypeInfo<FVector4f>::NumRows * TShaderParameterTypeInfo<FVector4f>::NumColumns) * sizeof(float);
+		Members.Emplace(InSrcOffset, OffsetAlign(NextMemberOffset, VectorPaddedSize), StructFinalSize, VectorPaddedSize);
+		InSrcOffset += sizeof(FVector4f);
+		NextMemberOffset = Members[Members.Num() - 1].DestOffset + Members[Members.Num() - 1].DestSize;
+	}
+	else if (InSrcStruct && (InSrcStruct->GetFName() == NAME_Quat || InSrcStruct->GetFName() == NAME_Quat4d))
+	{
+		uint32 StructFinalSize = (TShaderParameterTypeInfo<FQuat4f>::NumRows * TShaderParameterTypeInfo<FQuat4f>::NumColumns) * sizeof(float);
+		Members.Emplace(InSrcOffset, OffsetAlign(NextMemberOffset, VectorPaddedSize), StructFinalSize, VectorPaddedSize);
+		InSrcOffset += sizeof(FQuat4f);
+		NextMemberOffset = Members[Members.Num() - 1].DestOffset + Members[Members.Num() - 1].DestSize;
 	}
 	else if (InSrcStruct == FNiagaraTypeDefinition::GetVec4Struct() || InSrcStruct == FNiagaraTypeDefinition::GetColorStruct() || InSrcStruct == FNiagaraTypeDefinition::GetQuatStruct())
 	{
-		ScriptParameterStoreInternal::AddLayoutParameter<FVector4f>(Members, NextMemberOffset, InSrcOffset);
+		uint32 StructFinalSize = (TShaderParameterTypeInfo<FVector4f>::NumRows * TShaderParameterTypeInfo<FVector4f>::NumColumns) * sizeof(float);
+		Members.Emplace(InSrcOffset, Align(NextMemberOffset, TShaderParameterTypeInfo<FVector4f>::Alignment), StructFinalSize, StructFinalSize);
+		InSrcOffset += sizeof(FVector4f);
+		NextMemberOffset = Members[Members.Num() - 1].DestOffset + Members[Members.Num() - 1].DestSize;
 	}
 	else if (InSrcStruct == FNiagaraTypeDefinition::GetMatrix4Struct())
 	{
-		ScriptParameterStoreInternal::AddLayoutParameter<FMatrix44f>(Members, NextMemberOffset, InSrcOffset);
+		uint32 StructFinalSize = (TShaderParameterTypeInfo<FMatrix44f>::NumRows * TShaderParameterTypeInfo<FMatrix44f>::NumColumns) * sizeof(float);
+		Members.Emplace(InSrcOffset, Align(NextMemberOffset, TShaderParameterTypeInfo<FMatrix44f>::Alignment), StructFinalSize, StructFinalSize);
+		InSrcOffset += sizeof(FMatrix44f);
+		NextMemberOffset = Members[Members.Num() - 1].DestOffset + Members[Members.Num() - 1].DestSize;
+	}
+	else if (InSrcStruct == FNiagaraTypeDefinition::GetHalfStruct())
+	{
+		uint32 HalfSize = (/*NumRows*/1 * /*NumColumns*/ 1) * sizeof(FFloat16);
+		Members.Emplace(InSrcOffset, Align(NextMemberOffset, TShaderParameterTypeInfo<FFloat16>::Alignment), HalfSize, HalfSize);
+		InSrcOffset += sizeof(FFloat16);
+		NextMemberOffset = Members[Members.Num() - 1].DestOffset + Members[Members.Num() - 1].DestSize;
+	}
+	else if (InSrcStruct == FNiagaraTypeDefinition::GetHalfVec2Struct())
+	{
+		uint32 StructFinalSize = (/*NumRows*/2 * /*NumColumns*/1) * sizeof(FFloat16);
+		Members.Emplace(InSrcOffset, OffsetAlign(NextMemberOffset, VectorPaddedSize), StructFinalSize, VectorPaddedSize);
+		InSrcOffset += sizeof(FFloat16[2]);
+		NextMemberOffset = Members[Members.Num() - 1].DestOffset + Members[Members.Num() - 1].DestSize;
+	}
+	else if (InSrcStruct == FNiagaraTypeDefinition::GetHalfVec3Struct())
+	{
+		uint32 StructFinalSize = (/*NumRows*/3 * /*NumColumns*/1) * sizeof(FFloat16);
+		Members.Emplace(InSrcOffset, OffsetAlign(NextMemberOffset, VectorPaddedSize), StructFinalSize, VectorPaddedSize);
+		InSrcOffset += sizeof(FFloat16[3]);
+		NextMemberOffset = Members[Members.Num() - 1].DestOffset + Members[Members.Num() - 1].DestSize;
 	}
-	else if (InSrcStruct == FNiagaraTypeDefinition::GetHalfStruct() ||
-			 InSrcStruct == FNiagaraTypeDefinition::GetHalfVec2Struct() ||
-			 InSrcStruct == FNiagaraTypeDefinition::GetHalfVec3Struct() ||
-			 InSrcStruct == FNiagaraTypeDefinition::GetHalfVec4Struct())
+	else if (InSrcStruct == FNiagaraTypeDefinition::GetHalfVec4Struct())
 	{
-		UE_LOG(LogNiagara, Fatal, TEXT("Half is not supported inside parameter stores"));
+		uint32 StructFinalSize = (/*NumRows*/4 * /*NumColumns*/1) * sizeof(FFloat16);
+		Members.Emplace(InSrcOffset, OffsetAlign(NextMemberOffset, VectorPaddedSize), StructFinalSize, VectorPaddedSize);
+		InSrcOffset += sizeof(FFloat16[4]);
+		NextMemberOffset = Members[Members.Num() - 1].DestOffset + Members[Members.Num() - 1].DestSize;
 	}
 	else
 	{
diff --git a/Engine/Plugins/FX/Niagara/Source/NiagaraEditor/Private/NiagaraHlslTranslator.cpp b/Engine/Plugins/FX/Niagara/Source/NiagaraEditor/Private/NiagaraHlslTranslator.cpp
--- a/Engine/Plugins/FX/Niagara/Source/NiagaraEditor/Private/NiagaraHlslTranslator.cpp
+++ b/Engine/Plugins/FX/Niagara/Source/NiagaraEditor/Private/NiagaraHlslTranslator.cpp
@@ -996,7 +996,7 @@ void FHlslNiagaraTranslator::BuildConstantBuffer(ENiagaraCodeChunkMode ChunkMode
 	for (const FNiagaraVariable& Variable : T::GetVariables())
 	{
 		const FString SymbolName = GetSanitizedSymbolName(Variable.GetName().ToString(), true);
-		AddUniformChunk(SymbolName, Variable, ChunkMode);
+		AddUniformChunk(SymbolName, Variable, ChunkMode, false);
 	}
 }
 
@@ -4172,7 +4172,7 @@ FString FHlslNiagaraTranslator::GeneratedConstantString(FVector4 Constant)
 	return FString::Format(TEXT("float4({0}, {1}, {2}, {3})"), Args);
 }
 
-int32 FHlslNiagaraTranslator::AddUniformChunk(FString SymbolName, const FNiagaraVariable& InVariable, ENiagaraCodeChunkMode ChunkMode)
+int32 FHlslNiagaraTranslator::AddUniformChunk(FString SymbolName, const FNiagaraVariable& InVariable, ENiagaraCodeChunkMode ChunkMode, bool AddPadding)
 {
 	const FNiagaraTypeDefinition& Type = InVariable.GetType();
 
@@ -4192,6 +4192,20 @@ int32 FHlslNiagaraTranslator::AddUniformChunk(FString SymbolName, const FNiagara
 		Chunk.Type = Type;
 		Chunk.Original = InVariable;
 
+		if (AddPadding)
+		{
+			if (Type == FNiagaraTypeDefinition::GetVec2Def())
+			{
+				Chunk.Type = FNiagaraTypeDefinition::GetVec4Def();
+				Chunk.ComponentMask = TEXT(".xy");
+			}
+			else if (Type == FNiagaraTypeDefinition::GetVec3Def() || Type == FNiagaraTypeDefinition::GetPositionDef())
+			{
+				Chunk.Type = FNiagaraTypeDefinition::GetVec4Def();
+				Chunk.ComponentMask = TEXT(".xyz");
+			}
+		}
+
 		Chunk.Mode = ChunkMode;
 
 		ChunksByMode[static_cast<int32>(ChunkMode)].Add(Ret);
@@ -5526,7 +5540,7 @@ bool FHlslNiagaraTranslator::ParameterMapRegisterExternalConstantNamespaceVariab
 					CompilationOutput.ScriptData.Parameters.SetOrAdd(InVariable);
 				}
 
-				UniformChunk = AddUniformChunk(SymbolNameDefined, InVariable, ENiagaraCodeChunkMode::Uniform);
+				UniformChunk = AddUniformChunk(SymbolNameDefined, InVariable, ENiagaraCodeChunkMode::Uniform, UNiagaraScript::IsGPUScript(CompileOptions.TargetUsage));
 			}
 			else
 			{
diff --git a/Engine/Plugins/FX/Niagara/Source/NiagaraEditor/Private/NiagaraHlslTranslator.h b/Engine/Plugins/FX/Niagara/Source/NiagaraEditor/Private/NiagaraHlslTranslator.h
--- a/Engine/Plugins/FX/Niagara/Source/NiagaraEditor/Private/NiagaraHlslTranslator.h
+++ b/Engine/Plugins/FX/Niagara/Source/NiagaraEditor/Private/NiagaraHlslTranslator.h
@@ -539,7 +539,7 @@ protected:
 	int32 AddSourceChunk(FString SymbolName, const FNiagaraTypeDefinition& Type, bool bSanitize = true);
 
 	/** Add a chunk defining a uniform value. */
-	int32 AddUniformChunk(FString SymbolName, const FNiagaraVariable& InVariable, ENiagaraCodeChunkMode ChunkMode);
+	int32 AddUniformChunk(FString SymbolName, const FNiagaraVariable& InVariable, ENiagaraCodeChunkMode ChunkMode, bool AddPadding);
 
 	/* Add a chunk that is written to the body of the shader code. */
 	int32 AddBodyChunk(FString SymbolName, FString Definition, const FNiagaraTypeDefinition& Type, TArray<int32>& SourceChunks, bool bDecl = true, bool bIsTerminated = true);
diff --git a/Engine/Source/Runtime/Core/Private/UObject/UE5MainStreamObjectVersion.cpp b/Engine/Source/Runtime/Core/Private/UObject/UE5MainStreamObjectVersion.cpp
--- a/Engine/Source/Runtime/Core/Private/UObject/UE5MainStreamObjectVersion.cpp
+++ b/Engine/Source/Runtime/Core/Private/UObject/UE5MainStreamObjectVersion.cpp
@@ -14,7 +14,7 @@ TMap<FGuid, FGuid> FUE5MainStreamObjectVersion::GetSystemGuids()
 	SystemGuids.Add(DevGuids.MATERIALSHADERMAP_DERIVEDDATA_VER, FGuid("65C01D817C9A4EEFAE9E988D41A1F3DD"));
 	SystemGuids.Add(DevGuids.NANITE_DERIVEDDATA_VER, FGuid("7A3EA93546624597BEBE2ADFDE3A7628"));
 	SystemGuids.Add(DevGuids.NIAGARASHADERMAP_DERIVEDDATA_VER, FGuid("8A37C45D24F2423CBE5F8F371DE33575"));
-	SystemGuids.Add(DevGuids.Niagara_LatestScriptCompileVersion, FGuid("6BA472B9ADB03A4096668C11AA08DCD8"));
+	SystemGuids.Add(DevGuids.Niagara_LatestScriptCompileVersion, FGuid("0BFC42A8539AC143B6640DEBBA491892"));
 	SystemGuids.Add(DevGuids.SkeletalMeshDerivedDataVersion, FGuid("FB7B2FB546DE49B2A2F88B7F69763CEE"));
 	SystemGuids.Add(DevGuids.STATICMESH_DERIVEDDATA_VER, FGuid("98F7D79A5811013825E75BBCC41ED3E9"));
 	SystemGuids.Add(DevGuids.POSESEARCHDB_DERIVEDDATA_VER, FGuid("4E595C2AC5E947D6BA9ABC874353E5BC"));
