cubic interpolation for 3d grids
integration with advection
integration with 3d gas emitters and systems

diff --git a/Engine/Plugins/FX/Niagara/Shaders/Private/NiagaraDataInterfaceGrid3DCollection.ush b/Engine/Plugins/FX/Niagara/Shaders/Private/NiagaraDataInterfaceGrid3DCollection.ush
new file mode 100644
--- /dev/null
+++ b/Engine/Plugins/FX/Niagara/Shaders/Private/NiagaraDataInterfaceGrid3DCollection.ush
@@ -0,0 +1,112 @@
+// Copyright Epic Games, Inc. All Rights Reserved.
+
+float SampleExplicitTrilinear_{ParameterName}(SamplerState Sampler, float3 UVW, int MipLevel)
+{	
+	int3 TotalDim = {NumCellsName} * {NumTiles};
+	float3 GridPos = float3(TotalDim[0] * UVW[0] - 0.5,  TotalDim[1] * UVW[1] - 0.5, TotalDim[2] * UVW[2] - 0.5f);
+	int3 GridCell = floor(GridPos);
+						
+	float3 delta = float3( GridPos[0] - GridCell[0], GridPos[1] - GridCell[1], GridPos[2] - GridCell[2]);
+	// local values
+	float4 F; 
+						
+	// bottom face
+	F[0] = {GridName}.Load( int4( GridCell[0],     GridCell[1], GridCell[2], MipLevel) );
+	F[1] = {GridName}.Load( int4( GridCell[0] + 1, GridCell[1], GridCell[2], MipLevel) );
+	F[2] = {GridName}.Load( int4( GridCell[0] + 1, GridCell[1] +1 , GridCell[2], MipLevel) );
+	F[3] = {GridName}.Load( int4( GridCell[0],     GridCell[1] +1 , GridCell[2], MipLevel) );
+	// bottom value
+	float BV = (1.-delta.y) * (F[0] *(1.-delta.x) + F[1]*(delta.x)) + delta.y * (F[3]*(1.-delta.x) + F[2]*delta.x);
+
+	// top face
+	F[0] = {GridName}.Load( int4( GridCell[0],     GridCell[1],     GridCell[2] + 1, MipLevel) );
+	F[1] = {GridName}.Load( int4( GridCell[0] + 1, GridCell[1],     GridCell[2] + 1, MipLevel) );
+	F[2] = {GridName}.Load( int4( GridCell[0] + 1, GridCell[1] +1 , GridCell[2] + 1, MipLevel) );
+	F[3] = {GridName}.Load( int4( GridCell[0],     GridCell[1] +1 , GridCell[2] + 1, MipLevel) );
+	// top value
+	float TV = (1.-delta.y) * (F[0] *(1.-delta.x) + F[1]*(delta.x)) + delta.y * (F[3]*(1.-delta.x) + F[2]*delta.x);
+
+	// interp between bottom and top
+	return TV * delta.z + BV * (1.-delta.z);	
+}
+
+float SampleTriCubicLagrange_{ParameterName}(SamplerState Sampler, float3 UVW, int MipLevel)
+{		    
+	int3 TotalDim = {NumCellsName} * {NumTiles};	
+	float3 GridPos = float3(TotalDim[0] * UVW[0] - 0.5,  TotalDim[1] * UVW[1] - 0.5, TotalDim[2] * UVW[2] - 0.5f);
+								
+	// identify the lower-left-hand corner of the cell
+	int3 GridCell = floor(GridPos);
+								
+	int3 LocalCell = GridCell % {NumCellsName};
+								
+	bool bBoundary = !any(LocalCell);
+	LocalCell -= ({NumCellsName} -int3(2,2,2));
+	bBoundary = bBoundary || !any(LocalCell);
+		
+	if (bBoundary)
+	{
+		// revert to trilinear hardware sampling at the boundary cells.
+		return {GridName}.SampleLevel(Sampler, UVW, MipLevel);;
+	}
+	else 
+	{
+		// sample point offset from lower left
+		float3 delta = float3( GridPos[0] - GridCell[0], GridPos[1] - GridCell[1], GridPos[2] - GridCell[2]);
+							
+		
+		float t = 1. + delta.x;
+		float u = 1. + delta.y;
+		float v = 1. + delta.z;
+		
+		float4 DataInZ = float4(0.f, 0.f, 0.f, 0.f);
+		for (int zk = 0; zk < 4; zk++)
+		{
+			int plane = GridCell[2] + zk -1;
+		
+			float4 DataInY = float4(0.f, 0.f, 0.f, 0.f);
+			for (int yj = 0; yj < 4; yj++)
+			{
+				float4 DataInX = float4(0.f, 0.f, 0.f, 0.f);
+								
+				int row = GridCell[1] + yj - 1;
+		
+				for(int xi = 0; xi < 4; xi++)
+				{
+					int col = GridCell[0] + xi - 1;
+					DataInX[xi] = {GridName}.Load(int4(col, row, plane, MipLevel)); 
+				}
+		
+				// coefficients
+				float Cx0 =   6.f * DataInX[0];
+				float Cx1 = -11.f * DataInX[0] + 18.f * DataInX[1] - 9.f * DataInX[2] + 2.f * DataInX[3];
+				float Cx2 =   6.f * DataInX[0] - 15.f * DataInX[1] +12.f * DataInX[2] - 3.f * DataInX[3];
+				float Cx3 =        -DataInX[0] +  3.f * DataInX[1] - 3.f * DataInX[2] +       DataInX[3];
+		
+				DataInY[yj] = (1.f/6.f) * (Cx0 + t * ( Cx1 +  t * ( Cx2 + t * Cx3 ) ));
+		
+			}
+									
+			float Cy0 =   6.f * DataInY[0];
+			float Cy1 = -11.f * DataInY[0] + 18.f * DataInY[1] - 9.f * DataInY[2] + 2.f * DataInY[3];
+			float Cy2 =   6.f * DataInY[0] - 15.f * DataInY[1] +12.f * DataInY[2] - 3.f * DataInY[3];
+			float Cy3 =        -DataInY[0] +  3.f * DataInY[1] - 3.f * DataInY[2] +       DataInY[3];
+		
+			DataInZ[zk] = (1.f/6.f) * (Cy0 + u * ( Cy1 +  u * ( Cy2 + u * Cy3 ) ));
+									
+		}
+		
+		float Cz0 =   6.f * DataInZ[0];
+		float Cz1 = -11.f * DataInZ[0] + 18.f * DataInZ[1] - 9.f * DataInZ[2] + 2.f * DataInZ[3];
+		float Cz2 =   6.f * DataInZ[0] - 15.f * DataInZ[1] +12.f * DataInZ[2] - 3.f * DataInZ[3];
+		float Cz3 =        -DataInZ[0] +  3.f * DataInZ[1] - 3.f * DataInZ[2] +       DataInZ[3];
+								
+		float TriCubicValue = (1.f/6.f) * (Cz0 + v * ( Cz1 +  v * ( Cz2 + v * Cz3 ) ));
+									
+		// for the fluid sims, a nice look can be achieved with BlendAmount = 0.25f;
+		// float BlendAmount = 0.0f;
+		// return lerp(TriCubicValue, TriLinearValue, BlendAmount);		
+
+		return TriCubicValue;		
+	}
+}
\ No newline at end of file
diff --git a/Engine/Plugins/FX/Niagara/Source/Niagara/Classes/NiagaraDataInterfaceGrid3DCollection.h b/Engine/Plugins/FX/Niagara/Source/Niagara/Classes/NiagaraDataInterfaceGrid3DCollection.h
--- a/Engine/Plugins/FX/Niagara/Source/Niagara/Classes/NiagaraDataInterfaceGrid3DCollection.h
+++ b/Engine/Plugins/FX/Niagara/Source/Niagara/Classes/NiagaraDataInterfaceGrid3DCollection.h
@@ -251,10 +251,12 @@ public:
 	static const FName SetValueFunctionName;
 	static const FName GetValueFunctionName;
 	static const FName SampleGridFunctionName;
+	static const FName CubicSampleGridFunctionName;
 
 	static const FName SetFullGridValueFunctionName;
 	static const FName GetFullGridPreviousValueFunctionName;
 	static const FName SamplePreviousFullGridFunctionName;
+	static const FName CubicSamplePreviousFullGridFunctionName;
 
 	static const FName SetVector4ValueFunctionName;
 	static const FName SetVector3ValueFunctionName;
@@ -263,17 +265,22 @@ public:
 	static const FName SetFloatValueFunctionName;
 	static const FName GetPreviousValueAtIndexFunctionName;
 	static const FName SamplePreviousGridAtIndexFunctionName;
+	static const FName CubicSamplePreviousGridAtIndexFunctionName;
 
 	static const FName GetPreviousVector4ValueFunctionName;
 	static const FName SamplePreviousGridVector4FunctionName;
+	static const FName CubicSamplePreviousGridVector4FunctionName;
 	static const FName SetVectorValueFunctionName;
 	static const FName GetPreviousVectorValueFunctionName;
 	static const FName SamplePreviousGridVectorFunctionName;
+	static const FName CubicSamplePreviousGridVectorFunctionName;
 	static const FName SetVector2DValueFunctionName;
 	static const FName GetPreviousVector2DValueFunctionName;
 	static const FName SamplePreviousGridVector2DFunctionName;
+	static const FName CubicSamplePreviousGridVector2DFunctionName;
 	static const FName GetPreviousFloatValueFunctionName;
 	static const FName SamplePreviousGridFloatFunctionName;
+	static const FName CubicSamplePreviousGridFloatFunctionName;
 	static const FString AttributeIndicesBaseName;
 	static const FString PerAttributeDataName;
 	static const TCHAR* VectorComponentNames[];
@@ -287,6 +294,8 @@ public:
 
 	static const FString AnonymousAttributeString;
 
+	static const TCHAR* TemplateShaderFilePath;
+
 #if WITH_EDITOR
 	virtual bool SupportsSetupAndTeardownHLSL() const { return true; }
 	virtual bool GenerateSetupHLSL(FNiagaraDataInterfaceGPUParamInfo& DIInstanceInfo, TConstArrayView<FNiagaraVariable> InArguments, bool bSpawnOnly, bool bPartialWrites, TArray<FText>& OutErrors, FString& OutHLSL) const;
@@ -305,7 +314,7 @@ protected:
 #if WITH_EDITORONLY_DATA
 	void WriteSetHLSL(const FNiagaraDataInterfaceGPUParamInfo& ParamInfo, const FNiagaraDataInterfaceGeneratedFunction& FunctionInfo, int FunctionInstanceIndex, int32 InNumChannels, FString& OutHLSL);
 	void WriteGetHLSL(const FNiagaraDataInterfaceGPUParamInfo& ParamInfo, const FNiagaraDataInterfaceGeneratedFunction& FunctionInfo, int FunctionInstanceIndex, int32 InNumChannels, FString& OutHLSL);
-	void WriteSampleHLSL(const FNiagaraDataInterfaceGPUParamInfo& ParamInfo, const FNiagaraDataInterfaceGeneratedFunction& FunctionInfo, int FunctionInstanceIndex, int32 InNumChannels, FString& OutHLSL);
+	void WriteSampleHLSL(const FNiagaraDataInterfaceGPUParamInfo& ParamInfo, const FNiagaraDataInterfaceGeneratedFunction& FunctionInfo, int FunctionInstanceIndex, int32 InNumChannels, FString SampleFunction, FString& OutHLSL);
 	void WriteAttributeGetIndexHLSL(const FNiagaraDataInterfaceGPUParamInfo& ParamInfo, const FNiagaraDataInterfaceGeneratedFunction& FunctionInfo, int FunctionInstanceIndex, int32 InNumChannels, FString& OutHLSL);
 
 	const TCHAR* TypeDefinitionToHLSLTypeString(const FNiagaraTypeDefinition& InDef) const;
diff --git a/Engine/Plugins/FX/Niagara/Source/Niagara/Private/NiagaraDataInterfaceGrid3DCollection.cpp b/Engine/Plugins/FX/Niagara/Source/Niagara/Private/NiagaraDataInterfaceGrid3DCollection.cpp
--- a/Engine/Plugins/FX/Niagara/Source/Niagara/Private/NiagaraDataInterfaceGrid3DCollection.cpp
+++ b/Engine/Plugins/FX/Niagara/Source/Niagara/Private/NiagaraDataInterfaceGrid3DCollection.cpp
@@ -62,6 +62,7 @@ const FName UNiagaraDataInterfaceGrid3DCollection::GetValueFunctionName("GetGrid
 const FName UNiagaraDataInterfaceGrid3DCollection::SetFullGridValueFunctionName("SetFullGridValue");
 const FName UNiagaraDataInterfaceGrid3DCollection::GetFullGridPreviousValueFunctionName("GetFullGridPreviousValue");
 const FName UNiagaraDataInterfaceGrid3DCollection::SamplePreviousFullGridFunctionName("SamplePreviousFullGrid");
+const FName UNiagaraDataInterfaceGrid3DCollection::CubicSamplePreviousFullGridFunctionName("CubicSamplePreviousFullGrid");
 
 const FName UNiagaraDataInterfaceGrid3DCollection::SetVector4ValueFunctionName("SetVector4Value");
 const FName UNiagaraDataInterfaceGrid3DCollection::SetVector3ValueFunctionName("SetVector3Value");
@@ -70,17 +71,22 @@ const FName UNiagaraDataInterfaceGrid3DCollection::GetVector2ValueFunctionName("
 const FName UNiagaraDataInterfaceGrid3DCollection::SetFloatValueFunctionName("SetFloatValue");
 const FName UNiagaraDataInterfaceGrid3DCollection::GetPreviousValueAtIndexFunctionName("GetPreviousValueAtIndex");
 const FName UNiagaraDataInterfaceGrid3DCollection::SamplePreviousGridAtIndexFunctionName("SamplePreviousGridAtIndex");
+const FName UNiagaraDataInterfaceGrid3DCollection::CubicSamplePreviousGridAtIndexFunctionName("CubicSamplePreviousGridAtIndex");
 
 const FName UNiagaraDataInterfaceGrid3DCollection::GetPreviousVector4ValueFunctionName("GetPreviousVector4Value");
 const FName UNiagaraDataInterfaceGrid3DCollection::SamplePreviousGridVector4FunctionName("SamplePreviousGridVector4Value");
+const FName UNiagaraDataInterfaceGrid3DCollection::CubicSamplePreviousGridVector4FunctionName("CubicSamplePreviousGridVector4Value");
 const FName UNiagaraDataInterfaceGrid3DCollection::SetVectorValueFunctionName("SetVectorValue");
 const FName UNiagaraDataInterfaceGrid3DCollection::GetPreviousVectorValueFunctionName("GetPreviousVectorValue");
 const FName UNiagaraDataInterfaceGrid3DCollection::SamplePreviousGridVectorFunctionName("SamplePreviousGridVector3Value");
+const FName UNiagaraDataInterfaceGrid3DCollection::CubicSamplePreviousGridVectorFunctionName("CubicSamplePreviousGridVector3Value");
 const FName UNiagaraDataInterfaceGrid3DCollection::SetVector2DValueFunctionName("SetVector2DValue");
 const FName UNiagaraDataInterfaceGrid3DCollection::GetPreviousVector2DValueFunctionName("GetPreviousVector2DValue");
 const FName UNiagaraDataInterfaceGrid3DCollection::SamplePreviousGridVector2DFunctionName("SamplePreviousGridVector2DValue");
+const FName UNiagaraDataInterfaceGrid3DCollection::CubicSamplePreviousGridVector2DFunctionName("CubicSamplePreviousGridVector2DValue");
 const FName UNiagaraDataInterfaceGrid3DCollection::GetPreviousFloatValueFunctionName("GetPreviousFloatValue");
 const FName UNiagaraDataInterfaceGrid3DCollection::SamplePreviousGridFloatFunctionName("SamplePreviousGridFloatValue");
+const FName UNiagaraDataInterfaceGrid3DCollection::CubicSamplePreviousGridFloatFunctionName("CubicSamplePreviousGridFloatValue");
 
 const FName UNiagaraDataInterfaceGrid3DCollection::GetVector4AttributeIndexFunctionName("GetVector4AttributeIndex");
 const FName UNiagaraDataInterfaceGrid3DCollection::GetVectorAttributeIndexFunctionName("GetVectorAttributeIndex");
@@ -88,6 +94,7 @@ const FName UNiagaraDataInterfaceGrid3DCollection::GetVector2DAttributeIndexFunc
 const FName UNiagaraDataInterfaceGrid3DCollection::GetFloatAttributeIndexFunctionName("GetFloatAttributeIndex");
 
 const FName UNiagaraDataInterfaceGrid3DCollection::SampleGridFunctionName("SampleGrid");
+const FName UNiagaraDataInterfaceGrid3DCollection::CubicSampleGridFunctionName("CubicSampleGrid");
 
 const FString UNiagaraDataInterfaceGrid3DCollection::AttributeIndicesBaseName(TEXT("AttributeIndices"));
 const FString UNiagaraDataInterfaceGrid3DCollection::PerAttributeDataName(TEXT("PerAttributeData"));
@@ -97,6 +104,7 @@ const TCHAR* UNiagaraDataInterfaceGrid3DCollection::VectorComponentNames[] = { T
 FNiagaraVariableBase UNiagaraDataInterfaceGrid3DCollection::ExposedRTVar;
 const FString UNiagaraDataInterfaceGrid3DCollection::AnonymousAttributeString("Attribute At Index");
 
+const TCHAR* UNiagaraDataInterfaceGrid3DCollection::TemplateShaderFilePath = TEXT("/Plugin/FX/Niagara/Private/NiagaraDataInterfaceGrid3DCollection.ush");
 
 struct FNiagaraGridCollection3DDIFunctionVersion
 {
@@ -111,26 +119,26 @@ struct FNiagaraGridCollection3DDIFunctionVersion
 
 bool UNiagaraDataInterfaceGrid3DCollection::CanCreateVarFromFuncName(const FName& FuncName)
 {
-	if (FuncName == UNiagaraDataInterfaceGrid3DCollection::SetVector4ValueFunctionName	|| FuncName == UNiagaraDataInterfaceGrid3DCollection::GetPreviousVector4ValueFunctionName || FuncName == UNiagaraDataInterfaceGrid3DCollection::SamplePreviousGridVector4FunctionName)
+	if (FuncName == UNiagaraDataInterfaceGrid3DCollection::SetVector4ValueFunctionName	|| FuncName == UNiagaraDataInterfaceGrid3DCollection::GetPreviousVector4ValueFunctionName || FuncName == UNiagaraDataInterfaceGrid3DCollection::SamplePreviousGridVector4FunctionName || FuncName == UNiagaraDataInterfaceGrid3DCollection::CubicSamplePreviousGridVector4FunctionName)
 		return true;
-	else if (FuncName == UNiagaraDataInterfaceGrid3DCollection::SetVector3ValueFunctionName || FuncName == UNiagaraDataInterfaceGrid3DCollection::SetVectorValueFunctionName ||  FuncName == UNiagaraDataInterfaceGrid3DCollection::GetPreviousVectorValueFunctionName || FuncName == UNiagaraDataInterfaceGrid3DCollection::SamplePreviousGridVectorFunctionName)
+	else if (FuncName == UNiagaraDataInterfaceGrid3DCollection::SetVector3ValueFunctionName || FuncName == UNiagaraDataInterfaceGrid3DCollection::SetVectorValueFunctionName ||  FuncName == UNiagaraDataInterfaceGrid3DCollection::GetPreviousVectorValueFunctionName || FuncName == UNiagaraDataInterfaceGrid3DCollection::SamplePreviousGridVectorFunctionName || FuncName == UNiagaraDataInterfaceGrid3DCollection::CubicSamplePreviousGridVectorFunctionName)
 		return true;
-	else if (FuncName == UNiagaraDataInterfaceGrid3DCollection::SetVector2ValueFunctionName || FuncName == UNiagaraDataInterfaceGrid3DCollection::SetVector2DValueFunctionName || FuncName == UNiagaraDataInterfaceGrid3DCollection::GetPreviousVector2DValueFunctionName || FuncName == UNiagaraDataInterfaceGrid3DCollection::SamplePreviousGridVector2DFunctionName)
+	else if (FuncName == UNiagaraDataInterfaceGrid3DCollection::SetVector2ValueFunctionName || FuncName == UNiagaraDataInterfaceGrid3DCollection::SetVector2DValueFunctionName || FuncName == UNiagaraDataInterfaceGrid3DCollection::GetPreviousVector2DValueFunctionName || FuncName == UNiagaraDataInterfaceGrid3DCollection::SamplePreviousGridVector2DFunctionName || FuncName == UNiagaraDataInterfaceGrid3DCollection::CubicSamplePreviousGridVector2DFunctionName)
 		return true;
-	else if (FuncName == UNiagaraDataInterfaceGrid3DCollection::SetFloatValueFunctionName       || FuncName == UNiagaraDataInterfaceGrid3DCollection::GetPreviousFloatValueFunctionName || FuncName == UNiagaraDataInterfaceGrid3DCollection::SamplePreviousGridFloatFunctionName )
+	else if (FuncName == UNiagaraDataInterfaceGrid3DCollection::SetFloatValueFunctionName       || FuncName == UNiagaraDataInterfaceGrid3DCollection::GetPreviousFloatValueFunctionName || FuncName == UNiagaraDataInterfaceGrid3DCollection::SamplePreviousGridFloatFunctionName || FuncName == UNiagaraDataInterfaceGrid3DCollection::CubicSamplePreviousGridFloatFunctionName)
 		return true;
 	return false;
 }
 
 FNiagaraTypeDefinition UNiagaraDataInterfaceGrid3DCollection::GetValueTypeFromFuncName(const FName& FuncName)
 {
-	if (FuncName == UNiagaraDataInterfaceGrid3DCollection::SetVector4ValueFunctionName || FuncName == UNiagaraDataInterfaceGrid3DCollection::GetPreviousVector4ValueFunctionName || FuncName == UNiagaraDataInterfaceGrid3DCollection::SamplePreviousGridVector4FunctionName)
+	if (FuncName == UNiagaraDataInterfaceGrid3DCollection::SetVector4ValueFunctionName || FuncName == UNiagaraDataInterfaceGrid3DCollection::GetPreviousVector4ValueFunctionName || FuncName == UNiagaraDataInterfaceGrid3DCollection::SamplePreviousGridVector4FunctionName || FuncName == UNiagaraDataInterfaceGrid3DCollection::CubicSamplePreviousGridVector4FunctionName)
 		return FNiagaraTypeDefinition::GetVec4Def();
-	else if (FuncName == UNiagaraDataInterfaceGrid3DCollection::SetVector3ValueFunctionName || FuncName == UNiagaraDataInterfaceGrid3DCollection::SetVectorValueFunctionName || FuncName == UNiagaraDataInterfaceGrid3DCollection::GetPreviousVectorValueFunctionName || FuncName == UNiagaraDataInterfaceGrid3DCollection::SamplePreviousGridVectorFunctionName)
+	else if (FuncName == UNiagaraDataInterfaceGrid3DCollection::SetVector3ValueFunctionName || FuncName == UNiagaraDataInterfaceGrid3DCollection::SetVectorValueFunctionName || FuncName == UNiagaraDataInterfaceGrid3DCollection::GetPreviousVectorValueFunctionName || FuncName == UNiagaraDataInterfaceGrid3DCollection::SamplePreviousGridVectorFunctionName || FuncName == UNiagaraDataInterfaceGrid3DCollection::CubicSamplePreviousGridVectorFunctionName)
 		return FNiagaraTypeDefinition::GetVec3Def();
-	else if (FuncName == UNiagaraDataInterfaceGrid3DCollection::SetVector2ValueFunctionName || FuncName == UNiagaraDataInterfaceGrid3DCollection::SetVector2DValueFunctionName ||  FuncName == UNiagaraDataInterfaceGrid3DCollection::SamplePreviousGridVector2DFunctionName)
+	else if (FuncName == UNiagaraDataInterfaceGrid3DCollection::SetVector2ValueFunctionName || FuncName == UNiagaraDataInterfaceGrid3DCollection::SetVector2DValueFunctionName ||  FuncName == UNiagaraDataInterfaceGrid3DCollection::SamplePreviousGridVector2DFunctionName || FuncName == UNiagaraDataInterfaceGrid3DCollection::CubicSamplePreviousGridVector2DFunctionName)
 		return FNiagaraTypeDefinition::GetVec2Def();
-	else if (FuncName == UNiagaraDataInterfaceGrid3DCollection::SetFloatValueFunctionName || FuncName == UNiagaraDataInterfaceGrid3DCollection::GetPreviousFloatValueFunctionName || FuncName == UNiagaraDataInterfaceGrid3DCollection::SamplePreviousGridFloatFunctionName)
+	else if (FuncName == UNiagaraDataInterfaceGrid3DCollection::SetFloatValueFunctionName || FuncName == UNiagaraDataInterfaceGrid3DCollection::GetPreviousFloatValueFunctionName || FuncName == UNiagaraDataInterfaceGrid3DCollection::SamplePreviousGridFloatFunctionName || FuncName == UNiagaraDataInterfaceGrid3DCollection::CubicSamplePreviousGridFloatFunctionName)
 		return FNiagaraTypeDefinition::GetFloatDef();
 
 	return FNiagaraTypeDefinition();
@@ -138,13 +146,13 @@ FNiagaraTypeDefinition UNiagaraDataInterfaceGrid3DCollection::GetValueTypeFromFu
 
 int32 UNiagaraDataInterfaceGrid3DCollection::GetComponentCountFromFuncName(const FName& FuncName)
 {
-	if (FuncName == UNiagaraDataInterfaceGrid3DCollection::SetVector4ValueFunctionName || FuncName == UNiagaraDataInterfaceGrid3DCollection::GetVector4AttributeIndexFunctionName || FuncName == UNiagaraDataInterfaceGrid3DCollection::GetPreviousVector4ValueFunctionName || FuncName == UNiagaraDataInterfaceGrid3DCollection::SamplePreviousGridVector4FunctionName)
+	if (FuncName == UNiagaraDataInterfaceGrid3DCollection::SetVector4ValueFunctionName || FuncName == UNiagaraDataInterfaceGrid3DCollection::GetVector4AttributeIndexFunctionName || FuncName == UNiagaraDataInterfaceGrid3DCollection::GetPreviousVector4ValueFunctionName || FuncName == UNiagaraDataInterfaceGrid3DCollection::SamplePreviousGridVector4FunctionName || FuncName == UNiagaraDataInterfaceGrid3DCollection::CubicSamplePreviousGridVector4FunctionName)
 		return 4;
-	else if (FuncName == UNiagaraDataInterfaceGrid3DCollection::SetVector3ValueFunctionName || FuncName == UNiagaraDataInterfaceGrid3DCollection::SetVectorValueFunctionName ||  FuncName == UNiagaraDataInterfaceGrid3DCollection::GetVectorAttributeIndexFunctionName || FuncName == UNiagaraDataInterfaceGrid3DCollection::GetPreviousVectorValueFunctionName || FuncName == UNiagaraDataInterfaceGrid3DCollection::SamplePreviousGridVectorFunctionName)
+	else if (FuncName == UNiagaraDataInterfaceGrid3DCollection::SetVector3ValueFunctionName || FuncName == UNiagaraDataInterfaceGrid3DCollection::SetVectorValueFunctionName ||  FuncName == UNiagaraDataInterfaceGrid3DCollection::GetVectorAttributeIndexFunctionName || FuncName == UNiagaraDataInterfaceGrid3DCollection::GetPreviousVectorValueFunctionName || FuncName == UNiagaraDataInterfaceGrid3DCollection::SamplePreviousGridVectorFunctionName || FuncName == UNiagaraDataInterfaceGrid3DCollection::CubicSamplePreviousGridVectorFunctionName)
 		return 3;
-	else if (FuncName == UNiagaraDataInterfaceGrid3DCollection::SetVector2ValueFunctionName || FuncName == UNiagaraDataInterfaceGrid3DCollection::SetVector2DValueFunctionName ||  FuncName == UNiagaraDataInterfaceGrid3DCollection::GetVector2DAttributeIndexFunctionName|| FuncName == UNiagaraDataInterfaceGrid3DCollection::GetPreviousVector2DValueFunctionName || FuncName == UNiagaraDataInterfaceGrid3DCollection::SamplePreviousGridVector2DFunctionName)
+	else if (FuncName == UNiagaraDataInterfaceGrid3DCollection::SetVector2ValueFunctionName || FuncName == UNiagaraDataInterfaceGrid3DCollection::SetVector2DValueFunctionName ||  FuncName == UNiagaraDataInterfaceGrid3DCollection::GetVector2DAttributeIndexFunctionName|| FuncName == UNiagaraDataInterfaceGrid3DCollection::GetPreviousVector2DValueFunctionName || FuncName == UNiagaraDataInterfaceGrid3DCollection::SamplePreviousGridVector2DFunctionName || FuncName == UNiagaraDataInterfaceGrid3DCollection::CubicSamplePreviousGridVector2DFunctionName)
 		return 2;
-	else if (FuncName == UNiagaraDataInterfaceGrid3DCollection::SetFloatValueFunctionName || FuncName == UNiagaraDataInterfaceGrid3DCollection::GetPreviousFloatValueFunctionName || FuncName == UNiagaraDataInterfaceGrid3DCollection::SamplePreviousGridFloatFunctionName|| FuncName == UNiagaraDataInterfaceGrid3DCollection::GetFloatAttributeIndexFunctionName)
+	else if (FuncName == UNiagaraDataInterfaceGrid3DCollection::SetFloatValueFunctionName || FuncName == UNiagaraDataInterfaceGrid3DCollection::GetPreviousFloatValueFunctionName || FuncName == UNiagaraDataInterfaceGrid3DCollection::SamplePreviousGridFloatFunctionName|| FuncName == UNiagaraDataInterfaceGrid3DCollection::GetFloatAttributeIndexFunctionName || FuncName == UNiagaraDataInterfaceGrid3DCollection::CubicSamplePreviousGridFloatFunctionName)
 		return 1;
 
 	return INDEX_NONE;
@@ -201,7 +209,10 @@ bool UNiagaraDataInterfaceGrid3DCollection::AppendCompileHash(FNiagaraCompileHas
 
 
 	InVisitor->UpdatePOD(TEXT("UNiagaraDataInterfaceGrid3DCollectionVersion"), (int32)FNiagaraGridCollection3DDIFunctionVersion::LatestVersion);
+	
+	//InVisitor->UpdateString(TEXT("UNiagaraDataInterfaceVolumeTextureHLSLSource"), GetShaderFileHash(TemplateShaderFilePath, EShaderPlatform::SP_PCD3D_SM5).ToString());
 	InVisitor->UpdateShaderParameters<FNDIGrid3DShaderParameters>();
+
 	return true;
 }
 #endif
@@ -522,6 +533,26 @@ void UNiagaraDataInterfaceGrid3DCollection::GetFunctions(TArray<FNiagaraFunction
 
 	}
 
+	{
+		FNiagaraFunctionSignature Sig;
+		Sig.Name = CubicSamplePreviousGridVector4FunctionName;
+		Sig.Inputs.Add(FNiagaraVariable(FNiagaraTypeDefinition(GetClass()), TEXT("Grid")));
+		Sig.Inputs.Add(FNiagaraVariable(FNiagaraTypeDefinition::GetVec3Def(), TEXT("Unit")));
+		Sig.Outputs.Add(FNiagaraVariable(FNiagaraTypeDefinition::GetVec4Def(), TEXT("Value")));
+		Sig.FunctionSpecifiers.Add(FName("Attribute"));
+		Sig.bMemberFunction = true;
+		Sig.bRequiresContext = false;
+		Sig.bExperimental = true;
+		Sig.ModuleUsageBitmask = ENiagaraScriptUsageMask::Particle;
+		Sig.bSupportsCPU = false;
+		Sig.bSupportsGPU = true;
+#if WITH_EDITORONLY_DATA
+		Sig.Description = NSLOCTEXT("Niagara", "NiagaraDataInterfaceGridColl3D_CubicSampleVector4", "Cubic Sample a Vector4 value on the Grid by Attribute name. Note that this is the value from the previous execution stage.");
+#endif
+		OutFunctions.Add(Sig);
+
+	}
+
 	{
 		FNiagaraFunctionSignature Sig;
 		Sig.Name = SetVectorValueFunctionName;
@@ -588,6 +619,26 @@ void UNiagaraDataInterfaceGrid3DCollection::GetFunctions(TArray<FNiagaraFunction
 
 	}
 
+	{
+		FNiagaraFunctionSignature Sig;
+		Sig.Name = CubicSamplePreviousGridVectorFunctionName;
+		Sig.Inputs.Add(FNiagaraVariable(FNiagaraTypeDefinition(GetClass()), TEXT("Grid")));
+		Sig.Inputs.Add(FNiagaraVariable(FNiagaraTypeDefinition::GetVec3Def(), TEXT("Unit")));
+		Sig.Outputs.Add(FNiagaraVariable(FNiagaraTypeDefinition::GetVec3Def(), TEXT("Value")));
+		Sig.FunctionSpecifiers.Add(FName("Attribute"));
+		Sig.bMemberFunction = true;
+		Sig.bRequiresContext = false;
+		Sig.bExperimental = true;
+		Sig.ModuleUsageBitmask = ENiagaraScriptUsageMask::Particle;
+		Sig.bSupportsCPU = false;
+		Sig.bSupportsGPU = true;
+#if WITH_EDITORONLY_DATA
+		Sig.Description = NSLOCTEXT("Niagara", "NiagaraDataInterfaceGridColl3D_CubicSampleVector3", "Cubic Sample a Vector3 value on the Grid by Attribute name. Note that this is the value from the previous execution stage.");
+#endif
+		OutFunctions.Add(Sig);
+
+	}
+
 	{
 		FNiagaraFunctionSignature Sig;
 		Sig.Name = SetVector2DValueFunctionName;
@@ -654,6 +705,26 @@ void UNiagaraDataInterfaceGrid3DCollection::GetFunctions(TArray<FNiagaraFunction
 
 	}
 
+	{
+		FNiagaraFunctionSignature Sig;
+		Sig.Name = CubicSamplePreviousGridVector2DFunctionName;
+		Sig.Inputs.Add(FNiagaraVariable(FNiagaraTypeDefinition(GetClass()), TEXT("Grid")));
+		Sig.Inputs.Add(FNiagaraVariable(FNiagaraTypeDefinition::GetVec3Def(), TEXT("Unit")));
+		Sig.Outputs.Add(FNiagaraVariable(FNiagaraTypeDefinition::GetVec2Def(), TEXT("Value")));
+		Sig.FunctionSpecifiers.Add(FName("Attribute"));
+		Sig.bMemberFunction = true;
+		Sig.bRequiresContext = false;
+		Sig.bExperimental = true;
+		Sig.ModuleUsageBitmask = ENiagaraScriptUsageMask::Particle;
+		Sig.bSupportsCPU = false;
+		Sig.bSupportsGPU = true;
+#if WITH_EDITORONLY_DATA
+		Sig.Description = NSLOCTEXT("Niagara", "NiagaraDataInterfaceGridColl3D_CubicSampleVector2", "Cubic Sample a Vector2 value on the Grid by Attribute name. Note that this is the value from the previous execution stage.");
+#endif
+		OutFunctions.Add(Sig);
+
+	}
+
 	{
 		FNiagaraFunctionSignature Sig;
 		Sig.Name = SetFloatValueFunctionName;
@@ -719,6 +790,26 @@ void UNiagaraDataInterfaceGrid3DCollection::GetFunctions(TArray<FNiagaraFunction
 
 	}
 
+	{
+		FNiagaraFunctionSignature Sig;
+		Sig.Name = CubicSamplePreviousGridFloatFunctionName;
+		Sig.Inputs.Add(FNiagaraVariable(FNiagaraTypeDefinition(GetClass()), TEXT("Grid")));
+		Sig.Inputs.Add(FNiagaraVariable(FNiagaraTypeDefinition::GetVec3Def(), TEXT("Unit")));
+		Sig.Outputs.Add(FNiagaraVariable(FNiagaraTypeDefinition::GetFloatDef(), TEXT("Value")));
+		Sig.FunctionSpecifiers.Add(FName("Attribute"));
+		Sig.bMemberFunction = true;
+		Sig.bRequiresContext = false;
+		Sig.bExperimental = true;
+		Sig.ModuleUsageBitmask = ENiagaraScriptUsageMask::Particle;
+		Sig.bSupportsCPU = false;
+		Sig.bSupportsGPU = true;
+#if WITH_EDITORONLY_DATA
+		Sig.Description = NSLOCTEXT("Niagara", "NiagaraDataInterfaceGridColl3D_CubicSampleFloat", "Cubic Sample a float value on the Grid by Attribute name.");
+#endif
+		OutFunctions.Add(Sig);
+
+	}
+
 	{
 		FNiagaraFunctionSignature Sig;
 		Sig.Name = SampleGridFunctionName;
@@ -739,6 +830,26 @@ void UNiagaraDataInterfaceGrid3DCollection::GetFunctions(TArray<FNiagaraFunction
 		OutFunctions.Add(Sig);
 	}
 
+	{
+		FNiagaraFunctionSignature Sig;
+		Sig.Name = CubicSampleGridFunctionName;
+		Sig.Inputs.Add(FNiagaraVariable(FNiagaraTypeDefinition(GetClass()), TEXT("Grid")));
+		Sig.Inputs.Add(FNiagaraVariable(FNiagaraTypeDefinition::GetFloatDef(), TEXT("UnitX")));
+		Sig.Inputs.Add(FNiagaraVariable(FNiagaraTypeDefinition::GetFloatDef(), TEXT("UnitY")));
+		Sig.Inputs.Add(FNiagaraVariable(FNiagaraTypeDefinition::GetFloatDef(), TEXT("UnitZ")));
+		Sig.Inputs.Add(FNiagaraVariable(FNiagaraTypeDefinition::GetIntDef(), TEXT("AttributeIndex")));
+		Sig.Outputs.Add(FNiagaraVariable(FNiagaraTypeDefinition::GetFloatDef(), TEXT("Value")));
+
+		Sig.bExperimental = true;
+		Sig.bMemberFunction = true;
+		Sig.bRequiresContext = false;
+		Sig.ModuleUsageBitmask = ENiagaraScriptUsageMask::Particle;
+		Sig.bSupportsCPU = false;
+		Sig.bSupportsGPU = true;
+		Sig.bSoftDeprecatedFunction = true;
+		OutFunctions.Add(Sig);
+	}
+
 	{
 		FNiagaraFunctionSignature Sig;
 		Sig.Name = SamplePreviousGridAtIndexFunctionName;
@@ -758,6 +869,25 @@ void UNiagaraDataInterfaceGrid3DCollection::GetFunctions(TArray<FNiagaraFunction
 		OutFunctions.Add(Sig);
 	}
 
+	{
+		FNiagaraFunctionSignature Sig;
+		Sig.Name = CubicSamplePreviousGridAtIndexFunctionName;
+		Sig.Inputs.Add(FNiagaraVariable(FNiagaraTypeDefinition(GetClass()), TEXT("Grid")));
+		Sig.Inputs.Add(FNiagaraVariable(FNiagaraTypeDefinition::GetFloatDef(), TEXT("UnitX")));
+		Sig.Inputs.Add(FNiagaraVariable(FNiagaraTypeDefinition::GetFloatDef(), TEXT("UnitY")));
+		Sig.Inputs.Add(FNiagaraVariable(FNiagaraTypeDefinition::GetFloatDef(), TEXT("UnitZ")));
+		Sig.Inputs.Add(FNiagaraVariable(FNiagaraTypeDefinition::GetIntDef(), TEXT("AttributeIndex")));
+		Sig.Outputs.Add(FNiagaraVariable(FNiagaraTypeDefinition::GetFloatDef(), TEXT("Value")));
+
+		Sig.bExperimental = true;
+		Sig.bMemberFunction = true;
+		Sig.bRequiresContext = false;
+		Sig.ModuleUsageBitmask = ENiagaraScriptUsageMask::Particle;
+		Sig.bSupportsCPU = false;
+		Sig.bSupportsGPU = true;
+		OutFunctions.Add(Sig);
+	}
+
 	{
 		FNiagaraFunctionSignature Sig;
 		Sig.Name = SamplePreviousFullGridFunctionName;
@@ -774,6 +904,22 @@ void UNiagaraDataInterfaceGrid3DCollection::GetFunctions(TArray<FNiagaraFunction
 		OutFunctions.Add(Sig);
 	}
 
+	{
+		FNiagaraFunctionSignature Sig;
+		Sig.Name = CubicSamplePreviousFullGridFunctionName;
+		Sig.Inputs.Add(FNiagaraVariable(FNiagaraTypeDefinition(GetClass()), TEXT("Grid")));
+		Sig.Inputs.Add(FNiagaraVariable(FNiagaraTypeDefinition::GetVec3Def(), TEXT("Unit")));
+		Sig.Outputs.Add(FNiagaraVariable(FNiagaraTypeDefinition::GetFloatDef(), TEXT("Value")));
+
+		Sig.bExperimental = true;
+		Sig.bMemberFunction = true;
+		Sig.bRequiresContext = false;
+		Sig.ModuleUsageBitmask = ENiagaraScriptUsageMask::Particle;
+		Sig.bSupportsCPU = false;
+		Sig.bSupportsGPU = true;
+		OutFunctions.Add(Sig);
+	}
+
 	{
 		FNiagaraFunctionSignature Sig;
 		Sig.Name = GetVector4AttributeIndexFunctionName;
@@ -1002,7 +1148,8 @@ void UNiagaraDataInterfaceGrid3DCollection::GetParameterDefinitionHLSL(const FNi
 	const int AttributeCount = ParamInfo.GeneratedFunctions.Num();
 	const int AttributeInt4Count = FMath::Max(1, FMath::DivideAndRoundUp(AttributeCount, 4));
 
-	TMap<FString, FStringFormatArg> ArgsDeclarations = {				
+	TMap<FString, FStringFormatArg> ArgsDeclarations = {	
+		{ TEXT("ParameterName"),	ParamInfo.DataInterfaceHLSLSymbol},
 		{ TEXT("GridName"),			ParamInfo.DataInterfaceHLSLSymbol + GridName},
 		{ TEXT("SamplerName"),		ParamInfo.DataInterfaceHLSLSymbol + SamplerName},
 		{ TEXT("OutputGridName"),	ParamInfo.DataInterfaceHLSLSymbol + OutputGridName},
@@ -1010,7 +1157,7 @@ void UNiagaraDataInterfaceGrid3DCollection::GetParameterDefinitionHLSL(const FNi
 		{ TEXT("OneOverNumTiles"),	ParamInfo.DataInterfaceHLSLSymbol + OneOverNumTilesName},
 		{ TEXT("UnitClampMin"),		ParamInfo.DataInterfaceHLSLSymbol + UnitClampMinName},
 		{ TEXT("UnitClampMax"),		ParamInfo.DataInterfaceHLSLSymbol + UnitClampMaxName},
-
+		{ TEXT("NumCellsName"),				ParamInfo.DataInterfaceHLSLSymbol + UNiagaraDataInterfaceRWBase::NumCellsName},
 		{ TEXT("UnitToUVName"),				ParamInfo.DataInterfaceHLSLSymbol + UNiagaraDataInterfaceRWBase::UnitToUVName},
 		{ TEXT("AttributeIndicesName"),		ParamInfo.DataInterfaceHLSLSymbol + TEXT("_") + AttributeIndicesBaseName},
 		{ TEXT("PerAttributeDataName"),		ParamInfo.DataInterfaceHLSLSymbol + TEXT("_") + PerAttributeDataName},
@@ -1018,7 +1165,12 @@ void UNiagaraDataInterfaceGrid3DCollection::GetParameterDefinitionHLSL(const FNi
 		{ TEXT("NumAttributesName"),		ParamInfo.DataInterfaceHLSLSymbol + UNiagaraDataInterfaceRWBase::NumAttributesName},
 		{ TEXT("NumNamedAttributesName"),	ParamInfo.DataInterfaceHLSLSymbol + UNiagaraDataInterfaceRWBase::NumNamedAttributesName},
 	};
+
 	OutHLSL += FString::Format(FormatDeclarations, ArgsDeclarations);
+
+	FString TemplateFile;
+	LoadShaderSourceFile(TemplateShaderFilePath, EShaderPlatform::SP_PCD3D_SM5, &TemplateFile, nullptr);
+	OutHLSL += FString::Format(*TemplateFile, ArgsDeclarations);
 }
 void UNiagaraDataInterfaceGrid3DCollection::WriteSetHLSL(const FNiagaraDataInterfaceGPUParamInfo& ParamInfo, const FNiagaraDataInterfaceGeneratedFunction& FunctionInfo, int FunctionInstanceIndex, int32 InNumChannels, FString& OutHLSL)
 {
@@ -1052,6 +1204,7 @@ void UNiagaraDataInterfaceGrid3DCollection::WriteSetHLSL(const FNiagaraDataInter
 			}
 		)");
 	TMap<FString, FStringFormatArg> ArgsBounds = {
+		{TEXT("ParameterName"),				ParamInfo.DataInterfaceHLSLSymbol},
 		{TEXT("FunctionName"),				FunctionInfo.InstanceName},
 		{TEXT("OutputGrid"),				ParamInfo.DataInterfaceHLSLSymbol + OutputGridName},
 		{TEXT("NumCellsName"),				ParamInfo.DataInterfaceHLSLSymbol + UNiagaraDataInterfaceRWBase::NumCellsName},
@@ -1097,6 +1250,7 @@ void UNiagaraDataInterfaceGrid3DCollection::WriteGetHLSL(const FNiagaraDataInter
 			}
 		)");
 	TMap<FString, FStringFormatArg> ArgsBounds = {
+		{TEXT("ParameterName"),				ParamInfo.DataInterfaceHLSLSymbol},
 		{TEXT("FunctionName"),				FunctionInfo.InstanceName},
 		{TEXT("OutputGrid"),				ParamInfo.DataInterfaceHLSLSymbol + OutputGridName},
 		{TEXT("Grid"),						ParamInfo.DataInterfaceHLSLSymbol + GridName},
@@ -1114,7 +1268,7 @@ void UNiagaraDataInterfaceGrid3DCollection::WriteGetHLSL(const FNiagaraDataInter
 	OutHLSL += FString::Format(*FormatBounds, ArgsBounds);
 }
 
-void UNiagaraDataInterfaceGrid3DCollection::WriteSampleHLSL(const FNiagaraDataInterfaceGPUParamInfo& ParamInfo, const FNiagaraDataInterfaceGeneratedFunction& FunctionInfo, int FunctionInstanceIndex, int32 InNumChannels, FString& OutHLSL)
+void UNiagaraDataInterfaceGrid3DCollection::WriteSampleHLSL(const FNiagaraDataInterfaceGPUParamInfo& ParamInfo, const FNiagaraDataInterfaceGeneratedFunction& FunctionInfo, int FunctionInstanceIndex, int32 InNumChannels, FString SampleFunction, FString& OutHLSL)
 {
 	FString FormatBounds = TEXT(R"(
 			void {FunctionName}(float3 In_Unit, out float{NumChannelsVariableSuffix} Out_Val)
@@ -1127,7 +1281,7 @@ void UNiagaraDataInterfaceGrid3DCollection::WriteSampleHLSL(const FNiagaraDataIn
 					int CurAttributeIndex = In_AttributeIndex + i;
 					float3 UVW = {PerAttributeDataName}[CurAttributeIndex + {NumAttributesName}].xyz + TileUV;
 
-					float Val = {Grid}.SampleLevel({SamplerName}, UVW, 0);
+					float Val = {SampleFunction}({SamplerName}, UVW, 0);
 					)");
 	if (InNumChannels == 1)
 	{
@@ -1144,9 +1298,11 @@ void UNiagaraDataInterfaceGrid3DCollection::WriteSampleHLSL(const FNiagaraDataIn
 		)");
 
 	TMap<FString, FStringFormatArg> ArgsBounds = {
+		{TEXT("ParameterName"),				ParamInfo.DataInterfaceHLSLSymbol},
 		{TEXT("FunctionName"),				FunctionInfo.InstanceName},
 		{TEXT("Grid"),						ParamInfo.DataInterfaceHLSLSymbol + GridName},
 		{TEXT("SamplerName"),				ParamInfo.DataInterfaceHLSLSymbol + SamplerName },
+		{TEXT("SampleFunction"),			ParamInfo.DataInterfaceHLSLSymbol + SampleFunction },		
 		{TEXT("NumCellsName"),				ParamInfo.DataInterfaceHLSLSymbol + UNiagaraDataInterfaceRWBase::NumCellsName},
 		{TEXT("UnitToUVName"),				ParamInfo.DataInterfaceHLSLSymbol + UNiagaraDataInterfaceRWBase::UnitToUVName},
 		{TEXT("NumChannels"),				FString::FromInt(InNumChannels)},
@@ -1176,6 +1332,7 @@ void UNiagaraDataInterfaceGrid3DCollection::WriteAttributeGetIndexHLSL(const FNi
 
 
 	TMap<FString, FStringFormatArg> ArgsBounds = {
+		{TEXT("ParameterName"),	ParamInfo.DataInterfaceHLSLSymbol},
 		{TEXT("FunctionName"), FunctionInfo.InstanceName},
 		{TEXT("AttributeIndicesName"), ParamInfo.DataInterfaceHLSLSymbol + TEXT("_") + AttributeIndicesBaseName},
 		{TEXT("AttributeIndexGroup"), FunctionInstanceIndex / 4},
@@ -1244,6 +1401,7 @@ bool UNiagaraDataInterfaceGrid3DCollection::GetFunctionHLSL(const FNiagaraDataIn
 
 	TMap<FString, FStringFormatArg> ArgsBounds =
 	{
+		{TEXT("ParameterName"),		ParamInfo.DataInterfaceHLSLSymbol},
 		{TEXT("FunctionName"),			FunctionInfo.InstanceName},
 		{TEXT("Grid"),					ParamInfo.DataInterfaceHLSLSymbol + GridName},
 		{TEXT("OutputGrid"),			ParamInfo.DataInterfaceHLSLSymbol + OutputGridName},
@@ -1411,7 +1569,12 @@ bool UNiagaraDataInterfaceGrid3DCollection::GetFunctionHLSL(const FNiagaraDataIn
 	}
 	else if (FunctionInfo.DefinitionName == SamplePreviousGridVector4FunctionName)
 	{
-		WriteSampleHLSL(ParamInfo, FunctionInfo, FunctionInstanceIndex, 4, OutHLSL);
+		WriteSampleHLSL(ParamInfo, FunctionInfo, FunctionInstanceIndex, 4, "{Grid}.SampleLevel", OutHLSL);
+		return true;
+	}
+	else if (FunctionInfo.DefinitionName == CubicSamplePreviousGridVector4FunctionName)
+	{
+		WriteSampleHLSL(ParamInfo, FunctionInfo, FunctionInstanceIndex, 4, "SampleTriCubicLagrange_{ParameterName}", OutHLSL);
 		return true;
 	}
 	else if (FunctionInfo.DefinitionName == SetVector3ValueFunctionName || FunctionInfo.DefinitionName == SetVectorValueFunctionName)
@@ -1426,7 +1589,12 @@ bool UNiagaraDataInterfaceGrid3DCollection::GetFunctionHLSL(const FNiagaraDataIn
 	}
 	else if ( FunctionInfo.DefinitionName == SamplePreviousGridVectorFunctionName)
 	{
-		WriteSampleHLSL(ParamInfo, FunctionInfo, FunctionInstanceIndex, 3, OutHLSL);
+		WriteSampleHLSL(ParamInfo, FunctionInfo, FunctionInstanceIndex, 3, "{Grid}.SampleLevel", OutHLSL);
+		return true;
+	}
+	else if (FunctionInfo.DefinitionName == CubicSamplePreviousGridVectorFunctionName)
+	{
+		WriteSampleHLSL(ParamInfo, FunctionInfo, FunctionInstanceIndex, 3, "SampleTriCubicLagrange_{ParameterName}", OutHLSL);
 		return true;
 	}
 	else if (FunctionInfo.DefinitionName == SetVector2ValueFunctionName || FunctionInfo.DefinitionName == SetVector2DValueFunctionName)
@@ -1441,7 +1609,12 @@ bool UNiagaraDataInterfaceGrid3DCollection::GetFunctionHLSL(const FNiagaraDataIn
 	}
 	else if ( FunctionInfo.DefinitionName == SamplePreviousGridVector2DFunctionName)
 	{
-		WriteSampleHLSL(ParamInfo, FunctionInfo, FunctionInstanceIndex, 2, OutHLSL);
+		WriteSampleHLSL(ParamInfo, FunctionInfo, FunctionInstanceIndex, 2, "{Grid}.SampleLevel", OutHLSL);
+		return true;
+	}
+	else if (FunctionInfo.DefinitionName == CubicSamplePreviousGridVector2DFunctionName)
+	{
+		WriteSampleHLSL(ParamInfo, FunctionInfo, FunctionInstanceIndex, 2, "SampleTriCubicLagrange_{ParameterName}", OutHLSL);
 		return true;
 	}
 	else if (FunctionInfo.DefinitionName == SetFloatValueFunctionName)
@@ -1456,7 +1629,12 @@ bool UNiagaraDataInterfaceGrid3DCollection::GetFunctionHLSL(const FNiagaraDataIn
 	}
 	else if (FunctionInfo.DefinitionName == SamplePreviousGridFloatFunctionName)
 	{
-		WriteSampleHLSL(ParamInfo, FunctionInfo, FunctionInstanceIndex, 1, OutHLSL);
+		WriteSampleHLSL(ParamInfo, FunctionInfo, FunctionInstanceIndex, 1, "{Grid}.SampleLevel", OutHLSL);
+		return true;
+	}
+	else if (FunctionInfo.DefinitionName == CubicSamplePreviousGridFloatFunctionName)
+	{
+		WriteSampleHLSL(ParamInfo, FunctionInfo, FunctionInstanceIndex, 1, "SampleTriCubicLagrange_{ParameterName}", OutHLSL);
 		return true;
 	}
 	else if (FunctionInfo.DefinitionName == GetVector4AttributeIndexFunctionName)
@@ -1494,20 +1672,50 @@ bool UNiagaraDataInterfaceGrid3DCollection::GetFunctionHLSL(const FNiagaraDataIn
 						Out_Val = {Grid}.SampleLevel({SamplerName}, UVW, 0);
 					}
 				}
-			)");
+			)");		
 		OutHLSL += FString::Format(FormatBounds, ArgsBounds);
 		return true;
 	}
-	else if (FunctionInfo.DefinitionName == SamplePreviousFullGridFunctionName)
+	else if (FunctionInfo.DefinitionName == CubicSampleGridFunctionName || FunctionInfo.DefinitionName == CubicSamplePreviousGridAtIndexFunctionName)
 	{
-	static const TCHAR* FormatBounds = TEXT(R"(
-				void {FunctionName}(float3 In_Unit, out float Out_Val)
+		static const TCHAR* FormatBounds = TEXT(R"(
+				void {FunctionName}(float In_UnitX, float In_UnitY, float In_UnitZ, int In_AttributeIndex, out float Out_Val)
 				{
-					Out_Val = {Grid}.SampleLevel({SamplerName}, In_Unit, 0);
+					Out_Val = 0;
+					if ( In_AttributeIndex < {NumAttributesName} )
+					{
+						float3 Unit = float3(In_UnitX, In_UnitY, In_UnitZ);
+						Unit = clamp(Unit, {UnitClampMin}, {UnitClampMax}) * {OneOverNumTiles};
+
+						float3 UVW = {PerAttributeDataName}[In_AttributeIndex + {NumAttributesName}].xyz + Unit;
+						Out_Val = SampleTriCubicLagrange_{ParameterName}({SamplerName}, UVW, 0);
+					}
 				}
 			)");
-	OutHLSL += FString::Format(FormatBounds, ArgsBounds);
-	return true;
+		OutHLSL += FString::Format(FormatBounds, ArgsBounds);
+		return true;
+	}
+	else if (FunctionInfo.DefinitionName == SamplePreviousFullGridFunctionName)
+	{
+		static const TCHAR* FormatBounds = TEXT(R"(
+					void {FunctionName}(float3 In_Unit, out float Out_Val)
+					{
+						Out_Val = {Grid}.SampleLevel({SamplerName}, In_Unit, 0);
+					}
+				)");
+		OutHLSL += FString::Format(FormatBounds, ArgsBounds);
+		return true;
+	}
+	else if (FunctionInfo.DefinitionName == CubicSamplePreviousFullGridFunctionName)
+	{
+		static const TCHAR* FormatBounds = TEXT(R"(
+					void {FunctionName}(float3 In_Unit, out float Out_Val)
+					{
+						Out_Val = SampleTriCubicLagrange_{ParameterName}({SamplerName}, In_Unit, 0);
+					}
+				)");
+		OutHLSL += FString::Format(FormatBounds, ArgsBounds);
+		return true;
 	}
 	return false;
 }
