- Fix crash if an event spawn info had no space left

diff --git a/Engine/Plugins/FX/Niagara/Source/Niagara/Private/NiagaraEmitterInstance.cpp b/Engine/Plugins/FX/Niagara/Source/Niagara/Private/NiagaraEmitterInstance.cpp
--- a/Engine/Plugins/FX/Niagara/Source/Niagara/Private/NiagaraEmitterInstance.cpp
+++ b/Engine/Plugins/FX/Niagara/Source/Niagara/Private/NiagaraEmitterInstance.cpp
@@ -1823,6 +1823,10 @@ void FNiagaraEmitterInstance::Tick(float DeltaSeconds)
 							Info.SpawnCounts[i] = NumActuallySpawned;
 							SpawnCountRemaining -= NumActuallySpawned;
 						}
+						else
+						{
+							Info.SpawnCounts[i] = 0;
+						}
 					}
 				}
 			}
@@ -1929,8 +1933,8 @@ void FNiagaraEmitterInstance::Tick(float DeltaSeconds)
 
 						EventInstanceData->EventExecContexts[EventScriptIdx].Execute(EventNumToSpawn, EventConstantBufferTable);
 
-						uint32 PostHandlerNumInstances = Data.GetDestinationData()->GetNumInstances();
-						uint32 EventSpawnsStillAlive = PostHandlerNumInstances - EventSpawnStart;
+						const uint32 PostHandlerNumInstances = Data.GetDestinationData()->GetNumInstances();
+						const uint32 EventSpawnsStillAlive = PostHandlerNumInstances - EventSpawnStart;
 						if (GbDumpParticleData || System->bDumpDebugEmitterInfo)
 						{
 							EventInstanceData->EventHandlingInfo[EventScriptIdx].EventData->Dump(i, 1, FString::Printf(TEXT("=== Event Data %d [%d] ==="), EventScriptIdx, i));
