- Expose occluded line drawing color scale defaults to 5%

diff --git a/Engine/Plugins/FX/Niagara/Shaders/Private/NiagaraDebugDraw.usf b/Engine/Plugins/FX/Niagara/Shaders/Private/NiagaraDebugDraw.usf
--- a/Engine/Plugins/FX/Niagara/Shaders/Private/NiagaraDebugDraw.usf
+++ b/Engine/Plugins/FX/Niagara/Shaders/Private/NiagaraDebugDraw.usf
@@ -57,6 +57,7 @@ float2 OutputInvResolution;
 float2 OriginalViewRectMin;
 float2 OriginalViewSize;
 float2 OriginalBufferInvSize;
+float OccludedColorScale;
 
 void MainPS (
 	in float4 Position	: SV_POSITION,
@@ -70,7 +71,7 @@ void MainPS (
 	const float Depth		= DepthTexture.SampleLevel(DepthSampler, DepthUV, 0);
 	const bool bIsHidden	= Position.z < Depth; // Reverse-Z
 
-	const float ColorScale = bIsHidden ? 0.1f : 1;
+	const float ColorScale = bIsHidden ? OccludedColorScale : 1;
 	OutColor = Color * ColorScale;
 }
 #endif
diff --git a/Engine/Plugins/FX/Niagara/Source/NiagaraShader/Private/NiagaraDebugShaders.cpp b/Engine/Plugins/FX/Niagara/Source/NiagaraShader/Private/NiagaraDebugShaders.cpp
--- a/Engine/Plugins/FX/Niagara/Source/NiagaraShader/Private/NiagaraDebugShaders.cpp
+++ b/Engine/Plugins/FX/Niagara/Source/NiagaraShader/Private/NiagaraDebugShaders.cpp
@@ -29,6 +29,14 @@ static FAutoConsoleVariableRef CVarNiagaraGpuComputeDebug_FourComponentMode(
 	ECVF_Default
 );
 
+float GNiagaraGpuComputeDebug_OccludedLineColorScale = 0.05f;
+static FAutoConsoleVariableRef CVarNiagaraGpuComputeDebug_OccludedLineColorScale(
+	TEXT("fx.Niagara.GpuComputeDebug.OccludedLineColorScale"),
+	GNiagaraGpuComputeDebug_OccludedLineColorScale,
+	TEXT("Scalar value to adjust occluded lines, where 0 means transparent and 1 is opaque.  Default is 0.05 or 5%"),
+	ECVF_Default
+);
+
 //////////////////////////////////////////////////////////////////////////
 
 class NIAGARASHADER_API FNiagaraVisualizeTexturePS : public FGlobalShader
@@ -150,6 +158,7 @@ class NIAGARASHADER_API FNiagaraDebugDrawLinePS : public FGlobalShader
 		SHADER_PARAMETER(FVector2f, OriginalViewRectMin)
 		SHADER_PARAMETER(FVector2f, OriginalViewSize)
 		SHADER_PARAMETER(FVector2f, OriginalBufferInvSize)
+		SHADER_PARAMETER(float, OccludedColorScale)
 		SHADER_PARAMETER_RDG_TEXTURE(Texture2D, DepthTexture)
 		SHADER_PARAMETER_SAMPLER(SamplerState, DepthSampler)
 
@@ -205,6 +214,7 @@ void NiagaraDebugShaders::DrawDebugLines(
 	PassParameters->PSParameters.OriginalViewRectMin	= FVector2f(View.ViewRect.Min);
 	PassParameters->PSParameters.OriginalViewSize		= FVector2f(View.ViewRect.Width(), View.ViewRect.Height());
 	PassParameters->PSParameters.OriginalBufferInvSize	= FVector2f(1.f / SceneDepth->Desc.Extent.X, 1.f / SceneDepth->Desc.Extent.Y);
+	PassParameters->PSParameters.OccludedColorScale		= GNiagaraGpuComputeDebug_OccludedLineColorScale;
 	PassParameters->PSParameters.DepthTexture			= SceneDepth;
 	PassParameters->PSParameters.DepthSampler			= TStaticSamplerState<SF_Point, AM_Clamp, AM_Clamp, AM_Clamp>::GetRHI();
 	PassParameters->PSParameters.RenderTargets[0]		= FRenderTargetBinding(SceneColor, ERenderTargetLoadAction::ELoad);
@@ -250,6 +260,7 @@ void NiagaraDebugShaders::DrawDebugLines(
 	PassParameters->PSParameters.OriginalViewRectMin = FVector2f(View.ViewRect.Min);
 	PassParameters->PSParameters.OriginalViewSize = FVector2f(View.ViewRect.Width(), View.ViewRect.Height());
 	PassParameters->PSParameters.OriginalBufferInvSize = FVector2f(1.f / SceneDepth->Desc.Extent.X, 1.f / SceneDepth->Desc.Extent.Y);
+	PassParameters->PSParameters.OccludedColorScale = GNiagaraGpuComputeDebug_OccludedLineColorScale;
 	PassParameters->PSParameters.DepthTexture = SceneDepth;
 	PassParameters->PSParameters.DepthSampler = TStaticSamplerState<SF_Point, AM_Clamp, AM_Clamp, AM_Clamp>::GetRHI();
 	PassParameters->PSParameters.RenderTargets[0] = FRenderTargetBinding(SceneColor, ERenderTargetLoadAction::ELoad);
