First pass of pulling simulation data from live games via the Niagara Outliner and Niagara Sim Cache.

TODO:
- Sim Cache Asset UI needs some love for this to be really usable.
- Certain to be some cases that this falls over that we'll have to work through.

diff --git a/Engine/Plugins/FX/Niagara/Source/Niagara/Private/NiagaraDebuggerClient.cpp b/Engine/Plugins/FX/Niagara/Source/Niagara/Private/NiagaraDebuggerClient.cpp
--- a/Engine/Plugins/FX/Niagara/Source/Niagara/Private/NiagaraDebuggerClient.cpp
+++ b/Engine/Plugins/FX/Niagara/Source/Niagara/Private/NiagaraDebuggerClient.cpp
@@ -9,6 +9,8 @@
 #include "NiagaraComponent.h"
 #include "NiagaraDebugHud.h"
 #include "Containers/Ticker.h"
+#include "NiagaraSimCache.h"
+#include "Serialization\ObjectAndNameAsStringProxyArchive.h"
 
 #if WITH_NIAGARA_DEBUGGER
 
@@ -29,7 +31,8 @@ FNiagaraDebuggerClient::FNiagaraDebuggerClient()
 		.Handling<FNiagaraDebuggerExecuteConsoleCommand>(this, &FNiagaraDebuggerClient::HandleExecConsoleCommandMessage)
  		.Handling<FNiagaraDebugHUDSettingsData>(this, &FNiagaraDebuggerClient::HandleDebugHUDSettingsMessage)
 		.Handling<FNiagaraRequestSimpleClientInfoMessage>(this, &FNiagaraDebuggerClient::HandleRequestSimpleClientInfoMessage)
-		.Handling<FNiagaraOutlinerCaptureSettings>(this, &FNiagaraDebuggerClient::HandleOutlinerSettingsMessage);
+		.Handling<FNiagaraOutlinerCaptureSettings>(this, &FNiagaraDebuggerClient::HandleOutlinerSettingsMessage)
+		.Handling<FNiagaraSystemSimCacheCaptureRequest>(this, &FNiagaraDebuggerClient::HandleSimCacheCaptureRequestMessage);
 
  	if (MessageEndpoint.IsValid())
  	{
@@ -55,6 +58,30 @@ FNiagaraDebuggerClient::~FNiagaraDebuggerClient()
 bool FNiagaraDebuggerClient::Tick(float DeltaSeconds)
 {
 	//Keep ticking until we destroy the debugger.
+
+
+	//Process any pending sim cache captures we have.
+	for (auto It = SimCacheCaptures.CreateIterator(); It ; ++It)
+	{
+		FNiagaraSimCacheCaptureInfo& SimCacheCapture = *It;
+		if (SimCacheCapture.Process())
+		{
+			if (MessageEndpoint.IsValid() && Connection.IsValid())
+			{
+				//Sim cache is complete. Send it back to the connected debugger.
+				FNiagaraSystemSimCacheCaptureReply* SimCacheCaptureReply = FMessageEndpoint::MakeMessage<FNiagaraSystemSimCacheCaptureReply>();
+				SimCacheCaptureReply->ComponentName = SimCacheCapture.Request.ComponentName;
+
+				SimCacheCaptureReply->SimCacheData.Reset();
+				FMemoryWriter ArWriter(SimCacheCaptureReply->SimCacheData);
+				FObjectAndNameAsStringProxyArchive ProxyArWriter(ArWriter, false);
+				SimCacheCapture.SimCache->Serialize(ProxyArWriter);
+				
+				MessageEndpoint->Send(SimCacheCaptureReply, EMessageFlags::Reliable, nullptr, TArrayBuilder<FMessageAddress>().Add(Connection), FTimespan::Zero(), FDateTime::MaxValue());
+			}
+			It.RemoveCurrent();
+		}
+	}
 	return true;
 }
 
@@ -231,6 +258,30 @@ void FNiagaraDebuggerClient::HandleOutlinerSettingsMessage(const FNiagaraOutline
 	}
 }
 
+void FNiagaraDebuggerClient::HandleSimCacheCaptureRequestMessage(const FNiagaraSystemSimCacheCaptureRequest& Message, const TSharedRef<IMessageContext, ESPMode::ThreadSafe>& Context)
+{
+	if (ensure(Context->GetSender() == Connection))
+	{
+		UNiagaraComponent* Comp = nullptr;
+
+		for (TObjectIterator<UNiagaraComponent> It; It; ++It)
+		{
+			if (*It->GetPathName() == Message.ComponentName)
+			{
+				Comp = *It;
+				break;
+			}
+		}
+
+		if (Comp)
+		{
+			FNiagaraSimCacheCaptureInfo& NewCapture = SimCacheCaptures.AddDefaulted_GetRef();
+			NewCapture.Request = Message;
+			NewCapture.Component = Comp;
+		}
+	}
+}
+
 void FNiagaraDebuggerClient::CloseConnection()
 {
 	if (MessageEndpoint.IsValid() && Connection.IsValid())
@@ -441,4 +492,32 @@ bool FNiagaraDebuggerClient::UpdateOutliner(float DeltaSeconds)
 	return false;
 }
 
-#endif//WITH_NIAGARA_DEBUGGER
\ No newline at end of file
+//////////////////////////////////////////////////////////////////////////
+
+bool FNiagaraSimCacheCaptureInfo::Process()
+{
+	UNiagaraComponent* Comp = Component.Get();
+	if (Comp == nullptr || ProcessedFrames >= Request.CaptureDelayFrames + Request.CaptureFrames)
+	{
+		//Capture is complete		
+		SimCache->EndWrite();
+		return true;
+	}
+
+	if(ProcessedFrames >= Request.CaptureDelayFrames)
+	{
+		//First Frame of Capture, init the sim cache.
+		if (ProcessedFrames == Request.CaptureDelayFrames)
+		{	
+			SimCache.Reset(NewObject<UNiagaraSimCache>(GetTransientPackage()));
+			SimCache->BeginWrite(Comp);
+		}
+
+		SimCache->WriteFrame(Comp);
+	}
+	++ProcessedFrames;
+
+	return false;
+}
+
+#endif//WITH_NIAGARA_DEBUGGER
diff --git a/Engine/Plugins/FX/Niagara/Source/Niagara/Private/NiagaraDebuggerClient.h b/Engine/Plugins/FX/Niagara/Source/Niagara/Private/NiagaraDebuggerClient.h
--- a/Engine/Plugins/FX/Niagara/Source/Niagara/Private/NiagaraDebuggerClient.h
+++ b/Engine/Plugins/FX/Niagara/Source/Niagara/Private/NiagaraDebuggerClient.h
@@ -11,6 +11,8 @@ Controller class running on game clients that handles the passing of messages to
 #include "NiagaraDebuggerCommon.h"
 #include "IMessageContext.h"
 #include "Particles/ParticlePerfStatsManager.h"
+#include "NiagaraSimCache.h"
+#include "UObject/StrongObjectPtr.h"
 
 class FMessageEndpoint;
 
@@ -32,6 +34,19 @@ public:
 
 #if WITH_NIAGARA_DEBUGGER
 
+
+struct FNiagaraSimCacheCaptureInfo
+{
+	uint32 ProcessedFrames = 0;
+	FNiagaraSystemSimCacheCaptureRequest Request;
+	TWeakObjectPtr<UNiagaraComponent> Component;
+
+	TStrongObjectPtr<UNiagaraSimCache> SimCache = nullptr;
+
+	/** Process this request. Captures data where needed. Returns true if complete. */
+	bool Process();
+};
+
 DECLARE_LOG_CATEGORY_EXTERN(LogNiagaraDebuggerClient, Log, All);
 
 class FNiagaraDebuggerClient
@@ -55,6 +70,7 @@ private:
 	void HandleDebugHUDSettingsMessage(const FNiagaraDebugHUDSettingsData& Message, const TSharedRef<IMessageContext, ESPMode::ThreadSafe>& Context);
 	void HandleRequestSimpleClientInfoMessage(const FNiagaraRequestSimpleClientInfoMessage& Message, const TSharedRef<IMessageContext, ESPMode::ThreadSafe>& Context);
 	void HandleOutlinerSettingsMessage(const FNiagaraOutlinerCaptureSettings& Message, const TSharedRef<IMessageContext, ESPMode::ThreadSafe>& Context);
+	void HandleSimCacheCaptureRequestMessage(const FNiagaraSystemSimCacheCaptureRequest& Message, const TSharedRef<IMessageContext, ESPMode::ThreadSafe>& Context);
 
 	/** Closes any currently active connection. */
 	void CloseConnection();
@@ -85,6 +101,9 @@ private:
 #if WITH_PARTICLE_PERF_STATS
 	TSharedPtr<FNiagaraOutlinerPerfListener, ESPMode::ThreadSafe> StatsListener;
 #endif
+
+	/* All currently ongoing sim cache captures. */
+	TArray<FNiagaraSimCacheCaptureInfo> SimCacheCaptures;
 };
 
 #endif
\ No newline at end of file
diff --git a/Engine/Plugins/FX/Niagara/Source/Niagara/Public/NiagaraDebuggerCommon.h b/Engine/Plugins/FX/Niagara/Source/Niagara/Public/NiagaraDebuggerCommon.h
--- a/Engine/Plugins/FX/Niagara/Source/Niagara/Public/NiagaraDebuggerCommon.h
+++ b/Engine/Plugins/FX/Niagara/Source/Niagara/Public/NiagaraDebuggerCommon.h
@@ -10,6 +10,7 @@ All common code shared between the editor side debugger and debugger clients run
 #include "Misc/NotifyHook.h"
 #include "NiagaraTypes.h"
 #include "NiagaraCommon.h"
+#include "NiagaraSimCache.h"
 #include "NiagaraDebuggerCommon.generated.h"
 
 //////////////////////////////////////////////////////////////////////////
@@ -761,6 +762,10 @@ struct NIAGARA_API FNiagaraOutlinerCaptureSettings
 	
 	UPROPERTY(EditAnywhere, Config, Category = "Settings")
 	bool bGatherPerfData = true;
+
+	/** How many frames capture when capturing a sim cache. */
+	UPROPERTY(EditAnywhere, Config, Category = "Settings")
+	uint32 SimCacheCaptureFrames = 10;
 };
 
 /** Simple information on the connected client for use in continuous or immediate response UI elements. */
@@ -786,6 +791,39 @@ struct NIAGARA_API FNiagaraSimpleClientInfo
 	TArray<FString> Emitters;
 };
 
+/** Message sent from the debugger to a client to request a sim cache capture for a particular component. */
+USTRUCT()
+struct NIAGARA_API FNiagaraSystemSimCacheCaptureRequest
+{
+	GENERATED_BODY()
+
+	/** Name of the component we're going to capture. */
+	UPROPERTY(EditAnywhere, Category = "Settings")
+	FName ComponentName;
+
+	/** How many frames to delay capture. */
+	UPROPERTY(EditAnywhere, Config, Category="Settings")
+	uint32 CaptureDelayFrames = 60;
+
+	/** How many frames to capture. */
+	UPROPERTY(EditAnywhere, Config, Category = "Settings")
+	uint32 CaptureFrames = 10;
+};
+
+
+/** Message sent from a debugger client to a connected debugger containing the results of a sim cache capture. */
+USTRUCT()
+struct NIAGARA_API FNiagaraSystemSimCacheCaptureReply
+{
+	GENERATED_BODY()
+
+	/** Name of the captured component. */
+	UPROPERTY()
+	FName ComponentName;
+	
+	UPROPERTY()
+	TArray<uint8> SimCacheData;
+};
 
 enum class ENiagaraDebugMessageType : uint8
 {
diff --git a/Engine/Plugins/FX/Niagara/Source/NiagaraEditor/Private/Customizations/NiagaraOutlinerCustomization.cpp b/Engine/Plugins/FX/Niagara/Source/NiagaraEditor/Private/Customizations/NiagaraOutlinerCustomization.cpp
--- a/Engine/Plugins/FX/Niagara/Source/NiagaraEditor/Private/Customizations/NiagaraOutlinerCustomization.cpp
+++ b/Engine/Plugins/FX/Niagara/Source/NiagaraEditor/Private/Customizations/NiagaraOutlinerCustomization.cpp
@@ -22,12 +22,15 @@
 #include "Widgets/Layout/SSeparator.h"
 #include "Widgets/Layout/SSplitter.h"
 #include "Widgets/Images/SImage.h"
+#include "Widgets/Input/SButton.h"
 ///Niagara
 #include "NiagaraEditorModule.h"
 #include "NiagaraComponent.h"
 #include "NiagaraEditorStyle.h"
 #include "Widgets/SVerticalResizeBox.h"
 
+#include "NiagaraSimCache.h"
+
 
 #if WITH_NIAGARA_DEBUGGER
 
@@ -1269,12 +1272,107 @@ TSharedRef<SWidget> FNiagaraOutlinerTreeComponentItem::GetHeaderWidget()
 
 			//TODO: Add baseline comparison traffic light icon with tooltip containing details.
 		}
+		else if (Outliner->ViewSettings.ViewMode == ENiagaraOutlinerViewModes::Debug)
+		{			
+			Box->AddSlot()
+			.AutoWidth()
+			.Padding(FMargin(HeaderPadding, 0.0f, HeaderPadding, 0.0f))
+			[
+				SNew(SButton)
+				.OnClicked_Raw(this, &FNiagaraOutlinerTreeComponentItem::CaputreSimCache)
+				.ToolTipText(LOCTEXT("NiagaraOutlineCaptureSimCacheTooltip", "Capture a new Sim Cache for this component."))
+				.ButtonStyle(FAppStyle::Get(), "FlatButton.Success")
+				//.ForegroundColor(FLinearColor::White)
+				[
+					SNew(STextBlock)
+					.Text(LOCTEXT("NiagaraOutlinerCaptureSimCache", "Capture"))
+					//.TextStyle(FAppStyle::Get(), "ContentBrowser.TopBar.Font")
+				]	
+			];
+
+			if (TObjectPtr<UNiagaraSimCache> SimCache = Outliner->FindSimCache(*Data->ComponentName))
+			{
+				Box->AddSlot()
+				.AutoWidth()
+				.Padding(FMargin(HeaderPadding, 0.0f, HeaderPadding, 0.0f))
+				[
+					SNew(SButton)
+					.OnClicked_Raw(this, &FNiagaraOutlinerTreeComponentItem::OpenSimCache)
+					.ToolTipText(LOCTEXT("NiagaraOutlineOpenSimCacheTooltip", "Sim Cache is available for this component. Open in new window."))
+					.ButtonStyle(FAppStyle::Get(), "FlatButton.Success")
+					//.ForegroundColor(FLinearColor::White)
+					[
+						SNew(STextBlock)
+						.Text(LOCTEXT("NiagaraOutlinerOpenSimCache", "Sim Cache"))
+						//.TextStyle(FAppStyle::Get(), "ContentBrowser.TopBar.Font")
+					]	
+				];
+			}
+			else
+			{
+				Box->AddSlot()
+					.AutoWidth()
+					.Padding(FMargin(HeaderPadding, 0.0f, HeaderPadding, 0.0f))
+					[
+						SNew(SButton)						
+						.ToolTipText(LOCTEXT("NiagaraOutlineOpenSimCacheTooltipDisabled", "No Sim Cache available. Please capture one."))
+						.ButtonStyle(FAppStyle::Get(), "FlatButton.Default")
+						//.ForegroundColor(FLinearColor::White)
+					[
+						SNew(STextBlock)
+						.Text(LOCTEXT("NiagaraOutlinerOpenSimCache", "Sim Cache"))
+						//.TextStyle(FAppStyle::Get(), "ContentBrowser.TopBar.Font")
+					]
+				];
+			}
+
+			//TODO:	Add playback controls per component.
+			//TODO: Add controls to set the current debug hud component filter to this component.
+		}
 
 		return Box;
 	}
 	return SNullWidget::NullWidget;
 }
 
+FReply FNiagaraOutlinerTreeComponentItem::OpenSimCache()
+{
+	if (FNiagaraOutlinerSystemInstanceData* Data = (FNiagaraOutlinerSystemInstanceData*)GetData())
+	{
+		TSharedPtr<SNiagaraOutlinerTree> Tree = OwnerTree.Pin();
+		check(Tree.IsValid());//If we managed to get valid data then the tree should be valid.
+		TSharedPtr<FNiagaraDebugger>& Debugger = Tree->GetDebugger();
+		check(Debugger.IsValid());
+		UNiagaraOutliner* Outliner = Debugger->GetOutliner();
+		check(Outliner);
+
+		if (TObjectPtr<UNiagaraSimCache> SimCache = Outliner->FindSimCache(*Data->ComponentName))
+		{
+			GEditor->GetEditorSubsystem<UAssetEditorSubsystem>()->OpenEditorForAsset(Cast<UObject,UNiagaraSimCache>(SimCache.Get()), EToolkitMode::Standalone);
+		}
+	}
+	return FReply::Handled();
+}
+
+FReply FNiagaraOutlinerTreeComponentItem::CaputreSimCache()
+{
+	if (FNiagaraOutlinerSystemInstanceData* Data = (FNiagaraOutlinerSystemInstanceData*)GetData())
+	{
+		TSharedPtr<SNiagaraOutlinerTree> Tree = OwnerTree.Pin();
+		check(Tree.IsValid());//If we managed to get valid data then the tree should be valid.
+		TSharedPtr<FNiagaraDebugger>& Debugger = Tree->GetDebugger();
+		check(Debugger.IsValid());
+		UNiagaraOutliner* Outliner = Debugger->GetOutliner();
+		check(Outliner);
+
+		if (Outliner->CaptureSettings.SimCacheCaptureFrames > 0)
+		{
+			Debugger->TriggerSimCacheCapture(*Data->ComponentName, Outliner->CaptureSettings.CaptureDelayFrames, Outliner->CaptureSettings.SimCacheCaptureFrames);
+		}
+	}
+	return FReply::Handled();
+}
+
 void FNiagaraOutlinerTreeEmitterItem::SortChildren() 
 {
 
diff --git a/Engine/Plugins/FX/Niagara/Source/NiagaraEditor/Private/Customizations/NiagaraOutlinerCustomization.h b/Engine/Plugins/FX/Niagara/Source/NiagaraEditor/Private/Customizations/NiagaraOutlinerCustomization.h
--- a/Engine/Plugins/FX/Niagara/Source/NiagaraEditor/Private/Customizations/NiagaraOutlinerCustomization.h
+++ b/Engine/Plugins/FX/Niagara/Source/NiagaraEditor/Private/Customizations/NiagaraOutlinerCustomization.h
@@ -142,6 +142,10 @@ struct FNiagaraOutlinerTreeComponentItem : public FNiagaraOutlinerTreeItem
 
 		return Ret;
 	}
+
+	//Sim Cache Debugging Controls.
+	FReply OpenSimCache();
+	FReply CaputreSimCache();
 };
 
 struct FNiagaraOutlinerTreeEmitterItem : public FNiagaraOutlinerTreeItem
diff --git a/Engine/Plugins/FX/Niagara/Source/NiagaraEditor/Private/NiagaraDebugger.cpp b/Engine/Plugins/FX/Niagara/Source/NiagaraEditor/Private/NiagaraDebugger.cpp
--- a/Engine/Plugins/FX/Niagara/Source/NiagaraEditor/Private/NiagaraDebugger.cpp
+++ b/Engine/Plugins/FX/Niagara/Source/NiagaraEditor/Private/NiagaraDebugger.cpp
@@ -39,7 +39,8 @@ void FNiagaraDebugger::Init()
 		.Handling<FNiagaraDebuggerAcceptConnection>(this, &FNiagaraDebugger::HandleConnectionAcceptedMessage)
 		.Handling<FNiagaraDebuggerConnectionClosed>(this, &FNiagaraDebugger::HandleConnectionClosedMessage)
 		.Handling<FNiagaraDebuggerOutlinerUpdate>(this, &FNiagaraDebugger::HandleOutlinerUpdateMessage)
-		.Handling<FNiagaraSimpleClientInfo>(this, &FNiagaraDebugger::UpdateSimpleClientInfo);
+		.Handling<FNiagaraSimpleClientInfo>(this, &FNiagaraDebugger::UpdateSimpleClientInfo)
+		.Handling<FNiagaraSystemSimCacheCaptureReply>(this, &FNiagaraDebugger::HandleSimCacheCaptureReply);
 
 #if WITH_UNREAL_TARGET_DEVELOPER_TOOLS
 	ISessionServicesModule& SessionServicesModule = FModuleManager::LoadModuleChecked<ISessionServicesModule>("SessionServices");
@@ -142,6 +143,22 @@ void FNiagaraDebugger::TriggerOutlinerCapture()
 	}
 }
 
+void FNiagaraDebugger::TriggerSimCacheCapture(FName ComponentName, int32 CaptureDelay, int32 CaptureFrames)
+{
+	auto SendCaptureRequest = [&](FNiagaraDebugger::FClientInfo& Client)
+	{
+		FNiagaraSystemSimCacheCaptureRequest* Message = FMessageEndpoint::MakeMessage<FNiagaraSystemSimCacheCaptureRequest>();
+		Message->CaptureDelayFrames = CaptureDelay;
+		Message->CaptureFrames = CaptureFrames;
+		Message->ComponentName = ComponentName;
+
+		UE_LOG(LogNiagaraDebugger, Log, TEXT("Sending request to capture sim cache for component %s. | Session: %s | Instance: %s |"), *ComponentName.ToString(), *Client.SessionId.ToString(), *Client.InstanceId.ToString());
+		MessageEndpoint->Send(Message, EMessageFlags::Reliable, nullptr, TArrayBuilder<FMessageAddress>().Add(Client.Address), FTimespan::Zero(), FDateTime::MaxValue());
+	};
+
+	ForAllConnectedClients(SendCaptureRequest);
+}
+
 void FNiagaraDebugger::SessionManager_OnSessionSelectionChanged(const TSharedPtr<ISessionInfo>& Session)
 {
 #if WITH_UNREAL_TARGET_DEVELOPER_TOOLS
@@ -325,6 +342,16 @@ void FNiagaraDebugger::UpdateSimpleClientInfo(const FNiagaraSimpleClientInfo& Me
 	OnSimpleClientInfoChangedDelegate.Broadcast(SimpleClientInfo);
 }
 
+void FNiagaraDebugger::HandleSimCacheCaptureReply(const FNiagaraSystemSimCacheCaptureReply& Message, const TSharedRef<IMessageContext, ESPMode::ThreadSafe>& Context)
+{
+	UE_LOG(LogNiagaraDebugger, Log, TEXT("Recieved Niagara Sim Cache Capture Reply. | Sender: %s "), *Context->GetSender().ToString());
+	
+	if (UNiagaraOutliner* Outliner = GetOutliner())
+	{
+		Outliner->UpdateSystemSimCache(Message);
+	}
+}
+
 #undef LOCTEXT_NAMESPACE
 
 
diff --git a/Engine/Plugins/FX/Niagara/Source/NiagaraEditor/Private/NiagaraDebugger.h b/Engine/Plugins/FX/Niagara/Source/NiagaraEditor/Private/NiagaraDebugger.h
--- a/Engine/Plugins/FX/Niagara/Source/NiagaraEditor/Private/NiagaraDebugger.h
+++ b/Engine/Plugins/FX/Niagara/Source/NiagaraEditor/Private/NiagaraDebugger.h
@@ -57,6 +57,7 @@ public:
 	void ExecConsoleCommand(const TCHAR* Cmd, bool bRequiresWorld);
 	void UpdateDebugHUDSettings();
 	void TriggerOutlinerCapture();
+	void TriggerSimCacheCapture(FName ComponentName, int32 CaptureDelay, int32 CaptureFrames);
 
 	void RequestUpdatedClientInfo();
 
@@ -86,6 +87,9 @@ protected:
 	// Handles a message that updates the simple client info.
 	void UpdateSimpleClientInfo(const FNiagaraSimpleClientInfo& Message, const TSharedRef<IMessageContext, ESPMode::ThreadSafe>& Context);
 
+	// Handles an reply messages from a sim cache capture request.
+	void HandleSimCacheCaptureReply(const FNiagaraSystemSimCacheCaptureReply& Message, const TSharedRef<IMessageContext, ESPMode::ThreadSafe>& Context);
+
 	/** Callback from session manager when the selected session is changed. */
 	void SessionManager_OnSessionSelectionChanged(const TSharedPtr<ISessionInfo>& Session);
 
diff --git a/Engine/Plugins/FX/Niagara/Source/NiagaraEditor/Private/NiagaraOutliner.cpp b/Engine/Plugins/FX/Niagara/Source/NiagaraEditor/Private/NiagaraOutliner.cpp
--- a/Engine/Plugins/FX/Niagara/Source/NiagaraEditor/Private/NiagaraOutliner.cpp
+++ b/Engine/Plugins/FX/Niagara/Source/NiagaraEditor/Private/NiagaraOutliner.cpp
@@ -2,6 +2,10 @@
 
 #include "NiagaraOutliner.h"
 #include "NiagaraDebugger.h"
+#include "NiagaraEditorCommon.h"
+#include "NiagaraSimCache.h"
+#include "Serialization\ObjectAndNameAsStringProxyArchive.h"
+
 
 UNiagaraOutliner::UNiagaraOutliner(const FObjectInitializer& Initializer)
 {
@@ -39,6 +43,34 @@ void UNiagaraOutliner::UpdateData(const FNiagaraOutlinerData& NewData)
 	//Possibly keep them in the UI optionally but colour/mark them as dead until the user opts to remove them or on some timed interval.
 }
 
+void UNiagaraOutliner::UpdateSystemSimCache(const FNiagaraSystemSimCacheCaptureReply& Reply)
+{
+	if (Reply.SimCacheData.Num() > 0)
+	{
+		TObjectPtr<UNiagaraSimCache>& SimCache = SystemSimCaches.FindOrAdd(Reply.ComponentName);
+		
+		SimCache = NewObject<UNiagaraSimCache>(this);
+
+		FMemoryReader ArReader(Reply.SimCacheData);
+		FObjectAndNameAsStringProxyArchive ProxyArReader(ArReader, false);
+		SimCache->Serialize(ProxyArReader);
+	}
+	else
+	{
+		UE_LOG(LogNiagaraEditor, Warning, TEXT("Niagara Outliner recieved empty sim cache data."));
+	}
+	OnChanged();
+}
+
+UNiagaraSimCache* UNiagaraOutliner::FindSimCache(FName ComponentName)
+{
+	if (TObjectPtr<UNiagaraSimCache>* SimCachePtr = SystemSimCaches.Find(ComponentName))
+	{
+		return SimCachePtr->Get();
+	}
+	return nullptr;
+}
+
 const FNiagaraOutlinerWorldData* UNiagaraOutliner::FindWorldData(const FString& WorldName)
 {
 	return Data.WorldData.Find(WorldName);
diff --git a/Engine/Plugins/FX/Niagara/Source/NiagaraEditor/Private/NiagaraOutliner.h b/Engine/Plugins/FX/Niagara/Source/NiagaraEditor/Private/NiagaraOutliner.h
--- a/Engine/Plugins/FX/Niagara/Source/NiagaraEditor/Private/NiagaraOutliner.h
+++ b/Engine/Plugins/FX/Niagara/Source/NiagaraEditor/Private/NiagaraOutliner.h
@@ -4,6 +4,7 @@
 
 #include "NiagaraDebuggerCommon.h"
 #include "Misc/NotifyHook.h"
+#include "NiagaraSimCache.h"
 #include "NiagaraOutliner.generated.h"
 
 UENUM()
@@ -13,6 +14,8 @@ enum class ENiagaraOutlinerViewModes: uint8
 	State,
 	/** Outliner displays performance data for each item. */
 	Performance,
+	/** Outliner displays debugging controls for each item. */
+	Debug,
 };
 
 UENUM()
@@ -142,6 +145,9 @@ public:
 #endif
 
 	void UpdateData(const FNiagaraOutlinerData& NewData);
+	void UpdateSystemSimCache(const FNiagaraSystemSimCacheCaptureReply& Reply);
+
+	UNiagaraSimCache* FindSimCache(FName ComponentName);
 
 	const FNiagaraOutlinerWorldData* FindWorldData(const FString& WorldName);
 	const FNiagaraOutlinerSystemData* FindSystemData(const FString& WorldName, const FString& SystemName);
@@ -160,4 +166,7 @@ public:
 
 	UPROPERTY(VisibleAnywhere, Category="Outliner", Transient)
 	FNiagaraOutlinerData Data;
+
+	UPROPERTY()
+	TMap<FName, TObjectPtr<UNiagaraSimCache>> SystemSimCaches;
 };
\ No newline at end of file
diff --git a/Engine/Plugins/FX/Niagara/Source/NiagaraEditor/Private/Widgets/SNiagaraDebugger.cpp b/Engine/Plugins/FX/Niagara/Source/NiagaraEditor/Private/Widgets/SNiagaraDebugger.cpp
--- a/Engine/Plugins/FX/Niagara/Source/NiagaraEditor/Private/Widgets/SNiagaraDebugger.cpp
+++ b/Engine/Plugins/FX/Niagara/Source/NiagaraEditor/Private/Widgets/SNiagaraDebugger.cpp
@@ -334,7 +334,7 @@ namespace NiagaraOutlinerTab
  			{
 				TSharedRef<SWidget> DelayWidget = SNew(SBorder)
 				.BorderImage(FAppStyle::GetBrush("NoBorder"))
-				.Padding(FMargin(3.0, 0.0f, 3.0f, 0.0f))
+				.Padding(FMargin(6.0, 0.0f, 6.0f, 0.0f))
 				.ToolTipText(LOCTEXT("OutlinerDelayTooltip", "Number of frames to delay between a capture being triggered and it being taken.\nThis provides time to affect the scene and also defines the length of time performance data is gathered."))				
 				[
 					SNew(SEditableTextBox)
@@ -357,6 +357,34 @@ namespace NiagaraOutlinerTab
 				];
 				ToolbarBuilder.AddToolBarWidget(DelayWidget, LOCTEXT("OutlinerDelay", "Delay"));
 			}
+			
+ 			// Sim Cache Capture frames
+ 			{
+				TSharedRef<SWidget> FramesWidget = SNew(SBorder)
+				.BorderImage(FAppStyle::GetBrush("NoBorder"))
+				.Padding(FMargin(6.0, 0.0f, 6.0f, 0.0f))
+				.ToolTipText(LOCTEXT("OutlinerSimCacheCaptureFramesTooltip", "Number of frames to capture when capturing Sim Cache Data."))				
+				[
+					SNew(SEditableTextBox)
+					.OnTextCommitted(FOnTextCommitted::CreateLambda([Debugger](const FText& InText, ETextCommit::Type CommitInfo)
+					{
+						if (UNiagaraOutliner* Outliner = Debugger->GetOutliner())
+						{
+							LexFromString(Outliner->CaptureSettings.SimCacheCaptureFrames, *InText.ToString());
+							Outliner->OnChanged();
+						}
+					}))
+					.Text(MakeAttributeLambda([Debugger]()
+					{
+						if (UNiagaraOutliner* Outliner = Debugger->GetOutliner())
+						{
+							return FText::AsNumber(Outliner->CaptureSettings.SimCacheCaptureFrames);
+						}
+						return FText::GetEmpty();
+					}))					
+				];
+				ToolbarBuilder.AddToolBarWidget(FramesWidget, LOCTEXT("OutlinerSimCacheFrames", "Sim Cache Frames"));
+			}
 		}
 
 		ToolbarBuilder.AddSeparator();
