Added a search limit to dependent assets and improved reference count performance

diff --git a/Engine/Plugins/FX/Niagara/Source/NiagaraEditor/Private/NiagaraEditorSettings.cpp b/Engine/Plugins/FX/Niagara/Source/NiagaraEditor/Private/NiagaraEditorSettings.cpp
--- a/Engine/Plugins/FX/Niagara/Source/NiagaraEditor/Private/NiagaraEditorSettings.cpp
+++ b/Engine/Plugins/FX/Niagara/Source/NiagaraEditor/Private/NiagaraEditorSettings.cpp
@@ -458,6 +458,11 @@ bool UNiagaraEditorSettings::GetDisplayAffectedAssetStats() const
 	return bDisplayAffectedAssetStats;
 }
 
+int32 UNiagaraEditorSettings::GetAssetStatsSearchLimit() const
+{
+	return AffectedAssetSearchLimit;
+}
+
 FNiagaraNewAssetDialogConfig UNiagaraEditorSettings::GetNewAssetDailogConfig(FName InDialogConfigKey) const
 {
 	const FNiagaraNewAssetDialogConfig* Config = NewAssetDialogConfigMap.Find(InDialogConfigKey);
diff --git a/Engine/Plugins/FX/Niagara/Source/NiagaraEditor/Private/NiagaraEditorUtilities.cpp b/Engine/Plugins/FX/Niagara/Source/NiagaraEditor/Private/NiagaraEditorUtilities.cpp
--- a/Engine/Plugins/FX/Niagara/Source/NiagaraEditor/Private/NiagaraEditorUtilities.cpp
+++ b/Engine/Plugins/FX/Niagara/Source/NiagaraEditor/Private/NiagaraEditorUtilities.cpp
@@ -3831,6 +3831,65 @@ float FNiagaraEditorUtilities::GetScalabilityTintAlpha(FNiagaraEmitterHandle* Em
 	return EmitterHandle->GetIsEnabled() ? 1 : 0.5f;
 }
 
+int FNiagaraEditorUtilities::GetReferencedAssetCount(const FAssetData& SourceAsset,	TFunction<ETrackAssetResult(const FAssetData&)> Predicate)
+{
+	if (!SourceAsset.IsValid())
+	{
+		return 0;
+	}
+	
+	FAssetRegistryModule& AssetRegistryModule = FModuleManager::LoadModuleChecked<FAssetRegistryModule>(TEXT("AssetRegistry"));
+	IAssetRegistry& AssetRegistry = AssetRegistryModule.Get();
+	int32 SearchLimit = GetDefault<UNiagaraEditorSettings>()->GetAssetStatsSearchLimit();
+	int32 Count = 0;
+	
+	TSet<FName> SeenObjects;
+	SeenObjects.Add(SourceAsset.PackageName);
+	TArray<FName> AssetsToCheck;
+	AssetRegistry.GetReferencers(SourceAsset.PackageName, AssetsToCheck);
+	while (AssetsToCheck.Num() > 0)
+	{
+		FName AssetPath = AssetsToCheck[0];
+		AssetsToCheck.RemoveAtSwap(0);
+		if (SeenObjects.Contains(AssetPath))
+		{
+			// prevent asset loops
+			continue;
+		}
+		SeenObjects.Add(AssetPath);
+
+		if ((Count + AssetsToCheck.Num()) > SearchLimit)
+		{
+			// if we are over the search limit then just tally up the remaining references for a rough estimate
+			Count++;
+		}
+		else
+		{
+			// we should only ever get one asset per package, but the api wants to return a list
+			TArray<FAssetData> OutAssetData;
+			AssetRegistry.GetAssetsByPackageName(AssetPath, OutAssetData, true);
+			for (const FAssetData& AssetToCheck : OutAssetData)
+			{
+				ETrackAssetResult Result = Predicate(AssetToCheck);
+				if (Result == ETrackAssetResult::Count)
+				{
+					Count++;
+				}
+				else if (Result == ETrackAssetResult::CountRecursive)
+				{
+					Count++;
+					AssetRegistry.GetReferencers(AssetToCheck.PackageName, AssetsToCheck);
+				}
+				else
+				{
+					ensure(Result == ETrackAssetResult::Ignore);
+				}
+			}
+		}
+	}
+	return FMath::Max(Count, 0);
+}
+
 void FNiagaraParameterUtilities::FilterToRelevantStaticVariables(const TArray<FNiagaraVariable>& InVars, TArray<FNiagaraVariable>& OutVars, FName InOldEmitterAlias, FName InNewEmitterAlias, bool bFilterByEmitterAliasAndConvertToUnaliased)
 {
 	FNiagaraAliasContext RenameContext(ENiagaraScriptUsage::ParticleSpawnScript);
diff --git a/Engine/Plugins/FX/Niagara/Source/NiagaraEditor/Private/Widgets/SNiagaraScriptGraph.cpp b/Engine/Plugins/FX/Niagara/Source/NiagaraEditor/Private/Widgets/SNiagaraScriptGraph.cpp
--- a/Engine/Plugins/FX/Niagara/Source/NiagaraEditor/Private/Widgets/SNiagaraScriptGraph.cpp
+++ b/Engine/Plugins/FX/Niagara/Source/NiagaraEditor/Private/Widgets/SNiagaraScriptGraph.cpp
@@ -680,12 +680,24 @@ TOptional<SSearchBox::FSearchResultData> SNiagaraScriptGraph::GetSearchResultDat
 
 FText SNiagaraScriptGraph::GetScriptSubheaderText() const
 {
-	return FText::Format(LOCTEXT("EmitterSubheaderText", "Note: editing this script will affect {0} dependent assets (across all versions)!"), GetScriptAffectedAssets());
+	int32 SearchLimit = GetDefault<UNiagaraEditorSettings>()->GetAssetStatsSearchLimit();
+	int32 AffectedAssets = GetScriptAffectedAssets();
+	bool bVersioned = false;
+	if (UNiagaraScriptSource* ScriptSource = ViewModel->GetScriptSource())
+	{
+		if (UNiagaraScript* Script = ScriptSource->GetTypedOuter<UNiagaraScript>())
+		{
+			bVersioned = Script->IsVersioningEnabled();
+		}
+	}
+	FText VersionText = bVersioned ? LOCTEXT("ScriptSubheaderVersionInfoText", "(across all versions)") : FText();
+	FText LimitReachedText = AffectedAssets >= SearchLimit ? LOCTEXT("ScriptSubheaderLimitReachedText", "more than ") : FText();
+	return FText::Format(LOCTEXT("ScriptSubheaderText", "Note: editing this script will affect {0}{1} dependent {1}|plural(one=asset,other=assets)! {2}"), LimitReachedText, FText::AsNumber(AffectedAssets), VersionText);
 }
 
 EVisibility SNiagaraScriptGraph::GetScriptSubheaderVisibility() const
 {
-	return bShowHeader && (GetScriptAffectedAssets() > 1) ? EVisibility::Visible : EVisibility::Collapsed;
+	return bShowHeader && (GetScriptAffectedAssets() > 0) ? EVisibility::Visible : EVisibility::Collapsed;
 }
 
 FLinearColor SNiagaraScriptGraph::GetScriptSubheaderColor() const
@@ -699,7 +711,8 @@ int32 SNiagaraScriptGraph::GetScriptAffectedAssets() const
 	{
 		return 0;
 	}
-	if (GetDefault<UNiagaraEditorSettings>()->GetDisplayAffectedAssetStats() == false)
+	const UNiagaraEditorSettings* EditorSettings = GetDefault<UNiagaraEditorSettings>();
+	if (EditorSettings->GetDisplayAffectedAssetStats() == false)
 	{
 		return 0;
 	}
@@ -714,42 +727,18 @@ int32 SNiagaraScriptGraph::GetScriptAffectedAssets() const
 			return 0;
 		}
 
-		int32 Count = 0;
-		if (EditedAsset.IsValid())
+		ScriptAffectedAssets = FNiagaraEditorUtilities::GetReferencedAssetCount(EditedAsset, [](const FAssetData& AssetToCheck)
 		{
-			TSet<FName> SeenObjects;
-			TArray<FAssetData> AssetsToCheck;
-			AssetsToCheck.Add(EditedAsset);
-
-			// search for assets that reference this script
-			while (AssetsToCheck.Num() > 0)
+			if (AssetToCheck.GetClass() == UNiagaraSystem::StaticClass())
 			{
-				FAssetData AssetToCheck = AssetsToCheck[0];
-				AssetsToCheck.RemoveAtSwap(0);
-				if (SeenObjects.Contains(AssetToCheck.ObjectPath))
-				{
-					continue;
-				}
-				SeenObjects.Add(AssetToCheck.ObjectPath);
-
-				if (AssetToCheck.GetClass() == UNiagaraSystem::StaticClass())
-				{
-					Count++;
-				}
-				else if (AssetToCheck.GetClass() == UNiagaraEmitter::StaticClass() || AssetToCheck.GetClass() == UNiagaraScript::StaticClass())
-				{
-					Count++;
-					TArray<FName> Referencers;
-					AssetRegistry.GetReferencers(AssetToCheck.PackageName, Referencers);
-					for (FName& Referencer : Referencers)
-					{
-						AssetRegistry.GetAssetsByPackageName(Referencer, AssetsToCheck);
-					}
-				}
+				return FNiagaraEditorUtilities::ETrackAssetResult::Count;
 			}
-			Count--; // remove one for our own asset
-		}
-		ScriptAffectedAssets = Count;
+			if (AssetToCheck.GetClass() == UNiagaraEmitter::StaticClass() || AssetToCheck.GetClass() == UNiagaraScript::StaticClass())
+			{
+				return FNiagaraEditorUtilities::ETrackAssetResult::CountRecursive;
+			}
+			return FNiagaraEditorUtilities::ETrackAssetResult::Ignore;
+		});
 	}
 	return ScriptAffectedAssets.Get(0);
 }
diff --git a/Engine/Plugins/FX/Niagara/Source/NiagaraEditor/Public/NiagaraEditorSettings.h b/Engine/Plugins/FX/Niagara/Source/NiagaraEditor/Public/NiagaraEditorSettings.h
--- a/Engine/Plugins/FX/Niagara/Source/NiagaraEditor/Public/NiagaraEditorSettings.h
+++ b/Engine/Plugins/FX/Niagara/Source/NiagaraEditor/Public/NiagaraEditorSettings.h
@@ -298,6 +298,7 @@ public:
 	void SetDisplayAdvancedParameterPanelCategories(bool bInDisplayAdvancedParameterPanelCategories);
 
 	bool GetDisplayAffectedAssetStats() const;
+	int32 GetAssetStatsSearchLimit() const;
 
 	bool IsShowGridInViewport() const;
 	void SetShowGridInViewport(bool bShowGridInViewport);
@@ -370,6 +371,10 @@ private:
 	UPROPERTY(config, EditAnywhere, Category=Niagara)
 	bool bDisplayAffectedAssetStats = true;
 
+	/** The maximum amount of asset references that are searched before stopping. Setting this too high can lead to long load times when opening default assets (basically the same as disabling the breadth limit in the reference viewer). */
+	UPROPERTY(config, EditAnywhere, Category=Niagara, meta=(EditCondition="bDisplayAffectedAssetStats", ClampMin=1))
+	int32 AffectedAssetSearchLimit = 25;
+
 	/** Speeds used for slowing down and speeding up the playback speeds */
 	UPROPERTY(config, EditAnywhere, Category = Niagara)
 	TArray<float> PlaybackSpeeds;
diff --git a/Engine/Plugins/FX/Niagara/Source/NiagaraEditor/Public/NiagaraEditorUtilities.h b/Engine/Plugins/FX/Niagara/Source/NiagaraEditor/Public/NiagaraEditorUtilities.h
--- a/Engine/Plugins/FX/Niagara/Source/NiagaraEditor/Public/NiagaraEditorUtilities.h
+++ b/Engine/Plugins/FX/Niagara/Source/NiagaraEditor/Public/NiagaraEditorUtilities.h
@@ -345,6 +345,20 @@ namespace FNiagaraEditorUtilities
 	void VisitAllNodesConnectedToInputs(UEdGraphNode* StartNode, FNodeVisitor Visitor);
 	
 	NIAGARAEDITOR_API float GetScalabilityTintAlpha(FNiagaraEmitterHandle* EmitterHandle);
+
+	NIAGARAEDITOR_API enum class ETrackAssetResult
+	{
+		// do not count this asset
+		Ignore,
+
+		// count this asset
+		Count,
+
+		// count this asset and also check assets referencing it
+		CountRecursive
+	};
+
+	NIAGARAEDITOR_API int GetReferencedAssetCount(const FAssetData& SourceAsset, TFunction<ETrackAssetResult(const FAssetData&)> Predicate);
 	
 };
 
diff --git a/Engine/Plugins/FX/Niagara/Source/NiagaraEditorWidgets/Private/SNiagaraOverviewGraphTitleBar.cpp b/Engine/Plugins/FX/Niagara/Source/NiagaraEditorWidgets/Private/SNiagaraOverviewGraphTitleBar.cpp
--- a/Engine/Plugins/FX/Niagara/Source/NiagaraEditorWidgets/Private/SNiagaraOverviewGraphTitleBar.cpp
+++ b/Engine/Plugins/FX/Niagara/Source/NiagaraEditorWidgets/Private/SNiagaraOverviewGraphTitleBar.cpp
@@ -4,6 +4,7 @@
 
 #include "NiagaraEditorWidgetsStyle.h"
 #include "NiagaraEditorSettings.h"
+#include "NiagaraEditorUtilities.h"
 #include "ViewModels/NiagaraOverviewGraphViewModel.h"
 #include "SNiagaraScalabilityPreviewSettings.h"
 #include "SNiagaraSystemEffectTypeBar.h"
@@ -109,12 +110,15 @@ void SNiagaraOverviewGraphTitleBar::OnUpdateScalabilityMode(bool bActive)
 
 FText SNiagaraOverviewGraphTitleBar::GetEmitterSubheaderText() const
 {
-	return FText::Format(LOCTEXT("EmitterSubheaderText", "Note: editing this emitter will affect {0} child emitters and systems (across all versions)!"), GetEmitterAffectedAssets());
+	int32 SearchLimit = GetDefault<UNiagaraEditorSettings>()->GetAssetStatsSearchLimit();
+	int32 AffectedAssets = GetEmitterAffectedAssets();
+	FText LimitReachedText = AffectedAssets >= SearchLimit ? LOCTEXT("EmitterSubheaderLimitReachedText", "more than ") : FText();
+	return FText::Format(LOCTEXT("EmitterSubheaderText", "Note: editing this emitter will affect {0}{1} dependent {1}|plural(one=asset,other=assets)! (across all versions)"), LimitReachedText, FText::AsNumber(AffectedAssets));
 }
 
 EVisibility SNiagaraOverviewGraphTitleBar::GetEmitterSubheaderVisibility() const
 {
-	return GetEmitterAffectedAssets() > 1 ? EVisibility::Visible : EVisibility::Collapsed;
+	return GetEmitterAffectedAssets() > 0 ? EVisibility::Visible : EVisibility::Collapsed;
 }
 
 FLinearColor SNiagaraOverviewGraphTitleBar::GetEmitterSubheaderColor() const
@@ -155,8 +159,9 @@ int32 SNiagaraOverviewGraphTitleBar::GetEmitterAffectedAssets() const
 	{
 		return 0;
 	}
-	
-	if (GetDefault<UNiagaraEditorSettings>()->GetDisplayAffectedAssetStats() == false)
+
+	const UNiagaraEditorSettings* EditorSettings = GetDefault<UNiagaraEditorSettings>();
+	if (EditorSettings->GetDisplayAffectedAssetStats() == false)
 	{
 		return 0;
 	}
@@ -171,42 +176,18 @@ int32 SNiagaraOverviewGraphTitleBar::GetEmitterAffectedAssets() const
 			return 0;
 		}
 
-		int32 Count = 0;
-		if (EditedAsset.IsValid())
+		EmitterAffectedAssets = FNiagaraEditorUtilities::GetReferencedAssetCount(EditedAsset, [](const FAssetData& AssetToCheck)
 		{
-			TSet<FName> SeenObjects;
-			TArray<FAssetData> AssetsToCheck;
-			AssetsToCheck.Add(EditedAsset);
-
-			// search for assets that reference this emitter
-			while (AssetsToCheck.Num() > 0)
+			if (AssetToCheck.GetClass() == UNiagaraSystem::StaticClass())
 			{
-				FAssetData AssetToCheck = AssetsToCheck[0];
-				AssetsToCheck.RemoveAtSwap(0);
-				if (SeenObjects.Contains(AssetToCheck.ObjectPath))
-				{
-					continue;
-				}
-				SeenObjects.Add(AssetToCheck.ObjectPath);
-
-				if (AssetToCheck.GetClass() == UNiagaraEmitter::StaticClass())
-				{
-					Count++;
-					TArray<FName> Referencers;
-					AssetRegistry.GetReferencers(AssetToCheck.PackageName, Referencers);
-					for (FName& Referencer : Referencers)
-					{
-						AssetRegistry.GetAssetsByPackageName(Referencer, AssetsToCheck);
-					}
-				}
-				else if (AssetToCheck.GetClass() == UNiagaraSystem::StaticClass())
-				{
-					Count++;
-				}
+				return FNiagaraEditorUtilities::ETrackAssetResult::Count;
 			}
-			Count--; // remove one for our own asset
-		}
-		EmitterAffectedAssets = Count;
+			if (AssetToCheck.GetClass() == UNiagaraEmitter::StaticClass())
+			{
+				return FNiagaraEditorUtilities::ETrackAssetResult::CountRecursive;
+			}
+			return FNiagaraEditorUtilities::ETrackAssetResult::Ignore;
+		});
 	}
 	return EmitterAffectedAssets.Get(0);
 }
