- Add overrides for translucent sort priority / offset

diff --git a/Engine/Plugins/FX/Niagara/Source/Niagara/Classes/NiagaraSystem.h b/Engine/Plugins/FX/Niagara/Source/Niagara/Classes/NiagaraSystem.h
--- a/Engine/Plugins/FX/Niagara/Source/Niagara/Classes/NiagaraSystem.h
+++ b/Engine/Plugins/FX/Niagara/Source/Niagara/Classes/NiagaraSystem.h
@@ -546,6 +546,12 @@ public:
 	UPROPERTY(EditAnywhere, Category="Rendering", meta=(InlineEditConditionToggle="CustomDepthStencilWriteMask"))
 	uint8 bOverrideCustomDepthStencilWriteMask : 1;
 
+	UPROPERTY(EditAnywhere, Category="Rendering", meta=(InlineEditConditionToggle="TranslucencySortPriority"))
+	uint8 bOverrideTranslucencySortPriority : 1;
+
+	UPROPERTY(EditAnywhere, Category="Rendering", meta=(InlineEditConditionToggle="TranslucencySortDistanceOffset"))
+	uint8 bOverrideTranslucencySortDistanceOffset : 1;
+
 	/**
 	When enabled this is the default value set on the component.
 	Controls whether the primitive component should cast a shadow or not.
@@ -584,7 +590,20 @@ public:
 	*/
 	UPROPERTY(EditAnywhere, BlueprintReadOnly, AdvancedDisplay, Category="Rendering", meta=(DisplayName="Default CustomDepthStencil Value", editcondition="bOverrideCustomDepthStencilWriteMask", UIMin = "0", UIMax = "255"))
 	int32 CustomDepthStencilValue = 0;
-	//////////////////////////////////////////////////////////////////////////
+
+	/**
+	When enabled this is the default value set on the component.
+	Adjusts the translucent object sorting priority, see PrimitiveComponent description for more details.
+	*/
+	UPROPERTY(EditAnywhere, BlueprintReadOnly, AdvancedDisplay, Category="Rendering", meta=(editcondition="bOverrideTranslucencySortPriority"))
+	int32 TranslucencySortPriority = 0;
+
+	/**
+	When enabled this is the default value set on the component.
+	Modifies the sort distance for translucent objects, see PrimitiveComponent description for more details.
+	*/
+	UPROPERTY(EditAnywhere, BlueprintReadOnly, AdvancedDisplay, Category="Rendering", meta=(editcondition="bOverrideTranslucencySortDistanceOffset"))
+	float TranslucencySortDistanceOffset = 0.0f;
 
 	/** Computes emitter priorities based on the dependency information. */
 	bool ComputeEmitterPriority(int32 EmitterIdx, TArray<int32, TInlineAllocator<32>>& EmitterPriorities, const TBitArray<TInlineAllocator<32>>& EmitterDependencyGraph);
diff --git a/Engine/Plugins/FX/Niagara/Source/Niagara/Private/NiagaraComponent.cpp b/Engine/Plugins/FX/Niagara/Source/Niagara/Private/NiagaraComponent.cpp
--- a/Engine/Plugins/FX/Niagara/Source/Niagara/Private/NiagaraComponent.cpp
+++ b/Engine/Plugins/FX/Niagara/Source/Niagara/Private/NiagaraComponent.cpp
@@ -3590,16 +3590,19 @@ void UNiagaraComponent::SetAsset(UNiagaraSystem* InAsset, bool bResetExistingOve
 		CustomDepthStencilWriteMask = Asset->bOverrideCustomDepthStencilValue ? Asset->CustomDepthStencilWriteMask : CustomDepthStencilWriteMask;
 		CustomDepthStencilValue = Asset->bOverrideCustomDepthStencilWriteMask ? Asset->CustomDepthStencilValue : CustomDepthStencilValue;
 
+		TranslucencySortPriority = Asset->bOverrideTranslucencySortPriority ? Asset->TranslucencySortPriority : TranslucencySortPriority;
+		TranslucencySortDistanceOffset = Asset->bOverrideTranslucencySortDistanceOffset ? Asset->TranslucencySortDistanceOffset : TranslucencySortDistanceOffset;
+
 		// Activate
 		if ( IsRegistered() )
 		{
-		if (bAutoActivate || bWasActive)
-		{
-			Activate();
+			if (bAutoActivate || bWasActive)
+			{
+				Activate();
+			}
 		}
 	}
 }
-}
 
 void UNiagaraComponent::SetForceSolo(bool bInForceSolo) 
 { 
