Fixed build error by by removing unnecessary export declaration

diff --git a/Engine/Plugins/FX/Niagara/Source/NiagaraEditor/Public/NiagaraEditorUtilities.h b/Engine/Plugins/FX/Niagara/Source/NiagaraEditor/Public/NiagaraEditorUtilities.h
--- a/Engine/Plugins/FX/Niagara/Source/NiagaraEditor/Public/NiagaraEditorUtilities.h
+++ b/Engine/Plugins/FX/Niagara/Source/NiagaraEditor/Public/NiagaraEditorUtilities.h
@@ -346,7 +346,7 @@ namespace FNiagaraEditorUtilities
 	
 	NIAGARAEDITOR_API float GetScalabilityTintAlpha(FNiagaraEmitterHandle* EmitterHandle);
 
-	NIAGARAEDITOR_API enum class ETrackAssetResult
+	enum class ETrackAssetResult
 	{
 		// do not count this asset
 		Ignore,
