Allow modules to specify version dependencies

diff --git a/Engine/Plugins/FX/Niagara/Source/Niagara/Classes/NiagaraScript.h b/Engine/Plugins/FX/Niagara/Source/Niagara/Classes/NiagaraScript.h
--- a/Engine/Plugins/FX/Niagara/Source/Niagara/Classes/NiagaraScript.h
+++ b/Engine/Plugins/FX/Niagara/Source/Niagara/Classes/NiagaraScript.h
@@ -13,12 +13,10 @@
 #include "NiagaraParameters.h"
 #include "NiagaraDataSet.h"
 #include "NiagaraScriptExecutionParameterStore.h"
-#include "NiagaraScriptHighlight.h"
 #include "NiagaraStackSection.h"
 #include "NiagaraParameterDefinitionsSubscriber.h"
 #include "NiagaraVersionedObject.h"
 #include "HAL/CriticalSection.h"
-#include "HAL/PlatformAtomics.h"
 #include "VectorVM.h"
 
 #include "NiagaraScript.generated.h"
@@ -107,6 +105,14 @@ public:
 	/** Specifies constraints related to the source script a modules provides as dependency. */
 	UPROPERTY(AssetRegistrySearchable, EditAnywhere, Category = Script)
 	ENiagaraModuleDependencyScriptConstraint ScriptConstraint;
+
+	// Specifies the version constraint that module providing the dependency must fulfill.
+	// Example usages:
+	// '1.2' requires the exact version 1.2 of the source script
+	// '1.2+' requires at least version 1.2, but any higher version is also ok
+	// '1.2-2.0' requires any version between 1.2 and 2.0
+	UPROPERTY(AssetRegistrySearchable, EditAnywhere, Category = Script)
+	FString RequiredVersion;
 	
 	/** Detailed description of the dependency */
 	UPROPERTY(AssetRegistrySearchable, EditAnywhere, Category = Script, meta = (MultiLine = true))
@@ -117,6 +123,25 @@ public:
 		Type = ENiagaraModuleDependencyType::PreDependency;
 		ScriptConstraint = ENiagaraModuleDependencyScriptConstraint::SameScript;
 	}
+
+#if WITH_EDITOR
+	bool HasValidVersionDependency() const;
+	NIAGARA_API bool IsVersionAllowed(const FNiagaraAssetVersion& Version) const;
+	
+	void CheckVersionCache() const;
+#endif
+private:
+	
+	struct FResolvedVersions
+	{
+		bool bValid = false;
+		FString SourceProperty;
+		int32 MinMajorVersion = 0;
+		int32 MinMinorVersion = 0;
+		int32 MaxMajorVersion = 0;
+		int32 MaxMinorVersion = 0;
+	};
+	mutable FResolvedVersions VersionDependencyCache;
 };
 
 USTRUCT()
diff --git a/Engine/Plugins/FX/Niagara/Source/Niagara/Private/NiagaraScript.cpp b/Engine/Plugins/FX/Niagara/Source/Niagara/Private/NiagaraScript.cpp
--- a/Engine/Plugins/FX/Niagara/Source/Niagara/Private/NiagaraScript.cpp
+++ b/Engine/Plugins/FX/Niagara/Source/Niagara/Private/NiagaraScript.cpp
@@ -30,7 +30,7 @@
 	#include "DerivedDataCacheInterface.h"
 	#include "Interfaces/ITargetPlatform.h"
 	#include "NiagaraSettings.h"
-
+	#include "Internationalization/Regex.h"
 
 	// This is a version string that mimics the old versioning scheme. In case of merge conflicts with DDC versions,
 	// you MUST generate a new GUID and set this new version. If you want to bump this version, generate a new guid
@@ -1752,6 +1752,64 @@ void UNiagaraScript::SaveShaderStableKeys(const class ITargetPlatform* TP)
 		}
 	}
 }
+
+bool FNiagaraModuleDependency::HasValidVersionDependency() const
+{
+	CheckVersionCache();
+	return VersionDependencyCache.bValid;
+}
+
+bool FNiagaraModuleDependency::IsVersionAllowed(const FNiagaraAssetVersion& Version) const
+{
+	if (HasValidVersionDependency())
+	{
+		FNiagaraAssetVersion MinVersion;
+		FNiagaraAssetVersion MaxVersion;
+		MinVersion.MajorVersion = VersionDependencyCache.MinMajorVersion;
+		MinVersion.MinorVersion = VersionDependencyCache.MinMinorVersion;
+		MaxVersion.MajorVersion = VersionDependencyCache.MaxMajorVersion;
+		MaxVersion.MinorVersion = VersionDependencyCache.MaxMinorVersion;
+		return Version <= MaxVersion && MinVersion <= Version;
+	}
+	return true;
+}
+
+void FNiagaraModuleDependency::CheckVersionCache() const
+{
+	if (VersionDependencyCache.SourceProperty != RequiredVersion)
+	{
+		VersionDependencyCache.SourceProperty = RequiredVersion;
+		
+		FString PatternString(TEXT("^\\s*(\\d+)\\.(\\d+)\\s*(?:(\\+)|(-)\\s*(\\d+)\\.(\\d+))?\\s*$"));
+		FRegexPattern Pattern(PatternString);
+		FRegexMatcher Matcher(Pattern, RequiredVersion);
+		if (Matcher.FindNext())
+		{
+			VersionDependencyCache.bValid = true;
+			VersionDependencyCache.MinMajorVersion = FCString::Atoi(*Matcher.GetCaptureGroup(1));
+			VersionDependencyCache.MinMinorVersion = FCString::Atoi(*Matcher.GetCaptureGroup(2));
+			if (Matcher.GetCaptureGroup(3) == "+")
+			{
+				VersionDependencyCache.MaxMajorVersion = MAX_int32;
+				VersionDependencyCache.MaxMinorVersion = MAX_int32;
+			}
+			else if (Matcher.GetCaptureGroup(4) == "-")
+			{
+				VersionDependencyCache.MaxMajorVersion = FCString::Atoi(*Matcher.GetCaptureGroup(5));
+				VersionDependencyCache.MaxMinorVersion = FCString::Atoi(*Matcher.GetCaptureGroup(6));
+			}
+			else
+			{
+				VersionDependencyCache.MaxMajorVersion = VersionDependencyCache.MinMajorVersion;
+				VersionDependencyCache.MaxMinorVersion = VersionDependencyCache.MinMinorVersion;
+			}
+		}
+		else
+		{
+			VersionDependencyCache.bValid = false;
+		}
+	}
+}
 #endif
 
 FNiagaraCompilerTag* FNiagaraCompilerTag::FindTag(TArray< FNiagaraCompilerTag>& InTags, const FNiagaraVariableBase& InSearchVar)
diff --git a/Engine/Plugins/FX/Niagara/Source/Niagara/Public/NiagaraTypes.h b/Engine/Plugins/FX/Niagara/Source/Niagara/Public/NiagaraTypes.h
--- a/Engine/Plugins/FX/Niagara/Source/Niagara/Public/NiagaraTypes.h
+++ b/Engine/Plugins/FX/Niagara/Source/Niagara/Public/NiagaraTypes.h
@@ -200,7 +200,7 @@ struct NIAGARA_API FNiagaraAssetVersion
 	bool operator==(const FNiagaraAssetVersion& Other) const { return VersionGuid == Other.VersionGuid; }
 	bool operator!=(const FNiagaraAssetVersion& Other) const { return !(*this == Other); }
 	bool operator<(const FNiagaraAssetVersion& Other) const { return MajorVersion < Other.MajorVersion || (MajorVersion == Other.MajorVersion && MinorVersion < Other.MinorVersion); }
-	bool operator<=(const FNiagaraAssetVersion& Other) const { return *this < Other || *this == Other; }
+	bool operator<=(const FNiagaraAssetVersion& Other) const { return *this < Other || (MajorVersion == Other.MajorVersion && MinorVersion == Other.MinorVersion); }
 };
 
 struct NIAGARA_API FNiagaraLWCConverter
diff --git a/Engine/Plugins/FX/Niagara/Source/NiagaraEditor/Private/ViewModels/Stack/NiagaraStackModuleItem.cpp b/Engine/Plugins/FX/Niagara/Source/NiagaraEditor/Private/ViewModels/Stack/NiagaraStackModuleItem.cpp
--- a/Engine/Plugins/FX/Niagara/Source/NiagaraEditor/Private/ViewModels/Stack/NiagaraStackModuleItem.cpp
+++ b/Engine/Plugins/FX/Niagara/Source/NiagaraEditor/Private/ViewModels/Stack/NiagaraStackModuleItem.cpp
@@ -36,6 +36,7 @@
 #include "ViewModels/Stack/NiagaraStackModuleItemOutputCollection.h"
 
 // TODO: Remove these
+#include "ViewModels/Stack/NiagaraStackViewModel.h"
 #include "Widgets/Notifications/SNotificationList.h"
 
 #define LOCTEXT_NAMESPACE "NiagaraStackModuleItem"
@@ -300,352 +301,446 @@ TOptional<UNiagaraStackEntry::FDropRequestResponse> UNiagaraStackModuleItem::Dro
 	return TOptional<FDropRequestResponse>();
 }
 
-int32 GetIndexOfLastDependentModuleData(
-	const TArray<FNiagaraStackModuleData>& StackModuleData,
-	int32 StartIndex,
-	int32 EndIndex,
-	int32 NextIndexOffset,
-	ENiagaraScriptUsage DependentScriptUsage,
-	FGuid DependentScriptUsageId,
-	FNiagaraModuleDependency& RequiredDependency)
-{
-	int LastModuleDataIndexRequiringDependency = INDEX_NONE;
-	for (int32 StackModuleDataIndex = StartIndex; StackModuleDataIndex != EndIndex; StackModuleDataIndex += NextIndexOffset)
+namespace NiagaraStackModuleItemIssues {
+	int32 GetIndexOfLastDependentModuleData(
+		const TArray<FNiagaraStackModuleData>& StackModuleData,
+		int32 StartIndex,
+		int32 EndIndex,
+		int32 NextIndexOffset,
+		ENiagaraScriptUsage DependentScriptUsage,
+		FGuid DependentScriptUsageId,
+		FNiagaraModuleDependency& RequiredDependency)
 	{
-		const FNiagaraStackModuleData& CurrentStackModuleData = StackModuleData[StackModuleDataIndex];
-		if (RequiredDependency.ScriptConstraint == ENiagaraModuleDependencyScriptConstraint::SameScript &&
-			(UNiagaraScript::IsEquivalentUsage(CurrentStackModuleData.Usage, DependentScriptUsage) == false || CurrentStackModuleData.UsageId != DependentScriptUsageId))
+		int LastModuleDataIndexRequiringDependency = INDEX_NONE;
+		for (int32 StackModuleDataIndex = StartIndex; StackModuleDataIndex != EndIndex; StackModuleDataIndex += NextIndexOffset)
 		{
-			break;
-		}
+			const FNiagaraStackModuleData& CurrentStackModuleData = StackModuleData[StackModuleDataIndex];
+			if (RequiredDependency.ScriptConstraint == ENiagaraModuleDependencyScriptConstraint::SameScript &&
+				(UNiagaraScript::IsEquivalentUsage(CurrentStackModuleData.Usage, DependentScriptUsage) == false || CurrentStackModuleData.UsageId != DependentScriptUsageId))
+			{
+				break;
+			}
 
-		if (CurrentStackModuleData.ModuleNode->GetScriptData() != nullptr &&
-			CurrentStackModuleData.ModuleNode->GetScriptData()->RequiredDependencies.ContainsByPredicate(
-				[&RequiredDependency](const FNiagaraModuleDependency& CurrentRequiredDependency) { return CurrentRequiredDependency.Id == RequiredDependency.Id; }))
-		{
-			LastModuleDataIndexRequiringDependency = StackModuleDataIndex;
+			if (CurrentStackModuleData.ModuleNode->GetScriptData() != nullptr &&
+				CurrentStackModuleData.ModuleNode->GetScriptData()->RequiredDependencies.ContainsByPredicate(
+					[&RequiredDependency](const FNiagaraModuleDependency& CurrentRequiredDependency) { return CurrentRequiredDependency.Id == RequiredDependency.Id; }))
+			{
+				LastModuleDataIndexRequiringDependency = StackModuleDataIndex;
+			}
 		}
+		return LastModuleDataIndexRequiringDependency;
 	}
-	return LastModuleDataIndexRequiringDependency;
-}
 
-UNiagaraNodeOutput* GetCompatibleTargetOutputNodeFromOrderedScripts(const TArray<ENiagaraScriptUsage>& CompatibleUsages, const TArray<UNiagaraScript*>& OrderedScripts, int32 DependentScriptIndex, int32 LastScriptIndex, int32 NextScriptOffset)
-{
-	for (int32 CurrentScriptIndex = DependentScriptIndex + NextScriptOffset; CurrentScriptIndex != LastScriptIndex; CurrentScriptIndex += NextScriptOffset)
+	UNiagaraNodeOutput* GetCompatibleTargetOutputNodeFromOrderedScripts(const TArray<ENiagaraScriptUsage>& CompatibleUsages, const TArray<UNiagaraScript*>& OrderedScripts, int32 DependentScriptIndex, int32 LastScriptIndex, int32 NextScriptOffset)
 	{
-		UNiagaraScript* CurrentScript = OrderedScripts[CurrentScriptIndex];
-		if (UNiagaraScript::ContainsEquivilentUsage(CompatibleUsages, CurrentScript->GetUsage()))
+		for (int32 CurrentScriptIndex = DependentScriptIndex + NextScriptOffset; CurrentScriptIndex != LastScriptIndex; CurrentScriptIndex += NextScriptOffset)
 		{
-			UNiagaraScriptSource* CurrentScriptSource = CastChecked<UNiagaraScriptSource>(CurrentScript->GetLatestSource());
-			UNiagaraNodeOutput* CurrentOutputNode = CurrentScriptSource->NodeGraph->FindEquivalentOutputNode(CurrentScript->GetUsage(), CurrentScript->GetUsageId());
-			if (CurrentOutputNode != nullptr)
+			UNiagaraScript* CurrentScript = OrderedScripts[CurrentScriptIndex];
+			if (UNiagaraScript::ContainsEquivilentUsage(CompatibleUsages, CurrentScript->GetUsage()))
 			{
-				return CurrentOutputNode;
+				UNiagaraScriptSource* CurrentScriptSource = CastChecked<UNiagaraScriptSource>(CurrentScript->GetLatestSource());
+				UNiagaraNodeOutput* CurrentOutputNode = CurrentScriptSource->NodeGraph->FindEquivalentOutputNode(CurrentScript->GetUsage(), CurrentScript->GetUsageId());
+				if (CurrentOutputNode != nullptr)
+				{
+					return CurrentOutputNode;
+				}
 			}
 		}
+		return nullptr;
 	}
-	return nullptr;
-}
 
-void GetCompatibleOutputNodeAndIndex(TSharedRef<FNiagaraSystemViewModel> DependentSystemViewModel, FGuid EmitterHandleId, UNiagaraScript& DependencyProviderScript, const FNiagaraStackModuleData& LastDependentModuleData, FNiagaraModuleDependency RequiredDependency, UNiagaraNodeOutput*& OutTargetOutputNode, TOptional<int32>& OutTargetIndex)
-{
-	TArray<ENiagaraScriptUsage> CompatibleUsages = UNiagaraScript::GetSupportedUsageContextsForBitmask(DependencyProviderScript.GetLatestScriptData()->ModuleUsageBitmask);
-	if (UNiagaraScript::ContainsEquivilentUsage(CompatibleUsages, LastDependentModuleData.Usage))
-	{
-		// If the dependency provider is compatible with the last dependent usage it can be added directly before or after the last dependent.
-		OutTargetOutputNode = FNiagaraStackGraphUtilities::GetEmitterOutputNodeForStackNode(*LastDependentModuleData.ModuleNode);
-		OutTargetIndex = LastDependentModuleData.Index + (RequiredDependency.Type == ENiagaraModuleDependencyType::PostDependency ? 1 : 0);
-	}
-	else
+	void GetCompatibleOutputNodeAndIndex(TSharedRef<FNiagaraSystemViewModel> DependentSystemViewModel, FGuid EmitterHandleId, UNiagaraScript& DependencyProviderScript, const FNiagaraStackModuleData& LastDependentModuleData, FNiagaraModuleDependency RequiredDependency, UNiagaraNodeOutput*& OutTargetOutputNode, TOptional<int32>& OutTargetIndex)
 	{
-		// Otherwise we need to search for a compatible script to insert the module into.
-		TArray<UNiagaraScript*> OrderedScripts;
-		DependentSystemViewModel->GetOrderedScriptsForEmitterHandleId(EmitterHandleId, OrderedScripts);
-		int32 DependentScriptIndex = OrderedScripts.IndexOfByPredicate([&LastDependentModuleData](UNiagaraScript* Script)
-			{ return UNiagaraScript::IsEquivalentUsage(Script->GetUsage(), LastDependentModuleData.Usage) && Script->GetUsageId() == LastDependentModuleData.UsageId; });
-
-		if (DependentScriptIndex != INDEX_NONE)
+		TArray<ENiagaraScriptUsage> CompatibleUsages = UNiagaraScript::GetSupportedUsageContextsForBitmask(DependencyProviderScript.GetLatestScriptData()->ModuleUsageBitmask);
+		if (UNiagaraScript::ContainsEquivilentUsage(CompatibleUsages, LastDependentModuleData.Usage))
 		{
-			if (RequiredDependency.Type == ENiagaraModuleDependencyType::PreDependency)
+			// If the dependency provider is compatible with the last dependent usage it can be added directly before or after the last dependent.
+			OutTargetOutputNode = FNiagaraStackGraphUtilities::GetEmitterOutputNodeForStackNode(*LastDependentModuleData.ModuleNode);
+			OutTargetIndex = LastDependentModuleData.Index + (RequiredDependency.Type == ENiagaraModuleDependencyType::PostDependency ? 1 : 0);
+		}
+		else
+		{
+			// Otherwise we need to search for a compatible script to insert the module into.
+			TArray<UNiagaraScript*> OrderedScripts;
+			DependentSystemViewModel->GetOrderedScriptsForEmitterHandleId(EmitterHandleId, OrderedScripts);
+			int32 DependentScriptIndex = OrderedScripts.IndexOfByPredicate([&LastDependentModuleData](UNiagaraScript* Script)
+				{ return UNiagaraScript::IsEquivalentUsage(Script->GetUsage(), LastDependentModuleData.Usage) && Script->GetUsageId() == LastDependentModuleData.UsageId; });
+
+			if (DependentScriptIndex != INDEX_NONE)
 			{
-				UNiagaraNodeOutput* CompatibleOutputNode = GetCompatibleTargetOutputNodeFromOrderedScripts(CompatibleUsages, OrderedScripts, DependentScriptIndex, -1, -1);
-				if (CompatibleOutputNode != nullptr)
+				if (RequiredDependency.Type == ENiagaraModuleDependencyType::PreDependency)
 				{
-					OutTargetOutputNode = CompatibleOutputNode;
-					OutTargetIndex = INDEX_NONE;
+					UNiagaraNodeOutput* CompatibleOutputNode = GetCompatibleTargetOutputNodeFromOrderedScripts(CompatibleUsages, OrderedScripts, DependentScriptIndex, -1, -1);
+					if (CompatibleOutputNode != nullptr)
+					{
+						OutTargetOutputNode = CompatibleOutputNode;
+						OutTargetIndex = INDEX_NONE;
+					}
 				}
-			}
-			else if (RequiredDependency.Type == ENiagaraModuleDependencyType::PostDependency)
-			{
-				UNiagaraNodeOutput* CompatibleOutputNode = GetCompatibleTargetOutputNodeFromOrderedScripts(CompatibleUsages, OrderedScripts, DependentScriptIndex, OrderedScripts.Num(), 1);
-				if (CompatibleOutputNode != nullptr)
+				else if (RequiredDependency.Type == ENiagaraModuleDependencyType::PostDependency)
 				{
-					OutTargetOutputNode = CompatibleOutputNode;
-					OutTargetIndex = 0;
+					UNiagaraNodeOutput* CompatibleOutputNode = GetCompatibleTargetOutputNodeFromOrderedScripts(CompatibleUsages, OrderedScripts, DependentScriptIndex, OrderedScripts.Num(), 1);
+					if (CompatibleOutputNode != nullptr)
+					{
+						OutTargetOutputNode = CompatibleOutputNode;
+						OutTargetIndex = 0;
+					}
 				}
 			}
 		}
 	}
-}
 
-void AddModuleToFixDependencyIssue(
-	TWeakPtr<FNiagaraSystemViewModel> DependentSystemViewModelWeak,
-	FGuid DependentEmitterHandleId,
-	TWeakObjectPtr<UNiagaraNodeFunctionCall> DependentModuleNodeWeak,
-	FNiagaraModuleDependency RequiredDependency,
-	FAssetData DependencyProviderModuleAsset)
-{
-	TSharedPtr<FNiagaraSystemViewModel> DependentSystemViewModel = DependentSystemViewModelWeak.Pin();
-	UNiagaraNodeFunctionCall* DependentModuleNode = DependentModuleNodeWeak.Get();
-	if (DependentSystemViewModel.IsValid() == false || DependentModuleNode == nullptr)
+	void AddModuleToFixDependencyIssue(
+		TWeakPtr<FNiagaraSystemViewModel> DependentSystemViewModelWeak,
+		FGuid DependentEmitterHandleId,
+		TWeakObjectPtr<UNiagaraNodeFunctionCall> DependentModuleNodeWeak,
+		FNiagaraModuleDependency RequiredDependency,
+		FAssetData DependencyProviderModuleAsset)
 	{
-		FNiagaraEditorUtilities::WarnWithToastAndLog(LOCTEXT("AddDependencyModuleFailedInvalidSourceData", "Failed to add a dependency module because the fix source data was no longer valid."));
-		return;
-	}
+		TSharedPtr<FNiagaraSystemViewModel> DependentSystemViewModel = DependentSystemViewModelWeak.Pin();
+		UNiagaraNodeFunctionCall* DependentModuleNode = DependentModuleNodeWeak.Get();
+		if (DependentSystemViewModel.IsValid() == false || DependentModuleNode == nullptr)
+		{
+			FNiagaraEditorUtilities::WarnWithToastAndLog(LOCTEXT("AddDependencyModuleFailedInvalidSourceData", "Failed to add a dependency module because the fix source data was no longer valid."));
+			return;
+		}
 
-	UNiagaraNodeOutput* SourceOutputNode = FNiagaraStackGraphUtilities::GetEmitterOutputNodeForStackNode(*DependentModuleNode);
-	if (SourceOutputNode == nullptr)
-	{
-		FNiagaraEditorUtilities::WarnWithToastAndLog(LOCTEXT("AddDependencyModuleFailedInvalidSourceOutputData", "Failed to add a dependency module because the fix source data was no longer valid."));
-		return;
-	}
+		UNiagaraNodeOutput* SourceOutputNode = FNiagaraStackGraphUtilities::GetEmitterOutputNodeForStackNode(*DependentModuleNode);
+		if (SourceOutputNode == nullptr)
+		{
+			FNiagaraEditorUtilities::WarnWithToastAndLog(LOCTEXT("AddDependencyModuleFailedInvalidSourceOutputData", "Failed to add a dependency module because the fix source data was no longer valid."));
+			return;
+		}
 
-	UNiagaraScript* DependencyProviderScript = Cast<UNiagaraScript>(DependencyProviderModuleAsset.GetAsset());
-	if (DependencyProviderScript == nullptr)
-	{
-		FNiagaraEditorUtilities::WarnWithToastAndLog(LOCTEXT("AddDependencyModuleFailedInvalidScriptData", "Failed to add a dependency module because the script asset was not valid."));
-		return;
-	}
+		UNiagaraScript* DependencyProviderScript = Cast<UNiagaraScript>(DependencyProviderModuleAsset.GetAsset());
+		if (DependencyProviderScript == nullptr)
+		{
+			FNiagaraEditorUtilities::WarnWithToastAndLog(LOCTEXT("AddDependencyModuleFailedInvalidScriptData", "Failed to add a dependency module because the script asset was not valid."));
+			return;
+		}
 
-	const TArray<FNiagaraStackModuleData>& StackModuleData = DependentSystemViewModel->GetStackModuleDataByEmitterHandleId(DependentEmitterHandleId);
-	int32 DependentModuleIndex = StackModuleData.IndexOfByPredicate([DependentModuleNode](const FNiagaraStackModuleData& StackModuleDataItem) { return StackModuleDataItem.ModuleNode == DependentModuleNode; });
-	
-	UNiagaraNodeOutput* TargetOutputNode = nullptr;
-	TOptional<int32> TargetIndex;
+		const TArray<FNiagaraStackModuleData>& StackModuleData = DependentSystemViewModel->GetStackModuleDataByEmitterHandleId(DependentEmitterHandleId);
+		int32 DependentModuleIndex = StackModuleData.IndexOfByPredicate([DependentModuleNode](const FNiagaraStackModuleData& StackModuleDataItem) { return StackModuleDataItem.ModuleNode == DependentModuleNode; });
+		
+		UNiagaraNodeOutput* TargetOutputNode = nullptr;
+		TOptional<int32> TargetIndex;
 
-	int32 LastDependentModuleDataIndex = INDEX_NONE;
-	if (RequiredDependency.Type == ENiagaraModuleDependencyType::PreDependency)
-	{
-		LastDependentModuleDataIndex = GetIndexOfLastDependentModuleData(StackModuleData, DependentModuleIndex, -1, -1,
-			SourceOutputNode->GetUsage(), SourceOutputNode->GetUsageId(), RequiredDependency);
-	}
-	else if (RequiredDependency.Type == ENiagaraModuleDependencyType::PostDependency)
-	{
-		LastDependentModuleDataIndex = GetIndexOfLastDependentModuleData(StackModuleData, DependentModuleIndex, StackModuleData.Num(), 1,
-			SourceOutputNode->GetUsage(), SourceOutputNode->GetUsageId(), RequiredDependency);
-	}
+		int32 LastDependentModuleDataIndex = INDEX_NONE;
+		if (RequiredDependency.Type == ENiagaraModuleDependencyType::PreDependency)
+		{
+			LastDependentModuleDataIndex = GetIndexOfLastDependentModuleData(StackModuleData, DependentModuleIndex, -1, -1,
+				SourceOutputNode->GetUsage(), SourceOutputNode->GetUsageId(), RequiredDependency);
+		}
+		else if (RequiredDependency.Type == ENiagaraModuleDependencyType::PostDependency)
+		{
+			LastDependentModuleDataIndex = GetIndexOfLastDependentModuleData(StackModuleData, DependentModuleIndex, StackModuleData.Num(), 1,
+				SourceOutputNode->GetUsage(), SourceOutputNode->GetUsageId(), RequiredDependency);
+		}
 
-	if (LastDependentModuleDataIndex != INDEX_NONE)
-	{
-		GetCompatibleOutputNodeAndIndex(DependentSystemViewModel.ToSharedRef(), DependentEmitterHandleId, *DependencyProviderScript,
-			StackModuleData[LastDependentModuleDataIndex], RequiredDependency, TargetOutputNode, TargetIndex);
-	}
+		if (LastDependentModuleDataIndex != INDEX_NONE)
+		{
+			GetCompatibleOutputNodeAndIndex(DependentSystemViewModel.ToSharedRef(), DependentEmitterHandleId, *DependencyProviderScript,
+				StackModuleData[LastDependentModuleDataIndex], RequiredDependency, TargetOutputNode, TargetIndex);
+		}
 
-	if (TargetOutputNode == nullptr || TargetIndex.IsSet() == false)
-	{
-		FNiagaraEditorUtilities::WarnWithToastAndLog(LOCTEXT("AddDependencyModuleFailedNoValidLocation", "Failed to add a dependency module because an acceptable location could not be found."));
-		return;
-	}
+		if (TargetOutputNode == nullptr || TargetIndex.IsSet() == false)
+		{
+			FNiagaraEditorUtilities::WarnWithToastAndLog(LOCTEXT("AddDependencyModuleFailedNoValidLocation", "Failed to add a dependency module because an acceptable location could not be found."));
+			return;
+		}
 
-	FScopedTransaction ScopedTransaction(LOCTEXT("AddDependencyFixTransaction", "Add a module to fix a dependency"));
-	FNiagaraStackGraphUtilities::FAddScriptModuleToStackArgs AddScriptModuleToStackArgs(DependencyProviderScript, *TargetOutputNode);
-	AddScriptModuleToStackArgs.TargetIndex = TargetIndex.GetValue();
-	AddScriptModuleToStackArgs.bFixupTargetIndex = true;
-	FNiagaraStackGraphUtilities::AddScriptModuleToStack(AddScriptModuleToStackArgs);
-}
+		FScopedTransaction ScopedTransaction(LOCTEXT("AddDependencyFixTransaction", "Add a module to fix a dependency"));
+		FNiagaraStackGraphUtilities::FAddScriptModuleToStackArgs AddScriptModuleToStackArgs(DependencyProviderScript, *TargetOutputNode);
+		AddScriptModuleToStackArgs.TargetIndex = TargetIndex.GetValue();
+		AddScriptModuleToStackArgs.bFixupTargetIndex = true;
+		FNiagaraStackGraphUtilities::AddScriptModuleToStack(AddScriptModuleToStackArgs);
+	}
+
+	void GenerateFixesForAddingDependencyProviders(
+		TSharedRef<FNiagaraSystemViewModel> DependentSystemViewModel,
+		FGuid DependentEmitterHandleId,
+		UNiagaraNodeFunctionCall& DependentModuleNode,
+		ENiagaraScriptUsage DependentUsage,
+		const FNiagaraModuleDependency& RequiredDependency,
+		TArray<UNiagaraStackEntry::FStackIssueFix>& OutFixes)
+	{
+		TOptional<ENiagaraScriptUsage> RequiredUsage = RequiredDependency.ScriptConstraint == ENiagaraModuleDependencyScriptConstraint::SameScript
+			? DependentUsage
+			: TOptional<ENiagaraScriptUsage>();
+		TArray<FAssetData> ModuleAssetsForDependency;
+		FNiagaraStackGraphUtilities::DependencyUtilities::GetModuleScriptAssetsByDependencyProvided(RequiredDependency.Id, RequiredUsage, ModuleAssetsForDependency);
+
+		// Gather duplicate module names so their fixes can be disambiguated.
+		TSet<FName> ModuleNames;
+		TSet<FName> DuplicateModuleNames;
+		for (FAssetData ModuleAsset : ModuleAssetsForDependency)
+		{
+			bool bIsDuplicate;
+			ModuleNames.Add(ModuleAsset.AssetName, &bIsDuplicate);
+			if (bIsDuplicate)
+			{
+				DuplicateModuleNames.Add(ModuleAsset.AssetName);
+			}
+		}
 
-void GenerateFixesForAddingDependencyProviders(
-	TSharedRef<FNiagaraSystemViewModel> DependentSystemViewModel,
-	FGuid DependentEmitterHandleId,
-	UNiagaraNodeFunctionCall& DependentModuleNode,
-	ENiagaraScriptUsage DependentUsage,
-	const FNiagaraModuleDependency& RequiredDependency,
-	TArray<UNiagaraStackEntry::FStackIssueFix>& OutFixes)
-{
-	TOptional<ENiagaraScriptUsage> RequiredUsage = RequiredDependency.ScriptConstraint == ENiagaraModuleDependencyScriptConstraint::SameScript
-		? DependentUsage
-		: TOptional<ENiagaraScriptUsage>();
-	TArray<FAssetData> ModuleAssetsForDependency;
-	FNiagaraStackGraphUtilities::DependencyUtilities::GetModuleScriptAssetsByDependencyProvided(RequiredDependency.Id, RequiredUsage, ModuleAssetsForDependency);
-
-	// Gather duplicate module names so their fixes can be disambiguated.
-	TSet<FName> ModuleNames;
-	TSet<FName> DuplicateModuleNames;
-	for (FAssetData ModuleAsset : ModuleAssetsForDependency)
-	{
-		bool bIsDuplicate;
-		ModuleNames.Add(ModuleAsset.AssetName, &bIsDuplicate);
-		if (bIsDuplicate)
+		for (const FAssetData& ModuleAsset : ModuleAssetsForDependency)
 		{
-			DuplicateModuleNames.Add(ModuleAsset.AssetName);
+			FText DependencyAssetDisplayName = DuplicateModuleNames.Contains(ModuleAsset.AssetName) ? FText::FromName(ModuleAsset.PackageName) : FText::FromName(ModuleAsset.AssetName);
+			FText FixDescription = FText::Format(LOCTEXT("AddDependencyFixDescription", "Add new dependency module {0}"), DependencyAssetDisplayName);
+			OutFixes.Add(UNiagaraStackEntry::FStackIssueFix(FixDescription, UNiagaraStackEntry::FStackIssueFixDelegate::CreateStatic(&AddModuleToFixDependencyIssue, 
+				TWeakPtr<FNiagaraSystemViewModel>(DependentSystemViewModel),
+				DependentEmitterHandleId,
+				TWeakObjectPtr<UNiagaraNodeFunctionCall>(&DependentModuleNode),
+				RequiredDependency,
+				ModuleAsset)));
 		}
 	}
 
-	for (const FAssetData& ModuleAsset : ModuleAssetsForDependency)
+	void MoveModuleToFixDependencyIssue(
+		TWeakObjectPtr<UNiagaraSystem> DependentSystemWeak,
+		FGuid DependentEmitterHandleId,
+		TWeakObjectPtr<UNiagaraScript> DependentScriptWeak,
+		TWeakObjectPtr<UNiagaraNodeFunctionCall> DependentModuleWeak,
+		ENiagaraScriptUsage TargetUsage,
+		FGuid TargetUsageId,
+		int32 TargetMoveIndex)
 	{
-		FText DependencyAssetDisplayName = DuplicateModuleNames.Contains(ModuleAsset.AssetName) ? FText::FromName(ModuleAsset.PackageName) : FText::FromName(ModuleAsset.AssetName);
-		FText FixDescription = FText::Format(LOCTEXT("AddDependencyFixDescription", "Add new dependency module {0}"), DependencyAssetDisplayName);
-		OutFixes.Add(UNiagaraStackEntry::FStackIssueFix(FixDescription, UNiagaraStackEntry::FStackIssueFixDelegate::CreateStatic(&AddModuleToFixDependencyIssue, 
-			TWeakPtr<FNiagaraSystemViewModel>(DependentSystemViewModel),
-			DependentEmitterHandleId,
-			TWeakObjectPtr<UNiagaraNodeFunctionCall>(&DependentModuleNode),
-			RequiredDependency,
-			ModuleAsset)));
-	}
-}
+		UNiagaraSystem* TargetSystem = DependentSystemWeak.Get();
+		if (TargetSystem == nullptr)
+		{
+			FNiagaraEditorUtilities::WarnWithToastAndLog(LOCTEXT("MoveDependentModuleFailedInvalidSystem", "Failed to move a dependent module because the owning system was no longer valid."));
+			return;
+		}
 
-void MoveModuleToFixDependencyIssue(
-	TWeakObjectPtr<UNiagaraSystem> DependentSystemWeak,
-	FGuid DependentEmitterHandleId,
-	TWeakObjectPtr<UNiagaraScript> DependentScriptWeak,
-	TWeakObjectPtr<UNiagaraNodeFunctionCall> DependentModuleWeak,
-	ENiagaraScriptUsage TargetUsage,
-	FGuid TargetUsageId,
-	int32 TargetMoveIndex)
-{
-	UNiagaraSystem* TargetSystem = DependentSystemWeak.Get();
-	if (TargetSystem == nullptr)
-	{
-		FNiagaraEditorUtilities::WarnWithToastAndLog(LOCTEXT("MoveDependentModuleFailedInvalidSystem", "Failed to move a dependent module because the owning system was no longer valid."));
-		return;
-	}
+		UNiagaraScript* SourceScript = DependentScriptWeak.Get();
+		if (SourceScript == nullptr)
+		{
+			FNiagaraEditorUtilities::WarnWithToastAndLog(LOCTEXT("MoveDependentModuleFailedInvalidScript", "Failed to move a dependent module because the owning script was no longer valid."));
+			return;
+		}
 
-	UNiagaraScript* SourceScript = DependentScriptWeak.Get();
-	if (SourceScript == nullptr)
-	{
-		FNiagaraEditorUtilities::WarnWithToastAndLog(LOCTEXT("MoveDependentModuleFailedInvalidScript", "Failed to move a dependent module because the owning script was no longer valid."));
-		return;
+		UNiagaraNodeFunctionCall* ModuleToMove = DependentModuleWeak.Get();
+		if (ModuleToMove == nullptr)
+		{
+			FNiagaraEditorUtilities::WarnWithToastAndLog(LOCTEXT("MoveDependentModuleFailedInvalidModule", "Failed to move a dependent module because the module to move was no longer valid."));
+			return;
+		}
+
+		FScopedTransaction ScopedTransaction(LOCTEXT("MoveDependentFixTransaction", "Move a dependent module to fix a dependency"));
+		UNiagaraNodeFunctionCall* MovedNode;
+		FNiagaraStackGraphUtilities::MoveModule(*SourceScript, *ModuleToMove, *TargetSystem, DependentEmitterHandleId, TargetUsage, TargetUsageId, TargetMoveIndex, false, MovedNode);
 	}
 
-	UNiagaraNodeFunctionCall* ModuleToMove = DependentModuleWeak.Get();
-	if (ModuleToMove == nullptr)
+	void GenerateFixesForReorderingModules(
+		TSharedRef<FNiagaraSystemViewModel> DependentSystemViewModel,
+		FGuid DependentEmitterHandleId,
+		UNiagaraScript& DependentScript,
+		UNiagaraNodeFunctionCall& DependentModuleNode,
+		FNiagaraModuleDependency RequiredDependency,
+		const TArray<FNiagaraStackModuleData>& StackModuleData,
+		const TArray<int32>& WrongOrderDependencyProviderIndices,
+		TArray<UNiagaraStackEntry::FStackIssueFix>& OutFixes)
 	{
-		FNiagaraEditorUtilities::WarnWithToastAndLog(LOCTEXT("MoveDependentModuleFailedInvalidModule", "Failed to move a dependent module because the module to move was no longer valid."));
-		return;
-	}
+		for (int32 WrongOrderDependencyProviderIndex : WrongOrderDependencyProviderIndices)
+		{
+			const FNiagaraStackModuleData& CurrentStackModuleData = StackModuleData[WrongOrderDependencyProviderIndex];
 
-	FScopedTransaction ScopedTransaction(LOCTEXT("MoveDependentFixTransaction", "Move a dependent module to fix a dependency"));
-	UNiagaraNodeFunctionCall* MovedNode;
-	FNiagaraStackGraphUtilities::MoveModule(*SourceScript, *ModuleToMove, *TargetSystem, DependentEmitterHandleId, TargetUsage, TargetUsageId, TargetMoveIndex, false, MovedNode);
-}
+			FText LocationText;
+			int32 TargetIndex = INDEX_NONE;
+			if (RequiredDependency.Type == ENiagaraModuleDependencyType::PreDependency)
+			{
+				LocationText = LOCTEXT("MoveLocationAfter", "After");
+				TargetIndex = CurrentStackModuleData.Index + 1;
+			}
+			else if (RequiredDependency.Type == ENiagaraModuleDependencyType::PostDependency)
+			{
+				LocationText = LOCTEXT("MoveLocationBefore", "Before");
+				TargetIndex = CurrentStackModuleData.Index;
+			}
 
-void GenerateFixesForReorderingModules(
-	TSharedRef<FNiagaraSystemViewModel> DependentSystemViewModel,
-	FGuid DependentEmitterHandleId,
-	UNiagaraScript& DependentScript,
-	UNiagaraNodeFunctionCall& DependentModuleNode,
-	FNiagaraModuleDependency RequiredDependency,
-	const TArray<FNiagaraStackModuleData>& StackModuleData,
-	const TArray<int32>& WrongOrderDependencyProviderIndices,
-	TArray<UNiagaraStackEntry::FStackIssueFix>& OutFixes)
-{
-	for (int32 WrongOrderDependencyProviderIndex : WrongOrderDependencyProviderIndices)
-	{
-		const FNiagaraStackModuleData& CurrentStackModuleData = StackModuleData[WrongOrderDependencyProviderIndex];
+			if (TargetIndex != INDEX_NONE)
+			{
+				FText FixDescription = FText::Format(LOCTEXT("MoveDependentFixDescriptionFormat", "Move module {0} {1} {2}"), FText::FromString(DependentModuleNode.GetFunctionName()), LocationText, FText::FromString(CurrentStackModuleData.ModuleNode->GetFunctionName()));
+				OutFixes.Add(UNiagaraStackEntry::FStackIssueFix(FixDescription, UNiagaraStackEntry::FStackIssueFixDelegate::CreateStatic(&MoveModuleToFixDependencyIssue,
+					TWeakObjectPtr<UNiagaraSystem>(&DependentSystemViewModel->GetSystem()), DependentEmitterHandleId, TWeakObjectPtr<UNiagaraScript>(&DependentScript), TWeakObjectPtr<UNiagaraNodeFunctionCall>(&DependentModuleNode),
+					CurrentStackModuleData.Usage, CurrentStackModuleData.UsageId, TargetIndex)));
+			}
+		}
+	}
 
-		FText LocationText;
-		int32 TargetIndex = INDEX_NONE;
-		if (RequiredDependency.Type == ENiagaraModuleDependencyType::PreDependency)
+	UNiagaraStackModuleItem* FindStackModuleItem(TSharedRef<FNiagaraSystemViewModel> SourceSystemViewModel, const FGuid& SourceEmitterHandleId, const UNiagaraNodeFunctionCall* ModuleNode)
+	{
+		TSharedPtr<FNiagaraEmitterHandleViewModel> EmitterHandleViewModel = SourceSystemViewModel->GetEmitterHandleViewModelById(SourceEmitterHandleId);
+		if (ModuleNode == nullptr || !EmitterHandleViewModel.IsValid())
 		{
-			LocationText = LOCTEXT("MoveLocationAfter", "After");
-			TargetIndex = CurrentStackModuleData.Index + 1;
+			return nullptr;
 		}
-		else if (RequiredDependency.Type == ENiagaraModuleDependencyType::PostDependency)
+		UNiagaraStackViewModel* StackViewModel = EmitterHandleViewModel->GetEmitterStackViewModel();
+		
+		TArray<UNiagaraStackEntry*> EntriesToCheck;
+		if (UNiagaraStackEntry* RootEntry = StackViewModel->GetRootEntry())
 		{
-			LocationText = LOCTEXT("MoveLocationBefore", "Before");
-			TargetIndex = CurrentStackModuleData.Index;
+			RootEntry->GetUnfilteredChildren(EntriesToCheck);
+			while (EntriesToCheck.Num() > 0)
+			{
+				UNiagaraStackEntry* Entry = EntriesToCheck.Pop();
+				if (UNiagaraStackModuleItem* ItemToCheck = Cast<UNiagaraStackModuleItem>(Entry))
+				{
+					if (&ItemToCheck->GetModuleNode() == ModuleNode)
+					{
+						return ItemToCheck;
+					}
+				}
+				Entry->GetUnfilteredChildren(EntriesToCheck);
+			}
 		}
+		return nullptr;
+	}
 
-		if (TargetIndex != INDEX_NONE)
+	UNiagaraStackEntry::FStackIssueFix GetSwitchVersionFix(const FNiagaraAssetVersion& FromVersion, const FNiagaraAssetVersion& ToVersion, UNiagaraNodeFunctionCall* ModuleNode, TSharedRef<FNiagaraSystemViewModel> SourceSystemViewModel, const FGuid& SourceEmitterHandleId)
+	{
+		FText FixDescription = FText::Format(LOCTEXT("SwitchVersionFixDescriptionFormat", "Switch {0} module version from {1}.{2} to {3}.{4}"), FText::FromString(ModuleNode->GetFunctionName()), FText::AsNumber(FromVersion.MajorVersion), FText::AsNumber(FromVersion.MinorVersion), FText::AsNumber(ToVersion.MajorVersion), FText::AsNumber(ToVersion.MinorVersion));
+		TWeakPtr<FNiagaraSystemViewModel> SystemViewModelPtr = SourceSystemViewModel->AsWeak();
+		TWeakObjectPtr<UNiagaraNodeFunctionCall> ModuleNodePtr = ModuleNode;
+		return UNiagaraStackEntry::FStackIssueFix(FixDescription, UNiagaraStackEntry::FStackIssueFixDelegate::CreateLambda([SystemViewModelPtr, SourceEmitterHandleId, ModuleNodePtr, ToVersion]()
 		{
-			FText FixDescription = FText::Format(LOCTEXT("MoveDependentFixDescriptionFormat", "Move module {0} {1} {2}"), FText::FromString(DependentModuleNode.GetFunctionName()), LocationText, FText::FromString(CurrentStackModuleData.ModuleNode->GetFunctionName()));
-			OutFixes.Add(UNiagaraStackEntry::FStackIssueFix(FixDescription, UNiagaraStackEntry::FStackIssueFixDelegate::CreateStatic(&MoveModuleToFixDependencyIssue,
-				TWeakObjectPtr<UNiagaraSystem>(&DependentSystemViewModel->GetSystem()), DependentEmitterHandleId, TWeakObjectPtr<UNiagaraScript>(&DependentScript), TWeakObjectPtr<UNiagaraNodeFunctionCall>(&DependentModuleNode),
-				CurrentStackModuleData.Usage, CurrentStackModuleData.UsageId, TargetIndex)));
-		}
+			if (!SystemViewModelPtr.IsValid() || !ModuleNodePtr.IsValid())
+			{
+				return;
+			}
+			if(UNiagaraStackModuleItem* StackModule = FindStackModuleItem(SystemViewModelPtr.Pin().ToSharedRef(), SourceEmitterHandleId, ModuleNodePtr.Get()))
+			{
+				StackModule->ChangeScriptVersion(ToVersion.VersionGuid);
+			}		
+		}));
 	}
-}
 
-void EnableModuleToFixDependencyIssue(TWeakObjectPtr<UNiagaraNodeFunctionCall> DependencyProviderModuleNodeWeak)
-{
-	UNiagaraNodeFunctionCall* ModuleToEnable = DependencyProviderModuleNodeWeak.Get();
-	if (ModuleToEnable == nullptr)
+	void GenerateFixesForSwitchingModuleVersions(TSharedRef<FNiagaraSystemViewModel> SourceSystemViewModel, const FGuid& SourceEmitterHandleId, const TArray<FNiagaraStackModuleData>& StackModuleData, const FNiagaraModuleDependency& SourceRequiredDependency, const TArray<int32>& WrongVersionDependencyProviderIndices, TArray<UNiagaraStackEntry::FStackIssueFix>& NewIssues)
 	{
-		FNiagaraEditorUtilities::WarnWithToastAndLog(LOCTEXT("EnableDependencyProviderModuleFailedInvalidNode", "Failed to a dependency providing module because it was no longer valid."));
-		return;
+		for (int32 WrongVersionDependencyProviderIndex : WrongVersionDependencyProviderIndices)
+		{
+			const FNiagaraStackModuleData& CurrentStackModuleData = StackModuleData[WrongVersionDependencyProviderIndex];
+			UNiagaraNodeFunctionCall* ModuleNode = CurrentStackModuleData.ModuleNode;
+			
+			if (FindStackModuleItem(SourceSystemViewModel, SourceEmitterHandleId, ModuleNode) == nullptr)
+			{
+				continue;
+			}
+
+			if (UNiagaraScript* NiagaraScript = ModuleNode->FunctionScript)
+			{
+				FNiagaraAssetVersion CurrentVersion = ModuleNode->GetScriptData()->Version;
+			
+				// check the exposed version first
+				FNiagaraAssetVersion ExposedVersion = NiagaraScript->GetExposedVersion();
+				if (SourceRequiredDependency.IsVersionAllowed(ExposedVersion))
+				{
+					NewIssues.Add(GetSwitchVersionFix(CurrentVersion, ExposedVersion, ModuleNode, SourceSystemViewModel, SourceEmitterHandleId));
+					continue;
+				}
+
+				// no match, walk from highest to lowest version
+				TArray<FNiagaraAssetVersion> AvailableVersions = NiagaraScript->GetAllAvailableVersions();
+				for (int32 i = AvailableVersions.Num() - 1; i >= 0; i--)
+				{
+					FNiagaraAssetVersion Version = AvailableVersions[i];
+					if (SourceRequiredDependency.IsVersionAllowed(Version))
+					{
+						NewIssues.Add(GetSwitchVersionFix(CurrentVersion, Version, ModuleNode, SourceSystemViewModel, SourceEmitterHandleId));
+						break;
+					}
+				}
+			}
+		}
 	}
 
-	FScopedTransaction ScopedTransaction(LOCTEXT("EnableModuleFixTransaction", "Enable a dependency providing module to fix a dependency"));
-	FNiagaraStackGraphUtilities::SetModuleIsEnabled(*ModuleToEnable, true);
-}
+	void EnableModuleToFixDependencyIssue(TWeakObjectPtr<UNiagaraNodeFunctionCall> DependencyProviderModuleNodeWeak)
+	{
+		UNiagaraNodeFunctionCall* ModuleToEnable = DependencyProviderModuleNodeWeak.Get();
+		if (ModuleToEnable == nullptr)
+		{
+			FNiagaraEditorUtilities::WarnWithToastAndLog(LOCTEXT("EnableDependencyProviderModuleFailedInvalidNode", "Failed to a dependency providing module because it was no longer valid."));
+			return;
+		}
 
-void GenerateFixesForEnablingModules(
-	const TArray<FNiagaraStackModuleData>& StackModuleData,
-	const TArray<int32>& DisabledDependencyProviderIndices,
-	TArray<UNiagaraStackEntry::FStackIssueFix>& OutFixes)
-{
-	for (int32 DisabledDependencyProviderIndex : DisabledDependencyProviderIndices)
+		FScopedTransaction ScopedTransaction(LOCTEXT("EnableModuleFixTransaction", "Enable a dependency providing module to fix a dependency"));
+		FNiagaraStackGraphUtilities::SetModuleIsEnabled(*ModuleToEnable, true);
+	}
+
+	void GenerateFixesForEnablingModules(
+		const TArray<FNiagaraStackModuleData>& StackModuleData,
+		const TArray<int32>& DisabledDependencyProviderIndices,
+		TArray<UNiagaraStackEntry::FStackIssueFix>& OutFixes)
 	{
-		const FNiagaraStackModuleData& CurrentStackModuleData = StackModuleData[DisabledDependencyProviderIndex];
-		FText FixDescription = FText::Format(LOCTEXT("EnableDependencyProviderFixDescriptionFormat", "Enable module {0} which provides the dependency."),
-			FText::FromString(CurrentStackModuleData.ModuleNode->GetFunctionName()));
-		OutFixes.Add(UNiagaraStackEntry::FStackIssueFix(FixDescription, UNiagaraStackEntry::FStackIssueFixDelegate::CreateStatic(&EnableModuleToFixDependencyIssue,
-			TWeakObjectPtr<UNiagaraNodeFunctionCall>(CurrentStackModuleData.ModuleNode))));
+		for (int32 DisabledDependencyProviderIndex : DisabledDependencyProviderIndices)
+		{
+			const FNiagaraStackModuleData& CurrentStackModuleData = StackModuleData[DisabledDependencyProviderIndex];
+			FText FixDescription = FText::Format(LOCTEXT("EnableDependencyProviderFixDescriptionFormat", "Enable module {0} which provides the dependency."),
+				FText::FromString(CurrentStackModuleData.ModuleNode->GetFunctionName()));
+			OutFixes.Add(UNiagaraStackEntry::FStackIssueFix(FixDescription, UNiagaraStackEntry::FStackIssueFixDelegate::CreateStatic(&EnableModuleToFixDependencyIssue,
+				TWeakObjectPtr<UNiagaraNodeFunctionCall>(CurrentStackModuleData.ModuleNode))));
+		}
 	}
-}
 
-void GenerateDependencyIssues(
-	TSharedRef<FNiagaraSystemViewModel> SourceSystemViewModel,
-	FGuid SourceEmitterHandleId,
-	UNiagaraScript& SourceScript,
-	UNiagaraNodeFunctionCall& SourceModuleNode,
-	FString SourceStackEditorDataKey,
-	UNiagaraNodeOutput& SourceOutputNode,
-	const TArray<FNiagaraStackModuleData>& SourceStackModuleData,
-	TArray<UNiagaraStackEntry::FStackIssue>& NewIssues)
-{
-	if (SourceModuleNode.GetScriptData() == nullptr || SourceModuleNode.GetScriptData()->RequiredDependencies.Num() == 0)
+	bool IsCorrectDependencyVersion(const FNiagaraStackModuleData& StackModuleData, const FNiagaraModuleDependency& SourceModuleRequiredDependency)
 	{
-		return;
+		FVersionedNiagaraScriptData* ScriptData = StackModuleData.ModuleNode->FunctionScript->GetScriptData(StackModuleData.ModuleNode->SelectedScriptVersion);
+		return SourceModuleRequiredDependency.IsVersionAllowed(ScriptData->Version);
 	}
 
-	FVersionedNiagaraScriptData* ScriptData = SourceModuleNode.FunctionScript->GetScriptData(SourceModuleNode.SelectedScriptVersion);
-	int32 ModuleIndex = SourceStackModuleData.IndexOfByPredicate([&SourceModuleNode](const FNiagaraStackModuleData& ModuleData) { return ModuleData.ModuleNode == &SourceModuleNode || ModuleData.ModuleNode->GetFunctionName() == SourceModuleNode.GetFunctionName(); });
-	if (ensureMsgf(ModuleIndex != INDEX_NONE, TEXT("In system %s, module %s (%s) did not exist in the stack module data."),
-		*SourceSystemViewModel->GetSystem().GetPathName(), *SourceModuleNode.GetFunctionName(), *SourceModuleNode.GetName()))
+	void GenerateDependencyIssues(
+		TSharedRef<FNiagaraSystemViewModel> SourceSystemViewModel,
+		FGuid SourceEmitterHandleId,
+		UNiagaraScript& SourceScript,
+		UNiagaraNodeFunctionCall& SourceModuleNode,
+		FString SourceStackEditorDataKey,
+		UNiagaraNodeOutput& SourceOutputNode,
+		const TArray<FNiagaraStackModuleData>& SourceStackModuleData,
+		TArray<UNiagaraStackEntry::FStackIssue>& NewIssues)
 	{
-		for (const FNiagaraModuleDependency& SourceRequiredDependency : SourceModuleNode.GetScriptData()->RequiredDependencies)
+		if (SourceModuleNode.GetScriptData() == nullptr || SourceModuleNode.GetScriptData()->RequiredDependencies.Num() == 0)
 		{
-			TArray<int32> DependencyProviderIndices;
-			for (int32 StackModuleDataIndex = 0; StackModuleDataIndex < SourceStackModuleData.Num(); StackModuleDataIndex++)
-			{
-				if (FNiagaraStackGraphUtilities::DependencyUtilities::DoesStackModuleProvideDependency(SourceStackModuleData[StackModuleDataIndex], SourceRequiredDependency, SourceOutputNode))
-				{
-					DependencyProviderIndices.Add(StackModuleDataIndex);
-				}
-			}
+			return;
+		}
 
-			// Validate that dependency providers are enabled and in the correct direction.
-			bool bDependencyProviderFound = false;
-			TArray<int32> WrongOrderDependencyProviderIndices;
-			TArray<int32> DisabledDependencyProviderIndices;
-			TArray<ENiagaraScriptUsage> SupportedUsages = UNiagaraScript::GetSupportedUsageContextsForBitmask(ScriptData->ModuleUsageBitmask);
-			for (int32 DependencyProviderIndex : DependencyProviderIndices)
+		FVersionedNiagaraScriptData* ScriptData = SourceModuleNode.FunctionScript->GetScriptData(SourceModuleNode.SelectedScriptVersion);
+		int32 ModuleIndex = SourceStackModuleData.IndexOfByPredicate([&SourceModuleNode](const FNiagaraStackModuleData& ModuleData) { return ModuleData.ModuleNode == &SourceModuleNode || ModuleData.ModuleNode->GetFunctionName() == SourceModuleNode.GetFunctionName(); });
+		if (ensureMsgf(ModuleIndex != INDEX_NONE, TEXT("In system %s, module %s (%s) did not exist in the stack module data."),
+			*SourceSystemViewModel->GetSystem().GetPathName(), *SourceModuleNode.GetFunctionName(), *SourceModuleNode.GetName()))
+		{
+			for (const FNiagaraModuleDependency& SourceRequiredDependency : SourceModuleNode.GetScriptData()->RequiredDependencies)
 			{
-				const FNiagaraStackModuleData& DependencyProviderData = SourceStackModuleData[DependencyProviderIndex];
-				bool bCorrectOrder =
-					(SourceRequiredDependency.Type == ENiagaraModuleDependencyType::PreDependency && DependencyProviderIndex < ModuleIndex) ||
-					(SourceRequiredDependency.Type == ENiagaraModuleDependencyType::PostDependency && DependencyProviderIndex > ModuleIndex);
-				bool bEnabled = DependencyProviderData.ModuleNode->GetDesiredEnabledState() == ENodeEnabledState::Enabled;
-				bool bUsageIsSupported = UNiagaraScript::ContainsEquivilentUsage(SupportedUsages, DependencyProviderData.Usage);
-
-				if (bEnabled && bCorrectOrder)
+				TArray<int32> DependencyProviderIndices;
+				for (int32 StackModuleDataIndex = 0; StackModuleDataIndex < SourceStackModuleData.Num(); StackModuleDataIndex++)
 				{
-					bDependencyProviderFound = true;
-					break;
+					if (FNiagaraStackGraphUtilities::DependencyUtilities::DoesStackModuleProvideDependency(SourceStackModuleData[StackModuleDataIndex], SourceRequiredDependency, SourceOutputNode))
+					{
+						DependencyProviderIndices.Add(StackModuleDataIndex);
+					}
 				}
-				else
+
+				// Validate that dependency providers are enabled and in the correct direction.
+				bool bDependencyProviderFound = false;
+				TArray<int32> WrongOrderDependencyProviderIndices;
+				TArray<int32> WrongVersionDependencyProviderIndices;
+				TArray<int32> DisabledDependencyProviderIndices;
+				TArray<ENiagaraScriptUsage> SupportedUsages = UNiagaraScript::GetSupportedUsageContextsForBitmask(ScriptData->ModuleUsageBitmask);
+				for (int32 DependencyProviderIndex : DependencyProviderIndices)
 				{
+					const FNiagaraStackModuleData& DependencyProviderData = SourceStackModuleData[DependencyProviderIndex];
+					bool bCorrectOrder =
+						(SourceRequiredDependency.Type == ENiagaraModuleDependencyType::PreDependency && DependencyProviderIndex < ModuleIndex) ||
+						(SourceRequiredDependency.Type == ENiagaraModuleDependencyType::PostDependency && DependencyProviderIndex > ModuleIndex);
+					bool bCorrectVersion = IsCorrectDependencyVersion(DependencyProviderData, SourceRequiredDependency);
+					bool bEnabled = DependencyProviderData.ModuleNode->GetDesiredEnabledState() == ENodeEnabledState::Enabled;
+					bool bUsageIsSupported = UNiagaraScript::ContainsEquivilentUsage(SupportedUsages, DependencyProviderData.Usage);
+
+					if (bEnabled && bCorrectOrder && bCorrectVersion)
+					{
+						bDependencyProviderFound = true;
+						break;
+					}
+					
 					if (bCorrectOrder == false)
 					{
 						// We can only reorder a module if it supports being moved to the usage of the target module.
@@ -658,37 +753,45 @@ void GenerateDependencyIssues(
 					{
 						DisabledDependencyProviderIndices.Add(DependencyProviderIndex);
 					}
+					else if (bCorrectVersion == false)
+					{
+						WrongVersionDependencyProviderIndices.Add(DependencyProviderIndex);
+					}
 				}
-			}
 
-			if (bDependencyProviderFound == false)
-			{
-				TArray<UNiagaraStackEntry::FStackIssueFix> Fixes;
-				if (WrongOrderDependencyProviderIndices.Num() == 0 && DisabledDependencyProviderIndices.Num() == 0)
+				if (bDependencyProviderFound == false)
 				{
-					// No valid dependency providers found so add fixes for new providers to add.
-					GenerateFixesForAddingDependencyProviders(SourceSystemViewModel, SourceEmitterHandleId, SourceModuleNode, SourceOutputNode.GetUsage(), SourceRequiredDependency, Fixes);
-				}
-				else
-				{
-					if (WrongOrderDependencyProviderIndices.Num() > 0)
+					TArray<UNiagaraStackEntry::FStackIssueFix> Fixes;
+					if (WrongOrderDependencyProviderIndices.Num() == 0 && DisabledDependencyProviderIndices.Num() == 0 && WrongVersionDependencyProviderIndices.Num() == 0)
 					{
-						GenerateFixesForReorderingModules(SourceSystemViewModel, SourceEmitterHandleId, SourceScript, SourceModuleNode, SourceRequiredDependency, SourceStackModuleData, WrongOrderDependencyProviderIndices, Fixes);
+						// No valid dependency providers found so add fixes for new providers to add.
+						GenerateFixesForAddingDependencyProviders(SourceSystemViewModel, SourceEmitterHandleId, SourceModuleNode, SourceOutputNode.GetUsage(), SourceRequiredDependency, Fixes);
 					}
-					if (DisabledDependencyProviderIndices.Num() > 0)
+					else
 					{
-						GenerateFixesForEnablingModules(SourceStackModuleData, DisabledDependencyProviderIndices, Fixes);
+						if (WrongOrderDependencyProviderIndices.Num() > 0)
+						{
+							GenerateFixesForReorderingModules(SourceSystemViewModel, SourceEmitterHandleId, SourceScript, SourceModuleNode, SourceRequiredDependency, SourceStackModuleData, WrongOrderDependencyProviderIndices, Fixes);
+						}
+						if (DisabledDependencyProviderIndices.Num() > 0)
+						{
+							GenerateFixesForEnablingModules(SourceStackModuleData, DisabledDependencyProviderIndices, Fixes);
+						}
+						if (WrongVersionDependencyProviderIndices.Num() > 0)
+						{
+							GenerateFixesForSwitchingModuleVersions(SourceSystemViewModel, SourceEmitterHandleId, SourceStackModuleData, SourceRequiredDependency, WrongVersionDependencyProviderIndices, Fixes);
+						}
 					}
-				}
 
-				FText DependencyTypeString = SourceRequiredDependency.Type == ENiagaraModuleDependencyType::PreDependency ? LOCTEXT("PreDependency", "pre-dependency") : LOCTEXT("PostDependency", "post-dependency");
-				NewIssues.Add(UNiagaraStackEntry::FStackIssue(
-					EStackIssueSeverity::Error,
-					LOCTEXT("DependencyWarning", "The module has unmet dependencies."),
-					FText::Format(LOCTEXT("DependencyWarningLong", "The following {0} is not met: {1}; {2}"), DependencyTypeString, FText::FromName(SourceRequiredDependency.Id), SourceRequiredDependency.Description),
-					FString::Printf(TEXT("%s-dependency-%s"), *SourceStackEditorDataKey, *SourceRequiredDependency.Id.ToString()),
-					true,
-					Fixes));
+					FText DependencyTypeString = SourceRequiredDependency.Type == ENiagaraModuleDependencyType::PreDependency ? LOCTEXT("PreDependency", "pre-dependency") : LOCTEXT("PostDependency", "post-dependency");
+					NewIssues.Add(UNiagaraStackEntry::FStackIssue(
+						EStackIssueSeverity::Error,
+						LOCTEXT("DependencyWarning", "The module has unmet dependencies."),
+						FText::Format(LOCTEXT("DependencyWarningLong", "The following {0} is not met: {1}; {2}"), DependencyTypeString, FText::FromName(SourceRequiredDependency.Id), SourceRequiredDependency.Description),
+						FString::Printf(TEXT("%s-dependency-%s"), *SourceStackEditorDataKey, *SourceRequiredDependency.Id.ToString()),
+						true,
+						Fixes));
+				}
 			}
 		}
 	}
@@ -979,7 +1082,7 @@ void UNiagaraStackModuleItem::RefreshIssues(TArray<FStackIssue>& NewIssues)
 	UNiagaraScript* OwningScript = FNiagaraEditorUtilities::GetScriptFromSystem(GetSystemViewModel()->GetSystem(), EmitterHandleId, OutputNode->GetUsage(), OutputNode->GetUsageId());
 	if (OwningScript != nullptr)
 	{
-		GenerateDependencyIssues(
+		NiagaraStackModuleItemIssues::GenerateDependencyIssues(
 			GetSystemViewModel(),
 			EmitterHandleId,
 			*OwningScript,
