Fix Linux compile
#rnx

diff --git a/Engine/Plugins/FX/Niagara/Source/Niagara/Private/NiagaraDebuggerClient.cpp b/Engine/Plugins/FX/Niagara/Source/Niagara/Private/NiagaraDebuggerClient.cpp
--- a/Engine/Plugins/FX/Niagara/Source/Niagara/Private/NiagaraDebuggerClient.cpp
+++ b/Engine/Plugins/FX/Niagara/Source/Niagara/Private/NiagaraDebuggerClient.cpp
@@ -10,7 +10,7 @@
 #include "NiagaraDebugHud.h"
 #include "Containers/Ticker.h"
 #include "NiagaraSimCache.h"
-#include "Serialization\ObjectAndNameAsStringProxyArchive.h"
+#include "Serialization/ObjectAndNameAsStringProxyArchive.h"
 
 #if WITH_NIAGARA_DEBUGGER
 
diff --git a/Engine/Plugins/FX/Niagara/Source/NiagaraEditor/Private/NiagaraOutliner.cpp b/Engine/Plugins/FX/Niagara/Source/NiagaraEditor/Private/NiagaraOutliner.cpp
--- a/Engine/Plugins/FX/Niagara/Source/NiagaraEditor/Private/NiagaraOutliner.cpp
+++ b/Engine/Plugins/FX/Niagara/Source/NiagaraEditor/Private/NiagaraOutliner.cpp
@@ -4,7 +4,7 @@
 #include "NiagaraDebugger.h"
 #include "NiagaraEditorCommon.h"
 #include "NiagaraSimCache.h"
-#include "Serialization\ObjectAndNameAsStringProxyArchive.h"
+#include "Serialization/ObjectAndNameAsStringProxyArchive.h"
 
 
 UNiagaraOutliner::UNiagaraOutliner(const FObjectInitializer& Initializer)
