- Improve tooltip to be more descriptive

#rnx

diff --git a/Engine/Plugins/FX/Niagara/Source/Niagara/Classes/NiagaraEmitter.h b/Engine/Plugins/FX/Niagara/Source/Niagara/Classes/NiagaraEmitter.h
--- a/Engine/Plugins/FX/Niagara/Source/Niagara/Classes/NiagaraEmitter.h
+++ b/Engine/Plugins/FX/Niagara/Source/Niagara/Classes/NiagaraEmitter.h
@@ -125,7 +125,7 @@ enum class ENiagaraEmitterCalculateBoundMode : uint8
 	Dynamic,
 	/** Bounds are set from the emitter's fixed bounds. */
 	Fixed,
-	/** Bounds will be set from the script or blueprint.  If not set from either source the emitter has no bounds. */
+	/** Bounds will be set from the script using the Emitter Properties Data Interface, or blueprint.  If not set from either source the emitter has no bounds. */
 	Programmable
 };
 
