- Fix for SM6 and debug draw

#rnx
#9066

diff --git a/Engine/Plugins/FX/Niagara/Shaders/Private/NiagaraDebugDraw.ush b/Engine/Plugins/FX/Niagara/Shaders/Private/NiagaraDebugDraw.ush
--- a/Engine/Plugins/FX/Niagara/Shaders/Private/NiagaraDebugDraw.ush
+++ b/Engine/Plugins/FX/Niagara/Shaders/Private/NiagaraDebugDraw.ush
@@ -7,7 +7,7 @@
 	Includes for debug drawing
 =============================================================================*/
 
-void PackLine(inout RWBuffer<uint> OutBuffer, uint InstanceIndex, float3 InPosition0, float3 InPosition1, uint InPackedColor)
+void PackLine(RWBuffer<uint> OutBuffer, uint InstanceIndex, float3 InPosition0, float3 InPosition1, uint InPackedColor)
 {
 	uint BufferOffset = InstanceIndex * 7;
 
@@ -22,7 +22,7 @@ void PackLine(inout RWBuffer<uint> OutBuffer, uint InstanceIndex, float3 InPosit
 	OutBuffer[BufferOffset + 6] = InPackedColor;
 }
 
-void PackLine(inout RWBuffer<uint> OutBuffer, uint InstanceIndex, float3 InPosition0, float3 InPosition1, float4 InColor)
+void PackLine(RWBuffer<uint> OutBuffer, uint InstanceIndex, float3 InPosition0, float3 InPosition1, float4 InColor)
 {
 	uint PackedColor = 0;
 	PackedColor  = uint(saturate(InColor.x) * 255.0f) << 24;
