- Fix some override data interface handling issues
- Move OverrideParameters detection to cover both single & batched mode
- Ensure we flush any pending ticks if we invalidate data interfaces to ensure the RT does not read from stale data

#rnx

diff --git a/Engine/Plugins/FX/Niagara/Source/Niagara/Private/NiagaraComponent.cpp b/Engine/Plugins/FX/Niagara/Source/Niagara/Private/NiagaraComponent.cpp
--- a/Engine/Plugins/FX/Niagara/Source/Niagara/Private/NiagaraComponent.cpp
+++ b/Engine/Plugins/FX/Niagara/Source/Niagara/Private/NiagaraComponent.cpp
@@ -698,12 +698,6 @@ void UNiagaraComponent::TickComponent(float DeltaSeconds, enum ELevelTick TickTy
 		INC_DWORD_STAT_BY(STAT_TotalNiagaraSystemInstances, 1);
 		INC_DWORD_STAT_BY(STAT_TotalNiagaraSystemInstancesSolo, 1);
 
-		// If the interfaces have changed in a meaningful way, we need to potentially rebind and update the values.
-		if (OverrideParameters.GetInterfacesDirty())
-		{
-			SystemInstanceController->Reset(FNiagaraSystemInstance::EResetMode::ReInit);
-		}
-
 		// If we have a sim cache attached then use that
 		if ( SimCache != nullptr )
 		{
diff --git a/Engine/Plugins/FX/Niagara/Source/Niagara/Private/NiagaraSystemGpuComputeProxy.cpp b/Engine/Plugins/FX/Niagara/Source/Niagara/Private/NiagaraSystemGpuComputeProxy.cpp
--- a/Engine/Plugins/FX/Niagara/Source/Niagara/Private/NiagaraSystemGpuComputeProxy.cpp
+++ b/Engine/Plugins/FX/Niagara/Source/Niagara/Private/NiagaraSystemGpuComputeProxy.cpp
@@ -129,6 +129,19 @@ void FNiagaraSystemGpuComputeProxy::RemoveFromRenderThread(FNiagaraGpuComputeDis
 	);
 }
 
+void FNiagaraSystemGpuComputeProxy::ClearTicksFromRenderThread(FNiagaraGpuComputeDispatchInterface* ComputeDispatchInterface)
+{
+	check(IsInGameThread());
+	check(DebugOwnerComputeDispatchInterface == ComputeDispatchInterface);
+
+	ENQUEUE_RENDER_COMMAND(ClearTicksFromProxy)(
+		[=](FRHICommandListImmediate& RHICmdList)
+		{
+			ReleaseTicks(ComputeDispatchInterface->GetGPUInstanceCounterManager());
+		}
+	);
+}
+
 void FNiagaraSystemGpuComputeProxy::QueueTick(const FNiagaraGPUSystemTick& Tick)
 {
 	check(IsInRenderingThread());
diff --git a/Engine/Plugins/FX/Niagara/Source/Niagara/Private/NiagaraSystemInstance.cpp b/Engine/Plugins/FX/Niagara/Source/Niagara/Private/NiagaraSystemInstance.cpp
--- a/Engine/Plugins/FX/Niagara/Source/Niagara/Private/NiagaraSystemInstance.cpp
+++ b/Engine/Plugins/FX/Niagara/Source/Niagara/Private/NiagaraSystemInstance.cpp
@@ -1428,9 +1428,15 @@ void FNiagaraSystemInstance::InitDataInterfaces()
 	// Make sure the owner has flushed it's parameters before we initialize data interfaces
 	InstanceParameters.Tick();
 
-	//-TODO: Validate that any queued ticks have been executed
+	// Destroy data interface data
+	// Note: This invalidates any ticks pending on the render thread
 	DestroyDataInterfaceInstanceData();
 
+	if (SystemGpuComputeProxy)
+	{
+		SystemGpuComputeProxy->ClearTicksFromRenderThread(GetComputeDispatchInterface());
+	}
+
 	PerInstanceDIFunctions[(int32)ENiagaraSystemSimulationScript::Spawn].Reset();
 	PerInstanceDIFunctions[(int32)ENiagaraSystemSimulationScript::Update].Reset();
 
@@ -2347,6 +2353,13 @@ void FNiagaraSystemInstance::Tick_GameThread(float DeltaSeconds)
 		return;
 	}
 
+	// If the interfaces have changed in a meaningful way, we need to potentially rebind and update the values.
+	if (OverrideParameters->GetInterfacesDirty())
+	{
+		Reset(EResetMode::ReInit);
+		return;
+	}
+
 	CachedDeltaSeconds = DeltaSeconds;
 	FixedBounds_CNC = FixedBounds_GT;
 
diff --git a/Engine/Plugins/FX/Niagara/Source/Niagara/Public/NiagaraSystemGpuComputeProxy.h b/Engine/Plugins/FX/Niagara/Source/Niagara/Public/NiagaraSystemGpuComputeProxy.h
--- a/Engine/Plugins/FX/Niagara/Source/Niagara/Public/NiagaraSystemGpuComputeProxy.h
+++ b/Engine/Plugins/FX/Niagara/Source/Niagara/Public/NiagaraSystemGpuComputeProxy.h
@@ -20,6 +20,7 @@ public:
 
 	void AddToRenderThread(FNiagaraGpuComputeDispatchInterface* ComputeDispatchInterface);
 	void RemoveFromRenderThread(FNiagaraGpuComputeDispatchInterface* ComputeDispatchInterface, bool bDeleteProxy);
+	void ClearTicksFromRenderThread(FNiagaraGpuComputeDispatchInterface* ComputeDispatchInterface);
 
 	FNiagaraSystemInstanceID GetSystemInstanceID() const { return SystemInstanceID; }
 	ENiagaraGpuComputeTickStage::Type GetComputeTickStage() const { return ComputeTickStage; }
