- Fix for broken Grid 3D Functions

#rnx

diff --git a/Engine/Plugins/FX/Niagara/Source/Niagara/Private/NiagaraDataInterfaceGrid3DCollection.cpp b/Engine/Plugins/FX/Niagara/Source/Niagara/Private/NiagaraDataInterfaceGrid3DCollection.cpp
--- a/Engine/Plugins/FX/Niagara/Source/Niagara/Private/NiagaraDataInterfaceGrid3DCollection.cpp
+++ b/Engine/Plugins/FX/Niagara/Source/Niagara/Private/NiagaraDataInterfaceGrid3DCollection.cpp
@@ -1270,6 +1270,9 @@ void UNiagaraDataInterfaceGrid3DCollection::WriteGetHLSL(const FNiagaraDataInter
 
 void UNiagaraDataInterfaceGrid3DCollection::WriteSampleHLSL(const FNiagaraDataInterfaceGPUParamInfo& ParamInfo, const FNiagaraDataInterfaceGeneratedFunction& FunctionInfo, int FunctionInstanceIndex, int32 InNumChannels, FString SampleFunction, FString& OutHLSL)
 {
+	SampleFunction.ReplaceInline(TEXT("{Grid}"), *(ParamInfo.DataInterfaceHLSLSymbol + GridName));
+	SampleFunction.ReplaceInline(TEXT("{ParameterName}"), *ParamInfo.DataInterfaceHLSLSymbol);
+
 	FString FormatBounds = TEXT(R"(
 			void {FunctionName}(float3 In_Unit, out float{NumChannelsVariableSuffix} Out_Val)
 			{
@@ -1302,7 +1305,7 @@ void UNiagaraDataInterfaceGrid3DCollection::WriteSampleHLSL(const FNiagaraDataIn
 		{TEXT("FunctionName"),				FunctionInfo.InstanceName},
 		{TEXT("Grid"),						ParamInfo.DataInterfaceHLSLSymbol + GridName},
 		{TEXT("SamplerName"),				ParamInfo.DataInterfaceHLSLSymbol + SamplerName },
-		{TEXT("SampleFunction"),			ParamInfo.DataInterfaceHLSLSymbol + SampleFunction },		
+		{TEXT("SampleFunction"),			SampleFunction },
 		{TEXT("NumCellsName"),				ParamInfo.DataInterfaceHLSLSymbol + UNiagaraDataInterfaceRWBase::NumCellsName},
 		{TEXT("UnitToUVName"),				ParamInfo.DataInterfaceHLSLSymbol + UNiagaraDataInterfaceRWBase::UnitToUVName},
 		{TEXT("NumChannels"),				FString::FromInt(InNumChannels)},
diff --git a/Engine/Source/Runtime/Core/Private/UObject/UE5MainStreamObjectVersion.cpp b/Engine/Source/Runtime/Core/Private/UObject/UE5MainStreamObjectVersion.cpp
--- a/Engine/Source/Runtime/Core/Private/UObject/UE5MainStreamObjectVersion.cpp
+++ b/Engine/Source/Runtime/Core/Private/UObject/UE5MainStreamObjectVersion.cpp
@@ -14,7 +14,7 @@ TMap<FGuid, FGuid> FUE5MainStreamObjectVersion::GetSystemGuids()
 	SystemGuids.Add(DevGuids.MATERIALSHADERMAP_DERIVEDDATA_VER, FGuid("65C01D817C9A4EEFAE9E988D41A1F3DD"));
 	SystemGuids.Add(DevGuids.NANITE_DERIVEDDATA_VER, FGuid("DAC47956CDBF995D96F4520EDB56D875"));
 	SystemGuids.Add(DevGuids.NIAGARASHADERMAP_DERIVEDDATA_VER, FGuid("8A37C45D24F2423CBE5F8F371DE33575"));
-	SystemGuids.Add(DevGuids.Niagara_LatestScriptCompileVersion, FGuid("0BFC42A8539AC143B6640DEBBA491892"));
+	SystemGuids.Add(DevGuids.Niagara_LatestScriptCompileVersion, FGuid("28A43FC8EA00004E8B96B946BCA110F6"));
 	SystemGuids.Add(DevGuids.SkeletalMeshDerivedDataVersion, FGuid("FB7B2FB546DE49B2A2F88B7F69763CEE"));
 	SystemGuids.Add(DevGuids.STATICMESH_DERIVEDDATA_VER, FGuid("98F7D79A5811013825E75BBCC41ED3E9"));
 	SystemGuids.Add(DevGuids.POSESEARCHDB_DERIVEDDATA_VER, FGuid("4E595C2AC5E947D6BA9ABC874353E5BC"));
