Fixed changing emitter scalability overrides not updating active scalability values, by fixing post edit change property not being called on emitters when editing child properties

#rnx

diff --git a/Engine/Plugins/FX/Niagara/Source/NiagaraEditor/Private/NiagaraEmitterDetailsCustomization.cpp b/Engine/Plugins/FX/Niagara/Source/NiagaraEditor/Private/NiagaraEmitterDetailsCustomization.cpp
--- a/Engine/Plugins/FX/Niagara/Source/NiagaraEditor/Private/NiagaraEmitterDetailsCustomization.cpp
+++ b/Engine/Plugins/FX/Niagara/Source/NiagaraEditor/Private/NiagaraEmitterDetailsCustomization.cpp
@@ -78,11 +78,14 @@
 				TSharedPtr<FStructOnScope> StructData = MakeShareable(new FStructOnScope(FVersionedNiagaraEmitterData::StaticStruct(), reinterpret_cast<uint8*>(EmitterData)));
 				if (IDetailPropertyRow* PropertyRow = CategoryBuilder.AddExternalStructureProperty(StructData, ChildProperty->GetFName()))
 				{
-					PropertyRow->GetPropertyHandle()->SetOnPropertyValueChanged(FSimpleDelegate::CreateLambda([VersionedNiagaraEmitter, ChildProperty]()
+					const auto& PostEditChangeLambda = [VersionedNiagaraEmitter, ChildProperty]
 					{
 						FPropertyChangedEvent ChangeEvent(ChildProperty);
 						VersionedNiagaraEmitter.Emitter->PostEditChangeVersionedProperty(ChangeEvent, VersionedNiagaraEmitter.Version);
-					}));
+					};
+
+					PropertyRow->GetPropertyHandle()->SetOnChildPropertyValueChanged(FSimpleDelegate::CreateLambda(PostEditChangeLambda));
+					PropertyRow->GetPropertyHandle()->SetOnPropertyValueChanged(FSimpleDelegate::CreateLambda(PostEditChangeLambda));
 				}
 			}
 		}
