Reset transient system created duirng script merging to avoid later complicaitons

diff --git a/Engine/Plugins/FX/Niagara/Source/Niagara/Private/NiagaraSystem.cpp b/Engine/Plugins/FX/Niagara/Source/Niagara/Private/NiagaraSystem.cpp
--- a/Engine/Plugins/FX/Niagara/Source/Niagara/Private/NiagaraSystem.cpp
+++ b/Engine/Plugins/FX/Niagara/Source/Niagara/Private/NiagaraSystem.cpp
@@ -3215,6 +3215,10 @@ void UNiagaraSystem::InitSystemCompiledData()
 
 void UNiagaraSystem::ResetToEmptySystem()
 {
+	// in order to handle resetting transient systems that can be created (FNiagaraScriptMergeManager::UpdateModuleVersions) we make sure
+	// to mark the system as having been fully loaded to bypass any future work that might be done (generating compile requests)
+	bFullyLoaded = true;
+
 	EffectType = nullptr;
 	EmitterHandles.Empty();
 	ParameterCollectionOverrides.Empty();
diff --git a/Engine/Plugins/FX/Niagara/Source/NiagaraEditor/Private/NiagaraScriptMergeManager.cpp b/Engine/Plugins/FX/Niagara/Source/NiagaraEditor/Private/NiagaraScriptMergeManager.cpp
--- a/Engine/Plugins/FX/Niagara/Source/NiagaraEditor/Private/NiagaraScriptMergeManager.cpp
+++ b/Engine/Plugins/FX/Niagara/Source/NiagaraEditor/Private/NiagaraScriptMergeManager.cpp
@@ -1349,9 +1349,11 @@ void FNiagaraScriptMergeManager::UpdateModuleVersions(const FVersionedNiagaraEmi
 	}
 	
 	// the python update script needs a view model, so we need to create a temporary system here to instantiate the view models
-	TSharedPtr<FNiagaraSystemViewModel> SystemViewModel;
+	UNiagaraSystem* System = nullptr;
 	if (ChangedVersionOtherModules.Num() > 0)
 	{
+		TSharedPtr<FNiagaraSystemViewModel> SystemViewModel;
+
 		if (!FSlateApplication::IsInitialized())
 		{
 			// it's possible that during the cook process we don't have slate initialized, but we need it set up for the view model
@@ -1359,7 +1361,7 @@ void FNiagaraScriptMergeManager::UpdateModuleVersions(const FVersionedNiagaraEmi
 			TSharedRef<FSlateRenderer> SlateRenderer = FModuleManager::Get().LoadModuleChecked<ISlateNullRendererModule>("SlateNullRenderer").CreateSlateNullRenderer();
 			FSlateApplication::Get().InitializeRenderer(SlateRenderer);
 		}
-		UNiagaraSystem* System = NewObject<UNiagaraSystem>(GetTransientPackage(), NAME_None, RF_Transient | RF_Standalone);
+		System = NewObject<UNiagaraSystem>(GetTransientPackage(), NAME_None, RF_Transient | RF_Standalone);
 		UNiagaraSystemFactoryNew::InitializeSystem(System, true);
 		SystemViewModel = MakeShared<FNiagaraSystemViewModel>();
 		FNiagaraSystemViewModelOptions SystemOptions;
@@ -1372,49 +1374,56 @@ void FNiagaraScriptMergeManager::UpdateModuleVersions(const FVersionedNiagaraEmi
 		SystemViewModel->Initialize(*System, SystemOptions);
 		SystemViewModel->GetEditorData().SetOwningSystemIsPlaceholder(true, *System);
 		SystemViewModel->AddEmitter(Instance);
-	}
 
-	// apply version changes
-	for (int i = 0; i < ChangedVersionOtherModules.Num(); i++)
-	{
-		TSharedRef<FNiagaraStackFunctionMergeAdapter> BaseModule = ChangedVersionBaseModules[i];
-		TSharedRef<FNiagaraStackFunctionMergeAdapter> ChangedModule = ChangedVersionOtherModules[i];
-		FGuid NewScriptVersion = BaseModule->GetFunctionCallNode()->SelectedScriptVersion;
-
-		// search for the stack item of the module node we are merging
-		UNiagaraStackModuleItem* ModuleItem = nullptr;
-		UNiagaraStackViewModel* EmitterStackViewModel = SystemViewModel->GetEmitterHandleViewModels()[0]->GetEmitterStackViewModel();
-		TArray<UNiagaraStackEntry*> EntriesToCheck;
-		EmitterStackViewModel->GetRootEntry()->GetUnfilteredChildren(EntriesToCheck);
-		while (EntriesToCheck.Num() > 0)
+		// apply version changes
+		for (int i = 0; i < ChangedVersionOtherModules.Num(); i++)
 		{
-			UNiagaraStackEntry* Entry = EntriesToCheck.Pop();
-			UNiagaraStackModuleItem* ModuleItemToCheck = Cast<UNiagaraStackModuleItem>(Entry);
-			if (ModuleItemToCheck && ModuleItemToCheck->GetModuleNode().NodeGuid == ChangedModule->GetFunctionCallNode()->NodeGuid)
+			TSharedRef<FNiagaraStackFunctionMergeAdapter> BaseModule = ChangedVersionBaseModules[i];
+			TSharedRef<FNiagaraStackFunctionMergeAdapter> ChangedModule = ChangedVersionOtherModules[i];
+			FGuid NewScriptVersion = BaseModule->GetFunctionCallNode()->SelectedScriptVersion;
+
+			// search for the stack item of the module node we are merging
+			UNiagaraStackModuleItem* ModuleItem = nullptr;
+			UNiagaraStackViewModel* EmitterStackViewModel = SystemViewModel->GetEmitterHandleViewModels()[0]->GetEmitterStackViewModel();
+			TArray<UNiagaraStackEntry*> EntriesToCheck;
+			EmitterStackViewModel->GetRootEntry()->GetUnfilteredChildren(EntriesToCheck);
+			while (EntriesToCheck.Num() > 0)
 			{
-				ModuleItem = ModuleItemToCheck;
-				break;
+				UNiagaraStackEntry* Entry = EntriesToCheck.Pop();
+				UNiagaraStackModuleItem* ModuleItemToCheck = Cast<UNiagaraStackModuleItem>(Entry);
+				if (ModuleItemToCheck && ModuleItemToCheck->GetModuleNode().NodeGuid == ChangedModule->GetFunctionCallNode()->NodeGuid)
+				{
+					ModuleItem = ModuleItemToCheck;
+					break;
+				}
+				Entry->GetUnfilteredChildren(EntriesToCheck);
 			}
-			Entry->GetUnfilteredChildren(EntriesToCheck);
-		}
 		
-		FNiagaraScriptVersionUpgradeContext UpgradeContext;
-		UpgradeContext.ConstantResolver = FCompileConstantResolver(Instance, FNiagaraStackGraphUtilities::GetOutputNodeUsage(*ChangedModule->GetFunctionCallNode()));
-		if (ModuleItem)
-		{
-			UpgradeContext.CreateClipboardCallback = [ModuleItem](UNiagaraClipboardContent* ClipboardContent)
+			FNiagaraScriptVersionUpgradeContext UpgradeContext;
+			UpgradeContext.ConstantResolver = FCompileConstantResolver(Instance, FNiagaraStackGraphUtilities::GetOutputNodeUsage(*ChangedModule->GetFunctionCallNode()));
+			if (ModuleItem)
 			{
-				ModuleItem->RefreshChildren();
-				ModuleItem->Copy(ClipboardContent);
-				if (ClipboardContent->Functions.Num() > 0)
+				UpgradeContext.CreateClipboardCallback = [ModuleItem](UNiagaraClipboardContent* ClipboardContent)
 				{
-					ClipboardContent->FunctionInputs = ClipboardContent->Functions[0]->Inputs;
-					ClipboardContent->Functions.Empty();
-				}
-			};
-			UpgradeContext.ApplyClipboardCallback = [ModuleItem](UNiagaraClipboardContent* ClipboardContent, FText& OutWarning) { ModuleItem->Paste(ClipboardContent, OutWarning); };
+					ModuleItem->RefreshChildren();
+					ModuleItem->Copy(ClipboardContent);
+					if (ClipboardContent->Functions.Num() > 0)
+					{
+						ClipboardContent->FunctionInputs = ClipboardContent->Functions[0]->Inputs;
+						ClipboardContent->Functions.Empty();
+					}
+				};
+				UpgradeContext.ApplyClipboardCallback = [ModuleItem](UNiagaraClipboardContent* ClipboardContent, FText& OutWarning) { ModuleItem->Paste(ClipboardContent, OutWarning); };
+			}
+			ChangedModule->GetFunctionCallNode()->ChangeScriptVersion(NewScriptVersion, UpgradeContext, true);
 		}
-		ChangedModule->GetFunctionCallNode()->ChangeScriptVersion(NewScriptVersion, UpgradeContext, true);
+	}
+
+	// now that we're done with our transient System we need to reset it so that if it does get pulled in (via TObjectIterator as an example)
+	// it won't conflict with legitimate systems holding onto the supplied emitter
+	if (System)
+	{
+		System->ResetToEmptySystem();
 	}
 }
 
