- Fix for map size growth over time due to UB hash containing position / bounds

#rnx

diff --git a/Engine/Plugins/FX/Niagara/Source/Niagara/Private/NiagaraComponent.cpp b/Engine/Plugins/FX/Niagara/Source/Niagara/Private/NiagaraComponent.cpp
--- a/Engine/Plugins/FX/Niagara/Source/Niagara/Private/NiagaraComponent.cpp
+++ b/Engine/Plugins/FX/Niagara/Source/Niagara/Private/NiagaraComponent.cpp
@@ -226,6 +226,7 @@ void FNiagaraSceneProxy::ReleaseRenderThreadResources()
 	{
 		CustomUB.Value.ReleaseResource();
 	}
+	CustomUniformBuffers.Empty();
 }
 
 // FPrimitiveSceneProxy interface.
@@ -248,6 +249,10 @@ void FNiagaraSceneProxy::OnTransformChanged()
 	{
 		CustomUB.Value.ReleaseResource();
 	}
+
+	const int32 ExpectedRendererCount = RenderData ? RenderData->GetNumRenderers() : 0;
+	const int32 CurrentUBCount = CustomUniformBuffers.Num();
+	CustomUniformBuffers.Empty(FMath::Min(ExpectedRendererCount, CurrentUBCount));
 }
 
 FPrimitiveViewRelevance FNiagaraSceneProxy::GetViewRelevance(const FSceneView* View) const
diff --git a/Engine/Plugins/FX/Niagara/Source/Niagara/Public/NiagaraSystemRenderData.h b/Engine/Plugins/FX/Niagara/Source/Niagara/Public/NiagaraSystemRenderData.h
--- a/Engine/Plugins/FX/Niagara/Source/Niagara/Public/NiagaraSystemRenderData.h
+++ b/Engine/Plugins/FX/Niagara/Source/Niagara/Public/NiagaraSystemRenderData.h
@@ -80,6 +80,8 @@ public:
 	FORCEINLINE void SetRenderingEnabled(bool bInEnabled) { bRenderingEnabled = bInEnabled; }
 	FORCEINLINE bool HasAnyMotionBlurEnabled() const { return bAnyMotionBlurEnabled; }
 
+	FORCEINLINE int32 GetNumRenderers() const { return RendererDrawOrder.Num(); }
+
 	FVector3f LWCRenderTile = FVector3f::ZeroVector;
 private:
 	/** Emitter Renderers in the order they appear in the emitters. To be accessed by the GameThread */
