Fixed a crash that could happen when upgrading old Niagara systems to UE5

diff --git a/Engine/Plugins/FX/Niagara/Source/Niagara/Private/NiagaraSystem.cpp b/Engine/Plugins/FX/Niagara/Source/Niagara/Private/NiagaraSystem.cpp
--- a/Engine/Plugins/FX/Niagara/Source/Niagara/Private/NiagaraSystem.cpp
+++ b/Engine/Plugins/FX/Niagara/Source/Niagara/Private/NiagaraSystem.cpp
@@ -359,8 +359,10 @@ void UNiagaraSystem::UpdateSystemAfterLoad()
 	bFullyLoaded = true;
 
 #if WITH_EDITORONLY_DATA
-	if (SystemSpawnScript && !GetOutermost()->bIsCookedForEditor)
+	if (SystemSpawnScript && SystemUpdateScript && !GetOutermost()->bIsCookedForEditor)
 	{
+		SystemSpawnScript->ConditionalPostLoad();
+		SystemUpdateScript->ConditionalPostLoad();
 		ensure(SystemSpawnScript->GetLatestSource() != nullptr);
 		ensure(SystemSpawnScript->GetLatestSource() == SystemUpdateScript->GetLatestSource());
 		SystemSpawnScript->GetLatestSource()->OnChanged().AddUObject(this, &UNiagaraSystem::GraphSourceChanged);
