Fixing issue when resolving scalability for local player for FX that are in between deactivation and completion.
Their IsActive() wil return true here leading us to skip the required reactivation.

diff --git a/Engine/Plugins/FX/Niagara/Source/Niagara/Private/NiagaraComponent.cpp b/Engine/Plugins/FX/Niagara/Source/Niagara/Private/NiagaraComponent.cpp
--- a/Engine/Plugins/FX/Niagara/Source/Niagara/Private/NiagaraComponent.cpp
+++ b/Engine/Plugins/FX/Niagara/Source/Niagara/Private/NiagaraComponent.cpp
@@ -2029,14 +2029,12 @@ bool UNiagaraComponent::ResolveOwnerAllowsScalability(bool bRegister)
 		}
 		else
 		{
-			if (!IsActive() && bIsCulledByScalability)
+			if (bIsCulledByScalability)
 			{
 				ActivateInternal(false, true);
 			}
-			else
-			{
-				UnregisterWithScalabilityManager();
-			}
+
+			UnregisterWithScalabilityManager();
 		}
 	}
 	return bOwnerAllowsScalabiltiy;
