Fixing issues with default QL levels / uninitialized device profiles.

diff --git a/Engine/Plugins/FX/Niagara/Source/Niagara/Private/NiagaraPlatformSet.cpp b/Engine/Plugins/FX/Niagara/Source/Niagara/Private/NiagaraPlatformSet.cpp
--- a/Engine/Plugins/FX/Niagara/Source/Niagara/Private/NiagaraPlatformSet.cpp
+++ b/Engine/Plugins/FX/Niagara/Source/Niagara/Private/NiagaraPlatformSet.cpp
@@ -229,9 +229,20 @@ int32 FNiagaraPlatformSet::GetAvailableQualityLevelMask()
 			CachedAvailableQualityLevelMask = 0;
 			int32 Min = GetMinQualityLevel();
 			int32 Max = GetMaxQualityLevel();
-			for (int32 QL = Min; QL <= Max; ++QL)
+
+			if (Min == INDEX_NONE && Max == INDEX_NONE)
+			{
+				CachedAvailableQualityLevelMask = 0xFFFFFFFF;
+			}
+			else
 			{
-				CachedAvailableQualityLevelMask |= (1 << QL);
+				Min = Min == INDEX_NONE ? 0 : Min;
+				Max = Max == INDEX_NONE ? 31 : Max;
+				
+				for (int32 QL = Min; QL <= Max; ++QL)
+				{
+					CachedAvailableQualityLevelMask |= (1 << QL);
+				}
 			}
 		}
 	}
@@ -1023,7 +1034,7 @@ bool FNiagaraPlatformSet::IsEnabledForPlatform(const FString& PlatformName)const
 	{
 		if (const UDeviceProfile* Profile = Cast<const UDeviceProfile>(ProfileObj))
 		{
-			if (CanConsiderDeviceProfile(Profile) && Profile->DeviceType == PlatformName)
+			if (Profile->DeviceType == PlatformName)
 			{
 				if(IsEnabledForDeviceProfile(Profile))
 				{
@@ -1313,14 +1324,19 @@ FNiagaraPlatformSet::FPlatformIniSettings::FPlatformIniSettings(int32 InbCanChan
 	}
 	else
 	{
-		const UNiagaraSettings* Settings = GetDefault<UNiagaraSettings>();
-		check(Settings);
-		QualityLevelMask = 0;
-		int32 NumQualityLevels = Settings->QualityLevels.Num();
-		int32 PlatformQL = 0;
-		for (int32 QL = MinQualityLevel; QL <= MaxQualityLevel; ++QL)
+		if (InMinQualityLevel == INDEX_NONE && MaxQualityLevel == INDEX_NONE)
+		{
+			QualityLevelMask = INDEX_NONE;
+		}
+		else
 		{
-			QualityLevelMask |= (1 << QL);
+			QualityLevelMask = 0;
+			MinQualityLevel = MinQualityLevel == INDEX_NONE ? 0 : MinQualityLevel;
+			MaxQualityLevel = MaxQualityLevel == INDEX_NONE ? 31 : MaxQualityLevel;
+			for (int32 QL = MinQualityLevel; QL <= MaxQualityLevel; ++QL)
+			{
+				QualityLevelMask |= (1 << QL);
+			}
 		}
 	}
 }
