- Fix crash if we have no histories

#rnx

diff --git a/Engine/Plugins/FX/Niagara/Source/NiagaraEditor/Private/ViewModels/NiagaraParameterPanelViewModel.cpp b/Engine/Plugins/FX/Niagara/Source/NiagaraEditor/Private/ViewModels/NiagaraParameterPanelViewModel.cpp
--- a/Engine/Plugins/FX/Niagara/Source/NiagaraEditor/Private/ViewModels/NiagaraParameterPanelViewModel.cpp
+++ b/Engine/Plugins/FX/Niagara/Source/NiagaraEditor/Private/ViewModels/NiagaraParameterPanelViewModel.cpp
@@ -2035,38 +2035,42 @@ TArray<FNiagaraParameterPanelItem> FNiagaraSystemToolkitParameterPanelViewModel:
 					}
 				}
 
-			if (NodeToTraverse == nullptr)
-			{
-				continue;
-			}
+				if (NodeToTraverse == nullptr)
+				{
+					continue;
+				}
 
-			bool bIgnoreDisabled = true;
-			FNiagaraParameterMapHistoryBuilder Builder;
-			FVersionedNiagaraEmitter GraphOwningEmitter = EditableGraph->GetOwningEmitter();
-			FCompileConstantResolver ConstantResolver = GraphOwningEmitter.Emitter != nullptr
-				? FCompileConstantResolver(GraphOwningEmitter, OutputNode->GetUsage())
-				: FCompileConstantResolver();
+				bool bIgnoreDisabled = true;
+				FNiagaraParameterMapHistoryBuilder Builder;
+				FVersionedNiagaraEmitter GraphOwningEmitter = EditableGraph->GetOwningEmitter();
+				FCompileConstantResolver ConstantResolver = GraphOwningEmitter.Emitter != nullptr
+					? FCompileConstantResolver(GraphOwningEmitter, OutputNode->GetUsage())
+					: FCompileConstantResolver();
 				 
-			Builder.SetIgnoreDisabled(bIgnoreDisabled);
-			Builder.ConstantResolver = ConstantResolver;
-			FName StageName;
-			ENiagaraScriptUsage StageUsage = OutputNode->GetUsage();
-			if (StageUsage == ENiagaraScriptUsage::ParticleSimulationStageScript && GraphOwningEmitter.Emitter)
-			{
-				UNiagaraSimulationStageBase* Base = GraphOwningEmitter.GetEmitterData()->GetSimulationStageById(OutputNode->GetUsageId());
-				if (Base)
+				Builder.SetIgnoreDisabled(bIgnoreDisabled);
+				Builder.ConstantResolver = ConstantResolver;
+				FName StageName;
+				ENiagaraScriptUsage StageUsage = OutputNode->GetUsage();
+				if (StageUsage == ENiagaraScriptUsage::ParticleSimulationStageScript && GraphOwningEmitter.Emitter)
 				{
-					StageName = Base->GetStackContextReplacementName();
+					UNiagaraSimulationStageBase* Base = GraphOwningEmitter.GetEmitterData()->GetSimulationStageById(OutputNode->GetUsageId());
+					if (Base)
+					{
+						StageName = Base->GetStackContextReplacementName();
+					}
 				}
-			}
-			Builder.BeginUsage(StageUsage, StageName);
-			NodeToTraverse->BuildParameterMapHistory(Builder, true, true);
-			Builder.EndUsage();
+				Builder.BeginUsage(StageUsage, StageName);
+				NodeToTraverse->BuildParameterMapHistory(Builder, true, true);
+				Builder.EndUsage();
 
 				if (Builder.Histories.Num() != 1)
 				{
 					// We should only have traversed one emitter (have not visited more than one NiagaraNodeEmitter.)
-					ensureMsgf(false, TEXT("Encountered more than one parameter map history when collecting parameters for system parameter panel view model!"));
+					ensureMsgf(false, TEXT("Encountered '%d' parameter map history when collecting parameters for system parameter panel view model, we expect only 1!"), Builder.Histories.Num());
+				}
+				if (Builder.Histories.Num() == 0)
+				{
+					continue;
 				}
 
 				// Get all UNiagaraScriptVariables of visited graphs in the ParameterToScriptVariableMap so that generated items are in sync.
