Fix for persisting BP edits to User parameter array data.

diff --git a/Engine/Plugins/FX/Niagara/Source/Niagara/Classes/NiagaraDataInterfaceArrayImpl.h b/Engine/Plugins/FX/Niagara/Source/Niagara/Classes/NiagaraDataInterfaceArrayImpl.h
--- a/Engine/Plugins/FX/Niagara/Source/Niagara/Classes/NiagaraDataInterfaceArrayImpl.h
+++ b/Engine/Plugins/FX/Niagara/Source/Niagara/Classes/NiagaraDataInterfaceArrayImpl.h
@@ -272,19 +272,32 @@ struct FNDIArrayProxyImpl : public INDIArrayProxyBase
 
 	struct FWriteArrayRef
 	{
-		FWriteArrayRef(TOwnerType* Owner, FNDIArrayInstanceData_GameThread<TArrayType>* InstanceData)
+		FWriteArrayRef(TOwnerType* Owner, FNDIArrayInstanceData_GameThread<TArrayType>* InstanceData, bool bFromBP=false)
 		{
 			if (InstanceData)
 			{
 				LockObject = &InstanceData->ArrayRWGuard;
 				LockObject->WriteLock();
 
-				if (InstanceData->bIsModified == false)
+				if (bFromBP)//Writes from BP are for user parameters and we should discard modifications and write to the actual DI.
 				{
-					InstanceData->bIsModified = true;
-					InstanceData->ArrayData = Owner->GetArrayReference();
-				}
-				ArrayData = &InstanceData->ArrayData;
+					if (InstanceData->bIsModified)
+					{
+						Owner->GetArrayReference() = InstanceData->ArrayData;
+					}
+					InstanceData->bIsModified = false;
+					InstanceData->ArrayData.Empty();
+					ArrayData = &Owner->GetArrayReference();
+ 				}
+ 				else
+ 				{
+ 					if (InstanceData->bIsModified == false)
+ 					{
+ 						InstanceData->bIsModified = true;
+ 						InstanceData->ArrayData = Owner->GetArrayReference();
+ 					}
+					ArrayData = &InstanceData->ArrayData;
+ 				}
 			}
 			else
 			{
@@ -458,9 +471,10 @@ struct FNDIArrayProxyImpl : public INDIArrayProxyBase
 			ensure(PerInstanceData_GameThread.Num() == 1);
 			FNDIArrayInstanceData_GameThread<TArrayType>* InstanceData = PerInstanceData_GameThread.CreateConstIterator().Value();
 			FWriteScopeLock	ScopeLock(InstanceData->ArrayRWGuard);
-			InstanceData->bIsModified = true;
+			InstanceData->bIsModified = false;
 			InstanceData->bIsRenderDirty = bShouldSyncToGpu;
-			InstanceData->ArrayData = InArrayData;
+			InstanceData->ArrayData.Empty();
+			Owner->GetArrayReference() = InArrayData;
 		}
 	}
 
@@ -476,8 +490,7 @@ struct FNDIArrayProxyImpl : public INDIArrayProxyBase
 	{
 		ensure(PerInstanceData_GameThread.Num() <= 1);
 		FNDIArrayInstanceData_GameThread<TArrayType>* InstanceData = PerInstanceData_GameThread.IsEmpty() ? nullptr : PerInstanceData_GameThread.CreateConstIterator().Value();
-		FWriteArrayRef ArrayRef(Owner, InstanceData);
-
+		FWriteArrayRef ArrayRef(Owner, InstanceData, true);
 		if (!ArrayRef.GetArray().IsValidIndex(Index))
 		{
 			if (!bSizeToFit)
