Fixed a bug where Niagara systems without auto-activation would reinitialize after a compile request

diff --git a/Engine/Plugins/FX/Niagara/Source/Niagara/Private/NiagaraCommon.cpp b/Engine/Plugins/FX/Niagara/Source/Niagara/Private/NiagaraCommon.cpp
--- a/Engine/Plugins/FX/Niagara/Source/Niagara/Private/NiagaraCommon.cpp
+++ b/Engine/Plugins/FX/Niagara/Source/Niagara/Private/NiagaraCommon.cpp
@@ -246,7 +246,7 @@ void FNiagaraSystemUpdateContext::AddInternal(UNiagaraComponent* Comp, bool bReI
 		Comp->DeactivateImmediate();
 	}
 
-	if (bIsActive || bOnlyActive == false)
+	if (bIsActive || (bOnlyActive == false && Comp->bAutoActivate))
 	{
 		if (bReInit)
 		{
diff --git a/Engine/Plugins/FX/Niagara/Source/NiagaraEditor/Private/ViewModels/NiagaraSystemViewModel.cpp b/Engine/Plugins/FX/Niagara/Source/NiagaraEditor/Private/ViewModels/NiagaraSystemViewModel.cpp
--- a/Engine/Plugins/FX/Niagara/Source/NiagaraEditor/Private/ViewModels/NiagaraSystemViewModel.cpp
+++ b/Engine/Plugins/FX/Niagara/Source/NiagaraEditor/Private/ViewModels/NiagaraSystemViewModel.cpp
@@ -332,7 +332,6 @@
 		return;
 	}
 	SCOPE_CYCLE_COUNTER(STAT_NiagaraEditor_SystemViewModel_CompileSystem);
-	FNiagaraEditorUtilities::KillSystemInstances(GetSystem());
 	check(SystemScriptViewModel.IsValid());
 	SystemScriptViewModel->CompileSystem(bForce);
 	bCompilePendingCompletion = true;
