Improved audio player:
* no looping one-shot audio
* persistent audio components are destroyed when particles die
* audio modules work in emitter and system scripts now

diff --git a/Engine/Plugins/FX/Niagara/Source/Niagara/Classes/NiagaraDataInterfaceAudioPlayer.h b/Engine/Plugins/FX/Niagara/Source/Niagara/Classes/NiagaraDataInterfaceAudioPlayer.h
--- a/Engine/Plugins/FX/Niagara/Source/Niagara/Classes/NiagaraDataInterfaceAudioPlayer.h
+++ b/Engine/Plugins/FX/Niagara/Source/Niagara/Classes/NiagaraDataInterfaceAudioPlayer.h
@@ -44,12 +44,10 @@ struct FAudioPlayerInterface_InstanceData
 	FNiagaraLWCConverter LWCConverter;
 	int32 MaxPlaysPerTick = 0;
 	bool bStopWhenComponentIsDestroyed = true;
+	bool bValidOneShotSound = false;
 #if WITH_EDITORONLY_DATA
 	bool bOnlyActiveDuringGameplay = false;
 #endif
-	
-	// we track if at least one particle played a sound to prevent problems where sounds keep on playing when scalability culls an emitter (which the DI does not notice otherwise)
-	bool bHadPersistentAudioUpdateThisTick = false;
 };
 
 /** This Data Interface can be used to play one-shot audio effects driven by particle data. */
@@ -93,6 +91,10 @@ public:
 	UPROPERTY(EditAnywhere, AdvancedDisplay, Category = "Audio")
 	bool bOnlyActiveDuringGameplay = false;
 
+	/** Playing looping sounds as persistent audio is not a problem, as the sound is stopped when a particle dies, but one-shot audio outlives the niagara system and can never be stopped. */
+	UPROPERTY(EditAnywhere, AdvancedDisplay, Category = "Audio")
+	bool bAllowLoopingOneShotSounds = false;
+
 	virtual bool UpgradeFunctionCall(FNiagaraFunctionSignature& FunctionSignature) override;
 #endif
 	
diff --git a/Engine/Plugins/FX/Niagara/Source/Niagara/Private/NiagaraDataInterfaceAudioPlayer.cpp b/Engine/Plugins/FX/Niagara/Source/Niagara/Private/NiagaraDataInterfaceAudioPlayer.cpp
--- a/Engine/Plugins/FX/Niagara/Source/Niagara/Private/NiagaraDataInterfaceAudioPlayer.cpp
+++ b/Engine/Plugins/FX/Niagara/Source/Niagara/Private/NiagaraDataInterfaceAudioPlayer.cpp
@@ -172,15 +172,17 @@ bool UNiagaraDataInterfaceAudioPlayer::PerInstanceTick(void* PerInstanceData, FN
 		PIData->SoundToPlay = SoundToPlay;
 		PIData->Attenuation = Attenuation;
 		PIData->Concurrency = Concurrency;
+		PIData->bValidOneShotSound = SoundToPlay->IsLooping() ? bAllowLoopingOneShotSounds : true;
 	}
 	else
 	{
 		PIData->SoundToPlay.Reset();
 		PIData->Attenuation.Reset();
 		PIData->Concurrency.Reset();
+		PIData->bValidOneShotSound = false;
 	}
 
-	PIData->bHadPersistentAudioUpdateThisTick = false;
+	
 	PIData->ParameterNames = ParameterNames;
 	return false;
 }
@@ -219,24 +221,9 @@ bool UNiagaraDataInterfaceAudioPlayer::PerInstanceTickPostSimulate(void* PerInst
 		TGraphTask<FNiagaraAudioPlayerAsyncTask>::CreateTask().ConstructAndDispatchWhenReady(PIData->SoundToPlay, PIData->Attenuation, PIData->Concurrency, Data, World);
 	}
 
-	if (!PIData->bHadPersistentAudioUpdateThisTick)
-	{
-		ensureMsgf(PIData->PersistentAudioActionQueue.IsEmpty(), TEXT("There should be no update tasks without particle updates"));
-
-		for (auto& Entry : PIData->PersistentAudioMapping)
-		{
-			SCOPE_CYCLE_COUNTER(STAT_NiagaraAudioDIStopSound);
-			UAudioComponent* AudioComponent = Entry.Value.IsValid() ? Entry.Value.Get() : nullptr;
-			if (AudioComponent && AudioComponent->IsPlaying())
-			{
-				AudioComponent->Stop();
-			}
-		}
-		PIData->PersistentAudioMapping.Empty();
-	}
-
 	// process the persistent audio updates
 	FPersistentAudioParticleData Value;
+	TSet<int32> UpdatedAudioHandles;
 	while (PIData->PersistentAudioActionQueue.Dequeue(Value))
 	{
 		UAudioComponent* AudioComponent = nullptr;
@@ -248,9 +235,31 @@ bool UNiagaraDataInterfaceAudioPlayer::PerInstanceTickPostSimulate(void* PerInst
 				AudioComponent = MappedValue->Get();
 			}
 		}
+		
 		// since we are in the game thread here, it is safe for the callback to access the audio component
-		Value.UpdateCallback(PIData, AudioComponent, SystemInstance);
+		if (Value.UpdateCallback)
+		{
+			Value.UpdateCallback(PIData, AudioComponent, SystemInstance);
+		}
+		UpdatedAudioHandles.Add(Value.AudioHandle);
+	}
+
+	// destroy all persistent audio components that were not updated this frame - this prevents looping sounds to go on after a particle has died. It also stops sounds if an emitter is culled by scalability.
+	for (auto Iterator = PIData->PersistentAudioMapping.CreateIterator(); Iterator; ++Iterator)
+	{
+		if (!UpdatedAudioHandles.Contains(Iterator.Key()))
+		{
+			SCOPE_CYCLE_COUNTER(STAT_NiagaraAudioDIStopSound);
+			TWeakObjectPtr<UAudioComponent> WeakComponent = Iterator.Value();
+			UAudioComponent* AudioComponent = WeakComponent.IsValid() ? WeakComponent.Get() : nullptr;
+			if (AudioComponent && AudioComponent->IsPlaying())
+			{
+				AudioComponent->Stop();
+			}
+			Iterator.RemoveCurrent();
+		}
 	}
+	
 	return false;
 }
 
@@ -501,7 +510,6 @@ void UNiagaraDataInterfaceAudioPlayer::SetParameterBool(FVectorVMExternalFunctio
 	FNDIInputParam<FNiagaraBool> ValueParam(Context);
 	checkfSlow(InstData.Get(), TEXT("Audio player interface has invalid instance data. %s"), *GetPathName());
 
-	InstData->bHadPersistentAudioUpdateThisTick = true;
 	for (int32 i = 0; i < Context.GetNumInstances(); ++i)
 	{
 		int32 Handle = AudioHandleInParam.GetAndAdvance();
@@ -534,7 +542,6 @@ void UNiagaraDataInterfaceAudioPlayer::SetParameterInteger(FVectorVMExternalFunc
 	FNDIInputParam<int32> ValueParam(Context);
 	checkfSlow(InstData.Get(), TEXT("Audio player interface has invalid instance data. %s"), *GetPathName());
 
-	InstData->bHadPersistentAudioUpdateThisTick = true;
 	for (int32 i = 0; i < Context.GetNumInstances(); ++i)
 	{
 		int32 Handle = AudioHandleInParam.GetAndAdvance();
@@ -567,7 +574,6 @@ void UNiagaraDataInterfaceAudioPlayer::SetParameterFloat(FVectorVMExternalFuncti
 	FNDIInputParam<float> ValueParam(Context);
 	checkfSlow(InstData.Get(), TEXT("Audio player interface has invalid instance data. %s"), *GetPathName());
 
-	InstData->bHadPersistentAudioUpdateThisTick = true;
 	for (int32 i = 0; i < Context.GetNumInstances(); ++i)
 	{
 		int32 Handle = AudioHandleInParam.GetAndAdvance();
@@ -599,7 +605,6 @@ void UNiagaraDataInterfaceAudioPlayer::UpdateVolume(FVectorVMExternalFunctionCon
 	FNDIInputParam<float> VolumeParam(Context);
 	checkfSlow(InstData.Get(), TEXT("Audio player interface has invalid instance data. %s"), *GetPathName());
 
-	InstData->bHadPersistentAudioUpdateThisTick = true;
 	for (int32 i = 0; i < Context.GetNumInstances(); ++i)
 	{
 		int32 Handle = AudioHandleInParam.GetAndAdvance();
@@ -629,7 +634,6 @@ void UNiagaraDataInterfaceAudioPlayer::UpdatePitch(FVectorVMExternalFunctionCont
 	FNDIInputParam<float> PitchParam(Context);
 	checkfSlow(InstData.Get(), TEXT("Audio player interface has invalid instance data. %s"), *GetPathName());
 
-	InstData->bHadPersistentAudioUpdateThisTick = true;
 	for (int32 i = 0; i < Context.GetNumInstances(); ++i)
 	{
 		int32 Handle = AudioHandleInParam.GetAndAdvance();
@@ -659,7 +663,6 @@ void UNiagaraDataInterfaceAudioPlayer::UpdateLocation(FVectorVMExternalFunctionC
 	FNDIInputParam<FVector3f> LocationParam(Context);
 	checkfSlow(InstData.Get(), TEXT("Audio player interface has invalid instance data. %s"), *GetPathName());
 
-	InstData->bHadPersistentAudioUpdateThisTick = true;
 	for (int32 i = 0; i < Context.GetNumInstances(); ++i)
 	{
 		int32 Handle = AudioHandleInParam.GetAndAdvance();
@@ -689,7 +692,6 @@ void UNiagaraDataInterfaceAudioPlayer::UpdateRotation(FVectorVMExternalFunctionC
 	FNDIInputParam<FVector3f> RotationParam(Context);
 	checkfSlow(InstData.Get(), TEXT("Audio player interface has invalid instance data. %s"), *GetPathName());
 
-	InstData->bHadPersistentAudioUpdateThisTick = true;
 	for (int32 i = 0; i < Context.GetNumInstances(); ++i)
 	{
 		int32 Handle = AudioHandleInParam.GetAndAdvance();
@@ -720,7 +722,6 @@ void UNiagaraDataInterfaceAudioPlayer::SetPausedState(FVectorVMExternalFunctionC
 	FNDIInputParam<FNiagaraBool> PausedParam(Context);
 	checkfSlow(InstData.Get(), TEXT("Audio player interface has invalid instance data. %s"), *GetPathName());
 
-	InstData->bHadPersistentAudioUpdateThisTick = true;
 	for (int32 i = 0; i < Context.GetNumInstances(); ++i)
 	{
 		int32 Handle = AudioHandleInParam.GetAndAdvance();
@@ -763,7 +764,14 @@ void UNiagaraDataInterfaceAudioPlayer::PlayOneShotAudio(FVectorVMExternalFunctio
 	VectorVM::FExternalFuncRegisterHandler<FNiagaraBool> OutSample(Context);
 
 	checkfSlow(InstData.Get(), TEXT("Audio player interface has invalid instance data. %s"), *GetPathName());
-	bool ValidSoundData = InstData->SoundToPlay.IsValid();
+	bool ValidSoundData = InstData->SoundToPlay.IsValid() && InstData->bValidOneShotSound;
+
+#if WITH_EDITOR
+	if (InstData->SoundToPlay.IsValid() && !InstData->bValidOneShotSound)
+	{
+		UE_LOG(LogNiagara, Warning, TEXT("Suppressing Niagara one-shot sound %s, because it's a looping sound. Use the 'play persistent audio' module instead. Source: %s"), *InstData->SoundToPlay->GetPathName(), *GetPathNameSafe(this));
+	}
+#endif
 
 	for (int32 i = 0; i < Context.GetNumInstances(); ++i)
 	{
@@ -803,7 +811,6 @@ void UNiagaraDataInterfaceAudioPlayer::PlayPersistentAudio(FVectorVMExternalFunc
 
 	checkfSlow(InstData.Get(), TEXT("Audio player interface has invalid instance data. %s"), *GetPathName());
 
-	InstData->bHadPersistentAudioUpdateThisTick = true;
 	for (int32 i = 0; i < Context.GetNumInstances(); ++i)
 	{
 		bool ShouldPlay = PlayAudioParam.GetAndAdvance();
@@ -841,8 +848,13 @@ void UNiagaraDataInterfaceAudioPlayer::PlayPersistentAudio(FVectorVMExternalFunc
 						InstanceData->PersistentAudioMapping.Add(Handle, AudioComponent);
 					}
 				};
-				InstData->PersistentAudioActionQueue.Enqueue(AudioData);
-			}			
+			}
+			else
+			{
+				// add a dummy entry so the handle keeps playing
+				AudioData.AudioHandle = Handle;
+			}
+			InstData->PersistentAudioActionQueue.Enqueue(AudioData);
 			AudioHandleOutParam.SetAndAdvance(Handle);
 			continue;
 		}
