Fixed a Niagara system warning caused by any_package changes

diff --git a/Engine/Plugins/FX/Niagara/Source/NiagaraEditor/Private/Sequencer/NiagaraSequence/MovieSceneNiagaraEmitterTrack.cpp b/Engine/Plugins/FX/Niagara/Source/NiagaraEditor/Private/Sequencer/NiagaraSequence/MovieSceneNiagaraEmitterTrack.cpp
--- a/Engine/Plugins/FX/Niagara/Source/NiagaraEditor/Private/Sequencer/NiagaraSequence/MovieSceneNiagaraEmitterTrack.cpp
+++ b/Engine/Plugins/FX/Niagara/Source/NiagaraEditor/Private/Sequencer/NiagaraSequence/MovieSceneNiagaraEmitterTrack.cpp
@@ -280,11 +280,9 @@ void UMovieSceneNiagaraEmitterTrack::CreateSections(const FFrameRate& InFrameRes
 		UNiagaraNodeFunctionCall* FunctionNode = Cast<UNiagaraNodeFunctionCall>(StackGroup.EndNode);
 		if (FunctionNode != nullptr && FunctionNode->FunctionScript != nullptr)
 		{
-			FString* SectionClassName = FunctionNode->GetScriptData()->ScriptMetaData.Find("NiagaraTimelineSectionClass");
-			if (SectionClassName != nullptr)
+			if (FString* SectionClassName = FunctionNode->GetScriptData()->ScriptMetaData.Find("NiagaraTimelineSectionClass"))
 			{
-				UClass* SectionClass = UClass::TryFindTypeSlow<UClass>(*SectionClassName);
-				if (SectionClass != nullptr)
+				if (UClass* SectionClass = FindFirstObject<UClass>(**SectionClassName, EFindFirstObjectOptions::NativeFirst | EFindFirstObjectOptions::EnsureIfAmbiguous))
 				{
 					UMovieSceneNiagaraEmitterSectionBase* EmitterSection = nullptr;
 					bool bSectionCreated = false;
