Fixed various scalability mode UI bugs
- Effect type scalability defaults were not using the preview profile & quality level to determine defaults
- Scalability property highlighting for active features was sometimes wrong due to it
- Cleaned up default & resolved values by keeping them empty if there was no actual applicable values instead of using the default struct values
- Active scalability context selection now works via emitter/system properties primarily, and uses node selection as fallback

#rnx

diff --git a/Engine/Plugins/FX/Niagara/Source/Niagara/Classes/NiagaraEffectType.h b/Engine/Plugins/FX/Niagara/Source/Niagara/Classes/NiagaraEffectType.h
--- a/Engine/Plugins/FX/Niagara/Source/Niagara/Classes/NiagaraEffectType.h
+++ b/Engine/Plugins/FX/Niagara/Source/Niagara/Classes/NiagaraEffectType.h
@@ -180,7 +180,7 @@ struct FNiagaraSystemVisibilityCullingSettings
 
 /** Scalability settings for Niagara Systems for a particular platform set (unless overridden). */
 USTRUCT()
-struct FNiagaraSystemScalabilitySettings
+struct NIAGARA_API FNiagaraSystemScalabilitySettings
 {
 	GENERATED_USTRUCT_BODY()
 
@@ -254,7 +254,7 @@ struct FNiagaraSystemScalabilitySettings
 
 /** Container struct for an array of system scalability settings. Enables details customization and data validation. */
 USTRUCT()
-struct FNiagaraSystemScalabilitySettingsArray
+struct NIAGARA_API FNiagaraSystemScalabilitySettingsArray
 {
 	GENERATED_USTRUCT_BODY()
 
@@ -263,7 +263,7 @@ struct FNiagaraSystemScalabilitySettingsArray
 };
 
 USTRUCT()
-struct FNiagaraSystemScalabilityOverride : public FNiagaraSystemScalabilitySettings
+struct NIAGARA_API FNiagaraSystemScalabilityOverride : public FNiagaraSystemScalabilitySettings
 {
 	GENERATED_USTRUCT_BODY()
 
@@ -291,7 +291,7 @@ struct FNiagaraSystemScalabilityOverride : public FNiagaraSystemScalabilitySetti
 
 /** Container struct for an array of system scalability overrides. Enables details customization and data validation. */
 USTRUCT()
-struct FNiagaraSystemScalabilityOverrides
+struct NIAGARA_API FNiagaraSystemScalabilityOverrides
 {
 	GENERATED_USTRUCT_BODY()
 	
@@ -301,7 +301,7 @@ struct FNiagaraSystemScalabilityOverrides
 
 /** Scalability settings for Niagara Emitters on a particular platform set. */
 USTRUCT()
-struct FNiagaraEmitterScalabilitySettings
+struct NIAGARA_API FNiagaraEmitterScalabilitySettings
 {
 	GENERATED_USTRUCT_BODY()
 
@@ -323,7 +323,7 @@ struct FNiagaraEmitterScalabilitySettings
 
 /** Container struct for an array of emitter scalability settings. Enables details customization and data validation. */
 USTRUCT()
-struct FNiagaraEmitterScalabilitySettingsArray
+struct NIAGARA_API FNiagaraEmitterScalabilitySettingsArray
 {
 	GENERATED_USTRUCT_BODY()
 
@@ -332,7 +332,7 @@ struct FNiagaraEmitterScalabilitySettingsArray
 };
 
 USTRUCT()
-struct FNiagaraEmitterScalabilityOverride : public FNiagaraEmitterScalabilitySettings
+struct NIAGARA_API FNiagaraEmitterScalabilityOverride : public FNiagaraEmitterScalabilitySettings
 {
 	GENERATED_USTRUCT_BODY()
 
@@ -345,7 +345,7 @@ struct FNiagaraEmitterScalabilityOverride : public FNiagaraEmitterScalabilitySet
 
 /** Container struct for an array of emitter scalability overrides. Enables details customization and data validation. */
 USTRUCT()
-struct FNiagaraEmitterScalabilityOverrides
+struct NIAGARA_API FNiagaraEmitterScalabilityOverrides
 {
 	GENERATED_USTRUCT_BODY()
 
diff --git a/Engine/Plugins/FX/Niagara/Source/NiagaraEditor/Private/ViewModels/NiagaraSystemScalabilityViewModel.cpp b/Engine/Plugins/FX/Niagara/Source/NiagaraEditor/Private/ViewModels/NiagaraSystemScalabilityViewModel.cpp
--- a/Engine/Plugins/FX/Niagara/Source/NiagaraEditor/Private/ViewModels/NiagaraSystemScalabilityViewModel.cpp
+++ b/Engine/Plugins/FX/Niagara/Source/NiagaraEditor/Private/ViewModels/NiagaraSystemScalabilityViewModel.cpp
@@ -7,6 +7,7 @@
 #include "ISequencerModule.h"
 #include "NiagaraSystemToolkit.h"
 #include "DeviceProfiles/DeviceProfile.h"
+#include "DeviceProfiles/DeviceProfileManager.h"
 #include "Modules/ModuleManager.h"
 #include "ViewModels/NiagaraEmitterHandleViewModel.h"
 
@@ -49,15 +50,23 @@ void UNiagaraSystemScalabilityViewModel::UpdatePreviewDeviceProfile(UDeviceProfi
 	GetSystemViewModel().Pin()->GetSystem().UpdateScalability();
 }
 
-void UNiagaraSystemScalabilityViewModel::UpdatePreviewQualityLevel(int32 QualityLevel, bool bEnabled)
+void UNiagaraSystemScalabilityViewModel::UpdatePreviewQualityLevel(int32 QualityLevel)
 {
 	PreviewPlatforms->QualityLevelMask = 0;
-	PreviewPlatforms->SetEnabledForEffectQuality(QualityLevel, bEnabled);
+	PreviewPlatforms->SetEnabledForEffectQuality(QualityLevel, true);
 
 	FNiagaraPlatformSet::InvalidateCachedData();
 	GetSystemViewModel().Pin()->GetSystem().UpdateScalability();
 }
 
+bool UNiagaraSystemScalabilityViewModel::IsPlatformActive(const FNiagaraPlatformSet& PlatformSet)
+{
+	TOptional<TObjectPtr<UDeviceProfile>> DeviceProfile = GetPreviewDeviceProfile();
+	UDeviceProfileManager& DeviceProfileManager = UDeviceProfileManager::Get();
+	int32 PreviewQualityLevel = FNiagaraPlatformSet::QualityLevelFromMask(PreviewPlatforms->QualityLevelMask);
+	return PlatformSet.IsEnabled(DeviceProfile.Get(DeviceProfileManager.GetActiveProfile()), PreviewQualityLevel, true).bIsActive;
+}
+
 // void UNiagaraSystemScalabilityViewModel::NavigateToScalabilityProperty(UObject* Object, FName PropertyName)
 // {
 // 	TArray<UNiagaraOverviewNode*> OverviewNodes;
diff --git a/Engine/Plugins/FX/Niagara/Source/NiagaraEditor/Private/Widgets/Scalability/SNiagaraSystemResolvedScalabilitySettings.cpp b/Engine/Plugins/FX/Niagara/Source/NiagaraEditor/Private/Widgets/Scalability/SNiagaraSystemResolvedScalabilitySettings.cpp
--- a/Engine/Plugins/FX/Niagara/Source/NiagaraEditor/Private/Widgets/Scalability/SNiagaraSystemResolvedScalabilitySettings.cpp
+++ b/Engine/Plugins/FX/Niagara/Source/NiagaraEditor/Private/Widgets/Scalability/SNiagaraSystemResolvedScalabilitySettings.cpp
@@ -17,6 +17,11 @@
 #include "Widgets/Layout/SExpandableArea.h"
 #include "Widgets/Text/SInlineEditableTextBlock.h"
 #include "PropertyEditorModule.h"
+#include "ViewModels/NiagaraEmitterViewModel.h"
+#include "ViewModels/NiagaraSystemSelectionViewModel.h"
+#include "ViewModels/Stack/NiagaraStackEmitterPropertiesGroup.h"
+#include "ViewModels/Stack/NiagaraStackRendererItem.h"
+#include "ViewModels/Stack/NiagaraStackSystemSettingsGroup.h"
 
 #define LOCTEXT_NAMESPACE "NiagaraScalabilitySystem"
 
@@ -118,6 +123,54 @@ TSharedRef<SWidget> SScalabilityResolvedRow::GenerateWidgetForColumn(const FName
 		return bActive ? FNiagaraEditorStyle::Get().GetColor("NiagaraEditor.Scalability.System.Feature.Active") : FNiagaraEditorStyle::Get().GetColor("NiagaraEditor.Scalability.System.Feature.Inactive");
 	};
 
+	static FNiagaraSystemScalabilityOverride EmptySystemOverride;
+	static FNiagaraEmitterScalabilityOverride EmptyEmitterOverride;
+
+	const FNiagaraSystemScalabilityOverride* CurrentSystemOverride = &EmptySystemOverride;
+	const FNiagaraEmitterScalabilityOverride* CurrentEmitterOverride = &EmptyEmitterOverride;
+	
+	bool bIsResolvedNodeEditable = ScalabilityRowData->ResolvedTreeNode->CreatePropertyHandle()->IsEditable();
+	bool bIsDefaultNodeEditable = false;
+
+	if(ScalabilityRowData->DefaultTreeNode.IsValid())
+	{
+		bIsDefaultNodeEditable = ScalabilityRowData->DefaultTreeNode->CreatePropertyHandle()->IsEditable();
+	}
+	
+	if(UNiagaraSystem* System = Cast<UNiagaraSystem>(ScalabilityRowData->OwningObject))
+	{
+		for(FNiagaraSystemScalabilityOverride& SystemScalabilityOverride : System->GetScalabilityOverrides().Overrides)
+		{
+			if(ParentWidget.Pin()->GetSystemViewModel()->GetScalabilityViewModel()->IsPlatformActive(SystemScalabilityOverride.Platforms))
+			{
+				CurrentSystemOverride = &SystemScalabilityOverride;
+			}
+		}
+	}
+
+	if(UNiagaraEmitter* Emitter = Cast<UNiagaraEmitter>(ScalabilityRowData->OwningObject))
+	{
+		for(FNiagaraEmitterScalabilityOverride& EmitterScalabilityOverride : Emitter->GetEmitterData(ScalabilityRowData->OwningObjectVersion)->ScalabilityOverrides.Overrides)
+		{
+			if(ParentWidget.Pin()->GetSystemViewModel()->GetScalabilityViewModel()->IsPlatformActive(EmitterScalabilityOverride.Platforms))
+			{
+				CurrentEmitterOverride = &EmitterScalabilityOverride;
+			}
+		}
+	}
+
+	// we need to check whether the property is overridden so we can determine whether to check the resolved or default property for a scalability property being active
+
+	bool bIsPropertyOverridden = false;
+	if(ScalabilityRowData->OwningObject->IsA<UNiagaraEmitter>())
+	{
+		bIsPropertyOverridden = IsEmitterScalabilityValueOverridden(ScalabilityRowData->ResolvedTreeNode->GetNodeName(), *CurrentEmitterOverride);
+	}
+	else if(ScalabilityRowData->OwningObject->IsA<UNiagaraSystem>())
+	{
+		bIsPropertyOverridden = IsSystemScalabilityValueOverridden(ScalabilityRowData->ResolvedTreeNode->GetNodeName(), *CurrentSystemOverride);
+	}
+	
 	TSharedPtr<SWidget> RowContent;
 	if(ColumnId == ColumnPropertyName)
 	{
@@ -125,7 +178,7 @@ TSharedRef<SWidget> SScalabilityResolvedRow::GenerateWidgetForColumn(const FName
 		+ SHorizontalBox::Slot()
 		.VAlign(VAlign_Center)
 		[
-			SNew(STextBlock).Text(FText::FromName(ScalabilityRowData->ScalabilityAttribute)).ColorAndOpacity(IconColorActive(ScalabilityRowData->ResolvedTreeNode->CreatePropertyHandle()->IsEditable()))
+			SNew(STextBlock).Text(FText::FromName(ScalabilityRowData->ScalabilityAttribute)).ColorAndOpacity(IconColorActive(bIsPropertyOverridden ? bIsResolvedNodeEditable : bIsDefaultNodeEditable))
 		];
 	}
 	else if(ColumnId == ColumnDefaultValue)
@@ -138,7 +191,7 @@ TSharedRef<SWidget> SScalabilityResolvedRow::GenerateWidgetForColumn(const FName
 			+ SHorizontalBox::Slot()
 			.VAlign(VAlign_Center)
 			[
-				SNew(STextBlock).Text(DefaultValueText).ColorAndOpacity(IconColorActive(ScalabilityRowData->DefaultPropertyHandle->IsEditable()))
+				SNew(STextBlock).Text(DefaultValueText).ColorAndOpacity(IconColorActive(bIsDefaultNodeEditable))
 			];
 		}
 		else
@@ -148,28 +201,8 @@ TSharedRef<SWidget> SScalabilityResolvedRow::GenerateWidgetForColumn(const FName
 	}
 	// for the override value, we only add an entry if an override is active 
 	else if(ColumnId == ColumnResolvedValue)
-	{
-		bool bSkipWidget = false;
-		if(UNiagaraEmitter* Emitter = Cast<UNiagaraEmitter>(ScalabilityRowData->OwningObject))
-		{
-			const FNiagaraEmitterScalabilityOverride& CurrentOverrides = Emitter->GetEmitterData(ScalabilityRowData->OwningObjectVersion)->GetCurrentOverrideSettings();
-			if(!IsEmitterScalabilityValueOverridden(ScalabilityRowData->ResolvedTreeNode->GetNodeName(), CurrentOverrides))
-			{
-				RowContent = SNullWidget::NullWidget;
-				bSkipWidget = true;
-			}				
-		}
-		else if(UNiagaraSystem* System = Cast<UNiagaraSystem>(ScalabilityRowData->OwningObject))
-		{
-			const FNiagaraSystemScalabilityOverride& CurrentOverrides = System->GetCurrentOverrideSettings();
-			if(!IsSystemScalabilityValueOverridden(ScalabilityRowData->ResolvedTreeNode->GetNodeName(), CurrentOverrides))
-			{
-				RowContent = SNullWidget::NullWidget;
-				bSkipWidget = true;
-			}
-		}
-
-		if(!bSkipWidget)
+	{	
+		if(bIsPropertyOverridden)
 		{
 			FText ResolvedValueText;
 			ScalabilityRowData->ResolvedPropertyHandle->GetValueAsDisplayText(ResolvedValueText);
@@ -185,9 +218,13 @@ TSharedRef<SWidget> SScalabilityResolvedRow::GenerateWidgetForColumn(const FName
 			+ SHorizontalBox::Slot()
 			.VAlign(VAlign_Center)
 			[
-				SNew(STextBlock).Text(ResolvedValueText).ColorAndOpacity(IconColorActive(ScalabilityRowData->ResolvedPropertyHandle->IsEditable()))
+				SNew(STextBlock).Text(ResolvedValueText).ColorAndOpacity(IconColorActive(bIsResolvedNodeEditable))
 			];			
 		}
+		else
+		{
+			RowContent = SNullWidget::NullWidget;
+		}
 	}
 	// @Todo this won't be called currently as we removed the column from the header row
 	else if(ColumnId == ColumnLocate)
@@ -215,6 +252,27 @@ FReply SScalabilityResolvedRow::NavigateToValue(TSharedPtr<FScalabilityRowData>
 void SNiagaraSystemResolvedScalabilitySettings::RebuildWidget()
 {
 	ResolvedScalabilityContainer->ClearChildren();
+
+	TArray<UNiagaraStackEntry*> StackEntries;
+	SystemViewModel->GetSelectionViewModel()->GetSelectedEntries(StackEntries);
+
+	UObject* NewSelection = nullptr;
+	FGuid ObjectVersion;
+	if(StackEntries.Num() == 1)
+	{
+		UNiagaraStackEntry* StackEntry = StackEntries[0];
+		
+		if(StackEntry->IsA(UNiagaraStackEmitterPropertiesGroup::StaticClass()))
+		{
+			UNiagaraEmitter* Emitter = Cast<UNiagaraStackEmitterPropertiesGroup>(StackEntry)->GetEmitterViewModel()->GetEmitter().Emitter;
+			ObjectVersion = Cast<UNiagaraStackEmitterPropertiesGroup>(StackEntry)->GetEmitterViewModel()->GetEmitter().Version;
+			NewSelection = Emitter;
+		}
+		else if(StackEntry->IsA(UNiagaraStackSystemPropertiesGroup::StaticClass()))
+		{
+			NewSelection = System.Get();
+		}
+	}
 	
 	TSet<UNiagaraOverviewNode*> SelectedOverviewNodes;
 
@@ -225,8 +283,32 @@ void SNiagaraSystemResolvedScalabilitySettings::RebuildWidget()
 			SelectedOverviewNodes.Add(OverviewNode);
 		}
 	}
-	
-	if(SelectedOverviewNodes.Num() > 0)
+
+	// we prioritize stack selection
+	if(NewSelection != nullptr)
+	{
+		if(UNiagaraEmitter* Emitter = Cast<UNiagaraEmitter>(NewSelection))
+		{
+			Emitter->OnPropertiesChanged().RemoveAll(this);
+			Emitter->OnPropertiesChanged().AddSP(this, &SNiagaraSystemResolvedScalabilitySettings::RebuildWidget);
+
+			ResolvedScalabilityContainer->AddSlot()
+			.AutoHeight()
+			[
+				GenerateResolvedScalabilityTable(Emitter, ObjectVersion)
+			];
+		}
+		if(UNiagaraSystem* NiagaraSystem = Cast<UNiagaraSystem>(NewSelection))
+		{
+			ResolvedScalabilityContainer->AddSlot()
+			.AutoHeight()
+			[
+				GenerateResolvedScalabilityTable(NiagaraSystem, FGuid())
+			];
+		}
+	}
+	// we use node selection if we don't have any stack selection
+	else if(SelectedOverviewNodes.Num() > 0)
 	{
 		for(UObject* SelectedNode : SelectedOverviewNodes)
 		{
@@ -301,10 +383,23 @@ TSharedRef<SWidget> SNiagaraSystemResolvedScalabilitySettings::GenerateResolvedS
 
 	if(UNiagaraEmitter* Emitter = Cast<UNiagaraEmitter>(Object))
 	{
-		if(System->GetEffectType())
+		if(UNiagaraEffectType* EffectType = System->GetEffectType())
 		{
-			const FNiagaraEmitterScalabilitySettings& DefaultEmitterScalabilitySettings = System->GetEffectType()->GetActiveEmitterScalabilitySettings();
-			DefaultValueStruct = MakeShared<FStructOnScope>(FNiagaraEmitterScalabilitySettings::StaticStruct(), (uint8*) &DefaultEmitterScalabilitySettings);			
+			const FNiagaraEmitterScalabilitySettings* DefaultEmitterScalabilitySettings = nullptr;
+
+			for(const FNiagaraEmitterScalabilitySettings& EmitterScalabilitySettings : EffectType->GetEmitterScalabilitySettings().Settings)
+			{
+				// for the effect type's settings we haven't bound the platform QL & profile overrides, so we use a customized function to check if a platformset is active under scalability preview
+				if(GetSystemViewModel()->GetScalabilityViewModel()->IsPlatformActive(EmitterScalabilitySettings.Platforms))
+				{
+					DefaultEmitterScalabilitySettings = &EmitterScalabilitySettings;
+				}
+			}
+
+			if(DefaultEmitterScalabilitySettings != nullptr)
+			{
+				DefaultValueStruct = MakeShared<FStructOnScope>(FNiagaraEmitterScalabilitySettings::StaticStruct(), (uint8*) DefaultEmitterScalabilitySettings);			
+			}
 		}
 		const FNiagaraEmitterScalabilitySettings& EmitterScalabilitySettings = Emitter->GetEmitterData(Version)->GetScalabilitySettings();
 		ResolvedValueStruct = MakeShared<FStructOnScope>(FNiagaraEmitterScalabilitySettings::StaticStruct(), (uint8*) &EmitterScalabilitySettings);
@@ -313,8 +408,20 @@ TSharedRef<SWidget> SNiagaraSystemResolvedScalabilitySettings::GenerateResolvedS
 	{
 		if(UNiagaraEffectType* EffectType = NiagaraSystem->GetEffectType())
 		{
-			const FNiagaraSystemScalabilitySettings& DefaultSystemScalabilitySettings = EffectType->GetActiveSystemScalabilitySettings();
-			DefaultValueStruct = MakeShared<FStructOnScope>(FNiagaraSystemScalabilitySettings::StaticStruct(), (uint8*) &DefaultSystemScalabilitySettings);			
+			const FNiagaraSystemScalabilitySettings* DefaultSystemScalabilitySettings = nullptr;
+			for(const FNiagaraSystemScalabilitySettings& SystemScalabilitySettings : EffectType->GetSystemScalabilitySettings().Settings)
+			{
+				// see above
+				if(GetSystemViewModel()->GetScalabilityViewModel()->IsPlatformActive(SystemScalabilitySettings.Platforms))
+				{
+					DefaultSystemScalabilitySettings = &SystemScalabilitySettings;
+				}
+			}
+
+			if(DefaultSystemScalabilitySettings != nullptr)
+			{
+				DefaultValueStruct = MakeShared<FStructOnScope>(FNiagaraSystemScalabilitySettings::StaticStruct(), (uint8*) DefaultSystemScalabilitySettings);			
+			}
 		}
 		const FNiagaraSystemScalabilitySettings& SystemScalabilitySettings = NiagaraSystem->GetScalabilitySettings();
 		ResolvedValueStruct = MakeShared<FStructOnScope>(FNiagaraSystemScalabilitySettings::StaticStruct(), (uint8*) &SystemScalabilitySettings);
@@ -373,6 +480,7 @@ TSharedRef<SWidget> SNiagaraSystemResolvedScalabilitySettings::GenerateResolvedS
 					ScalabilityValue->ResolvedTreeNode = Child;
 					ScalabilityValue->OwningObject = Object;
 					ScalabilityValue->OwningObjectVersion = Version;
+					//ScalabilityValue->bIsOverridden = Object->IsA<UNiagaraEmitter>() ? IsEmitterScalabilityValueOverridden(Child->GetNodeName(), )
 
 					// if we found the same handle in the default handles, we have both handles needed to display default and override data
 					if(DefaultHandles.Contains(ChildPropertyHandle->GetProperty()->GetName()))
@@ -464,7 +572,8 @@ void SNiagaraSystemResolvedScalabilitySettings::Construct(const FArguments& InAr
 
 	ResolvedScalabilityContainer = SNew(SVerticalBox);
 
-	SystemViewModel->GetOverviewGraphViewModel()->GetNodeSelection()->OnSelectedObjectsChanged().AddSP(this, &SNiagaraSystemResolvedScalabilitySettings::RebuildWidget);	
+	SystemViewModel->GetOverviewGraphViewModel()->GetNodeSelection()->OnSelectedObjectsChanged().AddSP(this, &SNiagaraSystemResolvedScalabilitySettings::RebuildWidget);
+	SystemViewModel->GetSelectionViewModel()->OnEntrySelectionChanged().AddSP(this, &SNiagaraSystemResolvedScalabilitySettings::RebuildWidget);
 	System->OnScalabilityChanged().AddSP(this, &SNiagaraSystemResolvedScalabilitySettings::RebuildWidget);
 	
 	RebuildWidget();
diff --git a/Engine/Plugins/FX/Niagara/Source/NiagaraEditor/Public/ViewModels/NiagaraSystemScalabilityViewModel.h b/Engine/Plugins/FX/Niagara/Source/NiagaraEditor/Public/ViewModels/NiagaraSystemScalabilityViewModel.h
--- a/Engine/Plugins/FX/Niagara/Source/NiagaraEditor/Public/ViewModels/NiagaraSystemScalabilityViewModel.h
+++ b/Engine/Plugins/FX/Niagara/Source/NiagaraEditor/Public/ViewModels/NiagaraSystemScalabilityViewModel.h
@@ -48,7 +48,9 @@ public:
 	bool IsViewModeQualityEnabled(int32 QualityLevel) const { return PreviewPlatforms->IsEffectQualityEnabled(QualityLevel); }
 	
 	void UpdatePreviewDeviceProfile(UDeviceProfile* DeviceProfile);
-	void UpdatePreviewQualityLevel(int32 QualityLevel, bool bEnabled);
+	void UpdatePreviewQualityLevel(int32 QualityLevel);
+
+	bool IsPlatformActive(const FNiagaraPlatformSet& PlatformSet);
 
 	//void NavigateToScalabilityProperty(UObject* Object, FName PropertyName);
 	
diff --git a/Engine/Plugins/FX/Niagara/Source/NiagaraEditorWidgets/Private/SNiagaraScalabilityPreviewSettings.cpp b/Engine/Plugins/FX/Niagara/Source/NiagaraEditorWidgets/Private/SNiagaraScalabilityPreviewSettings.cpp
--- a/Engine/Plugins/FX/Niagara/Source/NiagaraEditorWidgets/Private/SNiagaraScalabilityPreviewSettings.cpp
+++ b/Engine/Plugins/FX/Niagara/Source/NiagaraEditorWidgets/Private/SNiagaraScalabilityPreviewSettings.cpp
@@ -193,7 +193,7 @@ ECheckBoxState SNiagaraScalabilityPreviewSettings::IsQLChecked(int32 QualityLeve
 
 void SNiagaraScalabilityPreviewSettings::QLCheckStateChanged(ECheckBoxState CheckState, int32 QualityLevel)
 {
-	ScalabilityViewModel->UpdatePreviewQualityLevel(QualityLevel, CheckState == ECheckBoxState::Checked);
+	ScalabilityViewModel->UpdatePreviewQualityLevel(QualityLevel);
 }
 
 FSlateColor SNiagaraScalabilityPreviewSettings::GetQualityLevelButtonTextColor(int32 QualityLevel) const
