Fixing static analysis warning

diff --git a/Engine/Plugins/FX/Niagara/Source/NiagaraEditorWidgets/Private/SNiagaraOverviewStack.cpp b/Engine/Plugins/FX/Niagara/Source/NiagaraEditorWidgets/Private/SNiagaraOverviewStack.cpp
--- a/Engine/Plugins/FX/Niagara/Source/NiagaraEditorWidgets/Private/SNiagaraOverviewStack.cpp
+++ b/Engine/Plugins/FX/Niagara/Source/NiagaraEditorWidgets/Private/SNiagaraOverviewStack.cpp
@@ -779,16 +779,8 @@ int32 SNiagaraOverviewStack::GetUsageIconHeight(UNiagaraStackEntry* StackEntry)
 
 EVisibility SNiagaraOverviewStack::GetUsageIconVisibility(UNiagaraStackEntry* StackEntry)
 {
-	bool bDisplayIcon;
-	if (StackEntry->IsA<UNiagaraStackItemGroup>() && StackEntry->GetCanExpandInOverview() && StackEntry->GetIsExpandedInOverview())
-	{
-		// If the entry is a group and it can expand and it is expanded, we only want to show the stack Usage icon if the group itself has Usages.
-		bDisplayIcon = StackEntry->HasUsagesOrAnyChildHasUsages();
-	}
-	else
-	{
-		bDisplayIcon = StackEntry->HasUsagesOrAnyChildHasUsages();
-	}
+	bool bDisplayIcon = StackEntry->HasUsagesOrAnyChildHasUsages();
+	
 	return bDisplayIcon ? EVisibility::Visible : EVisibility::Collapsed;
 }
 
