Niagara - Fix an issue in the merge manager where inputs linked to outputs or attributes of other modules would break if those modules were renamed as a result of the merge. This also fixes issues where renaming module parameters directly through the UI which wouldn't be propagated to linked module attributes.  Also add tests to verify these fixes.

diff --git a/Engine/Plugins/FX/Niagara/Source/NiagaraCore/Public/NiagaraCustomVersion.h b/Engine/Plugins/FX/Niagara/Source/NiagaraCore/Public/NiagaraCustomVersion.h
--- a/Engine/Plugins/FX/Niagara/Source/NiagaraCore/Public/NiagaraCustomVersion.h
+++ b/Engine/Plugins/FX/Niagara/Source/NiagaraCore/Public/NiagaraCustomVersion.h
@@ -178,6 +178,8 @@ struct FNiagaraCustomVersion
 
 		ComponentRendererSpawnProperty, // Changed the default value for the component renderer's OnlyCreateComponentsOnParticleSpawn property
 
+		RepopulateFunctionCallNodePinNameBindings, // Previous repopulate didn't handle module attributes like Particles.Module.Name so they need to be repopulated for renaming to work correctly.
+
 		// DO NOT ADD A NEW VERSION UNLESS YOU HAVE TALKED TO THE NIAGARA LEAD. Mismanagement of these versions can lead to data loss if it is adjusted in multiple streams simultaneously.
 		// -----<new versions can be added above this line>  -------------------------------------------------
 		VersionPlusOne,
diff --git a/Engine/Plugins/FX/Niagara/Source/NiagaraEditor/Private/NiagaraNodeFunctionCall.cpp b/Engine/Plugins/FX/Niagara/Source/NiagaraEditor/Private/NiagaraNodeFunctionCall.cpp
--- a/Engine/Plugins/FX/Niagara/Source/NiagaraEditor/Private/NiagaraNodeFunctionCall.cpp
+++ b/Engine/Plugins/FX/Niagara/Source/NiagaraEditor/Private/NiagaraNodeFunctionCall.cpp
@@ -94,7 +94,7 @@ void UNiagaraNodeFunctionCall::PostLoad()
 			UpdateStaticSwitchPinsWithPersistentGuids();
 		}
 
-		if (NiagaraCustomVersion < FNiagaraCustomVersion::PopulateFunctionCallNodePinNameBindings)
+		if (NiagaraCustomVersion < FNiagaraCustomVersion::RepopulateFunctionCallNodePinNameBindings)
 		{
 			FNiagaraStackGraphUtilities::PopulateFunctionCallNameBindings(*this);
 		}
diff --git a/Engine/Plugins/FX/Niagara/Source/NiagaraEditor/Private/NiagaraScriptMergeManager.cpp b/Engine/Plugins/FX/Niagara/Source/NiagaraEditor/Private/NiagaraScriptMergeManager.cpp
--- a/Engine/Plugins/FX/Niagara/Source/NiagaraEditor/Private/NiagaraScriptMergeManager.cpp
+++ b/Engine/Plugins/FX/Niagara/Source/NiagaraEditor/Private/NiagaraScriptMergeManager.cpp
@@ -81,7 +81,27 @@ FNiagaraStackFunctionInputOverrideMergeAdapter::FNiagaraStackFunctionInputOverri
 
 		if (OverridePin->LinkedTo[0]->GetOwningNode()->IsA<UNiagaraNodeParameterMapGet>())
 		{
-			LinkedValueHandle = FNiagaraParameterHandle(OverridePin->LinkedTo[0]->PinName);
+			LinkedValueData = FNiagaraStackLinkedValueData();
+			LinkedValueData->LinkedValueHandle = FNiagaraParameterHandle(OverridePin->LinkedTo[0]->PinName);
+			if (LinkedValueData->LinkedValueHandle.IsOutputHandle() ||
+				LinkedValueData->LinkedValueHandle.IsStackContextHandle() ||
+				LinkedValueData->LinkedValueHandle.IsEmitterHandle() ||
+				LinkedValueData->LinkedValueHandle.IsParticleAttributeHandle())
+			{
+				// If the linked handle is a module output or module data set attribute, record the node id so that we can check for renamed nodes
+				// when applying the diff.
+				TArray<FName> HandleParts = LinkedValueData->LinkedValueHandle.GetHandleParts();
+				if (HandleParts.Num() > 2)
+				{
+					FString FunctionName = HandleParts[1].ToString();
+					TObjectPtr<UEdGraphNode>* ReferencedFunctionCallNodePtr = OwningFunctionCallNode->GetGraph()->Nodes.FindByPredicate(
+						[FunctionName](UEdGraphNode* Node) { return Node->IsA<UNiagaraNodeFunctionCall>() && CastChecked<UNiagaraNodeFunctionCall>(Node)->GetFunctionName() == FunctionName; });
+					if (ReferencedFunctionCallNodePtr != nullptr)
+					{
+						LinkedValueData->LinkedFunctionNodeId = (*ReferencedFunctionCallNodePtr)->NodeGuid;
+					}
+				}
+			}
 		}
 		else if (OverridePin->LinkedTo[0]->GetOwningNode()->IsA<UNiagaraNodeInput>())
 		{
@@ -177,9 +197,9 @@ TOptional<FNiagaraVariable> FNiagaraStackFunctionInputOverrideMergeAdapter::GetL
 	return LocalValueRapidIterationParameter;
 }
 
-TOptional<FNiagaraParameterHandle> FNiagaraStackFunctionInputOverrideMergeAdapter::GetLinkedValueHandle() const
+TOptional<FNiagaraStackLinkedValueData> FNiagaraStackFunctionInputOverrideMergeAdapter::GetLinkedValueData() const
 {
-	return LinkedValueHandle;
+	return LinkedValueData;
 }
 
 TOptional<FName> FNiagaraStackFunctionInputOverrideMergeAdapter::GetDataValueInputName() const
@@ -2585,15 +2605,15 @@ TOptional<bool> FNiagaraScriptMergeManager::DoFunctionInputOverridesMatch(TShare
 	}
 
 	// Linked value
-	if ((BaseFunctionInputAdapter->GetLinkedValueHandle().IsSet() && OtherFunctionInputAdapter->GetLinkedValueHandle().IsSet() == false) ||
-		(BaseFunctionInputAdapter->GetLinkedValueHandle().IsSet() == false && OtherFunctionInputAdapter->GetLinkedValueHandle().IsSet()))
+	if ((BaseFunctionInputAdapter->GetLinkedValueData().IsSet() && OtherFunctionInputAdapter->GetLinkedValueData().IsSet() == false) ||
+		(BaseFunctionInputAdapter->GetLinkedValueData().IsSet() == false && OtherFunctionInputAdapter->GetLinkedValueData().IsSet()))
 	{
 		return false;
 	}
 
-	if (BaseFunctionInputAdapter->GetLinkedValueHandle().IsSet() && OtherFunctionInputAdapter->GetLinkedValueHandle().IsSet())
+	if (BaseFunctionInputAdapter->GetLinkedValueData().IsSet() && OtherFunctionInputAdapter->GetLinkedValueData().IsSet())
 	{
-		return BaseFunctionInputAdapter->GetLinkedValueHandle().GetValue() == OtherFunctionInputAdapter->GetLinkedValueHandle().GetValue();
+		return BaseFunctionInputAdapter->GetLinkedValueData().GetValue() == OtherFunctionInputAdapter->GetLinkedValueData().GetValue();
 	}
 
 	// Data value
@@ -2861,23 +2881,61 @@ FNiagaraScriptMergeManager::FApplyDiffResults FNiagaraScriptMergeManager::AddInp
 				InputOverridePin.DefaultValue = OverrideToAdd->GetLocalValueString().GetValue();
 				Results.bSucceeded = true;
 			}
-			else if (OverrideToAdd->GetLinkedValueHandle().IsSet())
+			else if (OverrideToAdd->GetLinkedValueData().IsSet())
 			{
 				check(OverrideToAdd->GetOverrideNode() && OverrideToAdd->GetOverrideNode()->GetNiagaraGraph());
-				FNiagaraVariable Parameter = FNiagaraVariable(InputType, OverrideToAdd->GetLinkedValueHandle().GetValue().GetParameterHandleString());
-				UNiagaraScriptVariable* ScriptVar = OverrideToAdd->GetOverrideNode()->GetNiagaraGraph()->GetScriptVariable(Parameter);
+				FNiagaraParameterHandle OldLinkedValueHandle = OverrideToAdd->GetLinkedValueData()->LinkedValueHandle;
+				FGuid LinkedFunctionCallNodeId = OverrideToAdd->GetLinkedValueData()->LinkedFunctionNodeId;
+				FNiagaraVariable OldLinkedValueVariable = FNiagaraVariable(InputType, OldLinkedValueHandle.GetParameterHandleString());
+				UNiagaraScriptVariable* ScriptVar = OverrideToAdd->GetOverrideNode()->GetNiagaraGraph()->GetScriptVariable(OldLinkedValueVariable);
 				ENiagaraDefaultMode DesiredMode = ENiagaraDefaultMode::Value;
 				if (ScriptVar)
 					DesiredMode = ScriptVar->DefaultMode;
 				if (GNiagaraForceFailIfPreviouslyNotSetOnMerge > 0)
 					DesiredMode = ENiagaraDefaultMode::FailIfPreviouslyNotSet;
-				TSet KnownParameters = { Parameter };
-				FNiagaraStackGraphUtilities::SetLinkedValueHandleForFunctionInput(InputOverridePin, OverrideToAdd->GetLinkedValueHandle().GetValue(), KnownParameters, DesiredMode, OverrideToAdd->GetOverrideNodeId());
-				FGuid LinkedOutputId = FNiagaraStackGraphUtilities::GetScriptVariableIdForLinkedOutputHandle(
-					OverrideToAdd->GetLinkedValueHandle().GetValue(), OverrideToAdd->GetType(), *TargetFunctionCall.GetNiagaraGraph());
+
+				// If the linked value handle has a valid function call node id, then we need to check to see if that function call node has been
+				// renamed due to the merge, and if so we need to update the handle.
+				FNiagaraParameterHandle NewLinkedValueHandle;
+				if (LinkedFunctionCallNodeId.IsValid() &&
+					(OldLinkedValueHandle.IsOutputHandle() || OldLinkedValueHandle.IsStackContextHandle() ||
+						OldLinkedValueHandle.IsEmitterHandle() || OldLinkedValueHandle.IsParticleAttributeHandle()) &&
+					OldLinkedValueHandle.GetHandleParts().Num() > 2)
+				{
+					TObjectPtr<UEdGraphNode>* ReferencedFunctionCallNodePtr = TargetFunctionCall.GetGraph()->Nodes.FindByPredicate(
+						[&LinkedFunctionCallNodeId](UEdGraphNode* Node) { return Node->IsA<UNiagaraNodeFunctionCall>() && Node->NodeGuid == LinkedFunctionCallNodeId; });
+					if (ReferencedFunctionCallNodePtr != nullptr)
+					{
+						UNiagaraNodeFunctionCall* ReferencedFunctionCallNode = CastChecked<UNiagaraNodeFunctionCall>(*ReferencedFunctionCallNodePtr);
+						FString FunctionName = OldLinkedValueHandle.GetHandleParts()[1].ToString();
+						if (ReferencedFunctionCallNode->GetFunctionName() != FunctionName)
+						{
+							// The function node has been renamed so we need to update the handle.  The handle is in the format [Namespace].[Function Call Name].[NameParts]
+							TArray<FName> OldHandleParts = OldLinkedValueHandle.GetHandleParts();
+							TArray<FString> NewHandleNameParts;
+							NewHandleNameParts.Add(OldHandleParts[0].ToString());
+							NewHandleNameParts.Add(ReferencedFunctionCallNode->GetFunctionName());
+							for (int32 i = 2; i < OldHandleParts.Num(); i++)
+							{
+								NewHandleNameParts.Add(OldHandleParts[i].ToString());
+							}
+							NewLinkedValueHandle = FNiagaraParameterHandle(*FString::Join(NewHandleNameParts, TEXT(".")));
+						}
+					}
+				}
+				else
+				{
+					NewLinkedValueHandle = OldLinkedValueHandle;
+				}
+
+				FNiagaraVariable NewLinkedValueVariable = FNiagaraVariable(OverrideToAdd->GetType(), NewLinkedValueHandle.GetParameterHandleString());
+				TSet KnownParameters = { NewLinkedValueVariable };
+				FNiagaraStackGraphUtilities::SetLinkedValueHandleForFunctionInput(InputOverridePin, NewLinkedValueHandle, KnownParameters, DesiredMode, OverrideToAdd->GetOverrideNodeId());
+				FGuid LinkedOutputId = FNiagaraStackGraphUtilities::GetScriptVariableIdForLinkedModuleParameterHandle(
+					OldLinkedValueHandle, OverrideToAdd->GetType(), *TargetFunctionCall.GetNiagaraGraph());
 				if (LinkedOutputId.IsValid())
 				{
-					TargetFunctionCall.UpdateInputNameBinding(LinkedOutputId, OverrideToAdd->GetLinkedValueHandle()->GetName());
+					TargetFunctionCall.UpdateInputNameBinding(LinkedOutputId, NewLinkedValueHandle.GetParameterHandleString());
 				}
 				Results.bSucceeded = true;
 			}
diff --git a/Engine/Plugins/FX/Niagara/Source/NiagaraEditor/Private/NiagaraScriptMergeManager.h b/Engine/Plugins/FX/Niagara/Source/NiagaraEditor/Private/NiagaraScriptMergeManager.h
--- a/Engine/Plugins/FX/Niagara/Source/NiagaraEditor/Private/NiagaraScriptMergeManager.h
+++ b/Engine/Plugins/FX/Niagara/Source/NiagaraEditor/Private/NiagaraScriptMergeManager.h
@@ -32,6 +32,17 @@ class UNiagaraRendererProperties;
 
 class FNiagaraStackFunctionMergeAdapter;
 
+struct FNiagaraStackLinkedValueData
+{
+	FNiagaraParameterHandle LinkedValueHandle;
+	FGuid LinkedFunctionNodeId;
+
+	bool operator==(const FNiagaraStackLinkedValueData& Other) const
+	{
+		return LinkedValueHandle == Other.LinkedValueHandle && LinkedFunctionNodeId == Other.LinkedFunctionNodeId;
+	}
+};
+
 class FNiagaraStackFunctionInputOverrideMergeAdapter
 {
 public:
@@ -59,7 +70,7 @@ public:
 
 	TOptional<FString> GetLocalValueString() const;
 	TOptional<FNiagaraVariable> GetLocalValueRapidIterationParameter() const;
-	TOptional<FNiagaraParameterHandle> GetLinkedValueHandle() const;
+	TOptional<FNiagaraStackLinkedValueData> GetLinkedValueData() const;
 	TOptional<FName> GetDataValueInputName() const;
 	UNiagaraDataInterface* GetDataValueObject() const;
 	TSharedPtr<FNiagaraStackFunctionMergeAdapter> GetDynamicValueFunction() const;
@@ -76,7 +87,7 @@ private:
 
 	TOptional<FString> LocalValueString;
 	TOptional<FNiagaraVariable> LocalValueRapidIterationParameter;
-	TOptional<FNiagaraParameterHandle> LinkedValueHandle;
+	TOptional<FNiagaraStackLinkedValueData> LinkedValueData;
 	TOptional<FName> DataValueInputName;
 	UNiagaraDataInterface* DataValueObject;
 	TSharedPtr<FNiagaraStackFunctionMergeAdapter> DynamicValueFunction;
diff --git a/Engine/Plugins/FX/Niagara/Source/NiagaraEditor/Private/ViewModels/Stack/NiagaraStackFunctionInput.cpp b/Engine/Plugins/FX/Niagara/Source/NiagaraEditor/Private/ViewModels/Stack/NiagaraStackFunctionInput.cpp
--- a/Engine/Plugins/FX/Niagara/Source/NiagaraEditor/Private/ViewModels/Stack/NiagaraStackFunctionInput.cpp
+++ b/Engine/Plugins/FX/Niagara/Source/NiagaraEditor/Private/ViewModels/Stack/NiagaraStackFunctionInput.cpp
@@ -1252,10 +1252,10 @@ void UNiagaraStackFunctionInput::SetLinkedValueHandle(const FNiagaraParameterHan
 		UEdGraphPin& OverridePin = GetOrCreateOverridePin();
 		TSet<FNiagaraVariable> KnownParameters = FNiagaraStackGraphUtilities::GetParametersForContext(OverridePin.GetOwningNode()->GetGraph(), GetSystemViewModel()->GetSystem());
 		FNiagaraStackGraphUtilities::SetLinkedValueHandleForFunctionInput(OverridePin, InParameterHandle, KnownParameters);
-		FGuid LinkedOutputId = FNiagaraStackGraphUtilities::GetScriptVariableIdForLinkedOutputHandle(InParameterHandle, InputType, *OwningFunctionCallNode->GetNiagaraGraph());
+		FGuid LinkedOutputId = FNiagaraStackGraphUtilities::GetScriptVariableIdForLinkedModuleParameterHandle(InParameterHandle, InputType, *OwningFunctionCallNode->GetNiagaraGraph());
 		if (LinkedOutputId.IsValid())
 		{
-			OwningFunctionCallNode->UpdateInputNameBinding(LinkedOutputId, InParameterHandle.GetName());
+			OwningFunctionCallNode->UpdateInputNameBinding(LinkedOutputId, InParameterHandle.GetParameterHandleString());
 		}
 	}
 
diff --git a/Engine/Plugins/FX/Niagara/Source/NiagaraEditor/Private/ViewModels/Stack/NiagaraStackGraphUtilities.cpp b/Engine/Plugins/FX/Niagara/Source/NiagaraEditor/Private/ViewModels/Stack/NiagaraStackGraphUtilities.cpp
--- a/Engine/Plugins/FX/Niagara/Source/NiagaraEditor/Private/ViewModels/Stack/NiagaraStackGraphUtilities.cpp
+++ b/Engine/Plugins/FX/Niagara/Source/NiagaraEditor/Private/ViewModels/Stack/NiagaraStackGraphUtilities.cpp
@@ -1310,7 +1310,7 @@ UEdGraphPin& FNiagaraStackGraphUtilities::GetOrCreateStackFunctionInputOverrideP
 
 		if(InputScriptVariableId.IsValid())
 		{
-			StackFunctionCall.UpdateInputNameBinding(InputScriptVariableId, AliasedInputParameterHandle.GetName());
+			StackFunctionCall.UpdateInputNameBinding(InputScriptVariableId, AliasedInputParameterHandle.GetParameterHandleString());
 		}
 	}
 	return *OverridePin;
@@ -3249,10 +3249,92 @@ TArray<UEdGraphPin*> GetStackLinkedOutputPinsForFunction(UNiagaraNodeFunctionCal
 	return LinkedOutputPins;
 }
 
-// Takes a parameter map output pin which represents a module output and finds the target function call which is linking that output.
-UNiagaraNodeFunctionCall* GetTargetFunctionCallForStackLinkedOutputPin(UEdGraphPin& StackLinkedOutputPin)
+// Searches the graph that contains InFunctionCallNode and any dependent graphs and finds parameter map get pins which reference data set attributes from the function call node.
+TArray<UEdGraphPin*> GetStackLinkedAttributePinsForFunction(UNiagaraNodeFunctionCall& InFunctionCallNode, const TArray<UNiagaraGraph*>& DependentGraphs)
 {
-	for (UEdGraphPin* LinkedOverridePin : StackLinkedOutputPin.LinkedTo)
+	TArray<UEdGraphPin*> LinkedAttributePins;
+	FName FunctionCallName = *InFunctionCallNode.GetFunctionName();
+
+	// Validate context.
+	UNiagaraGraph* CalledGraph = InFunctionCallNode.GetCalledGraph();
+	if (CalledGraph == nullptr)
+	{
+		return LinkedAttributePins;
+	}
+	UNiagaraNodeOutput* OutputNode = FNiagaraStackGraphUtilities::GetEmitterOutputNodeForStackNode(InFunctionCallNode);
+	if (OutputNode == nullptr)
+	{
+		return LinkedAttributePins;
+	}
+	TOptional<FName> AttributeNamespace = FNiagaraStackGraphUtilities::GetNamespaceForOutputNode(OutputNode);
+	if (AttributeNamespace.IsSet() == false)
+	{
+		return LinkedAttributePins;
+	}
+
+	// Check parameter map write nodes in the called graph to see if they write to any attributes in the current namespace.
+	TArray<UNiagaraNodeParameterMapSet*> CalledGraphSetNodes;
+	CalledGraph->GetNodesOfClass(CalledGraphSetNodes);
+	TSet<FName> WritePinNames;
+	for (UNiagaraNodeParameterMapSet* CalledGraphSetNode : CalledGraphSetNodes)
+	{
+		TArray<UEdGraphPin*> InputPins;
+		CalledGraphSetNode->GetInputPins(InputPins);
+		for (UEdGraphPin* InputPin : InputPins)
+		{
+			WritePinNames.Add(InputPin->PinName);
+		}
+	}
+
+	bool bHasAttributeWriteInNamespace = false;
+	for (FName WritePinName : WritePinNames)
+	{
+		FNiagaraParameterHandle WriteHandle(WritePinName);
+		TArray<FName> WriteHandleParts = WriteHandle.GetHandleParts();
+		if (WriteHandleParts.Num() >= 3 && 
+			(WriteHandleParts[0] == AttributeNamespace.GetValue() || WriteHandleParts[0] == FNiagaraConstants::StackContextNamespace) &&
+			WriteHandleParts[1] == FNiagaraConstants::ModuleNamespace)
+		{
+			bHasAttributeWriteInNamespace = true;
+			break;
+		}
+	}
+
+	// If the function call writes module attributes for the current namespace, check the current graph and any dependent graphs for reads from those
+	// module attributes.
+	if (bHasAttributeWriteInNamespace)
+	{
+		TArray<UNiagaraGraph*> GraphsToCheck = { InFunctionCallNode.GetNiagaraGraph() };
+		GraphsToCheck.Append(DependentGraphs);
+		for (UNiagaraGraph* GraphToCheck : GraphsToCheck)
+		{
+			TArray<UNiagaraNodeParameterMapGet*> MapGetNodes;
+			GraphToCheck->GetNodesOfClass(MapGetNodes);
+			for (UNiagaraNodeParameterMapGet* MapGetNode : MapGetNodes)
+			{
+				TArray<UEdGraphPin*> MapGetOutputPins;
+				MapGetNode->GetOutputPins(MapGetOutputPins);
+				for (UEdGraphPin* MapGetOutputPin : MapGetOutputPins)
+				{
+					FNiagaraParameterHandle MapGetHandle(MapGetOutputPin->PinName);
+					TArray<FName> MapGetHandleParts = MapGetHandle.GetHandleParts();
+					if (MapGetHandleParts.Num() >= 3 && 
+						(MapGetHandleParts[0] == AttributeNamespace.GetValue() || MapGetHandleParts[0] == FNiagaraConstants::StackContextNamespace) &&
+						MapGetHandleParts[1] == FunctionCallName)
+					{
+						LinkedAttributePins.Add(MapGetOutputPin);
+					}
+				}
+			}
+		}
+	}
+	return LinkedAttributePins;
+}
+
+// Takes a parameter map output pin which represents a module output or module attribute and finds the target function call which is linking that output or attribute.
+UNiagaraNodeFunctionCall* GetTargetFunctionCallForStackLinkedModuleParameterPin(UEdGraphPin& StackLinkedModuleParameterPin)
+{
+	for (UEdGraphPin* LinkedOverridePin : StackLinkedModuleParameterPin.LinkedTo)
 	{
 		UNiagaraNodeParameterMapSet* LinkedOverrideNode = Cast<UNiagaraNodeParameterMapSet>(LinkedOverridePin->GetOwningNode());
 		if (LinkedOverrideNode != nullptr)
@@ -3275,6 +3357,30 @@ UNiagaraNodeFunctionCall* GetTargetFunctionCallForStackLinkedOutputPin(UEdGraphP
 	return nullptr;
 }
 
+void GetDependentGraphsForOuterGraph(UNiagaraNodeFunctionCall& InFunctionCallNode, TArray<UNiagaraGraph*>& OutDependentGraphs)
+{
+	UNiagaraNodeOutput* OutputNode = FNiagaraStackGraphUtilities::GetEmitterOutputNodeForStackNode(InFunctionCallNode);
+	if (OutputNode != nullptr &&
+		(OutputNode->GetUsage() == ENiagaraScriptUsage::SystemSpawnScript || OutputNode->GetUsage() == ENiagaraScriptUsage::SystemUpdateScript))
+	{
+		UNiagaraSystem* OuterSystem = InFunctionCallNode.GetTypedOuter<UNiagaraSystem>();
+		if (OuterSystem != nullptr)
+		{
+			for (const FNiagaraEmitterHandle& EmitterHandle : OuterSystem->GetEmitterHandles())
+			{
+				if (EmitterHandle.GetEmitterData() != nullptr && EmitterHandle.GetEmitterData()->GraphSource != nullptr)
+				{
+					UNiagaraScriptSource* EmitterSource = Cast<UNiagaraScriptSource>(EmitterHandle.GetEmitterData()->GraphSource);
+					if (EmitterSource != nullptr && EmitterSource->NodeGraph != nullptr)
+					{
+						OutDependentGraphs.Add(EmitterSource->NodeGraph);
+					}
+				}
+			}
+		}
+	}
+}
+
 void FNiagaraStackGraphUtilities::PopulateFunctionCallNameBindings(UNiagaraNodeFunctionCall& InFunctionCallNode)
 {
 	UNiagaraNodeOutput* OutputNode = GetEmitterOutputNodeForStackNode(InFunctionCallNode);
@@ -3290,7 +3396,7 @@ void FNiagaraStackGraphUtilities::PopulateFunctionCallNameBindings(UNiagaraNodeF
 	TSharedPtr<TMap<FName, FGuid>> VariableNameToVariableIdMap;
 	const UEdGraphSchema_Niagara* NiagaraSchema = GetDefault<UEdGraphSchema_Niagara>();
 
-	// Populate input ids from the inputs on this function which are overriden.
+	// Populate input ids from the inputs on this function which are overridden.
 	TArray<UEdGraphPin*> OverridePins;
 	UNiagaraNodeParameterMapSet* OverrideNode = GetStackFunctionOverrideNode(InFunctionCallNode);
 	if (OverrideNode != nullptr)
@@ -3308,34 +3414,47 @@ void FNiagaraStackGraphUtilities::PopulateFunctionCallNameBindings(UNiagaraNodeF
 			FGuid* IdByInputVariableName = VariableNameToVariableIdMap->Find(InputVariable.GetName());
 			if (IdByInputVariableName != nullptr)
 			{
-				InFunctionCallNode.UpdateInputNameBinding(*IdByInputVariableName, OverrideHandle.GetName());
+				InFunctionCallNode.UpdateInputNameBinding(*IdByInputVariableName, OverrideHandle.GetParameterHandleString());
 			}
 		}
 	}
 
-	// Populate ids on target functions which use outputs from this function.
-	TArray<UEdGraphPin*> StackLinkedOutputPins = GetStackLinkedOutputPinsForFunction(InFunctionCallNode);
-	if (StackLinkedOutputPins.Num() > 0)
+	// Populate ids on target functions which use outputs or module attributes from this function.
+	auto PopulateForStackLinkedModuleParameterPins = [&NiagaraSchema, &InFunctionCallNode, &VariableNameToVariableIdMap](const TArray<UEdGraphPin*>& StackLinkedModuleParameterPins)
 	{
 		if (VariableNameToVariableIdMap.IsValid() == false)
 		{
 			VariableNameToVariableIdMap = GetVariableNameToVariableIdMap(InFunctionCallNode);
 		}
-		for (UEdGraphPin* StackLinkedOutputPin : StackLinkedOutputPins)
+		for (UEdGraphPin* StackLinkedModuleParameterPin : StackLinkedModuleParameterPins)
 		{
-			UNiagaraNodeFunctionCall* TargetFunctionCallNode = GetTargetFunctionCallForStackLinkedOutputPin(*StackLinkedOutputPin);
+			UNiagaraNodeFunctionCall* TargetFunctionCallNode = GetTargetFunctionCallForStackLinkedModuleParameterPin(*StackLinkedModuleParameterPin);
 			if (TargetFunctionCallNode != nullptr)
 			{
-				FNiagaraVariable LinkedOutputVariable = NiagaraSchema->PinToNiagaraVariable(StackLinkedOutputPin);
-				FNiagaraParameterHandle LinkedOutputHandle(LinkedOutputVariable.GetName());
-				FNiagaraVariable OutputVariable = FNiagaraUtilities::ResolveAliases(LinkedOutputVariable, FNiagaraAliasContext().ChangeModuleNameToModule(InFunctionCallNode.GetFunctionName()));
-				FGuid* IdByVariableName = VariableNameToVariableIdMap->Find(OutputVariable.GetName());
+				FNiagaraVariable LinkedModuleParameterVariable = NiagaraSchema->PinToNiagaraVariable(StackLinkedModuleParameterPin);
+				FNiagaraParameterHandle LinkedModuleParameterHandle(LinkedModuleParameterVariable.GetName());
+				FNiagaraVariable ModuleParameterVariable = FNiagaraUtilities::ResolveAliases(LinkedModuleParameterVariable, FNiagaraAliasContext().ChangeModuleNameToModule(InFunctionCallNode.GetFunctionName()));
+				FGuid* IdByVariableName = VariableNameToVariableIdMap->Find(ModuleParameterVariable.GetName());
 				if (IdByVariableName != nullptr)
 				{
-					InFunctionCallNode.UpdateInputNameBinding(*IdByVariableName, LinkedOutputHandle.GetName());
+					TargetFunctionCallNode->UpdateInputNameBinding(*IdByVariableName, LinkedModuleParameterHandle.GetParameterHandleString());
 				}
 			}
 		}
+	};
+
+	TArray<UEdGraphPin*> StackLinkedOutputPins = GetStackLinkedOutputPinsForFunction(InFunctionCallNode);
+	if (StackLinkedOutputPins.Num() > 0)
+	{
+		PopulateForStackLinkedModuleParameterPins(StackLinkedOutputPins);
+	}
+
+	TArray<UNiagaraGraph*> DependentGraphs;
+	GetDependentGraphsForOuterGraph(InFunctionCallNode, DependentGraphs);
+	TArray<UEdGraphPin*> StackLinkedAttributePins = GetStackLinkedAttributePinsForFunction(InFunctionCallNode, DependentGraphs);
+	if (StackLinkedAttributePins.Num() > 0)
+	{
+		PopulateForStackLinkedModuleParameterPins(StackLinkedAttributePins);
 	}
 }
 
@@ -3368,7 +3487,7 @@ void FNiagaraStackGraphUtilities::SynchronizeReferencingMapPinsWithFunctionCall(
 		{
 			FNiagaraVariable OverrideVariable = NiagaraSchema->PinToNiagaraVariable(OverridePin);
 			FNiagaraParameterHandle OverrideHandle(OverrideVariable.GetName());
-			TArray<FGuid> BoundGuids = InFunctionCallNode.GetBoundPinGuidsByName(OverrideHandle.GetName());
+			TArray<FGuid> BoundGuids = InFunctionCallNode.GetBoundPinGuidsByName(OverrideHandle.GetParameterHandleString());
 			for (FGuid BoundGuid : BoundGuids)
 			{
 				FNiagaraVariable* BoundVariable = VariableIdToVariableMap->Find(BoundGuid);
@@ -3379,7 +3498,7 @@ void FNiagaraStackGraphUtilities::SynchronizeReferencingMapPinsWithFunctionCall(
 					{
 						FNiagaraParameterHandle UpdatedOverrideHandle = FNiagaraParameterHandle::CreateAliasedModuleParameterHandle(BoundVariable->GetName(), &InFunctionCallNode);
 						OverridePin->PinName = UpdatedOverrideHandle.GetParameterHandleString();
-						InFunctionCallNode.UpdateInputNameBinding(BoundGuid, BoundVariableHandle.GetName());
+						InFunctionCallNode.UpdateInputNameBinding(BoundGuid, BoundVariableHandle.GetParameterHandleString());
 					}
 					if (OverrideVariable.GetType() != BoundVariable->GetType())
 					{
@@ -3400,22 +3519,21 @@ void FNiagaraStackGraphUtilities::SynchronizeReferencingMapPinsWithFunctionCall(
 		}
 	}
 
-	// Update linked outputs referenced by downstream parameter map gets.
-	TArray<UEdGraphPin*> StackLinkedOutputPins = GetStackLinkedOutputPinsForFunction(InFunctionCallNode);
-	if (StackLinkedOutputPins.Num() > 0)
+	// Update linked module outputs or module attributes referenced by downstream parameter map gets.
+	auto SynchronizeForStackLinkedModuleParameterPins = [&NiagaraSchema, &InFunctionCallNode, &VariableIdToVariableMap](const TArray<UEdGraphPin*>& StackLinkedModuleParameterPins)
 	{
 		if (VariableIdToVariableMap.IsValid() == false)
 		{
 			VariableIdToVariableMap = GetVariableIdToVariableMap(InFunctionCallNode);
 		}
-		for (UEdGraphPin* StackLinkedOutputPin : StackLinkedOutputPins)
+		for (UEdGraphPin* StackLinkedModuleParameterPin : StackLinkedModuleParameterPins)
 		{
-			FNiagaraVariable LinkedOutputVariable = NiagaraSchema->PinToNiagaraVariable(StackLinkedOutputPin);
-			FNiagaraParameterHandle LinkedOutputHandle(LinkedOutputVariable.GetName());
-			UNiagaraNodeFunctionCall* TargetFunctionCallNode = GetTargetFunctionCallForStackLinkedOutputPin(*StackLinkedOutputPin);
+			FNiagaraVariable LinkedModuleParameterVariable = NiagaraSchema->PinToNiagaraVariable(StackLinkedModuleParameterPin);
+			FNiagaraParameterHandle LinkedModuleParameterHandle(LinkedModuleParameterVariable.GetName());
+			UNiagaraNodeFunctionCall* TargetFunctionCallNode = GetTargetFunctionCallForStackLinkedModuleParameterPin(*StackLinkedModuleParameterPin);
 			if (TargetFunctionCallNode != nullptr)
 			{
-				TArray<FGuid> BoundGuids = TargetFunctionCallNode->GetBoundPinGuidsByName(LinkedOutputHandle.GetParameterHandleString());
+				TArray<FGuid> BoundGuids = TargetFunctionCallNode->GetBoundPinGuidsByName(LinkedModuleParameterHandle.GetParameterHandleString());
 				for (const FGuid& BoundGuid : BoundGuids)
 				{
 					FNiagaraVariable* BoundVariable = VariableIdToVariableMap->Find(BoundGuid);
@@ -3423,26 +3541,41 @@ void FNiagaraStackGraphUtilities::SynchronizeReferencingMapPinsWithFunctionCall(
 					{
 						FNiagaraVariable BoundVariableResolvedName = FNiagaraUtilities::ResolveAliases(*BoundVariable,FNiagaraAliasContext().ChangeModuleToModuleName(InFunctionCallNode.GetFunctionName()));
 						FNiagaraParameterHandle BoundVariableHandle(BoundVariableResolvedName.GetName());
-						if (LinkedOutputHandle.GetName() != BoundVariableHandle.GetName())
+						if (LinkedModuleParameterHandle.GetName() != BoundVariableHandle.GetName())
 						{
-							StackLinkedOutputPin->PinName = BoundVariableResolvedName.GetName();
-							TargetFunctionCallNode->UpdateInputNameBinding(BoundGuid, BoundVariableHandle.GetName());
+							StackLinkedModuleParameterPin->PinName = BoundVariableResolvedName.GetName();
+							TargetFunctionCallNode->UpdateInputNameBinding(BoundGuid, BoundVariableHandle.GetParameterHandleString());
 						}
-						if (LinkedOutputVariable.GetType() != BoundVariable->GetType())
+						if (LinkedModuleParameterVariable.GetType() != BoundVariable->GetType())
 						{
-							StackLinkedOutputPin->PinType = NiagaraSchema->TypeDefinitionToPinType(BoundVariable->GetType());
+							StackLinkedModuleParameterPin->PinType = NiagaraSchema->TypeDefinitionToPinType(BoundVariable->GetType());
 						}
 						break;
 					}
 				}
 			}
 		}
+	};
+
+	TArray<UEdGraphPin*> StackLinkedOutputPins = GetStackLinkedOutputPinsForFunction(InFunctionCallNode);
+	if (StackLinkedOutputPins.Num() > 0)
+	{
+		SynchronizeForStackLinkedModuleParameterPins(StackLinkedOutputPins);
+	}
+
+	TArray<UNiagaraGraph*> DependentGraphs;
+	GetDependentGraphsForOuterGraph(InFunctionCallNode, DependentGraphs);
+	TArray<UEdGraphPin*> StackLinkedAttributePins = GetStackLinkedAttributePinsForFunction(InFunctionCallNode, DependentGraphs);
+	if (StackLinkedAttributePins.Num() > 0)
+	{
+		SynchronizeForStackLinkedModuleParameterPins(StackLinkedAttributePins);
 	}
 }
 
-FGuid FNiagaraStackGraphUtilities::GetScriptVariableIdForLinkedOutputHandle(const FNiagaraParameterHandle& LinkedHandle, FNiagaraTypeDefinition LinkedType, UNiagaraGraph& TargetGraph)
+FGuid FNiagaraStackGraphUtilities::GetScriptVariableIdForLinkedModuleParameterHandle(const FNiagaraParameterHandle& LinkedHandle, FNiagaraTypeDefinition LinkedType, UNiagaraGraph& TargetGraph)
 {
-	if (LinkedHandle.IsOutputHandle() && LinkedHandle.GetHandleParts().Num() >= 3)
+	if (LinkedHandle.GetHandleParts().Num() >= 3 &&
+		(LinkedHandle.IsOutputHandle() || LinkedHandle.IsStackContextHandle() || LinkedHandle.IsSystemHandle() || LinkedHandle.IsEmitterHandle() || LinkedHandle.IsParticleAttributeHandle()))
 	{
 		// First find the source function call node.
 		FString LinkedFunctionName = LinkedHandle.GetHandleParts()[1].ToString();
diff --git a/Engine/Plugins/FX/Niagara/Source/NiagaraEditor/Public/NiagaraGraph.h b/Engine/Plugins/FX/Niagara/Source/NiagaraEditor/Public/NiagaraGraph.h
--- a/Engine/Plugins/FX/Niagara/Source/NiagaraEditor/Public/NiagaraGraph.h
+++ b/Engine/Plugins/FX/Niagara/Source/NiagaraEditor/Public/NiagaraGraph.h
@@ -161,7 +161,7 @@ class UNiagaraGraph : public UEdGraph
 			
 	/** Find the first output node bound to the target usage type.*/
 	class UNiagaraNodeOutput* FindOutputNode(ENiagaraScriptUsage TargetUsageType, FGuid TargetUsageId = FGuid()) const;
-	class UNiagaraNodeOutput* FindEquivalentOutputNode(ENiagaraScriptUsage TargetUsageType, FGuid TargetUsageId = FGuid()) const;
+	NIAGARAEDITOR_API class UNiagaraNodeOutput* FindEquivalentOutputNode(ENiagaraScriptUsage TargetUsageType, FGuid TargetUsageId = FGuid()) const;
 
 	/** Find all output nodes.*/
 	void FindOutputNodes(TArray<UNiagaraNodeOutput*>& OutputNodes) const;
diff --git a/Engine/Plugins/FX/Niagara/Source/NiagaraEditor/Public/NiagaraSystemFactoryNew.h b/Engine/Plugins/FX/Niagara/Source/NiagaraEditor/Public/NiagaraSystemFactoryNew.h
--- a/Engine/Plugins/FX/Niagara/Source/NiagaraEditor/Public/NiagaraSystemFactoryNew.h
+++ b/Engine/Plugins/FX/Niagara/Source/NiagaraEditor/Public/NiagaraSystemFactoryNew.h
@@ -26,7 +26,7 @@ class UNiagaraSystemFactoryNew : public UFactory
 	//~ Begin UFactory Interface	
 
 public:
-	static void InitializeSystem(UNiagaraSystem* System, bool bCreateDefaultNodes);
+	NIAGARAEDITOR_API static void InitializeSystem(UNiagaraSystem* System, bool bCreateDefaultNodes);
 };
 
 
diff --git a/Engine/Plugins/FX/Niagara/Source/NiagaraEditor/Public/ViewModels/Stack/NiagaraStackGraphUtilities.h b/Engine/Plugins/FX/Niagara/Source/NiagaraEditor/Public/ViewModels/Stack/NiagaraStackGraphUtilities.h
--- a/Engine/Plugins/FX/Niagara/Source/NiagaraEditor/Public/ViewModels/Stack/NiagaraStackGraphUtilities.h
+++ b/Engine/Plugins/FX/Niagara/Source/NiagaraEditor/Public/ViewModels/Stack/NiagaraStackGraphUtilities.h
@@ -189,7 +189,7 @@ namespace FNiagaraStackGraphUtilities
 
 	UNiagaraNodeFunctionCall* AddScriptModuleToStack(FAssetData ModuleScriptAsset, UNiagaraNodeOutput& TargetOutputNode, int32 TargetIndex = INDEX_NONE, FString SuggestedName = FString());
 
-	UNiagaraNodeFunctionCall* AddScriptModuleToStack(UNiagaraScript* ModuleScript, UNiagaraNodeOutput& TargetOutputNode, int32 TargetIndex = INDEX_NONE, FString SuggestedName = FString(), const FGuid& VersionGuid = FGuid());
+	NIAGARAEDITOR_API UNiagaraNodeFunctionCall* AddScriptModuleToStack(UNiagaraScript* ModuleScript, UNiagaraNodeOutput& TargetOutputNode, int32 TargetIndex = INDEX_NONE, FString SuggestedName = FString(), const FGuid& VersionGuid = FGuid());
 	
 	bool FindScriptModulesInStack(FAssetData ModuleScriptAsset, UNiagaraNodeOutput& TargetOutputNode, TArray<UNiagaraNodeFunctionCall*> OutFunctionCalls);
 
@@ -270,7 +270,7 @@ namespace FNiagaraStackGraphUtilities
 
 	NIAGARAEDITOR_API void SynchronizeReferencingMapPinsWithFunctionCall(UNiagaraNodeFunctionCall& InFunctionCallNode);
 
-	FGuid GetScriptVariableIdForLinkedOutputHandle(const FNiagaraParameterHandle& LinkedOutputHandle, FNiagaraTypeDefinition LinkedType, UNiagaraGraph& TargetGraph);
+	FGuid GetScriptVariableIdForLinkedModuleParameterHandle(const FNiagaraParameterHandle& LinkedOutputHandle, FNiagaraTypeDefinition LinkedType, UNiagaraGraph& TargetGraph);
 
 	namespace DependencyUtilities
 	{
