fix hlsl compiler warnings

diff --git a/Engine/Plugins/FX/Niagara/Source/Niagara/Private/NiagaraDataInterfaceGrid3DCollection.cpp b/Engine/Plugins/FX/Niagara/Source/Niagara/Private/NiagaraDataInterfaceGrid3DCollection.cpp
--- a/Engine/Plugins/FX/Niagara/Source/Niagara/Private/NiagaraDataInterfaceGrid3DCollection.cpp
+++ b/Engine/Plugins/FX/Niagara/Source/Niagara/Private/NiagaraDataInterfaceGrid3DCollection.cpp
@@ -1512,7 +1512,7 @@ bool UNiagaraDataInterfaceGrid3DCollection::GetFunctionHLSL(const FNiagaraDataIn
 					// we automatically pass through attributes higher than 31 since we can only
 					// store 32 attributes in the mask
 					[branch]
-					if ((AttributeMask & (1 << AttributeIndex)) == 0 || AttributeIndex >= 32)
+					if ((AttributeMask & (1l << AttributeIndex)) == 0 || AttributeIndex >= 32)
 					{						
 						int3 TileOffset = {PerAttributeDataName}[AttributeIndex].xyz;
 
@@ -1999,7 +1999,7 @@ bool UNiagaraDataInterfaceGrid3DCollection::GenerateIterationSourceNamespaceWrit
 
 			for (int j = 0; j < ComponentCount; ++j)
 			{
-				OutHLSL += FString::Printf(TEXT("\tAttributeIsSetMask |= 1 << (CurrAttributeIndex+%d);\n"), j);				
+				OutHLSL += FString::Printf(TEXT("\tAttributeIsSetMask |= 1l << (CurrAttributeIndex+%d);\n"), j);				
 				NumAttributesSet++;				
 			}			
 			OutHLSL += FString::Printf(TEXT("\t%s.%s<Attribute=\"%s\">(X, Y, Z,  %s);\n"), *DIVarName, *TypeDefinitionToSetFunctionName(InAttributes[i].GetType()).ToString(), *AttributeName, *InAttributeHLSLNames[i]);
