- Change GPU initialize LODResource to TRefCountPtr to avoid stream out

#9252

diff --git a/Engine/Plugins/FX/Niagara/Source/Niagara/Private/DataInterface/NiagaraDataInterfaceStaticMesh.cpp b/Engine/Plugins/FX/Niagara/Source/Niagara/Private/DataInterface/NiagaraDataInterfaceStaticMesh.cpp
--- a/Engine/Plugins/FX/Niagara/Source/Niagara/Private/DataInterface/NiagaraDataInterfaceStaticMesh.cpp
+++ b/Engine/Plugins/FX/Niagara/Source/Niagara/Private/DataInterface/NiagaraDataInterfaceStaticMesh.cpp
@@ -274,7 +274,7 @@ namespace NDIStaticMeshLocal
 	{
 		struct FRenderProxy*				RenderProxy = nullptr;
 		FNiagaraSystemInstanceID			SystemInstanceID;
-		const FStaticMeshLODResources*		LODResource = nullptr;
+		TRefCountPtr<const FStaticMeshLODResources>		LODResource = nullptr;
 
 		int32								NumFilteredTriangles = 0;
 		int32								NumUnfilteredTriangles = 0;
