- Fix GPU executing particle update script on spawn when interpolated spawning is disabled

diff --git a/Engine/Plugins/FX/Niagara/Source/Niagara/Classes/NiagaraEmitter.h b/Engine/Plugins/FX/Niagara/Source/Niagara/Classes/NiagaraEmitter.h
--- a/Engine/Plugins/FX/Niagara/Source/Niagara/Classes/NiagaraEmitter.h
+++ b/Engine/Plugins/FX/Niagara/Source/Niagara/Classes/NiagaraEmitter.h
@@ -266,6 +266,15 @@ struct NIAGARA_API FVersionedNiagaraEmitterData
 	UPROPERTY(EditAnywhere, BlueprintReadWrite, Category = "Emitter")
 	uint32 bInterpolatedSpawning : 1;
 
+#if WITH_EDITORONLY_DATA
+	/**
+	GPU scripts were incorrectly running both particle spawn & update, CPU only runs spawn when interpolated spawning is disabled.
+	This flag allows backwards compatability so content does not change post the fix.
+	*/
+	UPROPERTY()
+	uint32 bGpuAlwaysRunParticleUpdateScript : 1;
+#endif
+
 	/** An emitter-based seed for the deterministic random number generator. */
 	UPROPERTY(EditAnywhere, BlueprintReadWrite, Category = "Emitter", meta = (EditCondition = "bDeterminism"))
 	int32 RandomSeed = 0;
diff --git a/Engine/Plugins/FX/Niagara/Source/Niagara/Private/NiagaraEmitter.cpp b/Engine/Plugins/FX/Niagara/Source/Niagara/Private/NiagaraEmitter.cpp
--- a/Engine/Plugins/FX/Niagara/Source/Niagara/Private/NiagaraEmitter.cpp
+++ b/Engine/Plugins/FX/Niagara/Source/Niagara/Private/NiagaraEmitter.cpp
@@ -5,6 +5,7 @@
 
 #include "INiagaraEditorOnlyDataUtlities.h"
 #include "NiagaraCustomVersion.h"
+#include "NiagaraCustomVersionUEMain.h"
 #include "NiagaraEditorDataBase.h"
 #include "NiagaraModule.h"
 #include "NiagaraRenderer.h"
@@ -517,6 +518,7 @@ void UNiagaraEmitter::Serialize(FArchive& Ar)
 #endif
 
 	Ar.UsingCustomVersion(FNiagaraCustomVersion::GUID);
+	Ar.UsingCustomVersion(FNiagaraCustomVersionUEMain::GUID);
 }
 
 void FVersionedNiagaraEmitterData::EnsureScriptsPostLoaded()
@@ -615,15 +617,27 @@ void UNiagaraEmitter::PostLoad()
 	{
 		LibraryVisibility = ENiagaraScriptLibraryVisibility::Library;
 	}
-
 #endif
 
+	const int32 NiagaraVerUE5Main = GetLinkerCustomVersion(FNiagaraCustomVersionUEMain::GUID);
+
 	for (FVersionedNiagaraEmitterData& Data : VersionData)
 	{
 #if WITH_EDITORONLY_DATA
 		Data.PostLoad(*this, IsCooked);
 		Data.GPUComputeScript->OnGPUScriptCompiled().RemoveAll(this);
 		Data.GPUComputeScript->OnGPUScriptCompiled().AddUObject(this, &UNiagaraEmitter::RaiseOnEmitterGPUCompiled);
+		if (NiagaraVerUE5Main < FNiagaraCustomVersionUEMain::FixGpuAlwaysRunningUpdateScriptNoneInterpolated)
+		{
+			if ( Data.SimTarget == ENiagaraSimTarget::GPUComputeSim )
+			{
+				Data.bGpuAlwaysRunParticleUpdateScript = Data.bInterpolatedSpawning ? false : true;
+			}
+			else
+			{
+				Data.bGpuAlwaysRunParticleUpdateScript = false;
+			}
+		}
 #else
 		Data.PostLoad(*this, true);
 #endif
@@ -1113,6 +1127,18 @@ void UNiagaraEmitter::PostEditChangeVersionedProperty(FPropertyChangedEvent& Pro
 			bNeedsRecompile = true;
 		}
 	}
+	else if (PropertyName == GET_MEMBER_NAME_CHECKED(FVersionedNiagaraEmitterData, bGpuAlwaysRunParticleUpdateScript))
+	{
+		const bool bActualInterpolatedSpawning = EmitterData->SpawnScriptProps.Script->IsInterpolatedParticleSpawnScript();
+		if ((bActualInterpolatedSpawning == false) && (EmitterData->SimTarget == ENiagaraSimTarget::GPUComputeSim))
+		{
+			if (EmitterData->GraphSource != nullptr)
+			{
+				EmitterData->GraphSource->MarkNotSynchronized(TEXT("Emitter GPU interpolated spawn changed"));
+			}
+			bNeedsRecompile = true;
+		}
+	}
 	else if (PropertyName == GET_MEMBER_NAME_CHECKED(FVersionedNiagaraEmitterData, SimTarget))
 	{
 		if (EmitterData->GraphSource != nullptr)
@@ -2546,10 +2572,19 @@ void UNiagaraEmitter::CheckVersionDataAvailable()
 
 #endif
 
-FVersionedNiagaraEmitterData::FVersionedNiagaraEmitterData() :
-	bInterpolatedSpawning(false), bRequiresPersistentIDs(false), bCombineEventSpawn(false), bLimitDeltaTime(true), MaxGPUParticlesSpawnPerFrame(0), bRequiresViewUniformBuffer(false)
+FVersionedNiagaraEmitterData::FVersionedNiagaraEmitterData()
+	: bInterpolatedSpawning(false)
+#if WITH_EDITORONLY_DATA
+	, bGpuAlwaysRunParticleUpdateScript(false)
+#endif
+	, bRequiresPersistentIDs(false)
+	, bCombineEventSpawn(false)
+	, bLimitDeltaTime(true)
+	, MaxGPUParticlesSpawnPerFrame(0)
+	, bRequiresViewUniformBuffer(false)
 #if WITH_EDITORONLY_DATA
-	, EditorData(nullptr), EditorParameters(nullptr)
+	, EditorData(nullptr)
+	, EditorParameters(nullptr)
 #endif
 {
 	FixedBounds = GetDefaultFixedBounds();
diff --git a/Engine/Plugins/FX/Niagara/Source/Niagara/Private/NiagaraScript.cpp b/Engine/Plugins/FX/Niagara/Source/Niagara/Private/NiagaraScript.cpp
--- a/Engine/Plugins/FX/Niagara/Source/Niagara/Private/NiagaraScript.cpp
+++ b/Engine/Plugins/FX/Niagara/Source/Niagara/Private/NiagaraScript.cpp
@@ -1012,6 +1012,10 @@ void UNiagaraScript::ComputeVMCompilationId(FNiagaraVMExecutableDataId& Id, FGui
 		{
 			Id.AdditionalDefines.Add(TEXT("Emitter.InterpolatedSpawn"));
 		}
+		if (EmitterData->bInterpolatedSpawning == false && EmitterData->bGpuAlwaysRunParticleUpdateScript == true)
+		{
+			Id.AdditionalDefines.Add(TEXT("GpuAlwaysRunParticleUpdateScript"));
+		}
 		if (IsParticleScript(Usage) || IsParticleStageScript(Usage))
 		{
 			SimTargetToBuild = EmitterData->SimTarget;
diff --git a/Engine/Plugins/FX/Niagara/Source/NiagaraCore/Private/NiagaraCustomVersionUEMain.cpp b/Engine/Plugins/FX/Niagara/Source/NiagaraCore/Private/NiagaraCustomVersionUEMain.cpp
new file mode 100644
--- /dev/null
+++ b/Engine/Plugins/FX/Niagara/Source/NiagaraCore/Private/NiagaraCustomVersionUEMain.cpp
@@ -0,0 +1,8 @@
+// Copyright Epic Games, Inc. All Rights Reserved.
+
+#include "NiagaraCustomVersionUEMain.h"
+
+const FGuid FNiagaraCustomVersionUEMain::GUID(0xFACD84A7, 0xCEFF5F42, 0x8FB5674C, 0x40CE42BB);
+
+// Register the custom version with core
+FCustomVersionRegistration GRegisterNiagaraCustomVersionUE5Main(FNiagaraCustomVersionUEMain::GUID, FNiagaraCustomVersionUEMain::LatestVersion, TEXT("NiagaraVerUEMain"));
diff --git a/Engine/Plugins/FX/Niagara/Source/NiagaraCore/Public/NiagaraCustomVersionUEMain.h b/Engine/Plugins/FX/Niagara/Source/NiagaraCore/Public/NiagaraCustomVersionUEMain.h
new file mode 100644
--- /dev/null
+++ b/Engine/Plugins/FX/Niagara/Source/NiagaraCore/Public/NiagaraCustomVersionUEMain.h
@@ -0,0 +1,27 @@
+// Copyright Epic Games, Inc. All Rights Reserved.
+
+#pragma once
+
+#include "CoreMinimal.h"
+#include "Misc/Guid.h"
+
+// Custom serialization version for all packages containing Niagara asset types
+struct FNiagaraCustomVersionUEMain
+{
+	enum Type
+	{
+		// GPU none interpolated spawning no longer calls the update script
+		FixGpuAlwaysRunningUpdateScriptNoneInterpolated,
+
+		// DO NOT ADD A NEW VERSION UNLESS YOU HAVE TALKED TO THE NIAGARA LEAD. Mismanagement of these versions can lead to data loss if it is adjusted in multiple streams simultaneously.
+		// -----<new versions can be added above this line>  -------------------------------------------------
+		VersionPlusOne,
+		LatestVersion = VersionPlusOne - 1,
+	};
+
+	// The GUID for this custom version number
+	NIAGARACORE_API const static FGuid GUID;
+
+private:
+	FNiagaraCustomVersionUEMain() {}
+};
diff --git a/Engine/Plugins/FX/Niagara/Source/NiagaraEditor/Private/NiagaraHlslTranslator.cpp b/Engine/Plugins/FX/Niagara/Source/NiagaraEditor/Private/NiagaraHlslTranslator.cpp
--- a/Engine/Plugins/FX/Niagara/Source/NiagaraEditor/Private/NiagaraHlslTranslator.cpp
+++ b/Engine/Plugins/FX/Niagara/Source/NiagaraEditor/Private/NiagaraHlslTranslator.cpp
@@ -3066,8 +3066,6 @@ void FHlslNiagaraTranslator::DefineMainGPUFunctions(
 
 	//////////////////////
 	// TransferAttibutes()
-
-	
 	HlslOutput += TEXT("void TransferAttributes(inout FSimulationContext Context)\n{\n");
 	{
 		FExpressionPermutationContext PermutationContext(HlslOutput);
@@ -3086,6 +3084,10 @@ void FHlslNiagaraTranslator::DefineMainGPUFunctions(
 					{
 						HlslOutput += TEXT("\t\tContext.") + TranslationStages[i].PassNamespace + TEXT(".DataInstance = Context.") + TranslationStages[i - 1].PassNamespace + TEXT(".DataInstance;\n");
 					}
+					else if ( TranslationStages[i].bWritesAlive )
+					{
+						HlslOutput += TEXT("\t\tContext.") + TranslationStages[i].PassNamespace + TEXT(".DataInstance.Alive = true;\n");
+					}
 				}
 				
 				if (i == 1 && TranslationStages[i].ScriptUsage == ENiagaraScriptUsage::ParticleUpdateScript) // The Update Phase might need previous parameters set.
@@ -3115,7 +3117,6 @@ void FHlslNiagaraTranslator::DefineMainGPUFunctions(
 	
 	/////////////////////////
 	// StoreUpdateVariables()
-
 	HlslOutput += TEXT("void StoreUpdateVariables(in FSimulationContext Context, bool bIsValidInstance)\n{\n");
 	{
 		FExpressionPermutationContext PermutationContext(HlslOutput);
@@ -3217,22 +3218,6 @@ void FHlslNiagaraTranslator::DefineMainGPUFunctions(
 	}
 	HlslOutput += TEXT("\n}\n\n");
 
-	/////////////////////////
-	// SimulateDoWork()
-
-	//HlslOutput += TEXT("void SimulateDoWork(in FSimulationContext Context)\n{\n");
-	FString SimDoWorkString;// = TEXT("\t{ \n");
-	{
-		FExpressionPermutationContext PermutationContext(SimDoWorkString);
-		int32 StartIdx = 1;
-		for (int32 i = StartIdx; i < TranslationStages.Num(); i++)
-		{
-			PermutationContext.AddBranch(*this, TranslationStages[i]);
-			SimDoWorkString += TEXT("\t\tSimulate") + TranslationStages[i].PassNamespace + TEXT("(Context); \n");
-		}
-		//SimDoWorkString += TEXT("\t}\n");
-	}
-
 	/////////////////////////
 	// CopyInstance()
 	HlslOutput += TEXT("void CopyInstance(in int InstanceIdx)\n{\n");
@@ -3266,44 +3251,54 @@ void FHlslNiagaraTranslator::DefineMainGPUFunctions(
 	}
 	HlslOutput += TEXT("}\n");
 
-	FString SpawnLogicString;
-	
-	{
-		FExpressionPermutationContext PermutationContext(SpawnLogicString);
-		
-		if (TranslationStages.Num() > 1)
-		{
-			PermutationContext.AddBranch(*this, TranslationStages[0]);
-		}
+	//////////////////////////////////////////////////////////////////////////
+	// Generate common main body
+	HlslOutput += TEXT(
+			"\n\n/*\n"
+			"*	CS wrapper for our generated code; calls spawn and update functions on the corresponding instances in the buffer\n"
+			" */\n"
+			"\n"
+			"[numthreads(THREADGROUP_SIZE_X, THREADGROUP_SIZE_Y, THREADGROUP_SIZE_Z)]\n"
+			"void SimulateMainComputeCS(\n"
+			"	uint3 DispatchThreadId : SV_DispatchThreadID,\n"
+			"	uint3 GroupThreadId : SV_GroupThreadID)\n"
+			"{\n"
+			"	GLinearThreadId = DispatchThreadId.x + (DispatchThreadId.y * DispatchThreadIdToLinear.y);\n"
+			"	#if NIAGARA_DISPATCH_TYPE >= NIAGARA_DISPATCH_TYPE_THREE_D\n"
+			"		GLinearThreadId += DispatchThreadId.z * DispatchThreadIdToLinear.z;\n"
+			"	#endif\n"
+			"	GDispatchThreadId = DispatchThreadId;\n"
+			"	GGroupThreadId = GroupThreadId;\n"
+			"	GEmitterTickCounter = EmitterTickCounter;\n"
+			"	GRandomSeedOffset = 0;\n"
+			"\n"
+		);
 
-		SpawnLogicString += TEXT(
-				"		Context.MapSpawn.Particles.UniqueID = Engine_Emitter_TotalSpawnedParticles + ExecIndex(); \n"
-				"		ConditionalInterpolateParameters(Context); \n"
-				"		SimulateMapSpawn(Context); \n"
-				"		\n");
-	}
-	
 	//////////////////////////////////////////////////////////////////////////
-	// Generate per translation stage logic parts
-	// i.e. ones who vary based upon what we iterate over
-	FString SimRunSpawnUpdateLogicString;
-	FString SimSetupSpawnExecIndexLogicString;
+	// Generate each translation stages body
 	for (int32 i=1; i < TranslationStages.Num(); i++)
 	{
 		const FHlslNiagaraTranslationStage& TranslationStage = TranslationStages[i];
-		const TCHAR* IfOrElifText = (i == 1) ? TEXT("#if") : TEXT("#elif");
-
-		SimRunSpawnUpdateLogicString.Appendf(TEXT("%s SimulationStageIndex == %d	// %s\n"), IfOrElifText, TranslationStage.SimulationStageIndex, *TranslationStage.PassNamespace);
-		SimSetupSpawnExecIndexLogicString.Appendf(TEXT("%s SimulationStageIndex == %d	// %s\n"), IfOrElifText, TranslationStage.SimulationStageIndex, *TranslationStage.PassNamespace);
+		const bool bInterpolatedSpawning = CompileOptions.AdditionalDefines.Contains(TEXT("InterpolatedSpawn")) || i != 1;
+		const bool bAlwaysRunUpdateScript = CompileOptions.AdditionalDefines.Contains(TEXT("GpuAlwaysRunParticleUpdateScript"));
+		const bool bParticleIterationStage = TranslationStage.IterationSource.IsNone();
+		const bool bParticleSpawnStage = i == 1;
+
+		HlslOutput.Appendf(
+			TEXT("%s SimulationStageIndex == %d // %s\n"),
+			(i == 1) ? TEXT("#if") : TEXT("#elif"),
+			TranslationStage.SimulationStageIndex,
+			*TranslationStage.PassNamespace
+		);
 
 		// Particle iteration stage
-		if ( TranslationStage.IterationSource.IsNone() )
+		if (bParticleIterationStage)
 		{
 			// Do we have particle state iteration enable
-			SimRunSpawnUpdateLogicString += TEXT("	bool bRunSpawnUpdateLogic = true;\n");
-			if ( TranslationStage.bParticleIterationStateEnabled && TranslationStage.bWritesParticles )
+			HlslOutput.Append(TEXT("	bool bRunSpawnUpdateLogic = true;\n"));
+			if (TranslationStage.bParticleIterationStateEnabled && TranslationStage.bWritesParticles)
 			{
-				if ( TranslationStage.bPartialParticleUpdate == false )
+				if (TranslationStage.bPartialParticleUpdate == false)
 				{
 					FName StageName;
 					if (CompilationOutput.ScriptData.SimulationStageMetaData.IsValidIndex(TranslationStage.SimulationStageIndex))
@@ -3313,7 +3308,7 @@ void FHlslNiagaraTranslator::DefineMainGPUFunctions(
 
 					Error(FText::Format(LOCTEXT("ParticleIterationState_Invalid", "Simulation stage '{0}' is incompatible with particle state iteration due to killing particles or disabling particle updates."), FText::FromName(StageName)), nullptr, nullptr);
 				}
-				SimRunSpawnUpdateLogicString += TEXT(
+				HlslOutput += TEXT(
 					"	if ( ParticleIterationStateInfo.x != -1 )\n"
 					"	{\n"
 					"		int ParticleStateValue = InputDataInt(0, uint(ParticleIterationStateInfo.x), GLinearThreadId);\n"
@@ -3324,10 +3319,9 @@ void FHlslNiagaraTranslator::DefineMainGPUFunctions(
 
 			// We combine the update & spawn scripts together on GPU so we only need to check for spawning on the first translation stage
 			// Note: Depending on how spawning inside stages works we may need to enable the spawn logic for those stages *only*			
-			const bool bParticleSpawnStage = i == 1;
 			if (bParticleSpawnStage)
 			{
-				SimRunSpawnUpdateLogicString += TEXT(
+				HlslOutput += TEXT(
 					"	if (ReadInstanceCountOffset == 0xFFFFFFFF)\n"
 					"	{\n"
 					"		GSpawnStartInstance = 0;\n"
@@ -3343,14 +3337,12 @@ void FHlslNiagaraTranslator::DefineMainGPUFunctions(
 			}
 			else
 			{
-				SimRunSpawnUpdateLogicString += TEXT(
+				HlslOutput += TEXT(
 					"	GSpawnStartInstance = RWInstanceCounts[ReadInstanceCountOffset];\n"
 					"	const bool bRunUpdateLogic = bRunSpawnUpdateLogic && GLinearThreadId < GSpawnStartInstance;\n"
 					"	const bool bRunSpawnLogic = false;\n"
 				);
 			}
-
-			SimSetupSpawnExecIndexLogicString += TEXT("		SetupExecIndexAndSpawnInfoForGPU();\n");
 		}
 		// Iteration interface stage
 		else
@@ -3366,78 +3358,81 @@ void FHlslNiagaraTranslator::DefineMainGPUFunctions(
 
 			//-TODO: We can simplify the logic here with SimulationStage_GetInstanceCount() as only things that can provide an instance count offset
 			//		 can really be variable, everything else is driven from CPU code.
-			SimRunSpawnUpdateLogicString += TEXT(
+			HlslOutput += TEXT(
 				"	const uint MaxInstances = SimulationStage_GetInstanceCount();\n"
 				"	GLinearThreadId = all(DispatchThreadId < DispatchThreadIdBounds) ? GLinearThreadId : MaxInstances;\n"
 				"	GSpawnStartInstance = MaxInstances;\n"
 				"	const bool bRunUpdateLogic = (GLinearThreadId < GSpawnStartInstance) && (SimStart != 1);\n"
 				"	const bool bRunSpawnLogic = (GLinearThreadId < GSpawnStartInstance) && (SimStart == 1);\n"
 			);
-
-			SimSetupSpawnExecIndexLogicString += TEXT("		SetupExecIndexForGPU();\n");
 		}
-	}
-	SimRunSpawnUpdateLogicString.Append(TEXT("#endif\n"));
-	SimSetupSpawnExecIndexLogicString.Append(TEXT("#endif\n"));
 
-	//////////////////////////////////////////////////////////////////////////
-	// Generate main body
-	HlslOutput +=
-		TEXT("\n\n/*\n"
-			"*	CS wrapper for our generated code; calls spawn and update functions on the corresponding instances in the buffer\n"
-			" */\n"
-			"\n"
-			"[numthreads(THREADGROUP_SIZE_X, THREADGROUP_SIZE_Y, THREADGROUP_SIZE_Z)]\n"
-			"void SimulateMainComputeCS(\n"
-			"	uint3 DispatchThreadId : SV_DispatchThreadID,\n"
-			"	uint3 GroupThreadId : SV_GroupThreadID)\n"
-			"{\n"
-			"	GLinearThreadId = DispatchThreadId.x + (DispatchThreadId.y * DispatchThreadIdToLinear.y);\n"
-			"	#if NIAGARA_DISPATCH_TYPE >= NIAGARA_DISPATCH_TYPE_THREE_D\n"
-			"		GLinearThreadId += DispatchThreadId.z * DispatchThreadIdToLinear.z;\n"
-			"	#endif\n"
-			"	GDispatchThreadId = DispatchThreadId;\n"
-			"	GGroupThreadId = GroupThreadId;\n"
-			"	GEmitterTickCounter = EmitterTickCounter;\n"
-			"	GRandomSeedOffset = 0;\n"
-			"\n"
-			) + SimRunSpawnUpdateLogicString + TEXT(
+		HlslOutput += TEXT(
 			"	\n"
 			"	const float RandomSeedInitialisation = NiagaraInternalNoise(GLinearThreadId * 16384, 0 * 8196, (bRunUpdateLogic ? 4096 : 0) + EmitterTickCounter);	// initialise the random state seed\n"
 			"	\n"
 			"	FSimulationContext Context = (FSimulationContext)0;\n"
-			"	\n"
-			"	BRANCH\n"
-			"	if (bRunUpdateLogic)\n"
-			"	{\n"
-			"		SetupExecIndexForGPU();\n"
-			"		InitConstants(Context);\n"
-			"		LoadUpdateVariables(Context, GLinearThreadId);\n"
-			"		ReadDataSets(Context);\n"
-			"	}\n"
-			"	else if (bRunSpawnLogic)\n"
-			"	{\n"
-		) + SimSetupSpawnExecIndexLogicString + TEXT(
-			"		InitConstants(Context);\n"
-			"		InitSpawnVariables(Context);\n"
-			"		ReadDataSets(Context);\n"
-			"		\n") + SpawnLogicString
-			+ TEXT(
-			"		\n"
-			"		TransferAttributes(Context);\n"
-			"		\n"
-			"	}\n"
-			"\n"
-			"	if (bRunUpdateLogic || bRunSpawnLogic)\n"
-			"	{\n"
-		) + SimDoWorkString + TEXT(
-			"		WriteDataSets(Context);\n"
-			"	}\n"
-			"	\n"
-			"	StoreUpdateVariables(Context, bRunUpdateLogic || bRunSpawnLogic);\n\n"
 		);
 
-	HlslOutput += TEXT("}\n");
+		// Add Update Logic
+		HlslOutput.Append(TEXT("	BRANCH\n"));
+		HlslOutput.Append(TEXT("	if (bRunUpdateLogic)\n"));
+		HlslOutput.Append(TEXT("	{\n"));
+		HlslOutput.Append(TEXT("		SetupExecIndexForGPU();\n"));
+		HlslOutput.Append(TEXT("		InitConstants(Context);\n"));
+		HlslOutput.Append(TEXT("		LoadUpdateVariables(Context, GLinearThreadId);\n"));
+		HlslOutput.Append(TEXT("		ReadDataSets(Context);\n"));
+		if (bInterpolatedSpawning == false && bAlwaysRunUpdateScript == false)
+		{
+			HlslOutput.Appendf(TEXT("		Simulate%s(Context);\n"), *TranslationStages[i].PassNamespace);
+			HlslOutput.Append(TEXT("		WriteDataSets(Context);\n"));
+		}
+		HlslOutput.Append(TEXT("	}\n"));
+
+		// Add Spawn Logic
+		HlslOutput.Append(TEXT("	else if (bRunSpawnLogic)\n"));
+		HlslOutput.Append(TEXT("	{\n"));
+		if (bParticleIterationStage)
+		{
+			HlslOutput.Append(TEXT("		SetupExecIndexAndSpawnInfoForGPU();\n"));
+		}
+		else
+		{
+			HlslOutput.Append(TEXT("		SetupExecIndexForGPU();\n"));
+		}
+		HlslOutput.Append(TEXT("		InitConstants(Context);\n"));
+		HlslOutput.Append(TEXT("		InitSpawnVariables(Context);\n"));
+		HlslOutput.Append(TEXT("		ReadDataSets(Context);\n"));
+		if (bParticleSpawnStage == true)
+		{
+			HlslOutput.Append(TEXT("		Context.MapSpawn.Particles.UniqueID = Engine_Emitter_TotalSpawnedParticles + ExecIndex();\n"));
+			HlslOutput.Append(TEXT("		ConditionalInterpolateParameters(Context);\n"));
+			HlslOutput.Append(TEXT("		SimulateMapSpawn(Context);\n"));
+		}
+		HlslOutput.Append(TEXT("		TransferAttributes(Context);\n"));
+		if (bInterpolatedSpawning == false && bAlwaysRunUpdateScript == false)
+		{
+			HlslOutput.Append(TEXT("		WriteDataSets(Context);\n"));
+		}
+		HlslOutput.Append(TEXT("	}\n\n"));
+
+		// Interpolated spawning must also run the update logic if we have spawned
+		if (bInterpolatedSpawning == true || bAlwaysRunUpdateScript == true)
+		{
+			HlslOutput.Append(TEXT("	if (bRunUpdateLogic || bRunSpawnLogic)\n"));
+			HlslOutput.Append(TEXT("	{\n"));
+			HlslOutput.Appendf(TEXT("		Simulate%s(Context);\n"), *TranslationStages[i].PassNamespace);
+			HlslOutput.Append(TEXT("		WriteDataSets(Context);\n"));
+			HlslOutput.Append(TEXT("	}\n\n"));
+		}
+
+		// Store Data
+		HlslOutput.Append(TEXT("	StoreUpdateVariables(Context, bRunUpdateLogic || bRunSpawnLogic);\n\n"));
+	}
+
+	// End of logic
+	HlslOutput.Append(TEXT("#endif\n"));
+	HlslOutput.Append(TEXT("}\n"));
 }
 
 void FHlslNiagaraTranslator::DefineMain(FString &OutHlslOutput,
diff --git a/Engine/Plugins/FX/Niagara/Source/NiagaraEditor/Private/ViewModels/Stack/NiagaraStackEmitterSettingsGroup.cpp b/Engine/Plugins/FX/Niagara/Source/NiagaraEditor/Private/ViewModels/Stack/NiagaraStackEmitterSettingsGroup.cpp
--- a/Engine/Plugins/FX/Niagara/Source/NiagaraEditor/Private/ViewModels/Stack/NiagaraStackEmitterSettingsGroup.cpp
+++ b/Engine/Plugins/FX/Niagara/Source/NiagaraEditor/Private/ViewModels/Stack/NiagaraStackEmitterSettingsGroup.cpp
@@ -160,6 +160,44 @@ void UNiagaraStackEmitterPropertiesItem::RefreshIssues(TArray<FStackIssue>& NewI
 
 				NewIssues.Add(MissingRequiredFixedBoundsModuleError);
 			}
+
+			FStackIssue MissingRequiredFixedBoundsModuleError(
+				EStackIssueSeverity::Warning,
+				LOCTEXT("RequiredFixedBoundsWarningFormat", "The emitter is GPU and is using dynamic bounds mode.\r\nPlease update the Emitter or System properties otherwise bounds may be incorrect."),
+				LOCTEXT("MissingFixedBounds", "Missing fixed bounds."),
+				GetStackEditorDataKey(),
+				false);
+		}
+
+		if ( ActualEmitterData->SimTarget == ENiagaraSimTarget::GPUComputeSim && ActualEmitterData->bGpuAlwaysRunParticleUpdateScript )
+		{
+			const FText FixGpuInterpolatedDisabledText = LOCTEXT("FixGpuInterpolatedDisabled", "Fix GPU emitter to not run particle update script on particle spawn");
+			FVersionedNiagaraEmitterWeakPtr WeakEmitter = GetEmitterViewModel()->GetEmitter().ToWeakPtr();
+
+			NewIssues.Emplace(
+				EStackIssueSeverity::Warning,
+				LOCTEXT("GpuNoInterpolatedIncorrect", "GPU is incorrectly running particle update script on particle spawn."),
+				LOCTEXT("GpuNoInterpolatedIncorrectLong", "Due to a previous GPU codegen issue both particle spawn & update scripts were running on particle spawn when interpolated spawning is disabled.  This is now fixed for new content to match CPU behavior where only particle spawn will run when interpolated spawning is disabled."),
+				GetStackEditorDataKey(),
+				true,
+				FStackIssueFix(
+					FixGpuInterpolatedDisabledText,
+					FStackIssueFixDelegate::CreateLambda(
+						[WeakEmitter, FixGpuInterpolatedDisabledText]()
+						{
+							FVersionedNiagaraEmitter PinnedEmitter = WeakEmitter.ResolveWeakPtr();
+							if (PinnedEmitter.Emitter)
+							{
+								FScopedTransaction ScopedTransaction(FixGpuInterpolatedDisabledText);
+
+								PinnedEmitter.Emitter->Modify();
+								PinnedEmitter.GetEmitterData()->bGpuAlwaysRunParticleUpdateScript = false;
+								UNiagaraSystem::RequestCompileForEmitter(PinnedEmitter);
+							}
+						}
+					)
+				)
+			);
 		}
 
 		UNiagaraSystem& System = GetSystemViewModel()->GetSystem();
