Niagara - Merging 20492410 - Fix a bug in the merge manager introduced from my last change which breaks merging of linked value handles.

diff --git a/Engine/Plugins/FX/Niagara/Source/NiagaraEditor/Private/NiagaraScriptMergeManager.cpp b/Engine/Plugins/FX/Niagara/Source/NiagaraEditor/Private/NiagaraScriptMergeManager.cpp
--- a/Engine/Plugins/FX/Niagara/Source/NiagaraEditor/Private/NiagaraScriptMergeManager.cpp
+++ b/Engine/Plugins/FX/Niagara/Source/NiagaraEditor/Private/NiagaraScriptMergeManager.cpp
@@ -2923,7 +2923,7 @@ FNiagaraScriptMergeManager::FApplyDiffResults FNiagaraScriptMergeManager::AddInp
 						}
 					}
 				}
-				else
+				if (NewLinkedValueHandle.IsValid() == false)
 				{
 					NewLinkedValueHandle = OldLinkedValueHandle;
 				}
