Move new niagara version in to UE5Main version instead of creating a new custom version
#rnx

diff --git a/Engine/Plugins/FX/Niagara/Source/Niagara/Private/NiagaraEmitter.cpp b/Engine/Plugins/FX/Niagara/Source/Niagara/Private/NiagaraEmitter.cpp
--- a/Engine/Plugins/FX/Niagara/Source/Niagara/Private/NiagaraEmitter.cpp
+++ b/Engine/Plugins/FX/Niagara/Source/Niagara/Private/NiagaraEmitter.cpp
@@ -5,7 +5,7 @@
 
 #include "INiagaraEditorOnlyDataUtlities.h"
 #include "NiagaraCustomVersion.h"
-#include "NiagaraCustomVersionUEMain.h"
+#include "UObject/UE5MainStreamObjectVersion.h"
 #include "NiagaraEditorDataBase.h"
 #include "NiagaraModule.h"
 #include "NiagaraRenderer.h"
@@ -518,7 +518,7 @@ void UNiagaraEmitter::Serialize(FArchive& Ar)
 #endif
 
 	Ar.UsingCustomVersion(FNiagaraCustomVersion::GUID);
-	Ar.UsingCustomVersion(FNiagaraCustomVersionUEMain::GUID);
+	Ar.UsingCustomVersion(FUE5MainStreamObjectVersion::GUID);
 }
 
 void FVersionedNiagaraEmitterData::EnsureScriptsPostLoaded()
@@ -619,7 +619,7 @@ void UNiagaraEmitter::PostLoad()
 	}
 #endif
 
-	const int32 NiagaraVerUE5Main = GetLinkerCustomVersion(FNiagaraCustomVersionUEMain::GUID);
+	const int32 UE5MainVer = GetLinkerCustomVersion(FUE5MainStreamObjectVersion::GUID);
 
 	for (FVersionedNiagaraEmitterData& Data : VersionData)
 	{
@@ -627,7 +627,7 @@ void UNiagaraEmitter::PostLoad()
 		Data.PostLoad(*this, IsCooked);
 		Data.GPUComputeScript->OnGPUScriptCompiled().RemoveAll(this);
 		Data.GPUComputeScript->OnGPUScriptCompiled().AddUObject(this, &UNiagaraEmitter::RaiseOnEmitterGPUCompiled);
-		if (NiagaraVerUE5Main < FNiagaraCustomVersionUEMain::FixGpuAlwaysRunningUpdateScriptNoneInterpolated)
+		if (UE5MainVer < FUE5MainStreamObjectVersion::FixGpuAlwaysRunningUpdateScriptNoneInterpolated)
 		{
 			if ( Data.SimTarget == ENiagaraSimTarget::GPUComputeSim )
 			{
diff --git a/Engine/Plugins/FX/Niagara/Source/NiagaraCore/Private/NiagaraCustomVersionUEMain.cpp b/Engine/Plugins/FX/Niagara/Source/NiagaraCore/Private/NiagaraCustomVersionUEMain.cpp
deleted file mode 100644
diff --git a/Engine/Plugins/FX/Niagara/Source/NiagaraCore/Public/NiagaraCustomVersionUEMain.h b/Engine/Plugins/FX/Niagara/Source/NiagaraCore/Public/NiagaraCustomVersionUEMain.h
deleted file mode 100644
diff --git a/Engine/Source/Runtime/Core/Public/UObject/UE5MainStreamObjectVersion.h b/Engine/Source/Runtime/Core/Public/UObject/UE5MainStreamObjectVersion.h
--- a/Engine/Source/Runtime/Core/Public/UObject/UE5MainStreamObjectVersion.h
+++ b/Engine/Source/Runtime/Core/Public/UObject/UE5MainStreamObjectVersion.h
@@ -230,6 +230,9 @@ struct CORE_API FUE5MainStreamObjectVersion
 		// Moving all nodal offset handling to Lens Component
 		LensComponentNodalOffset,
 
+		// GPU none interpolated spawning no longer calls the update script
+		FixGpuAlwaysRunningUpdateScriptNoneInterpolated,
+
 		// -----<new versions can be added above this line>-------------------------------------------------
 		VersionPlusOne,
 		LatestVersion = VersionPlusOne - 1
