Remove toasts from execution checks on Copy and Paste. The UI will call these functions continually when displaying the options menu, and we have no way to determine if the call came from a keyboard input or the UI.

diff --git a/Engine/Plugins/FX/Niagara/Source/NiagaraEditor/Private/ViewModels/NiagaraOverviewGraphViewModel.cpp b/Engine/Plugins/FX/Niagara/Source/NiagaraEditor/Private/ViewModels/NiagaraOverviewGraphViewModel.cpp
--- a/Engine/Plugins/FX/Niagara/Source/NiagaraEditor/Private/ViewModels/NiagaraOverviewGraphViewModel.cpp
+++ b/Engine/Plugins/FX/Niagara/Source/NiagaraEditor/Private/ViewModels/NiagaraOverviewGraphViewModel.cpp
@@ -294,8 +294,6 @@ bool FNiagaraOverviewGraphViewModel::CanCopyNodes() const
 	// We don't store enough information on the node to find the emitter asset in question, so adding the emitter with the correct hierarchy would require a larger refactor to properly support.
 	if (GetSystemViewModel()->GetEditMode() == ENiagaraSystemViewModelEditMode::EmitterAsset)
 	{		
-		// Instead direct the user to add the emitter through the right-click menu, where the behavior is well defined.
-		FNiagaraEditorUtilities::InfoWithToastAndLog(LOCTEXT("CopyFromEmitterAsset", "Cannot copy emitters from emitter assets. Please use Add Emitter from the right click menu instead."));
 		return false;
 	}
 
@@ -374,7 +372,6 @@ bool FNiagaraOverviewGraphViewModel::CanPasteNodes() const
 	// Cannot paste into emitter asset
 	if (GetSystemViewModel()->GetEditMode() == ENiagaraSystemViewModelEditMode::EmitterAsset)
 	{
-		FNiagaraEditorUtilities::InfoWithToastAndLog(LOCTEXT("PasteToEmitterAsset", "Cannot paste nodes into emitter assets."));
 		return false;
 	}
 
