- Fix for holding onto task references longer than necessary

diff --git a/Engine/Plugins/FX/Niagara/Source/Niagara/Private/NiagaraSystem.cpp b/Engine/Plugins/FX/Niagara/Source/Niagara/Private/NiagaraSystem.cpp
--- a/Engine/Plugins/FX/Niagara/Source/Niagara/Private/NiagaraSystem.cpp
+++ b/Engine/Plugins/FX/Niagara/Source/Niagara/Private/NiagaraSystem.cpp
@@ -87,6 +87,38 @@ static FAutoConsoleVariableRef CVarNiagaraCompileWaitLoggingTerminationCap(
 
 //////////////////////////////////////////////////////////////////////////
 
+class FNiagaraOptimizeCompleteTask
+{
+public:
+	FNiagaraOptimizeCompleteTask(UNiagaraSystem* Owner, FGraphEventRef* InRefToClear)
+		: WeakOwner(Owner)
+		, RefToClear(InRefToClear)
+		, ReferenceValue(InRefToClear != nullptr ? InRefToClear->GetReference() : nullptr)
+	{
+	}
+
+	FORCEINLINE TStatId GetStatId() const { RETURN_QUICK_DECLARE_CYCLE_STAT(FNiagaraOptimizeCompleteTask, STATGROUP_TaskGraphTasks); }
+	ENamedThreads::Type GetDesiredThread() { return ENamedThreads::GameThread; }
+	static ESubsequentsMode::Type GetSubsequentsMode() { return ESubsequentsMode::TrackSubsequents; }
+
+	void DoTask(ENamedThreads::Type CurrentThread, const FGraphEventRef& MyCompletionGraphEvent)
+	{
+		if (UNiagaraSystem* Owner = WeakOwner.Get())
+		{
+			if (RefToClear && (RefToClear->GetReference() == ReferenceValue) )
+			{
+				*RefToClear = nullptr;
+			}
+		}
+	}
+
+	TWeakObjectPtr<UNiagaraSystem> WeakOwner;
+	FGraphEventRef* RefToClear = nullptr;
+	FGraphEvent* ReferenceValue = nullptr;
+};
+
+//////////////////////////////////////////////////////////////////////////
+
 UNiagaraSystem::UNiagaraSystem(const FObjectInitializer& ObjectInitializer)
 : Super(ObjectInitializer)
 #if WITH_EDITORONLY_DATA
@@ -3332,6 +3364,10 @@ void UNiagaraSystem::AsyncOptimizeAllScripts()
 	{
 		DECLARE_CYCLE_STAT(TEXT("FNullGraphTask.NiagaraScriptOptimizationCompletion"), STAT_FNullGraphTask_NiagaraScriptOptimizationCompletion, STATGROUP_TaskGraphTasks);
 		ScriptOptimizationCompletionEvent = TGraphTask<FNullGraphTask>::CreateTask(&Prereqs, ENamedThreads::GameThread).ConstructAndDispatchWhenReady(GET_STATID(STAT_FNullGraphTask_NiagaraScriptOptimizationCompletion), ENamedThreads::AnyThread);
+
+		// Inject task to clear out the reference to the graph task
+		FGraphEventArray CompleteTaskPrereqs({ ScriptOptimizationCompletionEvent });
+		TGraphTask<FNiagaraOptimizeCompleteTask>::CreateTask(&CompleteTaskPrereqs, ENamedThreads::GameThread).ConstructAndDispatchWhenReady(this, &ScriptOptimizationCompletionEvent);
 	}
 
 	bNeedsAsyncOptimize = false;
diff --git a/Engine/Plugins/FX/Niagara/Source/Niagara/Private/NiagaraSystemSimulation.cpp b/Engine/Plugins/FX/Niagara/Source/Niagara/Private/NiagaraSystemSimulation.cpp
--- a/Engine/Plugins/FX/Niagara/Source/Niagara/Private/NiagaraSystemSimulation.cpp
+++ b/Engine/Plugins/FX/Niagara/Source/Niagara/Private/NiagaraSystemSimulation.cpp
@@ -417,6 +417,7 @@ struct FNiagaraSystemSimulationSpawnConcurrentTask
 
 		Context.BeforeInstancesTickGraphEvents.Add(MyCompletionGraphEvent);
 		Context.Owner->Spawn_Concurrent(Context);
+		Context.BeforeInstancesTickGraphEvents.Empty();
 		CompletionTask->Unlock();
 	}
 
@@ -1577,7 +1578,7 @@ void FNiagaraSystemSimulation::DumpStalledInfo()
 {
 	TStringBuilder<128> Builder;
 	Builder.Appendf(TEXT("ConcurrentTickGraphEvent Complete (%d)\n"), ConcurrentTickGraphEvent ? ConcurrentTickGraphEvent->IsComplete() : true);
-	Builder.Appendf(TEXT("AllWorkCompleteGraphEvent Complete (%d)\n"), AllWorkCompleteGraphEvent ? ConcurrentTickGraphEvent->IsComplete() : true);
+	Builder.Appendf(TEXT("AllWorkCompleteGraphEvent Complete (%d)\n"), AllWorkCompleteGraphEvent ? AllWorkCompleteGraphEvent->IsComplete() : true);
 	Builder.Appendf(TEXT("SystemTickGroup (%d)\n"), SystemTickGroup);
 	Builder.Appendf(TEXT("bInSpawnPhase (%d)\n"), bInSpawnPhase);
 	Builder.Appendf(TEXT("bIsSolo (%d)\n"), bIsSolo);
@@ -1649,6 +1650,8 @@ void FNiagaraSystemSimulation::WaitForInstancesTickComplete(bool bEnsureComplete
 			}
 		}
 	}
+
+	ConcurrentTickGraphEvent = nullptr;
 	AllWorkCompleteGraphEvent = nullptr;
 }
 
@@ -2097,6 +2100,7 @@ void FNiagaraSystemSimulation::RemoveInstance(FNiagaraSystemInstance* Instance)
 	}
 
 	// We did not finalize and will not so clear the reference
+	Instance->ConcurrentTickGraphEvent = nullptr;
 	Instance->FinalizeRef.ConditionalClear();
 
 	// Remove the instance if it is still valid
