- Add Matrix User Parameters Setters
- This truncates to FMatrix44f so will loose precision for LWC

diff --git a/Engine/Plugins/FX/Niagara/Source/Niagara/Private/NiagaraComponent.cpp b/Engine/Plugins/FX/Niagara/Source/Niagara/Private/NiagaraComponent.cpp
--- a/Engine/Plugins/FX/Niagara/Source/Niagara/Private/NiagaraComponent.cpp
+++ b/Engine/Plugins/FX/Niagara/Source/Niagara/Private/NiagaraComponent.cpp
@@ -2136,6 +2136,28 @@ void UNiagaraComponent::SetNiagaraVariableQuat(const FString& InVariableName, co
 	SetVariableQuat(FName(*InVariableName), InValue);
 }
 
+void UNiagaraComponent::SetNiagaraVariableMatrix(const FString& InVariableName, const FMatrix& InValue)
+{
+	SetVariableMatrix(FName(*InVariableName), InValue);
+}
+
+void UNiagaraComponent::SetVariableMatrix(FName InVariableName, const FMatrix& InValue)
+{
+	FMatrix44f AsFloat = (FMatrix44f)InValue;	//-TODO: LWC Precision loss
+	const FNiagaraVariable VariableDesc(FNiagaraTypeDefinition::GetMatrix4Def(), InVariableName);
+	if (SystemInstanceController.IsValid())
+	{
+		SystemInstanceController->SetVariable_Deferred(InVariableName, AsFloat);
+	}
+	else
+	{
+		OverrideParameters.SetParameterValue(AsFloat, VariableDesc, true);
+	}
+#if WITH_EDITOR
+	SetParameterOverride(VariableDesc, FNiagaraVariant(&AsFloat, sizeof(FMatrix44f)));
+#endif
+}
+
 void UNiagaraComponent::SetVariableVec4(FName InVariableName, const FVector4& InValue)
 {
 	FVector4f AsFloat = (FVector4f)InValue;
diff --git a/Engine/Plugins/FX/Niagara/Source/Niagara/Private/NiagaraSystemInstanceController.cpp b/Engine/Plugins/FX/Niagara/Source/Niagara/Private/NiagaraSystemInstanceController.cpp
--- a/Engine/Plugins/FX/Niagara/Source/Niagara/Private/NiagaraSystemInstanceController.cpp
+++ b/Engine/Plugins/FX/Niagara/Source/Niagara/Private/NiagaraSystemInstanceController.cpp
@@ -426,6 +426,15 @@ void FNiagaraSystemInstanceController::SetVariable(FName InVariableName, FQuat4f
 	}
 }
 
+void FNiagaraSystemInstanceController::SetVariable(FName InVariableName, const FMatrix44f& InValue)
+{
+	if (OverrideParameters)
+	{
+		const FNiagaraVariable VariableDesc(FNiagaraTypeDefinition::GetMatrix4Def(), InVariableName);
+		OverrideParameters->SetParameterValue(InValue, VariableDesc, true);
+	}
+}
+
 void FNiagaraSystemInstanceController::SetVariable(FName InVariableName, TWeakObjectPtr<UObject> InValue)
 {
 	if (OverrideParameters)
diff --git a/Engine/Plugins/FX/Niagara/Source/Niagara/Public/NiagaraComponent.h b/Engine/Plugins/FX/Niagara/Source/Niagara/Public/NiagaraComponent.h
--- a/Engine/Plugins/FX/Niagara/Source/Niagara/Public/NiagaraComponent.h
+++ b/Engine/Plugins/FX/Niagara/Source/Niagara/Public/NiagaraComponent.h
@@ -409,7 +409,6 @@ public:
 	UFUNCTION(BlueprintCallable, Category = Niagara, meta = (DisplayName = "Set Niagara Variable (Vector4)"))
 	void SetVariableVec4(FName InVariableName, const FVector4& InValue);
 
-
 	/** Sets a Niagara Vector3 parameter by name, overriding locally if necessary.*/
 	UFUNCTION(BlueprintCallable, Category = Niagara, meta = (DisplayName = "Set Niagara Variable By String (Quaternion)"))
 	void SetNiagaraVariableQuat(const FString& InVariableName, const FQuat& InValue);
@@ -418,6 +417,14 @@ public:
 	UFUNCTION(BlueprintCallable, Category = Niagara, meta = (DisplayName = "Set Niagara Variable (Quaternion)"))
 	void SetVariableQuat(FName InVariableName, const FQuat& InValue);
 
+	/** Sets a Niagara Vector3 parameter by name, overriding locally if necessary.*/
+	UFUNCTION(BlueprintCallable, Category = Niagara, meta = (DisplayName = "Set Niagara Variable By String (Matrix)"))
+	void SetNiagaraVariableMatrix(const FString& InVariableName, const FMatrix& InValue);
+
+	/** Sets a Niagara Vector3 parameter by name, overriding locally if necessary.*/
+	UFUNCTION(BlueprintCallable, Category = Niagara, meta = (DisplayName = "Set Niagara Variable (Matrix)"))
+	void SetVariableMatrix(FName InVariableName, const FMatrix& InValue);
+
 	/** Sets a Niagara Vector3 parameter by name, overriding locally if necessary.*/
 	UFUNCTION(BlueprintCallable, Category = Niagara, meta = (DisplayName = "Set Niagara Variable By String (Vector3)"))
 	void SetNiagaraVariableVec3(const FString& InVariableName, FVector InValue);
diff --git a/Engine/Plugins/FX/Niagara/Source/Niagara/Public/NiagaraSystemInstanceController.h b/Engine/Plugins/FX/Niagara/Source/Niagara/Public/NiagaraSystemInstanceController.h
--- a/Engine/Plugins/FX/Niagara/Source/Niagara/Public/NiagaraSystemInstanceController.h
+++ b/Engine/Plugins/FX/Niagara/Source/Niagara/Public/NiagaraSystemInstanceController.h
@@ -160,6 +160,7 @@ private:
 	void SetVariable(FName InVariableName, FVector4f InValue);
 	void SetVariable(FName InVariableName, FLinearColor InValue);
 	void SetVariable(FName InVariableName, FQuat4f InValue);
+	void SetVariable(FName InVariableName, const FMatrix44f& InValue);
 	void SetVariable(FName InVariableName, TWeakObjectPtr<UObject> Object);
 	void SetVariable(FName InVariableName, TWeakObjectPtr<UMaterialInterface> Object);
 	void SetVariable(FName InVariableName, TWeakObjectPtr<UStaticMesh> Object);
