- Add console variable to allow CPU side data buffer to be released on reset
- This reduces peak memory usage of particle data

diff --git a/Engine/Plugins/FX/Niagara/Source/Niagara/Classes/NiagaraDataSet.h b/Engine/Plugins/FX/Niagara/Source/Niagara/Classes/NiagaraDataSet.h
--- a/Engine/Plugins/FX/Niagara/Source/Niagara/Classes/NiagaraDataSet.h
+++ b/Engine/Plugins/FX/Niagara/Source/Niagara/Classes/NiagaraDataSet.h
@@ -111,6 +111,7 @@ protected:
 public:
 	FNiagaraDataBuffer(FNiagaraDataSet* InOwner);
 	void Allocate(uint32 NumInstances, bool bMaintainExisting = false);
+	void ReleaseCPU();
 
 	void AllocateGPU(FRHICommandList& RHICmdList, uint32 InNumInstances, ERHIFeatureLevel::Type FeatureLevel, const TCHAR* DebugSimName);
 	void SwapGPU(FNiagaraDataBuffer* BufferToSwap);
diff --git a/Engine/Plugins/FX/Niagara/Source/Niagara/Private/NiagaraDataSet.cpp b/Engine/Plugins/FX/Niagara/Source/Niagara/Private/NiagaraDataSet.cpp
--- a/Engine/Plugins/FX/Niagara/Source/Niagara/Private/NiagaraDataSet.cpp
+++ b/Engine/Plugins/FX/Niagara/Source/Niagara/Private/NiagaraDataSet.cpp
@@ -93,6 +93,14 @@ static FAutoConsoleVariableRef CVarNiagaraGPUDataWarningSize(
 	ECVF_Scalability
 );
 
+static bool GNiagaraReleaseBuffersOnReset = true;
+static FAutoConsoleVariableRef CVarNiagaraReleaseBuffersOnReset(
+	TEXT("fx.NiagaraReleaseBuffersOnReset"),
+	GNiagaraReleaseBuffersOnReset,
+	TEXT("Will release all memory associated with data buffers when the dataset is reset."),
+	ECVF_Scalability
+);
+
 void NiagaraWarnGpuBufferSize(const int32 ByteSize, const TCHAR* DebugSimName)
 {
 	static int32 MaxEncountered = 0;
@@ -169,6 +177,14 @@ void FNiagaraDataSet::ResetBuffers()
 	
 	if (GetSimTarget() == ENiagaraSimTarget::CPUSim)
 	{
+		if (GNiagaraReleaseBuffersOnReset)
+		{
+			for (FNiagaraDataBuffer* Buffer : Data)
+			{
+				Buffer->ReleaseCPU();
+			}
+		}
+
 		ResetBuffersInternal();
 	}
 	else
@@ -747,6 +763,26 @@ void FNiagaraDataBuffer::Allocate(uint32 InNumInstances, bool bMaintainExisting)
 	}
 }
 
+void FNiagaraDataBuffer::ReleaseCPU()
+{
+	NumInstances = 0;
+	NumInstancesAllocated = 0;
+	IDToIndexTable.Reset();
+
+	FloatStride = 0;
+	Int32Stride = 0;
+	HalfStride = 0;
+
+	DEC_MEMORY_STAT_BY(STAT_NiagaraParticleMemory, FloatData.GetAllocatedSize() + Int32Data.GetAllocatedSize() + HalfData.GetAllocatedSize());
+#if NIAGARA_MEMORY_TRACKING
+	Owner->BufferSizeBytes -= FloatData.GetAllocatedSize() + Int32Data.GetAllocatedSize() + HalfData.GetAllocatedSize();
+#endif
+
+	FloatData.Empty();
+	Int32Data.Empty();
+	HalfData.Empty();
+}
+
 void FNiagaraDataBuffer::AllocateGPU(FRHICommandList& RHICmdList, uint32 InNumInstances, ERHIFeatureLevel::Type FeatureLevel, const TCHAR* DebugSimName)
 {
 	static constexpr EBufferUsageFlags GPUBufferFlags = BUF_Static | BUF_SourceCopy;
