- Handle edge case where we have valid emitter list but Emitters on instance are empty because we forced deactivate before InitEmitters was called

#rnx

diff --git a/Engine/Plugins/FX/Niagara/Source/Niagara/Private/NiagaraSystemRenderData.cpp b/Engine/Plugins/FX/Niagara/Source/Niagara/Private/NiagaraSystemRenderData.cpp
--- a/Engine/Plugins/FX/Niagara/Source/Niagara/Private/NiagaraSystemRenderData.cpp
+++ b/Engine/Plugins/FX/Niagara/Source/Niagara/Private/NiagaraSystemRenderData.cpp
@@ -333,22 +333,29 @@ void FNiagaraSystemRenderData::PostTickRenderers(const FNiagaraSystemInstance& S
 
 void FNiagaraSystemRenderData::OnSystemComplete(const FNiagaraSystemInstance& SystemInstance)
 {
-	// Give renderers a chance to handle completion
-	if (EmitterRenderers_GT.Num() > 0)
+	if (EmitterRenderers_GT.Num() == 0)
 	{
-		const UNiagaraSystem* System = SystemInstance.GetSystem();
-		for (const FNiagaraRendererExecutionIndex& ExecIdx : System->GetRendererCompletionOrder())
+		return;
+	}
+
+	const UNiagaraSystem* System = SystemInstance.GetSystem();
+	const TArray<TSharedRef<FNiagaraEmitterInstance, ESPMode::ThreadSafe>>& EmitterInstances = SystemInstance.GetEmitters();
+	if (EmitterInstances.Num() == 0)
+	{
+		return;
+	}
+
+	for (const FNiagaraRendererExecutionIndex& ExecIdx : System->GetRendererCompletionOrder())
+	{
+		if (EmitterRenderers_GT.IsValidIndex(ExecIdx.SystemRendererIndex) && ensure(EmitterInstances.IsValidIndex(ExecIdx.SystemRendererIndex)) )
 		{
-			if (EmitterRenderers_GT.IsValidIndex(ExecIdx.SystemRendererIndex))
+			const FNiagaraEmitterInstance& EmitterInstance = EmitterInstances[ExecIdx.SystemRendererIndex].Get();
+			FVersionedNiagaraEmitterData* EmitterData = EmitterInstance.GetCachedEmitterData();
+			FNiagaraRenderer* EmitterRenderer = EmitterRenderers_GT[ExecIdx.SystemRendererIndex];
+			if (EmitterData && EmitterRenderer)
 			{
-				const FNiagaraEmitterInstance& EmitterInst = SystemInstance.GetEmitters()[ExecIdx.EmitterIndex].Get();
-				FVersionedNiagaraEmitterData* EmitterData = EmitterInst.GetCachedEmitterData();
-				FNiagaraRenderer* EmitterRenderer = EmitterRenderers_GT[ExecIdx.SystemRendererIndex];
-				if (EmitterData && EmitterRenderer)
-				{
-					const UNiagaraRendererProperties* RendererProperties = EmitterData->GetRenderers()[ExecIdx.EmitterRendererIndex];
-					EmitterRenderer->OnSystemComplete_GameThread(RendererProperties, &EmitterInst);
-				}
+				const UNiagaraRendererProperties* RendererProperties = EmitterData->GetRenderers()[ExecIdx.EmitterRendererIndex];
+				EmitterRenderer->OnSystemComplete_GameThread(RendererProperties, &EmitterInstance);
 			}
 		}
 	}
