read/write 16 bit float grids for openvdb caches

diff --git a/Engine/Plugins/FX/Niagara/Source/Niagara/Classes/NiagaraOpenVDB.h b/Engine/Plugins/FX/Niagara/Source/Niagara/Classes/NiagaraOpenVDB.h
--- a/Engine/Plugins/FX/Niagara/Source/Niagara/Classes/NiagaraOpenVDB.h
+++ b/Engine/Plugins/FX/Niagara/Source/Niagara/Classes/NiagaraOpenVDB.h
@@ -47,10 +47,10 @@ using Vec4HTree = openvdb::tree::Tree4<Vec4h, 5, 4, 3>::Type;
 using Vec4HGrid = openvdb::Grid<Vec4HTree>;
 using Vec4HDense = openvdb::tools::Dense<Vec4h, openvdb::tools::MemoryLayout::LayoutXYZ>;
 
-using Vec4 = Vec4s;
-using Vec4Tree = Vec4STree;
-using Vec4Grid = Vec4SGrid;
-using Vec4Dense = Vec4SDense;
+using Vec4 = Vec4h;
+using Vec4Tree = openvdb::tree::Tree4<Vec4, 5, 4, 3>::Type;
+using Vec4Grid = openvdb::Grid<Vec4Tree>;
+using Vec4Dense = openvdb::tools::Dense<Vec4, openvdb::tools::MemoryLayout::LayoutXYZ>;
 using Vec4GridPtr = Vec4Grid::Ptr;
 using Vec4DensePtr = Vec4Dense::Ptr;
 
diff --git a/Engine/Plugins/FX/Niagara/Source/Niagara/Private/NiagaraDataInterfaceVolumeCache.cpp b/Engine/Plugins/FX/Niagara/Source/Niagara/Private/NiagaraDataInterfaceVolumeCache.cpp
--- a/Engine/Plugins/FX/Niagara/Source/Niagara/Private/NiagaraDataInterfaceVolumeCache.cpp
+++ b/Engine/Plugins/FX/Niagara/Source/Niagara/Private/NiagaraDataInterfaceVolumeCache.cpp
@@ -292,7 +292,8 @@ bool UNiagaraDataInterfaceVolumeCache::PerInstanceTick(void* PerInstanceData, FN
 	{
 		FNiagaraDataInterfaceVolumeCacheProxy* TextureProxy = GetProxyAs<FNiagaraDataInterfaceVolumeCacheProxy>();
 
-		EPixelFormat Format = PF_A32B32G32R32F;
+		//EPixelFormat Format = PF_A32B32G32R32F;
+		EPixelFormat Format = PF_FloatRGBA;
 		const int32 FormatSize = GPixelFormats[Format].BlockBytes;
 
 		// cannot read from cache...spew errors or let it go?
diff --git a/Engine/Plugins/FX/Niagara/Source/Niagara/Private/VolumeCache.cpp b/Engine/Plugins/FX/Niagara/Source/Niagara/Private/VolumeCache.cpp
--- a/Engine/Plugins/FX/Niagara/Source/Niagara/Private/VolumeCache.cpp
+++ b/Engine/Plugins/FX/Niagara/Source/Niagara/Private/VolumeCache.cpp
@@ -36,8 +36,8 @@ public:
 	virtual bool Fill3DTexture_RenderThread(int frame, FTextureRHIRef TextureToFill, FRHICommandListImmediate& RHICmdList);
 	
 private:
-	TMap<int32, Vec4SGrid::Ptr> OpenVDBGrids;
-	openvdb::tools::Dense<openvdb::Vec4s, openvdb::tools::MemoryLayout::LayoutXYZ>::Ptr DenseGridPtr;
+	TMap<int32, Vec4Grid::Ptr> OpenVDBGrids;
+	Vec4Dense::Ptr DenseGridPtr;
 };
 #endif
 
@@ -251,8 +251,14 @@ bool FOpenVDBCacheData::Fill3DTexture_RenderThread(int frame, FTextureRHIRef Tex
 		}
 		else
 		{
-			FUpdateTextureRegion3D UpdateRegion(0, 0, 0, 0, 0, 0, DenseResolution.X, DenseResolution.Y, DenseResolution.Z);			
-			RHICmdList.UpdateTexture3D(TextureToFill, 0, UpdateRegion, FormatSize * DenseResolution.X, FormatSize * DenseResolution.X * DenseResolution.Y, DataPtr);			
+			const FUpdateTextureRegion3D UpdateRegion(0, 0, 0, 0, 0, 0, DenseResolution.X, DenseResolution.Y, DenseResolution.Z);			
+			const SIZE_T MemorySize = static_cast<SIZE_T>(UpdateRegion.Width * UpdateRegion.Height * UpdateRegion.Depth * FormatSize);
+
+			FUpdateTexture3DData TheData = RHICmdList.BeginUpdateTexture3D(TextureToFill, 0, UpdateRegion);
+			
+			FMemory::Memcpy(TheData.Data, DataPtr, MemorySize);
+
+			RHICmdList.EndUpdateTexture3D(TheData);						
 		}
 
 		return true;
