Change the behavior of non-auto-activated Niagara components to stay off when changing user parameters

diff --git a/Engine/Plugins/FX/Niagara/Source/Niagara/Private/NiagaraComponent.cpp b/Engine/Plugins/FX/Niagara/Source/Niagara/Private/NiagaraComponent.cpp
--- a/Engine/Plugins/FX/Niagara/Source/Niagara/Private/NiagaraComponent.cpp
+++ b/Engine/Plugins/FX/Niagara/Source/Niagara/Private/NiagaraComponent.cpp
@@ -3153,14 +3153,17 @@ void UNiagaraComponent::AssetExposedParametersChanged()
 		// In that case we skip the sync with the system, as that happens in PostLoad anyways.
 		return;
 	}
+	bool bCurrentlyActive = IsActive();
 	SynchronizeWithSourceSystem();
 #if WITH_EDITOR
-	//-TODO: We can possibly remove the ReinitializeSystem here after more testing.
-	//       This fixes an issue with deleting a renderer and hitting undo in the editor.
 	if ( !GIsTransacting )
 #endif
 	{
-		ReinitializeSystem();
+		if (bCurrentlyActive)
+		{
+			// only reactivate systems that were running before
+			ReinitializeSystem();
+		}
 	}
 }
 
