Fix static analysis warning.

diff --git a/Engine/Plugins/FX/Niagara/Source/Niagara/Private/VolumeCache.cpp b/Engine/Plugins/FX/Niagara/Source/Niagara/Private/VolumeCache.cpp
--- a/Engine/Plugins/FX/Niagara/Source/Niagara/Private/VolumeCache.cpp
+++ b/Engine/Plugins/FX/Niagara/Source/Niagara/Private/VolumeCache.cpp
@@ -252,7 +252,7 @@ bool FOpenVDBCacheData::Fill3DTexture_RenderThread(int frame, FTextureRHIRef Tex
 		else
 		{
 			const FUpdateTextureRegion3D UpdateRegion(0, 0, 0, 0, 0, 0, DenseResolution.X, DenseResolution.Y, DenseResolution.Z);			
-			const SIZE_T MemorySize = static_cast<SIZE_T>(UpdateRegion.Width * UpdateRegion.Height * UpdateRegion.Depth * FormatSize);
+			const SIZE_T MemorySize = static_cast<SIZE_T>(UpdateRegion.Width) * static_cast<SIZE_T>(UpdateRegion.Height) * static_cast<SIZE_T>(UpdateRegion.Depth) * static_cast<SIZE_T>(FormatSize);
 
 			FUpdateTexture3DData TheData = RHICmdList.BeginUpdateTexture3D(TextureToFill, 0, UpdateRegion);
 			
