GPU Ribbons!
Heavily refactored ribbon rendering to improve maintainability and support improvements.
Implemented GPU sim support as well as GPU initialized ribbons for CPU systems.
Made some of tessellation initialization code optional when tessellation is disabled or using a constant factor to improve performance of CPU initialization.
Various other improvements/fixes to ribbons

diff --git a/Engine/Plugins/FX/Niagara/Shaders/Private/NiagaraRibbonVertexFactory.ush b/Engine/Plugins/FX/Niagara/Shaders/Private/NiagaraRibbonVertexFactory.ush
--- a/Engine/Plugins/FX/Niagara/Shaders/Private/NiagaraRibbonVertexFactory.ush
+++ b/Engine/Plugins/FX/Niagara/Shaders/Private/NiagaraRibbonVertexFactory.ush
@@ -24,6 +24,9 @@
 #define SLICEVERTEXDATA_TEXTUREV_OFFSET 4
 #define SLICEVERTEXDATA_STRIDE 5
 
+#define INDIRECTDRAWARGS_INTERPCOUNT_OFFSET 10
+#define INDIRECTDRAWARGS_ONEOVERINTERPCOUNT_OFFSET 11
+
 uint UnpackParticleId(uint InterpIdx)
 {
 	return (InterpIdx >> NiagaraRibbonVF.ParticleIdShift) & NiagaraRibbonVF.ParticleIdMask;
@@ -93,15 +96,15 @@ float3 CubicInterp(float3 P0, float3 T0, float3 P1, float3 T1, float A, float TS
  *							Particle specific
  *********************************************************************************/
  
-Buffer<int> SortedIndices;
-uint SortedIndicesOffset;
+Buffer<uint> SortedIndices;
 
-Buffer<float4> TangentsAndDistances;
+Buffer<float> TangentsAndDistances;
 Buffer<uint> MultiRibbonIndices;
-Buffer<float> PackedPerRibbonDataByIndex;
+Buffer<uint> PackedPerRibbonDataByIndex;
 
 uint FacingMode;
 
+
 // Naming convention :
 // RawParticleId : Used to access per particle data stored in the VF. Actual particle data needs to be indirected through SortedIndices.
 // ParticleId : Used to access particle data in NiagaraVFParticleAccess.
@@ -140,6 +143,7 @@ struct FSegmentData
 	float U0DistributionScaler;
 	float U1DistributionScaler;
 	uint StartParticleId;
+	uint LastParticleId;
 };
 
 struct FRibbonPositionIntermediates
@@ -284,15 +288,33 @@ void SetUV(inout FVertexFactoryInterpolantsVSToPS Interpolants, int UVIndex, flo
 
 int GetParticleIdByInterpVtxId(uint InterpVtxId)
 {
-	uint ParticleId = UnpackParticleId(InterpVtxId);
+	return NiagaraRibbonVFLooseParameters.SortedIndices[UnpackParticleId(InterpVtxId)];
+}
 
-	FLATTEN
-	if(NiagaraRibbonVFLooseParameters.SortedIndicesOffset != 0xFFFFFFFF)
+uint GetInterpCount()
+{
+	BRANCH
+	if (NiagaraRibbonVFLooseParameters.IndirectDrawOutputOffset != -1)
 	{
-		ParticleId = NiagaraRibbonVFLooseParameters.SortedIndices[NiagaraRibbonVFLooseParameters.SortedIndicesOffset + ParticleId];
+		return NiagaraRibbonVFLooseParameters.IndirectDrawOutput[NiagaraRibbonVFLooseParameters.IndirectDrawOutputOffset + INDIRECTDRAWARGS_INTERPCOUNT_OFFSET];
 	}
+	else
+	{
+		return NiagaraRibbonVF.InterpCount;
+	}
+}
 
-	return ParticleId;
+float GetOneOverInterpCount()
+{
+	BRANCH
+	if (NiagaraRibbonVFLooseParameters.IndirectDrawOutputOffset != -1)
+	{
+		return asfloat(NiagaraRibbonVFLooseParameters.IndirectDrawOutput[NiagaraRibbonVFLooseParameters.IndirectDrawOutputOffset + INDIRECTDRAWARGS_ONEOVERINTERPCOUNT_OFFSET]);
+	}
+	else
+	{
+		return NiagaraRibbonVF.OneOverInterpCount;
+	}
 }
 
 FInterpVtxData GetInterpVtxData(FVertexFactoryInput Input)
@@ -302,20 +324,12 @@ FInterpVtxData GetInterpVtxData(FVertexFactoryInput Input)
 	InterpVtxData.RawParticleId0 = UnpackParticleId(Input.InterpVtxId);
 	InterpVtxData.RawParticleId1 = InterpVtxData.RawParticleId0 + 1;
 
-	if (NiagaraRibbonVFLooseParameters.SortedIndicesOffset != 0xFFFFFFFF)
-	{
-		InterpVtxData.ParticleId0 = NiagaraRibbonVFLooseParameters.SortedIndices[NiagaraRibbonVFLooseParameters.SortedIndicesOffset + InterpVtxData.RawParticleId0];
-		InterpVtxData.ParticleId1 = NiagaraRibbonVFLooseParameters.SortedIndices[NiagaraRibbonVFLooseParameters.SortedIndicesOffset + InterpVtxData.RawParticleId1];
-	}
-	else
-	{
-		InterpVtxData.ParticleId0 = InterpVtxData.RawParticleId0;
-		InterpVtxData.ParticleId1 = InterpVtxData.RawParticleId1;
-	}
+	InterpVtxData.ParticleId0 = NiagaraRibbonVFLooseParameters.SortedIndices[InterpVtxData.RawParticleId0];
+	InterpVtxData.ParticleId1 = NiagaraRibbonVFLooseParameters.SortedIndices[InterpVtxData.RawParticleId1];
 
 #if USE_PARTICLE_INTERPOLATION
 	{
-		InterpVtxData.Alpha = float(UnpackInterpSlice(Input.InterpVtxId)) * NiagaraRibbonVF.OneOverInterpCount;
+		InterpVtxData.Alpha = float(UnpackInterpSlice(Input.InterpVtxId)) * GetOneOverInterpCount();
 	}
 #else
 	InterpVtxData.Alpha = 0;
@@ -327,17 +341,18 @@ FInterpVtxData GetInterpVtxData(FVertexFactoryInput Input)
 FSegmentData UnpackPerRibbonDataByIndex(int RibbonIndex)
 {
 	FSegmentData SegmentData = (FSegmentData)0;
-	const int Index = RibbonIndex * 7;
+	const int Index = RibbonIndex * 8;
 
-	SegmentData.UV0Scale = float2(NiagaraRibbonVFLooseParameters.PackedPerRibbonDataByIndex[Index], NiagaraRibbonVF.PackedVData.x);
-	SegmentData.UV0Offset = float2(NiagaraRibbonVFLooseParameters.PackedPerRibbonDataByIndex[Index + 1], NiagaraRibbonVF.PackedVData.y);
-	SegmentData.U0DistributionScaler = NiagaraRibbonVFLooseParameters.PackedPerRibbonDataByIndex[Index + 2];
+	SegmentData.UV0Scale = float2(asfloat(NiagaraRibbonVFLooseParameters.PackedPerRibbonDataByIndex[Index + 0]), NiagaraRibbonVF.PackedVData.x);
+	SegmentData.UV0Offset = float2(asfloat(NiagaraRibbonVFLooseParameters.PackedPerRibbonDataByIndex[Index + 1]), NiagaraRibbonVF.PackedVData.y);
+	SegmentData.U0DistributionScaler = asfloat(NiagaraRibbonVFLooseParameters.PackedPerRibbonDataByIndex[Index + 2]);
 	
-	SegmentData.UV1Scale = float2(NiagaraRibbonVFLooseParameters.PackedPerRibbonDataByIndex[Index + 3], NiagaraRibbonVF.PackedVData.z);
-	SegmentData.UV1Offset = float2(NiagaraRibbonVFLooseParameters.PackedPerRibbonDataByIndex[Index + 4], NiagaraRibbonVF.PackedVData.w);
-	SegmentData.U1DistributionScaler = NiagaraRibbonVFLooseParameters.PackedPerRibbonDataByIndex[Index + 5];
+	SegmentData.UV1Scale = float2(asfloat(NiagaraRibbonVFLooseParameters.PackedPerRibbonDataByIndex[Index + 3]), NiagaraRibbonVF.PackedVData.z);
+	SegmentData.UV1Offset = float2(asfloat(NiagaraRibbonVFLooseParameters.PackedPerRibbonDataByIndex[Index + 4]), NiagaraRibbonVF.PackedVData.w);
+	SegmentData.U1DistributionScaler = asfloat(NiagaraRibbonVFLooseParameters.PackedPerRibbonDataByIndex[Index + 5]);
 
-	SegmentData.StartParticleId = asuint(NiagaraRibbonVFLooseParameters.PackedPerRibbonDataByIndex[Index + 6]);
+	SegmentData.StartParticleId = NiagaraRibbonVFLooseParameters.PackedPerRibbonDataByIndex[Index + 6];
+	SegmentData.LastParticleId = NiagaraRibbonVFLooseParameters.PackedPerRibbonDataByIndex[Index + 7];
 
 	return SegmentData;
 }
@@ -493,6 +508,11 @@ float3x3 CalcTangentBasis(FVertexFactoryInput Input, FVertexFactoryIntermediates
 	return float3x3(Intermediates.Current.TangentX, Intermediates.Current.TangentY, Intermediates.Current.TangentZ);
 }
 
+float GetTangentAndDistanceComponent(int RawParticleID, int ElementIndex)
+{
+	return NiagaraRibbonVFLooseParameters.TangentsAndDistances[RawParticleID * 4 + ElementIndex];
+}
+
 // known limitations:
 // -TangentsAndDistances is from the current frame and if we're generating intermediates for the previous frame for velocity
 // rendering we won't get the exact same results.  Issues around trying to supply a previous version of the TangentsAndDistances
@@ -515,8 +535,16 @@ FRibbonPositionIntermediates GetRibbonPositionIntermediates(
 	{
 		const float3 Pos0 = GetVec3(PositionDataOffset, InterpVtxData.ParticleId0);
 		const float3 Pos1 = GetVec3(PositionDataOffset, InterpVtxData.ParticleId1);
-		const float4 TangentAndDist0 = NiagaraRibbonVFLooseParameters.TangentsAndDistances[InterpVtxData.RawParticleId0];
-		const float4 TangentAndDist1 = NiagaraRibbonVFLooseParameters.TangentsAndDistances[InterpVtxData.RawParticleId1];
+		const float4 TangentAndDist0 = float4(
+			GetTangentAndDistanceComponent(InterpVtxData.RawParticleId0, 0),
+			GetTangentAndDistanceComponent(InterpVtxData.RawParticleId0, 1), 
+			GetTangentAndDistanceComponent(InterpVtxData.RawParticleId0, 2),
+			GetTangentAndDistanceComponent(InterpVtxData.RawParticleId0, 3));
+		const float4 TangentAndDist1 = float4(
+			GetTangentAndDistanceComponent(InterpVtxData.RawParticleId1, 0),
+			GetTangentAndDistanceComponent(InterpVtxData.RawParticleId1, 1), 
+			GetTangentAndDistanceComponent(InterpVtxData.RawParticleId1, 2),
+			GetTangentAndDistanceComponent(InterpVtxData.RawParticleId1, 3));
 		
 		Intermediates.Position = TransformPosition(LocalToWorld, CubicInterp(Pos0, TangentAndDist0.xyz, Pos1, TangentAndDist1.xyz, InterpVtxData.Alpha, TangentAndDist1.w - TangentAndDist0.w));
 		Intermediates.Direction = normalize(TransformVector(LocalToWorld, InvScale, lerp(TangentAndDist0.xyz, TangentAndDist1.xyz, InterpVtxData.Alpha)));
@@ -599,7 +627,7 @@ FVertexFactoryIntermediates GetVertexFactoryIntermediates(FVertexFactoryInput In
 	
 	Intermediates.SceneData = VF_GPUSCENE_GET_INTERMEDIATES(Input);
 	
-	uint RibbonIndex = NiagaraRibbonVFLooseParameters.MultiRibbonIndices[InterpVtxData.RawParticleId0];
+	uint RibbonIndex = NiagaraRibbonVF.ShouldUseMultiRibbon == 1? NiagaraRibbonVFLooseParameters.MultiRibbonIndices[InterpVtxData.RawParticleId0] : 0;
 	FSegmentData SegmentData = UnpackPerRibbonDataByIndex(RibbonIndex);
 
 	const float2 SliceVertexNormals = UnpackParticleVertexNormal(Input.InterpVtxId);
diff --git a/Engine/Plugins/FX/Niagara/Shaders/Private/Ribbons/NiagaraRibbonInitializeIndices.usf b/Engine/Plugins/FX/Niagara/Shaders/Private/Ribbons/NiagaraRibbonInitializeIndices.usf
new file mode 100644
--- /dev/null
+++ b/Engine/Plugins/FX/Niagara/Shaders/Private/Ribbons/NiagaraRibbonInitializeIndices.usf
@@ -0,0 +1,141 @@
+// Copyright Epic Games, Inc. All Rights Reserved.
+
+#include "/Engine/Private/Common.ush"
+#include "NiagaraRibbonCommon.ush"
+
+RWBuffer<uint> IndirectDrawOutput;
+
+// Stats of the vertex gen, including tessellation stats
+Buffer<uint> VertexGenerationResults;
+
+uint IndirectDrawOutputIndex;
+uint VertexGenerationResultsIndex;
+
+uint IndexGenThreadSize;
+uint TrianglesPerSegment;
+
+float ViewDistance;
+float LODDistanceFactor;
+uint TessellationMode;
+uint bCustomUseConstantFactor;
+uint CustomTessellationFactor;
+float CustomTessellationMinAngle;
+uint bCustomUseScreenSpace;
+uint GNiagaraRibbonMaxTessellation;
+float GNiagaraRibbonTessellationAngle;
+float GNiagaraRibbonTessellationScreenPercentage;
+uint GNiagaraRibbonTessellationEnabled;
+float GNiagaraRibbonTessellationMinDisplacementError;
+
+#define TESSELLATIONMODE_AUTOMATIC 0
+#define TESSELLATIONMODE_CUSTOM 1
+#define TESSELLATIONMODE_DISABLED 2
+
+
+[numthreads(1, 1, 1)]
+void InitializeIndices(uint3 DispatchThreadId : SV_DispatchThreadID)
+{
+	const uint TotalNumParticles = GetTotalNumParticles();
+#if RIBBONS_WANTS_AUTOMATIC_TESSELLATION
+	
+	bool bUseConstantFactor = false;
+	uint TessellationFactor = GNiagaraRibbonMaxTessellation;
+	float TessellationMinAngle = GNiagaraRibbonTessellationAngle;
+	float ScreenPercentage = GNiagaraRibbonTessellationScreenPercentage;
+
+	switch (TessellationMode)
+	{
+	case TESSELLATIONMODE_AUTOMATIC:
+		break;
+	case TESSELLATIONMODE_CUSTOM:
+		TessellationFactor = min(TessellationFactor, CustomTessellationFactor); // Don't allow factors bigger than the platform limit.
+		bUseConstantFactor = bCustomUseConstantFactor == 1;
+		TessellationMinAngle = CustomTessellationMinAngle;
+		ScreenPercentage = bCustomUseScreenSpace == 1 && !bUseConstantFactor ? GNiagaraRibbonTessellationScreenPercentage : 0.f;
+		break;
+	case TESSELLATIONMODE_DISABLED:
+		TessellationFactor = 1;
+		break;
+	default:
+		break;
+	}
+	
+	uint SegmentTessellation = 1;
+
+	float TessellationAngle = asfloat(VertexGenerationResults[VertexGenerationResultsIndex * VERTEX_GEN_OUTPUT_DATA_STRIDE + VERTEX_GEN_OUTPUT_DATA_TESSELLATION_ANGLE]);
+	float TessellationCurvature = asfloat(VertexGenerationResults[VertexGenerationResultsIndex * VERTEX_GEN_OUTPUT_DATA_STRIDE + VERTEX_GEN_OUTPUT_DATA_TESSELLATION_CURVATURE]);
+	float TessellationTwistAngle = asfloat(VertexGenerationResults[VertexGenerationResultsIndex * VERTEX_GEN_OUTPUT_DATA_STRIDE + VERTEX_GEN_OUTPUT_DATA_TESSELLATION_TWIST_ANGLE]);
+	float TessellationTwistCurvature = asfloat(VertexGenerationResults[VertexGenerationResultsIndex * VERTEX_GEN_OUTPUT_DATA_STRIDE + VERTEX_GEN_OUTPUT_DATA_TESSELLATION_TWIST_CURVATURE]);
+	float TessellationTotalSegmentLength = asfloat(VertexGenerationResults[VertexGenerationResultsIndex * VERTEX_GEN_OUTPUT_DATA_STRIDE + VERTEX_GEN_OUTPUT_DATA_TESSELLATION_TOTAL_SEGMENT_LENGTH]);
+
+
+	if (bUseConstantFactor)
+	{
+		SegmentTessellation = TessellationFactor;
+	}
+	else if (GNiagaraRibbonTessellationEnabled == 1 && TessellationFactor > 1 && TessellationCurvature > UE_SMALL_NUMBER)
+	{	
+		const float MinTesselation = (TessellationMinAngle == 0.f || bUseConstantFactor)?
+			                             float(TessellationFactor)	:
+			                             max(1.f, max(TessellationTwistAngle, TessellationAngle) / max(UE_SMALL_NUMBER, TessellationMinAngle));
+
+		const float MAX_CURVATURE_FACTOR = 0.002f; // This will clamp the curvature to around 2.5 km and avoid numerical issues.
+		
+		const float MaxDisplacementError = max(GNiagaraRibbonTessellationMinDisplacementError, ScreenPercentage * sqrt(ViewDistance) / LODDistanceFactor);
+		float Tess = TessellationAngle / max(MAX_CURVATURE_FACTOR, acos(TessellationCurvature / (TessellationCurvature + MaxDisplacementError)));
+		// FMath::RoundUpToPowerOfTwo ? This could avoid vertices moving around as tesselation increases
+
+		if (TessellationTwistAngle > 0 && TessellationTwistCurvature > 0)
+		{
+			const float TwistTess = TessellationTwistAngle / max(MAX_CURVATURE_FACTOR, acos(TessellationTwistCurvature / (TessellationTwistCurvature + MaxDisplacementError)));
+			Tess = max(TwistTess, Tess);
+		}
+		SegmentTessellation = clamp(uint(round(Tess)), uint(round(MinTesselation)), TessellationFactor);
+	}
+#elif RIBBONS_WANTS_CONSTANT_TESSELLATION
+	// Has a constant factor so just use it
+	const uint SegmentTessellation = min(GNiagaraRibbonMaxTessellation, CustomTessellationFactor);
+#else
+	// Tessellation is disabled, so just set it to a factor of 1
+	const uint SegmentTessellation = 1;
+#endif
+
+	// We need to run for every possible segment * tessellation so one less than total particles * tessellation
+	const uint TotalNumInvocations = (TotalNumParticles - 1) * SegmentTessellation;	
+	const uint NumThreadGroups = DivideAndRoundUp(TotalNumInvocations, IndexGenThreadSize);
+
+	const uint SegmentCount = VertexGenerationResults[VertexGenerationResultsIndex * VERTEX_GEN_OUTPUT_DATA_STRIDE + VERTEX_GEN_OUTPUT_DATA_TOTAL_NUM_SEGMENTS_OFFSET];
+
+	const uint IndirectDrawOffset = IndirectDrawOutputIndex * INDEX_GEN_INDIRECT_ARGS_STRIDE;
+
+	const float OneOverSubSegmentTessellation = 1.0f / float(SegmentTessellation);
+	
+	// Output the dispatch indirect args followed by the segmentation amount
+	IndirectDrawOutput[IndirectDrawOffset + INDEX_GEN_EXECUTE_INDIRECT_OFFSET + 0] = NumThreadGroups;
+	IndirectDrawOutput[IndirectDrawOffset + INDEX_GEN_EXECUTE_INDIRECT_OFFSET + 1] = 1;
+	IndirectDrawOutput[IndirectDrawOffset + INDEX_GEN_EXECUTE_INDIRECT_OFFSET + 2] = 1;
+	IndirectDrawOutput[IndirectDrawOffset + INDEX_GEN_TESSELLATION_FACTOR_OFFSET] = SegmentTessellation;
+
+			
+	// Index Count
+	IndirectDrawOutput[IndirectDrawOffset + INDEX_GEN_INDIRECT_DRAW_OFFSET + 0] = SegmentCount * SegmentTessellation * TrianglesPerSegment * 3;
+	// Num Instances
+	IndirectDrawOutput[IndirectDrawOffset + INDEX_GEN_INDIRECT_DRAW_OFFSET + 1] = 1;
+	// First Index Offset
+	IndirectDrawOutput[IndirectDrawOffset + INDEX_GEN_INDIRECT_DRAW_OFFSET + 2] = 0;
+	// First Vertex Offset
+	IndirectDrawOutput[IndirectDrawOffset + INDEX_GEN_INDIRECT_DRAW_OFFSET + 3] = 0;
+	// First Instance Offset
+	IndirectDrawOutput[IndirectDrawOffset + INDEX_GEN_INDIRECT_DRAW_OFFSET + 4] = 0;
+	
+	// SegmentCount
+	IndirectDrawOutput[IndirectDrawOffset + INDEX_GEN_NUM_SEGMENTS_OFFSET] = SegmentCount;
+	// SubSegmentCount
+	IndirectDrawOutput[IndirectDrawOffset + INDEX_GEN_NUM_SUB_SEGMENTS_OFFSET] = SegmentTessellation;
+	// 1/SubSegmentCount
+	IndirectDrawOutput[IndirectDrawOffset + INDEX_GEN_ONE_OVER_NUM_SUB_SEGMENTS_OFFSET] = asuint(OneOverSubSegmentTessellation);
+
+}
+
+
+
diff --git a/Engine/Plugins/FX/Niagara/Shaders/Private/Ribbons/NiagaraRibbonSortParticles.usf b/Engine/Plugins/FX/Niagara/Shaders/Private/Ribbons/NiagaraRibbonSortParticles.usf
new file mode 100644
--- /dev/null
+++ b/Engine/Plugins/FX/Niagara/Shaders/Private/Ribbons/NiagaraRibbonSortParticles.usf
@@ -0,0 +1,188 @@
+﻿// Copyright Epic Games, Inc. All Rights Reserved.
+
+#include "/Engine/Private/Common.ush"
+#include "NiagaraRibbonCommon.ush"
+
+Buffer<uint> SortedIndices;
+RWBuffer<uint> DestinationSortedIndices;
+Buffer<uint> EmitterCounts;
+uint CountOffset;
+
+int MergeSortSourceBlockSize;
+int MergeSortDestinationBlockSize;
+
+struct FBubbleShortEntry
+{
+	uint ParticleID;
+	float LinkID;
+	
+#if HAS_RIBBON_ID
+	FRibbonID RibbonID;
+#endif
+};
+
+groupshared FBubbleShortEntry BubbleSortTemp[THREADGROUP_SIZE];
+
+
+
+[numthreads(THREADGROUP_SIZE, 1, 1)]
+void InitialSortList(uint3 DispatchThreadId : SV_DispatchThreadID, uint InGroupIndex : SV_GroupIndex, uint3 InGroupID : SV_GroupID)
+{
+	const uint TotalNumParticles = GetTotalNumParticles();
+	const int GroupStartOffset = InGroupID.x * THREADGROUP_SIZE;
+	const int ParticlesRemaining = TotalNumParticles - GroupStartOffset;
+	const int MaxParticlesInSet = THREADGROUP_SIZE;
+	const int TotalNumParticlesInSet = min(MaxParticlesInSet, ParticlesRemaining);
+	
+	const int GlobalIndexOffset = DispatchThreadId.x;
+	const int LocalIndexOffset = InGroupIndex;
+
+	FBubbleShortEntry NewEntry;
+	NewEntry.ParticleID = GlobalIndexOffset; //(TotalNumParticles - 1) - GlobalIndexOffset;
+	NewEntry.LinkID = GetRibbonLinkID(NewEntry.ParticleID);	
+#if HAS_RIBBON_ID
+	NewEntry.RibbonID = GetRibbonID(NewEntry.ParticleID);
+#endif
+	
+	BubbleSortTemp[LocalIndexOffset] = NewEntry;	
+	
+	int SortGroup = 1;
+	
+	for (int Index = 0; Index < THREADGROUP_SIZE; Index++)
+	{
+		// Switch the offset and sync all threads
+		SortGroup ^= 0x1;		
+		GroupMemoryBarrierWithGroupSync();
+
+		if (LocalIndexOffset % 2 == SortGroup && LocalIndexOffset + 1 < TotalNumParticlesInSet)
+		{
+			const FBubbleShortEntry ThisEntry = BubbleSortTemp[LocalIndexOffset + 0];
+			const FBubbleShortEntry NextEntry = BubbleSortTemp[LocalIndexOffset + 1];		
+
+#if HAS_RIBBON_ID			
+			const int RibbonIDComparisonResult = CompareRibbonIDs(ThisEntry.RibbonID, NextEntry.RibbonID);
+			const bool bShouldSwap = RibbonIDComparisonResult > 0 || (RibbonIDComparisonResult == 0 && ThisEntry.LinkID > NextEntry.LinkID);
+#else
+			const bool bShouldSwap = ThisEntry.LinkID > NextEntry.LinkID;
+#endif
+			
+			if (bShouldSwap)
+			{
+				BubbleSortTemp[LocalIndexOffset + 0] = NextEntry;
+				BubbleSortTemp[LocalIndexOffset + 1] = ThisEntry;
+			}
+		}
+	}
+
+	GroupMemoryBarrierWithGroupSync();
+
+	if (LocalIndexOffset < TotalNumParticlesInSet)
+	{
+		DestinationSortedIndices[GlobalIndexOffset] = BubbleSortTemp[LocalIndexOffset].ParticleID;
+	}
+}
+
+
+
+
+
+
+
+int MergeSortCompare(float SearchLinkOrder, float TestLinkOrder)
+{
+	if (SearchLinkOrder > TestLinkOrder)
+	{
+		return 1;
+	}
+
+	if (SearchLinkOrder == TestLinkOrder)
+	{
+		return 0;
+	}
+	
+	return -1;
+}
+
+#if HAS_RIBBON_ID
+int MergeSortCompare(FRibbonID Left, float LeftLinkOrder, FRibbonID Right, float RightLinkOrder)
+{
+	const int RibbonIDComparisonResult = CompareRibbonIDs(Left, Right);
+
+	if (RibbonIDComparisonResult > 0)
+	{
+		return 1;
+	}
+
+	if (RibbonIDComparisonResult == 0)
+	{
+		return MergeSortCompare(LeftLinkOrder, RightLinkOrder);
+	}
+
+	return -1;
+}
+#endif
+
+
+
+
+// This is a scatter based mergesort kernel
+[numthreads(THREADGROUP_SIZE, 1, 1)]
+void MergeSort(uint3 DispatchThreadId : SV_DispatchThreadID)
+{
+	const uint TotalNumParticles = GetTotalNumParticles();
+	const uint Index = DispatchThreadId.x;
+	
+	if (Index < TotalNumParticles)
+	{		
+		const uint DestinationBlockStart = (Index / MergeSortDestinationBlockSize) * MergeSortDestinationBlockSize;
+		const uint SourceBlockStart = (Index / MergeSortSourceBlockSize) * MergeSortSourceBlockSize;
+		const bool bIsInPrimaryBlock = Index < (DestinationBlockStart + MergeSortSourceBlockSize);
+		const uint OtherSourceBlockStart = DestinationBlockStart + (bIsInPrimaryBlock? MergeSortSourceBlockSize : 0);
+		
+		const uint IndexWithinSourceBlock = Index - SourceBlockStart;
+		
+		const uint ParticleID = SortedIndices[Index];
+		const float LinkID = GetRibbonLinkID(ParticleID);
+#if HAS_RIBBON_ID
+		const FRibbonID RibbonID = GetRibbonID(ParticleID);
+#endif
+			
+		// we find the index of the last particle before this in the other half,
+		// add that index to ours and we have the number of particles preceeding
+		// so we can scatter into the destination
+		// Handle equal values by placing all of those in the primary block before those in the secondary block
+		
+		int SearchStart = 0;
+		int SearchStop = min(int(MergeSortSourceBlockSize), max(int(TotalNumParticles) - int(OtherSourceBlockStart), 0));
+		
+		while (SearchStart != SearchStop)
+		{
+			const int Mid = (SearchStart + SearchStop) / 2;
+			const int SearchIndex = OtherSourceBlockStart + Mid;
+			const uint OtherParticleID = SortedIndices[SearchIndex];
+			const float OtherLinkID = GetRibbonLinkID(OtherParticleID);
+#if HAS_RIBBON_ID
+			const FRibbonID OtherRibbonID = GetRibbonID(OtherParticleID);
+			const int Result = MergeSortCompare(RibbonID, LinkID, OtherRibbonID, OtherLinkID);
+#else
+			const int Result = MergeSortCompare(LinkID, OtherLinkID);			
+#endif
+			
+			if (Result > 0 || (!bIsInPrimaryBlock && Result == 0))
+			{
+				SearchStart = min(Mid + 1, SearchStop);							
+			}
+			else
+			{
+				SearchStop = max(Mid, SearchStart);				
+			}
+		}
+		
+		const int FinalIndex = DestinationBlockStart + IndexWithinSourceBlock + SearchStop;
+		DestinationSortedIndices[FinalIndex] = ParticleID;
+	}
+}
+
+
+
+
