GPU Ribbons!

diff --git a/Engine/Plugins/FX/Niagara/Source/Niagara/Private/NiagaraRendererRibbons.cpp b/Engine/Plugins/FX/Niagara/Source/Niagara/Private/NiagaraRendererRibbons.cpp
--- a/Engine/Plugins/FX/Niagara/Source/Niagara/Private/NiagaraRendererRibbons.cpp
+++ b/Engine/Plugins/FX/Niagara/Source/Niagara/Private/NiagaraRendererRibbons.cpp
@@ -1,6 +1,8 @@
 // Copyright Epic Games, Inc. All Rights Reserved.
 
 #include "NiagaraRendererRibbons.h"
+
+#include "GPUSortManager.h"
 #include "ParticleResources.h"
 #include "NiagaraRibbonVertexFactory.h"
 #include "NiagaraDataSet.h"
@@ -12,8 +14,35 @@
 #include "RayTracingInstance.h"
 #include "Math/NumericLimits.h"
 #include "NiagaraCullProxyComponent.h"
+#include "NiagaraGpuComputeDispatchInterface.h"
+#include "NiagaraRibbonCompute.h"
+#include "Misc/LazySingleton.h"
 
 DECLARE_CYCLE_STAT(TEXT("Generate Ribbon Vertex Data [GT]"), STAT_NiagaraGenRibbonVertexData, STATGROUP_Niagara);
+DECLARE_CYCLE_STAT(TEXT("Render Ribbons [RT]"), STAT_NiagaraRenderRibbons, STATGROUP_Niagara);
+
+DECLARE_CYCLE_STAT(TEXT("Generate Indices CPU [GT]"), STAT_NiagaraRenderRibbonsGenIndiciesCPU, STATGROUP_Niagara);
+DECLARE_CYCLE_STAT(TEXT("Generate Indices GPU [RT]"), STAT_NiagaraRenderRibbonsGenIndiciesGPU, STATGROUP_Niagara);
+
+DECLARE_CYCLE_STAT(TEXT("Generate Vertices CPU [GT]"), STAT_NiagaraRenderRibbonsGenVerticesCPU, STATGROUP_Niagara);
+DECLARE_CYCLE_STAT(TEXT("Generate Vertices GPU [RT]"), STAT_NiagaraRenderRibbonsGenVerticesGPU, STATGROUP_Niagara);
+
+
+DECLARE_STATS_GROUP(TEXT("NiagaraRibbons"), STATGROUP_NiagaraRibbons, STATCAT_Advanced);
+
+DECLARE_CYCLE_STAT(TEXT("Generate Vertices GPU - Sort [RT]"), STAT_NiagaraRenderRibbonsGenVerticesSortGPU, STATGROUP_NiagaraRibbons);
+DECLARE_CYCLE_STAT(TEXT("Generate Vertices GPU - InitialSort [RT]"), STAT_NiagaraRenderRibbonsGenVerticesInitialSortGPU, STATGROUP_NiagaraRibbons);
+DECLARE_CYCLE_STAT(TEXT("Generate Vertices GPU - FinalSort [RT]"), STAT_NiagaraRenderRibbonsGenVerticesFinalSortGPU, STATGROUP_NiagaraRibbons);
+DECLARE_CYCLE_STAT(TEXT("Generate Vertices GPU - Reduction Phase 1 [RT]"), STAT_NiagaraRenderRibbonsGenVerticesReductionPhase1GPU, STATGROUP_NiagaraRibbons);
+DECLARE_CYCLE_STAT(TEXT("Generate Vertices GPU - Reduction Init [RT]"), STAT_NiagaraRenderRibbonsGenVerticesReductionInitGPU, STATGROUP_NiagaraRibbons);
+DECLARE_CYCLE_STAT(TEXT("Generate Vertices GPU - Reduction Propagate [RT]"), STAT_NiagaraRenderRibbonsGenVerticesReductionPropagateGPU, STATGROUP_NiagaraRibbons);
+DECLARE_CYCLE_STAT(TEXT("Generate Vertices GPU - Reduction Tessellation [RT]"), STAT_NiagaraRenderRibbonsGenVerticesReductionTessellationGPU, STATGROUP_NiagaraRibbons);
+DECLARE_CYCLE_STAT(TEXT("Generate Vertices GPU - Reduction Phase 2 [RT]"), STAT_NiagaraRenderRibbonsGenVerticesReductionPhase2GPU, STATGROUP_NiagaraRibbons);
+
+DECLARE_CYCLE_STAT(TEXT("Generate Vertices GPU - Reduction Finalize [RT]"), STAT_NiagaraRenderRibbonsGenVerticesReductionFinalizeGPU, STATGROUP_NiagaraRibbons);
+DECLARE_CYCLE_STAT(TEXT("Generate Vertices GPU - MultiRibbon Init [RT]"), STAT_NiagaraRenderRibbonsGenVerticesMultiRibbonInitGPU, STATGROUP_NiagaraRibbons);
+DECLARE_CYCLE_STAT(TEXT("Generate Vertices GPU - MultiRibbon Init Compute [RT]"), STAT_NiagaraRenderRibbonsGenVerticesMultiRibbonInitComputeGPU, STATGROUP_NiagaraRibbons);
+
 
 int32 GNiagaraRibbonTessellationEnabled = 1;
 static FAutoConsoleVariableRef CVarNiagaraRibbonTessellationEnabled(
@@ -80,6 +109,24 @@ static TAutoConsoleVariable<int32> CVarRayTracingNiagaraRibbons(
 	1,
 	TEXT("Include Niagara ribbons in ray tracing effects (default = 1 (Niagara ribbons enabled in ray tracing))"));
 
+
+static int32 GbEnableNiagaraRibbonComputeShaderGen = 1;
+static FAutoConsoleVariableRef CVarEnableNiagaraRibbonComputeShaderGen(
+	TEXT("fx.EnableNiagaraRibbonComputeShaderGen"),
+	GbEnableNiagaraRibbonComputeShaderGen,
+	TEXT("If == 1, Niagara Ribbon Renderer will use compute shaders if available to generate intermediate data. \n"),
+	ECVF_Default
+);
+
+static int32 GbForceNiagaraRibbonGPUInit = 0;
+static FAutoConsoleVariableRef CVarForceNiagaraRibbonGPUInit(
+	TEXT("fx.ForceNiagaraRibbonGPUInit"),
+	GbForceNiagaraRibbonGPUInit,
+	TEXT("If == 1, Niagara Ribbon Renderer will use compute shader based initialization for all CPU systems. \n"),
+	ECVF_Default
+);
+
+
 // max absolute error 9.0x10^-3
 // Eberly's polynomial degree 1 - respect bounds
 // input [-1, 1] and output [0, PI]
@@ -91,12 +138,130 @@ FORCEINLINE float AcosFast(float InX)
 	return (InX >= 0) ? Res : PI - Res;
 }
 
+// Calculates the number of bits needed to store a maximum value
+FORCEINLINE uint32 CalculateBitsForRange(uint32 Range)
+{
+	return FMath::CeilToInt(FMath::Loge(static_cast<float>(Range)) / FMath::Loge(static_cast<float>(2)));
+}
+
+// Generates the mask to remove unecessary bits after a range of bits
+FORCEINLINE uint32 CalculateBitMask(uint32 NumBits)
+{
+	return static_cast<uint32>(static_cast<uint64>(0xFFFFFFFF) >> (32 - NumBits));
+}
+
+struct FTessellationStatsEntry
+{
+	static constexpr int32 NumElements = 5;
+	
+	float TotalLength;
+	float AverageSegmentLength;
+	float AverageSegmentAngle;
+	float AverageTwistAngle;
+	float AverageWidth;
+};
+
+struct FTessellationStatsEntryNoTwist
+{
+	static constexpr int32 NumElements = 3;
+	
+	float TotalLength;
+	float AverageSegmentLength;
+	float AverageSegmentAngle;
+};
+
+struct FNiagaraRibbonCommandBufferLayout
+{
+	static constexpr int32 NumElements = 15;
+	
+	uint32 FinalizationIndirectArgsXDim;
+	uint32 FinalizationIndirectArgsYDim;
+	uint32 FinalizationIndirectArgsZDim;
+	uint32 NumSegments;
+	uint32 NumRibbons;
+	
+	float Tessellation_Angle;
+	float Tessellation_Curvature;
+	float Tessellation_TwistAngle;
+	float Tessellation_TwistCurvature;
+	float Tessellation_TotalLength;
+
+	float TessCurrentFrame_TotalLength;
+	float TessCurrentFrame_AverageSegmentLength;
+	float TessCurrentFrame_AverageSegmentAngle;
+	float TessCurrentFrame_AverageTwistAngle;
+	float TessCurrentFrame_AverageWidth;
+};
+
+
+struct FNiagaraRibbonIndirectDrawBufferLayout
+{
+	static constexpr int32 NumElements = 12;
+	static constexpr int32 GenerateIndicesCommandOffset = 0;
+	static constexpr int32 IndirectDrawCommandIndex = 4;
+	static constexpr int32 IndirectDrawCommandByteOffset = IndirectDrawCommandIndex * sizeof(uint32);
+
+	// This is passed from InitializeIndices to GenerateIndices
+	uint32 IndexGenIndirectArgsXDim;
+	uint32 IndexGenIndirectArgsYDim;
+	uint32 IndexGenIndirectArgsZDim;
+	uint32 TessellationFactor;
+
+	// This is the indirect draw args and then resulting information for the vertex shader	
+	uint32 IndexCount;
+	uint32 NumInstances;
+	uint32 FirstIndexOffset;
+	uint32 FirstVertexOffset;
+	uint32 FirstInstanceOffset;
+	
+	uint32 NumSegments;
+	uint32 NumSubSegments;
+	float OneOverSubSegmentCount;	
+};
+
+
+struct FRWIndexBuffer final : FIndexBuffer
+{
+	FUnorderedAccessViewRHIRef UAV;
+
+	FRWIndexBuffer()
+	{}
+
+	virtual ~FRWIndexBuffer() override
+	{
+		FRenderResource::ReleaseResource();
+	}
+
+	void Initialize(const TCHAR* InDebugName, const uint32 BytesPerElement, const uint32 NumElements, const EPixelFormat Format, const ERHIAccess InResourceState, const EBufferUsageFlags AdditionalUsage = BUF_None, FResourceArrayInterface *InResourceArray = nullptr)
+	{
+		InitResource();
+		
+		// Provide a debug name if using Fast VRAM so the allocators diagnostics will work
+		ensure(!(EnumHasAnyFlags(AdditionalUsage, BUF_FastVRAM) && !InDebugName));
+		const int32 NumBytes = BytesPerElement * NumElements;
+		FRHIResourceCreateInfo CreateInfo(InDebugName);
+		CreateInfo.ResourceArray = InResourceArray;
+		IndexBufferRHI = RHICreateIndexBuffer(BytesPerElement, NumBytes, BUF_UnorderedAccess | AdditionalUsage, InResourceState, CreateInfo);
+		UAV = RHICreateUnorderedAccessView(IndexBufferRHI, UE_PIXELFORMAT_TO_UINT8(Format));
+	}
+
+	virtual void ReleaseRHI() override
+	{
+		UAV.SafeRelease();
+
+		FIndexBuffer::ReleaseRHI();
+	}
+};
+
+
 struct FNiagaraDynamicDataRibbon : public FNiagaraDynamicDataBase
 {
 	FNiagaraDynamicDataRibbon(const FNiagaraEmitterInstance* InEmitter)
 		: FNiagaraDynamicDataBase(InEmitter)
 		, Material(nullptr)
-		, MaxParticleIndex(0)
+		, MaxAllocatedParticleCount(0)
+		, bUseGPUInit(false)
+		, bIsGPUSystem(false)
 	{
 	}
 
@@ -108,142 +273,372 @@ struct FNiagaraDynamicDataRibbon : public FNiagaraDynamicDataBase
 		}
 	}
 
+	
 	/** Material to use passed to the Renderer. */
 	FMaterialRenderProxy* Material;
+	int32 MaxAllocatedParticleCount;
 
-	// The list of all segments, each one connecting SortedIndices[SegmentId] to SortedIndices[SegmentId + 1].
-	// We use this format because the final index buffer gets generated based on view sorting and InterpCount.
-	TArray<int32> SegmentData;
-	int32 MaxParticleIndex;
-
-	/** The list of all particle (instance) indices. Converts raw indices to particles indices. Ordered along each ribbons, from head to tail. */
-	TArray<int32> SortedIndices;
-	/** The tangent and distance between segments, for each raw index (raw VS particle indices). */
-	TArray<FVector4f> TangentAndDistances;
-	/** The multi ribbon index, for each raw index. (raw VS particle indices). */
-	TArray<uint32> MultiRibbonIndices;
-	/** Data for each multi ribbon. There are several entries per ribbon. */
-	TArray<float> PackedPerRibbonDataByIndex;
-	/** Position offsets for each vertex within a slice, used for volumetric ribbons */
-	TArray<float> SliceVertexData;
-
-	struct FMultiRibbonInfo
-	{
-		/** start and end world space position of the ribbon, to figure out draw direction */
-		FVector StartPos;
-		FVector EndPos;
-		int32 BaseSegmentDataIndex = 0;
-		int32 NumSegmentDataIndices = 0;
-
-		FORCEINLINE bool UseInvertOrder(const FVector& ViewDirection, const FVector& ViewOriginForDistanceCulling, ENiagaraRibbonDrawDirection DrawDirection) const
-		{
-			const float StartDist = FVector::DotProduct(ViewDirection, StartPos - ViewOriginForDistanceCulling);
-			const float EndDist = FVector::DotProduct(ViewDirection, EndPos - ViewOriginForDistanceCulling);
-			return ((StartDist >= EndDist) && DrawDirection == ENiagaraRibbonDrawDirection::BackToFront)
-				|| ((StartDist < EndDist) && DrawDirection == ENiagaraRibbonDrawDirection::FrontToBack);
-		}
-	};
+	bool bUseGPUInit;
+	bool bIsGPUSystem;
+	
+	TSharedPtr<FNiagaraRibbonCPUGeneratedVertexData> GenerationOutput;
+	
+	int32 GetAllocatedSize()const
+	{
+		return GenerationOutput.IsValid()? GenerationOutput->GetAllocatedSize() : 0;
+	}
+};
 
-	/** Ribbon perperties required for sorting. */
-	TArray<FMultiRibbonInfo> MultiRibbonInfos;
 
-	void PackPerRibbonData(float U0Scale, float U0Offset, float U0DistributionScaler, float U1Scale, float U1Offset, float U1DistributionScaler, uint32 FirstParticleId)
+struct FNiagaraRibbonRenderingFrameViewResources
+{
+	FNiagaraRibbonVertexFactory VertexFactory;
+	FNiagaraRibbonUniformBufferRef UniformBuffer;
+
+	FRWIndexBuffer IndexBuffer;
+	FRWBuffer IndirectDrawBuffer;
+
+	FNiagaraIndexGenerationInput IndexGenerationSettings;
+
+	int32 IndirectDrawBufferStartOffset = FNiagaraRibbonIndirectDrawBufferLayout::IndirectDrawCommandIndex;
+	int32 IndirectDrawBufferStartByteOffset = FNiagaraRibbonIndirectDrawBufferLayout::IndirectDrawCommandByteOffset;
+
+	~FNiagaraRibbonRenderingFrameViewResources()
+	{
+		UniformBuffer.SafeRelease();
+		VertexFactory.ReleaseResource();
+		IndexBuffer.ReleaseResource();
+		IndirectDrawBuffer.Release();
+	}
+};
+
+struct FNiagaraRibbonRenderingFrameResources
+{
+	TArray<TSharedPtr<FNiagaraRibbonRenderingFrameViewResources>> ViewResources;
+
+	FParticleRenderData ParticleData;
+		
+	FRHIShaderResourceView*	ParticleFloatSRV;
+	FRHIShaderResourceView* ParticleHalfSRV;
+	FRHIShaderResourceView* ParticleIntSRV;
+	
+	int32 ParticleFloatDataStride = INDEX_NONE;
+	int32 ParticleHalfDataStride = INDEX_NONE;
+	int32 ParticleIntDataStride = INDEX_NONE;
+	
+	int32 RibbonIdParamOffset = INDEX_NONE;
+	
+	~FNiagaraRibbonRenderingFrameResources()
+	{
+		ViewResources.Empty();
+
+		ParticleFloatSRV = nullptr;
+		ParticleHalfSRV = nullptr;
+		ParticleIntSRV = nullptr;
+
+		ParticleFloatDataStride = INDEX_NONE;
+		ParticleHalfDataStride = INDEX_NONE;
+		ParticleIntDataStride = INDEX_NONE;
+		
+		RibbonIdParamOffset = INDEX_NONE;
+	}	
+};
+
+class FNiagaraRibbonMeshCollectorResources : public FOneFrameResource
+{
+public:
+	TSharedRef<FNiagaraRibbonRenderingFrameResources> RibbonResources;
+
+	FNiagaraRibbonMeshCollectorResources()
+		: RibbonResources(new FNiagaraRibbonRenderingFrameResources())
 	{
-		PackedPerRibbonDataByIndex.Add(U0Scale);
-		PackedPerRibbonDataByIndex.Add(U0Offset);
-		PackedPerRibbonDataByIndex.Add(U0DistributionScaler);
-		PackedPerRibbonDataByIndex.Add(U1Scale);
-		PackedPerRibbonDataByIndex.Add(U1Offset);
-		PackedPerRibbonDataByIndex.Add(U1DistributionScaler);
-		PackedPerRibbonDataByIndex.Add(*(float*)&FirstParticleId);
+		
 	}
+};
 
-	void PackSliceVertexData(const FVector2D& Position, const FVector2D& Normal, float TextureV)
+bool FNiagaraRibbonVertexBuffers::InitOrUpdateBuffer(bool bEnabled, FRWBuffer& Buffer, int32& CurrentLength, int32 NeededLength, int32 MaxLength, FRWBuffer(* InitFunction)(int32, ERHIAccess), ERHIAccess InitialAccessFlags)
+{
+	static constexpr float UpsizeMultipler = 1.1f;
+	static constexpr float DownsizeMultiplier = 1.2f;
+
+	bEnabled &= NeededLength > 0;
+
+	if (bEnabled)
+	{
+		check(NeededLength <= MaxLength);
+		
+		// we resize the buffer if it's too small or we're over the allowed free space.
+		if (CurrentLength < NeededLength || CurrentLength > (NeededLength * DownsizeMultiplier))
+		{
+			const int32 NewLength = FMath::Min(MaxLength, FMath::RoundToInt32(NeededLength * UpsizeMultipler));
+
+			Buffer = InitFunction(NewLength, InitialAccessFlags);
+			CurrentLength = NewLength;
+
+			return true;
+		}
+	}
+	else if (CurrentLength != 0)
 	{
-		// Add Position
-		SliceVertexData.Add(Position.X);
-		SliceVertexData.Add(Position.Y);
+		Buffer.Release();
+		CurrentLength = 0;
+		return true;
+	}
 
-		// Add Normal
-		SliceVertexData.Add(Normal.X);
-		SliceVertexData.Add(Normal.Y);
+	return false;
+}
 
-		// Add Texture V
-		SliceVertexData.Add(TextureV);
+void FNiagaraRibbonVertexBuffers::InitializeOrUpdateBuffers(const FNiagaraRibbonGenerationConfig& GenerationConfig, const TSharedPtr<FNiagaraRibbonCPUGeneratedVertexData>& GeneratedGeometryData, const FNiagaraDataBuffer* SourceParticleData, int32 MaxAllocatedCount, bool bIsUsingGPUInit)
+{	
+	const uint32 MaxAllocatedRibbons = GenerationConfig.HasRibbonIDs()? (GenerationConfig.GetMaxNumRibbons() > 0? GenerationConfig.GetMaxNumRibbons() : MaxAllocatedCount / 2) : 1;
+	
+	constexpr ERHIAccess InitialBufferAccessFlags = ERHIAccess::SRVMask | ERHIAccess::VertexOrIndexBuffer;
+	
+	if (!bIsUsingGPUInit)
+	{
+		check(GeneratedGeometryData.IsValid());
+		
+		InitOrUpdateBuffer(true,									SortedIndicesBuffer, SortedIndicesLength, GeneratedGeometryData->SortedIndices.Num(), MaxAllocatedCount, &CreateSortedIndicesBuffer, InitialBufferAccessFlags);
+		InitOrUpdateBuffer(true,									TangentsAndDistancesBuffer, TangentsLength, GeneratedGeometryData->TangentAndDistances.Num(), MaxAllocatedCount, &CreateTangentsAndDistancesBuffer, InitialBufferAccessFlags);
+		InitOrUpdateBuffer(GenerationConfig.HasRibbonIDs(),		MultiRibbonIndicesBuffer, MultiRibbonIndexLength, GeneratedGeometryData->MultiRibbonIndices.Num(), MaxAllocatedCount, &CreateMultiRibbonIndicesBuffer, InitialBufferAccessFlags);
+		InitOrUpdateBuffer(true,									RibbonLookupTableBuffer, RibbonLookupTableLength, GeneratedGeometryData->RibbonInfoLookup.Num(), MaxAllocatedRibbons, &CreateRibbonLookupTableBuffer, InitialBufferAccessFlags);
+		InitOrUpdateBuffer(false,								SegmentsBuffer, SegmentLength, 0, 0, &CreateSegmentsBuffer);
+		bJustCreatedCommandBuffer |= InitOrUpdateBuffer(false,	GPUComputeCommandBuffer, GPUComputeCommandLength, 0, 0, &CreateCommandBuffer);
 	}
+	else
+	{		
+		const uint32 TotalParticles = SourceParticleData->GetNumInstancesAllocated();
+
+		// We assume to have at least 2 particles in each ribbon, therefor we can't have more than instances/2 ribbons
+		const int32 TotalRibbons = FMath::Clamp<int32>(TotalParticles / 2, 1, MaxAllocatedRibbons);
+		
+		InitOrUpdateBuffer(true,									SortedIndicesBuffer, SortedIndicesLength, TotalParticles, MaxAllocatedCount, &CreateSortedIndicesBuffer, InitialBufferAccessFlags);
+		InitOrUpdateBuffer(true,									TangentsAndDistancesBuffer, TangentsLength, TotalParticles, MaxAllocatedCount, &CreateTangentsAndDistancesBuffer, InitialBufferAccessFlags);
+		InitOrUpdateBuffer(GenerationConfig.HasRibbonIDs(),		MultiRibbonIndicesBuffer, MultiRibbonIndexLength, TotalParticles, MaxAllocatedCount, &CreateMultiRibbonIndicesBuffer, InitialBufferAccessFlags);
+		InitOrUpdateBuffer(true,									RibbonLookupTableBuffer, RibbonLookupTableLength, TotalRibbons, MaxAllocatedRibbons, &CreateRibbonLookupTableBuffer, InitialBufferAccessFlags);
+		InitOrUpdateBuffer(true,									SegmentsBuffer, SegmentLength, TotalParticles, MaxAllocatedCount, &CreateSegmentsBuffer, InitialBufferAccessFlags);
+		bJustCreatedCommandBuffer |= InitOrUpdateBuffer(true,	GPUComputeCommandBuffer, GPUComputeCommandLength, FNiagaraRibbonCommandBufferLayout::NumElements, FNiagaraRibbonCommandBufferLayout::NumElements, &CreateCommandBuffer, InitialBufferAccessFlags | ERHIAccess::IndirectArgs);
+	}	
+}
+
+
+FRWBuffer FNiagaraRibbonVertexBuffers::CreateSortedIndicesBuffer(int32 Size, ERHIAccess InitialAccessFlags)
+{
+	FRWBuffer NewBuffer;
+	NewBuffer.Initialize(TEXT("RibbonSortedIndices"), sizeof(uint32), Size, EPixelFormat::PF_R32_UINT, InitialAccessFlags, BUF_Static);
+	return NewBuffer;
+}
+
+FRWBuffer FNiagaraRibbonVertexBuffers::CreateTangentsAndDistancesBuffer(int32 Size, ERHIAccess InitialAccessFlags)
+{
+	FRWBuffer NewBuffer;
+	NewBuffer.Initialize(TEXT("RibbonTangentsAndDistances"), sizeof(float), Size * 4, EPixelFormat::PF_R32_FLOAT, InitialAccessFlags, BUF_Static);
+	return NewBuffer;
+}
+
+FRWBuffer FNiagaraRibbonVertexBuffers::CreateMultiRibbonIndicesBuffer(int32 Size, ERHIAccess InitialAccessFlags)
+{
+	FRWBuffer NewBuffer;
+	NewBuffer.Initialize(TEXT("RibbonMultiRibbonIndices"), sizeof(uint32), Size, EPixelFormat::PF_R32_UINT, InitialAccessFlags, BUF_Static);
+	return NewBuffer;		
+}
+
+FRWBuffer FNiagaraRibbonVertexBuffers::CreateRibbonLookupTableBuffer(int32 Size, ERHIAccess InitialAccessFlags)
+{
+	FRWBuffer NewBuffer;
+	NewBuffer.Initialize(TEXT("RibbonDataLookupTable"), sizeof(uint32), Size * FRibbonMultiRibbonInfoBufferEntry::NumElements, EPixelFormat::PF_R32_UINT, InitialAccessFlags, BUF_Static);
+	return NewBuffer;		
+}
+
+FRWBuffer FNiagaraRibbonVertexBuffers::CreateSegmentsBuffer(int32 Size, ERHIAccess InitialAccessFlags)
+{
+	FRWBuffer NewBuffer;
+	NewBuffer.Initialize(TEXT("RibbonSegments"), sizeof(uint32), Size, EPixelFormat::PF_R32_UINT, InitialAccessFlags, BUF_Static);
+	return NewBuffer;		
+}
+
+FRWBuffer FNiagaraRibbonVertexBuffers::CreateCommandBuffer(int32 Size, ERHIAccess InitialAccessFlags)
+{
+	FRWBuffer NewBuffer;
+	NewBuffer.Initialize(TEXT("RibbonGPUInitCommandBuffer"), sizeof(uint32), FNiagaraRibbonCommandBufferLayout::NumElements, EPixelFormat::PF_R32_UINT, InitialAccessFlags, BUF_DrawIndirect | BUF_Static);
+	return NewBuffer;		
+}
+
+struct FNiagaraRibbonGPUInitParameters
+{
+	const FNiagaraRendererRibbons* Renderer;
+	FNiagaraGpuComputeDispatchInterface* ComputeDispatchInterface;
+	const FNiagaraDataBuffer* SourceParticleData;
+	TWeakPtr<FNiagaraRibbonRenderingFrameResources> RenderingResources;
+	
+	FNiagaraRibbonGPUInitParameters(const FNiagaraRendererRibbons* InRenderer, FNiagaraGpuComputeDispatchInterface* InComputeDispatchInterface,
+		const FNiagaraDataBuffer* InSourceParticleData, const TSharedPtr<FNiagaraRibbonRenderingFrameResources>& InRenderingResources)
+		: Renderer(InRenderer)
+		, ComputeDispatchInterface(InComputeDispatchInterface)
+		, SourceParticleData(InSourceParticleData)
+		, RenderingResources(InRenderingResources)
+	{
+		
+	}	 
+};
+
+struct FNiagaraRibbonGPUInitComputeBuffers
+{
+	FRWBuffer SortBuffer;
+	FRWBuffer TempSegments;
+	FRWBuffer TempDistances;
+	FRWBuffer TempMultiRibbon;
+	FRWBuffer TempTessellationStats[2];
+	
+	FNiagaraRibbonGPUInitComputeBuffers() { }
+
+	void InitOrUpdateBuffers(int32 NeededSize, bool bWantsMultiRibbon, bool bWantsTessellation, bool bWantsTessellationTwist)
+	{
+		// TODO: Downsize these when we haven't needed the size for a bit
+
+		constexpr ERHIAccess InitialAccess = ERHIAccess::SRVMask | ERHIAccess::VertexOrIndexBuffer;
+		
+		if (SortBuffer.NumBytes < NeededSize * sizeof(int32))
+		{
+			SortBuffer.Initialize(TEXT("NiagarGPUInit-SortedIndices"), sizeof(uint32), NeededSize,
+			                      EPixelFormat::PF_R32_UINT, InitialAccess);
+		}
+
+		if (TempSegments.NumBytes < NeededSize * sizeof(int32))
+		{
+			TempSegments.Initialize(TEXT("NiagaraGPUInit-Segments"), sizeof(uint32), NeededSize,
+			                        EPixelFormat::PF_R32_UINT, InitialAccess, BUF_Static);
+		}
+
+		if (TempDistances.NumBytes < NeededSize * sizeof(float) * 4)
+		{
+			TempDistances.Initialize(TEXT("NiagaraGPUInit-Distances"), sizeof(float), NeededSize * 4,
+			                         EPixelFormat::PF_R32_FLOAT, InitialAccess, BUF_Static);
+		}
+
+		const uint32 MultiRibbonBufferSize = NeededSize * (bWantsMultiRibbon? 1 : 0);
+		if (TempMultiRibbon.NumBytes < MultiRibbonBufferSize * sizeof(int32))
+		{
+			TempMultiRibbon.Initialize(
+				TEXT("NiagaraGPUInit-MultiRibbon"), sizeof(uint32), MultiRibbonBufferSize, EPixelFormat::PF_R32_UINT,
+				InitialAccess, BUF_Static);
+		}
+
+		const uint32 TessellationBufferSize = NeededSize * (bWantsTessellation? (bWantsTessellationTwist? FTessellationStatsEntry::NumElements : FTessellationStatsEntryNoTwist::NumElements) : 0);
+		if (TempTessellationStats[0].NumBytes < TessellationBufferSize * sizeof(float))
+		{
+			TempTessellationStats[0].Initialize(TEXT("NiagaraGPUInit-Tessellation-0"), sizeof(float), TessellationBufferSize, EPixelFormat::PF_R32_FLOAT, InitialAccess, BUF_Static);
+			TempTessellationStats[1].Initialize(TEXT("NiagaraGPUInit-Tessellation-1"), sizeof(float), TessellationBufferSize, EPixelFormat::PF_R32_FLOAT, InitialAccess, BUF_Static);
+		}
+	}	
 };
 
-class FNiagaraMeshCollectorResourcesRibbon : public FOneFrameResource
+
+
+class FNiagaraRibbonComputeDispatchManager	
 {
+	friend class FLazySingleton;
+private:
+	TArray<FNiagaraRibbonGPUInitParameters> RenderersToGenerate;
+
+	FGPUSortManager* SortManager;	
+	FDelegateHandle SortManagerEventRef;
+
+	FNiagaraRibbonGPUInitComputeBuffers ComputeBuffers;
+	
+	void RegisterForEvent(FGPUSortManager* InSortManager)
+	{
+		if (!SortManagerEventRef.IsValid())
+		{
+			SortManager = InSortManager;
+			SortManagerEventRef = SortManager->PostPreRenderEvent.AddRaw(this, &FNiagaraRibbonComputeDispatchManager::OnPostPreRender);
+		}
+	}
+
+	void OnPostPreRender(FRHICommandListImmediate& CMDList)
+	{
+		GenerateAllGPUData(CMDList);
+		
+		SortManager->PostPreRenderEvent.Remove(SortManagerEventRef);
+		SortManagerEventRef.Reset();
+		SortManager = nullptr;
+	}
+
+	void GenerateAllGPUData(FRHICommandListImmediate& CMDList);
+
 public:
-	FNiagaraRibbonVertexFactory VertexFactory;
-	FNiagaraRibbonUniformBufferRef UniformBuffer;
+	static FNiagaraRibbonComputeDispatchManager& Get()
+	{
+		return TLazySingleton<FNiagaraRibbonComputeDispatchManager>::Get();
+	}
 
-	virtual ~FNiagaraMeshCollectorResourcesRibbon()
+	void RegisterRenderer(FGPUSortManager* InSortManager, FNiagaraRibbonGPUInitParameters NewRegistration)
 	{
-		VertexFactory.ReleaseResource();
+		RenderersToGenerate.Add(NewRegistration);
+		RegisterForEvent(InSortManager);
+	}
+	
+	void UnRegister(FNiagaraRendererRibbons* Renderer)
+	{		
+		RenderersToGenerate.RemoveAll([Renderer](const FNiagaraRibbonGPUInitParameters& Params) { return Params.Renderer == Renderer; });
 	}
 };
 
+
+
 FNiagaraRendererRibbons::FNiagaraRendererRibbons(ERHIFeatureLevel::Type FeatureLevel, const UNiagaraRendererProperties *InProps, const FNiagaraEmitterInstance* Emitter)
 	: FNiagaraRenderer(FeatureLevel, InProps, Emitter)
+	, GenerationConfig(CastChecked<const UNiagaraRibbonRendererProperties>(InProps))
 	, FacingMode(ENiagaraRibbonFacingMode::Screen)
-	, Shape(ENiagaraRibbonShapeMode::Plane)
-	, bEnableAccurateGeometry(false)
-	, WidthSegmentationCount(1)
-	, MultiPlaneCount(2)
-	, TubeSubdivisions(3)
-	, TessellationMode(ENiagaraRibbonTessellationMode::Automatic)
-	, CustomCurveTension(0.f)
-	, CustomTessellationFactor(16)
-	, bCustomUseConstantFactor(false)
-	, CustomTessellationMinAngle(15.f * PI / 180.f)
-	, bCustomUseScreenSpace(true)
-	, bNeedsPreciseMotionVectors(false)
 {
 	const UNiagaraRibbonRendererProperties* Properties = CastChecked<const UNiagaraRibbonRendererProperties>(InProps);
-	FacingMode = Properties->FacingMode;
+
+	int32 IgnoredFloatOffset, IgnoredHalfOffset;
+	Emitter->GetData().GetVariableComponentOffsets(Properties->RibbonIdBinding.GetDataSetBindableVariable(), IgnoredFloatOffset, RibbonIDParamDataSetOffset, IgnoredHalfOffset);
+
+	// Check we actually have ribbon id if we claim we do
+	check(!GenerationConfig.HasRibbonIDs() || RibbonIDParamDataSetOffset != INDEX_NONE);
+	
 	UV0Settings = Properties->UV0Settings;
 	UV1Settings = Properties->UV1Settings;
+	FacingMode = Properties->FacingMode;
 	DrawDirection = Properties->DrawDirection;
-	Shape = Properties->Shape;
-	bEnableAccurateGeometry = Properties->bEnableAccurateGeometry;
-	WidthSegmentationCount = FMath::Max(Properties->WidthSegmentationCount, 1);
-	MultiPlaneCount = Properties->MultiPlaneCount;
-	TubeSubdivisions = Properties->TubeSubdivisions;
-	CustomVertices = Properties->CustomVertices;
-	TessellationMode = Properties->TessellationMode;
-	CustomCurveTension = FMath::Clamp<float>(Properties->CurveTension, 0.f, 0.9999f);
-	CustomTessellationFactor = Properties->TessellationFactor;
-	bCustomUseConstantFactor = Properties->bUseConstantFactor;
-	CustomTessellationMinAngle = Properties->TessellationAngle > 0.f && Properties->TessellationAngle < 1.f ? 1.f : Properties->TessellationAngle;
-	CustomTessellationMinAngle *= PI / 180.f;
-	bCustomUseScreenSpace = Properties->bScreenSpaceTessellation;
-	MaterialParamValidMask = Properties->MaterialParamValidMask;
 	RendererLayout = &Properties->RendererLayout;
-	bNeedsPreciseMotionVectors = Properties->NeedsPreciseMotionVectors();
+	
+	InitializeShape(Properties);
+	InitializeTessellation(Properties);	
 }
 
 FNiagaraRendererRibbons::~FNiagaraRendererRibbons()
 {
-}
-
-void FNiagaraRendererRibbons::ReleaseRenderThreadResources()
-{
-	FNiagaraRenderer::ReleaseRenderThreadResources();
-#if RHI_RAYTRACING
-	if (IsRayTracingEnabled())
-	{
-		RayTracingGeometry.ReleaseResource();
-		RayTracingDynamicVertexBuffer.Release();
-	}
-#endif
+	FNiagaraRibbonComputeDispatchManager::Get().UnRegister(this);
 }
 
 // FPrimitiveSceneProxy interface.
 void FNiagaraRendererRibbons::CreateRenderThreadResources()
 {
 	FNiagaraRenderer::CreateRenderThreadResources();
+
+
+	{
+		// Initialize the shape vertex buffer. This doesn't change frame-to-frame, so we can set it up once
+		const int32 NumElements = ShapeState.SliceTriangleToVertexIds.Num();
+		ShapeState.SliceTriangleToVertexIdsBuffer.Initialize(TEXT("SliceTriangleToVertexIdsBuffer"), sizeof(uint32), NumElements, EPixelFormat::PF_R32_UINT, BUF_Static);
+		void* SliceTriangleToVertexIdsBufferPtr = RHILockBuffer(ShapeState.SliceTriangleToVertexIdsBuffer.Buffer, 0, sizeof(uint32) * NumElements, RLM_WriteOnly);
+		FMemory::Memcpy(SliceTriangleToVertexIdsBufferPtr, ShapeState.SliceTriangleToVertexIds.GetData(), sizeof(uint32) * NumElements);
+		RHIUnlockBuffer(ShapeState.SliceTriangleToVertexIdsBuffer.Buffer);
+	}
+
+	{
+		// Initialize the shape vertex buffer. This doesn't change frame-to-frame, so we can set it up once
+		const int32 NumElements = ShapeState.SliceVertexData.Num() * FNiagaraRibbonShapeGeometryData::FVertex::NumElements;
+		ShapeState.SliceVertexDataBuffer.Initialize(TEXT("NiagaraShapeVertexDataBuffer"), sizeof(float), NumElements, EPixelFormat::PF_R32_FLOAT, BUF_Static);
+		void* SliceVertexDataBufferPtr = RHILockBuffer(ShapeState.SliceVertexDataBuffer.Buffer, 0, sizeof(float) * NumElements, RLM_WriteOnly);
+		FMemory::Memcpy(SliceVertexDataBufferPtr, ShapeState.SliceVertexData.GetData(), sizeof(float) * NumElements);
+		RHIUnlockBuffer(ShapeState.SliceVertexDataBuffer.Buffer);
+	}
+	
+	
 #if RHI_RAYTRACING
 	if (IsRayTracingEnabled())
 	{
@@ -262,275 +657,534 @@ void FNiagaraRendererRibbons::CreateRenderThreadResources()
 #endif
 }
 
-template <typename TValue>
-TValue* FNiagaraRendererRibbons::AppendToIndexBuffer(
-	TValue* OutIndices,
-	uint32& OutMaxUsedIndex,
-	const TArrayView<int32>& SegmentData,
-	const FRibbonRenderingIndexOffsets& Offsets,
-	int32 InterpCount,
-	bool bInvertOrder) const
+void FNiagaraRendererRibbons::ReleaseRenderThreadResources()
 {
-	TValue MaxIndex = 0;
-	if ( SegmentData.Num() == 0)
+	FNiagaraRenderer::ReleaseRenderThreadResources();
+
+	ShapeState.SliceTriangleToVertexIdsBuffer.Release();
+	ShapeState.SliceVertexDataBuffer.Release();
+	
+#if RHI_RAYTRACING
+	if (IsRayTracingEnabled())
 	{
-		return OutIndices;
+		RayTracingGeometry.ReleaseResource();
+		RayTracingDynamicVertexBuffer.Release();
 	}
+#endif
+}
 
-	// This sets up the first and next vertex for each pair of triangles in the slice.
-	// For a plane this will just be a linear set
-	// For a multiplane it will be multiple separate linear sets
-	// For a tube it will be a linear set that wraps back around to itself,
-	// Same with the custom vertices.
-	TArray<int32, TInlineAllocator<32>> SliceTriangleToVertexIds;
+void FNiagaraRendererRibbons::GetDynamicMeshElements(const TArray<const FSceneView*>& Views, const FSceneViewFamily& ViewFamily, uint32 VisibilityMap, FMeshElementCollector& Collector, const FNiagaraSceneProxy *SceneProxy) const
+{
+	SCOPE_CYCLE_COUNTER(STAT_NiagaraRenderRibbons);
+	PARTICLE_PERF_STAT_CYCLES_RT(SceneProxy->GetProxyDynamicData().PerfStatsContext, GetDynamicMeshElements);
 
-	if (Shape == ENiagaraRibbonShapeMode::MultiPlane)
+	FNiagaraDynamicDataRibbon* DynamicData = static_cast<FNiagaraDynamicDataRibbon*>(DynamicDataRender);
+	if (!DynamicData)
 	{
-		const int32 FrontFaceVertexCount = MultiPlaneCount * (WidthSegmentationCount + 1);
-
-		SliceTriangleToVertexIds.Reserve(WidthSegmentationCount * MultiPlaneCount * (bEnableAccurateGeometry ? 2 : 1));
-		for (int32 PlaneIndex = 0; PlaneIndex < MultiPlaneCount; PlaneIndex++)
-		{
-			const int32 BaseVertexId = (PlaneIndex * (WidthSegmentationCount + 1));
+		return;
+	}
 
-			for (int32 VertexIdx = 0; VertexIdx < WidthSegmentationCount; VertexIdx++)
-			{
-				SliceTriangleToVertexIds.Add(BaseVertexId + VertexIdx);
-				SliceTriangleToVertexIds.Add(BaseVertexId + VertexIdx + 1);
-			}
+	FNiagaraDataBuffer* SourceParticleData = DynamicData->GetParticleDataToRender();
 
-			if (bEnableAccurateGeometry)
-			{
-				for (int32 VertexIdx = 0; VertexIdx < WidthSegmentationCount; VertexIdx++)
-				{
-					SliceTriangleToVertexIds.Add(FrontFaceVertexCount + BaseVertexId + VertexIdx + 1);
-					SliceTriangleToVertexIds.Add(FrontFaceVertexCount + BaseVertexId + VertexIdx);
-				}
-			}
-		}
+	if (GbEnableNiagaraRibbonRendering == 0 || SourceParticleData == nullptr)
+	{
+		return;
 	}
-	else if (Shape == ENiagaraRibbonShapeMode::Tube)
+
+	if (DynamicData->bIsGPUSystem)
 	{
-		SliceTriangleToVertexIds.Reserve(TubeSubdivisions);
-		for (int32 VertexIdx = 0; VertexIdx < TubeSubdivisions; VertexIdx++)
+		// Bail if we don't have enough particle data to have a valid ribbon
+		// or if somehow the sim targets don't match
+		if (SimTarget != ENiagaraSimTarget::GPUComputeSim || SourceParticleData->GetNumInstancesAllocated() < 2)
 		{
-			SliceTriangleToVertexIds.Add(VertexIdx);
-			SliceTriangleToVertexIds.Add(VertexIdx + 1);
+			return;
 		}
 	}
-	else if (Shape == ENiagaraRibbonShapeMode::Custom && CustomVertices.Num() >= 2)
+	else
 	{
-		SliceTriangleToVertexIds.Reserve(CustomVertices.Num());
-		for (int32 VertexIdx = 0; VertexIdx < CustomVertices.Num(); VertexIdx++)
+		check(SimTarget == ENiagaraSimTarget::CPUSim);
+
+		if (SourceParticleData->GetNumInstances() < 2)
 		{
-			SliceTriangleToVertexIds.Add(VertexIdx);
-			SliceTriangleToVertexIds.Add(VertexIdx + 1);
+			// Bail if we don't have enough particle data to have a valid ribbon
+			return;
 		}
-	}
-	else // Plane
-	{
-		SliceTriangleToVertexIds.Reserve(WidthSegmentationCount);
-		for (int32 VertexIdx = 0; VertexIdx < WidthSegmentationCount; VertexIdx++)
+
+		if (!DynamicData->bUseGPUInit && (
+			!DynamicData->GenerationOutput.IsValid() ||
+			DynamicData->GenerationOutput->SegmentData.Num() <= 0))
 		{
-			SliceTriangleToVertexIds.Add(VertexIdx);
-			SliceTriangleToVertexIds.Add(VertexIdx + 1);
+			return;
 		}
 	}
 
-	int32 SegmentDataIndex = bInvertOrder ? SegmentData.Num() - 1 : 0;
-	const int32 LastSegmentDataIndex = bInvertOrder ? -1 : SegmentData.Num();
-	const int32 SegmentDataIndexInc = bInvertOrder ? -1 : 1;
-	const int32 FlipGeometryIndex = SliceTriangleToVertexIds.Num() / 2;
+#if STATS
+	FScopeCycleCounter EmitterStatsCounter(EmitterStatID);
+#endif
+	
+	auto* ComputeDispatchInterface = SceneProxy->GetComputeDispatchInterface();
 
-	while ( SegmentDataIndex != LastSegmentDataIndex )
+	const FNiagaraRibbonMeshCollectorResources& RenderingResources = Collector.AllocateOneFrameResource<FNiagaraRibbonMeshCollectorResources>();
+		
+	InitializeVertexBuffersResources(DynamicData, SourceParticleData, Collector.GetDynamicReadBuffer(), RenderingResources.RibbonResources, DynamicData->bUseGPUInit);
+	
+	// Compute the per-view uniform buffers.
+	for (int32 ViewIndex = 0; ViewIndex < Views.Num(); ViewIndex++)
 	{
-		const int32 SegmentIndex = SegmentData[SegmentDataIndex];
-		for (int32 SubSegmentIndex = 0; SubSegmentIndex < InterpCount; ++SubSegmentIndex)
+		if (VisibilityMap & (1 << ViewIndex))
 		{
-			const bool bIsFinalInterp = SubSegmentIndex == InterpCount - 1;
-
-			const int32 ThisSegmentOffset = SegmentIndex << Offsets.SegmentBitShift;
-			const int32 NextSegmentOffset = (SegmentIndex + (bIsFinalInterp ? 1 : 0)) << Offsets.SegmentBitShift;
+			const FSceneView* View = Views[ViewIndex];
+			check(View);
 
-			const int32 ThisSubSegmentOffset = SubSegmentIndex << Offsets.InterpBitShift;
-			const int32 NextSubSegmentOffset = (bIsFinalInterp ? 0 : SubSegmentIndex + 1) << Offsets.InterpBitShift;
+			if (View->bIsInstancedStereoEnabled && IStereoRendering::IsStereoEyeView(*View) && !IStereoRendering::IsAPrimaryView(*View))
+			{
+				// We don't have to generate batches for non-primary views in stereo instance rendering
+				continue;
+			}
+			
+			FMeshBatch& MeshBatch = Collector.AllocateMesh();
+			
+			const FVector ViewOriginForDistanceCulling = View->ViewMatrices.GetViewOrigin();
+
+			auto& RenderingViewResources = RenderingResources.RibbonResources->ViewResources.Add_GetRef(MakeShared<FNiagaraRibbonRenderingFrameViewResources>());
+			RenderingViewResources->IndexGenerationSettings = CalculateIndexBufferConfiguration(DynamicData->GenerationOutput, SourceParticleData, SceneProxy, View, ViewOriginForDistanceCulling, DynamicData->bUseGPUInit, DynamicData->bIsGPUSystem);
+			
+			GenerateIndexBufferForView(RenderingViewResources->IndexGenerationSettings, DynamicData, RenderingViewResources, View, ViewOriginForDistanceCulling, DynamicData->bUseGPUInit);
+			
+			SetupPerViewUniformBuffer(RenderingViewResources->IndexGenerationSettings, View, ViewFamily, SceneProxy, RenderingViewResources->UniformBuffer);
+			
+			SetupMeshBatchAndCollectorResourceForView(RenderingViewResources->IndexGenerationSettings, DynamicData, SourceParticleData, View, ViewFamily, SceneProxy, RenderingResources.RibbonResources, RenderingViewResources, MeshBatch, DynamicData->bUseGPUInit);
 
-			const int32 CurrSegment = ThisSegmentOffset | ThisSubSegmentOffset;
-			const int32 NextSegment = NextSegmentOffset | NextSubSegmentOffset;
-
-			int32 TriangleId = 0;
-
-			for (; TriangleId < FlipGeometryIndex; TriangleId += 2)
-			{
-				const int32 FirstIndex = SliceTriangleToVertexIds[TriangleId];
-				const int32 SecondIndex = SliceTriangleToVertexIds[TriangleId + 1];
+			Collector.AddMesh(ViewIndex, MeshBatch);
+		}
+	}
 
-				OutIndices[0] = CurrSegment | FirstIndex;
-				MaxIndex = FMath::Max<TValue>(MaxIndex, OutIndices[0]);
+	// Register this renderer for generation this frame if we're a gpu system or using gpu init
+	if (DynamicData->bUseGPUInit || DynamicData->bIsGPUSystem)
+	{
+		FGPUSortManager* SortManager = ComputeDispatchInterface->GetGPUSortManager();
+		FNiagaraRibbonComputeDispatchManager::Get().RegisterRenderer(SortManager, FNiagaraRibbonGPUInitParameters(this, ComputeDispatchInterface, SourceParticleData, RenderingResources.RibbonResources));
+	}
+}
 
-				OutIndices[1] = CurrSegment | SecondIndex;
-				MaxIndex = FMath::Max<TValue>(MaxIndex, OutIndices[1]);
+FNiagaraDynamicDataBase* FNiagaraRendererRibbons::GenerateDynamicData(const FNiagaraSceneProxy* Proxy, const UNiagaraRendererProperties* InProperties, const FNiagaraEmitterInstance* Emitter)const
+{
+	SCOPE_CYCLE_COUNTER(STAT_NiagaraGenRibbonVertexData);
+	check(Emitter && Emitter->GetParentSystemInstance());
 
-				OutIndices[2] = NextSegment | FirstIndex;
-				MaxIndex = FMath::Max<TValue>(MaxIndex, OutIndices[2]);
+	FNiagaraDynamicDataRibbon* DynamicData = nullptr;
+	const UNiagaraRibbonRendererProperties* Properties = CastChecked<const UNiagaraRibbonRendererProperties>(InProperties);
 
-				OutIndices[3] = OutIndices[1];
+	if (Properties)
+	{
+		if (!IsRendererEnabled(Properties, Emitter))
+		{
+			return nullptr;
+		}
 
-				OutIndices[4] = NextSegment | SecondIndex;
-				MaxIndex = FMath::Max<TValue>(MaxIndex, OutIndices[4]);
+		if (InProperties->bAllowInCullProxies == false &&
+			Cast<UNiagaraCullProxyComponent>(Emitter->GetParentSystemInstance()->GetAttachComponent()) != nullptr)
+		{
+			return nullptr;
+		}
 
-				OutIndices[5] = OutIndices[2];
+		FNiagaraDataBuffer* DataToRender = Emitter->GetData().GetCurrentData();		
+		if(SimTarget == ENiagaraSimTarget::GPUComputeSim || (DataToRender != nullptr && DataToRender->GetNumInstances() > 1))
+		{
+			DynamicData = new FNiagaraDynamicDataRibbon(Emitter);
+	
+			//In preparation for a material override feature, we pass our material(s) and relevance in via dynamic data.
+			//The renderer ensures we have the correct usage and relevance for materials in BaseMaterials_GT.
+			//Any override feature must also do the same for materials that are set.
+			check(BaseMaterials_GT.Num() == 1);
+			check(BaseMaterials_GT[0]->CheckMaterialUsage_Concurrent(MATUSAGE_NiagaraRibbons));
+			DynamicData->Material = BaseMaterials_GT[0]->GetRenderProxy();
+			DynamicData->SetMaterialRelevance(BaseMaterialRelevance_GT);
+		}
 
-				OutIndices += 6;
-			}
-			for (; TriangleId < SliceTriangleToVertexIds.Num(); TriangleId += 2)
+		if (DynamicData)
+		{		
+			// We always run GPU init when we're a GPU system
+			const bool bIsGPUSystem = SimTarget == ENiagaraSimTarget::GPUComputeSim;
+			
+			// We disable compute initialization when compute isn't available or they're CVar'd off
+			const bool bCanUseComputeGenForCPUSystems = (GbEnableNiagaraRibbonComputeShaderGen && FNiagaraUtilities::AllowComputeShaders(GShaderPlatformForFeatureLevel[FeatureLevel]));
+						
+			const bool bWantsGPUInit = (bCanUseComputeGenForCPUSystems && (GbForceNiagaraRibbonGPUInit || Properties->bUseGPUInit));
+			
+			DynamicData->bUseGPUInit = bIsGPUSystem || bWantsGPUInit;
+			DynamicData->bIsGPUSystem = bIsGPUSystem;
+			DynamicData->MaxAllocatedParticleCount = Emitter->GetData().GetMaxAllocationCount();
+			
+			if (!DynamicData->bUseGPUInit)
 			{
-				const int32 FirstIndex = SliceTriangleToVertexIds[TriangleId];
-				const int32 SecondIndex = SliceTriangleToVertexIds[TriangleId + 1];
-
-				OutIndices[0] = CurrSegment | FirstIndex;
-				MaxIndex = FMath::Max<TValue>(MaxIndex, OutIndices[0]);
-
-				OutIndices[1] = CurrSegment | SecondIndex;
-				MaxIndex = FMath::Max<TValue>(MaxIndex, OutIndices[1]);
-
-				OutIndices[2] = NextSegment | SecondIndex;
-				MaxIndex = FMath::Max<TValue>(MaxIndex, OutIndices[2]);
+				const FNiagaraGenerationInputDataCPUAccessors CPUData(Properties, Emitter->GetData());
+				
+				DynamicData->GenerationOutput = MakeShared<FNiagaraRibbonCPUGeneratedVertexData>();
 
-				OutIndices[3] = OutIndices[0];
-
-				OutIndices[4] = OutIndices[2];
-
-				OutIndices[5] = NextSegment | FirstIndex;
-				MaxIndex = FMath::Max<TValue>(MaxIndex, OutIndices[5]);
-
-				OutIndices += 6;
+				if (CPUData.PosData.IsValid() && CPUData.SortKeyReader.IsValid() && CPUData.TotalNumParticles >= 2)
+				{
+					GenerateVertexBufferCPU(CPUData, *DynamicData->GenerationOutput);
+				}
+				else
+				{
+					// We don't have valid data so remove the dynamic data
+					delete DynamicData;
+					DynamicData = nullptr;
+				}
 			}
 		}
 
-		SegmentDataIndex += SegmentDataIndexInc;
+		if (DynamicData && Properties->MaterialParameterBindings.Num() != 0)
+		{
+			ProcessMaterialParameterBindings(MakeArrayView(Properties->MaterialParameterBindings), Emitter, MakeArrayView(BaseMaterials_GT));
+		}
 	}
-
-	OutMaxUsedIndex = MaxIndex;
-	return OutIndices;
+	
+	return DynamicData;
 }
 
-template <typename TValue>
-void FNiagaraRendererRibbons::GenerateIndexBuffer(
-	FGlobalDynamicIndexBuffer::FAllocationEx& InOutIndexAllocation,
-	const FRibbonRenderingIndexOffsets& Offsets,
-	int32 InterpCount,
-	const FVector& ViewDirection,
-	const FVector& ViewOriginForDistanceCulling,
-	FNiagaraDynamicDataRibbon* DynamicData) const
+int FNiagaraRendererRibbons::GetDynamicDataSize()const
 {
-	check(DynamicData);
+	uint32 Size = sizeof(FNiagaraDynamicDataRibbon);
 
-	FMaterialRenderProxy* MaterialRenderProxy = DynamicData->Material;
-	check(MaterialRenderProxy);
-	const EBlendMode BlendMode = MaterialRenderProxy->GetIncompleteMaterialWithFallback(FeatureLevel).GetBlendMode();
+	Size += ShapeState.SliceVertexData.GetAllocatedSize();
 
-	TValue* CurrentIndexBuffer = (TValue*)InOutIndexAllocation.Buffer;
-	if (IsTranslucentBlendMode(BlendMode) && DynamicData->MultiRibbonInfos.Num())
-	{
-		for (const FNiagaraDynamicDataRibbon::FMultiRibbonInfo& MultiRibbonInfo : DynamicData->MultiRibbonInfos)
-		{
-			TArrayView<int32> CurrentSegmentData(DynamicData->SegmentData.GetData() + MultiRibbonInfo.BaseSegmentDataIndex, MultiRibbonInfo.NumSegmentDataIndices);
-			CurrentIndexBuffer = AppendToIndexBuffer(CurrentIndexBuffer, InOutIndexAllocation.MaxUsedIndex, CurrentSegmentData, Offsets, InterpCount, MultiRibbonInfo.UseInvertOrder(ViewDirection, ViewOriginForDistanceCulling, DrawDirection));
-		}
-	}
-	else // Otherwise ignore multi ribbon ordering.
+	if (DynamicDataRender)
 	{
-		TArrayView<int32> CurrentSegmentData(DynamicData->SegmentData.GetData(), DynamicData->SegmentData.Num());
-		CurrentIndexBuffer = AppendToIndexBuffer(CurrentIndexBuffer, InOutIndexAllocation.MaxUsedIndex, CurrentSegmentData, Offsets, InterpCount, false);
+		const FNiagaraDynamicDataRibbon* RibbonDynamicData = static_cast<FNiagaraDynamicDataRibbon*>(DynamicDataRender);
+		Size += RibbonDynamicData->GetAllocatedSize();
 	}
+	return Size;
 }
 
-void FNiagaraRendererRibbons::GetDynamicMeshElements(const TArray<const FSceneView*>& Views, const FSceneViewFamily& ViewFamily, uint32 VisibilityMap, FMeshElementCollector& Collector, const FNiagaraSceneProxy *SceneProxy) const
+bool FNiagaraRendererRibbons::IsMaterialValid(const UMaterialInterface* Mat)const
 {
-	PARTICLE_PERF_STAT_CYCLES_RT(SceneProxy->GetProxyDynamicData().PerfStatsContext, GetDynamicMeshElements);
+	return Mat && Mat->CheckMaterialUsage_Concurrent(MATUSAGE_NiagaraRibbons);
+}
 
+#if RHI_RAYTRACING
+void FNiagaraRendererRibbons::GetDynamicRayTracingInstances(FRayTracingMaterialGatheringContext& Context, TArray<FRayTracingInstance>& OutRayTracingInstances, const FNiagaraSceneProxy* SceneProxy)
+{
+	if (!CVarRayTracingNiagaraRibbons.GetValueOnRenderThread())
+	{
+		return;
+	}
+	
+	SCOPE_CYCLE_COUNTER(STAT_NiagaraRenderRibbons);
+	check(SceneProxy);
+	
 	FNiagaraDynamicDataRibbon *DynamicDataRibbon = static_cast<FNiagaraDynamicDataRibbon*>(DynamicDataRender);
-	if (!DynamicDataRibbon)
+	FNiagaraGpuComputeDispatchInterface* ComputeDispatchInterface = SceneProxy->GetComputeDispatchInterface();
+	
+	if (!ComputeDispatchInterface || !DynamicDataRibbon)
 	{
 		return;
 	}
 
 	FNiagaraDataBuffer* SourceParticleData = DynamicDataRibbon->GetParticleDataToRender();
-	if (!SourceParticleData ||
-		SourceParticleData->GetNumInstances() < 2 ||
-		DynamicDataRibbon->SegmentData.Num() == 0 ||
-		GbEnableNiagaraRibbonRendering == 0)
+
+	if (GbEnableNiagaraRibbonRendering == 0 || SourceParticleData == nullptr)
 	{
 		return;
 	}
 
-#if STATS
-	FScopeCycleCounter EmitterStatsCounter(EmitterStatID);
+	if (DynamicDataRibbon->bIsGPUSystem)
+	{
+		// Bail if we don't have enough particle data to have a valid ribbon
+		// or if somehow the sim targets don't match
+		if (SimTarget != ENiagaraSimTarget::GPUComputeSim || SourceParticleData->GetNumInstancesAllocated() < 2)
+		{
+			return;
+		}
+	}
+	else
+	{
+		check(SimTarget == ENiagaraSimTarget::CPUSim);
+
+		if (SourceParticleData->GetNumInstances() < 2)
+		{
+			// Bail if we don't have enough particle data to have a valid ribbon
+			return;
+		}
+
+		if (!DynamicDataRibbon->bUseGPUInit && (
+			!DynamicDataRibbon->GenerationOutput.IsValid() ||
+			DynamicDataRibbon->GenerationOutput->SegmentData.Num() <= 0))
+		{
+			return;
+		}
+	}
+	
+	auto& View = Context.ReferenceView;
+	auto& ViewFamily = Context.ReferenceViewFamily;
+	// Setup material for our ray tracing instance
+	
+	const FVector ViewOriginForDistanceCulling = View->ViewMatrices.GetViewOrigin();
+	
+	FNiagaraRibbonMeshCollectorResources& RenderingResources = Context.RayTracingMeshResourceCollector.AllocateOneFrameResource<FNiagaraRibbonMeshCollectorResources>();
+	auto& RenderingViewResources = RenderingResources.RibbonResources->ViewResources.Add_GetRef(MakeShared<FNiagaraRibbonRenderingFrameViewResources>());
+	RenderingViewResources->IndexGenerationSettings = CalculateIndexBufferConfiguration(DynamicDataRibbon->GenerationOutput, SourceParticleData, SceneProxy, View, ViewOriginForDistanceCulling, DynamicDataRibbon->bUseGPUInit, DynamicDataRibbon->bIsGPUSystem);
+	
+	if (!RenderingViewResources->VertexFactory.GetType()->SupportsRayTracingDynamicGeometry())
+	{
+		return;
+	}
+	
+	InitializeVertexBuffersResources(DynamicDataRibbon, SourceParticleData, Context.RayTracingMeshResourceCollector.GetDynamicReadBuffer(), RenderingResources.RibbonResources, DynamicDataRibbon->bUseGPUInit);
+	
+	GenerateIndexBufferForView(RenderingViewResources->IndexGenerationSettings, DynamicDataRibbon, RenderingViewResources, View, ViewOriginForDistanceCulling, DynamicDataRibbon->bUseGPUInit);
+			
+	SetupPerViewUniformBuffer(RenderingViewResources->IndexGenerationSettings, View, ViewFamily, SceneProxy, RenderingViewResources->UniformBuffer);
+	
+	if (RenderingViewResources->IndexGenerationSettings.TotalNumIndices <= 0)
+	{
+		return;
+	}
+	
+	FRayTracingInstance RayTracingInstance;
+	RayTracingInstance.Geometry = &RayTracingGeometry;
+	RayTracingInstance.InstanceTransforms.Add(FMatrix::Identity);
+	
+	RayTracingGeometry.Initializer.IndexBuffer = RenderingViewResources->IndexBuffer.IndexBufferRHI;// PerViewGeneratedData.IndexAllocation.IndexBuffer->IndexBufferRHI;
+	RayTracingGeometry.Initializer.IndexBufferOffset = 0;//PerViewGeneratedData.IndexAllocation.FirstIndex * PerViewGeneratedData.IndexAllocation.IndexStride;
+	
+	FMeshBatch MeshBatch;
+	
+	SetupMeshBatchAndCollectorResourceForView(RenderingViewResources->IndexGenerationSettings, DynamicDataRibbon, SourceParticleData, View, ViewFamily, SceneProxy, RenderingResources.RibbonResources, RenderingViewResources, MeshBatch, DynamicDataRibbon->bUseGPUInit);
+
+	RayTracingInstance.Materials.Add(MeshBatch);
+	
+	// Use the internal vertex buffer only when initialized otherwise used the shared vertex buffer - needs to be updated every frame
+	FRWBuffer* VertexBuffer = RayTracingDynamicVertexBuffer.NumBytes > 0 ? &RayTracingDynamicVertexBuffer : nullptr;
+	
+	const uint32 VertexCount = DynamicDataRibbon->bUseGPUInit? SourceParticleData->GetNumInstances() : DynamicDataRibbon->GenerationOutput->SortedIndices.Num();
+
+	
+	const int32 MaxTriangleCount = RenderingViewResources->IndexGenerationSettings.MaxSegmentCount * RenderingViewResources->IndexGenerationSettings.SubSegmentCount * ShapeState.TrianglesPerSegment;
+	
+	Context.DynamicRayTracingGeometriesToUpdate.Add(
+		FRayTracingDynamicGeometryUpdateParams
+		{
+			RayTracingInstance.Materials,
+			DynamicDataRibbon->bUseGPUInit,
+			VertexCount,
+			VertexCount * static_cast<uint32>(sizeof(FVector3f)),
+			DynamicDataRibbon->bUseGPUInit? MaxTriangleCount : MeshBatch.Elements[0].NumPrimitives,
+			&RayTracingGeometry,
+			VertexBuffer,
+			true
+		}
+	);
+	
+	RayTracingInstance.BuildInstanceMaskAndFlags(FeatureLevel);
+	
+	OutRayTracingInstances.Add(RayTracingInstance);
+}
 #endif
 
-	// Compute the per-view uniform buffers.
-	for (int32 ViewIndex = 0; ViewIndex < Views.Num(); ViewIndex++)
+void FNiagaraRendererRibbons::GenerateShapeStateMultiPlane(FNiagaraRibbonShapeGeometryData& State, int32 MultiPlaneCount, int32 WidthSegmentationCount, bool bEnableAccurateGeometry)
+{
+	State.Shape = ENiagaraRibbonShapeMode::MultiPlane;
+	State.bDisableBackfaceCulling = !bEnableAccurateGeometry;
+	State.bShouldFlipNormalToView = !bEnableAccurateGeometry;
+	State.TrianglesPerSegment = 2 * MultiPlaneCount * WidthSegmentationCount * (bEnableAccurateGeometry? 2 : 1);
+	State.NumVerticesInSlice = MultiPlaneCount * (WidthSegmentationCount + 1) * (bEnableAccurateGeometry ? 2 : 1);
+	State.BitsNeededForShape = CalculateBitsForRange(State.NumVerticesInSlice);
+	State.BitMaskForShape = CalculateBitMask(State.BitsNeededForShape);
+	
+	for (int32 PlaneIndex = 0; PlaneIndex < MultiPlaneCount; PlaneIndex++)
 	{
-		if (VisibilityMap & (1 << ViewIndex))
+		const float RotationAngle = (static_cast<float>(PlaneIndex) / MultiPlaneCount) * 180.0f;
+
+		for (int32 VertexId = 0; VertexId <= WidthSegmentationCount; VertexId++)
 		{
-			const FSceneView* View = Views[ViewIndex];
-			check(View);
+			const FVector2f Position = FVector2f((static_cast<float>(VertexId) / WidthSegmentationCount) - 0.5f, 0).GetRotated(RotationAngle);
+			const FVector2f Normal = FVector2f(0, 1).GetRotated(RotationAngle);
+			const float TextureV = static_cast<float>(VertexId) / WidthSegmentationCount;
 
-			if (View->bIsInstancedStereoEnabled && IStereoRendering::IsStereoEyeView(*View) && !IStereoRendering::IsAPrimaryView(*View))
+			State.SliceVertexData.Emplace(Position, Normal, TextureV);
+		}
+	}
+
+	if (bEnableAccurateGeometry)
+	{
+		for (int32 PlaneIndex = 0; PlaneIndex < MultiPlaneCount; PlaneIndex++)
+		{
+			const float RotationAngle = (static_cast<float>(PlaneIndex) / MultiPlaneCount) * 180.0f;
+
+			for (int32 VertexId = 0; VertexId <= WidthSegmentationCount; VertexId++)
 			{
-				// We don't have to generate batches for non-primary views in stereo instance rendering
-				continue;
+				const FVector2f Position = FVector2f((static_cast<float>(VertexId) / WidthSegmentationCount) - 0.5f, 0).GetRotated(RotationAngle);
+				const FVector2f Normal = FVector2f(0, -1).GetRotated(RotationAngle);
+				const float TextureV = static_cast<float>(VertexId) / WidthSegmentationCount;
+
+				State.SliceVertexData.Emplace(Position, Normal, TextureV);
 			}
+		}
+	}
 
-			FMeshBatch& MeshBatch = Collector.AllocateMesh();
 
-			FGlobalDynamicIndexBuffer::FAllocationEx DynamicIndexAllocation;
-			FNiagaraMeshCollectorResourcesRibbon& CollectorResources = Collector.AllocateOneFrameResource<FNiagaraMeshCollectorResourcesRibbon>();
+	const int32 FrontFaceVertexCount = MultiPlaneCount * (WidthSegmentationCount + 1);
 
-			CreatePerViewResources(View, ViewFamily, SceneProxy, Collector, CollectorResources.UniformBuffer, DynamicIndexAllocation);
+	State.SliceTriangleToVertexIds.Reserve(WidthSegmentationCount * MultiPlaneCount * (bEnableAccurateGeometry ? 2 : 1));
+	for (int32 PlaneIndex = 0; PlaneIndex < MultiPlaneCount; PlaneIndex++)
+	{
+		const int32 BaseVertexId = (PlaneIndex * (WidthSegmentationCount + 1));
 
-			SetupMeshBatchAndCollectorResourceForView(View, ViewFamily, SceneProxy, Collector, DynamicDataRibbon, DynamicIndexAllocation, MeshBatch, CollectorResources);
+		for (int32 VertexIdx = 0; VertexIdx < WidthSegmentationCount; VertexIdx++)
+		{
+			State.SliceTriangleToVertexIds.Add(BaseVertexId + VertexIdx);
+			State.SliceTriangleToVertexIds.Add(BaseVertexId + VertexIdx + 1);
+		}
 
-			Collector.AddMesh(ViewIndex, MeshBatch);
+		if (bEnableAccurateGeometry)
+		{
+			for (int32 VertexIdx = 0; VertexIdx < WidthSegmentationCount; VertexIdx++)
+			{
+				State.SliceTriangleToVertexIds.Add(FrontFaceVertexCount + BaseVertexId + VertexIdx + 1);
+				State.SliceTriangleToVertexIds.Add(FrontFaceVertexCount + BaseVertexId + VertexIdx);
+			}
 		}
 	}
 }
 
-int FNiagaraRendererRibbons::GetDynamicDataSize()const
+void FNiagaraRendererRibbons::GenerateShapeStateTube(FNiagaraRibbonShapeGeometryData& State, int32 TubeSubdivisions)
 {
-	uint32 Size = sizeof(FNiagaraDynamicDataRibbon);
-	if (DynamicDataRender)
+	State.Shape = ENiagaraRibbonShapeMode::Tube;
+	State.bDisableBackfaceCulling = false;
+	State.bShouldFlipNormalToView = true;
+	State.TrianglesPerSegment = 2 * TubeSubdivisions;
+	State.NumVerticesInSlice = TubeSubdivisions + 1;
+	State.BitsNeededForShape = CalculateBitsForRange(State.NumVerticesInSlice);
+	State.BitMaskForShape = CalculateBitMask(State.BitsNeededForShape);
+	
+	for (int32 VertexId = 0; VertexId <= TubeSubdivisions; VertexId++)
+	{
+		const float RotationAngle = (static_cast<float>(VertexId) / TubeSubdivisions) * -360.0f;
+		const FVector2f Position = FVector2f(-0.5f, 0.0f).GetRotated(RotationAngle);
+		const FVector2f Normal = FVector2f(-1, 0).GetRotated(RotationAngle);
+		const float TextureV = static_cast<float>(VertexId) / TubeSubdivisions;
+
+		State.SliceVertexData.Emplace(Position, Normal, TextureV);
+	}
+	
+	State.SliceTriangleToVertexIds.Reserve(TubeSubdivisions);
+	for (int32 VertexIdx = 0; VertexIdx < TubeSubdivisions; VertexIdx++)
 	{
-		FNiagaraDynamicDataRibbon* RibbonDynamicData = static_cast<FNiagaraDynamicDataRibbon*>(DynamicDataRender);
-		Size += RibbonDynamicData->SegmentData.GetAllocatedSize();
-		Size += RibbonDynamicData->SortedIndices.GetAllocatedSize();
-		Size += RibbonDynamicData->TangentAndDistances.GetAllocatedSize();
-		Size += RibbonDynamicData->MultiRibbonIndices.GetAllocatedSize();
-		Size += RibbonDynamicData->PackedPerRibbonDataByIndex.GetAllocatedSize();
-		Size += RibbonDynamicData->SliceVertexData.GetAllocatedSize();
+		State.SliceTriangleToVertexIds.Add(VertexIdx);
+		State.SliceTriangleToVertexIds.Add(VertexIdx + 1);
 	}
+}
 
-	return Size;
+void FNiagaraRendererRibbons::GenerateShapeStateCustom(FNiagaraRibbonShapeGeometryData& State, const TArray<FNiagaraRibbonShapeCustomVertex>& CustomVertices)
+{
+	State.Shape = ENiagaraRibbonShapeMode::Custom;
+	State.bDisableBackfaceCulling = false;
+	State.bShouldFlipNormalToView = true;
+	State.TrianglesPerSegment = 2 * CustomVertices.Num();
+	State.NumVerticesInSlice = CustomVertices.Num() + 1;
+	State.BitsNeededForShape = CalculateBitsForRange(State.NumVerticesInSlice);
+	State.BitMaskForShape = CalculateBitMask(State.BitsNeededForShape);
+	
+	bool bHasCustomUVs = false;
+	for (int32 VertexId = 0; VertexId < CustomVertices.Num(); VertexId++)
+	{
+		if (!FMath::IsNearlyZero(CustomVertices[VertexId].TextureV))
+		{
+			bHasCustomUVs = true;
+			break;
+		}
+	}
+
+	for (int32 VertexId = 0; VertexId <= CustomVertices.Num(); VertexId++)
+	{
+		const auto& CustomVert = CustomVertices[VertexId % CustomVertices.Num()];
+
+		const FVector2f Position = CustomVert.Position;
+		const FVector2f Normal = CustomVert.Normal.IsNearlyZero() ? Position.GetSafeNormal() : CustomVert.Normal;
+		const float TextureV = bHasCustomUVs ? CustomVert.TextureV : static_cast<float>(VertexId) / CustomVertices.Num();
+
+		State.SliceVertexData.Emplace(Position, Normal, TextureV);
+	}
+
+	State.SliceTriangleToVertexIds.Reserve(CustomVertices.Num());
+	for (int32 VertexIdx = 0; VertexIdx < CustomVertices.Num(); VertexIdx++)
+	{
+		State.SliceTriangleToVertexIds.Add(VertexIdx);
+		State.SliceTriangleToVertexIds.Add(VertexIdx + 1);
+	}
+}
+
+void FNiagaraRendererRibbons::GenerateShapeStatePlane(FNiagaraRibbonShapeGeometryData& State, int32 WidthSegmentationCount)
+{
+	State.Shape = ENiagaraRibbonShapeMode::Plane;
+	State.bDisableBackfaceCulling = true;
+	State.bShouldFlipNormalToView = true;
+	State.TrianglesPerSegment = 2 * WidthSegmentationCount;
+	State.NumVerticesInSlice = WidthSegmentationCount + 1;
+	State.BitsNeededForShape = CalculateBitsForRange(State.NumVerticesInSlice);
+	State.BitMaskForShape = CalculateBitMask(State.BitsNeededForShape);
+	
+	for (int32 VertexId = 0; VertexId <= WidthSegmentationCount; VertexId++)
+	{
+		const FVector2f Position = FVector2f((static_cast<float>(VertexId) / WidthSegmentationCount) - 0.5f, 0);
+		const FVector2f Normal = FVector2f(0, 1);
+		const float TextureV = static_cast<float>(VertexId) / WidthSegmentationCount;
+
+		State.SliceVertexData.Emplace(Position, Normal, TextureV);
+	}
+	
+	State.SliceTriangleToVertexIds.Reserve(WidthSegmentationCount);
+	for (int32 VertexIdx = 0; VertexIdx < WidthSegmentationCount; VertexIdx++)
+	{
+		State.SliceTriangleToVertexIds.Add(VertexIdx);
+		State.SliceTriangleToVertexIds.Add(VertexIdx + 1);
+	}
+}
+
+void FNiagaraRendererRibbons::InitializeShape(const UNiagaraRibbonRendererProperties* Properties)
+{
+	if (Properties->Shape == ENiagaraRibbonShapeMode::Custom && Properties->CustomVertices.Num() > 2)
+	{
+		GenerateShapeStateCustom(ShapeState, Properties->CustomVertices);
+	}
+	else if (Properties->Shape == ENiagaraRibbonShapeMode::Tube && Properties->TubeSubdivisions > 2 && Properties->TubeSubdivisions <= 16)
+	{
+		GenerateShapeStateTube(ShapeState, Properties->TubeSubdivisions);
+	}
+	else if (Properties->Shape == ENiagaraRibbonShapeMode::MultiPlane && Properties->MultiPlaneCount > 1 && Properties->MultiPlaneCount <= 16)
+	{
+		GenerateShapeStateMultiPlane(ShapeState, Properties->MultiPlaneCount, Properties->WidthSegmentationCount, Properties->bEnableAccurateGeometry);
+	}
+	else
+	{
+		GenerateShapeStatePlane(ShapeState, Properties->WidthSegmentationCount);
+	}	
+}
+
+void FNiagaraRendererRibbons::InitializeTessellation(const UNiagaraRibbonRendererProperties* Properties)
+{
+	TessellationConfig.TessellationMode = Properties->TessellationMode;
+	TessellationConfig.CustomTessellationFactor = Properties->TessellationFactor;
+	TessellationConfig.bCustomUseConstantFactor = Properties->bUseConstantFactor;
+	TessellationConfig.CustomTessellationMinAngle = Properties->TessellationAngle > 0.f && Properties->TessellationAngle < 1.f ? 1.f : Properties->TessellationAngle;
+	TessellationConfig.CustomTessellationMinAngle *= PI / 180.f;
+	TessellationConfig.bCustomUseScreenSpace = Properties->bScreenSpaceTessellation;
 }
 
-void CalculateUVScaleAndOffsets(
-	const FNiagaraRibbonUVSettings& UVSettings,
-	const TArray<int32>& RibbonIndices,
-	const TArray<FVector4f>& RibbonTangentsAndDistances,
-	const FNiagaraDataSetReaderFloat<float>& NormalizedAgeReader,
-	int32 StartIndex, int32 EndIndex,
-	int32 NumSegments, float TotalLength,
-	float& OutUScale, float& OutUOffset, float& OutUDistributionScaler)
+template<typename IntType>
+void FNiagaraRendererRibbons::CalculateUVScaleAndOffsets(const FNiagaraRibbonUVSettings& UVSettings, const TArray<IntType>& RibbonIndices, const TArray<FVector4f>& RibbonTangentsAndDistances, const FNiagaraDataSetReaderFloat<float>& NormalizedAgeReader,
+                                                         int32 StartIndex, int32 EndIndex, int32 NumSegments, float TotalLength, float& OutUScale, float& OutUOffset, float& OutUDistributionScaler)
 {
 	float NormalizedLeadingSegmentOffset;
 	if (UVSettings.LeadingEdgeMode == ENiagaraRibbonUVEdgeMode::SmoothTransition)
 	{
-		float FirstAge = NormalizedAgeReader[RibbonIndices[StartIndex]];
-		float SecondAge = NormalizedAgeReader[RibbonIndices[StartIndex + 1]];
+		const float FirstAge = NormalizedAgeReader[RibbonIndices[StartIndex]];
+		const float SecondAge = NormalizedAgeReader[RibbonIndices[StartIndex + 1]];
 
-		float StartTimeStep = SecondAge - FirstAge;
-		float StartTimeOffset = FirstAge < StartTimeStep ? StartTimeStep - FirstAge : 0;
+		const float StartTimeStep = SecondAge - FirstAge;
+		const float StartTimeOffset = FirstAge < StartTimeStep ? StartTimeStep - FirstAge : 0;
 
 		NormalizedLeadingSegmentOffset = StartTimeStep > 0 ? StartTimeOffset / StartTimeStep : 0.0f;
 	}
@@ -547,11 +1201,11 @@ void CalculateUVScaleAndOffsets(
 	float NormalizedTrailingSegmentOffset;
 	if (UVSettings.TrailingEdgeMode == ENiagaraRibbonUVEdgeMode::SmoothTransition)
 	{
-		float SecondToLastAge = NormalizedAgeReader[RibbonIndices[EndIndex - 1]];
-		float LastAge = NormalizedAgeReader[RibbonIndices[EndIndex]];
+		const float SecondToLastAge = NormalizedAgeReader[RibbonIndices[EndIndex - 1]];
+		const float LastAge = NormalizedAgeReader[RibbonIndices[EndIndex]];
 
-		float EndTimeStep = LastAge - SecondToLastAge;
-		float EndTimeOffset = 1 - LastAge < EndTimeStep ? EndTimeStep - (1 - LastAge) : 0;
+		const float EndTimeStep = LastAge - SecondToLastAge;
+		const float EndTimeOffset = 1 - LastAge < EndTimeStep ? EndTimeStep - (1 - LastAge) : 0;
 
 		NormalizedTrailingSegmentOffset = EndTimeStep > 0 ? EndTimeOffset / EndTimeStep : 0.0f;
 	}
@@ -569,21 +1223,21 @@ void CalculateUVScaleAndOffsets(
 	float CalculatedUScale;
 	if (UVSettings.DistributionMode == ENiagaraRibbonUVDistributionMode::ScaledUniformly)
 	{
-		float AvailableSegments = NumSegments - (NormalizedLeadingSegmentOffset + NormalizedTrailingSegmentOffset);
+		const float AvailableSegments = NumSegments - (NormalizedLeadingSegmentOffset + NormalizedTrailingSegmentOffset);
 		CalculatedUScale = NumSegments / AvailableSegments;
 		CalculatedUOffset = -((NormalizedLeadingSegmentOffset / NumSegments) * CalculatedUScale);
 		OutUDistributionScaler = 1.0f / NumSegments;
 	}
 	else if (UVSettings.DistributionMode == ENiagaraRibbonUVDistributionMode::ScaledUsingRibbonSegmentLength)
 	{
-		float SecondDistance = RibbonTangentsAndDistances[StartIndex + 1].W;
-		float LeadingDistanceOffset = SecondDistance * NormalizedLeadingSegmentOffset;
+		const float SecondDistance = RibbonTangentsAndDistances[StartIndex + 1].W;
+		const float LeadingDistanceOffset = SecondDistance * NormalizedLeadingSegmentOffset;
 
-		float SecondToLastDistance = RibbonTangentsAndDistances[EndIndex - 1].W;
-		float LastDistance = RibbonTangentsAndDistances[EndIndex].W;
-		float TrailingDistanceOffset = (LastDistance - SecondToLastDistance) * NormalizedTrailingSegmentOffset;
+		const float SecondToLastDistance = RibbonTangentsAndDistances[EndIndex - 1].W;
+		const float LastDistance = RibbonTangentsAndDistances[EndIndex].W;
+		const float TrailingDistanceOffset = (LastDistance - SecondToLastDistance) * NormalizedTrailingSegmentOffset;
 
-		float AvailableLength = TotalLength - (LeadingDistanceOffset + TrailingDistanceOffset);
+		const float AvailableLength = TotalLength - (LeadingDistanceOffset + TrailingDistanceOffset);
 
 		CalculatedUScale = TotalLength / AvailableLength;
 		CalculatedUOffset = -((LeadingDistanceOffset / TotalLength) * CalculatedUScale);
@@ -591,8 +1245,8 @@ void CalculateUVScaleAndOffsets(
 	}
 	else if (UVSettings.DistributionMode == ENiagaraRibbonUVDistributionMode::TiledOverRibbonLength)
 	{
-		float SecondDistance = RibbonTangentsAndDistances[StartIndex + 1].W;
-		float LeadingDistanceOffset = SecondDistance * NormalizedLeadingSegmentOffset;
+		const float SecondDistance = RibbonTangentsAndDistances[StartIndex + 1].W;
+		const float LeadingDistanceOffset = SecondDistance * NormalizedLeadingSegmentOffset;
 
 		CalculatedUScale = TotalLength / UVSettings.TilingLength;
 		CalculatedUOffset = -(LeadingDistanceOffset / UVSettings.TilingLength);
@@ -615,684 +1269,373 @@ void CalculateUVScaleAndOffsets(
 	OutUOffset = (CalculatedUOffset * UVSettings.Scale.X) + UVSettings.Offset.X;
 }
 
-FNiagaraDynamicDataBase* FNiagaraRendererRibbons::GenerateDynamicData(const FNiagaraSceneProxy* Proxy, const UNiagaraRendererProperties* InProperties, const FNiagaraEmitterInstance* Emitter)const
-{
-	SCOPE_CYCLE_COUNTER(STAT_NiagaraGenRibbonVertexData);
 
-	if (SimTarget == ENiagaraSimTarget::GPUComputeSim)
+template<bool bWantsTessellation, bool bHasTwist, bool bWantsMultiRibbon>
+void FNiagaraRendererRibbons::GenerateVertexBufferForRibbonPart(const FNiagaraGenerationInputDataCPUAccessors& CPUData, const TArray<uint32>& RibbonIndices, uint32 RibbonIndex, FNiagaraRibbonCPUGeneratedVertexData& OutputData) const
+{
+	TArray<uint32>& SegmentData = OutputData.SegmentData;
+	TArray<FRibbonMultiRibbonInfo>& MultiRibbonInfos = OutputData.RibbonInfoLookup;
+	
+	const FNiagaraDataSetReaderFloat<FNiagaraPosition>& PosData = CPUData.PosData;	
+	const FNiagaraDataSetReaderFloat<float>& AgeData = CPUData.AgeData;
+	const FNiagaraDataSetReaderFloat<float>& SizeData = CPUData.SizeData;
+	const FNiagaraDataSetReaderFloat<float>& TwistData = CPUData.TwistData;
+	
+	const int32 StartIndex = OutputData.SortedIndices.Num();
+
+	const FVector FirstPos = static_cast<FVector>(PosData[RibbonIndices[0]]);
+	FVector CurrPos = FirstPos;
+	FVector LastToCurrVec = FVector::ZeroVector;
+	float LastToCurrSize = 0;	
+	float LastTwist = 0;
+	float LastWidth = 0;
+	double TotalDistance = 0.0f;
+
+	// Find the first position with enough distance.
+	int32 CurrentIndex = 1;
+	while (CurrentIndex < RibbonIndices.Num())
 	{
-		return nullptr;
-	}
+		const int32 CurrentDataIndex = RibbonIndices[CurrentIndex];
+		CurrPos = static_cast<FVector>(PosData[CurrentDataIndex]);
+		LastToCurrVec = CurrPos - FirstPos;
+		LastToCurrSize = LastToCurrVec.Size();
+		if constexpr (bHasTwist)
+		{
+			LastTwist = TwistData[CurrentDataIndex];
+			LastWidth = SizeData[CurrentDataIndex];
+		}
 
-	FNiagaraDataSet& Data = Emitter->GetData();
-	const UNiagaraRibbonRendererProperties* Properties = CastChecked<const UNiagaraRibbonRendererProperties>(InProperties);
+		// Find the first segment, or unique segment
+		if (LastToCurrSize > GNiagaraRibbonMinSegmentLength)
+		{
+			// Normalize LastToCurrVec
+			LastToCurrVec *= 1.f / LastToCurrSize;
 
-	if (!IsRendererEnabled(Properties, Emitter))
-	{
-		return nullptr;
+			// Add the first point. Tangent follows first segment.
+			OutputData.SortedIndices.Add(RibbonIndices[0]);
+			OutputData.TangentAndDistances.Add(FVector4f(LastToCurrVec.X, LastToCurrVec.Y, LastToCurrVec.Z, 0));
+			if constexpr (bWantsMultiRibbon)
+			{
+				OutputData.MultiRibbonIndices.Add(RibbonIndex);
+			}
+			break;
+		}
+		else
+		{
+			LastToCurrSize = 0; // Ensure that the segment gets ignored if too small
+			++CurrentIndex;
+		}
 	}
 
-	if (Properties->bAllowInCullProxies == false)
+	// Now iterate on all other points, to proceed each particle connected to 2 segments.
+	int32 NextIndex = CurrentIndex + 1;
+	while (NextIndex < RibbonIndices.Num())
 	{
-		check(Emitter);
-
-		FNiagaraSystemInstance* Inst = Emitter->GetParentSystemInstance();
-		check(Emitter->GetParentSystemInstance());
-
-		//TODO: Probably should push some state into the system instance for this?
-		bool bIsCullProxy = Cast<UNiagaraCullProxyComponent>(Inst->GetAttachComponent()) != nullptr;
-		if (bIsCullProxy)
+		const int32 NextDataIndex = RibbonIndices[NextIndex];
+		const FVector NextPos = static_cast<FVector>(PosData[NextDataIndex]);
+		FVector CurrToNextVec = NextPos - CurrPos;
+		const float CurrToNextSize = CurrToNextVec.Size();
+
+		float NextTwist = 0;
+		float NextWidth = 0;
+		if constexpr (bHasTwist)
 		{
-			return nullptr;
+			NextTwist = TwistData[NextDataIndex];
+			NextWidth = SizeData[NextDataIndex];
 		}
-	}
 
-	FNiagaraDataBuffer* DataToRender = Emitter->GetData().GetCurrentData();
-	if (DataToRender == nullptr || DataToRender->GetNumInstances() < 2 || !Properties->PositionDataSetAccessor.IsValid() || !Properties->SortKeyDataSetAccessor.IsValid())
-	{
-		return nullptr;
-	}
-
-	const bool bSortKeyIsAge = Properties->bSortKeyDataSetAccessorIsAge;
-	const auto SortKeyReader = Properties->SortKeyDataSetAccessor.GetReader(Data);
-
-	const auto PosData = Properties->PositionDataSetAccessor.GetReader(Data);
-	const auto AgeData = Properties->NormalizedAgeAccessor.GetReader(Data);
-	const auto SizeData = Properties->SizeDataSetAccessor.GetReader(Data);
-	const auto TwistData = Properties->TwistDataSetAccessor.GetReader(Data);
-	const auto FacingData = Properties->FacingDataSetAccessor.GetReader(Data);
-
-	const auto MaterialParam0Data = Properties->MaterialParam0DataSetAccessor.GetReader(Data);
-	const auto MaterialParam1Data = Properties->MaterialParam1DataSetAccessor.GetReader(Data);
-	const auto MaterialParam2Data = Properties->MaterialParam2DataSetAccessor.GetReader(Data);
-	const auto MaterialParam3Data = Properties->MaterialParam3DataSetAccessor.GetReader(Data);
-
-	bool U0OverrideIsBound = Properties->U0OverrideIsBound;
-	bool U1OverrideIsBound = Properties->U1OverrideIsBound;
-
-	const auto RibbonIdData = Properties->RibbonIdDataSetAccessor.GetReader(Data);
-	const auto RibbonFullIDData = Properties->RibbonFullIDDataSetAccessor.GetReader(Data);
+		// It the next is far enough, or the last element
+		if (CurrToNextSize > GNiagaraRibbonMinSegmentLength || NextIndex == RibbonIndices.Num() - 1)
+		{
+			// Normalize CurrToNextVec
+			CurrToNextVec *= 1.f / FMath::Max(GNiagaraRibbonMinSegmentLength, CurrToNextSize);
+			const FVector Tangent = (1.f - GenerationConfig.GetCurveTension()) * (LastToCurrVec + CurrToNextVec).GetSafeNormal();
 
-	FNiagaraDynamicDataRibbon* DynamicData = new FNiagaraDynamicDataRibbon(Emitter);
+			// Update the distance for CurrentIndex.
+			TotalDistance += LastToCurrSize;
 
-	if (Properties->Shape == ENiagaraRibbonShapeMode::MultiPlane)
-	{
-		for (int32 PlaneIndex = 0; PlaneIndex < MultiPlaneCount; PlaneIndex++)
-		{
-			float RotationAngle = (float(PlaneIndex) / MultiPlaneCount) * 180.0f;
+			// Add the current point, which tangent is computed from neighbors
+			OutputData.SortedIndices.Add(RibbonIndices[CurrentIndex]);
+			OutputData.TangentAndDistances.Add(FVector4f(Tangent.X, Tangent.Y, Tangent.Z, TotalDistance));
 
-			for (int32 VertexId = 0; VertexId <= WidthSegmentationCount; VertexId++)
+			if constexpr (bWantsMultiRibbon)
 			{
-				FVector2D Position = FVector2D((((float)VertexId) / WidthSegmentationCount) - 0.5f, 0).GetRotated(RotationAngle);
-				FVector2D Normal = FVector2D(0, 1).GetRotated(RotationAngle);
-				float TextureV = ((float)VertexId) / WidthSegmentationCount;
-
-				DynamicData->PackSliceVertexData(Position, Normal, TextureV);
+				OutputData.MultiRibbonIndices.Add(RibbonIndex);
 			}
-		}
 
-		if (bEnableAccurateGeometry)
-		{
-			for (int32 PlaneIndex = 0; PlaneIndex < MultiPlaneCount; PlaneIndex++)
+			// Assumed equal to dot(Tangent, CurrToNextVec)
+			OutputData.TotalSegmentLength += CurrToNextSize;
+			
+			if constexpr (bWantsTessellation)
 			{
-				float RotationAngle = (float(PlaneIndex) / MultiPlaneCount) * 180.0f;
-
-				for (int32 VertexId = 0; VertexId <= WidthSegmentationCount; VertexId++)
+				OutputData.AverageSegmentLength += CurrToNextSize * CurrToNextSize;
+				OutputData.AverageSegmentAngle += CurrToNextSize * AcosFast(FVector::DotProduct(LastToCurrVec, CurrToNextVec));
+				if constexpr (bHasTwist)
 				{
-					FVector2D Position = FVector2D((((float)VertexId) / WidthSegmentationCount) - 0.5f, 0).GetRotated(RotationAngle);
-					FVector2D Normal = FVector2D(0, -1).GetRotated(RotationAngle);
-					float TextureV = ((float)VertexId) / WidthSegmentationCount;
-
-					DynamicData->PackSliceVertexData(Position, Normal, TextureV);
+					OutputData.AverageTwistAngle += CurrToNextSize * FMath::Abs(NextTwist - LastTwist);
+					OutputData.AverageWidth += CurrToNextSize * LastWidth;
 				}
 			}
-		}
-	}
-	else if (Properties->Shape == ENiagaraRibbonShapeMode::Tube)
-	{
-		for (int32 VertexId = 0; VertexId <= TubeSubdivisions; VertexId++)
-		{
-			float RotationAngle = (float(VertexId) / TubeSubdivisions) * -360.0f;
-			FVector2D Position = FVector2D(-0.5f, 0.0f).GetRotated(RotationAngle);
-			FVector2D Normal = FVector2D(-1, 0).GetRotated(RotationAngle);
-			float TextureV = float(VertexId) / TubeSubdivisions;
 
-			DynamicData->PackSliceVertexData(Position, Normal, TextureV);
-		}
-	}
-	else if (Properties->Shape == ENiagaraRibbonShapeMode::Custom && CustomVertices.Num() >= 2)
-	{
-		bool bHasCustomUVs = false;
-		for (int32 VertexId = 0; VertexId < CustomVertices.Num(); VertexId++)
-		{
-			if (!FMath::IsNearlyZero(CustomVertices[VertexId].TextureV))
-			{
-				bHasCustomUVs = true;
-				break;
-			}
+			// Move to next segment.
+			CurrentIndex = NextIndex;
+			CurrPos = NextPos;
+			LastToCurrVec = CurrToNextVec;
+			LastToCurrSize = CurrToNextSize;
+			LastTwist = NextTwist;
+			LastWidth = NextWidth;
 		}
 
-		for (int32 VertexId = 0; VertexId <= CustomVertices.Num(); VertexId++)
-		{
-			const auto& CustomVert = CustomVertices[VertexId % CustomVertices.Num()];
-
-			FVector2D Position = CustomVert.Position;
-			FVector2D Normal = CustomVert.Normal.IsNearlyZero() ? Position.GetSafeNormal() : CustomVert.Normal;
-			float TextureV = bHasCustomUVs ? CustomVert.TextureV : ((float)VertexId) / WidthSegmentationCount;
-
-			DynamicData->PackSliceVertexData(Position, Normal, TextureV);
-		}
+		// Try next if there is one.
+		++NextIndex;
 	}
-	else // Plane
+
+	// Close the last point and segment if there was at least 2.
+	if (LastToCurrSize > 0)
 	{
-		for (int32 VertexId = 0; VertexId <= WidthSegmentationCount; VertexId++)
-		{
-			FVector2D Position = FVector2D((((float)VertexId) / WidthSegmentationCount) - 0.5f, 0);
-			FVector2D Normal = FVector2D(0, 1);
-			float TextureV = ((float)VertexId) / WidthSegmentationCount;
+		// Update the distance for CurrentIndex.
+		TotalDistance += LastToCurrSize;
 
-			DynamicData->PackSliceVertexData(Position, Normal, TextureV);
+		// Add the last point, which tangent follows the last segment.
+		OutputData.SortedIndices.Add(RibbonIndices[CurrentIndex]);
+		OutputData.TangentAndDistances.Add(FVector4f(LastToCurrVec.X, LastToCurrVec.Y, LastToCurrVec.Z, TotalDistance));
+		if constexpr (bWantsMultiRibbon)
+		{
+			OutputData.MultiRibbonIndices.Add(RibbonIndex);
 		}
 	}
 
+	const int32 EndIndex = OutputData.SortedIndices.Num() - 1;
+	const int32 NumSegments = EndIndex - StartIndex;
 
-
-
-	//In preparation for a material override feature, we pass our material(s) and relevance in via dynamic data.
-	//The renderer ensures we have the correct usage and relevance for materials in BaseMaterials_GT.
-	//Any override feature must also do the same for materials that are set.
-	check(BaseMaterials_GT.Num() == 1);
-	check(BaseMaterials_GT[0]->CheckMaterialUsage_Concurrent(MATUSAGE_NiagaraRibbons));
-	DynamicData->Material = BaseMaterials_GT[0]->GetRenderProxy();
-	DynamicData->SetMaterialRelevance(BaseMaterialRelevance_GT);
-
-	if (DynamicData && Properties->MaterialParameterBindings.Num() != 0)
-	{
-		ProcessMaterialParameterBindings(MakeArrayView(Properties->MaterialParameterBindings), Emitter, MakeArrayView(BaseMaterials_GT));
-	}
-
-	TArray<int32>& SegmentData = DynamicData->SegmentData;
-	int32& MaxParticleIndex = DynamicData->MaxParticleIndex;
-	float TotalSegmentLength = 0;
-	// weighted sum based on the segment length :
-	float AverageSegmentLength = 0;
-	float AverageSegmentAngle = 0;
-	float AverageTwistAngle = 0;
-	float AverageWidth = 0;
-
-	bool bFullIDs = RibbonFullIDData.IsValid();
-	bool bSimpleIDs = !bFullIDs && RibbonIdData.IsValid();
-	bool bMultiRibbons = bFullIDs || bSimpleIDs;
-	bool bHasTwist = TwistData.IsValid() && SizeData.IsValid();
-
-	auto AddRibbonVerts = [&](TArray<int32>& RibbonIndices, uint32 RibbonIndex)
+	if (NumSegments > 0)
 	{
-		const int32 StartIndex = DynamicData->SortedIndices.Num();
-
-		float TotalDistance = 0.0f;
-
-		const FVector FirstPos = (FVector)PosData[RibbonIndices[0]];
-		FVector CurrPos = FirstPos;
-		FVector LastToCurrVec = FVector::ZeroVector;
-		float LastToCurrSize = 0;
-		float LastTwist = 0;
-		float LastWidth = 0;
-
-		// Find the first position with enough distance.
-		int32 CurrentIndex = 1;
-		while (CurrentIndex < RibbonIndices.Num())
+		FRibbonMultiRibbonInfo& MultiRibbonInfo = MultiRibbonInfos[RibbonIndex];
+		MultiRibbonInfo.StartPos = (FVector)PosData[RibbonIndices[0]];
+		MultiRibbonInfo.EndPos = (FVector)PosData[RibbonIndices.Last()];
+		MultiRibbonInfo.BaseSegmentDataIndex = SegmentData.Num();
+		MultiRibbonInfo.NumSegmentDataIndices = NumSegments;
+
+		// Update the tangents for the first and last vertex, apply a reflect vector logic so that the initial and final curvature is continuous.
+		if (NumSegments > 1)
 		{
-			const int32 CurrentDataIndex = RibbonIndices[CurrentIndex];
-			CurrPos = (FVector)PosData[CurrentDataIndex];
-			LastToCurrVec = CurrPos - FirstPos;
-			LastToCurrSize = LastToCurrVec.Size();
-			if (bHasTwist)
-			{
-				LastTwist = TwistData[CurrentDataIndex];
-				LastWidth = SizeData[CurrentDataIndex];
-			}
+			FVector3f& FirstTangent =  reinterpret_cast<FVector3f&>(OutputData.TangentAndDistances[StartIndex]);
+			FVector3f& NextToFirstTangent = reinterpret_cast<FVector3f&>(OutputData.TangentAndDistances[StartIndex + 1]);
+			FirstTangent = (2.f * FVector3f::DotProduct(FirstTangent, NextToFirstTangent)) * FirstTangent - NextToFirstTangent;
 
-			// Find the first segment, or unique segment
-			if (LastToCurrSize > GNiagaraRibbonMinSegmentLength)
-			{
-				// Normalize LastToCurrVec
-				LastToCurrVec *= 1.f / LastToCurrSize;
-
-				// Add the first point. Tangent follows first segment.
-				DynamicData->SortedIndices.Add(RibbonIndices[0]);
-				MaxParticleIndex = FMath::Max(RibbonIndices[0], MaxParticleIndex);
-				DynamicData->TangentAndDistances.Add(FVector4f(LastToCurrVec.X, LastToCurrVec.Y, LastToCurrVec.Z, 0));
-				DynamicData->MultiRibbonIndices.Add(RibbonIndex);
-				break;
-			}
-			else
-			{
-				LastToCurrSize = 0; // Ensure that the segment gets ignored if too small
-				++CurrentIndex;
-			}
+			FVector3f& LastTangent = reinterpret_cast<FVector3f&>(OutputData.TangentAndDistances[EndIndex]);
+			FVector3f& PrevToLastTangent = reinterpret_cast<FVector3f&>(OutputData.TangentAndDistances[EndIndex - 1]);
+			LastTangent = (2.f * FVector3f::DotProduct(LastTangent, PrevToLastTangent)) * LastTangent - PrevToLastTangent;
 		}
 
-		// Now iterate on all other points, to proceed each particle connected to 2 segments.
-		int32 NextIndex = CurrentIndex + 1;
-		while (NextIndex < RibbonIndices.Num())
+		// Add segment data
+		for (int32 SegmentIndex = StartIndex; SegmentIndex < EndIndex; ++SegmentIndex)
 		{
-			const int32 NextDataIndex = RibbonIndices[NextIndex];
-			const FVector NextPos = (FVector)PosData[NextDataIndex];
-			FVector CurrToNextVec = NextPos - CurrPos;
-			const float CurrToNextSize = CurrToNextVec.Size();
-
-			float NextTwist = 0;
-			float NextWidth = 0;
-			if (bHasTwist)
-			{
-				NextTwist = TwistData[NextDataIndex];
-				NextWidth = SizeData[NextDataIndex];
-			}
-
-			// It the next is far enough, or the last element
-			if (CurrToNextSize > GNiagaraRibbonMinSegmentLength || NextIndex == RibbonIndices.Num() - 1)
-			{
-				// Normalize CurrToNextVec
-				CurrToNextVec *= 1.f / FMath::Max(GNiagaraRibbonMinSegmentLength, CurrToNextSize);
-				const FVector Tangent = (1.f - CustomCurveTension) * (LastToCurrVec + CurrToNextVec).GetSafeNormal();
-
-				// Update the distance for CurrentIndex.
-				TotalDistance += LastToCurrSize;
-
-				// Add the current point, which tangent is computed from neighbors
-				DynamicData->SortedIndices.Add(RibbonIndices[CurrentIndex]);
-				MaxParticleIndex = FMath::Max(RibbonIndices[CurrentIndex], MaxParticleIndex);
-				DynamicData->TangentAndDistances.Add(FVector4f(Tangent.X, Tangent.Y, Tangent.Z, TotalDistance));
-				DynamicData->MultiRibbonIndices.Add(RibbonIndex);
-
-				// Assumed equal to dot(Tangent, CurrToNextVec)
-				TotalSegmentLength += CurrToNextSize;
-				AverageSegmentLength += CurrToNextSize * CurrToNextSize;
-				AverageSegmentAngle += CurrToNextSize * AcosFast(FVector::DotProduct(LastToCurrVec, CurrToNextVec));
-				AverageTwistAngle += FMath::Abs(NextTwist - LastTwist) * CurrToNextSize;
-				AverageWidth += LastWidth * CurrToNextSize;
-
-				// Move to next segment.
-				CurrentIndex = NextIndex;
-				CurrPos = NextPos;
-				LastToCurrVec = CurrToNextVec;
-				LastToCurrSize = CurrToNextSize;
-				LastTwist = NextTwist;
-				LastWidth = NextWidth;
-			}
-
-			// Try next if there is one.
-			++NextIndex;
+			SegmentData.Add(SegmentIndex);
 		}
 
-		// Close the last point and segment if there was at least 2.
-		if (LastToCurrSize > 0)
+		float U0Offset;
+		float U0Scale;
+		float U0DistributionScaler;
+		if(GenerationConfig.HasCustomU0Data())
 		{
-			// Update the distance for CurrentIndex.
-			TotalDistance += LastToCurrSize;
-
-			// Add the last point, which tangent follows the last segment.
-			DynamicData->SortedIndices.Add(RibbonIndices[CurrentIndex]);
-			MaxParticleIndex = FMath::Max(RibbonIndices[CurrentIndex], MaxParticleIndex);
-			DynamicData->TangentAndDistances.Add(FVector4f(LastToCurrVec.X, LastToCurrVec.Y, LastToCurrVec.Z, TotalDistance));
-			DynamicData->MultiRibbonIndices.Add(RibbonIndex);
+			U0Offset = 0;
+			U0Scale = 1.0f;
+			U0DistributionScaler = 1;
+		}
+		else
+		{
+			CalculateUVScaleAndOffsets(
+				UV0Settings, OutputData.SortedIndices, OutputData.TangentAndDistances,
+				AgeData,
+				StartIndex, EndIndex,
+				NumSegments, TotalDistance,
+				U0Scale, U0Offset, U0DistributionScaler);
 		}
 
-		const int32 EndIndex = DynamicData->SortedIndices.Num() - 1;
-		const int32 NumSegments = EndIndex - StartIndex;
-
-		if (NumSegments > 0)
+		float U1Offset;
+		float U1Scale;
+		float U1DistributionScaler;
+		if (GenerationConfig.HasCustomU1Data())
+		{
+			U1Offset = 0;
+			U1Scale = 1.0f;
+			U1DistributionScaler = 1;
+		}
+		else
 		{
-			FNiagaraDynamicDataRibbon::FMultiRibbonInfo& MultiRibbonInfo = DynamicData->MultiRibbonInfos[RibbonIndex];
-			MultiRibbonInfo.StartPos = (FVector)PosData[RibbonIndices[0]];
-			MultiRibbonInfo.EndPos = (FVector)PosData[RibbonIndices.Last()];
-			MultiRibbonInfo.BaseSegmentDataIndex = SegmentData.Num();
-			MultiRibbonInfo.NumSegmentDataIndices = NumSegments;
+			CalculateUVScaleAndOffsets(
+				UV1Settings, OutputData.SortedIndices, OutputData.TangentAndDistances,
+				AgeData,
+				StartIndex, EndIndex,
+				NumSegments, TotalDistance,
+				U1Scale, U1Offset, U1DistributionScaler);
+		}
 
-			// Update the tangents for the first and last vertex, apply a reflect vector logic so that the initial and final curvature is continuous.
-			if (NumSegments > 1)
-			{
-				FVector3f& FirstTangent =  reinterpret_cast<FVector3f&>(DynamicData->TangentAndDistances[StartIndex]);
-				FVector3f& NextToFirstTangent = reinterpret_cast<FVector3f&>(DynamicData->TangentAndDistances[StartIndex + 1]);
-				FirstTangent = (2.f * FVector3f::DotProduct(FirstTangent, NextToFirstTangent)) * FirstTangent - NextToFirstTangent;
+		MultiRibbonInfo.BufferEntry.U0Scale = U0Scale;
+		MultiRibbonInfo.BufferEntry.U0Offset = U0Offset;
+		MultiRibbonInfo.BufferEntry.U0DistributionScaler = U0DistributionScaler;
+		MultiRibbonInfo.BufferEntry.U1Scale = U1Scale;
+		MultiRibbonInfo.BufferEntry.U1Offset = U1Offset;
+		MultiRibbonInfo.BufferEntry.U1DistributionScaler = U1DistributionScaler;
+		MultiRibbonInfo.BufferEntry.FirstParticleId = StartIndex;
+		MultiRibbonInfo.BufferEntry.LastParticleId = EndIndex;
+	}
+}
 
-				FVector3f& LastTangent = reinterpret_cast<FVector3f&>(DynamicData->TangentAndDistances[EndIndex]);
-				FVector3f& PrevToLastTangent = reinterpret_cast<FVector3f&>(DynamicData->TangentAndDistances[EndIndex - 1]);
-				LastTangent = (2.f * FVector3f::DotProduct(LastTangent, PrevToLastTangent)) * LastTangent - PrevToLastTangent;
-			}
+template<typename IDType, typename ReaderType, bool bWantsTessellation, bool bHasTwist>
+void FNiagaraRendererRibbons::GenerateVertexBufferForMultiRibbonInternal(const FNiagaraGenerationInputDataCPUAccessors& CPUData, const ReaderType& IDReader, FNiagaraRibbonCPUGeneratedVertexData& OutputData) const
+{
+	TMap<IDType, TArray<uint32>> MultiRibbonSortedIndices;
 
-			// Add segment data
-			for (int32 SegmentIndex = StartIndex; SegmentIndex < EndIndex; ++SegmentIndex)
-			{
-				SegmentData.Add(SegmentIndex);
-			}
+	for (uint32 i = 0; i < CPUData.TotalNumParticles; ++i)
+	{
+		TArray<uint32>& Indices = MultiRibbonSortedIndices.FindOrAdd(IDReader[i]);
+		Indices.Add(i);
+	}
+	OutputData.RibbonInfoLookup.AddZeroed(MultiRibbonSortedIndices.Num());
 
-			float U0Offset;
-			float U0Scale;
-			float U0DistributionScaler;
-			if(UV0Settings.bEnablePerParticleUOverride && U0OverrideIsBound)
-			{
-				U0Offset = 0;
-				U0Scale = 1.0f;
-				U0DistributionScaler = 1;
-			}
-			else
-			{
-				CalculateUVScaleAndOffsets(
-					UV0Settings, DynamicData->SortedIndices, DynamicData->TangentAndDistances,
-					AgeData,
-					StartIndex, DynamicData->SortedIndices.Num() - 1,
-					NumSegments, TotalDistance,
-					U0Scale, U0Offset, U0DistributionScaler);
-			}
+	// Sort the ribbons by ID so that the draw order stays consistent.
+	MultiRibbonSortedIndices.KeySort(TLess<IDType>());
 
-			float U1Offset;
-			float U1Scale;
-			float U1DistributionScaler;
-			if (UV1Settings.bEnablePerParticleUOverride && U1OverrideIsBound)
-			{
-				U1Offset = 0;
-				U1Scale = 1.0f;
-				U1DistributionScaler = 1;
-			}
-			else
-			{
-				CalculateUVScaleAndOffsets(
-					UV1Settings, DynamicData->SortedIndices, DynamicData->TangentAndDistances,
-					AgeData,
-					StartIndex, DynamicData->SortedIndices.Num() - 1,
-					NumSegments, TotalDistance,
-					U1Scale, U1Offset, U1DistributionScaler);
-			}
+	uint32 RibbonIndex = 0;
+	for (TPair<IDType, TArray<uint32>>& Pair : MultiRibbonSortedIndices)
+	{
+		TArray<uint32>& SortedIndices = Pair.Value;
+		const auto& SortKeyReader = CPUData.SortKeyReader;
+		SortedIndices.Sort([&SortKeyReader](const uint32& A, const uint32& B) { return (SortKeyReader[A] < SortKeyReader[B]); });
+		GenerateVertexBufferForRibbonPart<bWantsTessellation, bHasTwist, true>(CPUData, SortedIndices, RibbonIndex, OutputData);
+		RibbonIndex++;
+	};
+}
 
-			DynamicData->PackPerRibbonData(U0Scale, U0Offset, U0DistributionScaler, U1Scale, U1Offset, U1DistributionScaler, StartIndex);
+template<typename IDType, typename ReaderType>
+void FNiagaraRendererRibbons::GenerateVertexBufferForMultiRibbon(const FNiagaraGenerationInputDataCPUAccessors& CPUData, const ReaderType& IDReader, FNiagaraRibbonCPUGeneratedVertexData& OutputData) const
+{
+	if (GenerationConfig.WantsAutomaticTessellation())
+	{
+		if (GenerationConfig.HasTwist())
+		{
+			GenerateVertexBufferForMultiRibbonInternal<IDType, ReaderType, true, true>(CPUData, IDReader, OutputData);				
 		}
 		else
 		{
-			DynamicData->PackPerRibbonData(0, 0, 0, 0, 0, 0, 0);
+			GenerateVertexBufferForMultiRibbonInternal<IDType, ReaderType, true, false>(CPUData, IDReader, OutputData);				
 		}
-	};
-
-	DynamicData->MultiRibbonInfos.Reset();
+	}
+	else
+	{
+		GenerateVertexBufferForMultiRibbonInternal<IDType, ReaderType, false, false>(CPUData, IDReader, OutputData);		
+	}
+}
 
-	//TODO: Move sorting to share code with sprite and mesh sorting and support the custom sorting key.
-	int32 TotalIndices = Data.GetCurrentDataChecked().GetNumInstances();
+void FNiagaraRendererRibbons::GenerateVertexBufferCPU(const FNiagaraGenerationInputDataCPUAccessors& CPUData, FNiagaraRibbonCPUGeneratedVertexData& OutputData) const
+{
+	SCOPE_CYCLE_COUNTER(STAT_NiagaraRenderRibbonsGenVerticesCPU);
+	
+	check(CPUData.PosData.IsValid() && CPUData.SortKeyReader.IsValid());
 
-	if (!bMultiRibbons)
+	// TODO: Move sorting to share code with sprite and mesh sorting and support the custom sorting key.
+	if (GenerationConfig.HasRibbonIDs())
 	{
-		TArray<int32> SortedIndices;
-		for (int32 i = 0; i < TotalIndices; ++i)
+		if (GenerationConfig.HasFullRibbonIDs())
 		{
-			SortedIndices.Add(i);
+			GenerateVertexBufferForMultiRibbon<FNiagaraID>(CPUData, CPUData.FullRibbonIDData, OutputData);
 		}
-		DynamicData->MultiRibbonInfos.AddZeroed(1);
-
-		SortedIndices.Sort([&SortKeyReader](const int32& A, const int32& B) {	return (SortKeyReader[A] < SortKeyReader[B]); });
+		else
+		{
+			// TODO: Remove simple ID path
+			check(GenerationConfig.HasSimpleRibbonIDs());
 
-		AddRibbonVerts(SortedIndices, 0);
+			GenerateVertexBufferForMultiRibbon<uint32>(CPUData, CPUData.SimpleRibbonIDData, OutputData);
+		}		
 	}
 	else
-	{
-		if (bFullIDs)
+	{		
+		TArray<uint32> SortedIndices;
+		for (uint32 i = 0; i < CPUData.TotalNumParticles; ++i)
 		{
-			TMap<FNiagaraID, TArray<int32>> MultiRibbonSortedIndices;
+			SortedIndices.Add(i);
+		}
+		OutputData.RibbonInfoLookup.AddZeroed(1);
 
-			for (int32 i = 0; i < TotalIndices; ++i)
+		const auto& SortKeyReader = CPUData.SortKeyReader;
+		SortedIndices.Sort([&SortKeyReader](const uint32& A, const uint32& B) {	return (SortKeyReader[A] < SortKeyReader[B]); });
+		
+		if (GenerationConfig.WantsAutomaticTessellation())
+		{
+			if (GenerationConfig.HasTwist())
 			{
-				TArray<int32>& Indices = MultiRibbonSortedIndices.FindOrAdd(RibbonFullIDData[i]);
-				Indices.Add(i);
+				GenerateVertexBufferForRibbonPart<true, true, false>(CPUData, SortedIndices, 0/*RibbonIndex*/, OutputData);			
 			}
-			DynamicData->MultiRibbonInfos.AddZeroed(MultiRibbonSortedIndices.Num());
-
-			// Sort the ribbons by ID so that the draw order stays consistent.
-			MultiRibbonSortedIndices.KeySort(TLess<FNiagaraID>());
-
-			uint32 RibbonIndex = 0;
-			for (TPair<FNiagaraID, TArray<int32>>& Pair : MultiRibbonSortedIndices)
+			else
 			{
-				TArray<int32>& SortedIndices = Pair.Value;
-				SortedIndices.Sort([&SortKeyReader](const int32& A, const int32& B) {	return (SortKeyReader[A] < SortKeyReader[B]); });
-				AddRibbonVerts(SortedIndices, RibbonIndex);
-
-				RibbonIndex++;
-			};
+				GenerateVertexBufferForRibbonPart<true, false, false>(CPUData, SortedIndices, 0/*RibbonIndex*/, OutputData);			
+			}
 		}
 		else
 		{
-			//TODO: Remove simple ID path
-			check(bSimpleIDs);
-
-			TMap<int32, TArray<int32>> MultiRibbonSortedIndices;
-
-			for (int32 i = 0; i < TotalIndices; ++i)
-			{
-				TArray<int32>& Indices = MultiRibbonSortedIndices.FindOrAdd(RibbonIdData[i]);
-				Indices.Add(i);
-			}
-			DynamicData->MultiRibbonInfos.AddZeroed(MultiRibbonSortedIndices.Num());
-
-			// Sort the ribbons by ID so that the draw order stays consistent.
-			MultiRibbonSortedIndices.KeySort(TLess<int32>());
-
-			uint32 RibbonIndex = 0;
-			for (TPair<int32, TArray<int32>>& Pair : MultiRibbonSortedIndices)
-			{
-				TArray<int32>& SortedIndices = Pair.Value;
-				SortedIndices.Sort([&SortKeyReader](const int32& A, const int32& B) {	return (SortKeyReader[A] < SortKeyReader[B]); });
-				AddRibbonVerts(SortedIndices, RibbonIndex);
-				RibbonIndex++;
-			};
+			GenerateVertexBufferForRibbonPart<false, false, false>(CPUData, SortedIndices, 0/*RibbonIndex*/, OutputData);		
 		}
 	}
-
-	if (TotalSegmentLength > 0)
+	
+	if (OutputData.TotalSegmentLength > 0.0)
 	{
+		const double& TotalSegmentLength = OutputData.TotalSegmentLength;
+	
+		// weighted sum based on the segment length :
+		double& AverageSegmentLength = OutputData.AverageSegmentLength;
+		double& AverageSegmentAngle = OutputData.AverageSegmentAngle;
+		double& AverageTwistAngle = OutputData.AverageTwistAngle;
+		double& AverageWidth = OutputData.AverageWidth;
+		
 		// Blend the result between the last frame tessellation factors and the current frame base on the total length of all segments.
 		// This is only used to increase the tessellation value of the current frame data to prevent glitches where tessellation is significantly changin between frames.
 		const float OneOverTotalSegmentLength = 1.f / FMath::Max(1.f, TotalSegmentLength);
-		const float AveragingFactor = TessellationTotalSegmentLength / (TotalSegmentLength + TessellationTotalSegmentLength);
-		TessellationTotalSegmentLength = TotalSegmentLength;
+		const float AveragingFactor = TessellationSmoothingData.TessellationTotalSegmentLength / (TotalSegmentLength + TessellationSmoothingData.TessellationTotalSegmentLength);
+		TessellationSmoothingData.TessellationTotalSegmentLength = TotalSegmentLength;
 
 		AverageSegmentAngle *= OneOverTotalSegmentLength;
 		AverageSegmentLength *= OneOverTotalSegmentLength;
 		const float AverageSegmentCurvature = AverageSegmentLength / (FMath::Max(SMALL_NUMBER, FMath::Abs(FMath::Sin(AverageSegmentAngle))));
 
-		TessellationAngle = FMath::Lerp<float>(AverageSegmentAngle, FMath::Max(TessellationAngle, AverageSegmentAngle), AveragingFactor);
-		TessellationCurvature = FMath::Lerp<float>(AverageSegmentCurvature, FMath::Max(TessellationCurvature, AverageSegmentCurvature), AveragingFactor);
+		TessellationSmoothingData.TessellationAngle = FMath::Lerp<float>(AverageSegmentAngle, FMath::Max(TessellationSmoothingData.TessellationAngle, AverageSegmentAngle), AveragingFactor);
+		TessellationSmoothingData.TessellationCurvature = FMath::Lerp<float>(AverageSegmentCurvature, FMath::Max(TessellationSmoothingData.TessellationCurvature, AverageSegmentCurvature), AveragingFactor);
 
-		if (bHasTwist)
+		if (GenerationConfig.HasTwist())
 		{
 			AverageTwistAngle *= OneOverTotalSegmentLength;
 			AverageWidth *= OneOverTotalSegmentLength;
 
-			TessellationTwistAngle = FMath::Lerp<float>(AverageTwistAngle, FMath::Max(TessellationTwistAngle, AverageTwistAngle), AveragingFactor);
-			TessellationTwistCurvature = FMath::Lerp<float>(AverageWidth, FMath::Max(TessellationTwistCurvature, AverageWidth), AveragingFactor);
+			TessellationSmoothingData.TessellationTwistAngle = FMath::Lerp<float>(AverageTwistAngle, FMath::Max(TessellationSmoothingData.TessellationTwistAngle, AverageTwistAngle), AveragingFactor);
+			TessellationSmoothingData.TessellationTwistCurvature = FMath::Lerp<float>(AverageWidth, FMath::Max(TessellationSmoothingData.TessellationTwistCurvature, AverageWidth), AveragingFactor);
 		}
 	}
 	else // Reset the metrics when the ribbons are reset.
 	{
-		TessellationAngle = 0;
-		TessellationCurvature = 0;
-		TessellationTwistAngle = 0;
-		TessellationTwistCurvature = 0;
-		TessellationTotalSegmentLength = 0;
+		TessellationSmoothingData.TessellationAngle = 0;
+		TessellationSmoothingData.TessellationCurvature = 0;
+		TessellationSmoothingData.TessellationTwistAngle = 0;
+		TessellationSmoothingData.TessellationTwistCurvature = 0;
+		TessellationSmoothingData.TessellationTotalSegmentLength = 0;
 	}
-
-	return DynamicData;
-}
-
-void FNiagaraRendererRibbons::AddDynamicParam(TArray<FNiagaraRibbonVertexDynamicParameter>& ParamData, const FVector4f& DynamicParam)
-{
-	FNiagaraRibbonVertexDynamicParameter Param;
-	Param.DynamicValue[0] = DynamicParam.X;
-	Param.DynamicValue[1] = DynamicParam.Y;
-	Param.DynamicValue[2] = DynamicParam.Z;
-	Param.DynamicValue[3] = DynamicParam.W;
-	ParamData.Add(Param);
 }
 
-bool FNiagaraRendererRibbons::IsMaterialValid(const UMaterialInterface* Mat)const
+int32 FNiagaraRendererRibbons::CalculateTessellationFactor(const FNiagaraSceneProxy* SceneProxy, const FSceneView* View, const FVector& ViewOriginForDistanceCulling) const
 {
-	return Mat && Mat->CheckMaterialUsage_Concurrent(MATUSAGE_NiagaraRibbons);
-}
-
-void FNiagaraRendererRibbons::SetupMeshBatchAndCollectorResourceForView(
-	const FSceneView* View,
-	const FSceneViewFamily& ViewFamily,
-	const FNiagaraSceneProxy* SceneProxy,
-	FMeshElementCollector& Collector,
-	struct FNiagaraDynamicDataRibbon* DynamicDataRibbon,
-	const FGlobalDynamicIndexBuffer::FAllocationEx& IndexAllocation,
-	FMeshBatch& MeshBatch,
-	class FNiagaraMeshCollectorResourcesRibbon& CollectorResources) const
-{
-	const bool bIsWireframe = ViewFamily.EngineShowFlags.Wireframe;
-	FMaterialRenderProxy* MaterialRenderProxy = DynamicDataRibbon->Material;
-	check(MaterialRenderProxy);
-
-
-	FNiagaraDataBuffer* SourceParticleData = DynamicDataRibbon->GetParticleDataToRender();
-	check(SourceParticleData); // if this is nullptr, should already be early out before entering this function
-
-	// Set common data on vertex factory
-	DynamicDataRibbon->SetVertexFactoryData(CollectorResources.VertexFactory);
-
-	FCPUSimParticleDataAllocation CPUSimParticleDataAllocation = AllocateParticleDataIfCPUSim(DynamicDataRibbon, Collector.GetDynamicReadBuffer());
-	auto& ParticleData = CPUSimParticleDataAllocation.ParticleData;
-
-	int32 ParticleDataFloatStride = SourceParticleData->GetNumInstances();
-	int32 ParticleDataHalfStride = SourceParticleData->GetNumInstances();
-
-	check(ParticleDataFloatStride == ParticleDataHalfStride);
-
-	// TODO: need to make these two a global alloc buffer as well, not recreate
-				// pass in the sorted indices so the VS can fetch the particle data in order
-	FReadBuffer SortedIndicesBuffer;
-	SortedIndicesBuffer.Initialize(TEXT("SortedIndicesBuffer"), sizeof(int32), DynamicDataRibbon->SortedIndices.Num(), EPixelFormat::PF_R32_SINT, BUF_Volatile);
-	void *IndexPtr = RHILockBuffer(SortedIndicesBuffer.Buffer, 0, DynamicDataRibbon->SortedIndices.Num() * sizeof(int32), RLM_WriteOnly);
-	FMemory::Memcpy(IndexPtr, DynamicDataRibbon->SortedIndices.GetData(), DynamicDataRibbon->SortedIndices.Num() * sizeof(int32));
-	RHIUnlockBuffer(SortedIndicesBuffer.Buffer);
-	CollectorResources.VertexFactory.SetSortedIndices(SortedIndicesBuffer.Buffer, SortedIndicesBuffer.SRV, 0);
-	// pass in the CPU generated total segment distance (for tiling distance modes); needs to be a buffer so we can fetch them in the correct order based on Draw Direction (front->back or back->front)
-	//	otherwise UVs will pop when draw direction changes based on camera view point
-	FReadBuffer TangentsAndDistancesBuffer;
-	TangentsAndDistancesBuffer.Initialize(TEXT("TangentsAndDistancesBuffer"), sizeof(FVector4f), DynamicDataRibbon->TangentAndDistances.Num(), EPixelFormat::PF_A32B32G32R32F, BUF_Volatile);
-	void *TangentsAndDistancesPtr = RHILockBuffer(TangentsAndDistancesBuffer.Buffer, 0, DynamicDataRibbon->TangentAndDistances.Num() * sizeof(FVector4f), RLM_WriteOnly);
-	FMemory::Memcpy(TangentsAndDistancesPtr, DynamicDataRibbon->TangentAndDistances.GetData(), DynamicDataRibbon->TangentAndDistances.Num() * sizeof(FVector4f));
-	RHIUnlockBuffer(TangentsAndDistancesBuffer.Buffer);
-	CollectorResources.VertexFactory.SetTangentAndDistances(TangentsAndDistancesBuffer.Buffer, TangentsAndDistancesBuffer.SRV);
-	// Copy a buffer which has the per particle multi ribbon index.
-	FReadBuffer MultiRibbonIndicesBuffer;
-	MultiRibbonIndicesBuffer.Initialize(TEXT("MultiRibbonIndicesBuffer"), sizeof(uint32), DynamicDataRibbon->MultiRibbonIndices.Num(), EPixelFormat::PF_R32_UINT, BUF_Volatile);
-	void* MultiRibbonIndexPtr = RHILockBuffer(MultiRibbonIndicesBuffer.Buffer, 0, DynamicDataRibbon->MultiRibbonIndices.Num() * sizeof(uint32), RLM_WriteOnly);
-	FMemory::Memcpy(MultiRibbonIndexPtr, DynamicDataRibbon->MultiRibbonIndices.GetData(), DynamicDataRibbon->MultiRibbonIndices.Num() * sizeof(uint32));
-	RHIUnlockBuffer(MultiRibbonIndicesBuffer.Buffer);
-	CollectorResources.VertexFactory.SetMultiRibbonIndicesSRV(MultiRibbonIndicesBuffer.Buffer, MultiRibbonIndicesBuffer.SRV);
-	// Copy the packed u data for stable age based uv generation.
-	FReadBuffer PackedPerRibbonDataByIndexBuffer;
-	PackedPerRibbonDataByIndexBuffer.Initialize(TEXT("PackedPerRibbonDataByIndexBuffer"), sizeof(float), DynamicDataRibbon->PackedPerRibbonDataByIndex.Num(), EPixelFormat::PF_R32_FLOAT, BUF_Volatile);
-	void *PackedPerRibbonDataByIndexPtr = RHILockBuffer(PackedPerRibbonDataByIndexBuffer.Buffer, 0, DynamicDataRibbon->PackedPerRibbonDataByIndex.Num() * sizeof(float), RLM_WriteOnly);
-	FMemory::Memcpy(PackedPerRibbonDataByIndexPtr, DynamicDataRibbon->PackedPerRibbonDataByIndex.GetData(), DynamicDataRibbon->PackedPerRibbonDataByIndex.Num() * sizeof(float));
-	RHIUnlockBuffer(PackedPerRibbonDataByIndexBuffer.Buffer);
-	CollectorResources.VertexFactory.SetPackedPerRibbonDataByIndexSRV(PackedPerRibbonDataByIndexBuffer.Buffer, PackedPerRibbonDataByIndexBuffer.SRV);
-
-	// Copy the packed offset data for slice vertices
-	FReadBuffer SliceVertexDataBuffer;
-	SliceVertexDataBuffer.Initialize(TEXT("SliceVertexDataBuffer"), sizeof(float), DynamicDataRibbon->SliceVertexData.Num(), EPixelFormat::PF_R32_FLOAT, BUF_Volatile);
-	void* SliceVertexDataBufferPtr = RHILockBuffer(SliceVertexDataBuffer.Buffer, 0, DynamicDataRibbon->SliceVertexData.Num() * sizeof(float), RLM_WriteOnly);
-	FMemory::Memcpy(SliceVertexDataBufferPtr, DynamicDataRibbon->SliceVertexData.GetData(), DynamicDataRibbon->SliceVertexData.Num() * sizeof(float));
-	RHIUnlockBuffer(SliceVertexDataBuffer.Buffer);
-	CollectorResources.VertexFactory.SetSliceVertexDataSRV(SliceVertexDataBuffer.Buffer, SliceVertexDataBuffer.SRV);
-
-	FRHIShaderResourceView* FloatSRV = ParticleData.FloatData.IsValid() ? ParticleData.FloatData.SRV : (FRHIShaderResourceView*)FNiagaraRenderer::GetDummyFloatBuffer();
-	FRHIShaderResourceView* HalfSRV = ParticleData.HalfData.IsValid() ? ParticleData.HalfData.SRV : (FRHIShaderResourceView*)FNiagaraRenderer::GetDummyHalfBuffer();
-
-	FNiagaraRibbonVFLooseParameters VFLooseParams;
-	VFLooseParams.SortedIndices = SortedIndicesBuffer.SRV;
-	VFLooseParams.TangentsAndDistances = TangentsAndDistancesBuffer.SRV;
-	VFLooseParams.MultiRibbonIndices = MultiRibbonIndicesBuffer.SRV;
-	VFLooseParams.PackedPerRibbonDataByIndex = PackedPerRibbonDataByIndexBuffer.SRV;
-	VFLooseParams.SliceVertexData = SliceVertexDataBuffer.SRV;
-	VFLooseParams.NiagaraParticleDataFloat = FloatSRV;
-	VFLooseParams.NiagaraParticleDataHalf = HalfSRV;
-	VFLooseParams.NiagaraFloatDataStride = ParticleDataFloatStride;
-	VFLooseParams.SortedIndicesOffset = CollectorResources.VertexFactory.GetSortedIndicesOffset();
-	VFLooseParams.FacingMode = static_cast<uint32>(FacingMode);
-	VFLooseParams.Shape = static_cast<uint32>(Shape);
-	VFLooseParams.NeedsPreciseMotionVectors = bNeedsPreciseMotionVectors;
-
-	// Collector.AllocateOneFrameResource uses default ctor, initialize the vertex factory
-	CollectorResources.VertexFactory.SetParticleFactoryType(NVFT_Ribbon);
-	CollectorResources.VertexFactory.LooseParameterUniformBuffer = FNiagaraRibbonVFLooseParametersRef::CreateUniformBufferImmediate(VFLooseParams, UniformBuffer_SingleFrame);
-	CollectorResources.VertexFactory.InitResource();
-	CollectorResources.VertexFactory.SetRibbonUniformBuffer(CollectorResources.UniformBuffer);
-	CollectorResources.VertexFactory.SetFacingMode(static_cast<uint32>(FacingMode));
-
-
-	MeshBatch.VertexFactory = &CollectorResources.VertexFactory;
-	MeshBatch.CastShadow = SceneProxy->CastsDynamicShadow();
-#if RHI_RAYTRACING
-	MeshBatch.CastRayTracedShadow = SceneProxy->CastsDynamicShadow();
-#endif
-	MeshBatch.bUseAsOccluder = false;
-	MeshBatch.ReverseCulling = SceneProxy->IsLocalToWorldDeterminantNegative();
-	MeshBatch.bDisableBackfaceCulling = Shape == ENiagaraRibbonShapeMode::Plane || !bEnableAccurateGeometry;
-	MeshBatch.Type = PT_TriangleList;
-	MeshBatch.DepthPriorityGroup = SceneProxy->GetDepthPriorityGroup(View);
-	MeshBatch.bCanApplyViewModeOverrides = true;
-	MeshBatch.bUseWireframeSelectionColoring = SceneProxy->IsSelected();
-	MeshBatch.SegmentIndex = 0;
-
-	if (bIsWireframe)
-	{
-		MeshBatch.MaterialRenderProxy = UMaterial::GetDefaultMaterial(MD_Surface)->GetRenderProxy();
-	}
-	else
-	{
-		MeshBatch.MaterialRenderProxy = MaterialRenderProxy;
-	}
-
-	FMeshBatchElement& MeshElement = MeshBatch.Elements[0];
-	MeshElement.IndexBuffer = IndexAllocation.IndexBuffer;
-	MeshElement.FirstIndex = IndexAllocation.FirstIndex;
-	MeshElement.NumPrimitives = IndexAllocation.NumIndices / 3; // 3 indices per triangle
-	check(MeshElement.NumPrimitives > 0);
-	MeshElement.NumInstances = 1;
-	MeshElement.MinVertexIndex = 0;
-	MeshElement.MaxVertexIndex = 0;
-	MeshElement.PrimitiveUniformBuffer = SceneProxy->GetCustomUniformBuffer(false);	// Note: Ribbons don't generate accurate velocities so disabling
-}
-
-FNiagaraRendererRibbons::FCPUSimParticleDataAllocation FNiagaraRendererRibbons::AllocateParticleDataIfCPUSim(FNiagaraDynamicDataRibbon* DynamicDataRibbon, FGlobalDynamicReadBuffer& DynamicReadBuffer) const
-{
-	FNiagaraDataBuffer* SourceParticleData = DynamicDataRibbon->GetParticleDataToRender();
-	check(SourceParticleData);//Can be null but should be checked before here.
-
-	FCPUSimParticleDataAllocation CPUSimParticleDataAllocation{ DynamicReadBuffer };
-
-	if (SimTarget == ENiagaraSimTarget::CPUSim)
-	{
-		CPUSimParticleDataAllocation.ParticleData = TransferDataToGPU(DynamicReadBuffer, RendererLayout, MakeArrayView<uint32>(nullptr, 0), SourceParticleData);
-	}
-
-	return CPUSimParticleDataAllocation;
-}
-
-int32 FNiagaraRendererRibbons::CalculateBitsForRange(int32 Range)
-{
-	return FMath::CeilToInt(FMath::Loge((float)Range) / FMath::Loge((float)2));
-}
-
-FNiagaraRendererRibbons::FRibbonRenderingIndexOffsets FNiagaraRendererRibbons::CalculateIndexBufferPacking(
-	int32 NumSegments,
-	int32 NumInterpolations,
-	int32 NumSliceVertices)
-{
-	uint32 NumSegmentBits = CalculateBitsForRange(NumSegments);
-	uint32 NumInterpolationBits = CalculateBitsForRange(NumInterpolations);
-	uint32 NumSliceVerticesBits = CalculateBitsForRange(NumSliceVertices);
-
-	FRibbonRenderingIndexOffsets Offsets;
-	Offsets.TotalBitCount = NumSegmentBits + NumInterpolationBits + NumSliceVerticesBits;
-
-	Offsets.SegmentBitShift = NumInterpolationBits + NumSliceVerticesBits;
-	Offsets.InterpBitShift = NumSliceVerticesBits;
-
-	Offsets.SegmentBitMask = uint64(0xFFFFFFFFul) >> (32 - NumSegmentBits);
-	Offsets.InterpBitMask = uint64(0xFFFFFFFF) >> (32 - NumInterpolationBits);
-	Offsets.SliceVertexBitMask = uint64(0xFFFFFFFF) >> (32 - NumSliceVerticesBits);
-
-	return Offsets;
-}
-
-void FNiagaraRendererRibbons::CreatePerViewResources(
-	const FSceneView* View,
-	const FSceneViewFamily& ViewFamily,
-	const FNiagaraSceneProxy* SceneProxy,
-	FMeshElementCollector& Collector,
-	FNiagaraRibbonUniformBufferRef& OutUniformBuffer,
-	FGlobalDynamicIndexBuffer::FAllocationEx& InOutIndexAllocation) const
-{
-	FNiagaraDynamicDataRibbon* DynamicDataRibbon = static_cast<FNiagaraDynamicDataRibbon*>(DynamicDataRender);
-	FNiagaraDataBuffer* SourceParticleData = DynamicDataRibbon->GetParticleDataToRender();
-	check(DynamicDataRibbon);
-	check(SourceParticleData);
-
 	bool bUseConstantFactor = false;
 	int32 TessellationFactor = GNiagaraRibbonMaxTessellation;
 	float TessellationMinAngle = GNiagaraRibbonTessellationAngle;
 	float ScreenPercentage = GNiagaraRibbonTessellationScreenPercentage;
-	switch (TessellationMode)
-
+	switch (TessellationConfig.TessellationMode)
 	{
 	case ENiagaraRibbonTessellationMode::Automatic:
 		break;
 	case ENiagaraRibbonTessellationMode::Custom:
-		TessellationFactor = FMath::Min<int32>(TessellationFactor, CustomTessellationFactor); // Don't allow factors bigger than the platform limit.
-		bUseConstantFactor = bCustomUseConstantFactor;
-		bUseConstantFactor = bCustomUseConstantFactor;
-		TessellationMinAngle = CustomTessellationMinAngle;
-		ScreenPercentage = bCustomUseScreenSpace && !bUseConstantFactor ? GNiagaraRibbonTessellationScreenPercentage : 0.f;
+		TessellationFactor = FMath::Min<int32>(TessellationFactor, TessellationConfig.CustomTessellationFactor); // Don't allow factors bigger than the platform limit.
+		bUseConstantFactor = TessellationConfig.bCustomUseConstantFactor;
+		TessellationMinAngle = TessellationConfig.CustomTessellationMinAngle;
+		ScreenPercentage = TessellationConfig.bCustomUseScreenSpace && !bUseConstantFactor ? GNiagaraRibbonTessellationScreenPercentage : 0.f;
 		break;
 	case ENiagaraRibbonTessellationMode::Disabled:
 		TessellationFactor = 1;
@@ -1301,81 +1644,223 @@ void FNiagaraRendererRibbons::CreatePerViewResources(
 		break;
 	}
 
-	const FVector ViewOriginForDistanceCulling = View->ViewMatrices.GetViewOrigin();
-
+	if (bUseConstantFactor)
+	{
+		return TessellationFactor;
+	}
+	
 	int32 SegmentTessellation = 1;
-	int32 NumSegments = DynamicDataRibbon->SegmentData.Num();
-	if (GNiagaraRibbonTessellationEnabled && TessellationFactor > 1 && TessellationCurvature > SMALL_NUMBER)
+	
+	if (GNiagaraRibbonTessellationEnabled && TessellationFactor > 1 && TessellationSmoothingData.TessellationCurvature > SMALL_NUMBER)
 	{
-		const float MinTesselation = [&]
-		{
-			if (TessellationMinAngle == 0.f || bUseConstantFactor)
-			{
-				return static_cast<float>(TessellationFactor);
-			}
-			else
-			{
-				return FMath::Max<float>(1.f, FMath::Max(TessellationTwistAngle, TessellationAngle) / FMath::Max<float>(SMALL_NUMBER, TessellationMinAngle));
-			}
-		}();
-		const float MAX_CURVATURE_FACTOR = 0.002f; // This will clamp the curvature to around 2.5 km and avoid numerical issues.
+		const float MinTesselation = (TessellationMinAngle == 0.f || bUseConstantFactor)?
+			                             static_cast<float>(TessellationFactor)	:
+			                             FMath::Max<float>(1.f, FMath::Max(TessellationSmoothingData.TessellationTwistAngle, TessellationSmoothingData.TessellationAngle) / FMath::Max<float>(SMALL_NUMBER, TessellationMinAngle));
+
+		constexpr float MAX_CURVATURE_FACTOR = 0.002f; // This will clamp the curvature to around 2.5 km and avoid numerical issues.
 		const float ViewDistance = SceneProxy->GetProxyDynamicData().LODDistanceOverride >= 0.0f ? SceneProxy->GetProxyDynamicData().LODDistanceOverride : SceneProxy->GetBounds().ComputeSquaredDistanceFromBoxToPoint(ViewOriginForDistanceCulling);
 		const float MaxDisplacementError = FMath::Max(GNiagaraRibbonTessellationMinDisplacementError, ScreenPercentage * FMath::Sqrt(ViewDistance) / View->LODDistanceFactor);
-		float Tess = TessellationAngle / FMath::Max(MAX_CURVATURE_FACTOR, AcosFast(TessellationCurvature / (TessellationCurvature + MaxDisplacementError)));
+		float Tess = TessellationSmoothingData.TessellationAngle / FMath::Max(MAX_CURVATURE_FACTOR, AcosFast(TessellationSmoothingData.TessellationCurvature / (TessellationSmoothingData.TessellationCurvature + MaxDisplacementError)));
 		// FMath::RoundUpToPowerOfTwo ? This could avoid vertices moving around as tesselation increases
 
-		if (TessellationTwistAngle > 0 && TessellationTwistCurvature > 0)
+		if (TessellationSmoothingData.TessellationTwistAngle > 0 && TessellationSmoothingData.TessellationTwistCurvature > 0)
 		{
-			const float TwistTess = TessellationTwistAngle / FMath::Max(MAX_CURVATURE_FACTOR, AcosFast(TessellationTwistCurvature / (TessellationTwistCurvature + MaxDisplacementError)));
+			const float TwistTess = TessellationSmoothingData.TessellationTwistAngle / FMath::Max(MAX_CURVATURE_FACTOR, AcosFast(TessellationSmoothingData.TessellationTwistCurvature / (TessellationSmoothingData.TessellationTwistCurvature + MaxDisplacementError)));
 			Tess = FMath::Max(TwistTess, Tess);
 		}
 		SegmentTessellation = FMath::Clamp<int32>(FMath::RoundToInt(Tess), FMath::RoundToInt(MinTesselation), TessellationFactor);
-		NumSegments *= SegmentTessellation;
 	}
 
-	int32 TrianglesPerSegment = 2;
-	int32 NumVerticesInSlice = 0;
+	return SegmentTessellation;
+}
+
 
-	if (Shape == ENiagaraRibbonShapeMode::MultiPlane)
+FNiagaraIndexGenerationInput FNiagaraRendererRibbons::CalculateIndexBufferConfiguration(const TSharedPtr<FNiagaraRibbonCPUGeneratedVertexData>& GeneratedVertices, const FNiagaraDataBuffer* SourceParticleData,
+	const FNiagaraSceneProxy* SceneProxy, const FSceneView* View, const FVector& ViewOriginForDistanceCulling, bool bShouldUseGPUInitIndices, bool bIsGPUSim) const
+{
+	FNiagaraIndexGenerationInput IndexGenInput;
+
+	IndexGenInput.ViewDistance = SceneProxy->GetProxyDynamicData().LODDistanceOverride >= 0.0f ? SceneProxy->GetProxyDynamicData().LODDistanceOverride : SceneProxy->GetBounds().ComputeSquaredDistanceFromBoxToPoint(ViewOriginForDistanceCulling);
+	IndexGenInput.LODDistanceFactor = View->LODDistanceFactor;
+
+	if (bShouldUseGPUInitIndices)
+	{
+		// If we are a gpu sim, we rely on num instances allocated since between now and render we could have a different number of particles living
+		// If we're not a gpu sim, we're a cpu sim with gpu init, we can rely on the num particles - 1 being the max cap for segments
+		IndexGenInput.MaxSegmentCount = bIsGPUSim? SourceParticleData->GetNumInstancesAllocated() : SourceParticleData->GetNumInstances();
+	}
+	else
 	{
-		TrianglesPerSegment *= MultiPlaneCount * WidthSegmentationCount * (bEnableAccurateGeometry? 2 : 1);
-		NumVerticesInSlice = MultiPlaneCount * (WidthSegmentationCount + 1) * (bEnableAccurateGeometry ? 2 : 1);
+		IndexGenInput.MaxSegmentCount = GeneratedVertices->SortedIndices.Num();
 	}
-	else if (Shape == ENiagaraRibbonShapeMode::Tube)
+	
+	
+
+	IndexGenInput.SubSegmentCount = 1;
+	if (GenerationConfig.WantsAutomaticTessellation() || GenerationConfig.WantsConstantTessellation())
+	{
+		if (bShouldUseGPUInitIndices)
+		{
+			// if we have a constant factor, use it, if not set it to the max allowed since we won't know what we need exactly until later on.
+			IndexGenInput.SubSegmentCount = (TessellationConfig.TessellationMode == ENiagaraRibbonTessellationMode::Custom && TessellationConfig.bCustomUseConstantFactor)?
+				TessellationConfig.CustomTessellationFactor : GNiagaraRibbonMaxTessellation;
+		}
+		else
+		{
+			IndexGenInput.SubSegmentCount = CalculateTessellationFactor(SceneProxy, View, ViewOriginForDistanceCulling);
+		}
+	}	
+	const uint32 NumSegmentBits = CalculateBitsForRange(IndexGenInput.MaxSegmentCount);
+	const uint32 NumSubSegmentBits = CalculateBitsForRange(IndexGenInput.SubSegmentCount);
+	
+	IndexGenInput.SegmentBitShift = NumSubSegmentBits + ShapeState.BitsNeededForShape;
+	IndexGenInput.SubSegmentBitShift = ShapeState.BitsNeededForShape;
+
+	IndexGenInput.SegmentBitMask = CalculateBitMask(NumSegmentBits);
+	IndexGenInput.SubSegmentBitMask = CalculateBitMask(NumSubSegmentBits);
+
+	IndexGenInput.ShapeBitMask = ShapeState.BitMaskForShape;
+	
+	IndexGenInput.TotalBitCount = NumSegmentBits + NumSubSegmentBits + ShapeState.BitsNeededForShape;
+	IndexGenInput.TotalNumIndices = IndexGenInput.MaxSegmentCount * IndexGenInput.SubSegmentCount * ShapeState.TrianglesPerSegment * 3;
+
+	return IndexGenInput;
+}
+
+void FNiagaraRendererRibbons::GenerateIndexBufferForView(FNiagaraIndexGenerationInput& GeneratedData, FNiagaraDynamicDataRibbon* DynamicDataRibbon,
+                                                         const TSharedPtr<FNiagaraRibbonRenderingFrameViewResources>& RenderingViewResources, const FSceneView* View,
+                                                         const FVector& ViewOriginForDistanceCulling, bool bShouldUseGPUInitIndices) const
+{
+	if (GeneratedData.MaxSegmentCount > 0)
 	{
-		TrianglesPerSegment *= TubeSubdivisions;
-		NumVerticesInSlice = TubeSubdivisions + 1;
+		if (bShouldUseGPUInitIndices)
+		{
+			RenderingViewResources->IndirectDrawBuffer.Initialize(TEXT("RibbonIndirectDrawBuffer"), sizeof(uint32), FNiagaraRibbonIndirectDrawBufferLayout::NumElements, EPixelFormat::PF_R32_UINT, ERHIAccess::IndirectArgs | ERHIAccess::SRVMask, BUF_Static | BUF_DrawIndirect);			
+		}
+		
+		if (GeneratedData.TotalBitCount <= 16/* Number of bits in a ushort*/ && !bShouldUseGPUInitIndices)
+		{
+			RenderingViewResources->IndexBuffer.Initialize(TEXT("NiagaraRibbonIndexBuffer"), sizeof(uint16), GeneratedData.TotalNumIndices, PF_R16_UINT, ERHIAccess::VertexOrIndexBuffer, BUF_Static);
+			if (!bShouldUseGPUInitIndices)
+			{
+				GenerateIndexBufferCPU<uint16>(GeneratedData, DynamicDataRibbon, ShapeState, RenderingViewResources, View, ViewOriginForDistanceCulling, FeatureLevel, DrawDirection);
+			}
+		}
+		else
+		{		
+			RenderingViewResources->IndexBuffer.Initialize(TEXT("NiagaraRibbonIndexBuffer"), sizeof(uint32), GeneratedData.TotalNumIndices, PF_R32_UINT, ERHIAccess::VertexOrIndexBuffer, BUF_Static);
+			if (!bShouldUseGPUInitIndices)
+			{
+				GenerateIndexBufferCPU<uint32>(GeneratedData, DynamicDataRibbon, ShapeState, RenderingViewResources, View, ViewOriginForDistanceCulling, FeatureLevel, DrawDirection);				
+			}
+		}
 	}
-	else if (Shape == ENiagaraRibbonShapeMode::Custom && CustomVertices.Num() >= 2)
+}
+
+template <typename TValue>
+void FNiagaraRendererRibbons::GenerateIndexBufferCPU(FNiagaraIndexGenerationInput& GeneratedData, FNiagaraDynamicDataRibbon* DynamicDataRibbon, const FNiagaraRibbonShapeGeometryData& ShapeState,
+    const TSharedPtr<FNiagaraRibbonRenderingFrameViewResources>& RenderingViewResources, const FSceneView* View, const FVector& ViewOriginForDistanceCulling, ERHIFeatureLevel::Type FeatureLevel, ENiagaraRibbonDrawDirection DrawDirection)
+{
+	SCOPE_CYCLE_COUNTER(STAT_NiagaraRenderRibbonsGenIndiciesCPU);
+	FMaterialRenderProxy* MaterialRenderProxy = DynamicDataRibbon->Material;
+	check(MaterialRenderProxy);
+	const EBlendMode BlendMode = MaterialRenderProxy->GetIncompleteMaterialWithFallback(FeatureLevel).GetBlendMode();
+	
+	const TSharedPtr<FNiagaraRibbonCPUGeneratedVertexData>& GeneratedGeometryData = DynamicDataRibbon->GenerationOutput;
+
+	TArray<TValue> IndexBufferData;
+	IndexBufferData.SetNum(GeneratedData.TotalNumIndices);
+	TValue* CurrentIndexBuffer = IndexBufferData.GetData();
+	if (IsTranslucentBlendMode(BlendMode) && GeneratedGeometryData->RibbonInfoLookup.Num())
 	{
-		TrianglesPerSegment *= CustomVertices.Num();
-		NumVerticesInSlice = CustomVertices.Num() + 1;
+		for (const FRibbonMultiRibbonInfo& MultiRibbonInfo : GeneratedGeometryData->RibbonInfoLookup)
+		{
+			const TArrayView<uint32> CurrentSegmentData(GeneratedGeometryData->SegmentData.GetData() + MultiRibbonInfo.BaseSegmentDataIndex, MultiRibbonInfo.NumSegmentDataIndices);
+			CurrentIndexBuffer = AppendToIndexBufferCPU<TValue>(CurrentIndexBuffer, GeneratedData, ShapeState, CurrentSegmentData, MultiRibbonInfo.UseInvertOrder(View->GetViewDirection(), ViewOriginForDistanceCulling, DrawDirection));
+		}
 	}
-	else // Plane
+	else // Otherwise ignore multi ribbon ordering.
 	{
-		TrianglesPerSegment *= WidthSegmentationCount;
-		NumVerticesInSlice = WidthSegmentationCount + 1;
+		const TArrayView<uint32> CurrentSegmentData(GeneratedGeometryData->SegmentData.GetData(), GeneratedGeometryData->SegmentData.Num());
+		CurrentIndexBuffer = AppendToIndexBufferCPU<TValue>(CurrentIndexBuffer, GeneratedData, ShapeState, CurrentSegmentData, false);
 	}
+	
+	// Copy the indices over to the GPU
+	void* IndexBufferPtr = RHILockBuffer(RenderingViewResources->IndexBuffer.IndexBufferRHI, 0, IndexBufferData.Num() * sizeof(TValue), RLM_WriteOnly);
+	FMemory::Memcpy(IndexBufferPtr, IndexBufferData.GetData(), IndexBufferData.Num() * sizeof(TValue));
+	RHIUnlockBuffer(RenderingViewResources->IndexBuffer.IndexBufferRHI);
+}
 
-	FRibbonRenderingIndexOffsets IndexBufferOffsets = CalculateIndexBufferPacking(
-		DynamicDataRibbon->MaxParticleIndex + 1, /* Add one as this needs to be a count, not a max index */
-		SegmentTessellation, NumVerticesInSlice);
 
-	// Copy the index data over.
-	FGlobalDynamicIndexBuffer& DynamicIndexBuffer = Collector.GetDynamicIndexBuffer();
 
-	const uint32 NumIndices = NumSegments * TrianglesPerSegment * 3;
-	if (IndexBufferOffsets.TotalBitCount <= 16/* Number of bits in a ushort*/)
+template <typename TValue>
+TValue* FNiagaraRendererRibbons::AppendToIndexBufferCPU(TValue* OutIndices, const FNiagaraIndexGenerationInput& GeneratedData, const FNiagaraRibbonShapeGeometryData& ShapeState, const TArrayView<uint32>& SegmentData, bool bInvertOrder)
+{
+	if (SegmentData.Num() == 0)
 	{
-		InOutIndexAllocation = DynamicIndexBuffer.Allocate<uint16>(NumIndices);
-		GenerateIndexBuffer<uint16>(InOutIndexAllocation, IndexBufferOffsets, SegmentTessellation, View->GetViewDirection(), ViewOriginForDistanceCulling, DynamicDataRibbon);
+		return OutIndices;
 	}
-	else
+
+	const int32 FirstSegmentDataIndex = bInvertOrder ? SegmentData.Num() - 1 : 0;
+	const int32 LastSegmentDataIndex = bInvertOrder ? -1 : SegmentData.Num();
+	const int32 SegmentDataIndexInc = bInvertOrder ? -1 : 1;
+	const int32 FlipGeometryIndex = FMath::Min(FMath::Max(ShapeState.SliceTriangleToVertexIds.Num() / 2, 2), ShapeState.SliceTriangleToVertexIds.Num());
+	
+	for (int32 SegmentDataIndex = FirstSegmentDataIndex; SegmentDataIndex != LastSegmentDataIndex; SegmentDataIndex += SegmentDataIndexInc)
 	{
-		InOutIndexAllocation = DynamicIndexBuffer.Allocate<uint32>(NumIndices);
-		GenerateIndexBuffer<uint32>(InOutIndexAllocation, IndexBufferOffsets, SegmentTessellation, View->GetViewDirection(), ViewOriginForDistanceCulling, DynamicDataRibbon);
+		const int32 SegmentIndex = SegmentData[SegmentDataIndex];
+		for (uint32 SubSegmentIndex = 0; SubSegmentIndex < GeneratedData.SubSegmentCount; ++SubSegmentIndex)
+		{
+			const bool bIsFinalInterp = SubSegmentIndex == GeneratedData.SubSegmentCount - 1;
+
+			const int32 ThisSegmentOffset = SegmentIndex << GeneratedData.SegmentBitShift;
+			const int32 NextSegmentOffset = (SegmentIndex + (bIsFinalInterp ? 1 : 0)) << GeneratedData.SegmentBitShift;
+
+			const int32 ThisSubSegmentOffset = SubSegmentIndex << GeneratedData.SubSegmentBitShift;
+			const int32 NextSubSegmentOffset = (bIsFinalInterp ? 0 : SubSegmentIndex + 1) << GeneratedData.SubSegmentBitShift;
+
+			const int32 CurrSegment = ThisSegmentOffset | ThisSubSegmentOffset;
+			const int32 NextSegment = NextSegmentOffset | NextSubSegmentOffset;
+
+			int32 TriangleId = 0;
+			
+			for (; TriangleId < FlipGeometryIndex; TriangleId += 2)
+			{
+				const int32 FirstIndex = ShapeState.SliceTriangleToVertexIds[TriangleId];
+				const int32 SecondIndex = ShapeState.SliceTriangleToVertexIds[TriangleId + 1];
+				
+				OutIndices[0] = CurrSegment | FirstIndex;
+				OutIndices[1] = CurrSegment | SecondIndex;
+				OutIndices[2] = NextSegment | FirstIndex;
+				OutIndices[3] = OutIndices[1];
+				OutIndices[4] = NextSegment | SecondIndex;
+				OutIndices[5] = OutIndices[2];
+
+				OutIndices += 6;
+			}
+			for (; TriangleId < ShapeState.SliceTriangleToVertexIds.Num(); TriangleId += 2)
+			{
+				const int32 FirstIndex = ShapeState.SliceTriangleToVertexIds[TriangleId];
+				const int32 SecondIndex = ShapeState.SliceTriangleToVertexIds[TriangleId + 1];
+
+				OutIndices[0] = CurrSegment | FirstIndex;
+				OutIndices[1] = CurrSegment | SecondIndex;
+				OutIndices[2] = NextSegment | SecondIndex;
+				OutIndices[3] = OutIndices[0];
+				OutIndices[4] = OutIndices[2];
+				OutIndices[5] = NextSegment | FirstIndex;
+
+				OutIndices += 6;
+			}			
+		}		
 	}
 
+	return OutIndices;
+}
+
+void FNiagaraRendererRibbons::SetupPerViewUniformBuffer(FNiagaraIndexGenerationInput& GeneratedData, const FSceneView* View,
+	const FSceneViewFamily& ViewFamily, const FNiagaraSceneProxy* SceneProxy, FNiagaraRibbonUniformBufferRef& OutUniformBuffer) const
+{	
 	FNiagaraRibbonUniformParameters PerViewUniformParameters;
 	FMemory::Memzero(&PerViewUniformParameters,sizeof(PerViewUniformParameters)); // Clear unset bytes
 
@@ -1383,18 +1868,18 @@ void FNiagaraRendererRibbons::CreatePerViewResources(
 	PerViewUniformParameters.bLocalSpace = bUseLocalSpace;
 	PerViewUniformParameters.DeltaSeconds = ViewFamily.Time.GetDeltaWorldTimeSeconds();
 	PerViewUniformParameters.SystemLWCTile = SceneProxy->GetLWCRenderTile();
-	PerViewUniformParameters.CameraUp = (FVector3f)View->GetViewUp(); // FVector4(0.0f, 0.0f, 1.0f, 0.0f);
-	PerViewUniformParameters.CameraRight = (FVector3f)View->GetViewRight();//	FVector4(1.0f, 0.0f, 0.0f, 0.0f);
+	PerViewUniformParameters.CameraUp = static_cast<FVector3f>(View->GetViewUp()); // FVector4(0.0f, 0.0f, 1.0f, 0.0f);
+	PerViewUniformParameters.CameraRight = static_cast<FVector3f>(View->GetViewRight());//	FVector4(1.0f, 0.0f, 0.0f, 0.0f);
 	PerViewUniformParameters.ScreenAlignment = FVector4f(0.0f, 0.0f, 0.0f, 0.0f);
-	PerViewUniformParameters.TotalNumInstances = SourceParticleData->GetNumInstances();
-	PerViewUniformParameters.InterpCount = SegmentTessellation;
-	PerViewUniformParameters.OneOverInterpCount = 1.f / (float)SegmentTessellation;
-	PerViewUniformParameters.ParticleIdShift = IndexBufferOffsets.SegmentBitShift;
-	PerViewUniformParameters.ParticleIdMask = IndexBufferOffsets.SegmentBitMask;
-	PerViewUniformParameters.InterpIdShift = IndexBufferOffsets.InterpBitShift;
-	PerViewUniformParameters.InterpIdMask = IndexBufferOffsets.InterpBitMask;
-	PerViewUniformParameters.SliceVertexIdMask = IndexBufferOffsets.SliceVertexBitMask;
-	PerViewUniformParameters.ShouldFlipNormalToView = Shape == ENiagaraRibbonShapeMode::MultiPlane && !bEnableAccurateGeometry;
+	PerViewUniformParameters.InterpCount = GeneratedData.SubSegmentCount;
+	PerViewUniformParameters.OneOverInterpCount = 1.f / static_cast<float>(GeneratedData.SubSegmentCount);
+	PerViewUniformParameters.ParticleIdShift = GeneratedData.SegmentBitShift;
+	PerViewUniformParameters.ParticleIdMask = GeneratedData.SegmentBitMask;
+	PerViewUniformParameters.InterpIdShift = GeneratedData.SubSegmentBitShift;
+	PerViewUniformParameters.InterpIdMask = GeneratedData.SubSegmentBitMask;
+	PerViewUniformParameters.SliceVertexIdMask = ShapeState.BitMaskForShape;
+	PerViewUniformParameters.ShouldFlipNormalToView = ShapeState.bShouldFlipNormalToView;
+	PerViewUniformParameters.ShouldUseMultiRibbon = GenerationConfig.HasRibbonIDs()? 1 : 0;
 
 	TConstArrayView<FNiagaraRendererVariableInfo> VFVariables = RendererLayout->GetVFVariables_RenderThread();
 	PerViewUniformParameters.PositionDataOffset = VFVariables[ENiagaraRibbonVFLayout::Position].GetGPUOffset();
@@ -1420,104 +1905,722 @@ void FNiagaraRendererRibbons::CreatePerViewResources(
 	PerViewUniformParameters.U1OverrideDataOffset = UV1Settings.bEnablePerParticleUOverride ? VFVariables[ENiagaraRibbonVFLayout::U1Override].GetGPUOffset() : -1;
 	PerViewUniformParameters.V1RangeOverrideDataOffset = UV1Settings.bEnablePerParticleVRangeOverride ? VFVariables[ENiagaraRibbonVFLayout::V1RangeOverride].GetGPUOffset() : -1;
 
-	PerViewUniformParameters.MaterialParamValidMask = MaterialParamValidMask;
+	PerViewUniformParameters.MaterialParamValidMask = GenerationConfig.GetMaterialParamValidMask();
 
 	bool bShouldDoFacing = FacingMode == ENiagaraRibbonFacingMode::Custom || FacingMode == ENiagaraRibbonFacingMode::CustomSideVector;
 	PerViewUniformParameters.FacingDataOffset = bShouldDoFacing ? VFVariables[ENiagaraRibbonVFLayout::Facing].GetGPUOffset() : -1;
 	PerViewUniformParameters.PrevFacingDataOffset = bShouldDoFacing ? VFVariables[ENiagaraRibbonVFLayout::PrevRibbonFacing].GetGPUOffset() : -1;
 
-	PerViewUniformParameters.U0DistributionMode = (int32)UV0Settings.DistributionMode;
-	PerViewUniformParameters.U1DistributionMode = (int32)UV1Settings.DistributionMode;
+	PerViewUniformParameters.LinkOrderDataOffset = VFVariables[ENiagaraRibbonVFLayout::LinkOrder].GetGPUOffset();
+
+	PerViewUniformParameters.U0DistributionMode = static_cast<int32>(UV0Settings.DistributionMode);
+	PerViewUniformParameters.U1DistributionMode = static_cast<int32>(UV1Settings.DistributionMode);
 	PerViewUniformParameters.PackedVData = FVector4f(UV0Settings.Scale.Y, UV0Settings.Offset.Y, UV1Settings.Scale.Y, UV1Settings.Offset.Y);
 
 	OutUniformBuffer = FNiagaraRibbonUniformBufferRef::CreateUniformBufferImmediate(PerViewUniformParameters, UniformBuffer_SingleFrame);
-
 }
 
-#if RHI_RAYTRACING
-void FNiagaraRendererRibbons::GetDynamicRayTracingInstances(FRayTracingMaterialGatheringContext& Context, TArray<FRayTracingInstance>& OutRayTracingInstances, const FNiagaraSceneProxy* SceneProxy)
+inline void FNiagaraRendererRibbons::SetupMeshBatchAndCollectorResourceForView(const FNiagaraIndexGenerationInput& GeneratedData, FNiagaraDynamicDataRibbon* DynamicDataRibbon, const FNiagaraDataBuffer* SourceParticleData, const FSceneView* View,
+    const FSceneViewFamily& ViewFamily, const FNiagaraSceneProxy* SceneProxy, const TSharedPtr<FNiagaraRibbonRenderingFrameResources>& RenderingResources, const TSharedPtr<FNiagaraRibbonRenderingFrameViewResources>& RenderingViewResources,
+    FMeshBatch& OutMeshBatch, bool bShouldUseGPUInitIndices) const
 {
-	if (!CVarRayTracingNiagaraRibbons.GetValueOnRenderThread())
+	const bool bIsWireframe = ViewFamily.EngineShowFlags.Wireframe;
+	FMaterialRenderProxy* MaterialRenderProxy = DynamicDataRibbon->Material;
+	check(MaterialRenderProxy);
+	
+	// Set common data on vertex factory
+	DynamicDataRibbon->SetVertexFactoryData(RenderingViewResources->VertexFactory);
+
+	FNiagaraRibbonVFLooseParameters VFLooseParams;
+	VFLooseParams.SortedIndices = VertexBuffers.SortedIndicesBuffer.SRV;
+	VFLooseParams.TangentsAndDistances = VertexBuffers.TangentsAndDistancesBuffer.SRV;
+	VFLooseParams.MultiRibbonIndices = GetSrvOrDefaultUInt(VertexBuffers.MultiRibbonIndicesBuffer.SRV);
+	VFLooseParams.PackedPerRibbonDataByIndex = VertexBuffers.RibbonLookupTableBuffer.SRV;
+	VFLooseParams.SliceVertexData = ShapeState.SliceVertexDataBuffer.SRV;
+	VFLooseParams.NiagaraParticleDataFloat = RenderingResources->ParticleFloatSRV;
+	VFLooseParams.NiagaraParticleDataHalf = RenderingResources->ParticleHalfSRV;
+	VFLooseParams.NiagaraFloatDataStride = RenderingResources->ParticleFloatDataStride;
+	VFLooseParams.FacingMode = static_cast<uint32>(FacingMode);
+	VFLooseParams.Shape = static_cast<uint32>(ShapeState.Shape);
+	VFLooseParams.NeedsPreciseMotionVectors = GenerationConfig.NeedsPreciseMotionVectors();
+
+	VFLooseParams.IndirectDrawOutput = bShouldUseGPUInitIndices? (FRHIShaderResourceView*)RenderingViewResources->IndirectDrawBuffer.SRV : GetDummyIntBuffer();
+	VFLooseParams.IndirectDrawOutputOffset = bShouldUseGPUInitIndices? 0 : -1;
+
+	// Collector.AllocateOneFrameResource uses default ctor, initialize the vertex factory
+	RenderingViewResources->VertexFactory.SetParticleFactoryType(NVFT_Ribbon);
+	RenderingViewResources->VertexFactory.LooseParameterUniformBuffer = FNiagaraRibbonVFLooseParametersRef::CreateUniformBufferImmediate(VFLooseParams, UniformBuffer_SingleFrame);
+	RenderingViewResources->VertexFactory.InitResource();
+	RenderingViewResources->VertexFactory.SetRibbonUniformBuffer(RenderingViewResources->UniformBuffer);
+
+
+	OutMeshBatch.VertexFactory = &RenderingViewResources->VertexFactory;
+	OutMeshBatch.CastShadow = SceneProxy->CastsDynamicShadow();
+#if RHI_RAYTRACING
+	OutMeshBatch.CastRayTracedShadow = SceneProxy->CastsDynamicShadow();
+#endif
+	OutMeshBatch.bUseAsOccluder = false;
+	OutMeshBatch.ReverseCulling = SceneProxy->IsLocalToWorldDeterminantNegative();
+	OutMeshBatch.bDisableBackfaceCulling = ShapeState.bDisableBackfaceCulling;
+	OutMeshBatch.Type = PT_TriangleList;
+	OutMeshBatch.DepthPriorityGroup = SceneProxy->GetDepthPriorityGroup(View);
+	OutMeshBatch.bCanApplyViewModeOverrides = true;
+	OutMeshBatch.bUseWireframeSelectionColoring = SceneProxy->IsSelected();
+	OutMeshBatch.SegmentIndex = 0;
+	OutMeshBatch.MaterialRenderProxy = bIsWireframe? UMaterial::GetDefaultMaterial(MD_Surface)->GetRenderProxy() : MaterialRenderProxy;
+	
+	FMeshBatchElement& MeshElement = OutMeshBatch.Elements[0];
+	MeshElement.IndexBuffer = &RenderingViewResources->IndexBuffer;
+	MeshElement.FirstIndex = 0;
+	MeshElement.NumInstances = 1;
+	MeshElement.MinVertexIndex = 0;
+	MeshElement.MaxVertexIndex = 0;
+
+	if (bShouldUseGPUInitIndices)
 	{
-		return;
+		MeshElement.NumPrimitives = 0;
+		MeshElement.IndirectArgsBuffer = RenderingViewResources->IndirectDrawBuffer.Buffer;
+		MeshElement.IndirectArgsOffset = RenderingViewResources->IndirectDrawBufferStartByteOffset;
 	}
+	else
+	{
+		MeshElement.NumPrimitives = GeneratedData.TotalNumIndices / 3; // 3 indices per triangle
+		MeshElement.MaxVertexIndex = SourceParticleData->GetNumInstances();
+		check(MeshElement.NumPrimitives > 0);
+	}	
 	
-	check(SceneProxy);
+	// TODO: MotionVector/Velocity? Probably need to look into this?
+	MeshElement.PrimitiveUniformBuffer = SceneProxy->GetCustomUniformBuffer(false);	// Note: Ribbons don't generate accurate velocities so disabling	
+}
 
-	FNiagaraDynamicDataRibbon *DynamicDataRibbon = static_cast<FNiagaraDynamicDataRibbon*>(DynamicDataRender);
-	FNiagaraGpuComputeDispatchInterface* ComputeDispatchInterface = SceneProxy->GetComputeDispatchInterface();
-	if (!DynamicDataRibbon || !ComputeDispatchInterface)
+template<typename Type>
+static inline TArray<Type> ReadBackBuffer(FRHIBuffer* Buffer, int32 NumValues)
+{
+	TArray<Type> Test;
+	Test.SetNum(NumValues);	
+	const void* DebugDataIntPtr = RHILockBuffer(Buffer, 0, NumValues * sizeof(Type), RLM_ReadOnly);
+	FMemory::Memcpy(Test.GetData(), DebugDataIntPtr, NumValues * sizeof(Type));
+	RHIUnlockBuffer(Buffer);
+	return Test;
+}
+
+void FNiagaraRendererRibbons::InitializeViewIndexBuffersGPU(FRHICommandListImmediate& CMDList, FNiagaraGpuComputeDispatchInterface* ComputeDispatchInterface,
+	const FNiagaraDataBuffer* SourceParticleData, const TSharedPtr<FNiagaraRibbonRenderingFrameViewResources>& RenderingViewResources) const
+{
+	SCOPE_CYCLE_COUNTER(STAT_NiagaraRenderRibbonsGenIndiciesGPU);
+	
+	const uint32 NumInstances = SourceParticleData->GetNumInstances();
+		
+	if (!RenderingViewResources->IndirectDrawBuffer.Buffer.IsValid())
 	{
 		return;
 	}
-
-	if (!DynamicDataRibbon->SortedIndices.Num())
+	
 	{
-		return;
+		FNiagaraRibbonCreateIndexBufferParamsCS::FPermutationDomain PermutationVector;
+		PermutationVector.Set<FRibbonWantsAutomaticTessellation>(GenerationConfig.WantsAutomaticTessellation());
+		PermutationVector.Set<FRibbonWantsConstantTessellation>(GenerationConfig.WantsConstantTessellation());
+			
+		TShaderMapRef<FNiagaraRibbonCreateIndexBufferParamsCS> ComputeShader(GetGlobalShaderMap(GMaxRHIFeatureLevel), PermutationVector);
+
+		FNiagaraRibbonInitializeIndices Params;
+		FMemory::Memzero(Params);
+
+		Params.IndirectDrawOutput = RenderingViewResources->IndirectDrawBuffer.UAV;
+		Params.VertexGenerationResults = VertexBuffers.GPUComputeCommandBuffer.SRV;
+
+		// Total particle Count
+		Params.TotalNumParticlesDirect = SourceParticleData->GetNumInstances();
+
+		// Indirect particle Count
+		Params.EmitterParticleCountsBuffer = GetSrvOrDefaultUInt(ComputeDispatchInterface->GetGPUInstanceCounterManager().GetInstanceCountBuffer());
+		Params.EmitterParticleCountsBufferOffset = SourceParticleData->GetGPUInstanceCountBufferOffset();
+
+		Params.IndirectDrawOutputIndex = 0;
+		Params.VertexGenerationResultsIndex = 0; /*Offset into command buffer*/
+		Params.IndexGenThreadSize = FNiagaraRibbonComputeCommon::IndexGenThreadSize;
+		Params.TrianglesPerSegment = ShapeState.TrianglesPerSegment;
+
+		Params.ViewDistance = RenderingViewResources->IndexGenerationSettings.ViewDistance;
+		Params.LODDistanceFactor = RenderingViewResources->IndexGenerationSettings.LODDistanceFactor;
+		Params.TessellationMode = static_cast<uint32>(TessellationConfig.TessellationMode);
+		Params.bCustomUseConstantFactor = TessellationConfig.bCustomUseConstantFactor ? 1 : 0;
+		Params.CustomTessellationFactor = TessellationConfig.CustomTessellationFactor;
+		Params.CustomTessellationMinAngle = TessellationConfig.CustomTessellationMinAngle;
+		Params.bCustomUseScreenSpace = TessellationConfig.bCustomUseScreenSpace ? 1 : 0;
+		Params.GNiagaraRibbonMaxTessellation = GNiagaraRibbonMaxTessellation;
+		Params.GNiagaraRibbonTessellationAngle = GNiagaraRibbonTessellationAngle;
+		Params.GNiagaraRibbonTessellationScreenPercentage = GNiagaraRibbonTessellationScreenPercentage;
+		Params.GNiagaraRibbonTessellationEnabled = GNiagaraRibbonTessellationEnabled ? 1 : 0;
+		Params.GNiagaraRibbonTessellationMinDisplacementError = GNiagaraRibbonTessellationMinDisplacementError;
+
+		CMDList.Transition(FRHITransitionInfo(RenderingViewResources->IndirectDrawBuffer.UAV, ERHIAccess::SRVMask | ERHIAccess::IndirectArgs, ERHIAccess::UAVCompute));
+		SetComputePipelineState(CMDList, ComputeShader.GetComputeShader());
+		SetShaderParameters(CMDList, ComputeShader, ComputeShader.GetComputeShader(), Params);
+		DispatchComputeShader(CMDList, ComputeShader, 1, 1, 1);
+		UnsetShaderUAVs(CMDList, ComputeShader, ComputeShader.GetComputeShader());
+		CMDList.Transition(FRHITransitionInfo(RenderingViewResources->IndirectDrawBuffer.UAV, ERHIAccess::UAVCompute, ERHIAccess::SRVMask | ERHIAccess::IndirectArgs));
 	}
-
-	FNiagaraDataBuffer* SourceParticleData = DynamicDataRibbon->GetParticleDataToRender();
-	if (SourceParticleData == nullptr ||
-		SourceParticleData->GetNumInstancesAllocated() == 0 ||
-		SourceParticleData->GetNumInstances() == 0 ||
-		GbEnableNiagaraRibbonRendering == 0)
+	
+	// Not possible to have a valid ribbon with less than 2 particles, abort!
+	// but we do need to write out the indirect draw so it will behave correctly.
+	// So the initialize call above sets up the indirect draw, but we'll skip the actual index gen below.
+	if (NumInstances < 2)
 	{
 		return;
 	}
+		
+	{
+		FNiagaraRibbonCreateIndexBufferCS::FPermutationDomain PermutationVector;
+		PermutationVector.Set<FRibbonHasFullRibbonID>(GenerationConfig.HasFullRibbonIDs());
+		PermutationVector.Set<FRibbonHasRibbonID>(GenerationConfig.HasSimpleRibbonIDs());
+
+		// This switches the index gen from a unrolled limited loop for performance to a full loop that can handle anything thrown at it
+		PermutationVector.Set<FRibbonHasHighSliceComplexity>(ShapeState.TrianglesPerSegment > 32);
+		
+		TShaderMapRef<FNiagaraRibbonCreateIndexBufferCS> ComputeShader(GetGlobalShaderMap(GMaxRHIFeatureLevel), PermutationVector);
+		// const int TotalNumInvocations = Params.TotalNumParticles * IndicesConfig.SubSegmentCount;
+		// const uint32 NumThreadGroups = FMath::DivideAndRoundUp<uint32>(TotalNumInvocations, FNiagaraRibbonComputeCommon::IndexGenThreadSize);
+
+		constexpr uint32 IndirectDispatchArgsOffset = 0;
+
+		FNiagaraRibbonGenerateIndices Params;
+		FMemory::Memzero(Params);
+		
+		Params.GeneratedIndicesBuffer = RenderingViewResources->IndexBuffer.UAV;
+		Params.SortedIndices = VertexBuffers.SortedIndicesBuffer.SRV;
+		Params.MultiRibbonIndices = VertexBuffers.MultiRibbonIndicesBuffer.SRV;
+		Params.Segments = VertexBuffers.SegmentsBuffer.SRV;
+
+		Params.IndirectDrawInfo = RenderingViewResources->IndirectDrawBuffer.SRV;
+		Params.TriangleToVertexIds = ShapeState.SliceTriangleToVertexIdsBuffer.SRV;
+
+		// Total particle Count
+		Params.TotalNumParticlesDirect = SourceParticleData->GetNumInstances();
+
+		// Indirect particle Count
+		Params.EmitterParticleCountsBuffer = GetSrvOrDefaultUInt(ComputeDispatchInterface->GetGPUInstanceCounterManager().GetInstanceCountBuffer());
+		Params.EmitterParticleCountsBufferOffset = SourceParticleData->GetGPUInstanceCountBufferOffset();
+
+		Params.IndexBufferOffset = 0;
+		Params.IndirectDrawInfoIndex = 0;
+		Params.TriangleToVertexIdsCount = ShapeState.SliceTriangleToVertexIds.Num();
+
+		Params.TrianglesPerSegment = ShapeState.TrianglesPerSegment;
+		Params.NumVerticesInSlice = ShapeState.NumVerticesInSlice;
+		Params.BitsNeededForShape = ShapeState.BitsNeededForShape;
+		Params.BitMaskForShape = ShapeState.BitMaskForShape;
+		Params.SegmentBitShift = RenderingViewResources->IndexGenerationSettings.SegmentBitShift;
+		Params.SegmentBitMask = RenderingViewResources->IndexGenerationSettings.SegmentBitMask;
+		Params.SubSegmentBitShift = RenderingViewResources->IndexGenerationSettings.SubSegmentBitShift;
+		Params.SubSegmentBitMask = RenderingViewResources->IndexGenerationSettings.SubSegmentBitMask;
+		
+		CMDList.Transition(FRHITransitionInfo(RenderingViewResources->IndexBuffer.UAV, ERHIAccess::VertexOrIndexBuffer, ERHIAccess::UAVCompute));
+		SetComputePipelineState(CMDList, ComputeShader.GetComputeShader());
+		SetShaderParameters(CMDList, ComputeShader, ComputeShader.GetComputeShader(), Params);
+		DispatchIndirectComputeShader(CMDList, ComputeShader.GetShader(), RenderingViewResources->IndirectDrawBuffer.Buffer, IndirectDispatchArgsOffset);
+		UnsetShaderUAVs(CMDList, ComputeShader, ComputeShader.GetComputeShader());
+		CMDList.Transition(FRHITransitionInfo(RenderingViewResources->IndexBuffer.UAV, ERHIAccess::UAVCompute, ERHIAccess::VertexOrIndexBuffer));
+	}
+}
 
-	auto& View = Context.ReferenceView;
-	auto& ViewFamily = Context.ReferenceViewFamily;
-	// Setup material for our ray tracing instance
-	FNiagaraMeshCollectorResourcesRibbon& CollectorResources = Context.RayTracingMeshResourceCollector.AllocateOneFrameResource<FNiagaraMeshCollectorResourcesRibbon>();
+void FNiagaraRendererRibbons::InitializeVertexBuffersResources(const FNiagaraDynamicDataRibbon* DynamicDataRibbon, FNiagaraDataBuffer* SourceParticleData,
+                                                               FGlobalDynamicReadBuffer& DynamicReadBuffer, const TSharedPtr<FNiagaraRibbonRenderingFrameResources>& RenderingResources, bool bShouldUseGPUInit) const
+{
 
-	if (!CollectorResources.VertexFactory.GetType()->SupportsRayTracingDynamicGeometry())
-	{
-		return;
+	// Make sure our ribbon data buffers are setup
+	VertexBuffers.InitializeOrUpdateBuffers(GenerationConfig, DynamicDataRibbon->GenerationOutput, SourceParticleData, DynamicDataRibbon->MaxAllocatedParticleCount, bShouldUseGPUInit);
+	
+	// Now we need to bind the source particle data, copying it to the gpu if necessary
+	if (DynamicDataRibbon->bIsGPUSystem)
+	{		
+		RenderingResources->ParticleFloatSRV = GetSrvOrDefaultFloat(SourceParticleData->GetGPUBufferFloat());
+		RenderingResources->ParticleHalfSRV = GetSrvOrDefaultHalf(SourceParticleData->GetGPUBufferHalf());
+		RenderingResources->ParticleIntSRV = GetSrvOrDefaultInt(SourceParticleData->GetGPUBufferInt());
+		
+		RenderingResources->ParticleFloatDataStride = SourceParticleData->GetFloatStride() / sizeof(float);
+		RenderingResources->ParticleHalfDataStride = SourceParticleData->GetHalfStride() / sizeof(FFloat16);
+		RenderingResources->ParticleIntDataStride = SourceParticleData->GetInt32Stride() / sizeof(int32);
+		
+		RenderingResources->RibbonIdParamOffset = RibbonIDParamDataSetOffset;
 	}
+	else 
+	{
+		TArray<uint32, TInlineAllocator<2>> IntParamsToCopy;
+		if (bShouldUseGPUInit && GenerationConfig.HasRibbonIDs())
+		{
+			RenderingResources->RibbonIdParamOffset = IntParamsToCopy.Add(RibbonIDParamDataSetOffset);
 
-	FGlobalDynamicIndexBuffer::FAllocationEx DynamicIndexAllocation;
-	CreatePerViewResources(Context.ReferenceView, Context.ReferenceViewFamily, SceneProxy, Context.RayTracingMeshResourceCollector, CollectorResources.UniformBuffer, DynamicIndexAllocation);
+			// Also add acquire index if we're running full sized ids.
+			if (GenerationConfig.HasFullRibbonIDs())
+			{
+				IntParamsToCopy.Add(RibbonIDParamDataSetOffset + 1);
+			}		
+		}
+		
+		RenderingResources->ParticleData = TransferDataToGPU(DynamicReadBuffer, RendererLayout, IntParamsToCopy, SourceParticleData);
+
+		RenderingResources->ParticleFloatSRV = GetSrvOrDefaultFloat(RenderingResources->ParticleData.FloatData);
+		RenderingResources->ParticleHalfSRV = GetSrvOrDefaultHalf(RenderingResources->ParticleData.HalfData);
+		RenderingResources->ParticleIntSRV = GetSrvOrDefaultInt(RenderingResources->ParticleData.IntData);
+		
+		RenderingResources->ParticleFloatDataStride = RenderingResources->ParticleData.FloatStride / sizeof(float);
+		RenderingResources->ParticleHalfDataStride = RenderingResources->ParticleData.HalfStride / sizeof(FFloat16);
+		RenderingResources->ParticleIntDataStride = RenderingResources->ParticleData.IntStride / sizeof(int32);
+	}
 
-	if (DynamicIndexAllocation.MaxUsedIndex <= 0)
+	
+	// If the data was generated sync it here, otherwise we rely on the generation step later to populate it
+	if (DynamicDataRibbon->GenerationOutput.IsValid() && DynamicDataRibbon->GenerationOutput->SegmentData.Num() > 0)
 	{
-		return;
+		const auto& GeneratedGeometryData = *DynamicDataRibbon->GenerationOutput;
+		
+		void *IndexPtr = RHILockBuffer(VertexBuffers.SortedIndicesBuffer.Buffer, 0, GeneratedGeometryData.SortedIndices.Num() * sizeof(int32), RLM_WriteOnly);
+		FMemory::Memcpy(IndexPtr, GeneratedGeometryData.SortedIndices.GetData(), GeneratedGeometryData.SortedIndices.Num() * sizeof(int32));
+		RHIUnlockBuffer(VertexBuffers.SortedIndicesBuffer.Buffer);
+
+		// pass in the CPU generated total segment distance (for tiling distance modes); needs to be a buffer so we can fetch them in the correct order based on Draw Direction (front->back or back->front)
+		//	otherwise UVs will pop when draw direction changes based on camera view point
+		void *TangentsAndDistancesPtr = RHILockBuffer(VertexBuffers.TangentsAndDistancesBuffer.Buffer, 0, GeneratedGeometryData.TangentAndDistances.Num() * sizeof(FVector4f), RLM_WriteOnly);
+		FMemory::Memcpy(TangentsAndDistancesPtr, GeneratedGeometryData.TangentAndDistances.GetData(), GeneratedGeometryData.TangentAndDistances.Num() * sizeof(FVector4f));
+		RHIUnlockBuffer(VertexBuffers.TangentsAndDistancesBuffer.Buffer);
+		
+		// Copy a buffer which has the per particle multi ribbon index.
+		if (GenerationConfig.HasRibbonIDs())
+		{
+			void* MultiRibbonIndexPtr = RHILockBuffer(VertexBuffers.MultiRibbonIndicesBuffer.Buffer, 0, GeneratedGeometryData.MultiRibbonIndices.Num() * sizeof(uint32), RLM_WriteOnly);
+			FMemory::Memcpy(MultiRibbonIndexPtr, GeneratedGeometryData.MultiRibbonIndices.GetData(), GeneratedGeometryData.MultiRibbonIndices.Num() * sizeof(uint32));
+			RHIUnlockBuffer(VertexBuffers.MultiRibbonIndicesBuffer.Buffer);
+		}
+		
+		// Copy the packed u data for stable age based uv generation.
+		TArray<uint32> PackedRibbonLookupTable;
+		PackedRibbonLookupTable.Reserve(GeneratedGeometryData.RibbonInfoLookup.Num() * FRibbonMultiRibbonInfoBufferEntry::NumElements);
+		for (int32 Index = 0; Index < GeneratedGeometryData.RibbonInfoLookup.Num(); Index++)
+		{
+			GeneratedGeometryData.RibbonInfoLookup[Index].PackElementsToLookupTableBuffer(PackedRibbonLookupTable);
+		}
+		
+		void *PackedPerRibbonDataByIndexPtr = RHILockBuffer(VertexBuffers.RibbonLookupTableBuffer.Buffer, 0, PackedRibbonLookupTable.Num() * sizeof(uint32), RLM_WriteOnly);
+		FMemory::Memcpy(PackedPerRibbonDataByIndexPtr, PackedRibbonLookupTable.GetData(), PackedRibbonLookupTable.Num() * sizeof(uint32));
+		RHIUnlockBuffer(VertexBuffers.RibbonLookupTableBuffer.Buffer);		
 	}
+}
 
-	FRayTracingInstance RayTracingInstance;
-	RayTracingInstance.Geometry = &RayTracingGeometry;
-	RayTracingInstance.InstanceTransforms.Add(FMatrix::Identity);
+FRibbonComputeUniformParameters FNiagaraRendererRibbons::SetupComputeVertexGenParams(FNiagaraGpuComputeDispatchInterface* ComputeDispatchInterface,
+	const TSharedPtr<FNiagaraRibbonRenderingFrameResources>& RenderingResources, const FNiagaraDataBuffer* SourceParticleData) const
+{
+	FRibbonComputeUniformParameters CommonParams;
+	FMemory::Memzero(CommonParams);
 
-	RayTracingGeometry.Initializer.IndexBuffer = DynamicIndexAllocation.IndexBuffer->IndexBufferRHI;
-	RayTracingGeometry.Initializer.IndexBufferOffset = DynamicIndexAllocation.FirstIndex * DynamicIndexAllocation.IndexStride;
+	// Total particle Count
+	CommonParams.TotalNumParticlesDirect = SourceParticleData->GetNumInstances();
 
-	FMeshBatch MeshBatch;
+	// Indirect particle Count
+	CommonParams.EmitterParticleCountsBuffer = GetSrvOrDefaultUInt(ComputeDispatchInterface->GetGPUInstanceCounterManager().GetInstanceCountBuffer());
+	CommonParams.EmitterParticleCountsBufferOffset = SourceParticleData->GetGPUInstanceCountBufferOffset();
 
-	SetupMeshBatchAndCollectorResourceForView(Context.ReferenceView, Context.ReferenceViewFamily, SceneProxy, Context.RayTracingMeshResourceCollector, DynamicDataRibbon, DynamicIndexAllocation, MeshBatch, CollectorResources);
+	// Niagara sim data
+	CommonParams.NiagaraParticleDataFloat = GetSrvOrDefaultFloat(RenderingResources->ParticleFloatSRV);
+	CommonParams.NiagaraParticleDataHalf = RenderingResources->ParticleHalfSRV? RenderingResources->ParticleHalfSRV : GetDummyHalfBuffer();
+	CommonParams.NiagaraParticleDataInt = GetSrvOrDefaultInt(RenderingResources->ParticleIntSRV);
+	CommonParams.NiagaraFloatDataStride = RenderingResources->ParticleFloatDataStride;
+	CommonParams.NiagaraIntDataStride = RenderingResources->ParticleIntDataStride;
 
-	RayTracingInstance.Materials.Add(MeshBatch);
+	
+	// Int bindings
+	CommonParams.RibbonIdDataOffset = RenderingResources->RibbonIdParamOffset;
+
+	// Float bindings
+	const TConstArrayView<FNiagaraRendererVariableInfo> VFVariables = RendererLayout->GetVFVariables_RenderThread();
+	CommonParams.PositionDataOffset = VFVariables[ENiagaraRibbonVFLayout::Position].GetGPUOffset();
+	CommonParams.PrevPositionDataOffset = VFVariables[ENiagaraRibbonVFLayout::PrevPosition].GetGPUOffset();
+	CommonParams.VelocityDataOffset = VFVariables[ENiagaraRibbonVFLayout::Velocity].GetGPUOffset();
+	CommonParams.ColorDataOffset = VFVariables[ENiagaraRibbonVFLayout::Color].GetGPUOffset();
+	CommonParams.WidthDataOffset = VFVariables[ENiagaraRibbonVFLayout::Width].GetGPUOffset();
+	CommonParams.PrevWidthDataOffset = VFVariables[ENiagaraRibbonVFLayout::PrevRibbonWidth].GetGPUOffset();
+	CommonParams.TwistDataOffset = VFVariables[ENiagaraRibbonVFLayout::Twist].GetGPUOffset();
+	CommonParams.PrevTwistDataOffset = VFVariables[ENiagaraRibbonVFLayout::PrevRibbonTwist].GetGPUOffset();
+	CommonParams.NormalizedAgeDataOffset = VFVariables[ENiagaraRibbonVFLayout::NormalizedAge].GetGPUOffset();
+	CommonParams.MaterialRandomDataOffset = VFVariables[ENiagaraRibbonVFLayout::MaterialRandom].GetGPUOffset();
+	CommonParams.MaterialParamDataOffset = VFVariables[ENiagaraRibbonVFLayout::MaterialParam0].GetGPUOffset();
+	CommonParams.MaterialParam1DataOffset = VFVariables[ENiagaraRibbonVFLayout::MaterialParam1].GetGPUOffset();
+	CommonParams.MaterialParam2DataOffset = VFVariables[ENiagaraRibbonVFLayout::MaterialParam2].GetGPUOffset();
+	CommonParams.MaterialParam3DataOffset = VFVariables[ENiagaraRibbonVFLayout::MaterialParam3].GetGPUOffset();
+
+	const bool bShouldLinkDistanceFromStart = (UV0Settings.DistributionMode == ENiagaraRibbonUVDistributionMode::TiledFromStartOverRibbonLength ||
+		UV1Settings.DistributionMode == ENiagaraRibbonUVDistributionMode::TiledFromStartOverRibbonLength);
+	
+	CommonParams.DistanceFromStartOffset = bShouldLinkDistanceFromStart ? VFVariables[ENiagaraRibbonVFLayout::DistanceFromStart].GetGPUOffset() : -1;
+	CommonParams.U0OverrideDataOffset = UV0Settings.bEnablePerParticleUOverride ? VFVariables[ENiagaraRibbonVFLayout::U0Override].GetGPUOffset() : -1;
+	CommonParams.V0RangeOverrideDataOffset = UV0Settings.bEnablePerParticleVRangeOverride ? VFVariables[ENiagaraRibbonVFLayout::V0RangeOverride].GetGPUOffset() : -1;
+	CommonParams.U1OverrideDataOffset = UV1Settings.bEnablePerParticleUOverride ? VFVariables[ENiagaraRibbonVFLayout::U1Override].GetGPUOffset() : -1;
+	CommonParams.V1RangeOverrideDataOffset = UV1Settings.bEnablePerParticleVRangeOverride ? VFVariables[ENiagaraRibbonVFLayout::V1RangeOverride].GetGPUOffset() : -1;
 
-	// Use the internal vertex buffer only when initialized otherwise used the shared vertex buffer - needs to be updated every frame
-	FRWBuffer* VertexBuffer = RayTracingDynamicVertexBuffer.NumBytes > 0 ? &RayTracingDynamicVertexBuffer : nullptr;
+	CommonParams.MaterialParamValidMask = GenerationConfig.GetMaterialParamValidMask();
 
-	// Vertex count is the maximum value in the index buffer + 1
-	const uint32 VertexCount = DynamicIndexAllocation.MaxUsedIndex + 1;
+	const bool bShouldDoFacing = FacingMode == ENiagaraRibbonFacingMode::Custom || FacingMode == ENiagaraRibbonFacingMode::CustomSideVector;
+	CommonParams.FacingDataOffset = bShouldDoFacing ? VFVariables[ENiagaraRibbonVFLayout::Facing].GetGPUOffset() : -1;
+	CommonParams.PrevFacingDataOffset = bShouldDoFacing ? VFVariables[ENiagaraRibbonVFLayout::PrevRibbonFacing].GetGPUOffset() : -1;
 
-	Context.DynamicRayTracingGeometriesToUpdate.Add(
-		FRayTracingDynamicGeometryUpdateParams
+	CommonParams.RibbonLinkOrderDataOffset = GenerationConfig.HasCustomLinkOrder()? VFVariables[ENiagaraRibbonVFLayout::LinkOrder].GetGPUOffset() : -1;
+
+	CommonParams.U0DistributionMode = static_cast<int32>(UV0Settings.DistributionMode);
+	CommonParams.U1DistributionMode = static_cast<int32>(UV1Settings.DistributionMode);
+
+	return CommonParams;
+}
+
+void FNiagaraRendererRibbons::InitializeVertexBuffersGPU(FRHICommandListImmediate& CMDList, FNiagaraGpuComputeDispatchInterface* ComputeDispatchInterface, const FNiagaraDataBuffer* SourceParticleData,
+	FNiagaraRibbonGPUInitComputeBuffers& TempBuffers, const TSharedPtr<FNiagaraRibbonRenderingFrameResources>& RenderingResources) const
+{	
+	SCOPE_CYCLE_COUNTER(STAT_NiagaraRenderRibbonsGenVerticesGPU);
+
+	FRibbonComputeUniformParameters CommonParams = SetupComputeVertexGenParams(ComputeDispatchInterface, RenderingResources, SourceParticleData);
+
+	const int32 NumExecutableInstances = CommonParams.EmitterParticleCountsBufferOffset != INDEX_NONE? SourceParticleData->GetNumInstancesAllocated() : SourceParticleData->GetNumInstances();
+
+	const bool bCanRun = NumExecutableInstances >= 2;
+	
+	// Clear the command buffer if we just initialized it, or if the sim doesn't have enough data to run
+	if ((!bCanRun || VertexBuffers.bJustCreatedCommandBuffer) && VertexBuffers.GPUComputeCommandLength > 0)
+	{
+		CMDList.Transition(FRHITransitionInfo(VertexBuffers.GPUComputeCommandBuffer.Buffer, ERHIAccess::SRVMask | ERHIAccess::VertexOrIndexBuffer | ERHIAccess::IndirectArgs, ERHIAccess::UAVCompute));
+		CMDList.ClearUAVUint(VertexBuffers.GPUComputeCommandBuffer.UAV, FUintVector4(0));
+		CMDList.Transition(FRHITransitionInfo(VertexBuffers.GPUComputeCommandBuffer.Buffer, ERHIAccess::UAVCompute, ERHIAccess::SRVMask | ERHIAccess::VertexOrIndexBuffer | ERHIAccess::IndirectArgs));
+		VertexBuffers.bJustCreatedCommandBuffer = false;
+	}
+	
+	// Not possible to have a valid ribbon with less than 2 particles, so the remaining work is unnecessary as there's nothing needed here
+	if (!bCanRun)
+	{
+		return;
+	}
+
+	{
+		SCOPE_CYCLE_COUNTER(STAT_NiagaraRenderRibbonsGenVerticesSortGPU);
+		
+		FNiagaraRibbonSortPhase1CS::FPermutationDomain PermutationVector;
+		PermutationVector.Set<FRibbonHasFullRibbonID>(GenerationConfig.HasFullRibbonIDs());
+		PermutationVector.Set<FRibbonHasRibbonID>(GenerationConfig.HasSimpleRibbonIDs());
+		PermutationVector.Set<FRibbonHasCustomLinkOrder>(GenerationConfig.HasCustomLinkOrder());
+		
+		TShaderMapRef<FNiagaraRibbonSortPhase1CS> BubbleSortShader(GetGlobalShaderMap(GMaxRHIFeatureLevel), PermutationVector);
+		TShaderMapRef<FNiagaraRibbonSortPhase2CS> MergeSortShader(GetGlobalShaderMap(GMaxRHIFeatureLevel), PermutationVector);
+			
+		FRibbonOrderSortParameters SortParams;
+		FMemory::Memzero(SortParams);
+		SortParams.Common = CommonParams;
+		SortParams.DestinationSortedIndices = VertexBuffers.SortedIndicesBuffer.UAV;
+		SortParams.SortedIndices = GetSrvOrDefaultUInt(TempBuffers.SortBuffer.SRV);
+		
+		int CurrentBufferOrientation = 0;		
+		const auto SwapBuffers = [&]()
 		{
-			RayTracingInstance.Materials,
-			false,
-			VertexCount,
-			VertexCount * (uint32)sizeof(FVector3f),
-			MeshBatch.Elements[0].NumPrimitives,
-			&RayTracingGeometry,
-			VertexBuffer,
-			true
+			CurrentBufferOrientation ^= 0x1;	
+			const bool bComputeOnOutputBuffer = CurrentBufferOrientation == 0;
+			
+			SortParams.DestinationSortedIndices = bComputeOnOutputBuffer? VertexBuffers.SortedIndicesBuffer.UAV : TempBuffers.SortBuffer.UAV;
+			SortParams.SortedIndices = bComputeOnOutputBuffer? TempBuffers.SortBuffer.SRV : VertexBuffers.SortedIndicesBuffer.SRV;			
+		};
+	
+		const uint32 NumInitialThreadGroups = FMath::DivideAndRoundUp<uint32>(NumExecutableInstances, FNiagaraRibbonSortPhase1CS::BubbleSortGroupWidth);	
+		const uint32 NumMergeSortThreadGroups = FMath::DivideAndRoundUp<uint32>(NumExecutableInstances, FNiagaraRibbonSortPhase2CS::ThreadGroupSize);
+		const uint32 MergeSortPasses = FMath::CeilLogTwo(NumInitialThreadGroups);
+		
+		// If should do an initial flip so we start with the temp buffer to end in the correct buffer
+		if (MergeSortPasses % 2 != 0)
+		{
+			SwapBuffers();
 		}
-	);
 
-	RayTracingInstance.BuildInstanceMaskAndFlags(FeatureLevel);
+		{			
+			SCOPE_CYCLE_COUNTER(STAT_NiagaraRenderRibbonsGenVerticesInitialSortGPU);
+			
+			// Initial sort, sets up the buffer, and runs a parallel bubble sort to create groups of BubbleSortGroupWidth size
+			CMDList.Transition(FRHITransitionInfo(SortParams.DestinationSortedIndices, ERHIAccess::SRVMask | ERHIAccess::VertexOrIndexBuffer, ERHIAccess::UAVCompute));
+			SetComputePipelineState(CMDList, BubbleSortShader.GetComputeShader());
+			ValidateShaderParameters(BubbleSortShader, SortParams);
+			SetShaderParameters(CMDList, BubbleSortShader, BubbleSortShader.GetComputeShader(), SortParams);
+			DispatchComputeShader(CMDList, BubbleSortShader, NumInitialThreadGroups, 1, 1);
+			UnsetShaderUAVs(CMDList, BubbleSortShader, BubbleSortShader.GetComputeShader());
+			CMDList.Transition(FRHITransitionInfo(SortParams.DestinationSortedIndices, ERHIAccess::UAVCompute, ERHIAccess::SRVMask | ERHIAccess::VertexOrIndexBuffer));
+		}
+		
+		{
+			SCOPE_CYCLE_COUNTER(STAT_NiagaraRenderRibbonsGenVerticesFinalSortGPU);
+			
+			// Repeatedly runs a scatter based merge sort until we have the final buffer
+			uint32 SortGroupSize = FNiagaraRibbonSortPhase1CS::BubbleSortGroupWidth;
+			for (uint32 Idx = 0; Idx < MergeSortPasses; Idx++)
+			{
+				SortParams.MergeSortSourceBlockSize = SortGroupSize;
+				SortParams.MergeSortDestinationBlockSize = SortGroupSize * 2;
+		
+				SwapBuffers();
+			
+				CMDList.Transition(FRHITransitionInfo(SortParams.DestinationSortedIndices, ERHIAccess::SRVMask | ERHIAccess::VertexOrIndexBuffer, ERHIAccess::UAVCompute));
+				SetComputePipelineState(CMDList, MergeSortShader.GetComputeShader());
+				SetShaderParameters(CMDList, MergeSortShader, MergeSortShader.GetComputeShader(), SortParams);
+				DispatchComputeShader(CMDList, MergeSortShader, NumMergeSortThreadGroups, 1, 1);
+				UnsetShaderUAVs(CMDList, MergeSortShader, MergeSortShader.GetComputeShader());
+				CMDList.Transition(FRHITransitionInfo(SortParams.DestinationSortedIndices, ERHIAccess::UAVCompute, ERHIAccess::SRVMask | ERHIAccess::VertexOrIndexBuffer));
+				
+				SortGroupSize *= 2;
+			}			
+		}		
+	}
+	
+	{
+		SCOPE_CYCLE_COUNTER(STAT_NiagaraRenderRibbonsGenVerticesReductionPhase1GPU);
+		
+		FNiagaraRibbonVertexReductionInitializationCS::FPermutationDomain PermutationVector;
+		PermutationVector.Set<FRibbonHasFullRibbonID>(GenerationConfig.HasFullRibbonIDs());
+		PermutationVector.Set<FRibbonHasRibbonID>(GenerationConfig.HasSimpleRibbonIDs());
+		PermutationVector.Set<FRibbonWantsAutomaticTessellation>(GenerationConfig.WantsAutomaticTessellation());
+		PermutationVector.Set<FRibbonWantsConstantTessellation>(GenerationConfig.WantsConstantTessellation());
+		
+		TShaderMapRef<FNiagaraRibbonVertexReductionInitializationCS> ReductionInitializationShader(GetGlobalShaderMap(GMaxRHIFeatureLevel), PermutationVector);
+		TShaderMapRef<FNiagaraRibbonVertexReductionPropagateCS> ReductionPropgateShader(GetGlobalShaderMap(GMaxRHIFeatureLevel), PermutationVector);	
+		
+		const int32 NumPrefixScanPasses = FMath::CeilLogTwo(NumExecutableInstances);
+	
+		FNiagaraRibbonVertexReductionParameters Params;
+		FMemory::Memzero(Params);
+		Params.Common = CommonParams;
+		Params.SortedIndices = VertexBuffers.SortedIndicesBuffer.SRV;
+		Params.CurveTension = GenerationConfig.GetCurveTension();
+	
+		uint32 CurrentBufferOrientation = 0x0;
+		const auto SwapBuffers = [&]()
+		{
+			CurrentBufferOrientation ^= 0x1;
+			const bool bComputeOnOutputBuffer = CurrentBufferOrientation == 0;
+	
+			if (bComputeOnOutputBuffer)
+			{
+				Params.InputTangentsAndDistances = TempBuffers.TempDistances.SRV;
+				Params.OutputTangentsAndDistances = VertexBuffers.TangentsAndDistancesBuffer.UAV;
+				Params.InputMultiRibbonIndices = TempBuffers.TempMultiRibbon.SRV;
+				Params.OutputMultiRibbonIndices = VertexBuffers.MultiRibbonIndicesBuffer.UAV;
+				Params.InputSegments = TempBuffers.TempSegments.SRV;
+				Params.OutputSegments = VertexBuffers.SegmentsBuffer.UAV;
+				Params.InputTessellationStats = TempBuffers.TempTessellationStats[1].SRV;
+				Params.OutputTessellationStats = TempBuffers.TempTessellationStats[0].UAV;
+			}
+			else
+			{
+				Params.InputTangentsAndDistances = VertexBuffers.TangentsAndDistancesBuffer.SRV;
+				Params.OutputTangentsAndDistances = TempBuffers.TempDistances.UAV;
+				Params.InputMultiRibbonIndices = VertexBuffers.MultiRibbonIndicesBuffer.SRV;
+				Params.OutputMultiRibbonIndices = TempBuffers.TempMultiRibbon.UAV;
+				Params.InputSegments = VertexBuffers.SegmentsBuffer.SRV;
+				Params.OutputSegments = TempBuffers.TempSegments.UAV;				
+				Params.InputTessellationStats = TempBuffers.TempTessellationStats[0].SRV;
+				Params.OutputTessellationStats = TempBuffers.TempTessellationStats[1].UAV;
+			}				
+		};	
+	
+		// Setup buffers
+		if (NumPrefixScanPasses % 2 == 0)
+		{
+			CurrentBufferOrientation = 0x1;
+			SwapBuffers();
+		}
+		else
+		{
+			CurrentBufferOrientation = 0x0;
+			SwapBuffers();
+		}		
+		
+		const auto TransitionOutputBuffers = [this, &CMDList, &Params, &TempBuffers](ERHIAccess Previous, ERHIAccess Next)
+		{
+			FRHITransitionInfo DataBufferTransitions[] =
+			{
+				FRHITransitionInfo(Params.OutputTangentsAndDistances, Previous, Next),
+				FRHITransitionInfo(Params.OutputMultiRibbonIndices, Previous, Next),
+				FRHITransitionInfo(Params.OutputSegments, Previous, Next),
+				FRHITransitionInfo(Params.OutputTessellationStats, Previous, Next),
+			};
+			CMDList.Transition(MakeArrayView(DataBufferTransitions, UE_ARRAY_COUNT(DataBufferTransitions)));			
+		};
+		
+	
+		{			
+			
+			const uint32 NumThreadGroupsInitialization = FMath::DivideAndRoundUp<uint32>(NumExecutableInstances, FNiagaraRibbonComputeCommon::VertexGenReductionInitializationThreadSize);	
+			TransitionOutputBuffers(ERHIAccess::SRVMask | ERHIAccess::VertexOrIndexBuffer, ERHIAccess::UAVCompute);
+			SetComputePipelineState(CMDList, ReductionInitializationShader.GetComputeShader());
+			SetShaderParameters(CMDList, ReductionInitializationShader, ReductionInitializationShader.GetComputeShader(), Params);
+			DispatchComputeShader(CMDList, ReductionInitializationShader, NumThreadGroupsInitialization, 1, 1);
+			UnsetShaderUAVs(CMDList, ReductionInitializationShader, ReductionInitializationShader.GetComputeShader());
+			TransitionOutputBuffers(ERHIAccess::UAVCompute, ERHIAccess::SRVMask | ERHIAccess::VertexOrIndexBuffer);
+		}
+		
+		{
+			SCOPE_CYCLE_COUNTER(STAT_NiagaraRenderRibbonsGenVerticesReductionPropagateGPU);
+			
+			const uint32 NumThreadGroups = FMath::DivideAndRoundUp<uint32>(NumExecutableInstances, FNiagaraRibbonComputeCommon::VertexGenReductionPropagationThreadSize);
+			
+			for (Params.PrefixScanStride = 1; Params.PrefixScanStride < NumExecutableInstances; Params.PrefixScanStride *= 2)
+			{
+				SwapBuffers();
+				
+				TransitionOutputBuffers(ERHIAccess::SRVMask | ERHIAccess::VertexOrIndexBuffer, ERHIAccess::UAVCompute);
+			
+				SetComputePipelineState(CMDList, ReductionPropgateShader.GetComputeShader());
+				SetShaderParameters(CMDList, ReductionPropgateShader, ReductionPropgateShader.GetComputeShader(), Params);
+				DispatchComputeShader(CMDList, ReductionPropgateShader, NumThreadGroups, 1, 1);
+				UnsetShaderUAVs(CMDList, ReductionPropgateShader, ReductionPropgateShader.GetComputeShader());
+				
+				TransitionOutputBuffers(ERHIAccess::UAVCompute, ERHIAccess::SRVMask | ERHIAccess::VertexOrIndexBuffer);
+			}
+		}
 
-	OutRayTracingInstances.Add(RayTracingInstance);
+		check(CurrentBufferOrientation == 0x0);		
+	}
+	
+	static constexpr int32 CommandBufferOffset = 0;
+		
+	{
+		SCOPE_CYCLE_COUNTER(STAT_NiagaraRenderRibbonsGenVerticesReductionPhase2GPU);
+		
+		FNiagaraRibbonVertexReductionFinalizationParameters FinalizationParams;
+		FMemory::Memzero(FinalizationParams);
+		FinalizationParams.Common = CommonParams;
+		FinalizationParams.SortedIndices = VertexBuffers.SortedIndicesBuffer.SRV;
+		FinalizationParams.TangentsAndDistances = VertexBuffers.TangentsAndDistancesBuffer.SRV;
+		FinalizationParams.MultiRibbonIndices = VertexBuffers.MultiRibbonIndicesBuffer.SRV;
+		FinalizationParams.Segments = VertexBuffers.SegmentsBuffer.SRV;
+		FinalizationParams.TessellationStats = TempBuffers.TempTessellationStats[0].SRV;
+		FinalizationParams.PackedPerRibbonData = VertexBuffers.RibbonLookupTableBuffer.UAV;
+		FinalizationParams.OutputCommandBuffer = VertexBuffers.GPUComputeCommandBuffer.UAV;
+		FinalizationParams.OutputCommandBufferIndex= CommandBufferOffset;
+		FinalizationParams.FinalizationThreadBlockSize = FNiagaraRibbonComputeCommon::VertexGenFinalizationThreadSize;
+			
+		{
+			SCOPE_CYCLE_COUNTER(STAT_NiagaraRenderRibbonsGenVerticesReductionFinalizeGPU);
+			
+			FNiagaraRibbonVertexReductionFinalizeCS::FPermutationDomain PermutationVector;
+			PermutationVector.Set<FRibbonHasFullRibbonID>(GenerationConfig.HasFullRibbonIDs());
+			PermutationVector.Set<FRibbonHasRibbonID>(GenerationConfig.HasSimpleRibbonIDs());
+			PermutationVector.Set<FRibbonWantsAutomaticTessellation>(GenerationConfig.WantsAutomaticTessellation());
+			PermutationVector.Set<FRibbonWantsConstantTessellation>(GenerationConfig.WantsConstantTessellation());
+			PermutationVector.Set<FRibbonHasTwist>(GenerationConfig.HasTwist());
+			
+			TShaderMapRef<FNiagaraRibbonVertexReductionFinalizeCS> ComputeShader(GetGlobalShaderMap(GMaxRHIFeatureLevel), PermutationVector);
+
+			// We only run a single threadgroup when we're not running multi-ribbon since we assume start/end is the first/last particle
+			const uint32 NumThreadGroups = GenerationConfig.HasRibbonIDs() ?
+				FMath::DivideAndRoundUp<uint32>(NumExecutableInstances, FNiagaraRibbonComputeCommon::VertexGenReductionFinalizationThreadSize) :
+				1;
+			CMDList.Transition({
+				FRHITransitionInfo(VertexBuffers.RibbonLookupTableBuffer.Buffer, ERHIAccess::SRVMask | ERHIAccess::VertexOrIndexBuffer, ERHIAccess::UAVCompute),
+				FRHITransitionInfo(VertexBuffers.GPUComputeCommandBuffer.Buffer, ERHIAccess::SRVMask | ERHIAccess::VertexOrIndexBuffer | ERHIAccess::IndirectArgs, ERHIAccess::UAVCompute)});
+			
+			SetComputePipelineState(CMDList, ComputeShader.GetComputeShader());
+			SetShaderParameters(CMDList, ComputeShader, ComputeShader.GetComputeShader(), FinalizationParams);
+			DispatchComputeShader(CMDList, ComputeShader, NumThreadGroups, 1, 1);
+			UnsetShaderUAVs(CMDList, ComputeShader, ComputeShader.GetComputeShader());
+			
+			// We don't need to transition RibbonLookupTableBuffer as it's still needed for the next shader
+			CMDList.Transition({
+				FRHITransitionInfo(VertexBuffers.RibbonLookupTableBuffer.Buffer, ERHIAccess::UAVCompute, ERHIAccess::UAVCompute),
+				FRHITransitionInfo(VertexBuffers.GPUComputeCommandBuffer.Buffer, ERHIAccess::UAVCompute, ERHIAccess::SRVMask | ERHIAccess::VertexOrIndexBuffer | ERHIAccess::IndirectArgs)});		
+		}		
+	}
+		
+	{
+		SCOPE_CYCLE_COUNTER(STAT_NiagaraRenderRibbonsGenVerticesMultiRibbonInitGPU);
+		
+		FNiagaraRibbonVertexFinalizationParameters FinalizeParams;
+		FMemory::Memzero(FinalizeParams);
+		FinalizeParams.Common = CommonParams;
+		FinalizeParams.SortedIndices = VertexBuffers.SortedIndicesBuffer.SRV;
+		FinalizeParams.TangentsAndDistances = VertexBuffers.TangentsAndDistancesBuffer.UAV;
+		FinalizeParams.PackedPerRibbonData = VertexBuffers.RibbonLookupTableBuffer.UAV;
+		FinalizeParams.CommandBuffer = VertexBuffers.GPUComputeCommandBuffer.SRV;
+		FinalizeParams.CommandBufferOffset = CommandBufferOffset;
+		
+		constexpr auto AddUVChannelParams = [](const FNiagaraRibbonUVSettings& Input, FNiagaraRibbonUVSettingsParams& Output)
+		{
+			Output.Offset = FVector2f(Input.Offset);
+			Output.Scale = FVector2f(Input.Scale);
+			Output.TilingLength = Input.TilingLength;
+			Output.DistributionMode = static_cast<int32>(Input.DistributionMode);
+			Output.LeadingEdgeMode = static_cast<int32>(Input.LeadingEdgeMode);
+			Output.TrailingEdgeMode = static_cast<int32>(Input.TrailingEdgeMode);
+			Output.bEnablePerParticleUOverride = Input.bEnablePerParticleUOverride ? 1 : 0;
+			Output.bEnablePerParticleVRangeOverride = Input.bEnablePerParticleVRangeOverride ? 1 : 0;			
+		};
+	
+		AddUVChannelParams(UV0Settings, FinalizeParams.UV0Settings);
+		AddUVChannelParams(UV1Settings, FinalizeParams.UV1Settings);
+	
+		{
+			SCOPE_CYCLE_COUNTER(STAT_NiagaraRenderRibbonsGenVerticesMultiRibbonInitComputeGPU);
+			
+			FNiagaraRibbonUVParamCalculationCS::FPermutationDomain PermutationVector;
+			PermutationVector.Set<FRibbonHasFullRibbonID>(GenerationConfig.HasFullRibbonIDs());
+			PermutationVector.Set<FRibbonHasRibbonID>(GenerationConfig.HasSimpleRibbonIDs());
+			PermutationVector.Set<FRibbonWantsAutomaticTessellation>(GenerationConfig.WantsAutomaticTessellation());
+			PermutationVector.Set<FRibbonWantsConstantTessellation>(GenerationConfig.WantsConstantTessellation());
+			PermutationVector.Set<FRibbonHasTwist>(GenerationConfig.HasTwist());
+
+			TShaderMapRef<FNiagaraRibbonUVParamCalculationCS> ComputeShader(GetGlobalShaderMap(GMaxRHIFeatureLevel), PermutationVector);
+
+			// We don't need to transition RibbonLookupTableBuffer as it's still setup for UAV from the last shader
+			CMDList.Transition({
+				FRHITransitionInfo(VertexBuffers.RibbonLookupTableBuffer.Buffer, ERHIAccess::UAVCompute, ERHIAccess::UAVCompute),
+				FRHITransitionInfo(VertexBuffers.TangentsAndDistancesBuffer.Buffer, ERHIAccess::SRVMask | ERHIAccess::VertexOrIndexBuffer, ERHIAccess::UAVCompute)});
+			
+			SetComputePipelineState(CMDList, ComputeShader.GetComputeShader());			
+			SetShaderParameters(CMDList, ComputeShader, ComputeShader.GetComputeShader(), FinalizeParams);
+			DispatchIndirectComputeShader(CMDList, ComputeShader.GetShader(), VertexBuffers.GPUComputeCommandBuffer.Buffer, CommandBufferOffset * FNiagaraRibbonCommandBufferLayout::NumElements);
+			UnsetShaderUAVs(CMDList, ComputeShader, ComputeShader.GetComputeShader());
+			
+			CMDList.Transition({
+				FRHITransitionInfo(VertexBuffers.TangentsAndDistancesBuffer.Buffer, ERHIAccess::UAVCompute, ERHIAccess::SRVMask | ERHIAccess::VertexOrIndexBuffer),
+				FRHITransitionInfo(VertexBuffers.RibbonLookupTableBuffer.Buffer, ERHIAccess::UAVCompute, ERHIAccess::SRVMask | ERHIAccess::VertexOrIndexBuffer)});
+		}		
+	}
 }
-#endif
+
+void FNiagaraRibbonComputeDispatchManager::GenerateAllGPUData(FRHICommandListImmediate& CMDList)
+{
+	// Handle all vertex gens first
+	for (int32 Index = 0; Index < RenderersToGenerate.Num(); Index++)
+	{
+		const auto& Params = RenderersToGenerate[Index];
+
+		const auto RenderingResources = Params.RenderingResources.Pin();
+		if (RenderingResources.IsValid())
+		{
+			const bool bIsGPUSim = Params.SourceParticleData->GetGPUInstanceCountBufferOffset() != INDEX_NONE;
+
+			ComputeBuffers.InitOrUpdateBuffers(bIsGPUSim? Params.SourceParticleData->GetNumInstancesAllocated() : Params.SourceParticleData->GetNumInstances(),
+				Params.Renderer->GenerationConfig.HasRibbonIDs(),
+				Params.Renderer->GenerationConfig.WantsAutomaticTessellation(),
+				Params.Renderer->GenerationConfig.HasTwist());
+		
+		
+			Params.Renderer->InitializeVertexBuffersGPU(CMDList, Params.ComputeDispatchInterface, Params.SourceParticleData, ComputeBuffers, RenderingResources);
+		}
+	}
+		
+	// Now handle all index gens
+	for (const auto& RendererToGen : RenderersToGenerate)
+	{
+		const auto RenderingResources = RendererToGen.RenderingResources.Pin();
+		if (RenderingResources.IsValid())
+		{
+			for (int32 Index = 0; Index < RenderingResources->ViewResources.Num(); Index++)
+			{
+				const auto& RenderingResourcesView = RenderingResources->ViewResources[Index];
+				RendererToGen.Renderer->InitializeViewIndexBuffersGPU(CMDList, RendererToGen.ComputeDispatchInterface, RendererToGen.SourceParticleData, RenderingResourcesView);
+			}
+		}
+	}
+	
+	RenderersToGenerate.Empty();
+}
\ No newline at end of file
