GPU Ribbons!

diff --git a/Engine/Plugins/FX/Niagara/Source/Niagara/Private/NiagaraRibbonRendererProperties.cpp b/Engine/Plugins/FX/Niagara/Source/Niagara/Private/NiagaraRibbonRendererProperties.cpp
--- a/Engine/Plugins/FX/Niagara/Source/Niagara/Private/NiagaraRibbonRendererProperties.cpp
+++ b/Engine/Plugins/FX/Niagara/Source/Niagara/Private/NiagaraRibbonRendererProperties.cpp
@@ -51,6 +51,8 @@ UNiagaraRibbonRendererProperties::UNiagaraRibbonRendererProperties()
 	, UV1Scale_DEPRECATED(FVector2D(1.0f, 1.0f))
 	, UV1AgeOffsetMode_DEPRECATED(ENiagaraRibbonAgeOffsetMode::Scale)
 #endif
+	, MaxNumRibbons(0)
+	, bUseGPUInit(false)
 	, Shape(ENiagaraRibbonShapeMode::Plane)
 	, bEnableAccurateGeometry(false)
 	, WidthSegmentationCount(1)
@@ -73,6 +75,7 @@ UNiagaraRibbonRendererProperties::UNiagaraRibbonRendererProperties()
 	AttributeBindings.Add(&RibbonFacingBinding);
 	AttributeBindings.Add(&RibbonIdBinding);
 	AttributeBindings.Add(&RibbonLinkOrderBinding);
+	
 	AttributeBindings.Add(&MaterialRandomBinding);
 	AttributeBindings.Add(&DynamicMaterialBinding);
 	AttributeBindings.Add(&DynamicMaterial1Binding);
@@ -246,8 +249,8 @@ void UNiagaraRibbonRendererProperties::InitBindings()
 		RibbonTwistBinding = FNiagaraConstants::GetAttributeDefaultBinding(SYS_PARAM_PARTICLES_RIBBONTWIST);
 		RibbonWidthBinding = FNiagaraConstants::GetAttributeDefaultBinding(SYS_PARAM_PARTICLES_RIBBONWIDTH);
 		RibbonFacingBinding = FNiagaraConstants::GetAttributeDefaultBinding(SYS_PARAM_PARTICLES_RIBBONFACING);
-		RibbonIdBinding = FNiagaraConstants::GetAttributeDefaultBinding(SYS_PARAM_PARTICLES_RIBBONID);
-		RibbonLinkOrderBinding = FNiagaraConstants::GetAttributeDefaultBinding(SYS_PARAM_PARTICLES_RIBBONLINKORDER);
+		RibbonIdBinding = FNiagaraConstants::GetAttributeDefaultBinding(SYS_PARAM_PARTICLES_RIBBONID);		
+		RibbonLinkOrderBinding = FNiagaraConstants::GetAttributeDefaultBinding(SYS_PARAM_PARTICLES_RIBBONLINKORDER);		
 		MaterialRandomBinding = FNiagaraConstants::GetAttributeDefaultBinding(SYS_PARAM_PARTICLES_MATERIAL_RANDOM);
 		RibbonUVDistance = FNiagaraConstants::GetAttributeDefaultBinding(RIBBONUVDISTANCE);
 		U0OverrideBinding = FNiagaraConstants::GetAttributeDefaultBinding(SYS_PARAM_PARTICLES_RIBBONU0OVERRIDE);
@@ -310,6 +313,8 @@ void UNiagaraRibbonRendererProperties::CacheFromCompiledData(const FNiagaraDataS
 		RibbonIdDataSetAccessor.Init(CompiledData, RibbonIdBinding.GetDataSetBindableVariable().GetName());
 	}
 
+	RibbonLinkOrderDataSetAccessor.Init(CompiledData, RibbonLinkOrderBinding.GetDataSetBindableVariable().GetName());
+
 	const bool bShouldDoFacing = FacingMode == ENiagaraRibbonFacingMode::Custom || FacingMode == ENiagaraRibbonFacingMode::CustomSideVector;
 
 	// Initialize layout
@@ -342,6 +347,9 @@ void UNiagaraRibbonRendererProperties::CacheFromCompiledData(const FNiagaraDataS
 		RendererLayout.SetVariableFromBinding(CompiledData, PrevRibbonFacingBinding, ENiagaraRibbonVFLayout::PrevRibbonFacing);
 		RendererLayout.SetVariableFromBinding(CompiledData, PrevRibbonTwistBinding, ENiagaraRibbonVFLayout::PrevRibbonTwist);
 	}
+
+	RendererLayout.SetVariableFromBinding(CompiledData, RibbonLinkOrderBinding, ENiagaraRibbonVFLayout::LinkOrder);
+	
 	RendererLayout.Finalize();
 }
 
@@ -467,9 +475,42 @@ void UNiagaraRibbonRendererProperties::GetRendererTooltipWidgets(const FNiagaraE
 }
 
 
-void UNiagaraRibbonRendererProperties::GetRendererFeedback(const FVersionedNiagaraEmitter& InEmitter, TArray<FText>& OutErrors, TArray<FText>& OutWarnings, TArray<FText>& OutInfo) const
+void UNiagaraRibbonRendererProperties::GetRendererFeedback(const FVersionedNiagaraEmitter& InEmitter, TArray<FNiagaraRendererFeedback>& OutErrors, TArray<FNiagaraRendererFeedback>& OutWarnings, TArray<FNiagaraRendererFeedback>& OutInfo) const
 {
 	Super::GetRendererFeedback(InEmitter, OutErrors, OutWarnings, OutInfo);
+
+	const FVersionedNiagaraEmitterData* EmitterData = InEmitter.GetEmitterData();
+
+
+	// If we're in a gpu sim, then uv mode uniform by segment can cause some visual oddity due to non-existent
+	// culling of near particles like the cpu initialization pipeline runs
+	if (EmitterData && EmitterData->SimTarget == ENiagaraSimTarget::GPUComputeSim)
+	{
+		const auto CheckUVSettingsForChannel = [&](const FNiagaraRibbonUVSettings& UVSettings, int32 Index)
+		{
+			if (UVSettings.DistributionMode == ENiagaraRibbonUVDistributionMode::ScaledUniformly)
+			{
+				const FText ErrorDescription = FText::Format(LOCTEXT("NiagaraRibbonRendererUVBySegmentGPUDesc", "The specified UV Distribution for Channel {0} on GPU may result in different visual look than a CPU sim due to increased particle density in GPU sim."), FText::AsNumber(Index));
+				const FText ErrorSummary = FText::Format(LOCTEXT("NiagaraRibbonRendererUVBySegmentGPUSummary", "The specified UV Settings on Channel {0} on GPU may result in undesirable look."), FText::AsNumber(Index));
+				OutWarnings.Add(FNiagaraRendererFeedback(ErrorDescription, ErrorSummary, FText(), FNiagaraRendererFeedbackFix(), true));				
+			}
+		};
+
+		CheckUVSettingsForChannel(UV0Settings, 0);
+		CheckUVSettingsForChannel(UV1Settings, 1);
+	}
+
+
+	// If we're in multiplane shape, and multiplane count is even while we're in camera facing mode then one
+	// slice out of the set will be invisible because the camera will be coplanar to it
+	if (FacingMode == ENiagaraRibbonFacingMode::Screen && Shape == ENiagaraRibbonShapeMode::MultiPlane && MultiPlaneCount % 2 == 0)
+	{
+		const FText ErrorDescription = LOCTEXT("NiagaraRibbonRendererMultiPlaneInvisibleFaceDesc", "The specified MultiPlaneCount (Even Count) with ScreenFacing will result in a hidden face due to the camera being coplanar to one face.");
+		const FText ErrorSummary = LOCTEXT("NiagaraRibbonRendererMultiPlaneInvisibleFaceSummary", "The specified MultiPlaneCount+ScreenFacing will result in a hidden face.");
+		const FText ErrorFix = LOCTEXT("NiagaraRibbonRendererMultiPlaneInvisibleFaceFix", "Fix by decreasing MultiPlane count by 1.");
+		const FNiagaraRendererFeedbackFix MultiPlaneFix = FNiagaraRendererFeedbackFix::CreateLambda([this]() { const_cast<UNiagaraRibbonRendererProperties*>(this)->MultiPlaneCount = FMath::Clamp(this->MultiPlaneCount - 1, 1, 16); });
+		OutWarnings.Add(FNiagaraRendererFeedback(ErrorDescription, ErrorSummary, ErrorFix, MultiPlaneFix, true));	
+	}	
 }
 
 
diff --git a/Engine/Plugins/FX/Niagara/Source/Niagara/Public/NiagaraRendererRibbons.h b/Engine/Plugins/FX/Niagara/Source/Niagara/Public/NiagaraRendererRibbons.h
--- a/Engine/Plugins/FX/Niagara/Source/Niagara/Public/NiagaraRendererRibbons.h
+++ b/Engine/Plugins/FX/Niagara/Source/Niagara/Public/NiagaraRendererRibbons.h
@@ -5,12 +5,338 @@ NiagaraRenderer.h: Base class for Niagara render modules
 ==============================================================================*/
 #pragma once
 
+#include "NiagaraGpuComputeDispatchInterface.h"
 #include "NiagaraRibbonVertexFactory.h"
+#include "NiagaraRibbonCompute.h"
 #include "NiagaraRenderer.h"
 #include "NiagaraRibbonRendererProperties.h"
+#include "Chaos/AABB.h"
+#include "Chaos/AABB.h"
+#include "Chaos/AABB.h"
+#include "Chaos/AABB.h"
+#include "Chaos/AABB.h"
+#include "Chaos/AABB.h"
 
 class FNiagaraDataSet;
 
+struct FNiagaraDynamicDataRibbon;
+struct FNiagaraRibbonRenderingFrameResources;
+struct FNiagaraRibbonRenderingFrameViewResources;
+
+
+struct FNiagaraGenerationInputDataCPUAccessors
+{
+	FNiagaraGenerationInputDataCPUAccessors(const UNiagaraRibbonRendererProperties* Properties, const FNiagaraDataSet& Data)
+		: TotalNumParticles(Data.GetCurrentDataChecked().GetNumInstances())
+		, SortKeyReader(Properties->SortKeyDataSetAccessor.GetReader(Data))
+		, RibbonLinkOrderData(Properties->RibbonLinkOrderDataSetAccessor.GetReader(Data))
+		, SimpleRibbonIDData(Properties->RibbonIdDataSetAccessor.GetReader(Data))
+		, FullRibbonIDData(Properties->RibbonFullIDDataSetAccessor.GetReader(Data))
+		, PosData(Properties->PositionDataSetAccessor.GetReader(Data))
+		, AgeData(Properties->NormalizedAgeAccessor.GetReader(Data))
+		, SizeData(Properties->SizeDataSetAccessor.GetReader(Data))
+		, TwistData(Properties->TwistDataSetAccessor.GetReader(Data))
+	{
+		
+	}
+	
+	const uint32 TotalNumParticles;
+	
+	const FNiagaraDataSetReaderFloat<float> SortKeyReader;
+	const FNiagaraDataSetReaderFloat<float> RibbonLinkOrderData;
+	
+	const FNiagaraDataSetReaderInt32<int> SimpleRibbonIDData;
+	const FNiagaraDataSetReaderStruct<FNiagaraID> FullRibbonIDData;
+	
+	const FNiagaraDataSetReaderFloat<FNiagaraPosition> PosData;
+	const FNiagaraDataSetReaderFloat<float> AgeData;
+	const FNiagaraDataSetReaderFloat<float> SizeData;
+	const FNiagaraDataSetReaderFloat<float> TwistData;
+};
+
+struct FNiagaraIndexGenerationInput
+{
+	float ViewDistance;
+	int32 LODDistanceFactor;
+	
+	uint32 MaxSegmentCount;
+	uint32 SubSegmentCount;	
+
+	uint32 SegmentBitShift;
+	uint32 SegmentBitMask;
+
+	uint32 SubSegmentBitShift;
+	uint32 SubSegmentBitMask;
+
+	uint32 ShapeBitMask;
+	
+	uint32 TotalBitCount;
+	uint32 TotalNumIndices;
+};
+
+struct FNiagaraRibbonGenerationConfig
+{
+public:
+	FNiagaraRibbonGenerationConfig(const UNiagaraRibbonRendererProperties* Properties)
+		: MaterialParamValidMask(Properties->MaterialParamValidMask)
+		, CurveTension(Properties->CurveTension)
+		, MaxNumRibbons(Properties->MaxNumRibbons)
+		, bHasFullRibbonIDs(Properties->RibbonFullIDDataSetAccessor.IsValid())
+		, bHasSimpleRibbonIDs(Properties->RibbonIdDataSetAccessor.IsValid())
+		, bHasCustomLinkOrder(Properties->RibbonLinkOrderDataSetAccessor.IsValid())
+		, bHasTwist(Properties->TwistDataSetAccessor.IsValid())
+		, bHasCustomU0Data(Properties->UV0Settings.bEnablePerParticleUOverride && Properties->U0OverrideIsBound)
+		, bHasCustomU1Data(Properties->UV1Settings.bEnablePerParticleUOverride && Properties->U1OverrideIsBound)
+		, bWantsConstantTessellation(Properties->TessellationMode == ENiagaraRibbonTessellationMode::Custom && Properties->bUseConstantFactor)
+		, bWantsAutomaticTessellation(Properties->TessellationMode != ENiagaraRibbonTessellationMode::Disabled && !bWantsConstantTessellation)
+		, bNeedsPreciseMotionVectors(Properties->NeedsPreciseMotionVectors())
+	{
+		
+	}
+	
+	uint32 GetMaterialParamValidMask() const { return MaterialParamValidMask; }
+	float GetCurveTension() const { return FMath::Min(CurveTension, 0.99f); }
+	int32 GetMaxNumRibbons() const { return MaxNumRibbons; }
+	
+	bool HasFullRibbonIDs() const { return bHasFullRibbonIDs; }
+	bool HasSimpleRibbonIDs() const { return bHasSimpleRibbonIDs; }
+	bool HasRibbonIDs() const { return HasFullRibbonIDs() || HasSimpleRibbonIDs(); };
+	
+	bool HasCustomLinkOrder() const { return bHasCustomLinkOrder; }
+	bool HasTwist() const { return bHasTwist; }
+	
+	bool HasCustomU0Data() const { return bHasCustomU0Data; }
+	bool HasCustomU1Data() const { return bHasCustomU1Data; }
+	
+	bool WantsConstantTessellation () const { return bWantsConstantTessellation; }
+	bool WantsAutomaticTessellation() const { return bWantsAutomaticTessellation; }
+	
+	bool NeedsPreciseMotionVectors() const { return bNeedsPreciseMotionVectors; }
+
+private:
+	const uint32 MaterialParamValidMask;
+	const float CurveTension;
+	const int32 MaxNumRibbons;
+	
+	const uint32 bHasFullRibbonIDs : 1;
+	const uint32 bHasSimpleRibbonIDs : 1;
+	const uint32 bHasCustomLinkOrder : 1;
+	const uint32 bHasTwist : 1;
+	
+	const uint32 bHasCustomU0Data : 1;
+	const uint32 bHasCustomU1Data : 1;
+	const uint32 bWantsConstantTessellation : 1;
+	const uint32 bWantsAutomaticTessellation : 1;	
+	const uint32 bNeedsPreciseMotionVectors : 1;
+};
+
+struct FRibbonMultiRibbonInfoBufferEntry
+{
+	static constexpr int32 NumElements = 8;
+	
+	float U0Scale = 1.0;
+	float U0Offset = 0.0;
+	float U0DistributionScaler = 1.0;
+	float U1Scale = 1.0;
+	float U1Offset = 0.0;
+	float U1DistributionScaler = 1.0;
+	int32 FirstParticleId = INDEX_NONE;
+	int32 LastParticleId = INDEX_NONE;
+};
+
+struct FRibbonMultiRibbonInfo
+{	
+	/** start and end world space position of the ribbon, to figure out draw direction */
+	FVector StartPos;
+	FVector EndPos;
+	int32 BaseSegmentDataIndex = 0;
+	int32 NumSegmentDataIndices = 0;
+
+	FRibbonMultiRibbonInfoBufferEntry BufferEntry;
+	
+
+	FORCEINLINE bool UseInvertOrder(const FVector& ViewDirection, const FVector& ViewOriginForDistanceCulling, ENiagaraRibbonDrawDirection DrawDirection) const
+	{
+		const float StartDist = FVector::DotProduct(ViewDirection, StartPos - ViewOriginForDistanceCulling);
+		const float EndDist = FVector::DotProduct(ViewDirection, EndPos - ViewOriginForDistanceCulling);
+		return ((StartDist >= EndDist) && DrawDirection == ENiagaraRibbonDrawDirection::BackToFront)
+			|| ((StartDist < EndDist) && DrawDirection == ENiagaraRibbonDrawDirection::FrontToBack);
+	}
+
+	void PackElementsToLookupTableBuffer(TArray<uint32>& OutputData) const
+	{
+		const int32 Index = OutputData.Num();
+		OutputData.AddUninitialized(FRibbonMultiRibbonInfoBufferEntry::NumElements);		
+		*reinterpret_cast<FRibbonMultiRibbonInfoBufferEntry*>(&OutputData[Index]) = BufferEntry;
+	}
+};
+
+
+struct FNiagaraRibbonCPUGeneratedVertexData
+{	
+	// The list of all segments, each one connecting SortedIndices[SegmentId] to SortedIndices[SegmentId + 1].
+	// We use this format because the final index buffer gets generated based on view sorting and InterpCount.
+	TArray<uint32> SegmentData;
+
+	/** The list of all particle (instance) indices. Converts raw indices to particles indices. Ordered along each ribbons, from head to tail. */
+	TArray<uint32> SortedIndices;
+	
+	/** The tangent and distance between segments, for each raw index (raw VS particle indices). */
+	TArray<FVector4f> TangentAndDistances;
+	
+	/** The multi ribbon index, for each raw index. (raw VS particle indices). */
+	TArray<uint32> MultiRibbonIndices;
+	
+	/** Ribbon perperties required for sorting. */
+	TArray<FRibbonMultiRibbonInfo> RibbonInfoLookup;
+
+	double TotalSegmentLength;
+	double AverageSegmentLength;
+	double AverageSegmentAngle;
+	double AverageTwistAngle;
+	double AverageWidth;
+
+	FNiagaraRibbonCPUGeneratedVertexData()
+		: TotalSegmentLength(0.0)
+		, AverageSegmentLength(0.0)
+		, AverageSegmentAngle(0.0)
+		, AverageTwistAngle(0.0)
+		, AverageWidth(0.0)
+	{ }
+
+	int32 GetAllocatedSize()const
+	{
+		int32 Size = 0;
+		Size += SegmentData.GetAllocatedSize();
+		Size += SortedIndices.GetAllocatedSize();
+		Size += TangentAndDistances.GetAllocatedSize();
+		Size += MultiRibbonIndices.GetAllocatedSize();
+		Size += RibbonInfoLookup.GetAllocatedSize();
+
+		return Size;
+	}
+};
+
+struct FNiagaraRibbonTessellationConfig
+{
+	ENiagaraRibbonTessellationMode TessellationMode;
+	int32 CustomTessellationFactor;
+	bool bCustomUseConstantFactor;
+	float CustomTessellationMinAngle;
+	bool bCustomUseScreenSpace;
+};
+
+struct FNiagaraRibbonTessellationSmoothingData
+{
+	// Average curvature of the segments.
+	float TessellationAngle = 0;
+	// Average curvature of the segments (computed from the segment angle in radian).
+	float TessellationCurvature = 0;
+	// Average twist of the segments.
+	float TessellationTwistAngle = 0;
+	// Average twist curvature of the segments.
+	float TessellationTwistCurvature = 0;
+	// Average twist curvature of the segments.
+	float TessellationTotalSegmentLength = 0;
+};
+
+struct FNiagaraRibbonVertexBuffers
+{
+	FRWBuffer SortedIndicesBuffer;
+	FRWBuffer TangentsAndDistancesBuffer;
+	FRWBuffer MultiRibbonIndicesBuffer;
+	FRWBuffer RibbonLookupTableBuffer;
+	FRWBuffer SegmentsBuffer;
+
+	FRWBuffer GPUComputeCommandBuffer;
+
+	int32 SortedIndicesLength;
+	int32 TangentsLength;
+	int32 MultiRibbonIndexLength;
+	int32 RibbonLookupTableLength;
+	int32 SegmentLength;
+	int32 GPUComputeCommandLength;
+	bool bJustCreatedCommandBuffer;
+
+	FNiagaraRibbonVertexBuffers()
+		: SortedIndicesLength(0)
+		, TangentsLength(0)
+		, MultiRibbonIndexLength(0)
+		, RibbonLookupTableLength(0)
+		, SegmentLength(0)
+		, GPUComputeCommandLength(0)
+		, bJustCreatedCommandBuffer(false)
+	{
+		
+	}
+
+	static bool InitOrUpdateBuffer(bool bEnabled, FRWBuffer& Buffer, int32& CurrentLength, int32 NeededLength, int32 MaxLength, FRWBuffer (*InitFunction)(int32, ERHIAccess), ERHIAccess InitialAccessFlags = ERHIAccess::None);
+
+	void InitializeOrUpdateBuffers(const FNiagaraRibbonGenerationConfig& GenerationConfig, const TSharedPtr<FNiagaraRibbonCPUGeneratedVertexData>& GeneratedGeometryData,
+		const FNiagaraDataBuffer* SourceParticleData, int32 MaxAllocatedCount, bool bIsUsingGPUInit);
+
+	void Release()
+	{
+		SortedIndicesBuffer.Release();
+		TangentsAndDistancesBuffer.Release();
+		MultiRibbonIndicesBuffer.Release();
+		RibbonLookupTableBuffer.Release();
+		SegmentsBuffer.Release();
+		GPUComputeCommandBuffer.Release();
+		
+		SortedIndicesLength = 0;
+		TangentsLength = 0;
+		MultiRibbonIndexLength = 0;
+		RibbonLookupTableLength = 0;
+		SegmentLength = 0;
+		GPUComputeCommandLength = 0;
+	}
+
+	static FRWBuffer CreateSortedIndicesBuffer(int32 Size, ERHIAccess InitialAccessFlags = ERHIAccess::None);
+	static FRWBuffer CreateTangentsAndDistancesBuffer(int32 Size, ERHIAccess InitialAccessFlags = ERHIAccess::None);
+	static FRWBuffer CreateMultiRibbonIndicesBuffer(int32 Size, ERHIAccess InitialAccessFlags = ERHIAccess::None);
+	static FRWBuffer CreateRibbonLookupTableBuffer(int32 Size, ERHIAccess InitialAccessFlags = ERHIAccess::None);
+	static FRWBuffer CreateSegmentsBuffer(int32 Size, ERHIAccess InitialAccessFlags = ERHIAccess::None);
+	static FRWBuffer CreateCommandBuffer(int32 Size, ERHIAccess InitialAccessFlags = ERHIAccess::None);
+};
+
+struct FNiagaraRibbonShapeGeometryData
+{
+	struct FVertex
+	{
+		static constexpr int32 NumElements = 5;
+		
+		FVector2f Position;
+		FVector2f Normal;
+		float TextureV;
+
+		FVertex(const FVector2f& InPosition, const FVector2f& InNormal, float InTextureV)
+			: Position(InPosition), Normal(InNormal), TextureV(InTextureV) { }
+	};		
+	static_assert(sizeof(FVertex) == (sizeof(float) * FVertex::NumElements));
+	
+	// This sets up the first and next vertex for each pair of triangles in the slice.
+	// For a plane this will just be a linear set
+	// For a multiplane it will be multiple separate linear sets
+	// For a tube it will be a linear set that wraps back around to itself,
+	// Same with the custom vertices.
+	TArray<uint32, TInlineAllocator<32>> SliceTriangleToVertexIds;
+	FReadBuffer SliceTriangleToVertexIdsBuffer;
+	
+	TArray<FVertex> SliceVertexData;		
+	FReadBuffer SliceVertexDataBuffer;
+
+	ENiagaraRibbonShapeMode Shape;
+	
+	int32 TrianglesPerSegment;
+	int32 NumVerticesInSlice;
+	int32 BitsNeededForShape;
+	int32 BitMaskForShape;
+	bool bDisableBackfaceCulling;
+	bool bShouldFlipNormalToView;
+};
+
+
 /**
 * NiagaraRendererRibbons renders an FNiagaraEmitterInstance as a ribbon connecting all particles
 * in order by particle age.
@@ -32,106 +358,90 @@ public:
 #if RHI_RAYTRACING
 	virtual void GetDynamicRayTracingInstances(FRayTracingMaterialGatheringContext& Context, TArray<FRayTracingInstance>& OutRayTracingInstances, const FNiagaraSceneProxy* Proxy) final override;
 #endif
-	//FNiagaraInterface END
 
-	FORCEINLINE void AddDynamicParam(TArray<FNiagaraRibbonVertexDynamicParameter>& ParamData, const FVector4f& DynamicParam);
 protected:
-	struct FRibbonRenderingIndexOffsets
-	{
-		uint32 TotalBitCount;
-		bool bWantsEnd;
+	
+	static void GenerateShapeStateMultiPlane(FNiagaraRibbonShapeGeometryData& State, int32 MultiPlaneCount, int32 WidthSegmentationCount, bool bEnableAccurateGeometry);
+	static void GenerateShapeStateTube(FNiagaraRibbonShapeGeometryData& State, int32 TubeSubdivisions);
+	static void GenerateShapeStateCustom(FNiagaraRibbonShapeGeometryData& State, const TArray<FNiagaraRibbonShapeCustomVertex>& CustomVertices);
+	static void GenerateShapeStatePlane(FNiagaraRibbonShapeGeometryData& State, int32 WidthSegmentationCount);	
+	void InitializeShape(const UNiagaraRibbonRendererProperties* Properties);
+	
+	void InitializeTessellation(const UNiagaraRibbonRendererProperties* Properties);
+	
+	
+	template<typename IntType>
+	static void CalculateUVScaleAndOffsets(const FNiagaraRibbonUVSettings& UVSettings, const TArray<IntType>& RibbonIndices, const TArray<FVector4f>& RibbonTangentsAndDistances, const FNiagaraDataSetReaderFloat<float>& NormalizedAgeReader,
+		int32 StartIndex, int32 EndIndex, int32 NumSegments, float TotalLength, float& OutUScale, float& OutUOffset, float& OutUDistributionScaler);
 
-		uint32 SegmentBitShift;
-		uint32 SegmentBitMask;
+	template<bool bWantsTessellation, bool bHasTwist, bool bWantsMultiRibbon>
+	void GenerateVertexBufferForRibbonPart(const FNiagaraGenerationInputDataCPUAccessors& CPUData, const TArray<uint32>& RibbonIndices, uint32 RibbonIndex, FNiagaraRibbonCPUGeneratedVertexData& OutputData) const;
 
-		uint32 InterpBitShift;
-		uint32 InterpBitMask;
+	template<typename IDType, typename ReaderType, bool bWantsTessellation, bool bHasTwist>
+	void GenerateVertexBufferForMultiRibbonInternal(const FNiagaraGenerationInputDataCPUAccessors& CPUData, const ReaderType& IDReader, FNiagaraRibbonCPUGeneratedVertexData& OutputData) const;
+	
+	template<typename IDType, typename ReaderType>
+	void GenerateVertexBufferForMultiRibbon(const FNiagaraGenerationInputDataCPUAccessors& CPUData, const ReaderType& IDReader, FNiagaraRibbonCPUGeneratedVertexData& OutputData) const;
+	
+	void GenerateVertexBufferCPU(const FNiagaraGenerationInputDataCPUAccessors& CPUData, FNiagaraRibbonCPUGeneratedVertexData& OutputData) const;
 
-		uint32 SliceVertexBitMask;
-	};
 
-	static int32 CalculateBitsForRange(int32 Range);
-	static FRibbonRenderingIndexOffsets CalculateIndexBufferPacking(int32 NumSegments, int32 NumInterpolations, int32 NumSliceVertices);
+	int32 CalculateTessellationFactor(const FNiagaraSceneProxy* SceneProxy, const FSceneView* View, const FVector& ViewOriginForDistanceCulling) const;
+	FNiagaraIndexGenerationInput CalculateIndexBufferConfiguration(const TSharedPtr<FNiagaraRibbonCPUGeneratedVertexData>& GeneratedVertices, const FNiagaraDataBuffer* SourceParticleData,
+		const FNiagaraSceneProxy* SceneProxy, const FSceneView* View, const FVector& ViewOriginForDistanceCulling, bool bShouldUseGPUInitIndices, bool bIsGPUSim) const;
+	
+	void GenerateIndexBufferForView(FNiagaraIndexGenerationInput& GeneratedData, FNiagaraDynamicDataRibbon* DynamicDataRibbon,
+	                                const TSharedPtr<FNiagaraRibbonRenderingFrameViewResources>& RenderingViewResources, const FSceneView* View, const FVector& ViewOriginForDistanceCulling, bool bShouldUseGPUInitIndices) const;
+	
+	template <typename TValue>
+	static void GenerateIndexBufferCPU(FNiagaraIndexGenerationInput& GeneratedData, FNiagaraDynamicDataRibbon* DynamicDataRibbon, const FNiagaraRibbonShapeGeometryData& ShapeState,
+	                                   const TSharedPtr<FNiagaraRibbonRenderingFrameViewResources>& RenderingViewResources, const FSceneView* View, const FVector& ViewOriginForDistanceCulling, ERHIFeatureLevel::Type FeatureLevel, ENiagaraRibbonDrawDirection DrawDirection);
 
 	template <typename TValue>
-	TValue* AppendToIndexBuffer(
-		TValue* OutIndices,
-		uint32& OutMaxUsedIndex,
-		const TArrayView<int32>& SegmentData,
-		const FRibbonRenderingIndexOffsets& Offsets,
-		int32 InterpCount,
-		bool bInvertOrder) const;
+	static TValue* AppendToIndexBufferCPU(TValue* OutIndices, const FNiagaraIndexGenerationInput& GeneratedData, const FNiagaraRibbonShapeGeometryData& ShapeState, const TArrayView<uint32>& SegmentData, bool bInvertOrder);
+	
+	void SetupPerViewUniformBuffer(FNiagaraIndexGenerationInput& GeneratedData,
+	                               const FSceneView* View, const FSceneViewFamily& ViewFamily, const FNiagaraSceneProxy* SceneProxy, FNiagaraRibbonUniformBufferRef& OutUniformBuffer) const;
 
+	void SetupMeshBatchAndCollectorResourceForView(const FNiagaraIndexGenerationInput& GeneratedData, FNiagaraDynamicDataRibbon* DynamicDataRibbon,
+	                                               const FNiagaraDataBuffer* SourceParticleData, const FSceneView* View, const FSceneViewFamily& ViewFamily, const FNiagaraSceneProxy* SceneProxy,
+	                                               const TSharedPtr<FNiagaraRibbonRenderingFrameResources>& RenderingResources, const TSharedPtr<FNiagaraRibbonRenderingFrameViewResources>& RenderingViewResources, FMeshBatch& OutMeshBatch, bool bShouldUseGPUInitIndices) const;
 
-	/** Generate the raw index buffer preserving multi ribbon ordering. */
-	template <typename TValue>
-	void GenerateIndexBuffer(
-		FGlobalDynamicIndexBuffer::FAllocationEx& InOutIndexAllocation, 
-		const FRibbonRenderingIndexOffsets& Offsets,
-		int32 InterpCount,
-		const FVector& ViewDirection, 
-		const FVector& ViewOriginForDistanceCulling, 
-		struct FNiagaraDynamicDataRibbon* DynamicData) const;
 
-private:
-	struct FCPUSimParticleDataAllocation
-	{
-		FGlobalDynamicReadBuffer& DynamicReadBuffer;
-		FParticleRenderData ParticleData;
-	};
-
-	void SetupMeshBatchAndCollectorResourceForView(
-		const FSceneView* View,
-		const FSceneViewFamily& ViewFamily,
-		const FNiagaraSceneProxy* SceneProxy,
-		FMeshElementCollector& Collector,
-		struct FNiagaraDynamicDataRibbon* DynamicData,
-		const FGlobalDynamicIndexBuffer::FAllocationEx& IndexAllocation,
-		FMeshBatch& OutMeshBatch,
-		class FNiagaraMeshCollectorResourcesRibbon& OutCollectorResources) const;
-
-	void CreatePerViewResources(
-		const FSceneView* View,
-		const FSceneViewFamily& ViewFamily,
-		const FNiagaraSceneProxy* SceneProxy,
-		FMeshElementCollector& Collector,
-		FNiagaraRibbonUniformBufferRef& OutUniformBuffer,
-		FGlobalDynamicIndexBuffer::FAllocationEx& InOutIndexAllocation) const;
-
-	FCPUSimParticleDataAllocation AllocateParticleDataIfCPUSim(struct FNiagaraDynamicDataRibbon* DynamicDataRibbon, FGlobalDynamicReadBuffer& DynamicReadBuffer) const;
+	void InitializeViewIndexBuffersGPU(FRHICommandListImmediate& CMDList, FNiagaraGpuComputeDispatchInterface* ComputeDispatchInterface, const FNiagaraDataBuffer* SourceParticleData,
+		const TSharedPtr<FNiagaraRibbonRenderingFrameViewResources>& RenderingViewResources) const;
 
-	ENiagaraRibbonFacingMode FacingMode;
+	void InitializeVertexBuffersResources(const FNiagaraDynamicDataRibbon* DynamicDataRibbon, FNiagaraDataBuffer* SourceParticleData,
+	                                      FGlobalDynamicReadBuffer& DynamicReadBuffer, const TSharedPtr<FNiagaraRibbonRenderingFrameResources>& RenderingResources, bool bShouldUseGPUInit) const;
+	
+	void InitializeVertexBuffersGPU(FRHICommandListImmediate& CMDList, FNiagaraGpuComputeDispatchInterface* ComputeDispatchInterface, const FNiagaraDataBuffer* SourceParticleData,
+		struct FNiagaraRibbonGPUInitComputeBuffers& TempBuffers, const TSharedPtr<FNiagaraRibbonRenderingFrameResources>& RenderingResources) const;
+
+	FRibbonComputeUniformParameters SetupComputeVertexGenParams(FNiagaraGpuComputeDispatchInterface* ComputeDispatchInterface, const TSharedPtr<FNiagaraRibbonRenderingFrameResources>& RenderingResources, const FNiagaraDataBuffer* SourceParticleData) const;
+
+	FNiagaraRibbonGenerationConfig GenerationConfig;
+	
 	FNiagaraRibbonUVSettings UV0Settings;
 	FNiagaraRibbonUVSettings UV1Settings;
+	FNiagaraRibbonShapeGeometryData ShapeState;
+	FNiagaraRibbonTessellationConfig TessellationConfig;
+	
+	ENiagaraRibbonFacingMode FacingMode;
 	ENiagaraRibbonDrawDirection DrawDirection;
+	
+	const FNiagaraRendererLayout* RendererLayout;
+
+	mutable FNiagaraRibbonVertexBuffers VertexBuffers;
+	
+	mutable FNiagaraRibbonTessellationSmoothingData TessellationSmoothingData;
+	int32 RibbonIDParamDataSetOffset;
+
+	friend class FNiagaraRibbonComputeDispatchManager;
+};
+
+
 
-	ENiagaraRibbonShapeMode Shape;
-	bool bEnableAccurateGeometry;
-	int32 WidthSegmentationCount;
-	int32 MultiPlaneCount;
-	int32 TubeSubdivisions;
-	TArray<FNiagaraRibbonShapeCustomVertex> CustomVertices;
 
-	ENiagaraRibbonTessellationMode TessellationMode;
-	float CustomCurveTension;
-	int32 CustomTessellationFactor;
-	bool bCustomUseConstantFactor;
-	float CustomTessellationMinAngle;
-	bool bCustomUseScreenSpace;
-	bool bNeedsPreciseMotionVectors;
 
 
-	uint32 MaterialParamValidMask;
-	const FNiagaraRendererLayout* RendererLayout;
 
-	// Average curvature of the segments.
-	mutable float TessellationAngle = 0;
-	// Average curvature of the segments (computed from the segment angle in radian).
-	mutable float TessellationCurvature = 0;
-	// Average twist of the segments.
-	mutable float TessellationTwistAngle = 0;
-	// Average twist curvature of the segments.
-	mutable float TessellationTwistCurvature = 0;
-	// Average twist curvature of the segments.
-	mutable float TessellationTotalSegmentLength = 0;
-};
diff --git a/Engine/Plugins/FX/Niagara/Source/Niagara/Public/NiagaraRibbonRendererProperties.h b/Engine/Plugins/FX/Niagara/Source/Niagara/Public/NiagaraRibbonRendererProperties.h
--- a/Engine/Plugins/FX/Niagara/Source/Niagara/Public/NiagaraRibbonRendererProperties.h
+++ b/Engine/Plugins/FX/Niagara/Source/Niagara/Public/NiagaraRibbonRendererProperties.h
@@ -67,10 +67,10 @@ struct FNiagaraRibbonShapeCustomVertex
 	FNiagaraRibbonShapeCustomVertex();
 
 	UPROPERTY(EditAnywhere, Category = "Ribbon Rendering")
-	FVector2D Position;
+	FVector2f Position;
 
 	UPROPERTY(EditAnywhere, Category = "Ribbon Rendering")
-	FVector2D Normal;
+	FVector2f Normal;
 
 	UPROPERTY(EditAnywhere, Category = "Ribbon Rendering")
 	float TextureV;
@@ -114,6 +114,8 @@ enum class ENiagaraRibbonUVDistributionMode
 	TiledFromStartOverRibbonLength UMETA(DisplayName = "Tiled By Distance (By Particles.RibbonUVDistance)")
 };
 
+
+
 /** Defines settings for UV behavior for a UV channel on ribbons. */
 USTRUCT()
 struct FNiagaraRibbonUVSettings
@@ -181,6 +183,7 @@ namespace ENiagaraRibbonVFLayout
 		PrevRibbonWidth,
 		PrevRibbonFacing,
 		PrevRibbonTwist,
+		LinkOrder,
 		Num,
 	};
 };
@@ -210,8 +213,8 @@ public:
 	virtual FNiagaraRenderer* CreateEmitterRenderer(ERHIFeatureLevel::Type FeatureLevel, const FNiagaraEmitterInstance* Emitter, const FNiagaraSystemInstanceController& InController) override;
 	virtual class FNiagaraBoundsCalculator* CreateBoundsCalculator() override;
 	virtual void GetUsedMaterials(const FNiagaraEmitterInstance* InEmitter, TArray<UMaterialInterface*>& OutMaterials) const override;
-	virtual bool IsSimTargetSupported(ENiagaraSimTarget InSimTarget) const override { return (InSimTarget == ENiagaraSimTarget::CPUSim); };
-	bool PopulateRequiredBindings(FNiagaraParameterStore& InParameterStore);
+	virtual bool IsSimTargetSupported(ENiagaraSimTarget InSimTarget) const override { return true; };
+	virtual bool PopulateRequiredBindings(FNiagaraParameterStore& InParameterStore) override;
 
 #if WITH_EDITOR
 	virtual const TArray<FNiagaraVariable>& GetOptionalAttributes() override;
@@ -219,12 +222,12 @@ public:
 	virtual FNiagaraVariable GetBoundAttribute(const FNiagaraVariableAttributeBinding* Binding) const override;
 	virtual void GetRendererWidgets(const FNiagaraEmitterInstance* InEmitter, TArray<TSharedPtr<SWidget>>& OutWidgets, TSharedPtr<FAssetThumbnailPool> InThumbnailPool) const override;
 	virtual void GetRendererTooltipWidgets(const FNiagaraEmitterInstance* InEmitter, TArray<TSharedPtr<SWidget>>& OutWidgets, TSharedPtr<FAssetThumbnailPool> InThumbnailPool) const override;
-	virtual void GetRendererFeedback(const FVersionedNiagaraEmitter& InEmitter, TArray<FText>& OutErrors, TArray<FText>& OutWarnings, TArray<FText>& OutInfo) const override;
+	virtual void GetRendererFeedback(const FVersionedNiagaraEmitter& InEmitter, TArray<FNiagaraRendererFeedback>& OutErrors, TArray<FNiagaraRendererFeedback>& OutWarnings, TArray<FNiagaraRendererFeedback>& OutInfo) const override;
 #endif
 	virtual void CacheFromCompiledData(const FNiagaraDataSetCompiledData* CompiledData) override;
 	
 #if WITH_EDITORONLY_DATA
-	bool IsSupportedVariableForBinding(const FNiagaraVariableBase& InSourceForBinding, const FName& InTargetBindingName) const;
+	virtual bool IsSupportedVariableForBinding(const FNiagaraVariableBase& InSourceForBinding, const FName& InTargetBindingName) const override;
 #endif
 	//UNiagaraRendererProperties Interface END
 
@@ -273,6 +276,20 @@ private:
 
 public:
 
+	UPROPERTY(EditAnywhere, Category = "Ribbon Rendering")	
+	int32 MaxNumRibbons;
+	
+	/**
+	*	Whether we use the CPU or GPU to generate ribbon geometry for CPU systems.
+	*	GPU systems will always use a fully GPU initialization pipeline,
+	*	Will fall back to CPU init when GPU init isn't available.
+	*/ 
+	UPROPERTY(EditAnywhere, Category = "Ribbon Rendering")	
+	bool bUseGPUInit;
+
+
+	
+
 	/** If true, the particles are only sorted when using a translucent material. */
 	UPROPERTY(EditAnywhere, Category = "Ribbon Rendering")
 	ENiagaraRibbonDrawDirection DrawDirection;
@@ -310,7 +327,7 @@ public:
 	/** Defines the curve tension, or how long the curve's tangents are.
 	  * Ranges from 0 to 1. The higher the value, the sharper the curve becomes.
 	  */
-	UPROPERTY(EditAnywhere, Category = "Tessellation", meta = (ClampMin = "0", ClampMax = "1"))
+	UPROPERTY(EditAnywhere, Category = "Tessellation", meta = (ClampMin = "0", ClampMax = "0.99"))
 	float CurveTension;
 
 	/** Defines the tessellation mode allowing custom tessellation parameters or disabling tessellation entirely. */
@@ -452,6 +469,8 @@ public:
 
 	FNiagaraDataSetAccessor<int32>		RibbonIdDataSetAccessor;
 	FNiagaraDataSetAccessor<FNiagaraID>	RibbonFullIDDataSetAccessor;
+	
+	FNiagaraDataSetAccessor<float>		RibbonLinkOrderDataSetAccessor;
 
 	uint32 MaterialParamValidMask = 0;
 	FNiagaraRendererLayout RendererLayout;
diff --git a/Engine/Plugins/FX/Niagara/Source/NiagaraShader/Public/NiagaraRibbonCompute.h b/Engine/Plugins/FX/Niagara/Source/NiagaraShader/Public/NiagaraRibbonCompute.h
new file mode 100644
--- /dev/null
+++ b/Engine/Plugins/FX/Niagara/Source/NiagaraShader/Public/NiagaraRibbonCompute.h
@@ -0,0 +1,352 @@
+﻿// Copyright Epic Games, Inc. All Rights Reserved.
+
+/*==============================================================================
+NiagaraRibbonCompute.h: Niagara ribbon compute shaders for initialization of ribbons on the gpu
+==============================================================================*/
+
+#pragma once
+
+#include "CoreMinimal.h"
+#include "GlobalShader.h"
+#include "ShaderPermutation.h"
+#include "ShaderParameterStruct.h"
+
+BEGIN_SHADER_PARAMETER_STRUCT(FRibbonComputeUniformParameters, NIAGARASHADER_API)
+
+	// Total particle Count
+	SHADER_PARAMETER(uint32, TotalNumParticlesDirect)
+	SHADER_PARAMETER_SRV(Buffer<uint32>, EmitterParticleCountsBuffer)
+	SHADER_PARAMETER(int32, EmitterParticleCountsBufferOffset)
+
+	// Niagara sim data
+	SHADER_PARAMETER_SRV(Buffer<float>, NiagaraParticleDataFloat)
+	SHADER_PARAMETER_SRV(Buffer<float>, NiagaraParticleDataHalf)
+	SHADER_PARAMETER_SRV(Buffer<int>, NiagaraParticleDataInt)
+	SHADER_PARAMETER(int, NiagaraFloatDataStride)
+	SHADER_PARAMETER(int, NiagaraIntDataStride)
+
+	// Int bindings
+	SHADER_PARAMETER(int, RibbonIdDataOffset)	
+	SHADER_PARAMETER(int, RibbonLinkOrderDataOffset)
+
+	// Float bindings
+	SHADER_PARAMETER(int, PositionDataOffset)
+	SHADER_PARAMETER(int, PrevPositionDataOffset)
+	SHADER_PARAMETER(int, VelocityDataOffset)
+	SHADER_PARAMETER(int, WidthDataOffset)
+	SHADER_PARAMETER(int, PrevWidthDataOffset)
+	SHADER_PARAMETER(int, TwistDataOffset)
+	SHADER_PARAMETER(int, PrevTwistDataOffset)
+	SHADER_PARAMETER(int, ColorDataOffset)
+	SHADER_PARAMETER(int, FacingDataOffset)
+	SHADER_PARAMETER(int, PrevFacingDataOffset)
+	SHADER_PARAMETER(int, NormalizedAgeDataOffset)
+	SHADER_PARAMETER(int, MaterialRandomDataOffset)
+	SHADER_PARAMETER(uint32, MaterialParamValidMask)
+	SHADER_PARAMETER(int, MaterialParamDataOffset)
+	SHADER_PARAMETER(int, MaterialParam1DataOffset)
+	SHADER_PARAMETER(int, MaterialParam2DataOffset)
+	SHADER_PARAMETER(int, MaterialParam3DataOffset)
+	SHADER_PARAMETER(int, DistanceFromStartOffset)
+	SHADER_PARAMETER(int, U0OverrideDataOffset)
+	SHADER_PARAMETER(int, V0RangeOverrideDataOffset)
+	SHADER_PARAMETER(int, U1OverrideDataOffset)
+	SHADER_PARAMETER(int, V1RangeOverrideDataOffset)
+	SHADER_PARAMETER(int, U0DistributionMode)
+	SHADER_PARAMETER(int, U1DistributionMode)
+
+END_SHADER_PARAMETER_STRUCT()
+
+
+class FRibbonHasFullRibbonID : SHADER_PERMUTATION_BOOL("RIBBONID_IS_NIAGARAID");
+class FRibbonHasRibbonID : SHADER_PERMUTATION_BOOL("RIBBONID_IS_INT");
+class FRibbonHasCustomLinkOrder : SHADER_PERMUTATION_BOOL("RIBBON_HAS_LINK_ORDER");
+
+class FRibbonWantsConstantTessellation : SHADER_PERMUTATION_BOOL("RIBBONS_WANTS_CONSTANT_TESSELLATION");
+class FRibbonWantsAutomaticTessellation : SHADER_PERMUTATION_BOOL("RIBBONS_WANTS_AUTOMATIC_TESSELLATION");
+
+class FRibbonHasTwist : SHADER_PERMUTATION_BOOL("RIBBON_HAS_TWIST");
+
+class FRibbonHasHighSliceComplexity : SHADER_PERMUTATION_BOOL("RIBBON_HAS_HIGH_SLICE_COMPLEXITY");
+using FRibbonComputePermutationDomain = TShaderPermutationDomain<FRibbonHasFullRibbonID, FRibbonHasRibbonID, FRibbonHasCustomLinkOrder, FRibbonWantsConstantTessellation, FRibbonWantsAutomaticTessellation, FRibbonHasTwist, FRibbonHasHighSliceComplexity>;
+
+
+struct FNiagaraRibbonComputeCommon
+{	
+	static void ModifyCompilationEnvironment(const FGlobalShaderPermutationParameters& Parameters, FShaderCompilerEnvironment& OutEnvironment, int32 ThreadGroupSize);
+	static bool ShouldCompilePermutation(const FGlobalShaderPermutationParameters& Parameters);
+
+	
+	static constexpr uint32 VertexGenReductionInitializationThreadSize = 64;
+	static constexpr uint32 VertexGenReductionPropagationThreadSize = 64;
+	static constexpr uint32 VertexGenReductionFinalizationThreadSize = 64;
+	static constexpr uint32 VertexGenFinalizationThreadSize = 64;
+	
+	static constexpr uint32 IndexGenThreadSize = 64;
+
+	static constexpr uint32 IndexGenOptimalLoopVertexLimit = 32;
+};
+
+
+
+
+BEGIN_SHADER_PARAMETER_STRUCT(FRibbonOrderSortParameters, NIAGARASHADER_API)
+	SHADER_PARAMETER_STRUCT_INCLUDE(FRibbonComputeUniformParameters, Common)
+	SHADER_PARAMETER_SRV(Buffer<uint32>, SortedIndices)
+	SHADER_PARAMETER_UAV(RWBuffer<uint32>, DestinationSortedIndices)	
+	SHADER_PARAMETER(uint32, MergeSortSourceBlockSize)
+	SHADER_PARAMETER(uint32, MergeSortDestinationBlockSize)
+END_SHADER_PARAMETER_STRUCT()
+
+/**
+* Compute shader used to generate particle sort keys.
+*/
+class NIAGARASHADER_API FNiagaraRibbonSortPhase1CS : public FGlobalShader
+{
+	DECLARE_GLOBAL_SHADER(FNiagaraRibbonSortPhase1CS);
+	SHADER_USE_PARAMETER_STRUCT(FNiagaraRibbonSortPhase1CS, FGlobalShader);
+
+	static constexpr uint32 BubbleSortGroupWidth = 32;
+
+	using FParameters = FRibbonOrderSortParameters;
+	using FPermutationDomain = TShaderPermutationDomain<FRibbonHasFullRibbonID, FRibbonHasRibbonID, FRibbonHasCustomLinkOrder>;
+
+	static bool ShouldCompilePermutation(const FGlobalShaderPermutationParameters& Parameters)
+	{
+		return FNiagaraRibbonComputeCommon::ShouldCompilePermutation(Parameters);
+	}
+
+	static void ModifyCompilationEnvironment(const FGlobalShaderPermutationParameters& Parameters, FShaderCompilerEnvironment& OutEnvironment);
+};
+
+/**
+* Compute shader used to generate particle sort keys.
+*/
+class NIAGARASHADER_API FNiagaraRibbonSortPhase2CS : public FGlobalShader
+{
+	DECLARE_GLOBAL_SHADER(FNiagaraRibbonSortPhase2CS);
+	SHADER_USE_PARAMETER_STRUCT(FNiagaraRibbonSortPhase2CS, FGlobalShader);
+
+	static constexpr uint32 ThreadGroupSize = 64;
+
+	using FParameters = FRibbonOrderSortParameters;
+	using FPermutationDomain = TShaderPermutationDomain<FRibbonHasFullRibbonID, FRibbonHasRibbonID, FRibbonHasCustomLinkOrder>;
+
+	static bool ShouldCompilePermutation(const FGlobalShaderPermutationParameters& Parameters)
+	{
+		return RHISupportsComputeShaders(Parameters.Platform);
+	}
+
+	static void ModifyCompilationEnvironment(const FGlobalShaderPermutationParameters& Parameters, FShaderCompilerEnvironment& OutEnvironment);
+};
+
+
+BEGIN_SHADER_PARAMETER_STRUCT(FNiagaraRibbonVertexReductionParameters, NIAGARASHADER_API)
+	SHADER_PARAMETER_STRUCT_INCLUDE(FRibbonComputeUniformParameters, Common)
+	SHADER_PARAMETER_SRV(Buffer<uint32>, SortedIndices)
+	SHADER_PARAMETER_SRV(Buffer<float>, InputTangentsAndDistances)
+	SHADER_PARAMETER_UAV(RWBuffer<float>, OutputTangentsAndDistances)
+	SHADER_PARAMETER_SRV(Buffer<uint32>, InputMultiRibbonIndices)
+	SHADER_PARAMETER_UAV(RWBuffer<uint32>, OutputMultiRibbonIndices)
+	SHADER_PARAMETER_SRV(Buffer<uint32>, InputSegments)
+	SHADER_PARAMETER_UAV(RWBuffer<uint32>, OutputSegments)
+	SHADER_PARAMETER_SRV(Buffer<float>, InputTessellationStats)
+	SHADER_PARAMETER_UAV(RWBuffer<float>, OutputTessellationStats)
+	SHADER_PARAMETER(float, CurveTension)
+	SHADER_PARAMETER(int, PrefixScanStride)
+END_SHADER_PARAMETER_STRUCT()
+
+class NIAGARASHADER_API FNiagaraRibbonVertexReductionInitializationCS : public FGlobalShader
+{
+	DECLARE_GLOBAL_SHADER(FNiagaraRibbonVertexReductionInitializationCS);
+	SHADER_USE_PARAMETER_STRUCT(FNiagaraRibbonVertexReductionInitializationCS, FGlobalShader);
+	
+	using FPermutationDomain = TShaderPermutationDomain<FRibbonHasFullRibbonID, FRibbonHasRibbonID, FRibbonWantsConstantTessellation, FRibbonWantsAutomaticTessellation, FRibbonHasTwist>;
+	using FParameters = FNiagaraRibbonVertexReductionParameters;
+	
+	static bool ShouldCompilePermutation(const FGlobalShaderPermutationParameters& Parameters)
+	{
+		return RHISupportsComputeShaders(Parameters.Platform);
+	}
+
+	static void ModifyCompilationEnvironment(const FGlobalShaderPermutationParameters& Parameters, FShaderCompilerEnvironment& OutEnvironment);
+};
+
+class NIAGARASHADER_API FNiagaraRibbonVertexReductionPropagateCS : public FGlobalShader
+{
+	DECLARE_GLOBAL_SHADER(FNiagaraRibbonVertexReductionPropagateCS);
+	SHADER_USE_PARAMETER_STRUCT(FNiagaraRibbonVertexReductionPropagateCS, FGlobalShader);
+	
+	using FPermutationDomain = TShaderPermutationDomain<FRibbonHasFullRibbonID, FRibbonHasRibbonID, FRibbonWantsConstantTessellation, FRibbonWantsAutomaticTessellation, FRibbonHasTwist>;
+	using FParameters = FNiagaraRibbonVertexReductionParameters
+	;
+	
+	static bool ShouldCompilePermutation(const FGlobalShaderPermutationParameters& Parameters)
+	{
+		return RHISupportsComputeShaders(Parameters.Platform);
+	}
+
+	static void ModifyCompilationEnvironment(const FGlobalShaderPermutationParameters& Parameters, FShaderCompilerEnvironment& OutEnvironment);
+};
+
+BEGIN_SHADER_PARAMETER_STRUCT(FNiagaraRibbonVertexReductionFinalizationParameters, NIAGARASHADER_API)
+	SHADER_PARAMETER_STRUCT_INCLUDE(FRibbonComputeUniformParameters, Common)
+	SHADER_PARAMETER_SRV(Buffer<uint32>, SortedIndices)
+	SHADER_PARAMETER_SRV(Buffer<float>, TangentsAndDistances)
+	SHADER_PARAMETER_SRV(Buffer<uint32>, MultiRibbonIndices)
+	SHADER_PARAMETER_SRV(Buffer<uint32>, Segments)
+	SHADER_PARAMETER_SRV(Buffer<float>, TessellationStats)
+	SHADER_PARAMETER_UAV(RWBuffer<uint32>, PackedPerRibbonData)
+	SHADER_PARAMETER_UAV(RWBuffer<uint32>, OutputCommandBuffer)
+	SHADER_PARAMETER(int, OutputCommandBufferIndex)
+	SHADER_PARAMETER(int, FinalizationThreadBlockSize)
+END_SHADER_PARAMETER_STRUCT()
+
+class NIAGARASHADER_API FNiagaraRibbonVertexReductionFinalizeCS : public FGlobalShader
+{
+	DECLARE_GLOBAL_SHADER(FNiagaraRibbonVertexReductionFinalizeCS);
+	SHADER_USE_PARAMETER_STRUCT(FNiagaraRibbonVertexReductionFinalizeCS, FGlobalShader);
+	
+	using FPermutationDomain = TShaderPermutationDomain<FRibbonHasFullRibbonID, FRibbonHasRibbonID, FRibbonWantsConstantTessellation, FRibbonWantsAutomaticTessellation, FRibbonHasTwist>;
+	using FParameters = FNiagaraRibbonVertexReductionFinalizationParameters;
+	
+	static bool ShouldCompilePermutation(const FGlobalShaderPermutationParameters& Parameters)
+	{
+		return RHISupportsComputeShaders(Parameters.Platform);
+	}
+
+	static void ModifyCompilationEnvironment(const FGlobalShaderPermutationParameters& Parameters, FShaderCompilerEnvironment& OutEnvironment);
+};
+
+BEGIN_SHADER_PARAMETER_STRUCT(FNiagaraRibbonUVSettingsParams, NIAGARASHADER_API)
+	SHADER_PARAMETER(FVector2f, Offset)
+	SHADER_PARAMETER(FVector2f, Scale)
+	SHADER_PARAMETER(float, TilingLength)
+	SHADER_PARAMETER(int, DistributionMode)
+	SHADER_PARAMETER(int, LeadingEdgeMode)
+	SHADER_PARAMETER(int, TrailingEdgeMode)
+	SHADER_PARAMETER(int, bEnablePerParticleUOverride)
+	SHADER_PARAMETER(int, bEnablePerParticleVRangeOverride)
+END_SHADER_PARAMETER_STRUCT()
+
+BEGIN_SHADER_PARAMETER_STRUCT(FNiagaraRibbonVertexFinalizationParameters, NIAGARASHADER_API)
+	SHADER_PARAMETER_STRUCT_INCLUDE(FRibbonComputeUniformParameters, Common)
+	SHADER_PARAMETER_STRUCT(FNiagaraRibbonUVSettingsParams, UV0Settings)
+	SHADER_PARAMETER_STRUCT(FNiagaraRibbonUVSettingsParams, UV1Settings)
+	SHADER_PARAMETER_SRV(Buffer<uint32>, SortedIndices)
+	SHADER_PARAMETER_UAV(RWBuffer<float>, TangentsAndDistances)
+	SHADER_PARAMETER_UAV(RWBuffer<uint32>, PackedPerRibbonData)
+	SHADER_PARAMETER_SRV(Buffer<uint32>, CommandBuffer)
+	SHADER_PARAMETER(int32, CommandBufferOffset)
+	SHADER_PARAMETER(int32, TotalNumRibbons)
+END_SHADER_PARAMETER_STRUCT()
+
+class NIAGARASHADER_API FNiagaraRibbonUVParamCalculationCS : public FGlobalShader
+{
+	DECLARE_GLOBAL_SHADER(FNiagaraRibbonUVParamCalculationCS);
+	SHADER_USE_PARAMETER_STRUCT(FNiagaraRibbonUVParamCalculationCS, FGlobalShader);
+	
+	using FPermutationDomain = TShaderPermutationDomain<FRibbonHasFullRibbonID, FRibbonHasRibbonID, FRibbonWantsConstantTessellation, FRibbonWantsAutomaticTessellation, FRibbonHasTwist>;
+	using FParameters = FNiagaraRibbonVertexFinalizationParameters;
+	
+	static bool ShouldCompilePermutation(const FGlobalShaderPermutationParameters& Parameters)
+	{
+		return RHISupportsComputeShaders(Parameters.Platform);
+	}
+
+	static void ModifyCompilationEnvironment(const FGlobalShaderPermutationParameters& Parameters, FShaderCompilerEnvironment& OutEnvironment);
+};
+
+
+
+BEGIN_SHADER_PARAMETER_STRUCT(FNiagaraRibbonInitializeIndices, NIAGARASHADER_API)
+	SHADER_PARAMETER_UAV(RWBuffer<uint32>, IndirectDrawOutput)
+	SHADER_PARAMETER_SRV(Buffer<uint32>, VertexGenerationResults)
+
+	// Direct and Indirect particle counts
+	SHADER_PARAMETER(uint32, TotalNumParticlesDirect)
+	SHADER_PARAMETER_SRV(Buffer<uint32>, EmitterParticleCountsBuffer)
+	SHADER_PARAMETER(int32, EmitterParticleCountsBufferOffset)
+
+	SHADER_PARAMETER(uint32, IndirectDrawOutputIndex)
+	SHADER_PARAMETER(int32, VertexGenerationResultsIndex)
+	SHADER_PARAMETER(uint32, IndexGenThreadSize)
+	SHADER_PARAMETER(uint32, TrianglesPerSegment)
+
+	SHADER_PARAMETER(float, ViewDistance)
+	SHADER_PARAMETER(float, LODDistanceFactor)
+	SHADER_PARAMETER(uint32, TessellationMode)
+	SHADER_PARAMETER(uint32, bCustomUseConstantFactor)
+	SHADER_PARAMETER(uint32, CustomTessellationFactor)
+	SHADER_PARAMETER(float, CustomTessellationMinAngle)
+	SHADER_PARAMETER(uint32, bCustomUseScreenSpace)
+	SHADER_PARAMETER(uint32, GNiagaraRibbonMaxTessellation)
+	SHADER_PARAMETER(float, GNiagaraRibbonTessellationAngle)
+	SHADER_PARAMETER(float, GNiagaraRibbonTessellationScreenPercentage)
+	SHADER_PARAMETER(uint32, GNiagaraRibbonTessellationEnabled)
+	SHADER_PARAMETER(float, GNiagaraRibbonTessellationMinDisplacementError)
+END_SHADER_PARAMETER_STRUCT()
+
+class NIAGARASHADER_API FNiagaraRibbonCreateIndexBufferParamsCS : public FGlobalShader
+{
+	DECLARE_GLOBAL_SHADER(FNiagaraRibbonCreateIndexBufferParamsCS);
+	SHADER_USE_PARAMETER_STRUCT(FNiagaraRibbonCreateIndexBufferParamsCS, FGlobalShader);
+
+	
+	using FPermutationDomain = TShaderPermutationDomain<FRibbonWantsConstantTessellation, FRibbonWantsAutomaticTessellation>;
+	using FParameters = FNiagaraRibbonInitializeIndices;
+
+	static bool ShouldCompilePermutation(const FGlobalShaderPermutationParameters& Parameters)
+	{
+		return RHISupportsComputeShaders(Parameters.Platform);
+	}
+
+	static void ModifyCompilationEnvironment(const FGlobalShaderPermutationParameters& Parameters, FShaderCompilerEnvironment& OutEnvironment);
+};
+
+
+BEGIN_SHADER_PARAMETER_STRUCT(FNiagaraRibbonGenerateIndices, NIAGARASHADER_API)
+	SHADER_PARAMETER_UAV(RWBuffer<uint32>, GeneratedIndicesBuffer)
+	SHADER_PARAMETER_SRV(Buffer<uint32>, SortedIndices)
+	SHADER_PARAMETER_SRV(Buffer<uint32>, MultiRibbonIndices)
+	SHADER_PARAMETER_SRV(Buffer<uint32>, Segments)
+
+	SHADER_PARAMETER_SRV(Buffer<uint32>, IndirectDrawInfo)
+	SHADER_PARAMETER_SRV(Buffer<uint32>, TriangleToVertexIds)
+
+	// Direct and Indirect particle counts
+	SHADER_PARAMETER(uint32, TotalNumParticlesDirect)
+	SHADER_PARAMETER_SRV(Buffer<uint32>, EmitterParticleCountsBuffer)
+	SHADER_PARAMETER(int32, EmitterParticleCountsBufferOffset)
+
+	SHADER_PARAMETER(uint32, IndexBufferOffset)
+	SHADER_PARAMETER(uint32, IndirectDrawInfoIndex)
+	SHADER_PARAMETER(uint32, TriangleToVertexIdsCount)
+
+	SHADER_PARAMETER(uint32, TrianglesPerSegment)
+	SHADER_PARAMETER(uint32, NumVerticesInSlice)
+	SHADER_PARAMETER(uint32, BitsNeededForShape)
+	SHADER_PARAMETER(uint32, BitMaskForShape)
+	SHADER_PARAMETER(uint32, SegmentBitShift)
+	SHADER_PARAMETER(uint32, SegmentBitMask)
+	SHADER_PARAMETER(uint32, SubSegmentBitShift)
+	SHADER_PARAMETER(uint32, SubSegmentBitMask)
+END_SHADER_PARAMETER_STRUCT()
+
+class NIAGARASHADER_API FNiagaraRibbonCreateIndexBufferCS : public FGlobalShader
+{
+	DECLARE_GLOBAL_SHADER(FNiagaraRibbonCreateIndexBufferCS);
+	SHADER_USE_PARAMETER_STRUCT(FNiagaraRibbonCreateIndexBufferCS, FGlobalShader);
+
+	
+	using FPermutationDomain = TShaderPermutationDomain<FRibbonHasFullRibbonID, FRibbonHasRibbonID, FRibbonHasHighSliceComplexity>;;
+	using FParameters = FNiagaraRibbonGenerateIndices;
+
+	static bool ShouldCompilePermutation(const FGlobalShaderPermutationParameters& Parameters)
+	{
+		return RHISupportsComputeShaders(Parameters.Platform);
+	}
+
+	static void ModifyCompilationEnvironment(const FGlobalShaderPermutationParameters& Parameters, FShaderCompilerEnvironment& OutEnvironment);
+};
+
diff --git a/Engine/Plugins/FX/Niagara/Source/NiagaraVertexFactories/Public/NiagaraRibbonVertexFactory.h b/Engine/Plugins/FX/Niagara/Source/NiagaraVertexFactories/Public/NiagaraRibbonVertexFactory.h
--- a/Engine/Plugins/FX/Niagara/Source/NiagaraVertexFactories/Public/NiagaraRibbonVertexFactory.h
+++ b/Engine/Plugins/FX/Niagara/Source/NiagaraVertexFactories/Public/NiagaraRibbonVertexFactory.h
@@ -39,6 +39,7 @@ BEGIN_GLOBAL_SHADER_PARAMETER_STRUCT(FNiagaraRibbonUniformParameters, NIAGARAVER
 	SHADER_PARAMETER(int, ColorDataOffset)
 	SHADER_PARAMETER(int, FacingDataOffset)
 	SHADER_PARAMETER(int, PrevFacingDataOffset)
+	SHADER_PARAMETER(int, LinkOrderDataOffset)
 	SHADER_PARAMETER(int, NormalizedAgeDataOffset)
 	SHADER_PARAMETER(int, MaterialRandomDataOffset)
 	SHADER_PARAMETER(uint32, MaterialParamValidMask)
@@ -51,7 +52,6 @@ BEGIN_GLOBAL_SHADER_PARAMETER_STRUCT(FNiagaraRibbonUniformParameters, NIAGARAVER
 	SHADER_PARAMETER(int, V0RangeOverrideDataOffset)
 	SHADER_PARAMETER(int, U1OverrideDataOffset)
 	SHADER_PARAMETER(int, V1RangeOverrideDataOffset)
-	SHADER_PARAMETER(int, TotalNumInstances)
 	SHADER_PARAMETER(int, InterpCount)
 	SHADER_PARAMETER(float, OneOverInterpCount)
 	SHADER_PARAMETER(int, ParticleIdShift)
@@ -60,6 +60,7 @@ BEGIN_GLOBAL_SHADER_PARAMETER_STRUCT(FNiagaraRibbonUniformParameters, NIAGARAVER
 	SHADER_PARAMETER(int, InterpIdMask)
 	SHADER_PARAMETER(int, SliceVertexIdMask)
 	SHADER_PARAMETER(int, ShouldFlipNormalToView)
+	SHADER_PARAMETER(int, ShouldUseMultiRibbon)
 	SHADER_PARAMETER(int, U0DistributionMode)
 	SHADER_PARAMETER(int, U1DistributionMode)
 	SHADER_PARAMETER(FVector3f, SystemLWCTile)
@@ -71,17 +72,18 @@ typedef TUniformBufferRef<FNiagaraRibbonUniformParameters> FNiagaraRibbonUniform
 
 BEGIN_GLOBAL_SHADER_PARAMETER_STRUCT(FNiagaraRibbonVFLooseParameters, NIAGARAVERTEXFACTORIES_API)
 	SHADER_PARAMETER(uint32, NiagaraFloatDataStride)
-	SHADER_PARAMETER(uint32, SortedIndicesOffset)
 	SHADER_PARAMETER(uint32, FacingMode)
 	SHADER_PARAMETER(uint32, Shape)
 	SHADER_PARAMETER(uint32, NeedsPreciseMotionVectors)
-	SHADER_PARAMETER_SRV(Buffer<int>, SortedIndices)
-	SHADER_PARAMETER_SRV(Buffer<float4>, TangentsAndDistances)
+	SHADER_PARAMETER_SRV(Buffer<uint>, SortedIndices)
+	SHADER_PARAMETER_SRV(Buffer<float>, TangentsAndDistances)
 	SHADER_PARAMETER_SRV(Buffer<uint>, MultiRibbonIndices)
-	SHADER_PARAMETER_SRV(Buffer<float>, PackedPerRibbonDataByIndex)
+	SHADER_PARAMETER_SRV(Buffer<uint>, PackedPerRibbonDataByIndex)
 	SHADER_PARAMETER_SRV(Buffer<float>, NiagaraParticleDataFloat)
 	SHADER_PARAMETER_SRV(Buffer<float>, NiagaraParticleDataHalf)
 	SHADER_PARAMETER_SRV(Buffer<float>, SliceVertexData)
+	SHADER_PARAMETER_SRV(Buffer<uint>, IndirectDrawOutput)
+	SHADER_PARAMETER(int32, IndirectDrawOutputOffset)
 END_GLOBAL_SHADER_PARAMETER_STRUCT()
 typedef TUniformBufferRef<FNiagaraRibbonVFLooseParameters> FNiagaraRibbonVFLooseParametersRef;
 
@@ -97,18 +99,12 @@ public:
 	/** Default constructor. */
 	FNiagaraRibbonVertexFactory(ENiagaraVertexFactoryType InType, ERHIFeatureLevel::Type InFeatureLevel)
 		: FNiagaraVertexFactoryBase(InType, InFeatureLevel)
-		, IndexBuffer(nullptr)
-		, FirstIndex(0)
-		, OutTriangleCount(0)
-		, DataSet(0)
+		, DataSet(nullptr)
 	{}
 
 	FNiagaraRibbonVertexFactory()
 		: FNiagaraVertexFactoryBase(NVFT_MAX, ERHIFeatureLevel::Num)
-		, IndexBuffer(nullptr)
-		, FirstIndex(0)
-		, OutTriangleCount(0)
-		, DataSet(0)
+		, DataSet(nullptr)
 	{}
 
 	/**
@@ -150,91 +146,6 @@ public:
 	*/
 	void SetDynamicParameterBuffer(const FVertexBuffer* InDynamicParameterBuffer, int32 ParameterIndex, uint32 StreamOffset, uint32 Stride);
 
-	void SetSortedIndices(const FBufferRHIRef& InSortedIndicesBuffer, const FShaderResourceViewRHIRef& InSortedIndicesSRV, uint32 InSortedIndicesOffset)
-	{
-		SortedIndicesBuffer = InSortedIndicesBuffer;
-		SortedIndicesSRV = InSortedIndicesSRV;
-		SortedIndicesOffset = InSortedIndicesOffset;
-	}
-
-	void SetTangentAndDistances(const FBufferRHIRef& InTangentAndDistancesBuffer, const FShaderResourceViewRHIRef& InTangentAndDistancesSRV)
-	{
-		TangentAndDistancesBuffer = InTangentAndDistancesBuffer;
-		TangentAndDistancesSRV = InTangentAndDistancesSRV;
-	}
-
-
-	void SetMultiRibbonIndicesSRV(const FBufferRHIRef& InMultiRibbonIndicesBuffer, const FShaderResourceViewRHIRef& InMultiRibbonIndicesSRV)
-	{
-		MultiRibbonIndicesBuffer = InMultiRibbonIndicesBuffer;
-		MultiRibbonIndicesSRV = InMultiRibbonIndicesSRV;
-	}
-
-	void SetPackedPerRibbonDataByIndexSRV(const FBufferRHIRef& InPackedPerRibbonDataByIndexBuffer, const FShaderResourceViewRHIRef& InPackedPerRibbonDataByIndexSRV)
-	{
-		PackedPerRibbonDataByIndexBuffer = InPackedPerRibbonDataByIndexBuffer;
-		PackedPerRibbonDataByIndexSRV = InPackedPerRibbonDataByIndexSRV;
-	}
-
-	void SetSliceVertexDataSRV(const FBufferRHIRef& InSliceVertexDataBuffer, const FShaderResourceViewRHIRef& InSliceVertexDataSRV)
-	{
-		SliceVertexDataBuffer = InSliceVertexDataBuffer;
-		SliceVertexDataSRV = InSliceVertexDataSRV;
-	}
-
-	void SetFacingMode(uint32 InFacingMode)
-	{
-		FacingMode = InFacingMode;
-	}
-
-	FORCEINLINE FRHIShaderResourceView* GetSortedIndicesSRV()
-	{
-		return SortedIndicesSRV;
-	}
-
-	FORCEINLINE int32 GetSortedIndicesOffset() const
-	{
-		return SortedIndicesOffset;
-	}
-
-	FORCEINLINE FRHIShaderResourceView* GetTangentAndDistancesSRV() const
-	{
-		return TangentAndDistancesSRV;
-	}
-
-	FORCEINLINE FRHIShaderResourceView* GetMultiRibbonIndicesSRV() const
-	{
-		return MultiRibbonIndicesSRV;
-	}
-
-	FORCEINLINE FRHIShaderResourceView* GetPackedPerRibbonDataByIndexSRV() const
-	{
-		return PackedPerRibbonDataByIndexSRV;
-	}
-
-	FORCEINLINE FRHIShaderResourceView* GetSliceVertexDataSRV() const
-	{
-		return SliceVertexDataSRV;
-	}
-
-	FORCEINLINE int32 GetFacingMode() const
-	{
-		return FacingMode;
-	}
-	FIndexBuffer*& GetIndexBuffer()
-	{
-		return IndexBuffer;
-	}
-
-	uint32& GetFirstIndex()
-	{
-		return FirstIndex;
-	}
-
-	int32& GetOutTriangleCount()
-	{
-		return OutTriangleCount;
-	}
 
 	FUniformBufferRHIRef LooseParameterUniformBuffer;
 
@@ -243,25 +154,6 @@ private:
 	/** Uniform buffer with beam/trail parameters. */
 	FNiagaraRibbonUniformBufferRef NiagaraRibbonUniformBuffer;
 
-	/** Used to hold the index buffer allocation information when we call GDME more than once per frame. */
-	FIndexBuffer* IndexBuffer;
-	uint32 FirstIndex;
-	int32 OutTriangleCount;
 
 	const FNiagaraDataSet *DataSet;
-
-	FBufferRHIRef SortedIndicesBuffer;
-	FBufferRHIRef TangentAndDistancesBuffer;
-	FBufferRHIRef MultiRibbonIndicesBuffer;
-	FBufferRHIRef PackedPerRibbonDataByIndexBuffer;
-	FBufferRHIRef SliceVertexDataBuffer;
-
-	FShaderResourceViewRHIRef SortedIndicesSRV;
-	FShaderResourceViewRHIRef TangentAndDistancesSRV;
-	FShaderResourceViewRHIRef MultiRibbonIndicesSRV;
-	FShaderResourceViewRHIRef PackedPerRibbonDataByIndexSRV;
-	FShaderResourceViewRHIRef SliceVertexDataSRV;
-
-	uint32 SortedIndicesOffset;
-	int32 FacingMode;
 };
\ No newline at end of file
