- Fix SystemRendererIndex vs EmitterIndex

#rnx

diff --git a/Engine/Plugins/FX/Niagara/Source/Niagara/Private/NiagaraSystemRenderData.cpp b/Engine/Plugins/FX/Niagara/Source/Niagara/Private/NiagaraSystemRenderData.cpp
--- a/Engine/Plugins/FX/Niagara/Source/Niagara/Private/NiagaraSystemRenderData.cpp
+++ b/Engine/Plugins/FX/Niagara/Source/Niagara/Private/NiagaraSystemRenderData.cpp
@@ -347,9 +347,9 @@ void FNiagaraSystemRenderData::OnSystemComplete(const FNiagaraSystemInstance& Sy
 
 	for (const FNiagaraRendererExecutionIndex& ExecIdx : System->GetRendererCompletionOrder())
 	{
-		if (EmitterRenderers_GT.IsValidIndex(ExecIdx.SystemRendererIndex) && ensure(EmitterInstances.IsValidIndex(ExecIdx.SystemRendererIndex)) )
+		if (EmitterRenderers_GT.IsValidIndex(ExecIdx.SystemRendererIndex) && ensure(EmitterInstances.IsValidIndex(ExecIdx.EmitterIndex)) )
 		{
-			const FNiagaraEmitterInstance& EmitterInstance = EmitterInstances[ExecIdx.SystemRendererIndex].Get();
+			const FNiagaraEmitterInstance& EmitterInstance = EmitterInstances[ExecIdx.EmitterIndex].Get();
 			FVersionedNiagaraEmitterData* EmitterData = EmitterInstance.GetCachedEmitterData();
 			FNiagaraRenderer* EmitterRenderer = EmitterRenderers_GT[ExecIdx.SystemRendererIndex];
 			if (EmitterData && EmitterRenderer)
