- Disable early data release as it may be causing issues

diff --git a/Engine/Plugins/FX/Niagara/Source/Niagara/Private/NiagaraDataSet.cpp b/Engine/Plugins/FX/Niagara/Source/Niagara/Private/NiagaraDataSet.cpp
--- a/Engine/Plugins/FX/Niagara/Source/Niagara/Private/NiagaraDataSet.cpp
+++ b/Engine/Plugins/FX/Niagara/Source/Niagara/Private/NiagaraDataSet.cpp
@@ -93,7 +93,7 @@ static FAutoConsoleVariableRef CVarNiagaraGPUDataWarningSize(
 	ECVF_Scalability
 );
 
-static bool GNiagaraReleaseBuffersOnReset = true;
+static bool GNiagaraReleaseBuffersOnReset = false;
 static FAutoConsoleVariableRef CVarNiagaraReleaseBuffersOnReset(
 	TEXT("fx.NiagaraReleaseBuffersOnReset"),
 	GNiagaraReleaseBuffersOnReset,
