Moved emitter property

diff --git a/Engine/Plugins/FX/Niagara/Source/Niagara/Classes/NiagaraEmitter.h b/Engine/Plugins/FX/Niagara/Source/Niagara/Classes/NiagaraEmitter.h
--- a/Engine/Plugins/FX/Niagara/Source/Niagara/Classes/NiagaraEmitter.h
+++ b/Engine/Plugins/FX/Niagara/Source/Niagara/Classes/NiagaraEmitter.h
@@ -262,6 +262,10 @@ struct NIAGARA_API FVersionedNiagaraEmitterData
 	UPROPERTY(EditAnywhere, BlueprintReadWrite, Category = "Emitter")
 	bool bDeterminism = false;
 
+	/** An emitter-based seed for the deterministic random number generator. */
+	UPROPERTY(EditAnywhere, BlueprintReadWrite, Category = "Emitter", meta = (EditCondition = "bDeterminism", EditConditionHides="true"))
+	int32 RandomSeed = 0;
+
 	/** When enabled, this will spawn using interpolated parameter values and perform a partial update at spawn time. This adds significant additional cost for spawning but will produce much smoother spawning for high spawn rates, erratic frame rates and fast moving emitters. */
 	UPROPERTY(EditAnywhere, BlueprintReadWrite, Category = "Emitter")
 	uint32 bInterpolatedSpawning : 1;
@@ -275,10 +279,6 @@ struct NIAGARA_API FVersionedNiagaraEmitterData
 	uint32 bGpuAlwaysRunParticleUpdateScript : 1;
 #endif
 
-	/** An emitter-based seed for the deterministic random number generator. */
-	UPROPERTY(EditAnywhere, BlueprintReadWrite, Category = "Emitter", meta = (EditCondition = "bDeterminism"))
-	int32 RandomSeed = 0;
-
 	UPROPERTY(EditAnywhere, BlueprintReadWrite, Category = "Emitter")
 	ENiagaraSimTarget SimTarget = ENiagaraSimTarget::CPUSim;
 
