- Fix for simulations not showing up when launching into a level

diff --git a/Engine/Plugins/FX/Niagara/Source/Niagara/Classes/NiagaraPlatformSet.h b/Engine/Plugins/FX/Niagara/Source/Niagara/Classes/NiagaraPlatformSet.h
--- a/Engine/Plugins/FX/Niagara/Source/Niagara/Classes/NiagaraPlatformSet.h
+++ b/Engine/Plugins/FX/Niagara/Source/Niagara/Classes/NiagaraPlatformSet.h
@@ -315,7 +315,6 @@
 	static FText GetQualityLevelText(int32 QualityLevel);
 	static FText GetQualityLevelMaskText(int32 QualityLevelMask);
 
-	static void OnQualityLevelChanged(IConsoleVariable* Variable);
 	static int32 GetQualityLevel();
 	static int32 GetMinQualityLevel();
 	static int32 GetMaxQualityLevel();
diff --git a/Engine/Plugins/FX/Niagara/Source/Niagara/Private/NiagaraPlatformSet.cpp b/Engine/Plugins/FX/Niagara/Source/Niagara/Private/NiagaraPlatformSet.cpp
--- a/Engine/Plugins/FX/Niagara/Source/Niagara/Private/NiagaraPlatformSet.cpp
+++ b/Engine/Plugins/FX/Niagara/Source/Niagara/Private/NiagaraPlatformSet.cpp
@@ -77,14 +77,29 @@ static FAutoConsoleCommand GCmdSetNiagaraQualityLevelOverride(
 		}
 	)
 ); 
+
 static FAutoConsoleVariableRef CVarNiagaraQualityLevel(
 	NiagaraQualityLevelName,
 	GNiagaraQualityLevel,
 	TEXT("The quality level for Niagara Effects. \n"),
-	FConsoleVariableDelegate::CreateStatic(&FNiagaraPlatformSet::OnQualityLevelChanged),
 	ECVF_Scalability
 );
 
+// Sync function for checking quality level hasn't changed
+// This must be done at the end of all console variable changes otherwise we can encounter multiple changes to the Q level resulting in some systems being left deactivated
+FAutoConsoleVariableSink CVarNiagaraQualityLevelSync(
+	FConsoleCommandDelegate::CreateLambda(
+		[]()
+		{
+			const int32 CurrentLevel = FNiagaraPlatformSet::GetQualityLevel();
+			if (CurrentLevel != GNiagaraQualityLevel)
+			{
+				SetGNiagaraQualityLevel(GNiagaraQualityLevel);
+			}
+		}
+	)
+);
+
 /** Minimum quality level available for a platform. Defaults to -1, indicating there is no minimum. */
 static const int32 DefaultMinQualityLevel = INDEX_NONE;
 const TCHAR* NiagaraMinQualityLevelName = TEXT("fx.Niagara.QualityLevel.Min");
@@ -261,17 +276,6 @@ int32 FNiagaraPlatformSet::GetFullQualityLevelMask(int32 NumQualityLevels)
 	return QualityLevelMask;
 }
 
-void FNiagaraPlatformSet::OnQualityLevelChanged(IConsoleVariable* Variable)
-{
-	int32 NewQualityLevel = Variable->GetInt();
-	int32 CurrentLevel = GetQualityLevel();
-
-	if (CurrentLevel != NewQualityLevel)
-	{
-		SetGNiagaraQualityLevel(NewQualityLevel);
-	}
-}
-
 uint32 FNiagaraPlatformSet::LastDirtiedFrame = 0;
 #if WITH_EDITOR
 TMap<const UDeviceProfile*, int32> FNiagaraPlatformSet::CachedQLMasksPerDeviceProfile;
