Niagara - Fix an issue in the niagara component customization where the category name wasn't being passed along to the custom node builder which prevents the details panel code from handling that row correctly for filtering.

diff --git a/Engine/Plugins/FX/Niagara/Source/NiagaraEditor/Private/Customizations/NiagaraComponentDetails.cpp b/Engine/Plugins/FX/Niagara/Source/NiagaraEditor/Private/Customizations/NiagaraComponentDetails.cpp
--- a/Engine/Plugins/FX/Niagara/Source/NiagaraEditor/Private/Customizations/NiagaraComponentDetails.cpp
+++ b/Engine/Plugins/FX/Niagara/Source/NiagaraEditor/Private/Customizations/NiagaraComponentDetails.cpp
@@ -264,9 +264,10 @@ private:
 class FNiagaraComponentNodeBuilder : public IDetailCustomNodeBuilder, public TSharedFromThis<FNiagaraComponentNodeBuilder>
 {
 public:
-	FNiagaraComponentNodeBuilder(UNiagaraComponent* InComponent, TArray<TSharedPtr<IPropertyHandle>> InOverridePropertyHandles) 
+	FNiagaraComponentNodeBuilder(UNiagaraComponent* InComponent, TArray<TSharedPtr<IPropertyHandle>> InOverridePropertyHandles, FName InCustomBuilderRowName) 
 	{
 		OverridePropertyHandles = InOverridePropertyHandles;
+		CustomBuilderRowName = InCustomBuilderRowName;
 
 		Component = InComponent;
 		bDelegatesInitialized = false;
@@ -296,8 +297,7 @@ public:
 	virtual bool InitiallyCollapsed() const { return false; }
 	virtual FName GetName() const  override
 	{
-		static const FName NiagaraComponentNodeBuilder("FNiagaraComponentNodeBuilder");
-		return NiagaraComponentNodeBuilder;
+		return CustomBuilderRowName;
 	}
 
 	virtual void GenerateChildContent(IDetailChildrenBuilder& ChildrenBuilder) override
@@ -525,6 +525,7 @@ private:
 	FSimpleDelegate OnRebuildChildren;
 	TArray<TSharedPtr<FNiagaraParameterProxy>> ParameterProxies;
 	TMap<FName, TWeakPtr<FStructOnScope>> ParameterNameToDisplayStruct;
+	FName CustomBuilderRowName;
 };
 
 TSharedRef<IDetailCustomization> FNiagaraComponentDetails::MakeInstance()
@@ -640,7 +641,7 @@ void FNiagaraComponentDetails::CustomizeDetails(IDetailLayoutBuilder& DetailBuil
 		FGameDelegates::Get().GetEndPlayMapDelegate().AddRaw(this, &FNiagaraComponentDetails::OnPiEEnd);
 			
 		IDetailCategoryBuilder& InputParamCategory = DetailBuilder.EditCategory(ParamCategoryName, LOCTEXT("ParamCategoryName", "Override Parameters"), ECategoryPriority::Important);
-		InputParamCategory.AddCustomBuilder(MakeShared<FNiagaraComponentNodeBuilder>(Component.Get(), PropertyHandles));
+		InputParamCategory.AddCustomBuilder(MakeShared<FNiagaraComponentNodeBuilder>(Component.Get(), PropertyHandles, ParamCategoryName));
 	}
 	else if (ObjectsCustomized.Num() > 1)
 	{
