Allow script dependencies to be filtered by context

diff --git a/Engine/Plugins/FX/Niagara/Source/Niagara/Classes/NiagaraScript.h b/Engine/Plugins/FX/Niagara/Source/Niagara/Classes/NiagaraScript.h
--- a/Engine/Plugins/FX/Niagara/Source/Niagara/Classes/NiagaraScript.h
+++ b/Engine/Plugins/FX/Niagara/Source/Niagara/Classes/NiagaraScript.h
@@ -89,6 +89,22 @@ enum class ENiagaraScriptTemplateSpecification : uint8
 	Behavior UMETA(DisplayName = "Behavior Example")
 };
 
+/** Defines different usages for a niagara script's module dependecies. */
+UENUM()
+enum class ENiagaraModuleDependencyUsage : uint8
+{
+	/** Default entry to catch invalid usages */
+	None UMETA(Hidden),
+	/** Evaluate when the script is called during the spawn phase. */
+	Spawn,
+	/** Evaluate when the script is called during the update phase. */
+	Update,
+	/** Evaluate when the script is called in an event context. */
+	Event,
+	/** Evaluate when the script is called in a simulation stage. */
+	SimulationStage
+};
+
 USTRUCT()
 struct FNiagaraModuleDependency
 {
@@ -113,6 +129,10 @@ public:
 	// '1.2-2.0' requires any version between 1.2 and 2.0
 	UPROPERTY(AssetRegistrySearchable, EditAnywhere, Category = Script)
 	FString RequiredVersion;
+
+	/** This property can limit where the dependency is evaluated. By default, the dependency is enforced in all script usages */
+	UPROPERTY(AssetRegistrySearchable, EditAnywhere, Category = Script, meta = (Bitmask, BitmaskEnum = "/Script/Niagara.ENiagaraModuleDependencyUsage"))
+	int32 OnlyEvaluateInScriptUsage;
 	
 	/** Detailed description of the dependency */
 	UPROPERTY(AssetRegistrySearchable, EditAnywhere, Category = Script, meta = (MultiLine = true))
@@ -122,6 +142,10 @@ public:
 	{
 		Type = ENiagaraModuleDependencyType::PreDependency;
 		ScriptConstraint = ENiagaraModuleDependencyScriptConstraint::SameScript;
+		OnlyEvaluateInScriptUsage = 1 << static_cast<int32>(ENiagaraModuleDependencyUsage::Spawn) |
+									1 << static_cast<int32>(ENiagaraModuleDependencyUsage::Update) |
+									1 << static_cast<int32>(ENiagaraModuleDependencyUsage::Event) |
+									1 << static_cast<int32>(ENiagaraModuleDependencyUsage::SimulationStage);
 	}
 
 #if WITH_EDITOR
diff --git a/Engine/Plugins/FX/Niagara/Source/NiagaraEditor/Private/ViewModels/Stack/NiagaraStackModuleItem.cpp b/Engine/Plugins/FX/Niagara/Source/NiagaraEditor/Private/ViewModels/Stack/NiagaraStackModuleItem.cpp
--- a/Engine/Plugins/FX/Niagara/Source/NiagaraEditor/Private/ViewModels/Stack/NiagaraStackModuleItem.cpp
+++ b/Engine/Plugins/FX/Niagara/Source/NiagaraEditor/Private/ViewModels/Stack/NiagaraStackModuleItem.cpp
@@ -530,6 +530,33 @@ namespace NiagaraStackModuleItemIssues {
 		UNiagaraNodeFunctionCall* MovedNode;
 		FNiagaraStackGraphUtilities::MoveModule(*SourceScript, *ModuleToMove, *TargetSystem, DependentEmitterHandleId, TargetUsage, TargetUsageId, TargetMoveIndex, false, MovedNode);
 	}
+	
+	ENiagaraModuleDependencyUsage ConvertScriptUsageToDependencyUsage(ENiagaraScriptUsage ScriptUsage)
+	{
+		if (ScriptUsage == ENiagaraScriptUsage::ParticleEventScript)
+		{
+			return ENiagaraModuleDependencyUsage::Event;
+		}
+		if (ScriptUsage == ENiagaraScriptUsage::ParticleSimulationStageScript)
+		{
+			return ENiagaraModuleDependencyUsage::SimulationStage;
+		}
+		if (ScriptUsage == ENiagaraScriptUsage::EmitterSpawnScript || ScriptUsage == ENiagaraScriptUsage::SystemSpawnScript || ScriptUsage == ENiagaraScriptUsage::ParticleSpawnScriptInterpolated || ScriptUsage == ENiagaraScriptUsage::ParticleSpawnScript)
+		{
+			return ENiagaraModuleDependencyUsage::Spawn;
+		}
+		if (ScriptUsage == ENiagaraScriptUsage::EmitterUpdateScript || ScriptUsage == ENiagaraScriptUsage::SystemUpdateScript || ScriptUsage == ENiagaraScriptUsage::ParticleUpdateScript)
+		{
+			return ENiagaraModuleDependencyUsage::Update;
+		}
+		return ENiagaraModuleDependencyUsage::None;
+	}
+
+	bool IsUsageAllowed(ENiagaraScriptUsage ModuleUsage, int32 AllowedUsageBitmask)
+	{
+		ENiagaraModuleDependencyUsage Usage = ConvertScriptUsageToDependencyUsage(ModuleUsage);
+		return AllowedUsageBitmask & (1 << static_cast<int32>(Usage));
+	}
 
 	void GenerateFixesForReorderingModules(
 		TSharedRef<FNiagaraSystemViewModel> DependentSystemViewModel,
@@ -694,7 +721,7 @@ namespace NiagaraStackModuleItemIssues {
 		UNiagaraScript& SourceScript,
 		UNiagaraNodeFunctionCall& SourceModuleNode,
 		FString SourceStackEditorDataKey,
-		UNiagaraNodeOutput& SourceOutputNode,
+		const UNiagaraNodeOutput& SourceOutputNode,
 		const TArray<FNiagaraStackModuleData>& SourceStackModuleData,
 		TArray<UNiagaraStackEntry::FStackIssue>& NewIssues)
 	{
@@ -710,6 +737,10 @@ namespace NiagaraStackModuleItemIssues {
 		{
 			for (const FNiagaraModuleDependency& SourceRequiredDependency : SourceModuleNode.GetScriptData()->RequiredDependencies)
 			{
+				if (!IsUsageAllowed(SourceOutputNode.GetUsage(), SourceRequiredDependency.OnlyEvaluateInScriptUsage))
+				{
+					continue;
+				}
 				TArray<int32> DependencyProviderIndices;
 				for (int32 StackModuleDataIndex = 0; StackModuleDataIndex < SourceStackModuleData.Num(); StackModuleDataIndex++)
 				{
