Make niagara component a blueprint type

diff --git a/Engine/Plugins/FX/Niagara/Source/Niagara/Public/NiagaraComponent.h b/Engine/Plugins/FX/Niagara/Source/Niagara/Public/NiagaraComponent.h
--- a/Engine/Plugins/FX/Niagara/Source/Niagara/Public/NiagaraComponent.h
+++ b/Engine/Plugins/FX/Niagara/Source/Niagara/Public/NiagaraComponent.h
@@ -32,7 +32,7 @@ DECLARE_DYNAMIC_MULTICAST_DELEGATE_OneParam(FOnNiagaraSystemFinished, class UNia
 * @see ANiagaraActor
 * @see UNiagaraSystem
 */
-UCLASS(ClassGroup = (Rendering, Common), hidecategories = Object, hidecategories = Physics, hidecategories = Collision, showcategories = Trigger, editinlinenew, meta = (BlueprintSpawnableComponent, DisplayName = "Niagara Particle System Component"))
+UCLASS(ClassGroup = (Rendering, Common), Blueprintable, hidecategories = Object, hidecategories = Physics, hidecategories = Collision, showcategories = Trigger, editinlinenew, meta = (BlueprintSpawnableComponent, DisplayName = "Niagara Particle System Component"))
 class NIAGARA_API UNiagaraComponent : public UFXSystemComponent
 {
 	friend struct FNiagaraScalabilityManager;
diff --git a/Engine/Plugins/FX/Niagara/Source/Niagara/Public/NiagaraCullProxyComponent.h b/Engine/Plugins/FX/Niagara/Source/Niagara/Public/NiagaraCullProxyComponent.h
--- a/Engine/Plugins/FX/Niagara/Source/Niagara/Public/NiagaraCullProxyComponent.h
+++ b/Engine/Plugins/FX/Niagara/Source/Niagara/Public/NiagaraCullProxyComponent.h
@@ -19,7 +19,7 @@ struct FNiagaraCulledComponentInfo
 /**
 A specialization of UNiagaraComponent that can act as a proxy for many other NiagaraComponents that have been culled by scalability. 
 */
-UCLASS()
+UCLASS(NotBlueprintable)
 class UNiagaraCullProxyComponent : public UNiagaraComponent
 {
 	GENERATED_UCLASS_BODY()
