- Fix crash when dragging parameters

#rnx

diff --git a/Engine/Plugins/FX/Niagara/Source/NiagaraEditor/Private/ViewModels/NiagaraParameterPanelViewModel.cpp b/Engine/Plugins/FX/Niagara/Source/NiagaraEditor/Private/ViewModels/NiagaraParameterPanelViewModel.cpp
--- a/Engine/Plugins/FX/Niagara/Source/NiagaraEditor/Private/ViewModels/NiagaraParameterPanelViewModel.cpp
+++ b/Engine/Plugins/FX/Niagara/Source/NiagaraEditor/Private/ViewModels/NiagaraParameterPanelViewModel.cpp
@@ -1825,7 +1825,11 @@ const TArray<FNiagaraGraphParameterReference> FNiagaraSystemToolkitParameterPane
 			if (Builder.Histories.Num() != 1)
 			{
 				// We should only have traversed one emitter (have not visited more than one NiagaraNodeEmitter.)
-				ensureMsgf(false, TEXT("Encountered more than one parameter map history when collecting graph parameter reference collections for system parameter panel view model!"));
+				ensureMsgf(false, TEXT("Encountered '%d' parameter map history when collecting parameters for system parameter panel view model, we expect only 1!"), Builder.Histories.Num());
+			}
+			if (Builder.Histories.Num() == 0)
+			{
+				continue;
 			}
 
 			const TArray<FName>& CustomIterationSourceNamespaces = Builder.Histories[0].IterationNamespaceOverridesEncountered;
