openvdb cache emitter
integration with toolset/existing modules

diff --git a/Engine/Plugins/FX/Niagara/Source/Niagara/Classes/NiagaraDataInterfaceVolumeCache.h b/Engine/Plugins/FX/Niagara/Source/Niagara/Classes/NiagaraDataInterfaceVolumeCache.h
--- a/Engine/Plugins/FX/Niagara/Source/Niagara/Classes/NiagaraDataInterfaceVolumeCache.h
+++ b/Engine/Plugins/FX/Niagara/Source/Niagara/Classes/NiagaraDataInterfaceVolumeCache.h
@@ -53,8 +53,9 @@ public:
 	virtual void SetShaderParameters(const FNiagaraDataInterfaceSetShaderParametersContext& Context) const override;
 
 	void SetFrame(FVectorVMExternalFunctionContext& Context);
-	void ReadFile(FVectorVMExternalFunctionContext& Context);
-	
+	void ReadFile(FVectorVMExternalFunctionContext& Context);	
+	void GetNumCells(FVectorVMExternalFunctionContext& Context);
+
 	FString GetAssetPath(FString PathFormat, int32 FrameIndex) const;
 
 protected:
@@ -66,6 +67,7 @@ protected:
 
 	static const FName SetFrameName;
 	static const FName ReadFileName;
+	static const FName GetNumCellsName;
 	static const FName SampleCurrentFrameValueName;
 	static const FName GetCurrentFrameValue;
 	static const FName GetCurrentFrameNumCells;
diff --git a/Engine/Plugins/FX/Niagara/Source/Niagara/Private/NiagaraDataInterfaceVolumeCache.cpp b/Engine/Plugins/FX/Niagara/Source/Niagara/Private/NiagaraDataInterfaceVolumeCache.cpp
--- a/Engine/Plugins/FX/Niagara/Source/Niagara/Private/NiagaraDataInterfaceVolumeCache.cpp
+++ b/Engine/Plugins/FX/Niagara/Source/Niagara/Private/NiagaraDataInterfaceVolumeCache.cpp
@@ -21,6 +21,7 @@
 const TCHAR* UNiagaraDataInterfaceVolumeCache::TemplateShaderFilePath = TEXT("/Plugin/FX/Niagara/Private/NiagaraDataInterfaceVolumeCache.ush");
 const FName UNiagaraDataInterfaceVolumeCache::SetFrameName("SetFrame");
 const FName UNiagaraDataInterfaceVolumeCache::ReadFileName("ReadFile");
+const FName UNiagaraDataInterfaceVolumeCache::GetNumCellsName("GetNumCells");
 const FName UNiagaraDataInterfaceVolumeCache::SampleCurrentFrameValueName(TEXT("SampleCurrentFrameValue"));
 const FName UNiagaraDataInterfaceVolumeCache::GetCurrentFrameValue(TEXT("GetCurrentFrameValue"));
 const FName UNiagaraDataInterfaceVolumeCache::GetCurrentFrameNumCells(TEXT("GetCurrentFrameNumCells"));
@@ -40,6 +41,8 @@ struct FVolumeCacheInstanceData_GameThread
 	int CurrFrame = 0;
 	int PrevFrame = -1;
 
+	FIntVector NumCells = FIntVector::ZeroValue;
+
 	bool ReadFile = false;		
 };
 
@@ -109,6 +112,11 @@ bool UNiagaraDataInterfaceVolumeCache::InitPerInstanceData(void* PerInstanceData
 	FVolumeCacheInstanceData_GameThread* InstanceData = new (PerInstanceData) FVolumeCacheInstanceData_GameThread();
 	SystemInstancesToProxyData_GT.Emplace(SystemInstance->GetId(), InstanceData);	
 
+	if (VolumeCache != nullptr)
+	{
+		InstanceData->NumCells = VolumeCache->Resolution;
+	}
+
 	// Push Updates to Proxy.
 	FNiagaraDataInterfaceVolumeCacheProxy* TheProxy = GetProxyAs<FNiagaraDataInterfaceVolumeCacheProxy>();
 	ENQUEUE_RENDER_COMMAND(FUpdateData)(
@@ -175,6 +183,19 @@ void UNiagaraDataInterfaceVolumeCache::GetFunctions(TArray<FNiagaraFunctionSigna
 		OutFunctions.Add(Sig);
 	}
 
+	{
+		FNiagaraFunctionSignature Sig;
+		Sig.Name = GetNumCellsName;
+		Sig.Inputs.Add(FNiagaraVariable(FNiagaraTypeDefinition(GetClass()), TEXT("Grid")));
+		Sig.Outputs.Add(FNiagaraVariable(FNiagaraTypeDefinition::GetIntDef(), TEXT("NumCellsX")));
+		Sig.Outputs.Add(FNiagaraVariable(FNiagaraTypeDefinition::GetIntDef(), TEXT("NumCellsY")));
+		Sig.Outputs.Add(FNiagaraVariable(FNiagaraTypeDefinition::GetIntDef(), TEXT("NumCellsZ")));
+
+		Sig.bMemberFunction = true;
+		Sig.bRequiresContext = false;
+		OutFunctions.Add(Sig);
+	}
+
 	{
 		FNiagaraFunctionSignature Sig;
 		Sig.Name = SampleCurrentFrameValueName;
@@ -251,6 +272,26 @@ void UNiagaraDataInterfaceVolumeCache::ReadFile(FVectorVMExternalFunctionContext
 	*OutSuccess.GetDestAndAdvance() = true;
 }
 
+void UNiagaraDataInterfaceVolumeCache::GetNumCells(FVectorVMExternalFunctionContext& Context)
+{
+	VectorVM::FUserPtrHandler<FVolumeCacheInstanceData_GameThread> InstData(Context);
+
+	FNDIOutputParam<int32> NumCellsX(Context);
+	FNDIOutputParam<int32> NumCellsY(Context);
+	FNDIOutputParam<int32> NumCellsZ(Context);
+
+	int32 TmpNumCellsX = InstData->NumCells.X;
+	int32 TmpNumCellsY = InstData->NumCells.Y;
+	int32 TmpNumCellsZ = InstData->NumCells.Z;
+
+	for (int32 InstanceIdx = 0; InstanceIdx < Context.GetNumInstances(); ++InstanceIdx)
+	{
+		NumCellsX.SetAndAdvance(TmpNumCellsX);
+		NumCellsY.SetAndAdvance(TmpNumCellsY);
+		NumCellsZ.SetAndAdvance(TmpNumCellsZ);
+	}
+}
+
 FString UNiagaraDataInterfaceVolumeCache::GetAssetPath(FString PathFormat, int32 FrameIndex) const
 {
 	UNiagaraSystem* NiagaraSystem = GetTypedOuter<UNiagaraSystem>();
@@ -268,6 +309,7 @@ FString UNiagaraDataInterfaceVolumeCache::GetAssetPath(FString PathFormat, int32
 
 DEFINE_NDI_DIRECT_FUNC_BINDER(UNiagaraDataInterfaceVolumeCache, SetFrame);
 DEFINE_NDI_DIRECT_FUNC_BINDER(UNiagaraDataInterfaceVolumeCache, ReadFile);
+DEFINE_NDI_DIRECT_FUNC_BINDER(UNiagaraDataInterfaceVolumeCache, GetNumCells);
 void UNiagaraDataInterfaceVolumeCache::GetVMExternalFunction(const FVMExternalFunctionBindingInfo& BindingInfo, void* InstanceData, FVMExternalFunction& OutFunc)
 {
 	if (BindingInfo.Name == SetFrameName)
@@ -280,6 +322,11 @@ void UNiagaraDataInterfaceVolumeCache::GetVMExternalFunction(const FVMExternalFu
 		check(BindingInfo.GetNumInputs() == 2 && BindingInfo.GetNumOutputs() == 1);
 		NDI_FUNC_BINDER(UNiagaraDataInterfaceVolumeCache, ReadFile)::Bind(this, OutFunc);
 	}
+	else if (BindingInfo.Name == GetNumCellsName)
+	{
+		check(BindingInfo.GetNumInputs() == 1 && BindingInfo.GetNumOutputs() == 3);
+		NDI_FUNC_BINDER(UNiagaraDataInterfaceVolumeCache, GetNumCells)::Bind(this, OutFunc);
+	}
 }
 
 bool UNiagaraDataInterfaceVolumeCache::PerInstanceTick(void* PerInstanceData, FNiagaraSystemInstance* SystemInstance, float DeltaSeconds)
@@ -289,7 +336,7 @@ bool UNiagaraDataInterfaceVolumeCache::PerInstanceTick(void* PerInstanceData, FN
 	// we can run into the case where depending on the ordering of DI initialization, we might have not been able to grab the other grid's InstanceData
 	// in InitPerInstanceData.  If this is the case, we ensure it is correct here.
 	if (InstanceData && InstanceData->ReadFile && VolumeCache != nullptr)
-	{
+	{	
 		FNiagaraDataInterfaceVolumeCacheProxy* TextureProxy = GetProxyAs<FNiagaraDataInterfaceVolumeCacheProxy>();
 
 		//EPixelFormat Format = PF_A32B32G32R32F;
@@ -302,7 +349,7 @@ bool UNiagaraDataInterfaceVolumeCache::PerInstanceTick(void* PerInstanceData, FN
 			UE_LOG(LogNiagara, Warning, TEXT("Cache Read failed: %s"), *VolumeCache->GetName());
 			return false;
 		}
-
+		
 		ENQUEUE_RENDER_COMMAND(FVolumeCacheFillTexture)(
 			[RT_Frame = InstanceData->CurrFrame, RT_VolumeCacheData = VolumeCache->GetData(), Format, TextureProxy, SystemID = SystemInstance->GetId()](FRHICommandListImmediate& RHICmdList)
 		{
diff --git a/Engine/Plugins/FX/Niagara/Source/NiagaraEditor/Private/NiagaraBakerRendererOutputVolumeTexture.cpp b/Engine/Plugins/FX/Niagara/Source/NiagaraEditor/Private/NiagaraBakerRendererOutputVolumeTexture.cpp
--- a/Engine/Plugins/FX/Niagara/Source/NiagaraEditor/Private/NiagaraBakerRendererOutputVolumeTexture.cpp
+++ b/Engine/Plugins/FX/Niagara/Source/NiagaraEditor/Private/NiagaraBakerRendererOutputVolumeTexture.cpp
@@ -594,6 +594,18 @@ void FNiagaraBakerRendererOutputVolumeTexture::BakeFrame(UNiagaraBakerOutput* In
 			}
 
 			const FString FilePath = OutputVolumeTexture->GetExportPath(OutputVolumeTexture->FramesExportPathFormat, FrameIndex);
+			const FString FileDirectory = FString(FPathViews::GetPath(FilePath));
+
+			IPlatformFile& PlatformFile = FPlatformFileManager::Get().GetPlatformFile();
+			if (!PlatformFile.DirectoryExists(*FileDirectory))
+			{
+				if (!PlatformFile.CreateDirectoryTree(*FileDirectory))
+				{
+					UE_LOG(LogNiagaraBaker, Warning, TEXT("Cannot Create Directory : %s"), *FileDirectory);
+					return;
+				}
+			}
+
 			FNiagaraBakerRenderer::ExportVolume(FilePath, TextureSize, TextureData);
 		}
 	}
