Update NiagaraDataInterfaceUtilities ForEach functions to include a version for scene instance

diff --git a/Engine/Plugins/FX/Niagara/Source/Niagara/Private/DataInterface/NiagaraDataInterfaceUtilities.cpp b/Engine/Plugins/FX/Niagara/Source/Niagara/Private/DataInterface/NiagaraDataInterfaceUtilities.cpp
--- a/Engine/Plugins/FX/Niagara/Source/Niagara/Private/DataInterface/NiagaraDataInterfaceUtilities.cpp
+++ b/Engine/Plugins/FX/Niagara/Source/Niagara/Private/DataInterface/NiagaraDataInterfaceUtilities.cpp
@@ -6,7 +6,10 @@
 #include "NiagaraDataInterface.h"
 #include "NiagaraSystem.h"
 
-void FNiagaraDataInterfaceUtilities::ForEachVMFunctionEquals(class UNiagaraDataInterface* DataInterface, class UNiagaraSystem* NiagaraSystem, class UNiagaraComponent* Component, TFunction<bool(const FVMExternalFunctionBindingInfo&)> Action)
+namespace FNiagaraDataInterfaceUtilities
+{
+
+void ForEachVMFunctionEqualsImpl(class UNiagaraDataInterface* DataInterface, class UNiagaraSystem* NiagaraSystem, const FNiagaraParameterStore& ParameterStore, TFunction<bool(const FVMExternalFunctionBindingInfo&)> Action)
 {
 	if ( DataInterface == nullptr || NiagaraSystem == nullptr )
 	{
@@ -15,26 +18,12 @@ void FNiagaraDataInterfaceUtilities::ForEachVMFunctionEquals(class UNiagaraDataI
 
 	// Find all override parameter names
 	TArray<FName, TInlineAllocator<8>> OverrideParameterNames;
-	if ( Component )
-	{
-		for (const UNiagaraDataInterface* OverrideDI : Component->GetOverrideParameters().GetDataInterfaces())
-		{
-			if ( (OverrideDI != nullptr) && ((OverrideDI == DataInterface) || OverrideDI->Equals(DataInterface)) )
-			{
-				if ( const FNiagaraVariableBase* Variable = Component->GetOverrideParameters().FindVariable(OverrideDI) )
-				{
-					OverrideParameterNames.AddUnique(Variable->GetName());
-				}
-			}
-		}
-	}
-	else
 	{
-		for (const UNiagaraDataInterface* OverrideDI : NiagaraSystem->GetExposedParameters().GetDataInterfaces())
+		for (const UNiagaraDataInterface* OverrideDI : ParameterStore.GetDataInterfaces())
 		{
-			if ( (OverrideDI != nullptr) && ((OverrideDI == DataInterface) || OverrideDI->Equals(DataInterface)) )
+			if ((OverrideDI != nullptr) && ((OverrideDI == DataInterface) || OverrideDI->Equals(DataInterface)))
 			{
-				if ( const FNiagaraVariableBase* Variable = NiagaraSystem->GetExposedParameters().FindVariable(OverrideDI) )
+				if (const FNiagaraVariableBase* Variable = ParameterStore.FindVariable(OverrideDI))
 				{
 					OverrideParameterNames.AddUnique(Variable->GetName());
 				}
@@ -94,35 +83,21 @@ void FNiagaraDataInterfaceUtilities::ForEachVMFunctionEquals(class UNiagaraDataI
 	);
 }
 
-void FNiagaraDataInterfaceUtilities::ForEachGpuFunctionEquals(class UNiagaraDataInterface* DataInterface, class UNiagaraSystem* NiagaraSystem, class UNiagaraComponent* Component, TFunction<bool(const FNiagaraDataInterfaceGeneratedFunction&)> Action)
+void ForEachGpuFunctionEqualsImpl(class UNiagaraDataInterface* DataInterface, class UNiagaraSystem* NiagaraSystem, const FNiagaraParameterStore& ParameterStore, TFunction<bool(const FNiagaraDataInterfaceGeneratedFunction&)> Action)
 {
-	if ( DataInterface == nullptr || NiagaraSystem == nullptr )
+	if (DataInterface == nullptr || NiagaraSystem == nullptr)
 	{
 		return;
 	}
 
 	// Find all override parameter names
 	TArray<FName, TInlineAllocator<8>> OverrideParameterNames;
-	if ( Component )
-	{
-		for (const UNiagaraDataInterface* OverrideDI : Component->GetOverrideParameters().GetDataInterfaces())
-		{
-			if ( (OverrideDI != nullptr) && ((OverrideDI == DataInterface) || OverrideDI->Equals(DataInterface)) )
-			{
-				if ( const FNiagaraVariableBase* Variable = Component->GetOverrideParameters().FindVariable(OverrideDI) )
-				{
-					OverrideParameterNames.AddUnique(Variable->GetName());
-				}
-			}
-		}
-	}
-	else
 	{
-		for (const UNiagaraDataInterface* OverrideDI : NiagaraSystem->GetExposedParameters().GetDataInterfaces())
+		for (const UNiagaraDataInterface* OverrideDI : ParameterStore.GetDataInterfaces())
 		{
-			if ( (OverrideDI != nullptr) && ((OverrideDI == DataInterface) || OverrideDI->Equals(DataInterface)) )
+			if ((OverrideDI != nullptr) && ((OverrideDI == DataInterface) || OverrideDI->Equals(DataInterface)))
 			{
-				if ( const FNiagaraVariableBase* Variable = NiagaraSystem->GetExposedParameters().FindVariable(OverrideDI) )
+				if (const FNiagaraVariableBase* Variable = ParameterStore.FindVariable(OverrideDI))
 				{
 					OverrideParameterNames.AddUnique(Variable->GetName());
 				}
@@ -170,6 +145,42 @@ void FNiagaraDataInterfaceUtilities::ForEachGpuFunctionEquals(class UNiagaraData
 	);
 }
 
+}
+
+void FNiagaraDataInterfaceUtilities::ForEachVMFunctionEquals(UNiagaraDataInterface* DataInterface, UNiagaraSystem* NiagaraSystem, TFunction<bool(const FVMExternalFunctionBindingInfo&)> Action)
+{
+	ForEachVMFunctionEqualsImpl(DataInterface, NiagaraSystem, NiagaraSystem->GetExposedParameters(), Action);
+}
+
+void FNiagaraDataInterfaceUtilities::ForEachVMFunctionEquals(UNiagaraDataInterface* DataInterface, UNiagaraSystem* NiagaraSystem, UNiagaraComponent* Component, TFunction<bool(const FVMExternalFunctionBindingInfo&)> Action)
+{
+	const FNiagaraParameterStore& Parameters = Component ? Component->GetOverrideParameters() : NiagaraSystem->GetExposedParameters();
+	ForEachVMFunctionEqualsImpl(DataInterface, NiagaraSystem, Parameters, Action);
+}
+
+void FNiagaraDataInterfaceUtilities::ForEachVMFunctionEquals(UNiagaraDataInterface* DataInterface, UNiagaraSystem* NiagaraSystem, FNiagaraSystemInstance* SystemInstance, TFunction<bool(const FVMExternalFunctionBindingInfo&)> Action)
+{
+	const FNiagaraParameterStore& Parameters = (SystemInstance && SystemInstance->GetOverrideParameters()) ? *SystemInstance->GetOverrideParameters() : NiagaraSystem->GetExposedParameters();
+	ForEachVMFunctionEqualsImpl(DataInterface, NiagaraSystem, Parameters, Action);
+}
+
+void FNiagaraDataInterfaceUtilities::ForEachGpuFunctionEquals(class UNiagaraDataInterface* DataInterface, class UNiagaraSystem* NiagaraSystem, TFunction<bool(const FNiagaraDataInterfaceGeneratedFunction&)> Action)
+{
+	ForEachGpuFunctionEqualsImpl(DataInterface, NiagaraSystem, NiagaraSystem->GetExposedParameters(), Action);
+}
+
+void FNiagaraDataInterfaceUtilities::ForEachGpuFunctionEquals(class UNiagaraDataInterface* DataInterface, class UNiagaraSystem* NiagaraSystem, UNiagaraComponent* Component, TFunction<bool(const FNiagaraDataInterfaceGeneratedFunction&)> Action)
+{
+	const FNiagaraParameterStore& Parameters = Component ? Component->GetOverrideParameters() : NiagaraSystem->GetExposedParameters();
+	ForEachGpuFunctionEqualsImpl(DataInterface, NiagaraSystem, Parameters, Action);
+}
+
+void FNiagaraDataInterfaceUtilities::ForEachGpuFunctionEquals(class UNiagaraDataInterface* DataInterface, class UNiagaraSystem* NiagaraSystem, FNiagaraSystemInstance* SystemInstance, TFunction<bool(const FNiagaraDataInterfaceGeneratedFunction&)> Action)
+{
+	const FNiagaraParameterStore& Parameters = (SystemInstance && SystemInstance->GetOverrideParameters()) ? *SystemInstance->GetOverrideParameters() : NiagaraSystem->GetExposedParameters();
+	ForEachGpuFunctionEqualsImpl(DataInterface, NiagaraSystem, Parameters, Action);
+}
+
 void FNiagaraDataInterfaceUtilities::ForEachVMFunction(class UNiagaraDataInterface* DataInterface, class FNiagaraSystemInstance* SystemInstance, TFunction<bool(const FVMExternalFunctionBindingInfo&)> Action)
 {
 	if ( !DataInterface || !SystemInstance )
diff --git a/Engine/Plugins/FX/Niagara/Source/Niagara/Public/NiagaraDataInterfaceUtilities.h b/Engine/Plugins/FX/Niagara/Source/Niagara/Public/NiagaraDataInterfaceUtilities.h
--- a/Engine/Plugins/FX/Niagara/Source/Niagara/Public/NiagaraDataInterfaceUtilities.h
+++ b/Engine/Plugins/FX/Niagara/Source/Niagara/Public/NiagaraDataInterfaceUtilities.h
@@ -4,23 +4,33 @@
 #include "NiagaraCommon.h"
 #include "NiagaraShared.h"
 
+class UNiagaraComponent;
+class UNiagaraDataInterface;
+class UNiagaraSystem;
+class FNiagaraSystemInstance;
+
 namespace FNiagaraDataInterfaceUtilities
 {
 	// Finds all VM function calls made using the data interface that equals this one (i.e. A->Equals(B))
 	// The same function call made be made multiple times across different scripts so you may see the same function multiple times
 	// The action should return True to continue iteration or False to stop
-	NIAGARA_API void ForEachVMFunctionEquals(class UNiagaraDataInterface* DataInterface, class UNiagaraSystem* NiagaraSystem, class UNiagaraComponent* Component, TFunction<bool(const FVMExternalFunctionBindingInfo&)> Action);
+	NIAGARA_API void ForEachVMFunctionEquals(UNiagaraDataInterface* DataInterface, UNiagaraSystem* NiagaraSystem, TFunction<bool(const FVMExternalFunctionBindingInfo&)> Action);
+	NIAGARA_API void ForEachVMFunctionEquals(UNiagaraDataInterface* DataInterface, UNiagaraSystem* NiagaraSystem, UNiagaraComponent* Component, TFunction<bool(const FVMExternalFunctionBindingInfo&)> Action);
+	NIAGARA_API void ForEachVMFunctionEquals(UNiagaraDataInterface* DataInterface, UNiagaraSystem* NiagaraSystem, FNiagaraSystemInstance* SystemInstance, TFunction<bool(const FVMExternalFunctionBindingInfo&)> Action);
+
 	// Finds all Gpu function calls made using the data interface that equals this one (i.e. A->Equals(B))
 	// The same function call made be made multiple times across different scripts so you may see the same function multiple times
 	// The action should return True to continue iteration or False to stop
-	NIAGARA_API void ForEachGpuFunctionEquals(class UNiagaraDataInterface* DataInterface, class UNiagaraSystem* NiagaraSystem, class UNiagaraComponent* Component, TFunction<bool(const FNiagaraDataInterfaceGeneratedFunction&)> Action);
+	NIAGARA_API void ForEachGpuFunctionEquals(UNiagaraDataInterface* DataInterface, UNiagaraSystem* NiagaraSystem, TFunction<bool(const FNiagaraDataInterfaceGeneratedFunction&)> Action);
+	NIAGARA_API void ForEachGpuFunctionEquals(UNiagaraDataInterface* DataInterface, UNiagaraSystem* NiagaraSystem, UNiagaraComponent* Component, TFunction<bool(const FNiagaraDataInterfaceGeneratedFunction&)> Action);
+	NIAGARA_API void ForEachGpuFunctionEquals(UNiagaraDataInterface* DataInterface, UNiagaraSystem* NiagaraSystem, FNiagaraSystemInstance* SystemInstance, TFunction<bool(const FNiagaraDataInterfaceGeneratedFunction&)> Action);
 
 	// Finds all VM function calls made using the data interface (i.e. pointer comparison A == B)
 	// The same function call made be made multiple times across different scripts so you may see the same function multiple times
 	// The action should return True to continue iteration or False to stop
-	NIAGARA_API void ForEachVMFunction(class UNiagaraDataInterface* DataInterface, class FNiagaraSystemInstance* SystemInstance, TFunction<bool(const FVMExternalFunctionBindingInfo&)> Action);
+	NIAGARA_API void ForEachVMFunction(UNiagaraDataInterface* DataInterface, FNiagaraSystemInstance* SystemInstance, TFunction<bool(const FVMExternalFunctionBindingInfo&)> Action);
 	// Finds all Gpu function calls made using the data interface (i.e. pointer comparison A == B)
 	// The same function call made be made multiple times across different scripts so you may see the same function multiple times
 	// The action should return True to continue iteration or False to stop
-	NIAGARA_API void ForEachGpuFunction(class UNiagaraDataInterface* DataInterface, class FNiagaraSystemInstance* SystemInstance, TFunction<bool(const FNiagaraDataInterfaceGeneratedFunction&)> Action);
+	NIAGARA_API void ForEachGpuFunction(UNiagaraDataInterface* DataInterface, FNiagaraSystemInstance* SystemInstance, TFunction<bool(const FNiagaraDataInterfaceGeneratedFunction&)> Action);
 }
