Improvements to Niagara's RigidBody DataInterface
-debug visualization of the rigid bodies available from niagara debugger
-added support for skinned meshes
-allows a spatial search through an overlapped collision check
-fixed some leaks and a use after deletes
-added ActorTags and ComponentTags for filtering out results

diff --git a/Engine/Plugins/FX/Niagara/Source/Niagara/Private/DataInterface/NiagaraDistanceFieldParameters.h b/Engine/Plugins/FX/Niagara/Source/Niagara/Private/DataInterface/NiagaraDistanceFieldParameters.h
--- a/Engine/Plugins/FX/Niagara/Source/Niagara/Private/DataInterface/NiagaraDistanceFieldParameters.h
+++ b/Engine/Plugins/FX/Niagara/Source/Niagara/Private/DataInterface/NiagaraDistanceFieldParameters.h
@@ -2,6 +2,7 @@
 
 #pragma once
 #include "Renderer/Private/DistanceFieldLightingShared.h"
+#include "NiagaraRenderer.h"
 
 class FDistanceFieldParameters
 {
diff --git a/Engine/Plugins/FX/Niagara/Source/Niagara/Private/NiagaraDataInterfaceRigidMeshCollisionQuery.cpp b/Engine/Plugins/FX/Niagara/Source/Niagara/Private/NiagaraDataInterfaceRigidMeshCollisionQuery.cpp
--- a/Engine/Plugins/FX/Niagara/Source/Niagara/Private/NiagaraDataInterfaceRigidMeshCollisionQuery.cpp
+++ b/Engine/Plugins/FX/Niagara/Source/Niagara/Private/NiagaraDataInterfaceRigidMeshCollisionQuery.cpp
@@ -1,24 +1,38 @@
 // Copyright Epic Games, Inc. All Rights Reserved.
 
 #include "NiagaraDataInterfaceRigidMeshCollisionQuery.h"
+#include "Algo/ForEach.h"
 #include "Animation/SkeletalMeshActor.h"
-#include "Components/SkeletalMeshComponent.h"
-#include "SkeletalRenderPublic.h"
 #include "AnimationRuntime.h"
-#include "NiagaraShader.h"
+#include "Components/SkeletalMeshComponent.h"
+#include "DataInterface/NiagaraDistanceFieldParameters.h"
+#include "Engine/Canvas.h"
+#include "EngineUtils.h"
 #include "NiagaraComponent.h"
+#include "NiagaraDataInterfaceUtilities.h"
+#include "NiagaraFunctionLibrary.h"
+#include "NiagaraGpuComputeDispatch.h"
+#include "NiagaraGpuComputeDispatchInterface.h"
 #include "NiagaraRenderer.h"
+#include "NiagaraShader.h"
 #include "NiagaraSimStageData.h"
 #include "NiagaraSystemInstance.h"
 #include "ShaderParameterUtils.h"
-#include "EngineUtils.h"
-#include "NiagaraGpuComputeDispatchInterface.h"
-#include "NiagaraGpuComputeDispatch.h"
-#include "DataInterface/NiagaraDistanceFieldParameters.h"
+#include "SkeletalRenderPublic.h"
 
 #define LOCTEXT_NAMESPACE "NiagaraDataInterfaceRigidMeshCollisionQuery"
 DEFINE_LOG_CATEGORY_STATIC(LogRigidMeshCollision, Log, All);
 
+// outstanding/known issues:
+// -when actors change and the arrays are fully updated we'll experience a frame of 0 velocities
+//		-potentially we could keep track of ranges of rigid bodies for given actors and then smartly reassign
+//		the previous frame's transforms
+// -could add a vM function for setting the maximum number of primitives
+
+
+namespace NDIRigidMeshCollisionLocal
+{
+
 struct FNiagaraRigidMeshCollisionDIFunctionVersion
 {
 	enum Type
@@ -34,6 +48,10 @@ struct FNiagaraRigidMeshCollisionDIFunctionVersion
 
 //------------------------------------------------------------------------------------------------------------
 
+static const FName FindActorsName(TEXT("FindActors"));
+
+//------------------------------------------------------------------------------------------------------------
+
 static const FName GetNumBoxesName(TEXT("GetNumBoxes"));
 static const FName GetNumSpheresName(TEXT("GetNumSpheres"));
 static const FName GetNumCapsulesName(TEXT("GetNumCapsules"));
@@ -52,18 +70,6 @@ static const FName GetClosestPointMeshDistanceFieldNoNormalName(TEXT("GetClosest
 
 //------------------------------------------------------------------------------------------------------------
 
-const FString UNiagaraDataInterfaceRigidMeshCollisionQuery::MaxTransformsName(TEXT("MaxTransforms_"));
-const FString UNiagaraDataInterfaceRigidMeshCollisionQuery::CurrentOffsetName(TEXT("CurrentOffset_"));
-const FString UNiagaraDataInterfaceRigidMeshCollisionQuery::PreviousOffsetName(TEXT("PreviousOffset_"));
-
-const FString UNiagaraDataInterfaceRigidMeshCollisionQuery::ElementOffsetsName(TEXT("ElementOffsets_"));
-
-const FString UNiagaraDataInterfaceRigidMeshCollisionQuery::WorldTransformBufferName(TEXT("WorldTransformBuffer_"));
-const FString UNiagaraDataInterfaceRigidMeshCollisionQuery::InverseTransformBufferName(TEXT("InverseTransformBuffer_"));
-const FString UNiagaraDataInterfaceRigidMeshCollisionQuery::ElementExtentBufferName(TEXT("ElementExtentBuffer_"));
-const FString UNiagaraDataInterfaceRigidMeshCollisionQuery::PhysicsTypeBufferName(TEXT("PhysicsTypeBuffer_"));
-const FString UNiagaraDataInterfaceRigidMeshCollisionQuery::DFIndexBufferName(TEXT("DFIndexBuffer_"));
-
 bool IsMeshDistanceFieldEnabled()
 {
 	static const auto* CVarGenerateMeshDistanceFields = IConsoleManager::Get().FindTConsoleVariableDataInt(TEXT("r.GenerateMeshDistanceFields"));
@@ -128,7 +134,7 @@ void UpdateInternalBuffer(const TArray<BufferType>& InputData, FRWBuffer& Output
 	}
 }
 
-void FillCurrentTransforms(const FTransform& ElementTransform, uint32& ElementCount, TArray<FVector4f>& OutCurrentTransform, TArray<FVector4f>& OutCurrentInverse)
+void FillCurrentTransforms(const FTransform& ElementTransform, uint32 ElementCount, TArray<FVector4f>& OutCurrentTransform, TArray<FVector4f>& OutCurrentInverse)
 {
 	// LWC_TODO: precision loss
 	const uint32 ElementOffset = 3 * ElementCount;
@@ -137,289 +143,359 @@ void FillCurrentTransforms(const FTransform& ElementTransform, uint32& ElementCo
 
 	ElementMatrix.To3x4MatrixTranspose(&OutCurrentTransform[ElementOffset].X);
 	ElementInverse.To3x4MatrixTranspose(&OutCurrentInverse[ElementOffset].X);
-	++ElementCount;
 }
 
-void GetNumPrimitives(TArray<TWeakObjectPtr<AActor>> Actors, uint32& NumBoxes, uint32& NumSpheres, uint32& NumCapsules)
+template<typename TComponentType, typename TComponentFilterPredicate>
+static void GenerateComponentList(
+	TConstArrayView<AActor*> Actors,
+	TConstArrayView<FName> ComponentTags,
+	TComponentFilterPredicate FilterPredicate,
+	TInlineComponentArray<TComponentType*>& Components)
+{
+	for (const AActor* Actor : Actors)
+	{
+		if (Actor)
+		{
+			for (UActorComponent* ActorComponent : Actor->GetComponents())
+			{
+				if (TComponentType* TypedComponent = Cast<TComponentType>(ActorComponent))
+				{
+					if (IsValid(TypedComponent) && FilterPredicate(TypedComponent))
+					{
+						if (ComponentTags.IsEmpty() || ComponentTags.ContainsByPredicate([&](const FName& Tag) { return Tag == NAME_None || TypedComponent->ComponentHasTag(Tag); }))
+						{
+							Components.Add(TypedComponent);
+						}
+					}
+				}
+			}
+		}
+	}
+}
+
+template<typename TComponentType>
+void CollectComponents(TConstArrayView<AActor*> Actors, TConstArrayView<FName> ComponentTags, TInlineComponentArray<TComponentType*>& Components);
+
+template<typename TComponentType, typename TBodySetupPredicate>
+void ForEachBodySetup(TComponentType* Component, TBodySetupPredicate Predicate);
+
+/// Begin UkeletalMeshComponent
+
+template<>
+void CollectComponents(TConstArrayView<AActor*> Actors, TConstArrayView<FName> ComponentTags, TInlineComponentArray<USkeletalMeshComponent*>& Components)
 {
-	NumBoxes = 0;
-	NumSpheres = 0;
-	NumCapsules = 0;
-
-	for (int32 ActorIndex = 0; ActorIndex < Actors.Num(); ++ActorIndex)
-	{		
-		UStaticMeshComponent* StaticMeshComponent = nullptr;
-		if (Actors[ActorIndex].IsValid())
+	auto SkeletalFilterPredicate = [&](USkeletalMeshComponent* Component)
+	{
+		if (UPhysicsAsset* PhysicsAsset = Component->GetPhysicsAsset())
 		{
-			StaticMeshComponent = Cast<UStaticMeshComponent>(Actors[ActorIndex]->GetComponentByClass(UStaticMeshComponent::StaticClass()));
+			USkeletalMesh* MeshAsset = Component->SkeletalMesh ? Component->SkeletalMesh.Get() : PhysicsAsset->GetPreviewMesh();
+			if (!MeshAsset || !MeshAsset->GetRefSkeleton().GetNum())
+			{
+				return false;
+			}
+
+			return true;
 		}
-		UBodySetup* BodySetup = StaticMeshComponent != nullptr ? StaticMeshComponent->GetBodySetup() : nullptr;
-		if (BodySetup != nullptr)
+
+		return false;
+	};
+
+	GenerateComponentList(Actors, ComponentTags, SkeletalFilterPredicate, Components);
+}
+
+template<typename TBodySetupPredicate>
+void ForEachBodySetup(USkeletalMeshComponent* Component, TBodySetupPredicate Predicate)
+{
+	if (UPhysicsAsset* PhysicsAsset = Component->GetPhysicsAsset())
+	{
+		USkeletalMesh* SkeletalMesh = Component->SkeletalMesh ? Component->SkeletalMesh.Get() : PhysicsAsset->GetPreviewMesh();
+		const FReferenceSkeleton& RefSkeleton = SkeletalMesh->GetRefSkeleton();
+
+		for (const UBodySetup* BodySetup : PhysicsAsset->SkeletalBodySetups)
 		{
-			// #todo(dmp): save static mesh component reference
+			const FName BoneName = BodySetup->BoneName;
+			const int32 BoneIndex = RefSkeleton.FindBoneIndex(BoneName);
+			if (BoneIndex != INDEX_NONE && BoneIndex < RefSkeleton.GetNum())
+			{
+				Predicate(Component, BodySetup);
+			}
+		}
+	}
+}
+
+/// End UkeletalMeshComponent
 
+/// Begin UStaticMeshComponent
+
+template<>
+void CollectComponents(TConstArrayView<AActor*> Actors, TConstArrayView<FName> ComponentTags, TInlineComponentArray<UStaticMeshComponent*>& Components)
+{
+	auto StaticFilterPredicate = [&](UStaticMeshComponent* Component)
+	{
+		return Component->GetBodySetup() != nullptr;
+	};
+
+	GenerateComponentList(Actors, ComponentTags, StaticFilterPredicate, Components);
+}
+
+template<typename TBodySetupPredicate>
+void ForEachBodySetup(UStaticMeshComponent* Component, TBodySetupPredicate Predicate)
+{
+	Predicate(Component, Component->GetBodySetup());
+}
+
+/// End UStaticMeshComponent
+
+template<typename TComponentType>
+void CountCollisionPrimitives(TConstArrayView<TComponentType*> Components, uint32& BoxCount, uint32& SphereCount, uint32& CapsuleCount)
+{
+	for (TComponentType* Component : Components)
+	{
+		bool HasConvexElements = false;
+
+		ForEachBodySetup(Component, [&](TComponentType* Component, const UBodySetup* BodySetup)
+		{
 			for (const FKConvexElem& ConvexElem : BodySetup->AggGeom.ConvexElems)
 			{
 				if (CollisionEnabledHasPhysics(ConvexElem.GetCollisionEnabled()))
 				{
-					NumBoxes += 1;
+					HasConvexElements = true;
+					++BoxCount;
 				}
 			}
 			for (const FKBoxElem& BoxElem : BodySetup->AggGeom.BoxElems)
 			{
 				if (CollisionEnabledHasPhysics(BoxElem.GetCollisionEnabled()))
 				{
-					NumBoxes += 1;
+					++BoxCount;
 				}
 			}
 			for (const FKSphereElem& SphereElem : BodySetup->AggGeom.SphereElems)
 			{
 				if (CollisionEnabledHasPhysics(SphereElem.GetCollisionEnabled()))
 				{
-					NumSpheres += 1;
+					++SphereCount;
 				}
 			}
 			for (const FKSphylElem& CapsuleElem : BodySetup->AggGeom.SphylElems)
 			{
 				if (CollisionEnabledHasPhysics(CapsuleElem.GetCollisionEnabled()))
 				{
-					NumCapsules += 1;
+					++CapsuleCount;
 				}
-			}				
+			}
+		});
+
+		if (HasConvexElements)
+		{
+			UE_LOG(LogRigidMeshCollision, Warning, TEXT("Convex collision objects encountered and will be interpreted as a bounding box on %s"), *Component->GetOwner()->GetName());
 		}
 	}
 }
 
-void CompactInternalArrays(FNDIRigidMeshCollisionArrays* OutAssetArrays)
+template<typename TComponentType>
+FTransform CreateElementTransform(const TComponentType* Component, const UBodySetup* BodySetup)
 {
-	for (uint32 TransformIndex = 0; TransformIndex < OutAssetArrays->MaxTransforms; ++TransformIndex)
-	{
-		uint32 OffsetIndex = TransformIndex;
-		OutAssetArrays->WorldTransform[OffsetIndex] = OutAssetArrays->CurrentTransform[TransformIndex];
-		OutAssetArrays->InverseTransform[OffsetIndex] = OutAssetArrays->CurrentInverse[TransformIndex];
-
-		OffsetIndex += OutAssetArrays->MaxTransforms;
-		OutAssetArrays->WorldTransform[OffsetIndex] = OutAssetArrays->PreviousTransform[TransformIndex];
-		OutAssetArrays->InverseTransform[OffsetIndex] = OutAssetArrays->PreviousInverse[TransformIndex];		
-	}
+	return Component->GetComponentTransform();
 }
 
-void CreateInternalArrays(TArray<TWeakObjectPtr<AActor>> Actors, FNDIRigidMeshCollisionArrays* OutAssetArrays, FVector LWCTile)
+template<>
+FTransform CreateElementTransform<USkeletalMeshComponent>(const USkeletalMeshComponent* Component, const UBodySetup* BodySetup)
 {
-	if (OutAssetArrays != nullptr)
+	if (USkeletalMesh* SkeletalMesh = Component->SkeletalMesh.Get())
 	{
-		OutAssetArrays->ElementOffsets.BoxOffset = 0;
-		OutAssetArrays->ElementOffsets.SphereOffset = 0;
-		OutAssetArrays->ElementOffsets.CapsuleOffset = 0;
-		OutAssetArrays->ElementOffsets.NumElements = 0;
-
-		uint32 NumBoxes = 0;
-		uint32 NumSpheres = 0;
-		uint32 NumCapsules = 0;
+		const FReferenceSkeleton& RefSkeleton = SkeletalMesh->GetRefSkeleton();
+		const int32 BoneCount = RefSkeleton.GetNum();
 
-		GetNumPrimitives(Actors, NumBoxes, NumSpheres, NumCapsules);
-		
-		if ((NumBoxes + NumSpheres + NumCapsules) < OutAssetArrays->MaxPrimitives)
+		if (BoneCount > 0)
 		{
-			OutAssetArrays->ElementOffsets.BoxOffset = 0;
-			OutAssetArrays->ElementOffsets.SphereOffset = OutAssetArrays->ElementOffsets.BoxOffset + NumBoxes;
-			OutAssetArrays->ElementOffsets.CapsuleOffset = OutAssetArrays->ElementOffsets.SphereOffset + NumSpheres;
-			OutAssetArrays->ElementOffsets.NumElements = OutAssetArrays->ElementOffsets.CapsuleOffset + NumCapsules;
-
-			uint32 BoxCount = OutAssetArrays->ElementOffsets.BoxOffset;
-			uint32 SphereCount = OutAssetArrays->ElementOffsets.SphereOffset;
-			uint32 CapsuleCount = OutAssetArrays->ElementOffsets.CapsuleOffset;
-
-			for (int32 ActorIndex = 0; ActorIndex < Actors.Num(); ++ActorIndex)
-			{				
-				TWeakObjectPtr<AActor> Actor = Actors[ActorIndex].Get();
-				
-				if (Actor.IsValid() && Actor->GetComponentByClass(UStaticMeshComponent::StaticClass()) != nullptr)
-				{
-					UStaticMeshComponent* StaticMeshComponent = Cast< UStaticMeshComponent>(Actor->GetComponentByClass(UStaticMeshComponent::StaticClass()));
-					UBodySetup* BodySetup = StaticMeshComponent->GetBodySetup();
-					bool FoundCollisionShapes = false;
-					if (BodySetup != nullptr)
-					{
-						FTransform MeshTransform = Actor->GetTransform();
-						MeshTransform.AddToTranslation(LWCTile * -FLargeWorldRenderScalar::GetTileSize());
-
-						for (const FKConvexElem& ConvexElem : BodySetup->AggGeom.ConvexElems)
-						{							
-							if (CollisionEnabledHasPhysics(ConvexElem.GetCollisionEnabled()))
-							{
-								UE_LOG(LogRigidMeshCollision, Warning, TEXT("Convex collision objects encountered and will be interpreted as a bounding box on %s"), *Actor->GetName());
-
-
-								FBox BBox = ConvexElem.ElemBox;
-								FVector3f Extent = FVector3f(BBox.Max - BBox.Min);
-								FVector Center = (BBox.Max + BBox.Min) * .5;
-								OutAssetArrays->PhysicsType[BoxCount] = (ConvexElem.GetCollisionEnabled() == ECollisionEnabled::QueryAndPhysics);
-								OutAssetArrays->SourceSceneProxy[BoxCount] = StaticMeshComponent->SceneProxy;
-								
-								const FTransform ElementTransform = FTransform(Center) * MeshTransform;
-								OutAssetArrays->ElementExtent[BoxCount] = FVector4f(Extent.X, Extent.Y, Extent.Z, 0);
-								FillCurrentTransforms(ElementTransform, BoxCount, OutAssetArrays->CurrentTransform, OutAssetArrays->CurrentInverse);
-
-								FoundCollisionShapes = true;
-							}
-						}
-						for (const FKBoxElem& BoxElem : BodySetup->AggGeom.BoxElems)
-						{
-							if (CollisionEnabledHasPhysics(BoxElem.GetCollisionEnabled()))
-							{
-								OutAssetArrays->PhysicsType[BoxCount] = (BoxElem.GetCollisionEnabled() == ECollisionEnabled::QueryAndPhysics);
-								OutAssetArrays->SourceSceneProxy[BoxCount] = StaticMeshComponent->SceneProxy;
-
-								const FTransform ElementTransform = FTransform(BoxElem.Rotation, BoxElem.Center) * MeshTransform;
-								OutAssetArrays->ElementExtent[BoxCount] = FVector4f(BoxElem.X, BoxElem.Y, BoxElem.Z, 0);
-								FillCurrentTransforms(ElementTransform, BoxCount, OutAssetArrays->CurrentTransform, OutAssetArrays->CurrentInverse);
-
-								FoundCollisionShapes = true;
-							}
-						}
+			const FName BoneName = BodySetup->BoneName;
+			const int32 BoneIndex = RefSkeleton.FindBoneIndex(BoneName);
+			if (BoneIndex != INDEX_NONE && BoneIndex < BoneCount)
+			{
+				return Component->GetBoneTransform(BoneIndex);
+			}
+		}
+	}
 
-						for (const FKSphereElem& SphereElem : BodySetup->AggGeom.SphereElems)
-						{
-							if (CollisionEnabledHasPhysics(SphereElem.GetCollisionEnabled()))
-							{
-								OutAssetArrays->PhysicsType[SphereCount] = (SphereElem.GetCollisionEnabled() == ECollisionEnabled::QueryAndPhysics);
-								OutAssetArrays->SourceSceneProxy[BoxCount] = StaticMeshComponent->SceneProxy;
+	return Component->GetComponentTransform();
+}
 
-								const FTransform ElementTransform = FTransform(SphereElem.Center) * MeshTransform;
-								OutAssetArrays->ElementExtent[SphereCount] = FVector4f(SphereElem.Radius, 0, 0, 0);
-								FillCurrentTransforms(ElementTransform, SphereCount, OutAssetArrays->CurrentTransform, OutAssetArrays->CurrentInverse);
+template<typename TComponentType, bool InitializeStatics>
+void UpdateAssetArrays(TConstArrayView<TComponentType*> Components, const FVector& LWCTile, FNDIRigidMeshCollisionArrays* OutAssetArrays, uint32& BoxIndex, uint32& SphereIndex, uint32& CapsuleIndex)
+{
+	auto UpdateAssetPredicate = [&](TComponentType* Component, const UBodySetup* BodySetup)
+	{
+		FTransform MeshTransform = CreateElementTransform(Component, BodySetup);
+		MeshTransform.AddToTranslation(LWCTile * -FLargeWorldRenderScalar::GetTileSize());
 
-								FoundCollisionShapes = true;
-							}
-						}
+		const int32 ComponentIdIndex = OutAssetArrays->UniqueCompnentId.AddUnique(Component->ComponentId);
 
-						for (const FKSphylElem& CapsuleElem : BodySetup->AggGeom.SphylElems)
-						{
-							if (CollisionEnabledHasPhysics(CapsuleElem.GetCollisionEnabled()))
-							{
-								OutAssetArrays->PhysicsType[CapsuleCount] = (CapsuleElem.GetCollisionEnabled() == ECollisionEnabled::QueryAndPhysics);
-								OutAssetArrays->SourceSceneProxy[BoxCount] = StaticMeshComponent->SceneProxy;
+		for (const FKConvexElem& ConvexElem : BodySetup->AggGeom.ConvexElems)
+		{
+			if (CollisionEnabledHasPhysics(ConvexElem.GetCollisionEnabled()))
+			{
+				FBox BBox = ConvexElem.ElemBox;
 
-								const FTransform ElementTransform = FTransform(CapsuleElem.Rotation, CapsuleElem.Center) * MeshTransform;
-								OutAssetArrays->ElementExtent[CapsuleCount] = FVector4f(CapsuleElem.Radius, CapsuleElem.Length, 0, 0);
-								FillCurrentTransforms(ElementTransform, CapsuleCount, OutAssetArrays->CurrentTransform, OutAssetArrays->CurrentInverse);
+				if (InitializeStatics)
+				{
+					FVector3f Extent = FVector3f(BBox.Max - BBox.Min);
+					OutAssetArrays->ElementExtent[BoxIndex] = FVector4f(Extent.X, Extent.Y, Extent.Z, 0);
+					OutAssetArrays->PhysicsType[BoxIndex] = (ConvexElem.GetCollisionEnabled() == ECollisionEnabled::QueryAndPhysics);
+					OutAssetArrays->ComponentIdIndex[BoxIndex] = ComponentIdIndex;
+				}
 
-								FoundCollisionShapes = true;
-							}
-						}
-					}
+				FVector Center = (BBox.Max + BBox.Min) * .5;
+				const FTransform ElementTransform = FTransform(Center) * MeshTransform;
+				FillCurrentTransforms(ElementTransform, BoxIndex, OutAssetArrays->CurrentTransform, OutAssetArrays->CurrentInverse);
+				++BoxIndex;
+			}
+		}
+		for (const FKBoxElem& BoxElem : BodySetup->AggGeom.BoxElems)
+		{
+			if (CollisionEnabledHasPhysics(BoxElem.GetCollisionEnabled()))
+			{
+				if (InitializeStatics)
+				{
+					OutAssetArrays->ElementExtent[BoxIndex] = FVector4f(BoxElem.X, BoxElem.Y, BoxElem.Z, 0);
+					OutAssetArrays->PhysicsType[BoxIndex] = (BoxElem.GetCollisionEnabled() == ECollisionEnabled::QueryAndPhysics);
+					OutAssetArrays->ComponentIdIndex[BoxIndex] = ComponentIdIndex;
+				}
 
-					if (!FoundCollisionShapes)
-					{
-						UE_LOG(LogRigidMeshCollision, Error, TEXT("No useable collision body setup found on mesh %s"), *Actor->GetName());
+				const FTransform ElementTransform = FTransform(BoxElem.Rotation, BoxElem.Center) * MeshTransform;
+				FillCurrentTransforms(ElementTransform, BoxIndex, OutAssetArrays->CurrentTransform, OutAssetArrays->CurrentInverse);
+				++BoxIndex;
+			}
+		}
 
-					}
+		for (const FKSphereElem& SphereElem : BodySetup->AggGeom.SphereElems)
+		{
+			if (CollisionEnabledHasPhysics(SphereElem.GetCollisionEnabled()))
+			{
+				if (InitializeStatics)
+				{
+					OutAssetArrays->ElementExtent[SphereIndex] = FVector4f(SphereElem.Radius, 0, 0, 0);
+					OutAssetArrays->PhysicsType[SphereIndex] = (SphereElem.GetCollisionEnabled() == ECollisionEnabled::QueryAndPhysics);
+					OutAssetArrays->ComponentIdIndex[SphereIndex] = ComponentIdIndex;
 				}
-			}
-			OutAssetArrays->PreviousTransform = OutAssetArrays->CurrentTransform;
-			OutAssetArrays->PreviousInverse = OutAssetArrays->CurrentInverse;
 
-			CompactInternalArrays(OutAssetArrays);
+				const FTransform ElementTransform = FTransform(SphereElem.Center) * MeshTransform;
+				FillCurrentTransforms(ElementTransform, SphereIndex, OutAssetArrays->CurrentTransform, OutAssetArrays->CurrentInverse);
+				++SphereIndex;
+			}
 		}
-		else
+
+		for (const FKSphylElem& CapsuleElem : BodySetup->AggGeom.SphylElems)
 		{
-			UE_LOG(LogRigidMeshCollision, Warning, TEXT("Number of Collision DI primitives is higher than the %d limit.  Please increase it."), OutAssetArrays->MaxPrimitives);
+			if (CollisionEnabledHasPhysics(CapsuleElem.GetCollisionEnabled()))
+			{
+				if (InitializeStatics)
+				{
+					OutAssetArrays->ElementExtent[CapsuleIndex] = FVector4f(CapsuleElem.Radius, CapsuleElem.Length, 0, 0);
+					OutAssetArrays->PhysicsType[CapsuleIndex] = (CapsuleElem.GetCollisionEnabled() == ECollisionEnabled::QueryAndPhysics);
+					OutAssetArrays->ComponentIdIndex[CapsuleIndex] = ComponentIdIndex;
+				}
+
+				const FTransform ElementTransform = FTransform(CapsuleElem.Rotation, CapsuleElem.Center) * MeshTransform;
+				FillCurrentTransforms(ElementTransform, CapsuleIndex, OutAssetArrays->CurrentTransform, OutAssetArrays->CurrentInverse);
+				++CapsuleIndex;
+			}
 		}
+	};
+
+	for (TComponentType* Component : Components)
+	{
+		ForEachBodySetup(Component, UpdateAssetPredicate);
 	}
 }
 
-void UpdateInternalArrays(const TArray<TWeakObjectPtr<AActor>> &Actors, FNDIRigidMeshCollisionArrays* OutAssetArrays, FVector LWCTile)
+void UpdateInternalArrays(TConstArrayView<AActor*> Actors, TConstArrayView<FName> ComponentTags, FVector LWCTile, bool bFullUpdate, FNDIRigidMeshCollisionArrays* OutAssetArrays)
 {
 	if (OutAssetArrays != nullptr && OutAssetArrays->ElementOffsets.NumElements < OutAssetArrays->MaxPrimitives)
 	{
-		uint32 NumBoxes = 0;
-		uint32 NumSpheres = 0;
-		uint32 NumCapsules = 0;		
+		TInlineComponentArray<UStaticMeshComponent*> StaticMeshes;
+		TInlineComponentArray<USkeletalMeshComponent*> SkeletalMeshes;
 
-		GetNumPrimitives(Actors, NumBoxes, NumSpheres, NumCapsules);
+		uint32 BoxCount = 0;
+		uint32 SphereCount = 0;
+		uint32 CapsuleCount = 0;
 
-		if (((OutAssetArrays->ElementOffsets.SphereOffset - OutAssetArrays->ElementOffsets.BoxOffset) != NumBoxes) || 
-			((OutAssetArrays->ElementOffsets.CapsuleOffset - OutAssetArrays->ElementOffsets.SphereOffset) != NumSpheres) ||
-			((OutAssetArrays->ElementOffsets.NumElements - OutAssetArrays->ElementOffsets.CapsuleOffset) != NumCapsules))
-		{
-			CreateInternalArrays(Actors, OutAssetArrays, LWCTile);
-		}
+		CollectComponents(Actors, ComponentTags, StaticMeshes);
+		CollectComponents(Actors, ComponentTags, SkeletalMeshes);
+
+		TConstArrayView<UStaticMeshComponent*> StaticMeshView = MakeArrayView(StaticMeshes.GetData(), StaticMeshes.Num());
+		TConstArrayView<USkeletalMeshComponent*> SkeletalMeshView = MakeArrayView(SkeletalMeshes.GetData(), SkeletalMeshes.Num());
+
+		CountCollisionPrimitives(StaticMeshView, BoxCount, SphereCount, CapsuleCount);
+		CountCollisionPrimitives(SkeletalMeshView, BoxCount, SphereCount, CapsuleCount);
 
-		OutAssetArrays->PreviousTransform = OutAssetArrays->CurrentTransform;
-		OutAssetArrays->PreviousInverse = OutAssetArrays->CurrentInverse;
+		const bool MismatchOffsets = ((OutAssetArrays->ElementOffsets.SphereOffset - OutAssetArrays->ElementOffsets.BoxOffset) != BoxCount) ||
+			((OutAssetArrays->ElementOffsets.CapsuleOffset - OutAssetArrays->ElementOffsets.SphereOffset) != SphereCount) ||
+			((OutAssetArrays->ElementOffsets.NumElements - OutAssetArrays->ElementOffsets.CapsuleOffset) != CapsuleCount);
 
-		uint32 BoxCount = OutAssetArrays->ElementOffsets.BoxOffset;
-		uint32 SphereCount = OutAssetArrays->ElementOffsets.SphereOffset;
-		uint32 CapsuleCount = OutAssetArrays->ElementOffsets.CapsuleOffset;
+		// if we're only running an update, then make sure that the offsets aren't mismatched
+		check(!MismatchOffsets || bFullUpdate);
 
-		for (int32 ActorIndex = 0; ActorIndex < Actors.Num(); ++ActorIndex)
+		if (bFullUpdate)
 		{
-			TWeakObjectPtr<AActor> Actor = Actors[ActorIndex].Get();
-			UStaticMeshComponent* StaticMeshComponent = nullptr;
-			if (Actor.IsValid())
+			if ((BoxCount + SphereCount + CapsuleCount) < OutAssetArrays->MaxPrimitives)
 			{
-				StaticMeshComponent = Cast<UStaticMeshComponent>(Actor->GetComponentByClass(UStaticMeshComponent::StaticClass()));
+				OutAssetArrays->ElementOffsets.BoxOffset = 0;
+				OutAssetArrays->ElementOffsets.SphereOffset = OutAssetArrays->ElementOffsets.BoxOffset + BoxCount;
+				OutAssetArrays->ElementOffsets.CapsuleOffset = OutAssetArrays->ElementOffsets.SphereOffset + SphereCount;
+				OutAssetArrays->ElementOffsets.NumElements = OutAssetArrays->ElementOffsets.CapsuleOffset + CapsuleCount;
 			}
-			UBodySetup* BodySetup = StaticMeshComponent != nullptr ? StaticMeshComponent->GetBodySetup() : nullptr;
-			if (BodySetup != nullptr)
+			else
 			{
-				FTransform MeshTransform = Actor->GetTransform();
-				MeshTransform.AddToTranslation(LWCTile * -FLargeWorldRenderScalar::GetTileSize());
-
-				for (const FKConvexElem& ConvexElem : BodySetup->AggGeom.ConvexElems)
-				{
-					if (CollisionEnabledHasPhysics(ConvexElem.GetCollisionEnabled()))
-					{
-						FBox BBox = ConvexElem.ElemBox;												
-						FVector Center = (BBox.Max + BBox.Min) * .5;
-
-						OutAssetArrays->SourceSceneProxy[BoxCount] = StaticMeshComponent->SceneProxy;
-						const FTransform ElementTransform = FTransform(Center) * MeshTransform;
-						FillCurrentTransforms(ElementTransform, BoxCount, OutAssetArrays->CurrentTransform, OutAssetArrays->CurrentInverse);
-					}
-				}
+				UE_LOG(LogRigidMeshCollision, Warning, TEXT("Number of Collision DI primitives is higher than the %d limit.  Please increase it."), OutAssetArrays->MaxPrimitives);
+			}
+		}
 
-				for (const FKBoxElem& BoxElem : BodySetup->AggGeom.BoxElems)
-				{
-					if (CollisionEnabledHasPhysics(BoxElem.GetCollisionEnabled()))
-					{
-						OutAssetArrays->SourceSceneProxy[BoxCount] = StaticMeshComponent->SceneProxy;
-						const FTransform ElementTransform = FTransform(BoxElem.Rotation, BoxElem.Center) * MeshTransform;
-						FillCurrentTransforms(ElementTransform, BoxCount, OutAssetArrays->CurrentTransform, OutAssetArrays->CurrentInverse);
-					}
-				}
+		uint32 BoxIndex = OutAssetArrays->ElementOffsets.BoxOffset;
+		uint32 SphereIndex = OutAssetArrays->ElementOffsets.SphereOffset;
+		uint32 CapsuleIndex = OutAssetArrays->ElementOffsets.CapsuleOffset;
 
-				for (const FKSphereElem& SphereElem : BodySetup->AggGeom.SphereElems)
-				{
-					if (CollisionEnabledHasPhysics(SphereElem.GetCollisionEnabled()))
-					{
-						OutAssetArrays->SourceSceneProxy[SphereCount] = StaticMeshComponent->SceneProxy;
-						const FTransform ElementTransform = FTransform(SphereElem.Center) * MeshTransform;
-						FillCurrentTransforms(ElementTransform, SphereCount, OutAssetArrays->CurrentTransform, OutAssetArrays->CurrentInverse);
-					}
-				}
+		if (bFullUpdate)
+		{
+			UpdateAssetArrays<UStaticMeshComponent, true>(StaticMeshView, LWCTile, OutAssetArrays, BoxIndex, SphereIndex, CapsuleIndex);
+			UpdateAssetArrays<USkeletalMeshComponent, true>(SkeletalMeshView, LWCTile, OutAssetArrays, BoxIndex, SphereIndex, CapsuleIndex);
 
-				for (const FKSphylElem& CapsuleElem : BodySetup->AggGeom.SphylElems)
-				{
-					if (CollisionEnabledHasPhysics(CapsuleElem.GetCollisionEnabled()))
-					{
-						OutAssetArrays->SourceSceneProxy[CapsuleCount] = StaticMeshComponent->SceneProxy;
-						const FTransform ElementTransform = FTransform(CapsuleElem.Rotation, CapsuleElem.Center) * MeshTransform;
-						FillCurrentTransforms(ElementTransform, CapsuleCount, OutAssetArrays->CurrentTransform, OutAssetArrays->CurrentInverse);
-					}
-				}
-			}
+			// for newly created array data we need to duplicate the transforms to our previous transforms
+			OutAssetArrays->PreviousTransform = OutAssetArrays->CurrentTransform;
+			OutAssetArrays->PreviousInverse = OutAssetArrays->CurrentInverse;
 		}
+		else
+		{
+			// if we're updating, then copy over last frame's transforms before we generate new ones
+			OutAssetArrays->PreviousTransform = OutAssetArrays->CurrentTransform;
+			OutAssetArrays->PreviousInverse = OutAssetArrays->CurrentInverse;
 
-		CompactInternalArrays(OutAssetArrays);
+			UpdateAssetArrays<UStaticMeshComponent, false>(StaticMeshView, LWCTile, OutAssetArrays, BoxIndex, SphereIndex, CapsuleIndex);
+			UpdateAssetArrays<USkeletalMeshComponent, false>(SkeletalMeshView, LWCTile, OutAssetArrays, BoxIndex, SphereIndex, CapsuleIndex);
+		}
 	}
 }
 
+} // NDIRigidMeshCollisionLocal
+
+const FString UNiagaraDataInterfaceRigidMeshCollisionQuery::MaxTransformsName(TEXT("MaxTransforms_"));
+const FString UNiagaraDataInterfaceRigidMeshCollisionQuery::CurrentOffsetName(TEXT("CurrentOffset_"));
+const FString UNiagaraDataInterfaceRigidMeshCollisionQuery::PreviousOffsetName(TEXT("PreviousOffset_"));
+
+const FString UNiagaraDataInterfaceRigidMeshCollisionQuery::ElementOffsetsName(TEXT("ElementOffsets_"));
+
+const FString UNiagaraDataInterfaceRigidMeshCollisionQuery::WorldTransformBufferName(TEXT("WorldTransformBuffer_"));
+const FString UNiagaraDataInterfaceRigidMeshCollisionQuery::InverseTransformBufferName(TEXT("InverseTransformBuffer_"));
+const FString UNiagaraDataInterfaceRigidMeshCollisionQuery::ElementExtentBufferName(TEXT("ElementExtentBuffer_"));
+const FString UNiagaraDataInterfaceRigidMeshCollisionQuery::PhysicsTypeBufferName(TEXT("PhysicsTypeBuffer_"));
+const FString UNiagaraDataInterfaceRigidMeshCollisionQuery::DFIndexBufferName(TEXT("DFIndexBuffer_"));
+
 //------------------------------------------------------------------------------------------------------------
 
 void FNDIRigidMeshCollisionBuffer::InitRHI()
 {
+	using namespace NDIRigidMeshCollisionLocal;
+
 	CreateInternalBuffer<FVector4f, EPixelFormat::PF_A32B32G32R32F>(WorldTransformBuffer, 3 * MaxNumTransforms);
 	CreateInternalBuffer<FVector4f, EPixelFormat::PF_A32B32G32R32F>(InverseTransformBuffer, 3 * MaxNumTransforms);
 
@@ -441,7 +517,7 @@ void FNDIRigidMeshCollisionBuffer::ReleaseRHI()
 
 
 
-void FNDIRigidMeshCollisionData::Release()
+void FNDIRigidMeshCollisionData::ReleaseBuffers()
 {
 	if (AssetBuffer)
 	{
@@ -455,103 +531,246 @@ void FNDIRigidMeshCollisionData::Release()
 	}
 }
 
-void FNDIRigidMeshCollisionData::Init(UNiagaraDataInterfaceRigidMeshCollisionQuery* Interface, FNiagaraSystemInstance* SystemInstance)
+bool FNDIRigidMeshCollisionData::HasActors() const
 {
-	AssetBuffer = nullptr;
+	return !ExplicitActors.IsEmpty() || !FoundActors.IsEmpty();
+}
 
-	if (Interface != nullptr && SystemInstance != nullptr)
-	{
-		UWorld* World = SystemInstance->GetWorld();
+bool FNDIRigidMeshCollisionData::ShouldRunGlobalSearch(UNiagaraDataInterfaceRigidMeshCollisionQuery* Interface) const
+{
+	return Interface->GlobalSearchAllowed && (Interface->GlobalSearchForced || (Interface->GlobalSearchFallback_Unscripted && !bHasScriptedFindActor));
+}
 
-		Actors.Empty();
+void FNDIRigidMeshCollisionData::MergeActors(FMergedActorArray& MergedActors) const
+{
+	MergedActors.Reserve(ExplicitActors.Num() + FoundActors.Num());
 
-		for (TActorIterator<AActor> It(World); It; ++It)
+	auto AppendActors = [&](const TWeakObjectPtr<AActor>& ActorPtr)
+	{
+		if (AActor* Actor = ActorPtr.Get())
 		{
-			AActor* Actor = *It;
+			MergedActors.AddUnique(Actor);
+		}
+	};
 
-			if ((!Interface->OnlyUseMoveable || (Interface->OnlyUseMoveable && Actor->IsRootComponentMovable())) &&
-				(Interface->Tag == FString("") || (Interface->Tag != FString("") && Actor->Tags.Contains(FName(Interface->Tag)))))
-			{
-				Actors.Add(TWeakObjectPtr<AActor>(Actor));
-			}
-		}	
+	Algo::ForEach(ExplicitActors, AppendActors);
+	Algo::ForEach(FoundActors, AppendActors);
+}
 
-		if (0 < Actors.Num() && Actors[0] != nullptr)
+void FNDIRigidMeshCollisionData::Init(int32 MaxNumPrimitives)
+{
+	const bool bHasActors = HasActors();
+	const bool bWasInitialized = AssetArrays.IsValid();
+
+	if (bHasActors)
+	{
+		if (!bWasInitialized)
 		{
-			// @note: we are not creating internal arrays here and letting it happen on the call to update
-			 //CreateInternalArrays(Actors, &AssetArrays);
+			AssetArrays = MakeUnique<FNDIRigidMeshCollisionArrays>(MaxNumPrimitives);
+
+			AssetBuffer = new FNDIRigidMeshCollisionBuffer();
+			AssetBuffer->SetMaxNumPrimitives(MaxNumPrimitives);
+
+			BeginInitResource(AssetBuffer);
 		}
-		AssetArrays = new FNDIRigidMeshCollisionArrays();
-		AssetArrays->Resize(Interface->MaxNumPrimitives);
 
-		AssetBuffer = new FNDIRigidMeshCollisionBuffer();
-		AssetBuffer->SetMaxNumPrimitives(Interface->MaxNumPrimitives);		
-		BeginInitResource(AssetBuffer);
+		AssetArrays->Reset();
+	}
+	else if (bWasInitialized)
+	{
+		AssetArrays = nullptr;
+		ReleaseBuffers();
 	}
+
+	bFoundActorsUpdated = false;
+	bRequiresFullUpdate = true;
 }
 
-void FNDIRigidMeshCollisionData::Update(UNiagaraDataInterfaceRigidMeshCollisionQuery* Interface, FNiagaraSystemInstance* SystemInstance)
+void FNDIRigidMeshCollisionData::Update(UNiagaraDataInterfaceRigidMeshCollisionQuery* Interface)
 {
-	if (Interface != nullptr && SystemInstance != nullptr)
-	{		
-		TArray<TWeakObjectPtr<AActor> > ActorsTmp;		
+	using namespace NDIRigidMeshCollisionLocal;
 
+	if (!Interface || !SystemInstance || !bRequiresSourceActors)
+	{
+		return;
+	}
 
-		// #todo(dmp): loop over actors and count number of potential colliders.  This might change frame to frame and we might need a
-		// better way of doing this
-		UWorld* World = SystemInstance->GetWorld();
-		for (TActorIterator<AActor> It(World); It; ++It)
+	const bool bExplicitActorsChanged = Interface->GetExplicitActors(*this);
+	if (ShouldRunGlobalSearch(Interface))
+	{
+		if (Interface->GlobalFindActors(SystemInstance->GetWorld(), *this))
 		{
-			AActor* Actor = *It;
-
-			if ((!Interface->OnlyUseMoveable || (Interface->OnlyUseMoveable && Actor->IsRootComponentMovable())) &&
-				(Interface->Tag == FString("") || (Interface->Tag != FString("") && Actor->Tags.Contains(FName(Interface->Tag)))))
-			{
-				ActorsTmp.Add(TWeakObjectPtr<AActor>(Actor));
-			}
+			bFoundActorsUpdated = true;
 		}
-		
+	}
 
-		// check if the number of active actors is the same.  If not, reinitialize
-		if (ActorsTmp.Num() != Actors.Num())
-		{
-			Init(Interface, SystemInstance);
-		}
+	// see if we need to reinitialize the internals
+	const bool bAlreadyInited = AssetArrays != nullptr;
+	const bool bHasActors = HasActors();
+	if (bAlreadyInited != bHasActors)
+	{
+		Init(Interface->MaxNumPrimitives);
+	}
 
-		TickingGroup = ComputeTickingGroup();
+	if (bHasActors)
+	{
+		FMergedActorArray MergedActors;
+		MergeActors(MergedActors);
 
-		if (0 < Actors.Num() && Actors[0].IsValid())
-		{
-			UpdateInternalArrays(Actors, AssetArrays, FVector(SystemInstance->GetLWCTile()));
-		}
+		const bool bFullUpdate = bRequiresFullUpdate || bExplicitActorsChanged || bFoundActorsUpdated;
+		UpdateInternalArrays(MergedActors, Interface->ComponentTags, FVector(SystemInstance->GetLWCTile()), bFullUpdate, AssetArrays.Get());
 	}
+
+	bRequiresFullUpdate = false;
+	bFoundActorsUpdated = false;
+}
+
+//------------------------------------------------------------------------------------------------------------
+
+/** Proxy to send data to gpu */
+struct FNDIRigidMeshCollisionProxy : public FNiagaraDataInterfaceProxy
+{
+	struct FGameThreadData
+	{
+		FElementOffset ElementOffsets;
+		TArray<FVector4f> WorldTransform;
+		TArray<FVector4f> InverseTransform;
+		TArray<FVector4f> ElementExtent;
+		TArray<uint32> PhysicsType;
+		TArray<int32> ComponentIdIndex;
+		uint32 MaxPrimitiveCount;
+		TArray<FPrimitiveComponentId> UniqueComponentIds;
+
+		FNDIRigidMeshCollisionBuffer* AssetBuffer = nullptr;
+	};
+
+	/** Get the size of the data that will be passed to render*/
+	virtual int32 PerInstanceDataPassedToRenderThreadSize() const override;
+
+	/** Get the data that will be passed to render*/
+	virtual void ConsumePerInstanceDataFromGameThread(void* PerInstanceData, const FNiagaraSystemInstanceID& Instance) override;
+
+	/** Initialize the Proxy data buffer */
+	void InitializePerInstanceData(const FNiagaraSystemInstanceID& SystemInstance);
+
+	/** Destroy the proxy data if necessary */
+	void DestroyPerInstanceData(const FNiagaraSystemInstanceID& SystemInstance);
+
+	/** Launch all pre stage functions */
+	virtual void PreStage(FRHICommandList& RHICmdList, const FNiagaraDataInterfaceStageArgs& Context) override;
+
+	/** Reset the buffers  */
+	virtual void ResetData(FRHICommandList& RHICmdList, const FNiagaraDataInterfaceArgs& Context) override;
+
+	/** List of proxy data for each system instances*/
+	TMap<FNiagaraSystemInstanceID, FGameThreadData> SystemInstancesToProxyData;
+};
+
+int32 FNDIRigidMeshCollisionProxy::PerInstanceDataPassedToRenderThreadSize() const
+{
+	return sizeof(FGameThreadData);
 }
 
-ETickingGroup FNDIRigidMeshCollisionData::ComputeTickingGroup()
-{
-	TickingGroup = NiagaraFirstTickGroup;
-	for (int32 ComponentIndex = 0; ComponentIndex < Actors.Num(); ++ComponentIndex)
-	{
-		TWeakObjectPtr<AActor> Actor = Actors[ComponentIndex].Get();
-		if (Actor.IsValid())
-		{			
-		
-			UStaticMeshComponent* Component = Cast<UStaticMeshComponent>(Actor->GetComponentByClass(UStaticMeshComponent::StaticClass()));
+void FNDIRigidMeshCollisionProxy::ConsumePerInstanceDataFromGameThread(void* PerInstanceData, const FNiagaraSystemInstanceID& Instance)
+{
+	check(IsInRenderingThread());
+
+	const FGameThreadData* SourceData = reinterpret_cast<FGameThreadData*>(PerInstanceData);
+	FGameThreadData& TargetData = (SystemInstancesToProxyData.FindOrAdd(Instance));
+
+	if (ensure(SourceData))
+	{
+		TargetData = *SourceData;
+		SourceData->~FGameThreadData();
+	}
+}
+
+void FNDIRigidMeshCollisionProxy::InitializePerInstanceData(const FNiagaraSystemInstanceID& SystemInstance)
+{
+	check(IsInRenderingThread());
+
+	check(!SystemInstancesToProxyData.Contains(SystemInstance));
+	SystemInstancesToProxyData.Add(SystemInstance);
+}
+
+void FNDIRigidMeshCollisionProxy::DestroyPerInstanceData(const FNiagaraSystemInstanceID& SystemInstance)
+{
+	check(IsInRenderingThread());	
+
+	SystemInstancesToProxyData.Remove(SystemInstance);
+}
+
+void FNDIRigidMeshCollisionProxy::PreStage(FRHICommandList& RHICmdList, const FNiagaraDataInterfaceStageArgs& Context)
+{
+	using namespace NDIRigidMeshCollisionLocal;
+
+	check(SystemInstancesToProxyData.Contains(Context.SystemInstanceID));
+
+	FGameThreadData* ProxyData =
+		SystemInstancesToProxyData.Find(Context.SystemInstanceID);
+
+	if (ProxyData != nullptr && ProxyData->AssetBuffer)
+	{
+		if (Context.SimStageData->bFirstStage)
+		{
+			FRHITransitionInfo Transitions[] = {
+				FRHITransitionInfo(ProxyData->AssetBuffer->WorldTransformBuffer.UAV, ERHIAccess::Unknown, ERHIAccess::UAVCompute),
+				FRHITransitionInfo(ProxyData->AssetBuffer->InverseTransformBuffer.UAV, ERHIAccess::Unknown, ERHIAccess::UAVCompute),
+				FRHITransitionInfo(ProxyData->AssetBuffer->PhysicsTypeBuffer.UAV, ERHIAccess::Unknown, ERHIAccess::UAVCompute),
+				FRHITransitionInfo(ProxyData->AssetBuffer->ElementExtentBuffer.UAV, ERHIAccess::Unknown, ERHIAccess::UAVCompute),
+				FRHITransitionInfo(ProxyData->AssetBuffer->DFIndexBuffer.UAV, ERHIAccess::Unknown, ERHIAccess::UAVCompute)
+			};
+			RHICmdList.Transition(MakeArrayView(Transitions, UE_ARRAY_COUNT(Transitions)));
+
+			UpdateInternalBuffer<FVector4f, EPixelFormat::PF_A32B32G32R32F>(ProxyData->WorldTransform, ProxyData->AssetBuffer->WorldTransformBuffer);
+			UpdateInternalBuffer<FVector4f, EPixelFormat::PF_A32B32G32R32F>(ProxyData->InverseTransform, ProxyData->AssetBuffer->InverseTransformBuffer);
+			UpdateInternalBuffer<FVector4f, EPixelFormat::PF_A32B32G32R32F>(ProxyData->ElementExtent, ProxyData->AssetBuffer->ElementExtentBuffer);
+			UpdateInternalBuffer<uint32, EPixelFormat::PF_R32_UINT>(ProxyData->PhysicsType, ProxyData->AssetBuffer->PhysicsTypeBuffer);
 
-			if (Component == nullptr)
+			// the distance field indexing needs to be generated using the scene
+			if (!ProxyData->ComponentIdIndex.IsEmpty() && ProxyData->AssetBuffer->DFIndexBuffer.Buffer.IsValid())
 			{
-				continue;
-			}
-				const ETickingGroup ComponentTickGroup = FMath::Max(Component->PrimaryComponentTick.TickGroup, Component->PrimaryComponentTick.EndTickGroup);
-				const ETickingGroup PhysicsTickGroup = ComponentTickGroup;
-				const ETickingGroup ClampedTickGroup = FMath::Clamp(static_cast<ETickingGroup>(PhysicsTickGroup + 1), NiagaraFirstTickGroup, NiagaraLastTickGroup);		
+				const int32 ElementCount = ProxyData->ComponentIdIndex.Num();
+				const uint32 BufferBytes = sizeof(uint32) * ElementCount;
+				void* BufferData = RHILockBuffer(ProxyData->AssetBuffer->DFIndexBuffer.Buffer, 0, BufferBytes, RLM_WriteOnly);
 
-			TickingGroup = FMath::Max(TickingGroup, ClampedTickGroup);
+				const FScene* Scene = Context.ComputeDispatchInterface->GetScene();
+				if (Scene && !ProxyData->UniqueComponentIds.IsEmpty())
+				{
+					TArray<uint32> UniqueDistanceFieldIndices;
+					UniqueDistanceFieldIndices.Reserve(ProxyData->UniqueComponentIds.Num());
+
+					for (const FPrimitiveComponentId& ComponentId : ProxyData->UniqueComponentIds)
+					{
+						uint32& DistanceFieldIndex = UniqueDistanceFieldIndices.Add_GetRef(INDEX_NONE);
+						const int32 PrimitiveSceneIndex = Scene->PrimitiveComponentIds.Find(ComponentId);
+						if (PrimitiveSceneIndex != INDEX_NONE)
+						{
+							const TArray<int32, TInlineAllocator<1>>& DFIndices = Scene->Primitives[PrimitiveSceneIndex]->DistanceFieldInstanceIndices;
+							DistanceFieldIndex = DFIndices.IsEmpty() ? INDEX_NONE : DFIndices[0];
+						}
+					}
+
+					TArrayView<uint32> BufferView(reinterpret_cast<uint32*>(BufferData), ElementCount);
+					for (int32 ElementIt = 0; ElementIt < ElementCount; ++ElementIt)
+					{
+						const int32 UniqueIdIndex = ProxyData->ComponentIdIndex[ElementIt];
+						BufferView[ElementIt] = UniqueDistanceFieldIndices.IsValidIndex(UniqueIdIndex) ? UniqueDistanceFieldIndices[UniqueIdIndex] : INDEX_NONE;
+					}
+				}
+				else
+				{
+					FMemory::Memset(BufferData, 0xFF, BufferBytes);
+				}
+				RHIUnlockBuffer(ProxyData->AssetBuffer->DFIndexBuffer.Buffer);
+			}
 		}
 	}
-	return TickingGroup;
 }
 
+void FNDIRigidMeshCollisionProxy::ResetData(FRHICommandList& RHICmdList, const FNiagaraDataInterfaceArgs& Context)
+{}
+
 //------------------------------------------------------------------------------------------------------------
 
 struct FNDIRigidMeshCollisionParametersCS : public FNiagaraDataInterfaceParametersCS
@@ -560,6 +779,8 @@ struct FNDIRigidMeshCollisionParametersCS : public FNiagaraDataInterfaceParamete
 public:
 	void Bind(const FNiagaraDataInterfaceGPUParamInfo& ParameterInfo, const class FShaderParameterMap& ParameterMap)
 	{
+		using namespace NDIRigidMeshCollisionLocal;
+
 		FNDIRigidMeshCollisionParametersName ParamNames(*ParameterInfo.DataInterfaceHLSLSymbol);
 
 		MaxTransforms.Bind(ParameterMap, *ParamNames.MaxTransformsName);
@@ -571,7 +792,7 @@ public:
 		WorldTransformBuffer.Bind(ParameterMap, *ParamNames.WorldTransformBufferName);
 		InverseTransformBuffer.Bind(ParameterMap, *ParamNames.InverseTransformBufferName);
 		ElementExtentBuffer.Bind(ParameterMap, *ParamNames.ElementExtentBufferName);
-		PhysicsTypeBuffer.Bind(ParameterMap, *ParamNames.PhysicsTypeBufferName);	
+		PhysicsTypeBuffer.Bind(ParameterMap, *ParamNames.PhysicsTypeBufferName);
 		DFIndexBuffer.Bind(ParameterMap, *ParamNames.DFIndexBufferName);
 
 		DistanceFieldParameters.Bind(ParameterMap);
@@ -579,13 +800,15 @@ public:
 
 	void Set(FRHICommandList& RHICmdList, const FNiagaraDataInterfaceSetArgs& Context) const
 	{
+		using namespace NDIRigidMeshCollisionLocal;
+
 		check(IsInRenderingThread());
 
 		FRHIComputeShader* ComputeShaderRHI = RHICmdList.GetBoundComputeShader();
 
 		FNDIRigidMeshCollisionProxy* InterfaceProxy =
 			static_cast<FNDIRigidMeshCollisionProxy*>(Context.DataInterface);
-		FNDIRigidMeshCollisionData* ProxyData =
+		FNDIRigidMeshCollisionProxy::FGameThreadData* ProxyData =
 			InterfaceProxy->SystemInstancesToProxyData.Find(Context.SystemInstanceID);
 
 		if (ProxyData != nullptr && ProxyData->AssetBuffer && ProxyData->AssetBuffer->IsInitialized())
@@ -608,22 +831,22 @@ public:
 			SetSRVParameter(RHICmdList, ComputeShaderRHI, DFIndexBuffer, AssetBuffer->DFIndexBuffer.SRV);
 
 
-			SetShaderValue(RHICmdList, ComputeShaderRHI, MaxTransforms, ProxyData->AssetArrays->MaxTransforms);
+			SetShaderValue(RHICmdList, ComputeShaderRHI, MaxTransforms, ProxyData->MaxPrimitiveCount * 2);
 			SetShaderValue(RHICmdList, ComputeShaderRHI, CurrentOffset, 0);
-			SetShaderValue(RHICmdList, ComputeShaderRHI, PreviousOffset, ProxyData->AssetArrays->MaxTransforms);
+			SetShaderValue(RHICmdList, ComputeShaderRHI, PreviousOffset, ProxyData->MaxPrimitiveCount * 3);
 
-			SetShaderValue(RHICmdList, ComputeShaderRHI, ElementOffsets, ProxyData->AssetArrays->ElementOffsets);
+			SetShaderValue(RHICmdList, ComputeShaderRHI, ElementOffsets, ProxyData->ElementOffsets);
 
 			if (DistanceFieldParameters.IsBound())
 			{
-				const FDistanceFieldSceneData *DistanceFieldSceneData = static_cast<const FNiagaraGpuComputeDispatch*>(Context.ComputeDispatchInterface)->GetMeshDistanceFieldParameters();	//-BATCHERTODO:
+				const FDistanceFieldSceneData* DistanceFieldSceneData = static_cast<const FNiagaraGpuComputeDispatch*>(Context.ComputeDispatchInterface)->GetMeshDistanceFieldParameters();	//-BATCHERTODO:
 
 				if (DistanceFieldSceneData == nullptr)
 				{
 					// UE_LOG(LogRigidMeshCollision, Error, TEXT("Distance fields are not available for use"));
 					// #todo(dmp): for now, we'll disable collisions when distance field data is not available
 					// There is no Dummy distance field data we can use.
-									
+
 					//FDistanceFieldSceneData DummyDistanceFieldSceneData(Context.Shader->GetShaderPlatform());
 					//DistanceFieldParameters.SetEmpty(RHICmdList, ComputeShaderRHI, DummyDistanceFieldSceneData);
 
@@ -639,10 +862,10 @@ public:
 					SetShaderValue(RHICmdList, ComputeShaderRHI, PreviousOffset, 0);
 					SetShaderValue(RHICmdList, ComputeShaderRHI, ElementOffsets, DummyOffsets);
 				}
-				else 
+				else
 				{
 					DistanceFieldParameters.Set(RHICmdList, ComputeShaderRHI, DistanceFieldSceneData);
-				}				
+				}
 			}
 		}
 		else
@@ -657,7 +880,7 @@ public:
 			SetShaderValue(RHICmdList, ComputeShaderRHI, MaxTransforms, 0);
 			SetShaderValue(RHICmdList, ComputeShaderRHI, CurrentOffset, 0);
 			SetShaderValue(RHICmdList, ComputeShaderRHI, PreviousOffset, 0);
-			SetShaderValue(RHICmdList, ComputeShaderRHI, ElementOffsets, DummyOffsets);	
+			SetShaderValue(RHICmdList, ComputeShaderRHI, ElementOffsets, DummyOffsets);
 		}
 	}
 
@@ -679,125 +902,181 @@ private:
 	LAYOUT_FIELD(FShaderResourceParameter, PhysicsTypeBuffer);
 	LAYOUT_FIELD(FShaderResourceParameter, DFIndexBuffer);
 
-	LAYOUT_FIELD(FDistanceFieldParameters, DistanceFieldParameters);	
+	LAYOUT_FIELD(FDistanceFieldParameters, DistanceFieldParameters);
 };
 
 IMPLEMENT_TYPE_LAYOUT(FNDIRigidMeshCollisionParametersCS);
 
 IMPLEMENT_NIAGARA_DI_PARAMETER(UNiagaraDataInterfaceRigidMeshCollisionQuery, FNDIRigidMeshCollisionParametersCS);
 
-
 //------------------------------------------------------------------------------------------------------------
 
-void FNDIRigidMeshCollisionProxy::ConsumePerInstanceDataFromGameThread(void* PerInstanceData, const FNiagaraSystemInstanceID& Instance)
+UNiagaraDataInterfaceRigidMeshCollisionQuery::UNiagaraDataInterfaceRigidMeshCollisionQuery(FObjectInitializer const& ObjectInitializer)
+	: Super(ObjectInitializer)	
 {
-	check(IsInRenderingThread());
+	Proxy.Reset(new FNDIRigidMeshCollisionProxy());
+}
+
+#if WITH_NIAGARA_DEBUGGER
+
+void UNiagaraDataInterfaceRigidMeshCollisionQuery::DrawDebugHud(UCanvas* Canvas, FNiagaraSystemInstance* SystemInstance, FString& VariableDataString, bool bVerbose) const
+{
+	FNDIRigidMeshCollisionData* InstanceData_GT = SystemInstance->FindTypedDataInterfaceInstanceData<FNDIRigidMeshCollisionData>(this);
+	if (InstanceData_GT == nullptr || !InstanceData_GT->AssetArrays.IsValid())
+	{
+		return;
+	}
+
+	const FElementOffset& ElementOffsets = InstanceData_GT->AssetArrays->ElementOffsets;
+
+	const uint32 BoxCount = ElementOffsets.SphereOffset - ElementOffsets.BoxOffset;
+	const uint32 SphereCount = ElementOffsets.CapsuleOffset - ElementOffsets.SphereOffset;
+	const uint32 CapsuleCount = ElementOffsets.NumElements - ElementOffsets.CapsuleOffset;
 
-	FNDIRigidMeshCollisionData* SourceData = static_cast<FNDIRigidMeshCollisionData*>(PerInstanceData);	
-	FNDIRigidMeshCollisionData* TargetData = &(SystemInstancesToProxyData.FindOrAdd(Instance));
+	VariableDataString = FString::Printf(TEXT("Boxes(%d) Spheres(%d) Capsules(%d)"), BoxCount, SphereCount, CapsuleCount);
 
-	ensure(TargetData);
-	if (TargetData)
+	auto GetCurrentTransform = [&](int32 ElementIndex)
 	{
-		TargetData->AssetBuffer = SourceData->AssetBuffer;		
-		TargetData->AssetArrays = SourceData->AssetArrays;
-		TargetData->TickingGroup = SourceData->TickingGroup;
+		const uint32 ElementOffset = 3 * ElementIndex;
+		FVector4f* TransformVec = InstanceData_GT->AssetArrays->CurrentTransform.GetData() + ElementOffset;
 
-		// loop over proxies and compute df indices on the RT only on new data
-		// @todo(dmp): for now we do this every frame because it seems like sometimes DFindices are not set
+		FMatrix ElementMatrix;
+		ElementMatrix.SetIdentity();
 
-		for (uint32 i = 0; i < SourceData->AssetArrays->ElementOffsets.NumElements; ++i)
+		for (int32 RowIt = 0; RowIt < 3; ++RowIt)
 		{
-			FPrimitiveSceneProxy* Proxy = SourceData->AssetArrays->SourceSceneProxy[i];
-								
-			if (Proxy != nullptr && Proxy->GetPrimitiveSceneInfo() != nullptr)
+			for (int32 ColIt = 0; ColIt < 4; ++ColIt)
 			{
-				const TArray<int32, TInlineAllocator<1>>& DFIndices = Proxy->GetPrimitiveSceneInfo()->DistanceFieldInstanceIndices;				
-				TargetData->AssetArrays->DFIndex[i] = DFIndices.Num() > 0 ? DFIndices[0] : -1;
+				ElementMatrix.M[RowIt][ColIt] = TransformVec[RowIt][ColIt];
 			}
-			else
+		}
+
+		return ElementMatrix.GetTransposed();
+	};
+
+	if (bVerbose)
+	{
+		// the DrawDebugCanvas* functions don't reasoanbly handle the near clip plane (both in terms of clipping and in terms of
+		// objects being behind the camera); so we introduce this culling behavior to work around it
+		auto ShouldClip = [&](UCanvas* Canvas, const FMatrix& Transform, const FBoxSphereBounds& Bounds)
+		{
+			const FVector Origin = Transform.TransformPosition(Bounds.Origin);
+			return (Canvas->Project(Origin).GetMin() < UE_KINDA_SMALL_NUMBER);
+		};
+
+		// Boxes
+		for (uint32 BoxIt = 0; BoxIt < BoxCount; ++BoxIt)
+		{
+			const FVector3f HalfBoxExtent = 0.5f * InstanceData_GT->AssetArrays->ElementExtent[ElementOffsets.BoxOffset + BoxIt];
+			const FBox Box(-HalfBoxExtent, HalfBoxExtent);
+			const FMatrix CurrentTransform = GetCurrentTransform(ElementOffsets.BoxOffset + BoxIt);
+			if (!ShouldClip(Canvas, CurrentTransform, FSphere(FVector::ZeroVector, HalfBoxExtent.Size())))
 			{
-				TargetData->AssetArrays->DFIndex[i] = -1;
+				DrawDebugCanvasWireBox(Canvas, CurrentTransform, Box, FColor::Blue);
 			}
-		}	
-	}
-	else
-	{
-		UE_LOG(LogRigidMeshCollision, Log, TEXT("ConsumePerInstanceDataFromGameThread() ... could not find %d"), Instance);
-	}
-	SourceData->~FNDIRigidMeshCollisionData();
-}
+		}
 
-void FNDIRigidMeshCollisionProxy::InitializePerInstanceData(const FNiagaraSystemInstanceID& SystemInstance)
-{
-	check(IsInRenderingThread());
+		// Spheres
+		for (uint32 SphereIt = 0; SphereIt < SphereCount; ++SphereIt)
+		{
+			const float Radius = InstanceData_GT->AssetArrays->ElementExtent[ElementOffsets.SphereOffset + SphereIt].X;
+			const FMatrix CurrentTransform = GetCurrentTransform(ElementOffsets.SphereOffset + SphereIt);
+			if (!ShouldClip(Canvas, CurrentTransform, FSphere(FVector::ZeroVector, Radius)))
+			{
+				DrawDebugCanvasWireSphere(Canvas, CurrentTransform.TransformPosition(FVector::ZeroVector), FColor::Blue, Radius, 20);
+			}
+		}
 
-	check(!SystemInstancesToProxyData.Contains(SystemInstance));
-	SystemInstancesToProxyData.Add(SystemInstance);
+		// Capsules
+		for (uint32 CapsuleIt = 0; CapsuleIt < CapsuleCount; ++CapsuleIt)
+		{
+			const FVector2f RadiusLength(InstanceData_GT->AssetArrays->ElementExtent[ElementOffsets.CapsuleOffset + CapsuleIt]);
+			const FMatrix CurrentTransform = GetCurrentTransform(ElementOffsets.CapsuleOffset + CapsuleIt);
+			const float HalfTotalLength = RadiusLength.X + 0.5f * RadiusLength.Y;
+			if (!ShouldClip(Canvas, CurrentTransform, FSphere(FVector::ZeroVector, HalfTotalLength)))
+			{
+				DrawDebugCanvasCapsule(Canvas, CurrentTransform, HalfTotalLength, RadiusLength.X, FColor::Blue);
+			}
+		}
+	}
 }
+#endif
 
-void FNDIRigidMeshCollisionProxy::DestroyPerInstanceData(const FNiagaraSystemInstanceID& SystemInstance)
+bool UNiagaraDataInterfaceRigidMeshCollisionQuery::InitPerInstanceData(void* PerInstanceData, FNiagaraSystemInstance* SystemInstance)
 {
-	check(IsInRenderingThread());	
-
-	SystemInstancesToProxyData.Remove(SystemInstance);
-}
+	using namespace NDIRigidMeshCollisionLocal;
 
-void FNDIRigidMeshCollisionProxy::PreStage(FRHICommandList& RHICmdList, const FNiagaraDataInterfaceStageArgs& Context)
-{
-	check(SystemInstancesToProxyData.Contains(Context.SystemInstanceID));
+	bool RequiresSourceActors = false;
+	bool HasScriptedFindActor = false;
 
-	FNDIRigidMeshCollisionData* ProxyData =
-		SystemInstancesToProxyData.Find(Context.SystemInstanceID);
+	FNiagaraDataInterfaceUtilities::ForEachGpuFunctionEquals(this, SystemInstance->GetSystem(), SystemInstance, [&](const FNiagaraDataInterfaceGeneratedFunction Function)
+	{
+		RequiresSourceActors = true;
+		return false;
+	});
 
-	if (ProxyData != nullptr && ProxyData->AssetBuffer)
+	FNiagaraDataInterfaceUtilities::ForEachVMFunctionEquals(this, SystemInstance->GetSystem(), SystemInstance, [&](const FVMExternalFunctionBindingInfo& Binding)
 	{
-		if (Context.SimStageData->bFirstStage)
+		if (Binding.Name == FindActorsName)
 		{
-			FRHITransitionInfo Transitions[] = {
-				FRHITransitionInfo(ProxyData->AssetBuffer->WorldTransformBuffer.UAV, ERHIAccess::Unknown, ERHIAccess::UAVCompute),
-				FRHITransitionInfo(ProxyData->AssetBuffer->InverseTransformBuffer.UAV, ERHIAccess::Unknown, ERHIAccess::UAVCompute),
-				FRHITransitionInfo(ProxyData->AssetBuffer->PhysicsTypeBuffer.UAV, ERHIAccess::Unknown, ERHIAccess::UAVCompute),
-				FRHITransitionInfo(ProxyData->AssetBuffer->ElementExtentBuffer.UAV, ERHIAccess::Unknown, ERHIAccess::UAVCompute),
-				FRHITransitionInfo(ProxyData->AssetBuffer->DFIndexBuffer.UAV, ERHIAccess::Unknown, ERHIAccess::UAVCompute)
-			};
-			RHICmdList.Transition(MakeArrayView(Transitions, UE_ARRAY_COUNT(Transitions)));
-
-			UpdateInternalBuffer<FVector4f, EPixelFormat::PF_A32B32G32R32F>(ProxyData->AssetArrays->WorldTransform, ProxyData->AssetBuffer->WorldTransformBuffer);
-			UpdateInternalBuffer<FVector4f, EPixelFormat::PF_A32B32G32R32F>(ProxyData->AssetArrays->InverseTransform, ProxyData->AssetBuffer->InverseTransformBuffer);
-			UpdateInternalBuffer<FVector4f, EPixelFormat::PF_A32B32G32R32F>(ProxyData->AssetArrays->ElementExtent, ProxyData->AssetBuffer->ElementExtentBuffer);
-			UpdateInternalBuffer<uint32, EPixelFormat::PF_R32_UINT>(ProxyData->AssetArrays->PhysicsType, ProxyData->AssetBuffer->PhysicsTypeBuffer);
-			UpdateInternalBuffer<uint32, EPixelFormat::PF_R32_UINT>(ProxyData->AssetArrays->DFIndex, ProxyData->AssetBuffer->DFIndexBuffer);
+			HasScriptedFindActor = true;
+			return false;
 		}
-	}
-}
-
-void FNDIRigidMeshCollisionProxy::ResetData(FRHICommandList& RHICmdList, const FNiagaraDataInterfaceArgs& Context)
-{}
+		return true;
+	});
 
-//------------------------------------------------------------------------------------------------------------
+	FNDIRigidMeshCollisionData* InstanceData = new (PerInstanceData) FNDIRigidMeshCollisionData(SystemInstance, RequiresSourceActors, HasScriptedFindActor);
 
-UNiagaraDataInterfaceRigidMeshCollisionQuery::UNiagaraDataInterfaceRigidMeshCollisionQuery(FObjectInitializer const& ObjectInitializer)
-	: Super(ObjectInitializer)	
-{
-	Proxy.Reset(new FNDIRigidMeshCollisionProxy());
-}
+	GetExplicitActors(*InstanceData);
 
-bool UNiagaraDataInterfaceRigidMeshCollisionQuery::InitPerInstanceData(void* PerInstanceData, FNiagaraSystemInstance* SystemInstance)
-{
-	FNDIRigidMeshCollisionData* InstanceData = new (PerInstanceData) FNDIRigidMeshCollisionData();
+	// if we're running a global search, then run that now
+	if (InstanceData->ShouldRunGlobalSearch(this))
+	{
+		GlobalFindActors(SystemInstance->GetWorld(), *InstanceData);
+	}
 
-	check(InstanceData);
-	InstanceData->Init(this, SystemInstance);
+	InstanceData->Init(MaxNumPrimitives);
 
 	return true;
 }
 
 ETickingGroup UNiagaraDataInterfaceRigidMeshCollisionQuery::CalculateTickGroup(const void* PerInstanceData) const
 {
+	using namespace NDIRigidMeshCollisionLocal;
+
 	if (const FNDIRigidMeshCollisionData* InstanceData = static_cast<const FNDIRigidMeshCollisionData*>(PerInstanceData))
 	{
-		return InstanceData->TickingGroup;
+		ETickingGroup TickingGroup = NiagaraFirstTickGroup;
+
+		TInlineComponentArray<UStaticMeshComponent*> StaticMeshes;
+		TInlineComponentArray<USkeletalMeshComponent*> SkeletalMeshes;
+
+		FNDIRigidMeshCollisionData::FMergedActorArray MergedActors;
+		InstanceData->MergeActors(MergedActors);
+
+		CollectComponents(MergedActors, ComponentTags, StaticMeshes);
+		CollectComponents(MergedActors, ComponentTags, SkeletalMeshes);
+
+		auto ProcessComponent = [&](const UActorComponent* Component)
+		{
+			const ETickingGroup ComponentTickGroup = FMath::Max(Component->PrimaryComponentTick.TickGroup, Component->PrimaryComponentTick.EndTickGroup);
+			const ETickingGroup PhysicsTickGroup = ComponentTickGroup;
+			const ETickingGroup ClampedTickGroup = FMath::Clamp(static_cast<ETickingGroup>(PhysicsTickGroup + 1), NiagaraFirstTickGroup, NiagaraLastTickGroup);
+
+			TickingGroup = FMath::Max(TickingGroup, ClampedTickGroup);
+		};
+
+		for (const UStaticMeshComponent* Component : StaticMeshes)
+		{
+			ProcessComponent(Component);
+		}
+
+		for (const USkeletalMeshComponent* Component : SkeletalMeshes)
+		{
+			ProcessComponent(Component);
+		}
+
+		return TickingGroup;
 	}
 	return NiagaraFirstTickGroup;
 }
@@ -806,31 +1085,24 @@ void UNiagaraDataInterfaceRigidMeshCollisionQuery::DestroyPerInstanceData(void*
 {
 	FNDIRigidMeshCollisionData* InstanceData = static_cast<FNDIRigidMeshCollisionData*>(PerInstanceData);
 
-	InstanceData->Release();	
+	InstanceData->ReleaseBuffers();	
 	InstanceData->~FNDIRigidMeshCollisionData();
 
 	FNDIRigidMeshCollisionProxy* ThisProxy = GetProxyAs<FNDIRigidMeshCollisionProxy>();
 	ENQUEUE_RENDER_COMMAND(FNiagaraDIDestroyInstanceData) (
 		[ThisProxy, InstanceID = SystemInstance->GetId()](FRHICommandListImmediate& CmdList)
 	{
-		FNDIRigidMeshCollisionData* ProxyData =
-			ThisProxy->SystemInstancesToProxyData.Find(InstanceID);
-
-		if (ProxyData != nullptr && ProxyData->AssetArrays)
-		{			
-			ThisProxy->SystemInstancesToProxyData.Remove(InstanceID);
-			delete ProxyData->AssetArrays;
-		}		
-	}
-	);
+		ThisProxy->SystemInstancesToProxyData.Remove(InstanceID);
+	});
 }
 
 bool UNiagaraDataInterfaceRigidMeshCollisionQuery::PerInstanceTick(void* PerInstanceData, FNiagaraSystemInstance* SystemInstance, float InDeltaSeconds)
 {
 	FNDIRigidMeshCollisionData* InstanceData = static_cast<FNDIRigidMeshCollisionData*>(PerInstanceData);
-	if (InstanceData && InstanceData->AssetBuffer && SystemInstance)
+	if (InstanceData && SystemInstance)
 	{
-		InstanceData->Update(this, SystemInstance);
+		check(InstanceData->SystemInstance == SystemInstance);
+		InstanceData->Update(this);
 	}
 	return false;
 }
@@ -844,8 +1116,14 @@ bool UNiagaraDataInterfaceRigidMeshCollisionQuery::CopyToInternal(UNiagaraDataIn
 
 	UNiagaraDataInterfaceRigidMeshCollisionQuery* OtherTyped = CastChecked<UNiagaraDataInterfaceRigidMeshCollisionQuery>(Destination);		
 	
-	OtherTyped->Tag = Tag;
+	OtherTyped->ActorTags = ActorTags;
+	OtherTyped->ComponentTags = ComponentTags;
+	OtherTyped->SourceActors = SourceActors;
 	OtherTyped->OnlyUseMoveable = OnlyUseMoveable;
+	OtherTyped->GlobalSearchAllowed = GlobalSearchAllowed;
+	OtherTyped->GlobalSearchForced = GlobalSearchForced;
+	OtherTyped->GlobalSearchFallback_Unscripted = GlobalSearchFallback_Unscripted;
+	OtherTyped->MaxNumPrimitives = MaxNumPrimitives;
 
 	return true;
 }
@@ -858,7 +1136,14 @@ bool UNiagaraDataInterfaceRigidMeshCollisionQuery::Equals(const UNiagaraDataInte
 	}
 	const UNiagaraDataInterfaceRigidMeshCollisionQuery* OtherTyped = CastChecked<const UNiagaraDataInterfaceRigidMeshCollisionQuery>(Other);
 
-	return  (OtherTyped->Tag == Tag) && (OtherTyped->OnlyUseMoveable == OnlyUseMoveable);
+	return (OtherTyped->ActorTags == ActorTags)
+		&& (OtherTyped->ComponentTags == ComponentTags)
+		&& (OtherTyped->SourceActors == SourceActors)
+		&& (OtherTyped->OnlyUseMoveable == OnlyUseMoveable)
+		&& (OtherTyped->GlobalSearchAllowed == GlobalSearchAllowed)
+		&& (OtherTyped->GlobalSearchForced == GlobalSearchForced)
+		&& (OtherTyped->GlobalSearchFallback_Unscripted == GlobalSearchFallback_Unscripted)
+		&& (OtherTyped->MaxNumPrimitives == MaxNumPrimitives);
 }
 
 void UNiagaraDataInterfaceRigidMeshCollisionQuery::PostInitProperties()
@@ -872,8 +1157,63 @@ void UNiagaraDataInterfaceRigidMeshCollisionQuery::PostInitProperties()
 	}
 }
 
+void UNiagaraDataInterfaceRigidMeshCollisionQuery::PostLoad()
+{
+	Super::PostLoad();
+
+#if WITH_EDITORONLY_DATA
+	if (!Tag_DEPRECATED.IsEmpty())
+	{
+		FName Tag = *Tag_DEPRECATED;
+		ActorTags.AddUnique(Tag);
+		Tag_DEPRECATED = TEXT("");
+	}
+#endif
+}
+
 void UNiagaraDataInterfaceRigidMeshCollisionQuery::GetFunctions(TArray<FNiagaraFunctionSignature>& OutFunctions)
 {
+	using namespace NDIRigidMeshCollisionLocal;
+
+	{
+		const FText OverlapOriginDescription = IF_WITH_EDITORONLY_DATA(
+			LOCTEXT("RigidBodyOverlapOriginDescription", "The center point, in world space, where the overlap trace will be performed."),
+			FText()
+		);
+
+		const FText OverlapExtentDescription = IF_WITH_EDITORONLY_DATA(
+			LOCTEXT("RigidBodyOverlapExtentDescription", "The extent, in world space, of the overlap trace."),
+			FText()
+		);
+
+		const FText TraceChannelDescription = IF_WITH_EDITORONLY_DATA(
+			LOCTEXT("RigidBodyTraceChannelDescription", "The trace channel to collide against. Trace channels can be configured in the project settings."),
+			FText()
+		);
+
+		const FText SkipOverlapDescription = IF_WITH_EDITORONLY_DATA(
+			LOCTEXT("RigidBodySkipTraceDescription", "If enabled, the overlap test will not be performed."),
+			FText()
+		);
+
+		FNiagaraFunctionSignature Sig;
+		Sig.Name = FindActorsName;
+		Sig.SetDescription(LOCTEXT("FindActorsDescription", "Triggers an overlap test on the world to find actors to represent.."));
+		Sig.SetFunctionVersion(FNiagaraRigidMeshCollisionDIFunctionVersion::LatestVersion);
+		Sig.bSupportsGPU = false;
+		Sig.bSupportsCPU = true;
+		Sig.bMemberFunction = true;
+		Sig.bRequiresExecPin = true;
+		Sig.AddInput(FNiagaraVariable(FNiagaraTypeDefinition(GetClass()), TEXT("RigidBody DI")));
+		Sig.AddInput(FNiagaraVariable(FNiagaraTypeDefinition::GetPositionDef(), TEXT("Overlap Origin")), OverlapOriginDescription);
+		Sig.AddInput(FNiagaraVariable(FNiagaraTypeDefinition::GetVec3Def(), TEXT("Overlap Extent")), OverlapExtentDescription);
+		Sig.AddInput(FNiagaraVariable(FNiagaraTypeDefinition(StaticEnum<ECollisionChannel>()), TEXT("TraceChannel")), TraceChannelDescription);
+		Sig.AddInput(FNiagaraVariable(FNiagaraTypeDefinition::GetBoolDef(), TEXT("Skip Overlap")), SkipOverlapDescription);
+		Sig.AddOutput(FNiagaraVariable(FNiagaraTypeDefinition::GetBoolDef(), TEXT("Actors Changed")));
+
+		OutFunctions.Add(Sig);
+	}
+
 	{
 		FNiagaraFunctionSignature Sig;
 		Sig.Name = GetNumBoxesName;
@@ -1061,13 +1401,28 @@ void UNiagaraDataInterfaceRigidMeshCollisionQuery::GetFunctions(TArray<FNiagaraF
 	}
 }
 
+DEFINE_NDI_DIRECT_FUNC_BINDER(UNiagaraDataInterfaceRigidMeshCollisionQuery, FindActorsCPU);
+
 void UNiagaraDataInterfaceRigidMeshCollisionQuery::GetVMExternalFunction(const FVMExternalFunctionBindingInfo& BindingInfo, void* InstanceData, FVMExternalFunction& OutFunc)
 {
+	using namespace NDIRigidMeshCollisionLocal;
+
+	if (BindingInfo.Name == FindActorsName)
+	{
+		NDI_FUNC_BINDER(UNiagaraDataInterfaceRigidMeshCollisionQuery, FindActorsCPU)::Bind(this, OutFunc);
+	}
+	else
+	{
+		UE_LOG(LogNiagara, Display, TEXT("Could not find data interface external function in %s. %s\n"),
+			*GetPathNameSafe(this), *BindingInfo.Name.ToString());
+	}
 }
 
 #if WITH_EDITORONLY_DATA
 bool UNiagaraDataInterfaceRigidMeshCollisionQuery::GetFunctionHLSL(const FNiagaraDataInterfaceGPUParamInfo& ParamInfo, const FNiagaraDataInterfaceGeneratedFunction& FunctionInfo, int FunctionInstanceIndex, FString& OutHLSL)
 {
+	using namespace NDIRigidMeshCollisionLocal;
+
 	FNDIRigidMeshCollisionParametersName ParamNames(ParamInfo.DataInterfaceHLSLSymbol);
 
 	TMap<FString, FStringFormatArg> ArgsSample = {
@@ -1224,6 +1579,8 @@ bool UNiagaraDataInterfaceRigidMeshCollisionQuery::GetFunctionHLSL(const FNiagar
 
 bool UNiagaraDataInterfaceRigidMeshCollisionQuery::UpgradeFunctionCall(FNiagaraFunctionSignature& FunctionSignature)
 {
+	using namespace NDIRigidMeshCollisionLocal;
+
 	bool bChanged = false;
 	
 	// upgrade from lwc changes, only parameter types changed there
@@ -1306,6 +1663,8 @@ void UNiagaraDataInterfaceRigidMeshCollisionQuery::GetParameterDefinitionHLSL(co
 #if WITH_EDITOR
 void UNiagaraDataInterfaceRigidMeshCollisionQuery::ValidateFunction(const FNiagaraFunctionSignature& Function, TArray<FText>& OutValidationErrors)
 {
+	using namespace NDIRigidMeshCollisionLocal;
+
 	if (Function.Name == GetClosestPointMeshDistanceFieldName || Function.Name == GetClosestPointMeshDistanceFieldNoNormalName || Function.Name == GetClosestPointMeshDistanceFieldAccurateName)
 	{
 		if (!IsMeshDistanceFieldEnabled())
@@ -1318,18 +1677,205 @@ void UNiagaraDataInterfaceRigidMeshCollisionQuery::ValidateFunction(const FNiaga
 
 void UNiagaraDataInterfaceRigidMeshCollisionQuery::ProvidePerInstanceDataForRenderThread(void* DataForRenderThread, void* PerInstanceData, const FNiagaraSystemInstanceID& SystemInstance)
 {
-	FNDIRigidMeshCollisionData* GameThreadData = static_cast<FNDIRigidMeshCollisionData*>(PerInstanceData);
-	FNDIRigidMeshCollisionData* RenderThreadData = new(DataForRenderThread) FNDIRigidMeshCollisionData();
+	const FNDIRigidMeshCollisionData* GameThreadData = reinterpret_cast<const FNDIRigidMeshCollisionData*>(PerInstanceData);
+	FNDIRigidMeshCollisionProxy::FGameThreadData* RenderThreadData = new(DataForRenderThread) FNDIRigidMeshCollisionProxy::FGameThreadData();
+
+	if (ensure(GameThreadData != nullptr && RenderThreadData != nullptr))
+	{
+		if (const FNDIRigidMeshCollisionArrays* SourceArrayData = GameThreadData->AssetArrays.Get())
+		{
+			RenderThreadData->ElementOffsets = GameThreadData->AssetArrays->ElementOffsets;
+
+			// compact the world/inverse transforms
+			const int32 TransformVectorCount = GameThreadData->AssetArrays->MaxPrimitives * 3;
+
+			auto CompactTransforms = [&](const TArray<FVector4f>& Current, const TArray<FVector4f>& Previous, TArray<FVector4f>& Compact)
+			{
+				Compact.Reset(2 * TransformVectorCount);
+				Compact.Append(Current);
+				Compact.SetNumUninitialized(TransformVectorCount, false);
+				Compact.Append(Previous);
+				Compact.SetNumUninitialized(2 * TransformVectorCount, false);
+			};
 
-	if (GameThreadData != nullptr && RenderThreadData != nullptr)
-	{		
-		RenderThreadData->AssetBuffer = GameThreadData->AssetBuffer;				
+			CompactTransforms(GameThreadData->AssetArrays->CurrentTransform, GameThreadData->AssetArrays->PreviousTransform, RenderThreadData->WorldTransform);
+			CompactTransforms(GameThreadData->AssetArrays->CurrentInverse, GameThreadData->AssetArrays->PreviousInverse, RenderThreadData->InverseTransform);
 
-		RenderThreadData->AssetArrays = new FNDIRigidMeshCollisionArrays();
-		RenderThreadData->AssetArrays->CopyFrom(GameThreadData->AssetArrays);		
-		RenderThreadData->TickingGroup = GameThreadData->TickingGroup;
+			RenderThreadData->ElementExtent = GameThreadData->AssetArrays->ElementExtent;
+			RenderThreadData->PhysicsType = GameThreadData->AssetArrays->PhysicsType;
+			RenderThreadData->ComponentIdIndex = GameThreadData->AssetArrays->ComponentIdIndex;
+			RenderThreadData->UniqueComponentIds = GameThreadData->AssetArrays->UniqueCompnentId;
+			RenderThreadData->MaxPrimitiveCount = GameThreadData->AssetArrays->MaxPrimitives;
+			RenderThreadData->AssetBuffer = GameThreadData->AssetBuffer;
+		}
 	}
 	check(Proxy);
 }
 
+bool UNiagaraDataInterfaceRigidMeshCollisionQuery::FilterComponent(const UPrimitiveComponent* Component) const
+{
+	return !(Component->IsA<USkeletalMeshComponent>() || Component->IsA<UStaticMeshComponent>());
+}
+
+bool UNiagaraDataInterfaceRigidMeshCollisionQuery::FilterActor(const AActor* Actor) const
+{
+	if (OnlyUseMoveable && !Actor->IsRootComponentMovable())
+	{
+		return true;
+	}
+
+	if (!ActorTags.ContainsByPredicate([&](const FName& Tag) { return Tag == NAME_None || Actor->Tags.Contains(Tag); }))
+	{
+		return true;
+	}
+
+	return false;
+}
+
+bool UNiagaraDataInterfaceRigidMeshCollisionQuery::GlobalFindActors(UWorld* World, FNDIRigidMeshCollisionData& InstanceData) const
+{
+	TArray<TWeakObjectPtr<AActor>> PreviousActors;
+	Swap(InstanceData.FoundActors, PreviousActors);
+
+	if (ensure(World))
+	{
+		for (TActorIterator<AActor> It(World); It; ++It)
+		{
+			if (AActor* Actor = *It)
+			{
+				if (FilterActor(Actor))
+				{
+					continue;
+				}
+
+				InstanceData.FoundActors.AddUnique(Actor);
+			}
+		}
+	}
+
+	return PreviousActors != InstanceData.FoundActors;
+}
+
+bool UNiagaraDataInterfaceRigidMeshCollisionQuery::FindActors(UWorld* World, FNDIRigidMeshCollisionData& InstanceData, ECollisionChannel Channel, const FVector& OverlapLocation, const FVector& OverlapExtent) const
+{
+	TArray<TWeakObjectPtr<AActor>> PreviousActors;
+	Swap(InstanceData.FoundActors, PreviousActors);
+
+	if (ensure(World))
+	{
+		FCollisionObjectQueryParams ObjectParams;
+		ObjectParams.AddObjectTypesToQuery(Channel);
+
+		TArray<FOverlapResult> Overlaps;
+		FCollisionQueryParams Params(SCENE_QUERY_STAT(NiagaraRigidMeshCollisionQuery), false);
+
+		World->OverlapMultiByChannel(Overlaps, OverlapLocation, FQuat::Identity, Channel, FCollisionShape::MakeBox(0.5f * OverlapExtent), Params);
+
+		for (const FOverlapResult& OverlapResult : Overlaps)
+		{
+			if (UPrimitiveComponent* PrimitiveComponent = OverlapResult.GetComponent())
+			{
+				if (FilterComponent(PrimitiveComponent))
+				{
+					continue;
+				}
+
+				if (AActor* ComponentActor = PrimitiveComponent->GetOwner())
+				{
+					if (FilterActor(ComponentActor))
+					{
+						continue;
+					}
+					InstanceData.FoundActors.AddUnique(ComponentActor);
+				}
+			}
+		}
+	}
+
+	return PreviousActors != InstanceData.FoundActors;
+}
+
+bool UNiagaraDataInterfaceRigidMeshCollisionQuery::GetExplicitActors(FNDIRigidMeshCollisionData& InstanceData)
+{
+	if (!InstanceData.bRequiresSourceActors)
+	{
+		return false;
+	}
+
+	TArray<TWeakObjectPtr<AActor>> PreviousActors;
+	Swap(InstanceData.ExplicitActors, PreviousActors);
+
+	for (const TSoftObjectPtr<AActor>& ActorPtr : SourceActors)
+	{
+		if (AActor* Actor = ActorPtr.Get())
+		{
+			InstanceData.ExplicitActors.AddUnique(Actor);
+		}
+	}
+
+	return InstanceData.ExplicitActors != PreviousActors;
+}
+
+void UNiagaraDataInterfaceRigidMeshCollisionQuery::FindActorsCPU(FVectorVMExternalFunctionContext& Context)
+{
+	VectorVM::FUserPtrHandler<FNDIRigidMeshCollisionData> InstanceData(Context);
+
+	FNDIInputParam<FNiagaraPosition> OverlapOriginParam(Context);
+	FNDIInputParam<FVector3f> OverlapExtentParam(Context);
+	FNDIInputParam<ECollisionChannel> TraceChannelParam(Context);
+	FNDIInputParam<FNiagaraBool> SkipOverlapParam(Context);
+
+	FNDIOutputParam<FNiagaraBool> ActorsChangedParam(Context);
+
+	if (ensure(InstanceData->SystemInstance))
+	{
+		FNiagaraLWCConverter LWCConverter = InstanceData->SystemInstance->GetLWCConverter();
+		UWorld* World = InstanceData->SystemInstance->GetWorld();
+
+		if (World)
+		{
+			for (int32 i = 0; i < Context.GetNumInstances(); ++i)
+			{
+				FNiagaraPosition OverlapOrigin = OverlapOriginParam.GetAndAdvance();
+				FVector3f OverlapExtent = OverlapExtentParam.GetAndAdvance();
+				ECollisionChannel TraceChannel = TraceChannelParam.GetAndAdvance();
+				bool SkipOverlap = SkipOverlapParam.GetAndAdvance() || !InstanceData->bRequiresSourceActors;
+
+				bool ActorsChanged = false;
+
+				if (!SkipOverlap)
+				{
+					const FVector ConvertedOrigin = LWCConverter.ConvertSimulationPositionToWorld(OverlapOrigin);
+					if (FindActors(World, *InstanceData, TraceChannel, ConvertedOrigin, FVector(OverlapExtent)))
+					{
+						ActorsChanged = true;
+						InstanceData->bFoundActorsUpdated = true;
+					}
+				}
+
+				ActorsChangedParam.SetAndAdvance(ActorsChanged);
+			}
+
+			return;
+		}
+	}
+
+	for (int32 i = 0; i < Context.GetNumInstances(); ++i)
+	{
+		ActorsChangedParam.SetAndAdvance(false);
+	}
+}
+
+void UNiagaraDIRigidMeshCollisionFunctionLibrary::SetSourceActors(UNiagaraComponent* NiagaraComponent, FName OverrideName, const TArray<AActor*>& InSourceActors)
+{
+	if (UNiagaraDataInterfaceRigidMeshCollisionQuery* QueryDI = UNiagaraFunctionLibrary::GetDataInterface<UNiagaraDataInterfaceRigidMeshCollisionQuery>(NiagaraComponent, OverrideName))
+	{
+		QueryDI->SourceActors.Reset(InSourceActors.Num());
+		Algo::Transform(InSourceActors, QueryDI->SourceActors, [&](AActor* Actor)
+		{
+			return Actor;
+		});
+	}
+}
+
 #undef LOCTEXT_NAMESPACE
\ No newline at end of file
diff --git a/Engine/Plugins/FX/Niagara/Source/Niagara/Public/NiagaraDataInterfaceRigidMeshCollisionQuery.h b/Engine/Plugins/FX/Niagara/Source/Niagara/Public/NiagaraDataInterfaceRigidMeshCollisionQuery.h
--- a/Engine/Plugins/FX/Niagara/Source/Niagara/Public/NiagaraDataInterfaceRigidMeshCollisionQuery.h
+++ b/Engine/Plugins/FX/Niagara/Source/Niagara/Public/NiagaraDataInterfaceRigidMeshCollisionQuery.h
@@ -2,6 +2,7 @@
 
 #pragma once
 
+#include "Kismet/BlueprintFunctionLibrary.h"
 #include "NiagaraDataInterface.h"
 #include "NiagaraCommon.h"
 #include "VectorVM.h"
@@ -30,63 +31,36 @@ struct FElementOffset
 struct FNDIRigidMeshCollisionArrays
 {
 	FElementOffset ElementOffsets;
-	TArray<FVector4f> WorldTransform;
-	TArray<FVector4f> InverseTransform;
 	TArray<FVector4f> CurrentTransform;
 	TArray<FVector4f> CurrentInverse;
 	TArray<FVector4f> PreviousTransform;
 	TArray<FVector4f> PreviousInverse;
 	TArray<FVector4f> ElementExtent;
 	TArray<uint32> PhysicsType;
-	TArray<uint32> DFIndex;
-	TArray<FPrimitiveSceneProxy*> SourceSceneProxy;
-
-	FNDIRigidMeshCollisionArrays()
-	{
-		Resize(100);
-	}
+	TArray<int32> ComponentIdIndex;
+	TArray<FPrimitiveComponentId> UniqueCompnentId;
 
+	FNDIRigidMeshCollisionArrays() = delete;
 	FNDIRigidMeshCollisionArrays(uint32 Num)
+		: MaxPrimitives(Num)
 	{
-		Resize(Num);
+		Reset();
 	}
 
-	void CopyFrom(const FNDIRigidMeshCollisionArrays* Other)
+	void Reset()
 	{
-		Resize(Other->MaxPrimitives);
-
-		ElementOffsets = Other->ElementOffsets;
-		WorldTransform = Other->WorldTransform;
-		InverseTransform = Other->InverseTransform;
-		CurrentTransform = Other->CurrentTransform;
-		CurrentInverse = Other->CurrentInverse;
-		PreviousTransform = Other->PreviousTransform;
-		PreviousInverse = Other->PreviousInverse;
-		ElementExtent = Other->ElementExtent;
-		PhysicsType = Other->PhysicsType;
-		DFIndex = Other->DFIndex;
-		SourceSceneProxy = Other->SourceSceneProxy;
-	}
-
-	void Resize(uint32 Num)
-	{
-		MaxPrimitives = Num;
-		MaxTransforms = MaxPrimitives * 2;				
-
-		WorldTransform.Init(FVector4f(0, 0, 0, 0), 3 * MaxTransforms);
-		InverseTransform.Init(FVector4f(0, 0, 0, 0), 3 * MaxTransforms);
+		ElementOffsets = FElementOffset();
 		CurrentTransform.Init(FVector4f(0, 0, 0, 0), 3 * MaxPrimitives);
 		CurrentInverse.Init(FVector4f(0, 0, 0, 0), 3 * MaxPrimitives);
 		PreviousTransform.Init(FVector4f(0, 0, 0, 0), 3 * MaxPrimitives);
 		PreviousInverse.Init(FVector4f(0, 0, 0, 0), 3 * MaxPrimitives);
 		ElementExtent.Init(FVector4f(0, 0, 0, 0), MaxPrimitives);
 		PhysicsType.Init(0, MaxPrimitives);
-		DFIndex.Init(0, MaxPrimitives);
-		SourceSceneProxy.Init(nullptr, MaxPrimitives);
+		ComponentIdIndex.Init(INDEX_NONE, MaxPrimitives);
+		UniqueCompnentId.Reset();
 	}
 
-	uint32 MaxPrimitives = 100;
-	uint32 MaxTransforms = MaxPrimitives * 2;
+	const uint32 MaxPrimitives;
 };
 
 /** Render buffers that will be used in hlsl functions */
@@ -129,31 +103,56 @@ struct FNDIRigidMeshCollisionBuffer : public FRenderResource
 	}
 };
 
-/** Data stored per physics asset instance*/
+/** Data stored per DI instance*/
 struct FNDIRigidMeshCollisionData
 {
+	FNDIRigidMeshCollisionData(const FNiagaraSystemInstance* InSystemInstance, bool InRequiresSourceActors, bool InHasScriptedFindActor)
+	: SystemInstance(InSystemInstance)
+	, bRequiresSourceActors(InRequiresSourceActors)
+	, bHasScriptedFindActor(InHasScriptedFindActor)
+	{}
+
 	/** Initialize the cpu datas */
-	void Init(class UNiagaraDataInterfaceRigidMeshCollisionQuery* Interface, FNiagaraSystemInstance* SystemInstance);
+	void Init(int32 MaxNumPrimitives);
 
 	/** Update the gpu datas */
-	void Update(class UNiagaraDataInterfaceRigidMeshCollisionQuery* Interface, FNiagaraSystemInstance* SystemInstance);
+	void Update(UNiagaraDataInterfaceRigidMeshCollisionQuery* Interface);
 
 	/** Release the buffers */
-	void Release();
+	void ReleaseBuffers();
 
-	ETickingGroup ComputeTickingGroup();
+	bool HasActors() const;
+	bool ShouldRunGlobalSearch(UNiagaraDataInterfaceRigidMeshCollisionQuery* Interface) const;
 
-	/** The instance ticking group */
-	ETickingGroup TickingGroup;
+	using FMergedActorArray = TArray<AActor*, TInlineAllocator<16>>;
+	void MergeActors(FMergedActorArray& MergedActors) const;
 
 	/** Physics asset Gpu buffer */
 	FNDIRigidMeshCollisionBuffer* AssetBuffer = nullptr;
 
 	/** Physics asset Cpu arrays */
-	FNDIRigidMeshCollisionArrays *AssetArrays = nullptr;
+	TUniquePtr<FNDIRigidMeshCollisionArrays> AssetArrays = nullptr;
+
+	/** Source actors **/
+	TArray<TWeakObjectPtr<AActor>> ExplicitActors;
+
+	/** Source actors **/
+	TArray<TWeakObjectPtr<AActor>> FoundActors;
+
+	const FNiagaraSystemInstance* SystemInstance = nullptr;
+
+	/** If false indicates that the instance is not dependent on any RigidBodies, and processing can be skipped. */
+	const bool bRequiresSourceActors;
+
+	/** Indicates that the instance is being used in a call to FindActors */
+	const bool bHasScriptedFindActor;
+
+	/** Indicates that something has updated the list of FoundActors */
+	bool bFoundActorsUpdated = false;
+
+	/** Indicates that a full update of the arrays is required */
+	bool bRequiresFullUpdate = false;
 
-	/** Static Mesh Components **/
-	TArray<TWeakObjectPtr<AActor>> Actors;
 };
 
 /** Data Interface used to collide against static meshes - whether it is the mesh distance field or a physics asset's collision primitive */
@@ -166,22 +165,56 @@ public:
 
 	DECLARE_NIAGARA_DI_PARAMETER();
 
-	UPROPERTY(EditAnywhere, Category = "Static Mesh")
-	FString Tag = TEXT("");
+#if WITH_EDITORONLY_DATA
+	UPROPERTY()
+	FString Tag_DEPRECATED = TEXT("");
+#endif
+
+	/** Set of tags used to match against actors when searching for RigidBody providers.  Empty tags will be ignored, and only a single 
+		tag is required for an actor to be matched. */
+	UPROPERTY(EditAnywhere, Category = "Search")
+	TArray<FName> ActorTags;
+
+	/** Set of tags used to match against components when searching for RigidBody providers.  Empty tags will be ignored, and only a
+		single tag is required for a component to be matched. */
+	UPROPERTY(EditAnywhere, Category = "Source")
+	TArray<FName> ComponentTags;
 
-	UPROPERTY(EditAnywhere, Category = "Static Mesh")
+	/** Hardcoded references to actors that will be used as RigidBody providers. */
+	UPROPERTY(EditAnywhere, Category = "Source")
+	TArray<TSoftObjectPtr<AActor>> SourceActors;
+
+	/** If enabled only actors that are considered moveable will be searched for RigidBodies. */
+	UPROPERTY(EditAnywhere, Category = "Source")
 	bool OnlyUseMoveable = true;
 
+	/** If enabled the global search can be executed dependeing on GlobalSearchForced and GlobalSearchFallback_Unscripted */
+	UPROPERTY(EditAnywhere, Category = "Source", meta = (DisplayName = "Global Search Allowed"))
+	bool GlobalSearchAllowed = true;
+
+	/** If enabled the global search will be performed only if there are no explicit actors specified */
+	UPROPERTY(EditAnywhere, Category = "Source", meta = (DisplayName = "Global Search Forced", EditCondition = GlobalSearchAllowed))
+	bool GlobalSearchForced = false;
+
+	/** If enabled the global search will be performed only if there are no explicit actors specified */
+	UPROPERTY(EditAnywhere, Category = "Source", meta = (DisplayName = "Global Search Only If Unscripted", EditCondition = GlobalSearchAllowed))
+	bool GlobalSearchFallback_Unscripted = true;
+
+	/** Maximum number of RigidBody represented by this DataInterface. */
 	UPROPERTY(EditAnywhere, Category = "General")
 	int MaxNumPrimitives = 100;
 
 	/** UObject Interface */
 	virtual void PostInitProperties() override;
+	virtual void PostLoad() override;
 
 	/** UNiagaraDataInterface Interface */
 	virtual void GetFunctions(TArray<FNiagaraFunctionSignature>& OutFunctions) override;
 	virtual void GetVMExternalFunction(const FVMExternalFunctionBindingInfo& BindingInfo, void* InstanceData, FVMExternalFunction& OutFunc) override;
 	virtual bool CanExecuteOnTarget(ENiagaraSimTarget Target) const override { return Target == ENiagaraSimTarget::GPUComputeSim; }
+#if WITH_NIAGARA_DEBUGGER
+	virtual void DrawDebugHud(UCanvas* Canvas, FNiagaraSystemInstance* SystemInstance, FString& VariableDataString, bool bVerbose) const override;
+#endif
 	virtual bool InitPerInstanceData(void* PerInstanceData, FNiagaraSystemInstance* SystemInstance) override;
 	virtual void DestroyPerInstanceData(void* PerInstanceData, FNiagaraSystemInstance* SystemInstance)override;
 	virtual bool PerInstanceTick(void* PerInstanceData, FNiagaraSystemInstance* SystemInstance, float DeltaSeconds) override;
@@ -204,6 +237,14 @@ public:
 #endif
 	virtual void ProvidePerInstanceDataForRenderThread(void* DataForRenderThread, void* PerInstanceData, const FNiagaraSystemInstanceID& SystemInstance) override;
 
+	bool UpdateSourceActors(FNiagaraSystemInstance* SystemInstance, FNDIRigidMeshCollisionData& InstanceData) const;
+
+	bool GetExplicitActors(FNDIRigidMeshCollisionData& InstanceData);
+	bool FindActors(UWorld* World, FNDIRigidMeshCollisionData& InstanceData, ECollisionChannel Channel, const FVector& OverlapLocation, const FVector& OverlapExtent) const;
+	bool GlobalFindActors(UWorld* World, FNDIRigidMeshCollisionData& InstanceData) const;
+
+	void FindActorsCPU(FVectorVMExternalFunctionContext& Context);
+
 	/** Name of element offsets */
 	static const FString MaxTransformsName;
 
@@ -234,29 +275,21 @@ public:
 protected:
 	/** Copy one niagara DI to this */
 	virtual bool CopyToInternal(UNiagaraDataInterface* Destination) const override;
+
+	bool FilterComponent(const UPrimitiveComponent* Component) const;
+	bool FilterActor(const AActor* Actor) const;
 };
 
-/** Proxy to send data to gpu */
-struct FNDIRigidMeshCollisionProxy : public FNiagaraDataInterfaceProxy
+/**
+* C++ and Blueprint library for accessing array types
+*/
+UCLASS()
+class NIAGARA_API UNiagaraDIRigidMeshCollisionFunctionLibrary : public UBlueprintFunctionLibrary
 {
-	/** Get the size of the data that will be passed to render*/
-	virtual int32 PerInstanceDataPassedToRenderThreadSize() const override { return sizeof(FNDIRigidMeshCollisionData); }
-
-	/** Get the data that will be passed to render*/
-	virtual void ConsumePerInstanceDataFromGameThread(void* PerInstanceData, const FNiagaraSystemInstanceID& Instance) override;
+	GENERATED_BODY()
 
-	/** Initialize the Proxy data buffer */
-	void InitializePerInstanceData(const FNiagaraSystemInstanceID& SystemInstance);
-
-	/** Destroy the proxy data if necessary */
-	void DestroyPerInstanceData(const FNiagaraSystemInstanceID& SystemInstance);
-
-	/** Launch all pre stage functions */
-	virtual void PreStage(FRHICommandList& RHICmdList, const FNiagaraDataInterfaceStageArgs& Context) override;
-
-	/** Reset the buffers  */
-	virtual void ResetData(FRHICommandList& RHICmdList, const FNiagaraDataInterfaceArgs& Context) override;
+public:
 
-	/** List of proxy data for each system instances*/
-	TMap<FNiagaraSystemInstanceID, FNDIRigidMeshCollisionData> SystemInstancesToProxyData;
+	UFUNCTION(BlueprintCallable, Category = Niagara, meta = (DisplayName = "Niagara Set Source Actors"))
+	static void SetSourceActors(UNiagaraComponent* NiagaraSystem, FName OverrideName, const TArray<AActor*>& SourceActors);
 };
\ No newline at end of file
