Niagara - Allow numeric pins to be converted to enums again.

diff --git a/Engine/Plugins/FX/Niagara/Source/NiagaraEditor/Private/EdGraphSchema_Niagara.cpp b/Engine/Plugins/FX/Niagara/Source/NiagaraEditor/Private/EdGraphSchema_Niagara.cpp
--- a/Engine/Plugins/FX/Niagara/Source/NiagaraEditor/Private/EdGraphSchema_Niagara.cpp
+++ b/Engine/Plugins/FX/Niagara/Source/NiagaraEditor/Private/EdGraphSchema_Niagara.cpp
@@ -1986,11 +1986,17 @@ void UEdGraphSchema_Niagara::ConvertNumericPinToType(UEdGraphPin* InGraphPin, FN
 bool UEdGraphSchema_Niagara::PinTypesValidForNumericConversion(FEdGraphPinType AType, FEdGraphPinType BType) const
 {
 	if (AType == BType)
+	{
 		return true;
-	else if (AType.PinCategory == PinCategoryStaticType && BType.PinCategory == PinCategoryType)
+	}
+	else if ((AType.PinCategory == PinCategoryStaticType || AType.PinCategory == PinCategoryStaticEnum) && BType.PinCategory == PinCategoryType)
+	{
 		return true;
-	else if (AType.PinCategory == PinCategoryType && BType.PinCategory == PinCategoryType)
+	}
+	else if ((AType.PinCategory == PinCategoryType || AType.PinCategory == PinCategoryEnum) && BType.PinCategory == PinCategoryType)
+	{
 		return true;
+	}
 
 	return false;
 }
