- Add BUF_Static to buffer creation to avoid issues with some RHIs

#rnx

diff --git a/Engine/Plugins/FX/Niagara/Source/Niagara/Classes/NiagaraDataInterfaceArrayImpl.h b/Engine/Plugins/FX/Niagara/Source/Niagara/Classes/NiagaraDataInterfaceArrayImpl.h
--- a/Engine/Plugins/FX/Niagara/Source/Niagara/Classes/NiagaraDataInterfaceArrayImpl.h
+++ b/Engine/Plugins/FX/Niagara/Source/Niagara/Classes/NiagaraDataInterfaceArrayImpl.h
@@ -164,7 +164,7 @@ struct FNDIArrayInstanceData_RenderThread
 
 			// Create Buffer
 			FRHIResourceCreateInfo CreateInfo(TEXT("NiagaraDataInterfaceArray"));
-			const EBufferUsageFlags BufferUsage = BUF_ShaderResource | BUF_VertexBuffer | (IsReadOnly() ? BUF_None : BUF_UnorderedAccess | BUF_SourceCopy);
+			const EBufferUsageFlags BufferUsage = BUF_Static | BUF_ShaderResource | BUF_VertexBuffer | (IsReadOnly() ? BUF_None : BUF_UnorderedAccess | BUF_SourceCopy);
 			const ERHIAccess DefaultAccess = IsReadOnly() ? ERHIAccess::SRVCompute : ERHIAccess::UAVCompute;
 			ArrayBuffer = RHICreateBuffer(ArrayNumBytes, BufferUsage, TypeStride, DefaultAccess, CreateInfo);
 
