fix copy paste bug with missing parameter

diff --git a/Engine/Plugins/FX/Niagara/Shaders/Private/NiagaraDataInterfaceRigidMeshCollisionQuery.ush b/Engine/Plugins/FX/Niagara/Shaders/Private/NiagaraDataInterfaceRigidMeshCollisionQuery.ush
--- a/Engine/Plugins/FX/Niagara/Shaders/Private/NiagaraDataInterfaceRigidMeshCollisionQuery.ush
+++ b/Engine/Plugins/FX/Niagara/Shaders/Private/NiagaraDataInterfaceRigidMeshCollisionQuery.ush
@@ -535,7 +535,7 @@ void DIRigidMeshCollision_GetClosestPointMeshDistanceFieldNoNormal(in FDIRigidMe
 		{
 			// #todo(dmp): this does a dynamic branch based on intersecting the bbox.  Maybe we can factor that out due to the broadphase here?
 			MaxDistance = max(.001, MaxDistance);
-			OutClosestDistance = DistanceToNearestSurfaceForObject(DFIndex, WorldPosition, 1e-3);		
+			OutClosestDistance = DistanceToNearestSurfaceForObject(DFIndex, WorldPosition, MaxDistance);		
 
 			const float3 PreviousPosition = mul(DIRigidMeshCollision_GetPreviousTransform(DIContext,ElementIndex), float4(CollisionPosition,1.0)).xyz;
 			const float3 CurrentPosition = mul(DIRigidMeshCollision_GetCurrentTransform(DIContext,ElementIndex), float4(CollisionPosition,1.0)).xyz;
