Adding functions to interact with particle data payloads. Currently we just check the update or spawn script in certain situations. This can break down for GPU scripts with variables defined in different sim stages.

diff --git a/Engine/Plugins/FX/Niagara/Source/Niagara/Classes/NiagaraEmitter.h b/Engine/Plugins/FX/Niagara/Source/Niagara/Classes/NiagaraEmitter.h
--- a/Engine/Plugins/FX/Niagara/Source/Niagara/Classes/NiagaraEmitter.h
+++ b/Engine/Plugins/FX/Niagara/Source/Niagara/Classes/NiagaraEmitter.h
@@ -454,6 +454,8 @@ struct NIAGARA_API FVersionedNiagaraEmitterData
 #else
 	const TCHAR* GetDebugSimName() const { return TEXT(""); }
 #endif
+
+	void GatherCompiledParticleAttributes(TArray<FNiagaraVariable>& OutVariables) const;
 	
 private:
 	UPROPERTY()
diff --git a/Engine/Plugins/FX/Niagara/Source/Niagara/Private/NiagaraDataInterfaceParticleRead.cpp b/Engine/Plugins/FX/Niagara/Source/Niagara/Private/NiagaraDataInterfaceParticleRead.cpp
--- a/Engine/Plugins/FX/Niagara/Source/Niagara/Private/NiagaraDataInterfaceParticleRead.cpp
+++ b/Engine/Plugins/FX/Niagara/Source/Niagara/Private/NiagaraDataInterfaceParticleRead.cpp
@@ -2530,7 +2530,9 @@ void UNiagaraDataInterfaceParticleRead::GetFeedback(UNiagaraSystem* Asset, UNiag
 									return Var.GetName() == *AttributeName;
 								};
 
-								const FNiagaraVariableBase* FoundVar = FoundSourceEmitter.GetEmitterData()->SpawnScriptProps.Script->GetVMExecutableData().Attributes.FindByPredicate(AttribFilter);
+								TArray<FNiagaraVariable> Variables;
+								FoundSourceEmitter.GetEmitterData()->GatherCompiledParticleAttributes(Variables);
+								const FNiagaraVariableBase* FoundVar = Variables.FindByPredicate(AttribFilter);
 								if (FoundVar && !CheckVariableType(FoundVar->GetType(), AttributeType))
 								{
 									FText Msg = FText::Format(LOCTEXT("SourceEmitterTypeMismatchError", "Source Emitter has attribute named, \"{0}\" but the type isn't compatible with the function \"{1}\", and will not succeed."), FText::FromName(*AttributeName), FText::FromName(Func.Name));
diff --git a/Engine/Plugins/FX/Niagara/Source/Niagara/Private/NiagaraEmitter.cpp b/Engine/Plugins/FX/Niagara/Source/Niagara/Private/NiagaraEmitter.cpp
--- a/Engine/Plugins/FX/Niagara/Source/Niagara/Private/NiagaraEmitter.cpp
+++ b/Engine/Plugins/FX/Niagara/Source/Niagara/Private/NiagaraEmitter.cpp
@@ -2086,6 +2086,25 @@ void UNiagaraEmitter::RebindNotifications()
 
 #endif
 
+
+void FVersionedNiagaraEmitterData::GatherCompiledParticleAttributes(TArray<FNiagaraVariable>& OutVariables) const
+{
+	if (SimTarget == ENiagaraSimTarget::GPUComputeSim)
+	{
+		OutVariables = GetGPUComputeScript()->GetVMExecutableData().Attributes;
+	}
+	else
+	{
+		OutVariables = UpdateScriptProps.Script->GetVMExecutableData().Attributes;
+
+		for (const FNiagaraVariable& Var : SpawnScriptProps.Script->GetVMExecutableData().Attributes)
+		{
+			OutVariables.AddUnique(Var);
+		}
+	}
+
+}
+
 bool UNiagaraEmitter::CanObtainParticleAttribute(const FNiagaraVariableBase& InVar, const FGuid& EmitterVersion, FNiagaraTypeDefinition& OutBoundType) const
 {
 	check(!HasAnyFlags(RF_NeedPostLoad));
@@ -2111,6 +2130,7 @@ bool UNiagaraEmitter::CanObtainParticleAttribute(const FNiagaraVariableBase& InV
 	return false;
 }
 
+
 bool UNiagaraEmitter::CanObtainEmitterAttribute(const FNiagaraVariableBase& InVarWithUniqueNameNamespace, FNiagaraTypeDefinition& OutBoundType) const
 {
 	check(!HasAnyFlags(RF_NeedPostLoad));
diff --git a/Engine/Plugins/FX/Niagara/Source/Niagara/Private/NiagaraSystem.cpp b/Engine/Plugins/FX/Niagara/Source/Niagara/Private/NiagaraSystem.cpp
--- a/Engine/Plugins/FX/Niagara/Source/Niagara/Private/NiagaraSystem.cpp
+++ b/Engine/Plugins/FX/Niagara/Source/Niagara/Private/NiagaraSystem.cpp
@@ -2754,19 +2754,7 @@ void UNiagaraSystem::InitEmitterDataSetCompiledData(FNiagaraDataSetCompiledData&
 {
 	DataSetToInit.Empty();
 	const FVersionedNiagaraEmitterData* AssociatedEmitterData = InAssociatedEmitterHandle.GetEmitterData();
-	if (AssociatedEmitterData->SimTarget == ENiagaraSimTarget::GPUComputeSim)
-	{
-		DataSetToInit.Variables = AssociatedEmitterData->GetGPUComputeScript()->GetVMExecutableData().Attributes;
-	}
-	else
-	{
-		DataSetToInit.Variables = AssociatedEmitterData->UpdateScriptProps.Script->GetVMExecutableData().Attributes;
-
-		for (const FNiagaraVariable& Var : AssociatedEmitterData->SpawnScriptProps.Script->GetVMExecutableData().Attributes)
-		{
-			DataSetToInit.Variables.AddUnique(Var);
-		}
-	}
+	AssociatedEmitterData->GatherCompiledParticleAttributes(DataSetToInit.Variables);
 
 	DataSetToInit.bRequiresPersistentIDs = AssociatedEmitterData->RequiresPersistentIDs() || DataSetToInit.Variables.Contains(SYS_PARAM_PARTICLES_ID);
 	DataSetToInit.ID = FNiagaraDataSetID(InAssociatedEmitterHandle.GetIdName(), ENiagaraDataSetType::ParticleData);
