- Add constant material parameter bindings for renderers

diff --git a/Engine/Plugins/FX/Niagara/Source/Niagara/Private/NiagaraMeshRendererProperties.cpp b/Engine/Plugins/FX/Niagara/Source/Niagara/Private/NiagaraMeshRendererProperties.cpp
--- a/Engine/Plugins/FX/Niagara/Source/Niagara/Private/NiagaraMeshRendererProperties.cpp
+++ b/Engine/Plugins/FX/Niagara/Source/Niagara/Private/NiagaraMeshRendererProperties.cpp
@@ -301,7 +301,7 @@ void UNiagaraMeshRendererProperties::UpdateSourceModeDerivates(ENiagaraRendererS
 	FVersionedNiagaraEmitter SrcEmitter = GetOuterEmitter();
 	if (SrcEmitter.Emitter)
 	{
-		for (FNiagaraMaterialAttributeBinding& MaterialParamBinding : MaterialParameterBindings)
+		for (FNiagaraMaterialAttributeBinding& MaterialParamBinding : MaterialParameters.AttributeBindings)
 		{
 			MaterialParamBinding.CacheValues(SrcEmitter.Emitter);
 		}
@@ -502,7 +502,7 @@ bool UNiagaraMeshRendererProperties::PopulateRequiredBindings(FNiagaraParameterS
 		}
 	}
 
-	for (FNiagaraMaterialAttributeBinding& MaterialParamBinding : MaterialParameterBindings)
+	for (FNiagaraMaterialAttributeBinding& MaterialParamBinding : MaterialParameters.AttributeBindings)
 	{
 		InParameterStore.AddParameter(MaterialParamBinding.GetParamMapBindableVariable(), false);
 		bAnyAdded = true;
@@ -565,6 +565,14 @@ void UNiagaraMeshRendererProperties::PostLoad()
 			OverrideMaterial.ExplicitMat->ConditionalPostLoad();
 		}
 	}
+
+#if WITH_EDITORONLY_DATA
+	if (MaterialParameterBindings_DEPRECATED.Num() > 0)
+	{
+		MaterialParameters.AttributeBindings = MaterialParameterBindings_DEPRECATED;
+		MaterialParameterBindings_DEPRECATED.Empty();
+	}
+#endif
 }
 
 #if WITH_EDITORONLY_DATA
@@ -668,9 +676,16 @@ void UNiagaraMeshRendererProperties::GetRendererTooltipWidgets(const FNiagaraEmi
 }
 
 
-void UNiagaraMeshRendererProperties::GetRendererFeedback(const FVersionedNiagaraEmitter& InEmitter, TArray<FText>& OutErrors, TArray<FText>& OutWarnings, TArray<FText>& OutInfo) const
+void UNiagaraMeshRendererProperties::GetRendererFeedback(const FVersionedNiagaraEmitter& InEmitter, TArray<FNiagaraRendererFeedback>& OutErrors, TArray<FNiagaraRendererFeedback>& OutWarnings, TArray<FNiagaraRendererFeedback>& OutInfo) const
 {
 	Super::GetRendererFeedback(InEmitter, OutErrors, OutWarnings, OutInfo);
+
+	if (MaterialParameters.HasAnyBindings())
+	{
+		TArray<UMaterialInterface*> Materials;
+		GetUsedMaterials(nullptr, Materials);
+		MaterialParameters.GetFeedback(Materials, OutWarnings);
+	}
 }
 
 void UNiagaraMeshRendererProperties::BeginDestroy()
@@ -807,7 +822,7 @@ void UNiagaraMeshRendererProperties::RenameVariable(const FNiagaraVariableBase&
 	Super::RenameVariable(OldVariable, NewVariable, InEmitter);
 
 	// Handle renaming material bindings
-	for (FNiagaraMaterialAttributeBinding& Binding : MaterialParameterBindings)
+	for (FNiagaraMaterialAttributeBinding& Binding : MaterialParameters.AttributeBindings)
 	{
 		Binding.RenameVariableIfMatching(OldVariable, NewVariable, InEmitter.Emitter, GetCurrentSourceMode());
 	}
@@ -818,7 +833,7 @@ void UNiagaraMeshRendererProperties::RemoveVariable(const FNiagaraVariableBase&
 	Super::RemoveVariable(OldVariable, InEmitter);
 
 	// Handle resetting material bindings to defaults
-	for (FNiagaraMaterialAttributeBinding& Binding : MaterialParameterBindings)
+	for (FNiagaraMaterialAttributeBinding& Binding : MaterialParameters.AttributeBindings)
 	{
 		if (Binding.Matches(OldVariable, InEmitter.Emitter, GetCurrentSourceMode()))
 		{
diff --git a/Engine/Plugins/FX/Niagara/Source/Niagara/Private/NiagaraRenderer.cpp b/Engine/Plugins/FX/Niagara/Source/Niagara/Private/NiagaraRenderer.cpp
--- a/Engine/Plugins/FX/Niagara/Source/Niagara/Private/NiagaraRenderer.cpp
+++ b/Engine/Plugins/FX/Niagara/Source/Niagara/Private/NiagaraRenderer.cpp
@@ -493,10 +493,12 @@ bool FNiagaraRenderer::UseLocalSpace(const FNiagaraSceneProxy* Proxy)const
 	return bLocalSpace || Proxy->GetProxyDynamicData().bUseCullProxy;
 }
 
-void FNiagaraRenderer::ProcessMaterialParameterBindings(TConstArrayView< FNiagaraMaterialAttributeBinding > InMaterialParameterBindings, const FNiagaraEmitterInstance* InEmitter, TConstArrayView<UMaterialInterface*> InMaterials) const
+void FNiagaraRenderer::ProcessMaterialParameterBindings(const FNiagaraRendererMaterialParameters& MaterialParameters, const FNiagaraEmitterInstance* InEmitter, TConstArrayView<UMaterialInterface*> InMaterials) const
 {
-	if (InMaterialParameterBindings.Num() == 0 || !InEmitter)
+	if (MaterialParameters.HasAnyBindings() == false || !InEmitter)
+	{
 		return;
+	}
 
 	FNiagaraSystemInstance* SystemInstance = InEmitter->GetParentSystemInstance();
 	if (SystemInstance)
@@ -510,7 +512,7 @@ void FNiagaraRenderer::ProcessMaterialParameterBindings(TConstArrayView< FNiagar
 				UMaterialInstanceDynamic* MatDyn = Cast<UMaterialInstanceDynamic>(Mat);
 				if (MatDyn)
 				{
-					for (const FNiagaraMaterialAttributeBinding& Binding : InMaterialParameterBindings)
+					for (const FNiagaraMaterialAttributeBinding& Binding : MaterialParameters.AttributeBindings)
 					{
 
 						if (Binding.GetParamMapBindableVariable().GetType() == FNiagaraTypeDefinition::GetVec4Def() ||
@@ -567,6 +569,24 @@ void FNiagaraRenderer::ProcessMaterialParameterBindings(TConstArrayView< FNiagar
 							}
 						}
 					}
+
+					for (const FNiagaraRendererMaterialScalarParameter& ScalarParameter : MaterialParameters.ScalarParameters)
+					{
+						MatDyn->SetScalarParameterValue(ScalarParameter.MaterialParameterName, ScalarParameter.Value);
+					}
+
+					for (const FNiagaraRendererMaterialVectorParameter& VectorParameter : MaterialParameters.VectorParameters)
+					{
+						MatDyn->SetVectorParameterValue(VectorParameter.MaterialParameterName, VectorParameter.Value);
+					}
+
+					for (const FNiagaraRendererMaterialTextureParameter& TextureParameter : MaterialParameters.TextureParameters)
+					{
+						if (IsValid(TextureParameter.Texture))
+						{
+							MatDyn->SetTextureParameterValue(TextureParameter.MaterialParameterName, TextureParameter.Texture);
+						}
+					}
 				}
 			}
 		}
diff --git a/Engine/Plugins/FX/Niagara/Source/Niagara/Private/NiagaraRendererMeshes.cpp b/Engine/Plugins/FX/Niagara/Source/Niagara/Private/NiagaraRendererMeshes.cpp
--- a/Engine/Plugins/FX/Niagara/Source/Niagara/Private/NiagaraRendererMeshes.cpp
+++ b/Engine/Plugins/FX/Niagara/Source/Niagara/Private/NiagaraRendererMeshes.cpp
@@ -1770,9 +1770,9 @@ FNiagaraDynamicDataBase* FNiagaraRendererMeshes::GenerateDynamicData(const FNiag
 		DynamicData->ParameterDataBound = ParameterData.GetParameterDataArray();
 	}
 
-	if (DynamicData && Properties->MaterialParameterBindings.Num() != 0)
+	if (DynamicData && Properties->MaterialParameters.HasAnyBindings())
 	{
-		ProcessMaterialParameterBindings(MakeArrayView(Properties->MaterialParameterBindings), Emitter, MakeArrayView(BaseMaterials_GT));
+		ProcessMaterialParameterBindings(Properties->MaterialParameters, Emitter, MakeArrayView(BaseMaterials_GT));
 	}
 
 	return DynamicData;
diff --git a/Engine/Plugins/FX/Niagara/Source/Niagara/Private/NiagaraRendererProperties.cpp b/Engine/Plugins/FX/Niagara/Source/Niagara/Private/NiagaraRendererProperties.cpp
--- a/Engine/Plugins/FX/Niagara/Source/Niagara/Private/NiagaraRendererProperties.cpp
+++ b/Engine/Plugins/FX/Niagara/Source/Niagara/Private/NiagaraRendererProperties.cpp
@@ -9,8 +9,11 @@
 #include "NiagaraSettings.h"
 #include "NiagaraSystem.h"
 #include "Interfaces/ITargetPlatform.h"
+#include "Materials/MaterialInterface.h"
 #include "Styling/SlateIconFinder.h"
 
+#define LOCTEXT_NAMESPACE "UNiagaraRendererProperties"
+
 void FNiagaraRendererLayout::Initialize(int32 NumVariables)
 {
 	VFVariables_GT.Reset(NumVariables);
@@ -100,6 +103,125 @@ void FNiagaraRendererLayout::Finalize()
 	);
 }
 
+//////////////////////////////////////////////////////////////////////////
+
+#if WITH_EDITORONLY_DATA
+void FNiagaraRendererMaterialParameters::GetFeedback(TArrayView<UMaterialInterface*> Materials, TArray<FNiagaraRendererFeedback>& OutWarnings) const
+{
+	TArray<bool> AttributeBindingsValid;
+	TArray<bool> ScalarParametersValid;
+	TArray<bool> VectorParametersValid;
+	TArray<bool> TextureParametersValid;
+	AttributeBindingsValid.AddDefaulted(AttributeBindings.Num());
+	ScalarParametersValid.AddDefaulted(ScalarParameters.Num());
+	VectorParametersValid.AddDefaulted(VectorParameters.Num());
+	TextureParametersValid.AddDefaulted(TextureParameters.Num());
+
+	TArray<FMaterialParameterInfo> TempParameterInfo;
+	TArray<FGuid> TempParameterIds;
+	auto ContainsParameter =
+		[&TempParameterInfo](FName InName)
+		{
+			for ( const FMaterialParameterInfo& Parameter : TempParameterInfo )
+			{
+				if (Parameter.Name == InName)
+				{
+					return true;
+				}
+			}
+			return false;
+		};
+
+	for (UMaterialInterface* Material : Materials)
+	{
+		if (Material == nullptr)
+		{
+			continue;
+		}
+		
+		if (AttributeBindingsValid.Num() > 0 || ScalarParametersValid.Num() > 0)
+		{
+			Material->GetAllScalarParameterInfo(TempParameterInfo, TempParameterIds);
+			for (int32 i = 0; i < AttributeBindings.Num(); ++i)
+			{
+				AttributeBindingsValid[i] |= ContainsParameter(AttributeBindings[i].MaterialParameterName);
+			}
+			for (int32 i = 0; i < ScalarParameters.Num(); ++i)
+			{
+				ScalarParametersValid[i] |= ContainsParameter(ScalarParameters[i].MaterialParameterName);
+			}
+		}
+		if (AttributeBindingsValid.Num() > 0 || VectorParametersValid.Num() > 0)
+		{
+			Material->GetAllVectorParameterInfo(TempParameterInfo, TempParameterIds);
+			for (int32 i = 0; i < AttributeBindings.Num(); ++i)
+			{
+				AttributeBindingsValid[i] |= ContainsParameter(AttributeBindings[i].MaterialParameterName);
+			}
+			for (int32 i = 0; i < VectorParameters.Num(); ++i)
+			{
+				VectorParametersValid[i] |= ContainsParameter(VectorParameters[i].MaterialParameterName);
+			}
+		}
+		if (AttributeBindingsValid.Num() > 0 || TextureParametersValid.Num() > 0)
+		{
+			Material->GetAllTextureParameterInfo(TempParameterInfo, TempParameterIds);
+			for (int32 i=0; i < AttributeBindings.Num(); ++i)
+			{
+				AttributeBindingsValid[i] |= ContainsParameter(AttributeBindings[i].MaterialParameterName);
+			}
+			for (int32 i = 0; i < TextureParameters.Num(); ++i)
+			{
+				TextureParametersValid[i] |= ContainsParameter(TextureParameters[i].MaterialParameterName);
+			}
+		}
+	}
+
+	for (int32 i=0; i < AttributeBindingsValid.Num(); ++i)
+	{
+		if (AttributeBindingsValid[i] == false)
+		{
+			OutWarnings.Emplace(
+				FText::Format(LOCTEXT("AttributeBindingMissingDesc", "AttributeBinding '{0}' could not be found in the renderer materials.  We will still create the MID which may be unnecessary."), FText::FromName(AttributeBindings[i].MaterialParameterName)),
+				FText::Format(LOCTEXT("AttributeBindingMissing", "AttributeBinding '{0}' not found on materials."), FText::FromName(AttributeBindings[i].MaterialParameterName))
+			);
+		}
+	}
+	for (int32 i = 0; i < ScalarParametersValid.Num(); ++i)
+	{
+		if (ScalarParametersValid[i] == false)
+		{
+			OutWarnings.Emplace(
+				FText::Format(LOCTEXT("ScalarParameterMissingDesc", "ScalarParameter '{0}' could not be found in the renderer materials.  We will still create the MID which may be unnecessary."), FText::FromName(ScalarParameters[i].MaterialParameterName)),
+				FText::Format(LOCTEXT("ScalarParameterMissing", "ScalarParameter '{0}' not found on materials."), FText::FromName(ScalarParameters[i].MaterialParameterName))
+			);
+		}
+	}
+	for (int32 i = 0; i < VectorParametersValid.Num(); ++i)
+	{
+		if (VectorParametersValid[i] == false)
+		{
+			OutWarnings.Emplace(
+				FText::Format(LOCTEXT("VectorParameterMissingDesc", "VectorParameter '{0}' could not be found in the renderer materials.  We will still create the MID which may be unnecessary."), FText::FromName(VectorParameters[i].MaterialParameterName)),
+				FText::Format(LOCTEXT("VectorParameterMissing", "VectorParameter '{0}' not found on materials."), FText::FromName(VectorParameters[i].MaterialParameterName))
+			);
+		}
+	}
+	for (int32 i = 0; i < TextureParametersValid.Num(); ++i)
+	{
+		if (TextureParametersValid[i] == false)
+		{
+			OutWarnings.Emplace(
+				FText::Format(LOCTEXT("TextureParameterMissingDesc", "TextureParameter '{0}' could not be found in the renderer materials.  We will still create the MID which may be unnecessary."), FText::FromName(TextureParameters[i].MaterialParameterName)),
+				FText::Format(LOCTEXT("TextureParameterMissing", "TextureParameter '{0}' not found on materials."), FText::FromName(TextureParameters[i].MaterialParameterName))
+			);
+		}
+	}
+}
+#endif //WITH_EDITORONLY_DATA
+
+//////////////////////////////////////////////////////////////////////////
+
 #if WITH_EDITORONLY_DATA
 bool UNiagaraRendererProperties::IsSupportedVariableForBinding(const FNiagaraVariableBase& InSourceForBinding, const FName& InTargetBindingName) const
 {
@@ -484,3 +606,5 @@ bool UNiagaraRendererProperties::IsSortHighPrecision(ENiagaraRendererSortPrecisi
 	}
 	return SortPrecision == ENiagaraRendererSortPrecision::High;
 }
+
+#undef LOCTEXT_NAMESPACE
diff --git a/Engine/Plugins/FX/Niagara/Source/Niagara/Private/NiagaraRendererRibbons.cpp b/Engine/Plugins/FX/Niagara/Source/Niagara/Private/NiagaraRendererRibbons.cpp
--- a/Engine/Plugins/FX/Niagara/Source/Niagara/Private/NiagaraRendererRibbons.cpp
+++ b/Engine/Plugins/FX/Niagara/Source/Niagara/Private/NiagaraRendererRibbons.cpp
@@ -835,9 +835,9 @@ FNiagaraDynamicDataBase* FNiagaraRendererRibbons::GenerateDynamicData(const FNia
 			}
 		}
 
-		if (DynamicData && Properties->MaterialParameterBindings.Num() != 0)
+		if (DynamicData && Properties->MaterialParameters.HasAnyBindings())
 		{
-			ProcessMaterialParameterBindings(MakeArrayView(Properties->MaterialParameterBindings), Emitter, MakeArrayView(BaseMaterials_GT));
+			ProcessMaterialParameterBindings(Properties->MaterialParameters, Emitter, MakeArrayView(BaseMaterials_GT));
 		}
 	}
 	
diff --git a/Engine/Plugins/FX/Niagara/Source/Niagara/Private/NiagaraRendererSprites.cpp b/Engine/Plugins/FX/Niagara/Source/Niagara/Private/NiagaraRendererSprites.cpp
--- a/Engine/Plugins/FX/Niagara/Source/Niagara/Private/NiagaraRendererSprites.cpp
+++ b/Engine/Plugins/FX/Niagara/Source/Niagara/Private/NiagaraRendererSprites.cpp
@@ -1199,9 +1199,9 @@ FNiagaraDynamicDataBase *FNiagaraRendererSprites::GenerateDynamicData(const FNia
 			DynamicData->ParameterDataBound = ParameterData.GetParameterDataArray();
 		}
 
-		if (DynamicData && Properties->MaterialParameterBindings.Num() != 0)
+		if (DynamicData && Properties->MaterialParameters.HasAnyBindings())
 		{
-			ProcessMaterialParameterBindings(MakeArrayView(Properties->MaterialParameterBindings), Emitter, MakeArrayView(BaseMaterials_GT));
+			ProcessMaterialParameterBindings(Properties->MaterialParameters, Emitter, MakeArrayView(BaseMaterials_GT));
 		}
 	}
 
diff --git a/Engine/Plugins/FX/Niagara/Source/Niagara/Private/NiagaraRibbonRendererProperties.cpp b/Engine/Plugins/FX/Niagara/Source/Niagara/Private/NiagaraRibbonRendererProperties.cpp
--- a/Engine/Plugins/FX/Niagara/Source/Niagara/Private/NiagaraRibbonRendererProperties.cpp
+++ b/Engine/Plugins/FX/Niagara/Source/Niagara/Private/NiagaraRibbonRendererProperties.cpp
@@ -126,7 +126,6 @@ void UNiagaraRibbonRendererProperties::PostLoad()
 	Super::PostLoad();
 
 #if WITH_EDITORONLY_DATA
-
 	if (MaterialUserParamBinding.Parameter.GetType().GetClass() != UMaterialInterface::StaticClass())
 	{
 		FNiagaraTypeDefinition MaterialDef(UMaterialInterface::StaticClass());
@@ -149,6 +148,14 @@ void UNiagaraRibbonRendererProperties::PostLoad()
 	{
 		Material->ConditionalPostLoad();
 	}
+
+#if WITH_EDITORONLY_DATA
+	if (MaterialParameterBindings_DEPRECATED.Num() > 0)
+	{
+		MaterialParameters.AttributeBindings = MaterialParameterBindings_DEPRECATED;
+		MaterialParameterBindings_DEPRECATED.Empty();
+	}
+#endif
 }
 
 FNiagaraBoundsCalculator* UNiagaraRibbonRendererProperties::CreateBoundsCalculator()
@@ -180,7 +187,7 @@ bool UNiagaraRibbonRendererProperties::PopulateRequiredBindings(FNiagaraParamete
 		}
 	}
 
-	for (FNiagaraMaterialAttributeBinding& MaterialParamBinding : MaterialParameterBindings)
+	for (FNiagaraMaterialAttributeBinding& MaterialParamBinding : MaterialParameters.AttributeBindings)
 	{
 		InParameterStore.AddParameter(MaterialParamBinding.GetParamMapBindableVariable(), false);
 		bAnyAdded = true;
@@ -194,7 +201,7 @@ void UNiagaraRibbonRendererProperties::UpdateSourceModeDerivates(ENiagaraRendere
 	UNiagaraEmitter* SrcEmitter = GetTypedOuter<UNiagaraEmitter>();
 	if (SrcEmitter)
 	{
-		for (FNiagaraMaterialAttributeBinding& MaterialParamBinding : MaterialParameterBindings)
+		for (FNiagaraMaterialAttributeBinding& MaterialParamBinding : MaterialParameters.AttributeBindings)
 		{
 			MaterialParamBinding.CacheValues(SrcEmitter);
 		}
@@ -511,6 +518,14 @@ void UNiagaraRibbonRendererProperties::GetRendererFeedback(const FVersionedNiaga
 		const FNiagaraRendererFeedbackFix MultiPlaneFix = FNiagaraRendererFeedbackFix::CreateLambda([this]() { const_cast<UNiagaraRibbonRendererProperties*>(this)->MultiPlaneCount = FMath::Clamp(this->MultiPlaneCount - 1, 1, 16); });
 		OutWarnings.Add(FNiagaraRendererFeedback(ErrorDescription, ErrorSummary, ErrorFix, MultiPlaneFix, true));	
 	}	
+
+
+	if (MaterialParameters.HasAnyBindings())
+	{
+		TArray<UMaterialInterface*> Materials;
+		GetUsedMaterials(nullptr, Materials);
+		MaterialParameters.GetFeedback(Materials, OutWarnings);
+	}
 }
 
 
@@ -542,7 +557,7 @@ void UNiagaraRibbonRendererProperties::RenameVariable(const FNiagaraVariableBase
 	Super::RenameVariable(OldVariable, NewVariable, InEmitter);
 
 	// Handle renaming material bindings
-	for (FNiagaraMaterialAttributeBinding& Binding : MaterialParameterBindings)
+	for (FNiagaraMaterialAttributeBinding& Binding : MaterialParameters.AttributeBindings)
 	{
 		Binding.RenameVariableIfMatching(OldVariable, NewVariable, InEmitter.Emitter, GetCurrentSourceMode());
 	}
@@ -553,7 +568,7 @@ void UNiagaraRibbonRendererProperties::RemoveVariable(const FNiagaraVariableBase
 	Super::RemoveVariable(OldVariable, InEmitter);
 
 	// Handle resetting material bindings to defaults
-	for (FNiagaraMaterialAttributeBinding& Binding : MaterialParameterBindings)
+	for (FNiagaraMaterialAttributeBinding& Binding : MaterialParameters.AttributeBindings)
 	{
 		if (Binding.Matches(OldVariable, InEmitter.Emitter, GetCurrentSourceMode()))
 		{
diff --git a/Engine/Plugins/FX/Niagara/Source/Niagara/Private/NiagaraSpriteRendererProperties.cpp b/Engine/Plugins/FX/Niagara/Source/Niagara/Private/NiagaraSpriteRendererProperties.cpp
--- a/Engine/Plugins/FX/Niagara/Source/Niagara/Private/NiagaraSpriteRendererProperties.cpp
+++ b/Engine/Plugins/FX/Niagara/Source/Niagara/Private/NiagaraSpriteRendererProperties.cpp
@@ -149,7 +149,13 @@ void UNiagaraSpriteRendererProperties::PostLoad()
 
 	// Fix up these bindings from their loaded source bindings
 	SetPreviousBindings(FVersionedNiagaraEmitter(), SourceMode);
-#endif // WITH_EDITORONLY_DATA
+
+	if (MaterialParameterBindings_DEPRECATED.Num() > 0)
+	{
+		MaterialParameters.AttributeBindings = MaterialParameterBindings_DEPRECATED;
+		MaterialParameterBindings_DEPRECATED.Empty();
+	}
+#endif
 }
 
 void UNiagaraSpriteRendererProperties::PostInitProperties()
@@ -334,7 +340,7 @@ bool UNiagaraSpriteRendererProperties::PopulateRequiredBindings(FNiagaraParamete
 		}
 	}
 
-	for (FNiagaraMaterialAttributeBinding& MaterialParamBinding : MaterialParameterBindings)
+	for (FNiagaraMaterialAttributeBinding& MaterialParamBinding : MaterialParameters.AttributeBindings)
 	{
 		InParameterStore.AddParameter(MaterialParamBinding.GetParamMapBindableVariable(), false);
 		bAnyAdded = true;
@@ -350,7 +356,7 @@ void UNiagaraSpriteRendererProperties::UpdateSourceModeDerivates(ENiagaraRendere
 	UNiagaraEmitter* SrcEmitter = GetTypedOuter<UNiagaraEmitter>();
 	if (SrcEmitter)
 	{
-		for (FNiagaraMaterialAttributeBinding& MaterialParamBinding : MaterialParameterBindings)
+		for (FNiagaraMaterialAttributeBinding& MaterialParamBinding : MaterialParameters.AttributeBindings)
 		{
 			MaterialParamBinding.CacheValues(SrcEmitter);
 		}
@@ -445,7 +451,7 @@ void UNiagaraSpriteRendererProperties::RenameVariable(const FNiagaraVariableBase
 	Super::RenameVariable(OldVariable, NewVariable, InEmitter);
 
 	// Handle renaming material bindings
-	for (FNiagaraMaterialAttributeBinding& Binding : MaterialParameterBindings)
+	for (FNiagaraMaterialAttributeBinding& Binding : MaterialParameters.AttributeBindings)
 	{
 		Binding.RenameVariableIfMatching(OldVariable, NewVariable, InEmitter.Emitter, GetCurrentSourceMode());
 	}
@@ -456,7 +462,7 @@ void UNiagaraSpriteRendererProperties::RemoveVariable(const FNiagaraVariableBase
 	Super::RemoveVariable(OldVariable, InEmitter);
 
 	// Handle resetting material bindings to defaults
-	for (FNiagaraMaterialAttributeBinding& Binding : MaterialParameterBindings)
+	for (FNiagaraMaterialAttributeBinding& Binding : MaterialParameters.AttributeBindings)
 	{
 		if (Binding.Matches(OldVariable, InEmitter.Emitter, GetCurrentSourceMode()))
 		{
@@ -555,6 +561,18 @@ void UNiagaraSpriteRendererProperties::GetRendererFeedback(const FVersionedNiaga
 	}
 }
 
+void UNiagaraSpriteRendererProperties::GetRendererFeedback(const FVersionedNiagaraEmitter & InEmitter, TArray<FNiagaraRendererFeedback>&OutErrors, TArray<FNiagaraRendererFeedback>&OutWarnings, TArray<FNiagaraRendererFeedback>&OutInfo) const
+{
+	Super::GetRendererFeedback(InEmitter, OutErrors, OutWarnings, OutInfo);
+
+	if (MaterialParameters.HasAnyBindings())
+	{
+		TArray<UMaterialInterface*> Materials;
+		GetUsedMaterials(nullptr, Materials);
+		MaterialParameters.GetFeedback(Materials, OutWarnings);
+	}
+}
+
 void UNiagaraSpriteRendererProperties::GetRendererTooltipWidgets(const FNiagaraEmitterInstance* InEmitter, TArray<TSharedPtr<SWidget>>& OutWidgets, TSharedPtr<FAssetThumbnailPool> InThumbnailPool) const
 {
 	TArray<UMaterialInterface*> Materials;
diff --git a/Engine/Plugins/FX/Niagara/Source/Niagara/Public/NiagaraMeshRendererProperties.h b/Engine/Plugins/FX/Niagara/Source/Niagara/Public/NiagaraMeshRendererProperties.h
--- a/Engine/Plugins/FX/Niagara/Source/Niagara/Public/NiagaraMeshRendererProperties.h
+++ b/Engine/Plugins/FX/Niagara/Source/Niagara/Public/NiagaraMeshRendererProperties.h
@@ -184,7 +184,7 @@ public:
 	virtual void GetAdditionalVariables(TArray<FNiagaraVariableBase>& OutArray) const override;
 	virtual	void GetRendererWidgets(const FNiagaraEmitterInstance* InEmitter, TArray<TSharedPtr<SWidget>>& OutWidgets, TSharedPtr<FAssetThumbnailPool> InThumbnailPool) const override;
 	virtual	void GetRendererTooltipWidgets(const FNiagaraEmitterInstance* InEmitter, TArray<TSharedPtr<SWidget>>& OutWidgets, TSharedPtr<FAssetThumbnailPool> InThumbnailPool) const override;
-	virtual void GetRendererFeedback(const FVersionedNiagaraEmitter& InEmitter, TArray<FText>& OutErrors, TArray<FText>& OutWarnings, TArray<FText>& OutInfo) const override;
+	virtual void GetRendererFeedback(const FVersionedNiagaraEmitter& InEmitter, TArray<FNiagaraRendererFeedback>& OutErrors, TArray<FNiagaraRendererFeedback>& OutWarnings, TArray<FNiagaraRendererFeedback>& OutInfo) const override;
 	void OnMeshChanged();
 	void OnMeshPostBuild(UStaticMesh*);
 	void OnAssetReimported(UObject*);
@@ -356,7 +356,12 @@ public:
 
 	/** If this array has entries, we will create a MaterialInstanceDynamic per Emitter instance from Material and set the Material parameters using the Niagara simulation variables listed.*/
 	UPROPERTY(EditAnywhere, Category = "Bindings")
-	TArray<FNiagaraMaterialAttributeBinding> MaterialParameterBindings;
+	FNiagaraRendererMaterialParameters MaterialParameters;
+
+#if WITH_EDITORONLY_DATA
+	UPROPERTY()
+	TArray<FNiagaraMaterialAttributeBinding> MaterialParameterBindings_DEPRECATED;
+#endif
 
 	// The following bindings are not provided by the user, but are filled based on what other bindings are set to, and the value of bGenerateAccurateMotionVectors
 
@@ -409,7 +414,7 @@ protected:
 	void InitBindings();
 	void SetPreviousBindings(const FVersionedNiagaraEmitter& SrcEmitter, ENiagaraRendererSourceDataMode InSourceMode);
 	virtual void UpdateSourceModeDerivates(ENiagaraRendererSourceDataMode InSourceMode, bool bFromPropertyEdit = false) override;
-	virtual bool NeedsMIDsForMaterials() const override { return MaterialParameterBindings.Num() > 0; }	
+	virtual bool NeedsMIDsForMaterials() const override { return MaterialParameters.HasAnyBindings(); }
 #if WITH_EDITORONLY_DATA
 	bool ChangeRequiresMeshListRebuild(const FProperty* Property);
 	void RebuildMeshList();
diff --git a/Engine/Plugins/FX/Niagara/Source/Niagara/Public/NiagaraRenderer.h b/Engine/Plugins/FX/Niagara/Source/Niagara/Public/NiagaraRenderer.h
--- a/Engine/Plugins/FX/Niagara/Source/Niagara/Public/NiagaraRenderer.h
+++ b/Engine/Plugins/FX/Niagara/Source/Niagara/Public/NiagaraRenderer.h
@@ -148,7 +148,7 @@ public:
 
 protected:
 
-	virtual void ProcessMaterialParameterBindings(TConstArrayView< FNiagaraMaterialAttributeBinding > InMaterialParameterBindings, const FNiagaraEmitterInstance* InEmitter, TConstArrayView<UMaterialInterface*> InMaterials) const;
+	virtual void ProcessMaterialParameterBindings(const FNiagaraRendererMaterialParameters& MaterialParameters, const FNiagaraEmitterInstance* InEmitter, TConstArrayView<UMaterialInterface*> InMaterials) const;
 
 	bool IsRendererEnabled(const UNiagaraRendererProperties* InProperties, const FNiagaraEmitterInstance* Emitter) const;
 	bool UseLocalSpace(const FNiagaraSceneProxy* Proxy) const;
diff --git a/Engine/Plugins/FX/Niagara/Source/Niagara/Public/NiagaraRendererProperties.h b/Engine/Plugins/FX/Niagara/Source/Niagara/Public/NiagaraRendererProperties.h
--- a/Engine/Plugins/FX/Niagara/Source/Niagara/Public/NiagaraRendererProperties.h
+++ b/Engine/Plugins/FX/Niagara/Source/Niagara/Public/NiagaraRendererProperties.h
@@ -161,6 +161,73 @@ enum class ENiagaraRendererSortPrecision : uint8
 	High,
 };
 
+USTRUCT()
+struct FNiagaraRendererMaterialScalarParameter
+{
+	GENERATED_BODY()
+
+	UPROPERTY(EditAnywhere, Category = "Material")
+	FName MaterialParameterName;
+
+	UPROPERTY(EditAnywhere, Category = "Material")
+	float Value = 0.0f;
+};
+
+USTRUCT()
+struct FNiagaraRendererMaterialVectorParameter
+{
+	GENERATED_BODY()
+
+	UPROPERTY(EditAnywhere, Category = "Material")
+	FName MaterialParameterName;
+
+	UPROPERTY(EditAnywhere, Category = "Material")
+	FLinearColor Value = FLinearColor::Black;
+};
+
+USTRUCT()
+struct FNiagaraRendererMaterialTextureParameter
+{
+	GENERATED_BODY()
+
+	UPROPERTY(EditAnywhere, Category = "Material")
+	FName MaterialParameterName;
+
+	UPROPERTY(EditAnywhere, Category = "Material")
+	TObjectPtr<UTexture> Texture;
+};
+
+/**
+* Parameters to apply to the material, these are both constant and dynamic bindings
+* Having any bindings set will cause a MID to be generated
+*/
+USTRUCT()
+struct FNiagaraRendererMaterialParameters
+{
+	GENERATED_BODY()
+
+	UPROPERTY(EditAnywhere, Category = "Material")
+	TArray<FNiagaraMaterialAttributeBinding> AttributeBindings;
+
+	UPROPERTY(EditAnywhere, Category = "Material")
+	TArray<FNiagaraRendererMaterialScalarParameter> ScalarParameters;
+
+	UPROPERTY(EditAnywhere, Category = "Material")
+	TArray<FNiagaraRendererMaterialVectorParameter> VectorParameters;
+
+	UPROPERTY(EditAnywhere, Category = "Material")
+	TArray<FNiagaraRendererMaterialTextureParameter> TextureParameters;
+
+#if WITH_EDITORONLY_DATA
+	void GetFeedback(TArrayView<UMaterialInterface*> Materials, TArray<FNiagaraRendererFeedback>& OutWarnings) const;
+#endif
+
+	bool HasAnyBindings() const
+	{
+		return AttributeBindings.Num() > 0 || ScalarParameters.Num() > 0 || VectorParameters.Num() > 0 || TextureParameters.Num() > 0;
+	}
+};
+
 /**
 * Emitter properties base class
 * Each EmitterRenderer derives from this with its own class, and returns it in GetProperties; a copy
diff --git a/Engine/Plugins/FX/Niagara/Source/Niagara/Public/NiagaraRibbonRendererProperties.h b/Engine/Plugins/FX/Niagara/Source/Niagara/Public/NiagaraRibbonRendererProperties.h
--- a/Engine/Plugins/FX/Niagara/Source/Niagara/Public/NiagaraRibbonRendererProperties.h
+++ b/Engine/Plugins/FX/Niagara/Source/Niagara/Public/NiagaraRibbonRendererProperties.h
@@ -433,7 +433,12 @@ public:
 
 	/** If this array has entries, we will create a MaterialInstanceDynamic per Emitter instance from Material and set the Material parameters using the Niagara simulation variables listed.*/
 	UPROPERTY(EditAnywhere, Category = "Bindings")
-	TArray<FNiagaraMaterialAttributeBinding> MaterialParameterBindings;
+	FNiagaraRendererMaterialParameters MaterialParameters;
+
+#if WITH_EDITORONLY_DATA
+	UPROPERTY()
+	TArray<FNiagaraMaterialAttributeBinding> MaterialParameterBindings_DEPRECATED;
+#endif
 
 	/** Implicit binding for previous position */
 	UPROPERTY(Transient)
@@ -481,7 +486,7 @@ protected:
 
 	void UpdateSourceModeDerivates(ENiagaraRendererSourceDataMode InSourceMode, bool bFromPropertyEdit);
 
-	virtual bool NeedsMIDsForMaterials() const { return MaterialParameterBindings.Num() > 0; }
+	virtual bool NeedsMIDsForMaterials() const { return MaterialParameters.HasAnyBindings(); }
 private: 
 	static TArray<TWeakObjectPtr<UNiagaraRibbonRendererProperties>> RibbonRendererPropertiesToDeferredInit;
 };
diff --git a/Engine/Plugins/FX/Niagara/Source/Niagara/Public/NiagaraSpriteRendererProperties.h b/Engine/Plugins/FX/Niagara/Source/Niagara/Public/NiagaraSpriteRendererProperties.h
--- a/Engine/Plugins/FX/Niagara/Source/Niagara/Public/NiagaraSpriteRendererProperties.h
+++ b/Engine/Plugins/FX/Niagara/Source/Niagara/Public/NiagaraSpriteRendererProperties.h
@@ -131,6 +131,7 @@ public:
 	virtual void GetRendererWidgets(const FNiagaraEmitterInstance* InEmitter, TArray<TSharedPtr<SWidget>>& OutWidgets, TSharedPtr<FAssetThumbnailPool> InThumbnailPool) const override;
 	virtual void GetRendererTooltipWidgets(const FNiagaraEmitterInstance* InEmitter, TArray<TSharedPtr<SWidget>>& OutWidgets, TSharedPtr<FAssetThumbnailPool> InThumbnailPool) const override;
 	virtual void GetRendererFeedback(const FVersionedNiagaraEmitter& InEmitter, TArray<FText>& OutErrors, TArray<FText>& OutWarnings, TArray<FText>& OutInfo) const override;
+	virtual void GetRendererFeedback(const FVersionedNiagaraEmitter& InEmitter, TArray<FNiagaraRendererFeedback>& OutErrors, TArray<FNiagaraRendererFeedback>& OutWarnings, TArray<FNiagaraRendererFeedback>& OutInfo) const override;
 #endif
 	virtual ENiagaraRendererSourceDataMode GetCurrentSourceMode() const override { return SourceMode; }
 
@@ -321,7 +322,12 @@ public:
 
 	/** If this array has entries, we will create a MaterialInstanceDynamic per Emitter instance from Material and set the Material parameters using the Niagara simulation variables listed.*/
 	UPROPERTY(EditAnywhere, Category = "Bindings")
-	TArray< FNiagaraMaterialAttributeBinding > MaterialParameterBindings;
+	FNiagaraRendererMaterialParameters MaterialParameters;
+
+#if WITH_EDITORONLY_DATA
+	UPROPERTY()
+	TArray<FNiagaraMaterialAttributeBinding> MaterialParameterBindings_DEPRECATED;
+#endif
 
 	// The following bindings are only needed for accurate motion vectors
 
@@ -342,7 +348,7 @@ public:
 	UPROPERTY(Transient)
 	FNiagaraVariableAttributeBinding PrevPivotOffsetBinding;
 
-	virtual bool NeedsMIDsForMaterials() const { return MaterialParameterBindings.Num() > 0; }
+	virtual bool NeedsMIDsForMaterials() const { return MaterialParameters.HasAnyBindings(); }
 
 
 #if WITH_EDITORONLY_DATA
diff --git a/Engine/Plugins/FX/Niagara/Source/NiagaraEditor/Private/Customizations/NiagaraTypeCustomizations.cpp b/Engine/Plugins/FX/Niagara/Source/NiagaraEditor/Private/Customizations/NiagaraTypeCustomizations.cpp
--- a/Engine/Plugins/FX/Niagara/Source/NiagaraEditor/Private/Customizations/NiagaraTypeCustomizations.cpp
+++ b/Engine/Plugins/FX/Niagara/Source/NiagaraEditor/Private/Customizations/NiagaraTypeCustomizations.cpp
@@ -28,6 +28,7 @@
 #include "SGraphActionMenu.h"
 #include "DeviceProfiles/DeviceProfileManager.h"
 #include "Framework/Application/SlateApplication.h"
+#include "Framework/Multibox/MultiBoxBuilder.h"
 #include "ViewModels/NiagaraSystemViewModel.h"
 #include "ViewModels/TNiagaraViewModelManager.h"
 #include "Widgets/SNiagaraParameterName.h"
@@ -38,7 +39,6 @@
 #include "Widgets/Text/STextBlock.h"
 #include "IDetailGroup.h"
 
-
 #define LOCTEXT_NAMESPACE "FNiagaraVariableAttributeBindingCustomization"
 #define ALLOW_LIBRARY_TO_LIBRARY_DEFAULT_BINDING 0
 
@@ -515,7 +515,6 @@ void FNiagaraVariableAttributeBindingCustomization::CustomizeHeader(TSharedRef<I
 	}
 }
 
-
 //////////////////////////////////////////////////////////////////////////
 
 FName FNiagaraUserParameterBindingCustomization::GetVariableName() const
@@ -1985,4 +1984,155 @@ void FNiagaraSystemScalabilityOverrideCustomization::CustomizeChildren(TSharedRe
 	AddOverrideProperties(TEXT("bOverrideCullProxySettings"), { TEXT("CullProxyMode"), TEXT("MaxSystemProxies") });
 }
 
+//////////////////////////////////////////////////////////////////////////
+
+void FNiagaraRendererMaterialParameterCustomization::CustomizeHeader(TSharedRef<IPropertyHandle> PropertyHandle, class FDetailWidgetRow& HeaderRow, IPropertyTypeCustomizationUtils& StructCustomizationUtils)
+{
+	TArray<UObject*> Objects;
+	PropertyHandle->GetOuterObjects(Objects);
+	if (Objects.Num() > 0)
+	{
+		WeakRenderProperties = Cast<UNiagaraRendererProperties>(Objects[0]);
+	}
+
+	HeaderRow
+		.NameContent()
+		[
+			PropertyHandle->CreatePropertyNameWidget()
+		]
+		.ValueContent()
+		[
+			PropertyHandle->CreatePropertyValueWidget()
+		];
+}
+
+void FNiagaraRendererMaterialParameterCustomization::CustomizeChildren(TSharedRef<IPropertyHandle> StructPropertyHandle, class IDetailChildrenBuilder& ChildBuilder, IPropertyTypeCustomizationUtils& StructCustomizationUtils)
+{
+	uint32 NumChildren;
+	StructPropertyHandle->GetNumChildren(NumChildren);
+
+	static const FName NAME_MaterialParameterName("MaterialParameterName");
+	for (uint32 ChildIndex=0; ChildIndex < NumChildren; ++ChildIndex)
+	{
+		TSharedPtr<IPropertyHandle> ChildHandle = StructPropertyHandle->GetChildHandle(ChildIndex);
+		if ( ChildHandle->GetProperty()->GetFName() == NAME_MaterialParameterName )
+		{
+			MaterialBindingNameProperty = ChildHandle;
+
+			ChildBuilder.AddCustomRow(ChildHandle->GetPropertyDisplayName())
+			.NameWidget
+			[
+				ChildHandle->CreatePropertyNameWidget()
+			]
+			.ValueWidget
+			[
+				SNew(SComboButton)
+				.OnGetMenuContent(this, &FNiagaraRendererMaterialParameterCustomization::OnGetMaterialBindingNameMenuContent)
+				.ContentPadding(1)
+				.ButtonStyle(FAppStyle::Get(), "PropertyEditor.AssetComboStyle")
+				.ForegroundColor(FAppStyle::GetColor("PropertyEditor.AssetName.ColorAndOpacity"))
+				.ButtonContent()
+				[
+					SNew(STextBlock)
+					.Text(this, &FNiagaraRendererMaterialParameterCustomization::GetMaterialBindingNameText)
+					.Font(IDetailLayoutBuilder::GetDetailFont())
+				]
+			];
+		}
+		else
+		{
+			ChildBuilder.AddProperty(ChildHandle.ToSharedRef());
+		}
+	}
+}
+
+FText FNiagaraRendererMaterialParameterCustomization::GetMaterialBindingNameText() const
+{
+	FName BindingName;
+	if ( MaterialBindingNameProperty.IsValid() )
+	{
+		MaterialBindingNameProperty->GetValue(BindingName);
+	}
+	return FText::FromName(BindingName);
+}
+
+TSharedRef<SWidget> FNiagaraRendererMaterialParameterCustomization::OnGetMaterialBindingNameMenuContent() const
+{
+	FMenuBuilder MenuBuilder(true, NULL);
+
+	UNiagaraRendererProperties* RenderProperties = WeakRenderProperties.Get();
+	if (RenderProperties != nullptr && MaterialBindingNameProperty.IsValid())
+	{
+		TArray<UMaterialInterface*> UsedMaterials;
+		RenderProperties->GetUsedMaterials(nullptr, UsedMaterials);
+
+		TArray<FName> ValidBindings;
+		for (UMaterialInterface* Material : UsedMaterials)
+		{
+			if (Material == nullptr)
+			{
+				continue;
+			}
+
+			GetMaterialBindingNames(Material, ValidBindings);
+		}
+
+		for (FName ValidBinding : ValidBindings)
+		{
+			MenuBuilder.AddMenuEntry(
+				FText::FromName(ValidBinding),
+				TAttribute<FText>(),
+				FSlateIcon(),
+				FUIAction(
+					FExecuteAction::CreateLambda(
+						[ValidBinding, PropertyToSet=MaterialBindingNameProperty]()
+						{
+							PropertyToSet->SetValue(ValidBinding);
+						}
+					)
+				)
+			);
+		}
+	}
+
+	return MenuBuilder.MakeWidget();
+}
+
+void FNiagaraRendererMaterialScalarParameterCustomization::GetMaterialBindingNames(UMaterialInterface* Material, TArray<FName>& OutBindings) const
+{
+	TArray<FName> ValidBindings;
+	TArray<FMaterialParameterInfo> MaterialParameterInfos;
+	TArray<FGuid> Guids;
+	Material->GetAllScalarParameterInfo(MaterialParameterInfos, Guids);
+
+	for (const FMaterialParameterInfo& MaterialParameterInfo : MaterialParameterInfos)
+	{
+		OutBindings.AddUnique(MaterialParameterInfo.Name);
+	}
+}
+
+void FNiagaraRendererMaterialVectorParameterCustomization::GetMaterialBindingNames(UMaterialInterface* Material, TArray<FName>& OutBindings) const
+{
+	TArray<FMaterialParameterInfo> MaterialParameterInfos;
+	TArray<FGuid> Guids;
+	Material->GetAllVectorParameterInfo(MaterialParameterInfos, Guids);
+
+	for (const FMaterialParameterInfo& MaterialParameterInfo : MaterialParameterInfos)
+	{
+		OutBindings.AddUnique(MaterialParameterInfo.Name);
+	}
+}
+
+void FNiagaraRendererMaterialTextureParameterCustomization::GetMaterialBindingNames(UMaterialInterface* Material, TArray<FName>& OutBindings) const
+{
+	TArray<FMaterialParameterInfo> MaterialParameterInfos;
+	TArray<FGuid> Guids;
+	Material->GetAllTextureParameterInfo(MaterialParameterInfos, Guids);
+
+	for (const FMaterialParameterInfo& MaterialParameterInfo : MaterialParameterInfos)
+	{
+		OutBindings.AddUnique(MaterialParameterInfo.Name);
+	}
+}
+
 #undef LOCTEXT_NAMESPACE
diff --git a/Engine/Plugins/FX/Niagara/Source/NiagaraEditor/Private/Customizations/NiagaraTypeCustomizations.h b/Engine/Plugins/FX/Niagara/Source/NiagaraEditor/Private/Customizations/NiagaraTypeCustomizations.h
--- a/Engine/Plugins/FX/Niagara/Source/NiagaraEditor/Private/Customizations/NiagaraTypeCustomizations.h
+++ b/Engine/Plugins/FX/Niagara/Source/NiagaraEditor/Private/Customizations/NiagaraTypeCustomizations.h
@@ -14,8 +14,10 @@ class FDetailWidgetRow;
 class FNiagaraScriptViewModel;
 class IPropertyHandle;
 class IPropertyHandleArray;
+class UMaterialInterface;
 class UNiagaraGraph;
 class UNiagaraParameterDefinitions;
+class UNiagaraRendererProperties;
 class UNiagaraScriptVariable;
 class SComboButton;
 
@@ -183,7 +185,7 @@ private:
 	TSharedPtr<IPropertyHandle> PropertyHandle;
 	TSharedPtr<SComboButton> ComboButton;
 	FVersionedNiagaraEmitter BaseEmitter;
-	class UNiagaraRendererProperties* RenderProps = nullptr;
+	UNiagaraRendererProperties* RenderProps = nullptr;
 	class UNiagaraSimulationStageBase* SimulationStage = nullptr;
 	struct FNiagaraVariableAttributeBinding* TargetVariableBinding = nullptr;
 	const struct FNiagaraVariableAttributeBinding* DefaultVariableBinding = nullptr;
@@ -267,7 +269,7 @@ private:
 	TSharedPtr<SComboButton> NiagaraParameterButton;
 	class UNiagaraSystem* BaseSystem;
 	FVersionedNiagaraEmitter BaseEmitter;
-	class UNiagaraRendererProperties* RenderProps;
+	UNiagaraRendererProperties* RenderProps;
 	struct FNiagaraMaterialAttributeBinding* TargetParameterBinding;
 
 };
@@ -367,3 +369,54 @@ public:
 
 	virtual void CustomizeChildren(TSharedRef<IPropertyHandle> PropertyHandle, IDetailChildrenBuilder& ChildBuilder, IPropertyTypeCustomizationUtils& CustomizationUtils);
 };
+
+class FNiagaraRendererMaterialParameterCustomization : public IPropertyTypeCustomization
+{
+public:
+	// IPropertyTypeCustomization interface begin
+	virtual void CustomizeHeader(TSharedRef<IPropertyHandle> PropertyHandle, class FDetailWidgetRow& HeaderRow, IPropertyTypeCustomizationUtils& StructCustomizationUtils) override;
+	virtual void CustomizeChildren(TSharedRef<IPropertyHandle> StructPropertyHandle, class IDetailChildrenBuilder& ChildBuilder, IPropertyTypeCustomizationUtils& StructCustomizationUtils) override;
+	// IPropertyTypeCustomization interface end
+
+	FText GetMaterialBindingNameText() const;
+	TSharedRef<SWidget> OnGetMaterialBindingNameMenuContent() const;
+
+	virtual void GetMaterialBindingNames(UMaterialInterface* Material, TArray<FName>& OutBindings) const = 0;
+
+private:
+	TWeakObjectPtr<UNiagaraRendererProperties>	WeakRenderProperties;
+	TSharedPtr<IPropertyHandle>					MaterialBindingNameProperty;
+};
+
+class FNiagaraRendererMaterialScalarParameterCustomization : public FNiagaraRendererMaterialParameterCustomization
+{
+public:
+	static TSharedRef<IPropertyTypeCustomization> MakeInstance()
+	{
+		return MakeShared<FNiagaraRendererMaterialScalarParameterCustomization>();
+	}
+
+	virtual void GetMaterialBindingNames(UMaterialInterface* Material, TArray<FName>& OutBindings) const override;
+};
+
+class FNiagaraRendererMaterialVectorParameterCustomization : public FNiagaraRendererMaterialParameterCustomization
+{
+public:
+	static TSharedRef<IPropertyTypeCustomization> MakeInstance()
+	{
+		return MakeShared<FNiagaraRendererMaterialVectorParameterCustomization>();
+	}
+
+	virtual void GetMaterialBindingNames(UMaterialInterface* Material, TArray<FName>& OutBindings) const override;
+};
+
+class FNiagaraRendererMaterialTextureParameterCustomization : public FNiagaraRendererMaterialParameterCustomization
+{
+public:
+	static TSharedRef<IPropertyTypeCustomization> MakeInstance()
+	{
+		return MakeShared<FNiagaraRendererMaterialTextureParameterCustomization>();
+	}
+
+	virtual void GetMaterialBindingNames(UMaterialInterface* Material, TArray<FName>& OutBindings) const override;
+};
diff --git a/Engine/Plugins/FX/Niagara/Source/NiagaraEditor/Private/NiagaraEditorModule.cpp b/Engine/Plugins/FX/Niagara/Source/NiagaraEditor/Private/NiagaraEditorModule.cpp
--- a/Engine/Plugins/FX/Niagara/Source/NiagaraEditor/Private/NiagaraEditorModule.cpp
+++ b/Engine/Plugins/FX/Niagara/Source/NiagaraEditor/Private/NiagaraEditorModule.cpp
@@ -980,6 +980,18 @@ void FNiagaraEditorModule::StartupModule()
 		FNiagaraOutlinerEmitterInstanceData::StaticStruct()->GetFName(),
 		FOnGetPropertyTypeCustomizationInstance::CreateStatic(&FNiagaraOutlinerEmitterInstanceDetailsCustomization::MakeInstance));
 #endif
+	
+	PropertyModule.RegisterCustomPropertyTypeLayout(
+		FNiagaraRendererMaterialScalarParameter::StaticStruct()->GetFName(),
+		FOnGetPropertyTypeCustomizationInstance::CreateStatic(&FNiagaraRendererMaterialScalarParameterCustomization::MakeInstance));
+
+	PropertyModule.RegisterCustomPropertyTypeLayout(
+		FNiagaraRendererMaterialVectorParameter::StaticStruct()->GetFName(),
+		FOnGetPropertyTypeCustomizationInstance::CreateStatic(&FNiagaraRendererMaterialVectorParameterCustomization::MakeInstance));
+
+	PropertyModule.RegisterCustomPropertyTypeLayout(
+		FNiagaraRendererMaterialTextureParameter::StaticStruct()->GetFName(),
+		FOnGetPropertyTypeCustomizationInstance::CreateStatic(&FNiagaraRendererMaterialTextureParameterCustomization::MakeInstance));
 
 	// Outliner Customizations end.
 
