- Fix case in include

#rnx

diff --git a/Engine/Plugins/FX/Niagara/Source/NiagaraEditor/Private/Customizations/NiagaraTypeCustomizations.cpp b/Engine/Plugins/FX/Niagara/Source/NiagaraEditor/Private/Customizations/NiagaraTypeCustomizations.cpp
--- a/Engine/Plugins/FX/Niagara/Source/NiagaraEditor/Private/Customizations/NiagaraTypeCustomizations.cpp
+++ b/Engine/Plugins/FX/Niagara/Source/NiagaraEditor/Private/Customizations/NiagaraTypeCustomizations.cpp
@@ -28,7 +28,7 @@
 #include "SGraphActionMenu.h"
 #include "DeviceProfiles/DeviceProfileManager.h"
 #include "Framework/Application/SlateApplication.h"
-#include "Framework/Multibox/MultiBoxBuilder.h"
+#include "Framework/MultiBox/MultiBoxBuilder.h"
 #include "ViewModels/NiagaraSystemViewModel.h"
 #include "ViewModels/TNiagaraViewModelManager.h"
 #include "Widgets/SNiagaraParameterName.h"
