Added a rotation field for the overlap region for the FindActors function of teh rigid body DI

diff --git a/Engine/Plugins/FX/Niagara/Source/Niagara/Private/NiagaraDataInterfaceRigidMeshCollisionQuery.cpp b/Engine/Plugins/FX/Niagara/Source/Niagara/Private/NiagaraDataInterfaceRigidMeshCollisionQuery.cpp
--- a/Engine/Plugins/FX/Niagara/Source/Niagara/Private/NiagaraDataInterfaceRigidMeshCollisionQuery.cpp
+++ b/Engine/Plugins/FX/Niagara/Source/Niagara/Private/NiagaraDataInterfaceRigidMeshCollisionQuery.cpp
@@ -40,6 +40,7 @@ struct FNiagaraRigidMeshCollisionDIFunctionVersion
 		InitialVersion = 0,
 		LargeWorldCoordinates = 1,
 		SetMaxDistance = 2,
+		FindActorRotation = 3,
 
 		VersionPlusOne,
 		LatestVersion = VersionPlusOne - 1
@@ -67,6 +68,30 @@ static const FName GetClosestPointMeshDistanceFieldName(TEXT("GetClosestPointMes
 static const FName GetClosestPointMeshDistanceFieldAccurateName(TEXT("GetClosestPointMeshDistanceFieldAccurate"));
 static const FName GetClosestPointMeshDistanceFieldNoNormalName(TEXT("GetClosestPointMeshDistanceFieldNoNormal"));
 
+static const FText OverlapOriginDescription = IF_WITH_EDITORONLY_DATA(
+	LOCTEXT("RigidBodyOverlapOriginDescription", "The center point, in world space, where the overlap trace will be performed."),
+	FText()
+);
+
+static const FText OverlapRotationDescription = IF_WITH_EDITORONLY_DATA(
+	LOCTEXT("RigidBodyOVerlapRotationDesciption", "The orientation of the box to be used for hte overlap trace."),
+	FText()
+);
+
+static const FText OverlapExtentDescription = IF_WITH_EDITORONLY_DATA(
+	LOCTEXT("RigidBodyOverlapExtentDescription", "The extent, in world space, of the overlap trace."),
+	FText()
+);
+
+static const FText TraceChannelDescription = IF_WITH_EDITORONLY_DATA(
+	LOCTEXT("RigidBodyTraceChannelDescription", "The trace channel to collide against. Trace channels can be configured in the project settings."),
+	FText()
+);
+
+static const FText SkipOverlapDescription = IF_WITH_EDITORONLY_DATA(
+	LOCTEXT("RigidBodySkipTraceDescription", "If enabled, the overlap test will not be performed."),
+	FText()
+);
 
 //------------------------------------------------------------------------------------------------------------
 
@@ -1176,26 +1201,6 @@ void UNiagaraDataInterfaceRigidMeshCollisionQuery::GetFunctions(TArray<FNiagaraF
 	using namespace NDIRigidMeshCollisionLocal;
 
 	{
-		const FText OverlapOriginDescription = IF_WITH_EDITORONLY_DATA(
-			LOCTEXT("RigidBodyOverlapOriginDescription", "The center point, in world space, where the overlap trace will be performed."),
-			FText()
-		);
-
-		const FText OverlapExtentDescription = IF_WITH_EDITORONLY_DATA(
-			LOCTEXT("RigidBodyOverlapExtentDescription", "The extent, in world space, of the overlap trace."),
-			FText()
-		);
-
-		const FText TraceChannelDescription = IF_WITH_EDITORONLY_DATA(
-			LOCTEXT("RigidBodyTraceChannelDescription", "The trace channel to collide against. Trace channels can be configured in the project settings."),
-			FText()
-		);
-
-		const FText SkipOverlapDescription = IF_WITH_EDITORONLY_DATA(
-			LOCTEXT("RigidBodySkipTraceDescription", "If enabled, the overlap test will not be performed."),
-			FText()
-		);
-
 		FNiagaraFunctionSignature Sig;
 		Sig.Name = FindActorsName;
 		Sig.SetDescription(LOCTEXT("FindActorsDescription", "Triggers an overlap test on the world to find actors to represent.."));
@@ -1206,6 +1211,7 @@ void UNiagaraDataInterfaceRigidMeshCollisionQuery::GetFunctions(TArray<FNiagaraF
 		Sig.bRequiresExecPin = true;
 		Sig.AddInput(FNiagaraVariable(FNiagaraTypeDefinition(GetClass()), TEXT("RigidBody DI")));
 		Sig.AddInput(FNiagaraVariable(FNiagaraTypeDefinition::GetPositionDef(), TEXT("Overlap Origin")), OverlapOriginDescription);
+		Sig.AddInput(FNiagaraVariable(FNiagaraTypeDefinition::GetQuatDef(), TEXT("Overlap Rotation")), OverlapRotationDescription);
 		Sig.AddInput(FNiagaraVariable(FNiagaraTypeDefinition::GetVec3Def(), TEXT("Overlap Extent")), OverlapExtentDescription);
 		Sig.AddInput(FNiagaraVariable(FNiagaraTypeDefinition(StaticEnum<ECollisionChannel>()), TEXT("TraceChannel")), TraceChannelDescription);
 		Sig.AddInput(FNiagaraVariable(FNiagaraTypeDefinition::GetBoolDef(), TEXT("Skip Overlap")), SkipOverlapDescription);
@@ -1642,6 +1648,18 @@ bool UNiagaraDataInterfaceRigidMeshCollisionQuery::UpgradeFunctionCall(FNiagaraF
 		}
 	}
 
+	if (FunctionSignature.FunctionVersion < FNiagaraRigidMeshCollisionDIFunctionVersion::FindActorRotation)
+	{
+		if (FunctionSignature.Name == FindActorsName)
+		{
+			FNiagaraVariable OverlapRotation(FNiagaraTypeDefinition::GetQuatDef(), TEXT("Overlap Rotation"));
+
+			FunctionSignature.Inputs.Insert(OverlapRotation, 2);
+			FunctionSignature.InputDescriptions.Add(OverlapRotation, OverlapRotationDescription);
+			bChanged = true;
+		}
+	}
+
 	FunctionSignature.FunctionVersion = FNiagaraRigidMeshCollisionDIFunctionVersion::LatestVersion;
 
 	return bChanged;
@@ -1756,7 +1774,7 @@ bool UNiagaraDataInterfaceRigidMeshCollisionQuery::GlobalFindActors(UWorld* Worl
 	return PreviousActors != InstanceData.FoundActors;
 }
 
-bool UNiagaraDataInterfaceRigidMeshCollisionQuery::FindActors(UWorld* World, FNDIRigidMeshCollisionData& InstanceData, ECollisionChannel Channel, const FVector& OverlapLocation, const FVector& OverlapExtent) const
+bool UNiagaraDataInterfaceRigidMeshCollisionQuery::FindActors(UWorld* World, FNDIRigidMeshCollisionData& InstanceData, ECollisionChannel Channel, const FVector& OverlapLocation, const FVector& OverlapExtent, const FQuat& OverlapRotation) const
 {
 	TArray<TWeakObjectPtr<AActor>> PreviousActors;
 	Swap(InstanceData.FoundActors, PreviousActors);
@@ -1769,7 +1787,7 @@ bool UNiagaraDataInterfaceRigidMeshCollisionQuery::FindActors(UWorld* World, FND
 		TArray<FOverlapResult> Overlaps;
 		FCollisionQueryParams Params(SCENE_QUERY_STAT(NiagaraRigidMeshCollisionQuery), false);
 
-		World->OverlapMultiByChannel(Overlaps, OverlapLocation, FQuat::Identity, Channel, FCollisionShape::MakeBox(0.5f * OverlapExtent), Params);
+		World->OverlapMultiByChannel(Overlaps, OverlapLocation, OverlapRotation, Channel, FCollisionShape::MakeBox(0.5f * OverlapExtent), Params);
 
 		for (const FOverlapResult& OverlapResult : Overlaps)
 		{
@@ -1821,6 +1839,7 @@ void UNiagaraDataInterfaceRigidMeshCollisionQuery::FindActorsCPU(FVectorVMExtern
 	VectorVM::FUserPtrHandler<FNDIRigidMeshCollisionData> InstanceData(Context);
 
 	FNDIInputParam<FNiagaraPosition> OverlapOriginParam(Context);
+	FNDIInputParam<FQuat4f> OverlapRotationParam(Context);
 	FNDIInputParam<FVector3f> OverlapExtentParam(Context);
 	FNDIInputParam<ECollisionChannel> TraceChannelParam(Context);
 	FNDIInputParam<FNiagaraBool> SkipOverlapParam(Context);
@@ -1837,6 +1856,7 @@ void UNiagaraDataInterfaceRigidMeshCollisionQuery::FindActorsCPU(FVectorVMExtern
 			for (int32 i = 0; i < Context.GetNumInstances(); ++i)
 			{
 				FNiagaraPosition OverlapOrigin = OverlapOriginParam.GetAndAdvance();
+				FQuat4f OverlapRotation = OverlapRotationParam.GetAndAdvance();
 				FVector3f OverlapExtent = OverlapExtentParam.GetAndAdvance();
 				ECollisionChannel TraceChannel = TraceChannelParam.GetAndAdvance();
 				bool SkipOverlap = SkipOverlapParam.GetAndAdvance() || !InstanceData->bRequiresSourceActors;
@@ -1846,7 +1866,7 @@ void UNiagaraDataInterfaceRigidMeshCollisionQuery::FindActorsCPU(FVectorVMExtern
 				if (!SkipOverlap)
 				{
 					const FVector ConvertedOrigin = LWCConverter.ConvertSimulationPositionToWorld(OverlapOrigin);
-					if (FindActors(World, *InstanceData, TraceChannel, ConvertedOrigin, FVector(OverlapExtent)))
+					if (FindActors(World, *InstanceData, TraceChannel, ConvertedOrigin, FVector(OverlapExtent), FQuat(OverlapRotation)))
 					{
 						ActorsChanged = true;
 						InstanceData->bFoundActorsUpdated = true;
diff --git a/Engine/Plugins/FX/Niagara/Source/Niagara/Public/NiagaraDataInterfaceRigidMeshCollisionQuery.h b/Engine/Plugins/FX/Niagara/Source/Niagara/Public/NiagaraDataInterfaceRigidMeshCollisionQuery.h
--- a/Engine/Plugins/FX/Niagara/Source/Niagara/Public/NiagaraDataInterfaceRigidMeshCollisionQuery.h
+++ b/Engine/Plugins/FX/Niagara/Source/Niagara/Public/NiagaraDataInterfaceRigidMeshCollisionQuery.h
@@ -240,7 +240,7 @@ public:
 	bool UpdateSourceActors(FNiagaraSystemInstance* SystemInstance, FNDIRigidMeshCollisionData& InstanceData) const;
 
 	bool GetExplicitActors(FNDIRigidMeshCollisionData& InstanceData);
-	bool FindActors(UWorld* World, FNDIRigidMeshCollisionData& InstanceData, ECollisionChannel Channel, const FVector& OverlapLocation, const FVector& OverlapExtent) const;
+	bool FindActors(UWorld* World, FNDIRigidMeshCollisionData& InstanceData, ECollisionChannel Channel, const FVector& OverlapLocation, const FVector& OverlapExtent, const FQuat& OverlapRotation) const;
 	bool GlobalFindActors(UWorld* World, FNDIRigidMeshCollisionData& InstanceData) const;
 
 	void FindActorsCPU(FVectorVMExternalFunctionContext& Context);
