better error case for #prims exceeds limit
single DI function call can query both mesh DFs and prims

diff --git a/Engine/Plugins/FX/Niagara/Shaders/Private/NiagaraDataInterfaceRigidMeshCollisionQuery.ush b/Engine/Plugins/FX/Niagara/Shaders/Private/NiagaraDataInterfaceRigidMeshCollisionQuery.ush
--- a/Engine/Plugins/FX/Niagara/Shaders/Private/NiagaraDataInterfaceRigidMeshCollisionQuery.ush
+++ b/Engine/Plugins/FX/Niagara/Shaders/Private/NiagaraDataInterfaceRigidMeshCollisionQuery.ush
@@ -523,9 +523,8 @@ void DIRigidMeshCollision_GetClosestPointMeshDistanceFieldNoNormal(in FDIRigidMe
 	
 	OutClosestPosition = float3(0,0,0);
 	OutClosestVelocity = float3(0,0,0);
-
-	float OutClosestDistanceTmp;
-	const int ElementIndex = DIRigidMeshCollision_GetClosestElement(DIContext,WorldPosition,CollisionPosition,CollisionNormal,OutClosestDistanceTmp,TimeFraction);
+	
+	const int ElementIndex = DIRigidMeshCollision_GetClosestElement(DIContext,WorldPosition,CollisionPosition,CollisionNormal,OutClosestDistance,TimeFraction);
 
 	if (ElementIndex != -1)
 	{
@@ -536,11 +535,11 @@ void DIRigidMeshCollision_GetClosestPointMeshDistanceFieldNoNormal(in FDIRigidMe
 			// #todo(dmp): this does a dynamic branch based on intersecting the bbox.  Maybe we can factor that out due to the broadphase here?
 			MaxDistance = max(.001, MaxDistance);
 			OutClosestDistance = DistanceToNearestSurfaceForObject(DFIndex, WorldPosition, MaxDistance);		
-
-			const float3 PreviousPosition = mul(DIRigidMeshCollision_GetPreviousTransform(DIContext,ElementIndex), float4(CollisionPosition,1.0)).xyz;
-			const float3 CurrentPosition = mul(DIRigidMeshCollision_GetCurrentTransform(DIContext,ElementIndex), float4(CollisionPosition,1.0)).xyz;
-			OutClosestVelocity = ( CurrentPosition - PreviousPosition ) / DeltaTime;
-			OutClosestPosition = PreviousPosition + TimeFraction * (CurrentPosition-PreviousPosition);
 		}
-	}
+		
+		const float3 PreviousPosition = mul(DIRigidMeshCollision_GetPreviousTransform(DIContext,ElementIndex), float4(CollisionPosition,1.0)).xyz;
+		const float3 CurrentPosition = mul(DIRigidMeshCollision_GetCurrentTransform(DIContext,ElementIndex), float4(CollisionPosition,1.0)).xyz;
+		OutClosestVelocity = ( CurrentPosition - PreviousPosition ) / DeltaTime;
+		OutClosestPosition = PreviousPosition + TimeFraction * (CurrentPosition-PreviousPosition);
+	}	
 }
\ No newline at end of file
diff --git a/Engine/Plugins/FX/Niagara/Source/Niagara/Private/NiagaraDataInterfaceRigidMeshCollisionQuery.cpp b/Engine/Plugins/FX/Niagara/Source/Niagara/Private/NiagaraDataInterfaceRigidMeshCollisionQuery.cpp
--- a/Engine/Plugins/FX/Niagara/Source/Niagara/Private/NiagaraDataInterfaceRigidMeshCollisionQuery.cpp
+++ b/Engine/Plugins/FX/Niagara/Source/Niagara/Private/NiagaraDataInterfaceRigidMeshCollisionQuery.cpp
@@ -454,6 +454,12 @@ void UpdateInternalArrays(TConstArrayView<AActor*> Actors, TConstArrayView<FName
 		CountCollisionPrimitives(StaticMeshView, BoxCount, SphereCount, CapsuleCount);
 		CountCollisionPrimitives(SkeletalMeshView, BoxCount, SphereCount, CapsuleCount);
 
+		if ((BoxCount + SphereCount + CapsuleCount) >= OutAssetArrays->MaxPrimitives)
+		{
+			UE_LOG(LogRigidMeshCollision, Error, TEXT("Number of Collision DI primitives is higher than the %d limit.  Please increase it."), OutAssetArrays->MaxPrimitives);
+			return;
+		}
+
 		const bool MismatchOffsets = ((OutAssetArrays->ElementOffsets.SphereOffset - OutAssetArrays->ElementOffsets.BoxOffset) != BoxCount) ||
 			((OutAssetArrays->ElementOffsets.CapsuleOffset - OutAssetArrays->ElementOffsets.SphereOffset) != SphereCount) ||
 			((OutAssetArrays->ElementOffsets.NumElements - OutAssetArrays->ElementOffsets.CapsuleOffset) != CapsuleCount);
@@ -462,18 +468,11 @@ void UpdateInternalArrays(TConstArrayView<AActor*> Actors, TConstArrayView<FName
 		check(!MismatchOffsets || bFullUpdate);
 
 		if (bFullUpdate)
-		{
-			if ((BoxCount + SphereCount + CapsuleCount) < OutAssetArrays->MaxPrimitives)
-			{
-				OutAssetArrays->ElementOffsets.BoxOffset = 0;
-				OutAssetArrays->ElementOffsets.SphereOffset = OutAssetArrays->ElementOffsets.BoxOffset + BoxCount;
-				OutAssetArrays->ElementOffsets.CapsuleOffset = OutAssetArrays->ElementOffsets.SphereOffset + SphereCount;
-				OutAssetArrays->ElementOffsets.NumElements = OutAssetArrays->ElementOffsets.CapsuleOffset + CapsuleCount;
-			}
-			else
-			{
-				UE_LOG(LogRigidMeshCollision, Warning, TEXT("Number of Collision DI primitives is higher than the %d limit.  Please increase it."), OutAssetArrays->MaxPrimitives);
-			}
+		{			
+			OutAssetArrays->ElementOffsets.BoxOffset = 0;
+			OutAssetArrays->ElementOffsets.SphereOffset = OutAssetArrays->ElementOffsets.BoxOffset + BoxCount;
+			OutAssetArrays->ElementOffsets.CapsuleOffset = OutAssetArrays->ElementOffsets.SphereOffset + SphereCount;
+			OutAssetArrays->ElementOffsets.NumElements = OutAssetArrays->ElementOffsets.CapsuleOffset + CapsuleCount;
 		}
 
 		uint32 BoxIndex = OutAssetArrays->ElementOffsets.BoxOffset;
