- Niagara Compute Dipsatch modified to use FRDGBuilder through the pipeline
- Legacy data interfaces run on Graph execution timeline, some caveats now exist like peeking the current buffer on the MainDataSet willl be wrong
- Converted data interfaces run on Graph build timeline executing passes as required
- Converted data interfaces have been updated to use new PreStage / PostStage / PostSimulate function
- Various other parts of Niagara have been updated to allow operation on the Graph builder timeline, i.e. readbacks, debug drawing, profiiling, etc
- Various global compute shaders have been updated to use shader parameter structs
- More work to come to move various resources to RDG and be part of the tracked process to avoid everything being submitted on the immediate command list

diff --git a/Engine/Plugins/Experimental/ChaosNiagara/Source/ChaosNiagara/Private/NiagaraDataInterfaceGeometryCollection.cpp b/Engine/Plugins/Experimental/ChaosNiagara/Source/ChaosNiagara/Private/NiagaraDataInterfaceGeometryCollection.cpp
--- a/Engine/Plugins/Experimental/ChaosNiagara/Source/ChaosNiagara/Private/NiagaraDataInterfaceGeometryCollection.cpp
+++ b/Engine/Plugins/Experimental/ChaosNiagara/Source/ChaosNiagara/Private/NiagaraDataInterfaceGeometryCollection.cpp
@@ -450,9 +450,6 @@ void FNDIGeometryCollectionProxy::PreStage(FRHICommandList& RHICmdList, const FN
 	}
 }
 
-void FNDIGeometryCollectionProxy::ResetData(FRHICommandList& RHICmdList, const FNiagaraDataInterfaceArgs& Context)
-{}
-
 //------------------------------------------------------------------------------------------------------------
 
 UNiagaraDataInterfaceGeometryCollection::UNiagaraDataInterfaceGeometryCollection(FObjectInitializer const& ObjectInitializer)
diff --git a/Engine/Plugins/Experimental/ChaosNiagara/Source/ChaosNiagara/Public/NiagaraDataInterfaceGeometryCollection.h b/Engine/Plugins/Experimental/ChaosNiagara/Source/ChaosNiagara/Public/NiagaraDataInterfaceGeometryCollection.h
--- a/Engine/Plugins/Experimental/ChaosNiagara/Source/ChaosNiagara/Public/NiagaraDataInterfaceGeometryCollection.h
+++ b/Engine/Plugins/Experimental/ChaosNiagara/Source/ChaosNiagara/Public/NiagaraDataInterfaceGeometryCollection.h
@@ -191,9 +191,6 @@ struct FNDIGeometryCollectionProxy : public FNiagaraDataInterfaceProxy
 	/** Launch all pre stage functions */
 	virtual void PreStage(FRHICommandList& RHICmdList, const FNiagaraDataInterfaceStageArgs& Context) override;
 
-	/** Reset the buffers  */
-	virtual void ResetData(FRHICommandList& RHICmdList, const FNiagaraDataInterfaceArgs& Context) override;
-
 	/** List of proxy data for each system instances*/
 	TMap<FNiagaraSystemInstanceID, FNDIGeometryCollectionData> SystemInstancesToProxyData;
 };
\ No newline at end of file
diff --git a/Engine/Plugins/FX/Niagara/Shaders/Private/NiagaraComputeFreeIDs.usf b/Engine/Plugins/FX/Niagara/Shaders/Private/NiagaraComputeFreeIDs.usf
--- a/Engine/Plugins/FX/Niagara/Shaders/Private/NiagaraComputeFreeIDs.usf
+++ b/Engine/Plugins/FX/Niagara/Shaders/Private/NiagaraComputeFreeIDs.usf
@@ -18,7 +18,7 @@ uint FreeIDListIndex;
 #if USE_WAVE_INTRINSICS
 
 // This happens to be fixed on PS4 and XB1, but we'll need to configure it externally
-[numthreads(THREAD_COUNT, 1, 1)]
+[numthreads(THREADGROUP_SIZE, 1, 1)]
 void ComputeFreeIDs(uint3 GroupId : SV_GroupID, uint3 GroupThreadId : SV_GroupThreadID, uint3 DispatchThreadId : SV_DispatchThreadID)
 {
 	uint InputStart = GroupId.x * WaveGetLaneCount();
@@ -60,14 +60,14 @@ void ComputeFreeIDs(uint3 GroupId : SV_GroupID, uint3 GroupThreadId : SV_GroupTh
 #else
 
 // Temporary buffer where we compute the number of elements to write before each input element.
-groupshared uint GroupWriteOffsets[2*THREAD_COUNT];
+groupshared uint GroupWriteOffsets[2*THREADGROUP_SIZE];
 // The offset where this thread group will start writing its output.
 groupshared uint GroupWriteStart;
 
-[numthreads(THREAD_COUNT, 1, 1)]
+[numthreads(THREADGROUP_SIZE, 1, 1)]
 void ComputeFreeIDs(uint3 GroupId : SV_GroupID, uint3 GroupThreadId : SV_GroupThreadID, uint3 DispatchThreadId : SV_DispatchThreadID)
 {
-	uint InputStart = GroupId.x * THREAD_COUNT;
+	uint InputStart = GroupId.x * THREADGROUP_SIZE;
 	uint Thread = GroupThreadId.x;
 	uint ParticleID = InputStart + Thread;
 
@@ -87,18 +87,18 @@ void ComputeFreeIDs(uint3 GroupId : SV_GroupID, uint3 GroupThreadId : SV_GroupTh
 
 	// Compute prefix sum.
 	[unroll]
-	for(uint Offset = 1; Offset < THREAD_COUNT; Offset = Offset << 1)
+	for(uint Offset = 1; Offset < THREADGROUP_SIZE; Offset = Offset << 1)
 	{
 		// Swap the input and output buffers.
 		OutBuffer = 1 - OutBuffer;
 		InBuffer = 1 - InBuffer;
 		if(Thread >= Offset)
 		{
-			GroupWriteOffsets[OutBuffer*THREAD_COUNT + Thread] = GroupWriteOffsets[InBuffer*THREAD_COUNT + Thread - Offset] + GroupWriteOffsets[InBuffer*THREAD_COUNT + Thread];
+			GroupWriteOffsets[OutBuffer*THREADGROUP_SIZE + Thread] = GroupWriteOffsets[InBuffer*THREADGROUP_SIZE + Thread - Offset] + GroupWriteOffsets[InBuffer*THREADGROUP_SIZE + Thread];
 		}
 		else
 		{
-			GroupWriteOffsets[OutBuffer*THREAD_COUNT + Thread] = GroupWriteOffsets[InBuffer*THREAD_COUNT + Thread];
+			GroupWriteOffsets[OutBuffer*THREADGROUP_SIZE + Thread] = GroupWriteOffsets[InBuffer*THREADGROUP_SIZE + Thread];
 		}
 		GroupMemoryBarrierWithGroupSync();
 	}
@@ -106,8 +106,8 @@ void ComputeFreeIDs(uint3 GroupId : SV_GroupID, uint3 GroupThreadId : SV_GroupTh
 	if(Thread == 0)
 	{
 		// Add the value of the last predicate to get the total number of IDs that this group will write.
-		uint LastIsFree = (IDToIndexTable[InputStart + THREAD_COUNT - 1] == -1) ? 1 : 0;
-		uint NumGroupFreeIDs = GroupWriteOffsets[(OutBuffer + 1)*THREAD_COUNT - 1] + LastIsFree;
+		uint LastIsFree = (IDToIndexTable[InputStart + THREADGROUP_SIZE - 1] == -1) ? 1 : 0;
+		uint NumGroupFreeIDs = GroupWriteOffsets[(OutBuffer + 1)*THREADGROUP_SIZE - 1] + LastIsFree;
 		// Add to the global write offset. The previous offset is where we start writing in the output buffer.
 		InterlockedAdd(RWFreeIDListSizes[FreeIDListIndex], NumGroupFreeIDs, GroupWriteStart);
 	}
@@ -116,7 +116,7 @@ void ComputeFreeIDs(uint3 GroupId : SV_GroupID, uint3 GroupThreadId : SV_GroupTh
 	if(IDToIndexTable[ParticleID] == -1)
 	{
 		// This is a free ID, output it.
-		uint WriteOffset = GroupWriteStart + GroupWriteOffsets[OutBuffer*THREAD_COUNT + Thread];
+		uint WriteOffset = GroupWriteStart + GroupWriteOffsets[OutBuffer*THREADGROUP_SIZE + Thread];
 		// TODO: consider buffering the results before writing to the output UAV.
 		RWFreeIDList[WriteOffset] = ParticleID;
 	}
diff --git a/Engine/Plugins/FX/Niagara/Shaders/Private/NiagaraDataInterfaceIntRenderTarget2D.ush b/Engine/Plugins/FX/Niagara/Shaders/Private/NiagaraDataInterfaceIntRenderTarget2D.ush
--- a/Engine/Plugins/FX/Niagara/Shaders/Private/NiagaraDataInterfaceIntRenderTarget2D.ush
+++ b/Engine/Plugins/FX/Niagara/Shaders/Private/NiagaraDataInterfaceIntRenderTarget2D.ush
@@ -1,28 +1,31 @@
 // Copyright Epic Games, Inc. All Rights Reserved.
 
-float4				TextureSizeAndInvSize_{ParameterName};
-RWTexture2D<int>	RWTextureUAV_{ParameterName};
+float4				{ParameterName}_TextureSizeAndInvSize;
+RWTexture2D<int>	{ParameterName}_RWTextureUAV;
 
 void GetValue_{ParameterName}(int X, int Y, out int CurrValue)
 {
-	CurrValue = RWTextureUAV_{ParameterName}[int2(X, Y)];
+	CurrValue = {ParameterName}_RWTextureUAV[int2(X, Y)];
 }
 
-void SetValue_{ParameterName}_UEImpureCall(int X, int Y, int Value)
+void SetValue_{ParameterName}_UEImpureCall(in bool bExecute, int X, int Y, int Value)
 {
-	RWTextureUAV_{ParameterName}[int2(X, Y)] = Value;
+	if ( bExecute )
+	{
+		{ParameterName}_RWTextureUAV[int2(X, Y)] = Value;
+	}
 }
 
 void AtomicAdd_{ParameterName}_UEImpureCall(in bool bExecute, int X, int Y, int Amount, out int CurrValue, out int PrevValue)
 {
 	if ( bExecute )
 	{
-		InterlockedAdd(RWTextureUAV_{ParameterName}[int2(X, Y)], Amount, PrevValue);
+		InterlockedAdd({ParameterName}_RWTextureUAV[int2(X, Y)], Amount, PrevValue);
 		CurrValue = PrevValue + Amount;
 	}
 	else
 	{
-		PrevValue = RWTextureUAV_{ParameterName}[int2(X, Y)];
+		PrevValue = {ParameterName}_RWTextureUAV[int2(X, Y)];
 		CurrValue = PrevValue;
 	}
 }
@@ -31,11 +34,11 @@ void AtomicCompareAndExchange_{ParameterName}_UEImpureCall(in bool bExecute, int
 {
 	if ( bExecute )
 	{
-		InterlockedCompareExchange(RWTextureUAV_{ParameterName}[int2(X, Y)], ComparisonValue, Value, OriginalValue);
+		InterlockedCompareExchange({ParameterName}_RWTextureUAV[int2(X, Y)], ComparisonValue, Value, OriginalValue);
 	}
 	else
 	{
-		OriginalValue = RWTextureUAV_{ParameterName}[int2(X, Y)];
+		OriginalValue = {ParameterName}_RWTextureUAV[int2(X, Y)];
 	}
 }
 
@@ -43,12 +46,12 @@ void AtomicMax_{ParameterName}_UEImpureCall(in bool bExecute, int X, int Y, int
 {
 	if ( bExecute )
 	{
-		InterlockedMax(RWTextureUAV_{ParameterName}[int2(X, Y)], InValue, PrevValue);
+		InterlockedMax({ParameterName}_RWTextureUAV[int2(X, Y)], InValue, PrevValue);
 		CurrValue = max(PrevValue, InValue);
 	}
 	else
 	{
-		PrevValue = RWTextureUAV_{ParameterName}[int2(X, Y)];
+		PrevValue = {ParameterName}_RWTextureUAV[int2(X, Y)];
 		CurrValue = PrevValue;
 	}
 }
@@ -57,30 +60,40 @@ void AtomicMin_{ParameterName}_UEImpureCall(in bool bExecute, int X, int Y, int
 {
 	if ( bExecute )
 	{
-		InterlockedMin(RWTextureUAV_{ParameterName}[int2(X, Y)], InValue, PrevValue);
+		InterlockedMin({ParameterName}_RWTextureUAV[int2(X, Y)], InValue, PrevValue);
 		CurrValue = min(PrevValue, InValue);
 	}
 	else
 	{
-		PrevValue = RWTextureUAV_{ParameterName}[int2(X, Y)];
+		PrevValue = {ParameterName}_RWTextureUAV[int2(X, Y)];
 		CurrValue = PrevValue;
 	}
 }
 
 void GetRenderTargetSize_{ParameterName}(out int Width, out int Height)
 {
-	Width = int(TextureSizeAndInvSize_{ParameterName}.x);
-	Height = int(TextureSizeAndInvSize_{ParameterName}.y);
+	Width = int({ParameterName}_TextureSizeAndInvSize.x);
+	Height = int({ParameterName}_TextureSizeAndInvSize.y);
 }
 
 void LinearToIndex_{ParameterName}(int Linear, out int X, out int Y)
 {
-	X = Linear % int(TextureSizeAndInvSize_{ParameterName}.x);
-	Y = Linear / int(TextureSizeAndInvSize_{ParameterName}.x);
+	X = int(uint(Linear) % uint({ParameterName}_TextureSizeAndInvSize.x));
+	Y = int(uint(Linear) / uint({ParameterName}_TextureSizeAndInvSize.x));
 }
 
 void LinearToUV_{ParameterName}(int Linear, out float2 UV)
 {
-	UV.x = (float(Linear % int(TextureSizeAndInvSize_{ParameterName}.x)) + 0.5f) * TextureSizeAndInvSize_{ParameterName}.z;
-	UV.y = (float(Linear / int(TextureSizeAndInvSize_{ParameterName}.x)) + 0.5f) * TextureSizeAndInvSize_{ParameterName}.w;
+	UV.x = (float(uint(Linear) % uint({ParameterName}_TextureSizeAndInvSize.x)) + 0.5f) * {ParameterName}_TextureSizeAndInvSize.z;
+	UV.y = (float(uint(Linear) / uint({ParameterName}_TextureSizeAndInvSize.x)) + 0.5f) * {ParameterName}_TextureSizeAndInvSize.w;
+}
+
+void ExecToIndex_{ParameterName}(out int X, out int Y)
+{
+	LinearToIndex_{ParameterName}(ExecIndex(), X, Y);
+}
+
+void ExecToUV_{ParameterName}(out float2 UV)
+{
+	LinearToUV_{ParameterName}(ExecIndex(), UV);
 }
diff --git a/Engine/Plugins/FX/Niagara/Shaders/Private/NiagaraFillIntBuffer.usf b/Engine/Plugins/FX/Niagara/Shaders/Private/NiagaraFillIntBuffer.usf
--- a/Engine/Plugins/FX/Niagara/Shaders/Private/NiagaraFillIntBuffer.usf
+++ b/Engine/Plugins/FX/Niagara/Shaders/Private/NiagaraFillIntBuffer.usf
@@ -1,17 +1,18 @@
+
 // Copyright Epic Games, Inc. All Rights Reserved.
 
 #include "/Engine/Private/Common.ush"
 
-RWBuffer<int> TargetBuffer;
-int FillValue;
-uint BufferSize;
+RWBuffer<int>	BufferUAV;
+uint			NumElements;
+int				FillValue;
 
-[numthreads(64, 1, 1)]
+[numthreads(THREADGROUP_SIZE, 1, 1)]
 void FillIntBuffer(uint3 DispatchThreadId : SV_DispatchThreadID)
 {
 	uint Index = DispatchThreadId.x;
-	if(Index < BufferSize)
+	if(Index < NumElements)
 	{
-		TargetBuffer[Index] = FillValue;
+		BufferUAV[Index] = FillValue;
 	}
 }
diff --git a/Engine/Plugins/FX/Niagara/Shaders/Private/NiagaraInitFreeIDBuffer.usf b/Engine/Plugins/FX/Niagara/Shaders/Private/NiagaraInitFreeIDBuffer.usf
--- a/Engine/Plugins/FX/Niagara/Shaders/Private/NiagaraInitFreeIDBuffer.usf
+++ b/Engine/Plugins/FX/Niagara/Shaders/Private/NiagaraInitFreeIDBuffer.usf
@@ -2,12 +2,12 @@
 
 #include "/Engine/Private/Common.ush"
 
-RWBuffer<int> RWNewBuffer;
-Buffer<int> ExistingBuffer;
-uint NumNewElements;
-uint NumExistingElements;
+RWBuffer<int>	RWNewBuffer;
+Buffer<int>		ExistingBuffer;
+uint			NumNewElements;
+uint			NumExistingElements;
 
-[numthreads(THREAD_COUNT, 1, 1)]
+[numthreads(THREADGROUP_SIZE, 1, 1)]
 void InitIDBufferCS(uint Index : SV_DispatchThreadID)
 {
 	if(Index.x < NumNewElements)
diff --git a/Engine/Plugins/FX/Niagara/Source/Niagara/Classes/NiagaraComputeExecutionContext.h b/Engine/Plugins/FX/Niagara/Source/Niagara/Classes/NiagaraComputeExecutionContext.h
--- a/Engine/Plugins/FX/Niagara/Source/Niagara/Classes/NiagaraComputeExecutionContext.h
+++ b/Engine/Plugins/FX/Niagara/Source/Niagara/Classes/NiagaraComputeExecutionContext.h
@@ -22,43 +22,6 @@ public:
 	}
 };
 
-struct FNiagaraGpuDispatchInstance
-{
-	FNiagaraGpuDispatchInstance(const FNiagaraGPUSystemTick& InTick, const FNiagaraComputeInstanceData& InInstanceData)
-		: Tick(InTick)
-		, InstanceData(InInstanceData)
-	{
-	}
-
-	const FNiagaraGPUSystemTick& Tick;
-	const FNiagaraComputeInstanceData& InstanceData;
-	FNiagaraSimStageData SimStageData;
-};
-
-struct FNiagaraGpuDispatchGroup
-{
-	TArray<FNiagaraGPUSystemTick*>				TicksWithPerInstanceData;
-	TArray<FNiagaraGpuDispatchInstance>			DispatchInstances;
-	TArray<FNiagaraComputeExecutionContext*>	FreeIDUpdates;
-};
-
-struct FNiagaraGpuDispatchList
-{
-	void PreAllocateGroups(int32 LastGroup)
-	{
-		const int32 GroupsToAllocate = LastGroup - DispatchGroups.Num();
-		if (GroupsToAllocate > 0)
-		{
-			DispatchGroups.AddDefaulted(GroupsToAllocate);
-		}
-	}
-
-	bool HasWork() const { return DispatchGroups.Num() > 0; }
-
-	TArray<uint32>						CountsToRelease;
-	TArray<FNiagaraGpuDispatchGroup>	DispatchGroups;
-};
-
 struct FNiagaraGpuSpawnInfoParams
 {
 	float IntervalDt;
diff --git a/Engine/Plugins/FX/Niagara/Source/Niagara/Classes/NiagaraDataInterface.h b/Engine/Plugins/FX/Niagara/Source/Niagara/Classes/NiagaraDataInterface.h
--- a/Engine/Plugins/FX/Niagara/Source/Niagara/Classes/NiagaraDataInterface.h
+++ b/Engine/Plugins/FX/Niagara/Source/Niagara/Classes/NiagaraDataInterface.h
@@ -8,11 +8,15 @@
 #include "NiagaraScriptBase.h"
 #include "NiagaraDataInterface.generated.h"
 
+class RENDERCORE_API FRDGBuilder;
+class FRDGExternalAccessQueue;
+
 class INiagaraCompiler;
 class UCurveVector;
 class UCurveLinearColor;
 class UCurveFloat;
 struct FNiagaraDataInterfaceProxy;
+struct FNiagaraDataInterfaceProxyRW;
 struct FNiagaraComputeInstanceData;
 class FNiagaraEmitterInstance;
 class FNiagaraGpuComputeDispatchInterface;
@@ -239,7 +243,87 @@ private:
 
 //////////////////////////////////////////////////////////////////////////
 
-struct FNiagaraDataInterfaceProxyRW;
+struct NIAGARA_API FNDIGpuComputeContext
+{
+	FNDIGpuComputeContext(FRDGBuilder& InGraphBuilder, const FNiagaraGpuComputeDispatchInterface& InComputeDispatchInterface)
+		: GraphBuilder(InGraphBuilder)
+		, ComputeDispatchInterface(InComputeDispatchInterface)
+	{
+	}
+
+	FRDGBuilder& GetGraphBuilder() const { return GraphBuilder; }
+	FRDGExternalAccessQueue& GetRDGExternalAccessQueue() const;
+
+	const FNiagaraGpuComputeDispatchInterface& GetComputeDispatchInterface() const { return ComputeDispatchInterface; }
+
+private:
+	FRDGBuilder& GraphBuilder;
+	const FNiagaraGpuComputeDispatchInterface& ComputeDispatchInterface;
+};
+
+struct NIAGARA_API FNDIGpuComputeResetContext : public FNDIGpuComputeContext
+{
+	FNDIGpuComputeResetContext(FRDGBuilder& InGraphBuilder, const FNiagaraGpuComputeDispatchInterface& InComputeDispatchInterface, FNiagaraSystemInstanceID InSystemInstanceID)
+		: FNDIGpuComputeContext(InGraphBuilder, InComputeDispatchInterface)
+		, SystemInstanceID(InSystemInstanceID)
+	{
+	}
+
+	FNiagaraSystemInstanceID GetSystemInstanceID() const { return SystemInstanceID; }
+
+private:
+	FNiagaraSystemInstanceID SystemInstanceID = FNiagaraSystemInstanceID();
+};
+
+struct NIAGARA_API FNDIGpuComputePrePostStageContext : public FNDIGpuComputeContext
+{
+	FNDIGpuComputePrePostStageContext(FRDGBuilder& InGraphBuilder, const FNiagaraGpuComputeDispatchInterface& InComputeDispatchInterface, const FNiagaraGPUSystemTick& InSystemTick, const FNiagaraComputeInstanceData& InComputeInstanceData, const FNiagaraSimStageData& InSimStageData)
+		: FNDIGpuComputeContext(InGraphBuilder, InComputeDispatchInterface)
+		, SystemTick(InSystemTick)
+		, ComputeInstanceData(InComputeInstanceData)
+		, SimStageData(InSimStageData)
+	{
+	}
+
+	const FNiagaraComputeInstanceData& GetComputeInstanceData() const { return ComputeInstanceData; }
+	const FNiagaraSimStageData& GetSimStageData() const { return SimStageData; }
+
+	FNiagaraSystemInstanceID GetSystemInstanceID() const;
+	FVector3f GetSystemLWCTile() const;
+	bool IsOutputStage() const;
+	bool IsIterationStage() const;
+
+	void SetDataInterfaceProxy(FNiagaraDataInterfaceProxy* InDataInterfaceProxy) { DataInterfaceProxy = InDataInterfaceProxy; }
+
+private:
+	const FNiagaraGPUSystemTick& SystemTick;
+	const FNiagaraComputeInstanceData& ComputeInstanceData;
+	const FNiagaraSimStageData& SimStageData;
+	FNiagaraDataInterfaceProxy* DataInterfaceProxy = nullptr;
+};
+
+using FNDIGpuComputePreStageContext = FNDIGpuComputePrePostStageContext;
+using FNDIGpuComputePostStageContext = FNDIGpuComputePrePostStageContext;
+
+struct NIAGARA_API FNDIGpuComputePostSimulateContext : public FNDIGpuComputeContext
+{
+	FNDIGpuComputePostSimulateContext(FRDGBuilder& InGraphBuilder, const FNiagaraGpuComputeDispatchInterface& InComputeDispatchInterface, FNiagaraSystemInstanceID InSystemInstanceID, bool InFinalPostSimulate)
+		: FNDIGpuComputeContext(InGraphBuilder, InComputeDispatchInterface)
+		, SystemInstanceID(InSystemInstanceID)
+		, bFinalPostSimulate(InFinalPostSimulate)
+	{
+	}
+
+	FNiagaraSystemInstanceID GetSystemInstanceID() const { return SystemInstanceID; }
+	bool IsFinalPostSimulate() const { return bFinalPostSimulate; }
+
+private:
+	FNiagaraSystemInstanceID SystemInstanceID = FNiagaraSystemInstanceID();
+	bool bFinalPostSimulate = false;
+};
+
+//////////////////////////////////////////////////////////////////////////
+
 struct FNiagaraDataInterfaceProxy
 {
 	FNiagaraDataInterfaceProxy() {}
@@ -250,17 +334,29 @@ struct FNiagaraDataInterfaceProxy
 
 	// #todo(dmp): move all of this stuff to the RW interface to keep it out of here?
 	FName SourceDIName;
-	
+
+	// New data interface path
+	virtual void ResetData(const FNDIGpuComputeResetContext& Context) { }
+	virtual void PreStage(const FNDIGpuComputePostStageContext& Context) {}
+	virtual void PostStage(const FNDIGpuComputePostStageContext& Context) {}
+	virtual void PostSimulate(const FNDIGpuComputePostSimulateContext& Context) {}
+
+	virtual bool RequiresPreStageFinalize() const { return false; }
+	virtual void FinalizePreStage(FRDGBuilder& GraphBuilder, const FNiagaraGpuComputeDispatchInterface& ComputeDispatchInterface) {}
+
+	virtual bool RequiresPostStageFinalize() const { return false; }
+	virtual void FinalizePostStage(FRDGBuilder& GraphBuilder, const FNiagaraGpuComputeDispatchInterface& ComputeDispatchInterface) {}
+
+
+	// Legacy data interface path
 	virtual void ResetData(FRHICommandList& RHICmdList, const FNiagaraDataInterfaceArgs& Context) { }
 
 	virtual void PreStage(FRHICommandList& RHICmdList, const FNiagaraDataInterfaceStageArgs& Context) {}
 	virtual void PostStage(FRHICommandList& RHICmdList, const FNiagaraDataInterfaceStageArgs& Context) {}
 	virtual void PostSimulate(FRHICommandList& RHICmdList, const FNiagaraDataInterfaceArgs& Context) {}
 
-	virtual bool RequiresPreStageFinalize() const { return false; }
 	virtual void FinalizePreStage(FRHICommandList& RHICmdList, const FNiagaraGpuComputeDispatchInterface* ComputeDispatchInterface) {}
 
-	virtual bool RequiresPostStageFinalize() const { return false; }
 	virtual void FinalizePostStage(FRHICommandList& RHICmdList, const FNiagaraGpuComputeDispatchInterface* ComputeDispatchInterface) {}
 
 	virtual FNiagaraDataInterfaceProxyRW* AsIterationProxy() { return nullptr; }
@@ -270,10 +366,8 @@ struct FNiagaraDataInterfaceProxy
 
 struct NIAGARA_API FNiagaraDataInterfaceSetShaderParametersContext
 {
-	friend class FNiagaraGpuComputeDispatch;
-
-	FNiagaraDataInterfaceSetShaderParametersContext(FRHICommandList& InRHICmdList, const FNiagaraGpuComputeDispatchInterface& InComputeDispatchInterface, const FNiagaraGPUSystemTick& InSystemTick, const FNiagaraComputeInstanceData& InComputeInstanceData, const FNiagaraSimStageData& InSimStageData, const FNiagaraShaderRef& InShaderRef, uint8* InBaseParameters)
-		: RHICmdList(InRHICmdList)
+	FNiagaraDataInterfaceSetShaderParametersContext(FRDGBuilder& InGraphBuilder, const FNiagaraGpuComputeDispatchInterface& InComputeDispatchInterface, const FNiagaraGPUSystemTick& InSystemTick, const FNiagaraComputeInstanceData& InComputeInstanceData, const FNiagaraSimStageData& InSimStageData, const FNiagaraShaderRef& InShaderRef, uint8* InBaseParameters)
+		: GraphBuilder(InGraphBuilder)
 		, ComputeDispatchInterface(InComputeDispatchInterface)
 		, SystemTick(InSystemTick)
 		, ComputeInstanceData(InComputeInstanceData)
@@ -283,19 +377,21 @@ struct NIAGARA_API FNiagaraDataInterfaceSetShaderParametersContext
 	{
 	}
 
+	FRDGBuilder& GetGraphBuilder() const { return GraphBuilder; }
+	FRDGExternalAccessQueue& GetRDGExternalAccessQueue() const;
+
 	template<typename T> T& GetProxy() const { check(DataInterfaceProxy); return static_cast<T&>(*DataInterfaceProxy); }
 	const FNiagaraGpuComputeDispatchInterface& GetComputeDispatchInterface() const { return ComputeDispatchInterface; }
 	const FNiagaraGPUSystemTick& GetSystemTick() const { return SystemTick; }
 	const FNiagaraComputeInstanceData& GetComputeInstanceData() const { return ComputeInstanceData; }
 	const FNiagaraSimStageData& GetSimStageData() const { return SimStageData; }
-	FRHICommandList& GetRHICmdList() const { return RHICmdList; }
 
 	FNiagaraSystemInstanceID GetSystemInstanceID() const;
 	bool IsResourceBound(const void* ResourceAddress) const;
+	bool IsParameterBound(const void* ParameterAddress) const;
 	bool IsOutputStage() const;
 	bool IsIterationStage() const;
 
-
 	template<typename T> T* GetParameterNestedStruct() const
 	{
 		const uint32 StructOffset = Align(ParametersOffset, TShaderParameterStructTypeInfo<T>::Alignment);
@@ -324,7 +420,7 @@ struct NIAGARA_API FNiagaraDataInterfaceSetShaderParametersContext
 	}
 
 private:
-	FRHICommandList& RHICmdList;
+	FRDGBuilder& GraphBuilder;
 	const FNiagaraGpuComputeDispatchInterface& ComputeDispatchInterface;
 	const FNiagaraGPUSystemTick& SystemTick;
 	const FNiagaraComputeInstanceData& ComputeInstanceData;
diff --git a/Engine/Plugins/FX/Niagara/Source/Niagara/Classes/NiagaraDataInterface2DArrayTexture.h b/Engine/Plugins/FX/Niagara/Source/Niagara/Classes/NiagaraDataInterface2DArrayTexture.h
--- a/Engine/Plugins/FX/Niagara/Source/Niagara/Classes/NiagaraDataInterface2DArrayTexture.h
+++ b/Engine/Plugins/FX/Niagara/Source/Niagara/Classes/NiagaraDataInterface2DArrayTexture.h
@@ -17,7 +17,7 @@ class NIAGARA_API UNiagaraDataInterface2DArrayTexture : public UNiagaraDataInter
 
 	BEGIN_SHADER_PARAMETER_STRUCT(FShaderParameters, )
 		SHADER_PARAMETER(FVector3f,					TextureSize)
-		SHADER_PARAMETER_TEXTURE(Texture2DArray,	Texture)
+		SHADER_PARAMETER_RDG_TEXTURE(Texture2DArray,	Texture)
 		SHADER_PARAMETER_SAMPLER(SamplerState,		TextureSampler)
 	END_SHADER_PARAMETER_STRUCT()
 
diff --git a/Engine/Plugins/FX/Niagara/Source/Niagara/Classes/NiagaraDataInterfaceCubeTexture.h b/Engine/Plugins/FX/Niagara/Source/Niagara/Classes/NiagaraDataInterfaceCubeTexture.h
--- a/Engine/Plugins/FX/Niagara/Source/Niagara/Classes/NiagaraDataInterfaceCubeTexture.h
+++ b/Engine/Plugins/FX/Niagara/Source/Niagara/Classes/NiagaraDataInterfaceCubeTexture.h
@@ -17,7 +17,7 @@ class NIAGARA_API UNiagaraDataInterfaceCubeTexture : public UNiagaraDataInterfac
 
 	BEGIN_SHADER_PARAMETER_STRUCT(FShaderParameters, )
 		SHADER_PARAMETER(FIntPoint,				TextureSize)
-		SHADER_PARAMETER_TEXTURE(TextureCube,	Texture)
+		SHADER_PARAMETER_RDG_TEXTURE(TextureCube,	Texture)
 		SHADER_PARAMETER_SAMPLER(SamplerState,	TextureSampler)
 	END_SHADER_PARAMETER_STRUCT()
 
diff --git a/Engine/Plugins/FX/Niagara/Source/Niagara/Classes/NiagaraDataInterfaceGrid2DCollection.h b/Engine/Plugins/FX/Niagara/Source/Niagara/Classes/NiagaraDataInterfaceGrid2DCollection.h
--- a/Engine/Plugins/FX/Niagara/Source/Niagara/Classes/NiagaraDataInterfaceGrid2DCollection.h
+++ b/Engine/Plugins/FX/Niagara/Source/Niagara/Classes/NiagaraDataInterfaceGrid2DCollection.h
@@ -5,6 +5,7 @@
 #include "NiagaraDataInterfaceRW.h"
 #include "ClearQuad.h"
 #include "NiagaraComponent.h"
+#include "NiagaraRenderGraphUtils.h"
 #include "Niagara/Private/NiagaraStats.h"
 
 #include "NiagaraDataInterfaceGrid2DCollection.generated.h"
@@ -13,35 +14,7 @@ class FNiagaraSystemInstance;
 class UTextureRenderTarget;
 class UTextureRenderTarget2DArray;
 
-class FGrid2DBuffer
-{
-public:
-	FGrid2DBuffer(int NumX, int NumY, int NumAttributes, EPixelFormat PixelFormat)
-	{
-		const FRHITextureCreateDesc Desc =
-			FRHITextureCreateDesc::Create2DArray(TEXT("FGrid2DBuffer"), NumX, NumY, NumAttributes, PixelFormat)
-			.SetFlags(ETextureCreateFlags::ShaderResource | ETextureCreateFlags::UAV);
-
-		GridTexture = RHICreateTexture(Desc);
-
-		FRHITextureSRVCreateInfo SRVCreateInfo;
-		GridSRV = RHICreateShaderResourceView(GridTexture, SRVCreateInfo);
-		GridUAV = RHICreateUnorderedAccessView(GridTexture);
-
-		INC_MEMORY_STAT_BY(STAT_NiagaraGPUDataInterfaceMemory, RHIComputeMemorySize(GridTexture));
-	}
-	~FGrid2DBuffer()
-	{
-		DEC_MEMORY_STAT_BY(STAT_NiagaraGPUDataInterfaceMemory, RHIComputeMemorySize(GridTexture));
-		GridTexture.SafeRelease();
-		GridSRV.SafeRelease();
-		GridUAV.SafeRelease();
-	}
-
-	FTextureRHIRef GridTexture;
-	FShaderResourceViewRHIRef GridSRV;
-	FUnorderedAccessViewRHIRef GridUAV;
-};
+using FGrid2DBuffer = FNiagaraPooledRWTexture;
 
 struct FGrid2DCollectionRWInstanceData_GameThread
 {
@@ -83,7 +56,7 @@ struct FGrid2DCollectionRWInstanceData_RenderThread
 	FVector2D WorldBBoxSize = FVector2D::ZeroVector;
 	EPixelFormat PixelFormat = EPixelFormat::PF_R32_FLOAT;
 
-	TArray<TUniquePtr<FGrid2DBuffer>> Buffers;
+	TArray<FGrid2DBuffer, TInlineAllocator<2>> Buffers;
 	FGrid2DBuffer* CurrentData = nullptr;
 	FGrid2DBuffer* DestinationData = nullptr;
 
@@ -103,19 +76,18 @@ struct FGrid2DCollectionRWInstanceData_RenderThread
 	// overrides the render thread data, which in this case is for a grid reader
 	FNiagaraDataInterfaceProxy* OtherProxy = nullptr;
 
-	void BeginSimulate(FRHICommandList& RHICmdList);
-	void EndSimulate(FRHICommandList& RHICmdList);
+	void BeginSimulate(FRDGBuilder& GraphBuilder);
+	void EndSimulate();
 };
 
 struct FNiagaraDataInterfaceProxyGrid2DCollectionProxy : public FNiagaraDataInterfaceProxyRW
 {
 	FNiagaraDataInterfaceProxyGrid2DCollectionProxy() {}
 	
-	virtual void PreStage(FRHICommandList& RHICmdList, const FNiagaraDataInterfaceStageArgs& Context) override;
-	virtual void PostStage(FRHICommandList& RHICmdList, const FNiagaraDataInterfaceStageArgs& Context) override;
-	virtual void PostSimulate(FRHICommandList& RHICmdList, const FNiagaraDataInterfaceArgs& Context) override;
-
-	virtual void ResetData(FRHICommandList& RHICmdList, const FNiagaraDataInterfaceArgs& Context) override;
+	virtual void ResetData(const FNDIGpuComputeResetContext& Context) override;
+	virtual void PreStage(const FNDIGpuComputePreStageContext& Context) override;
+	virtual void PostStage(const FNDIGpuComputePostStageContext& Context) override;
+	virtual void PostSimulate(const FNDIGpuComputePostSimulateContext& Context) override;
 
 	virtual FIntVector GetElementCount(FNiagaraSystemInstanceID SystemInstanceID) const override;
 
diff --git a/Engine/Plugins/FX/Niagara/Source/Niagara/Classes/NiagaraDataInterfaceGrid3DCollection.h b/Engine/Plugins/FX/Niagara/Source/Niagara/Classes/NiagaraDataInterfaceGrid3DCollection.h
--- a/Engine/Plugins/FX/Niagara/Source/Niagara/Classes/NiagaraDataInterfaceGrid3DCollection.h
+++ b/Engine/Plugins/FX/Niagara/Source/Niagara/Classes/NiagaraDataInterfaceGrid3DCollection.h
@@ -5,6 +5,7 @@
 #include "NiagaraDataInterfaceRW.h"
 #include "ClearQuad.h"
 #include "NiagaraComponent.h"
+#include "NiagaraRenderGraphUtils.h"
 #include "Niagara/Private/NiagaraStats.h"
 
 #include "NiagaraDataInterfaceGrid3DCollection.generated.h"
@@ -13,23 +14,7 @@ class FNiagaraSystemInstance;
 class UTextureRenderTarget;
 class UTextureRenderTargetVolume;
 
-class FGrid3DBuffer
-{
-public:
-	FGrid3DBuffer(int NumX, int NumY, int NumZ, EPixelFormat PixelFormat)
-	{
-		GridBuffer.Initialize3D(TEXT("FGrid3DBuffer"), GPixelFormats[PixelFormat].BlockBytes, NumX, NumY, NumZ, PixelFormat);
-		INC_MEMORY_STAT_BY(STAT_NiagaraGPUDataInterfaceMemory, GridBuffer.NumBytes);
-	}
-
-	~FGrid3DBuffer()
-	{
-		DEC_MEMORY_STAT_BY(STAT_NiagaraGPUDataInterfaceMemory, GridBuffer.NumBytes);
-		GridBuffer.Release();
-	}
-
-	FTextureRWBuffer GridBuffer;
-};
+using FGrid3DBuffer = FNiagaraPooledRWTexture;
 
 struct FGrid3DCollectionRWInstanceData_GameThread
 {
@@ -83,7 +68,7 @@ struct FGrid3DCollectionRWInstanceData_RenderThread
 	FIntVector4 PreviewAttribute = FIntVector4(INDEX_NONE, INDEX_NONE, INDEX_NONE, INDEX_NONE);
 #endif
 
-	TArray<TUniquePtr<FGrid3DBuffer>> Buffers;
+	TArray<FGrid3DBuffer, TInlineAllocator<2>> Buffers;
 	FGrid3DBuffer* CurrentData = nullptr;
 	FGrid3DBuffer* DestinationData = nullptr;
 
@@ -98,18 +83,18 @@ struct FGrid3DCollectionRWInstanceData_RenderThread
 	// overrides the render thread data, which in this case is for a grid reader
 	FNiagaraDataInterfaceProxy* OtherProxy = nullptr;
 
-	void BeginSimulate(FRHICommandList& RHICmdList);
-	void EndSimulate(FRHICommandList& RHICmdList);
+	void BeginSimulate(FRDGBuilder& GraphBuilder);
+	void EndSimulate();
 };
 
 struct FNiagaraDataInterfaceProxyGrid3DCollectionProxy : public FNiagaraDataInterfaceProxyRW
 {
 	FNiagaraDataInterfaceProxyGrid3DCollectionProxy() {}
 
-	virtual void PreStage(FRHICommandList& RHICmdList, const FNiagaraDataInterfaceStageArgs& Context) override;
-	virtual void PostStage(FRHICommandList& RHICmdList, const FNiagaraDataInterfaceStageArgs& Context) override;
-	virtual void PostSimulate(FRHICommandList& RHICmdList, const FNiagaraDataInterfaceArgs& Context) override;
-	virtual void ResetData(FRHICommandList& RHICmdList, const FNiagaraDataInterfaceArgs& Context) override;
+	virtual void ResetData(const FNDIGpuComputeResetContext& Context) override;
+	virtual void PreStage(const FNDIGpuComputePreStageContext& Context) override;
+	virtual void PostStage(const FNDIGpuComputePostStageContext& Context) override;
+	virtual void PostSimulate(const FNDIGpuComputePostSimulateContext& Context) override;
 
 	virtual FIntVector GetElementCount(FNiagaraSystemInstanceID SystemInstanceID) const override;
 
diff --git a/Engine/Plugins/FX/Niagara/Source/Niagara/Classes/NiagaraDataInterfaceIntRenderTarget2D.h b/Engine/Plugins/FX/Niagara/Source/Niagara/Classes/NiagaraDataInterfaceIntRenderTarget2D.h
--- a/Engine/Plugins/FX/Niagara/Source/Niagara/Classes/NiagaraDataInterfaceIntRenderTarget2D.h
+++ b/Engine/Plugins/FX/Niagara/Source/Niagara/Classes/NiagaraDataInterfaceIntRenderTarget2D.h
@@ -12,14 +12,20 @@ class NIAGARA_API UNiagaraDataInterfaceIntRenderTarget2D : public UNiagaraDataIn
 {
 	GENERATED_UCLASS_BODY()
 
+	BEGIN_SHADER_PARAMETER_STRUCT(FShaderParameters, )
+		SHADER_PARAMETER(FVector4f,								TextureSizeAndInvSize)
+		SHADER_PARAMETER_RDG_TEXTURE_UAV(RWTexture2D<float4>,	RWTextureUAV)
+	END_SHADER_PARAMETER_STRUCT()
+
 public:
-	DECLARE_NIAGARA_DI_PARAMETER();	
-		
 	virtual void PostInitProperties() override;
 	
 	//~ UNiagaraDataInterface interface
 	virtual bool CanExecuteOnTarget(ENiagaraSimTarget Target)const override { return true; }
 	virtual void GetFunctions(TArray<FNiagaraFunctionSignature>& OutFunctions) override;
+#if WITH_EDITORONLY_DATA
+	virtual bool UpgradeFunctionCall(FNiagaraFunctionSignature& FunctionSignature) override;
+#endif
 	virtual void GetVMExternalFunction(const FVMExternalFunctionBindingInfo& BindingInfo, void* InstanceData, FVMExternalFunction &OutFunc) override;
 
 	virtual bool Equals(const UNiagaraDataInterface* Other) const override;
@@ -30,6 +36,9 @@ public:
 	virtual void GetParameterDefinitionHLSL(const FNiagaraDataInterfaceGPUParamInfo& ParamInfo, FString& OutHLSL) override;
 	virtual bool GetFunctionHLSL(const FNiagaraDataInterfaceGPUParamInfo& ParamInfo, const FNiagaraDataInterfaceGeneratedFunction& FunctionInfo, int FunctionInstanceIndex, FString& OutHLSL) override;
 #endif
+	virtual bool UseLegacyShaderBindings() const  override { return false; }
+	virtual void BuildShaderParameters(FNiagaraShaderParametersBuilder& ShaderParametersBuilder) const override;
+	virtual void SetShaderParameters(const FNiagaraDataInterfaceSetShaderParametersContext& Context) const override;
 
 	virtual void ProvidePerInstanceDataForRenderThread(void* DataForRenderThread, void* PerInstanceData, const FNiagaraSystemInstanceID& SystemInstance) override {}
 	virtual bool InitPerInstanceData(void* PerInstanceData, FNiagaraSystemInstance* SystemInstance) override;
diff --git a/Engine/Plugins/FX/Niagara/Source/Niagara/Classes/NiagaraDataInterfaceNeighborGrid3D.h b/Engine/Plugins/FX/Niagara/Source/Niagara/Classes/NiagaraDataInterfaceNeighborGrid3D.h
--- a/Engine/Plugins/FX/Niagara/Source/Niagara/Classes/NiagaraDataInterfaceNeighborGrid3D.h
+++ b/Engine/Plugins/FX/Niagara/Source/Niagara/Classes/NiagaraDataInterfaceNeighborGrid3D.h
@@ -4,6 +4,7 @@
 #include "NiagaraDataInterface.h"
 #include "NiagaraCommon.h"
 #include "NiagaraDataInterfaceRW.h"
+#include "NiagaraRenderGraphUtils.h"
 #include "ClearQuad.h"
 #include "Niagara/Private/NiagaraStats.h"
 
@@ -18,35 +19,26 @@ class FNiagaraSystemInstance;
 class NeighborGrid3DRWInstanceData
 {
 public:
+	void ResizeBuffers(FRDGBuilder& GraphBuilder);
 
-	~NeighborGrid3DRWInstanceData()
-	{
-#if STATS
-		DEC_MEMORY_STAT_BY(STAT_NiagaraGPUDataInterfaceMemory, GPUMemory);
-#endif
-	}
-
-	void ResizeBuffers();
+	FIntVector	NumCells = FIntVector::ZeroValue;
+	float		CellSize = 0.0f;
+	bool		SetGridFromCellSize = false;
+	uint32		MaxNeighborsPerCell = 0;
+	FVector		WorldBBoxSize = FVector::ZeroVector;
 
-	FIntVector NumCells;
-	float CellSize;
-	bool SetGridFromCellSize;
-	uint32 MaxNeighborsPerCell;	
-	FVector WorldBBoxSize;
+	bool		NeedsRealloc_GT = false;
+	bool		NeedsRealloc_RT = false;
 
-	bool NeedsRealloc = false;
-
-	FRWBuffer NeighborhoodBuffer;
-	FRWBuffer NeighborhoodCountBuffer;
-
-#if STATS
-	int32 GPUMemory = 0;
-#endif
+	FNiagaraPooledRWBuffer NeighborhoodBuffer;
+	FNiagaraPooledRWBuffer NeighborhoodCountBuffer;
 };
 
 struct FNiagaraDataInterfaceProxyNeighborGrid3D : public FNiagaraDataInterfaceProxyRW
 {	
-	virtual void PreStage(FRHICommandList& RHICmdList, const FNiagaraDataInterfaceStageArgs& Context) override;
+	virtual void PreStage(const FNDIGpuComputePreStageContext& Context) override;
+	virtual void PostSimulate(const FNDIGpuComputePostSimulateContext& Context) override;
+
 	virtual void ConsumePerInstanceDataFromGameThread(void* PerInstanceData, const FNiagaraSystemInstanceID& Instance) override {}
 	virtual int32 PerInstanceDataPassedToRenderThreadSize() const override { return sizeof(NeighborGrid3DRWInstanceData); }	
 
@@ -68,11 +60,11 @@ class NIAGARA_API UNiagaraDataInterfaceNeighborGrid3D : public UNiagaraDataInter
 		SHADER_PARAMETER(FVector3f,			CellSize)
 		SHADER_PARAMETER(FVector3f,			WorldBBoxSize)
 
-		SHADER_PARAMETER(int32,				MaxNeighborsPerCellValue)
-		SHADER_PARAMETER_SRV(Buffer<int>,	ParticleNeighbors)
-		SHADER_PARAMETER_SRV(Buffer<int>,	ParticleNeighborCount)
-		SHADER_PARAMETER_UAV(RWBuffer<int>, OutputParticleNeighbors)
-		SHADER_PARAMETER_UAV(RWBuffer<int>, OutputParticleNeighborCount)
+		SHADER_PARAMETER(int32,							MaxNeighborsPerCellValue)
+		SHADER_PARAMETER_RDG_BUFFER_SRV(Buffer<int>,	ParticleNeighbors)
+		SHADER_PARAMETER_RDG_BUFFER_SRV(Buffer<int>,	ParticleNeighborCount)
+		SHADER_PARAMETER_RDG_BUFFER_UAV(RWBuffer<int>,	OutputParticleNeighbors)
+		SHADER_PARAMETER_RDG_BUFFER_UAV(RWBuffer<int>,	OutputParticleNeighborCount)
 	END_SHADER_PARAMETER_STRUCT()
 
 public:	
diff --git a/Engine/Plugins/FX/Niagara/Source/Niagara/Classes/NiagaraDataInterfaceParticleRead.h b/Engine/Plugins/FX/Niagara/Source/Niagara/Classes/NiagaraDataInterfaceParticleRead.h
--- a/Engine/Plugins/FX/Niagara/Source/Niagara/Classes/NiagaraDataInterfaceParticleRead.h
+++ b/Engine/Plugins/FX/Niagara/Source/Niagara/Classes/NiagaraDataInterfaceParticleRead.h
@@ -11,9 +11,23 @@ UCLASS(EditInlineNew, Category = "ParticleRead", meta = (DisplayName = "Particle
 class NIAGARA_API UNiagaraDataInterfaceParticleRead : public UNiagaraDataInterfaceRWBase
 {
 	GENERATED_UCLASS_BODY()
-public:
-	DECLARE_NIAGARA_DI_PARAMETER();
 
+	BEGIN_SHADER_PARAMETER_STRUCT(FShaderParameters, )
+		SHADER_PARAMETER(int,				NumSpawnedParticles)
+		SHADER_PARAMETER(int,				SpawnedParticlesAcquireTag)
+		SHADER_PARAMETER(uint32,			InstanceCountOffset)
+		SHADER_PARAMETER(uint32,			ParticleStrideFloat)
+		SHADER_PARAMETER(uint32,			ParticleStrideInt)
+		SHADER_PARAMETER(uint32,			ParticleStrideHalf)
+		SHADER_PARAMETER(int,				AcquireTagRegisterIndex)
+		SHADER_PARAMETER_SRV(Buffer<int>,	SpawnedIDsBuffer)
+		SHADER_PARAMETER_SRV(Buffer<int>,	IDToIndexTable)
+		SHADER_PARAMETER_SRV(Buffer<float>,	InputFloatBuffer)
+		SHADER_PARAMETER_SRV(Buffer<int>,	InputIntBuffer)
+		SHADER_PARAMETER_SRV(Buffer<half>,	InputHalfBuffer)
+	END_SHADER_PARAMETER_STRUCT()
+
+public:
 	UPROPERTY(EditAnywhere, Category = "ParticleRead")
 	FString EmitterName;
 
@@ -33,9 +47,16 @@ public:
 	virtual bool CanExecuteOnTarget(ENiagaraSimTarget Target) const override { return true; }
 	virtual bool Equals(const UNiagaraDataInterface* Other) const override;
 #if WITH_EDITORONLY_DATA
+	virtual bool AppendCompileHash(FNiagaraCompileHashVisitor* InVisitor) const override;
 	virtual void GetParameterDefinitionHLSL(const FNiagaraDataInterfaceGPUParamInfo& ParamInfo, FString& OutHLSL) override;
 	virtual bool GetFunctionHLSL(const FNiagaraDataInterfaceGPUParamInfo& ParamInfo, const FNiagaraDataInterfaceGeneratedFunction& FunctionInfo, int FunctionInstanceIndex, FString& OutHLSL) override;
 #endif
+	virtual bool UseLegacyShaderBindings() const override { return false; }
+	virtual void BuildShaderParameters(FNiagaraShaderParametersBuilder& ShaderParametersBuilder) const override;
+	virtual void SetShaderParameters(const FNiagaraDataInterfaceSetShaderParametersContext& Context) const override;
+	virtual FNiagaraDataInterfaceParametersCS* CreateShaderStorage(const FNiagaraDataInterfaceGPUParamInfo& ParameterInfo, const FShaderParameterMap& ParameterMap) const override;
+	virtual const FTypeLayoutDesc* GetShaderStorageType() const override;
+
 	virtual void ProvidePerInstanceDataForRenderThread(void* DataForRenderThread, void* PerInstanceData, const FNiagaraSystemInstanceID& SystemInstance) override;
 #if WITH_EDITOR	
 	virtual void GetFeedback(UNiagaraSystem* Asset, UNiagaraComponent* Component, TArray<FNiagaraDataInterfaceError>& OutErrors, TArray<FNiagaraDataInterfaceFeedback>& Warnings, TArray<FNiagaraDataInterfaceFeedback>& Info) override;
diff --git a/Engine/Plugins/FX/Niagara/Source/Niagara/Classes/NiagaraDataInterfaceRW.h b/Engine/Plugins/FX/Niagara/Source/Niagara/Classes/NiagaraDataInterfaceRW.h
--- a/Engine/Plugins/FX/Niagara/Source/Niagara/Classes/NiagaraDataInterfaceRW.h
+++ b/Engine/Plugins/FX/Niagara/Source/Niagara/Classes/NiagaraDataInterfaceRW.h
@@ -28,8 +28,6 @@ public:
 	// The dispatch will use the CPU count, which is potentially an over-estimation, and the value inside the buffer will be used to clip instances that are not valid
 	virtual uint32 GetGPUInstanceCountOffset(FNiagaraSystemInstanceID SystemInstanceID) const { return INDEX_NONE; }
 
-	virtual void ClearBuffers(FRHICommandList& RHICmdList) {}
-
 	virtual FNiagaraDataInterfaceProxyRW* AsIterationProxy() override { return this; }
 };
 
diff --git a/Engine/Plugins/FX/Niagara/Source/Niagara/Classes/NiagaraDataInterfaceRenderTarget2D.h b/Engine/Plugins/FX/Niagara/Source/Niagara/Classes/NiagaraDataInterfaceRenderTarget2D.h
--- a/Engine/Plugins/FX/Niagara/Source/Niagara/Classes/NiagaraDataInterfaceRenderTarget2D.h
+++ b/Engine/Plugins/FX/Niagara/Source/Niagara/Classes/NiagaraDataInterfaceRenderTarget2D.h
@@ -47,11 +47,14 @@ struct FRenderTarget2DRWInstanceData_RenderThread
 	bool bRebuildMips = false;
 	bool bReadThisFrame = false;
 	bool bWroteThisFrame = false;
-	bool bNeedsTransition = false;
 
-	FSamplerStateRHIRef SamplerStateRHI;
-	FTexture2DRHIRef TextureRHI;
-	FUnorderedAccessViewRHIRef UnorderedAccessViewRHI;
+	FSamplerStateRHIRef	SamplerStateRHI;
+	FTexture2DRHIRef	TextureRHI;
+
+	FRDGTextureRef		TransientRDGTexture = nullptr;
+	FRDGTextureSRVRef	TransientRDGSRV = nullptr;
+	FRDGTextureUAVRef	TransientRDGUAV = nullptr;
+
 #if WITH_EDITORONLY_DATA
 	uint32 bPreviewTexture : 1;
 #endif
@@ -67,8 +70,8 @@ struct FNiagaraDataInterfaceProxyRenderTarget2DProxy : public FNiagaraDataInterf
 	virtual void ConsumePerInstanceDataFromGameThread(void* PerInstanceData, const FNiagaraSystemInstanceID& Instance) override {}
 	virtual int32 PerInstanceDataPassedToRenderThreadSize() const override { return 0; }
 
-	virtual void PostStage(FRHICommandList& RHICmdList, const FNiagaraDataInterfaceStageArgs& Context) override;
-	virtual void PostSimulate(FRHICommandList& RHICmdList, const FNiagaraDataInterfaceArgs& Context) override;
+	virtual void PostStage(const FNDIGpuComputePostStageContext& Context) override;
+	virtual void PostSimulate(const FNDIGpuComputePostSimulateContext& Context) override;
 
 	virtual FIntVector GetElementCount(FNiagaraSystemInstanceID SystemInstanceID) const override;
 
diff --git a/Engine/Plugins/FX/Niagara/Source/Niagara/Classes/NiagaraDataInterfaceRenderTarget2DArray.h b/Engine/Plugins/FX/Niagara/Source/Niagara/Classes/NiagaraDataInterfaceRenderTarget2DArray.h
--- a/Engine/Plugins/FX/Niagara/Source/Niagara/Classes/NiagaraDataInterfaceRenderTarget2DArray.h
+++ b/Engine/Plugins/FX/Niagara/Source/Niagara/Classes/NiagaraDataInterfaceRenderTarget2DArray.h
@@ -42,11 +42,14 @@ struct FRenderTarget2DArrayRWInstanceData_RenderThread
 	FIntVector Size = FIntVector(EForceInit::ForceInitToZero);
 	bool bWroteThisFrame = false;
 	bool bReadThisFrame = false;
-	bool bNeedsTransition = false;
 
 	FSamplerStateRHIRef SamplerStateRHI;
-	FTexture2DArrayRHIRef TextureRHI;
-	FUnorderedAccessViewRHIRef UnorderedAccessViewRHI;
+	TRefCountPtr<IPooledRenderTarget> RenderTarget;
+
+	FRDGTextureRef		TransientRDGTexture = nullptr;
+	FRDGTextureSRVRef	TransientRDGSRV = nullptr;
+	FRDGTextureUAVRef	TransientRDGUAV = nullptr;
+
 #if WITH_EDITORONLY_DATA
 	uint32 bPreviewTexture : 1;
 #endif
@@ -65,9 +68,7 @@ struct FNiagaraDataInterfaceProxyRenderTarget2DArrayProxy : public FNiagaraDataI
 		return 0;
 	}
 
-	virtual void ClearBuffers(FRHICommandList& RHICmdList) {}
-	virtual void PostStage(FRHICommandList& RHICmdList, const FNiagaraDataInterfaceStageArgs& Context) override;
-	virtual void PostSimulate(FRHICommandList& RHICmdList, const FNiagaraDataInterfaceArgs& Context) override;
+	virtual void PostSimulate(const FNDIGpuComputePostSimulateContext& Context) override;
 
 	virtual FIntVector GetElementCount(FNiagaraSystemInstanceID SystemInstanceID) const override;
 
diff --git a/Engine/Plugins/FX/Niagara/Source/Niagara/Classes/NiagaraDataInterfaceRenderTargetCube.h b/Engine/Plugins/FX/Niagara/Source/Niagara/Classes/NiagaraDataInterfaceRenderTargetCube.h
--- a/Engine/Plugins/FX/Niagara/Source/Niagara/Classes/NiagaraDataInterfaceRenderTargetCube.h
+++ b/Engine/Plugins/FX/Niagara/Source/Niagara/Classes/NiagaraDataInterfaceRenderTargetCube.h
@@ -44,8 +44,11 @@ struct FRenderTargetCubeRWInstanceData_RenderThread
 	bool bNeedsTransition = false;
 
 	FSamplerStateRHIRef SamplerStateRHI;
-	FTextureCubeRHIRef TextureRHI;
-	FUnorderedAccessViewRHIRef UnorderedAccessViewRHI;
+	TRefCountPtr<IPooledRenderTarget> RenderTarget;
+
+	FRDGTextureRef		TransientRDGTexture = nullptr;
+	FRDGTextureSRVRef	TransientRDGSRV = nullptr;
+	FRDGTextureUAVRef	TransientRDGUAV = nullptr;
 #if WITH_EDITORONLY_DATA
 	uint32 bPreviewTexture : 1;
 #endif
@@ -58,8 +61,7 @@ struct FRenderTargetCubeRWInstanceData_RenderThread
 struct FNiagaraDataInterfaceProxyRenderTargetCubeProxy : public FNiagaraDataInterfaceProxyRW
 {
 	virtual int32 PerInstanceDataPassedToRenderThreadSize() const override { return 0; }
-	virtual void PostStage(FRHICommandList& RHICmdList, const FNiagaraDataInterfaceStageArgs& Context) override;
-	virtual void PostSimulate(FRHICommandList& RHICmdList, const FNiagaraDataInterfaceArgs& Context) override;
+	virtual void PostSimulate(const FNDIGpuComputePostSimulateContext& Context) override;
 
 	virtual FIntVector GetElementCount(FNiagaraSystemInstanceID SystemInstanceID) const override;
 
diff --git a/Engine/Plugins/FX/Niagara/Source/Niagara/Classes/NiagaraDataInterfaceRenderTargetVolume.h b/Engine/Plugins/FX/Niagara/Source/Niagara/Classes/NiagaraDataInterfaceRenderTargetVolume.h
--- a/Engine/Plugins/FX/Niagara/Source/Niagara/Classes/NiagaraDataInterfaceRenderTargetVolume.h
+++ b/Engine/Plugins/FX/Niagara/Source/Niagara/Classes/NiagaraDataInterfaceRenderTargetVolume.h
@@ -41,11 +41,14 @@ struct FRenderTargetVolumeRWInstanceData_RenderThread
 	FIntVector Size = FIntVector(EForceInit::ForceInitToZero);
 	bool bWroteThisFrame = false;
 	bool bReadThisFrame = false;
-	bool bNeedsTransition = false;
 
 	FSamplerStateRHIRef SamplerStateRHI;
-	FTexture3DRHIRef TextureRHI;
-	FUnorderedAccessViewRHIRef UnorderedAccessViewRHI;
+	TRefCountPtr<IPooledRenderTarget> RenderTarget;
+
+	FRDGTextureRef		TransientRDGTexture = nullptr;
+	FRDGTextureSRVRef	TransientRDGSRV = nullptr;
+	FRDGTextureUAVRef	TransientRDGUAV = nullptr;
+
 #if WITH_EDITORONLY_DATA
 	uint32 bPreviewTexture : 1;
 #endif
@@ -58,8 +61,7 @@ struct FRenderTargetVolumeRWInstanceData_RenderThread
 struct FNiagaraDataInterfaceProxyRenderTargetVolumeProxy : public FNiagaraDataInterfaceProxyRW
 {
 	virtual int32 PerInstanceDataPassedToRenderThreadSize() const override { return 0; }
-	virtual void PostStage(FRHICommandList& RHICmdList, const FNiagaraDataInterfaceStageArgs& Context) override;
-	virtual void PostSimulate(FRHICommandList& RHICmdList, const FNiagaraDataInterfaceArgs& Context) override;
+	virtual void PostSimulate(const FNDIGpuComputePostSimulateContext& Context) override;
 
 	virtual FIntVector GetElementCount(FNiagaraSystemInstanceID SystemInstanceID) const override;
 
diff --git a/Engine/Plugins/FX/Niagara/Source/Niagara/Classes/NiagaraDataInterfaceTexture.h b/Engine/Plugins/FX/Niagara/Source/Niagara/Classes/NiagaraDataInterfaceTexture.h
--- a/Engine/Plugins/FX/Niagara/Source/Niagara/Classes/NiagaraDataInterfaceTexture.h
+++ b/Engine/Plugins/FX/Niagara/Source/Niagara/Classes/NiagaraDataInterfaceTexture.h
@@ -15,7 +15,7 @@ class NIAGARA_API UNiagaraDataInterfaceTexture : public UNiagaraDataInterface
 
 	BEGIN_SHADER_PARAMETER_STRUCT(FShaderParameters, )
 		SHADER_PARAMETER(FVector2f,				TextureSize)
-		SHADER_PARAMETER_TEXTURE(Texture2D,		Texture)
+		SHADER_PARAMETER_RDG_TEXTURE(Texture2D,	Texture)
 		SHADER_PARAMETER_SAMPLER(SamplerState,	TextureSampler)
 	END_SHADER_PARAMETER_STRUCT()
 
diff --git a/Engine/Plugins/FX/Niagara/Source/Niagara/Classes/NiagaraDataInterfaceVolumeTexture.h b/Engine/Plugins/FX/Niagara/Source/Niagara/Classes/NiagaraDataInterfaceVolumeTexture.h
--- a/Engine/Plugins/FX/Niagara/Source/Niagara/Classes/NiagaraDataInterfaceVolumeTexture.h
+++ b/Engine/Plugins/FX/Niagara/Source/Niagara/Classes/NiagaraDataInterfaceVolumeTexture.h
@@ -17,7 +17,7 @@ class NIAGARA_API UNiagaraDataInterfaceVolumeTexture : public UNiagaraDataInterf
 
 	BEGIN_SHADER_PARAMETER_STRUCT(FShaderParameters, )
 		SHADER_PARAMETER(FVector3f,				TextureSize)
-		SHADER_PARAMETER_TEXTURE(Texture3D,		Texture)
+		SHADER_PARAMETER_RDG_TEXTURE(Texture3D,		Texture)
 		SHADER_PARAMETER_SAMPLER(SamplerState,	TextureSampler)
 	END_SHADER_PARAMETER_STRUCT()
 
diff --git a/Engine/Plugins/FX/Niagara/Source/Niagara/Classes/NiagaraGPUProfilerInterface.h b/Engine/Plugins/FX/Niagara/Source/Niagara/Classes/NiagaraGPUProfilerInterface.h
--- a/Engine/Plugins/FX/Niagara/Source/Niagara/Classes/NiagaraGPUProfilerInterface.h
+++ b/Engine/Plugins/FX/Niagara/Source/Niagara/Classes/NiagaraGPUProfilerInterface.h
@@ -6,27 +6,43 @@
 #include "NiagaraEmitter.h"
 #include "NiagaraSystemGpuComputeProxy.h"
 
+struct FNiagaraComputeInstanceData;
+struct FNiagaraSimStageData;
+
 //////////////////////////////////////////////////////////////////////////
 // Public API for tracking GPU time when the profiler is enabled
+struct NIAGARA_API FNiagaraGpuProfileEvent
+{
+#if WITH_NIAGARA_GPU_PROFILER
+	friend class FNiagaraGPUProfiler;
+
+	explicit FNiagaraGpuProfileEvent(const FNiagaraComputeInstanceData& InstanceData, const FNiagaraSimStageData& SimStageData, const bool bFirstInstanceData);
+	explicit FNiagaraGpuProfileEvent(const FNiagaraComputeInstanceData& InstanceData, FName CustomStageName);
+	explicit FNiagaraGpuProfileEvent(const FNiagaraGpuProfileEvent& Other) = default;
+
+private:
+	uint32									bUniqueInstance : 1;
+	TWeakObjectPtr<class USceneComponent>	OwnerComponent;
+	FVersionedNiagaraEmitterWeakPtr			OwnerEmitter;
+	FName									StageName;
+#else
+	explicit FNiagaraGpuProfileEvent(const struct FNiagaraGpuDispatchInstance& DispatchInstance) {}
+	explicit FNiagaraGpuProfileEvent(const struct FNiagaraGpuDispatchInstance& DispatchInstance, FName CustomStageName) {}
+	explicit FNiagaraGpuProfileEvent(const FNiagaraGpuProfileEvent& Other) = default;
+#endif
+};
+
 struct NIAGARA_API FNiagaraGpuProfileScope
 {
 #if WITH_NIAGARA_GPU_PROFILER
-	FNiagaraGpuProfileScope(FRHICommandList& RHICmdList, const class FNiagaraGpuComputeDispatchInterface* ComputeDispatchInterface, FName StageName);
-	FNiagaraGpuProfileScope(FRHICommandList& RHICmdList, const class FNiagaraGpuComputeDispatchInterface* ComputeDispatchInterface, const struct FNiagaraGpuDispatchInstance& DispatchInstance);
-	FNiagaraGpuProfileScope(FRHICommandList& RHICmdList, const struct FNiagaraDataInterfaceArgs& Context, FName StageName);
-	FNiagaraGpuProfileScope(FRHICommandList& RHICmdList, const struct FNiagaraDataInterfaceSetArgs& Context, FName StageName);
-	FNiagaraGpuProfileScope(FRHICommandList& RHICmdList, const struct FNiagaraDataInterfaceStageArgs& Context, FName StageName);
+	explicit FNiagaraGpuProfileScope(FRHICommandList& RHICmdList, const class FNiagaraGpuComputeDispatchInterface* ComputeDispatchInterface, const FNiagaraGpuProfileEvent& Event);
 	~FNiagaraGpuProfileScope();
 
 private:
 	FRHICommandList& RHICmdList;
 	class FNiagaraGPUProfiler* GPUProfiler = nullptr;
 #else
-	FNiagaraGpuProfileScope(FRHICommandList& RHICmdList, const class FNiagaraGpuComputeDispatchInterface* ComputeDispatchInterface, FName StageName) {}
-	FNiagaraGpuProfileScope(FRHICommandList& RHICmdList, const class FNiagaraGpuComputeDispatchInterface* ComputeDispatchInterface, const struct FNiagaraGpuDispatchInstance& DispatchInstance) {}
-	FNiagaraGpuProfileScope(FRHICommandList& RHICmdList, const struct FNiagaraDataInterfaceArgs& Context, FName StageName) {}
-	FNiagaraGpuProfileScope(FRHICommandList& RHICmdList, const struct FNiagaraDataInterfaceSetArgs& Context, FName StageName) {}
-	FNiagaraGpuProfileScope(FRHICommandList& RHICmdList, const struct FNiagaraDataInterfaceStageArgs& Context, FName StageName) {}
+	explicit FNiagaraGpuProfileScope(FRHICommandList& RHICmdList, const class FNiagaraGpuComputeDispatchInterface* ComputeDispatchInterface, const FNiagaraGpuProfileEvent& Event) {}
 	~FNiagaraGpuProfileScope() {}
 #endif
 };
@@ -36,13 +52,6 @@ private:
 /** Results generated when the frame is ready and sent to that game thread */
 struct FNiagaraGpuFrameResults : public TSharedFromThis<FNiagaraGpuFrameResults, ESPMode::ThreadSafe>
 {
-	struct FStageResults
-	{
-		int32	NumDispatchGroups = 0;
-		int32	NumDispatches = 0;
-		uint64	DurationMicroseconds = 0;
-	};
-
 	struct FDispatchResults
 	{
 		uint32									bUniqueInstance : 1;		// Set only once for all dispatches from an instance across all ticks
@@ -53,7 +62,9 @@ struct FNiagaraGpuFrameResults : public TSharedFromThis<FNiagaraGpuFrameResults,
 	};
 
 	uintptr_t					OwnerContext = 0;
-	FStageResults				StageResults[ENiagaraGpuComputeTickStage::Max];
+	int32						TotalDispatches = 0;
+	int32						TotalDispatchGroups = 0;
+	uint64						TotalDurationMicroseconds = 0;
 	TArray<FDispatchResults>	DispatchResults;
 };
 
@@ -96,4 +107,3 @@ protected:
 };
 
 #endif //WITH_NIAGARA_GPU_PROFILER
-
diff --git a/Engine/Plugins/FX/Niagara/Source/Niagara/Classes/NiagaraGPUSystemTick.h b/Engine/Plugins/FX/Niagara/Source/Niagara/Classes/NiagaraGPUSystemTick.h
--- a/Engine/Plugins/FX/Niagara/Source/Niagara/Classes/NiagaraGPUSystemTick.h
+++ b/Engine/Plugins/FX/Niagara/Source/Niagara/Classes/NiagaraGPUSystemTick.h
@@ -7,6 +7,9 @@
 
 struct FNiagaraComputeDataInterfaceInstanceData
 {
+	UE_NONCOPYABLE(FNiagaraComputeDataInterfaceInstanceData);
+	FNiagaraComputeDataInterfaceInstanceData() {}
+
 	void* PerInstanceDataForRT = nullptr;
 	TMap<FNiagaraDataInterfaceProxy*, int32> InterfaceProxiesToOffsets;
 	uint32 PerInstanceDataSize = 0;
@@ -15,6 +18,8 @@ struct FNiagaraComputeDataInterfaceInstanceData
 
 struct FNiagaraComputeInstanceData
 {
+	UE_NONCOPYABLE(FNiagaraComputeInstanceData);
+
 	struct FPerStageInfo
 	{
 		FPerStageInfo() {}
diff --git a/Engine/Plugins/FX/Niagara/Source/Niagara/Private/DataInterface/NiagaraDataInterfaceEmitterProperties.cpp b/Engine/Plugins/FX/Niagara/Source/Niagara/Private/DataInterface/NiagaraDataInterfaceEmitterProperties.cpp
--- a/Engine/Plugins/FX/Niagara/Source/Niagara/Private/DataInterface/NiagaraDataInterfaceEmitterProperties.cpp
+++ b/Engine/Plugins/FX/Niagara/Source/Niagara/Private/DataInterface/NiagaraDataInterfaceEmitterProperties.cpp
@@ -65,9 +65,6 @@ namespace NDIEmitterPropertiesLocal
 			return sizeof(FInstanceData_GameToRender);
 		}
 	
-		virtual void PreStage(FRHICommandList& RHICmdList, const FNiagaraDataInterfaceStageArgs& Context) override {}
-		virtual void PostSimulate(FRHICommandList& RHICmdList, const FNiagaraDataInterfaceArgs& Context) override {}
-
 		TMap<FNiagaraSystemInstanceID, FInstanceData_RenderThread>	PerInstanceData_RenderThread;
 	};
 }
diff --git a/Engine/Plugins/FX/Niagara/Source/Niagara/Private/DataInterface/NiagaraDataInterfaceGBuffer.cpp b/Engine/Plugins/FX/Niagara/Source/Niagara/Private/DataInterface/NiagaraDataInterfaceGBuffer.cpp
--- a/Engine/Plugins/FX/Niagara/Source/Niagara/Private/DataInterface/NiagaraDataInterfaceGBuffer.cpp
+++ b/Engine/Plugins/FX/Niagara/Source/Niagara/Private/DataInterface/NiagaraDataInterfaceGBuffer.cpp
@@ -17,7 +17,7 @@
 namespace NiagaraDataInterfaceGBufferLocal
 {
 	BEGIN_SHADER_PARAMETER_STRUCT(FShaderParameters, )
-		SHADER_PARAMETER_TEXTURE(Texture2D, VelocityTexture)
+		SHADER_PARAMETER_RDG_TEXTURE(Texture2D, VelocityTexture)
 		SHADER_PARAMETER_SAMPLER(SamplerState, VelocityTextureSampler)
 	END_SHADER_PARAMETER_STRUCT()
 
@@ -173,16 +173,21 @@ void UNiagaraDataInterfaceGBuffer::BuildShaderParameters(FNiagaraShaderParameter
 
 void UNiagaraDataInterfaceGBuffer::SetShaderParameters(const FNiagaraDataInterfaceSetShaderParametersContext& Context) const
 {
-	FRHITexture* VelocityTexture = GBlackTexture->TextureRHI;
-
-	if (FNiagaraSceneTextureParameters* NiagaraSceneTextures = static_cast<const FNiagaraGpuComputeDispatch&>(Context.GetComputeDispatchInterface()).GetNiagaraSceneTextures())	//-BATCHERTODO:
+	FRDGTextureRef VelocityTexture = nullptr;
+	NiagaraDataInterfaceGBufferLocal::FShaderParameters* Parameters = Context.GetParameterNestedStruct<NiagaraDataInterfaceGBufferLocal::FShaderParameters>();
+	if (Context.IsResourceBound(&Parameters->VelocityTexture))
 	{
-		if (FRDGTexture* VelocityRDGTexture = NiagaraSceneTextures->Velocity.GetTexture())
+		if (FNiagaraSceneTextureParameters* NiagaraSceneTextures = static_cast<const FNiagaraGpuComputeDispatch&>(Context.GetComputeDispatchInterface()).GetNiagaraSceneTextures())	//-BATCHERTODO:
 		{
-			VelocityTexture = VelocityRDGTexture->GetRHI();
+			VelocityTexture = NiagaraSceneTextures->Velocity.GetTexture();
+		}
+
+		if (VelocityTexture == nullptr)
+		{
+			VelocityTexture = Context.GetComputeDispatchInterface().GetBlackTexture(Context.GetGraphBuilder(), ETextureDimension::Texture2D);
 		}
 	}
-	NiagaraDataInterfaceGBufferLocal::FShaderParameters* Parameters = Context.GetParameterNestedStruct<NiagaraDataInterfaceGBufferLocal::FShaderParameters>();
+
 	Parameters->VelocityTexture = VelocityTexture;
 	Parameters->VelocityTextureSampler = TStaticSamplerState<SF_Bilinear, AM_Clamp, AM_Clamp, AM_Clamp>::GetRHI();
 }
diff --git a/Engine/Plugins/FX/Niagara/Source/Niagara/Private/DataInterface/NiagaraDataInterfaceSimpleCounter.cpp b/Engine/Plugins/FX/Niagara/Source/Niagara/Private/DataInterface/NiagaraDataInterfaceSimpleCounter.cpp
--- a/Engine/Plugins/FX/Niagara/Source/Niagara/Private/DataInterface/NiagaraDataInterfaceSimpleCounter.cpp
+++ b/Engine/Plugins/FX/Niagara/Source/Niagara/Private/DataInterface/NiagaraDataInterfaceSimpleCounter.cpp
@@ -50,53 +50,72 @@ struct FNDISimpleCounterProxy : public FNiagaraDataInterfaceProxy
 	virtual void ConsumePerInstanceDataFromGameThread(void* PerInstanceData, const FNiagaraSystemInstanceID& Instance) override {}
 	virtual int32 PerInstanceDataPassedToRenderThreadSize() const override { return 0; }
 
-	virtual void PreStage(FRHICommandList& RHICmdList, const FNiagaraDataInterfaceStageArgs& Context) override
+	virtual void PreStage(const FNDIGpuComputePreStageContext& Context) override
 	{
-		if (FNDISimpleCounterInstanceData_RenderThread* InstanceData = PerInstanceData_RenderThread.Find(Context.SystemInstanceID) )
+		if (FNDISimpleCounterInstanceData_RenderThread* InstanceData = PerInstanceData_RenderThread.Find(Context.GetSystemInstanceID()) )
 		{
 			if ( InstanceData->CountValue.IsSet() )
 			{
 				//-OPT: We could push this into the count manager and batch set as part of the clear process
-				const FNiagaraGPUInstanceCountManager& CounterManager = Context.ComputeDispatchInterface->GetGPUInstanceCounterManager();
+				const FNiagaraGPUInstanceCountManager& CounterManager = Context.GetComputeDispatchInterface().GetGPUInstanceCounterManager();
 				const FRWBuffer& CountBuffer = CounterManager.GetInstanceCountBuffer();
 
-				const TPair<uint32, int32> DataToClear(InstanceData->CountOffset, InstanceData->CountValue.GetValue());
-				RHICmdList.Transition(FRHITransitionInfo(CountBuffer.UAV, ERHIAccess::UAVCompute, ERHIAccess::UAVCompute));
-				NiagaraClearCounts::ClearCountsInt(RHICmdList, CountBuffer.UAV, MakeArrayView(&DataToClear, 1) );
-				RHICmdList.Transition(FRHITransitionInfo(CountBuffer.UAV, ERHIAccess::UAVCompute, ERHIAccess::UAVCompute));
+				//-TODO:RDG: Once the count buffer is a graph resource this can be changed
+				AddPass(
+					Context.GetGraphBuilder(),
+					RDG_EVENT_NAME("NiagaraSimpleCounter::PreStage"),
+					[CountBufferUAV=CountBuffer.UAV, CountOffset=InstanceData->CountOffset, CountValue=InstanceData->CountValue.GetValue()](FRHICommandListImmediate& RHICmdList)
+					{
+						const TPair<uint32, int32> DataToClear(CountOffset, CountValue);
+						RHICmdList.Transition(FRHITransitionInfo(CountBufferUAV, ERHIAccess::UAVCompute, ERHIAccess::UAVCompute));
+						NiagaraClearCounts::ClearCountsInt(RHICmdList, CountBufferUAV, MakeArrayView(&DataToClear, 1) );
+						RHICmdList.Transition(FRHITransitionInfo(CountBufferUAV, ERHIAccess::UAVCompute, ERHIAccess::UAVCompute));
+					}
+				);
+
 				InstanceData->CountValue.Reset();
 			}
 		}
 	}
 
-	virtual void PostSimulate(FRHICommandList& RHICmdList, const FNiagaraDataInterfaceArgs& Context) override
+	virtual void PostSimulate(const FNDIGpuComputePostSimulateContext& Context) override
 	{
-		if (FNiagaraUtilities::ShouldSyncGpuToCpu(GpuSyncMode))
+		if (FNiagaraUtilities::ShouldSyncGpuToCpu(GpuSyncMode) && Context.IsFinalPostSimulate())
 		{
-			if (FNDISimpleCounterInstanceData_RenderThread* InstanceData = PerInstanceData_RenderThread.Find(Context.SystemInstanceID))
+			FNiagaraSystemInstanceID SystemInstanceID = Context.GetSystemInstanceID();
+			if (FNDISimpleCounterInstanceData_RenderThread* InstanceData = PerInstanceData_RenderThread.Find(SystemInstanceID))
 			{
-				const FNiagaraGPUInstanceCountManager& CountManager = Context.ComputeDispatchInterface->GetGPUInstanceCounterManager();
-				FNiagaraGpuReadbackManager* ReadbackManager = Context.ComputeDispatchInterface->GetGpuReadbackManager();
-				ReadbackManager->EnqueueReadback(
-					RHICmdList,
-					CountManager.GetInstanceCountBuffer().Buffer,
-					InstanceData->CountOffset * sizeof(uint32), sizeof(uint32),
-					[SystemInstanceID=Context.SystemInstanceID, WeakOwner=WeakOwner, Proxy=this](TConstArrayView<TPair<void*, uint32>> ReadbackData)
+				//-TODO:RDG: This should be done using the graph builder
+				const FNiagaraGPUInstanceCountManager& CountManager = Context.GetComputeDispatchInterface().GetGPUInstanceCounterManager();
+				FNiagaraGpuReadbackManager* ReadbackManager = Context.GetComputeDispatchInterface().GetGpuReadbackManager();
+
+				AddPass(
+					Context.GetGraphBuilder(),
+					RDG_EVENT_NAME("NiagaraSimpleCounter::PostStage"),
+					[InstanceData, SystemInstanceID, ReadbackManager, CountBuffer=CountManager.GetInstanceCountBuffer().Buffer, Proxy=this](FRHICommandListImmediate& RHICmdList)
 					{
-						const int32 CounterValue = *reinterpret_cast<const int32*>(ReadbackData[0].Key);
-						AsyncTask(
-							ENamedThreads::GameThread,
-							[SystemInstanceID, CounterValue, WeakOwner, Proxy]()
+						ReadbackManager->EnqueueReadback(
+							RHICmdList,
+							CountBuffer,
+							InstanceData->CountOffset * sizeof(uint32), sizeof(uint32),
+							[SystemInstanceID, WeakOwner=Proxy->WeakOwner, Proxy](TConstArrayView<TPair<void*, uint32>> ReadbackData)
 							{
-								// FNiagaraDataInterfaceProxy do not outlive UNiagaraDataInterface so if our Object is valid so is the proxy
-								// Equally because we do not share instance IDs (monotonically increasing number) we won't ever stomp something that has 'gone away'
-								if ( WeakOwner.Get() )
-								{
-									if ( FNDISimpleCounterInstanceData_GameThread* InstanceData_GT = Proxy->PerInstanceData_GameThread.FindRef(SystemInstanceID) )
+								const int32 CounterValue = *reinterpret_cast<const int32*>(ReadbackData[0].Key);
+								AsyncTask(
+									ENamedThreads::GameThread,
+									[SystemInstanceID, CounterValue, WeakOwner, Proxy]()
 									{
-										InstanceData_GT->Counter = CounterValue;
+										// FNiagaraDataInterfaceProxy do not outlive UNiagaraDataInterface so if our Object is valid so is the proxy
+										// Equally because we do not share instance IDs (monotonically increasing number) we won't ever stomp something that has 'gone away'
+										if ( WeakOwner.Get() )
+										{
+											if ( FNDISimpleCounterInstanceData_GameThread* InstanceData_GT = Proxy->PerInstanceData_GameThread.FindRef(SystemInstanceID) )
+											{
+												InstanceData_GT->Counter = CounterValue;
+											}
+										}
 									}
-								}
+								);
 							}
 						);
 					}
diff --git a/Engine/Plugins/FX/Niagara/Source/Niagara/Private/DataInterface/NiagaraDataInterfaceUObjectPropertyReader.cpp b/Engine/Plugins/FX/Niagara/Source/Niagara/Private/DataInterface/NiagaraDataInterfaceUObjectPropertyReader.cpp
--- a/Engine/Plugins/FX/Niagara/Source/Niagara/Private/DataInterface/NiagaraDataInterfaceUObjectPropertyReader.cpp
+++ b/Engine/Plugins/FX/Niagara/Source/Niagara/Private/DataInterface/NiagaraDataInterfaceUObjectPropertyReader.cpp
@@ -135,9 +135,6 @@ namespace NDIUObjectPropertyReaderLocal
 			return sizeof(FInstanceData_GameToRender);
 		}
 	
-		virtual void PreStage(FRHICommandList& RHICmdList, const FNiagaraDataInterfaceStageArgs& Context) override {}
-		virtual void PostSimulate(FRHICommandList& RHICmdList, const FNiagaraDataInterfaceArgs& Context) override {}
-
 		TMap<FNiagaraSystemInstanceID, FInstanceData_RenderThread>	PerInstanceData_RenderThread;
 	};
 
diff --git a/Engine/Plugins/FX/Niagara/Source/Niagara/Private/NiagaraDataInterface.cpp b/Engine/Plugins/FX/Niagara/Source/Niagara/Private/NiagaraDataInterface.cpp
--- a/Engine/Plugins/FX/Niagara/Source/Niagara/Private/NiagaraDataInterface.cpp
+++ b/Engine/Plugins/FX/Niagara/Source/Niagara/Private/NiagaraDataInterface.cpp
@@ -6,6 +6,7 @@
 #include "NiagaraTypes.h"
 #include "ShaderParameterUtils.h"
 #include "NiagaraGPUSystemTick.h"
+#include "NiagaraGpuComputeDispatch.h"
 #include "NiagaraShader.h"
 #include "NiagaraComponent.h"
 #include "ShaderCompilerCore.h"
@@ -221,6 +222,38 @@ FSimpleMulticastDelegate& UNiagaraDataInterface::OnErrorsRefreshed()
 
 #endif
 
+FRDGExternalAccessQueue& FNDIGpuComputeContext::GetRDGExternalAccessQueue() const
+{
+	return static_cast<const FNiagaraGpuComputeDispatch&>(ComputeDispatchInterface).GetCurrentPassExternalAccessQueue();
+}
+
+FNiagaraSystemInstanceID FNDIGpuComputePrePostStageContext::GetSystemInstanceID() const
+{
+	return SystemTick.SystemInstanceID;
+}
+
+FVector3f FNDIGpuComputePrePostStageContext::GetSystemLWCTile() const
+{
+	return SystemTick.SystemGpuComputeProxy->GetSystemLWCTile();
+}
+
+bool FNDIGpuComputePrePostStageContext::IsOutputStage() const
+{
+	check(DataInterfaceProxy);
+	return ComputeInstanceData.IsOutputStage(DataInterfaceProxy, SimStageData.StageIndex);
+}
+
+bool FNDIGpuComputePrePostStageContext::IsIterationStage() const
+{
+	check(DataInterfaceProxy);
+	return ComputeInstanceData.IsIterationStage(DataInterfaceProxy, SimStageData.StageIndex);
+}
+
+FRDGExternalAccessQueue& FNiagaraDataInterfaceSetShaderParametersContext::GetRDGExternalAccessQueue() const
+{
+	return static_cast<const FNiagaraGpuComputeDispatch&>(ComputeDispatchInterface).GetCurrentPassExternalAccessQueue();
+}
+
 FNiagaraSystemInstanceID FNiagaraDataInterfaceSetShaderParametersContext::GetSystemInstanceID() const
 {
 	return SystemTick.SystemInstanceID;
@@ -239,6 +272,25 @@ bool FNiagaraDataInterfaceSetShaderParametersContext::IsResourceBound(const void
 	return false;
 }
 
+bool FNiagaraDataInterfaceSetShaderParametersContext::IsParameterBound(const void* ParameterAddress) const
+{
+	const uint16 ByteOffset = uint16(uintptr_t(ParameterAddress) - uintptr_t(BaseParameters));
+	for (const FShaderParameterBindings::FParameter& Parameter : ShaderRef->Bindings.Parameters)
+	{
+		if (Parameter.ByteOffset > ByteOffset)
+		{
+			return false;
+		}
+
+		const uint16 Delta = ByteOffset - Parameter.ByteOffset;
+		if (Delta < Parameter.ByteSize)
+		{
+			return true;
+		}
+	}
+	return false;
+}
+
 bool FNiagaraDataInterfaceSetShaderParametersContext::IsOutputStage() const
 {
 	return ComputeInstanceData.IsOutputStage(DataInterfaceProxy, SimStageData.StageIndex);
diff --git a/Engine/Plugins/FX/Niagara/Source/Niagara/Private/NiagaraDataInterface2DArrayTexture.cpp b/Engine/Plugins/FX/Niagara/Source/Niagara/Private/NiagaraDataInterface2DArrayTexture.cpp
--- a/Engine/Plugins/FX/Niagara/Source/Niagara/Private/NiagaraDataInterface2DArrayTexture.cpp
+++ b/Engine/Plugins/FX/Niagara/Source/Niagara/Private/NiagaraDataInterface2DArrayTexture.cpp
@@ -2,13 +2,13 @@
 
 #include "NiagaraDataInterface2DArrayTexture.h"
 #include "NiagaraComputeExecutionContext.h"
+#include "NiagaraGpuComputeDispatchInterface.h"
 #include "NiagaraShader.h"
 #include "NiagaraShaderParametersBuilder.h"
 #include "NiagaraSystemInstance.h"
 
 #include "Engine/Texture2DArray.h"
 #include "Engine/TextureRenderTarget2DArray.h"
-#include "ShaderParameterUtils.h"
 
 #define LOCTEXT_NAMESPACE "UNiagaraDataInterface2DArrayTexture"
 
@@ -27,8 +27,9 @@ struct FNDITexture2DArrayInstanceData_RenderThread
 {
 	FSamplerStateRHIRef		SamplerStateRHI;
 	FTextureReferenceRHIRef	TextureReferenceRHI;
-	FTextureRHIRef			ResolvedTextureRHI;
 	FVector3f				TextureSize;
+
+	FRDGTextureRef			TransientRDGTexture = nullptr;
 };
 
 struct FNiagaraDataInterfaceProxyTexture2DArray : public FNiagaraDataInterfaceProxy
@@ -36,26 +37,13 @@ struct FNiagaraDataInterfaceProxyTexture2DArray : public FNiagaraDataInterfacePr
 	virtual void ConsumePerInstanceDataFromGameThread(void* PerInstanceData, const FNiagaraSystemInstanceID& Instance) override { check(false); }
 	virtual int32 PerInstanceDataPassedToRenderThreadSize() const override { return 0; }
 
-	virtual void PreStage(FRHICommandList& RHICmdList, const FNiagaraDataInterfaceStageArgs& Context) override
+	virtual void PostSimulate(const FNDIGpuComputePostSimulateContext& Context) override
 	{
-		if (FNDITexture2DArrayInstanceData_RenderThread* InstanceData = InstanceData_RT.Find(Context.SystemInstanceID))
+		if (Context.IsFinalPostSimulate())
 		{
-			// Because the underlying reference can have a switch in flight on the RHI we get the referenced texture
-			// here, ensure it's valid (as it could be queued for delete) and cache until next round.  If we were
-			// to release the reference in PostStage / PostSimulate we still stand a chance the the transition we
-			// queue will be invalid by the time it is processed on the RHI thread.
-			if (Context.SimStageData->bFirstStage && InstanceData->TextureReferenceRHI.IsValid())
+			if (FNDITexture2DArrayInstanceData_RenderThread* InstanceData = InstanceData_RT.Find(Context.GetSystemInstanceID()))
 			{
-				InstanceData->ResolvedTextureRHI = InstanceData->TextureReferenceRHI->GetReferencedTexture();
-				if (InstanceData->ResolvedTextureRHI && !InstanceData->ResolvedTextureRHI->IsValid())
-				{
-					InstanceData->ResolvedTextureRHI = nullptr;
-				}
-			}
-			if (InstanceData->ResolvedTextureRHI)
-			{
-				// Make sure the texture is readable, we don't know where it's coming from.
-				RHICmdList.Transition(FRHITransitionInfo(InstanceData->ResolvedTextureRHI, ERHIAccess::Unknown, ERHIAccess::SRVMask));
+				InstanceData->TransientRDGTexture = nullptr;
 			}
 		}
 	}
@@ -199,14 +187,15 @@ bool UNiagaraDataInterface2DArrayTexture::PerInstanceTick(void* PerInstanceData,
 					if (RT_Texture)
 					{
 						InstanceData.TextureReferenceRHI = RT_Texture->TextureReference.TextureReferenceRHI;
-						InstanceData.SamplerStateRHI = RT_Texture->GetResource() ? RT_Texture->GetResource()->SamplerStateRHI : nullptr;
+						InstanceData.SamplerStateRHI = RT_Texture->GetResource() ? RT_Texture->GetResource()->SamplerStateRHI.GetReference() : TStaticSamplerState<SF_Point>::GetRHI();
+						InstanceData.TextureSize = FVector3f(RT_TextureSize.X, RT_TextureSize.Y, RT_TextureSize.Z);
 					}
 					else
 					{
 						InstanceData.TextureReferenceRHI = nullptr;
 						InstanceData.SamplerStateRHI = nullptr;
+						InstanceData.TextureSize = FVector3f::ZeroVector;
 					}
-					InstanceData.TextureSize = FVector3f(RT_TextureSize.X, RT_TextureSize.Y, RT_TextureSize.Z);
 				}
 			);
 		}
@@ -265,21 +254,39 @@ void UNiagaraDataInterface2DArrayTexture::BuildShaderParameters(FNiagaraShaderPa
 
 void UNiagaraDataInterface2DArrayTexture::SetShaderParameters(const FNiagaraDataInterfaceSetShaderParametersContext& Context) const
 {
-	const FNiagaraDataInterfaceProxyTexture2DArray& TextureProxy = Context.GetProxy<FNiagaraDataInterfaceProxyTexture2DArray>();
-	const FNDITexture2DArrayInstanceData_RenderThread* InstanceData = TextureProxy.InstanceData_RT.Find(Context.GetSystemInstanceID());
+	FNiagaraDataInterfaceProxyTexture2DArray& TextureProxy = Context.GetProxy<FNiagaraDataInterfaceProxyTexture2DArray>();
+	FNDITexture2DArrayInstanceData_RenderThread* InstanceData = TextureProxy.InstanceData_RT.Find(Context.GetSystemInstanceID());
 
 	FShaderParameters* Parameters = Context.GetParameterNestedStruct<FShaderParameters>();
-	if (InstanceData && InstanceData->ResolvedTextureRHI.IsValid())
+	if (InstanceData && InstanceData->TextureReferenceRHI.IsValid())
 	{
-		Parameters->TextureSize		= InstanceData->TextureSize;
-		Parameters->Texture			= InstanceData->ResolvedTextureRHI;
-		Parameters->TextureSampler	= InstanceData->SamplerStateRHI ? InstanceData->SamplerStateRHI : GBlackArrayTexture->SamplerStateRHI;
+		Parameters->TextureSize = InstanceData->TextureSize;
+		Parameters->TextureSampler = InstanceData->SamplerStateRHI;
+		if (Context.IsResourceBound(&Parameters->Texture))
+		{
+			FRDGTextureRef RDGTexture = InstanceData ? InstanceData->TransientRDGTexture : nullptr;
+			if (InstanceData && RDGTexture == nullptr)
+			{
+				FTextureRHIRef ResolvedTextureRHI = InstanceData->TextureReferenceRHI->GetReferencedTexture();
+				if (ResolvedTextureRHI.IsValid())
+				{
+					InstanceData->TransientRDGTexture = Context.GetGraphBuilder().FindExternalTexture(ResolvedTextureRHI);
+					if (InstanceData->TransientRDGTexture == nullptr)
+					{
+						InstanceData->TransientRDGTexture = Context.GetGraphBuilder().RegisterExternalTexture(CreateRenderTarget(ResolvedTextureRHI, TEXT("NiagaraTexture2DArray")));
+					}
+					RDGTexture = InstanceData->TransientRDGTexture;
+				}
+			}
+
+			Parameters->Texture = RDGTexture ? RDGTexture : Context.GetComputeDispatchInterface().GetBlackTexture(Context.GetGraphBuilder(), ETextureDimension::Texture2DArray);
+		}
 	}
 	else
 	{
-		Parameters->TextureSize		= FVector3f::ZeroVector;
-		Parameters->Texture			= GBlackArrayTexture->TextureRHI;
-		Parameters->TextureSampler	= GBlackArrayTexture->SamplerStateRHI;
+		Parameters->TextureSize = FVector3f::ZeroVector;
+		Parameters->TextureSampler = TStaticSamplerState<SF_Point>::GetRHI();
+		Parameters->Texture = Context.GetComputeDispatchInterface().GetBlackTexture(Context.GetGraphBuilder(), ETextureDimension::Texture2DArray);
 	}
 }
 
diff --git a/Engine/Plugins/FX/Niagara/Source/Niagara/Private/NiagaraDataInterfaceCubeTexture.cpp b/Engine/Plugins/FX/Niagara/Source/Niagara/Private/NiagaraDataInterfaceCubeTexture.cpp
--- a/Engine/Plugins/FX/Niagara/Source/Niagara/Private/NiagaraDataInterfaceCubeTexture.cpp
+++ b/Engine/Plugins/FX/Niagara/Source/Niagara/Private/NiagaraDataInterfaceCubeTexture.cpp
@@ -2,13 +2,13 @@
 
 #include "NiagaraDataInterfaceCubeTexture.h"
 #include "NiagaraComputeExecutionContext.h"
+#include "NiagaraGpuComputeDispatchInterface.h"
 #include "NiagaraShader.h"
 #include "NiagaraShaderParametersBuilder.h"
 #include "NiagaraSystemInstance.h"
 
 #include "Engine/TextureCube.h"
 #include "Engine/TextureRenderTargetCube.h"
-#include "ShaderParameterUtils.h"
 
 #define LOCTEXT_NAMESPACE "UNiagaraDataInterfaceCubeTexture"
 
@@ -27,8 +27,9 @@ struct FNDICubeTextureInstanceData_RenderThread
 {
 	FSamplerStateRHIRef		SamplerStateRHI;
 	FTextureReferenceRHIRef	TextureReferenceRHI;
-	FTextureRHIRef			ResolvedTextureRHI;
 	FIntPoint				TextureSize;
+
+	FRDGTextureRef			TransientRDGTexture = nullptr;
 };
 
 struct FNiagaraDataInterfaceProxyCubeTexture : public FNiagaraDataInterfaceProxy
@@ -36,26 +37,13 @@ struct FNiagaraDataInterfaceProxyCubeTexture : public FNiagaraDataInterfaceProxy
 	virtual void ConsumePerInstanceDataFromGameThread(void* PerInstanceData, const FNiagaraSystemInstanceID& Instance) override { check(false); }
 	virtual int32 PerInstanceDataPassedToRenderThreadSize() const override { return 0; }
 
-	virtual void PreStage(FRHICommandList& RHICmdList, const FNiagaraDataInterfaceStageArgs& Context) override
+	virtual void PostSimulate(const FNDIGpuComputePostSimulateContext& Context) override
 	{
-		if (FNDICubeTextureInstanceData_RenderThread* InstanceData = InstanceData_RT.Find(Context.SystemInstanceID))
+		if (Context.IsFinalPostSimulate())
 		{
-			// Because the underlying reference can have a switch in flight on the RHI we get the referenced texture
-			// here, ensure it's valid (as it could be queued for delete) and cache until next round.  If we were
-			// to release the reference in PostStage / PostSimulate we still stand a chance the the transition we
-			// queue will be invalid by the time it is processed on the RHI thread.
-			if (Context.SimStageData->bFirstStage && InstanceData->TextureReferenceRHI.IsValid())
+			if (FNDICubeTextureInstanceData_RenderThread* InstanceData = InstanceData_RT.Find(Context.GetSystemInstanceID()))
 			{
-				InstanceData->ResolvedTextureRHI = InstanceData->TextureReferenceRHI->GetReferencedTexture();
-				if (InstanceData->ResolvedTextureRHI && !InstanceData->ResolvedTextureRHI->IsValid())
-				{
-					InstanceData->ResolvedTextureRHI = nullptr;
-				}
-			}
-			if (InstanceData->ResolvedTextureRHI)
-			{
-				// Make sure the texture is readable, we don't know where it's coming from.
-				RHICmdList.Transition(FRHITransitionInfo(InstanceData->ResolvedTextureRHI, ERHIAccess::Unknown, ERHIAccess::SRVMask));
+				InstanceData->TransientRDGTexture = nullptr;
 			}
 		}
 	}
@@ -204,14 +192,15 @@ bool UNiagaraDataInterfaceCubeTexture::PerInstanceTick(void* PerInstanceData, FN
 					if (RT_Texture)
 					{
 						InstanceData.TextureReferenceRHI = RT_Texture->TextureReference.TextureReferenceRHI;
-						InstanceData.SamplerStateRHI = RT_Texture->GetResource() ? RT_Texture->GetResource()->SamplerStateRHI : nullptr;
+						InstanceData.SamplerStateRHI = RT_Texture->GetResource() ? RT_Texture->GetResource()->SamplerStateRHI.GetReference() : TStaticSamplerState<SF_Point>::GetRHI();
+						InstanceData.TextureSize = RT_TextureSize;
 					}
 					else
 					{
 						InstanceData.TextureReferenceRHI = nullptr;
 						InstanceData.SamplerStateRHI = nullptr;
+						InstanceData.TextureSize = FIntPoint::ZeroValue;
 					}
-					InstanceData.TextureSize = RT_TextureSize;
 				}
 			);
 		}
@@ -285,21 +274,39 @@ void UNiagaraDataInterfaceCubeTexture::BuildShaderParameters(FNiagaraShaderParam
 
 void UNiagaraDataInterfaceCubeTexture::SetShaderParameters(const FNiagaraDataInterfaceSetShaderParametersContext& Context) const
 {
-	const FNiagaraDataInterfaceProxyCubeTexture& TextureProxy = Context.GetProxy<FNiagaraDataInterfaceProxyCubeTexture>();
-	const FNDICubeTextureInstanceData_RenderThread* InstanceData = TextureProxy.InstanceData_RT.Find(Context.GetSystemInstanceID());
+	FNiagaraDataInterfaceProxyCubeTexture& TextureProxy = Context.GetProxy<FNiagaraDataInterfaceProxyCubeTexture>();
+	FNDICubeTextureInstanceData_RenderThread* InstanceData = TextureProxy.InstanceData_RT.Find(Context.GetSystemInstanceID());
 
 	FShaderParameters* Parameters = Context.GetParameterNestedStruct<FShaderParameters>();
-	if (InstanceData && InstanceData->ResolvedTextureRHI.IsValid())
+	if (InstanceData && InstanceData->TextureReferenceRHI.IsValid())
 	{
-		Parameters->TextureSize		= InstanceData->TextureSize;
-		Parameters->Texture			= InstanceData->ResolvedTextureRHI;
-		Parameters->TextureSampler	= InstanceData->SamplerStateRHI ? InstanceData->SamplerStateRHI : GBlackTextureCube->SamplerStateRHI;
+		Parameters->TextureSize = InstanceData->TextureSize;
+		Parameters->TextureSampler = InstanceData->SamplerStateRHI;
+		if (Context.IsResourceBound(&Parameters->Texture))
+		{
+			FRDGTextureRef RDGTexture = InstanceData ? InstanceData->TransientRDGTexture : nullptr;
+			if (InstanceData && RDGTexture == nullptr)
+			{
+				FTextureRHIRef ResolvedTextureRHI = InstanceData->TextureReferenceRHI->GetReferencedTexture();
+				if (ResolvedTextureRHI.IsValid())
+				{
+					InstanceData->TransientRDGTexture = Context.GetGraphBuilder().FindExternalTexture(ResolvedTextureRHI);
+					if (InstanceData->TransientRDGTexture == nullptr)
+					{
+						InstanceData->TransientRDGTexture = Context.GetGraphBuilder().RegisterExternalTexture(CreateRenderTarget(ResolvedTextureRHI, TEXT("NiagaraTextureCube")));
+					}
+					RDGTexture = InstanceData->TransientRDGTexture;
+				}
+			}
+
+			Parameters->Texture = RDGTexture ? RDGTexture : Context.GetComputeDispatchInterface().GetBlackTexture(Context.GetGraphBuilder(), ETextureDimension::TextureCube);
+		}
 	}
 	else
 	{
-		Parameters->TextureSize		= FIntPoint::ZeroValue;
-		Parameters->Texture			= GBlackTextureCube->TextureRHI;
-		Parameters->TextureSampler	= GBlackTextureCube->SamplerStateRHI;
+		Parameters->TextureSize = FIntPoint::ZeroValue;
+		Parameters->TextureSampler = TStaticSamplerState<SF_Point>::GetRHI();
+		Parameters->Texture = Context.GetComputeDispatchInterface().GetBlackTexture(Context.GetGraphBuilder(), ETextureDimension::TextureCube);
 	}
 }
 
diff --git a/Engine/Plugins/FX/Niagara/Source/Niagara/Private/NiagaraDataInterfaceExport.cpp b/Engine/Plugins/FX/Niagara/Source/Niagara/Private/NiagaraDataInterfaceExport.cpp
--- a/Engine/Plugins/FX/Niagara/Source/Niagara/Private/NiagaraDataInterfaceExport.cpp
+++ b/Engine/Plugins/FX/Niagara/Source/Niagara/Private/NiagaraDataInterfaceExport.cpp
@@ -18,8 +18,8 @@
 namespace NDIExportLocal
 {
 	BEGIN_SHADER_PARAMETER_STRUCT(FShaderParameters, )
-		SHADER_PARAMETER(uint32,				WriteBufferSize)
-		SHADER_PARAMETER_UAV(RWBuffer<uint>,	RWWriteBuffer)
+		SHADER_PARAMETER(uint32,						WriteBufferSize)
+		SHADER_PARAMETER_RDG_BUFFER_UAV(RWBuffer<uint>,	RWWriteBuffer)
 	END_SHADER_PARAMETER_STRUCT()
 
 	static const TCHAR*		TemplateShaderFile = TEXT("/Plugin/FX/Niagara/Private/NiagaraDataInterfaceExportTemplate.ush");
@@ -107,7 +107,9 @@ struct FNDIExportInstanceData_RenderThread
 
 	bool							bWriteBufferUsed = false;
 	uint32							WriteBufferInstanceCount = 0;
-	FRWBuffer						WriteBuffer;
+
+	FRDGBufferRef					RDGWriteBuffer = nullptr;
+	FRDGBufferUAVRef				RDGWriteBufferUAV = nullptr;
 };
 
 struct FNDIExportProxy : public FNiagaraDataInterfaceProxy
@@ -119,63 +121,20 @@ struct FNDIExportProxy : public FNiagaraDataInterfaceProxy
 		return 0;
 	}
 
-	virtual void PreStage(FRHICommandList& RHICmdList, const FNiagaraDataInterfaceStageArgs& Context) override
+	virtual void PostStage(const FNDIGpuComputePostStageContext& Context) override
 	{
-		FNDIExportInstanceData_RenderThread* InstanceData = SystemInstancesToProxyData_RT.Find(Context.SystemInstanceID);
-		check(InstanceData);
-
-		// Ensure our buffer if big enough to hold all the data
-		const int32 AllocationRounding = 64;
-		const int32 DataBufferNumInstances = Context.ComputeInstanceData->Context->CurrentMaxInstances_RT;
-		int32 NumInstances;
-		if (InstanceData->AllocationMode == ENDIExport_GPUAllocationMode::PerParticle)
-		{
-			NumInstances = int32(float(DataBufferNumInstances) * InstanceData->AllocationPerParticleSize);
-			NumInstances = FMath::Max(NumInstances, 0);
-			if (NDIExportLocal::GGPUMaxReadbackCount > 0)
-			{
-				NumInstances = FMath::Min(NumInstances, NDIExportLocal::GGPUMaxReadbackCount);
-			}
-		}
-		else
-		{
-			NumInstances = InstanceData->AllocationFixedSize;
-		}
-
-		const int32 AllocationNumInstances = FMath::DivideAndRoundUp(NumInstances, AllocationRounding) * AllocationRounding;
-		const int32 RequiredElements = 1 + (AllocationNumInstances * NDIExportLocal::NumFloatsPerInstance);
-		const int32 RequiredBytes = sizeof(float) * RequiredElements;
-
-		InstanceData->WriteBufferInstanceCount = NumInstances;
-		if (InstanceData->WriteBuffer.NumBytes != RequiredBytes)
-		{
-			InstanceData->WriteBuffer.Release();
-			InstanceData->WriteBuffer.Initialize(TEXT("FNDIExportProxy_WriteBuffer"), sizeof(float), RequiredElements, PF_R32_UINT, BUF_SourceCopy);
-
-			// Clear counter when we initialize the buffer
-			ClearCounter(RHICmdList, InstanceData);
-		}
-
-		RHICmdList.Transition(FRHITransitionInfo(InstanceData->WriteBuffer.UAV, ERHIAccess::Unknown, ERHIAccess::UAVCompute));
-	}
-
-	virtual void PostStage(FRHICommandList& RHICmdList, const FNiagaraDataInterfaceStageArgs& Context) override
-	{
-		FNDIExportInstanceData_RenderThread* InstanceData = SystemInstancesToProxyData_RT.Find(Context.SystemInstanceID);
-		check(InstanceData);
+		FNDIExportInstanceData_RenderThread& InstanceData = SystemInstancesToProxyData_RT.FindChecked(Context.GetSystemInstanceID());
 
 		// If we bound the buffer as a UAV the stage may have written to it
-		if (InstanceData->bWriteBufferUsed)
+		if (InstanceData.bWriteBufferUsed)
 		{
-			InstanceData->bWriteBufferUsed = false;
+			InstanceData.bWriteBufferUsed = false;
 
-			RHICmdList.Transition(FRHITransitionInfo(InstanceData->WriteBuffer.UAV, ERHIAccess::UAVCompute, ERHIAccess::CopySrc));
-
-			FNiagaraGpuReadbackManager* ReadbackManager = Context.ComputeDispatchInterface->GetGpuReadbackManager();
+			FNiagaraGpuReadbackManager* ReadbackManager = Context.GetComputeDispatchInterface().GetGpuReadbackManager();
 			ReadbackManager->EnqueueReadback(
-				RHICmdList,
-				InstanceData->WriteBuffer.Buffer,
-				[MaxInstances=InstanceData->WriteBufferInstanceCount, WeakCallbackHandler=InstanceData->WeakCallbackHandler, WeakSystem=InstanceData->WeakSystem, SystemLWCTile=Context.SystemLWCTile](TConstArrayView<TPair<void*, uint32>> Buffers)
+				Context.GetGraphBuilder(),
+				InstanceData.RDGWriteBuffer,
+				[MaxInstances= InstanceData.WriteBufferInstanceCount, WeakCallbackHandler=InstanceData.WeakCallbackHandler, WeakSystem=InstanceData.WeakSystem, SystemLWCTile= Context.GetSystemLWCTile()](TConstArrayView<TPair<void*, uint32>> Buffers)
 				{
 					const uint32 ReadbackInstanceCount = *reinterpret_cast<uint32*>(Buffers[0].Key);
 					check(ReadbackInstanceCount <= MaxInstances);
@@ -203,21 +162,9 @@ struct FNDIExportProxy : public FNiagaraDataInterfaceProxy
 					}
 				}
 			);
-
-			RHICmdList.Transition(FRHITransitionInfo(InstanceData->WriteBuffer.UAV, ERHIAccess::CopySrc, ERHIAccess::UAVCompute));
-
-			// Ensure counter is clear for the next user, they will perform the transition on the buffer
-			ClearCounter(RHICmdList, InstanceData);
 		}
-	}
-
-	void ClearCounter(FRHICommandList& RHICmdList, FNDIExportInstanceData_RenderThread* InstanceData)
-	{
-		check(InstanceData);
-		check(InstanceData->WriteBuffer.NumBytes > 0);
-
-		//-OPT: We only need to reset the first 4 bytes for the count
-		RHICmdList.ClearUAVUint(InstanceData->WriteBuffer.UAV, FUintVector4(0, 0, 0, 0));
+		InstanceData.RDGWriteBuffer = nullptr;
+		InstanceData.RDGWriteBufferUAV = nullptr;
 	}
 
 	TMap<FNiagaraSystemInstanceID, FNDIExportInstanceData_RenderThread> SystemInstancesToProxyData_RT;
@@ -445,21 +392,59 @@ void UNiagaraDataInterfaceExport::BuildShaderParameters(FNiagaraShaderParameters
 void UNiagaraDataInterfaceExport::SetShaderParameters(const FNiagaraDataInterfaceSetShaderParametersContext& Context) const
 {
 	FNDIExportProxy& DIProxy = Context.GetProxy<FNDIExportProxy>();
-	FNDIExportInstanceData_RenderThread* InstanceData_RT = DIProxy.SystemInstancesToProxyData_RT.Find(Context.GetSystemInstanceID());
+	FNDIExportInstanceData_RenderThread& InstanceData_RT = DIProxy.SystemInstancesToProxyData_RT.FindChecked(Context.GetSystemInstanceID());
 
 	NDIExportLocal::FShaderParameters* Parameters = Context.GetParameterNestedStruct<NDIExportLocal::FShaderParameters>();
 	if ( Context.IsResourceBound(&Parameters->RWWriteBuffer) )
 	{
-		if ( (InstanceData_RT != nullptr) && !InstanceData_RT->WeakCallbackHandler.IsExplicitlyNull() )
+		if ( !InstanceData_RT.WeakCallbackHandler.IsExplicitlyNull() )
 		{
-			InstanceData_RT->bWriteBufferUsed = true;
-			Parameters->WriteBufferSize = InstanceData_RT->WriteBufferInstanceCount;
-			Parameters->RWWriteBuffer = InstanceData_RT->WriteBuffer.UAV;
+			FRDGBuilder& GraphBuilder = Context.GetGraphBuilder();
+			if (InstanceData_RT.RDGWriteBuffer == nullptr)
+			{
+				// Ensure our buffer if big enough to hold all the data
+				const int32 AllocationRounding = 64;
+				const int32 DataBufferNumInstances = Context.GetComputeInstanceData().Context->CurrentMaxInstances_RT;
+				int32 NumInstances = 0;
+				if (InstanceData_RT.AllocationMode == ENDIExport_GPUAllocationMode::PerParticle)
+				{
+					NumInstances = int32(float(DataBufferNumInstances) * InstanceData_RT.AllocationPerParticleSize);
+					NumInstances = FMath::Max(NumInstances, 0);
+					if (NDIExportLocal::GGPUMaxReadbackCount > 0)
+					{
+						NumInstances = FMath::Min(NumInstances, NDIExportLocal::GGPUMaxReadbackCount);
+					}
+				}
+				else
+				{
+					NumInstances = InstanceData_RT.AllocationFixedSize;
+				}
+
+				InstanceData_RT.WriteBufferInstanceCount = NumInstances;
+
+				const int32 AllocationNumInstances = FMath::DivideAndRoundUp(NumInstances, AllocationRounding) * AllocationRounding;
+				const int32 RequiredElements = 1 + (AllocationNumInstances * NDIExportLocal::NumFloatsPerInstance);
+				const int32 RequiredBytes = sizeof(float) * RequiredElements;
+
+				FRDGBufferDesc BufferDesc = FRDGBufferDesc::CreateBufferDesc(sizeof(uint32), RequiredElements);
+				BufferDesc.Usage |= BUF_SourceCopy;
+				InstanceData_RT.RDGWriteBuffer = GraphBuilder.CreateBuffer(BufferDesc, TEXT("FNDIExportProxy_WriteBuffer"));
+				InstanceData_RT.RDGWriteBufferUAV = GraphBuilder.CreateUAV(InstanceData_RT.RDGWriteBuffer, PF_R32_UINT);
+
+			}
+
+			//-OPT: We only need to reset the first 4 bytes for the count
+			//-OPT: We only need to clear if we actual use the buffer
+			AddClearUAVPass(GraphBuilder, InstanceData_RT.RDGWriteBufferUAV, 0);
+
+			InstanceData_RT.bWriteBufferUsed = true;
+			Parameters->WriteBufferSize = InstanceData_RT.WriteBufferInstanceCount;
+			Parameters->RWWriteBuffer = InstanceData_RT.RDGWriteBufferUAV;
 		}
 		else
 		{
 			Parameters->WriteBufferSize = 0;
-			Parameters->RWWriteBuffer = Context.GetComputeDispatchInterface().GetEmptyUAVFromPool(Context.GetRHICmdList(), PF_R32_UINT, ENiagaraEmptyUAVType::Buffer);
+			Parameters->RWWriteBuffer = Context.GetComputeDispatchInterface().GetEmptyBufferUAV(Context.GetGraphBuilder(), PF_R32_UINT);
 		}
 	}
 	else
