- Niagara Compute Dipsatch modified to use FRDGBuilder through the pipeline

diff --git a/Engine/Plugins/FX/Niagara/Source/Niagara/Private/NiagaraDataInterfaceGrid2DCollection.cpp b/Engine/Plugins/FX/Niagara/Source/Niagara/Private/NiagaraDataInterfaceGrid2DCollection.cpp
--- a/Engine/Plugins/FX/Niagara/Source/Niagara/Private/NiagaraDataInterfaceGrid2DCollection.cpp
+++ b/Engine/Plugins/FX/Niagara/Source/Niagara/Private/NiagaraDataInterfaceGrid2DCollection.cpp
@@ -30,9 +30,9 @@ BEGIN_SHADER_PARAMETER_STRUCT(FNDIGrid2DShaderParameters, )
 	SHADER_PARAMETER(FVector2f,						CellSize)
 	SHADER_PARAMETER(FVector2f,						WorldBBoxSize)
 
-	SHADER_PARAMETER_SRV(Texture2DArray<float>,		Grid)
-	SHADER_PARAMETER_SAMPLER(SamplerState,			GridSampler)
-	SHADER_PARAMETER_UAV(RWTexture2DArray<float>,	OutputGrid)
+	SHADER_PARAMETER_RDG_TEXTURE_SRV(Texture2DArray<float>,		Grid)
+	SHADER_PARAMETER_SAMPLER(SamplerState,						GridSampler)
+	SHADER_PARAMETER_RDG_TEXTURE_UAV(RWTexture2DArray<float>,	OutputGrid)
 END_SHADER_PARAMETER_STRUCT()
 
 const FString UNiagaraDataInterfaceGrid2DCollection::GridName(TEXT("_Grid"));
@@ -182,6 +182,26 @@ struct FNiagaraGrid2DLegacyTiled2DInfo
 		}
 	}
 
+	void CopyTo2D(FRDGBuilder& GraphBuilder, FRDGTexture* SourceTexture, FRHITexture* DestinationTextureRHI, const TCHAR* NameIfNotRegistered) const
+	{
+		FRDGTexture* DestinationTexture = GraphBuilder.FindExternalTexture(DestinationTextureRHI);
+		if (DestinationTexture == nullptr)
+		{
+			DestinationTexture = GraphBuilder.RegisterExternalTexture(CreateRenderTarget(DestinationTextureRHI, NameIfNotRegistered));
+		}
+
+		for (int iAttribute = 0; iAttribute < NumAttributes; ++iAttribute)
+		{
+			FRHICopyTextureInfo CopyInfo;
+			CopyInfo.Size = FIntVector(NumCells.X, NumCells.Y, 1);
+			CopyInfo.SourceSliceIndex = iAttribute;
+			CopyInfo.DestPosition.X = (iAttribute % NumTiles.X) * NumCells.X;
+			CopyInfo.DestPosition.Y = (iAttribute / NumTiles.X) * NumCells.Y;
+			CopyInfo.DestPosition.Z = 0;
+			AddCopyTexturePass(GraphBuilder, SourceTexture, DestinationTexture, CopyInfo);
+		}
+	}
+
 	void CopyTo2D(FRHICommandList& RHICmdList, FRHITexture* Src, FRHITexture* Dst) const
 	{
 		check(Src != nullptr && Dst != nullptr);
@@ -1472,24 +1492,36 @@ void UNiagaraDataInterfaceGrid2DCollection::SetShaderParameters(const FNiagaraDa
 		}
 	}
 
+	FRDGBuilder& GraphBuilder = Context.GetGraphBuilder();
+
 	FNDIGrid2DShaderParameters* Parameters = Context.GetParameterNestedStruct<FNDIGrid2DShaderParameters>();
 	Parameters->NumAttributes = ProxyData->NumAttributes;
 	Parameters->UnitToUV = FVector2f(1.0f) / FVector2f(ProxyData->NumCells);
 	Parameters->NumCells = ProxyData->NumCells;
 	Parameters->CellSize = FVector2f(ProxyData->CellSize);
 	Parameters->WorldBBoxSize = FVector2f(ProxyData->WorldBBoxSize);	// LWC_TODO: Precision loss?
-	Parameters->Grid = ProxyData->CurrentData ? ProxyData->CurrentData->GridSRV.GetReference() : FNiagaraRenderer::GetDummyTextureReadBuffer2DArray();
+	if (Context.IsResourceBound(&Parameters->Grid))
+	{
+		if (ProxyData->CurrentData)
+		{
+			Parameters->Grid = ProxyData->CurrentData->GetOrCreateSRV(GraphBuilder);
+		}
+		else
+		{
+			Parameters->Grid = Context.GetComputeDispatchInterface().GetBlackTextureSRV(GraphBuilder, ETextureDimension::Texture2DArray);
+		}
+	}
 	Parameters->GridSampler = TStaticSamplerState<SF_Bilinear, AM_Clamp, AM_Clamp, AM_Clamp>::GetRHI();
 
 	if (Context.IsResourceBound(&Parameters->OutputGrid))
 	{
 		if (Context.IsOutputStage() && ProxyData->DestinationData)
 		{
-			Parameters->OutputGrid = ProxyData->DestinationData->GridUAV;
+			Parameters->OutputGrid = ProxyData->DestinationData->GetOrCreateUAV(GraphBuilder);
 		}
 		else
 		{
-			Parameters->OutputGrid = Context.GetComputeDispatchInterface().GetEmptyUAVFromPool(Context.GetRHICmdList(), PF_R32_FLOAT, ENiagaraEmptyUAVType::Texture2DArray);
+			Parameters->OutputGrid = Context.GetComputeDispatchInterface().GetEmptyTextureUAV(GraphBuilder, PF_R32_FLOAT, ETextureDimension::Texture2DArray);
 		}
 	}
 
@@ -2097,7 +2129,7 @@ bool UNiagaraDataInterfaceGrid2DCollection::RenderVariableToCanvas(FNiagaraSyste
 			{
 				if ( RT_InstanceData->CurrentData != nullptr )
 				{
-					OutTexture = RT_InstanceData->CurrentData->GridTexture;
+					OutTexture = RT_InstanceData->CurrentData->GetPooledTexture()->GetRHI();
 					OutSamplerState = TStaticSamplerState<SF_Bilinear>::GetRHI();
 				}
 			}
@@ -2146,7 +2178,7 @@ bool UNiagaraDataInterfaceGrid2DCollection::FillTexture2D(const UNiagaraComponen
 			FRHICopyTextureInfo CopyInfo;
 			CopyInfo.Size = FIntVector(Grid2DInstanceData->NumCells.X, Grid2DInstanceData->NumCells.Y, 1);
 			CopyInfo.SourcePosition = FIntVector(0, 0, AttributeIndex);
-			TransitionAndCopyTexture(RHICmdList, Grid2DInstanceData->CurrentData->GridTexture, RT_TextureResource->TextureRHI, CopyInfo);
+			TransitionAndCopyTexture(RHICmdList, Grid2DInstanceData->CurrentData->GetPooledTexture()->GetRHI(), RT_TextureResource->TextureRHI, CopyInfo);
 		}
 	});
 
@@ -2199,7 +2231,7 @@ bool UNiagaraDataInterfaceGrid2DCollection::FillRawTexture2D(const UNiagaraCompo
 		if (RT_TextureResource && RT_TextureResource->TextureRHI.IsValid() && RT_Grid2DInstanceData && RT_Grid2DInstanceData->CurrentData)
 		{
 			const FNiagaraGrid2DLegacyTiled2DInfo Tiled2DInfo(RT_Grid2DInstanceData->NumCells, RT_Grid2DInstanceData->NumAttributes);
-			Tiled2DInfo.CopyTo2D(RHICmdList, RT_Grid2DInstanceData->CurrentData->GridTexture, RT_TextureResource->TextureRHI);
+			Tiled2DInfo.CopyTo2D(RHICmdList, RT_Grid2DInstanceData->CurrentData->GetPooledTexture()->GetRHI(), RT_TextureResource->TextureRHI);
 		}
 	});
 
@@ -2497,134 +2529,132 @@ bool FGrid2DCollectionRWInstanceData_GameThread::UpdateTargetTexture(ENiagaraGpu
 	return false;
 }
 
-void FGrid2DCollectionRWInstanceData_RenderThread::BeginSimulate(FRHICommandList& RHICmdList)
+void FGrid2DCollectionRWInstanceData_RenderThread::BeginSimulate(FRDGBuilder& GraphBuilder)
 {
-	for (TUniquePtr<FGrid2DBuffer>& Buffer : Buffers)
+	for (FGrid2DBuffer& Buffer : Buffers)
 	{
 		check(Buffer.IsValid());
-		if (Buffer.Get() != CurrentData)
+		if (&Buffer != CurrentData)
 		{
-			DestinationData = Buffer.Get();
+			DestinationData = &Buffer;
 			break;
 		}
 	}
 
 	if (DestinationData == nullptr)
 	{
-		DestinationData = new FGrid2DBuffer(NumCells.X, NumCells.Y, NumAttributes, PixelFormat);
-		Buffers.Emplace(DestinationData);
+		DestinationData = &Buffers.AddDefaulted_GetRef();
 
-		// The rest of the code expects to find the buffers readable, and will transition from there to UAVCompute as necessary.
-		RHICmdList.Transition(FRHITransitionInfo(DestinationData->GridUAV, ERHIAccess::Unknown, ERHIAccess::SRVMask));
+		const FRDGTextureDesc TextureDesc = FRDGTextureDesc::Create2DArray(NumCells, PixelFormat, FClearValueBinding::Black, ETextureCreateFlags::ShaderResource | ETextureCreateFlags::UAV, NumAttributes);
+		DestinationData->Initialize(GraphBuilder, TEXT("FGrid2DBuffer"), TextureDesc);
 	}
 }
 
-void FGrid2DCollectionRWInstanceData_RenderThread::EndSimulate(FRHICommandList& RHICmdList)
+void FGrid2DCollectionRWInstanceData_RenderThread::EndSimulate()
 {
 	CurrentData = DestinationData;
 	DestinationData = nullptr;
 }
 
-void FNiagaraDataInterfaceProxyGrid2DCollectionProxy::PreStage(FRHICommandList& RHICmdList, const FNiagaraDataInterfaceStageArgs& Context)
+void FNiagaraDataInterfaceProxyGrid2DCollectionProxy::ResetData(const FNDIGpuComputeResetContext& Context)
+{
+	FGrid2DCollectionRWInstanceData_RenderThread* ProxyData = SystemInstancesToProxyData_RT.Find(Context.GetSystemInstanceID());
+	if (!ProxyData)
+	{
+		return;
+	}
+
+	FRDGBuilder& GraphBuilder = Context.GetGraphBuilder();
+	for (FGrid2DBuffer& Buffer : ProxyData->Buffers)
+	{
+		AddClearUAVPass(GraphBuilder, Buffer.GetOrCreateUAV(GraphBuilder), FVector4f(ForceInitToZero));
+	}
+}
+
+void FNiagaraDataInterfaceProxyGrid2DCollectionProxy::PreStage(const FNDIGpuComputePreStageContext& Context)
 {
 	// #todo(dmp): Context doesnt need to specify if a stage is output or not since we moved pre/post stage to the DI itself.  Not sure which design is better for the future
-	if (Context.IsOutputStage)
+	if (Context.IsOutputStage())
 	{
-		FGrid2DCollectionRWInstanceData_RenderThread* ProxyData = SystemInstancesToProxyData_RT.Find(Context.SystemInstanceID);
+		FGrid2DCollectionRWInstanceData_RenderThread* ProxyData = SystemInstancesToProxyData_RT.Find(Context.GetSystemInstanceID());
 
-		ProxyData->BeginSimulate(RHICmdList);
+		FRDGBuilder& GraphBuilder = Context.GetGraphBuilder();
+		ProxyData->BeginSimulate(GraphBuilder);
 
 		// If we don't have an iteration stage, then we should manually clear the buffer to make sure there is no residual data.  If we are doing something like rasterizing particles into a grid, we want it to be clear before
 		// we start.  If a user wants to access data from the previous stage, then they can read from the current data.
 
 		// #todo(dmp): we might want to expose an option where we have buffers that are write only and need a clear (ie: no buffering like the neighbor grid).  They would be considered transient perhaps?  It'd be more
 		// memory efficient since it would theoretically not require any double buffering.
-		RHICmdList.Transition(FRHITransitionInfo(ProxyData->DestinationData->GridUAV, ERHIAccess::SRVMask, ERHIAccess::UAVCompute));
-		if (!Context.IsIterationStage)
+		if (!Context.IsIterationStage())
 		{
-			SCOPED_DRAW_EVENT(RHICmdList, Grid2DCollection_PreStage);
-			RHICmdList.ClearUAVFloat(ProxyData->DestinationData->GridUAV, FVector4f(ForceInitToZero));
-			RHICmdList.Transition(FRHITransitionInfo(ProxyData->DestinationData->GridUAV, ERHIAccess::UAVCompute, ERHIAccess::UAVCompute));
+			AddClearUAVPass(GraphBuilder, ProxyData->DestinationData->GetOrCreateUAV(GraphBuilder), FVector4f(ForceInitToZero));
 		}
 	}
 }
 
-void FNiagaraDataInterfaceProxyGrid2DCollectionProxy::PostStage(FRHICommandList& RHICmdList, const FNiagaraDataInterfaceStageArgs& Context)
+void FNiagaraDataInterfaceProxyGrid2DCollectionProxy::PostStage(const FNDIGpuComputePostStageContext& Context)
 {
-	if (Context.IsOutputStage)
+	if (Context.IsOutputStage())
 	{
-		FGrid2DCollectionRWInstanceData_RenderThread* ProxyData = SystemInstancesToProxyData_RT.Find(Context.SystemInstanceID);
-		RHICmdList.Transition(FRHITransitionInfo(ProxyData->DestinationData->GridUAV, ERHIAccess::UAVCompute, ERHIAccess::SRVMask));
-		ProxyData->EndSimulate(RHICmdList);
+		FGrid2DCollectionRWInstanceData_RenderThread* ProxyData = SystemInstancesToProxyData_RT.Find(Context.GetSystemInstanceID());
+		ProxyData->EndSimulate();
 	}
 }
 
-void FNiagaraDataInterfaceProxyGrid2DCollectionProxy::PostSimulate(FRHICommandList& RHICmdList, const FNiagaraDataInterfaceArgs& Context)
+void FNiagaraDataInterfaceProxyGrid2DCollectionProxy::PostSimulate(const FNDIGpuComputePostSimulateContext& Context)
 {
-	FGrid2DCollectionRWInstanceData_RenderThread* ProxyData = SystemInstancesToProxyData_RT.Find(Context.SystemInstanceID);
+	FGrid2DCollectionRWInstanceData_RenderThread* ProxyData = SystemInstancesToProxyData_RT.Find(Context.GetSystemInstanceID());
 
-	if (ProxyData->RenderTargetToCopyTo != nullptr && ProxyData->CurrentData != nullptr && ProxyData->CurrentData->GridTexture != nullptr)
+	if (ProxyData->RenderTargetToCopyTo != nullptr && ProxyData->CurrentData != nullptr && ProxyData->CurrentData->IsValid())
 	{
-		SCOPED_DRAW_EVENT(RHICmdList, Grid2DCollection_PostSimulate);
+		FRDGBuilder& GraphBuilder = Context.GetGraphBuilder();
 		if (ProxyData->RenderTargetToCopyTo->GetTexture2DArray() != nullptr)
 		{
-			FRHICopyTextureInfo CopyInfo;
-			TransitionAndCopyTexture(RHICmdList, ProxyData->CurrentData->GridTexture, ProxyData->RenderTargetToCopyTo, CopyInfo);
+			ProxyData->CurrentData->CopyToTexture(GraphBuilder, ProxyData->RenderTargetToCopyTo, TEXT("NiagaraRenderTargetToCopyTo"));
 		}
 		else if (ensure(ProxyData->RenderTargetToCopyTo->GetTexture2D() != nullptr))
 		{
 			const FNiagaraGrid2DLegacyTiled2DInfo Tiled2DInfo(ProxyData->NumCells, ProxyData->NumAttributes);
-			Tiled2DInfo.CopyTo2D(RHICmdList, ProxyData->CurrentData->GridTexture, ProxyData->RenderTargetToCopyTo);
+			Tiled2DInfo.CopyTo2D(GraphBuilder, ProxyData->CurrentData->GetOrCreateTexture(GraphBuilder), ProxyData->RenderTargetToCopyTo, TEXT("NiagaraRenderTargetToCopyTo"));
 		}
 	}
 
 #if NIAGARA_COMPUTEDEBUG_ENABLED && WITH_EDITORONLY_DATA
-	if (ProxyData->bPreviewGrid && ProxyData->CurrentData)
+	if (ProxyData->bPreviewGrid && ProxyData->CurrentData && ProxyData->CurrentData->IsValid())
 	{
-		if (FNiagaraGpuComputeDebug* GpuComputeDebug = Context.ComputeDispatchInterface->GetGpuComputeDebug())
+		if (FNiagaraGpuComputeDebug* GpuComputeDebug = Context.GetComputeDispatchInterface().GetGpuComputeDebug())
 		{
+			FRDGBuilder& GraphBuilder = Context.GetGraphBuilder();
 			if (ProxyData->PreviewAttribute[0] != INDEX_NONE)
 			{
-				GpuComputeDebug->AddAttributeTexture(RHICmdList, Context.SystemInstanceID, SourceDIName, ProxyData->CurrentData->GridTexture, FIntPoint::ZeroValue, ProxyData->PreviewAttribute);
+				GpuComputeDebug->AddAttributeTexture(GraphBuilder, Context.GetSystemInstanceID(), SourceDIName, ProxyData->CurrentData->GetOrCreateTexture(GraphBuilder), FIntPoint::ZeroValue, ProxyData->PreviewAttribute);
 			}
 			else
 			{
-				GpuComputeDebug->AddTexture(RHICmdList, Context.SystemInstanceID, SourceDIName, ProxyData->CurrentData->GridTexture);
+				GpuComputeDebug->AddTexture(GraphBuilder, Context.GetSystemInstanceID(), SourceDIName, ProxyData->CurrentData->GetOrCreateTexture(GraphBuilder));
 			}
 		}
 	}
 #endif
-}
 
-void FNiagaraDataInterfaceProxyGrid2DCollectionProxy::ResetData(FRHICommandList& RHICmdList, const FNiagaraDataInterfaceArgs& Context)
-{
-	FGrid2DCollectionRWInstanceData_RenderThread* ProxyData = SystemInstancesToProxyData_RT.Find(Context.SystemInstanceID);
-	if (!ProxyData)
+	// Clear out the transient resource we cached
+	if (Context.IsFinalPostSimulate())
 	{
-		return;
-	}
+		for (FGrid2DBuffer& Buffer : ProxyData->Buffers)
+		{
+			Buffer.EndGraphUsage();
+		}
 
-	for (TUniquePtr<FGrid2DBuffer>& Buffer : ProxyData->Buffers)
-	{
-		if (Buffer.IsValid())
+		// Readers point to data not owned by themselves so can be caching resources on the 'other' proxy
+		// Therefore we need to ensure the transient buffers are correctly cleared
+		if ( FNiagaraDataInterfaceProxyGrid2DCollectionProxy* OtherGrid3DProxy = static_cast<FNiagaraDataInterfaceProxyGrid2DCollectionProxy*>(ProxyData->OtherProxy) )
 		{
-			ERHIAccess AccessAfter;
-			const bool bIsDestination = (ProxyData->DestinationData == Buffer.Get());
-			if (bIsDestination)
+			FGrid2DCollectionRWInstanceData_RenderThread& OtherProxyData = OtherGrid3DProxy->SystemInstancesToProxyData_RT.FindChecked(Context.GetSystemInstanceID());
+			for (FGrid2DBuffer& Buffer : OtherProxyData.Buffers)
 			{
-				// The destination buffer is already in UAVCompute because PreStage() runs first. It must stay in UAVCompute after the clear
-				// because the shader is going to use it.
-				AccessAfter = ERHIAccess::UAVCompute;
+				Buffer.EndGraphUsage();
 			}
-			else
-			{
-				// The other buffers are in SRVMask and must be returned to that state after the clear.
-				RHICmdList.Transition(FRHITransitionInfo(Buffer->GridUAV, ERHIAccess::SRVMask, ERHIAccess::UAVCompute));
-				AccessAfter = ERHIAccess::SRVMask;
-			}
-
-			RHICmdList.ClearUAVFloat(Buffer->GridUAV, FVector4f(ForceInitToZero));
-			RHICmdList.Transition(FRHITransitionInfo(Buffer->GridUAV, ERHIAccess::UAVCompute, AccessAfter));
 		}
 	}
 }
diff --git a/Engine/Plugins/FX/Niagara/Source/Niagara/Private/NiagaraDataInterfaceGrid3DCollection.cpp b/Engine/Plugins/FX/Niagara/Source/Niagara/Private/NiagaraDataInterfaceGrid3DCollection.cpp
--- a/Engine/Plugins/FX/Niagara/Source/Niagara/Private/NiagaraDataInterfaceGrid3DCollection.cpp
+++ b/Engine/Plugins/FX/Niagara/Source/Niagara/Private/NiagaraDataInterfaceGrid3DCollection.cpp
@@ -34,10 +34,10 @@ BEGIN_SHADER_PARAMETER_STRUCT(FNDIGrid3DShaderParameters, )
 	SHADER_PARAMETER(FVector3f,						UnitClampMax)
 	SHADER_PARAMETER(FVector3f,						WorldBBoxSize)
 
-	SHADER_PARAMETER_SRV(Texture3D<float>,			Grid)
-	SHADER_PARAMETER_SAMPLER(SamplerState,			GridSampler)
-	SHADER_PARAMETER_UAV(RWTexture3D<float>,		OutputGrid)
-	SHADER_PARAMETER_SRV(Buffer<float4>,			PerAttributeData)
+	SHADER_PARAMETER_RDG_TEXTURE_SRV(Texture3D<float>,		Grid)
+	SHADER_PARAMETER_SAMPLER(SamplerState,					GridSampler)
+	SHADER_PARAMETER_RDG_TEXTURE_UAV(RWTexture3D<float>,	OutputGrid)
+	SHADER_PARAMETER_SRV(Buffer<float4>,					PerAttributeData)
 END_SHADER_PARAMETER_STRUCT()
 
 const FString UNiagaraDataInterfaceGrid3DCollection::NumTilesName(TEXT("_NumTiles"));
@@ -1798,6 +1798,8 @@ void UNiagaraDataInterfaceGrid3DCollection::SetShaderParameters(const FNiagaraDa
 	// Set parameters
 	const FVector3f HalfPixelOffset = FVector3f(0.5f / ProxyData->NumCells.X, 0.5f / ProxyData->NumCells.Y, 0.5f / ProxyData->NumCells.Z);
 
+	FRDGBuilder& GraphBuilder = Context.GetGraphBuilder();
+
 	FNDIGrid3DShaderParameters* Parameters = Context.GetParameterNestedStruct<FNDIGrid3DShaderParameters>();
 	Parameters->NumAttributes		= ProxyData->TotalNumAttributes;
 	Parameters->NumNamedAttributes	= ProxyData->TotalNumNamedAttributes;
@@ -1810,18 +1812,28 @@ void UNiagaraDataInterfaceGrid3DCollection::SetShaderParameters(const FNiagaraDa
 	Parameters->UnitClampMax		= FVector3f::OneVector - HalfPixelOffset;
 	Parameters->WorldBBoxSize		= FVector3f(ProxyData->WorldBBoxSize);
 
-	Parameters->Grid = ProxyData->CurrentData ? ProxyData->CurrentData->GridBuffer.SRV.GetReference() : FNiagaraRenderer::GetDummyTextureReadBuffer3D();
+	if (Context.IsResourceBound(&Parameters->Grid))
+	{
+		if (ProxyData->CurrentData)
+		{
+			Parameters->Grid = ProxyData->CurrentData->GetOrCreateSRV(GraphBuilder);
+		}
+		else
+		{
+			Parameters->Grid = Context.GetComputeDispatchInterface().GetBlackTextureSRV(GraphBuilder, ETextureDimension::Texture3D);
+		}
+	}
 	Parameters->GridSampler = TStaticSamplerState<SF_Bilinear, AM_Clamp, AM_Clamp, AM_Clamp>::GetRHI();
 
 	if (Context.IsResourceBound(&Parameters->OutputGrid))
 	{
 		if (Context.IsOutputStage() && ProxyData->DestinationData)
 		{
-			Parameters->OutputGrid = ProxyData->DestinationData->GridBuffer.UAV;
+			Parameters->OutputGrid = ProxyData->DestinationData->GetOrCreateUAV(GraphBuilder);
 		}
 		else
 		{
-			Parameters->OutputGrid = Context.GetComputeDispatchInterface().GetEmptyUAVFromPool(Context.GetRHICmdList(), PF_R32_FLOAT, ENiagaraEmptyUAVType::Texture3D);
+			Parameters->OutputGrid = Context.GetComputeDispatchInterface().GetEmptyTextureUAV(GraphBuilder, PF_R32_FLOAT, ETextureDimension::Texture3D);
 		}
 	}
 	Parameters->PerAttributeData = ProxyData->PerAttributeData.SRV;
@@ -2907,40 +2919,57 @@ bool FGrid3DCollectionRWInstanceData_GameThread::UpdateTargetTexture(ENiagaraGpu
 	return false;
 }
 
-void FGrid3DCollectionRWInstanceData_RenderThread::BeginSimulate(FRHICommandList& RHICmdList)
+void FGrid3DCollectionRWInstanceData_RenderThread::BeginSimulate(FRDGBuilder& GraphBuilder)
 {
-	for (TUniquePtr<FGrid3DBuffer>& Buffer : Buffers)
+	for (FGrid3DBuffer& Buffer : Buffers)
 	{
 		check(Buffer.IsValid());
-		if (Buffer.Get() != CurrentData)
+		if (&Buffer != CurrentData)
 		{
-			DestinationData = Buffer.Get();
+			DestinationData = &Buffer;
 			break;
 		}
 	}
 
 	if (DestinationData == nullptr)
 	{
-		DestinationData = new FGrid3DBuffer(NumCells.X * NumTiles.X, NumCells.Y * NumTiles.Y, NumCells.Z * NumTiles.Z, PixelFormat);
-		RHICmdList.Transition(FRHITransitionInfo(DestinationData->GridBuffer.UAV, ERHIAccess::Unknown, ERHIAccess::SRVMask));
-		Buffers.Emplace(DestinationData);
+		DestinationData = &Buffers.AddDefaulted_GetRef();
+
+		const FIntVector TextureSize(NumCells.X * NumTiles.X, NumCells.Y * NumTiles.Y, NumCells.Z * NumTiles.Z);
+		const FRDGTextureDesc TextureDesc = FRDGTextureDesc::Create3D(TextureSize, PixelFormat, FClearValueBinding::Black, ETextureCreateFlags::ShaderResource | ETextureCreateFlags::UAV);
+		DestinationData->Initialize(GraphBuilder, TEXT("FGrid3DBuffer"), TextureDesc);
 	}
 }
 
-void FGrid3DCollectionRWInstanceData_RenderThread::EndSimulate(FRHICommandList& RHICmdList)
+void FGrid3DCollectionRWInstanceData_RenderThread::EndSimulate()
 {
 	CurrentData = DestinationData;
 	DestinationData = nullptr;
 }
 
-void FNiagaraDataInterfaceProxyGrid3DCollectionProxy::PreStage(FRHICommandList& RHICmdList, const FNiagaraDataInterfaceStageArgs& Context)
+void FNiagaraDataInterfaceProxyGrid3DCollectionProxy::ResetData(const FNDIGpuComputeResetContext& Context)
 {
-	FGrid3DCollectionRWInstanceData_RenderThread* ProxyData = SystemInstancesToProxyData_RT.Find(Context.SystemInstanceID);
+	FGrid3DCollectionRWInstanceData_RenderThread* ProxyData = SystemInstancesToProxyData_RT.Find(Context.GetSystemInstanceID());
+	if (!ProxyData)
+	{
+		return;
+	}
 
+	FRDGBuilder& GraphBuilder = Context.GetGraphBuilder();
+	for (FGrid3DBuffer& Buffer : ProxyData->Buffers)
+	{
+		AddClearUAVPass(GraphBuilder, Buffer.GetOrCreateUAV(GraphBuilder), FVector4f(ForceInitToZero));
+	}
+}
+
+void FNiagaraDataInterfaceProxyGrid3DCollectionProxy::PreStage(const FNDIGpuComputePreStageContext& Context)
+{
 	// #todo(dmp): Context doesnt need to specify if a stage is output or not since we moved pre/post stage to the DI itself.  Not sure which design is better for the future
-	if (Context.IsOutputStage)
+	if (Context.IsOutputStage())
 	{
-		ProxyData->BeginSimulate(RHICmdList);
+		FGrid3DCollectionRWInstanceData_RenderThread* ProxyData = SystemInstancesToProxyData_RT.Find(Context.GetSystemInstanceID());
+		FRDGBuilder& GraphBuilder = Context.GetGraphBuilder();
+		ProxyData->BeginSimulate(GraphBuilder);
 
 		// If we don't have an iteration stage, then we should manually clear the buffer to make sure there is no residual data.  If we are doing something like rasterizing particles into a grid, we want it to be clear before
 		// we start.  If a user wants to access data from the previous stage, then they can read from the current data.
@@ -2950,13 +2979,12 @@ void FNiagaraDataInterfaceProxyGrid3DCollectionProxy::PreStage(FRHICommandList&
 
 		// #todo(dmp): for now, if there is an output DI that is NOT an iteration DI AND if both the iteration DI and this DI have the same total number of instances, do not do the UAV clear prior to the stage.
 		// this isn't optimal, but should work to some degree to reduce overhead in cases where we have multiple grids of the same resolution being processed/written to in 1 stage
-		FNiagaraDataInterfaceProxyRW* IterationInterface = Context.SimStageData->AlternateIterationSource;
-		RHICmdList.Transition(FRHITransitionInfo(ProxyData->DestinationData->GridBuffer.UAV, ERHIAccess::SRVMask, ERHIAccess::UAVCompute));
-		if (!Context.IsIterationStage)
+		if (!Context.IsIterationStage())
 		{
+			FNiagaraDataInterfaceProxyRW* IterationInterface = Context.GetSimStageData().AlternateIterationSource;
 			if (IterationInterface != nullptr)
 			{
-				const FIntVector ElementCount = IterationInterface->GetElementCount(Context.SystemInstanceID);
+				const FIntVector ElementCount = IterationInterface->GetElementCount(Context.GetSystemInstanceID());
 				const uint64 TotalNumInstances = ElementCount.X * ElementCount.Y * ElementCount.Z;
 
 				if (ElementCount.X == ProxyData->NumCells.X && ElementCount.Y == ProxyData->NumCells.Y && ElementCount.Z == ProxyData->NumCells.Z)
@@ -2965,99 +2993,65 @@ void FNiagaraDataInterfaceProxyGrid3DCollectionProxy::PreStage(FRHICommandList&
 				}
 			}
 
-			check(ProxyData->DestinationData);
-			RHICmdList.ClearUAVFloat(ProxyData->DestinationData->GridBuffer.UAV, FVector4f(ForceInitToZero));
-			RHICmdList.Transition(FRHITransitionInfo(ProxyData->DestinationData->GridBuffer.UAV, ERHIAccess::UAVCompute, ERHIAccess::UAVCompute));
+			AddClearUAVPass(GraphBuilder, ProxyData->DestinationData->GetOrCreateUAV(GraphBuilder), FVector4f(ForceInitToZero));
 		}
 	}
 }
 
-void FNiagaraDataInterfaceProxyGrid3DCollectionProxy::PostStage(FRHICommandList& RHICmdList, const FNiagaraDataInterfaceStageArgs& Context)
+void FNiagaraDataInterfaceProxyGrid3DCollectionProxy::PostStage(const FNDIGpuComputePostStageContext& Context)
 {
-	if (Context.IsOutputStage)
+	if (Context.IsOutputStage())
 	{
-		FGrid3DCollectionRWInstanceData_RenderThread* ProxyData = SystemInstancesToProxyData_RT.Find(Context.SystemInstanceID);
-		RHICmdList.Transition(FRHITransitionInfo(ProxyData->DestinationData->GridBuffer.UAV, ERHIAccess::UAVCompute, ERHIAccess::SRVMask));
-		ProxyData->EndSimulate(RHICmdList);
+		FGrid3DCollectionRWInstanceData_RenderThread* ProxyData = SystemInstancesToProxyData_RT.Find(Context.GetSystemInstanceID());
+		ProxyData->EndSimulate();
 	}
 }
 
-void FNiagaraDataInterfaceProxyGrid3DCollectionProxy::PostSimulate(FRHICommandList& RHICmdList, const FNiagaraDataInterfaceArgs& Context)
+void FNiagaraDataInterfaceProxyGrid3DCollectionProxy::PostSimulate(const FNDIGpuComputePostSimulateContext& Context)
 {
-	FGrid3DCollectionRWInstanceData_RenderThread* ProxyData = SystemInstancesToProxyData_RT.Find(Context.SystemInstanceID);
+	FGrid3DCollectionRWInstanceData_RenderThread* ProxyData = SystemInstancesToProxyData_RT.Find(Context.GetSystemInstanceID());
 
-	if (ProxyData->RenderTargetToCopyTo != nullptr && ProxyData->CurrentData != nullptr && ProxyData->CurrentData->GridBuffer.Buffer != nullptr)
+	if (ProxyData->RenderTargetToCopyTo != nullptr && ProxyData->CurrentData != nullptr && ProxyData->CurrentData->IsValid())
 	{
-		FRHITexture* Source = ProxyData->CurrentData->GridBuffer.Buffer;
-		FRHITexture* Destination = ProxyData->RenderTargetToCopyTo;
-
-		RHICmdList.Transition(
-			{
-				FRHITransitionInfo(Source, ERHIAccess::SRVMask, ERHIAccess::CopySrc),
-				FRHITransitionInfo(Destination, ERHIAccess::SRVMask, ERHIAccess::CopyDest)
-			}
-		);
-
-		FRHICopyTextureInfo CopyInfo;
-		RHICmdList.CopyTexture(ProxyData->CurrentData->GridBuffer.Buffer, ProxyData->RenderTargetToCopyTo, CopyInfo);
-
-		RHICmdList.Transition(
-			{
-				FRHITransitionInfo(Source, ERHIAccess::CopySrc, ERHIAccess::SRVMask),
-				FRHITransitionInfo(Destination, ERHIAccess::CopyDest, ERHIAccess::SRVMask)
-			}
-		);
+		ProxyData->CurrentData->CopyToTexture(Context.GetGraphBuilder(), ProxyData->RenderTargetToCopyTo, TEXT("NiagaraRenderTargetToCopyTo"));
 	}
 
 #if WITH_EDITOR
-	if (ProxyData->bPreviewGrid && ProxyData->CurrentData)
+	if (ProxyData->bPreviewGrid && ProxyData->CurrentData && ProxyData->CurrentData->IsValid())
 	{
-		if (FNiagaraGpuComputeDebug* GpuComputeDebug = Context.ComputeDispatchInterface->GetGpuComputeDebug())
+		if (FNiagaraGpuComputeDebug* GpuComputeDebug = Context.GetComputeDispatchInterface().GetGpuComputeDebug())
 		{
-				if (ProxyData->PreviewAttribute[0] != INDEX_NONE)
+			FRDGBuilder& GraphBuilder = Context.GetGraphBuilder();
+			if (ProxyData->PreviewAttribute[0] != INDEX_NONE)
 			{
-
-				FIntVector4 TotalNumAttributeVector = FIntVector4(ProxyData->NumTiles.X, ProxyData->NumTiles.Y, ProxyData->NumTiles.Z, 0);
-				GpuComputeDebug->AddAttributeTexture(RHICmdList, Context.SystemInstanceID, SourceDIName, ProxyData->CurrentData->GridBuffer.Buffer, TotalNumAttributeVector, ProxyData->PreviewAttribute);
+				const FIntVector4 TotalNumAttributeVector = FIntVector4(ProxyData->NumTiles.X, ProxyData->NumTiles.Y, ProxyData->NumTiles.Z, 0);
+				GpuComputeDebug->AddAttributeTexture(GraphBuilder, Context.GetSystemInstanceID(), SourceDIName, ProxyData->CurrentData->GetOrCreateTexture(GraphBuilder), TotalNumAttributeVector, ProxyData->PreviewAttribute);
 			}
 			else
 			{
-				GpuComputeDebug->AddTexture(RHICmdList, Context.SystemInstanceID, SourceDIName, ProxyData->CurrentData->GridBuffer.Buffer);
+				GpuComputeDebug->AddTexture(GraphBuilder, Context.GetSystemInstanceID(), SourceDIName, ProxyData->CurrentData->GetOrCreateTexture(GraphBuilder));
 			}
 		}
 	}
 #endif
-}
 
-void FNiagaraDataInterfaceProxyGrid3DCollectionProxy::ResetData(FRHICommandList& RHICmdList, const FNiagaraDataInterfaceArgs& Context)
-{
-	FGrid3DCollectionRWInstanceData_RenderThread* ProxyData = SystemInstancesToProxyData_RT.Find(Context.SystemInstanceID);
-	if (!ProxyData)
+	// Clear out the transient resource we cached
+	if (Context.IsFinalPostSimulate())
 	{
-		return;
-	}
+		for (FGrid3DBuffer& Buffer : ProxyData->Buffers)
+		{
+			Buffer.EndGraphUsage();
+		}
 
-	for (TUniquePtr<FGrid3DBuffer>& Buffer : ProxyData->Buffers)
-	{
-		if (Buffer.IsValid())
+		// Readers point to data not owned by themselves so can be caching resources on the 'other' proxy
+		// Therefore we need to ensure the transient buffers are correctly cleared
+		if (FNiagaraDataInterfaceProxyGrid3DCollectionProxy* OtherGrid3DProxy = static_cast<FNiagaraDataInterfaceProxyGrid3DCollectionProxy*>(ProxyData->OtherProxy))
 		{
-			ERHIAccess AccessAfter;
-			const bool bIsDestination = (ProxyData->DestinationData == Buffer.Get());
-			if (bIsDestination)
+			FGrid3DCollectionRWInstanceData_RenderThread& OtherProxyData = OtherGrid3DProxy->SystemInstancesToProxyData_RT.FindChecked(Context.GetSystemInstanceID());
+			for (FGrid3DBuffer& Buffer : OtherProxyData.Buffers)
 			{
-				// The destination buffer is already in UAVCompute because PreStage() runs first. It must stay in UAVCompute after the clear
-				// because the shader is going to use it.
-				AccessAfter = ERHIAccess::UAVCompute;
+				Buffer.EndGraphUsage();
 			}
-			else
-			{
-				// The other buffers are in SRVMask and must be returned to that state after the clear.
-				RHICmdList.Transition(FRHITransitionInfo(Buffer->GridBuffer.UAV, ERHIAccess::SRVMask, ERHIAccess::UAVCompute));
-				AccessAfter = ERHIAccess::SRVMask;
-			}
-
-			RHICmdList.ClearUAVFloat(Buffer->GridBuffer.UAV, FVector4f(ForceInitToZero));
-			RHICmdList.Transition(FRHITransitionInfo(Buffer->GridBuffer.UAV, ERHIAccess::UAVCompute, AccessAfter));
 		}
 	}
 }
diff --git a/Engine/Plugins/FX/Niagara/Source/Niagara/Private/NiagaraDataInterfaceIntRenderTarget2D.cpp b/Engine/Plugins/FX/Niagara/Source/Niagara/Private/NiagaraDataInterfaceIntRenderTarget2D.cpp
--- a/Engine/Plugins/FX/Niagara/Source/Niagara/Private/NiagaraDataInterfaceIntRenderTarget2D.cpp
+++ b/Engine/Plugins/FX/Niagara/Source/Niagara/Private/NiagaraDataInterfaceIntRenderTarget2D.cpp
@@ -8,6 +8,8 @@
 #include "NiagaraGpuComputeDebug.h"
 #endif
 #include "NiagaraGpuComputeDispatchInterface.h"
+#include "NiagaraShader.h"
+#include "NiagaraShaderParametersBuilder.h"
 #include "NiagaraStats.h"
 
 #include "ShaderParameterUtils.h"
@@ -24,9 +26,21 @@
 
 namespace NDIIntRenderTarget2DLocal
 {
+	struct FFunctionVersion
+	{
+		enum Type
+		{
+			InitialVersion = 0,
+			AddedOptionalExecute = 1,
+
+			VersionPlusOne,
+			LatestVersion = VersionPlusOne - 1
+		};
+	};
+
 	static const TCHAR* TemplateShaderFile = TEXT("/Plugin/FX/Niagara/Private/NiagaraDataInterfaceIntRenderTarget2D.ush");
 
-	static const int32 NumFunctions = 9;
+	static const int32 NumFunctions = 12;
 	static const FName GetValueFunctionName("GetValue");
 	static const FName SetValueFunctionName("SetValue");
 	static const FName AtomicAddFunctionName("AtomicAdd");
@@ -45,6 +59,9 @@ namespace NDIIntRenderTarget2DLocal
 	static const FName LinearToIndexFunctionName("LinearToIndex");
 	static const FName LinearToUVFunctionName("LinearToUV");
 
+	static const FName ExecToIndexFunctionName("ExecToIndex");
+	static const FName ExecToUVFunctionName("ExecToUV");
+
 	// Shader parameters
 	static const FString TextureSizeAndInvSizeName(TEXT("TextureSizeAndInvSize_"));
 	static const FString TextureUAVName(TEXT("TextureUAV_"));
@@ -77,9 +94,11 @@ struct FNDIIntRenderTarget2DInstanceData_RenderThread
 	bool bWroteThisFrame = false;
 	bool bReadThisFrame = false;
 
-	FSamplerStateRHIRef SamplerStateRHI;
-	FTextureRHIRef TextureRHI;
-	FUnorderedAccessViewRHIRef UnorderedAccessViewRHI;
+	FTextureRHIRef		TextureRHI;
+
+	FRDGTextureRef		TransientRDGTexture = nullptr;
+	FRDGTextureUAVRef	TransientRDGUAV = nullptr;
+
 #if STATS
 	void UpdateMemoryStats()
 	{
@@ -105,37 +124,36 @@ struct FNDIIntRenderTarget2DProxy : public FNiagaraDataInterfaceProxyRW
 	virtual void ConsumePerInstanceDataFromGameThread(void* PerInstanceData, const FNiagaraSystemInstanceID& Instance) override {}
 	virtual int32 PerInstanceDataPassedToRenderThreadSize() const override { return 0; }
 
-	virtual void PostSimulate(FRHICommandList& RHICmdList, const FNiagaraDataInterfaceArgs& Context) override
+	virtual void PostSimulate(const FNDIGpuComputePostSimulateContext& Context) override
 	{
-		FNDIIntRenderTarget2DInstanceData_RenderThread* InstanceData = SystemInstancesToProxyData_RT.Find(Context.SystemInstanceID);
-		if ( InstanceData == nullptr )
-		{
-			return;
-		}
+		FNDIIntRenderTarget2DInstanceData_RenderThread& InstanceData = SystemInstancesToProxyData_RT.FindChecked(Context.GetSystemInstanceID());
 
-		// We only need to transfer this frame if it was written to.
-		// If also read then we need to notify that the texture is important for the simulation
-		// We also assume the texture is important for rendering, without discovering renderer bindings we don't really know
-		if (InstanceData->bWroteThisFrame)
-		{
-			Context.ComputeDispatchInterface->MultiGPUResourceModified(RHICmdList, InstanceData->TextureRHI, InstanceData->bReadThisFrame, true);
-		}
+//-TODO:RDG:mGPU
+//		// We only need to transfer this frame if it was written to.
+//		// If also read then we need to notify that the texture is important for the simulation
+//		// We also assume the texture is important for rendering, without discovering renderer bindings we don't really know
+//		if (InstanceData->bWroteThisFrame)
+//		{
+//			Context.ComputeDispatchInterface->MultiGPUResourceModified(RHICmdList, InstanceData->TextureRHI, InstanceData->bReadThisFrame, true);
+//		}
 
-		InstanceData->bWroteThisFrame = true;
-		InstanceData->bReadThisFrame = true;
+		InstanceData.bWroteThisFrame = true;
+		InstanceData.bReadThisFrame = true;
 
 #if NIAGARA_COMPUTEDEBUG_ENABLED && WITH_EDITORONLY_DATA
-		if (InstanceData->bPreviewRenderTarget)
+		if (InstanceData.bPreviewRenderTarget && InstanceData.TransientRDGTexture)
 		{
-			if (FNiagaraGpuComputeDebug* GpuComputeDebug = Context.ComputeDispatchInterface->GetGpuComputeDebug())
+			if (FNiagaraGpuComputeDebug* GpuComputeDebug = Context.GetComputeDispatchInterface().GetGpuComputeDebug())
 			{
-				if (FRHITexture* RHITexture = InstanceData->TextureRHI)
-				{
-					GpuComputeDebug->AddTexture(RHICmdList, Context.SystemInstanceID, SourceDIName, RHITexture, InstanceData->PreviewDisplayRange);
-				}
+				GpuComputeDebug->AddTexture(Context.GetGraphBuilder(), Context.GetSystemInstanceID(), SourceDIName, InstanceData.TransientRDGTexture);
 			}
 		}
 #endif
+		if (Context.IsFinalPostSimulate())
+		{
+			InstanceData.TransientRDGTexture = nullptr;
+			InstanceData.TransientRDGUAV = nullptr;
+		}
 	}
 
 	virtual FIntVector GetElementCount(FNiagaraSystemInstanceID SystemInstanceID) const override
@@ -152,75 +170,6 @@ struct FNDIIntRenderTarget2DProxy : public FNiagaraDataInterfaceProxyRW
 
 //////////////////////////////////////////////////////////////////////////
 
-struct FNDIIntRenderTarget2DParametersCS : public FNiagaraDataInterfaceParametersCS
-{
-	DECLARE_TYPE_LAYOUT(FNDIIntRenderTarget2DParametersCS, NonVirtual);
-public:
-	void Bind(const FNiagaraDataInterfaceGPUParamInfo& ParameterInfo, const class FShaderParameterMap& ParameterMap)
-	{			
-		TextureSizeAndInvSizeParam.Bind(ParameterMap, *(NDIIntRenderTarget2DLocal::TextureSizeAndInvSizeName + ParameterInfo.DataInterfaceHLSLSymbol));
-		TextureUAVParam.Bind(ParameterMap, *(NDIIntRenderTarget2DLocal::TextureUAVName + ParameterInfo.DataInterfaceHLSLSymbol));
-	}
-
-	void Set(FRHICommandList& RHICmdList, const FNiagaraDataInterfaceSetArgs& Context) const
-	{
-		check(IsInRenderingThread());
-
-		// Get shader and DI
-		FRHIComputeShader* ComputeShaderRHI = Context.Shader.GetComputeShader();
-		FNDIIntRenderTarget2DProxy* DataInterfaceProxy = static_cast<FNDIIntRenderTarget2DProxy*>(Context.DataInterface);
-		FNDIIntRenderTarget2DInstanceData_RenderThread* InstanceData = DataInterfaceProxy->SystemInstancesToProxyData_RT.Find(Context.SystemInstanceID);
-		check(InstanceData);
-
-		const FVector4f TextureSizeAndInvSize(InstanceData->Size.X, InstanceData->Size.Y, 1.0f / float(InstanceData->Size.X), 1.0f / float(InstanceData->Size.Y));
-		SetShaderValue(RHICmdList, ComputeShaderRHI, TextureSizeAndInvSizeParam, TextureSizeAndInvSize);
-	
-		if (TextureUAVParam.IsUAVBound())
-		{
-			InstanceData->bWroteThisFrame = true;
-			InstanceData->bReadThisFrame = true;
-
-			FRHIUnorderedAccessView* OutputUAV = InstanceData->UnorderedAccessViewRHI;
-			if (OutputUAV != nullptr)
-			{
-				RHICmdList.Transition(FRHITransitionInfo(OutputUAV, ERHIAccess::Unknown, ERHIAccess::UAVCompute));
-			}
-			else
-			{
-				OutputUAV = Context.ComputeDispatchInterface->GetEmptyUAVFromPool(RHICmdList, EPixelFormat::PF_A16B16G16R16, ENiagaraEmptyUAVType::Texture2D);
-			}
-
-			RHICmdList.SetUAVParameter(ComputeShaderRHI, TextureUAVParam.GetUAVIndex(), OutputUAV);
-		}
-	}
-
-	void Unset(FRHICommandList& RHICmdList, const FNiagaraDataInterfaceSetArgs& Context) const 
-	{
-		if (TextureUAVParam.IsBound())
-		{
-			TextureUAVParam.UnsetUAV(RHICmdList, Context.Shader.GetComputeShader());
-
-			FNDIIntRenderTarget2DProxy* DIProxy = static_cast<FNDIIntRenderTarget2DProxy*>(Context.DataInterface);
-			if ( FNDIIntRenderTarget2DInstanceData_RenderThread* InstanceData = DIProxy->SystemInstancesToProxyData_RT.Find(Context.SystemInstanceID) )
-			{
-				if (FRHIUnorderedAccessView* OutputUAV = InstanceData->UnorderedAccessViewRHI)
-				{
-					RHICmdList.Transition(FRHITransitionInfo(OutputUAV, ERHIAccess::UAVCompute, ERHIAccess::SRVMask));
-				}
-			}
-		}
-	}
-
-private:
-	LAYOUT_FIELD(FShaderParameter,		TextureSizeAndInvSizeParam);
-	LAYOUT_FIELD(FRWShaderParameter,	TextureUAVParam);
-};
-
-IMPLEMENT_TYPE_LAYOUT(FNDIIntRenderTarget2DParametersCS);
-IMPLEMENT_NIAGARA_DI_PARAMETER(UNiagaraDataInterfaceIntRenderTarget2D, FNDIIntRenderTarget2DParametersCS);
-
-//////////////////////////////////////////////////////////////////////////
-
 FNiagaraVariableBase UNiagaraDataInterfaceIntRenderTarget2D::ExposedRTVar;
 
 UNiagaraDataInterfaceIntRenderTarget2D::UNiagaraDataInterfaceIntRenderTarget2D(FObjectInitializer const& ObjectInitializer)
@@ -269,12 +218,14 @@ void UNiagaraDataInterfaceIntRenderTarget2D::GetFunctions(TArray<FNiagaraFunctio
 		Sig.bSupportsGPU = true;
 #if WITH_EDITORONLY_DATA
 		Sig.Description = LOCTEXT("GetValueDesc", "Gets the value from the render target at the pixel offset");
+		Sig.FunctionVersion = NDIIntRenderTarget2DLocal::FFunctionVersion::LatestVersion;
 #endif
 	}
 	{
 		FNiagaraFunctionSignature& Sig = OutFunctions.AddDefaulted_GetRef();
 		Sig.Name = NDIIntRenderTarget2DLocal::SetValueFunctionName;
 		Sig.Inputs.Add(FNiagaraVariable(FNiagaraTypeDefinition(GetClass()), TEXT("RenderTarget")));
+		Sig.Inputs.Add_GetRef(FNiagaraVariable(FNiagaraTypeDefinition::GetBoolDef(), TEXT("Execute"))).SetValue(true);
 		Sig.Inputs.Add(FNiagaraVariable(FNiagaraTypeDefinition::GetIntDef(), TEXT("PixelX")));
 		Sig.Inputs.Add(FNiagaraVariable(FNiagaraTypeDefinition::GetIntDef(), TEXT("PixelY")));
 		Sig.Inputs.Add(FNiagaraVariable(FNiagaraTypeDefinition::GetIntDef(), TEXT("Value")));
@@ -287,6 +238,7 @@ void UNiagaraDataInterfaceIntRenderTarget2D::GetFunctions(TArray<FNiagaraFunctio
 		Sig.bSupportsGPU = true;
 #if WITH_EDITORONLY_DATA
 		Sig.Description = LOCTEXT("SetValueDesc", "Sets the value on the render target at the pixel offset.");
+		Sig.FunctionVersion = NDIIntRenderTarget2DLocal::FFunctionVersion::LatestVersion;
 #endif
 	}
 	{
@@ -308,6 +260,7 @@ void UNiagaraDataInterfaceIntRenderTarget2D::GetFunctions(TArray<FNiagaraFunctio
 		Sig.bSupportsGPU = true;
 #if WITH_EDITORONLY_DATA
 		Sig.Description = LOCTEXT("AtomicAddDesc", "Atomic min the value to the pixel at the offset, returns the current & previous values.  This opertion is thread safe.");
+		Sig.FunctionVersion = NDIIntRenderTarget2DLocal::FFunctionVersion::LatestVersion;
 #endif
 	}
 	////static const FName AtomicAndFunctionName("AtomicAnd");
@@ -331,6 +284,7 @@ void UNiagaraDataInterfaceIntRenderTarget2D::GetFunctions(TArray<FNiagaraFunctio
 		Sig.bSupportsGPU = true;
 #if WITH_EDITORONLY_DATA
 		Sig.Description = LOCTEXT("AtomicCASFunctionDesc", "Compares the pixel value against the comparison value, if they are equal the value is replaced.  Original Value is the pixel value before the operation completes.  This opertion is thread safe.");
+		Sig.FunctionVersion = NDIIntRenderTarget2DLocal::FFunctionVersion::LatestVersion;
 #endif
 	}
 	////static const FName AtomicCSFunctionName("AtomicCompareStore");
@@ -354,6 +308,7 @@ void UNiagaraDataInterfaceIntRenderTarget2D::GetFunctions(TArray<FNiagaraFunctio
 		Sig.bSupportsGPU = true;
 #if WITH_EDITORONLY_DATA
 		Sig.Description = LOCTEXT("AtomicMaxDesc", "Atomic max the value to the pixel at the offset, returns the current & previous values.  This opertion is thread safe.");
+		Sig.FunctionVersion = NDIIntRenderTarget2DLocal::FFunctionVersion::LatestVersion;
 #endif
 	}
 	{
@@ -375,6 +330,7 @@ void UNiagaraDataInterfaceIntRenderTarget2D::GetFunctions(TArray<FNiagaraFunctio
 		Sig.bSupportsGPU = true;
 #if WITH_EDITORONLY_DATA
 		Sig.Description = LOCTEXT("AtomicMinDesc", "Atomic min the value to the pixel at the offset, returns the current & previous values.  This opertion is thread safe.");
+		Sig.FunctionVersion = NDIIntRenderTarget2DLocal::FFunctionVersion::LatestVersion;
 #endif
 	}
 	////static const FName AtomicOrFunctionName("AtomicOr");
@@ -393,6 +349,7 @@ void UNiagaraDataInterfaceIntRenderTarget2D::GetFunctions(TArray<FNiagaraFunctio
 		Sig.bSupportsGPU = true;
 #if WITH_EDITORONLY_DATA
 		Sig.Description = LOCTEXT("GetSizeDesc", "Gets the size of the rendertarget");
+		Sig.FunctionVersion = NDIIntRenderTarget2DLocal::FFunctionVersion::LatestVersion;
 #endif
 	}
 	{
@@ -411,6 +368,7 @@ void UNiagaraDataInterfaceIntRenderTarget2D::GetFunctions(TArray<FNiagaraFunctio
 		Sig.bSupportsGPU = false;
 #if WITH_EDITORONLY_DATA
 		Sig.Description = LOCTEXT("SetSizeDesc", "Sets the size of the rendertarget");
+		Sig.FunctionVersion = NDIIntRenderTarget2DLocal::FFunctionVersion::LatestVersion;
 #endif
 	}
 
@@ -428,6 +386,7 @@ void UNiagaraDataInterfaceIntRenderTarget2D::GetFunctions(TArray<FNiagaraFunctio
 		Sig.bSupportsGPU = true;
 #if WITH_EDITORONLY_DATA
 		Sig.Description = LOCTEXT("LinearToIndexDesc", "Converts a linear index into a pixel coordinate");
+		Sig.FunctionVersion = NDIIntRenderTarget2DLocal::FFunctionVersion::LatestVersion;
 #endif
 	}
 	{
@@ -443,10 +402,65 @@ void UNiagaraDataInterfaceIntRenderTarget2D::GetFunctions(TArray<FNiagaraFunctio
 		Sig.bSupportsGPU = true;
 #if WITH_EDITORONLY_DATA
 		Sig.Description = LOCTEXT("LinearToUVDesc", "Converts a linear index into a UV coordinate");
+		Sig.FunctionVersion = NDIIntRenderTarget2DLocal::FFunctionVersion::LatestVersion;
+#endif
+	}
+
+	{
+		FNiagaraFunctionSignature& Sig = OutFunctions.AddDefaulted_GetRef();
+		Sig.Name = NDIIntRenderTarget2DLocal::ExecToIndexFunctionName;
+		Sig.Inputs.Add(FNiagaraVariable(FNiagaraTypeDefinition(GetClass()), TEXT("RenderTarget")));
+		Sig.Outputs.Add(FNiagaraVariable(FNiagaraTypeDefinition::GetIntDef(), TEXT("PixelX")));
+		Sig.Outputs.Add(FNiagaraVariable(FNiagaraTypeDefinition::GetIntDef(), TEXT("PixelY")));
+		Sig.bExperimental = true;
+		Sig.bMemberFunction = true;
+		Sig.bRequiresContext = false;
+		Sig.bSupportsCPU = false;
+		Sig.bSupportsGPU = true;
+#if WITH_EDITORONLY_DATA
+		Sig.Description = LOCTEXT("ExecToIndexDesc", "Returns the execution index as a pixel coordinate");
+		Sig.FunctionVersion = NDIIntRenderTarget2DLocal::FFunctionVersion::LatestVersion;
+#endif
+	}
+	{
+		FNiagaraFunctionSignature& Sig = OutFunctions.AddDefaulted_GetRef();
+		Sig.Name = NDIIntRenderTarget2DLocal::ExecToUVFunctionName;
+		Sig.Inputs.Add(FNiagaraVariable(FNiagaraTypeDefinition(GetClass()), TEXT("RenderTarget")));
+		Sig.Outputs.Add(FNiagaraVariable(FNiagaraTypeDefinition::GetVec2Def(), TEXT("UV")));
+		Sig.bExperimental = true;
+		Sig.bMemberFunction = true;
+		Sig.bRequiresContext = false;
+		Sig.bSupportsCPU = false;
+		Sig.bSupportsGPU = true;
+#if WITH_EDITORONLY_DATA
+		Sig.Description = LOCTEXT("ExecToUVDesc", "Returns the execution index as a UV coordinate");
+		Sig.FunctionVersion = NDIIntRenderTarget2DLocal::FFunctionVersion::LatestVersion;
 #endif
 	}
 }
 
+#if WITH_EDITORONLY_DATA
+bool UNiagaraDataInterfaceIntRenderTarget2D::UpgradeFunctionCall(FNiagaraFunctionSignature& FunctionSignature)
+{
+	bool bWasChanged = false;
+
+	if (FunctionSignature.FunctionVersion < NDIIntRenderTarget2DLocal::FFunctionVersion::AddedOptionalExecute)
+	{
+		if (FunctionSignature.Name == NDIIntRenderTarget2DLocal::SetValueFunctionName)
+		{
+			check(FunctionSignature.Inputs.Num() == 4);
+			FunctionSignature.Inputs.Insert_GetRef(FNiagaraVariable(FNiagaraTypeDefinition::GetBoolDef(), TEXT("Execute")), 1).SetValue(true);
+			bWasChanged = true;
+		}
+	}
+
+	// Set latest version
+	FunctionSignature.FunctionVersion = NDIIntRenderTarget2DLocal::FFunctionVersion::LatestVersion;
+
+	return bWasChanged;
+}
+#endif
+
 void UNiagaraDataInterfaceIntRenderTarget2D::GetVMExternalFunction(const FVMExternalFunctionBindingInfo& BindingInfo, void* InstanceData, FVMExternalFunction &OutFunc)
 {
 	Super::GetVMExternalFunction(BindingInfo, InstanceData, OutFunc);
@@ -506,6 +520,7 @@ bool UNiagaraDataInterfaceIntRenderTarget2D::AppendCompileHash(FNiagaraCompileHa
 	bool bSuccess = Super::AppendCompileHash(InVisitor);
 	FSHAHash Hash = GetShaderFileHash(NDIIntRenderTarget2DLocal::TemplateShaderFile, EShaderPlatform::SP_PCD3D_SM5);
 	InVisitor->UpdateString(TEXT("NiagaraDataInterfaceExportTemplateHLSLSource"), Hash.ToString());
+	InVisitor->UpdateShaderParameters<FShaderParameters>();
 	return bSuccess;
 }
 
@@ -538,7 +553,9 @@ bool UNiagaraDataInterfaceIntRenderTarget2D::GetFunctionHLSL(const FNiagaraDataI
 		(FunctionInfo.DefinitionName == NDIIntRenderTarget2DLocal::GetSizeFunctionName) ||
 		//(FunctionInfo.DefinitionName == NDIIntRenderTarget2DLocal::SetSizeFunctionName) ||
 		(FunctionInfo.DefinitionName == NDIIntRenderTarget2DLocal::LinearToIndexFunctionName) ||
-		(FunctionInfo.DefinitionName == NDIIntRenderTarget2DLocal::LinearToUVFunctionName) )
+		(FunctionInfo.DefinitionName == NDIIntRenderTarget2DLocal::LinearToUVFunctionName) ||
+		(FunctionInfo.DefinitionName == NDIIntRenderTarget2DLocal::ExecToIndexFunctionName) ||
+		(FunctionInfo.DefinitionName == NDIIntRenderTarget2DLocal::ExecToUVFunctionName) )
 	{
 		return true;
 	}
@@ -546,6 +563,55 @@ bool UNiagaraDataInterfaceIntRenderTarget2D::GetFunctionHLSL(const FNiagaraDataI
 }
 #endif
 
+void UNiagaraDataInterfaceIntRenderTarget2D::BuildShaderParameters(FNiagaraShaderParametersBuilder& ShaderParametersBuilder) const
+{
+	ShaderParametersBuilder.AddNestedStruct<FShaderParameters>();
+}
+
+void UNiagaraDataInterfaceIntRenderTarget2D::SetShaderParameters(const FNiagaraDataInterfaceSetShaderParametersContext& Context) const
+{
+	check(IsInRenderingThread());
+
+	FNDIIntRenderTarget2DProxy& DataInterfaceProxy = Context.GetProxy<FNDIIntRenderTarget2DProxy>();
+	FNDIIntRenderTarget2DInstanceData_RenderThread& InstanceData = DataInterfaceProxy.SystemInstancesToProxyData_RT.FindChecked(Context.GetSystemInstanceID());
+
+	// Set Parameters
+	FShaderParameters* Parameters = Context.GetParameterNestedStruct<FShaderParameters>();
+	Parameters->TextureSizeAndInvSize.X = InstanceData.Size.X;
+	Parameters->TextureSizeAndInvSize.Y = InstanceData.Size.Y;
+	Parameters->TextureSizeAndInvSize.Z = 1.0f / float(InstanceData.Size.X);
+	Parameters->TextureSizeAndInvSize.W = 1.0f / float(InstanceData.Size.Y);
+
+	if (Context.IsResourceBound(&Parameters->RWTextureUAV))
+	{
+		InstanceData.bWroteThisFrame = true;
+		InstanceData.bReadThisFrame = true;
+
+		FRDGBuilder& GraphBuilder = Context.GetGraphBuilder();
+
+		// Ensure RDG resources are ready to use
+		if (InstanceData.TransientRDGTexture == nullptr && InstanceData.TextureRHI.IsValid())
+		{
+			InstanceData.TransientRDGTexture = GraphBuilder.FindExternalTexture(InstanceData.TextureRHI);
+			if (InstanceData.TransientRDGTexture == nullptr)
+			{
+				InstanceData.TransientRDGTexture = GraphBuilder.RegisterExternalTexture(CreateRenderTarget(InstanceData.TextureRHI, TEXT("NiagaraIntRenderTarget2D")));
+			}
+			InstanceData.TransientRDGUAV = GraphBuilder.CreateUAV(InstanceData.TransientRDGTexture);
+			Context.GetRDGExternalAccessQueue().Add(InstanceData.TransientRDGTexture);
+		}
+
+		if (InstanceData.TransientRDGUAV != nullptr)
+		{
+			Parameters->RWTextureUAV = InstanceData.TransientRDGUAV;
+		}
+		else
+		{
+			Parameters->RWTextureUAV = Context.GetComputeDispatchInterface().GetEmptyTextureUAV(GraphBuilder, EPixelFormat::PF_R32_SINT, ETextureDimension::Texture2D);
+		}
+	}
+}
+
 bool UNiagaraDataInterfaceIntRenderTarget2D::InitPerInstanceData(void* PerInstanceData, FNiagaraSystemInstance* SystemInstance)
 {
 	check(Proxy);
@@ -579,9 +645,7 @@ bool UNiagaraDataInterfaceIntRenderTarget2D::InitPerInstanceData(void* PerInstan
 			{
 				if (FTextureRenderTarget2DResource* Resource2D = RT_TargetTexture->GetTextureRenderTarget2DResource())
 				{
-					InstanceData->SamplerStateRHI = Resource2D->SamplerStateRHI;
 					InstanceData->TextureRHI = Resource2D->GetTextureRHI();
-					InstanceData->UnorderedAccessViewRHI = Resource2D->GetUnorderedAccessViewRHI();
 				}
 			}
 		}
@@ -656,16 +720,12 @@ bool UNiagaraDataInterfaceIntRenderTarget2D::PerInstanceTickPostSimulate(void* P
 						InstanceData->bPreviewRenderTarget = RT_InstanceData.bPreviewRenderTarget;
 						InstanceData->PreviewDisplayRange = RT_InstanceData.PreviewDisplayRange;
 					#endif
-						InstanceData->SamplerStateRHI.SafeRelease();
 						InstanceData->TextureRHI.SafeRelease();
-						InstanceData->UnorderedAccessViewRHI.SafeRelease();
 						if (RT_TargetTexture)
 						{
 							if (FTextureRenderTarget2DResource* Resource2D = RT_TargetTexture->GetTextureRenderTarget2DResource())
 							{
-								InstanceData->SamplerStateRHI = Resource2D->SamplerStateRHI;
 								InstanceData->TextureRHI = Resource2D->GetTextureRHI();
-								InstanceData->UnorderedAccessViewRHI = Resource2D->GetUnorderedAccessViewRHI();
 							}
 						}
 					}
diff --git a/Engine/Plugins/FX/Niagara/Source/Niagara/Private/NiagaraDataInterfaceNeighborGrid3D.cpp b/Engine/Plugins/FX/Niagara/Source/Niagara/Private/NiagaraDataInterfaceNeighborGrid3D.cpp
--- a/Engine/Plugins/FX/Niagara/Source/Niagara/Private/NiagaraDataInterfaceNeighborGrid3D.cpp
+++ b/Engine/Plugins/FX/Niagara/Source/Niagara/Private/NiagaraDataInterfaceNeighborGrid3D.cpp
@@ -35,23 +35,24 @@ static FAutoConsoleVariableRef CVarMaxNiagaraNeighborGridCells(
 	ECVF_Default
 );
 
-void NeighborGrid3DRWInstanceData::ResizeBuffers()
-{
-	uint32 NumTotalCells = NumCells.X * NumCells.Y * NumCells.Z;
+//////////////////////////////////////////////////////////////////////////
 
-	uint32 NumIntsInGridBuffer = NumTotalCells * MaxNeighborsPerCell;
+void NeighborGrid3DRWInstanceData::ResizeBuffers(FRDGBuilder& GraphBuilder)
+{
+	const uint32 NumTotalCells = NumCells.X * NumCells.Y * NumCells.Z;
+	const uint32 NumIntsInGridBuffer = NumTotalCells * MaxNeighborsPerCell;
 
+	NeedsRealloc_RT = false;
 	if (NumTotalCells > (uint32)GMaxNiagaraNeighborGridCells)
+	{
 		return;
+	}
 
-	NeighborhoodCountBuffer.Initialize(TEXT("NiagaraNeighborGrid3D::NeighborCount"), sizeof(int32), NumTotalCells, EPixelFormat::PF_R32_SINT, BUF_Static);
-	NeighborhoodBuffer.Initialize(TEXT("NiagaraNeighborGrid3D::NeighborsGrid"), sizeof(int32), NumIntsInGridBuffer, EPixelFormat::PF_R32_SINT, BUF_Static);
+	NeighborhoodCountBuffer.Release();
+	NeighborhoodBuffer.Release();
 
-	#if STATS
-		DEC_MEMORY_STAT_BY(STAT_NiagaraGPUDataInterfaceMemory, GPUMemory);
-		GPUMemory = NumTotalCells * sizeof(int32) + NumIntsInGridBuffer * sizeof(int32);
-		INC_MEMORY_STAT_BY(STAT_NiagaraGPUDataInterfaceMemory, GPUMemory);
-	#endif
+	NeighborhoodCountBuffer.Initialize(GraphBuilder, TEXT("NiagaraNeighborGrid3D::NeighborCount"), EPixelFormat::PF_R32_SINT, sizeof(int32), NumTotalCells, BUF_Static);
+	NeighborhoodBuffer.Initialize(GraphBuilder, TEXT("NiagaraNeighborGrid3D::NeighborsGrid"), EPixelFormat::PF_R32_SINT, sizeof(int32), NumIntsInGridBuffer, BUF_Static);
 }
 
 //////////////////////////////////////////////////////////////////////////
@@ -65,7 +66,6 @@ UNiagaraDataInterfaceNeighborGrid3D::UNiagaraDataInterfaceNeighborGrid3D(FObject
 	Proxy.Reset(new FNiagaraDataInterfaceProxyNeighborGrid3D());	
 }
 
-
 void UNiagaraDataInterfaceNeighborGrid3D::GetFunctions(TArray<FNiagaraFunctionSignature>& OutFunctions)
 {
 	Super::GetFunctions(OutFunctions);
@@ -407,9 +407,10 @@ void UNiagaraDataInterfaceNeighborGrid3D::SetShaderParameters(const FNiagaraData
 {
 	FNiagaraDataInterfaceProxyNeighborGrid3D& DIProxy = Context.GetProxy<FNiagaraDataInterfaceProxyNeighborGrid3D>();
 	NeighborGrid3DRWInstanceData* ProxyData = DIProxy.SystemInstancesToProxyData.Find(Context.GetSystemInstanceID());
+	FRDGBuilder& GraphBuilder = Context.GetGraphBuilder();
 
 	FShaderParameters* ShaderParameters = Context.GetParameterNestedStruct<FShaderParameters>();
-	if (ProxyData && ProxyData->NeighborhoodBuffer.Buffer.IsValid())
+	if (ProxyData && ProxyData->NeighborhoodBuffer.IsValid())
 	{
 		ShaderParameters->NumCells		= ProxyData->NumCells;
 		ShaderParameters->UnitToUV		= FVector3f(1.0f) / FVector3f(ProxyData->NumCells);
@@ -419,50 +420,19 @@ void UNiagaraDataInterfaceNeighborGrid3D::SetShaderParameters(const FNiagaraData
 		ShaderParameters->WorldBBoxSize	= FVector3f(ProxyData->WorldBBoxSize);
 		ShaderParameters->MaxNeighborsPerCellValue = ProxyData->MaxNeighborsPerCell;
 
-		TArray<FRHITransitionInfo, TInlineAllocator<4>> PendingTransitions;
 		if (Context.IsOutputStage())
 		{
-			ShaderParameters->ParticleNeighbors = FNiagaraRenderer::GetDummyIntBuffer();
-			ShaderParameters->ParticleNeighborCount = FNiagaraRenderer::GetDummyIntBuffer();
-
-			if (Context.IsResourceBound(&ShaderParameters->OutputParticleNeighbors))
-			{
-				PendingTransitions.Emplace(ProxyData->NeighborhoodBuffer.UAV, ERHIAccess::Unknown, ERHIAccess::UAVCompute);
-				ShaderParameters->OutputParticleNeighbors = ProxyData->NeighborhoodBuffer.UAV;
-			}
-			if (Context.IsResourceBound(&ShaderParameters->OutputParticleNeighborCount))
-			{
-				PendingTransitions.Emplace(ProxyData->NeighborhoodCountBuffer.UAV, ERHIAccess::Unknown, ERHIAccess::UAVCompute);
-				ShaderParameters->OutputParticleNeighborCount = ProxyData->NeighborhoodCountBuffer.UAV;
-			}
+			ShaderParameters->ParticleNeighbors = Context.GetComputeDispatchInterface().GetEmptyBufferSRV(GraphBuilder, PF_R32_SINT);
+			ShaderParameters->ParticleNeighborCount = Context.GetComputeDispatchInterface().GetEmptyBufferSRV(GraphBuilder, PF_R32_SINT);
+			ShaderParameters->OutputParticleNeighbors = ProxyData->NeighborhoodBuffer.GetOrCreateUAV(GraphBuilder);
+			ShaderParameters->OutputParticleNeighborCount = ProxyData->NeighborhoodCountBuffer.GetOrCreateUAV(GraphBuilder);
 		}
 		else
 		{
-			if (Context.IsResourceBound(&ShaderParameters->ParticleNeighbors))
-			{
-				PendingTransitions.Emplace(ProxyData->NeighborhoodBuffer.UAV, ERHIAccess::Unknown, ERHIAccess::SRVCompute);
-				ShaderParameters->ParticleNeighbors = ProxyData->NeighborhoodBuffer.SRV;
-			}
-
-			if (Context.IsResourceBound(&ShaderParameters->ParticleNeighborCount))
-			{
-				PendingTransitions.Emplace(ProxyData->NeighborhoodCountBuffer.UAV, ERHIAccess::Unknown, ERHIAccess::SRVCompute);
-				ShaderParameters->ParticleNeighborCount = ProxyData->NeighborhoodCountBuffer.SRV;
-			}
-
-			if (Context.IsResourceBound(&ShaderParameters->OutputParticleNeighbors))
-			{
-				ShaderParameters->OutputParticleNeighbors = Context.GetComputeDispatchInterface().GetEmptyUAVFromPool(Context.GetRHICmdList(), PF_R32_SINT, ENiagaraEmptyUAVType::Buffer);
-			}
-			if (Context.IsResourceBound(&ShaderParameters->OutputParticleNeighborCount))
-			{
-				ShaderParameters->OutputParticleNeighborCount = Context.GetComputeDispatchInterface().GetEmptyUAVFromPool(Context.GetRHICmdList(), PF_R32_SINT, ENiagaraEmptyUAVType::Buffer);
-			}
-		}
-
-		if (PendingTransitions.Num() > 0)
-		{
-			Context.GetRHICmdList().Transition(PendingTransitions);
+			ShaderParameters->ParticleNeighbors = ProxyData->NeighborhoodBuffer.GetOrCreateSRV(GraphBuilder);
+			ShaderParameters->ParticleNeighborCount = ProxyData->NeighborhoodCountBuffer.GetOrCreateSRV(GraphBuilder);
+			ShaderParameters->OutputParticleNeighbors = Context.GetComputeDispatchInterface().GetEmptyBufferUAV(GraphBuilder, PF_R32_SINT);
+			ShaderParameters->OutputParticleNeighborCount = Context.GetComputeDispatchInterface().GetEmptyBufferUAV(GraphBuilder, PF_R32_SINT);
 		}
 	}
 	else
@@ -473,16 +443,10 @@ void UNiagaraDataInterfaceNeighborGrid3D::SetShaderParameters(const FNiagaraData
 		ShaderParameters->WorldBBoxSize	= FVector3f::ZeroVector;
 
 		ShaderParameters->MaxNeighborsPerCellValue	= 0;
-		ShaderParameters->ParticleNeighbors			= FNiagaraRenderer::GetDummyIntBuffer();
-		ShaderParameters->ParticleNeighborCount		= FNiagaraRenderer::GetDummyIntBuffer();
-		if (Context.IsResourceBound(&ShaderParameters->OutputParticleNeighbors))
-		{
-			ShaderParameters->OutputParticleNeighbors = Context.GetComputeDispatchInterface().GetEmptyUAVFromPool(Context.GetRHICmdList(), PF_R32_SINT, ENiagaraEmptyUAVType::Buffer);
-		}
-		if (Context.IsResourceBound(&ShaderParameters->OutputParticleNeighborCount))
-		{
-			ShaderParameters->OutputParticleNeighborCount = Context.GetComputeDispatchInterface().GetEmptyUAVFromPool(Context.GetRHICmdList(), PF_R32_SINT, ENiagaraEmptyUAVType::Buffer);
-		}
+		ShaderParameters->ParticleNeighbors = Context.GetComputeDispatchInterface().GetEmptyBufferSRV(GraphBuilder, PF_R32_SINT);
+		ShaderParameters->ParticleNeighborCount = Context.GetComputeDispatchInterface().GetEmptyBufferSRV(GraphBuilder, PF_R32_SINT);
+		ShaderParameters->OutputParticleNeighbors = Context.GetComputeDispatchInterface().GetEmptyBufferUAV(GraphBuilder, PF_R32_SINT);
+		ShaderParameters->OutputParticleNeighborCount = Context.GetComputeDispatchInterface().GetEmptyBufferUAV(GraphBuilder, PF_R32_SINT);
 	}
 }
 
@@ -581,8 +545,7 @@ bool UNiagaraDataInterfaceNeighborGrid3D::InitPerInstanceData(void* PerInstanceD
 		TargetData->NumCells = RT_NumCells;
 		TargetData->MaxNeighborsPerCell = RT_MaxNeighborsPerCell;		
 		TargetData->WorldBBoxSize = RT_WorldBBoxSize;
-
-		TargetData->ResizeBuffers();
+		TargetData->NeedsRealloc_RT = true;
 	});
 
 	return true;
@@ -617,7 +580,7 @@ void UNiagaraDataInterfaceNeighborGrid3D::SetNumCells(FVectorVMExternalFunctionC
 
 			InstData->MaxNeighborsPerCell = NewMaxNeighborsPerCell;
 		
-			InstData->NeedsRealloc = OldNumCells != InstData->NumCells || OldMaxNeighborsPerCell != InstData->MaxNeighborsPerCell;
+			InstData->NeedsRealloc_GT = OldNumCells != InstData->NumCells || OldMaxNeighborsPerCell != InstData->MaxNeighborsPerCell;
 		}
 	}
 }
@@ -627,9 +590,9 @@ bool UNiagaraDataInterfaceNeighborGrid3D::PerInstanceTickPostSimulate(void* PerI
 	NeighborGrid3DRWInstanceData* InstanceData = static_cast<NeighborGrid3DRWInstanceData*>(PerInstanceData);
 	bool bNeedsReset = false;
 
-	if (InstanceData->NeedsRealloc && InstanceData->NumCells.X > 0 && InstanceData->NumCells.Y > 0 && InstanceData->NumCells.Z > 0 && InstanceData->MaxNeighborsPerCell > 0)
+	if (InstanceData->NeedsRealloc_GT && InstanceData->NumCells.X > 0 && InstanceData->NumCells.Y > 0 && InstanceData->NumCells.Z > 0 && InstanceData->MaxNeighborsPerCell > 0)
 	{
-		InstanceData->NeedsRealloc = false;
+		InstanceData->NeedsRealloc_GT = false;
 
 		InstanceData->CellSize = (InstanceData->WorldBBoxSize / FVector(InstanceData->NumCells.X, InstanceData->NumCells.Y, InstanceData->NumCells.Z))[0];
 
@@ -643,7 +606,7 @@ bool UNiagaraDataInterfaceNeighborGrid3D::PerInstanceTickPostSimulate(void* PerI
 			TargetData->NumCells = RT_NumCells;
 			TargetData->MaxNeighborsPerCell = RT_MaxNeighborsPerCell;			
 			TargetData->CellSize = RT_CellSize;
-			TargetData->ResizeBuffers();
+			TargetData->NeedsRealloc_RT = true;
 		});
 	}
 
@@ -669,23 +632,30 @@ void UNiagaraDataInterfaceNeighborGrid3D::DestroyPerInstanceData(void* PerInstan
 	);
 }
 
-void FNiagaraDataInterfaceProxyNeighborGrid3D::PreStage(FRHICommandList& RHICmdList, const FNiagaraDataInterfaceStageArgs& Context)
+void FNiagaraDataInterfaceProxyNeighborGrid3D::PreStage(const FNDIGpuComputePreStageContext& Context)
 {
-	if (Context.IsOutputStage)
+	FRDGBuilder& GraphBuilder = Context.GetGraphBuilder();
+	NeighborGrid3DRWInstanceData& ProxyData = SystemInstancesToProxyData.FindChecked(Context.GetSystemInstanceID());
+	if (ProxyData.NeedsRealloc_RT)
 	{
-		NeighborGrid3DRWInstanceData* ProxyData = SystemInstancesToProxyData.Find(Context.SystemInstanceID);
+		ProxyData.ResizeBuffers(GraphBuilder);
+	}
 
-		SCOPED_DRAW_EVENT(RHICmdList, NiagaraNeighborGrid3DClearNeighborInfo);
-		ERHIFeatureLevel::Type FeatureLevel = Context.ComputeDispatchInterface->GetFeatureLevel();
+	if (Context.IsOutputStage() && ProxyData.NeighborhoodBuffer.IsValid())
+	{
+		RDG_RHI_EVENT_SCOPE(GraphBuilder, NiagaraNeighborGrid3DClearNeighborInfo);
+		AddClearUAVPass(GraphBuilder, ProxyData.NeighborhoodBuffer.GetOrCreateUAV(GraphBuilder), -1);
+		AddClearUAVPass(GraphBuilder, ProxyData.NeighborhoodCountBuffer.GetOrCreateUAV(GraphBuilder), 0);
+	}
+}
 
-		FRHITransitionInfo TransitionInfos[] =
-		{
-			FRHITransitionInfo(ProxyData->NeighborhoodBuffer.UAV, ERHIAccess::Unknown, ERHIAccess::UAVCompute),
-			FRHITransitionInfo(ProxyData->NeighborhoodCountBuffer.UAV, ERHIAccess::Unknown, ERHIAccess::UAVCompute),
-		};
-		RHICmdList.Transition(MakeArrayView(TransitionInfos, UE_ARRAY_COUNT(TransitionInfos)));
-		NiagaraFillGPUIntBuffer(RHICmdList, FeatureLevel, ProxyData->NeighborhoodBuffer, -1);
-		NiagaraFillGPUIntBuffer(RHICmdList, FeatureLevel, ProxyData->NeighborhoodCountBuffer, 0);
+void FNiagaraDataInterfaceProxyNeighborGrid3D::PostSimulate(const FNDIGpuComputePostSimulateContext& Context)
+{
+	if (Context.IsFinalPostSimulate())
+	{
+		NeighborGrid3DRWInstanceData& ProxyData = SystemInstancesToProxyData.FindChecked(Context.GetSystemInstanceID());
+		ProxyData.NeighborhoodBuffer.EndGraphUsage();
+		ProxyData.NeighborhoodCountBuffer.EndGraphUsage();
 	}
 }
 
