- Niagara Compute Dipsatch modified to use FRDGBuilder through the pipeline

diff --git a/Engine/Plugins/FX/Niagara/Source/Niagara/Private/NiagaraDataInterfaceParticleRead.cpp b/Engine/Plugins/FX/Niagara/Source/Niagara/Private/NiagaraDataInterfaceParticleRead.cpp
--- a/Engine/Plugins/FX/Niagara/Source/Niagara/Private/NiagaraDataInterfaceParticleRead.cpp
+++ b/Engine/Plugins/FX/Niagara/Source/Niagara/Private/NiagaraDataInterfaceParticleRead.cpp
@@ -2,12 +2,14 @@
 
 #include "NiagaraDataInterfaceParticleRead.h"
 #include "NiagaraComputeExecutionContext.h"
-#include "NiagaraGPUSystemTick.h"
 #include "NiagaraConstants.h"
+#include "NiagaraDataSetAccessor.h"
+#include "NiagaraGPUSystemTick.h"
+#include "NiagaraRenderer.h"
 #include "NiagaraSystemInstance.h"
+#include "NiagaraShaderParametersBuilder.h"
+
 #include "ShaderParameterUtils.h"
-#include "NiagaraRenderer.h"
-#include "NiagaraDataSetAccessor.h"
 
 #define LOCTEXT_NAMESPACE "NiagaraDataInterfaceParticleRead"
 
@@ -56,20 +58,22 @@ static const FName GetParticleIndexFromIDTableName("Get Particle Index From ID T
 
 static const FName ParticleReadIDName(TEXT("ID"));
 
-static const FString NumSpawnedParticlesBaseName(TEXT("NumSpawnedParticles_"));
-static const FString SpawnedParticlesAcquireTagBaseName(TEXT("SpawnedParticlesAcquireTag_"));
-static const FString InstanceCountOffsetBaseName(TEXT("InstanceCountOffset_"));
-static const FString SpawnedIDsBufferBaseName(TEXT("SpawnedIDsBuffer_"));
-static const FString IDToIndexTableBaseName(TEXT("IDToIndexTable_"));
-static const FString InputFloatBufferBaseName(TEXT("InputFloatBuffer_"));
-static const FString InputIntBufferBaseName(TEXT("InputIntBuffer_"));
-static const FString InputHalfBufferBaseName(TEXT("InputHalfBuffer_"));
-static const FString ParticleStrideFloatBaseName(TEXT("ParticleStrideFloat_"));
-static const FString ParticleStrideIntBaseName(TEXT("ParticleStrideInt_"));
-static const FString ParticleStrideHalfBaseName(TEXT("ParticleStrideHalf_"));
-static const FString AttributeIndicesBaseName(TEXT("AttributeIndices_"));
-static const FString AttributeCompressedBaseName(TEXT("AttributeCompressed_"));
-static const FString AcquireTagRegisterIndexBaseName(TEXT("AcquireTagRegisterIndex_"));
+static const TCHAR* NumSpawnedParticlesName = TEXT("_NumSpawnedParticles");
+static const TCHAR* SpawnedParticlesAcquireTagName = TEXT("_SpawnedParticlesAcquireTag");
+static const TCHAR* InstanceCountOffsetName = TEXT("_InstanceCountOffset");
+static const TCHAR* SpawnedIDsBufferName = TEXT("_SpawnedIDsBuffer");
+static const TCHAR* IDToIndexTableName = TEXT("_IDToIndexTable");
+static const TCHAR* InputFloatBufferName = TEXT("_InputFloatBuffer");
+static const TCHAR* InputIntBufferName = TEXT("_InputIntBuffer");
+static const TCHAR* InputHalfBufferName = TEXT("_InputHalfBuffer");
+static const TCHAR* ParticleStrideFloatName = TEXT("_ParticleStrideFloat");
+static const TCHAR* ParticleStrideIntName = TEXT("_ParticleStrideInt");
+static const TCHAR* ParticleStrideHalfName = TEXT("_ParticleStrideHalf");
+static const TCHAR* AcquireTagRegisterIndexName = TEXT("_AcquireTagRegisterIndex");
+static const TCHAR* AttributeIndicesName = TEXT("_AttributeIndices");
+static const TCHAR* AttributeCompressedName = TEXT("_AttributeCompressed");
+static const TCHAR* AttributeIndicesMetadataName = TEXT("AttributeIndices");
+static const TCHAR* AttributeCompressedMetadataName = TEXT("AttributeCompressed");
 
 enum class ENiagaraParticleDataComponentType : uint8
 {
@@ -96,6 +100,16 @@ enum class ENiagaraParticleDataValueType : uint8
 
 DECLARE_INTRINSIC_TYPE_LAYOUT(ENiagaraParticleDataValueType);
 
+struct FNiagaraDataInterfaceParametersCS_ParticleRead : public FNiagaraDataInterfaceParametersCS
+{
+	DECLARE_TYPE_LAYOUT(FNiagaraDataInterfaceParametersCS_ParticleRead, NonVirtual);
+
+	LAYOUT_FIELD(TMemoryImageArray<FMemoryImageName>, AttributeNames);
+	LAYOUT_FIELD(TMemoryImageArray<ENiagaraParticleDataValueType>, AttributeTypes);
+};
+
+IMPLEMENT_TYPE_LAYOUT(FNiagaraDataInterfaceParametersCS_ParticleRead);
+
 static const TCHAR* NiagaraParticleDataValueTypeName(ENiagaraParticleDataValueType Type)
 {
 	switch (Type)
@@ -119,17 +133,32 @@ static bool CheckVariableType(const FNiagaraTypeDefinition& VarType, ENiagaraPar
 {
 	switch (AttributeType)
 	{
-	case ENiagaraParticleDataValueType::Int: return VarType == FNiagaraTypeDefinition::GetIntDef();
-	case ENiagaraParticleDataValueType::Bool: return VarType == FNiagaraTypeDefinition::GetBoolDef();
-	case ENiagaraParticleDataValueType::Float: return VarType == FNiagaraTypeDefinition::GetFloatDef();
-	case ENiagaraParticleDataValueType::Vec2: return VarType == FNiagaraTypeDefinition::GetVec2Def();
-	case ENiagaraParticleDataValueType::Vec3: return VarType == FNiagaraTypeDefinition::GetVec3Def() || VarType == FNiagaraTypeDefinition::GetPositionDef();
-	case ENiagaraParticleDataValueType::Vec4: return VarType == FNiagaraTypeDefinition::GetVec4Def();
-	case ENiagaraParticleDataValueType::Color: return VarType == FNiagaraTypeDefinition::GetColorDef();
-	case ENiagaraParticleDataValueType::Position: return VarType == FNiagaraTypeDefinition::GetPositionDef();
-	case ENiagaraParticleDataValueType::Quat: return VarType == FNiagaraTypeDefinition::GetQuatDef();
-	case ENiagaraParticleDataValueType::ID: return VarType == FNiagaraTypeDefinition::GetIDDef();
-	default: return false;
+		case ENiagaraParticleDataValueType::Int: return VarType == FNiagaraTypeDefinition::GetIntDef();
+		case ENiagaraParticleDataValueType::Bool: return VarType == FNiagaraTypeDefinition::GetBoolDef();
+		case ENiagaraParticleDataValueType::Float: return VarType == FNiagaraTypeDefinition::GetFloatDef();
+		case ENiagaraParticleDataValueType::Vec2: return VarType == FNiagaraTypeDefinition::GetVec2Def();
+		case ENiagaraParticleDataValueType::Vec3: return VarType == FNiagaraTypeDefinition::GetVec3Def() || VarType == FNiagaraTypeDefinition::GetPositionDef();
+		case ENiagaraParticleDataValueType::Vec4: return VarType == FNiagaraTypeDefinition::GetVec4Def();
+		case ENiagaraParticleDataValueType::Color: return VarType == FNiagaraTypeDefinition::GetColorDef();
+		case ENiagaraParticleDataValueType::Position: return VarType == FNiagaraTypeDefinition::GetPositionDef();
+		case ENiagaraParticleDataValueType::Quat: return VarType == FNiagaraTypeDefinition::GetQuatDef();
+		case ENiagaraParticleDataValueType::ID: return VarType == FNiagaraTypeDefinition::GetIDDef();
+		default: return false;
+	}
+}
+
+static bool CheckHalfVariableType(const FNiagaraTypeDefinition& VarType, ENiagaraParticleDataValueType AttributeType)
+{
+	switch (AttributeType)
+	{
+		case ENiagaraParticleDataValueType::Float: return VarType == FNiagaraTypeDefinition::GetHalfDef();
+		case ENiagaraParticleDataValueType::Vec2: return VarType == FNiagaraTypeDefinition::GetHalfVec2Def();
+		case ENiagaraParticleDataValueType::Vec3: return VarType == FNiagaraTypeDefinition::GetHalfVec3Def();
+		case ENiagaraParticleDataValueType::Vec4: return VarType == FNiagaraTypeDefinition::GetHalfVec4Def();
+		case ENiagaraParticleDataValueType::Color: return VarType == FNiagaraTypeDefinition::GetHalfVec4Def();
+		case ENiagaraParticleDataValueType::Position: return VarType == FNiagaraTypeDefinition::GetHalfVec3Def();
+		case ENiagaraParticleDataValueType::Quat: return VarType == FNiagaraTypeDefinition::GetHalfVec4Def();
+		default: return false;
 	}
 }
 
@@ -179,6 +208,8 @@ struct FNDIParticleRead_RenderInstanceData
 	uint32 bWarnFailedToFindAcquireTag : 1;
 };
 
+//////////////////////////////////////////////////////////////////////////
+
 struct FNiagaraDataInterfaceProxyParticleRead : public FNiagaraDataInterfaceProxyRW
 {
 	virtual void ConsumePerInstanceDataFromGameThread(void* Data, const FNiagaraSystemInstanceID& InstanceID) override
@@ -266,359 +297,95 @@ private:
 	TMap<FNiagaraSystemInstanceID, FNDIParticleRead_RenderInstanceData> SystemsRenderData;
 };
 
-struct FNiagaraDataInterfaceParametersCS_ParticleRead : public FNiagaraDataInterfaceParametersCS
-{
-	DECLARE_TYPE_LAYOUT(FNiagaraDataInterfaceParametersCS_ParticleRead, NonVirtual);
-
-	
-
-	void Bind(const FNiagaraDataInterfaceGPUParamInfo& ParameterInfo, const FShaderParameterMap& ParameterMap)
-	{
-		NumSpawnedParticlesParam.Bind(ParameterMap, *(NumSpawnedParticlesBaseName + ParameterInfo.DataInterfaceHLSLSymbol));
-		SpawnedParticlesAcquireTagParam.Bind(ParameterMap, *(SpawnedParticlesAcquireTagBaseName + ParameterInfo.DataInterfaceHLSLSymbol));
-		InstanceCountOffsetParam.Bind(ParameterMap, *(InstanceCountOffsetBaseName + ParameterInfo.DataInterfaceHLSLSymbol));
-		SpawnedIDsBufferParam.Bind(ParameterMap, *(SpawnedIDsBufferBaseName + ParameterInfo.DataInterfaceHLSLSymbol));
-		IDToIndexTableParam.Bind(ParameterMap, *(IDToIndexTableBaseName + ParameterInfo.DataInterfaceHLSLSymbol));
-		InputFloatBufferParam.Bind(ParameterMap, *(InputFloatBufferBaseName + ParameterInfo.DataInterfaceHLSLSymbol));
-		InputIntBufferParam.Bind(ParameterMap, *(InputIntBufferBaseName + ParameterInfo.DataInterfaceHLSLSymbol));
-		InputHalfBufferParam.Bind(ParameterMap, *(InputHalfBufferBaseName + ParameterInfo.DataInterfaceHLSLSymbol));
-		ParticleStrideFloatParam.Bind(ParameterMap, *(ParticleStrideFloatBaseName + ParameterInfo.DataInterfaceHLSLSymbol));
-		ParticleStrideIntParam.Bind(ParameterMap, *(ParticleStrideIntBaseName + ParameterInfo.DataInterfaceHLSLSymbol));
-		ParticleStrideHalfParam.Bind(ParameterMap, *(ParticleStrideHalfBaseName + ParameterInfo.DataInterfaceHLSLSymbol));
-		AttributeIndicesParam.Bind(ParameterMap, *(AttributeIndicesBaseName + ParameterInfo.DataInterfaceHLSLSymbol));
-		AttributeCompressedParam.Bind(ParameterMap, *(AttributeCompressedBaseName + ParameterInfo.DataInterfaceHLSLSymbol));
-		AcquireTagRegisterIndexParam.Bind(ParameterMap, *(AcquireTagRegisterIndexBaseName + ParameterInfo.DataInterfaceHLSLSymbol));
-
-		int32 NumFuncs = ParameterInfo.GeneratedFunctions.Num();
-		AttributeNames.SetNum(NumFuncs);
-		AttributeTypes.SetNum(NumFuncs);
-		for (int32 FuncIdx = 0; FuncIdx < NumFuncs; ++FuncIdx)
-		{
-			const FNiagaraDataInterfaceGeneratedFunction& Func = ParameterInfo.GeneratedFunctions[FuncIdx];
-			static const FName NAME_Attribute("Attribute");
-			const FName* AttributeName = Func.FindSpecifierValue(NAME_Attribute);
-			if (AttributeName != nullptr)
-			{
-				AttributeNames[FuncIdx] = *AttributeName;
-				AttributeTypes[FuncIdx] = GetValueTypeFromFuncName(Func.DefinitionName);
-			}
-			else
-			{
-				// This is not an error. GetNumSpawnedParticles and GetIDAtSpawnIndexFunctionName don't use specifiers,
-				// but they take up slots in the attribute indices array for simplicity. Just stick NAME_None in here to ignore them.
-				AttributeNames[FuncIdx] = NAME_None;
-				AttributeTypes[FuncIdx] = ENiagaraParticleDataValueType::Invalid;
-			}
-		}
-	}
-
-	void SetErrorParams(FRHICommandList& RHICmdList, FRHIComputeShader* ComputeShader, bool bSkipSpawnInfo) const
-	{
-		// Set all the indices to -1, so that any reads return false.
-		TArray<int32, TInlineAllocator<32>> AttributeIndices;
-		TArray<int32, TInlineAllocator<32>> AttributeCompressed;
-
-		int NumAttrIndices = Align(AttributeNames.Num(), 4);
-		AttributeIndices.SetNumUninitialized(NumAttrIndices);
-		AttributeCompressed.SetNumUninitialized(NumAttrIndices);
-
-		for (int AttrIdx = 0; AttrIdx < AttributeIndices.Num(); ++AttrIdx)
-		{
-			AttributeIndices[AttrIdx] = -1;
-			AttributeCompressed[AttrIdx] = 0;
-		}
-
-		int32 AcquireTagRegisterIndex = -1;
-
-		if (!bSkipSpawnInfo)
-		{
-			SetShaderValue(RHICmdList, ComputeShader, NumSpawnedParticlesParam, 0);
-			SetShaderValue(RHICmdList, ComputeShader, SpawnedParticlesAcquireTagParam, 0);
-			SetSRVParameter(RHICmdList, ComputeShader, SpawnedIDsBufferParam, FNiagaraRenderer::GetDummyIntBuffer());
-		}
+//////////////////////////////////////////////////////////////////////////
 
-		SetShaderValue(RHICmdList, ComputeShader, InstanceCountOffsetParam, -1);
-		SetSRVParameter(RHICmdList, ComputeShader, IDToIndexTableParam, FNiagaraRenderer::GetDummyIntBuffer());
-		SetSRVParameter(RHICmdList, ComputeShader, InputFloatBufferParam, FNiagaraRenderer::GetDummyFloatBuffer());
-		SetSRVParameter(RHICmdList, ComputeShader, InputIntBufferParam, FNiagaraRenderer::GetDummyIntBuffer());
-		SetSRVParameter(RHICmdList, ComputeShader, InputHalfBufferParam, FNiagaraRenderer::GetDummyHalfBuffer());
-		SetShaderValue(RHICmdList, ComputeShader, ParticleStrideFloatParam, 0);
-		SetShaderValue(RHICmdList, ComputeShader, ParticleStrideIntParam, 0);
-		SetShaderValue(RHICmdList, ComputeShader, ParticleStrideHalfParam, 0);
-		SetShaderValueArray(RHICmdList, ComputeShader, AttributeIndicesParam, AttributeIndices.GetData(), AttributeIndices.Num());
-		SetShaderValueArray(RHICmdList, ComputeShader, AttributeCompressedParam, AttributeCompressed.GetData(), AttributeCompressed.Num());
-		SetShaderValue(RHICmdList, ComputeShader, AcquireTagRegisterIndexParam, AcquireTagRegisterIndex);
-	}
+static void NDIParticleRead_FindAttributeIndices(FNDIParticleRead_RenderInstanceData* InstanceData, const FNiagaraDataSet* SourceDataSet, const FNiagaraDataInterfaceParametersCS_ParticleRead& ShaderStorage)
+{
+	check(ShaderStorage.AttributeNames.Num() == ShaderStorage.AttributeTypes.Num());
 
-	
+	const int NumAttrIndices = Align(ShaderStorage.AttributeNames.Num(), 4);
+	InstanceData->AttributeIndices.SetNumUninitialized(NumAttrIndices);
+	InstanceData->AttributeCompressed.SetNumUninitialized(NumAttrIndices);
 
-	bool CheckHalfVariableType(const FNiagaraTypeDefinition& VarType, ENiagaraParticleDataValueType AttributeType) const
+	// Find the register index for each named attribute in the source emitter.
+	const TArray<FNiagaraVariable>& SourceEmitterVariables = SourceDataSet->GetVariables();
+	const TArray<FNiagaraVariableLayoutInfo>& SourceEmitterVariableLayouts = SourceDataSet->GetVariableLayouts();
+	for (int AttrNameIdx = 0; AttrNameIdx < ShaderStorage.AttributeNames.Num(); ++AttrNameIdx)
 	{
-		switch (AttributeType)
+		const FMemoryImageName& AttrName = ShaderStorage.AttributeNames[AttrNameIdx];
+		if (AttrName == NAME_None)
 		{
-			case ENiagaraParticleDataValueType::Float: return VarType == FNiagaraTypeDefinition::GetHalfDef();
-			case ENiagaraParticleDataValueType::Vec2: return VarType == FNiagaraTypeDefinition::GetHalfVec2Def();
-			case ENiagaraParticleDataValueType::Vec3: return VarType == FNiagaraTypeDefinition::GetHalfVec3Def();
-			case ENiagaraParticleDataValueType::Vec4: return VarType == FNiagaraTypeDefinition::GetHalfVec4Def();
-			case ENiagaraParticleDataValueType::Color: return VarType == FNiagaraTypeDefinition::GetHalfVec4Def();
-			case ENiagaraParticleDataValueType::Position: return VarType == FNiagaraTypeDefinition::GetHalfVec3Def();
-			case ENiagaraParticleDataValueType::Quat: return VarType == FNiagaraTypeDefinition::GetHalfVec4Def();
-			default: return false;
+			InstanceData->AttributeIndices[AttrNameIdx] = -1;
+			InstanceData->AttributeCompressed[AttrNameIdx] = 0;
+			continue;
 		}
-	}
-
-	void FindAttributeIndices(FNDIParticleRead_RenderInstanceData* InstanceData, const FNiagaraDataSet* SourceDataSet) const
-	{
-		check(AttributeNames.Num() == AttributeTypes.Num());
-
-		int NumAttrIndices = Align(AttributeNames.Num(), 4);
-		InstanceData->AttributeIndices.SetNumUninitialized(NumAttrIndices);
-		InstanceData->AttributeCompressed.SetNumUninitialized(NumAttrIndices);
 
-		// Find the register index for each named attribute in the source emitter.
-		const TArray<FNiagaraVariable>& SourceEmitterVariables = SourceDataSet->GetVariables();
-		const TArray<FNiagaraVariableLayoutInfo>& SourceEmitterVariableLayouts = SourceDataSet->GetVariableLayouts();
-		for (int AttrNameIdx = 0; AttrNameIdx < AttributeNames.Num(); ++AttrNameIdx)
+		bool FoundVariable = false;
+		for (int VarIdx = 0; VarIdx < SourceEmitterVariables.Num(); ++VarIdx)
 		{
-			const FMemoryImageName& AttrName = AttributeNames[AttrNameIdx];
-			if (AttrName == NAME_None)
+			const FNiagaraVariable& Var = SourceEmitterVariables[VarIdx];
+			if (Var.GetName() == AttrName)
 			{
-				InstanceData->AttributeIndices[AttrNameIdx] = -1;
-				continue;
-			}
-
-			bool FoundVariable = false;
-			for (int VarIdx = 0; VarIdx < SourceEmitterVariables.Num(); ++VarIdx)
-			{
-				const FNiagaraVariable& Var = SourceEmitterVariables[VarIdx];
-				if (Var.GetName() == AttrName)
+				ENiagaraParticleDataValueType AttributeType = ShaderStorage.AttributeTypes[AttrNameIdx];
+				if (CheckVariableType(Var.GetType(), AttributeType))
 				{
-					ENiagaraParticleDataValueType AttributeType = AttributeTypes[AttrNameIdx];
-					if (CheckVariableType(Var.GetType(), AttributeType))
-					{
-						const FNiagaraVariableLayoutInfo& Layout = SourceEmitterVariableLayouts[VarIdx];
-						InstanceData->AttributeIndices[AttrNameIdx] = 
-							(AttributeType == ENiagaraParticleDataValueType::Int || AttributeType == ENiagaraParticleDataValueType::Bool || AttributeType == ENiagaraParticleDataValueType::ID) ? Layout.Int32ComponentStart : Layout.FloatComponentStart;
-						InstanceData->AttributeCompressed[AttrNameIdx] = 0;
-					}
-					else if (CheckHalfVariableType(Var.GetType(), AttributeType))
-					{
-						const FNiagaraVariableLayoutInfo& Layout = SourceEmitterVariableLayouts[VarIdx];
-						InstanceData->AttributeIndices[AttrNameIdx] = Layout.HalfComponentStart;
-						InstanceData->AttributeCompressed[AttrNameIdx] = 1;
-					}
-					else
-					{
-						UE_LOG(LogNiagara, Error, TEXT("Variable '%s' in emitter '%s' has type '%s', but particle read DI tried to access it as '%s'."),
-							*Var.GetName().ToString(), *InstanceData->DebugSourceName, *Var.GetType().GetName(), NiagaraParticleDataValueTypeName(AttributeType)
-						);
-						InstanceData->AttributeIndices[AttrNameIdx] = -1;
-						InstanceData->AttributeCompressed[AttrNameIdx] = 0;
-					}
-					FoundVariable = true;
-					break;
+					const FNiagaraVariableLayoutInfo& Layout = SourceEmitterVariableLayouts[VarIdx];
+					InstanceData->AttributeIndices[AttrNameIdx] =
+						(AttributeType == ENiagaraParticleDataValueType::Int || AttributeType == ENiagaraParticleDataValueType::Bool || AttributeType == ENiagaraParticleDataValueType::ID) ? Layout.Int32ComponentStart : Layout.FloatComponentStart;
+					InstanceData->AttributeCompressed[AttrNameIdx] = 0;
 				}
-			}
-
-			if (!FoundVariable)
-			{
-				UE_LOG(LogNiagara, Error, TEXT("Particle read DI is trying to access inexistent variable '%s' in emitter '%s'."), *AttrName.ToString(), *InstanceData->DebugSourceName);
-				InstanceData->AttributeIndices[AttrNameIdx] = -1;
-				InstanceData->AttributeCompressed[AttrNameIdx] = 0;
-			}
-		}
-
-		// Find the register index for the AcquireTag part of the particle ID in the source emitter.
-		{
-			const FName FName_ID(TEXT("ID"));
-
-			InstanceData->AcquireTagRegisterIndex = -1;
-			for (int VarIdx = 0; VarIdx < SourceEmitterVariables.Num(); ++VarIdx)
-			{
-				const FNiagaraVariable& Var = SourceEmitterVariables[VarIdx];
-				if (Var.GetName() == FName_ID)
+				else if (CheckHalfVariableType(Var.GetType(), AttributeType))
 				{
-					InstanceData->AcquireTagRegisterIndex = SourceEmitterVariableLayouts[VarIdx].Int32ComponentStart + 1;
-					break;
+					const FNiagaraVariableLayoutInfo& Layout = SourceEmitterVariableLayouts[VarIdx];
+					InstanceData->AttributeIndices[AttrNameIdx] = Layout.HalfComponentStart;
+					InstanceData->AttributeCompressed[AttrNameIdx] = 1;
 				}
+				else
+				{
+					UE_LOG(LogNiagara, Error, TEXT("Variable '%s' in emitter '%s' has type '%s', but particle read DI tried to access it as '%s'."),
+						*Var.GetName().ToString(), *InstanceData->DebugSourceName, *Var.GetType().GetName(), NiagaraParticleDataValueTypeName(AttributeType)
+					);
+					InstanceData->AttributeIndices[AttrNameIdx] = -1;
+					InstanceData->AttributeCompressed[AttrNameIdx] = 0;
+				}
+				FoundVariable = true;
+				break;
 			}
 		}
 
-		// Initialize the buffer padding too, so we don't move garbage around.
-		for (int AttrIdx = AttributeNames.Num(); AttrIdx < InstanceData->AttributeIndices.Num(); ++AttrIdx)
+		if (!FoundVariable)
 		{
-			InstanceData->AttributeIndices[AttrIdx] = -1;
-			InstanceData->AttributeCompressed[AttrIdx] = 0;
+			UE_LOG(LogNiagara, Error, TEXT("Particle read DI is trying to access inexistent variable '%s' in emitter '%s'."), *AttrName.ToString(), *InstanceData->DebugSourceName);
+			InstanceData->AttributeIndices[AttrNameIdx] = -1;
+			InstanceData->AttributeCompressed[AttrNameIdx] = 0;
 		}
 	}
 
-	FRHIShaderResourceView* GetIntSRVWithFallback(FRWBuffer& Buffer) const
-	{
-		return Buffer.SRV ? Buffer.SRV.GetReference() : FNiagaraRenderer::GetDummyIntBuffer();
-	}
-
-	FRHIShaderResourceView* GetFloatSRVWithFallback(FRWBuffer& Buffer) const
-	{
-		return Buffer.SRV ? Buffer.SRV.GetReference() : FNiagaraRenderer::GetDummyFloatBuffer();
-	}
-
-	FRHIShaderResourceView* GetHalfSRVWithFallback(FRWBuffer& Buffer) const
-	{
-		return Buffer.SRV ? Buffer.SRV.GetReference() : FNiagaraRenderer::GetDummyHalfBuffer();
-	}
-
-	void Set(FRHICommandList& RHICmdList, const FNiagaraDataInterfaceSetArgs& Context) const
+	// Find the register index for the AcquireTag part of the particle ID in the source emitter.
 	{
-		check(IsInRenderingThread());
-
-		FRHIComputeShader* ComputeShader = RHICmdList.GetBoundComputeShader();
-
-		FNiagaraDataInterfaceProxyParticleRead* Proxy = static_cast<FNiagaraDataInterfaceProxyParticleRead*>(Context.DataInterface);
-		check(Proxy);
-
-		FNDIParticleRead_RenderInstanceData* InstanceData = Proxy->GetRenderDataForSystem(Context.SystemInstanceID);
-		if (!InstanceData)
-		{
-			SetErrorParams(RHICmdList, ComputeShader, false);
-			return;
-		}
-
-		if (InstanceData->SourceEmitterGPUContext == nullptr)
-		{
-			// This means the source emitter isn't running on GPU.
-			if (!InstanceData->bSourceEmitterNotGPUErrorShown)
-			{
-				UE_LOG(LogNiagara, Error, TEXT("GPU particle read DI is set to access CPU emitter '%s'."), *InstanceData->DebugSourceName);
-				InstanceData->bSourceEmitterNotGPUErrorShown = true;
-			}
-			SetErrorParams(RHICmdList, ComputeShader, false);
-			InstanceData->CachedDataSet = nullptr;
-			InstanceData->bWarnFailedToFindAcquireTag = true;
-			return;
-		}
-
-		InstanceData->bSourceEmitterNotGPUErrorShown = false;
-
-		FNiagaraDataSet* SourceDataSet = InstanceData->SourceEmitterGPUContext->MainDataSet;
-		if (!SourceDataSet)
-		{
-			SetErrorParams(RHICmdList, ComputeShader, false);
-			InstanceData->CachedDataSet = nullptr;
-			InstanceData->bWarnFailedToFindAcquireTag = true;
-			return;
-		}
-
-		FNiagaraDataBuffer* SourceData = nullptr;
-		uint32 NumSpawnedInstances = 0;
-		uint32 IDAcquireTag = 0;
-
-		// Test for reading ourself
-		if (Context.ComputeInstanceData->Context == InstanceData->SourceEmitterGPUContext)
-		{
-			// Pull source data from SimStageData as this has the correct information until all stages are dispatched
-			// Note: In the case of an in place (aka partial update) stage this does not work as we need to bind the
-			//       same buffer as both UAV & SRV.  For the moment the user must disable partial writes.
-			//       Another option would be to read from the UAV, but that would equally race.
-			//-TODO: Automate this in some way or surface as feedback...
-			if (const FSimulationStageMetaData* StageMetaData = Context.SimStageData->StageMetaData )
-			{
-				if ( StageMetaData->bPartialParticleUpdate )
-				{
-					UE_LOG(LogNiagara, Error, TEXT("Particle read DI reading self '%s' on stage '%s' is unsafe, please disable partial writes on the stage."), *InstanceData->DebugSourceName, *StageMetaData->SimulationStageName.ToString());
-					SetErrorParams(RHICmdList, ComputeShader, true);
-					return;
-				}
-			}
-
-			SourceData = Context.SimStageData->Source;
+		const FName FName_ID(TEXT("ID"));
 
-			// Pull some spawned / acquire tag from destination data as we are reading from ourselves this is what we are currently doing in this stage.
-			if (Context.SimStageData->Destination)
-			{
-				ensure(SourceData != nullptr);
-				NumSpawnedInstances = Context.SimStageData->Destination->GetNumSpawnedInstances();
-				IDAcquireTag = Context.SimStageData->Destination->GetIDAcquireTag();
-			}
-			// In cases where we do not have destination data pull from source (i.e. this is an iteration on none particle stages)
-			else if (SourceData)
-			{
-				NumSpawnedInstances = SourceData->GetNumSpawnedInstances();
-				IDAcquireTag = SourceData->GetIDAcquireTag();
-			}
-		}
-		else
+		InstanceData->AcquireTagRegisterIndex = -1;
+		for (int VarIdx = 0; VarIdx < SourceEmitterVariables.Num(); ++VarIdx)
 		{
-			SourceData = SourceDataSet->GetCurrentData();
-			if ( SourceData )
+			const FNiagaraVariable& Var = SourceEmitterVariables[VarIdx];
+			if (Var.GetName() == FName_ID)
 			{
-				NumSpawnedInstances = SourceData->GetNumSpawnedInstances();
-				IDAcquireTag = SourceData->GetIDAcquireTag();
+				InstanceData->AcquireTagRegisterIndex = SourceEmitterVariableLayouts[VarIdx].Int32ComponentStart + 1;
+				break;
 			}
 		}
+	}
 
-		SetShaderValue(RHICmdList, ComputeShader, NumSpawnedParticlesParam, NumSpawnedInstances);
-		SetShaderValue(RHICmdList, ComputeShader, SpawnedParticlesAcquireTagParam, IDAcquireTag);
-		SetSRVParameter(RHICmdList, ComputeShader, SpawnedIDsBufferParam, GetIntSRVWithFallback(SourceDataSet->GetGPUFreeIDs()));
-
-		if (!SourceData)
-		{
-			SetErrorParams(RHICmdList, ComputeShader, true);
-			return;
-		}
-
-		if (InstanceData->CachedDataSet != SourceDataSet)
-		{
-			FindAttributeIndices(InstanceData, SourceDataSet);
-			InstanceData->CachedDataSet = SourceDataSet;
-			InstanceData->bWarnFailedToFindAcquireTag = true;
-		}
-
-		const uint32 ParticleStrideFloat = SourceData->GetFloatStride() / sizeof(float);
-		const uint32 ParticleStrideInt = SourceData->GetInt32Stride() / sizeof(int32);
-		const uint32 ParticleStrideHalf = SourceData->GetHalfStride() / sizeof(FFloat16);
-
-		// There's no need to transition the input buffers, because the grouping logic inside FNiagaraGpuComputeDispatch ensures that our source emitter has ran before us,
-		// and its buffers have been transitioned to readable.
-
-		SetShaderValue(RHICmdList, ComputeShader, InstanceCountOffsetParam, SourceData->GetGPUInstanceCountBufferOffset());
-		SetSRVParameter(RHICmdList, ComputeShader, IDToIndexTableParam, GetIntSRVWithFallback(SourceData->GetGPUIDToIndexTable()));
-		SetSRVParameter(RHICmdList, ComputeShader, InputFloatBufferParam, GetFloatSRVWithFallback(SourceData->GetGPUBufferFloat()));
-		SetSRVParameter(RHICmdList, ComputeShader, InputIntBufferParam, GetIntSRVWithFallback(SourceData->GetGPUBufferInt()));
-		SetSRVParameter(RHICmdList, ComputeShader, InputHalfBufferParam, GetHalfSRVWithFallback(SourceData->GetGPUBufferHalf()));
-		SetShaderValue(RHICmdList, ComputeShader, ParticleStrideFloatParam, ParticleStrideFloat);
-		SetShaderValue(RHICmdList, ComputeShader, ParticleStrideIntParam, ParticleStrideInt);
-		SetShaderValue(RHICmdList, ComputeShader, ParticleStrideHalfParam, ParticleStrideHalf);
-		SetShaderValueArray(RHICmdList, ComputeShader, AttributeIndicesParam, InstanceData->AttributeIndices.GetData(), InstanceData->AttributeIndices.Num());
-		SetShaderValueArray(RHICmdList, ComputeShader, AttributeCompressedParam, InstanceData->AttributeCompressed.GetData(), InstanceData->AttributeCompressed.Num());
-		SetShaderValue(RHICmdList, ComputeShader, AcquireTagRegisterIndexParam, InstanceData->AcquireTagRegisterIndex);
-
-		if (InstanceData->bWarnFailedToFindAcquireTag && AcquireTagRegisterIndexParam.IsBound() && (InstanceData->AcquireTagRegisterIndex == -1))
-		{
-			InstanceData->bWarnFailedToFindAcquireTag = false;
-			UE_LOG(LogNiagara, Error, TEXT("Particle read DI cannot find ID variable in emitter '%s'."), *InstanceData->DebugSourceName);
-		}
+	// Initialize the buffer padding too, so we don't move garbage around.
+	for (int AttrIdx = ShaderStorage.AttributeNames.Num(); AttrIdx < InstanceData->AttributeIndices.Num(); ++AttrIdx)
+	{
+		InstanceData->AttributeIndices[AttrIdx] = -1;
+		InstanceData->AttributeCompressed[AttrIdx] = 0;
 	}
-	
-private:
-	LAYOUT_FIELD(FShaderParameter, NumSpawnedParticlesParam);
-	LAYOUT_FIELD(FShaderParameter, SpawnedParticlesAcquireTagParam);
-	LAYOUT_FIELD(FShaderParameter, InstanceCountOffsetParam);
-	LAYOUT_FIELD(FShaderResourceParameter, SpawnedIDsBufferParam);
-	LAYOUT_FIELD(FShaderResourceParameter, IDToIndexTableParam);
-	LAYOUT_FIELD(FShaderResourceParameter, InputFloatBufferParam);
-	LAYOUT_FIELD(FShaderResourceParameter, InputIntBufferParam);
-	LAYOUT_FIELD(FShaderResourceParameter, InputHalfBufferParam);
-	LAYOUT_FIELD(FShaderParameter, ParticleStrideFloatParam);
-	LAYOUT_FIELD(FShaderParameter, ParticleStrideHalfParam);
-	LAYOUT_FIELD(FShaderParameter, ParticleStrideIntParam);
-	LAYOUT_FIELD(FShaderParameter, AttributeIndicesParam);
-	LAYOUT_FIELD(FShaderParameter, AttributeCompressedParam);
-	LAYOUT_FIELD(FShaderParameter, AcquireTagRegisterIndexParam);
-	LAYOUT_FIELD(TMemoryImageArray<FMemoryImageName>, AttributeNames);
-	LAYOUT_FIELD(TMemoryImageArray<ENiagaraParticleDataValueType>, AttributeTypes);
-};
+}
 
-IMPLEMENT_TYPE_LAYOUT(FNiagaraDataInterfaceParametersCS_ParticleRead);
+//////////////////////////////////////////////////////////////////////////
 
 UNiagaraDataInterfaceParticleRead::UNiagaraDataInterfaceParticleRead(FObjectInitializer const& ObjectInitializer)
 	: Super(ObjectInitializer)
@@ -1924,49 +1691,35 @@ bool UNiagaraDataInterfaceParticleRead::CopyToInternal(UNiagaraDataInterface* De
 }
 
 #if WITH_EDITORONLY_DATA
-void UNiagaraDataInterfaceParticleRead::GetParameterDefinitionHLSL(const FNiagaraDataInterfaceGPUParamInfo& ParamInfo, FString& OutHLSL)
+bool UNiagaraDataInterfaceParticleRead::AppendCompileHash(FNiagaraCompileHashVisitor* InVisitor) const
 {
-	static const TCHAR *FormatDeclarations = TEXT(
-		"int {NumSpawnedParticlesName};\n"
-		"int {SpawnedParticlesAcquireTagName};\n"
-		"uint {InstanceCountOffsetName};\n"
-		"uint {ParticleStrideFloatName};\n"
-		"uint {ParticleStrideIntName};\n"
-		"uint {ParticleStrideHalfName};\n"
-		"int {AcquireTagRegisterIndexName};\n"
-		"Buffer<int> {SpawnedIDsBufferName};\n"
-		"Buffer<int> {IDToIndexTableName};\n"
-		"Buffer<float> {InputFloatBufferName};\n"
-		"Buffer<int> {InputIntBufferName};\n"
-		"Buffer<half> {InputHalfBufferName};\n"
-		"int4 {AttributeIndicesName}[{AttributeInt4Count}];\n"
-		"int4 {AttributeCompressedName}[{AttributeInt4Count}];\n\n"
-	);
+	bool bSuccess = Super::AppendCompileHash(InVisitor);
+	bSuccess &= InVisitor->UpdateShaderParameters<FShaderParameters>();
+	return bSuccess;
+}
 
+void UNiagaraDataInterfaceParticleRead::GetParameterDefinitionHLSL(const FNiagaraDataInterfaceGPUParamInfo& ParamInfo, FString& OutHLSL)
+{
 	// If we use an int array for the attribute indices, the shader compiler will actually use int4 due to the packing rules,
 	// and leave 3 elements unused. Besides being wasteful, this means that the array we send to the CS would need to be padded,
 	// which is a hassle. Instead, use int4 explicitly, and access individual components in the generated code.
 	// Note that we have to have at least one here because hlsl doesn't support arrays of size 0.
 	const int AttributeInt4Count = FMath::Max(1, FMath::DivideAndRoundUp(ParamInfo.GeneratedFunctions.Num(), 4));
 
-	TMap<FString, FStringFormatArg> ArgsDeclarations;
-	ArgsDeclarations.Add(TEXT("NumSpawnedParticlesName"), NumSpawnedParticlesBaseName + ParamInfo.DataInterfaceHLSLSymbol);
-	ArgsDeclarations.Add(TEXT("SpawnedParticlesAcquireTagName"), SpawnedParticlesAcquireTagBaseName + ParamInfo.DataInterfaceHLSLSymbol);
-	ArgsDeclarations.Add(TEXT("InstanceCountOffsetName"), InstanceCountOffsetBaseName + ParamInfo.DataInterfaceHLSLSymbol);
-	ArgsDeclarations.Add(TEXT("ParticleStrideFloatName"), ParticleStrideFloatBaseName + ParamInfo.DataInterfaceHLSLSymbol);
-	ArgsDeclarations.Add(TEXT("ParticleStrideIntName"), ParticleStrideIntBaseName + ParamInfo.DataInterfaceHLSLSymbol);
-	ArgsDeclarations.Add(TEXT("ParticleStrideHalfName"), ParticleStrideHalfBaseName + ParamInfo.DataInterfaceHLSLSymbol);
-	ArgsDeclarations.Add(TEXT("AcquireTagRegisterIndexName"), AcquireTagRegisterIndexBaseName + ParamInfo.DataInterfaceHLSLSymbol);
-	ArgsDeclarations.Add(TEXT("SpawnedIDsBufferName"), SpawnedIDsBufferBaseName + ParamInfo.DataInterfaceHLSLSymbol);
-	ArgsDeclarations.Add(TEXT("IDToIndexTableName"), IDToIndexTableBaseName + ParamInfo.DataInterfaceHLSLSymbol);
-	ArgsDeclarations.Add(TEXT("InputFloatBufferName"), InputFloatBufferBaseName + ParamInfo.DataInterfaceHLSLSymbol);
-	ArgsDeclarations.Add(TEXT("InputIntBufferName"), InputIntBufferBaseName + ParamInfo.DataInterfaceHLSLSymbol);
-	ArgsDeclarations.Add(TEXT("InputHalfBufferName"), InputHalfBufferBaseName + ParamInfo.DataInterfaceHLSLSymbol);
-	ArgsDeclarations.Add(TEXT("AttributeIndicesName"), AttributeIndicesBaseName + ParamInfo.DataInterfaceHLSLSymbol);
-	ArgsDeclarations.Add(TEXT("AttributeCompressedName"), AttributeCompressedBaseName + ParamInfo.DataInterfaceHLSLSymbol);
-	ArgsDeclarations.Add(TEXT("AttributeInt4Count"), AttributeInt4Count);
-
-	OutHLSL += FString::Format(FormatDeclarations, ArgsDeclarations);
+	OutHLSL.Appendf(TEXT("int			%s%s;"), *ParamInfo.DataInterfaceHLSLSymbol, NumSpawnedParticlesName);
+	OutHLSL.Appendf(TEXT("int			%s%s;"), *ParamInfo.DataInterfaceHLSLSymbol, SpawnedParticlesAcquireTagName);
+	OutHLSL.Appendf(TEXT("uint			%s%s;"), *ParamInfo.DataInterfaceHLSLSymbol, InstanceCountOffsetName);
+	OutHLSL.Appendf(TEXT("uint			%s%s;"), *ParamInfo.DataInterfaceHLSLSymbol, ParticleStrideFloatName);
+	OutHLSL.Appendf(TEXT("uint			%s%s;"), *ParamInfo.DataInterfaceHLSLSymbol, ParticleStrideIntName);
+	OutHLSL.Appendf(TEXT("uint			%s%s;"), *ParamInfo.DataInterfaceHLSLSymbol, ParticleStrideHalfName);
+	OutHLSL.Appendf(TEXT("int			%s%s;"), *ParamInfo.DataInterfaceHLSLSymbol, AcquireTagRegisterIndexName);
+	OutHLSL.Appendf(TEXT("Buffer<int>	%s%s;"), *ParamInfo.DataInterfaceHLSLSymbol, SpawnedIDsBufferName);
+	OutHLSL.Appendf(TEXT("Buffer<int>	%s%s;"), *ParamInfo.DataInterfaceHLSLSymbol, IDToIndexTableName);
+	OutHLSL.Appendf(TEXT("Buffer<float>	%s%s;"), *ParamInfo.DataInterfaceHLSLSymbol, InputFloatBufferName);
+	OutHLSL.Appendf(TEXT("Buffer<int>	%s%s;"), *ParamInfo.DataInterfaceHLSLSymbol, InputIntBufferName);
+	OutHLSL.Appendf(TEXT("Buffer<half>	%s%s;"), *ParamInfo.DataInterfaceHLSLSymbol, InputHalfBufferName);
+	OutHLSL.Appendf(TEXT("int4			%s%s[%d];"), *ParamInfo.DataInterfaceHLSLSymbol, AttributeIndicesName, AttributeInt4Count);
+	OutHLSL.Appendf(TEXT("int4			%s%s[%d];"), *ParamInfo.DataInterfaceHLSLSymbol, AttributeCompressedName, AttributeInt4Count);
 }
 
 static FString GenerateFetchValueHLSL(int NumComponents, const TCHAR* ComponentNames[], const TCHAR* ComponentTypeName, const FString& InputBufferName, const FString& ParticleStrideName, bool bExtraIndent)
@@ -2024,12 +1777,12 @@ static bool GenerateGetFunctionHLSL(const FNiagaraDataInterfaceGPUParamInfo& Par
 	static const TCHAR* VectorComponentNames[] = { TEXT(".x"), TEXT(".y"), TEXT(".z"), TEXT(".w") };
 	static const TCHAR* IDComponentNames[] = { TEXT(".Index"), TEXT(".AcquireTag") };
 
-	const FString ParticleStrideFloatName = ParticleStrideFloatBaseName + ParamInfo.DataInterfaceHLSLSymbol;
-	const FString ParticleStrideIntName = ParticleStrideIntBaseName + ParamInfo.DataInterfaceHLSLSymbol;
-	const FString ParticleStrideHalfName = ParticleStrideHalfBaseName + ParamInfo.DataInterfaceHLSLSymbol;
-	const FString InputFloatBufferName = InputFloatBufferBaseName + ParamInfo.DataInterfaceHLSLSymbol;
-	const FString InputIntBufferName = InputIntBufferBaseName + ParamInfo.DataInterfaceHLSLSymbol;
-	const FString InputHalfBufferName = InputHalfBufferBaseName + ParamInfo.DataInterfaceHLSLSymbol;
+	const FString FullParticleStrideFloatName = ParamInfo.DataInterfaceHLSLSymbol + ParticleStrideFloatName;
+	const FString FullParticleStrideIntName = ParamInfo.DataInterfaceHLSLSymbol + ParticleStrideIntName;
+	const FString FullParticleStrideHalfName = ParamInfo.DataInterfaceHLSLSymbol + ParticleStrideHalfName;
+	const FString FullInputFloatBufferName = ParamInfo.DataInterfaceHLSLSymbol + InputFloatBufferName;
+	const FString FullInputIntBufferName = ParamInfo.DataInterfaceHLSLSymbol + InputIntBufferName;
+	const FString FullInputHalfBufferName = ParamInfo.DataInterfaceHLSLSymbol + InputHalfBufferName;
 
 	const TCHAR* ComponentTypeName;
 	FString FetchValueCode;
@@ -2055,14 +1808,14 @@ static bool GenerateGetFunctionHLSL(const FNiagaraDataInterfaceGPUParamInfo& Par
 				return false;
 		}
 
-		FetchValueCode = GenerateFetchValueHLSL(NumComponents, ValueComponentNames, ComponentTypeName, InputIntBufferName, ParticleStrideIntName, false);
+		FetchValueCode = GenerateFetchValueHLSL(NumComponents, ValueComponentNames, ComponentTypeName, FullInputIntBufferName, FullParticleStrideIntName, false);
 	}
 	else
 	{
 		ComponentTypeName = TEXT("float");
 		// Floats and vectors can be compressed, so we need to add extra code which checks.
-		FString FetchFloatCode = GenerateFetchValueHLSL(NumComponents, VectorComponentNames, ComponentTypeName, InputFloatBufferName, ParticleStrideFloatName, true);
-		FString FetchHalfCode = GenerateFetchValueHLSL(NumComponents, VectorComponentNames, ComponentTypeName, InputHalfBufferName, ParticleStrideHalfName, true);
+		FString FetchFloatCode = GenerateFetchValueHLSL(NumComponents, VectorComponentNames, ComponentTypeName, FullInputFloatBufferName, FullParticleStrideFloatName, true);
+		FString FetchHalfCode = GenerateFetchValueHLSL(NumComponents, VectorComponentNames, ComponentTypeName, FullInputHalfBufferName, FullParticleStrideHalfName, true);
 		FetchValueCode = FString(
 			TEXT(
 			"#if NIAGARA_COMPRESSED_ATTRIBUTES_ENABLED\n"
@@ -2115,16 +1868,16 @@ static bool GenerateGetFunctionHLSL(const FNiagaraDataInterfaceGPUParamInfo& Par
 	TMap<FString, FStringFormatArg> FuncTemplateArgs;
 	FuncTemplateArgs.Add(TEXT("FunctionName"), FunctionInfo.InstanceName);
 	FuncTemplateArgs.Add(TEXT("ValueType"), ValueTypeName);
-	FuncTemplateArgs.Add(TEXT("AttributeIndicesName"), AttributeIndicesBaseName + ParamInfo.DataInterfaceHLSLSymbol);
-	FuncTemplateArgs.Add(TEXT("AttributeCompressedName"), AttributeCompressedBaseName + ParamInfo.DataInterfaceHLSLSymbol);
+	FuncTemplateArgs.Add(TEXT("AttributeIndicesName"), ParamInfo.DataInterfaceHLSLSymbol + AttributeIndicesName);
+	FuncTemplateArgs.Add(TEXT("AttributeCompressedName"), ParamInfo.DataInterfaceHLSLSymbol + AttributeCompressedName);
 	FuncTemplateArgs.Add(TEXT("AttributeIndexGroup"), FunctionInstanceIndex / 4);
 	FuncTemplateArgs.Add(TEXT("AttributeIndexComponent"), VectorComponentNames[FunctionInstanceIndex % 4]);
-	FuncTemplateArgs.Add(TEXT("IDToIndexTableName"), IDToIndexTableBaseName + ParamInfo.DataInterfaceHLSLSymbol);
-	FuncTemplateArgs.Add(TEXT("InputIntBufferName"), InputIntBufferName);
-	FuncTemplateArgs.Add(TEXT("AcquireTagRegisterIndexName"), AcquireTagRegisterIndexBaseName + ParamInfo.DataInterfaceHLSLSymbol);
-	FuncTemplateArgs.Add(TEXT("ParticleStrideIntName"), ParticleStrideIntName);
+	FuncTemplateArgs.Add(TEXT("IDToIndexTableName"),  ParamInfo.DataInterfaceHLSLSymbol + IDToIndexTableName);
+	FuncTemplateArgs.Add(TEXT("InputIntBufferName"), ParamInfo.DataInterfaceHLSLSymbol + InputIntBufferName);
+	FuncTemplateArgs.Add(TEXT("AcquireTagRegisterIndexName"), ParamInfo.DataInterfaceHLSLSymbol + AcquireTagRegisterIndexName);
+	FuncTemplateArgs.Add(TEXT("ParticleStrideIntName"), ParamInfo.DataInterfaceHLSLSymbol + ParticleStrideIntName);
 	FuncTemplateArgs.Add(TEXT("FetchValueCode"), FetchValueCode);
-	FuncTemplateArgs.Add(TEXT("InstanceCountOffsetName"), InstanceCountOffsetBaseName + ParamInfo.DataInterfaceHLSLSymbol);
+	FuncTemplateArgs.Add(TEXT("InstanceCountOffsetName"), ParamInfo.DataInterfaceHLSLSymbol + InstanceCountOffsetName);
 
 	OutHLSL += FString::Format(*FuncTemplate, FuncTemplateArgs);
 
@@ -2147,7 +1900,7 @@ bool UNiagaraDataInterfaceParticleRead::GetFunctionHLSL(const FNiagaraDataInterf
 
 		TMap<FString, FStringFormatArg> FuncTemplateArgs;
 		FuncTemplateArgs.Add(TEXT("FunctionName"), FunctionInfo.InstanceName);
-		FuncTemplateArgs.Add(TEXT("NumSpawnedParticlesName"), NumSpawnedParticlesBaseName + ParamInfo.DataInterfaceHLSLSymbol);
+		FuncTemplateArgs.Add(TEXT("NumSpawnedParticlesName"), ParamInfo.DataInterfaceHLSLSymbol + NumSpawnedParticlesName);
 
 		OutHLSL += FString::Format(FuncTemplate, FuncTemplateArgs);
 		return true;
@@ -2175,9 +1928,9 @@ bool UNiagaraDataInterfaceParticleRead::GetFunctionHLSL(const FNiagaraDataInterf
 
 		TMap<FString, FStringFormatArg> FuncTemplateArgs;
 		FuncTemplateArgs.Add(TEXT("FunctionName"), FunctionInfo.InstanceName);
-		FuncTemplateArgs.Add(TEXT("NumSpawnedParticlesName"), NumSpawnedParticlesBaseName + ParamInfo.DataInterfaceHLSLSymbol);
-		FuncTemplateArgs.Add(TEXT("SpawnedParticlesAcquireTagName"), SpawnedParticlesAcquireTagBaseName + ParamInfo.DataInterfaceHLSLSymbol);
-		FuncTemplateArgs.Add(TEXT("SpawnedIDsBufferName"), SpawnedIDsBufferBaseName + ParamInfo.DataInterfaceHLSLSymbol);
+		FuncTemplateArgs.Add(TEXT("NumSpawnedParticlesName"), ParamInfo.DataInterfaceHLSLSymbol + NumSpawnedParticlesName);
+		FuncTemplateArgs.Add(TEXT("SpawnedParticlesAcquireTagName"), ParamInfo.DataInterfaceHLSLSymbol + SpawnedParticlesAcquireTagName);
+		FuncTemplateArgs.Add(TEXT("SpawnedIDsBufferName"), ParamInfo.DataInterfaceHLSLSymbol + SpawnedIDsBufferName);
 
 		OutHLSL += FString::Format(FuncTemplate, FuncTemplateArgs);
 		return true;
@@ -2201,7 +1954,7 @@ bool UNiagaraDataInterfaceParticleRead::GetFunctionHLSL(const FNiagaraDataInterf
 
 		TMap<FString, FStringFormatArg> FuncTemplateArgs;
 		FuncTemplateArgs.Add(TEXT("FunctionName"), FunctionInfo.InstanceName);
-		FuncTemplateArgs.Add(TEXT("InstanceCountOffsetName"), InstanceCountOffsetBaseName + ParamInfo.DataInterfaceHLSLSymbol);
+		FuncTemplateArgs.Add(TEXT("InstanceCountOffsetName"), ParamInfo.DataInterfaceHLSLSymbol + InstanceCountOffsetName);
 
 		OutHLSL += FString::Format(FuncTemplate, FuncTemplateArgs);
 		return true;
@@ -2223,10 +1976,10 @@ bool UNiagaraDataInterfaceParticleRead::GetFunctionHLSL(const FNiagaraDataInterf
 
 		TMap<FString, FStringFormatArg> FuncTemplateArgs;
 		FuncTemplateArgs.Add(TEXT("FunctionName"), FunctionInfo.InstanceName);
-		FuncTemplateArgs.Add(TEXT("IDToIndexTableName"), IDToIndexTableBaseName + ParamInfo.DataInterfaceHLSLSymbol);
-		FuncTemplateArgs.Add(TEXT("InputIntBufferName"), InputIntBufferBaseName + ParamInfo.DataInterfaceHLSLSymbol);
-		FuncTemplateArgs.Add(TEXT("AcquireTagRegisterIndexName"), AcquireTagRegisterIndexBaseName + ParamInfo.DataInterfaceHLSLSymbol);
-		FuncTemplateArgs.Add(TEXT("ParticleStrideIntName"), ParticleStrideIntBaseName + ParamInfo.DataInterfaceHLSLSymbol);
+		FuncTemplateArgs.Add(TEXT("IDToIndexTableName"), ParamInfo.DataInterfaceHLSLSymbol + IDToIndexTableName);
+		FuncTemplateArgs.Add(TEXT("InputIntBufferName"), ParamInfo.DataInterfaceHLSLSymbol + InputIntBufferName);
+		FuncTemplateArgs.Add(TEXT("AcquireTagRegisterIndexName"), ParamInfo.DataInterfaceHLSLSymbol + AcquireTagRegisterIndexName);
+		FuncTemplateArgs.Add(TEXT("ParticleStrideIntName"), ParamInfo.DataInterfaceHLSLSymbol + ParticleStrideIntName);
 
 		OutHLSL += FString::Format(FuncTemplate, FuncTemplateArgs);
 		return true;
@@ -2244,7 +1997,7 @@ bool UNiagaraDataInterfaceParticleRead::GetFunctionHLSL(const FNiagaraDataInterf
 
 		TMap<FString, FStringFormatArg> FuncTemplateArgs;
 		FuncTemplateArgs.Add(TEXT("FunctionName"), FunctionInfo.InstanceName);
-		FuncTemplateArgs.Add(TEXT("IDToIndexTableName"), IDToIndexTableBaseName + ParamInfo.DataInterfaceHLSLSymbol);
+		FuncTemplateArgs.Add(TEXT("IDToIndexTableName"), ParamInfo.DataInterfaceHLSLSymbol + IDToIndexTableName);
 
 		OutHLSL += FString::Format(FuncTemplate, FuncTemplateArgs);
 		return true;
@@ -2361,6 +2114,210 @@ bool UNiagaraDataInterfaceParticleRead::GetFunctionHLSL(const FNiagaraDataInterf
 }
 #endif
 
+void UNiagaraDataInterfaceParticleRead::BuildShaderParameters(FNiagaraShaderParametersBuilder& ShaderParametersBuilder) const
+{
+	ShaderParametersBuilder.AddNestedStruct<FShaderParameters>();
+
+	const int NumAttributesInt4 = FMath::Max(1, FMath::DivideAndRoundUp(ShaderParametersBuilder.GetGeneratedFunctions().Num(), 4));
+	ShaderParametersBuilder.AddLooseParamArray<FIntVector4>(AttributeIndicesMetadataName, NumAttributesInt4);
+	ShaderParametersBuilder.AddLooseParamArray<FIntVector4>(AttributeCompressedMetadataName, NumAttributesInt4);
+}
+
+void UNiagaraDataInterfaceParticleRead::SetShaderParameters(const FNiagaraDataInterfaceSetShaderParametersContext& Context) const
+{
+	FNiagaraDataInterfaceProxyParticleRead& DIProxy = Context.GetProxy<FNiagaraDataInterfaceProxyParticleRead>();
+
+	const FNiagaraDataInterfaceParametersCS_ParticleRead& ShaderStorage = Context.GetShaderStorage<FNiagaraDataInterfaceParametersCS_ParticleRead>();
+	const int32 NumAttributes = ShaderStorage.AttributeNames.Num();
+	const int32 NumAttributesInt4 = FMath::Max(1, FMath::DivideAndRoundUp(NumAttributes, 4));
+
+	FShaderParameters* ShaderParameters = Context.GetParameterNestedStruct<FShaderParameters>();
+	TArrayView<FIntVector4> AttributeIndices = Context.GetParameterLooseArray<FIntVector4>(NumAttributesInt4);
+	TArrayView<FIntVector4> AttributeCompressed = Context.GetParameterLooseArray<FIntVector4>(NumAttributesInt4);
+
+	auto SetErrorParams =
+		[&](bool bClearSpawnInfo)
+		{
+			if (bClearSpawnInfo)
+			{
+				ShaderParameters->NumSpawnedParticles			= 0;
+				ShaderParameters->SpawnedParticlesAcquireTag	= 0;
+				ShaderParameters->SpawnedIDsBuffer				= FNiagaraRenderer::GetDummyIntBuffer();
+			}
+
+			ShaderParameters->InstanceCountOffset		= INDEX_NONE;
+			ShaderParameters->ParticleStrideFloat		= 0;
+			ShaderParameters->ParticleStrideInt			= 0;
+			ShaderParameters->ParticleStrideHalf		= 0;
+			ShaderParameters->AcquireTagRegisterIndex	= INDEX_NONE;
+			ShaderParameters->IDToIndexTable			= FNiagaraRenderer::GetDummyIntBuffer();
+			ShaderParameters->InputFloatBuffer			= FNiagaraRenderer::GetDummyFloatBuffer();
+			ShaderParameters->InputIntBuffer			= FNiagaraRenderer::GetDummyIntBuffer();
+			ShaderParameters->InputHalfBuffer			= FNiagaraRenderer::GetDummyHalfBuffer();
+
+			FMemory::Memset(AttributeIndices.GetData(), 0xFF, AttributeIndices.GetTypeSize());
+			FMemory::Memset(AttributeCompressed.GetData(), 0x00, AttributeCompressed.GetTypeSize());
+		};
+
+	FNDIParticleRead_RenderInstanceData* InstanceData = DIProxy.GetRenderDataForSystem(Context.GetSystemInstanceID());
+	if (InstanceData == nullptr)
+	{
+		SetErrorParams(true);
+		return;
+	}
+
+	if (InstanceData->SourceEmitterGPUContext == nullptr)
+	{
+		// This means the source emitter isn't running on GPU.
+		if (!InstanceData->bSourceEmitterNotGPUErrorShown)
+		{
+			UE_LOG(LogNiagara, Error, TEXT("GPU particle read DI is set to access CPU emitter '%s'."), *InstanceData->DebugSourceName);
+			InstanceData->bSourceEmitterNotGPUErrorShown = true;
+		}
+		SetErrorParams(true);
+		InstanceData->CachedDataSet = nullptr;
+		InstanceData->bWarnFailedToFindAcquireTag = true;
+		return;
+	}
+	InstanceData->bSourceEmitterNotGPUErrorShown = false;
+
+	FNiagaraDataSet* SourceDataSet = InstanceData->SourceEmitterGPUContext->MainDataSet;
+	if (!SourceDataSet)
+	{
+		SetErrorParams(true);
+		InstanceData->CachedDataSet = nullptr;
+		InstanceData->bWarnFailedToFindAcquireTag = true;
+		return;
+	}
+
+	FNiagaraDataBuffer* SourceData = nullptr;
+	uint32 NumSpawnedInstances = 0;
+	uint32 IDAcquireTag = 0;
+
+	// Test for reading ourself
+	if (Context.GetComputeInstanceData().Context == InstanceData->SourceEmitterGPUContext)
+	{
+		// Pull source data from SimStageData as this has the correct information until all stages are dispatched
+		// Note: In the case of an in place (aka partial update) stage this does not work as we need to bind the
+		//       same buffer as both UAV & SRV.  For the moment the user must disable partial writes.
+		//       Another option would be to read from the UAV, but that would equally race.
+		//-TODO: Automate this in some way or surface as feedback...
+		if (const FSimulationStageMetaData* StageMetaData = Context.GetSimStageData().StageMetaData)
+		{
+			if (StageMetaData->bPartialParticleUpdate)
+			{
+				UE_LOG(LogNiagara, Error, TEXT("Particle read DI reading self '%s' on stage '%s' is unsafe, please disable partial writes on the stage."), *InstanceData->DebugSourceName, *StageMetaData->SimulationStageName.ToString());
+				SetErrorParams(true);
+				return;
+			}
+		}
+
+		SourceData = Context.GetSimStageData().Source;
+
+		// Pull some spawned / acquire tag from destination data as we are reading from ourselves this is what we are currently doing in this stage.
+		if (Context.GetSimStageData().Destination)
+		{
+			ensure(SourceData != nullptr);
+			NumSpawnedInstances = Context.GetSimStageData().Destination->GetNumSpawnedInstances();
+			IDAcquireTag = Context.GetSimStageData().Destination->GetIDAcquireTag();
+		}
+		// In cases where we do not have destination data pull from source (i.e. this is an iteration on none particle stages)
+		else if (SourceData)
+		{
+			NumSpawnedInstances = SourceData->GetNumSpawnedInstances();
+			IDAcquireTag = SourceData->GetIDAcquireTag();
+		}
+	}
+	else
+	{
+		SourceData = SourceDataSet->GetCurrentData();
+		if (SourceData)
+		{
+			NumSpawnedInstances = SourceData->GetNumSpawnedInstances();
+			IDAcquireTag = SourceData->GetIDAcquireTag();
+		}
+	}
+
+	ShaderParameters->NumSpawnedParticles = NumSpawnedInstances;
+	ShaderParameters->SpawnedParticlesAcquireTag = IDAcquireTag;
+	ShaderParameters->SpawnedIDsBuffer = FNiagaraRenderer::GetSrvOrDefaultInt(SourceDataSet->GetGPUFreeIDs());
+
+	if (!SourceData)
+	{
+		SetErrorParams(false);
+		return;
+	}
+
+	if (InstanceData->CachedDataSet != SourceDataSet)
+	{
+		NDIParticleRead_FindAttributeIndices(InstanceData, SourceDataSet, ShaderStorage);
+		InstanceData->CachedDataSet = SourceDataSet;
+		InstanceData->bWarnFailedToFindAcquireTag = true;
+	}
+
+	ShaderParameters->InstanceCountOffset		= SourceData->GetGPUInstanceCountBufferOffset();
+	ShaderParameters->ParticleStrideFloat		= SourceData->GetFloatStride() / sizeof(float);
+	ShaderParameters->ParticleStrideInt			= SourceData->GetInt32Stride() / sizeof(int32);
+	ShaderParameters->ParticleStrideHalf		= SourceData->GetHalfStride() / sizeof(FFloat16);
+	ShaderParameters->AcquireTagRegisterIndex	= InstanceData->AcquireTagRegisterIndex;
+
+	// There's no need to transition the input buffers, because the grouping logic inside FNiagaraGpuComputeDispatch ensures that our source emitter has ran before us,
+	// and its buffers have been transitioned to readable.
+	ShaderParameters->IDToIndexTable	= FNiagaraRenderer::GetSrvOrDefaultInt(SourceData->GetGPUIDToIndexTable());
+	ShaderParameters->InputFloatBuffer	= FNiagaraRenderer::GetSrvOrDefaultFloat(SourceData->GetGPUBufferFloat());
+	ShaderParameters->InputIntBuffer	= FNiagaraRenderer::GetSrvOrDefaultInt(SourceData->GetGPUBufferInt());
+	ShaderParameters->InputHalfBuffer	= FNiagaraRenderer::GetSrvOrDefaultHalf(SourceData->GetGPUBufferHalf());
+
+	if (NumAttributes > 0)
+	{
+		check(InstanceData->AttributeIndices.Num() * InstanceData->AttributeIndices.GetTypeSize() == AttributeIndices.Num() * AttributeIndices.GetTypeSize());
+		check(InstanceData->AttributeCompressed.Num() * InstanceData->AttributeCompressed.GetTypeSize() == AttributeCompressed.Num() * AttributeCompressed.GetTypeSize());
+
+		FMemory::Memcpy(AttributeIndices.GetData(), InstanceData->AttributeIndices.GetData(), AttributeIndices.Num() * AttributeIndices.GetTypeSize());
+		FMemory::Memcpy(AttributeCompressed.GetData(), InstanceData->AttributeCompressed.GetData(), AttributeCompressed.Num() * AttributeCompressed.GetTypeSize());
+	}
+
+	if (InstanceData->bWarnFailedToFindAcquireTag && (InstanceData->AcquireTagRegisterIndex == -1) && Context.IsParameterBound(&ShaderParameters->AcquireTagRegisterIndex) )
+	{
+		InstanceData->bWarnFailedToFindAcquireTag = false;
+		UE_LOG(LogNiagara, Error, TEXT("Particle read DI cannot find ID variable in emitter '%s'."), *InstanceData->DebugSourceName);
+	}
+}
+
+FNiagaraDataInterfaceParametersCS* UNiagaraDataInterfaceParticleRead::CreateShaderStorage(const FNiagaraDataInterfaceGPUParamInfo& ParameterInfo, const FShaderParameterMap& ParameterMap) const
+{
+	FNiagaraDataInterfaceParametersCS_ParticleRead* ShaderStorage = new FNiagaraDataInterfaceParametersCS_ParticleRead();
+
+	const int32 NumFuncs = ParameterInfo.GeneratedFunctions.Num();
+	ShaderStorage->AttributeNames.SetNum(NumFuncs);
+	ShaderStorage->AttributeTypes.SetNum(NumFuncs);
+	for (int32 FuncIdx = 0; FuncIdx < NumFuncs; ++FuncIdx)
+	{
+		const FNiagaraDataInterfaceGeneratedFunction& Func = ParameterInfo.GeneratedFunctions[FuncIdx];
+		static const FName NAME_Attribute("Attribute");
+		const FName* AttributeName = Func.FindSpecifierValue(NAME_Attribute);
+		if (AttributeName != nullptr)
+		{
+			ShaderStorage->AttributeNames[FuncIdx] = *AttributeName;
+			ShaderStorage->AttributeTypes[FuncIdx] = GetValueTypeFromFuncName(Func.DefinitionName);
+		}
+		else
+		{
+			// This is not an error. GetNumSpawnedParticles and GetIDAtSpawnIndexFunctionName don't use specifiers,
+			// but they take up slots in the attribute indices array for simplicity. Just stick NAME_None in here to ignore them.
+			ShaderStorage->AttributeNames[FuncIdx] = NAME_None;
+			ShaderStorage->AttributeTypes[FuncIdx] = ENiagaraParticleDataValueType::Invalid;
+		}
+	}
+
+	return ShaderStorage;
+}
+
+const FTypeLayoutDesc* UNiagaraDataInterfaceParticleRead::GetShaderStorageType() const
+{
+	return &StaticGetTypeLayoutDesc<FNiagaraDataInterfaceParametersCS_ParticleRead>();
+}
+
 void UNiagaraDataInterfaceParticleRead::ProvidePerInstanceDataForRenderThread(void* DataForRenderThread, void* PerInstanceData, const FNiagaraSystemInstanceID& SystemInstance)
 {
 	FNDIParticleRead_GameToRenderData* RTData = new (DataForRenderThread) FNDIParticleRead_GameToRenderData;
@@ -2620,6 +2577,4 @@ bool UNiagaraDataInterfaceParticleRead::ReadsEmitterParticleData(const FString&
 	return EmitterName == InEmitterName;
 }
 
-IMPLEMENT_NIAGARA_DI_PARAMETER(UNiagaraDataInterfaceParticleRead, FNiagaraDataInterfaceParametersCS_ParticleRead);
-
 #undef LOCTEXT_NAMESPACE
diff --git a/Engine/Plugins/FX/Niagara/Source/Niagara/Private/NiagaraDataInterfaceRenderTarget2D.cpp b/Engine/Plugins/FX/Niagara/Source/Niagara/Private/NiagaraDataInterfaceRenderTarget2D.cpp
--- a/Engine/Plugins/FX/Niagara/Source/Niagara/Private/NiagaraDataInterfaceRenderTarget2D.cpp
+++ b/Engine/Plugins/FX/Niagara/Source/Niagara/Private/NiagaraDataInterfaceRenderTarget2D.cpp
@@ -23,10 +23,10 @@
 namespace NDIRenderTarget2DLocal
 {
 	BEGIN_SHADER_PARAMETER_STRUCT(FShaderParameters, )
-		SHADER_PARAMETER(FIntPoint,					TextureSize)
-		SHADER_PARAMETER_UAV(RWTexture2D<float4>,	RWTexture)
-		SHADER_PARAMETER_TEXTURE(Texture2D<float4>, Texture)
-		SHADER_PARAMETER_SAMPLER(SamplerState,		TextureSampler)
+		SHADER_PARAMETER(FIntPoint,								TextureSize)
+		SHADER_PARAMETER_RDG_TEXTURE_UAV(RWTexture2D<float4>,	RWTexture)
+		SHADER_PARAMETER_RDG_TEXTURE_SRV(Texture2D<float4>,		Texture)
+		SHADER_PARAMETER_SAMPLER(SamplerState,					TextureSampler)
 	END_SHADER_PARAMETER_STRUCT()
 
 	static const TCHAR* TemplateShaderFile = TEXT("/Plugin/FX/Niagara/Private/NiagaraDataInterfaceRenderTarget2DTemplate.ush");
@@ -39,8 +39,6 @@ namespace NDIRenderTarget2DLocal
 	static const FName GetSizeFunctionName("GetRenderTargetSize");
 	static const FName LinearToIndexName("LinearToIndex");
 	static const FName ExecToIndexName("ExecToIndex");
-
-	static const FName GenerateMipsName("RenderTarget2D::GenerateMips");
 }
 
 FNiagaraVariableBase UNiagaraDataInterfaceRenderTarget2D::ExposedRTVar;
@@ -505,6 +503,21 @@ void UNiagaraDataInterfaceRenderTarget2D::SetShaderParameters(const FNiagaraData
 	FRenderTarget2DRWInstanceData_RenderThread* InstanceData_RT = DIProxy.SystemInstancesToProxyData_RT.Find(Context.GetSystemInstanceID());
 	check(InstanceData_RT);
 
+	FRDGBuilder& GraphBuilder = Context.GetGraphBuilder();
+
+	// Ensure RDG resources are ready to use
+	if (InstanceData_RT->TransientRDGTexture == nullptr)
+	{
+		InstanceData_RT->TransientRDGTexture = GraphBuilder.FindExternalTexture(InstanceData_RT->TextureRHI);
+		if (InstanceData_RT->TransientRDGTexture == nullptr)
+		{
+			InstanceData_RT->TransientRDGTexture = GraphBuilder.RegisterExternalTexture(CreateRenderTarget(InstanceData_RT->TextureRHI, TEXT("NiagaraRenderTarget2D")));
+		}
+		InstanceData_RT->TransientRDGSRV = GraphBuilder.CreateSRV(InstanceData_RT->TransientRDGTexture);
+		InstanceData_RT->TransientRDGUAV = GraphBuilder.CreateUAV(InstanceData_RT->TransientRDGTexture);
+		Context.GetRDGExternalAccessQueue().Add(InstanceData_RT->TransientRDGTexture);
+	}
+
 	NDIRenderTarget2DLocal::FShaderParameters* Parameters = Context.GetParameterNestedStruct<NDIRenderTarget2DLocal::FShaderParameters>();
 	Parameters->TextureSize = InstanceData_RT->Size;
 
@@ -513,38 +526,38 @@ void UNiagaraDataInterfaceRenderTarget2D::SetShaderParameters(const FNiagaraData
 
 	if (bRTWrite)
 	{
-		Parameters->RWTexture = InstanceData_RT->UnorderedAccessViewRHI;
-		if (Parameters->RWTexture)
+		if (InstanceData_RT->TransientRDGUAV)
 		{
-			// FIXME: this transition needs to happen in FNiagaraDataInterfaceProxyRenderTarget2DProxy::PreStage so it doesn't break up the overlap group,
-			// but for some reason it stops working if I move it in there.
-			Context.GetRHICmdList().Transition(FRHITransitionInfo(Parameters->RWTexture, ERHIAccess::Unknown, ERHIAccess::UAVCompute));
+			Parameters->RWTexture = InstanceData_RT->TransientRDGUAV;
 			InstanceData_RT->bRebuildMips = true;
 			InstanceData_RT->bWroteThisFrame = true;
-			InstanceData_RT->bNeedsTransition = true;
 		}
 		else
 		{
-			Parameters->RWTexture = Context.GetComputeDispatchInterface().GetEmptyUAVFromPool(Context.GetRHICmdList(), EPixelFormat::PF_A16B16G16R16, ENiagaraEmptyUAVType::Texture2D);
+			Parameters->RWTexture = Context.GetComputeDispatchInterface().GetEmptyTextureUAV(GraphBuilder, EPixelFormat::PF_A16B16G16R16, ETextureDimension::Texture2D);
 		}
 	}
 	
 	if (bRTRead)
 	{
-		Parameters->Texture = GBlackTexture->TextureRHI;
-		Parameters->TextureSampler = TStaticSamplerState<SF_Point, AM_Clamp, AM_Clamp, AM_Clamp>::GetRHI();
-		if (ensureMsgf(bRTWrite == false, TEXT("RenderTarget DataInterface is both wrote and read from in the same stage, this is not allowed, read will be invalid")))
+		ensureMsgf(bRTWrite == false, TEXT("RenderTarget DataInterface is both wrote and read from in the same stage, this is not allowed, read will be invalid"));
+		if (bRTWrite == false && InstanceData_RT->TransientRDGSRV)
 		{
 			InstanceData_RT->bReadThisFrame = true;
+			Parameters->Texture = InstanceData_RT->TransientRDGSRV;
+		}
+		else
+		{
+			Parameters->Texture = Context.GetComputeDispatchInterface().GetBlackTextureSRV(GraphBuilder, ETextureDimension::Texture2D);
+		}
 
-			if (InstanceData_RT->TextureRHI)
-			{
-				Parameters->Texture = InstanceData_RT->TextureRHI;
-			}
-			if (InstanceData_RT->SamplerStateRHI)
-			{
-				Parameters->TextureSampler = InstanceData_RT->SamplerStateRHI;
-			}
+		if (InstanceData_RT->SamplerStateRHI)
+		{
+			Parameters->TextureSampler = InstanceData_RT->SamplerStateRHI;
+		}
+		else
+		{
+			Parameters->TextureSampler = TStaticSamplerState<SF_Point, AM_Clamp, AM_Clamp, AM_Clamp>::GetRHI();
 		}
 	}
 }
@@ -591,8 +604,8 @@ void UNiagaraDataInterfaceRenderTarget2D::DestroyPerInstanceData(void* PerInstan
 #if STATS
 			if (FRenderTarget2DRWInstanceData_RenderThread* TargetData = RT_Proxy->SystemInstancesToProxyData_RT.Find(InstanceID))
 			{
-				TargetData->SamplerStateRHI = nullptr;
-				TargetData->TextureRHI = nullptr;
+				TargetData->SamplerStateRHI.SafeRelease();
+				TargetData->TextureRHI.SafeRelease();
 				TargetData->UpdateMemoryStats();
 			}
 #endif
@@ -811,14 +824,12 @@ bool UNiagaraDataInterfaceRenderTarget2D::PerInstanceTickPostSimulate(void* PerI
 			#endif
 				TargetData->SamplerStateRHI.SafeRelease();
 				TargetData->TextureRHI.SafeRelease();
-				TargetData->UnorderedAccessViewRHI.SafeRelease();
 				if (RT_TargetTexture)
 				{
 					if (FTextureRenderTarget2DResource* Resource2D = RT_TargetTexture->GetTextureRenderTarget2DResource())
 					{
 						TargetData->SamplerStateRHI = Resource2D->SamplerStateRHI;
 						TargetData->TextureRHI = Resource2D->GetTextureRHI();
-						TargetData->UnorderedAccessViewRHI = Resource2D->GetUnorderedAccessViewRHI();
 					}
 				}
 			#if STATS
@@ -831,33 +842,22 @@ bool UNiagaraDataInterfaceRenderTarget2D::PerInstanceTickPostSimulate(void* PerI
 	return false;
 }
 
-void FNiagaraDataInterfaceProxyRenderTarget2DProxy::PostStage(FRHICommandList& RHICmdList, const FNiagaraDataInterfaceStageArgs& Context)
+void FNiagaraDataInterfaceProxyRenderTarget2DProxy::PostStage(const FNDIGpuComputePostStageContext& Context)
 {
-	if (FRenderTarget2DRWInstanceData_RenderThread* ProxyData = SystemInstancesToProxyData_RT.Find(Context.SystemInstanceID))
+	if (FRenderTarget2DRWInstanceData_RenderThread* ProxyData = SystemInstancesToProxyData_RT.Find(Context.GetSystemInstanceID()))
 	{
-		if (ProxyData->bNeedsTransition)
-		{
-			ProxyData->bNeedsTransition = false;
-
-			if (FRHIUnorderedAccessView* OutputUAV = ProxyData->UnorderedAccessViewRHI)
-			{
-				// FIXME: move to FNiagaraDataInterfaceProxyRenderTarget2DProxy::PostStage, same as for the transition in Set() above.
-				RHICmdList.Transition(FRHITransitionInfo(OutputUAV, ERHIAccess::UAVCompute, ERHIAccess::SRVMask));
-			}
-		}
-
 		if (ProxyData->bRebuildMips && (ProxyData->MipMapGeneration == ENiagaraMipMapGeneration::PostStage))
 		{
 			ProxyData->bRebuildMips = false;
-			FNiagaraGpuProfileScope GpuProfileScope(RHICmdList, Context, NDIRenderTarget2DLocal::GenerateMipsName);
-			NiagaraGenerateMips::GenerateMips(RHICmdList, ProxyData->TextureRHI, ProxyData->MipMapGenerationType);
+			//-TODO:RDG:GPUProfiler:FNiagaraGpuProfileScope GpuProfileScope(RHICmdList, Context, GNiagaraRenderTarget2DGenerateMipsName);
+			NiagaraGenerateMips::GenerateMips(Context.GetGraphBuilder(), ProxyData->TransientRDGTexture, ProxyData->MipMapGenerationType);
 		}
 	}
 }
 
-void FNiagaraDataInterfaceProxyRenderTarget2DProxy::PostSimulate(FRHICommandList& RHICmdList, const FNiagaraDataInterfaceArgs& Context)
+void FNiagaraDataInterfaceProxyRenderTarget2DProxy::PostSimulate(const FNDIGpuComputePostSimulateContext& Context)
 {
-	FRenderTarget2DRWInstanceData_RenderThread* ProxyData = SystemInstancesToProxyData_RT.Find(Context.SystemInstanceID);
+	FRenderTarget2DRWInstanceData_RenderThread* ProxyData = SystemInstancesToProxyData_RT.Find(Context.GetSystemInstanceID());
 	if (ProxyData == nullptr)
 	{
 		return;
@@ -865,8 +865,8 @@ void FNiagaraDataInterfaceProxyRenderTarget2DProxy::PostSimulate(FRHICommandList
 
 	if (ProxyData->bRebuildMips && (ProxyData->MipMapGeneration == ENiagaraMipMapGeneration::PostSimulate))
 	{
-		FNiagaraGpuProfileScope GpuProfileScope(RHICmdList, Context, NDIRenderTarget2DLocal::GenerateMipsName);
-		NiagaraGenerateMips::GenerateMips(RHICmdList, ProxyData->TextureRHI, ProxyData->MipMapGenerationType);
+		//-TODO:RDG:GPUProfiler:FNiagaraGpuProfileScope GpuProfileScope(RHICmdList, Context, GNiagaraRenderTarget2DGenerateMipsName);
+		NiagaraGenerateMips::GenerateMips(Context.GetGraphBuilder(), ProxyData->TransientRDGTexture, ProxyData->MipMapGenerationType);
 	}
 
 	// We only need to transfer this frame if it was written to.
@@ -874,25 +874,27 @@ void FNiagaraDataInterfaceProxyRenderTarget2DProxy::PostSimulate(FRHICommandList
 	// We also assume the texture is important for rendering, without discovering renderer bindings we don't really know
 	if (ProxyData->bWroteThisFrame)
 	{
-		Context.ComputeDispatchInterface->MultiGPUResourceModified(RHICmdList, ProxyData->TextureRHI, ProxyData->bReadThisFrame, true);
+		//-TODO:RDG:mGPU:Context.ComputeDispatchInterface->MultiGPUResourceModified(RHICmdList, ProxyData->TextureRHI, ProxyData->bReadThisFrame, true);
 	}
 
-	ProxyData->bRebuildMips = false;
-	ProxyData->bReadThisFrame = false;
-	ProxyData->bWroteThisFrame = false;
-
 #if NIAGARA_COMPUTEDEBUG_ENABLED && WITH_EDITORONLY_DATA
-	if (ProxyData->bPreviewTexture)
+	if (ProxyData->bPreviewTexture && ProxyData->TransientRDGTexture)
 	{
-		if (FNiagaraGpuComputeDebug* GpuComputeDebug = Context.ComputeDispatchInterface->GetGpuComputeDebug())
+		if (FNiagaraGpuComputeDebug* GpuComputeDebug = Context.GetComputeDispatchInterface().GetGpuComputeDebug())
 		{
-			if (FRHITexture* RHITexture = ProxyData->TextureRHI)
-			{
-				GpuComputeDebug->AddTexture(RHICmdList, Context.SystemInstanceID, SourceDIName, RHITexture);
-			}
+			GpuComputeDebug->AddTexture(Context.GetGraphBuilder(), Context.GetSystemInstanceID(), SourceDIName, ProxyData->TransientRDGTexture);
 		}
 	}
 #endif
+
+	// Clean up our temporary tracking
+	ProxyData->bRebuildMips = false;
+	ProxyData->bReadThisFrame = false;
+	ProxyData->bWroteThisFrame = false;
+
+	ProxyData->TransientRDGTexture = nullptr;
+	ProxyData->TransientRDGSRV = nullptr;
+	ProxyData->TransientRDGUAV = nullptr;
 }
 
 FIntVector FNiagaraDataInterfaceProxyRenderTarget2DProxy::GetElementCount(FNiagaraSystemInstanceID SystemInstanceID) const
diff --git a/Engine/Plugins/FX/Niagara/Source/Niagara/Private/NiagaraDataInterfaceRenderTarget2DArray.cpp b/Engine/Plugins/FX/Niagara/Source/Niagara/Private/NiagaraDataInterfaceRenderTarget2DArray.cpp
--- a/Engine/Plugins/FX/Niagara/Source/Niagara/Private/NiagaraDataInterfaceRenderTarget2DArray.cpp
+++ b/Engine/Plugins/FX/Niagara/Source/Niagara/Private/NiagaraDataInterfaceRenderTarget2DArray.cpp
@@ -21,10 +21,10 @@
 namespace NDIRenderTarget2DArrayLocal
 {
 	BEGIN_SHADER_PARAMETER_STRUCT(FShaderParameters, )
-		SHADER_PARAMETER(FIntVector3,						TextureSize)
-		SHADER_PARAMETER_UAV(RWTexture2DArray<float4>,		RWTexture)
-		SHADER_PARAMETER_TEXTURE(Texture2DArray<float4>,	Texture)
-		SHADER_PARAMETER_SAMPLER(SamplerState,				TextureSampler)
+		SHADER_PARAMETER(FIntVector3,								TextureSize)
+		SHADER_PARAMETER_RDG_TEXTURE_UAV(RWTexture2DArray<float4>,	RWTexture)
+		SHADER_PARAMETER_RDG_TEXTURE_SRV(Texture2DArray<float4>,	Texture)
+		SHADER_PARAMETER_SAMPLER(SamplerState,						TextureSampler)
 	END_SHADER_PARAMETER_STRUCT()
 
 	static const TCHAR* TemplateShaderFile = TEXT("/Plugin/FX/Niagara/Private/NiagaraDataInterfaceRenderTarget2DArrayTemplate.ush");
@@ -63,9 +63,9 @@ void FRenderTarget2DArrayRWInstanceData_RenderThread::UpdateMemoryStats()
 	DEC_MEMORY_STAT_BY(STAT_NiagaraRenderTargetMemory, MemorySize);
 
 	MemorySize = 0;
-	if (FRHITexture* RHITexture = TextureRHI)
+	if (RenderTarget.IsValid())
 	{
-		MemorySize = RHIComputeMemorySize(RHITexture);
+		MemorySize = RHIComputeMemorySize(RenderTarget->GetRHI());
 	}
 
 	INC_MEMORY_STAT_BY(STAT_NiagaraRenderTargetMemory, MemorySize);
@@ -389,6 +389,17 @@ void UNiagaraDataInterfaceRenderTarget2DArray::SetShaderParameters(const FNiagar
 	FRenderTarget2DArrayRWInstanceData_RenderThread* InstanceData_RT = DIProxy.SystemInstancesToProxyData_RT.Find(Context.GetSystemInstanceID());
 	check(InstanceData_RT);
 
+	FRDGBuilder& GraphBuilder = Context.GetGraphBuilder();
+
+	// Ensure RDG resources are ready to use
+	if (InstanceData_RT->TransientRDGTexture == nullptr)
+	{
+		InstanceData_RT->TransientRDGTexture = GraphBuilder.RegisterExternalTexture(InstanceData_RT->RenderTarget);
+		InstanceData_RT->TransientRDGSRV = GraphBuilder.CreateSRV(InstanceData_RT->TransientRDGTexture);
+		InstanceData_RT->TransientRDGUAV = GraphBuilder.CreateUAV(InstanceData_RT->TransientRDGTexture);
+		Context.GetRDGExternalAccessQueue().Add(InstanceData_RT->TransientRDGTexture);
+	}
+
 	NDIRenderTarget2DArrayLocal::FShaderParameters* Parameters = Context.GetParameterNestedStruct<NDIRenderTarget2DArrayLocal::FShaderParameters>();
 	Parameters->TextureSize = InstanceData_RT->Size;
 
@@ -397,37 +408,37 @@ void UNiagaraDataInterfaceRenderTarget2DArray::SetShaderParameters(const FNiagar
 
 	if (bRTWrite)
 	{
-		Parameters->RWTexture = InstanceData_RT->UnorderedAccessViewRHI;
-		if (Parameters->RWTexture)
+		if (InstanceData_RT->RenderTarget.IsValid())
 		{
-			// FIXME: this transition needs to happen in FNiagaraDataInterfaceProxyRenderTarget2DProxy::PreStage so it doesn't break up the overlap group,
-			// but for some reason it stops working if I move it in there.
-			Context.GetRHICmdList().Transition(FRHITransitionInfo(Parameters->RWTexture, ERHIAccess::Unknown, ERHIAccess::UAVCompute));
+			Parameters->RWTexture = InstanceData_RT->TransientRDGUAV;
 			InstanceData_RT->bWroteThisFrame = true;
-			InstanceData_RT->bNeedsTransition = true;
 		}
 		else
 		{
-			Parameters->RWTexture = Context.GetComputeDispatchInterface().GetEmptyUAVFromPool(Context.GetRHICmdList(), EPixelFormat::PF_A16B16G16R16, ENiagaraEmptyUAVType::Texture2DArray);
+			Parameters->RWTexture = Context.GetComputeDispatchInterface().GetEmptyTextureUAV(GraphBuilder, EPixelFormat::PF_A16B16G16R16, ETextureDimension::Texture2DArray);
 		}
 	}
 	
 	if (bRTRead)
 	{
-		Parameters->Texture = GBlackArrayTexture->TextureRHI;
-		Parameters->TextureSampler = TStaticSamplerState<SF_Point, AM_Clamp, AM_Clamp, AM_Clamp>::GetRHI();
-		if (ensureMsgf(bRTWrite == false, TEXT("RenderTarget DataInterface is both wrote and read from in the same stage, this is not allowed, read will be invalid")))
+		ensureMsgf(bRTWrite == false, TEXT("RenderTarget DataInterface is both wrote and read from in the same stage, this is not allowed, read will be invalid"));
+		if (bRTWrite == false && InstanceData_RT->RenderTarget.IsValid())
 		{
 			InstanceData_RT->bReadThisFrame = true;
+			Parameters->Texture = InstanceData_RT->TransientRDGSRV;
+		}
+		else
+		{
+			Parameters->Texture = Context.GetComputeDispatchInterface().GetBlackTextureSRV(GraphBuilder, ETextureDimension::Texture2DArray);
+		}
 
-			if (InstanceData_RT->TextureRHI)
-			{
-				Parameters->Texture = InstanceData_RT->TextureRHI;
-			}
-			if (InstanceData_RT->SamplerStateRHI)
-			{
-				Parameters->TextureSampler = InstanceData_RT->SamplerStateRHI;
-			}
+		if (InstanceData_RT->SamplerStateRHI)
+		{
+			Parameters->TextureSampler = InstanceData_RT->SamplerStateRHI;
+		}
+		else
+		{
+			Parameters->TextureSampler = TStaticSamplerState<SF_Point, AM_Clamp, AM_Clamp, AM_Clamp>::GetRHI();
 		}
 	}
 }
@@ -471,8 +482,8 @@ void UNiagaraDataInterfaceRenderTarget2DArray::DestroyPerInstanceData(void* PerI
 #if STATS
 			if (FRenderTarget2DArrayRWInstanceData_RenderThread* TargetData = RT_Proxy->SystemInstancesToProxyData_RT.Find(InstanceID))
 			{
-				TargetData->SamplerStateRHI = nullptr;
-				TargetData->TextureRHI = nullptr;
+				TargetData->SamplerStateRHI.SafeRelease();
+				TargetData->RenderTarget.SafeRelease();
 				TargetData->UpdateMemoryStats();
 			}
 #endif
@@ -659,15 +670,13 @@ bool UNiagaraDataInterfaceRenderTarget2DArray::PerInstanceTickPostSimulate(void*
 				TargetData->bPreviewTexture = RT_InstanceData.bPreviewTexture;
 			#endif
 				TargetData->SamplerStateRHI.SafeRelease();
-				TargetData->TextureRHI.SafeRelease();
-				TargetData->UnorderedAccessViewRHI.SafeRelease();
+				TargetData->RenderTarget.SafeRelease();
 				if (RT_TargetTexture)
 				{
 					if (FTextureRenderTarget2DArrayResource* Resource2DArray = RT_TargetTexture->GetTextureRenderTarget2DArrayResource())
 					{
 						TargetData->SamplerStateRHI = Resource2DArray->SamplerStateRHI;
-						TargetData->TextureRHI = Resource2DArray->GetTextureRHI();
-						TargetData->UnorderedAccessViewRHI = Resource2DArray->GetUnorderedAccessViewRHI();
+						TargetData->RenderTarget = CreateRenderTarget(Resource2DArray->GetTextureRHI(), TEXT("NiagaraRenderTarget2DArray"));
 					}
 				}
 			#if STATS
@@ -680,26 +689,9 @@ bool UNiagaraDataInterfaceRenderTarget2DArray::PerInstanceTickPostSimulate(void*
 	return false;
 }
 
-void FNiagaraDataInterfaceProxyRenderTarget2DArrayProxy::PostStage(FRHICommandList& RHICmdList, const FNiagaraDataInterfaceStageArgs& Context)
-{
-	if (FRenderTarget2DArrayRWInstanceData_RenderThread* ProxyData = SystemInstancesToProxyData_RT.Find(Context.SystemInstanceID))
-	{
-		if (ProxyData->bNeedsTransition)
-		{
-			ProxyData->bNeedsTransition = false;
-
-			if (FRHIUnorderedAccessView* OutputUAV = ProxyData->UnorderedAccessViewRHI)
-			{
-				// FIXME: move to FNiagaraDataInterfaceProxyRenderTargetVolumeProxy::PostStage, same as for the transition in Set() above.
-				RHICmdList.Transition(FRHITransitionInfo(OutputUAV, ERHIAccess::UAVCompute, ERHIAccess::SRVMask));
-			}
-		}
-	}
-}
-
-void FNiagaraDataInterfaceProxyRenderTarget2DArrayProxy::PostSimulate(FRHICommandList& RHICmdList, const FNiagaraDataInterfaceArgs& Context)
+void FNiagaraDataInterfaceProxyRenderTarget2DArrayProxy::PostSimulate(const FNDIGpuComputePostSimulateContext& Context)
 {
-	FRenderTarget2DArrayRWInstanceData_RenderThread* ProxyData = SystemInstancesToProxyData_RT.Find(Context.SystemInstanceID);
+	FRenderTarget2DArrayRWInstanceData_RenderThread* ProxyData = SystemInstancesToProxyData_RT.Find(Context.GetSystemInstanceID());
 	if (ProxyData == nullptr)
 	{
 		return;
@@ -710,24 +702,26 @@ void FNiagaraDataInterfaceProxyRenderTarget2DArrayProxy::PostSimulate(FRHIComman
 	// We also assume the texture is important for rendering, without discovering renderer bindings we don't really know
 	if (ProxyData->bWroteThisFrame)
 	{
-		Context.ComputeDispatchInterface->MultiGPUResourceModified(RHICmdList, ProxyData->TextureRHI, ProxyData->bReadThisFrame, true);
+		//-TODO:RDG:mGPU:Context.ComputeDispatchInterface->MultiGPUResourceModified(RHICmdList, ProxyData->TextureRHI, ProxyData->bReadThisFrame, true);
 	}
 
-	ProxyData->bReadThisFrame = false;
-	ProxyData->bWroteThisFrame = false;
-
 #if NIAGARA_COMPUTEDEBUG_ENABLED && WITH_EDITORONLY_DATA
-	if (ProxyData->bPreviewTexture)
+	if (ProxyData->bPreviewTexture && ProxyData->TransientRDGTexture)
 	{
-		if (FNiagaraGpuComputeDebug* GpuComputeDebug = Context.ComputeDispatchInterface->GetGpuComputeDebug())
+		if (FNiagaraGpuComputeDebug* GpuComputeDebug = Context.GetComputeDispatchInterface().GetGpuComputeDebug())
 		{
-			if (FRHITexture* RHITexture = ProxyData->TextureRHI)
-			{
-				GpuComputeDebug->AddTexture(RHICmdList, Context.SystemInstanceID, SourceDIName, RHITexture);
-			}
+			GpuComputeDebug->AddTexture(Context.GetGraphBuilder(), Context.GetSystemInstanceID(), SourceDIName, ProxyData->TransientRDGTexture);
 		}
 	}
 #endif
+
+	// Clean up our temporary tracking
+	ProxyData->bReadThisFrame = false;
+	ProxyData->bWroteThisFrame = false;
+
+	ProxyData->TransientRDGTexture = nullptr;
+	ProxyData->TransientRDGSRV = nullptr;
+	ProxyData->TransientRDGUAV = nullptr;
 }
 
 FIntVector FNiagaraDataInterfaceProxyRenderTarget2DArrayProxy::GetElementCount(FNiagaraSystemInstanceID SystemInstanceID) const
diff --git a/Engine/Plugins/FX/Niagara/Source/Niagara/Private/NiagaraDataInterfaceRenderTargetCube.cpp b/Engine/Plugins/FX/Niagara/Source/Niagara/Private/NiagaraDataInterfaceRenderTargetCube.cpp
--- a/Engine/Plugins/FX/Niagara/Source/Niagara/Private/NiagaraDataInterfaceRenderTargetCube.cpp
+++ b/Engine/Plugins/FX/Niagara/Source/Niagara/Private/NiagaraDataInterfaceRenderTargetCube.cpp
@@ -20,8 +20,8 @@ namespace NDIRenderTargetCubeLocal
 {
 	BEGIN_SHADER_PARAMETER_STRUCT(FShaderParameters, )
 		SHADER_PARAMETER(int, TextureSize)
-		SHADER_PARAMETER_UAV(RWTextureCube<float4>, RWTexture)
-		SHADER_PARAMETER_TEXTURE(TextureCube<float4>, Texture)
+		SHADER_PARAMETER_RDG_TEXTURE_UAV(RWTextureCube<float4>, RWTexture)
+		SHADER_PARAMETER_RDG_TEXTURE_SRV(TextureCube<float4>, Texture)
 		SHADER_PARAMETER_SAMPLER(SamplerState, TextureSampler)
 	END_SHADER_PARAMETER_STRUCT()
 
@@ -59,9 +59,9 @@ void FRenderTargetCubeRWInstanceData_RenderThread::UpdateMemoryStats()
 	DEC_MEMORY_STAT_BY(STAT_NiagaraRenderTargetMemory, MemorySize);
 
 	MemorySize = 0;
-	if (FRHITexture* RHITexture = TextureRHI)
+	if (RenderTarget.IsValid())
 	{
-		MemorySize = RHIComputeMemorySize(RHITexture);
+		MemorySize = RHIComputeMemorySize(RenderTarget->GetRHI());
 	}
 
 	INC_MEMORY_STAT_BY(STAT_NiagaraRenderTargetMemory, MemorySize);
@@ -377,6 +377,17 @@ void UNiagaraDataInterfaceRenderTargetCube::SetShaderParameters(const FNiagaraDa
 	FRenderTargetCubeRWInstanceData_RenderThread* InstanceData_RT = DIProxy.SystemInstancesToProxyData_RT.Find(Context.GetSystemInstanceID());
 	check(InstanceData_RT);
 
+	FRDGBuilder& GraphBuilder = Context.GetGraphBuilder();
+
+	// Ensure RDG resources are ready to use
+	if (InstanceData_RT->TransientRDGTexture == nullptr)
+	{
+		InstanceData_RT->TransientRDGTexture = GraphBuilder.RegisterExternalTexture(InstanceData_RT->RenderTarget);
+		InstanceData_RT->TransientRDGSRV = GraphBuilder.CreateSRV(InstanceData_RT->TransientRDGTexture);
+		InstanceData_RT->TransientRDGUAV = GraphBuilder.CreateUAV(InstanceData_RT->TransientRDGTexture);
+		Context.GetRDGExternalAccessQueue().Add(InstanceData_RT->TransientRDGTexture);
+	}
+
 	NDIRenderTargetCubeLocal::FShaderParameters* Parameters = Context.GetParameterNestedStruct<NDIRenderTargetCubeLocal::FShaderParameters>();
 	Parameters->TextureSize = InstanceData_RT->Size;
 
@@ -385,38 +396,37 @@ void UNiagaraDataInterfaceRenderTargetCube::SetShaderParameters(const FNiagaraDa
 
 	if (bRTWrite)
 	{
-		Parameters->RWTexture = InstanceData_RT->UnorderedAccessViewRHI;
-		if (Parameters->RWTexture)
+		if (InstanceData_RT->RenderTarget.IsValid())
 		{
-			// FIXME: this transition needs to happen in FNiagaraDataInterfaceProxyRenderTargetCubeProxy::PreStage so it doesn't break up the overlap group,
-			// but for some reason it stops working if I move it in there.
-			Context.GetRHICmdList().Transition(FRHITransitionInfo(Parameters->RWTexture, ERHIAccess::Unknown, ERHIAccess::UAVCompute));
-			//InstanceData_RT->bRebuildMips = true;
+			Parameters->RWTexture = InstanceData_RT->TransientRDGUAV;
 			InstanceData_RT->bWroteThisFrame = true;
-			InstanceData_RT->bNeedsTransition = true;
 		}
 		else
 		{
-			Parameters->RWTexture = Context.GetComputeDispatchInterface().GetEmptyUAVFromPool(Context.GetRHICmdList(), EPixelFormat::PF_A16B16G16R16, ENiagaraEmptyUAVType::TextureCube);
+			Parameters->RWTexture = Context.GetComputeDispatchInterface().GetEmptyTextureUAV(GraphBuilder, EPixelFormat::PF_A16B16G16R16, ETextureDimension::TextureCube);
 		}
 	}
 
 	if (bRTRead)
 	{
-		Parameters->Texture = GBlackTextureCube->TextureRHI;
-		Parameters->TextureSampler = TStaticSamplerState<SF_Point, AM_Clamp, AM_Clamp, AM_Clamp>::GetRHI();
-		if (ensureMsgf(bRTWrite == false, TEXT("RenderTarget DataInterface is both wrote and read from in the same stage, this is not allowed, read will be invalid")))
+		ensureMsgf(bRTWrite == false, TEXT("RenderTarget DataInterface is both wrote and read from in the same stage, this is not allowed, read will be invalid"));
+		if (bRTWrite == false && InstanceData_RT->RenderTarget.IsValid())
 		{
 			InstanceData_RT->bReadThisFrame = true;
+			Parameters->Texture = InstanceData_RT->TransientRDGSRV;
+		}
+		else
+		{
+			Parameters->Texture = Context.GetComputeDispatchInterface().GetBlackTextureSRV(GraphBuilder, ETextureDimension::TextureCube);
+		}
 
-			if (InstanceData_RT->TextureRHI)
-			{
-				Parameters->Texture = InstanceData_RT->TextureRHI;
-			}
-			if (InstanceData_RT->SamplerStateRHI)
-			{
-				Parameters->TextureSampler = InstanceData_RT->SamplerStateRHI;
-			}
+		if (InstanceData_RT->SamplerStateRHI)
+		{
+			Parameters->TextureSampler = InstanceData_RT->SamplerStateRHI;
+		}
+		else
+		{
+			Parameters->TextureSampler = TStaticSamplerState<SF_Point, AM_Clamp, AM_Clamp, AM_Clamp>::GetRHI();
 		}
 	}
 }
@@ -458,8 +468,8 @@ void UNiagaraDataInterfaceRenderTargetCube::DestroyPerInstanceData(void* PerInst
 #if STATS
 			if (FRenderTargetCubeRWInstanceData_RenderThread* TargetData = RT_Proxy->SystemInstancesToProxyData_RT.Find(InstanceID))
 			{
-				TargetData->SamplerStateRHI = nullptr;
-				TargetData->TextureRHI = nullptr;
+				TargetData->SamplerStateRHI.SafeRelease();
+				TargetData->RenderTarget.SafeRelease();
 				TargetData->UpdateMemoryStats();
 			}
 #endif
@@ -635,15 +645,13 @@ bool UNiagaraDataInterfaceRenderTargetCube::PerInstanceTickPostSimulate(void* Pe
 				TargetData->bPreviewTexture = RT_InstanceData.bPreviewTexture;
 			#endif
 				TargetData->SamplerStateRHI.SafeRelease();
-				TargetData->TextureRHI.SafeRelease();
-				TargetData->UnorderedAccessViewRHI.SafeRelease();
+				TargetData->RenderTarget.SafeRelease();
 				if (RT_TargetTexture)
 				{
 					if (FTextureRenderTargetCubeResource* ResourceCube = RT_TargetTexture->GetTextureRenderTargetCubeResource())
 					{
 						TargetData->SamplerStateRHI = ResourceCube->SamplerStateRHI;
-						TargetData->TextureRHI = ResourceCube->GetTextureRHI();
-						TargetData->UnorderedAccessViewRHI = ResourceCube->GetUnorderedAccessViewRHI();
+						TargetData->RenderTarget = CreateRenderTarget(ResourceCube->GetTextureRHI(), TEXT("NiagaraRenderTargetCube"));
 					}
 				}
 			#if STATS
@@ -656,54 +664,40 @@ bool UNiagaraDataInterfaceRenderTargetCube::PerInstanceTickPostSimulate(void* Pe
 	return false;
 }
 
-void FNiagaraDataInterfaceProxyRenderTargetCubeProxy::PostStage(FRHICommandList& RHICmdList, const FNiagaraDataInterfaceStageArgs& Context)
+void FNiagaraDataInterfaceProxyRenderTargetCubeProxy::PostSimulate(const FNDIGpuComputePostSimulateContext& Context)
 {
-	if (FRenderTargetCubeRWInstanceData_RenderThread* ProxyData = SystemInstancesToProxyData_RT.Find(Context.SystemInstanceID))
-	{
-		if (ProxyData->bNeedsTransition)
-		{
-			ProxyData->bNeedsTransition = false;
-
-			if (FRHIUnorderedAccessView* OutputUAV = ProxyData->UnorderedAccessViewRHI)
-			{
-				// FIXME: move to FNiagaraDataInterfaceProxyRenderTargetCubeProxy::PostStage, same as for the transition in Set() above.
-				RHICmdList.Transition(FRHITransitionInfo(OutputUAV, ERHIAccess::UAVCompute, ERHIAccess::SRVMask));
-			}
-		}
-	}
-}
-
-void FNiagaraDataInterfaceProxyRenderTargetCubeProxy::PostSimulate(FRHICommandList& RHICmdList, const FNiagaraDataInterfaceArgs& Context)
-{
-	FRenderTargetCubeRWInstanceData_RenderThread* ProxyData = SystemInstancesToProxyData_RT.Find(Context.SystemInstanceID);
+	FRenderTargetCubeRWInstanceData_RenderThread* ProxyData = SystemInstancesToProxyData_RT.Find(Context.GetSystemInstanceID());
 	if (ProxyData == nullptr)
 	{
 		return;
 	}
 
+	
 	// We only need to transfer this frame if it was written to.
 	// If also read then we need to notify that the texture is important for the simulation
 	// We also assume the texture is important for rendering, without discovering renderer bindings we don't really know
 	if (ProxyData->bWroteThisFrame)
 	{
-		Context.ComputeDispatchInterface->MultiGPUResourceModified(RHICmdList, ProxyData->TextureRHI, ProxyData->bReadThisFrame, true);
+		//-TODO:RDG:mGPU:Context.ComputeDispatchInterface->MultiGPUResourceModified(RHICmdList, ProxyData->TextureRHI, ProxyData->bReadThisFrame, true);
 	}
 
-	ProxyData->bReadThisFrame = false;
-	ProxyData->bWroteThisFrame = false;
-
 #if NIAGARA_COMPUTEDEBUG_ENABLED && WITH_EDITORONLY_DATA
-	if (ProxyData->bPreviewTexture && ProxyData->TextureRHI.IsValid())
+	if (ProxyData->bPreviewTexture && ProxyData->TransientRDGTexture)
 	{
-		if (FNiagaraGpuComputeDebug* GpuComputeDebug = Context.ComputeDispatchInterface->GetGpuComputeDebug())
+		if (FNiagaraGpuComputeDebug* GpuComputeDebug = Context.GetComputeDispatchInterface().GetGpuComputeDebug())
 		{
-			if ( FRHITexture* RHITexture = ProxyData->TextureRHI )
-			{
-				GpuComputeDebug->AddTexture(RHICmdList, Context.SystemInstanceID, SourceDIName, RHITexture);
-			}
+			GpuComputeDebug->AddTexture(Context.GetGraphBuilder(), Context.GetSystemInstanceID(), SourceDIName, ProxyData->TransientRDGTexture);
 		}
 	}
 #endif
+
+	// Clean up our temporary tracking
+	ProxyData->bReadThisFrame = false;
+	ProxyData->bWroteThisFrame = false;
+
+	ProxyData->TransientRDGTexture = nullptr;
+	ProxyData->TransientRDGSRV = nullptr;
+	ProxyData->TransientRDGUAV = nullptr;
 }
 
 FIntVector FNiagaraDataInterfaceProxyRenderTargetCubeProxy::GetElementCount(FNiagaraSystemInstanceID SystemInstanceID) const
diff --git a/Engine/Plugins/FX/Niagara/Source/Niagara/Private/NiagaraDataInterfaceRenderTargetVolume.cpp b/Engine/Plugins/FX/Niagara/Source/Niagara/Private/NiagaraDataInterfaceRenderTargetVolume.cpp
--- a/Engine/Plugins/FX/Niagara/Source/Niagara/Private/NiagaraDataInterfaceRenderTargetVolume.cpp
+++ b/Engine/Plugins/FX/Niagara/Source/Niagara/Private/NiagaraDataInterfaceRenderTargetVolume.cpp
@@ -21,8 +21,8 @@ namespace NDIRenderTargetVolumeLocal
 {
 	BEGIN_SHADER_PARAMETER_STRUCT(FShaderParameters, )
 		SHADER_PARAMETER(FIntVector3, TextureSize)
-		SHADER_PARAMETER_UAV(RWTexture3D<float4>, RWTexture)
-		SHADER_PARAMETER_TEXTURE(Texture3D<float4>, Texture)
+		SHADER_PARAMETER_RDG_TEXTURE_UAV(RWTexture3D<float4>, RWTexture)
+		SHADER_PARAMETER_RDG_TEXTURE_SRV(Texture3D<float4>, Texture)
 		SHADER_PARAMETER_SAMPLER(SamplerState, TextureSampler)
 	END_SHADER_PARAMETER_STRUCT()
 
@@ -60,9 +60,9 @@ void FRenderTargetVolumeRWInstanceData_RenderThread::UpdateMemoryStats()
 	DEC_MEMORY_STAT_BY(STAT_NiagaraRenderTargetMemory, MemorySize);
 
 	MemorySize = 0;
-	if (FRHITexture* RHITexture = TextureRHI)
+	if (RenderTarget.IsValid())
 	{
-		MemorySize = RHIComputeMemorySize(RHITexture);
+		MemorySize = RHIComputeMemorySize(RenderTarget->GetRHI());
 	}
 
 	INC_MEMORY_STAT_BY(STAT_NiagaraRenderTargetMemory, MemorySize);
@@ -381,6 +381,17 @@ void UNiagaraDataInterfaceRenderTargetVolume::SetShaderParameters(const FNiagara
 	FRenderTargetVolumeRWInstanceData_RenderThread* InstanceData_RT = DIProxy.SystemInstancesToProxyData_RT.Find(Context.GetSystemInstanceID());
 	check(InstanceData_RT);
 
+	FRDGBuilder& GraphBuilder = Context.GetGraphBuilder();
+
+	// Ensure RDG resources are ready to use
+	if (InstanceData_RT->TransientRDGTexture == nullptr)
+	{
+		InstanceData_RT->TransientRDGTexture = GraphBuilder.RegisterExternalTexture(InstanceData_RT->RenderTarget);
+		InstanceData_RT->TransientRDGSRV = GraphBuilder.CreateSRV(InstanceData_RT->TransientRDGTexture);
+		InstanceData_RT->TransientRDGUAV = GraphBuilder.CreateUAV(InstanceData_RT->TransientRDGTexture);
+		Context.GetRDGExternalAccessQueue().Add(InstanceData_RT->TransientRDGTexture);
+	}
+
 	NDIRenderTargetVolumeLocal::FShaderParameters* Parameters = Context.GetParameterNestedStruct<NDIRenderTargetVolumeLocal::FShaderParameters>();
 	Parameters->TextureSize = InstanceData_RT->Size;
 
@@ -389,38 +400,37 @@ void UNiagaraDataInterfaceRenderTargetVolume::SetShaderParameters(const FNiagara
 
 	if (bRTWrite)
 	{
-		Parameters->RWTexture = InstanceData_RT->UnorderedAccessViewRHI;
-		if (Parameters->RWTexture)
+		if (InstanceData_RT->RenderTarget.IsValid())
 		{
-			// FIXME: this transition needs to happen in FNiagaraDataInterfaceProxyRenderTargetVolumeProxy::PreStage so it doesn't break up the overlap group,
-			// but for some reason it stops working if I move it in there.
-			Context.GetRHICmdList().Transition(FRHITransitionInfo(Parameters->RWTexture, ERHIAccess::Unknown, ERHIAccess::UAVCompute));
-			//InstanceData_RT->bRebuildMips = true;
+			Parameters->RWTexture = InstanceData_RT->TransientRDGUAV;
 			InstanceData_RT->bWroteThisFrame = true;
-			InstanceData_RT->bNeedsTransition = true;
 		}
 		else
 		{
-			Parameters->RWTexture = Context.GetComputeDispatchInterface().GetEmptyUAVFromPool(Context.GetRHICmdList(), EPixelFormat::PF_A16B16G16R16, ENiagaraEmptyUAVType::Texture3D);
+			Parameters->RWTexture = Context.GetComputeDispatchInterface().GetEmptyTextureUAV(GraphBuilder, EPixelFormat::PF_A16B16G16R16, ETextureDimension::Texture3D);
 		}
 	}
 
 	if (bRTRead)
 	{
-		Parameters->Texture = GBlackVolumeTexture->TextureRHI;
-		Parameters->TextureSampler = TStaticSamplerState<SF_Point, AM_Clamp, AM_Clamp, AM_Clamp>::GetRHI();
-		if (ensureMsgf(bRTWrite == false, TEXT("RenderTarget DataInterface is both wrote and read from in the same stage, this is not allowed, read will be invalid")))
+		ensureMsgf(bRTWrite == false, TEXT("RenderTarget DataInterface is both wrote and read from in the same stage, this is not allowed, read will be invalid"));
+		if (bRTWrite == false && InstanceData_RT->RenderTarget.IsValid())
 		{
 			InstanceData_RT->bReadThisFrame = true;
+			Parameters->Texture = InstanceData_RT->TransientRDGSRV;
+		}
+		else
+		{
+			Parameters->Texture = Context.GetComputeDispatchInterface().GetBlackTextureSRV(GraphBuilder, ETextureDimension::Texture3D);
+		}
 
-			if (InstanceData_RT->TextureRHI)
-			{
-				Parameters->Texture = InstanceData_RT->TextureRHI;
-			}
-			if (InstanceData_RT->SamplerStateRHI)
-			{
-				Parameters->TextureSampler = InstanceData_RT->SamplerStateRHI;
-			}
+		if (InstanceData_RT->SamplerStateRHI)
+		{
+			Parameters->TextureSampler = InstanceData_RT->SamplerStateRHI;
+		}
+		else
+		{
+			Parameters->TextureSampler = TStaticSamplerState<SF_Point, AM_Clamp, AM_Clamp, AM_Clamp>::GetRHI();
 		}
 	}
 }
@@ -464,8 +474,8 @@ void UNiagaraDataInterfaceRenderTargetVolume::DestroyPerInstanceData(void* PerIn
 #if STATS
 			if (FRenderTargetVolumeRWInstanceData_RenderThread* TargetData = RT_Proxy->SystemInstancesToProxyData_RT.Find(InstanceID))
 			{
-				TargetData->SamplerStateRHI = nullptr;
-				TargetData->TextureRHI = nullptr;
+				TargetData->SamplerStateRHI.SafeRelease();
+				TargetData->RenderTarget.SafeRelease();
 				TargetData->UpdateMemoryStats();
 			}
 #endif
@@ -654,15 +664,13 @@ bool UNiagaraDataInterfaceRenderTargetVolume::PerInstanceTickPostSimulate(void*
 				TargetData->bPreviewTexture = RT_InstanceData.bPreviewTexture;
 			#endif
 				TargetData->SamplerStateRHI.SafeRelease();
-				TargetData->TextureRHI.SafeRelease();
-				TargetData->UnorderedAccessViewRHI.SafeRelease();
+				TargetData->RenderTarget.SafeRelease();
 				if (RT_TargetTexture)
 				{
 					if (FTextureRenderTargetVolumeResource* ResourceVolume = RT_TargetTexture->GetTextureRenderTargetVolumeResource())
 					{
 						TargetData->SamplerStateRHI = ResourceVolume->SamplerStateRHI;
-						TargetData->TextureRHI = ResourceVolume->GetTextureRHI();
-						TargetData->UnorderedAccessViewRHI = ResourceVolume->GetUnorderedAccessViewRHI();
+						TargetData->RenderTarget = CreateRenderTarget(ResourceVolume->GetTextureRHI(), TEXT("NiagaraRenderTargetVolume"));
 					}
 				}
 			#if STATS
@@ -675,26 +683,9 @@ bool UNiagaraDataInterfaceRenderTargetVolume::PerInstanceTickPostSimulate(void*
 	return false;
 }
 
-void FNiagaraDataInterfaceProxyRenderTargetVolumeProxy::PostStage(FRHICommandList& RHICmdList, const FNiagaraDataInterfaceStageArgs& Context)
-{
-	if (FRenderTargetVolumeRWInstanceData_RenderThread* ProxyData = SystemInstancesToProxyData_RT.Find(Context.SystemInstanceID))
-	{
-		if (ProxyData->bNeedsTransition)
-		{
-			ProxyData->bNeedsTransition = false;
-
-			if (FRHIUnorderedAccessView* OutputUAV = ProxyData->UnorderedAccessViewRHI)
-			{
-				// FIXME: move to FNiagaraDataInterfaceProxyRenderTargetVolumeProxy::PostStage, same as for the transition in Set() above.
-				RHICmdList.Transition(FRHITransitionInfo(OutputUAV, ERHIAccess::UAVCompute, ERHIAccess::SRVMask));
-			}
-		}
-	}
-}
-
-void FNiagaraDataInterfaceProxyRenderTargetVolumeProxy::PostSimulate(FRHICommandList& RHICmdList, const FNiagaraDataInterfaceArgs& Context)
+void FNiagaraDataInterfaceProxyRenderTargetVolumeProxy::PostSimulate(const FNDIGpuComputePostSimulateContext& Context)
 {
-	FRenderTargetVolumeRWInstanceData_RenderThread* ProxyData = SystemInstancesToProxyData_RT.Find(Context.SystemInstanceID);
+	FRenderTargetVolumeRWInstanceData_RenderThread* ProxyData = SystemInstancesToProxyData_RT.Find(Context.GetSystemInstanceID());
 	if (ProxyData == nullptr)
 	{
 		return;
@@ -705,24 +696,26 @@ void FNiagaraDataInterfaceProxyRenderTargetVolumeProxy::PostSimulate(FRHICommand
 	// We also assume the texture is important for rendering, without discovering renderer bindings we don't really know
 	if (ProxyData->bWroteThisFrame)
 	{
-		Context.ComputeDispatchInterface->MultiGPUResourceModified(RHICmdList, ProxyData->TextureRHI, ProxyData->bReadThisFrame, true);
+		//-TODO:RDG:mGPU:Context.ComputeDispatchInterface->MultiGPUResourceModified(RHICmdList, ProxyData->TextureRHI, ProxyData->bReadThisFrame, true);
 	}
 
-	ProxyData->bReadThisFrame = false;
-	ProxyData->bWroteThisFrame = false;
-
 #if NIAGARA_COMPUTEDEBUG_ENABLED && WITH_EDITORONLY_DATA
-	if (ProxyData->bPreviewTexture && ProxyData->TextureRHI.IsValid())
+	if (ProxyData->bPreviewTexture && ProxyData->TransientRDGTexture)
 	{
-		if (FNiagaraGpuComputeDebug* GpuComputeDebug = Context.ComputeDispatchInterface->GetGpuComputeDebug())
+		if (FNiagaraGpuComputeDebug* GpuComputeDebug = Context.GetComputeDispatchInterface().GetGpuComputeDebug())
 		{
-			if ( FRHITexture* RHITexture = ProxyData->TextureRHI )
-			{
-				GpuComputeDebug->AddTexture(RHICmdList, Context.SystemInstanceID, SourceDIName, RHITexture);
-			}
+			GpuComputeDebug->AddTexture(Context.GetGraphBuilder(), Context.GetSystemInstanceID(), SourceDIName, ProxyData->TransientRDGTexture);
 		}
 	}
 #endif
+
+	// Clean up our temporary tracking
+	ProxyData->bReadThisFrame = false;
+	ProxyData->bWroteThisFrame = false;
+
+	ProxyData->TransientRDGTexture = nullptr;
+	ProxyData->TransientRDGSRV = nullptr;
+	ProxyData->TransientRDGUAV = nullptr;
 }
 
 FIntVector FNiagaraDataInterfaceProxyRenderTargetVolumeProxy::GetElementCount(FNiagaraSystemInstanceID SystemInstanceID) const
diff --git a/Engine/Plugins/FX/Niagara/Source/Niagara/Private/NiagaraDataInterfaceRigidMeshCollisionQuery.cpp b/Engine/Plugins/FX/Niagara/Source/Niagara/Private/NiagaraDataInterfaceRigidMeshCollisionQuery.cpp
--- a/Engine/Plugins/FX/Niagara/Source/Niagara/Private/NiagaraDataInterfaceRigidMeshCollisionQuery.cpp
+++ b/Engine/Plugins/FX/Niagara/Source/Niagara/Private/NiagaraDataInterfaceRigidMeshCollisionQuery.cpp
@@ -683,9 +683,6 @@ struct FNDIRigidMeshCollisionProxy : public FNiagaraDataInterfaceProxy
 	/** Launch all pre stage functions */
 	virtual void PreStage(FRHICommandList& RHICmdList, const FNiagaraDataInterfaceStageArgs& Context) override;
 
-	/** Reset the buffers  */
-	virtual void ResetData(FRHICommandList& RHICmdList, const FNiagaraDataInterfaceArgs& Context) override;
-
 	/** List of proxy data for each system instances*/
 	TMap<FNiagaraSystemInstanceID, FGameThreadData> SystemInstancesToProxyData;
 };
@@ -792,9 +789,6 @@ void FNDIRigidMeshCollisionProxy::PreStage(FRHICommandList& RHICmdList, const FN
 	}
 }
 
-void FNDIRigidMeshCollisionProxy::ResetData(FRHICommandList& RHICmdList, const FNiagaraDataInterfaceArgs& Context)
-{}
-
 //------------------------------------------------------------------------------------------------------------
 
 struct FNDIRigidMeshCollisionParametersCS : public FNiagaraDataInterfaceParametersCS
