- Niagara Compute Dipsatch modified to use FRDGBuilder through the pipeline

diff --git a/Engine/Plugins/FX/Niagara/Source/Niagara/Private/NiagaraDataInterfaceTexture.cpp b/Engine/Plugins/FX/Niagara/Source/Niagara/Private/NiagaraDataInterfaceTexture.cpp
--- a/Engine/Plugins/FX/Niagara/Source/Niagara/Private/NiagaraDataInterfaceTexture.cpp
+++ b/Engine/Plugins/FX/Niagara/Source/Niagara/Private/NiagaraDataInterfaceTexture.cpp
@@ -3,12 +3,11 @@
 #include "NiagaraDataInterfaceTexture.h"
 #include "NiagaraComputeExecutionContext.h"
 #include "NiagaraCustomVersion.h"
+#include "NiagaraGpuComputeDispatchInterface.h"
 #include "NiagaraShader.h"
 #include "NiagaraShaderParametersBuilder.h"
 #include "NiagaraSystemInstance.h"
 
-#include "ShaderParameterUtils.h"
-
 #define LOCTEXT_NAMESPACE "UNiagaraDataInterfaceTexture"
 
 const TCHAR* UNiagaraDataInterfaceTexture::TemplateShaderFilePath = TEXT("/Plugin/FX/Niagara/Private/NiagaraDataInterfaceTextureTemplate.ush");
@@ -27,42 +26,23 @@ struct FNDITextureInstanceData_RenderThread
 {
 	FSamplerStateRHIRef		SamplerStateRHI;
 	FTextureReferenceRHIRef	TextureReferenceRHI;
-	FTextureRHIRef			ResolvedTextureRHI;
 	FVector2f				TextureSize;
+
+	FRDGTextureRef			TransientRDGTexture = nullptr;
 };
 
 struct FNiagaraDataInterfaceProxyTexture : public FNiagaraDataInterfaceProxy
 {
-	virtual void ConsumePerInstanceDataFromGameThread(void* PerInstanceData, const FNiagaraSystemInstanceID& Instance) override
-	{
-		checkNoEntry();
-	}
+	virtual void ConsumePerInstanceDataFromGameThread(void* PerInstanceData, const FNiagaraSystemInstanceID& Instance) override { checkNoEntry(); }
+	virtual int32 PerInstanceDataPassedToRenderThreadSize() const override { return 0; }
 
-	virtual int32 PerInstanceDataPassedToRenderThreadSize() const override
+	virtual void PostSimulate(const FNDIGpuComputePostSimulateContext& Context) override
 	{
-		return 0;
-	}
-
-	virtual void PreStage(FRHICommandList& RHICmdList, const FNiagaraDataInterfaceStageArgs& Context) override
-	{
-		if ( FNDITextureInstanceData_RenderThread* InstanceData = InstanceData_RT.Find(Context.SystemInstanceID) )
+		if (Context.IsFinalPostSimulate())
 		{
-			// Because the underlying reference can have a switch in flight on the RHI we get the referenced texture
-			// here, ensure it's valid (as it could be queued for delete) and cache until next round.  If we were
-			// to release the reference in PostStage / PostSimulate we still stand a chance the the transition we
-			// queue will be invalid by the time it is processed on the RHI thread.
-			if (Context.SimStageData->bFirstStage && InstanceData->TextureReferenceRHI.IsValid())
-			{
-				InstanceData->ResolvedTextureRHI = InstanceData->TextureReferenceRHI->GetReferencedTexture();
-				if (InstanceData->ResolvedTextureRHI && !InstanceData->ResolvedTextureRHI->IsValid())
-				{
-					InstanceData->ResolvedTextureRHI = nullptr;
-				}
-			}
-			if (InstanceData->ResolvedTextureRHI.IsValid())
+			if (FNDITextureInstanceData_RenderThread* InstanceData = InstanceData_RT.Find(Context.GetSystemInstanceID()))
 			{
-				// Make sure the texture is readable, we don't know where it's coming from.
-				RHICmdList.Transition(FRHITransitionInfo(InstanceData->ResolvedTextureRHI, ERHIAccess::Unknown, ERHIAccess::SRVMask));
+				InstanceData->TransientRDGTexture = nullptr;
 			}
 		}
 	}
@@ -263,14 +243,15 @@ bool UNiagaraDataInterfaceTexture::PerInstanceTick(void* PerInstanceData, FNiaga
 				if (RT_Texture)
 				{
 					InstanceData.TextureReferenceRHI = RT_Texture->TextureReference.TextureReferenceRHI;
-					InstanceData.SamplerStateRHI = RT_Texture->GetResource() ? RT_Texture->GetResource()->SamplerStateRHI : nullptr;
+					InstanceData.SamplerStateRHI = RT_Texture->GetResource() ? RT_Texture->GetResource()->SamplerStateRHI.GetReference() : TStaticSamplerState<SF_Point>::GetRHI();
+					InstanceData.TextureSize = FVector2f(RT_TextureSize.X, RT_TextureSize.Y);
 				}
 				else
 				{
 					InstanceData.TextureReferenceRHI = nullptr;
 					InstanceData.SamplerStateRHI = nullptr;
+					InstanceData.TextureSize = FVector2f::ZeroVector;
 				}
-				InstanceData.TextureSize = FVector2f(RT_TextureSize.X, RT_TextureSize.Y);
 			}
 		);
 	}
@@ -409,21 +390,39 @@ void UNiagaraDataInterfaceTexture::BuildShaderParameters(FNiagaraShaderParameter
 
 void UNiagaraDataInterfaceTexture::SetShaderParameters(const FNiagaraDataInterfaceSetShaderParametersContext& Context) const
 {
-	const FNiagaraDataInterfaceProxyTexture& TextureProxy = Context.GetProxy<FNiagaraDataInterfaceProxyTexture>();
-	const FNDITextureInstanceData_RenderThread* InstanceData = TextureProxy.InstanceData_RT.Find(Context.GetSystemInstanceID());
+	FNiagaraDataInterfaceProxyTexture& TextureProxy = Context.GetProxy<FNiagaraDataInterfaceProxyTexture>();
+	FNDITextureInstanceData_RenderThread* InstanceData = TextureProxy.InstanceData_RT.Find(Context.GetSystemInstanceID());
 
 	FShaderParameters* Parameters = Context.GetParameterNestedStruct<FShaderParameters>();
-	if (InstanceData && InstanceData->ResolvedTextureRHI.IsValid())
+	if (InstanceData && InstanceData->TextureReferenceRHI.IsValid())
 	{
-		Parameters->TextureSize		= InstanceData->TextureSize;
-		Parameters->Texture			= InstanceData->ResolvedTextureRHI;
-		Parameters->TextureSampler	= InstanceData->SamplerStateRHI ? InstanceData->SamplerStateRHI : GBlackTexture->SamplerStateRHI;
+		Parameters->TextureSize = InstanceData->TextureSize;
+		Parameters->TextureSampler = InstanceData->SamplerStateRHI;
+		if (Context.IsResourceBound(&Parameters->Texture))
+		{
+			FRDGTextureRef RDGTexture = InstanceData ? InstanceData->TransientRDGTexture : nullptr;
+			if (InstanceData && RDGTexture == nullptr)
+			{
+				FTextureRHIRef ResolvedTextureRHI = InstanceData->TextureReferenceRHI->GetReferencedTexture();
+				if (ResolvedTextureRHI.IsValid())
+				{
+					InstanceData->TransientRDGTexture = Context.GetGraphBuilder().FindExternalTexture(ResolvedTextureRHI);
+					if (InstanceData->TransientRDGTexture == nullptr)
+					{
+						InstanceData->TransientRDGTexture = Context.GetGraphBuilder().RegisterExternalTexture(CreateRenderTarget(ResolvedTextureRHI, TEXT("NiagaraTexture2D")));
+					}
+					RDGTexture = InstanceData->TransientRDGTexture;
+				}
+			}
+
+			Parameters->Texture = RDGTexture ? RDGTexture : Context.GetComputeDispatchInterface().GetBlackTexture(Context.GetGraphBuilder(), ETextureDimension::Texture2D);
+		}
 	}
 	else
 	{
-		Parameters->TextureSize		= FVector2f::ZeroVector;
-		Parameters->Texture			= GBlackTexture->TextureRHI;
-		Parameters->TextureSampler	= GBlackTexture->SamplerStateRHI;
+		Parameters->TextureSize = FVector2f::ZeroVector;
+		Parameters->TextureSampler = TStaticSamplerState<SF_Point>::GetRHI();
+		Parameters->Texture = Context.GetComputeDispatchInterface().GetBlackTexture(Context.GetGraphBuilder(), ETextureDimension::Texture2D);
 	}
 }
 
diff --git a/Engine/Plugins/FX/Niagara/Source/Niagara/Private/NiagaraDataInterfaceVolumeTexture.cpp b/Engine/Plugins/FX/Niagara/Source/Niagara/Private/NiagaraDataInterfaceVolumeTexture.cpp
--- a/Engine/Plugins/FX/Niagara/Source/Niagara/Private/NiagaraDataInterfaceVolumeTexture.cpp
+++ b/Engine/Plugins/FX/Niagara/Source/Niagara/Private/NiagaraDataInterfaceVolumeTexture.cpp
@@ -2,13 +2,13 @@
 
 #include "NiagaraDataInterfaceVolumeTexture.h"
 #include "NiagaraComputeExecutionContext.h"
+#include "NiagaraGpuComputeDispatchInterface.h"
 #include "NiagaraShader.h"
 #include "NiagaraShaderParametersBuilder.h"
 #include "NiagaraSystemInstance.h"
 
 #include "Engine/TextureRenderTargetVolume.h"
 #include "Engine/VolumeTexture.h"
-#include "ShaderParameterUtils.h"
 
 #define LOCTEXT_NAMESPACE "UNiagaraDataInterfaceVolumeTexture"
 
@@ -27,8 +27,9 @@ struct FNDIVolumeTextureInstanceData_RenderThread
 {
 	FSamplerStateRHIRef		SamplerStateRHI;
 	FTextureReferenceRHIRef	TextureReferenceRHI;
-	FTextureRHIRef			ResolvedTextureRHI;
 	FVector3f				TextureSize;
+
+	FRDGTextureRef			TransientRDGTexture = nullptr;
 };
 
 struct FNiagaraDataInterfaceProxyVolumeTexture : public FNiagaraDataInterfaceProxy
@@ -36,26 +37,13 @@ struct FNiagaraDataInterfaceProxyVolumeTexture : public FNiagaraDataInterfacePro
 	virtual void ConsumePerInstanceDataFromGameThread(void* PerInstanceData, const FNiagaraSystemInstanceID& Instance) override { check(false); }
 	virtual int32 PerInstanceDataPassedToRenderThreadSize() const override { return 0; }
 
-	virtual void PreStage(FRHICommandList& RHICmdList, const FNiagaraDataInterfaceStageArgs& Context) override
+	virtual void PostSimulate(const FNDIGpuComputePostSimulateContext& Context) override
 	{
-		if (FNDIVolumeTextureInstanceData_RenderThread* InstanceData = InstanceData_RT.Find(Context.SystemInstanceID))
+		if (Context.IsFinalPostSimulate())
 		{
-			// Because the underlying reference can have a switch in flight on the RHI we get the referenced texture
-			// here, ensure it's valid (as it could be queued for delete) and cache until next round.  If we were
-			// to release the reference in PostStage / PostSimulate we still stand a chance the the transition we
-			// queue will be invalid by the time it is processed on the RHI thread.
-			if (Context.SimStageData->bFirstStage && InstanceData->TextureReferenceRHI.IsValid())
+			if (FNDIVolumeTextureInstanceData_RenderThread* InstanceData = InstanceData_RT.Find(Context.GetSystemInstanceID()))
 			{
-				InstanceData->ResolvedTextureRHI = InstanceData->TextureReferenceRHI->GetReferencedTexture();
-				if (InstanceData->ResolvedTextureRHI && !InstanceData->ResolvedTextureRHI->IsValid())
-				{
-					InstanceData->ResolvedTextureRHI = nullptr;
-				}
-			}
-			if (InstanceData->ResolvedTextureRHI)
-			{
-				// Make sure the texture is readable, we don't know where it's coming from.
-				RHICmdList.Transition(FRHITransitionInfo(InstanceData->ResolvedTextureRHI, ERHIAccess::Unknown, ERHIAccess::SRVMask));
+				InstanceData->TransientRDGTexture = nullptr;
 			}
 		}
 	}
@@ -209,14 +197,15 @@ bool UNiagaraDataInterfaceVolumeTexture::PerInstanceTick(void* PerInstanceData,
 					if (RT_Texture)
 					{
 						InstanceData.TextureReferenceRHI = RT_Texture->TextureReference.TextureReferenceRHI;
-						InstanceData.SamplerStateRHI = RT_Texture->GetResource() ? RT_Texture->GetResource()->SamplerStateRHI : nullptr;
+						InstanceData.SamplerStateRHI = RT_Texture->GetResource() ? RT_Texture->GetResource()->SamplerStateRHI.GetReference() : TStaticSamplerState<SF_Point>::GetRHI();
+						InstanceData.TextureSize = FVector3f(RT_TextureSize.X, RT_TextureSize.Y, RT_TextureSize.Z);
 					}
 					else
 					{
 						InstanceData.TextureReferenceRHI = nullptr;
 						InstanceData.SamplerStateRHI = nullptr;
+						InstanceData.TextureSize = FVector3f::ZeroVector;
 					}
-					InstanceData.TextureSize = FVector3f(RT_TextureSize.X, RT_TextureSize.Y, RT_TextureSize.Z);
 				}
 			);
 		}
@@ -305,21 +294,39 @@ void UNiagaraDataInterfaceVolumeTexture::BuildShaderParameters(FNiagaraShaderPar
 
 void UNiagaraDataInterfaceVolumeTexture::SetShaderParameters(const FNiagaraDataInterfaceSetShaderParametersContext& Context) const
 {
-	const FNiagaraDataInterfaceProxyVolumeTexture& TextureProxy = Context.GetProxy<FNiagaraDataInterfaceProxyVolumeTexture>();
-	const FNDIVolumeTextureInstanceData_RenderThread* InstanceData = TextureProxy.InstanceData_RT.Find(Context.GetSystemInstanceID());
+	FNiagaraDataInterfaceProxyVolumeTexture& TextureProxy = Context.GetProxy<FNiagaraDataInterfaceProxyVolumeTexture>();
+	FNDIVolumeTextureInstanceData_RenderThread* InstanceData = TextureProxy.InstanceData_RT.Find(Context.GetSystemInstanceID());
 
 	FShaderParameters* Parameters = Context.GetParameterNestedStruct<FShaderParameters>();
-	if (InstanceData && InstanceData->ResolvedTextureRHI.IsValid())
+	if (InstanceData && InstanceData->TextureReferenceRHI.IsValid())
 	{
-		Parameters->TextureSize		= InstanceData->TextureSize;
-		Parameters->Texture			= InstanceData->ResolvedTextureRHI;
-		Parameters->TextureSampler	= InstanceData->SamplerStateRHI ? InstanceData->SamplerStateRHI : GBlackVolumeTexture->SamplerStateRHI;
+		Parameters->TextureSize = InstanceData->TextureSize;
+		Parameters->TextureSampler = InstanceData->SamplerStateRHI;
+		if (Context.IsResourceBound(&Parameters->Texture))
+		{
+			FRDGTextureRef RDGTexture = InstanceData ? InstanceData->TransientRDGTexture : nullptr;
+			if (InstanceData && RDGTexture == nullptr)
+			{
+				FTextureRHIRef ResolvedTextureRHI = InstanceData->TextureReferenceRHI->GetReferencedTexture();
+				if (ResolvedTextureRHI.IsValid())
+				{
+					InstanceData->TransientRDGTexture = Context.GetGraphBuilder().FindExternalTexture(ResolvedTextureRHI);
+					if (InstanceData->TransientRDGTexture == nullptr)
+					{
+						InstanceData->TransientRDGTexture = Context.GetGraphBuilder().RegisterExternalTexture(CreateRenderTarget(ResolvedTextureRHI, TEXT("NiagaraTextureVolume")));
+					}
+					RDGTexture = InstanceData->TransientRDGTexture;
+				}
+			}
+
+			Parameters->Texture = RDGTexture ? RDGTexture : Context.GetComputeDispatchInterface().GetBlackTexture(Context.GetGraphBuilder(), ETextureDimension::Texture3D);
+		}
 	}
 	else
 	{
-		Parameters->TextureSize		= FVector3f::ZeroVector;
-		Parameters->Texture			= GBlackVolumeTexture->TextureRHI;
-		Parameters->TextureSampler	= GBlackVolumeTexture->SamplerStateRHI;
+		Parameters->TextureSize = FVector3f::ZeroVector;
+		Parameters->TextureSampler = TStaticSamplerState<SF_Point>::GetRHI();
+		Parameters->Texture = Context.GetComputeDispatchInterface().GetBlackTexture(Context.GetGraphBuilder(), ETextureDimension::Texture3D);
 	}
 }
 
diff --git a/Engine/Plugins/FX/Niagara/Source/Niagara/Private/NiagaraDataSet.cpp b/Engine/Plugins/FX/Niagara/Source/Niagara/Private/NiagaraDataSet.cpp
--- a/Engine/Plugins/FX/Niagara/Source/Niagara/Private/NiagaraDataSet.cpp
+++ b/Engine/Plugins/FX/Niagara/Source/Niagara/Private/NiagaraDataSet.cpp
@@ -464,7 +464,7 @@ void FNiagaraDataSet::AllocateGPUFreeIDs(uint32 InNumInstances, FRHICommandList&
 	FRHIShaderResourceView* ExistingBuffer = GPUNumAllocatedIDs > 0 ? GPUFreeIDs.SRV.GetReference() : FNiagaraRenderer::GetDummyIntBuffer();
 
 	RHICmdList.Transition(FRHITransitionInfo(NewFreeIDsBuffer.UAV, ERHIAccess::Unknown, ERHIAccess::UAVCompute));
-	NiagaraInitGPUFreeIDList(RHICmdList, FeatureLevel, NumIDsToAlloc, NewFreeIDsBuffer, GPUNumAllocatedIDs, ExistingBuffer);
+	NiagaraInitGPUFreeIDList(RHICmdList, FeatureLevel, NewFreeIDsBuffer.UAV, NumIDsToAlloc, ExistingBuffer, GPUNumAllocatedIDs);
 	RHICmdList.Transition(FRHITransitionInfo(NewFreeIDsBuffer.UAV, ERHIAccess::UAVCompute, ERHIAccess::SRVCompute));
 
 	GPUFreeIDs = MoveTemp(NewFreeIDsBuffer);
diff --git a/Engine/Plugins/FX/Niagara/Source/Niagara/Private/NiagaraDebugHud.cpp b/Engine/Plugins/FX/Niagara/Source/Niagara/Private/NiagaraDebugHud.cpp
--- a/Engine/Plugins/FX/Niagara/Source/Niagara/Private/NiagaraDebugHud.cpp
+++ b/Engine/Plugins/FX/Niagara/Source/Niagara/Private/NiagaraDebugHud.cpp
@@ -549,12 +549,8 @@ FNiagaraDebugHud::FNiagaraDebugHud(UWorld* World)
 				GpuResults = InGpuResults;
 				GpuResultsGameFrameCounter = GFrameCounter;
 
-				for (const auto& StageResults : GpuResults->StageResults )
-				{
-					GpuTotalDispatches.Accumulate(GpuResultsGameFrameCounter, StageResults.NumDispatches);
-					GpuTotalDispatchGroups.Accumulate(GpuResultsGameFrameCounter, StageResults.NumDispatchGroups);
-					GpuTotalMicroseconds.Accumulate(GpuResultsGameFrameCounter, StageResults.DurationMicroseconds);
-				}
+				GpuTotalDispatches.Accumulate(GpuResultsGameFrameCounter, GpuResults->TotalDispatches);
+				GpuTotalMicroseconds.Accumulate(GpuResultsGameFrameCounter, GpuResults->TotalDurationMicroseconds);
 
 				for ( const auto& DispatchResult : GpuResults->DispatchResults )
 				{
@@ -1903,7 +1899,6 @@ void FNiagaraDebugHud::DrawGpuComputeOverriew(class FNiagaraWorldManager* WorldM
 
 		SimpleTable.GetColumnText(0).Appendf(TEXT("TotalMicroseconds : %llu"), GpuTotalMicroseconds.GetAverage());
 		SimpleTable.GetColumnText(1).Appendf(TEXT("TotalDispatches : %d"), GpuTotalDispatches.GetAverage());
-		SimpleTable.GetColumnText(2).Appendf(TEXT("TotalDispatchGroups : %d"), GpuTotalDispatchGroups.GetAverage());
 
 		SimpleTable.Draw(Font, DrawCanvas, TextLocation, DetailColor, BackgroundColor);
 		TextLocation.Y += fAdvanceHeight;
diff --git a/Engine/Plugins/FX/Niagara/Source/Niagara/Private/NiagaraDebugHud.h b/Engine/Plugins/FX/Niagara/Source/Niagara/Private/NiagaraDebugHud.h
--- a/Engine/Plugins/FX/Niagara/Source/Niagara/Private/NiagaraDebugHud.h
+++ b/Engine/Plugins/FX/Niagara/Source/Niagara/Private/NiagaraDebugHud.h
@@ -280,7 +280,6 @@ private:
 	uint64						GpuResultsGameFrameCounter = 0;
 
 	FSmoothedCounter<int32>		GpuTotalDispatches;
-	FSmoothedCounter<int32>		GpuTotalDispatchGroups;
 	FSmoothedCounter<uint64>	GpuTotalMicroseconds;
 
 	struct FGpuUsagePerStage
diff --git a/Engine/Plugins/FX/Niagara/Source/Niagara/Private/NiagaraGPUProfiler.cpp b/Engine/Plugins/FX/Niagara/Source/Niagara/Private/NiagaraGPUProfiler.cpp
--- a/Engine/Plugins/FX/Niagara/Source/Niagara/Private/NiagaraGPUProfiler.cpp
+++ b/Engine/Plugins/FX/Niagara/Source/Niagara/Private/NiagaraGPUProfiler.cpp
@@ -26,11 +26,6 @@ FNiagaraGPUProfiler::~FNiagaraGPUProfiler()
 	for (FGpuFrameData& Frame : GpuFrames)
 	{
 		Frame.EndQuery.ReleaseQuery();
-		for (FGpuStageTimer& StageTimer : Frame.StageTimers)
-		{
-			StageTimer.StartQuery.ReleaseQuery();
-			StageTimer.EndQuery.ReleaseQuery();
-		}
 
 		for (FGpuDispatchTimer& DispatchTimer : Frame.DispatchTimers)
 		{
@@ -81,52 +76,7 @@ void FNiagaraGPUProfiler::EndFrame(FRHICommandList& RHICmdList)
 	CurrentWriteFrame = (CurrentWriteFrame + 1) % NumBufferFrames;
 }
 
-void FNiagaraGPUProfiler::BeginStage(FRHICommandList& RHICmdList, ENiagaraGpuComputeTickStage::Type TickStage, int32 NumDispatchGroups)
-{
-	if (ActiveWriteFrame == nullptr)
-	{
-		return;
-	}
-
-	FGpuStageTimer& StageTimer = ActiveWriteFrame->StageTimers[TickStage];
-	StageTimer.NumDispatchGroups = NumDispatchGroups;
-	StageTimer.StartQuery = QueryPool->AllocateQuery();
-	RHICmdList.EndRenderQuery(StageTimer.StartQuery.GetQuery());
-}
-
-void FNiagaraGPUProfiler::EndStage(FRHICommandList& RHICmdList, ENiagaraGpuComputeTickStage::Type TickStage, int32 NumDispatches)
-{
-	if (ActiveWriteFrame == nullptr)
-	{
-		return;
-	}
-
-	FGpuStageTimer& StageTimer = ActiveWriteFrame->StageTimers[TickStage];
-	StageTimer.NumDispatches = NumDispatches;
-	StageTimer.EndQuery = QueryPool->AllocateQuery();
-	RHICmdList.EndRenderQuery(StageTimer.EndQuery.GetQuery());
-}
-
-void FNiagaraGPUProfiler::BeginDispatch(FRHICommandList& RHICmdList, const FNiagaraGpuDispatchInstance& DispatchInstance)
-{
-	if (ActiveWriteFrame == nullptr)
-	{
-		return;
-	}
-	check(bDispatchRecursionGuard == false);
-	bDispatchRecursionGuard = true;
-
-	FGpuDispatchTimer& DispatchTimer = ActiveWriteFrame->DispatchTimers.AddDefaulted_GetRef();
-
-	DispatchTimer.bUniqueInstance = DispatchInstance.SimStageData.bSetDataToRender && (&DispatchInstance.InstanceData == &DispatchInstance.Tick.GetInstances()[0]);
-	DispatchTimer.OwnerComponent = DispatchInstance.InstanceData.Context->ProfilingComponentPtr;
-	DispatchTimer.OwnerEmitter = DispatchInstance.InstanceData.Context->ProfilingEmitterPtr;
-	DispatchTimer.StageName = DispatchInstance.SimStageData.StageMetaData->SimulationStageName;
-	DispatchTimer.StartQuery = QueryPool->AllocateQuery();
-	RHICmdList.EndRenderQuery(DispatchTimer.StartQuery.GetQuery());
-}
-
-void FNiagaraGPUProfiler::BeginDispatch(FRHICommandList& RHICmdList, const struct FNiagaraComputeInstanceData& InstanceData, FName StageName)
+void FNiagaraGPUProfiler::BeginDispatch(FRHICommandList& RHICmdList, const FNiagaraGpuProfileEvent& Event)
 {
 	if (ActiveWriteFrame == nullptr)
 	{
@@ -135,29 +85,8 @@ void FNiagaraGPUProfiler::BeginDispatch(FRHICommandList& RHICmdList, const struc
 	check(bDispatchRecursionGuard == false);
 	bDispatchRecursionGuard = true;
 
-	FGpuDispatchTimer& DispatchTimer = ActiveWriteFrame->DispatchTimers.AddDefaulted_GetRef();
-	DispatchTimer.bUniqueInstance = false;
-	DispatchTimer.OwnerComponent = InstanceData.Context->ProfilingComponentPtr;
-	DispatchTimer.OwnerEmitter = InstanceData.Context->ProfilingEmitterPtr;
-	DispatchTimer.StageName = StageName;
-	DispatchTimer.StartQuery = QueryPool->AllocateQuery();
-	RHICmdList.EndRenderQuery(DispatchTimer.StartQuery.GetQuery());
-}
-
-void FNiagaraGPUProfiler::BeginDispatch(FRHICommandList& RHICmdList, FName StageName)
-{
-	if (ActiveWriteFrame == nullptr)
-	{
-		return;
-	}
-	check(bDispatchRecursionGuard == false);
-	bDispatchRecursionGuard = true;
+	FGpuDispatchTimer& DispatchTimer = ActiveWriteFrame->DispatchTimers.Emplace_GetRef(Event);
 
-	FGpuDispatchTimer& DispatchTimer = ActiveWriteFrame->DispatchTimers.AddDefaulted_GetRef();
-	DispatchTimer.bUniqueInstance = false;
-	DispatchTimer.OwnerComponent = nullptr;
-	DispatchTimer.OwnerEmitter = FVersionedNiagaraEmitterWeakPtr();
-	DispatchTimer.StageName = StageName;
 	DispatchTimer.StartQuery = QueryPool->AllocateQuery();
 	RHICmdList.EndRenderQuery(DispatchTimer.StartQuery.GetQuery());
 }
@@ -190,30 +119,11 @@ bool FNiagaraGPUProfiler::ProcessFrame(FRHICommandListImmediate& RHICmdList, FGp
 	//-OPT: Potentially pool these
 	FNiagaraGpuFrameResultsPtr FrameResults = MakeShared<FNiagaraGpuFrameResults, ESPMode::ThreadSafe>();
 	FrameResults->OwnerContext = OwnerContext;
+	FrameResults->TotalDispatches = 0;
+	FrameResults->TotalDurationMicroseconds = 0;
 	FrameResults->DispatchResults.Reserve(ReadFrame.DispatchTimers.Num());
 
 	// Process results
-	for (int32 i = 0; i < ENiagaraGpuComputeTickStage::Max; ++i)
-	{
-		FGpuStageTimer& StageTimer = ReadFrame.StageTimers[i];
-		auto& StageResults = FrameResults->StageResults[i];
-		StageResults.NumDispatches = StageTimer.NumDispatches;
-		StageResults.NumDispatchGroups = StageTimer.NumDispatchGroups;
-		StageResults.DurationMicroseconds = 0;
-		if ( StageTimer.StartQuery.GetQuery() != nullptr )
-		{
-			uint64 StartMicroseconds = 0;
-			uint64 EndMicroseconds = 0;
-			ensure(RHICmdList.GetRenderQueryResult(StageTimer.StartQuery.GetQuery(), StartMicroseconds, false));
-			ensure(RHICmdList.GetRenderQueryResult(StageTimer.EndQuery.GetQuery(), EndMicroseconds, false));
-			StageResults.DurationMicroseconds = EndMicroseconds - StartMicroseconds;
-			StageTimer.StartQuery.ReleaseQuery();
-			StageTimer.EndQuery.ReleaseQuery();
-		}
-		StageTimer.NumDispatches = 0;
-		StageTimer.NumDispatchGroups = 0;
-	}
-
 	for (FGpuDispatchTimer& DispatchTimer : ReadFrame.DispatchTimers)
 	{
 		auto& DispatchResults = FrameResults->DispatchResults.AddDefaulted_GetRef();
@@ -225,11 +135,14 @@ bool FNiagaraGPUProfiler::ProcessFrame(FRHICommandListImmediate& RHICmdList, FGp
 		DispatchTimer.StartQuery.ReleaseQuery();
 		DispatchTimer.EndQuery.ReleaseQuery();
 
-		DispatchResults.bUniqueInstance = DispatchTimer.bUniqueInstance;
-		DispatchResults.OwnerComponent = DispatchTimer.OwnerComponent;
-		DispatchResults.OwnerEmitter = DispatchTimer.OwnerEmitter;
-		DispatchResults.StageName = DispatchTimer.StageName;
+		DispatchResults.bUniqueInstance = DispatchTimer.Event.bUniqueInstance;
+		DispatchResults.OwnerComponent = DispatchTimer.Event.OwnerComponent;
+		DispatchResults.OwnerEmitter = DispatchTimer.Event.OwnerEmitter;
+		DispatchResults.StageName = DispatchTimer.Event.StageName;
 		DispatchResults.DurationMicroseconds = EndMicroseconds - StartMicroseconds;
+
+		++FrameResults->TotalDispatches;
+		FrameResults->TotalDurationMicroseconds += DispatchResults.DurationMicroseconds;
 	}
 
 	ReadFrame.DispatchTimers.Empty();
diff --git a/Engine/Plugins/FX/Niagara/Source/Niagara/Private/NiagaraGPUProfiler.h b/Engine/Plugins/FX/Niagara/Source/Niagara/Private/NiagaraGPUProfiler.h
--- a/Engine/Plugins/FX/Niagara/Source/Niagara/Private/NiagaraGPUProfiler.h
+++ b/Engine/Plugins/FX/Niagara/Source/Niagara/Private/NiagaraGPUProfiler.h
@@ -20,31 +20,21 @@ class FNiagaraGPUProfiler : public FNiagaraGPUProfilerInterface
 {
 	static constexpr int32 NumBufferFrames = 5;
 
-	struct FGpuStageTimer
+	struct FGpuDispatchTimer
 	{
-		int32					NumDispatchGroups = 0;
-		int32					NumDispatches = 0;
+		explicit FGpuDispatchTimer(const FNiagaraGpuProfileEvent& InEvent) : Event(InEvent) {}
+
+		FNiagaraGpuProfileEvent	Event;
 		FRHIPooledRenderQuery	StartQuery;
 		FRHIPooledRenderQuery	EndQuery;
 	};
 
-	struct FGpuDispatchTimer
-	{
-		uint32									bUniqueInstance : 1;
-		TWeakObjectPtr<class USceneComponent>	OwnerComponent;
-		FVersionedNiagaraEmitterWeakPtr			OwnerEmitter;
-		FName									StageName;
-		FRHIPooledRenderQuery					StartQuery;
-		FRHIPooledRenderQuery					EndQuery;
-	};
-
 	struct FGpuFrameData
 	{
 		bool CanWrite() const { return EndQuery.GetQuery() == nullptr; }
 		bool CanRead() const { return EndQuery.GetQuery() != nullptr; }
 
 		FRHIPooledRenderQuery		EndQuery;
-		FGpuStageTimer				StageTimers[ENiagaraGpuComputeTickStage::Max];
 		TArray<FGpuDispatchTimer>	DispatchTimers;
 	};
 
@@ -57,12 +47,7 @@ public:
 	void BeginFrame(FRHICommandListImmediate& RHICmdList);
 	void EndFrame(FRHICommandList& RHICmdList);
 
-	void BeginStage(FRHICommandList& RHICmdList, ENiagaraGpuComputeTickStage::Type TickStage, int32 NumDispatchGroups);
-	void EndStage(FRHICommandList& RHICmdList, ENiagaraGpuComputeTickStage::Type TickStage, int32 NumDispatches);
-
-	void BeginDispatch(FRHICommandList& RHICmdList, const struct FNiagaraGpuDispatchInstance& DispatchInstance);
-	void BeginDispatch(FRHICommandList& RHICmdList, const struct FNiagaraComputeInstanceData& InstanceData, FName StageName);
-	void BeginDispatch(FRHICommandList& RHICmdList, FName StageName);
+	void BeginDispatch(FRHICommandList& RHICmdList, const FNiagaraGpuProfileEvent& Event);
 	void EndDispatch(FRHICommandList& RHICmdList);
 
 private:
diff --git a/Engine/Plugins/FX/Niagara/Source/Niagara/Private/NiagaraGPUProfilerInterface.cpp b/Engine/Plugins/FX/Niagara/Source/Niagara/Private/NiagaraGPUProfilerInterface.cpp
--- a/Engine/Plugins/FX/Niagara/Source/Niagara/Private/NiagaraGPUProfilerInterface.cpp
+++ b/Engine/Plugins/FX/Niagara/Source/Niagara/Private/NiagaraGPUProfilerInterface.cpp
@@ -49,53 +49,32 @@ void FNiagaraGpuProfilerListener::SetHandler(TFunction<void(const FNiagaraGpuFra
 }
 
 //////////////////////////////////////////////////////////////////////////
-FNiagaraGpuProfileScope::FNiagaraGpuProfileScope(FRHICommandList& InRHICmdList, const FNiagaraGpuComputeDispatchInterface* ComputeDispatchInterface, FName StageName)
-	: RHICmdList(InRHICmdList)
-	, GPUProfiler(static_cast<FNiagaraGPUProfiler*>(ComputeDispatchInterface->GetGPUProfiler()))
-{
-	if (GPUProfiler)
-	{
-		GPUProfiler->BeginDispatch(RHICmdList, StageName);
-	}
-}
 
-FNiagaraGpuProfileScope::FNiagaraGpuProfileScope(FRHICommandList& InRHICmdList, const FNiagaraGpuComputeDispatchInterface* ComputeDispatchInterface, const struct FNiagaraGpuDispatchInstance& DispatchInstance)
-	: RHICmdList(InRHICmdList)
-	, GPUProfiler(static_cast<FNiagaraGPUProfiler*>(ComputeDispatchInterface->GetGPUProfiler()))
+FNiagaraGpuProfileEvent::FNiagaraGpuProfileEvent(const FNiagaraComputeInstanceData& InstanceData, const FNiagaraSimStageData& SimStageData, const bool bFirstInstanceData)
 {
-	if (GPUProfiler)
-	{
-		GPUProfiler->BeginDispatch(RHICmdList, DispatchInstance);
-	}
+	bUniqueInstance = SimStageData.bSetDataToRender && bFirstInstanceData;
+	OwnerComponent = InstanceData.Context->ProfilingComponentPtr;
+	OwnerEmitter = InstanceData.Context->ProfilingEmitterPtr;
+	StageName = SimStageData.StageMetaData->SimulationStageName;
 }
 
-FNiagaraGpuProfileScope::FNiagaraGpuProfileScope(FRHICommandList& InRHICmdList, const struct FNiagaraDataInterfaceArgs& Context, FName StageName)
-	: RHICmdList(InRHICmdList)
-	, GPUProfiler(static_cast<FNiagaraGPUProfiler*>(Context.ComputeDispatchInterface->GetGPUProfiler()))
+FNiagaraGpuProfileEvent::FNiagaraGpuProfileEvent(const FNiagaraComputeInstanceData& InstanceData, FName CustomStageName)
 {
-	if (GPUProfiler)
-	{
-		GPUProfiler->BeginDispatch(RHICmdList, StageName);
-	}
+	bUniqueInstance = false;
+	OwnerComponent = InstanceData.Context->ProfilingComponentPtr;
+	OwnerEmitter = InstanceData.Context->ProfilingEmitterPtr;
+	StageName = CustomStageName;
 }
 
-FNiagaraGpuProfileScope::FNiagaraGpuProfileScope(FRHICommandList& InRHICmdList, const struct FNiagaraDataInterfaceSetArgs& Context, FName StageName)
-	: RHICmdList(InRHICmdList)
-	, GPUProfiler(static_cast<FNiagaraGPUProfiler*>(Context.ComputeDispatchInterface->GetGPUProfiler()))
-{
-	if (GPUProfiler)
-	{
-		GPUProfiler->BeginDispatch(RHICmdList, *Context.ComputeInstanceData, StageName);
-	}
-}
+//////////////////////////////////////////////////////////////////////////
 
-FNiagaraGpuProfileScope::FNiagaraGpuProfileScope(FRHICommandList& InRHICmdList, const struct FNiagaraDataInterfaceStageArgs& Context, FName StageName)
+FNiagaraGpuProfileScope::FNiagaraGpuProfileScope(FRHICommandList& InRHICmdList, const class FNiagaraGpuComputeDispatchInterface* ComputeDispatchInterface, const FNiagaraGpuProfileEvent& Event)
 	: RHICmdList(InRHICmdList)
-	, GPUProfiler(static_cast<FNiagaraGPUProfiler*>(Context.ComputeDispatchInterface->GetGPUProfiler()))
+	, GPUProfiler(static_cast<FNiagaraGPUProfiler*>(ComputeDispatchInterface->GetGPUProfiler()))
 {
 	if (GPUProfiler)
 	{
-		GPUProfiler->BeginDispatch(RHICmdList, *Context.ComputeInstanceData, StageName);
+		GPUProfiler->BeginDispatch(RHICmdList, Event);
 	}
 }
 
diff --git a/Engine/Plugins/FX/Niagara/Source/Niagara/Private/NiagaraGpuComputeDebug.cpp b/Engine/Plugins/FX/Niagara/Source/Niagara/Private/NiagaraGpuComputeDebug.cpp
--- a/Engine/Plugins/FX/Niagara/Source/Niagara/Private/NiagaraGpuComputeDebug.cpp
+++ b/Engine/Plugins/FX/Niagara/Source/Niagara/Private/NiagaraGpuComputeDebug.cpp
@@ -114,18 +114,18 @@ void FNiagaraGpuComputeDebug::OnSystemDeallocated(FNiagaraSystemInstanceID Syste
 	DebugDrawBuffers.Remove(SystemInstanceID);
 }
 
-void FNiagaraGpuComputeDebug::AddTexture(FRHICommandList& RHICmdList, FNiagaraSystemInstanceID SystemInstanceID, FName SourceName, FRHITexture* Texture, FVector2D PreviewDisplayRange)
+void FNiagaraGpuComputeDebug::AddTexture(FRDGBuilder& GraphBuilder, FNiagaraSystemInstanceID SystemInstanceID, FName SourceName, FRDGTextureRef Texture, FVector2D PreviewDisplayRange)
 {
-	AddAttributeTexture(RHICmdList, SystemInstanceID, SourceName, Texture, FIntPoint::ZeroValue, FIntVector4(INDEX_NONE, INDEX_NONE, INDEX_NONE, INDEX_NONE), PreviewDisplayRange);
+	AddAttributeTexture(GraphBuilder, SystemInstanceID, SourceName, Texture, FIntPoint::ZeroValue, FIntVector4(INDEX_NONE, INDEX_NONE, INDEX_NONE, INDEX_NONE), PreviewDisplayRange);
 }
 
-void FNiagaraGpuComputeDebug::AddAttributeTexture(FRHICommandList& RHICmdList, FNiagaraSystemInstanceID SystemInstanceID, FName SourceName, FRHITexture* Texture, FIntPoint NumTextureAttributes, FIntVector4 AttributeIndices, FVector2D PreviewDisplayRange)
+void FNiagaraGpuComputeDebug::AddAttributeTexture(FRDGBuilder& GraphBuilder, FNiagaraSystemInstanceID SystemInstanceID, FName SourceName, FRDGTextureRef Texture, FIntPoint NumTextureAttributes, FIntVector4 AttributeIndices, FVector2D PreviewDisplayRange)
 {
 	FIntVector4 TextureAttributesInt4 = FIntVector4(NumTextureAttributes.X, NumTextureAttributes.Y, 0, 0);
-	AddAttributeTexture( RHICmdList, SystemInstanceID, SourceName,  Texture, TextureAttributesInt4, AttributeIndices, PreviewDisplayRange);
+	AddAttributeTexture(GraphBuilder, SystemInstanceID, SourceName, Texture, TextureAttributesInt4, AttributeIndices, PreviewDisplayRange);
 }
 
-void FNiagaraGpuComputeDebug::AddAttributeTexture(FRHICommandList& RHICmdList, FNiagaraSystemInstanceID SystemInstanceID, FName SourceName, FRHITexture* Texture, FIntVector4 NumTextureAttributes, FIntVector4 AttributeIndices, FVector2D PreviewDisplayRange)
+void FNiagaraGpuComputeDebug::AddAttributeTexture(FRDGBuilder& GraphBuilder, FNiagaraSystemInstanceID SystemInstanceID, FName SourceName, FRDGTextureRef Texture, FIntVector4 NumTextureAttributes, FIntVector4 AttributeIndices, FVector2D PreviewDisplayRange)
 {
 	if (!SystemInstancesToWatch.Contains(SystemInstanceID))
 	{
@@ -137,13 +137,10 @@ void FNiagaraGpuComputeDebug::AddAttributeTexture(FRHICommandList& RHICmdList, F
 		return;
 	}
 
-	const FRHITextureDesc& TextureDesc = Texture->GetDesc();
+	const FRDGTextureDesc& TextureDesc = Texture->Desc;
 
 	bool bCreateTexture = false;
 
-	const FIntVector SrcSize = Texture->GetSizeXYZ();
-	const EPixelFormat SrcFormat = Texture->GetFormat();
-
 	FNiagaraVisualizeTexture* VisualizeEntry = VisualizeTextures.FindByPredicate([&SourceName, &SystemInstanceID](const FNiagaraVisualizeTexture& Texture) -> bool { return Texture.SystemInstanceID == SystemInstanceID && Texture.SourceName == SourceName; });
 	if (!VisualizeEntry)
 	{
@@ -154,14 +151,18 @@ void FNiagaraGpuComputeDebug::AddAttributeTexture(FRHICommandList& RHICmdList, F
 	}
 	else
 	{
-		bCreateTexture = (VisualizeEntry->Texture->GetSizeXYZ() != SrcSize) || (VisualizeEntry->Texture->GetFormat() != SrcFormat);
+		const FPooledRenderTargetDesc& EntryTextureDesc = VisualizeEntry->Texture->GetDesc();
+		bCreateTexture = 
+			(EntryTextureDesc.GetSize() != TextureDesc.GetSize()) ||
+			(EntryTextureDesc.Format != TextureDesc.Format) ||
+			(EntryTextureDesc.NumMips != TextureDesc.NumMips) ||
+			(EntryTextureDesc.ArraySize != TextureDesc.ArraySize);
 	}
 	VisualizeEntry->NumTextureAttributes = NumTextureAttributes;
 	VisualizeEntry->AttributesToVisualize = AttributeIndices;
 	VisualizeEntry->PreviewDisplayRange = PreviewDisplayRange;
 
 	// Do we need to create a texture to copy into?
-	FTextureRHIRef Destination;
 	if ( bCreateTexture )
 	{
 		// Create a minimal copy of the Texture's Desc
@@ -174,36 +175,14 @@ void FNiagaraGpuComputeDebug::AddAttributeTexture(FRHICommandList& RHICmdList, F
 			.SetFlags(ETextureCreateFlags::ShaderResource)
 			.SetInitialState(ERHIAccess::SRVMask);
 
-		Destination = RHICreateTexture(NewTextureDesc);
-		VisualizeEntry->Texture = Destination;
-	}
-	else
-	{
-		Destination = VisualizeEntry->Texture;
-		check(Destination != nullptr);
-	}
-
-	// Copy texture
-	{
-		FRHITransitionInfo TransitionsBefore[] =
-		{
-			FRHITransitionInfo(Texture, ERHIAccess::SRVMask, ERHIAccess::CopySrc),
-			FRHITransitionInfo(Destination, ERHIAccess::SRVMask, ERHIAccess::CopyDest)
-		};
-		RHICmdList.Transition(MakeArrayView(TransitionsBefore, UE_ARRAY_COUNT(TransitionsBefore)));
+		VisualizeEntry->Texture = CreateRenderTarget(RHICreateTexture(NewTextureDesc), TEXT("FNiagaraGpuComputeDebug"));
 	}
+	check(VisualizeEntry->Texture.IsValid());
 
 	FRHICopyTextureInfo CopyInfo;
-	RHICmdList.CopyTexture(Texture, Destination, CopyInfo);
-
-	{
-		FRHITransitionInfo TransitionsAfter[] =
-		{
-			FRHITransitionInfo(Texture, ERHIAccess::CopySrc, ERHIAccess::SRVMask),
-			FRHITransitionInfo(Destination, ERHIAccess::CopyDest, ERHIAccess::SRVMask)
-		};
-		RHICmdList.Transition(MakeArrayView(TransitionsAfter, UE_ARRAY_COUNT(TransitionsAfter)));
-	}
+	CopyInfo.NumMips = TextureDesc.NumMips;
+	CopyInfo.NumSlices = TextureDesc.ArraySize;
+	AddCopyTexturePass(GraphBuilder, Texture, GraphBuilder.RegisterExternalTexture(VisualizeEntry->Texture), CopyInfo);
 }
 
 FNiagaraSimulationDebugDrawData* FNiagaraGpuComputeDebug::GetSimulationDebugDrawData(FNiagaraSystemInstanceID SystemInstanceID, bool bRequiresGpuBuffers, uint32 OverrideMaxDebugLines)
@@ -256,7 +235,7 @@ bool FNiagaraGpuComputeDebug::ShouldDrawDebug() const
 	return GNiagaraGpuComputeDebug_DrawDebugEnabled && (VisualizeTextures.Num() > 0);
 }
 
-void FNiagaraGpuComputeDebug::DrawDebug(class FRDGBuilder& GraphBuilder, const FViewInfo& View, const FScreenPassRenderTarget& Output)
+void FNiagaraGpuComputeDebug::DrawDebug(FRDGBuilder& GraphBuilder, const FViewInfo& View, const FScreenPassRenderTarget& Output)
 {
 	if (!GNiagaraGpuComputeDebug_DrawDebugEnabled || (VisualizeTextures.Num() == 0))
 	{
@@ -275,7 +254,7 @@ void FNiagaraGpuComputeDebug::DrawDebug(class FRDGBuilder& GraphBuilder, const F
 
 	for (const FNiagaraVisualizeTexture& VisualizeEntry : VisualizeTextures)
 	{
-		FIntVector TextureSize = VisualizeEntry.Texture->GetSizeXYZ();
+		FIntVector TextureSize = VisualizeEntry.Texture->GetDesc().GetSize();
 		if ( VisualizeEntry.NumTextureAttributes.X > 0 )
 		{
 			check(VisualizeEntry.NumTextureAttributes.Y > 0);
@@ -290,7 +269,7 @@ void FNiagaraGpuComputeDebug::DrawDebug(class FRDGBuilder& GraphBuilder, const F
 
 		Location.Y -= DisplayHeight;
 
-		NiagaraDebugShaders::VisualizeTexture(GraphBuilder, View, Output, Location, DisplayHeight, VisualizeEntry.AttributesToVisualize, VisualizeEntry.Texture, VisualizeEntry.NumTextureAttributes, TickCounter, VisualizeEntry.PreviewDisplayRange);
+		NiagaraDebugShaders::VisualizeTexture(GraphBuilder, View, Output, Location, DisplayHeight, VisualizeEntry.AttributesToVisualize, GraphBuilder.RegisterExternalTexture(VisualizeEntry.Texture), VisualizeEntry.NumTextureAttributes, TickCounter, VisualizeEntry.PreviewDisplayRange);
 
 		Location.Y -= FontHeight;
 
diff --git a/Engine/Plugins/FX/Niagara/Source/Niagara/Private/NiagaraGpuComputeDebug.h b/Engine/Plugins/FX/Niagara/Source/Niagara/Private/NiagaraGpuComputeDebug.h
--- a/Engine/Plugins/FX/Niagara/Source/Niagara/Private/NiagaraGpuComputeDebug.h
+++ b/Engine/Plugins/FX/Niagara/Source/Niagara/Private/NiagaraGpuComputeDebug.h
@@ -7,7 +7,6 @@ Class used help realtime debug Gpu Compute simulations
 
 #include "CoreMinimal.h"
 #include "NiagaraCommon.h"
-#include "RHICommandList.h"
 #include "RenderGraphDefinitions.h"
 
 #if NIAGARA_COMPUTEDEBUG_ENABLED
@@ -38,12 +37,12 @@ class FNiagaraGpuComputeDebug
 public:
 	struct FNiagaraVisualizeTexture
 	{
-		FNiagaraSystemInstanceID	SystemInstanceID;
-		FName						SourceName;
-		FTextureRHIRef				Texture;
-		FIntVector4					NumTextureAttributes = FIntVector4(0,0,0,0);
-		FIntVector4					AttributesToVisualize = FIntVector4(-1, -1, -1, -1);
-		FVector2D					PreviewDisplayRange = FVector2D::ZeroVector;
+		FNiagaraSystemInstanceID			SystemInstanceID;
+		FName								SourceName;
+		TRefCountPtr<IPooledRenderTarget>	Texture;
+		FIntVector4							NumTextureAttributes = FIntVector4(0,0,0,0);
+		FIntVector4							AttributesToVisualize = FIntVector4(-1, -1, -1, -1);
+		FVector2D							PreviewDisplayRange = FVector2D::ZeroVector;
 	};
 
 	FNiagaraGpuComputeDebug(ERHIFeatureLevel::Type InFeatureLevel);
@@ -60,17 +59,17 @@ public:
 	void OnSystemDeallocated(FNiagaraSystemInstanceID SystemInstanceID);
 
 	// Add a texture to visualize
-	void AddTexture(FRHICommandList& RHICmdList, FNiagaraSystemInstanceID SystemInstanceID, FName SourceName, FRHITexture* Texture, FVector2D PreviewDisplayRange = FVector2D::ZeroVector);
+	void AddTexture(FRDGBuilder& GraphBuilder, FNiagaraSystemInstanceID SystemInstanceID, FName SourceName, FRDGTextureRef Texture, FVector2D PreviewDisplayRange = FVector2D::ZeroVector);
 
 	// Add a texture to visualize that contains a number of attributes and select which attributes to push into RGBA where -1 means ignore that channel
 	// The first -1 in the attribute indices list will also limit the number of attributes we attempt to read.
 	// NumTextureAttributes in this version is meant for a 2D atlas
-	void AddAttributeTexture(FRHICommandList& RHICmdList, FNiagaraSystemInstanceID SystemInstanceID, FName SourceName, FRHITexture* Texture, FIntPoint NumTextureAttributes, FIntVector4 AttributeIndices, FVector2D PreviewDisplayRange = FVector2D::ZeroVector);
+	void AddAttributeTexture(FRDGBuilder& GraphBuilder, FNiagaraSystemInstanceID SystemInstanceID, FName SourceName, FRDGTextureRef Texture, FIntPoint NumTextureAttributes, FIntVector4 AttributeIndices, FVector2D PreviewDisplayRange = FVector2D::ZeroVector);
 
 	// Add a texture to visualize that contains a number of attributes and select which attributes to push into RGBA where -1 means ignore that channel
 	// The first -1 in the attribute indices list will also limit the number of attributes we attempt to read
 	// NumTextureAttributes in this version is meant for a 3D atlas
-	void AddAttributeTexture(FRHICommandList& RHICmdList, FNiagaraSystemInstanceID SystemInstanceID, FName SourceName, FRHITexture* Texture, FIntVector4 NumTextureAttributes, FIntVector4 AttributeIndices, FVector2D PreviewDisplayRange = FVector2D::ZeroVector);
+	void AddAttributeTexture(FRDGBuilder& GraphBuilder, FNiagaraSystemInstanceID SystemInstanceID, FName SourceName, FRDGTextureRef Texture, FIntVector4 NumTextureAttributes, FIntVector4 AttributeIndices, FVector2D PreviewDisplayRange = FVector2D::ZeroVector);
 
 	// Get Debug draw buffers for a system instance
 	FNiagaraSimulationDebugDrawData* GetSimulationDebugDrawData(FNiagaraSystemInstanceID SystemInstanceID, bool bRequiresGpuBuffers, uint32 OverrideMaxDebugLines = 0);
@@ -82,10 +81,10 @@ public:
 	bool ShouldDrawDebug() const;
 
 	// Draw all the debug information for the system
-	void DrawDebug(class FRDGBuilder& GraphBuilder, const class FViewInfo& View, const struct FScreenPassRenderTarget& Output);
+	void DrawDebug(FRDGBuilder& GraphBuilder, const class FViewInfo& View, const struct FScreenPassRenderTarget& Output);
 
 	// Draw debug information that requires rendering into the scene
-	void DrawSceneDebug(class FRDGBuilder& GraphBuilder, const class FViewInfo& View, FRDGTextureRef SceneColor, FRDGTextureRef SceneDepth);
+	void DrawSceneDebug(FRDGBuilder& GraphBuilder, const class FViewInfo& View, FRDGTextureRef SceneColor, FRDGTextureRef SceneDepth);
 
 private:
 	ERHIFeatureLevel::Type FeatureLevel;
