- Niagara Compute Dipsatch modified to use FRDGBuilder through the pipeline

diff --git a/Engine/Plugins/FX/Niagara/Source/Niagara/Private/NiagaraGpuComputeDispatch.cpp b/Engine/Plugins/FX/Niagara/Source/Niagara/Private/NiagaraGpuComputeDispatch.cpp
--- a/Engine/Plugins/FX/Niagara/Source/Niagara/Private/NiagaraGpuComputeDispatch.cpp
+++ b/Engine/Plugins/FX/Niagara/Source/Niagara/Private/NiagaraGpuComputeDispatch.cpp
@@ -78,11 +78,6 @@ static FAutoConsoleVariableRef CVarNiagaraBatcherFreeBufferEarly(
 
 const FName FNiagaraGpuComputeDispatch::Name(TEXT("FNiagaraGpuComputeDispatch"));
 
-BEGIN_SHADER_PARAMETER_STRUCT(FNiagaraComputePassParameters, )
-	SHADER_PARAMETER_STRUCT_INCLUDE(FNiagaraSceneTextureParameters,	SceneTextures)
-	SHADER_PARAMETER_STRUCT_REF(FViewUniformShaderParameters, View)
-END_SHADER_PARAMETER_STRUCT()
-
 namespace FNiagaraGpuComputeDispatchLocal
 {
 #if WITH_MGPU
@@ -110,14 +105,6 @@ namespace FNiagaraGpuComputeDispatchLocal
 		ECVF_Default
 	);
 
-	static int32 GAddDispatchGroupDrawEvent = 0;
-	static FAutoConsoleVariableRef CVarAddDispatchGroupDrawEvent(
-		TEXT("fx.Niagara.Batcher.AddDispatchGroupDrawEvent"),
-		GAddDispatchGroupDrawEvent,
-		TEXT("Add a draw event marker around each dispatch group."),
-		ECVF_Default
-	);
-
 	#if !WITH_EDITOR
 		constexpr int32 GDebugLogging = 0;
 	#else
@@ -151,6 +138,8 @@ namespace FNiagaraGpuComputeDispatchLocal
 	}
 }
 
+//////////////////////////////////////////////////////////////////////////
+
 FFXSystemInterface* FNiagaraGpuComputeDispatch::GetInterface(const FName& InName)
 {
 	return InName == Name ? this : nullptr;
@@ -480,148 +469,12 @@ void FNiagaraGpuComputeDispatch::FinishDispatches()
 	}
 }
 
-void FNiagaraGpuComputeDispatch::ResetDataInterfaces(FRHICommandList& RHICmdList, const FNiagaraGPUSystemTick& Tick, const FNiagaraComputeInstanceData& InstanceData) const
-{
-	// Note: All stages will contain the same bindings so if they are valid for one they are valid for all, this could change in the future
-	const FNiagaraShaderRef& ComputeShader = InstanceData.Context->GPUScript_RT->GetShader(0);
-	const TMemoryImageArray<FNiagaraDataInterfaceParamRef>& DIParameters = ComputeShader->GetDIParameters();
-
-	uint32 InterfaceIndex = 0;
-	for (FNiagaraDataInterfaceProxy* Interface : InstanceData.DataInterfaceProxies)
-	{
-		const FNiagaraDataInterfaceParamRef& DIParam = DIParameters[InterfaceIndex];
-		if (DIParam.Parameters.IsValid() || (DIParam.ShaderParametersOffset != INDEX_NONE))
-		{
-			const FNiagaraDataInterfaceArgs TmpContext(Interface, Tick.SystemInstanceID, Tick.SystemGpuComputeProxy->GetSystemLWCTile(), this);
-			Interface->ResetData(RHICmdList, TmpContext);
-		}
-		InterfaceIndex++;
-	}
-}
-
 FNiagaraDataInterfaceProxyRW* FNiagaraGpuComputeDispatch::FindIterationInterface(FNiagaraComputeInstanceData* Instance, const uint32 SimulationStageIndex) const
 {
 	// Determine if the iteration is outputting to a custom data size
 	return Instance->FindIterationInterface(SimulationStageIndex);
 }
 
-void FNiagaraGpuComputeDispatch::PreStageInterface(FRHICommandList& RHICmdList, const FNiagaraGPUSystemTick& Tick, const FNiagaraComputeInstanceData& InstanceData, const FNiagaraSimStageData& SimStageData, TSet<FNiagaraDataInterfaceProxy*>& ProxiesToFinalize) const
-{
-	// Note: All stages will contain the same bindings so if they are valid for one they are valid for all, this could change in the future
-	const FNiagaraShaderRef& ComputeShader = InstanceData.Context->GPUScript_RT->GetShader(0);
-	const TMemoryImageArray<FNiagaraDataInterfaceParamRef>& DIParameters = ComputeShader->GetDIParameters();
-
-	uint32 InterfaceIndex = 0;
-	for (FNiagaraDataInterfaceProxy* Interface : InstanceData.DataInterfaceProxies)
-	{
-		const FNiagaraDataInterfaceParamRef& DIParam = DIParameters[InterfaceIndex];
-		if (DIParam.Parameters.IsValid() || (DIParam.ShaderParametersOffset != INDEX_NONE))
-		{
-			const FNiagaraDataInterfaceStageArgs TmpContext(Interface, Tick.SystemInstanceID, Tick.SystemGpuComputeProxy->GetSystemLWCTile(), this, &InstanceData, &SimStageData, InstanceData.IsOutputStage(Interface, SimStageData.StageIndex), InstanceData.IsIterationStage(Interface, SimStageData.StageIndex));
-			Interface->PreStage(RHICmdList, TmpContext);
-
-			if (Interface->RequiresPreStageFinalize())
-			{
-				ProxiesToFinalize.Add(Interface);
-			}
-		}
-		InterfaceIndex++;
-	}
-}
-
-void FNiagaraGpuComputeDispatch::PostStageInterface(FRHICommandList& RHICmdList, const FNiagaraGPUSystemTick& Tick, const FNiagaraComputeInstanceData& InstanceData, const FNiagaraSimStageData& SimStageData, TSet<FNiagaraDataInterfaceProxy*>& ProxiesToFinalize) const
-{
-	// Note: All stages will contain the same bindings so if they are valid for one they are valid for all, this could change in the future
-	const FNiagaraShaderRef& ComputeShader = InstanceData.Context->GPUScript_RT->GetShader(0);
-	const TMemoryImageArray<FNiagaraDataInterfaceParamRef>& DIParameters = ComputeShader->GetDIParameters();
-
-	uint32 InterfaceIndex = 0;
-	for (FNiagaraDataInterfaceProxy* Interface : InstanceData.DataInterfaceProxies)
-	{
-		const FNiagaraDataInterfaceParamRef& DIParam = DIParameters[InterfaceIndex];
-		if (DIParam.Parameters.IsValid() || (DIParam.ShaderParametersOffset != INDEX_NONE))
-		{
-			const FNiagaraDataInterfaceStageArgs TmpContext(Interface, Tick.SystemInstanceID, Tick.SystemGpuComputeProxy->GetSystemLWCTile(), this, &InstanceData, &SimStageData, InstanceData.IsOutputStage(Interface, SimStageData.StageIndex), InstanceData.IsIterationStage(Interface, SimStageData.StageIndex));
-			Interface->PostStage(RHICmdList, TmpContext);
-
-			if (Interface->RequiresPostStageFinalize())
-			{
-				ProxiesToFinalize.Add(Interface);
-			}
-		}
-		InterfaceIndex++;
-	}
-}
-
-void FNiagaraGpuComputeDispatch::PostSimulateInterface(FRHICommandList& RHICmdList, const FNiagaraGPUSystemTick& Tick, const FNiagaraComputeInstanceData& InstanceData) const
-{
-	// Note: All stages will contain the same bindings so if they are valid for one they are valid for all, this could change in the future
-	const FNiagaraShaderRef& ComputeShader = InstanceData.Context->GPUScript_RT->GetShader(0);
-	const TMemoryImageArray<FNiagaraDataInterfaceParamRef>& DIParameters = ComputeShader->GetDIParameters();
-
-	uint32 InterfaceIndex = 0;
-	for (FNiagaraDataInterfaceProxy* Interface : InstanceData.DataInterfaceProxies)
-	{
-		const FNiagaraDataInterfaceParamRef& DIParam = DIParameters[InterfaceIndex];
-		if (DIParam.Parameters.IsValid() || (DIParam.ShaderParametersOffset != INDEX_NONE))
-		{
-			const FNiagaraDataInterfaceArgs TmpContext(Interface, Tick.SystemInstanceID, Tick.SystemGpuComputeProxy->GetSystemLWCTile(), this);
-			Interface->PostSimulate(RHICmdList, TmpContext);
-		}
-		InterfaceIndex++;
-	}
-}
-
-
-void FNiagaraGpuComputeDispatch::UpdateFreeIDsListSizesBuffer(FRHICommandList& RHICmdList, uint32 NumInstances)
-{
-	if (NumInstances > NumAllocatedFreeIDListSizes)
-	{
-		constexpr uint32 ALLOC_CHUNK_SIZE = 128;
-		NumAllocatedFreeIDListSizes = Align(NumInstances, ALLOC_CHUNK_SIZE);
-		if (FreeIDListSizesBuffer.Buffer)
-		{
-			FreeIDListSizesBuffer.Release();
-		}
-		FreeIDListSizesBuffer.Initialize(TEXT("NiagaraFreeIDListSizes"), sizeof(uint32), NumAllocatedFreeIDListSizes, EPixelFormat::PF_R32_SINT, ERHIAccess::UAVCompute, BUF_Static);
-	}
-
-	{
-		SCOPED_DRAW_EVENT(RHICmdList, NiagaraGPUComputeClearFreeIDListSizes);
-		RHICmdList.Transition(FRHITransitionInfo(FreeIDListSizesBuffer.UAV, ERHIAccess::UAVCompute, ERHIAccess::UAVCompute));
-		NiagaraFillGPUIntBuffer(RHICmdList, FeatureLevel, FreeIDListSizesBuffer, 0);
-	}
-}
-
-void FNiagaraGpuComputeDispatch::UpdateFreeIDBuffers(FRHICommandList& RHICmdList, TConstArrayView<FNiagaraComputeExecutionContext*> Instances)
-{
-	if (Instances.Num() == 0)
-	{
-		return;
-	}
-
-	SCOPED_DRAW_EVENT(RHICmdList, NiagaraGPUComputeFreeIDs);
-	SCOPED_GPU_STAT(RHICmdList, NiagaraGPUComputeFreeIDs);
-
-	TArray<FRHITransitionInfo, TMemStackAllocator<>> TransitionsBefore;
-	TransitionsBefore.Emplace(FreeIDListSizesBuffer.UAV, ERHIAccess::UAVCompute, ERHIAccess::UAVCompute);
-	RHICmdList.Transition(TransitionsBefore);
-
-	check((uint32)Instances.Num() <= NumAllocatedFreeIDListSizes);
-
-	RHICmdList.BeginUAVOverlap(FreeIDListSizesBuffer.UAV);
-	for (uint32 iInstance = 0; iInstance < (uint32)Instances.Num(); ++iInstance)
-	{
-		FNiagaraComputeExecutionContext* ComputeContext = Instances[iInstance];
-		FNiagaraDataSet* MainDataSet = ComputeContext->MainDataSet;
-		FNiagaraDataBuffer* CurrentData = ComputeContext->MainDataSet->GetCurrentData();
-
-		SCOPED_DRAW_EVENTF(RHICmdList, NiagaraGPUComputeFreeIDsEmitter, TEXT("Update Free ID Buffer - %s"), ComputeContext->GetDebugSimName());
-		NiagaraComputeGPUFreeIDs(RHICmdList, FeatureLevel, MainDataSet->GetGPUNumAllocatedIDs(), CurrentData->GetGPUIDToIndexTable().SRV, MainDataSet->GetGPUFreeIDs(), FreeIDListSizesBuffer, iInstance);
-	}
-	RHICmdList.EndUAVOverlap(FreeIDListSizesBuffer.UAV);
-}
-
 void FNiagaraGpuComputeDispatch::DumpDebugFrame()
 {
 	// Anything doing?
@@ -723,9 +576,9 @@ void FNiagaraGpuComputeDispatch::DumpDebugFrame()
 			if (DispatchGroup.FreeIDUpdates.Num() > 0)
 			{
 				UE_LOG(LogNiagara, Warning, TEXT("====== FreeIDUpdates"));
-				for (FNiagaraComputeExecutionContext* ComputeContext : DispatchGroup.FreeIDUpdates)
+				for (const FNiagaraGpuFreeIDUpdate& FreeIDUpdate : DispatchGroup.FreeIDUpdates)
 				{
-					UE_LOG(LogNiagara, Warning, TEXT("ComputeContext(%p) Emitter(%s)"), ComputeContext, ComputeContext->GetDebugSimName());
+					UE_LOG(LogNiagara, Warning, TEXT("ComputeContext(%p) Emitter(%s)"), FreeIDUpdate.ComputeContext, FreeIDUpdate.ComputeContext->GetDebugSimName());
 				}
 			}
 		}
@@ -1049,9 +902,9 @@ void FNiagaraGpuComputeDispatch::PrepareTicksForProxy(FRHICommandListImmediate&
 					continue;
 				}
 
-				if ( ComputeContext->MainDataSet->RequiresPersistentIDs() )
+				if ( ComputeContext->MainDataSet->RequiresPersistentIDs() && ComputeContext->bHasTickedThisFrame_RT )
 				{
-					FinalDispatchGroup.FreeIDUpdates.Add(ComputeContext);
+					FinalDispatchGroup.FreeIDUpdates.Emplace(ComputeContext);
 				}
 			}
 		}
@@ -1133,11 +986,21 @@ void FNiagaraGpuComputeDispatch::PrepareAllTicks(FRHICommandListImmediate& RHICm
 	}
 }
 
-void FNiagaraGpuComputeDispatch::ExecuteTicks(FRHICommandList& RHICmdList, ENiagaraGpuComputeTickStage::Type TickStage)
+void FNiagaraGpuComputeDispatch::ExecuteTicks(FRDGBuilder& GraphBuilder, TConstArrayView<FViewInfo> Views, ENiagaraGpuComputeTickStage::Type TickStage)
 {
 #if WITH_MGPU
-	WaitForMultiGPUBuffers(RHICmdList, TickStage);
-#endif // WITH_MGPU
+	if (StageToWaitForGPUTransfers == TickStage)
+	{
+ 		GraphBuilder.AddPass(
+			RDG_EVENT_NAME("Niagara::ExecuteTicksPre"),
+			ERDGPassFlags::None,
+			[this, TickStage, NumDispatchGroups=DispatchListPerStage[TickStage].DispatchGroups.Num()](FRHICommandListImmediate& RHICmdList)
+			{
+				WaitForMultiGPUBuffers(RHICmdList, TickStage);
+			}
+		);
+	}
+#endif
 
 	// Anything to execute for this stage
 	FNiagaraGpuDispatchList& DispatchList = DispatchListPerStage[TickStage];
@@ -1145,224 +1008,339 @@ void FNiagaraGpuComputeDispatch::ExecuteTicks(FRHICommandList& RHICmdList, ENiag
 	{
 		return;
 	}
+	const int32 StageStartTotalDispatches = TotalDispatchesThisFrame;
 
 	TRACE_CPUPROFILER_EVENT_SCOPE(FNiagaraGpuComputeDispatch_ExecuteTicks);
-	SCOPED_DRAW_EVENTF(RHICmdList, FNiagaraGpuComputeDispatch_ExecuteTicks, TEXT("FNiagaraGpuComputeDispatch_ExecuteTicks - TickStage(%d)"), TickStage);
 	SCOPE_CYCLE_COUNTER(STAT_NiagaraGPUSimTick_RT);
-	SCOPED_GPU_STAT(RHICmdList, NiagaraGPUSimulation);
+	RDG_GPU_STAT_SCOPE(GraphBuilder, NiagaraGPUSimulation);
+	RDG_RHI_EVENT_SCOPE(GraphBuilder, NiagaraGpuComputeDispatch);	//-TODO:RDG: Show TickStage
 
-	FMemMark Mark(FMemStack::Get());
-	TArray<FRHITransitionInfo, TMemStackAllocator<>> TransitionsBefore;
-	TArray<FRHITransitionInfo, TMemStackAllocator<>> TransitionsAfter;
-	TArray<FNiagaraDataBuffer*, TMemStackAllocator<>> IDToIndexInit;
-
-#if WITH_NIAGARA_GPU_PROFILER
-	GPUProfilerPtr->BeginStage(RHICmdList, TickStage, DispatchList.DispatchGroups.Num());
-	const int32 StageStartTotalDispatches = TotalDispatchesThisFrame;
-#endif
+	// Setup Parameters that can be read from data interfaces
+	FNiagaraSceneTextureParameters SceneTextures;
+	GNiagaraViewDataManager.GetSceneTextureParameters(GraphBuilder, Views.Num() > 0 ? GetViewFamilyInfo(Views).GetSceneTexturesChecked() : nullptr, SceneTextures);
+	NiagaraSceneTextures = &SceneTextures;
+	CurrentPassViews = Views;
 
+	// Loop over dispatches
 	for ( const FNiagaraGpuDispatchGroup& DispatchGroup : DispatchList.DispatchGroups )
 	{
-		SCOPED_CONDITIONAL_DRAW_EVENT(RHICmdList, NiagaraDispatchGroup, FNiagaraGpuComputeDispatchLocal::GAddDispatchGroupDrawEvent);
-
 		const bool bIsFirstGroup = &DispatchGroup == &DispatchList.DispatchGroups[0];
 		const bool bIsLastGroup = &DispatchGroup == &DispatchList.DispatchGroups.Last();
 
-		// Generate transitions and discover free / ID table updates
-		TransitionsBefore.Reserve(DispatchGroup.DispatchInstances.Num() * 3);
-		TransitionsAfter.Reserve(DispatchGroup.DispatchInstances.Num() * 3);
+		// Consume per tick data from the game thread
+		//-TODO: This does not work currently as some senders assume the data will not be deferred processed
+		//for (FNiagaraGPUSystemTick* Tick : DispatchGroup.TicksWithPerInstanceData)
+		//{
+		//	uint8* BasePointer = reinterpret_cast<uint8*>(Tick->DIInstanceData->PerInstanceDataForRT);
+		//
+		//	for (const auto& Pair : Tick->DIInstanceData->InterfaceProxiesToOffsets)
+		//	{
+		//		FNiagaraDataInterfaceProxy* Proxy = Pair.Key;
+		//		uint8* InstanceDataPtr = BasePointer + Pair.Value;
+		//		Proxy->ConsumePerInstanceDataFromGameThread(InstanceDataPtr, Tick->SystemInstanceID);
+		//	}
+		//}
+
+		// Generate pre / post stage graph data
+		//-OPT: Once we have RenderGraph dependencies between stages a lot of this tracking can go away as the graph will handle it
+		TArray<FRHITransitionInfo> PreStageTransitions;
+		TArray<FRHITransitionInfo> PostStageTransitions;
+		TArray<TPair<FRHIUnorderedAccessView*, int32>> PreStageIDToIndexInit;
+		PreStageTransitions.Reserve(DispatchGroup.DispatchInstances.Num() * 3);
+		PostStageTransitions.Reserve(DispatchGroup.DispatchInstances.Num() * 3);
+
 		for (const FNiagaraGpuDispatchInstance& DispatchInstance : DispatchGroup.DispatchInstances)
 		{
 			if (FNiagaraDataBuffer* DestinationBuffer = DispatchInstance.SimStageData.Destination)
 			{
-				FNiagaraGpuComputeDispatchLocal::AddDataBufferTransitions(TransitionsBefore, TransitionsAfter, DestinationBuffer);
+				FNiagaraGpuComputeDispatchLocal::AddDataBufferTransitions(PreStageTransitions, PostStageTransitions, DestinationBuffer);
 			}
 
 			FNiagaraComputeExecutionContext* ComputeContext = DispatchInstance.InstanceData.Context;
 			const bool bRequiresPersistentIDs = ComputeContext->MainDataSet->RequiresPersistentIDs();
 			if (bRequiresPersistentIDs)
 			{
-				if ( FNiagaraDataBuffer* IDtoIndexBuffer = DispatchInstance.SimStageData.Destination )
+				if ( FNiagaraDataBuffer* IDToIndexBuffer = DispatchInstance.SimStageData.Destination )
 				{
-					IDToIndexInit.Add(IDtoIndexBuffer);
-					TransitionsBefore.Emplace(IDtoIndexBuffer->GetGPUIDToIndexTable().UAV, ERHIAccess::SRVCompute, ERHIAccess::UAVCompute);
-					TransitionsAfter.Emplace(IDtoIndexBuffer->GetGPUIDToIndexTable().UAV, ERHIAccess::UAVCompute, ERHIAccess::SRVCompute);
+					FRHIUnorderedAccessView* BufferUAV = IDToIndexBuffer->GetGPUIDToIndexTable().UAV;
+					const int32 BufferNumEntries = IDToIndexBuffer->GetGPUIDToIndexTable().NumBytes / sizeof(int32);
+
+					PreStageIDToIndexInit.Emplace(BufferUAV, BufferNumEntries);
+					PreStageTransitions.Emplace(BufferUAV, ERHIAccess::SRVCompute, ERHIAccess::UAVCompute);
+					PostStageTransitions.Emplace(BufferUAV, ERHIAccess::UAVCompute, ERHIAccess::SRVCompute);
 				}
 			}
 		}
 
-		TransitionsBefore.Emplace(GPUInstanceCounterManager.GetInstanceCountBuffer().UAV, bIsFirstGroup ? FNiagaraGPUInstanceCountManager::kCountBufferDefaultState : ERHIAccess::UAVCompute, ERHIAccess::UAVCompute);
+		PreStageTransitions.Emplace(GPUInstanceCounterManager.GetInstanceCountBuffer().UAV, bIsFirstGroup ? FNiagaraGPUInstanceCountManager::kCountBufferDefaultState : ERHIAccess::UAVCompute, ERHIAccess::UAVCompute);
 		if (bIsLastGroup)
 		{
-			TransitionsAfter.Emplace(GPUInstanceCounterManager.GetInstanceCountBuffer().UAV, ERHIAccess::UAVCompute, FNiagaraGPUInstanceCountManager::kCountBufferDefaultState);
+			PostStageTransitions.Emplace(GPUInstanceCounterManager.GetInstanceCountBuffer().UAV, ERHIAccess::UAVCompute, FNiagaraGPUInstanceCountManager::kCountBufferDefaultState);
 		}
 
 		if (DispatchGroup.FreeIDUpdates.Num() > 0)
 		{
-			for (FNiagaraComputeExecutionContext* ComputeContext : DispatchGroup.FreeIDUpdates )
-			{
-				TransitionsAfter.Emplace(ComputeContext->MainDataSet->GetGPUFreeIDs().UAV, ERHIAccess::SRVCompute, ERHIAccess::UAVCompute);
-			}
-		}
-
-		// Consume per tick data from the game thread
-		//for (FNiagaraGPUSystemTick* Tick : DispatchGroup.TicksWithPerInstanceData)
-		//{
-		//	uint8* BasePointer = reinterpret_cast<uint8*>(Tick->DIInstanceData->PerInstanceDataForRT);
-
-		//	for (const auto& Pair : Tick->DIInstanceData->InterfaceProxiesToOffsets)
-		//	{
-		//		FNiagaraDataInterfaceProxy* Proxy = Pair.Key;
-		//		uint8* InstanceDataPtr = BasePointer + Pair.Value;
-		//		Proxy->ConsumePerInstanceDataFromGameThread(InstanceDataPtr, Tick->SystemInstanceID);
-		//	}
-		//}
-
-		// Execute Before Transitions
-		RHICmdList.Transition(TransitionsBefore);
-		TransitionsBefore.Reset();
-
-		// Initialize the IDtoIndex tables
-		if (IDToIndexInit.Num() > 0)
-		{
-			SCOPED_DRAW_EVENT(RHICmdList, NiagaraGPUComputeClearIDToIndexBuffer);
-
-			TArray<FRHITransitionInfo, TMemStackAllocator<>> IDToIndexTransitions;
-			IDToIndexTransitions.Reserve(IDToIndexInit.Num());
-
-			for (FNiagaraDataBuffer* IDtoIndexBuffer : IDToIndexInit)
+			PostStageTransitions.Reserve(PostStageTransitions.Num() + DispatchGroup.FreeIDUpdates.Num());
+			for (const FNiagaraGpuFreeIDUpdate& FreeIDUpdate : DispatchGroup.FreeIDUpdates)
 			{
-				NiagaraFillGPUIntBuffer(RHICmdList, FeatureLevel, IDtoIndexBuffer->GetGPUIDToIndexTable(), INDEX_NONE);
-				IDToIndexTransitions.Emplace(IDtoIndexBuffer->GetGPUIDToIndexTable().UAV, ERHIAccess::UAVCompute, ERHIAccess::UAVCompute);
+				PostStageTransitions.Emplace(FreeIDUpdate.ComputeContext->MainDataSet->GetGPUFreeIDs().UAV, ERHIAccess::SRVCompute, ERHIAccess::UAVCompute);
 			}
-			IDToIndexInit.Reset();
-			RHICmdList.Transition(IDToIndexTransitions);
 		}
 
-		// Execute PreStage
+		// Pre Stage
 		{
 			TSet<FNiagaraDataInterfaceProxy*> ProxiesToFinalize;
 			for (const FNiagaraGpuDispatchInstance& DispatchInstance : DispatchGroup.DispatchInstances)
 			{
-				PreStageInterface(RHICmdList, DispatchInstance.Tick, DispatchInstance.InstanceData, DispatchInstance.SimStageData, ProxiesToFinalize);
+				PreStageInterface(GraphBuilder, DispatchInstance.Tick, DispatchInstance.InstanceData, DispatchInstance.SimStageData, ProxiesToFinalize);
 			}
-
 			for (FNiagaraDataInterfaceProxy* ProxyToFinalize : ProxiesToFinalize)
 			{
-				ProxyToFinalize->FinalizePreStage(RHICmdList, this);
+				ProxyToFinalize->FinalizePreStage(GraphBuilder, *this);
 			}
 		}
 
+ 		GraphBuilder.AddPass(
+			{},//RDG_EVENT_NAME("Niagara::ExecuteTicks::DispatchGroupPre"),
+			ERDGPassFlags::None,
+			[this, PreStageTransitions=MoveTemp(PreStageTransitions), PreStageIDToIndexInit=MoveTemp(PreStageIDToIndexInit), DispatchInstances=MakeArrayView(DispatchGroup.DispatchInstances)](FRHICommandListImmediate& RHICmdList)
+			{
+				// Execute Before Transitions
+				RHICmdList.Transition(PreStageTransitions);
+
+				// Initialize the IDtoIndex tables
+				if (PreStageIDToIndexInit.Num() > 0)
+				{
+					SCOPED_DRAW_EVENT(RHICmdList, NiagaraGPUComputeClearIDToIndexBuffer);
+
+					TArray<FRHITransitionInfo> IDToIndexTransitions;
+					IDToIndexTransitions.Reserve(PreStageIDToIndexInit.Num());
+
+					for (const TPair<FRHIUnorderedAccessView*, int32>& IDtoIndexBuffer : PreStageIDToIndexInit)
+					{
+						NiagaraFillGPUIntBuffer(RHICmdList, FeatureLevel, IDtoIndexBuffer.Key, IDtoIndexBuffer.Value, INDEX_NONE);
+						IDToIndexTransitions.Emplace(IDtoIndexBuffer.Key, ERHIAccess::UAVCompute, ERHIAccess::UAVCompute);
+					}
+					RHICmdList.Transition(IDToIndexTransitions);
+				}
+
+				// Execute PreStage
+				TSet<FNiagaraDataInterfaceProxy*> ProxiesToFinalize;
+				for (const FNiagaraGpuDispatchInstance& DispatchInstance : DispatchInstances)
+				{
+					LegacyPreStageInterface(RHICmdList, DispatchInstance.Tick, DispatchInstance.InstanceData, DispatchInstance.SimStageData, ProxiesToFinalize);
+				}
+
+				for (FNiagaraDataInterfaceProxy* ProxyToFinalize : ProxiesToFinalize)
+				{
+					ProxyToFinalize->FinalizePreStage(RHICmdList, this);
+				}
+
+				//-TODO:RDG: Remove once legacy is no longer supported
+				for (const FNiagaraGpuDispatchInstance& DispatchInstance : DispatchInstances)
+				{
+					if (DispatchInstance.InstanceData.bResetData && DispatchInstance.SimStageData.bFirstStage )
+					{
+						LegacyResetDataInterfaces(RHICmdList, DispatchInstance.Tick, DispatchInstance.InstanceData);
+					}
+				}
+				
+				RHICmdList.BeginUAVOverlap(GPUInstanceCounterManager.GetInstanceCountBuffer().UAV);
+			}
+		);
+
 		// Execute Stage
-		RHICmdList.BeginUAVOverlap(GPUInstanceCounterManager.GetInstanceCountBuffer().UAV);
 		for (const FNiagaraGpuDispatchInstance& DispatchInstance : DispatchGroup.DispatchInstances)
 		{
 			++FNiagaraComputeExecutionContext::TickCounter;
-			if (DispatchInstance.InstanceData.bResetData && DispatchInstance.SimStageData.bFirstStage )
+			if (DispatchInstance.InstanceData.bResetData && DispatchInstance.SimStageData.bFirstStage)
 			{
-				ResetDataInterfaces(RHICmdList, DispatchInstance.Tick, DispatchInstance.InstanceData);
+				ResetDataInterfaces(GraphBuilder, DispatchInstance.Tick, DispatchInstance.InstanceData);
 			}
 
-			FNiagaraGpuProfileScope GpuProfileDispatchScope(RHICmdList, this, DispatchInstance);
-			DispatchStage(RHICmdList, DispatchInstance.Tick, DispatchInstance.InstanceData, DispatchInstance.SimStageData);
+			DispatchStage(GraphBuilder, DispatchInstance.Tick, DispatchInstance.InstanceData, DispatchInstance.SimStageData);
 		}
-		RHICmdList.EndUAVOverlap(GPUInstanceCounterManager.GetInstanceCountBuffer().UAV);
 
-		// Execute PostStage
-		{
-			TSet<FNiagaraDataInterfaceProxy*> ProxiesToFinalize;
-			for (const FNiagaraGpuDispatchInstance& DispatchInstance : DispatchGroup.DispatchInstances)
+		// Execute legacy Post Stage
+ 		GraphBuilder.AddPass(
+			{},//RDG_EVENT_NAME("Niagara::ExecuteTicks::DispatchGroupPost"),
+			ERDGPassFlags::None,
+			[this, PostStageTransitions=MoveTemp(PostStageTransitions), DispatchInstances=MakeArrayView(DispatchGroup.DispatchInstances)](FRHICommandListImmediate& RHICmdList)
 			{
-				PostStageInterface(RHICmdList, DispatchInstance.Tick, DispatchInstance.InstanceData, DispatchInstance.SimStageData, ProxiesToFinalize);
-				if ( DispatchInstance.SimStageData.bLastStage )
-				{
-					PostSimulateInterface(RHICmdList, DispatchInstance.Tick, DispatchInstance.InstanceData);
-
-					// Update CurrentData with the latest information as things like ParticleReads can use this data
-					FNiagaraComputeExecutionContext* ComputeContext = DispatchInstance.InstanceData.Context;
-					const FNiagaraSimStageData& FinalSimStageData = DispatchInstance.SimStageData;
-					FNiagaraDataBuffer* FinalSimStageDataBuffer = FinalSimStageData.Destination != nullptr ? FinalSimStageData.Destination : FinalSimStageData.Source;
-					check(FinalSimStageDataBuffer != nullptr);
+				RHICmdList.EndUAVOverlap(GPUInstanceCounterManager.GetInstanceCountBuffer().UAV);
 
-					// If we are setting the data to render we need to ensure we switch back to the original CurrentData then swap the GPU buffers into it
-					if (DispatchInstance.SimStageData.bSetDataToRender)
+				TSet<FNiagaraDataInterfaceProxy*> ProxiesToFinalize;
+				for (const FNiagaraGpuDispatchInstance& DispatchInstance : DispatchInstances)
+				{
+					LegacyPostStageInterface(RHICmdList, DispatchInstance.Tick, DispatchInstance.InstanceData, DispatchInstance.SimStageData, ProxiesToFinalize);
+					if ( DispatchInstance.SimStageData.bLastStage )
 					{
-						check(ComputeContext->DataSetOriginalBuffer_RT != nullptr);
-						FNiagaraDataBuffer* CurrentData = ComputeContext->DataSetOriginalBuffer_RT;
-						ComputeContext->DataSetOriginalBuffer_RT = nullptr;
+						LegacyPostSimulateInterface(RHICmdList, DispatchInstance.Tick, DispatchInstance.InstanceData);
+					}
+				}
 
-						ComputeContext->MainDataSet->CurrentData = CurrentData;
-						CurrentData->SwapGPU(FinalSimStageDataBuffer);
+				for (FNiagaraDataInterfaceProxy* ProxyToFinalize : ProxiesToFinalize)
+				{
+					ProxyToFinalize->FinalizePostStage(RHICmdList, this);
+				}
 
-						// Mark data as ready for anyone who picks up the buffer on the next frame
-						CurrentData->SetGPUDataReadyStage(ENiagaraGpuComputeTickStage::First);
-                    
-						ComputeContext->SetTranslucentDataToRender(nullptr);
-						ComputeContext->SetDataToRender(CurrentData);
+				// Execute After Transitions
+				RHICmdList.Transition(PostStageTransitions);
+			}
+		);
 
-#if WITH_MGPU
-						if (bAFREnabled)
-						{
-							AddAFRBuffer(CurrentData->GetGPUBufferFloat().Buffer);
-							AddAFRBuffer(CurrentData->GetGPUBufferHalf().Buffer);
-							AddAFRBuffer(CurrentData->GetGPUBufferInt().Buffer);
-							if (ComputeContext->MainDataSet->RequiresPersistentIDs())
-							{
-								AddAFRBuffer(ComputeContext->MainDataSet->GetGPUFreeIDs().Buffer);
-							}
-						}
-						if (bCrossGPUTransferEnabled)
-						{
-							AddCrossGPUTransfer(RHICmdList, CurrentData->GetGPUBufferFloat().Buffer);
-							AddCrossGPUTransfer(RHICmdList, CurrentData->GetGPUBufferHalf().Buffer);
-							AddCrossGPUTransfer(RHICmdList, CurrentData->GetGPUBufferInt().Buffer);
-						}
-#endif // WITH_MGPU
-					}
-					// If this is not the final tick of the final stage we need set our temporary buffer for data interfaces, etc, that may snoop from CurrentData
-					else
-					{
-						ComputeContext->MainDataSet->CurrentData = FinalSimStageDataBuffer;
-					}
-				}
+		// Execute PostStage / Post Simulate
+		TSet<FNiagaraDataInterfaceProxy*> ProxiesToFinalize;
+		for (const FNiagaraGpuDispatchInstance& DispatchInstance : DispatchGroup.DispatchInstances)
+		{
+			PostStageInterface(GraphBuilder, DispatchInstance.Tick, DispatchInstance.InstanceData, DispatchInstance.SimStageData, ProxiesToFinalize);
+			if (DispatchInstance.SimStageData.bLastStage)
+			{
+				PostSimulateInterface(GraphBuilder, DispatchInstance.Tick, DispatchInstance.InstanceData, DispatchInstance.SimStageData);
 			}
 
-			for (FNiagaraDataInterfaceProxy* ProxyToFinalize : ProxiesToFinalize)
+			// Update CurrentData with the latest information as things like ParticleReads can use this data
+			FNiagaraComputeExecutionContext* ComputeContext = DispatchInstance.InstanceData.Context;
+			const FNiagaraSimStageData& FinalSimStageData = DispatchInstance.SimStageData;
+			FNiagaraDataBuffer* FinalSimStageDataBuffer = FinalSimStageData.Destination != nullptr ? FinalSimStageData.Destination : FinalSimStageData.Source;
+			check(FinalSimStageDataBuffer != nullptr);
+
+			// If we are setting the data to render we need to ensure we switch back to the original CurrentData then swap the GPU buffers into it
+			if (DispatchInstance.SimStageData.bSetDataToRender)
+			{
+				check(ComputeContext->DataSetOriginalBuffer_RT != nullptr);
+				FNiagaraDataBuffer* CurrentData = ComputeContext->DataSetOriginalBuffer_RT;
+				ComputeContext->DataSetOriginalBuffer_RT = nullptr;
+
+				ComputeContext->MainDataSet->CurrentData = CurrentData;
+				CurrentData->SwapGPU(FinalSimStageDataBuffer);
+
+				// Mark data as ready for anyone who picks up the buffer on the next frame
+				CurrentData->SetGPUDataReadyStage(ENiagaraGpuComputeTickStage::First);
+                    
+				ComputeContext->SetTranslucentDataToRender(nullptr);
+				ComputeContext->SetDataToRender(CurrentData);
+
+//-TODO:RDG:mGPU:
+//#if WITH_MGPU
+//				if (bAFREnabled)
+//				{
+//					AddAFRBuffer(CurrentData->GetGPUBufferFloat().Buffer);
+//					AddAFRBuffer(CurrentData->GetGPUBufferHalf().Buffer);
+//					AddAFRBuffer(CurrentData->GetGPUBufferInt().Buffer);
+//					if (ComputeContext->MainDataSet->RequiresPersistentIDs())
+//					{
+//						AddAFRBuffer(ComputeContext->MainDataSet->GetGPUFreeIDs().Buffer);
+//					}
+//				}
+//				if (bCrossGPUTransferEnabled)
+//				{
+//					AddCrossGPUTransfer(RHICmdList, CurrentData->GetGPUBufferFloat().Buffer);
+//					AddCrossGPUTransfer(RHICmdList, CurrentData->GetGPUBufferHalf().Buffer);
+//					AddCrossGPUTransfer(RHICmdList, CurrentData->GetGPUBufferInt().Buffer);
+//				}
+//#endif // WITH_MGPU
+			}
+			// If this is not the final tick of the final stage we need set our temporary buffer for data interfaces, etc, that may snoop from CurrentData
+			else
 			{
-				ProxyToFinalize->FinalizePostStage(RHICmdList, this);
+				ComputeContext->MainDataSet->CurrentData = FinalSimStageDataBuffer;
 			}
 		}
 
-		// Execute After Transitions
-		RHICmdList.Transition(TransitionsAfter);
-		TransitionsAfter.Reset();
+		for (FNiagaraDataInterfaceProxy* ProxyToFinalize : ProxiesToFinalize)
+		{
+			ProxyToFinalize->FinalizePostStage(GraphBuilder, *this);
+		}
 
-		// Update free IDs
+		// Free ID updates
 		if (DispatchGroup.FreeIDUpdates.Num() > 0)
 		{
-			UpdateFreeIDsListSizesBuffer(RHICmdList, DispatchGroup.FreeIDUpdates.Num());
-			UpdateFreeIDBuffers(RHICmdList, DispatchGroup.FreeIDUpdates);
-
-			for (FNiagaraComputeExecutionContext* ComputeContext : DispatchGroup.FreeIDUpdates)
+			// Grab information for the pass
+			for (const FNiagaraGpuFreeIDUpdate& FreeIDUpdate : DispatchGroup.FreeIDUpdates)
 			{
-				TransitionsAfter.Emplace(ComputeContext->MainDataSet->GetGPUFreeIDs().UAV, ERHIAccess::UAVCompute, ERHIAccess::SRVCompute);
+				FNiagaraDataSet* MainDataSet = FreeIDUpdate.ComputeContext->MainDataSet;
+				FNiagaraDataBuffer* CurrentData = FreeIDUpdate.ComputeContext->MainDataSet->GetCurrentData();
+
+				FreeIDUpdate.IDToIndexSRV		= CurrentData->GetGPUIDToIndexTable().SRV;
+				FreeIDUpdate.FreeIDsUAV			= MainDataSet->GetGPUFreeIDs().UAV;
+				FreeIDUpdate.NumAllocatedIDs	= MainDataSet->GetGPUNumAllocatedIDs();
+
+				check(FreeIDUpdate.IDToIndexSRV != nullptr);
+				check(FreeIDUpdate.FreeIDsUAV != nullptr);
 			}
-			RHICmdList.Transition(TransitionsAfter);
-			TransitionsAfter.Reset();
+
+ 			GraphBuilder.AddPass(
+				{},//RDG_EVENT_NAME("Niagara::ExecuteTicks::FreeIDUpdates"),
+				ERDGPassFlags::None,
+				[this, FreeIDUpdates=MakeArrayView(DispatchGroup.FreeIDUpdates)](FRHICommandListImmediate& RHICmdList)
+				{
+					const uint32 NumFreeIDUpdates = FreeIDUpdates.Num();
+
+					// Initialize the free ID size buffer
+					{
+						if (NumFreeIDUpdates > NumAllocatedFreeIDListSizes)
+						{
+							constexpr uint32 ALLOC_CHUNK_SIZE = 128;
+							NumAllocatedFreeIDListSizes = Align(NumFreeIDUpdates, ALLOC_CHUNK_SIZE);
+							if (FreeIDListSizesBuffer.Buffer)
+							{
+								FreeIDListSizesBuffer.Release();
+							}
+							FreeIDListSizesBuffer.Initialize(TEXT("NiagaraFreeIDListSizes"), sizeof(uint32), NumAllocatedFreeIDListSizes, EPixelFormat::PF_R32_SINT, ERHIAccess::UAVCompute, BUF_Static);
+						}
+
+						SCOPED_DRAW_EVENT(RHICmdList, NiagaraGPUComputeClearFreeIDListSizes);
+						RHICmdList.Transition(FRHITransitionInfo(FreeIDListSizesBuffer.UAV, ERHIAccess::UAVCompute, ERHIAccess::UAVCompute));
+						NiagaraFillGPUIntBuffer(RHICmdList, FeatureLevel, FreeIDListSizesBuffer.UAV, FreeIDListSizesBuffer.NumBytes / sizeof(uint32), 0);
+					}
+
+					// Update Free IDs
+					{
+						SCOPED_DRAW_EVENT(RHICmdList, NiagaraGPUComputeFreeIDs);
+						SCOPED_GPU_STAT(RHICmdList, NiagaraGPUComputeFreeIDs);
+
+						RHICmdList.Transition(FRHITransitionInfo(FreeIDListSizesBuffer.UAV, ERHIAccess::UAVCompute, ERHIAccess::UAVCompute));
+
+						TArray<FRHITransitionInfo> TransitionsFreeIDs;
+						TransitionsFreeIDs.Reserve(NumFreeIDUpdates);
+
+						RHICmdList.BeginUAVOverlap(FreeIDListSizesBuffer.UAV);
+						for (uint32 iInstance=0; iInstance < NumFreeIDUpdates; ++iInstance)
+						{
+							const FNiagaraGpuFreeIDUpdate& FreeIDUpdateInfo = FreeIDUpdates[iInstance];
+							SCOPED_DRAW_EVENTF(RHICmdList, NiagaraGPUComputeFreeIDsEmitter, TEXT("Update Free ID Buffer - %s"), FreeIDUpdateInfo.ComputeContext->GetDebugSimName());
+							NiagaraComputeGPUFreeIDs(RHICmdList, FeatureLevel, FreeIDUpdateInfo.IDToIndexSRV, FreeIDUpdateInfo.NumAllocatedIDs, FreeIDUpdateInfo.FreeIDsUAV, FreeIDListSizesBuffer.UAV, iInstance);
+							TransitionsFreeIDs.Emplace(FreeIDUpdateInfo.FreeIDsUAV, ERHIAccess::UAVCompute, ERHIAccess::SRVCompute);
+						}
+						RHICmdList.EndUAVOverlap(FreeIDListSizesBuffer.UAV);
+						RHICmdList.Transition(TransitionsFreeIDs);
+					}
+				}
+			);
 		}
 	}
 
-#if WITH_NIAGARA_GPU_PROFILER
 	const int32 StageTotalDispatches = TotalDispatchesThisFrame - StageStartTotalDispatches;
-	GPUProfilerPtr->EndStage(RHICmdList, TickStage, StageTotalDispatches);
-#endif
+	GraphBuilder.AddPass(
+		RDG_EVENT_NAME("Niagara::ExecuteTicksPost"),
+		ERDGPassFlags::None,
+		[this, StageTotalDispatches, TickStage, DispatchListPtr=&DispatchList](FRHICommandListImmediate& RHICmdList)
+		{
+			// Clear dispatch groups
+			// We do not release the counts as we won't do that until we finish the dispatches
+			DispatchListPtr->DispatchGroups.Empty();
+		}
+	);
+
+	// Tear down for tick pass
+	NiagaraSceneTextures = &SceneTextures;
+	CurrentPassViews = Views;
 
-	// Clear dispatch groups
-	// We do not release the counts as we won't do that until we finish the dispatches
-	DispatchList.DispatchGroups.Empty();
+	CurrentPassExternalAccessQueue.Submit(GraphBuilder);
 }
 
-void FNiagaraGpuComputeDispatch::DispatchStage(FRHICommandList& RHICmdList, const FNiagaraGPUSystemTick& Tick, const FNiagaraComputeInstanceData& InstanceData, const FNiagaraSimStageData& SimStageData)
+void FNiagaraGpuComputeDispatch::DispatchStage(FRDGBuilder& GraphBuilder, const FNiagaraGPUSystemTick& Tick, const FNiagaraComputeInstanceData& InstanceData, const FNiagaraSimStageData& SimStageData)
 {
 	// Setup source buffer
 	if ( SimStageData.Source != nullptr )
@@ -1431,23 +1409,9 @@ void FNiagaraGpuComputeDispatch::DispatchStage(FRHICommandList& RHICmdList, cons
 
 	checkf(DispatchNumThreads.X * DispatchNumThreads.Y * DispatchNumThreads.Z > 0, TEXT("DispatchNumThreads(%d, %d, %d) is invalid"), DispatchNumThreads.X, DispatchNumThreads.Y, DispatchNumThreads.Z);
 
-	SCOPED_DRAW_EVENTF(RHICmdList, NiagaraGPUSimulationCS,
-		TEXT("NiagaraGpuSim(%s) DispatchCount(%dx%dx%d) Stage(%s %u) Iteration(%u) NumThreads(%dx%dx%d)"),
-		InstanceData.Context->GetDebugSimName(),
-		DispatchCount.X, DispatchCount.Y, DispatchCount.Z,
-		*SimStageData.StageMetaData->SimulationStageName.ToString(),
-		SimStageData.StageIndex,
-		SimStageData.IterationIndex,
-		DispatchNumThreads.X, DispatchNumThreads.Y, DispatchNumThreads.Z
-	);
-	FNiagaraEmptyUAVPoolScopedAccess UAVPoolAccessScope(GetEmptyUAVPool());
-
-	//-TODO: Optimize, we don't need to keep allocating for each dispatch
-	FMemStack& MemStack = FMemStack::Get();
-	FMemMark Mark(MemStack);
-	const FShaderParametersMetadata& ShaderParametersMetadata = *InstanceData.Context->GPUScript_RT->GetScriptParametersMetadata()->ShaderParametersMetadata.Get();
-	FNiagaraShader::FParameters* DispatchParameters = reinterpret_cast<FNiagaraShader::FParameters*>(MemStack.Alloc(ShaderParametersMetadata.GetSize(), SHADER_PARAMETER_STRUCT_ALIGNMENT));
-	FMemory::Memset(DispatchParameters, 0, ShaderParametersMetadata.GetSize());
+	const FShaderParametersMetadata* ShaderParametersMetadata = InstanceData.Context->GPUScript_RT->GetScriptParametersMetadata()->ShaderParametersMetadata.Get();
+	FNiagaraShader::FParameters* DispatchParameters = reinterpret_cast<FNiagaraShader::FParameters*>(GraphBuilder.Alloc(ShaderParametersMetadata->GetSize(), SHADER_PARAMETER_STRUCT_ALIGNMENT));
+	FMemory::Memset(DispatchParameters, 0, ShaderParametersMetadata->GetSize());
 
 	// Set Parameters
 	const bool bRequiresPersistentIDs = InstanceData.Context->MainDataSet->RequiresPersistentIDs();
@@ -1465,7 +1429,7 @@ void FNiagaraGpuComputeDispatch::DispatchStage(FRHICommandList& RHICmdList, cons
 	static_assert(sizeof(DispatchParameters->EmitterSpawnInfoOffsets) == sizeof(InstanceData.SpawnInfo.SpawnInfoStartOffsets), "Mismatch between shader parameter size and SpawnInfoStartOffsets");
 	static_assert(sizeof(DispatchParameters->EmitterSpawnInfoParams) == sizeof(InstanceData.SpawnInfo.SpawnInfoParams), "Mismatch between shader parameter size and SpawnInfoParams");
 	FMemory::Memcpy(&DispatchParameters->EmitterSpawnInfoOffsets, &InstanceData.SpawnInfo.SpawnInfoStartOffsets, sizeof(InstanceData.SpawnInfo.SpawnInfoStartOffsets));
-	FMemory::Memcpy(&DispatchParameters->EmitterSpawnInfoParams,	 &InstanceData.SpawnInfo.SpawnInfoParams, sizeof(InstanceData.SpawnInfo.SpawnInfoParams));
+	FMemory::Memcpy(&DispatchParameters->EmitterSpawnInfoParams, &InstanceData.SpawnInfo.SpawnInfoParams, sizeof(InstanceData.SpawnInfo.SpawnInfoParams));
 
 	// Setup instance counts
 	DispatchParameters->RWInstanceCounts			= GPUInstanceCounterManager.GetInstanceCountBuffer().UAV;
@@ -1538,12 +1502,10 @@ void FNiagaraGpuComputeDispatch::DispatchStage(FRHICommandList& RHICmdList, cons
 	}
 
 	// Set Compute Shader
-	const FNiagaraShaderRef ComputeShader = InstanceData.Context->GPUScript_RT->GetShader(SimStageData.StageIndex);
-	FRHIComputeShader* RHIComputeShader = ComputeShader.GetComputeShader();
-	SetComputePipelineState(RHICmdList, RHIComputeShader);
+	const TShaderRef<FNiagaraShader> ComputeShader = InstanceData.Context->GPUScript_RT->GetShader(SimStageData.StageIndex);
 
 	// Set data interface parameters
-	SetDataInterfaceParameters(RHICmdList, Tick, InstanceData, ComputeShader, SimStageData, reinterpret_cast<uint8*>(DispatchParameters));
+	SetDataInterfaceParameters(GraphBuilder, Tick, InstanceData, ComputeShader, SimStageData, reinterpret_cast<uint8*>(DispatchParameters));
 
 	// Set tick parameters
 	Tick.GetGlobalParameters(InstanceData, &DispatchParameters->GlobalParameters);
@@ -1551,17 +1513,13 @@ void FNiagaraGpuComputeDispatch::DispatchStage(FRHICommandList& RHICmdList, cons
 	Tick.GetOwnerParameters(InstanceData, &DispatchParameters->OwnerParameters);
 	Tick.GetEmitterParameters(InstanceData, &DispatchParameters->EmitterParameters);
 
-	// Set external constant buffer
-	//-TODO: This should be replace with parameters structure which is dynamically created
-	if (ComputeShader->ExternalConstantBufferParam[0].IsBound())
-	{
-		RHICmdList.SetShaderUniformBuffer(RHIComputeShader, ComputeShader->ExternalConstantBufferParam[0].GetBaseIndex(), Tick.GetExternalUniformBuffer(InstanceData, false));
-	}
-	if (ComputeShader->ExternalConstantBufferParam[1].IsBound())
+	// Set ViewUniformBuffer if it's required
+	if (ComputeShader->bNeedsViewUniformBuffer)
 	{
-		check(InstanceData.Context->HasInterpolationParameters);
-		RHICmdList.SetShaderUniformBuffer(RHIComputeShader, ComputeShader->ExternalConstantBufferParam[1].GetBaseIndex(), Tick.GetExternalUniformBuffer(InstanceData, true));
+		DispatchParameters->View = CurrentPassViews[0].ViewUniformBuffer;
 	}
+	DispatchParameters->SceneTextures.SceneTextures			= NiagaraSceneTextures->SceneTextures;
+	DispatchParameters->SceneTextures.MobileSceneTextures	= NiagaraSceneTextures->MobileSceneTextures;
 
 	// Execute the dispatch
 	{
@@ -1581,26 +1539,73 @@ void FNiagaraGpuComputeDispatch::DispatchStage(FRHICommandList& RHICmdList, cons
 		DispatchParameters->DispatchThreadIdToLinear	= FUintVector3(1, DispatchCount.X, DispatchCount.X * DispatchCount.Y);
 		DispatchParameters->DispatchThreadIdBounds		= FUintVector3(DispatchCount.X, DispatchCount.Y, DispatchCount.Z);
 
-		SetShaderParameters<FRHICommandList, FNiagaraShader>(RHICmdList, ComputeShader, RHIComputeShader, &ShaderParametersMetadata, *DispatchParameters);
-		DispatchComputeShader(RHICmdList, ComputeShader, ThreadGroupCount.X, ThreadGroupCount.Y, ThreadGroupCount.Z);
-		UnsetShaderUAVs<FRHICommandList, FNiagaraShader>(RHICmdList, ComputeShader, RHIComputeShader);
+		// Simple case can be this, i.e. no legacy DI's, no External Constants, no profiling
+		//FComputeShaderUtils::AddPass(
+		//	GraphBuilder,
+		//	RDG_EVENT_NAME("NiagaraStage"),
+		//	ERDGPassFlags::Compute | ERDGPassFlags::NeverCull,
+		//	ComputeShader,
+		//	&ShaderParametersMetadata,
+		//	DispatchParameters,
+		//	ThreadGroupCount
+		//);
+
+		GraphBuilder.AddPass(
+			RDG_EVENT_NAME(
+				"NiagaraGpuSim(%s) DispatchCount(%dx%dx%d) Stage(%s %u) Iteration(%u) NumThreads(%dx%dx%d)",
+				InstanceData.Context->GetDebugSimName(),
+				DispatchCount.X, DispatchCount.Y, DispatchCount.Z,
+				*SimStageData.StageMetaData->SimulationStageName.ToString(),
+				SimStageData.StageIndex,
+				SimStageData.IterationIndex,
+				DispatchNumThreads.X, DispatchNumThreads.Y, DispatchNumThreads.Z
+			),
+			ShaderParametersMetadata,
+			DispatchParameters,
+			ERDGPassFlags::Compute | ERDGPassFlags::NeverCull,
+			[this, ShaderParametersMetadata, DispatchParameters, ComputeShader, ThreadGroupCount, TickPtr=&Tick, InstanceDataPtr=&InstanceData, SimStageDataPtr=&SimStageData](FRHICommandListImmediate& RHICmdList)		//-TODO:RDG: When legacy is removed this can be FRHIComputeCommandList
+			{
+				FRHIComputeShader* RHIComputeShader = ComputeShader.GetComputeShader();
+				SetComputePipelineState(RHICmdList, RHIComputeShader);
+
+				if (ComputeShader->ExternalConstantBufferParam[0].IsBound())
+				{
+					RHICmdList.SetShaderUniformBuffer(RHIComputeShader, ComputeShader->ExternalConstantBufferParam[0].GetBaseIndex(), TickPtr->GetExternalUniformBuffer(*InstanceDataPtr, false));
+				}
+				if (ComputeShader->ExternalConstantBufferParam[1].IsBound())
+				{
+					check(InstanceDataPtr->Context->HasInterpolationParameters);
+					RHICmdList.SetShaderUniformBuffer(RHIComputeShader, ComputeShader->ExternalConstantBufferParam[1].GetBaseIndex(), TickPtr->GetExternalUniformBuffer(*InstanceDataPtr, true));
+				}
+
+				FNiagaraEmptyUAVPoolScopedAccess UAVPoolAccessScope(GetEmptyUAVPool());
+
+				LegacySetDataInterfaceParameters(RHICmdList, *TickPtr, *InstanceDataPtr, ComputeShader, *SimStageDataPtr);
+
+				FNiagaraGpuProfileScope GpuProfileDispatchScope(RHICmdList, this, FNiagaraGpuProfileEvent(*InstanceDataPtr, *SimStageDataPtr, InstanceDataPtr == &TickPtr->GetInstances()[0]));
+
+				SetShaderParameters<FRHICommandList, FNiagaraShader>(RHICmdList, ComputeShader, RHIComputeShader, ShaderParametersMetadata, *DispatchParameters);
+				DispatchComputeShader(RHICmdList, ComputeShader, ThreadGroupCount.X, ThreadGroupCount.Y, ThreadGroupCount.Z);
+				UnsetShaderUAVs<FRHICommandList, FNiagaraShader>(RHICmdList, ComputeShader, RHIComputeShader);
+
+				LegacyUnsetDataInterfaceParameters(RHICmdList, *TickPtr, *InstanceDataPtr, ComputeShader, *SimStageDataPtr);
+			}
+		);
 
 		INC_DWORD_STAT(STAT_NiagaraGPUDispatches);
 	}
 
-	// Unset UAV parameters
-	UnsetDataInterfaceParameters(RHICmdList, Tick, InstanceData, ComputeShader, SimStageData);
-
 	// Optionally submit commands to the GPU
 	// This can be used to avoid accidental TDR detection in the editor especially when issuing multiple ticks in the same frame
 	++TotalDispatchesThisFrame;
-	if (GNiagaraGpuSubmitCommandHint > 0)
-	{
-		if ((TotalDispatchesThisFrame % GNiagaraGpuSubmitCommandHint) == 0)
-		{
-			RHICmdList.SubmitCommandsHint();
-		}
-	}
+	//-TODO:RDG: Fix submit commands hint
+	//if (GNiagaraGpuSubmitCommandHint > 0)
+	//{
+	//	if ((TotalDispatchesThisFrame % GNiagaraGpuSubmitCommandHint) == 0)
+	//	{
+	//		RHICmdList.SubmitCommandsHint();
+	//	}
+	//}
 }
 
 void FNiagaraGpuComputeDispatch::PreInitViews(FRDGBuilder& GraphBuilder, bool bAllowGPUParticleUpdate)
@@ -1663,20 +1668,7 @@ void FNiagaraGpuComputeDispatch::PreInitViews(FRDGBuilder& GraphBuilder, bool bA
 				DumpDebugFrame();
 			}
 
-			FNiagaraComputePassParameters* PassParameters = GraphBuilder.AllocParameters<FNiagaraComputePassParameters>();
-			GNiagaraViewDataManager.GetSceneTextureParameters(GraphBuilder, nullptr, PassParameters->SceneTextures);
-
-			GraphBuilder.AddPass(
-				RDG_EVENT_NAME("Niagara::PreInitViews"),
-				PassParameters,
-				ERDGPassFlags::Compute | ERDGPassFlags::NeverCull,
-				[this, PassParameters](FRHICommandListImmediate& RHICmdList)
-			{
-				NiagaraSceneTextures = &PassParameters->SceneTextures;
-				ON_SCOPE_EXIT { NiagaraSceneTextures = nullptr; };
-
-				ExecuteTicks(RHICmdList, ENiagaraGpuComputeTickStage::PreInitViews);
-			});
+			ExecuteTicks(GraphBuilder, TConstArrayView<FViewInfo>(), ENiagaraGpuComputeTickStage::PreInitViews);
 		}
 	}
 	else
@@ -1694,25 +1686,7 @@ void FNiagaraGpuComputeDispatch::PostInitViews(FRDGBuilder& GraphBuilder, TArray
 
 	if (bAllowGPUParticleUpdate && FNiagaraUtilities::AllowGPUParticles(GetShaderPlatform()))
 	{
-		FNiagaraComputePassParameters* PassParameters = GraphBuilder.AllocParameters<FNiagaraComputePassParameters>();
-		GNiagaraViewDataManager.GetSceneTextureParameters(GraphBuilder, GetViewFamilyInfo(Views).GetSceneTexturesChecked(), PassParameters->SceneTextures);
-		PassParameters->View = GetReferenceViewUniformBuffer(Views);
-
-		GraphBuilder.AddPass(
-			RDG_EVENT_NAME("Niagara::PostInitViews"),
-			PassParameters,
-			ERDGPassFlags::Compute | ERDGPassFlags::NeverCull,
-			[this, PassParameters, Views](FRHICommandListImmediate& RHICmdList)
-			{
-				NiagaraSceneTextures = &PassParameters->SceneTextures;
-				ON_SCOPE_EXIT{ NiagaraSceneTextures = nullptr; };
-
-				ExecuteTicks(RHICmdList, ENiagaraGpuComputeTickStage::PostInitViews);
-			#if WITH_MGPU
-				TransferMultiGPUBufers(RHICmdList, ENiagaraGpuComputeTickStage::PostInitViews);
-			#endif // WITH_MGPU
-			}
-		);
+		ExecuteTicks(GraphBuilder, Views, ENiagaraGpuComputeTickStage::PostInitViews);
 	}
 }
 
@@ -1729,55 +1703,46 @@ void FNiagaraGpuComputeDispatch::PostRenderOpaque(FRDGBuilder& GraphBuilder, TCo
 
 	if (bAllowGPUParticleUpdate && FNiagaraUtilities::AllowGPUParticles(GetShaderPlatform()))
 	{
-		FNiagaraComputePassParameters* PassParameters = GraphBuilder.AllocParameters<FNiagaraComputePassParameters>();
-		// TODO: This will cause a fragment->compute barrier on a scene textures which could be costly especially on mobile GPUs
-		// Will be nice to avoid executing this if we know that there are no simulations that require access to a scene textures
-		GNiagaraViewDataManager.GetSceneTextureParameters(GraphBuilder, GetViewFamilyInfo(Views).GetSceneTexturesChecked(), PassParameters->SceneTextures);
-		PassParameters->View = GetReferenceViewUniformBuffer(Views);
-
- 		GraphBuilder.AddPass(
-			RDG_EVENT_NAME("Niagara::PostRenderOpaque"),
-			PassParameters,
-			ERDGPassFlags::Compute | ERDGPassFlags::NeverCull,
-			[this, PassParameters, Views](FRHICommandListImmediate& RHICmdList)
-		{
-			NiagaraSceneTextures = &PassParameters->SceneTextures;
-			ON_SCOPE_EXIT{ NiagaraSceneTextures = nullptr; };
-
-			AsyncGpuTraceHelper->PostRenderOpaque(RHICmdList, this, Views);
-
-			CurrentPassViews = Views;
-
-			// Setup new readback since if there is no pending request, there is no risk of having invalid data read (offset being allocated after the readback was sent).
-			ExecuteTicks(RHICmdList, ENiagaraGpuComputeTickStage::PostOpaqueRender);
+		AddPass(
+			GraphBuilder,
+			RDG_EVENT_NAME("AsyncGpuTraceHelper::PostRenderOpaque"),
+			[this, Views](FRHICommandListImmediate& RHICmdList)
+			{
+				AsyncGpuTraceHelper->PostRenderOpaque(RHICmdList, this, Views);
+				CurrentPassViews = Views;
+			}
+		);
 
-			FinishDispatches();
+		ExecuteTicks(GraphBuilder, Views, ENiagaraGpuComputeTickStage::PostOpaqueRender);
+	}
 
-			AsyncGpuTraceHelper->EndFrame(RHICmdList, this);
+	AddPass(
+		GraphBuilder,
+		RDG_EVENT_NAME("Niagara::PostRenderFinish"),
+		[this, bExecuteReadback=bRequiresReadback, bAllowGPUParticleUpdate](FRHICommandListImmediate& RHICmdList)
+		{
+			if (bAllowGPUParticleUpdate)
+			{
+				FinishDispatches();
 
-			// Clear CurrentPassViews
+				AsyncGpuTraceHelper->EndFrame(RHICmdList, this);
+			}
 			CurrentPassViews = TConstArrayView<FViewInfo>();
 
 			ProcessDebugReadbacks(RHICmdList, false);
-		});
-	}
 
-	if (bRequiresReadback)
-	{
-		AddPass(GraphBuilder, RDG_EVENT_NAME("Niagara::GPUReadback"), [this](FRHICommandListImmediate& RHICmdList)
-		{
-			check(!GPUInstanceCounterManager.HasPendingGPUReadback());
-			GPUInstanceCounterManager.EnqueueGPUReadback(RHICmdList);
-		});
-		bRequiresReadback = false;
-	}
+			if (bExecuteReadback)
+			{
+				check(!GPUInstanceCounterManager.HasPendingGPUReadback());
+				GPUInstanceCounterManager.EnqueueGPUReadback(RHICmdList);
+			}
 
-#if WITH_NIAGARA_GPU_PROFILER
-	AddPass(GraphBuilder, RDG_EVENT_NAME("Niagara::GPUProfiler_EndFrame"), [this](FRHICommandListImmediate& RHICmdList)
-	{
-		GPUProfilerPtr->EndFrame(RHICmdList);
-	});
-#endif
+		#if WITH_NIAGARA_GPU_PROFILER
+			GPUProfilerPtr->EndFrame(RHICmdList);
+		#endif
+		}
+	);
+	bRequiresReadback = false;
 
 	GNiagaraViewDataManager.ClearSceneTextureParameters();
 }
@@ -2016,9 +1981,12 @@ FNiagaraAsyncGpuTraceHelper& FNiagaraGpuComputeDispatch::GetAsyncGpuTraceHelper(
 	return *AsyncGpuTraceHelper.Get();
 }
 
-/* Set shader parameters for data interfaces
- */
-void FNiagaraGpuComputeDispatch::SetDataInterfaceParameters(FRHICommandList& RHICmdList, const FNiagaraGPUSystemTick& Tick, const FNiagaraComputeInstanceData& InstanceData, const FNiagaraShaderRef& ComputeShader, const FNiagaraSimStageData& SimStageData, uint8* ParametersStructure) const
+//////////////////////////////////////////////////////////////////////////////////////////////////////////////////////////////////////////////////////////////////////////////////////////////////////////////////////////////////////////////////////////////////////////////////////////////////////////////////////////////////////////////////////////////////////////////////
+//////////////////////////////////////////////////////////////////////////////////////////////////////////////////////////////////////////////////////////////////////////////////////////////////////////////////////////////////////////////////////////////////////////////////////////////////////////////////////////////////////////////////////////////////////////////////
+//////////////////////////////////////////////////////////////////////////////////////////////////////////////////////////////////////////////////////////////////////////////////////////////////////////////////////////////////////////////////////////////////////////////////////////////////////////////////////////////////////////////////////////////////////////////////
+//////////////////////////////////////////////////////////////////////////////////////////////////////////////////////////////////////////////////////////////////////////////////////////////////////////////////////////////////////////////////////////////////////////////////////////////////////////////////////////////////////////////////////////////////////////////////
+
+void FNiagaraGpuComputeDispatch::LegacySetDataInterfaceParameters(FRHICommandList& RHICmdList, const FNiagaraGPUSystemTick& Tick, const FNiagaraComputeInstanceData& InstanceData, const FNiagaraShaderRef& ComputeShader, const FNiagaraSimStageData& SimStageData) const
 {
 	const int32 NumDataInterfaces = InstanceData.DataInterfaceProxies.Num();
 	if (NumDataInterfaces == 0)
@@ -2026,23 +1994,15 @@ void FNiagaraGpuComputeDispatch::SetDataInterfaceParameters(FRHICommandList& RHI
 		return;
 	}
 
-	// @todo-threadsafety This is a bit gross. Need to rethink this api.
 	const FNiagaraShaderMapPointerTable& PointerTable = ComputeShader.GetPointerTable();
 	const TMemoryImageArray<FNiagaraDataInterfaceParamRef>& DIParameters = ComputeShader->GetDIParameters();
 
-	FNiagaraDataInterfaceSetShaderParametersContext SetParametersContext(RHICmdList , *this, Tick, InstanceData, SimStageData, ComputeShader, ParametersStructure);
-
-	for ( int32 iDataInterface=0; iDataInterface < NumDataInterfaces; ++iDataInterface )
+	for (int32 iDataInterface=0; iDataInterface < NumDataInterfaces; ++iDataInterface)
 	{
 		FNiagaraDataInterfaceProxy* DataInterfaceProxy = InstanceData.DataInterfaceProxies[iDataInterface];
 
 		const FNiagaraDataInterfaceParamRef& DIParam = DIParameters[iDataInterface];
-		if (DIParam.ShaderParametersOffset != INDEX_NONE)
-		{
-			SetParametersContext.SetDataInterface(DataInterfaceProxy, DIParam.ShaderParametersOffset, DIParam.Parameters);
-			CastChecked<UNiagaraDataInterface>(DIParam.DIType.Get(PointerTable.DITypes))->SetShaderParameters(SetParametersContext);
-		}
-		else if (DIParam.Parameters.IsValid())
+		if (DIParam.Parameters.IsValid() && (DIParam.ShaderParametersOffset == INDEX_NONE))
 		{
 			FNiagaraDataInterfaceSetArgs Context(DataInterfaceProxy, Tick.SystemInstanceID, Tick.SystemGpuComputeProxy->GetSystemLWCTile(), this, ComputeShader, &InstanceData, &SimStageData, InstanceData.IsOutputStage(DataInterfaceProxy, SimStageData.StageIndex), InstanceData.IsIterationStage(DataInterfaceProxy, SimStageData.StageIndex));
 			DIParam.DIType.Get(PointerTable.DITypes)->SetParameters(DIParam.Parameters.Get(), RHICmdList, Context);
@@ -2050,7 +2010,7 @@ void FNiagaraGpuComputeDispatch::SetDataInterfaceParameters(FRHICommandList& RHI
 	}
 }
 
-void FNiagaraGpuComputeDispatch::UnsetDataInterfaceParameters(FRHICommandList& RHICmdList, const FNiagaraGPUSystemTick& Tick, const FNiagaraComputeInstanceData& InstanceData, const FNiagaraShaderRef& ComputeShader, const FNiagaraSimStageData& SimStageData) const
+void FNiagaraGpuComputeDispatch::LegacyUnsetDataInterfaceParameters(FRHICommandList& RHICmdList, const FNiagaraGPUSystemTick& Tick, const FNiagaraComputeInstanceData& InstanceData, const FNiagaraShaderRef& ComputeShader, const FNiagaraSimStageData& SimStageData) const
 {
 	// @todo-threadsafety This is a bit gross. Need to rethink this api.
 	const FNiagaraShaderMapPointerTable& PointerTable = ComputeShader.GetPointerTable();
@@ -2060,7 +2020,7 @@ void FNiagaraGpuComputeDispatch::UnsetDataInterfaceParameters(FRHICommandList& R
 	{
 		const TMemoryImageArray<FNiagaraDataInterfaceParamRef>& DIParameters = ComputeShader->GetDIParameters();
 		const FNiagaraDataInterfaceParamRef& DIParam = DIParameters[InterfaceIndex];
-		if (DIParam.ShaderParametersOffset != INDEX_NONE && DIParam.Parameters.IsValid())
+		if (DIParam.Parameters.IsValid() && (DIParam.ShaderParametersOffset == INDEX_NONE))
 		{
 			FNiagaraDataInterfaceSetArgs Context(Interface, Tick.SystemInstanceID, Tick.SystemGpuComputeProxy->GetSystemLWCTile(), this, ComputeShader, &InstanceData, &SimStageData, InstanceData.IsOutputStage(Interface, SimStageData.StageIndex), InstanceData.IsIterationStage(Interface, SimStageData.StageIndex));
 			DIParam.DIType.Get(PointerTable.DITypes)->UnsetParameters(DIParam.Parameters.Get(), RHICmdList, Context);
@@ -2070,6 +2030,216 @@ void FNiagaraGpuComputeDispatch::UnsetDataInterfaceParameters(FRHICommandList& R
 	}
 }
 
+void FNiagaraGpuComputeDispatch::LegacyResetDataInterfaces(FRHICommandList& RHICmdList, const FNiagaraGPUSystemTick& Tick, const FNiagaraComputeInstanceData& InstanceData) const
+{
+	// Note: All stages will contain the same bindings so if they are valid for one they are valid for all, this could change in the future
+	const FNiagaraShaderRef& ComputeShader = InstanceData.Context->GPUScript_RT->GetShader(0);
+	const TMemoryImageArray<FNiagaraDataInterfaceParamRef>& DIParameters = ComputeShader->GetDIParameters();
+
+	uint32 InterfaceIndex = 0;
+	for (FNiagaraDataInterfaceProxy* Interface : InstanceData.DataInterfaceProxies)
+	{
+		const FNiagaraDataInterfaceParamRef& DIParam = DIParameters[InterfaceIndex];
+		if (DIParam.Parameters.IsValid() && (DIParam.ShaderParametersOffset == INDEX_NONE))
+		{
+			const FNiagaraDataInterfaceArgs TmpContext(Interface, Tick.SystemInstanceID, Tick.SystemGpuComputeProxy->GetSystemLWCTile(), this);
+			Interface->ResetData(RHICmdList, TmpContext);
+		}
+		InterfaceIndex++;
+	}
+}
+
+void FNiagaraGpuComputeDispatch::LegacyPreStageInterface(FRHICommandList& RHICmdList, const FNiagaraGPUSystemTick& Tick, const FNiagaraComputeInstanceData& InstanceData, const FNiagaraSimStageData& SimStageData, TSet<FNiagaraDataInterfaceProxy*>& ProxiesToFinalize) const
+{
+	// Note: All stages will contain the same bindings so if they are valid for one they are valid for all, this could change in the future
+	const FNiagaraShaderRef& ComputeShader = InstanceData.Context->GPUScript_RT->GetShader(0);
+	const TMemoryImageArray<FNiagaraDataInterfaceParamRef>& DIParameters = ComputeShader->GetDIParameters();
+
+	uint32 InterfaceIndex = 0;
+	for (FNiagaraDataInterfaceProxy* Interface : InstanceData.DataInterfaceProxies)
+	{
+		const FNiagaraDataInterfaceParamRef& DIParam = DIParameters[InterfaceIndex];
+		if (DIParam.Parameters.IsValid() && (DIParam.ShaderParametersOffset == INDEX_NONE))
+		{
+			const FNiagaraDataInterfaceStageArgs TmpContext(Interface, Tick.SystemInstanceID, Tick.SystemGpuComputeProxy->GetSystemLWCTile(), this, &InstanceData, &SimStageData, InstanceData.IsOutputStage(Interface, SimStageData.StageIndex), InstanceData.IsIterationStage(Interface, SimStageData.StageIndex));
+			Interface->PreStage(RHICmdList, TmpContext);
+
+			if (Interface->RequiresPreStageFinalize())
+			{
+				ProxiesToFinalize.Add(Interface);
+			}
+		}
+		InterfaceIndex++;
+	}
+}
+
+void FNiagaraGpuComputeDispatch::LegacyPostStageInterface(FRHICommandList& RHICmdList, const FNiagaraGPUSystemTick& Tick, const FNiagaraComputeInstanceData& InstanceData, const FNiagaraSimStageData& SimStageData, TSet<FNiagaraDataInterfaceProxy*>& ProxiesToFinalize) const
+{
+	// Note: All stages will contain the same bindings so if they are valid for one they are valid for all, this could change in the future
+	const FNiagaraShaderRef& ComputeShader = InstanceData.Context->GPUScript_RT->GetShader(0);
+	const TMemoryImageArray<FNiagaraDataInterfaceParamRef>& DIParameters = ComputeShader->GetDIParameters();
+
+	uint32 InterfaceIndex = 0;
+	for (FNiagaraDataInterfaceProxy* Interface : InstanceData.DataInterfaceProxies)
+	{
+		const FNiagaraDataInterfaceParamRef& DIParam = DIParameters[InterfaceIndex];
+		if (DIParam.Parameters.IsValid() && (DIParam.ShaderParametersOffset == INDEX_NONE))
+		{
+			const FNiagaraDataInterfaceStageArgs TmpContext(Interface, Tick.SystemInstanceID, Tick.SystemGpuComputeProxy->GetSystemLWCTile(), this, &InstanceData, &SimStageData, InstanceData.IsOutputStage(Interface, SimStageData.StageIndex), InstanceData.IsIterationStage(Interface, SimStageData.StageIndex));
+			Interface->PostStage(RHICmdList, TmpContext);
+
+			if (Interface->RequiresPostStageFinalize())
+			{
+				ProxiesToFinalize.Add(Interface);
+			}
+		}
+		InterfaceIndex++;
+	}
+}
+
+void FNiagaraGpuComputeDispatch::LegacyPostSimulateInterface(FRHICommandList& RHICmdList, const FNiagaraGPUSystemTick& Tick, const FNiagaraComputeInstanceData& InstanceData) const
+{
+	// Note: All stages will contain the same bindings so if they are valid for one they are valid for all, this could change in the future
+	const FNiagaraShaderRef& ComputeShader = InstanceData.Context->GPUScript_RT->GetShader(0);
+	const TMemoryImageArray<FNiagaraDataInterfaceParamRef>& DIParameters = ComputeShader->GetDIParameters();
+
+	uint32 InterfaceIndex = 0;
+	for (FNiagaraDataInterfaceProxy* Interface : InstanceData.DataInterfaceProxies)
+	{
+		const FNiagaraDataInterfaceParamRef& DIParam = DIParameters[InterfaceIndex];
+		if (DIParam.Parameters.IsValid() && (DIParam.ShaderParametersOffset == INDEX_NONE))
+		{
+			const FNiagaraDataInterfaceArgs TmpContext(Interface, Tick.SystemInstanceID, Tick.SystemGpuComputeProxy->GetSystemLWCTile(), this);
+			Interface->PostSimulate(RHICmdList, TmpContext);
+		}
+		InterfaceIndex++;
+	}
+}
+
+//////////////////////////////////////////////////////////////////////////////////////////////////////////////////////////////////////////////////////////////////////////////////////////////////////////////////////////////////////////////////////////////////////////////////////////////////////////////////////////////////////////////////////////////////////////////////
+//////////////////////////////////////////////////////////////////////////////////////////////////////////////////////////////////////////////////////////////////////////////////////////////////////////////////////////////////////////////////////////////////////////////////////////////////////////////////////////////////////////////////////////////////////////////////
+//////////////////////////////////////////////////////////////////////////////////////////////////////////////////////////////////////////////////////////////////////////////////////////////////////////////////////////////////////////////////////////////////////////////////////////////////////////////////////////////////////////////////////////////////////////////////
+//////////////////////////////////////////////////////////////////////////////////////////////////////////////////////////////////////////////////////////////////////////////////////////////////////////////////////////////////////////////////////////////////////////////////////////////////////////////////////////////////////////////////////////////////////////////////
+
+void FNiagaraGpuComputeDispatch::ResetDataInterfaces(FRDGBuilder& GraphBuilder, const FNiagaraGPUSystemTick& Tick, const FNiagaraComputeInstanceData& InstanceData) const
+{
+	const int32 NumDataInterfaces = InstanceData.DataInterfaceProxies.Num();
+	if (NumDataInterfaces == 0)
+	{
+		return;
+	}
+
+	const FNiagaraShaderRef& ComputeShader = InstanceData.Context->GPUScript_RT->GetShader(0);
+	const TMemoryImageArray<FNiagaraDataInterfaceParamRef>& DIParameters = ComputeShader->GetDIParameters();
+
+	FNDIGpuComputeResetContext Context(GraphBuilder, *this, Tick.SystemInstanceID);
+	for (int32 iDataInterface=0; iDataInterface < NumDataInterfaces; ++iDataInterface)
+	{
+		const FNiagaraDataInterfaceParamRef& DIParam = DIParameters[iDataInterface];
+		if (DIParam.ShaderParametersOffset != INDEX_NONE)
+		{
+			FNiagaraDataInterfaceProxy* DataInterfaceProxy = InstanceData.DataInterfaceProxies[iDataInterface];
+			DataInterfaceProxy->ResetData(Context);
+		}
+	}
+}
+
+void FNiagaraGpuComputeDispatch::SetDataInterfaceParameters(FRDGBuilder& GraphBuilder, const FNiagaraGPUSystemTick& Tick, const FNiagaraComputeInstanceData& InstanceData, const FNiagaraShaderRef& ComputeShader, const FNiagaraSimStageData& SimStageData, uint8* ParametersStructure) const
+{
+	const int32 NumDataInterfaces = InstanceData.DataInterfaceProxies.Num();
+	if (NumDataInterfaces == 0)
+	{
+		return;
+	}
+
+	const FNiagaraShaderMapPointerTable& PointerTable = ComputeShader.GetPointerTable();
+	const TMemoryImageArray<FNiagaraDataInterfaceParamRef>& DIParameters = ComputeShader->GetDIParameters();
+
+	FNiagaraDataInterfaceSetShaderParametersContext Context(GraphBuilder, *this, Tick, InstanceData, SimStageData, ComputeShader, ParametersStructure);
+	for ( int32 iDataInterface=0; iDataInterface < NumDataInterfaces; ++iDataInterface )
+	{
+		FNiagaraDataInterfaceProxy* DataInterfaceProxy = InstanceData.DataInterfaceProxies[iDataInterface];
+
+		const FNiagaraDataInterfaceParamRef& DIParam = DIParameters[iDataInterface];
+		if (DIParam.ShaderParametersOffset != INDEX_NONE)
+		{
+			Context.SetDataInterface(DataInterfaceProxy, DIParam.ShaderParametersOffset, DIParam.Parameters);
+			CastChecked<UNiagaraDataInterface>(DIParam.DIType.Get(PointerTable.DITypes))->SetShaderParameters(Context);
+		}
+	}
+}
+
+void FNiagaraGpuComputeDispatch::PreStageInterface(FRDGBuilder& GraphBuilder, const FNiagaraGPUSystemTick& Tick, const FNiagaraComputeInstanceData& InstanceData, const FNiagaraSimStageData& SimStageData, TSet<FNiagaraDataInterfaceProxy*>& ProxiesToFinalize) const
+{
+	const int32 NumDataInterfaces = InstanceData.DataInterfaceProxies.Num();
+	if (NumDataInterfaces == 0)
+	{
+		return;
+	}
+
+	const FNiagaraShaderRef& ComputeShader = InstanceData.Context->GPUScript_RT->GetShader(0);
+	const TMemoryImageArray<FNiagaraDataInterfaceParamRef>& DIParameters = ComputeShader->GetDIParameters();
+
+	FNDIGpuComputePreStageContext Context(GraphBuilder, *this, Tick, InstanceData, SimStageData);
+	for (int32 iDataInterface = 0; iDataInterface < NumDataInterfaces; ++iDataInterface)
+	{
+		const FNiagaraDataInterfaceParamRef& DIParam = DIParameters[iDataInterface];
+		if (DIParam.ShaderParametersOffset != INDEX_NONE)
+		{
+			FNiagaraDataInterfaceProxy* DataInterfaceProxy = InstanceData.DataInterfaceProxies[iDataInterface];
+			Context.SetDataInterfaceProxy(DataInterfaceProxy);
+			DataInterfaceProxy->PreStage(Context);
+		}
+	}
+}
+
+void FNiagaraGpuComputeDispatch::PostStageInterface(FRDGBuilder& GraphBuilder, const FNiagaraGPUSystemTick& Tick, const FNiagaraComputeInstanceData& InstanceData, const FNiagaraSimStageData& SimStageData, TSet<FNiagaraDataInterfaceProxy*>& ProxiesToFinalize) const
+{
+	const int32 NumDataInterfaces = InstanceData.DataInterfaceProxies.Num();
+	if (NumDataInterfaces == 0)
+	{
+		return;
+	}
+
+	const FNiagaraShaderRef& ComputeShader = InstanceData.Context->GPUScript_RT->GetShader(0);
+	const TMemoryImageArray<FNiagaraDataInterfaceParamRef>& DIParameters = ComputeShader->GetDIParameters();
+
+	FNDIGpuComputePostStageContext Context(GraphBuilder, *this, Tick, InstanceData, SimStageData);
+	for (int32 iDataInterface = 0; iDataInterface < NumDataInterfaces; ++iDataInterface)
+	{
+		const FNiagaraDataInterfaceParamRef& DIParam = DIParameters[iDataInterface];
+		if (DIParam.ShaderParametersOffset != INDEX_NONE)
+		{
+			FNiagaraDataInterfaceProxy* DataInterfaceProxy = InstanceData.DataInterfaceProxies[iDataInterface];
+			Context.SetDataInterfaceProxy(DataInterfaceProxy);
+			DataInterfaceProxy->PostStage(Context);
+		}
+	}
+}
+
+void FNiagaraGpuComputeDispatch::PostSimulateInterface(FRDGBuilder& GraphBuilder, const FNiagaraGPUSystemTick& Tick, const FNiagaraComputeInstanceData& InstanceData, const FNiagaraSimStageData& SimStageData) const
+{
+	const int32 NumDataInterfaces = InstanceData.DataInterfaceProxies.Num();
+	if (NumDataInterfaces == 0)
+	{
+		return;
+	}
+
+	const FNiagaraShaderRef& ComputeShader = InstanceData.Context->GPUScript_RT->GetShader(0);
+	const TMemoryImageArray<FNiagaraDataInterfaceParamRef>& DIParameters = ComputeShader->GetDIParameters();
+
+	FNDIGpuComputePostSimulateContext Context(GraphBuilder, *this, Tick.SystemInstanceID, SimStageData.bSetDataToRender);
+	for (int32 iDataInterface = 0; iDataInterface < NumDataInterfaces; ++iDataInterface)
+	{
+		const FNiagaraDataInterfaceParamRef& DIParam = DIParameters[iDataInterface];
+		if (DIParam.ShaderParametersOffset != INDEX_NONE)
+		{
+			FNiagaraDataInterfaceProxy* DataInterfaceProxy = InstanceData.DataInterfaceProxies[iDataInterface];
+			DataInterfaceProxy->PostSimulate(Context);
+		}
+	}
+}
+
 FGPUSortManager* FNiagaraGpuComputeDispatch::GetGPUSortManager() const
 {
 	return GPUSortManager;
diff --git a/Engine/Plugins/FX/Niagara/Source/Niagara/Private/NiagaraGpuComputeDispatch.h b/Engine/Plugins/FX/Niagara/Source/Niagara/Private/NiagaraGpuComputeDispatch.h
--- a/Engine/Plugins/FX/Niagara/Source/Niagara/Private/NiagaraGpuComputeDispatch.h
+++ b/Engine/Plugins/FX/Niagara/Source/Niagara/Private/NiagaraGpuComputeDispatch.h
@@ -30,6 +30,61 @@ class FGPUSortManager;
 class FDistanceFieldSceneData;
 class FNiagaraAsyncGpuTraceHelper;
 
+//////////////////////////////////////////////////////////////////////////
+
+struct FNiagaraGpuDispatchInstance
+{
+	FNiagaraGpuDispatchInstance(const FNiagaraGPUSystemTick& InTick, const FNiagaraComputeInstanceData& InInstanceData)
+		: Tick(InTick)
+		, InstanceData(InInstanceData)
+	{
+	}
+
+	const FNiagaraGPUSystemTick& Tick;
+	const FNiagaraComputeInstanceData& InstanceData;
+	FNiagaraSimStageData SimStageData;
+};
+
+
+struct FNiagaraGpuFreeIDUpdate
+{
+	FNiagaraGpuFreeIDUpdate(FNiagaraComputeExecutionContext* InComputeContext)
+		: ComputeContext(InComputeContext)
+	{
+	}
+
+	FNiagaraComputeExecutionContext*	ComputeContext = nullptr;
+	mutable FRHIShaderResourceView*		IDToIndexSRV = nullptr;
+	mutable FRHIUnorderedAccessView*	FreeIDsUAV = nullptr;
+	mutable uint32						NumAllocatedIDs = 0;
+};
+
+struct FNiagaraGpuDispatchGroup
+{
+	TArray<FNiagaraGPUSystemTick*>				TicksWithPerInstanceData;
+	TArray<FNiagaraGpuDispatchInstance>			DispatchInstances;
+	TArray<FNiagaraGpuFreeIDUpdate>				FreeIDUpdates;
+};
+
+struct FNiagaraGpuDispatchList
+{
+	void PreAllocateGroups(int32 LastGroup)
+	{
+		const int32 GroupsToAllocate = LastGroup - DispatchGroups.Num();
+		if (GroupsToAllocate > 0)
+		{
+			DispatchGroups.AddDefaulted(GroupsToAllocate);
+		}
+	}
+
+	bool HasWork() const { return DispatchGroups.Num() > 0; }
+
+	TArray<uint32>						CountsToRelease;
+	TArray<FNiagaraGpuDispatchGroup>	DispatchGroups;
+};
+
+//////////////////////////////////////////////////////////////////////////
+
 class FNiagaraGpuComputeDispatch : public FNiagaraGpuComputeDispatchInterface
 {
 public:
@@ -95,24 +150,22 @@ public:
 	const FGlobalDistanceFieldParameterData* GetGlobalDistanceFieldParameters() const;
 	const FDistanceFieldSceneData* GetMeshDistanceFieldParameters() const;
 
-	void SetDataInterfaceParameters(FRHICommandList& RHICmdList, const FNiagaraGPUSystemTick& Tick, const FNiagaraComputeInstanceData& InstanceData, const FNiagaraShaderRef& ComputeShader, const FNiagaraSimStageData& SimStageData, uint8* ParametersStructure) const;
-	void UnsetDataInterfaceParameters(FRHICommandList& RHICmdList, const FNiagaraGPUSystemTick& Tick, const FNiagaraComputeInstanceData& InstanceData, const FNiagaraShaderRef& ComputeShader, const FNiagaraSimStageData& SimStageData) const;
+	void LegacySetDataInterfaceParameters(FRHICommandList& RHICmdList, const FNiagaraGPUSystemTick& Tick, const FNiagaraComputeInstanceData& InstanceData, const FNiagaraShaderRef& ComputeShader, const FNiagaraSimStageData& SimStageData) const;
+	void LegacyUnsetDataInterfaceParameters(FRHICommandList& RHICmdList, const FNiagaraGPUSystemTick& Tick, const FNiagaraComputeInstanceData& InstanceData, const FNiagaraShaderRef& ComputeShader, const FNiagaraSimStageData& SimStageData) const;
+	void LegacyResetDataInterfaces(FRHICommandList& RHICmdList, const FNiagaraGPUSystemTick& Tick, const FNiagaraComputeInstanceData& InstanceData) const;
+	void LegacyPreStageInterface(FRHICommandList& RHICmdList, const FNiagaraGPUSystemTick& Tick, const FNiagaraComputeInstanceData& InstanceData, const FNiagaraSimStageData& SimStageData, TSet<FNiagaraDataInterfaceProxy*>& ProxiesToFinalize) const;
+	void LegacyPostStageInterface(FRHICommandList& RHICmdList, const FNiagaraGPUSystemTick& Tick, const FNiagaraComputeInstanceData& InstanceData, const FNiagaraSimStageData& SimStageData, TSet<FNiagaraDataInterfaceProxy*>& ProxiesToFinalize) const;
+	void LegacyPostSimulateInterface(FRHICommandList& RHICmdList, const FNiagaraGPUSystemTick& Tick, const FNiagaraComputeInstanceData& InstanceData) const;
 
-	/** Reset the data interfaces and check if the spawn stages are valid */
-	void ResetDataInterfaces(FRHICommandList& RHICmdList, const FNiagaraGPUSystemTick& Tick, const FNiagaraComputeInstanceData& InstanceData) const;
+	void ResetDataInterfaces(FRDGBuilder& GraphBuilder, const FNiagaraGPUSystemTick& Tick, const FNiagaraComputeInstanceData& InstanceData) const;
+	void SetDataInterfaceParameters(FRDGBuilder& GraphBuilder, const FNiagaraGPUSystemTick& Tick, const FNiagaraComputeInstanceData& InstanceData, const FNiagaraShaderRef& ComputeShader, const FNiagaraSimStageData& SimStageData, uint8* ParametersStructure) const;
+	void PreStageInterface(FRDGBuilder& GraphBuilder, const FNiagaraGPUSystemTick& Tick, const FNiagaraComputeInstanceData& InstanceData, const FNiagaraSimStageData& SimStageData, TSet<FNiagaraDataInterfaceProxy*>& ProxiesToFinalize) const;
+	void PostStageInterface(FRDGBuilder& GraphBuilder, const FNiagaraGPUSystemTick& Tick, const FNiagaraComputeInstanceData& InstanceData, const FNiagaraSimStageData& SimStageData, TSet<FNiagaraDataInterfaceProxy*>& ProxiesToFinalize) const;
+	void PostSimulateInterface(FRDGBuilder& GraphBuilder, const FNiagaraGPUSystemTick& Tick, const FNiagaraComputeInstanceData& InstanceData, const FNiagaraSimStageData& SimStageData) const;
 
 	/** Given a shader stage index, find the corresponding data interface */
 	FNiagaraDataInterfaceProxyRW* FindIterationInterface(FNiagaraComputeInstanceData* Instance, const uint32 SimulationStageIndex) const;
 
-	/** Loop over all the data interfaces and call the pre-stage methods */
-	void PreStageInterface(FRHICommandList& RHICmdList, const FNiagaraGPUSystemTick& Tick, const FNiagaraComputeInstanceData& InstanceData, const FNiagaraSimStageData& SimStageData, TSet<FNiagaraDataInterfaceProxy*>& ProxiesToFinalize) const;
-
-	/** Loop over all the data interfaces and call the post-stage methods */
-	void PostStageInterface(FRHICommandList& RHICmdList, const FNiagaraGPUSystemTick& Tick, const FNiagaraComputeInstanceData& InstanceData, const FNiagaraSimStageData& SimStageData, TSet<FNiagaraDataInterfaceProxy*>& ProxiesToFinalize) const;
-
-	/** Loop over all data interfaces and call the post-simulate methods */
-	void PostSimulateInterface(FRHICommandList& RHICmdList, const FNiagaraGPUSystemTick& Tick, const FNiagaraComputeInstanceData& InstanceData) const;
-
 	/** Get the shared SortManager, used in the rendering loop to call FGPUSortManager::OnPreRender() and FGPUSortManager::OnPostRenderOpaque() */
 	virtual FGPUSortManager* GetGPUSortManager() const override;
 
@@ -123,6 +176,9 @@ public:
 	/** When inside a pass returns the active scene textures for the pass, can be nullptr. */
 	class FNiagaraSceneTextureParameters* GetNiagaraSceneTextures() const { return NiagaraSceneTextures; }
 
+	/** Allows access to the current passes external access queue. */
+	FRDGExternalAccessQueue& GetCurrentPassExternalAccessQueue() const { return CurrentPassExternalAccessQueue; }
+
 	/** Debug only function to readback data. */
 	virtual void AddDebugReadback(FNiagaraSystemInstanceID InstanceID, TSharedPtr<struct FNiagaraScriptDebuggerInfo, ESPMode::ThreadSafe> DebugInfo, FNiagaraComputeExecutionContext* Context) override;
 
@@ -138,8 +194,8 @@ private:
 	void UpdateInstanceCountManager(FRHICommandListImmediate& RHICmdList);
 	void PrepareTicksForProxy(FRHICommandListImmediate& RHICmdList, FNiagaraSystemGpuComputeProxy* ComputeProxy, FNiagaraGpuDispatchList& GpuDispatchList);
 	void PrepareAllTicks(FRHICommandListImmediate& RHICmdList);
-	void ExecuteTicks(FRHICommandList& RHICmdList, ENiagaraGpuComputeTickStage::Type TickStage);
-	void DispatchStage(FRHICommandList& RHICmdList, const FNiagaraGPUSystemTick& Tick, const FNiagaraComputeInstanceData& InstanceData, const FNiagaraSimStageData& SimStageData);
+	void ExecuteTicks(FRDGBuilder& GraphBuilder, TConstArrayView<FViewInfo> Views, ENiagaraGpuComputeTickStage::Type TickStage);
+	void DispatchStage(FRDGBuilder& GraphBuilder, const FNiagaraGPUSystemTick& Tick, const FNiagaraComputeInstanceData& InstanceData, const FNiagaraSimStageData& SimStageData);
 
 	/**
 	 * Generate all the initial keys and values for a GPUSortManager sort batch.
@@ -164,9 +220,6 @@ private:
 
 	void FinishDispatches();
 
-	void UpdateFreeIDsListSizesBuffer(FRHICommandList& RHICmdList, uint32 NumInstances);
-	void UpdateFreeIDBuffers(FRHICommandList& RHICmdList, TConstArrayView<FNiagaraComputeExecutionContext*> Instances);
-
 	/** The shared GPUSortManager, used to register GPU sort tasks in order to generate sorted particle indices per emitter. */
 	TRefCountPtr<FGPUSortManager> GPUSortManager;
 	/** All sort tasks registered in AddSortedGPUSimulation(). Holds all the data required in GenerateSortKeys(). */
@@ -181,6 +234,7 @@ private:
 	uint32 FramesBeforeTickFlush = 0;
 
 	TConstArrayView<FViewInfo> CurrentPassViews;
+	mutable FRDGExternalAccessQueue CurrentPassExternalAccessQueue;
 
 	/** A buffer of list sizes used by UpdateFreeIDBuffers to allow overlapping several dispatches. */
 	FRWBuffer FreeIDListSizesBuffer;
diff --git a/Engine/Plugins/FX/Niagara/Source/Niagara/Private/NiagaraGpuComputeDispatchInterface.cpp b/Engine/Plugins/FX/Niagara/Source/Niagara/Private/NiagaraGpuComputeDispatchInterface.cpp
--- a/Engine/Plugins/FX/Niagara/Source/Niagara/Private/NiagaraGpuComputeDispatchInterface.cpp
+++ b/Engine/Plugins/FX/Niagara/Source/Niagara/Private/NiagaraGpuComputeDispatchInterface.cpp
@@ -6,8 +6,10 @@
 #include "NiagaraGpuReadbackManager.h"
 
 #include "Engine/World.h"
+#include "Runtime/Renderer/Private/SystemTextures.h"
 #include "SceneInterface.h"
 #include "FXSystem.h"
+#include "RenderGraphBuilder.h"
 
 FNiagaraGpuComputeDispatchInterface::FNiagaraGpuComputeDispatchInterface(EShaderPlatform InShaderPlatform, ERHIFeatureLevel::Type InFeatureLevel)
 	: ShaderPlatform(InShaderPlatform)
@@ -34,3 +36,114 @@ FNiagaraGpuComputeDispatchInterface* FNiagaraGpuComputeDispatchInterface::Get(FF
 {
 	return FXSceneInterface ? static_cast<FNiagaraGpuComputeDispatchInterface*>(FXSceneInterface->GetInterface(FNiagaraGpuComputeDispatch::Name)) : nullptr;
 }
+
+FRDGTextureRef FNiagaraGpuComputeDispatchInterface::GetBlackTexture(FRDGBuilder& GraphBuilder, ETextureDimension TextureDimension) const
+{
+	const FRDGSystemTextures& SystemTextures = FRDGSystemTextures::Get(GraphBuilder);
+
+	switch (TextureDimension)
+	{
+		case ETextureDimension::Texture2D:			return SystemTextures.Black;
+		case ETextureDimension::Texture2DArray:		return SystemTextures.BlackArray;
+		case ETextureDimension::Texture3D:			return SystemTextures.VolumetricBlack;
+		case ETextureDimension::TextureCube:		return SystemTextures.CubeBlack;
+		case ETextureDimension::TextureCubeArray:	return SystemTextures.CubeArrayBlack;
+		default: checkNoEntry(); return nullptr;
+	}
+}
+
+FRDGTextureSRVRef FNiagaraGpuComputeDispatchInterface::GetBlackTextureSRV(FRDGBuilder& GraphBuilder, ETextureDimension TextureDimension) const
+{
+	const FRDGSystemTextures& SystemTextures = FRDGSystemTextures::Get(GraphBuilder);
+
+	switch (TextureDimension)
+	{
+		case ETextureDimension::Texture2D:			return GraphBuilder.CreateSRV(SystemTextures.Black);
+		case ETextureDimension::Texture2DArray:		return GraphBuilder.CreateSRV(SystemTextures.BlackArray);
+		case ETextureDimension::Texture3D:			return GraphBuilder.CreateSRV(SystemTextures.VolumetricBlack);
+		case ETextureDimension::TextureCube:		return GraphBuilder.CreateSRV(SystemTextures.CubeBlack);
+		case ETextureDimension::TextureCubeArray:	return GraphBuilder.CreateSRV(SystemTextures.CubeArrayBlack);
+		default: checkNoEntry(); return nullptr;
+	}
+}
+
+FRDGTextureUAVRef FNiagaraGpuComputeDispatchInterface::GetEmptyTextureUAV(FRDGBuilder& GraphBuilder, EPixelFormat Format, ETextureDimension TextureDimension) const
+{
+	check(RHIIsTypedUAVStoreSupported(Format));
+
+	switch (TextureDimension)
+	{
+		case ETextureDimension::Texture2D:
+			return
+				GraphBuilder.CreateUAV(
+					GraphBuilder.CreateTexture(
+						FRDGTextureDesc::Create2D(FIntPoint(1, 1), Format, FClearValueBinding::Black, ETextureCreateFlags::ShaderResource | ETextureCreateFlags::UAV),
+						TEXT("NiagaraEmptyTextureUAV::Texture2D")
+					),
+					ERDGUnorderedAccessViewFlags::SkipBarrier
+				);
+
+		case ETextureDimension::Texture2DArray:
+			return
+				GraphBuilder.CreateUAV(
+					GraphBuilder.CreateTexture(
+						FRDGTextureDesc::Create2DArray(FIntPoint(1, 1), Format, FClearValueBinding::Black, ETextureCreateFlags::ShaderResource | ETextureCreateFlags::UAV, 1),
+						TEXT("NiagaraEmptyTextureUAV::Texture2DArray")
+					),
+					ERDGUnorderedAccessViewFlags::SkipBarrier
+				);
+
+		case ETextureDimension::Texture3D:
+			return
+				GraphBuilder.CreateUAV(
+					GraphBuilder.CreateTexture(
+						FRDGTextureDesc::Create3D(FIntVector(1, 1, 1), Format, FClearValueBinding::Black, ETextureCreateFlags::ShaderResource | ETextureCreateFlags::UAV),
+						TEXT("NiagaraEmptyTextureUAV::Texture3D")
+					),
+					ERDGUnorderedAccessViewFlags::SkipBarrier
+				);
+
+		case ETextureDimension::TextureCube:
+			return
+				GraphBuilder.CreateUAV(
+					GraphBuilder.CreateTexture(
+						FRDGTextureDesc::CreateCube(1, Format, FClearValueBinding::Black, ETextureCreateFlags::ShaderResource | ETextureCreateFlags::UAV),
+						TEXT("NiagaraEmptyTextureUAV::TextureCube")
+					),
+					ERDGUnorderedAccessViewFlags::SkipBarrier
+				);
+
+		case ETextureDimension::TextureCubeArray:
+			return
+				GraphBuilder.CreateUAV(
+					GraphBuilder.CreateTexture(
+						FRDGTextureDesc::CreateCubeArray(1, Format, FClearValueBinding::Black, ETextureCreateFlags::ShaderResource | ETextureCreateFlags::UAV, 1),
+						TEXT("NiagaraEmptyTextureUAV::TextureCubeArray")
+					),
+					ERDGUnorderedAccessViewFlags::SkipBarrier
+				);
+
+		default:
+			checkNoEntry();
+			return nullptr;
+	}
+}
+
+FRDGBufferUAVRef FNiagaraGpuComputeDispatchInterface::GetEmptyBufferUAV(FRDGBuilder& GraphBuilder, EPixelFormat Format) const
+{
+	const uint32 BytesPerElement = GPixelFormats[Format].BlockBytes;
+	return
+		GraphBuilder.CreateUAV(
+			GraphBuilder.CreateBuffer(
+				FRDGBufferDesc::CreateBufferDesc(BytesPerElement, 1),
+				TEXT("ENiagaraEmptyUAVType::Buffer")
+			),
+			Format,
+			ERDGUnorderedAccessViewFlags::SkipBarrier
+		);
+}
+
+FRDGBufferSRVRef FNiagaraGpuComputeDispatchInterface::GetEmptyBufferSRV(FRDGBuilder& GraphBuilder, EPixelFormat Format) const
+{
+	return GraphBuilder.CreateSRV(GSystemTextures.GetDefaultBuffer(GraphBuilder, GPixelFormats[Format].BlockBytes, FUintVector4(0, 0, 0, 0)), Format);
+}
diff --git a/Engine/Plugins/FX/Niagara/Source/Niagara/Private/NiagaraGpuReadbackManager.cpp b/Engine/Plugins/FX/Niagara/Source/Niagara/Private/NiagaraGpuReadbackManager.cpp
--- a/Engine/Plugins/FX/Niagara/Source/Niagara/Private/NiagaraGpuReadbackManager.cpp
+++ b/Engine/Plugins/FX/Niagara/Source/Niagara/Private/NiagaraGpuReadbackManager.cpp
@@ -4,9 +4,12 @@
 #include "NiagaraCommon.h"
 
 #include "CoreMinimal.h"
+#include "RenderGraphUtils.h"
 #include "RHI.h"
 #include "RHICommandList.h"
 
+const FName FNiagaraGpuReadbackManager::FenceName("NiagaraGpuReadback");
+
 FNiagaraGpuReadbackManager::FNiagaraGpuReadbackManager()
 {
 }
@@ -67,6 +70,106 @@ void FNiagaraGpuReadbackManager::WaitCompletion(FRHICommandListImmediate& RHICmd
 	TickInternal(true);
 }
 
+void FNiagaraGpuReadbackManager::EnqueueReadback(FRDGBuilder& GraphBuilder, FRDGBufferRef Buffer, FCompletionCallback Callback)
+{
+	check(IsInRenderingThread());
+	EnqueueReadback(GraphBuilder, Buffer, 0, Buffer->Desc.GetSize(), Callback);
+}
+
+void FNiagaraGpuReadbackManager::EnqueueReadback(FRDGBuilder& GraphBuilder, FRDGBufferRef Buffer, uint32 Offset, uint32 NumBytes, FCompletionCallback Callback)
+{
+	check(IsInRenderingThread());
+
+	FPendingReadback Readback;
+	Readback.StagingBuffers.Emplace(RHICreateStagingBuffer(), NumBytes);
+	Readback.Fence = RHICreateGPUFence(FenceName);
+	Readback.Callback = Callback;
+	Readback.Fence->Clear();
+
+	AddReadbackBufferPass(
+		GraphBuilder, RDG_EVENT_NAME("NiagaraReadback"), Buffer,
+		[Buffer, ReadbackBuffer=Readback.StagingBuffers[0].Key, Offset, NumBytes, Fence=Readback.Fence](FRHICommandList& RHICmdList)
+		{
+			RHICmdList.CopyToStagingBuffer(Buffer->GetRHI(), ReadbackBuffer, Offset, NumBytes);
+			RHICmdList.WriteGPUFence(Fence);
+		}
+	);
+
+	PendingReadbacks.Enqueue(Readback);
+}
+
+void FNiagaraGpuReadbackManager::EnqueueReadbacks(FRDGBuilder& GraphBuilder, TConstArrayView<FRDGBufferRef> Buffers, FCompletionCallback Callback)
+{
+	check(IsInRenderingThread());
+
+	FPendingReadback Readback;
+	Readback.Fence = RHICreateGPUFence(FenceName);
+	Readback.Callback = Callback;
+	Readback.Fence->Clear();
+
+	for (int32 i=0; i < Buffers.Num(); ++i)
+	{
+		FRDGBufferRef Buffer = Buffers[i];
+		FStagingBufferRHIRef ReadbackBuffer = RHICreateStagingBuffer();
+		const uint32 Offset = 0;
+		const uint32 NumBytes = Buffer->Desc.GetSize();
+
+		Readback.StagingBuffers.Emplace_GetRef(ReadbackBuffer, NumBytes);
+
+		const bool bWriteFence = i == Buffers.Num() - 1;
+		AddReadbackBufferPass(
+			GraphBuilder, RDG_EVENT_NAME("NiagaraReadback"), Buffer,
+			[Buffer, ReadbackBuffer, Offset=0, NumBytes, Fence=bWriteFence ? Readback.Fence : nullptr](FRHICommandListImmediate& RHICmdList)
+			{
+				RHICmdList.CopyToStagingBuffer(Buffer->GetRHI(), ReadbackBuffer, Offset, NumBytes);
+				if (Fence != nullptr)
+				{
+					RHICmdList.WriteGPUFence(Fence);
+				}
+			}
+		);
+	}
+
+	PendingReadbacks.Enqueue(Readback);
+}
+
+void FNiagaraGpuReadbackManager::EnqueueReadbacks(FRDGBuilder& GraphBuilder, TConstArrayView<FRDGBufferRequest> BufferRequests, FCompletionCallback Callback)
+{
+	check(IsInRenderingThread());
+
+	FPendingReadback Readback;
+	Readback.Fence = RHICreateGPUFence(FenceName);
+	Readback.Callback = Callback;
+	Readback.Fence->Clear();
+
+	for (int32 i = 0; i < BufferRequests.Num(); ++i)
+	{
+		FRDGBufferRef Buffer = BufferRequests[i].Buffer;
+		FStagingBufferRHIRef ReadbackBuffer = RHICreateStagingBuffer();
+		const uint32 Offset = BufferRequests[i].Offset;
+		const uint32 NumBytes = BufferRequests[i].Size;
+
+		check(NumBytes > 0);
+
+		Readback.StagingBuffers.Emplace_GetRef(ReadbackBuffer, NumBytes);
+
+		const bool bWriteFence = i == BufferRequests.Num() - 1;
+		AddReadbackBufferPass(
+			GraphBuilder, RDG_EVENT_NAME("NiagaraReadback"), Buffer,
+			[Buffer, ReadbackBuffer, Offset = 0, NumBytes, Fence=bWriteFence ? Readback.Fence : nullptr](FRHICommandListImmediate& RHICmdList)
+			{
+				RHICmdList.CopyToStagingBuffer(Buffer->GetRHI(), ReadbackBuffer, Offset, NumBytes);
+				if (Fence != nullptr)
+				{
+					RHICmdList.WriteGPUFence(Fence);
+				}
+			}
+		);
+	}
+
+	PendingReadbacks.Enqueue(Readback);
+}
+
 void FNiagaraGpuReadbackManager::EnqueueReadback(FRHICommandList& RHICmdList, FRHIBuffer* Buffer, FCompletionCallback Callback)
 {
 	check(IsInRenderingThread());
@@ -76,8 +179,6 @@ void FNiagaraGpuReadbackManager::EnqueueReadback(FRHICommandList& RHICmdList, FR
 
 void FNiagaraGpuReadbackManager::EnqueueReadback(FRHICommandList& RHICmdList, FRHIBuffer* Buffer, uint32 Offset, uint32 NumBytes, FCompletionCallback Callback)
 {
-	static const FName FenceName(TEXT("NiagaraGpuReadback"));
-
 	check(IsInRenderingThread());
 
 	FPendingReadback Readback;
@@ -95,8 +196,6 @@ void FNiagaraGpuReadbackManager::EnqueueReadback(FRHICommandList& RHICmdList, FR
 
 void FNiagaraGpuReadbackManager::EnqueueReadbacks(FRHICommandList& RHICmdList, TConstArrayView<FRHIBuffer*> Buffers, FCompletionCallback Callback)
 {
-	static const FName FenceName(TEXT("NiagaraGpuReadback"));
-
 	check(IsInRenderingThread());
 
 	FPendingReadback Readback;
@@ -116,8 +215,6 @@ void FNiagaraGpuReadbackManager::EnqueueReadbacks(FRHICommandList& RHICmdList, T
 
 void FNiagaraGpuReadbackManager::EnqueueReadbacks(FRHICommandList& RHICmdList, TConstArrayView<FBufferRequest> BufferRequests, FCompletionCallback Callback)
 {
-	static const FName FenceName(TEXT("NiagaraGpuReadback"));
-
 	check(IsInRenderingThread());
 
 	FPendingReadback Readback;
diff --git a/Engine/Plugins/FX/Niagara/Source/Niagara/Public/NiagaraGpuComputeDispatchInterface.h b/Engine/Plugins/FX/Niagara/Source/Niagara/Public/NiagaraGpuComputeDispatchInterface.h
--- a/Engine/Plugins/FX/Niagara/Source/Niagara/Public/NiagaraGpuComputeDispatchInterface.h
+++ b/Engine/Plugins/FX/Niagara/Source/Niagara/Public/NiagaraGpuComputeDispatchInterface.h
@@ -7,6 +7,9 @@
 #include "NiagaraGPUInstanceCountManager.h"
 #include "FXSystem.h"
 
+class RENDERCORE_API FRDGBuilder;
+class FRDGExternalAccessQueue;
+
 class FNiagaraAsyncGpuTraceHelper;
 struct FNiagaraComputeExecutionContext;
 class FNiagaraGpuComputeDebug;
@@ -66,11 +69,28 @@ public:
 	FNiagaraGpuReadbackManager* GetGpuReadbackManager() const { return GpuReadbackManagerPtr.Get(); }
 
 	/** Get access to Niagara's GpuReadbackManager. */
+	//UE_DEPRECATED(5.1, "The UAV Pool will be removed, please update your code to support RenderGraph.")
 	FNiagaraEmptyUAVPool* GetEmptyUAVPool() const { return EmptyUAVPoolPtr.Get(); }
 
-	/** Convenience wrapped to get a UAV from the poo. */
+	/** Convenience wrapper to get a UAV from the pool. */
+	//UE_DEPRECATED(5.1, "The UAV Pool will be removed, please update your code to support RenderGraph.")
 	FRHIUnorderedAccessView* GetEmptyUAVFromPool(FRHICommandList& RHICmdList, EPixelFormat Format, ENiagaraEmptyUAVType Type) const { return EmptyUAVPoolPtr->GetEmptyUAVFromPool(RHICmdList, Format, Type); }
 
+	/** Helper function to return an RDG Texture where the texture contains 0 for all channels. */
+	FRDGTextureRef GetBlackTexture(FRDGBuilder& GraphBuilder, ETextureDimension TextureDimension) const;
+
+	/** Helper function to return a RDG Texture SRV where the texture contains 0 for all channels. */
+	FRDGTextureSRVRef GetBlackTextureSRV(FRDGBuilder& GraphBuilder, ETextureDimension TextureDimension) const;
+
+	/** Helper function to return a RDG Texture UAV you don't care about the contents of or the results, i.e. to use as a dummy binding. */
+	FRDGTextureUAVRef GetEmptyTextureUAV(FRDGBuilder& GraphBuilder, EPixelFormat Format, ETextureDimension TextureDimension) const;
+
+	/** Helper function to return a Buffer UAV you don't care about the contents of or the results, i.e. to use as a dummy binding. */
+	FRDGBufferUAVRef GetEmptyBufferUAV(FRDGBuilder& GraphBuilder, EPixelFormat Format) const;
+
+	/** Helper function to return a Buffer SRV which will contain 1 element of 0 value, i.e. to use as a dummy binding. */
+	FRDGBufferSRVRef GetEmptyBufferSRV(FRDGBuilder& GraphBuilder, EPixelFormat Format) const;
+
 	/**
 	Call this to force all pending ticks to be flushed from the batcher.
 	Doing so will execute them outside of a view context which may result in undesirable results.
diff --git a/Engine/Plugins/FX/Niagara/Source/Niagara/Public/NiagaraRenderGraphUtils.h b/Engine/Plugins/FX/Niagara/Source/Niagara/Public/NiagaraRenderGraphUtils.h
new file mode 100644
--- /dev/null
+++ b/Engine/Plugins/FX/Niagara/Source/Niagara/Public/NiagaraRenderGraphUtils.h
@@ -0,0 +1,65 @@
+// Copyright Epic Games, Inc. All Rights Reserved.
+#pragma once
+
+#include "RenderGraphResources.h"
+
+//////////////////////////////////////////////////////////////////////////
+// Helper for wrapping a persistent RW buffer
+struct FNiagaraPooledRWBuffer
+{
+	void Initialize(FRDGBuilder& GraphBuilder, const TCHAR* ResourceName, EPixelFormat PixelFormat, const FRDGBufferDesc& BufferDesc);
+	void Initialize(FRDGBuilder& GraphBuilder, const TCHAR* ResourceName, EPixelFormat PixelFormat, const uint32 BytesPerElemenet, const uint32 NumElements, EBufferUsageFlags UsageFlags = EBufferUsageFlags::None);
+	void Release();
+
+	bool IsValid() const { return Buffer.IsValid(); }
+
+
+	const TRefCountPtr<FRDGPooledBuffer>& GetPooledBuffer() const { return Buffer; }
+	FRDGBufferRef GetOrCreateBuffer(FRDGBuilder& GraphBuilder);
+	FRDGBufferSRVRef GetOrCreateSRV(FRDGBuilder& GraphBuilder);
+	FRDGBufferUAVRef GetOrCreateUAV(FRDGBuilder& GraphBuilder);
+
+	void EndGraphUsage();
+
+private:
+	TRefCountPtr<FRDGPooledBuffer>	Buffer;
+	EPixelFormat					PixelFormat = PF_Unknown;
+
+	// Transient for the lifetime of the GraphBuilder, call PostSimulate to clear
+	FRDGBufferRef					TransientRDGBuffer = nullptr;
+	FRDGBufferSRVRef				TransientRDGSRV = nullptr;
+	FRDGBufferUAVRef				TransientRDGUAV = nullptr;
+#if STATS
+	uint64							NumBytes = 0;
+#endif
+};
+
+//////////////////////////////////////////////////////////////////////////
+// Helper for wrapping a persistent RW texture
+struct FNiagaraPooledRWTexture
+{
+	void Initialize(FRDGBuilder& GraphBuilder, const TCHAR* ResourceName, const FRDGTextureDesc& TextureDesc);
+	void Release();
+
+	bool IsValid() const { return Texture.IsValid(); }
+
+	const TRefCountPtr<IPooledRenderTarget>& GetPooledTexture() const { return Texture; }
+	FRDGTextureRef GetOrCreateTexture(FRDGBuilder& GraphBuilder);
+	FRDGTextureSRVRef GetOrCreateSRV(FRDGBuilder& GraphBuilder);
+	FRDGTextureUAVRef GetOrCreateUAV(FRDGBuilder& GraphBuilder);
+
+	void CopyToTexture(FRDGBuilder& GraphBuilder, FRHITexture* DestinationTexture, const TCHAR* NameIfNotRegistered);
+
+	void EndGraphUsage();
+
+private:
+	TRefCountPtr<IPooledRenderTarget>	Texture;
+
+	// Transient for the lifetime of the GraphBuilder, call PostSimulate to clear
+	FRDGTextureRef						TransientRDGTexture = nullptr;
+	FRDGTextureSRVRef					TransientRDGSRV = nullptr;
+	FRDGTextureUAVRef					TransientRDGUAV = nullptr;
+#if STATS
+	uint64								NumBytes = 0;
+#endif
+};
diff --git a/Engine/Plugins/FX/Niagara/Source/NiagaraCore/Public/NiagaraDataInterfaceBase.h b/Engine/Plugins/FX/Niagara/Source/NiagaraCore/Public/NiagaraDataInterfaceBase.h
--- a/Engine/Plugins/FX/Niagara/Source/NiagaraCore/Public/NiagaraDataInterfaceBase.h
+++ b/Engine/Plugins/FX/Niagara/Source/NiagaraCore/Public/NiagaraDataInterfaceBase.h
@@ -107,8 +107,7 @@ class NIAGARACORE_API UNiagaraDataInterfaceBase : public UNiagaraMergeable
 	GENERATED_UCLASS_BODY()
 
 public:
-	/** Does this data interface require legacy shader bindings. */
-	//UE_DEPRECATED(5.1, "UseLegacyShaderBindings is deprecated and will be removed.")
+	//UE_DEPRECATED(5.1, "Support for legacy shaders bindings will be removed in a future release, please upgrade your data interface.")
 	virtual bool UseLegacyShaderBindings() const { return true; }
 
 	/**
diff --git a/Engine/Plugins/FX/Niagara/Source/NiagaraEditor/Private/NiagaraBakerRendererOutputVolumeTexture.cpp b/Engine/Plugins/FX/Niagara/Source/NiagaraEditor/Private/NiagaraBakerRendererOutputVolumeTexture.cpp
--- a/Engine/Plugins/FX/Niagara/Source/NiagaraEditor/Private/NiagaraBakerRendererOutputVolumeTexture.cpp
+++ b/Engine/Plugins/FX/Niagara/Source/NiagaraEditor/Private/NiagaraBakerRendererOutputVolumeTexture.cpp
@@ -214,9 +214,8 @@ void FNiagaraBakerRendererOutputVolumeTexture::RenderPreview(UNiagaraBakerOutput
 				{
 					if ( const FRenderTargetVolumeRWInstanceData_RenderThread* InstanceData_RT = Proxy_RT->SystemInstancesToProxyData_RT.Find(SystemInstanceID_RT) )
 					{
-						OutTexture = InstanceData_RT->TextureRHI;
+						OutTexture = InstanceData_RT->RenderTarget->GetRHI();
 						OutSamplerState = InstanceData_RT->SamplerStateRHI;
-			//				OutSamplerState = TStaticSamplerState<SF_Bilinear>::GetRHI();
 					}
 				}
 			);
@@ -275,7 +274,7 @@ void FNiagaraBakerRendererOutputVolumeTexture::RenderPreview(UNiagaraBakerOutput
 						{
 							if ( RT_InstanceData->CurrentData != nullptr )
 							{
-								OutTexture = RT_InstanceData->CurrentData->GridBuffer.Buffer;
+								OutTexture = RT_InstanceData->CurrentData->GetPooledTexture()->GetRHI();
 								OutSamplerState = TStaticSamplerState<SF_Bilinear>::GetRHI();
 							}
 						}
@@ -424,7 +423,7 @@ void FNiagaraBakerRendererOutputVolumeTexture::BakeFrame(UNiagaraBakerOutput* In
 				{
 					*OutTextureSize_RT = InstanceData_RT->Size;
 					RHICmdList.Read3DSurfaceFloatData(
-						InstanceData_RT->TextureRHI,
+						InstanceData_RT->RenderTarget->GetRHI(),
 						FIntRect(0, 0, InstanceData_RT->Size.X, InstanceData_RT->Size.Y),
 						FIntPoint(0, InstanceData_RT->Size.Z),
 						*OutTextureData_RT
@@ -452,7 +451,7 @@ void FNiagaraBakerRendererOutputVolumeTexture::BakeFrame(UNiagaraBakerOutput* In
 						*OutTextureSize_RT = TextureSize;
 
 						RHICmdList.Read3DSurfaceFloatData(
-							InstanceData_RT->CurrentData->GridBuffer.Buffer,
+							InstanceData_RT->CurrentData->GetPooledTexture()->GetRHI(),
 							FIntRect(0, 0, TextureSize.X, TextureSize.Y),
 							FIntPoint(0, TextureSize.Z),
 							*OutTextureData_RT
diff --git a/Engine/Plugins/FX/Niagara/Source/NiagaraShader/Private/NiagaraAsyncGpuTraceProvider.cpp b/Engine/Plugins/FX/Niagara/Source/NiagaraShader/Private/NiagaraAsyncGpuTraceProvider.cpp
--- a/Engine/Plugins/FX/Niagara/Source/NiagaraShader/Private/NiagaraAsyncGpuTraceProvider.cpp
+++ b/Engine/Plugins/FX/Niagara/Source/NiagaraShader/Private/NiagaraAsyncGpuTraceProvider.cpp
@@ -260,9 +260,9 @@ void FNiagaraAsyncGpuTraceProvider::BuildCollisionGroupHashMap(FRHICommandList&
 	RHICmdList.Transition(FRHITransitionInfo(Result.HashToCollisionGroups.UAV, ERHIAccess::Unknown, ERHIAccess::UAVCompute));
 
 	//First we have to clear the buffers. Can probably do this better.
-	NiagaraFillGPUIntBuffer(RHICmdList, FeatureLevel, Result.PrimIdHashTable, 0);
+	NiagaraFillGPUIntBuffer(RHICmdList, FeatureLevel, Result.PrimIdHashTable.UAV, Result.PrimIdHashTable.NumBytes / sizeof(uint32), 0);
 	RHICmdList.Transition(FRHITransitionInfo(Result.PrimIdHashTable.UAV, ERHIAccess::UAVCompute, ERHIAccess::UAVCompute));
-	NiagaraFillGPUIntBuffer(RHICmdList, FeatureLevel, Result.HashToCollisionGroups, INDEX_NONE);
+	NiagaraFillGPUIntBuffer(RHICmdList, FeatureLevel, Result.HashToCollisionGroups.UAV, Result.HashToCollisionGroups.NumBytes / sizeof(uint32), INDEX_NONE);
 	RHICmdList.Transition(FRHITransitionInfo(Result.HashToCollisionGroups.UAV, ERHIAccess::UAVCompute, ERHIAccess::UAVCompute));
 
 	TShaderMapRef<FNiagaraUpdateCollisionGroupMapCS> ComputeShader(GetGlobalShaderMap(GMaxRHIFeatureLevel));
