- Niagara Compute Dipsatch modified to use FRDGBuilder through the pipeline

diff --git a/Engine/Plugins/FX/Niagara/Source/NiagaraShader/Private/NiagaraClearCounts.cpp b/Engine/Plugins/FX/Niagara/Source/NiagaraShader/Private/NiagaraClearCounts.cpp
--- a/Engine/Plugins/FX/Niagara/Source/NiagaraShader/Private/NiagaraClearCounts.cpp
+++ b/Engine/Plugins/FX/Niagara/Source/NiagaraShader/Private/NiagaraClearCounts.cpp
@@ -12,7 +12,7 @@
 
 //////////////////////////////////////////////////////////////////////////
 
-class NIAGARASHADER_API FNiagaraClearCountsIntCS : public FGlobalShader
+class FNiagaraClearCountsIntCS : public FGlobalShader
 {
 	DECLARE_GLOBAL_SHADER(FNiagaraClearCountsIntCS);
 	SHADER_USE_PARAMETER_STRUCT(FNiagaraClearCountsIntCS, FGlobalShader);
@@ -32,14 +32,68 @@ class NIAGARASHADER_API FNiagaraClearCountsIntCS : public FGlobalShader
 	}
 
 	BEGIN_SHADER_PARAMETER_STRUCT(FParameters, )
-		SHADER_PARAMETER_UAV(RWBuffer<int>,	CountBuffer)
-		SHADER_PARAMETER_SRV(Buffer<int>,	CountsAndValuesBuffer)
-		SHADER_PARAMETER(uint32,			NumCounts)
+		SHADER_PARAMETER_RDG_BUFFER_UAV(RWBuffer<int>,	CountBuffer)
+		SHADER_PARAMETER_RDG_BUFFER_SRV(Buffer<int>,	CountsAndValuesBuffer)
+		SHADER_PARAMETER(uint32,						NumCounts)
 	END_SHADER_PARAMETER_STRUCT()
 };
 
 IMPLEMENT_GLOBAL_SHADER(FNiagaraClearCountsIntCS, "/Plugin/FX/Niagara/Private/NiagaraClearCounts.usf", "ClearCountsIntCS", SF_Compute);
 
+void NiagaraClearCounts::ClearCountsInt(FRDGBuilder& GraphBuilder, FRDGBufferUAVRef UAV, TConstArrayView<TPair<uint32, int32>> IndexAndValueArray)
+{
+	const int32 NumInts = IndexAndValueArray.Num() * 2;
+	FRDGBufferRef CountsAndValuesBuffer = GraphBuilder.CreateBuffer(FRDGBufferDesc::CreateBufferDesc(sizeof(int32), NumInts), TEXT("NiagaraClearCountsInt::CountsAndValuesBuffer"));
+	GraphBuilder.QueueBufferUpload(CountsAndValuesBuffer, IndexAndValueArray.GetData(), sizeof(int32) * NumInts);
+
+	TShaderMapRef<FNiagaraClearCountsIntCS> ComputeShader(GetGlobalShaderMap(GMaxRHIFeatureLevel));
+	const uint32 NumThreadGroups = FMath::DivideAndRoundUp<uint32>(IndexAndValueArray.Num(), FNiagaraClearCountsIntCS::ThreadGroupSize);
+
+	FNiagaraClearCountsIntCS::FParameters* Parameters = GraphBuilder.AllocParameters<FNiagaraClearCountsIntCS::FParameters>();
+	Parameters->CountBuffer = UAV;
+	Parameters->CountsAndValuesBuffer = GraphBuilder.CreateSRV(CountsAndValuesBuffer, PF_R32_SINT);
+	Parameters->NumCounts = IndexAndValueArray.Num();
+
+	FComputeShaderUtils::AddPass(
+		GraphBuilder,
+		RDG_EVENT_NAME("NiagaraClearCountsInt"),
+		ERDGPassFlags::Compute,
+		ComputeShader,
+		Parameters,
+		FIntVector(NumThreadGroups, 1, 1)
+	);
+}
+
+//////////////////////////////////////////////////////////////////////////
+
+class FNiagaraClearCountsIntCS_Legacy : public FGlobalShader
+{
+	DECLARE_GLOBAL_SHADER(FNiagaraClearCountsIntCS_Legacy);
+	SHADER_USE_PARAMETER_STRUCT(FNiagaraClearCountsIntCS_Legacy, FGlobalShader);
+
+	static constexpr uint32 ThreadGroupSize = 64;
+
+	static bool ShouldCompilePermutation(const FGlobalShaderPermutationParameters& Parameters)
+	{
+		return !IsMobilePlatform(Parameters.Platform);
+	}
+
+	static void ModifyCompilationEnvironment(const FGlobalShaderPermutationParameters& Parameters, FShaderCompilerEnvironment& OutEnvironment)
+	{
+		FGlobalShader::ModifyCompilationEnvironment(Parameters, OutEnvironment);
+		OutEnvironment.SetDefine(TEXT("NIAGARA_CLEAR_COUNTS_INT_CS"), 1);
+		OutEnvironment.SetDefine(TEXT("THREADGROUP_SIZE"), ThreadGroupSize);
+	}
+
+	BEGIN_SHADER_PARAMETER_STRUCT(FParameters, )
+		SHADER_PARAMETER_UAV(RWBuffer<int>, CountBuffer)
+		SHADER_PARAMETER_SRV(Buffer<int>, CountsAndValuesBuffer)
+		SHADER_PARAMETER(uint32, NumCounts)
+	END_SHADER_PARAMETER_STRUCT()
+};
+
+IMPLEMENT_GLOBAL_SHADER(FNiagaraClearCountsIntCS_Legacy, "/Plugin/FX/Niagara/Private/NiagaraClearCounts.usf", "ClearCountsIntCS", SF_Compute);
+
 void NiagaraClearCounts::ClearCountsInt(FRHICommandList& RHICmdList, FRHIUnorderedAccessView* UAV, TConstArrayView<TPair<uint32, int32>> IndexAndValueArray)
 {
 	check(UAV != nullptr);
@@ -54,14 +108,14 @@ void NiagaraClearCounts::ClearCountsInt(FRHICommandList& RHICmdList, FRHIUnorder
 		RHIUnlockBuffer(IndexAndValueBuffer.Buffer);
 	}
 
-	FNiagaraClearCountsIntCS::FParameters ShaderParameters;
+	FNiagaraClearCountsIntCS_Legacy::FParameters ShaderParameters;
 	ShaderParameters.CountBuffer			= UAV;
 	ShaderParameters.CountsAndValuesBuffer	= IndexAndValueBuffer.SRV;
 	ShaderParameters.NumCounts				= IndexAndValueArray.Num();
 
-	TShaderMapRef<FNiagaraClearCountsIntCS> ComputeShader(GetGlobalShaderMap(GMaxRHIFeatureLevel));
+	TShaderMapRef<FNiagaraClearCountsIntCS_Legacy> ComputeShader(GetGlobalShaderMap(GMaxRHIFeatureLevel));
 	FRHIComputeShader* ShaderRHI = ComputeShader.GetComputeShader();
-	const uint32 NumThreadGroups = FMath::DivideAndRoundUp<uint32>(IndexAndValueArray.Num(), FNiagaraClearCountsIntCS::ThreadGroupSize);
+	const uint32 NumThreadGroups = FMath::DivideAndRoundUp<uint32>(IndexAndValueArray.Num(), FNiagaraClearCountsIntCS_Legacy::ThreadGroupSize);
 
 	RHICmdList.Transition(FRHITransitionInfo(UAV, ERHIAccess::Unknown, ERHIAccess::UAVCompute));
 	SetComputePipelineState(RHICmdList, ShaderRHI);
diff --git a/Engine/Plugins/FX/Niagara/Source/NiagaraShader/Private/NiagaraDebugShaders.cpp b/Engine/Plugins/FX/Niagara/Source/NiagaraShader/Private/NiagaraDebugShaders.cpp
--- a/Engine/Plugins/FX/Niagara/Source/NiagaraShader/Private/NiagaraDebugShaders.cpp
+++ b/Engine/Plugins/FX/Niagara/Source/NiagaraShader/Private/NiagaraDebugShaders.cpp
@@ -66,11 +66,11 @@ class NIAGARASHADER_API FNiagaraVisualizeTexturePS : public FGlobalShader
 		SHADER_PARAMETER(uint32,		TickCounter)
 		SHADER_PARAMETER(uint32,		TextureSlice)
 
-		SHADER_PARAMETER_TEXTURE(Texture2D, Texture2DObject)
-		SHADER_PARAMETER_TEXTURE(Texture2DArray, Texture2DArrayObject)
-		SHADER_PARAMETER_TEXTURE(Texture3D, Texture3DObject)
-		SHADER_PARAMETER_TEXTURE(TextureCube, TextureCubeObject)
-		SHADER_PARAMETER_SAMPLER(SamplerState, TextureSampler)
+		SHADER_PARAMETER_RDG_TEXTURE(Texture2D,			Texture2DObject)
+		SHADER_PARAMETER_RDG_TEXTURE(Texture2DArray,	Texture2DArrayObject)
+		SHADER_PARAMETER_RDG_TEXTURE(Texture3D,			Texture3DObject)
+		SHADER_PARAMETER_RDG_TEXTURE(TextureCube,		TextureCubeObject)
+		SHADER_PARAMETER_SAMPLER(SamplerState,			TextureSampler)
 
 		RENDER_TARGET_BINDING_SLOTS()
 	END_SHADER_PARAMETER_STRUCT()
@@ -293,11 +293,11 @@ void NiagaraDebugShaders::DrawDebugLines(
 void NiagaraDebugShaders::VisualizeTexture(
 	class FRDGBuilder& GraphBuilder, const FViewInfo& View, const FScreenPassRenderTarget& Output,
 	const FIntPoint& Location, const int32& DisplayHeight,
-	const FIntVector4& InAttributesToVisualize, FRHITexture* Texture, const FIntVector4& NumTextureAttributes, uint32 TickCounter,
+	const FIntVector4& InAttributesToVisualize, FRDGTextureRef Texture, const FIntVector4& NumTextureAttributes, uint32 TickCounter,
 	const FVector2D& PreviewDisplayRange
 )
 {
-	FIntVector TextureSize = Texture->GetSizeXYZ();
+	FIntVector TextureSize = Texture->Desc.GetSize();
 	if (NumTextureAttributes.X > 0)
 	{
 		check(NumTextureAttributes.Y > 0);
@@ -338,26 +338,22 @@ void NiagaraDebugShaders::VisualizeTexture(
 		}
 	}
 
-	FRHITexture2D* Texture2D = Texture->GetTexture2D();
-	FRHITexture2DArray* Texture2DArray = Texture->GetTexture2DArray();
-	FRHITexture3D* Texture3D = Texture->GetTexture3D();
-	FRHITextureCube* TextureCube = Texture->GetTextureCube();
-
 	// Set Shaders & State
 	FNiagaraVisualizeTexturePS::FPermutationDomain PermutationVector;
-	if (Texture2D != nullptr)
+	if (Texture->Desc.Dimension == ETextureDimension::Texture2D)
 	{
 		PermutationVector.Set<FNiagaraVisualizeTexturePS::FTextureType>(0);
 	}
-	else if (Texture2DArray != nullptr)
+	else if (Texture->Desc.Dimension == ETextureDimension::Texture2DArray)
 	{
 		PermutationVector.Set<FNiagaraVisualizeTexturePS::FTextureType>(1);
+		TextureSize.Z = NumAttributesToVisualizeValue == 0 ? Texture->Desc.ArraySize : 1;
 	}
-	else if (Texture3D != nullptr)
+	else if (Texture->Desc.Dimension == ETextureDimension::Texture3D)
 	{
 		PermutationVector.Set<FNiagaraVisualizeTexturePS::FTextureType>(2);
 	}
-	else if (TextureCube != nullptr)
+	else if (Texture->Desc.Dimension == ETextureDimension::TextureCube)
 	{
 		PermutationVector.Set<FNiagaraVisualizeTexturePS::FTextureType>(3);
 		TextureSize.X *= 3;
@@ -368,7 +364,7 @@ void NiagaraDebugShaders::VisualizeTexture(
 		return;
 	}
 
-	switch (Texture->GetFormat())
+	switch (Texture->Desc.Format)
 	{
 		case PF_R32_UINT:
 		case PF_R32_SINT:
@@ -401,29 +397,27 @@ void NiagaraDebugShaders::VisualizeTexture(
 	const int32 SlicesWidth = FMath::Clamp(FMath::DivideAndRoundUp(AvailableWidth, DisplaySize.X + 1), 1, TextureSize.Z);
 
 	const FVector2D::FReal DisplayScale = (PreviewDisplayRange.Y > PreviewDisplayRange.X) ? (1.0f / (PreviewDisplayRange.Y - PreviewDisplayRange.X)) : 1.0f;
-	const FVector4 PerChannelScale(DisplayScale, DisplayScale, DisplayScale, DisplayScale);
-	const FVector4 PerChannelBias(-PreviewDisplayRange.X, -PreviewDisplayRange.X, -PreviewDisplayRange.X, -PreviewDisplayRange.X);
+	const FVector4f PerChannelScale(DisplayScale, DisplayScale, DisplayScale, DisplayScale);
+	const FVector4f PerChannelBias(-PreviewDisplayRange.X, -PreviewDisplayRange.X, -PreviewDisplayRange.X, -PreviewDisplayRange.X);
 
 	for (int32 iSlice = 0; iSlice < SlicesWidth; ++iSlice)
 	{
 		FNiagaraVisualizeTexturePS::FParameters* PassParameters = GraphBuilder.AllocParameters<FNiagaraVisualizeTexturePS::FParameters>();
-		{
-			PassParameters->NumTextureAttributes = NumTextureAttributes;
-			PassParameters->NumAttributesToVisualize = NumAttributesToVisualizeValue;
-			PassParameters->AttributesToVisualize = AttributesToVisualize;
-			PassParameters->TextureDimensions = TextureSize;
-			PassParameters->PerChannelScale = (FVector4f)PerChannelScale; // LWC_TODO: precision loss
-			PassParameters->PerChannelBias = (FVector4f)PerChannelBias; // LWC_TODO: precision loss
-			PassParameters->DebugFlags = GNiagaraGpuComputeDebug_ShowNaNInf != 0 ? 1 : 0;
-			PassParameters->TickCounter = TickCounter;
-			PassParameters->TextureSlice = iSlice;
-			PassParameters->Texture2DObject = Texture2D;
-			PassParameters->Texture2DArrayObject = Texture2DArray;
-			PassParameters->Texture3DObject = Texture3D;
-			PassParameters->TextureCubeObject = TextureCube;
-			PassParameters->TextureSampler = TStaticSamplerState<SF_Point>::GetRHI();
-			PassParameters->RenderTargets[0] = Output.GetRenderTargetBinding();
-		}
+		PassParameters->NumTextureAttributes		= NumTextureAttributes;
+		PassParameters->NumAttributesToVisualize	= NumAttributesToVisualizeValue;
+		PassParameters->AttributesToVisualize		= AttributesToVisualize;
+		PassParameters->TextureDimensions			= TextureSize;
+		PassParameters->PerChannelScale				= PerChannelScale;
+		PassParameters->PerChannelBias				= PerChannelBias;
+		PassParameters->DebugFlags					= GNiagaraGpuComputeDebug_ShowNaNInf != 0 ? 1 : 0;
+		PassParameters->TickCounter					= TickCounter;
+		PassParameters->TextureSlice				= iSlice;
+		PassParameters->Texture2DObject				= Texture;
+		PassParameters->Texture2DArrayObject		= Texture;
+		PassParameters->Texture3DObject				= Texture;
+		PassParameters->TextureCubeObject			= Texture;
+		PassParameters->TextureSampler				= TStaticSamplerState<SF_Point>::GetRHI();
+		PassParameters->RenderTargets[0]			= Output.GetRenderTargetBinding();
 
 		const float OffsetX = float(iSlice) * DisplaySize.X + 1;
 
diff --git a/Engine/Plugins/FX/Niagara/Source/NiagaraShader/Private/NiagaraGenerateMips.cpp b/Engine/Plugins/FX/Niagara/Source/NiagaraShader/Private/NiagaraGenerateMips.cpp
--- a/Engine/Plugins/FX/Niagara/Source/NiagaraShader/Private/NiagaraGenerateMips.cpp
+++ b/Engine/Plugins/FX/Niagara/Source/NiagaraShader/Private/NiagaraGenerateMips.cpp
@@ -20,8 +20,8 @@ public:
 		SHADER_PARAMETER(FVector2f, DstTexelSize)
 		SHADER_PARAMETER(int32, KernelHWidth)
 		SHADER_PARAMETER(FIntPoint, MipOutSize)
-		SHADER_PARAMETER_UAV(RWTexture2D<float4>, MipOutUAV)
-		SHADER_PARAMETER_SRV(Texture2D, MipInSRV)
+		SHADER_PARAMETER_RDG_TEXTURE_UAV(RWTexture2D<float4>, MipOutUAV)
+		SHADER_PARAMETER_RDG_TEXTURE_SRV(Texture2D, MipInSRV)
 		SHADER_PARAMETER_SAMPLER(SamplerState, MipInSampler)
 	END_SHADER_PARAMETER_STRUCT()
 
@@ -39,12 +39,15 @@ public:
 
 IMPLEMENT_GLOBAL_SHADER(FNiagaraGenerateMipsCS, "/Plugin/FX/Niagara/Private/NiagaraGenerateMips.usf", "MainCS", SF_Compute);
 
-void NiagaraGenerateMips::GenerateMips(FRHICommandList& RHICmdList, FRHITexture2D* TextureRHI, ENiagaraMipMapGenerationType GenType)
+void NiagaraGenerateMips::GenerateMips(FRDGBuilder& GraphBuilder, FRDGTextureRef RDGTexture, ENiagaraMipMapGenerationType GenType)
 {
-	SCOPED_DRAW_EVENT(RHICmdList, NiagaraGenerateMips);
+	RDG_RHI_EVENT_SCOPE(GraphBuilder, NiagaraGenerateMips);
+
+	const FRDGTextureDesc& TextureDesc = RDGTexture->Desc;
+	check(TextureDesc.Dimension == ETextureDimension::Texture2D);
 
 	// Select compute shader variant (normal vs. sRGB etc.)
-	bool bMipsSRGB = EnumHasAnyFlags(TextureRHI->GetFlags(), TexCreate_SRGB);
+	bool bMipsSRGB = EnumHasAnyFlags(TextureDesc.Flags, TexCreate_SRGB);
 #if PLATFORM_ANDROID
 	if (IsVulkanPlatform(GMaxRHIShaderPlatform))
 	{
@@ -59,38 +62,30 @@ void NiagaraGenerateMips::GenerateMips(FRHICommandList& RHICmdList, FRHITexture2
 	PermutationVector.Set<FNiagaraGenerateMipsCS::FGenMipsSRGB>(bMipsSRGB);
 	PermutationVector.Set<FNiagaraGenerateMipsCS::FGaussianBlur>(bBlurMips);
 	TShaderMapRef<FNiagaraGenerateMipsCS> ComputeShader(GetGlobalShaderMap(GMaxRHIFeatureLevel), PermutationVector);
-	FRHIComputeShader* ShaderRHI = ComputeShader.GetComputeShader();
 
 	const int32 KernelHWidth = bBlurMips ? (int32(GenType) + 1 - int32(ENiagaraMipMapGenerationType::Blur1)) : 1;
 
-	const int32 NumMips = TextureRHI->GetNumMips();
-	for ( int32 iDstMip=1; iDstMip < NumMips; ++iDstMip)
+	for ( int32 iDstMip=1; iDstMip < TextureDesc.NumMips; ++iDstMip)
 	{
 		const int32 iSrcMip = iDstMip - 1;
-		const FIntPoint SrcMipSize(FMath::Max<int32>(TextureRHI->GetSizeX() >> iSrcMip, 1), FMath::Max<int32>(TextureRHI->GetSizeY() >> iSrcMip, 1));
-		const FIntPoint DstMipSize(FMath::Max<int32>(TextureRHI->GetSizeX() >> iDstMip, 1), FMath::Max<int32>(TextureRHI->GetSizeY() >> iDstMip, 1));
-
-		FShaderResourceViewRHIRef MipInSRV = RHICreateShaderResourceView(TextureRHI, iSrcMip);
-		FUnorderedAccessViewRHIRef MipOutUAV = RHICreateUnorderedAccessView(TextureRHI, iDstMip);
-
-		FNiagaraGenerateMipsCS::FParameters PassParameters;
-		PassParameters.SrcTexelSize	= FVector2f(1.0f / float(SrcMipSize.X), 1.0f / float(SrcMipSize.Y));
-		PassParameters.DstTexelSize = FVector2f(1.0f / float(DstMipSize.X), 1.0f / float(DstMipSize.Y));
-		PassParameters.KernelHWidth = KernelHWidth;
-		PassParameters.MipOutSize	= DstMipSize;
-		PassParameters.MipInSRV		= MipInSRV;
-		PassParameters.MipInSampler	= GenType == ENiagaraMipMapGenerationType::Unfiltered ? TStaticSamplerState<SF_Point>::GetRHI() : TStaticSamplerState<SF_Bilinear>::GetRHI();
-		PassParameters.MipOutUAV	= MipOutUAV;
-
-		const FIntVector NumThreadGroups = FComputeShaderUtils::GetGroupCount(DstMipSize, FComputeShaderUtils::kGolden2DGroupSize);
-
-		RHICmdList.Transition(FRHITransitionInfo(MipOutUAV, ERHIAccess::SRVMask, ERHIAccess::UAVCompute));
-
-		SetComputePipelineState(RHICmdList, ShaderRHI);
-		SetShaderParameters(RHICmdList, ComputeShader, ShaderRHI, PassParameters);
-		RHICmdList.DispatchComputeShader(NumThreadGroups.X, NumThreadGroups.Y, NumThreadGroups.Z);
-		UnsetShaderUAVs(RHICmdList, ComputeShader, ShaderRHI);
-
-		RHICmdList.Transition(FRHITransitionInfo(MipOutUAV, ERHIAccess::UAVCompute, ERHIAccess::SRVMask));
+		const FIntPoint SrcMipSize(FMath::Max<int32>(TextureDesc.GetSize().X >> iSrcMip, 1), FMath::Max<int32>(TextureDesc.GetSize().Y >> iSrcMip, 1));
+		const FIntPoint DstMipSize(FMath::Max<int32>(TextureDesc.GetSize().X >> iDstMip, 1), FMath::Max<int32>(TextureDesc.GetSize().Y >> iDstMip, 1));
+
+		FNiagaraGenerateMipsCS::FParameters* PassParameters = GraphBuilder.AllocParameters<FNiagaraGenerateMipsCS::FParameters>();
+		PassParameters->SrcTexelSize	= FVector2f(1.0f / float(SrcMipSize.X), 1.0f / float(SrcMipSize.Y));
+		PassParameters->DstTexelSize	 = FVector2f(1.0f / float(DstMipSize.X), 1.0f / float(DstMipSize.Y));
+		PassParameters->KernelHWidth	 = KernelHWidth;
+		PassParameters->MipOutSize		= DstMipSize;
+		PassParameters->MipInSRV		= GraphBuilder.CreateSRV(FRDGTextureSRVDesc::CreateForMipLevel(RDGTexture, iSrcMip));
+		PassParameters->MipInSampler	= GenType == ENiagaraMipMapGenerationType::Unfiltered ? TStaticSamplerState<SF_Point>::GetRHI() : TStaticSamplerState<SF_Bilinear>::GetRHI();
+		PassParameters->MipOutUAV		= GraphBuilder.CreateUAV(FRDGTextureUAVDesc(RDGTexture, iDstMip));
+
+		FComputeShaderUtils::AddPass(
+			GraphBuilder,
+			{},
+			ComputeShader,
+			PassParameters,
+			FComputeShaderUtils::GetGroupCount(DstMipSize, FComputeShaderUtils::kGolden2DGroupSize)
+		);
 	}
 }
diff --git a/Engine/Plugins/FX/Niagara/Source/NiagaraShader/Private/NiagaraShader.cpp b/Engine/Plugins/FX/Niagara/Source/NiagaraShader/Private/NiagaraShader.cpp
--- a/Engine/Plugins/FX/Niagara/Source/NiagaraShader/Private/NiagaraShader.cpp
+++ b/Engine/Plugins/FX/Niagara/Source/NiagaraShader/Private/NiagaraShader.cpp
@@ -1132,8 +1132,10 @@ FNiagaraShader::FNiagaraShader(const FNiagaraShaderType::CompiledShaderInitializ
 {
 	check(!DebugDescription.IsEmpty());
 
-	// Legacy bindings
+	// Cache off requirements for later queries
 	bNeedsViewUniformBuffer = Initializer.ParameterMap.ContainsParameterAllocation(TEXT("View"));
+
+	// Legacy bindings
 	ExternalConstantBufferParam[0].Bind(Initializer.ParameterMap, TEXT("FNiagaraExternalParameters"));
 	ExternalConstantBufferParam[1].Bind(Initializer.ParameterMap, TEXT("PREV_FNiagaraExternalParameters"));
 
@@ -1145,7 +1147,6 @@ FNiagaraShader::FNiagaraShader(const FNiagaraShaderType::CompiledShaderInitializ
 	
 	if (GNiagaraShaderForceBindEverything == 0)
 	{
-		bShouldBindEverything &= bNeedsViewUniformBuffer == false;
 		bShouldBindEverything &= ExternalConstantBufferParam[0].IsBound() == false;
 		bShouldBindEverything &= ExternalConstantBufferParam[1].IsBound() == false;
 
diff --git a/Engine/Plugins/FX/Niagara/Source/NiagaraShader/Private/NiagaraShaderParticleID.cpp b/Engine/Plugins/FX/Niagara/Source/NiagaraShader/Private/NiagaraShaderParticleID.cpp
--- a/Engine/Plugins/FX/Niagara/Source/NiagaraShader/Private/NiagaraShaderParticleID.cpp
+++ b/Engine/Plugins/FX/Niagara/Source/NiagaraShader/Private/NiagaraShaderParticleID.cpp
@@ -2,6 +2,7 @@
 
 #include "NiagaraShaderParticleID.h"
 #include "GlobalShader.h"
+#include "RenderGraphUtils.h"
 #include "ShaderParameterUtils.h"
 #include "PipelineStateCache.h"
 
@@ -12,20 +13,14 @@ FAutoConsoleVariableRef CVarGNiagaraWaveIntrinsics(
 	TEXT("")
 );
 
+//////////////////////////////////////////////////////////////////////////
+
 class FNiagaraInitFreeIDBufferCS : public FGlobalShader
 {
 	DECLARE_GLOBAL_SHADER(FNiagaraInitFreeIDBufferCS);
+	SHADER_USE_PARAMETER_STRUCT(FNiagaraInitFreeIDBufferCS, FGlobalShader);
 
-public:
-	FNiagaraInitFreeIDBufferCS() : FGlobalShader() {}
-
-	FNiagaraInitFreeIDBufferCS(const ShaderMetaType::CompiledShaderInitializerType& Initializer) : FGlobalShader(Initializer)
-	{
-		NewBufferParam.Bind(Initializer.ParameterMap, TEXT("NewBuffer"));
-		ExistingBufferParam.Bind(Initializer.ParameterMap, TEXT("ExistingBuffer"));
-		NumNewElementsParam.Bind(Initializer.ParameterMap, TEXT("NumNewElements"));
-		NumExistingElementsParam.Bind(Initializer.ParameterMap, TEXT("NumExistingElements"));
-	}
+	static constexpr uint32 ThreadGroupSize = 64;
 
 	static bool ShouldCompilePermutation(const FGlobalShaderPermutationParameters& Parameters)
 	{
@@ -35,66 +30,52 @@ public:
 	static void ModifyCompilationEnvironment(const FGlobalShaderPermutationParameters& Parameters, FShaderCompilerEnvironment& OutEnvironment)
 	{
 		FGlobalShader::ModifyCompilationEnvironment(Parameters, OutEnvironment);
-		OutEnvironment.SetDefine(TEXT("THREAD_COUNT"), THREAD_COUNT);
+		OutEnvironment.SetDefine(TEXT("THREADGROUP_SIZE"), ThreadGroupSize);
 	}
 
-	void Execute(FRHICommandList& RHICmdList, FRHIComputeShader* ComputeShader, uint32 NumElementsToAlloc, FRWBuffer& NewBuffer, uint32 NumExistingElements, FRHIShaderResourceView* ExistingBuffer)
-	{
-		// To simplify the shader code, the size of the ID table must be a multiple of the thread count.
-		check(NumElementsToAlloc % THREAD_COUNT == 0);
-
-		// Shrinking is not supported.
-		check(NumElementsToAlloc >= NumExistingElements);
-		uint32 NumNewElements = NumElementsToAlloc - NumExistingElements;
-
-		SetComputePipelineState(RHICmdList, ComputeShader);
+	BEGIN_SHADER_PARAMETER_STRUCT(FParameters, )
+		SHADER_PARAMETER_UAV(RWBuffer<int>,	RWNewBuffer)
+		SHADER_PARAMETER_SRV(Buffer<int>,	ExistingBuffer)
+		SHADER_PARAMETER(uint32,			NumNewElements)
+		SHADER_PARAMETER(uint32,			NumExistingElements)
+	END_SHADER_PARAMETER_STRUCT()
+};
 
-		NewBufferParam.SetBuffer(RHICmdList, ComputeShader, NewBuffer);
-		SetSRVParameter(RHICmdList, ComputeShader, ExistingBufferParam, ExistingBuffer);
-		SetShaderValue(RHICmdList, ComputeShader, NumNewElementsParam, NumNewElements);
-		SetShaderValue(RHICmdList, ComputeShader, NumExistingElementsParam, NumExistingElements);
+IMPLEMENT_GLOBAL_SHADER(FNiagaraInitFreeIDBufferCS, "/Plugin/FX/Niagara/Private/NiagaraInitFreeIDBuffer.usf", "InitIDBufferCS", SF_Compute);
 
-		DispatchComputeShader(RHICmdList, this, FMath::DivideAndRoundUp(NumElementsToAlloc, THREAD_COUNT), 1, 1);
+void NiagaraInitGPUFreeIDList(FRHICommandList& RHICmdList, ERHIFeatureLevel::Type FeatureLevel, FRHIUnorderedAccessView* NewBufferUAV, uint32 NewBufferNumElements, FRHIShaderResourceView* ExistingBufferSRV, uint32 ExistingBufferNumElements)
+{
+	// To simplify the shader code, the size of the ID table must be a multiple of the thread count.
+	check(NewBufferNumElements % FNiagaraInitFreeIDBufferCS::ThreadGroupSize == 0);
 
-		RHICmdList.SetUAVParameter(ComputeShader, NewBufferParam.GetUAVIndex(), nullptr);
-		RHICmdList.SetShaderResourceViewParameter(ComputeShader, ExistingBufferParam.GetBaseIndex(), nullptr);
-	}
+	// Shrinking is not supported.
+	check(NewBufferNumElements >= ExistingBufferNumElements);
+	const uint32 NumNewElements = NewBufferNumElements - ExistingBufferNumElements;
 
-private:
-	static constexpr uint32 THREAD_COUNT = 64;
+	TShaderMapRef<FNiagaraInitFreeIDBufferCS> ComputeShader(GetGlobalShaderMap(FeatureLevel));
+	const uint32 NumThreadGroups = FMath::DivideAndRoundUp(NewBufferNumElements, FNiagaraInitFreeIDBufferCS::ThreadGroupSize);
 
-	LAYOUT_FIELD(FRWShaderParameter, NewBufferParam);
-	LAYOUT_FIELD(FShaderResourceParameter, ExistingBufferParam);
-	LAYOUT_FIELD(FShaderParameter, NumNewElementsParam);
-	LAYOUT_FIELD(FShaderParameter, NumExistingElementsParam);
-};
+	FNiagaraInitFreeIDBufferCS::FParameters Parameters;
+	Parameters.RWNewBuffer			= NewBufferUAV;
+	Parameters.ExistingBuffer		= ExistingBufferSRV;
+	Parameters.NumNewElements		= NumNewElements;
+	Parameters.NumExistingElements	= ExistingBufferNumElements;
 
-IMPLEMENT_GLOBAL_SHADER(FNiagaraInitFreeIDBufferCS, "/Plugin/FX/Niagara/Private/NiagaraInitFreeIDBuffer.usf", "InitIDBufferCS", SF_Compute);
-
-void NiagaraInitGPUFreeIDList(FRHICommandList& RHICmdList, ERHIFeatureLevel::Type FeatureLevel, uint32 NumElementsToAlloc, FRWBuffer& NewBuffer, uint32 NumExistingElements, FRHIShaderResourceView* ExistingBuffer)
-{
-	TShaderMapRef<FNiagaraInitFreeIDBufferCS> InitIDBufferCS(GetGlobalShaderMap(FeatureLevel));
-	FRHIComputeShader* ComputeShader = InitIDBufferCS.GetComputeShader();
-	InitIDBufferCS->Execute(RHICmdList, ComputeShader, NumElementsToAlloc, NewBuffer, NumExistingElements, ExistingBuffer);
+	FComputeShaderUtils::Dispatch(RHICmdList, ComputeShader, Parameters, FIntVector(NumThreadGroups, 1, 1));
 }
 
+//////////////////////////////////////////////////////////////////////////
+
 class NiagaraComputeFreeIDsCS : public FGlobalShader
 {
 	DECLARE_GLOBAL_SHADER(NiagaraComputeFreeIDsCS);
+	SHADER_USE_PARAMETER_STRUCT(NiagaraComputeFreeIDsCS, FGlobalShader);
 
 	class FWaveIntrinsicsDim : SHADER_PERMUTATION_BOOL("USE_WAVE_INTRINSICS");
 	using FPermutationDomain = TShaderPermutationDomain<FWaveIntrinsicsDim>;
 
-public:
-	NiagaraComputeFreeIDsCS() : FGlobalShader() {}
-
-	NiagaraComputeFreeIDsCS(const ShaderMetaType::CompiledShaderInitializerType& Initializer) : FGlobalShader(Initializer)
-	{
-		IDToIndexTableParam.Bind(Initializer.ParameterMap, TEXT("IDToIndexTable"));
-		FreeIDListParam.Bind(Initializer.ParameterMap, TEXT("FreeIDList"));
-		FreeIDListSizesParam.Bind(Initializer.ParameterMap, TEXT("FreeIDListSizes"));
-		FreeIDListIndexParam.Bind(Initializer.ParameterMap, TEXT("FreeIDListIndex"));
-	}
+	static constexpr uint32 ThreadGroupSize_WaveEnabled = 64;
+	static constexpr uint32 ThreadGroupSize_WaveDisabled = 128;
 
 	static bool ShouldCompilePermutation(const FGlobalShaderPermutationParameters& Parameters)
 	{
@@ -116,123 +97,80 @@ public:
 		const bool bWithWaveIntrinsics = PermutationVector.Get<FWaveIntrinsicsDim>();
 
 		OutEnvironment.SetDefine(TEXT("USE_WAVE_INTRINSICS"), bWithWaveIntrinsics ? 1 : 0);
-		OutEnvironment.SetDefine(TEXT("THREAD_COUNT"), bWithWaveIntrinsics ? 64 : 128);
+		OutEnvironment.SetDefine(TEXT("THREADGROUP_SIZE"), bWithWaveIntrinsics ? ThreadGroupSize_WaveEnabled : ThreadGroupSize_WaveDisabled);
 	}
 
-	void Execute(FRHICommandList& RHICmdList, FRHIComputeShader* ComputeShader, uint32 ThreadCount, uint32 NumIDs, FRHIShaderResourceView* IDToIndexTable, FRWBuffer& FreeIDList, FRWBuffer& FreeIDListSizes, uint32 FreeIDListIndex)
-	{
-		// To simplify the shader code, the size of the ID table must be a multiple of the thread count.
-		check(NumIDs % ThreadCount == 0);
-
-		SetComputePipelineState(RHICmdList, ComputeShader);
-
-		SetSRVParameter(RHICmdList, ComputeShader, IDToIndexTableParam, IDToIndexTable);
-		FreeIDListParam.SetBuffer(RHICmdList, ComputeShader, FreeIDList);
-		FreeIDListSizesParam.SetBuffer(RHICmdList, ComputeShader, FreeIDListSizes);
-		SetShaderValue(RHICmdList, ComputeShader, FreeIDListIndexParam, FreeIDListIndex);
-
-		DispatchComputeShader(RHICmdList, this, NumIDs / ThreadCount, 1, 1);
-
-		RHICmdList.SetShaderResourceViewParameter(ComputeShader, IDToIndexTableParam.GetBaseIndex(), nullptr);
-		RHICmdList.SetUAVParameter(ComputeShader, FreeIDListParam.GetUAVIndex(), nullptr);
-		RHICmdList.SetUAVParameter(ComputeShader, FreeIDListSizesParam.GetUAVIndex(), nullptr);
-	}
-
-private:
-	LAYOUT_FIELD(FShaderResourceParameter, IDToIndexTableParam);
-	LAYOUT_FIELD(FRWShaderParameter, FreeIDListParam);
-	LAYOUT_FIELD(FRWShaderParameter, FreeIDListSizesParam);
-	LAYOUT_FIELD(FShaderParameter, FreeIDListIndexParam);
+	BEGIN_SHADER_PARAMETER_STRUCT(FParameters, )
+		SHADER_PARAMETER_SRV(Buffer<int>,	IDToIndexTable)
+		SHADER_PARAMETER_UAV(RWBuffer<int>,	RWFreeIDList)
+		SHADER_PARAMETER_UAV(RWBuffer<int>,	RWFreeIDListSizes)
+		SHADER_PARAMETER(uint32,			FreeIDListIndex)
+	END_SHADER_PARAMETER_STRUCT()
 };
 
 IMPLEMENT_GLOBAL_SHADER(NiagaraComputeFreeIDsCS, "/Plugin/FX/Niagara/Private/NiagaraComputeFreeIDs.usf", "ComputeFreeIDs", SF_Compute);
 
-void NiagaraComputeGPUFreeIDs(FRHICommandList& RHICmdList, ERHIFeatureLevel::Type FeatureLevel, uint32 NumIDs, FRHIShaderResourceView* IDToIndexTable, FRWBuffer& FreeIDList, FRWBuffer& FreeIDListSizes, uint32 FreeIDListIndex)
+void NiagaraComputeGPUFreeIDs(FRHICommandList& RHICmdList, ERHIFeatureLevel::Type FeatureLevel, FRHIShaderResourceView* IDToIndexTableSRV, uint32 NumIDs, FRHIUnorderedAccessView* FreeIDUAV, FRHIUnorderedAccessView* FreeIDListSizesUAV, uint32 FreeIDListIndex)
 {
-	auto ShaderMap = GetGlobalShaderMap(FeatureLevel);
 	const EShaderPlatform Platform = GShaderPlatformForFeatureLevel[FeatureLevel];
+	const bool bUseWaveOps = FDataDrivenShaderPlatformInfo::GetSupportsIntrinsicWaveOnce(Platform) && GNiagaraWaveIntrinsics != 0;
+	const uint32 ThreadGroupSize = bUseWaveOps ? NiagaraComputeFreeIDsCS::ThreadGroupSize_WaveEnabled : NiagaraComputeFreeIDsCS::ThreadGroupSize_WaveDisabled;
+
+	check(NumIDs % ThreadGroupSize == 0);
 
-	const bool bUseWaveIntrinsics = FDataDrivenShaderPlatformInfo::GetSupportsIntrinsicWaveOnce(Platform) && GNiagaraWaveIntrinsics != 0;
-	const uint32 ThreadCount = bUseWaveIntrinsics ? 64 : 128;
-	
 	NiagaraComputeFreeIDsCS::FPermutationDomain PermutationVector;
-	PermutationVector.Set<NiagaraComputeFreeIDsCS::FWaveIntrinsicsDim>(bUseWaveIntrinsics);
-	
-	auto ComputeShaderEntry = ShaderMap->GetShader<NiagaraComputeFreeIDsCS>(PermutationVector);
-	FRHIComputeShader* ComputeShader = ComputeShaderEntry.GetComputeShader();
-	ComputeShaderEntry->Execute(RHICmdList, ComputeShader, ThreadCount, NumIDs, IDToIndexTable, FreeIDList, FreeIDListSizes, FreeIDListIndex);
+	PermutationVector.Set<NiagaraComputeFreeIDsCS::FWaveIntrinsicsDim>(bUseWaveOps);
+
+	TShaderMapRef<NiagaraComputeFreeIDsCS> ComputeShader(GetGlobalShaderMap(FeatureLevel), PermutationVector);
+	const uint32 NumThreadGroups = FMath::DivideAndRoundUp(NumIDs, ThreadGroupSize);
+
+	NiagaraComputeFreeIDsCS::FParameters Parameters;
+	Parameters.IDToIndexTable		= IDToIndexTableSRV;
+	Parameters.RWFreeIDList			= FreeIDUAV;
+	Parameters.RWFreeIDListSizes	= FreeIDListSizesUAV;
+	Parameters.FreeIDListIndex		= FreeIDListIndex;
+
+	FComputeShaderUtils::Dispatch(RHICmdList, ComputeShader, Parameters, FIntVector(NumThreadGroups, 1, 1));
 }
 
+//////////////////////////////////////////////////////////////////////////
+
 class NiagaraFillIntBufferCS : public FGlobalShader
 {
 	DECLARE_GLOBAL_SHADER(NiagaraFillIntBufferCS);
+	SHADER_USE_PARAMETER_STRUCT(NiagaraFillIntBufferCS, FGlobalShader);
 
-public:
-	NiagaraFillIntBufferCS() : FGlobalShader() {}
-
-	NiagaraFillIntBufferCS(const ShaderMetaType::CompiledShaderInitializerType& Initializer) : FGlobalShader(Initializer)
-	{
-		TargetBufferParam.Bind(Initializer.ParameterMap, TEXT("TargetBuffer"));
-		FillValueParam.Bind(Initializer.ParameterMap, TEXT("FillValue"));
-		BufferSizeParam.Bind(Initializer.ParameterMap, TEXT("BufferSize"));
-	}
+	static constexpr uint32 ThreadGroupSize = 64;
 
 	static bool ShouldCompilePermutation(const FGlobalShaderPermutationParameters& Parameters)
 	{
 		return RHISupportsComputeShaders(Parameters.Platform);
 	}
 
-	void Execute(FRHICommandList& RHICmdList, FRHIComputeShader* ComputeShader, FRWBuffer& Buffer, int32 Value)
+	static void ModifyCompilationEnvironment(const FGlobalShaderPermutationParameters& Parameters, FShaderCompilerEnvironment& OutEnvironment)
 	{
-		const uint32 THREAD_COUNT = 64;
-		const uint32 NumInts = Buffer.NumBytes / sizeof(int32);
-		const uint32 ThreadGroups = FMath::DivideAndRoundUp(NumInts, THREAD_COUNT);
-
-		SetComputePipelineState(RHICmdList, ComputeShader);
-
-		SetUAVParameter(RHICmdList, ComputeShader, TargetBufferParam, Buffer.UAV);
-		SetShaderValue(RHICmdList, ComputeShader, FillValueParam, Value);
-		SetShaderValue(RHICmdList, ComputeShader, BufferSizeParam, NumInts);
-
-		DispatchComputeShader(RHICmdList, this, ThreadGroups, 1, 1);
-
-		SetUAVParameter(RHICmdList, ComputeShader, TargetBufferParam, nullptr);
+		FGlobalShader::ModifyCompilationEnvironment(Parameters, OutEnvironment);
+		OutEnvironment.SetDefine(TEXT("THREADGROUP_SIZE"), ThreadGroupSize);
 	}
 
-	void Execute(FRHICommandList& RHICmdList, FRHIComputeShader* ComputeShader, FRWBufferStructured& Buffer, int32 Value)
-	{
-		const uint32 THREAD_COUNT = 64;
-		const uint32 NumInts = Buffer.NumBytes / sizeof(int32);
-		const uint32 ThreadGroups = FMath::DivideAndRoundUp(NumInts, THREAD_COUNT);
-
-		SetComputePipelineState(RHICmdList, ComputeShader);
-
-		SetUAVParameter(RHICmdList, ComputeShader, TargetBufferParam, Buffer.UAV);
-		SetShaderValue(RHICmdList, ComputeShader, FillValueParam, Value);
-		SetShaderValue(RHICmdList, ComputeShader, BufferSizeParam, NumInts);
-
-		DispatchComputeShader(RHICmdList, this, ThreadGroups, 1, 1);
-
-		SetUAVParameter(RHICmdList, ComputeShader, TargetBufferParam, nullptr);
-	}
-private:
-	LAYOUT_FIELD(FShaderResourceParameter, TargetBufferParam);
-	LAYOUT_FIELD(FShaderParameter, FillValueParam);
-	LAYOUT_FIELD(FShaderParameter, BufferSizeParam);
+	BEGIN_SHADER_PARAMETER_STRUCT(FParameters, )
+		SHADER_PARAMETER_UAV(RWBuffer<int>,	BufferUAV)
+		SHADER_PARAMETER(uint32,			NumElements)
+		SHADER_PARAMETER(int,				FillValue)
+	END_SHADER_PARAMETER_STRUCT()
 };
 
 IMPLEMENT_GLOBAL_SHADER(NiagaraFillIntBufferCS, "/Plugin/FX/Niagara/Private/NiagaraFillIntBuffer.usf", "FillIntBuffer", SF_Compute);
 
-void NiagaraFillGPUIntBuffer(FRHICommandList& RHICmdList, ERHIFeatureLevel::Type FeatureLevel, FRWBuffer& Buffer, int32 Value)
+void NiagaraFillGPUIntBuffer(FRHICommandList& RHICmdList, ERHIFeatureLevel::Type FeatureLevel, FRHIUnorderedAccessView* BufferUAV, uint32 NumElements, int32 FillValue)
 {
-	TShaderMapRef<NiagaraFillIntBufferCS> FillCS(GetGlobalShaderMap(FeatureLevel));
-	FRHIComputeShader* ComputeShader = FillCS.GetComputeShader();
-	FillCS->Execute(RHICmdList, ComputeShader, Buffer, Value);
-}
+	TShaderMapRef<NiagaraFillIntBufferCS> ComputeShader(GetGlobalShaderMap(FeatureLevel));
+	const uint32 NumThreadGroups = FMath::DivideAndRoundUp(NumElements, NiagaraFillIntBufferCS::ThreadGroupSize);
 
-void NiagaraFillGPUIntBuffer(FRHICommandList& RHICmdList, ERHIFeatureLevel::Type FeatureLevel, FRWBufferStructured& Buffer, int32 Value)
-{
-	TShaderMapRef<NiagaraFillIntBufferCS> FillCS(GetGlobalShaderMap(FeatureLevel));
-	FRHIComputeShader* ComputeShader = FillCS.GetComputeShader();
-	FillCS->Execute(RHICmdList, ComputeShader, Buffer, Value);
+	NiagaraFillIntBufferCS::FParameters Parameters;
+	Parameters.BufferUAV	= BufferUAV;
+	Parameters.NumElements	= NumElements;
+	Parameters.FillValue	= FillValue;
+
+	FComputeShaderUtils::Dispatch(RHICmdList, ComputeShader, Parameters, FIntVector(NumThreadGroups, 1, 1));
 }
diff --git a/Engine/Plugins/FX/Niagara/Source/NiagaraShader/Private/NiagaraShared.cpp b/Engine/Plugins/FX/Niagara/Source/NiagaraShader/Private/NiagaraShared.cpp
--- a/Engine/Plugins/FX/Niagara/Source/NiagaraShader/Private/NiagaraShared.cpp
+++ b/Engine/Plugins/FX/Niagara/Source/NiagaraShader/Private/NiagaraShared.cpp
@@ -694,10 +694,7 @@ void FNiagaraShaderScript::BuildScriptParametersMetadata(const FNiagaraShaderScr
 			const uint32 NextMemberOffset = ShaderMetadataBuilder.GetNextMemberOffset();
 			FNiagaraShaderParametersBuilder ShaderParametersBuilder(DataInterfaceParamInfo, ScriptParametersMetadata->LooseMetadataNames, ShaderMetadataBuilder);
 			CDODataInterface->BuildShaderParameters(ShaderParametersBuilder);
-			if (NextMemberOffset != ShaderMetadataBuilder.GetNextMemberOffset())
-			{
-				DataInterfaceParamInfo.ShaderParametersOffset = NextMemberOffset;
-			}
+			DataInterfaceParamInfo.ShaderParametersOffset = NextMemberOffset;
 		}
 	}
 
diff --git a/Engine/Plugins/FX/Niagara/Source/NiagaraShader/Public/NiagaraClearCounts.h b/Engine/Plugins/FX/Niagara/Source/NiagaraShader/Public/NiagaraClearCounts.h
--- a/Engine/Plugins/FX/Niagara/Source/NiagaraShader/Public/NiagaraClearCounts.h
+++ b/Engine/Plugins/FX/Niagara/Source/NiagaraShader/Public/NiagaraClearCounts.h
@@ -2,10 +2,14 @@
 
 #pragma once
 
-#include "RHI.h"
 #include "CommonRenderResources.h"
+#include "RenderGraphResources.h"
+#include "RHI.h"
 
 namespace NiagaraClearCounts
 {
+	NIAGARASHADER_API void ClearCountsInt(FRDGBuilder& GraphBuilder, FRDGBufferUAVRef UAV, TConstArrayView<TPair<uint32, int32>> IndexAndValueArray);
+
+	//-TODO:RDG: Deprecated
 	NIAGARASHADER_API void ClearCountsInt(FRHICommandList& RHICmdList, FRHIUnorderedAccessView* UAV, TConstArrayView<TPair<uint32, int32>> IndexAndValueArray);
 }
diff --git a/Engine/Plugins/FX/Niagara/Source/NiagaraShader/Public/NiagaraDebugShaders.h b/Engine/Plugins/FX/Niagara/Source/NiagaraShader/Public/NiagaraDebugShaders.h
--- a/Engine/Plugins/FX/Niagara/Source/NiagaraShader/Public/NiagaraDebugShaders.h
+++ b/Engine/Plugins/FX/Niagara/Source/NiagaraShader/Public/NiagaraDebugShaders.h
@@ -24,7 +24,7 @@ namespace NiagaraDebugShaders
 	NIAGARASHADER_API void VisualizeTexture(
 		class FRDGBuilder& GraphBuilder, const FViewInfo& View, const FScreenPassRenderTarget& Output,
 		const FIntPoint& Location, const int32& DisplayHeight,
-		const FIntVector4& AttributesToVisualize, FRHITexture* Texture, const FIntVector4& NumTextureAttributes, uint32 TickCounter,
+		const FIntVector4& AttributesToVisualize, FRDGTextureRef Texture, const FIntVector4& NumTextureAttributes, uint32 TickCounter,
 		const FVector2D& PreviewDisplayRange = FVector2D::ZeroVector
 	);
 }
diff --git a/Engine/Plugins/FX/Niagara/Source/NiagaraShader/Public/NiagaraGenerateMips.h b/Engine/Plugins/FX/Niagara/Source/NiagaraShader/Public/NiagaraGenerateMips.h
--- a/Engine/Plugins/FX/Niagara/Source/NiagaraShader/Public/NiagaraGenerateMips.h
+++ b/Engine/Plugins/FX/Niagara/Source/NiagaraShader/Public/NiagaraGenerateMips.h
@@ -2,7 +2,7 @@
 
 #pragma once
 
-#include "RHI.h"
+#include "RenderGraphDefinitions.h"
 #include "NiagaraGenerateMips.generated.h"
 
 UENUM()
@@ -24,5 +24,5 @@ enum class ENiagaraMipMapGenerationType : uint8
 
 namespace NiagaraGenerateMips
 {
-	NIAGARASHADER_API void GenerateMips(FRHICommandList& RHICmdList, FRHITexture2D* TextureRHI, ENiagaraMipMapGenerationType GenType);
+	NIAGARASHADER_API void GenerateMips(FRDGBuilder& GraphBuilder, FRDGTextureRef RDGTexture, ENiagaraMipMapGenerationType GenType);
 }
diff --git a/Engine/Plugins/FX/Niagara/Source/NiagaraShader/Public/NiagaraShader.h b/Engine/Plugins/FX/Niagara/Source/NiagaraShader/Public/NiagaraShader.h
--- a/Engine/Plugins/FX/Niagara/Source/NiagaraShader/Public/NiagaraShader.h
+++ b/Engine/Plugins/FX/Niagara/Source/NiagaraShader/Public/NiagaraShader.h
@@ -6,18 +6,19 @@
 
 #pragma once
 
+#include "GlobalShader.h"
 #include "ShaderParameters.h"
 #include "SceneView.h"
 #include "Shader.h"
-#include "GlobalShader.h"
 #include "Math/Vector.h"
 #include "NiagaraCommon.h"
 #include "NiagaraScriptBase.h"
 #include "NiagaraShared.h"
 #include "NiagaraShaderType.h"
-#if UE_ENABLE_INCLUDE_ORDER_DEPRECATED_IN_5_1
+//#if UE_ENABLE_INCLUDE_ORDER_DEPRECATED_IN_5_1
+#include "RenderGraph.h"
 #include "SceneRenderTargetParameters.h"
-#endif
+//#endif
 
 class UClass;
 
@@ -183,6 +184,9 @@ public:
 		SHADER_PARAMETER_STRUCT_INCLUDE(FSystemParameters,			SystemParameters)
 		SHADER_PARAMETER_STRUCT_INCLUDE(FOwnerParameters,			OwnerParameters)
 		SHADER_PARAMETER_STRUCT_INCLUDE(FEmitterParameters,			EmitterParameters)
+
+		SHADER_PARAMETER_STRUCT_REF(FViewUniformShaderParameters,	View)
+		SHADER_PARAMETER_STRUCT_INCLUDE(FSceneTextureShaderParameters,	SceneTextures)
 	END_SHADER_PARAMETER_STRUCT()
 
 	using FPermutationParameters = FNiagaraShaderPermutationParameters;
diff --git a/Engine/Plugins/FX/Niagara/Source/NiagaraShader/Public/NiagaraShaderParticleID.h b/Engine/Plugins/FX/Niagara/Source/NiagaraShader/Public/NiagaraShaderParticleID.h
--- a/Engine/Plugins/FX/Niagara/Source/NiagaraShader/Public/NiagaraShaderParticleID.h
+++ b/Engine/Plugins/FX/Niagara/Source/NiagaraShader/Public/NiagaraShaderParticleID.h
@@ -4,7 +4,6 @@
 
 #include "RHI.h"
 
-extern NIAGARASHADER_API void NiagaraInitGPUFreeIDList(FRHICommandList& RHICmdList, ERHIFeatureLevel::Type FeatureLevel, uint32 NumElementsToAlloc, FRWBuffer& NewBuffer, uint32 NumExistingElements, FRHIShaderResourceView* ExistingBuffer);
-extern NIAGARASHADER_API void NiagaraComputeGPUFreeIDs(FRHICommandList& RHICmdList, ERHIFeatureLevel::Type FeatureLevel, uint32 NumIDs, FRHIShaderResourceView* IDToIndexTable, FRWBuffer& FreeIDList, FRWBuffer& FreeIDListSizes, uint32 FreeIDListIndex);
-extern NIAGARASHADER_API void NiagaraFillGPUIntBuffer(FRHICommandList& RHICmdList, ERHIFeatureLevel::Type FeatureLevel, FRWBuffer& Buffer, int32 Value);
-extern NIAGARASHADER_API void NiagaraFillGPUIntBuffer(FRHICommandList& RHICmdList, ERHIFeatureLevel::Type FeatureLevel, FRWBufferStructured& Buffer, int32 Value);
+extern NIAGARASHADER_API void NiagaraInitGPUFreeIDList(FRHICommandList& RHICmdList, ERHIFeatureLevel::Type FeatureLevel, FRHIUnorderedAccessView* NewBufferUAV, uint32 NewBufferNumElements, FRHIShaderResourceView* ExistingBufferSRV, uint32 ExistingBufferNumElements);
+extern NIAGARASHADER_API void NiagaraComputeGPUFreeIDs(FRHICommandList& RHICmdList, ERHIFeatureLevel::Type FeatureLevel, FRHIShaderResourceView* IDToIndexTableSRV, uint32 NumIDs, FRHIUnorderedAccessView* FreeIDUAV, FRHIUnorderedAccessView* FreeIDListSizesUAV, uint32 FreeIDListIndex);
+extern NIAGARASHADER_API void NiagaraFillGPUIntBuffer(FRHICommandList& RHICmdList, ERHIFeatureLevel::Type FeatureLevel, FRHIUnorderedAccessView* BufferUAV, uint32 NumElements, int32 FillValue);
diff --git a/Engine/Plugins/FX/Niagara/Source/NiagaraShader/Public/NiagaraShared.h b/Engine/Plugins/FX/Niagara/Source/NiagaraShader/Public/NiagaraShared.h
--- a/Engine/Plugins/FX/Niagara/Source/NiagaraShader/Public/NiagaraShared.h
+++ b/Engine/Plugins/FX/Niagara/Source/NiagaraShader/Public/NiagaraShared.h
@@ -210,6 +210,7 @@ struct NIAGARASHADER_API FNiagaraDataInterfaceGPUParamInfo
 	UPROPERTY()
 	FString DIClassName;
 
+	/** The offset for any shader parameters, when this member is INDEX_NONE we are in legacy mode for the data interface. */
 	UPROPERTY()
 	uint32 ShaderParametersOffset = INDEX_NONE;
 
diff --git a/Engine/Plugins/Runtime/HairStrands/Source/HairStrandsCore/Private/Niagara/NiagaraDataInterfacePhysicsAsset.cpp b/Engine/Plugins/Runtime/HairStrands/Source/HairStrandsCore/Private/Niagara/NiagaraDataInterfacePhysicsAsset.cpp
--- a/Engine/Plugins/Runtime/HairStrands/Source/HairStrandsCore/Private/Niagara/NiagaraDataInterfacePhysicsAsset.cpp
+++ b/Engine/Plugins/Runtime/HairStrands/Source/HairStrandsCore/Private/Niagara/NiagaraDataInterfacePhysicsAsset.cpp
@@ -675,9 +675,6 @@ void FNDIPhysicsAssetProxy::PreStage(FRHICommandList& RHICmdList, const FNiagara
 	}
 }
 
-void FNDIPhysicsAssetProxy::ResetData(FRHICommandList& RHICmdList, const FNiagaraDataInterfaceArgs& Context)
-{}
-
 //------------------------------------------------------------------------------------------------------------
 
 UNiagaraDataInterfacePhysicsAsset::UNiagaraDataInterfacePhysicsAsset(FObjectInitializer const& ObjectInitializer)
diff --git a/Engine/Plugins/Runtime/HairStrands/Source/HairStrandsCore/Public/Niagara/NiagaraDataInterfacePhysicsAsset.h b/Engine/Plugins/Runtime/HairStrands/Source/HairStrandsCore/Public/Niagara/NiagaraDataInterfacePhysicsAsset.h
--- a/Engine/Plugins/Runtime/HairStrands/Source/HairStrandsCore/Public/Niagara/NiagaraDataInterfacePhysicsAsset.h
+++ b/Engine/Plugins/Runtime/HairStrands/Source/HairStrandsCore/Public/Niagara/NiagaraDataInterfacePhysicsAsset.h
@@ -231,9 +231,6 @@ struct FNDIPhysicsAssetProxy : public FNiagaraDataInterfaceProxy
 	/** Launch all pre stage functions */
 	virtual void PreStage(FRHICommandList& RHICmdList, const FNiagaraDataInterfaceStageArgs& Context) override;
 
-	/** Reset the buffers  */
-	virtual void ResetData(FRHICommandList& RHICmdList, const FNiagaraDataInterfaceArgs& Context) override;
-
 	/** List of proxy data for each system instances*/
 	TMap<FNiagaraSystemInstanceID, FNDIPhysicsAssetData> SystemInstancesToProxyData;
 };
diff --git a/Engine/Source/Runtime/Core/Private/UObject/UE5MainStreamObjectVersion.cpp b/Engine/Source/Runtime/Core/Private/UObject/UE5MainStreamObjectVersion.cpp
--- a/Engine/Source/Runtime/Core/Private/UObject/UE5MainStreamObjectVersion.cpp
+++ b/Engine/Source/Runtime/Core/Private/UObject/UE5MainStreamObjectVersion.cpp
@@ -12,9 +12,9 @@ TMap<FGuid, FGuid> FUE5MainStreamObjectVersion::GetSystemGuids()
 	SystemGuids.Add(DevGuids.GLOBALSHADERMAP_DERIVEDDATA_VER, FGuid("62F5564D1FED4A2D8864DF300EC5AA2F"));
 	SystemGuids.Add(DevGuids.LANDSCAPE_MOBILE_COOK_VERSION, FGuid("71000000000000000000000000000035"));
 	SystemGuids.Add(DevGuids.MATERIALSHADERMAP_DERIVEDDATA_VER, FGuid("65C01D817C9A4EEFAE9E988D41A1F3DD"));
-	SystemGuids.Add(DevGuids.NANITE_DERIVEDDATA_VER, FGuid("DAC47956CDBF995D96F4520EDB56D875"));
+	SystemGuids.Add(DevGuids.NANITE_DERIVEDDATA_VER, FGuid("7A3EA93546624597BEBE2ADFDE3A7628"));
 	SystemGuids.Add(DevGuids.NIAGARASHADERMAP_DERIVEDDATA_VER, FGuid("8A37C45D24F2423CBE5F8F371DE33575"));
-	SystemGuids.Add(DevGuids.Niagara_LatestScriptCompileVersion, FGuid("28A43FC8EA00004E8B96B946BCA110F6"));
+	SystemGuids.Add(DevGuids.Niagara_LatestScriptCompileVersion, FGuid("B531E01D70C42444B4C76046D11CDA22"));
 	SystemGuids.Add(DevGuids.SkeletalMeshDerivedDataVersion, FGuid("FB7B2FB546DE49B2A2F88B7F69763CEE"));
 	SystemGuids.Add(DevGuids.STATICMESH_DERIVEDDATA_VER, FGuid("98F7D79A5811013825E75BBCC41ED3E9"));
 	SystemGuids.Add(DevGuids.POSESEARCHDB_DERIVEDDATA_VER, FGuid("4E595C2AC5E947D6BA9ABC874353E5BC"));
diff --git a/Engine/Source/Runtime/Renderer/Private/SystemTextures.cpp b/Engine/Source/Runtime/Renderer/Private/SystemTextures.cpp
--- a/Engine/Source/Runtime/Renderer/Private/SystemTextures.cpp
+++ b/Engine/Source/Runtime/Renderer/Private/SystemTextures.cpp
@@ -29,6 +29,7 @@ const FRDGSystemTextures& FRDGSystemTextures::Create(FRDGBuilder& GraphBuilder)
 	SystemTextures.White = Register(GSystemTextures.WhiteDummy);
 	SystemTextures.Black = Register(GSystemTextures.BlackDummy);
 	SystemTextures.BlackAlphaOne = Register(GSystemTextures.BlackAlphaOneDummy);
+	SystemTextures.BlackArray = Register(GSystemTextures.BlackArrayDummy);
 	SystemTextures.MaxFP16Depth = Register(GSystemTextures.MaxFP16Depth);
 	SystemTextures.DepthDummy = Register(GSystemTextures.DepthDummy);
 	SystemTextures.BlackDepthCube = Register(GSystemTextures.BlackDepthCube);
@@ -39,6 +40,8 @@ const FRDGSystemTextures& FRDGSystemTextures::Create(FRDGBuilder& GraphBuilder)
 	SystemTextures.VolumetricBlack = Register(GSystemTextures.VolumetricBlackDummy);
 	SystemTextures.VolumetricBlackAlphaOne = Register(GSystemTextures.VolumetricBlackAlphaOneDummy);
 	SystemTextures.VolumetricBlackUint = Register(GSystemTextures.VolumetricBlackUintDummy);
+	SystemTextures.CubeBlack = Register(GSystemTextures.CubeBlackDummy);
+	SystemTextures.CubeArrayBlack = Register(GSystemTextures.CubeArrayBlackDummy);
 	SystemTextures.StencilDummySRV = GraphBuilder.CreateSRV(FRDGTextureSRVDesc::Create(SystemTextures.StencilDummy));
 	return SystemTextures;
 }
@@ -269,6 +272,12 @@ void FSystemTextures::InitializeCommonTextures(FRHICommandListImmediate& RHICmdL
 		VolumetricBlackAlphaOneDummy = CreateRenderTarget(GBlackAlpha1VolumeTexture->TextureRHI, TEXT("VolumetricBlackAlphaOneDummy"));
 		VolumetricBlackUintDummy = CreateRenderTarget(GBlackUintVolumeTexture->TextureRHI, TEXT("VolumetricBlackUintDummy"));
 	}
+
+	// Create Cube BlackDummy textures
+	{
+		CubeBlackDummy = CreateRenderTarget(GBlackTextureCube->TextureRHI, TEXT("CubeBlackDummy"));
+		CubeArrayBlackDummy = CreateRenderTarget(GBlackCubeArrayTexture->TextureRHI, TEXT("CubeArrayBlackDummy"));
+	}
 }
 
 void FSystemTextures::InitializeFeatureLevelDependentTextures(FRHICommandListImmediate& RHICmdList, const ERHIFeatureLevel::Type InFeatureLevel)
@@ -873,6 +882,7 @@ void FSystemTextures::ReleaseDynamicRHI()
 	BlackDummy.SafeRelease();
 	BlackArrayDummy.SafeRelease();
 	BlackAlphaOneDummy.SafeRelease();
+	BlackArrayDummy.SafeRelease();
 	PerlinNoiseGradient.SafeRelease();
 	PerlinNoise3D.SafeRelease();
 	SobolSampling.SafeRelease();
@@ -891,6 +901,8 @@ void FSystemTextures::ReleaseDynamicRHI()
 	VolumetricBlackDummy.SafeRelease();
 	VolumetricBlackAlphaOneDummy.SafeRelease();
 	VolumetricBlackUintDummy.SafeRelease();
+	CubeBlackDummy.SafeRelease();
+	CubeArrayBlackDummy.SafeRelease();
 	ZeroUIntDummy.SafeRelease();
 	ZeroUIntArrayDummy.SafeRelease();
 	MidGreyDummy.SafeRelease();
@@ -971,6 +983,16 @@ FRDGTextureRef FSystemTextures::GetVolumetricBlackUintDummy(FRDGBuilder& GraphBu
 	return GraphBuilder.RegisterExternalTexture(VolumetricBlackUintDummy, TEXT("VolumetricBlackUintDummy"), ERDGTextureFlags::SkipTracking);
 }
 
+FRDGTextureRef FSystemTextures::GetCubeBlackDummy(FRDGBuilder& GraphBuilder) const
+{
+	return GraphBuilder.RegisterExternalTexture(CubeBlackDummy, TEXT("CubeBlackDummy"), ERDGTextureFlags::SkipTracking);
+}
+
+FRDGTextureRef FSystemTextures::GetCubeArrayBlackDummy(FRDGBuilder& GraphBuilder) const
+{
+	return GraphBuilder.RegisterExternalTexture(CubeArrayBlackDummy, TEXT("CubeArrayBlackDummy"), ERDGTextureFlags::SkipTracking);
+}
+
 FRDGTextureRef FSystemTextures::GetZeroUIntDummy(FRDGBuilder& GraphBuilder) const
 {
 	return GraphBuilder.RegisterExternalTexture(ZeroUIntDummy, TEXT("ZeroUIntDummy"), ERDGTextureFlags::SkipTracking);
