- Fix for shipping build

#rnx

diff --git a/Engine/Plugins/FX/Niagara/Source/Niagara/Classes/NiagaraGPUProfilerInterface.h b/Engine/Plugins/FX/Niagara/Source/Niagara/Classes/NiagaraGPUProfilerInterface.h
--- a/Engine/Plugins/FX/Niagara/Source/Niagara/Classes/NiagaraGPUProfilerInterface.h
+++ b/Engine/Plugins/FX/Niagara/Source/Niagara/Classes/NiagaraGPUProfilerInterface.h
@@ -26,8 +26,8 @@ private:
 	FVersionedNiagaraEmitterWeakPtr			OwnerEmitter;
 	FName									StageName;
 #else
-	explicit FNiagaraGpuProfileEvent(const struct FNiagaraGpuDispatchInstance& DispatchInstance) {}
-	explicit FNiagaraGpuProfileEvent(const struct FNiagaraGpuDispatchInstance& DispatchInstance, FName CustomStageName) {}
+	explicit FNiagaraGpuProfileEvent(const FNiagaraComputeInstanceData& InstanceData, const FNiagaraSimStageData& SimStageData, const bool bFirstInstanceData) {}
+	explicit FNiagaraGpuProfileEvent(const FNiagaraComputeInstanceData& InstanceData, FName CustomStageName) {}
 	explicit FNiagaraGpuProfileEvent(const FNiagaraGpuProfileEvent& Other) = default;
 #endif
 };
