Added debug logging to find emitter cook problems

diff --git a/Engine/Plugins/FX/Niagara/Source/NiagaraEditor/Private/NiagaraScriptMergeManager.cpp b/Engine/Plugins/FX/Niagara/Source/NiagaraEditor/Private/NiagaraScriptMergeManager.cpp
--- a/Engine/Plugins/FX/Niagara/Source/NiagaraEditor/Private/NiagaraScriptMergeManager.cpp
+++ b/Engine/Plugins/FX/Niagara/Source/NiagaraEditor/Private/NiagaraScriptMergeManager.cpp
@@ -1372,6 +1372,9 @@ void FNiagaraScriptMergeManager::UpdateModuleVersions(const FVersionedNiagaraEmi
 	UNiagaraSystem* System = nullptr;
 	if (ChangedVersionOtherModules.Num() > 0)
 	{
+		// TODO (mga) REMOVE - debug log to find cook problems
+		UE_LOG(LogNiagaraEditor, Log, TEXT("Upgrading module version for emitter %s"), *GetPathNameSafe(Instance.Emitter));
+
 		TSharedPtr<FNiagaraSystemViewModel> SystemViewModel;
 
 		if (!FSlateApplication::IsInitialized())
