Fixed a crash when switching between emitter versions

diff --git a/Engine/Plugins/FX/Niagara/Source/Niagara/Private/NiagaraSystem.cpp b/Engine/Plugins/FX/Niagara/Source/Niagara/Private/NiagaraSystem.cpp
--- a/Engine/Plugins/FX/Niagara/Source/Niagara/Private/NiagaraSystem.cpp
+++ b/Engine/Plugins/FX/Niagara/Source/Niagara/Private/NiagaraSystem.cpp
@@ -309,7 +309,6 @@ bool UNiagaraSystem::ChangeEmitterVersion(const FVersionedNiagaraEmitter& Versio
 			KillAllActiveCompilations();
 			Modify();
 			EmitterHandles[i].GetInstance().Version = NewVersion;
-			RequestCompile(false);
 			return true;
 		}
 	}
diff --git a/Engine/Plugins/FX/Niagara/Source/NiagaraEditor/Private/ViewModels/NiagaraSystemViewModel.cpp b/Engine/Plugins/FX/Niagara/Source/NiagaraEditor/Private/ViewModels/NiagaraSystemViewModel.cpp
--- a/Engine/Plugins/FX/Niagara/Source/NiagaraEditor/Private/ViewModels/NiagaraSystemViewModel.cpp
+++ b/Engine/Plugins/FX/Niagara/Source/NiagaraEditor/Private/ViewModels/NiagaraSystemViewModel.cpp
@@ -2331,6 +2331,7 @@ bool FNiagaraSystemViewModel::ChangeEmitterVersion(const FVersionedNiagaraEmitte
 	if (System->ChangeEmitterVersion(Emitter, NewVersion))
 	{
 		ResetStack();
+		System->RequestCompile(false);
 		return true;
 	}
 	return false;
