- Convert Audio DIs to use shader parameter structs

diff --git a/Engine/Plugins/FX/Niagara/Source/Niagara/Classes/NiagaraDataInterfaceAudioOscilloscope.h b/Engine/Plugins/FX/Niagara/Source/Niagara/Classes/NiagaraDataInterfaceAudioOscilloscope.h
--- a/Engine/Plugins/FX/Niagara/Source/Niagara/Classes/NiagaraDataInterfaceAudioOscilloscope.h
+++ b/Engine/Plugins/FX/Niagara/Source/Niagara/Classes/NiagaraDataInterfaceAudioOscilloscope.h
@@ -91,10 +91,13 @@ UCLASS(EditInlineNew, Category = "Audio", meta = (DisplayName = "Audio Oscillosc
 class NIAGARA_API UNiagaraDataInterfaceAudioOscilloscope : public UNiagaraDataInterface
 {
 	GENERATED_UCLASS_BODY()
-public:
 
-	DECLARE_NIAGARA_DI_PARAMETER();
-	
+	BEGIN_SHADER_PARAMETER_STRUCT(FShaderParameters, )
+		SHADER_PARAMETER(int,				NumChannels)
+		SHADER_PARAMETER_SRV(Buffer<float>,	AudioBuffer)
+	END_SHADER_PARAMETER_STRUCT()
+
+public:
 	UPROPERTY(EditAnywhere, Category = "Oscilloscope")
 	TObjectPtr<USoundSubmix> Submix;
 
@@ -125,9 +128,13 @@ public:
 	virtual bool RequiresDistanceFieldData() const override { return false; }
 
 #if WITH_EDITORONLY_DATA
+	virtual bool AppendCompileHash(FNiagaraCompileHashVisitor* InVisitor) const override;
 	virtual bool GetFunctionHLSL(const FNiagaraDataInterfaceGPUParamInfo& ParamInfo, const FNiagaraDataInterfaceGeneratedFunction& FunctionInfo, int FunctionInstanceIndex, FString& OutHLSL) override;
 	virtual void GetParameterDefinitionHLSL(const FNiagaraDataInterfaceGPUParamInfo& ParamInfo, FString& OutHLSL) override;
 #endif
+	virtual bool UseLegacyShaderBindings() const override { return false; }
+	virtual void BuildShaderParameters(FNiagaraShaderParametersBuilder& ShaderParametersBuilder) const override;
+	virtual void SetShaderParameters(const FNiagaraDataInterfaceSetShaderParametersContext& Context) const override;
 
 	virtual bool Equals(const UNiagaraDataInterface* Other) const override;
 
diff --git a/Engine/Plugins/FX/Niagara/Source/Niagara/Classes/NiagaraDataInterfaceAudioSpectrum.h b/Engine/Plugins/FX/Niagara/Source/Niagara/Classes/NiagaraDataInterfaceAudioSpectrum.h
--- a/Engine/Plugins/FX/Niagara/Source/Niagara/Classes/NiagaraDataInterfaceAudioSpectrum.h
+++ b/Engine/Plugins/FX/Niagara/Source/Niagara/Classes/NiagaraDataInterfaceAudioSpectrum.h
@@ -129,10 +129,14 @@ UCLASS(EditInlineNew, Category = "Audio", meta = (DisplayName = "Audio Spectrum"
 class NIAGARA_API UNiagaraDataInterfaceAudioSpectrum : public UNiagaraDataInterfaceAudioSubmix
 {
 	GENERATED_UCLASS_BODY()
-public:
 
-	DECLARE_NIAGARA_DI_PARAMETER();
+	BEGIN_SHADER_PARAMETER_STRUCT(FShaderParameters, )
+		SHADER_PARAMETER(int,				NumChannels)
+		SHADER_PARAMETER(int,				Resolution)
+		SHADER_PARAMETER_SRV(Buffer<float>, SpectrumBuffer)
+	END_SHADER_PARAMETER_STRUCT()
 
+public:
 	// VM function names
 	static const FName GetSpectrumFunctionName;
 	static const FName GetNumChannelsFunctionName;
@@ -171,10 +175,13 @@ public:
 	}
 
 #if WITH_EDITORONLY_DATA
+	virtual bool AppendCompileHash(FNiagaraCompileHashVisitor* InVisitor) const override;
 	virtual bool GetFunctionHLSL(const FNiagaraDataInterfaceGPUParamInfo& ParamInfo, const FNiagaraDataInterfaceGeneratedFunction& FunctionInfo, int FunctionInstanceIndex, FString& OutHLSL) override;
-
 	virtual void GetParameterDefinitionHLSL(const FNiagaraDataInterfaceGPUParamInfo& ParamInfo, FString& OutHLSL) override;
 #endif
+	virtual bool UseLegacyShaderBindings() const override { return false; }
+	virtual void BuildShaderParameters(FNiagaraShaderParametersBuilder& ShaderParametersBuilder) const override;
+	virtual void SetShaderParameters(const FNiagaraDataInterfaceSetShaderParametersContext& Context) const override;
 
 	virtual bool Equals(const UNiagaraDataInterface* Other) const override;
 
diff --git a/Engine/Plugins/FX/Niagara/Source/Niagara/Private/NiagaraDataInterfaceAudioOscilloscope.cpp b/Engine/Plugins/FX/Niagara/Source/Niagara/Private/NiagaraDataInterfaceAudioOscilloscope.cpp
--- a/Engine/Plugins/FX/Niagara/Source/Niagara/Private/NiagaraDataInterfaceAudioOscilloscope.cpp
+++ b/Engine/Plugins/FX/Niagara/Source/Niagara/Private/NiagaraDataInterfaceAudioOscilloscope.cpp
@@ -9,6 +9,7 @@
 #include "ClearQuad.h"
 #include "TextureResource.h"
 #include "Engine/Texture2D.h"
+#include "NiagaraShaderParametersBuilder.h"
 #include "NiagaraSystemInstance.h"
 #include "Engine/TextureRenderTarget2D.h"
 #include "Engine/Engine.h"
@@ -22,8 +23,8 @@ static const FName SampleAudioBufferFunctionName("SampleAudioBuffer");
 static const FName GetAudioBufferNumChannelsFunctionName("GetAudioBufferNumChannels");
 
 // Global variable prefixes, used in HLSL parameter declarations.
-static const FString AudioBufferName(TEXT("AudioBuffer_"));
-static const FString NumChannelsName(TEXT("NumChannels_"));
+static const FString AudioBufferName(TEXT("_AudioBuffer"));
+static const FString NumChannelsName(TEXT("_NumChannels"));
 
 FNiagaraDataInterfaceProxyOscilloscope::FNiagaraDataInterfaceProxyOscilloscope(int32 InResolution, float InScopeInMillseconds)
 	: PatchMixer()
@@ -239,6 +240,13 @@ void UNiagaraDataInterfaceAudioOscilloscope::GetVMExternalFunction(const FVMExte
 }
 
 #if WITH_EDITORONLY_DATA
+bool UNiagaraDataInterfaceAudioOscilloscope::AppendCompileHash(FNiagaraCompileHashVisitor* InVisitor) const
+{
+	bool bSuccess = Super::AppendCompileHash(InVisitor);
+	bSuccess &= InVisitor->UpdateShaderParameters<FShaderParameters>();
+	return bSuccess;
+}
+
 bool UNiagaraDataInterfaceAudioOscilloscope::GetFunctionHLSL(const FNiagaraDataInterfaceGPUParamInfo& ParamInfo, const FNiagaraDataInterfaceGeneratedFunction& FunctionInfo, int FunctionInstanceIndex, FString& OutHLSL)
 {
 	bool ParentRet = Super::GetFunctionHLSL(ParamInfo, FunctionInfo, FunctionInstanceIndex, OutHLSL);
@@ -272,9 +280,9 @@ bool UNiagaraDataInterfaceAudioOscilloscope::GetFunctionHLSL(const FNiagaraDataI
 		)");
 		TMap<FString, FStringFormatArg> ArgsBounds = {
 			{TEXT("FunctionName"), FStringFormatArg(FunctionInfo.InstanceName)},
-			{TEXT("ChannelCount"), FStringFormatArg(NumChannelsName + ParamInfo.DataInterfaceHLSLSymbol)},
+			{TEXT("ChannelCount"), FStringFormatArg(ParamInfo.DataInterfaceHLSLSymbol + NumChannelsName)},
 			{TEXT("AudioBufferNumSamples"), FStringFormatArg(Resolution)},
-			{TEXT("AudioBuffer"), FStringFormatArg(AudioBufferName + ParamInfo.DataInterfaceHLSLSymbol)},
+			{TEXT("AudioBuffer"), FStringFormatArg(ParamInfo.DataInterfaceHLSLSymbol + AudioBufferName)},
 		};
 		OutHLSL += FString::Format(FormatBounds, ArgsBounds);
 		return true;
@@ -290,7 +298,7 @@ bool UNiagaraDataInterfaceAudioOscilloscope::GetFunctionHLSL(const FNiagaraDataI
 		)");
 		TMap<FString, FStringFormatArg> ArgsBounds = {
 			{TEXT("FunctionName"), FStringFormatArg(FunctionInfo.InstanceName)},
-			{TEXT("ChannelCount"), FStringFormatArg(NumChannelsName + ParamInfo.DataInterfaceHLSLSymbol)},
+			{TEXT("ChannelCount"), FStringFormatArg(ParamInfo.DataInterfaceHLSLSymbol + NumChannelsName)},
 		};
 		OutHLSL += FString::Format(FormatBounds, ArgsBounds);
 		return true;
@@ -305,49 +313,26 @@ void UNiagaraDataInterfaceAudioOscilloscope::GetParameterDefinitionHLSL(const FN
 {
 	Super::GetParameterDefinitionHLSL(ParamInfo, OutHLSL);
 
-	static const TCHAR *FormatDeclarations = TEXT(R"(				
-		Buffer<float> {AudioBufferName};
-		int {NumChannelsName};
+	OutHLSL.Appendf(TEXT("int %s%s;\n"), *ParamInfo.DataInterfaceHLSLSymbol, *NumChannelsName);
+	OutHLSL.Appendf(TEXT("Buffer<float> %s%s;\n"), *ParamInfo.DataInterfaceHLSLSymbol, *AudioBufferName);
 
-	)");
-	TMap<FString, FStringFormatArg> ArgsDeclarations = {
-		{ TEXT("AudioBufferName"),    FStringFormatArg(AudioBufferName + ParamInfo.DataInterfaceHLSLSymbol) },
-		{ TEXT("NumChannelsName"),    FStringFormatArg(NumChannelsName + ParamInfo.DataInterfaceHLSLSymbol) },
-	};
-	OutHLSL += FString::Format(FormatDeclarations, ArgsDeclarations);
 }
 #endif
 
-struct FNiagaraDataInterfaceParametersCS_AudioOscilloscope : public FNiagaraDataInterfaceParametersCS
+void UNiagaraDataInterfaceAudioOscilloscope::BuildShaderParameters(FNiagaraShaderParametersBuilder& ShaderParametersBuilder) const
 {
-	DECLARE_TYPE_LAYOUT(FNiagaraDataInterfaceParametersCS_AudioOscilloscope, NonVirtual);
-
-	void Bind(const FNiagaraDataInterfaceGPUParamInfo& ParameterInfo, const class FShaderParameterMap& ParameterMap)
-	{
-		NumChannels.Bind(ParameterMap, *(NumChannelsName + ParameterInfo.DataInterfaceHLSLSymbol));
-		AudioBuffer.Bind(ParameterMap, *(AudioBufferName + ParameterInfo.DataInterfaceHLSLSymbol));
-	}
-
-	void Set(FRHICommandList& RHICmdList, const FNiagaraDataInterfaceSetArgs& Context) const
-	{
-		check(IsInRenderingThread());
-
-		FRHIComputeShader* ComputeShaderRHI = Context.Shader.GetComputeShader();
-
-		FNiagaraDataInterfaceProxyOscilloscope* NDI = (FNiagaraDataInterfaceProxyOscilloscope*)Context.DataInterface;
-		FReadBuffer& AudioBufferSRV = NDI->ComputeAndPostSRV();
-
-		SetShaderValue(RHICmdList, ComputeShaderRHI, NumChannels, NDI->GetNumChannels());
-		RHICmdList.SetShaderResourceViewParameter(ComputeShaderRHI, AudioBuffer.GetBaseIndex(), AudioBufferSRV.SRV);
-	}
-
-	LAYOUT_FIELD(FShaderParameter, NumChannels);
-	LAYOUT_FIELD(FShaderResourceParameter, AudioBuffer);
-};
+	ShaderParametersBuilder.AddNestedStruct<FShaderParameters>();
+}
 
-IMPLEMENT_TYPE_LAYOUT(FNiagaraDataInterfaceParametersCS_AudioOscilloscope);
+void UNiagaraDataInterfaceAudioOscilloscope::SetShaderParameters(const FNiagaraDataInterfaceSetShaderParametersContext& Context) const
+{
+	FNiagaraDataInterfaceProxyOscilloscope& DIProxy = Context.GetProxy<FNiagaraDataInterfaceProxyOscilloscope>();
+	FReadBuffer& AudioBufferSRV = DIProxy.ComputeAndPostSRV();
 
-IMPLEMENT_NIAGARA_DI_PARAMETER(UNiagaraDataInterfaceAudioOscilloscope, FNiagaraDataInterfaceParametersCS_AudioOscilloscope);
+	FShaderParameters* ShaderParameters = Context.GetParameterNestedStruct<FShaderParameters>();
+	ShaderParameters->NumChannels = DIProxy.GetNumChannels();
+	ShaderParameters->AudioBuffer = FNiagaraRenderer::GetSrvOrDefaultFloat(AudioBufferSRV.SRV);
+}
 
 void FNiagaraDataInterfaceProxyOscilloscope::OnUpdateResampling(int32 InResolution, float InScopeInMilliseconds)
 {
diff --git a/Engine/Plugins/FX/Niagara/Source/Niagara/Private/NiagaraDataInterfaceAudioSpectrum.cpp b/Engine/Plugins/FX/Niagara/Source/Niagara/Private/NiagaraDataInterfaceAudioSpectrum.cpp
--- a/Engine/Plugins/FX/Niagara/Source/Niagara/Private/NiagaraDataInterfaceAudioSpectrum.cpp
+++ b/Engine/Plugins/FX/Niagara/Source/Niagara/Private/NiagaraDataInterfaceAudioSpectrum.cpp
@@ -4,6 +4,7 @@
 #include "NiagaraTypes.h"
 #include "NiagaraCustomVersion.h"
 #include "NiagaraShader.h"
+#include "NiagaraShaderParametersBuilder.h"
 #include "ShaderParameterUtils.h"
 #include "Engine/Engine.h"
 #include "NiagaraComponent.h"
@@ -645,6 +646,13 @@ void UNiagaraDataInterfaceAudioSpectrum::GetVMExternalFunction(const FVMExternal
 }
 
 #if WITH_EDITORONLY_DATA
+bool UNiagaraDataInterfaceAudioSpectrum::AppendCompileHash(FNiagaraCompileHashVisitor* InVisitor) const
+{
+	bool bSuccess = Super::AppendCompileHash(InVisitor);
+	bSuccess &= InVisitor->UpdateShaderParameters<FShaderParameters>();
+	return bSuccess;
+}
+
 bool UNiagaraDataInterfaceAudioSpectrum::GetFunctionHLSL(const FNiagaraDataInterfaceGPUParamInfo& ParamInfo, const FNiagaraDataInterfaceGeneratedFunction& FunctionInfo, int FunctionInstanceIndex, FString& OutHLSL)
 {
 	bool ParentRet = Super::GetFunctionHLSL(ParamInfo, FunctionInfo, FunctionInstanceIndex, OutHLSL);
@@ -684,9 +692,9 @@ bool UNiagaraDataInterfaceAudioSpectrum::GetFunctionHLSL(const FNiagaraDataInter
 
 		TMap<FString, FStringFormatArg> ArgsBounds = {
 			{TEXT("FunctionName"), FStringFormatArg(FunctionInfo.InstanceName)},
-			{TEXT("ChannelCount"), FStringFormatArg(NumChannelsName + ParamInfo.DataInterfaceHLSLSymbol)},
-			{TEXT("SpectrumResolution"), FStringFormatArg(ResolutionName + ParamInfo.DataInterfaceHLSLSymbol)},
-			{TEXT("SpectrumBuffer"), FStringFormatArg(GetSpectrumName + ParamInfo.DataInterfaceHLSLSymbol)},
+			{TEXT("ChannelCount"), FStringFormatArg(ParamInfo.DataInterfaceHLSLSymbol + NumChannelsName)},
+			{TEXT("SpectrumResolution"), FStringFormatArg(ParamInfo.DataInterfaceHLSLSymbol + ResolutionName)},
+			{TEXT("SpectrumBuffer"), FStringFormatArg(ParamInfo.DataInterfaceHLSLSymbol + GetSpectrumName)},
 		};
 		OutHLSL += FString::Format(FormatBounds, ArgsBounds);
 		return true;
@@ -702,7 +710,7 @@ bool UNiagaraDataInterfaceAudioSpectrum::GetFunctionHLSL(const FNiagaraDataInter
 		)");
 		TMap<FString, FStringFormatArg> ArgsBounds = {
 			{TEXT("FunctionName"), FStringFormatArg(FunctionInfo.InstanceName)},
-			{TEXT("ChannelCount"), FStringFormatArg(NumChannelsName + ParamInfo.DataInterfaceHLSLSymbol)},
+			{TEXT("ChannelCount"), FStringFormatArg(ParamInfo.DataInterfaceHLSLSymbol + NumChannelsName)},
 		};
 		OutHLSL += FString::Format(FormatBounds, ArgsBounds);
 		return true;
@@ -724,14 +732,30 @@ void UNiagaraDataInterfaceAudioSpectrum::GetParameterDefinitionHLSL(const FNiaga
 
 	)");
 	TMap<FString, FStringFormatArg> ArgsDeclarations = {
-		{ TEXT("GetSpectrumName"),    FStringFormatArg(GetSpectrumName + ParamInfo.DataInterfaceHLSLSymbol) },
-		{ TEXT("NumChannelsName"),    FStringFormatArg(NumChannelsName + ParamInfo.DataInterfaceHLSLSymbol) },
-		{ TEXT("ResolutionName"),     FStringFormatArg(ResolutionName + ParamInfo.DataInterfaceHLSLSymbol) },
+		{ TEXT("GetSpectrumName"),    FStringFormatArg(ParamInfo.DataInterfaceHLSLSymbol + GetSpectrumName) },
+		{ TEXT("NumChannelsName"),    FStringFormatArg(ParamInfo.DataInterfaceHLSLSymbol + NumChannelsName) },
+		{ TEXT("ResolutionName"),     FStringFormatArg(ParamInfo.DataInterfaceHLSLSymbol + ResolutionName) },
 	};
 	OutHLSL += FString::Format(FormatDeclarations, ArgsDeclarations);
 }
 #endif
 
+void UNiagaraDataInterfaceAudioSpectrum::BuildShaderParameters(FNiagaraShaderParametersBuilder& ShaderParametersBuilder) const
+{
+	ShaderParametersBuilder.AddNestedStruct<FShaderParameters>();
+}
+
+void UNiagaraDataInterfaceAudioSpectrum::SetShaderParameters(const FNiagaraDataInterfaceSetShaderParametersContext& Context) const
+{
+	FNiagaraDataInterfaceProxySpectrum& DIProxy = Context.GetProxy<FNiagaraDataInterfaceProxySpectrum>();
+	FReadBuffer& SpectrumBuffer = DIProxy.ComputeAndPostSRV();
+
+	FShaderParameters* ShaderParameters = Context.GetParameterNestedStruct<FShaderParameters>();
+	ShaderParameters->NumChannels = DIProxy.GetNumChannels();
+	ShaderParameters->Resolution = DIProxy.GetNumBands();
+	ShaderParameters->SpectrumBuffer = FNiagaraRenderer::GetSrvOrDefaultFloat(SpectrumBuffer.SRV);
+}
+
 bool UNiagaraDataInterfaceAudioSpectrum::Equals(const UNiagaraDataInterface* Other) const
 {
 	bool bIsEqual = Super::Equals(Other);
@@ -746,43 +770,6 @@ bool UNiagaraDataInterfaceAudioSpectrum::Equals(const UNiagaraDataInterface* Oth
 	return bIsEqual;
 }
 
-struct FNiagaraDataInterfaceParametersCS_AudioSpectrum : public FNiagaraDataInterfaceParametersCS
-{
-	DECLARE_TYPE_LAYOUT(FNiagaraDataInterfaceParametersCS_AudioSpectrum, NonVirtual);
-
-	void Bind(const FNiagaraDataInterfaceGPUParamInfo& ParameterInfo, const class FShaderParameterMap& ParameterMap)
-	{
-		NumChannels.Bind(ParameterMap, *(UNiagaraDataInterfaceAudioSpectrum::NumChannelsName + ParameterInfo.DataInterfaceHLSLSymbol));
-
-		Resolution.Bind(ParameterMap, *(UNiagaraDataInterfaceAudioSpectrum::ResolutionName + ParameterInfo.DataInterfaceHLSLSymbol));
-
-		SpectrumBuffer.Bind(ParameterMap, *(UNiagaraDataInterfaceAudioSpectrum::GetSpectrumName + ParameterInfo.DataInterfaceHLSLSymbol));
-	}
-
-	void Set(FRHICommandList& RHICmdList, const FNiagaraDataInterfaceSetArgs& Context) const 
-	{
-		check(IsInRenderingThread());
-
-		FRHIComputeShader* ComputeShaderRHI = Context.Shader.GetComputeShader();
-
-		FNiagaraDataInterfaceProxySpectrum* NDI = (FNiagaraDataInterfaceProxySpectrum*)Context.DataInterface;
-		FReadBuffer& SpectrumBufferSRV = NDI->ComputeAndPostSRV();
-
-		SetShaderValue(RHICmdList, ComputeShaderRHI, NumChannels, NDI->GetNumChannels());
-		SetShaderValue(RHICmdList, ComputeShaderRHI, Resolution, NDI->GetNumBands());
-		RHICmdList.SetShaderResourceViewParameter(ComputeShaderRHI, SpectrumBuffer.GetBaseIndex(), SpectrumBufferSRV.SRV);
-	}
-
-	LAYOUT_FIELD(FShaderParameter, NumChannels);
-	LAYOUT_FIELD(FShaderParameter, Resolution);
-	LAYOUT_FIELD(FShaderResourceParameter,SpectrumBuffer);
-};
-
-IMPLEMENT_TYPE_LAYOUT(FNiagaraDataInterfaceParametersCS_AudioSpectrum);
-
-IMPLEMENT_NIAGARA_DI_PARAMETER(UNiagaraDataInterfaceAudioSpectrum, FNiagaraDataInterfaceParametersCS_AudioSpectrum);
-
-
 #if WITH_EDITOR
 void UNiagaraDataInterfaceAudioSpectrum::PostEditChangeProperty(struct FPropertyChangedEvent& PropertyChangedEvent)
 {
@@ -870,9 +857,9 @@ const FName UNiagaraDataInterfaceAudioSpectrum::GetSpectrumFunctionName("AudioSp
 const FName UNiagaraDataInterfaceAudioSpectrum::GetNumChannelsFunctionName("GetNumChannels");
 
 // Global variable prefixes, used in HLSL parameter declarations.
-const FString UNiagaraDataInterfaceAudioSpectrum::GetSpectrumName(TEXT("SpectrumBuffer_"));
-const FString UNiagaraDataInterfaceAudioSpectrum::NumChannelsName(TEXT("NumChannels_"));
-const FString UNiagaraDataInterfaceAudioSpectrum::ResolutionName(TEXT("Resolution_"));
+const FString UNiagaraDataInterfaceAudioSpectrum::GetSpectrumName(TEXT("_SpectrumBuffer"));
+const FString UNiagaraDataInterfaceAudioSpectrum::NumChannelsName(TEXT("_NumChannels"));
+const FString UNiagaraDataInterfaceAudioSpectrum::ResolutionName(TEXT("_Resolution"));
 
 
 #undef LOCTEXT_NAMESPACE
