Fix loading & duplicating of niagara cooked data when editor only data is not filtered

diff --git a/Engine/Plugins/FX/Niagara/Source/Niagara/Private/NiagaraSpriteRendererProperties.cpp b/Engine/Plugins/FX/Niagara/Source/Niagara/Private/NiagaraSpriteRendererProperties.cpp
--- a/Engine/Plugins/FX/Niagara/Source/Niagara/Private/NiagaraSpriteRendererProperties.cpp
+++ b/Engine/Plugins/FX/Niagara/Source/Niagara/Private/NiagaraSpriteRendererProperties.cpp
@@ -189,7 +189,7 @@ void UNiagaraSpriteRendererProperties::Serialize(FStructuredArchive::FRecord Rec
 
 	bool bIsCookedForEditor = false;
 #if WITH_EDITORONLY_DATA
-	bIsCookedForEditor = GetOutermost()->bIsCookedForEditor;
+	bIsCookedForEditor = ((Ar.GetPortFlags() & PPF_Duplicate) == 0) && GetPackage()->HasAnyPackageFlags(PKG_Cooked);
 #endif // WITH_EDITORONLY_DATA
 
 	FArchive& UnderlyingArchive = Record.GetUnderlyingArchive();
