- Landscape DI uses Shader Template & Shader Params

diff --git a/Engine/Plugins/FX/Niagara/Shaders/Private/NiagaraDataInterfaceLandscape.ush b/Engine/Plugins/FX/Niagara/Shaders/Private/NiagaraDataInterfaceLandscape.ush
--- a/Engine/Plugins/FX/Niagara/Shaders/Private/NiagaraDataInterfaceLandscape.ush
+++ b/Engine/Plugins/FX/Niagara/Shaders/Private/NiagaraDataInterfaceLandscape.ush
@@ -1,7 +1,5 @@
 // Copyright Epic Games, Inc. All Rights Reserved.
 
-#pragma once
-
 #ifndef NUM_VIRTUALTEXTURE_SAMPLES
 	#define NUM_VIRTUALTEXTURE_SAMPLES 1
 #endif
@@ -13,81 +11,72 @@
 #define VIRTUAL_TEXTURE_UNPACK_NORMAL_BC3BC3 1
 #define VIRTUAL_TEXTURE_UNPACK_NORMAL_BC5BC1 2
 
-struct FNDILandscapeContext
-{
-	Texture2D BaseColorVirtualTexture;
-	Texture2D<uint4> BaseColorVirtualTexturePageTable;
-	SamplerState BaseColorVirtualTextureSampler;
-	float4x4 BaseColorVirtualTextureWorldToUvTransform;
-	uint4 BaseColorVirtualTexturePackedUniform0;
-	uint4 BaseColorVirtualTexturePackedUniform1;
-	uint4 BaseColorVirtualTextureUniforms;
-	bool BaseColorVirtualTextureSRGB;
-	bool BaseColorVirtualTextureEnabled;
-
-	Texture2D HeightVirtualTexture;
-	Texture2D<uint4> HeightVirtualTexturePageTable;
-	SamplerState HeightVirtualTextureSampler;
-	float4x4 HeightVirtualTextureWorldToUvTransform;
-	uint4 HeightVirtualTexturePackedUniform0;
-	uint4 HeightVirtualTexturePackedUniform1;
-	uint4 HeightVirtualTextureUniforms;
-	bool HeightVirtualTextureEnabled;
-
-	Texture2D NormalVirtualTexture0;
-	Texture2D NormalVirtualTexture1;
-	Texture2D<uint4> NormalVirtualTexturePageTable;
-	SamplerState NormalVirtualTextureSampler;
-	float4x4 NormalVirtualTextureWorldToUvTransform;
-	uint4 NormalVirtualTexturePackedUniform0;
-	uint4 NormalVirtualTexturePackedUniform1;
-	uint4 NormalVirtualTextureUniforms0;
-	uint4 NormalVirtualTextureUniforms1;
-	int NormalVirtualTextureUnpackMode;
-	bool NormalVirtualTextureEnabled;
-
-	Texture2D CachedHeightTexture;
-	SamplerState CachedHeightTextureSampler;
-	bool CachedHeightTextureEnabled;
-	float4x4 CachedHeightTextureWorldToUvTransform;
-	float4x4 CachedHeightTextureUvToWorldTransform;
-	float4 CachedHeightTextureUvScaleBias;
-	float2 CachedHeightTextureGridSize;
-	SamplerState PointClampedSampler;
-
-	Texture2D<int> CachedPhysMatTexture;
-	int2 CachedPhysMatTextureDimension;
-
-	float3 SystemLWCTile;
-};
-
-void NDILandscape_GetBaseColor(in FNDILandscapeContext DIContext, float3 InWorldPos, out float3 OutBaseColor, out bool OutIsValid)
+Texture2D			{ParameterName}_BaseColorVirtualTexture;
+Texture2D<uint4>	{ParameterName}_BaseColorVirtualTexturePageTable;
+SamplerState		{ParameterName}_BaseColorVirtualTextureSampler;
+float4x4			{ParameterName}_BaseColorVirtualTextureWorldToUvTransform;
+uint				{ParameterName}_BaseColorVirtualTextureSRGB;
+uint				{ParameterName}_BaseColorVirtualTextureEnabled;
+uint4				{ParameterName}_BaseColorVirtualTexturePackedUniform0;
+uint4				{ParameterName}_BaseColorVirtualTexturePackedUniform1;
+uint4				{ParameterName}_BaseColorVirtualTextureUniforms;
+Texture2D			{ParameterName}_HeightVirtualTexture;
+Texture2D<uint4>	{ParameterName}_HeightVirtualTexturePageTable;
+SamplerState		{ParameterName}_HeightVirtualTextureSampler;
+float4x4			{ParameterName}_HeightVirtualTextureWorldToUvTransform;
+uint				{ParameterName}_HeightVirtualTextureEnabled;
+uint4				{ParameterName}_HeightVirtualTexturePackedUniform0;
+uint4				{ParameterName}_HeightVirtualTexturePackedUniform1;
+uint4				{ParameterName}_HeightVirtualTextureUniforms;
+Texture2D			{ParameterName}_NormalVirtualTexture0;
+Texture2D			{ParameterName}_NormalVirtualTexture1;
+Texture2D<uint4>	{ParameterName}_NormalVirtualTexturePageTable;
+SamplerState		{ParameterName}_NormalVirtualTextureSampler;
+float4x4			{ParameterName}_NormalVirtualTextureWorldToUvTransform;
+uint4				{ParameterName}_NormalVirtualTexturePackedUniform0;
+uint4				{ParameterName}_NormalVirtualTexturePackedUniform1;
+uint4				{ParameterName}_NormalVirtualTextureUniforms0;
+uint4				{ParameterName}_NormalVirtualTextureUniforms1;
+int					{ParameterName}_NormalVirtualTextureUnpackMode;
+uint				{ParameterName}_NormalVirtualTextureEnabled;
+Texture2D			{ParameterName}_CachedHeightTexture;
+SamplerState		{ParameterName}_CachedHeightTextureSampler;
+uint				{ParameterName}_CachedHeightTextureEnabled;
+float4x4			{ParameterName}_CachedHeightTextureWorldToUvTransform;
+float4x4			{ParameterName}_CachedHeightTextureUvToWorldTransform;
+float4				{ParameterName}_CachedHeightTextureUvScaleBias;
+float2				{ParameterName}_CachedHeightTextureGridSize;
+SamplerState		{ParameterName}_PointClampedSampler;
+Texture2D<int>		{ParameterName}_CachedPhysMatTexture;
+int2				{ParameterName}_CachedPhysMatTextureDimension;
+
+void GetBaseColor_{ParameterName}(float3 InWorldPos, out float3 OutBaseColor, out bool OutIsValid)
 {
 	// todo[mg] convert world pos and transform origin to LWC pos
-	if (DIContext.BaseColorVirtualTextureEnabled)
+	if ({ParameterName}_BaseColorVirtualTextureEnabled != 0)
 	{
-		float3 UnpackBaseColorOrigin = DIContext.BaseColorVirtualTextureWorldToUvTransform[0].xyz;
-		float3 UnpackBaseColorU = DIContext.BaseColorVirtualTextureWorldToUvTransform[1].xyz;
-		float3 UnpackBaseColorV = DIContext.BaseColorVirtualTextureWorldToUvTransform[2].xyz;
+		float3 UnpackBaseColorOrigin = {ParameterName}_BaseColorVirtualTextureWorldToUvTransform[0].xyz;
+		float3 UnpackBaseColorU = {ParameterName}_BaseColorVirtualTextureWorldToUvTransform[1].xyz;
+		float3 UnpackBaseColorV = {ParameterName}_BaseColorVirtualTextureWorldToUvTransform[2].xyz;
 
 		float2 SampleUv = VirtualTextureWorldToUV(LWCPromote(InWorldPos), LWCPromote(UnpackBaseColorOrigin), UnpackBaseColorU, UnpackBaseColorV);
 
 		VTPageTableResult PageTable = TextureLoadVirtualPageTableLevel(
-			DIContext.BaseColorVirtualTexturePageTable,
-			VTPageTableUniform_Unpack(DIContext.BaseColorVirtualTexturePackedUniform0, DIContext.BaseColorVirtualTexturePackedUniform1),
+			{ParameterName}_BaseColorVirtualTexturePageTable,
+			VTPageTableUniform_Unpack({ParameterName}_BaseColorVirtualTexturePackedUniform0, {ParameterName}_BaseColorVirtualTexturePackedUniform1),
 			SampleUv,
 			VTADDRESSMODE_CLAMP,
 			VTADDRESSMODE_CLAMP,
 			0.0f);
 
 		float4 PackedValue = TextureVirtualSample(
-			DIContext.BaseColorVirtualTexture,
-			DIContext.BaseColorVirtualTextureSampler,
+			{ParameterName}_BaseColorVirtualTexture,
+			{ParameterName}_BaseColorVirtualTextureSampler,
 			PageTable,
 			0,
-			VTUniform_Unpack(DIContext.BaseColorVirtualTextureUniforms));
+			VTUniform_Unpack({ParameterName}_BaseColorVirtualTextureUniforms));
 
-		OutBaseColor = DIContext.BaseColorVirtualTextureSRGB ? VirtualTextureUnpackBaseColorSRGB(PackedValue) : VirtualTextureUnpackBaseColorYCoCg(PackedValue);
+		OutBaseColor = {ParameterName}_BaseColorVirtualTextureSRGB != 0 ? VirtualTextureUnpackBaseColorSRGB(PackedValue) : VirtualTextureUnpackBaseColorYCoCg(PackedValue);
 		OutIsValid = true;
 	}
 	else
@@ -97,41 +86,41 @@ void NDILandscape_GetBaseColor(in FNDILandscapeContext DIContext, float3 InWorld
 	}
 }
 
-void NDILandscape_GetHeight(in FNDILandscapeContext DIContext, float3 InWorldPos, out float OutHeight, out bool OutIsValid)
+void GetHeight_{ParameterName}(float3 InWorldPos, out float OutHeight, out bool OutIsValid)
 {
 	// todo[mg] convert world pos and transform origin to LWC pos
-	if (DIContext.HeightVirtualTextureEnabled)
+	if ({ParameterName}_HeightVirtualTextureEnabled != 0)
 	{
-		float3 UnpackHeightOrigin = DIContext.HeightVirtualTextureWorldToUvTransform[0].xyz;
-		float3 UnpackHeightU = DIContext.HeightVirtualTextureWorldToUvTransform[1].xyz;
-		float3 UnpackHeightV = DIContext.HeightVirtualTextureWorldToUvTransform[2].xyz;
-		float2 UnpackHeightScaleBias = DIContext.HeightVirtualTextureWorldToUvTransform[3].xy;
+		float3 UnpackHeightOrigin = {ParameterName}_HeightVirtualTextureWorldToUvTransform[0].xyz;
+		float3 UnpackHeightU = {ParameterName}_HeightVirtualTextureWorldToUvTransform[1].xyz;
+		float3 UnpackHeightV = {ParameterName}_HeightVirtualTextureWorldToUvTransform[2].xyz;
+		float2 UnpackHeightScaleBias = {ParameterName}_HeightVirtualTextureWorldToUvTransform[3].xy;
 
 		float2 SampleUv = VirtualTextureWorldToUV(LWCPromote(InWorldPos), LWCPromote(UnpackHeightOrigin), UnpackHeightU, UnpackHeightV);
 
 		VTPageTableResult PageTable = TextureLoadVirtualPageTableLevel(
-			DIContext.HeightVirtualTexturePageTable,
-			VTPageTableUniform_Unpack(DIContext.HeightVirtualTexturePackedUniform0, DIContext.HeightVirtualTexturePackedUniform1),
+			{ParameterName}_HeightVirtualTexturePageTable,
+			VTPageTableUniform_Unpack({ParameterName}_HeightVirtualTexturePackedUniform0, {ParameterName}_HeightVirtualTexturePackedUniform1),
 			SampleUv,
 			VTADDRESSMODE_CLAMP,
 			VTADDRESSMODE_CLAMP,
 			0.0f);
 
 		float4 PackedValue = TextureVirtualSample(
-			DIContext.HeightVirtualTexture,
-			DIContext.HeightVirtualTextureSampler,
+			{ParameterName}_HeightVirtualTexture,
+			{ParameterName}_HeightVirtualTextureSampler,
 			PageTable,
 			0,
-			VTUniform_Unpack(DIContext.HeightVirtualTextureUniforms));
+			VTUniform_Unpack({ParameterName}_HeightVirtualTextureUniforms));
 
 		OutHeight = VirtualTextureUnpackHeight(PackedValue, UnpackHeightScaleBias);
 		OutIsValid = true;
 	}
-	else if (DIContext.CachedHeightTextureEnabled)
+	else if ({ParameterName}_CachedHeightTextureEnabled != 0)
 	{
-		float2 Uv = mul(float4(InWorldPos, 1.0f), DIContext.CachedHeightTextureWorldToUvTransform).xy;
-		Uv = Uv * DIContext.CachedHeightTextureUvScaleBias.xy + DIContext.CachedHeightTextureUvScaleBias.zw;
-		OutHeight = DIContext.CachedHeightTexture.SampleLevel(DIContext.CachedHeightTextureSampler, Uv, 0).x;
+		float2 Uv = mul(float4(InWorldPos, 1.0f), {ParameterName}_CachedHeightTextureWorldToUvTransform).xy;
+		Uv = Uv * {ParameterName}_CachedHeightTextureUvScaleBias.xy + {ParameterName}_CachedHeightTextureUvScaleBias.zw;
+		OutHeight = {ParameterName}_CachedHeightTexture.SampleLevel({ParameterName}_CachedHeightTextureSampler, Uv, 0).x;
 		OutIsValid = true;
 	}
 	else
@@ -141,44 +130,44 @@ void NDILandscape_GetHeight(in FNDILandscapeContext DIContext, float3 InWorldPos
 	}
 }
 
-void NDILandscape_GetWorldNormal(in FNDILandscapeContext DIContext, float3 InWorldPos, out float3 OutNormal, out bool OutIsValid)
+void GetWorldNormal_{ParameterName}(float3 InWorldPos, out float3 OutNormal, out bool OutIsValid)
 {
 	// todo[mg] convert world pos and transform origin to LWC pos
-	if (DIContext.NormalVirtualTextureEnabled)
+	if ({ParameterName}_NormalVirtualTextureEnabled != 0)
 	{
-		float3 UnpackNormalOrigin = DIContext.NormalVirtualTextureWorldToUvTransform[0].xyz;
-		float3 UnpackNormalU = DIContext.NormalVirtualTextureWorldToUvTransform[1].xyz;
-		float3 UnpackNormalV = DIContext.NormalVirtualTextureWorldToUvTransform[2].xyz;
+		float3 UnpackNormalOrigin = {ParameterName}_NormalVirtualTextureWorldToUvTransform[0].xyz;
+		float3 UnpackNormalU = {ParameterName}_NormalVirtualTextureWorldToUvTransform[1].xyz;
+		float3 UnpackNormalV = {ParameterName}_NormalVirtualTextureWorldToUvTransform[2].xyz;
 
 		float2 SampleUv = VirtualTextureWorldToUV(LWCPromote(InWorldPos), LWCPromote(UnpackNormalOrigin), UnpackNormalU, UnpackNormalV);
 
 		VTPageTableResult PageTable = TextureLoadVirtualPageTableLevel(
-			DIContext.NormalVirtualTexturePageTable,
-			VTPageTableUniform_Unpack(DIContext.NormalVirtualTexturePackedUniform0, DIContext.NormalVirtualTexturePackedUniform1),
+			{ParameterName}_NormalVirtualTexturePageTable,
+			VTPageTableUniform_Unpack({ParameterName}_NormalVirtualTexturePackedUniform0, {ParameterName}_NormalVirtualTexturePackedUniform1),
 			SampleUv,
 			VTADDRESSMODE_CLAMP,
 			VTADDRESSMODE_CLAMP,
 			0.0f);
 
 		float4 PackedValue0 = TextureVirtualSample(
-			DIContext.NormalVirtualTexture0,
-			DIContext.NormalVirtualTextureSampler,
+			{ParameterName}_NormalVirtualTexture0,
+			{ParameterName}_NormalVirtualTextureSampler,
 			PageTable,
 			0,
-			VTUniform_Unpack(DIContext.NormalVirtualTextureUniforms0));
+			VTUniform_Unpack({ParameterName}_NormalVirtualTextureUniforms0));
 
 		float4 PackedValue1 = TextureVirtualSample(
-			DIContext.NormalVirtualTexture1,
-			DIContext.NormalVirtualTextureSampler,
+			{ParameterName}_NormalVirtualTexture1,
+			{ParameterName}_NormalVirtualTextureSampler,
 			PageTable,
 			0,
-			VTUniform_Unpack(DIContext.NormalVirtualTextureUniforms1));
+			VTUniform_Unpack({ParameterName}_NormalVirtualTextureUniforms1));
 
-		if (DIContext.NormalVirtualTextureUnpackMode == VIRTUAL_TEXTURE_UNPACK_NORMAL_BC3BC3)
+		if ({ParameterName}_NormalVirtualTextureUnpackMode == VIRTUAL_TEXTURE_UNPACK_NORMAL_BC3BC3)
 		{
 			OutNormal = VirtualTextureUnpackNormalBC3BC3(PackedValue0, PackedValue1);
 		}
-		else if (DIContext.NormalVirtualTextureUnpackMode == VIRTUAL_TEXTURE_UNPACK_NORMAL_BC5BC1)
+		else if ({ParameterName}_NormalVirtualTextureUnpackMode == VIRTUAL_TEXTURE_UNPACK_NORMAL_BC5BC1)
 		{
 			OutNormal = VirtualTextureUnpackNormalBC5BC1(PackedValue0, PackedValue1);
 		}
@@ -191,24 +180,24 @@ void NDILandscape_GetWorldNormal(in FNDILandscapeContext DIContext, float3 InWor
 
 		OutIsValid = true;
 	}
-	else if (DIContext.CachedHeightTextureEnabled)
+	else if ({ParameterName}_CachedHeightTextureEnabled != 0)
 	{
-		float2 Uv = mul(float4(InWorldPos, 1.0f), DIContext.CachedHeightTextureWorldToUvTransform).xy;
-		Uv = Uv * DIContext.CachedHeightTextureUvScaleBias.xy + DIContext.CachedHeightTextureUvScaleBias.zw;
+		float2 Uv = mul(float4(InWorldPos, 1.0f), {ParameterName}_CachedHeightTextureWorldToUvTransform).xy;
+		Uv = Uv * {ParameterName}_CachedHeightTextureUvScaleBias.xy + {ParameterName}_CachedHeightTextureUvScaleBias.zw;
 
-		float Left = DIContext.CachedHeightTexture.SampleLevel(DIContext.PointClampedSampler, Uv, 0, int2(-1,0)).x;
-		float Right = DIContext.CachedHeightTexture.SampleLevel(DIContext.PointClampedSampler, Uv, 0, int2(1,0)).x;
-		float Top = DIContext.CachedHeightTexture.SampleLevel(DIContext.PointClampedSampler, Uv, 0, int2(0,1)).x;
-		float Bottom = DIContext.CachedHeightTexture.SampleLevel(DIContext.PointClampedSampler, Uv, 0, int2(0,-1)).x;
+		float Left = {ParameterName}_CachedHeightTexture.SampleLevel({ParameterName}_PointClampedSampler, Uv, 0, int2(-1,0)).x;
+		float Right = {ParameterName}_CachedHeightTexture.SampleLevel({ParameterName}_PointClampedSampler, Uv, 0, int2(1,0)).x;
+		float Top = {ParameterName}_CachedHeightTexture.SampleLevel({ParameterName}_PointClampedSampler, Uv, 0, int2(0,1)).x;
+		float Bottom = {ParameterName}_CachedHeightTexture.SampleLevel({ParameterName}_PointClampedSampler, Uv, 0, int2(0,-1)).x;
 
 		// cross product of our two gradient vectors
 		OutNormal = float3(
-			DIContext.CachedHeightTextureGridSize.y * (Left - Right),
-			DIContext.CachedHeightTextureGridSize.x * (Bottom - Top),
-			2.0f * DIContext.CachedHeightTextureGridSize.x * DIContext.CachedHeightTextureGridSize.y);
+			{ParameterName}_CachedHeightTextureGridSize.y * (Left - Right),
+			{ParameterName}_CachedHeightTextureGridSize.x * (Bottom - Top),
+			2.0f * {ParameterName}_CachedHeightTextureGridSize.x * {ParameterName}_CachedHeightTextureGridSize.y);
 
 		// next rotate the normal to make it world space
-		OutNormal = mul(float4(OutNormal, 0.0f), DIContext.CachedHeightTextureUvToWorldTransform).xyz;
+		OutNormal = mul(float4(OutNormal, 0.0f), {ParameterName}_CachedHeightTextureUvToWorldTransform).xyz;
 		OutNormal = normalize(OutNormal);
 
 		OutIsValid = true;
@@ -220,14 +209,14 @@ void NDILandscape_GetWorldNormal(in FNDILandscapeContext DIContext, float3 InWor
 	}
 }
 
-void NDILandscape_GetPhysicalMaterialIndex(in FNDILandscapeContext DIContext, float3 InWorldPos, out int OutIndex, out bool OutIsValid)
+void GetPhysicalMaterialIndex_{ParameterName}(float3 InWorldPos, out int OutIndex, out bool OutIsValid)
 {
 	// todo[mg] convert world pos to LWC pos
-	float2 Uv = mul(float4(InWorldPos, 1.0f), DIContext.CachedHeightTextureWorldToUvTransform).xy;
-	Uv = Uv * DIContext.CachedHeightTextureUvScaleBias.xy + DIContext.CachedHeightTextureUvScaleBias.zw;
+	float2 Uv = mul(float4(InWorldPos, 1.0f), {ParameterName}_CachedHeightTextureWorldToUvTransform).xy;
+	Uv = Uv * {ParameterName}_CachedHeightTextureUvScaleBias.xy + {ParameterName}_CachedHeightTextureUvScaleBias.zw;
 
-	int3 SampleIndex = int3(Uv * DIContext.CachedPhysMatTextureDimension, 0);
-	OutIndex = DIContext.CachedPhysMatTexture.Load(SampleIndex).x;
+	int3 SampleIndex = int3(Uv * {ParameterName}_CachedPhysMatTextureDimension, 0);
+	OutIndex = {ParameterName}_CachedPhysMatTexture.Load(SampleIndex).x;
 
 	if (OutIndex >= 255)
 	{
@@ -240,87 +229,6 @@ void NDILandscape_GetPhysicalMaterialIndex(in FNDILandscapeContext DIContext, fl
 	}
 }
 
-////////////////////////////////////////////////////////////////////////////////////////////////////////////////////////////////////////////////////////////////
-
-#define NDILANDSCAPE_DECLARE_CONSTANTS(NAME) \
-	Texture2D			BaseColorVirtualTexture_ ## NAME; \
-	Texture2D<uint4>	BaseColorVirtualTexturePageTable_ ## NAME; \
-	SamplerState		BaseColorVirtualTextureSampler_ ## NAME; \
-	float4x4			BaseColorVirtualTextureWorldToUvTransform_ ## NAME; \
-	bool				BaseColorVirtualTextureSRGB_ ## NAME; \
-	bool				BaseColorVirtualTextureEnabled_ ## NAME; \
-	uint4				BaseColorVirtualTexturePackedUniform0_ ## NAME; \
-	uint4				BaseColorVirtualTexturePackedUniform1_ ## NAME; \
-	uint4				BaseColorVirtualTextureUniforms_ ## NAME; \
-	Texture2D			HeightVirtualTexture_ ## NAME; \
-	Texture2D<uint4>	HeightVirtualTexturePageTable_ ## NAME; \
-	SamplerState		HeightVirtualTextureSampler_ ## NAME; \
-	float4x4			HeightVirtualTextureWorldToUvTransform_ ## NAME; \
-	bool				HeightVirtualTextureEnabled_ ## NAME; \
-	uint4				HeightVirtualTexturePackedUniform0_ ## NAME; \
-	uint4				HeightVirtualTexturePackedUniform1_ ## NAME; \
-	uint4				HeightVirtualTextureUniforms_ ## NAME; \
-	Texture2D			NormalVirtualTexture0_ ## NAME; \
-	Texture2D			NormalVirtualTexture1_ ## NAME; \
-	Texture2D<uint4>	NormalVirtualTexturePageTable_ ## NAME; \
-	SamplerState		NormalVirtualTextureSampler_ ## NAME; \
-	float4x4			NormalVirtualTextureWorldToUvTransform_ ## NAME; \
-	uint4				NormalVirtualTexturePackedUniform0_ ## NAME; \
-	uint4				NormalVirtualTexturePackedUniform1_ ## NAME; \
-	uint4				NormalVirtualTextureUniforms0_ ## NAME; \
-	uint4				NormalVirtualTextureUniforms1_ ## NAME; \
-	int					NormalVirtualTextureUnpackMode_ ## NAME; \
-	bool				NormalVirtualTextureEnabled_ ## NAME; \
-	Texture2D			CachedHeightTexture_ ## NAME; \
-	SamplerState		CachedHeightTextureSampler_ ## NAME; \
-	bool				CachedHeightTextureEnabled_ ## NAME; \
-	float4x4			CachedHeightTextureWorldToUvTransform_ ## NAME; \
-	float4x4			CachedHeightTextureUvToWorldTransform_ ## NAME; \
-	float4				CachedHeightTextureUvScaleBias_ ## NAME; \
-	float2				CachedHeightTextureGridSize_ ## NAME; \
-	SamplerState		PointClampedSampler_ ## NAME; \
-	Texture2D<int>		CachedPhysMatTexture_ ## NAME; \
-	int2				CachedPhysMatTextureDimension_ ## NAME; \
-	float3				SystemLWCTile_ ## NAME;
-
-#define NDILANDSCAPE_MAKE_CONTEXT(NAME) \
-	FNDILandscapeContext DIContext; \
-	DIContext.BaseColorVirtualTexture = BaseColorVirtualTexture_ ## NAME; \
-	DIContext.BaseColorVirtualTexturePageTable = BaseColorVirtualTexturePageTable_ ## NAME; \
-	DIContext.BaseColorVirtualTextureSampler = BaseColorVirtualTextureSampler_ ## NAME; \
-	DIContext.BaseColorVirtualTextureWorldToUvTransform = BaseColorVirtualTextureWorldToUvTransform_ ## NAME; \
-	DIContext.BaseColorVirtualTexturePackedUniform0 = BaseColorVirtualTexturePackedUniform0_ ## NAME; \
-	DIContext.BaseColorVirtualTexturePackedUniform1 = BaseColorVirtualTexturePackedUniform1_ ## NAME; \
-	DIContext.BaseColorVirtualTextureUniforms = BaseColorVirtualTextureUniforms_ ## NAME; \
-	DIContext.BaseColorVirtualTextureSRGB = BaseColorVirtualTextureSRGB_ ## NAME; \
-	DIContext.BaseColorVirtualTextureEnabled = BaseColorVirtualTextureEnabled_ ## NAME; \
-	DIContext.HeightVirtualTexture = HeightVirtualTexture_ ## NAME; \
-	DIContext.HeightVirtualTexturePageTable = HeightVirtualTexturePageTable_ ## NAME; \
-	DIContext.HeightVirtualTextureSampler = HeightVirtualTextureSampler_ ## NAME; \
-	DIContext.HeightVirtualTextureWorldToUvTransform = HeightVirtualTextureWorldToUvTransform_ ## NAME; \
-	DIContext.HeightVirtualTextureEnabled = HeightVirtualTextureEnabled_ ## NAME; \
-	DIContext.HeightVirtualTexturePackedUniform0 = HeightVirtualTexturePackedUniform0_ ## NAME; \
-	DIContext.HeightVirtualTexturePackedUniform1 = HeightVirtualTexturePackedUniform1_ ## NAME; \
-	DIContext.HeightVirtualTextureUniforms = HeightVirtualTextureUniforms_ ## NAME; \
-	DIContext.NormalVirtualTexture0 = NormalVirtualTexture0_ ## NAME; \
-	DIContext.NormalVirtualTexture1 = NormalVirtualTexture1_ ## NAME; \
-	DIContext.NormalVirtualTexturePageTable = NormalVirtualTexturePageTable_ ## NAME; \
-	DIContext.NormalVirtualTextureSampler = NormalVirtualTextureSampler_ ## NAME; \
-	DIContext.NormalVirtualTextureWorldToUvTransform = NormalVirtualTextureWorldToUvTransform_ ## NAME; \
-	DIContext.NormalVirtualTexturePackedUniform0 = NormalVirtualTexturePackedUniform0_ ## NAME; \
-	DIContext.NormalVirtualTexturePackedUniform1 = NormalVirtualTexturePackedUniform1_ ## NAME; \
-	DIContext.NormalVirtualTextureUniforms0 = NormalVirtualTextureUniforms0_ ## NAME; \
-	DIContext.NormalVirtualTextureUniforms1 = NormalVirtualTextureUniforms1_ ## NAME; \
-	DIContext.NormalVirtualTextureUnpackMode = NormalVirtualTextureUnpackMode_ ## NAME; \
-	DIContext.NormalVirtualTextureEnabled = NormalVirtualTextureEnabled_ ## NAME; \
-	DIContext.CachedHeightTexture = CachedHeightTexture_ ## NAME; \
-	DIContext.CachedHeightTextureSampler = CachedHeightTextureSampler_ ## NAME; \
-	DIContext.CachedHeightTextureEnabled = CachedHeightTextureEnabled_ ## NAME; \
-	DIContext.CachedHeightTextureWorldToUvTransform = CachedHeightTextureWorldToUvTransform_ ## NAME; \
-	DIContext.CachedHeightTextureUvToWorldTransform = CachedHeightTextureUvToWorldTransform_ ## NAME; \
-	DIContext.CachedHeightTextureUvScaleBias = CachedHeightTextureUvScaleBias_ ## NAME; \
-	DIContext.CachedHeightTextureGridSize = CachedHeightTextureGridSize_ ## NAME; \
-	DIContext.PointClampedSampler = PointClampedSampler_ ## NAME; \
-	DIContext.CachedPhysMatTexture = CachedPhysMatTexture_ ## NAME; \
-	DIContext.CachedPhysMatTextureDimension = CachedPhysMatTextureDimension_ ## NAME; \
-	DIContext.SystemLWCTile = SystemLWCTile_ ## NAME;
+#undef VIRTUAL_TEXTURE_UNPACK_NORMAL_NONE
+#undef VIRTUAL_TEXTURE_UNPACK_NORMAL_BC3BC3
+#undef VIRTUAL_TEXTURE_UNPACK_NORMAL_BC5BC1
diff --git a/Engine/Plugins/FX/Niagara/Source/Niagara/Classes/NiagaraDataInterfaceLandscape.h b/Engine/Plugins/FX/Niagara/Source/Niagara/Classes/NiagaraDataInterfaceLandscape.h
--- a/Engine/Plugins/FX/Niagara/Source/Niagara/Classes/NiagaraDataInterfaceLandscape.h
+++ b/Engine/Plugins/FX/Niagara/Source/Niagara/Classes/NiagaraDataInterfaceLandscape.h
@@ -38,10 +38,8 @@ UCLASS(EditInlineNew, Category = "Landscape", meta = (DisplayName = "Landscape S
 class NIAGARA_API UNiagaraDataInterfaceLandscape : public UNiagaraDataInterface
 {
 	GENERATED_UCLASS_BODY()
-public:
-
-	DECLARE_NIAGARA_DI_PARAMETER();
 
+public:
 	UPROPERTY(EditAnywhere, Category = "Landscape")
 	TObjectPtr<AActor> SourceLandscape;
 
@@ -69,9 +67,11 @@ public:
 #if WITH_EDITORONLY_DATA
 	virtual void GetParameterDefinitionHLSL(const FNiagaraDataInterfaceGPUParamInfo& ParamInfo, FString& OutHLSL) override;
 	virtual bool GetFunctionHLSL(const FNiagaraDataInterfaceGPUParamInfo& ParamInfo, const FNiagaraDataInterfaceGeneratedFunction& FunctionInfo, int FunctionInstanceIndex, FString& OutHLSL) override;
-	virtual void GetCommonHLSL(FString& OutHLSL) override;
 	virtual bool AppendCompileHash(FNiagaraCompileHashVisitor* InVisitor) const override;
 #endif
+	virtual bool UseLegacyShaderBindings() const override { return false; }
+	virtual void BuildShaderParameters(FNiagaraShaderParametersBuilder& ShaderParametersBuilder) const override;
+	virtual void SetShaderParameters(const FNiagaraDataInterfaceSetShaderParametersContext& Context) const override;
 
 	virtual void ProvidePerInstanceDataForRenderThread(void* DataForRenderThread, void* PerInstanceData, const FNiagaraSystemInstanceID& SystemInstance) override;
 	virtual bool InitPerInstanceData(void* PerInstanceData, FNiagaraSystemInstance* SystemInstance) override;
@@ -89,4 +89,4 @@ protected:
 	static const FName GetHeightName;
 	static const FName GetWorldNormalName;
 	static const FName GetPhysicalMaterialIndexName;
-};
\ No newline at end of file
+};
diff --git a/Engine/Plugins/FX/Niagara/Source/Niagara/Private/NiagaraDataInterfaceLandscape.cpp b/Engine/Plugins/FX/Niagara/Source/Niagara/Private/NiagaraDataInterfaceLandscape.cpp
--- a/Engine/Plugins/FX/Niagara/Source/Niagara/Private/NiagaraDataInterfaceLandscape.cpp
+++ b/Engine/Plugins/FX/Niagara/Source/Niagara/Private/NiagaraDataInterfaceLandscape.cpp
@@ -2,20 +2,23 @@
 
 #include "NiagaraDataInterfaceLandscape.h"
 
+#include "NiagaraComponent.h"
+#include "NiagaraShader.h"
+#include "NiagaraShaderParametersBuilder.h"
+#include "NiagaraStats.h"
+#include "NiagaraWorldManager.h"
+
 #include "Algo/RemoveIf.h"
+#include "VT/RuntimeVirtualTexture.h"
 #include "EngineModule.h"
 #include "EngineUtils.h"
 #include "Landscape.h"
 #include "LandscapeHeightfieldCollisionComponent.h"
 #include "LandscapeInfo.h"
 #include "LandscapeProxy.h"
-#include "NiagaraComponent.h"
-#include "NiagaraShader.h"
-#include "NiagaraStats.h"
-#include "NiagaraWorldManager.h"
-#include "PhysicalMaterials/PhysicalMaterial.h"
+#include "ShaderCompilerCore.h"
 #include "ShaderParameterUtils.h"
-#include "VT/RuntimeVirtualTexture.h"
+#include "PhysicalMaterials/PhysicalMaterial.h"
 
 #define LOCTEXT_NAMESPACE "UNiagaraDataInterfaceLandscape"
 
@@ -38,6 +41,49 @@ namespace NiagaraDataInterfaceLandscape
 		VersionPlusOne,
 		LatestVersion = VersionPlusOne - 1
 	};
+
+	static const TCHAR* TemplateShaderFile = TEXT("/Plugin/FX/Niagara/Private/NiagaraDataInterfaceLandscape.ush");
+
+	BEGIN_SHADER_PARAMETER_STRUCT(FShaderParameters,)
+		SHADER_PARAMETER_SRV(Texture2D,				BaseColorVirtualTexture)
+		SHADER_PARAMETER_TEXTURE(Texture2D<uint4>,	BaseColorVirtualTexturePageTable)
+		SHADER_PARAMETER_SAMPLER(SamplerState,		BaseColorVirtualTextureSampler)
+		SHADER_PARAMETER(FMatrix44f,				BaseColorVirtualTextureWorldToUvTransform)
+		SHADER_PARAMETER(uint32,					BaseColorVirtualTextureSRGB)
+		SHADER_PARAMETER(uint32,					BaseColorVirtualTextureEnabled)
+		SHADER_PARAMETER(FUintVector4,				BaseColorVirtualTexturePackedUniform0)
+		SHADER_PARAMETER(FUintVector4,				BaseColorVirtualTexturePackedUniform1)
+		SHADER_PARAMETER(FUintVector4,				BaseColorVirtualTextureUniforms)
+		SHADER_PARAMETER_SRV(Texture2D,				HeightVirtualTexture)
+		SHADER_PARAMETER_TEXTURE(Texture2D<uint4>,	HeightVirtualTexturePageTable)
+		SHADER_PARAMETER_SAMPLER(SamplerState,		HeightVirtualTextureSampler)
+		SHADER_PARAMETER(FMatrix44f,				HeightVirtualTextureWorldToUvTransform)
+		SHADER_PARAMETER(uint32,					HeightVirtualTextureEnabled)
+		SHADER_PARAMETER(FUintVector4,				HeightVirtualTexturePackedUniform0)
+		SHADER_PARAMETER(FUintVector4,				HeightVirtualTexturePackedUniform1)
+		SHADER_PARAMETER(FUintVector4,				HeightVirtualTextureUniforms)
+		SHADER_PARAMETER_SRV(Texture2D,				NormalVirtualTexture0)
+		SHADER_PARAMETER_SRV(Texture2D,				NormalVirtualTexture1)
+		SHADER_PARAMETER_TEXTURE(Texture2D<uint4>,	NormalVirtualTexturePageTable)
+		SHADER_PARAMETER_SAMPLER(SamplerState,		NormalVirtualTextureSampler)
+		SHADER_PARAMETER(FMatrix44f,				NormalVirtualTextureWorldToUvTransform)
+		SHADER_PARAMETER(FUintVector4,				NormalVirtualTexturePackedUniform0)
+		SHADER_PARAMETER(FUintVector4,				NormalVirtualTexturePackedUniform1)
+		SHADER_PARAMETER(FUintVector4,				NormalVirtualTextureUniforms0)
+		SHADER_PARAMETER(FUintVector4,				NormalVirtualTextureUniforms1)
+		SHADER_PARAMETER(int32,						NormalVirtualTextureUnpackMode)
+		SHADER_PARAMETER(uint32,					NormalVirtualTextureEnabled)
+		SHADER_PARAMETER_SRV(Texture2D,				CachedHeightTexture)
+		SHADER_PARAMETER_SAMPLER(SamplerState,		CachedHeightTextureSampler)
+		SHADER_PARAMETER(uint32,					CachedHeightTextureEnabled)
+		SHADER_PARAMETER(FMatrix44f,				CachedHeightTextureWorldToUvTransform)
+		SHADER_PARAMETER(FMatrix44f,				CachedHeightTextureUvToWorldTransform)
+		SHADER_PARAMETER(FVector4f,					CachedHeightTextureUvScaleBias)
+		SHADER_PARAMETER(FVector2f,					CachedHeightTextureGridSize)
+		SHADER_PARAMETER_SAMPLER(SamplerState,		PointClampedSampler)
+		SHADER_PARAMETER_SRV(Texture2D<int>,		CachedPhysMatTexture)
+		SHADER_PARAMETER(FIntPoint,					CachedPhysMatTextureDimension)
+	END_SHADER_PARAMETER_STRUCT()
 }
 
 const FName UNiagaraDataInterfaceLandscape::GetBaseColorName(TEXT("GetBaseColor"));
@@ -199,6 +245,335 @@ struct FNDILandscapeData_RenderThread
 	FVector4 CachedHeightTextureUvScaleBias = FVector4(ForceInitToZero);
 	FVector2D CachedHeightTextureGridSize = FVector2D(1.0f, 1.0f);
 	ERuntimeVirtualTextureMaterialType NormalVirtualTextureMode = ERuntimeVirtualTextureMaterialType::Count;
+
+	enum class ENormalUnpackType : int32
+	{
+		None = 0,
+		BC3BC3,
+		BC5BC1,
+		B5G6R5
+	};
+
+	bool SetBaseColorVirtualTextureParameters(NiagaraDataInterfaceLandscape::FShaderParameters* ShaderParameters) const
+	{
+		const IAllocatedVirtualTexture* BaseColorAllocatedTexture = BaseColorVirtualTexture.IsValid()
+			? BaseColorVirtualTexture->GetAllocatedVirtualTexture()
+			: nullptr;
+
+		if (!BaseColorAllocatedTexture)
+		{
+			return false;
+		}
+
+		// todo - need to figure out a way to confirm that this is in fact the best/only option for the basecolor
+		constexpr uint32 BaseColorVirtualTextureLayerIndex = 0;
+		constexpr uint32 BaseColorVirtualTexturePageIndex = 0;
+		constexpr bool BaseColorVirtualTextureSrgb = false;
+
+		FRHIShaderResourceView* PhysicalTextureSrv = BaseColorAllocatedTexture->GetPhysicalTextureSRV(BaseColorVirtualTextureLayerIndex, BaseColorVirtualTextureSrgb);
+		if (!PhysicalTextureSrv)
+		{
+			return false;
+		}
+
+		FRHITexture* PageTableTexture = nullptr;
+
+		if (FRHITexture* PageTable = BaseColorAllocatedTexture->GetPageTableTexture(BaseColorVirtualTexturePageIndex))
+		{
+			if (FRHITextureReference* TextureReference = PageTable->GetTextureReference())
+			{
+				PageTableTexture = TextureReference->GetReferencedTexture();
+			}
+		}
+
+		if (!PageTableTexture)
+		{
+			return false;
+		}
+
+		FMatrix44f WorldToUvTransform(
+			FVector3f((FVector4f)BaseColorVirtualTextureWorldToUvParameters[0]),
+			FVector3f((FVector4f)BaseColorVirtualTextureWorldToUvParameters[1]),
+			FVector3f((FVector4f)BaseColorVirtualTextureWorldToUvParameters[2]),
+			FVector3f(0, 0, 0));
+
+		ShaderParameters->BaseColorVirtualTexture = PhysicalTextureSrv;
+		ShaderParameters->BaseColorVirtualTexturePageTable = PageTableTexture;
+		ShaderParameters->BaseColorVirtualTextureWorldToUvTransform = WorldToUvTransform;
+		ShaderParameters->BaseColorVirtualTextureSRGB = BaseColorVirtualTextureSRGB ? 1 : 0;
+		ShaderParameters->BaseColorVirtualTextureEnabled = 1;
+
+		FUintVector4 BaseColorVirtualTexturePageTableUniforms[2];
+		FUintVector4 BaseColorVirtualTextureUniforms;
+
+		BaseColorAllocatedTexture->GetPackedPageTableUniform(BaseColorVirtualTexturePageTableUniforms);
+		BaseColorAllocatedTexture->GetPackedUniform(&BaseColorVirtualTextureUniforms, BaseColorVirtualTextureLayerIndex);
+		ShaderParameters->BaseColorVirtualTexturePackedUniform0 = BaseColorVirtualTexturePageTableUniforms[0];
+		ShaderParameters->BaseColorVirtualTexturePackedUniform1 = BaseColorVirtualTexturePageTableUniforms[1];
+
+		ShaderParameters->BaseColorVirtualTextureUniforms = BaseColorVirtualTextureUniforms;
+
+		return true;
+	}
+
+	static void SetBaseColorVirtualTextureParameters_Default(NiagaraDataInterfaceLandscape::FShaderParameters* ShaderParameters)
+	{
+		const FUintVector4 DummyUint4(ForceInitToZero);
+
+		ShaderParameters->BaseColorVirtualTexture = GBlackTextureWithSRV->ShaderResourceViewRHI;
+		ShaderParameters->BaseColorVirtualTexturePageTable = GBlackTexture->TextureRHI;
+		ShaderParameters->BaseColorVirtualTextureWorldToUvTransform = FMatrix44f::Identity;
+		ShaderParameters->BaseColorVirtualTextureSRGB = 0;
+		ShaderParameters->BaseColorVirtualTextureEnabled = 0;
+		ShaderParameters->BaseColorVirtualTexturePackedUniform0 = DummyUint4;
+		ShaderParameters->BaseColorVirtualTexturePackedUniform1 = DummyUint4;
+		ShaderParameters->BaseColorVirtualTextureUniforms = DummyUint4;
+	}
+
+	bool SetHeightVirtualTextureParameters(NiagaraDataInterfaceLandscape::FShaderParameters* ShaderParameters) const
+	{
+		const IAllocatedVirtualTexture* HeightAllocatedTexture = HeightVirtualTexture.IsValid()
+			? HeightVirtualTexture->GetAllocatedVirtualTexture()
+			: nullptr;
+
+		if (!HeightAllocatedTexture)
+		{
+			return false;
+		}
+
+		// todo - need to figure out a way to confirm that this is in fact the best/only option for the height
+		constexpr uint32 HeightVirtualTextureLayerIndex = 0;
+		constexpr uint32 HeightVirtualTexturePageIndex = 0;
+		constexpr bool HeightVirtualTextureSrgb = false;
+
+		FRHIShaderResourceView* PhysicalTextureSrv = HeightAllocatedTexture->GetPhysicalTextureSRV(HeightVirtualTextureLayerIndex, HeightVirtualTextureSrgb);
+		if (!PhysicalTextureSrv)
+		{
+			return false;
+		}
+
+		FRHITexture* PageTableTexture = nullptr;
+
+		if (FRHITexture* PageTable = HeightAllocatedTexture->GetPageTableTexture(HeightVirtualTexturePageIndex))
+		{
+			if (FRHITextureReference* TextureReference = PageTable->GetTextureReference())
+			{
+				PageTableTexture = TextureReference->GetReferencedTexture();
+			}
+		}
+
+		if (!PageTableTexture)
+		{
+			return false;
+		}
+
+		FMatrix44f WorldToUvTransform(
+			FVector3f( (FVector4f)HeightVirtualTextureWorldToUvParameters[0]),
+			FVector3f( (FVector4f)HeightVirtualTextureWorldToUvParameters[1]),
+			FVector3f( (FVector4f)HeightVirtualTextureWorldToUvParameters[2]),
+			FVector3f( (FVector4f)HeightVirtualTextureWorldToUvParameters[3]));
+
+		ShaderParameters->HeightVirtualTexture = PhysicalTextureSrv;
+		ShaderParameters->HeightVirtualTexturePageTable = PageTableTexture;
+		ShaderParameters->HeightVirtualTextureWorldToUvTransform = WorldToUvTransform;
+		ShaderParameters->HeightVirtualTextureEnabled = 1;
+
+		FUintVector4 HeightVirtualTexturePageTableUniforms[2];
+		FUintVector4 HeightVirtualTextureUniforms;
+
+		HeightAllocatedTexture->GetPackedPageTableUniform(HeightVirtualTexturePageTableUniforms);
+		HeightAllocatedTexture->GetPackedUniform(&HeightVirtualTextureUniforms, HeightVirtualTextureLayerIndex);
+		ShaderParameters->HeightVirtualTexturePackedUniform0 = HeightVirtualTexturePageTableUniforms[0];
+		ShaderParameters->HeightVirtualTexturePackedUniform1 = HeightVirtualTexturePageTableUniforms[1];
+
+		ShaderParameters->HeightVirtualTextureUniforms = HeightVirtualTextureUniforms;
+
+		return true;
+	}
+
+	static void SetHeightVirtualTextureParameters_Default(NiagaraDataInterfaceLandscape::FShaderParameters* ShaderParameters)
+	{
+		const FUintVector4 DummyUint4(ForceInitToZero);
+
+		ShaderParameters->HeightVirtualTexture = GBlackTextureWithSRV->ShaderResourceViewRHI;
+		ShaderParameters->HeightVirtualTexturePageTable = GBlackTexture->TextureRHI;
+		ShaderParameters->HeightVirtualTextureWorldToUvTransform = FMatrix44f::Identity;
+		ShaderParameters->HeightVirtualTextureEnabled = 0;
+		ShaderParameters->HeightVirtualTexturePackedUniform0 = DummyUint4;
+		ShaderParameters->HeightVirtualTexturePackedUniform1 = DummyUint4;
+		ShaderParameters->HeightVirtualTextureUniforms = DummyUint4;
+	}
+
+	bool SetNormalVirtualTextureParameters(NiagaraDataInterfaceLandscape::FShaderParameters* ShaderParameters) const
+	{
+		const IAllocatedVirtualTexture* NormalAllocatedTexture = NormalVirtualTexture.IsValid()
+			? NormalVirtualTexture->GetAllocatedVirtualTexture()
+			: nullptr;
+
+		if (!NormalAllocatedTexture)
+		{
+			return false;
+		}
+
+		constexpr uint32 NormalVirtualTexturePageIndex = 0;
+
+		FRHITexture* PageTableTexture = nullptr;
+
+		if (FRHITexture* PageTable = NormalAllocatedTexture->GetPageTableTexture(NormalVirtualTexturePageIndex))
+		{
+			if (FRHITextureReference* TextureReference = PageTable->GetTextureReference())
+			{
+				PageTableTexture = TextureReference->GetReferencedTexture();
+			}
+		}
+
+		if (!PageTableTexture)
+		{
+			return false;
+		}
+
+		FRHIShaderResourceView* PhysicalTextureSrv[2] = { GBlackTextureWithSRV->ShaderResourceViewRHI, GBlackTextureWithSRV->ShaderResourceViewRHI };
+		FUintVector4 NormalVirtualTextureUniforms[2];
+
+		switch (NormalVirtualTextureMode)
+		{
+			case ERuntimeVirtualTextureMaterialType::BaseColor_Normal_Roughness:
+				PhysicalTextureSrv[0] = NormalAllocatedTexture->GetPhysicalTextureSRV(0, false);
+				NormalAllocatedTexture->GetPackedUniform(&NormalVirtualTextureUniforms[0], 0);
+
+				PhysicalTextureSrv[1] = NormalAllocatedTexture->GetPhysicalTextureSRV(1, false);
+				NormalAllocatedTexture->GetPackedUniform(&NormalVirtualTextureUniforms[1], 1);
+
+				ShaderParameters->NormalVirtualTextureUnpackMode = (int32)ENormalUnpackType::B5G6R5;
+				break;
+			case ERuntimeVirtualTextureMaterialType::BaseColor_Normal_Specular:
+				PhysicalTextureSrv[0] = NormalAllocatedTexture->GetPhysicalTextureSRV(0, false);
+				NormalAllocatedTexture->GetPackedUniform(&NormalVirtualTextureUniforms[0], 0);
+
+				PhysicalTextureSrv[1] = NormalAllocatedTexture->GetPhysicalTextureSRV(1, false);
+				NormalAllocatedTexture->GetPackedUniform(&NormalVirtualTextureUniforms[1], 1);
+
+				ShaderParameters->NormalVirtualTextureUnpackMode = (int32)ENormalUnpackType::BC3BC3;
+				break;
+
+			case ERuntimeVirtualTextureMaterialType::BaseColor_Normal_Specular_YCoCg:
+			case ERuntimeVirtualTextureMaterialType::BaseColor_Normal_Specular_Mask_YCoCg:
+				PhysicalTextureSrv[0] = NormalAllocatedTexture->GetPhysicalTextureSRV(1, false);
+				NormalAllocatedTexture->GetPackedUniform(&NormalVirtualTextureUniforms[0], 1);
+
+				PhysicalTextureSrv[1] = NormalAllocatedTexture->GetPhysicalTextureSRV(2, false);
+				NormalAllocatedTexture->GetPackedUniform(&NormalVirtualTextureUniforms[1], 2);
+
+				ShaderParameters->NormalVirtualTextureUnpackMode = (int32)ENormalUnpackType::BC5BC1;
+				break;
+
+			default:
+				return false;
+		}
+
+		FMatrix44f WorldToUvTransform(
+			FPlane4f( (FVector4f)NormalVirtualTextureWorldToUvParameters[0]),
+			FPlane4f( (FVector4f)NormalVirtualTextureWorldToUvParameters[1]),
+			FPlane4f( (FVector4f)NormalVirtualTextureWorldToUvParameters[2]),
+			FPlane4f(0.0f, 0.0f, 0.0f, 1.0f));
+
+		ShaderParameters->NormalVirtualTexture0 = PhysicalTextureSrv[0];
+		ShaderParameters->NormalVirtualTexture1 = PhysicalTextureSrv[1];
+		ShaderParameters->NormalVirtualTexturePageTable = PageTableTexture;
+		ShaderParameters->NormalVirtualTextureWorldToUvTransform = WorldToUvTransform;
+		ShaderParameters->NormalVirtualTextureEnabled = 1;
+
+		FUintVector4 NormalVirtualTexturePageTableUniforms[2];
+		NormalAllocatedTexture->GetPackedPageTableUniform(NormalVirtualTexturePageTableUniforms);
+		ShaderParameters->NormalVirtualTexturePackedUniform0 = NormalVirtualTexturePageTableUniforms[0];
+		ShaderParameters->NormalVirtualTexturePackedUniform1 = NormalVirtualTexturePageTableUniforms[1];
+
+		ShaderParameters->NormalVirtualTextureUniforms0 = NormalVirtualTextureUniforms[0];
+		ShaderParameters->NormalVirtualTextureUniforms1 = NormalVirtualTextureUniforms[1];
+
+		return true;
+	}
+
+	static void SetNormalVirtualTextureParameters_Default(NiagaraDataInterfaceLandscape::FShaderParameters* ShaderParameters)
+	{
+		const FUintVector4 DummyUint4(ForceInitToZero);
+
+		ShaderParameters->NormalVirtualTexture0 =  GBlackTextureWithSRV->ShaderResourceViewRHI;
+		ShaderParameters->NormalVirtualTexture1 = GBlackTextureWithSRV->ShaderResourceViewRHI;
+		ShaderParameters->NormalVirtualTexturePageTable = GBlackTexture->TextureRHI;
+		ShaderParameters->NormalVirtualTextureWorldToUvTransform = FMatrix44f::Identity;
+		ShaderParameters->NormalVirtualTextureEnabled = 0;
+		ShaderParameters->NormalVirtualTextureUnpackMode = (int32)ENormalUnpackType::None;
+		ShaderParameters->NormalVirtualTexturePackedUniform0 = DummyUint4;
+		ShaderParameters->NormalVirtualTexturePackedUniform1 = DummyUint4;
+		ShaderParameters->NormalVirtualTextureUniforms0 = DummyUint4;
+		ShaderParameters->NormalVirtualTextureUniforms1 = DummyUint4;
+	}
+
+	bool SetCachedHeightTextureParameters(NiagaraDataInterfaceLandscape::FShaderParameters* ShaderParameters) const
+	{
+		if (TextureResources)
+		{
+			FRHIShaderResourceView* HeightTextureSrv = TextureResources->GetHeightTexture();
+			FRHIShaderResourceView* PhysMatTextureSrv = TextureResources->GetPhysMatTexture();
+
+			if (HeightTextureSrv || PhysMatTextureSrv)
+			{
+				ShaderParameters->CachedHeightTextureWorldToUvTransform = (FMatrix44f)CachedHeightTextureWorldToUvTransform;
+				ShaderParameters->CachedHeightTextureUvToWorldTransform = (FMatrix44f)CachedHeightTextureUvToWorldTransform;
+				ShaderParameters->CachedHeightTextureUvScaleBias = (FVector4f)CachedHeightTextureUvScaleBias;
+				ShaderParameters->CachedHeightTextureGridSize = (FVector2f)CachedHeightTextureGridSize;
+
+				if (HeightTextureSrv)
+				{
+					ShaderParameters->CachedHeightTexture = HeightTextureSrv;
+					ShaderParameters->CachedHeightTextureEnabled = 1;
+				}
+				else
+				{
+					ShaderParameters->CachedHeightTexture = GBlackTextureWithSRV->ShaderResourceViewRHI;
+					ShaderParameters->CachedHeightTextureEnabled = 0;
+				}
+
+				if (PhysMatTextureSrv)
+				{
+					const FIntPoint PhysMatDimensions(TextureResources->GetDimensions());
+					ShaderParameters->CachedPhysMatTexture = PhysMatTextureSrv;
+					ShaderParameters->CachedPhysMatTextureDimension = PhysMatDimensions;
+				}
+				else
+				{
+					const FIntPoint PhysMatDimensions(ForceInitToZero);
+					ShaderParameters->CachedPhysMatTexture = GBlackTextureWithSRV->ShaderResourceViewRHI;
+					ShaderParameters->CachedPhysMatTextureDimension = PhysMatDimensions;
+				}
+
+				FRHISamplerState* PointClampedSampler = TStaticSamplerState<SF_Point, AM_Clamp, AM_Clamp, AM_Clamp>::GetRHI();
+				ShaderParameters->PointClampedSampler = PointClampedSampler;
+
+				return true;
+			}
+		}
+
+		return false;
+	}
+
+	static void SetCachedHeightTextureParameters_Defaults(NiagaraDataInterfaceLandscape::FShaderParameters* ShaderParameters)
+	{
+		const FVector4f DummyVector4(ForceInitToZero);
+
+		ShaderParameters->CachedHeightTexture = GBlackTextureWithSRV->ShaderResourceViewRHI;
+		ShaderParameters->CachedHeightTextureWorldToUvTransform = FMatrix44f::Identity;
+		ShaderParameters->CachedHeightTextureUvToWorldTransform = FMatrix44f::Identity;
+		ShaderParameters->CachedHeightTextureUvScaleBias = DummyVector4;
+		ShaderParameters->CachedHeightTextureGridSize = FVector2f::ZeroVector;
+		ShaderParameters->CachedHeightTextureEnabled = 0;
+
+		ShaderParameters->CachedPhysMatTexture = GBlackTextureWithSRV->ShaderResourceViewRHI;
+
+		FRHISamplerState* PointClampedSampler = TStaticSamplerState<SF_Point, AM_Clamp, AM_Clamp, AM_Clamp>::GetRHI();
+		ShaderParameters->PointClampedSampler = PointClampedSampler;
+	}
 };
 
 class FNDI_Landscape_GeneratedData : public FNDI_GeneratedData
@@ -886,55 +1261,35 @@ bool UNiagaraDataInterfaceLandscape::AppendCompileHash(FNiagaraCompileHashVisito
 		return false;
 	}
 
-	FSHAHash Hash = GetShaderFileHash((TEXT("/Plugin/FX/Niagara/Private/NiagaraDataInterfaceLandscape.ush")), EShaderPlatform::SP_PCD3D_SM5);
+	FSHAHash Hash = GetShaderFileHash(NiagaraDataInterfaceLandscape::TemplateShaderFile, EShaderPlatform::SP_PCD3D_SM5);
 	InVisitor->UpdateString(TEXT("NiagaraDataInterfaceLandscapeHLSLSource"), Hash.ToString());
+	InVisitor->UpdateShaderParameters<NiagaraDataInterfaceLandscape::FShaderParameters>();
 
 	return true;
 
 }
 
-void UNiagaraDataInterfaceLandscape::GetCommonHLSL(FString& OutHLSL)
-{
-	OutHLSL += TEXT("#include \"/Plugin/FX/Niagara/Private/NiagaraDataInterfaceLandscape.ush\"\n");
-}
-
 void UNiagaraDataInterfaceLandscape::GetParameterDefinitionHLSL(const FNiagaraDataInterfaceGPUParamInfo& ParamInfo, FString& OutHLSL)
 {
 	Super::GetParameterDefinitionHLSL(ParamInfo, OutHLSL);
 
-	OutHLSL += TEXT("NDILANDSCAPE_DECLARE_CONSTANTS(") + ParamInfo.DataInterfaceHLSLSymbol + TEXT(")\n");
+	TMap<FString, FStringFormatArg> TemplateArgs =
+	{
+		{TEXT("ParameterName"),	ParamInfo.DataInterfaceHLSLSymbol},
+	};
+
+	FString TemplateFile;
+	LoadShaderSourceFile(NiagaraDataInterfaceLandscape::TemplateShaderFile, EShaderPlatform::SP_PCD3D_SM5, &TemplateFile, nullptr);
+	OutHLSL += FString::Format(*TemplateFile, TemplateArgs);
 }
 
 bool UNiagaraDataInterfaceLandscape::GetFunctionHLSL(const FNiagaraDataInterfaceGPUParamInfo& ParamInfo, const FNiagaraDataInterfaceGeneratedFunction& FunctionInfo, int FunctionInstanceIndex, FString& OutHLSL)
 {
-	TMap<FString, FStringFormatArg> ArgsSample =
+	if ( (FunctionInfo.DefinitionName == GetBaseColorName) ||
+		 (FunctionInfo.DefinitionName == GetHeightName) ||
+		 (FunctionInfo.DefinitionName == GetWorldNormalName) ||
+		 (FunctionInfo.DefinitionName == GetPhysicalMaterialIndexName) )
 	{
-		{TEXT("InstanceFunctionName"), FunctionInfo.InstanceName},
-		{TEXT("NDIGetContextName"), TEXT("NDILANDSCAPE_MAKE_CONTEXT(") + ParamInfo.DataInterfaceHLSLSymbol + TEXT(")")},
-	};
-
-	if (FunctionInfo.DefinitionName == GetBaseColorName)
-	{
-		static const TCHAR* FormatSample = TEXT("void {InstanceFunctionName}(in float3 InWorldPos, out float3 OutBaseColor, out bool OutIsValid) { {NDIGetContextName} NDILandscape_GetBaseColor(DIContext, InWorldPos, OutBaseColor, OutIsValid); }\n");
-		OutHLSL += FString::Format(FormatSample, ArgsSample);
-		return true;
-	}
-	else if (FunctionInfo.DefinitionName == GetHeightName)
-	{
-		static const TCHAR* FormatSample = TEXT("void {InstanceFunctionName}(in float3 InWorldPos, out float OutHeight, out bool OutIsValid) { {NDIGetContextName} NDILandscape_GetHeight(DIContext, InWorldPos, OutHeight, OutIsValid); }\n");
-		OutHLSL += FString::Format(FormatSample, ArgsSample);
-		return true;
-	}
-	else if (FunctionInfo.DefinitionName == GetWorldNormalName)
-	{
-		static const TCHAR* FormatSample = TEXT("void {InstanceFunctionName}(in float3 InWorldPos, out float3 OutWorldNormal, out bool OutIsValid) { {NDIGetContextName} NDILandscape_GetWorldNormal(DIContext, InWorldPos, OutWorldNormal, OutIsValid); }\n");
-		OutHLSL += FString::Format(FormatSample, ArgsSample);
-		return true;
-	}
-	else if (FunctionInfo.DefinitionName == GetPhysicalMaterialIndexName)
-	{
-		static const TCHAR* FormatSample = TEXT("void {InstanceFunctionName}(in float3 InWorldPos, out int OutIndex, out bool OutIsValid) { {NDIGetContextName} NDILandscape_GetPhysicalMaterialIndex(DIContext, InWorldPos, OutIndex, OutIsValid); }\n");
-		OutHLSL += FString::Format(FormatSample, ArgsSample);
 		return true;
 	}
 
@@ -942,10 +1297,64 @@ bool UNiagaraDataInterfaceLandscape::GetFunctionHLSL(const FNiagaraDataInterface
 }
 #endif
 
-bool UNiagaraDataInterfaceLandscape::InitPerInstanceData(void* PerInstanceData, FNiagaraSystemInstance* SystemInstance)
+void UNiagaraDataInterfaceLandscape::BuildShaderParameters(FNiagaraShaderParametersBuilder& ShaderParametersBuilder) const
 {
-	FNDILandscapeData_GameThread* InstanceData = new(PerInstanceData) FNDILandscapeData_GameThread();
-	ApplyLandscape(*SystemInstance, *InstanceData);
+	ShaderParametersBuilder.AddNestedStruct<NiagaraDataInterfaceLandscape::FShaderParameters>();
+}
+
+void UNiagaraDataInterfaceLandscape::SetShaderParameters(const FNiagaraDataInterfaceSetShaderParametersContext& Context) const
+{
+	FNiagaraDataInterfaceProxyLandscape& RT_Proxy = Context.GetProxy<FNiagaraDataInterfaceProxyLandscape>();
+	FNDILandscapeData_RenderThread* ProxyData = RT_Proxy.SystemInstancesToProxyData_RT.Find(Context.GetSystemInstanceID());
+
+	NiagaraDataInterfaceLandscape::FShaderParameters* ShaderParameters = Context.GetParameterNestedStruct<NiagaraDataInterfaceLandscape::FShaderParameters>();
+
+	// Set Samplers
+	FRHISamplerState* BilinearSamplerState				= TStaticSamplerState<SF_Bilinear, AM_Clamp, AM_Clamp, AM_Clamp>::GetRHI();
+	ShaderParameters->BaseColorVirtualTextureSampler	= BilinearSamplerState;
+	ShaderParameters->HeightVirtualTextureSampler		= BilinearSamplerState;
+	ShaderParameters->NormalVirtualTextureSampler		= BilinearSamplerState;
+	ShaderParameters->CachedHeightTextureSampler		= BilinearSamplerState;
+
+	// Set Textures
+	bool ApplyBaseColorVirtualTextureDefaults = true;
+	bool ApplyHeightVirtualTextureDefaults = true;
+	bool ApplyNormalVirtualTextureDefaults = true;
+	bool ApplyCachedHeightTextureDefaults = true;
+
+	if (ProxyData)
+	{
+		ApplyBaseColorVirtualTextureDefaults = ProxyData->SetBaseColorVirtualTextureParameters(ShaderParameters) != true;
+		ApplyHeightVirtualTextureDefaults = ProxyData->SetHeightVirtualTextureParameters(ShaderParameters) != true;
+		ApplyNormalVirtualTextureDefaults = ProxyData->SetNormalVirtualTextureParameters(ShaderParameters) != true;
+		ApplyCachedHeightTextureDefaults = ProxyData->SetCachedHeightTextureParameters(ShaderParameters) != true;
+	}
+
+	if (ApplyBaseColorVirtualTextureDefaults)
+	{
+		FNDILandscapeData_RenderThread::SetBaseColorVirtualTextureParameters_Default(ShaderParameters);
+	}
+
+	if (ApplyHeightVirtualTextureDefaults)
+	{
+		FNDILandscapeData_RenderThread::SetHeightVirtualTextureParameters_Default(ShaderParameters);
+	}
+
+	if (ApplyNormalVirtualTextureDefaults)
+	{
+		FNDILandscapeData_RenderThread::SetNormalVirtualTextureParameters_Default(ShaderParameters);
+	}
+
+	if (ApplyCachedHeightTextureDefaults)
+	{
+		FNDILandscapeData_RenderThread::SetCachedHeightTextureParameters_Defaults(ShaderParameters);
+	}
+}
+
+bool UNiagaraDataInterfaceLandscape::InitPerInstanceData(void* PerInstanceData, FNiagaraSystemInstance* SystemInstance)
+{
+	FNDILandscapeData_GameThread* InstanceData = new(PerInstanceData) FNDILandscapeData_GameThread();
+	ApplyLandscape(*SystemInstance, *InstanceData);
 
 	FNiagaraDataInterfaceProxyLandscape* RT_Proxy = GetProxyAs<FNiagaraDataInterfaceProxyLandscape>();
 	ENQUEUE_RENDER_COMMAND(FNiagaraDICreateProxy) (
@@ -1154,602 +1563,4 @@ ALandscape* UNiagaraDataInterfaceLandscape::GetLandscape(const FNiagaraSystemIns
 	return nullptr;
 }
 
-struct FNiagaraDataInterfaceParametersCS_Landscape : public FNiagaraDataInterfaceParametersCS
-{
-	DECLARE_TYPE_LAYOUT(FNiagaraDataInterfaceParametersCS_Landscape, NonVirtual);
-
-private:
-	enum class ENormalUnpackType : int32
-	{
-		None = 0,
-		BC3BC3,
-		BC5BC1,
-		B5G6R5
-	};
-
-public:
-	void Bind(const FNiagaraDataInterfaceGPUParamInfo& ParameterInfo, const class FShaderParameterMap& ParameterMap)
-	{
-		BaseColorVirtualTextureParam.Bind(ParameterMap, *(BaseColorVirtualTextureName + ParameterInfo.DataInterfaceHLSLSymbol));
-		BaseColorVirtualTexturePageTableParam.Bind(ParameterMap, *(BaseColorVirtualTexturePageTableName + ParameterInfo.DataInterfaceHLSLSymbol));
-		BaseColorVirtualTextureSamplerParam.Bind(ParameterMap, *(BaseColorVirtualTextureSamplerName + ParameterInfo.DataInterfaceHLSLSymbol));
-		BaseColorVirtualTextureWorldToUvTransformParam.Bind(ParameterMap, *(BaseColorVirtualTextureWorldToUvTransformName + ParameterInfo.DataInterfaceHLSLSymbol));
-		BaseColorVirtualTextureSRGBParam.Bind(ParameterMap, *(BaseColorVirtualTextureSRGBName + ParameterInfo.DataInterfaceHLSLSymbol));
-		BaseColorVirtualTextureEnabledParam.Bind(ParameterMap, *(BaseColorVirtualTextureEnabledName + ParameterInfo.DataInterfaceHLSLSymbol));
-		BaseColorVirtualTexturePageTableUniform0Param.Bind(ParameterMap, *(BaseColorVirtualTexturePageTableUniform0Name + ParameterInfo.DataInterfaceHLSLSymbol));
-		BaseColorVirtualTexturePageTableUniform1Param.Bind(ParameterMap, *(BaseColorVirtualTexturePageTableUniform1Name + ParameterInfo.DataInterfaceHLSLSymbol));
-		BaseColorVirtualTextureUniformsParam.Bind(ParameterMap, *(BaseColorVirtualTextureUniformsName + ParameterInfo.DataInterfaceHLSLSymbol));
-
-		HeightVirtualTextureParam.Bind(ParameterMap, *(HeightVirtualTextureName + ParameterInfo.DataInterfaceHLSLSymbol));
-		HeightVirtualTexturePageTableParam.Bind(ParameterMap, *(HeightVirtualTexturePageTableName + ParameterInfo.DataInterfaceHLSLSymbol));
-		HeightVirtualTextureSamplerParam.Bind(ParameterMap, *(HeightVirtualTextureSamplerName + ParameterInfo.DataInterfaceHLSLSymbol));
-		HeightVirtualTextureWorldToUvTransformParam.Bind(ParameterMap, *(HeightVirtualTextureWorldToUvTransformName + ParameterInfo.DataInterfaceHLSLSymbol));
-		HeightVirtualTextureEnabledParam.Bind(ParameterMap, *(HeightVirtualTextureEnabledName + ParameterInfo.DataInterfaceHLSLSymbol));
-		HeightVirtualTexturePageTableUniform0Param.Bind(ParameterMap, *(HeightVirtualTexturePageTableUniform0Name + ParameterInfo.DataInterfaceHLSLSymbol));
-		HeightVirtualTexturePageTableUniform1Param.Bind(ParameterMap, *(HeightVirtualTexturePageTableUniform1Name + ParameterInfo.DataInterfaceHLSLSymbol));
-		HeightVirtualTextureUniformsParam.Bind(ParameterMap, *(HeightVirtualTextureUniformsName + ParameterInfo.DataInterfaceHLSLSymbol));
-
-		NormalVirtualTexture0Param.Bind(ParameterMap, *(NormalVirtualTexture0Name + ParameterInfo.DataInterfaceHLSLSymbol));
-		NormalVirtualTexture1Param.Bind(ParameterMap, *(NormalVirtualTexture1Name + ParameterInfo.DataInterfaceHLSLSymbol));
-		NormalVirtualTexturePageTableParam.Bind(ParameterMap, *(NormalVirtualTexturePageTableName + ParameterInfo.DataInterfaceHLSLSymbol));
-		NormalVirtualTextureSamplerParam.Bind(ParameterMap, *(NormalVirtualTextureSamplerName + ParameterInfo.DataInterfaceHLSLSymbol));
-		NormalVirtualTextureWorldToUvTransformParam.Bind(ParameterMap, *(NormalVirtualTextureWorldToUvTransformName + ParameterInfo.DataInterfaceHLSLSymbol));
-		NormalVirtualTextureEnabledParam.Bind(ParameterMap, *(NormalVirtualTextureEnabledName + ParameterInfo.DataInterfaceHLSLSymbol));
-		NormalVirtualTextureUnpackModeParam.Bind(ParameterMap, *(NormalVirtualTextureUnpackModeName + ParameterInfo.DataInterfaceHLSLSymbol));
-		NormalVirtualTexturePageTableUniform0Param.Bind(ParameterMap, *(NormalVirtualTexturePageTableUniform0Name + ParameterInfo.DataInterfaceHLSLSymbol));
-		NormalVirtualTexturePageTableUniform1Param.Bind(ParameterMap, *(NormalVirtualTexturePageTableUniform1Name + ParameterInfo.DataInterfaceHLSLSymbol));
-		NormalVirtualTextureUniforms0Param.Bind(ParameterMap, *(NormalVirtualTextureUniforms0Name + ParameterInfo.DataInterfaceHLSLSymbol));
-		NormalVirtualTextureUniforms1Param.Bind(ParameterMap, *(NormalVirtualTextureUniforms1Name + ParameterInfo.DataInterfaceHLSLSymbol));
-
-		CachedHeightTextureParam.Bind(ParameterMap, *(CachedHeightTextureName + ParameterInfo.DataInterfaceHLSLSymbol));
-		CachedHeightTextureSamplerParam.Bind(ParameterMap, *(CachedHeightTextureSamplerName + ParameterInfo.DataInterfaceHLSLSymbol));
-		CachedHeightTextureWorldToUvTransformParam.Bind(ParameterMap, *(CachedHeightTextureWorldToUvTransformName + ParameterInfo.DataInterfaceHLSLSymbol));
-		CachedHeightTextureUvToWorldTransformParam.Bind(ParameterMap, *(CachedHeightTextureUvToWorldTransformName + ParameterInfo.DataInterfaceHLSLSymbol));
-		CachedHeightTextureUvScaleBiasParam.Bind(ParameterMap, *(CachedHeightTextureUvScaleBiasName + ParameterInfo.DataInterfaceHLSLSymbol));
-		CachedHeightTexturEnabledParam.Bind(ParameterMap, *(CachedHeightTextureEnabledName + ParameterInfo.DataInterfaceHLSLSymbol));
-		CachedHeightTextureGridSizeParam.Bind(ParameterMap, *(CachedHeightTextureGridSizeName + ParameterInfo.DataInterfaceHLSLSymbol));
-		PointClampedSamplerParam.Bind(ParameterMap, *(PointClampedSamplerName + ParameterInfo.DataInterfaceHLSLSymbol));
-		CachedPhysMatTextureParam.Bind(ParameterMap, *(CachedPhysMatTextureName + ParameterInfo.DataInterfaceHLSLSymbol));
-		CachedPhysMatTextureDimensionParam.Bind(ParameterMap, *(CachedPhysMatTextureDimensionName + ParameterInfo.DataInterfaceHLSLSymbol));
-		SystemLWCTileParam.Bind(ParameterMap, *(SystemLWCTileName + ParameterInfo.DataInterfaceHLSLSymbol));
-	}
-
-	bool SetBaseColorVirtualTextureParameters(FRHICommandList& RHICmdList, FRHIComputeShader* ComputeShaderRHI, const FNDILandscapeData_RenderThread& ProxyData) const
-	{
-		const IAllocatedVirtualTexture* BaseColorAllocatedTexture = ProxyData.BaseColorVirtualTexture.IsValid()
-			? ProxyData.BaseColorVirtualTexture->GetAllocatedVirtualTexture()
-			: nullptr;
-
-		if (!BaseColorAllocatedTexture)
-		{
-			return false;
-		}
-
-		// todo - need to figure out a way to confirm that this is in fact the best/only option for the basecolor
-		constexpr uint32 BaseColorVirtualTextureLayerIndex = 0;
-		constexpr uint32 BaseColorVirtualTexturePageIndex = 0;
-		constexpr bool BaseColorVirtualTextureSrgb = false;
-
-		FRHIShaderResourceView* PhysicalTextureSrv = BaseColorAllocatedTexture->GetPhysicalTextureSRV(BaseColorVirtualTextureLayerIndex, BaseColorVirtualTextureSrgb);
-		if (!PhysicalTextureSrv)
-		{
-			return false;
-		}
-
-		FRHITexture* PageTableTexture = nullptr;
-
-		if (FRHITexture* PageTable = BaseColorAllocatedTexture->GetPageTableTexture(BaseColorVirtualTexturePageIndex))
-		{
-			if (FRHITextureReference* TextureReference = PageTable->GetTextureReference())
-			{
-				PageTableTexture = TextureReference->GetReferencedTexture();
-			}
-		}
-
-		if (!PageTableTexture)
-		{
-			return false;
-		}
-
-		FMatrix44f WorldToUvTransform(
-			FVector3f( (FVector4f)ProxyData.BaseColorVirtualTextureWorldToUvParameters[0]),
-			FVector3f( (FVector4f)ProxyData.BaseColorVirtualTextureWorldToUvParameters[1]),
-			FVector3f( (FVector4f)ProxyData.BaseColorVirtualTextureWorldToUvParameters[2]),
-			FVector3f(0, 0, 0));
-
-		SetSRVParameter(RHICmdList, ComputeShaderRHI, BaseColorVirtualTextureParam, PhysicalTextureSrv);
-		SetTextureParameter(RHICmdList, ComputeShaderRHI, BaseColorVirtualTexturePageTableParam, PageTableTexture);
-		SetShaderValue(RHICmdList, ComputeShaderRHI, BaseColorVirtualTextureWorldToUvTransformParam, WorldToUvTransform);
-		SetShaderValue(RHICmdList, ComputeShaderRHI, BaseColorVirtualTextureSRGBParam, ProxyData.BaseColorVirtualTextureSRGB ? 1 : 0);
-		SetShaderValue(RHICmdList, ComputeShaderRHI, BaseColorVirtualTextureEnabledParam, 1 /* true */);
-
-		FUintVector4 BaseColorVirtualTexturePageTableUniforms[2];
-		FUintVector4 BaseColorVirtualTextureUniforms;
-
-		BaseColorAllocatedTexture->GetPackedPageTableUniform(BaseColorVirtualTexturePageTableUniforms);
-		BaseColorAllocatedTexture->GetPackedUniform(&BaseColorVirtualTextureUniforms, BaseColorVirtualTextureLayerIndex);
-		SetShaderValue(RHICmdList, ComputeShaderRHI, BaseColorVirtualTexturePageTableUniform0Param, BaseColorVirtualTexturePageTableUniforms[0]);
-		SetShaderValue(RHICmdList, ComputeShaderRHI, BaseColorVirtualTexturePageTableUniform1Param, BaseColorVirtualTexturePageTableUniforms[1]);
-
-		SetShaderValue(RHICmdList, ComputeShaderRHI, BaseColorVirtualTextureUniformsParam, BaseColorVirtualTextureUniforms);
-
-		return true;
-	}
-
-	void SetBaseColorVirtualTextureParameters_Default(FRHICommandList& RHICmdList, FRHIComputeShader* ComputeShaderRHI) const
-	{
-		FUintVector4 DummyUint4(ForceInitToZero);
-
-		SetSRVParameter(RHICmdList, ComputeShaderRHI, BaseColorVirtualTextureParam, GBlackTextureWithSRV->ShaderResourceViewRHI);
-		SetTextureParameter(RHICmdList, ComputeShaderRHI, BaseColorVirtualTexturePageTableParam, GBlackTexture->TextureRHI);
-		SetShaderValue(RHICmdList, ComputeShaderRHI, BaseColorVirtualTextureWorldToUvTransformParam, FMatrix44f::Identity);
-		SetShaderValue(RHICmdList, ComputeShaderRHI, BaseColorVirtualTextureSRGBParam, 0 /* false */);
-		SetShaderValue(RHICmdList, ComputeShaderRHI, BaseColorVirtualTextureEnabledParam, 0 /* false */);
-		SetShaderValue(RHICmdList, ComputeShaderRHI, BaseColorVirtualTexturePageTableUniform0Param, DummyUint4);
-		SetShaderValue(RHICmdList, ComputeShaderRHI, BaseColorVirtualTexturePageTableUniform1Param, DummyUint4);
-		SetShaderValue(RHICmdList, ComputeShaderRHI, BaseColorVirtualTextureUniformsParam, DummyUint4);
-	}
-
-	bool SetHeightVirtualTextureParameters(FRHICommandList& RHICmdList, FRHIComputeShader* ComputeShaderRHI, const FNDILandscapeData_RenderThread& ProxyData) const
-	{
-		const IAllocatedVirtualTexture* HeightAllocatedTexture = ProxyData.HeightVirtualTexture.IsValid()
-			? ProxyData.HeightVirtualTexture->GetAllocatedVirtualTexture()
-			: nullptr;
-
-		if (!HeightAllocatedTexture)
-		{
-			return false;
-		}
-
-		// todo - need to figure out a way to confirm that this is in fact the best/only option for the height
-		constexpr uint32 HeightVirtualTextureLayerIndex = 0;
-		constexpr uint32 HeightVirtualTexturePageIndex = 0;
-		constexpr bool HeightVirtualTextureSrgb = false;
-
-		FRHIShaderResourceView* PhysicalTextureSrv = HeightAllocatedTexture->GetPhysicalTextureSRV(HeightVirtualTextureLayerIndex, HeightVirtualTextureSrgb);
-		if (!PhysicalTextureSrv)
-		{
-			return false;
-		}
-
-		FRHITexture* PageTableTexture = nullptr;
-
-		if (FRHITexture* PageTable = HeightAllocatedTexture->GetPageTableTexture(HeightVirtualTexturePageIndex))
-		{
-			if (FRHITextureReference* TextureReference = PageTable->GetTextureReference())
-			{
-				PageTableTexture = TextureReference->GetReferencedTexture();
-			}
-		}
-
-		if (!PageTableTexture)
-		{
-			return false;
-		}
-
-		FMatrix44f WorldToUvTransform(
-			FVector3f( (FVector4f)ProxyData.HeightVirtualTextureWorldToUvParameters[0]),
-			FVector3f( (FVector4f)ProxyData.HeightVirtualTextureWorldToUvParameters[1]),
-			FVector3f( (FVector4f)ProxyData.HeightVirtualTextureWorldToUvParameters[2]),
-			FVector3f( (FVector4f)ProxyData.HeightVirtualTextureWorldToUvParameters[3]));
-
-		SetSRVParameter(RHICmdList, ComputeShaderRHI, HeightVirtualTextureParam, PhysicalTextureSrv);
-		SetTextureParameter(RHICmdList, ComputeShaderRHI, HeightVirtualTexturePageTableParam, PageTableTexture);
-		SetShaderValue(RHICmdList, ComputeShaderRHI, HeightVirtualTextureWorldToUvTransformParam, WorldToUvTransform);
-		SetShaderValue(RHICmdList, ComputeShaderRHI, HeightVirtualTextureEnabledParam, 1 /* true */);
-
-		FUintVector4 HeightVirtualTexturePageTableUniforms[2];
-		FUintVector4 HeightVirtualTextureUniforms;
-
-		HeightAllocatedTexture->GetPackedPageTableUniform(HeightVirtualTexturePageTableUniforms);
-		HeightAllocatedTexture->GetPackedUniform(&HeightVirtualTextureUniforms, HeightVirtualTextureLayerIndex);
-		SetShaderValue(RHICmdList, ComputeShaderRHI, HeightVirtualTexturePageTableUniform0Param, HeightVirtualTexturePageTableUniforms[0]);
-		SetShaderValue(RHICmdList, ComputeShaderRHI, HeightVirtualTexturePageTableUniform1Param, HeightVirtualTexturePageTableUniforms[1]);
-
-		SetShaderValue(RHICmdList, ComputeShaderRHI, HeightVirtualTextureUniformsParam, HeightVirtualTextureUniforms);
-
-		return true;
-	}
-
-	void SetHeightVirtualTextureParameters_Default(FRHICommandList& RHICmdList, FRHIComputeShader* ComputeShaderRHI) const
-	{
-		FUintVector4 DummyUint4(ForceInitToZero);
-
-		SetSRVParameter(RHICmdList, ComputeShaderRHI, HeightVirtualTextureParam, GBlackTextureWithSRV->ShaderResourceViewRHI);
-		SetTextureParameter(RHICmdList, ComputeShaderRHI, HeightVirtualTexturePageTableParam, GBlackTexture->TextureRHI);
-		SetShaderValue(RHICmdList, ComputeShaderRHI, HeightVirtualTextureWorldToUvTransformParam, FMatrix44f::Identity);
-		SetShaderValue(RHICmdList, ComputeShaderRHI, HeightVirtualTextureEnabledParam, 0 /* false */);
-		SetShaderValue(RHICmdList, ComputeShaderRHI, HeightVirtualTexturePageTableUniform0Param, DummyUint4);
-		SetShaderValue(RHICmdList, ComputeShaderRHI, HeightVirtualTexturePageTableUniform1Param, DummyUint4);
-		SetShaderValue(RHICmdList, ComputeShaderRHI, HeightVirtualTextureUniformsParam, DummyUint4);
-	}
-
-	bool SetNormalVirtualTextureParameters(FRHICommandList& RHICmdList, FRHIComputeShader* ComputeShaderRHI, const FNDILandscapeData_RenderThread& ProxyData) const
-	{
-		const IAllocatedVirtualTexture* NormalAllocatedTexture = ProxyData.NormalVirtualTexture.IsValid()
-			? ProxyData.NormalVirtualTexture->GetAllocatedVirtualTexture()
-			: nullptr;
-
-		if (!NormalAllocatedTexture)
-		{
-			return false;
-		}
-
-		constexpr uint32 NormalVirtualTexturePageIndex = 0;
-
-		FRHITexture* PageTableTexture = nullptr;
-
-		if (FRHITexture* PageTable = NormalAllocatedTexture->GetPageTableTexture(NormalVirtualTexturePageIndex))
-		{
-			if (FRHITextureReference* TextureReference = PageTable->GetTextureReference())
-			{
-				PageTableTexture = TextureReference->GetReferencedTexture();
-			}
-		}
-
-		if (!PageTableTexture)
-		{
-			return false;
-		}
-
-		FRHIShaderResourceView* PhysicalTextureSrv[2] = { GBlackTextureWithSRV->ShaderResourceViewRHI, GBlackTextureWithSRV->ShaderResourceViewRHI };
-		FUintVector4 NormalVirtualTextureUniforms[2];
-
-		switch (ProxyData.NormalVirtualTextureMode)
-		{
-			case ERuntimeVirtualTextureMaterialType::BaseColor_Normal_Roughness:
-				PhysicalTextureSrv[0] = NormalAllocatedTexture->GetPhysicalTextureSRV(0, false);
-				NormalAllocatedTexture->GetPackedUniform(&NormalVirtualTextureUniforms[0], 0);
-
-				PhysicalTextureSrv[1] = NormalAllocatedTexture->GetPhysicalTextureSRV(1, false);
-				NormalAllocatedTexture->GetPackedUniform(&NormalVirtualTextureUniforms[1], 1);
-
-				SetShaderValue(RHICmdList, ComputeShaderRHI, NormalVirtualTextureUnpackModeParam, ENormalUnpackType::B5G6R5);
-				break;
-			case ERuntimeVirtualTextureMaterialType::BaseColor_Normal_Specular:
-				PhysicalTextureSrv[0] = NormalAllocatedTexture->GetPhysicalTextureSRV(0, false);
-				NormalAllocatedTexture->GetPackedUniform(&NormalVirtualTextureUniforms[0], 0);
-
-				PhysicalTextureSrv[1] = NormalAllocatedTexture->GetPhysicalTextureSRV(1, false);
-				NormalAllocatedTexture->GetPackedUniform(&NormalVirtualTextureUniforms[1], 1);
-
-				SetShaderValue(RHICmdList, ComputeShaderRHI, NormalVirtualTextureUnpackModeParam, ENormalUnpackType::BC3BC3);
-				break;
-
-			case ERuntimeVirtualTextureMaterialType::BaseColor_Normal_Specular_YCoCg:
-			case ERuntimeVirtualTextureMaterialType::BaseColor_Normal_Specular_Mask_YCoCg:
-				PhysicalTextureSrv[0] = NormalAllocatedTexture->GetPhysicalTextureSRV(1, false);
-				NormalAllocatedTexture->GetPackedUniform(&NormalVirtualTextureUniforms[0], 1);
-
-				PhysicalTextureSrv[1] = NormalAllocatedTexture->GetPhysicalTextureSRV(2, false);
-				NormalAllocatedTexture->GetPackedUniform(&NormalVirtualTextureUniforms[1], 2);
-
-				SetShaderValue(RHICmdList, ComputeShaderRHI, NormalVirtualTextureUnpackModeParam, ENormalUnpackType::BC5BC1);
-				break;
-
-			default:
-				return false;
-		}
-
-		FMatrix44f WorldToUvTransform(
-			FPlane4f( (FVector4f)ProxyData.NormalVirtualTextureWorldToUvParameters[0]),
-			FPlane4f( (FVector4f)ProxyData.NormalVirtualTextureWorldToUvParameters[1]),
-			FPlane4f( (FVector4f)ProxyData.NormalVirtualTextureWorldToUvParameters[2]),
-			FPlane4f(0.0f, 0.0f, 0.0f, 1.0f));
-
-		SetSRVParameter(RHICmdList, ComputeShaderRHI, NormalVirtualTexture0Param, PhysicalTextureSrv[0]);
-		SetSRVParameter(RHICmdList, ComputeShaderRHI, NormalVirtualTexture1Param, PhysicalTextureSrv[1]);
-		SetTextureParameter(RHICmdList, ComputeShaderRHI, NormalVirtualTexturePageTableParam, PageTableTexture);
-		SetShaderValue(RHICmdList, ComputeShaderRHI, NormalVirtualTextureWorldToUvTransformParam, WorldToUvTransform);
-		SetShaderValue(RHICmdList, ComputeShaderRHI, NormalVirtualTextureEnabledParam, 1 /* true */);
-
-		FUintVector4 NormalVirtualTexturePageTableUniforms[2];
-		NormalAllocatedTexture->GetPackedPageTableUniform(NormalVirtualTexturePageTableUniforms);
-		SetShaderValue(RHICmdList, ComputeShaderRHI, NormalVirtualTexturePageTableUniform0Param, NormalVirtualTexturePageTableUniforms[0]);
-		SetShaderValue(RHICmdList, ComputeShaderRHI, NormalVirtualTexturePageTableUniform1Param, NormalVirtualTexturePageTableUniforms[1]);
-
-		SetShaderValue(RHICmdList, ComputeShaderRHI, NormalVirtualTextureUniforms0Param, NormalVirtualTextureUniforms[0]);
-		SetShaderValue(RHICmdList, ComputeShaderRHI, NormalVirtualTextureUniforms1Param, NormalVirtualTextureUniforms[1]);
-
-		return true;
-	}
-
-	void SetNormalVirtualTextureParameters_Default(FRHICommandList& RHICmdList, FRHIComputeShader* ComputeShaderRHI) const
-	{
-		FUintVector4 DummyUint4(ForceInitToZero);
-
-		SetSRVParameter(RHICmdList, ComputeShaderRHI, NormalVirtualTexture0Param, GBlackTextureWithSRV->ShaderResourceViewRHI);
-		SetSRVParameter(RHICmdList, ComputeShaderRHI, NormalVirtualTexture1Param, GBlackTextureWithSRV->ShaderResourceViewRHI);
-		SetTextureParameter(RHICmdList, ComputeShaderRHI, NormalVirtualTexturePageTableParam, GBlackTexture->TextureRHI);
-		SetShaderValue(RHICmdList, ComputeShaderRHI, NormalVirtualTextureWorldToUvTransformParam, FMatrix44f::Identity);
-		SetShaderValue(RHICmdList, ComputeShaderRHI, NormalVirtualTextureEnabledParam, 0 /* false */);
-		SetShaderValue(RHICmdList, ComputeShaderRHI, NormalVirtualTextureUnpackModeParam, ENormalUnpackType::None);
-		SetShaderValue(RHICmdList, ComputeShaderRHI, NormalVirtualTexturePageTableUniform0Param, DummyUint4);
-		SetShaderValue(RHICmdList, ComputeShaderRHI, NormalVirtualTexturePageTableUniform1Param, DummyUint4);
-		SetShaderValue(RHICmdList, ComputeShaderRHI, NormalVirtualTextureUniforms0Param, DummyUint4);
-		SetShaderValue(RHICmdList, ComputeShaderRHI, NormalVirtualTextureUniforms1Param, DummyUint4);
-	}
-
-	bool SetCachedHeightTextureParameters(FRHICommandList& RHICmdList, FRHIComputeShader* ComputeShaderRHI, const FNDILandscapeData_RenderThread& ProxyData) const
-	{
-		if (ProxyData.TextureResources)
-		{
-			FRHIShaderResourceView* HeightTextureSrv = ProxyData.TextureResources->GetHeightTexture();
-			FRHIShaderResourceView* PhysMatTextureSrv = ProxyData.TextureResources->GetPhysMatTexture();
-
-			if (HeightTextureSrv || PhysMatTextureSrv)
-			{
-				SetShaderValue(RHICmdList, ComputeShaderRHI, CachedHeightTextureWorldToUvTransformParam, (FMatrix44f)ProxyData.CachedHeightTextureWorldToUvTransform);
-				SetShaderValue(RHICmdList, ComputeShaderRHI, CachedHeightTextureUvToWorldTransformParam, (FMatrix44f)ProxyData.CachedHeightTextureUvToWorldTransform);
-				SetShaderValue(RHICmdList, ComputeShaderRHI, CachedHeightTextureUvScaleBiasParam, (FVector4f)ProxyData.CachedHeightTextureUvScaleBias);
-				SetShaderValue(RHICmdList, ComputeShaderRHI, CachedHeightTextureGridSizeParam, (FVector2f)ProxyData.CachedHeightTextureGridSize);
-
-				if (HeightTextureSrv)
-				{
-					SetSRVParameter(RHICmdList, ComputeShaderRHI, CachedHeightTextureParam, HeightTextureSrv);
-					SetShaderValue(RHICmdList, ComputeShaderRHI, CachedHeightTexturEnabledParam, 1 /* true */);
-				}
-				else
-				{
-					SetSRVParameter(RHICmdList, ComputeShaderRHI, CachedHeightTextureParam, GBlackTextureWithSRV->ShaderResourceViewRHI);
-					SetShaderValue(RHICmdList, ComputeShaderRHI, CachedHeightTexturEnabledParam, 0 /* true */);
-				}
-
-				if (PhysMatTextureSrv)
-				{
-					FIntPoint PhysMatDimensions(ProxyData.TextureResources->GetDimensions());
-					SetSRVParameter(RHICmdList, ComputeShaderRHI, CachedPhysMatTextureParam, PhysMatTextureSrv);
-					SetShaderValue(RHICmdList, ComputeShaderRHI, CachedPhysMatTextureDimensionParam, PhysMatDimensions);
-				}
-				else
-				{
-					FIntPoint PhysMatDimensions(ForceInitToZero);
-					SetSRVParameter(RHICmdList, ComputeShaderRHI, CachedPhysMatTextureParam, GBlackTextureWithSRV->ShaderResourceViewRHI);
-					SetShaderValue(RHICmdList, ComputeShaderRHI, CachedPhysMatTextureDimensionParam, PhysMatDimensions);
-				}
-
-				FRHISamplerState* PointClampedSampler = TStaticSamplerState<SF_Point, AM_Clamp, AM_Clamp, AM_Clamp>::GetRHI();;
-				SetSamplerParameter(RHICmdList, ComputeShaderRHI, PointClampedSamplerParam, PointClampedSampler);
-
-				return true;
-			}
-		}
-
-		return false;
-	}
-
-	void SetCachedHeightTextureParameters_Defaults(FRHICommandList& RHICmdList, FRHIComputeShader* ComputeShaderRHI) const
-	{
-		FVector4f DummyVector4(ForceInitToZero);
-
-		SetSRVParameter(RHICmdList, ComputeShaderRHI, CachedHeightTextureParam, GBlackTextureWithSRV->ShaderResourceViewRHI);
-		SetShaderValue(RHICmdList, ComputeShaderRHI, CachedHeightTextureWorldToUvTransformParam, FMatrix44f::Identity);
-		SetShaderValue(RHICmdList, ComputeShaderRHI, CachedHeightTextureUvToWorldTransformParam, FMatrix44f::Identity);
-		SetShaderValue(RHICmdList, ComputeShaderRHI, CachedHeightTextureUvScaleBiasParam, DummyVector4);
-		SetShaderValue(RHICmdList, ComputeShaderRHI, CachedHeightTextureGridSizeParam, 0.0f);
-		SetShaderValue(RHICmdList, ComputeShaderRHI, CachedHeightTexturEnabledParam, 0 /* false */);
-
-		SetSRVParameter(RHICmdList, ComputeShaderRHI, CachedPhysMatTextureParam, GBlackTextureWithSRV->ShaderResourceViewRHI);
-
-		FRHISamplerState* PointClampedSampler = TStaticSamplerState<SF_Point, AM_Clamp, AM_Clamp, AM_Clamp>::GetRHI();;
-		SetSamplerParameter(RHICmdList, ComputeShaderRHI, PointClampedSamplerParam, PointClampedSampler);
-	}
-
-	void Set(FRHICommandList& RHICmdList, const FNiagaraDataInterfaceSetArgs& Context) const
-	{
-		check(IsInRenderingThread());
-
-		FRHIComputeShader* ComputeShaderRHI = Context.Shader.GetComputeShader();
-		SetShaderValue(RHICmdList, ComputeShaderRHI, SystemLWCTileParam, Context.SystemLWCTile);
-
-		FNiagaraDataInterfaceProxyLandscape* RT_Proxy = static_cast<FNiagaraDataInterfaceProxyLandscape*>(Context.DataInterface);
-		FNDILandscapeData_RenderThread* ProxyData = RT_Proxy ? RT_Proxy->SystemInstancesToProxyData_RT.Find(Context.SystemInstanceID) : nullptr;
-
-		// global samplers
-		FRHISamplerState* BilinearSamplerState = TStaticSamplerState<SF_Bilinear, AM_Clamp, AM_Clamp, AM_Clamp>::GetRHI();
-		SetSamplerParameter(RHICmdList, ComputeShaderRHI, BaseColorVirtualTextureSamplerParam, BilinearSamplerState);
-		SetSamplerParameter(RHICmdList, ComputeShaderRHI, HeightVirtualTextureSamplerParam, BilinearSamplerState);
-		SetSamplerParameter(RHICmdList, ComputeShaderRHI, NormalVirtualTextureSamplerParam, BilinearSamplerState);
-		SetSamplerParameter(RHICmdList, ComputeShaderRHI, CachedHeightTextureSamplerParam, BilinearSamplerState);
-
-		bool ApplyBaseColorVirtualTextureDefaults = true;
-		bool ApplyHeightVirtualTextureDefaults = true;
-		bool ApplyNormalVirtualTextureDefaults = true;
-		bool ApplyCachedHeightTextureDefaults = true;
-
-		if (ProxyData)
-		{
-			if (SetBaseColorVirtualTextureParameters(RHICmdList, ComputeShaderRHI, *ProxyData))
-			{
-				ApplyBaseColorVirtualTextureDefaults = false;
-			}
-
-			if (SetHeightVirtualTextureParameters(RHICmdList, ComputeShaderRHI, *ProxyData))
-			{
-				ApplyHeightVirtualTextureDefaults = false;
-			}
-
-			if (SetNormalVirtualTextureParameters(RHICmdList, ComputeShaderRHI, *ProxyData))
-			{
-				ApplyNormalVirtualTextureDefaults = false;
-			}
-
-			if (SetCachedHeightTextureParameters(RHICmdList, ComputeShaderRHI, *ProxyData))
-			{
-				ApplyCachedHeightTextureDefaults = false;
-			}
-		}
-
-		if (ApplyBaseColorVirtualTextureDefaults)
-		{
-			SetBaseColorVirtualTextureParameters_Default(RHICmdList, ComputeShaderRHI);
-		}
-
-		if (ApplyHeightVirtualTextureDefaults)
-		{
-			SetHeightVirtualTextureParameters_Default(RHICmdList, ComputeShaderRHI);
-		}
-
-		if (ApplyNormalVirtualTextureDefaults)
-		{
-			SetNormalVirtualTextureParameters_Default(RHICmdList, ComputeShaderRHI);
-		}
-
-		if (ApplyCachedHeightTextureDefaults)
-		{
-			SetCachedHeightTextureParameters_Defaults(RHICmdList, ComputeShaderRHI);
-		}
-	}
-
-private:
-	// virtual texture parameters - basecolor
-	LAYOUT_FIELD(FShaderResourceParameter, BaseColorVirtualTextureParam);
-	LAYOUT_FIELD(FShaderResourceParameter, BaseColorVirtualTexturePageTableParam);
-	LAYOUT_FIELD(FShaderResourceParameter, BaseColorVirtualTextureSamplerParam);
-	LAYOUT_FIELD(FShaderParameter, BaseColorVirtualTextureWorldToUvTransformParam);
-	LAYOUT_FIELD(FShaderParameter, BaseColorVirtualTextureSRGBParam);
-	LAYOUT_FIELD(FShaderParameter, BaseColorVirtualTextureEnabledParam);
-	LAYOUT_FIELD(FShaderParameter, BaseColorVirtualTexturePageTableUniform0Param);
-	LAYOUT_FIELD(FShaderParameter, BaseColorVirtualTexturePageTableUniform1Param);
-	LAYOUT_FIELD(FShaderParameter, BaseColorVirtualTextureUniformsParam);
-
-	// virtual texture parameters - height
-	LAYOUT_FIELD(FShaderResourceParameter, HeightVirtualTextureParam);
-	LAYOUT_FIELD(FShaderResourceParameter, HeightVirtualTexturePageTableParam);
-	LAYOUT_FIELD(FShaderResourceParameter, HeightVirtualTextureSamplerParam);
-	LAYOUT_FIELD(FShaderParameter, HeightVirtualTextureWorldToUvTransformParam);
-	LAYOUT_FIELD(FShaderParameter, HeightVirtualTextureEnabledParam);
-	LAYOUT_FIELD(FShaderParameter, HeightVirtualTexturePageTableUniform0Param);
-	LAYOUT_FIELD(FShaderParameter, HeightVirtualTexturePageTableUniform1Param);
-	LAYOUT_FIELD(FShaderParameter, HeightVirtualTextureUniformsParam);
-
-	// virtual texture parameters - normal
-	LAYOUT_FIELD(FShaderResourceParameter, NormalVirtualTexture0Param);
-	LAYOUT_FIELD(FShaderResourceParameter, NormalVirtualTexture1Param);
-	LAYOUT_FIELD(FShaderResourceParameter, NormalVirtualTexturePageTableParam);
-	LAYOUT_FIELD(FShaderResourceParameter, NormalVirtualTextureSamplerParam);
-	LAYOUT_FIELD(FShaderParameter, NormalVirtualTextureWorldToUvTransformParam);
-	LAYOUT_FIELD(FShaderParameter, NormalVirtualTextureEnabledParam);
-	LAYOUT_FIELD(FShaderParameter, NormalVirtualTextureUnpackModeParam);
-	LAYOUT_FIELD(FShaderParameter, NormalVirtualTexturePageTableUniform0Param);
-	LAYOUT_FIELD(FShaderParameter, NormalVirtualTexturePageTableUniform1Param);
-	LAYOUT_FIELD(FShaderParameter, NormalVirtualTextureUniforms0Param);
-	LAYOUT_FIELD(FShaderParameter, NormalVirtualTextureUniforms1Param);
-
-	// cached texture parameters - height
-	LAYOUT_FIELD(FShaderResourceParameter, CachedHeightTextureParam);
-	LAYOUT_FIELD(FShaderResourceParameter, CachedHeightTextureSamplerParam);
-	LAYOUT_FIELD(FShaderParameter, CachedHeightTextureWorldToUvTransformParam);
-	LAYOUT_FIELD(FShaderParameter, CachedHeightTextureUvToWorldTransformParam);
-	LAYOUT_FIELD(FShaderParameter, CachedHeightTextureUvScaleBiasParam);
-	LAYOUT_FIELD(FShaderParameter, CachedHeightTexturEnabledParam);
-	LAYOUT_FIELD(FShaderParameter, CachedHeightTextureGridSizeParam);
-	LAYOUT_FIELD(FShaderResourceParameter, PointClampedSamplerParam);
-
-	// cached texture parameters - phys mat
-	LAYOUT_FIELD(FShaderResourceParameter, CachedPhysMatTextureParam);
-	LAYOUT_FIELD(FShaderParameter, CachedPhysMatTextureDimensionParam);
-
-	// LWC data
-	LAYOUT_FIELD(FShaderParameter, SystemLWCTileParam);
-
-	static const FString BaseColorVirtualTextureSRGBName;
-	static const FString BaseColorVirtualTextureEnabledName;
-	static const FString BaseColorVirtualTextureName;
-	static const FString BaseColorVirtualTexturePageTableName;
-	static const FString BaseColorVirtualTexturePageTableUniform0Name;
-	static const FString BaseColorVirtualTexturePageTableUniform1Name;
-	static const FString BaseColorVirtualTextureSamplerName;
-	static const FString BaseColorVirtualTextureUniformsName;
-	static const FString BaseColorVirtualTextureWorldToUvTransformName;
-
-	static const FString HeightVirtualTextureEnabledName;
-	static const FString HeightVirtualTextureName;
-	static const FString HeightVirtualTexturePageTableName;
-	static const FString HeightVirtualTexturePageTableUniform0Name;
-	static const FString HeightVirtualTexturePageTableUniform1Name;
-	static const FString HeightVirtualTextureSamplerName;
-	static const FString HeightVirtualTextureUniformsName;
-	static const FString HeightVirtualTextureWorldToUvTransformName;
-
-	static const FString NormalVirtualTexture0Name;
-	static const FString NormalVirtualTexture1Name;
-	static const FString NormalVirtualTexturePageTableName;
-	static const FString NormalVirtualTextureSamplerName;
-	static const FString NormalVirtualTextureWorldToUvTransformName;
-	static const FString NormalVirtualTextureEnabledName;
-	static const FString NormalVirtualTextureUnpackModeName;
-	static const FString NormalVirtualTexturePageTableUniform0Name;
-	static const FString NormalVirtualTexturePageTableUniform1Name;
-	static const FString NormalVirtualTextureUniforms0Name;
-	static const FString NormalVirtualTextureUniforms1Name;
-
-	static const FString CachedHeightTextureEnabledName;
-	static const FString CachedHeightTextureName;
-	static const FString CachedHeightTextureSamplerName;
-	static const FString CachedHeightTextureUvScaleBiasName;
-	static const FString CachedHeightTextureWorldToUvTransformName;
-	static const FString CachedHeightTextureUvToWorldTransformName;
-	static const FString CachedHeightTextureGridSizeName;
-	static const FString PointClampedSamplerName;
-	static const FString CachedPhysMatTextureName;
-	static const FString CachedPhysMatTextureDimensionName;
-	
-	static const FString SystemLWCTileName;
-};
-
-// virtual texture parameters - basecolor
-const FString FNiagaraDataInterfaceParametersCS_Landscape::BaseColorVirtualTextureSRGBName(TEXT("BaseColorVirtualTextureSRGB_"));
-const FString FNiagaraDataInterfaceParametersCS_Landscape::BaseColorVirtualTextureEnabledName(TEXT("BaseColorVirtualTextureEnabled_"));
-const FString FNiagaraDataInterfaceParametersCS_Landscape::BaseColorVirtualTextureName(TEXT("BaseColorVirtualTexture_"));
-const FString FNiagaraDataInterfaceParametersCS_Landscape::BaseColorVirtualTexturePageTableName(TEXT("BaseColorVirtualTexturePageTable_"));
-const FString FNiagaraDataInterfaceParametersCS_Landscape::BaseColorVirtualTexturePageTableUniform0Name(TEXT("BaseColorVirtualTexturePackedUniform0_"));
-const FString FNiagaraDataInterfaceParametersCS_Landscape::BaseColorVirtualTexturePageTableUniform1Name(TEXT("BaseColorVirtualTexturePackedUniform1_"));
-const FString FNiagaraDataInterfaceParametersCS_Landscape::BaseColorVirtualTextureSamplerName(TEXT("BaseColorVirtualTextureSampler_"));
-const FString FNiagaraDataInterfaceParametersCS_Landscape::BaseColorVirtualTextureUniformsName(TEXT("BaseColorVirtualTextureUniforms_"));
-const FString FNiagaraDataInterfaceParametersCS_Landscape::BaseColorVirtualTextureWorldToUvTransformName(TEXT("BaseColorVirtualTextureWorldToUvTransform_"));
-
-// virtual texture parameters - height
-const FString FNiagaraDataInterfaceParametersCS_Landscape::HeightVirtualTextureEnabledName(TEXT("HeightVirtualTextureEnabled_"));
-const FString FNiagaraDataInterfaceParametersCS_Landscape::HeightVirtualTextureName(TEXT("HeightVirtualTexture_"));
-const FString FNiagaraDataInterfaceParametersCS_Landscape::HeightVirtualTexturePageTableName(TEXT("HeightVirtualTexturePageTable_"));
-const FString FNiagaraDataInterfaceParametersCS_Landscape::HeightVirtualTexturePageTableUniform0Name(TEXT("HeightVirtualTexturePackedUniform0_"));
-const FString FNiagaraDataInterfaceParametersCS_Landscape::HeightVirtualTexturePageTableUniform1Name(TEXT("HeightVirtualTexturePackedUniform1_"));
-const FString FNiagaraDataInterfaceParametersCS_Landscape::HeightVirtualTextureSamplerName(TEXT("HeightVirtualTextureSampler_"));
-const FString FNiagaraDataInterfaceParametersCS_Landscape::HeightVirtualTextureUniformsName(TEXT("HeightVirtualTextureUniforms_"));
-const FString FNiagaraDataInterfaceParametersCS_Landscape::HeightVirtualTextureWorldToUvTransformName(TEXT("HeightVirtualTextureWorldToUvTransform_"));
-
-// virtual texture parameters - normal
-const FString FNiagaraDataInterfaceParametersCS_Landscape::NormalVirtualTexture0Name(TEXT("NormalVirtualTexture0_"));
-const FString FNiagaraDataInterfaceParametersCS_Landscape::NormalVirtualTexture1Name(TEXT("NormalVirtualTexture1_"));
-const FString FNiagaraDataInterfaceParametersCS_Landscape::NormalVirtualTexturePageTableName(TEXT("NormalVirtualTexturePageTable_"));
-const FString FNiagaraDataInterfaceParametersCS_Landscape::NormalVirtualTextureSamplerName(TEXT("NormalVirtualTextureSampler_"));
-const FString FNiagaraDataInterfaceParametersCS_Landscape::NormalVirtualTextureWorldToUvTransformName(TEXT("NormalVirtualTextureWorldToUvTransform_"));
-const FString FNiagaraDataInterfaceParametersCS_Landscape::NormalVirtualTextureEnabledName(TEXT("NormalVirtualTextureEnabled_"));
-const FString FNiagaraDataInterfaceParametersCS_Landscape::NormalVirtualTextureUnpackModeName(TEXT("NormalVirtualTextureUnpackMode_"));
-const FString FNiagaraDataInterfaceParametersCS_Landscape::NormalVirtualTexturePageTableUniform0Name(TEXT("NormalVirtualTexturePackedUniform0_"));
-const FString FNiagaraDataInterfaceParametersCS_Landscape::NormalVirtualTexturePageTableUniform1Name(TEXT("NormalVirtualTexturePackedUniform1_"));
-const FString FNiagaraDataInterfaceParametersCS_Landscape::NormalVirtualTextureUniforms0Name(TEXT("NormalVirtualTextureUniforms0_"));
-const FString FNiagaraDataInterfaceParametersCS_Landscape::NormalVirtualTextureUniforms1Name(TEXT("NormalVirtualTextureUniforms1_"));
-
-// cached texture - height
-const FString FNiagaraDataInterfaceParametersCS_Landscape::CachedHeightTextureEnabledName(TEXT("CachedHeightTextureEnabled_"));
-const FString FNiagaraDataInterfaceParametersCS_Landscape::CachedHeightTextureName(TEXT("CachedHeightTexture_"));
-const FString FNiagaraDataInterfaceParametersCS_Landscape::CachedHeightTextureSamplerName(TEXT("CachedHeightTextureSampler_"));
-const FString FNiagaraDataInterfaceParametersCS_Landscape::CachedHeightTextureUvScaleBiasName(TEXT("CachedHeightTextureUvScaleBias_"));
-const FString FNiagaraDataInterfaceParametersCS_Landscape::CachedHeightTextureWorldToUvTransformName(TEXT("CachedHeightTextureWorldToUvTransform_"));
-const FString FNiagaraDataInterfaceParametersCS_Landscape::CachedHeightTextureUvToWorldTransformName(TEXT("CachedHeightTextureUvToWorldTransform_"));
-const FString FNiagaraDataInterfaceParametersCS_Landscape::CachedHeightTextureGridSizeName(TEXT("CachedHeightTextureGridSize_"));
-const FString FNiagaraDataInterfaceParametersCS_Landscape::PointClampedSamplerName(TEXT("PointClampedSampler_"));
-
-// cached texture - physmat
-const FString FNiagaraDataInterfaceParametersCS_Landscape::CachedPhysMatTextureName(TEXT("CachedPhysMatTexture_"));
-const FString FNiagaraDataInterfaceParametersCS_Landscape::CachedPhysMatTextureDimensionName(TEXT("CachedPhysMatTextureDimension_"));
-
-// LWC
-const FString FNiagaraDataInterfaceParametersCS_Landscape::SystemLWCTileName(TEXT("SystemLWCTile_"));
-
-IMPLEMENT_TYPE_LAYOUT(FNiagaraDataInterfaceParametersCS_Landscape);
-IMPLEMENT_NIAGARA_DI_PARAMETER(UNiagaraDataInterfaceLandscape, FNiagaraDataInterfaceParametersCS_Landscape);
-
 #undef LOCTEXT_NAMESPACE
