- Rigid Mesh Collision DI converted to RDG

diff --git a/Engine/Plugins/FX/Niagara/Shaders/Private/NiagaraDataInterfaceRigidMeshCollisionQuery.ush b/Engine/Plugins/FX/Niagara/Shaders/Private/NiagaraDataInterfaceRigidMeshCollisionQuery.ush
--- a/Engine/Plugins/FX/Niagara/Shaders/Private/NiagaraDataInterfaceRigidMeshCollisionQuery.ush
+++ b/Engine/Plugins/FX/Niagara/Shaders/Private/NiagaraDataInterfaceRigidMeshCollisionQuery.ush
@@ -4,8 +4,6 @@
 NiagaraDataInterfaceRigidMeshCollisionQuery.ush 
 =============================================================================*/
 
-#pragma once
-
 /* -----------------------------------------------------------------
  * static mesh constants and context
  * -----------------------------------------------------------------
@@ -17,104 +15,63 @@ NiagaraDataInterfaceRigidMeshCollisionQuery.ush
 #define SPHERE_INDEX 1
 #define CAPSULE_INDEX 2
 
-struct FDIRigidMeshCollisionContext
-{
-	uint MaxTransforms;
-	uint CurrentOffset;
-	uint PreviousOffset;
-	uint4 ElementOffsets;
-	Buffer<float4>		WorldTransformBuffer;
-	Buffer<float4>		InverseTransformBuffer;
-	Buffer<float4>		ElementExtentBuffer;
-	Buffer<uint>		PhysicsTypeBuffer;
-	Buffer<uint>		DFIndexBuffer;
-}; 
-
-#define DIRIGIDMESHCOLLISIONQUERY_DECLARE_CONSTANTS(NAME)\
-uint MaxTransforms_##NAME;\
-uint CurrentOffset_##NAME;\
-uint PreviousOffset_##NAME;\
-uint4 ElementOffsets_##NAME;\
-Buffer<float4>		WorldTransformBuffer_##NAME;\
-Buffer<float4>		InverseTransformBuffer_##NAME;\
-Buffer<float4>		ElementExtentBuffer_##NAME;\
-Buffer<uint>		PhysicsTypeBuffer_##NAME;\
-Buffer<uint>		DFIndexBuffer_##NAME;\
-
-
-#define DIRIGIDMESHCOLLISIONQUERY_MAKE_CONTEXT(NAME)\
-FDIRigidMeshCollisionContext DIContext; \
-DIContext.MaxTransforms = MaxTransforms_##NAME;\
-DIContext.CurrentOffset =  CurrentOffset_##NAME;\
-DIContext.PreviousOffset = PreviousOffset_##NAME;\
-DIContext.ElementOffsets = ElementOffsets_##NAME;\
-DIContext.WorldTransformBuffer = WorldTransformBuffer_##NAME;\
-DIContext.InverseTransformBuffer = InverseTransformBuffer_##NAME;\
-DIContext.ElementExtentBuffer = ElementExtentBuffer_##NAME;\
-DIContext.PhysicsTypeBuffer = PhysicsTypeBuffer_##NAME;\
-DIContext.DFIndexBuffer = DFIndexBuffer_##NAME;\
-
+uint			{ParameterName}_MaxTransforms;
+uint			{ParameterName}_CurrentOffset;
+uint			{ParameterName}_PreviousOffset;
+uint4			{ParameterName}_ElementOffsets;
+Buffer<float4>	{ParameterName}_WorldTransformBuffer;
+Buffer<float4>	{ParameterName}_InverseTransformBuffer;
+Buffer<float4>	{ParameterName}_ElementExtentBuffer;
+Buffer<uint>	{ParameterName}_PhysicsTypeBuffer;
+Buffer<uint>	{ParameterName}_DFIndexBuffer;
 
 /* -----------------------------------------------------------------
  * Element accessors
  * -----------------------------------------------------------------
  */
 
-int DIRigidMeshCollision_GetNumBoxes(in FDIRigidMeshCollisionContext DIContext)
+void GetNumBoxes_{ParameterName}(out int OutNumBoxes)
 {
-	return DIContext.ElementOffsets[BOX_INDEX+1] - DIContext.ElementOffsets[BOX_INDEX];
+	OutNumBoxes = {ParameterName}_ElementOffsets[BOX_INDEX+1] - {ParameterName}_ElementOffsets[BOX_INDEX];
 }
 
-int DIRigidMeshCollision_GetNumSpheres(in FDIRigidMeshCollisionContext DIContext)
+void GetNumSpheres_{ParameterName}(out int OutNumSpheres)
 {
-	return DIContext.ElementOffsets[SPHERE_INDEX+1] - DIContext.ElementOffsets[SPHERE_INDEX];
+	OutNumSpheres = {ParameterName}_ElementOffsets[SPHERE_INDEX+1] - {ParameterName}_ElementOffsets[SPHERE_INDEX];
 }
 
-int DIRigidMeshCollision_GetNumCapsules(in FDIRigidMeshCollisionContext DIContext)
+void GetNumCapsules_{ParameterName}(out int OutNumCapsules)
 {
-	return DIContext.ElementOffsets[CAPSULE_INDEX+1] - DIContext.ElementOffsets[CAPSULE_INDEX];
+	OutNumCapsules = {ParameterName}_ElementOffsets[CAPSULE_INDEX+1] - {ParameterName}_ElementOffsets[CAPSULE_INDEX];
 }
 
-
 /* -----------------------------------------------------------------
  * Position / Transform utils
  * -----------------------------------------------------------------
  */
 
-float3 DIRigidMeshCollision_GetCurrentPosition(in FDIRigidMeshCollisionContext DIContext, in int ElementIndex)
-{
-	const int BufferOffset = 3 * ElementIndex + DIContext.CurrentOffset;
-    return float3(DIContext.WorldTransformBuffer[BufferOffset].w, DIContext.WorldTransformBuffer[BufferOffset+1].w, DIContext.WorldTransformBuffer[BufferOffset+2].w);
-}
-
-float3 DIRigidMeshCollision_GetPreviousPosition(in FDIRigidMeshCollisionContext DIContext, in int ElementIndex)
-{
-	const int BufferOffset = 3 * ElementIndex + DIContext.PreviousOffset;
-    return float3(DIContext.WorldTransformBuffer[BufferOffset].w, DIContext.WorldTransformBuffer[BufferOffset+1].w, DIContext.WorldTransformBuffer[BufferOffset+2].w);
-}
-
-float3x4 DIRigidMeshCollision_GetCurrentTransform(in FDIRigidMeshCollisionContext DIContext, in int ElementIndex)
+float3x4 GetCurrentTransform_{ParameterName}(in int ElementIndex)
 {
-	const int BufferOffset = 3 * ElementIndex + DIContext.CurrentOffset;
-    return float3x4(DIContext.WorldTransformBuffer[BufferOffset], DIContext.WorldTransformBuffer[BufferOffset+1], DIContext.WorldTransformBuffer[BufferOffset+2]);
+	const int BufferOffset = 3 * ElementIndex + {ParameterName}_CurrentOffset;
+    return float3x4({ParameterName}_WorldTransformBuffer[BufferOffset], {ParameterName}_WorldTransformBuffer[BufferOffset+1], {ParameterName}_WorldTransformBuffer[BufferOffset+2]);
 }
 
-float3x4 DIRigidMeshCollision_GetInverseTransform(in FDIRigidMeshCollisionContext DIContext, in int ElementIndex)
+float3x4 GetInverseTransform_{ParameterName}(in int ElementIndex)
 {
-	const int BufferOffset = 3 * ElementIndex + DIContext.CurrentOffset;
-    return float3x4(DIContext.InverseTransformBuffer[BufferOffset], DIContext.InverseTransformBuffer[BufferOffset+1], DIContext.InverseTransformBuffer[BufferOffset+2]);
+	const int BufferOffset = 3 * ElementIndex + {ParameterName}_CurrentOffset;
+    return float3x4({ParameterName}_InverseTransformBuffer[BufferOffset], {ParameterName}_InverseTransformBuffer[BufferOffset+1], {ParameterName}_InverseTransformBuffer[BufferOffset+2]);
 }
 
-float3x4 DIRigidMeshCollision_GetPreviousTransform(in FDIRigidMeshCollisionContext DIContext, in int ElementIndex)
+float3x4 GetPreviousTransform_{ParameterName}(in int ElementIndex)
 {
-	const int BufferOffset = 3 * ElementIndex + DIContext.PreviousOffset;
-    return float3x4(DIContext.WorldTransformBuffer[BufferOffset], DIContext.WorldTransformBuffer[BufferOffset+1], DIContext.WorldTransformBuffer[BufferOffset+2]);
+	const int BufferOffset = 3 * ElementIndex + {ParameterName}_PreviousOffset;
+    return float3x4({ParameterName}_WorldTransformBuffer[BufferOffset], {ParameterName}_WorldTransformBuffer[BufferOffset+1], {ParameterName}_WorldTransformBuffer[BufferOffset+2]);
 }
 
-float3x4 DIRigidMeshCollision_GetPreviousInverse(in FDIRigidMeshCollisionContext DIContext, in int ElementIndex)
+float3x4 GetPreviousInverse_{ParameterName}(in int ElementIndex)
 {
-	const int BufferOffset = 3 * ElementIndex + DIContext.PreviousOffset;
-    return float3x4(DIContext.InverseTransformBuffer[BufferOffset], DIContext.InverseTransformBuffer[BufferOffset+1], DIContext.InverseTransformBuffer[BufferOffset+2]);
+	const int BufferOffset = 3 * ElementIndex + {ParameterName}_PreviousOffset;
+    return float3x4({ParameterName}_InverseTransformBuffer[BufferOffset], {ParameterName}_InverseTransformBuffer[BufferOffset+1], {ParameterName}_InverseTransformBuffer[BufferOffset+2]);
 }
 
 
@@ -124,7 +81,7 @@ float3x4 DIRigidMeshCollision_GetPreviousInverse(in FDIRigidMeshCollisionContext
  */
 
 // Given a world space position (WorldPosition) compute the sphere closest point (position,normal,velocity)
-float DIRigidMeshCollision_GetSphereProjection(in float3 LocalPosition, in float3 SphereCenter, in float SphereRadius, in int SphereIndex,
+float GetSphereProjection_{ParameterName}(in float3 LocalPosition, in float3 SphereCenter, in float SphereRadius, in int SphereIndex,
 			inout float3 OutClosestPosition, inout float3 OutClosestNormal, inout int OutElementIndex, inout float OutMinDistance)
 {
 	const float3 DeltaPosition = LocalPosition - SphereCenter;
@@ -143,18 +100,18 @@ float DIRigidMeshCollision_GetSphereProjection(in float3 LocalPosition, in float
 }
 
 // Given a world space position (WorldPosition) compute the sphere closest point (position,normal,velocity)
-float DIRigidMeshCollision_GetCapsuleProjection(in float3 LocalPosition, in float2 CapsuleSize, in int CapsuleIndex,
+float GetCapsuleProjection_{ParameterName}(in float3 LocalPosition, in float2 CapsuleSize, in int CapsuleIndex,
 			inout float3 OutClosestPosition, inout float3 OutClosestNormal, inout int OutElementIndex, inout float OutMinDistance)
 {
 	const float HalfLength = 0.5 * CapsuleSize.y;
 	const float3 SphereCenter = float3(0,0,clamp(LocalPosition.z, -HalfLength, HalfLength));
 
-	return DIRigidMeshCollision_GetSphereProjection(LocalPosition,SphereCenter,CapsuleSize.x,CapsuleIndex,
+	return GetSphereProjection_{ParameterName}(LocalPosition,SphereCenter,CapsuleSize.x,CapsuleIndex,
 			OutClosestPosition,OutClosestNormal,OutElementIndex,OutMinDistance);
 }
 
 // Given a world space position (WorldPosition) compute the sphere closest point (position,normal,velocity)
-float DIRigidMeshCollision_GetBoxProjection(in float3 LocalPosition, in float3 BoxExtent, in int BoxIndex, 
+float GetBoxProjection_{ParameterName}(in float3 LocalPosition, in float3 BoxExtent, in int BoxIndex, 
 			inout float3 OutClosestPosition, inout float3 OutClosestNormal, inout int OutElementIndex, inout float OutMinDistance)
 {
 	const float3 HalfExtent = 0.5 * BoxExtent;
@@ -206,50 +163,49 @@ float DIRigidMeshCollision_GetBoxProjection(in float3 LocalPosition, in float3 B
  * -----------------------------------------------------------------
  */
 
-float3 DIRigidMeshCollision_GetLocalPosition(in FDIRigidMeshCollisionContext DIContext, in float3 WorldPosition, in int ElementIndex, in float TimeFraction )
+float3 GetLocalPosition_{ParameterName}(in float3 WorldPosition, in int ElementIndex, in float TimeFraction )
 {
-	const float3 CurrentLocal = mul(DIRigidMeshCollision_GetInverseTransform(DIContext,ElementIndex), float4(WorldPosition,1.0)).xyz;
-	const float3 PreviousLocal = mul(DIRigidMeshCollision_GetPreviousInverse(DIContext,ElementIndex), float4(WorldPosition,1.0)).xyz;
+	const float3 CurrentLocal = mul(GetInverseTransform_{ParameterName}(ElementIndex), float4(WorldPosition,1.0)).xyz;
+	const float3 PreviousLocal = mul(GetPreviousInverse_{ParameterName}(ElementIndex), float4(WorldPosition,1.0)).xyz;
 
 	return PreviousLocal + TimeFraction * (CurrentLocal-PreviousLocal);
 }
 
 // Given a world space position (WorldPosition) compute the sphere closest point (position,normal,velocity)
-int DIRigidMeshCollision_GetClosestElement(in FDIRigidMeshCollisionContext DIContext, in float3 WorldPosition, out float3 OutClosestPosition, 
-										  out float3 OutClosestNormal, out float OutMinDistance, in float TimeFraction)
+int GetClosestElement_{ParameterName}(in float3 WorldPosition, out float3 OutClosestPosition, out float3 OutClosestNormal, out float OutMinDistance, in float TimeFraction)
 {	
 	float MinDistance = MAX_DISTANCE;
 	int ElementIndex = -1;
 	float3 CollisionPosition = float3(0,0,0);
 	float3 CollisionNormal = float3(0,0,0);
 
-	const int SpheresBegin = DIContext.ElementOffsets[SPHERE_INDEX];
-	const int SpheresEnd = DIContext.ElementOffsets[SPHERE_INDEX+1];
+	const int SpheresBegin = {ParameterName}_ElementOffsets[SPHERE_INDEX];
+	const int SpheresEnd = {ParameterName}_ElementOffsets[SPHERE_INDEX+1];
 	for (int SphereIndex = SpheresBegin; SphereIndex < SpheresEnd; ++SphereIndex)
 	{
-		const float3 LocalPosition = DIRigidMeshCollision_GetLocalPosition(DIContext,WorldPosition,SphereIndex,TimeFraction);
+		const float3 LocalPosition = GetLocalPosition_{ParameterName}(WorldPosition,SphereIndex,TimeFraction);
 
-		DIRigidMeshCollision_GetSphereProjection(LocalPosition, float3(0,0,0), DIContext.ElementExtentBuffer[SphereIndex].x, SphereIndex,
+		GetSphereProjection_{ParameterName}(LocalPosition, float3(0,0,0), {ParameterName}_ElementExtentBuffer[SphereIndex].x, SphereIndex,
 				CollisionPosition, CollisionNormal, ElementIndex, MinDistance);
 	}
 
-	const int CapsulesBegin = DIContext.ElementOffsets[CAPSULE_INDEX];
-	const int CapsulesEnd = DIContext.ElementOffsets[CAPSULE_INDEX+1];
+	const int CapsulesBegin = {ParameterName}_ElementOffsets[CAPSULE_INDEX];
+	const int CapsulesEnd = {ParameterName}_ElementOffsets[CAPSULE_INDEX+1];
 	for (int CapsuleIndex = CapsulesBegin; CapsuleIndex < CapsulesEnd; ++CapsuleIndex)
 	{
-		const float3 LocalPosition = DIRigidMeshCollision_GetLocalPosition(DIContext,WorldPosition,CapsuleIndex,TimeFraction);
+		const float3 LocalPosition = GetLocalPosition_{ParameterName}(WorldPosition,CapsuleIndex,TimeFraction);
 
-		DIRigidMeshCollision_GetCapsuleProjection(LocalPosition, DIContext.ElementExtentBuffer[CapsuleIndex].xy, CapsuleIndex,
+		GetCapsuleProjection_{ParameterName}(LocalPosition, {ParameterName}_ElementExtentBuffer[CapsuleIndex].xy, CapsuleIndex,
 				CollisionPosition, CollisionNormal, ElementIndex, MinDistance);
 	}
 
-	const int BoxesBegin = DIContext.ElementOffsets[BOX_INDEX];
-	const int BoxesEnd = DIContext.ElementOffsets[BOX_INDEX+1];
+	const int BoxesBegin = {ParameterName}_ElementOffsets[BOX_INDEX];
+	const int BoxesEnd = {ParameterName}_ElementOffsets[BOX_INDEX+1];
 	for (int BoxIndex = BoxesBegin; BoxIndex < BoxesEnd; ++BoxIndex)
 	{
-		const float3 LocalPosition = DIRigidMeshCollision_GetLocalPosition(DIContext,WorldPosition,BoxIndex,TimeFraction);
+		const float3 LocalPosition = GetLocalPosition_{ParameterName}(WorldPosition,BoxIndex,TimeFraction);
 
-		DIRigidMeshCollision_GetBoxProjection(LocalPosition, DIContext.ElementExtentBuffer[BoxIndex].xyz, BoxIndex,
+		GetBoxProjection_{ParameterName}(LocalPosition, {ParameterName}_ElementExtentBuffer[BoxIndex].xyz, BoxIndex,
 				CollisionPosition, CollisionNormal, ElementIndex, MinDistance);
 	}
 	OutClosestPosition = CollisionPosition;
@@ -264,8 +220,7 @@ int DIRigidMeshCollision_GetClosestElement(in FDIRigidMeshCollisionContext DICon
  */
 
 // Given a world space position (WorldPosition) and an element index compute the static mesh closest point (position,normal,velocity)
-void DIRigidMeshCollision_GetElementPoint(in FDIRigidMeshCollisionContext DIContext, in float3 WorldPosition, in float DeltaTime, in float TimeFraction, in int ElementIndex, out float3 OutClosestPosition, 
-										  out float3 OutClosestNormal, out float3 OutClosestVelocity)
+void GetElementPoint_{ParameterName}(in float3 WorldPosition, in float DeltaTime, in float TimeFraction, in int ElementIndex, out float3 OutClosestPosition, out float3 OutClosestNormal, out float3 OutClosestVelocity)
 {	
 	OutClosestNormal = float3(0,0,0);
 	OutClosestPosition = float3(0,0,0);
@@ -274,36 +229,36 @@ void DIRigidMeshCollision_GetElementPoint(in FDIRigidMeshCollisionContext DICont
 	if (ElementIndex >= 0)
 	{
 		const uint ElementIndexUint = uint(ElementIndex);
-		const float3 LocalPosition = DIRigidMeshCollision_GetLocalPosition(DIContext,WorldPosition,ElementIndex,TimeFraction);
+		const float3 LocalPosition = GetLocalPosition_{ParameterName}(WorldPosition,ElementIndex,TimeFraction);
 
 		float MinDistance = MAX_DISTANCE;
 		int OutputIndex = -1;
 		float3 CollisionPosition = float3(0,0,0);
 		float3 CollisionNormal = float3(0,0,0);
 
-		if ( ElementIndexUint >= DIContext.ElementOffsets[SPHERE_INDEX] && ElementIndexUint < DIContext.ElementOffsets[SPHERE_INDEX+1] )
+		if ( ElementIndexUint >= {ParameterName}_ElementOffsets[SPHERE_INDEX] && ElementIndexUint < {ParameterName}_ElementOffsets[SPHERE_INDEX+1] )
 		{
-			DIRigidMeshCollision_GetSphereProjection(LocalPosition, float3(0,0,0), DIContext.ElementExtentBuffer[ElementIndex].x, ElementIndex,
+			GetSphereProjection_{ParameterName}(LocalPosition, float3(0,0,0), {ParameterName}_ElementExtentBuffer[ElementIndex].x, ElementIndex,
 					CollisionPosition, CollisionNormal, OutputIndex, MinDistance);
 		}
-		else if ( ElementIndexUint >= DIContext.ElementOffsets[CAPSULE_INDEX] && ElementIndexUint < DIContext.ElementOffsets[CAPSULE_INDEX+1] )
+		else if ( ElementIndexUint >= {ParameterName}_ElementOffsets[CAPSULE_INDEX] && ElementIndexUint < {ParameterName}_ElementOffsets[CAPSULE_INDEX+1] )
 		{
-			DIRigidMeshCollision_GetCapsuleProjection(LocalPosition, DIContext.ElementExtentBuffer[ElementIndex].xy, ElementIndex,
+			GetCapsuleProjection_{ParameterName}(LocalPosition, {ParameterName}_ElementExtentBuffer[ElementIndex].xy, ElementIndex,
 					CollisionPosition, CollisionNormal, OutputIndex, MinDistance);
 		}
-		else if ( ElementIndexUint >= DIContext.ElementOffsets[BOX_INDEX] && ElementIndexUint < DIContext.ElementOffsets[BOX_INDEX+1] )
+		else if ( ElementIndexUint >= {ParameterName}_ElementOffsets[BOX_INDEX] && ElementIndexUint < {ParameterName}_ElementOffsets[BOX_INDEX+1] )
 		{
-			DIRigidMeshCollision_GetBoxProjection(LocalPosition, DIContext.ElementExtentBuffer[ElementIndex].xyz, ElementIndex,
+			GetBoxProjection_{ParameterName}(LocalPosition, {ParameterName}_ElementExtentBuffer[ElementIndex].xyz, ElementIndex,
 					CollisionPosition, CollisionNormal, OutputIndex, MinDistance);
 		}
 
 		if (ElementIndex != -1)
 		{
-			const float3 PreviousPosition = mul(DIRigidMeshCollision_GetPreviousTransform(DIContext,ElementIndex), float4(CollisionPosition,1.0)).xyz;
-			const float3 CurrentPosition = mul(DIRigidMeshCollision_GetCurrentTransform(DIContext,ElementIndex), float4(CollisionPosition,1.0)).xyz;
+			const float3 PreviousPosition = mul(GetPreviousTransform_{ParameterName}(ElementIndex), float4(CollisionPosition,1.0)).xyz;
+			const float3 CurrentPosition = mul(GetCurrentTransform_{ParameterName}(ElementIndex), float4(CollisionPosition,1.0)).xyz;
 
-			const float3 PreviousNormal = mul(DIRigidMeshCollision_GetPreviousTransform(DIContext,ElementIndex), float4(CollisionNormal,0.0)).xyz;
-			const float3 CurrentNormal = mul(DIRigidMeshCollision_GetCurrentTransform(DIContext,ElementIndex), float4(CollisionNormal,0.0)).xyz;
+			const float3 PreviousNormal = mul(GetPreviousTransform_{ParameterName}(ElementIndex), float4(CollisionNormal,0.0)).xyz;
+			const float3 CurrentNormal = mul(GetCurrentTransform_{ParameterName}(ElementIndex), float4(CollisionNormal,0.0)).xyz;
 
 			OutClosestPosition = PreviousPosition + TimeFraction * (CurrentPosition-PreviousPosition);
 			OutClosestNormal = normalize(PreviousNormal + TimeFraction * (CurrentNormal-PreviousNormal));
@@ -314,61 +269,61 @@ void DIRigidMeshCollision_GetElementPoint(in FDIRigidMeshCollisionContext DICont
 }
 
 // Given a world space position (WorldPosition) and an element index compute the static mesh closest distance
-void DIRigidMeshCollision_GetElementDistance(in FDIRigidMeshCollisionContext DIContext, in float3 WorldPosition, in float TimeFraction, in int ElementIndex, out float OutClosestDistance)
+void GetElementDistance_{ParameterName}(in float3 WorldPosition, in float TimeFraction, in int ElementIndex, out float OutClosestDistance)
 {	
 	OutClosestDistance = 0.0;
 
 	if (ElementIndex >= 0)
 	{
 		const uint ElementIndexUint = uint(ElementIndex);
-		const float3 LocalPosition = DIRigidMeshCollision_GetLocalPosition(DIContext,WorldPosition,ElementIndex,TimeFraction);
+		const float3 LocalPosition = GetLocalPosition_{ParameterName}(WorldPosition,ElementIndex,TimeFraction);
 
 		float MinDistance = MAX_DISTANCE;
 		int OutputIndex = -1;
 		float3 CollisionPosition = float3(0,0,0);
 		float3 CollisionNormal = float3(0,0,0);
 
-		if ( ElementIndexUint >= DIContext.ElementOffsets[SPHERE_INDEX] && ElementIndexUint < DIContext.ElementOffsets[SPHERE_INDEX+1] )
+		if ( ElementIndexUint >= {ParameterName}_ElementOffsets[SPHERE_INDEX] && ElementIndexUint < {ParameterName}_ElementOffsets[SPHERE_INDEX+1] )
 		{
-			DIRigidMeshCollision_GetSphereProjection(LocalPosition, float3(0,0,0), DIContext.ElementExtentBuffer[ElementIndex].x, ElementIndex,
+			GetSphereProjection_{ParameterName}(LocalPosition, float3(0,0,0), {ParameterName}_ElementExtentBuffer[ElementIndex].x, ElementIndex,
 					CollisionPosition, CollisionNormal, OutputIndex, MinDistance);
 		}
-		else if ( ElementIndexUint >= DIContext.ElementOffsets[CAPSULE_INDEX] && ElementIndexUint < DIContext.ElementOffsets[CAPSULE_INDEX+1] )
+		else if ( ElementIndexUint >= {ParameterName}_ElementOffsets[CAPSULE_INDEX] && ElementIndexUint < {ParameterName}_ElementOffsets[CAPSULE_INDEX+1] )
 		{
-			DIRigidMeshCollision_GetCapsuleProjection(LocalPosition, DIContext.ElementExtentBuffer[ElementIndex].xy, ElementIndex,
+			GetCapsuleProjection_{ParameterName}(LocalPosition, {ParameterName}_ElementExtentBuffer[ElementIndex].xy, ElementIndex,
 					CollisionPosition, CollisionNormal, OutputIndex, MinDistance);
 		}
-		else if ( ElementIndexUint >= DIContext.ElementOffsets[BOX_INDEX] && ElementIndexUint < DIContext.ElementOffsets[BOX_INDEX+1] )
+		else if ( ElementIndexUint >= {ParameterName}_ElementOffsets[BOX_INDEX] && ElementIndexUint < {ParameterName}_ElementOffsets[BOX_INDEX+1] )
 		{
-			DIRigidMeshCollision_GetBoxProjection(LocalPosition, DIContext.ElementExtentBuffer[ElementIndex].xyz, ElementIndex,
+			GetBoxProjection_{ParameterName}(LocalPosition, {ParameterName}_ElementExtentBuffer[ElementIndex].xyz, ElementIndex,
 					CollisionPosition, CollisionNormal, OutputIndex, MinDistance);
 		}
 	
 		if (ElementIndex != -1)
 		{
-			OutClosestDistance = (DIContext.PhysicsTypeBuffer[ElementIndex] == 1) ? MinDistance : 0;
+			OutClosestDistance = ({ParameterName}_PhysicsTypeBuffer[ElementIndex] == 1) ? MinDistance : 0;
 		}
 	}
 }
 
  // Given a world space position (WorldPosition) compute the static mesh closest element
-void DIRigidMeshCollision_GetClosestElement(in FDIRigidMeshCollisionContext DIContext, in float3 WorldPosition, in float TimeFraction, out int OutElementIndex)
+void GetClosestElement_{ParameterName}(in float3 WorldPosition, in float TimeFraction, out int OutElementIndex)
 {	
 	float3 CollisionPosition = float3(0,0,0);
 	float3 CollisionNormal = float3(0,0,0);
 	float MinDistance = 0.0;
 
-	OutElementIndex = DIRigidMeshCollision_GetClosestElement(DIContext,WorldPosition,CollisionPosition,CollisionNormal,MinDistance,TimeFraction);
+	OutElementIndex = GetClosestElement_{ParameterName}(WorldPosition,CollisionPosition,CollisionNormal,MinDistance,TimeFraction);
 }
 
 // Given a world space position (WorldPosition) compute the static mesh closest point (position,normal,velocity)
-void DIRigidMeshCollision_GetClosestPoint(in FDIRigidMeshCollisionContext DIContext, in float3 WorldPosition, in float DeltaTime, in float TimeFraction, out float OutClosestDistance, out float3 OutClosestPosition, out float3 OutClosestNormal, out float3 OutClosestVelocity)
+void GetClosestPoint_{ParameterName}(in float3 WorldPosition, in float DeltaTime, in float TimeFraction, out float OutClosestDistance, out float3 OutClosestPosition, out float3 OutClosestNormal, out float3 OutClosestVelocity)
 {	
 	float3 CollisionPosition = float3(0,0,0);
 	float3 CollisionNormal = float3(0,0,0);
 	OutClosestDistance = 0.0;
 
-	const int ElementIndex = DIRigidMeshCollision_GetClosestElement(DIContext,WorldPosition,CollisionPosition,CollisionNormal,OutClosestDistance,TimeFraction);
+	const int ElementIndex = GetClosestElement_{ParameterName}(WorldPosition,CollisionPosition,CollisionNormal,OutClosestDistance,TimeFraction);
 
 	OutClosestNormal = float3(0,0,0);
 	OutClosestPosition = float3(0,0,0);
@@ -376,11 +331,11 @@ void DIRigidMeshCollision_GetClosestPoint(in FDIRigidMeshCollisionContext DICont
 
 	if (ElementIndex != -1)
 	{
-		const float3 PreviousPosition = mul(DIRigidMeshCollision_GetPreviousTransform(DIContext,ElementIndex), float4(CollisionPosition,1.0)).xyz;
-		const float3 CurrentPosition = mul(DIRigidMeshCollision_GetCurrentTransform(DIContext,ElementIndex), float4(CollisionPosition,1.0)).xyz;
+		const float3 PreviousPosition = mul(GetPreviousTransform_{ParameterName}(ElementIndex), float4(CollisionPosition,1.0)).xyz;
+		const float3 CurrentPosition = mul(GetCurrentTransform_{ParameterName}(ElementIndex), float4(CollisionPosition,1.0)).xyz;
 
-		const float3 PreviousNormal = mul(DIRigidMeshCollision_GetPreviousTransform(DIContext,ElementIndex), float4(CollisionNormal,0.0)).xyz;
-		const float3 CurrentNormal = mul(DIRigidMeshCollision_GetCurrentTransform(DIContext,ElementIndex), float4(CollisionNormal,0.0)).xyz;
+		const float3 PreviousNormal = mul(GetPreviousTransform_{ParameterName}(ElementIndex), float4(CollisionNormal,0.0)).xyz;
+		const float3 CurrentNormal = mul(GetCurrentTransform_{ParameterName}(ElementIndex), float4(CollisionNormal,0.0)).xyz;
 
 		OutClosestPosition = PreviousPosition + TimeFraction * (CurrentPosition-PreviousPosition);
 		OutClosestNormal = normalize(PreviousNormal + TimeFraction * (CurrentNormal-PreviousNormal));
@@ -390,24 +345,24 @@ void DIRigidMeshCollision_GetClosestPoint(in FDIRigidMeshCollisionContext DICont
 }
 
 // Given a world space position (WorldPosition) compute the static mesh closest distance
-void DIRigidMeshCollision_GetClosestDistance(in FDIRigidMeshCollisionContext DIContext, in float3 WorldPosition, in float TimeFraction, out float OutClosestDistance)
+void GetClosestDistance_{ParameterName}(in float3 WorldPosition, in float TimeFraction, out float OutClosestDistance)
 {	
 	float3 CollisionPosition = float3(0,0,0);
 	float3 CollisionNormal = float3(0,0,0);
 	float MinDistance = 0.0;
 
-	const int ElementIndex = DIRigidMeshCollision_GetClosestElement(DIContext,WorldPosition,CollisionPosition,CollisionNormal,MinDistance,TimeFraction);
+	const int ElementIndex = GetClosestElement_{ParameterName}(WorldPosition,CollisionPosition,CollisionNormal,MinDistance,TimeFraction);
 
 	OutClosestDistance = 0.0;
 
 	if (ElementIndex != -1)
 	{
-		OutClosestDistance = (DIContext.PhysicsTypeBuffer[ElementIndex] == 1) ? MinDistance : 0;
+		OutClosestDistance = ({ParameterName}_PhysicsTypeBuffer[ElementIndex] == 1) ? MinDistance : 0;
 	}
 }
 
 // Given a world space position (WorldPosition) compute the static mesh closest point (position,normal,velocity)
-void DIRigidMeshCollision_GetClosestPointMeshDistanceField(in FDIRigidMeshCollisionContext DIContext, in float3 WorldPosition, in float DeltaTime, in float TimeFraction, in float MaxDistance, out float OutClosestDistance, out float3 OutClosestPosition, out float3 OutClosestNormal, out float3 OutClosestVelocity, out bool NormalIsValid)
+void GetClosestPointMeshDistanceField_{ParameterName}(in float3 WorldPosition, in float DeltaTime, in float TimeFraction, in float MaxDistance, out float OutClosestDistance, out float3 OutClosestPosition, out float3 OutClosestNormal, out float3 OutClosestVelocity, out bool NormalIsValid)
 {	
 	float3 CollisionPosition = float3(0,0,0);
 	float3 CollisionNormal = float3(0,0,0);
@@ -420,11 +375,11 @@ void DIRigidMeshCollision_GetClosestPointMeshDistanceField(in FDIRigidMeshCollis
 	NormalIsValid = false;
 
 	float OutClosestDistanceTmp;
-	const int ElementIndex = DIRigidMeshCollision_GetClosestElement(DIContext,WorldPosition,CollisionPosition,CollisionNormal,OutClosestDistanceTmp,TimeFraction);
+	const int ElementIndex = GetClosestElement_{ParameterName}(WorldPosition,CollisionPosition,CollisionNormal,OutClosestDistanceTmp,TimeFraction);
 
 	if (ElementIndex != -1)
 	{
-		uint DFIndex = DIContext.DFIndexBuffer[ElementIndex];
+		uint DFIndex = {ParameterName}_DFIndexBuffer[ElementIndex];
 
 		if (DFIndex >= 0 && DFIndex < NumSceneObjects)
 		{
@@ -447,10 +402,10 @@ void DIRigidMeshCollision_GetClosestPointMeshDistanceField(in FDIRigidMeshCollis
 			OutClosestNormal = WorldGradientLength > 0.0f ? WorldGradient / WorldGradientLength : 0;
 			
 			CollisionPosition = WorldPosition - OutClosestNormal * OutClosestDistance;			
-			const float3 LocalCollisionPosition = mul(DIRigidMeshCollision_GetInverseTransform(DIContext,ElementIndex), float4(CollisionPosition,1.0)).xyz;			
+			const float3 LocalCollisionPosition = mul(GetInverseTransform_{ParameterName}(ElementIndex), float4(CollisionPosition,1.0)).xyz;			
 			
-			const float3 PreviousPosition = mul(DIRigidMeshCollision_GetPreviousTransform(DIContext,ElementIndex), float4(LocalCollisionPosition,1.0)).xyz;
-			const float3 CurrentPosition = mul(DIRigidMeshCollision_GetCurrentTransform(DIContext,ElementIndex), float4(LocalCollisionPosition,1.0)).xyz;
+			const float3 PreviousPosition = mul(GetPreviousTransform_{ParameterName}(ElementIndex), float4(LocalCollisionPosition,1.0)).xyz;
+			const float3 CurrentPosition = mul(GetCurrentTransform_{ParameterName}(ElementIndex), float4(LocalCollisionPosition,1.0)).xyz;
 			OutClosestVelocity = ( CurrentPosition - PreviousPosition ) / DeltaTime;
 			OutClosestPosition = PreviousPosition + TimeFraction * (CurrentPosition-PreviousPosition);
 
@@ -462,7 +417,7 @@ void DIRigidMeshCollision_GetClosestPointMeshDistanceField(in FDIRigidMeshCollis
 }
 
 // Given a world space position (WorldPosition) compute the static mesh closest point (position,normal,velocity) by querying all overlapping mesh distance fields
-void DIRigidMeshCollision_GetClosestPointMeshDistanceFieldAccurate(in FDIRigidMeshCollisionContext DIContext, in float3 WorldPosition, in float DeltaTime, in float TimeFraction, in float MaxDistance, out float OutClosestDistance, out float3 OutClosestPosition, out float3 OutClosestNormal, out float3 OutClosestVelocity, out bool NormalIsValid)
+void GetClosestPointMeshDistanceFieldAccurate_{ParameterName}(in float3 WorldPosition, in float DeltaTime, in float TimeFraction, in float MaxDistance, out float OutClosestDistance, out float3 OutClosestPosition, out float3 OutClosestNormal, out float3 OutClosestVelocity, out bool NormalIsValid)
 {	
 	float3 CollisionPosition = float3(0,0,0);
 	float3 CollisionNormal = float3(0,0,0);
@@ -474,9 +429,9 @@ void DIRigidMeshCollision_GetClosestPointMeshDistanceFieldAccurate(in FDIRigidMe
 	
 	NormalIsValid = false;
 
-	for (int ElementIndex = 0; ElementIndex < DIContext.ElementOffsets[3]; ++ElementIndex)
+	for (int ElementIndex = 0; ElementIndex < {ParameterName}_ElementOffsets[3]; ++ElementIndex)
 	{
-		uint DFIndex = DIContext.DFIndexBuffer[ElementIndex];
+		uint DFIndex = {ParameterName}_DFIndexBuffer[ElementIndex];
 
 		// #todo(dmp): this does a dynamic branch based on intersecting the bbox.  Maybe we can factor that out due to the broadphase here?
 		// #todo(dmp): Check the distance from the physics asset for a speedup?
@@ -500,10 +455,10 @@ void DIRigidMeshCollision_GetClosestPointMeshDistanceFieldAccurate(in FDIRigidMe
 			OutClosestNormal = WorldGradientLength > 0.0f ? WorldGradient / WorldGradientLength : 0;
 			
 			CollisionPosition = WorldPosition - OutClosestNormal * OutClosestDistance;			
-			const float3 LocalCollisionPosition = mul(DIRigidMeshCollision_GetInverseTransform(DIContext,ElementIndex), float4(CollisionPosition,1.0)).xyz;			
+			const float3 LocalCollisionPosition = mul(GetInverseTransform_{ParameterName}(ElementIndex), float4(CollisionPosition,1.0)).xyz;			
 			
-			const float3 PreviousPosition = mul(DIRigidMeshCollision_GetPreviousTransform(DIContext,ElementIndex), float4(LocalCollisionPosition,1.0)).xyz;
-			const float3 CurrentPosition = mul(DIRigidMeshCollision_GetCurrentTransform(DIContext,ElementIndex), float4(LocalCollisionPosition,1.0)).xyz;
+			const float3 PreviousPosition = mul(GetPreviousTransform_{ParameterName}(ElementIndex), float4(LocalCollisionPosition,1.0)).xyz;
+			const float3 CurrentPosition = mul(GetCurrentTransform_{ParameterName}(ElementIndex), float4(LocalCollisionPosition,1.0)).xyz;
 			OutClosestVelocity = ( CurrentPosition - PreviousPosition ) / DeltaTime;
 			OutClosestPosition = PreviousPosition + TimeFraction * (CurrentPosition-PreviousPosition);
 
@@ -515,7 +470,7 @@ void DIRigidMeshCollision_GetClosestPointMeshDistanceFieldAccurate(in FDIRigidMe
 }
 
 // Given a world space position (WorldPosition) compute the static mesh closest point (position,normal,velocity)
-void DIRigidMeshCollision_GetClosestPointMeshDistanceFieldNoNormal(in FDIRigidMeshCollisionContext DIContext, in float3 WorldPosition, in float DeltaTime, in float TimeFraction, in float MaxDistance, out float OutClosestDistance, out float3 OutClosestPosition, out float3 OutClosestVelocity)
+void GetClosestPointMeshDistanceFieldNoNormal_{ParameterName}(in float3 WorldPosition, in float DeltaTime, in float TimeFraction, in float MaxDistance, out float OutClosestDistance, out float3 OutClosestPosition, out float3 OutClosestVelocity)
 {	
 	float3 CollisionPosition = float3(0,0,0);
 	float3 CollisionNormal = float3(0,0,0);
@@ -524,11 +479,11 @@ void DIRigidMeshCollision_GetClosestPointMeshDistanceFieldNoNormal(in FDIRigidMe
 	OutClosestPosition = float3(0,0,0);
 	OutClosestVelocity = float3(0,0,0);
 	
-	const int ElementIndex = DIRigidMeshCollision_GetClosestElement(DIContext,WorldPosition,CollisionPosition,CollisionNormal,OutClosestDistance,TimeFraction);
+	const int ElementIndex = GetClosestElement_{ParameterName}(WorldPosition,CollisionPosition,CollisionNormal,OutClosestDistance,TimeFraction);
 
 	if (ElementIndex != -1)
 	{
-		int DFIndex = DIContext.DFIndexBuffer[ElementIndex];
+		int DFIndex = {ParameterName}_DFIndexBuffer[ElementIndex];
 
 		if (DFIndex >= 0 && DFIndex < NumSceneObjects)
 		{
@@ -537,9 +492,15 @@ void DIRigidMeshCollision_GetClosestPointMeshDistanceFieldNoNormal(in FDIRigidMe
 			OutClosestDistance = DistanceToNearestSurfaceForObject(DFIndex, WorldPosition, MaxDistance);		
 		}
 		
-		const float3 PreviousPosition = mul(DIRigidMeshCollision_GetPreviousTransform(DIContext,ElementIndex), float4(CollisionPosition,1.0)).xyz;
-		const float3 CurrentPosition = mul(DIRigidMeshCollision_GetCurrentTransform(DIContext,ElementIndex), float4(CollisionPosition,1.0)).xyz;
+		const float3 PreviousPosition = mul(GetPreviousTransform_{ParameterName}(ElementIndex), float4(CollisionPosition,1.0)).xyz;
+		const float3 CurrentPosition = mul(GetCurrentTransform_{ParameterName}(ElementIndex), float4(CollisionPosition,1.0)).xyz;
 		OutClosestVelocity = ( CurrentPosition - PreviousPosition ) / DeltaTime;
 		OutClosestPosition = PreviousPosition + TimeFraction * (CurrentPosition-PreviousPosition);
 	}	
-}
\ No newline at end of file
+}
+
+#undef SMALL_NUMBER
+#undef MAX_DISTANCE
+#undef BOX_INDEX
+#undef SPHERE_INDEX
+#undef CAPSULE_INDEX
diff --git a/Engine/Plugins/FX/Niagara/Source/Niagara/Classes/NiagaraDataInterface.h b/Engine/Plugins/FX/Niagara/Source/Niagara/Classes/NiagaraDataInterface.h
--- a/Engine/Plugins/FX/Niagara/Source/Niagara/Classes/NiagaraDataInterface.h
+++ b/Engine/Plugins/FX/Niagara/Source/Niagara/Classes/NiagaraDataInterface.h
@@ -389,6 +389,7 @@ struct NIAGARA_API FNiagaraDataInterfaceSetShaderParametersContext
 	FNiagaraSystemInstanceID GetSystemInstanceID() const;
 	bool IsResourceBound(const void* ResourceAddress) const;
 	bool IsParameterBound(const void* ParameterAddress) const;
+	template<typename T> bool IsStructBound(const T* StructAddress) const { return IsStructBoundInternal(StructAddress, sizeof(T)); }
 	bool IsOutputStage() const;
 	bool IsIterationStage() const;
 
@@ -399,6 +400,13 @@ struct NIAGARA_API FNiagaraDataInterfaceSetShaderParametersContext
 
 		return reinterpret_cast<T*>(BaseParameters + StructOffset);
 	}
+	template<typename T> T* GetParameterIncludedStruct() const
+	{
+		const uint32 StructOffset = Align(ParametersOffset, TShaderParameterStructTypeInfo<T>::Alignment);
+		ParametersOffset = StructOffset + TShaderParameterStructTypeInfo<T>::GetStructMetadata()->GetSize();
+
+		return reinterpret_cast<T*>(BaseParameters + StructOffset);
+	}
 	template<typename T> TArrayView<T> GetParameterLooseArray(int32 NumElements) const
 	{
 		const uint32 ArrayOffset = Align(ParametersOffset, SHADER_PARAMETER_ARRAY_ELEMENT_ALIGNMENT);
@@ -419,6 +427,9 @@ struct NIAGARA_API FNiagaraDataInterfaceSetShaderParametersContext
 		ShaderStorage = InShaderStorage;
 	}
 
+private:
+	bool IsStructBoundInternal(const void* StructAddress, uint32 StructSize) const;
+
 private:
 	FRDGBuilder& GraphBuilder;
 	const FNiagaraGpuComputeDispatchInterface& ComputeDispatchInterface;
diff --git a/Engine/Plugins/FX/Niagara/Source/Niagara/Private/NiagaraDataInterface.cpp b/Engine/Plugins/FX/Niagara/Source/Niagara/Private/NiagaraDataInterface.cpp
--- a/Engine/Plugins/FX/Niagara/Source/Niagara/Private/NiagaraDataInterface.cpp
+++ b/Engine/Plugins/FX/Niagara/Source/Niagara/Private/NiagaraDataInterface.cpp
@@ -291,6 +291,41 @@ bool FNiagaraDataInterfaceSetShaderParametersContext::IsParameterBound(const voi
 	return false;
 }
 
+bool FNiagaraDataInterfaceSetShaderParametersContext::IsStructBoundInternal(const void* StructAddress, uint32 StructSize) const
+{
+	const uint16 ByteStart = uint16(uintptr_t(StructAddress) - uintptr_t(BaseParameters));
+	const uint16 ByteEnd = ByteStart + uint16(StructSize);
+
+	// Loop over resources
+	for (const FShaderParameterBindings::FResourceParameter& ResourceParameter : ShaderRef->Bindings.ResourceParameters)
+	{
+		if (ResourceParameter.ByteOffset >= ByteEnd)
+		{
+			break;
+		}
+		if (ResourceParameter.ByteOffset >= ByteStart)
+		{
+			return true;
+		}
+	}
+
+	// Loop over parameters
+	for (const FShaderParameterBindings::FParameter& Parameter : ShaderRef->Bindings.Parameters)
+	{
+		if (Parameter.ByteOffset >= ByteEnd)
+		{
+			break;
+		}
+
+		if (Parameter.ByteOffset >= ByteStart)
+		{
+			return true;
+		}
+	}
+
+	return false;
+}
+
 bool FNiagaraDataInterfaceSetShaderParametersContext::IsOutputStage() const
 {
 	return ComputeInstanceData.IsOutputStage(DataInterfaceProxy, SimStageData.StageIndex);
diff --git a/Engine/Plugins/FX/Niagara/Source/Niagara/Private/NiagaraDataInterfaceRigidMeshCollisionQuery.cpp b/Engine/Plugins/FX/Niagara/Source/Niagara/Private/NiagaraDataInterfaceRigidMeshCollisionQuery.cpp
--- a/Engine/Plugins/FX/Niagara/Source/Niagara/Private/NiagaraDataInterfaceRigidMeshCollisionQuery.cpp
+++ b/Engine/Plugins/FX/Niagara/Source/Niagara/Private/NiagaraDataInterfaceRigidMeshCollisionQuery.cpp
@@ -2,12 +2,14 @@
 
 #include "NiagaraDataInterfaceRigidMeshCollisionQuery.h"
 #include "Algo/ForEach.h"
-#include "Animation/SkeletalMeshActor.h"
-#include "AnimationRuntime.h"
 #include "Components/SkeletalMeshComponent.h"
-#include "DataInterface/NiagaraDistanceFieldParameters.h"
 #include "Engine/Canvas.h"
+#include "Renderer/Private/DistanceFieldLightingShared.h"
+#include "Animation/SkeletalMeshActor.h"
+#include "AnimationRuntime.h"
+#include "SkeletalRenderPublic.h"
 #include "EngineUtils.h"
+
 #include "NiagaraComponent.h"
 #include "NiagaraDataInterfaceUtilities.h"
 #include "NiagaraFunctionLibrary.h"
@@ -15,10 +17,10 @@
 #include "NiagaraGpuComputeDispatchInterface.h"
 #include "NiagaraRenderer.h"
 #include "NiagaraShader.h"
+#include "NiagaraShaderParametersBuilder.h"
 #include "NiagaraSimStageData.h"
 #include "NiagaraSystemInstance.h"
 #include "ShaderParameterUtils.h"
-#include "SkeletalRenderPublic.h"
 
 #define LOCTEXT_NAMESPACE "NiagaraDataInterfaceRigidMeshCollisionQuery"
 DEFINE_LOG_CATEGORY_STATIC(LogRigidMeshCollision, Log, All);
@@ -49,6 +51,22 @@ struct FNiagaraRigidMeshCollisionDIFunctionVersion
 
 //------------------------------------------------------------------------------------------------------------
 
+BEGIN_SHADER_PARAMETER_STRUCT(FShaderParameters,)
+	SHADER_PARAMETER(uint32,				MaxTransforms)
+	SHADER_PARAMETER(uint32,				CurrentOffset)
+	SHADER_PARAMETER(uint32,				PreviousOffset)
+	SHADER_PARAMETER(FUintVector4,			ElementOffsets)
+	SHADER_PARAMETER_SRV(Buffer<float4>,	WorldTransformBuffer)
+	SHADER_PARAMETER_SRV(Buffer<float4>,	InverseTransformBuffer)
+	SHADER_PARAMETER_SRV(Buffer<float4>,	ElementExtentBuffer)
+	SHADER_PARAMETER_SRV(Buffer<uint32>,	PhysicsTypeBuffer)
+	SHADER_PARAMETER_SRV(Buffer<uint32>,	DFIndexBuffer)
+END_SHADER_PARAMETER_STRUCT()
+
+static const TCHAR* TemplateShaderFile = TEXT("/Plugin/FX/Niagara/Private/NiagaraDataInterfaceRigidMeshCollisionQuery.ush");
+
+//------------------------------------------------------------------------------------------------------------
+
 static const FName FindActorsName(TEXT("FindActors"));
 
 //------------------------------------------------------------------------------------------------------------
@@ -103,40 +121,8 @@ bool IsMeshDistanceFieldEnabled()
 
 //------------------------------------------------------------------------------------------------------------
 
-struct FNDIRigidMeshCollisionParametersName
-{
-	FNDIRigidMeshCollisionParametersName(const FString& Suffix)
-	{
-		MaxTransformsName = UNiagaraDataInterfaceRigidMeshCollisionQuery::MaxTransformsName + Suffix;
-		CurrentOffsetName = UNiagaraDataInterfaceRigidMeshCollisionQuery::CurrentOffsetName + Suffix;
-		PreviousOffsetName = UNiagaraDataInterfaceRigidMeshCollisionQuery::PreviousOffsetName + Suffix;
-
-		ElementOffsetsName = UNiagaraDataInterfaceRigidMeshCollisionQuery::ElementOffsetsName + Suffix;
-
-		WorldTransformBufferName = UNiagaraDataInterfaceRigidMeshCollisionQuery::WorldTransformBufferName + Suffix;
-		InverseTransformBufferName = UNiagaraDataInterfaceRigidMeshCollisionQuery::InverseTransformBufferName + Suffix;
-		ElementExtentBufferName = UNiagaraDataInterfaceRigidMeshCollisionQuery::ElementExtentBufferName + Suffix;
-		PhysicsTypeBufferName = UNiagaraDataInterfaceRigidMeshCollisionQuery::PhysicsTypeBufferName + Suffix;		
-		DFIndexBufferName = UNiagaraDataInterfaceRigidMeshCollisionQuery::DFIndexBufferName + Suffix;
-	}
-
-	FString MaxTransformsName;
-	FString CurrentOffsetName;
-	FString PreviousOffsetName;
-
-	FString ElementOffsetsName;
-
-	FString WorldTransformBufferName;
-	FString InverseTransformBufferName;
-	FString ElementExtentBufferName;
-	FString PhysicsTypeBufferName;
-	FString DFIndexBufferName;
-};
-
-//------------------------------------------------------------------------------------------------------------
-
 template<typename BufferType, EPixelFormat PixelFormat>
-void CreateInternalBuffer(FRWBuffer& OutputBuffer, uint32 ElementCount)
+void CreateInternalBuffer(FReadBuffer& OutputBuffer, uint32 ElementCount)
 {
 	if (ElementCount > 0)
 	{
@@ -145,7 +131,7 @@ void CreateInternalBuffer(FRWBuffer& OutputBuffer, uint32 ElementCount)
 }
 
 template<typename BufferType, EPixelFormat PixelFormat>
-void UpdateInternalBuffer(const TArray<BufferType>& InputData, FRWBuffer& OutputBuffer)
+void UpdateInternalBuffer(const TArray<BufferType>& InputData, FReadBuffer& OutputBuffer)
 {
 	uint32 ElementCount = InputData.Num();
 	if (ElementCount > 0 && OutputBuffer.Buffer.IsValid())
@@ -502,18 +488,6 @@ void UpdateInternalArrays(TConstArrayView<AActor*> Actors, TConstArrayView<FName
 
 } // NDIRigidMeshCollisionLocal
 
-const FString UNiagaraDataInterfaceRigidMeshCollisionQuery::MaxTransformsName(TEXT("MaxTransforms_"));
-const FString UNiagaraDataInterfaceRigidMeshCollisionQuery::CurrentOffsetName(TEXT("CurrentOffset_"));
-const FString UNiagaraDataInterfaceRigidMeshCollisionQuery::PreviousOffsetName(TEXT("PreviousOffset_"));
-
-const FString UNiagaraDataInterfaceRigidMeshCollisionQuery::ElementOffsetsName(TEXT("ElementOffsets_"));
-
-const FString UNiagaraDataInterfaceRigidMeshCollisionQuery::WorldTransformBufferName(TEXT("WorldTransformBuffer_"));
-const FString UNiagaraDataInterfaceRigidMeshCollisionQuery::InverseTransformBufferName(TEXT("InverseTransformBuffer_"));
-const FString UNiagaraDataInterfaceRigidMeshCollisionQuery::ElementExtentBufferName(TEXT("ElementExtentBuffer_"));
-const FString UNiagaraDataInterfaceRigidMeshCollisionQuery::PhysicsTypeBufferName(TEXT("PhysicsTypeBuffer_"));
-const FString UNiagaraDataInterfaceRigidMeshCollisionQuery::DFIndexBufferName(TEXT("DFIndexBuffer_"));
-
 //------------------------------------------------------------------------------------------------------------
 
 void FNDIRigidMeshCollisionBuffer::InitRHI()
@@ -654,7 +628,7 @@ void FNDIRigidMeshCollisionData::Update(UNiagaraDataInterfaceRigidMeshCollisionQ
 /** Proxy to send data to gpu */
 struct FNDIRigidMeshCollisionProxy : public FNiagaraDataInterfaceProxy
 {
-	struct FGameThreadData
+	struct FRenderThreadData
 	{
 		FElementOffset ElementOffsets;
 		TArray<FVector4f> WorldTransform;
@@ -669,264 +643,89 @@ struct FNDIRigidMeshCollisionProxy : public FNiagaraDataInterfaceProxy
 	};
 
 	/** Get the size of the data that will be passed to render*/
-	virtual int32 PerInstanceDataPassedToRenderThreadSize() const override;
+	virtual int32 PerInstanceDataPassedToRenderThreadSize() const override
+	{
+		return sizeof(FRenderThreadData);
+	}
 
 	/** Get the data that will be passed to render*/
-	virtual void ConsumePerInstanceDataFromGameThread(void* PerInstanceData, const FNiagaraSystemInstanceID& Instance) override;
+	virtual void ConsumePerInstanceDataFromGameThread(void* PerInstanceData, const FNiagaraSystemInstanceID& Instance) override
+	{
+		check(IsInRenderingThread());
 
-	/** Initialize the Proxy data buffer */
-	void InitializePerInstanceData(const FNiagaraSystemInstanceID& SystemInstance);
+		const FRenderThreadData* SourceData = reinterpret_cast<FRenderThreadData*>(PerInstanceData);
+		FRenderThreadData& TargetData = SystemInstancesToProxyData_RT.FindOrAdd(Instance);
 
-	/** Destroy the proxy data if necessary */
-	void DestroyPerInstanceData(const FNiagaraSystemInstanceID& SystemInstance);
+		if (ensure(SourceData))
+		{
+			TargetData = *SourceData;
+			SourceData->~FRenderThreadData();
+		}
+	}
 
 	/** Launch all pre stage functions */
-	virtual void PreStage(FRHICommandList& RHICmdList, const FNiagaraDataInterfaceStageArgs& Context) override;
-
-	/** List of proxy data for each system instances*/
-	TMap<FNiagaraSystemInstanceID, FGameThreadData> SystemInstancesToProxyData;
-};
-
-int32 FNDIRigidMeshCollisionProxy::PerInstanceDataPassedToRenderThreadSize() const
-{
-	return sizeof(FGameThreadData);
-}
-
-void FNDIRigidMeshCollisionProxy::ConsumePerInstanceDataFromGameThread(void* PerInstanceData, const FNiagaraSystemInstanceID& Instance)
-{
-	check(IsInRenderingThread());
-
-	const FGameThreadData* SourceData = reinterpret_cast<FGameThreadData*>(PerInstanceData);
-	FGameThreadData& TargetData = (SystemInstancesToProxyData.FindOrAdd(Instance));
-
-	if (ensure(SourceData))
+	virtual void PreStage(const FNDIGpuComputePostStageContext& Context) override
 	{
-		TargetData = *SourceData;
-		SourceData->~FGameThreadData();
-	}
-}
-
-void FNDIRigidMeshCollisionProxy::InitializePerInstanceData(const FNiagaraSystemInstanceID& SystemInstance)
-{
-	check(IsInRenderingThread());
-
-	check(!SystemInstancesToProxyData.Contains(SystemInstance));
-	SystemInstancesToProxyData.Add(SystemInstance);
-}
-
-void FNDIRigidMeshCollisionProxy::DestroyPerInstanceData(const FNiagaraSystemInstanceID& SystemInstance)
-{
-	check(IsInRenderingThread());	
-
-	SystemInstancesToProxyData.Remove(SystemInstance);
-}
-
-void FNDIRigidMeshCollisionProxy::PreStage(FRHICommandList& RHICmdList, const FNiagaraDataInterfaceStageArgs& Context)
-{
-	using namespace NDIRigidMeshCollisionLocal;
-
-	check(SystemInstancesToProxyData.Contains(Context.SystemInstanceID));
+		using namespace NDIRigidMeshCollisionLocal;
 
-	FGameThreadData* ProxyData =
-		SystemInstancesToProxyData.Find(Context.SystemInstanceID);
+		check(SystemInstancesToProxyData_RT.Contains(Context.GetSystemInstanceID()));
 
-	if (ProxyData != nullptr && ProxyData->AssetBuffer)
-	{
-		if (Context.SimStageData->bFirstStage)
+		FRenderThreadData* ProxyData = SystemInstancesToProxyData_RT.Find(Context.GetSystemInstanceID());
+		if (ProxyData != nullptr && ProxyData->AssetBuffer)
 		{
-			FRHITransitionInfo Transitions[] = {
-				FRHITransitionInfo(ProxyData->AssetBuffer->WorldTransformBuffer.UAV, ERHIAccess::Unknown, ERHIAccess::UAVCompute),
-				FRHITransitionInfo(ProxyData->AssetBuffer->InverseTransformBuffer.UAV, ERHIAccess::Unknown, ERHIAccess::UAVCompute),
-				FRHITransitionInfo(ProxyData->AssetBuffer->PhysicsTypeBuffer.UAV, ERHIAccess::Unknown, ERHIAccess::UAVCompute),
-				FRHITransitionInfo(ProxyData->AssetBuffer->ElementExtentBuffer.UAV, ERHIAccess::Unknown, ERHIAccess::UAVCompute),
-				FRHITransitionInfo(ProxyData->AssetBuffer->DFIndexBuffer.UAV, ERHIAccess::Unknown, ERHIAccess::UAVCompute)
-			};
-			RHICmdList.Transition(MakeArrayView(Transitions, UE_ARRAY_COUNT(Transitions)));
-
-			UpdateInternalBuffer<FVector4f, EPixelFormat::PF_A32B32G32R32F>(ProxyData->WorldTransform, ProxyData->AssetBuffer->WorldTransformBuffer);
-			UpdateInternalBuffer<FVector4f, EPixelFormat::PF_A32B32G32R32F>(ProxyData->InverseTransform, ProxyData->AssetBuffer->InverseTransformBuffer);
-			UpdateInternalBuffer<FVector4f, EPixelFormat::PF_A32B32G32R32F>(ProxyData->ElementExtent, ProxyData->AssetBuffer->ElementExtentBuffer);
-			UpdateInternalBuffer<uint32, EPixelFormat::PF_R32_UINT>(ProxyData->PhysicsType, ProxyData->AssetBuffer->PhysicsTypeBuffer);
-
-			// the distance field indexing needs to be generated using the scene
-			if (!ProxyData->ComponentIdIndex.IsEmpty() && ProxyData->AssetBuffer->DFIndexBuffer.Buffer.IsValid())
+			if (Context.GetSimStageData().bFirstStage)
 			{
-				const int32 ElementCount = ProxyData->ComponentIdIndex.Num();
-				const uint32 BufferBytes = sizeof(uint32) * ElementCount;
-				void* BufferData = RHILockBuffer(ProxyData->AssetBuffer->DFIndexBuffer.Buffer, 0, BufferBytes, RLM_WriteOnly);
-
-				const FScene* Scene = Context.ComputeDispatchInterface->GetScene();
-				if (Scene && !ProxyData->UniqueComponentIds.IsEmpty())
+				//-OPT: We may be able to avoid updating the buffers all the time
+				UpdateInternalBuffer<FVector4f, EPixelFormat::PF_A32B32G32R32F>(ProxyData->WorldTransform, ProxyData->AssetBuffer->WorldTransformBuffer);
+				UpdateInternalBuffer<FVector4f, EPixelFormat::PF_A32B32G32R32F>(ProxyData->InverseTransform, ProxyData->AssetBuffer->InverseTransformBuffer);
+				UpdateInternalBuffer<FVector4f, EPixelFormat::PF_A32B32G32R32F>(ProxyData->ElementExtent, ProxyData->AssetBuffer->ElementExtentBuffer);
+				UpdateInternalBuffer<uint32, EPixelFormat::PF_R32_UINT>(ProxyData->PhysicsType, ProxyData->AssetBuffer->PhysicsTypeBuffer);
+
+				// the distance field indexing needs to be generated using the scene
+				if (!ProxyData->ComponentIdIndex.IsEmpty() && ProxyData->AssetBuffer->DFIndexBuffer.Buffer.IsValid())
 				{
-					TArray<uint32> UniqueDistanceFieldIndices;
-					UniqueDistanceFieldIndices.Reserve(ProxyData->UniqueComponentIds.Num());
+					const int32 ElementCount = ProxyData->ComponentIdIndex.Num();
+					const uint32 BufferBytes = sizeof(uint32) * ElementCount;
+					void* BufferData = RHILockBuffer(ProxyData->AssetBuffer->DFIndexBuffer.Buffer, 0, BufferBytes, RLM_WriteOnly);
 
-					for (const FPrimitiveComponentId& ComponentId : ProxyData->UniqueComponentIds)
+					const FScene* Scene = Context.GetComputeDispatchInterface().GetScene();
+					if (Scene && !ProxyData->UniqueComponentIds.IsEmpty())
 					{
-						uint32& DistanceFieldIndex = UniqueDistanceFieldIndices.Add_GetRef(INDEX_NONE);
-						const int32 PrimitiveSceneIndex = Scene->PrimitiveComponentIds.Find(ComponentId);
-						if (PrimitiveSceneIndex != INDEX_NONE)
+						TArray<uint32> UniqueDistanceFieldIndices;
+						UniqueDistanceFieldIndices.Reserve(ProxyData->UniqueComponentIds.Num());
+
+						for (const FPrimitiveComponentId& ComponentId : ProxyData->UniqueComponentIds)
 						{
-							const TArray<int32, TInlineAllocator<1>>& DFIndices = Scene->Primitives[PrimitiveSceneIndex]->DistanceFieldInstanceIndices;
-							DistanceFieldIndex = DFIndices.IsEmpty() ? INDEX_NONE : DFIndices[0];
+							uint32& DistanceFieldIndex = UniqueDistanceFieldIndices.Add_GetRef(INDEX_NONE);
+							const int32 PrimitiveSceneIndex = Scene->PrimitiveComponentIds.Find(ComponentId);
+							if (PrimitiveSceneIndex != INDEX_NONE)
+							{
+								const TArray<int32, TInlineAllocator<1>>& DFIndices = Scene->Primitives[PrimitiveSceneIndex]->DistanceFieldInstanceIndices;
+								DistanceFieldIndex = DFIndices.IsEmpty() ? INDEX_NONE : DFIndices[0];
+							}
 						}
-					}
 
-					TArrayView<uint32> BufferView(reinterpret_cast<uint32*>(BufferData), ElementCount);
-					for (int32 ElementIt = 0; ElementIt < ElementCount; ++ElementIt)
+						TArrayView<uint32> BufferView(reinterpret_cast<uint32*>(BufferData), ElementCount);
+						for (int32 ElementIt = 0; ElementIt < ElementCount; ++ElementIt)
+						{
+							const int32 UniqueIdIndex = ProxyData->ComponentIdIndex[ElementIt];
+							BufferView[ElementIt] = UniqueDistanceFieldIndices.IsValidIndex(UniqueIdIndex) ? UniqueDistanceFieldIndices[UniqueIdIndex] : INDEX_NONE;
+						}
+					}
+					else
 					{
-						const int32 UniqueIdIndex = ProxyData->ComponentIdIndex[ElementIt];
-						BufferView[ElementIt] = UniqueDistanceFieldIndices.IsValidIndex(UniqueIdIndex) ? UniqueDistanceFieldIndices[UniqueIdIndex] : INDEX_NONE;
+						FMemory::Memset(BufferData, 0xFF, BufferBytes);
 					}
+					RHIUnlockBuffer(ProxyData->AssetBuffer->DFIndexBuffer.Buffer);
 				}
-				else
-				{
-					FMemory::Memset(BufferData, 0xFF, BufferBytes);
-				}
-				RHIUnlockBuffer(ProxyData->AssetBuffer->DFIndexBuffer.Buffer);
 			}
 		}
 	}
-}
-
-//------------------------------------------------------------------------------------------------------------
-
-struct FNDIRigidMeshCollisionParametersCS : public FNiagaraDataInterfaceParametersCS
-{
-	DECLARE_TYPE_LAYOUT(FNDIRigidMeshCollisionParametersCS, NonVirtual);
-public:
-	void Bind(const FNiagaraDataInterfaceGPUParamInfo& ParameterInfo, const class FShaderParameterMap& ParameterMap)
-	{
-		using namespace NDIRigidMeshCollisionLocal;
-
-		FNDIRigidMeshCollisionParametersName ParamNames(*ParameterInfo.DataInterfaceHLSLSymbol);
-
-		MaxTransforms.Bind(ParameterMap, *ParamNames.MaxTransformsName);
-		CurrentOffset.Bind(ParameterMap, *ParamNames.CurrentOffsetName);
-		PreviousOffset.Bind(ParameterMap, *ParamNames.PreviousOffsetName);
-
-		ElementOffsets.Bind(ParameterMap, *ParamNames.ElementOffsetsName);
 
-		WorldTransformBuffer.Bind(ParameterMap, *ParamNames.WorldTransformBufferName);
-		InverseTransformBuffer.Bind(ParameterMap, *ParamNames.InverseTransformBufferName);
-		ElementExtentBuffer.Bind(ParameterMap, *ParamNames.ElementExtentBufferName);
-		PhysicsTypeBuffer.Bind(ParameterMap, *ParamNames.PhysicsTypeBufferName);
-		DFIndexBuffer.Bind(ParameterMap, *ParamNames.DFIndexBufferName);
-
-		DistanceFieldParameters.Bind(ParameterMap);
-	}
-
-	void Set(FRHICommandList& RHICmdList, const FNiagaraDataInterfaceSetArgs& Context) const
-	{
-		using namespace NDIRigidMeshCollisionLocal;
-
-		check(IsInRenderingThread());
-
-		FRHIComputeShader* ComputeShaderRHI = RHICmdList.GetBoundComputeShader();
-
-		FNDIRigidMeshCollisionProxy* InterfaceProxy =
-			static_cast<FNDIRigidMeshCollisionProxy*>(Context.DataInterface);
-		FNDIRigidMeshCollisionProxy::FGameThreadData* ProxyData =
-			InterfaceProxy->SystemInstancesToProxyData.Find(Context.SystemInstanceID);
-
-		if (ProxyData != nullptr && ProxyData->AssetBuffer && ProxyData->AssetBuffer->IsInitialized())
-		{
-			FNDIRigidMeshCollisionBuffer* AssetBuffer = ProxyData->AssetBuffer;
-
-			FRHITransitionInfo Transitions[] = {
-				FRHITransitionInfo(AssetBuffer->WorldTransformBuffer.UAV, ERHIAccess::Unknown, ERHIAccess::SRVCompute),
-				FRHITransitionInfo(AssetBuffer->InverseTransformBuffer.UAV, ERHIAccess::Unknown, ERHIAccess::SRVCompute),
-				FRHITransitionInfo(AssetBuffer->ElementExtentBuffer.UAV, ERHIAccess::Unknown, ERHIAccess::SRVCompute),
-				FRHITransitionInfo(AssetBuffer->PhysicsTypeBuffer.UAV, ERHIAccess::Unknown, ERHIAccess::SRVCompute),
-				FRHITransitionInfo(AssetBuffer->DFIndexBuffer.UAV, ERHIAccess::Unknown, ERHIAccess::SRVCompute)
-			};
-			RHICmdList.Transition(MakeArrayView(Transitions, UE_ARRAY_COUNT(Transitions)));
-
-			SetSRVParameter(RHICmdList, ComputeShaderRHI, WorldTransformBuffer, AssetBuffer->WorldTransformBuffer.SRV);
-			SetSRVParameter(RHICmdList, ComputeShaderRHI, InverseTransformBuffer, AssetBuffer->InverseTransformBuffer.SRV);
-			SetSRVParameter(RHICmdList, ComputeShaderRHI, ElementExtentBuffer, AssetBuffer->ElementExtentBuffer.SRV);
-			SetSRVParameter(RHICmdList, ComputeShaderRHI, PhysicsTypeBuffer, AssetBuffer->PhysicsTypeBuffer.SRV);
-			SetSRVParameter(RHICmdList, ComputeShaderRHI, DFIndexBuffer, AssetBuffer->DFIndexBuffer.SRV);
-
-
-			SetShaderValue(RHICmdList, ComputeShaderRHI, MaxTransforms, ProxyData->MaxPrimitiveCount * 2);
-			SetShaderValue(RHICmdList, ComputeShaderRHI, CurrentOffset, 0);
-			SetShaderValue(RHICmdList, ComputeShaderRHI, PreviousOffset, ProxyData->MaxPrimitiveCount * 3);
-
-			SetShaderValue(RHICmdList, ComputeShaderRHI, ElementOffsets, ProxyData->ElementOffsets);
-
-			if (DistanceFieldParameters.IsBound())
-			{
-				const FDistanceFieldSceneData* DistanceFieldSceneData = static_cast<const FNiagaraGpuComputeDispatch*>(Context.ComputeDispatchInterface)->GetMeshDistanceFieldParameters();	//-BATCHERTODO:
-
-				if (DistanceFieldSceneData == nullptr)
-				{
-					// UE_LOG(LogRigidMeshCollision, Error, TEXT("Distance fields are not available for use"));
-					// #todo(dmp): for now, we'll disable collisions when distance field data is not available
-					// There is no Dummy distance field data we can use.
-
-					//FDistanceFieldSceneData DummyDistanceFieldSceneData(Context.Shader->GetShaderPlatform());
-					//DistanceFieldParameters.SetEmpty(RHICmdList, ComputeShaderRHI, DummyDistanceFieldSceneData);
-
-					SetSRVParameter(RHICmdList, ComputeShaderRHI, WorldTransformBuffer, FNiagaraRenderer::GetDummyFloatBuffer());
-					SetSRVParameter(RHICmdList, ComputeShaderRHI, InverseTransformBuffer, FNiagaraRenderer::GetDummyFloatBuffer());
-					SetSRVParameter(RHICmdList, ComputeShaderRHI, ElementExtentBuffer, FNiagaraRenderer::GetDummyFloatBuffer());
-					SetSRVParameter(RHICmdList, ComputeShaderRHI, PhysicsTypeBuffer, FNiagaraRenderer::GetDummyIntBuffer());
-					SetSRVParameter(RHICmdList, ComputeShaderRHI, DFIndexBuffer, FNiagaraRenderer::GetDummyIntBuffer());
-
-					static const FElementOffset DummyOffsets(0, 0, 0, 0);
-					SetShaderValue(RHICmdList, ComputeShaderRHI, MaxTransforms, 0);
-					SetShaderValue(RHICmdList, ComputeShaderRHI, CurrentOffset, 0);
-					SetShaderValue(RHICmdList, ComputeShaderRHI, PreviousOffset, 0);
-					SetShaderValue(RHICmdList, ComputeShaderRHI, ElementOffsets, DummyOffsets);
-				}
-				else
-				{
-					DistanceFieldParameters.Set(RHICmdList, ComputeShaderRHI, DistanceFieldSceneData);
-				}
-			}
-		}
-		else
-		{
-			SetSRVParameter(RHICmdList, ComputeShaderRHI, WorldTransformBuffer, FNiagaraRenderer::GetDummyFloatBuffer());
-			SetSRVParameter(RHICmdList, ComputeShaderRHI, InverseTransformBuffer, FNiagaraRenderer::GetDummyFloatBuffer());
-			SetSRVParameter(RHICmdList, ComputeShaderRHI, ElementExtentBuffer, FNiagaraRenderer::GetDummyFloatBuffer());
-			SetSRVParameter(RHICmdList, ComputeShaderRHI, PhysicsTypeBuffer, FNiagaraRenderer::GetDummyIntBuffer());
-			SetSRVParameter(RHICmdList, ComputeShaderRHI, DFIndexBuffer, FNiagaraRenderer::GetDummyIntBuffer());
-
-			static const FElementOffset DummyOffsets(0, 0, 0, 0);
-			SetShaderValue(RHICmdList, ComputeShaderRHI, MaxTransforms, 0);
-			SetShaderValue(RHICmdList, ComputeShaderRHI, CurrentOffset, 0);
-			SetShaderValue(RHICmdList, ComputeShaderRHI, PreviousOffset, 0);
-			SetShaderValue(RHICmdList, ComputeShaderRHI, ElementOffsets, DummyOffsets);
-		}
-	}
-
-	void Unset(FRHICommandList& RHICmdList, const FNiagaraDataInterfaceSetArgs& Context) const
-	{
-	}
-
-private:
-
-	LAYOUT_FIELD(FShaderParameter, MaxTransforms);
-	LAYOUT_FIELD(FShaderParameter, CurrentOffset);
-	LAYOUT_FIELD(FShaderParameter, PreviousOffset);
-
-	LAYOUT_FIELD(FShaderParameter, ElementOffsets);
-
-	LAYOUT_FIELD(FShaderResourceParameter, WorldTransformBuffer);
-	LAYOUT_FIELD(FShaderResourceParameter, InverseTransformBuffer);
-	LAYOUT_FIELD(FShaderResourceParameter, ElementExtentBuffer);
-	LAYOUT_FIELD(FShaderResourceParameter, PhysicsTypeBuffer);
-	LAYOUT_FIELD(FShaderResourceParameter, DFIndexBuffer);
-
-	LAYOUT_FIELD(FDistanceFieldParameters, DistanceFieldParameters);
+	/** List of proxy data for each system instances*/
+	TMap<FNiagaraSystemInstanceID, FRenderThreadData> SystemInstancesToProxyData_RT;
 };
 
-IMPLEMENT_TYPE_LAYOUT(FNDIRigidMeshCollisionParametersCS);
-
-IMPLEMENT_NIAGARA_DI_PARAMETER(UNiagaraDataInterfaceRigidMeshCollisionQuery, FNDIRigidMeshCollisionParametersCS);
-
 //------------------------------------------------------------------------------------------------------------
 
 UNiagaraDataInterfaceRigidMeshCollisionQuery::UNiagaraDataInterfaceRigidMeshCollisionQuery(FObjectInitializer const& ObjectInitializer)
@@ -1110,7 +909,7 @@ void UNiagaraDataInterfaceRigidMeshCollisionQuery::DestroyPerInstanceData(void*
 	ENQUEUE_RENDER_COMMAND(FNiagaraDIDestroyInstanceData) (
 		[ThisProxy, InstanceID = SystemInstance->GetId()](FRHICommandListImmediate& CmdList)
 	{
-		ThisProxy->SystemInstancesToProxyData.Remove(InstanceID);
+		ThisProxy->SystemInstancesToProxyData_RT.Remove(InstanceID);
 	});
 }
 
@@ -1422,156 +1221,20 @@ bool UNiagaraDataInterfaceRigidMeshCollisionQuery::GetFunctionHLSL(const FNiagar
 {
 	using namespace NDIRigidMeshCollisionLocal;
 
-	FNDIRigidMeshCollisionParametersName ParamNames(ParamInfo.DataInterfaceHLSLSymbol);
-
-	TMap<FString, FStringFormatArg> ArgsSample = {
-		{TEXT("InstanceFunctionName"), FunctionInfo.InstanceName},
-		{TEXT("MaxTransformsName"), ParamNames.MaxTransformsName},
-		{TEXT("CurrentOffsetName"), ParamNames.CurrentOffsetName},
-		{TEXT("PreviousOffsetName"), ParamNames.PreviousOffsetName},
-		{TEXT("ElementOffsetsName"), ParamNames.ElementOffsetsName},
-		{TEXT("WorldTransformBufferName"), ParamNames.WorldTransformBufferName},
-		{TEXT("InverseTransformBufferName"), ParamNames.InverseTransformBufferName},
-		{TEXT("ElementExtentBufferName"), ParamNames.ElementExtentBufferName},
-		{TEXT("RigidMeshCollisionContextName"), TEXT("DIRIGIDMESHCOLLISIONQUERY_MAKE_CONTEXT(") + ParamInfo.DataInterfaceHLSLSymbol + TEXT(")")},
-	};
-
-	if (FunctionInfo.DefinitionName == GetNumBoxesName)
-	{
-		static const TCHAR* FormatSample = TEXT(R"(
-		void {InstanceFunctionName}(out int OutNumBoxes)
-		{
-			{RigidMeshCollisionContextName}
-			OutNumBoxes = DIRigidMeshCollision_GetNumBoxes(DIContext);
-		}
-		)");
-		OutHLSL += FString::Format(FormatSample, ArgsSample);
-		return true;
-	}
-	else if (FunctionInfo.DefinitionName == GetNumCapsulesName)
+	if ((FunctionInfo.DefinitionName == GetNumBoxesName) ||
+		(FunctionInfo.DefinitionName == GetNumCapsulesName) ||
+		(FunctionInfo.DefinitionName == GetNumSpheresName) ||
+		(FunctionInfo.DefinitionName == GetClosestPointName) ||
+		(FunctionInfo.DefinitionName == GetClosestElementName) ||
+		(FunctionInfo.DefinitionName == GetElementPointName) ||
+		(FunctionInfo.DefinitionName == GetElementDistanceName) ||
+		(FunctionInfo.DefinitionName == GetClosestDistanceName) ||
+		(FunctionInfo.DefinitionName == GetClosestPointMeshDistanceFieldName) ||
+		(FunctionInfo.DefinitionName == GetClosestPointMeshDistanceFieldAccurateName) ||
+		(FunctionInfo.DefinitionName == GetClosestPointMeshDistanceFieldNoNormalName) )
 	{
-		static const TCHAR* FormatSample = TEXT(R"(
-		void {InstanceFunctionName}(out int OutNumCapsules)
-		{
-			{RigidMeshCollisionContextName}
-			OutNumCapsules = DIRigidMeshCollision_GetNumCapsules(DIContext);
-		}
-		)");
-		OutHLSL += FString::Format(FormatSample, ArgsSample);
-		return true;
-	}
-	else if (FunctionInfo.DefinitionName == GetNumSpheresName)
-	{
-		static const TCHAR* FormatSample = TEXT(R"(
-		void {InstanceFunctionName}(out int OutNumSpheres)
-		{
-			{RigidMeshCollisionContextName}
-			OutNumSpheres = DIRigidMeshCollision_GetNumSpheres(DIContext);
-		}
-		)");
-		OutHLSL += FString::Format(FormatSample, ArgsSample);
-		return true;
-	}
-	else if (FunctionInfo.DefinitionName == GetClosestPointName)
-	{
-		static const TCHAR* FormatSample = TEXT(R"(
-		void {InstanceFunctionName}(in float3 WorldPosition, in float DeltaTime, in float TimeFraction, out float ClosestDistance, out float3 OutClosestPosition, 
-							out float3 OutClosestNormal, out float3 OutClosestVelocity)
-		{
-			{RigidMeshCollisionContextName} DIRigidMeshCollision_GetClosestPoint(DIContext,WorldPosition,DeltaTime,TimeFraction, ClosestDistance,
-				OutClosestPosition,OutClosestNormal,OutClosestVelocity);
-		}
-		)");
-		OutHLSL += FString::Format(FormatSample, ArgsSample);
-		return true;
-	}
-	else if (FunctionInfo.DefinitionName == GetClosestElementName)
-	{
-		static const TCHAR* FormatSample = TEXT(R"(
-		void {InstanceFunctionName}(in float3 WorldPosition, in float TimeFraction, out int OutClosestElement)
-		{
-			{RigidMeshCollisionContextName} DIRigidMeshCollision_GetClosestElement(DIContext,WorldPosition,TimeFraction,
-				OutClosestElement);
-		}
-		)");
-		OutHLSL += FString::Format(FormatSample, ArgsSample);
-		return true;
-	}
-	else if (FunctionInfo.DefinitionName == GetElementPointName)
-	{
-		static const TCHAR* FormatSample = TEXT(R"(
-		void {InstanceFunctionName}(in float3 WorldPosition, in float DeltaTime, in float TimeFraction, in int ElementIndex, out float3 OutClosestPosition, 
-							out float3 OutClosestNormal, out float3 OutClosestVelocity)
-		{
-			{RigidMeshCollisionContextName} DIRigidMeshCollision_GetElementPoint(DIContext,WorldPosition,DeltaTime,TimeFraction,ElementIndex,
-				OutClosestPosition,OutClosestNormal,OutClosestVelocity);
-		}
-		)");
-		OutHLSL += FString::Format(FormatSample, ArgsSample);
-		return true;
-	}
-	else if (FunctionInfo.DefinitionName == GetElementDistanceName)
-	{
-		static const TCHAR* FormatSample = TEXT(R"(
-		void {InstanceFunctionName}(in float3 WorldPosition, in float TimeFraction, in int ElementIndex, out float OutClosestDistance)
-		{
-			{RigidMeshCollisionContextName} DIRigidMeshCollision_GetElementDistance(DIContext,WorldPosition,TimeFraction,ElementIndex,
-				OutClosestDistance);
-		}
-		)");
-		OutHLSL += FString::Format(FormatSample, ArgsSample);
-		return true;
-	}
-	else if (FunctionInfo.DefinitionName == GetClosestDistanceName)
-	{
-		static const TCHAR* FormatSample = TEXT(R"(
-		void {InstanceFunctionName}(in float3 WorldPosition, in float TimeFraction, out float OutClosestDistance)
-		{
-			{RigidMeshCollisionContextName} DIRigidMeshCollision_GetClosestDistance(DIContext,WorldPosition,TimeFraction,OutClosestDistance);
-		}
-		)");
-		OutHLSL += FString::Format(FormatSample, ArgsSample);
 		return true;
 	}
-	else if (FunctionInfo.DefinitionName == GetClosestPointMeshDistanceFieldName)
-	{
-		static const TCHAR* FormatSample = TEXT(R"(
-		void {InstanceFunctionName}(in float3 WorldPosition, in float DeltaTime, in float TimeFraction,  in float MaxDistance, out float ClosestDistance, out float3 OutClosestPosition, 
-							out float3 OutClosestNormal, out float3 OutClosestVelocity, out bool NormalIsValid)
-		{
-			{RigidMeshCollisionContextName} DIRigidMeshCollision_GetClosestPointMeshDistanceField(DIContext,WorldPosition,DeltaTime,TimeFraction, MaxDistance, ClosestDistance,
-				OutClosestPosition,OutClosestNormal,OutClosestVelocity, NormalIsValid);
-		}
-		)");
-		OutHLSL += FString::Format(FormatSample, ArgsSample);
-		return true;
-	}
-	else if (FunctionInfo.DefinitionName == GetClosestPointMeshDistanceFieldAccurateName)
-	{
-	static const TCHAR* FormatSample = TEXT(R"(
-		void {InstanceFunctionName}(in float3 WorldPosition, in float DeltaTime, in float TimeFraction,  in float MaxDistance, out float ClosestDistance, out float3 OutClosestPosition, 
-							out float3 OutClosestNormal, out float3 OutClosestVelocity, out bool NormalIsValid)
-		{
-			{RigidMeshCollisionContextName} DIRigidMeshCollision_GetClosestPointMeshDistanceFieldAccurate(DIContext,WorldPosition,DeltaTime,TimeFraction, MaxDistance, ClosestDistance,
-				OutClosestPosition,OutClosestNormal,OutClosestVelocity, NormalIsValid);
-		}
-		)");
-	OutHLSL += FString::Format(FormatSample, ArgsSample);
-	return true;
-	}
-	else if (FunctionInfo.DefinitionName == GetClosestPointMeshDistanceFieldNoNormalName)
-	{
-	static const TCHAR* FormatSample = TEXT(R"(
-		void {InstanceFunctionName}(in float3 WorldPosition, in float DeltaTime, in float TimeFraction, in float MaxDistance, out float ClosestDistance, out float3 OutClosestPosition, 
-							out float3 OutClosestVelocity)
-		{
-			{RigidMeshCollisionContextName} DIRigidMeshCollision_GetClosestPointMeshDistanceFieldNoNormal(DIContext,WorldPosition,DeltaTime,TimeFraction, MaxDistance, ClosestDistance,
-				OutClosestPosition,OutClosestVelocity);
-		}
-		)");
-	OutHLSL += FString::Format(FormatSample, ArgsSample);
-	return true;
-	}
 	OutHLSL += TEXT("\n");
 	return false;
 }
@@ -1660,17 +1323,133 @@ bool UNiagaraDataInterfaceRigidMeshCollisionQuery::UpgradeFunctionCall(FNiagaraF
 
 void UNiagaraDataInterfaceRigidMeshCollisionQuery::GetCommonHLSL(FString& OutHLSL)
 {
+	Super::GetCommonHLSL(OutHLSL);
+
 	OutHLSL += TEXT("#include \"/Engine/Private/DistanceFieldLightingShared.ush\"\n");
 	OutHLSL += TEXT("#include \"/Engine/Private/MeshDistanceFieldCommon.ush\"\n");
-	OutHLSL += TEXT("#include \"/Plugin/FX/Niagara/Private/NiagaraDataInterfaceRigidMeshCollisionQuery.ush\"\n");			
 }
 
 void UNiagaraDataInterfaceRigidMeshCollisionQuery::GetParameterDefinitionHLSL(const FNiagaraDataInterfaceGPUParamInfo& ParamInfo, FString& OutHLSL)
 {
-	OutHLSL += TEXT("DIRIGIDMESHCOLLISIONQUERY_DECLARE_CONSTANTS(") + ParamInfo.DataInterfaceHLSLSymbol + TEXT(")\n");
+	Super::GetParameterDefinitionHLSL(ParamInfo, OutHLSL);
+
+	TMap<FString, FStringFormatArg> TemplateArgs =
+	{
+		{TEXT("ParameterName"),	ParamInfo.DataInterfaceHLSLSymbol},
+	};
+
+	FString TemplateFile;
+	LoadShaderSourceFile(NDIRigidMeshCollisionLocal::TemplateShaderFile, EShaderPlatform::SP_PCD3D_SM5, &TemplateFile, nullptr);
+	OutHLSL += FString::Format(*TemplateFile, TemplateArgs);
+}
+
+bool UNiagaraDataInterfaceRigidMeshCollisionQuery::AppendCompileHash(FNiagaraCompileHashVisitor* InVisitor) const
+{
+	if (!Super::AppendCompileHash(InVisitor))
+	{
+		return false;
+	}
+
+	FSHAHash Hash = GetShaderFileHash(NDIRigidMeshCollisionLocal::TemplateShaderFile, EShaderPlatform::SP_PCD3D_SM5);
+	InVisitor->UpdateString(TEXT("NiagaraDataInterfaceRigidMeshCollisionQueryHLSLSource"), Hash.ToString());
+	InVisitor->UpdateShaderParameters<NDIRigidMeshCollisionLocal::FShaderParameters>();
+
+	return true;
 }
 #endif
 
+void UNiagaraDataInterfaceRigidMeshCollisionQuery::BuildShaderParameters(FNiagaraShaderParametersBuilder& ShaderParametersBuilder) const
+{
+	ShaderParametersBuilder.AddNestedStruct<NDIRigidMeshCollisionLocal::FShaderParameters>();
+	ShaderParametersBuilder.AddIncludedStruct<FDistanceFieldObjectBufferParameters>();
+	ShaderParametersBuilder.AddIncludedStruct<FDistanceFieldAtlasParameters>();
+}
+
+void UNiagaraDataInterfaceRigidMeshCollisionQuery::SetShaderParameters(const FNiagaraDataInterfaceSetShaderParametersContext& Context) const
+{
+	FRDGBuilder& GraphBuilder = Context.GetGraphBuilder();
+	const FNDIRigidMeshCollisionProxy& InterfaceProxy = Context.GetProxy<FNDIRigidMeshCollisionProxy>();
+	const FNDIRigidMeshCollisionProxy::FRenderThreadData* ProxyData = InterfaceProxy.SystemInstancesToProxyData_RT.Find(Context.GetSystemInstanceID());
+
+	NDIRigidMeshCollisionLocal::FShaderParameters* ShaderParameters = Context.GetParameterNestedStruct<NDIRigidMeshCollisionLocal::FShaderParameters>();
+	FDistanceFieldObjectBufferParameters* ShaderDistanceFieldObjectParameters = Context.GetParameterIncludedStruct<FDistanceFieldObjectBufferParameters>();
+	FDistanceFieldAtlasParameters* ShaderDistanceFieldAtlasParameters = Context.GetParameterIncludedStruct<FDistanceFieldAtlasParameters>();
+
+	const bool bDistanceFieldDataBound = Context.IsStructBound<FDistanceFieldObjectBufferParameters>(ShaderDistanceFieldObjectParameters) || Context.IsStructBound<FDistanceFieldAtlasParameters>(ShaderDistanceFieldAtlasParameters);
+	bool bDistanceFieldDataSet = false;
+
+	if (ProxyData != nullptr && ProxyData->AssetBuffer != nullptr && ProxyData->AssetBuffer->IsInitialized())
+	{
+		FNDIRigidMeshCollisionBuffer* AssetBuffer = ProxyData->AssetBuffer;
+
+		ShaderParameters->WorldTransformBuffer = FNiagaraRenderer::GetSrvOrDefaultFloat4(AssetBuffer->WorldTransformBuffer.SRV);
+		ShaderParameters->InverseTransformBuffer = FNiagaraRenderer::GetSrvOrDefaultFloat4(AssetBuffer->InverseTransformBuffer.SRV);
+		ShaderParameters->ElementExtentBuffer = FNiagaraRenderer::GetSrvOrDefaultFloat4(AssetBuffer->ElementExtentBuffer.SRV);
+		ShaderParameters->PhysicsTypeBuffer = FNiagaraRenderer::GetSrvOrDefaultUInt(AssetBuffer->PhysicsTypeBuffer.SRV);
+		ShaderParameters->DFIndexBuffer = FNiagaraRenderer::GetSrvOrDefaultUInt(AssetBuffer->DFIndexBuffer.SRV);
+
+		ShaderParameters->MaxTransforms = ProxyData->MaxPrimitiveCount * 2;
+		ShaderParameters->CurrentOffset = 0;
+		ShaderParameters->PreviousOffset = ProxyData->MaxPrimitiveCount * 3;
+
+		ShaderParameters->ElementOffsets.X = ProxyData->ElementOffsets.BoxOffset;
+		ShaderParameters->ElementOffsets.Y = ProxyData->ElementOffsets.SphereOffset;
+		ShaderParameters->ElementOffsets.Z = ProxyData->ElementOffsets.CapsuleOffset;
+		ShaderParameters->ElementOffsets.W = ProxyData->ElementOffsets.NumElements;
+
+		if (bDistanceFieldDataBound)
+		{
+			const FDistanceFieldSceneData* DistanceFieldSceneData = static_cast<const FNiagaraGpuComputeDispatch&>(Context.GetComputeDispatchInterface()).GetMeshDistanceFieldParameters();	//-BATCHERTODO:
+			if (DistanceFieldSceneData != nullptr)
+			{
+				bDistanceFieldDataSet = true;
+
+				*ShaderDistanceFieldObjectParameters = DistanceField::SetupObjectBufferParameters(GraphBuilder, *DistanceFieldSceneData);
+				*ShaderDistanceFieldAtlasParameters = DistanceField::SetupAtlasParameters(GraphBuilder, *DistanceFieldSceneData);
+			}
+		}
+	}
+	else
+	{
+		ShaderParameters->WorldTransformBuffer = FNiagaraRenderer::GetDummyFloat4Buffer();
+		ShaderParameters->InverseTransformBuffer = FNiagaraRenderer::GetDummyFloat4Buffer();
+		ShaderParameters->ElementExtentBuffer = FNiagaraRenderer::GetDummyFloat4Buffer();
+		ShaderParameters->PhysicsTypeBuffer = FNiagaraRenderer::GetDummyUIntBuffer();
+		ShaderParameters->DFIndexBuffer = FNiagaraRenderer::GetDummyUIntBuffer();
+
+		ShaderParameters->MaxTransforms = 0;
+		ShaderParameters->CurrentOffset = 0;
+		ShaderParameters->PreviousOffset = 0;
+		ShaderParameters->ElementOffsets = FUintVector4(0, 0, 0, 0);
+	}
+
+	if (bDistanceFieldDataBound && !bDistanceFieldDataSet)
+	{
+		ShaderDistanceFieldObjectParameters->SceneObjectBounds				= FNiagaraRenderer::GetDummyFloat4Buffer();
+		ShaderDistanceFieldObjectParameters->SceneObjectData				= FNiagaraRenderer::GetDummyFloat4Buffer();
+		ShaderDistanceFieldObjectParameters->NumSceneObjects				= 0;
+		ShaderDistanceFieldObjectParameters->SceneHeightfieldObjectBounds	= Context.GetComputeDispatchInterface().GetEmptyBufferSRV(GraphBuilder, PF_R8G8B8A8);
+		ShaderDistanceFieldObjectParameters->SceneHeightfieldObjectData		= Context.GetComputeDispatchInterface().GetEmptyBufferSRV(GraphBuilder, PF_R8G8B8A8);
+		ShaderDistanceFieldObjectParameters->NumSceneHeightfieldObjects		= 0;
+
+		ShaderDistanceFieldAtlasParameters->SceneDistanceFieldAssetData						= Context.GetComputeDispatchInterface().GetEmptyBufferSRV(GraphBuilder, PF_R8G8B8A8);
+		ShaderDistanceFieldAtlasParameters->DistanceFieldIndirectionTable					= Context.GetComputeDispatchInterface().GetEmptyBufferSRV(GraphBuilder, PF_R32_UINT);
+		ShaderDistanceFieldAtlasParameters->DistanceFieldIndirection2Table					= Context.GetComputeDispatchInterface().GetEmptyBufferSRV(GraphBuilder, PF_R8G8B8A8);
+		ShaderDistanceFieldAtlasParameters->DistanceFieldIndirectionAtlas					= Context.GetComputeDispatchInterface().GetBlackTexture(GraphBuilder, ETextureDimension::Texture3D);
+		ShaderDistanceFieldAtlasParameters->DistanceFieldBrickTexture						= GBlackVolumeTexture->GetTextureRHI();
+		ShaderDistanceFieldAtlasParameters->DistanceFieldSampler							= TStaticSamplerState<SF_Bilinear, AM_Clamp, AM_Clamp, AM_Clamp>::GetRHI();
+		ShaderDistanceFieldAtlasParameters->DistanceFieldBrickSize							= FVector3f::ZeroVector;
+		ShaderDistanceFieldAtlasParameters->DistanceFieldUniqueDataBrickSize				= FVector3f::ZeroVector;
+		ShaderDistanceFieldAtlasParameters->DistanceFieldBrickAtlasSizeInBricks				= FIntVector::ZeroValue;
+		ShaderDistanceFieldAtlasParameters->DistanceFieldBrickAtlasMask						= FIntVector::ZeroValue;
+		ShaderDistanceFieldAtlasParameters->DistanceFieldBrickAtlasSizeLog2					= FIntVector::ZeroValue;
+		ShaderDistanceFieldAtlasParameters->DistanceFieldBrickAtlasTexelSize				= FVector3f::ZeroVector;
+		ShaderDistanceFieldAtlasParameters->DistanceFieldBrickAtlasHalfTexelSize			= FVector3f::ZeroVector;
+		ShaderDistanceFieldAtlasParameters->DistanceFieldBrickOffsetToAtlasUVScale			= FVector3f::ZeroVector;
+		ShaderDistanceFieldAtlasParameters->DistanceFieldUniqueDataBrickSizeInAtlasTexels	= FVector3f::ZeroVector;
+	}
+}
+
 #if WITH_EDITOR
 void UNiagaraDataInterfaceRigidMeshCollisionQuery::ValidateFunction(const FNiagaraFunctionSignature& Function, TArray<FText>& OutValidationErrors)
 {
@@ -1689,7 +1468,7 @@ void UNiagaraDataInterfaceRigidMeshCollisionQuery::ValidateFunction(const FNiaga
 void UNiagaraDataInterfaceRigidMeshCollisionQuery::ProvidePerInstanceDataForRenderThread(void* DataForRenderThread, void* PerInstanceData, const FNiagaraSystemInstanceID& SystemInstance)
 {
 	const FNDIRigidMeshCollisionData* GameThreadData = reinterpret_cast<const FNDIRigidMeshCollisionData*>(PerInstanceData);
-	FNDIRigidMeshCollisionProxy::FGameThreadData* RenderThreadData = new(DataForRenderThread) FNDIRigidMeshCollisionProxy::FGameThreadData();
+	FNDIRigidMeshCollisionProxy::FRenderThreadData* RenderThreadData = new(DataForRenderThread) FNDIRigidMeshCollisionProxy::FRenderThreadData();
 
 	if (ensure(GameThreadData != nullptr && RenderThreadData != nullptr))
 	{
@@ -1891,4 +1670,4 @@ void UNiagaraDIRigidMeshCollisionFunctionLibrary::SetSourceActors(UNiagaraCompon
 	}
 }
 
-#undef LOCTEXT_NAMESPACE
\ No newline at end of file
+#undef LOCTEXT_NAMESPACE
diff --git a/Engine/Plugins/FX/Niagara/Source/Niagara/Public/NiagaraDataInterfaceRigidMeshCollisionQuery.h b/Engine/Plugins/FX/Niagara/Source/Niagara/Public/NiagaraDataInterfaceRigidMeshCollisionQuery.h
--- a/Engine/Plugins/FX/Niagara/Source/Niagara/Public/NiagaraDataInterfaceRigidMeshCollisionQuery.h
+++ b/Engine/Plugins/FX/Niagara/Source/Niagara/Public/NiagaraDataInterfaceRigidMeshCollisionQuery.h
@@ -76,19 +76,19 @@ struct FNDIRigidMeshCollisionBuffer : public FRenderResource
 	virtual FString GetFriendlyName() const override { return TEXT("FNDIRigidMeshCollisionBuffer"); }
 
 	/** World transform buffer */
-	FRWBuffer WorldTransformBuffer;
+	FReadBuffer WorldTransformBuffer;
 
 	/** Inverse transform buffer*/
-	FRWBuffer InverseTransformBuffer;
+	FReadBuffer InverseTransformBuffer;
 
 	/** Element extent buffer */
-	FRWBuffer ElementExtentBuffer;
+	FReadBuffer ElementExtentBuffer;
 
 	/** Physics type buffer */
-	FRWBuffer PhysicsTypeBuffer;
+	FReadBuffer PhysicsTypeBuffer;
 
 	/** Distance field index buffer */
-	FRWBuffer DFIndexBuffer;
+	FReadBuffer DFIndexBuffer;
 
 	/** Max number of primitives */
 	uint32 MaxNumPrimitives;
@@ -162,9 +162,6 @@ class UNiagaraDataInterfaceRigidMeshCollisionQuery : public UNiagaraDataInterfac
 	GENERATED_UCLASS_BODY()
 
 public:
-
-	DECLARE_NIAGARA_DI_PARAMETER();
-
 #if WITH_EDITORONLY_DATA
 	UPROPERTY()
 	FString Tag_DEPRECATED = TEXT("");
@@ -230,11 +227,17 @@ public:
 #if WITH_EDITORONLY_DATA
 	virtual void GetCommonHLSL(FString& OutHLSL) override;
 	virtual void GetParameterDefinitionHLSL(const FNiagaraDataInterfaceGPUParamInfo& ParamInfo, FString& OutHLSL) override;
+	virtual bool AppendCompileHash(FNiagaraCompileHashVisitor* InVisitor) const override;
+
 	virtual bool GetFunctionHLSL(const FNiagaraDataInterfaceGPUParamInfo& ParamInfo, const FNiagaraDataInterfaceGeneratedFunction& FunctionInfo, int FunctionInstanceIndex, FString& OutHLSL) override;
 	virtual bool UpgradeFunctionCall(FNiagaraFunctionSignature& FunctionSignature) override;
 
 	virtual void ValidateFunction(const FNiagaraFunctionSignature& Function, TArray<FText>& OutValidationErrors) override;
 #endif
+	virtual bool UseLegacyShaderBindings() const override { return false; }
+	virtual void BuildShaderParameters(FNiagaraShaderParametersBuilder& ShaderParametersBuilder) const override;
+	virtual void SetShaderParameters(const FNiagaraDataInterfaceSetShaderParametersContext& Context) const override;
+
 	virtual void ProvidePerInstanceDataForRenderThread(void* DataForRenderThread, void* PerInstanceData, const FNiagaraSystemInstanceID& SystemInstance) override;
 
 	bool UpdateSourceActors(FNiagaraSystemInstance* SystemInstance, FNDIRigidMeshCollisionData& InstanceData) const;
@@ -245,33 +248,6 @@ public:
 
 	void FindActorsCPU(FVectorVMExternalFunctionContext& Context);
 
-	/** Name of element offsets */
-	static const FString MaxTransformsName;
-
-	/** Name of element offsets */
-	static const FString CurrentOffsetName;
-
-	/** Name of element offsets */
-	static const FString PreviousOffsetName;
-
-	/** Name of element offsets */
-	static const FString ElementOffsetsName;
-
-	/** Name of the world transform buffer */
-	static const FString WorldTransformBufferName;
-
-	/** Name of the inverse transform buffer */
-	static const FString InverseTransformBufferName;
-
-	/** Name of the element extent buffer */
-	static const FString ElementExtentBufferName;
-
-	/** Name of the physics type buffer */
-	static const FString PhysicsTypeBufferName;
-
-	/** Name of the DF Index type buffer */
-	static const FString DFIndexBufferName;
-
 protected:
 	/** Copy one niagara DI to this */
 	virtual bool CopyToInternal(UNiagaraDataInterface* Destination) const override;
diff --git a/Engine/Plugins/FX/Niagara/Source/NiagaraShader/Public/NiagaraShaderParametersBuilder.h b/Engine/Plugins/FX/Niagara/Source/NiagaraShader/Public/NiagaraShaderParametersBuilder.h
--- a/Engine/Plugins/FX/Niagara/Source/NiagaraShader/Public/NiagaraShaderParametersBuilder.h
+++ b/Engine/Plugins/FX/Niagara/Source/NiagaraShader/Public/NiagaraShaderParametersBuilder.h
@@ -15,21 +15,36 @@ public:
 	{
 	}
 
+	// Adds a loose parameter that is scoped to the data interface
+	// i.e. if the parameter was called "MyFloat" the shader variable would be "UniqueDataInterfaceName_MyFloat"
 	template<typename T> void AddLooseParam(const TCHAR* Name)
 	{
 		LooseNames.Emplace(FString::Printf(TEXT("%s_%s"), *GPUParamInfo.DataInterfaceHLSLSymbol, Name));
 		MetadataBuilder.AddParam<T>(*LooseNames.Last());
 	}
+
+	// Adds a loose array parameter that is scoped to the data interface
+	// i.e. if the parameter was called "MyFloat" the shader variable would be "UniqueDataInterfaceName_MyFloat"
 	template<typename T> void AddLooseParamArray(const TCHAR* Name, int32 NumElements)
 	{
 		LooseNames.Emplace(FString::Printf(TEXT("%s_%s"), *GPUParamInfo.DataInterfaceHLSLSymbol, Name));
 		MetadataBuilder.AddParamArray<T>(*LooseNames.Last(), NumElements);
 	}
+
+	// Adds a shader parameters structure that is scoped to the data interface
+	// i.e. if the structured contained "MyFloat" the shader variable would be "UniqueDataInterfaceName_MyFloat"
 	template<typename T> void AddNestedStruct()
 	{
 		MetadataBuilder.AddNestedStruct<T>(*GPUParamInfo.DataInterfaceHLSLSymbol);
 	}
 
+	// Adds a shadera parameters structure that is global in scope
+	// i.e. if the structure contained "MyFloat" the shader variable would be named "MyFloat"
+	template<typename T> void AddIncludedStruct()
+	{
+		MetadataBuilder.AddIncludedStruct<T>();
+	}
+
 	TConstArrayView<FNiagaraDataInterfaceGeneratedFunction> GetGeneratedFunctions() const { return GPUParamInfo.GeneratedFunctions; }
 
 private:
diff --git a/Engine/Source/Runtime/RenderCore/Private/ShaderParameterMetadataBuilder.cpp b/Engine/Source/Runtime/RenderCore/Private/ShaderParameterMetadataBuilder.cpp
--- a/Engine/Source/Runtime/RenderCore/Private/ShaderParameterMetadataBuilder.cpp
+++ b/Engine/Source/Runtime/RenderCore/Private/ShaderParameterMetadataBuilder.cpp
@@ -26,6 +26,29 @@ void FShaderParametersMetadataBuilder::AddReferencedStruct(
 	NextMemberOffset += Align(StructMetadata->GetSize(), SHADER_PARAMETER_STRUCT_ALIGNMENT);
 }
 
+void FShaderParametersMetadataBuilder::AddIncludedStruct(
+	const FShaderParametersMetadata* StructMetadata,
+	EShaderPrecisionModifier::Type Precision
+)
+{
+	NextMemberOffset = Align(NextMemberOffset, SHADER_PARAMETER_STRUCT_ALIGNMENT);
+
+	new(Members) FShaderParametersMetadata::FMember(
+		TEXT(""),
+		StructMetadata->GetStructTypeName(),
+		__LINE__,
+		NextMemberOffset,
+		UBMT_INCLUDED_STRUCT,
+		Precision,
+		1,
+		1,
+		0,
+		StructMetadata
+	);
+
+	NextMemberOffset += Align(StructMetadata->GetSize(), SHADER_PARAMETER_STRUCT_ALIGNMENT);
+}
+
 uint32 FShaderParametersMetadataBuilder::AddNestedStruct(
 	const TCHAR* Name,
 	const FShaderParametersMetadata* StructMetadata,
diff --git a/Engine/Source/Runtime/RenderCore/Public/ShaderParameterMetadataBuilder.h b/Engine/Source/Runtime/RenderCore/Public/ShaderParameterMetadataBuilder.h
--- a/Engine/Source/Runtime/RenderCore/Public/ShaderParameterMetadataBuilder.h
+++ b/Engine/Source/Runtime/RenderCore/Public/ShaderParameterMetadataBuilder.h
@@ -88,6 +88,19 @@ public:
 		EShaderPrecisionModifier::Type Precision = EShaderPrecisionModifier::Float
 		);
 
+	template<typename T>
+	void AddIncludedStruct(
+		EShaderPrecisionModifier::Type Precision = EShaderPrecisionModifier::Float
+	)
+	{
+		AddIncludedStruct(TShaderParameterStructTypeInfo<T>::GetStructMetadata(), Precision);
+	}
+
+	void AddIncludedStruct(
+		const FShaderParametersMetadata* StructMetadata,
+		EShaderPrecisionModifier::Type Precision = EShaderPrecisionModifier::Float
+	);
+
 	template<typename T>
 	uint32 AddNestedStruct(
 		const TCHAR* Name,
diff --git a/Engine/Source/Runtime/Renderer/Private/DistanceFieldLightingShared.h b/Engine/Source/Runtime/Renderer/Private/DistanceFieldLightingShared.h
--- a/Engine/Source/Runtime/Renderer/Private/DistanceFieldLightingShared.h
+++ b/Engine/Source/Runtime/Renderer/Private/DistanceFieldLightingShared.h
@@ -125,8 +125,8 @@ namespace DistanceField
 		FBox LumenBounds;
 	};
 
-	FDistanceFieldObjectBufferParameters SetupObjectBufferParameters(FRDGBuilder& GraphBuilder, const FDistanceFieldSceneData& DistanceFieldSceneData);
-	FDistanceFieldAtlasParameters SetupAtlasParameters(FRDGBuilder& GraphBuilder, const FDistanceFieldSceneData& DistanceFieldSceneData);
+	RENDERER_API FDistanceFieldObjectBufferParameters SetupObjectBufferParameters(FRDGBuilder& GraphBuilder, const FDistanceFieldSceneData& DistanceFieldSceneData);
+	RENDERER_API FDistanceFieldAtlasParameters SetupAtlasParameters(FRDGBuilder& GraphBuilder, const FDistanceFieldSceneData& DistanceFieldSceneData);
 	void BuildUpdateTrackingBounds(const TArray<FViewInfo>& Views, DistanceField::FUpdateTrackingBounds& UpdateTrackingBounds);
 };
 
