- Only attempt to create the SRV if the buffer was created with SRV access

diff --git a/Engine/Plugins/FX/Niagara/Source/Niagara/Private/DataInterface/NiagaraDataInterfaceStaticMesh.cpp b/Engine/Plugins/FX/Niagara/Source/Niagara/Private/DataInterface/NiagaraDataInterfaceStaticMesh.cpp
--- a/Engine/Plugins/FX/Niagara/Source/Niagara/Private/DataInterface/NiagaraDataInterfaceStaticMesh.cpp
+++ b/Engine/Plugins/FX/Niagara/Source/Niagara/Private/DataInterface/NiagaraDataInterfaceStaticMesh.cpp
@@ -276,6 +276,11 @@ namespace NDIStaticMeshLocal
 		FNiagaraSystemInstanceID			SystemInstanceID;
 		TRefCountPtr<const FStaticMeshLODResources>		LODResource = nullptr;
 
+#if DO_CHECK
+		FName								SystemFName;
+		FName								StaticMeshFName;
+#endif
+
 		int32								NumFilteredTriangles = 0;
 		int32								NumUnfilteredTriangles = 0;
 
@@ -340,7 +345,17 @@ namespace NDIStaticMeshLocal
 			// Gather mesh buffers
 			if (GpuInitializeData.LODResource)
 			{
-				MeshIndexBufferSRV = RHICreateShaderResourceView(GpuInitializeData.LODResource->IndexBuffer.IndexBufferRHI);
+				FBufferRHIRef IndexBufferRHIRef = GpuInitializeData.LODResource->IndexBuffer.IndexBufferRHI;
+				const bool bCanCreateIndexSRV = IndexBufferRHIRef.IsValid() && ((IndexBufferRHIRef->GetUsage() & EBufferUsageFlags::ShaderResource) == EBufferUsageFlags::ShaderResource);
+
+			#if DO_CHECK
+				if (bCanCreateIndexSRV == false)
+				{
+					UE_LOG(LogNiagara, Log, TEXT("NiagaraStaticMeshDataInterface used by GPU emitter but does not have SRV access on this platform.  Enable CPU access to fix this issue. System: %s, Mesh: %s"), *GpuInitializeData.SystemFName.ToString(), *GpuInitializeData.StaticMeshFName.ToString());
+				}
+			#endif
+
+				MeshIndexBufferSRV = bCanCreateIndexSRV ? RHICreateShaderResourceView(IndexBufferRHIRef) : nullptr;
 				MeshPositionBufferSRV = GpuInitializeData.LODResource->VertexBuffers.PositionVertexBuffer.GetSRV();
 				MeshTangentBufferSRV = GpuInitializeData.LODResource->VertexBuffers.StaticMeshVertexBuffer.GetTangentsSRV();
 				MeshUVBufferSRV = GpuInitializeData.LODResource->VertexBuffers.StaticMeshVertexBuffer.GetTexCoordsSRV();
@@ -1484,6 +1499,14 @@ bool UNiagaraDataInterfaceStaticMesh::InitPerInstanceData(void* PerInstanceData,
 		GpuInitializeData->LODResource = InstanceData->GetCurrentFirstLOD();
 		if ( GpuInitializeData->LODResource )
 		{
+		#if DO_CHECK
+			GpuInitializeData->SystemFName = SystemInstance->GetSystem()->GetFName();
+			if ( UStaticMesh* StaticMesh = InstanceData->StaticMeshWeakPtr.Get() )
+			{
+				GpuInitializeData->StaticMeshFName = StaticMesh->GetFName();
+			}
+		#endif
+
 			GpuInitializeData->NumFilteredTriangles = InstanceData->NumFilteredTriangles;
 			GpuInitializeData->NumUnfilteredTriangles = InstanceData->NumUnfilteredTriangles;
 
