- Sprite Renderer Info DataInterface for reading information from a sprite renderer inside script

diff --git a/Engine/Plugins/FX/Niagara/Shaders/Private/NiagaraDataInterfaceSpriteRendererInfoTemplate.ush b/Engine/Plugins/FX/Niagara/Shaders/Private/NiagaraDataInterfaceSpriteRendererInfoTemplate.ush
new file mode 100644
--- /dev/null
+++ b/Engine/Plugins/FX/Niagara/Shaders/Private/NiagaraDataInterfaceSpriteRendererInfoTemplate.ush
@@ -0,0 +1,35 @@
+// Copyright Epic Games, Inc. All Rights Reserved.
+
+uint	{ParameterName}_bIsValid;
+int		{ParameterName}_SourceMode;
+int		{ParameterName}_Alignment;
+int		{ParameterName}_FacingMode;
+uint	{ParameterName}_bSubImageBlend;
+float2	{ParameterName}_SubImageSize;
+
+void IsValid_{ParameterName}(out bool bIsValid)
+{
+	bIsValid = {ParameterName}_bIsValid != 0;
+}
+
+void GetSourceMode_{ParameterName}(out int SourceMode)
+{
+	SourceMode = {ParameterName}_SourceMode;
+}
+
+void GetAlignment_{ParameterName}(out int Alignment)
+{
+	Alignment = {ParameterName}_Alignment;
+}
+
+void GetFacingMode_{ParameterName}(out int FacingMode)
+{
+	FacingMode = {ParameterName}_FacingMode;
+}
+
+void GetSubUVDetails_{ParameterName}(out bool bBlendEnabled, out float2 SubImageSize)
+{
+	bBlendEnabled = {ParameterName}_bSubImageBlend != 0;
+	SubImageSize = {ParameterName}_SubImageSize;
+}
+
diff --git a/Engine/Plugins/FX/Niagara/Source/Niagara/Classes/NiagaraDataInterfaceSpriteRendererInfo.h b/Engine/Plugins/FX/Niagara/Source/Niagara/Classes/NiagaraDataInterfaceSpriteRendererInfo.h
new file mode 100644
--- /dev/null
+++ b/Engine/Plugins/FX/Niagara/Source/Niagara/Classes/NiagaraDataInterfaceSpriteRendererInfo.h
@@ -0,0 +1,68 @@
+// Copyright Epic Games, Inc. All Rights Reserved.
+#pragma once
+
+#include "NiagaraCommon.h"
+#include "NiagaraShared.h"
+#include "NiagaraDataInterface.h"
+#include "NiagaraParameterStore.h"
+#include "NiagaraDataInterfaceSpriteRendererInfo.generated.h"
+
+class UNiagaraSpriteRendererProperties;
+
+/** This Data Interface can be used to query information about the sprite renderers of an emitter */
+UCLASS(EditInlineNew, Category = "Sprite Particles", meta = (DisplayName = "Sprite Renderer Info"))
+class NIAGARA_API UNiagaraDataInterfaceSpriteRendererInfo : public UNiagaraDataInterface
+{
+	GENERATED_UCLASS_BODY()
+
+public:
+	UNiagaraSpriteRendererProperties* GetSpriteRenderer() const { return SpriteRenderer; }
+
+	//UObject Interface
+	virtual void PostInitProperties() override;
+	virtual void PostLoad() override;
+#if WITH_EDITOR	
+	virtual void PreEditChange(FProperty* PropertyAboutToChange) override;
+	virtual void PostEditChangeProperty(struct FPropertyChangedEvent& PropertyChangedEvent) override;
+#endif 
+	//UObject Interface End
+
+	void OnSpriteRendererChanged(UNiagaraSpriteRendererProperties* NewSpriteRenderer);
+
+	//UNiagaraDataInterface Interface
+	virtual void GetFunctions(TArray<FNiagaraFunctionSignature>& OutFunctions) override;
+	virtual void GetVMExternalFunction(const FVMExternalFunctionBindingInfo& BindingInfo, void* InstanceData, FVMExternalFunction& OutFunc) override;
+	virtual bool Equals(const UNiagaraDataInterface* Other) const override;
+#if WITH_EDITORONLY_DATA
+	virtual bool AppendCompileHash(FNiagaraCompileHashVisitor* InVisitor) const override;
+	virtual void GetParameterDefinitionHLSL(const FNiagaraDataInterfaceGPUParamInfo& ParamInfo, FString& OutHLSL) override;
+	virtual bool GetFunctionHLSL(const FNiagaraDataInterfaceGPUParamInfo& ParamInfo, const FNiagaraDataInterfaceGeneratedFunction& FunctionInfo, int FunctionInstanceIndex, FString& OutHLSL) override;
+#endif
+	virtual bool UseLegacyShaderBindings() const override { return false; }
+	virtual void BuildShaderParameters(FNiagaraShaderParametersBuilder& ShaderParametersBuilder) const override;
+	virtual void SetShaderParameters(const FNiagaraDataInterfaceSetShaderParametersContext& Context) const override;
+
+	virtual bool CanExecuteOnTarget(ENiagaraSimTarget Target) const override { return true; }
+#if WITH_EDITOR
+	virtual void GetFeedback(UNiagaraSystem* InAsset, UNiagaraComponent* InComponent, TArray<FNiagaraDataInterfaceError>& OutErrors, TArray<FNiagaraDataInterfaceFeedback>& OutWarnings, TArray<FNiagaraDataInterfaceFeedback>& OutInfo) override;
+#endif
+	//UNiagaraDataInterface Interface	
+
+protected:
+	virtual bool CopyToInternal(UNiagaraDataInterface* Destination) const override;
+	virtual void PushToRenderThreadImpl() override;
+
+	void VMIsValid(FVectorVMExternalFunctionContext& Context);
+	void VMGetSourceMode(FVectorVMExternalFunctionContext& Context);
+	void VMGetAlignment(FVectorVMExternalFunctionContext& Context);
+	void VMGetFacingMode(FVectorVMExternalFunctionContext& Context);
+	void VMGetSubUVDetails(FVectorVMExternalFunctionContext& Context);
+
+	/** The name of the sprite renderer */
+	UPROPERTY(EditAnywhere, Category = "Source")
+	TObjectPtr<UNiagaraSpriteRendererProperties> SpriteRenderer;
+
+#if WITH_EDITOR
+	FDelegateHandle OnSpriteRendererChangedHandle;
+#endif
+};
diff --git a/Engine/Plugins/FX/Niagara/Source/Niagara/Private/NiagaraDataInterfaceSpriteRendererInfo.cpp b/Engine/Plugins/FX/Niagara/Source/Niagara/Private/NiagaraDataInterfaceSpriteRendererInfo.cpp
new file mode 100644
--- /dev/null
+++ b/Engine/Plugins/FX/Niagara/Source/Niagara/Private/NiagaraDataInterfaceSpriteRendererInfo.cpp
@@ -0,0 +1,360 @@
+// Copyright Epic Games, Inc. All Rights Reserved.
+#include "NiagaraDataInterfaceSpriteRendererInfo.h"
+#include "NiagaraSystem.h"
+#include "NiagaraEmitter.h"
+#include "NiagaraSpriteRendererProperties.h"
+#include "ShaderParameterUtils.h"
+#include "NiagaraStats.h"
+#include "NiagaraRenderer.h"
+#include "NiagaraShaderParametersBuilder.h"
+
+#define LOCTEXT_NAMESPACE "NiagaraDataInterfaceSpriteRendererInfo"
+
+namespace NDISpriteRendererInfoLocal
+{
+	static const FName IsValidName("IsValid");
+	static const FName GetSourceModeName("GetSourceMode");
+	static const FName GetAlignmentName("GetAlignment");
+	static const FName GetFacingModeName("GetFacingMode");
+	static const FName GetSubUVDetailsName("GetSubUVDetails");
+
+	BEGIN_SHADER_PARAMETER_STRUCT(FShaderParameters, )
+		SHADER_PARAMETER(uint32,			bIsValid)
+		SHADER_PARAMETER(int32,				SourceMode)
+		SHADER_PARAMETER(int32,				Alignment)
+		SHADER_PARAMETER(int32,				FacingMode)
+		SHADER_PARAMETER(uint32,			bSubImageBlend)
+		SHADER_PARAMETER(FVector2f,			SubImageSize)
+	END_SHADER_PARAMETER_STRUCT()
+
+	static const TCHAR* TemplateShaderFile = TEXT("/Plugin/FX/Niagara/Private/NiagaraDataInterfaceSpriteRendererInfoTemplate.ush");
+
+	struct FNDIProxy : public FNiagaraDataInterfaceProxy
+	{
+		FShaderParameters	ShaderParameters_RT;
+
+		virtual int32 PerInstanceDataPassedToRenderThreadSize() const override { return 0; }
+
+		void PushToRenderThread(const UNiagaraSpriteRendererProperties* InProperties)
+		{
+			const UNiagaraSpriteRendererProperties* Properties = InProperties ? InProperties : GetDefault<UNiagaraSpriteRendererProperties>();
+
+			FShaderParameters Parameters;
+			Parameters.bIsValid			= InProperties != nullptr;
+			Parameters.SourceMode		= uint32(Properties->SourceMode);
+			Parameters.Alignment		= uint32(Properties->Alignment);
+			Parameters.FacingMode		= uint32(Properties->FacingMode);
+			Parameters.bSubImageBlend	= Properties->bSubImageBlend ? 1 : 0;
+			Parameters.SubImageSize		= FVector2f(Properties->SubImageSize);
+
+			ENQUEUE_RENDER_COMMAND(NDISpriteRendererInfo_Update)
+			(
+				[this, Parameters](FRHICommandList& RHICmdList)
+				{
+					ShaderParameters_RT = Parameters;
+				}
+			);
+		}
+	};
+}
+
+////////////////////////////////////////////////////////////////////////////////////////////////
+
+UNiagaraDataInterfaceSpriteRendererInfo::UNiagaraDataInterfaceSpriteRendererInfo(FObjectInitializer const& ObjectInitializer)
+	: Super(ObjectInitializer)
+{
+	Proxy.Reset(new NDISpriteRendererInfoLocal::FNDIProxy());
+	MarkRenderDataDirty();
+}
+
+void UNiagaraDataInterfaceSpriteRendererInfo::PostInitProperties()
+{
+	Super::PostInitProperties();
+
+	if (HasAnyFlags(RF_ClassDefaultObject))
+	{
+		ENiagaraTypeRegistryFlags Flags = ENiagaraTypeRegistryFlags::AllowAnyVariable | ENiagaraTypeRegistryFlags::AllowParameter;
+		
+		// We can't allow user variables of this type because it will cause components to have external reference (the renderer)
+		Flags &= ~ENiagaraTypeRegistryFlags::AllowUserVariable;
+		Flags &= ~ENiagaraTypeRegistryFlags::AllowSystemVariable;
+
+		FNiagaraTypeRegistry::Register(FNiagaraTypeDefinition(GetClass()), Flags);
+	}	
+}
+
+void UNiagaraDataInterfaceSpriteRendererInfo::PostLoad()
+{
+	Super::PostLoad();
+
+	if (SpriteRenderer)
+	{
+		SpriteRenderer->ConditionalPostLoad();
+		OnSpriteRendererChanged(SpriteRenderer);
+	}
+}
+
+#if WITH_EDITOR
+void UNiagaraDataInterfaceSpriteRendererInfo::PreEditChange(FProperty* PropertyAboutToChange)
+{
+	if (SpriteRenderer && PropertyAboutToChange && PropertyAboutToChange->GetFName() == GET_MEMBER_NAME_CHECKED(UNiagaraDataInterfaceSpriteRendererInfo, SpriteRenderer))
+	{
+		OnSpriteRendererChanged(nullptr);
+	}
+}
+
+void UNiagaraDataInterfaceSpriteRendererInfo::PostEditChangeProperty(struct FPropertyChangedEvent& PropertyChangedEvent)
+{
+	Super::PostEditChangeProperty(PropertyChangedEvent);
+
+	// If coming from undo, property will be nullptr and since we copy the info, we need to reacquire if new.
+	if (PropertyChangedEvent.Property == nullptr || (PropertyChangedEvent.Property &&
+		PropertyChangedEvent.Property->GetFName() == GET_MEMBER_NAME_CHECKED(UNiagaraDataInterfaceSpriteRendererInfo, SpriteRenderer)))
+	{
+		OnSpriteRendererChanged(SpriteRenderer);
+	}
+}
+#endif // WITH_EDITOR
+
+void UNiagaraDataInterfaceSpriteRendererInfo::OnSpriteRendererChanged(UNiagaraSpriteRendererProperties* NewSpriteRenderer)
+{
+#if WITH_EDITOR
+	if ( SpriteRenderer )
+	{
+		SpriteRenderer->OnChanged().Remove(OnSpriteRendererChangedHandle);
+		OnSpriteRendererChangedHandle.Reset();
+	}
+
+	if ( NewSpriteRenderer )
+	{
+		OnSpriteRendererChangedHandle = NewSpriteRenderer->OnChanged().AddLambda([this]() { MarkRenderDataDirty(); } );
+	}
+#endif
+	MarkRenderDataDirty();
+}
+
+void UNiagaraDataInterfaceSpriteRendererInfo::GetFunctions(TArray<FNiagaraFunctionSignature>& OutFunctions)
+{
+	FNiagaraFunctionSignature DefaultImmutableSignature;
+	DefaultImmutableSignature.bMemberFunction = true;
+	DefaultImmutableSignature.Inputs.Emplace(GetClass(), TEXT("SpriteRendererInfo"));
+
+	{
+		FNiagaraFunctionSignature& Sig = OutFunctions.Add_GetRef(DefaultImmutableSignature);
+		Sig.Name = NDISpriteRendererInfoLocal::IsValidName;
+		Sig.Outputs.Emplace(FNiagaraTypeDefinition::GetBoolDef(), TEXT("IsValid"));
+	}
+	{
+		FNiagaraFunctionSignature& Sig = OutFunctions.Add_GetRef(DefaultImmutableSignature);
+		Sig.Name = NDISpriteRendererInfoLocal::GetSourceModeName;
+		Sig.Outputs.Emplace(FNiagaraTypeDefinition(StaticEnum<ENiagaraRendererSourceDataMode>()), TEXT("SourceMode"));
+	}
+	{
+		FNiagaraFunctionSignature& Sig = OutFunctions.Add_GetRef(DefaultImmutableSignature);
+		Sig.Name = NDISpriteRendererInfoLocal::GetAlignmentName;
+		Sig.Outputs.Emplace(FNiagaraTypeDefinition(StaticEnum<ENiagaraSpriteAlignment>()), TEXT("Alignment"));
+	}
+	{
+		FNiagaraFunctionSignature& Sig = OutFunctions.Add_GetRef(DefaultImmutableSignature);
+		Sig.Name = NDISpriteRendererInfoLocal::GetFacingModeName;
+		Sig.Outputs.Emplace(FNiagaraTypeDefinition(StaticEnum<ENiagaraSpriteFacingMode>()), TEXT("FacingMode"));
+	}
+	{
+		FNiagaraFunctionSignature& Sig = OutFunctions.Add_GetRef(DefaultImmutableSignature);
+		Sig.Name = NDISpriteRendererInfoLocal::GetSubUVDetailsName;
+		Sig.Outputs.Emplace(FNiagaraTypeDefinition::GetBoolDef(), TEXT("BlendEnabled"));
+		Sig.Outputs.Emplace(FNiagaraTypeDefinition::GetVec2Def(), TEXT("SubImageSize"));
+	}
+}
+
+void UNiagaraDataInterfaceSpriteRendererInfo::GetVMExternalFunction(const FVMExternalFunctionBindingInfo& BindingInfo, void* InstanceData, FVMExternalFunction& OutFunc)
+{
+	if (BindingInfo.Name == NDISpriteRendererInfoLocal::IsValidName)
+	{
+		OutFunc = FVMExternalFunction::CreateUObject(this, &UNiagaraDataInterfaceSpriteRendererInfo::VMIsValid);
+	}
+	else if (BindingInfo.Name == NDISpriteRendererInfoLocal::GetSourceModeName)
+	{
+		OutFunc = FVMExternalFunction::CreateUObject(this, &UNiagaraDataInterfaceSpriteRendererInfo::VMGetSourceMode);
+	}
+	else if (BindingInfo.Name == NDISpriteRendererInfoLocal::GetAlignmentName)
+	{
+		OutFunc = FVMExternalFunction::CreateUObject(this, &UNiagaraDataInterfaceSpriteRendererInfo::VMGetAlignment);
+	}
+	else if (BindingInfo.Name == NDISpriteRendererInfoLocal::GetFacingModeName)
+	{
+		OutFunc = FVMExternalFunction::CreateUObject(this, &UNiagaraDataInterfaceSpriteRendererInfo::VMGetFacingMode);
+	}
+	else if (BindingInfo.Name == NDISpriteRendererInfoLocal::GetSubUVDetailsName)
+	{
+		OutFunc = FVMExternalFunction::CreateUObject(this, &UNiagaraDataInterfaceSpriteRendererInfo::VMGetSubUVDetails);
+	}
+}
+
+bool UNiagaraDataInterfaceSpriteRendererInfo::Equals(const UNiagaraDataInterface* Other) const
+{
+	if (!Super::Equals(Other))
+	{
+		return false;
+	}
+
+	auto OtherTyped = CastChecked<const UNiagaraDataInterfaceSpriteRendererInfo>(Other);
+	return SpriteRenderer == OtherTyped->SpriteRenderer;
+}
+
+#if WITH_EDITORONLY_DATA
+bool UNiagaraDataInterfaceSpriteRendererInfo::AppendCompileHash(FNiagaraCompileHashVisitor* InVisitor) const
+{
+	bool bSuccess = Super::AppendCompileHash(InVisitor);
+	bSuccess &= InVisitor->UpdateString(TEXT("NiagaraDataInterfaceSpriteRendererInfoHLSLSource"), GetShaderFileHash(NDISpriteRendererInfoLocal::TemplateShaderFile, EShaderPlatform::SP_PCD3D_SM5).ToString());
+	bSuccess &= InVisitor->UpdateShaderParameters<NDISpriteRendererInfoLocal::FShaderParameters>();
+	return bSuccess;
+}
+
+void UNiagaraDataInterfaceSpriteRendererInfo::GetParameterDefinitionHLSL(const FNiagaraDataInterfaceGPUParamInfo& ParamInfo, FString& OutHLSL)
+{
+	TMap<FString, FStringFormatArg> TemplateArgs =
+	{
+		{TEXT("ParameterName"),	ParamInfo.DataInterfaceHLSLSymbol},
+	};
+
+	FString TemplateFile;
+	LoadShaderSourceFile(NDISpriteRendererInfoLocal::TemplateShaderFile, EShaderPlatform::SP_PCD3D_SM5, &TemplateFile, nullptr);
+	OutHLSL += FString::Format(*TemplateFile, TemplateArgs);
+}
+
+bool UNiagaraDataInterfaceSpriteRendererInfo::GetFunctionHLSL(const FNiagaraDataInterfaceGPUParamInfo& ParamInfo, const FNiagaraDataInterfaceGeneratedFunction& FunctionInfo, int FunctionInstanceIndex, FString& OutHLSL)
+{
+	if ((FunctionInfo.DefinitionName == NDISpriteRendererInfoLocal::IsValidName) ||
+		(FunctionInfo.DefinitionName == NDISpriteRendererInfoLocal::GetSourceModeName) ||
+		(FunctionInfo.DefinitionName == NDISpriteRendererInfoLocal::GetAlignmentName) ||
+		(FunctionInfo.DefinitionName == NDISpriteRendererInfoLocal::GetFacingModeName) ||
+		(FunctionInfo.DefinitionName == NDISpriteRendererInfoLocal::GetSubUVDetailsName) )
+	{
+		return true;
+	}
+
+	return false;
+}
+#endif
+
+void UNiagaraDataInterfaceSpriteRendererInfo::BuildShaderParameters(FNiagaraShaderParametersBuilder& ShaderParametersBuilder) const
+{
+	ShaderParametersBuilder.AddNestedStruct<NDISpriteRendererInfoLocal::FShaderParameters>();
+}
+
+void UNiagaraDataInterfaceSpriteRendererInfo::SetShaderParameters(const FNiagaraDataInterfaceSetShaderParametersContext& Context) const
+{
+	NDISpriteRendererInfoLocal::FNDIProxy& DIProxy = Context.GetProxy<NDISpriteRendererInfoLocal::FNDIProxy>();
+	NDISpriteRendererInfoLocal::FShaderParameters* ShaderParameters = Context.GetParameterNestedStruct<NDISpriteRendererInfoLocal::FShaderParameters>();
+	*ShaderParameters = DIProxy.ShaderParameters_RT;
+}
+
+#if WITH_EDITOR
+
+void UNiagaraDataInterfaceSpriteRendererInfo::GetFeedback(UNiagaraSystem* InAsset, UNiagaraComponent* InComponent, TArray<FNiagaraDataInterfaceError>& OutErrors, TArray<FNiagaraDataInterfaceFeedback>& OutWarnings, TArray<FNiagaraDataInterfaceFeedback>& OutInfo)
+{
+	if (SpriteRenderer == nullptr)
+	{
+		OutWarnings.Emplace(
+			LOCTEXT("NoRendererSelectedWarning", "A Sprite Renderer applied to an emitter in this system is expected to be selected here"),
+			LOCTEXT("NoRendererSelectedWarningSummary", "No Sprite Renderer selected"),
+			FNiagaraDataInterfaceFix()
+		);
+	}
+	else if ( !SpriteRenderer->GetIsEnabled() )
+	{
+		OutWarnings.Emplace(
+			LOCTEXT("RendererDisabledWarning", "The selected Sprite Renderer is disabled"),
+			LOCTEXT("RendererDisabledWarningSummary", "Sprite Renderer is disabled"),
+			FNiagaraDataInterfaceFix::CreateLambda(
+				[this]()
+				{
+					SpriteRenderer->SetIsEnabled(true);
+					return true;
+				}
+			)
+		);
+	}
+}
+
+#endif // WITH_EDITOR
+
+bool UNiagaraDataInterfaceSpriteRendererInfo::CopyToInternal(UNiagaraDataInterface* Destination) const
+{
+	if (!Super::CopyToInternal(Destination))
+	{
+		return false;
+	}
+
+	UNiagaraDataInterfaceSpriteRendererInfo* DestinationTyped = CastChecked<UNiagaraDataInterfaceSpriteRendererInfo>(Destination);
+	DestinationTyped->SpriteRenderer = SpriteRenderer;
+	DestinationTyped->OnSpriteRendererChanged(SpriteRenderer);
+
+	return true;
+}
+
+void UNiagaraDataInterfaceSpriteRendererInfo::PushToRenderThreadImpl()
+{
+	GetProxyAs<NDISpriteRendererInfoLocal::FNDIProxy>()->PushToRenderThread(SpriteRenderer);
+}
+
+void UNiagaraDataInterfaceSpriteRendererInfo::VMIsValid(FVectorVMExternalFunctionContext& Context)
+{
+	FNDIOutputParam<bool> OutIsValid(Context);
+
+	const bool bIsValid = SpriteRenderer != nullptr;
+	for (int32 i=0; i < Context.GetNumInstances(); ++i)
+	{
+		OutIsValid.SetAndAdvance(bIsValid);
+	}
+}
+
+void UNiagaraDataInterfaceSpriteRendererInfo::VMGetSourceMode(FVectorVMExternalFunctionContext& Context)
+{
+	FNDIOutputParam<uint32> OutSourceMode(Context);
+
+	const UNiagaraSpriteRendererProperties* Properties = SpriteRenderer ? SpriteRenderer : GetDefault<UNiagaraSpriteRendererProperties>();
+	for (int32 i = 0; i < Context.GetNumInstances(); ++i)
+	{
+		OutSourceMode.SetAndAdvance(uint32(Properties->SourceMode));
+	}
+}
+
+void UNiagaraDataInterfaceSpriteRendererInfo::VMGetAlignment(FVectorVMExternalFunctionContext& Context)
+{
+	FNDIOutputParam<uint32> OutAlignment(Context);
+
+	const UNiagaraSpriteRendererProperties* Properties = SpriteRenderer ? SpriteRenderer : GetDefault<UNiagaraSpriteRendererProperties>();
+	for (int32 i = 0; i < Context.GetNumInstances(); ++i)
+	{
+		OutAlignment.SetAndAdvance(uint32(Properties->Alignment));
+	}
+}
+
+void UNiagaraDataInterfaceSpriteRendererInfo::VMGetFacingMode(FVectorVMExternalFunctionContext& Context)
+{
+	FNDIOutputParam<uint32> OutFacingMode(Context);
+
+	const UNiagaraSpriteRendererProperties* Properties = SpriteRenderer ? SpriteRenderer : GetDefault<UNiagaraSpriteRendererProperties>();
+	for (int32 i = 0; i < Context.GetNumInstances(); ++i)
+	{
+		OutFacingMode.SetAndAdvance(uint32(Properties->FacingMode));
+	}
+}
+
+void UNiagaraDataInterfaceSpriteRendererInfo::VMGetSubUVDetails(FVectorVMExternalFunctionContext& Context)
+{
+	FNDIOutputParam<bool> OutBlendEnabled(Context);
+	FNDIOutputParam<FVector2f> OutSubImageSize(Context);
+
+	const UNiagaraSpriteRendererProperties* Properties = SpriteRenderer ? SpriteRenderer : GetDefault<UNiagaraSpriteRendererProperties>();
+	const FVector2f SubImageSize(Properties->SubImageSize);
+	for (int32 i = 0; i < Context.GetNumInstances(); ++i)
+	{
+		OutBlendEnabled.SetAndAdvance(Properties->bSubImageBlend != 0);
+		OutSubImageSize.SetAndAdvance(SubImageSize);
+	}
+}
+
+#undef LOCTEXT_NAMESPACE
diff --git a/Engine/Plugins/FX/Niagara/Source/NiagaraEditorWidgets/Private/DetailCustomizations/NiagaraDataInterfaceSpriteRendererInfoDetails.cpp b/Engine/Plugins/FX/Niagara/Source/NiagaraEditorWidgets/Private/DetailCustomizations/NiagaraDataInterfaceSpriteRendererInfoDetails.cpp
new file mode 100644
--- /dev/null
+++ b/Engine/Plugins/FX/Niagara/Source/NiagaraEditorWidgets/Private/DetailCustomizations/NiagaraDataInterfaceSpriteRendererInfoDetails.cpp
@@ -0,0 +1,201 @@
+// Copyright Epic Games, Inc. All Rights Reserved.
+
+#include "NiagaraDataInterfaceSpriteRendererInfoDetails.h"
+#include "NiagaraDetailSourcedArrayBuilder.h"
+#include "NiagaraDataInterfaceDetails.h"
+#include "NiagaraDataInterfaceSpriteRendererInfo.h"
+#include "NiagaraSpriteRendererProperties.h"
+#include "DetailLayoutBuilder.h"
+#include "DetailCategoryBuilder.h"
+#include "NiagaraComponent.h"
+#include "NiagaraEmitter.h"
+#include "NiagaraStackEditorData.h"
+#include "NiagaraEmitterEditorData.h"
+#include "NiagaraEditorModule.h"
+#include "ScopedTransaction.h"
+
+#define LOCTEXT_NAMESPACE "FNiagaraDataInterfaceSpriteRendererInfoDetails"
+
+FNiagaraDataInterfaceSpriteRendererInfoDetails::~FNiagaraDataInterfaceSpriteRendererInfoDetails()
+{
+	if (DataInterface.IsValid())
+	{
+		DataInterface->OnChanged().RemoveAll(this);
+	}
+}
+
+void FNiagaraDataInterfaceSpriteRendererInfoDetails::CustomizeDetails(IDetailLayoutBuilder& DetailBuilder)
+{
+	LayoutBuilder = &DetailBuilder;
+	FNiagaraDataInterfaceDetailsBase::CustomizeDetails(DetailBuilder);
+
+	TArray<TWeakObjectPtr<UObject>> SelectedObjects;
+	DetailBuilder.GetObjectsBeingCustomized(SelectedObjects);
+	if (SelectedObjects.Num() != 1 || SelectedObjects[0]->IsA<UNiagaraDataInterfaceSpriteRendererInfo>() == false)
+	{
+		return;
+	}
+
+	UNiagaraDataInterfaceSpriteRendererInfo* Interface = CastChecked<UNiagaraDataInterfaceSpriteRendererInfo>(SelectedObjects[0].Get());
+	DataInterface = Interface;
+
+	Interface->OnChanged().AddSP(this, &FNiagaraDataInterfaceSpriteRendererInfoDetails::OnInterfaceChanged);
+
+	GenerateRendererList();	
+
+	static const FName SourceCategoryName = TEXT("Source");
+	SourceCategory = &DetailBuilder.EditCategory(SourceCategoryName, LOCTEXT("SourceCategory", "Source"));
+	{
+		TArray<TSharedRef<IPropertyHandle>> Properties;
+		SourceCategory->GetDefaultProperties(Properties, true, true);
+
+		static const FName SpriteRendererPropertyName = TEXT("SpriteRenderer");
+		SpriteRendererProperty = DetailBuilder.GetProperty(SpriteRendererPropertyName);
+		for (TSharedPtr<IPropertyHandle> Property : Properties)
+		{
+			FProperty* PropertyPtr = Property->GetProperty();
+			if (PropertyPtr == SpriteRendererProperty->GetProperty())
+			{				
+				SpriteRendererWidget = SNew(SSpriteRendererComboBox)
+					.OptionsSource(&RendererList)
+					.InitiallySelectedItem(Interface->GetSpriteRenderer())
+					.OnComboBoxOpening(this, &FNiagaraDataInterfaceSpriteRendererInfoDetails::GenerateRendererList)
+					.OnSelectionChanged(this, &FNiagaraDataInterfaceSpriteRendererInfoDetails::SetSelectedRenderer)
+					.OnGenerateWidget(this, &FNiagaraDataInterfaceSpriteRendererInfoDetails::CreateRendererItemWidget)
+					.AddMetaData<FTagMetaData>(TEXT("SelectSpriteRendererCobmo"))
+					[
+						SNew(STextBlock)
+							.Text(this, &FNiagaraDataInterfaceSpriteRendererInfoDetails::GetSelectedRendererTextLabel)
+					];
+
+				IDetailPropertyRow& RendererRow = SourceCategory->AddProperty(Property);
+				RendererRow.CustomWidget(false)
+					.NameContent()
+					[
+						Property->CreatePropertyNameWidget()
+					]
+					.ValueContent()
+					.MaxDesiredWidth(TOptional<float>())
+					[
+						SpriteRendererWidget.ToSharedRef()
+					];
+			}
+			else
+			{
+				SourceCategory->AddProperty(Property);
+			}
+		}
+	}
+}
+
+TSharedRef<IDetailCustomization> FNiagaraDataInterfaceSpriteRendererInfoDetails::MakeInstance()
+{
+	return MakeShared<FNiagaraDataInterfaceSpriteRendererInfoDetails>();
+}
+
+void FNiagaraDataInterfaceSpriteRendererInfoDetails::OnInterfaceChanged()
+{
+	if (!bSettingSelection && SpriteRendererWidget.IsValid())
+	{
+		SpriteRendererWidget->SetSelectedItem(GetSelectedRenderer());
+	}
+}
+
+void FNiagaraDataInterfaceSpriteRendererInfoDetails::GenerateRendererList()
+{
+	RendererList.SetNum(0, false);
+	RendererLabels.SetNum(0, false);
+
+	// Get a list of all Sprite renderers in the emitter
+	if (auto Interface = DataInterface.Get())
+	{
+		UNiagaraSystem* System = nullptr;
+		FVersionedNiagaraEmitter EmitterHandle;
+		FNiagaraEditorModule::Get().GetTargetSystemAndEmitterForDataInterface(Interface, System, EmitterHandle);
+		if (System && EmitterHandle.GetEmitterData())
+		{
+			for (auto RendererProps : EmitterHandle.GetEmitterData()->GetRenderers())
+			{
+				if (auto SpriteProps = Cast<UNiagaraSpriteRendererProperties>(RendererProps))
+				{
+					RendererList.Add(SpriteProps);
+					RendererLabels.Add(CreateRendererTextLabel(SpriteProps));
+				}
+			}
+
+			if(SpriteRendererWidget.IsValid())
+			{
+				SpriteRendererWidget->RefreshOptions();
+			}
+		}
+	}
+}
+
+UNiagaraSpriteRendererProperties* FNiagaraDataInterfaceSpriteRendererInfoDetails::GetSelectedRenderer() const
+{
+	UObject* Value = nullptr;
+	if (SpriteRendererProperty.IsValid())
+	{
+		SpriteRendererProperty->GetValue(Value);
+	}
+	return Cast<UNiagaraSpriteRendererProperties>(Value);
+}
+
+void FNiagaraDataInterfaceSpriteRendererInfoDetails::SetSelectedRenderer(TRendererPtr Selection, ESelectInfo::Type)
+{
+	if (SpriteRendererProperty.IsValid() && DataInterface.IsValid())
+	{
+		bSettingSelection = true;
+		UObject* Obj = nullptr;
+		SpriteRendererProperty->GetValue(Obj);
+		if (Obj != Selection.Get())
+		{
+			const FScopedTransaction Transaction(NSLOCTEXT("FNiagaraDataInterfaceSpriteRendererInfoDetails", "Change Sprite Renderer", "Change Sprite Renderer"));
+			DataInterface->Modify();
+			SpriteRendererProperty->NotifyPreChange();
+			SpriteRendererProperty->SetValue(Selection.Get());
+			SpriteRendererProperty->NotifyPostChange(EPropertyChangeType::ValueSet);
+		}
+		bSettingSelection = false;
+	}	
+}
+
+FText FNiagaraDataInterfaceSpriteRendererInfoDetails::GetRendererTextLabel(TRendererPtr Renderer) const
+{
+	int32 RendererIndex = Renderer.IsValid() ? RendererList.Find(Renderer) : INDEX_NONE;
+	if (!RendererLabels.IsValidIndex(RendererIndex))
+	{
+		return LOCTEXT("NoneOption", "None");
+	}
+	return RendererLabels[RendererIndex];
+}
+
+TSharedRef<SWidget> FNiagaraDataInterfaceSpriteRendererInfoDetails::CreateRendererItemWidget(TRendererPtr Item)
+{
+	return SNew(STextBlock)
+		.Text(this, &FNiagaraDataInterfaceSpriteRendererInfoDetails::GetRendererTextLabel, Item);
+}
+
+FText FNiagaraDataInterfaceSpriteRendererInfoDetails::CreateRendererTextLabel(const UNiagaraSpriteRendererProperties* Properties)
+{
+	if (Properties == nullptr)
+	{
+		return LOCTEXT("NoneOption", "None");
+	}
+
+	UNiagaraEmitter* Emitter = Properties->GetTypedOuter<UNiagaraEmitter>();
+	check(Emitter);
+
+	UNiagaraEmitterEditorData* EmitterEditorData = static_cast<UNiagaraEmitterEditorData*>(Properties->GetEmitterData()->GetEditorData());
+	if (EmitterEditorData == nullptr)
+	{
+		return LOCTEXT("NoneOption", "None");
+	}
+	UNiagaraStackEditorData& EmitterStackEditorData = EmitterEditorData->GetStackEditorData();
+	FString RendererStackEditorDataKey = FString::Printf(TEXT("Renderer-%s"), *Properties->GetName());
+	const FText* RendererDisplayName = EmitterStackEditorData.GetStackEntryDisplayName(RendererStackEditorDataKey);
+
+	return RendererDisplayName ? *RendererDisplayName : Properties->GetWidgetDisplayName();
+}
+
+#undef LOCTEXT_NAMESPACE
\ No newline at end of file
diff --git a/Engine/Plugins/FX/Niagara/Source/NiagaraEditorWidgets/Private/DetailCustomizations/NiagaraDataInterfaceSpriteRendererInfoDetails.h b/Engine/Plugins/FX/Niagara/Source/NiagaraEditorWidgets/Private/DetailCustomizations/NiagaraDataInterfaceSpriteRendererInfoDetails.h
new file mode 100644
--- /dev/null
+++ b/Engine/Plugins/FX/Niagara/Source/NiagaraEditorWidgets/Private/DetailCustomizations/NiagaraDataInterfaceSpriteRendererInfoDetails.h
@@ -0,0 +1,47 @@
+// Copyright Epic Games, Inc. All Rights Reserved.
+
+#pragma once
+
+#include "NiagaraDataInterfaceDetails.h"
+#include "Types/SlateEnums.h"
+#include "Widgets/Input/SComboBox.h"
+
+class IDetailLayoutBuilder;
+class IPropertyHandle;
+class SWidget;
+class UNiagaraDataInterfaceSpriteRendererInfo;
+class UNiagaraSpriteRendererProperties;
+
+/** Details customization for Niagara Sprite renderer info data interface. */
+class FNiagaraDataInterfaceSpriteRendererInfoDetails : public FNiagaraDataInterfaceDetailsBase
+{
+public:
+	virtual ~FNiagaraDataInterfaceSpriteRendererInfoDetails();
+	virtual void CustomizeDetails(IDetailLayoutBuilder& DetailBuilder) override;
+	static TSharedRef<IDetailCustomization> MakeInstance();
+
+private:
+	using TRendererPtr = TWeakObjectPtr<UNiagaraSpriteRendererProperties>;
+	using SSpriteRendererComboBox = SComboBox<TRendererPtr>;
+
+	void OnInterfaceChanged();
+	void GenerateRendererList();
+	UNiagaraSpriteRendererProperties* GetSelectedRenderer() const;
+	void SetSelectedRenderer(TRendererPtr Selection, ESelectInfo::Type SelectInfo);
+	FText GetRendererTextLabel(TRendererPtr Renderer) const;
+	FText GetSelectedRendererTextLabel() const { return CreateRendererTextLabel(GetSelectedRenderer()); }
+	TSharedRef<SWidget> CreateRendererItemWidget(TRendererPtr Item);
+	
+	static FText CreateRendererTextLabel(const UNiagaraSpriteRendererProperties* Properties);
+
+private:
+
+	TSharedPtr<SSpriteRendererComboBox> SpriteRendererWidget;
+	TArray<TRendererPtr> RendererList;
+	TArray<FText> RendererLabels;
+	IDetailLayoutBuilder* LayoutBuilder = nullptr;
+	TWeakObjectPtr<UNiagaraDataInterfaceSpriteRendererInfo>  DataInterface;
+	IDetailCategoryBuilder* SourceCategory = nullptr;
+	TSharedPtr<IPropertyHandle> SpriteRendererProperty;
+	bool bSettingSelection = false;
+};
diff --git a/Engine/Plugins/FX/Niagara/Source/NiagaraEditorWidgets/Private/NiagaraEditorWidgetsModule.cpp b/Engine/Plugins/FX/Niagara/Source/NiagaraEditorWidgets/Private/NiagaraEditorWidgetsModule.cpp
--- a/Engine/Plugins/FX/Niagara/Source/NiagaraEditorWidgets/Private/NiagaraEditorWidgetsModule.cpp
+++ b/Engine/Plugins/FX/Niagara/Source/NiagaraEditorWidgets/Private/NiagaraEditorWidgetsModule.cpp
@@ -11,6 +11,7 @@
 #include "DetailCustomizations/NiagaraDataInterfaceSkeletalMeshDetails.h"
 #include "DetailCustomizations/NiagaraDataInterfaceStaticMeshDetails.h"
 #include "DetailCustomizations/NiagaraDataInterfaceMeshRendererInfoDetails.h"
+#include "DetailCustomizations/NiagaraDataInterfaceSpriteRendererInfoDetails.h"
 #include "DetailCustomizations/NiagaraMeshRendererDetails.h"
 #include "DetailCustomizations/NiagaraSimCacheDetailsCustomization.h"
 #include "ViewModels/NiagaraSystemViewModel.h"
@@ -130,6 +131,7 @@ void FNiagaraEditorWidgetsModule::StartupModule()
 	PropertyModule.RegisterCustomClassLayout("NiagaraDataInterfaceGrid2DCollection", FOnGetDetailCustomizationInstance::CreateStatic(&FNiagaraDataInterfaceGrid2DCollectionDetails::MakeInstance));
 	PropertyModule.RegisterCustomClassLayout("NiagaraDataInterfaceGrid3DCollection", FOnGetDetailCustomizationInstance::CreateStatic(&FNiagaraDataInterfaceGrid3DCollectionDetails::MakeInstance));
 	PropertyModule.RegisterCustomClassLayout("NiagaraDataInterfaceMeshRendererInfo", FOnGetDetailCustomizationInstance::CreateStatic(&FNiagaraDataInterfaceMeshRendererInfoDetails::MakeInstance));
+	PropertyModule.RegisterCustomClassLayout("NiagaraDataInterfaceSpriteRendererInfo", FOnGetDetailCustomizationInstance::CreateStatic(&FNiagaraDataInterfaceSpriteRendererInfoDetails::MakeInstance));
 	PropertyModule.RegisterCustomClassLayout("NiagaraMeshRendererProperties", FOnGetDetailCustomizationInstance::CreateStatic(&FNiagaraMeshRendererDetails::MakeInstance));
 	PropertyModule.RegisterCustomClassLayout("NiagaraSimCache", FOnGetDetailCustomizationInstance::CreateStatic(&FNiagaraSimCacheDetailsCustomization::MakeInstance));
 
