- When fixed DT is enabled force solo mode
- There are a few issues where we assume the world manager will run after each tick in batched mode, this makes it safe

diff --git a/Engine/Plugins/FX/Niagara/Source/Niagara/Private/NiagaraComponent.cpp b/Engine/Plugins/FX/Niagara/Source/Niagara/Private/NiagaraComponent.cpp
--- a/Engine/Plugins/FX/Niagara/Source/Niagara/Private/NiagaraComponent.cpp
+++ b/Engine/Plugins/FX/Niagara/Source/Niagara/Private/NiagaraComponent.cpp
@@ -3637,7 +3637,7 @@ void UNiagaraComponent::SetForceSolo(bool bInForceSolo)
 
 bool UNiagaraComponent::RequiresSoloMode() const
 {
-	return bForceSolo || SimCache != nullptr || !FMath::IsNearlyEqual(CustomTimeDilation, 1.0f);
+	return bForceSolo || SimCache != nullptr || !FMath::IsNearlyEqual(CustomTimeDilation, 1.0f) || (Asset && Asset->HasFixedTickDelta());
 }
 
 void UNiagaraComponent::UpdateInstanceSoloMode()
