- Convert Collision DI to RDG
- Track global structures and only allow them to be added once to avoid binding errors

diff --git a/Engine/Plugins/FX/Niagara/Shaders/Private/NiagaraDataInterfaceCollisionQueryTemplate.ush b/Engine/Plugins/FX/Niagara/Shaders/Private/NiagaraDataInterfaceCollisionQueryTemplate.ush
--- a/Engine/Plugins/FX/Niagara/Shaders/Private/NiagaraDataInterfaceCollisionQueryTemplate.ush
+++ b/Engine/Plugins/FX/Niagara/Shaders/Private/NiagaraDataInterfaceCollisionQueryTemplate.ush
@@ -1,20 +1,20 @@
 // Copyright Epic Games, Inc. All Rights Reserved.
 
-float3										SystemLWCTile_{ParameterName};
+float3 {ParameterName}_SystemLWCTile;
 
 void QuerySceneDepthGPU_{ParameterName}(in float3 In_SamplePos, out float Out_SceneDepth, out float3 Out_CameraPosWorld, out bool Out_IsInsideView, out float3 Out_WorldPos, out float3 Out_WorldNormal)
 {
-	NDICollisionQuery_QuerySceneDepthGPU(In_SamplePos, SystemLWCTile_{ParameterName}, Out_SceneDepth, Out_CameraPosWorld, Out_IsInsideView, Out_WorldPos, Out_WorldNormal);
+	NDICollisionQuery_QuerySceneDepthGPU(In_SamplePos, {ParameterName}_SystemLWCTile, Out_SceneDepth, Out_CameraPosWorld, Out_IsInsideView, Out_WorldPos, Out_WorldNormal);
 }
 
 void QueryCustomDepthGPU_{ParameterName}(in float3 In_SamplePos, out float Out_SceneDepth, out float3 Out_CameraPosWorld, out bool Out_IsInsideView, out float3 Out_WorldPos, out float3 Out_WorldNormal)
 {
-	NDICollisionQuery_QueryCustomDepthGPU(In_SamplePos, SystemLWCTile_{ParameterName}, Out_SceneDepth, Out_CameraPosWorld, Out_IsInsideView, Out_WorldPos, Out_WorldNormal);
+	NDICollisionQuery_QueryCustomDepthGPU(In_SamplePos, {ParameterName}_SystemLWCTile, Out_SceneDepth, Out_CameraPosWorld, Out_IsInsideView, Out_WorldPos, Out_WorldNormal);
 }
 
 void QueryMeshDistanceFieldGPU_{ParameterName}(in float3 In_SamplePos, out float Out_DistanceToNearestSurface, out float3 Out_FieldGradient, out bool Out_IsDistanceFieldValid)
 {
-	NDICollisionQuery_QueryMeshDistanceFieldGPU(In_SamplePos, SystemLWCTile_{ParameterName}, Out_DistanceToNearestSurface, Out_FieldGradient, Out_IsDistanceFieldValid);
+	NDICollisionQuery_QueryMeshDistanceFieldGPU(In_SamplePos, {ParameterName}_SystemLWCTile, Out_DistanceToNearestSurface, Out_FieldGradient, Out_IsDistanceFieldValid);
 }
 
 // DEPRECATE_BEGIN
diff --git a/Engine/Plugins/FX/Niagara/Source/Niagara/Classes/NiagaraDataInterface.h b/Engine/Plugins/FX/Niagara/Source/Niagara/Classes/NiagaraDataInterface.h
--- a/Engine/Plugins/FX/Niagara/Source/Niagara/Classes/NiagaraDataInterface.h
+++ b/Engine/Plugins/FX/Niagara/Source/Niagara/Classes/NiagaraDataInterface.h
@@ -366,13 +366,14 @@ struct FNiagaraDataInterfaceProxy
 
 struct NIAGARA_API FNiagaraDataInterfaceSetShaderParametersContext
 {
-	FNiagaraDataInterfaceSetShaderParametersContext(FRDGBuilder& InGraphBuilder, const FNiagaraGpuComputeDispatchInterface& InComputeDispatchInterface, const FNiagaraGPUSystemTick& InSystemTick, const FNiagaraComputeInstanceData& InComputeInstanceData, const FNiagaraSimStageData& InSimStageData, const FNiagaraShaderRef& InShaderRef, uint8* InBaseParameters)
+	FNiagaraDataInterfaceSetShaderParametersContext(FRDGBuilder& InGraphBuilder, const FNiagaraGpuComputeDispatchInterface& InComputeDispatchInterface, const FNiagaraGPUSystemTick& InSystemTick, const FNiagaraComputeInstanceData& InComputeInstanceData, const FNiagaraSimStageData& InSimStageData, const FNiagaraShaderRef& InShaderRef, const FNiagaraShaderScriptParametersMetadata& InShaderParametersMetadata, uint8* InBaseParameters)
 		: GraphBuilder(InGraphBuilder)
 		, ComputeDispatchInterface(InComputeDispatchInterface)
 		, SystemTick(InSystemTick)
 		, ComputeInstanceData(InComputeInstanceData)
 		, SimStageData(InSimStageData)
 		, ShaderRef(InShaderRef)
+		, ShaderParametersMetadata(InShaderParametersMetadata)
 		, BaseParameters(InBaseParameters)
 	{
 	}
@@ -402,9 +403,7 @@ struct NIAGARA_API FNiagaraDataInterfaceSetShaderParametersContext
 	}
 	template<typename T> T* GetParameterIncludedStruct() const
 	{
-		const uint32 StructOffset = Align(ParametersOffset, TShaderParameterStructTypeInfo<T>::Alignment);
-		ParametersOffset = StructOffset + TShaderParameterStructTypeInfo<T>::GetStructMetadata()->GetSize();
-
+		const uint16 StructOffset = GetParameterIncludedStructInternal(TShaderParameterStructTypeInfo<T>::GetStructMetadata());
 		return reinterpret_cast<T*>(BaseParameters + StructOffset);
 	}
 	template<typename T> TArrayView<T> GetParameterLooseArray(int32 NumElements) const
@@ -429,6 +428,7 @@ struct NIAGARA_API FNiagaraDataInterfaceSetShaderParametersContext
 
 private:
 	bool IsStructBoundInternal(const void* StructAddress, uint32 StructSize) const;
+	uint16 GetParameterIncludedStructInternal(const FShaderParametersMetadata* StructMetadata) const;
 
 private:
 	FRDGBuilder& GraphBuilder;
@@ -437,6 +437,7 @@ private:
 	const FNiagaraComputeInstanceData& ComputeInstanceData;
 	const FNiagaraSimStageData& SimStageData;
 	const FNiagaraShaderRef& ShaderRef;
+	const FNiagaraShaderScriptParametersMetadata& ShaderParametersMetadata;
 	uint8* BaseParameters = nullptr;
 	const FNiagaraDataInterfaceParametersCS* ShaderStorage = nullptr;
 
diff --git a/Engine/Plugins/FX/Niagara/Source/Niagara/Classes/NiagaraDataInterfaceCollisionQuery.h b/Engine/Plugins/FX/Niagara/Source/Niagara/Classes/NiagaraDataInterfaceCollisionQuery.h
--- a/Engine/Plugins/FX/Niagara/Source/Niagara/Classes/NiagaraDataInterfaceCollisionQuery.h
+++ b/Engine/Plugins/FX/Niagara/Source/Niagara/Classes/NiagaraDataInterfaceCollisionQuery.h
@@ -5,6 +5,7 @@
 #include "NiagaraCollision.h"
 #include "NiagaraSystemInstance.h"
 #include "NiagaraShared.h"
+#include "ShaderParameterMacros.h"
 #include "VectorVM.h"
 #include "NiagaraDataInterface.h"
 #include "NiagaraDataInterfaceCollisionQuery.generated.h"
@@ -25,10 +26,12 @@ UCLASS(EditInlineNew, Category = "Collision", meta = (DisplayName = "Collision Q
 class NIAGARA_API UNiagaraDataInterfaceCollisionQuery : public UNiagaraDataInterface
 {
 	GENERATED_UCLASS_BODY()
-public:
 
-	DECLARE_NIAGARA_DI_PARAMETER();
+	BEGIN_SHADER_PARAMETER_STRUCT(FShaderParameters, )
+		SHADER_PARAMETER(FVector3f,	SystemLWCTile)
+	END_SHADER_PARAMETER_STRUCT()
 
+public:
 	//UObject Interface
 	virtual void PostInitProperties() override;
 	//UObject Interface End
@@ -62,6 +65,9 @@ public:
 	virtual bool GetFunctionHLSL(const FNiagaraDataInterfaceGPUParamInfo& ParamInfo, const FNiagaraDataInterfaceGeneratedFunction& FunctionInfo, int FunctionInstanceIndex, FString& OutHLSL) override;
 	virtual bool UpgradeFunctionCall(FNiagaraFunctionSignature& FunctionSignature) override;
 #endif
+	virtual bool UseLegacyShaderBindings() const override { return false; }
+	virtual void BuildShaderParameters(FNiagaraShaderParametersBuilder& ShaderParametersBuilder) const override;
+	virtual void SetShaderParameters(const FNiagaraDataInterfaceSetShaderParametersContext& Context) const override;
 
 #if WITH_EDITOR	
 	virtual void ValidateFunction(const FNiagaraFunctionSignature& Function, TArray<FText>& OutValidationErrors) override;
diff --git a/Engine/Plugins/FX/Niagara/Source/Niagara/Private/NiagaraDataInterface.cpp b/Engine/Plugins/FX/Niagara/Source/Niagara/Private/NiagaraDataInterface.cpp
--- a/Engine/Plugins/FX/Niagara/Source/Niagara/Private/NiagaraDataInterface.cpp
+++ b/Engine/Plugins/FX/Niagara/Source/Niagara/Private/NiagaraDataInterface.cpp
@@ -326,6 +326,20 @@ bool FNiagaraDataInterfaceSetShaderParametersContext::IsStructBoundInternal(cons
 	return false;
 }
 
+uint16 FNiagaraDataInterfaceSetShaderParametersContext::GetParameterIncludedStructInternal(const FShaderParametersMetadata* StructMetadata) const
+{
+	for (const FNiagaraDataInterfaceStructIncludeInfo& StructIncludeInfo : ShaderParametersMetadata.StructIncludeInfos)
+	{
+		if (StructIncludeInfo.StructMetadata == StructMetadata)
+		{
+			return StructIncludeInfo.ParamterOffset;
+		}
+	}
+
+	UE_LOG(LogNiagara, Fatal, TEXT("Failed to find Included Parameter Struct '%s' in parameters"), StructMetadata->GetStructTypeName());
+	return 0;
+}
+
 bool FNiagaraDataInterfaceSetShaderParametersContext::IsOutputStage() const
 {
 	return ComputeInstanceData.IsOutputStage(DataInterfaceProxy, SimStageData.StageIndex);
diff --git a/Engine/Plugins/FX/Niagara/Source/Niagara/Private/NiagaraDataInterfaceCollisionQuery.cpp b/Engine/Plugins/FX/Niagara/Source/Niagara/Private/NiagaraDataInterfaceCollisionQuery.cpp
--- a/Engine/Plugins/FX/Niagara/Source/Niagara/Private/NiagaraDataInterfaceCollisionQuery.cpp
+++ b/Engine/Plugins/FX/Niagara/Source/Niagara/Private/NiagaraDataInterfaceCollisionQuery.cpp
@@ -4,9 +4,11 @@
 
 #include "GlobalDistanceFieldParameters.h"
 #include "NiagaraAsyncGpuTraceHelper.h"
+#include "NiagaraDistanceFieldHelper.h"
 #include "NiagaraComponent.h"
 #include "NiagaraGpuComputeDispatchInterface.h"
 #include "NiagaraGpuComputeDispatch.h"
+#include "NiagaraShaderParametersBuilder.h"
 #include "NiagaraStats.h"
 #include "NiagaraTypes.h"
 #include "NiagaraWorldManager.h"
@@ -522,12 +524,33 @@ bool UNiagaraDataInterfaceCollisionQuery::AppendCompileHash(FNiagaraCompileHashV
 	InVisitor->UpdatePOD(TEXT("NiagaraCollisionDI_DistanceField"), IsDistanceFieldEnabled());
 	InVisitor->UpdateString(TEXT("NDICollisionQueryCommonHLSLSource"), GetShaderFileHash(NDICollisionQueryLocal::CommonShaderFile, EShaderPlatform::SP_PCD3D_SM5).ToString());
 	InVisitor->UpdateString(TEXT("NDICollisionQueryTemplateHLSLSource"), GetShaderFileHash(NDICollisionQueryLocal::TemplateShaderFile, EShaderPlatform::SP_PCD3D_SM5).ToString());
+	InVisitor->UpdateShaderParameters<FShaderParameters>();
+	InVisitor->UpdateShaderParameters<FGlobalDistanceFieldParameters2>();
 
 	return true;
 }
 
 #endif
 
+void UNiagaraDataInterfaceCollisionQuery::BuildShaderParameters(FNiagaraShaderParametersBuilder& ShaderParametersBuilder) const
+{
+	ShaderParametersBuilder.AddNestedStruct<FShaderParameters>();
+	ShaderParametersBuilder.AddIncludedStruct<FGlobalDistanceFieldParameters2>();
+}
+
+void UNiagaraDataInterfaceCollisionQuery::SetShaderParameters(const FNiagaraDataInterfaceSetShaderParametersContext& Context) const
+{
+	FShaderParameters* ShaderParameters = Context.GetParameterNestedStruct<FShaderParameters>();
+	ShaderParameters->SystemLWCTile = Context.GetSystemTick().SystemGpuComputeProxy->GetSystemLWCTile();	//-OPT: Seems like a lot of DIs might want tile
+
+	FGlobalDistanceFieldParameters2* ShaderGDFParameters = Context.GetParameterIncludedStruct<FGlobalDistanceFieldParameters2>();
+	if (Context.IsStructBound(ShaderGDFParameters))
+	{
+		const FGlobalDistanceFieldParameterData* GDFParameterData = static_cast<const FNiagaraGpuComputeDispatch&>(Context.GetComputeDispatchInterface()).GetGlobalDistanceFieldParameters();//-BATCHERTODO:
+		FNiagaraDistanceFieldHelper::SetGlobalDistanceFieldParameters(GDFParameterData, *ShaderGDFParameters);
+	}
+}
+
 void UNiagaraDataInterfaceCollisionQuery::PerformQuerySyncCPU(FVectorVMExternalFunctionContext & Context)
 {
 	VectorVM::FUserPtrHandler<CQDIPerInstanceData> InstanceData(Context);
@@ -657,49 +680,4 @@ bool UNiagaraDataInterfaceCollisionQuery::PerInstanceTickPostSimulate(void* PerI
 	return false;
 }
 
-//////////////////////////////////////////////////////////////////////////
-
-struct FNiagaraDataInterfaceParametersCS_CollisionQuery : public FNiagaraDataInterfaceParametersCS
-{
-	DECLARE_TYPE_LAYOUT(FNiagaraDataInterfaceParametersCS_CollisionQuery, NonVirtual);
-public:
-	void Bind(const FNiagaraDataInterfaceGPUParamInfo& ParameterInfo, const class FShaderParameterMap& ParameterMap)
-	{
-		GlobalDistanceFieldParameters.Bind(ParameterMap);
-		SystemLWCTileParam.Bind(ParameterMap, *(NDICollisionQueryLocal::SystemLWCTileName + ParameterInfo.DataInterfaceHLSLSymbol));
-	}
-
-	void Set(FRHICommandList& RHICmdList, const FNiagaraDataInterfaceSetArgs& Context) const
-	{
-		check(IsInRenderingThread());
-
-		FRHIComputeShader* ComputeShaderRHI = Context.Shader.GetComputeShader();
-		SetShaderValue(RHICmdList, ComputeShaderRHI, SystemLWCTileParam, Context.SystemLWCTile);
-		
-		// Bind distance field parameters
-		if (GlobalDistanceFieldParameters.IsBound())
-		{
-			check(Context.ComputeDispatchInterface);
-			const FGlobalDistanceFieldParameterData* GlobalDistanceFieldParameterData = static_cast<const FNiagaraGpuComputeDispatch*>(Context.ComputeDispatchInterface)->GetGlobalDistanceFieldParameters();//-BATCHERTODO:
-
-			if (GlobalDistanceFieldParameterData)
-			{
-				GlobalDistanceFieldParameters.Set(RHICmdList, ComputeShaderRHI, *GlobalDistanceFieldParameterData);
-			}
-			else
-			{
-				GlobalDistanceFieldParameters.Set(RHICmdList, ComputeShaderRHI, FGlobalDistanceFieldParameterData());
-			}
-		}
-	}
-
-private:
-	LAYOUT_FIELD(FGlobalDistanceFieldParameters, GlobalDistanceFieldParameters);
-	LAYOUT_FIELD(FShaderParameter, SystemLWCTileParam);
-};
-
-IMPLEMENT_TYPE_LAYOUT(FNiagaraDataInterfaceParametersCS_CollisionQuery);
-
-IMPLEMENT_NIAGARA_DI_PARAMETER(UNiagaraDataInterfaceCollisionQuery, FNiagaraDataInterfaceParametersCS_CollisionQuery);
-
 #undef LOCTEXT_NAMESPACE
diff --git a/Engine/Plugins/FX/Niagara/Source/Niagara/Private/NiagaraGpuComputeDispatch.cpp b/Engine/Plugins/FX/Niagara/Source/Niagara/Private/NiagaraGpuComputeDispatch.cpp
--- a/Engine/Plugins/FX/Niagara/Source/Niagara/Private/NiagaraGpuComputeDispatch.cpp
+++ b/Engine/Plugins/FX/Niagara/Source/Niagara/Private/NiagaraGpuComputeDispatch.cpp
@@ -1409,7 +1409,8 @@ void FNiagaraGpuComputeDispatch::DispatchStage(FRDGBuilder& GraphBuilder, const
 
 	checkf(DispatchNumThreads.X * DispatchNumThreads.Y * DispatchNumThreads.Z > 0, TEXT("DispatchNumThreads(%d, %d, %d) is invalid"), DispatchNumThreads.X, DispatchNumThreads.Y, DispatchNumThreads.Z);
 
-	const FShaderParametersMetadata* ShaderParametersMetadata = InstanceData.Context->GPUScript_RT->GetScriptParametersMetadata()->ShaderParametersMetadata.Get();
+	const FNiagaraShaderScriptParametersMetadata& NiagaraShaderParametersMetadata = InstanceData.Context->GPUScript_RT->GetScriptParametersMetadata().Get();
+	const FShaderParametersMetadata* ShaderParametersMetadata = NiagaraShaderParametersMetadata.ShaderParametersMetadata.Get();
 	FNiagaraShader::FParameters* DispatchParameters = reinterpret_cast<FNiagaraShader::FParameters*>(GraphBuilder.Alloc(ShaderParametersMetadata->GetSize(), SHADER_PARAMETER_STRUCT_ALIGNMENT));
 	FMemory::Memset(DispatchParameters, 0, ShaderParametersMetadata->GetSize());
 
@@ -1505,7 +1506,7 @@ void FNiagaraGpuComputeDispatch::DispatchStage(FRDGBuilder& GraphBuilder, const
 	const TShaderRef<FNiagaraShader> ComputeShader = InstanceData.Context->GPUScript_RT->GetShader(SimStageData.StageIndex);
 
 	// Set data interface parameters
-	SetDataInterfaceParameters(GraphBuilder, Tick, InstanceData, ComputeShader, SimStageData, reinterpret_cast<uint8*>(DispatchParameters));
+	SetDataInterfaceParameters(GraphBuilder, Tick, InstanceData, ComputeShader, SimStageData, NiagaraShaderParametersMetadata, reinterpret_cast<uint8*>(DispatchParameters));
 
 	// Set tick parameters
 	Tick.GetGlobalParameters(InstanceData, &DispatchParameters->GlobalParameters);
@@ -1995,7 +1996,7 @@ void FNiagaraGpuComputeDispatch::LegacySetDataInterfaceParameters(FRHICommandLis
 	}
 
 	const FNiagaraShaderMapPointerTable& PointerTable = ComputeShader.GetPointerTable();
-	const TMemoryImageArray<FNiagaraDataInterfaceParamRef>& DIParameters = ComputeShader->GetDIParameters();
+	TConstArrayView<FNiagaraDataInterfaceParamRef> DIParameters = ComputeShader->GetDIParameters();
 
 	for (int32 iDataInterface=0; iDataInterface < NumDataInterfaces; ++iDataInterface)
 	{
@@ -2018,7 +2019,7 @@ void FNiagaraGpuComputeDispatch::LegacyUnsetDataInterfaceParameters(FRHICommandL
 	uint32 InterfaceIndex = 0;
 	for (FNiagaraDataInterfaceProxy* Interface : InstanceData.DataInterfaceProxies)
 	{
-		const TMemoryImageArray<FNiagaraDataInterfaceParamRef>& DIParameters = ComputeShader->GetDIParameters();
+		TConstArrayView<FNiagaraDataInterfaceParamRef> DIParameters = ComputeShader->GetDIParameters();
 		const FNiagaraDataInterfaceParamRef& DIParam = DIParameters[InterfaceIndex];
 		if (DIParam.Parameters.IsValid() && (DIParam.ShaderParametersOffset == INDEX_NONE))
 		{
@@ -2034,7 +2035,7 @@ void FNiagaraGpuComputeDispatch::LegacyResetDataInterfaces(FRHICommandList& RHIC
 {
 	// Note: All stages will contain the same bindings so if they are valid for one they are valid for all, this could change in the future
 	const FNiagaraShaderRef& ComputeShader = InstanceData.Context->GPUScript_RT->GetShader(0);
-	const TMemoryImageArray<FNiagaraDataInterfaceParamRef>& DIParameters = ComputeShader->GetDIParameters();
+	TConstArrayView<FNiagaraDataInterfaceParamRef> DIParameters = ComputeShader->GetDIParameters();
 
 	uint32 InterfaceIndex = 0;
 	for (FNiagaraDataInterfaceProxy* Interface : InstanceData.DataInterfaceProxies)
@@ -2053,7 +2054,7 @@ void FNiagaraGpuComputeDispatch::LegacyPreStageInterface(FRHICommandList& RHICmd
 {
 	// Note: All stages will contain the same bindings so if they are valid for one they are valid for all, this could change in the future
 	const FNiagaraShaderRef& ComputeShader = InstanceData.Context->GPUScript_RT->GetShader(0);
-	const TMemoryImageArray<FNiagaraDataInterfaceParamRef>& DIParameters = ComputeShader->GetDIParameters();
+	TConstArrayView<FNiagaraDataInterfaceParamRef> DIParameters = ComputeShader->GetDIParameters();
 
 	uint32 InterfaceIndex = 0;
 	for (FNiagaraDataInterfaceProxy* Interface : InstanceData.DataInterfaceProxies)
@@ -2077,7 +2078,7 @@ void FNiagaraGpuComputeDispatch::LegacyPostStageInterface(FRHICommandList& RHICm
 {
 	// Note: All stages will contain the same bindings so if they are valid for one they are valid for all, this could change in the future
 	const FNiagaraShaderRef& ComputeShader = InstanceData.Context->GPUScript_RT->GetShader(0);
-	const TMemoryImageArray<FNiagaraDataInterfaceParamRef>& DIParameters = ComputeShader->GetDIParameters();
+	TConstArrayView<FNiagaraDataInterfaceParamRef> DIParameters = ComputeShader->GetDIParameters();
 
 	uint32 InterfaceIndex = 0;
 	for (FNiagaraDataInterfaceProxy* Interface : InstanceData.DataInterfaceProxies)
@@ -2101,7 +2102,7 @@ void FNiagaraGpuComputeDispatch::LegacyPostSimulateInterface(FRHICommandList& RH
 {
 	// Note: All stages will contain the same bindings so if they are valid for one they are valid for all, this could change in the future
 	const FNiagaraShaderRef& ComputeShader = InstanceData.Context->GPUScript_RT->GetShader(0);
-	const TMemoryImageArray<FNiagaraDataInterfaceParamRef>& DIParameters = ComputeShader->GetDIParameters();
+	TConstArrayView<FNiagaraDataInterfaceParamRef> DIParameters = ComputeShader->GetDIParameters();
 
 	uint32 InterfaceIndex = 0;
 	for (FNiagaraDataInterfaceProxy* Interface : InstanceData.DataInterfaceProxies)
@@ -2130,7 +2131,7 @@ void FNiagaraGpuComputeDispatch::ResetDataInterfaces(FRDGBuilder& GraphBuilder,
 	}
 
 	const FNiagaraShaderRef& ComputeShader = InstanceData.Context->GPUScript_RT->GetShader(0);
-	const TMemoryImageArray<FNiagaraDataInterfaceParamRef>& DIParameters = ComputeShader->GetDIParameters();
+	TConstArrayView<FNiagaraDataInterfaceParamRef> DIParameters = ComputeShader->GetDIParameters();
 
 	FNDIGpuComputeResetContext Context(GraphBuilder, *this, Tick.SystemInstanceID);
 	for (int32 iDataInterface=0; iDataInterface < NumDataInterfaces; ++iDataInterface)
@@ -2144,7 +2145,7 @@ void FNiagaraGpuComputeDispatch::ResetDataInterfaces(FRDGBuilder& GraphBuilder,
 	}
 }
 
-void FNiagaraGpuComputeDispatch::SetDataInterfaceParameters(FRDGBuilder& GraphBuilder, const FNiagaraGPUSystemTick& Tick, const FNiagaraComputeInstanceData& InstanceData, const FNiagaraShaderRef& ComputeShader, const FNiagaraSimStageData& SimStageData, uint8* ParametersStructure) const
+void FNiagaraGpuComputeDispatch::SetDataInterfaceParameters(FRDGBuilder& GraphBuilder, const FNiagaraGPUSystemTick& Tick, const FNiagaraComputeInstanceData& InstanceData, const FNiagaraShaderRef& ComputeShader, const FNiagaraSimStageData& SimStageData, const FNiagaraShaderScriptParametersMetadata& NiagaraShaderParametersMetadata, uint8* ParametersStructure) const
 {
 	const int32 NumDataInterfaces = InstanceData.DataInterfaceProxies.Num();
 	if (NumDataInterfaces == 0)
@@ -2153,9 +2154,9 @@ void FNiagaraGpuComputeDispatch::SetDataInterfaceParameters(FRDGBuilder& GraphBu
 	}
 
 	const FNiagaraShaderMapPointerTable& PointerTable = ComputeShader.GetPointerTable();
-	const TMemoryImageArray<FNiagaraDataInterfaceParamRef>& DIParameters = ComputeShader->GetDIParameters();
+	TConstArrayView<FNiagaraDataInterfaceParamRef> DIParameters = ComputeShader->GetDIParameters();
 
-	FNiagaraDataInterfaceSetShaderParametersContext Context(GraphBuilder, *this, Tick, InstanceData, SimStageData, ComputeShader, ParametersStructure);
+	FNiagaraDataInterfaceSetShaderParametersContext Context(GraphBuilder, *this, Tick, InstanceData, SimStageData, ComputeShader, NiagaraShaderParametersMetadata, ParametersStructure);
 	for ( int32 iDataInterface=0; iDataInterface < NumDataInterfaces; ++iDataInterface )
 	{
 		FNiagaraDataInterfaceProxy* DataInterfaceProxy = InstanceData.DataInterfaceProxies[iDataInterface];
@@ -2178,7 +2179,7 @@ void FNiagaraGpuComputeDispatch::PreStageInterface(FRDGBuilder& GraphBuilder, co
 	}
 
 	const FNiagaraShaderRef& ComputeShader = InstanceData.Context->GPUScript_RT->GetShader(0);
-	const TMemoryImageArray<FNiagaraDataInterfaceParamRef>& DIParameters = ComputeShader->GetDIParameters();
+	TConstArrayView<FNiagaraDataInterfaceParamRef> DIParameters = ComputeShader->GetDIParameters();
 
 	FNDIGpuComputePreStageContext Context(GraphBuilder, *this, Tick, InstanceData, SimStageData);
 	for (int32 iDataInterface = 0; iDataInterface < NumDataInterfaces; ++iDataInterface)
@@ -2202,7 +2203,7 @@ void FNiagaraGpuComputeDispatch::PostStageInterface(FRDGBuilder& GraphBuilder, c
 	}
 
 	const FNiagaraShaderRef& ComputeShader = InstanceData.Context->GPUScript_RT->GetShader(0);
-	const TMemoryImageArray<FNiagaraDataInterfaceParamRef>& DIParameters = ComputeShader->GetDIParameters();
+	TConstArrayView<FNiagaraDataInterfaceParamRef> DIParameters = ComputeShader->GetDIParameters();
 
 	FNDIGpuComputePostStageContext Context(GraphBuilder, *this, Tick, InstanceData, SimStageData);
 	for (int32 iDataInterface = 0; iDataInterface < NumDataInterfaces; ++iDataInterface)
@@ -2226,7 +2227,7 @@ void FNiagaraGpuComputeDispatch::PostSimulateInterface(FRDGBuilder& GraphBuilder
 	}
 
 	const FNiagaraShaderRef& ComputeShader = InstanceData.Context->GPUScript_RT->GetShader(0);
-	const TMemoryImageArray<FNiagaraDataInterfaceParamRef>& DIParameters = ComputeShader->GetDIParameters();
+	TConstArrayView<FNiagaraDataInterfaceParamRef> DIParameters = ComputeShader->GetDIParameters();
 
 	FNDIGpuComputePostSimulateContext Context(GraphBuilder, *this, Tick.SystemInstanceID, SimStageData.bSetDataToRender);
 	for (int32 iDataInterface = 0; iDataInterface < NumDataInterfaces; ++iDataInterface)
diff --git a/Engine/Plugins/FX/Niagara/Source/Niagara/Private/NiagaraGpuComputeDispatch.h b/Engine/Plugins/FX/Niagara/Source/Niagara/Private/NiagaraGpuComputeDispatch.h
--- a/Engine/Plugins/FX/Niagara/Source/Niagara/Private/NiagaraGpuComputeDispatch.h
+++ b/Engine/Plugins/FX/Niagara/Source/Niagara/Private/NiagaraGpuComputeDispatch.h
@@ -158,7 +158,7 @@ public:
 	void LegacyPostSimulateInterface(FRHICommandList& RHICmdList, const FNiagaraGPUSystemTick& Tick, const FNiagaraComputeInstanceData& InstanceData) const;
 
 	void ResetDataInterfaces(FRDGBuilder& GraphBuilder, const FNiagaraGPUSystemTick& Tick, const FNiagaraComputeInstanceData& InstanceData) const;
-	void SetDataInterfaceParameters(FRDGBuilder& GraphBuilder, const FNiagaraGPUSystemTick& Tick, const FNiagaraComputeInstanceData& InstanceData, const FNiagaraShaderRef& ComputeShader, const FNiagaraSimStageData& SimStageData, uint8* ParametersStructure) const;
+	void SetDataInterfaceParameters(FRDGBuilder& GraphBuilder, const FNiagaraGPUSystemTick& Tick, const FNiagaraComputeInstanceData& InstanceData, const FNiagaraShaderRef& ComputeShader, const FNiagaraSimStageData& SimStageData, const FNiagaraShaderScriptParametersMetadata& NiagaraShaderParametersMetadata, uint8* ParametersStructure) const;
 	void PreStageInterface(FRDGBuilder& GraphBuilder, const FNiagaraGPUSystemTick& Tick, const FNiagaraComputeInstanceData& InstanceData, const FNiagaraSimStageData& SimStageData, TSet<FNiagaraDataInterfaceProxy*>& ProxiesToFinalize) const;
 	void PostStageInterface(FRDGBuilder& GraphBuilder, const FNiagaraGPUSystemTick& Tick, const FNiagaraComputeInstanceData& InstanceData, const FNiagaraSimStageData& SimStageData, TSet<FNiagaraDataInterfaceProxy*>& ProxiesToFinalize) const;
 	void PostSimulateInterface(FRDGBuilder& GraphBuilder, const FNiagaraGPUSystemTick& Tick, const FNiagaraComputeInstanceData& InstanceData, const FNiagaraSimStageData& SimStageData) const;
diff --git a/Engine/Plugins/FX/Niagara/Source/NiagaraShader/Private/NiagaraAsyncGpuTraceProviderGsdf.cpp b/Engine/Plugins/FX/Niagara/Source/NiagaraShader/Private/NiagaraAsyncGpuTraceProviderGsdf.cpp
--- a/Engine/Plugins/FX/Niagara/Source/NiagaraShader/Private/NiagaraAsyncGpuTraceProviderGsdf.cpp
+++ b/Engine/Plugins/FX/Niagara/Source/NiagaraShader/Private/NiagaraAsyncGpuTraceProviderGsdf.cpp
@@ -4,6 +4,7 @@
 
 #include "GlobalShader.h"
 #include "NiagaraGpuComputeDispatchInterface.h"
+#include "NiagaraDistanceFieldHelper.h"
 #include "NiagaraSettings.h"
 #include "Renderer/Private/ScenePrivate.h"
 #include "Renderer/Private/SceneRendering.h"
@@ -73,38 +74,6 @@ IMPLEMENT_GLOBAL_SHADER(FNiagaraRayMarchGlobalSdfCS, "/Plugin/FX/Niagara/Private
 ////////////////////////////////////////////////////////////////////////////////
 ////////////////////////////////////////////////////////////////////////////////
 
-// todo - currently duplicated from SetupGlobalDistanceFieldParameters (GlobalDistanceField.cpp) because of
-// problems getting it properly exported from the dll
-static void AssignGlobalDistanceFieldParameters(const FGlobalDistanceFieldParameterData& ParameterData, FGlobalDistanceFieldParameters2& ShaderParameters)
-{
-	ShaderParameters.GlobalDistanceFieldPageAtlasTexture = OrBlack3DIfNull(ParameterData.PageAtlasTexture);
-	ShaderParameters.GlobalDistanceFieldCoverageAtlasTexture = OrBlack3DIfNull(ParameterData.CoverageAtlasTexture);
-	ShaderParameters.GlobalDistanceFieldPageTableTexture = OrBlack3DUintIfNull(ParameterData.PageTableTexture);
-	ShaderParameters.GlobalDistanceFieldMipTexture = OrBlack3DIfNull(ParameterData.MipTexture);
-
-	for (int32 Index = 0; Index < GlobalDistanceField::MaxClipmaps; Index++)
-	{
-		ShaderParameters.GlobalVolumeCenterAndExtent[Index] = ParameterData.CenterAndExtent[Index];
-		ShaderParameters.GlobalVolumeWorldToUVAddAndMul[Index] = ParameterData.WorldToUVAddAndMul[Index];
-		ShaderParameters.GlobalDistanceFieldMipWorldToUVScale[Index] = ParameterData.MipWorldToUVScale[Index];
-		ShaderParameters.GlobalDistanceFieldMipWorldToUVBias[Index] = ParameterData.MipWorldToUVBias[Index];
-	}
-
-	ShaderParameters.GlobalDistanceFieldMipFactor = ParameterData.MipFactor;
-	ShaderParameters.GlobalDistanceFieldMipTransition = ParameterData.MipTransition;
-	ShaderParameters.GlobalDistanceFieldClipmapSizeInPages = ParameterData.ClipmapSizeInPages;
-	ShaderParameters.GlobalDistanceFieldInvPageAtlasSize = (FVector3f)ParameterData.InvPageAtlasSize;
-	ShaderParameters.GlobalDistanceFieldInvCoverageAtlasSize = (FVector3f)ParameterData.InvCoverageAtlasSize;
-	ShaderParameters.GlobalVolumeDimension = ParameterData.GlobalDFResolution;
-	ShaderParameters.GlobalVolumeTexelSize = 1.0f / ParameterData.GlobalDFResolution;
-	ShaderParameters.MaxGlobalDFAOConeDistance = ParameterData.MaxDFAOConeDistance;
-	ShaderParameters.NumGlobalSDFClipmaps = ParameterData.NumGlobalSDFClipmaps;
-}
-
-////////////////////////////////////////////////////////////////////////////////
-////////////////////////////////////////////////////////////////////////////////
-////////////////////////////////////////////////////////////////////////////////
-
 const FNiagaraAsyncGpuTraceProvider::EProviderType FNiagaraAsyncGpuTraceProviderGsdf::Type = ENDICollisionQuery_AsyncGpuTraceProvider::GSDF;
 
 FNiagaraAsyncGpuTraceProviderGsdf::FNiagaraAsyncGpuTraceProviderGsdf(EShaderPlatform InShaderPlatform, FNiagaraGpuComputeDispatchInterface* Dispatcher)
@@ -159,7 +128,7 @@ void FNiagaraAsyncGpuTraceProviderGsdf::IssueTraces(FRHICommandList& RHICmdList,
 	Params.Traces = Request.TracesBuffer->SRV;
 	Params.TraceCounts = Request.TraceCountsBuffer->SRV;
 	Params.Results = Request.ResultsBuffer->UAV;
-	AssignGlobalDistanceFieldParameters(m_DistanceFieldData, Params.GlobalDistanceFieldParameters);
+	FNiagaraDistanceFieldHelper::SetGlobalDistanceFieldParameters(&m_DistanceFieldData, Params.GlobalDistanceFieldParameters);
 	Params.TracesOffset = Request.TracesOffset;
 	Params.TraceCountsOffset = Request.TraceCountsOffset;
 	Params.ResultsOffset = Request.ResultsOffset;
diff --git a/Engine/Plugins/FX/Niagara/Source/NiagaraShader/Private/NiagaraDistanceFieldHelper.cpp b/Engine/Plugins/FX/Niagara/Source/NiagaraShader/Private/NiagaraDistanceFieldHelper.cpp
new file mode 100644
--- /dev/null
+++ b/Engine/Plugins/FX/Niagara/Source/NiagaraShader/Private/NiagaraDistanceFieldHelper.cpp
@@ -0,0 +1,67 @@
+// Copyright Epic Games, Inc. All Rights Reserved.
+
+#include "NiagaraDistanceFieldHelper.h"
+
+// todo - currently duplicated from SetupGlobalDistanceFieldParameters (GlobalDistanceField.cpp) because of problems getting it properly exported from the dll
+void FNiagaraDistanceFieldHelper::SetGlobalDistanceFieldParameters(const FGlobalDistanceFieldParameterData* OptionalParameterData, FGlobalDistanceFieldParameters2& ShaderParameters)
+{
+	if (OptionalParameterData != nullptr)
+	{
+		ShaderParameters.GlobalDistanceFieldPageAtlasTexture = OrBlack3DIfNull(OptionalParameterData->PageAtlasTexture);
+		ShaderParameters.GlobalDistanceFieldCoverageAtlasTexture = OrBlack3DIfNull(OptionalParameterData->CoverageAtlasTexture);
+		ShaderParameters.GlobalDistanceFieldPageTableTexture = OrBlack3DUintIfNull(OptionalParameterData->PageTableTexture);
+		ShaderParameters.GlobalDistanceFieldMipTexture = OrBlack3DIfNull(OptionalParameterData->MipTexture);
+
+		for (int32 Index = 0; Index < GlobalDistanceField::MaxClipmaps; Index++)
+		{
+			ShaderParameters.GlobalVolumeCenterAndExtent[Index] = OptionalParameterData->CenterAndExtent[Index];
+			ShaderParameters.GlobalVolumeWorldToUVAddAndMul[Index] = OptionalParameterData->WorldToUVAddAndMul[Index];
+			ShaderParameters.GlobalDistanceFieldMipWorldToUVScale[Index] = OptionalParameterData->MipWorldToUVScale[Index];
+			ShaderParameters.GlobalDistanceFieldMipWorldToUVBias[Index] = OptionalParameterData->MipWorldToUVBias[Index];
+		}
+
+		ShaderParameters.GlobalDistanceFieldMipFactor = OptionalParameterData->MipFactor;
+		ShaderParameters.GlobalDistanceFieldMipTransition = OptionalParameterData->MipTransition;
+		ShaderParameters.GlobalDistanceFieldClipmapSizeInPages = OptionalParameterData->ClipmapSizeInPages;
+		ShaderParameters.GlobalDistanceFieldInvPageAtlasSize = (FVector3f)OptionalParameterData->InvPageAtlasSize;
+		ShaderParameters.GlobalDistanceFieldInvCoverageAtlasSize = (FVector3f)OptionalParameterData->InvCoverageAtlasSize;
+		ShaderParameters.GlobalVolumeDimension = OptionalParameterData->GlobalDFResolution;
+		ShaderParameters.GlobalVolumeTexelSize = 1.0f / OptionalParameterData->GlobalDFResolution;
+		ShaderParameters.MaxGlobalDFAOConeDistance = OptionalParameterData->MaxDFAOConeDistance;
+		ShaderParameters.NumGlobalSDFClipmaps = OptionalParameterData->NumGlobalSDFClipmaps;
+		ShaderParameters.FullyCoveredExpandSurfaceScale = 0.0f;//GLumenSceneGlobalSDFFullyCoveredExpandSurfaceScale;
+		ShaderParameters.UncoveredExpandSurfaceScale = 0.0f;//GLumenSceneGlobalSDFUncoveredExpandSurfaceScale;
+		ShaderParameters.UncoveredMinStepScale = 0.0f;//GLumenSceneGlobalSDFUncoveredMinStepScale;
+	}
+	else
+	{
+		ShaderParameters.GlobalDistanceFieldPageAtlasTexture = GBlackVolumeTexture->TextureRHI;
+		ShaderParameters.GlobalDistanceFieldCoverageAtlasTexture = GBlackVolumeTexture->TextureRHI;
+		ShaderParameters.GlobalDistanceFieldPageTableTexture = GBlackUintVolumeTexture->TextureRHI;
+		ShaderParameters.GlobalDistanceFieldMipTexture = GBlackVolumeTexture->TextureRHI;
+
+		for (int32 Index = 0; Index < GlobalDistanceField::MaxClipmaps; Index++)
+		{
+			ShaderParameters.GlobalVolumeCenterAndExtent[Index] = FVector4f::Zero();
+			ShaderParameters.GlobalVolumeWorldToUVAddAndMul[Index] = FVector4f::Zero();
+			ShaderParameters.GlobalDistanceFieldMipWorldToUVScale[Index] = FVector4f::Zero();
+			ShaderParameters.GlobalDistanceFieldMipWorldToUVBias[Index] = FVector4f::Zero();
+		}
+
+		ShaderParameters.GlobalDistanceFieldMipFactor = 0.0f;
+		ShaderParameters.GlobalDistanceFieldMipTransition = 0.0f;
+		ShaderParameters.GlobalDistanceFieldClipmapSizeInPages = 0;
+		ShaderParameters.GlobalDistanceFieldInvPageAtlasSize = FVector3f::ZeroVector;
+		ShaderParameters.GlobalDistanceFieldInvCoverageAtlasSize = FVector3f::ZeroVector;
+		ShaderParameters.GlobalVolumeDimension = 0.0f;
+		ShaderParameters.GlobalVolumeTexelSize = 0.0f;
+		ShaderParameters.MaxGlobalDFAOConeDistance = 0.0f;
+		ShaderParameters.NumGlobalSDFClipmaps = 0;
+		ShaderParameters.FullyCoveredExpandSurfaceScale = 0.0f;
+		ShaderParameters.UncoveredExpandSurfaceScale = 0.0f;
+		ShaderParameters.UncoveredMinStepScale = 0.0f;
+		ShaderParameters.FullyCoveredExpandSurfaceScale = 0.0f;//GLumenSceneGlobalSDFFullyCoveredExpandSurfaceScale;
+		ShaderParameters.UncoveredExpandSurfaceScale = 0.0f;//GLumenSceneGlobalSDFUncoveredExpandSurfaceScale;
+		ShaderParameters.UncoveredMinStepScale = 0.0f;//GLumenSceneGlobalSDFUncoveredMinStepScale;
+	}
+}
diff --git a/Engine/Plugins/FX/Niagara/Source/NiagaraShader/Private/NiagaraShader.cpp b/Engine/Plugins/FX/Niagara/Source/NiagaraShader/Private/NiagaraShader.cpp
--- a/Engine/Plugins/FX/Niagara/Source/NiagaraShader/Private/NiagaraShader.cpp
+++ b/Engine/Plugins/FX/Niagara/Source/NiagaraShader/Private/NiagaraShader.cpp
@@ -1192,6 +1192,7 @@ FNiagaraShader::FNiagaraShader(const FNiagaraShaderType::CompiledShaderInitializ
 		{
 			DataInterfaceParamRef.Parameters = CDODataInterface->CreateShaderStorage(DataInterfaceParamInfo, Initializer.ParameterMap);
 			checkf(DataInterfaceParamRef.Parameters == nullptr || CDODataInterface->GetShaderStorageType() != nullptr, TEXT("DataInterface(%s) provides shader storage but did not implement GetShaderStorageType"), *GetNameSafe(CDODataInterface->GetClass()));
+			checkf(DataInterfaceParamRef.Parameters == nullptr || CDODataInterface->GetShaderStorageType()->Interface == ETypeLayoutInterface::NonVirtual, TEXT("DataInterface(%s) shader storage is either abstract or virtual which is not allowed"), *GetNameSafe(CDODataInterface->GetClass()));
 		}
 	}
 }
diff --git a/Engine/Plugins/FX/Niagara/Source/NiagaraShader/Private/NiagaraShared.cpp b/Engine/Plugins/FX/Niagara/Source/NiagaraShader/Private/NiagaraShared.cpp
--- a/Engine/Plugins/FX/Niagara/Source/NiagaraShader/Private/NiagaraShared.cpp
+++ b/Engine/Plugins/FX/Niagara/Source/NiagaraShader/Private/NiagaraShared.cpp
@@ -692,7 +692,7 @@ void FNiagaraShaderScript::BuildScriptParametersMetadata(const FNiagaraShaderScr
 		if (CDODataInterface->UseLegacyShaderBindings() == false)
 		{
 			const uint32 NextMemberOffset = ShaderMetadataBuilder.GetNextMemberOffset();
-			FNiagaraShaderParametersBuilder ShaderParametersBuilder(DataInterfaceParamInfo, ScriptParametersMetadata->LooseMetadataNames, ShaderMetadataBuilder);
+			FNiagaraShaderParametersBuilder ShaderParametersBuilder(DataInterfaceParamInfo, ScriptParametersMetadata->LooseMetadataNames, ScriptParametersMetadata->StructIncludeInfos, ShaderMetadataBuilder);
 			CDODataInterface->BuildShaderParameters(ShaderParametersBuilder);
 			DataInterfaceParamInfo.ShaderParametersOffset = NextMemberOffset;
 		}
diff --git a/Engine/Plugins/FX/Niagara/Source/NiagaraShader/Public/NiagaraDistanceFieldHelper.h b/Engine/Plugins/FX/Niagara/Source/NiagaraShader/Public/NiagaraDistanceFieldHelper.h
new file mode 100644
--- /dev/null
+++ b/Engine/Plugins/FX/Niagara/Source/NiagaraShader/Public/NiagaraDistanceFieldHelper.h
@@ -0,0 +1,11 @@
+// Copyright Epic Games, Inc. All Rights Reserved.
+
+#pragma once
+
+#include "GlobalDistanceFieldParameters.h"
+#include "Renderer/Private/DistanceFieldLightingShared.h"
+
+namespace FNiagaraDistanceFieldHelper
+{
+	NIAGARASHADER_API void SetGlobalDistanceFieldParameters(const FGlobalDistanceFieldParameterData* OptionalParameterData, FGlobalDistanceFieldParameters2& ShaderParameters);
+}
diff --git a/Engine/Plugins/FX/Niagara/Source/NiagaraShader/Public/NiagaraShader.h b/Engine/Plugins/FX/Niagara/Source/NiagaraShader/Public/NiagaraShader.h
--- a/Engine/Plugins/FX/Niagara/Source/NiagaraShader/Public/NiagaraShader.h
+++ b/Engine/Plugins/FX/Niagara/Source/NiagaraShader/Public/NiagaraShader.h
@@ -220,9 +220,9 @@ public:
 	FNiagaraShader() {}
 	FNiagaraShader(const FNiagaraShaderType::CompiledShaderInitializerType& Initializer);
 
-	const TMemoryImageArray<FNiagaraDataInterfaceParamRef>& GetDIParameters()
+	TConstArrayView<FNiagaraDataInterfaceParamRef> GetDIParameters()
 	{
-		return DataInterfaceParameters;
+		return MakeArrayView(DataInterfaceParameters);
 	}
 
 	LAYOUT_FIELD(bool, bNeedsViewUniformBuffer);
@@ -232,10 +232,6 @@ private:
 	// Data about parameters used for each Data Interface.
 	LAYOUT_FIELD(TMemoryImageArray<FNiagaraDataInterfaceParamRef>, DataInterfaceParameters);
 
-	/*
-	FDebugUniformExpressionSet	DebugUniformExpressionSet;
-	FRHIUniformBufferLayout		DebugUniformExpressionUBLayout;
-	*/
 	LAYOUT_FIELD(FMemoryImageString, DebugDescription);
 };
 
diff --git a/Engine/Plugins/FX/Niagara/Source/NiagaraShader/Public/NiagaraShaderParametersBuilder.h b/Engine/Plugins/FX/Niagara/Source/NiagaraShader/Public/NiagaraShaderParametersBuilder.h
--- a/Engine/Plugins/FX/Niagara/Source/NiagaraShader/Public/NiagaraShaderParametersBuilder.h
+++ b/Engine/Plugins/FX/Niagara/Source/NiagaraShader/Public/NiagaraShaderParametersBuilder.h
@@ -8,8 +8,9 @@
 class NIAGARASHADER_API FNiagaraShaderParametersBuilder
 {
 public:
-	explicit FNiagaraShaderParametersBuilder(const FNiagaraDataInterfaceGPUParamInfo& InGPUParamInfo, TArray<FString>& InLooseNames, FShaderParametersMetadataBuilder& InMetadataBuilder)
+	explicit FNiagaraShaderParametersBuilder(const FNiagaraDataInterfaceGPUParamInfo& InGPUParamInfo, TArray<FString>& InLooseNames, TArray<FNiagaraDataInterfaceStructIncludeInfo>& InStructIncludeInfos, FShaderParametersMetadataBuilder& InMetadataBuilder)
 		: GPUParamInfo(InGPUParamInfo)
+		, StructIncludeInfos(InStructIncludeInfos)
 		, LooseNames(InLooseNames)
 		, MetadataBuilder(InMetadataBuilder)
 	{
@@ -38,10 +39,23 @@ public:
 		MetadataBuilder.AddNestedStruct<T>(*GPUParamInfo.DataInterfaceHLSLSymbol);
 	}
 
-	// Adds a shadera parameters structure that is global in scope
+	// Adds a shader parameters structure that is global in scope
 	// i.e. if the structure contained "MyFloat" the shader variable would be named "MyFloat"
 	template<typename T> void AddIncludedStruct()
 	{
+		const FShaderParametersMetadata* StructMetadata = TShaderParameterStructTypeInfo<T>::GetStructMetadata();
+		for (const FNiagaraDataInterfaceStructIncludeInfo& Existing : StructIncludeInfos)
+		{
+			if (Existing.StructMetadata == StructMetadata)
+			{
+				return;
+			}
+		}
+
+		FNiagaraDataInterfaceStructIncludeInfo& NewInfo = StructIncludeInfos.AddDefaulted_GetRef();
+		NewInfo.StructMetadata = StructMetadata;
+		NewInfo.ParamterOffset = Align(MetadataBuilder.GetNextMemberOffset(), SHADER_PARAMETER_STRUCT_ALIGNMENT);
+
 		MetadataBuilder.AddIncludedStruct<T>();
 	}
 
@@ -49,6 +63,7 @@ public:
 
 private:
 	const FNiagaraDataInterfaceGPUParamInfo& GPUParamInfo;
+	TArray<FNiagaraDataInterfaceStructIncludeInfo>& StructIncludeInfos;
 	TArray<FString>& LooseNames;
 	FShaderParametersMetadataBuilder& MetadataBuilder;
 };
diff --git a/Engine/Plugins/FX/Niagara/Source/NiagaraShader/Public/NiagaraShared.h b/Engine/Plugins/FX/Niagara/Source/NiagaraShader/Public/NiagaraShared.h
--- a/Engine/Plugins/FX/Niagara/Source/NiagaraShader/Public/NiagaraShared.h
+++ b/Engine/Plugins/FX/Niagara/Source/NiagaraShader/Public/NiagaraShared.h
@@ -243,6 +243,14 @@ struct FNiagaraShaderScriptExternalConstant
 	FString	Name;
 };
 
+/** Shader side information about global structures includes. */
+struct FNiagaraDataInterfaceStructIncludeInfo
+{
+	const FShaderParametersMetadata* StructMetadata = nullptr;
+	int32 ParamterOffset = 0;
+};
+
+
 /* Data provided from the translator to track information requires to bind data interface or parameters at runtime. */
 USTRUCT()
 struct NIAGARASHADER_API FNiagaraShaderScriptParametersMetadata
@@ -262,6 +270,7 @@ struct NIAGARASHADER_API FNiagaraShaderScriptParametersMetadata
 	TArray<FNiagaraShaderScriptExternalConstant> ExternalConstants;
 
 	// Runtime generated bindings, not serialized
+	TArray<FNiagaraDataInterfaceStructIncludeInfo> StructIncludeInfos;
 	TSharedPtr<FShaderParametersMetadata> ShaderParametersMetadata;
 };
 
@@ -286,7 +295,6 @@ public:
 	LAYOUT_FIELD_WITH_WRITER(TMemoryImagePtr<FNiagaraDataInterfaceParametersCS>, Parameters, WriteFrozenParameters);
 };
 
-
 /** Stores outputs from the script compile that need to be saved. */
 class FNiagaraComputeShaderCompilationOutput
 {
diff --git a/Engine/Source/Runtime/Core/Private/UObject/UE5MainStreamObjectVersion.cpp b/Engine/Source/Runtime/Core/Private/UObject/UE5MainStreamObjectVersion.cpp
--- a/Engine/Source/Runtime/Core/Private/UObject/UE5MainStreamObjectVersion.cpp
+++ b/Engine/Source/Runtime/Core/Private/UObject/UE5MainStreamObjectVersion.cpp
@@ -14,7 +14,7 @@ TMap<FGuid, FGuid> FUE5MainStreamObjectVersion::GetSystemGuids()
 	SystemGuids.Add(DevGuids.MATERIALSHADERMAP_DERIVEDDATA_VER, FGuid("65C01D817C9A4EEFAE9E988D41A1F3DD"));
 	SystemGuids.Add(DevGuids.NANITE_DERIVEDDATA_VER, FGuid("7A3EA93546624597BEBE2ADFDE3A7628"));
 	SystemGuids.Add(DevGuids.NIAGARASHADERMAP_DERIVEDDATA_VER, FGuid("8A37C45D24F2423CBE5F8F371DE33575"));
-	SystemGuids.Add(DevGuids.Niagara_LatestScriptCompileVersion, FGuid("B531E01D70C42444B4C76046D11CDA22"));
+	SystemGuids.Add(DevGuids.Niagara_LatestScriptCompileVersion, FGuid("F8EAE1D4DF12704EB35657412E887A7D"));
 	SystemGuids.Add(DevGuids.SkeletalMeshDerivedDataVersion, FGuid("FB7B2FB546DE49B2A2F88B7F69763CEE"));
 	SystemGuids.Add(DevGuids.STATICMESH_DERIVEDDATA_VER, FGuid("98F7D79A5811013825E75BBCC41ED3E9"));
 	SystemGuids.Add(DevGuids.POSESEARCHDB_DERIVEDDATA_VER, FGuid("4E595C2AC5E947D6BA9ABC874353E5BC"));
