Don't double buffer Grid3D collections when encountering a stage that only writes to grids

diff --git a/Engine/Plugins/FX/Niagara/Source/Niagara/Classes/NiagaraComputeExecutionContext.h b/Engine/Plugins/FX/Niagara/Source/Niagara/Classes/NiagaraComputeExecutionContext.h
--- a/Engine/Plugins/FX/Niagara/Source/Niagara/Classes/NiagaraComputeExecutionContext.h
+++ b/Engine/Plugins/FX/Niagara/Source/Niagara/Classes/NiagaraComputeExecutionContext.h
@@ -176,6 +176,7 @@ public:
 	TArray<FSimulationStageMetaData> SimStageInfo;
 
 	bool IsOutputStage(FNiagaraDataInterfaceProxy* DIProxy, uint32 SimulationStageIndex) const;
+	bool IsInputStage(FNiagaraDataInterfaceProxy* DIProxy, uint32 SimulationStageIndex) const;
 	bool IsIterationStage(FNiagaraDataInterfaceProxy* DIProxy, uint32 SimulationStageIndex) const;
 	FNiagaraDataInterfaceProxyRW* FindIterationInterface(const TArray<FNiagaraDataInterfaceProxyRW*>& InProxies, uint32 SimulationStageIndex) const;
 };
diff --git a/Engine/Plugins/FX/Niagara/Source/Niagara/Classes/NiagaraDataInterface.h b/Engine/Plugins/FX/Niagara/Source/Niagara/Classes/NiagaraDataInterface.h
--- a/Engine/Plugins/FX/Niagara/Source/Niagara/Classes/NiagaraDataInterface.h
+++ b/Engine/Plugins/FX/Niagara/Source/Niagara/Classes/NiagaraDataInterface.h
@@ -291,6 +291,7 @@ struct NIAGARA_API FNDIGpuComputePrePostStageContext : public FNDIGpuComputeCont
 	FNiagaraSystemInstanceID GetSystemInstanceID() const;
 	FVector3f GetSystemLWCTile() const;
 	bool IsOutputStage() const;
+	bool IsInputStage() const;
 	bool IsIterationStage() const;
 
 	void SetDataInterfaceProxy(FNiagaraDataInterfaceProxy* InDataInterfaceProxy) { DataInterfaceProxy = InDataInterfaceProxy; }
diff --git a/Engine/Plugins/FX/Niagara/Source/Niagara/Classes/NiagaraDataInterfaceGrid3DCollection.h b/Engine/Plugins/FX/Niagara/Source/Niagara/Classes/NiagaraDataInterfaceGrid3DCollection.h
--- a/Engine/Plugins/FX/Niagara/Source/Niagara/Classes/NiagaraDataInterfaceGrid3DCollection.h
+++ b/Engine/Plugins/FX/Niagara/Source/Niagara/Classes/NiagaraDataInterfaceGrid3DCollection.h
@@ -51,6 +51,8 @@ struct FGrid3DCollectionRWInstanceData_GameThread
 
 struct FGrid3DCollectionRWInstanceData_RenderThread
 {
+	FName SourceDIName;
+
 	FIntVector NumCells = FIntVector::ZeroValue;
 	FIntVector NumTiles = FIntVector::ZeroValue;
 	int32 TotalNumAttributes = 0;
@@ -83,7 +85,7 @@ struct FGrid3DCollectionRWInstanceData_RenderThread
 	// overrides the render thread data, which in this case is for a grid reader
 	FNiagaraDataInterfaceProxy* OtherProxy = nullptr;
 
-	void BeginSimulate(FRDGBuilder& GraphBuilder);
+	void BeginSimulate(FRDGBuilder& GraphBuilder, bool RequiresBuffering);
 	void EndSimulate();
 };
 
diff --git a/Engine/Plugins/FX/Niagara/Source/Niagara/Classes/NiagaraGPUSystemTick.h b/Engine/Plugins/FX/Niagara/Source/Niagara/Classes/NiagaraGPUSystemTick.h
--- a/Engine/Plugins/FX/Niagara/Source/Niagara/Classes/NiagaraGPUSystemTick.h
+++ b/Engine/Plugins/FX/Niagara/Source/Niagara/Classes/NiagaraGPUSystemTick.h
@@ -53,6 +53,7 @@ struct FNiagaraComputeInstanceData
 	uint32 bHasMultipleStages : 1;
 
 	bool IsOutputStage(FNiagaraDataInterfaceProxy* DIProxy, uint32 CurrentStage) const;
+	bool IsInputStage(FNiagaraDataInterfaceProxy* DIProxy, uint32 CurrentStage) const;
 	bool IsIterationStage(FNiagaraDataInterfaceProxy* DIProxy, uint32 CurrentStage) const;
 	FNiagaraDataInterfaceProxyRW* FindIterationInterface(uint32 SimulationStageIndex) const;
 };
diff --git a/Engine/Plugins/FX/Niagara/Source/Niagara/Private/NiagaraComputeExecutionContext.cpp b/Engine/Plugins/FX/Niagara/Source/Niagara/Private/NiagaraComputeExecutionContext.cpp
--- a/Engine/Plugins/FX/Niagara/Source/Niagara/Private/NiagaraComputeExecutionContext.cpp
+++ b/Engine/Plugins/FX/Niagara/Source/Niagara/Private/NiagaraComputeExecutionContext.cpp
@@ -90,6 +90,15 @@ bool FNiagaraComputeExecutionContext::IsOutputStage(FNiagaraDataInterfaceProxy*
 	return false;
 }
 
+bool FNiagaraComputeExecutionContext::IsInputStage(FNiagaraDataInterfaceProxy* DIProxy, uint32 SimulationStageIndex) const
+{
+	if (DIProxy && !DIProxy->SourceDIName.IsNone())
+	{
+		return SimStageInfo[SimulationStageIndex].InputDataInterfaces.Contains(DIProxy->SourceDIName);
+	}
+	return false;
+}
+
 bool FNiagaraComputeExecutionContext::IsIterationStage(FNiagaraDataInterfaceProxy* DIProxy, uint32 SimulationStageIndex) const
 {
 	if (DIProxy && !DIProxy->SourceDIName.IsNone())
@@ -254,6 +263,11 @@ bool FNiagaraComputeInstanceData::IsOutputStage(FNiagaraDataInterfaceProxy* DIPr
 	return Context->IsOutputStage(DIProxy, SimulationStageIndex);
 }
 
+bool FNiagaraComputeInstanceData::IsInputStage(FNiagaraDataInterfaceProxy* DIProxy, uint32 SimulationStageIndex) const
+{
+	return Context->IsInputStage(DIProxy, SimulationStageIndex);
+}
+
 bool FNiagaraComputeInstanceData::IsIterationStage(FNiagaraDataInterfaceProxy* DIProxy, uint32 SimulationStageIndex) const
 {
 	return Context->IsIterationStage(DIProxy, SimulationStageIndex);
diff --git a/Engine/Plugins/FX/Niagara/Source/Niagara/Private/NiagaraDataInterface.cpp b/Engine/Plugins/FX/Niagara/Source/Niagara/Private/NiagaraDataInterface.cpp
--- a/Engine/Plugins/FX/Niagara/Source/Niagara/Private/NiagaraDataInterface.cpp
+++ b/Engine/Plugins/FX/Niagara/Source/Niagara/Private/NiagaraDataInterface.cpp
@@ -243,6 +243,12 @@ bool FNDIGpuComputePrePostStageContext::IsOutputStage() const
 	return ComputeInstanceData.IsOutputStage(DataInterfaceProxy, SimStageData.StageIndex);
 }
 
+bool FNDIGpuComputePrePostStageContext::IsInputStage() const
+{
+	check(DataInterfaceProxy);
+	return ComputeInstanceData.IsInputStage(DataInterfaceProxy, SimStageData.StageIndex);
+}
+
 bool FNDIGpuComputePrePostStageContext::IsIterationStage() const
 {
 	check(DataInterfaceProxy);
diff --git a/Engine/Plugins/FX/Niagara/Source/Niagara/Private/NiagaraDataInterfaceGrid3DCollection.cpp b/Engine/Plugins/FX/Niagara/Source/Niagara/Private/NiagaraDataInterfaceGrid3DCollection.cpp
--- a/Engine/Plugins/FX/Niagara/Source/Niagara/Private/NiagaraDataInterfaceGrid3DCollection.cpp
+++ b/Engine/Plugins/FX/Niagara/Source/Niagara/Private/NiagaraDataInterfaceGrid3DCollection.cpp
@@ -299,6 +299,7 @@ void UNiagaraDataInterfaceGrid3DCollection::GetFunctions(TArray<FNiagaraFunction
 		Sig.bSoftDeprecatedFunction = true;
 		Sig.bSupportsCPU = false;
 		Sig.bSupportsGPU = true;
+		Sig.bReadFunction = true;
 #if WITH_EDITORONLY_DATA
 		Sig.Description = NSLOCTEXT("Niagara", "NiagaraDataInterfaceGridColl3D_GetValueFunction", "Get the value at a specific index. Note that this is an older way of working with Grids. Consider using the SetFloat or other typed, named functions or parameter map variables with StackContext namespace instead.");
 #endif
@@ -322,6 +323,7 @@ void UNiagaraDataInterfaceGrid3DCollection::GetFunctions(TArray<FNiagaraFunction
 		Sig.ModuleUsageBitmask = ENiagaraScriptUsageMask::Particle;
 		Sig.bSupportsCPU = false;
 		Sig.bSupportsGPU = true;
+		Sig.bReadFunction = true;
 #if WITH_EDITORONLY_DATA
 		Sig.Description = NSLOCTEXT("Niagara", "NiagaraDataInterfaceGridColl2D_PreviousValueAtIndexFunction", "Get the value at a specific index.");
 #endif
@@ -344,6 +346,7 @@ void UNiagaraDataInterfaceGrid3DCollection::GetFunctions(TArray<FNiagaraFunction
 		Sig.ModuleUsageBitmask = ENiagaraScriptUsageMask::Particle;
 		Sig.bSupportsCPU = false;
 		Sig.bSupportsGPU = true;
+		Sig.bReadFunction = true;
 #if WITH_EDITORONLY_DATA
 		Sig.Description = NSLOCTEXT("Niagara", "NiagaraDataInterfaceGridColl2D_PreviousValueAtIndexFunction", "Get the value at a specific index.");
 #endif
@@ -438,6 +441,7 @@ void UNiagaraDataInterfaceGrid3DCollection::GetFunctions(TArray<FNiagaraFunction
 		Sig.ModuleUsageBitmask = ENiagaraScriptUsageMask::Particle;
 		Sig.bSupportsCPU = false;
 		Sig.bSupportsGPU = true;
+		Sig.bReadFunction = true;
 #if WITH_EDITORONLY_DATA
 		Sig.Description = NSLOCTEXT("Niagara", "NiagaraDataInterfaceGridColl2D_CopyPreviousToCurrentForCell", "Take the previous contents of the cell and copy to the output location for the cell.");
 #endif
@@ -462,6 +466,7 @@ void UNiagaraDataInterfaceGrid3DCollection::GetFunctions(TArray<FNiagaraFunction
 		Sig.ModuleUsageBitmask = ENiagaraScriptUsageMask::Particle;
 		Sig.bSupportsCPU = false;
 		Sig.bSupportsGPU = true;
+		Sig.bReadFunction = true;
 #if WITH_EDITORONLY_DATA
 		Sig.Description = NSLOCTEXT("Niagara", "NiagaraDataInterfaceGridColl2D_CopyMaskedPreviousToCurrentForCell", "Take the previous contents of the cell and copy to the output location for the cell.");
 #endif
@@ -506,6 +511,7 @@ void UNiagaraDataInterfaceGrid3DCollection::GetFunctions(TArray<FNiagaraFunction
 		Sig.ModuleUsageBitmask = ENiagaraScriptUsageMask::Particle;
 		Sig.bSupportsCPU = false;
 		Sig.bSupportsGPU = true;
+		Sig.bReadFunction = true;
 #if WITH_EDITORONLY_DATA
 		Sig.Description = NSLOCTEXT("Niagara", "NiagaraDataInterfaceGridColl3D_GetVector4", "Gets a Vector4 value on the Grid by Attribute name. Note that this is the value from the previous execution stage.");
 #endif
@@ -526,6 +532,7 @@ void UNiagaraDataInterfaceGrid3DCollection::GetFunctions(TArray<FNiagaraFunction
 		Sig.ModuleUsageBitmask = ENiagaraScriptUsageMask::Particle;
 		Sig.bSupportsCPU = false;
 		Sig.bSupportsGPU = true;
+		Sig.bReadFunction = true;
 #if WITH_EDITORONLY_DATA
 		Sig.Description = NSLOCTEXT("Niagara", "NiagaraDataInterfaceGridColl3D_SampleVector4", "Sample a Vector4 value on the Grid by Attribute name. Note that this is the value from the previous execution stage.");
 #endif
@@ -546,6 +553,7 @@ void UNiagaraDataInterfaceGrid3DCollection::GetFunctions(TArray<FNiagaraFunction
 		Sig.ModuleUsageBitmask = ENiagaraScriptUsageMask::Particle;
 		Sig.bSupportsCPU = false;
 		Sig.bSupportsGPU = true;
+		Sig.bReadFunction = true;
 #if WITH_EDITORONLY_DATA
 		Sig.Description = NSLOCTEXT("Niagara", "NiagaraDataInterfaceGridColl3D_CubicSampleVector4", "Cubic Sample a Vector4 value on the Grid by Attribute name. Note that this is the value from the previous execution stage.");
 #endif
@@ -592,6 +600,7 @@ void UNiagaraDataInterfaceGrid3DCollection::GetFunctions(TArray<FNiagaraFunction
 		Sig.ModuleUsageBitmask = ENiagaraScriptUsageMask::Particle;
 		Sig.bSupportsCPU = false;
 		Sig.bSupportsGPU = true;
+		Sig.bReadFunction = true;
 #if WITH_EDITORONLY_DATA
 		Sig.Description = NSLOCTEXT("Niagara", "NiagaraDataInterfaceGridColl3D_GetVector3", "Gets a Vector3 value on the Grid by Attribute name. Note that this is the value from the previous execution stage.");
 #endif
@@ -612,6 +621,7 @@ void UNiagaraDataInterfaceGrid3DCollection::GetFunctions(TArray<FNiagaraFunction
 		Sig.ModuleUsageBitmask = ENiagaraScriptUsageMask::Particle;
 		Sig.bSupportsCPU = false;
 		Sig.bSupportsGPU = true;
+		Sig.bReadFunction = true;
 #if WITH_EDITORONLY_DATA
 		Sig.Description = NSLOCTEXT("Niagara", "NiagaraDataInterfaceGridColl3D_SampleVector3", "Sample a Vector3 value on the Grid by Attribute name. Note that this is the value from the previous execution stage.");
 #endif
@@ -632,6 +642,7 @@ void UNiagaraDataInterfaceGrid3DCollection::GetFunctions(TArray<FNiagaraFunction
 		Sig.ModuleUsageBitmask = ENiagaraScriptUsageMask::Particle;
 		Sig.bSupportsCPU = false;
 		Sig.bSupportsGPU = true;
+		Sig.bReadFunction = true;
 #if WITH_EDITORONLY_DATA
 		Sig.Description = NSLOCTEXT("Niagara", "NiagaraDataInterfaceGridColl3D_CubicSampleVector3", "Cubic Sample a Vector3 value on the Grid by Attribute name. Note that this is the value from the previous execution stage.");
 #endif
@@ -678,6 +689,7 @@ void UNiagaraDataInterfaceGrid3DCollection::GetFunctions(TArray<FNiagaraFunction
 		Sig.ModuleUsageBitmask = ENiagaraScriptUsageMask::Particle;
 		Sig.bSupportsCPU = false;
 		Sig.bSupportsGPU = true;
+		Sig.bReadFunction = true;
 #if WITH_EDITORONLY_DATA
 		Sig.Description = NSLOCTEXT("Niagara", "NiagaraDataInterfaceGridColl3D_GetVector2", "Gets a Vector2 value on the Grid by Attribute name. Note that this is the value from the previous execution stage.");
 #endif
@@ -698,6 +710,7 @@ void UNiagaraDataInterfaceGrid3DCollection::GetFunctions(TArray<FNiagaraFunction
 		Sig.ModuleUsageBitmask = ENiagaraScriptUsageMask::Particle;
 		Sig.bSupportsCPU = false;
 		Sig.bSupportsGPU = true;
+		Sig.bReadFunction = true;
 #if WITH_EDITORONLY_DATA
 		Sig.Description = NSLOCTEXT("Niagara", "NiagaraDataInterfaceGridColl3D_SampleVector2", "Sample a Vector2 value on the Grid by Attribute name. Note that this is the value from the previous execution stage.");
 #endif
@@ -718,6 +731,7 @@ void UNiagaraDataInterfaceGrid3DCollection::GetFunctions(TArray<FNiagaraFunction
 		Sig.ModuleUsageBitmask = ENiagaraScriptUsageMask::Particle;
 		Sig.bSupportsCPU = false;
 		Sig.bSupportsGPU = true;
+		Sig.bReadFunction = true;
 #if WITH_EDITORONLY_DATA
 		Sig.Description = NSLOCTEXT("Niagara", "NiagaraDataInterfaceGridColl3D_CubicSampleVector2", "Cubic Sample a Vector2 value on the Grid by Attribute name. Note that this is the value from the previous execution stage.");
 #endif
@@ -763,6 +777,7 @@ void UNiagaraDataInterfaceGrid3DCollection::GetFunctions(TArray<FNiagaraFunction
 		Sig.ModuleUsageBitmask = ENiagaraScriptUsageMask::Particle;
 		Sig.bSupportsCPU = false;
 		Sig.bSupportsGPU = true;
+		Sig.bReadFunction = true;
 #if WITH_EDITORONLY_DATA
 		Sig.Description = NSLOCTEXT("Niagara", "NiagaraDataInterfaceGridColl3D_GetFloat", "Gets a float value on the Grid by Attribute name. Note that this is the value from the previous execution stage.");
 #endif
@@ -783,6 +798,7 @@ void UNiagaraDataInterfaceGrid3DCollection::GetFunctions(TArray<FNiagaraFunction
 		Sig.ModuleUsageBitmask = ENiagaraScriptUsageMask::Particle;
 		Sig.bSupportsCPU = false;
 		Sig.bSupportsGPU = true;
+		Sig.bReadFunction = true;
 #if WITH_EDITORONLY_DATA
 		Sig.Description = NSLOCTEXT("Niagara", "NiagaraDataInterfaceGridColl3D_SampleFloat", "Sample a float value on the Grid by Attribute name.");
 #endif
@@ -803,6 +819,7 @@ void UNiagaraDataInterfaceGrid3DCollection::GetFunctions(TArray<FNiagaraFunction
 		Sig.ModuleUsageBitmask = ENiagaraScriptUsageMask::Particle;
 		Sig.bSupportsCPU = false;
 		Sig.bSupportsGPU = true;
+		Sig.bReadFunction = true;
 #if WITH_EDITORONLY_DATA
 		Sig.Description = NSLOCTEXT("Niagara", "NiagaraDataInterfaceGridColl3D_CubicSampleFloat", "Cubic Sample a float value on the Grid by Attribute name.");
 #endif
@@ -827,6 +844,7 @@ void UNiagaraDataInterfaceGrid3DCollection::GetFunctions(TArray<FNiagaraFunction
 		Sig.bSupportsCPU = false;
 		Sig.bSupportsGPU = true;
 		Sig.bSoftDeprecatedFunction = true;
+		Sig.bReadFunction = true;
 		OutFunctions.Add(Sig);
 	}
 
@@ -847,6 +865,7 @@ void UNiagaraDataInterfaceGrid3DCollection::GetFunctions(TArray<FNiagaraFunction
 		Sig.bSupportsCPU = false;
 		Sig.bSupportsGPU = true;
 		Sig.bSoftDeprecatedFunction = true;
+		Sig.bReadFunction = true;
 		OutFunctions.Add(Sig);
 	}
 
@@ -866,6 +885,7 @@ void UNiagaraDataInterfaceGrid3DCollection::GetFunctions(TArray<FNiagaraFunction
 		Sig.ModuleUsageBitmask = ENiagaraScriptUsageMask::Particle;
 		Sig.bSupportsCPU = false;
 		Sig.bSupportsGPU = true;
+		Sig.bReadFunction = true;
 		OutFunctions.Add(Sig);
 	}
 
@@ -885,6 +905,7 @@ void UNiagaraDataInterfaceGrid3DCollection::GetFunctions(TArray<FNiagaraFunction
 		Sig.ModuleUsageBitmask = ENiagaraScriptUsageMask::Particle;
 		Sig.bSupportsCPU = false;
 		Sig.bSupportsGPU = true;
+		Sig.bReadFunction = true;
 		OutFunctions.Add(Sig);
 	}
 
@@ -901,6 +922,7 @@ void UNiagaraDataInterfaceGrid3DCollection::GetFunctions(TArray<FNiagaraFunction
 		Sig.ModuleUsageBitmask = ENiagaraScriptUsageMask::Particle;
 		Sig.bSupportsCPU = false;
 		Sig.bSupportsGPU = true;
+		Sig.bReadFunction = true;
 		OutFunctions.Add(Sig);
 	}
 
@@ -917,6 +939,7 @@ void UNiagaraDataInterfaceGrid3DCollection::GetFunctions(TArray<FNiagaraFunction
 		Sig.ModuleUsageBitmask = ENiagaraScriptUsageMask::Particle;
 		Sig.bSupportsCPU = false;
 		Sig.bSupportsGPU = true;
+		Sig.bReadFunction = true;
 		OutFunctions.Add(Sig);
 	}
 
@@ -2265,6 +2288,7 @@ bool UNiagaraDataInterfaceGrid3DCollection::InitPerInstanceData(void* PerInstanc
 		check(!RT_Proxy->SystemInstancesToProxyData_RT.Contains(InstanceID));
 		FGrid3DCollectionRWInstanceData_RenderThread* TargetData = &RT_Proxy->SystemInstancesToProxyData_RT.Add(InstanceID);
 
+		TargetData->SourceDIName = RT_Proxy->SourceDIName;
 		TargetData->NumCells = RT_InstanceData.NumCells;
 		TargetData->NumTiles = RT_InstanceData.NumTiles;
 		TargetData->TotalNumAttributes = RT_InstanceData.TotalNumAttributes;
@@ -2919,12 +2943,12 @@ bool FGrid3DCollectionRWInstanceData_GameThread::UpdateTargetTexture(ENiagaraGpu
 	return false;
 }
 
-void FGrid3DCollectionRWInstanceData_RenderThread::BeginSimulate(FRDGBuilder& GraphBuilder)
+void FGrid3DCollectionRWInstanceData_RenderThread::BeginSimulate(FRDGBuilder& GraphBuilder, bool RequiresBuffering)
 {
 	for (FGrid3DBuffer& Buffer : Buffers)
 	{
 		check(Buffer.IsValid());
-		if (&Buffer != CurrentData)
+		if (&Buffer != CurrentData || !RequiresBuffering)
 		{
 			DestinationData = &Buffer;
 			break;
@@ -2937,7 +2961,18 @@ void FGrid3DCollectionRWInstanceData_RenderThread::BeginSimulate(FRDGBuilder& Gr
 
 		const FIntVector TextureSize(NumCells.X * NumTiles.X, NumCells.Y * NumTiles.Y, NumCells.Z * NumTiles.Z);
 		const FRDGTextureDesc TextureDesc = FRDGTextureDesc::Create3D(TextureSize, PixelFormat, FClearValueBinding::Black, ETextureCreateFlags::ShaderResource | ETextureCreateFlags::UAV);
-		DestinationData->Initialize(GraphBuilder, TEXT("FGrid3DBuffer"), TextureDesc);
+		
+		TStringBuilder<128> StringBuilder;
+		SourceDIName.AppendString(StringBuilder);		
+		StringBuilder.Append("_");		
+		for (FName CurrName : Vars)
+		{
+			CurrName.AppendString(StringBuilder);			
+			StringBuilder.Append("_");			
+		}
+		StringBuilder.Append("Grid3DCollection");		
+		
+		DestinationData->Initialize(GraphBuilder, *StringBuilder, TextureDesc);
 	}
 }
 
@@ -2969,7 +3004,7 @@ void FNiagaraDataInterfaceProxyGrid3DCollectionProxy::PreStage(const FNDIGpuComp
 	{
 		FGrid3DCollectionRWInstanceData_RenderThread* ProxyData = SystemInstancesToProxyData_RT.Find(Context.GetSystemInstanceID());
 		FRDGBuilder& GraphBuilder = Context.GetGraphBuilder();
-		ProxyData->BeginSimulate(GraphBuilder);
+		ProxyData->BeginSimulate(GraphBuilder, Context.IsInputStage());
 
 		// If we don't have an iteration stage, then we should manually clear the buffer to make sure there is no residual data.  If we are doing something like rasterizing particles into a grid, we want it to be clear before
 		// we start.  If a user wants to access data from the previous stage, then they can read from the current data.
diff --git a/Engine/Plugins/FX/Niagara/Source/Niagara/Private/NiagaraScript.cpp b/Engine/Plugins/FX/Niagara/Source/Niagara/Private/NiagaraScript.cpp
--- a/Engine/Plugins/FX/Niagara/Source/Niagara/Private/NiagaraScript.cpp
+++ b/Engine/Plugins/FX/Niagara/Source/Niagara/Private/NiagaraScript.cpp
@@ -3303,6 +3303,15 @@ void UNiagaraScript::SyncAliases(const FNiagaraAliasContext& ResolveAliasesConte
 				OutputDestination = FNiagaraUtilities::ResolveAliases(Var, ResolveAliasesContext).GetName();
 			}
 		}
+
+		for (FName& InputDI : SimStageMetaData.InputDataInterfaces)
+		{
+			if (!InputDI.IsNone())
+			{
+				FNiagaraVariable Var(FNiagaraTypeDefinition(UNiagaraDataInterface::StaticClass()), InputDI);
+				InputDI = FNiagaraUtilities::ResolveAliases(Var, ResolveAliasesContext).GetName();
+			}
+		}
 	}
 
 	// Also handle any data set mappings...
diff --git a/Engine/Plugins/FX/Niagara/Source/Niagara/Public/NiagaraCommon.h b/Engine/Plugins/FX/Niagara/Source/Niagara/Public/NiagaraCommon.h
--- a/Engine/Plugins/FX/Niagara/Source/Niagara/Public/NiagaraCommon.h
+++ b/Engine/Plugins/FX/Niagara/Source/Niagara/Public/NiagaraCommon.h
@@ -387,10 +387,14 @@ struct NIAGARA_API FNiagaraFunctionSignature
 	UPROPERTY()
 	uint32 bSupportsGPU : 1;
 
-	/** Writes to the variable this is bound to */
+	/** Writes data owned by the data interface.  Any stage using it will be marked as an Output stage. */
 	UPROPERTY()
 	uint32 bWriteFunction : 1;
 
+	/** Reads data owned by the data interface.  Any stage using it will be marked as an Input stage. */
+	UPROPERTY()
+	uint32 bReadFunction : 1;
+
 	/** Whether or not this function should show up in normal usage. */
 	UPROPERTY()
 	uint32 bSoftDeprecatedFunction : 1;
@@ -435,6 +439,7 @@ struct NIAGARA_API FNiagaraFunctionSignature
 		, bSupportsCPU(true)
 		, bSupportsGPU(true)
 		, bWriteFunction(false)
+		, bReadFunction(false)
 		, bSoftDeprecatedFunction(false)
 		, bIsCompileTagGenerator(false)
 		, bHidden(false)
@@ -452,6 +457,7 @@ struct NIAGARA_API FNiagaraFunctionSignature
 		, bSupportsCPU(true)
 		, bSupportsGPU(true)
 		, bWriteFunction(false)
+		, bReadFunction(false)
 		, bSoftDeprecatedFunction(false)
 		, bIsCompileTagGenerator(false)
 		, bHidden(false)
diff --git a/Engine/Plugins/FX/Niagara/Source/NiagaraEditor/Private/NiagaraHlslTranslator.cpp b/Engine/Plugins/FX/Niagara/Source/NiagaraEditor/Private/NiagaraHlslTranslator.cpp
--- a/Engine/Plugins/FX/Niagara/Source/NiagaraEditor/Private/NiagaraHlslTranslator.cpp
+++ b/Engine/Plugins/FX/Niagara/Source/NiagaraEditor/Private/NiagaraHlslTranslator.cpp
@@ -2019,6 +2019,11 @@ const FNiagaraTranslateResults &FHlslNiagaraTranslator::Translate(const FNiagara
 			{
 				Preamble.Appendf(TEXT("//\tOutputs to: \"%s\"\n"), *Dest.ToString());
 			}
+
+			for (const FName& Dest : SimStageMetaData.InputDataInterfaces)
+			{
+				Preamble.Appendf(TEXT("//\tReads from: \"%s\"\n"), *Dest.ToString());
+			}
 		}
 
 		// Display the computed compile tags in the source hlsl to make checking easier.
@@ -4058,6 +4063,7 @@ void FHlslNiagaraTranslator::EnterFunction(const FString& Name, FNiagaraFunction
 	FunctionContextStack.Emplace(Name, Signature, Inputs, InGuid);
 	TArray<FName> Entries;
 	ActiveStageWriteTargets.Push(Entries);
+	ActiveStageReadTargets.Push(Entries);
 	//May need some more heavy and scoped symbol tracking?
 
 	//Add new scope for pin reuse.
@@ -4081,6 +4087,16 @@ void FHlslNiagaraTranslator::ExitFunction()
 			ActiveStageWriteTargets.Top().AddUnique(Entry);
 		}
 	}
+
+	// Accumulate the read targets.
+	Entries = ActiveStageReadTargets.Pop();
+	if (ActiveStageReadTargets.Num())
+	{
+		for (const FName& Entry : Entries)
+		{
+			ActiveStageReadTargets.Top().AddUnique(Entry);
+		}
+	}
 }
 
 FString FHlslNiagaraTranslator::GeneratedConstantString(float Constant)
@@ -7721,6 +7737,15 @@ void FHlslNiagaraTranslator::HandleDataInterfaceCall(FNiagaraScriptDataInterface
 			ActiveStageWriteTargets.Top().AddUnique(Info.Name);
 		}
 	}
+	if (InMatchingSignature.bReadFunction && CompilationOutput.ScriptData.SimulationStageMetaData.Num() > 1 && TranslationStages[ActiveStageIdx].SimulationStageIndex != -1)
+	{
+		const int32 SourceSimStage = TranslationStages[ActiveStageIdx].SimulationStageIndex;
+		CompilationOutput.ScriptData.SimulationStageMetaData[SourceSimStage].InputDataInterfaces.AddUnique(Info.Name);
+		if (ActiveStageReadTargets.Num() > 0)
+		{
+			ActiveStageReadTargets.Top().AddUnique(Info.Name);
+		}
+	}
 }
 
 bool IsVariableWriteBeforeRead(const TArray<FNiagaraParameterMapHistory::FReadHistory>& ReadHistory)
@@ -7976,6 +8001,7 @@ void FHlslNiagaraTranslator::RegisterFunctionCall(ENiagaraScriptUsage ScriptUsag
 				}
 
 				FunctionStageWriteTargets.Add(OutSignature, ActiveStageWriteTargets.Top());
+				FunctionStageReadTargets.Add(OutSignature, ActiveStageReadTargets.Top());
 			}
 
 			ExitFunction();
@@ -7984,7 +8010,7 @@ void FHlslNiagaraTranslator::RegisterFunctionCall(ENiagaraScriptUsage ScriptUsag
 		{
 			FuncBody->StageIndices.AddUnique(ActiveStageIdx);
 
-			// Just because we had a cached call, doesn't mean that we should ignore adding an writetargets
+			// Just because we had a cached call, doesn't mean that we should ignore adding read or writetargets
 			TArray<FName>* Entries = FunctionStageWriteTargets.Find(OutSignature);
 			if (Entries)
 			{
@@ -7995,6 +8021,17 @@ void FHlslNiagaraTranslator::RegisterFunctionCall(ENiagaraScriptUsage ScriptUsag
 					ActiveStageWriteTargets.Top().AddUnique(Entry);
 				}
 			}
+
+			Entries = FunctionStageReadTargets.Find(OutSignature);
+			if (Entries)
+			{
+				for (const FName& Entry : *Entries)
+				{
+					const int32 SourceSimStage = TranslationStages[ActiveStageIdx].SimulationStageIndex;
+					CompilationOutput.ScriptData.SimulationStageMetaData[SourceSimStage].InputDataInterfaces.AddUnique(Entry);
+					ActiveStageReadTargets.Top().AddUnique(Entry);
+				}
+			}
 		}
 	}
 	else
@@ -8026,6 +8063,7 @@ void FHlslNiagaraTranslator::RegisterFunctionCall(ENiagaraScriptUsage ScriptUsag
 				}
 
 				FunctionStageWriteTargets.Add(OutSignature, ActiveStageWriteTargets.Top());
+				FunctionStageReadTargets.Add(OutSignature, ActiveStageReadTargets.Top());
 
 				for (const FNiagaraCustomHlslInclude& FileInclude : InCustomHlslIncludeFilePaths)
 				{
diff --git a/Engine/Plugins/FX/Niagara/Source/NiagaraEditor/Private/NiagaraHlslTranslator.h b/Engine/Plugins/FX/Niagara/Source/NiagaraEditor/Private/NiagaraHlslTranslator.h
--- a/Engine/Plugins/FX/Niagara/Source/NiagaraEditor/Private/NiagaraHlslTranslator.h
+++ b/Engine/Plugins/FX/Niagara/Source/NiagaraEditor/Private/NiagaraHlslTranslator.h
@@ -491,6 +491,9 @@ protected:
 	TMap<FNiagaraFunctionSignature, TArray<FName>> FunctionStageWriteTargets;
 	TArray<TArray<FName>> ActiveStageWriteTargets;
 
+	TMap<FNiagaraFunctionSignature, TArray<FName>> FunctionStageReadTargets;
+	TArray<TArray<FName>> ActiveStageReadTargets;
+
 	void RegisterFunctionCall(ENiagaraScriptUsage ScriptUsage, const FString& InName, const FString& InFullName, const FGuid& CallNodeId, const FString& InFunctionNameSuffix, UNiagaraScriptSource* Source, FNiagaraFunctionSignature& InSignature, bool bIsCustomHlsl, const FString& InCustomHlsl, const TArray<FNiagaraCustomHlslInclude>& InCustomHlslIncludeFilePaths, TArray<int32>& Inputs, TArrayView<UEdGraphPin* const> CallInputs, TArrayView<UEdGraphPin* const> CallOutputs,
 		FNiagaraFunctionSignature& OutSignature);
 	void GenerateFunctionCall(ENiagaraScriptUsage ScriptUsage, FNiagaraFunctionSignature& FunctionSignature, TArrayView<const int32> Inputs, TArray<int32>& Outputs);
diff --git a/Engine/Plugins/FX/Niagara/Source/NiagaraShader/Public/NiagaraScriptBase.h b/Engine/Plugins/FX/Niagara/Source/NiagaraShader/Public/NiagaraScriptBase.h
--- a/Engine/Plugins/FX/Niagara/Source/NiagaraShader/Public/NiagaraScriptBase.h
+++ b/Engine/Plugins/FX/Niagara/Source/NiagaraShader/Public/NiagaraScriptBase.h
@@ -83,6 +83,11 @@ public:
 	UPROPERTY()
 	TArray<FName> OutputDestinations;
 
+	/** DataInterfaces that we read from in this stage.*/
+	UPROPERTY()
+	TArray<FName> InputDataInterfaces;
+
+
 	/** The number of iterations for the stage. */
 	UPROPERTY()
 	int32 NumIterations = 1;
