Temporary fix to prevent cooking crashes due to emitter parent pointers becoming stale

diff --git a/Engine/Plugins/FX/Niagara/Source/NiagaraEditor/Private/ViewModels/Stack/NiagaraStackEntry.cpp b/Engine/Plugins/FX/Niagara/Source/NiagaraEditor/Private/ViewModels/Stack/NiagaraStackEntry.cpp
--- a/Engine/Plugins/FX/Niagara/Source/NiagaraEditor/Private/ViewModels/Stack/NiagaraStackEntry.cpp
+++ b/Engine/Plugins/FX/Niagara/Source/NiagaraEditor/Private/ViewModels/Stack/NiagaraStackEntry.cpp
@@ -554,6 +554,12 @@ void UNiagaraStackEntry::SetIsSearchResult(bool bInIsSearchResult)
 
 bool UNiagaraStackEntry::HasBaseEmitter() const
 {
+	if (GetSystemViewModel()->GetIsForDataProcessingOnly())
+	{
+		// If the model is just for data processing we don't want to go through the whole merge procedure for all the stack entries and just treat all entries as non-inherited.
+		return false;
+	}
+	
 	if (bHasBaseEmitterCache.IsSet() == false)
 	{
 		const UNiagaraEmitter* BaseEmitter = GetEmitterViewModel().IsValid() ? GetEmitterViewModel()->GetParentEmitter().Emitter : nullptr;
diff --git a/Engine/Plugins/FX/Niagara/Source/NiagaraEditor/Private/ViewModels/Stack/NiagaraStackModuleItem.cpp b/Engine/Plugins/FX/Niagara/Source/NiagaraEditor/Private/ViewModels/Stack/NiagaraStackModuleItem.cpp
--- a/Engine/Plugins/FX/Niagara/Source/NiagaraEditor/Private/ViewModels/Stack/NiagaraStackModuleItem.cpp
+++ b/Engine/Plugins/FX/Niagara/Source/NiagaraEditor/Private/ViewModels/Stack/NiagaraStackModuleItem.cpp
@@ -863,7 +863,7 @@ UNiagaraStackEntry::FStackIssueFixDelegate UNiagaraStackModuleItem::GetUpgradeVe
 
 void UNiagaraStackModuleItem::RefreshIssues(TArray<FStackIssue>& NewIssues)
 {
-	if (!GetIsEnabled())
+	if (!GetIsEnabled() || GetSystemViewModel()->GetIsForDataProcessingOnly())
 	{
 		NewIssues.Empty();
 		return;
diff --git a/Engine/Plugins/FX/Niagara/Source/NiagaraEditor/Private/ViewModels/Stack/NiagaraStackSimulationStageGroup.cpp b/Engine/Plugins/FX/Niagara/Source/NiagaraEditor/Private/ViewModels/Stack/NiagaraStackSimulationStageGroup.cpp
--- a/Engine/Plugins/FX/Niagara/Source/NiagaraEditor/Private/ViewModels/Stack/NiagaraStackSimulationStageGroup.cpp
+++ b/Engine/Plugins/FX/Niagara/Source/NiagaraEditor/Private/ViewModels/Stack/NiagaraStackSimulationStageGroup.cpp
@@ -157,6 +157,11 @@ void UNiagaraStackSimulationStagePropertiesItem::SimulationStagePropertiesChange
 
 bool UNiagaraStackSimulationStagePropertiesItem::HasBaseSimulationStage() const
 {
+	if (GetSystemViewModel()->GetIsForDataProcessingOnly())
+	{
+		// If the model is just for data processing we don't want to go through the whole merge procedure and treat the stage entry as non-inherited.
+		return false;
+	}
 	if (bHasBaseSimulationStageCache.IsSet() == false)
 	{
 		FVersionedNiagaraEmitter VersionedEmitter = GetEmitterViewModel()->GetEmitter();
