- Convert Niagara debug drawing to RDG

diff --git a/Engine/Plugins/FX/Niagara/Shaders/Private/NiagaraDataInterfaceDebugDrawTemplate.ush b/Engine/Plugins/FX/Niagara/Shaders/Private/NiagaraDataInterfaceDebugDrawTemplate.ush
new file mode 100644
--- /dev/null
+++ b/Engine/Plugins/FX/Niagara/Shaders/Private/NiagaraDataInterfaceDebugDrawTemplate.ush
@@ -0,0 +1,100 @@
+// Copyright Epic Games, Inc. All Rights Reserved.
+
+void DrawLine_{ParameterName}_UEImpureCall(bool bExecute, float3 LineStart, float3 LineEnd, float4 Color)
+{
+	NDIDebugDraw_Line(bExecute, LineStart, LineEnd, Color);
+}
+
+void DrawRectangle_{ParameterName}_UEImpureCall(bool bExecute, float3 Location, float3 XAxis, float3 YAxis, float2 Extents, int NumXSegments, int NumYSegments, float4 Color, bool bUnbounded)
+{
+	NDIDebugDraw_DrawRectangle(bExecute, Location, XAxis, YAxis, Extents, NumXSegments, NumYSegments, Color, bUnbounded);
+}
+
+
+void DrawCircle_{ParameterName}_UEImpureCall(bool bExecute, float3 Location, float3 XAxis, float3 YAxis, float Scale, int Segments, float4 Color)
+{
+	NDIDebugDraw_Circle(bExecute, Location, XAxis, YAxis, Scale, Segments, Color);
+}
+
+void DrawBox_{ParameterName}_UEImpureCall(bool bExecute, float3 Location, float4 Rotation, float3 Extents, float4 Color)
+{
+	NDIDebugDraw_DrawBox(bExecute, Location, Rotation, Extents, Color);
+}
+
+void DrawSphere_{ParameterName}_UEImpureCall(bool bExecute, float3 LineStart, float Radius, int Segments, float4 Color)
+{
+	NDIDebugDraw_Sphere(bExecute, LineStart, Radius, Segments, Color);
+}
+
+void DrawCylinder_{ParameterName}_UEImpureCall(bool bExecute, float3 Location, float3 Axis, float Height, float Radius, int NumHeightSegments, int NumRadiusSegments, float4 Color)
+{
+	NDIDebugDraw_Cylinder(bExecute, Location, Axis, Height, Radius, NumHeightSegments, NumRadiusSegments, Color);
+}
+
+void DrawCone_{ParameterName}_UEImpureCall(bool bExecute, float3 Location, float3 Axis, float Height, float RadiusTop, float RadiusBottom, int NumHeightSegments, int NumRadiusSegments, float4 Color)
+{
+}
+
+void DrawTorus_{ParameterName}_UEImpureCall(bool bExecute, float3 Location, float3 Axis, float MajorRadius, float MinorRadius, float MajorRadiusSegments, float MinorRadiusSegments, float4 Color)
+{
+}
+
+void DrawCoordinateSystem_{ParameterName}_UEImpureCall(bool bExecute, float3 Location, float4 Rotation, float Scale)
+{
+	NDIDebugDraw_CoordinateSystem(bExecute, Location, Rotation, Scale);
+}
+
+void DrawGrid2D_{ParameterName}_UEImpureCall(bool bExecute, float3 Center, float4 Rotation, float2 Extents, int NumCellsX, int NumCellsY, float4 Color)
+{
+	NDIDebugDraw_Grid2D(bExecute, Center, Rotation, Extents, int2(NumCellsX, NumCellsY), Color);
+}
+
+void DrawGrid3D_{ParameterName}_UEImpureCall(bool bExecute, float3 Center, float4 Rotation, float3 Extents, int NumCellsX, int NumCellsY, int NumCellsZ, float4 Color)
+{
+	NDIDebugDraw_Grid3D(bExecute, Center, Rotation, Extents, int3(NumCellsX, NumCellsY, NumCellsZ), Color);
+}
+
+void DrawLinePersistent_{ParameterName}_UEImpureCall(float3 StartLocation, int StartLocationCoordinateSpace, float3 EndLocation, int EndLocationCoordinateSpace, float4 Color)
+{
+}
+
+void DrawRectanglePersistent_{ParameterName}_UEImpureCall(float3 Center, int CenterCoordinateSpace, float3 Offset, int OffsetCoordinateSpace, float3 Extents, bool HalfExtents, float3 XAxis, float3 YAxis, int AxisCoordinateSpace, int NumXSegments, int NumYSegments, bool bUnboundedPlane, float4 Color)
+{
+}
+
+void DrawCirclePersistent_{ParameterName}_UEImpureCall(float3 Center, int CenterCoordinateSpace, float3 Offset, int OffsetCoordinateSpace, float Radius, float3 XAxis, float3 YAxis, int AxisCoordinateSpace, int NumSegments, float4 Color)
+{
+}
+
+void DrawBoxPersistent_{ParameterName}_UEImpureCall(float3 Center, int CenterCoordinateSpace, float3 Offset, int OffsetCoordinateSpace, float3 Extents, bool HalfExtents, float3 RotationAxis, float RotationAngle, int RotationCoordinateSpace, float4 Color)
+{
+}
+
+void DrawSpherePersistent_{ParameterName}_UEImpureCall(float3 Center, int CenterCoordinateSpace, float3 OffsetFromCenter, int OffsetCoordinateSpace, float Radius,  int RadiusCoordinateSpace, bool HemisphereX, bool HemisphereY, bool HemisphereZ, float3 OrientationAxis, int RotationCoordinateSpace, float4 AdditionalRotation, float3 NonUniformScale, int NumSegments, float4 Color)
+{
+}
+
+void DrawCylinderPersistent_{ParameterName}_UEImpureCall(float3 Center, int CenterCoordinateSpace, float3 OffsetFromCenter, int OffsetCoordinateSpace, float Radius, float Height, bool IsHalfHeight, bool HemisphereX, bool HemisphereY, float3 OrientationAxis, int OrientationAxisCoordinateSpace, float3 NonUniformScale, int NumRadiusSegments, int NumHeightSegments, float4 Color)
+{
+}
+
+void DrawConePersistent_{ParameterName}_UEImpureCall(float3 Center, int CenterCoordinateSpace, float3 OffsetFromCenter, int OffsetCoordinateSpace, float Angle, float Length, int OrientationAxis, int OrientationAxisCoordinateSpace, float3 NonUniformScale, int NumRadiusSegments, int NumHeightSegments, float4 Color)
+{
+}
+
+void DrawTorusPersistent_{ParameterName}_UEImpureCall(float3 Center, int CenterCoordinateSpace, float3 OffsetFromCenter, int OffsetCoordinateSpace, float MajorRadius, float MinorRadius, float3 OrientationAxis, int OrientationAxisCoordinateSpace, float3 NonUniformScale, int MajorRadiusSegments, int MinorRadiusSegments, float4 Color)
+{
+}
+
+void DrawCoordinateSystemPersistent_{ParameterName}_UEImpureCall(float3 Center, int CenterCoordinateSpace, float3 OffsetFromCenter, int OffsetCoordinateSpace, float4 Rotation, int RotationCoordinateSpace, float Scale)
+{
+}
+
+void DrawGrid2DPersistent_{ParameterName}_UEImpureCall(float3 Center, int CenterCoordinateSpace, float3 OffsetFromCenter, int OffsetCoordinateSpace, float4 Rotation, int RotationCoordinateSpace, float2 Extents, int NumCellsX, int NumCellsY, float4 Color)
+{
+}
+
+void DrawGrid3DPersistent_{ParameterName}_UEImpureCall(float3 Center, int CenterCoordinateSpace, float3 OffsetFromCenter, int OffsetCoordinateSpace, float4 Rotation, int RotationCoordinateSpace, float3 Extents, int NumCellsX, int NumCellsY, int NumCellsZ, float4 Color)
+{
+}
+
diff --git a/Engine/Plugins/FX/Niagara/Source/Niagara/Private/DataInterface/NiagaraDataInterfaceDebugDraw.cpp b/Engine/Plugins/FX/Niagara/Source/Niagara/Private/DataInterface/NiagaraDataInterfaceDebugDraw.cpp
--- a/Engine/Plugins/FX/Niagara/Source/Niagara/Private/DataInterface/NiagaraDataInterfaceDebugDraw.cpp
+++ b/Engine/Plugins/FX/Niagara/Source/Niagara/Private/DataInterface/NiagaraDataInterfaceDebugDraw.cpp
@@ -5,6 +5,7 @@
 #include "NiagaraGpuComputeDispatchInterface.h"
 #include "NiagaraGpuComputeDebug.h"
 #include "NiagaraWorldManager.h"
+#include "NiagaraShaderParametersBuilder.h"
 #include "NiagaraSystemInstance.h"
 
 #include "Async/Async.h"
@@ -13,7 +14,6 @@
 
 //////////////////////////////////////////////////////////////////////////
 
-FName UNiagaraDataInterfaceDebugDraw::CompileTagKey = TEXT("CompilerTagKey");
 struct FNiagaraDebugDrawDIFunctionVersion
 {
 	enum Type
@@ -31,6 +31,12 @@ struct FNiagaraDebugDrawDIFunctionVersion
 	};
 };
 
+BEGIN_SHADER_PARAMETER_STRUCT(FNDIDebugDrawShaderParameters,)
+	SHADER_PARAMETER_RDG_BUFFER_UAV(RWBuffer<uint>,	RWNDIDebugDrawArgs)
+	SHADER_PARAMETER_RDG_BUFFER_UAV(RWBuffer<uint>,	RWNDIDebugDrawLineVertex)
+	SHADER_PARAMETER(uint32,						NDIDebugDrawLineMaxInstances)
+END_SHADER_PARAMETER_STRUCT()
+
 struct FNDIDebugDrawInstanceData_GameThread
 {
 	FNDIDebugDrawInstanceData_GameThread()
@@ -1430,6 +1436,10 @@ struct FNDIDebugDrawProxy : public FNiagaraDataInterfaceProxy
 
 namespace NDIDebugDrawLocal
 {
+	static const FName CompileTagKey = TEXT("CompilerTagKey");
+	static const TCHAR* CommonShaderFile = TEXT("/Plugin/FX/Niagara/Private/NiagaraDataInterfaceDebugDraw.ush");
+	static const TCHAR* TemplateShaderFile = TEXT("/Plugin/FX/Niagara/Private/NiagaraDataInterfaceDebugDrawTemplate.ush");
+
 	static const FName DrawLineName(TEXT("DrawLine"));
 	static const FName DrawRectangleName(TEXT("DrawRectangle"));
 	static const FName DrawCircleName(TEXT("DrawCircle"));
@@ -2272,103 +2282,6 @@ namespace NDIDebugDrawLocal
 
 //////////////////////////////////////////////////////////////////////////
 
-struct FNiagaraDataInterfaceParametersCS_DebugDraw : public FNiagaraDataInterfaceParametersCS
-{
-	DECLARE_TYPE_LAYOUT(FNiagaraDataInterfaceParametersCS_DebugDraw, NonVirtual);
-
-public:
-	void Bind(const FNiagaraDataInterfaceGPUParamInfo& ParameterInfo, const class FShaderParameterMap& ParameterMap)
-	{
-		DrawArgsParams.Bind(ParameterMap, TEXT("NDIDebugDrawArgs"));
-		DrawLineVertexParam.Bind(ParameterMap, TEXT("NDIDebugDrawLineVertex"));
-		DrawLineMaxInstancesParam.Bind(ParameterMap, TEXT("NDIDebugDrawLineMaxInstances"));
-	}
-
-	void Set(FRHICommandList& RHICmdList, const FNiagaraDataInterfaceSetArgs& Context) const
-	{
-		FRHIComputeShader* ComputeShaderRHI = Context.Shader.GetComputeShader();
-#if NIAGARA_COMPUTEDEBUG_ENABLED
-		auto* DIProxy = static_cast<FNDIDebugDrawProxy*>(Context.DataInterface);
-		auto* InstanceData = &DIProxy->SystemInstancesToProxyData_RT.FindChecked(Context.SystemInstanceID);
-
-		FNiagaraSimulationDebugDrawData* DebugDraw = nullptr;
-		if ( InstanceData->GpuComputeDebug )
-		{
-			DebugDraw = InstanceData->GpuComputeDebug->GetSimulationDebugDrawData(Context.SystemInstanceID, true, InstanceData->OverrideMaxLineInstances);
-		}
-
-		const bool bIsValid =
-			NDIDebugDrawLocal::GNiagaraDebugDrawEnabled &&
-			(DebugDraw != nullptr) &&
-			DrawArgsParams.IsUAVBound() &&
-			DrawLineVertexParam.IsUAVBound();
-
-		if ( bIsValid )
-		{
-			FRHITransitionInfo Transitions[] =
-			{
-				FRHITransitionInfo(DebugDraw->GpuLineBufferArgs.UAV, ERHIAccess::IndirectArgs, ERHIAccess::UAVCompute),
-				FRHITransitionInfo(DebugDraw->GpuLineVertexBuffer.UAV, ERHIAccess::SRVMask, ERHIAccess::UAVCompute),
-			};
-			RHICmdList.Transition(Transitions);
-
-			RHICmdList.SetUAVParameter(ComputeShaderRHI, DrawArgsParams.GetUAVIndex(), DebugDraw->GpuLineBufferArgs.UAV);
-			RHICmdList.SetUAVParameter(ComputeShaderRHI, DrawLineVertexParam.GetUAVIndex(), DebugDraw->GpuLineVertexBuffer.UAV);
-			SetShaderValue(RHICmdList, ComputeShaderRHI, DrawLineMaxInstancesParam, DebugDraw->GpuLineMaxInstances);
-		}
-		else
-#endif //NIAGARA_COMPUTEDEBUG_ENABLED
-		{
-			SetShaderValue(RHICmdList, ComputeShaderRHI, DrawLineMaxInstancesParam, 0);
-		}
-	}
-
-	void Unset(FRHICommandList& RHICmdList, const FNiagaraDataInterfaceSetArgs& Context) const
-	{
-		FRHIComputeShader* ComputeShaderRHI = Context.Shader.GetComputeShader();
-		DrawArgsParams.UnsetUAV(RHICmdList, ComputeShaderRHI);
-		DrawLineVertexParam.UnsetUAV(RHICmdList, ComputeShaderRHI);
-
-#if NIAGARA_COMPUTEDEBUG_ENABLED
-		auto* DIProxy = static_cast<FNDIDebugDrawProxy*>(Context.DataInterface);
-		auto* InstanceData = &DIProxy->SystemInstancesToProxyData_RT.FindChecked(Context.SystemInstanceID);
-
-		FNiagaraSimulationDebugDrawData* DebugDraw = nullptr;
-		if (InstanceData->GpuComputeDebug)
-		{
-			DebugDraw = InstanceData->GpuComputeDebug->GetSimulationDebugDrawData(Context.SystemInstanceID, true, InstanceData->OverrideMaxLineInstances);
-		}
-
-		const bool bIsValid =
-			NDIDebugDrawLocal::GNiagaraDebugDrawEnabled &&
-			(DebugDraw != nullptr) &&
-			DrawArgsParams.IsUAVBound() &&
-			DrawLineVertexParam.IsUAVBound();
-
-		if (bIsValid)
-		{
-			FRHITransitionInfo Transitions[] =
-			{
-				FRHITransitionInfo(DebugDraw->GpuLineBufferArgs.UAV, ERHIAccess::UAVCompute, ERHIAccess::IndirectArgs),
-				FRHITransitionInfo(DebugDraw->GpuLineVertexBuffer.UAV, ERHIAccess::UAVCompute, ERHIAccess::SRVMask),
-			};
-			RHICmdList.Transition(Transitions);
-		}
-#endif //NIAGARA_COMPUTEDEBUG_ENABLED
-	}
-
-private:
-	LAYOUT_FIELD(FRWShaderParameter,	DrawArgsParams);
-	LAYOUT_FIELD(FRWShaderParameter,	DrawLineVertexParam);
-	LAYOUT_FIELD(FShaderParameter,		DrawLineMaxInstancesParam);
-};
-
-IMPLEMENT_TYPE_LAYOUT(FNiagaraDataInterfaceParametersCS_DebugDraw);
-
-IMPLEMENT_NIAGARA_DI_PARAMETER(UNiagaraDataInterfaceDebugDraw, FNiagaraDataInterfaceParametersCS_DebugDraw);
-
-//////////////////////////////////////////////////////////////////////////
-
 UNiagaraDataInterfaceDebugDraw::UNiagaraDataInterfaceDebugDraw(FObjectInitializer const& ObjectInitializer)
 	: Super(ObjectInitializer)
 {
@@ -2818,6 +2731,48 @@ bool UNiagaraDataInterfaceDebugDraw::UpgradeFunctionCall(FNiagaraFunctionSignatu
 }
 #endif
 
+void UNiagaraDataInterfaceDebugDraw::BuildShaderParameters(FNiagaraShaderParametersBuilder& ShaderParametersBuilder) const
+{
+	ShaderParametersBuilder.AddIncludedStruct<FNDIDebugDrawShaderParameters>();
+}
+
+void UNiagaraDataInterfaceDebugDraw::SetShaderParameters(const FNiagaraDataInterfaceSetShaderParametersContext& Context) const
+{
+	FRDGBuilder& GraphBuilder = Context.GetGraphBuilder();
+	FNDIDebugDrawShaderParameters* ShaderParameters = Context.GetParameterIncludedStruct<FNDIDebugDrawShaderParameters>();
+
+#if NIAGARA_COMPUTEDEBUG_ENABLED
+	FNDIDebugDrawProxy& DIProxy = Context.GetProxy<FNDIDebugDrawProxy>();
+	FNDIDebugDrawInstanceData_RenderThread& InstanceData = DIProxy.SystemInstancesToProxyData_RT.FindChecked(Context.GetSystemInstanceID());
+
+	FNiagaraSimulationDebugDrawData* DebugDraw = nullptr;
+	if (InstanceData.GpuComputeDebug)
+	{
+		DebugDraw = InstanceData.GpuComputeDebug->GetSimulationDebugDrawData(GraphBuilder, Context.GetSystemInstanceID(), InstanceData.OverrideMaxLineInstances);
+	}
+
+	const bool bIsValid =
+		NDIDebugDrawLocal::GNiagaraDebugDrawEnabled &&
+		(DebugDraw != nullptr) &&
+		(DebugDraw->GpuLineMaxInstances > 0) &&
+		Context.IsResourceBound(&ShaderParameters->RWNDIDebugDrawArgs) &&
+		Context.IsResourceBound(&ShaderParameters->RWNDIDebugDrawLineVertex);
+
+	if ( bIsValid )
+	{
+		ShaderParameters->RWNDIDebugDrawArgs			= DebugDraw->GpuLineBufferArgs.GetOrCreateUAV(GraphBuilder);
+		ShaderParameters->RWNDIDebugDrawLineVertex		= DebugDraw->GpuLineVertexBuffer.GetOrCreateUAV(GraphBuilder);
+		ShaderParameters->NDIDebugDrawLineMaxInstances	= DebugDraw->GpuLineMaxInstances;
+	}
+	else
+#endif
+	{
+		ShaderParameters->RWNDIDebugDrawArgs			= Context.GetComputeDispatchInterface().GetEmptyBufferUAV(GraphBuilder, PF_R32_UINT);
+		ShaderParameters->RWNDIDebugDrawLineVertex		= Context.GetComputeDispatchInterface().GetEmptyBufferUAV(GraphBuilder, PF_R32_UINT);
+		ShaderParameters->NDIDebugDrawLineMaxInstances	= 0;
+	}
+}
+
 void UNiagaraDataInterfaceDebugDraw::GetVMExternalFunction(const FVMExternalFunctionBindingInfo& BindingInfo, void* InstanceData, FVMExternalFunction& OutFunc)
 {
 	if (BindingInfo.Name == NDIDebugDrawLocal::DrawBoxName)
@@ -2867,7 +2822,7 @@ void UNiagaraDataInterfaceDebugDraw::GetVMExternalFunction(const FVMExternalFunc
 	else if (BindingInfo.FunctionSpecifiers.Num() != 0)
 	{
 		// The HLSL translator adds this function specifier in so that we have a unqiue key during compilation.
-		const FVMFunctionSpecifier* Specifier = BindingInfo.FunctionSpecifiers.FindByPredicate([&](const FVMFunctionSpecifier& Info) { return Info.Key == UNiagaraDataInterfaceDebugDraw::CompileTagKey; });
+		const FVMFunctionSpecifier* Specifier = BindingInfo.FunctionSpecifiers.FindByPredicate([&](const FVMFunctionSpecifier& Info) { return Info.Key == NDIDebugDrawLocal::CompileTagKey; });
 
 		if (Specifier && !Specifier->Value.IsNone())
 		{
@@ -2961,161 +2916,62 @@ bool UNiagaraDataInterfaceDebugDraw::AppendCompileHash(FNiagaraCompileHashVisito
 	if (!Super::AppendCompileHash(InVisitor))
 		return false;
 
-	FSHAHash Hash = GetShaderFileHash((TEXT("/Plugin/FX/Niagara/Private/NiagaraDataInterfaceDebugDraw.ush")), EShaderPlatform::SP_PCD3D_SM5);
-	InVisitor->UpdateString(TEXT("NiagaraDataInterfaceDebugDrawHLSLSource"), Hash.ToString());
+	InVisitor->UpdateString(TEXT("NiagaraDataInterfaceDebugDrawHLSLSource"), GetShaderFileHash(NDIDebugDrawLocal::CommonShaderFile, EShaderPlatform::SP_PCD3D_SM5).ToString());
+	InVisitor->UpdateString(TEXT("NiagaraDataInterfaceDebugDrawTemplateHLSLSource"), GetShaderFileHash(NDIDebugDrawLocal::TemplateShaderFile, EShaderPlatform::SP_PCD3D_SM5).ToString());
 
 	InVisitor->UpdatePOD<int32>(TEXT("NiagaraDataInterfaceDebugDrawVersion"), FNiagaraDebugDrawDIFunctionVersion::LatestVersion);	
+	InVisitor->UpdateShaderParameters<FNDIDebugDrawShaderParameters>();
 
 	return true;
 }
 
 void UNiagaraDataInterfaceDebugDraw::GetCommonHLSL(FString& OutHLSL)
 {
-	OutHLSL += TEXT("#include \"/Plugin/FX/Niagara/Private/NiagaraDataInterfaceDebugDraw.ush\"\n");
+	OutHLSL.Appendf(TEXT("#include \"%s\"\n"), NDIDebugDrawLocal::CommonShaderFile);
 }
 
 
-bool UNiagaraDataInterfaceDebugDraw::GetFunctionHLSL(const FNiagaraDataInterfaceGPUParamInfo& ParamInfo, const FNiagaraDataInterfaceGeneratedFunction& FunctionInfo, int FunctionInstanceIndex, FString& OutHLSL)
+void UNiagaraDataInterfaceDebugDraw::GetParameterDefinitionHLSL(const FNiagaraDataInterfaceGPUParamInfo& ParamInfo, FString& OutHLSL)
 {
-	TMap<FString, FStringFormatArg> ArgsSample =
+	TMap<FString, FStringFormatArg> TemplateArgs =
 	{
-		{TEXT("InstanceFunctionName"), FunctionInfo.InstanceName},
+		{TEXT("ParameterName"),	ParamInfo.DataInterfaceHLSLSymbol},
 	};
 
-	if (FunctionInfo.DefinitionName == NDIDebugDrawLocal::DrawLineName)
-	{
-		static const TCHAR* FormatSample = TEXT("void {InstanceFunctionName}(bool bExecute, float3 LineStart, float3 LineEnd, float4 Color) { NDIDebugDraw_Line(bExecute, LineStart, LineEnd, Color); }\n");
-		OutHLSL += FString::Format(FormatSample, ArgsSample);
-		return true;
-	}
-	else if (FunctionInfo.DefinitionName == NDIDebugDrawLocal::DrawRectangleName)
-	{
-		static const TCHAR* FormatSample = TEXT("void {InstanceFunctionName}(bool bExecute, float3 Location, float3 XAxis, float3 YAxis, float2 Extents, int NumXSegments, int NumYSegments, float4 Color, bool bUnbounded) { NDIDebugDraw_DrawRectangle(bExecute, Location, XAxis, YAxis, Extents, NumXSegments, NumYSegments, Color, bUnbounded); }\n");
-		OutHLSL += FString::Format(FormatSample, ArgsSample);
-		return true;
-	}
-	else if (FunctionInfo.DefinitionName == NDIDebugDrawLocal::DrawCircleName)
-	{
-		static const TCHAR* FormatSample = TEXT("void {InstanceFunctionName}(bool bExecute, float3 Location, float3 XAxis, float3 YAxis, float Scale, int Segments, float4 Color) { NDIDebugDraw_Circle(bExecute, Location, XAxis, YAxis, Scale, Segments, Color); }\n");
-		OutHLSL += FString::Format(FormatSample, ArgsSample);
-		return true;
-	}
-	if (FunctionInfo.DefinitionName == NDIDebugDrawLocal::DrawBoxName)
-	{
-		static const TCHAR* FormatSample = TEXT("void {InstanceFunctionName}(bool bExecute, float3 Location, float4 Rotation, float3 Extents, float4 Color) { NDIDebugDraw_DrawBox(bExecute, Location, Rotation, Extents, Color); }\n");
-		OutHLSL += FString::Format(FormatSample, ArgsSample);
-		return true;
-	}
-	else if (FunctionInfo.DefinitionName == NDIDebugDrawLocal::DrawSphereName)
-	{
-		static const TCHAR* FormatSample = TEXT("void {InstanceFunctionName}(bool bExecute, float3 LineStart, float Radius, int Segments, float4 Color) { NDIDebugDraw_Sphere(bExecute, LineStart, Radius, Segments, Color); }\n");
-		OutHLSL += FString::Format(FormatSample, ArgsSample);
-		return true;
-	}
-	else if (FunctionInfo.DefinitionName == NDIDebugDrawLocal::DrawCylinderName)
-	{
-		static const TCHAR* FormatSample = TEXT("void {InstanceFunctionName}(bool bExecute, float3 Location, float3 Axis, float Height, float Radius, int NumHeightSegments, int NumRadiusSegments, float4 Color) { NDIDebugDraw_Cylinder(bExecute, Location, Axis, Height, Radius, NumHeightSegments, NumRadiusSegments, Color); }\n");
-		OutHLSL += FString::Format(FormatSample, ArgsSample);
-		return true;
-	}
-	else if (FunctionInfo.DefinitionName == NDIDebugDrawLocal::DrawConeName)
-	{
-		static const TCHAR* FormatSample = TEXT("void {InstanceFunctionName}(bool bExecute, float3 Location, float3 Axis, float Height, float RadiusTop, float RadiusBottom, int NumHeightSegments, int NumRadiusSegments, float4 Color) {  }\n");
-		OutHLSL += FString::Format(FormatSample, ArgsSample);
-		return true;
-	}
-	else if (FunctionInfo.DefinitionName == NDIDebugDrawLocal::DrawTorusName)
-	{
-		static const TCHAR* FormatSample = TEXT("void {InstanceFunctionName}(bool bExecute, float3 Location, float3 Axis, float MajorRadius, float MinorRadius, float MajorRadiusSegments, float MinorRadiusSegments, float4 Color) {  }\n");
-		OutHLSL += FString::Format(FormatSample, ArgsSample);
-		return true;
-	}
-	else if (FunctionInfo.DefinitionName == NDIDebugDrawLocal::DrawCoordinateSystemName)
-	{
-		static const TCHAR* FormatSample = TEXT("void {InstanceFunctionName}(bool bExecute, float3 Location, float4 Rotation, float Scale) { NDIDebugDraw_CoordinateSystem(bExecute, Location, Rotation, Scale); }\n");
-		OutHLSL += FString::Format(FormatSample, ArgsSample);
-		return true;
-	}
-	else if (FunctionInfo.DefinitionName == NDIDebugDrawLocal::DrawGrid2DName)
-	{
-		static const TCHAR* FormatSample = TEXT("void {InstanceFunctionName}(bool bExecute, float3 Center, float4 Rotation, float2 Extents, int NumCellsX, int NumCellsY, float4 Color) { NDIDebugDraw_Grid2D(bExecute, Center, Rotation, Extents, int2(NumCellsX, NumCellsY), Color); }\n");
-		OutHLSL += FString::Format(FormatSample, ArgsSample);
-		return true;
-	}
-	else if (FunctionInfo.DefinitionName == NDIDebugDrawLocal::DrawGrid3DName)
-	{
-		static const TCHAR* FormatSample = TEXT("void {InstanceFunctionName}(bool bExecute, float3 Center, float4 Rotation, float3 Extents, int NumCellsX, int NumCellsY, int NumCellsZ, float4 Color) { NDIDebugDraw_Grid3D(bExecute, Center, Rotation, Extents, int3(NumCellsX, NumCellsY, NumCellsZ), Color); }\n");
-		OutHLSL += FString::Format(FormatSample, ArgsSample);
-		return true;
-	}
-	else if (FunctionInfo.DefinitionName == NDIDebugDrawLocal::DrawLinePersistentName)
-	{
-		static const TCHAR* FormatSample = TEXT("void {InstanceFunctionName}(float3 StartLocation, int StartLocationCoordinateSpace, float3 EndLocation, int EndLocationCoordinateSpace, float4 Color) { }\n");
-		OutHLSL += FString::Format(FormatSample, ArgsSample);
-		return true;
-	}
-	else if (FunctionInfo.DefinitionName == NDIDebugDrawLocal::DrawRectanglePersistentName)
-	{
-		static const TCHAR* FormatSample = TEXT("void {InstanceFunctionName}(float3 Center, int CenterCoordinateSpace, float3 Offset, int OffsetCoordinateSpace, float3 Extents, bool HalfExtents, float3 XAxis, float3 YAxis, int AxisCoordinateSpace, int NumXSegments, int NumYSegments, bool bUnboundedPlane, float4 Color) { }\n");
-		OutHLSL += FString::Format(FormatSample, ArgsSample);
-		return true;
-	}
-	else if (FunctionInfo.DefinitionName == NDIDebugDrawLocal::DrawCirclePersistentName)
-	{
-		static const TCHAR* FormatSample = TEXT("void {InstanceFunctionName}(float3 Center, int CenterCoordinateSpace, float3 Offset, int OffsetCoordinateSpace, float Radius, float3 XAxis, float3 YAxis, int AxisCoordinateSpace, int NumSegments, float4 Color) { }\n");
-		OutHLSL += FString::Format(FormatSample, ArgsSample);
-		return true;
-	}
-	else if (FunctionInfo.DefinitionName == NDIDebugDrawLocal::DrawBoxPersistentName)
-	{
-		static const TCHAR* FormatSample = TEXT("void {InstanceFunctionName}(float3 Center, int CenterCoordinateSpace, float3 Offset, int OffsetCoordinateSpace, float3 Extents, bool HalfExtents, float3 RotationAxis, float RotationAngle, int RotationCoordinateSpace, float4 Color) { }\n");
-		OutHLSL += FString::Format(FormatSample, ArgsSample);
-		return true;
-	}
-	else if (FunctionInfo.DefinitionName == NDIDebugDrawLocal::DrawSpherePersistentName)
-	{
-		static const TCHAR* FormatSample = TEXT("void {InstanceFunctionName}(float3 Center, int CenterCoordinateSpace, float3 OffsetFromCenter, int OffsetCoordinateSpace, float Radius,  int RadiusCoordinateSpace, bool HemisphereX, bool HemisphereY, bool HemisphereZ, float3 OrientationAxis, int RotationCoordinateSpace, float4 AdditionalRotation, float3 NonUniformScale, int NumSegments, float4 Color) { }\n");
-		OutHLSL += FString::Format(FormatSample, ArgsSample);
-		return true;
-	}
-	else if (FunctionInfo.DefinitionName == NDIDebugDrawLocal::DrawCylinderPersistentName)
-	{
-		static const TCHAR* FormatSample = TEXT("void {InstanceFunctionName}(float3 Center, int CenterCoordinateSpace, float3 OffsetFromCenter, int OffsetCoordinateSpace, float Radius, float Height, bool IsHalfHeight, bool HemisphereX, bool HemisphereY, float3 OrientationAxis, int OrientationAxisCoordinateSpace, float3 NonUniformScale, int NumRadiusSegments, int NumHeightSegments, float4 Color) { }\n");
-		OutHLSL += FString::Format(FormatSample, ArgsSample);
-		return true;
-	}
-	else if (FunctionInfo.DefinitionName == NDIDebugDrawLocal::DrawConePersistentName)
-	{
-		static const TCHAR* FormatSample = TEXT("void {InstanceFunctionName}(float3 Center, int CenterCoordinateSpace, float3 OffsetFromCenter, int OffsetCoordinateSpace, float Angle, float Length, int OrientationAxis, int OrientationAxisCoordinateSpace, float3 NonUniformScale, int NumRadiusSegments, int NumHeightSegments, float4 Color) { }\n");
-		OutHLSL += FString::Format(FormatSample, ArgsSample);
-		return true;
-	}
-	else if (FunctionInfo.DefinitionName == NDIDebugDrawLocal::DrawTorusPersistentName)
-	{
-		static const TCHAR* FormatSample = TEXT("void {InstanceFunctionName}(float3 Center, int CenterCoordinateSpace, float3 OffsetFromCenter, int OffsetCoordinateSpace, float MajorRadius, float MinorRadius, float3 OrientationAxis, int OrientationAxisCoordinateSpace, float3 NonUniformScale, int MajorRadiusSegments, int MinorRadiusSegments, float4 Color) { }\n");
-		OutHLSL += FString::Format(FormatSample, ArgsSample);
-		return true;
-	}
-	else if (FunctionInfo.DefinitionName == NDIDebugDrawLocal::DrawCoordinateSystemPersistentName)
-	{
-		static const TCHAR* FormatSample = TEXT("void {InstanceFunctionName}(float3 Center, int CenterCoordinateSpace, float3 OffsetFromCenter, int OffsetCoordinateSpace, float4 Rotation, int RotationCoordinateSpace, float Scale) { }\n");
-		OutHLSL += FString::Format(FormatSample, ArgsSample);
-		return true;
-	}
-	else if (FunctionInfo.DefinitionName == NDIDebugDrawLocal::DrawGrid2DPersistentName)
-	{
-		static const TCHAR* FormatSample = TEXT("void {InstanceFunctionName}(float3 Center, int CenterCoordinateSpace, float3 OffsetFromCenter, int OffsetCoordinateSpace, float4 Rotation, int RotationCoordinateSpace, float2 Extents, int NumCellsX, int NumCellsY, float4 Color) { }\n");
-		OutHLSL += FString::Format(FormatSample, ArgsSample);
-		return true;
-	}
-	else if (FunctionInfo.DefinitionName == NDIDebugDrawLocal::DrawGrid3DPersistentName)
-	{
-		static const TCHAR* FormatSample = TEXT("void {InstanceFunctionName}(float3 Center, int CenterCoordinateSpace, float3 OffsetFromCenter, int OffsetCoordinateSpace, float4 Rotation, int RotationCoordinateSpace, float3 Extents, int NumCellsX, int NumCellsY, int NumCellsZ, float4 Color) { }\n");
-		OutHLSL += FString::Format(FormatSample, ArgsSample);
-		return true;
-	}
+	FString TemplateFile;
+	LoadShaderSourceFile(NDIDebugDrawLocal::TemplateShaderFile, EShaderPlatform::SP_PCD3D_SM5, &TemplateFile, nullptr);
+	OutHLSL += FString::Format(*TemplateFile, TemplateArgs);
+}
 
-	return false;
+bool UNiagaraDataInterfaceDebugDraw::GetFunctionHLSL(const FNiagaraDataInterfaceGPUParamInfo& ParamInfo, const FNiagaraDataInterfaceGeneratedFunction& FunctionInfo, int FunctionInstanceIndex, FString& OutHLSL)
+{
+	static const TSet<FName> ValidGpuFunctions =
+	{
+		NDIDebugDrawLocal::DrawLineName,
+		NDIDebugDrawLocal::DrawRectangleName,
+		NDIDebugDrawLocal::DrawCircleName,
+		NDIDebugDrawLocal::DrawBoxName,
+		NDIDebugDrawLocal::DrawSphereName,
+		NDIDebugDrawLocal::DrawCylinderName,
+		NDIDebugDrawLocal::DrawConeName,
+		NDIDebugDrawLocal::DrawTorusName,
+		NDIDebugDrawLocal::DrawCoordinateSystemName,
+		NDIDebugDrawLocal::DrawGrid2DName,
+		NDIDebugDrawLocal::DrawGrid3DName,
+		NDIDebugDrawLocal::DrawLinePersistentName,
+		NDIDebugDrawLocal::DrawRectanglePersistentName,
+		NDIDebugDrawLocal::DrawCirclePersistentName,
+		NDIDebugDrawLocal::DrawBoxPersistentName,
+		NDIDebugDrawLocal::DrawSpherePersistentName,
+		NDIDebugDrawLocal::DrawCylinderPersistentName,
+		NDIDebugDrawLocal::DrawConePersistentName,
+		NDIDebugDrawLocal::DrawTorusPersistentName,
+		NDIDebugDrawLocal::DrawCoordinateSystemPersistentName,
+		NDIDebugDrawLocal::DrawGrid2DPersistentName,
+		NDIDebugDrawLocal::DrawGrid3DPersistentName,
+	};
+
+	return ValidGpuFunctions.Contains(FunctionInfo.DefinitionName);
 }
 
 bool UNiagaraDataInterfaceDebugDraw::GenerateCompilerTagPrefix(const FNiagaraFunctionSignature& InSignature, FString& OutPrefix) const 
@@ -3145,7 +3001,7 @@ bool UNiagaraDataInterfaceDebugDraw::GPUContextInit(const FNiagaraScriptDataInte
 		if (PerInstanceData && Sig.FunctionSpecifiers.Num() > 0)
 		{
 			// The HLSL translator adds this function specifier in so that we have a unqiue key during compilation.
-			const FName* Specifier = Sig.FunctionSpecifiers.Find(UNiagaraDataInterfaceDebugDraw::CompileTagKey);
+			const FName* Specifier = Sig.FunctionSpecifiers.Find(NDIDebugDrawLocal::CompileTagKey);
 
 			if (Specifier && !Specifier->IsNone())
 			{
@@ -3204,7 +3060,6 @@ bool UNiagaraDataInterfaceDebugDraw::PerInstanceTickPostSimulate(void* PerInstan
 			InstanceData->HandlePersistentShapes(SystemInstance, DeltaSeconds);
 		}
 
-
 		// Dispatch information to the RT proxy
 		ENQUEUE_RENDER_COMMAND(NDIDebugDrawUpdate)(
 			[RT_Proxy=GetProxyAs<FNDIDebugDrawProxy>(), RT_InstanceID=SystemInstance->GetId(), RT_TickCount=SystemInstance->GetTickCount(), RT_LineBuffer=MoveTemp(InstanceData->LineBuffer)](FRHICommandListImmediate& RHICmdList) mutable
@@ -3215,7 +3070,7 @@ bool UNiagaraDataInterfaceDebugDraw::PerInstanceTickPostSimulate(void* PerInstan
 
 					if (RT_InstanceData && RT_InstanceData->GpuComputeDebug)
 					{
-						if (FNiagaraSimulationDebugDrawData* DebugDraw = RT_InstanceData->GpuComputeDebug->GetSimulationDebugDrawData(RT_InstanceID, false, RT_InstanceData->OverrideMaxLineInstances))
+						if (FNiagaraSimulationDebugDrawData* DebugDraw = RT_InstanceData->GpuComputeDebug->GetSimulationDebugDrawData(RT_InstanceID))
 						{
 							if (DebugDraw->LastUpdateTickCount != RT_TickCount)
 							{
@@ -3305,4 +3160,4 @@ bool UNiagaraDataInterfaceDebugDraw::CopyToInternal(UNiagaraDataInterface* Desti
 	OtherTyped->OverrideMaxLineInstances = OverrideMaxLineInstances;	
 
 	return true;
-}
\ No newline at end of file
+}
diff --git a/Engine/Plugins/FX/Niagara/Source/Niagara/Private/DataInterface/NiagaraDataInterfaceDebugDraw.h b/Engine/Plugins/FX/Niagara/Source/Niagara/Private/DataInterface/NiagaraDataInterfaceDebugDraw.h
--- a/Engine/Plugins/FX/Niagara/Source/Niagara/Private/DataInterface/NiagaraDataInterfaceDebugDraw.h
+++ b/Engine/Plugins/FX/Niagara/Source/Niagara/Private/DataInterface/NiagaraDataInterfaceDebugDraw.h
@@ -12,8 +12,6 @@ class NIAGARA_API UNiagaraDataInterfaceDebugDraw : public UNiagaraDataInterface
 	GENERATED_UCLASS_BODY()
 
 public:
-	DECLARE_NIAGARA_DI_PARAMETER();
-		
 	// Number of debug lines is set as the max of OverrideMaxLineInstances and fx.Niagara.GpuComputeDebug.MaxLineInstances
 	UPROPERTY(EditAnywhere, Category = "Debug Draw")
 	uint32 OverrideMaxLineInstances = 0;
@@ -43,15 +41,17 @@ public:
 		Max
 	};
 
-	static FName CompileTagKey;
-
 #if WITH_EDITORONLY_DATA
 	virtual bool AppendCompileHash(FNiagaraCompileHashVisitor* InVisitor) const override;
 	virtual void GetCommonHLSL(FString& OutHLSL) override;
+	virtual void GetParameterDefinitionHLSL(const FNiagaraDataInterfaceGPUParamInfo& ParamInfo, FString& OutHLSL) override;
 	virtual bool GetFunctionHLSL(const FNiagaraDataInterfaceGPUParamInfo& ParamInfo, const FNiagaraDataInterfaceGeneratedFunction& FunctionInfo, int FunctionInstanceIndex, FString& OutHLSL) override;
 	virtual bool GenerateCompilerTagPrefix(const FNiagaraFunctionSignature& InSignature, FString& OutPrefix) const override;	
 	virtual bool UpgradeFunctionCall(FNiagaraFunctionSignature& FunctionSignature) override;
 #endif
+	virtual bool UseLegacyShaderBindings() const override { return false; }
+	virtual void BuildShaderParameters(FNiagaraShaderParametersBuilder& ShaderParametersBuilder) const override;
+	virtual void SetShaderParameters(const FNiagaraDataInterfaceSetShaderParametersContext& Context) const override;
 
 	virtual bool PerInstanceTick(void* PerInstanceData, FNiagaraSystemInstance* SystemInstance, float DeltaSeconds) override;
 	virtual bool PerInstanceTickPostSimulate(void* PerInstanceData, FNiagaraSystemInstance* SystemInstance, float DeltaSeconds) override;
@@ -71,6 +71,5 @@ public:
 protected:
 	/** Copy one niagara DI to this */
 	virtual bool CopyToInternal(UNiagaraDataInterface* Destination) const override;
-
 	//UNiagaraDataInterface Interface
 };
diff --git a/Engine/Plugins/FX/Niagara/Source/Niagara/Private/NiagaraGpuComputeDebug.cpp b/Engine/Plugins/FX/Niagara/Source/Niagara/Private/NiagaraGpuComputeDebug.cpp
--- a/Engine/Plugins/FX/Niagara/Source/Niagara/Private/NiagaraGpuComputeDebug.cpp
+++ b/Engine/Plugins/FX/Niagara/Source/Niagara/Private/NiagaraGpuComputeDebug.cpp
@@ -46,6 +46,12 @@ static FAutoConsoleVariableRef CVarNiagaraGpuComputeDebug_DrawDebugEnabled(
 
 #if NIAGARA_COMPUTEDEBUG_ENABLED
 
+namespace NiagaraGpuComputeDebugLocal
+{
+	static constexpr uint32 NumUintsPerLine = 7;
+}
+static_assert(sizeof(FNiagaraSimulationDebugDrawData::FGpuLine) == (NiagaraGpuComputeDebugLocal::NumUintsPerLine * sizeof(uint32)), "Line size does not match expected GPU size");
+
 //////////////////////////////////////////////////////////////////////////
 
 FNiagaraGpuComputeDebug::FNiagaraGpuComputeDebug(ERHIFeatureLevel::Type InFeatureLevel)
@@ -53,46 +59,39 @@ FNiagaraGpuComputeDebug::FNiagaraGpuComputeDebug(ERHIFeatureLevel::Type InFeatur
 {
 }
 
-void FNiagaraGpuComputeDebug::Tick(FRHICommandListImmediate& RHICmdList)
+void FNiagaraGpuComputeDebug::Tick(FRDGBuilder& GraphBuilder)
 {
 	for (auto it=DebugDrawBuffers.CreateConstIterator(); it; ++it)
 	{
 		FNiagaraSimulationDebugDrawData* DebugDrawData = it.Value().Get();
+		DebugDrawData->RDGStaticLineBuffer = nullptr;
+		DebugDrawData->GpuLineBufferArgs.EndGraphUsage();
+		DebugDrawData->GpuLineVertexBuffer.EndGraphUsage();
+
 		if ( !DebugDrawData->bRequiresUpdate )
 		{
 			continue;
 		}
 
 		DebugDrawData->bRequiresUpdate = false;
-		if (DebugDrawData->GpuLineMaxInstances > 0)
+		if (DebugDrawData->GpuLineBufferArgs.IsValid())
 		{
-			//-OPT: Batch UAV clears
-			NiagaraDebugShaders::ClearUAV(RHICmdList, DebugDrawData->GpuLineBufferArgs.UAV, FUintVector4(2, 0, 0, 0), 4);
-			RHICmdList.Transition(FRHITransitionInfo(DebugDrawData->GpuLineBufferArgs.UAV, ERHIAccess::UAVCompute, ERHIAccess::IndirectArgs));
+			NiagaraDebugShaders::ClearUAV(GraphBuilder, DebugDrawData->GpuLineBufferArgs.GetOrCreateUAV(GraphBuilder), FUintVector4(2, 0, 0, 0), 4);
 		}
 
 		DebugDrawData->StaticLineCount = DebugDrawData->StaticLines.Num();
 		if (DebugDrawData->StaticLineCount > 0 )
 		{
-			constexpr uint32 NumUintsPerLine = 7;
-			static_assert(sizeof(FNiagaraSimulationDebugDrawData::FGpuLine) == (NumUintsPerLine * sizeof(uint32)), "Line size does not match expected GPU size");
-
-			const uint32 NumElements = FMath::DivideAndRoundUp(DebugDrawData->StaticLineCount, 64u) * 64u * NumUintsPerLine;
-			const uint32 RequiredBytes = NumElements * sizeof(uint32);
-			if ( DebugDrawData->StaticLineBuffer.NumBytes < RequiredBytes )
-			{
-				DebugDrawData->StaticLineBuffer.Release();
-				DebugDrawData->StaticLineBuffer.Initialize(TEXT("NiagaraGpuComputeDebug::StaticLineBuffer"), sizeof(uint32), NumElements, EPixelFormat::PF_R32_UINT);
-			}
-			void* VertexData = RHILockBuffer(DebugDrawData->StaticLineBuffer.Buffer, 0, RequiredBytes, RLM_WriteOnly);
-			FMemory::Memcpy(VertexData, DebugDrawData->StaticLines.GetData(), DebugDrawData->StaticLineCount * DebugDrawData->StaticLines.GetTypeSize());
-			RHIUnlockBuffer(DebugDrawData->StaticLineBuffer.Buffer);
+			const uint32 BufferNumUint32 = NiagaraGpuComputeDebugLocal::NumUintsPerLine * DebugDrawData->StaticLineCount;
+			DebugDrawData->RDGStaticLineBuffer = GraphBuilder.CreateBuffer(FRDGBufferDesc::CreateBufferDesc(sizeof(uint32), BufferNumUint32), TEXT("NiagaraGpuComputeDebug::DrawLineVertexBuffer"));
+			GraphBuilder.QueueBufferUpload(DebugDrawData->RDGStaticLineBuffer, DebugDrawData->StaticLines.GetData(), BufferNumUint32 * sizeof(uint32));
+			GraphBuilder.QueueBufferExtraction(DebugDrawData->RDGStaticLineBuffer, &DebugDrawData->StaticLineBuffer);
 
 			DebugDrawData->StaticLines.Reset();
 		}
 		else
 		{
-			//-OPT: Release buffer?
+			DebugDrawData->StaticLineBuffer.SafeRelease();
 		}
 	}
 }
@@ -185,44 +184,43 @@ void FNiagaraGpuComputeDebug::AddAttributeTexture(FRDGBuilder& GraphBuilder, FNi
 	AddCopyTexturePass(GraphBuilder, Texture, GraphBuilder.RegisterExternalTexture(VisualizeEntry->Texture), CopyInfo);
 }
 
-FNiagaraSimulationDebugDrawData* FNiagaraGpuComputeDebug::GetSimulationDebugDrawData(FNiagaraSystemInstanceID SystemInstanceID, bool bRequiresGpuBuffers, uint32 OverrideMaxDebugLines)
+FNiagaraSimulationDebugDrawData* FNiagaraGpuComputeDebug::GetSimulationDebugDrawData(FNiagaraSystemInstanceID SystemInstanceID)
 {
-	TUniquePtr<FNiagaraSimulationDebugDrawData>& DebugDrawDataPtr = DebugDrawBuffers.FindOrAdd(SystemInstanceID);
-	if (!DebugDrawDataPtr.IsValid())
+	check(IsInRenderingThread());
+
+	TUniquePtr<FNiagaraSimulationDebugDrawData>& DebugDrawData = DebugDrawBuffers.FindOrAdd(SystemInstanceID);
+	if (!DebugDrawData.IsValid())
 	{
-		DebugDrawDataPtr.Reset(new FNiagaraSimulationDebugDrawData());
+		DebugDrawData.Reset(new FNiagaraSimulationDebugDrawData());
 	}
 
-	int MaxLineInstancesToUse = FMath::Max3(DebugDrawDataPtr->GpuLineMaxInstances, (uint32) GNiagaraGpuComputeDebug_MaxLineInstances, OverrideMaxDebugLines);
+	return DebugDrawData.Get();
+}
+
+FNiagaraSimulationDebugDrawData* FNiagaraGpuComputeDebug::GetSimulationDebugDrawData(FRDGBuilder& GraphBuilder, FNiagaraSystemInstanceID SystemInstanceID, uint32 OverrideMaxDebugLines)
+{
+	check(IsInRenderingThread());
+
+	FNiagaraSimulationDebugDrawData* DebugDrawData = GetSimulationDebugDrawData(SystemInstanceID);
 
-	if (bRequiresGpuBuffers && (DebugDrawDataPtr->GpuLineMaxInstances != MaxLineInstancesToUse))
+	const int MaxLineInstancesToUse = FMath::Max3(DebugDrawData->GpuLineMaxInstances, (uint32)GNiagaraGpuComputeDebug_MaxLineInstances, OverrideMaxDebugLines);
+	if (DebugDrawData->GpuLineMaxInstances != MaxLineInstancesToUse)
 	{
-		check(IsInRenderingThread());
-		DebugDrawDataPtr->GpuLineBufferArgs.Release();
-		DebugDrawDataPtr->GpuLineVertexBuffer.Release();
-		DebugDrawDataPtr->GpuLineMaxInstances = MaxLineInstancesToUse;
-		if (DebugDrawDataPtr->GpuLineMaxInstances > 0)
+		DebugDrawData->GpuLineBufferArgs.Release();
+		DebugDrawData->GpuLineVertexBuffer.Release();
+
+		if (MaxLineInstancesToUse > 0)
 		{
-			DebugDrawDataPtr->GpuLineBufferArgs.Initialize(TEXT("NiagaraGpuComputeDebug::DrawLineBufferArgs"), sizeof(uint32), 4, EPixelFormat::PF_R32_UINT, BUF_Static | BUF_DrawIndirect);
-			DebugDrawDataPtr->GpuLineVertexBuffer.Initialize(TEXT("NiagaraGpuComputeDebug::DrawLineVertexBuffer"), sizeof(uint32), 7 * DebugDrawDataPtr->GpuLineMaxInstances, EPixelFormat::PF_R32_UINT, BUF_Static);
+			DebugDrawData->GpuLineMaxInstances = MaxLineInstancesToUse;
 
-			auto& RHICmdList = FRHICommandListExecutor::GetImmediateCommandList();
-			{
-				FUintVector4* IndirectArgs = reinterpret_cast<FUintVector4*>(RHILockBuffer(DebugDrawDataPtr->GpuLineBufferArgs.Buffer, 0, sizeof(uint32) * 4, RLM_WriteOnly));
-				*IndirectArgs = FUintVector4(2, 0, 0, 0);
-				RHIUnlockBuffer(DebugDrawDataPtr->GpuLineBufferArgs.Buffer);
-			}
+			DebugDrawData->GpuLineBufferArgs.Initialize(GraphBuilder, TEXT("NiagaraGpuComputeDebug::DrawLineBufferArgs"), PF_R32_UINT, sizeof(uint32), sizeof(FRHIDrawIndirectParameters) / sizeof(uint32), EBufferUsageFlags::DrawIndirect);
+			DebugDrawData->GpuLineVertexBuffer.Initialize(GraphBuilder, TEXT("NiagaraGpuComputeDebug::DrawLineVertexBuffer"), PF_R32_UINT, sizeof(uint32), NiagaraGpuComputeDebugLocal::NumUintsPerLine * DebugDrawData->GpuLineMaxInstances);
 
-			FRHITransitionInfo Transitions[] =
-			{
-				FRHITransitionInfo(DebugDrawDataPtr->GpuLineBufferArgs.UAV, ERHIAccess::Unknown, ERHIAccess::IndirectArgs),
-				FRHITransitionInfo(DebugDrawDataPtr->GpuLineVertexBuffer.UAV, ERHIAccess::Unknown, ERHIAccess::SRVMask),
-			};
-			RHICmdList.Transition(Transitions);
+			NiagaraDebugShaders::ClearUAV(GraphBuilder, DebugDrawData->GpuLineBufferArgs.GetOrCreateUAV(GraphBuilder), FUintVector4(2, 0, 0, 0), 4);
 		}
 	}
 
-	return DebugDrawDataPtr.Get();
+	return DebugDrawData;
 }
 
 void FNiagaraGpuComputeDebug::RemoveSimulationDebugDrawData(FNiagaraSystemInstanceID SystemInstanceID)
@@ -297,18 +295,23 @@ void FNiagaraGpuComputeDebug::DrawSceneDebug(class FRDGBuilder& GraphBuilder, co
 		FNiagaraSimulationDebugDrawData* DebugDrawData = it.Value().Get();
 		if (DebugDrawData->StaticLineCount > 0)
 		{
+			if (DebugDrawData->RDGStaticLineBuffer == nullptr)
+			{
+				DebugDrawData->RDGStaticLineBuffer = GraphBuilder.RegisterExternalBuffer(DebugDrawData->StaticLineBuffer);
+			}
+
 			NiagaraDebugShaders::DrawDebugLines(
 				GraphBuilder, View, SceneColor, SceneDepth,
 				DebugDrawData->StaticLineCount,
-				DebugDrawData->StaticLineBuffer.SRV
+				DebugDrawData->RDGStaticLineBuffer
 			);
 		}
 		if (DebugDrawData->GpuLineMaxInstances > 0)
 		{
 			NiagaraDebugShaders::DrawDebugLines(
 				GraphBuilder, View, SceneColor, SceneDepth,
-				DebugDrawData->GpuLineBufferArgs.Buffer,
-				DebugDrawData->GpuLineVertexBuffer.SRV
+				DebugDrawData->GpuLineBufferArgs.GetOrCreateBuffer(GraphBuilder),
+				DebugDrawData->GpuLineVertexBuffer.GetOrCreateBuffer(GraphBuilder)
 			);
 		}
 	}
diff --git a/Engine/Plugins/FX/Niagara/Source/Niagara/Private/NiagaraGpuComputeDebug.h b/Engine/Plugins/FX/Niagara/Source/Niagara/Private/NiagaraGpuComputeDebug.h
--- a/Engine/Plugins/FX/Niagara/Source/Niagara/Private/NiagaraGpuComputeDebug.h
+++ b/Engine/Plugins/FX/Niagara/Source/Niagara/Private/NiagaraGpuComputeDebug.h
@@ -7,6 +7,7 @@ Class used help realtime debug Gpu Compute simulations
 
 #include "CoreMinimal.h"
 #include "NiagaraCommon.h"
+#include "NiagaraRenderGraphUtils.h"
 #include "RenderGraphDefinitions.h"
 
 #if NIAGARA_COMPUTEDEBUG_ENABLED
@@ -23,13 +24,14 @@ struct FNiagaraSimulationDebugDrawData
 	bool				bRequiresUpdate = true;
 	int32				LastUpdateTickCount = INDEX_NONE;
 
-	TArray<FGpuLine>	StaticLines;
-	uint32				StaticLineCount = 0;
-	FReadBuffer			StaticLineBuffer;
+	TArray<FGpuLine>				StaticLines;
+	uint32							StaticLineCount = 0;
+	TRefCountPtr<FRDGPooledBuffer>	StaticLineBuffer;
+	FRDGBufferRef					RDGStaticLineBuffer = nullptr;
 
-	FRWBuffer			GpuLineBufferArgs;
-	FRWBuffer			GpuLineVertexBuffer;
-	uint32				GpuLineMaxInstances = 0;
+	uint32							GpuLineMaxInstances = 0;
+	FNiagaraPooledRWBuffer			GpuLineBufferArgs;
+	FNiagaraPooledRWBuffer			GpuLineVertexBuffer;
 };
 
 class FNiagaraGpuComputeDebug
@@ -48,7 +50,7 @@ public:
 	FNiagaraGpuComputeDebug(ERHIFeatureLevel::Type InFeatureLevel);
 
 	// Called at the start of the frame
-	void Tick(FRHICommandListImmediate& RHICmdList);
+	void Tick(FRDGBuilder& GraphBuilder);
 
 	// Enables providing debug information for the system instance
 	void AddSystemInstance(FNiagaraSystemInstanceID SystemInstanceID, FString SystemName);
@@ -72,7 +74,8 @@ public:
 	void AddAttributeTexture(FRDGBuilder& GraphBuilder, FNiagaraSystemInstanceID SystemInstanceID, FName SourceName, FRDGTextureRef Texture, FIntVector4 NumTextureAttributes, FIntVector4 AttributeIndices, FVector2D PreviewDisplayRange = FVector2D::ZeroVector);
 
 	// Get Debug draw buffers for a system instance
-	FNiagaraSimulationDebugDrawData* GetSimulationDebugDrawData(FNiagaraSystemInstanceID SystemInstanceID, bool bRequiresGpuBuffers, uint32 OverrideMaxDebugLines = 0);
+	FNiagaraSimulationDebugDrawData* GetSimulationDebugDrawData(FNiagaraSystemInstanceID SystemInstanceID);
+	FNiagaraSimulationDebugDrawData* GetSimulationDebugDrawData(FRDGBuilder& GraphBuilder, FNiagaraSystemInstanceID SystemInstanceID, uint32 OverrideMaxDebugLines);
 
 	// Force remove debug draw data
 	void RemoveSimulationDebugDrawData(FNiagaraSystemInstanceID SystemInstanceID);
diff --git a/Engine/Plugins/FX/Niagara/Source/Niagara/Private/NiagaraGpuComputeDispatch.cpp b/Engine/Plugins/FX/Niagara/Source/Niagara/Private/NiagaraGpuComputeDispatch.cpp
--- a/Engine/Plugins/FX/Niagara/Source/Niagara/Private/NiagaraGpuComputeDispatch.cpp
+++ b/Engine/Plugins/FX/Niagara/Source/Niagara/Private/NiagaraGpuComputeDispatch.cpp
@@ -1629,7 +1629,7 @@ void FNiagaraGpuComputeDispatch::PreInitViews(FRDGBuilder& GraphBuilder, bool bA
 #if NIAGARA_COMPUTEDEBUG_ENABLED
 	if ( FNiagaraGpuComputeDebug* GpuComputeDebug = GetGpuComputeDebug() )
 	{
-		GpuComputeDebug->Tick(GraphBuilder.RHICmdList);
+		GpuComputeDebug->Tick(GraphBuilder);
 	}
 #endif
 
diff --git a/Engine/Plugins/FX/Niagara/Source/NiagaraShader/Private/NiagaraDebugShaders.cpp b/Engine/Plugins/FX/Niagara/Source/NiagaraShader/Private/NiagaraDebugShaders.cpp
--- a/Engine/Plugins/FX/Niagara/Source/NiagaraShader/Private/NiagaraDebugShaders.cpp
+++ b/Engine/Plugins/FX/Niagara/Source/NiagaraShader/Private/NiagaraDebugShaders.cpp
@@ -105,9 +105,9 @@ class NIAGARASHADER_API FNiagaraClearUAVCS : public FGlobalShader
 	//NIAGARA_DEBUGDRAW_CLEARUAV_UINT_CS
 
 	BEGIN_SHADER_PARAMETER_STRUCT(FParameters, )
-		SHADER_PARAMETER(FUintVector4,			ClearValue)
-		SHADER_PARAMETER(uint32,				ClearSize)
-		SHADER_PARAMETER_UAV(RWBuffer<uint>, BufferToClear)
+		SHADER_PARAMETER(FUintVector4,					ClearValue)
+		SHADER_PARAMETER(uint32,						ClearSize)
+		SHADER_PARAMETER_RDG_BUFFER_UAV(RWBuffer<uint>,	BufferToClear)
 	END_SHADER_PARAMETER_STRUCT()
 };
 
@@ -133,7 +133,7 @@ class NIAGARASHADER_API FNiagaraDebugDrawLineVS : public FGlobalShader
 
 	BEGIN_SHADER_PARAMETER_STRUCT(FParameters, )
 		SHADER_PARAMETER_STRUCT_REF(FViewUniformShaderParameters, View)
-		SHADER_PARAMETER_SRV(Buffer<uint>, GpuLineBuffer)
+		SHADER_PARAMETER_RDG_BUFFER_SRV(Buffer<uint>, GpuLineBuffer)
 	END_SHADER_PARAMETER_STRUCT()
 };
 
@@ -171,44 +171,52 @@ BEGIN_SHADER_PARAMETER_STRUCT(FNiagaraDebugDrawLineParameters, )
 	SHADER_PARAMETER_STRUCT_INCLUDE(FNiagaraDebugDrawLinePS::FParameters, PSParameters)
 END_SHADER_PARAMETER_STRUCT()
 
+BEGIN_SHADER_PARAMETER_STRUCT(FNiagaraDebugDrawLineIndirectParameters, )
+	RDG_BUFFER_ACCESS(AccessIndirectDrawArgsBuffer,	ERHIAccess::IndirectArgs)
+
+	SHADER_PARAMETER_STRUCT_INCLUDE(FNiagaraDebugDrawLineVS::FParameters, VSParameters)
+	SHADER_PARAMETER_STRUCT_INCLUDE(FNiagaraDebugDrawLinePS::FParameters, PSParameters)
+END_SHADER_PARAMETER_STRUCT()
+
 IMPLEMENT_GLOBAL_SHADER(FNiagaraDebugDrawLineVS, "/Plugin/FX/Niagara/Private/NiagaraDebugDraw.usf", "MainVS", SF_Vertex);
 IMPLEMENT_GLOBAL_SHADER(FNiagaraDebugDrawLinePS, "/Plugin/FX/Niagara/Private/NiagaraDebugDraw.usf", "MainPS", SF_Pixel);
 
 //////////////////////////////////////////////////////////////////////////
 
-void NiagaraDebugShaders::ClearUAV(FRHICommandList& RHICmdList, FRHIUnorderedAccessView* UAV, FUintVector4 ClearValues, uint32 UIntsToSet)
+void NiagaraDebugShaders::ClearUAV(FRDGBuilder& GraphBuilder, FRDGBufferUAVRef UAV, FUintVector4 ClearValues, uint32 UIntsToSet)
 {
 	check(UIntsToSet > 0);
 	check(UAV != nullptr);
 
-	FNiagaraClearUAVCS::FParameters	PassParameters;
-	PassParameters.BufferToClear = UAV;
-	PassParameters.ClearValue = ClearValues;
-	PassParameters.ClearSize = UIntsToSet;
-
 	TShaderMapRef<FNiagaraClearUAVCS> ComputeShader(GetGlobalShaderMap(GMaxRHIFeatureLevel));
-	FRHIComputeShader* ShaderRHI = ComputeShader.GetComputeShader();
 	const uint32 NumThreadGroups = FMath::DivideAndRoundUp(UIntsToSet, FNiagaraClearUAVCS::ThreadGroupSize);
 
-	RHICmdList.Transition(FRHITransitionInfo(UAV, ERHIAccess::Unknown, ERHIAccess::UAVCompute));
-	SetComputePipelineState(RHICmdList, ShaderRHI);
-	SetShaderParameters(RHICmdList, ComputeShader, ShaderRHI, PassParameters);
-	RHICmdList.DispatchComputeShader(NumThreadGroups, 1, 1);
-	UnsetShaderUAVs(RHICmdList, ComputeShader, ShaderRHI);
-	RHICmdList.Transition(FRHITransitionInfo(UAV, ERHIAccess::UAVCompute, ERHIAccess::UAVCompute));
+	FNiagaraClearUAVCS::FParameters* PassParameters = GraphBuilder.AllocParameters<FNiagaraClearUAVCS::FParameters>();
+	PassParameters->BufferToClear = UAV;
+	PassParameters->ClearValue = ClearValues;
+	PassParameters->ClearSize = UIntsToSet;
+
+	FComputeShaderUtils::AddPass(
+		GraphBuilder,
+		RDG_EVENT_NAME("NiagaraDebugShaders::ClearUAV"),
+		ERDGPassFlags::Compute,
+		ComputeShader,
+		PassParameters,
+		FIntVector(NumThreadGroups, 1, 1)
+	);
 }
 
 void NiagaraDebugShaders::DrawDebugLines(
 	class FRDGBuilder& GraphBuilder, const class FViewInfo& View, FRDGTextureRef SceneColor, FRDGTextureRef SceneDepth,
-	const uint32 LineInstanceCount, const FShaderResourceViewRHIRef& LineBuffer
+	const uint32 LineInstanceCount, FRDGBufferRef LineBuffer
 )
 {
 	TShaderMapRef<FNiagaraDebugDrawLineVS> VertexShader(View.ShaderMap);
 	TShaderMapRef<FNiagaraDebugDrawLinePS> PixelShader(View.ShaderMap);
 
-	auto* PassParameters = GraphBuilder.AllocParameters<FNiagaraDebugDrawLineParameters>();
-	PassParameters->VSParameters.View = View.ViewUniformBuffer;
-	PassParameters->VSParameters.GpuLineBuffer = LineBuffer;
+	FNiagaraDebugDrawLineParameters* PassParameters = GraphBuilder.AllocParameters<FNiagaraDebugDrawLineParameters>();
+	PassParameters->VSParameters.View					= View.ViewUniformBuffer;
+	PassParameters->VSParameters.GpuLineBuffer			= GraphBuilder.CreateSRV(LineBuffer, PF_R32_UINT);
 
 	PassParameters->PSParameters.OutputInvResolution	= FVector2f(1.0f / View.UnscaledViewRect.Width(), 1.0f / View.UnconstrainedViewRect.Height());
 	PassParameters->PSParameters.OriginalViewRectMin	= FVector2f(View.ViewRect.Min);
@@ -244,26 +252,25 @@ void NiagaraDebugShaders::DrawDebugLines(
 	);
 }
 
-void NiagaraDebugShaders::DrawDebugLines(
-	class FRDGBuilder& GraphBuilder, const class FViewInfo& View, FRDGTextureRef SceneColor, FRDGTextureRef SceneDepth,
-	const FBufferRHIRef& ArgsBuffer, const FShaderResourceViewRHIRef& LineBuffer
-)
+void NiagaraDebugShaders::DrawDebugLines(FRDGBuilder& GraphBuilder, const FViewInfo& View, FRDGTextureRef SceneColor, FRDGTextureRef SceneDepth, FRDGBufferRef ArgsBuffer, FRDGBufferRef LineBuffer)
 {
 	TShaderMapRef<FNiagaraDebugDrawLineVS> VertexShader(View.ShaderMap);
 	TShaderMapRef<FNiagaraDebugDrawLinePS> PixelShader(View.ShaderMap);
 
-	auto* PassParameters = GraphBuilder.AllocParameters<FNiagaraDebugDrawLineParameters>();
-	PassParameters->VSParameters.View = View.ViewUniformBuffer;
-	PassParameters->VSParameters.GpuLineBuffer = LineBuffer;
+	FNiagaraDebugDrawLineIndirectParameters* PassParameters = GraphBuilder.AllocParameters<FNiagaraDebugDrawLineIndirectParameters>();
+	PassParameters->AccessIndirectDrawArgsBuffer		= ArgsBuffer;
+
+	PassParameters->VSParameters.View					= View.ViewUniformBuffer;
+	PassParameters->VSParameters.GpuLineBuffer			= GraphBuilder.CreateSRV(LineBuffer, PF_R32_UINT);
 
-	PassParameters->PSParameters.OutputInvResolution = FVector2f(1.0f / View.UnscaledViewRect.Width(), 1.0f / View.UnconstrainedViewRect.Height());
-	PassParameters->PSParameters.OriginalViewRectMin = FVector2f(View.ViewRect.Min);
-	PassParameters->PSParameters.OriginalViewSize = FVector2f(View.ViewRect.Width(), View.ViewRect.Height());
-	PassParameters->PSParameters.OriginalBufferInvSize = FVector2f(1.f / SceneDepth->Desc.Extent.X, 1.f / SceneDepth->Desc.Extent.Y);
-	PassParameters->PSParameters.OccludedColorScale = GNiagaraGpuComputeDebug_OccludedLineColorScale;
-	PassParameters->PSParameters.DepthTexture = SceneDepth;
-	PassParameters->PSParameters.DepthSampler = TStaticSamplerState<SF_Point, AM_Clamp, AM_Clamp, AM_Clamp>::GetRHI();
-	PassParameters->PSParameters.RenderTargets[0] = FRenderTargetBinding(SceneColor, ERenderTargetLoadAction::ELoad);
+	PassParameters->PSParameters.OutputInvResolution	= FVector2f(1.0f / View.UnscaledViewRect.Width(), 1.0f / View.UnconstrainedViewRect.Height());
+	PassParameters->PSParameters.OriginalViewRectMin	= FVector2f(View.ViewRect.Min);
+	PassParameters->PSParameters.OriginalViewSize		= FVector2f(View.ViewRect.Width(), View.ViewRect.Height());
+	PassParameters->PSParameters.OriginalBufferInvSize	= FVector2f(1.f / SceneDepth->Desc.Extent.X, 1.f / SceneDepth->Desc.Extent.Y);
+	PassParameters->PSParameters.OccludedColorScale		= GNiagaraGpuComputeDebug_OccludedLineColorScale;
+	PassParameters->PSParameters.DepthTexture			= SceneDepth;
+	PassParameters->PSParameters.DepthSampler			= TStaticSamplerState<SF_Point, AM_Clamp, AM_Clamp, AM_Clamp>::GetRHI();
+	PassParameters->PSParameters.RenderTargets[0]		= FRenderTargetBinding(SceneColor, ERenderTargetLoadAction::ELoad);
 
 	GraphBuilder.AddPass(
 		RDG_EVENT_NAME("NiagaraDrawDebugLines"),
@@ -285,7 +292,7 @@ void NiagaraDebugShaders::DrawDebugLines(
 			SetShaderParameters(RHICmdList, VertexShader, VertexShader.GetVertexShader(), PassParameters->VSParameters);
 			SetShaderParameters(RHICmdList, PixelShader, PixelShader.GetPixelShader(), PassParameters->PSParameters);
 			RHICmdList.SetViewport(ViewRect.Min.X, ViewRect.Min.Y, 0.0f, ViewRect.Max.X, ViewRect.Max.Y, 1.0f);
-			RHICmdList.DrawPrimitiveIndirect(ArgsBuffer, 0);
+			RHICmdList.DrawPrimitiveIndirect(ArgsBuffer->GetIndirectRHICallBuffer(), 0);
 		}
 	);
 }
diff --git a/Engine/Plugins/FX/Niagara/Source/NiagaraShader/Public/NiagaraDebugShaders.h b/Engine/Plugins/FX/Niagara/Source/NiagaraShader/Public/NiagaraDebugShaders.h
--- a/Engine/Plugins/FX/Niagara/Source/NiagaraShader/Public/NiagaraDebugShaders.h
+++ b/Engine/Plugins/FX/Niagara/Source/NiagaraShader/Public/NiagaraDebugShaders.h
@@ -9,16 +9,16 @@
 struct FScreenPassRenderTarget;
 namespace NiagaraDebugShaders
 {
-	NIAGARASHADER_API void ClearUAV(FRHICommandList& RHICmdList, FRHIUnorderedAccessView* UAV, FUintVector4 ClearValues, uint32 UIntsToSet);
+	NIAGARASHADER_API void ClearUAV(FRDGBuilder& GraphBuilder, FRDGBufferUAVRef UAV, FUintVector4 ClearValues, uint32 UIntsToSet);
 
 	NIAGARASHADER_API void DrawDebugLines(
 		class FRDGBuilder& GraphBuilder, const class FViewInfo& View, FRDGTextureRef SceneColor, FRDGTextureRef SceneDepth,
-		const uint32 LineInstanceCount, const FShaderResourceViewRHIRef& LineBuffer
+		const uint32 LineInstanceCount, FRDGBufferRef LineBuffer
 	);
 
 	NIAGARASHADER_API void DrawDebugLines(
 		class FRDGBuilder& GraphBuilder, const class FViewInfo& View, FRDGTextureRef SceneColor, FRDGTextureRef SceneDepth,
-		const FBufferRHIRef& ArgsBuffer, const FShaderResourceViewRHIRef& LineBuffer
+		FRDGBufferRef ArgsBuffer, FRDGBufferRef LineBuffer
 	);
 
 	NIAGARASHADER_API void VisualizeTexture(
