Allow module scripts to define validation rules

diff --git a/Engine/Plugins/FX/Niagara/Source/Niagara/Classes/NiagaraScript.h b/Engine/Plugins/FX/Niagara/Source/Niagara/Classes/NiagaraScript.h
--- a/Engine/Plugins/FX/Niagara/Source/Niagara/Classes/NiagaraScript.h
+++ b/Engine/Plugins/FX/Niagara/Source/Niagara/Classes/NiagaraScript.h
@@ -15,6 +15,7 @@
 #include "NiagaraScriptExecutionParameterStore.h"
 #include "NiagaraStackSection.h"
 #include "NiagaraParameterDefinitionsSubscriber.h"
+#include "NiagaraValidationRule.h"
 #include "NiagaraVersionedObject.h"
 #include "HAL/CriticalSection.h"
 #include "VectorVM.h"
@@ -669,7 +670,7 @@ public:
 	
 	UPROPERTY(EditAnywhere, Category = Script)
 	TArray<FNiagaraStackSection> InputSections;
-	
+
 	/** Adjusted every time ComputeVMCompilationId is called.*/
 	UPROPERTY()
 	mutable FNiagaraVMExecutableDataId LastGeneratedVMId;
@@ -751,6 +752,10 @@ public:
 
 	/** Creates a shallow transient copy of this script for compilation purposes. */
 	NIAGARA_API UNiagaraScript* CreateCompilationCopy();
+
+	/** A set of rules to apply when this script is used in the stack. To create your own rules, write a custom class that extends UNiagaraValidationRule. */
+	UPROPERTY(EditAnywhere, Category = "Validation", Instanced)
+	TArray<TObjectPtr<UNiagaraValidationRule>> ValidationRules;
 #endif
 
 	/** Workaround for emitter versioning because we used a lot of Script->GetOuter() previously. */
diff --git a/Engine/Plugins/FX/Niagara/Source/Niagara/Classes/NiagaraValidationRule.h b/Engine/Plugins/FX/Niagara/Source/Niagara/Classes/NiagaraValidationRule.h
--- a/Engine/Plugins/FX/Niagara/Source/Niagara/Classes/NiagaraValidationRule.h
+++ b/Engine/Plugins/FX/Niagara/Source/Niagara/Classes/NiagaraValidationRule.h
@@ -7,6 +7,7 @@
 #include "NiagaraValidationRule.generated.h"
 
 class FNiagaraSystemViewModel;
+class UNiagaraStackEntry;
 
 UENUM()
 enum class ENiagaraValidationSeverity
@@ -50,6 +51,15 @@ struct FNiagaraValidationResult
 	TArray<FNiagaraValidationFix> Links;
 };
 
+struct FNiagaraValidationContext
+{
+	// the view model for the whole Niagara system
+	TSharedPtr<FNiagaraSystemViewModel> ViewModel;
+
+	// if set then this stack entry requested the validation
+	UNiagaraStackEntry* Source = nullptr;
+};
+
 /**
 Base class for system validation logic. 
 These allow Niagara systems to be inspected for content validation either at save time or from a commandlet.
@@ -61,7 +71,7 @@ class NIAGARA_API UNiagaraValidationRule : public UObject
 public:
 
 #if WITH_EDITOR
-	virtual void CheckValidity(TSharedPtr<FNiagaraSystemViewModel> ViewModel, TArray<FNiagaraValidationResult>& OutResults) const;
+	virtual void CheckValidity(const FNiagaraValidationContext& Context, TArray<FNiagaraValidationResult>& OutResults) const;
 #endif
 };
 
diff --git a/Engine/Plugins/FX/Niagara/Source/Niagara/Private/NiagaraValidationRule.cpp b/Engine/Plugins/FX/Niagara/Source/Niagara/Private/NiagaraValidationRule.cpp
--- a/Engine/Plugins/FX/Niagara/Source/Niagara/Private/NiagaraValidationRule.cpp
+++ b/Engine/Plugins/FX/Niagara/Source/Niagara/Private/NiagaraValidationRule.cpp
@@ -3,7 +3,7 @@
 #include "NiagaraValidationRule.h"
 
 #if WITH_EDITOR
-void UNiagaraValidationRule::CheckValidity(TSharedPtr<FNiagaraSystemViewModel> ViewModel, TArray<FNiagaraValidationResult>& Results) const
+void UNiagaraValidationRule::CheckValidity(const FNiagaraValidationContext& Context, TArray<FNiagaraValidationResult>& Results) const
 {
 }
 #endif
\ No newline at end of file
diff --git a/Engine/Plugins/FX/Niagara/Source/NiagaraEditor/Private/NiagaraValidationRules.cpp b/Engine/Plugins/FX/Niagara/Source/NiagaraEditor/Private/NiagaraValidationRules.cpp
--- a/Engine/Plugins/FX/Niagara/Source/NiagaraEditor/Private/NiagaraValidationRules.cpp
+++ b/Engine/Plugins/FX/Niagara/Source/NiagaraEditor/Private/NiagaraValidationRules.cpp
@@ -21,6 +21,7 @@
 #include "ViewModels/NiagaraEmitterViewModel.h"
 
 #include "AssetToolsModule.h"
+#include "Materials/MaterialInterface.h"
 
 #define LOCTEXT_NAMESPACE "NiagaraValidationRules"
 
@@ -110,27 +111,78 @@ namespace NiagaraValidation
 
 // --------------------------------------------------------------------------------------------------------------------------------------------
 
-void UNiagaraValidationRule_NoWarmupTime::CheckValidity(TSharedPtr<FNiagaraSystemViewModel> ViewModel, TArray<FNiagaraValidationResult>& Results)  const
+void NiagaraValidation::ValidateAllRulesInSystem(TSharedPtr<FNiagaraSystemViewModel> SysViewModel, TFunction<void(const FNiagaraValidationResult& Result)> ResultCallback)
 {
-	UNiagaraSystem& System = ViewModel->GetSystem();
+	if (SysViewModel == nullptr)
+	{
+		return;
+	}
+
+	FNiagaraValidationContext Context;
+	Context.ViewModel = SysViewModel;
+	TArray<FNiagaraValidationResult> NiagaraValidationResults;
+	
+	UNiagaraSystem& NiagaraSystem = SysViewModel->GetSystem();
+	if (NiagaraSystem.GetEffectType())
+	{
+		// go over the validation rules in the effect type
+		for (UNiagaraValidationRule* ValidationRule : NiagaraSystem.GetEffectType()->ValidationRules)
+		{
+			if (ValidationRule)
+			{
+				ValidationRule->CheckValidity(Context, NiagaraValidationResults);
+			}
+		}
+	}
+
+	// go over the module-specific rules
+	TArray<UNiagaraStackModuleItem*> StackModuleItems =	NiagaraValidation::GetAllStackEntriesInSystem<UNiagaraStackModuleItem>(Context.ViewModel);
+	for (UNiagaraStackModuleItem* Module : StackModuleItems)
+	{
+		if (Module && Module->GetIsEnabled())
+		{
+			if (UNiagaraScript* Script = Module->GetModuleNode().FunctionScript)
+			{
+				Context.Source = Module;
+				for (UNiagaraValidationRule* ValidationRule : Script->ValidationRules)
+				{
+					if (ValidationRule)
+					{
+						ValidationRule->CheckValidity(Context, NiagaraValidationResults);
+					}
+				}
+			}
+		}
+	}
+
+	// process results
+	for (const FNiagaraValidationResult& Result : NiagaraValidationResults)
+	{
+		ResultCallback(Result);
+	}
+}
+
+void UNiagaraValidationRule_NoWarmupTime::CheckValidity(const FNiagaraValidationContext& Context, TArray<FNiagaraValidationResult>& Results)  const
+{
+	UNiagaraSystem& System = Context.ViewModel->GetSystem();
 	if (System.NeedsWarmup())
 	{
-		UNiagaraStackSystemPropertiesItem* SystemProperties = NiagaraValidation::GetStackEntry<UNiagaraStackSystemPropertiesItem>(ViewModel->GetSystemStackViewModel());
+		UNiagaraStackSystemPropertiesItem* SystemProperties = NiagaraValidation::GetStackEntry<UNiagaraStackSystemPropertiesItem>(Context.ViewModel->GetSystemStackViewModel());
 		FNiagaraValidationResult Result(ENiagaraValidationSeverity::Error, LOCTEXT("WarumupSummary", "Warmuptime > 0 is not allowed"), LOCTEXT("WarmupDescription", "Systems with the chosen effect type do not allow warmup time, as it costs too much performance.\nPlease set the warmup time to 0 in the system properties."), SystemProperties);
 		Results.Add(Result);
 	}
 }
 
-void UNiagaraValidationRule_FixedGPUBoundsSet::CheckValidity(TSharedPtr<FNiagaraSystemViewModel> ViewModel, TArray<FNiagaraValidationResult>& Results)  const
+void UNiagaraValidationRule_FixedGPUBoundsSet::CheckValidity(const FNiagaraValidationContext& Context, TArray<FNiagaraValidationResult>& Results)  const
 {
 	// if the system has fixed bounds set then it overrides the emitter settings
-	if (ViewModel->GetSystem().bFixedBounds)
+	if (Context.ViewModel->GetSystem().bFixedBounds)
 	{
 		return;
 	}
 
 	// check that all the gpu emitters have fixed bounds set
-	TArray<TSharedRef<FNiagaraEmitterHandleViewModel>> EmitterHandleViewModels = ViewModel->GetEmitterHandleViewModels();
+	TArray<TSharedRef<FNiagaraEmitterHandleViewModel>> EmitterHandleViewModels = Context.ViewModel->GetEmitterHandleViewModels();
 	for (TSharedRef<FNiagaraEmitterHandleViewModel> EmitterHandleModel : EmitterHandleViewModels)
 	{
 		FVersionedNiagaraEmitterData* EmitterData = EmitterHandleModel.Get().GetEmitterHandle()->GetEmitterData();
@@ -155,10 +207,10 @@ bool IsEnabledForMaxQualityLevel(FNiagaraPlatformSet Platforms, int32 MaxQuality
 	return false;
 }
 
-void UNiagaraValidationRule_BannedRenderers::CheckValidity(TSharedPtr<FNiagaraSystemViewModel> ViewModel, TArray<FNiagaraValidationResult>& Results)  const
+void UNiagaraValidationRule_BannedRenderers::CheckValidity(const FNiagaraValidationContext& Context, TArray<FNiagaraValidationResult>& Results)  const
 {
-	UNiagaraSystem& System = ViewModel->GetSystem();
-	TArray<TSharedRef<FNiagaraEmitterHandleViewModel>> EmitterHandleViewModels = ViewModel->GetEmitterHandleViewModels();
+	UNiagaraSystem& System = Context.ViewModel->GetSystem();
+	TArray<TSharedRef<FNiagaraEmitterHandleViewModel>> EmitterHandleViewModels = Context.ViewModel->GetEmitterHandleViewModels();
 	for (TSharedRef<FNiagaraEmitterHandleViewModel> EmitterHandleModel : EmitterHandleViewModels)
 	{
 		FVersionedNiagaraEmitterData* EmitterData = EmitterHandleModel.Get().GetEmitterHandle()->GetEmitterData();
@@ -211,11 +263,11 @@ void UNiagaraValidationRule_BannedRenderers::CheckValidity(TSharedPtr<FNiagaraSy
 	}
 } 
 
-void UNiagaraValidationRule_BannedModules::CheckValidity(TSharedPtr<FNiagaraSystemViewModel> ViewModel, TArray<FNiagaraValidationResult>& Results)  const
+void UNiagaraValidationRule_BannedModules::CheckValidity(const FNiagaraValidationContext& Context, TArray<FNiagaraValidationResult>& Results)  const
 {
-	UNiagaraSystem& System = ViewModel->GetSystem();
+	UNiagaraSystem& System = Context.ViewModel->GetSystem();
 
-	TArray<UNiagaraStackModuleItem*> StackModuleItems =	NiagaraValidation::GetAllStackEntriesInSystem<UNiagaraStackModuleItem>(ViewModel);
+	TArray<UNiagaraStackModuleItem*> StackModuleItems =	NiagaraValidation::GetAllStackEntriesInSystem<UNiagaraStackModuleItem>(Context.ViewModel);
 
 	for (UNiagaraStackModuleItem* Item : StackModuleItems)
 	{
@@ -274,17 +326,17 @@ void UNiagaraValidationRule_BannedModules::CheckValidity(TSharedPtr<FNiagaraSyst
 	}
 }
 
-void UNiagaraValidationRule_InvalidEffectType::CheckValidity(TSharedPtr<FNiagaraSystemViewModel> ViewModel, TArray<FNiagaraValidationResult>& Results)  const
+void UNiagaraValidationRule_InvalidEffectType::CheckValidity(const FNiagaraValidationContext& Context, TArray<FNiagaraValidationResult>& Results)  const
 {
-	UNiagaraStackSystemPropertiesItem* SystemProperties = NiagaraValidation::GetStackEntry<UNiagaraStackSystemPropertiesItem>(ViewModel->GetSystemStackViewModel());
+	UNiagaraStackSystemPropertiesItem* SystemProperties = NiagaraValidation::GetStackEntry<UNiagaraStackSystemPropertiesItem>(Context.ViewModel->GetSystemStackViewModel());
 	FNiagaraValidationResult Result(ENiagaraValidationSeverity::Error, LOCTEXT("InvalidEffectSummary", "Invalid Effect Type"), LOCTEXT("InvalidEffectDescription", "The effect type on this system was marked as invalid for production content and should only be used as placeholder."), SystemProperties);
 	Results.Add(Result);
 }
 
-void UNiagaraValidationRule_LWC::CheckValidity(TSharedPtr<FNiagaraSystemViewModel> ViewModel, TArray<FNiagaraValidationResult>& Results)  const
+void UNiagaraValidationRule_LWC::CheckValidity(const FNiagaraValidationContext& Context, TArray<FNiagaraValidationResult>& Results)  const
 {
 	const UNiagaraSettings* Settings = GetDefault<UNiagaraSettings>();
-	UNiagaraSystem& System = ViewModel->GetSystem();
+	UNiagaraSystem& System = Context.ViewModel->GetSystem();
 	if (!System.SupportsLargeWorldCoordinates())
 	{
 		return;
@@ -292,8 +344,8 @@ void UNiagaraValidationRule_LWC::CheckValidity(TSharedPtr<FNiagaraSystemViewMode
 
 	// gather all the modules in the system, excluding localspace emitters
 	TArray<UNiagaraStackModuleItem*> AllModules;
-	AllModules.Append(NiagaraValidation::GetStackEntries<UNiagaraStackModuleItem>(ViewModel->GetSystemStackViewModel()));
-	TArray<TSharedRef<FNiagaraEmitterHandleViewModel>> EmitterHandleViewModels = ViewModel->GetEmitterHandleViewModels();
+	AllModules.Append(NiagaraValidation::GetStackEntries<UNiagaraStackModuleItem>(Context.ViewModel->GetSystemStackViewModel()));
+	TArray<TSharedRef<FNiagaraEmitterHandleViewModel>> EmitterHandleViewModels = Context.ViewModel->GetEmitterHandleViewModels();
 	for (TSharedRef<FNiagaraEmitterHandleViewModel> EmitterHandleModel : EmitterHandleViewModels)
 	{
 		if (EmitterHandleModel->GetEmitterHandle()->GetEmitterData()->bLocalSpace == false)
@@ -355,10 +407,97 @@ void UNiagaraValidationRule_LWC::CheckValidity(TSharedPtr<FNiagaraSystemViewMode
 	}
 }
 
-void UNiagaraValidationRule_SimulationStageBudget::CheckValidity(TSharedPtr<FNiagaraSystemViewModel> ViewModel, TArray<FNiagaraValidationResult>& OutResults) const
+void UNiagaraValidationRule_NoOpaqueRenderMaterial::CheckValidity(const FNiagaraValidationContext& Context,	TArray<FNiagaraValidationResult>& Results) const
+{
+	// check that we are called from a valid module
+	UNiagaraStackModuleItem* SourceModule = Cast<UNiagaraStackModuleItem>(Context.Source);
+	if (SourceModule && SourceModule->GetIsEnabled() && SourceModule->GetEmitterViewModel())
+	{
+		auto GetCollisionTypeInput = [](const UNiagaraStackModuleItem* Module)
+		{
+			TArray<UNiagaraStackFunctionInput*> ModuleInputs;
+			Module->GetParameterInputs(ModuleInputs);
+			for (UNiagaraStackFunctionInput* Input : ModuleInputs)
+			{
+				if (Input->IsStaticParameter() && Input->GetInputParameterHandle().GetName() == FName("GPU Collision Type"))
+				{
+					return Input;
+				}
+			}
+			return static_cast<UNiagaraStackFunctionInput*>(nullptr);
+		};
+		
+		// search for the right emitter view model
+		for (const TSharedRef<FNiagaraEmitterHandleViewModel>& EmitterHandleModel : Context.ViewModel->GetEmitterHandleViewModels())
+		{
+			FVersionedNiagaraEmitterData* EmitterData = EmitterHandleModel->GetEmitterHandle()->GetEmitterData();
+			if (EmitterHandleModel->GetIsEnabled() && EmitterData->SimTarget == ENiagaraSimTarget::GPUComputeSim && EmitterHandleModel->GetEmitterViewModel() == SourceModule->GetEmitterViewModel())
+			{
+				// check if we are using depth buffer collisions
+				if (UNiagaraStackFunctionInput* Input = GetCollisionTypeInput(SourceModule))
+				{
+					// the first entry of the enum is depth buffer collision
+					// Unfortunately it's a BP enum, so we can't match the named values in C++
+					if (int32 Value = *(int32*)Input->GetLocalValueStruct()->GetStructMemory(); Value != 0)
+					{
+						continue;
+					}
+				}
+				else
+				{
+					continue;
+				}
+				
+				// check the renderers
+				TArray<UNiagaraStackRendererItem*> RendererItems = NiagaraValidation::GetStackEntries<UNiagaraStackRendererItem>(EmitterHandleModel->GetEmitterStackViewModel());
+				for (UNiagaraStackRendererItem* Renderer : RendererItems)
+				{
+					if (UNiagaraRendererProperties* RendererProperties = Renderer->GetRendererProperties())
+					{
+						TArray<UMaterialInterface*> OutMaterials;
+						RendererProperties->GetUsedMaterials(nullptr, OutMaterials);
+						for (UMaterialInterface* Material : OutMaterials)
+						{
+							if (Material->GetBlendMode() == BLEND_Opaque || Material->GetBlendMode() == BLEND_Masked)
+							{
+								FText Description = LOCTEXT("NoOpaqueRenderMaterialDescription", "This renderer uses a material with a masked or opaque blend mode, which writes to the depth buffer.\nThis will cause conflicts when the collision module also uses depth buffer collisions.");
+								FNiagaraValidationResult Result(ENiagaraValidationSeverity::Warning, FText::Format(LOCTEXT("NoOpaqueRenderMaterialSummary", "Renderer '{0}' has an opaque material"), Renderer->GetDisplayName()), Description, Renderer);
+								
+								//Add autofix to switch to distance field collisions if possible
+								static const auto* CVarGenerateMeshDistanceFields = IConsoleManager::Get().FindTConsoleVariableDataInt(TEXT("r.GenerateMeshDistanceFields"));
+								if (CVarGenerateMeshDistanceFields != nullptr && CVarGenerateMeshDistanceFields->GetValueOnGameThread() > 0)
+								{
+									FNiagaraValidationFix& DisableRendererFix = Result.Fixes.AddDefaulted_GetRef();
+									DisableRendererFix.Description = LOCTEXT("SwitchCollisionFix", "Change collision type to distance fields");
+									TWeakObjectPtr<UNiagaraStackModuleItem> WeakSourceModule = SourceModule;
+								
+									DisableRendererFix.FixDelegate = FNiagaraValidationFixDelegate::CreateLambda(
+										[WeakSourceModule, GetCollisionTypeInput]()
+										{
+											if (UNiagaraStackModuleItem* CollisionModule = WeakSourceModule.Get())
+											{
+												if (UNiagaraStackFunctionInput* Input = GetCollisionTypeInput(CollisionModule))
+												{
+													TSharedRef<FStructOnScope> ValueStruct = MakeShared<FStructOnScope>(Input->GetLocalValueStruct()->GetStruct());
+													*(int32*)ValueStruct->GetStructMemory() = 1;
+													Input->SetLocalValue(ValueStruct);
+												}
+											}
+										});
+								}
+								Results.Add(Result);
+							}
+						}
+					}
+				}
+			}
+		}
+	}
+}
+
+void UNiagaraValidationRule_SimulationStageBudget::CheckValidity(const FNiagaraValidationContext& Context, TArray<FNiagaraValidationResult>& OutResults) const
 {
-	UNiagaraSystem& System = ViewModel->GetSystem();
-	for (const TSharedRef<FNiagaraEmitterHandleViewModel>& EmitterHandleModel : ViewModel->GetEmitterHandleViewModels())
+	for (const TSharedRef<FNiagaraEmitterHandleViewModel>& EmitterHandleModel : Context.ViewModel->GetEmitterHandleViewModels())
 	{
 		// Skip disabled
 		if ( EmitterHandleModel->GetIsEnabled() == false )
diff --git a/Engine/Plugins/FX/Niagara/Source/NiagaraEditor/Private/ViewModels/Stack/NiagaraStackViewModel.cpp b/Engine/Plugins/FX/Niagara/Source/NiagaraEditor/Private/ViewModels/Stack/NiagaraStackViewModel.cpp
--- a/Engine/Plugins/FX/Niagara/Source/NiagaraEditor/Private/ViewModels/Stack/NiagaraStackViewModel.cpp
+++ b/Engine/Plugins/FX/Niagara/Source/NiagaraEditor/Private/ViewModels/Stack/NiagaraStackViewModel.cpp
@@ -7,6 +7,7 @@
 #include "NiagaraStackEditorData.h"
 #include "NiagaraSystem.h"
 #include "NiagaraSystemEditorData.h"
+#include "NiagaraValidationRules.h"
 #include "ScopedTransaction.h"
 #include "ViewModels/NiagaraEmitterHandleViewModel.h"
 #include "ViewModels/NiagaraEmitterViewModel.h"
@@ -89,32 +90,14 @@ void UNiagaraStackViewModel::UpdateStackWithValidationResults()
 	}
 	
 	TSharedPtr<FNiagaraSystemViewModel> SysViewModel = SystemViewModel.Pin();
-	if (SysViewModel == nullptr)
+	NiagaraValidation::ValidateAllRulesInSystem(SysViewModel, [](const FNiagaraValidationResult& Result)
 	{
-		return;
-	}
-
-	UNiagaraSystem& NiagaraSystem = SysViewModel->GetSystem();
-	if (NiagaraSystem.GetEffectType())
-	{
-		// go over the validation rules
-		for (UNiagaraValidationRule* ValidationRule : NiagaraSystem.GetEffectType()->ValidationRules)
+		EStackIssueSeverity Severity = (Result.Severity == ENiagaraValidationSeverity::Error) ? EStackIssueSeverity::Error : (Result.Severity == ENiagaraValidationSeverity::Warning ? EStackIssueSeverity::Warning : EStackIssueSeverity::Info);
+		if (UNiagaraStackEntry* SourceEntry = Cast<UNiagaraStackEntry>(Result.SourceObject.Get()))
 		{
-			if (ValidationRule)
-			{
-				TArray<FNiagaraValidationResult> NiagaraValidationResults;
-				ValidationRule->CheckValidity(SystemViewModel.Pin(), NiagaraValidationResults);
-				for (const FNiagaraValidationResult& Result : NiagaraValidationResults)
-				{
-					EStackIssueSeverity Severity = (Result.Severity == ENiagaraValidationSeverity::Error) ? EStackIssueSeverity::Error : (Result.Severity == ENiagaraValidationSeverity::Warning ? EStackIssueSeverity::Warning : EStackIssueSeverity::Info);
-					if (UNiagaraStackEntry* SourceEntry = Cast<UNiagaraStackEntry>(Result.SourceObject.Get()))
-					{
-						SourceEntry->AddValidationIssue(Severity, Result.SummaryText, Result.Description, Result.Severity == ENiagaraValidationSeverity::Info, Result.Fixes, Result.Links);
-					}
-				}
-			}
+			SourceEntry->AddValidationIssue(Severity, Result.SummaryText, Result.Description, Result.Severity == ENiagaraValidationSeverity::Info, Result.Fixes, Result.Links);
 		}
-	}
+	});
 }
 
 void UNiagaraStackViewModel::InitializeWithViewModels(TSharedPtr<FNiagaraSystemViewModel> InSystemViewModel, TSharedPtr<FNiagaraEmitterHandleViewModel> InEmitterHandleViewModel, FNiagaraStackViewModelOptions InOptions)
diff --git a/Engine/Plugins/FX/Niagara/Source/NiagaraEditor/Public/NiagaraValidationRules.h b/Engine/Plugins/FX/Niagara/Source/NiagaraEditor/Public/NiagaraValidationRules.h
--- a/Engine/Plugins/FX/Niagara/Source/NiagaraEditor/Public/NiagaraValidationRules.h
+++ b/Engine/Plugins/FX/Niagara/Source/NiagaraEditor/Public/NiagaraValidationRules.h
@@ -10,13 +10,18 @@
 
 class UNiagaraScript;
 
+namespace NiagaraValidation
+{
+	NIAGARAEDITOR_API void ValidateAllRulesInSystem(TSharedPtr<FNiagaraSystemViewModel> ViewModel, TFunction<void(const FNiagaraValidationResult& Result)> ResultCallback);
+}
+
 /** This validation rule ensures that no Systems have a warm up time set. */
 UCLASS(Category = "Validation", DisplayName = "No Warmup Time")
 class UNiagaraValidationRule_NoWarmupTime : public UNiagaraValidationRule
 {
 	GENERATED_BODY()
 public:
-	virtual void CheckValidity(TSharedPtr<FNiagaraSystemViewModel> ViewModel, TArray<FNiagaraValidationResult>& OutResults) const override;
+	virtual void CheckValidity(const FNiagaraValidationContext& Context, TArray<FNiagaraValidationResult>& OutResults) const override;
 };
 
 /** This validation rule ensures that GPU emitters have fixed bounds set. */
@@ -25,7 +30,7 @@ class UNiagaraValidationRule_FixedGPUBoundsSet : public UNiagaraValidationRule
 {
 	GENERATED_BODY()
 public:
-	virtual void CheckValidity(TSharedPtr<FNiagaraSystemViewModel> ViewModel, TArray<FNiagaraValidationResult>& OutResults) const override;
+	virtual void CheckValidity(const FNiagaraValidationContext& Context, TArray<FNiagaraValidationResult>& OutResults) const override;
 };
 
 /** This validation rule can ban the use of certain renderers on all or a subset of platforms. */
@@ -43,7 +48,7 @@ public:
 	UPROPERTY(EditAnywhere, Category = Validation)
 	TArray<TSubclassOf<UNiagaraRendererProperties>> BannedRenderers;
 
-	virtual void CheckValidity(TSharedPtr<FNiagaraSystemViewModel> ViewModel, TArray<FNiagaraValidationResult>& OutResults) const override;
+	virtual void CheckValidity(const FNiagaraValidationContext& Context, TArray<FNiagaraValidationResult>& OutResults) const override;
 };
 
 /** This validation rule can ban the use of certain modules on all or a subset of platforms. */
@@ -61,7 +66,7 @@ public:
 	UPROPERTY(EditAnywhere, Category = Validation)
 	TArray<TObjectPtr<UNiagaraScript>> BannedModules;
 
-	virtual void CheckValidity(TSharedPtr<FNiagaraSystemViewModel> ViewModel, TArray<FNiagaraValidationResult>& OutResults) const override;
+	virtual void CheckValidity(const FNiagaraValidationContext& Context, TArray<FNiagaraValidationResult>& OutResults) const override;
 };
 
 /** This validation rule can marks this effect type as invalid and so must be changed. Forces a choice of correct Effect Type for an System rather than. Leaving as the default. */
@@ -70,7 +75,7 @@ class UNiagaraValidationRule_InvalidEffectType : public UNiagaraValidationRule
 {
 	GENERATED_BODY()
 public:
-	virtual void CheckValidity(TSharedPtr<FNiagaraSystemViewModel> ViewModel, TArray<FNiagaraValidationResult>& OutResults) const override;
+	virtual void CheckValidity(const FNiagaraValidationContext& Context, TArray<FNiagaraValidationResult>& OutResults) const override;
 };
 
 /** This validation rule checks for various issue with Large World Coordinates. */
@@ -79,7 +84,16 @@ class UNiagaraValidationRule_LWC : public UNiagaraValidationRule
 {
 	GENERATED_BODY()
 public:
-	virtual void CheckValidity(TSharedPtr<FNiagaraSystemViewModel> ViewModel, TArray<FNiagaraValidationResult>& OutResults) const override;
+	virtual void CheckValidity(const FNiagaraValidationContext& Context, TArray<FNiagaraValidationResult>& OutResults) const override;
+};
+
+/** This validation rule is used by the collision module to check that renderers don't use any opaque or masked materials when depth buffer collisions are used. */
+UCLASS(Category = "Validation", DisplayName = "Validate GPU Depth Collision Module")
+class UNiagaraValidationRule_NoOpaqueRenderMaterial : public UNiagaraValidationRule
+{
+	GENERATED_BODY()
+public:
+	virtual void CheckValidity(const FNiagaraValidationContext& Context, TArray<FNiagaraValidationResult>& OutResults) const override;
 };
 
 /** This validation rule can be used to enforce a budget on the number of simulation stages and the iterations that may execute. */
@@ -89,7 +103,7 @@ class UNiagaraValidationRule_SimulationStageBudget : public UNiagaraValidationRu
 	GENERATED_BODY()
 
 public:
-	virtual void CheckValidity(TSharedPtr<FNiagaraSystemViewModel> ViewModel, TArray<FNiagaraValidationResult>& OutResults) const override;
+	virtual void CheckValidity(const FNiagaraValidationContext& Context, TArray<FNiagaraValidationResult>& OutResults) const override;
 
 	UPROPERTY(EditAnywhere, Category = Validation)
 	bool bMaxSimulationStagesEnabled = false;
