Migrate NiagaraDataInterfacePhysicsAsset to the NiagaraPlugin.  Maintains consistent behavior by implementing an interface on the GroomComponent.

diff --git a/Engine/Plugins/Experimental/HairModelingToolset/Source/HairModelingToolset/HairModelingToolset.Build.cs b/Engine/Plugins/Experimental/HairModelingToolset/Source/HairModelingToolset/HairModelingToolset.Build.cs
--- a/Engine/Plugins/Experimental/HairModelingToolset/Source/HairModelingToolset/HairModelingToolset.Build.cs
+++ b/Engine/Plugins/Experimental/HairModelingToolset/Source/HairModelingToolset/HairModelingToolset.Build.cs
@@ -56,6 +56,7 @@ public class HairModelingToolset : ModuleRules
 				"PhysicsCore",
 				"RenderCore",
 				"RHI",
+				"Niagara",
 
 				"ApplicationCore",
 				"Slate",
diff --git a/Engine/Plugins/FX/Niagara/Config/BaseNiagara.ini b/Engine/Plugins/FX/Niagara/Config/BaseNiagara.ini
--- a/Engine/Plugins/FX/Niagara/Config/BaseNiagara.ini
+++ b/Engine/Plugins/FX/Niagara/Config/BaseNiagara.ini
@@ -10,6 +10,7 @@
 +ClassRedirects=(OldName="NiagaraEffectRendererMeshes", NewName="/Script/Niagara.NiagaraRendererMeshes")
 +ClassRedirects=(OldName="NiagaraEffectRendererProperties", NewName="/Script/Niagara.NiagaraRendererProperties")
 +ClassRedirects=(OldName="NiagaraDataInterfaceTexture2D", NewName="/Script/Niagara.NiagaraDataInterfaceTexture")
++ClassRedirects=(OldName="NiagaraDataInterfacePhysicsAsset", NewName="/Script/Niagara.NiagaraDataInterfacePhysicsAsset")
 
 +EnumRedirects=(OldName="ENiagaraScriptUsage",ValueChanges=(("ENiagaraScriptUsage::SpawnScript","ENiagaraScriptUsage::ParticleSpawnScript"), ("ENiagaraScriptUsage::SpawnScriptInterpolated","ENiagaraScriptUsage::ParticleSpawnScriptInterpolated"),("ENiagaraScriptUsage::UpdateScript","ENiagaraScriptUsage::ParticleUpdateScript"), ("ENiagaraScriptUsage::EventScript","ENiagaraScriptUsage::ParticleEventScript"), ("ENiagaraScriptUsage::SystemScript","ENiagaraScriptUsage::SystemSpawnScript")) )
 +EnumRedirects=(OldName="/Niagara/Functions/Localspace/ENiagaraCoordinateSpace.ENiagaraCoordinateSpace",NewName="/Script/Niagara.ENiagaraCoordinateSpace",OverrideClassName="/Script/CoreUObject.Enum",ValueChanges=(("NewEnumerator0","Simulation"),("NewEnumerator1","World"),("NewEnumerator2","Local"),("NewEnumerator3","Max")))
diff --git a/Engine/Plugins/FX/Niagara/Shaders/Private/NiagaraDataInterfacePhysicsAsset.ush b/Engine/Plugins/FX/Niagara/Shaders/Private/NiagaraDataInterfacePhysicsAsset.ush
new file mode 100644
--- /dev/null
+++ b/Engine/Plugins/FX/Niagara/Shaders/Private/NiagaraDataInterfacePhysicsAsset.ush
@@ -0,0 +1,533 @@
+// Copyright Epic Games, Inc. All Rights Reserved.
+
+/*=============================================================================
+NiagaraDataInterfacePhysicsAsset.ush 
+=============================================================================*/
+
+#pragma once
+
+/* -----------------------------------------------------------------
+ * Physics Asset constants and context
+ * -----------------------------------------------------------------
+ */
+
+ #define PHYSICS_ASSET_MAX_TRANSFORMS 300
+
+ #define CURRENT_OFFSET 0
+ #define PREVIOUS_OFFSET PHYSICS_ASSET_MAX_TRANSFORMS
+ #define REST_OFFSET 2*PHYSICS_ASSET_MAX_TRANSFORMS
+
+#define MAX_DISTANCE 1e+38
+#define BOX_INDEX 0
+#define SPHERE_INDEX 1
+#define CAPSULE_INDEX 2
+
+struct FDIPhysicsAssetContext
+{
+	uint4 ElementOffsets;
+	Buffer<float4>		WorldTransformBuffer;
+	Buffer<float4>		InverseTransformBuffer;
+	Buffer<float4>		ElementExtentBuffer;
+	Buffer<uint>		PhysicsTypeBuffer;
+	float3				BoxOrigin;
+	float3				BoxExtent;
+}; 
+
+#define DIPHYSICSASSET_DECLARE_CONSTANTS(NAME)\
+uint4 ElementOffsets_##NAME;\
+Buffer<float4>		WorldTransformBuffer_##NAME;\
+Buffer<float4>		InverseTransformBuffer_##NAME;\
+Buffer<float4>		ElementExtentBuffer_##NAME;\
+Buffer<uint>		PhysicsTypeBuffer_##NAME;\
+float3 BoxOrigin_##NAME;\
+float3 BoxExtent_##NAME;\
+
+#define DIPHYSICSASSET_MAKE_CONTEXT(NAME)\
+FDIPhysicsAssetContext DIContext; \
+DIContext.ElementOffsets = ElementOffsets_##NAME;\
+DIContext.WorldTransformBuffer = WorldTransformBuffer_##NAME;\
+DIContext.InverseTransformBuffer = InverseTransformBuffer_##NAME;\
+DIContext.ElementExtentBuffer = ElementExtentBuffer_##NAME;\
+DIContext.PhysicsTypeBuffer = PhysicsTypeBuffer_##NAME;\
+DIContext.BoxOrigin = BoxOrigin_##NAME;\
+DIContext.BoxExtent = BoxExtent_##NAME;\
+
+/* -----------------------------------------------------------------
+ * Element accessors
+ * -----------------------------------------------------------------
+ */
+
+int DIPhysicsAsset_GetNumBoxes(in FDIPhysicsAssetContext DIContext)
+{
+	return DIContext.ElementOffsets[BOX_INDEX+1] - DIContext.ElementOffsets[BOX_INDEX];
+}
+
+int DIPhysicsAsset_GetNumSpheres(in FDIPhysicsAssetContext DIContext)
+{
+	return DIContext.ElementOffsets[SPHERE_INDEX+1] - DIContext.ElementOffsets[SPHERE_INDEX];
+}
+
+int DIPhysicsAsset_GetNumCapsules(in FDIPhysicsAssetContext DIContext)
+{
+	return DIContext.ElementOffsets[CAPSULE_INDEX+1] - DIContext.ElementOffsets[CAPSULE_INDEX];
+}
+
+
+/* -----------------------------------------------------------------
+ * Position / Transform utils
+ * -----------------------------------------------------------------
+ */
+
+float3 DIPhysicsAsset_GetCurrentPosition(in FDIPhysicsAssetContext DIContext, in int ElementIndex)
+{
+	const int BufferOffset = 3 * ElementIndex + CURRENT_OFFSET;
+    return float3(DIContext.WorldTransformBuffer[BufferOffset].w, DIContext.WorldTransformBuffer[BufferOffset+1].w, DIContext.WorldTransformBuffer[BufferOffset+2].w);
+}
+
+float3 DIPhysicsAsset_GetPreviousPosition(in FDIPhysicsAssetContext DIContext, in int ElementIndex)
+{
+	const int BufferOffset = 3 * ElementIndex + PREVIOUS_OFFSET;
+    return float3(DIContext.WorldTransformBuffer[BufferOffset].w, DIContext.WorldTransformBuffer[BufferOffset+1].w, DIContext.WorldTransformBuffer[BufferOffset+2].w);
+}
+
+float4 DIPhysicsAsset_GetCurrentOrientation(in FDIPhysicsAssetContext DIContext, in int ElementIndex)
+{
+	const int BufferOffset = 3 * ElementIndex + CURRENT_OFFSET;
+	const float3 OrientationMatrix[3] = {float3(DIContext.WorldTransformBuffer[BufferOffset].x, DIContext.WorldTransformBuffer[BufferOffset+1].x, DIContext.WorldTransformBuffer[BufferOffset+2].x),
+									     float3(DIContext.WorldTransformBuffer[BufferOffset].y, DIContext.WorldTransformBuffer[BufferOffset+1].y, DIContext.WorldTransformBuffer[BufferOffset+2].y),
+										 float3(DIContext.WorldTransformBuffer[BufferOffset].z, DIContext.WorldTransformBuffer[BufferOffset+1].z, DIContext.WorldTransformBuffer[BufferOffset+2].z)};
+
+
+	return QuatFromMatrix(OrientationMatrix);
+}
+
+float4 DIPhysicsAsset_GetPreviousOrientation(in FDIPhysicsAssetContext DIContext, in int ElementIndex)
+{
+	const int BufferOffset = 3 * ElementIndex + PREVIOUS_OFFSET;
+	const float3 OrientationMatrix[3] = {float3(DIContext.WorldTransformBuffer[BufferOffset].x, DIContext.WorldTransformBuffer[BufferOffset+1].x, DIContext.WorldTransformBuffer[BufferOffset+2].x),
+										 float3(DIContext.WorldTransformBuffer[BufferOffset].y, DIContext.WorldTransformBuffer[BufferOffset+1].y, DIContext.WorldTransformBuffer[BufferOffset+2].y),
+										 float3(DIContext.WorldTransformBuffer[BufferOffset].z, DIContext.WorldTransformBuffer[BufferOffset+1].z, DIContext.WorldTransformBuffer[BufferOffset+2].z)};
+
+	return QuatFromMatrix(OrientationMatrix);
+}
+
+float3x4 DIPhysicsAsset_GetCurrentTransform(in FDIPhysicsAssetContext DIContext, in int ElementIndex)
+{
+	const int BufferOffset = 3 * ElementIndex + CURRENT_OFFSET;
+    return float3x4(DIContext.WorldTransformBuffer[BufferOffset], DIContext.WorldTransformBuffer[BufferOffset+1], DIContext.WorldTransformBuffer[BufferOffset+2]);
+}
+
+float3x4 DIPhysicsAsset_GetInverseTransform(in FDIPhysicsAssetContext DIContext, in int ElementIndex)
+{
+	const int BufferOffset = 3 * ElementIndex + CURRENT_OFFSET;
+    return float3x4(DIContext.InverseTransformBuffer[BufferOffset], DIContext.InverseTransformBuffer[BufferOffset+1], DIContext.InverseTransformBuffer[BufferOffset+2]);
+}
+
+float3x4 DIPhysicsAsset_GetPreviousTransform(in FDIPhysicsAssetContext DIContext, in int ElementIndex)
+{
+	const int BufferOffset = 3 * ElementIndex + PREVIOUS_OFFSET;
+    return float3x4(DIContext.WorldTransformBuffer[BufferOffset], DIContext.WorldTransformBuffer[BufferOffset+1], DIContext.WorldTransformBuffer[BufferOffset+2]);
+}
+
+float3x4 DIPhysicsAsset_GetPreviousInverse(in FDIPhysicsAssetContext DIContext, in int ElementIndex)
+{
+	const int BufferOffset = 3 * ElementIndex + PREVIOUS_OFFSET;
+    return float3x4(DIContext.InverseTransformBuffer[BufferOffset], DIContext.InverseTransformBuffer[BufferOffset+1], DIContext.InverseTransformBuffer[BufferOffset+2]);
+}
+
+float3x4 DIPhysicsAsset_GetRestTransform(in FDIPhysicsAssetContext DIContext, in int ElementIndex)
+{
+	const int BufferOffset = 3 * ElementIndex + REST_OFFSET;
+    return float3x4(DIContext.WorldTransformBuffer[BufferOffset], DIContext.WorldTransformBuffer[BufferOffset+1], DIContext.WorldTransformBuffer[BufferOffset+2]);
+}
+
+float3x4 DIPhysicsAsset_GetRestInverse(in FDIPhysicsAssetContext DIContext, in int ElementIndex)
+{
+	const int BufferOffset = 3 * ElementIndex + REST_OFFSET;
+    return float3x4(DIContext.InverseTransformBuffer[BufferOffset], DIContext.InverseTransformBuffer[BufferOffset+1], DIContext.InverseTransformBuffer[BufferOffset+2]);
+}
+
+/* -----------------------------------------------------------------
+ * Collision detection utils
+ * -----------------------------------------------------------------
+ */
+
+// Given a world space position (NodePosition) compute the sphere closest point (position,normal,velocity)
+float DIPhysicsAsset_GetSphereProjection(in float3 LocalPosition, in float3 SphereCenter, in float SphereRadius, in int SphereIndex,
+			inout float3 OutClosestPosition, inout float3 OutClosestNormal, inout int OutElementIndex, inout float OutMinDistance)
+{
+	const float3 DeltaPosition = LocalPosition - SphereCenter;
+	const float DeltaLength = length(DeltaPosition);
+	const float SphereDistance = DeltaLength - SphereRadius;
+
+	if (SphereDistance < OutMinDistance)
+	{
+		OutMinDistance = SphereDistance;
+		OutElementIndex = SphereIndex;
+
+		OutClosestNormal = (DeltaLength > SMALL_NUMBER) ? DeltaPosition / DeltaLength : float3(0,0,0);
+		OutClosestPosition = LocalPosition - OutClosestNormal*SphereDistance;
+	}
+	return SphereDistance;
+}
+
+// Given a world space position (NodePosition) compute the sphere closest point (position,normal,velocity)
+float DIPhysicsAsset_GetCapsuleProjection(in float3 LocalPosition, in float2 CapsuleSize, in int CapsuleIndex,
+			inout float3 OutClosestPosition, inout float3 OutClosestNormal, inout int OutElementIndex, inout float OutMinDistance)
+{
+	const float HalfLength = 0.5 * CapsuleSize.y;
+	const float3 SphereCenter = float3(0,0,clamp(LocalPosition.z, -HalfLength, HalfLength));
+
+	return DIPhysicsAsset_GetSphereProjection(LocalPosition,SphereCenter,CapsuleSize.x,CapsuleIndex,
+			OutClosestPosition,OutClosestNormal,OutElementIndex,OutMinDistance);
+}
+
+// Given a world space position (NodePosition) compute the sphere closest point (position,normal,velocity)
+float DIPhysicsAsset_GetBoxProjection(in float3 LocalPosition, in float3 BoxExtent, in int BoxIndex, 
+			inout float3 OutClosestPosition, inout float3 OutClosestNormal, inout int OutElementIndex, inout float OutMinDistance)
+{
+	const float3 HalfExtent = 0.5 * BoxExtent;
+	const float3 DeltaPosition = abs(LocalPosition) - HalfExtent;
+
+	const int ClosestAxis = ((DeltaPosition.x > DeltaPosition.y) && (DeltaPosition.x > DeltaPosition.z)) ? 0 : ( DeltaPosition.y > DeltaPosition.z) ? 1 : 2;
+	const float OutsideDistance = length(max(DeltaPosition,0.0));
+	const float BoxDistance = OutsideDistance + min(DeltaPosition[ClosestAxis],0.0);
+
+	if (BoxDistance < OutMinDistance)
+	{
+		OutMinDistance = BoxDistance;
+		OutElementIndex = BoxIndex;
+
+		if (BoxDistance <= 0)
+		{
+			const bool NegativeSide = LocalPosition[ClosestAxis] < 0.0;
+
+			OutClosestPosition = LocalPosition;
+			OutClosestNormal = float3(0,0,0);
+			if( ClosestAxis == 0) 
+			{
+				OutClosestPosition.x = NegativeSide ? -HalfExtent.x : HalfExtent.x;
+				OutClosestNormal.x = NegativeSide ? -1.0 : 1.0;
+			}
+			else if( ClosestAxis == 1) 
+			{
+				OutClosestPosition.y = NegativeSide ? -HalfExtent.y : HalfExtent.y;
+				OutClosestNormal.y = NegativeSide ? -1.0 : 1.0;
+			}
+			else if( ClosestAxis == 2) 
+			{
+				OutClosestPosition.z = NegativeSide ? -HalfExtent.z : HalfExtent.z;
+				OutClosestNormal.z = NegativeSide ? -1.0 : 1.0;
+			}
+		}
+		else
+		{
+			OutClosestPosition = clamp(LocalPosition,-HalfExtent,HalfExtent);
+			OutClosestNormal = (LocalPosition - OutClosestPosition) / OutsideDistance;
+		}
+	}
+	return BoxDistance;
+}
+
+
+/* -----------------------------------------------------------------
+ * Get the closest element to the node position
+ * -----------------------------------------------------------------
+ */
+
+/*float3 DIPhysicsAsset_GetLocalPosition(in FDIPhysicsAssetContext DIContext, in float3 NodePosition, in int ElementIndex, in float TimeFraction )
+{
+	//return mul(DIPhysicsAsset_GetInverseTransform(DIContext,ElementIndex), float4(NodePosition,1.0)).xyz;
+
+	const float3 CurrentPosition = DIPhysicsAsset_GetCurrentPosition(DIContext,ElementIndex);
+	const float4 CurrentOrientation = DIPhysicsAsset_GetCurrentOrientation(DIContext,ElementIndex);
+
+	const float3 PreviousPosition = DIPhysicsAsset_GetPreviousPosition(DIContext,ElementIndex);
+	const float4 PreviousOrientation = DIPhysicsAsset_GetPreviousOrientation(DIContext,ElementIndex);
+
+	const float4 SlerpedOrientation = InverseQuat(SlerpQuat(PreviousOrientation,CurrentOrientation,TimeFraction));
+	const float3 SlerpedPosition = PreviousPosition + TimeFraction * (CurrentPosition-PreviousPosition);
+
+	return RotateVectorByQuat( NodePosition - SlerpedPosition, SlerpedOrientation );
+}*/
+
+float3 DIPhysicsAsset_GetLocalPosition(in FDIPhysicsAssetContext DIContext, in float3 NodePosition, in int ElementIndex, in float TimeFraction )
+{
+	const float3 CurrentLocal = mul(DIPhysicsAsset_GetInverseTransform(DIContext,ElementIndex), float4(NodePosition,1.0)).xyz;
+	const float3 PreviousLocal = mul(DIPhysicsAsset_GetPreviousInverse(DIContext,ElementIndex), float4(NodePosition,1.0)).xyz;
+
+	return PreviousLocal + TimeFraction * (CurrentLocal-PreviousLocal);
+}
+
+// Given a world space position (NodePosition) compute the sphere closest point (position,normal,velocity)
+int DIPhysicsAsset_GetClosestElement(in FDIPhysicsAssetContext DIContext, in float3 NodePosition, out float3 OutClosestPosition, 
+										  out float3 OutClosestNormal, out float OutMinDistance, in float TimeFraction, in bool IsRestSpace )
+{	
+	float MinDistance = MAX_DISTANCE;
+	int ElementIndex = -1;
+	float3 CollisionPosition = float3(0,0,0);
+	float3 CollisionNormal = float3(0,0,0);
+
+	const int SpheresBegin = DIContext.ElementOffsets[SPHERE_INDEX];
+	const int SpheresEnd = DIContext.ElementOffsets[SPHERE_INDEX+1];
+	for (int SphereIndex = SpheresBegin; SphereIndex < SpheresEnd; ++SphereIndex)
+	{
+		const float3 LocalPosition = IsRestSpace ? mul(DIPhysicsAsset_GetRestInverse(DIContext,SphereIndex), float4(NodePosition,1.0)).xyz : 
+						DIPhysicsAsset_GetLocalPosition(DIContext,NodePosition,SphereIndex,TimeFraction);
+
+		DIPhysicsAsset_GetSphereProjection(LocalPosition, float3(0,0,0), DIContext.ElementExtentBuffer[SphereIndex].x, SphereIndex,
+				CollisionPosition, CollisionNormal, ElementIndex, MinDistance);
+	}
+
+	const int CapsulesBegin = DIContext.ElementOffsets[CAPSULE_INDEX];
+	const int CapsulesEnd = DIContext.ElementOffsets[CAPSULE_INDEX+1];
+	for (int CapsuleIndex = CapsulesBegin; CapsuleIndex < CapsulesEnd; ++CapsuleIndex)
+	{
+		const float3 LocalPosition = IsRestSpace ? mul(DIPhysicsAsset_GetRestInverse(DIContext,CapsuleIndex), float4(NodePosition,1.0)).xyz : 
+						DIPhysicsAsset_GetLocalPosition(DIContext,NodePosition,CapsuleIndex,TimeFraction);
+
+		DIPhysicsAsset_GetCapsuleProjection(LocalPosition, DIContext.ElementExtentBuffer[CapsuleIndex].xy, CapsuleIndex,
+				CollisionPosition, CollisionNormal, ElementIndex, MinDistance);
+	}
+
+	const int BoxesBegin = DIContext.ElementOffsets[BOX_INDEX];
+	const int BoxesEnd = DIContext.ElementOffsets[BOX_INDEX+1];
+	for (int BoxIndex = BoxesBegin; BoxIndex < BoxesEnd; ++BoxIndex)
+	{
+		const float3 LocalPosition = IsRestSpace ? mul(DIPhysicsAsset_GetRestInverse(DIContext,BoxIndex), float4(NodePosition,1.0)).xyz : 
+						DIPhysicsAsset_GetLocalPosition(DIContext,NodePosition,BoxIndex,TimeFraction);
+
+		DIPhysicsAsset_GetBoxProjection(LocalPosition, DIContext.ElementExtentBuffer[BoxIndex].xyz, BoxIndex,
+				CollisionPosition, CollisionNormal, ElementIndex, MinDistance);
+	}
+	OutClosestPosition = CollisionPosition;
+	OutClosestNormal = CollisionNormal;
+	OutMinDistance = MinDistance;
+	return ElementIndex;
+}
+
+/* -----------------------------------------------------------------
+ * Get the closest point to the physics asset 
+ * -----------------------------------------------------------------
+ */
+
+// Given a world space position (NodePosition) and an element index compute the physics asset closest point (position,normal,velocity)
+void DIPhysicsAsset_GetElementPoint(in FDIPhysicsAssetContext DIContext, in float3 NodePosition, in float DeltaTime, in float TimeFraction, in int ElementIndex, out float3 OutClosestPosition, 
+										  out float3 OutClosestNormal, out float3 OutClosestVelocity)
+{	
+	OutClosestNormal = float3(0,0,0);
+	OutClosestPosition = float3(0,0,0);
+	OutClosestVelocity = float3(0,0,0);
+
+	if (ElementIndex >= 0)
+	{
+		const uint ElementIndexUint = uint(ElementIndex);
+		const float3 LocalPosition = DIPhysicsAsset_GetLocalPosition(DIContext,NodePosition,ElementIndex,TimeFraction);
+
+		float MinDistance = MAX_DISTANCE;
+		int OutputIndex = -1;
+		float3 CollisionPosition = float3(0,0,0);
+		float3 CollisionNormal = float3(0,0,0);
+
+		if ( ElementIndexUint >= DIContext.ElementOffsets[SPHERE_INDEX] && ElementIndexUint < DIContext.ElementOffsets[SPHERE_INDEX+1] )
+		{
+			DIPhysicsAsset_GetSphereProjection(LocalPosition, float3(0,0,0), DIContext.ElementExtentBuffer[ElementIndex].x, ElementIndex,
+					CollisionPosition, CollisionNormal, OutputIndex, MinDistance);
+		}
+		else if ( ElementIndexUint >= DIContext.ElementOffsets[CAPSULE_INDEX] && ElementIndexUint < DIContext.ElementOffsets[CAPSULE_INDEX+1] )
+		{
+			DIPhysicsAsset_GetCapsuleProjection(LocalPosition, DIContext.ElementExtentBuffer[ElementIndex].xy, ElementIndex,
+					CollisionPosition, CollisionNormal, OutputIndex, MinDistance);
+		}
+		else if ( ElementIndexUint >= DIContext.ElementOffsets[BOX_INDEX] && ElementIndexUint < DIContext.ElementOffsets[BOX_INDEX+1] )
+		{
+			DIPhysicsAsset_GetBoxProjection(LocalPosition, DIContext.ElementExtentBuffer[ElementIndex].xyz, ElementIndex,
+					CollisionPosition, CollisionNormal, OutputIndex, MinDistance);
+		}
+
+		if (ElementIndex != -1)
+		{
+			const float3 PreviousPosition = mul(DIPhysicsAsset_GetPreviousTransform(DIContext,ElementIndex), float4(CollisionPosition,1.0)).xyz;
+			const float3 CurrentPosition = mul(DIPhysicsAsset_GetCurrentTransform(DIContext,ElementIndex), float4(CollisionPosition,1.0)).xyz;
+
+			const float3 PreviousNormal = mul(DIPhysicsAsset_GetPreviousTransform(DIContext,ElementIndex), float4(CollisionNormal,0.0)).xyz;
+			const float3 CurrentNormal = mul(DIPhysicsAsset_GetCurrentTransform(DIContext,ElementIndex), float4(CollisionNormal,0.0)).xyz;
+
+			OutClosestPosition = PreviousPosition + TimeFraction * (CurrentPosition-PreviousPosition);
+			OutClosestNormal = normalize(PreviousNormal + TimeFraction * (CurrentNormal-PreviousNormal));
+
+			OutClosestVelocity = ( CurrentPosition - PreviousPosition ) / DeltaTime;
+		}
+	}
+}
+
+// Given a world space position (NodePosition) and an element index compute the physics asset closest distance
+void DIPhysicsAsset_GetElementDistance(in FDIPhysicsAssetContext DIContext, in float3 NodePosition, in float TimeFraction, in int ElementIndex, out float OutClosestDistance)
+{	
+	OutClosestDistance = 0.0;
+
+	if (ElementIndex >= 0)
+	{
+		const uint ElementIndexUint = uint(ElementIndex);
+		const float3 LocalPosition = DIPhysicsAsset_GetLocalPosition(DIContext,NodePosition,ElementIndex,TimeFraction);
+
+		float MinDistance = MAX_DISTANCE;
+		int OutputIndex = -1;
+		float3 CollisionPosition = float3(0,0,0);
+		float3 CollisionNormal = float3(0,0,0);
+
+		if( ElementIndexUint >= DIContext.ElementOffsets[SPHERE_INDEX] && ElementIndexUint < DIContext.ElementOffsets[SPHERE_INDEX+1] )
+		{
+			DIPhysicsAsset_GetSphereProjection(LocalPosition, float3(0,0,0), DIContext.ElementExtentBuffer[ElementIndex].x, ElementIndex,
+					CollisionPosition, CollisionNormal, OutputIndex, MinDistance);
+		}
+		else if ( ElementIndexUint >= DIContext.ElementOffsets[CAPSULE_INDEX] && ElementIndexUint < DIContext.ElementOffsets[CAPSULE_INDEX+1] )
+		{
+			DIPhysicsAsset_GetCapsuleProjection(LocalPosition, DIContext.ElementExtentBuffer[ElementIndex].xy, ElementIndex,
+					CollisionPosition, CollisionNormal, OutputIndex, MinDistance);
+		}
+		else if ( ElementIndexUint >= DIContext.ElementOffsets[BOX_INDEX] && ElementIndexUint < DIContext.ElementOffsets[BOX_INDEX+1] )
+		{
+			DIPhysicsAsset_GetBoxProjection(LocalPosition, DIContext.ElementExtentBuffer[ElementIndex].xyz, ElementIndex,
+					CollisionPosition, CollisionNormal, OutputIndex, MinDistance);
+		}
+
+		if (ElementIndex != -1)
+		{
+			OutClosestDistance = (DIContext.PhysicsTypeBuffer[ElementIndex] == 1) ? MinDistance : 0;
+		}
+	}
+}
+
+ // Given a world space position (NodePosition) compute the physics asset closest element
+void DIPhysicsAsset_GetClosestElement(in FDIPhysicsAssetContext DIContext, in float3 NodePosition, in float TimeFraction, out int OutElementIndex)
+{	
+	float3 CollisionPosition = float3(0,0,0);
+	float3 CollisionNormal = float3(0,0,0);
+	float MinDistance = 0.0;
+
+	OutElementIndex = DIPhysicsAsset_GetClosestElement(DIContext,NodePosition,CollisionPosition,CollisionNormal,MinDistance,TimeFraction, false);
+}
+
+// Given a world space position (NodePosition) compute the physics asset closest point (position,normal,velocity)
+void DIPhysicsAsset_GetClosestPoint(in FDIPhysicsAssetContext DIContext, in float3 NodePosition, in float DeltaTime, in float TimeFraction, out float3 OutClosestPosition, 
+										  out float3 OutClosestNormal, out float3 OutClosestVelocity)
+{	
+	float3 CollisionPosition = float3(0,0,0);
+	float3 CollisionNormal = float3(0,0,0);
+	float MinDistance = 0.0;
+
+	const int ElementIndex = DIPhysicsAsset_GetClosestElement(DIContext,NodePosition,CollisionPosition,CollisionNormal,MinDistance,TimeFraction, false);
+
+	OutClosestNormal = float3(0,0,0);
+	OutClosestPosition = float3(0,0,0);
+	OutClosestVelocity = float3(0,0,0);
+
+	if (ElementIndex != -1)
+	{
+		const float3 PreviousPosition = mul(DIPhysicsAsset_GetPreviousTransform(DIContext,ElementIndex), float4(CollisionPosition,1.0)).xyz;
+		const float3 CurrentPosition = mul(DIPhysicsAsset_GetCurrentTransform(DIContext,ElementIndex), float4(CollisionPosition,1.0)).xyz;
+
+		const float3 PreviousNormal = mul(DIPhysicsAsset_GetPreviousTransform(DIContext,ElementIndex), float4(CollisionNormal,0.0)).xyz;
+		const float3 CurrentNormal = mul(DIPhysicsAsset_GetCurrentTransform(DIContext,ElementIndex), float4(CollisionNormal,0.0)).xyz;
+
+		OutClosestPosition = PreviousPosition + TimeFraction * (CurrentPosition-PreviousPosition);
+		OutClosestNormal = normalize(PreviousNormal + TimeFraction * (CurrentNormal-PreviousNormal));
+
+		OutClosestVelocity = ( CurrentPosition - PreviousPosition ) / DeltaTime;
+	}
+}
+
+// Given a world space position (NodePosition) compute the physics asset closest distance
+void DIPhysicsAsset_GetClosestDistance(in FDIPhysicsAssetContext DIContext, in float3 NodePosition, in float TimeFraction, out float OutClosestDistance)
+{	
+	float3 CollisionPosition = float3(0,0,0);
+	float3 CollisionNormal = float3(0,0,0);
+	float MinDistance = 0.0;
+
+	const int ElementIndex = DIPhysicsAsset_GetClosestElement(DIContext,NodePosition,CollisionPosition,CollisionNormal,MinDistance,TimeFraction, false);
+
+	OutClosestDistance = 0.0;
+
+	if (ElementIndex != -1)
+	{
+		OutClosestDistance = (DIContext.PhysicsTypeBuffer[ElementIndex] == 1) ? MinDistance : 0;
+	}
+}
+
+
+/* -----------------------------------------------------------------
+ * Get the texture point to the physics asset 
+ * -----------------------------------------------------------------
+ */
+
+void DIPhysicsAsset_GetRestDistance(in FDIPhysicsAssetContext DIContext, in float3 RestPosition, out float OutRestDistance)
+{
+	float3 CollisionPosition = float3(0,0,0);
+	float3 CollisionNormal = float3(0,0,0);
+	float MinDistance = 0.0;
+
+	const int ElementIndex = DIPhysicsAsset_GetClosestElement(DIContext,RestPosition,CollisionPosition,CollisionNormal,MinDistance, 1.0, true);
+
+	OutRestDistance = 0.0;
+
+	if (ElementIndex != -1)
+	{
+		OutRestDistance = MinDistance;
+	}
+}
+
+void DIPhysicsAsset_GetTexturePoint(in FDIPhysicsAssetContext DIContext, in float3 NodePosition, out int OutElementIndex, out float3 OutTexturePosition)
+{
+	float3 CollisionPosition = float3(0,0,0);
+	float3 CollisionNormal = float3(0,0,0);
+	float MinDistance = 0.0;
+
+	OutTexturePosition = float3(0,0,0);
+
+	OutElementIndex = DIPhysicsAsset_GetClosestElement(DIContext,NodePosition,CollisionPosition,CollisionNormal,MinDistance, 1.0, false);
+
+	const float AxisExtent = max( max(DIContext.BoxExtent.x, DIContext.BoxExtent.y), DIContext.BoxExtent.z );
+	const float3 BoxExtent = float3(AxisExtent,AxisExtent,AxisExtent);
+
+	const float3 BoxMinimum = DIContext.BoxOrigin - BoxExtent;
+	const float3 BoxSize = 2.0 * BoxExtent;
+
+	if (OutElementIndex != -1)
+	{
+		const float3 LocalPosition = mul(DIPhysicsAsset_GetInverseTransform(DIContext,OutElementIndex), float4(NodePosition,1.0)).xyz;
+		const float3 RestPosition = mul(DIPhysicsAsset_GetRestTransform(DIContext,OutElementIndex), float4(LocalPosition,1.0)).xyz;
+
+		OutTexturePosition = clamp( (RestPosition - BoxMinimum)/BoxSize, float3(0,0,0), float3(1,1,1) );
+	}
+}
+
+void DIPhysicsAsset_GetProjectionPoint(in FDIPhysicsAssetContext DIContext, in float3 NodePosition, in float DeltaTime, in int ElementIndex,
+			in float TextureValue, in float3 TextureGradient, out float3 OutClosestPosition, out float3 OutClosestNormal, out float3 OutClosestVelocity, out float OutClosestDistance)
+{
+	OutClosestNormal = float3(0,0,0);
+	OutClosestPosition = float3(0,0,0);
+	OutClosestVelocity = float3(0,0,0);
+	OutClosestDistance = 0.0;
+
+	if (ElementIndex != -1)
+	{
+		const float AxisExtent = max( max(DIContext.BoxExtent.x, DIContext.BoxExtent.y), DIContext.BoxExtent.z );
+
+		const float3 LocalPosition = mul(DIPhysicsAsset_GetInverseTransform(DIContext,ElementIndex), float4(NodePosition,1.0)).xyz;
+		const float3 RestPosition = mul(DIPhysicsAsset_GetRestTransform(DIContext,ElementIndex), float4(LocalPosition,1.0)).xyz;
+
+		const float3 RestProjection = RestPosition - TextureValue * 2.0 * AxisExtent * TextureGradient;
+
+		const float3 CollisionPosition = mul(DIPhysicsAsset_GetRestInverse(DIContext,ElementIndex), float4(RestProjection,1.0)).xyz;
+		const float3 CollisionNormal = mul(DIPhysicsAsset_GetRestInverse(DIContext,ElementIndex), float4(TextureGradient,0.0)).xyz;
+
+		const float3 PreviousPosition = mul(DIPhysicsAsset_GetPreviousTransform(DIContext,ElementIndex), float4(CollisionPosition,1.0)).xyz;
+
+		OutClosestPosition = mul(DIPhysicsAsset_GetCurrentTransform(DIContext,ElementIndex), float4(CollisionPosition,1.0)).xyz;
+		OutClosestNormal = mul(DIPhysicsAsset_GetCurrentTransform(DIContext,ElementIndex), float4(CollisionNormal,0.0)).xyz;
+		OutClosestVelocity = ( OutClosestPosition - PreviousPosition ) / DeltaTime;
+		OutClosestDistance = TextureValue * 2.0 * AxisExtent;
+	}
+}
+
diff --git a/Engine/Plugins/FX/Niagara/Shaders/Private/NiagaraQuaternionUtils.ush b/Engine/Plugins/FX/Niagara/Shaders/Private/NiagaraQuaternionUtils.ush
new file mode 100644
--- /dev/null
+++ b/Engine/Plugins/FX/Niagara/Shaders/Private/NiagaraQuaternionUtils.ush
@@ -0,0 +1,163 @@
+// Copyright Epic Games, Inc. All Rights Reserved.
+
+#pragma once
+
+#define QUATERNION_IDENTITY float4(0,0,0,1)
+#define SMALL_NUMBER 1e-8
+
+float4 NormalizeQuat(in float4 Quat)
+{
+	float SquaredNorm = dot(Quat,Quat);
+	return (SquaredNorm >= SMALL_NUMBER) ? Quat / sqrt(SquaredNorm) : QUATERNION_IDENTITY;
+}
+
+float4 InverseQuat(in float4 Quat)
+{
+	return float4(-Quat.x,-Quat.y,-Quat.z,Quat.w);
+}
+
+float4 MultiplyQuat(in float4 QuatA, in float4 QuatB)
+{
+	return float4(
+        QuatB.xyz * QuatA.w + QuatA.xyz * QuatB.w + cross(QuatA.xyz, QuatB.xyz),
+        QuatA.w * QuatB.w - dot(QuatA.xyz, QuatB.xyz));
+}
+
+float3 RotateVectorByQuat(in float3 Vector, in float4 Quat)
+{
+	float3 T = 2.0 * cross(Quat.xyz,Vector);
+	return Vector + Quat.w * T + cross(Quat.xyz,T);
+}
+
+float3 UnrotateVectorByQuat(in float3 Vector, in float4 Quat)
+{
+	float3 T = 2.0 * cross(Quat.xyz,Vector);
+	return Vector - Quat.w * T + cross(Quat.xyz,T);
+}
+
+float4 SlerpQuat(in float4 QuatA, in float4 QuatB, float Slerp)
+{
+	// Get cosine of angle between quats.
+	const float RawCosom = 
+		    QuatA.x * QuatB.x +
+			QuatA.y * QuatB.y +
+			QuatA.z * QuatB.z +
+			QuatA.w * QuatB.w;
+	// Unaligned quats - compensate, results in taking shorter route.
+	const float Cosom = (RawCosom >= 0.0) ? RawCosom : -RawCosom;
+	
+	float Scale0, Scale1;
+
+	if( Cosom < 0.9999f )
+	{	
+		const float Omega = acos(Cosom);
+		const float InvSin = 1.f/sin(Omega);
+		Scale0 = sin( (1.f - Slerp) * Omega ) * InvSin;
+		Scale1 = sin( Slerp * Omega ) * InvSin;
+	}
+	else
+	{
+		// Use linear interpolation.
+		Scale0 = 1.0f - Slerp;
+		Scale1 = Slerp;	
+	}
+
+	// In keeping with our flipped Cosom:
+	Scale1 = (RawCosom >= 0.0) ? Scale1 : -Scale1; 
+
+	float4 Result;
+		
+	Result.x = Scale0 * QuatA.x + Scale1 * QuatB.x;
+	Result.y = Scale0 * QuatA.y + Scale1 * QuatB.y;
+	Result.z = Scale0 * QuatA.z + Scale1 * QuatB.z;
+	Result.w = Scale0 * QuatA.w + Scale1 * QuatB.w;
+
+	return NormalizeQuat(Result);
+}
+
+float4 FindQuatBetweenHelper(in float3 A, in float3 B, in float NormAB)
+{
+	float W = NormAB + dot(A,B);
+	float4 Quat = (W>1e-6*NormAB) ? float4(A.y*B.z-A.z*B.y,A.z*B.x-A.x*B.z,A.x*B.y-A.y*B.x,W) :
+			(abs(A.x) > abs(A.y)) ? float4(-A.z,0.0,A.x,0.0) : float4(0.0,-A.z,A.y,0.0);
+
+	return NormalizeQuat(Quat);
+}
+
+float4 FindQuatBetweenInternal(in float3 An, in float3 Bn)
+{
+	float3 Cn = normalize(An+Bn);
+	return float4(cross(An,Cn),dot(An,Cn));
+}
+
+float4 FindQuatBetweenNormals(in float3 NormalA, in float3 NormalB)
+{
+	float NormAB = 1.0;
+	return FindQuatBetweenHelper(NormalA,NormalB,NormAB);
+}
+
+float4 FindQuatBetweenVectors(in float3 VectorA, in float3 VectorB)
+{
+	float NormAB = sqrt(dot(VectorA,VectorA) * dot(VectorB,VectorB));
+	return FindQuatBetweenHelper(VectorA,VectorB,NormAB);
+}
+
+float4 QuatFromMatrix(in float3 R[3])
+{
+	float4 Q;
+	float s;
+	float tr = R[0][0] + R[1][1] + R[2][2];
+
+	if (tr > 0.0f) 
+	{
+		float InvS = 1.0 / sqrt(tr + 1.f);
+		s = 0.5f * InvS;
+
+		Q.w = 0.5f * (1.f / InvS);
+		Q.x = (R[1][2] - R[2][1]) * s;
+		Q.y = (R[2][0] - R[0][2]) * s;
+		Q.z = (R[0][1] - R[1][0]) * s;
+	} 
+	else 
+	{
+		uint i = 0;
+
+		if (R[1][1] > R[0][0])
+			i = 1;
+
+		if (R[2][2] > R[i][i])
+			i = 2;
+
+		const uint j = (i+1)%3;
+		const uint k = (j+1)%3;
+
+		tr = R[i][i] - R[j][j] - R[k][k];
+		float InvS = 1.0 / sqrt(tr + 1.0);
+		s = 0.5 * InvS;
+
+		const float Qi = 0.5 * (1.0 / InvS);
+		const float Qj = (R[i][j] + R[j][i]) * s;
+		const float Qk = (R[i][k] + R[k][i]) * s;
+
+		Q.w = (R[j][k] - R[k][j]) * s;
+		if(i == 0)
+		{
+			Q.x = Qi;
+			Q.y = Qj;
+			Q.z = Qk;
+		}
+		else if(i == 1)
+		{
+			Q.y = Qi;
+			Q.z = Qj;
+			Q.x = Qk;
+		}
+		else if(i == 2)
+		{
+			Q.z = Qi;
+			Q.x = Qj;
+			Q.y = Qk;
+		}
+	}
+	return NormalizeQuat(Q);
+}
diff --git a/Engine/Plugins/FX/Niagara/Source/Niagara/Private/DataInterface/NiagaraDataInterfacePhysicsAsset.cpp b/Engine/Plugins/FX/Niagara/Source/Niagara/Private/DataInterface/NiagaraDataInterfacePhysicsAsset.cpp
new file mode 100644
--- /dev/null
+++ b/Engine/Plugins/FX/Niagara/Source/Niagara/Private/DataInterface/NiagaraDataInterfacePhysicsAsset.cpp
@@ -0,0 +1,1396 @@
+// Copyright Epic Games, Inc. All Rights Reserved.
+
+#include "NiagaraDataInterfacePhysicsAsset.h"
+#include "Animation/SkeletalMeshActor.h"
+#include "Components/SkeletalMeshComponent.h"
+#include "SkeletalRenderPublic.h"
+#include "SkeletalMeshTypes.h"
+#include "AnimationRuntime.h"
+#include "NiagaraShader.h"
+#include "NiagaraComponent.h"
+#include "NiagaraRenderer.h"
+#include "NiagaraSimStageData.h"
+#include "NiagaraSystemInstance.h"
+#include "ShaderParameterUtils.h"
+
+#define LOCTEXT_NAMESPACE "NiagaraDataInterfacePhysicsAsset"
+DEFINE_LOG_CATEGORY_STATIC(LogPhysicsAsset, Log, All);
+
+//------------------------------------------------------------------------------------------------------------
+
+static const FName GetNumBoxesName(TEXT("GetNumBoxes"));
+static const FName GetNumSpheresName(TEXT("GetNumSpheres"));
+static const FName GetNumCapsulesName(TEXT("GetNumCapsules"));
+
+//------------------------------------------------------------------------------------------------------------
+
+static const FName GetClosestElementName(TEXT("GetClosestElement"));
+static const FName GetElementPointName(TEXT("GetElementPoint"));
+static const FName GetElementDistanceName(TEXT("GetElementDistance"));
+static const FName GetClosestPointName(TEXT("GetClosestPoint"));
+static const FName GetClosestDistanceName(TEXT("GetClosestDistance"));
+static const FName GetRestDistanceName(TEXT("GetRestDistance"));
+static const FName GetTexturePointName(TEXT("GetTexturePoint"));
+static const FName GetProjectionPointName(TEXT("GetProjectionPoint"));
+
+//------------------------------------------------------------------------------------------------------------
+
+const FString UNiagaraDataInterfacePhysicsAsset::ElementOffsetsName(TEXT("ElementOffsets_"));
+
+const FString UNiagaraDataInterfacePhysicsAsset::WorldTransformBufferName(TEXT("WorldTransformBuffer_"));
+const FString UNiagaraDataInterfacePhysicsAsset::InverseTransformBufferName(TEXT("InverseTransformBuffer_"));
+const FString UNiagaraDataInterfacePhysicsAsset::ElementExtentBufferName(TEXT("ElementExtentBuffer_"));
+const FString UNiagaraDataInterfacePhysicsAsset::PhysicsTypeBufferName(TEXT("PhysicsTypeBuffer_"));
+
+const FString UNiagaraDataInterfacePhysicsAsset::BoxOriginName(TEXT("BoxOrigin_"));
+const FString UNiagaraDataInterfacePhysicsAsset::BoxExtentName(TEXT("BoxExtent_"));
+
+//------------------------------------------------------------------------------------------------------------
+
+struct FNDIPhysicsAssetParametersName
+{
+	FNDIPhysicsAssetParametersName(const FString& Suffix)
+	{
+		ElementOffsetsName = UNiagaraDataInterfacePhysicsAsset::ElementOffsetsName + Suffix;
+
+		WorldTransformBufferName = UNiagaraDataInterfacePhysicsAsset::WorldTransformBufferName + Suffix;
+		InverseTransformBufferName = UNiagaraDataInterfacePhysicsAsset::InverseTransformBufferName + Suffix;
+		ElementExtentBufferName = UNiagaraDataInterfacePhysicsAsset::ElementExtentBufferName + Suffix;
+		PhysicsTypeBufferName = UNiagaraDataInterfacePhysicsAsset::PhysicsTypeBufferName + Suffix;
+
+		BoxOriginName = UNiagaraDataInterfacePhysicsAsset::BoxOriginName + Suffix;
+		BoxExtentName = UNiagaraDataInterfacePhysicsAsset::BoxExtentName + Suffix;
+	}
+
+	FString ElementOffsetsName;
+
+	FString WorldTransformBufferName;
+	FString InverseTransformBufferName;
+	FString ElementExtentBufferName;
+	FString PhysicsTypeBufferName;
+
+	FString BoxOriginName;
+	FString BoxExtentName;
+};
+
+//------------------------------------------------------------------------------------------------------------
+
+template<typename BufferType, EPixelFormat PixelFormat, uint32 ElementCount, uint32 BufferCount = 1>
+void CreateInternalBuffer(FRWBuffer& OutputBuffer)
+{
+	if (ElementCount > 0)
+	{
+		OutputBuffer.Initialize(TEXT("FNDIPhysicsAssetBuffer"), sizeof(BufferType), ElementCount * BufferCount, PixelFormat, BUF_Static);
+	}
+}
+
+template<typename BufferType, EPixelFormat PixelFormat, uint32 ElementCount, uint32 BufferCount = 1>
+void UpdateInternalBuffer(const TStaticArray<BufferType,ElementCount*BufferCount>& InputData, FRWBuffer& OutputBuffer)
+{
+	if (ElementCount > 0 && OutputBuffer.Buffer.IsValid())
+	{
+		const uint32 BufferBytes = sizeof(BufferType) * ElementCount * BufferCount;
+
+		void* OutputData = RHILockBuffer(OutputBuffer.Buffer, 0, BufferBytes, RLM_WriteOnly);
+
+		FMemory::Memcpy(OutputData, InputData.GetData(), BufferBytes);
+		RHIUnlockBuffer(OutputBuffer.Buffer);
+	}
+}
+
+void FillCurrentTransforms(const FTransform& ElementTransform, uint32& ElementCount,
+	TStaticArray<FVector4f,PHYSICS_ASSET_MAX_TRANSFORMS>& OutCurrentTransform, TStaticArray<FVector4f, PHYSICS_ASSET_MAX_TRANSFORMS>& OutCurrentInverse)
+{
+	const uint32 ElementOffset = 3 * ElementCount;
+
+	// LWC_TODO: precision loss
+	const FMatrix44f ElementMatrix = FMatrix44f(ElementTransform.ToMatrixWithScale());
+	const FMatrix44f ElementInverse = FMatrix44f(ElementMatrix.Inverse());
+
+	ElementMatrix.To3x4MatrixTranspose(&OutCurrentTransform[ElementOffset].X);
+	ElementInverse.To3x4MatrixTranspose(&OutCurrentInverse[ElementOffset].X);
+	++ElementCount;
+}
+
+void GetNumPrimitives(const TArray<TWeakObjectPtr<UPhysicsAsset>>& PhysicsAssets, const TArray<TWeakObjectPtr<USkeletalMeshComponent>>& SkeletalMeshs, uint32& NumBoxes, uint32& NumSpheres, uint32& NumCapsules)
+{
+	NumBoxes = 0;
+	NumSpheres = 0;
+	NumCapsules = 0;
+
+	for (int32 ComponentIndex = 0; ComponentIndex < SkeletalMeshs.Num(); ++ComponentIndex)
+	{
+		TWeakObjectPtr<UPhysicsAsset> PhysicsAsset = PhysicsAssets[ComponentIndex];
+		if (PhysicsAsset.IsValid() && PhysicsAsset.Get() != nullptr)
+		{
+			USkeletalMesh* SkelMesh = (SkeletalMeshs[ComponentIndex].Get() && SkeletalMeshs[ComponentIndex]->GetSkeletalMesh()) ? ToRawPtr(SkeletalMeshs[ComponentIndex]->GetSkeletalMesh()) : PhysicsAsset->GetPreviewMesh();
+			if (!SkelMesh)
+			{
+				continue;
+			}
+
+			const FReferenceSkeleton& RefSkeleton = SkelMesh->GetRefSkeleton();
+			if (RefSkeleton.GetNum() > 0)
+			{
+				for (const UBodySetup* BodySetup : PhysicsAsset->SkeletalBodySetups)
+				{
+					const FName BoneName = BodySetup->BoneName;
+					const int32 BoneIndex = RefSkeleton.FindBoneIndex(BoneName);
+					if (BoneIndex != INDEX_NONE && BoneIndex < RefSkeleton.GetNum())
+					{
+						for (const FKBoxElem& BoxElem : BodySetup->AggGeom.BoxElems)
+						{
+							if (CollisionEnabledHasPhysics(BoxElem.GetCollisionEnabled()))
+							{
+								NumBoxes += 1;
+							}
+						}
+						for (const FKSphereElem& SphereElem : BodySetup->AggGeom.SphereElems)
+						{
+							if (CollisionEnabledHasPhysics(SphereElem.GetCollisionEnabled()))
+							{
+								NumSpheres += 1;
+							}
+						}
+						for (const FKSphylElem& CapsuleElem : BodySetup->AggGeom.SphylElems)
+						{
+							if (CollisionEnabledHasPhysics(CapsuleElem.GetCollisionEnabled()))
+							{
+								NumCapsules += 1;
+							}
+						}
+					}
+				}
+			}
+		}
+	}
+}
+
+void CompactInternalArrays(FNDIPhysicsAssetArrays* OutAssetArrays)
+{
+	for (uint32 TransformIndex = 0; TransformIndex < PHYSICS_ASSET_MAX_TRANSFORMS; ++TransformIndex)
+	{
+		uint32 OffsetIndex = TransformIndex;
+		OutAssetArrays->WorldTransform[OffsetIndex] = OutAssetArrays->CurrentTransform[TransformIndex];
+		OutAssetArrays->InverseTransform[OffsetIndex] = OutAssetArrays->CurrentInverse[TransformIndex];
+
+		OffsetIndex += PHYSICS_ASSET_MAX_TRANSFORMS;
+		OutAssetArrays->WorldTransform[OffsetIndex] = OutAssetArrays->PreviousTransform[TransformIndex];
+		OutAssetArrays->InverseTransform[OffsetIndex] = OutAssetArrays->PreviousInverse[TransformIndex];
+
+		OffsetIndex += PHYSICS_ASSET_MAX_TRANSFORMS;
+		OutAssetArrays->WorldTransform[OffsetIndex] = OutAssetArrays->RestTransform[TransformIndex];
+		OutAssetArrays->InverseTransform[OffsetIndex] = OutAssetArrays->RestInverse[TransformIndex];
+	}
+}
+
+// Due to large world coordinate we compute the relative world transform in double precision
+inline FTransform ComputeRelativeTransform(const FTransform& TransformA, const FTransform& TransformB)
+{
+	const FMatrix44d TransformAMatrix = TransformA.ToMatrixWithScale();
+	const FMatrix44d TransformBMatrix = TransformB.ToMatrixWithScale();
+
+	const FMatrix44d RelativeMatrix  = TransformAMatrix * TransformBMatrix.Inverse();
+	return FTransform(RelativeMatrix);
+}
+
+void CreateInternalArrays(const TArray<TWeakObjectPtr<UPhysicsAsset>>& PhysicsAssets, const TArray<TWeakObjectPtr<USkeletalMeshComponent>>& SkeletalMeshs,
+	FNDIPhysicsAssetArrays* OutAssetArrays, const FTransform& InWorldTransform, const FTransform& LocalTransform)
+{
+	if (OutAssetArrays != nullptr)
+	{
+		OutAssetArrays->ElementOffsets.BoxOffset = 0;
+		OutAssetArrays->ElementOffsets.SphereOffset = 0;
+		OutAssetArrays->ElementOffsets.CapsuleOffset = 0;
+		OutAssetArrays->ElementOffsets.NumElements = 0;
+
+		uint32 NumBoxes = 0;
+		uint32 NumSpheres = 0;
+		uint32 NumCapsules = 0;
+
+		GetNumPrimitives(PhysicsAssets, SkeletalMeshs, NumBoxes, NumSpheres, NumCapsules);
+		
+		if ((NumBoxes + NumSpheres + NumCapsules) < PHYSICS_ASSET_MAX_PRIMITIVES)
+		{
+			OutAssetArrays->ElementOffsets.BoxOffset = 0;
+			OutAssetArrays->ElementOffsets.SphereOffset = OutAssetArrays->ElementOffsets.BoxOffset + NumBoxes;
+			OutAssetArrays->ElementOffsets.CapsuleOffset = OutAssetArrays->ElementOffsets.SphereOffset + NumSpheres;
+			OutAssetArrays->ElementOffsets.NumElements = OutAssetArrays->ElementOffsets.CapsuleOffset + NumCapsules;
+
+			const uint32 NumTransforms = OutAssetArrays->ElementOffsets.NumElements * 3;
+			const uint32 NumExtents = OutAssetArrays->ElementOffsets.NumElements;
+
+			uint32 BoxCount = OutAssetArrays->ElementOffsets.BoxOffset;
+			uint32 SphereCount = OutAssetArrays->ElementOffsets.SphereOffset;
+			uint32 CapsuleCount = OutAssetArrays->ElementOffsets.CapsuleOffset;
+
+			for (int32 ComponentIndex = 0; ComponentIndex < SkeletalMeshs.Num(); ++ComponentIndex)
+			{
+				TWeakObjectPtr<USkeletalMeshComponent> SkeletalMesh = SkeletalMeshs[ComponentIndex];
+				const bool IsSkelMeshValid = SkeletalMesh.IsValid() && SkeletalMesh.Get() != nullptr;
+				const FTransform WorldTransform = IsSkelMeshValid ? SkeletalMesh->GetComponentTransform() : InWorldTransform;
+
+				TWeakObjectPtr<UPhysicsAsset> PhysicsAsset = PhysicsAssets[ComponentIndex];
+				if (PhysicsAsset.IsValid() && PhysicsAsset.Get() != nullptr)
+				{
+					USkeletalMesh* SkelMesh = (SkeletalMeshs[ComponentIndex].Get() && SkeletalMeshs[ComponentIndex]->GetSkeletalMesh()) ? ToRawPtr(SkeletalMeshs[ComponentIndex]->GetSkeletalMesh()) : PhysicsAsset->GetPreviewMesh();
+					if (!SkelMesh)
+					{
+						continue;
+					}
+					const FReferenceSkeleton& RefSkeleton = SkelMesh->GetRefSkeleton();
+					TArray<FTransform> RestTransforms;
+					FAnimationRuntime::FillUpComponentSpaceTransforms(RefSkeleton, RefSkeleton.GetRefBonePose(), RestTransforms);
+
+					if (RefSkeleton.GetNum() > 0)
+					{
+						for (const UBodySetup* BodySetup : PhysicsAsset->SkeletalBodySetups)
+						{
+							const FName BoneName = BodySetup->BoneName;
+							const int32 BoneIndex = RefSkeleton.FindBoneIndex(BoneName);
+							if (BoneIndex != INDEX_NONE && BoneIndex < RestTransforms.Num())
+							{
+								const FTransform RestTransform = RestTransforms[BoneIndex];
+								const FTransform BoneTransform = IsSkelMeshValid ? ComputeRelativeTransform(SkeletalMesh->GetBoneTransform(BoneIndex), LocalTransform) : 
+									RestTransform * ComputeRelativeTransform(WorldTransform, LocalTransform);
+
+								for (const FKBoxElem& BoxElem : BodySetup->AggGeom.BoxElems)
+								{
+									if (CollisionEnabledHasPhysics(BoxElem.GetCollisionEnabled()))
+									{
+										const FTransform RestElement = FTransform(BoxElem.Rotation, BoxElem.Center) * RestTransform;
+										FillCurrentTransforms(RestElement, BoxCount, OutAssetArrays->RestTransform, OutAssetArrays->RestInverse);
+										--BoxCount;
+
+										OutAssetArrays->PhysicsType[BoxCount] = (BoxElem.GetCollisionEnabled() == ECollisionEnabled::QueryAndPhysics);
+
+										const FTransform ElementTransform = FTransform(BoxElem.Rotation, BoxElem.Center) * BoneTransform;
+										OutAssetArrays->ElementExtent[BoxCount] = FVector4f(BoxElem.X, BoxElem.Y, BoxElem.Z, 0);
+										FillCurrentTransforms(ElementTransform, BoxCount, OutAssetArrays->CurrentTransform, OutAssetArrays->CurrentInverse);
+									}
+								}
+
+								for (const FKSphereElem& SphereElem : BodySetup->AggGeom.SphereElems)
+								{
+									if (CollisionEnabledHasPhysics(SphereElem.GetCollisionEnabled()))
+									{
+										const FTransform RestElement = FTransform(SphereElem.Center) * RestTransform;
+										FillCurrentTransforms(RestElement, SphereCount, OutAssetArrays->RestTransform, OutAssetArrays->RestInverse);
+										--SphereCount;
+
+										OutAssetArrays->PhysicsType[SphereCount] = (SphereElem.GetCollisionEnabled() == ECollisionEnabled::QueryAndPhysics);
+
+										const FTransform ElementTransform = FTransform(SphereElem.Center) * BoneTransform;
+										OutAssetArrays->ElementExtent[SphereCount] = FVector4f(SphereElem.Radius, 0, 0, 0);
+										FillCurrentTransforms(ElementTransform, SphereCount, OutAssetArrays->CurrentTransform, OutAssetArrays->CurrentInverse);
+									}
+								}
+
+								for (const FKSphylElem& CapsuleElem : BodySetup->AggGeom.SphylElems)
+								{
+									if (CollisionEnabledHasPhysics(CapsuleElem.GetCollisionEnabled()))
+									{
+										const FTransform RestElement = FTransform(CapsuleElem.Rotation, CapsuleElem.Center) * RestTransform;
+										FillCurrentTransforms(RestElement, CapsuleCount, OutAssetArrays->RestTransform, OutAssetArrays->RestInverse);
+										--CapsuleCount;
+
+										OutAssetArrays->PhysicsType[CapsuleCount] = (CapsuleElem.GetCollisionEnabled() == ECollisionEnabled::QueryAndPhysics);
+
+										const FTransform ElementTransform = FTransform(CapsuleElem.Rotation, CapsuleElem.Center) * BoneTransform;
+										OutAssetArrays->ElementExtent[CapsuleCount] = FVector4f(CapsuleElem.Radius, CapsuleElem.Length, 0, 0);
+										FillCurrentTransforms(ElementTransform, CapsuleCount, OutAssetArrays->CurrentTransform, OutAssetArrays->CurrentInverse);
+									}
+								}
+							}
+						}
+					}
+				}
+			}
+			OutAssetArrays->PreviousTransform = OutAssetArrays->CurrentTransform;
+			OutAssetArrays->PreviousInverse = OutAssetArrays->CurrentInverse;
+
+			CompactInternalArrays(OutAssetArrays);
+		}
+		else
+		{
+			UE_LOG(LogPhysicsAsset, Warning, TEXT("Number of physics asset primitives is higher than the niagara %d limit"), PHYSICS_ASSET_MAX_PRIMITIVES);
+		}
+	}
+}
+
+void UpdateInternalArrays(const TArray<TWeakObjectPtr<UPhysicsAsset>>& PhysicsAssets, const TArray<TWeakObjectPtr<USkeletalMeshComponent>>& SkeletalMeshs,
+	FNDIPhysicsAssetArrays* OutAssetArrays, const FTransform& InWorldTransform, const FTransform& LocalTransform)
+{
+	if (OutAssetArrays != nullptr && OutAssetArrays->ElementOffsets.NumElements < PHYSICS_ASSET_MAX_PRIMITIVES)
+	{
+		uint32 NumBoxes = 0;
+		uint32 NumSpheres = 0;
+		uint32 NumCapsules = 0;
+
+		GetNumPrimitives(PhysicsAssets, SkeletalMeshs, NumBoxes, NumSpheres, NumCapsules);
+
+		if ((NumBoxes + NumSpheres + NumCapsules) < PHYSICS_ASSET_MAX_PRIMITIVES)
+		{
+			if (((OutAssetArrays->ElementOffsets.SphereOffset - OutAssetArrays->ElementOffsets.BoxOffset) != NumBoxes) || 
+				((OutAssetArrays->ElementOffsets.CapsuleOffset - OutAssetArrays->ElementOffsets.SphereOffset) != NumSpheres) ||
+				((OutAssetArrays->ElementOffsets.NumElements - OutAssetArrays->ElementOffsets.CapsuleOffset) != NumCapsules))
+			{
+				CreateInternalArrays(PhysicsAssets, SkeletalMeshs, OutAssetArrays, InWorldTransform, LocalTransform);
+			}
+
+			OutAssetArrays->PreviousTransform = OutAssetArrays->CurrentTransform;
+			OutAssetArrays->PreviousInverse = OutAssetArrays->CurrentInverse;
+
+			uint32 BoxCount = OutAssetArrays->ElementOffsets.BoxOffset;
+			uint32 SphereCount = OutAssetArrays->ElementOffsets.SphereOffset;
+			uint32 CapsuleCount = OutAssetArrays->ElementOffsets.CapsuleOffset;
+
+			for (int32 ComponentIndex = 0; ComponentIndex < SkeletalMeshs.Num(); ++ComponentIndex)
+			{
+				TWeakObjectPtr<USkeletalMeshComponent> SkeletalMesh = SkeletalMeshs[ComponentIndex];
+				const bool IsSkelMeshValid = SkeletalMesh.IsValid() && SkeletalMesh.Get() != nullptr;
+				const FTransform WorldTransform = IsSkelMeshValid ? SkeletalMesh->GetComponentTransform() : InWorldTransform;
+
+				TWeakObjectPtr<UPhysicsAsset> PhysicsAsset = PhysicsAssets[ComponentIndex];
+				if (PhysicsAsset.IsValid() && PhysicsAsset.Get() != nullptr)
+				{
+					USkeletalMesh* SkelMesh = (SkeletalMeshs[ComponentIndex].Get() && SkeletalMeshs[ComponentIndex]->GetSkeletalMesh()) ? ToRawPtr(SkeletalMeshs[ComponentIndex]->GetSkeletalMesh()) : PhysicsAsset->GetPreviewMesh();
+					if (!SkelMesh)
+					{
+						continue;
+					}
+					const FReferenceSkeleton& RefSkeleton = SkelMesh->GetRefSkeleton();
+
+					{
+						TArray<FTransform> RestTransforms;
+						FAnimationRuntime::FillUpComponentSpaceTransforms(RefSkeleton, RefSkeleton.GetRefBonePose(), RestTransforms);
+						
+						for (const UBodySetup* BodySetup : PhysicsAsset->SkeletalBodySetups)
+						{
+							const FName BoneName = BodySetup->BoneName;
+							const int32 BoneIndex = RefSkeleton.FindBoneIndex(BoneName);
+							if (BoneIndex != INDEX_NONE && BoneIndex < RestTransforms.Num())
+							{
+								const FTransform BoneTransform = IsSkelMeshValid ? ComputeRelativeTransform(SkeletalMesh->GetBoneTransform(BoneIndex), LocalTransform) : 
+									RestTransforms[BoneIndex] * ComputeRelativeTransform(WorldTransform, LocalTransform);
+
+								for (const FKBoxElem& BoxElem : BodySetup->AggGeom.BoxElems)
+								{
+									if (CollisionEnabledHasPhysics(BoxElem.GetCollisionEnabled()))
+									{
+										const FTransform ElementTransform = FTransform(BoxElem.Rotation, BoxElem.Center) * BoneTransform;
+										FillCurrentTransforms(ElementTransform, BoxCount, OutAssetArrays->CurrentTransform, OutAssetArrays->CurrentInverse);
+									}
+								}
+
+								for (const FKSphereElem& SphereElem : BodySetup->AggGeom.SphereElems)
+								{
+									if (CollisionEnabledHasPhysics(SphereElem.GetCollisionEnabled()))
+									{
+										const FTransform ElementTransform = FTransform(SphereElem.Center) * BoneTransform;
+										FillCurrentTransforms(ElementTransform, SphereCount, OutAssetArrays->CurrentTransform, OutAssetArrays->CurrentInverse);
+									}
+								}
+
+								for (const FKSphylElem& CapsuleElem : BodySetup->AggGeom.SphylElems)
+								{
+									if (CollisionEnabledHasPhysics(CapsuleElem.GetCollisionEnabled()))
+									{
+										const FTransform ElementTransform = FTransform(CapsuleElem.Rotation, CapsuleElem.Center) * BoneTransform;
+										FillCurrentTransforms(ElementTransform, CapsuleCount, OutAssetArrays->CurrentTransform, OutAssetArrays->CurrentInverse);
+									}
+								}
+							}
+						}
+					}
+				}
+			}
+		}
+		CompactInternalArrays(OutAssetArrays);
+	}
+}
+
+//------------------------------------------------------------------------------------------------------------
+
+void FNDIPhysicsAssetBuffer::InitRHI()
+{
+	CreateInternalBuffer<FVector4f, EPixelFormat::PF_A32B32G32R32F, PHYSICS_ASSET_MAX_TRANSFORMS, 3>(WorldTransformBuffer);
+	CreateInternalBuffer<FVector4f, EPixelFormat::PF_A32B32G32R32F, PHYSICS_ASSET_MAX_TRANSFORMS, 3>(InverseTransformBuffer);
+
+	CreateInternalBuffer<FVector4f, EPixelFormat::PF_A32B32G32R32F, PHYSICS_ASSET_MAX_PRIMITIVES>(ElementExtentBuffer);
+	CreateInternalBuffer<uint32, EPixelFormat::PF_R32_UINT, PHYSICS_ASSET_MAX_PRIMITIVES>(PhysicsTypeBuffer);
+}
+
+void FNDIPhysicsAssetBuffer::ReleaseRHI()
+{
+	WorldTransformBuffer.Release();
+	InverseTransformBuffer.Release();
+	ElementExtentBuffer.Release();
+	PhysicsTypeBuffer.Release();
+}
+
+//------------------------------------------------------------------------------------------------------------
+
+ETickingGroup ComputeTickingGroup(const TArray<TWeakObjectPtr<class USkeletalMeshComponent>> SkeletalMeshes)
+{
+	ETickingGroup TickingGroup = NiagaraFirstTickGroup;
+	for (int32 ComponentIndex = 0; ComponentIndex < SkeletalMeshes.Num(); ++ComponentIndex)
+	{
+		if (SkeletalMeshes[ComponentIndex].Get() != nullptr)
+		{
+			const USkeletalMeshComponent* Component = Cast<USkeletalMeshComponent>(SkeletalMeshes[ComponentIndex].Get());
+
+			const ETickingGroup ComponentTickGroup = FMath::Max(Component->PrimaryComponentTick.TickGroup, Component->PrimaryComponentTick.EndTickGroup);
+			const ETickingGroup PhysicsTickGroup = Component->bBlendPhysics ? FMath::Max(ComponentTickGroup, TG_EndPhysics) : ComponentTickGroup;
+			const ETickingGroup ClampedTickGroup = FMath::Clamp(ETickingGroup(PhysicsTickGroup + 1), NiagaraFirstTickGroup, NiagaraLastTickGroup);
+
+			TickingGroup = FMath::Max(TickingGroup, ClampedTickGroup);
+		}
+	}
+	return TickingGroup;
+}
+
+void FNDIPhysicsAssetData::Release()
+{
+	if (AssetBuffer)
+	{
+		BeginReleaseResource(AssetBuffer);
+		ENQUEUE_RENDER_COMMAND(DeleteResource)(
+			[ParamPointerToRelease = AssetBuffer](FRHICommandListImmediate& RHICmdList)
+			{
+				delete ParamPointerToRelease;
+			});
+		AssetBuffer = nullptr;
+	}
+}
+
+void FNDIPhysicsAssetData::Init(UNiagaraDataInterfacePhysicsAsset* Interface, FNiagaraSystemInstance* SystemInstance)
+{
+	AssetBuffer = nullptr;
+
+	if (Interface != nullptr && SystemInstance != nullptr)
+	{
+		FTransform BoneTransform = FTransform::Identity;
+		Interface->ExtractSourceComponent(SystemInstance, BoneTransform);
+		TickingGroup = ComputeTickingGroup(Interface->SourceComponents);
+
+		if(0 < Interface->PhysicsAssets.Num() && Interface->PhysicsAssets[0].IsValid() && Interface->PhysicsAssets[0].Get() != nullptr && 
+			Interface->PhysicsAssets.Num() == Interface->SourceComponents.Num() )
+		{
+			CreateInternalArrays(Interface->PhysicsAssets, Interface->SourceComponents, &AssetArrays, SystemInstance->GetWorldTransform(), BoneTransform);
+		}
+
+		AssetBuffer = new FNDIPhysicsAssetBuffer();
+		BeginInitResource(AssetBuffer);
+
+		FVector BoxMax(-FLT_MAX, -FLT_MAX, -FLT_MAX);
+		FVector BoxMin = -BoxMax;
+
+		FBox BoundingBox(BoxMin, BoxMax);
+		for (int32 ComponentIndex = 0; ComponentIndex < Interface->PhysicsAssets.Num(); ++ComponentIndex)
+		{
+			TWeakObjectPtr<UPhysicsAsset> PhysicsAsset = Interface->PhysicsAssets[ComponentIndex];
+			if (PhysicsAsset.IsValid() && PhysicsAsset.Get() != nullptr)
+			{
+				USkeletalMesh* SkelMesh = (Interface->SourceComponents[ComponentIndex].IsValid() && Interface->SourceComponents[ComponentIndex].Get() && Interface->SourceComponents[ComponentIndex]->GetSkeletalMesh()) ?
+					ToRawPtr(Interface->SourceComponents[ComponentIndex]->GetSkeletalMesh()) : PhysicsAsset->GetPreviewMesh();
+				if (SkelMesh)
+				{
+					BoundingBox += SkelMesh->GetImportedBounds().GetBox();
+				}
+			}
+		}
+		BoxOrigin = 0.5 * (BoundingBox.Max + BoundingBox.Min);
+		BoxExtent = (BoundingBox.Max - BoundingBox.Min);
+	}
+}
+
+void FNDIPhysicsAssetData::Update(UNiagaraDataInterfacePhysicsAsset* Interface, FNiagaraSystemInstance* SystemInstance)
+{
+	if (Interface != nullptr && SystemInstance != nullptr)
+	{
+		FTransform BoneTransform = FTransform::Identity;
+		Interface->ExtractSourceComponent(SystemInstance, BoneTransform);
+		TickingGroup = ComputeTickingGroup(Interface->SourceComponents);
+
+		if (0 < Interface->PhysicsAssets.Num() && Interface->PhysicsAssets[0].IsValid() && Interface->PhysicsAssets[0].Get() != nullptr &&
+			Interface->PhysicsAssets.Num() == Interface->SourceComponents.Num())
+		{
+			UpdateInternalArrays(Interface->PhysicsAssets, Interface->SourceComponents, &AssetArrays, SystemInstance->GetWorldTransform(), BoneTransform);
+		}
+	}
+}
+
+//------------------------------------------------------------------------------------------------------------
+
+struct FNDIPhysicsAssetParametersCS : public FNiagaraDataInterfaceParametersCS
+{
+	DECLARE_TYPE_LAYOUT(FNDIPhysicsAssetParametersCS, NonVirtual);
+public:
+	void Bind(const FNiagaraDataInterfaceGPUParamInfo& ParameterInfo, const class FShaderParameterMap& ParameterMap)
+	{
+		FNDIPhysicsAssetParametersName ParamNames(*ParameterInfo.DataInterfaceHLSLSymbol);
+
+		ElementOffsets.Bind(ParameterMap, *ParamNames.ElementOffsetsName);
+
+		WorldTransformBuffer.Bind(ParameterMap, *ParamNames.WorldTransformBufferName);
+		InverseTransformBuffer.Bind(ParameterMap, *ParamNames.InverseTransformBufferName);
+		ElementExtentBuffer.Bind(ParameterMap, *ParamNames.ElementExtentBufferName);
+		PhysicsTypeBuffer.Bind(ParameterMap, *ParamNames.PhysicsTypeBufferName);
+
+		BoxOrigin.Bind(ParameterMap, *ParamNames.BoxOriginName);
+		BoxExtent.Bind(ParameterMap, *ParamNames.BoxExtentName);
+
+		
+	}
+
+	void Set(FRHICommandList& RHICmdList, const FNiagaraDataInterfaceSetArgs& Context) const
+	{
+		check(IsInRenderingThread());
+
+		FRHIComputeShader* ComputeShaderRHI = RHICmdList.GetBoundComputeShader();
+
+		FNDIPhysicsAssetProxy* InterfaceProxy =
+			static_cast<FNDIPhysicsAssetProxy*>(Context.DataInterface);
+		FNDIPhysicsAssetData* ProxyData =
+			InterfaceProxy->SystemInstancesToProxyData.Find(Context.SystemInstanceID);
+
+		if (ProxyData != nullptr && ProxyData->AssetBuffer && ProxyData->AssetBuffer->IsInitialized())
+		{
+			FNDIPhysicsAssetBuffer* AssetBuffer = ProxyData->AssetBuffer;
+
+			FRHITransitionInfo Transitions[] = {
+				FRHITransitionInfo(AssetBuffer->WorldTransformBuffer.UAV, ERHIAccess::Unknown, ERHIAccess::SRVCompute),
+				FRHITransitionInfo(AssetBuffer->InverseTransformBuffer.UAV, ERHIAccess::Unknown, ERHIAccess::SRVCompute),
+				FRHITransitionInfo(AssetBuffer->ElementExtentBuffer.UAV, ERHIAccess::Unknown, ERHIAccess::SRVCompute),
+				FRHITransitionInfo(AssetBuffer->PhysicsTypeBuffer.UAV, ERHIAccess::Unknown, ERHIAccess::SRVCompute)
+			};
+			RHICmdList.Transition(MakeArrayView(Transitions, UE_ARRAY_COUNT(Transitions)));
+
+			SetSRVParameter(RHICmdList, ComputeShaderRHI, WorldTransformBuffer, AssetBuffer->WorldTransformBuffer.SRV);
+			SetSRVParameter(RHICmdList, ComputeShaderRHI, InverseTransformBuffer, AssetBuffer->InverseTransformBuffer.SRV);
+			SetSRVParameter(RHICmdList, ComputeShaderRHI, ElementExtentBuffer, AssetBuffer->ElementExtentBuffer.SRV);
+			SetSRVParameter(RHICmdList, ComputeShaderRHI, PhysicsTypeBuffer, AssetBuffer->PhysicsTypeBuffer.SRV);
+
+			SetShaderValue(RHICmdList, ComputeShaderRHI, ElementOffsets, ProxyData->AssetArrays.ElementOffsets);
+			SetShaderValue(RHICmdList, ComputeShaderRHI, BoxOrigin, (FVector3f)ProxyData->BoxOrigin);
+			SetShaderValue(RHICmdList, ComputeShaderRHI, BoxExtent, (FVector3f)ProxyData->BoxExtent);
+		}
+		else
+		{
+			SetSRVParameter(RHICmdList, ComputeShaderRHI, WorldTransformBuffer, FNiagaraRenderer::GetDummyFloatBuffer());
+			SetSRVParameter(RHICmdList, ComputeShaderRHI, InverseTransformBuffer, FNiagaraRenderer::GetDummyFloatBuffer());
+			SetSRVParameter(RHICmdList, ComputeShaderRHI, ElementExtentBuffer, FNiagaraRenderer::GetDummyFloatBuffer());
+			SetSRVParameter(RHICmdList, ComputeShaderRHI, PhysicsTypeBuffer, FNiagaraRenderer::GetDummyIntBuffer());
+
+			static const FElementOffset DummyOffsets(0, 0, 0, 0);
+			SetShaderValue(RHICmdList, ComputeShaderRHI, ElementOffsets, DummyOffsets);
+			SetShaderValue(RHICmdList, ComputeShaderRHI, BoxOrigin, FVector3f::ZeroVector);
+			SetShaderValue(RHICmdList, ComputeShaderRHI, BoxExtent, FVector3f::ZeroVector);
+		}
+	}
+
+	void Unset(FRHICommandList& RHICmdList, const FNiagaraDataInterfaceSetArgs& Context) const
+	{
+	}
+
+private:
+
+	LAYOUT_FIELD(FShaderParameter, ElementOffsets);
+
+	LAYOUT_FIELD(FShaderResourceParameter, WorldTransformBuffer);
+	LAYOUT_FIELD(FShaderResourceParameter, InverseTransformBuffer);
+	LAYOUT_FIELD(FShaderResourceParameter, ElementExtentBuffer);
+	LAYOUT_FIELD(FShaderResourceParameter, PhysicsTypeBuffer);
+
+	LAYOUT_FIELD(FShaderParameter, BoxOrigin);
+	LAYOUT_FIELD(FShaderParameter, BoxExtent);
+};
+
+IMPLEMENT_TYPE_LAYOUT(FNDIPhysicsAssetParametersCS);
+
+IMPLEMENT_NIAGARA_DI_PARAMETER(UNiagaraDataInterfacePhysicsAsset, FNDIPhysicsAssetParametersCS);
+
+
+//------------------------------------------------------------------------------------------------------------
+
+void FNDIPhysicsAssetProxy::ConsumePerInstanceDataFromGameThread(void* PerInstanceData, const FNiagaraSystemInstanceID& Instance)
+{
+	FNDIPhysicsAssetData* SourceData = static_cast<FNDIPhysicsAssetData*>(PerInstanceData);
+	FNDIPhysicsAssetData* TargetData = &(SystemInstancesToProxyData.FindOrAdd(Instance));
+
+	ensure(TargetData);
+	if (TargetData)
+	{
+		TargetData->AssetBuffer = SourceData->AssetBuffer;
+		TargetData->BoxOrigin = SourceData->BoxOrigin;
+		TargetData->BoxExtent = SourceData->BoxExtent;
+		TargetData->AssetArrays = SourceData->AssetArrays;
+		TargetData->TickingGroup = SourceData->TickingGroup;
+	}
+	else
+	{
+		UE_LOG(LogPhysicsAsset, Log, TEXT("ConsumePerInstanceDataFromGameThread() ... could not find %d"), Instance);
+	}
+	SourceData->~FNDIPhysicsAssetData();
+}
+
+void FNDIPhysicsAssetProxy::InitializePerInstanceData(const FNiagaraSystemInstanceID& SystemInstance)
+{
+	check(IsInRenderingThread());
+
+	FNDIPhysicsAssetData* TargetData = SystemInstancesToProxyData.Find(SystemInstance);
+	TargetData = &SystemInstancesToProxyData.Add(SystemInstance);
+}
+
+void FNDIPhysicsAssetProxy::DestroyPerInstanceData(const FNiagaraSystemInstanceID& SystemInstance)
+{
+	check(IsInRenderingThread());
+	SystemInstancesToProxyData.Remove(SystemInstance);
+}
+
+void FNDIPhysicsAssetProxy::PreStage(FRHICommandList& RHICmdList, const FNiagaraDataInterfaceStageArgs& Context)
+{
+	FNDIPhysicsAssetData* ProxyData =
+		SystemInstancesToProxyData.Find(Context.SystemInstanceID);
+
+	if (ProxyData != nullptr && ProxyData->AssetBuffer)
+	{
+		if (Context.SimStageData->bFirstStage)
+		{
+			FRHITransitionInfo Transitions[] = {
+				FRHITransitionInfo(ProxyData->AssetBuffer->WorldTransformBuffer.UAV, ERHIAccess::Unknown, ERHIAccess::UAVCompute),
+				FRHITransitionInfo(ProxyData->AssetBuffer->InverseTransformBuffer.UAV, ERHIAccess::Unknown, ERHIAccess::UAVCompute),
+				FRHITransitionInfo(ProxyData->AssetBuffer->PhysicsTypeBuffer.UAV, ERHIAccess::Unknown, ERHIAccess::UAVCompute),
+				FRHITransitionInfo(ProxyData->AssetBuffer->ElementExtentBuffer.UAV, ERHIAccess::Unknown, ERHIAccess::UAVCompute)
+			};
+			RHICmdList.Transition(MakeArrayView(Transitions, UE_ARRAY_COUNT(Transitions)));
+
+			UpdateInternalBuffer<FVector4f, EPixelFormat::PF_A32B32G32R32F, PHYSICS_ASSET_MAX_TRANSFORMS, 3>(ProxyData->AssetArrays.WorldTransform, ProxyData->AssetBuffer->WorldTransformBuffer);
+			UpdateInternalBuffer<FVector4f, EPixelFormat::PF_A32B32G32R32F, PHYSICS_ASSET_MAX_TRANSFORMS, 3>(ProxyData->AssetArrays.InverseTransform, ProxyData->AssetBuffer->InverseTransformBuffer);
+			UpdateInternalBuffer<FVector4f, EPixelFormat::PF_A32B32G32R32F, PHYSICS_ASSET_MAX_PRIMITIVES>(ProxyData->AssetArrays.ElementExtent, ProxyData->AssetBuffer->ElementExtentBuffer);
+			UpdateInternalBuffer<uint32, EPixelFormat::PF_R32_UINT, PHYSICS_ASSET_MAX_PRIMITIVES>(ProxyData->AssetArrays.PhysicsType, ProxyData->AssetBuffer->PhysicsTypeBuffer);
+		}
+	}
+}
+
+//------------------------------------------------------------------------------------------------------------
+
+UNiagaraDataInterfacePhysicsAsset::UNiagaraDataInterfacePhysicsAsset(FObjectInitializer const& ObjectInitializer)
+	: Super(ObjectInitializer)
+	, DefaultSource(nullptr)
+	, SoftSourceActor(nullptr)
+	, SourceComponents()
+	, PhysicsAssets()
+{
+	FNiagaraTypeDefinition Def(UObject::StaticClass());
+	MeshUserParameter.Parameter.SetType(Def);
+
+	Proxy.Reset(new FNDIPhysicsAssetProxy());
+}
+
+void UNiagaraDataInterfacePhysicsAsset::ExtractSourceComponent(FNiagaraSystemInstance* SystemInstance, FTransform& BoneTransform)
+{
+	// Helper to scour an actor (or its parents) for a valid skeletal mesh component
+	auto FindActorSkelMeshComponent = [](AActor* Actor, bool bRecurseParents = false) -> USkeletalMeshComponent*
+	{
+		if (ASkeletalMeshActor* SkelMeshActor = Cast<ASkeletalMeshActor>(Actor))
+		{
+			USkeletalMeshComponent* Comp = SkelMeshActor->GetSkeletalMeshComponent();
+			if (IsValid(Comp))
+			{
+				return Comp;
+			}
+		}
+
+		// Fall back on any valid component on the actor
+		while (Actor)
+		{
+			for (UActorComponent* ActorComp : Actor->GetComponents())
+			{
+				USkeletalMeshComponent* Comp = Cast<USkeletalMeshComponent>(ActorComp);
+				if (IsValid(Comp) && Comp->GetSkeletalMesh() != nullptr)
+				{
+					return Comp;
+				}
+			}
+
+			if (bRecurseParents)
+			{
+				Actor = Actor->GetParentActor();
+			}
+			else
+			{
+				break;
+			}
+		}
+
+		return nullptr;
+	};
+		
+	BoneTransform = FTransform::Identity;
+
+	// Track down the source component
+	TWeakObjectPtr<USkeletalMeshComponent> SourceComponent;
+	
+	
+	if (MeshUserParameter.Parameter.IsValid() && SystemInstance != nullptr)
+	{
+		// Initialize the binding and retrieve the object. If a valid object is bound, we'll try and retrieve the SkelMesh component from it.
+		// If it's not valid yet, we'll reset and do this again when/if a valid object is set on the binding
+		FNiagaraParameterDirectBinding<UObject*> UserParamBinding;
+		UObject* UserParamObject = UserParamBinding.Init(SystemInstance->GetInstanceParameters(), MeshUserParameter.Parameter);
+
+		if (UserParamObject)
+		{
+			if (USkeletalMeshComponent* UserSkelMeshComp = Cast<USkeletalMeshComponent>(UserParamObject))
+			{
+				if (IsValid(UserSkelMeshComp))
+				{
+					SourceComponent = UserSkelMeshComp;
+				}
+			}
+			else if (AActor* Actor = Cast<AActor>(UserParamObject))
+			{
+				SourceComponent = FindActorSkelMeshComponent(Actor);
+			}
+			else
+			{
+				//We have a valid, non-null UObject parameter type but it is not a type we can use to get a skeletal mesh from.
+				UE_LOG(LogPhysicsAsset, Warning, TEXT("SkeletalMesh data interface using object parameter with invalid type. Skeletal Mesh Data Interfaces can only get a valid mesh from SkeletalMeshComponents, SkeletalMeshActors or Actors."));
+				UE_LOG(LogPhysicsAsset, Warning, TEXT("Invalid Parameter : %s"), *UserParamObject->GetFullName());
+				UE_LOG(LogPhysicsAsset, Warning, TEXT("System : %s"), *GetFullNameSafe(SystemInstance->GetSystem()));
+			}
+		}
+		else
+		{
+			// The binding exists, but no object is bound. Not warning here in case the user knows what they're doing.
+		}
+	}
+	else if (AActor* SourceActor = SoftSourceActor.Get())
+	{
+		ASkeletalMeshActor* SkeletalMeshActor = Cast<ASkeletalMeshActor>(SourceActor);
+		if (SkeletalMeshActor != nullptr)
+		{
+			SourceComponent = SkeletalMeshActor->GetSkeletalMeshComponent();
+		}
+		else
+		{
+			SourceComponent = SourceActor->FindComponentByClass<USkeletalMeshComponent>();
+		}
+	}
+	else if (USceneComponent* AttachComponent = SystemInstance->GetAttachComponent())
+	{
+		// Try to find the first component by walking the attachment hierarchy
+		for (USceneComponent* Curr = AttachComponent; Curr; Curr = Curr->GetAttachParent())
+		{
+			USkeletalMeshComponent* SkelMeshComp = Cast<USkeletalMeshComponent>(Curr);
+			if (SkelMeshComp && SkelMeshComp->GetSkeletalMesh())
+			{
+				SourceComponent = SkelMeshComp;
+				break;
+			}
+		}
+
+		if (!SourceComponent.IsValid())
+		{
+			// Fall back on the attach component's outer chain if we aren't attached to the skeletal mesh 
+			if (USkeletalMeshComponent* OuterComp = AttachComponent->GetTypedOuter<USkeletalMeshComponent>())
+			{
+				SourceComponent = OuterComp;
+			}
+		}
+	}
+	
+	SourceComponents.Empty();
+	PhysicsAssets.Empty();
+	
+	// Try to find the groom physics asset by walking the attachment hierarchy
+	UPhysicsAsset* PhysicsAsset = DefaultSource;
+	for (USceneComponent* Curr = SystemInstance->GetAttachComponent(); Curr; Curr = Curr->GetAttachParent())
+	{
+		if (INiagaraPhysicsAssetDICollectorInterface* RetrieverInterface = Cast<INiagaraPhysicsAssetDICollectorInterface>(Curr))
+		{
+			PhysicsAsset = RetrieverInterface->BuildAndCollect(BoneTransform, SourceComponents, PhysicsAssets);;
+			break;
+		}
+	}
+	
+	if (SourceComponent != nullptr)
+	{
+		if (PhysicsAsset || SourceComponent->GetPhysicsAsset())
+		{
+			SourceComponents.Add(SourceComponent);
+			if (PhysicsAsset)
+			{
+				PhysicsAssets.Add(PhysicsAsset);
+			}
+			else
+			{
+				PhysicsAssets.Add(SourceComponent->GetPhysicsAsset());
+			}
+		}
+	}
+	else if (PhysicsAsset != nullptr)
+	{
+		SourceComponents.Add(nullptr);
+		PhysicsAssets.Add(PhysicsAsset);
+	}
+}
+
+bool UNiagaraDataInterfacePhysicsAsset::InitPerInstanceData(void* PerInstanceData, FNiagaraSystemInstance* SystemInstance)
+{
+	FNDIPhysicsAssetData* InstanceData = new (PerInstanceData) FNDIPhysicsAssetData();
+
+	check(InstanceData);
+	InstanceData->Init(this, SystemInstance);
+
+	return true;
+}
+
+ETickingGroup UNiagaraDataInterfacePhysicsAsset::CalculateTickGroup(const void* PerInstanceData) const
+{
+	const FNDIPhysicsAssetData* InstanceData = static_cast<const FNDIPhysicsAssetData*>(PerInstanceData);
+
+	if (InstanceData)
+	{
+		return InstanceData->TickingGroup;
+	}
+	return NiagaraFirstTickGroup;
+}
+
+void UNiagaraDataInterfacePhysicsAsset::DestroyPerInstanceData(void* PerInstanceData, FNiagaraSystemInstance* SystemInstance)
+{
+	FNDIPhysicsAssetData* InstanceData = static_cast<FNDIPhysicsAssetData*>(PerInstanceData);
+
+	InstanceData->Release();
+	InstanceData->~FNDIPhysicsAssetData();
+
+	FNDIPhysicsAssetProxy* ThisProxy = GetProxyAs<FNDIPhysicsAssetProxy>();
+	ENQUEUE_RENDER_COMMAND(FNiagaraDIDestroyInstanceData) (
+		[ThisProxy, InstanceID = SystemInstance->GetId()](FRHICommandListImmediate& CmdList)
+		{
+			ThisProxy->SystemInstancesToProxyData.Remove(InstanceID);
+		}
+	);
+}
+
+bool UNiagaraDataInterfacePhysicsAsset::PerInstanceTick(void* PerInstanceData, FNiagaraSystemInstance* SystemInstance, float InDeltaSeconds)
+{
+	FNDIPhysicsAssetData* InstanceData = static_cast<FNDIPhysicsAssetData*>(PerInstanceData);
+	if (InstanceData && InstanceData->AssetBuffer && SystemInstance)
+	{
+		InstanceData->Update(this, SystemInstance);
+	}
+	return false;
+}
+
+bool UNiagaraDataInterfacePhysicsAsset::CopyToInternal(UNiagaraDataInterface* Destination) const
+{
+	if (!Super::CopyToInternal(Destination))
+	{
+		return false;
+	}
+
+	UNiagaraDataInterfacePhysicsAsset* OtherTyped = CastChecked<UNiagaraDataInterfacePhysicsAsset>(Destination);
+	OtherTyped->PhysicsAssets = PhysicsAssets;
+	OtherTyped->SoftSourceActor = SoftSourceActor;
+	OtherTyped->SourceComponents = SourceComponents;
+	OtherTyped->DefaultSource = DefaultSource;
+	OtherTyped->MeshUserParameter = MeshUserParameter;
+
+	return true;
+}
+
+bool UNiagaraDataInterfacePhysicsAsset::Equals(const UNiagaraDataInterface* Other) const
+{
+	if (!Super::Equals(Other))
+	{
+		return false;
+	}
+	const UNiagaraDataInterfacePhysicsAsset* OtherTyped = CastChecked<const UNiagaraDataInterfacePhysicsAsset>(Other);
+
+	return  (OtherTyped->PhysicsAssets == PhysicsAssets) && 
+		(OtherTyped->SoftSourceActor == SoftSourceActor) &&
+		(OtherTyped->SourceComponents == SourceComponents) && 
+		(OtherTyped->DefaultSource == DefaultSource) &&
+		(OtherTyped->MeshUserParameter == MeshUserParameter);
+}
+
+void UNiagaraDataInterfacePhysicsAsset::PostInitProperties()
+{
+	Super::PostInitProperties();
+
+	if (HasAnyFlags(RF_ClassDefaultObject))
+	{
+		ENiagaraTypeRegistryFlags Flags = ENiagaraTypeRegistryFlags::AllowAnyVariable | ENiagaraTypeRegistryFlags::AllowParameter;
+		FNiagaraTypeRegistry::Register(FNiagaraTypeDefinition(GetClass()), Flags);
+	}
+}
+
+void UNiagaraDataInterfacePhysicsAsset::PostLoad()
+{
+	Super::PostLoad();
+
+#if WITH_EDITORONLY_DATA
+	if (SourceActor_DEPRECATED != nullptr)
+	{
+		SoftSourceActor = SourceActor_DEPRECATED;
+	}
+#endif
+}
+
+void UNiagaraDataInterfacePhysicsAsset::GetFunctions(TArray<FNiagaraFunctionSignature>& OutFunctions)
+{
+	{
+		FNiagaraFunctionSignature Sig;
+		Sig.Name = GetNumBoxesName;
+		Sig.bSupportsGPU = true;
+		Sig.bSupportsCPU = false;
+		Sig.bMemberFunction = true;
+		Sig.bRequiresContext = false;
+		Sig.Inputs.Add(FNiagaraVariable(FNiagaraTypeDefinition(GetClass()), TEXT("Physics Asset")));
+		Sig.Outputs.Add(FNiagaraVariable(FNiagaraTypeDefinition::GetIntDef(), TEXT("Num Boxes")));
+
+		OutFunctions.Add(Sig);
+	}
+	{
+		FNiagaraFunctionSignature Sig;
+		Sig.Name = GetNumSpheresName;
+		Sig.bSupportsGPU = true;
+		Sig.bSupportsCPU = false;
+		Sig.bMemberFunction = true;
+		Sig.bRequiresContext = false;
+		Sig.Inputs.Add(FNiagaraVariable(FNiagaraTypeDefinition(GetClass()), TEXT("Physics Asset")));
+		Sig.Outputs.Add(FNiagaraVariable(FNiagaraTypeDefinition::GetIntDef(), TEXT("Num Spheres")));
+
+		OutFunctions.Add(Sig);
+	}
+	{
+		FNiagaraFunctionSignature Sig;
+		Sig.Name = GetNumCapsulesName;
+		Sig.bSupportsGPU = true;
+		Sig.bSupportsCPU = false;
+		Sig.bMemberFunction = true;
+		Sig.bRequiresContext = false;
+		Sig.Inputs.Add(FNiagaraVariable(FNiagaraTypeDefinition(GetClass()), TEXT("Physics Asset")));
+		Sig.Outputs.Add(FNiagaraVariable(FNiagaraTypeDefinition::GetIntDef(), TEXT("Num Capsules")));
+
+		OutFunctions.Add(Sig);
+	}
+	{
+		FNiagaraFunctionSignature Sig;
+		Sig.Name = GetClosestPointName;
+		Sig.bSupportsGPU = true;
+		Sig.bSupportsCPU = false;
+		Sig.bMemberFunction = true;
+		Sig.bRequiresContext = false;
+		Sig.Inputs.Add(FNiagaraVariable(FNiagaraTypeDefinition(GetClass()), TEXT("Physics Asset")));
+		Sig.Inputs.Add(FNiagaraVariable(FNiagaraTypeDefinition::GetVec3Def(), TEXT("Node Position")));
+		Sig.Inputs.Add(FNiagaraVariable(FNiagaraTypeDefinition::GetFloatDef(), TEXT("Delta Time")));
+		Sig.Inputs.Add(FNiagaraVariable(FNiagaraTypeDefinition::GetFloatDef(), TEXT("Time Fraction")));
+		Sig.Outputs.Add(FNiagaraVariable(FNiagaraTypeDefinition::GetVec3Def(), TEXT("Closest Position")));
+		Sig.Outputs.Add(FNiagaraVariable(FNiagaraTypeDefinition::GetVec3Def(), TEXT("Closest Normal")));
+		Sig.Outputs.Add(FNiagaraVariable(FNiagaraTypeDefinition::GetVec3Def(), TEXT("Closest Velocity")));
+
+		OutFunctions.Add(Sig);
+	}
+	{
+		FNiagaraFunctionSignature Sig;
+		Sig.Name = GetClosestElementName;
+		Sig.bSupportsGPU = true;
+		Sig.bSupportsCPU = false;
+		Sig.bMemberFunction = true;
+		Sig.bRequiresContext = false;
+		Sig.Inputs.Add(FNiagaraVariable(FNiagaraTypeDefinition(GetClass()), TEXT("Physics Asset")));
+		Sig.Inputs.Add(FNiagaraVariable(FNiagaraTypeDefinition::GetVec3Def(), TEXT("Node Position")));
+		Sig.Inputs.Add(FNiagaraVariable(FNiagaraTypeDefinition::GetFloatDef(), TEXT("Time Fraction")));
+		Sig.Outputs.Add(FNiagaraVariable(FNiagaraTypeDefinition::GetIntDef(), TEXT("Closest Element")));
+
+		OutFunctions.Add(Sig);
+	}
+	{
+		FNiagaraFunctionSignature Sig;
+		Sig.Name = GetElementPointName;
+		Sig.bSupportsGPU = true;
+		Sig.bSupportsCPU = false;
+		Sig.bMemberFunction = true;
+		Sig.bRequiresContext = false;
+		Sig.Inputs.Add(FNiagaraVariable(FNiagaraTypeDefinition(GetClass()), TEXT("Physics Asset")));
+		Sig.Inputs.Add(FNiagaraVariable(FNiagaraTypeDefinition::GetVec3Def(), TEXT("Node Position")));
+		Sig.Inputs.Add(FNiagaraVariable(FNiagaraTypeDefinition::GetFloatDef(), TEXT("Delta Time")));
+		Sig.Inputs.Add(FNiagaraVariable(FNiagaraTypeDefinition::GetFloatDef(), TEXT("Time Fraction")));
+		Sig.Inputs.Add(FNiagaraVariable(FNiagaraTypeDefinition::GetIntDef(), TEXT("Element Index")));
+		Sig.Outputs.Add(FNiagaraVariable(FNiagaraTypeDefinition::GetVec3Def(), TEXT("Closest Position")));
+		Sig.Outputs.Add(FNiagaraVariable(FNiagaraTypeDefinition::GetVec3Def(), TEXT("Closest Normal")));
+		Sig.Outputs.Add(FNiagaraVariable(FNiagaraTypeDefinition::GetVec3Def(), TEXT("Closest Velocity")));
+
+		OutFunctions.Add(Sig);
+	}
+	{
+		FNiagaraFunctionSignature Sig;
+		Sig.Name = GetElementDistanceName;
+		Sig.bSupportsGPU = true;
+		Sig.bSupportsCPU = false;
+		Sig.bMemberFunction = true;
+		Sig.bRequiresContext = false;
+		Sig.Inputs.Add(FNiagaraVariable(FNiagaraTypeDefinition(GetClass()), TEXT("Physics Asset")));
+		Sig.Inputs.Add(FNiagaraVariable(FNiagaraTypeDefinition::GetVec3Def(), TEXT("Node Position")));
+		Sig.Inputs.Add(FNiagaraVariable(FNiagaraTypeDefinition::GetFloatDef(), TEXT("Time Fraction")));
+		Sig.Inputs.Add(FNiagaraVariable(FNiagaraTypeDefinition::GetIntDef(), TEXT("Element Index")));
+		Sig.Outputs.Add(FNiagaraVariable(FNiagaraTypeDefinition::GetFloatDef(), TEXT("Closest Distance")));
+
+		OutFunctions.Add(Sig);
+	}
+	{
+		FNiagaraFunctionSignature Sig;
+		Sig.Name = GetClosestDistanceName;
+		Sig.bSupportsGPU = true;
+		Sig.bSupportsCPU = false;
+		Sig.bMemberFunction = true;
+		Sig.bRequiresContext = false;
+		Sig.Inputs.Add(FNiagaraVariable(FNiagaraTypeDefinition(GetClass()), TEXT("Physics Asset")));
+		Sig.Inputs.Add(FNiagaraVariable(FNiagaraTypeDefinition::GetVec3Def(), TEXT("Node Position")));
+		Sig.Inputs.Add(FNiagaraVariable(FNiagaraTypeDefinition::GetFloatDef(), TEXT("Time Fraction")));
+		Sig.Outputs.Add(FNiagaraVariable(FNiagaraTypeDefinition::GetFloatDef(), TEXT("Closest Distance")));
+
+		OutFunctions.Add(Sig);
+	}
+	{
+		FNiagaraFunctionSignature Sig;
+		Sig.Name = GetRestDistanceName;
+		Sig.bSupportsGPU = true;
+		Sig.bSupportsCPU = false;
+		Sig.bMemberFunction = true;
+		Sig.bRequiresContext = false;
+		Sig.Inputs.Add(FNiagaraVariable(FNiagaraTypeDefinition(GetClass()), TEXT("Physics Asset")));
+		Sig.Inputs.Add(FNiagaraVariable(FNiagaraTypeDefinition::GetVec3Def(), TEXT("Rest Position")));
+		Sig.Outputs.Add(FNiagaraVariable(FNiagaraTypeDefinition::GetFloatDef(), TEXT("Rest Distance")));
+
+		OutFunctions.Add(Sig);
+	}
+	{
+		FNiagaraFunctionSignature Sig;
+		Sig.Name = GetTexturePointName;
+		Sig.bSupportsGPU = true;
+		Sig.bSupportsCPU = false;
+		Sig.bMemberFunction = true;
+		Sig.bRequiresContext = false;
+		Sig.Inputs.Add(FNiagaraVariable(FNiagaraTypeDefinition(GetClass()), TEXT("Physics Asset")));
+		Sig.Inputs.Add(FNiagaraVariable(FNiagaraTypeDefinition::GetVec3Def(), TEXT("Node Position")));
+		Sig.Outputs.Add(FNiagaraVariable(FNiagaraTypeDefinition::GetIntDef(), TEXT("Element Index")));
+		Sig.Outputs.Add(FNiagaraVariable(FNiagaraTypeDefinition::GetVec3Def(), TEXT("Texture Position")));
+
+		OutFunctions.Add(Sig);
+	}
+	{
+		FNiagaraFunctionSignature Sig;
+		Sig.Name = GetProjectionPointName;
+		Sig.bSupportsGPU = true;
+		Sig.bSupportsCPU = false;
+		Sig.bMemberFunction = true;
+		Sig.bRequiresContext = false;
+		Sig.Inputs.Add(FNiagaraVariable(FNiagaraTypeDefinition(GetClass()), TEXT("Physics Asset")));
+		Sig.Inputs.Add(FNiagaraVariable(FNiagaraTypeDefinition::GetVec3Def(), TEXT("Node Position")));
+		Sig.Inputs.Add(FNiagaraVariable(FNiagaraTypeDefinition::GetFloatDef(), TEXT("Delta Time")));
+		Sig.Inputs.Add(FNiagaraVariable(FNiagaraTypeDefinition::GetIntDef(), TEXT("Element Index")));
+		Sig.Inputs.Add(FNiagaraVariable(FNiagaraTypeDefinition::GetFloatDef(), TEXT("Texture Value")));
+		Sig.Inputs.Add(FNiagaraVariable(FNiagaraTypeDefinition::GetVec3Def(), TEXT("Texture Gradient")));
+		Sig.Outputs.Add(FNiagaraVariable(FNiagaraTypeDefinition::GetVec3Def(), TEXT("Closest Position")));
+		Sig.Outputs.Add(FNiagaraVariable(FNiagaraTypeDefinition::GetVec3Def(), TEXT("Closest Normal")));
+		Sig.Outputs.Add(FNiagaraVariable(FNiagaraTypeDefinition::GetVec3Def(), TEXT("Closest Velocity")));
+		Sig.Outputs.Add(FNiagaraVariable(FNiagaraTypeDefinition::GetFloatDef(), TEXT("Closest Distance")));
+
+		OutFunctions.Add(Sig);
+	}
+}
+
+DEFINE_NDI_DIRECT_FUNC_BINDER(UNiagaraDataInterfacePhysicsAsset, GetNumBoxes);
+DEFINE_NDI_DIRECT_FUNC_BINDER(UNiagaraDataInterfacePhysicsAsset, GetNumSpheres);
+DEFINE_NDI_DIRECT_FUNC_BINDER(UNiagaraDataInterfacePhysicsAsset, GetNumCapsules);
+DEFINE_NDI_DIRECT_FUNC_BINDER(UNiagaraDataInterfacePhysicsAsset, GetClosestPoint);
+DEFINE_NDI_DIRECT_FUNC_BINDER(UNiagaraDataInterfacePhysicsAsset, GetClosestElement);
+DEFINE_NDI_DIRECT_FUNC_BINDER(UNiagaraDataInterfacePhysicsAsset, GetElementPoint);
+DEFINE_NDI_DIRECT_FUNC_BINDER(UNiagaraDataInterfacePhysicsAsset, GetElementDistance);
+DEFINE_NDI_DIRECT_FUNC_BINDER(UNiagaraDataInterfacePhysicsAsset, GetClosestDistance);
+DEFINE_NDI_DIRECT_FUNC_BINDER(UNiagaraDataInterfacePhysicsAsset, GetTexturePoint);
+DEFINE_NDI_DIRECT_FUNC_BINDER(UNiagaraDataInterfacePhysicsAsset, GetProjectionPoint);
+
+void UNiagaraDataInterfacePhysicsAsset::GetVMExternalFunction(const FVMExternalFunctionBindingInfo& BindingInfo, void* InstanceData, FVMExternalFunction& OutFunc)
+{
+	if (BindingInfo.Name == GetNumBoxesName)
+	{
+		check(BindingInfo.GetNumInputs() == 1 && BindingInfo.GetNumOutputs() == 1);
+		NDI_FUNC_BINDER(UNiagaraDataInterfacePhysicsAsset, GetNumBoxes)::Bind(this, OutFunc);
+	}
+	else if (BindingInfo.Name == GetNumSpheresName)
+	{
+		check(BindingInfo.GetNumInputs() == 1 && BindingInfo.GetNumOutputs() == 1);
+		NDI_FUNC_BINDER(UNiagaraDataInterfacePhysicsAsset, GetNumSpheres)::Bind(this, OutFunc);
+	}
+	else if (BindingInfo.Name == GetNumCapsulesName)
+	{
+		check(BindingInfo.GetNumInputs() == 1 && BindingInfo.GetNumOutputs() == 1);
+		NDI_FUNC_BINDER(UNiagaraDataInterfacePhysicsAsset, GetNumCapsules)::Bind(this, OutFunc);
+	}
+	else if (BindingInfo.Name == GetClosestPointName)
+	{
+		check(BindingInfo.GetNumInputs() == 6 && BindingInfo.GetNumOutputs() == 9);
+		NDI_FUNC_BINDER(UNiagaraDataInterfacePhysicsAsset, GetClosestPoint)::Bind(this, OutFunc);
+	}
+	else if (BindingInfo.Name == GetClosestElementName)
+	{
+		check(BindingInfo.GetNumInputs() == 5 && BindingInfo.GetNumOutputs() == 1);
+		NDI_FUNC_BINDER(UNiagaraDataInterfacePhysicsAsset, GetClosestElement)::Bind(this, OutFunc);
+	}
+	else if (BindingInfo.Name == GetElementPointName)
+	{
+		check(BindingInfo.GetNumInputs() == 7 && BindingInfo.GetNumOutputs() == 9);
+		NDI_FUNC_BINDER(UNiagaraDataInterfacePhysicsAsset, GetElementPoint)::Bind(this, OutFunc);
+	}
+	else if (BindingInfo.Name == GetElementDistanceName)
+	{
+		check(BindingInfo.GetNumInputs() == 6 && BindingInfo.GetNumOutputs() == 1);
+		NDI_FUNC_BINDER(UNiagaraDataInterfacePhysicsAsset, GetElementDistance)::Bind(this, OutFunc);
+	}
+	else if (BindingInfo.Name == GetClosestDistanceName)
+	{
+		check(BindingInfo.GetNumInputs() == 5 && BindingInfo.GetNumOutputs() == 1);
+		NDI_FUNC_BINDER(UNiagaraDataInterfacePhysicsAsset, GetClosestDistance)::Bind(this, OutFunc);
+	}
+	else if (BindingInfo.Name == GetTexturePointName)
+	{
+		check(BindingInfo.GetNumInputs() == 4 && BindingInfo.GetNumOutputs() == 4);
+		NDI_FUNC_BINDER(UNiagaraDataInterfacePhysicsAsset, GetTexturePoint)::Bind(this, OutFunc);
+	}
+	else if (BindingInfo.Name == GetProjectionPointName)
+	{
+		check(BindingInfo.GetNumInputs() == 10 && BindingInfo.GetNumOutputs() == 10);
+		NDI_FUNC_BINDER(UNiagaraDataInterfacePhysicsAsset, GetProjectionPoint)::Bind(this, OutFunc);
+	}
+}
+
+void UNiagaraDataInterfacePhysicsAsset::GetNumBoxes(FVectorVMExternalFunctionContext& Context)
+{
+}
+
+void UNiagaraDataInterfacePhysicsAsset::GetNumSpheres(FVectorVMExternalFunctionContext& Context)
+{
+}
+
+void UNiagaraDataInterfacePhysicsAsset::GetNumCapsules(FVectorVMExternalFunctionContext& Context)
+{
+}
+
+void UNiagaraDataInterfacePhysicsAsset::GetClosestPoint(FVectorVMExternalFunctionContext& Context)
+{
+}
+
+void UNiagaraDataInterfacePhysicsAsset::GetClosestElement(FVectorVMExternalFunctionContext& Context)
+{
+}
+
+void UNiagaraDataInterfacePhysicsAsset::GetElementPoint(FVectorVMExternalFunctionContext& Context)
+{
+}
+
+void UNiagaraDataInterfacePhysicsAsset::GetElementDistance(FVectorVMExternalFunctionContext& Context)
+{
+}
+
+void UNiagaraDataInterfacePhysicsAsset::GetClosestDistance(FVectorVMExternalFunctionContext& Context)
+{
+}
+
+void UNiagaraDataInterfacePhysicsAsset::GetTexturePoint(FVectorVMExternalFunctionContext& Context)
+{
+}
+
+void UNiagaraDataInterfacePhysicsAsset::GetProjectionPoint(FVectorVMExternalFunctionContext& Context)
+{
+}
+
+#if WITH_EDITORONLY_DATA
+bool UNiagaraDataInterfacePhysicsAsset::GetFunctionHLSL(const FNiagaraDataInterfaceGPUParamInfo& ParamInfo, const FNiagaraDataInterfaceGeneratedFunction& FunctionInfo, int FunctionInstanceIndex, FString& OutHLSL)
+{
+	FNDIPhysicsAssetParametersName ParamNames(ParamInfo.DataInterfaceHLSLSymbol);
+
+	TMap<FString, FStringFormatArg> ArgsSample = {
+		{TEXT("InstanceFunctionName"), FunctionInfo.InstanceName},
+		{TEXT("ElementOffsetsName"), ParamNames.ElementOffsetsName},
+		{TEXT("WorldTransformBufferName"), ParamNames.WorldTransformBufferName},
+		{TEXT("InverseTransformBufferName"), ParamNames.InverseTransformBufferName},
+		{TEXT("ElementExtentBufferName"), ParamNames.ElementExtentBufferName},
+		{TEXT("PhysicsAssetContextName"), TEXT("DIPHYSICSASSET_MAKE_CONTEXT(") + ParamInfo.DataInterfaceHLSLSymbol + TEXT(")")},
+	};
+
+	if (FunctionInfo.DefinitionName == GetNumBoxesName)
+	{
+		static const TCHAR* FormatSample = TEXT(R"(
+		void {InstanceFunctionName}(out int OutNumBoxes)
+		{
+			{PhysicsAssetContextName}
+			OutNumBoxes = DIPhysicsAsset_GetNumBoxes(DIContext);
+		}
+		)");
+		OutHLSL += FString::Format(FormatSample, ArgsSample);
+		return true;
+	}
+	else if (FunctionInfo.DefinitionName == GetNumCapsulesName)
+	{
+		static const TCHAR* FormatSample = TEXT(R"(
+		void {InstanceFunctionName}(out int OutNumCapsules)
+		{
+			{PhysicsAssetContextName}
+			OutNumCapsules = DIPhysicsAsset_GetNumCapsules(DIContext);
+		}
+		)");
+		OutHLSL += FString::Format(FormatSample, ArgsSample);
+		return true;
+	}
+	else if (FunctionInfo.DefinitionName == GetNumSpheresName)
+	{
+		static const TCHAR* FormatSample = TEXT(R"(
+		void {InstanceFunctionName}(out int OutNumSpheres)
+		{
+			{PhysicsAssetContextName}
+			OutNumSpheres = DIPhysicsAsset_GetNumSpheres(DIContext);
+		}
+		)");
+		OutHLSL += FString::Format(FormatSample, ArgsSample);
+		return true;
+	}
+	else if (FunctionInfo.DefinitionName == GetClosestPointName)
+	{
+		static const TCHAR* FormatSample = TEXT(R"(
+		void {InstanceFunctionName}(in float3 NodePosition, in float DeltaTime, in float TimeFraction, out float3 OutClosestPosition, 
+							out float3 OutClosestNormal, out float3 OutClosestVelocity)
+		{
+			{PhysicsAssetContextName} DIPhysicsAsset_GetClosestPoint(DIContext,NodePosition,DeltaTime,TimeFraction,
+				OutClosestPosition,OutClosestNormal,OutClosestVelocity);
+		}
+		)");
+		OutHLSL += FString::Format(FormatSample, ArgsSample);
+		return true;
+	}
+	else if (FunctionInfo.DefinitionName == GetClosestElementName)
+	{
+		static const TCHAR* FormatSample = TEXT(R"(
+		void {InstanceFunctionName}(in float3 NodePosition, in float TimeFraction, out int OutClosestElement)
+		{
+			{PhysicsAssetContextName} DIPhysicsAsset_GetClosestElement(DIContext,NodePosition,TimeFraction,
+				OutClosestElement);
+		}
+		)");
+		OutHLSL += FString::Format(FormatSample, ArgsSample);
+		return true;
+	}
+	else if (FunctionInfo.DefinitionName == GetElementPointName)
+	{
+		static const TCHAR* FormatSample = TEXT(R"(
+		void {InstanceFunctionName}(in float3 NodePosition, in float DeltaTime, in float TimeFraction, in int ElementIndex, out float3 OutClosestPosition, 
+							out float3 OutClosestNormal, out float3 OutClosestVelocity)
+		{
+			{PhysicsAssetContextName} DIPhysicsAsset_GetElementPoint(DIContext,NodePosition,DeltaTime,TimeFraction,ElementIndex,
+				OutClosestPosition,OutClosestNormal,OutClosestVelocity);
+		}
+		)");
+		OutHLSL += FString::Format(FormatSample, ArgsSample);
+		return true;
+	}
+	else if (FunctionInfo.DefinitionName == GetElementDistanceName)
+	{
+		static const TCHAR* FormatSample = TEXT(R"(
+		void {InstanceFunctionName}(in float3 NodePosition, in float TimeFraction, in int ElementIndex, out float OutClosestDistance)
+		{
+			{PhysicsAssetContextName} DIPhysicsAsset_GetElementDistance(DIContext,NodePosition,TimeFraction,ElementIndex,
+				OutClosestDistance);
+		}
+		)");
+		OutHLSL += FString::Format(FormatSample, ArgsSample);
+		return true;
+	}
+	else if (FunctionInfo.DefinitionName == GetClosestDistanceName)
+	{
+		static const TCHAR* FormatSample = TEXT(R"(
+		void {InstanceFunctionName}(in float3 NodePosition, in float TimeFraction, out float OutClosestDistance)
+		{
+			{PhysicsAssetContextName} DIPhysicsAsset_GetClosestDistance(DIContext,NodePosition,TimeFraction,OutClosestDistance);
+		}
+		)");
+		OutHLSL += FString::Format(FormatSample, ArgsSample);
+		return true;
+	}
+	else if (FunctionInfo.DefinitionName == GetRestDistanceName)
+	{
+		static const TCHAR* FormatSample = TEXT(R"(
+		void {InstanceFunctionName}(in float3 RestPosition, out float OutRestDistance)
+		{
+			{PhysicsAssetContextName} DIPhysicsAsset_GetRestDistance(DIContext,RestPosition,OutRestDistance);
+		}
+		)");
+		OutHLSL += FString::Format(FormatSample, ArgsSample);
+		return true;
+	}
+	else if (FunctionInfo.DefinitionName == GetTexturePointName)
+	{
+		static const TCHAR* FormatSample = TEXT(R"(
+		void {InstanceFunctionName}(in float3 NodePosition, out int OutElementIndex, out float3 OutTexturePosition)
+		{
+			{PhysicsAssetContextName} DIPhysicsAsset_GetTexturePoint(DIContext,NodePosition,OutElementIndex,OutTexturePosition);
+		}
+		)");
+		OutHLSL += FString::Format(FormatSample, ArgsSample);
+		return true;
+	}
+	else if (FunctionInfo.DefinitionName == GetProjectionPointName)
+	{
+		static const TCHAR* FormatSample = TEXT(R"(
+		void {InstanceFunctionName}(in float3 NodePosition, in float DeltaTime, in int ElementIndex, in float TextureValue, in float3 TextureGradient, out float3 OutClosestPosition, 
+							out float3 OutClosestNormal, out float3 OutClosestVelocity, out float OutClosestDistance)
+		{
+			{PhysicsAssetContextName} DIPhysicsAsset_GetProjectionPoint(DIContext,NodePosition,DeltaTime,ElementIndex,TextureValue,TextureGradient,
+				OutClosestPosition,OutClosestNormal,OutClosestVelocity,OutClosestDistance);
+		}
+		)");
+		OutHLSL += FString::Format(FormatSample, ArgsSample);
+		return true;
+	}
+
+	OutHLSL += TEXT("\n");
+	return false;
+}
+
+void UNiagaraDataInterfacePhysicsAsset::GetCommonHLSL(FString& OutHLSL)
+{
+	OutHLSL += TEXT("#include \"/Plugin/FX/Niagara/Private/NiagaraQuaternionUtils.ush\"\n");
+	OutHLSL += TEXT("#include \"/Plugin/FX/Niagara/Private/NiagaraDataInterfacePhysicsAsset.ush\"\n");	
+}
+
+void UNiagaraDataInterfacePhysicsAsset::GetParameterDefinitionHLSL(const FNiagaraDataInterfaceGPUParamInfo& ParamInfo, FString& OutHLSL)
+{
+	OutHLSL += TEXT("DIPHYSICSASSET_DECLARE_CONSTANTS(") + ParamInfo.DataInterfaceHLSLSymbol + TEXT(")\n");
+}
+#endif
+
+void UNiagaraDataInterfacePhysicsAsset::ProvidePerInstanceDataForRenderThread(void* DataForRenderThread, void* PerInstanceData, const FNiagaraSystemInstanceID& SystemInstance)
+{
+	FNDIPhysicsAssetData* GameThreadData = static_cast<FNDIPhysicsAssetData*>(PerInstanceData);
+	FNDIPhysicsAssetData* RenderThreadData = static_cast<FNDIPhysicsAssetData*>(DataForRenderThread);
+
+	if (GameThreadData != nullptr && RenderThreadData != nullptr)
+	{
+		RenderThreadData->AssetBuffer = GameThreadData->AssetBuffer;
+		RenderThreadData->BoxOrigin = GameThreadData->BoxOrigin;
+		RenderThreadData->BoxExtent = GameThreadData->BoxExtent;
+		RenderThreadData->AssetArrays = GameThreadData->AssetArrays;
+		RenderThreadData->TickingGroup = GameThreadData->TickingGroup;
+	}
+	check(Proxy);
+}
+
+UNiagaraPhysicsAssetDICollectorInterface::UNiagaraPhysicsAssetDICollectorInterface(const FObjectInitializer& ObjectInitializer)
+	: Super(ObjectInitializer)
+{}
+
+#undef LOCTEXT_NAMESPACE
\ No newline at end of file
diff --git a/Engine/Plugins/FX/Niagara/Source/Niagara/Public/NiagaraDataInterfacePhysicsAsset.h b/Engine/Plugins/FX/Niagara/Source/Niagara/Public/NiagaraDataInterfacePhysicsAsset.h
new file mode 100644
--- /dev/null
+++ b/Engine/Plugins/FX/Niagara/Source/Niagara/Public/NiagaraDataInterfacePhysicsAsset.h
@@ -0,0 +1,254 @@
+// Copyright Epic Games, Inc. All Rights Reserved.
+
+#pragma once
+
+#include "NiagaraDataInterface.h"
+#include "NiagaraCommon.h"
+#include "VectorVM.h"
+#include "PhysicsEngine/PhysicsAsset.h"
+#include "UObject/Interface.h"
+
+#include "NiagaraDataInterfacePhysicsAsset.generated.h"
+
+#define PHYSICS_ASSET_MAX_PRIMITIVES 100
+#define PHYSICS_ASSET_MAX_TRANSFORMS PHYSICS_ASSET_MAX_PRIMITIVES * 3
+
+/** Element offsets in the array list */
+struct FElementOffset
+{
+	FElementOffset(const uint32 InBoxOffset, const uint32 InSphereOffset, const uint32 InCapsuleOffset, const uint32 InNumElements) :
+		BoxOffset(InBoxOffset), SphereOffset(InSphereOffset), CapsuleOffset(InCapsuleOffset), NumElements(InNumElements)
+	{}
+
+	FElementOffset() :
+		BoxOffset(0), SphereOffset(0), CapsuleOffset(0), NumElements(0)
+	{}
+	uint32 BoxOffset;
+	uint32 SphereOffset;
+	uint32 CapsuleOffset;
+	uint32 NumElements;
+};
+
+/** Arrays in which the cpu datas will be str */
+struct FNDIPhysicsAssetArrays
+{
+	FElementOffset ElementOffsets;
+	TStaticArray<FVector4f, 3 * PHYSICS_ASSET_MAX_TRANSFORMS> WorldTransform;
+	TStaticArray<FVector4f, 3 * PHYSICS_ASSET_MAX_TRANSFORMS> InverseTransform;
+	TStaticArray<FVector4f, PHYSICS_ASSET_MAX_TRANSFORMS> CurrentTransform;
+	TStaticArray<FVector4f, PHYSICS_ASSET_MAX_TRANSFORMS> CurrentInverse;
+	TStaticArray<FVector4f, PHYSICS_ASSET_MAX_TRANSFORMS> PreviousTransform;
+	TStaticArray<FVector4f, PHYSICS_ASSET_MAX_TRANSFORMS> PreviousInverse;
+	TStaticArray<FVector4f, PHYSICS_ASSET_MAX_TRANSFORMS> RestTransform;
+	TStaticArray<FVector4f, PHYSICS_ASSET_MAX_TRANSFORMS> RestInverse;
+	TStaticArray<FVector4f, PHYSICS_ASSET_MAX_PRIMITIVES> ElementExtent;
+	TStaticArray<uint32, PHYSICS_ASSET_MAX_PRIMITIVES> PhysicsType;
+};
+
+/** Render buffers that will be used in hlsl functions */
+struct FNDIPhysicsAssetBuffer : public FRenderResource
+{
+	/** Init the buffer */
+	virtual void InitRHI() override;
+
+	/** Release the buffer */
+	virtual void ReleaseRHI() override;
+
+	/** Get the resource name */
+	virtual FString GetFriendlyName() const override { return TEXT("FNDIPhysicsAssetBuffer"); }
+
+	/** World transform buffer */
+	FRWBuffer WorldTransformBuffer;
+
+	/** Inverse transform buffer*/
+	FRWBuffer InverseTransformBuffer;
+
+	/** Element extent buffer */
+	FRWBuffer ElementExtentBuffer;
+
+	/** Physics type buffer */
+	FRWBuffer PhysicsTypeBuffer;
+};
+
+/** Data stored per physics asset instance*/
+struct FNDIPhysicsAssetData
+{
+	/** Initialize the cpu datas */
+	void Init(class UNiagaraDataInterfacePhysicsAsset* Interface, FNiagaraSystemInstance* SystemInstance);
+
+	/** Update the gpu datas */
+	void Update(class UNiagaraDataInterfacePhysicsAsset* Interface, FNiagaraSystemInstance* SystemInstance);
+
+	/** Release the buffers */
+	void Release();
+
+	/** Physics asset Gpu buffer */
+	FNDIPhysicsAssetBuffer* AssetBuffer;
+
+	/** Bounding box center */
+	FVector BoxOrigin;
+
+	/** Bounding box extent */
+	FVector BoxExtent;
+
+	/** Physics asset Cpu arrays */
+	FNDIPhysicsAssetArrays AssetArrays;  
+
+	/** The instance ticking group */
+	ETickingGroup TickingGroup;
+};
+
+/** Data Interface for interacting with PhysicsAssets */
+UCLASS(EditInlineNew, Category = "Physics", meta = (DisplayName = "Physics Asset"))
+class NIAGARA_API UNiagaraDataInterfacePhysicsAsset : public UNiagaraDataInterface
+{
+	GENERATED_UCLASS_BODY()
+
+public:
+
+	DECLARE_NIAGARA_DI_PARAMETER();
+
+	/** Skeletal Mesh from which the Physics Asset will be found. */
+	UPROPERTY(EditAnywhere, Category = "Source")
+	TObjectPtr<UPhysicsAsset> DefaultSource;
+
+	/** The source actor from which to sample */
+	UPROPERTY(EditAnywhere, Category = "Source", meta = (DisplayName = "Source Actor"))
+	TSoftObjectPtr<AActor> SoftSourceActor;
+
+#if WITH_EDITORONLY_DATA
+	UPROPERTY()
+	TObjectPtr<AActor> SourceActor_DEPRECATED;
+#endif
+
+	/** Reference to a user parameter if we're reading one. */
+	UPROPERTY(EditAnywhere, Category = "Source")
+	FNiagaraUserParameterBinding MeshUserParameter;
+
+	/** The source component from which to sample */
+	TArray<TWeakObjectPtr<class USkeletalMeshComponent>> SourceComponents;
+
+	/** The source asset from which to sample */
+	TArray<TWeakObjectPtr<class UPhysicsAsset>> PhysicsAssets;
+
+	/** UObject Interface */
+	virtual void PostInitProperties() override;
+	virtual void PostLoad() override;
+
+	/** UNiagaraDataInterface Interface */
+	virtual void GetFunctions(TArray<FNiagaraFunctionSignature>& OutFunctions) override;
+	virtual void GetVMExternalFunction(const FVMExternalFunctionBindingInfo& BindingInfo, void* InstanceData, FVMExternalFunction& OutFunc) override;
+	virtual bool CanExecuteOnTarget(ENiagaraSimTarget Target) const override { return Target == ENiagaraSimTarget::GPUComputeSim; }
+	virtual bool InitPerInstanceData(void* PerInstanceData, FNiagaraSystemInstance* SystemInstance) override;
+	virtual void DestroyPerInstanceData(void* PerInstanceData, FNiagaraSystemInstance* SystemInstance)override;
+	virtual bool PerInstanceTick(void* PerInstanceData, FNiagaraSystemInstance* SystemInstance, float DeltaSeconds) override;
+	virtual int32 PerInstanceDataSize()const override { return sizeof(FNDIPhysicsAssetData); }
+	virtual bool Equals(const UNiagaraDataInterface* Other) const override;
+	virtual bool HasPreSimulateTick() const override { return true; }
+	virtual bool HasTickGroupPrereqs() const override { return true; }
+	virtual ETickingGroup CalculateTickGroup(const void* PerInstanceData) const override;
+
+	/** GPU simulation  functionality */
+#if WITH_EDITORONLY_DATA
+	virtual void GetCommonHLSL(FString& OutHLSL) override;
+	virtual void GetParameterDefinitionHLSL(const FNiagaraDataInterfaceGPUParamInfo& ParamInfo, FString& OutHLSL) override;
+	virtual bool GetFunctionHLSL(const FNiagaraDataInterfaceGPUParamInfo& ParamInfo, const FNiagaraDataInterfaceGeneratedFunction& FunctionInfo, int FunctionInstanceIndex, FString& OutHLSL) override;
+#endif
+	virtual void ProvidePerInstanceDataForRenderThread(void* DataForRenderThread, void* PerInstanceData, const FNiagaraSystemInstanceID& SystemInstance) override;
+
+	/** Extract the source component */
+	void ExtractSourceComponent(FNiagaraSystemInstance* SystemInstance, FTransform& LocalTransform);
+
+	/** Get the number of boxes*/
+	void GetNumBoxes(FVectorVMExternalFunctionContext& Context);
+
+	/** Get the number of spheres  */
+	void GetNumSpheres(FVectorVMExternalFunctionContext& Context);
+
+	/** Get the number of capsules */
+	void GetNumCapsules(FVectorVMExternalFunctionContext& Context);
+
+	/** Get the element point */
+	void GetElementPoint(FVectorVMExternalFunctionContext& Context);
+
+	/** Get the element distance */
+	void GetElementDistance(FVectorVMExternalFunctionContext& Context);
+
+	/** Get the closest element */
+	void GetClosestElement(FVectorVMExternalFunctionContext& Context);
+
+	/** Get the closest point */
+	void GetClosestPoint(FVectorVMExternalFunctionContext& Context);
+
+	/** Get the closest distance */
+	void GetClosestDistance(FVectorVMExternalFunctionContext& Context);
+
+	/** Get the closest texture point */
+	void GetTexturePoint(FVectorVMExternalFunctionContext& Context);
+
+	/** Get the projection point */
+	void GetProjectionPoint(FVectorVMExternalFunctionContext& Context);
+
+	/** Name of element offsets */
+	static const FString ElementOffsetsName;
+
+	/** Name of the world transform buffer */
+	static const FString WorldTransformBufferName;
+
+	/** Name of the inverse transform buffer */
+	static const FString InverseTransformBufferName;
+
+	/** Name of the element extent buffer */
+	static const FString ElementExtentBufferName;
+
+	/** Name of the physics type buffer */
+	static const FString PhysicsTypeBufferName;
+
+	/** Init Box Origin */
+	static const FString BoxOriginName;
+
+	/** Init Box extent */
+	static const FString BoxExtentName;
+
+protected:
+	/** Copy one niagara DI to this */
+	virtual bool CopyToInternal(UNiagaraDataInterface* Destination) const override;
+};
+
+/** Proxy to send data to gpu */
+struct FNDIPhysicsAssetProxy : public FNiagaraDataInterfaceProxy
+{
+	/** Get the size of the data that will be passed to render*/
+	virtual int32 PerInstanceDataPassedToRenderThreadSize() const override { return sizeof(FNDIPhysicsAssetData); }
+
+	/** Get the data that will be passed to render*/
+	virtual void ConsumePerInstanceDataFromGameThread(void* PerInstanceData, const FNiagaraSystemInstanceID& Instance) override;
+
+	/** Initialize the Proxy data buffer */
+	void InitializePerInstanceData(const FNiagaraSystemInstanceID& SystemInstance);
+
+	/** Destroy the proxy data if necessary */
+	void DestroyPerInstanceData(const FNiagaraSystemInstanceID& SystemInstance);
+
+	/** Launch all pre stage functions */
+	virtual void PreStage(FRHICommandList& RHICmdList, const FNiagaraDataInterfaceStageArgs& Context) override;
+
+	/** List of proxy data for each system instances*/
+	TMap<FNiagaraSystemInstanceID, FNDIPhysicsAssetData> SystemInstancesToProxyData;
+};
+
+UINTERFACE()
+class NIAGARA_API UNiagaraPhysicsAssetDICollectorInterface : public UInterface
+{
+	GENERATED_UINTERFACE_BODY()
+};
+
+class NIAGARA_API INiagaraPhysicsAssetDICollectorInterface
+{
+	GENERATED_IINTERFACE_BODY()
+
+	virtual UPhysicsAsset* BuildAndCollect(
+		FTransform& BoneTransform,
+		TArray<TWeakObjectPtr<USkeletalMeshComponent>>& SourceComponents,
+		TArray<TWeakObjectPtr<UPhysicsAsset>>& PhysicsAssets) const = 0;
+};
\ No newline at end of file
diff --git a/Engine/Plugins/Runtime/HairStrands/Shaders/Private/NiagaraDataInterfacePhysicsAsset.ush b/Engine/Plugins/Runtime/HairStrands/Shaders/Private/NiagaraDataInterfacePhysicsAsset.ush
deleted file mode 100644
diff --git a/Engine/Plugins/Runtime/HairStrands/Shaders/Private/NiagaraGroomInterpolationUtils.ush b/Engine/Plugins/Runtime/HairStrands/Shaders/Private/NiagaraGroomInterpolationUtils.ush
--- a/Engine/Plugins/Runtime/HairStrands/Shaders/Private/NiagaraGroomInterpolationUtils.ush
+++ b/Engine/Plugins/Runtime/HairStrands/Shaders/Private/NiagaraGroomInterpolationUtils.ush
@@ -6,7 +6,7 @@ NiagaraGroomInterpolationUtils.ush
 
 #pragma once
 
-#include "NiagaraQuaternionUtils.ush"
+#include "/Plugin/FX/Niagara/Private/NiagaraQuaternionUtils.ush"
 
 // Compute the projection triangle
 void DIHairStrands_ComputeProjectionTriangle(in float2 ProjectionUV, in float3 PA, in float3 PB, in float3 PC, in float3 Offset, 
diff --git a/Engine/Plugins/Runtime/HairStrands/Shaders/Private/NiagaraQuaternionUtils.ush b/Engine/Plugins/Runtime/HairStrands/Shaders/Private/NiagaraQuaternionUtils.ush
deleted file mode 100644
diff --git a/Engine/Plugins/Runtime/HairStrands/Source/HairStrandsCore/Private/GroomComponent.cpp b/Engine/Plugins/Runtime/HairStrands/Source/HairStrandsCore/Private/GroomComponent.cpp
--- a/Engine/Plugins/Runtime/HairStrands/Source/HairStrandsCore/Private/GroomComponent.cpp
+++ b/Engine/Plugins/Runtime/HairStrands/Source/HairStrandsCore/Private/GroomComponent.cpp
@@ -3304,6 +3304,22 @@ void UGroomComponent::BuildSimulationTransform(FTransform& SimulationTransform)
 	}
 }
 
+UPhysicsAsset* UGroomComponent::BuildAndCollect(FTransform& BoneTransform, TArray<TWeakObjectPtr<USkeletalMeshComponent>>& SourceComponents, TArray<TWeakObjectPtr<UPhysicsAsset>>& PhysicsAssets) const
+{
+	BuildSimulationTransform(BoneTransform);
+
+	for (auto& CollisionComponent : CollisionComponents)
+	{
+		if (CollisionComponent.IsValid() && CollisionComponent->GetPhysicsAsset())
+		{
+			SourceComponents.Add(CollisionComponent);
+			PhysicsAssets.Add(CollisionComponent->GetPhysicsAsset());
+		}
+	}
+
+	return PhysicsAsset;
+}
+
 void UGroomComponent::TickComponent(float DeltaTime, enum ELevelTick TickType, FActorComponentTickFunction *ThisTickFunction)
 {
 	Super::TickComponent(DeltaTime, TickType, ThisTickFunction);	
diff --git a/Engine/Plugins/Runtime/HairStrands/Source/HairStrandsCore/Private/Niagara/NiagaraDataInterfaceHairStrands.cpp b/Engine/Plugins/Runtime/HairStrands/Source/HairStrandsCore/Private/Niagara/NiagaraDataInterfaceHairStrands.cpp
--- a/Engine/Plugins/Runtime/HairStrands/Source/HairStrandsCore/Private/Niagara/NiagaraDataInterfaceHairStrands.cpp
+++ b/Engine/Plugins/Runtime/HairStrands/Source/HairStrandsCore/Private/Niagara/NiagaraDataInterfaceHairStrands.cpp
@@ -4404,7 +4404,7 @@ in float RestLength, in float DeltaTime, in int NodeOffset, in float MaterialDam
 
 void UNiagaraDataInterfaceHairStrands::GetCommonHLSL(FString& OutHLSL)
 {
-	OutHLSL += TEXT("#include \"/Plugin/Runtime/HairStrands/Private/NiagaraQuaternionUtils.ush\"\n");
+	OutHLSL += TEXT("#include \"/Plugin/FX/Niagara/Private/NiagaraQuaternionUtils.ush\"\n");
 	//OutHLSL += TEXT("#include \"/Plugin/Runtime/HairStrands/Private/NiagaraDirectSolver.ush\"\n");
 	OutHLSL += TEXT("#include \"/Plugin/Runtime/HairStrands/Private/NiagaraStrandsExternalForce.ush\"\n");
 	OutHLSL += TEXT("#include \"/Plugin/Runtime/HairStrands/Private/NiagaraHookeSpringMaterial.ush\"\n");
@@ -4422,7 +4422,7 @@ bool UNiagaraDataInterfaceHairStrands::AppendCompileHash(FNiagaraCompileHashVisi
 
 	for (const TCHAR* VirtualFilePath :
 		{
-			TEXT("/Plugin/Runtime/HairStrands/Private/NiagaraQuaternionUtils.ush"),
+			TEXT("/Plugin/FX/Niagara/Private/NiagaraQuaternionUtils.ush"),
 			TEXT("/Plugin/Runtime/HairStrands/Private/NiagaraStrandsExternalForce.ush"),
 			TEXT("/Plugin/Runtime/HairStrands/Private/NiagaraHookeSpringMaterial.ush"),
 			TEXT("/Plugin/Runtime/HairStrands/Private/NiagaraAngularSpringMaterial.ush"),
diff --git a/Engine/Plugins/Runtime/HairStrands/Source/HairStrandsCore/Private/Niagara/NiagaraDataInterfacePhysicsAsset.cpp b/Engine/Plugins/Runtime/HairStrands/Source/HairStrandsCore/Private/Niagara/NiagaraDataInterfacePhysicsAsset.cpp
deleted file mode 100644
diff --git a/Engine/Plugins/Runtime/HairStrands/Source/HairStrandsCore/Public/GroomComponent.h b/Engine/Plugins/Runtime/HairStrands/Source/HairStrandsCore/Public/GroomComponent.h
--- a/Engine/Plugins/Runtime/HairStrands/Source/HairStrandsCore/Public/GroomComponent.h
+++ b/Engine/Plugins/Runtime/HairStrands/Source/HairStrandsCore/Public/GroomComponent.h
@@ -12,13 +12,14 @@
 #include "GroomDesc.h"
 #include "LODSyncInterface.h"
 #include "GroomInstance.h"
+#include "NiagaraDataInterfacePhysicsAsset.h"
 
 #include "GroomComponent.generated.h"
 
 class UGroomCache;
 
 UCLASS(HideCategories = (Object, Physics, Activation, Mobility, "Components|Activation"), editinlinenew, meta = (BlueprintSpawnableComponent), ClassGroup = Rendering)
-class HAIRSTRANDSCORE_API UGroomComponent : public UMeshComponent, public ILODSyncInterface
+class HAIRSTRANDSCORE_API UGroomComponent : public UMeshComponent, public ILODSyncInterface, public INiagaraPhysicsAssetDICollectorInterface
 {
 	GENERATED_UCLASS_BODY()
 
@@ -269,6 +270,10 @@ public:
 	/** Build the local simulation transform that could be used in strands simulation */
 	void BuildSimulationTransform(FTransform& SimulationTransform) const;
 
+	//~ Begin INiagaraPhysicsAssetDICollectorInterface Interface
+	virtual UPhysicsAsset* BuildAndCollect(FTransform& BoneTransform, TArray<TWeakObjectPtr<USkeletalMeshComponent>>& SourceComponents, TArray<TWeakObjectPtr<UPhysicsAsset>>& PhysicsAssets) const override;
+	//~ End INiagaraPhysicsAssetDICollectorInterface Interface
+
 private:
 	void UpdateGroomCache(float Time);
 
diff --git a/Engine/Plugins/Runtime/HairStrands/Source/HairStrandsCore/Public/Niagara/NiagaraDataInterfacePhysicsAsset.h b/Engine/Plugins/Runtime/HairStrands/Source/HairStrandsCore/Public/Niagara/NiagaraDataInterfacePhysicsAsset.h
deleted file mode 100644
