Added a validation rule to warn about fixed tick delta times for modules and systems

diff --git a/Engine/Plugins/FX/Niagara/Source/NiagaraEditor/Private/NiagaraValidationRules.cpp b/Engine/Plugins/FX/Niagara/Source/NiagaraEditor/Private/NiagaraValidationRules.cpp
--- a/Engine/Plugins/FX/Niagara/Source/NiagaraEditor/Private/NiagaraValidationRules.cpp
+++ b/Engine/Plugins/FX/Niagara/Source/NiagaraEditor/Private/NiagaraValidationRules.cpp
@@ -495,6 +495,40 @@ void UNiagaraValidationRule_NoOpaqueRenderMaterial::CheckValidity(const FNiagara
 	}
 }
 
+void UNiagaraValidationRule_NoFixedDeltaTime::CheckValidity(const FNiagaraValidationContext& Context, TArray<FNiagaraValidationResult>& OutResults) const
+{
+	// check to see if we're called from a module or the effect type
+	if (UNiagaraStackModuleItem* SourceModule = Cast<UNiagaraStackModuleItem>(Context.Source))
+	{
+		if (SourceModule->GetIsEnabled())
+		{
+			UNiagaraSystem& System = SourceModule->GetSystemViewModel()->GetSystem();
+			if (System.HasFixedTickDelta())
+			{
+				OutResults.Emplace_GetRef(
+					ENiagaraValidationSeverity::Warning,
+					LOCTEXT("NoFixedDeltaTimeModule", "Module does not support fixed tick delta time"),
+					LOCTEXT("NoFixedDeltaTimeModuleDetailed", "This system uses a fixed tick delta time, which means it might tick multiple times per frame or might skip ticks depending on the global tick rate.\nModules that depend on external assets such as render targets or collision data will NOT work correctly when their tick is different from the engine tick.\nConsider disabling the fixed tick delta time."),
+					SourceModule
+				);
+			}
+		}
+	}
+	else
+	{
+		UNiagaraSystem& System = Context.ViewModel->GetSystem();
+		if (System.HasFixedTickDelta())
+		{
+			OutResults.Emplace_GetRef(
+				ENiagaraValidationSeverity::Error,
+				LOCTEXT("NoFixedDeltaTime", "Effect tyoe does not allow fixed tick delta time"),
+				LOCTEXT("NoFixedDeltaTimeDetailed", "This system uses a fixed tick delta time, which means it might tick multiple times per frame or might skip ticks depending on the global tick rate.\nThe selected effect type does not allow fixed tick delta times."),
+				NiagaraValidation::GetStackEntry<UNiagaraStackSystemPropertiesItem>(Context.ViewModel->GetSystemStackViewModel())
+			);
+		}
+	}
+}
+
 void UNiagaraValidationRule_SimulationStageBudget::CheckValidity(const FNiagaraValidationContext& Context, TArray<FNiagaraValidationResult>& OutResults) const
 {
 	for (const TSharedRef<FNiagaraEmitterHandleViewModel>& EmitterHandleModel : Context.ViewModel->GetEmitterHandleViewModels())
diff --git a/Engine/Plugins/FX/Niagara/Source/NiagaraEditor/Public/NiagaraValidationRules.h b/Engine/Plugins/FX/Niagara/Source/NiagaraEditor/Public/NiagaraValidationRules.h
--- a/Engine/Plugins/FX/Niagara/Source/NiagaraEditor/Public/NiagaraValidationRules.h
+++ b/Engine/Plugins/FX/Niagara/Source/NiagaraEditor/Public/NiagaraValidationRules.h
@@ -96,6 +96,15 @@ public:
 	virtual void CheckValidity(const FNiagaraValidationContext& Context, TArray<FNiagaraValidationResult>& OutResults) const override;
 };
 
+/** This validation rule is used by modules or effect types to warn that they don't support systems with fixed delta time ticks. */
+UCLASS(Category = "Validation", DisplayName = "No Fixed DT Tick Support")
+class UNiagaraValidationRule_NoFixedDeltaTime : public UNiagaraValidationRule
+{
+	GENERATED_BODY()
+public:
+	virtual void CheckValidity(const FNiagaraValidationContext& Context, TArray<FNiagaraValidationResult>& OutResults) const override;
+};
+
 /** This validation rule can be used to enforce a budget on the number of simulation stages and the iterations that may execute. */
 UCLASS(Category = "Validation", DisplayName = "Simulation Stage Budget")
 class UNiagaraValidationRule_SimulationStageBudget : public UNiagaraValidationRule
