- Ensure systems are complete / deactivated before adjusting the renderer information

#rnx

diff --git a/Engine/Plugins/FX/Niagara/Source/Niagara/Private/NiagaraEmitter.cpp b/Engine/Plugins/FX/Niagara/Source/Niagara/Private/NiagaraEmitter.cpp
--- a/Engine/Plugins/FX/Niagara/Source/Niagara/Private/NiagaraEmitter.cpp
+++ b/Engine/Plugins/FX/Niagara/Source/Niagara/Private/NiagaraEmitter.cpp
@@ -2820,6 +2820,13 @@ bool UNiagaraEmitter::SetUniqueEmitterName(const FString& InName)
 
 void UNiagaraEmitter::AddRenderer(UNiagaraRendererProperties* Renderer, FGuid EmitterVersion)
 {
+	FNiagaraSystemUpdateContext UpdateContext;
+	UpdateContext.SetDestroyOnAdd(true);
+	if (UNiagaraSystem* Owner = GetTypedOuter<UNiagaraSystem>())
+	{
+		UpdateContext.Add(Owner, true);
+	}
+
 	Modify();
 	Renderer->OuterEmitterVersion = EmitterVersion;
 	FVersionedNiagaraEmitterData* EmitterData = GetEmitterData(EmitterVersion);
@@ -2830,17 +2837,17 @@ void UNiagaraEmitter::AddRenderer(UNiagaraRendererProperties* Renderer, FGuid Em
 	OnRenderersChangedDelegate.Broadcast();
 #endif
 	EmitterData->RebuildRendererBindings(*this);
+}
 
+void UNiagaraEmitter::RemoveRenderer(UNiagaraRendererProperties* Renderer, FGuid EmitterVersion)
+{
+	FNiagaraSystemUpdateContext UpdateContext;
+	UpdateContext.SetDestroyOnAdd(true);
 	if (UNiagaraSystem* Owner = GetTypedOuter<UNiagaraSystem>())
 	{
-		FNiagaraSystemUpdateContext UpdateContext;
 		UpdateContext.Add(Owner, true);
-		UpdateContext.CommitUpdate();
 	}
-}
 
-void UNiagaraEmitter::RemoveRenderer(UNiagaraRendererProperties* Renderer, FGuid EmitterVersion)
-{
 	Modify();
 	FVersionedNiagaraEmitterData* EmitterData = GetEmitterData(EmitterVersion);
 	EmitterData->RendererProperties.Remove(Renderer);
@@ -2850,11 +2857,6 @@ void UNiagaraEmitter::RemoveRenderer(UNiagaraRendererProperties* Renderer, FGuid
 	OnRenderersChangedDelegate.Broadcast();
 #endif
 	EmitterData->RebuildRendererBindings(*this);
-
-	if (UNiagaraSystem* Owner = GetTypedOuter<UNiagaraSystem>())
-	{
-		Owner->ComputeRenderersDrawOrder();
-	}
 }
 
 FNiagaraEventScriptProperties* FVersionedNiagaraEmitterData::GetEventHandlerByIdUnsafe(FGuid ScriptUsageId)
