- Convert RasterizationGrid to shader params / RDG

diff --git a/Engine/Plugins/FX/Niagara/Source/Niagara/Classes/NiagaraDataInterfaceRasterizationGrid3D.h b/Engine/Plugins/FX/Niagara/Source/Niagara/Classes/NiagaraDataInterfaceRasterizationGrid3D.h
--- a/Engine/Plugins/FX/Niagara/Source/Niagara/Classes/NiagaraDataInterfaceRasterizationGrid3D.h
+++ b/Engine/Plugins/FX/Niagara/Source/Niagara/Classes/NiagaraDataInterfaceRasterizationGrid3D.h
@@ -4,6 +4,7 @@
 #include "NiagaraDataInterface.h"
 #include "NiagaraCommon.h"
 #include "NiagaraDataInterfaceRW.h"
+#include "NiagaraRenderGraphUtils.h"
 #include "ClearQuad.h"
 #include "NiagaraStats.h"
 
@@ -18,14 +19,6 @@ class FNiagaraSystemInstance;
 class RasterizationGrid3DRWInstanceData
 {
 public:
-
-	~RasterizationGrid3DRWInstanceData()
-	{
-#if STATS
-		DEC_MEMORY_STAT_BY(STAT_NiagaraGPUDataInterfaceMemory, GPUMemory);
-#endif
-	}
-
 	void ResizeBuffers();
 	int32 TotalNumAttributes = 0;
 	FIntVector NumCells = FIntVector::ZeroValue;
@@ -35,18 +28,16 @@ public:
 
 	bool NeedsRealloc = false;
 	
-	FTextureRWBuffer RasterizationBuffer;
+	FNiagaraPooledRWTexture RasterizationTexture;
 
 	FReadBuffer PerAttributeData;
-
-#if STATS
-	int32 GPUMemory = 0;
-#endif
 };
 
 struct FNiagaraDataInterfaceProxyRasterizationGrid3D : public FNiagaraDataInterfaceProxyRW
 {	
-	virtual void PreStage(FRHICommandList& RHICmdList, const FNiagaraDataInterfaceStageArgs& Context) override;
+	virtual void PreStage(const FNDIGpuComputePreStageContext& Context) override;
+	virtual void PostSimulate(const FNDIGpuComputePostSimulateContext& Context) override;
+
 	virtual void ConsumePerInstanceDataFromGameThread(void* PerInstanceData, const FNiagaraSystemInstanceID& Instance) override {}
 	virtual int32 PerInstanceDataPassedToRenderThreadSize() const override { return sizeof(RasterizationGrid3DRWInstanceData); }	
 
@@ -75,9 +66,6 @@ public:
 	int ResetValue;
 
 public:
-
-	DECLARE_NIAGARA_DI_PARAMETER();
-
 	virtual void PostInitProperties() override
 	{
 		Super::PostInitProperties();
@@ -102,9 +90,13 @@ public:
 
 	// GPU sim functionality
 #if WITH_EDITORONLY_DATA
+	virtual bool AppendCompileHash(FNiagaraCompileHashVisitor* InVisitor) const override;
 	virtual void GetParameterDefinitionHLSL(const FNiagaraDataInterfaceGPUParamInfo& ParamInfo, FString& OutHLSL) override;
 	virtual bool GetFunctionHLSL(const FNiagaraDataInterfaceGPUParamInfo& ParamInfo, const FNiagaraDataInterfaceGeneratedFunction& FunctionInfo, int FunctionInstanceIndex, FString& OutHLSL) override;
 #endif
+	virtual bool UseLegacyShaderBindings() const  override { return false; }
+	virtual void BuildShaderParameters(FNiagaraShaderParametersBuilder& ShaderParametersBuilder) const override;
+	virtual void SetShaderParameters(const FNiagaraDataInterfaceSetShaderParametersContext& Context) const override;
 #if WITH_EDITOR
 	virtual ENiagaraGpuDispatchType GetGpuDispatchType() const override { return ENiagaraGpuDispatchType::ThreeD; }
 #endif
@@ -119,14 +111,9 @@ public:
 	virtual bool HasPreSimulateTick() const override { return true; }
 	//~ UNiagaraDataInterface interface END
 
-	void GetNumCells(FVectorVMExternalFunctionContext& Context);
-	void SetNumCells(FVectorVMExternalFunctionContext& Context);
-	void SetFloatResetValue(FVectorVMExternalFunctionContext& Context);
-
-	static const FName SetNumCellsFunctionName;
-	static const FName SetFloatResetValueFunctionName;	
-
-	static const FString PerAttributeDataName;
+	void VMGetNumCells(FVectorVMExternalFunctionContext& Context);
+	void VMSetNumCells(FVectorVMExternalFunctionContext& Context);
+	void VMSetFloatResetValue(FVectorVMExternalFunctionContext& Context);
 
 protected:
 	//~ UNiagaraDataInterface interface
diff --git a/Engine/Plugins/FX/Niagara/Source/Niagara/Private/NiagaraDataInterfaceRasterizationGrid3D.cpp b/Engine/Plugins/FX/Niagara/Source/Niagara/Private/NiagaraDataInterfaceRasterizationGrid3D.cpp
--- a/Engine/Plugins/FX/Niagara/Source/Niagara/Private/NiagaraDataInterfaceRasterizationGrid3D.cpp
+++ b/Engine/Plugins/FX/Niagara/Source/Niagara/Private/NiagaraDataInterfaceRasterizationGrid3D.cpp
@@ -3,174 +3,63 @@
 #include "NiagaraShader.h"
 #include "ShaderParameterUtils.h"
 #include "NiagaraGpuComputeDispatchInterface.h"
+#include "NiagaraShaderParametersBuilder.h"
 #include "NiagaraSystemInstance.h"
 #include "NiagaraRenderer.h"
 #include "NiagaraShaderParticleID.h"
 
 #define LOCTEXT_NAMESPACE "NiagaraDataInterfaceRasterizationGrid3D"
 
-static const FString IntGridName(TEXT("_IntGrid"));
-static const FString OutputIntGridName(TEXT("_OutputIntGrid"));
-static const FString PrecisionName(TEXT("_Precision"));
-
-const FName UNiagaraDataInterfaceRasterizationGrid3D::SetNumCellsFunctionName("SetNumCells");
-const FName UNiagaraDataInterfaceRasterizationGrid3D::SetFloatResetValueFunctionName("SetFloatResetValue");
-
-const FString UNiagaraDataInterfaceRasterizationGrid3D::PerAttributeDataName(TEXT("_PerAttributeDataName"));
-
-// Global VM function names, also used by the shaders code generation methods.
-
-static const FName SetFloatValueFunctionName("SetFloatGridValue");
-static const FName GetFloatValueFunctionName("GetFloatGridValue");
-
-static const FName InterlockedAddFloatGridValueFunctionName("InterlockedAddFloatGridValue");
-static const FName InterlockedAddIntGridValueFunctionName("InterlockedAddIntGridValue");
-static const FName InterlockedAddFloatGridValueSafeFunctionName("InterlockedAddFloatGridValueSafe");
-
-
-static const FName InterlockedMinFloatGridValueFunctionName("InterlockedMinFloatGridValue");
-static const FName InterlockedMaxFloatGridValueFunctionName("InterlockedMaxFloatGridValue");
-
-
-
-
-static const FName IntToFloatFunctionName("IntToFloat");
-static const FName FloatToIntFunctionName("FloatToInt");
-
-
-static const FName SetIntValueFunctionName("SetIntGridValue");
-static const FName GetIntValueFunctionName("GetIntGridValue");
-
-
-static int32 GMaxNiagaraRasterizationGridCells = (1024 * 1024 * 1024);
-static FAutoConsoleVariableRef CVarMaxNiagaraRasterizationGridCells(
-	TEXT("fx.MaxNiagaraRasterizationGridCells"),
-	GMaxNiagaraRasterizationGridCells,
-	TEXT("The max number of supported grid cells in Niagara. Overflowing this threshold will cause the sim to warn and fail. \n"),
-	ECVF_Default
-);
-
-
-/*--------------------------------------------------------------------------------------------------------------------------*/
-struct FNiagaraDataInterfaceParametersCS_RasterizationGrid3D : public FNiagaraDataInterfaceParametersCS
+namespace NDIRasterizationGrid3DLocal
 {
-	DECLARE_TYPE_LAYOUT(FNiagaraDataInterfaceParametersCS_RasterizationGrid3D, NonVirtual);
-public:
-	void Bind(const FNiagaraDataInterfaceGPUParamInfo& ParameterInfo, const class FShaderParameterMap& ParameterMap)
-	{
-		NumAttributesParam.Bind(ParameterMap, *(ParameterInfo.DataInterfaceHLSLSymbol + UNiagaraDataInterfaceRWBase::NumAttributesName));
-		NumCellsParam.Bind(ParameterMap, *(ParameterInfo.DataInterfaceHLSLSymbol + UNiagaraDataInterfaceRWBase::NumCellsName));
-		UnitToUVParam.Bind(ParameterMap, *(ParameterInfo.DataInterfaceHLSLSymbol + UNiagaraDataInterfaceRWBase::UnitToUVName));
-		PrecisionParam.Bind(ParameterMap, *(ParameterInfo.DataInterfaceHLSLSymbol + PrecisionName));
-				
-		IntGridParam.Bind(ParameterMap,  *(ParameterInfo.DataInterfaceHLSLSymbol + IntGridName));
-		
-		OutputIntGridParam.Bind(ParameterMap, *(ParameterInfo.DataInterfaceHLSLSymbol + OutputIntGridName));
+	BEGIN_SHADER_PARAMETER_STRUCT(FShaderParameters,)
+		SHADER_PARAMETER(FIntVector,						NumCells)
+		SHADER_PARAMETER(FVector3f,							UnitToUVParam)
+		SHADER_PARAMETER(float,								Precision)
+		SHADER_PARAMETER(int,								NumAttributes)
+		SHADER_PARAMETER_RDG_TEXTURE_SRV(Texture3D<int>,	IntGrid)
+		SHADER_PARAMETER_RDG_TEXTURE_UAV(RWTexture3D<int>,	OutputIntGrid)
+		SHADER_PARAMETER_SRV(Buffer<float4>,				PerAttributeData)
+	END_SHADER_PARAMETER_STRUCT()
 
-		PerAttributeDataParam.Bind(ParameterMap, *(ParameterInfo.DataInterfaceHLSLSymbol + UNiagaraDataInterfaceRasterizationGrid3D::PerAttributeDataName));
-	}
+	static const FString IntGridName(TEXT("_IntGrid"));
+	static const FString OutputIntGridName(TEXT("_OutputIntGrid"));
+	static const FString PrecisionName(TEXT("_Precision"));
 
-	// #todo(dmp): make resource transitions batched
-	void Set(FRHICommandList& RHICmdList, const FNiagaraDataInterfaceSetArgs& Context) const
-	{
-		check(IsInRenderingThread());
-
-		// Get shader and DI
-		FRHIComputeShader* ComputeShaderRHI = Context.Shader.GetComputeShader();
-		FNiagaraDataInterfaceProxyRasterizationGrid3D* VFDI = static_cast<FNiagaraDataInterfaceProxyRasterizationGrid3D*>(Context.DataInterface);
-
-		RasterizationGrid3DRWInstanceData* ProxyData = VFDI->SystemInstancesToProxyData.Find(Context.SystemInstanceID);	
-
-		if (!(ProxyData && ProxyData->RasterizationBuffer.Buffer.IsValid()))
-		{			
-			SetShaderValue(RHICmdList, ComputeShaderRHI, NumAttributesParam, 0);
-			SetShaderValue(RHICmdList, ComputeShaderRHI, NumCellsParam, 0);
-			SetShaderValue(RHICmdList, ComputeShaderRHI, UnitToUVParam, FVector3f::ZeroVector);						
-			SetShaderValue(RHICmdList, ComputeShaderRHI, PrecisionParam, 0);
-			SetSRVParameter(RHICmdList, ComputeShaderRHI, IntGridParam, FNiagaraRenderer::GetDummyIntBuffer());
-			SetSRVParameter(RHICmdList, ComputeShaderRHI, PerAttributeDataParam, FNiagaraRenderer::GetDummyIntBuffer());
-			if (OutputIntGridParam.IsUAVBound())
-			{
-				RHICmdList.SetUAVParameter(ComputeShaderRHI, OutputIntGridParam.GetUAVIndex(), Context.ComputeDispatchInterface->GetEmptyUAVFromPool(RHICmdList, PF_R32_SINT, ENiagaraEmptyUAVType::Buffer));
-			}
+	static const FName SetNumCellsFunctionName("SetNumCells");
+	static const FName SetFloatResetValueFunctionName("SetFloatResetValue");
 
-			return;
-		}
+	static const FString PerAttributeDataName(TEXT("_PerAttributeData"));
 
-		SetShaderValue(RHICmdList, ComputeShaderRHI, NumAttributesParam, ProxyData->TotalNumAttributes);
-		SetShaderValue(RHICmdList, ComputeShaderRHI, NumCellsParam, ProxyData->NumCells);
-		SetShaderValue(RHICmdList, ComputeShaderRHI, UnitToUVParam, FVector3f(1.0f) / FVector3f(ProxyData->NumCells));
-		SetShaderValue(RHICmdList, ComputeShaderRHI, PrecisionParam, ProxyData->Precision);
-		SetSRVParameter(RHICmdList, ComputeShaderRHI, PerAttributeDataParam, ProxyData->PerAttributeData.SRV);
+	// Global VM function names, also used by the shaders code generation methods.
 
-		if (!Context.IsOutputStage)
-		{			
-			if (IntGridParam.IsBound())
-			{
-				RHICmdList.Transition(FRHITransitionInfo(ProxyData->RasterizationBuffer.UAV, ERHIAccess::Unknown, ERHIAccess::SRVCompute));
-				SetSRVParameter(RHICmdList, ComputeShaderRHI, IntGridParam, ProxyData->RasterizationBuffer.SRV);
-			}
+	static const FName SetFloatValueFunctionName("SetFloatGridValue");
+	static const FName GetFloatValueFunctionName("GetFloatGridValue");
 
-			if (OutputIntGridParam.IsUAVBound())
-			{
-				RHICmdList.SetUAVParameter(ComputeShaderRHI, OutputIntGridParam.GetUAVIndex(), Context.ComputeDispatchInterface->GetEmptyUAVFromPool(RHICmdList, PF_R32_SINT, ENiagaraEmptyUAVType::Buffer));
-			}
-		}
-		else
-		{			
-			SetSRVParameter(RHICmdList, ComputeShaderRHI, IntGridParam, FNiagaraRenderer::GetDummyIntBuffer());
+	static const FName InterlockedAddFloatGridValueFunctionName("InterlockedAddFloatGridValue");
+	static const FName InterlockedAddIntGridValueFunctionName("InterlockedAddIntGridValue");
+	static const FName InterlockedAddFloatGridValueSafeFunctionName("InterlockedAddFloatGridValueSafe");
 
-			if (OutputIntGridParam.IsBound())
-			{
-				RHICmdList.Transition(FRHITransitionInfo(ProxyData->RasterizationBuffer.UAV, ERHIAccess::Unknown, ERHIAccess::UAVCompute));
 
-				//inline void SetTexture(TRHICmdList & RHICmdList, const TShaderRHIRef & Shader, FRHITexture * Texture, FRHIUnorderedAccessView * UAV) const;
-				//OutputIntGridParam.SetBuffer(RHICmdList, ComputeShaderRHI, ProxyData->RasterizationBuffer.Buffer, ProxyData->RasterizationBuffer.UAV);
-				RHICmdList.SetUAVParameter(ComputeShaderRHI, OutputIntGridParam.GetUAVIndex(), ProxyData->RasterizationBuffer.UAV);
-			}
-		}
-		// Note: There is a flush in PreEditChange to make sure everything is synced up at this point 
-	}
+	static const FName InterlockedMinFloatGridValueFunctionName("InterlockedMinFloatGridValue");
+	static const FName InterlockedMaxFloatGridValueFunctionName("InterlockedMaxFloatGridValue");
 
-	void Unset(FRHICommandList& RHICmdList, const FNiagaraDataInterfaceSetArgs& Context) const 
-	{
-		if (OutputIntGridParam.IsBound())
-		{
-			OutputIntGridParam.UnsetUAV(RHICmdList, Context.Shader.GetComputeShader());
-		}
-	}
-
-private:
-	LAYOUT_FIELD(FShaderParameter, NumAttributesParam);
-	LAYOUT_FIELD(FShaderParameter, NumCellsParam);
-	LAYOUT_FIELD(FShaderParameter, UnitToUVParam);		
-	LAYOUT_FIELD(FShaderParameter, PrecisionParam);
-	LAYOUT_FIELD(FShaderResourceParameter, IntGridParam);
-	LAYOUT_FIELD(FRWShaderParameter, OutputIntGridParam);	
-	LAYOUT_FIELD(FShaderResourceParameter, PerAttributeDataParam);
-};
-
-void RasterizationGrid3DRWInstanceData::ResizeBuffers()
-{
-	uint32 NumTotalCells = NumCells.X * NumCells.Y * NumCells.Z;
-
-	uint32 NumIntsInGridBuffer = NumTotalCells * 1;
-
-	if (NumTotalCells > (uint32)GMaxNiagaraRasterizationGridCells)
-		return;
+	static const FName IntToFloatFunctionName("IntToFloat");
+	static const FName FloatToIntFunctionName("FloatToInt");
 
-	RasterizationBuffer.Initialize3D(TEXT("NiagaraRasterizationGrid3D::IntGrid"), sizeof(int32), NumCells.X * NumTiles.X, NumCells.Y * NumTiles.Y, NumCells.Z * NumTiles.Z, EPixelFormat::PF_R32_SINT);
+	static const FName SetIntValueFunctionName("SetIntGridValue");
+	static const FName GetIntValueFunctionName("GetIntGridValue");
 
-	#if STATS
-		DEC_MEMORY_STAT_BY(STAT_NiagaraGPUDataInterfaceMemory, GPUMemory);
-		GPUMemory = NumTotalCells * sizeof(int32) + NumIntsInGridBuffer * sizeof(int32);
-		INC_MEMORY_STAT_BY(STAT_NiagaraGPUDataInterfaceMemory, GPUMemory);
-	#endif
+	static int32 GMaxNiagaraRasterizationGridCells = (1024 * 1024 * 1024);
+	static FAutoConsoleVariableRef CVarMaxNiagaraRasterizationGridCells(
+		TEXT("fx.MaxNiagaraRasterizationGridCells"),
+		GMaxNiagaraRasterizationGridCells,
+		TEXT("The max number of supported grid cells in Niagara. Overflowing this threshold will cause the sim to warn and fail. \n"),
+		ECVF_Default
+	);
 }
 
-IMPLEMENT_TYPE_LAYOUT(FNiagaraDataInterfaceParametersCS_RasterizationGrid3D);
-
-IMPLEMENT_NIAGARA_DI_PARAMETER(UNiagaraDataInterfaceRasterizationGrid3D, FNiagaraDataInterfaceParametersCS_RasterizationGrid3D);
+/*--------------------------------------------------------------------------------------------------------------------------*/
 
 UNiagaraDataInterfaceRasterizationGrid3D::UNiagaraDataInterfaceRasterizationGrid3D(FObjectInitializer const& ObjectInitializer)
 	: Super(ObjectInitializer)
@@ -181,9 +70,10 @@ UNiagaraDataInterfaceRasterizationGrid3D::UNiagaraDataInterfaceRasterizationGrid
 	Proxy.Reset(new FNiagaraDataInterfaceProxyRasterizationGrid3D());	
 }
 
-
 void UNiagaraDataInterfaceRasterizationGrid3D::GetFunctions(TArray<FNiagaraFunctionSignature>& OutFunctions)
 {
+	using namespace NDIRasterizationGrid3DLocal;
+
 	Super::GetFunctions(OutFunctions);
 
 	{
@@ -529,32 +419,32 @@ void UNiagaraDataInterfaceRasterizationGrid3D::GetFunctions(TArray<FNiagaraFunct
 	}
 }
 
-DEFINE_NDI_DIRECT_FUNC_BINDER(UNiagaraDataInterfaceRasterizationGrid3D, SetNumCells);
-DEFINE_NDI_DIRECT_FUNC_BINDER(UNiagaraDataInterfaceRasterizationGrid3D, SetFloatResetValue);
 void UNiagaraDataInterfaceRasterizationGrid3D::GetVMExternalFunction(const FVMExternalFunctionBindingInfo& BindingInfo, void* InstanceData, FVMExternalFunction &OutFunc)
 {
+	using namespace NDIRasterizationGrid3DLocal;
+
 	Super::GetVMExternalFunction(BindingInfo, InstanceData, OutFunc);
 
 	// #todo(dmp): this overrides the empty function set by the super class
 	if (BindingInfo.Name == UNiagaraDataInterfaceRWBase::NumCellsFunctionName)
 	{
 		check(BindingInfo.GetNumInputs() == 1 && BindingInfo.GetNumOutputs() == 3);
-		OutFunc = FVMExternalFunction::CreateLambda([&](FVectorVMExternalFunctionContext& Context) { GetNumCells(Context); });
+		OutFunc.BindUObject(this, &UNiagaraDataInterfaceRasterizationGrid3D::VMGetNumCells);
 	}
 	else if (BindingInfo.Name == SetNumCellsFunctionName)
 	{
 		check(BindingInfo.GetNumInputs() == 4 && BindingInfo.GetNumOutputs() == 1);
-		NDI_FUNC_BINDER(UNiagaraDataInterfaceRasterizationGrid3D, SetNumCells)::Bind(this, OutFunc);
+		OutFunc.BindUObject(this, &UNiagaraDataInterfaceRasterizationGrid3D::VMSetNumCells);
 	}
 	else if (BindingInfo.Name == SetFloatResetValueFunctionName)
 	{
 		check(BindingInfo.GetNumInputs() == 2 && BindingInfo.GetNumOutputs() == 1);
-		NDI_FUNC_BINDER(UNiagaraDataInterfaceRasterizationGrid3D, SetFloatResetValue)::Bind(this, OutFunc);
+		OutFunc.BindUObject(this, &UNiagaraDataInterfaceRasterizationGrid3D::VMSetFloatResetValue);
 	}
 }
 
 
-void UNiagaraDataInterfaceRasterizationGrid3D::GetNumCells(FVectorVMExternalFunctionContext& Context)
+void UNiagaraDataInterfaceRasterizationGrid3D::VMGetNumCells(FVectorVMExternalFunctionContext& Context)
 {
 	VectorVM::FUserPtrHandler<RasterizationGrid3DRWInstanceData> InstData(Context);
 
@@ -595,13 +485,22 @@ bool UNiagaraDataInterfaceRasterizationGrid3D::ShouldCompile(EShaderPlatform Sha
 #endif
 
 #if WITH_EDITORONLY_DATA
+bool UNiagaraDataInterfaceRasterizationGrid3D::AppendCompileHash(FNiagaraCompileHashVisitor* InVisitor) const
+{
+	bool bSuccess = Super::AppendCompileHash(InVisitor);
+	bSuccess &= InVisitor->UpdateShaderParameters<NDIRasterizationGrid3DLocal::FShaderParameters>();
+	return bSuccess;
+}
+
 void UNiagaraDataInterfaceRasterizationGrid3D::GetParameterDefinitionHLSL(const FNiagaraDataInterfaceGPUParamInfo& ParamInfo, FString& OutHLSL)
 {
 	Super::GetParameterDefinitionHLSL(ParamInfo, OutHLSL);
 
+	using namespace NDIRasterizationGrid3DLocal;
+
 	static const TCHAR *FormatDeclarations = TEXT(R"(			
 		Texture3D<int> {IntGridName};		
-		RWTexture3D<int> RW{OutputIntGridName};
+		RWTexture3D<int> {OutputIntGridName};
 		float {Precision};
 		Buffer<float4> {PerAttributeDataName};
 		int {NumAttributesName};
@@ -654,6 +553,8 @@ void UNiagaraDataInterfaceRasterizationGrid3D::GetParameterDefinitionHLSL(const
 
 bool UNiagaraDataInterfaceRasterizationGrid3D::GetFunctionHLSL(const FNiagaraDataInterfaceGPUParamInfo& ParamInfo, const FNiagaraDataInterfaceGeneratedFunction& FunctionInfo, int FunctionInstanceIndex, FString& OutHLSL)
 {
+	using namespace NDIRasterizationGrid3DLocal;
+
 	bool ParentRet = Super::GetFunctionHLSL(ParamInfo, FunctionInfo, FunctionInstanceIndex, OutHLSL);
 
 	TMap<FString, FStringFormatArg> ArgsBounds =
@@ -696,7 +597,7 @@ bool UNiagaraDataInterfaceRasterizationGrid3D::GetFunctionHLSL(const FNiagaraDat
 				if ( In_AttributeIndex < {NumAttributesName} )
 				{
 					int3 TileOffset = {PerAttributeDataName}[In_AttributeIndex].xyz;
-					RW{OutputIntGrid}[int3(In_IndexX + TileOffset.x, In_IndexY + TileOffset.y, In_IndexZ + TileOffset.z)] = {FloatToIntFunctionName}(In_Value);
+					{OutputIntGrid}[int3(In_IndexX + TileOffset.x, In_IndexY + TileOffset.y, In_IndexZ + TileOffset.z)] = {FloatToIntFunctionName}(In_Value);
 				}
 			}
 		)");
@@ -713,7 +614,7 @@ bool UNiagaraDataInterfaceRasterizationGrid3D::GetFunctionHLSL(const FNiagaraDat
 				if ( In_AttributeIndex < {NumAttributesName} )
 				{	
 					int3 TileOffset = {PerAttributeDataName}[In_AttributeIndex].xyz;			
-					RW{OutputIntGrid}[int3(In_IndexX + TileOffset.x, In_IndexY + TileOffset.y, In_IndexZ + TileOffset.z)] = In_Value;				
+					{OutputIntGrid}[int3(In_IndexX + TileOffset.x, In_IndexY + TileOffset.y, In_IndexZ + TileOffset.z)] = In_Value;				
 				}
 			}
 		)");
@@ -729,7 +630,7 @@ bool UNiagaraDataInterfaceRasterizationGrid3D::GetFunctionHLSL(const FNiagaraDat
 				if ( In_AttributeIndex < {NumAttributesName} )
 				{
 					int3 TileOffset = {PerAttributeDataName}[In_AttributeIndex].xyz;		
-					InterlockedAdd(RW{OutputIntGrid}[int3(In_IndexX + TileOffset.x, In_IndexY + TileOffset.y, In_IndexZ + TileOffset.z)], {FloatToIntFunctionName}(In_Value), val);					
+					InterlockedAdd({OutputIntGrid}[int3(In_IndexX + TileOffset.x, In_IndexY + TileOffset.y, In_IndexZ + TileOffset.z)], {FloatToIntFunctionName}(In_Value), val);					
 				}
 			}
 
@@ -740,7 +641,7 @@ bool UNiagaraDataInterfaceRasterizationGrid3D::GetFunctionHLSL(const FNiagaraDat
 				{
 					int3 TileOffset = {PerAttributeDataName}[In_AttributeIndex].xyz;		
 					int OriginalIntValue;
-					InterlockedAdd(RW{OutputIntGrid}[int3(In_IndexX + TileOffset.x, In_IndexY + TileOffset.y, In_IndexZ + TileOffset.z)], {FloatToIntFunctionName}(In_Value), OriginalIntValue);
+					InterlockedAdd({OutputIntGrid}[int3(In_IndexX + TileOffset.x, In_IndexY + TileOffset.y, In_IndexZ + TileOffset.z)], {FloatToIntFunctionName}(In_Value), OriginalIntValue);
 					Original_Value = {IntToFloatFunctionName}(OriginalIntValue);
 				}
 			}
@@ -757,7 +658,7 @@ bool UNiagaraDataInterfaceRasterizationGrid3D::GetFunctionHLSL(const FNiagaraDat
 				if ( In_AttributeIndex < {NumAttributesName} )
 				{
 					int3 TileOffset = {PerAttributeDataName}[In_AttributeIndex].xyz;		
-					InterlockedAdd(RW{OutputIntGrid}[int3(In_IndexX + TileOffset.x, In_IndexY + TileOffset.y, In_IndexZ + TileOffset.z)], In_Value);
+					InterlockedAdd({OutputIntGrid}[int3(In_IndexX + TileOffset.x, In_IndexY + TileOffset.y, In_IndexZ + TileOffset.z)], In_Value);
 				}
 			}
 		)");
@@ -777,7 +678,7 @@ bool UNiagaraDataInterfaceRasterizationGrid3D::GetFunctionHLSL(const FNiagaraDat
 					int3 TileOffset = {PerAttributeDataName}[In_AttributeIndex].xyz;	
 
 					int3 Index = int3(In_IndexX + TileOffset.x, In_IndexY + TileOffset.y, In_IndexZ + TileOffset.z);
-					InterlockedAdd(RW{OutputIntGrid}[Index], IntValue, OriginalValue);
+					InterlockedAdd({OutputIntGrid}[Index], IntValue, OriginalValue);
 
 					int StoredValue = IntValue + OriginalValue;
 
@@ -788,7 +689,7 @@ bool UNiagaraDataInterfaceRasterizationGrid3D::GetFunctionHLSL(const FNiagaraDat
 						val = 0;
 
 						int NewValue = IntValue > 0 ? 2147483647 : -2147483648;
-						RW{OutputIntGrid}[Index] = In_Value;
+						{OutputIntGrid}[Index] = In_Value;
 					}
 				}
 			}
@@ -806,7 +707,7 @@ bool UNiagaraDataInterfaceRasterizationGrid3D::GetFunctionHLSL(const FNiagaraDat
 				if ( In_AttributeIndex < {NumAttributesName} )
 				{
 					int3 TileOffset = {PerAttributeDataName}[In_AttributeIndex].xyz;			
-					InterlockedMin(RW{OutputIntGrid}[int3(In_IndexX + TileOffset.x, In_IndexY + TileOffset.y, In_IndexZ + TileOffset.z)], {FloatToIntFunctionName}(In_Value), val);
+					InterlockedMin({OutputIntGrid}[int3(In_IndexX + TileOffset.x, In_IndexY + TileOffset.y, In_IndexZ + TileOffset.z)], {FloatToIntFunctionName}(In_Value), val);
 				}
 			}
 
@@ -818,7 +719,7 @@ bool UNiagaraDataInterfaceRasterizationGrid3D::GetFunctionHLSL(const FNiagaraDat
 					int3 TileOffset = {PerAttributeDataName}[In_AttributeIndex].xyz;			
 
 					int OriginalIntValue;
-					InterlockedMin(RW{OutputIntGrid}[int3(In_IndexX + TileOffset.x, In_IndexY + TileOffset.y, In_IndexZ + TileOffset.z)], {FloatToIntFunctionName}(In_Value), OriginalIntValue);
+					InterlockedMin({OutputIntGrid}[int3(In_IndexX + TileOffset.x, In_IndexY + TileOffset.y, In_IndexZ + TileOffset.z)], {FloatToIntFunctionName}(In_Value), OriginalIntValue);
 					Original_Value = {IntToFloatFunctionName}(OriginalIntValue);
 				}
 			}
@@ -835,7 +736,7 @@ bool UNiagaraDataInterfaceRasterizationGrid3D::GetFunctionHLSL(const FNiagaraDat
 				if ( In_AttributeIndex < {NumAttributesName} )
 				{
 					int3 TileOffset = {PerAttributeDataName}[In_AttributeIndex].xyz;			
-					InterlockedMax(RW{OutputIntGrid}[int3(In_IndexX + TileOffset.x, In_IndexY + TileOffset.y, In_IndexZ + TileOffset.z)], {FloatToIntFunctionName}(In_Value), val);
+					InterlockedMax({OutputIntGrid}[int3(In_IndexX + TileOffset.x, In_IndexY + TileOffset.y, In_IndexZ + TileOffset.z)], {FloatToIntFunctionName}(In_Value), val);
 				}
 			}
 
@@ -847,7 +748,7 @@ bool UNiagaraDataInterfaceRasterizationGrid3D::GetFunctionHLSL(const FNiagaraDat
 					int3 TileOffset = {PerAttributeDataName}[In_AttributeIndex].xyz;			
 
 					int OriginalIntValue;
-					InterlockedMax(RW{OutputIntGrid}[int3(In_IndexX + TileOffset.x, In_IndexY + TileOffset.y, In_IndexZ + TileOffset.z)], {FloatToIntFunctionName}(In_Value), OriginalIntValue);
+					InterlockedMax({OutputIntGrid}[int3(In_IndexX + TileOffset.x, In_IndexY + TileOffset.y, In_IndexZ + TileOffset.z)], {FloatToIntFunctionName}(In_Value), OriginalIntValue);
 					Original_Value = {IntToFloatFunctionName}(OriginalIntValue);
 				}
 			}
@@ -891,8 +792,69 @@ bool UNiagaraDataInterfaceRasterizationGrid3D::GetFunctionHLSL(const FNiagaraDat
 }
 #endif
 
+void UNiagaraDataInterfaceRasterizationGrid3D::BuildShaderParameters(FNiagaraShaderParametersBuilder& ShaderParametersBuilder) const
+{
+	ShaderParametersBuilder.AddNestedStruct<NDIRasterizationGrid3DLocal::FShaderParameters>();
+}
+
+void UNiagaraDataInterfaceRasterizationGrid3D::SetShaderParameters(const FNiagaraDataInterfaceSetShaderParametersContext& Context) const
+{
+	FNiagaraDataInterfaceProxyRasterizationGrid3D& DIProxy = Context.GetProxy<FNiagaraDataInterfaceProxyRasterizationGrid3D>();
+	RasterizationGrid3DRWInstanceData* InstanceData = DIProxy.SystemInstancesToProxyData.Find(Context.GetSystemInstanceID());
+	NDIRasterizationGrid3DLocal::FShaderParameters* ShaderParameters = Context.GetParameterNestedStruct<NDIRasterizationGrid3DLocal::FShaderParameters>();
+	FRDGBuilder& GraphBuilder = Context.GetGraphBuilder();
+
+	const bool bBindInput = Context.IsResourceBound(&ShaderParameters->IntGrid);
+	const bool bBindOutput = Context.IsResourceBound(&ShaderParameters->OutputIntGrid);
+	bool bInputBound = false;
+	bool bOutputBound = false;
+	if (InstanceData && InstanceData->RasterizationTexture.IsValid())
+	{
+		ShaderParameters->NumCells = InstanceData->NumCells;
+		ShaderParameters->UnitToUVParam = FVector3f(1.0f) / FVector3f(InstanceData->NumCells);
+		ShaderParameters->Precision = InstanceData->Precision;
+		ShaderParameters->NumAttributes = InstanceData->TotalNumAttributes;
+		ShaderParameters->PerAttributeData = InstanceData->PerAttributeData.SRV;
+
+		if (Context.IsOutputStage())
+		{
+			if (bBindOutput)
+			{
+				bOutputBound = true;
+				ShaderParameters->OutputIntGrid = InstanceData->RasterizationTexture.GetOrCreateUAV(GraphBuilder);
+			}
+		}
+		else
+		{
+			if (bBindInput)
+			{
+				bInputBound = true;
+				ShaderParameters->IntGrid = InstanceData->RasterizationTexture.GetOrCreateSRV(GraphBuilder);
+			}
+		}
+	}
+	else
+	{
+		ShaderParameters->NumCells = FIntVector::ZeroValue;
+		ShaderParameters->UnitToUVParam = FVector3f::ZeroVector;
+		ShaderParameters->Precision = 0;
+		ShaderParameters->NumAttributes = 0;
+		ShaderParameters->PerAttributeData = FNiagaraRenderer::GetDummyFloat4Buffer();
+	}
+
+	if (bBindInput && !bInputBound)
+	{
+		ShaderParameters->IntGrid = Context.GetComputeDispatchInterface().GetBlackTextureSRV(GraphBuilder, ETextureDimension::Texture3D);
+	}
+	if (bBindOutput && !bOutputBound)
+	{
+		ShaderParameters->OutputIntGrid = Context.GetComputeDispatchInterface().GetEmptyTextureUAV(GraphBuilder, PF_A32B32G32R32F, ETextureDimension::Texture3D);
+	}
+}
+
 bool UNiagaraDataInterfaceRasterizationGrid3D::InitPerInstanceData(void* PerInstanceData, FNiagaraSystemInstance* SystemInstance)
 {
+	using namespace NDIRasterizationGrid3DLocal;
 
 	RasterizationGrid3DRWInstanceData* InstanceData = new (PerInstanceData) RasterizationGrid3DRWInstanceData();
 
@@ -960,14 +922,13 @@ bool UNiagaraDataInterfaceRasterizationGrid3D::InitPerInstanceData(void* PerInst
 		TargetData->NumTiles = RT_InstanceData.NumTiles;
 		TargetData->Precision = RT_Precision;
 		TargetData->ResetValue = RT_ResetValue;
-
-		TargetData->ResizeBuffers();
+		TargetData->NeedsRealloc = true;
 	});
 
 	return true;
 }
 
-void UNiagaraDataInterfaceRasterizationGrid3D::SetNumCells(FVectorVMExternalFunctionContext& Context)
+void UNiagaraDataInterfaceRasterizationGrid3D::VMSetNumCells(FVectorVMExternalFunctionContext& Context)
 {
 	// This should only be called from a system or emitter script due to a need for only setting up initially.
 	VectorVM::FUserPtrHandler<RasterizationGrid3DRWInstanceData> InstData(Context);
@@ -996,7 +957,7 @@ void UNiagaraDataInterfaceRasterizationGrid3D::SetNumCells(FVectorVMExternalFunc
 	}
 }
 
-void UNiagaraDataInterfaceRasterizationGrid3D::SetFloatResetValue(FVectorVMExternalFunctionContext& Context)
+void UNiagaraDataInterfaceRasterizationGrid3D::VMSetFloatResetValue(FVectorVMExternalFunctionContext& Context)
 {
 	VectorVM::FUserPtrHandler<RasterizationGrid3DRWInstanceData> InstData(Context);
 	VectorVM::FExternalFuncInputHandler<float> InResetValue(Context);
@@ -1057,7 +1018,7 @@ bool UNiagaraDataInterfaceRasterizationGrid3D::PerInstanceTickPostSimulate(void*
 			TargetData->Precision = RT_Precision;
 			TargetData->TotalNumAttributes = RT_InstanceData.TotalNumAttributes;
 			TargetData->ResetValue = RT_ResetValue;
-			TargetData->ResizeBuffers();
+			TargetData->NeedsRealloc = true;
 			TargetData->PerAttributeData.Release();
 		});
 	}
@@ -1077,56 +1038,75 @@ void UNiagaraDataInterfaceRasterizationGrid3D::DestroyPerInstanceData(void* PerI
 
 	ENQUEUE_RENDER_COMMAND(FNiagaraDIDestroyInstanceData) (
 		[ThisProxy, InstanceID = SystemInstance->GetId()](FRHICommandListImmediate& CmdList)
-	{
-		//check(ThisProxy->SystemInstancesToProxyData.Contains(InstanceID));
-		ThisProxy->SystemInstancesToProxyData.Remove(InstanceID);
-	}
+		{
+			//check(ThisProxy->SystemInstancesToProxyData.Contains(InstanceID));
+			ThisProxy->SystemInstancesToProxyData.Remove(InstanceID);
+		}
 	);
 }
 
-void FNiagaraDataInterfaceProxyRasterizationGrid3D::PreStage(FRHICommandList& RHICmdList, const FNiagaraDataInterfaceStageArgs& Context)
+void FNiagaraDataInterfaceProxyRasterizationGrid3D::PreStage(const FNDIGpuComputePreStageContext& Context)
 {
-	RasterizationGrid3DRWInstanceData* ProxyData = SystemInstancesToProxyData.Find(Context.SystemInstanceID);
+	using namespace NDIRasterizationGrid3DLocal;
+
+	RasterizationGrid3DRWInstanceData& InstanceData = SystemInstancesToProxyData.FindChecked(Context.GetSystemInstanceID());
+
+	// Resize was requested
+	if (InstanceData.NeedsRealloc)
+	{
+		InstanceData.NeedsRealloc = false;
+
+		const uint32 NumTotalCells = InstanceData.NumCells.X * InstanceData.NumCells.Y * InstanceData.NumCells.Z;
+		if (NumTotalCells <= (uint32)GMaxNiagaraRasterizationGridCells)
+		{
+			const FIntVector TextureSize(InstanceData.NumCells.X * InstanceData.NumTiles.X, InstanceData.NumCells.Y * InstanceData.NumTiles.Y, InstanceData.NumCells.Z * InstanceData.NumTiles.Z);
+			const FRDGTextureDesc TextureDesc = FRDGTextureDesc::Create3D(TextureSize, PF_R32_SINT, FClearValueBinding::Black, ETextureCreateFlags::ShaderResource | ETextureCreateFlags::UAV);
+			InstanceData.RasterizationTexture.Initialize(Context.GetGraphBuilder(), TEXT("NiagaraRasterizationGrid3D::IntGrid"), TextureDesc);
+		}
+	}
 
 	// Generate per-attribute data
-	if (ProxyData->PerAttributeData.NumBytes == 0)
+	if (InstanceData.PerAttributeData.NumBytes == 0)
 	{
 		TResourceArray<FVector4f> PerAttributeData;
-		PerAttributeData.AddUninitialized((ProxyData->TotalNumAttributes * 2) + 1);
-		for (int32 iAttribute = 0; iAttribute < ProxyData->TotalNumAttributes; ++iAttribute)
+		PerAttributeData.AddUninitialized((InstanceData.TotalNumAttributes * 2) + 1);
+		for (int32 iAttribute = 0; iAttribute < InstanceData.TotalNumAttributes; ++iAttribute)
 		{
-			const FIntVector AttributeTileIndex(iAttribute % ProxyData->NumTiles.X, (iAttribute / ProxyData->NumTiles.X) % ProxyData->NumTiles.Y, iAttribute / (ProxyData->NumTiles.X * ProxyData->NumTiles.Y));
+			const FIntVector AttributeTileIndex(iAttribute % InstanceData.NumTiles.X, (iAttribute / InstanceData.NumTiles.X) % InstanceData.NumTiles.Y, iAttribute / (InstanceData.NumTiles.X * InstanceData.NumTiles.Y));
 			PerAttributeData[iAttribute] = FVector4f(
-				AttributeTileIndex.X * ProxyData->NumCells.X,
-				AttributeTileIndex.Y * ProxyData->NumCells.Y,
-				AttributeTileIndex.Z * ProxyData->NumCells.Z,
+				AttributeTileIndex.X * InstanceData.NumCells.X,
+				AttributeTileIndex.Y * InstanceData.NumCells.Y,
+				AttributeTileIndex.Z * InstanceData.NumCells.Z,
 				0
 			);
-			PerAttributeData[iAttribute + ProxyData->TotalNumAttributes] = FVector4f(
-				(1.0f / ProxyData->NumTiles.X) * float(AttributeTileIndex.X),
-				(1.0f / ProxyData->NumTiles.Y) * float(AttributeTileIndex.Y),
-				(1.0f / ProxyData->NumTiles.Z) * float(AttributeTileIndex.Z),
+			PerAttributeData[iAttribute + InstanceData.TotalNumAttributes] = FVector4f(
+				(1.0f / InstanceData.NumTiles.X) * float(AttributeTileIndex.X),
+				(1.0f / InstanceData.NumTiles.Y) * float(AttributeTileIndex.Y),
+				(1.0f / InstanceData.NumTiles.Z) * float(AttributeTileIndex.Z),
 				0.0f
 			);
 		}
-		PerAttributeData[ProxyData->TotalNumAttributes * 2] = FVector4f(65535, 65535, 65535, 65535);
-		ProxyData->PerAttributeData.Initialize(TEXT("Grid3D::PerAttributeData"), sizeof(FVector4f), PerAttributeData.Num(), EPixelFormat::PF_A32B32G32R32F, BUF_Static, &PerAttributeData);
+		PerAttributeData[InstanceData.TotalNumAttributes * 2] = FVector4f(65535, 65535, 65535, 65535);
+		InstanceData.PerAttributeData.Initialize(TEXT("Grid3D::PerAttributeData"), sizeof(FVector4f), PerAttributeData.Num(), EPixelFormat::PF_A32B32G32R32F, BUF_Static, &PerAttributeData);
 	}
 
-	if (Context.IsOutputStage)
+	if (Context.IsOutputStage())
 	{
-		
-
-		SCOPED_DRAW_EVENT(RHICmdList, NiagaraRasterizationGrid3DClearNeighborInfo);
-		ERHIFeatureLevel::Type FeatureLevel = Context.ComputeDispatchInterface->GetFeatureLevel();
+		FRDGBuilder& GraphBuilder = Context.GetGraphBuilder();
+		const FUintVector4 ResetValue(InstanceData.ResetValue, InstanceData.ResetValue, InstanceData.ResetValue, InstanceData.ResetValue);
+		AddClearUAVPass(GraphBuilder, InstanceData.RasterizationTexture.GetOrCreateUAV(GraphBuilder), ResetValue);
+	}
+}
 
-		FRHITransitionInfo TransitionInfos[] =
-		{			
-			FRHITransitionInfo(ProxyData->RasterizationBuffer.UAV, ERHIAccess::Unknown, ERHIAccess::UAVCompute),
-		};
-		RHICmdList.Transition(MakeArrayView(TransitionInfos, UE_ARRAY_COUNT(TransitionInfos)));		
-					
-		RHICmdList.ClearUAVUint(ProxyData->RasterizationBuffer.UAV, FUintVector4(ProxyData->ResetValue, ProxyData->ResetValue, ProxyData->ResetValue, ProxyData->ResetValue));		
+void FNiagaraDataInterfaceProxyRasterizationGrid3D::PostSimulate(const FNDIGpuComputePostSimulateContext& Context)
+{
+	if (Context.IsFinalPostSimulate())
+	{
+		for ( auto it=SystemInstancesToProxyData.CreateIterator(); it; ++it )
+		{
+			RasterizationGrid3DRWInstanceData& InstanceData = it.Value();
+			InstanceData.RasterizationTexture.EndGraphUsage();
+		}
 	}
 }
 
