- Fix name collision between RigidMesh & PhysicsAsset DI

#rnx

diff --git a/Engine/Plugins/FX/Niagara/Source/Niagara/Private/NiagaraDataInterfaceRigidMeshCollisionQuery.cpp b/Engine/Plugins/FX/Niagara/Source/Niagara/Private/NiagaraDataInterfaceRigidMeshCollisionQuery.cpp
--- a/Engine/Plugins/FX/Niagara/Source/Niagara/Private/NiagaraDataInterfaceRigidMeshCollisionQuery.cpp
+++ b/Engine/Plugins/FX/Niagara/Source/Niagara/Private/NiagaraDataInterfaceRigidMeshCollisionQuery.cpp
@@ -630,7 +630,7 @@ struct FNDIRigidMeshCollisionProxy : public FNiagaraDataInterfaceProxy
 {
 	struct FRenderThreadData
 	{
-		FElementOffset ElementOffsets;
+		FNDIRigidMeshCollisionElementOffset ElementOffsets;
 		TArray<FVector4f> WorldTransform;
 		TArray<FVector4f> InverseTransform;
 		TArray<FVector4f> ElementExtent;
@@ -744,7 +744,7 @@ void UNiagaraDataInterfaceRigidMeshCollisionQuery::DrawDebugHud(UCanvas* Canvas,
 		return;
 	}
 
-	const FElementOffset& ElementOffsets = InstanceData_GT->AssetArrays->ElementOffsets;
+	const FNDIRigidMeshCollisionElementOffset& ElementOffsets = InstanceData_GT->AssetArrays->ElementOffsets;
 
 	const uint32 BoxCount = ElementOffsets.SphereOffset - ElementOffsets.BoxOffset;
 	const uint32 SphereCount = ElementOffsets.CapsuleOffset - ElementOffsets.SphereOffset;
diff --git a/Engine/Plugins/FX/Niagara/Source/Niagara/Public/NiagaraDataInterfaceRigidMeshCollisionQuery.h b/Engine/Plugins/FX/Niagara/Source/Niagara/Public/NiagaraDataInterfaceRigidMeshCollisionQuery.h
--- a/Engine/Plugins/FX/Niagara/Source/Niagara/Public/NiagaraDataInterfaceRigidMeshCollisionQuery.h
+++ b/Engine/Plugins/FX/Niagara/Source/Niagara/Public/NiagaraDataInterfaceRigidMeshCollisionQuery.h
@@ -12,13 +12,13 @@
 class AActor;
 
 /** Element offsets in the array list */
-struct FElementOffset
+struct FNDIRigidMeshCollisionElementOffset
 {
-	FElementOffset(const uint32 InBoxOffset, const uint32 InSphereOffset, const uint32 InCapsuleOffset, const uint32 InNumElements) :
+	FNDIRigidMeshCollisionElementOffset(const uint32 InBoxOffset, const uint32 InSphereOffset, const uint32 InCapsuleOffset, const uint32 InNumElements) :
 		BoxOffset(InBoxOffset), SphereOffset(InSphereOffset), CapsuleOffset(InCapsuleOffset), NumElements(InNumElements)
 	{}
 
-	FElementOffset() :
+	FNDIRigidMeshCollisionElementOffset() :
 		BoxOffset(0), SphereOffset(0), CapsuleOffset(0), NumElements(0)
 	{}
 	uint32 BoxOffset;
@@ -30,7 +30,7 @@ struct FElementOffset
 /** Arrays in which the cpu datas will be str */
 struct FNDIRigidMeshCollisionArrays
 {
-	FElementOffset ElementOffsets;
+	FNDIRigidMeshCollisionElementOffset ElementOffsets;
 	TArray<FVector4f> CurrentTransform;
 	TArray<FVector4f> CurrentInverse;
 	TArray<FVector4f> PreviousTransform;
@@ -49,7 +49,7 @@ struct FNDIRigidMeshCollisionArrays
 
 	void Reset()
 	{
-		ElementOffsets = FElementOffset();
+		ElementOffsets = FNDIRigidMeshCollisionElementOffset();
 		CurrentTransform.Init(FVector4f(0, 0, 0, 0), 3 * MaxPrimitives);
 		CurrentInverse.Init(FVector4f(0, 0, 0, 0), 3 * MaxPrimitives);
 		PreviousTransform.Init(FVector4f(0, 0, 0, 0), 3 * MaxPrimitives);
