added index to unit function for volume cache DI

diff --git a/Engine/Plugins/FX/Niagara/Shaders/Private/NiagaraDataInterfaceVolumeCache.ush b/Engine/Plugins/FX/Niagara/Shaders/Private/NiagaraDataInterfaceVolumeCache.ush
--- a/Engine/Plugins/FX/Niagara/Shaders/Private/NiagaraDataInterfaceVolumeCache.ush
+++ b/Engine/Plugins/FX/Niagara/Shaders/Private/NiagaraDataInterfaceVolumeCache.ush
@@ -18,3 +18,8 @@ void GetCurrentFrameNumCells_{ParameterName}(out float3 OutValue)
 {
 	OutValue = {ParameterName}_TextureSize;
 }
+
+void IndexToUnit_{ParameterName}(float3 Index, out float3 Unit)
+{
+	Unit = (Index + .5) / {ParameterName}_TextureSize;
+}
diff --git a/Engine/Plugins/FX/Niagara/Source/Niagara/Classes/NiagaraDataInterfaceVolumeCache.h b/Engine/Plugins/FX/Niagara/Source/Niagara/Classes/NiagaraDataInterfaceVolumeCache.h
--- a/Engine/Plugins/FX/Niagara/Source/Niagara/Classes/NiagaraDataInterfaceVolumeCache.h
+++ b/Engine/Plugins/FX/Niagara/Source/Niagara/Classes/NiagaraDataInterfaceVolumeCache.h
@@ -68,6 +68,7 @@ protected:
 	static const FName SetFrameName;
 	static const FName ReadFileName;
 	static const FName GetNumCellsName;
+	static const FName IndexToUnitName;
 	static const FName SampleCurrentFrameValueName;
 	static const FName GetCurrentFrameValue;
 	static const FName GetCurrentFrameNumCells;
diff --git a/Engine/Plugins/FX/Niagara/Source/Niagara/Private/NiagaraDataInterfaceVolumeCache.cpp b/Engine/Plugins/FX/Niagara/Source/Niagara/Private/NiagaraDataInterfaceVolumeCache.cpp
--- a/Engine/Plugins/FX/Niagara/Source/Niagara/Private/NiagaraDataInterfaceVolumeCache.cpp
+++ b/Engine/Plugins/FX/Niagara/Source/Niagara/Private/NiagaraDataInterfaceVolumeCache.cpp
@@ -22,6 +22,7 @@ const TCHAR* UNiagaraDataInterfaceVolumeCache::TemplateShaderFilePath = TEXT("/P
 const FName UNiagaraDataInterfaceVolumeCache::SetFrameName("SetFrame");
 const FName UNiagaraDataInterfaceVolumeCache::ReadFileName("ReadFile");
 const FName UNiagaraDataInterfaceVolumeCache::GetNumCellsName("GetNumCells");
+const FName UNiagaraDataInterfaceVolumeCache::IndexToUnitName("IndexToUnit");
 const FName UNiagaraDataInterfaceVolumeCache::SampleCurrentFrameValueName(TEXT("SampleCurrentFrameValue"));
 const FName UNiagaraDataInterfaceVolumeCache::GetCurrentFrameValue(TEXT("GetCurrentFrameValue"));
 const FName UNiagaraDataInterfaceVolumeCache::GetCurrentFrameNumCells(TEXT("GetCurrentFrameNumCells"));
@@ -196,6 +197,18 @@ void UNiagaraDataInterfaceVolumeCache::GetFunctions(TArray<FNiagaraFunctionSigna
 		OutFunctions.Add(Sig);
 	}
 
+	{
+		FNiagaraFunctionSignature Sig;
+		Sig.Name = IndexToUnitName;
+		Sig.Inputs.Add(FNiagaraVariable(FNiagaraTypeDefinition(GetClass()), TEXT("Grid")));
+		Sig.Inputs.Add(FNiagaraVariable(FNiagaraTypeDefinition::GetVec3Def(), TEXT("Index")));		
+		Sig.Outputs.Add(FNiagaraVariable(FNiagaraTypeDefinition::GetVec3Def(), TEXT("Unit")));
+
+		Sig.bMemberFunction = true;
+		Sig.bRequiresContext = false;
+		OutFunctions.Add(Sig);
+	}
+
 	{
 		FNiagaraFunctionSignature Sig;
 		Sig.Name = SampleCurrentFrameValueName;
@@ -401,7 +414,8 @@ bool UNiagaraDataInterfaceVolumeCache::GetFunctionHLSL(const FNiagaraDataInterfa
 {
 	if ((FunctionInfo.DefinitionName == SampleCurrentFrameValueName) ||
 		(FunctionInfo.DefinitionName == GetCurrentFrameNumCells) ||
-		(FunctionInfo.DefinitionName == GetCurrentFrameValue))
+		(FunctionInfo.DefinitionName == GetCurrentFrameValue) || 
+		(FunctionInfo.DefinitionName == IndexToUnitName))
 	{
 		return true;
 	}
