- Fix for input keys not working due to a change in the method client's have to implement

diff --git a/Engine/Plugins/FX/Niagara/Source/NiagaraEditor/Private/Widgets/SNiagaraBakerViewport.cpp b/Engine/Plugins/FX/Niagara/Source/NiagaraEditor/Private/Widgets/SNiagaraBakerViewport.cpp
--- a/Engine/Plugins/FX/Niagara/Source/NiagaraEditor/Private/Widgets/SNiagaraBakerViewport.cpp
+++ b/Engine/Plugins/FX/Niagara/Source/NiagaraEditor/Private/Widgets/SNiagaraBakerViewport.cpp
@@ -31,7 +31,7 @@ public:
 		SetViewportType(ELevelViewportType::LVT_OrthoXZ);
 	}
 
-	virtual bool InputKey(FViewport* InViewport, int32 ControllerId, FKey Key, EInputEvent Event, float AmountDepressed, bool bGamepad) override
+	virtual bool InputKey(const FInputKeyEventArgs& EventArgs) override
 	{
 		// Only allow movement when in preview mode
 		auto ViewModel = WeakViewModel.Pin();
