- Add camera reset to default option to baker

diff --git a/Engine/Plugins/FX/Niagara/Source/Niagara/Classes/NiagaraBakerSettings.h b/Engine/Plugins/FX/Niagara/Source/Niagara/Classes/NiagaraBakerSettings.h
--- a/Engine/Plugins/FX/Niagara/Source/Niagara/Classes/NiagaraBakerSettings.h
+++ b/Engine/Plugins/FX/Niagara/Source/Niagara/Classes/NiagaraBakerSettings.h
@@ -58,6 +58,25 @@ struct FNiagaraBakerCameraSettings
 	bool IsOrthographic() const { return ViewMode != ENiagaraBakerViewMode::Perspective; }
 	bool IsPerspective() const { return ViewMode == ENiagaraBakerViewMode::Perspective; }
 
+	void ResetToDefault()
+	{
+		if (IsPerspective())
+		{
+			ViewportLocation = FVector(0.0f, -200.0f, 0.0f);
+			ViewportRotation = FRotator(180.0f, 0.0f, 90.0f);
+		}
+		else
+		{
+			ViewportLocation = FVector::ZeroVector;
+			ViewportRotation = FRotator::ZeroRotator;
+		}
+		OrbitDistance = 200.0f;
+		FOV = 90.0f;
+		OrthoWidth = 512.0f;
+		bUseAspectRatio = false;
+		AspectRatio = 1.0f;
+	}
+
 	bool operator==(const FNiagaraBakerCameraSettings& Other) const
 	{
 		return
diff --git a/Engine/Plugins/FX/Niagara/Source/Niagara/Private/NiagaraBakerSettings.cpp b/Engine/Plugins/FX/Niagara/Source/Niagara/Private/NiagaraBakerSettings.cpp
--- a/Engine/Plugins/FX/Niagara/Source/Niagara/Private/NiagaraBakerSettings.cpp
+++ b/Engine/Plugins/FX/Niagara/Source/Niagara/Private/NiagaraBakerSettings.cpp
@@ -20,9 +20,8 @@ UNiagaraBakerSettings::UNiagaraBakerSettings(const FObjectInitializer& Init)
 	for (int i=0; i < int(ENiagaraBakerViewMode::Num); ++i)
 	{
 		CameraSettings[i].ViewMode = ENiagaraBakerViewMode(i);
+		CameraSettings[i].ResetToDefault();
 	}
-	CameraSettings[int(ENiagaraBakerViewMode::Perspective)].ViewportLocation = FVector(0.0f, -200.0f, 0.0f);
-	CameraSettings[int(ENiagaraBakerViewMode::Perspective)].ViewportRotation = FRotator(180.0f, 0.0f, 90.0f);
 
 	// Deprecate properties
 	for ( int i=0; i < (int)ENiagaraBakerViewMode::Num; ++i )
diff --git a/Engine/Plugins/FX/Niagara/Source/NiagaraEditor/Private/ViewModels/NiagaraBakerViewModel.cpp b/Engine/Plugins/FX/Niagara/Source/NiagaraEditor/Private/ViewModels/NiagaraBakerViewModel.cpp
--- a/Engine/Plugins/FX/Niagara/Source/NiagaraEditor/Private/ViewModels/NiagaraBakerViewModel.cpp
+++ b/Engine/Plugins/FX/Niagara/Source/NiagaraEditor/Private/ViewModels/NiagaraBakerViewModel.cpp
@@ -338,6 +338,16 @@ void FNiagaraBakerViewModel::SetCameraAspectRatio(float InAspectRatio)
 	}
 }
 
+void FNiagaraBakerViewModel::ResetCurrentCamera()
+{
+	if (UNiagaraBakerSettings* BakerSettings = GetBakerSettings())
+	{
+		const FScopedTransaction Transaction(LOCTEXT("ResetCamera", "Reset Camera"));
+		BakerSettings->Modify();
+		BakerSettings->GetCurrentCamera().ResetToDefault();
+	}
+}
+
 void FNiagaraBakerViewModel::AddOutput(UClass* Class)
 {
 	if (UNiagaraBakerSettings* BakerSettings = GetBakerSettings())
diff --git a/Engine/Plugins/FX/Niagara/Source/NiagaraEditor/Private/ViewModels/NiagaraBakerViewModel.h b/Engine/Plugins/FX/Niagara/Source/NiagaraEditor/Private/ViewModels/NiagaraBakerViewModel.h
--- a/Engine/Plugins/FX/Niagara/Source/NiagaraEditor/Private/ViewModels/NiagaraBakerViewModel.h
+++ b/Engine/Plugins/FX/Niagara/Source/NiagaraEditor/Private/ViewModels/NiagaraBakerViewModel.h
@@ -101,6 +101,8 @@ public:
 	float GetCameraAspectRatio() const;
 	void SetCameraAspectRatio(float InAspectRatio);
 
+	void ResetCurrentCamera();
+
 	void AddOutput(UClass* Class);
 	void RemoveCurrentOutput();
 	bool CanRemoveCurrentOutput() const;
diff --git a/Engine/Plugins/FX/Niagara/Source/NiagaraEditor/Private/Widgets/SNiagaraBakerWidget.cpp b/Engine/Plugins/FX/Niagara/Source/NiagaraEditor/Private/Widgets/SNiagaraBakerWidget.cpp
--- a/Engine/Plugins/FX/Niagara/Source/NiagaraEditor/Private/Widgets/SNiagaraBakerWidget.cpp
+++ b/Engine/Plugins/FX/Niagara/Source/NiagaraEditor/Private/Widgets/SNiagaraBakerWidget.cpp
@@ -697,6 +697,15 @@ TSharedRef<SWidget> SNiagaraBakerWidget::MakeViewOptionsMenu()
 				LOCTEXT("CameraAspectRatio", "Camera Aspect Ratio")
 			);
 		}
+
+		MenuBuilder.AddMenuEntry(
+			LOCTEXT("ResetCamera", "Reset Camera"),
+			LOCTEXT("ResetCameraTooltip", "Resets the current camera back to the default settings."),
+			FSlateIcon(),
+			FUIAction(
+				FExecuteAction::CreateSP(ViewModel, &FNiagaraBakerViewModel::ResetCurrentCamera)
+			)
+		);
 	}
 	MenuBuilder.EndSection();
 
