- Use correct gamma value with custom batched elements

#rnx

diff --git a/Engine/Plugins/FX/Niagara/Source/NiagaraShader/Private/NiagaraBatchedElements.cpp b/Engine/Plugins/FX/Niagara/Source/NiagaraShader/Private/NiagaraBatchedElements.cpp
--- a/Engine/Plugins/FX/Niagara/Source/NiagaraShader/Private/NiagaraBatchedElements.cpp
+++ b/Engine/Plugins/FX/Niagara/Source/NiagaraShader/Private/NiagaraBatchedElements.cpp
@@ -198,7 +198,7 @@ void FBatchedElementNiagaraInvertColorChannel::BindShaders(FRHICommandList& RHIC
 
 	FNiagaraSimpleElementPS::FParameters PassParameters;
 	PassParameters.InColorWeights	= FMatrix44f::Identity;
-	PassParameters.InGamma			= 1.0f;
+	PassParameters.InGamma			= InGamma;
 	PassParameters.InTexture		= GWhiteTexture->TextureRHI;
 	PassParameters.InTextureSampler = TStaticSamplerState<SF_Point>::GetRHI();
 	SetShaderParameters(RHICmdList, PixelShader, PixelShader.GetPixelShader(), PassParameters);
@@ -222,7 +222,7 @@ void FBatchedElementNiagaraSimple::BindShaders(FRHICommandList& RHICmdList, FGra
 
 	FNiagaraSimpleElementPS::FParameters PassParameters;
 	PassParameters.InColorWeights = FMatrix44f(ColorTransform);
-	PassParameters.InGamma = 1.0f;
+	PassParameters.InGamma = InGamma;
 	PassParameters.InTexture = Texture ? Texture->TextureRHI : GWhiteTexture->TextureRHI;
 	PassParameters.InTextureSampler = TStaticSamplerState<SF_Bilinear>::GetRHI();
 	SetShaderParameters(RHICmdList, PixelShader, PixelShader.GetPixelShader(), PassParameters);
