- Convert StaticMesh DI to ShaderParams / RDG

diff --git a/Engine/Plugins/FX/Niagara/Shaders/Private/NiagaraDataInterfaceStaticMeshTemplate.ush b/Engine/Plugins/FX/Niagara/Shaders/Private/NiagaraDataInterfaceStaticMeshTemplate.ush
--- a/Engine/Plugins/FX/Niagara/Shaders/Private/NiagaraDataInterfaceStaticMeshTemplate.ush
+++ b/Engine/Plugins/FX/Niagara/Shaders/Private/NiagaraDataInterfaceStaticMeshTemplate.ush
@@ -1,34 +1,38 @@
 // Copyright Epic Games, Inc. All Rights Reserved.
 
-int3			NumTriangles_{ParameterName};		// x = Num Sections, y = Num Filtered, z = Num Unfiltered
-int				NumVertices_{ParameterName};
-int				NumUVs_{ParameterName};
-uint			HasColors_{ParameterName};
-Buffer<uint>	IndexBuffer_{ParameterName};
-Buffer<float>	PositionBuffer_{ParameterName};
-Buffer<float4>	TangentBuffer_{ParameterName};
-Buffer<float2>	UVBuffer_{ParameterName};
-Buffer<float4>	ColorBuffer_{ParameterName};
-
-uint			HasUniformSampling_{ParameterName};
-Buffer<uint2>	UniformSamplingTriangles_{ParameterName};
-
-int3			SectionCounts_{ParameterName};		// x = Num Sections, y = Num Filtered, z = Num Unfiltered
-Buffer<uint4>	SectionInfos_{ParameterName};		// FirstTriangle, NumTriangles, Prob, Alias
-Buffer<uint>	FilteredAndUnfilteredSections_{ParameterName};
-
-int3			SocketCounts_{ParameterName};		// x = Num Sockets, y = Num Filtered, z = Num Unfiltered
-Buffer<float4>	SocketTransforms_{ParameterName};
-Buffer<uint>	FilteredAndUnfilteredSockets_{ParameterName};
-
-float			InvDeltaSeconds_{ParameterName};
-float4x4		InstanceTransform_{ParameterName};
-float4x4		InstanceTransformInverseTransposed_{ParameterName};
-float4			InstanceRotation_{ParameterName};
-float4x4		InstancePreviousTransform_{ParameterName};
-float4x4		InstancePreviousTransformInverseTransposed_{ParameterName};
-float4			InstancePreviousRotation_{ParameterName};
-float3			InstanceWorldVelocity_{ParameterName};
+int3			{ParameterName}_NumTriangles;		// x = Num Sections, y = Num Filtered, z = Num Unfiltered
+int				{ParameterName}_NumVertices;
+int				{ParameterName}_NumUVs;
+uint			{ParameterName}_HasColors;
+Buffer<uint>	{ParameterName}_IndexBuffer;
+Buffer<float>	{ParameterName}_PositionBuffer;
+Buffer<float4>	{ParameterName}_TangentBuffer;
+Buffer<float2>	{ParameterName}_UVBuffer;
+Buffer<float4>	{ParameterName}_ColorBuffer;
+
+uint			{ParameterName}_HasUniformSampling;
+Buffer<uint2>	{ParameterName}_UniformSamplingTriangles;
+
+int3			{ParameterName}_SectionCounts;		// x = Num Sections, y = Num Filtered, z = Num Unfiltered
+Buffer<uint4>	{ParameterName}_SectionInfos;		// FirstTriangle, NumTriangles, Prob, Alias
+Buffer<uint>	{ParameterName}_FilteredAndUnfilteredSections;
+
+int3			{ParameterName}_SocketCounts;		// x = Num Sockets, y = Num Filtered, z = Num Unfiltered
+Buffer<float4>	{ParameterName}_SocketTransforms;
+Buffer<uint>	{ParameterName}_FilteredAndUnfilteredSockets;
+
+float			{ParameterName}_InvDeltaSeconds;
+float4x4		{ParameterName}_InstanceTransform;
+float4x4		{ParameterName}_InstanceTransformInverseTransposed;
+float4			{ParameterName}_InstanceRotation;
+float4x4		{ParameterName}_InstancePreviousTransform;
+float4x4		{ParameterName}_InstancePreviousTransformInverseTransposed;
+float4			{ParameterName}_InstancePreviousRotation;
+float3			{ParameterName}_InstanceWorldVelocity;
+
+#if DISTATICMESH_ALLOWDISTANCEFIELD
+int				{ParameterName}_InstanceDistanceFieldIndex;
+#endif
 
 //////////////////////////////////////////////////////////////////////////
 // Fwd declarations
@@ -40,7 +44,7 @@ void RandomSectionTriangle_{ParameterName}(NiagaraRandInfo RandInfo, int Section
 //////////////////////////////////////////////////////////////////////////
 int UniformTriangle_{ParameterName}(NiagaraRandInfo RandInfo, int Tri, int SectionOffset)
 {
-	uint2 ProbAlias = UniformSamplingTriangles_{ParameterName}[Tri];
+	uint2 ProbAlias = {ParameterName}_UniformSamplingTriangles[Tri];
 	float Prob = asfloat(ProbAlias.x);
 	int Alias = ProbAlias.y;
 	return NiagaraRandomFloat(RandInfo) < Prob ? Tri : Alias + SectionOffset;
@@ -50,32 +54,32 @@ int UniformTriangle_{ParameterName}(NiagaraRandInfo RandInfo, int Tri, int Secti
 // Vertex Sampling
 void IsValidVertex_{ParameterName}(out bool IsValid)
 {
-	IsValid = NumVertices_{ParameterName} > 0;
+	IsValid = {ParameterName}_NumVertices > 0;
 }
 
 void RandomVertex_{ParameterName}(NiagaraRandInfo RandInfo, out int Vertex)
 {
-	Vertex = NiagaraRandomInt(RandInfo, NumVertices_{ParameterName});
+	Vertex = NiagaraRandomInt(RandInfo, {ParameterName}_NumVertices);
 }
 
 void GetVertexCount_{ParameterName}(out int Count)
 {
-	Count = NumVertices_{ParameterName};
+	Count = {ParameterName}_NumVertices;
 }
 
 void GetVertex_{ParameterName}(int Vertex, out float3 Position, out float3 Velocity, out float3 Normal, out float3 Bitangent, out float3 Tangent)
 {
 	[branch]
-	if ( NumVertices_{ParameterName} > 0 )
+	if ( {ParameterName}_NumVertices > 0 )
 	{
-		Vertex = clamp(Vertex, 0, NumVertices_{ParameterName} - 1);
+		Vertex = clamp(Vertex, 0, {ParameterName}_NumVertices - 1);
 
-		Position.x = PositionBuffer_{ParameterName}[Vertex * 3 + 0];
-		Position.y = PositionBuffer_{ParameterName}[Vertex * 3 + 1];
-		Position.z = PositionBuffer_{ParameterName}[Vertex * 3 + 2];
+		Position.x = {ParameterName}_PositionBuffer[Vertex * 3 + 0];
+		Position.y = {ParameterName}_PositionBuffer[Vertex * 3 + 1];
+		Position.z = {ParameterName}_PositionBuffer[Vertex * 3 + 2];
 
-		float4 TangentX = TangentBias(TangentBuffer_{ParameterName}[Vertex * 2 + 0]);
-		float4 TangentZ = TangentBias(TangentBuffer_{ParameterName}[Vertex * 2 + 1]);
+		float4 TangentX = TangentBias({ParameterName}_TangentBuffer[Vertex * 2 + 0]);
+		float4 TangentZ = TangentBias({ParameterName}_TangentBuffer[Vertex * 2 + 1]);
 		float3 TangentY = cross(TangentZ.xyz, TangentX.xyz) * TangentZ.w;
 
 		Normal		= TangentZ.xyz;
@@ -97,32 +101,32 @@ void GetVertexWS_{ParameterName}(int Vertex, out float3 Position, out float3 Vel
 	float3 LocalPosition;
 	GetVertex_{ParameterName}(Vertex, LocalPosition, Velocity, Normal, Bitangent, Tangent);
 	
-	Position				= mul(float4(LocalPosition, 1), InstanceTransform_{ParameterName}).xyz;
-	float3 PreviousPosition	= mul(float4(LocalPosition, 1), InstancePreviousTransform_{ParameterName}).xyz;
-	Velocity				= (Position - PreviousPosition) * InvDeltaSeconds_{ParameterName};
-	Normal					= normalize(mul(float4(Normal, 0), InstanceTransformInverseTransposed_{ParameterName}).xyz);
-	Bitangent				= normalize(mul(float4(Bitangent, 0), InstanceTransformInverseTransposed_{ParameterName}).xyz);
-	Tangent					= normalize(mul(float4(Tangent, 0), InstanceTransformInverseTransposed_{ParameterName}).xyz);
+	Position				= mul(float4(LocalPosition, 1), {ParameterName}_InstanceTransform).xyz;
+	float3 PreviousPosition	= mul(float4(LocalPosition, 1), {ParameterName}_InstancePreviousTransform).xyz;
+	Velocity				= (Position - PreviousPosition) * {ParameterName}_InvDeltaSeconds;
+	Normal					= normalize(mul(float4(Normal, 0), {ParameterName}_InstanceTransformInverseTransposed).xyz);
+	Bitangent				= normalize(mul(float4(Bitangent, 0), {ParameterName}_InstanceTransformInverseTransposed).xyz);
+	Tangent					= normalize(mul(float4(Tangent, 0), {ParameterName}_InstanceTransformInverseTransposed).xyz);
 }
 
 void GetVertexColor_{ParameterName}(int Vertex, out float4 Color)
 {
-	Vertex = clamp(Vertex, 0, NumVertices_{ParameterName} - 1);
-	Color = HasColors_{ParameterName} ? ColorBuffer_{ParameterName}[Vertex] FMANUALFETCH_COLOR_COMPONENT_SWIZZLE : float4(1, 1, 1, 1);
+	Vertex = clamp(Vertex, 0, {ParameterName}_NumVertices - 1);
+	Color = {ParameterName}_HasColors ? {ParameterName}_ColorBuffer[Vertex] FMANUALFETCH_COLOR_COMPONENT_SWIZZLE : float4(1, 1, 1, 1);
 }
 
 void GetVertexUV_{ParameterName}(int Vertex, int UVSet, out float2 UV)
 {
-	Vertex = clamp(Vertex, 0, NumVertices_{ParameterName} - 1);
-	UVSet = clamp(UVSet, 0, NumUVs_{ParameterName} - 1);
-	UV = UVSet >= 0 && UVSet < NumUVs_{ParameterName} ? UVBuffer_{ParameterName}[Vertex * NumUVs_{ParameterName} + UVSet] : float2(0, 0);
+	Vertex = clamp(Vertex, 0, {ParameterName}_NumVertices - 1);
+	UVSet = clamp(UVSet, 0, {ParameterName}_NumUVs - 1);
+	UV = UVSet >= 0 && UVSet < {ParameterName}_NumUVs ? {ParameterName}_UVBuffer[Vertex * {ParameterName}_NumUVs + UVSet] : float2(0, 0);
 }
 
 //////////////////////////////////////////////////////////////////////////
 // VM Triangle Sampling
 void IsValidTriangle_{ParameterName}(int Tri, out bool IsValid)
 {
-	IsValid = Tri >= 0 && Tri < NumTriangles_{ParameterName}.x;
+	IsValid = Tri >= 0 && Tri < {ParameterName}_NumTriangles.x;
 }
 
 void RandomTriangle_{ParameterName}(NiagaraRandInfo RandInfo, out int Tri, out float3 BaryCoord)
@@ -134,17 +138,17 @@ void RandomTriangle_{ParameterName}(NiagaraRandInfo RandInfo, out int Tri, out f
 
 void GetTriangleCount_{ParameterName}(out int Count)
 {
-	Count = NumTriangles_{ParameterName}.x;
+	Count = {ParameterName}_NumTriangles.x;
 }
 
 void IsValidFilteredTriangle_{ParameterName}(int Tri, out bool IsValid)
 {
-	IsValid = Tri >= 0 && Tri < NumTriangles_{ParameterName}.y;
+	IsValid = Tri >= 0 && Tri < {ParameterName}_NumTriangles.y;
 }
 
 void RandomFilteredTriangle_{ParameterName}(NiagaraRandInfo RandInfo, out int Tri, out float3 BaryCoord)
 {
-	if ( NumTriangles_{ParameterName}.y > 0 )
+	if ( {ParameterName}_NumTriangles.y > 0 )
 	{
 		int Section;
 		RandomFilteredSection_{ParameterName}(RandInfo, Section);
@@ -159,21 +163,21 @@ void RandomFilteredTriangle_{ParameterName}(NiagaraRandInfo RandInfo, out int Tr
 
 void GetFilteredTriangleCount_{ParameterName}(out int Count)
 {
-	Count = NumTriangles_{ParameterName}.y;
+	Count = {ParameterName}_NumTriangles.y;
 }
 
 void GetFilteredTriangleAt_{ParameterName}(int TriangleIndex, out int Tri)
 {
 	Tri = 0;
 
-	int SectionInfoOffset = SectionCounts_{ParameterName}.x;
+	int SectionInfoOffset = {ParameterName}_SectionCounts.x;
 	int SectionFirstTriangle = 0;
-	for ( int i=0; i < SectionCounts_{ParameterName}.y; ++i )
+	for ( int i=0; i < {ParameterName}_SectionCounts.y; ++i )
 	{
-		int SectionLastTriangle = SectionFirstTriangle + SectionInfos_{ParameterName}[SectionInfoOffset + i].y;
+		int SectionLastTriangle = SectionFirstTriangle + {ParameterName}_SectionInfos[SectionInfoOffset + i].y;
 		if ( TriangleIndex >= SectionFirstTriangle && TriangleIndex < SectionLastTriangle )
 		{
-			Tri = (TriangleIndex - SectionFirstTriangle) + SectionInfos_{ParameterName}[SectionInfoOffset + i].x;
+			Tri = (TriangleIndex - SectionFirstTriangle) + {ParameterName}_SectionInfos[SectionInfoOffset + i].x;
 			return;
 		}
 		SectionFirstTriangle = SectionLastTriangle;
@@ -182,12 +186,12 @@ void GetFilteredTriangleAt_{ParameterName}(int TriangleIndex, out int Tri)
 
 void IsValidUnfilteredTriangle_{ParameterName}(int Tri, out bool IsValid)
 {
-	IsValid = Tri >= 0 && Tri < NumTriangles_{ParameterName}.z;
+	IsValid = Tri >= 0 && Tri < {ParameterName}_NumTriangles.z;
 }
 
 void RandomUnfilteredTriangle_{ParameterName}(NiagaraRandInfo RandInfo, out int Tri, out float3 BaryCoord)
 {
-	if ( NumTriangles_{ParameterName}.z > 0 )
+	if ( {ParameterName}_NumTriangles.z > 0 )
 	{
 		int Section;
 		RandomUnfilteredSection_{ParameterName}(RandInfo, Section);
@@ -202,21 +206,21 @@ void RandomUnfilteredTriangle_{ParameterName}(NiagaraRandInfo RandInfo, out int
 
 void GetUnfilteredTriangleCount_{ParameterName}(out int Count)
 {
-	Count = NumTriangles_{ParameterName}.z;
+	Count = {ParameterName}_NumTriangles.z;
 }
 
 void GetUnfilteredTriangleAt_{ParameterName}(int TriangleIndex, out int Tri)
 {
 	Tri = 0;
 
-	int SectionInfoOffset = SectionCounts_{ParameterName}.x + SectionCounts_{ParameterName}.y;
+	int SectionInfoOffset = {ParameterName}_SectionCounts.x + {ParameterName}_SectionCounts.y;
 	int SectionFirstTriangle = 0;
-	for ( int i=0; i < SectionCounts_{ParameterName}.z; ++i )
+	for ( int i=0; i < {ParameterName}_SectionCounts.z; ++i )
 	{
-		int SectionLastTriangle = SectionFirstTriangle + SectionInfos_{ParameterName}[SectionInfoOffset + i].y;
+		int SectionLastTriangle = SectionFirstTriangle + {ParameterName}_SectionInfos[SectionInfoOffset + i].y;
 		if ( TriangleIndex >= SectionFirstTriangle && TriangleIndex < SectionLastTriangle )
 		{
-			Tri = (TriangleIndex - SectionFirstTriangle) + SectionInfos_{ParameterName}[SectionInfoOffset + i].x;
+			Tri = (TriangleIndex - SectionFirstTriangle) + {ParameterName}_SectionInfos[SectionInfoOffset + i].x;
 			return;
 		}
 		SectionFirstTriangle = SectionLastTriangle;
@@ -226,12 +230,12 @@ void GetUnfilteredTriangleAt_{ParameterName}(int TriangleIndex, out int Tri)
 void GetTriangleIndices_{ParameterName}(int Tri, out int Index0, out int Index1, out int Index2)
 {
 	[branch]
-	if ( NumTriangles_{ParameterName}.x > 0 )
+	if ( {ParameterName}_NumTriangles.x > 0 )
 	{
-		Tri = clamp(Tri, 0, NumTriangles_{ParameterName}.x - 1);
-		Index0 = IndexBuffer_{ParameterName}[Tri * 3 + 0];
-		Index1 = IndexBuffer_{ParameterName}[Tri * 3 + 1];
-		Index2 = IndexBuffer_{ParameterName}[Tri * 3 + 2];
+		Tri = clamp(Tri, 0, {ParameterName}_NumTriangles.x - 1);
+		Index0 = {ParameterName}_IndexBuffer[Tri * 3 + 0];
+		Index1 = {ParameterName}_IndexBuffer[Tri * 3 + 1];
+		Index2 = {ParameterName}_IndexBuffer[Tri * 3 + 2];
 	}
 	else
 	{
@@ -263,12 +267,12 @@ void GetTriangleWS_{ParameterName}(int Tri, float3 BaryCoord, out float3 Positio
 	float3 LocalPosition;
 	GetTriangle_{ParameterName}(Tri, BaryCoord, LocalPosition, Velocity, Normal, Bitangent, Tangent);
 
-	Position				= mul(float4(LocalPosition, 1), InstanceTransform_{ParameterName}).xyz;
-	float3 PreviousPosition	= mul(float4(LocalPosition, 1), InstancePreviousTransform_{ParameterName}).xyz;
-	Velocity				= (Position - PreviousPosition) * InvDeltaSeconds_{ParameterName};
-	Normal					= normalize(mul(float4(Normal, 0), InstanceTransformInverseTransposed_{ParameterName}).xyz);
-	Bitangent				= normalize(mul(float4(Bitangent, 0), InstanceTransformInverseTransposed_{ParameterName}).xyz);
-	Tangent					= normalize(mul(float4(Tangent, 0), InstanceTransformInverseTransposed_{ParameterName}).xyz);
+	Position				= mul(float4(LocalPosition, 1), {ParameterName}_InstanceTransform).xyz;
+	float3 PreviousPosition	= mul(float4(LocalPosition, 1), {ParameterName}_InstancePreviousTransform).xyz;
+	Velocity				= (Position - PreviousPosition) * {ParameterName}_InvDeltaSeconds;
+	Normal					= normalize(mul(float4(Normal, 0), {ParameterName}_InstanceTransformInverseTransposed).xyz);
+	Bitangent				= normalize(mul(float4(Bitangent, 0), {ParameterName}_InstanceTransformInverseTransposed).xyz);
+	Tangent					= normalize(mul(float4(Tangent, 0), {ParameterName}_InstanceTransformInverseTransposed).xyz);
 }
 
 void GetTriangleColor_{ParameterName}(int Tri, float3 BaryCoord, out float4 Color)
@@ -301,27 +305,27 @@ void GetTriangleUV_{ParameterName}(int Tri, float3 BaryCoord, int UVSet, out flo
 // Socket Functions
 void GetSocketCount_{ParameterName}(out int Count)
 {
-	Count = SocketCounts_{ParameterName}.x;
+	Count = {ParameterName}_SocketCounts.x;
 }
 
 void GetFilteredSocketCount_{ParameterName}(out int Count)
 {
-	Count = SocketCounts_{ParameterName}.y;
+	Count = {ParameterName}_SocketCounts.y;
 }
 void GetUnfilteredSocketCount_{ParameterName}(out int Count)
 {
-	Count = SocketCounts_{ParameterName}.z;
+	Count = {ParameterName}_SocketCounts.z;
 }
 
 void GetSocketTransform_{ParameterName}(int Socket, out float3 Position, out float4 Rotation, out float3 Scale)
 {
 	[branch]
-	if ( SocketCounts_{ParameterName}.x > 0 )
+	if ( {ParameterName}_SocketCounts.x > 0 )
 	{
-		Socket = clamp(Socket, 0, SocketCounts_{ParameterName}.x - 1);
-		Position = SocketTransforms_{ParameterName}[Socket * 3 + 0].xyz;
-		Rotation = SocketTransforms_{ParameterName}[Socket * 3 + 1];
-		Scale    = SocketTransforms_{ParameterName}[Socket * 3 + 2].xyz;
+		Socket = clamp(Socket, 0, {ParameterName}_SocketCounts.x - 1);
+		Position = {ParameterName}_SocketTransforms[Socket * 3 + 0].xyz;
+		Rotation = {ParameterName}_SocketTransforms[Socket * 3 + 1];
+		Scale    = {ParameterName}_SocketTransforms[Socket * 3 + 2].xyz;
 	}
 	else
 	{
@@ -334,18 +338,18 @@ void GetSocketTransform_{ParameterName}(int Socket, out float3 Position, out flo
 void GetSocketTransformWS_{ParameterName}(int Socket, out float3 Position, out float4 Rotation, out float3 Scale)
 {
 	GetSocketTransform_{ParameterName}(Socket, Position, Rotation, Scale);
-	Position = mul(float4(Position, 1.0), InstanceTransform_{ParameterName}).xyz;
-	Rotation = NiagaraGPU_QuatMul(InstanceRotation_{ParameterName}, Rotation);
-	Scale = mul(float4(Scale, 0.0), InstanceTransform_{ParameterName}).xyz;
+	Position = mul(float4(Position, 1.0), {ParameterName}_InstanceTransform).xyz;
+	Rotation = NiagaraGPU_QuatMul({ParameterName}_InstanceRotation, Rotation);
+	Scale = mul(float4(Scale, 0.0), {ParameterName}_InstanceTransform).xyz;
 }
 
 void GetFilteredSocketTransform_{ParameterName}(int SocketIndex, out float3 Position, out float4 Rotation, out float3 Scale)
 {
 	[branch]
-	if ( SocketCounts_{ParameterName}.y > 0 )
+	if ( {ParameterName}_SocketCounts.y > 0 )
 	{
-		SocketIndex = clamp(SocketIndex, 0, SocketCounts_{ParameterName}.y - 1);
-		int Socket = FilteredAndUnfilteredSockets_{ParameterName}[SocketIndex];
+		SocketIndex = clamp(SocketIndex, 0, {ParameterName}_SocketCounts.y - 1);
+		int Socket = {ParameterName}_FilteredAndUnfilteredSockets[SocketIndex];
 		GetSocketTransform_{ParameterName}(Socket, Position, Rotation, Scale);
 	}
 	else
@@ -359,18 +363,18 @@ void GetFilteredSocketTransform_{ParameterName}(int SocketIndex, out float3 Posi
 void GetFilteredSocketTransformWS_{ParameterName}(int SocketIndex, out float3 Position, out float4 Rotation, out float3 Scale)
 {
 	GetFilteredSocketTransform_{ParameterName}(SocketIndex, Position, Rotation, Scale);
-	Position = mul(float4(Position, 1.0), InstanceTransform_{ParameterName}).xyz;
-	Rotation = NiagaraGPU_QuatMul(InstanceRotation_{ParameterName}, Rotation);
-	Scale = mul(float4(Scale, 0.0), InstanceTransform_{ParameterName}).xyz;
+	Position = mul(float4(Position, 1.0), {ParameterName}_InstanceTransform).xyz;
+	Rotation = NiagaraGPU_QuatMul({ParameterName}_InstanceRotation, Rotation);
+	Scale = mul(float4(Scale, 0.0), {ParameterName}_InstanceTransform).xyz;
 }
 
 void GetUnfilteredSocketTransform_{ParameterName}(int SocketIndex, out float3 Position, out float4 Rotation, out float3 Scale)
 {
 	[branch]
-	if ( SocketCounts_{ParameterName}.z > 0 )
+	if ( {ParameterName}_SocketCounts.z > 0 )
 	{
-		SocketIndex = SocketCounts_{ParameterName}.y + clamp(SocketIndex, 0, SocketCounts_{ParameterName}.z - 1);
-		int Socket = FilteredAndUnfilteredSockets_{ParameterName}[SocketIndex];
+		SocketIndex = {ParameterName}_SocketCounts.y + clamp(SocketIndex, 0, {ParameterName}_SocketCounts.z - 1);
+		int Socket = {ParameterName}_FilteredAndUnfilteredSockets[SocketIndex];
 		GetSocketTransform_{ParameterName}(Socket, Position, Rotation, Scale);
 	}
 	else
@@ -384,21 +388,21 @@ void GetUnfilteredSocketTransform_{ParameterName}(int SocketIndex, out float3 Po
 void GetUnfilteredSocketTransformWS_{ParameterName}(int SocketIndex, out float3 Position, out float4 Rotation, out float3 Scale)
 {
 	GetUnfilteredSocketTransform_{ParameterName}(SocketIndex, Position, Rotation, Scale);
-	Position = mul(float4(Position, 1.0), InstanceTransform_{ParameterName}).xyz;
-	Rotation = NiagaraGPU_QuatMul(InstanceRotation_{ParameterName}, Rotation);
-	Scale = mul(float4(Scale, 0.0), InstanceTransform_{ParameterName}).xyz;
+	Position = mul(float4(Position, 1.0), {ParameterName}_InstanceTransform).xyz;
+	Rotation = NiagaraGPU_QuatMul({ParameterName}_InstanceRotation, Rotation);
+	Scale = mul(float4(Scale, 0.0), {ParameterName}_InstanceTransform).xyz;
 }
 
 //////////////////////////////////////////////////////////////////////////
 // Section functions
 void IsValidSection_{ParameterName}(int Section, out bool IsValid)
 {
-	IsValid = Section >= 0 && Section < SectionCounts_{ParameterName}.x;
+	IsValid = Section >= 0 && Section < {ParameterName}_SectionCounts.x;
 }
 
 void GetSectionTriangleCount_{ParameterName}(int Section, out int Count)
 {
-	Count = Section >= 0 && Section < SectionCounts_{ParameterName}.x ? SectionInfos_{ParameterName}[Section].y : 0;
+	Count = Section >= 0 && Section < {ParameterName}_SectionCounts.x ? {ParameterName}_SectionInfos[Section].y : 0;
 }
 
 void RandomSectionTriangle_{ParameterName}(NiagaraRandInfo RandInfo, int Section, out int Tri, out float3 BaryCoord)
@@ -408,11 +412,11 @@ void RandomSectionTriangle_{ParameterName}(NiagaraRandInfo RandInfo, int Section
 
 	if (SectionTriangleCount > 0)
 	{
-		int SectionTriangleOffset = SectionInfos_{ParameterName}[Section].x;
+		int SectionTriangleOffset = {ParameterName}_SectionInfos[Section].x;
 		Tri = NiagaraRandomInt(RandInfo, SectionTriangleCount) + SectionTriangleOffset;
 
 		[branch]
-		if ( HasUniformSampling_{ParameterName} )
+		if ( {ParameterName}_HasUniformSampling )
 		{
 			Tri = UniformTriangle_{ParameterName}(RandInfo, Tri, SectionTriangleOffset);
 		}
@@ -427,8 +431,8 @@ void GetSectionTriangleAt_{ParameterName}(int Section, int SectionTriangleIndex,
 
 	if ( SectionTriangleCount > 0 )
 	{
-		int SectionTriangleMax = SectionInfos_{ParameterName}[Section].y - 1;
-		int SectionTriangleOffset = SectionInfos_{ParameterName}[Section].x;
+		int SectionTriangleMax = {ParameterName}_SectionInfos[Section].y - 1;
+		int SectionTriangleOffset = {ParameterName}_SectionInfos[Section].x;
 		Tri = clamp(SectionTriangleIndex, 0, SectionTriangleMax) + SectionTriangleOffset;
 	}
 	else
@@ -439,11 +443,11 @@ void GetSectionTriangleAt_{ParameterName}(int Section, int SectionTriangleIndex,
 
 void GetFilteredSectionAt_{ParameterName}(int SectionIndex, out int Section)
 {
-	int FilteredSectionCount = SectionCounts_{ParameterName}.y;
+	int FilteredSectionCount = {ParameterName}_SectionCounts.y;
 	if ( FilteredSectionCount > 0 )
 	{
 		SectionIndex = clamp(SectionIndex, 0, FilteredSectionCount - 1);
-		Section = FilteredAndUnfilteredSections_{ParameterName}[SectionIndex];
+		Section = {ParameterName}_FilteredAndUnfilteredSections[SectionIndex];
 	}
 	else
 	{
@@ -453,12 +457,12 @@ void GetFilteredSectionAt_{ParameterName}(int SectionIndex, out int Section)
 
 void GetUnfilteredSectionAt_{ParameterName}(int SectionIndex, out int Section)
 {
-	int FilteredSectionCount = SectionCounts_{ParameterName}.y;
-	int UnfilteredSectionCount = SectionCounts_{ParameterName}.z;
+	int FilteredSectionCount = {ParameterName}_SectionCounts.y;
+	int UnfilteredSectionCount = {ParameterName}_SectionCounts.z;
 	if ( UnfilteredSectionCount > 0 )
 	{
 		SectionIndex = clamp(SectionIndex, 0, UnfilteredSectionCount - 1);
-		Section = FilteredAndUnfilteredSections_{ParameterName}[SectionIndex + FilteredSectionCount];
+		Section = {ParameterName}_FilteredAndUnfilteredSections[SectionIndex + FilteredSectionCount];
 	}
 	else
 	{
@@ -468,27 +472,27 @@ void GetUnfilteredSectionAt_{ParameterName}(int SectionIndex, out int Section)
 
 void GetSectionCount_{ParameterName}(out int Count)
 {
-	Count = SectionCounts_{ParameterName}.x;
+	Count = {ParameterName}_SectionCounts.x;
 }
 
 void GetFilteredSectionCount_{ParameterName}(out int Count)
 {
-	Count = SectionCounts_{ParameterName}.y;
+	Count = {ParameterName}_SectionCounts.y;
 }
 
 void GetUnfilteredSectionCount_{ParameterName}(out int Count)
 {
-	Count = SectionCounts_{ParameterName}.z;
+	Count = {ParameterName}_SectionCounts.z;
 }
 
 void RandomSection_{ParameterName}(NiagaraRandInfo RandInfo, out int Section)
 {
-	int SectionCount = SectionCounts_{ParameterName}.x;
+	int SectionCount = {ParameterName}_SectionCounts.x;
 	if (SectionCount > 0)
 	{
 		Section = NiagaraRandomInt(RandInfo, SectionCount);
-		float Prob = asfloat(SectionInfos_{ParameterName}[Section].z);
-		int Alias = SectionInfos_{ParameterName}[Section].w;
+		float Prob = asfloat({ParameterName}_SectionInfos[Section].z);
+		int Alias = {ParameterName}_SectionInfos[Section].w;
 		Section = NiagaraRandomFloat(RandInfo) < Prob ? Section : Alias;
 	}
 	else
@@ -499,16 +503,16 @@ void RandomSection_{ParameterName}(NiagaraRandInfo RandInfo, out int Section)
 
 void RandomFilteredSection_{ParameterName}(NiagaraRandInfo RandInfo, out int Section)
 {
-	int SectionCount = SectionCounts_{ParameterName}.x;
-	int FilteredSectionCount = SectionCounts_{ParameterName}.y;
+	int SectionCount = {ParameterName}_SectionCounts.x;
+	int FilteredSectionCount = {ParameterName}_SectionCounts.y;
 	if (FilteredSectionCount > 0)
 	{
 		int SectionIndex = NiagaraRandomInt(RandInfo, FilteredSectionCount);
 		int SectionInfoIndex = SectionIndex + SectionCount;
-		float Prob = asfloat(SectionInfos_{ParameterName}[SectionInfoIndex].z);
-		int Alias = SectionInfos_{ParameterName}[SectionInfoIndex].w;
+		float Prob = asfloat({ParameterName}_SectionInfos[SectionInfoIndex].z);
+		int Alias = {ParameterName}_SectionInfos[SectionInfoIndex].w;
 		SectionIndex = NiagaraRandomFloat(RandInfo) < Prob ? SectionIndex : Alias;
-		Section = FilteredAndUnfilteredSections_{ParameterName}[SectionIndex];
+		Section = {ParameterName}_FilteredAndUnfilteredSections[SectionIndex];
 	}
 	else
 	{
@@ -518,17 +522,17 @@ void RandomFilteredSection_{ParameterName}(NiagaraRandInfo RandInfo, out int Sec
 
 void RandomUnfilteredSection_{ParameterName}(NiagaraRandInfo RandInfo, out int Section)
 {
-	int SectionCount = SectionCounts_{ParameterName}.x;
-	int FilteredSectionCount = SectionCounts_{ParameterName}.y;
-	int UnfilteredSectionCount = SectionCounts_{ParameterName}.z;
+	int SectionCount = {ParameterName}_SectionCounts.x;
+	int FilteredSectionCount = {ParameterName}_SectionCounts.y;
+	int UnfilteredSectionCount = {ParameterName}_SectionCounts.z;
 	if (UnfilteredSectionCount > 0)
 	{
 		int SectionIndex = NiagaraRandomInt(RandInfo, UnfilteredSectionCount);
 		int SectionInfoIndex = SectionIndex + SectionCount + FilteredSectionCount;
-		float Prob = asfloat(SectionInfos_{ParameterName}[SectionInfoIndex].z);
-		int Alias = SectionInfos_{ParameterName}[SectionInfoIndex].w;
+		float Prob = asfloat({ParameterName}_SectionInfos[SectionInfoIndex].z);
+		int Alias = {ParameterName}_SectionInfos[SectionInfoIndex].w;
 		SectionIndex = NiagaraRandomFloat(RandInfo) < Prob ? SectionIndex : Alias;
-		Section = FilteredAndUnfilteredSections_{ParameterName}[SectionIndex + FilteredSectionCount];
+		Section = {ParameterName}_FilteredAndUnfilteredSections[SectionIndex + FilteredSectionCount];
 	}
 	else
 	{
@@ -540,30 +544,26 @@ void RandomUnfilteredSection_{ParameterName}(NiagaraRandInfo RandInfo, out int S
 // VM Misc Functions
 void IsValid_{ParameterName}(out bool IsValid)
 {
-	IsValid = NumVertices_{ParameterName} > 0;
+	IsValid = {ParameterName}_NumVertices > 0;
 }
 
 void GetLocalToWorld_{ParameterName}(out float4x4 OutMatrix)
 {
-	OutMatrix = InstanceTransform_{ParameterName};
+	OutMatrix = {ParameterName}_InstanceTransform;
 }
 
 void GetLocalToWorldInverseTransposed_{ParameterName}(out float4x4 OutMatrix)
 {
-	OutMatrix = InstanceTransformInverseTransposed_{ParameterName};
+	OutMatrix = {ParameterName}_InstanceTransformInverseTransposed;
 }
 
 void GetWorldVelocity_{ParameterName}(out float3 Velocity)
 {
-	Velocity = InstanceWorldVelocity_{ParameterName};
+	Velocity = {ParameterName}_InstanceWorldVelocity;
 }
 
 //////////////////////////////////////////////////////////////////////////
 // Experimental Distance Field Functions
-#if DISTATICMESH_ALLOWDISTANCEFIELD
-int		InstanceDistanceFieldIndex_{ParameterName};
-#endif
-
 void QueryDistanceField_{ParameterName}(bool bExecute, float3 WorldPosition, bool bUseMaxDistance, float MaxDistance, out bool OutIsValid, out float OutDistance)
 {
 	if ( !bUseMaxDistance )
@@ -572,13 +572,13 @@ void QueryDistanceField_{ParameterName}(bool bExecute, float3 WorldPosition, boo
 	}
 
 #if DISTATICMESH_ALLOWDISTANCEFIELD
-	OutIsValid = bExecute && InstanceDistanceFieldIndex_{ParameterName} != int(-1);
+	OutIsValid = bExecute && {ParameterName}_InstanceDistanceFieldIndex != int(-1);
 	OutDistance = MaxDistance;
 	if ( OutIsValid )
 	{
 		// Needs to be a smidge over 0.0 for DistanceToNearestSurfaceForObject to query the SDF
 		MaxDistance = max(0.001, MaxDistance);
-		OutDistance = DistanceToNearestSurfaceForObject(InstanceDistanceFieldIndex_{ParameterName}, WorldPosition, MaxDistance);
+		OutDistance = DistanceToNearestSurfaceForObject({ParameterName}_InstanceDistanceFieldIndex, WorldPosition, MaxDistance);
 	}
 #else
 	OutIsValid = false;
diff --git a/Engine/Plugins/FX/Niagara/Source/Niagara/Private/DataInterface/NiagaraDataInterfaceStaticMesh.cpp b/Engine/Plugins/FX/Niagara/Source/Niagara/Private/DataInterface/NiagaraDataInterfaceStaticMesh.cpp
--- a/Engine/Plugins/FX/Niagara/Source/Niagara/Private/DataInterface/NiagaraDataInterfaceStaticMesh.cpp
+++ b/Engine/Plugins/FX/Niagara/Source/Niagara/Private/DataInterface/NiagaraDataInterfaceStaticMesh.cpp
@@ -3,11 +3,13 @@
 #include "NiagaraDataInterfaceStaticMesh.h"
 #include "NiagaraEmitterInstance.h"
 #include "NiagaraComponent.h"
+#include "NiagaraDistanceFieldHelper.h"
 #include "NiagaraStats.h"
 #include "NiagaraSystemInstance.h"
 #include "NiagaraRenderer.h"
 #include "NiagaraSettings.h"
 #include "NiagaraScript.h"
+#include "NiagaraShaderParametersBuilder.h"
 #include "NiagaraDataInterfaceUtilities.h"
 #include "ShaderParameterUtils.h"
 #include "NiagaraStats.h"
@@ -25,7 +27,6 @@
 
 #include "NiagaraGpuComputeDispatchInterface.h"
 #include "NiagaraGpuComputeDispatch.h"
-#include "NiagaraDistanceFieldParameters.h"
 
 #if WITH_EDITOR
 #include "Subsystems/ImportSubsystem.h"
@@ -38,6 +39,40 @@
 
 namespace NDIStaticMeshLocal
 {
+	BEGIN_SHADER_PARAMETER_STRUCT(FShaderParameters,)
+		SHADER_PARAMETER(FIntVector,			NumTriangles)		// x = Num Sections, y = Num Filtered, z = Num Unfiltered
+		SHADER_PARAMETER(int,					NumVertices)
+		SHADER_PARAMETER(int,					NumUVs)
+		SHADER_PARAMETER(uint32,				HasColors)
+		SHADER_PARAMETER_SRV(Buffer<uint>,		IndexBuffer)
+		SHADER_PARAMETER_SRV(Buffer<float>,		PositionBuffer)
+		SHADER_PARAMETER_SRV(Buffer<float4>,	TangentBuffer)
+		SHADER_PARAMETER_SRV(Buffer<float2>,	UVBuffer)
+		SHADER_PARAMETER_SRV(Buffer<float4>,	ColorBuffer)
+
+		SHADER_PARAMETER(uint32,				HasUniformSampling)
+		SHADER_PARAMETER_SRV(Buffer<uint2>,		UniformSamplingTriangles)
+
+		SHADER_PARAMETER(FIntVector,			SectionCounts)		// x = Num Sections, y = Num Filtered, z = Num Unfiltered
+		SHADER_PARAMETER_SRV(Buffer<uint4>,		SectionInfos)		// FirstTriangle, NumTriangles, Prob, Alias
+		SHADER_PARAMETER_SRV(Buffer<uint>,		FilteredAndUnfilteredSections)
+
+		SHADER_PARAMETER(FIntVector,			SocketCounts)		// x = Num Sockets, y = Num Filtered, z = Num Unfiltered
+		SHADER_PARAMETER_SRV(Buffer<float4>,	SocketTransforms)
+		SHADER_PARAMETER_SRV(Buffer<uint>,		FilteredAndUnfilteredSockets)
+
+		SHADER_PARAMETER(float,					InvDeltaSeconds)
+		SHADER_PARAMETER(FMatrix44f,			InstanceTransform)
+		SHADER_PARAMETER(FMatrix44f,			InstanceTransformInverseTransposed)
+		SHADER_PARAMETER(FQuat4f,				InstanceRotation)
+		SHADER_PARAMETER(FMatrix44f,			InstancePreviousTransform)
+		SHADER_PARAMETER(FMatrix44f,			InstancePreviousTransformInverseTransposed)
+		SHADER_PARAMETER(FQuat4f,				InstancePreviousRotation)
+		SHADER_PARAMETER(FVector3f,				InstanceWorldVelocity)
+
+		SHADER_PARAMETER(int,					InstanceDistanceFieldIndex)
+	END_SHADER_PARAMETER_STRUCT()
+
 	static const TCHAR* TemplateShaderFile = TEXT("/Plugin/FX/Niagara/Private/NiagaraDataInterfaceStaticMeshTemplate.ush");
 
 	static int32 GNiagaraFailStaticMeshDataInterface = 0;
@@ -1211,191 +1246,6 @@ namespace NDIStaticMeshLocal
 
 //////////////////////////////////////////////////////////////////////////
 
-struct FNiagaraDataInterfaceParametersCS_StaticMesh : public FNiagaraDataInterfaceParametersCS
-{
-	DECLARE_TYPE_LAYOUT(FNiagaraDataInterfaceParametersCS_StaticMesh, NonVirtual);
-
-public:
-	void Bind(const FNiagaraDataInterfaceGPUParamInfo& ParameterInfo, const class FShaderParameterMap& ParameterMap)
-	{
-		using namespace NDIStaticMeshLocal;
-
-		NumTrianglesParam.Bind(ParameterMap, *(NumTriangles_String + ParameterInfo.DataInterfaceHLSLSymbol));
-		NumVerticesParam.Bind(ParameterMap, *(NumVertices_String + ParameterInfo.DataInterfaceHLSLSymbol));
-		NumUVsParam.Bind(ParameterMap, *(NumUVs_String + ParameterInfo.DataInterfaceHLSLSymbol));
-		HasColorsParam.Bind(ParameterMap, *(HasColors_String + ParameterInfo.DataInterfaceHLSLSymbol));
-		IndexBufferParam.Bind(ParameterMap, *(IndexBuffer_String + ParameterInfo.DataInterfaceHLSLSymbol));
-		PositionBufferParam.Bind(ParameterMap, *(PositionBuffer_String + ParameterInfo.DataInterfaceHLSLSymbol));
-		TangentBufferParam.Bind(ParameterMap, *(TangentBuffer_String + ParameterInfo.DataInterfaceHLSLSymbol));
-		UVBufferParam.Bind(ParameterMap, *(UVBuffer_String + ParameterInfo.DataInterfaceHLSLSymbol));
-		ColorBufferParam.Bind(ParameterMap, *(ColorBuffer_String + ParameterInfo.DataInterfaceHLSLSymbol));
-
-		HasUniformSamplingParam.Bind(ParameterMap, *(HasUniformSampling_String + ParameterInfo.DataInterfaceHLSLSymbol));
-		UniformSamplingTrianglesParam.Bind(ParameterMap, *(UniformSamplingTriangles_String + ParameterInfo.DataInterfaceHLSLSymbol));
-
-		SectionCountsParam.Bind(ParameterMap, *(SectionCounts_String + ParameterInfo.DataInterfaceHLSLSymbol));
-		SectionInfosParam.Bind(ParameterMap, *(SectionInfos_String + ParameterInfo.DataInterfaceHLSLSymbol));
-		FilteredAndUnfilteredSectionsParam.Bind(ParameterMap, *(FilteredAndUnfilteredSections_String + ParameterInfo.DataInterfaceHLSLSymbol));
-
-		SocketCountsParam.Bind(ParameterMap, *(SocketCounts_String + ParameterInfo.DataInterfaceHLSLSymbol));
-		SocketTransformsParam.Bind(ParameterMap, *(SocketTransforms_String + ParameterInfo.DataInterfaceHLSLSymbol));
-		FilteredAndUnfilteredSocketsParam.Bind(ParameterMap, *(FilteredAndUnfilteredSockets_String + ParameterInfo.DataInterfaceHLSLSymbol));
-
-		InvDeltaSecondsParam.Bind(ParameterMap, *(InvDeltaSeconds_String + ParameterInfo.DataInterfaceHLSLSymbol));
-		InstanceTransformParam.Bind(ParameterMap, *(InstanceTransform_String + ParameterInfo.DataInterfaceHLSLSymbol));
-		InstanceTransformInverseTransposedParam.Bind(ParameterMap, *(InstanceTransformInverseTransposed_String + ParameterInfo.DataInterfaceHLSLSymbol));
-		InstanceRotationParam.Bind(ParameterMap, *(InstanceRotation_String + ParameterInfo.DataInterfaceHLSLSymbol));
-		InstancePreviousTransformParam.Bind(ParameterMap, *(InstancePreviousTransform_String + ParameterInfo.DataInterfaceHLSLSymbol));
-		InstancePreviousTransformInverseTransposedParam.Bind(ParameterMap, *(InstancePreviousTransformInverseTransposed_String + ParameterInfo.DataInterfaceHLSLSymbol));
-		InstancePreviousRotationParam.Bind(ParameterMap, *(InstancePreviousRotation_String + ParameterInfo.DataInterfaceHLSLSymbol));
-		InstanceWorldVelocityParam.Bind(ParameterMap, *(InstanceWorldVelocity_String + ParameterInfo.DataInterfaceHLSLSymbol));
-
-		InstanceDistanceFieldIndexParam.Bind(ParameterMap, *(InstanceDistanceFieldIndex_String + ParameterInfo.DataInterfaceHLSLSymbol));
-		DistanceFieldParameters.Bind(ParameterMap);
-	}
-
-	void Set(FRHICommandList& RHICmdList, const FNiagaraDataInterfaceSetArgs& Context) const
-	{
-		check(IsInRenderingThread());
-
-		FRHIComputeShader* ComputeShaderRHI = Context.Shader.GetComputeShader();
-
-		NDIStaticMeshLocal::FRenderProxy* InterfaceProxy = static_cast<NDIStaticMeshLocal::FRenderProxy*>(Context.DataInterface);
-		const NDIStaticMeshLocal::FInstanceData_RenderThread& InstanceData = InterfaceProxy->PerInstanceData_RT.FindChecked(Context.SystemInstanceID);
-
-		// Set mesh sampling data
-		if ( InstanceData.bIsValid )
-		{
-			const bool bHasColors = InstanceData.MeshColorBufferSRV.IsValid();
-			SetShaderValue(RHICmdList, ComputeShaderRHI, NumTrianglesParam, InstanceData.NumTriangles);
-			SetShaderValue(RHICmdList, ComputeShaderRHI, NumVerticesParam, InstanceData.NumVertices);
-			SetShaderValue(RHICmdList, ComputeShaderRHI, NumUVsParam, InstanceData.NumUVs);
-			SetShaderValue(RHICmdList, ComputeShaderRHI, HasColorsParam, bHasColors ? 1 : 0);
-			SetSRVParameter(RHICmdList, ComputeShaderRHI, IndexBufferParam, InstanceData.MeshIndexBufferSRV);
-			SetSRVParameter(RHICmdList, ComputeShaderRHI, PositionBufferParam, InstanceData.MeshPositionBufferSRV);
-			SetSRVParameter(RHICmdList, ComputeShaderRHI, TangentBufferParam, InstanceData.MeshTangentBufferSRV);
-			SetSRVParameter(RHICmdList, ComputeShaderRHI, UVBufferParam, InstanceData.MeshUVBufferSRV.IsValid() ? InstanceData.MeshUVBufferSRV.GetReference() : FNiagaraRenderer::GetDummyFloat2Buffer());
-			SetSRVParameter(RHICmdList, ComputeShaderRHI, ColorBufferParam, bHasColors ? InstanceData.MeshColorBufferSRV.GetReference() : FNiagaraRenderer::GetDummyWhiteColorBuffer());
-
-			SetShaderValue(RHICmdList, ComputeShaderRHI, HasUniformSamplingParam, InstanceData.bGpuUniformDistribution ? 1 : 0);
-			SetSRVParameter(RHICmdList, ComputeShaderRHI, UniformSamplingTrianglesParam, InstanceData.bGpuUniformDistribution ? InstanceData.MeshUniformSamplingTriangleSRV.GetReference() : FNiagaraRenderer::GetDummyUInt2Buffer());
-		}
-		else
-		{
-			SetShaderValue(RHICmdList, ComputeShaderRHI, NumTrianglesParam, 0);
-			SetShaderValue(RHICmdList, ComputeShaderRHI, NumVerticesParam, 0);
-			SetShaderValue(RHICmdList, ComputeShaderRHI, NumUVsParam, 0);
-			SetShaderValue(RHICmdList, ComputeShaderRHI, HasColorsParam, 0);
-			SetSRVParameter(RHICmdList, ComputeShaderRHI, IndexBufferParam, FNiagaraRenderer::GetDummyUIntBuffer());
-			SetSRVParameter(RHICmdList, ComputeShaderRHI, PositionBufferParam, FNiagaraRenderer::GetDummyFloatBuffer());
-			SetSRVParameter(RHICmdList, ComputeShaderRHI, TangentBufferParam, FNiagaraRenderer::GetDummyFloat4Buffer());
-			SetSRVParameter(RHICmdList, ComputeShaderRHI, UVBufferParam, FNiagaraRenderer::GetDummyFloat2Buffer());
-			SetSRVParameter(RHICmdList, ComputeShaderRHI, ColorBufferParam, FNiagaraRenderer::GetDummyWhiteColorBuffer());
-
-			SetShaderValue(RHICmdList, ComputeShaderRHI, HasUniformSamplingParam, 0);
-			SetSRVParameter(RHICmdList, ComputeShaderRHI, UniformSamplingTrianglesParam, FNiagaraRenderer::GetDummyUInt2Buffer());
-		}
-
-		// Section information
-		SetShaderValue(RHICmdList, ComputeShaderRHI, SectionCountsParam, InstanceData.SectionCounts);
-		SetSRVParameter(RHICmdList, ComputeShaderRHI, SectionInfosParam, InstanceData.SectionInfos.SRV.IsValid() ? InstanceData.SectionInfos.SRV.GetReference() : FNiagaraRenderer::GetDummyUInt4Buffer());
-		SetSRVParameter(RHICmdList, ComputeShaderRHI, FilteredAndUnfilteredSectionsParam, InstanceData.FilteredAndUnfilteredSections.SRV.IsValid() ? InstanceData.FilteredAndUnfilteredSections.SRV.GetReference() : FNiagaraRenderer::GetDummyUIntBuffer());
-
-		// Set socket sampling information
-		SetShaderValue(RHICmdList, ComputeShaderRHI, SocketCountsParam, InstanceData.SocketCounts);
-		SetSRVParameter(RHICmdList, ComputeShaderRHI, SocketTransformsParam, InstanceData.SocketTransforms.SRV.IsValid() ? InstanceData.SocketTransforms.SRV.GetReference() : FNiagaraRenderer::GetDummyFloat4Buffer());
-		SetSRVParameter(RHICmdList, ComputeShaderRHI, FilteredAndUnfilteredSocketsParam, InstanceData.FilteredAndUnfilteredSockets.SRV.IsValid() ? InstanceData.FilteredAndUnfilteredSockets.SRV.GetReference() : FNiagaraRenderer::GetDummyUIntBuffer());
-
-		// Set misc data
-		const float InvDeltaTime = InstanceData.DeltaSeconds > 0.0f ? 1.0f / InstanceData.DeltaSeconds : 0.0f;
-		const FVector3f DeltaPosition = InstanceData.Transform.GetOrigin() - InstanceData.PrevTransform.GetOrigin();
-
-		SetShaderValue(RHICmdList, ComputeShaderRHI, InvDeltaSecondsParam, InvDeltaTime);
-		SetShaderValue(RHICmdList, ComputeShaderRHI, InstanceTransformParam, InstanceData.Transform);
-		SetShaderValue(RHICmdList, ComputeShaderRHI, InstanceTransformInverseTransposedParam, InstanceData.Transform.Inverse().GetTransposed());
-		SetShaderValue(RHICmdList, ComputeShaderRHI, InstanceRotationParam, InstanceData.Rotation);
-		SetShaderValue(RHICmdList, ComputeShaderRHI, InstancePreviousTransformParam, InstanceData.PrevTransform);
-		SetShaderValue(RHICmdList, ComputeShaderRHI, InstancePreviousTransformInverseTransposedParam, InstanceData.PrevTransform.Inverse().GetTransposed());
-		SetShaderValue(RHICmdList, ComputeShaderRHI, InstancePreviousRotationParam, InstanceData.PrevRotation);
-		SetShaderValue(RHICmdList, ComputeShaderRHI, InstanceWorldVelocityParam, DeltaPosition);
-
-		const FDistanceFieldSceneData* DistanceFieldSceneData = static_cast<const FNiagaraGpuComputeDispatch*>(Context.ComputeDispatchInterface)->GetMeshDistanceFieldParameters();	//-BATCHERTODO:
-
-		if ( InstanceDistanceFieldIndexParam.IsBound() )
-		{
-			int32 DistanceFieldIndex = -1;
-			if ( DistanceFieldSceneData != nullptr && InstanceData.DistanceFieldPrimitiveId.IsValid() )
-			{
-				if ( FScene* Scene = Context.ComputeDispatchInterface->GetScene() )
-				{
-					// Kind of gross, but currently no way to reference other primitive scene infos
-					const int32 PrimitiveSceneIndex = Scene->PrimitiveComponentIds.Find(InstanceData.DistanceFieldPrimitiveId);
-					if (PrimitiveSceneIndex != INDEX_NONE)
-					{
-						const TArray<int32, TInlineAllocator<1>>& DFIndices = Scene->Primitives[PrimitiveSceneIndex]->DistanceFieldInstanceIndices;
-						DistanceFieldIndex = DFIndices.Num() > 0 ? DFIndices[0] : -1;
-					}
-				}
-			}
-			SetShaderValue(RHICmdList, ComputeShaderRHI, InstanceDistanceFieldIndexParam, DistanceFieldIndex);
-		}
-
-		if (DistanceFieldParameters.IsBound())
-		{
-			if (ensure(DistanceFieldSceneData != nullptr))
-			{
-				DistanceFieldParameters.Set(RHICmdList, ComputeShaderRHI, DistanceFieldSceneData);
-			}
-			else 
-			{
-				//-TODO: We can't create dummy buffers here due to dll boundaries, we will skip the distance field calcs since DistanceFieldIndex will be set to -1
-				//       However some platforms may still complain due to lack of dummy data bound
-				//FDistanceFieldSceneData DummyDistanceFieldSceneData(Context.ComputeDispatchInterface->GetShaderPlatform());
-				//DistanceFieldParameters.Set(RHICmdList, ComputeShaderRHI, &DummyDistanceFieldSceneData);
-			}
-		}
-	}
-
-private:
-	LAYOUT_FIELD(FShaderParameter, NumTrianglesParam);
-	LAYOUT_FIELD(FShaderParameter, NumVerticesParam);
-	LAYOUT_FIELD(FShaderParameter, NumUVsParam);
-	LAYOUT_FIELD(FShaderParameter, HasColorsParam);
-	LAYOUT_FIELD(FShaderResourceParameter, IndexBufferParam);
-	LAYOUT_FIELD(FShaderResourceParameter, PositionBufferParam);
-	LAYOUT_FIELD(FShaderResourceParameter, TangentBufferParam);
-	LAYOUT_FIELD(FShaderResourceParameter, UVBufferParam);
-	LAYOUT_FIELD(FShaderResourceParameter, ColorBufferParam);
-
-	LAYOUT_FIELD(FShaderParameter, HasUniformSamplingParam);
-	LAYOUT_FIELD(FShaderResourceParameter, UniformSamplingTrianglesParam);
-
-	LAYOUT_FIELD(FShaderParameter,  SectionCountsParam);
-	LAYOUT_FIELD(FShaderResourceParameter, SectionInfosParam);
-	LAYOUT_FIELD(FShaderResourceParameter, FilteredAndUnfilteredSectionsParam);
-
-	LAYOUT_FIELD(FShaderParameter, SocketCountsParam);
-	LAYOUT_FIELD(FShaderResourceParameter, SocketTransformsParam);
-	LAYOUT_FIELD(FShaderResourceParameter, FilteredAndUnfilteredSocketsParam);
-
-	LAYOUT_FIELD(FShaderParameter, InvDeltaSecondsParam);
-	LAYOUT_FIELD(FShaderParameter, InstanceTransformParam);
-	LAYOUT_FIELD(FShaderParameter, InstanceTransformInverseTransposedParam);
-	LAYOUT_FIELD(FShaderParameter, InstanceRotationParam);
-	LAYOUT_FIELD(FShaderParameter, InstancePreviousTransformParam);
-	LAYOUT_FIELD(FShaderParameter, InstancePreviousTransformInverseTransposedParam);
-	LAYOUT_FIELD(FShaderParameter, InstancePreviousRotationParam);
-	LAYOUT_FIELD(FShaderParameter, InstanceWorldVelocityParam);
-
-	LAYOUT_FIELD(FShaderParameter, InstanceDistanceFieldIndexParam);
-	LAYOUT_FIELD(FDistanceFieldParameters, DistanceFieldParameters);	
-};
-
-IMPLEMENT_TYPE_LAYOUT(FNiagaraDataInterfaceParametersCS_StaticMesh);
-IMPLEMENT_NIAGARA_DI_PARAMETER(UNiagaraDataInterfaceStaticMesh, FNiagaraDataInterfaceParametersCS_StaticMesh);
-
-//////////////////////////////////////////////////////////////////////////
-
 UNiagaraDataInterfaceStaticMesh::UNiagaraDataInterfaceStaticMesh(FObjectInitializer const& ObjectInitializer)
 	: Super(ObjectInitializer)
 {
@@ -2429,8 +2279,9 @@ bool UNiagaraDataInterfaceStaticMesh::AppendCompileHash(FNiagaraCompileHashVisit
 {
 	bool bSuccess = Super::AppendCompileHash(InVisitor);
 	FSHAHash Hash = GetShaderFileHash(NDIStaticMeshLocal::TemplateShaderFile, EShaderPlatform::SP_PCD3D_SM5);
-	InVisitor->UpdateString(TEXT("NiagaraDataInterfaceStaticMeshTemplateHLSLSource"), Hash.ToString());
-	InVisitor->UpdatePOD(TEXT("NDIStaticMesh_AllowDistanceField"), GetDefault<UNiagaraSettings>()->NDIStaticMesh_AllowDistanceFields ? 1 : 0);
+	bSuccess &= InVisitor->UpdateString(TEXT("NiagaraDataInterfaceStaticMeshTemplateHLSLSource"), Hash.ToString());
+	bSuccess &= InVisitor->UpdatePOD(TEXT("NDIStaticMesh_AllowDistanceField"), GetDefault<UNiagaraSettings>()->NDIStaticMesh_AllowDistanceFields ? 1 : 0);
+	bSuccess &= InVisitor->UpdateShaderParameters<NDIStaticMeshLocal::FShaderParameters>();
 	return bSuccess;
 }
 
@@ -2547,6 +2398,105 @@ bool UNiagaraDataInterfaceStaticMesh::GetFunctionHLSL(const FNiagaraDataInterfac
 }
 #endif
 
+void UNiagaraDataInterfaceStaticMesh::BuildShaderParameters(FNiagaraShaderParametersBuilder& ShaderParametersBuilder) const
+{
+	ShaderParametersBuilder.AddNestedStruct<NDIStaticMeshLocal::FShaderParameters>();
+	ShaderParametersBuilder.AddIncludedStruct<FDistanceFieldObjectBufferParameters>();
+	ShaderParametersBuilder.AddIncludedStruct<FDistanceFieldAtlasParameters>();
+}
+
+void UNiagaraDataInterfaceStaticMesh::SetShaderParameters(const FNiagaraDataInterfaceSetShaderParametersContext& Context) const
+{
+	NDIStaticMeshLocal::FRenderProxy& DIProxy = Context.GetProxy<NDIStaticMeshLocal::FRenderProxy>();
+	const NDIStaticMeshLocal::FInstanceData_RenderThread& InstanceData = DIProxy.PerInstanceData_RT.FindChecked(Context.GetSystemInstanceID());
+
+	NDIStaticMeshLocal::FShaderParameters* ShaderParameters = Context.GetParameterNestedStruct<NDIStaticMeshLocal::FShaderParameters>();
+
+	// Set mesh sampling data
+	if (InstanceData.bIsValid)
+	{
+		const bool bHasColors = InstanceData.MeshColorBufferSRV.IsValid();
+		ShaderParameters->NumTriangles		= InstanceData.NumTriangles;
+		ShaderParameters->NumVertices		= InstanceData.NumVertices;
+		ShaderParameters->NumUVs			= InstanceData.NumUVs;
+		ShaderParameters->HasColors			= bHasColors ? 1 : 0;
+		ShaderParameters->IndexBuffer		= InstanceData.MeshIndexBufferSRV;
+		ShaderParameters->PositionBuffer	= InstanceData.MeshPositionBufferSRV;
+		ShaderParameters->TangentBuffer		= InstanceData.MeshTangentBufferSRV;
+		ShaderParameters->UVBuffer			= InstanceData.MeshUVBufferSRV.IsValid() ? InstanceData.MeshUVBufferSRV.GetReference() : FNiagaraRenderer::GetDummyFloat2Buffer();
+		ShaderParameters->ColorBuffer		= bHasColors ? InstanceData.MeshColorBufferSRV.GetReference() : FNiagaraRenderer::GetDummyWhiteColorBuffer();
+		ShaderParameters->HasUniformSampling = InstanceData.bGpuUniformDistribution ? 1 : 0;
+		ShaderParameters->UniformSamplingTriangles = InstanceData.bGpuUniformDistribution ? InstanceData.MeshUniformSamplingTriangleSRV.GetReference() : FNiagaraRenderer::GetDummyUInt2Buffer();
+	}
+	else
+	{
+		ShaderParameters->NumTriangles		= FIntVector::ZeroValue;
+		ShaderParameters->NumVertices		= 0;
+		ShaderParameters->NumUVs			= 0;
+		ShaderParameters->HasColors			= 0;
+		ShaderParameters->IndexBuffer		= FNiagaraRenderer::GetDummyUIntBuffer();
+		ShaderParameters->PositionBuffer	= FNiagaraRenderer::GetDummyFloatBuffer();
+		ShaderParameters->TangentBuffer		= FNiagaraRenderer::GetDummyFloat4Buffer();
+		ShaderParameters->UVBuffer			= FNiagaraRenderer::GetDummyFloat2Buffer();
+		ShaderParameters->ColorBuffer		= FNiagaraRenderer::GetDummyWhiteColorBuffer();
+
+		ShaderParameters->HasUniformSampling = 0;
+		ShaderParameters->UniformSamplingTriangles = FNiagaraRenderer::GetDummyUInt2Buffer();
+	}
+
+	// Section information
+	ShaderParameters->SectionCounts	= InstanceData.SectionCounts;
+	ShaderParameters->SectionInfos	= InstanceData.SectionInfos.SRV.IsValid() ? InstanceData.SectionInfos.SRV.GetReference() : FNiagaraRenderer::GetDummyUInt4Buffer();
+	ShaderParameters->FilteredAndUnfilteredSections = InstanceData.FilteredAndUnfilteredSections.SRV.IsValid() ? InstanceData.FilteredAndUnfilteredSections.SRV.GetReference() : FNiagaraRenderer::GetDummyUIntBuffer();
+
+	// Set socket sampling information
+	ShaderParameters->SocketCounts		= InstanceData.SocketCounts;
+	ShaderParameters->SocketTransforms	= InstanceData.SocketTransforms.SRV.IsValid() ? InstanceData.SocketTransforms.SRV.GetReference() : FNiagaraRenderer::GetDummyFloat4Buffer();
+	ShaderParameters->FilteredAndUnfilteredSockets = InstanceData.FilteredAndUnfilteredSockets.SRV.IsValid() ? InstanceData.FilteredAndUnfilteredSockets.SRV.GetReference() : FNiagaraRenderer::GetDummyUIntBuffer();
+
+	// Set misc data
+	const float InvDeltaTime = InstanceData.DeltaSeconds > 0.0f ? 1.0f / InstanceData.DeltaSeconds : 0.0f;
+	const FVector3f DeltaPosition = InstanceData.Transform.GetOrigin() - InstanceData.PrevTransform.GetOrigin();
+
+	ShaderParameters->InvDeltaSeconds								= InvDeltaTime;
+	ShaderParameters->InstanceTransform								= InstanceData.Transform;
+	ShaderParameters->InstanceTransformInverseTransposed			= InstanceData.Transform.Inverse().GetTransposed();
+	ShaderParameters->InstanceRotation								= InstanceData.Rotation;
+	ShaderParameters->InstancePreviousTransform						= InstanceData.PrevTransform;
+	ShaderParameters->InstancePreviousTransformInverseTransposed	= InstanceData.PrevTransform.Inverse().GetTransposed();
+	ShaderParameters->InstancePreviousRotation						= InstanceData.PrevRotation;
+	ShaderParameters->InstanceWorldVelocity							= DeltaPosition;
+
+	const FDistanceFieldSceneData* DistanceFieldSceneData = static_cast<const FNiagaraGpuComputeDispatch&>(Context.GetComputeDispatchInterface()).GetMeshDistanceFieldParameters();	//-BATCHERTODO:
+	if (Context.IsParameterBound(&ShaderParameters->InstanceDistanceFieldIndex))
+	{
+		int32 DistanceFieldIndex = -1;
+		if (DistanceFieldSceneData != nullptr && InstanceData.DistanceFieldPrimitiveId.IsValid())
+		{
+			if (FScene* Scene = Context.GetComputeDispatchInterface().GetScene())
+			{
+				// Kind of gross, but currently no way to reference other primitive scene infos
+				const int32 PrimitiveSceneIndex = Scene->PrimitiveComponentIds.Find(InstanceData.DistanceFieldPrimitiveId);
+				if (PrimitiveSceneIndex != INDEX_NONE)
+				{
+					const TArray<int32, TInlineAllocator<1>>& DFIndices = Scene->Primitives[PrimitiveSceneIndex]->DistanceFieldInstanceIndices;
+					DistanceFieldIndex = DFIndices.Num() > 0 ? DFIndices[0] : -1;
+				}
+			}
+		}
+		ShaderParameters->InstanceDistanceFieldIndex = DistanceFieldIndex;
+	}
+
+	// Bind Mesh Distance Field Data
+	FDistanceFieldObjectBufferParameters* ShaderDistanceFieldObjectParameters = Context.GetParameterIncludedStruct<FDistanceFieldObjectBufferParameters>();
+	FDistanceFieldAtlasParameters* ShaderDistanceFieldAtlasParameters = Context.GetParameterIncludedStruct<FDistanceFieldAtlasParameters>();
+	const bool bDistanceFieldDataBound = Context.IsStructBound<FDistanceFieldObjectBufferParameters>(ShaderDistanceFieldObjectParameters) || Context.IsStructBound<FDistanceFieldAtlasParameters>(ShaderDistanceFieldAtlasParameters);
+	if (bDistanceFieldDataBound)
+	{
+		FNiagaraDistanceFieldHelper::SetMeshDistanceFieldParameters(Context.GetGraphBuilder(), DistanceFieldSceneData, *ShaderDistanceFieldObjectParameters, *ShaderDistanceFieldAtlasParameters, FNiagaraRenderer::GetDummyFloat4Buffer());
+	}
+}
+
 bool UNiagaraDataInterfaceStaticMesh::Equals(const UNiagaraDataInterface* Other) const
 {
 	if (!Super::Equals(Other))
diff --git a/Engine/Plugins/FX/Niagara/Source/Niagara/Private/DataInterface/NiagaraDataInterfaceStaticMesh.h b/Engine/Plugins/FX/Niagara/Source/Niagara/Private/DataInterface/NiagaraDataInterfaceStaticMesh.h
--- a/Engine/Plugins/FX/Niagara/Source/Niagara/Private/DataInterface/NiagaraDataInterfaceStaticMesh.h
+++ b/Engine/Plugins/FX/Niagara/Source/Niagara/Private/DataInterface/NiagaraDataInterfaceStaticMesh.h
@@ -56,8 +56,6 @@ class NIAGARA_API UNiagaraDataInterfaceStaticMesh : public UNiagaraDataInterface
 	GENERATED_UCLASS_BODY()
 
 public:
-	DECLARE_NIAGARA_DI_PARAMETER();
-
 	/** Controls how to retrieve the Static Mesh Component to attach to. */
 	UPROPERTY(EditAnywhere, Category = "Mesh")
 	ENDIStaticMesh_SourceMode SourceMode = ENDIStaticMesh_SourceMode::Default;
@@ -133,6 +131,9 @@ public:
 	virtual void GetParameterDefinitionHLSL(const FNiagaraDataInterfaceGPUParamInfo& ParamInfo, FString& OutHLSL) override;
 	virtual bool GetFunctionHLSL(const FNiagaraDataInterfaceGPUParamInfo& ParamInfo, const FNiagaraDataInterfaceGeneratedFunction& FunctionInfo, int FunctionInstanceIndex, FString& OutHLSL) override;
 #endif
+	virtual bool UseLegacyShaderBindings() const override { return false; }
+	virtual void BuildShaderParameters(FNiagaraShaderParametersBuilder& ShaderParametersBuilder) const override;
+	virtual void SetShaderParameters(const FNiagaraDataInterfaceSetShaderParametersContext& Context) const override;
 
 	virtual bool Equals(const UNiagaraDataInterface* Other) const override;
 protected:
diff --git a/Engine/Plugins/FX/Niagara/Source/Niagara/Private/DataInterface/NiagaraDistanceFieldParameters.h b/Engine/Plugins/FX/Niagara/Source/Niagara/Private/DataInterface/NiagaraDistanceFieldParameters.h
deleted file mode 100644
diff --git a/Engine/Plugins/FX/Niagara/Source/Niagara/Private/NiagaraDataInterfaceRigidMeshCollisionQuery.cpp b/Engine/Plugins/FX/Niagara/Source/Niagara/Private/NiagaraDataInterfaceRigidMeshCollisionQuery.cpp
--- a/Engine/Plugins/FX/Niagara/Source/Niagara/Private/NiagaraDataInterfaceRigidMeshCollisionQuery.cpp
+++ b/Engine/Plugins/FX/Niagara/Source/Niagara/Private/NiagaraDataInterfaceRigidMeshCollisionQuery.cpp
@@ -12,6 +12,7 @@
 
 #include "NiagaraComponent.h"
 #include "NiagaraDataInterfaceUtilities.h"
+#include "NiagaraDistanceFieldHelper.h"
 #include "NiagaraFunctionLibrary.h"
 #include "NiagaraGpuComputeDispatch.h"
 #include "NiagaraGpuComputeDispatchInterface.h"
@@ -1376,7 +1377,7 @@ void UNiagaraDataInterfaceRigidMeshCollisionQuery::SetShaderParameters(const FNi
 	FDistanceFieldAtlasParameters* ShaderDistanceFieldAtlasParameters = Context.GetParameterIncludedStruct<FDistanceFieldAtlasParameters>();
 
 	const bool bDistanceFieldDataBound = Context.IsStructBound<FDistanceFieldObjectBufferParameters>(ShaderDistanceFieldObjectParameters) || Context.IsStructBound<FDistanceFieldAtlasParameters>(ShaderDistanceFieldAtlasParameters);
-	bool bDistanceFieldDataSet = false;
+	const FDistanceFieldSceneData* DistanceFieldSceneData = nullptr;
 
 	if (ProxyData != nullptr && ProxyData->AssetBuffer != nullptr && ProxyData->AssetBuffer->IsInitialized())
 	{
@@ -1399,14 +1400,7 @@ void UNiagaraDataInterfaceRigidMeshCollisionQuery::SetShaderParameters(const FNi
 
 		if (bDistanceFieldDataBound)
 		{
-			const FDistanceFieldSceneData* DistanceFieldSceneData = static_cast<const FNiagaraGpuComputeDispatch&>(Context.GetComputeDispatchInterface()).GetMeshDistanceFieldParameters();	//-BATCHERTODO:
-			if (DistanceFieldSceneData != nullptr)
-			{
-				bDistanceFieldDataSet = true;
-
-				*ShaderDistanceFieldObjectParameters = DistanceField::SetupObjectBufferParameters(GraphBuilder, *DistanceFieldSceneData);
-				*ShaderDistanceFieldAtlasParameters = DistanceField::SetupAtlasParameters(GraphBuilder, *DistanceFieldSceneData);
-			}
+			DistanceFieldSceneData = static_cast<const FNiagaraGpuComputeDispatch&>(Context.GetComputeDispatchInterface()).GetMeshDistanceFieldParameters();	//-BATCHERTODO:
 		}
 	}
 	else
@@ -1423,30 +1417,9 @@ void UNiagaraDataInterfaceRigidMeshCollisionQuery::SetShaderParameters(const FNi
 		ShaderParameters->ElementOffsets = FUintVector4(0, 0, 0, 0);
 	}
 
-	if (bDistanceFieldDataBound && !bDistanceFieldDataSet)
-	{
-		ShaderDistanceFieldObjectParameters->SceneObjectBounds				= FNiagaraRenderer::GetDummyFloat4Buffer();
-		ShaderDistanceFieldObjectParameters->SceneObjectData				= FNiagaraRenderer::GetDummyFloat4Buffer();
-		ShaderDistanceFieldObjectParameters->NumSceneObjects				= 0;
-		ShaderDistanceFieldObjectParameters->SceneHeightfieldObjectBounds	= Context.GetComputeDispatchInterface().GetEmptyBufferSRV(GraphBuilder, PF_R8G8B8A8);
-		ShaderDistanceFieldObjectParameters->SceneHeightfieldObjectData		= Context.GetComputeDispatchInterface().GetEmptyBufferSRV(GraphBuilder, PF_R8G8B8A8);
-		ShaderDistanceFieldObjectParameters->NumSceneHeightfieldObjects		= 0;
-
-		ShaderDistanceFieldAtlasParameters->SceneDistanceFieldAssetData						= Context.GetComputeDispatchInterface().GetEmptyBufferSRV(GraphBuilder, PF_R8G8B8A8);
-		ShaderDistanceFieldAtlasParameters->DistanceFieldIndirectionTable					= Context.GetComputeDispatchInterface().GetEmptyBufferSRV(GraphBuilder, PF_R32_UINT);
-		ShaderDistanceFieldAtlasParameters->DistanceFieldIndirection2Table					= Context.GetComputeDispatchInterface().GetEmptyBufferSRV(GraphBuilder, PF_R8G8B8A8);
-		ShaderDistanceFieldAtlasParameters->DistanceFieldIndirectionAtlas					= Context.GetComputeDispatchInterface().GetBlackTexture(GraphBuilder, ETextureDimension::Texture3D);
-		ShaderDistanceFieldAtlasParameters->DistanceFieldBrickTexture						= GBlackVolumeTexture->GetTextureRHI();
-		ShaderDistanceFieldAtlasParameters->DistanceFieldSampler							= TStaticSamplerState<SF_Bilinear, AM_Clamp, AM_Clamp, AM_Clamp>::GetRHI();
-		ShaderDistanceFieldAtlasParameters->DistanceFieldBrickSize							= FVector3f::ZeroVector;
-		ShaderDistanceFieldAtlasParameters->DistanceFieldUniqueDataBrickSize				= FVector3f::ZeroVector;
-		ShaderDistanceFieldAtlasParameters->DistanceFieldBrickAtlasSizeInBricks				= FIntVector::ZeroValue;
-		ShaderDistanceFieldAtlasParameters->DistanceFieldBrickAtlasMask						= FIntVector::ZeroValue;
-		ShaderDistanceFieldAtlasParameters->DistanceFieldBrickAtlasSizeLog2					= FIntVector::ZeroValue;
-		ShaderDistanceFieldAtlasParameters->DistanceFieldBrickAtlasTexelSize				= FVector3f::ZeroVector;
-		ShaderDistanceFieldAtlasParameters->DistanceFieldBrickAtlasHalfTexelSize			= FVector3f::ZeroVector;
-		ShaderDistanceFieldAtlasParameters->DistanceFieldBrickOffsetToAtlasUVScale			= FVector3f::ZeroVector;
-		ShaderDistanceFieldAtlasParameters->DistanceFieldUniqueDataBrickSizeInAtlasTexels	= FVector3f::ZeroVector;
+	if (bDistanceFieldDataBound)
+	{
+		FNiagaraDistanceFieldHelper::SetMeshDistanceFieldParameters(Context.GetGraphBuilder(), DistanceFieldSceneData, *ShaderDistanceFieldObjectParameters, *ShaderDistanceFieldAtlasParameters, FNiagaraRenderer::GetDummyFloat4Buffer());
 	}
 }
 
diff --git a/Engine/Plugins/FX/Niagara/Source/NiagaraShader/Private/NiagaraDistanceFieldHelper.cpp b/Engine/Plugins/FX/Niagara/Source/NiagaraShader/Private/NiagaraDistanceFieldHelper.cpp
--- a/Engine/Plugins/FX/Niagara/Source/NiagaraShader/Private/NiagaraDistanceFieldHelper.cpp
+++ b/Engine/Plugins/FX/Niagara/Source/NiagaraShader/Private/NiagaraDistanceFieldHelper.cpp
@@ -1,6 +1,7 @@
 // Copyright Epic Games, Inc. All Rights Reserved.
 
 #include "NiagaraDistanceFieldHelper.h"
+#include "Runtime/Renderer/Private/SystemTextures.h"
 
 // todo - currently duplicated from SetupGlobalDistanceFieldParameters (GlobalDistanceField.cpp) because of problems getting it properly exported from the dll
 void FNiagaraDistanceFieldHelper::SetGlobalDistanceFieldParameters(const FGlobalDistanceFieldParameterData* OptionalParameterData, FGlobalDistanceFieldParameters2& ShaderParameters)
@@ -65,3 +66,40 @@ void FNiagaraDistanceFieldHelper::SetGlobalDistanceFieldParameters(const FGlobal
 		ShaderParameters.UncoveredMinStepScale = 0.0f;//GLumenSceneGlobalSDFUncoveredMinStepScale;
 	}
 }
+
+void FNiagaraDistanceFieldHelper::SetMeshDistanceFieldParameters(FRDGBuilder& GraphBuilder, const FDistanceFieldSceneData* OptionalDistanceFieldData, FDistanceFieldObjectBufferParameters& ObjectShaderParameters, FDistanceFieldAtlasParameters& AtlasShaderParameters, FRHIShaderResourceView* DummyFloat4Buffer)
+{
+	if (OptionalDistanceFieldData != nullptr && OptionalDistanceFieldData->NumObjectsInBuffer > 0)
+	{
+		ObjectShaderParameters = DistanceField::SetupObjectBufferParameters(GraphBuilder, *OptionalDistanceFieldData);
+		AtlasShaderParameters = DistanceField::SetupAtlasParameters(GraphBuilder, *OptionalDistanceFieldData);
+	}
+	else
+	{
+		FRDGBufferSRVRef DefaultVector4 = GraphBuilder.CreateSRV(FRDGBufferSRVDesc(GSystemTextures.GetDefaultStructuredBuffer(GraphBuilder, sizeof(FVector4f))));
+		FRDGBufferSRVRef DefaultUInt32 = GraphBuilder.CreateSRV(FRDGBufferSRVDesc(GSystemTextures.GetDefaultStructuredBuffer(GraphBuilder, sizeof(uint32))));
+
+		ObjectShaderParameters.SceneObjectBounds = DummyFloat4Buffer;
+		ObjectShaderParameters.SceneObjectData = DummyFloat4Buffer;
+		ObjectShaderParameters.NumSceneObjects = 0;
+		ObjectShaderParameters.SceneHeightfieldObjectBounds = DefaultVector4;
+		ObjectShaderParameters.SceneHeightfieldObjectData = DefaultVector4;
+		ObjectShaderParameters.NumSceneHeightfieldObjects = 0;
+
+		AtlasShaderParameters.SceneDistanceFieldAssetData = DefaultVector4;
+		AtlasShaderParameters.DistanceFieldIndirectionTable = DefaultUInt32;
+		AtlasShaderParameters.DistanceFieldIndirection2Table = DefaultVector4;
+		AtlasShaderParameters.DistanceFieldIndirectionAtlas = FRDGSystemTextures::Get(GraphBuilder).VolumetricBlack;
+		AtlasShaderParameters.DistanceFieldBrickTexture = GBlackVolumeTexture->GetTextureRHI();
+		AtlasShaderParameters.DistanceFieldSampler = TStaticSamplerState<SF_Bilinear, AM_Clamp, AM_Clamp, AM_Clamp>::GetRHI();
+		AtlasShaderParameters.DistanceFieldBrickSize = FVector3f::ZeroVector;
+		AtlasShaderParameters.DistanceFieldUniqueDataBrickSize = FVector3f::ZeroVector;
+		AtlasShaderParameters.DistanceFieldBrickAtlasSizeInBricks = FIntVector::ZeroValue;
+		AtlasShaderParameters.DistanceFieldBrickAtlasMask = FIntVector::ZeroValue;
+		AtlasShaderParameters.DistanceFieldBrickAtlasSizeLog2 = FIntVector::ZeroValue;
+		AtlasShaderParameters.DistanceFieldBrickAtlasTexelSize = FVector3f::ZeroVector;
+		AtlasShaderParameters.DistanceFieldBrickAtlasHalfTexelSize = FVector3f::ZeroVector;
+		AtlasShaderParameters.DistanceFieldBrickOffsetToAtlasUVScale = FVector3f::ZeroVector;
+		AtlasShaderParameters.DistanceFieldUniqueDataBrickSizeInAtlasTexels = FVector3f::ZeroVector;
+	}
+}
diff --git a/Engine/Plugins/FX/Niagara/Source/NiagaraShader/Public/NiagaraDistanceFieldHelper.h b/Engine/Plugins/FX/Niagara/Source/NiagaraShader/Public/NiagaraDistanceFieldHelper.h
--- a/Engine/Plugins/FX/Niagara/Source/NiagaraShader/Public/NiagaraDistanceFieldHelper.h
+++ b/Engine/Plugins/FX/Niagara/Source/NiagaraShader/Public/NiagaraDistanceFieldHelper.h
@@ -4,8 +4,10 @@
 
 #include "GlobalDistanceFieldParameters.h"
 #include "Renderer/Private/DistanceFieldLightingShared.h"
+#include "RenderGraphBuilder.h"
 
 namespace FNiagaraDistanceFieldHelper
 {
 	NIAGARASHADER_API void SetGlobalDistanceFieldParameters(const FGlobalDistanceFieldParameterData* OptionalParameterData, FGlobalDistanceFieldParameters2& ShaderParameters);
+	NIAGARASHADER_API void SetMeshDistanceFieldParameters(FRDGBuilder& GraphBuilder, const FDistanceFieldSceneData* OptionalDistanceFieldData, FDistanceFieldObjectBufferParameters& ObjectShaderParameters, FDistanceFieldAtlasParameters& AtlasShaderParameters, FRHIShaderResourceView* DummyFloat4Buffer);
 }
