fixed indexing bug with Niagara Baker and Grid3D Collection output

diff --git a/Engine/Plugins/FX/Niagara/Source/NiagaraEditor/Private/NiagaraBakerRendererOutputVolumeTexture.cpp b/Engine/Plugins/FX/Niagara/Source/NiagaraEditor/Private/NiagaraBakerRendererOutputVolumeTexture.cpp
--- a/Engine/Plugins/FX/Niagara/Source/NiagaraEditor/Private/NiagaraBakerRendererOutputVolumeTexture.cpp
+++ b/Engine/Plugins/FX/Niagara/Source/NiagaraEditor/Private/NiagaraBakerRendererOutputVolumeTexture.cpp
@@ -490,7 +490,7 @@ void FNiagaraBakerRendererOutputVolumeTexture::BakeFrame(UNiagaraBakerOutput* In
 				const int32 SliceOffset = z * TextureSize.X * TextureSize.Y;
 				for ( int32 y=0; y < TileTextureSize.Y; ++y )
 				{
-					const int32 SliceRowOffset = SliceOffset + (y * TextureSize.Y);
+					const int32 SliceRowOffset = SliceOffset + (y * TextureSize.X);
 					for ( int32 x=0; x < TileTextureSize.X; ++x )
 					{
 						const int32 Offset = SliceRowOffset + x;
