- Convert Physics DI to RDG / Shader Params

diff --git a/Engine/Plugins/FX/Niagara/Shaders/Private/NiagaraDataInterfacePhysicsAsset.ush b/Engine/Plugins/FX/Niagara/Shaders/Private/NiagaraDataInterfacePhysicsAsset.ush
--- a/Engine/Plugins/FX/Niagara/Shaders/Private/NiagaraDataInterfacePhysicsAsset.ush
+++ b/Engine/Plugins/FX/Niagara/Shaders/Private/NiagaraDataInterfacePhysicsAsset.ush
@@ -4,72 +4,48 @@
 NiagaraDataInterfacePhysicsAsset.ush 
 =============================================================================*/
 
-#pragma once
-
 /* -----------------------------------------------------------------
  * Physics Asset constants and context
  * -----------------------------------------------------------------
  */
 
- #define PHYSICS_ASSET_MAX_TRANSFORMS 300
+#define PHYSICS_ASSET_MAX_TRANSFORMS 300
 
- #define CURRENT_OFFSET 0
- #define PREVIOUS_OFFSET PHYSICS_ASSET_MAX_TRANSFORMS
- #define REST_OFFSET 2*PHYSICS_ASSET_MAX_TRANSFORMS
+#define CURRENT_OFFSET 0
+#define PREVIOUS_OFFSET PHYSICS_ASSET_MAX_TRANSFORMS
+#define REST_OFFSET 2*PHYSICS_ASSET_MAX_TRANSFORMS
 
 #define MAX_DISTANCE 1e+38
 #define BOX_INDEX 0
 #define SPHERE_INDEX 1
 #define CAPSULE_INDEX 2
 
-struct FDIPhysicsAssetContext
-{
-	uint4 ElementOffsets;
-	Buffer<float4>		WorldTransformBuffer;
-	Buffer<float4>		InverseTransformBuffer;
-	Buffer<float4>		ElementExtentBuffer;
-	Buffer<uint>		PhysicsTypeBuffer;
-	float3				BoxOrigin;
-	float3				BoxExtent;
-}; 
-
-#define DIPHYSICSASSET_DECLARE_CONSTANTS(NAME)\
-uint4 ElementOffsets_##NAME;\
-Buffer<float4>		WorldTransformBuffer_##NAME;\
-Buffer<float4>		InverseTransformBuffer_##NAME;\
-Buffer<float4>		ElementExtentBuffer_##NAME;\
-Buffer<uint>		PhysicsTypeBuffer_##NAME;\
-float3 BoxOrigin_##NAME;\
-float3 BoxExtent_##NAME;\
-
-#define DIPHYSICSASSET_MAKE_CONTEXT(NAME)\
-FDIPhysicsAssetContext DIContext; \
-DIContext.ElementOffsets = ElementOffsets_##NAME;\
-DIContext.WorldTransformBuffer = WorldTransformBuffer_##NAME;\
-DIContext.InverseTransformBuffer = InverseTransformBuffer_##NAME;\
-DIContext.ElementExtentBuffer = ElementExtentBuffer_##NAME;\
-DIContext.PhysicsTypeBuffer = PhysicsTypeBuffer_##NAME;\
-DIContext.BoxOrigin = BoxOrigin_##NAME;\
-DIContext.BoxExtent = BoxExtent_##NAME;\
+uint4				{ParameterName}_ElementOffsets;
+Buffer<float4>		{ParameterName}_WorldTransformBuffer;
+Buffer<float4>		{ParameterName}_InverseTransformBuffer;
+Buffer<float4>		{ParameterName}_ElementExtentBuffer;
+Buffer<uint>		{ParameterName}_PhysicsTypeBuffer;
+float3				{ParameterName}_BoxOrigin;
+float3				{ParameterName}_BoxExtent;
 
 /* -----------------------------------------------------------------
  * Element accessors
  * -----------------------------------------------------------------
  */
 
-int DIPhysicsAsset_GetNumBoxes(in FDIPhysicsAssetContext DIContext)
+void GetNumBoxes_{ParameterName}(out int OutNumBoxes)
 {
-	return DIContext.ElementOffsets[BOX_INDEX+1] - DIContext.ElementOffsets[BOX_INDEX];
+	OutNumBoxes = {ParameterName}_ElementOffsets[BOX_INDEX+1] - {ParameterName}_ElementOffsets[BOX_INDEX];
 }
 
-int DIPhysicsAsset_GetNumSpheres(in FDIPhysicsAssetContext DIContext)
+void GetNumSpheres_{ParameterName}(out int OutNumSpheres)
 {
-	return DIContext.ElementOffsets[SPHERE_INDEX+1] - DIContext.ElementOffsets[SPHERE_INDEX];
+	OutNumSpheres = {ParameterName}_ElementOffsets[SPHERE_INDEX+1] - {ParameterName}_ElementOffsets[SPHERE_INDEX];
 }
 
-int DIPhysicsAsset_GetNumCapsules(in FDIPhysicsAssetContext DIContext)
+void GetNumCapsules_{ParameterName}(out int OutNumCapsules)
 {
-	return DIContext.ElementOffsets[CAPSULE_INDEX+1] - DIContext.ElementOffsets[CAPSULE_INDEX];
+	OutNumCapsules = {ParameterName}_ElementOffsets[CAPSULE_INDEX+1] - {ParameterName}_ElementOffsets[CAPSULE_INDEX];
 }
 
 
@@ -78,73 +54,73 @@ int DIPhysicsAsset_GetNumCapsules(in FDIPhysicsAssetContext DIContext)
  * -----------------------------------------------------------------
  */
 
-float3 DIPhysicsAsset_GetCurrentPosition(in FDIPhysicsAssetContext DIContext, in int ElementIndex)
+float3 GetCurrentPosition_{ParameterName}(in int ElementIndex)
 {
 	const int BufferOffset = 3 * ElementIndex + CURRENT_OFFSET;
-    return float3(DIContext.WorldTransformBuffer[BufferOffset].w, DIContext.WorldTransformBuffer[BufferOffset+1].w, DIContext.WorldTransformBuffer[BufferOffset+2].w);
+    return float3({ParameterName}_WorldTransformBuffer[BufferOffset].w, {ParameterName}_WorldTransformBuffer[BufferOffset+1].w, {ParameterName}_WorldTransformBuffer[BufferOffset+2].w);
 }
 
-float3 DIPhysicsAsset_GetPreviousPosition(in FDIPhysicsAssetContext DIContext, in int ElementIndex)
+float3 GetPreviousPosition_{ParameterName}(in int ElementIndex)
 {
 	const int BufferOffset = 3 * ElementIndex + PREVIOUS_OFFSET;
-    return float3(DIContext.WorldTransformBuffer[BufferOffset].w, DIContext.WorldTransformBuffer[BufferOffset+1].w, DIContext.WorldTransformBuffer[BufferOffset+2].w);
+    return float3({ParameterName}_WorldTransformBuffer[BufferOffset].w, {ParameterName}_WorldTransformBuffer[BufferOffset+1].w, {ParameterName}_WorldTransformBuffer[BufferOffset+2].w);
 }
 
-float4 DIPhysicsAsset_GetCurrentOrientation(in FDIPhysicsAssetContext DIContext, in int ElementIndex)
+float4 GetCurrentOrientation_{ParameterName}(in int ElementIndex)
 {
 	const int BufferOffset = 3 * ElementIndex + CURRENT_OFFSET;
-	const float3 OrientationMatrix[3] = {float3(DIContext.WorldTransformBuffer[BufferOffset].x, DIContext.WorldTransformBuffer[BufferOffset+1].x, DIContext.WorldTransformBuffer[BufferOffset+2].x),
-									     float3(DIContext.WorldTransformBuffer[BufferOffset].y, DIContext.WorldTransformBuffer[BufferOffset+1].y, DIContext.WorldTransformBuffer[BufferOffset+2].y),
-										 float3(DIContext.WorldTransformBuffer[BufferOffset].z, DIContext.WorldTransformBuffer[BufferOffset+1].z, DIContext.WorldTransformBuffer[BufferOffset+2].z)};
+	const float3 OrientationMatrix[3] = {float3({ParameterName}_WorldTransformBuffer[BufferOffset].x, {ParameterName}_WorldTransformBuffer[BufferOffset+1].x, {ParameterName}_WorldTransformBuffer[BufferOffset+2].x),
+									     float3({ParameterName}_WorldTransformBuffer[BufferOffset].y, {ParameterName}_WorldTransformBuffer[BufferOffset+1].y, {ParameterName}_WorldTransformBuffer[BufferOffset+2].y),
+										 float3({ParameterName}_WorldTransformBuffer[BufferOffset].z, {ParameterName}_WorldTransformBuffer[BufferOffset+1].z, {ParameterName}_WorldTransformBuffer[BufferOffset+2].z)};
 
 
 	return QuatFromMatrix(OrientationMatrix);
 }
 
-float4 DIPhysicsAsset_GetPreviousOrientation(in FDIPhysicsAssetContext DIContext, in int ElementIndex)
+float4 GetPreviousOrientation_{ParameterName}(in int ElementIndex)
 {
 	const int BufferOffset = 3 * ElementIndex + PREVIOUS_OFFSET;
-	const float3 OrientationMatrix[3] = {float3(DIContext.WorldTransformBuffer[BufferOffset].x, DIContext.WorldTransformBuffer[BufferOffset+1].x, DIContext.WorldTransformBuffer[BufferOffset+2].x),
-										 float3(DIContext.WorldTransformBuffer[BufferOffset].y, DIContext.WorldTransformBuffer[BufferOffset+1].y, DIContext.WorldTransformBuffer[BufferOffset+2].y),
-										 float3(DIContext.WorldTransformBuffer[BufferOffset].z, DIContext.WorldTransformBuffer[BufferOffset+1].z, DIContext.WorldTransformBuffer[BufferOffset+2].z)};
+	const float3 OrientationMatrix[3] = {float3({ParameterName}_WorldTransformBuffer[BufferOffset].x, {ParameterName}_WorldTransformBuffer[BufferOffset+1].x, {ParameterName}_WorldTransformBuffer[BufferOffset+2].x),
+										 float3({ParameterName}_WorldTransformBuffer[BufferOffset].y, {ParameterName}_WorldTransformBuffer[BufferOffset+1].y, {ParameterName}_WorldTransformBuffer[BufferOffset+2].y),
+										 float3({ParameterName}_WorldTransformBuffer[BufferOffset].z, {ParameterName}_WorldTransformBuffer[BufferOffset+1].z, {ParameterName}_WorldTransformBuffer[BufferOffset+2].z)};
 
 	return QuatFromMatrix(OrientationMatrix);
 }
 
-float3x4 DIPhysicsAsset_GetCurrentTransform(in FDIPhysicsAssetContext DIContext, in int ElementIndex)
+float3x4 GetCurrentTransform_{ParameterName}(in int ElementIndex)
 {
 	const int BufferOffset = 3 * ElementIndex + CURRENT_OFFSET;
-    return float3x4(DIContext.WorldTransformBuffer[BufferOffset], DIContext.WorldTransformBuffer[BufferOffset+1], DIContext.WorldTransformBuffer[BufferOffset+2]);
+    return float3x4({ParameterName}_WorldTransformBuffer[BufferOffset], {ParameterName}_WorldTransformBuffer[BufferOffset+1], {ParameterName}_WorldTransformBuffer[BufferOffset+2]);
 }
 
-float3x4 DIPhysicsAsset_GetInverseTransform(in FDIPhysicsAssetContext DIContext, in int ElementIndex)
+float3x4 GetInverseTransform_{ParameterName}(in int ElementIndex)
 {
 	const int BufferOffset = 3 * ElementIndex + CURRENT_OFFSET;
-    return float3x4(DIContext.InverseTransformBuffer[BufferOffset], DIContext.InverseTransformBuffer[BufferOffset+1], DIContext.InverseTransformBuffer[BufferOffset+2]);
+    return float3x4({ParameterName}_InverseTransformBuffer[BufferOffset], {ParameterName}_InverseTransformBuffer[BufferOffset+1], {ParameterName}_InverseTransformBuffer[BufferOffset+2]);
 }
 
-float3x4 DIPhysicsAsset_GetPreviousTransform(in FDIPhysicsAssetContext DIContext, in int ElementIndex)
+float3x4 GetPreviousTransform_{ParameterName}(in int ElementIndex)
 {
 	const int BufferOffset = 3 * ElementIndex + PREVIOUS_OFFSET;
-    return float3x4(DIContext.WorldTransformBuffer[BufferOffset], DIContext.WorldTransformBuffer[BufferOffset+1], DIContext.WorldTransformBuffer[BufferOffset+2]);
+    return float3x4({ParameterName}_WorldTransformBuffer[BufferOffset], {ParameterName}_WorldTransformBuffer[BufferOffset+1], {ParameterName}_WorldTransformBuffer[BufferOffset+2]);
 }
 
-float3x4 DIPhysicsAsset_GetPreviousInverse(in FDIPhysicsAssetContext DIContext, in int ElementIndex)
+float3x4 GetPreviousInverse_{ParameterName}(in int ElementIndex)
 {
 	const int BufferOffset = 3 * ElementIndex + PREVIOUS_OFFSET;
-    return float3x4(DIContext.InverseTransformBuffer[BufferOffset], DIContext.InverseTransformBuffer[BufferOffset+1], DIContext.InverseTransformBuffer[BufferOffset+2]);
+    return float3x4({ParameterName}_InverseTransformBuffer[BufferOffset], {ParameterName}_InverseTransformBuffer[BufferOffset+1], {ParameterName}_InverseTransformBuffer[BufferOffset+2]);
 }
 
-float3x4 DIPhysicsAsset_GetRestTransform(in FDIPhysicsAssetContext DIContext, in int ElementIndex)
+float3x4 GetRestTransform_{ParameterName}(in int ElementIndex)
 {
 	const int BufferOffset = 3 * ElementIndex + REST_OFFSET;
-    return float3x4(DIContext.WorldTransformBuffer[BufferOffset], DIContext.WorldTransformBuffer[BufferOffset+1], DIContext.WorldTransformBuffer[BufferOffset+2]);
+    return float3x4({ParameterName}_WorldTransformBuffer[BufferOffset], {ParameterName}_WorldTransformBuffer[BufferOffset+1], {ParameterName}_WorldTransformBuffer[BufferOffset+2]);
 }
 
-float3x4 DIPhysicsAsset_GetRestInverse(in FDIPhysicsAssetContext DIContext, in int ElementIndex)
+float3x4 GetRestInverse_{ParameterName}(in int ElementIndex)
 {
 	const int BufferOffset = 3 * ElementIndex + REST_OFFSET;
-    return float3x4(DIContext.InverseTransformBuffer[BufferOffset], DIContext.InverseTransformBuffer[BufferOffset+1], DIContext.InverseTransformBuffer[BufferOffset+2]);
+    return float3x4({ParameterName}_InverseTransformBuffer[BufferOffset], {ParameterName}_InverseTransformBuffer[BufferOffset+1], {ParameterName}_InverseTransformBuffer[BufferOffset+2]);
 }
 
 /* -----------------------------------------------------------------
@@ -153,7 +129,7 @@ float3x4 DIPhysicsAsset_GetRestInverse(in FDIPhysicsAssetContext DIContext, in i
  */
 
 // Given a world space position (NodePosition) compute the sphere closest point (position,normal,velocity)
-float DIPhysicsAsset_GetSphereProjection(in float3 LocalPosition, in float3 SphereCenter, in float SphereRadius, in int SphereIndex,
+float GetSphereProjection_{ParameterName}(in float3 LocalPosition, in float3 SphereCenter, in float SphereRadius, in int SphereIndex,
 			inout float3 OutClosestPosition, inout float3 OutClosestNormal, inout int OutElementIndex, inout float OutMinDistance)
 {
 	const float3 DeltaPosition = LocalPosition - SphereCenter;
@@ -172,18 +148,18 @@ float DIPhysicsAsset_GetSphereProjection(in float3 LocalPosition, in float3 Sphe
 }
 
 // Given a world space position (NodePosition) compute the sphere closest point (position,normal,velocity)
-float DIPhysicsAsset_GetCapsuleProjection(in float3 LocalPosition, in float2 CapsuleSize, in int CapsuleIndex,
+float GetCapsuleProjection_{ParameterName}(in float3 LocalPosition, in float2 CapsuleSize, in int CapsuleIndex,
 			inout float3 OutClosestPosition, inout float3 OutClosestNormal, inout int OutElementIndex, inout float OutMinDistance)
 {
 	const float HalfLength = 0.5 * CapsuleSize.y;
 	const float3 SphereCenter = float3(0,0,clamp(LocalPosition.z, -HalfLength, HalfLength));
 
-	return DIPhysicsAsset_GetSphereProjection(LocalPosition,SphereCenter,CapsuleSize.x,CapsuleIndex,
+	return GetSphereProjection_{ParameterName}(LocalPosition,SphereCenter,CapsuleSize.x,CapsuleIndex,
 			OutClosestPosition,OutClosestNormal,OutElementIndex,OutMinDistance);
 }
 
 // Given a world space position (NodePosition) compute the sphere closest point (position,normal,velocity)
-float DIPhysicsAsset_GetBoxProjection(in float3 LocalPosition, in float3 BoxExtent, in int BoxIndex, 
+float GetBoxProjection_{ParameterName}(in float3 LocalPosition, in float3 BoxExtent, in int BoxIndex, 
 			inout float3 OutClosestPosition, inout float3 OutClosestNormal, inout int OutElementIndex, inout float OutMinDistance)
 {
 	const float3 HalfExtent = 0.5 * BoxExtent;
@@ -235,15 +211,15 @@ float DIPhysicsAsset_GetBoxProjection(in float3 LocalPosition, in float3 BoxExte
  * -----------------------------------------------------------------
  */
 
-/*float3 DIPhysicsAsset_GetLocalPosition(in FDIPhysicsAssetContext DIContext, in float3 NodePosition, in int ElementIndex, in float TimeFraction )
+/*float3 GetLocalPosition_{ParameterName}(in float3 NodePosition, in int ElementIndex, in float TimeFraction )
 {
-	//return mul(DIPhysicsAsset_GetInverseTransform(DIContext,ElementIndex), float4(NodePosition,1.0)).xyz;
+	//return mul(GetInverseTransform_{ParameterName}(ElementIndex), float4(NodePosition,1.0)).xyz;
 
-	const float3 CurrentPosition = DIPhysicsAsset_GetCurrentPosition(DIContext,ElementIndex);
-	const float4 CurrentOrientation = DIPhysicsAsset_GetCurrentOrientation(DIContext,ElementIndex);
+	const float3 CurrentPosition = GetCurrentPosition_{ParameterName}(ElementIndex);
+	const float4 CurrentOrientation = GetCurrentOrientation_{ParameterName}(ElementIndex);
 
-	const float3 PreviousPosition = DIPhysicsAsset_GetPreviousPosition(DIContext,ElementIndex);
-	const float4 PreviousOrientation = DIPhysicsAsset_GetPreviousOrientation(DIContext,ElementIndex);
+	const float3 PreviousPosition = GetPreviousPosition_{ParameterName}(ElementIndex);
+	const float4 PreviousOrientation = GetPreviousOrientation_{ParameterName}(ElementIndex);
 
 	const float4 SlerpedOrientation = InverseQuat(SlerpQuat(PreviousOrientation,CurrentOrientation,TimeFraction));
 	const float3 SlerpedPosition = PreviousPosition + TimeFraction * (CurrentPosition-PreviousPosition);
@@ -251,16 +227,16 @@ float DIPhysicsAsset_GetBoxProjection(in float3 LocalPosition, in float3 BoxExte
 	return RotateVectorByQuat( NodePosition - SlerpedPosition, SlerpedOrientation );
 }*/
 
-float3 DIPhysicsAsset_GetLocalPosition(in FDIPhysicsAssetContext DIContext, in float3 NodePosition, in int ElementIndex, in float TimeFraction )
+float3 GetLocalPosition_{ParameterName}(in float3 NodePosition, in int ElementIndex, in float TimeFraction )
 {
-	const float3 CurrentLocal = mul(DIPhysicsAsset_GetInverseTransform(DIContext,ElementIndex), float4(NodePosition,1.0)).xyz;
-	const float3 PreviousLocal = mul(DIPhysicsAsset_GetPreviousInverse(DIContext,ElementIndex), float4(NodePosition,1.0)).xyz;
+	const float3 CurrentLocal = mul(GetInverseTransform_{ParameterName}(ElementIndex), float4(NodePosition,1.0)).xyz;
+	const float3 PreviousLocal = mul(GetPreviousInverse_{ParameterName}(ElementIndex), float4(NodePosition,1.0)).xyz;
 
 	return PreviousLocal + TimeFraction * (CurrentLocal-PreviousLocal);
 }
 
 // Given a world space position (NodePosition) compute the sphere closest point (position,normal,velocity)
-int DIPhysicsAsset_GetClosestElement(in FDIPhysicsAssetContext DIContext, in float3 NodePosition, out float3 OutClosestPosition, 
+int GetClosestElement_{ParameterName}(in float3 NodePosition, out float3 OutClosestPosition, 
 										  out float3 OutClosestNormal, out float OutMinDistance, in float TimeFraction, in bool IsRestSpace )
 {	
 	float MinDistance = MAX_DISTANCE;
@@ -268,36 +244,36 @@ int DIPhysicsAsset_GetClosestElement(in FDIPhysicsAssetContext DIContext, in flo
 	float3 CollisionPosition = float3(0,0,0);
 	float3 CollisionNormal = float3(0,0,0);
 
-	const int SpheresBegin = DIContext.ElementOffsets[SPHERE_INDEX];
-	const int SpheresEnd = DIContext.ElementOffsets[SPHERE_INDEX+1];
+	const int SpheresBegin = {ParameterName}_ElementOffsets[SPHERE_INDEX];
+	const int SpheresEnd = {ParameterName}_ElementOffsets[SPHERE_INDEX+1];
 	for (int SphereIndex = SpheresBegin; SphereIndex < SpheresEnd; ++SphereIndex)
 	{
-		const float3 LocalPosition = IsRestSpace ? mul(DIPhysicsAsset_GetRestInverse(DIContext,SphereIndex), float4(NodePosition,1.0)).xyz : 
-						DIPhysicsAsset_GetLocalPosition(DIContext,NodePosition,SphereIndex,TimeFraction);
+		const float3 LocalPosition = IsRestSpace ? mul(GetRestInverse_{ParameterName}(SphereIndex), float4(NodePosition,1.0)).xyz : 
+						GetLocalPosition_{ParameterName}(NodePosition,SphereIndex,TimeFraction);
 
-		DIPhysicsAsset_GetSphereProjection(LocalPosition, float3(0,0,0), DIContext.ElementExtentBuffer[SphereIndex].x, SphereIndex,
+		GetSphereProjection_{ParameterName}(LocalPosition, float3(0,0,0), {ParameterName}_ElementExtentBuffer[SphereIndex].x, SphereIndex,
 				CollisionPosition, CollisionNormal, ElementIndex, MinDistance);
 	}
 
-	const int CapsulesBegin = DIContext.ElementOffsets[CAPSULE_INDEX];
-	const int CapsulesEnd = DIContext.ElementOffsets[CAPSULE_INDEX+1];
+	const int CapsulesBegin = {ParameterName}_ElementOffsets[CAPSULE_INDEX];
+	const int CapsulesEnd = {ParameterName}_ElementOffsets[CAPSULE_INDEX+1];
 	for (int CapsuleIndex = CapsulesBegin; CapsuleIndex < CapsulesEnd; ++CapsuleIndex)
 	{
-		const float3 LocalPosition = IsRestSpace ? mul(DIPhysicsAsset_GetRestInverse(DIContext,CapsuleIndex), float4(NodePosition,1.0)).xyz : 
-						DIPhysicsAsset_GetLocalPosition(DIContext,NodePosition,CapsuleIndex,TimeFraction);
+		const float3 LocalPosition = IsRestSpace ? mul(GetRestInverse_{ParameterName}(CapsuleIndex), float4(NodePosition,1.0)).xyz : 
+						GetLocalPosition_{ParameterName}(NodePosition,CapsuleIndex,TimeFraction);
 
-		DIPhysicsAsset_GetCapsuleProjection(LocalPosition, DIContext.ElementExtentBuffer[CapsuleIndex].xy, CapsuleIndex,
+		GetCapsuleProjection_{ParameterName}(LocalPosition, {ParameterName}_ElementExtentBuffer[CapsuleIndex].xy, CapsuleIndex,
 				CollisionPosition, CollisionNormal, ElementIndex, MinDistance);
 	}
 
-	const int BoxesBegin = DIContext.ElementOffsets[BOX_INDEX];
-	const int BoxesEnd = DIContext.ElementOffsets[BOX_INDEX+1];
+	const int BoxesBegin = {ParameterName}_ElementOffsets[BOX_INDEX];
+	const int BoxesEnd = {ParameterName}_ElementOffsets[BOX_INDEX+1];
 	for (int BoxIndex = BoxesBegin; BoxIndex < BoxesEnd; ++BoxIndex)
 	{
-		const float3 LocalPosition = IsRestSpace ? mul(DIPhysicsAsset_GetRestInverse(DIContext,BoxIndex), float4(NodePosition,1.0)).xyz : 
-						DIPhysicsAsset_GetLocalPosition(DIContext,NodePosition,BoxIndex,TimeFraction);
+		const float3 LocalPosition = IsRestSpace ? mul(GetRestInverse_{ParameterName}(BoxIndex), float4(NodePosition,1.0)).xyz : 
+						GetLocalPosition_{ParameterName}(NodePosition,BoxIndex,TimeFraction);
 
-		DIPhysicsAsset_GetBoxProjection(LocalPosition, DIContext.ElementExtentBuffer[BoxIndex].xyz, BoxIndex,
+		GetBoxProjection_{ParameterName}(LocalPosition, {ParameterName}_ElementExtentBuffer[BoxIndex].xyz, BoxIndex,
 				CollisionPosition, CollisionNormal, ElementIndex, MinDistance);
 	}
 	OutClosestPosition = CollisionPosition;
@@ -312,7 +288,7 @@ int DIPhysicsAsset_GetClosestElement(in FDIPhysicsAssetContext DIContext, in flo
  */
 
 // Given a world space position (NodePosition) and an element index compute the physics asset closest point (position,normal,velocity)
-void DIPhysicsAsset_GetElementPoint(in FDIPhysicsAssetContext DIContext, in float3 NodePosition, in float DeltaTime, in float TimeFraction, in int ElementIndex, out float3 OutClosestPosition, 
+void GetElementPoint_{ParameterName}(in float3 NodePosition, in float DeltaTime, in float TimeFraction, in int ElementIndex, out float3 OutClosestPosition, 
 										  out float3 OutClosestNormal, out float3 OutClosestVelocity)
 {	
 	OutClosestNormal = float3(0,0,0);
@@ -322,36 +298,36 @@ void DIPhysicsAsset_GetElementPoint(in FDIPhysicsAssetContext DIContext, in floa
 	if (ElementIndex >= 0)
 	{
 		const uint ElementIndexUint = uint(ElementIndex);
-		const float3 LocalPosition = DIPhysicsAsset_GetLocalPosition(DIContext,NodePosition,ElementIndex,TimeFraction);
+		const float3 LocalPosition = GetLocalPosition_{ParameterName}(NodePosition,ElementIndex,TimeFraction);
 
 		float MinDistance = MAX_DISTANCE;
 		int OutputIndex = -1;
 		float3 CollisionPosition = float3(0,0,0);
 		float3 CollisionNormal = float3(0,0,0);
 
-		if ( ElementIndexUint >= DIContext.ElementOffsets[SPHERE_INDEX] && ElementIndexUint < DIContext.ElementOffsets[SPHERE_INDEX+1] )
+		if ( ElementIndexUint >= {ParameterName}_ElementOffsets[SPHERE_INDEX] && ElementIndexUint < {ParameterName}_ElementOffsets[SPHERE_INDEX+1] )
 		{
-			DIPhysicsAsset_GetSphereProjection(LocalPosition, float3(0,0,0), DIContext.ElementExtentBuffer[ElementIndex].x, ElementIndex,
+			GetSphereProjection_{ParameterName}(LocalPosition, float3(0,0,0), {ParameterName}_ElementExtentBuffer[ElementIndex].x, ElementIndex,
 					CollisionPosition, CollisionNormal, OutputIndex, MinDistance);
 		}
-		else if ( ElementIndexUint >= DIContext.ElementOffsets[CAPSULE_INDEX] && ElementIndexUint < DIContext.ElementOffsets[CAPSULE_INDEX+1] )
+		else if ( ElementIndexUint >= {ParameterName}_ElementOffsets[CAPSULE_INDEX] && ElementIndexUint < {ParameterName}_ElementOffsets[CAPSULE_INDEX+1] )
 		{
-			DIPhysicsAsset_GetCapsuleProjection(LocalPosition, DIContext.ElementExtentBuffer[ElementIndex].xy, ElementIndex,
+			GetCapsuleProjection_{ParameterName}(LocalPosition, {ParameterName}_ElementExtentBuffer[ElementIndex].xy, ElementIndex,
 					CollisionPosition, CollisionNormal, OutputIndex, MinDistance);
 		}
-		else if ( ElementIndexUint >= DIContext.ElementOffsets[BOX_INDEX] && ElementIndexUint < DIContext.ElementOffsets[BOX_INDEX+1] )
+		else if ( ElementIndexUint >= {ParameterName}_ElementOffsets[BOX_INDEX] && ElementIndexUint < {ParameterName}_ElementOffsets[BOX_INDEX+1] )
 		{
-			DIPhysicsAsset_GetBoxProjection(LocalPosition, DIContext.ElementExtentBuffer[ElementIndex].xyz, ElementIndex,
+			GetBoxProjection_{ParameterName}(LocalPosition, {ParameterName}_ElementExtentBuffer[ElementIndex].xyz, ElementIndex,
 					CollisionPosition, CollisionNormal, OutputIndex, MinDistance);
 		}
 
 		if (ElementIndex != -1)
 		{
-			const float3 PreviousPosition = mul(DIPhysicsAsset_GetPreviousTransform(DIContext,ElementIndex), float4(CollisionPosition,1.0)).xyz;
-			const float3 CurrentPosition = mul(DIPhysicsAsset_GetCurrentTransform(DIContext,ElementIndex), float4(CollisionPosition,1.0)).xyz;
+			const float3 PreviousPosition = mul(GetPreviousTransform_{ParameterName}(ElementIndex), float4(CollisionPosition,1.0)).xyz;
+			const float3 CurrentPosition = mul(GetCurrentTransform_{ParameterName}(ElementIndex), float4(CollisionPosition,1.0)).xyz;
 
-			const float3 PreviousNormal = mul(DIPhysicsAsset_GetPreviousTransform(DIContext,ElementIndex), float4(CollisionNormal,0.0)).xyz;
-			const float3 CurrentNormal = mul(DIPhysicsAsset_GetCurrentTransform(DIContext,ElementIndex), float4(CollisionNormal,0.0)).xyz;
+			const float3 PreviousNormal = mul(GetPreviousTransform_{ParameterName}(ElementIndex), float4(CollisionNormal,0.0)).xyz;
+			const float3 CurrentNormal = mul(GetCurrentTransform_{ParameterName}(ElementIndex), float4(CollisionNormal,0.0)).xyz;
 
 			OutClosestPosition = PreviousPosition + TimeFraction * (CurrentPosition-PreviousPosition);
 			OutClosestNormal = normalize(PreviousNormal + TimeFraction * (CurrentNormal-PreviousNormal));
@@ -362,62 +338,62 @@ void DIPhysicsAsset_GetElementPoint(in FDIPhysicsAssetContext DIContext, in floa
 }
 
 // Given a world space position (NodePosition) and an element index compute the physics asset closest distance
-void DIPhysicsAsset_GetElementDistance(in FDIPhysicsAssetContext DIContext, in float3 NodePosition, in float TimeFraction, in int ElementIndex, out float OutClosestDistance)
+void GetElementDistance_{ParameterName}(in float3 NodePosition, in float TimeFraction, in int ElementIndex, out float OutClosestDistance)
 {	
 	OutClosestDistance = 0.0;
 
 	if (ElementIndex >= 0)
 	{
 		const uint ElementIndexUint = uint(ElementIndex);
-		const float3 LocalPosition = DIPhysicsAsset_GetLocalPosition(DIContext,NodePosition,ElementIndex,TimeFraction);
+		const float3 LocalPosition = GetLocalPosition_{ParameterName}(NodePosition,ElementIndex,TimeFraction);
 
 		float MinDistance = MAX_DISTANCE;
 		int OutputIndex = -1;
 		float3 CollisionPosition = float3(0,0,0);
 		float3 CollisionNormal = float3(0,0,0);
 
-		if( ElementIndexUint >= DIContext.ElementOffsets[SPHERE_INDEX] && ElementIndexUint < DIContext.ElementOffsets[SPHERE_INDEX+1] )
+		if( ElementIndexUint >= {ParameterName}_ElementOffsets[SPHERE_INDEX] && ElementIndexUint < {ParameterName}_ElementOffsets[SPHERE_INDEX+1] )
 		{
-			DIPhysicsAsset_GetSphereProjection(LocalPosition, float3(0,0,0), DIContext.ElementExtentBuffer[ElementIndex].x, ElementIndex,
+			GetSphereProjection_{ParameterName}(LocalPosition, float3(0,0,0), {ParameterName}_ElementExtentBuffer[ElementIndex].x, ElementIndex,
 					CollisionPosition, CollisionNormal, OutputIndex, MinDistance);
 		}
-		else if ( ElementIndexUint >= DIContext.ElementOffsets[CAPSULE_INDEX] && ElementIndexUint < DIContext.ElementOffsets[CAPSULE_INDEX+1] )
+		else if ( ElementIndexUint >= {ParameterName}_ElementOffsets[CAPSULE_INDEX] && ElementIndexUint < {ParameterName}_ElementOffsets[CAPSULE_INDEX+1] )
 		{
-			DIPhysicsAsset_GetCapsuleProjection(LocalPosition, DIContext.ElementExtentBuffer[ElementIndex].xy, ElementIndex,
+			GetCapsuleProjection_{ParameterName}(LocalPosition, {ParameterName}_ElementExtentBuffer[ElementIndex].xy, ElementIndex,
 					CollisionPosition, CollisionNormal, OutputIndex, MinDistance);
 		}
-		else if ( ElementIndexUint >= DIContext.ElementOffsets[BOX_INDEX] && ElementIndexUint < DIContext.ElementOffsets[BOX_INDEX+1] )
+		else if ( ElementIndexUint >= {ParameterName}_ElementOffsets[BOX_INDEX] && ElementIndexUint < {ParameterName}_ElementOffsets[BOX_INDEX+1] )
 		{
-			DIPhysicsAsset_GetBoxProjection(LocalPosition, DIContext.ElementExtentBuffer[ElementIndex].xyz, ElementIndex,
+			GetBoxProjection_{ParameterName}(LocalPosition, {ParameterName}_ElementExtentBuffer[ElementIndex].xyz, ElementIndex,
 					CollisionPosition, CollisionNormal, OutputIndex, MinDistance);
 		}
 
 		if (ElementIndex != -1)
 		{
-			OutClosestDistance = (DIContext.PhysicsTypeBuffer[ElementIndex] == 1) ? MinDistance : 0;
+			OutClosestDistance = ({ParameterName}_PhysicsTypeBuffer[ElementIndex] == 1) ? MinDistance : 0;
 		}
 	}
 }
 
  // Given a world space position (NodePosition) compute the physics asset closest element
-void DIPhysicsAsset_GetClosestElement(in FDIPhysicsAssetContext DIContext, in float3 NodePosition, in float TimeFraction, out int OutElementIndex)
+void GetClosestElement_{ParameterName}(in float3 NodePosition, in float TimeFraction, out int OutElementIndex)
 {	
 	float3 CollisionPosition = float3(0,0,0);
 	float3 CollisionNormal = float3(0,0,0);
 	float MinDistance = 0.0;
 
-	OutElementIndex = DIPhysicsAsset_GetClosestElement(DIContext,NodePosition,CollisionPosition,CollisionNormal,MinDistance,TimeFraction, false);
+	OutElementIndex = GetClosestElement_{ParameterName}(NodePosition,CollisionPosition,CollisionNormal,MinDistance,TimeFraction, false);
 }
 
 // Given a world space position (NodePosition) compute the physics asset closest point (position,normal,velocity)
-void DIPhysicsAsset_GetClosestPoint(in FDIPhysicsAssetContext DIContext, in float3 NodePosition, in float DeltaTime, in float TimeFraction, out float3 OutClosestPosition, 
+void GetClosestPoint_{ParameterName}(in float3 NodePosition, in float DeltaTime, in float TimeFraction, out float3 OutClosestPosition, 
 										  out float3 OutClosestNormal, out float3 OutClosestVelocity)
 {	
 	float3 CollisionPosition = float3(0,0,0);
 	float3 CollisionNormal = float3(0,0,0);
 	float MinDistance = 0.0;
 
-	const int ElementIndex = DIPhysicsAsset_GetClosestElement(DIContext,NodePosition,CollisionPosition,CollisionNormal,MinDistance,TimeFraction, false);
+	const int ElementIndex = GetClosestElement_{ParameterName}(NodePosition,CollisionPosition,CollisionNormal,MinDistance,TimeFraction, false);
 
 	OutClosestNormal = float3(0,0,0);
 	OutClosestPosition = float3(0,0,0);
@@ -425,11 +401,11 @@ void DIPhysicsAsset_GetClosestPoint(in FDIPhysicsAssetContext DIContext, in floa
 
 	if (ElementIndex != -1)
 	{
-		const float3 PreviousPosition = mul(DIPhysicsAsset_GetPreviousTransform(DIContext,ElementIndex), float4(CollisionPosition,1.0)).xyz;
-		const float3 CurrentPosition = mul(DIPhysicsAsset_GetCurrentTransform(DIContext,ElementIndex), float4(CollisionPosition,1.0)).xyz;
+		const float3 PreviousPosition = mul(GetPreviousTransform_{ParameterName}(ElementIndex), float4(CollisionPosition,1.0)).xyz;
+		const float3 CurrentPosition = mul(GetCurrentTransform_{ParameterName}(ElementIndex), float4(CollisionPosition,1.0)).xyz;
 
-		const float3 PreviousNormal = mul(DIPhysicsAsset_GetPreviousTransform(DIContext,ElementIndex), float4(CollisionNormal,0.0)).xyz;
-		const float3 CurrentNormal = mul(DIPhysicsAsset_GetCurrentTransform(DIContext,ElementIndex), float4(CollisionNormal,0.0)).xyz;
+		const float3 PreviousNormal = mul(GetPreviousTransform_{ParameterName}(ElementIndex), float4(CollisionNormal,0.0)).xyz;
+		const float3 CurrentNormal = mul(GetCurrentTransform_{ParameterName}(ElementIndex), float4(CollisionNormal,0.0)).xyz;
 
 		OutClosestPosition = PreviousPosition + TimeFraction * (CurrentPosition-PreviousPosition);
 		OutClosestNormal = normalize(PreviousNormal + TimeFraction * (CurrentNormal-PreviousNormal));
@@ -439,19 +415,19 @@ void DIPhysicsAsset_GetClosestPoint(in FDIPhysicsAssetContext DIContext, in floa
 }
 
 // Given a world space position (NodePosition) compute the physics asset closest distance
-void DIPhysicsAsset_GetClosestDistance(in FDIPhysicsAssetContext DIContext, in float3 NodePosition, in float TimeFraction, out float OutClosestDistance)
+void GetClosestDistance_{ParameterName}(in float3 NodePosition, in float TimeFraction, out float OutClosestDistance)
 {	
 	float3 CollisionPosition = float3(0,0,0);
 	float3 CollisionNormal = float3(0,0,0);
 	float MinDistance = 0.0;
 
-	const int ElementIndex = DIPhysicsAsset_GetClosestElement(DIContext,NodePosition,CollisionPosition,CollisionNormal,MinDistance,TimeFraction, false);
+	const int ElementIndex = GetClosestElement_{ParameterName}(NodePosition,CollisionPosition,CollisionNormal,MinDistance,TimeFraction, false);
 
 	OutClosestDistance = 0.0;
 
 	if (ElementIndex != -1)
 	{
-		OutClosestDistance = (DIContext.PhysicsTypeBuffer[ElementIndex] == 1) ? MinDistance : 0;
+		OutClosestDistance = ({ParameterName}_PhysicsTypeBuffer[ElementIndex] == 1) ? MinDistance : 0;
 	}
 }
 
@@ -461,13 +437,13 @@ void DIPhysicsAsset_GetClosestDistance(in FDIPhysicsAssetContext DIContext, in f
  * -----------------------------------------------------------------
  */
 
-void DIPhysicsAsset_GetRestDistance(in FDIPhysicsAssetContext DIContext, in float3 RestPosition, out float OutRestDistance)
+void GetRestDistance_{ParameterName}(in float3 RestPosition, out float OutRestDistance)
 {
 	float3 CollisionPosition = float3(0,0,0);
 	float3 CollisionNormal = float3(0,0,0);
 	float MinDistance = 0.0;
 
-	const int ElementIndex = DIPhysicsAsset_GetClosestElement(DIContext,RestPosition,CollisionPosition,CollisionNormal,MinDistance, 1.0, true);
+	const int ElementIndex = GetClosestElement_{ParameterName}(RestPosition,CollisionPosition,CollisionNormal,MinDistance, 1.0, true);
 
 	OutRestDistance = 0.0;
 
@@ -477,7 +453,7 @@ void DIPhysicsAsset_GetRestDistance(in FDIPhysicsAssetContext DIContext, in floa
 	}
 }
 
-void DIPhysicsAsset_GetTexturePoint(in FDIPhysicsAssetContext DIContext, in float3 NodePosition, out int OutElementIndex, out float3 OutTexturePosition)
+void GetTexturePoint_{ParameterName}(in float3 NodePosition, out int OutElementIndex, out float3 OutTexturePosition)
 {
 	float3 CollisionPosition = float3(0,0,0);
 	float3 CollisionNormal = float3(0,0,0);
@@ -485,24 +461,24 @@ void DIPhysicsAsset_GetTexturePoint(in FDIPhysicsAssetContext DIContext, in floa
 
 	OutTexturePosition = float3(0,0,0);
 
-	OutElementIndex = DIPhysicsAsset_GetClosestElement(DIContext,NodePosition,CollisionPosition,CollisionNormal,MinDistance, 1.0, false);
+	OutElementIndex = GetClosestElement_{ParameterName}(NodePosition,CollisionPosition,CollisionNormal,MinDistance, 1.0, false);
 
-	const float AxisExtent = max( max(DIContext.BoxExtent.x, DIContext.BoxExtent.y), DIContext.BoxExtent.z );
+	const float AxisExtent = max( max({ParameterName}_BoxExtent.x, {ParameterName}_BoxExtent.y), {ParameterName}_BoxExtent.z );
 	const float3 BoxExtent = float3(AxisExtent,AxisExtent,AxisExtent);
 
-	const float3 BoxMinimum = DIContext.BoxOrigin - BoxExtent;
+	const float3 BoxMinimum = {ParameterName}_BoxOrigin - BoxExtent;
 	const float3 BoxSize = 2.0 * BoxExtent;
 
 	if (OutElementIndex != -1)
 	{
-		const float3 LocalPosition = mul(DIPhysicsAsset_GetInverseTransform(DIContext,OutElementIndex), float4(NodePosition,1.0)).xyz;
-		const float3 RestPosition = mul(DIPhysicsAsset_GetRestTransform(DIContext,OutElementIndex), float4(LocalPosition,1.0)).xyz;
+		const float3 LocalPosition = mul(GetInverseTransform_{ParameterName}(OutElementIndex), float4(NodePosition,1.0)).xyz;
+		const float3 RestPosition = mul(GetRestTransform_{ParameterName}(OutElementIndex), float4(LocalPosition,1.0)).xyz;
 
 		OutTexturePosition = clamp( (RestPosition - BoxMinimum)/BoxSize, float3(0,0,0), float3(1,1,1) );
 	}
 }
 
-void DIPhysicsAsset_GetProjectionPoint(in FDIPhysicsAssetContext DIContext, in float3 NodePosition, in float DeltaTime, in int ElementIndex,
+void GetProjectionPoint_{ParameterName}(in float3 NodePosition, in float DeltaTime, in int ElementIndex,
 			in float TextureValue, in float3 TextureGradient, out float3 OutClosestPosition, out float3 OutClosestNormal, out float3 OutClosestVelocity, out float OutClosestDistance)
 {
 	OutClosestNormal = float3(0,0,0);
@@ -512,22 +488,32 @@ void DIPhysicsAsset_GetProjectionPoint(in FDIPhysicsAssetContext DIContext, in f
 
 	if (ElementIndex != -1)
 	{
-		const float AxisExtent = max( max(DIContext.BoxExtent.x, DIContext.BoxExtent.y), DIContext.BoxExtent.z );
+		const float AxisExtent = max( max({ParameterName}_BoxExtent.x, {ParameterName}_BoxExtent.y), {ParameterName}_BoxExtent.z );
 
-		const float3 LocalPosition = mul(DIPhysicsAsset_GetInverseTransform(DIContext,ElementIndex), float4(NodePosition,1.0)).xyz;
-		const float3 RestPosition = mul(DIPhysicsAsset_GetRestTransform(DIContext,ElementIndex), float4(LocalPosition,1.0)).xyz;
+		const float3 LocalPosition = mul(GetInverseTransform_{ParameterName}(ElementIndex), float4(NodePosition,1.0)).xyz;
+		const float3 RestPosition = mul(GetRestTransform_{ParameterName}(ElementIndex), float4(LocalPosition,1.0)).xyz;
 
 		const float3 RestProjection = RestPosition - TextureValue * 2.0 * AxisExtent * TextureGradient;
 
-		const float3 CollisionPosition = mul(DIPhysicsAsset_GetRestInverse(DIContext,ElementIndex), float4(RestProjection,1.0)).xyz;
-		const float3 CollisionNormal = mul(DIPhysicsAsset_GetRestInverse(DIContext,ElementIndex), float4(TextureGradient,0.0)).xyz;
+		const float3 CollisionPosition = mul(GetRestInverse_{ParameterName}(ElementIndex), float4(RestProjection,1.0)).xyz;
+		const float3 CollisionNormal = mul(GetRestInverse_{ParameterName}(ElementIndex), float4(TextureGradient,0.0)).xyz;
 
-		const float3 PreviousPosition = mul(DIPhysicsAsset_GetPreviousTransform(DIContext,ElementIndex), float4(CollisionPosition,1.0)).xyz;
+		const float3 PreviousPosition = mul(GetPreviousTransform_{ParameterName}(ElementIndex), float4(CollisionPosition,1.0)).xyz;
 
-		OutClosestPosition = mul(DIPhysicsAsset_GetCurrentTransform(DIContext,ElementIndex), float4(CollisionPosition,1.0)).xyz;
-		OutClosestNormal = mul(DIPhysicsAsset_GetCurrentTransform(DIContext,ElementIndex), float4(CollisionNormal,0.0)).xyz;
+		OutClosestPosition = mul(GetCurrentTransform_{ParameterName}(ElementIndex), float4(CollisionPosition,1.0)).xyz;
+		OutClosestNormal = mul(GetCurrentTransform_{ParameterName}(ElementIndex), float4(CollisionNormal,0.0)).xyz;
 		OutClosestVelocity = ( OutClosestPosition - PreviousPosition ) / DeltaTime;
 		OutClosestDistance = TextureValue * 2.0 * AxisExtent;
 	}
 }
 
+#undef PHYSICS_ASSET_MAX_TRANSFORMS
+
+#undef CURRENT_OFFSET
+#undef PREVIOUS_OFFSET
+#undef REST_OFFSET
+
+#undef MAX_DISTANCE
+#undef BOX_INDEX
+#undef SPHERE_INDEX
+#undef CAPSULE_INDEX
diff --git a/Engine/Plugins/FX/Niagara/Source/Niagara/Private/DataInterface/NiagaraDataInterfacePhysicsAsset.cpp b/Engine/Plugins/FX/Niagara/Source/Niagara/Private/DataInterface/NiagaraDataInterfacePhysicsAsset.cpp
--- a/Engine/Plugins/FX/Niagara/Source/Niagara/Private/DataInterface/NiagaraDataInterfacePhysicsAsset.cpp
+++ b/Engine/Plugins/FX/Niagara/Source/Niagara/Private/DataInterface/NiagaraDataInterfacePhysicsAsset.cpp
@@ -10,12 +10,28 @@
 #include "NiagaraComponent.h"
 #include "NiagaraRenderer.h"
 #include "NiagaraSimStageData.h"
+#include "NiagaraShaderParametersBuilder.h"
 #include "NiagaraSystemInstance.h"
 #include "ShaderParameterUtils.h"
 
 #define LOCTEXT_NAMESPACE "NiagaraDataInterfacePhysicsAsset"
 DEFINE_LOG_CATEGORY_STATIC(LogPhysicsAsset, Log, All);
 
+namespace NDIPhysicsAssetLocal
+{
+
+BEGIN_SHADER_PARAMETER_STRUCT(FShaderParameters,)
+	SHADER_PARAMETER(FUintVector4,			ElementOffsets)
+	SHADER_PARAMETER_SRV(Buffer<float4>,	WorldTransformBuffer)
+	SHADER_PARAMETER_SRV(Buffer<float4>,	InverseTransformBuffer)
+	SHADER_PARAMETER_SRV(Buffer<float4>,	ElementExtentBuffer)
+	SHADER_PARAMETER_SRV(Buffer<uint>,		PhysicsTypeBuffer)
+	SHADER_PARAMETER(FVector3f,				BoxOrigin)
+	SHADER_PARAMETER(FVector3f,				BoxExtent)
+END_SHADER_PARAMETER_STRUCT()
+
+static const TCHAR* TemplateShaderFile = TEXT("/Plugin/FX/Niagara/Private/NiagaraDataInterfacePhysicsAsset.ush");
+
 //------------------------------------------------------------------------------------------------------------
 
 static const FName GetNumBoxesName(TEXT("GetNumBoxes"));
@@ -35,48 +51,8 @@ static const FName GetProjectionPointName(TEXT("GetProjectionPoint"));
 
 //------------------------------------------------------------------------------------------------------------
 
-const FString UNiagaraDataInterfacePhysicsAsset::ElementOffsetsName(TEXT("ElementOffsets_"));
-
-const FString UNiagaraDataInterfacePhysicsAsset::WorldTransformBufferName(TEXT("WorldTransformBuffer_"));
-const FString UNiagaraDataInterfacePhysicsAsset::InverseTransformBufferName(TEXT("InverseTransformBuffer_"));
-const FString UNiagaraDataInterfacePhysicsAsset::ElementExtentBufferName(TEXT("ElementExtentBuffer_"));
-const FString UNiagaraDataInterfacePhysicsAsset::PhysicsTypeBufferName(TEXT("PhysicsTypeBuffer_"));
-
-const FString UNiagaraDataInterfacePhysicsAsset::BoxOriginName(TEXT("BoxOrigin_"));
-const FString UNiagaraDataInterfacePhysicsAsset::BoxExtentName(TEXT("BoxExtent_"));
-
-//------------------------------------------------------------------------------------------------------------
-
-struct FNDIPhysicsAssetParametersName
-{
-	FNDIPhysicsAssetParametersName(const FString& Suffix)
-	{
-		ElementOffsetsName = UNiagaraDataInterfacePhysicsAsset::ElementOffsetsName + Suffix;
-
-		WorldTransformBufferName = UNiagaraDataInterfacePhysicsAsset::WorldTransformBufferName + Suffix;
-		InverseTransformBufferName = UNiagaraDataInterfacePhysicsAsset::InverseTransformBufferName + Suffix;
-		ElementExtentBufferName = UNiagaraDataInterfacePhysicsAsset::ElementExtentBufferName + Suffix;
-		PhysicsTypeBufferName = UNiagaraDataInterfacePhysicsAsset::PhysicsTypeBufferName + Suffix;
-
-		BoxOriginName = UNiagaraDataInterfacePhysicsAsset::BoxOriginName + Suffix;
-		BoxExtentName = UNiagaraDataInterfacePhysicsAsset::BoxExtentName + Suffix;
-	}
-
-	FString ElementOffsetsName;
-
-	FString WorldTransformBufferName;
-	FString InverseTransformBufferName;
-	FString ElementExtentBufferName;
-	FString PhysicsTypeBufferName;
-
-	FString BoxOriginName;
-	FString BoxExtentName;
-};
-
-//------------------------------------------------------------------------------------------------------------
-
 template<typename BufferType, EPixelFormat PixelFormat, uint32 ElementCount, uint32 BufferCount = 1>
-void CreateInternalBuffer(FRWBuffer& OutputBuffer)
+void CreateInternalBuffer(FReadBuffer& OutputBuffer)
 {
 	if (ElementCount > 0)
 	{
@@ -85,7 +61,7 @@ void CreateInternalBuffer(FRWBuffer& OutputBuffer)
 }
 
 template<typename BufferType, EPixelFormat PixelFormat, uint32 ElementCount, uint32 BufferCount = 1>
-void UpdateInternalBuffer(const TStaticArray<BufferType,ElementCount*BufferCount>& InputData, FRWBuffer& OutputBuffer)
+void UpdateInternalBuffer(const TStaticArray<BufferType,ElementCount*BufferCount>& InputData, FReadBuffer& OutputBuffer)
 {
 	if (ElementCount > 0 && OutputBuffer.Buffer.IsValid())
 	{
@@ -410,24 +386,7 @@ void UpdateInternalArrays(const TArray<TWeakObjectPtr<UPhysicsAsset>>& PhysicsAs
 	}
 }
 
-//------------------------------------------------------------------------------------------------------------
-
-void FNDIPhysicsAssetBuffer::InitRHI()
-{
-	CreateInternalBuffer<FVector4f, EPixelFormat::PF_A32B32G32R32F, PHYSICS_ASSET_MAX_TRANSFORMS, 3>(WorldTransformBuffer);
-	CreateInternalBuffer<FVector4f, EPixelFormat::PF_A32B32G32R32F, PHYSICS_ASSET_MAX_TRANSFORMS, 3>(InverseTransformBuffer);
-
-	CreateInternalBuffer<FVector4f, EPixelFormat::PF_A32B32G32R32F, PHYSICS_ASSET_MAX_PRIMITIVES>(ElementExtentBuffer);
-	CreateInternalBuffer<uint32, EPixelFormat::PF_R32_UINT, PHYSICS_ASSET_MAX_PRIMITIVES>(PhysicsTypeBuffer);
-}
-
-void FNDIPhysicsAssetBuffer::ReleaseRHI()
-{
-	WorldTransformBuffer.Release();
-	InverseTransformBuffer.Release();
-	ElementExtentBuffer.Release();
-	PhysicsTypeBuffer.Release();
-}
+} //namespace NDIPhysicsAssetLocal
 
 //------------------------------------------------------------------------------------------------------------
 
@@ -450,6 +409,27 @@ ETickingGroup ComputeTickingGroup(const TArray<TWeakObjectPtr<class USkeletalMes
 	return TickingGroup;
 }
 
+//------------------------------------------------------------------------------------------------------------
+
+void FNDIPhysicsAssetBuffer::InitRHI()
+{
+	using namespace NDIPhysicsAssetLocal;
+
+	CreateInternalBuffer<FVector4f, EPixelFormat::PF_A32B32G32R32F, PHYSICS_ASSET_MAX_TRANSFORMS, 3>(WorldTransformBuffer);
+	CreateInternalBuffer<FVector4f, EPixelFormat::PF_A32B32G32R32F, PHYSICS_ASSET_MAX_TRANSFORMS, 3>(InverseTransformBuffer);
+
+	CreateInternalBuffer<FVector4f, EPixelFormat::PF_A32B32G32R32F, PHYSICS_ASSET_MAX_PRIMITIVES>(ElementExtentBuffer);
+	CreateInternalBuffer<uint32, EPixelFormat::PF_R32_UINT, PHYSICS_ASSET_MAX_PRIMITIVES>(PhysicsTypeBuffer);
+}
+
+void FNDIPhysicsAssetBuffer::ReleaseRHI()
+{
+	WorldTransformBuffer.Release();
+	InverseTransformBuffer.Release();
+	ElementExtentBuffer.Release();
+	PhysicsTypeBuffer.Release();
+}
+
 void FNDIPhysicsAssetData::Release()
 {
 	if (AssetBuffer)
@@ -466,6 +446,8 @@ void FNDIPhysicsAssetData::Release()
 
 void FNDIPhysicsAssetData::Init(UNiagaraDataInterfacePhysicsAsset* Interface, FNiagaraSystemInstance* SystemInstance)
 {
+	using namespace NDIPhysicsAssetLocal;
+
 	AssetBuffer = nullptr;
 
 	if (Interface != nullptr && SystemInstance != nullptr)
@@ -516,102 +498,11 @@ void FNDIPhysicsAssetData::Update(UNiagaraDataInterfacePhysicsAsset* Interface,
 		if (0 < Interface->PhysicsAssets.Num() && Interface->PhysicsAssets[0].IsValid() && Interface->PhysicsAssets[0].Get() != nullptr &&
 			Interface->PhysicsAssets.Num() == Interface->SourceComponents.Num())
 		{
-			UpdateInternalArrays(Interface->PhysicsAssets, Interface->SourceComponents, &AssetArrays, SystemInstance->GetWorldTransform(), BoneTransform);
+			NDIPhysicsAssetLocal::UpdateInternalArrays(Interface->PhysicsAssets, Interface->SourceComponents, &AssetArrays, SystemInstance->GetWorldTransform(), BoneTransform);
 		}
 	}
 }
 
-//------------------------------------------------------------------------------------------------------------
-
-struct FNDIPhysicsAssetParametersCS : public FNiagaraDataInterfaceParametersCS
-{
-	DECLARE_TYPE_LAYOUT(FNDIPhysicsAssetParametersCS, NonVirtual);
-public:
-	void Bind(const FNiagaraDataInterfaceGPUParamInfo& ParameterInfo, const class FShaderParameterMap& ParameterMap)
-	{
-		FNDIPhysicsAssetParametersName ParamNames(*ParameterInfo.DataInterfaceHLSLSymbol);
-
-		ElementOffsets.Bind(ParameterMap, *ParamNames.ElementOffsetsName);
-
-		WorldTransformBuffer.Bind(ParameterMap, *ParamNames.WorldTransformBufferName);
-		InverseTransformBuffer.Bind(ParameterMap, *ParamNames.InverseTransformBufferName);
-		ElementExtentBuffer.Bind(ParameterMap, *ParamNames.ElementExtentBufferName);
-		PhysicsTypeBuffer.Bind(ParameterMap, *ParamNames.PhysicsTypeBufferName);
-
-		BoxOrigin.Bind(ParameterMap, *ParamNames.BoxOriginName);
-		BoxExtent.Bind(ParameterMap, *ParamNames.BoxExtentName);
-
-		
-	}
-
-	void Set(FRHICommandList& RHICmdList, const FNiagaraDataInterfaceSetArgs& Context) const
-	{
-		check(IsInRenderingThread());
-
-		FRHIComputeShader* ComputeShaderRHI = RHICmdList.GetBoundComputeShader();
-
-		FNDIPhysicsAssetProxy* InterfaceProxy =
-			static_cast<FNDIPhysicsAssetProxy*>(Context.DataInterface);
-		FNDIPhysicsAssetData* ProxyData =
-			InterfaceProxy->SystemInstancesToProxyData.Find(Context.SystemInstanceID);
-
-		if (ProxyData != nullptr && ProxyData->AssetBuffer && ProxyData->AssetBuffer->IsInitialized())
-		{
-			FNDIPhysicsAssetBuffer* AssetBuffer = ProxyData->AssetBuffer;
-
-			FRHITransitionInfo Transitions[] = {
-				FRHITransitionInfo(AssetBuffer->WorldTransformBuffer.UAV, ERHIAccess::Unknown, ERHIAccess::SRVCompute),
-				FRHITransitionInfo(AssetBuffer->InverseTransformBuffer.UAV, ERHIAccess::Unknown, ERHIAccess::SRVCompute),
-				FRHITransitionInfo(AssetBuffer->ElementExtentBuffer.UAV, ERHIAccess::Unknown, ERHIAccess::SRVCompute),
-				FRHITransitionInfo(AssetBuffer->PhysicsTypeBuffer.UAV, ERHIAccess::Unknown, ERHIAccess::SRVCompute)
-			};
-			RHICmdList.Transition(MakeArrayView(Transitions, UE_ARRAY_COUNT(Transitions)));
-
-			SetSRVParameter(RHICmdList, ComputeShaderRHI, WorldTransformBuffer, AssetBuffer->WorldTransformBuffer.SRV);
-			SetSRVParameter(RHICmdList, ComputeShaderRHI, InverseTransformBuffer, AssetBuffer->InverseTransformBuffer.SRV);
-			SetSRVParameter(RHICmdList, ComputeShaderRHI, ElementExtentBuffer, AssetBuffer->ElementExtentBuffer.SRV);
-			SetSRVParameter(RHICmdList, ComputeShaderRHI, PhysicsTypeBuffer, AssetBuffer->PhysicsTypeBuffer.SRV);
-
-			SetShaderValue(RHICmdList, ComputeShaderRHI, ElementOffsets, ProxyData->AssetArrays.ElementOffsets);
-			SetShaderValue(RHICmdList, ComputeShaderRHI, BoxOrigin, (FVector3f)ProxyData->BoxOrigin);
-			SetShaderValue(RHICmdList, ComputeShaderRHI, BoxExtent, (FVector3f)ProxyData->BoxExtent);
-		}
-		else
-		{
-			SetSRVParameter(RHICmdList, ComputeShaderRHI, WorldTransformBuffer, FNiagaraRenderer::GetDummyFloatBuffer());
-			SetSRVParameter(RHICmdList, ComputeShaderRHI, InverseTransformBuffer, FNiagaraRenderer::GetDummyFloatBuffer());
-			SetSRVParameter(RHICmdList, ComputeShaderRHI, ElementExtentBuffer, FNiagaraRenderer::GetDummyFloatBuffer());
-			SetSRVParameter(RHICmdList, ComputeShaderRHI, PhysicsTypeBuffer, FNiagaraRenderer::GetDummyIntBuffer());
-
-			static const FElementOffset DummyOffsets(0, 0, 0, 0);
-			SetShaderValue(RHICmdList, ComputeShaderRHI, ElementOffsets, DummyOffsets);
-			SetShaderValue(RHICmdList, ComputeShaderRHI, BoxOrigin, FVector3f::ZeroVector);
-			SetShaderValue(RHICmdList, ComputeShaderRHI, BoxExtent, FVector3f::ZeroVector);
-		}
-	}
-
-	void Unset(FRHICommandList& RHICmdList, const FNiagaraDataInterfaceSetArgs& Context) const
-	{
-	}
-
-private:
-
-	LAYOUT_FIELD(FShaderParameter, ElementOffsets);
-
-	LAYOUT_FIELD(FShaderResourceParameter, WorldTransformBuffer);
-	LAYOUT_FIELD(FShaderResourceParameter, InverseTransformBuffer);
-	LAYOUT_FIELD(FShaderResourceParameter, ElementExtentBuffer);
-	LAYOUT_FIELD(FShaderResourceParameter, PhysicsTypeBuffer);
-
-	LAYOUT_FIELD(FShaderParameter, BoxOrigin);
-	LAYOUT_FIELD(FShaderParameter, BoxExtent);
-};
-
-IMPLEMENT_TYPE_LAYOUT(FNDIPhysicsAssetParametersCS);
-
-IMPLEMENT_NIAGARA_DI_PARAMETER(UNiagaraDataInterfacePhysicsAsset, FNDIPhysicsAssetParametersCS);
-
-
 //------------------------------------------------------------------------------------------------------------
 
 void FNDIPhysicsAssetProxy::ConsumePerInstanceDataFromGameThread(void* PerInstanceData, const FNiagaraSystemInstanceID& Instance)
@@ -649,23 +540,16 @@ void FNDIPhysicsAssetProxy::DestroyPerInstanceData(const FNiagaraSystemInstanceI
 	SystemInstancesToProxyData.Remove(SystemInstance);
 }
 
-void FNDIPhysicsAssetProxy::PreStage(FRHICommandList& RHICmdList, const FNiagaraDataInterfaceStageArgs& Context)
+void FNDIPhysicsAssetProxy::PreStage(const FNDIGpuComputePostStageContext& Context)
 {
-	FNDIPhysicsAssetData* ProxyData =
-		SystemInstancesToProxyData.Find(Context.SystemInstanceID);
+	using namespace NDIPhysicsAssetLocal;
+
+	FNDIPhysicsAssetData* ProxyData = SystemInstancesToProxyData.Find(Context.GetSystemInstanceID());
 
 	if (ProxyData != nullptr && ProxyData->AssetBuffer)
 	{
-		if (Context.SimStageData->bFirstStage)
+		if (Context.GetSimStageData().bFirstStage)
 		{
-			FRHITransitionInfo Transitions[] = {
-				FRHITransitionInfo(ProxyData->AssetBuffer->WorldTransformBuffer.UAV, ERHIAccess::Unknown, ERHIAccess::UAVCompute),
-				FRHITransitionInfo(ProxyData->AssetBuffer->InverseTransformBuffer.UAV, ERHIAccess::Unknown, ERHIAccess::UAVCompute),
-				FRHITransitionInfo(ProxyData->AssetBuffer->PhysicsTypeBuffer.UAV, ERHIAccess::Unknown, ERHIAccess::UAVCompute),
-				FRHITransitionInfo(ProxyData->AssetBuffer->ElementExtentBuffer.UAV, ERHIAccess::Unknown, ERHIAccess::UAVCompute)
-			};
-			RHICmdList.Transition(MakeArrayView(Transitions, UE_ARRAY_COUNT(Transitions)));
-
 			UpdateInternalBuffer<FVector4f, EPixelFormat::PF_A32B32G32R32F, PHYSICS_ASSET_MAX_TRANSFORMS, 3>(ProxyData->AssetArrays.WorldTransform, ProxyData->AssetBuffer->WorldTransformBuffer);
 			UpdateInternalBuffer<FVector4f, EPixelFormat::PF_A32B32G32R32F, PHYSICS_ASSET_MAX_TRANSFORMS, 3>(ProxyData->AssetArrays.InverseTransform, ProxyData->AssetBuffer->InverseTransformBuffer);
 			UpdateInternalBuffer<FVector4f, EPixelFormat::PF_A32B32G32R32F, PHYSICS_ASSET_MAX_PRIMITIVES>(ProxyData->AssetArrays.ElementExtent, ProxyData->AssetBuffer->ElementExtentBuffer);
@@ -942,6 +826,7 @@ void UNiagaraDataInterfacePhysicsAsset::PostLoad()
 
 void UNiagaraDataInterfacePhysicsAsset::GetFunctions(TArray<FNiagaraFunctionSignature>& OutFunctions)
 {
+	using namespace NDIPhysicsAssetLocal;
 	{
 		FNiagaraFunctionSignature Sig;
 		Sig.Name = GetNumBoxesName;
@@ -1118,6 +1003,8 @@ DEFINE_NDI_DIRECT_FUNC_BINDER(UNiagaraDataInterfacePhysicsAsset, GetProjectionPo
 
 void UNiagaraDataInterfacePhysicsAsset::GetVMExternalFunction(const FVMExternalFunctionBindingInfo& BindingInfo, void* InstanceData, FVMExternalFunction& OutFunc)
 {
+	using namespace NDIPhysicsAssetLocal;
+
 	if (BindingInfo.Name == GetNumBoxesName)
 	{
 		check(BindingInfo.GetNumInputs() == 1 && BindingInfo.GetNumOutputs() == 1);
@@ -1213,166 +1100,99 @@ void UNiagaraDataInterfacePhysicsAsset::GetProjectionPoint(FVectorVMExternalFunc
 #if WITH_EDITORONLY_DATA
 bool UNiagaraDataInterfacePhysicsAsset::GetFunctionHLSL(const FNiagaraDataInterfaceGPUParamInfo& ParamInfo, const FNiagaraDataInterfaceGeneratedFunction& FunctionInfo, int FunctionInstanceIndex, FString& OutHLSL)
 {
-	FNDIPhysicsAssetParametersName ParamNames(ParamInfo.DataInterfaceHLSLSymbol);
-
-	TMap<FString, FStringFormatArg> ArgsSample = {
-		{TEXT("InstanceFunctionName"), FunctionInfo.InstanceName},
-		{TEXT("ElementOffsetsName"), ParamNames.ElementOffsetsName},
-		{TEXT("WorldTransformBufferName"), ParamNames.WorldTransformBufferName},
-		{TEXT("InverseTransformBufferName"), ParamNames.InverseTransformBufferName},
-		{TEXT("ElementExtentBufferName"), ParamNames.ElementExtentBufferName},
-		{TEXT("PhysicsAssetContextName"), TEXT("DIPHYSICSASSET_MAKE_CONTEXT(") + ParamInfo.DataInterfaceHLSLSymbol + TEXT(")")},
-	};
+	using namespace NDIPhysicsAssetLocal;
 
-	if (FunctionInfo.DefinitionName == GetNumBoxesName)
-	{
-		static const TCHAR* FormatSample = TEXT(R"(
-		void {InstanceFunctionName}(out int OutNumBoxes)
-		{
-			{PhysicsAssetContextName}
-			OutNumBoxes = DIPhysicsAsset_GetNumBoxes(DIContext);
-		}
-		)");
-		OutHLSL += FString::Format(FormatSample, ArgsSample);
-		return true;
-	}
-	else if (FunctionInfo.DefinitionName == GetNumCapsulesName)
-	{
-		static const TCHAR* FormatSample = TEXT(R"(
-		void {InstanceFunctionName}(out int OutNumCapsules)
-		{
-			{PhysicsAssetContextName}
-			OutNumCapsules = DIPhysicsAsset_GetNumCapsules(DIContext);
-		}
-		)");
-		OutHLSL += FString::Format(FormatSample, ArgsSample);
-		return true;
-	}
-	else if (FunctionInfo.DefinitionName == GetNumSpheresName)
-	{
-		static const TCHAR* FormatSample = TEXT(R"(
-		void {InstanceFunctionName}(out int OutNumSpheres)
-		{
-			{PhysicsAssetContextName}
-			OutNumSpheres = DIPhysicsAsset_GetNumSpheres(DIContext);
-		}
-		)");
-		OutHLSL += FString::Format(FormatSample, ArgsSample);
-		return true;
-	}
-	else if (FunctionInfo.DefinitionName == GetClosestPointName)
-	{
-		static const TCHAR* FormatSample = TEXT(R"(
-		void {InstanceFunctionName}(in float3 NodePosition, in float DeltaTime, in float TimeFraction, out float3 OutClosestPosition, 
-							out float3 OutClosestNormal, out float3 OutClosestVelocity)
-		{
-			{PhysicsAssetContextName} DIPhysicsAsset_GetClosestPoint(DIContext,NodePosition,DeltaTime,TimeFraction,
-				OutClosestPosition,OutClosestNormal,OutClosestVelocity);
-		}
-		)");
-		OutHLSL += FString::Format(FormatSample, ArgsSample);
-		return true;
-	}
-	else if (FunctionInfo.DefinitionName == GetClosestElementName)
-	{
-		static const TCHAR* FormatSample = TEXT(R"(
-		void {InstanceFunctionName}(in float3 NodePosition, in float TimeFraction, out int OutClosestElement)
-		{
-			{PhysicsAssetContextName} DIPhysicsAsset_GetClosestElement(DIContext,NodePosition,TimeFraction,
-				OutClosestElement);
-		}
-		)");
-		OutHLSL += FString::Format(FormatSample, ArgsSample);
-		return true;
-	}
-	else if (FunctionInfo.DefinitionName == GetElementPointName)
-	{
-		static const TCHAR* FormatSample = TEXT(R"(
-		void {InstanceFunctionName}(in float3 NodePosition, in float DeltaTime, in float TimeFraction, in int ElementIndex, out float3 OutClosestPosition, 
-							out float3 OutClosestNormal, out float3 OutClosestVelocity)
-		{
-			{PhysicsAssetContextName} DIPhysicsAsset_GetElementPoint(DIContext,NodePosition,DeltaTime,TimeFraction,ElementIndex,
-				OutClosestPosition,OutClosestNormal,OutClosestVelocity);
-		}
-		)");
-		OutHLSL += FString::Format(FormatSample, ArgsSample);
-		return true;
-	}
-	else if (FunctionInfo.DefinitionName == GetElementDistanceName)
-	{
-		static const TCHAR* FormatSample = TEXT(R"(
-		void {InstanceFunctionName}(in float3 NodePosition, in float TimeFraction, in int ElementIndex, out float OutClosestDistance)
-		{
-			{PhysicsAssetContextName} DIPhysicsAsset_GetElementDistance(DIContext,NodePosition,TimeFraction,ElementIndex,
-				OutClosestDistance);
-		}
-		)");
-		OutHLSL += FString::Format(FormatSample, ArgsSample);
-		return true;
-	}
-	else if (FunctionInfo.DefinitionName == GetClosestDistanceName)
-	{
-		static const TCHAR* FormatSample = TEXT(R"(
-		void {InstanceFunctionName}(in float3 NodePosition, in float TimeFraction, out float OutClosestDistance)
-		{
-			{PhysicsAssetContextName} DIPhysicsAsset_GetClosestDistance(DIContext,NodePosition,TimeFraction,OutClosestDistance);
-		}
-		)");
-		OutHLSL += FString::Format(FormatSample, ArgsSample);
-		return true;
-	}
-	else if (FunctionInfo.DefinitionName == GetRestDistanceName)
-	{
-		static const TCHAR* FormatSample = TEXT(R"(
-		void {InstanceFunctionName}(in float3 RestPosition, out float OutRestDistance)
-		{
-			{PhysicsAssetContextName} DIPhysicsAsset_GetRestDistance(DIContext,RestPosition,OutRestDistance);
-		}
-		)");
-		OutHLSL += FString::Format(FormatSample, ArgsSample);
-		return true;
-	}
-	else if (FunctionInfo.DefinitionName == GetTexturePointName)
+	static const TSet<FName> ValidGpuFunctions =
 	{
-		static const TCHAR* FormatSample = TEXT(R"(
-		void {InstanceFunctionName}(in float3 NodePosition, out int OutElementIndex, out float3 OutTexturePosition)
-		{
-			{PhysicsAssetContextName} DIPhysicsAsset_GetTexturePoint(DIContext,NodePosition,OutElementIndex,OutTexturePosition);
-		}
-		)");
-		OutHLSL += FString::Format(FormatSample, ArgsSample);
-		return true;
-	}
-	else if (FunctionInfo.DefinitionName == GetProjectionPointName)
+		GetNumBoxesName,
+		GetNumCapsulesName,
+		GetNumSpheresName,
+		GetClosestPointName,
+		GetClosestElementName,
+		GetElementPointName,
+		GetElementDistanceName,
+		GetClosestDistanceName,
+		GetRestDistanceName,
+		GetTexturePointName,
+		GetProjectionPointName,
+	};
+
+	return ValidGpuFunctions.Contains(FunctionInfo.DefinitionName);
+}
+
+bool UNiagaraDataInterfacePhysicsAsset::AppendCompileHash(FNiagaraCompileHashVisitor* InVisitor) const
+{
+	if (!Super::AppendCompileHash(InVisitor))
 	{
-		static const TCHAR* FormatSample = TEXT(R"(
-		void {InstanceFunctionName}(in float3 NodePosition, in float DeltaTime, in int ElementIndex, in float TextureValue, in float3 TextureGradient, out float3 OutClosestPosition, 
-							out float3 OutClosestNormal, out float3 OutClosestVelocity, out float OutClosestDistance)
-		{
-			{PhysicsAssetContextName} DIPhysicsAsset_GetProjectionPoint(DIContext,NodePosition,DeltaTime,ElementIndex,TextureValue,TextureGradient,
-				OutClosestPosition,OutClosestNormal,OutClosestVelocity,OutClosestDistance);
-		}
-		)");
-		OutHLSL += FString::Format(FormatSample, ArgsSample);
-		return true;
+		return false;
 	}
 
-	OutHLSL += TEXT("\n");
-	return false;
+	InVisitor->UpdateString(TEXT("NiagaraDataInterfaceRigidMeshCollisionQueryHLSLSource"), GetShaderFileHash(NDIPhysicsAssetLocal::TemplateShaderFile, EShaderPlatform::SP_PCD3D_SM5).ToString());
+	InVisitor->UpdateShaderParameters<NDIPhysicsAssetLocal::FShaderParameters>();
+
+	return true;
 }
 
 void UNiagaraDataInterfacePhysicsAsset::GetCommonHLSL(FString& OutHLSL)
 {
 	OutHLSL += TEXT("#include \"/Plugin/FX/Niagara/Private/NiagaraQuaternionUtils.ush\"\n");
-	OutHLSL += TEXT("#include \"/Plugin/FX/Niagara/Private/NiagaraDataInterfacePhysicsAsset.ush\"\n");	
 }
 
 void UNiagaraDataInterfacePhysicsAsset::GetParameterDefinitionHLSL(const FNiagaraDataInterfaceGPUParamInfo& ParamInfo, FString& OutHLSL)
 {
-	OutHLSL += TEXT("DIPHYSICSASSET_DECLARE_CONSTANTS(") + ParamInfo.DataInterfaceHLSLSymbol + TEXT(")\n");
+	Super::GetParameterDefinitionHLSL(ParamInfo, OutHLSL);
+
+	TMap<FString, FStringFormatArg> TemplateArgs =
+	{
+		{TEXT("ParameterName"),	ParamInfo.DataInterfaceHLSLSymbol},
+	};
+
+	FString TemplateFile;
+	LoadShaderSourceFile(NDIPhysicsAssetLocal::TemplateShaderFile, EShaderPlatform::SP_PCD3D_SM5, &TemplateFile, nullptr);
+	OutHLSL += FString::Format(*TemplateFile, TemplateArgs);
 }
 #endif
 
+void UNiagaraDataInterfacePhysicsAsset::BuildShaderParameters(FNiagaraShaderParametersBuilder& ShaderParametersBuilder) const
+{
+	ShaderParametersBuilder.AddNestedStruct<NDIPhysicsAssetLocal::FShaderParameters>();
+}
+
+void UNiagaraDataInterfacePhysicsAsset::SetShaderParameters(const FNiagaraDataInterfaceSetShaderParametersContext& Context) const
+{
+	FNDIPhysicsAssetProxy& DIProxy = Context.GetProxy<FNDIPhysicsAssetProxy>();
+	FNDIPhysicsAssetData* ProxyData = DIProxy.SystemInstancesToProxyData.Find(Context.GetSystemInstanceID());
+	NDIPhysicsAssetLocal::FShaderParameters* ShaderParameters = Context.GetParameterNestedStruct<NDIPhysicsAssetLocal::FShaderParameters>();
+
+	if (ProxyData != nullptr && ProxyData->AssetBuffer && ProxyData->AssetBuffer->IsInitialized())
+	{
+		FNDIPhysicsAssetBuffer* AssetBuffer = ProxyData->AssetBuffer;
+
+		ShaderParameters->WorldTransformBuffer = FNiagaraRenderer::GetSrvOrDefaultFloat4(AssetBuffer->WorldTransformBuffer.SRV);
+		ShaderParameters->InverseTransformBuffer = FNiagaraRenderer::GetSrvOrDefaultFloat4(AssetBuffer->InverseTransformBuffer.SRV);
+		ShaderParameters->ElementExtentBuffer = FNiagaraRenderer::GetSrvOrDefaultFloat4(AssetBuffer->ElementExtentBuffer.SRV);
+		ShaderParameters->PhysicsTypeBuffer = FNiagaraRenderer::GetSrvOrDefaultUInt(AssetBuffer->PhysicsTypeBuffer.SRV);
+
+		ShaderParameters->ElementOffsets.X = ProxyData->AssetArrays.ElementOffsets.BoxOffset;
+		ShaderParameters->ElementOffsets.Y = ProxyData->AssetArrays.ElementOffsets.SphereOffset;
+		ShaderParameters->ElementOffsets.Z = ProxyData->AssetArrays.ElementOffsets.CapsuleOffset;
+		ShaderParameters->ElementOffsets.W = ProxyData->AssetArrays.ElementOffsets.NumElements;
+		ShaderParameters->BoxOrigin = FVector3f(ProxyData->BoxOrigin);
+		ShaderParameters->BoxExtent = FVector3f(ProxyData->BoxExtent);
+	}
+	else
+	{
+		ShaderParameters->WorldTransformBuffer = FNiagaraRenderer::GetDummyFloat4Buffer();
+		ShaderParameters->InverseTransformBuffer = FNiagaraRenderer::GetDummyFloat4Buffer();
+		ShaderParameters->ElementExtentBuffer = FNiagaraRenderer::GetDummyFloat4Buffer();
+		ShaderParameters->PhysicsTypeBuffer = FNiagaraRenderer::GetDummyUIntBuffer();
+
+		ShaderParameters->ElementOffsets = FUintVector4(0, 0, 0, 0);
+		ShaderParameters->BoxOrigin = FVector3f::ZeroVector;
+		ShaderParameters->BoxExtent = FVector3f::ZeroVector;
+	}
+}
+
 void UNiagaraDataInterfacePhysicsAsset::ProvidePerInstanceDataForRenderThread(void* DataForRenderThread, void* PerInstanceData, const FNiagaraSystemInstanceID& SystemInstance)
 {
 	FNDIPhysicsAssetData* GameThreadData = static_cast<FNDIPhysicsAssetData*>(PerInstanceData);
@@ -1393,4 +1213,4 @@ UNiagaraPhysicsAssetDICollectorInterface::UNiagaraPhysicsAssetDICollectorInterfa
 	: Super(ObjectInitializer)
 {}
 
-#undef LOCTEXT_NAMESPACE
\ No newline at end of file
+#undef LOCTEXT_NAMESPACE
diff --git a/Engine/Plugins/FX/Niagara/Source/Niagara/Public/NiagaraDataInterfacePhysicsAsset.h b/Engine/Plugins/FX/Niagara/Source/Niagara/Public/NiagaraDataInterfacePhysicsAsset.h
--- a/Engine/Plugins/FX/Niagara/Source/Niagara/Public/NiagaraDataInterfacePhysicsAsset.h
+++ b/Engine/Plugins/FX/Niagara/Source/Niagara/Public/NiagaraDataInterfacePhysicsAsset.h
@@ -14,13 +14,13 @@
 #define PHYSICS_ASSET_MAX_TRANSFORMS PHYSICS_ASSET_MAX_PRIMITIVES * 3
 
 /** Element offsets in the array list */
-struct FElementOffset
+struct FNDIPhysicsAssetElementOffset
 {
-	FElementOffset(const uint32 InBoxOffset, const uint32 InSphereOffset, const uint32 InCapsuleOffset, const uint32 InNumElements) :
+	FNDIPhysicsAssetElementOffset(const uint32 InBoxOffset, const uint32 InSphereOffset, const uint32 InCapsuleOffset, const uint32 InNumElements) :
 		BoxOffset(InBoxOffset), SphereOffset(InSphereOffset), CapsuleOffset(InCapsuleOffset), NumElements(InNumElements)
 	{}
 
-	FElementOffset() :
+	FNDIPhysicsAssetElementOffset() :
 		BoxOffset(0), SphereOffset(0), CapsuleOffset(0), NumElements(0)
 	{}
 	uint32 BoxOffset;
@@ -32,7 +32,7 @@ struct FElementOffset
 /** Arrays in which the cpu datas will be str */
 struct FNDIPhysicsAssetArrays
 {
-	FElementOffset ElementOffsets;
+	FNDIPhysicsAssetElementOffset ElementOffsets;
 	TStaticArray<FVector4f, 3 * PHYSICS_ASSET_MAX_TRANSFORMS> WorldTransform;
 	TStaticArray<FVector4f, 3 * PHYSICS_ASSET_MAX_TRANSFORMS> InverseTransform;
 	TStaticArray<FVector4f, PHYSICS_ASSET_MAX_TRANSFORMS> CurrentTransform;
@@ -58,16 +58,16 @@ struct FNDIPhysicsAssetBuffer : public FRenderResource
 	virtual FString GetFriendlyName() const override { return TEXT("FNDIPhysicsAssetBuffer"); }
 
 	/** World transform buffer */
-	FRWBuffer WorldTransformBuffer;
+	FReadBuffer WorldTransformBuffer;
 
 	/** Inverse transform buffer*/
-	FRWBuffer InverseTransformBuffer;
+	FReadBuffer InverseTransformBuffer;
 
 	/** Element extent buffer */
-	FRWBuffer ElementExtentBuffer;
+	FReadBuffer ElementExtentBuffer;
 
 	/** Physics type buffer */
-	FRWBuffer PhysicsTypeBuffer;
+	FReadBuffer PhysicsTypeBuffer;
 };
 
 /** Data stored per physics asset instance*/
@@ -105,9 +105,6 @@ class NIAGARA_API UNiagaraDataInterfacePhysicsAsset : public UNiagaraDataInterfa
 	GENERATED_UCLASS_BODY()
 
 public:
-
-	DECLARE_NIAGARA_DI_PARAMETER();
-
 	/** Skeletal Mesh from which the Physics Asset will be found. */
 	UPROPERTY(EditAnywhere, Category = "Source")
 	TObjectPtr<UPhysicsAsset> DefaultSource;
@@ -150,10 +147,15 @@ public:
 
 	/** GPU simulation  functionality */
 #if WITH_EDITORONLY_DATA
+	virtual bool AppendCompileHash(FNiagaraCompileHashVisitor* InVisitor) const override;
 	virtual void GetCommonHLSL(FString& OutHLSL) override;
 	virtual void GetParameterDefinitionHLSL(const FNiagaraDataInterfaceGPUParamInfo& ParamInfo, FString& OutHLSL) override;
 	virtual bool GetFunctionHLSL(const FNiagaraDataInterfaceGPUParamInfo& ParamInfo, const FNiagaraDataInterfaceGeneratedFunction& FunctionInfo, int FunctionInstanceIndex, FString& OutHLSL) override;
 #endif
+	virtual bool UseLegacyShaderBindings() const override { return false; }
+	virtual void BuildShaderParameters(FNiagaraShaderParametersBuilder& ShaderParametersBuilder) const override;
+	virtual void SetShaderParameters(const FNiagaraDataInterfaceSetShaderParametersContext& Context) const override;
+
 	virtual void ProvidePerInstanceDataForRenderThread(void* DataForRenderThread, void* PerInstanceData, const FNiagaraSystemInstanceID& SystemInstance) override;
 
 	/** Extract the source component */
@@ -189,27 +191,6 @@ public:
 	/** Get the projection point */
 	void GetProjectionPoint(FVectorVMExternalFunctionContext& Context);
 
-	/** Name of element offsets */
-	static const FString ElementOffsetsName;
-
-	/** Name of the world transform buffer */
-	static const FString WorldTransformBufferName;
-
-	/** Name of the inverse transform buffer */
-	static const FString InverseTransformBufferName;
-
-	/** Name of the element extent buffer */
-	static const FString ElementExtentBufferName;
-
-	/** Name of the physics type buffer */
-	static const FString PhysicsTypeBufferName;
-
-	/** Init Box Origin */
-	static const FString BoxOriginName;
-
-	/** Init Box extent */
-	static const FString BoxExtentName;
-
 protected:
 	/** Copy one niagara DI to this */
 	virtual bool CopyToInternal(UNiagaraDataInterface* Destination) const override;
@@ -231,7 +212,7 @@ struct FNDIPhysicsAssetProxy : public FNiagaraDataInterfaceProxy
 	void DestroyPerInstanceData(const FNiagaraSystemInstanceID& SystemInstance);
 
 	/** Launch all pre stage functions */
-	virtual void PreStage(FRHICommandList& RHICmdList, const FNiagaraDataInterfaceStageArgs& Context) override;
+	virtual void PreStage(const FNDIGpuComputePostStageContext& Context) override;
 
 	/** List of proxy data for each system instances*/
 	TMap<FNiagaraSystemInstanceID, FNDIPhysicsAssetData> SystemInstancesToProxyData;
