- Fix for GPU scene mesh path dangling into a potentially reallocated buffer

diff --git a/Engine/Plugins/FX/Niagara/Source/Niagara/Private/NiagaraComponent.cpp b/Engine/Plugins/FX/Niagara/Source/Niagara/Private/NiagaraComponent.cpp
--- a/Engine/Plugins/FX/Niagara/Source/Niagara/Private/NiagaraComponent.cpp
+++ b/Engine/Plugins/FX/Niagara/Source/Niagara/Private/NiagaraComponent.cpp
@@ -206,13 +206,13 @@ FNiagaraSceneProxy::~FNiagaraSceneProxy()
 {
 	check(IsInRenderingThread());
 
-	for ( auto& CustomUB : CustomUniformBuffers)
+	if (RenderData)
 	{
-		CustomUB.Value.ReleaseResource();
+		delete RenderData;
+		RenderData = nullptr;
 	}
-	CustomUniformBuffers.Empty();
 
-	delete RenderData;
+	ReleaseUniformBuffers(true);
 }
 
 void FNiagaraSceneProxy::ReleaseRenderThreadResources()
@@ -222,11 +222,7 @@ void FNiagaraSceneProxy::ReleaseRenderThreadResources()
 		RenderData->ReleaseRenderThreadResources();
 	}
 
-	for (auto& CustomUB : CustomUniformBuffers)
-	{
-		CustomUB.Value.ReleaseResource();
-	}
-	CustomUniformBuffers.Empty();
+	ReleaseUniformBuffers(true);
 }
 
 // FPrimitiveSceneProxy interface.
@@ -245,14 +241,7 @@ void FNiagaraSceneProxy::OnTransformChanged()
 
 	LocalToWorldInverse = GetLocalToWorld().Inverse();
 
-	for (auto& CustomUB : CustomUniformBuffers)
-	{
-		CustomUB.Value.ReleaseResource();
-	}
-
-	const int32 ExpectedRendererCount = RenderData ? RenderData->GetNumRenderers() : 0;
-	const int32 CurrentUBCount = CustomUniformBuffers.Num();
-	CustomUniformBuffers.Empty(FMath::Min(ExpectedRendererCount, CurrentUBCount));
+	ReleaseUniformBuffers(false);
 }
 
 FPrimitiveViewRelevance FNiagaraSceneProxy::GetViewRelevance(const FSceneView* View) const
@@ -279,7 +268,24 @@ FPrimitiveViewRelevance FNiagaraSceneProxy::GetViewRelevance(const FSceneView* V
 	return Relevance;
 }
 
-TUniformBuffer<FPrimitiveUniformShaderParameters>& FNiagaraSceneProxy::GetCustomUniformBufferResource(bool bHasVelocity, const FBox& InstanceBounds) const
+void FNiagaraSceneProxy::ReleaseUniformBuffers(bool bEmpty)
+{
+	for (auto& CustomUB : CustomUniformBuffers)
+	{
+		if (CustomUB.Value != nullptr)
+		{
+			CustomUB.Value->ReleaseResource();
+			delete CustomUB.Value;
+			CustomUB.Value = nullptr;
+		}
+	}
+
+	const int32 ExpectedRendererCount = RenderData ? RenderData->GetNumRenderers() : 0;
+	const int32 CurrentUBCount = CustomUniformBuffers.Num();
+	CustomUniformBuffers.Empty(bEmpty ? 0 : FMath::Min(ExpectedRendererCount, CurrentUBCount));
+}
+
+TUniformBuffer<FPrimitiveUniformShaderParameters>* FNiagaraSceneProxy::GetCustomUniformBufferResource(bool bHasVelocity, const FBox& InstanceBounds) const
 {
 	// Use a hash to determine if we can re-use any uniform buffer
 	uint64 KeyHash = HashCombine(bHasVelocity, InstanceBounds.IsValid);
@@ -309,8 +315,8 @@ TUniformBuffer<FPrimitiveUniformShaderParameters>& FNiagaraSceneProxy::GetCustom
 		KeyHash = HashCombine(KeyHash, InstanceBounds.Max.Z);
 	}
 
-	TUniformBuffer<FPrimitiveUniformShaderParameters>& CustomUB = CustomUniformBuffers.FindOrAdd(KeyHash);
-	if (!CustomUB.IsInitialized())
+	TUniformBuffer<FPrimitiveUniformShaderParameters>*& CustomUBRef = CustomUniformBuffers.FindOrAdd(KeyHash);
+	if (CustomUBRef == nullptr)
 	{
 		FPrimitiveUniformShaderParametersBuilder UBBuilder = FPrimitiveUniformShaderParametersBuilder()
 			.Defaults()
@@ -343,11 +349,12 @@ TUniformBuffer<FPrimitiveUniformShaderParameters>& FNiagaraSceneProxy::GetCustom
 			UBBuilder.PreSkinnedLocalBounds(InstanceBounds);
 		}
 
-		CustomUB.SetContents(UBBuilder.Build());
-		CustomUB.InitResource();
+		CustomUBRef = new TUniformBuffer<FPrimitiveUniformShaderParameters>();
+		CustomUBRef->SetContents(UBBuilder.Build());
+		CustomUBRef->InitResource();
 	}
 
-	return CustomUB;
+	return CustomUBRef;
 }
 
 FRHIUniformBuffer* FNiagaraSceneProxy::GetCustomUniformBuffer(bool bHasVelocity, const FBox& InstanceBounds) const
@@ -358,7 +365,7 @@ FRHIUniformBuffer* FNiagaraSceneProxy::GetCustomUniformBuffer(bool bHasVelocity,
 		return GetUniformBuffer();
 	}
 
-	return GetCustomUniformBufferResource(bHasVelocity, InstanceBounds).GetUniformBufferRHI();
+	return GetCustomUniformBufferResource(bHasVelocity, InstanceBounds)->GetUniformBufferRHI();
 }
 
 uint32 FNiagaraSceneProxy::GetMemoryFootprint() const
diff --git a/Engine/Plugins/FX/Niagara/Source/Niagara/Private/NiagaraRendererMeshes.cpp b/Engine/Plugins/FX/Niagara/Source/Niagara/Private/NiagaraRendererMeshes.cpp
--- a/Engine/Plugins/FX/Niagara/Source/Niagara/Private/NiagaraRendererMeshes.cpp
+++ b/Engine/Plugins/FX/Niagara/Source/Niagara/Private/NiagaraRendererMeshes.cpp
@@ -1061,7 +1061,7 @@ void FNiagaraRendererMeshes::CreateMeshBatchForSection(
 	FMeshBatchElement& BatchElement = MeshBatch.Elements[0];
 	if (ParticleMeshRenderData.bUseGPUScene)
 	{
-		BatchElement.PrimitiveUniformBufferResource = &SceneProxy.GetCustomUniformBufferResource(IsMotionBlurEnabled(), MeshData.LocalBounds);
+		BatchElement.PrimitiveUniformBufferResource = SceneProxy.GetCustomUniformBufferResource(IsMotionBlurEnabled(), MeshData.LocalBounds);
 	}
 	else
 	{
diff --git a/Engine/Plugins/FX/Niagara/Source/Niagara/Public/NiagaraComponent.h b/Engine/Plugins/FX/Niagara/Source/Niagara/Public/NiagaraComponent.h
--- a/Engine/Plugins/FX/Niagara/Source/Niagara/Public/NiagaraComponent.h
+++ b/Engine/Plugins/FX/Niagara/Source/Niagara/Public/NiagaraComponent.h
@@ -841,7 +841,7 @@ public:
 
 	const FVector3f& GetLWCRenderTile() const { return RenderData ? RenderData->LWCRenderTile : FVector3f::ZeroVector; }
 
-	TUniformBuffer<FPrimitiveUniformShaderParameters>& GetCustomUniformBufferResource(bool bHasVelocity, const FBox& InstanceBounds = FBox(ForceInitToZero)) const;
+	TUniformBuffer<FPrimitiveUniformShaderParameters>* GetCustomUniformBufferResource(bool bHasVelocity, const FBox& InstanceBounds = FBox(ForceInitToZero)) const;
 	FRHIUniformBuffer* GetCustomUniformBuffer(bool bHasVelocity, const FBox& InstanceBounds = FBox(ForceInitToZero)) const;
 
 	virtual FPrimitiveViewRelevance GetViewRelevance(const FSceneView* View) const override;
@@ -861,6 +861,8 @@ public:
 private:
 	void ReleaseRenderThreadResources();
 
+	void ReleaseUniformBuffers(bool bEmpty);
+
 	//~ Begin FPrimitiveSceneProxy Interface.
 	virtual void CreateRenderThreadResources() override;
 
@@ -886,7 +888,7 @@ private:
 
 private:
 	/** Custom Uniform Buffers, allows us to have renderer specific data packed inside such as pre-skinned bounds. */
-	mutable TMap<uint64, TUniformBuffer<FPrimitiveUniformShaderParameters>> CustomUniformBuffers;
+	mutable TMap<uint64, TUniformBuffer<FPrimitiveUniformShaderParameters>*> CustomUniformBuffers;
 
 	/** The data required to render a single instance of a NiagaraSystem */
 	FNiagaraSystemRenderData* RenderData = nullptr;
