- Fix for using reference to a map that can be reallocated

#rnx

diff --git a/Engine/Plugins/FX/Niagara/Source/Niagara/Classes/NiagaraDataInterfaceGrid2DCollection.h b/Engine/Plugins/FX/Niagara/Source/Niagara/Classes/NiagaraDataInterfaceGrid2DCollection.h
--- a/Engine/Plugins/FX/Niagara/Source/Niagara/Classes/NiagaraDataInterfaceGrid2DCollection.h
+++ b/Engine/Plugins/FX/Niagara/Source/Niagara/Classes/NiagaraDataInterfaceGrid2DCollection.h
@@ -56,7 +56,7 @@ struct FGrid2DCollectionRWInstanceData_RenderThread
 	FVector2D WorldBBoxSize = FVector2D::ZeroVector;
 	EPixelFormat PixelFormat = EPixelFormat::PF_R32_FLOAT;
 
-	TArray<FGrid2DBuffer, TInlineAllocator<2>> Buffers;
+	TArray<TUniquePtr<FGrid2DBuffer>, TInlineAllocator<2>> Buffers;
 	FGrid2DBuffer* CurrentData = nullptr;
 	FGrid2DBuffer* DestinationData = nullptr;
 
diff --git a/Engine/Plugins/FX/Niagara/Source/Niagara/Classes/NiagaraDataInterfaceGrid3DCollection.h b/Engine/Plugins/FX/Niagara/Source/Niagara/Classes/NiagaraDataInterfaceGrid3DCollection.h
--- a/Engine/Plugins/FX/Niagara/Source/Niagara/Classes/NiagaraDataInterfaceGrid3DCollection.h
+++ b/Engine/Plugins/FX/Niagara/Source/Niagara/Classes/NiagaraDataInterfaceGrid3DCollection.h
@@ -70,7 +70,7 @@ struct FGrid3DCollectionRWInstanceData_RenderThread
 	FIntVector4 PreviewAttribute = FIntVector4(INDEX_NONE, INDEX_NONE, INDEX_NONE, INDEX_NONE);
 #endif
 
-	TArray<FGrid3DBuffer, TInlineAllocator<2>> Buffers;
+	TArray<TUniquePtr<FGrid3DBuffer>, TInlineAllocator<2>> Buffers;
 	FGrid3DBuffer* CurrentData = nullptr;
 	FGrid3DBuffer* DestinationData = nullptr;
 
diff --git a/Engine/Plugins/FX/Niagara/Source/Niagara/Private/NiagaraDataInterfaceGrid2DCollection.cpp b/Engine/Plugins/FX/Niagara/Source/Niagara/Private/NiagaraDataInterfaceGrid2DCollection.cpp
--- a/Engine/Plugins/FX/Niagara/Source/Niagara/Private/NiagaraDataInterfaceGrid2DCollection.cpp
+++ b/Engine/Plugins/FX/Niagara/Source/Niagara/Private/NiagaraDataInterfaceGrid2DCollection.cpp
@@ -2531,19 +2531,20 @@ bool FGrid2DCollectionRWInstanceData_GameThread::UpdateTargetTexture(ENiagaraGpu
 
 void FGrid2DCollectionRWInstanceData_RenderThread::BeginSimulate(FRDGBuilder& GraphBuilder)
 {
-	for (FGrid2DBuffer& Buffer : Buffers)
+	for (TUniquePtr<FGrid2DBuffer>& Buffer : Buffers)
 	{
 		check(Buffer.IsValid());
-		if (&Buffer != CurrentData)
+		if (Buffer.Get() != CurrentData)
 		{
-			DestinationData = &Buffer;
+			DestinationData = Buffer.Get();
 			break;
 		}
 	}
 
 	if (DestinationData == nullptr)
 	{
-		DestinationData = &Buffers.AddDefaulted_GetRef();
+		DestinationData = new FGrid2DBuffer();
+		Buffers.Emplace(DestinationData);
 
 		const FRDGTextureDesc TextureDesc = FRDGTextureDesc::Create2DArray(NumCells, PixelFormat, FClearValueBinding::Black, ETextureCreateFlags::ShaderResource | ETextureCreateFlags::UAV, NumAttributes);
 		DestinationData->Initialize(GraphBuilder, TEXT("FGrid2DBuffer"), TextureDesc);
@@ -2565,9 +2566,9 @@ void FNiagaraDataInterfaceProxyGrid2DCollectionProxy::ResetData(const FNDIGpuCom
 	}
 
 	FRDGBuilder& GraphBuilder = Context.GetGraphBuilder();
-	for (FGrid2DBuffer& Buffer : ProxyData->Buffers)
+	for (TUniquePtr<FGrid2DBuffer>& Buffer : ProxyData->Buffers)
 	{
-		AddClearUAVPass(GraphBuilder, Buffer.GetOrCreateUAV(GraphBuilder), FVector4f(ForceInitToZero));
+		AddClearUAVPass(GraphBuilder, Buffer->GetOrCreateUAV(GraphBuilder), FVector4f(ForceInitToZero));
 	}
 }
 
@@ -2641,9 +2642,9 @@ void FNiagaraDataInterfaceProxyGrid2DCollectionProxy::PostSimulate(const FNDIGpu
 	// Clear out the transient resource we cached
 	if (Context.IsFinalPostSimulate())
 	{
-		for (FGrid2DBuffer& Buffer : ProxyData->Buffers)
+		for (TUniquePtr<FGrid2DBuffer>& Buffer : ProxyData->Buffers)
 		{
-			Buffer.EndGraphUsage();
+			Buffer->EndGraphUsage();
 		}
 
 		// Readers point to data not owned by themselves so can be caching resources on the 'other' proxy
@@ -2651,9 +2652,9 @@ void FNiagaraDataInterfaceProxyGrid2DCollectionProxy::PostSimulate(const FNDIGpu
 		if ( FNiagaraDataInterfaceProxyGrid2DCollectionProxy* OtherGrid3DProxy = static_cast<FNiagaraDataInterfaceProxyGrid2DCollectionProxy*>(ProxyData->OtherProxy) )
 		{
 			FGrid2DCollectionRWInstanceData_RenderThread& OtherProxyData = OtherGrid3DProxy->SystemInstancesToProxyData_RT.FindChecked(Context.GetSystemInstanceID());
-			for (FGrid2DBuffer& Buffer : OtherProxyData.Buffers)
+			for (TUniquePtr<FGrid2DBuffer>& Buffer : OtherProxyData.Buffers)
 			{
-				Buffer.EndGraphUsage();
+				Buffer->EndGraphUsage();
 			}
 		}
 	}
diff --git a/Engine/Plugins/FX/Niagara/Source/Niagara/Private/NiagaraDataInterfaceGrid3DCollection.cpp b/Engine/Plugins/FX/Niagara/Source/Niagara/Private/NiagaraDataInterfaceGrid3DCollection.cpp
--- a/Engine/Plugins/FX/Niagara/Source/Niagara/Private/NiagaraDataInterfaceGrid3DCollection.cpp
+++ b/Engine/Plugins/FX/Niagara/Source/Niagara/Private/NiagaraDataInterfaceGrid3DCollection.cpp
@@ -2945,19 +2945,20 @@ bool FGrid3DCollectionRWInstanceData_GameThread::UpdateTargetTexture(ENiagaraGpu
 
 void FGrid3DCollectionRWInstanceData_RenderThread::BeginSimulate(FRDGBuilder& GraphBuilder, bool RequiresBuffering)
 {
-	for (FGrid3DBuffer& Buffer : Buffers)
+	for (TUniquePtr<FGrid3DBuffer>& Buffer : Buffers)
 	{
 		check(Buffer.IsValid());
-		if (&Buffer != CurrentData || !RequiresBuffering)
+		if (Buffer.Get() != CurrentData || !RequiresBuffering)
 		{
-			DestinationData = &Buffer;
+			DestinationData = Buffer.Get();
 			break;
 		}
 	}
 
 	if (DestinationData == nullptr)
 	{
-		DestinationData = &Buffers.AddDefaulted_GetRef();
+		DestinationData = new FGrid3DBuffer();
+		Buffers.Emplace(DestinationData);
 
 		const FIntVector TextureSize(NumCells.X * NumTiles.X, NumCells.Y * NumTiles.Y, NumCells.Z * NumTiles.Z);
 		const FRDGTextureDesc TextureDesc = FRDGTextureDesc::Create3D(TextureSize, PixelFormat, FClearValueBinding::Black, ETextureCreateFlags::ShaderResource | ETextureCreateFlags::UAV);
@@ -2991,9 +2992,9 @@ void FNiagaraDataInterfaceProxyGrid3DCollectionProxy::ResetData(const FNDIGpuCom
 	}
 
 	FRDGBuilder& GraphBuilder = Context.GetGraphBuilder();
-	for (FGrid3DBuffer& Buffer : ProxyData->Buffers)
+	for (TUniquePtr<FGrid3DBuffer>& Buffer : ProxyData->Buffers)
 	{
-		AddClearUAVPass(GraphBuilder, Buffer.GetOrCreateUAV(GraphBuilder), FVector4f(ForceInitToZero));
+		AddClearUAVPass(GraphBuilder, Buffer->GetOrCreateUAV(GraphBuilder), FVector4f(ForceInitToZero));
 	}
 }
 
@@ -3073,9 +3074,9 @@ void FNiagaraDataInterfaceProxyGrid3DCollectionProxy::PostSimulate(const FNDIGpu
 	// Clear out the transient resource we cached
 	if (Context.IsFinalPostSimulate())
 	{
-		for (FGrid3DBuffer& Buffer : ProxyData->Buffers)
+		for (TUniquePtr<FGrid3DBuffer>& Buffer : ProxyData->Buffers)
 		{
-			Buffer.EndGraphUsage();
+			Buffer->EndGraphUsage();
 		}
 
 		// Readers point to data not owned by themselves so can be caching resources on the 'other' proxy
@@ -3083,9 +3084,9 @@ void FNiagaraDataInterfaceProxyGrid3DCollectionProxy::PostSimulate(const FNDIGpu
 		if (FNiagaraDataInterfaceProxyGrid3DCollectionProxy* OtherGrid3DProxy = static_cast<FNiagaraDataInterfaceProxyGrid3DCollectionProxy*>(ProxyData->OtherProxy))
 		{
 			FGrid3DCollectionRWInstanceData_RenderThread& OtherProxyData = OtherGrid3DProxy->SystemInstancesToProxyData_RT.FindChecked(Context.GetSystemInstanceID());
-			for (FGrid3DBuffer& Buffer : OtherProxyData.Buffers)
+			for (TUniquePtr<FGrid3DBuffer>& Buffer : OtherProxyData.Buffers)
 			{
-				Buffer.EndGraphUsage();
+				Buffer->EndGraphUsage();
 			}
 		}
 	}
