- Converted Skel Mesh DI to RDG / Shader Params

diff --git a/Engine/Plugins/FX/Niagara/Shaders/Private/NiagaraDataInterfaceSkeletalMesh.ush b/Engine/Plugins/FX/Niagara/Shaders/Private/NiagaraDataInterfaceSkeletalMesh.ush
--- a/Engine/Plugins/FX/Niagara/Shaders/Private/NiagaraDataInterfaceSkeletalMesh.ush
+++ b/Engine/Plugins/FX/Niagara/Shaders/Private/NiagaraDataInterfaceSkeletalMesh.ush
@@ -27,60 +27,6 @@
 #define DISKELMESH_ADJ_INDEX_FORMAT DISKELMESH_ADJ_INDEX_FORMAT_FULL
 #endif
 
-struct FDISkelMeshContext
-{
-	Buffer<uint>	MeshIndexBuffer;
-	Buffer<float>	MeshVertexBuffer;
-	Buffer<uint>	MeshSkinWeightBuffer;
-	Buffer<uint>	MeshSkinWeightLookupBuffer;
-	Buffer<float4>	MeshCurrBonesBuffer;
-	Buffer<float4>	MeshPrevBonesBuffer;
-	Buffer<float4>	MeshCurrSamplingBonesBuffer;
-    Buffer<float4>	MeshPrevSamplingBonesBuffer;
-	Buffer<float4>	MeshTangentBuffer;
-	Buffer<float2>	MeshTexCoordBuffer;
-    Buffer<float4>	MeshColorBuffer;
-	Buffer<uint>	MeshTriangleSamplerProbAliasBuffer;
-	uint			MeshNumSamplingRegionTriangles;
-	uint			MeshNumSamplingRegionVertices;
-	Buffer<uint>	MeshSamplingRegionsProbAliasBuffer;
-	Buffer<uint>	MeshSampleRegionsTriangleIndices;
-	Buffer<uint>	MeshSampleRegionsVertices;
-	Buffer<uint>	MeshTriangleMatricesOffsetBuffer;
-	uint			MeshTriangleCount;
-	uint			MeshVertexCount;
-    uint			MeshWeightStride;
-	uint			MeshSkinWeightIndexSize;
-	uint			MeshNumTexCoord;
-	uint			MeshNumWeights;
-	int				NumBones;
-	int				NumFilteredBones;
-	int				NumUnfilteredBones;
-	int				RandomMaxBone;
-	int				ExcludeBoneIndex;
-	Buffer<uint>	FilteredAndUnfilteredBones;
-	int				NumFilteredSockets;
-	int				FilteredSocketBoneOffset;
-	Buffer<int>		UvMappingBuffer;
-	uint			UvMappingBufferLength;
-	uint			UvMappingSet;
-	Buffer<uint>		ConnectivityBuffer;
-	uint			ConnectivityMaxAdjacentPerVertex;
-	uint			ConnectivityBufferLength;
-	float4x4		InstanceTransform;
-	float4x4		InstancePrevTransform;
-	float4			InstanceRotation;
-	float4			InstancePrevRotation;
-	float			InstanceInvDeltaTime;
-	bool			UniformTriangleSamplingEnable;
-	bool			UniformSamplingRegionEnabled;
-	bool			UnlimitedBoneInfluences;
-	bool			bHasMeshColors;
-};
-
-////////////////////////////////////////////////////////////////////////////////////////////////////////////////////////////////////////////////////////////////
-// Triangle Sampling
-
 //-TODO: We need to consolodate the random interface inside NiagaraEmitterInstanceShader.usf it's a bit messy at the moment but should follow what FNDIRandomHelper does
 float DISKelMesh_Random(uint S1, uint S2, uint S3)
 {
@@ -105,36 +51,6 @@ float3 DISKelMesh_RandomBarycentricCoord(uint Seed1, uint Seed2, uint Seed3)
 	return float3(1.0f - sqrt0, sqrt0 * (1.0 - r1), r1 * sqrt0);
 }
 
-float3x4 DISkelMesh_GetPrevBoneSkinningMatrix(in FDISkelMeshContext DIContext, uint Bone)
-{
-    return float3x4(DIContext.MeshPrevBonesBuffer[Bone * 3], DIContext.MeshPrevBonesBuffer[Bone * 3 + 1], DIContext.MeshPrevBonesBuffer[Bone * 3 + 2]);
-}
-
-float3x4 DISkelMesh_GetPrevSkinningMatrix(in FDISkelMeshContext DIContext, uint SectionBoneOffset, int4 BlendIndices, float4 BlendWeights)
-{
-	float3x4 Result;
-    Result = DISkelMesh_GetPrevBoneSkinningMatrix(DIContext, SectionBoneOffset + BlendIndices.x) * BlendWeights.x;
-    Result += DISkelMesh_GetPrevBoneSkinningMatrix(DIContext, SectionBoneOffset + BlendIndices.y) * BlendWeights.y;
-    Result += DISkelMesh_GetPrevBoneSkinningMatrix(DIContext, SectionBoneOffset + BlendIndices.z) * BlendWeights.z;
-    Result += DISkelMesh_GetPrevBoneSkinningMatrix(DIContext, SectionBoneOffset + BlendIndices.w) * BlendWeights.w;
-	return Result;
-}
-
-float3x4 DISkelMesh_GetCurrBoneSkinningMatrix(in FDISkelMeshContext DIContext, uint Bone)
-{
-	return float3x4(DIContext.MeshCurrBonesBuffer[Bone * 3], DIContext.MeshCurrBonesBuffer[Bone * 3 + 1], DIContext.MeshCurrBonesBuffer[Bone * 3 + 2]);
-}
-
-float3x4 DISkelMesh_GetCurrSkinningMatrix(in FDISkelMeshContext DIContext, uint SectionBoneOffset, int4 BlendIndices, float4 BlendWeights)
-{
-	float3x4 Result;
-    Result = DISkelMesh_GetCurrBoneSkinningMatrix(DIContext, SectionBoneOffset + BlendIndices.x) * BlendWeights.x;
-    Result += DISkelMesh_GetCurrBoneSkinningMatrix(DIContext, SectionBoneOffset + BlendIndices.y) * BlendWeights.y;
-    Result += DISkelMesh_GetCurrBoneSkinningMatrix(DIContext, SectionBoneOffset + BlendIndices.z) * BlendWeights.z;
-    Result += DISkelMesh_GetCurrBoneSkinningMatrix(DIContext, SectionBoneOffset + BlendIndices.w) * BlendWeights.w;
-	return Result;
-}
-
 int4 DISkelMesh_UnpackIndices4(uint Packed)
 {
 	return int4(Packed & 0xff, Packed >> 8 & 0xff, Packed >> 16 & 0xff, Packed >> 24 & 0xff);
@@ -145,421 +61,6 @@ float4 DISkelMesh_UnpackWeights4(uint Packed)
 	return float4(Packed & 0xff, Packed >> 8 & 0xff, Packed >> 16 & 0xff, Packed >> 24 & 0xff) / 255.0f;
 }
 
-void DISkelMesh_GetVertex(in FDISkelMeshContext DIContext, uint VertexIndex, out float3 OutPosition, out float3 OutTangentX, out float3 OutTangentY, out float3 OutTangentZ)
-{
-	// Early out for bad data
-	if ( DIContext.MeshVertexCount == 0 )
-	{
-		OutPosition = float3(0.0f, 0.0f, 0.0f);
-		OutTangentX = float3(1.0f, 0.0f, 0.0f);
-		OutTangentY = float3(0.0f, 1.0f, 0.0f);
-		OutTangentZ = float3(0.0f, 0.0f, 1.0f);
-		return;
-	}
-
-	VertexIndex = clamp(VertexIndex, 0, DIContext.MeshVertexCount - 1);
-
-	//-TODO: Use a R32G32B32 format
-	float3 Position = float3(DIContext.MeshVertexBuffer[VertexIndex * 3], DIContext.MeshVertexBuffer[VertexIndex * 3 + 1], DIContext.MeshVertexBuffer[VertexIndex * 3 + 2]);
-	float3 TangentX = TangentBias(DIContext.MeshTangentBuffer[VertexIndex * 2    ].xyz);
-	float4 TangentZ = TangentBias(DIContext.MeshTangentBuffer[VertexIndex * 2 + 1].xyzw);	
-	float3 TangentY = cross(TangentZ.xyz, TangentX.xyz) * TangentZ.w;
-
-	OutPosition = Position;
-	OutTangentX = cross(TangentY, TangentZ.xyz) * TangentZ.w;
-	OutTangentY = TangentY;
-	OutTangentZ = TangentZ.xyz;
-}
-
-void DISkelMesh_GetPointOnTriangle(in FDISkelMeshContext DIContext, uint TriangleIndex, float3 BaryCoord, out float3 OutPosition, out float3 OutTangentX, out float3 OutTangentY, out float3 OutTangentZ)
-{
-	// Early out for bad data
-	if ( DIContext.MeshTriangleCount == 0 )
-	{
-		OutPosition = float3(0.0f, 0.0f, 0.0f);
-		OutTangentX = float3(1.0f, 0.0f, 0.0f);
-		OutTangentY = float3(0.0f, 1.0f, 0.0f);
-		OutTangentZ = float3(0.0f, 0.0f, 1.0f);
-		return;
-	}
-
-	TriangleIndex = clamp(TriangleIndex, 0, DIContext.MeshTriangleCount - 1);
-
-	uint IndexBufferOffset = TriangleIndex * 3;
-	uint VertexIndex0 = DIContext.MeshIndexBuffer[IndexBufferOffset];
-	uint VertexIndex1 = DIContext.MeshIndexBuffer[IndexBufferOffset + 1];
-	uint VertexIndex2 = DIContext.MeshIndexBuffer[IndexBufferOffset + 2];
-
-	float3 Positions[3], TangentsX[3], TangentsY[3], TangentsZ[3];
-	DISkelMesh_GetVertex(DIContext, VertexIndex0, Positions[0], TangentsX[0], TangentsY[0], TangentsZ[0]);
-	DISkelMesh_GetVertex(DIContext, VertexIndex1, Positions[1], TangentsX[1], TangentsY[1], TangentsZ[1]);
-	DISkelMesh_GetVertex(DIContext, VertexIndex2, Positions[2], TangentsX[2], TangentsY[2], TangentsZ[2]);
-	OutPosition = Positions[0] * BaryCoord.x + Positions[1] * BaryCoord.y  + Positions[2] * BaryCoord.z;
-	OutTangentX = TangentsX[0] * BaryCoord.x + TangentsX[1] * BaryCoord.y  + TangentsX[2] * BaryCoord.z;
-	OutTangentY = TangentsY[0] * BaryCoord.x + TangentsY[1] * BaryCoord.y  + TangentsY[2] * BaryCoord.z;
-	OutTangentZ = TangentsZ[0] * BaryCoord.x + TangentsZ[1] * BaryCoord.y  + TangentsZ[2] * BaryCoord.z;
-}
-
-FDISkelMeshSkinnedVertex DISkelMesh_GetSkinnedVertex(in FDISkelMeshContext DIContext, uint VertexIndex)
-{
-	// Early out for bad data
-	if ( DIContext.MeshVertexCount == 0 )
-	{
-		FDISkelMeshSkinnedVertex SkinnedVertex;
-		SkinnedVertex.Position = float3(0.0f, 0.0f, 0.0f);
-		SkinnedVertex.PrevPosition = float3(0.0f, 0.0f, 0.0f);
-		SkinnedVertex.TangentX = float3(1.0f, 0.0f, 0.0f);
-		SkinnedVertex.TangentY = float3(0.0f, 1.0f, 0.0f);
-		SkinnedVertex.TangentZ = float3(0.0f, 0.0f, 1.0f);
-		return SkinnedVertex;
-	}
-
-	VertexIndex = clamp(VertexIndex, 0, DIContext.MeshVertexCount - 1);
-
-	//-TODO: Use a R32G32B32 format
-	float3 CurrPosition = float3(DIContext.MeshVertexBuffer[VertexIndex * 3], DIContext.MeshVertexBuffer[VertexIndex * 3 + 1], DIContext.MeshVertexBuffer[VertexIndex * 3 + 2]);
-	float3 PrevPosition = CurrPosition;
-
-	float3 TangentX = TangentBias(DIContext.MeshTangentBuffer[VertexIndex * 2    ].xyz);
-	float4 TangentZ = TangentBias(DIContext.MeshTangentBuffer[VertexIndex * 2 + 1].xyzw);
-	
-	if ( DIContext.MeshNumWeights > 0 )
-	{
-		float3x4 PrevBoneMatrix;
-		float3x4 CurrBoneMatrix;
-
-		// Get the matrix offset for each vertex because BlendIndices are stored relatively to each section start vertex.
-		uint SectionBoneOffset = DIContext.MeshTriangleMatricesOffsetBuffer[VertexIndex];
-
-#if DISKELMESH_BONE_INFLUENCES == DISKELMESH_BONE_INFLUENCES_UNLIMITED
-		BRANCH
-		if (DIContext.UnlimitedBoneInfluences)
-		{
-			uint BlendOffsetCount = DIContext.MeshSkinWeightLookupBuffer[VertexIndex];
-			int NumBoneInfluences = BlendOffsetCount & 0xff;
-			int StreamOffset = BlendOffsetCount >> 8;
-			int WeightsOffset = StreamOffset + (DIContext.MeshSkinWeightIndexSize * NumBoneInfluences);
-
-			PrevBoneMatrix = float3x4(float4(0,0,0,0), float4(0,0,0,0), float4(0,0,0,0));
-			CurrBoneMatrix = float3x4(float4(0,0,0,0), float4(0,0,0,0), float4(0,0,0,0));
-			for (int InfluenceIdx = 0; InfluenceIdx < NumBoneInfluences; InfluenceIdx++)
-			{
-				int BoneIndexOffset = StreamOffset + (DIContext.MeshSkinWeightIndexSize * InfluenceIdx);
-				int BoneIndex = DIContext.MeshSkinWeightBuffer[BoneIndexOffset];
-				if (DIContext.MeshSkinWeightIndexSize > 1)
-				{
-					BoneIndex = DIContext.MeshSkinWeightBuffer[BoneIndexOffset + 1] << 8 | BoneIndex;
-					//@todo-lh: Workaround to fix issue in SPIRVEmitter of DXC; this block must be inside the if branch
-					float BoneWeight = float(DIContext.MeshSkinWeightBuffer[WeightsOffset + InfluenceIdx]) / 255.0f;
-
-					PrevBoneMatrix += DISkelMesh_GetPrevBoneSkinningMatrix(DIContext, SectionBoneOffset + BoneIndex) * BoneWeight;
-					CurrBoneMatrix += DISkelMesh_GetCurrBoneSkinningMatrix(DIContext, SectionBoneOffset + BoneIndex) * BoneWeight;
-				}
-				else
-				{
-					//@todo-lh: Workaround to fix issue in SPIRVEmitter of DXC; this block must be inside the if branch
-					float BoneWeight = float(DIContext.MeshSkinWeightBuffer[WeightsOffset + InfluenceIdx]) / 255.0f;
-
-					PrevBoneMatrix += DISkelMesh_GetPrevBoneSkinningMatrix(DIContext, SectionBoneOffset + BoneIndex) * BoneWeight;
-					CurrBoneMatrix += DISkelMesh_GetCurrBoneSkinningMatrix(DIContext, SectionBoneOffset + BoneIndex) * BoneWeight;
-				}
-			}	
-		}
-		else
-#endif
-		{
-			uint WeightOffset = DIContext.MeshNumWeights == 8 ? 2 : 1;
-			int4 BlendIndices = DISkelMesh_UnpackIndices4(DIContext.MeshSkinWeightBuffer[VertexIndex * DIContext.MeshWeightStride]);
-	        float4 BlendWeights = DISkelMesh_UnpackWeights4(DIContext.MeshSkinWeightBuffer[VertexIndex * DIContext.MeshWeightStride + WeightOffset]);
-			CurrBoneMatrix = DISkelMesh_GetCurrSkinningMatrix(DIContext, SectionBoneOffset, BlendIndices, BlendWeights);
-			PrevBoneMatrix = DISkelMesh_GetPrevSkinningMatrix(DIContext, SectionBoneOffset, BlendIndices, BlendWeights);
-#if DISKELMESH_BONE_INFLUENCES != DISKELMESH_BONE_INFLUENCES_MAX4
-			// We are in 8 or unlimited mode so need to check for the extended weights
-			if (DIContext.MeshNumWeights == 8)
-			{
-			    BlendIndices = DISkelMesh_UnpackIndices4(DIContext.MeshSkinWeightBuffer[VertexIndex * DIContext.MeshWeightStride + 1]);
-			    BlendWeights = DISkelMesh_UnpackWeights4(DIContext.MeshSkinWeightBuffer[VertexIndex * DIContext.MeshWeightStride + WeightOffset + 1]);
-				CurrBoneMatrix += DISkelMesh_GetCurrSkinningMatrix(DIContext, SectionBoneOffset, BlendIndices, BlendWeights);
-				PrevBoneMatrix += DISkelMesh_GetPrevSkinningMatrix(DIContext, SectionBoneOffset, BlendIndices, BlendWeights);
-			}
-#endif
-		}
-
-		CurrPosition = mul(CurrBoneMatrix, float4(CurrPosition, 1.0f)).xyz;
-		PrevPosition = mul(PrevBoneMatrix, float4(PrevPosition, 1.0f)).xyz;
-
-		// Not using InverseTranspose of matrices so assuming uniform scaling only (same as SkinCache)
-		TangentX.xyz = mul(CurrBoneMatrix, float4(TangentX.xyz, 0.0f)).xyz;
-		TangentZ.xyz = mul(CurrBoneMatrix, float4(TangentZ.xyz, 0.0f)).xyz;
-	}
-
-	float3 TangentY = cross(TangentZ.xyz, TangentX.xyz) * TangentZ.w;
-
-	FDISkelMeshSkinnedVertex SkinnedVertex;
-	SkinnedVertex.Position = CurrPosition;
-	SkinnedVertex.PrevPosition = PrevPosition;
-	SkinnedVertex.TangentX = cross(TangentY, TangentZ.xyz) * TangentZ.w;
-	SkinnedVertex.TangentY = TangentY;
-	SkinnedVertex.TangentZ = TangentZ.xyz;
-	return SkinnedVertex;
-}
-
-FDISkelMeshSkinnedVertex DISkelMesh_GetSkinnedPointOnTriangle(in FDISkelMeshContext DIContext, uint TriangleIndex, float3 BaryCoord)
-{
-	// Early out for bad data
-	if ( DIContext.MeshTriangleCount == 0 )
-	{
-		FDISkelMeshSkinnedVertex SkinnedVertex;
-		SkinnedVertex.Position = float3(0.0f, 0.0f, 0.0f);
-		SkinnedVertex.PrevPosition = float3(0.0f, 0.0f, 0.0f);
-		SkinnedVertex.TangentX = float3(1.0f, 0.0f, 0.0f);
-		SkinnedVertex.TangentY = float3(0.0f, 1.0f, 0.0f);
-		SkinnedVertex.TangentZ = float3(0.0f, 0.0f, 1.0f);
-		return SkinnedVertex;
-	}
-
-	TriangleIndex = clamp(TriangleIndex, 0, DIContext.MeshTriangleCount - 1);
-
-	uint IndexBufferOffset = TriangleIndex * 3;
-	uint VertexIndex0 = DIContext.MeshIndexBuffer[IndexBufferOffset];
-	uint VertexIndex1 = DIContext.MeshIndexBuffer[IndexBufferOffset + 1];
-	uint VertexIndex2 = DIContext.MeshIndexBuffer[IndexBufferOffset + 2];
-
-	FDISkelMeshSkinnedVertex SkinnedVertex0 = DISkelMesh_GetSkinnedVertex(DIContext, VertexIndex0);
-	FDISkelMeshSkinnedVertex SkinnedVertex1 = DISkelMesh_GetSkinnedVertex(DIContext, VertexIndex1);
-	FDISkelMeshSkinnedVertex SkinnedVertex2 = DISkelMesh_GetSkinnedVertex(DIContext, VertexIndex2);
-
-	FDISkelMeshSkinnedVertex FinalVertex;
-	FinalVertex.Position		= (SkinnedVertex0.Position * BaryCoord.x) + (SkinnedVertex1.Position * BaryCoord.y) + (SkinnedVertex2.Position * BaryCoord.z);
-	FinalVertex.PrevPosition	= (SkinnedVertex0.PrevPosition * BaryCoord.x) + (SkinnedVertex1.PrevPosition * BaryCoord.y) + (SkinnedVertex2.PrevPosition * BaryCoord.z);
-	FinalVertex.TangentX		= (SkinnedVertex0.TangentX * BaryCoord.x) + (SkinnedVertex1.TangentX * BaryCoord.y) + (SkinnedVertex2.TangentX * BaryCoord.z);
-	FinalVertex.TangentY		= (SkinnedVertex0.TangentY * BaryCoord.x) + (SkinnedVertex1.TangentY * BaryCoord.y) + (SkinnedVertex2.TangentY * BaryCoord.z);
-	FinalVertex.TangentZ		= (SkinnedVertex0.TangentZ * BaryCoord.x) + (SkinnedVertex1.TangentZ * BaryCoord.y) + (SkinnedVertex2.TangentZ * BaryCoord.z);
-
-	return FinalVertex;
-}
-
-FDISkelMeshSkinnedVertex DISkelMesh_GetSkinnedPointOnTriangleWS(in FDISkelMeshContext DIContext, uint TriangleIndex, float3 BaryCoord)
-{
-	FDISkelMeshSkinnedVertex FinalVertex = DISkelMesh_GetSkinnedPointOnTriangle(DIContext, TriangleIndex, BaryCoord);
-	FinalVertex.Position		= mul(float4(FinalVertex.Position, 1.0f), DIContext.InstanceTransform).xyz;
-	FinalVertex.PrevPosition	= mul(float4(FinalVertex.PrevPosition, 1.0f), DIContext.InstancePrevTransform).xyz;
-	FinalVertex.TangentX		= mul(float4(FinalVertex.TangentX, 0.0f), DIContext.InstanceTransform).xyz;
-	FinalVertex.TangentY		= mul(float4(FinalVertex.TangentY, 0.0f), DIContext.InstanceTransform).xyz;
-	FinalVertex.TangentZ		= mul(float4(FinalVertex.TangentZ, 0.0f), DIContext.InstanceTransform).xyz;
-
-	return FinalVertex;
-}
-
-void DISKelMesh_RandomTriCoord(in FDISkelMeshContext DIContext, uint Seed1, uint Seed2, uint Seed3, out int OutTriangle, out float3 OutBaryCoord)
-{
-	// Early out for bad data
-	if (DIContext.MeshTriangleCount == 0)
-	{
-		OutTriangle = -1;
-		OutBaryCoord = 1.0f / 3.0f;
-		return;
-	}
-
-	float RandT0 = DISKelMesh_Random(Seed1, Seed2, Seed3);
-
-	[branch]
-	if ( !DIContext.UniformTriangleSamplingEnable )
-	{
-		// Uniform triangle id selection
-		OutTriangle = min(uint(RandT0 * float(DIContext.MeshTriangleCount)), DIContext.MeshTriangleCount - 1); // avoid % by using mul/min to Tri = MeshTriangleCountName
-	}
-	else
-	{
-		// Uniform area weighted position selection (using alias method from Alias method from FWeightedRandomSampler)
-		uint TriangleIndex = min(uint(RandT0*float(DIContext.MeshTriangleCount)), DIContext.MeshTriangleCount - 1);
-
-#if DISKELMESH_PROBALIAS_FORMAT == DISKELMESH_PROBALIAS_FORMAT_64
-		float TriangleProbability = asfloat(DIContext.MeshTriangleSamplerProbAliasBuffer[(TriangleIndex * 2) + 0]);
-		uint TriangleAlias = DIContext.MeshTriangleSamplerProbAliasBuffer[(TriangleIndex * 2) + 1];
-#elif DISKELMESH_PROBALIAS_FORMAT == DISKELMESH_PROBALIAS_FORMAT_24_8
-		uint TriangleProbAlias = DIContext.MeshTriangleSamplerProbAliasBuffer[TriangleIndex];
-		float TriangleProbability = float(TriangleProbAlias & 0xff) / 255.0f;
-		uint TriangleAlias = TriangleProbAlias >> 8;
-#elif DISKELMESH_PROBALIAS_FORMAT == DISKELMESH_PROBALIAS_FORMAT_23_9
-		uint TriangleProbAlias = DIContext.MeshTriangleSamplerProbAliasBuffer[TriangleIndex];
-		float TriangleProbability = float(TriangleProbAlias & 0x1ff) / 511.0f;
-		uint TriangleAlias = TriangleProbAlias >> 9;
-#endif
-
-		// Alias check
-		float RandT1 = DISKelMesh_Random(Seed1, Seed2, Seed3);
-		if( RandT1 > TriangleProbability )
-		{
-			TriangleIndex = TriangleAlias;
-		}
-		OutTriangle = int(TriangleIndex);
-	}
-
-	OutBaryCoord = DISKelMesh_RandomBarycentricCoord(Seed1, Seed2, Seed3);
-}
-
-void DISKelMesh_GetSkinnedTriangleDataWS(in FDISkelMeshContext DIContext, in uint TriangleIndex, in float3 BaryCoord, out float3 OutPosition, out float3 OutVelocity, out float3 OutNormal, out float3 OutBinormal, out float3 OutTangent)
-{
-    FDISkelMeshSkinnedVertex SkinnedVertex = DISkelMesh_GetSkinnedPointOnTriangleWS(DIContext, TriangleIndex, BaryCoord);
-	OutPosition	= SkinnedVertex.Position;
-	OutVelocity	= (SkinnedVertex.Position - SkinnedVertex.PrevPosition) * DIContext.InstanceInvDeltaTime;
-	OutNormal	= normalize(SkinnedVertex.TangentZ);
-	OutBinormal	= normalize(SkinnedVertex.TangentY);
-	OutTangent	= normalize(SkinnedVertex.TangentX);
-}
-
-void DISKelMesh_GetSkinnedTriangleData(in FDISkelMeshContext DIContext, in uint TriangleIndex, in float3 BaryCoord, out float3 OutPosition, out float3 OutVelocity, out float3 OutNormal, out float3 OutBinormal, out float3 OutTangent)
-{
-    FDISkelMeshSkinnedVertex SkinnedVertex = DISkelMesh_GetSkinnedPointOnTriangle(DIContext, TriangleIndex, BaryCoord);
-	OutPosition	= SkinnedVertex.Position;
-	OutVelocity	= (SkinnedVertex.Position - SkinnedVertex.PrevPosition) * DIContext.InstanceInvDeltaTime;
-	OutNormal	= normalize(SkinnedVertex.TangentZ);
-	OutBinormal	= normalize(SkinnedVertex.TangentY);
-	OutTangent	= normalize(SkinnedVertex.TangentX);
-}
-
-void DISKelMesh_GetSkinnedTriangleDataInterpolatedWS(in FDISkelMeshContext DIContext, in uint TriangleIndex, in float3 BaryCoord, in float Interp, out float3 OutPosition, out float3 OutVelocity, out float3 OutNormal, out float3 OutBinormal, out float3 OutTangent)
-{
-    FDISkelMeshSkinnedVertex SkinnedVertex = DISkelMesh_GetSkinnedPointOnTriangleWS(DIContext, TriangleIndex, BaryCoord);
-	OutPosition	= lerp(SkinnedVertex.PrevPosition, SkinnedVertex.Position, Interp);
-	OutVelocity	= (SkinnedVertex.Position - SkinnedVertex.PrevPosition) * DIContext.InstanceInvDeltaTime;
-	OutNormal	= normalize(SkinnedVertex.TangentZ);
-	OutBinormal	= normalize(SkinnedVertex.TangentY);
-	OutTangent	= normalize(SkinnedVertex.TangentX);
-}
-
-void DISKelMesh_GetSkinnedTriangleDataInterpolated(in FDISkelMeshContext DIContext, in uint TriangleIndex, in float3 BaryCoord, in float Interp, out float3 OutPosition, out float3 OutVelocity, out float3 OutNormal, out float3 OutBinormal, out float3 OutTangent)
-{
-    FDISkelMeshSkinnedVertex SkinnedVertex = DISkelMesh_GetSkinnedPointOnTriangle(DIContext, TriangleIndex, BaryCoord);
-	OutPosition	= lerp(SkinnedVertex.PrevPosition, SkinnedVertex.Position, Interp);
-	OutVelocity	= (SkinnedVertex.Position - SkinnedVertex.PrevPosition) * DIContext.InstanceInvDeltaTime;
-	OutNormal	= normalize(SkinnedVertex.TangentZ);
-	OutBinormal	= normalize(SkinnedVertex.TangentY);
-	OutTangent	= normalize(SkinnedVertex.TangentX);
-}
-
-void DISKelMesh_GetTriUV(in FDISkelMeshContext DIContext, in uint TriangleIndex, in float3 BaryCoord, in int UVSet, out float2 OutUV)
-{
-    if (DIContext.MeshNumTexCoord > 0)
-	{
-        uint IndexBufferOffset = TriangleIndex * 3;
-		uint VertexIndex0 = DIContext.MeshIndexBuffer[IndexBufferOffset  ];
-		uint VertexIndex1 = DIContext.MeshIndexBuffer[IndexBufferOffset+1];
-		uint VertexIndex2 = DIContext.MeshIndexBuffer[IndexBufferOffset+2];
-
-        uint Stride = DIContext.MeshNumTexCoord;
-        uint SelectedUVSet = clamp((uint) UVSet, 0, DIContext.MeshNumTexCoord - 1);
-		float2 UV0 = DIContext.MeshTexCoordBuffer[VertexIndex0 * Stride + SelectedUVSet];
-		float2 UV1 = DIContext.MeshTexCoordBuffer[VertexIndex1 * Stride + SelectedUVSet];
-		float2 UV2 = DIContext.MeshTexCoordBuffer[VertexIndex2 * Stride + SelectedUVSet];
-
-		OutUV = UV0 * BaryCoord.x + UV1 * BaryCoord.y + UV2 * BaryCoord.z;
-	}
-	else	
-	{
-		OutUV = 0.0f;
-	}
-}
-
-void DISKelMesh_RandomTriangle(in FDISkelMeshContext DIContext, uint Seed1, uint Seed2, uint Seed3, out int OutTriangle, out float3 OutBaryCoord)
-{
-	float RandT0 = DISKelMesh_Random(Seed1, Seed2, Seed3);
-	int TriangleIndex = min(int(RandT0 * float(DIContext.MeshTriangleCount)), (int)DIContext.MeshTriangleCount - 1); // avoid % by using mul/min to Tri = MeshTriangleCountName
-
-	[branch]
-	if ( DIContext.UniformTriangleSamplingEnable )
-	{
-		// Uniform area weighted position selection (using alias method from Alias method from FWeightedRandomSampler)
-#if DISKELMESH_PROBALIAS_FORMAT == DISKELMESH_PROBALIAS_FORMAT_64
-		float TriangleProbability = asfloat(DIContext.MeshTriangleSamplerProbAliasBuffer[(TriangleIndex * 2) + 0]);
-		uint TriangleAlias = DIContext.MeshTriangleSamplerProbAliasBuffer[(TriangleIndex * 2) + 1];
-#elif DISKELMESH_PROBALIAS_FORMAT == DISKELMESH_PROBALIAS_FORMAT_24_8
-		uint TriangleProbAlias = DIContext.MeshTriangleSamplerProbAliasBuffer[TriangleIndex];
-		float TriangleProbability = float(TriangleProbAlias & 0xff) / 255.0f;
-		uint TriangleAlias = TriangleProbAlias >> 8;
-#elif DISKELMESH_PROBALIAS_FORMAT == DISKELMESH_PROBALIAS_FORMAT_23_9
-		uint TriangleProbAlias = DIContext.MeshTriangleSamplerProbAliasBuffer[TriangleIndex];
-		float TriangleProbability = float(TriangleProbAlias & 0x1ff) / 511.0f;
-		uint TriangleAlias = TriangleProbAlias >> 9;
-#endif
-
-		// Alias check
-		float RandT1 = DISKelMesh_Random(Seed1, Seed2, Seed3);
-		if( RandT1 > TriangleProbability )
-		{
-			TriangleIndex = int(TriangleAlias);
-		}
-	}
-
-	OutTriangle = TriangleIndex;
-	OutBaryCoord = DISKelMesh_RandomBarycentricCoord(Seed1, Seed2, Seed3);
-}
-
-void DISKelMesh_GetTriangleCount(in FDISkelMeshContext DIContext, out int OutCount)
-{
-	OutCount = DIContext.MeshTriangleCount;
-}
-
-void DISKelMesh_RandomFilteredTriangle(in FDISkelMeshContext DIContext, uint Seed1, uint Seed2, uint Seed3, out int OutTriangle, out float3 OutBaryCoord)
-{
-	// Early out for bad data
-	if (DIContext.MeshNumSamplingRegionTriangles == 0)
-	{
-		OutTriangle = -1;
-		OutBaryCoord = 1.0f / 3.0f;
-		return;
-	}
-
-	float RandT0 = DISKelMesh_Random(Seed1, Seed2, Seed3);
-	uint RegionTriangle = min(uint(RandT0 * float(DIContext.MeshNumSamplingRegionTriangles)), DIContext.MeshNumSamplingRegionTriangles - 1); // avoid % by using mul/min to Tri = MeshTriangleCountName
-
-	[branch]
-	if ( DIContext.UniformSamplingRegionEnabled )
-	{
-#if DISKELMESH_PROBALIAS_FORMAT == DISKELMESH_PROBALIAS_FORMAT_64
-		float TriangleProbability = asfloat(DIContext.MeshSamplingRegionsProbAliasBuffer[(RegionTriangle * 2) + 0]);
-		uint TriangleAlias = DIContext.MeshSamplingRegionsProbAliasBuffer[(RegionTriangle * 2) + 1];
-#elif DISKELMESH_PROBALIAS_FORMAT == DISKELMESH_PROBALIAS_FORMAT_24_8
-		uint TriangleProbAlias = DIContext.MeshSamplingRegionsProbAliasBuffer[RegionTriangle];
-		float TriangleProbability = float(TriangleProbAlias & 0xff) / 255.0f;
-		uint TriangleAlias = TriangleProbAlias >> 8;
-#elif DISKELMESH_PROBALIAS_FORMAT == DISKELMESH_PROBALIAS_FORMAT_23_9
-		uint TriangleProbAlias = DIContext.MeshSamplingRegionsProbAliasBuffer[RegionTriangle];
-		float TriangleProbability = float(TriangleProbAlias & 0x1ff) / 511.0f;
-		uint TriangleAlias = TriangleProbAlias >> 9;
-#endif
-		float RandT1 = DISKelMesh_Random(Seed1, Seed2, Seed3);
-		if ( RandT1 > TriangleProbability )
-		{
-			RegionTriangle = TriangleAlias;
-		}
-	}
-	OutTriangle = (int)DIContext.MeshSampleRegionsTriangleIndices[RegionTriangle];
-	OutBaryCoord = DISKelMesh_RandomBarycentricCoord(Seed1, Seed2, Seed3);
-}
-
-void DISKelMesh_GetFilteredTriangleCount(in FDISkelMeshContext DIContext, out int OutCount)
-{
-	OutCount = DIContext.MeshNumSamplingRegionTriangles;
-}
-
-void DISKelMesh_GetFilteredTriangleAt(in FDISkelMeshContext DIContext, int FilteredIndex, out int OutTriangle)
-{
-	if ( DIContext.MeshNumSamplingRegionTriangles > 0 )
-	{
-		FilteredIndex = clamp(FilteredIndex, 0, (int)DIContext.MeshNumSamplingRegionTriangles - 1);
-		OutTriangle = (int)DIContext.MeshSampleRegionsTriangleIndices[FilteredIndex];
-	}
-	else
-	{
-		OutTriangle = 0;
-	}
-}
-
-////////////////////////////////////////////////////////////////////////////////////////////////////////////////////////////////////////////////////////////////
-// Uv Mapping
-
 // UvMappingBuffer organized as an array of:
 //struct FFrozenNode
 //{
@@ -568,94 +69,11 @@ void DISKelMesh_GetFilteredTriangleAt(in FDISkelMeshContext DIContext, int Filte
 //	int32 TriangleIndices[]
 //};
 
-#define ElementCountOffset 4
-#define TriangleIndexOffset 5
-#define MinValidUvMappingBufferSize 24 // minimum size (in bytes) for a quadtree to have valid data
+#define DISkelMesh_ElementCountOffset 4
+#define DISkelMesh_TriangleIndexOffset 5
+#define DISkelMesh_MinValidUvMappingBufferSize 24 // minimum size (in bytes) for a quadtree to have valid data
 
-// Evaluates QuadTree for a single point
-bool QuadTreeTriangleOverlap(FDISkelMeshContext DIContext, float2 Uv, float Tolerance, int TriangleIndex, out float3 OutBaryCoord)
-{
-	if (DIContext.MeshNumTexCoord > 0)
-	{
-        uint IndexBufferOffset = TriangleIndex * 3;
-		uint VertexIndex0 = DIContext.MeshIndexBuffer[IndexBufferOffset  ];
-		uint VertexIndex1 = DIContext.MeshIndexBuffer[IndexBufferOffset+1];
-		uint VertexIndex2 = DIContext.MeshIndexBuffer[IndexBufferOffset+2];
-
-		uint Stride = DIContext.MeshNumTexCoord;
-		float2 UV0 = DIContext.MeshTexCoordBuffer[VertexIndex0 * Stride + DIContext.UvMappingSet];
-		float2 UV1 = DIContext.MeshTexCoordBuffer[VertexIndex1 * Stride + DIContext.UvMappingSet];
-		float2 UV2 = DIContext.MeshTexCoordBuffer[VertexIndex2 * Stride + DIContext.UvMappingSet];
-
-		OutBaryCoord = GetBaryCentric2D(Uv, UV0, UV1, UV2);
-
-		if (all(OutBaryCoord > -Tolerance) && all(OutBaryCoord < (1.0f + Tolerance)))
-		{
-			return true;
-		}
-	}
-
-	return false;
-}
-
-bool QuadTreeNodeContains(FDISkelMeshContext DIContext, int NodeBufferOffset, float2 Uv, float Tolerance, out int OutTriangleIndex, out float3 OutBaryCoord)
-{
-	if (NodeBufferOffset == -1)
-	{
-		return false;
-	}
-
-	const int TriangleCount = DIContext.UvMappingBuffer[NodeBufferOffset + ElementCountOffset];
-	for (int TriangleIt = 0; TriangleIt < TriangleCount; ++TriangleIt)
-	{
-		int TriangleIndex = DIContext.UvMappingBuffer[NodeBufferOffset + TriangleIndexOffset + TriangleIt];
-		if (QuadTreeTriangleOverlap(DIContext, Uv, Tolerance, TriangleIndex, OutBaryCoord))
-		{
-			OutTriangleIndex = TriangleIndex;
-			return true;
-		}
-	}
-
-	return false;
-}
-
-void DISkelMesh_GetTriangleCoordAtUV(in FDISkelMeshContext DIContext, bool Enabled, float2 InUv, float InTolerance, out int OutTriangle, out float3 OutBaryCoord, out bool OutIsValid)
-{
-	OutTriangle = -1;
-	OutBaryCoord = float3(0.0f, 0.0f, 0.0f);
-	OutIsValid = false;
-
-	if (!Enabled || DIContext.UvMappingBufferLength < MinValidUvMappingBufferSize)
-	{
-		return;
-	}
-
-	int NodeBufferOffset = 0;
-
-	float2 MidPoint = float2(0.5f, 0.5f);
-	float2 Extent = float2(0.5f, 0.5f);
-
-	while (NodeBufferOffset != -1)
-	{
-		if (QuadTreeNodeContains(DIContext, NodeBufferOffset, InUv, InTolerance, OutTriangle, OutBaryCoord))
-		{
-			OutIsValid = true;
-			return;
-		}
-
-		Extent *= 0.5f;
-
-		const bool LessX = InUv.x < MidPoint.x;
-		const bool LessY = InUv.y < MidPoint.y;
-
-		MidPoint += Extent * float2(LessX ? -1.0f : 1.0f, LessY ? -1.0f : 1.0f);
-		int ChildIndex = LessX ? (LessY ? 2 : 0) : (LessY ? 3 : 1);
-
-		NodeBufferOffset = DIContext.UvMappingBuffer[NodeBufferOffset + ChildIndex];
-	}
-}
-
-bool NormalizedAabbTriangleIntersection(float2 A, float2 B, float2 C)
+bool DISkelMesh_NormalizedAabbTriangleIntersection(float2 A, float2 B, float2 C)
 {
 	float2 TriAabbMin = float2(min3(A.x, B.x, C.x), min3(A.y, B.y, C.y));
 	float2 TriAabbMax = float2(max3(A.x, B.x, C.x), max3(A.y, B.y, C.y));
@@ -684,109 +102,6 @@ bool NormalizedAabbTriangleIntersection(float2 A, float2 B, float2 C)
 	return true;
 }
 
-// Evaluates QuadTree for an AABB region
-bool QuadTreeTriangleAabbOverlap(FDISkelMeshContext DIContext, float2 UvRef, float2 UvExtentScale, float2 UvExtentBias, int TriangleIndex, out float3 OutBaryCoord)
-{
-	if (DIContext.MeshNumTexCoord > 0)
-	{
-        uint IndexBufferOffset = TriangleIndex * 3;
-		uint VertexIndex0 = DIContext.MeshIndexBuffer[IndexBufferOffset  ];
-		uint VertexIndex1 = DIContext.MeshIndexBuffer[IndexBufferOffset+1];
-		uint VertexIndex2 = DIContext.MeshIndexBuffer[IndexBufferOffset+2];
-
-		uint Stride = DIContext.MeshNumTexCoord;
-		float2 UV0 = DIContext.MeshTexCoordBuffer[VertexIndex0 * Stride + DIContext.UvMappingSet];
-		float2 UV1 = DIContext.MeshTexCoordBuffer[VertexIndex1 * Stride + DIContext.UvMappingSet];
-		float2 UV2 = DIContext.MeshTexCoordBuffer[VertexIndex2 * Stride + DIContext.UvMappingSet];
-
-		// evaluate whether our triangle intersects with our AABB, if it does, then generate the barycentric coordinate for the center of the AABB
-		const bool Intersection = NormalizedAabbTriangleIntersection(
-				UV0 * UvExtentScale + UvExtentBias,
-				UV1 * UvExtentScale + UvExtentBias,
-				UV2 * UvExtentScale + UvExtentBias);
-
-		if (Intersection)
-		{
-			OutBaryCoord = GetClosestBaryCentric2D(UvRef, UV0, UV1, UV2);
-			return true;
-		}
-	}
-
-	return false;
-}
-
-bool QuadTreeNodeContainsAabb(FDISkelMeshContext DIContext, int NodeBufferOffset, float2 UvMin, float2 UvMax, out int OutTriangleIndex, out float3 OutBaryCoord)
-{
-	if (NodeBufferOffset == -1)
-	{
-		return false;
-	}
-
-	const int TriangleCount = DIContext.UvMappingBuffer[NodeBufferOffset + ElementCountOffset];
-	if (TriangleCount == 0)
-	{
-		return false;
-	}
-
-	const float2 UvRef = 0.5f * (UvMin + UvMax);
-	const float2 UvExtentScale = 1.0f / (UvMax - UvMin);
-	const float2 UvExtentBias = 1.0f - UvMax * UvExtentScale;
-
-	for (int TriangleIt = 0; TriangleIt < TriangleCount; ++TriangleIt)
-	{
-		float3 BaryCoord;
-		int TriangleIndex = DIContext.UvMappingBuffer[NodeBufferOffset + TriangleIndexOffset + TriangleIt];
-		if (QuadTreeTriangleAabbOverlap(DIContext, UvRef, UvExtentScale, UvExtentBias, TriangleIndex, BaryCoord))
-		{
-			OutTriangleIndex = TriangleIndex;
-			OutBaryCoord = BaryCoord;
-			return true;
-		}
-	}
-
-	return false;
-}
-
-void DISkelMesh_GetTriangleCoordInAabb(FDISkelMeshContext DIContext, bool Enabled, float2 InUvMin, float2 InUvMax, out int OutTriangle, out float3 OutBaryCoord, out bool OutIsValid)
-{
-	OutTriangle = -1;
-	OutBaryCoord = float3(0.0f, 0.0f, 0.0f);
-	OutIsValid = false;
-
-	if (!Enabled || DIContext.UvMappingBufferLength < MinValidUvMappingBufferSize)
-	{
-		return;
-	}
-
-	int NodeBufferOffset = 0;
-
-	const float2 UvRegionCenter = 0.5f * (InUvMin + InUvMax);
-	float2 MidPoint = float2(0.5f, 0.5f);
-	float2 Extent = float2(0.5f, 0.5f);
-
-	while (NodeBufferOffset != -1)
-	{
-		if (QuadTreeNodeContainsAabb(DIContext, NodeBufferOffset, InUvMin, InUvMax, OutTriangle, OutBaryCoord))
-		{
-			OutIsValid = true;
-			return;
-		}
-
-		Extent *= 0.5f;
-
-		const bool LessX = UvRegionCenter.x < MidPoint.x;
-		const bool LessY = UvRegionCenter.y < MidPoint.y;
-
-		MidPoint += Extent * float2(LessX ? -1.0f : 1.0f, LessY ? -1.0f : 1.0f);
-		int ChildIndex = LessX ? (LessY ? 2 : 0) : (LessY ? 3 : 1);
-
-		NodeBufferOffset = DIContext.UvMappingBuffer[NodeBufferOffset + ChildIndex];
-	}
-}
-
-////////////////////////////////////////////////////////////////////////////////////////////////////////////////////////////////////////////////////////////////
-// Bone Sampling
-
 float4 DISKelMesh_QuatSlerp(float4 Quat1, float4 Quat2, float Slerp)
 {
     const float RawCosom = dot(Quat1, Quat2);
@@ -811,543 +126,3 @@ float4 DISKelMesh_QuatSlerp(float4 Quat1, float4 Quat2, float Slerp)
 
     return (Scale0 * Quat1) + (Scale1 * Quat2);
 }
-
-void DISkelMesh_GetSkinnedBoneCommon(in FDISkelMeshContext DIContext, in int Bone, in float Interp, in bool bWorldSpace, out float3 OutPosition, out float4 OutRotation, out float3 OutScale, out float3 OutVelocity)
-{
-	// Early out for bad data
-	if ( DIContext.NumBones + DIContext.NumFilteredSockets == 0 )
-	{
-		OutPosition = float3(0.0f, 0.0f, 0.0f);
-		OutRotation = float4(0.0f, 0.0f, 0.0f, 1.0f);
-		OutScale	= float3(0.0f, 0.0f, 0.0f);
-		OutVelocity = float3(0.0f, 0.0f, 0.0f);
-		return;
-	}
-
-	Bone = clamp(Bone, 0, DIContext.NumBones + DIContext.NumFilteredSockets - 1);
-
-    float3 PrevPosition = DIContext.MeshPrevSamplingBonesBuffer[Bone * 3].xyz;
-    float4 PrevRotation = DIContext.MeshPrevSamplingBonesBuffer[Bone * 3 + 1];
-	float3 PrevScale	= DIContext.MeshPrevSamplingBonesBuffer[Bone * 3 + 2].xyz;
-    float3 CurrPosition = DIContext.MeshCurrSamplingBonesBuffer[Bone * 3].xyz;
-    float4 CurrRotation	= DIContext.MeshCurrSamplingBonesBuffer[Bone * 3 + 1];
-	float3 CurrScale	= DIContext.MeshCurrSamplingBonesBuffer[Bone * 3 + 2].xyz;
-
-    if (bWorldSpace)
-    {
-        PrevPosition = mul(float4(PrevPosition, 1), DIContext.InstancePrevTransform).xyz;
-        CurrPosition = mul(float4(CurrPosition, 1), DIContext.InstanceTransform).xyz;
-		PrevRotation = NiagaraGPU_QuatMul(DIContext.InstancePrevRotation, PrevRotation);
-		CurrRotation = NiagaraGPU_QuatMul(DIContext.InstanceRotation, CurrRotation);
-		PrevScale = mul(PrevScale, (float3x3)DIContext.InstancePrevTransform);
-		CurrScale = mul(CurrScale, (float3x3)DIContext.InstanceTransform);
-    }
-    CurrPosition = lerp(PrevPosition, CurrPosition, Interp);
-	CurrScale	 = lerp(PrevScale, CurrScale, Interp);
-
-    OutPosition = CurrPosition;
-    OutRotation = DISKelMesh_QuatSlerp(PrevRotation, CurrRotation, Interp);
-	OutScale	= CurrScale;
-    OutVelocity = (CurrPosition - PrevPosition) * DIContext.InstanceInvDeltaTime;
-}
-
-void DISkelMesh_GetSkinnedBone(in FDISkelMeshContext DIContext, in int Bone, out float3 OutPosition, out float4 OutRotation, out float3 OutScale, out float3 OutVelocity)
-{
-    DISkelMesh_GetSkinnedBoneCommon(DIContext, Bone, 1.0f, false, OutPosition, OutRotation, OutScale, OutVelocity);
-}
-
-void DISkelMesh_GetSkinnedBoneInterpolated(in FDISkelMeshContext DIContext, in int Bone, in float Interp, out float3 OutPosition, out float4 OutRotation, out float3 OutScale, out float3 OutVelocity)
-{
-    DISkelMesh_GetSkinnedBoneCommon(DIContext, Bone, Interp, false, OutPosition, OutRotation, OutScale, OutVelocity);
-}
-
-void DISkelMesh_GetSkinnedBoneWS(in FDISkelMeshContext DIContext, in int Bone, out float3 OutPosition, out float4 OutRotation, out float3 OutScale, out float3 OutVelocity)
-{
-    DISkelMesh_GetSkinnedBoneCommon(DIContext, Bone, 1.0f, true, OutPosition, OutRotation, OutScale, OutVelocity);
-}
-
-void DISkelMesh_GetSkinnedBoneInterpolatedWS(in FDISkelMeshContext DIContext, in int Bone, in float Interp, out float3 OutPosition, out float4 OutRotation, out float3 OutScale, out float3 OutVelocity)
-{
-    DISkelMesh_GetSkinnedBoneCommon(DIContext, Bone, Interp, true, OutPosition, OutRotation, OutScale, OutVelocity);
-}
-
-////////////////////////////////////////////////////////////////////////////////////////////////////////////////////////////////////////////////////////////////
-// Vertex Sampling
-
-void DISkelMesh_GetTriVertices(in FDISkelMeshContext DIContext, in uint TriangleIndex, out int VertexIndex0, out int VertexIndex1, out int VertexIndex2)
-{
-	// Early out for bad data
-	if ( DIContext.MeshTriangleCount == 0 )
-	{
-		VertexIndex0 = 0;
-		VertexIndex1 = 0;
-		VertexIndex2 = 0;
-		return;
-	}
-
-	TriangleIndex = clamp(TriangleIndex, 0, DIContext.MeshTriangleCount - 1);
-    uint IndexBufferOffset = TriangleIndex * 3;
-    VertexIndex0 = DIContext.MeshIndexBuffer[IndexBufferOffset];
-    VertexIndex1 = DIContext.MeshIndexBuffer[IndexBufferOffset + 1];
-    VertexIndex2 = DIContext.MeshIndexBuffer[IndexBufferOffset + 2];
-}
-
-void DISkelMesh_GetTriColor(in FDISkelMeshContext DIContext, in uint TriangleIndex, in float3 BaryCoord, out float4 OutColor)
-{
-	// Early out for bad data
-	if ( DIContext.bHasMeshColors == false )
-	{
-		OutColor = float4(1, 1, 1, 1);
-		return;
-	}
-	TriangleIndex = clamp(TriangleIndex, 0, DIContext.MeshTriangleCount - 1);
-
-    uint IndexBufferOffset = TriangleIndex * 3;
-    uint VertexIndex0 = DIContext.MeshIndexBuffer[IndexBufferOffset];
-    uint VertexIndex1 = DIContext.MeshIndexBuffer[IndexBufferOffset + 1];
-    uint VertexIndex2 = DIContext.MeshIndexBuffer[IndexBufferOffset + 2];
-    OutColor = (DIContext.MeshColorBuffer[VertexIndex0] FMANUALFETCH_COLOR_COMPONENT_SWIZZLE * BaryCoord.x) + (DIContext.MeshColorBuffer[VertexIndex1] FMANUALFETCH_COLOR_COMPONENT_SWIZZLE * BaryCoord.y) + (DIContext.MeshColorBuffer[VertexIndex2] FMANUALFETCH_COLOR_COMPONENT_SWIZZLE * BaryCoord.z);
-}
-
-void DISkelMesh_GetSkinnedVertex(in FDISkelMeshContext DIContext, in int VertexIndex, out float3 OutPosition, out float3 OutVelocity, out float3 OutNormal, out float3 OutBinormal, out float3 OutTangent)
-{
-    FDISkelMeshSkinnedVertex SkinnedVertex = DISkelMesh_GetSkinnedVertex(DIContext, VertexIndex);
-    OutPosition = SkinnedVertex.Position;
-    OutVelocity = (SkinnedVertex.Position - SkinnedVertex.PrevPosition) * DIContext.InstanceInvDeltaTime;
-	OutNormal = SkinnedVertex.TangentZ;
-	OutBinormal = SkinnedVertex.TangentY;
-	OutTangent = SkinnedVertex.TangentX;
-}
-
-void DISkelMesh_GetSkinnedVertexWS(in FDISkelMeshContext DIContext, in int VertexIndex, out float3 OutPosition, out float3 OutVelocity, out float3 OutNormal, out float3 OutBinormal, out float3 OutTangent)
-{
-    FDISkelMeshSkinnedVertex SkinnedVertex = DISkelMesh_GetSkinnedVertex(DIContext, VertexIndex);
-    OutPosition = mul(float4(SkinnedVertex.Position, 1.0f), DIContext.InstanceTransform).xyz;
-    float3 PrevPosition = mul(float4(SkinnedVertex.PrevPosition, 1.0f), DIContext.InstancePrevTransform).xyz;
-    OutVelocity = (SkinnedVertex.Position - SkinnedVertex.PrevPosition) * DIContext.InstanceInvDeltaTime;
-    OutNormal   = normalize(mul(float4(SkinnedVertex.TangentZ, 0.0f), DIContext.InstanceTransform).xyz);
-    OutBinormal = normalize(mul(float4(SkinnedVertex.TangentY, 0.0f), DIContext.InstanceTransform).xyz);
-    OutTangent  = normalize(mul(float4(SkinnedVertex.TangentX, 0.0f), DIContext.InstanceTransform).xyz);
-}
-
-void DISkelMesh_GetVertexColor(in FDISkelMeshContext DIContext, in int VertexIndex, out float4 OutColor)
-{
-    OutColor = DIContext.bHasMeshColors ? DIContext.MeshColorBuffer[VertexIndex] FMANUALFETCH_COLOR_COMPONENT_SWIZZLE : float4(1, 1, 1, 1);
-}
-
-void DISkelMesh_GetVertexUV(in FDISkelMeshContext DIContext, in int VertexIndex, in int UVSet, out float2 OutUV)
-{
-    if (DIContext.MeshNumTexCoord > 0)
-    {
-        uint Stride = DIContext.MeshNumTexCoord;
-        uint SelectedUVSet = clamp((uint) UVSet, 0, DIContext.MeshNumTexCoord - 1);
-        OutUV = DIContext.MeshTexCoordBuffer[VertexIndex * Stride + SelectedUVSet];
-    }
-    else
-    {
-        OutUV = 0.0f;
-    }
-}
-
-void DISkelMesh_IsValidVertex(in FDISkelMeshContext DIContext, int Vertex, out bool IsValid)
-{
-    IsValid = Vertex < (int)DIContext.MeshVertexCount;
-}
-
-void DISkelMesh_GetRandomVertex(in FDISkelMeshContext DIContext, uint Seed1, uint Seed2, uint Seed3, out int OutVertex)
-{
-	float RandT0 = DISKelMesh_Random(Seed1, Seed2, Seed3);
-	OutVertex = (int)min(uint(RandT0 * float(DIContext.MeshVertexCount)), DIContext.MeshVertexCount - 1); // avoid % by using mul/min to Tri = MeshVertexCount
-}
-
-void DISkelMesh_GetVertexCount(in FDISkelMeshContext DIContext, out int VertexCount)
-{
-	VertexCount = (int)DIContext.MeshVertexCount;
-}
-
-void DISkelMesh_IsValidFilteredVertex(in FDISkelMeshContext DIContext, int FilteredIndex, out bool IsValid)
-{
-	IsValid = FilteredIndex < (int)DIContext.MeshNumSamplingRegionVertices;
-}
-
-void DISkelMesh_GetRandomFilteredVertex(in FDISkelMeshContext DIContext, uint Seed1, uint Seed2, uint Seed3, out int OutVertex)
-{
-	// Early out for bad data
-	if ( DIContext.MeshNumSamplingRegionVertices == 0 )
-	{
-		OutVertex = -1;
-		return;
-	}
-
-	float RandT0 = DISKelMesh_Random(Seed1, Seed2, Seed3);
-	int FilteredIndex = (int)min(uint(RandT0 * float(DIContext.MeshNumSamplingRegionVertices)), DIContext.MeshNumSamplingRegionVertices - 1); // avoid % by using mul/min to Tri = MeshVertexCount
-	OutVertex = (int)DIContext.MeshSampleRegionsVertices[FilteredIndex];
-}
-
-void DISkelMesh_GetFilteredVertexCount(in FDISkelMeshContext DIContext, out int VertexCount)
-{
-	VertexCount = (int)DIContext.MeshNumSamplingRegionVertices;
-}
-
-void DISkelMesh_GetFilteredVertexAt(in FDISkelMeshContext DIContext, int FilteredIndex, out int VertexIndex)
-{
-	// Early out for bad data
-	if ( DIContext.MeshNumSamplingRegionVertices == 0 )
-	{
-		VertexIndex = -1;
-		return;
-	}
-
-	FilteredIndex = clamp(FilteredIndex, 0, (int)DIContext.MeshNumSamplingRegionVertices - 1);
-	VertexIndex = (int)DIContext.MeshSampleRegionsVertices[FilteredIndex];
-}
-
-////////////////////////////////////////////////////////////////////////////////////////////////////////////////////////////////////////////////////////////////
-// Filtered Bones / Sockets
-void DISkelMesh_IsValidBone(in FDISkelMeshContext DIContext, in int BoneIndex, out bool IsValid)
-{
-	IsValid = BoneIndex >= 0 && BoneIndex < DIContext.NumBones;
-}
-
-void DISkelMesh_RandomBone(in FDISkelMeshContext DIContext, uint Seed1, uint Seed2, uint Seed3, out int Bone)
-{
-	float RandT0 = DISKelMesh_Random(Seed1, Seed2, Seed3);
-	Bone = min(int(RandT0 * float(DIContext.RandomMaxBone)), DIContext.RandomMaxBone);
-	Bone = Bone != DIContext.ExcludeBoneIndex ? Bone : Bone + 1;
-}
-
-void DISkelMesh_GetBoneCount(in FDISkelMeshContext DIContext, out int Count)
-{
-	Count = DIContext.NumBones;
-}
-
-void DISkelMesh_GetParentBone(in FDISkelMeshContext DIContext, int BoneIndex, out int ParentIndex)
-{
-	ParentIndex = -1;
-	if ( BoneIndex >= 0 && BoneIndex < DIContext.NumBones )
-	{
-		ParentIndex = asint(DIContext.MeshCurrSamplingBonesBuffer[BoneIndex * 3 + 2].w);
-		
-	}
-}
-
-void DISkelMesh_GetFilteredBoneCount(in FDISkelMeshContext DIContext, out int Count)
-{
-    Count = DIContext.NumFilteredBones;
-}
-
-void DISkelMesh_GetFilteredBoneAt(in FDISkelMeshContext DIContext, in int InBoneIndex, out int Bone)
-{
-    int BoneIndex = min(InBoneIndex, DIContext.NumFilteredBones - 1);
-    Bone = DIContext.NumFilteredBones > 0 ? (int)DIContext.FilteredAndUnfilteredBones[BoneIndex] : -1;
-}
-
-void DISkelMesh_RandomFilteredBone(in FDISkelMeshContext DIContext, uint Seed1, uint Seed2, uint Seed3, out int Bone)
-{
-	float RandT0 = DISKelMesh_Random(Seed1, Seed2, Seed3);
-    int FilteredBone = min(int(RandT0 * float(DIContext.NumFilteredBones)), DIContext.NumFilteredBones - 1); // avoid % by using mul/min to Tri = MeshVertexCount
-    Bone = DIContext.NumFilteredBones > 0 ? (int)DIContext.FilteredAndUnfilteredBones[FilteredBone] : -1;
-}
-
-void DISkelMesh_GetUnfilteredBoneCount(in FDISkelMeshContext DIContext, out int Count)
-{
-    Count = DIContext.NumUnfilteredBones;
-}
-
-void DISkelMesh_GetUnfilteredBoneAt(in FDISkelMeshContext DIContext, in int InBoneIndex, out int Bone)
-{
-    int BoneIndex = min(InBoneIndex, DIContext.NumUnfilteredBones - 1);
-    Bone = DIContext.NumFilteredBones > 0 ? (int)DIContext.FilteredAndUnfilteredBones[BoneIndex + DIContext.NumFilteredBones] : BoneIndex;
-}
-
-void DISkelMesh_RandomUnfilteredBone(in FDISkelMeshContext DIContext, uint Seed1, uint Seed2, uint Seed3, out int Bone)
-{
-	float RandT0 = DISKelMesh_Random(Seed1, Seed2, Seed3);
-	if ( DIContext.NumFilteredBones == 0 )
-	{
-		Bone = min(int(RandT0 * float(DIContext.RandomMaxBone)), DIContext.RandomMaxBone);
-		Bone = Bone != DIContext.ExcludeBoneIndex ? Bone : Bone + 1;
-	}
-	else
-	{
-		int UnfilteredBone = min(int(RandT0 * float(DIContext.NumUnfilteredBones)), DIContext.NumUnfilteredBones - 1); // avoid % by using mul/min to Tri = MeshVertexCount
-		Bone = DIContext.NumFilteredBones > 0 ? (int)DIContext.FilteredAndUnfilteredBones[UnfilteredBone + DIContext.NumFilteredBones] : UnfilteredBone;
-	}
-}
-
-void DISkelMesh_GetFilteredSocketCount(in FDISkelMeshContext DIContext, out int Count)
-{
-    Count = DIContext.NumFilteredSockets;
-}
-
-void DISkelMesh_GetFilteredSocketTransform(in FDISkelMeshContext DIContext, in int SocketIndex, in bool bApplyTransform, out float3 OutTranslation, out float4 OutRotation, out float3 OutScale)
-{
-    if ( DIContext.NumFilteredSockets > 0 )
-	{
-		SocketIndex = clamp(SocketIndex, 0, DIContext.NumFilteredSockets - 1);
-		int BufferOffset = (DIContext.FilteredSocketBoneOffset + SocketIndex) * 3;
-
-		OutTranslation = DIContext.MeshCurrSamplingBonesBuffer[BufferOffset + 0].xyz;
-		OutRotation = DIContext.MeshCurrSamplingBonesBuffer[BufferOffset + 1];
-		OutScale = DIContext.MeshCurrSamplingBonesBuffer[BufferOffset + 2].xyz;
-
-		if (bApplyTransform)
-		{
-			OutTranslation = mul(float4(OutTranslation, 1), DIContext.InstanceTransform).xyz;
-			OutRotation = NiagaraGPU_QuatMul(DIContext.InstanceRotation, OutRotation);
-			OutScale = mul(OutScale, (float3x3)DIContext.InstanceTransform);
-		}
-	}
-	else
-	{
-		OutTranslation = float3(0.0f, 0.0f, 0.0f);
-		OutRotation = float4(0.0f, 0.0f, 0.0f, 1.0f);
-		OutScale = float3(0.0f, 0.0f, 0.0f);
-	}
-}
-
-void DISkelMesh_GetFilteredSocketBoneAt(in FDISkelMeshContext DIContext, in int InSocketIndex, out int Bone)
-{
-    Bone = DIContext.FilteredSocketBoneOffset + clamp(InSocketIndex, 0, DIContext.NumFilteredSockets - 1);
-}
-
-void DISkelMesh_RandomFilteredSocket(in FDISkelMeshContext DIContext, uint Seed1, uint Seed2, uint Seed3, out int SocketBone)
-{
-	float RandT0 = DISKelMesh_Random(Seed1, Seed2, Seed3);
-    SocketBone = DIContext.FilteredSocketBoneOffset + min(int(RandT0 * float(DIContext.NumFilteredSockets)), DIContext.NumFilteredSockets - 1); // avoid % by using mul/min to Tri = MeshVertexCount
-}
-
-void DISkelMesh_RandomFilteredSocketOrBone(in FDISkelMeshContext DIContext, uint Seed1, uint Seed2, uint Seed3, out int Bone)
-{
-	float RandT0 = DISKelMesh_Random(Seed1, Seed2, Seed3);
-	int NumFilteredSocketsAndBones = DIContext.NumFilteredBones + DIContext.NumFilteredSockets;
-    int FilteredIndex = min(int(RandT0 * float(NumFilteredSocketsAndBones)), NumFilteredSocketsAndBones - 1);
-	if ( FilteredIndex < DIContext.NumFilteredBones )
-	{
-		Bone = (int)DIContext.FilteredAndUnfilteredBones[FilteredIndex];
-	}
-	else
-	{
-		Bone = DIContext.FilteredSocketBoneOffset + FilteredIndex - DIContext.NumFilteredBones;
-	}
-}
-
-void DISkelMesh_GetFilteredSocketOrBoneCount(in FDISkelMeshContext DIContext, out int Count)
-{
-	Count = DIContext.NumFilteredBones + DIContext.NumFilteredSockets;
-}
-
-void DISkelMesh_GetFilteredSocketOrBoneAt(in FDISkelMeshContext DIContext, int FilteredIndex, out int Bone)
-{
-	int NumFilteredSocketsAndBones = DIContext.NumFilteredBones + DIContext.NumFilteredSockets;
-	FilteredIndex = clamp(FilteredIndex, 0, NumFilteredSocketsAndBones - 1);
-	if ( FilteredIndex < DIContext.NumFilteredBones )
-	{
-		Bone = (int)DIContext.FilteredAndUnfilteredBones[FilteredIndex];
-	}
-	else
-	{
-		Bone = DIContext.FilteredSocketBoneOffset + FilteredIndex - DIContext.NumFilteredBones;
-	}
-}
-
-////////////////////////////////////////////////////////////////////////////////////////////////////////////////////////////////////////////////////////////////
-// Connectivity
-
-void DISkelMesh_GetAdjacentTriangleIndex(in FDISkelMeshContext DIContext, int VertexId, uint AdjacencyIndex, out int TriangleIndex)
-{
-	TriangleIndex = -1;
-
-	if (AdjacencyIndex < DIContext.ConnectivityMaxAdjacentPerVertex)
-	{
-#if DISKELMESH_ADJ_INDEX_FORMAT == DISKELMESH_ADJ_INDEX_FORMAT_FULL
-		const uint ReadIndex = (VertexId * DIContext.ConnectivityMaxAdjacentPerVertex + AdjacencyIndex);
-		if (ReadIndex < DIContext.ConnectivityBufferLength)
-		{
-			TriangleIndex = DIContext.ConnectivityBuffer[ReadIndex];
-		}
-#elif DISKELMESH_ADJ_INDEX_FORMAT == DISKELMESH_ADJ_INDEX_FORMAT_HALF
-		const uint ReadIndex = ((VertexId * DIContext.ConnectivityMaxAdjacentPerVertex + AdjacencyIndex) & ~0x1) >> 1;
-		if (ReadIndex < DIContext.ConnectivityBufferLength)
-		{
-			uint PackedIndices = DIContext.ConnectivityBuffer[ReadIndex];
-			uint UnpackedIndex = uint2(PackedIndices & 0xFFFF, PackedIndices >> 16)[AdjacencyIndex & 1];
-			TriangleIndex = UnpackedIndex == 0xFFFF ? -1 : int(UnpackedIndex);
-		}
-#else
-		#error Unsupported index format for adjacency information
-#endif
-	}
-}
-
-bool IsAdjacent(FDISkelMeshContext DIContext, int VertexIndex, int TriangleIndex)
-{
-	for (uint AdjIt = 0; AdjIt < DIContext.ConnectivityMaxAdjacentPerVertex; ++AdjIt)
-	{
-		int AdjTriangleIndex = -1;
-		DISkelMesh_GetAdjacentTriangleIndex(DIContext, VertexIndex, AdjIt, AdjTriangleIndex);
-		if (AdjTriangleIndex == -1)
-		{
-			break;
-		}
-		if (AdjTriangleIndex == TriangleIndex)
-		{
-			return true;
-		}
-	}
-	return false;
-}
-
-void DISkelMesh_GetTriangleNeighbor(in FDISkelMeshContext DIContext, int TriangleIndex, int EdgeIndex, out int NeighborTriangleIndex, out int NeighborEdgeIndex)
-{
-	NeighborTriangleIndex = -1;
-	NeighborEdgeIndex = -1;
-
-	int3 SourceVertexIndices;
-
-	DISkelMesh_GetTriVertices(DIContext, TriangleIndex, SourceVertexIndices.x, SourceVertexIndices.y, SourceVertexIndices.z);
-
-	int StartSourceIndex = SourceVertexIndices[(EdgeIndex + 1) % 3];
-	int EndSourceIndex = SourceVertexIndices[(EdgeIndex + 2) % 3];
-
-	for (uint SourceAdjIt = 0; SourceAdjIt < DIContext.ConnectivityMaxAdjacentPerVertex; ++SourceAdjIt)
-	{
-		int SourceAdjTriangleIndex = -1;
-		DISkelMesh_GetAdjacentTriangleIndex(DIContext, StartSourceIndex, SourceAdjIt, SourceAdjTriangleIndex);
-
-		if (SourceAdjTriangleIndex == TriangleIndex) // ignore self
-		{
-			continue;
-		}
-		if (SourceAdjTriangleIndex == -1) // stop looking
-		{
-			break;
-		}
-
-		if (IsAdjacent(DIContext, EndSourceIndex, SourceAdjTriangleIndex))
-		{
-			NeighborTriangleIndex = SourceAdjTriangleIndex;
-
-			// now to find the edge index for the neighbor
-			int3 NeighborVertexIndices;
-			DISkelMesh_GetTriVertices(DIContext, NeighborTriangleIndex, NeighborVertexIndices.x, NeighborVertexIndices.y, NeighborVertexIndices.z);
-
-			bool3 NeighborVertMatches = bool3(
-				IsAdjacent(DIContext, NeighborVertexIndices.x, TriangleIndex),
-				IsAdjacent(DIContext, NeighborVertexIndices.y, TriangleIndex),
-				IsAdjacent(DIContext, NeighborVertexIndices.z, TriangleIndex));
-
-			if (all(NeighborVertMatches.xy))
-			{
-				NeighborEdgeIndex = 2;
-			}
-			else if (all(NeighborVertMatches.yz))
-			{
-				NeighborEdgeIndex = 0;
-			}
-			else if (all(NeighborVertMatches.zx))
-			{
-				NeighborEdgeIndex = 1;
-			}
-			break;
-		}
-	}
-}
-
-////////////////////////////////////////////////////////////////////////////////////////////////////////////////////////////////////////////////////////////////
-
-#define DISKELMESH_DECLARE_CONSTANTS(NAME) \
-	Buffer<uint>	MeshIndexBuffer_ ## NAME; \
-	Buffer<float>	MeshVertexBuffer_ ## NAME; \
-	Buffer<uint>	MeshSkinWeightBuffer_ ## NAME; \
-	Buffer<uint>	MeshSkinWeightLookupBuffer_ ## NAME; \
-	Buffer<float4>	MeshCurrBonesBuffer_ ## NAME; \
-	Buffer<float4>	MeshPrevBonesBuffer_ ## NAME; \
-	Buffer<float4>	MeshCurrSamplingBonesBuffer_ ## NAME; \
-	Buffer<float4>	MeshPrevSamplingBonesBuffer_ ## NAME; \
-	Buffer<float4>	MeshTangentBuffer_ ## NAME; \
-	Buffer<float2>	MeshTexCoordBuffer_ ## NAME; \
-	Buffer<float4>	MeshColorBuffer_ ## NAME; \
-	Buffer<uint>	MeshTriangleSamplerProbAliasBuffer_ ## NAME; \
-	uint			MeshNumSamplingRegionTriangles_ ## NAME; \
-	uint			MeshNumSamplingRegionVertices_ ## NAME; \
-	Buffer<uint>	MeshSamplingRegionsProbAliasBuffer_ ## NAME; \
-	Buffer<uint>	MeshSampleRegionsTriangleIndices_ ## NAME; \
-	Buffer<uint>	MeshSampleRegionsVertices_ ## NAME; \
-	Buffer<uint>	MeshTriangleMatricesOffsetBuffer_ ## NAME; \
-	uint			MeshTriangleCount_ ## NAME; \
-	uint			MeshVertexCount_ ## NAME; \
-	uint			MeshWeightStride_ ## NAME; \
-	uint			MeshSkinWeightIndexSize_ ## NAME; \
-	uint			MeshNumTexCoord_ ## NAME; \
-	uint			MeshNumWeights_ ## NAME; \
-	int				NumBones_ ## NAME; \
-	int				NumFilteredBones_ ## NAME; \
-	int				NumUnfilteredBones_ ## NAME; \
-	int				RandomMaxBone_ ## NAME; \
-	int				ExcludeBoneIndex_ ## NAME; \
-	Buffer<uint>	FilteredAndUnfilteredBones_ ## NAME; \
-	int				NumFilteredSockets_ ## NAME; \
-	int				FilteredSocketBoneOffset_ ## NAME; \
-	Buffer<int>		UvMappingBuffer_ ## NAME; \
-	uint			UvMappingBufferLength_ ## NAME; \
-	uint			UvMappingSet_ ## NAME; \
-	Buffer<uint>		ConnectivityBuffer_ ## NAME; \
-	uint			ConnectivityBufferLength_ ## NAME; \
-	uint			ConnectivityMaxAdjacentPerVertex_ ## NAME; \
-	float4x4		InstanceTransform_ ## NAME; \
-	float4x4		InstancePrevTransform_ ## NAME; \
-	float4			InstanceRotation_ ## NAME; \
-	float4			InstancePrevRotation_ ## NAME; \
-	float			InstanceInvDeltaTime_ ## NAME; \
-	uint			EnabledFeatures_ ## NAME; \
-
-#define DISKELMESH_MAKE_CONTEXT(NAME) \
-	FDISkelMeshContext	DIContext; \
-	DIContext.MeshIndexBuffer = MeshIndexBuffer_ ## NAME; \
-	DIContext.MeshVertexBuffer = MeshVertexBuffer_ ## NAME; \
-	DIContext.MeshSkinWeightBuffer = MeshSkinWeightBuffer_ ## NAME; \
-	DIContext.MeshSkinWeightLookupBuffer = MeshSkinWeightLookupBuffer_ ## NAME; \
-	DIContext.MeshCurrBonesBuffer = MeshCurrBonesBuffer_ ## NAME; \
-	DIContext.MeshPrevBonesBuffer = MeshPrevBonesBuffer_ ## NAME; \
-	DIContext.MeshCurrSamplingBonesBuffer = MeshCurrSamplingBonesBuffer_ ## NAME; \
-	DIContext.MeshPrevSamplingBonesBuffer = MeshPrevSamplingBonesBuffer_ ## NAME; \
-	DIContext.MeshTangentBuffer = MeshTangentBuffer_ ## NAME; \
-	DIContext.MeshTexCoordBuffer = MeshTexCoordBuffer_ ## NAME; \
-	DIContext.MeshColorBuffer = MeshColorBuffer_ ## NAME; \
-	DIContext.MeshTriangleSamplerProbAliasBuffer = MeshTriangleSamplerProbAliasBuffer_ ## NAME; \
-	DIContext.MeshNumSamplingRegionTriangles = MeshNumSamplingRegionTriangles_ ## NAME; \
-	DIContext.MeshNumSamplingRegionVertices = MeshNumSamplingRegionVertices_ ## NAME; \
-	DIContext.MeshSamplingRegionsProbAliasBuffer = MeshSamplingRegionsProbAliasBuffer_ ## NAME; \
-	DIContext.MeshSampleRegionsTriangleIndices = MeshSampleRegionsTriangleIndices_ ## NAME; \
-	DIContext.MeshSampleRegionsVertices = MeshSampleRegionsVertices_ ## NAME; \
-	DIContext.MeshTriangleMatricesOffsetBuffer = MeshTriangleMatricesOffsetBuffer_ ## NAME; \
-	DIContext.MeshTriangleCount = MeshTriangleCount_ ## NAME; \
-	DIContext.MeshVertexCount = MeshVertexCount_ ## NAME; \
-	DIContext.MeshNumTexCoord = MeshNumTexCoord_ ## NAME; \
-	DIContext.MeshNumWeights = MeshNumWeights_ ## NAME; \
-	DIContext.NumBones = NumBones_ ## NAME; \
-	DIContext.NumFilteredBones	= NumFilteredBones_ ## NAME; \
-	DIContext.NumUnfilteredBones = NumUnfilteredBones_ ## NAME; \
-	DIContext.RandomMaxBone = RandomMaxBone_ ## NAME; \
-	DIContext.ExcludeBoneIndex = ExcludeBoneIndex_ ## NAME; \
-	DIContext.FilteredAndUnfilteredBones	= FilteredAndUnfilteredBones_ ## NAME; \
-	DIContext.NumFilteredSockets = NumFilteredSockets_ ## NAME; \
-	DIContext.FilteredSocketBoneOffset = FilteredSocketBoneOffset_ ## NAME; \
-	DIContext.UvMappingBuffer = UvMappingBuffer_ ## NAME; \
-	DIContext.UvMappingBufferLength = UvMappingBufferLength_ ## NAME; \
-	DIContext.UvMappingSet = UvMappingSet_ ## NAME; \
-	DIContext.ConnectivityBuffer = ConnectivityBuffer_ ## NAME; \
-	DIContext.ConnectivityBufferLength = ConnectivityBufferLength_ ## NAME; \
-	DIContext.ConnectivityMaxAdjacentPerVertex = ConnectivityMaxAdjacentPerVertex_ ## NAME; \
-	DIContext.MeshWeightStride = MeshWeightStride_ ## NAME; \
-	DIContext.MeshSkinWeightIndexSize = MeshSkinWeightIndexSize_ ## NAME; \
-	DIContext.InstanceTransform = InstanceTransform_ ## NAME; \
-	DIContext.InstancePrevTransform = InstancePrevTransform_ ## NAME; \
-	DIContext.InstanceRotation = InstanceRotation_ ## NAME; \
-	DIContext.InstancePrevRotation = InstancePrevRotation_ ## NAME; \
-	DIContext.InstanceInvDeltaTime = InstanceInvDeltaTime_ ## NAME; \
-	DIContext.UniformTriangleSamplingEnable = EnabledFeatures_ ## NAME & 0x1; \
-	DIContext.UniformSamplingRegionEnabled = EnabledFeatures_ ## NAME & 0x2; \
-	DIContext.UnlimitedBoneInfluences = EnabledFeatures_ ## NAME & 0x4; \
-	DIContext.bHasMeshColors = EnabledFeatures_ ## NAME & 0x8; \
diff --git a/Engine/Plugins/FX/Niagara/Shaders/Private/NiagaraDataInterfaceSkeletalMeshTemplate.ush b/Engine/Plugins/FX/Niagara/Shaders/Private/NiagaraDataInterfaceSkeletalMeshTemplate.ush
new file mode 100644
--- /dev/null
+++ b/Engine/Plugins/FX/Niagara/Shaders/Private/NiagaraDataInterfaceSkeletalMeshTemplate.ush
@@ -0,0 +1,1153 @@
+// Copyright Epic Games, Inc. All Rights Reserved.
+
+////////////////////////////////////////////////////////////////////////////////////////////////////////////////////////////////////////////////////////////////
+// Parameters
+
+Buffer<uint>	{ParameterName}_MeshIndexBuffer;
+Buffer<float>	{ParameterName}_MeshVertexBuffer;
+Buffer<uint>	{ParameterName}_MeshSkinWeightBuffer;
+Buffer<uint>	{ParameterName}_MeshSkinWeightLookupBuffer;
+Buffer<float4>	{ParameterName}_MeshCurrBonesBuffer;
+Buffer<float4>	{ParameterName}_MeshPrevBonesBuffer;
+Buffer<float4>	{ParameterName}_MeshCurrSamplingBonesBuffer;
+Buffer<float4>	{ParameterName}_MeshPrevSamplingBonesBuffer;
+Buffer<float4>	{ParameterName}_MeshTangentBuffer;
+Buffer<float2>	{ParameterName}_MeshTexCoordBuffer;
+Buffer<float4>	{ParameterName}_MeshColorBuffer;
+Buffer<uint>	{ParameterName}_MeshTriangleSamplerProbAliasBuffer;
+uint			{ParameterName}_MeshNumSamplingRegionTriangles;
+uint			{ParameterName}_MeshNumSamplingRegionVertices;
+Buffer<uint>	{ParameterName}_MeshSamplingRegionsProbAliasBuffer;
+Buffer<uint>	{ParameterName}_MeshSampleRegionsTriangleIndices;
+Buffer<uint>	{ParameterName}_MeshSampleRegionsVertices;
+Buffer<uint>	{ParameterName}_MeshTriangleMatricesOffsetBuffer;
+uint			{ParameterName}_MeshTriangleCount;
+uint			{ParameterName}_MeshVertexCount;
+uint			{ParameterName}_MeshWeightStride;
+uint			{ParameterName}_MeshSkinWeightIndexSize;
+uint			{ParameterName}_MeshNumTexCoord;
+uint			{ParameterName}_MeshNumWeights;
+int				{ParameterName}_NumBones;
+int				{ParameterName}_NumFilteredBones;
+int				{ParameterName}_NumUnfilteredBones;
+int				{ParameterName}_RandomMaxBone;
+int				{ParameterName}_ExcludeBoneIndex;
+Buffer<uint>	{ParameterName}_FilteredAndUnfilteredBones;
+int				{ParameterName}_NumFilteredSockets;
+int				{ParameterName}_FilteredSocketBoneOffset;
+Buffer<int>		{ParameterName}_UvMappingBuffer;
+uint			{ParameterName}_UvMappingBufferLength;
+uint			{ParameterName}_UvMappingSet;
+Buffer<uint>	{ParameterName}_ConnectivityBuffer;
+uint			{ParameterName}_ConnectivityBufferLength;
+uint			{ParameterName}_ConnectivityMaxAdjacentPerVertex;
+float4x4		{ParameterName}_InstanceTransform;
+float4x4		{ParameterName}_InstancePrevTransform;
+float4			{ParameterName}_InstanceRotation;
+float4			{ParameterName}_InstancePrevRotation;
+float			{ParameterName}_InstanceInvDeltaTime;
+uint			{ParameterName}_EnabledFeatures;
+
+////////////////////////////////////////////////////////////////////////////////////////////////////////////////////////////////////////////////////////////////
+// Flags test
+bool UniformTriangleSamplingEnable_{ParameterName}() { return ({ParameterName}_EnabledFeatures & 0x1) == 1; }
+bool UniformSamplingRegionEnabled_{ParameterName}() { return ({ParameterName}_EnabledFeatures & 0x2) == 2; }
+bool UnlimitedBoneInfluences_{ParameterName}() { return ({ParameterName}_EnabledFeatures & 0x4) == 4; }
+bool HasMeshColors_{ParameterName}() { return ({ParameterName}_EnabledFeatures & 0x8) == 8; }
+
+////////////////////////////////////////////////////////////////////////////////////////////////////////////////////////////////////////////////////////////////
+// Triangle Sampling
+float3x4 GetPrevBoneSkinningMatrix_{ParameterName}(uint Bone)
+{
+    return float3x4({ParameterName}_MeshPrevBonesBuffer[Bone * 3], {ParameterName}_MeshPrevBonesBuffer[Bone * 3 + 1], {ParameterName}_MeshPrevBonesBuffer[Bone * 3 + 2]);
+}
+
+float3x4 GetPrevSkinningMatrix_{ParameterName}(uint SectionBoneOffset, int4 BlendIndices, float4 BlendWeights)
+{
+	float3x4 Result;
+    Result = GetPrevBoneSkinningMatrix_{ParameterName}(SectionBoneOffset + BlendIndices.x) * BlendWeights.x;
+    Result += GetPrevBoneSkinningMatrix_{ParameterName}(SectionBoneOffset + BlendIndices.y) * BlendWeights.y;
+    Result += GetPrevBoneSkinningMatrix_{ParameterName}(SectionBoneOffset + BlendIndices.z) * BlendWeights.z;
+    Result += GetPrevBoneSkinningMatrix_{ParameterName}(SectionBoneOffset + BlendIndices.w) * BlendWeights.w;
+	return Result;
+}
+
+float3x4 GetCurrBoneSkinningMatrix_{ParameterName}(uint Bone)
+{
+	return float3x4({ParameterName}_MeshCurrBonesBuffer[Bone * 3], {ParameterName}_MeshCurrBonesBuffer[Bone * 3 + 1], {ParameterName}_MeshCurrBonesBuffer[Bone * 3 + 2]);
+}
+
+float3x4 GetCurrSkinningMatrix_{ParameterName}(uint SectionBoneOffset, int4 BlendIndices, float4 BlendWeights)
+{
+	float3x4 Result;
+    Result = GetCurrBoneSkinningMatrix_{ParameterName}(SectionBoneOffset + BlendIndices.x) * BlendWeights.x;
+    Result += GetCurrBoneSkinningMatrix_{ParameterName}(SectionBoneOffset + BlendIndices.y) * BlendWeights.y;
+    Result += GetCurrBoneSkinningMatrix_{ParameterName}(SectionBoneOffset + BlendIndices.z) * BlendWeights.z;
+    Result += GetCurrBoneSkinningMatrix_{ParameterName}(SectionBoneOffset + BlendIndices.w) * BlendWeights.w;
+	return Result;
+}
+
+void IsValidTriCoord_{ParameterName}(in int Triangle, out bool IsValid)
+{
+	IsValid = Triangle >= 0 && Triangle < {ParameterName}_MeshTriangleCount;
+}
+
+void GetVertex_{ParameterName}(uint VertexIndex, out float3 OutPosition, out float3 OutTangentX, out float3 OutTangentY, out float3 OutTangentZ)
+{
+	// Early out for bad data
+	if ( {ParameterName}_MeshVertexCount == 0 )
+	{
+		OutPosition = float3(0.0f, 0.0f, 0.0f);
+		OutTangentX = float3(1.0f, 0.0f, 0.0f);
+		OutTangentY = float3(0.0f, 1.0f, 0.0f);
+		OutTangentZ = float3(0.0f, 0.0f, 1.0f);
+		return;
+	}
+
+	VertexIndex = clamp(VertexIndex, 0, {ParameterName}_MeshVertexCount - 1);
+
+	//-TODO: Use a R32G32B32 format
+	float3 Position = float3({ParameterName}_MeshVertexBuffer[VertexIndex * 3], {ParameterName}_MeshVertexBuffer[VertexIndex * 3 + 1], {ParameterName}_MeshVertexBuffer[VertexIndex * 3 + 2]);
+	float3 TangentX = TangentBias({ParameterName}_MeshTangentBuffer[VertexIndex * 2    ].xyz);
+	float4 TangentZ = TangentBias({ParameterName}_MeshTangentBuffer[VertexIndex * 2 + 1].xyzw);	
+	float3 TangentY = cross(TangentZ.xyz, TangentX.xyz) * TangentZ.w;
+
+	OutPosition = Position;
+	OutTangentX = cross(TangentY, TangentZ.xyz) * TangentZ.w;
+	OutTangentY = TangentY;
+	OutTangentZ = TangentZ.xyz;
+}
+
+void GetVertexData_{ParameterName}(int VertexIndex, out float3 OutPosition, out float3 OutTangentX, out float3 OutTangentY, out float3 OutTangentZ)
+{
+	GetVertex_{ParameterName}(VertexIndex, OutPosition, OutTangentX, OutTangentY, OutTangentZ);
+}
+
+void GetPointOnTriangle_{ParameterName}(uint TriangleIndex, float3 BaryCoord, out float3 OutPosition, out float3 OutTangentX, out float3 OutTangentY, out float3 OutTangentZ)
+{
+	// Early out for bad data
+	if ( {ParameterName}_MeshTriangleCount == 0 )
+	{
+		OutPosition = float3(0.0f, 0.0f, 0.0f);
+		OutTangentX = float3(1.0f, 0.0f, 0.0f);
+		OutTangentY = float3(0.0f, 1.0f, 0.0f);
+		OutTangentZ = float3(0.0f, 0.0f, 1.0f);
+		return;
+	}
+
+	TriangleIndex = clamp(TriangleIndex, 0, {ParameterName}_MeshTriangleCount - 1);
+
+	uint IndexBufferOffset = TriangleIndex * 3;
+	uint VertexIndex0 = {ParameterName}_MeshIndexBuffer[IndexBufferOffset];
+	uint VertexIndex1 = {ParameterName}_MeshIndexBuffer[IndexBufferOffset + 1];
+	uint VertexIndex2 = {ParameterName}_MeshIndexBuffer[IndexBufferOffset + 2];
+
+	float3 Positions[3], TangentsX[3], TangentsY[3], TangentsZ[3];
+	GetVertex_{ParameterName}(VertexIndex0, Positions[0], TangentsX[0], TangentsY[0], TangentsZ[0]);
+	GetVertex_{ParameterName}(VertexIndex1, Positions[1], TangentsX[1], TangentsY[1], TangentsZ[1]);
+	GetVertex_{ParameterName}(VertexIndex2, Positions[2], TangentsX[2], TangentsY[2], TangentsZ[2]);
+	OutPosition = Positions[0] * BaryCoord.x + Positions[1] * BaryCoord.y  + Positions[2] * BaryCoord.z;
+	OutTangentX = TangentsX[0] * BaryCoord.x + TangentsX[1] * BaryCoord.y  + TangentsX[2] * BaryCoord.z;
+	OutTangentY = TangentsY[0] * BaryCoord.x + TangentsY[1] * BaryCoord.y  + TangentsY[2] * BaryCoord.z;
+	OutTangentZ = TangentsZ[0] * BaryCoord.x + TangentsZ[1] * BaryCoord.y  + TangentsZ[2] * BaryCoord.z;
+}
+
+FDISkelMeshSkinnedVertex GetSkinnedVertex_{ParameterName}(uint VertexIndex)
+{
+	// Early out for bad data
+	if ( {ParameterName}_MeshVertexCount == 0 )
+	{
+		FDISkelMeshSkinnedVertex SkinnedVertex;
+		SkinnedVertex.Position = float3(0.0f, 0.0f, 0.0f);
+		SkinnedVertex.PrevPosition = float3(0.0f, 0.0f, 0.0f);
+		SkinnedVertex.TangentX = float3(1.0f, 0.0f, 0.0f);
+		SkinnedVertex.TangentY = float3(0.0f, 1.0f, 0.0f);
+		SkinnedVertex.TangentZ = float3(0.0f, 0.0f, 1.0f);
+		return SkinnedVertex;
+	}
+
+	VertexIndex = clamp(VertexIndex, 0, {ParameterName}_MeshVertexCount - 1);
+
+	//-TODO: Use a R32G32B32 format
+	float3 CurrPosition = float3({ParameterName}_MeshVertexBuffer[VertexIndex * 3], {ParameterName}_MeshVertexBuffer[VertexIndex * 3 + 1], {ParameterName}_MeshVertexBuffer[VertexIndex * 3 + 2]);
+	float3 PrevPosition = CurrPosition;
+
+	float3 TangentX = TangentBias({ParameterName}_MeshTangentBuffer[VertexIndex * 2    ].xyz);
+	float4 TangentZ = TangentBias({ParameterName}_MeshTangentBuffer[VertexIndex * 2 + 1].xyzw);
+	
+	if ( {ParameterName}_MeshNumWeights > 0 )
+	{
+		float3x4 PrevBoneMatrix;
+		float3x4 CurrBoneMatrix;
+
+		// Get the matrix offset for each vertex because BlendIndices are stored relatively to each section start vertex.
+		uint SectionBoneOffset = {ParameterName}_MeshTriangleMatricesOffsetBuffer[VertexIndex];
+
+#if DISKELMESH_BONE_INFLUENCES == DISKELMESH_BONE_INFLUENCES_UNLIMITED
+		BRANCH
+		if (UnlimitedBoneInfluences_{ParameterName}())
+		{
+			uint BlendOffsetCount = {ParameterName}_MeshSkinWeightLookupBuffer[VertexIndex];
+			int NumBoneInfluences = BlendOffsetCount & 0xff;
+			int StreamOffset = BlendOffsetCount >> 8;
+			int WeightsOffset = StreamOffset + ({ParameterName}_MeshSkinWeightIndexSize * NumBoneInfluences);
+
+			PrevBoneMatrix = float3x4(float4(0,0,0,0), float4(0,0,0,0), float4(0,0,0,0));
+			CurrBoneMatrix = float3x4(float4(0,0,0,0), float4(0,0,0,0), float4(0,0,0,0));
+			for (int InfluenceIdx = 0; InfluenceIdx < NumBoneInfluences; InfluenceIdx++)
+			{
+				int BoneIndexOffset = StreamOffset + ({ParameterName}_MeshSkinWeightIndexSize * InfluenceIdx);
+				int BoneIndex = {ParameterName}_MeshSkinWeightBuffer[BoneIndexOffset];
+				if ({ParameterName}_MeshSkinWeightIndexSize > 1)
+				{
+					BoneIndex = {ParameterName}_MeshSkinWeightBuffer[BoneIndexOffset + 1] << 8 | BoneIndex;
+					//@todo-lh: Workaround to fix issue in SPIRVEmitter of DXC; this block must be inside the if branch
+					float BoneWeight = float({ParameterName}_MeshSkinWeightBuffer[WeightsOffset + InfluenceIdx]) / 255.0f;
+
+					PrevBoneMatrix += GetPrevBoneSkinningMatrix_{ParameterName}(SectionBoneOffset + BoneIndex) * BoneWeight;
+					CurrBoneMatrix += GetCurrBoneSkinningMatrix_{ParameterName}(SectionBoneOffset + BoneIndex) * BoneWeight;
+				}
+				else
+				{
+					//@todo-lh: Workaround to fix issue in SPIRVEmitter of DXC; this block must be inside the if branch
+					float BoneWeight = float({ParameterName}_MeshSkinWeightBuffer[WeightsOffset + InfluenceIdx]) / 255.0f;
+
+					PrevBoneMatrix += GetPrevBoneSkinningMatrix_{ParameterName}(SectionBoneOffset + BoneIndex) * BoneWeight;
+					CurrBoneMatrix += GetCurrBoneSkinningMatrix_{ParameterName}(SectionBoneOffset + BoneIndex) * BoneWeight;
+				}
+			}	
+		}
+		else
+#endif
+		{
+			uint WeightOffset = {ParameterName}_MeshNumWeights == 8 ? 2 : 1;
+			int4 BlendIndices = DISkelMesh_UnpackIndices4({ParameterName}_MeshSkinWeightBuffer[VertexIndex * {ParameterName}_MeshWeightStride]);
+	        float4 BlendWeights = DISkelMesh_UnpackWeights4({ParameterName}_MeshSkinWeightBuffer[VertexIndex * {ParameterName}_MeshWeightStride + WeightOffset]);
+			CurrBoneMatrix = GetCurrSkinningMatrix_{ParameterName}(SectionBoneOffset, BlendIndices, BlendWeights);
+			PrevBoneMatrix = GetPrevSkinningMatrix_{ParameterName}(SectionBoneOffset, BlendIndices, BlendWeights);
+#if DISKELMESH_BONE_INFLUENCES != DISKELMESH_BONE_INFLUENCES_MAX4
+			// We are in 8 or unlimited mode so need to check for the extended weights
+			if ({ParameterName}_MeshNumWeights == 8)
+			{
+			    BlendIndices = DISkelMesh_UnpackIndices4({ParameterName}_MeshSkinWeightBuffer[VertexIndex * {ParameterName}_MeshWeightStride + 1]);
+			    BlendWeights = DISkelMesh_UnpackWeights4({ParameterName}_MeshSkinWeightBuffer[VertexIndex * {ParameterName}_MeshWeightStride + WeightOffset + 1]);
+				CurrBoneMatrix += GetCurrSkinningMatrix_{ParameterName}(SectionBoneOffset, BlendIndices, BlendWeights);
+				PrevBoneMatrix += GetPrevSkinningMatrix_{ParameterName}(SectionBoneOffset, BlendIndices, BlendWeights);
+			}
+#endif
+		}
+
+		CurrPosition = mul(CurrBoneMatrix, float4(CurrPosition, 1.0f)).xyz;
+		PrevPosition = mul(PrevBoneMatrix, float4(PrevPosition, 1.0f)).xyz;
+
+		// Not using InverseTranspose of matrices so assuming uniform scaling only (same as SkinCache)
+		TangentX.xyz = mul(CurrBoneMatrix, float4(TangentX.xyz, 0.0f)).xyz;
+		TangentZ.xyz = mul(CurrBoneMatrix, float4(TangentZ.xyz, 0.0f)).xyz;
+	}
+
+	float3 TangentY = cross(TangentZ.xyz, TangentX.xyz) * TangentZ.w;
+
+	FDISkelMeshSkinnedVertex SkinnedVertex;
+	SkinnedVertex.Position = CurrPosition;
+	SkinnedVertex.PrevPosition = PrevPosition;
+	SkinnedVertex.TangentX = cross(TangentY, TangentZ.xyz) * TangentZ.w;
+	SkinnedVertex.TangentY = TangentY;
+	SkinnedVertex.TangentZ = TangentZ.xyz;
+	return SkinnedVertex;
+}
+
+FDISkelMeshSkinnedVertex GetSkinnedPointOnTriangle_{ParameterName}(uint TriangleIndex, float3 BaryCoord)
+{
+	// Early out for bad data
+	if ( {ParameterName}_MeshTriangleCount == 0 )
+	{
+		FDISkelMeshSkinnedVertex SkinnedVertex;
+		SkinnedVertex.Position = float3(0.0f, 0.0f, 0.0f);
+		SkinnedVertex.PrevPosition = float3(0.0f, 0.0f, 0.0f);
+		SkinnedVertex.TangentX = float3(1.0f, 0.0f, 0.0f);
+		SkinnedVertex.TangentY = float3(0.0f, 1.0f, 0.0f);
+		SkinnedVertex.TangentZ = float3(0.0f, 0.0f, 1.0f);
+		return SkinnedVertex;
+	}
+
+	TriangleIndex = clamp(TriangleIndex, 0, {ParameterName}_MeshTriangleCount - 1);
+
+	uint IndexBufferOffset = TriangleIndex * 3;
+	uint VertexIndex0 = {ParameterName}_MeshIndexBuffer[IndexBufferOffset];
+	uint VertexIndex1 = {ParameterName}_MeshIndexBuffer[IndexBufferOffset + 1];
+	uint VertexIndex2 = {ParameterName}_MeshIndexBuffer[IndexBufferOffset + 2];
+
+	FDISkelMeshSkinnedVertex SkinnedVertex0 = GetSkinnedVertex_{ParameterName}(VertexIndex0);
+	FDISkelMeshSkinnedVertex SkinnedVertex1 = GetSkinnedVertex_{ParameterName}(VertexIndex1);
+	FDISkelMeshSkinnedVertex SkinnedVertex2 = GetSkinnedVertex_{ParameterName}(VertexIndex2);
+
+	FDISkelMeshSkinnedVertex FinalVertex;
+	FinalVertex.Position		= (SkinnedVertex0.Position * BaryCoord.x) + (SkinnedVertex1.Position * BaryCoord.y) + (SkinnedVertex2.Position * BaryCoord.z);
+	FinalVertex.PrevPosition	= (SkinnedVertex0.PrevPosition * BaryCoord.x) + (SkinnedVertex1.PrevPosition * BaryCoord.y) + (SkinnedVertex2.PrevPosition * BaryCoord.z);
+	FinalVertex.TangentX		= (SkinnedVertex0.TangentX * BaryCoord.x) + (SkinnedVertex1.TangentX * BaryCoord.y) + (SkinnedVertex2.TangentX * BaryCoord.z);
+	FinalVertex.TangentY		= (SkinnedVertex0.TangentY * BaryCoord.x) + (SkinnedVertex1.TangentY * BaryCoord.y) + (SkinnedVertex2.TangentY * BaryCoord.z);
+	FinalVertex.TangentZ		= (SkinnedVertex0.TangentZ * BaryCoord.x) + (SkinnedVertex1.TangentZ * BaryCoord.y) + (SkinnedVertex2.TangentZ * BaryCoord.z);
+
+	return FinalVertex;
+}
+
+FDISkelMeshSkinnedVertex GetSkinnedPointOnTriangleWS_{ParameterName}(uint TriangleIndex, float3 BaryCoord)
+{
+	FDISkelMeshSkinnedVertex FinalVertex = GetSkinnedPointOnTriangle_{ParameterName}(TriangleIndex, BaryCoord);
+	FinalVertex.Position		= mul(float4(FinalVertex.Position, 1.0f), {ParameterName}_InstanceTransform).xyz;
+	FinalVertex.PrevPosition	= mul(float4(FinalVertex.PrevPosition, 1.0f), {ParameterName}_InstancePrevTransform).xyz;
+	FinalVertex.TangentX		= mul(float4(FinalVertex.TangentX, 0.0f), {ParameterName}_InstanceTransform).xyz;
+	FinalVertex.TangentY		= mul(float4(FinalVertex.TangentY, 0.0f), {ParameterName}_InstanceTransform).xyz;
+	FinalVertex.TangentZ		= mul(float4(FinalVertex.TangentZ, 0.0f), {ParameterName}_InstanceTransform).xyz;
+
+	return FinalVertex;
+}
+
+void RandomTriCoord_{ParameterName}(uint Seed1, uint Seed2, uint Seed3, out int OutTriangle, out float3 OutBaryCoord)
+{
+	// Early out for bad data
+	if ({ParameterName}_MeshTriangleCount == 0)
+	{
+		OutTriangle = -1;
+		OutBaryCoord = 1.0f / 3.0f;
+		return;
+	}
+
+	float RandT0 = DISKelMesh_Random(Seed1, Seed2, Seed3);
+
+	[branch]
+	if ( !UniformTriangleSamplingEnable_{ParameterName}() )
+	{
+		// Uniform triangle id selection
+		OutTriangle = min(uint(RandT0 * float({ParameterName}_MeshTriangleCount)), {ParameterName}_MeshTriangleCount - 1); // avoid % by using mul/min to Tri = MeshTriangleCountName
+	}
+	else
+	{
+		// Uniform area weighted position selection (using alias method from Alias method from FWeightedRandomSampler)
+		uint TriangleIndex = min(uint(RandT0*float({ParameterName}_MeshTriangleCount)), {ParameterName}_MeshTriangleCount - 1);
+
+#if DISKELMESH_PROBALIAS_FORMAT == DISKELMESH_PROBALIAS_FORMAT_64
+		float TriangleProbability = asfloat({ParameterName}_MeshTriangleSamplerProbAliasBuffer[(TriangleIndex * 2) + 0]);
+		uint TriangleAlias = {ParameterName}_MeshTriangleSamplerProbAliasBuffer[(TriangleIndex * 2) + 1];
+#elif DISKELMESH_PROBALIAS_FORMAT == DISKELMESH_PROBALIAS_FORMAT_24_8
+		uint TriangleProbAlias = {ParameterName}_MeshTriangleSamplerProbAliasBuffer[TriangleIndex];
+		float TriangleProbability = float(TriangleProbAlias & 0xff) / 255.0f;
+		uint TriangleAlias = TriangleProbAlias >> 8;
+#elif DISKELMESH_PROBALIAS_FORMAT == DISKELMESH_PROBALIAS_FORMAT_23_9
+		uint TriangleProbAlias = {ParameterName}_MeshTriangleSamplerProbAliasBuffer[TriangleIndex];
+		float TriangleProbability = float(TriangleProbAlias & 0x1ff) / 511.0f;
+		uint TriangleAlias = TriangleProbAlias >> 9;
+#endif
+
+		// Alias check
+		float RandT1 = DISKelMesh_Random(Seed1, Seed2, Seed3);
+		if( RandT1 > TriangleProbability )
+		{
+			TriangleIndex = TriangleAlias;
+		}
+		OutTriangle = int(TriangleIndex);
+	}
+
+	OutBaryCoord = DISKelMesh_RandomBarycentricCoord(Seed1, Seed2, Seed3);
+}
+
+void GetSkinnedTriangleDataWS_{ParameterName}(in uint TriangleIndex, in float3 BaryCoord, out float3 OutPosition, out float3 OutVelocity, out float3 OutNormal, out float3 OutBinormal, out float3 OutTangent)
+{
+    FDISkelMeshSkinnedVertex SkinnedVertex = GetSkinnedPointOnTriangleWS_{ParameterName}(TriangleIndex, BaryCoord);
+	OutPosition	= SkinnedVertex.Position;
+	OutVelocity	= (SkinnedVertex.Position - SkinnedVertex.PrevPosition) * {ParameterName}_InstanceInvDeltaTime;
+	OutNormal	= normalize(SkinnedVertex.TangentZ);
+	OutBinormal	= normalize(SkinnedVertex.TangentY);
+	OutTangent	= normalize(SkinnedVertex.TangentX);
+}
+
+void GetSkinnedTriangleData_{ParameterName}(in uint TriangleIndex, in float3 BaryCoord, out float3 OutPosition, out float3 OutVelocity, out float3 OutNormal, out float3 OutBinormal, out float3 OutTangent)
+{
+    FDISkelMeshSkinnedVertex SkinnedVertex = GetSkinnedPointOnTriangle_{ParameterName}(TriangleIndex, BaryCoord);
+	OutPosition	= SkinnedVertex.Position;
+	OutVelocity	= (SkinnedVertex.Position - SkinnedVertex.PrevPosition) * {ParameterName}_InstanceInvDeltaTime;
+	OutNormal	= normalize(SkinnedVertex.TangentZ);
+	OutBinormal	= normalize(SkinnedVertex.TangentY);
+	OutTangent	= normalize(SkinnedVertex.TangentX);
+}
+
+void GetSkinnedTriangleDataInterpolatedWS_{ParameterName}(in uint TriangleIndex, in float3 BaryCoord, in float Interp, out float3 OutPosition, out float3 OutVelocity, out float3 OutNormal, out float3 OutBinormal, out float3 OutTangent)
+{
+    FDISkelMeshSkinnedVertex SkinnedVertex = GetSkinnedPointOnTriangleWS_{ParameterName}(TriangleIndex, BaryCoord);
+	OutPosition	= lerp(SkinnedVertex.PrevPosition, SkinnedVertex.Position, Interp);
+	OutVelocity	= (SkinnedVertex.Position - SkinnedVertex.PrevPosition) * {ParameterName}_InstanceInvDeltaTime;
+	OutNormal	= normalize(SkinnedVertex.TangentZ);
+	OutBinormal	= normalize(SkinnedVertex.TangentY);
+	OutTangent	= normalize(SkinnedVertex.TangentX);
+}
+
+void GetSkinnedTriangleDataInterpolated_{ParameterName}(in uint TriangleIndex, in float3 BaryCoord, in float Interp, out float3 OutPosition, out float3 OutVelocity, out float3 OutNormal, out float3 OutBinormal, out float3 OutTangent)
+{
+    FDISkelMeshSkinnedVertex SkinnedVertex = GetSkinnedPointOnTriangle_{ParameterName}(TriangleIndex, BaryCoord);
+	OutPosition	= lerp(SkinnedVertex.PrevPosition, SkinnedVertex.Position, Interp);
+	OutVelocity	= (SkinnedVertex.Position - SkinnedVertex.PrevPosition) * {ParameterName}_InstanceInvDeltaTime;
+	OutNormal	= normalize(SkinnedVertex.TangentZ);
+	OutBinormal	= normalize(SkinnedVertex.TangentY);
+	OutTangent	= normalize(SkinnedVertex.TangentX);
+}
+
+void GetTriUV_{ParameterName}(in uint TriangleIndex, in float3 BaryCoord, in int UVSet, out float2 OutUV)
+{
+    if ({ParameterName}_MeshNumTexCoord > 0)
+	{
+        uint IndexBufferOffset = TriangleIndex * 3;
+		uint VertexIndex0 = {ParameterName}_MeshIndexBuffer[IndexBufferOffset  ];
+		uint VertexIndex1 = {ParameterName}_MeshIndexBuffer[IndexBufferOffset+1];
+		uint VertexIndex2 = {ParameterName}_MeshIndexBuffer[IndexBufferOffset+2];
+
+        uint Stride = {ParameterName}_MeshNumTexCoord;
+        uint SelectedUVSet = clamp((uint) UVSet, 0, {ParameterName}_MeshNumTexCoord - 1);
+		float2 UV0 = {ParameterName}_MeshTexCoordBuffer[VertexIndex0 * Stride + SelectedUVSet];
+		float2 UV1 = {ParameterName}_MeshTexCoordBuffer[VertexIndex1 * Stride + SelectedUVSet];
+		float2 UV2 = {ParameterName}_MeshTexCoordBuffer[VertexIndex2 * Stride + SelectedUVSet];
+
+		OutUV = UV0 * BaryCoord.x + UV1 * BaryCoord.y + UV2 * BaryCoord.z;
+	}
+	else	
+	{
+		OutUV = 0.0f;
+	}
+}
+
+void RandomTriangle_{ParameterName}(uint Seed1, uint Seed2, uint Seed3, out int OutTriangle, out float3 OutBaryCoord)
+{
+	float RandT0 = DISKelMesh_Random(Seed1, Seed2, Seed3);
+	int TriangleIndex = min(int(RandT0 * float({ParameterName}_MeshTriangleCount)), (int){ParameterName}_MeshTriangleCount - 1); // avoid % by using mul/min to Tri = MeshTriangleCountName
+
+	[branch]
+	if ( UniformTriangleSamplingEnable_{ParameterName}() )
+	{
+		// Uniform area weighted position selection (using alias method from Alias method from FWeightedRandomSampler)
+#if DISKELMESH_PROBALIAS_FORMAT == DISKELMESH_PROBALIAS_FORMAT_64
+		float TriangleProbability = asfloat({ParameterName}_MeshTriangleSamplerProbAliasBuffer[(TriangleIndex * 2) + 0]);
+		uint TriangleAlias = {ParameterName}_MeshTriangleSamplerProbAliasBuffer[(TriangleIndex * 2) + 1];
+#elif DISKELMESH_PROBALIAS_FORMAT == DISKELMESH_PROBALIAS_FORMAT_24_8
+		uint TriangleProbAlias = {ParameterName}_MeshTriangleSamplerProbAliasBuffer[TriangleIndex];
+		float TriangleProbability = float(TriangleProbAlias & 0xff) / 255.0f;
+		uint TriangleAlias = TriangleProbAlias >> 8;
+#elif DISKELMESH_PROBALIAS_FORMAT == DISKELMESH_PROBALIAS_FORMAT_23_9
+		uint TriangleProbAlias = {ParameterName}_MeshTriangleSamplerProbAliasBuffer[TriangleIndex];
+		float TriangleProbability = float(TriangleProbAlias & 0x1ff) / 511.0f;
+		uint TriangleAlias = TriangleProbAlias >> 9;
+#endif
+
+		// Alias check
+		float RandT1 = DISKelMesh_Random(Seed1, Seed2, Seed3);
+		if( RandT1 > TriangleProbability )
+		{
+			TriangleIndex = int(TriangleAlias);
+		}
+	}
+
+	OutTriangle = TriangleIndex;
+	OutBaryCoord = DISKelMesh_RandomBarycentricCoord(Seed1, Seed2, Seed3);
+}
+
+void GetTriangleCount_{ParameterName}(out int OutCount)
+{
+	OutCount = {ParameterName}_MeshTriangleCount;
+}
+
+void RandomFilteredTriangle_{ParameterName}(uint Seed1, uint Seed2, uint Seed3, out int OutTriangle, out float3 OutBaryCoord)
+{
+	// Early out for bad data
+	if ({ParameterName}_MeshNumSamplingRegionTriangles == 0)
+	{
+		OutTriangle = -1;
+		OutBaryCoord = 1.0f / 3.0f;
+		return;
+	}
+
+	float RandT0 = DISKelMesh_Random(Seed1, Seed2, Seed3);
+	uint RegionTriangle = min(uint(RandT0 * float({ParameterName}_MeshNumSamplingRegionTriangles)), {ParameterName}_MeshNumSamplingRegionTriangles - 1); // avoid % by using mul/min to Tri = MeshTriangleCountName
+
+	[branch]
+	if ( UniformSamplingRegionEnabled_{ParameterName}() )
+	{
+#if DISKELMESH_PROBALIAS_FORMAT == DISKELMESH_PROBALIAS_FORMAT_64
+		float TriangleProbability = asfloat({ParameterName}_MeshSamplingRegionsProbAliasBuffer[(RegionTriangle * 2) + 0]);
+		uint TriangleAlias = {ParameterName}_MeshSamplingRegionsProbAliasBuffer[(RegionTriangle * 2) + 1];
+#elif DISKELMESH_PROBALIAS_FORMAT == DISKELMESH_PROBALIAS_FORMAT_24_8
+		uint TriangleProbAlias = {ParameterName}_MeshSamplingRegionsProbAliasBuffer[RegionTriangle];
+		float TriangleProbability = float(TriangleProbAlias & 0xff) / 255.0f;
+		uint TriangleAlias = TriangleProbAlias >> 8;
+#elif DISKELMESH_PROBALIAS_FORMAT == DISKELMESH_PROBALIAS_FORMAT_23_9
+		uint TriangleProbAlias = {ParameterName}_MeshSamplingRegionsProbAliasBuffer[RegionTriangle];
+		float TriangleProbability = float(TriangleProbAlias & 0x1ff) / 511.0f;
+		uint TriangleAlias = TriangleProbAlias >> 9;
+#endif
+		float RandT1 = DISKelMesh_Random(Seed1, Seed2, Seed3);
+		if ( RandT1 > TriangleProbability )
+		{
+			RegionTriangle = TriangleAlias;
+		}
+	}
+	OutTriangle = (int){ParameterName}_MeshSampleRegionsTriangleIndices[RegionTriangle];
+	OutBaryCoord = DISKelMesh_RandomBarycentricCoord(Seed1, Seed2, Seed3);
+}
+
+void GetFilteredTriangleCount_{ParameterName}(out int OutCount)
+{
+	OutCount = {ParameterName}_MeshNumSamplingRegionTriangles;
+}
+
+void GetFilteredTriangleAt_{ParameterName}(int FilteredIndex, out int OutTriangle)
+{
+	if ( {ParameterName}_MeshNumSamplingRegionTriangles > 0 )
+	{
+		FilteredIndex = clamp(FilteredIndex, 0, (int){ParameterName}_MeshNumSamplingRegionTriangles - 1);
+		OutTriangle = (int){ParameterName}_MeshSampleRegionsTriangleIndices[FilteredIndex];
+	}
+	else
+	{
+		OutTriangle = 0;
+	}
+}
+
+////////////////////////////////////////////////////////////////////////////////////////////////////////////////////////////////////////////////////////////////
+// Uv Mapping
+
+// Evaluates QuadTree for a single point
+bool QuadTreeTriangleOverlap_{ParameterName}(float2 Uv, float Tolerance, int TriangleIndex, out float3 OutBaryCoord)
+{
+	if ({ParameterName}_MeshNumTexCoord > 0)
+	{
+        uint IndexBufferOffset = TriangleIndex * 3;
+		uint VertexIndex0 = {ParameterName}_MeshIndexBuffer[IndexBufferOffset  ];
+		uint VertexIndex1 = {ParameterName}_MeshIndexBuffer[IndexBufferOffset+1];
+		uint VertexIndex2 = {ParameterName}_MeshIndexBuffer[IndexBufferOffset+2];
+
+		uint Stride = {ParameterName}_MeshNumTexCoord;
+		float2 UV0 = {ParameterName}_MeshTexCoordBuffer[VertexIndex0 * Stride + {ParameterName}_UvMappingSet];
+		float2 UV1 = {ParameterName}_MeshTexCoordBuffer[VertexIndex1 * Stride + {ParameterName}_UvMappingSet];
+		float2 UV2 = {ParameterName}_MeshTexCoordBuffer[VertexIndex2 * Stride + {ParameterName}_UvMappingSet];
+
+		OutBaryCoord = GetBaryCentric2D(Uv, UV0, UV1, UV2);
+
+		if (all(OutBaryCoord > -Tolerance) && all(OutBaryCoord < (1.0f + Tolerance)))
+		{
+			return true;
+		}
+	}
+
+	return false;
+}
+
+bool QuadTreeNodeContains_{ParameterName}(int NodeBufferOffset, float2 Uv, float Tolerance, out int OutTriangleIndex, out float3 OutBaryCoord)
+{
+	if (NodeBufferOffset == -1)
+	{
+		return false;
+	}
+
+	const int TriangleCount = {ParameterName}_UvMappingBuffer[NodeBufferOffset + DISkelMesh_ElementCountOffset];
+	for (int TriangleIt = 0; TriangleIt < TriangleCount; ++TriangleIt)
+	{
+		int TriangleIndex = {ParameterName}_UvMappingBuffer[NodeBufferOffset + DISkelMesh_TriangleIndexOffset + TriangleIt];
+		if (QuadTreeTriangleOverlap_{ParameterName}(Uv, Tolerance, TriangleIndex, OutBaryCoord))
+		{
+			OutTriangleIndex = TriangleIndex;
+			return true;
+		}
+	}
+
+	return false;
+}
+
+// Evaluates QuadTree for an AABB region
+bool QuadTreeTriangleAabbOverlap_{ParameterName}(float2 UvRef, float2 UvExtentScale, float2 UvExtentBias, int TriangleIndex, out float3 OutBaryCoord)
+{
+	if ({ParameterName}_MeshNumTexCoord > 0)
+	{
+        uint IndexBufferOffset = TriangleIndex * 3;
+		uint VertexIndex0 = {ParameterName}_MeshIndexBuffer[IndexBufferOffset  ];
+		uint VertexIndex1 = {ParameterName}_MeshIndexBuffer[IndexBufferOffset+1];
+		uint VertexIndex2 = {ParameterName}_MeshIndexBuffer[IndexBufferOffset+2];
+
+		uint Stride = {ParameterName}_MeshNumTexCoord;
+		float2 UV0 = {ParameterName}_MeshTexCoordBuffer[VertexIndex0 * Stride + {ParameterName}_UvMappingSet];
+		float2 UV1 = {ParameterName}_MeshTexCoordBuffer[VertexIndex1 * Stride + {ParameterName}_UvMappingSet];
+		float2 UV2 = {ParameterName}_MeshTexCoordBuffer[VertexIndex2 * Stride + {ParameterName}_UvMappingSet];
+
+		// evaluate whether our triangle intersects with our AABB, if it does, then generate the barycentric coordinate for the center of the AABB
+		const bool Intersection = DISkelMesh_NormalizedAabbTriangleIntersection(
+				UV0 * UvExtentScale + UvExtentBias,
+				UV1 * UvExtentScale + UvExtentBias,
+				UV2 * UvExtentScale + UvExtentBias);
+
+		if (Intersection)
+		{
+			OutBaryCoord = GetClosestBaryCentric2D(UvRef, UV0, UV1, UV2);
+			return true;
+		}
+	}
+
+	return false;
+}
+
+bool QuadTreeNodeContains_{ParameterName}Aabb(int NodeBufferOffset, float2 UvMin, float2 UvMax, out int OutTriangleIndex, out float3 OutBaryCoord)
+{
+	if (NodeBufferOffset == -1)
+	{
+		return false;
+	}
+
+	const int TriangleCount = {ParameterName}_UvMappingBuffer[NodeBufferOffset + DISkelMesh_ElementCountOffset];
+	if (TriangleCount == 0)
+	{
+		return false;
+	}
+
+	const float2 UvRef = 0.5f * (UvMin + UvMax);
+	const float2 UvExtentScale = 1.0f / (UvMax - UvMin);
+	const float2 UvExtentBias = 1.0f - UvMax * UvExtentScale;
+
+	for (int TriangleIt = 0; TriangleIt < TriangleCount; ++TriangleIt)
+	{
+		float3 BaryCoord;
+		int TriangleIndex = {ParameterName}_UvMappingBuffer[NodeBufferOffset + DISkelMesh_TriangleIndexOffset + TriangleIt];
+		if (QuadTreeTriangleAabbOverlap_{ParameterName}(UvRef, UvExtentScale, UvExtentBias, TriangleIndex, BaryCoord))
+		{
+			OutTriangleIndex = TriangleIndex;
+			OutBaryCoord = BaryCoord;
+			return true;
+		}
+	}
+
+	return false;
+}
+
+void GetTriangleCoordAtUV_{ParameterName}(bool Enabled, float2 InUv, float InTolerance, out int OutTriangle, out float3 OutBaryCoord, out bool OutIsValid)
+{
+	OutTriangle = -1;
+	OutBaryCoord = float3(0.0f, 0.0f, 0.0f);
+	OutIsValid = false;
+
+	if (!Enabled || {ParameterName}_UvMappingBufferLength < DISkelMesh_MinValidUvMappingBufferSize)
+	{
+		return;
+	}
+
+	int NodeBufferOffset = 0;
+
+	float2 MidPoint = float2(0.5f, 0.5f);
+	float2 Extent = float2(0.5f, 0.5f);
+
+	while (NodeBufferOffset != -1)
+	{
+		if (QuadTreeNodeContains_{ParameterName}(NodeBufferOffset, InUv, InTolerance, OutTriangle, OutBaryCoord))
+		{
+			OutIsValid = true;
+			return;
+		}
+
+		Extent *= 0.5f;
+
+		const bool LessX = InUv.x < MidPoint.x;
+		const bool LessY = InUv.y < MidPoint.y;
+
+		MidPoint += Extent * float2(LessX ? -1.0f : 1.0f, LessY ? -1.0f : 1.0f);
+		int ChildIndex = LessX ? (LessY ? 2 : 0) : (LessY ? 3 : 1);
+
+		NodeBufferOffset = {ParameterName}_UvMappingBuffer[NodeBufferOffset + ChildIndex];
+	}
+}
+
+void GetTriangleCoordInAabb_{ParameterName}(bool Enabled, float2 InUvMin, float2 InUvMax, out int OutTriangle, out float3 OutBaryCoord, out bool OutIsValid)
+{
+	OutTriangle = -1;
+	OutBaryCoord = float3(0.0f, 0.0f, 0.0f);
+	OutIsValid = false;
+
+	if (!Enabled || {ParameterName}_UvMappingBufferLength < DISkelMesh_MinValidUvMappingBufferSize)
+	{
+		return;
+	}
+
+	int NodeBufferOffset = 0;
+
+	const float2 UvRegionCenter = 0.5f * (InUvMin + InUvMax);
+	float2 MidPoint = float2(0.5f, 0.5f);
+	float2 Extent = float2(0.5f, 0.5f);
+
+	while (NodeBufferOffset != -1)
+	{
+		if (QuadTreeNodeContains_{ParameterName}Aabb(NodeBufferOffset, InUvMin, InUvMax, OutTriangle, OutBaryCoord))
+		{
+			OutIsValid = true;
+			return;
+		}
+
+		Extent *= 0.5f;
+
+		const bool LessX = UvRegionCenter.x < MidPoint.x;
+		const bool LessY = UvRegionCenter.y < MidPoint.y;
+
+		MidPoint += Extent * float2(LessX ? -1.0f : 1.0f, LessY ? -1.0f : 1.0f);
+		int ChildIndex = LessX ? (LessY ? 2 : 0) : (LessY ? 3 : 1);
+
+		NodeBufferOffset = {ParameterName}_UvMappingBuffer[NodeBufferOffset + ChildIndex];
+	}
+}
+
+////////////////////////////////////////////////////////////////////////////////////////////////////////////////////////////////////////////////////////////////
+// Bone Sampling
+
+void GetSkinnedBoneCommon_{ParameterName}(in int Bone, in float Interp, in bool bWorldSpace, out float3 OutPosition, out float4 OutRotation, out float3 OutScale, out float3 OutVelocity)
+{
+	// Early out for bad data
+	if ( {ParameterName}_NumBones + {ParameterName}_NumFilteredSockets == 0 )
+	{
+		OutPosition = float3(0.0f, 0.0f, 0.0f);
+		OutRotation = float4(0.0f, 0.0f, 0.0f, 1.0f);
+		OutScale	= float3(0.0f, 0.0f, 0.0f);
+		OutVelocity = float3(0.0f, 0.0f, 0.0f);
+		return;
+	}
+
+	Bone = clamp(Bone, 0, {ParameterName}_NumBones + {ParameterName}_NumFilteredSockets - 1);
+
+    float3 PrevPosition = {ParameterName}_MeshPrevSamplingBonesBuffer[Bone * 3].xyz;
+    float4 PrevRotation = {ParameterName}_MeshPrevSamplingBonesBuffer[Bone * 3 + 1];
+	float3 PrevScale	= {ParameterName}_MeshPrevSamplingBonesBuffer[Bone * 3 + 2].xyz;
+    float3 CurrPosition = {ParameterName}_MeshCurrSamplingBonesBuffer[Bone * 3].xyz;
+    float4 CurrRotation	= {ParameterName}_MeshCurrSamplingBonesBuffer[Bone * 3 + 1];
+	float3 CurrScale	= {ParameterName}_MeshCurrSamplingBonesBuffer[Bone * 3 + 2].xyz;
+
+    if (bWorldSpace)
+    {
+        PrevPosition = mul(float4(PrevPosition, 1), {ParameterName}_InstancePrevTransform).xyz;
+        CurrPosition = mul(float4(CurrPosition, 1), {ParameterName}_InstanceTransform).xyz;
+		PrevRotation = NiagaraGPU_QuatMul({ParameterName}_InstancePrevRotation, PrevRotation);
+		CurrRotation = NiagaraGPU_QuatMul({ParameterName}_InstanceRotation, CurrRotation);
+		PrevScale = mul(PrevScale, (float3x3){ParameterName}_InstancePrevTransform);
+		CurrScale = mul(CurrScale, (float3x3){ParameterName}_InstanceTransform);
+    }
+    CurrPosition = lerp(PrevPosition, CurrPosition, Interp);
+	CurrScale	 = lerp(PrevScale, CurrScale, Interp);
+
+    OutPosition = CurrPosition;
+    OutRotation = DISKelMesh_QuatSlerp(PrevRotation, CurrRotation, Interp);
+	OutScale	= CurrScale;
+    OutVelocity = (CurrPosition - PrevPosition) * {ParameterName}_InstanceInvDeltaTime;
+}
+
+void GetSkinnedBoneData_{ParameterName}(in int Bone, out float3 OutPosition, out float4 OutRotation, out float3 OutScale, out float3 OutVelocity)
+{
+    GetSkinnedBoneCommon_{ParameterName}(Bone, 1.0f, false, OutPosition, OutRotation, OutScale, OutVelocity);
+}
+
+void GetSkinnedBoneDataInterpolated_{ParameterName}(in int Bone, in float Interp, out float3 OutPosition, out float4 OutRotation, out float3 OutScale, out float3 OutVelocity)
+{
+    GetSkinnedBoneCommon_{ParameterName}(Bone, Interp, false, OutPosition, OutRotation, OutScale, OutVelocity);
+}
+
+void GetSkinnedBoneDataWS_{ParameterName}(in int Bone, out float3 OutPosition, out float4 OutRotation, out float3 OutScale, out float3 OutVelocity)
+{
+    GetSkinnedBoneCommon_{ParameterName}(Bone, 1.0f, true, OutPosition, OutRotation, OutScale, OutVelocity);
+}
+
+void GetSkinnedBoneDataWSInterpolated_{ParameterName}(in int Bone, in float Interp, out float3 OutPosition, out float4 OutRotation, out float3 OutScale, out float3 OutVelocity)
+{
+    GetSkinnedBoneCommon_{ParameterName}(Bone, Interp, true, OutPosition, OutRotation, OutScale, OutVelocity);
+}
+
+////////////////////////////////////////////////////////////////////////////////////////////////////////////////////////////////////////////////////////////////
+// Vertex Sampling
+
+void GetTriVertices_{ParameterName}(in uint TriangleIndex, out int VertexIndex0, out int VertexIndex1, out int VertexIndex2)
+{
+	// Early out for bad data
+	if ( {ParameterName}_MeshTriangleCount == 0 )
+	{
+		VertexIndex0 = 0;
+		VertexIndex1 = 0;
+		VertexIndex2 = 0;
+		return;
+	}
+
+	TriangleIndex = clamp(TriangleIndex, 0, {ParameterName}_MeshTriangleCount - 1);
+    uint IndexBufferOffset = TriangleIndex * 3;
+    VertexIndex0 = {ParameterName}_MeshIndexBuffer[IndexBufferOffset];
+    VertexIndex1 = {ParameterName}_MeshIndexBuffer[IndexBufferOffset + 1];
+    VertexIndex2 = {ParameterName}_MeshIndexBuffer[IndexBufferOffset + 2];
+}
+
+void GetTriCoordVertices_{ParameterName}(in int TriangleIndex, out int VertexIndex0, out int VertexIndex1, out int VertexIndex2)
+{
+	GetTriVertices_{ParameterName}(TriangleIndex, VertexIndex0, VertexIndex1, VertexIndex2);
+}
+
+void GetTriColor_{ParameterName}(in uint TriangleIndex, in float3 BaryCoord, out float4 OutColor)
+{
+	// Early out for bad data
+	if ( HasMeshColors_{ParameterName}() == false )
+	{
+		OutColor = float4(1, 1, 1, 1);
+		return;
+	}
+	TriangleIndex = clamp(TriangleIndex, 0, {ParameterName}_MeshTriangleCount - 1);
+
+    uint IndexBufferOffset = TriangleIndex * 3;
+    uint VertexIndex0 = {ParameterName}_MeshIndexBuffer[IndexBufferOffset];
+    uint VertexIndex1 = {ParameterName}_MeshIndexBuffer[IndexBufferOffset + 1];
+    uint VertexIndex2 = {ParameterName}_MeshIndexBuffer[IndexBufferOffset + 2];
+    OutColor = ({ParameterName}_MeshColorBuffer[VertexIndex0] FMANUALFETCH_COLOR_COMPONENT_SWIZZLE * BaryCoord.x) + ({ParameterName}_MeshColorBuffer[VertexIndex1] FMANUALFETCH_COLOR_COMPONENT_SWIZZLE * BaryCoord.y) + ({ParameterName}_MeshColorBuffer[VertexIndex2] FMANUALFETCH_COLOR_COMPONENT_SWIZZLE * BaryCoord.z);
+}
+
+void GetSkinnedVertexData_{ParameterName}(in int VertexIndex, out float3 OutPosition, out float3 OutVelocity, out float3 OutNormal, out float3 OutBinormal, out float3 OutTangent)
+{
+    FDISkelMeshSkinnedVertex SkinnedVertex = GetSkinnedVertex_{ParameterName}(VertexIndex);
+    OutPosition = SkinnedVertex.Position;
+    OutVelocity = (SkinnedVertex.Position - SkinnedVertex.PrevPosition) * {ParameterName}_InstanceInvDeltaTime;
+	OutNormal = SkinnedVertex.TangentZ;
+	OutBinormal = SkinnedVertex.TangentY;
+	OutTangent = SkinnedVertex.TangentX;
+}
+
+void GetSkinnedVertexDataWS_{ParameterName}(in int VertexIndex, out float3 OutPosition, out float3 OutVelocity, out float3 OutNormal, out float3 OutBinormal, out float3 OutTangent)
+{
+    FDISkelMeshSkinnedVertex SkinnedVertex = GetSkinnedVertex_{ParameterName}(VertexIndex);
+    OutPosition = mul(float4(SkinnedVertex.Position, 1.0f), {ParameterName}_InstanceTransform).xyz;
+    float3 PrevPosition = mul(float4(SkinnedVertex.PrevPosition, 1.0f), {ParameterName}_InstancePrevTransform).xyz;
+    OutVelocity = (SkinnedVertex.Position - SkinnedVertex.PrevPosition) * {ParameterName}_InstanceInvDeltaTime;
+    OutNormal   = normalize(mul(float4(SkinnedVertex.TangentZ, 0.0f), {ParameterName}_InstanceTransform).xyz);
+    OutBinormal = normalize(mul(float4(SkinnedVertex.TangentY, 0.0f), {ParameterName}_InstanceTransform).xyz);
+    OutTangent  = normalize(mul(float4(SkinnedVertex.TangentX, 0.0f), {ParameterName}_InstanceTransform).xyz);
+}
+
+void GetVertexColor_{ParameterName}(in int VertexIndex, out float4 OutColor)
+{
+    OutColor = HasMeshColors_{ParameterName}() ? {ParameterName}_MeshColorBuffer[VertexIndex] FMANUALFETCH_COLOR_COMPONENT_SWIZZLE : float4(1, 1, 1, 1);
+}
+
+void GetVertexUV_{ParameterName}(in int VertexIndex, in int UVSet, out float2 OutUV)
+{
+    if ({ParameterName}_MeshNumTexCoord > 0)
+    {
+        uint Stride = {ParameterName}_MeshNumTexCoord;
+        uint SelectedUVSet = clamp((uint) UVSet, 0, {ParameterName}_MeshNumTexCoord - 1);
+        OutUV = {ParameterName}_MeshTexCoordBuffer[VertexIndex * Stride + SelectedUVSet];
+    }
+    else
+    {
+        OutUV = 0.0f;
+    }
+}
+
+void IsValidVertex_{ParameterName}(int Vertex, out bool IsValid)
+{
+    IsValid = Vertex < (int){ParameterName}_MeshVertexCount;
+}
+
+void RandomVertex_{ParameterName}(NiagaraRandInfo RandInfo, out int OutVertex)
+{
+	float RandT0 = NiagaraRandomFloat(RandInfo);
+	OutVertex = (int)min(uint(RandT0 * float({ParameterName}_MeshVertexCount)), {ParameterName}_MeshVertexCount - 1); // avoid % by using mul/min to Tri = MeshVertexCount
+}
+
+void GetVertexCount_{ParameterName}(out int VertexCount)
+{
+	VertexCount = (int){ParameterName}_MeshVertexCount;
+}
+
+void IsValidFilteredVertex_{ParameterName}(int FilteredIndex, out bool IsValid)
+{
+	IsValid = FilteredIndex < (int){ParameterName}_MeshNumSamplingRegionVertices;
+}
+
+void RandomFilteredVertex_{ParameterName}(NiagaraRandInfo RandInfo, out int OutVertex)
+{
+	// Early out for bad data
+	if ( {ParameterName}_MeshNumSamplingRegionVertices == 0 )
+	{
+		OutVertex = -1;
+		return;
+	}
+
+	float RandT0 = NiagaraRandomFloat(RandInfo);
+	int FilteredIndex = (int)min(uint(RandT0 * float({ParameterName}_MeshNumSamplingRegionVertices)), {ParameterName}_MeshNumSamplingRegionVertices - 1); // avoid % by using mul/min to Tri = MeshVertexCount
+	OutVertex = (int){ParameterName}_MeshSampleRegionsVertices[FilteredIndex];
+}
+
+void GetFilteredVertexCount_{ParameterName}(out int VertexCount)
+{
+	VertexCount = (int){ParameterName}_MeshNumSamplingRegionVertices;
+}
+
+void GetFilteredVertex_{ParameterName}(int FilteredIndex, out int VertexIndex)
+{
+	// Early out for bad data
+	if ( {ParameterName}_MeshNumSamplingRegionVertices == 0 )
+	{
+		VertexIndex = -1;
+		return;
+	}
+
+	FilteredIndex = clamp(FilteredIndex, 0, (int){ParameterName}_MeshNumSamplingRegionVertices - 1);
+	VertexIndex = (int){ParameterName}_MeshSampleRegionsVertices[FilteredIndex];
+}
+
+////////////////////////////////////////////////////////////////////////////////////////////////////////////////////////////////////////////////////////////////
+// Filtered Bones / Sockets
+void IsValidBone_{ParameterName}(in int BoneIndex, out bool IsValid)
+{
+	IsValid = BoneIndex >= 0 && BoneIndex < {ParameterName}_NumBones;
+}
+
+void RandomBone_{ParameterName}(NiagaraRandInfo RandInfo, out int Bone)
+{
+	float RandT0 = NiagaraRandomFloat(RandInfo);
+	Bone = min(int(RandT0 * float({ParameterName}_RandomMaxBone)), {ParameterName}_RandomMaxBone);
+	Bone = Bone != {ParameterName}_ExcludeBoneIndex ? Bone : Bone + 1;
+}
+
+void GetBoneCount_{ParameterName}(out int Count)
+{
+	Count = {ParameterName}_NumBones;
+}
+
+void GetParentBone_{ParameterName}(int BoneIndex, out int ParentIndex)
+{
+	ParentIndex = -1;
+	if ( BoneIndex >= 0 && BoneIndex < {ParameterName}_NumBones )
+	{
+		ParentIndex = asint({ParameterName}_MeshCurrSamplingBonesBuffer[BoneIndex * 3 + 2].w);
+		
+	}
+}
+
+void GetFilteredBoneCount_{ParameterName}(out int Count)
+{
+    Count = {ParameterName}_NumFilteredBones;
+}
+
+void GetFilteredBone_{ParameterName}(in int InBoneIndex, out int Bone)
+{
+    int BoneIndex = min(InBoneIndex, {ParameterName}_NumFilteredBones - 1);
+    Bone = {ParameterName}_NumFilteredBones > 0 ? (int){ParameterName}_FilteredAndUnfilteredBones[BoneIndex] : -1;
+}
+
+void RandomFilteredBone_{ParameterName}(NiagaraRandInfo RandInfo, out int Bone)
+{
+	float RandT0 = NiagaraRandomFloat(RandInfo);
+    int FilteredBone = min(int(RandT0 * float({ParameterName}_NumFilteredBones)), {ParameterName}_NumFilteredBones - 1); // avoid % by using mul/min to Tri = MeshVertexCount
+    Bone = {ParameterName}_NumFilteredBones > 0 ? (int){ParameterName}_FilteredAndUnfilteredBones[FilteredBone] : -1;
+}
+
+void GetUnfilteredBoneCount_{ParameterName}(out int Count)
+{
+    Count = {ParameterName}_NumUnfilteredBones;
+}
+
+void GetUnfilteredBone_{ParameterName}(in int InBoneIndex, out int Bone)
+{
+    int BoneIndex = min(InBoneIndex, {ParameterName}_NumUnfilteredBones - 1);
+    Bone = {ParameterName}_NumFilteredBones > 0 ? (int){ParameterName}_FilteredAndUnfilteredBones[BoneIndex + {ParameterName}_NumFilteredBones] : BoneIndex;
+}
+
+void RandomUnfilteredBone_{ParameterName}(NiagaraRandInfo RandInfo, out int Bone)
+{
+	float RandT0 = NiagaraRandomFloat(RandInfo);
+	if ( {ParameterName}_NumFilteredBones == 0 )
+	{
+		Bone = min(int(RandT0 * float({ParameterName}_RandomMaxBone)), {ParameterName}_RandomMaxBone);
+		Bone = Bone != {ParameterName}_ExcludeBoneIndex ? Bone : Bone + 1;
+	}
+	else
+	{
+		int UnfilteredBone = min(int(RandT0 * float({ParameterName}_NumUnfilteredBones)), {ParameterName}_NumUnfilteredBones - 1); // avoid % by using mul/min to Tri = MeshVertexCount
+		Bone = {ParameterName}_NumFilteredBones > 0 ? (int){ParameterName}_FilteredAndUnfilteredBones[UnfilteredBone + {ParameterName}_NumFilteredBones] : UnfilteredBone;
+	}
+}
+
+void GetFilteredSocketCount_{ParameterName}(out int Count)
+{
+    Count = {ParameterName}_NumFilteredSockets;
+}
+
+void GetFilteredSocketTransform_{ParameterName}(in int SocketIndex, in bool bApplyTransform, out float3 OutTranslation, out float4 OutRotation, out float3 OutScale)
+{
+    if ( {ParameterName}_NumFilteredSockets > 0 )
+	{
+		SocketIndex = clamp(SocketIndex, 0, {ParameterName}_NumFilteredSockets - 1);
+		int BufferOffset = ({ParameterName}_FilteredSocketBoneOffset + SocketIndex) * 3;
+
+		OutTranslation = {ParameterName}_MeshCurrSamplingBonesBuffer[BufferOffset + 0].xyz;
+		OutRotation = {ParameterName}_MeshCurrSamplingBonesBuffer[BufferOffset + 1];
+		OutScale = {ParameterName}_MeshCurrSamplingBonesBuffer[BufferOffset + 2].xyz;
+
+		if (bApplyTransform)
+		{
+			OutTranslation = mul(float4(OutTranslation, 1), {ParameterName}_InstanceTransform).xyz;
+			OutRotation = NiagaraGPU_QuatMul({ParameterName}_InstanceRotation, OutRotation);
+			OutScale = mul(OutScale, (float3x3){ParameterName}_InstanceTransform);
+		}
+	}
+	else
+	{
+		OutTranslation = float3(0.0f, 0.0f, 0.0f);
+		OutRotation = float4(0.0f, 0.0f, 0.0f, 1.0f);
+		OutScale = float3(0.0f, 0.0f, 0.0f);
+	}
+}
+
+void GetFilteredSocket_{ParameterName}(in int InSocketIndex, out int Bone)
+{
+    Bone = {ParameterName}_FilteredSocketBoneOffset + clamp(InSocketIndex, 0, {ParameterName}_NumFilteredSockets - 1);
+}
+
+void RandomFilteredSocket_{ParameterName}(NiagaraRandInfo RandInfo, out int SocketBone)
+{
+	float RandT0 = NiagaraRandomFloat(RandInfo);
+    SocketBone = {ParameterName}_FilteredSocketBoneOffset + min(int(RandT0 * float({ParameterName}_NumFilteredSockets)), {ParameterName}_NumFilteredSockets - 1); // avoid % by using mul/min to Tri = MeshVertexCount
+}
+
+void RandomFilteredSocketOrBone_{ParameterName}(NiagaraRandInfo RandInfo, out int Bone)
+{
+	float RandT0 = NiagaraRandomFloat(RandInfo);
+	int NumFilteredSocketsAndBones = {ParameterName}_NumFilteredBones + {ParameterName}_NumFilteredSockets;
+    int FilteredIndex = min(int(RandT0 * float(NumFilteredSocketsAndBones)), NumFilteredSocketsAndBones - 1);
+	if ( FilteredIndex < {ParameterName}_NumFilteredBones )
+	{
+		Bone = (int){ParameterName}_FilteredAndUnfilteredBones[FilteredIndex];
+	}
+	else
+	{
+		Bone = {ParameterName}_FilteredSocketBoneOffset + FilteredIndex - {ParameterName}_NumFilteredBones;
+	}
+}
+
+void GetFilteredSocketOrBoneCount_{ParameterName}(out int Count)
+{
+	Count = {ParameterName}_NumFilteredBones + {ParameterName}_NumFilteredSockets;
+}
+
+void GetFilteredSocketOrBone_{ParameterName}(int FilteredIndex, out int Bone)
+{
+	int NumFilteredSocketsAndBones = {ParameterName}_NumFilteredBones + {ParameterName}_NumFilteredSockets;
+	FilteredIndex = clamp(FilteredIndex, 0, NumFilteredSocketsAndBones - 1);
+	if ( FilteredIndex < {ParameterName}_NumFilteredBones )
+	{
+		Bone = (int){ParameterName}_FilteredAndUnfilteredBones[FilteredIndex];
+	}
+	else
+	{
+		Bone = {ParameterName}_FilteredSocketBoneOffset + FilteredIndex - {ParameterName}_NumFilteredBones;
+	}
+}
+
+////////////////////////////////////////////////////////////////////////////////////////////////////////////////////////////////////////////////////////////////
+// Connectivity
+
+void GetAdjacentTriangleIndex_{ParameterName}(int VertexId, uint AdjacencyIndex, out int TriangleIndex)
+{
+	TriangleIndex = -1;
+
+	if (AdjacencyIndex < {ParameterName}_ConnectivityMaxAdjacentPerVertex)
+	{
+#if DISKELMESH_ADJ_INDEX_FORMAT == DISKELMESH_ADJ_INDEX_FORMAT_FULL
+		const uint ReadIndex = (VertexId * {ParameterName}_ConnectivityMaxAdjacentPerVertex + AdjacencyIndex);
+		if (ReadIndex < {ParameterName}_ConnectivityBufferLength)
+		{
+			TriangleIndex = {ParameterName}_ConnectivityBuffer[ReadIndex];
+		}
+#elif DISKELMESH_ADJ_INDEX_FORMAT == DISKELMESH_ADJ_INDEX_FORMAT_HALF
+		const uint ReadIndex = ((VertexId * {ParameterName}_ConnectivityMaxAdjacentPerVertex + AdjacencyIndex) & ~0x1) >> 1;
+		if (ReadIndex < {ParameterName}_ConnectivityBufferLength)
+		{
+			uint PackedIndices = {ParameterName}_ConnectivityBuffer[ReadIndex];
+			uint UnpackedIndex = uint2(PackedIndices & 0xFFFF, PackedIndices >> 16)[AdjacencyIndex & 1];
+			TriangleIndex = UnpackedIndex == 0xFFFF ? -1 : int(UnpackedIndex);
+		}
+#else
+		#error Unsupported index format for adjacency information
+#endif
+	}
+}
+
+bool IsAdjacent_{ParameterName}(int VertexIndex, int TriangleIndex)
+{
+	for (uint AdjIt = 0; AdjIt < {ParameterName}_ConnectivityMaxAdjacentPerVertex; ++AdjIt)
+	{
+		int AdjTriangleIndex = -1;
+		GetAdjacentTriangleIndex_{ParameterName}(VertexIndex, AdjIt, AdjTriangleIndex);
+		if (AdjTriangleIndex == -1)
+		{
+			break;
+		}
+		if (AdjTriangleIndex == TriangleIndex)
+		{
+			return true;
+		}
+	}
+	return false;
+}
+
+void GetTriangleNeighbor_{ParameterName}(int TriangleIndex, int EdgeIndex, out int NeighborTriangleIndex, out int NeighborEdgeIndex, out bool IsValid)
+{
+	NeighborTriangleIndex = -1;
+	NeighborEdgeIndex = -1;
+
+	int3 SourceVertexIndices;
+
+	GetTriVertices_{ParameterName}(TriangleIndex, SourceVertexIndices.x, SourceVertexIndices.y, SourceVertexIndices.z);
+
+	int StartSourceIndex = SourceVertexIndices[(EdgeIndex + 1) % 3];
+	int EndSourceIndex = SourceVertexIndices[(EdgeIndex + 2) % 3];
+
+	for (uint SourceAdjIt = 0; SourceAdjIt < {ParameterName}_ConnectivityMaxAdjacentPerVertex; ++SourceAdjIt)
+	{
+		int SourceAdjTriangleIndex = -1;
+		GetAdjacentTriangleIndex_{ParameterName}(StartSourceIndex, SourceAdjIt, SourceAdjTriangleIndex);
+
+		if (SourceAdjTriangleIndex == TriangleIndex) // ignore self
+		{
+			continue;
+		}
+		if (SourceAdjTriangleIndex == -1) // stop looking
+		{
+			break;
+		}
+
+		if (IsAdjacent_{ParameterName}(EndSourceIndex, SourceAdjTriangleIndex))
+		{
+			NeighborTriangleIndex = SourceAdjTriangleIndex;
+
+			// now to find the edge index for the neighbor
+			int3 NeighborVertexIndices;
+			GetTriVertices_{ParameterName}(NeighborTriangleIndex, NeighborVertexIndices.x, NeighborVertexIndices.y, NeighborVertexIndices.z);
+
+			bool3 NeighborVertMatches = bool3(
+				IsAdjacent_{ParameterName}(NeighborVertexIndices.x, TriangleIndex),
+				IsAdjacent_{ParameterName}(NeighborVertexIndices.y, TriangleIndex),
+				IsAdjacent_{ParameterName}(NeighborVertexIndices.z, TriangleIndex));
+
+			if (all(NeighborVertMatches.xy))
+			{
+				NeighborEdgeIndex = 2;
+			}
+			else if (all(NeighborVertMatches.yz))
+			{
+				NeighborEdgeIndex = 0;
+			}
+			else if (all(NeighborVertMatches.zx))
+			{
+				NeighborEdgeIndex = 1;
+			}
+			break;
+		}
+	}
+	IsValid = NeighborTriangleIndex != -1;
+}
+
+void GetAdjacentTriangleIndex_{ParameterName}(in int VertexId, in int AdjacencyIndex, out int TriangleIndex, out bool IsValid)
+{
+	GetAdjacentTriangleIndex_{ParameterName}(VertexId, max(0, AdjacencyIndex), TriangleIndex);
+	IsValid = TriangleIndex != -1;
+}
diff --git a/Engine/Plugins/FX/Niagara/Source/Niagara/Classes/NiagaraDataInterfaceSkeletalMesh.h b/Engine/Plugins/FX/Niagara/Source/Niagara/Classes/NiagaraDataInterfaceSkeletalMesh.h
--- a/Engine/Plugins/FX/Niagara/Source/Niagara/Classes/NiagaraDataInterfaceSkeletalMesh.h
+++ b/Engine/Plugins/FX/Niagara/Source/Niagara/Classes/NiagaraDataInterfaceSkeletalMesh.h
@@ -716,9 +716,6 @@ class NIAGARA_API UNiagaraDataInterfaceSkeletalMesh : public UNiagaraDataInterfa
 	GENERATED_UCLASS_BODY()
 
 public:
-
-	DECLARE_NIAGARA_DI_PARAMETER();
-
 	/** Controls how to retrieve the Skeletal Mesh Component to attach to. */
 	UPROPERTY(EditAnywhere, Category = "Mesh")
 	ENDISkeletalMesh_SourceMode SourceMode;
@@ -829,6 +826,9 @@ public:
 	virtual bool GetFunctionHLSL(const FNiagaraDataInterfaceGPUParamInfo& ParamInfo, const FNiagaraDataInterfaceGeneratedFunction& FunctionInfo, int FunctionInstanceIndex, FString& OutHLSL) override;
 	virtual bool UpgradeFunctionCall(FNiagaraFunctionSignature& FunctionSignature) override;
 #endif
+	virtual bool UseLegacyShaderBindings() const override { return false; }
+	virtual void BuildShaderParameters(FNiagaraShaderParametersBuilder& ShaderParametersBuilder) const override;
+	virtual void SetShaderParameters(const FNiagaraDataInterfaceSetShaderParametersContext& Context) const override;
 	//~ UNiagaraDataInterface interface END
 
 	/** This overload is for use when initializing per-instance data. It possibly uses the SystemInstance and instance data to initialize a user binding */
diff --git a/Engine/Plugins/FX/Niagara/Source/Niagara/Private/NiagaraDataInterfaceSkeletalMesh.cpp b/Engine/Plugins/FX/Niagara/Source/Niagara/Private/NiagaraDataInterfaceSkeletalMesh.cpp
--- a/Engine/Plugins/FX/Niagara/Source/Niagara/Private/NiagaraDataInterfaceSkeletalMesh.cpp
+++ b/Engine/Plugins/FX/Niagara/Source/Niagara/Private/NiagaraDataInterfaceSkeletalMesh.cpp
@@ -18,6 +18,7 @@
 #include "NiagaraSystemInstance.h"
 #include "NiagaraRenderer.h"
 #include "NiagaraScript.h"
+#include "NiagaraShaderParametersBuilder.h"
 #include "NiagaraWorldManager.h"
 #include "Templates/AlignmentTemplates.h"
 #include "ShaderParameterUtils.h"
@@ -51,6 +52,56 @@ struct FNiagaraSkelMeshDIFunctionVersion
 
 namespace NDISkelMeshLocal
 {
+	static const TCHAR* CommonShaderFile = TEXT("/Plugin/FX/Niagara/Private/NiagaraDataInterfaceSkeletalMesh.ush");
+	static const TCHAR* TemplateShaderFile = TEXT("/Plugin/FX/Niagara/Private/NiagaraDataInterfaceSkeletalMeshTemplate.ush");
+
+	BEGIN_SHADER_PARAMETER_STRUCT(FShaderParameters,)
+		SHADER_PARAMETER_SRV(Buffer<uint>,		MeshIndexBuffer)
+		SHADER_PARAMETER_SRV(Buffer<float>,		MeshVertexBuffer)
+		SHADER_PARAMETER_SRV(Buffer<uint>,		MeshSkinWeightBuffer)
+		SHADER_PARAMETER_SRV(Buffer<uint>,		MeshSkinWeightLookupBuffer)
+		SHADER_PARAMETER_SRV(Buffer<float4>,	MeshCurrBonesBuffer)
+		SHADER_PARAMETER_SRV(Buffer<float4>,	MeshPrevBonesBuffer)
+		SHADER_PARAMETER_SRV(Buffer<float4>,	MeshCurrSamplingBonesBuffer)
+		SHADER_PARAMETER_SRV(Buffer<float4>,	MeshPrevSamplingBonesBuffer)
+		SHADER_PARAMETER_SRV(Buffer<float4>,	MeshTangentBuffer)
+		SHADER_PARAMETER_SRV(Buffer<float2>,	MeshTexCoordBuffer)
+		SHADER_PARAMETER_SRV(Buffer<float4>,	MeshColorBuffer)
+		SHADER_PARAMETER_SRV(Buffer<uint>,		MeshTriangleSamplerProbAliasBuffer)
+		SHADER_PARAMETER(uint32,				MeshNumSamplingRegionTriangles)
+		SHADER_PARAMETER(uint32,				MeshNumSamplingRegionVertices)
+		SHADER_PARAMETER_SRV(Buffer<uint>,		MeshSamplingRegionsProbAliasBuffer)
+		SHADER_PARAMETER_SRV(Buffer<uint>,		MeshSampleRegionsTriangleIndices)
+		SHADER_PARAMETER_SRV(Buffer<uint>,		MeshSampleRegionsVertices)
+		SHADER_PARAMETER_SRV(Buffer<uint>,		MeshTriangleMatricesOffsetBuffer)
+		SHADER_PARAMETER(uint32,				MeshTriangleCount)
+		SHADER_PARAMETER(uint32,				MeshVertexCount)
+		SHADER_PARAMETER(uint32,				MeshWeightStride)
+		SHADER_PARAMETER(uint32,				MeshSkinWeightIndexSize)
+		SHADER_PARAMETER(uint32,				MeshNumTexCoord)
+		SHADER_PARAMETER(uint32,				MeshNumWeights)
+		SHADER_PARAMETER(int,					NumBones)
+		SHADER_PARAMETER(int,					NumFilteredBones)
+		SHADER_PARAMETER(int,					NumUnfilteredBones)
+		SHADER_PARAMETER(int,					RandomMaxBone)
+		SHADER_PARAMETER(int,					ExcludeBoneIndex)
+		SHADER_PARAMETER_SRV(Buffer<uint>,		FilteredAndUnfilteredBones)
+		SHADER_PARAMETER(int,					NumFilteredSockets)
+		SHADER_PARAMETER(int,					FilteredSocketBoneOffset)
+		SHADER_PARAMETER_SRV(Buffer<int>,		UvMappingBuffer)
+		SHADER_PARAMETER(uint32,				UvMappingBufferLength)
+		SHADER_PARAMETER(uint32,				UvMappingSet)
+		SHADER_PARAMETER_SRV(Buffer<uint>,		ConnectivityBuffer)
+		SHADER_PARAMETER(uint32,				ConnectivityBufferLength)
+		SHADER_PARAMETER(uint32,				ConnectivityMaxAdjacentPerVertex)
+		SHADER_PARAMETER(FMatrix44f,			InstanceTransform)
+		SHADER_PARAMETER(FMatrix44f,			InstancePrevTransform)
+		SHADER_PARAMETER(FQuat4f,				InstanceRotation)
+		SHADER_PARAMETER(FQuat4f,				InstancePrevRotation)
+		SHADER_PARAMETER(float,					InstanceInvDeltaTime)
+		SHADER_PARAMETER(uint32,				EnabledFeatures)
+	END_SHADER_PARAMETER_STRUCT()
+
 	int32 GetProbAliasDWORDSize(int32 TriangleCount)
 	{
 		const ENDISkelMesh_GpuUniformSamplingFormat::Type Format = GetDefault<UNiagaraSettings>()->NDISkelMesh_GpuUniformSamplingFormat;
@@ -1385,286 +1436,6 @@ static void GetNiagaraDataInterfaceParametersName(FNDISkeletalMeshParametersName
 	Names.EnabledFeaturesName = UNiagaraDataInterfaceSkeletalMesh::EnabledFeaturesName + Suffix;
 }
 
-struct FNiagaraDataInterfaceParametersCS_SkeletalMesh : public FNiagaraDataInterfaceParametersCS
-{
-	DECLARE_TYPE_LAYOUT(FNiagaraDataInterfaceParametersCS_SkeletalMesh, NonVirtual);
-public:
-	void Bind(const FNiagaraDataInterfaceGPUParamInfo& ParameterInfo, const class FShaderParameterMap& ParameterMap)
-	{
-		FNDISkeletalMeshParametersName ParamNames;
-		GetNiagaraDataInterfaceParametersName(ParamNames, ParameterInfo.DataInterfaceHLSLSymbol);
-
-		MeshIndexBuffer.Bind(ParameterMap, *ParamNames.MeshIndexBufferName);
-		MeshVertexBuffer.Bind(ParameterMap, *ParamNames.MeshVertexBufferName);
-		MeshSkinWeightBuffer.Bind(ParameterMap, *ParamNames.MeshSkinWeightBufferName);
-		MeshSkinWeightLookupBuffer.Bind(ParameterMap, *ParamNames.MeshSkinWeightLookupBufferName);
-		MeshCurrBonesBuffer.Bind(ParameterMap, *ParamNames.MeshCurrBonesBufferName);
-		MeshPrevBonesBuffer.Bind(ParameterMap, *ParamNames.MeshPrevBonesBufferName);
-		MeshCurrSamplingBonesBuffer.Bind(ParameterMap, *ParamNames.MeshCurrSamplingBonesBufferName);
-		MeshPrevSamplingBonesBuffer.Bind(ParameterMap, *ParamNames.MeshPrevSamplingBonesBufferName);
-		MeshTangentBuffer.Bind(ParameterMap, *ParamNames.MeshTangentBufferName);
-		MeshTexCoordBuffer.Bind(ParameterMap, *ParamNames.MeshTexCoordBufferName);
-		MeshColorBuffer.Bind(ParameterMap, *ParamNames.MeshColorBufferName);
-		MeshTriangleSamplerProbAliasBuffer.Bind(ParameterMap, *ParamNames.MeshTriangleSamplerProbAliasBufferName);
-		MeshNumSamplingRegionTriangles.Bind(ParameterMap, *ParamNames.MeshNumSamplingRegionTrianglesName);
-		MeshNumSamplingRegionVertices.Bind(ParameterMap, *ParamNames.MeshNumSamplingRegionVerticesName);
-		MeshSamplingRegionsProbAliasBuffer.Bind(ParameterMap, *ParamNames.MeshSamplingRegionsProbAliasBufferName);
-		MeshSampleRegionsTriangleIndices.Bind(ParameterMap, *ParamNames.MeshSampleRegionsTriangleIndicesName);
-		MeshSampleRegionsVertices.Bind(ParameterMap, *ParamNames.MeshSampleRegionsVerticesName);
-		MeshTriangleMatricesOffsetBuffer.Bind(ParameterMap, *ParamNames.MeshTriangleMatricesOffsetBufferName);
-		MeshTriangleCount.Bind(ParameterMap, *ParamNames.MeshTriangleCountName);
-		MeshVertexCount.Bind(ParameterMap, *ParamNames.MeshVertexCountName);
-		MeshWeightStride.Bind(ParameterMap, *ParamNames.MeshWeightStrideName);
-		MeshSkinWeightIndexSize.Bind(ParameterMap, *ParamNames.MeshSkinWeightIndexSizeName);
-		MeshNumTexCoord.Bind(ParameterMap, *ParamNames.MeshNumTexCoordName);
-		MeshNumWeights.Bind(ParameterMap, *ParamNames.MeshNumWeightsName);
-		NumBones.Bind(ParameterMap, *ParamNames.NumBonesName);
-		NumFilteredBones.Bind(ParameterMap, *ParamNames.NumFilteredBonesName);
-		NumUnfilteredBones.Bind(ParameterMap, *ParamNames.NumUnfilteredBonesName);
-		RandomMaxBone.Bind(ParameterMap, *ParamNames.RandomMaxBoneName);
-		ExcludeBoneIndex.Bind(ParameterMap, *ParamNames.ExcludeBoneIndexName);
-		FilteredAndUnfilteredBones.Bind(ParameterMap, *ParamNames.FilteredAndUnfilteredBonesName);
-		NumFilteredSockets.Bind(ParameterMap, *ParamNames.NumFilteredSocketsName);
-		FilteredSocketBoneOffset.Bind(ParameterMap, *ParamNames.FilteredSocketBoneOffsetName);
-		UvMappingBuffer.Bind(ParameterMap, *ParamNames.UvMappingBufferName);
-		UvMappingBufferLength.Bind(ParameterMap, *ParamNames.UvMappingBufferLengthName);
-		UvMappingSet.Bind(ParameterMap, *ParamNames.UvMappingSetName);
-		ConnectivityBuffer.Bind(ParameterMap, *ParamNames.ConnectivityBufferName);
-		ConnectivityBufferLength.Bind(ParameterMap, *ParamNames.ConnectivityBufferLengthName);
-		ConnectivityMaxAdjacentPerVertex.Bind(ParameterMap, *ParamNames.ConnectivityMaxAdjacentPerVertexName);
-		InstanceTransform.Bind(ParameterMap, *ParamNames.InstanceTransformName);
-		InstancePrevTransform.Bind(ParameterMap, *ParamNames.InstancePrevTransformName);
-		InstanceRotation.Bind(ParameterMap, *ParamNames.InstanceRotationName);
-		InstancePrevRotation.Bind(ParameterMap, *ParamNames.InstancePrevRotationName);
-		InstanceInvDeltaTime.Bind(ParameterMap, *ParamNames.InstanceInvDeltaTimeName);
-		EnabledFeatures.Bind(ParameterMap, *ParamNames.EnabledFeaturesName);
-	}
-
-	void Set(FRHICommandList& RHICmdList, const FNiagaraDataInterfaceSetArgs& Context) const
-	{
-		check(IsInRenderingThread());
-
-		FRHIComputeShader* ComputeShaderRHI = Context.Shader.GetComputeShader();
-		FNiagaraDataInterfaceProxySkeletalMesh* InterfaceProxy = static_cast<FNiagaraDataInterfaceProxySkeletalMesh*>(Context.DataInterface);
-		FNiagaraDataInterfaceProxySkeletalMeshData* InstanceData = InterfaceProxy->SystemInstancesToData.Find(Context.SystemInstanceID);
-		if (InstanceData && InstanceData->StaticBuffers)
-		{
-			FSkeletalMeshGpuSpawnStaticBuffers* StaticBuffers = InstanceData->StaticBuffers;
-
-			// Bind mesh buffers
-			SetSRVParameter(RHICmdList, ComputeShaderRHI, MeshVertexBuffer, StaticBuffers->GetBufferPositionSRV());
-			SetSRVParameter(RHICmdList, ComputeShaderRHI, MeshIndexBuffer, StaticBuffers->GetBufferIndexSRV());
-			SetSRVParameter(RHICmdList, ComputeShaderRHI, MeshTangentBuffer, StaticBuffers->GetBufferTangentSRV());
-			SetSRVParameter(RHICmdList, ComputeShaderRHI, MeshTexCoordBuffer, StaticBuffers->GetBufferTexCoordSRV());
-			SetSRVParameter(RHICmdList, ComputeShaderRHI, MeshColorBuffer, StaticBuffers->GetBufferColorSRV());
-
-			SetShaderValue(RHICmdList, ComputeShaderRHI, MeshNumTexCoord, StaticBuffers->GetNumTexCoord());
-			SetShaderValue(RHICmdList, ComputeShaderRHI, MeshTriangleCount, StaticBuffers->GetTriangleCount());
-			SetShaderValue(RHICmdList, ComputeShaderRHI, MeshVertexCount, StaticBuffers->GetVertexCount());
-
-			// Set triangle sampling buffer
-			SetSRVParameter(RHICmdList, ComputeShaderRHI, MeshTriangleSamplerProbAliasBuffer, StaticBuffers->GetBufferTriangleUniformSamplerProbAliasSRV());
-
-			// Set triangle sampling region buffer
-			SetShaderValue(RHICmdList, ComputeShaderRHI, MeshNumSamplingRegionTriangles, StaticBuffers->GetNumSamplingRegionTriangles());
-			SetShaderValue(RHICmdList, ComputeShaderRHI, MeshNumSamplingRegionVertices, StaticBuffers->GetNumSamplingRegionVertices());
-			SetSRVParameter(RHICmdList, ComputeShaderRHI, MeshSamplingRegionsProbAliasBuffer, StaticBuffers->GetSampleRegionsProbAliasSRV());
-			SetSRVParameter(RHICmdList, ComputeShaderRHI, MeshSampleRegionsTriangleIndices, StaticBuffers->GetSampleRegionsTriangleIndicesSRV());
-			SetSRVParameter(RHICmdList, ComputeShaderRHI, MeshSampleRegionsVertices, StaticBuffers->GetSampleRegionsVerticesSRV());
-
-			FRHIShaderResourceView* MeshSkinWeightBufferSRV = InstanceData->MeshSkinWeightBuffer->GetSRV();
-			SetSRVParameter(RHICmdList, ComputeShaderRHI, MeshSkinWeightBuffer, MeshSkinWeightBufferSRV ? MeshSkinWeightBufferSRV : FNiagaraRenderer::GetDummyUIntBuffer());
-			FRHIShaderResourceView* MeshSkinWeightLookupBufferSRV = InstanceData->MeshSkinWeightLookupBuffer->GetSRV();
-			SetSRVParameter(RHICmdList, ComputeShaderRHI, MeshSkinWeightLookupBuffer, MeshSkinWeightLookupBufferSRV ? MeshSkinWeightLookupBufferSRV : FNiagaraRenderer::GetDummyUIntBuffer());
-
-			SetShaderValue(RHICmdList, ComputeShaderRHI, MeshWeightStride, InstanceData->MeshWeightStrideByte/4);
-			SetShaderValue(RHICmdList, ComputeShaderRHI, MeshSkinWeightIndexSize, InstanceData->MeshSkinWeightIndexSizeByte);
-
-			uint32 EnabledFeaturesBits = 0;
-			EnabledFeaturesBits |= StaticBuffers->IsUseGpuUniformlyDistributedSampling() ? 1 : 0;
-			EnabledFeaturesBits |= StaticBuffers->IsSamplingRegionsAllAreaWeighted() ? 2 : 0;
-			EnabledFeaturesBits |= InstanceData->bUnlimitedBoneInfluences ? 4 : 0;
-			EnabledFeaturesBits |= StaticBuffers->HasMeshColors() ? 8 : 0;
-
-			FSkeletalMeshGpuDynamicBufferProxy* DynamicBuffers = InstanceData->DynamicBuffer;
-			check(DynamicBuffers);
-			if(DynamicBuffers->DoesBoneDataExist())
-			{
-				SetShaderValue(RHICmdList, ComputeShaderRHI, MeshNumWeights, StaticBuffers->GetNumWeights());
-				SetSRVParameter(RHICmdList, ComputeShaderRHI, MeshCurrBonesBuffer, DynamicBuffers->GetRWBufferBone().SectionSRV);
-				SetSRVParameter(RHICmdList, ComputeShaderRHI, MeshPrevBonesBuffer, DynamicBuffers->GetRWBufferPrevBone().SectionSRV);
-				SetSRVParameter(RHICmdList, ComputeShaderRHI, MeshCurrSamplingBonesBuffer, DynamicBuffers->GetRWBufferBone().SamplingSRV);
-				SetSRVParameter(RHICmdList, ComputeShaderRHI, MeshPrevSamplingBonesBuffer, DynamicBuffers->GetRWBufferPrevBone().SamplingSRV);
-				SetSRVParameter(RHICmdList, ComputeShaderRHI, MeshTriangleMatricesOffsetBuffer, StaticBuffers->GetBufferTriangleMatricesOffsetSRV());
-			}
-			// Bind dummy data for validation purposes only.  Code will not execute due to "EnabledFeatures" bits but validation can not determine that.
-			else
-			{
-				SetShaderValue(RHICmdList, ComputeShaderRHI, MeshNumWeights, 0);
-				SetSRVParameter(RHICmdList, ComputeShaderRHI, MeshCurrBonesBuffer, FNiagaraRenderer::GetDummyFloat4Buffer());
-				SetSRVParameter(RHICmdList, ComputeShaderRHI, MeshPrevBonesBuffer, FNiagaraRenderer::GetDummyFloat4Buffer());
-				SetSRVParameter(RHICmdList, ComputeShaderRHI, MeshCurrSamplingBonesBuffer, FNiagaraRenderer::GetDummyFloat4Buffer());
-				SetSRVParameter(RHICmdList, ComputeShaderRHI, MeshPrevSamplingBonesBuffer, FNiagaraRenderer::GetDummyFloat4Buffer());
-				SetSRVParameter(RHICmdList, ComputeShaderRHI, MeshTriangleMatricesOffsetBuffer, FNiagaraRenderer::GetDummyUIntBuffer());
-			}
-
-			SetShaderValue(RHICmdList, ComputeShaderRHI, NumBones, DynamicBuffers->GetNumBones());
-
-			SetShaderValue(RHICmdList, ComputeShaderRHI, NumFilteredBones, StaticBuffers->GetNumFilteredBones());
-			SetShaderValue(RHICmdList, ComputeShaderRHI, NumUnfilteredBones, StaticBuffers->GetNumUnfilteredBones());
-			SetShaderValue(RHICmdList, ComputeShaderRHI, RandomMaxBone, StaticBuffers->GetExcludedBoneIndex() >= 0 ? DynamicBuffers->GetNumBones() - 2 : DynamicBuffers->GetNumBones() - 1);
-			SetShaderValue(RHICmdList, ComputeShaderRHI, ExcludeBoneIndex, StaticBuffers->GetExcludedBoneIndex());
-			SetSRVParameter(RHICmdList, ComputeShaderRHI, FilteredAndUnfilteredBones, StaticBuffers->GetFilteredAndUnfilteredBonesSRV());
-
-			SetShaderValue(RHICmdList, ComputeShaderRHI, NumFilteredSockets, StaticBuffers->GetNumFilteredSockets());
-			SetShaderValue(RHICmdList, ComputeShaderRHI, FilteredSocketBoneOffset, StaticBuffers->GetFilteredSocketBoneOffset());
-
-			if (InstanceData->UvMappingBuffer)
-			{
-				SetSRVParameter(RHICmdList, ComputeShaderRHI, UvMappingBuffer, InstanceData->UvMappingBuffer->GetSrv());
-				SetShaderValue(RHICmdList, ComputeShaderRHI, UvMappingBufferLength, InstanceData->UvMappingBuffer->GetBufferSize());
-				SetShaderValue(RHICmdList, ComputeShaderRHI, UvMappingSet, InstanceData->UvMappingSet);
-			}
-			else
-			{
-				SetSRVParameter(RHICmdList, ComputeShaderRHI, UvMappingBuffer, FNiagaraRenderer::GetDummyIntBuffer());
-				SetShaderValue(RHICmdList, ComputeShaderRHI, UvMappingBufferLength, 0);
-				SetShaderValue(RHICmdList, ComputeShaderRHI, UvMappingSet, 0);
-			}
-
-			if (InstanceData->ConnectivityBuffer)
-			{
-				const uint32 NumBufferElements = FMath::DivideAndRoundUp<uint32>(InstanceData->ConnectivityBuffer->GetBufferSize(), sizeof(uint32));
-				SetSRVParameter(RHICmdList, ComputeShaderRHI, ConnectivityBuffer, InstanceData->ConnectivityBuffer->GetSrv());
-				SetShaderValue(RHICmdList, ComputeShaderRHI, ConnectivityBufferLength, NumBufferElements);
-				SetShaderValue(RHICmdList, ComputeShaderRHI, ConnectivityMaxAdjacentPerVertex, InstanceData->ConnectivityBuffer->MaxAdjacentTriangleCount);
-			}
-			else
-			{
-				SetSRVParameter(RHICmdList, ComputeShaderRHI, ConnectivityBuffer, FNiagaraRenderer::GetDummyUIntBuffer());
-				SetShaderValue(RHICmdList, ComputeShaderRHI, ConnectivityBufferLength, 0);
-				SetShaderValue(RHICmdList, ComputeShaderRHI, ConnectivityMaxAdjacentPerVertex, 0);
-			}
-
-			SetShaderValue(RHICmdList, ComputeShaderRHI, InstanceTransform, InstanceData->Transform);
-			SetShaderValue(RHICmdList, ComputeShaderRHI, InstancePrevTransform, InstanceData->PrevTransform);
-			SetShaderValue(RHICmdList, ComputeShaderRHI, InstanceRotation, InstanceData->Transform.GetMatrixWithoutScale().ToQuat());
-			SetShaderValue(RHICmdList, ComputeShaderRHI, InstancePrevRotation, InstanceData->PrevTransform.GetMatrixWithoutScale().ToQuat());
-			SetShaderValue(RHICmdList, ComputeShaderRHI, InstanceInvDeltaTime, 1.0f / InstanceData->DeltaSeconds);
-
-			SetShaderValue(RHICmdList, ComputeShaderRHI, EnabledFeatures, EnabledFeaturesBits);
-		}
-		else
-		{
-			// Bind dummy buffers
-			SetSRVParameter(RHICmdList, ComputeShaderRHI, MeshVertexBuffer, FNiagaraRenderer::GetDummyFloatBuffer());
-			SetSRVParameter(RHICmdList, ComputeShaderRHI, MeshIndexBuffer, FNiagaraRenderer::GetDummyUIntBuffer());
-			SetSRVParameter(RHICmdList, ComputeShaderRHI, MeshTangentBuffer, FNiagaraRenderer::GetDummyFloat4Buffer());
-
-			SetShaderValue(RHICmdList, ComputeShaderRHI, MeshNumTexCoord, 0);
-			SetSRVParameter(RHICmdList, ComputeShaderRHI, MeshTexCoordBuffer, FNiagaraRenderer::GetDummyFloat2Buffer());
-			SetSRVParameter(RHICmdList, ComputeShaderRHI, MeshColorBuffer, FNiagaraRenderer::GetDummyFloat4Buffer());
-			SetShaderValue(RHICmdList, ComputeShaderRHI, MeshTriangleCount, 0);
-			SetShaderValue(RHICmdList, ComputeShaderRHI, MeshVertexCount, 0);
-			SetSRVParameter(RHICmdList, ComputeShaderRHI, MeshTriangleSamplerProbAliasBuffer, FNiagaraRenderer::GetDummyUIntBuffer());
-
-			SetShaderValue(RHICmdList, ComputeShaderRHI, MeshNumSamplingRegionTriangles, 0);
-			SetShaderValue(RHICmdList, ComputeShaderRHI, MeshNumSamplingRegionVertices, 0);
-			SetSRVParameter(RHICmdList, ComputeShaderRHI, MeshSamplingRegionsProbAliasBuffer, FNiagaraRenderer::GetDummyUIntBuffer());
-			SetSRVParameter(RHICmdList, ComputeShaderRHI, MeshSampleRegionsTriangleIndices, FNiagaraRenderer::GetDummyUIntBuffer());
-			SetSRVParameter(RHICmdList, ComputeShaderRHI, MeshSampleRegionsVertices, FNiagaraRenderer::GetDummyUIntBuffer());
-
-			SetSRVParameter(RHICmdList, ComputeShaderRHI, MeshSkinWeightBuffer, FNiagaraRenderer::GetDummyUIntBuffer());
-			SetSRVParameter(RHICmdList, ComputeShaderRHI, MeshSkinWeightLookupBuffer, FNiagaraRenderer::GetDummyUIntBuffer());
-
-			SetShaderValue(RHICmdList, ComputeShaderRHI, MeshWeightStride, 0);
-			SetShaderValue(RHICmdList, ComputeShaderRHI, MeshSkinWeightIndexSize, 0);
-			SetShaderValue(RHICmdList, ComputeShaderRHI, MeshNumWeights, 0);
-
-			SetSRVParameter(RHICmdList, ComputeShaderRHI, MeshCurrBonesBuffer, FNiagaraRenderer::GetDummyFloat4Buffer());
-			SetSRVParameter(RHICmdList, ComputeShaderRHI, MeshPrevBonesBuffer, FNiagaraRenderer::GetDummyFloat4Buffer());
-			SetSRVParameter(RHICmdList, ComputeShaderRHI, MeshCurrSamplingBonesBuffer, FNiagaraRenderer::GetDummyFloat4Buffer());
-			SetSRVParameter(RHICmdList, ComputeShaderRHI, MeshPrevSamplingBonesBuffer, FNiagaraRenderer::GetDummyFloat4Buffer());
-			SetSRVParameter(RHICmdList, ComputeShaderRHI, MeshTriangleMatricesOffsetBuffer, FNiagaraRenderer::GetDummyUIntBuffer());
-
-			SetShaderValue(RHICmdList, ComputeShaderRHI, NumBones, 0);
-			SetShaderValue(RHICmdList, ComputeShaderRHI, NumFilteredBones, 0);
-			SetShaderValue(RHICmdList, ComputeShaderRHI, NumUnfilteredBones, 0);
-			SetShaderValue(RHICmdList, ComputeShaderRHI, RandomMaxBone, 0);
-			SetShaderValue(RHICmdList, ComputeShaderRHI, ExcludeBoneIndex, 0);
-			SetSRVParameter(RHICmdList, ComputeShaderRHI, FilteredAndUnfilteredBones, FNiagaraRenderer::GetDummyUIntBuffer());
-			SetShaderValue(RHICmdList, ComputeShaderRHI, NumFilteredSockets, 0);
-			SetShaderValue(RHICmdList, ComputeShaderRHI, FilteredSocketBoneOffset, 0);
-
-			SetSRVParameter(RHICmdList, ComputeShaderRHI, UvMappingBuffer, FNiagaraRenderer::GetDummyIntBuffer());
-			SetShaderValue(RHICmdList, ComputeShaderRHI, UvMappingBufferLength, 0);
-			SetShaderValue(RHICmdList, ComputeShaderRHI, UvMappingSet, 0);
-
-			SetSRVParameter(RHICmdList, ComputeShaderRHI, ConnectivityBuffer, FNiagaraRenderer::GetDummyUIntBuffer());
-			SetShaderValue(RHICmdList, ComputeShaderRHI, ConnectivityBufferLength, 0);
-			SetShaderValue(RHICmdList, ComputeShaderRHI, ConnectivityMaxAdjacentPerVertex, 0);
-
-			SetShaderValue(RHICmdList, ComputeShaderRHI, InstanceTransform, FMatrix44f::Identity);
-			SetShaderValue(RHICmdList, ComputeShaderRHI, InstancePrevTransform, FMatrix44f::Identity);
-			SetShaderValue(RHICmdList, ComputeShaderRHI, InstanceRotation, FQuat4f::Identity);
-			SetShaderValue(RHICmdList, ComputeShaderRHI, InstancePrevRotation, FQuat4f::Identity);
-			SetShaderValue(RHICmdList, ComputeShaderRHI, InstanceInvDeltaTime, 0.0f);
-
-			SetShaderValue(RHICmdList, ComputeShaderRHI, EnabledFeatures, 0);
-		}
-	}
-
-private:
-	LAYOUT_FIELD(FShaderResourceParameter, MeshIndexBuffer);
-	LAYOUT_FIELD(FShaderResourceParameter, MeshVertexBuffer);
-	LAYOUT_FIELD(FShaderResourceParameter, MeshSkinWeightBuffer);
-	LAYOUT_FIELD(FShaderResourceParameter, MeshSkinWeightLookupBuffer);
-	LAYOUT_FIELD(FShaderResourceParameter, MeshCurrBonesBuffer);
-	LAYOUT_FIELD(FShaderResourceParameter, MeshPrevBonesBuffer);
-	LAYOUT_FIELD(FShaderResourceParameter, MeshCurrSamplingBonesBuffer);
-	LAYOUT_FIELD(FShaderResourceParameter, MeshPrevSamplingBonesBuffer);
-	LAYOUT_FIELD(FShaderResourceParameter, MeshTangentBuffer);
-	LAYOUT_FIELD(FShaderResourceParameter, MeshTexCoordBuffer);
-	LAYOUT_FIELD(FShaderResourceParameter, MeshColorBuffer);
-	LAYOUT_FIELD(FShaderResourceParameter, MeshTriangleSamplerProbAliasBuffer);
-	LAYOUT_FIELD(FShaderParameter, MeshNumSamplingRegionTriangles);
-	LAYOUT_FIELD(FShaderParameter, MeshNumSamplingRegionVertices);
-	LAYOUT_FIELD(FShaderResourceParameter, MeshSamplingRegionsProbAliasBuffer);
-	LAYOUT_FIELD(FShaderResourceParameter, MeshSampleRegionsTriangleIndices);
-	LAYOUT_FIELD(FShaderResourceParameter, MeshSampleRegionsVertices);
-	LAYOUT_FIELD(FShaderResourceParameter, MeshTriangleMatricesOffsetBuffer);
-	LAYOUT_FIELD(FShaderParameter, MeshTriangleCount);
-	LAYOUT_FIELD(FShaderParameter, MeshVertexCount);
-	LAYOUT_FIELD(FShaderParameter, MeshWeightStride);
-	LAYOUT_FIELD(FShaderParameter, MeshSkinWeightIndexSize);
-	LAYOUT_FIELD(FShaderParameter, MeshNumTexCoord);
-	LAYOUT_FIELD(FShaderParameter, MeshNumWeights);
-	LAYOUT_FIELD(FShaderParameter, NumBones);
-	LAYOUT_FIELD(FShaderParameter, NumFilteredBones);
-	LAYOUT_FIELD(FShaderParameter, NumUnfilteredBones);
-	LAYOUT_FIELD(FShaderParameter, RandomMaxBone);
-	LAYOUT_FIELD(FShaderParameter, ExcludeBoneIndex);
-	LAYOUT_FIELD(FShaderResourceParameter, FilteredAndUnfilteredBones);
-	LAYOUT_FIELD(FShaderParameter, NumFilteredSockets);
-	LAYOUT_FIELD(FShaderParameter, FilteredSocketBoneOffset);
-	LAYOUT_FIELD(FShaderResourceParameter, UvMappingBuffer);
-	LAYOUT_FIELD(FShaderParameter, UvMappingBufferLength);
-	LAYOUT_FIELD(FShaderParameter, UvMappingSet);
-	LAYOUT_FIELD(FShaderResourceParameter, ConnectivityBuffer);
-	LAYOUT_FIELD(FShaderParameter, ConnectivityBufferLength);
-	LAYOUT_FIELD(FShaderParameter, ConnectivityMaxAdjacentPerVertex);
-	LAYOUT_FIELD(FShaderParameter, InstanceTransform);
-	LAYOUT_FIELD(FShaderParameter, InstancePrevTransform);
-	LAYOUT_FIELD(FShaderParameter, InstanceRotation);
-	LAYOUT_FIELD(FShaderParameter, InstancePrevRotation);
-	LAYOUT_FIELD(FShaderParameter, InstanceInvDeltaTime);
-	LAYOUT_FIELD(FShaderParameter, EnabledFeatures);
-};
-
-IMPLEMENT_TYPE_LAYOUT(FNiagaraDataInterfaceParametersCS_SkeletalMesh);
-
-IMPLEMENT_NIAGARA_DI_PARAMETER(UNiagaraDataInterfaceSkeletalMesh, FNiagaraDataInterfaceParametersCS_SkeletalMesh);
-
 //////////////////////////////////////////////////////////////////////////
 
 void FNiagaraDataInterfaceProxySkeletalMesh::ConsumePerInstanceDataFromGameThread(void* PerInstanceData, const FNiagaraSystemInstanceID& Instance)
@@ -3070,8 +2841,9 @@ bool UNiagaraDataInterfaceSkeletalMesh::AppendCompileHash(FNiagaraCompileHashVis
 	if (!Super::AppendCompileHash(InVisitor))
 		return false;
 
-	FSHAHash Hash = GetShaderFileHash((TEXT("/Plugin/FX/Niagara/Private/NiagaraDataInterfaceSkeletalMesh.ush")), EShaderPlatform::SP_PCD3D_SM5);
-	InVisitor->UpdateString(TEXT("NiagaraDataInterfaceSkeletalMeshHLSLSource"), Hash.ToString());
+	InVisitor->UpdateString(TEXT("NiagaraDataInterfaceSkeletalMeshHLSLSource"), GetShaderFileHash(NDISkelMeshLocal::CommonShaderFile, EShaderPlatform::SP_PCD3D_SM5).ToString());
+	InVisitor->UpdateString(TEXT("NiagaraDataInterfaceSkeletalMeshTemplateHLSLSource"), GetShaderFileHash(NDISkelMeshLocal::TemplateShaderFile, EShaderPlatform::SP_PCD3D_SM5).ToString());
+	InVisitor->UpdateShaderParameters<NDISkelMeshLocal::FShaderParameters>();
 
 	InVisitor->UpdatePOD(TEXT("NDISkelmesh_Influences"), int(GetDefault<UNiagaraSettings>()->NDISkelMesh_GpuMaxInfluences));
 	InVisitor->UpdatePOD(TEXT("NDISkelmesh_ProbAliasFormat"), int(GetDefault<UNiagaraSettings>()->NDISkelMesh_GpuUniformSamplingFormat));
@@ -3094,308 +2866,157 @@ void UNiagaraDataInterfaceSkeletalMesh::ModifyCompilationEnvironment(EShaderPlat
 #if WITH_EDITORONLY_DATA
 void UNiagaraDataInterfaceSkeletalMesh::GetCommonHLSL(FString& OutHLSL)
 {
-	OutHLSL += TEXT("#include \"/Plugin/FX/Niagara/Private/NiagaraDataInterfaceSkeletalMesh.ush\"\n");
+	OutHLSL.Appendf(TEXT("#include \"%s\"\n"), NDISkelMeshLocal::CommonShaderFile);
+}
+
+void UNiagaraDataInterfaceSkeletalMesh::GetParameterDefinitionHLSL(const FNiagaraDataInterfaceGPUParamInfo& ParamInfo, FString& OutHLSL)
+{
+	TMap<FString, FStringFormatArg> TemplateArgs =
+	{
+		{TEXT("ParameterName"),	ParamInfo.DataInterfaceHLSLSymbol},
+	};
+
+	FString TemplateFile;
+	LoadShaderSourceFile(NDISkelMeshLocal::TemplateShaderFile, EShaderPlatform::SP_PCD3D_SM5, &TemplateFile, nullptr);
+	OutHLSL += FString::Format(*TemplateFile, TemplateArgs);
 }
 
 bool UNiagaraDataInterfaceSkeletalMesh::GetFunctionHLSL(const FNiagaraDataInterfaceGPUParamInfo& ParamInfo, const FNiagaraDataInterfaceGeneratedFunction& FunctionInfo, int FunctionInstanceIndex, FString& OutHLSL)
 {
-	FNDISkeletalMeshParametersName ParamNames;
-	GetNiagaraDataInterfaceParametersName(ParamNames, ParamInfo.DataInterfaceHLSLSymbol);
-	TMap<FString, FStringFormatArg> ArgsSample = {
-		{TEXT("InstanceFunctionName"), FunctionInfo.InstanceName},
-		{TEXT("MeshTriCoordinateStructName"), TEXT("MeshTriCoordinate")},
-		{TEXT("MeshTriangleCount"), ParamNames.MeshTriangleCountName},
-		{TEXT("MeshVertexCount"), ParamNames.MeshVertexCountName},
-		{TEXT("GetDISkelMeshContextName"), TEXT("DISKELMESH_MAKE_CONTEXT(") + ParamInfo.DataInterfaceHLSLSymbol + TEXT(")")},
+	static const TSet<FName> ValidGpuFunctions =
+	{
+		FSkeletalMeshInterfaceHelper::GetTriCoordVerticesName,
+		FSkeletalMeshInterfaceHelper::GetTriangleCountName,
+		FSkeletalMeshInterfaceHelper::GetFilteredTriangleCountName,
+		//////////////////////////////////////////////////////////////////////////////////////////////////////////////////////////////
+		// Bone Sampling
+		FSkeletalMeshInterfaceHelper::GetSkinnedBoneDataName,
+		FSkeletalMeshInterfaceHelper::GetSkinnedBoneDataInterpolatedName,
+		FSkeletalMeshInterfaceHelper::GetSkinnedBoneDataWSName,
+		FSkeletalMeshInterfaceHelper::GetSkinnedBoneDataWSInterpolatedName,
+		//////////////////////////////////////////////////////////////////////////////////////////////////////////////////////////////
+		// Vertex Sampling
+		FSkeletalMeshInterfaceHelper::GetVertexDataName,
+		FSkeletalMeshInterfaceHelper::GetSkinnedVertexDataName,
+		FSkeletalMeshInterfaceHelper::GetSkinnedVertexDataWSName,
+		FSkeletalMeshInterfaceHelper::GetVertexColorName,
+		FSkeletalMeshInterfaceHelper::GetVertexUVName,
+		FSkeletalMeshInterfaceHelper::IsValidVertexName,
+		FSkeletalMeshInterfaceHelper::RandomVertexName,
+		FSkeletalMeshInterfaceHelper::GetVertexCountName,
+		FSkeletalMeshInterfaceHelper::IsValidFilteredVertexName,
+		FSkeletalMeshInterfaceHelper::RandomFilteredVertexName,
+		FSkeletalMeshInterfaceHelper::GetFilteredVertexCountName,
+		FSkeletalMeshInterfaceHelper::GetFilteredVertexAtName,
+		//////////////////////////////////////////////////////////////////////////////////////////////////////////////////////////////
+		// Filtered Bone
+		FSkeletalMeshInterfaceHelper::IsValidBoneName,
+		FSkeletalMeshInterfaceHelper::RandomBoneName,
+		FSkeletalMeshInterfaceHelper::GetBoneCountName,
+		FSkeletalMeshInterfaceHelper::GetParentBoneName,
+		FSkeletalMeshInterfaceHelper::GetFilteredBoneCountName,
+		FSkeletalMeshInterfaceHelper::GetFilteredBoneAtName,
+		FSkeletalMeshInterfaceHelper::RandomFilteredBoneName,
+		FSkeletalMeshInterfaceHelper::GetUnfilteredBoneCountName,
+		FSkeletalMeshInterfaceHelper::GetUnfilteredBoneAtName,
+		FSkeletalMeshInterfaceHelper::RandomUnfilteredBoneName,
+		//////////////////////////////////////////////////////////////////////////////////////////////////////////////////////////////
+		// Filtered Socket
+		FSkeletalMeshInterfaceHelper::GetFilteredSocketCountName,
+		FSkeletalMeshInterfaceHelper::GetFilteredSocketBoneAtName,
+		FSkeletalMeshInterfaceHelper::GetFilteredSocketTransformName,
+		FSkeletalMeshInterfaceHelper::RandomFilteredSocketName,
+		//////////////////////////////////////////////////////////////////////////////////////////////////////////////////////////////
+		// Misc bone functions
+		FSkeletalMeshInterfaceHelper::RandomFilteredSocketOrBoneName,
+		FSkeletalMeshInterfaceHelper::GetFilteredSocketOrBoneCountName,
+		FSkeletalMeshInterfaceHelper::GetFilteredSocketOrBoneAtName,
+		//////////////////////////////////////////////////////////////////////////////////////////////////////////////////////////////
+		// Uv Mapping
+		FSkeletalMeshInterfaceHelper::GetAdjacentTriangleIndexName,
+		FSkeletalMeshInterfaceHelper::GetTriangleNeighborName,
 	};
 
-	//////////////////////////////////////////////////////////////////////////////////////////////////////////////////////////////
-	// Triangle Sampling
-	if (FunctionInfo.DefinitionName == FSkeletalMeshInterfaceHelper::RandomTriCoordName)
-	{
-		static const TCHAR* FormatSample = TEXT("void {InstanceFunctionName} (NiagaraRandInfo InRandomInfo, out {MeshTriCoordinateStructName} OutCoord) { {GetDISkelMeshContextName} DISKelMesh_RandomTriCoord(DIContext, InRandomInfo.Seed1, InRandomInfo.Seed2, InRandomInfo.Seed3, OutCoord.Tri, OutCoord.BaryCoord); }");
-		OutHLSL += FString::Format(FormatSample, ArgsSample);
-	}
-	else if (FunctionInfo.DefinitionName == FSkeletalMeshInterfaceHelper::IsValidTriCoordName)
-	{
-		static const TCHAR* FormatSample = TEXT("void {InstanceFunctionName} (in {MeshTriCoordinateStructName} InCoord, out bool IsValid) { {GetDISkelMeshContextName} IsValid = InCoord.Tri < DIContext.MeshTriangleCount; }");
-		OutHLSL += FString::Format(FormatSample, ArgsSample);
-	}
-	else if (FunctionInfo.DefinitionName == FSkeletalMeshInterfaceHelper::GetTriangleDataName)
-	{
-		static const TCHAR* FormatSample = TEXT("void {InstanceFunctionName} (in {MeshTriCoordinateStructName} InCoord, out float3 OutPosition, out float3 OutNormal, out float3 OutBinormal, out float3 OutTangent) { {GetDISkelMeshContextName} DISkelMesh_GetPointOnTriangle(DIContext, InCoord.Tri, InCoord.BaryCoord, OutPosition, OutTangent, OutBinormal, OutNormal); }");
-		OutHLSL += FString::Format(FormatSample, ArgsSample);
-	}
-	else if (FunctionInfo.DefinitionName == FSkeletalMeshInterfaceHelper::GetSkinnedTriangleDataWSName)
-	{
-		static const TCHAR* FormatSample = TEXT("void {InstanceFunctionName} (in {MeshTriCoordinateStructName} InCoord, out float3 OutPosition, out float3 OutVelocity, out float3 OutNormal, out float3 OutBinormal, out float3 OutTangent) { {GetDISkelMeshContextName} DISKelMesh_GetSkinnedTriangleDataWS(DIContext, InCoord.Tri, InCoord.BaryCoord, OutPosition, OutVelocity, OutNormal, OutBinormal, OutTangent); }");
-		OutHLSL += FString::Format(FormatSample, ArgsSample);
-	}
-	else if (FunctionInfo.DefinitionName == FSkeletalMeshInterfaceHelper::GetSkinnedTriangleDataWSInterpName)
+	if ( ValidGpuFunctions.Contains(FunctionInfo.DefinitionName) )
 	{
-		static const TCHAR* FormatSample = TEXT("void {InstanceFunctionName} (in {MeshTriCoordinateStructName} InCoord, in float InInterp, out float3 OutPosition, out float3 OutVelocity, out float3 OutNormal, out float3 OutBinormal, out float3 OutTangent) { {GetDISkelMeshContextName} DISKelMesh_GetSkinnedTriangleDataInterpolatedWS(DIContext, InCoord.Tri, InCoord.BaryCoord, InInterp, OutPosition, OutVelocity, OutNormal, OutBinormal, OutTangent); }");
-		OutHLSL += FString::Format(FormatSample, ArgsSample);
-	}
-	else if (FunctionInfo.DefinitionName == FSkeletalMeshInterfaceHelper::GetSkinnedTriangleDataName)
-	{
-		static const TCHAR* FormatSample = TEXT("void {InstanceFunctionName} (in {MeshTriCoordinateStructName} InCoord, out float3 OutPosition, out float3 OutVelocity, out float3 OutNormal, out float3 OutBinormal, out float3 OutTangent) { {GetDISkelMeshContextName} DISKelMesh_GetSkinnedTriangleData(DIContext, InCoord.Tri, InCoord.BaryCoord, OutPosition, OutVelocity, OutNormal, OutBinormal, OutTangent); }");
-		OutHLSL += FString::Format(FormatSample, ArgsSample);
-	}
-	else if (FunctionInfo.DefinitionName == FSkeletalMeshInterfaceHelper::GetSkinnedTriangleDataInterpName)
-	{
-		static const TCHAR* FormatSample = TEXT("void {InstanceFunctionName} (in {MeshTriCoordinateStructName} InCoord, in float InInterp, out float3 OutPosition, out float3 OutVelocity, out float3 OutNormal, out float3 OutBinormal, out float3 OutTangent) { {GetDISkelMeshContextName} DISKelMesh_GetSkinnedTriangleDataInterpolated(DIContext, InCoord.Tri, InCoord.BaryCoord, InInterp, OutPosition, OutVelocity, OutNormal, OutBinormal, OutTangent); }");
-		OutHLSL += FString::Format(FormatSample, ArgsSample);
-	}
-	else if (FunctionInfo.DefinitionName == FSkeletalMeshInterfaceHelper::GetTriUVName)
-	{
-		static const TCHAR* FormatSample = TEXT("void {InstanceFunctionName} (in {MeshTriCoordinateStructName} InCoord, in int InUVSet, out float2 OutUV) { {GetDISkelMeshContextName} DISKelMesh_GetTriUV(DIContext, InCoord.Tri, InCoord.BaryCoord, InUVSet, OutUV); }");
-		OutHLSL += FString::Format(FormatSample, ArgsSample);
+		return true;
 	}
-	else if (FunctionInfo.DefinitionName == FSkeletalMeshInterfaceHelper::GetTriColorName)
+
+	//////////////////////////////////////////////////////////////////////////////////////////////////////////////////////////////
+	// Special cases due to using FMeshTriCoordinate
+	TMap<FString, FStringFormatArg> FormatArgs =
 	{
-		static const TCHAR* FormatSample = TEXT("void {InstanceFunctionName} (in {MeshTriCoordinateStructName} InCoord, out float4 OutColor) { {GetDISkelMeshContextName} DISkelMesh_GetTriColor(DIContext, InCoord.Tri, InCoord.BaryCoord, OutColor); }");
-		OutHLSL += FString::Format(FormatSample, ArgsSample);
-	}
-	else if (FunctionInfo.DefinitionName == FSkeletalMeshInterfaceHelper::GetTriCoordVerticesName)
+		{TEXT("InstanceFunctionName"),	FunctionInfo.InstanceName},
+		{TEXT("ParameterName"),			ParamInfo.DataInterfaceHLSLSymbol},
+	};
+
+	const TCHAR* FuncFormat = nullptr;
+	if (FunctionInfo.DefinitionName == FSkeletalMeshInterfaceHelper::IsValidTriCoordName)
 	{
-		static const TCHAR* FormatSample = TEXT("void {InstanceFunctionName} (in int TriangleIndex, out int OutVertexIndex0, out int OutVertexIndex1, out int OutVertexIndex2) { {GetDISkelMeshContextName} DISkelMesh_GetTriVertices(DIContext, TriangleIndex, OutVertexIndex0, OutVertexIndex1, OutVertexIndex2); }");
-		OutHLSL += FString::Format(FormatSample, ArgsSample);
+		FuncFormat = TEXT("void {InstanceFunctionName}(in MeshTriCoordinate InCoord, out bool bIsValid) { IsValidTriCoord_{ParameterName}(InCoord.Tri, bIsValid); }");
 	}
 	else if (FunctionInfo.DefinitionName == FSkeletalMeshInterfaceHelper::RandomTriangleName)
 	{
-		static const TCHAR* FormatSample = TEXT("void {InstanceFunctionName} (NiagaraRandInfo InRandomInfo, out {MeshTriCoordinateStructName} OutCoord) { {GetDISkelMeshContextName} DISKelMesh_RandomTriangle(DIContext, InRandomInfo.Seed1, InRandomInfo.Seed2, InRandomInfo.Seed3, OutCoord.Tri, OutCoord.BaryCoord); }");
-		OutHLSL += FString::Format(FormatSample, ArgsSample);
+		FuncFormat = TEXT("void {InstanceFunctionName}(NiagaraRandInfo InRandomInfo, out MeshTriCoordinate OutCoord) { RandomTriangle_{ParameterName}(InRandomInfo.Seed1, InRandomInfo.Seed2, InRandomInfo.Seed3, OutCoord.Tri, OutCoord.BaryCoord); }");
 	}
-	else if (FunctionInfo.DefinitionName == FSkeletalMeshInterfaceHelper::GetTriangleCountName)
+	else if (FunctionInfo.DefinitionName == FSkeletalMeshInterfaceHelper::RandomTriCoordName)
 	{
-		static const TCHAR* FormatSample = TEXT("void {InstanceFunctionName} (out int Count) { {GetDISkelMeshContextName} DISKelMesh_GetTriangleCount(DIContext, Count); }");
-		OutHLSL += FString::Format(FormatSample, ArgsSample);
+		FuncFormat = TEXT("void {InstanceFunctionName}(NiagaraRandInfo InRandomInfo, out MeshTriCoordinate OutCoord) { RandomTriangle_{ParameterName}(InRandomInfo.Seed1, InRandomInfo.Seed2, InRandomInfo.Seed3, OutCoord.Tri, OutCoord.BaryCoord); }");
 	}
 	else if (FunctionInfo.DefinitionName == FSkeletalMeshInterfaceHelper::RandomFilteredTriangleName)
 	{
-		static const TCHAR* FormatSample = TEXT("void {InstanceFunctionName} (NiagaraRandInfo InRandomInfo, out {MeshTriCoordinateStructName} OutCoord) { {GetDISkelMeshContextName} DISKelMesh_RandomFilteredTriangle(DIContext, InRandomInfo.Seed1, InRandomInfo.Seed2, InRandomInfo.Seed3, OutCoord.Tri, OutCoord.BaryCoord); }");
-		OutHLSL += FString::Format(FormatSample, ArgsSample);
-	}
-	else if (FunctionInfo.DefinitionName == FSkeletalMeshInterfaceHelper::GetFilteredTriangleCountName)
-	{
-		static const TCHAR* FormatSample = TEXT("void {InstanceFunctionName} (out int Count) { {GetDISkelMeshContextName} DISKelMesh_GetFilteredTriangleCount(DIContext, Count); }");
-		OutHLSL += FString::Format(FormatSample, ArgsSample);
+		FuncFormat = TEXT("void {InstanceFunctionName}(NiagaraRandInfo InRandomInfo, out MeshTriCoordinate OutCoord) { RandomFilteredTriangle_{ParameterName}(InRandomInfo.Seed1, InRandomInfo.Seed2, InRandomInfo.Seed3, OutCoord.Tri, OutCoord.BaryCoord); }");
 	}
 	else if (FunctionInfo.DefinitionName == FSkeletalMeshInterfaceHelper::GetFilteredTriangleAtName)
 	{
-		static const TCHAR* FormatSample = TEXT("void {InstanceFunctionName} (int FilteredIndex, in float3 BaryCoord, out {MeshTriCoordinateStructName} OutCoord) { {GetDISkelMeshContextName} DISKelMesh_GetFilteredTriangleAt(DIContext, FilteredIndex, OutCoord.Tri); OutCoord.BaryCoord = BaryCoord; }");
-		OutHLSL += FString::Format(FormatSample, ArgsSample);
-	}
-	//////////////////////////////////////////////////////////////////////////////////////////////////////////////////////////////
-	// Bone Sampling
-	else if (FunctionInfo.DefinitionName == FSkeletalMeshInterfaceHelper::GetSkinnedBoneDataName)
-	{
-		static const TCHAR* FormatSample = TEXT("void {InstanceFunctionName} (in int InBone, out float3 OutPosition, out float4 OutRotation, out float3 OutScale, out float3 OutVelocity) { {GetDISkelMeshContextName} DISkelMesh_GetSkinnedBone(DIContext, InBone, OutPosition, OutRotation, OutScale, OutVelocity); }");
-		OutHLSL += FString::Format(FormatSample, ArgsSample);
-	}
-	else if (FunctionInfo.DefinitionName == FSkeletalMeshInterfaceHelper::GetSkinnedBoneDataInterpolatedName)
-	{
-		static const TCHAR* FormatSample = TEXT("void {InstanceFunctionName} (in int InBone, in float Interp, out float3 OutPosition, out float4 OutRotation, out float3 OutScale, out float3 OutVelocity) { {GetDISkelMeshContextName} DISkelMesh_GetSkinnedBoneInterpolated(DIContext, InBone, Interp, OutPosition, OutRotation, OutScale, OutVelocity); }");
-		OutHLSL += FString::Format(FormatSample, ArgsSample);
-	}
-	else if (FunctionInfo.DefinitionName == FSkeletalMeshInterfaceHelper::GetSkinnedBoneDataWSName)
-	{
-		static const TCHAR* FormatSample = TEXT("void {InstanceFunctionName} (in int InBone, out float3 OutPosition, out float4 OutRotation, out float3 OutScale, out float3 OutVelocity) { {GetDISkelMeshContextName} DISkelMesh_GetSkinnedBoneWS(DIContext, InBone, OutPosition, OutRotation, OutScale, OutVelocity); }");
-		OutHLSL += FString::Format(FormatSample, ArgsSample);
-	}
-	else if (FunctionInfo.DefinitionName == FSkeletalMeshInterfaceHelper::GetSkinnedBoneDataWSInterpolatedName)
-	{
-		static const TCHAR* FormatSample = TEXT("void {InstanceFunctionName} (in int InBone, in float Interp, out float3 OutPosition, out float4 OutRotation, out float3 OutScale, out float3 OutVelocity) { {GetDISkelMeshContextName} DISkelMesh_GetSkinnedBoneInterpolatedWS(DIContext, InBone, Interp, OutPosition, OutRotation, OutScale, OutVelocity); }");
-		OutHLSL += FString::Format(FormatSample, ArgsSample);
-	}
-	//////////////////////////////////////////////////////////////////////////////////////////////////////////////////////////////
-	// Vertex Sampling
-	else if (FunctionInfo.DefinitionName == FSkeletalMeshInterfaceHelper::GetVertexDataName)
-	{
-		static const TCHAR* FormatSample = TEXT("void {InstanceFunctionName} (in int Vertex, out float3 OutPosition, out float3 OutNormal, out float3 OutBinormal, out float3 OutTangent) { {GetDISkelMeshContextName} DISkelMesh_GetVertex(DIContext, Vertex, OutPosition, OutTangent, OutBinormal, OutNormal); }");
-		OutHLSL += FString::Format(FormatSample, ArgsSample);
-	}
-	else if (FunctionInfo.DefinitionName == FSkeletalMeshInterfaceHelper::GetSkinnedVertexDataName)
-	{
-		static const TCHAR* FormatSample = TEXT("void {InstanceFunctionName} (in int Vertex, out float3 OutPosition, out float3 OutVelocity, out float3 OutNormal, out float3 OutBinormal, out float3 OutTangent) { {GetDISkelMeshContextName} DISkelMesh_GetSkinnedVertex(DIContext, Vertex, OutPosition, OutVelocity, OutNormal, OutBinormal, OutTangent); }");
-		OutHLSL += FString::Format(FormatSample, ArgsSample);
-	}
-	else if (FunctionInfo.DefinitionName == FSkeletalMeshInterfaceHelper::GetSkinnedVertexDataWSName)
-	{
-		static const TCHAR* FormatSample = TEXT("void {InstanceFunctionName} (in int Vertex, out float3 OutPosition, out float3 OutVelocity, out float3 OutNormal, out float3 OutBinormal, out float3 OutTangent) { {GetDISkelMeshContextName} DISkelMesh_GetSkinnedVertexWS(DIContext, Vertex, OutPosition, OutVelocity, OutNormal, OutBinormal, OutTangent); }");
-		OutHLSL += FString::Format(FormatSample, ArgsSample);
-	}
-	else if (FunctionInfo.DefinitionName == FSkeletalMeshInterfaceHelper::GetVertexColorName)
-	{
-		static const TCHAR* FormatSample = TEXT("void {InstanceFunctionName} (in int Vertex, out float4 OutColor) { {GetDISkelMeshContextName} DISkelMesh_GetVertexColor(DIContext, Vertex, OutColor); }");
-		OutHLSL += FString::Format(FormatSample, ArgsSample);
-	}
-	else if (FunctionInfo.DefinitionName == FSkeletalMeshInterfaceHelper::GetVertexUVName)
-	{
-		static const TCHAR* FormatSample = TEXT("void {InstanceFunctionName} (in int Vertex, in int UVSet, out float2 OutUV) { {GetDISkelMeshContextName} DISkelMesh_GetVertexUV(DIContext, Vertex, UVSet, OutUV); }");
-		OutHLSL += FString::Format(FormatSample, ArgsSample);
-	}
-	else if (FunctionInfo.DefinitionName == FSkeletalMeshInterfaceHelper::IsValidVertexName)
-	{
-		static const TCHAR* FormatSample = TEXT("void {InstanceFunctionName} (in int Vertex, out bool IsValid) { {GetDISkelMeshContextName} DISkelMesh_IsValidVertex(DIContext, Vertex, IsValid); }");
-		OutHLSL += FString::Format(FormatSample, ArgsSample);
-	}
-	else if (FunctionInfo.DefinitionName == FSkeletalMeshInterfaceHelper::RandomVertexName)
-	{
-		static const TCHAR* FormatSample = TEXT("void {InstanceFunctionName}(NiagaraRandInfo InRandomInfo, out int OutVertex) { {GetDISkelMeshContextName} DISkelMesh_GetRandomVertex(DIContext, InRandomInfo.Seed1, InRandomInfo.Seed2, InRandomInfo.Seed3, OutVertex); }");
-		OutHLSL += FString::Format(FormatSample, ArgsSample);
-	}
-	else if (FunctionInfo.DefinitionName == FSkeletalMeshInterfaceHelper::GetVertexCountName)
-	{
-		static const TCHAR* FormatSample = TEXT("void {InstanceFunctionName} (out int VertexCount) { {GetDISkelMeshContextName} DISkelMesh_GetVertexCount(DIContext, VertexCount); }");
-		OutHLSL += FString::Format(FormatSample, ArgsSample);
-	}
-	else if (FunctionInfo.DefinitionName == FSkeletalMeshInterfaceHelper::IsValidFilteredVertexName)
-	{
-		static const TCHAR* FormatSample = TEXT("void {InstanceFunctionName} (in int FilteredIndex, out bool IsValid) { {GetDISkelMeshContextName} DISkelMesh_IsValidFilteredVertex(DIContext, FilteredIndex, IsValid); }");
-		OutHLSL += FString::Format(FormatSample, ArgsSample);
+		FuncFormat = TEXT("void {InstanceFunctionName}(int FilteredIndex, in float3 BaryCoord, out MeshTriCoordinate OutCoord) { GetFilteredTriangleAt_{ParameterName}(FilteredIndex, OutCoord.Tri); OutCoord.BaryCoord = BaryCoord; }");
 	}
-	else if (FunctionInfo.DefinitionName == FSkeletalMeshInterfaceHelper::RandomFilteredVertexName)
-	{
-		static const TCHAR* FormatSample = TEXT("void {InstanceFunctionName}(NiagaraRandInfo InRandomInfo, out int OutVertex) { {GetDISkelMeshContextName} DISkelMesh_GetRandomFilteredVertex(DIContext, InRandomInfo.Seed1, InRandomInfo.Seed2, InRandomInfo.Seed3, OutVertex); }");
-		OutHLSL += FString::Format(FormatSample, ArgsSample);
-	}
-	else if (FunctionInfo.DefinitionName == FSkeletalMeshInterfaceHelper::GetFilteredVertexCountName)
-	{
-		static const TCHAR* FormatSample = TEXT("void {InstanceFunctionName} (out int VertexCount) { {GetDISkelMeshContextName} DISkelMesh_GetFilteredVertexCount(DIContext, VertexCount); }");
-		OutHLSL += FString::Format(FormatSample, ArgsSample);
-	}
-	else if (FunctionInfo.DefinitionName == FSkeletalMeshInterfaceHelper::GetFilteredVertexAtName)
-	{
-		static const TCHAR* FormatSample = TEXT("void {InstanceFunctionName} (in int FilteredIndex, out int VertexIndex) { {GetDISkelMeshContextName} DISkelMesh_GetFilteredVertexAt(DIContext, FilteredIndex, VertexIndex); }");
-		OutHLSL += FString::Format(FormatSample, ArgsSample);
-	}
-
-	//////////////////////////////////////////////////////////////////////////////////////////////////////////////////////////////
-	// Filtered Bone
-	else if (FunctionInfo.DefinitionName == FSkeletalMeshInterfaceHelper::IsValidBoneName)
-	{
-		static const TCHAR* FormatSample = TEXT("void {InstanceFunctionName} (in int BoneIndex, out bool IsValid) { {GetDISkelMeshContextName} DISkelMesh_IsValidBone(DIContext, BoneIndex, IsValid); }");
-		OutHLSL += FString::Format(FormatSample, ArgsSample);
-	}
-	else if (FunctionInfo.DefinitionName == FSkeletalMeshInterfaceHelper::RandomBoneName)
-	{
-		static const TCHAR* FormatSample = TEXT("void {InstanceFunctionName} (NiagaraRandInfo InRandomInfo, out int Bone) { {GetDISkelMeshContextName} DISkelMesh_RandomBone(DIContext, InRandomInfo.Seed1, InRandomInfo.Seed2, InRandomInfo.Seed3, Bone); }");
-		OutHLSL += FString::Format(FormatSample, ArgsSample);
-	}
-	else if (FunctionInfo.DefinitionName == FSkeletalMeshInterfaceHelper::GetBoneCountName)
-	{
-		static const TCHAR* FormatSample = TEXT("void {InstanceFunctionName} (out int Count) { {GetDISkelMeshContextName} DISkelMesh_GetBoneCount(DIContext, Count); }");
-		OutHLSL += FString::Format(FormatSample, ArgsSample);
-	}
-	else if (FunctionInfo.DefinitionName == FSkeletalMeshInterfaceHelper::GetParentBoneName)
-	{
-		static const TCHAR* FormatSample = TEXT("void {InstanceFunctionName} (int BoneIndex, out int ParentIndex) { {GetDISkelMeshContextName} DISkelMesh_GetParentBone(DIContext, BoneIndex, ParentIndex); }");
-		OutHLSL += FString::Format(FormatSample, ArgsSample);
-	}
-	else if (FunctionInfo.DefinitionName == FSkeletalMeshInterfaceHelper::GetFilteredBoneCountName)
-	{
-		static const TCHAR* FormatSample = TEXT("void {InstanceFunctionName} (out int Count) { {GetDISkelMeshContextName} DISkelMesh_GetFilteredBoneCount(DIContext, Count); }");
-		OutHLSL += FString::Format(FormatSample, ArgsSample);
-	}
-	else if (FunctionInfo.DefinitionName == FSkeletalMeshInterfaceHelper::GetFilteredBoneAtName)
-	{
-		static const TCHAR* FormatSample = TEXT("void {InstanceFunctionName} (in int BoneIndex, out int Bone) { {GetDISkelMeshContextName} DISkelMesh_GetFilteredBoneAt(DIContext, BoneIndex, Bone); }");
-		OutHLSL += FString::Format(FormatSample, ArgsSample);
-	}
-	else if (FunctionInfo.DefinitionName == FSkeletalMeshInterfaceHelper::RandomFilteredBoneName)
-	{
-		static const TCHAR* FormatSample = TEXT("void {InstanceFunctionName} (NiagaraRandInfo InRandomInfo, out int Bone) { {GetDISkelMeshContextName} DISkelMesh_RandomFilteredBone(DIContext, InRandomInfo.Seed1, InRandomInfo.Seed2, InRandomInfo.Seed3, Bone); }");
-		OutHLSL += FString::Format(FormatSample, ArgsSample);
-	}
-	else if (FunctionInfo.DefinitionName == FSkeletalMeshInterfaceHelper::GetUnfilteredBoneCountName)
-	{
-		static const TCHAR* FormatSample = TEXT("void {InstanceFunctionName} (out int Count) { {GetDISkelMeshContextName} DISkelMesh_GetUnfilteredBoneCount(DIContext, Count); }");
-		OutHLSL += FString::Format(FormatSample, ArgsSample);
-	}
-	else if (FunctionInfo.DefinitionName == FSkeletalMeshInterfaceHelper::GetUnfilteredBoneAtName)
-	{
-		static const TCHAR* FormatSample = TEXT("void {InstanceFunctionName} (in int BoneIndex, out int Bone) { {GetDISkelMeshContextName} DISkelMesh_GetUnfilteredBoneAt(DIContext, BoneIndex, Bone); }");
-		OutHLSL += FString::Format(FormatSample, ArgsSample);
-	}
-	else if (FunctionInfo.DefinitionName == FSkeletalMeshInterfaceHelper::RandomUnfilteredBoneName)
-	{
-		static const TCHAR* FormatSample = TEXT("void {InstanceFunctionName} (NiagaraRandInfo InRandomInfo, out int Bone) { {GetDISkelMeshContextName} DISkelMesh_RandomUnfilteredBone(DIContext, InRandomInfo.Seed1, InRandomInfo.Seed2, InRandomInfo.Seed3, Bone); }");
-		OutHLSL += FString::Format(FormatSample, ArgsSample);
-	}
-	//////////////////////////////////////////////////////////////////////////////////////////////////////////////////////////////
-	// Filtered Socket
-	else if (FunctionInfo.DefinitionName == FSkeletalMeshInterfaceHelper::GetFilteredSocketCountName)
+	else if (FunctionInfo.DefinitionName == FSkeletalMeshInterfaceHelper::GetTriangleDataName)
 	{
-		static const TCHAR* FormatSample = TEXT("void {InstanceFunctionName} (out int Count) { {GetDISkelMeshContextName} DISkelMesh_GetFilteredSocketCount(DIContext, Count); }");
-		OutHLSL += FString::Format(FormatSample, ArgsSample);
+		FuncFormat = TEXT("void {InstanceFunctionName}(in MeshTriCoordinate InCoord, out float3 OutPosition, out float3 OutNormal, out float3 OutBinormal, out float3 OutTangent) { GetPointOnTriangle_{ParameterName}(InCoord.Tri, InCoord.BaryCoord, OutPosition, OutTangent, OutBinormal, OutNormal); }");
 	}
-	else if (FunctionInfo.DefinitionName == FSkeletalMeshInterfaceHelper::GetFilteredSocketBoneAtName)
+	else if (FunctionInfo.DefinitionName == FSkeletalMeshInterfaceHelper::GetSkinnedTriangleDataWSName)
 	{
-		static const TCHAR* FormatSample = TEXT("void {InstanceFunctionName} (in int SocketIndex, out int Bone) { {GetDISkelMeshContextName} DISkelMesh_GetFilteredSocketBoneAt(DIContext, SocketIndex, Bone); }");
-		OutHLSL += FString::Format(FormatSample, ArgsSample);
+		FuncFormat = TEXT("void {InstanceFunctionName}(in MeshTriCoordinate InCoord, out float3 OutPosition, out float3 OutVelocity, out float3 OutNormal, out float3 OutBinormal, out float3 OutTangent) { GetSkinnedTriangleDataWS_{ParameterName}(InCoord.Tri, InCoord.BaryCoord, OutPosition, OutVelocity, OutNormal, OutBinormal, OutTangent); }");
 	}
-	else if (FunctionInfo.DefinitionName == FSkeletalMeshInterfaceHelper::GetFilteredSocketTransformName)
+	else if (FunctionInfo.DefinitionName == FSkeletalMeshInterfaceHelper::GetSkinnedTriangleDataWSInterpName)
 	{
-		static const TCHAR* FormatSample = TEXT("void {InstanceFunctionName} (in int SocketIndex, in bool bApplyTransform, out float3 OutTranslation, out float4 OutRotation, out float3 OutScale) { {GetDISkelMeshContextName} DISkelMesh_GetFilteredSocketTransform(DIContext, SocketIndex, bApplyTransform, OutTranslation, OutRotation, OutScale); }");
-		OutHLSL += FString::Format(FormatSample, ArgsSample);
+		FuncFormat = TEXT("void {InstanceFunctionName}(in MeshTriCoordinate InCoord, in float InInterp, out float3 OutPosition, out float3 OutVelocity, out float3 OutNormal, out float3 OutBinormal, out float3 OutTangent) { GetSkinnedTriangleDataInterpolatedWS_{ParameterName}(InCoord.Tri, InCoord.BaryCoord, InInterp, OutPosition, OutVelocity, OutNormal, OutBinormal, OutTangent); }");
 	}
-	else if (FunctionInfo.DefinitionName == FSkeletalMeshInterfaceHelper::RandomFilteredSocketName)
+	else if (FunctionInfo.DefinitionName == FSkeletalMeshInterfaceHelper::GetSkinnedTriangleDataName)
 	{
-		static const TCHAR* FormatSample = TEXT("void {InstanceFunctionName} (NiagaraRandInfo InRandomInfo, out int SocketBone) { {GetDISkelMeshContextName} DISkelMesh_RandomFilteredSocket(DIContext, InRandomInfo.Seed1, InRandomInfo.Seed2, InRandomInfo.Seed3, SocketBone); }");
-		OutHLSL += FString::Format(FormatSample, ArgsSample);
+		FuncFormat = TEXT("void {InstanceFunctionName}(in MeshTriCoordinate InCoord, out float3 OutPosition, out float3 OutVelocity, out float3 OutNormal, out float3 OutBinormal, out float3 OutTangent) { GetSkinnedTriangleData_{ParameterName}(InCoord.Tri, InCoord.BaryCoord, OutPosition, OutVelocity, OutNormal, OutBinormal, OutTangent); }");
 	}
-	//////////////////////////////////////////////////////////////////////////////////////////////////////////////////////////////
-	// Misc bone functions
-	else if (FunctionInfo.DefinitionName == FSkeletalMeshInterfaceHelper::RandomFilteredSocketOrBoneName)
+	else if (FunctionInfo.DefinitionName == FSkeletalMeshInterfaceHelper::GetSkinnedTriangleDataInterpName)
 	{
-		static const TCHAR* FormatSample = TEXT("void {InstanceFunctionName} (NiagaraRandInfo InRandomInfo, out int Bone) { {GetDISkelMeshContextName} DISkelMesh_RandomFilteredSocketOrBone(DIContext, InRandomInfo.Seed1, InRandomInfo.Seed2, InRandomInfo.Seed3, Bone); }");
-		OutHLSL += FString::Format(FormatSample, ArgsSample);
+		FuncFormat = TEXT("void {InstanceFunctionName}(in MeshTriCoordinate InCoord, in float InInterp, out float3 OutPosition, out float3 OutVelocity, out float3 OutNormal, out float3 OutBinormal, out float3 OutTangent) { GetSkinnedTriangleDataInterpolated_{ParameterName}(InCoord.Tri, InCoord.BaryCoord, InInterp, OutPosition, OutVelocity, OutNormal, OutBinormal, OutTangent); }");
 	}
-	else if (FunctionInfo.DefinitionName == FSkeletalMeshInterfaceHelper::GetFilteredSocketOrBoneCountName)
+	else if (FunctionInfo.DefinitionName == FSkeletalMeshInterfaceHelper::GetTriUVName)
 	{
-		static const TCHAR* FormatSample = TEXT("void {InstanceFunctionName} (out int Count) { {GetDISkelMeshContextName} DISkelMesh_GetFilteredSocketOrBoneCount(DIContext, Count); }");
-		OutHLSL += FString::Format(FormatSample, ArgsSample);
+		FuncFormat = TEXT("void {InstanceFunctionName}(in MeshTriCoordinate InCoord, in int InUVSet, out float2 OutUV) { GetTriUV_{ParameterName}(InCoord.Tri, InCoord.BaryCoord, InUVSet, OutUV); }");
 	}
-	else if (FunctionInfo.DefinitionName == FSkeletalMeshInterfaceHelper::GetFilteredSocketOrBoneAtName)
+	else if (FunctionInfo.DefinitionName == FSkeletalMeshInterfaceHelper::GetTriColorName)
 	{
-		static const TCHAR* FormatSample = TEXT("void {InstanceFunctionName} (in int FilteredIndex, out int Bone) { {GetDISkelMeshContextName} DISkelMesh_GetFilteredSocketOrBoneAt(DIContext, FilteredIndex, Bone); }");
-		OutHLSL += FString::Format(FormatSample, ArgsSample);
+		FuncFormat = TEXT("void {InstanceFunctionName}(in MeshTriCoordinate InCoord, out float4 OutColor) { GetTriColor_{ParameterName}(InCoord.Tri, InCoord.BaryCoord, OutColor); }");
 	}
-	//////////////////////////////////////////////////////////////////////////////////////////////////////////////////////////////
-	// Uv Mapping
 	else if (FunctionInfo.DefinitionName == FSkeletalMeshInterfaceHelper::GetTriangleCoordAtUVName)
 	{
-		static const TCHAR* FormatSample = TEXT("void {InstanceFunctionName} (in bool InEnabled, in float2 InUV, in float InTolerance, out {MeshTriCoordinateStructName} OutCoord, out bool OutIsValid) { {GetDISkelMeshContextName} DISkelMesh_GetTriangleCoordAtUV(DIContext, InEnabled, InUV, InTolerance, OutCoord.Tri, OutCoord.BaryCoord, OutIsValid); }");
-		OutHLSL += FString::Format(FormatSample, ArgsSample);
+		FuncFormat = TEXT("void {InstanceFunctionName}(in bool InEnabled, in float2 InUV, in float InTolerance, out MeshTriCoordinate OutCoord, out bool OutIsValid) { GetTriangleCoordAtUV_{ParameterName}(InEnabled, InUV, InTolerance, OutCoord.Tri, OutCoord.BaryCoord, OutIsValid); }");
 	}
 	else if (FunctionInfo.DefinitionName == FSkeletalMeshInterfaceHelper::GetTriangleCoordInAabbName)
 	{
-		static const TCHAR* FormatSample = TEXT("void {InstanceFunctionName} (in bool InEnabled, in float2 InUvMin, in float2 InUvMax, out {MeshTriCoordinateStructName} OutCoord, out bool OutIsValid) { {GetDISkelMeshContextName} DISkelMesh_GetTriangleCoordInAabb(DIContext, InEnabled, InUvMin, InUvMax, OutCoord.Tri, OutCoord.BaryCoord, OutIsValid); }");
-		OutHLSL += FString::Format(FormatSample, ArgsSample);
-	}
-	//////////////////////////////////////////////////////////////////////////////////////////////////////////////////////////////
-	// Adjacency
-	else if (FunctionInfo.DefinitionName == FSkeletalMeshInterfaceHelper::GetAdjacentTriangleIndexName)
-	{
-		static const TCHAR* FormatSample = TEXT("void {InstanceFunctionName} (in int VertexId, in int AdjacencyIndex, out int TriangleIndex, out bool IsValid) { {GetDISkelMeshContextName} DISkelMesh_GetAdjacentTriangleIndex(DIContext, VertexId, max(0, AdjacencyIndex), TriangleIndex); IsValid = TriangleIndex != -1; }");
-		OutHLSL += FString::Format(FormatSample, ArgsSample);
+		FuncFormat = TEXT("void {InstanceFunctionName}(in bool InEnabled, in float2 InUvMin, in float2 InUvMax, out MeshTriCoordinate OutCoord, out bool OutIsValid) { GetTriangleCoordInAabb_{ParameterName}(InEnabled, InUvMin, InUvMax, OutCoord.Tri, OutCoord.BaryCoord, OutIsValid); }");
 	}
-	else if (FunctionInfo.DefinitionName == FSkeletalMeshInterfaceHelper::GetTriangleNeighborName)
+
+	if ( FuncFormat != nullptr )
 	{
-		static const TCHAR* FormatSample = TEXT("void {InstanceFunctionName} (in int TriangleIndex, in int EdgeIndex, out int NeighborTriangleIndex, out int NeighborEdgeIndex, out bool IsValid) { {GetDISkelMeshContextName} DISkelMesh_GetTriangleNeighbor(DIContext, TriangleIndex, EdgeIndex, NeighborTriangleIndex, NeighborEdgeIndex); IsValid = NeighborTriangleIndex != -1;} ");
-		OutHLSL += FString::Format(FormatSample, ArgsSample);
+		OutHLSL += FString::Format(FuncFormat, FormatArgs);
+		return true;
 	}
+
 	//////////////////////////////////////////////////////////////////////////////////////////////////////////////////////////////
 	// Unsupported functionality
-	else
-	{
-		// This function is not support
-		return false;
-	}
-
-	OutHLSL += TEXT("\n");
-	return true;
+	return false;
 }
 
 bool UNiagaraDataInterfaceSkeletalMesh::UpgradeFunctionCall(FNiagaraFunctionSignature& FunctionSignature)
@@ -3614,12 +3235,183 @@ bool UNiagaraDataInterfaceSkeletalMesh::UpgradeFunctionCall(FNiagaraFunctionSign
 
 	return bWasChanged;
 }
+#endif
 
-void UNiagaraDataInterfaceSkeletalMesh::GetParameterDefinitionHLSL(const FNiagaraDataInterfaceGPUParamInfo& ParamInfo, FString& OutHLSL)
+void UNiagaraDataInterfaceSkeletalMesh::BuildShaderParameters(FNiagaraShaderParametersBuilder& ShaderParametersBuilder) const
 {
-	OutHLSL += TEXT("DISKELMESH_DECLARE_CONSTANTS(") + ParamInfo.DataInterfaceHLSLSymbol + TEXT(")\n");
+	ShaderParametersBuilder.AddNestedStruct<NDISkelMeshLocal::FShaderParameters>();
+}
+
+void UNiagaraDataInterfaceSkeletalMesh::SetShaderParameters(const FNiagaraDataInterfaceSetShaderParametersContext& Context) const
+{
+	FNiagaraDataInterfaceProxySkeletalMesh& DIProxy = Context.GetProxy<FNiagaraDataInterfaceProxySkeletalMesh>();
+	FNiagaraDataInterfaceProxySkeletalMeshData* InstanceData = DIProxy.SystemInstancesToData.Find(Context.GetSystemInstanceID());
+
+	NDISkelMeshLocal::FShaderParameters* ShaderParameters = Context.GetParameterNestedStruct<NDISkelMeshLocal::FShaderParameters>();
+	if (InstanceData && InstanceData->StaticBuffers)
+	{
+		FSkeletalMeshGpuSpawnStaticBuffers* StaticBuffers = InstanceData->StaticBuffers;
+
+		// Bind mesh buffers
+		ShaderParameters->MeshVertexBuffer = StaticBuffers->GetBufferPositionSRV();
+		ShaderParameters->MeshIndexBuffer = StaticBuffers->GetBufferIndexSRV();
+		ShaderParameters->MeshTangentBuffer = StaticBuffers->GetBufferTangentSRV();
+		ShaderParameters->MeshTexCoordBuffer = StaticBuffers->GetBufferTexCoordSRV();
+		ShaderParameters->MeshColorBuffer = StaticBuffers->GetBufferColorSRV();
+
+		ShaderParameters->MeshNumTexCoord = StaticBuffers->GetNumTexCoord();
+		ShaderParameters->MeshTriangleCount = StaticBuffers->GetTriangleCount();
+		ShaderParameters->MeshVertexCount = StaticBuffers->GetVertexCount();
+
+		// Set triangle sampling buffer
+		ShaderParameters->MeshTriangleSamplerProbAliasBuffer = StaticBuffers->GetBufferTriangleUniformSamplerProbAliasSRV();
+
+		// Set triangle sampling region buffer
+		ShaderParameters->MeshNumSamplingRegionTriangles = StaticBuffers->GetNumSamplingRegionTriangles();
+		ShaderParameters->MeshNumSamplingRegionVertices = StaticBuffers->GetNumSamplingRegionVertices();
+		ShaderParameters->MeshSamplingRegionsProbAliasBuffer = StaticBuffers->GetSampleRegionsProbAliasSRV();
+		ShaderParameters->MeshSampleRegionsTriangleIndices = StaticBuffers->GetSampleRegionsTriangleIndicesSRV();
+		ShaderParameters->MeshSampleRegionsVertices = StaticBuffers->GetSampleRegionsVerticesSRV();
+
+		ShaderParameters->MeshSkinWeightBuffer = FNiagaraRenderer::GetSrvOrDefaultUInt(InstanceData->MeshSkinWeightBuffer->GetSRV());
+		ShaderParameters->MeshSkinWeightLookupBuffer = FNiagaraRenderer::GetSrvOrDefaultUInt(InstanceData->MeshSkinWeightLookupBuffer->GetSRV());
+
+		ShaderParameters->MeshWeightStride = InstanceData->MeshWeightStrideByte / 4;
+		ShaderParameters->MeshSkinWeightIndexSize = InstanceData->MeshSkinWeightIndexSizeByte;
+
+		uint32 EnabledFeaturesBits = 0;
+		EnabledFeaturesBits |= StaticBuffers->IsUseGpuUniformlyDistributedSampling() ? 1 : 0;
+		EnabledFeaturesBits |= StaticBuffers->IsSamplingRegionsAllAreaWeighted() ? 2 : 0;
+		EnabledFeaturesBits |= InstanceData->bUnlimitedBoneInfluences ? 4 : 0;
+		EnabledFeaturesBits |= StaticBuffers->HasMeshColors() ? 8 : 0;
+
+		FSkeletalMeshGpuDynamicBufferProxy* DynamicBuffers = InstanceData->DynamicBuffer;
+		check(DynamicBuffers);
+		if (DynamicBuffers->DoesBoneDataExist())
+		{
+			ShaderParameters->MeshNumWeights = StaticBuffers->GetNumWeights();
+			ShaderParameters->MeshCurrBonesBuffer = DynamicBuffers->GetRWBufferBone().SectionSRV;
+			ShaderParameters->MeshPrevBonesBuffer = DynamicBuffers->GetRWBufferPrevBone().SectionSRV;
+			ShaderParameters->MeshCurrSamplingBonesBuffer = DynamicBuffers->GetRWBufferBone().SamplingSRV;
+			ShaderParameters->MeshPrevSamplingBonesBuffer = DynamicBuffers->GetRWBufferPrevBone().SamplingSRV;
+			ShaderParameters->MeshTriangleMatricesOffsetBuffer = StaticBuffers->GetBufferTriangleMatricesOffsetSRV();
+		}
+		// Bind dummy data for validation purposes only.  Code will not execute due to "EnabledFeatures" bits but validation can not determine that.
+		else
+		{
+			ShaderParameters->MeshNumWeights = 0;
+			ShaderParameters->MeshCurrBonesBuffer = FNiagaraRenderer::GetDummyFloat4Buffer();
+			ShaderParameters->MeshPrevBonesBuffer = FNiagaraRenderer::GetDummyFloat4Buffer();
+			ShaderParameters->MeshCurrSamplingBonesBuffer = FNiagaraRenderer::GetDummyFloat4Buffer();
+			ShaderParameters->MeshPrevSamplingBonesBuffer = FNiagaraRenderer::GetDummyFloat4Buffer();
+			ShaderParameters->MeshTriangleMatricesOffsetBuffer = FNiagaraRenderer::GetDummyUIntBuffer();
+		}
+
+		ShaderParameters->NumBones = DynamicBuffers->GetNumBones();
+
+		ShaderParameters->NumFilteredBones = StaticBuffers->GetNumFilteredBones();
+		ShaderParameters->NumUnfilteredBones = StaticBuffers->GetNumUnfilteredBones();
+		ShaderParameters->RandomMaxBone = StaticBuffers->GetExcludedBoneIndex() >= 0 ? DynamicBuffers->GetNumBones() - 2 : DynamicBuffers->GetNumBones() - 1;
+		ShaderParameters->ExcludeBoneIndex = StaticBuffers->GetExcludedBoneIndex();
+		ShaderParameters->FilteredAndUnfilteredBones = StaticBuffers->GetFilteredAndUnfilteredBonesSRV();
+
+		ShaderParameters->NumFilteredSockets = StaticBuffers->GetNumFilteredSockets();
+		ShaderParameters->FilteredSocketBoneOffset = StaticBuffers->GetFilteredSocketBoneOffset();
+
+		if (InstanceData->UvMappingBuffer)
+		{
+			ShaderParameters->UvMappingBuffer = InstanceData->UvMappingBuffer->GetSrv();
+			ShaderParameters->UvMappingBufferLength = InstanceData->UvMappingBuffer->GetBufferSize();
+			ShaderParameters->UvMappingSet = InstanceData->UvMappingSet;
+		}
+		else
+		{
+			ShaderParameters->UvMappingBuffer = FNiagaraRenderer::GetDummyIntBuffer();
+			ShaderParameters->UvMappingBufferLength = 0;
+			ShaderParameters->UvMappingSet = 0;
+		}
+
+		if (InstanceData->ConnectivityBuffer)
+		{
+			const uint32 NumBufferElements = FMath::DivideAndRoundUp<uint32>(InstanceData->ConnectivityBuffer->GetBufferSize(), sizeof(uint32));
+			ShaderParameters->ConnectivityBuffer = InstanceData->ConnectivityBuffer->GetSrv();
+			ShaderParameters->ConnectivityBufferLength = NumBufferElements;
+			ShaderParameters->ConnectivityMaxAdjacentPerVertex = InstanceData->ConnectivityBuffer->MaxAdjacentTriangleCount;
+		}
+		else
+		{
+			ShaderParameters->ConnectivityBuffer = FNiagaraRenderer::GetDummyUIntBuffer();
+			ShaderParameters->ConnectivityBufferLength = 0;
+			ShaderParameters->ConnectivityMaxAdjacentPerVertex = 0;
+		}
+
+		ShaderParameters->InstanceTransform = InstanceData->Transform;
+		ShaderParameters->InstancePrevTransform = InstanceData->PrevTransform;
+		ShaderParameters->InstanceRotation = InstanceData->Transform.GetMatrixWithoutScale().ToQuat();
+		ShaderParameters->InstancePrevRotation = InstanceData->PrevTransform.GetMatrixWithoutScale().ToQuat();
+		ShaderParameters->InstanceInvDeltaTime = 1.0f / InstanceData->DeltaSeconds;
+
+		ShaderParameters->EnabledFeatures = EnabledFeaturesBits;
+	}
+	else
+	{
+		// Bind dummy buffers
+		ShaderParameters->MeshVertexBuffer = FNiagaraRenderer::GetDummyFloatBuffer();
+		ShaderParameters->MeshIndexBuffer = FNiagaraRenderer::GetDummyUIntBuffer();
+		ShaderParameters->MeshTangentBuffer = FNiagaraRenderer::GetDummyFloat4Buffer();
+
+		ShaderParameters->MeshNumTexCoord = 0;
+		ShaderParameters->MeshTexCoordBuffer = FNiagaraRenderer::GetDummyFloat2Buffer();
+		ShaderParameters->MeshColorBuffer = FNiagaraRenderer::GetDummyFloat4Buffer();
+		ShaderParameters->MeshTriangleCount = 0;
+		ShaderParameters->MeshVertexCount = 0;
+		ShaderParameters->MeshTriangleSamplerProbAliasBuffer = FNiagaraRenderer::GetDummyUIntBuffer();
+
+		ShaderParameters->MeshNumSamplingRegionTriangles = 0;
+		ShaderParameters->MeshNumSamplingRegionVertices = 0;
+		ShaderParameters->MeshSamplingRegionsProbAliasBuffer = FNiagaraRenderer::GetDummyUIntBuffer();
+		ShaderParameters->MeshSampleRegionsTriangleIndices = FNiagaraRenderer::GetDummyUIntBuffer();
+		ShaderParameters->MeshSampleRegionsVertices = FNiagaraRenderer::GetDummyUIntBuffer();
+
+		ShaderParameters->MeshSkinWeightBuffer = FNiagaraRenderer::GetDummyUIntBuffer();
+		ShaderParameters->MeshSkinWeightLookupBuffer = FNiagaraRenderer::GetDummyUIntBuffer();
+
+		ShaderParameters->MeshWeightStride = 0;
+		ShaderParameters->MeshSkinWeightIndexSize = 0;
+		ShaderParameters->MeshNumWeights = 0;
+
+		ShaderParameters->MeshCurrBonesBuffer = FNiagaraRenderer::GetDummyFloat4Buffer();
+		ShaderParameters->MeshPrevBonesBuffer = FNiagaraRenderer::GetDummyFloat4Buffer();
+		ShaderParameters->MeshCurrSamplingBonesBuffer = FNiagaraRenderer::GetDummyFloat4Buffer();
+		ShaderParameters->MeshPrevSamplingBonesBuffer = FNiagaraRenderer::GetDummyFloat4Buffer();
+		ShaderParameters->MeshTriangleMatricesOffsetBuffer = FNiagaraRenderer::GetDummyUIntBuffer();
+
+		ShaderParameters->NumBones = 0;
+		ShaderParameters->NumFilteredBones = 0;
+		ShaderParameters->NumUnfilteredBones = 0;
+		ShaderParameters->RandomMaxBone = 0;
+		ShaderParameters->ExcludeBoneIndex = 0;
+		ShaderParameters->FilteredAndUnfilteredBones = FNiagaraRenderer::GetDummyUIntBuffer();
+		ShaderParameters->NumFilteredSockets = 0;
+		ShaderParameters->FilteredSocketBoneOffset = 0;
+
+		ShaderParameters->UvMappingBuffer = FNiagaraRenderer::GetDummyIntBuffer();
+		ShaderParameters->UvMappingBufferLength = 0;
+		ShaderParameters->UvMappingSet = 0;
+
+		ShaderParameters->ConnectivityBuffer = FNiagaraRenderer::GetDummyUIntBuffer();
+		ShaderParameters->ConnectivityBufferLength = 0;
+		ShaderParameters->ConnectivityMaxAdjacentPerVertex = 0;
+
+		ShaderParameters->InstanceTransform = FMatrix44f::Identity;
+		ShaderParameters->InstancePrevTransform = FMatrix44f::Identity;
+		ShaderParameters->InstanceRotation = FQuat4f::Identity;
+		ShaderParameters->InstancePrevRotation = FQuat4f::Identity;
+		ShaderParameters->InstanceInvDeltaTime = 0.0f;
+
+		ShaderParameters->EnabledFeatures = 0;
+	}
 }
-#endif
 
 void UNiagaraDataInterfaceSkeletalMesh::SetSourceComponentFromBlueprints(USkeletalMeshComponent* ComponentToUse)
 {
