- Replace Triangle with TriangleIndex as Triangle is a reserved keyword

#rnx

diff --git a/Engine/Plugins/FX/Niagara/Shaders/Private/NiagaraDataInterfaceSkeletalMeshTemplate.ush b/Engine/Plugins/FX/Niagara/Shaders/Private/NiagaraDataInterfaceSkeletalMeshTemplate.ush
--- a/Engine/Plugins/FX/Niagara/Shaders/Private/NiagaraDataInterfaceSkeletalMeshTemplate.ush
+++ b/Engine/Plugins/FX/Niagara/Shaders/Private/NiagaraDataInterfaceSkeletalMeshTemplate.ush
@@ -87,9 +87,9 @@ float3x4 GetCurrSkinningMatrix_{ParameterName}(uint SectionBoneOffset, int4 Blen
 	return Result;
 }
 
-void IsValidTriCoord_{ParameterName}(in int Triangle, out bool IsValid)
+void IsValidTriCoord_{ParameterName}(in int TriangleIndex, out bool IsValid)
 {
-	IsValid = Triangle >= 0 && Triangle < {ParameterName}_MeshTriangleCount;
+	IsValid = TriangleIndex >= 0 && TriangleIndex < {ParameterName}_MeshTriangleCount;
 }
 
 void GetVertex_{ParameterName}(uint VertexIndex, out float3 OutPosition, out float3 OutTangentX, out float3 OutTangentY, out float3 OutTangentZ)
