- Add information to check when the system simulation becomes nullptr

#rnx

diff --git a/Engine/Plugins/FX/Niagara/Source/Niagara/Private/NiagaraSystemInstance.cpp b/Engine/Plugins/FX/Niagara/Source/Niagara/Private/NiagaraSystemInstance.cpp
--- a/Engine/Plugins/FX/Niagara/Source/Niagara/Private/NiagaraSystemInstance.cpp
+++ b/Engine/Plugins/FX/Niagara/Source/Niagara/Private/NiagaraSystemInstance.cpp
@@ -2576,7 +2576,13 @@ void FNiagaraSystemInstance::FinalizeTick_GameThread(bool bEnqueueGPUTickIfNeede
 		//Enqueue a GPU tick for this sim if we have to do this from the GameThread.
 		//If we're batching our tick passing we may still need to enqueue here if not called from the regular finalize task. The caller will tell us with bEnqueueGPUTickIfNeeded.
 		FNiagaraSystemSimulation* Sim = SystemSimulation.Get();
-		check(Sim);
+		checkf(Sim != nullptr,
+			TEXT("SystemSimulation is nullptr during Finalize and the Asset(%s) AttachComponent(%s) ActualExecutionState(%d) SystemInstanceIndex(%d) SystemInstanceState(%s)"),
+			*GetFullNameSafe(Asset.Get()), *GetFullNameSafe(AttachComponent.Get()), ActualExecutionState, SystemInstanceIndex,
+			*StaticEnum<ENiagaraSystemInstanceState>()->GetValueAsString(SystemInstanceState)
+		);
+		
+
 		ENiagaraGPUTickHandlingMode Mode = Sim->GetGPUTickHandlingMode();
 		if (Mode == ENiagaraGPUTickHandlingMode::GameThread || (Mode == ENiagaraGPUTickHandlingMode::GameThreadBatched && bEnqueueGPUTickIfNeeded))
 		{
