Fixed a bug where several script source copies were left in the root set after compiling a niagara system

diff --git a/Engine/Plugins/FX/Niagara/Source/NiagaraEditor/Private/NiagaraCompiler.cpp b/Engine/Plugins/FX/Niagara/Source/NiagaraEditor/Private/NiagaraCompiler.cpp
--- a/Engine/Plugins/FX/Niagara/Source/NiagaraEditor/Private/NiagaraCompiler.cpp
+++ b/Engine/Plugins/FX/Niagara/Source/NiagaraEditor/Private/NiagaraCompiler.cpp
@@ -186,6 +186,7 @@ void FNiagaraCompileRequestDuplicateData::DuplicateReferencedGraphsRecursive(UNi
 						UNiagaraScript* DupeScript = FunctionScript->CreateCompilationCopy();
 						TArray<ENiagaraScriptUsage> CompileUsages = { DupeScript->GetUsage() };
 						UNiagaraScriptSource* DupeScriptSource = CastChecked<UNiagaraScriptSource>(DupeScript->GetSource(FunctionCallNode->SelectedScriptVersion))->CreateCompilationCopy(CompileUsages);
+						TrackedScriptSourceCopies.Add(DupeScriptSource);
 						UNiagaraGraph* DupeGraph = DupeScriptSource->NodeGraph;
 						DupeScript->SetSource(DupeScriptSource, FunctionCallNode->SelectedScriptVersion);
 						
@@ -446,12 +447,26 @@ void FNiagaraCompileRequestDuplicateData::ReleaseCompilationCopies()
 		SharedSourceGraphToDuplicatedGraphsMap->Empty();
 	}
 
-	// clean up script source
+	// clean up script sources
+	for (TWeakObjectPtr<UNiagaraScriptSource> Source : TrackedScriptSourceCopies)
+	{
+		if (Source.IsValid())
+		{
+			Source->ReleaseCompilationCopy();
+		}
+	}
+	TrackedScriptSourceCopies.Empty();
 	if (SourceDeepCopy.IsValid())
 	{
 		SourceDeepCopy->ReleaseCompilationCopy();
 	}
 	SourceDeepCopy = nullptr;
+
+	// clean up emitter data
+	for (TSharedPtr<FNiagaraCompileRequestDuplicateData, ESPMode::ThreadSafe> EmitterRequest : EmitterData)
+	{
+		EmitterRequest->ReleaseCompilationCopies();
+	}
 }
 
 void FNiagaraCompileRequestData::CompareAgainst(FNiagaraGraphCachedBuiltHistory* InCachedDataBase)
diff --git a/Engine/Plugins/FX/Niagara/Source/NiagaraEditor/Private/NiagaraHlslTranslator.h b/Engine/Plugins/FX/Niagara/Source/NiagaraEditor/Private/NiagaraHlslTranslator.h
--- a/Engine/Plugins/FX/Niagara/Source/NiagaraEditor/Private/NiagaraHlslTranslator.h
+++ b/Engine/Plugins/FX/Niagara/Source/NiagaraEditor/Private/NiagaraHlslTranslator.h
@@ -204,6 +204,8 @@ public:
 
 protected:
 	void DuplicateReferencedGraphsRecursive(UNiagaraGraph* InGraph, const FCompileConstantResolver& ConstantResolver, TMap<UNiagaraNodeFunctionCall*, ENiagaraScriptUsage> FunctionsWithUsage);
+private:
+	TArray<TWeakObjectPtr<UNiagaraScriptSource>> TrackedScriptSourceCopies;
 };
 
 
