- When loading DDC data ensure we find the LWC struct for any SWC structs so we can create the SWC struct again
- Thanks to [REDACTED_PERSON] for prodiving this fix

#rnx

diff --git a/Engine/Plugins/FX/Niagara/Source/Niagara/Private/NiagaraScript.cpp b/Engine/Plugins/FX/Niagara/Source/Niagara/Private/NiagaraScript.cpp
--- a/Engine/Plugins/FX/Niagara/Source/Niagara/Private/NiagaraScript.cpp
+++ b/Engine/Plugins/FX/Niagara/Source/Niagara/Private/NiagaraScript.cpp
@@ -2409,6 +2409,51 @@ static bool ValidateExecData(const UNiagaraScript* Script, const FNiagaraVMExecu
 	return IsValid;
 }
 
+struct FObjectAndNameAsStringProxyArchive_NiagaraExecData : public FObjectAndNameAsStringProxyArchive
+{
+	FObjectAndNameAsStringProxyArchive_NiagaraExecData(FArchive& InInnerArchive, bool bInLoadIfFindFails)
+		: FObjectAndNameAsStringProxyArchive(InInnerArchive, bInLoadIfFindFails)
+	{ }
+
+	virtual FArchive& operator<<(UObject*& Obj) override
+	{
+		if (IsLoading())
+		{
+			// load the path name to the object
+			FString LoadedString;
+			InnerArchive << LoadedString;
+			// look up the object by fully qualified pathname
+			Obj = FindObject<UObject>(nullptr, *LoadedString, false);
+			// If we couldn't find it, and we want to load it, do that
+			if (!Obj)
+			{
+				if (bLoadIfFindFails)
+				{
+					Obj = LoadObject<UObject>(nullptr, *LoadedString);
+				}
+				if (!Obj)
+				{
+					if (LoadedString.EndsWith(TEXT("_SWC")))
+					{
+						FString LWCShortName = FPackageName::ObjectPathToObjectName(LoadedString).LeftChop(4);
+						if (UScriptStruct* LWCStruct = FindFirstObject<UScriptStruct>(*LWCShortName))
+						{
+							Obj = FNiagaraTypeHelper::GetSWCStruct(LWCStruct);
+						}
+					}
+				}
+			}
+		}
+		else
+		{
+			// save out the fully qualified object name
+			FString SavedString(Obj->GetPathName());
+			InnerArchive << SavedString;
+		}
+		return *this;
+	}
+};
+
 bool UNiagaraScript::BinaryToExecData(const UNiagaraScript* Script, const TArray<uint8>& InBinaryData, FNiagaraVMExecutableData& OutExecData)
 {
 	check(IsInGameThread());
@@ -2418,7 +2463,7 @@ bool UNiagaraScript::BinaryToExecData(const UNiagaraScript* Script, const TArray
 	}
 
 	FMemoryReader Ar(InBinaryData, true);
-	FObjectAndNameAsStringProxyArchive SafeAr(Ar, false);
+	FObjectAndNameAsStringProxyArchive_NiagaraExecData SafeAr(Ar, false);
 	OutExecData.SerializeData(SafeAr, true);
 
 	FString ValidationErrors;
