fix off by 1 frame issue with reading caches
support writing to volume texture from render thread

diff --git a/Engine/Plugins/FX/Niagara/Source/Niagara/Classes/NiagaraDataInterfaceVolumeCache.h b/Engine/Plugins/FX/Niagara/Source/Niagara/Classes/NiagaraDataInterfaceVolumeCache.h
--- a/Engine/Plugins/FX/Niagara/Source/Niagara/Classes/NiagaraDataInterfaceVolumeCache.h
+++ b/Engine/Plugins/FX/Niagara/Source/Niagara/Classes/NiagaraDataInterfaceVolumeCache.h
@@ -41,6 +41,8 @@ public:
 	virtual bool HasPreSimulateTick() const override { return true; }
 	virtual void GetVMExternalFunction(const FVMExternalFunctionBindingInfo& BindingInfo, void* InstanceData, FVMExternalFunction& OutFunc) override;
 	virtual void DestroyPerInstanceData(void* PerInstanceData, FNiagaraSystemInstance* SystemInstance) override;
+	
+	virtual void ProvidePerInstanceDataForRenderThread(void* InDataFromGT, void* InInstanceData, const FNiagaraSystemInstanceID& SystemInstance) override;
 
 #if WITH_EDITORONLY_DATA
 	virtual bool AppendCompileHash(FNiagaraCompileHashVisitor* InVisitor) const override;
diff --git a/Engine/Plugins/FX/Niagara/Source/Niagara/Classes/VolumeCache.h b/Engine/Plugins/FX/Niagara/Source/Niagara/Classes/VolumeCache.h
--- a/Engine/Plugins/FX/Niagara/Source/Niagara/Classes/VolumeCache.h
+++ b/Engine/Plugins/FX/Niagara/Source/Niagara/Classes/VolumeCache.h
@@ -71,6 +71,7 @@ public:
 	virtual bool LoadRange(FString Path, int Start, int End) = 0;
 	virtual void UnloadAll() = 0;
 	virtual bool Fill3DTexture_RenderThread(int frame, FTextureRHIRef TextureToFill, FRHICommandListImmediate& RHICmdList) = 0;
+	virtual bool Fill3DTexture(int frame, FTextureRHIRef TextureToFill) = 0;
 
 protected:
 	FIntVector DenseResolution;
diff --git a/Engine/Plugins/FX/Niagara/Source/Niagara/Private/NiagaraDataInterfaceVolumeCache.cpp b/Engine/Plugins/FX/Niagara/Source/Niagara/Private/NiagaraDataInterfaceVolumeCache.cpp
--- a/Engine/Plugins/FX/Niagara/Source/Niagara/Private/NiagaraDataInterfaceVolumeCache.cpp
+++ b/Engine/Plugins/FX/Niagara/Source/Niagara/Private/NiagaraDataInterfaceVolumeCache.cpp
@@ -30,11 +30,14 @@ const FName UNiagaraDataInterfaceVolumeCache::GetCurrentFrameNumCells(TEXT("GetC
 struct FVolumeCacheInstanceData_RenderThread
 {
 	int						CurrFrame = 0;
+	bool					ReadFile = false;
 
 	FSamplerStateRHIRef		SamplerStateRHI = nullptr;
 	FTextureReferenceRHIRef	TextureReferenceRHI = nullptr;
 	FTextureRHIRef			ResolvedTextureRHI = nullptr;
 	FVector3f				TextureSize = FVector3f::ZeroVector;
+	
+	TSharedPtr<FVolumeCacheData> VolumeCacheData = nullptr;
 };
 
 struct FVolumeCacheInstanceData_GameThread
@@ -48,9 +51,25 @@ struct FVolumeCacheInstanceData_GameThread
 };
 
 struct FNiagaraDataInterfaceVolumeCacheProxy : public FNiagaraDataInterfaceProxy
-{
-	virtual void ConsumePerInstanceDataFromGameThread(void* PerInstanceData, const FNiagaraSystemInstanceID& Instance) override { check(false); }
-	virtual int32 PerInstanceDataPassedToRenderThreadSize() const override { return 0; }
+{	
+	virtual int32 PerInstanceDataPassedToRenderThreadSize() const override { return sizeof(FVolumeCacheInstanceData_GameThread); }
+
+	virtual void ConsumePerInstanceDataFromGameThread(void* FromGameThreadData, const FNiagaraSystemInstanceID& InstanceID) override
+	{
+		FVolumeCacheInstanceData_GameThread* FromGameThread = reinterpret_cast<FVolumeCacheInstanceData_GameThread*>(FromGameThreadData);
+		FVolumeCacheInstanceData_RenderThread* InstanceData = InstanceData_RT.Find(InstanceID);
+
+		if (FromGameThread != nullptr && InstanceData != nullptr)
+		{
+			InstanceData->CurrFrame = FromGameThread->CurrFrame;
+			InstanceData->ReadFile = FromGameThread->ReadFile;
+			InstanceData->TextureSize = FVector3f(FromGameThread->NumCells.X, FromGameThread->NumCells.Y, FromGameThread->NumCells.Z);		
+
+			FromGameThread->~FVolumeCacheInstanceData_GameThread();
+		}
+	}
+	virtual void PreStage(const FNDIGpuComputePreStageContext& Context) override;
+
 
 	TMap<FNiagaraSystemInstanceID, FVolumeCacheInstanceData_RenderThread> InstanceData_RT;
 };
@@ -116,17 +135,19 @@ bool UNiagaraDataInterfaceVolumeCache::InitPerInstanceData(void* PerInstanceData
 	if (VolumeCache != nullptr)
 	{
 		InstanceData->NumCells = VolumeCache->Resolution;
-	}
 
-	// Push Updates to Proxy.
-	FNiagaraDataInterfaceVolumeCacheProxy* TheProxy = GetProxyAs<FNiagaraDataInterfaceVolumeCacheProxy>();
-	ENQUEUE_RENDER_COMMAND(FUpdateData)(
-		[TheProxy, InstanceID = SystemInstance->GetId()](FRHICommandListImmediate& RHICmdList)
-	{
-		check(!TheProxy->InstanceData_RT.Contains(InstanceID));
-		FVolumeCacheInstanceData_RenderThread* TargetData = &TheProxy->InstanceData_RT.Add(InstanceID);
+		// Push Updates to Proxy.
+		FNiagaraDataInterfaceVolumeCacheProxy* TheProxy = GetProxyAs<FNiagaraDataInterfaceVolumeCacheProxy>();
+		ENQUEUE_RENDER_COMMAND(FUpdateData)(
+			[TheProxy, RT_VolumeCacheData = VolumeCache->GetData(), InstanceID = SystemInstance->GetId()](FRHICommandListImmediate& RHICmdList)
+		{
+			check(!TheProxy->InstanceData_RT.Contains(InstanceID));
+			FVolumeCacheInstanceData_RenderThread* TargetData = &TheProxy->InstanceData_RT.Add(InstanceID);
+			TargetData->VolumeCacheData = RT_VolumeCacheData;
+		});
+	}
 
-	});
+	
 
 	return true;
 }
@@ -147,9 +168,47 @@ void UNiagaraDataInterfaceVolumeCache::DestroyPerInstanceData(void* PerInstanceD
 	);
 }
 
-void UNiagaraDataInterfaceVolumeCache::GetFunctions(TArray<FNiagaraFunctionSignature>& OutFunctions)
+void FNiagaraDataInterfaceVolumeCacheProxy::PreStage(const FNDIGpuComputePreStageContext& Context)
 {	
+	FVolumeCacheInstanceData_RenderThread* InstanceData = InstanceData_RT.Find(Context.GetSystemInstanceID());
+	
+	if (Context.IsInputStage() && InstanceData && InstanceData->ReadFile && InstanceData->VolumeCacheData != nullptr)
+	{
+	
+		FIntVector Size = FIntVector(InstanceData->TextureSize.X, InstanceData->TextureSize.Y, InstanceData->TextureSize.Z);
+
+		if (!InstanceData->ResolvedTextureRHI.IsValid())
+		{
+
+			// #todo(dmp): this should be in instance data
+			EPixelFormat Format = PF_FloatRGBA;
+
+			const FRHITextureCreateDesc Desc =
+				FRHITextureCreateDesc::Create3D(TEXT("stuff"), Size.X, Size.Y, Size.Z, Format)
+				.SetFlags(ETextureCreateFlags::ShaderResource | ETextureCreateFlags::NoTiling);
+
+			InstanceData->ResolvedTextureRHI = RHICreateTexture(Desc);
+			InstanceData->TextureSize = FVector3f(Size.X, Size.Y, Size.Z);
+			InstanceData->SamplerStateRHI = TStaticSamplerState<SF_Bilinear, AM_Clamp, AM_Clamp, AM_Clamp>::GetRHI();
+		}
+		// must make sure readfile, curr frame, size are in the instance data - must sync in provide instance data for render thread
+
+		InstanceData->VolumeCacheData->Fill3DTexture(InstanceData->CurrFrame, InstanceData->ResolvedTextureRHI);
+	}
+}
 
+void UNiagaraDataInterfaceVolumeCache::ProvidePerInstanceDataForRenderThread(void* InDataFromGT, void* InInstanceData, const FNiagaraSystemInstanceID& SystemInstance)
+{
+	FVolumeCacheInstanceData_GameThread* InstanceData = static_cast<FVolumeCacheInstanceData_GameThread*>(InInstanceData);
+	FVolumeCacheInstanceData_GameThread* DataFromGT = static_cast<FVolumeCacheInstanceData_GameThread*>(InDataFromGT);
+
+	DataFromGT->CurrFrame = InstanceData->CurrFrame;
+	DataFromGT->ReadFile = InstanceData->ReadFile;
+	DataFromGT->NumCells = InstanceData->NumCells;
+}
+
+void UNiagaraDataInterfaceVolumeCache::GetFunctions(TArray<FNiagaraFunctionSignature>& OutFunctions)
+{	
 	{
 		FNiagaraFunctionSignature Sig;
 		Sig.Name = SetFrameName;
@@ -216,6 +275,7 @@ void UNiagaraDataInterfaceVolumeCache::GetFunctions(TArray<FNiagaraFunctionSigna
 		Sig.bRequiresContext = false;
 		Sig.bSupportsCPU = false;
 		Sig.bSupportsGPU = true;
+		Sig.bReadFunction = true;
 		Sig.Inputs.Add(FNiagaraVariable(FNiagaraTypeDefinition(GetClass()), TEXT("Texture")));
 		Sig.Inputs.Add(FNiagaraVariable(FNiagaraTypeDefinition::GetVec3Def(), TEXT("UVW")));
 		Sig.Inputs.Add(FNiagaraVariable(FNiagaraTypeDefinition::GetFloatDef(), TEXT("MipLevel")));
@@ -233,6 +293,7 @@ void UNiagaraDataInterfaceVolumeCache::GetFunctions(TArray<FNiagaraFunctionSigna
 		Sig.bRequiresContext = false;
 		Sig.bSupportsCPU = false;
 		Sig.bSupportsGPU = true;
+		Sig.bReadFunction = true;
 		Sig.Inputs.Add(FNiagaraVariable(FNiagaraTypeDefinition(GetClass()), TEXT("Texture")));
 		Sig.Inputs.Add(FNiagaraVariable(FNiagaraTypeDefinition::GetIntDef(), TEXT("x")));
 		Sig.Inputs.Add(FNiagaraVariable(FNiagaraTypeDefinition::GetIntDef(), TEXT("y")));
@@ -270,6 +331,12 @@ void UNiagaraDataInterfaceVolumeCache::SetFrame(FVectorVMExternalFunctionContext
 
 	InstData->CurrFrame = InFrame.GetAndAdvance();
 
+	// cannot read from cache...spew errors or let it go?
+	if (VolumeCache == nullptr || !VolumeCache->LoadFile(InstData->CurrFrame))
+	{
+		UE_LOG(LogNiagara, Warning, TEXT("Cache Read failed: %s"), *GetName());
+	}
+
 	*OutSuccess.GetDestAndAdvance() = true;
 }
 
@@ -343,50 +410,8 @@ void UNiagaraDataInterfaceVolumeCache::GetVMExternalFunction(const FVMExternalFu
 }
 
 bool UNiagaraDataInterfaceVolumeCache::PerInstanceTick(void* PerInstanceData, FNiagaraSystemInstance* SystemInstance, float DeltaSeconds)
-{
-	FVolumeCacheInstanceData_GameThread* InstanceData = SystemInstancesToProxyData_GT.FindRef(SystemInstance->GetId());	
-
-	// we can run into the case where depending on the ordering of DI initialization, we might have not been able to grab the other grid's InstanceData
-	// in InitPerInstanceData.  If this is the case, we ensure it is correct here.
-	if (InstanceData && InstanceData->ReadFile && VolumeCache != nullptr)
-	{	
-		FNiagaraDataInterfaceVolumeCacheProxy* TextureProxy = GetProxyAs<FNiagaraDataInterfaceVolumeCacheProxy>();
-
-		//EPixelFormat Format = PF_A32B32G32R32F;
-		EPixelFormat Format = PF_FloatRGBA;
-		const int32 FormatSize = GPixelFormats[Format].BlockBytes;
-
-		// cannot read from cache...spew errors or let it go?
-		if (!VolumeCache->LoadFile(InstanceData->CurrFrame))
-		{
-			UE_LOG(LogNiagara, Warning, TEXT("Cache Read failed: %s"), *VolumeCache->GetName());
-			return false;
-		}
-		
-		ENQUEUE_RENDER_COMMAND(FVolumeCacheFillTexture)(
-			[RT_Frame = InstanceData->CurrFrame, RT_VolumeCacheData = VolumeCache->GetData(), Format, TextureProxy, SystemID = SystemInstance->GetId()](FRHICommandListImmediate& RHICmdList)
-		{
-			FVolumeCacheInstanceData_RenderThread* InstanceData = TextureProxy->InstanceData_RT.Find(SystemID);
-			FIntVector Size = RT_VolumeCacheData->GetDenseResolution();
-
-			if (!InstanceData->ResolvedTextureRHI.IsValid())
-			{
-
-				const FRHITextureCreateDesc Desc =
-					FRHITextureCreateDesc::Create3D(TEXT("stuff"), Size.X, Size.Y, Size.Z, Format)
-					.SetFlags(ETextureCreateFlags::ShaderResource | ETextureCreateFlags::NoTiling);
-
-				InstanceData->ResolvedTextureRHI = RHICreateTexture(Desc);
-				InstanceData->TextureSize = FVector3f(Size.X, Size.Y, Size.Z);
-				InstanceData->SamplerStateRHI = TStaticSamplerState<SF_Bilinear, AM_Clamp, AM_Clamp, AM_Clamp>::GetRHI();
-			}
-				
-			RT_VolumeCacheData->Fill3DTexture_RenderThread(RT_Frame, InstanceData->ResolvedTextureRHI, RHICmdList);			
-		});
-	}
-
+{	
 	return false;
-
 }
 
 #if WITH_EDITORONLY_DATA
diff --git a/Engine/Plugins/FX/Niagara/Source/Niagara/Private/VolumeCache.cpp b/Engine/Plugins/FX/Niagara/Source/Niagara/Private/VolumeCache.cpp
--- a/Engine/Plugins/FX/Niagara/Source/Niagara/Private/VolumeCache.cpp
+++ b/Engine/Plugins/FX/Niagara/Source/Niagara/Private/VolumeCache.cpp
@@ -34,6 +34,7 @@ public:
 	virtual bool LoadRange(FString Path, int Start, int End);
 	virtual void UnloadAll();
 	virtual bool Fill3DTexture_RenderThread(int frame, FTextureRHIRef TextureToFill, FRHICommandListImmediate& RHICmdList);
+	virtual bool Fill3DTexture(int frame, FTextureRHIRef TextureToFill);
 	
 private:
 	TMap<int32, Vec4Grid::Ptr> OpenVDBGrids;
@@ -267,6 +268,36 @@ bool FOpenVDBCacheData::Fill3DTexture_RenderThread(int frame, FTextureRHIRef Tex
 	return false;
 }
 
+bool FOpenVDBCacheData::Fill3DTexture(int frame, FTextureRHIRef TextureToFill)
+{
+	if (OpenVDBGrids.Contains(frame) && OpenVDBGrids[frame] != nullptr)
+	{
+		uint8* DataPtr = (uint8*)DenseGridPtr->data();
+
+		const int32 FormatSize = GPixelFormats[TextureToFill->GetFormat()].BlockBytes;
+
+		if (TextureToFill->GetSizeXYZ() != DenseResolution)
+		{
+			UE_LOG(LogVolumeCache, Warning, TEXT("Target texture resolution %s doesn't match OpenVDB grid resolution %s.  Regenerate your cache"), *TextureToFill->GetSizeXYZ().ToString(), *DenseResolution.ToString());
+		}
+		else
+		{
+			const FUpdateTextureRegion3D UpdateRegion(0, 0, 0, 0, 0, 0, DenseResolution.X, DenseResolution.Y, DenseResolution.Z);
+			const SIZE_T MemorySize = static_cast<SIZE_T>(UpdateRegion.Width) * static_cast<SIZE_T>(UpdateRegion.Height) * static_cast<SIZE_T>(UpdateRegion.Depth) * static_cast<SIZE_T>(FormatSize);
+
+			FUpdateTexture3DData TheData = RHIBeginUpdateTexture3D(TextureToFill, 0, UpdateRegion);
+
+			FMemory::Memcpy(TheData.Data, DataPtr, MemorySize);
+
+			RHIEndUpdateTexture3D(TheData);
+		}
+
+		return true;
+	}
+
+	return false;
+}
+
 bool OpenVDBTools::WriteImageDataToOpenVDBFile(FStringView FilePath, FIntVector ImageSize, TArrayView<FFloat16Color> ImageData, bool UseFloatGrids)
 {
 	openvdb::math::Transform::Ptr TransformPtr = openvdb::math::Transform::createLinearTransform(/*voxel size=*/1.0);
