- Add newline constants to end of functions

#rnx

diff --git a/Engine/Plugins/FX/Niagara/Source/Niagara/Private/NiagaraDataInterfaceSkeletalMesh.cpp b/Engine/Plugins/FX/Niagara/Source/Niagara/Private/NiagaraDataInterfaceSkeletalMesh.cpp
--- a/Engine/Plugins/FX/Niagara/Source/Niagara/Private/NiagaraDataInterfaceSkeletalMesh.cpp
+++ b/Engine/Plugins/FX/Niagara/Source/Niagara/Private/NiagaraDataInterfaceSkeletalMesh.cpp
@@ -2953,59 +2953,59 @@ bool UNiagaraDataInterfaceSkeletalMesh::GetFunctionHLSL(const FNiagaraDataInterf
 	const TCHAR* FuncFormat = nullptr;
 	if (FunctionInfo.DefinitionName == FSkeletalMeshInterfaceHelper::IsValidTriCoordName)
 	{
-		FuncFormat = TEXT("void {InstanceFunctionName}(in MeshTriCoordinate InCoord, out bool bIsValid) { IsValidTriCoord_{ParameterName}(InCoord.Tri, bIsValid); }");
+		FuncFormat = TEXT("void {InstanceFunctionName}(in MeshTriCoordinate InCoord, out bool bIsValid) { IsValidTriCoord_{ParameterName}(InCoord.Tri, bIsValid); }\n");
 	}
 	else if (FunctionInfo.DefinitionName == FSkeletalMeshInterfaceHelper::RandomTriangleName)
 	{
-		FuncFormat = TEXT("void {InstanceFunctionName}(NiagaraRandInfo InRandomInfo, out MeshTriCoordinate OutCoord) { RandomTriangle_{ParameterName}(InRandomInfo.Seed1, InRandomInfo.Seed2, InRandomInfo.Seed3, OutCoord.Tri, OutCoord.BaryCoord); }");
+		FuncFormat = TEXT("void {InstanceFunctionName}(NiagaraRandInfo InRandomInfo, out MeshTriCoordinate OutCoord) { RandomTriangle_{ParameterName}(InRandomInfo.Seed1, InRandomInfo.Seed2, InRandomInfo.Seed3, OutCoord.Tri, OutCoord.BaryCoord); }\n");
 	}
 	else if (FunctionInfo.DefinitionName == FSkeletalMeshInterfaceHelper::RandomTriCoordName)
 	{
-		FuncFormat = TEXT("void {InstanceFunctionName}(NiagaraRandInfo InRandomInfo, out MeshTriCoordinate OutCoord) { RandomTriangle_{ParameterName}(InRandomInfo.Seed1, InRandomInfo.Seed2, InRandomInfo.Seed3, OutCoord.Tri, OutCoord.BaryCoord); }");
+		FuncFormat = TEXT("void {InstanceFunctionName}(NiagaraRandInfo InRandomInfo, out MeshTriCoordinate OutCoord) { RandomTriangle_{ParameterName}(InRandomInfo.Seed1, InRandomInfo.Seed2, InRandomInfo.Seed3, OutCoord.Tri, OutCoord.BaryCoord); }\n");
 	}
 	else if (FunctionInfo.DefinitionName == FSkeletalMeshInterfaceHelper::RandomFilteredTriangleName)
 	{
-		FuncFormat = TEXT("void {InstanceFunctionName}(NiagaraRandInfo InRandomInfo, out MeshTriCoordinate OutCoord) { RandomFilteredTriangle_{ParameterName}(InRandomInfo.Seed1, InRandomInfo.Seed2, InRandomInfo.Seed3, OutCoord.Tri, OutCoord.BaryCoord); }");
+		FuncFormat = TEXT("void {InstanceFunctionName}(NiagaraRandInfo InRandomInfo, out MeshTriCoordinate OutCoord) { RandomFilteredTriangle_{ParameterName}(InRandomInfo.Seed1, InRandomInfo.Seed2, InRandomInfo.Seed3, OutCoord.Tri, OutCoord.BaryCoord); }\n");
 	}
 	else if (FunctionInfo.DefinitionName == FSkeletalMeshInterfaceHelper::GetFilteredTriangleAtName)
 	{
-		FuncFormat = TEXT("void {InstanceFunctionName}(int FilteredIndex, in float3 BaryCoord, out MeshTriCoordinate OutCoord) { GetFilteredTriangleAt_{ParameterName}(FilteredIndex, OutCoord.Tri); OutCoord.BaryCoord = BaryCoord; }");
+		FuncFormat = TEXT("void {InstanceFunctionName}(int FilteredIndex, in float3 BaryCoord, out MeshTriCoordinate OutCoord) { GetFilteredTriangleAt_{ParameterName}(FilteredIndex, OutCoord.Tri); OutCoord.BaryCoord = BaryCoord; }\n");
 	}
 	else if (FunctionInfo.DefinitionName == FSkeletalMeshInterfaceHelper::GetTriangleDataName)
 	{
-		FuncFormat = TEXT("void {InstanceFunctionName}(in MeshTriCoordinate InCoord, out float3 OutPosition, out float3 OutNormal, out float3 OutBinormal, out float3 OutTangent) { GetPointOnTriangle_{ParameterName}(InCoord.Tri, InCoord.BaryCoord, OutPosition, OutTangent, OutBinormal, OutNormal); }");
+		FuncFormat = TEXT("void {InstanceFunctionName}(in MeshTriCoordinate InCoord, out float3 OutPosition, out float3 OutNormal, out float3 OutBinormal, out float3 OutTangent) { GetPointOnTriangle_{ParameterName}(InCoord.Tri, InCoord.BaryCoord, OutPosition, OutTangent, OutBinormal, OutNormal); }\n");
 	}
 	else if (FunctionInfo.DefinitionName == FSkeletalMeshInterfaceHelper::GetSkinnedTriangleDataWSName)
 	{
-		FuncFormat = TEXT("void {InstanceFunctionName}(in MeshTriCoordinate InCoord, out float3 OutPosition, out float3 OutVelocity, out float3 OutNormal, out float3 OutBinormal, out float3 OutTangent) { GetSkinnedTriangleDataWS_{ParameterName}(InCoord.Tri, InCoord.BaryCoord, OutPosition, OutVelocity, OutNormal, OutBinormal, OutTangent); }");
+		FuncFormat = TEXT("void {InstanceFunctionName}(in MeshTriCoordinate InCoord, out float3 OutPosition, out float3 OutVelocity, out float3 OutNormal, out float3 OutBinormal, out float3 OutTangent) { GetSkinnedTriangleDataWS_{ParameterName}(InCoord.Tri, InCoord.BaryCoord, OutPosition, OutVelocity, OutNormal, OutBinormal, OutTangent); }\n");
 	}
 	else if (FunctionInfo.DefinitionName == FSkeletalMeshInterfaceHelper::GetSkinnedTriangleDataWSInterpName)
 	{
-		FuncFormat = TEXT("void {InstanceFunctionName}(in MeshTriCoordinate InCoord, in float InInterp, out float3 OutPosition, out float3 OutVelocity, out float3 OutNormal, out float3 OutBinormal, out float3 OutTangent) { GetSkinnedTriangleDataInterpolatedWS_{ParameterName}(InCoord.Tri, InCoord.BaryCoord, InInterp, OutPosition, OutVelocity, OutNormal, OutBinormal, OutTangent); }");
+		FuncFormat = TEXT("void {InstanceFunctionName}(in MeshTriCoordinate InCoord, in float InInterp, out float3 OutPosition, out float3 OutVelocity, out float3 OutNormal, out float3 OutBinormal, out float3 OutTangent) { GetSkinnedTriangleDataInterpolatedWS_{ParameterName}(InCoord.Tri, InCoord.BaryCoord, InInterp, OutPosition, OutVelocity, OutNormal, OutBinormal, OutTangent); }\n");
 	}
 	else if (FunctionInfo.DefinitionName == FSkeletalMeshInterfaceHelper::GetSkinnedTriangleDataName)
 	{
-		FuncFormat = TEXT("void {InstanceFunctionName}(in MeshTriCoordinate InCoord, out float3 OutPosition, out float3 OutVelocity, out float3 OutNormal, out float3 OutBinormal, out float3 OutTangent) { GetSkinnedTriangleData_{ParameterName}(InCoord.Tri, InCoord.BaryCoord, OutPosition, OutVelocity, OutNormal, OutBinormal, OutTangent); }");
+		FuncFormat = TEXT("void {InstanceFunctionName}(in MeshTriCoordinate InCoord, out float3 OutPosition, out float3 OutVelocity, out float3 OutNormal, out float3 OutBinormal, out float3 OutTangent) { GetSkinnedTriangleData_{ParameterName}(InCoord.Tri, InCoord.BaryCoord, OutPosition, OutVelocity, OutNormal, OutBinormal, OutTangent); }\n");
 	}
 	else if (FunctionInfo.DefinitionName == FSkeletalMeshInterfaceHelper::GetSkinnedTriangleDataInterpName)
 	{
-		FuncFormat = TEXT("void {InstanceFunctionName}(in MeshTriCoordinate InCoord, in float InInterp, out float3 OutPosition, out float3 OutVelocity, out float3 OutNormal, out float3 OutBinormal, out float3 OutTangent) { GetSkinnedTriangleDataInterpolated_{ParameterName}(InCoord.Tri, InCoord.BaryCoord, InInterp, OutPosition, OutVelocity, OutNormal, OutBinormal, OutTangent); }");
+		FuncFormat = TEXT("void {InstanceFunctionName}(in MeshTriCoordinate InCoord, in float InInterp, out float3 OutPosition, out float3 OutVelocity, out float3 OutNormal, out float3 OutBinormal, out float3 OutTangent) { GetSkinnedTriangleDataInterpolated_{ParameterName}(InCoord.Tri, InCoord.BaryCoord, InInterp, OutPosition, OutVelocity, OutNormal, OutBinormal, OutTangent); }\n");
 	}
 	else if (FunctionInfo.DefinitionName == FSkeletalMeshInterfaceHelper::GetTriUVName)
 	{
-		FuncFormat = TEXT("void {InstanceFunctionName}(in MeshTriCoordinate InCoord, in int InUVSet, out float2 OutUV) { GetTriUV_{ParameterName}(InCoord.Tri, InCoord.BaryCoord, InUVSet, OutUV); }");
+		FuncFormat = TEXT("void {InstanceFunctionName}(in MeshTriCoordinate InCoord, in int InUVSet, out float2 OutUV) { GetTriUV_{ParameterName}(InCoord.Tri, InCoord.BaryCoord, InUVSet, OutUV); }\n");
 	}
 	else if (FunctionInfo.DefinitionName == FSkeletalMeshInterfaceHelper::GetTriColorName)
 	{
-		FuncFormat = TEXT("void {InstanceFunctionName}(in MeshTriCoordinate InCoord, out float4 OutColor) { GetTriColor_{ParameterName}(InCoord.Tri, InCoord.BaryCoord, OutColor); }");
+		FuncFormat = TEXT("void {InstanceFunctionName}(in MeshTriCoordinate InCoord, out float4 OutColor) { GetTriColor_{ParameterName}(InCoord.Tri, InCoord.BaryCoord, OutColor); }\n");
 	}
 	else if (FunctionInfo.DefinitionName == FSkeletalMeshInterfaceHelper::GetTriangleCoordAtUVName)
 	{
-		FuncFormat = TEXT("void {InstanceFunctionName}(in bool InEnabled, in float2 InUV, in float InTolerance, out MeshTriCoordinate OutCoord, out bool OutIsValid) { GetTriangleCoordAtUV_{ParameterName}(InEnabled, InUV, InTolerance, OutCoord.Tri, OutCoord.BaryCoord, OutIsValid); }");
+		FuncFormat = TEXT("void {InstanceFunctionName}(in bool InEnabled, in float2 InUV, in float InTolerance, out MeshTriCoordinate OutCoord, out bool OutIsValid) { GetTriangleCoordAtUV_{ParameterName}(InEnabled, InUV, InTolerance, OutCoord.Tri, OutCoord.BaryCoord, OutIsValid); }\n");
 	}
 	else if (FunctionInfo.DefinitionName == FSkeletalMeshInterfaceHelper::GetTriangleCoordInAabbName)
 	{
-		FuncFormat = TEXT("void {InstanceFunctionName}(in bool InEnabled, in float2 InUvMin, in float2 InUvMax, out MeshTriCoordinate OutCoord, out bool OutIsValid) { GetTriangleCoordInAabb_{ParameterName}(InEnabled, InUvMin, InUvMax, OutCoord.Tri, OutCoord.BaryCoord, OutIsValid); }");
+		FuncFormat = TEXT("void {InstanceFunctionName}(in bool InEnabled, in float2 InUvMin, in float2 InUvMax, out MeshTriCoordinate OutCoord, out bool OutIsValid) { GetTriangleCoordInAabb_{ParameterName}(InEnabled, InUvMin, InUvMax, OutCoord.Tri, OutCoord.BaryCoord, OutIsValid); }\n");
 	}
 
 	if ( FuncFormat != nullptr )
diff --git a/Engine/Source/Runtime/Core/Private/UObject/UE5MainStreamObjectVersion.cpp b/Engine/Source/Runtime/Core/Private/UObject/UE5MainStreamObjectVersion.cpp
--- a/Engine/Source/Runtime/Core/Private/UObject/UE5MainStreamObjectVersion.cpp
+++ b/Engine/Source/Runtime/Core/Private/UObject/UE5MainStreamObjectVersion.cpp
@@ -14,7 +14,7 @@ TMap<FGuid, FGuid> FUE5MainStreamObjectVersion::GetSystemGuids()
 	SystemGuids.Add(DevGuids.MATERIALSHADERMAP_DERIVEDDATA_VER, FGuid("65C01D817C9A4EEFAE9E988D41A1F3DD"));
 	SystemGuids.Add(DevGuids.NANITE_DERIVEDDATA_VER, FGuid("7A3EA93546624597BEBE2ADFDE3A7628"));
 	SystemGuids.Add(DevGuids.NIAGARASHADERMAP_DERIVEDDATA_VER, FGuid("8A37C45D24F2423CBE5F8F371DE33575"));
-	SystemGuids.Add(DevGuids.Niagara_LatestScriptCompileVersion, FGuid("2026F2A19D0C46688E92214FDBAE4CA0"));
+	SystemGuids.Add(DevGuids.Niagara_LatestScriptCompileVersion, FGuid("C4BD9AE35DC59C4C985DD5316ACC8CC9"));
 	SystemGuids.Add(DevGuids.SkeletalMeshDerivedDataVersion, FGuid("FB7B2FB546DE49B2A2F88B7F69763CEE"));
 	SystemGuids.Add(DevGuids.STATICMESH_DERIVEDDATA_VER, FGuid("98F7D79A5811013825E75BBCC41ED3E9"));
 	SystemGuids.Add(DevGuids.POSESEARCHDB_DERIVEDDATA_VER, FGuid("4E595C2AC5E947D6BA9ABC874353E5BC"));
