C++20 fixes in Niagara.

diff --git a/Engine/Plugins/FX/Niagara/Source/Niagara/Classes/NiagaraDataSetAccessor.h b/Engine/Plugins/FX/Niagara/Source/Niagara/Classes/NiagaraDataSetAccessor.h
--- a/Engine/Plugins/FX/Niagara/Source/Niagara/Classes/NiagaraDataSetAccessor.h
+++ b/Engine/Plugins/FX/Niagara/Source/Niagara/Classes/NiagaraDataSetAccessor.h
@@ -192,7 +192,7 @@ struct FNiagaraDataSetAccessorFloat
 	{
 		ComponentIndex = INDEX_NONE;
 
-		const int32 FloatVariableIndex = DataSetCompiledData.Variables.IndexOfByKey(FNiagaraVariableBase(FNiagaraDataSetAccessorTypeInfo<TType>::GetFloatType(), VariableName));
+		const int32 FloatVariableIndex = DataSetCompiledData.Variables.IndexOfByPredicate(FNiagaraVariableMatch(FNiagaraDataSetAccessorTypeInfo<TType>::GetFloatType(), VariableName));
 		if (FloatVariableIndex != INDEX_NONE)
 		{
 			bIsFloat = true;
@@ -202,7 +202,7 @@ struct FNiagaraDataSetAccessorFloat
 
 		if (FNiagaraDataSetAccessorTypeInfo<TType>::bSupportsAlternateType)
 		{
-			const int32 AlternateVariableIndex = DataSetCompiledData.Variables.IndexOfByKey(FNiagaraVariableBase(FNiagaraDataSetAccessorTypeInfo<TType>::GetAlternateType(), VariableName));
+			const int32 AlternateVariableIndex = DataSetCompiledData.Variables.IndexOfByPredicate(FNiagaraVariableMatch(FNiagaraDataSetAccessorTypeInfo<TType>::GetAlternateType(), VariableName));
 			if (AlternateVariableIndex != INDEX_NONE)
 			{
 				bIsFloat = true;
@@ -213,7 +213,7 @@ struct FNiagaraDataSetAccessorFloat
 		
 		if (FNiagaraDataSetAccessorTypeInfo<TType>::bSupportsHalf)
 		{
-			const int32 HalfVariableIndex = DataSetCompiledData.Variables.IndexOfByKey(FNiagaraVariableBase(FNiagaraDataSetAccessorTypeInfo<TType>::GetHalfType(), VariableName));
+			const int32 HalfVariableIndex = DataSetCompiledData.Variables.IndexOfByPredicate(FNiagaraVariableMatch(FNiagaraDataSetAccessorTypeInfo<TType>::GetHalfType(), VariableName));
 			if (HalfVariableIndex != INDEX_NONE)
 			{
 				bIsFloat = false;
@@ -456,7 +456,7 @@ struct FNiagaraDataSetAccessorInt32
 	{
 		ComponentIndex = INDEX_NONE;
 
-		const int32 VariableIndex = DataSetCompiledData.Variables.IndexOfByKey(FNiagaraVariableBase(FNiagaraDataSetAccessorTypeInfo<TType>::GetIntType(), VariableName));
+		const int32 VariableIndex = DataSetCompiledData.Variables.IndexOfByPredicate(FNiagaraVariableMatch(FNiagaraDataSetAccessorTypeInfo<TType>::GetIntType(), VariableName));
 		if (VariableIndex != INDEX_NONE)
 		{
 			ComponentIndex = DataSetCompiledData.VariableLayouts[VariableIndex].Int32ComponentStart;
@@ -598,7 +598,7 @@ struct FNiagaraDataSetAccessorStruct
 	{
 		bIsValid = false;
 
-		const int32 VariableIndex = DataSetCompiledData.Variables.IndexOfByKey(FNiagaraVariableBase(FNiagaraDataSetAccessorTypeInfo<TType>::GetStructType(), VariableName));
+		const int32 VariableIndex = DataSetCompiledData.Variables.IndexOfByPredicate(FNiagaraVariableMatch(FNiagaraDataSetAccessorTypeInfo<TType>::GetStructType(), VariableName));
 		if (VariableIndex != INDEX_NONE)
 		{
 			bIsValid = true;
diff --git a/Engine/Plugins/FX/Niagara/Source/Niagara/Classes/NiagaraScript.h b/Engine/Plugins/FX/Niagara/Source/Niagara/Classes/NiagaraScript.h
--- a/Engine/Plugins/FX/Niagara/Source/Niagara/Classes/NiagaraScript.h
+++ b/Engine/Plugins/FX/Niagara/Source/Niagara/Classes/NiagaraScript.h
@@ -1042,7 +1042,7 @@ public:
 	{
 		for (const FNiagaraCompilerTag& Tag : CachedScriptVM.CompileTags)
 		{
-			if (Tag.Variable == InVar)
+			if (static_cast<const FNiagaraVariableBase&>(Tag.Variable) == InVar)
 			{
 				if (Tag.Variable.IsDataAllocated())
 				{
diff --git a/Engine/Plugins/FX/Niagara/Source/Niagara/Private/NiagaraComponentRendererProperties.cpp b/Engine/Plugins/FX/Niagara/Source/Niagara/Private/NiagaraComponentRendererProperties.cpp
--- a/Engine/Plugins/FX/Niagara/Source/Niagara/Private/NiagaraComponentRendererProperties.cpp
+++ b/Engine/Plugins/FX/Niagara/Source/Niagara/Private/NiagaraComponentRendererProperties.cpp
@@ -216,7 +216,7 @@ void UNiagaraComponentRendererProperties::PostLoad()
 		// Move old bindings over to new position type 
 		for (FNiagaraVariable OldVarType : OldTypes)
 		{
-			if (Binding.AttributeBinding.GetParamMapBindableVariable() == OldVarType)
+			if (Binding.AttributeBinding.GetParamMapBindableVariable() == static_cast<const FNiagaraVariableBase&>(OldVarType))
 			{
 				FNiagaraVariable NewVarType(FNiagaraTypeDefinition::GetPositionDef(), OldVarType.GetName());
 				Binding.AttributeBinding.Setup(NewVarType, NewVarType, InSourceMode);
diff --git a/Engine/Plugins/FX/Niagara/Source/Niagara/Private/NiagaraDataInterfaceColorCurve.cpp b/Engine/Plugins/FX/Niagara/Source/Niagara/Private/NiagaraDataInterfaceColorCurve.cpp
--- a/Engine/Plugins/FX/Niagara/Source/Niagara/Private/NiagaraDataInterfaceColorCurve.cpp
+++ b/Engine/Plugins/FX/Niagara/Source/Niagara/Private/NiagaraDataInterfaceColorCurve.cpp
@@ -207,8 +207,8 @@ FORCEINLINE_DEBUGGABLE FLinearColor UNiagaraDataInterfaceColorCurve::SampleCurve
 	float NextEntry = PrevEntry < LUTNumSamplesMinusOne ? PrevEntry + 1.0f : PrevEntry;
 	float Interp = RemappedX - PrevEntry;
 	
-	int32 AIndex = PrevEntry * CurveLUTNumElems;
-	int32 BIndex = NextEntry * CurveLUTNumElems;
+	int32 AIndex = (int32)(PrevEntry * (float)CurveLUTNumElems);
+	int32 BIndex = (int32)(NextEntry * (float)CurveLUTNumElems);
 	FLinearColor A = FLinearColor(ShaderLUT[AIndex], ShaderLUT[AIndex + 1], ShaderLUT[AIndex + 2], ShaderLUT[AIndex + 3]);
 	FLinearColor B = FLinearColor(ShaderLUT[BIndex], ShaderLUT[BIndex + 1], ShaderLUT[BIndex + 2], ShaderLUT[BIndex + 3]);
 	return FMath::Lerp(A, B, Interp);
diff --git a/Engine/Plugins/FX/Niagara/Source/Niagara/Private/NiagaraDataInterfaceCurve.cpp b/Engine/Plugins/FX/Niagara/Source/Niagara/Private/NiagaraDataInterfaceCurve.cpp
--- a/Engine/Plugins/FX/Niagara/Source/Niagara/Private/NiagaraDataInterfaceCurve.cpp
+++ b/Engine/Plugins/FX/Niagara/Source/Niagara/Private/NiagaraDataInterfaceCurve.cpp
@@ -170,8 +170,8 @@ FORCEINLINE_DEBUGGABLE float UNiagaraDataInterfaceCurve::SampleCurveInternal<TIn
 	float NextEntry = PrevEntry < LUTNumSamplesMinusOne ? PrevEntry + 1.0f : PrevEntry;
 	float Interp = RemappedX - PrevEntry;
 
-	int32 AIndex = PrevEntry * CurveLUTNumElems;
-	int32 BIndex = NextEntry * CurveLUTNumElems;
+	int32 AIndex = (int32)(PrevEntry * (float)CurveLUTNumElems);
+	int32 BIndex = (int32)(NextEntry * (float)CurveLUTNumElems);
 	float A = ShaderLUT[AIndex];
 	float B = ShaderLUT[BIndex];
 	return FMath::Lerp(A, B, Interp);
diff --git a/Engine/Plugins/FX/Niagara/Source/Niagara/Private/NiagaraDataInterfaceVector2DCurve.cpp b/Engine/Plugins/FX/Niagara/Source/Niagara/Private/NiagaraDataInterfaceVector2DCurve.cpp
--- a/Engine/Plugins/FX/Niagara/Source/Niagara/Private/NiagaraDataInterfaceVector2DCurve.cpp
+++ b/Engine/Plugins/FX/Niagara/Source/Niagara/Private/NiagaraDataInterfaceVector2DCurve.cpp
@@ -185,8 +185,8 @@ FORCEINLINE_DEBUGGABLE FVector2D UNiagaraDataInterfaceVector2DCurve::SampleCurve
 	float NextEntry = PrevEntry < LUTNumSamplesMinusOne ? PrevEntry + 1.0f : PrevEntry;
 	float Interp = RemappedX - PrevEntry;
 
-	int32 AIndex = PrevEntry * CurveLUTNumElems;
-	int32 BIndex = NextEntry * CurveLUTNumElems;
+	int32 AIndex = (int32)(PrevEntry * (float)CurveLUTNumElems);
+	int32 BIndex = (int32)(NextEntry * (float)CurveLUTNumElems);
 	FVector2D A = FVector2D(ShaderLUT[AIndex], ShaderLUT[AIndex + 1]);
 	FVector2D B = FVector2D(ShaderLUT[BIndex], ShaderLUT[BIndex + 1]);
 	return FMath::Lerp(A, B, Interp);
diff --git a/Engine/Plugins/FX/Niagara/Source/Niagara/Private/NiagaraDataInterfaceVector4Curve.cpp b/Engine/Plugins/FX/Niagara/Source/Niagara/Private/NiagaraDataInterfaceVector4Curve.cpp
--- a/Engine/Plugins/FX/Niagara/Source/Niagara/Private/NiagaraDataInterfaceVector4Curve.cpp
+++ b/Engine/Plugins/FX/Niagara/Source/Niagara/Private/NiagaraDataInterfaceVector4Curve.cpp
@@ -206,8 +206,8 @@ FORCEINLINE_DEBUGGABLE FVector4 UNiagaraDataInterfaceVector4Curve::SampleCurveIn
 	float NextEntry = PrevEntry < LUTNumSamplesMinusOne ? PrevEntry + 1.0f : PrevEntry;
 	float Interp = RemappedX - PrevEntry;
 
-	int32 AIndex = PrevEntry * CurveLUTNumElems;
-	int32 BIndex = NextEntry * CurveLUTNumElems;
+	int32 AIndex = (int32)(PrevEntry * (float)CurveLUTNumElems);
+	int32 BIndex = (int32)(NextEntry * (float)CurveLUTNumElems);
 	FVector4 A = FVector4(ShaderLUT[AIndex], ShaderLUT[AIndex + 1], ShaderLUT[AIndex + 2], ShaderLUT[AIndex + 3]);
 	FVector4 B = FVector4(ShaderLUT[BIndex], ShaderLUT[BIndex + 1], ShaderLUT[BIndex + 2], ShaderLUT[BIndex + 3]);
 	return FMath::Lerp(A, B, Interp);
diff --git a/Engine/Plugins/FX/Niagara/Source/Niagara/Private/NiagaraDataInterfaceVectorCurve.cpp b/Engine/Plugins/FX/Niagara/Source/Niagara/Private/NiagaraDataInterfaceVectorCurve.cpp
--- a/Engine/Plugins/FX/Niagara/Source/Niagara/Private/NiagaraDataInterfaceVectorCurve.cpp
+++ b/Engine/Plugins/FX/Niagara/Source/Niagara/Private/NiagaraDataInterfaceVectorCurve.cpp
@@ -197,8 +197,8 @@ FORCEINLINE_DEBUGGABLE FVector UNiagaraDataInterfaceVectorCurve::SampleCurveInte
 	float NextEntry = PrevEntry < LUTNumSamplesMinusOne ? PrevEntry + 1.0f : PrevEntry;
 	float Interp = RemappedX - PrevEntry;
 
-	int32 AIndex = PrevEntry * CurveLUTNumElems;
-	int32 BIndex = NextEntry * CurveLUTNumElems;
+	int32 AIndex = (int32)(PrevEntry * (float)CurveLUTNumElems);
+	int32 BIndex = (int32)(NextEntry * (float)CurveLUTNumElems);
 	FVector A = FVector(ShaderLUT[AIndex], ShaderLUT[AIndex + 1], ShaderLUT[AIndex + 2]);
 	FVector B = FVector(ShaderLUT[BIndex], ShaderLUT[BIndex + 1], ShaderLUT[BIndex + 2]);
 	return FMath::Lerp(A, B, Interp);
diff --git a/Engine/Plugins/FX/Niagara/Source/Niagara/Private/NiagaraEmitter.cpp b/Engine/Plugins/FX/Niagara/Source/Niagara/Private/NiagaraEmitter.cpp
--- a/Engine/Plugins/FX/Niagara/Source/Niagara/Private/NiagaraEmitter.cpp
+++ b/Engine/Plugins/FX/Niagara/Source/Niagara/Private/NiagaraEmitter.cpp
@@ -2116,14 +2116,12 @@ bool UNiagaraEmitter::CanObtainParticleAttribute(const FNiagaraVariableBase& InV
 		const FNiagaraEmitterScriptProperties& SpawnScriptProps = EmitterData->SpawnScriptProps;
 		check(!SpawnScriptProps.Script->HasAnyFlags(RF_NeedPostLoad));
 
-		bool bContainsAttribute = SpawnScriptProps.Script->GetVMExecutableData().Attributes.Contains(InVar);
+		bool bContainsAttribute = SpawnScriptProps.Script->GetVMExecutableData().Attributes.ContainsByPredicate(FNiagaraVariableMatch(InVar.GetType(), InVar.GetName()));
 		if (!bContainsAttribute && InVar.GetType() == FNiagaraTypeDefinition::GetPositionDef())
 		{
 			// if we don't find a position type var we check for a vec3 type for backwards compatibility
 			OutBoundType = FNiagaraTypeDefinition::GetVec3Def();
-			FNiagaraVariableBase VarCopy = InVar;
-			VarCopy.SetType(OutBoundType);
-			bContainsAttribute = SpawnScriptProps.Script->GetVMExecutableData().Attributes.Contains(VarCopy);
+			bContainsAttribute = SpawnScriptProps.Script->GetVMExecutableData().Attributes.ContainsByPredicate(FNiagaraVariableMatch(OutBoundType, InVar.GetName()));
 		}
 		return bContainsAttribute;
 	}
diff --git a/Engine/Plugins/FX/Niagara/Source/Niagara/Private/NiagaraParameterStore.cpp b/Engine/Plugins/FX/Niagara/Source/Niagara/Private/NiagaraParameterStore.cpp
--- a/Engine/Plugins/FX/Niagara/Source/Niagara/Private/NiagaraParameterStore.cpp
+++ b/Engine/Plugins/FX/Niagara/Source/Niagara/Private/NiagaraParameterStore.cpp
@@ -720,7 +720,7 @@ bool FNiagaraParameterStore::RemoveParameter(const FNiagaraVariableBase& ToRemov
 			const int32 ExistingOffset = Existing.Offset;
 
 			//Add all but the one to remove to our
-			if (ExistingVar != ToRemove)
+			if (static_cast<const FNiagaraVariableBase&>(ExistingVar) != ToRemove)
 			{
 				if (ExistingVar.GetType().IsDataInterface())
 				{
diff --git a/Engine/Plugins/FX/Niagara/Source/Niagara/Private/NiagaraScript.cpp b/Engine/Plugins/FX/Niagara/Source/Niagara/Private/NiagaraScript.cpp
--- a/Engine/Plugins/FX/Niagara/Source/Niagara/Private/NiagaraScript.cpp
+++ b/Engine/Plugins/FX/Niagara/Source/Niagara/Private/NiagaraScript.cpp
@@ -1820,7 +1820,7 @@ FNiagaraCompilerTag* FNiagaraCompilerTag::FindTag(TArray< FNiagaraCompilerTag>&
 {
 	for (FNiagaraCompilerTag& Tag : InTags)
 	{
-		if (Tag.Variable == InSearchVar)
+		if (static_cast<const FNiagaraVariableBase&>(Tag.Variable) == InSearchVar)
 			return &Tag;
 	}
 	return nullptr;
@@ -1831,7 +1831,7 @@ const FNiagaraCompilerTag* FNiagaraCompilerTag::FindTag(const TArray< FNiagaraCo
 {
 	for (const FNiagaraCompilerTag& Tag : InTags)
 	{
-		if (Tag.Variable == InSearchVar)
+		if (static_cast<const FNiagaraVariableBase&>(Tag.Variable) == InSearchVar)
 			return &Tag;
 	}
 	return nullptr;
diff --git a/Engine/Plugins/FX/Niagara/Source/Niagara/Private/NiagaraSimCache.cpp b/Engine/Plugins/FX/Niagara/Source/Niagara/Private/NiagaraSimCache.cpp
--- a/Engine/Plugins/FX/Niagara/Source/Niagara/Private/NiagaraSimCache.cpp
+++ b/Engine/Plugins/FX/Niagara/Source/Niagara/Private/NiagaraSimCache.cpp
@@ -131,7 +131,7 @@ struct FNiagaraSimCacheHelper
 		for (const FNiagaraSimCacheVariable& SourceVariable : CacheLayout.Variables)
 		{
 			// Find variable, if it doesn't exist that's ok as the cache contains more data than is required
-			const int32 DataSetVariableIndex = CompiledData.Variables.IndexOfByPredicate([&](const FNiagaraVariable& DataSetVariable) { return DataSetVariable == SourceVariable.Variable; });
+			const int32 DataSetVariableIndex = CompiledData.Variables.IndexOfByPredicate(FNiagaraVariableMatch(SourceVariable.Variable.GetType(), SourceVariable.Variable.GetName()));
 			const FNiagaraVariableLayoutInfo* DestVariableLayout = nullptr;
 			if (DataSetVariableIndex != INDEX_NONE)
 			{
diff --git a/Engine/Plugins/FX/Niagara/Source/Niagara/Private/NiagaraSystem.cpp b/Engine/Plugins/FX/Niagara/Source/Niagara/Private/NiagaraSystem.cpp
--- a/Engine/Plugins/FX/Niagara/Source/Niagara/Private/NiagaraSystem.cpp
+++ b/Engine/Plugins/FX/Niagara/Source/Niagara/Private/NiagaraSystem.cpp
@@ -1935,14 +1935,12 @@ bool UNiagaraSystem::CanObtainSystemAttribute(const FNiagaraVariableBase& InVar,
 	if (SystemSpawnScript)
 	{
 		OutBoundType = InVar.GetType();
-		bool bContainsAttribute = SystemSpawnScript->GetVMExecutableData().Attributes.Contains(InVar);
+		bool bContainsAttribute = SystemSpawnScript->GetVMExecutableData().Attributes.ContainsByPredicate(FNiagaraVariableMatch(OutBoundType, InVar.GetName()));
 		if (!bContainsAttribute && InVar.GetType() == FNiagaraTypeDefinition::GetPositionDef())
 		{
 			// if we don't find a position type var we check for a vec3 type for backwards compatibility
 			OutBoundType = FNiagaraTypeDefinition::GetVec3Def();
-			FNiagaraVariableBase VarCopy = InVar;
-			VarCopy.SetType(OutBoundType);
-			bContainsAttribute = SystemSpawnScript->GetVMExecutableData().Attributes.Contains(VarCopy);
+			bContainsAttribute = SystemSpawnScript->GetVMExecutableData().Attributes.ContainsByPredicate(FNiagaraVariableMatch(OutBoundType, InVar.GetName()));
 		}
 		return bContainsAttribute;
 	}
diff --git a/Engine/Plugins/FX/Niagara/Source/Niagara/Public/NiagaraTypes.h b/Engine/Plugins/FX/Niagara/Source/Niagara/Public/NiagaraTypes.h
--- a/Engine/Plugins/FX/Niagara/Source/Niagara/Public/NiagaraTypes.h
+++ b/Engine/Plugins/FX/Niagara/Source/Niagara/Public/NiagaraTypes.h
@@ -1962,6 +1962,22 @@ private:
 	TArray<uint8> VarData;
 };
 
+struct FNiagaraVariableMatch
+{
+	explicit FNiagaraVariableMatch(const FNiagaraTypeDefinition& InType, const FName& InName)
+		: VariableBase(InType, InName)
+	{
+	}
+
+	bool operator()(const FNiagaraVariable& Other) const
+	{
+		return static_cast<const FNiagaraVariableBase&>(Other) == VariableBase;
+	}
+
+private:
+	FNiagaraVariableBase VariableBase;
+};
+
 
 template<>
 struct TStructOpsTypeTraits<FNiagaraVariable> : public TStructOpsTypeTraitsBase2<FNiagaraVariable>
diff --git a/Engine/Plugins/FX/Niagara/Source/NiagaraEditor/Private/NiagaraEditorUtilities.cpp b/Engine/Plugins/FX/Niagara/Source/NiagaraEditor/Private/NiagaraEditorUtilities.cpp
--- a/Engine/Plugins/FX/Niagara/Source/NiagaraEditor/Private/NiagaraEditorUtilities.cpp
+++ b/Engine/Plugins/FX/Niagara/Source/NiagaraEditor/Private/NiagaraEditorUtilities.cpp
@@ -1727,7 +1727,7 @@ int32 FNiagaraEditorUtilities::GetWeightForItem(const TSharedPtr<FNiagaraMenuAct
 			
 			FilterTerm = &InFilterTerms[FilterIndex];
 
-			float MaxWeightOfKeywords = INDEX_NONE;
+			float MaxWeightOfKeywords = (float)INDEX_NONE;
 			EWordMatchStyle MaxFoundMatchStyle = EWordMatchStyle::Contains;
 
 			auto CalculateSupersedingMatch = [&](float IterativeWeight, EWordMatchStyle IterativeFoundWordMatchStyle)
@@ -1774,7 +1774,7 @@ int32 FNiagaraEditorUtilities::GetWeightForItem(const TSharedPtr<FNiagaraMenuAct
 				}
 			}
 
-			if(MaxWeightOfKeywords != INDEX_NONE)
+			if(MaxWeightOfKeywords != (float)INDEX_NONE)
 			{
 				// if the match style that applies is 'stricter' than the required match style of the array, we consider it a hit
 				if(MaxFoundMatchStyle >= WordMatchStyle)
diff --git a/Engine/Plugins/FX/Niagara/Source/NiagaraEditor/Private/ViewModels/NiagaraParameterPanelViewModel.cpp b/Engine/Plugins/FX/Niagara/Source/NiagaraEditor/Private/ViewModels/NiagaraParameterPanelViewModel.cpp
--- a/Engine/Plugins/FX/Niagara/Source/NiagaraEditor/Private/ViewModels/NiagaraParameterPanelViewModel.cpp
+++ b/Engine/Plugins/FX/Niagara/Source/NiagaraEditor/Private/ViewModels/NiagaraParameterPanelViewModel.cpp
@@ -2312,7 +2312,7 @@ void FNiagaraSystemToolkitParameterPanelViewModel::OnParameterRenamedExternally(
 
 					if (!bFound)
 					{
-						if (const FNiagaraParameterPanelItem* FoundItemPtr = CachedViewedItems.FindByPredicate([&InOldVar](const FNiagaraParameterPanelItem& Item) { return Item.GetVariable() == InOldVar; }))
+						if (const FNiagaraParameterPanelItem* FoundItemPtr = CachedViewedItems.FindByPredicate([&InOldVar](const FNiagaraParameterPanelItem& Item) { return static_cast<const FNiagaraVariableBase&>(Item.GetVariable()) == InOldVar; }))
 						{
 							RenameParameter(*FoundItemPtr, InNewVar.GetName());
 						}
@@ -2352,7 +2352,7 @@ void FNiagaraSystemToolkitParameterPanelViewModel::OnParameterRenamedExternally(
 
 		if (!bFound)
 		{
-			if (const FNiagaraParameterPanelItem* FoundItemPtr = CachedViewedItems.FindByPredicate([&InOldVar](const FNiagaraParameterPanelItem& Item) { return Item.GetVariable() == InOldVar; }))
+			if (const FNiagaraParameterPanelItem* FoundItemPtr = CachedViewedItems.FindByPredicate([&InOldVar](const FNiagaraParameterPanelItem& Item) { return static_cast<const FNiagaraVariableBase&>(Item.GetVariable()) == InOldVar; }))
 			{
 				RenameParameter(*FoundItemPtr, InNewVar.GetName());
 			}
