- Add SubUV to renderer info
- Move to a shader template and clean up the code a little
- This removes the GPU data deduplication, but this should be done via DI deduplication and is not really robust when using pointers

diff --git a/Engine/Plugins/FX/Niagara/Shaders/Private/NiagaraDataInterfaceMeshRendererInfoTemplate.ush b/Engine/Plugins/FX/Niagara/Shaders/Private/NiagaraDataInterfaceMeshRendererInfoTemplate.ush
new file mode 100644
--- /dev/null
+++ b/Engine/Plugins/FX/Niagara/Shaders/Private/NiagaraDataInterfaceMeshRendererInfoTemplate.ush
@@ -0,0 +1,40 @@
+// Copyright Epic Games, Inc. All Rights Reserved.
+
+uint			{ParameterName}_bSubImageBlend;
+float2			{ParameterName}_SubImageSize;
+uint			{ParameterName}_NumMeshes;
+Buffer<float>	{ParameterName}_MeshDataBuffer;
+
+void GetNumMeshes(out int OutNumMeshes)
+{
+	OutNumMeshes = int({ParameterName}_NumMeshes);
+}
+
+void GetMeshLocalBounds(int MeshIndex, out float3 OutMinBounds, out float3 OutMaxBounds, out float3 OutSize)
+{
+	OutMinBounds = (float3)0;
+	OutMaxBounds = (float3)0;
+	OutSize = (float3)0;
+	if ( {ParameterName}_NumMeshes > 0 )
+	{
+		const uint MeshDataNumFloats = 6;
+		const uint BufferOffs = clamp(MeshIndex, 0, int({ParameterName}_NumMeshes - 1)) * MeshDataNumFloats;
+		OutMinBounds = float3(
+			{ParameterName}_MeshDataBuffer[BufferOffs + 0],
+			{ParameterName}_MeshDataBuffer[BufferOffs + 1],
+			{ParameterName}_MeshDataBuffer[BufferOffs + 2]
+		);
+		OutMaxBounds = float3(
+			{ParameterName}_MeshDataBuffer[BufferOffs + 3],
+			{ParameterName}_MeshDataBuffer[BufferOffs + 4],
+			{ParameterName}_MeshDataBuffer[BufferOffs + 5]
+		);
+		OutSize = OutMaxBounds - OutMinBounds;
+	}
+}
+
+void GetSubUVDetails_{ParameterName}(out bool bBlendEnabled, out float2 SubImageSize)
+{
+	bBlendEnabled = {ParameterName}_bSubImageBlend != 0;
+	SubImageSize = {ParameterName}_SubImageSize;
+}
diff --git a/Engine/Plugins/FX/Niagara/Source/Niagara/Classes/NiagaraDataInterfaceMeshRendererInfo.h b/Engine/Plugins/FX/Niagara/Source/Niagara/Classes/NiagaraDataInterfaceMeshRendererInfo.h
--- a/Engine/Plugins/FX/Niagara/Source/Niagara/Classes/NiagaraDataInterfaceMeshRendererInfo.h
+++ b/Engine/Plugins/FX/Niagara/Source/Niagara/Classes/NiagaraDataInterfaceMeshRendererInfo.h
@@ -10,19 +10,17 @@
 class UNiagaraMeshRendererProperties;
 class FNDIMeshRendererInfo;
 
-using FNDIMeshRendererInfoRef = TSharedRef<FNDIMeshRendererInfo, ESPMode::ThreadSafe>;
-using FNDIMeshRendererInfoPtr = TSharedPtr<FNDIMeshRendererInfo, ESPMode::ThreadSafe>;
-
 /** This Data Interface can be used to query information about the mesh renderers of an emitter */
 UCLASS(EditInlineNew, Category = "Mesh Particles", meta = (DisplayName = "Mesh Renderer Info"))
 class NIAGARA_API UNiagaraDataInterfaceMeshRendererInfo : public UNiagaraDataInterface
 {
 	GENERATED_UCLASS_BODY()
 
-	BEGIN_SHADER_PARAMETER_STRUCT(FShaderParameters,)
-		SHADER_PARAMETER(uint32,			NumMeshes)
-		SHADER_PARAMETER_SRV(Buffer<float>,	MeshDataBuffer)
-	END_SHADER_PARAMETER_STRUCT()
+	struct FMeshData
+	{
+		FVector3f MinLocalBounds = FVector3f(ForceInitToZero);
+		FVector3f MaxLocalBounds = FVector3f(ForceInitToZero);
+	};
 
 public:
 	UNiagaraMeshRendererProperties* GetMeshRenderer() const { return MeshRenderer; }
@@ -37,6 +35,8 @@ public:
 #endif 
 	//UObject Interface End
 
+	void OnMeshRendererChanged(UNiagaraMeshRendererProperties* NewMeshRenderer);
+
 	//UNiagaraDataInterface Interface
 	virtual void GetFunctions(TArray<FNiagaraFunctionSignature>& OutFunctions) override;
 	virtual void GetVMExternalFunction(const FVMExternalFunctionBindingInfo& BindingInfo, void* InstanceData, FVMExternalFunction& OutFunc) override;
@@ -62,12 +62,17 @@ protected:
 	virtual bool CopyToInternal(UNiagaraDataInterface* Destination) const override;
 	virtual void PushToRenderThreadImpl() override;
 
-	void GetNumMeshes(FVectorVMExternalFunctionContext& Context);
-	void GetMeshLocalBounds(FVectorVMExternalFunctionContext& Context);
+	void VMGetNumMeshes(FVectorVMExternalFunctionContext& Context);
+	void VMGetMeshLocalBounds(FVectorVMExternalFunctionContext& Context);
+	void VMGetSubUVDetails(FVectorVMExternalFunctionContext& Context);
 
 	/** The name of the mesh renderer */
 	UPROPERTY(EditAnywhere, Category = "Source")
 	TObjectPtr<UNiagaraMeshRendererProperties> MeshRenderer;
 
-	FNDIMeshRendererInfoPtr Info;
+	TArray<FMeshData> CachedMeshData;
+
+#if WITH_EDITOR
+	FDelegateHandle OnMeshRendererChangedHandle;
+#endif
 };
diff --git a/Engine/Plugins/FX/Niagara/Source/Niagara/Private/NiagaraDataInterfaceMeshRendererInfo.cpp b/Engine/Plugins/FX/Niagara/Source/Niagara/Private/NiagaraDataInterfaceMeshRendererInfo.cpp
--- a/Engine/Plugins/FX/Niagara/Source/Niagara/Private/NiagaraDataInterfaceMeshRendererInfo.cpp
+++ b/Engine/Plugins/FX/Niagara/Source/Niagara/Private/NiagaraDataInterfaceMeshRendererInfo.cpp
@@ -11,14 +11,21 @@
 
 #define LOCTEXT_NAMESPACE "NiagaraDataInterfaceMeshRendererInfo"
 
-class FNDIMeshRendererInfoGPUData;
-using FNDIMeshRendererInfoGPUDataRef = TSharedPtr<FNDIMeshRendererInfoGPUData, ESPMode::ThreadSafe>;
-using FNDIMeshRendererInfoGPUDataPtr = TSharedPtr<FNDIMeshRendererInfoGPUData, ESPMode::ThreadSafe>;
-
 namespace NDIMeshRendererInfoInternal
 {
-	const FName GetNumMeshesName("GetNumMeshes");
-	const FName GetMeshLocalBoundsName("GetMeshLocalBounds");
+	BEGIN_SHADER_PARAMETER_STRUCT(FShaderParameters,)
+		SHADER_PARAMETER(uint32,			bSubImageBlend)
+		SHADER_PARAMETER(FVector2f,			SubImageSize)
+		SHADER_PARAMETER(uint32,			NumMeshes)
+		SHADER_PARAMETER_SRV(Buffer<float>,	MeshDataBuffer)
+	END_SHADER_PARAMETER_STRUCT()
+
+
+	static const TCHAR* TemplateShaderFile = TEXT("/Plugin/FX/Niagara/Private/NiagaraDataInterfaceSpriteRendererInfoTemplate.ush");
+
+	static const FName GetNumMeshesName("GetNumMeshes");
+	static const FName GetMeshLocalBoundsName("GetMeshLocalBounds");
+	static const FName GetSubUVDetailsName("GetSubUVDetails");
 }
 
 enum class ENDIMeshRendererInfoVersion : uint32
@@ -30,228 +37,17 @@ enum class ENDIMeshRendererInfoVersion : uint32
 	LatestVersion = VersionPlusOne - 1
 };
 
-/** Holds information (for both CPU and GPU) accessed by this data interface for a given renderer */
-class FNDIMeshRendererInfo
-{
-public:
-	struct FMeshData
-	{
-		FVector3f MinLocalBounds = FVector3f(ForceInitToZero);
-		FVector3f MaxLocalBounds = FVector3f(ForceInitToZero);
-	};
-
-	using FMeshDataArray = TArray<FMeshData>;
-
-	void AddRef() { ++RefCount; }
-	const FMeshDataArray& GetMeshData() { return MeshData; }
-	FNDIMeshRendererInfoGPUDataRef GetOrCreateGPUData();
-
-	static FNDIMeshRendererInfoRef Acquire(UNiagaraMeshRendererProperties& Renderer);
-	static void Release(UNiagaraMeshRendererProperties& Renderer, FNDIMeshRendererInfoPtr& Info);
-
-private:
-	FMeshDataArray MeshData;
-	FNDIMeshRendererInfoGPUDataPtr GPUData;
-	uint32 RefCount = 0;
-#if WITH_EDITOR
-	FDelegateHandle OnChangedHandle;
-#endif
-
-	using FCacheMap = TMap<UNiagaraMeshRendererProperties*, FNDIMeshRendererInfoRef>;
-	static FCacheMap CachedData;
-
-	static void ResetMeshData(const UNiagaraMeshRendererProperties& Renderer, FMeshDataArray& OutArray);
-};
-
-/** This is a resource that holds the static GPU buffer data of the info for a given renderer */
-class FNDIMeshRendererInfoGPUData : public FRenderResource
-{
-public:
-	FNDIMeshRendererInfoGPUData(const FNDIMeshRendererInfo::FMeshDataArray& InMeshData)
-		: MeshData(InMeshData)
-	{}
-
-	uint32 GetNumMeshes() { return MeshData.Num(); }
-	FBufferRHIRef GetMeshDataBufferRHI() const { return BufferMeshDataRHI; }
-	FShaderResourceViewRHIRef GetMeshDataBufferSRV() const { return BufferMeshDataSRV; }
-
-	virtual void InitRHI() override final
-	{
-		uint32 SizeByte = MeshData.Num() * sizeof(FVector3f) * 2;
-		
-		if (SizeByte > 0)
-		{
-			FRHIResourceCreateInfo CreateInfo(TEXT("FNDIMeshRendererInfoGPUData"));
-			BufferMeshDataRHI = RHICreateBuffer(SizeByte, BUF_Static | BUF_VertexBuffer | BUF_ShaderResource, 0, ERHIAccess::VertexOrIndexBuffer | ERHIAccess::SRVMask, CreateInfo);
-			FVector3f* BufferDataVector = reinterpret_cast<FVector3f*>(RHILockBuffer(BufferMeshDataRHI, 0, SizeByte, RLM_WriteOnly));
-			for (const auto& Mesh : MeshData)
-			{
-				*BufferDataVector++ = Mesh.MinLocalBounds;
-				*BufferDataVector++ = Mesh.MaxLocalBounds;
-			}
-			RHIUnlockBuffer(BufferMeshDataRHI);
-
-			BufferMeshDataSRV = RHICreateShaderResourceView(BufferMeshDataRHI, sizeof(float), PF_R32_FLOAT);
-		}
-
-#if STATS
-		GPUMemory = SizeByte;
-		INC_MEMORY_STAT_BY(STAT_NiagaraGPUDataInterfaceMemory, SizeByte);
-#endif
-	}
-
-	virtual void ReleaseRHI() override final
-	{
-		BufferMeshDataSRV.SafeRelease();
-		BufferMeshDataRHI.SafeRelease();
-		
-#if STATS
-		DEC_MEMORY_STAT_BY(STAT_NiagaraGPUDataInterfaceMemory, GPUMemory);
-		GPUMemory = 0;
-#endif
-	}
-
-private:
-	FBufferRHIRef BufferMeshDataRHI = nullptr;
-	FShaderResourceViewRHIRef BufferMeshDataSRV = nullptr;
-	const FNDIMeshRendererInfo::FMeshDataArray& MeshData; // cached from FNDIMeshRendererInfo, which is guaranteed to live longer than we are
-#if STATS
-	uint32 GPUMemory = 0;
-#endif
-};
-
 /** The render thread proxy of the data interface */
 struct FNDIMeshRendererInfoProxy : public FNiagaraDataInterfaceProxy
 {
-	FNDIMeshRendererInfoGPUDataPtr GPUData;
+	bool		bSubImageBlend = false;
+	FVector2f	SubImageSize = FVector2f::ZeroVector;
+	uint32		NumMeshes;
+	FReadBuffer	GPUMeshData;		//-OPT: Can shared if pointing to same resources
 
 	virtual int32 PerInstanceDataPassedToRenderThreadSize() const override { return 0; }
 };
 
-///////////////////////////////////////////////////////////////////////////////////////////////
-
-FNDIMeshRendererInfo::FCacheMap FNDIMeshRendererInfo::CachedData;
-
-FNDIMeshRendererInfoGPUDataRef FNDIMeshRendererInfo::GetOrCreateGPUData()
-{
-	if (!GPUData.IsValid())
-	{
-		GPUData = MakeShared<FNDIMeshRendererInfoGPUData, ESPMode::ThreadSafe>(MeshData);
-		BeginInitResource(GPUData.Get());
-	}
-	return GPUData.ToSharedRef();
-}
-
-FNDIMeshRendererInfoRef FNDIMeshRendererInfo::Acquire(UNiagaraMeshRendererProperties& Renderer)
-{
-	FNDIMeshRendererInfoRef* RefPtr = CachedData.Find(&Renderer);
-	if (RefPtr)
-	{
-		++(*RefPtr)->RefCount;
-		return *RefPtr;
-	}
-
-	FNDIMeshRendererInfoRef Info = MakeShared<FNDIMeshRendererInfo, ESPMode::ThreadSafe>();
-	++Info->RefCount;
-
-	ResetMeshData(Renderer, Info->MeshData);
-
-#if WITH_EDITOR
-	Info->OnChangedHandle = Renderer.OnChanged().AddLambda(
-		[&Renderer, Info]()
-		{
-			if (Info->GPUData.IsValid())
-			{
-				// The render thread could be accessing the mesh data, so we have to update it in a render cmd
-				FMeshDataArray TempArray;
-				ResetMeshData(Renderer, TempArray);
-				ENQUEUE_RENDER_COMMAND(FDIMeshRendererUpdateMeshDataBuffer)
-				(
-					[Info, MeshData = MoveTemp(TempArray)](FRHICommandList& RHICmdList)
-					{
-						Info->MeshData = MeshData;
-						if (Info->GPUData.IsValid())
-						{
-							// Re-create the buffers
-							Info->GPUData->ReleaseRHI();
-							Info->GPUData->InitRHI();
-						}
-					}
-				);
-			}
-			else
-			{
-				// We've never pushed our data to the RT so we're safe to stomp this data without worry of data race
-				ResetMeshData(Renderer, Info->MeshData);
-			}
-		}
-	);
-#endif
-	
-	CachedData.Add(&Renderer, Info);
-	return Info;
-}
-
-void FNDIMeshRendererInfo::Release(UNiagaraMeshRendererProperties& Renderer, FNDIMeshRendererInfoPtr& Info)
-{
-	if (!Info.IsValid())
-	{
-		return;
-	}
-
-	check(Info->RefCount > 0);
-	if (--Info->RefCount == 0)
-	{
-		// Last reference is out, we can stop caching it and release it
-		CachedData.Remove(&Renderer);
-
-#if WITH_EDITOR
-		Renderer.OnChanged().Remove(Info->OnChangedHandle);
-		Info->OnChangedHandle.Reset();
-#endif
-
-		if (Info->GPUData.IsValid())
-		{
-			BeginReleaseResource(Info->GPUData.Get());
-
-			// We have to release this on the render thread because a render command could be holding references
-			ENQUEUE_RENDER_COMMAND(FDIMeshRendererInfoReleaseInfo)
-			(
-				[Info_RT = MoveTemp(Info)](FRHICommandList& RHICmdList) mutable
-				{
-					Info_RT->GPUData = nullptr;
-					Info_RT = nullptr;
-				}
-			);
-		}
-	}
-
-	Info = nullptr;
-}
-
-void FNDIMeshRendererInfo::ResetMeshData(const UNiagaraMeshRendererProperties& Renderer, FMeshDataArray& OutMeshData)
-{
-	OutMeshData.Reset(Renderer.Meshes.Num());
-	for (const auto& MeshSlot : Renderer.Meshes)
-	{
-		FMeshData& NewMeshData = OutMeshData.AddDefaulted_GetRef();
-		NewMeshData.MinLocalBounds = FVector3f(EForceInit::ForceInitToZero);
-		NewMeshData.MaxLocalBounds = FVector3f(EForceInit::ForceInitToZero);
-
-		if (MeshSlot.Mesh)
-		{
-			const FBox LocalBounds = MeshSlot.Mesh->GetExtendedBounds().GetBox();
-			if (LocalBounds.IsValid)
-			{
-				// Scale the local bounds if there's a scale on this slot
-				// TODO: Should we also apply the pivot offset if it's in mesh space? Seems like that might be strange
-				NewMeshData.MinLocalBounds = FVector3f(LocalBounds.Min * MeshSlot.Scale);
-				NewMeshData.MaxLocalBounds = FVector3f(LocalBounds.Max * MeshSlot.Scale);
-			}
-		}
-	}	
-}
-
 ////////////////////////////////////////////////////////////////////////////////////////////////
 
 UNiagaraDataInterfaceMeshRendererInfo::UNiagaraDataInterfaceMeshRendererInfo(FObjectInitializer const& ObjectInitializer)
@@ -283,9 +79,7 @@ void UNiagaraDataInterfaceMeshRendererInfo::PostLoad()
 	if (MeshRenderer)
 	{
 		MeshRenderer->ConditionalPostLoad();
-
-		Info = FNDIMeshRendererInfo::Acquire(*MeshRenderer);
-		MarkRenderDataDirty();
+		OnMeshRendererChanged(MeshRenderer);
 	}
 }
 
@@ -293,10 +87,7 @@ void UNiagaraDataInterfaceMeshRendererInfo::BeginDestroy()
 {
 	Super::BeginDestroy();
 
-	if (MeshRenderer)
-	{
-		FNDIMeshRendererInfo::Release(*MeshRenderer, Info);
-	}
+	OnMeshRendererChanged(nullptr);
 }
 
 #if WITH_EDITOR
@@ -305,7 +96,7 @@ void UNiagaraDataInterfaceMeshRendererInfo::PreEditChange(FProperty* PropertyAbo
 {
 	if (MeshRenderer && PropertyAboutToChange && PropertyAboutToChange->GetFName() == GET_MEMBER_NAME_CHECKED(UNiagaraDataInterfaceMeshRendererInfo, MeshRenderer))
 	{		
-		FNDIMeshRendererInfo::Release(*MeshRenderer, Info);
+		OnMeshRendererChanged(nullptr);
 	}
 }
 
@@ -317,16 +108,48 @@ void UNiagaraDataInterfaceMeshRendererInfo::PostEditChangeProperty(struct FPrope
 	if (PropertyChangedEvent.Property == nullptr || (PropertyChangedEvent.Property &&
 		PropertyChangedEvent.Property->GetFName() == GET_MEMBER_NAME_CHECKED(UNiagaraDataInterfaceMeshRendererInfo, MeshRenderer)))
 	{
-		if (MeshRenderer)
-		{
-			Info = FNDIMeshRendererInfo::Acquire(*MeshRenderer);
-		}
-		MarkRenderDataDirty();
+		OnMeshRendererChanged(MeshRenderer);
 	}
 }
 
 #endif // WITH_EDITOR
 
+void UNiagaraDataInterfaceMeshRendererInfo::OnMeshRendererChanged(UNiagaraMeshRendererProperties* NewMeshRenderer)
+{
+#if WITH_EDITOR
+	if (MeshRenderer)
+	{
+		MeshRenderer->OnChanged().Remove(OnMeshRendererChangedHandle);
+		OnMeshRendererChangedHandle.Reset();
+	}
+	CachedMeshData.Empty();
+
+	if (NewMeshRenderer)
+	{
+		OnMeshRendererChangedHandle = NewMeshRenderer->OnChanged().AddLambda([this]() { MarkRenderDataDirty(); });
+
+		CachedMeshData.AddDefaulted(NewMeshRenderer->Meshes.Num());
+		for ( int32 i=0; i < NewMeshRenderer->Meshes.Num(); ++i )
+		{
+			FNiagaraMeshRendererMeshProperties& MeshProperties = NewMeshRenderer->Meshes[i];
+			if ( MeshProperties.Mesh == nullptr )
+			{
+				continue;
+			}
+
+			const FBox LocalBounds = MeshProperties.Mesh->GetExtendedBounds().GetBox();
+			if ( LocalBounds.IsValid )
+			{
+				FMeshData& CachedMesh = CachedMeshData[i];
+				CachedMesh.MinLocalBounds = FVector3f(LocalBounds.Min * MeshProperties.Scale);
+				CachedMesh.MaxLocalBounds = FVector3f(LocalBounds.Max * MeshProperties.Scale);
+			}
+		}
+	}
+#endif
+	MarkRenderDataDirty();
+}
+
 void UNiagaraDataInterfaceMeshRendererInfo::GetFunctions(TArray<FNiagaraFunctionSignature>& OutFunctions)
 {
 	{
@@ -340,7 +163,6 @@ void UNiagaraDataInterfaceMeshRendererInfo::GetFunctions(TArray<FNiagaraFunction
 		Signature.Inputs.Add(FNiagaraVariable(FNiagaraTypeDefinition(GetClass()), TEXT("MeshRendererInfo")));
 		Signature.Outputs.Add(FNiagaraVariable(FNiagaraTypeDefinition::GetIntDef(), TEXT("OutNumMeshes")));
 	}
-
 	{
 		FNiagaraFunctionSignature& Signature = OutFunctions.AddDefaulted_GetRef();
 		Signature.Name = NDIMeshRendererInfoInternal::GetMeshLocalBoundsName;
@@ -355,17 +177,32 @@ void UNiagaraDataInterfaceMeshRendererInfo::GetFunctions(TArray<FNiagaraFunction
 		Signature.Outputs.Add(FNiagaraVariable(FNiagaraTypeDefinition::GetVec3Def(), TEXT("OutMaxBounds")));
 		Signature.Outputs.Add(FNiagaraVariable(FNiagaraTypeDefinition::GetVec3Def(), TEXT("OutSize")));
 	}
+	{
+		FNiagaraFunctionSignature& Signature = OutFunctions.AddDefaulted_GetRef();
+		Signature.Name = NDIMeshRendererInfoInternal::GetSubUVDetailsName;
+#if WITH_EDITORONLY_DATA
+		Signature.FunctionVersion = (uint32)ENDIMeshRendererInfoVersion::LatestVersion;
+#endif
+		Signature.bMemberFunction = true;
+		Signature.Inputs.Emplace(FNiagaraTypeDefinition(GetClass()), TEXT("MeshRendererInfo"));
+		Signature.Outputs.Emplace(FNiagaraTypeDefinition::GetBoolDef(), TEXT("BlendEnabled"));
+		Signature.Outputs.Emplace(FNiagaraTypeDefinition::GetVec2Def(), TEXT("SubImageSize"));
+	}
 }
 
 void UNiagaraDataInterfaceMeshRendererInfo::GetVMExternalFunction(const FVMExternalFunctionBindingInfo& BindingInfo, void* InstanceData, FVMExternalFunction& OutFunc)
 {
 	if (BindingInfo.Name == NDIMeshRendererInfoInternal::GetNumMeshesName)
 	{
-		OutFunc = FVMExternalFunction::CreateUObject(this, &UNiagaraDataInterfaceMeshRendererInfo::GetNumMeshes);
+		OutFunc = FVMExternalFunction::CreateUObject(this, &UNiagaraDataInterfaceMeshRendererInfo::VMGetNumMeshes);
 	}
 	else if (BindingInfo.Name == NDIMeshRendererInfoInternal::GetMeshLocalBoundsName)
 	{
-		OutFunc = FVMExternalFunction::CreateUObject(this, &UNiagaraDataInterfaceMeshRendererInfo::GetMeshLocalBounds);
+		OutFunc = FVMExternalFunction::CreateUObject(this, &UNiagaraDataInterfaceMeshRendererInfo::VMGetMeshLocalBounds);
+	}
+	else if (BindingInfo.Name == NDIMeshRendererInfoInternal::GetSubUVDetailsName)
+	{
+		OutFunc = FVMExternalFunction::CreateUObject(this, &UNiagaraDataInterfaceMeshRendererInfo::VMGetSubUVDetails);
 	}
 }
 
@@ -384,89 +221,51 @@ bool UNiagaraDataInterfaceMeshRendererInfo::Equals(const UNiagaraDataInterface*
 bool UNiagaraDataInterfaceMeshRendererInfo::AppendCompileHash(FNiagaraCompileHashVisitor* InVisitor) const
 {
 	bool bSuccess = Super::AppendCompileHash(InVisitor);
-	bSuccess &= InVisitor->UpdateShaderParameters<FShaderParameters>();
+	bSuccess &= InVisitor->UpdateString(TEXT("NiagaraDataInterfaceMeshRendererInfoHLSLSource"), GetShaderFileHash(NDIMeshRendererInfoInternal::TemplateShaderFile, EShaderPlatform::SP_PCD3D_SM5).ToString());
+	bSuccess &= InVisitor->UpdateShaderParameters<NDIMeshRendererInfoInternal::FShaderParameters>();
 	return bSuccess;
 }
 
 void UNiagaraDataInterfaceMeshRendererInfo::GetParameterDefinitionHLSL(const FNiagaraDataInterfaceGPUParamInfo& ParamInfo, FString& OutHLSL)
 {
-	OutHLSL.Appendf(TEXT("uint %s_NumMeshes;\n"), *ParamInfo.DataInterfaceHLSLSymbol);
-	OutHLSL.Appendf(TEXT("Buffer<float> %s_MeshDataBuffer;\n"), *ParamInfo.DataInterfaceHLSLSymbol);
+	TMap<FString, FStringFormatArg> TemplateArgs =
+	{
+		{TEXT("ParameterName"),	ParamInfo.DataInterfaceHLSLSymbol},
+	};
+
+	FString TemplateFile;
+	LoadShaderSourceFile(NDIMeshRendererInfoInternal::TemplateShaderFile, EShaderPlatform::SP_PCD3D_SM5, &TemplateFile, nullptr);
+	OutHLSL += FString::Format(*TemplateFile, TemplateArgs);
 }
 
 bool UNiagaraDataInterfaceMeshRendererInfo::GetFunctionHLSL(const FNiagaraDataInterfaceGPUParamInfo& ParamInfo, const FNiagaraDataInterfaceGeneratedFunction& FunctionInfo, int FunctionInstanceIndex, FString& OutHLSL)
 {
-	const FStringFormatNamedArguments Args = {
-		{TEXT("FuncName"), FunctionInfo.InstanceName},
-		{TEXT("DIName"), ParamInfo.DataInterfaceHLSLSymbol}
-	};
-
-	if (FunctionInfo.DefinitionName == NDIMeshRendererInfoInternal::GetNumMeshesName)
+	static const TSet<FName> ValidGpuFunction =
 	{
-		OutHLSL += FString::Format(TEXT(R"(
-			void {FuncName}(out int OutNumMeshes)
-			{
-				OutNumMeshes = {DIName}_NumMeshes;
-			}
-			)"),
-			Args);
-
-		return true;
-	}
-	else if (FunctionInfo.DefinitionName == NDIMeshRendererInfoInternal::GetMeshLocalBoundsName)
-	{
-		OutHLSL += FString::Format(TEXT(R"(
-			void {FuncName}(in int MeshIndex, out float3 OutMinBounds, out float3 OutMaxBounds, out float3 OutSize)
-			{
-				OutMinBounds = (float3)0;
-				OutMaxBounds = (float3)0;
-				OutSize = (float3)0;
-				if ({DIName}_NumMeshes > 0)
-				{
-					const uint MeshDataNumFloats = 6;
-					const uint BufferOffs = clamp(MeshIndex, 0, int({DIName}_NumMeshes - 1)) * MeshDataNumFloats;
-					OutMinBounds = float3(
-						{DIName}_MeshDataBuffer[BufferOffs + 0],
-						{DIName}_MeshDataBuffer[BufferOffs + 1],
-						{DIName}_MeshDataBuffer[BufferOffs + 2]
-					);
-					OutMaxBounds = float3(
-						{DIName}_MeshDataBuffer[BufferOffs + 3],
-						{DIName}_MeshDataBuffer[BufferOffs + 4],
-						{DIName}_MeshDataBuffer[BufferOffs + 5]
-					);
-					OutSize = OutMaxBounds - OutMinBounds;
-				}
-			}
-			)"),
-			Args);
-
-		return true;
-	}
+		NDIMeshRendererInfoInternal::GetNumMeshesName,
+		NDIMeshRendererInfoInternal::GetMeshLocalBoundsName,
+		NDIMeshRendererInfoInternal::GetSubUVDetailsName,
+	};
 
-	return false;
+	return ValidGpuFunction.Contains(FunctionInfo.DefinitionName);
 }
 #endif
 
 void UNiagaraDataInterfaceMeshRendererInfo::BuildShaderParameters(FNiagaraShaderParametersBuilder& ShaderParametersBuilder) const
 {
-	ShaderParametersBuilder.AddNestedStruct<FShaderParameters>();
+	ShaderParametersBuilder.AddNestedStruct<NDIMeshRendererInfoInternal::FShaderParameters>();
 }
 
 void UNiagaraDataInterfaceMeshRendererInfo::SetShaderParameters(const FNiagaraDataInterfaceSetShaderParametersContext& Context) const
 {
+	using namespace NDIMeshRendererInfoInternal;
+
 	FNDIMeshRendererInfoProxy& DIProxy = Context.GetProxy<FNDIMeshRendererInfoProxy>();
-	FShaderParameters* ShaderParameters = Context.GetParameterNestedStruct<FShaderParameters>();
-	if (DIProxy.GPUData.IsValid())
-	{
-		ShaderParameters->NumMeshes			= DIProxy.GPUData->GetNumMeshes();
-		ShaderParameters->MeshDataBuffer	= DIProxy.GPUData->GetMeshDataBufferSRV();
-	}
-	else
-	{
-		ShaderParameters->NumMeshes			= 0;
-		ShaderParameters->MeshDataBuffer	= FNiagaraRenderer::GetDummyFloatBuffer();
-	}
+	FShaderParameters* ShaderParameters	= Context.GetParameterNestedStruct<FShaderParameters>();
+	ShaderParameters->bSubImageBlend	= DIProxy.bSubImageBlend;
+	ShaderParameters->SubImageSize		= DIProxy.SubImageSize;
+	ShaderParameters->NumMeshes			= DIProxy.NumMeshes;
+	ShaderParameters->MeshDataBuffer	= FNiagaraRenderer::GetSrvOrDefaultFloat(DIProxy.GPUMeshData.SRV);
 }
 
 #if WITH_EDITOR
@@ -528,56 +327,61 @@ bool UNiagaraDataInterfaceMeshRendererInfo::CopyToInternal(UNiagaraDataInterface
 
 
 	auto OtherTyped = CastChecked<UNiagaraDataInterfaceMeshRendererInfo>(Destination);
-	
-	if (OtherTyped->MeshRenderer)
-	{
-		FNDIMeshRendererInfo::Release(*OtherTyped->MeshRenderer, OtherTyped->Info);
-	}
-	
-	OtherTyped->MeshRenderer = MeshRenderer;	
-	OtherTyped->Info = Info;
+	OtherTyped->OnMeshRendererChanged(nullptr);
+	OtherTyped->MeshRenderer = MeshRenderer;
+	OtherTyped->OnMeshRendererChanged(MeshRenderer);
 
-	// Check to add a reference to the per-renderer data
-	if (Info.IsValid())
-	{
-		Info->AddRef();
-	}
-	else if (MeshRenderer != nullptr && !EnumHasAnyFlags(OtherTyped->GetFlags(), RF_NeedPostLoad))
-	{
-		// This is currently necessary because, for some reason, it's possible that data interfaces that have not been post-loaded to be copied to
-		// another data interface object that has, and was previously causing this data to never get acquired on the copy
-		OtherTyped->Info = FNDIMeshRendererInfo::Acquire(*MeshRenderer);
-	}
 	return true;
 }
 
 void UNiagaraDataInterfaceMeshRendererInfo::PushToRenderThreadImpl()
 {
-	if (MeshRenderer && Info.IsValid())
-	{
-		auto TypedProxy = GetProxyAs<FNDIMeshRendererInfoProxy>();
-		ENQUEUE_RENDER_COMMAND(FDIMeshRendererInfoPushToRT)
-		(
-			[TypedProxy, GPUData_RT=Info->GetOrCreateGPUData()](FRHICommandList& RHICmdList)
+	FNDIMeshRendererInfoProxy* TypedProxy = GetProxyAs<FNDIMeshRendererInfoProxy>();
+
+	const UNiagaraMeshRendererProperties* Properties = MeshRenderer ? MeshRenderer : GetDefault<UNiagaraMeshRendererProperties>();
+
+	const bool bSubImageBlend_RT	= Properties->bSubImageBlend;
+	const FVector2f SubImageSize_RT	= FVector2f(Properties->SubImageSize);
+	ENQUEUE_RENDER_COMMAND(FDIMeshRendererInfoPushToRT)
+	(
+		[TypedProxy, CachedMeshData_RT=CachedMeshData, bSubImageBlend_RT, SubImageSize_RT](FRHICommandList& RHICmdList)
+		{
+			TypedProxy->bSubImageBlend = bSubImageBlend_RT;
+			TypedProxy->SubImageSize = SubImageSize_RT;
+			TypedProxy->NumMeshes = CachedMeshData_RT.Num();
+			TypedProxy->GPUMeshData.Release();
+			if ( TypedProxy->NumMeshes > 0 )
 			{
-				TypedProxy->GPUData = GPUData_RT;
+				TypedProxy->GPUMeshData.Initialize(TEXT("UNiagaraDataInterfaceMeshRendererInfo"), sizeof(float), CachedMeshData_RT.Num() * 6, PF_R32_FLOAT, BUF_Static);
+				float* GpuMeshData = reinterpret_cast<float*>(RHILockBuffer(TypedProxy->GPUMeshData.Buffer, 0, TypedProxy->GPUMeshData.NumBytes, RLM_WriteOnly));
+				for ( const UNiagaraDataInterfaceMeshRendererInfo::FMeshData& MeshData : CachedMeshData_RT )
+				{
+					*GpuMeshData++ = MeshData.MinLocalBounds.X;
+					*GpuMeshData++ = MeshData.MinLocalBounds.Y;
+					*GpuMeshData++ = MeshData.MinLocalBounds.Z;
+
+					*GpuMeshData++ = MeshData.MaxLocalBounds.X;
+					*GpuMeshData++ = MeshData.MaxLocalBounds.Y;
+					*GpuMeshData++ = MeshData.MaxLocalBounds.Z;
+				}
+				RHIUnlockBuffer(TypedProxy->GPUMeshData.Buffer);
 			}
-		);
-	}
+		}
+	);
 }
 
-void UNiagaraDataInterfaceMeshRendererInfo::GetNumMeshes(FVectorVMExternalFunctionContext& Context)
+void UNiagaraDataInterfaceMeshRendererInfo::VMGetNumMeshes(FVectorVMExternalFunctionContext& Context)
 {
 	FNDIOutputParam<int32> OutNum(Context);
 
-	int32 NumMeshes = Info.IsValid() ? Info->GetMeshData().Num() : 0;
+	const int32 NumMeshes = CachedMeshData.Num();
 	for (int32 Instance = 0; Instance < Context.GetNumInstances(); ++Instance)
 	{
 		OutNum.SetAndAdvance(NumMeshes);
 	}
 }
 
-void UNiagaraDataInterfaceMeshRendererInfo::GetMeshLocalBounds(FVectorVMExternalFunctionContext& Context)
+void UNiagaraDataInterfaceMeshRendererInfo::VMGetMeshLocalBounds(FVectorVMExternalFunctionContext& Context)
 {
 	FNDIInputParam<int32> InMeshIdx(Context);
 	FNDIOutputParam<FVector3f> OutMinBounds(Context);
@@ -588,16 +392,12 @@ void UNiagaraDataInterfaceMeshRendererInfo::GetMeshLocalBounds(FVectorVMExternal
 	{
 		FVector3f MinLocalBounds(ForceInitToZero);
 		FVector3f MaxLocalBounds(ForceInitToZero);
-		if (Info.IsValid())
+		if (CachedMeshData.Num() > 0)
 		{
-			const FNDIMeshRendererInfo::FMeshDataArray& MeshData = Info->GetMeshData();
-			if (MeshData.Num() > 0)
-			{
-				const int32 MeshIdx = FMath::Clamp(InMeshIdx.GetAndAdvance(), 0, MeshData.Num() - 1);			
-				const FNDIMeshRendererInfo::FMeshData& Mesh = MeshData[MeshIdx];
-				MinLocalBounds = Mesh.MinLocalBounds;
-				MaxLocalBounds = Mesh.MaxLocalBounds;
-			}
+			const int32 MeshIdx = FMath::Clamp(InMeshIdx.GetAndAdvance(), 0, CachedMeshData.Num() - 1);
+			const FMeshData& MeshData = CachedMeshData[MeshIdx];
+			MinLocalBounds = MeshData.MinLocalBounds;
+			MaxLocalBounds = MeshData.MaxLocalBounds;
 		}
 		OutMinBounds.SetAndAdvance(MinLocalBounds);
 		OutMaxBounds.SetAndAdvance(MaxLocalBounds);
@@ -605,4 +405,18 @@ void UNiagaraDataInterfaceMeshRendererInfo::GetMeshLocalBounds(FVectorVMExternal
 	}
 }
 
-#undef LOCTEXT_NAMESPACE
\ No newline at end of file
+void UNiagaraDataInterfaceMeshRendererInfo::VMGetSubUVDetails(FVectorVMExternalFunctionContext& Context)
+{
+	FNDIOutputParam<bool> OutBlendEnabled(Context);
+	FNDIOutputParam<FVector2f> OutSubImageSize(Context);
+
+	const UNiagaraMeshRendererProperties* Properties = MeshRenderer ? MeshRenderer : GetDefault<UNiagaraMeshRendererProperties>();
+	const FVector2f SubImageSize(Properties->SubImageSize);
+	for (int32 i = 0; i < Context.GetNumInstances(); ++i)
+	{
+		OutBlendEnabled.SetAndAdvance(Properties->bSubImageBlend != 0);
+		OutSubImageSize.SetAndAdvance(SubImageSize);
+	}
+}
+
+#undef LOCTEXT_NAMESPACE
diff --git a/Engine/Plugins/FX/Niagara/Source/Niagara/Private/NiagaraDataInterfaceSpriteRendererInfo.cpp b/Engine/Plugins/FX/Niagara/Source/Niagara/Private/NiagaraDataInterfaceSpriteRendererInfo.cpp
--- a/Engine/Plugins/FX/Niagara/Source/Niagara/Private/NiagaraDataInterfaceSpriteRendererInfo.cpp
+++ b/Engine/Plugins/FX/Niagara/Source/Niagara/Private/NiagaraDataInterfaceSpriteRendererInfo.cpp
@@ -288,6 +288,7 @@ bool UNiagaraDataInterfaceSpriteRendererInfo::CopyToInternal(UNiagaraDataInterfa
 	}
 
 	UNiagaraDataInterfaceSpriteRendererInfo* DestinationTyped = CastChecked<UNiagaraDataInterfaceSpriteRendererInfo>(Destination);
+	DestinationTyped->OnSpriteRendererChanged(nullptr);
 	DestinationTyped->SpriteRenderer = SpriteRenderer;
 	DestinationTyped->OnSpriteRendererChanged(SpriteRenderer);
 
