Sometimes using static variables, changing the values doesn't update the stack for downstream static switches hiding parameters. This is due to stale rapid iteration variables somehow being around that is preferentially looked at when creating the stack UI. Most paths properly remove this data when swapping to a different thing, like linked values, but somehow there are paths where it isn't removed. In this case, we provide UI fix issue paths to resolve the problem.

diff --git a/Engine/Plugins/FX/Niagara/Source/NiagaraEditor/Private/ViewModels/Stack/NiagaraStackFunctionInput.cpp b/Engine/Plugins/FX/Niagara/Source/NiagaraEditor/Private/ViewModels/Stack/NiagaraStackFunctionInput.cpp
--- a/Engine/Plugins/FX/Niagara/Source/NiagaraEditor/Private/ViewModels/Stack/NiagaraStackFunctionInput.cpp
+++ b/Engine/Plugins/FX/Niagara/Source/NiagaraEditor/Private/ViewModels/Stack/NiagaraStackFunctionInput.cpp
@@ -565,6 +565,34 @@ void UNiagaraStackFunctionInput::RefreshChildrenInternal(const TArray<UNiagaraSt
 	RefreshFromMetaData(NewIssues);
 	RefreshValues();
 
+	// If we keep around rapid iteration parameters that should go away, it bloats the parameter store.
+	if (RapidIterationParameter.IsValid() && InputValues.Mode != EValueMode::Local)
+	{
+		bool bHasStaleRapidIterationVar = false;
+		for (TWeakObjectPtr<UNiagaraScript> Script : AffectedScripts)
+		{
+			if (Script->RapidIterationParameters.IndexOf(RapidIterationParameter) != INDEX_NONE)
+			{
+				bHasStaleRapidIterationVar = true;
+			}
+		}
+
+		if (bHasStaleRapidIterationVar)
+		{
+			NewIssues.Add(FStackIssue(
+				EStackIssueSeverity::Warning,
+				LOCTEXT("InvalidRapidIterationOverrideShort", "Invalid Rapid Iteration override"),
+				LOCTEXT("InvalidRapidIterationOverrideLong", "There is a rapid iteration value left around that could make the UI not match the actual state. Hit 'Fix issue' to clear out the old, bad value."),
+				GetStackEditorDataKey(),
+				false,
+				{ FStackIssueFix(
+					LOCTEXT("ResetRapidIterationInputFix", "Fix stale value"),
+					FStackIssueFixDelegate::CreateLambda([this]() { this->RemoveRapidIterationParametersForAffectedScripts(true); }))
+				}
+			));
+		}
+	}
+
 	if (InputValues.DynamicNode.IsValid())
 	{
 		FNiagaraStackGraphUtilities::CheckForDeprecatedScriptVersion(GetDynamicInputNode(), GetStackEditorDataKey(), GetUpgradeDynamicInputVersionFix(), NewIssues);
@@ -1825,7 +1853,7 @@ bool UNiagaraStackFunctionInput::UpdateRapidIterationParametersForAffectedScript
 	return true;
 }
 
-bool UNiagaraStackFunctionInput::RemoveRapidIterationParametersForAffectedScripts()
+bool UNiagaraStackFunctionInput::RemoveRapidIterationParametersForAffectedScripts(bool bUpdateGraphGuidsForAffected)
 {
 	for (TWeakObjectPtr<UNiagaraScript> Script : AffectedScripts)
 	{
@@ -1836,6 +1864,12 @@ bool UNiagaraStackFunctionInput::RemoveRapidIterationParametersForAffectedScript
 	{
 		if (Script->RapidIterationParameters.RemoveParameter(RapidIterationParameter))
 		{
+			if (bUpdateGraphGuidsForAffected)
+			{
+				// Because these scripts are not versioned usually, we pass in the 0-0-0-0 id.
+				Script->MarkScriptAndSourceDesynchronized(TEXT("Invalidated GUIDS at request of RemoveRapidIterationParametersForAffectedScripts"), FGuid());
+			}
+
 			UE_LOG(LogNiagaraEditor, Log, TEXT("Removed Var '%s' from Script %s"), *RapidIterationParameter.GetName().ToString(), *Script->GetFullName());
 		}
 	}
diff --git a/Engine/Plugins/FX/Niagara/Source/NiagaraEditor/Public/ViewModels/Stack/NiagaraStackFunctionInput.h b/Engine/Plugins/FX/Niagara/Source/NiagaraEditor/Public/ViewModels/Stack/NiagaraStackFunctionInput.h
--- a/Engine/Plugins/FX/Niagara/Source/NiagaraEditor/Public/ViewModels/Stack/NiagaraStackFunctionInput.h
+++ b/Engine/Plugins/FX/Niagara/Source/NiagaraEditor/Public/ViewModels/Stack/NiagaraStackFunctionInput.h
@@ -288,7 +288,7 @@ protected:
 	virtual void RefreshChildrenInternal(const TArray<UNiagaraStackEntry*>& CurrentChildren, TArray<UNiagaraStackEntry*>& NewChildren, TArray<FStackIssue>& NewIssues) override;
 
 	bool UpdateRapidIterationParametersForAffectedScripts(const uint8* Data);
-	bool RemoveRapidIterationParametersForAffectedScripts();
+	bool RemoveRapidIterationParametersForAffectedScripts(bool bUpdateGraphGuidsForAffected = false);
 	FString ResolveDisplayNameArgument(const FString& InArg) const;
 	FStackIssueFixDelegate GetUpgradeDynamicInputVersionFix();
 
