- Add 'NiagaraDebugHud' console command to quickly enable / disable the HUD

diff --git a/Engine/Plugins/FX/Niagara/Source/Niagara/Private/NiagaraDebugHud.cpp b/Engine/Plugins/FX/Niagara/Source/Niagara/Private/NiagaraDebugHud.cpp
--- a/Engine/Plugins/FX/Niagara/Source/Niagara/Private/NiagaraDebugHud.cpp
+++ b/Engine/Plugins/FX/Niagara/Source/Niagara/Private/NiagaraDebugHud.cpp
@@ -135,6 +135,30 @@ namespace NiagaraDebugLocal
 		MakeTuple(TEXT("ShowParticlesVariablesWithSystem="), TEXT("When enabled particle variables are shown with the system display"), [](FString Arg) {Settings.bShowParticlesVariablesWithSystem = FCString::Atoi(*Arg) != 0; }),
 	};
 
+	static FAutoConsoleCommandWithWorldAndArgs CmdNiagaraDebugHud(
+		TEXT("NiagaraDebugHud"),
+		TEXT("Shorter version to quickly toggle debug hud modes\n")
+		TEXT(" No value will toggle the overview on / off\n")
+		TEXT(" A numberic value selects which overmode to set, where 0 is off\n"),
+		FConsoleCommandWithWorldAndArgsDelegate::CreateLambda(
+			[](const TArray<FString>& Args, UWorld*)
+			{
+				if ( Args.Num() == 0 )
+				{
+					Settings.bOverviewEnabled = !Settings.bOverviewEnabled;
+					Settings.OverviewMode = ENiagaraDebugHUDOverviewMode::Overview;
+				}
+				else
+				{
+					const int32 Mode = FCString::Atoi(*Args[0]);
+
+					Settings.bOverviewEnabled = Mode > 0;
+					Settings.OverviewMode = FMath::Clamp(ENiagaraDebugHUDOverviewMode(Mode - 1), ENiagaraDebugHUDOverviewMode::Overview, ENiagaraDebugHUDOverviewMode::GpuComputePerformance);
+				}
+			}
+		)
+	);
+
 	static FAutoConsoleCommandWithWorldAndArgs CmdDebugHud(
 		TEXT("fx.Niagara.Debug.Hud"),
 		TEXT("Set options for debug hud display"),
