- Fix for missing _{ParameterName}

#rnx

diff --git a/Engine/Plugins/FX/Niagara/Shaders/Private/NiagaraDataInterfaceMeshRendererInfoTemplate.ush b/Engine/Plugins/FX/Niagara/Shaders/Private/NiagaraDataInterfaceMeshRendererInfoTemplate.ush
--- a/Engine/Plugins/FX/Niagara/Shaders/Private/NiagaraDataInterfaceMeshRendererInfoTemplate.ush
+++ b/Engine/Plugins/FX/Niagara/Shaders/Private/NiagaraDataInterfaceMeshRendererInfoTemplate.ush
@@ -5,12 +5,12 @@ float2			{ParameterName}_SubImageSize;
 uint			{ParameterName}_NumMeshes;
 Buffer<float>	{ParameterName}_MeshDataBuffer;
 
-void GetNumMeshes(out int OutNumMeshes)
+void GetNumMeshes_{ParameterName}(out int OutNumMeshes)
 {
 	OutNumMeshes = int({ParameterName}_NumMeshes);
 }
 
-void GetMeshLocalBounds(int MeshIndex, out float3 OutMinBounds, out float3 OutMaxBounds, out float3 OutSize)
+void GetMeshLocalBounds_{ParameterName}(int MeshIndex, out float3 OutMinBounds, out float3 OutMaxBounds, out float3 OutSize)
 {
 	OutMinBounds = (float3)0;
 	OutMaxBounds = (float3)0;
diff --git a/Engine/Plugins/FX/Niagara/Source/Niagara/Private/NiagaraDataInterfaceMeshRendererInfo.cpp b/Engine/Plugins/FX/Niagara/Source/Niagara/Private/NiagaraDataInterfaceMeshRendererInfo.cpp
--- a/Engine/Plugins/FX/Niagara/Source/Niagara/Private/NiagaraDataInterfaceMeshRendererInfo.cpp
+++ b/Engine/Plugins/FX/Niagara/Source/Niagara/Private/NiagaraDataInterfaceMeshRendererInfo.cpp
@@ -20,8 +20,7 @@ namespace NDIMeshRendererInfoInternal
 		SHADER_PARAMETER_SRV(Buffer<float>,	MeshDataBuffer)
 	END_SHADER_PARAMETER_STRUCT()
 
-
-	static const TCHAR* TemplateShaderFile = TEXT("/Plugin/FX/Niagara/Private/NiagaraDataInterfaceSpriteRendererInfoTemplate.ush");
+	static const TCHAR* TemplateShaderFile = TEXT("/Plugin/FX/Niagara/Private/NiagaraDataInterfaceMeshRendererInfoTemplate.ush");
 
 	static const FName GetNumMeshesName("GetNumMeshes");
 	static const FName GetMeshLocalBoundsName("GetMeshLocalBounds");
