Temp fix for a crash in the UI

#rnx

diff --git a/Engine/Plugins/FX/Niagara/Source/NiagaraEditor/Private/ViewModels/Stack/NiagaraStackEntry.cpp b/Engine/Plugins/FX/Niagara/Source/NiagaraEditor/Private/ViewModels/Stack/NiagaraStackEntry.cpp
--- a/Engine/Plugins/FX/Niagara/Source/NiagaraEditor/Private/ViewModels/Stack/NiagaraStackEntry.cpp
+++ b/Engine/Plugins/FX/Niagara/Source/NiagaraEditor/Private/ViewModels/Stack/NiagaraStackEntry.cpp
@@ -554,7 +554,8 @@ void UNiagaraStackEntry::SetIsSearchResult(bool bInIsSearchResult)
 
 bool UNiagaraStackEntry::HasBaseEmitter() const
 {
-	if (GetSystemViewModel()->GetIsForDataProcessingOnly())
+	// TODO (me) the is valid check is a temp fix as it can happen that it gets deleted for some reason which leads to a crash here
+	if (!SystemViewModel.IsValid() || GetSystemViewModel()->GetIsForDataProcessingOnly())
 	{
 		// If the model is just for data processing we don't want to go through the whole merge procedure for all the stack entries and just treat all entries as non-inherited.
 		return false;
