Work around issue where Niagara GPU scripts for the editor platform will have compile requests (because they are dirty on load) that are never properly handled, including being put into DDC.

diff --git a/Engine/Plugins/FX/Niagara/Source/Niagara/Private/NiagaraScript.cpp b/Engine/Plugins/FX/Niagara/Source/Niagara/Private/NiagaraScript.cpp
--- a/Engine/Plugins/FX/Niagara/Source/Niagara/Private/NiagaraScript.cpp
+++ b/Engine/Plugins/FX/Niagara/Source/Niagara/Private/NiagaraScript.cpp
@@ -3177,6 +3177,14 @@ void UNiagaraScript::CacheResourceShadersForCooking(EShaderPlatform ShaderPlatfo
 
 void UNiagaraScript::CacheShadersForResources(FNiagaraShaderScript* ResourceToCache, bool bApplyCompletedShaderMapForRendering, bool bForceRecompile, bool bCooking, const ITargetPlatform* TargetPlatform)
 {
+	// if we're running the cook commandlet and this script isn't being processed as part of the cooking process, then let's not worry
+	// about kicking off a compile for it.  As things currently stand we don't process the results (including storing them in the DDC)
+	// so we'll just be kicking off the compile each time for no reason.
+	if (IsRunningCookCommandlet() && !bCooking)
+	{
+		return;
+	}
+
 	if (CanBeRunOnGpu())
 	{
 		// When not running in the editor, the shaders are created in-sync (in the postload) to avoid update issues.
