- Fix for deprecated method

#rnx

diff --git a/Engine/Plugins/FX/Niagara/Source/NiagaraEditor/Private/Widgets/SNiagaraBakerViewport.cpp b/Engine/Plugins/FX/Niagara/Source/NiagaraEditor/Private/Widgets/SNiagaraBakerViewport.cpp
--- a/Engine/Plugins/FX/Niagara/Source/NiagaraEditor/Private/Widgets/SNiagaraBakerViewport.cpp
+++ b/Engine/Plugins/FX/Niagara/Source/NiagaraEditor/Private/Widgets/SNiagaraBakerViewport.cpp
@@ -97,7 +97,7 @@ public:
 		return true;
 	}
 
-	virtual bool InputAxis(FViewport* InViewport, int32 ControllerId, FKey Key, float Delta, float InDeltaTime, int32 NumSamples, bool bGamepad) override
+	virtual bool InputAxis(FViewport* InViewport, FInputDeviceId InDeviceID, FKey Key, float Delta, float InDeltaTime, int32 InNumSamples, bool InbGamepad) override
 	{
 		// Viewport movement only enabled when preview is enabled otherwise there is no feedback
 		auto ViewModel = WeakViewModel.Pin();
