- Fix for WorldVelocity calculation
- Change velocity sampler to use Point, interpolated results are not valid

#rnx

diff --git a/Engine/Plugins/FX/Niagara/Shaders/Private/NiagaraDataInterfaceGBufferTemplate.ush b/Engine/Plugins/FX/Niagara/Shaders/Private/NiagaraDataInterfaceGBufferTemplate.ush
--- a/Engine/Plugins/FX/Niagara/Shaders/Private/NiagaraDataInterfaceGBufferTemplate.ush
+++ b/Engine/Plugins/FX/Niagara/Shaders/Private/NiagaraDataInterfaceGBufferTemplate.ush
@@ -85,28 +85,26 @@ void DecodeWorldVelocity_{ParameterName}(in float2 ScreenUV, in bool bApplyViewp
 	IsValid = true;
 	if ( EncodedVelocity.x > 0 )
 	{
+		float3 Velocity = DecodeVelocityFromTexture(EncodedVelocity);
+
+		float SceneDeviceZ = LookupDeviceZ(ScreenUV);
+		float SceneDepth = ConvertFromDeviceZ(SceneDeviceZ);
+ 		float2 ScreenPosition = (ScreenUV - View.ScreenPositionScaleBias.wz) / View.ScreenPositionScaleBias.xy;
+		float3 TranslatedWorldPosition = mul(float4(ScreenPosition * SceneDepth, SceneDepth, 1), PrimaryView.ScreenToTranslatedWorld).xyz;
+ 
+		float PrevSceneDepth = 0;
 	#if VELOCITY_ENCODE_DEPTH
-		OutValue = DecodeVelocityFromTexture(EncodedVelocity);
-		OutValue = mul(float4(OutValue, 0), View.ViewToTranslatedWorld).xyz;
+		PrevSceneDepth = ConvertFromDeviceZ(SceneDeviceZ - Velocity.z);
 	#else
-		float2 ScreenVelocity = DecodeVelocityFromTexture(EncodedVelocity).xy;
-		ScreenVelocity.y = -ScreenVelocity.y;
-
-		{
-			float2 PrevScreenUV = ScreenUV - ScreenVelocity;
-
-			float SceneDepth = CalcSceneDepth(ScreenUV);
-			float2 ScreenPosition = (ScreenUV - View.ScreenPositionScaleBias.wz) / View.ScreenPositionScaleBias.xy;
-			float3 WorldPosition = mul(float4(ScreenPosition * SceneDepth, SceneDepth, 1), LWCHackToFloat(PrimaryView.ScreenToWorld)).xyz;
-
-			// Ideally we would use CalcPrevSceneDepth(PrevScreenUV) here to get an accurate location, however we do not store PrevSceneDepth so this is a very rough approximation
-			float PrevSceneDepth = SceneDepth;	//CalcPrevSceneDepth(PrevScreenUV);
-			float2 PrevScreenPosition = (PrevScreenUV - View.ScreenPositionScaleBias.wz) / View.ScreenPositionScaleBias.xy;
-			float3 PrevWorldPosition = mul(float4(PrevScreenPosition * PrevSceneDepth, PrevSceneDepth, 1), LWCHackToFloat(PrimaryView.ScreenToWorld)).xyz;
-
-			OutValue = WorldPosition - PrevWorldPosition;
-		}
+ 		// Ideally we would use CalcPrevSceneDepth(PrevScreenUV) here to get an accurate location, however we do not store PrevSceneDepth so this is a very rough approximation
+		PrevSceneDepth = SceneDepth; // CalcPrevSceneDepth(PrevScreenUV);
 	#endif
+		float2 PrevScreenPosition = ScreenPosition - Velocity.xy;
+		float3 PrevTranslatedWorldPosition = mul(float4(PrevScreenPosition * PrevSceneDepth, PrevSceneDepth, 1), PrimaryView.PrevScreenToTranslatedWorld).xyz;
+
+		float3 PreViewTranslationOffset = LWCToFloat(LWCSubtract(PrimaryView.PreViewTranslation, PrimaryView.PrevPreViewTranslation));
+ 
+		OutValue = TranslatedWorldPosition - PrevTranslatedWorldPosition - PreViewTranslationOffset;
 	}
 #endif
 }
diff --git a/Engine/Plugins/FX/Niagara/Source/Niagara/Private/DataInterface/NiagaraDataInterfaceGBuffer.cpp b/Engine/Plugins/FX/Niagara/Source/Niagara/Private/DataInterface/NiagaraDataInterfaceGBuffer.cpp
--- a/Engine/Plugins/FX/Niagara/Source/Niagara/Private/DataInterface/NiagaraDataInterfaceGBuffer.cpp
+++ b/Engine/Plugins/FX/Niagara/Source/Niagara/Private/DataInterface/NiagaraDataInterfaceGBuffer.cpp
@@ -189,7 +189,7 @@ void UNiagaraDataInterfaceGBuffer::SetShaderParameters(const FNiagaraDataInterfa
 	}
 
 	Parameters->VelocityTexture = VelocityTexture;
-	Parameters->VelocityTextureSampler = TStaticSamplerState<SF_Bilinear, AM_Clamp, AM_Clamp, AM_Clamp>::GetRHI();
+	Parameters->VelocityTextureSampler = TStaticSamplerState<SF_Point, AM_Clamp, AM_Clamp, AM_Clamp>::GetRHI();
 }
 
 #if WITH_EDITORONLY_DATA
