- Ensure cached data is updated when OnChanged notification is handled

#rnx

diff --git a/Engine/Plugins/FX/Niagara/Source/Niagara/Classes/NiagaraDataInterfaceMeshRendererInfo.h b/Engine/Plugins/FX/Niagara/Source/Niagara/Classes/NiagaraDataInterfaceMeshRendererInfo.h
--- a/Engine/Plugins/FX/Niagara/Source/Niagara/Classes/NiagaraDataInterfaceMeshRendererInfo.h
+++ b/Engine/Plugins/FX/Niagara/Source/Niagara/Classes/NiagaraDataInterfaceMeshRendererInfo.h
@@ -61,6 +61,7 @@ public:
 protected:
 	virtual bool CopyToInternal(UNiagaraDataInterface* Destination) const override;
 	virtual void PushToRenderThreadImpl() override;
+	void UpdateCachedData();
 
 	void VMGetNumMeshes(FVectorVMExternalFunctionContext& Context);
 	void VMGetMeshLocalBounds(FVectorVMExternalFunctionContext& Context);
diff --git a/Engine/Plugins/FX/Niagara/Source/Niagara/Private/NiagaraDataInterfaceMeshRendererInfo.cpp b/Engine/Plugins/FX/Niagara/Source/Niagara/Private/NiagaraDataInterfaceMeshRendererInfo.cpp
--- a/Engine/Plugins/FX/Niagara/Source/Niagara/Private/NiagaraDataInterfaceMeshRendererInfo.cpp
+++ b/Engine/Plugins/FX/Niagara/Source/Niagara/Private/NiagaraDataInterfaceMeshRendererInfo.cpp
@@ -121,32 +121,13 @@ void UNiagaraDataInterfaceMeshRendererInfo::OnMeshRendererChanged(UNiagaraMeshRe
 		MeshRenderer->OnChanged().Remove(OnMeshRendererChangedHandle);
 		OnMeshRendererChangedHandle.Reset();
 	}
-	CachedMeshData.Empty();
 
 	if (NewMeshRenderer)
 	{
-		OnMeshRendererChangedHandle = NewMeshRenderer->OnChanged().AddLambda([this]() { MarkRenderDataDirty(); });
-
-		CachedMeshData.AddDefaulted(NewMeshRenderer->Meshes.Num());
-		for ( int32 i=0; i < NewMeshRenderer->Meshes.Num(); ++i )
-		{
-			FNiagaraMeshRendererMeshProperties& MeshProperties = NewMeshRenderer->Meshes[i];
-			if ( MeshProperties.Mesh == nullptr )
-			{
-				continue;
-			}
-
-			const FBox LocalBounds = MeshProperties.Mesh->GetExtendedBounds().GetBox();
-			if ( LocalBounds.IsValid )
-			{
-				FMeshData& CachedMesh = CachedMeshData[i];
-				CachedMesh.MinLocalBounds = FVector3f(LocalBounds.Min * MeshProperties.Scale);
-				CachedMesh.MaxLocalBounds = FVector3f(LocalBounds.Max * MeshProperties.Scale);
-			}
-		}
+		OnMeshRendererChangedHandle = NewMeshRenderer->OnChanged().AddLambda([this]() { UpdateCachedData(); });
 	}
 #endif
-	MarkRenderDataDirty();
+	UpdateCachedData();
 }
 
 void UNiagaraDataInterfaceMeshRendererInfo::GetFunctions(TArray<FNiagaraFunctionSignature>& OutFunctions)
@@ -369,6 +350,32 @@ void UNiagaraDataInterfaceMeshRendererInfo::PushToRenderThreadImpl()
 	);
 }
 
+void UNiagaraDataInterfaceMeshRendererInfo::UpdateCachedData()
+{
+	CachedMeshData.Empty();
+	if ( MeshRenderer != nullptr )
+	{
+		CachedMeshData.AddDefaulted(MeshRenderer->Meshes.Num());
+		for (int32 i = 0; i < MeshRenderer->Meshes.Num(); ++i)
+		{
+			FNiagaraMeshRendererMeshProperties& MeshProperties = MeshRenderer->Meshes[i];
+			if (MeshProperties.Mesh == nullptr)
+			{
+				continue;
+			}
+
+			const FBox LocalBounds = MeshProperties.Mesh->GetExtendedBounds().GetBox();
+			if (LocalBounds.IsValid)
+			{
+				FMeshData& CachedMesh = CachedMeshData[i];
+				CachedMesh.MinLocalBounds = FVector3f(LocalBounds.Min * MeshProperties.Scale);
+				CachedMesh.MaxLocalBounds = FVector3f(LocalBounds.Max * MeshProperties.Scale);
+			}
+		}
+	}
+	MarkRenderDataDirty();
+}
+
 void UNiagaraDataInterfaceMeshRendererInfo::VMGetNumMeshes(FVectorVMExternalFunctionContext& Context)
 {
 	FNDIOutputParam<int32> OutNum(Context);
