- Add markers for GPU ribbons

diff --git a/Engine/Plugins/FX/Niagara/Source/Niagara/Private/NiagaraRendererRibbons.cpp b/Engine/Plugins/FX/Niagara/Source/Niagara/Private/NiagaraRendererRibbons.cpp
--- a/Engine/Plugins/FX/Niagara/Source/Niagara/Private/NiagaraRendererRibbons.cpp
+++ b/Engine/Plugins/FX/Niagara/Source/Niagara/Private/NiagaraRendererRibbons.cpp
@@ -2015,7 +2015,8 @@ void FNiagaraRendererRibbons::InitializeViewIndexBuffersGPU(FRHICommandListImmed
 	{
 		return;
 	}
-	
+
+	SCOPED_DRAW_EVENT(CMDList, NiagaraRenderRibbonsGenIndiciesGPU);
 	{
 		FNiagaraRibbonCreateIndexBufferParamsCS::FPermutationDomain PermutationVector;
 		PermutationVector.Set<FRibbonWantsAutomaticTessellation>(GenerationConfig.WantsAutomaticTessellation());
@@ -2298,6 +2299,7 @@ void FNiagaraRendererRibbons::InitializeVertexBuffersGPU(FRHICommandListImmediat
 	}
 
 	{
+		SCOPED_DRAW_EVENT(CMDList, NiagaraRenderRibbonsGenVerticesSortGPU);
 		SCOPE_CYCLE_COUNTER(STAT_NiagaraRenderRibbonsGenVerticesSortGPU);
 		
 		FNiagaraRibbonSortPhase1CS::FPermutationDomain PermutationVector;
@@ -2335,6 +2337,7 @@ void FNiagaraRendererRibbons::InitializeVertexBuffersGPU(FRHICommandListImmediat
 		}
 
 		{			
+			SCOPED_DRAW_EVENT(CMDList, NiagaraRenderRibbonsGenVerticesInitialSortGPU);
 			SCOPE_CYCLE_COUNTER(STAT_NiagaraRenderRibbonsGenVerticesInitialSortGPU);
 			
 			// Initial sort, sets up the buffer, and runs a parallel bubble sort to create groups of BubbleSortGroupWidth size
@@ -2348,6 +2351,7 @@ void FNiagaraRendererRibbons::InitializeVertexBuffersGPU(FRHICommandListImmediat
 		}
 		
 		{
+			SCOPED_DRAW_EVENT(CMDList, NiagaraRenderRibbonsGenVerticesFinalSortGPU);
 			SCOPE_CYCLE_COUNTER(STAT_NiagaraRenderRibbonsGenVerticesFinalSortGPU);
 			
 			// Repeatedly runs a scatter based merge sort until we have the final buffer
@@ -2372,6 +2376,7 @@ void FNiagaraRendererRibbons::InitializeVertexBuffersGPU(FRHICommandListImmediat
 	}
 	
 	{
+		SCOPED_DRAW_EVENT(CMDList, NiagaraRenderRibbonsGenVerticesReductionPhase1GPU);
 		SCOPE_CYCLE_COUNTER(STAT_NiagaraRenderRibbonsGenVerticesReductionPhase1GPU);
 		
 		FNiagaraRibbonVertexReductionInitializationCS::FPermutationDomain PermutationVector;
@@ -2483,6 +2488,7 @@ void FNiagaraRendererRibbons::InitializeVertexBuffersGPU(FRHICommandListImmediat
 	static constexpr int32 CommandBufferOffset = 0;
 		
 	{
+		SCOPED_DRAW_EVENT(CMDList, NiagaraRenderRibbonsGenVerticesReductionPhase2GPU);
 		SCOPE_CYCLE_COUNTER(STAT_NiagaraRenderRibbonsGenVerticesReductionPhase2GPU);
 		
 		FNiagaraRibbonVertexReductionFinalizationParameters FinalizationParams;
@@ -2531,6 +2537,7 @@ void FNiagaraRendererRibbons::InitializeVertexBuffersGPU(FRHICommandListImmediat
 	}
 		
 	{
+		SCOPED_DRAW_EVENT(CMDList, NiagaraRenderRibbonsGenVerticesMultiRibbonInitGPU);
 		SCOPE_CYCLE_COUNTER(STAT_NiagaraRenderRibbonsGenVerticesMultiRibbonInitGPU);
 		
 		FNiagaraRibbonVertexFinalizationParameters FinalizeParams;
