- Add settings menu to baker
- Move none camera related settings to here
- Expose FPS option for component ticking

#rnx

diff --git a/Engine/Plugins/FX/Niagara/Source/NiagaraEditor/Private/NiagaraEditorStyle.cpp b/Engine/Plugins/FX/Niagara/Source/NiagaraEditor/Private/NiagaraEditorStyle.cpp
--- a/Engine/Plugins/FX/Niagara/Source/NiagaraEditor/Private/NiagaraEditorStyle.cpp
+++ b/Engine/Plugins/FX/Niagara/Source/NiagaraEditor/Private/NiagaraEditorStyle.cpp
@@ -569,6 +569,7 @@ void FNiagaraEditorStyle::InitDebuggerStyle()
 void FNiagaraEditorStyle::InitBakerStyle()
 {
 	Set("NiagaraEditor.BakerIcon", new IMAGE_BRUSH("Icons/Baker/BakerIcon", Icon40x40));
+	Set("NiagaraEditor.BakerSettings", new IMAGE_BRUSH("Icons/Baker/BakerSettings", Icon40x40));
 }
 
 void FNiagaraEditorStyle::InitCommonColors()
diff --git a/Engine/Plugins/FX/Niagara/Source/NiagaraEditor/Private/ViewModels/NiagaraBakerViewModel.cpp b/Engine/Plugins/FX/Niagara/Source/NiagaraEditor/Private/ViewModels/NiagaraBakerViewModel.cpp
--- a/Engine/Plugins/FX/Niagara/Source/NiagaraEditor/Private/ViewModels/NiagaraBakerViewModel.cpp
+++ b/Engine/Plugins/FX/Niagara/Source/NiagaraEditor/Private/ViewModels/NiagaraBakerViewModel.cpp
@@ -348,6 +348,32 @@ void FNiagaraBakerViewModel::ResetCurrentCamera()
 	}
 }
 
+bool FNiagaraBakerViewModel::IsSimTickRate(int TickRate) const
+{
+	UNiagaraBakerSettings* BakerSettings = GetBakerSettings();
+	return BakerSettings ? BakerSettings->FramesPerSecond == TickRate : false;
+}
+
+int FNiagaraBakerViewModel::GetSimTickRate() const
+{
+	UNiagaraBakerSettings* BakerSettings = GetBakerSettings();
+	return BakerSettings ? BakerSettings->FramesPerSecond : 60;
+}
+
+void FNiagaraBakerViewModel::SetSimTickRate(int TickRate)
+{
+	if (UNiagaraBakerSettings* BakerSettings = GetBakerSettings())
+	{
+		TickRate = FMath::Max(TickRate, 1);
+		if (BakerSettings->FramesPerSecond != TickRate)
+		{
+			const FScopedTransaction Transaction(LOCTEXT("SetSimTickRate", "SetSimTickRate"));
+			BakerSettings->Modify();
+			BakerSettings->FramesPerSecond = TickRate;
+		}
+	}
+}
+
 void FNiagaraBakerViewModel::AddOutput(UClass* Class)
 {
 	if (UNiagaraBakerSettings* BakerSettings = GetBakerSettings())
diff --git a/Engine/Plugins/FX/Niagara/Source/NiagaraEditor/Private/ViewModels/NiagaraBakerViewModel.h b/Engine/Plugins/FX/Niagara/Source/NiagaraEditor/Private/ViewModels/NiagaraBakerViewModel.h
--- a/Engine/Plugins/FX/Niagara/Source/NiagaraEditor/Private/ViewModels/NiagaraBakerViewModel.h
+++ b/Engine/Plugins/FX/Niagara/Source/NiagaraEditor/Private/ViewModels/NiagaraBakerViewModel.h
@@ -103,6 +103,10 @@ public:
 
 	void ResetCurrentCamera();
 
+	bool IsSimTickRate(int TickRate) const;
+	int GetSimTickRate() const;
+	void SetSimTickRate(int TickRate);
+
 	void AddOutput(UClass* Class);
 	void RemoveCurrentOutput();
 	bool CanRemoveCurrentOutput() const;
diff --git a/Engine/Plugins/FX/Niagara/Source/NiagaraEditor/Private/Widgets/SNiagaraBakerWidget.cpp b/Engine/Plugins/FX/Niagara/Source/NiagaraEditor/Private/Widgets/SNiagaraBakerWidget.cpp
--- a/Engine/Plugins/FX/Niagara/Source/NiagaraEditor/Private/Widgets/SNiagaraBakerWidget.cpp
+++ b/Engine/Plugins/FX/Niagara/Source/NiagaraEditor/Private/Widgets/SNiagaraBakerWidget.cpp
@@ -231,6 +231,14 @@ void SNiagaraBakerWidget::Construct(const FArguments& InArgs)
 		FSlateIcon(FNiagaraEditorStyle::Get().GetStyleSetName(), "NiagaraEditor.BakerIcon")
 	);
 
+	BakerToolbarBuilder.AddComboButton(
+		FUIAction(),
+		FOnGetContent::CreateSP(this, &SNiagaraBakerWidget::MakeSettingsMenu),
+		FText::GetEmpty(),
+		LOCTEXT("SettingsToolTip", "Modify baker settings"),
+		FSlateIcon(FNiagaraEditorStyle::Get().GetStyleSetName(), "NiagaraEditor.BakerSettings")
+	);
+
 	BakerToolbarBuilder.AddComboButton(
 		FUIAction(),
 		FOnGetContent::CreateSP(this, &SNiagaraBakerWidget::MakeCameraModeMenu),
@@ -558,7 +566,7 @@ TSharedRef<SWidget> SNiagaraBakerWidget::MakeCameraModeMenu()
 	return MenuBuilder.MakeWidget();
 }
 
-TSharedRef<SWidget> SNiagaraBakerWidget::MakeViewOptionsMenu()
+TSharedRef<SWidget> SNiagaraBakerWidget::MakeSettingsMenu()
 {
 	using namespace NiagaraBakerWidgetLocal;
 
@@ -567,147 +575,155 @@ TSharedRef<SWidget> SNiagaraBakerWidget::MakeViewOptionsMenu()
 
 	FMenuBuilder MenuBuilder(true, nullptr);
 
-	MenuBuilder.BeginSection(NAME_None, LOCTEXT("PreviewSettings", "Preview Settings"));
-	{
-		MenuBuilder.AddMenuEntry(
-			LOCTEXT("ShowRealtimePreview", "Show Realtime Preview"),
-			LOCTEXT("ShowRealtimePreviewTooltip", "When enabled shows a live preview of what will be rendered, this may not be accurate with all visualization modes."),
-			FSlateIcon(),
-			FUIAction(
-				FExecuteAction::CreateSP(ViewModel, &FNiagaraBakerViewModel::ToggleRealtimePreview),
-				FCanExecuteAction(),
-				FIsActionChecked::CreateSP(ViewModel, &FNiagaraBakerViewModel::ShowRealtimePreview)
-			),
-			NAME_None,
-			EUserInterfaceActionType::ToggleButton
-		);
+	MenuBuilder.AddMenuEntry(
+		LOCTEXT("ShowRealtimePreview", "Show Realtime Preview"),
+		LOCTEXT("ShowRealtimePreviewTooltip", "When enabled shows a live preview of what will be rendered, this may not be accurate with all visualization modes."),
+		FSlateIcon(),
+		FUIAction(
+			FExecuteAction::CreateSP(ViewModel, &FNiagaraBakerViewModel::ToggleRealtimePreview),
+			FCanExecuteAction(),
+			FIsActionChecked::CreateSP(ViewModel, &FNiagaraBakerViewModel::ShowRealtimePreview)
+		),
+		NAME_None,
+		EUserInterfaceActionType::ToggleButton
+	);
 
-		MenuBuilder.AddMenuEntry(
-			LOCTEXT("ShowBakedView", "Show Baked View"),
-			LOCTEXT("ShowBakedViewTooltip", "When enabled shows the baked texture."),
-			FSlateIcon(),
-			FUIAction(
-				FExecuteAction::CreateSP(ViewModel, &FNiagaraBakerViewModel::ToggleBakedView),
-				FCanExecuteAction(),
-				FIsActionChecked::CreateSP(ViewModel, &FNiagaraBakerViewModel::ShowBakedView)
-			),
-			NAME_None,
-			EUserInterfaceActionType::ToggleButton
-		);
+	MenuBuilder.AddMenuEntry(
+		LOCTEXT("ShowBakedView", "Show Baked View"),
+		LOCTEXT("ShowBakedViewTooltip", "When enabled shows the baked texture."),
+		FSlateIcon(),
+		FUIAction(
+			FExecuteAction::CreateSP(ViewModel, &FNiagaraBakerViewModel::ToggleBakedView),
+			FCanExecuteAction(),
+			FIsActionChecked::CreateSP(ViewModel, &FNiagaraBakerViewModel::ShowBakedView)
+		),
+		NAME_None,
+		EUserInterfaceActionType::ToggleButton
+	);
 
-		MenuBuilder.AddMenuEntry(
-			LOCTEXT("ShowCheckerboard", "Show Checkerboard"),
-			LOCTEXT("ShowCheckerboardTooltip", "Show a checkerboard rather than a solid color to easily visualize alpha blending."),
-			FSlateIcon(),
-			FUIAction(
-				FExecuteAction::CreateSP(ViewModel, &FNiagaraBakerViewModel::ToggleCheckerboardEnabled),
-				FCanExecuteAction(),
-				FIsActionChecked::CreateSP(ViewModel, &FNiagaraBakerViewModel::IsCheckerboardEnabled)
-			),
-			NAME_None,
-			EUserInterfaceActionType::ToggleButton
-		);
+	MenuBuilder.AddMenuEntry(
+		LOCTEXT("ShowCheckerboard", "Show Checkerboard"),
+		LOCTEXT("ShowCheckerboardTooltip", "Show a checkerboard rather than a solid color to easily visualize alpha blending."),
+		FSlateIcon(),
+		FUIAction(
+			FExecuteAction::CreateSP(ViewModel, &FNiagaraBakerViewModel::ToggleCheckerboardEnabled),
+			FCanExecuteAction(),
+			FIsActionChecked::CreateSP(ViewModel, &FNiagaraBakerViewModel::IsCheckerboardEnabled)
+		),
+		NAME_None,
+		EUserInterfaceActionType::ToggleButton
+	);
 
-		MenuBuilder.AddMenuEntry(
-			LOCTEXT("ShowInfoText", "Show Info Text"),
-			LOCTEXT("ShowInfoTextTooltip", "Shows information about the preview and baked outputs."),
-			FSlateIcon(),
-			FUIAction(
-				FExecuteAction::CreateSP(ViewModel, &FNiagaraBakerViewModel::ToggleInfoText),
-				FCanExecuteAction(),
-				FIsActionChecked::CreateSP(ViewModel, &FNiagaraBakerViewModel::ShowInfoText)
-			),
-			NAME_None,
-			EUserInterfaceActionType::ToggleButton
-		);
+	MenuBuilder.AddMenuEntry(
+		LOCTEXT("ShowInfoText", "Show Info Text"),
+		LOCTEXT("ShowInfoTextTooltip", "Shows information about the preview and baked outputs."),
+		FSlateIcon(),
+		FUIAction(
+			FExecuteAction::CreateSP(ViewModel, &FNiagaraBakerViewModel::ToggleInfoText),
+			FCanExecuteAction(),
+			FIsActionChecked::CreateSP(ViewModel, &FNiagaraBakerViewModel::ShowInfoText)
+		),
+		NAME_None,
+		EUserInterfaceActionType::ToggleButton
+	);
 
-		MenuBuilder.AddMenuEntry(
-			LOCTEXT("ShowNiagaraOnly", "Render Niagara Only"),
-			LOCTEXT("ShowNiagaraOnlyTooltip", "Renders only the Niagara System when enabled."),
-			FSlateIcon(),
-			FUIAction(
-				FExecuteAction::CreateSP(ViewModel, &FNiagaraBakerViewModel::ToggleRenderComponentOnly),
-				FCanExecuteAction(),
-				FIsActionChecked::CreateSP(ViewModel, &FNiagaraBakerViewModel::ShowRenderComponentOnly)
-			),
-			NAME_None,
-			EUserInterfaceActionType::ToggleButton
-		);
+	MenuBuilder.AddMenuEntry(
+		LOCTEXT("ShowNiagaraOnly", "Render Niagara Only"),
+		LOCTEXT("ShowNiagaraOnlyTooltip", "Renders only the Niagara System when enabled."),
+		FSlateIcon(),
+		FUIAction(
+			FExecuteAction::CreateSP(ViewModel, &FNiagaraBakerViewModel::ToggleRenderComponentOnly),
+			FCanExecuteAction(),
+			FIsActionChecked::CreateSP(ViewModel, &FNiagaraBakerViewModel::ShowRenderComponentOnly)
+		),
+		NAME_None,
+		EUserInterfaceActionType::ToggleButton
+	);
 
-		//DetailCategory.HideProperty(GET_MEMBER_NAME_CHECKED(UNiagaraBakerSettings, bPreviewLooping));
-		//DetailCategory.HideProperty(GET_MEMBER_NAME_CHECKED(UNiagaraBakerSettings, bRenderComponentOnly));
-	}
+	MenuBuilder.AddSubMenu(
+		LOCTEXT("SimulationTickRate", "Simulation Tick Rate"),
+		LOCTEXT("SimulationTickRateToolTip", "The rate at which the simulation will tick, i.e. 120fps, 60fps, 30fps, etc."),
+		FNewMenuDelegate::CreateSP(this, &SNiagaraBakerWidget::MakeSimTickRateMenu)
+	);
+
+	return MenuBuilder.MakeWidget();
+}
+
+TSharedRef<SWidget> SNiagaraBakerWidget::MakeViewOptionsMenu()
+{
+	using namespace NiagaraBakerWidgetLocal;
+
+	FNiagaraBakerViewModel* ViewModel = WeakViewModel.Pin().Get();
+	check(ViewModel != nullptr);
+
+	FMenuBuilder MenuBuilder(true, nullptr);
 
-	MenuBuilder.BeginSection(NAME_None, LOCTEXT("Camera", "Camera"));
+	MenuBuilder.AddWidget(
+		FMakeVectorBox::Construct(
+			FMakeVectorBox::FGetter::CreateSP(ViewModel, &FNiagaraBakerViewModel::GetCurrentCameraLocation),
+			FMakeVectorBox::FSetter::CreateSP(ViewModel, &FNiagaraBakerViewModel::SetCurrentCameraLocation)
+		),
+		LOCTEXT("CameraLocation", "Camera Location")
+	);
+
+	MenuBuilder.AddWidget(
+		FMakeRotatorBox::Construct(
+			FMakeRotatorBox::FGetter::CreateSP(ViewModel, &FNiagaraBakerViewModel::GetCurrentCameraRotation),
+			FMakeRotatorBox::FSetter::CreateSP(ViewModel, &FNiagaraBakerViewModel::SetCurrentCameraRotation)
+		),
+		LOCTEXT("CameraRotation", "Camera Rotation")
+	);
+
+	if ( ViewModel->IsCurrentCameraPerspective() )
 	{
 		MenuBuilder.AddWidget(
-			FMakeVectorBox::Construct(
-				FMakeVectorBox::FGetter::CreateSP(ViewModel, &FNiagaraBakerViewModel::GetCurrentCameraLocation),
-				FMakeVectorBox::FSetter::CreateSP(ViewModel, &FNiagaraBakerViewModel::SetCurrentCameraLocation)
-			),
-			LOCTEXT("CameraLocation", "Camera Location")
+			MakeSpinBox<float>(1.0f, 170.0f, 1.0f, 170.0f, TAttribute<float>::CreateSP(ViewModel, &FNiagaraBakerViewModel::GetCameraFOV), SSpinBox<float>::FOnValueChanged::CreateSP(ViewModel, &FNiagaraBakerViewModel::SetCameraFOV)),
+			LOCTEXT("FOVAngle", "Field of View (H)")
 		);
 
 		MenuBuilder.AddWidget(
-			FMakeRotatorBox::Construct(
-				FMakeRotatorBox::FGetter::CreateSP(ViewModel, &FNiagaraBakerViewModel::GetCurrentCameraRotation),
-				FMakeRotatorBox::FSetter::CreateSP(ViewModel, &FNiagaraBakerViewModel::SetCurrentCameraRotation)
-			),
-			LOCTEXT("CameraRotation", "Camera Rotation")
+			MakeSpinBox<float>(0.0f, TOptional<float>(), 0.0f, 1000.0f, TAttribute<float>::CreateSP(ViewModel, &FNiagaraBakerViewModel::GetCameraOrbitDistance), SSpinBox<float>::FOnValueChanged::CreateSP(ViewModel, &FNiagaraBakerViewModel::SetCameraOrbitDistance)),
+			LOCTEXT("OrbitDistance", "Orbit Distance")
 		);
-
-		if ( ViewModel->IsCurrentCameraPerspective() )
-		{
-			MenuBuilder.AddWidget(
-				MakeSpinBox<float>(1.0f, 170.0f, 1.0f, 170.0f, TAttribute<float>::CreateSP(ViewModel, &FNiagaraBakerViewModel::GetCameraFOV), SSpinBox<float>::FOnValueChanged::CreateSP(ViewModel, &FNiagaraBakerViewModel::SetCameraFOV)),
-				LOCTEXT("FOVAngle", "Field of View (H)")
-			);
-
-			MenuBuilder.AddWidget(
-				MakeSpinBox<float>(0.0f, TOptional<float>(), 0.0f, 1000.0f, TAttribute<float>::CreateSP(ViewModel, &FNiagaraBakerViewModel::GetCameraOrbitDistance), SSpinBox<float>::FOnValueChanged::CreateSP(ViewModel, &FNiagaraBakerViewModel::SetCameraOrbitDistance)),
-				LOCTEXT("OrbitDistance", "Orbit Distance")
-			);
-		}
-		else
-		{
-			MenuBuilder.AddWidget(
-				MakeSpinBox<float>(0.1f, TOptional<float>(), 0.1f, 1000.0f, TAttribute<float>::CreateSP(ViewModel, &FNiagaraBakerViewModel::GetCameraOrthoWidth), SSpinBox<float>::FOnValueChanged::CreateSP(ViewModel, &FNiagaraBakerViewModel::SetCameraOrthoWidth)),
-				LOCTEXT("CameraOrthoWidth", "Camera Orthographic Width")
-			);
-		}
-
-		//-TODO: Clean this up
-		MenuBuilder.AddMenuEntry(
-			LOCTEXT("CustomAspectRatio", "Custom Aspect Ratio"),
-			LOCTEXT("CustomAspectRatioTooltip", "Use a custom aspect ratio to render with."),
-			FSlateIcon(),
-			FUIAction(
-				FExecuteAction::CreateSP(ViewModel, &FNiagaraBakerViewModel::ToggleCameraAspectRatioEnabled),
-				FCanExecuteAction(),
-				FIsActionChecked::CreateSP(ViewModel, &FNiagaraBakerViewModel::IsCameraAspectRatioEnabled)
-			),
-			NAME_None,
-			EUserInterfaceActionType::ToggleButton
+	}
+	else
+	{
+		MenuBuilder.AddWidget(
+			MakeSpinBox<float>(0.1f, TOptional<float>(), 0.1f, 1000.0f, TAttribute<float>::CreateSP(ViewModel, &FNiagaraBakerViewModel::GetCameraOrthoWidth), SSpinBox<float>::FOnValueChanged::CreateSP(ViewModel, &FNiagaraBakerViewModel::SetCameraOrthoWidth)),
+			LOCTEXT("CameraOrthoWidth", "Camera Orthographic Width")
 		);
+	}
 
-		if ( ViewModel->IsCameraAspectRatioEnabled() )
-		{
-			MenuBuilder.AddWidget(
-				MakeSpinBox<float>(0.1f, 5.0f, 0.1f, 5.0f, TAttribute<float>::CreateSP(ViewModel, &FNiagaraBakerViewModel::GetCameraAspectRatio), SSpinBox<float>::FOnValueChanged::CreateSP(ViewModel, &FNiagaraBakerViewModel::SetCameraAspectRatio)),
-				LOCTEXT("CameraAspectRatio", "Camera Aspect Ratio")
-			);
-		}
+	//-TODO: Clean this up
+	MenuBuilder.AddMenuEntry(
+		LOCTEXT("CustomAspectRatio", "Custom Aspect Ratio"),
+		LOCTEXT("CustomAspectRatioTooltip", "Use a custom aspect ratio to render with."),
+		FSlateIcon(),
+		FUIAction(
+			FExecuteAction::CreateSP(ViewModel, &FNiagaraBakerViewModel::ToggleCameraAspectRatioEnabled),
+			FCanExecuteAction(),
+			FIsActionChecked::CreateSP(ViewModel, &FNiagaraBakerViewModel::IsCameraAspectRatioEnabled)
+		),
+		NAME_None,
+		EUserInterfaceActionType::ToggleButton
+	);
 
-		MenuBuilder.AddMenuEntry(
-			LOCTEXT("ResetCamera", "Reset Camera"),
-			LOCTEXT("ResetCameraTooltip", "Resets the current camera back to the default settings."),
-			FSlateIcon(),
-			FUIAction(
-				FExecuteAction::CreateSP(ViewModel, &FNiagaraBakerViewModel::ResetCurrentCamera)
-			)
+	if ( ViewModel->IsCameraAspectRatioEnabled() )
+	{
+		MenuBuilder.AddWidget(
+			MakeSpinBox<float>(0.1f, 5.0f, 0.1f, 5.0f, TAttribute<float>::CreateSP(ViewModel, &FNiagaraBakerViewModel::GetCameraAspectRatio), SSpinBox<float>::FOnValueChanged::CreateSP(ViewModel, &FNiagaraBakerViewModel::SetCameraAspectRatio)),
+			LOCTEXT("CameraAspectRatio", "Camera Aspect Ratio")
 		);
 	}
-	MenuBuilder.EndSection();
+
+	MenuBuilder.AddMenuEntry(
+		LOCTEXT("ResetCamera", "Reset Camera"),
+		LOCTEXT("ResetCameraTooltip", "Resets the current camera back to the default settings."),
+		FSlateIcon(),
+		FUIAction(
+			FExecuteAction::CreateSP(ViewModel, &FNiagaraBakerViewModel::ResetCurrentCamera)
+		)
+	);
 
 	return MenuBuilder.MakeWidget();
 }
@@ -761,6 +777,42 @@ TSharedRef<SWidget> SNiagaraBakerWidget::MakeAddOutputMenu()
 	return MenuBuilder.MakeWidget();
 }
 
+void SNiagaraBakerWidget::MakeSimTickRateMenu(FMenuBuilder& MenuBuilder) const
+{
+	using namespace NiagaraBakerWidgetLocal;
+
+	FNiagaraBakerViewModel* ViewModel = WeakViewModel.Pin().Get();
+	check(ViewModel != nullptr);
+
+	MenuBuilder.AddWidget(
+		MakeSpinBox<int>(1, 480, 1, 480, TAttribute<int>::CreateSP(ViewModel, &FNiagaraBakerViewModel::GetSimTickRate), SSpinBox<int>::FOnValueChanged::CreateSP(ViewModel, &FNiagaraBakerViewModel::SetSimTickRate)),
+		LOCTEXT("SimTickRate", "Simulation Tick Rate")
+	);
+
+	MenuBuilder.BeginSection("FPS", LOCTEXT("FPS", "FPS"));
+	{
+		static const int32 DefaultFPS[] = { 240, 120, 60, 50, 30, 20, 15 };
+
+		for (int32 FPS : DefaultFPS)
+		{
+			FText FPSText = FText::Format(LOCTEXT("FPSFormat", "{0} fps"), FPS);
+			MenuBuilder.AddMenuEntry(
+				FPSText,
+				FText::GetEmpty(),
+				FSlateIcon(),
+				FUIAction(
+					FExecuteAction::CreateSP(ViewModel, &FNiagaraBakerViewModel::SetSimTickRate, FPS),
+					FCanExecuteAction(),
+					FIsActionChecked::CreateSP(ViewModel, &FNiagaraBakerViewModel::IsSimTickRate, FPS)
+				),
+				NAME_None,
+				EUserInterfaceActionType::ToggleButton
+			);
+		}
+	}
+	MenuBuilder.EndSection();
+}
+
 void SNiagaraBakerWidget::FindWarnings()
 {
 	FoundWarnings.Reset();
diff --git a/Engine/Plugins/FX/Niagara/Source/NiagaraEditor/Private/Widgets/SNiagaraBakerWidget.h b/Engine/Plugins/FX/Niagara/Source/NiagaraEditor/Private/Widgets/SNiagaraBakerWidget.h
--- a/Engine/Plugins/FX/Niagara/Source/NiagaraEditor/Private/Widgets/SNiagaraBakerWidget.h
+++ b/Engine/Plugins/FX/Niagara/Source/NiagaraEditor/Private/Widgets/SNiagaraBakerWidget.h
@@ -9,6 +9,7 @@ class FNiagaraBakerViewModel;
 class SNiagaraBakerViewport;
 class SNiagaraBakerTimelineWidget;
 class IDetailsView;
+class FMenuBuilder;
 
 class SNiagaraBakerWidget : public SCompoundWidget
 {
@@ -27,10 +28,13 @@ public:
 	TSharedRef<SWidget> MakeSettingsWidget();
 
 	TSharedRef<SWidget> MakeCameraModeMenu();
+	TSharedRef<SWidget> MakeSettingsMenu();
 	TSharedRef<SWidget> MakeViewOptionsMenu();
 	TSharedRef<SWidget> MakeOutputSelectMenu();
 	TSharedRef<SWidget> MakeAddOutputMenu();
 
+	void MakeSimTickRateMenu(FMenuBuilder& MenuBuilder) const;
+
 	void FindWarnings();
 	bool HasWarnings() const { return FoundWarnings.Num() > 0; }
 	TSharedRef<SWidget> MakeWarningsMenu();
