- Correctly handle WaveOps by adding a permutation as some platforms wave op support is runtime optional

diff --git a/Engine/Plugins/FX/Niagara/Shaders/Private/NiagaraSortKeyGen.usf b/Engine/Plugins/FX/Niagara/Shaders/Private/NiagaraSortKeyGen.usf
--- a/Engine/Plugins/FX/Niagara/Shaders/Private/NiagaraSortKeyGen.usf
+++ b/Engine/Plugins/FX/Niagara/Shaders/Private/NiagaraSortKeyGen.usf
@@ -16,7 +16,6 @@
 
 #include "/Engine/Private/Common.ush"
 
-#define NIAGARA_KEY_GEN_MAX_CULL_PLANES 10 // must match the code
 Buffer<float> NiagaraParticleDataFloat;
 Buffer<half> NiagaraParticleDataHalf;
 Buffer<int> NiagaraParticleDataInt;
diff --git a/Engine/Plugins/FX/Niagara/Source/Niagara/Private/NiagaraGpuComputeDispatch.cpp b/Engine/Plugins/FX/Niagara/Source/Niagara/Private/NiagaraGpuComputeDispatch.cpp
--- a/Engine/Plugins/FX/Niagara/Source/Niagara/Private/NiagaraGpuComputeDispatch.cpp
+++ b/Engine/Plugins/FX/Niagara/Source/Niagara/Private/NiagaraGpuComputeDispatch.cpp
@@ -1882,13 +1882,17 @@ void FNiagaraGpuComputeDispatch::GenerateSortKeys(FRHICommandListImmediate& RHIC
 	const bool bHighPrecision = EnumHasAnyFlags(Flags, EGPUSortFlags::HighPrecisionKeys);
 	const FGPUSortManager::FKeyGenInfo KeyGenInfo((uint32)NumElementsInBatch, bHighPrecision);
 
+	const bool bUseWaveOps = FNiagaraSortKeyGenCS::UseWaveOps(ShaderPlatform);
+
 	FNiagaraSortKeyGenCS::FPermutationDomain SortPermutationVector;
 	SortPermutationVector.Set<FNiagaraSortKeyGenCS::FSortUsingMaxPrecision>(bHighPrecision);
 	SortPermutationVector.Set<FNiagaraSortKeyGenCS::FEnableCulling>(false);
+	SortPermutationVector.Set<FNiagaraSortKeyGenCS::FUseWaveOps>(bUseWaveOps);
 
 	FNiagaraSortKeyGenCS::FPermutationDomain SortAndCullPermutationVector;
 	SortAndCullPermutationVector.Set<FNiagaraSortKeyGenCS::FSortUsingMaxPrecision>(bHighPrecision);
 	SortAndCullPermutationVector.Set<FNiagaraSortKeyGenCS::FEnableCulling>(true);
+	SortAndCullPermutationVector.Set<FNiagaraSortKeyGenCS::FUseWaveOps>(bUseWaveOps);
 
 	TShaderMapRef<FNiagaraSortKeyGenCS> SortKeyGenCS(GetGlobalShaderMap(FeatureLevel), SortPermutationVector);
 	TShaderMapRef<FNiagaraSortKeyGenCS> SortAndCullKeyGenCS(GetGlobalShaderMap(FeatureLevel), SortAndCullPermutationVector);
@@ -1964,12 +1968,12 @@ void FNiagaraGpuComputeDispatch::GenerateSortKeys(FRHICommandListImmediate& RHIC
 			}
 
 			// Choose the shader to bind
-			TShaderMapRef<FNiagaraSortKeyGenCS> KeyGenCS = SortInfo.bEnableCulling ? SortAndCullKeyGenCS : SortKeyGenCS;
-			SetComputePipelineState(RHICmdList, KeyGenCS.GetComputeShader());
-
-			SetShaderParameters(RHICmdList, KeyGenCS, KeyGenCS.GetComputeShader(), Params);
-			DispatchComputeShader(RHICmdList, KeyGenCS, FMath::DivideAndRoundUp(SortInfo.ParticleCount, NIAGARA_KEY_GEN_THREAD_COUNT), 1, 1);
-			UnsetShaderUAVs(RHICmdList, KeyGenCS, KeyGenCS.GetComputeShader());
+			FComputeShaderUtils::Dispatch(
+				RHICmdList,
+				SortInfo.bEnableCulling ? SortAndCullKeyGenCS : SortKeyGenCS,
+				Params,
+				FIntVector(FMath::DivideAndRoundUp(SortInfo.ParticleCount, NIAGARA_KEY_GEN_THREAD_COUNT), 1, 1)
+			);
 		}
 	}
 
diff --git a/Engine/Plugins/FX/Niagara/Source/NiagaraVertexFactories/Private/NiagaraSortingGPU.cpp b/Engine/Plugins/FX/Niagara/Source/NiagaraVertexFactories/Private/NiagaraSortingGPU.cpp
--- a/Engine/Plugins/FX/Niagara/Source/NiagaraVertexFactories/Private/NiagaraSortingGPU.cpp
+++ b/Engine/Plugins/FX/Niagara/Source/NiagaraVertexFactories/Private/NiagaraSortingGPU.cpp
@@ -27,8 +27,34 @@ static FAutoConsoleVariableRef CVarNiagaraGPUCullingCPUToGPUThreshold(
 
 IMPLEMENT_GLOBAL_SHADER(FNiagaraSortKeyGenCS, "/Plugin/FX/Niagara/Private/NiagaraSortKeyGen.usf", "GenerateParticleSortKeys", SF_Compute);
 
+bool FNiagaraSortKeyGenCS::ShouldCompilePermutation(const FGlobalShaderPermutationParameters& Parameters)
+{
+	FPermutationDomain PermutationVector(Parameters.PermutationId);
+
+	ERHIFeatureSupport WaveOpsSupport = FDataDrivenShaderPlatformInfo::GetSupportsWaveOperations(Parameters.Platform);
+	if (PermutationVector.Get<FUseWaveOps>())
+	{
+		if ( WaveOpsSupport == ERHIFeatureSupport::Unsupported )
+		{
+			return false;
+		}
+	}
+	else
+	{
+		if (WaveOpsSupport == ERHIFeatureSupport::RuntimeGuaranteed)
+		{
+			return false;
+		}
+	}
+
+	return RHISupportsComputeShaders(Parameters.Platform);
+}
+
 void FNiagaraSortKeyGenCS::ModifyCompilationEnvironment(const FGlobalShaderPermutationParameters& Parameters, FShaderCompilerEnvironment& OutEnvironment)
 {
+	FPermutationDomain PermutationVector(Parameters.PermutationId);
+	const bool bUseWaveIntrinsics = PermutationVector.Get<FUseWaveOps>();
+
 	FGlobalShader::ModifyCompilationEnvironment(Parameters, OutEnvironment);
 	OutEnvironment.SetDefine(TEXT("THREAD_COUNT"), NIAGARA_KEY_GEN_THREAD_COUNT);
 	OutEnvironment.SetDefine(TEXT("SORT_NONE"), (uint8)ENiagaraSortMode::None);
@@ -36,8 +62,7 @@ void FNiagaraSortKeyGenCS::ModifyCompilationEnvironment(const FGlobalShaderPermu
 	OutEnvironment.SetDefine(TEXT("SORT_VIEW_DISTANCE"), (uint8)ENiagaraSortMode::ViewDistance);
 	OutEnvironment.SetDefine(TEXT("SORT_CUSTOM_ASCENDING"), (uint8)ENiagaraSortMode::CustomAscending);
 	OutEnvironment.SetDefine(TEXT("SORT_CUSTOM_DESCENDING"), (uint8)ENiagaraSortMode::CustomDecending);
-
-	bool bUseWaveIntrinsics = FDataDrivenShaderPlatformInfo::GetSupportsWaveOperations(Parameters.Platform) != ERHIFeatureSupport::Unsupported;
+	OutEnvironment.SetDefine(TEXT("NIAGARA_KEY_GEN_MAX_CULL_PLANES"), NIAGARA_KEY_GEN_MAX_CULL_PLANES);
 	OutEnvironment.SetDefine(TEXT("USE_WAVE_INTRINSICS"), bUseWaveIntrinsics);
 	if (bUseWaveIntrinsics)
 	{
@@ -45,3 +70,8 @@ void FNiagaraSortKeyGenCS::ModifyCompilationEnvironment(const FGlobalShaderPermu
 	}
 }
 
+bool FNiagaraSortKeyGenCS::UseWaveOps(EShaderPlatform ShaderPlatform)
+{
+	const ERHIFeatureSupport UseWaveOpsSupport = FDataDrivenShaderPlatformInfo::GetSupportsWaveOperations(ShaderPlatform);
+	return (GRHISupportsWaveOperations && UseWaveOpsSupport == ERHIFeatureSupport::RuntimeDependent) || (UseWaveOpsSupport == ERHIFeatureSupport::RuntimeGuaranteed);
+}
diff --git a/Engine/Plugins/FX/Niagara/Source/NiagaraVertexFactories/Public/NiagaraSortingGPU.h b/Engine/Plugins/FX/Niagara/Source/NiagaraVertexFactories/Public/NiagaraSortingGPU.h
--- a/Engine/Plugins/FX/Niagara/Source/NiagaraVertexFactories/Public/NiagaraSortingGPU.h
+++ b/Engine/Plugins/FX/Niagara/Source/NiagaraVertexFactories/Public/NiagaraSortingGPU.h
@@ -32,7 +32,8 @@ class NIAGARAVERTEXFACTORIES_API FNiagaraSortKeyGenCS : public FGlobalShader
 public:
 	class FEnableCulling : SHADER_PERMUTATION_BOOL("ENABLE_CULLING");
 	class FSortUsingMaxPrecision : SHADER_PERMUTATION_BOOL("SORT_MAX_PRECISION");
-	using FPermutationDomain = TShaderPermutationDomain<FEnableCulling, FSortUsingMaxPrecision>;
+	class FUseWaveOps : SHADER_PERMUTATION_BOOL("DIM_USE_WAVE_OPS");
+	using FPermutationDomain = TShaderPermutationDomain<FEnableCulling, FSortUsingMaxPrecision, FUseWaveOps>;
 
 	BEGIN_SHADER_PARAMETER_STRUCT(FParameters, NIAGARAVERTEXFACTORIES_API)
 		SHADER_PARAMETER_SRV(Buffer, NiagaraParticleDataFloat)
@@ -74,10 +75,9 @@ public:
 		SHADER_PARAMETER_UAV(Buffer, OutCulledParticleCounts)
 	END_SHADER_PARAMETER_STRUCT()
 
-	static bool ShouldCompilePermutation(const FGlobalShaderPermutationParameters& Parameters)
-	{
-		return RHISupportsComputeShaders(Parameters.Platform);
-	}
+	static bool ShouldCompilePermutation(const FGlobalShaderPermutationParameters& Parameters);
 
 	static void ModifyCompilationEnvironment(const FGlobalShaderPermutationParameters& Parameters, FShaderCompilerEnvironment& OutEnvironment);
+
+	static bool UseWaveOps(EShaderPlatform ShaderPlatform);
 };
