Another Niagara UI crash temp fix

#rnx

diff --git a/Engine/Plugins/FX/Niagara/Source/NiagaraEditor/Private/ViewModels/Stack/NiagaraStackSimulationStageGroup.cpp b/Engine/Plugins/FX/Niagara/Source/NiagaraEditor/Private/ViewModels/Stack/NiagaraStackSimulationStageGroup.cpp
--- a/Engine/Plugins/FX/Niagara/Source/NiagaraEditor/Private/ViewModels/Stack/NiagaraStackSimulationStageGroup.cpp
+++ b/Engine/Plugins/FX/Niagara/Source/NiagaraEditor/Private/ViewModels/Stack/NiagaraStackSimulationStageGroup.cpp
@@ -422,7 +422,8 @@ void UNiagaraStackSimulationStageGroup::SimulationStagePropertiesChanged()
 
 bool UNiagaraStackSimulationStageGroup::HasBaseSimulationStage() const
 {
-	if (bHasBaseSimulationStageCache.IsSet() == false)
+	// todo (me) emitter view model validity check should not be required but fixes a crash in UI when the view model is invalid for some reason
+	if (bHasBaseSimulationStageCache.IsSet() == false && GetEmitterViewModel().IsValid())
 	{
 		FVersionedNiagaraEmitter BaseEmitter = GetEmitterViewModel()->GetEmitter().GetEmitterData()->GetParent();
 		bHasBaseSimulationStageCache = BaseEmitter.Emitter != nullptr && FNiagaraScriptMergeManager::Get()->HasBaseSimulationStage(BaseEmitter, GetScriptUsageId());
