Making sure a value is set before GetValue is called
Addendum to previous fix

#rnx

diff --git a/Engine/Plugins/FX/Niagara/Source/NiagaraEditor/Private/ViewModels/Stack/NiagaraStackSimulationStageGroup.cpp b/Engine/Plugins/FX/Niagara/Source/NiagaraEditor/Private/ViewModels/Stack/NiagaraStackSimulationStageGroup.cpp
--- a/Engine/Plugins/FX/Niagara/Source/NiagaraEditor/Private/ViewModels/Stack/NiagaraStackSimulationStageGroup.cpp
+++ b/Engine/Plugins/FX/Niagara/Source/NiagaraEditor/Private/ViewModels/Stack/NiagaraStackSimulationStageGroup.cpp
@@ -422,11 +422,18 @@ void UNiagaraStackSimulationStageGroup::SimulationStagePropertiesChanged()
 
 bool UNiagaraStackSimulationStageGroup::HasBaseSimulationStage() const
 {
-	// todo (me) emitter view model validity check should not be required but fixes a crash in UI when the view model is invalid for some reason
-	if (bHasBaseSimulationStageCache.IsSet() == false && GetEmitterViewModel().IsValid())
+	if (bHasBaseSimulationStageCache.IsSet() == false)
 	{
-		FVersionedNiagaraEmitter BaseEmitter = GetEmitterViewModel()->GetEmitter().GetEmitterData()->GetParent();
-		bHasBaseSimulationStageCache = BaseEmitter.Emitter != nullptr && FNiagaraScriptMergeManager::Get()->HasBaseSimulationStage(BaseEmitter, GetScriptUsageId());
+		// todo (me) emitter view model validity check should not be required but fixes a crash in UI when the view model is invalid for some reason
+		if(!GetEmitterViewModel().IsValid())
+		{
+			bHasBaseSimulationStageCache = false;
+		}
+		else
+		{
+			FVersionedNiagaraEmitter BaseEmitter = GetEmitterViewModel()->GetEmitter().GetEmitterData()->GetParent();
+			bHasBaseSimulationStageCache = BaseEmitter.Emitter != nullptr && FNiagaraScriptMergeManager::Get()->HasBaseSimulationStage(BaseEmitter, GetScriptUsageId());
+		}
 	}
 	return bHasBaseSimulationStageCache.GetValue();
 }
