Added optional ComponentTags to the SkelMeshDI to allow more control over selection of the SkelMeshComponent being selected

diff --git a/Engine/Plugins/FX/Niagara/Source/Niagara/Classes/NiagaraDataInterfaceSkeletalMesh.h b/Engine/Plugins/FX/Niagara/Source/Niagara/Classes/NiagaraDataInterfaceSkeletalMesh.h
--- a/Engine/Plugins/FX/Niagara/Source/Niagara/Classes/NiagaraDataInterfaceSkeletalMesh.h
+++ b/Engine/Plugins/FX/Niagara/Source/Niagara/Classes/NiagaraDataInterfaceSkeletalMesh.h
@@ -730,6 +730,10 @@ public:
 	UPROPERTY(EditAnywhere, Category = "Mesh", meta = (DisplayName = "Source Actor"))
 	TSoftObjectPtr<AActor> SoftSourceActor;
 
+	/** If defined, the supplied tags will be used to identify a valid component */
+	UPROPERTY(EditAnywhere, Category = "Mesh")
+	TArray<FName> ComponentTags;
+
 #if WITH_EDITORONLY_DATA
 	UPROPERTY()
 	TObjectPtr<AActor> Source_DEPRECATED;
diff --git a/Engine/Plugins/FX/Niagara/Source/Niagara/Private/NiagaraDataInterfaceSkeletalMesh.cpp b/Engine/Plugins/FX/Niagara/Source/Niagara/Private/NiagaraDataInterfaceSkeletalMesh.cpp
--- a/Engine/Plugins/FX/Niagara/Source/Niagara/Private/NiagaraDataInterfaceSkeletalMesh.cpp
+++ b/Engine/Plugins/FX/Niagara/Source/Niagara/Private/NiagaraDataInterfaceSkeletalMesh.cpp
@@ -1494,13 +1494,20 @@ void UNiagaraDataInterfaceSkeletalMesh::ProvidePerInstanceDataForRenderThread(vo
 
 USkeletalMesh* UNiagaraDataInterfaceSkeletalMesh::GetSkeletalMesh(FNiagaraSystemInstance* SystemInstance, USceneComponent* AttachComponent, TWeakObjectPtr<USceneComponent>& SceneComponent, USkeletalMeshComponent*& FoundSkelComp, FNDISkeletalMesh_InstanceData* InstData)
 {
+	auto IsValidComponent = [&](const USkeletalMeshComponent* Component)
+	{
+		return IsValid(Component) && IsValid(Component->SkeletalMesh) &&
+			(ComponentTags.IsEmpty()
+			|| ComponentTags.ContainsByPredicate([&](const FName& Tag) { return Tag == NAME_None || Component->ComponentHasTag(Tag); }));
+	};
+
 	// Helper to scour an actor (or its parents) for a valid skeletal mesh component
-	auto FindActorSkelMeshComponent = [](AActor* Actor, bool bRecurseParents = false) -> USkeletalMeshComponent*
+	auto FindActorSkelMeshComponent = [&](AActor* Actor, bool bRecurseParents = false) -> USkeletalMeshComponent*
 	{
 		if (ASkeletalMeshActor* SkelMeshActor = Cast<ASkeletalMeshActor>(Actor))
 		{
 			USkeletalMeshComponent* Comp = SkelMeshActor->GetSkeletalMeshComponent();
-			if (IsValid(Comp))
+			if (IsValidComponent(Comp))
 			{
 				return Comp;
 			}
@@ -1512,7 +1519,7 @@ USkeletalMesh* UNiagaraDataInterfaceSkeletalMesh::GetSkeletalMesh(FNiagaraSystem
 			for (UActorComponent* ActorComp : Actor->GetComponents())
 			{
 				USkeletalMeshComponent* Comp = Cast<USkeletalMeshComponent>(ActorComp);
-				if (IsValid(Comp) && Comp->GetSkeletalMesh() != nullptr)
+				if (IsValidComponent(Comp))
 				{
 					return Comp;
 				}
@@ -1581,7 +1588,7 @@ USkeletalMesh* UNiagaraDataInterfaceSkeletalMesh::GetSkeletalMesh(FNiagaraSystem
 		for (USceneComponent* Curr = AttachComponent; Curr; Curr = Curr->GetAttachParent())
 		{
 			USkeletalMeshComponent* ParentComp = Cast<USkeletalMeshComponent>(Curr);
-			if (IsValid(ParentComp))
+			if (IsValidComponent(ParentComp))
 			{
 				FoundSkelComp = ParentComp;
 				break;
@@ -1592,7 +1599,7 @@ USkeletalMesh* UNiagaraDataInterfaceSkeletalMesh::GetSkeletalMesh(FNiagaraSystem
 		{
 			// Next, try to find one in our outer chain
 			USkeletalMeshComponent* OuterComp = AttachComponent->GetTypedOuter<USkeletalMeshComponent>();
-			if (IsValid(OuterComp))
+			if (IsValidComponent(OuterComp))
 			{
 				FoundSkelComp = OuterComp;
 			}
