Fix component search in SkelMeshDI to leverage the scene component hierarchy rather than the potentially arbitrary ordering of the OwnedComponents of the actor.

diff --git a/Engine/Plugins/FX/Niagara/Source/Niagara/Private/NiagaraDataInterfaceSkeletalMesh.cpp b/Engine/Plugins/FX/Niagara/Source/Niagara/Private/NiagaraDataInterfaceSkeletalMesh.cpp
--- a/Engine/Plugins/FX/Niagara/Source/Niagara/Private/NiagaraDataInterfaceSkeletalMesh.cpp
+++ b/Engine/Plugins/FX/Niagara/Source/Niagara/Private/NiagaraDataInterfaceSkeletalMesh.cpp
@@ -1516,12 +1516,22 @@ USkeletalMesh* UNiagaraDataInterfaceSkeletalMesh::GetSkeletalMesh(FNiagaraSystem
 		// Fall back on any valid component on the actor
 		while (Actor)
 		{
-			for (UActorComponent* ActorComp : Actor->GetComponents())
+			USceneComponent* ActorRootComponent = Actor->GetRootComponent();
+			if (ensure(ActorRootComponent))
 			{
-				USkeletalMeshComponent* Comp = Cast<USkeletalMeshComponent>(ActorComp);
-				if (IsValidComponent(Comp))
+				TArray<USceneComponent*> ActorComponents;
+				ActorRootComponent->GetChildrenComponents(true /*bIncludeAllDescendants*/, ActorComponents);
+				ActorComponents.Insert(ActorRootComponent, 0);
+
+				for (USceneComponent* SceneComponent : ActorComponents)
 				{
-					return Comp;
+					if (USkeletalMeshComponent* SkeletalComp = Cast<USkeletalMeshComponent>(SceneComponent))
+					{
+						if (IsValidComponent(SkeletalComp))
+						{
+							return SkeletalComp;
+						}
+					}
 				}
 			}
 
