Fixing issues with new scratch panel
+ Fixing minor compile warning
+ Fixing ability to go to parent scratch modules
+ Fixing selection issues when applying
+ If you open a scratch pad and then take a merge from the parent, the opened scratch pad editor points to an invalid script.
+ Double-clicked to open a scratch pad after taking a merge from the parent causes issues.
+ Applying scratchpad changes switches the tab to overview
+ Parameters panel now has slightly different terminology
+ Parameters panel now auto selects right section when document change happens
+ Parameters panel now works with node right-click and selecting the active parameter in a get node
+ Parameters panel now remembers the state of categories opened from tab to tab

diff --git a/Engine/Plugins/FX/Niagara/Source/Niagara/Classes/NiagaraScratchPadContainer.h b/Engine/Plugins/FX/Niagara/Source/Niagara/Classes/NiagaraScratchPadContainer.h
--- a/Engine/Plugins/FX/Niagara/Source/Niagara/Classes/NiagaraScratchPadContainer.h
+++ b/Engine/Plugins/FX/Niagara/Source/Niagara/Classes/NiagaraScratchPadContainer.h
@@ -21,6 +21,11 @@ public:
 	void SetScripts(const TArray<TObjectPtr<UNiagaraScript>>& InScripts);
 	void AppendScripts(const TArray<TObjectPtr<UNiagaraScript>>& InScripts);
 	void AppendScripts(TObjectPtr<UNiagaraScratchPadContainer> InScripts);
+	
+	// The scripts get added in specific ordering through inheritance. These accessors just allow us to deal with checks 
+	// for known scripts 
+	int32 FindIndexForScript(UNiagaraScript* InScript) const;
+	UNiagaraScript* FindScriptAtIndex(int32 Index) const;
 
 #if WITH_EDITORONLY_DATA
 
diff --git a/Engine/Plugins/FX/Niagara/Source/Niagara/Private/NiagaraScratchPadContainer.cpp b/Engine/Plugins/FX/Niagara/Source/Niagara/Private/NiagaraScratchPadContainer.cpp
--- a/Engine/Plugins/FX/Niagara/Source/Niagara/Private/NiagaraScratchPadContainer.cpp
+++ b/Engine/Plugins/FX/Niagara/Source/Niagara/Private/NiagaraScratchPadContainer.cpp
@@ -57,3 +57,28 @@ void UNiagaraScratchPadContainer::AppendScripts(TObjectPtr<UNiagaraScratchPadCon
 	}
 #endif
 }
+int32 UNiagaraScratchPadContainer::FindIndexForScript(UNiagaraScript* InScript) const
+{
+#if WITH_EDITORONLY_DATA
+	for (int32 i = 0; i < Scripts.Num(); i++)
+	{
+		if (Scripts[i] == InScript)
+		{
+			return i;
+		}
+	}
+#endif
+	return INDEX_NONE;
+}
+
+UNiagaraScript* UNiagaraScratchPadContainer::FindScriptAtIndex(int32 Index) const
+{
+#if WITH_EDITORONLY_DATA
+	if (Scripts.IsValidIndex(Index))
+	{
+		return Scripts[Index];
+	}
+#endif
+	return nullptr;
+}
+
diff --git a/Engine/Plugins/FX/Niagara/Source/NiagaraEditor/Private/NiagaraEditorSettings.cpp b/Engine/Plugins/FX/Niagara/Source/NiagaraEditor/Private/NiagaraEditorSettings.cpp
--- a/Engine/Plugins/FX/Niagara/Source/NiagaraEditor/Private/NiagaraEditorSettings.cpp
+++ b/Engine/Plugins/FX/Niagara/Source/NiagaraEditor/Private/NiagaraEditorSettings.cpp
@@ -6,19 +6,20 @@
 #include "NiagaraConstants.h"
 #include "NiagaraEditorModule.h"
 
-const FGuid FNiagaraEditorGuids::SystemNamespaceMetaDataGuid = FGuid::NewGuid();
-const FGuid FNiagaraEditorGuids::EmitterNamespaceMetaDataGuid = FGuid::NewGuid();
-const FGuid FNiagaraEditorGuids::ParticleAttributeNamespaceMetaDataGuid = FGuid::NewGuid();
-const FGuid FNiagaraEditorGuids::ModuleNamespaceMetaDataGuid = FGuid::NewGuid();
-const FGuid FNiagaraEditorGuids::ModuleOutputNamespaceMetaDataGuid = FGuid::NewGuid();
-const FGuid FNiagaraEditorGuids::ModuleLocalNamespaceMetaDataGuid = FGuid::NewGuid();
-const FGuid FNiagaraEditorGuids::TransientNamespaceMetaDataGuid = FGuid::NewGuid();
-const FGuid FNiagaraEditorGuids::StackContextNamespaceMetaDataGuid = FGuid::NewGuid();
-const FGuid FNiagaraEditorGuids::EngineNamespaceMetaDataGuid = FGuid::NewGuid();
-const FGuid FNiagaraEditorGuids::UserNamespaceMetaDataGuid = FGuid::NewGuid();
-const FGuid FNiagaraEditorGuids::ParameterCollectionNamespaceMetaDataGuid = FGuid::NewGuid();
-const FGuid FNiagaraEditorGuids::DataInstanceNamespaceMetaDataGuid = FGuid::NewGuid();
-const FGuid FNiagaraEditorGuids::StaticSwitchNamespaceMetaDataGuid = FGuid::NewGuid();
+// These GUIDs are now statically defined so that they can be serialized out in settings for the parameter panel category expansion.
+const FGuid FNiagaraEditorGuids::SystemNamespaceMetaDataGuid = FGuid(TEXT("7B4AFB34D0DF46189A05349E361CE735"));
+const FGuid FNiagaraEditorGuids::EmitterNamespaceMetaDataGuid = FGuid(TEXT("1BA31433B3314F6BB258AECFBB466AC7"));
+const FGuid FNiagaraEditorGuids::ParticleAttributeNamespaceMetaDataGuid = FGuid(TEXT("530A57A84EF444F482CCC0A4B6D8D364"));
+const FGuid FNiagaraEditorGuids::ModuleNamespaceMetaDataGuid = FGuid(TEXT("FCA14ABEFAD14BB58786362887FEED09"));
+const FGuid FNiagaraEditorGuids::ModuleOutputNamespaceMetaDataGuid = FGuid(TEXT("8945EAA4041E43F6824AC3B0AF2AEA20"));
+const FGuid FNiagaraEditorGuids::ModuleLocalNamespaceMetaDataGuid = FGuid(TEXT("1DAA70FAC2ED4512914BD788D5BC3C9D"));
+const FGuid FNiagaraEditorGuids::TransientNamespaceMetaDataGuid = FGuid(TEXT("39BA69CE1CB74A168469C1F49C4E37DE"));
+const FGuid FNiagaraEditorGuids::StackContextNamespaceMetaDataGuid = FGuid(TEXT("65275FF7723C4CF28A82EB033EB3FAA6"));
+const FGuid FNiagaraEditorGuids::EngineNamespaceMetaDataGuid = FGuid(TEXT("18363BDD94D549038AA825175F92DDF9"));
+const FGuid FNiagaraEditorGuids::UserNamespaceMetaDataGuid = FGuid(TEXT("A3AC42514BF24B84AD2C52D2276414DA"));
+const FGuid FNiagaraEditorGuids::ParameterCollectionNamespaceMetaDataGuid = FGuid(TEXT("A5843E1CF9924F16B52016676332DEDC"));
+const FGuid FNiagaraEditorGuids::DataInstanceNamespaceMetaDataGuid = FGuid(TEXT("17F5FC610A914BA58D7EA8FD39B9A1D0"));
+const FGuid FNiagaraEditorGuids::StaticSwitchNamespaceMetaDataGuid = FGuid(TEXT("6A44CDD2EC3D4495BDC7FE28CAE00604"));
 
 
 FNiagaraNamespaceMetadata::FNiagaraNamespaceMetadata()
@@ -128,6 +129,7 @@ void UNiagaraEditorSettings::SetupNamespaceMetadata()
 			.SetDescription(LOCTEXT("ModuleDescription", "A value which exposes a module input to the system and emitter editor."))
 			.SetBackgroundColor(FLinearColor(FColor(136, 66, 65)))
 			.SetSortId(40)
+			.AddOption(ENiagaraNamespaceMetadataOptions::HideInSystem)
 			.AddOption(ENiagaraNamespaceMetadataOptions::HideInDefinitions)
 			.SetGuid(FNiagaraEditorGuids::ModuleNamespaceMetaDataGuid),
 		FNiagaraNamespaceMetadata({FNiagaraConstants::OutputNamespace}, FNiagaraConstants::ModuleNamespace)
@@ -148,6 +150,7 @@ void UNiagaraEditorSettings::SetupNamespaceMetadata()
 			.SetBackgroundColor(FLinearColor(FColor(191, 176, 84)))
 			.SetForegroundStyle("NiagaraEditor.ParameterName.NamespaceTextDark")
 			.SetSortId(50)
+			.AddOption(ENiagaraNamespaceMetadataOptions::HideInSystem)
 			.AddOption(ENiagaraNamespaceMetadataOptions::HideInDefinitions)
 			.AddOption(ENiagaraNamespaceMetadataOptions::PreventEditingNamespaceModifier)
 			.SetGuid(FNiagaraEditorGuids::ModuleLocalNamespaceMetaDataGuid),
@@ -219,6 +222,7 @@ void UNiagaraEditorSettings::SetupNamespaceMetadata()
 			.SetBackgroundColor(FLinearColor(FColor(170, 170, 170)))
 			.SetForegroundStyle("NiagaraEditor.ParameterName.NamespaceTextDark")
 			.SetSortId(100)
+			.AddOption(ENiagaraNamespaceMetadataOptions::HideInSystem)
 			.AddOption(ENiagaraNamespaceMetadataOptions::AdvancedInScript)
 			.AddOption(ENiagaraNamespaceMetadataOptions::AdvancedInSystem)
 			.AddOption(ENiagaraNamespaceMetadataOptions::HideInDefinitions)
@@ -230,6 +234,7 @@ void UNiagaraEditorSettings::SetupNamespaceMetadata()
 			.SetDisplayName(LOCTEXT("StatisSwitchDisplayName", "Static Switch Inputs"))
 			.SetDescription(LOCTEXT("StaticSwitchDescription", "Values which can only be set at edit time."))
 			.SetSortId(45)
+			.AddOption(ENiagaraNamespaceMetadataOptions::HideInSystem)
 			.AddOption(ENiagaraNamespaceMetadataOptions::HideInDefinitions)
 			.AddOption(ENiagaraNamespaceMetadataOptions::PreventEditingNamespace)
 			.AddOption(ENiagaraNamespaceMetadataOptions::PreventEditingNamespaceModifier)
@@ -260,7 +265,8 @@ void UNiagaraEditorSettings::SetupNamespaceMetadata()
 		FNiagaraNamespaceMetadata({FNiagaraConstants::ModuleNamespace})
 			.SetDisplayName(LOCTEXT("ModuleModifierDisplayName", "Module"))
 			.SetDescription(LOCTEXT("ModuleModifierDescription", "A namespace modifier which makes that attribute unique to the module\ninstance by appending the unique module name."))
-			.SetBackgroundColor(FLinearColor(FColor(102, 102, 152))),
+			.SetBackgroundColor(FLinearColor(FColor(102, 102, 152)))
+			.AddOption(ENiagaraNamespaceMetadataOptions::HideInSystem),
 		FNiagaraNamespaceMetadata({FNiagaraConstants::SystemNamespace})
 			.SetDisplayName(LOCTEXT("SystemModifierDisplayName", "System"))
 			.SetDescription(LOCTEXT("SystemModifierDescription", "A namespace modifier which specifies that an engine provided parameter comes from the system."))
@@ -453,6 +459,25 @@ void UNiagaraEditorSettings::SetDisplayAdvancedParameterPanelCategories(bool bIn
 	}
 }
 
+// Use the guids from the list at the top of this file to register last known expansion state.
+FNiagaraParameterPanelSectionStorage& UNiagaraEditorSettings::FindOrAddParameterPanelSectionStorage(FGuid PanelSectionId, bool& bOutAdded)
+{
+	bOutAdded = false;
+	check(PanelSectionId.IsValid());
+	for (FNiagaraParameterPanelSectionStorage& Storage : SystemParameterPanelSectionData)
+	{
+		if (PanelSectionId == Storage.ParamStorageId)
+		{
+			return Storage;
+		}
+	}
+
+	int32 Idx = SystemParameterPanelSectionData.Emplace(PanelSectionId);
+	bOutAdded = true;
+	return SystemParameterPanelSectionData[Idx];
+
+}
+
 bool UNiagaraEditorSettings::GetDisplayAffectedAssetStats() const
 {
 	return bDisplayAffectedAssetStats;
diff --git a/Engine/Plugins/FX/Niagara/Source/NiagaraEditor/Private/NiagaraScriptMergeManager.cpp b/Engine/Plugins/FX/Niagara/Source/NiagaraEditor/Private/NiagaraScriptMergeManager.cpp
--- a/Engine/Plugins/FX/Niagara/Source/NiagaraEditor/Private/NiagaraScriptMergeManager.cpp
+++ b/Engine/Plugins/FX/Niagara/Source/NiagaraEditor/Private/NiagaraScriptMergeManager.cpp
@@ -681,6 +681,41 @@ UNiagaraScript* FNiagaraScratchPadMergeAdapter::GetScratchPadScriptForFunctionId
 	return ParentScratchPadScriptPtr != nullptr ? *ParentScratchPadScriptPtr : nullptr;
 }
 
+
+TArray< FNiagaraScratchPadMergeAdapter::FMergeRecord > FNiagaraScratchPadMergeAdapter::GetMergedEmitterRecords() const
+{
+	TArray< FNiagaraScratchPadMergeAdapter::FMergeRecord > FoundArray;
+
+	// Logic needs to be kept in sync with FNiagaraScratchPadMergeAdapter::Initialize()
+	TArray<FVersionedNiagaraEmitter> ParentEmitters;
+	FVersionedNiagaraEmitter CurrentParent = ParentEmitter;
+	while (CurrentParent.Emitter != nullptr)
+	{
+		ParentEmitters.Insert(CurrentParent, 0);
+		CurrentParent = CurrentParent.GetEmitterData()->GetParent();
+	}
+
+	for (FVersionedNiagaraEmitter CurrentParentEmitter : ParentEmitters)
+	{
+		FVersionedNiagaraEmitterData* CurrentParentEmitterData = CurrentParentEmitter.GetEmitterData();
+		if (CurrentParentEmitterData->GetParent().Emitter == nullptr)
+		{
+			for (UNiagaraScript* Script : CurrentParentEmitterData->ParentScratchPads->Scripts)
+			{
+				// Right now we don't version scratch scripts. If we end up doing that, we'll need to account for that here.
+				FoundArray.Emplace(CurrentParentEmitter, Script, FGuid());
+			}
+		}
+		for (UNiagaraScript* Script : CurrentParentEmitterData->ScratchPads->Scripts)
+		{
+			// Right now we don't version scratch scripts. If we end up doing that, we'll need to account for that here.
+			FoundArray.Emplace(CurrentParentEmitter, Script, FGuid());
+		}
+	}
+
+	return FoundArray;
+}
+
 void FNiagaraScratchPadMergeAdapter::Initialize()
 {
 	FVersionedNiagaraEmitterData* TargetEmitterData = TargetEmitter.GetEmitterData();
@@ -1779,6 +1814,73 @@ bool FNiagaraScriptMergeManager::HasBaseModule(const FVersionedNiagaraEmitter& B
 	return BaseScriptStackAdapter.IsValid() && BaseScriptStackAdapter->GetModuleFunctionById(ModuleId).IsValid();
 }
 
+bool FNiagaraScriptMergeManager::FindBaseModule(const FVersionedNiagaraEmitter& BaseEmitter, ENiagaraScriptUsage ScriptUsage, FGuid ScriptUsageId, FGuid ModuleId, class UNiagaraScript*& OutActualScript, FGuid& OutScriptVersionGuid, FVersionedNiagaraEmitter& OutBaseEmitter)
+{
+	// Search through the existing scratch pads local to this versioned emitter and the ones inherited from parents.
+	OutScriptVersionGuid = FGuid();
+	OutActualScript = nullptr;
+	OutBaseEmitter = FVersionedNiagaraEmitter();
+
+	TSharedRef<FNiagaraEmitterMergeAdapter> BaseEmitterAdapter = GetEmitterMergeAdapterUsingCache(BaseEmitter);
+	TSharedRef<FNiagaraScratchPadMergeAdapter> ScratchPadMergeAdapter = GetScratchPadMergeAdapterUsingCache(BaseEmitter);
+	TSharedPtr<FNiagaraScriptStackMergeAdapter> BaseScriptStackAdapter = BaseEmitterAdapter->GetScriptStack(ScriptUsage, ScriptUsageId);
+	if (BaseScriptStackAdapter.IsValid())
+	{
+		TSharedPtr<FNiagaraStackFunctionMergeAdapter> AdapterFound = BaseScriptStackAdapter->GetModuleFunctionById(ModuleId);
+		if (AdapterFound.IsValid())
+		{
+			UNiagaraNodeFunctionCall* CallNode = AdapterFound->GetFunctionCallNode();
+			if (CallNode)
+			{
+				UNiagaraScript* Script = CallNode->FunctionScript;
+				if (Script)
+				{
+					int32 IndexToMatch = INDEX_NONE;
+					bool bParents = false;
+					FVersionedNiagaraEmitterData* EmitterData = BaseEmitter.Emitter->GetEmitterData(BaseEmitter.Version);
+					if (EmitterData)
+					{
+						if (!EmitterData->ScratchPads.IsNull())
+						{
+							IndexToMatch = EmitterData->ScratchPads->FindIndexForScript(Script);
+						}
+						if (INDEX_NONE == IndexToMatch && !EmitterData->ParentScratchPads.IsNull())
+						{
+							IndexToMatch = EmitterData->ParentScratchPads->FindIndexForScript(Script);
+							bParents = true;
+						}
+					}
+
+					if (bParents && IndexToMatch != INDEX_NONE)
+					{
+						TArray< FNiagaraScratchPadMergeAdapter::FMergeRecord > MergeRecordsInOrder = ScratchPadMergeAdapter->GetMergedEmitterRecords();
+						if (MergeRecordsInOrder.IsValidIndex(IndexToMatch))
+						{
+							OutScriptVersionGuid = MergeRecordsInOrder[IndexToMatch].OriginalScriptVersion;
+							OutActualScript = MergeRecordsInOrder[IndexToMatch].OriginalScript;
+							OutBaseEmitter = MergeRecordsInOrder[IndexToMatch].OriginalEmitter;
+							return true;
+						}
+
+					}
+					else if (!bParents && IndexToMatch != INDEX_NONE)
+					{
+						OutScriptVersionGuid = CallNode->SelectedScriptVersion;
+						OutActualScript = Script;
+						OutBaseEmitter = BaseEmitter;
+						return true;
+					}
+
+					
+				}
+				return false;
+			}
+		}
+	}
+
+	return false;
+}
+
 bool FNiagaraScriptMergeManager::IsModuleInputDifferentFromBase(const FVersionedNiagaraEmitter& Emitter, const FVersionedNiagaraEmitter& BaseEmitter, ENiagaraScriptUsage ScriptUsage, FGuid ScriptUsageId, FGuid ModuleId, FString InputName)
 {
 	SCOPE_CYCLE_COUNTER(STAT_NiagaraEditor_ScriptMergeManager_IsModuleInputDifferentFromBase);
diff --git a/Engine/Plugins/FX/Niagara/Source/NiagaraEditor/Private/NiagaraScriptMergeManager.h b/Engine/Plugins/FX/Niagara/Source/NiagaraEditor/Private/NiagaraScriptMergeManager.h
--- a/Engine/Plugins/FX/Niagara/Source/NiagaraEditor/Private/NiagaraScriptMergeManager.h
+++ b/Engine/Plugins/FX/Niagara/Source/NiagaraEditor/Private/NiagaraScriptMergeManager.h
@@ -240,6 +240,19 @@ public:
 
 	UNiagaraScript* GetScratchPadScriptForFunctionId(FGuid FunctionId);
 
+	// Struct to aid in book-keeping around scratch pads and merge inheritance.
+	struct FMergeRecord
+	{
+		FMergeRecord(FVersionedNiagaraEmitter& InOriginalEmitter, UNiagaraScript* InOriginalScript, const FGuid& InOriginalScriptVersion) : OriginalEmitter(InOriginalEmitter), OriginalScript(InOriginalScript), OriginalScriptVersion(InOriginalScriptVersion){}
+
+		FVersionedNiagaraEmitter OriginalEmitter;
+		UNiagaraScript* OriginalScript;
+		FGuid OriginalScriptVersion;
+	};
+
+	// Get the collapsed merge records for this merge session
+	TArray< FMergeRecord > GetMergedEmitterRecords() const;
+
 private:
 	void Initialize();
 
@@ -442,6 +455,9 @@ public:
 
 	bool HasBaseModule(const FVersionedNiagaraEmitter& BaseEmitter, ENiagaraScriptUsage ScriptUsage, FGuid ScriptUsageId, FGuid ModuleId);
 
+	// Find the module originating this module id in the scratch lists
+	bool FindBaseModule(const FVersionedNiagaraEmitter& BaseEmitter, ENiagaraScriptUsage ScriptUsage, FGuid ScriptUsageId, FGuid ModuleId, class UNiagaraScript*& OutActualScript, FGuid& OutScriptVersionGuid, FVersionedNiagaraEmitter& OutBaseEmitter);
+
 	bool IsModuleInputDifferentFromBase(const FVersionedNiagaraEmitter& Emitter, const FVersionedNiagaraEmitter& BaseEmitter, ENiagaraScriptUsage ScriptUsage, FGuid ScriptUsageId, FGuid ModuleId, FString InputName);
 
 	FApplyDiffResults ResetModuleInputToBase(const FVersionedNiagaraEmitter& VersionedEmitter, const FVersionedNiagaraEmitter& VersionedBaseEmitter, ENiagaraScriptUsage ScriptUsage, FGuid ScriptUsageId, FGuid ModuleId, FString InputName);
diff --git a/Engine/Plugins/FX/Niagara/Source/NiagaraEditor/Private/Toolkits/NiagaraSystemToolkit.h b/Engine/Plugins/FX/Niagara/Source/NiagaraEditor/Private/Toolkits/NiagaraSystemToolkit.h
--- a/Engine/Plugins/FX/Niagara/Source/NiagaraEditor/Private/Toolkits/NiagaraSystemToolkit.h
+++ b/Engine/Plugins/FX/Niagara/Source/NiagaraEditor/Private/Toolkits/NiagaraSystemToolkit.h
@@ -95,6 +95,9 @@
 	void CompileSystem(bool bFullRebuild);
 
 	TSharedPtr<FNiagaraSystemViewModel> GetSystemViewModel();
+	TSharedPtr<FNiagaraSystemGraphSelectionViewModel> GetSystemGraphSelectionViewModel() {
+		return SystemGraphSelectionViewModel;
+	}
 	
 	// @todo This is a hack for now until we reconcile the default toolbar with application modes [duplicated from counterpart in Blueprint Editor]
 	void RegisterToolbarTab(const TSharedRef<class FTabManager>& TabManager);
diff --git a/Engine/Plugins/FX/Niagara/Source/NiagaraEditor/Private/Toolkits/SystemToolkitModes/NiagaraSystemToolkitModeBase.cpp b/Engine/Plugins/FX/Niagara/Source/NiagaraEditor/Private/Toolkits/SystemToolkitModes/NiagaraSystemToolkitModeBase.cpp
--- a/Engine/Plugins/FX/Niagara/Source/NiagaraEditor/Private/Toolkits/SystemToolkitModes/NiagaraSystemToolkitModeBase.cpp
+++ b/Engine/Plugins/FX/Niagara/Source/NiagaraEditor/Private/Toolkits/SystemToolkitModes/NiagaraSystemToolkitModeBase.cpp
@@ -126,7 +126,7 @@ void FNiagaraSystemToolkitModeBase::OnEditedScriptGraphChanged(const FEdGraphEdi
 	TSharedPtr<FNiagaraSystemToolkit> Toolkit = StaticCastSharedPtr<FNiagaraSystemToolkit>(SystemToolkit.Pin());
 	// We need to update the parameter panel view model with new parameters potentially
 	if (Toolkit->GetSystemViewModel()->GetParameterPanelViewModel())
-		Toolkit->GetSystemViewModel()->GetParameterPanelViewModel()->RefreshNextTick();
+		Toolkit->GetSystemViewModel()->GetParameterPanelViewModel()->RefreshFullNextTick(true);
 }
 
 void FNiagaraSystemToolkitModeBase::OnParameterPanelViewModelExternalSelectionChanged()
@@ -150,20 +150,7 @@ void FNiagaraSystemToolkitModeBase::OnParameterPanelViewModelExternalSelectionCh
 
 void FNiagaraSystemToolkitModeBase::OnSystemSelectionChanged()
 {
-	TSharedPtr<FNiagaraSystemToolkit> Toolkit = StaticCastSharedPtr<FNiagaraSystemToolkit>(SystemToolkit.Pin());
-	TSharedPtr<SDockTab>  ActiveTab = Toolkit->GetSystemViewModel()->GetDocumentViewModel()->GetActiveDocumentTab().Pin();
-	if (ActiveTab.IsValid())
-	{
-		if (ActiveTab->GetLayoutIdentifier().TabType != SystemOverviewTabID)
-		{
-			TSharedPtr<SDockTab> FoundOverviewTab = Toolkit->GetTabManager()->FindExistingLiveTab(SystemOverviewTabID);
-			if (FoundOverviewTab.IsValid())
-			{
-				FoundOverviewTab->DrawAttention();
-				FoundOverviewTab->ActivateInParent(ETabActivationCause::SetDirectly);
-			}
-		}
-	}
+	// Do nothing for now
 }
 
 void FNiagaraSystemToolkitModeBase::UpdateSelectionForActiveDocument()
@@ -203,6 +190,10 @@ void FNiagaraSystemToolkitModeBase::UpdateSelectionForActiveDocument()
 		ObjectSelection->ClearSelectedObjects();
 		ObjectSelectionSubHeaderText = LOCTEXT("EmptySelection", "Nothing selected");
 	}
+
+	// We need to update the parameter panel view model with new parameters potentially
+	if (Toolkit->GetSystemViewModel()->GetParameterPanelViewModel())
+		Toolkit->GetSystemViewModel()->GetParameterPanelViewModel()->RefreshDueToActiveDocumentChanged();
 }
 
 void FNiagaraSystemToolkitModeBase::RegisterTabFactories(TSharedPtr<FTabManager> InTabManager)
@@ -277,7 +268,7 @@ void FNiagaraSystemToolkitModeBase::RegisterTabFactories(TSharedPtr<FTabManager>
 		.SetIcon(FSlateIcon(FNiagaraEditorStyle::Get().GetStyleSetName(), "Tab.SystemOverview"));
 
 	InTabManager->RegisterTabSpawner(ScratchPadScriptsTabID, FOnSpawnTab::CreateSP(this, &FNiagaraSystemToolkitModeBase::SpawnTab_ScratchPadScripts))
-		.SetDisplayName(LOCTEXT("ScratchPadScriptsTabName", "Local Scripts"))
+		.SetDisplayName(LOCTEXT("ScratchPadScriptsTabName", "Local Modules"))
 		.SetGroup(WorkspaceMenuCategory.ToSharedRef())
 		.SetIcon(FSlateIcon(FNiagaraEditorStyle::Get().GetStyleSetName(), "Tab.ScratchPad"));
 
@@ -703,7 +694,7 @@ TSharedRef<SDockTab> FNiagaraSystemToolkitModeBase::SpawnTab_ScratchPadScripts(c
 	}
 
 	TSharedRef<SDockTab> SpawnedTab = SNew(SDockTab)
-		.Label(LOCTEXT("ScratchPadTabLabel", "Local Scripts"))
+		.Label(LOCTEXT("ScratchPadTabLabel", "Local Modules"))
 		[
 			SystemToolkit.Pin()->GetScriptScratchpadManager().ToSharedRef()
 		];
diff --git a/Engine/Plugins/FX/Niagara/Source/NiagaraEditor/Private/Toolkits/SystemToolkitModes/NiagaraSystemToolkitMode_Default.cpp b/Engine/Plugins/FX/Niagara/Source/NiagaraEditor/Private/Toolkits/SystemToolkitModes/NiagaraSystemToolkitMode_Default.cpp
--- a/Engine/Plugins/FX/Niagara/Source/NiagaraEditor/Private/Toolkits/SystemToolkitModes/NiagaraSystemToolkitMode_Default.cpp
+++ b/Engine/Plugins/FX/Niagara/Source/NiagaraEditor/Private/Toolkits/SystemToolkitModes/NiagaraSystemToolkitMode_Default.cpp
@@ -9,6 +9,7 @@
 #include "NiagaraEditorSettings.h"
 #include "ViewModels/NiagaraEmitterHandleViewModel.h"
 #include "ViewModels/NiagaraSystemViewModel.h"
+#include "ViewModels/NiagaraSystemEditorDocumentsViewModel.h"
 
 #include "Framework/MultiBox/MultiBoxBuilder.h"
 
@@ -300,6 +301,22 @@
 	// SystemToolkit.Pin()->AddToolbarExtender(NiagaraEditorModule.GetToolBarExtensibilityManager()->GetAllExtenders(SystemToolkit.Pin()->GetToolkitCommands(), SystemToolkit.Pin()->GetObjectsBeingEdited()));
 }
 
+void FNiagaraSystemToolkitMode_Default::PostActivateMode()
+{
+	// We need to register the primary tab ID with the document system so that we can easily update and 
+	// reference it elsewhere and only this class knows it's identity since it is dynamically allocated.
+	TSharedPtr<FNiagaraSystemToolkit> Toolkit = StaticCastSharedPtr<FNiagaraSystemToolkit>(SystemToolkit.Pin());
+	if (Toolkit.IsValid())
+	{
+		TSharedPtr<FNiagaraSystemViewModel> SystemVM = Toolkit->GetSystemViewModel();
+		if (SystemVM.IsValid())
+		{
+			SystemVM->GetDocumentViewModel()->SetPrimaryDocumentID(SystemOverviewTabID);
+		}
+	}
+}
+
+
 void FNiagaraSystemToolkitMode_Default::RegisterTabFactories(TSharedPtr<FTabManager> InTabManager)
 {
 	// this will register the common tab factories shared across modes and will call ExtendToolbar
diff --git a/Engine/Plugins/FX/Niagara/Source/NiagaraEditor/Private/Toolkits/SystemToolkitModes/NiagaraSystemToolkitMode_Default.h b/Engine/Plugins/FX/Niagara/Source/NiagaraEditor/Private/Toolkits/SystemToolkitModes/NiagaraSystemToolkitMode_Default.h
--- a/Engine/Plugins/FX/Niagara/Source/NiagaraEditor/Private/Toolkits/SystemToolkitModes/NiagaraSystemToolkitMode_Default.h
+++ b/Engine/Plugins/FX/Niagara/Source/NiagaraEditor/Private/Toolkits/SystemToolkitModes/NiagaraSystemToolkitMode_Default.h
@@ -15,4 +15,7 @@ public:
 	virtual void RegisterTabFactories(TSharedPtr<FTabManager> InTabManager) override;
 	
 	void ExtendToolbar();
+
+
+	virtual void PostActivateMode() override;
 };
diff --git a/Engine/Plugins/FX/Niagara/Source/NiagaraEditor/Private/ViewModels/NiagaraParameterDefinitionsPanelViewModel.cpp b/Engine/Plugins/FX/Niagara/Source/NiagaraEditor/Private/ViewModels/NiagaraParameterDefinitionsPanelViewModel.cpp
--- a/Engine/Plugins/FX/Niagara/Source/NiagaraEditor/Private/ViewModels/NiagaraParameterDefinitionsPanelViewModel.cpp
+++ b/Engine/Plugins/FX/Niagara/Source/NiagaraEditor/Private/ViewModels/NiagaraParameterDefinitionsPanelViewModel.cpp
@@ -227,7 +227,7 @@ void FNiagaraSystemToolkitParameterDefinitionsPanelViewModel::Init(const FSystem
 
 const TArray<UNiagaraGraph*> FNiagaraSystemToolkitParameterDefinitionsPanelViewModel::GetEditableGraphsConst() const
 {
-	return FNiagaraSystemToolkitParameterPanelUtilities::GetEditableGraphs(SystemViewModel, SystemGraphSelectionViewModelWeak);
+	return FNiagaraSystemToolkitParameterPanelUtilities::GetEditableGraphs(SystemViewModel, SystemGraphSelectionViewModelWeak, false);
 }
 
 const TArray<UNiagaraParameterDefinitions*> FNiagaraSystemToolkitParameterDefinitionsPanelViewModel::GetAvailableParameterDefinitionsAssets(bool bSkipSubscribedParameterDefinitions) const
diff --git a/Engine/Plugins/FX/Niagara/Source/NiagaraEditor/Private/ViewModels/NiagaraParameterPanelViewModel.cpp b/Engine/Plugins/FX/Niagara/Source/NiagaraEditor/Private/ViewModels/NiagaraParameterPanelViewModel.cpp
--- a/Engine/Plugins/FX/Niagara/Source/NiagaraEditor/Private/ViewModels/NiagaraParameterPanelViewModel.cpp
+++ b/Engine/Plugins/FX/Niagara/Source/NiagaraEditor/Private/ViewModels/NiagaraParameterPanelViewModel.cpp
@@ -32,6 +32,7 @@
 #include "Widgets/SNiagaraParameterMenu.h"
 #include "Widgets/SNiagaraParameterPanel.h"
 #include "ViewModels/Stack/NiagaraStackViewModel.h"
+#include "ViewModels/NiagaraScratchPadScriptViewModel.h"
 #include "ViewModels/NiagaraSystemEditorDocumentsViewModel.h"
 
 #include "Framework/Commands/GenericCommands.h"
@@ -50,103 +51,47 @@ template<> TMap<UNiagaraSystem*, TArray<FNiagaraSystemToolkitParameterPanelViewM
 TArray<FNiagaraParameterPanelCategory> FNiagaraSystemToolkitParameterPanelViewModel::DefaultCategories;
 TArray<FNiagaraParameterPanelCategory> FNiagaraSystemToolkitParameterPanelViewModel::DefaultAdvancedCategories;
 TArray<FNiagaraParameterPanelCategory> FNiagaraSystemToolkitParameterPanelViewModel::UserCategories;
+TArray<FNiagaraParameterPanelCategory> FNiagaraSystemToolkitParameterPanelViewModel::DefaultScriptCategories;
+TArray<FNiagaraParameterPanelCategory> FNiagaraSystemToolkitParameterPanelViewModel::DefaultAdvancedScriptCategories;
+
 TArray<FNiagaraParameterPanelCategory> FNiagaraScriptToolkitParameterPanelViewModel::DefaultCategories;
 TArray<FNiagaraParameterPanelCategory> FNiagaraScriptToolkitParameterPanelViewModel::DefaultAdvancedCategories;
 TArray<FNiagaraParameterPanelCategory> FNiagaraParameterDefinitionsToolkitParameterPanelViewModel::DefaultCategories;
 
+
+
 ///////////////////////////////////////////////////////////////////////////////
 /// System Toolkit Parameter Panel Utilities								///
 ///////////////////////////////////////////////////////////////////////////////
 
-TArray<UNiagaraGraph*> FNiagaraSystemToolkitParameterPanelUtilities::GetAllGraphs(const TSharedPtr<FNiagaraSystemViewModel>& SystemViewModel)
+TArray<UNiagaraGraph*> FNiagaraSystemToolkitParameterPanelUtilities::GetAllGraphs(const TSharedPtr<FNiagaraSystemViewModel>& SystemViewModel, bool bActiveScriptMode)
 {
-	TArray<UNiagaraGraph*> OutGraphs;
-
-	// Helper lambda to null check graph weak object ptrs before adding them to the retval array.
-	auto AddToOutGraphsChecked = [&OutGraphs](const TWeakObjectPtr<UNiagaraGraph>& WeakGraph) {
-		UNiagaraGraph* Graph = WeakGraph.Get();
-		if (Graph == nullptr)
-		{
-			ensureMsgf(false, TEXT("Encountered null graph when gathering editable graphs for system parameter panel viewmodel!"));
-			return;
-		}
-		OutGraphs.Add(Graph);
-	};
-
-	if (SystemViewModel->GetEditMode() == ENiagaraSystemViewModelEditMode::SystemAsset)
+	if (bActiveScriptMode)
 	{
-		TArray<UNiagaraGraph*> Graphs;
-		AddToOutGraphsChecked(SystemViewModel->GetSystemScriptViewModel()->GetGraphViewModel()->GetGraph());
-
-		const TArray<TSharedRef<FNiagaraEmitterHandleViewModel>>& EmitterHandleViewModels = SystemViewModel->GetEmitterHandleViewModels();
-		for (const TSharedRef<FNiagaraEmitterHandleViewModel>& EmitterHandleViewModel : EmitterHandleViewModels)
-		{
-			FNiagaraEmitterHandle* EmitterHandle = EmitterHandleViewModel->GetEmitterHandle();
-			if (EmitterHandle == nullptr)
-			{
-				continue;
-			}
-			UNiagaraGraph* Graph = Cast<UNiagaraScriptSource>(EmitterHandle->GetEmitterData()->GraphSource)->NodeGraph;
-			if (Graph)
-			{
-				OutGraphs.Add(Graph);
-			}
-		}
+		TArray<UNiagaraGraph*> EdGraphs = SystemViewModel->GetDocumentViewModel()->GetAllGraphsForActiveScriptDocument();
+		return EdGraphs;
 	}
 	else
 	{
-		FNiagaraEmitterHandle* EmitterHandle = SystemViewModel->GetEmitterHandleViewModels()[0]->GetEmitterHandle();
-		if (EmitterHandle != nullptr)
-		{
-			OutGraphs.Add(Cast<UNiagaraScriptSource>(EmitterHandle->GetEmitterData()->GraphSource)->NodeGraph);
-		}
-	}
-
-	TArray<UNiagaraGraph*> EdGraphs = SystemViewModel->GetDocumentViewModel()->GetEditableGraphsForActiveDocument();
-	for (UNiagaraGraph* Graph : EdGraphs)
-	{
-		OutGraphs.AddUnique(Graph);
+		TArray<UNiagaraGraph*> EdGraphs = SystemViewModel->GetDocumentViewModel()->GetAllGraphsForPrimaryDocument();
+		return EdGraphs;
 	}
-	return OutGraphs;
 }
 
-TArray<UNiagaraGraph*> FNiagaraSystemToolkitParameterPanelUtilities::GetEditableGraphs(const TSharedPtr<FNiagaraSystemViewModel>& SystemViewModel, const TWeakPtr<FNiagaraSystemGraphSelectionViewModel>& SystemGraphSelectionViewModelWeak)
+TArray<UNiagaraGraph*> FNiagaraSystemToolkitParameterPanelUtilities::GetEditableGraphs(const TSharedPtr<FNiagaraSystemViewModel>& SystemViewModel, const TWeakPtr<FNiagaraSystemGraphSelectionViewModel>& SystemGraphSelectionViewModelWeak, bool bActiveScriptMode)
 {
-	TArray<UNiagaraGraph*> EditableGraphs;
-
-	// Helper lambda to null check graph weak object ptrs before adding them to the retval array.
-	auto AddToEditableGraphChecked = [&EditableGraphs](const TWeakObjectPtr<UNiagaraGraph>& WeakGraph) {
-		UNiagaraGraph* Graph = WeakGraph.Get();
-		if (Graph == nullptr)
-		{
-			ensureMsgf(false, TEXT("Encountered null graph when gathering editable graphs for system parameter panel viewmodel!"));
-			return;
-		}
-		EditableGraphs.Add(Graph);
-	};
-
-	if (SystemViewModel->GetEditMode() == ENiagaraSystemViewModelEditMode::SystemAsset && ensureMsgf(SystemGraphSelectionViewModelWeak.IsValid(), TEXT("SystemGraphSelectionViewModel was null for System edit mode!")))
+	if (bActiveScriptMode)
 	{
-		for (const TWeakObjectPtr<UNiagaraGraph>& WeakGraph : SystemGraphSelectionViewModelWeak.Pin()->GetSelectedEmitterScriptGraphs())
-		{
-			AddToEditableGraphChecked(WeakGraph);
-		}
-		AddToEditableGraphChecked(SystemViewModel->GetSystemScriptViewModel()->GetGraphViewModel()->GetGraph());
+		TArray<UNiagaraGraph*> EdGraphs = SystemViewModel->GetDocumentViewModel()->GetEditableGraphsForActiveScriptDocument();
+		return EdGraphs;
 	}
 	else
 	{
-		EditableGraphs.Add(Cast<UNiagaraScriptSource>(SystemViewModel->GetEmitterHandleViewModels()[0]->GetEmitterHandle()->GetEmitterData()->GraphSource)->NodeGraph);
+		TArray<UNiagaraGraph*> EdGraphs = SystemViewModel->GetDocumentViewModel()->GetEditableGraphsForPrimaryDocument();
+		return EdGraphs;
 	}
-
-	TArray<UNiagaraGraph*> EdGraphs = SystemViewModel->GetDocumentViewModel()->GetEditableGraphsForActiveDocument();
-	for (UNiagaraGraph* Graph : EdGraphs)
-	{
-		EditableGraphs.AddUnique(Graph);
-	}
-	return EditableGraphs;
 }
 
-
 FReply FNiagaraSystemToolkitParameterPanelUtilities::CreateDragEventForParameterItem(const FNiagaraParameterPanelItemBase& DraggedItem, const FPointerEvent& MouseEvent, const TArray<FNiagaraGraphParameterReference>& GraphParameterReferencesForItem, const TSharedPtr<TArray<FName>>& ParametersWithNamespaceModifierRenamePending)
 {
 	const static FText TooltipFormat = LOCTEXT("Parameters", "Name: {0} \nType: {1}");
@@ -278,14 +223,14 @@ bool FNiagaraParameterPanelUtilities::GetCanSetParameterCustomNamespaceModifierA
 /// Immutable Parameter Panel View Model									///
 ///////////////////////////////////////////////////////////////////////////////
 
-void INiagaraImmutableParameterPanelViewModel::Refresh() const
+void INiagaraImmutableParameterPanelViewModel::RefreshFull(bool bDoCategoryExpansion) const
 {
-	OnRequestRefreshDelegate.ExecuteIfBound();
+	OnRequestRefreshDelegate.ExecuteIfBound(bDoCategoryExpansion);
 }
 
-void INiagaraImmutableParameterPanelViewModel::RefreshNextTick() const
+void INiagaraImmutableParameterPanelViewModel::RefreshFullNextTick(bool bDoCategoryExpansion) const
 {
-	OnRequestRefreshNextTickDelegate.ExecuteIfBound();
+	OnRequestRefreshNextTickDelegate.ExecuteIfBound(bDoCategoryExpansion);
 }
 
 void INiagaraImmutableParameterPanelViewModel::PostUndo(bool bSuccess)
@@ -346,6 +291,18 @@ INiagaraParameterPanelViewModel::~INiagaraParameterPanelViewModel()
 
 }
 
+
+void INiagaraParameterPanelViewModel::RefreshFull(bool bDoCategoryExpansion) const
+{
+	INiagaraImmutableParameterPanelViewModel::RefreshFull(bDoCategoryExpansion);
+
+	TSharedPtr<INiagaraParameterPanelViewModel> MasterVM = MasterParameterPanelViewModel.Pin();
+	if (MasterVM.IsValid())
+	{
+		MasterVM->RefreshFull(bDoCategoryExpansion);
+	}
+}
+
 bool INiagaraParameterPanelViewModel::GetCanDeleteParameterAndToolTip(const FNiagaraParameterPanelItem& ItemToDelete, FText& OutCanDeleteParameterToolTip) const
 {
 	if (ItemToDelete.bExternallyReferenced)
@@ -397,7 +354,7 @@ bool INiagaraParameterPanelViewModel::GetCanPasteParameterMetaDataAndToolTip(FTe
 	return true;
 }
 
-void INiagaraParameterPanelViewModel::DuplicateParameters(const TArray<FNiagaraParameterPanelItem> ItemsToDuplicate) const
+void INiagaraParameterPanelViewModel::DuplicateParameters(const TArray<FNiagaraParameterPanelItem> ItemsToDuplicate) 
 {
 	bool bAnyChange = false;
 	FScopedTransaction Transaction(LOCTEXT("DuplicateParameter(s)Transaction", "Duplicate parameter(s)"));
@@ -538,7 +495,7 @@ bool INiagaraParameterPanelViewModel::GetCanSubscribeParameterToLibraryAndToolTi
 	return false;
 }
 
-void INiagaraParameterPanelViewModel::SetParameterIsSubscribedToLibrary(const FNiagaraParameterPanelItem ItemToModify, const bool bSubscribed) const
+void INiagaraParameterPanelViewModel::SetParameterIsSubscribedToLibrary(const FNiagaraParameterPanelItem ItemToModify, const bool bSubscribed) 
 {
 	if (ensureMsgf(ItemToModify.bExternallyReferenced == false, TEXT("Cannot modify an externally referenced parameter.")) == false)
 	{
@@ -548,7 +505,7 @@ void INiagaraParameterPanelViewModel::SetParameterIsSubscribedToLibrary(const FN
 	SetParameterIsSubscribedToLibrary(ItemToModify.ScriptVariable, bSubscribed);
 }
 
-void INiagaraParameterPanelViewModel::SetParameterNamespace(const FNiagaraParameterPanelItem ItemToModify, FNiagaraNamespaceMetadata NewNamespaceMetaData, bool bDuplicateParameter) const
+void INiagaraParameterPanelViewModel::SetParameterNamespace(const FNiagaraParameterPanelItem ItemToModify, FNiagaraNamespaceMetadata NewNamespaceMetaData, bool bDuplicateParameter) 
 {
 	FName NewName = FNiagaraParameterUtilities::ChangeNamespace(ItemToModify.GetVariable().GetName(), NewNamespaceMetaData);
 	if (NewName != NAME_None)
@@ -583,7 +540,7 @@ void INiagaraParameterPanelViewModel::SetParameterNamespace(const FNiagaraParame
 	}
 }
 
-void INiagaraParameterPanelViewModel::SetParameterNamespaceModifier(const FNiagaraParameterPanelItem ItemToModify, const FName NewNamespaceModifier, bool bDuplicateParameter) const
+void INiagaraParameterPanelViewModel::SetParameterNamespaceModifier(const FNiagaraParameterPanelItem ItemToModify, const FName NewNamespaceModifier, bool bDuplicateParameter) 
 {
 	FName NewName = FNiagaraParameterUtilities::SetSpecificNamespaceModifier(ItemToModify.GetVariable().GetName(), NewNamespaceModifier);
 	if (NewName != NAME_None)
@@ -619,7 +576,7 @@ void INiagaraParameterPanelViewModel::SetParameterNamespaceModifier(const FNiaga
 	}
 }
 
-void INiagaraParameterPanelViewModel::ChangeParameterType(const TArray<FNiagaraParameterPanelItem> ItemsToModify, const FNiagaraTypeDefinition NewType) const
+void INiagaraParameterPanelViewModel::ChangeParameterType(const TArray<FNiagaraParameterPanelItem> ItemsToModify, const FNiagaraTypeDefinition NewType) 
 {
 	ensureMsgf(false, TEXT("type change not supported for this view"));
 }
@@ -663,7 +620,7 @@ bool INiagaraParameterPanelViewModel::GetCanChangeParameterType(const TArray<FNi
 	return true;
 }
 
-void INiagaraParameterPanelViewModel::GetChangeTypeSubMenu(FMenuBuilder& MenuBuilder, TArray<FNiagaraParameterPanelItem> Items) const
+void INiagaraParameterPanelViewModel::GetChangeTypeSubMenu(FMenuBuilder& MenuBuilder, TArray<FNiagaraParameterPanelItem> Items) 
 {
 	TArray<FNiagaraTypeDefinition> FilteredTypes;
 	for (const FNiagaraTypeDefinition& RegisteredType : FNiagaraTypeRegistry::GetRegisteredTypes())
@@ -690,7 +647,7 @@ void INiagaraParameterPanelViewModel::GetChangeTypeSubMenu(FMenuBuilder& MenuBui
 	}
 }
 
-void INiagaraParameterPanelViewModel::SetParameterCustomNamespaceModifier(const FNiagaraParameterPanelItem ItemToModify, bool bDuplicateParameter) const
+void INiagaraParameterPanelViewModel::SetParameterCustomNamespaceModifier(const FNiagaraParameterPanelItem ItemToModify, bool bDuplicateParameter) 
 {
 	if (bDuplicateParameter == false && ItemToModify.bExternallyReferenced)
 	{
@@ -738,7 +695,7 @@ void INiagaraParameterPanelViewModel::SetParameterCustomNamespaceModifier(const
 }
 
 
-void INiagaraParameterPanelViewModel::GetChangeNamespaceSubMenu(FMenuBuilder& MenuBuilder, bool bDuplicateParameter, FNiagaraParameterPanelItem Item) const
+void INiagaraParameterPanelViewModel::GetChangeNamespaceSubMenu(FMenuBuilder& MenuBuilder, bool bDuplicateParameter, FNiagaraParameterPanelItem Item) 
 {
 	TArray<FNiagaraParameterUtilities::FChangeNamespaceMenuData> MenuData;
 	FNiagaraParameterUtilities::GetChangeNamespaceMenuData(Item.GetVariable().GetName(), GetParameterContext(), MenuData);
@@ -774,7 +731,7 @@ void INiagaraParameterPanelViewModel::GetChangeNamespaceSubMenu(FMenuBuilder& Me
 	}
 }
 
-void INiagaraParameterPanelViewModel::GetChangeNamespaceModifierSubMenu(FMenuBuilder& MenuBuilder, bool bDuplicateParameter, FNiagaraParameterPanelItem Item) const
+void INiagaraParameterPanelViewModel::GetChangeNamespaceModifierSubMenu(FMenuBuilder& MenuBuilder, bool bDuplicateParameter, FNiagaraParameterPanelItem Item) 
 {
 	TArray<FName> OptionalNamespaceModifiers;
 	FNiagaraParameterUtilities::GetOptionalNamespaceModifiers(Item.GetVariable().GetName(), GetParameterContext(), OptionalNamespaceModifiers);
@@ -823,7 +780,7 @@ const TArray<FNiagaraParameterPanelItem>& INiagaraParameterPanelViewModel::GetCa
 	return CachedViewedItems;
 }
 
-void INiagaraParameterPanelViewModel::SelectParameterItemByName(const FName ParameterName, const bool bRequestRename) const
+void INiagaraParameterPanelViewModel::SelectParameterItemByName(const FName ParameterName, const bool bRequestRename)
 {
 	if (bRequestRename)
 	{
@@ -833,9 +790,32 @@ void INiagaraParameterPanelViewModel::SelectParameterItemByName(const FName Para
 	{
 		OnSelectParameterItemByNameDelegate.ExecuteIfBound(ParameterName);
 	}
+
+	// Propagate to the selection object (if found)
+	UNiagaraScriptVariable* FoundScriptVariable = nullptr;
+	for (const UNiagaraGraph* Graph : GetEditableGraphsConst())
+	{
+		if (UNiagaraScriptVariable* ScriptVariable = Graph->GetScriptVariable(ParameterName))
+		{
+			FoundScriptVariable = ScriptVariable;
+			break;
+		}
+	}
+
+	const TSharedPtr<FNiagaraObjectSelection>& Selection = GetVariableObjectSelection();
+	if (Selection.IsValid() && FoundScriptVariable)
+	{
+		Selection->SetSelectedObject(FoundScriptVariable);
+	}
+
+	TSharedPtr<INiagaraParameterPanelViewModel> MasterVM = MasterParameterPanelViewModel.Pin();
+	if (MasterVM.IsValid())
+	{
+		MasterVM->SelectParameterItemByName(ParameterName, bRequestRename);
+	}
 }
 
-void INiagaraParameterPanelViewModel::SubscribeParameterToLibraryIfMatchingDefinition(const UNiagaraScriptVariable* ScriptVarToModify, const FName ScriptVarName) const
+void INiagaraParameterPanelViewModel::SubscribeParameterToLibraryIfMatchingDefinition(const UNiagaraScriptVariable* ScriptVarToModify, const FName ScriptVarName) 
 {
 	const TArray<const UNiagaraScriptVariable*> ReservedParametersForName = FNiagaraParameterDefinitionsUtilities::FindReservedParametersByName(ScriptVarName);
 
@@ -953,6 +933,7 @@ void FNiagaraSystemToolkitParameterPanelViewModel::Cleanup()
 
 }
 
+
 void FNiagaraSystemToolkitParameterPanelViewModel::Init(const FSystemToolkitUIContext& InUIContext)
 {
 	UIContext = InUIContext;
@@ -960,14 +941,63 @@ void FNiagaraSystemToolkitParameterPanelViewModel::Init(const FSystemToolkitUICo
 
 	VariableObjectSelection = MakeShared< FNiagaraObjectSelection>();
 
-	ActiveSectionIndex = 0;
-	UserOnlyIdx = Sections.Add(LOCTEXT("SectionUser", "User"));
-	CommonOnlyIdx = Sections.Add(LOCTEXT("SectionCommon", "Common"));
-	ActiveIdx = Sections.Add(LOCTEXT("SectionActive", "Active"));
+	// Standalone emitter assets have different default section setups than systems
+	if (SystemViewModel->GetEditMode() == ENiagaraSystemViewModelEditMode::EmitterAsset)
+	{
+		ActiveSectionIndex = 1;
+		LastActiveSystemSectionIdx = ActiveSectionIndex;
+	}
+	else
+	{
+		ActiveSectionIndex = 0;
+		LastActiveSystemSectionIdx = ActiveSectionIndex;
+	}
+
+	// Define the sections
+	UserOnlyIdx = Sections.Emplace(FGuid(0x9E29BC43, 0x31524AB5, 0xA9D6675D, 0x741137CB), LOCTEXT("SectionUser", "User"), LOCTEXT("SectionUserDesc", "Displays only User exposed parameters"));
+	CommonOnlyIdx = Sections.Emplace(FGuid(0x3054B620, 0xBE714E54, 0xB7DB81E9, 0x00C0DE3D ),LOCTEXT("SectionCommon", "Common"), LOCTEXT("SectionCommonDesc", "Displays parameters defined in Parameter Libraries"));
+	ActiveSystemIdx = Sections.Emplace(FGuid(0x7B4AFB34, 0xD0DF4618, 0x9A05349E, 0x361CE735), LOCTEXT("SectionSystemActive", "Active Overview"), LOCTEXT("SectionActiveSystemDesc", "Displays parameters that are context-sensitive to the selected items in the System Overview."));
+	ActiveScriptIdx = Sections.Emplace(FGuid(0x6E11F3EA, 0x5E314E82, 0xA64FEA41, 0x4F1DB891), LOCTEXT("SectionScriptActive", "Active Module"), LOCTEXT("SectionActiveScriptDesc", "Displays parameters that are context-sensitive to the active Scratch module document. Use this for renaming/editing local parameters to the module."));
 
-	SectionDescriptions.Add(LOCTEXT("SectionUserDesc", "Displays only User exposed parameters"));
-	SectionDescriptions.Add(LOCTEXT("SectionCommonDesc", "Displays parameters defined in Parameter Libraries"));
-	SectionDescriptions.Add(LOCTEXT("SectionActiveDesc", "Displays parameters that are context-sensitive to the selected items in the active document tab."));
+	// Go through all the settings and propagate with the default values if not already defined.
+	UNiagaraEditorSettings* Settings = GetMutableDefault<UNiagaraEditorSettings>();
+	if (Settings)
+	{
+		bool bSaveConfig = false;
+		bool bAdded = false;
+		FNiagaraParameterPanelSectionStorage& UserStorage = Settings->FindOrAddParameterPanelSectionStorage(Sections[UserOnlyIdx].SectionId, bAdded);
+		if (bAdded)
+		{
+			UserStorage.ExpandedCategories.Emplace(FNiagaraEditorGuids::UserNamespaceMetaDataGuid);
+			bSaveConfig = true;
+		}
+
+		FNiagaraParameterPanelSectionStorage& CommonStorage = Settings->FindOrAddParameterPanelSectionStorage(Sections[CommonOnlyIdx].SectionId, bAdded);
+		if (bAdded)
+		{
+			CommonStorage.ExpandedCategories.Emplace(FNiagaraEditorGuids::ParticleAttributeNamespaceMetaDataGuid);
+			bSaveConfig = true;
+		}
+
+		FNiagaraParameterPanelSectionStorage& ScriptStorage = Settings->FindOrAddParameterPanelSectionStorage(Sections[ActiveScriptIdx].SectionId, bAdded);
+		if (bAdded)
+		{
+			ScriptStorage.ExpandedCategories.Emplace(FNiagaraEditorGuids::ModuleNamespaceMetaDataGuid);
+			bSaveConfig = true;
+		}
+
+		FNiagaraParameterPanelSectionStorage& SystemStorage = Settings->FindOrAddParameterPanelSectionStorage(Sections[ActiveSystemIdx].SectionId, bAdded);
+		if (bAdded)
+		{
+			SystemStorage.ExpandedCategories.Emplace(FNiagaraEditorGuids::SystemNamespaceMetaDataGuid);
+			bSaveConfig = true;
+		}
+
+		if (bSaveConfig)
+		{
+			Settings->SaveConfig();
+		}
+	}
 
 	// Init bindings
 	// Bind OnChanged() and OnNeedsRecompile() callbacks for all script source graphs.
@@ -1052,6 +1082,26 @@ void FNiagaraSystemToolkitParameterPanelViewModel::Init(const FSystemToolkitUICo
 				UserCategories.Add(NamespaceMetadatum);
 		}
 	}
+	// Init Default Script categories
+	if (DefaultScriptCategories.Num() == 0)
+	{
+		TArray<FNiagaraNamespaceMetadata> NamespaceMetadata = GetDefault<UNiagaraEditorSettings>()->GetAllNamespaceMetadata();
+		for (const FNiagaraNamespaceMetadata& NamespaceMetadatum : NamespaceMetadata)
+		{
+			if (NamespaceMetadatum.Options.Contains(ENiagaraNamespaceMetadataOptions::HideInScript) == false)
+			{
+				if (NamespaceMetadatum.Options.Contains(ENiagaraNamespaceMetadataOptions::AdvancedInScript))
+				{
+					DefaultAdvancedScriptCategories.Add(NamespaceMetadatum);
+				}
+				else
+				{
+					DefaultScriptCategories.Add(NamespaceMetadatum);
+					DefaultAdvancedScriptCategories.Add(NamespaceMetadatum);
+				}
+			}
+		}
+	}
 
 	RegisteredHandle = RegisterViewModelWithMap(&System, this);
 }
@@ -1118,50 +1168,94 @@ const TArray<UNiagaraScriptVariable*> FNiagaraSystemToolkitParameterPanelViewMod
 	}
 	return EditableScriptVariables;
 }
+bool FNiagaraSystemToolkitParameterPanelViewModel::GetSectionEnabled(FText Section) const
+{
+	if (ActiveScriptIdx != -1 && Sections.IsValidIndex(ActiveScriptIdx) && Section.IdenticalTo(Sections[ActiveScriptIdx].DisplayName))
+	{
+		TSharedPtr<FNiagaraScratchPadScriptViewModel> ScratchDocumentVM = SystemViewModel->GetDocumentViewModel()->GetActiveScratchPadViewModelIfSet();
+		return ScratchDocumentVM.IsValid();
+	}
+
+	return true;
+}
+
+bool  FNiagaraSystemToolkitParameterPanelViewModel::ShouldRouteThroughScratchParameterMap(const FNiagaraParameterPanelCategory* Category, const FNiagaraVariableBase* NewVariable)
+{
+	bool bForceScript = false;
+	if (ActiveSectionIndex == ActiveScriptIdx && ActiveScriptIdx != -1)
+	{
+		bForceScript = true;
+	}
+	if (NewVariable != nullptr && bForceScript)
+	{
+		const FGuid NamespaceId = FNiagaraEditorUtilities::GetNamespaceMetaDataForVariableName(NewVariable->GetName()).GetGuid();
+		/*if (NamespaceId == FNiagaraEditorGuids::UserNamespaceMetaDataGuid)
+			bForceScript = false;
+		if (NamespaceId == FNiagaraEditorGuids::UserNamespaceMetaDataGuid)
+			bForceScript = false;
+		if (NamespaceId == FNiagaraEditorGuids::UserNamespaceMetaDataGuid)
+			bForceScript = false;*/
+	}
+
+	return bForceScript;
+}
 
-void FNiagaraSystemToolkitParameterPanelViewModel::AddParameter(FNiagaraVariable NewVariable, const FNiagaraParameterPanelCategory Category, const bool bRequestRename, const bool bMakeUniqueName) const
+void FNiagaraSystemToolkitParameterPanelViewModel::AddParameter(FNiagaraVariable NewVariable, const FNiagaraParameterPanelCategory Category, const bool bRequestRename, const bool bMakeUniqueName) 
 {
 	TGuardValue<bool> AddParameterRefreshGuard(bIsAddingParameter, true);
 	bool bSuccess = false;
 	UNiagaraSystem& System = SystemViewModel->GetSystem();
 
-	FScopedTransaction AddTransaction(LOCTEXT("AddSystemParameterTransaction", "Add parameter to system."));
-	System.Modify();
-	const FGuid NamespaceId = FNiagaraEditorUtilities::GetNamespaceMetaDataForVariableName(NewVariable.GetName()).GetGuid();
-	if (NamespaceId == FNiagaraEditorGuids::UserNamespaceMetaDataGuid)
+	bool bForceScript = ShouldRouteThroughScratchParameterMap(&Category, &NewVariable);
+
+	TSharedPtr<FNiagaraScratchPadScriptViewModel> ScratchDocumentVM = SystemViewModel->GetDocumentViewModel()->GetActiveScratchPadViewModelIfSet();
+	if (ScratchDocumentVM.IsValid() && bForceScript)
 	{
-		bSuccess = FNiagaraEditorUtilities::AddParameter(NewVariable, System.GetExposedParameters(), System, nullptr);
+		ScratchDocumentVM->GetParameterPanelViewModel()->AddParameter(NewVariable, Category, bRequestRename, bMakeUniqueName);
+		RefreshNextTick();
+		return;
 	}
-	else
+	
 	{
-		UNiagaraEditorParametersAdapter* EditorParametersAdapter = SystemViewModel->GetEditorOnlyParametersAdapter();
-		TArray<UNiagaraScriptVariable*>& EditorOnlyScriptVars = EditorParametersAdapter->GetParameters();
-		bool bNewScriptVarAlreadyExists = EditorOnlyScriptVars.ContainsByPredicate([&NewVariable](const UNiagaraScriptVariable* ScriptVar) { return ScriptVar->Variable == NewVariable; });
 
-		// unless the namespace prevents name changes we make sure the new parameter has a unique name  
-		if (bMakeUniqueName && !Category.NamespaceMetaData.Options.Contains(ENiagaraNamespaceMetadataOptions::PreventEditingName))
+		FScopedTransaction AddTransaction(LOCTEXT("AddSystemParameterTransaction", "Add parameter to system."));
+		System.Modify();
+		const FGuid NamespaceId = FNiagaraEditorUtilities::GetNamespaceMetaDataForVariableName(NewVariable.GetName()).GetGuid();
+		if (NamespaceId == FNiagaraEditorGuids::UserNamespaceMetaDataGuid)
 		{
-			TSet<FName> Names;
-			for (UNiagaraScriptVariable* ScriptVar : EditorOnlyScriptVars)
-			{
-				Names.Add(ScriptVar->Variable.GetName());
-			}
-			NewVariable.SetName(FNiagaraUtilities::GetUniqueName(NewVariable.GetName(), Names));
-			bNewScriptVarAlreadyExists = false;
+			bSuccess = FNiagaraEditorUtilities::AddParameter(NewVariable, System.GetExposedParameters(), System, nullptr);
 		}
-
-		if (bNewScriptVarAlreadyExists == false)
+		else
 		{
-			EditorParametersAdapter->Modify();
-			UNiagaraScriptVariable* NewScriptVar = NewObject<UNiagaraScriptVariable>(EditorParametersAdapter, FName(), RF_Transactional);
-			NewScriptVar->Init(NewVariable, FNiagaraVariableMetaData());
-			NewScriptVar->SetIsStaticSwitch(false);
-			NewScriptVar->SetIsSubscribedToParameterDefinitions(false);
-			EditorOnlyScriptVars.Add(NewScriptVar);
-			bSuccess = true;
+			UNiagaraEditorParametersAdapter* EditorParametersAdapter = SystemViewModel->GetEditorOnlyParametersAdapter();
+			TArray<UNiagaraScriptVariable*>& EditorOnlyScriptVars = EditorParametersAdapter->GetParameters();
+			bool bNewScriptVarAlreadyExists = EditorOnlyScriptVars.ContainsByPredicate([&NewVariable](const UNiagaraScriptVariable* ScriptVar) { return ScriptVar->Variable == NewVariable; });
 
-			// Check if the new parameter has the same name and type as an existing parameter definition, and if so, link to the definition automatically.
-			SubscribeParameterToLibraryIfMatchingDefinition(NewScriptVar, NewScriptVar->Variable.GetName());
+			// unless the namespace prevents name changes we make sure the new parameter has a unique name  
+			if (bMakeUniqueName && !Category.NamespaceMetaData.Options.Contains(ENiagaraNamespaceMetadataOptions::PreventEditingName))
+			{
+				TSet<FName> Names;
+				for (UNiagaraScriptVariable* ScriptVar : EditorOnlyScriptVars)
+				{
+					Names.Add(ScriptVar->Variable.GetName());
+				}
+				NewVariable.SetName(FNiagaraUtilities::GetUniqueName(NewVariable.GetName(), Names));
+				bNewScriptVarAlreadyExists = false;
+			}
+
+			if (bNewScriptVarAlreadyExists == false)
+			{
+				EditorParametersAdapter->Modify();
+				UNiagaraScriptVariable* NewScriptVar = NewObject<UNiagaraScriptVariable>(EditorParametersAdapter, FName(), RF_Transactional);
+				NewScriptVar->Init(NewVariable, FNiagaraVariableMetaData());
+				NewScriptVar->SetIsStaticSwitch(false);
+				NewScriptVar->SetIsSubscribedToParameterDefinitions(false);
+				EditorOnlyScriptVars.Add(NewScriptVar);
+				bSuccess = true;
+
+				// Check if the new parameter has the same name and type as an existing parameter definition, and if so, link to the definition automatically.
+				SubscribeParameterToLibraryIfMatchingDefinition(NewScriptVar, NewScriptVar->Variable.GetName());
+			}
 		}
 	}
 
@@ -1172,12 +1266,19 @@ void FNiagaraSystemToolkitParameterPanelViewModel::AddParameter(FNiagaraVariable
 	}
 }
 
-void FNiagaraSystemToolkitParameterPanelViewModel::FindOrAddParameter(FNiagaraVariable Variable, const FNiagaraParameterPanelCategory Category) const
+void FNiagaraSystemToolkitParameterPanelViewModel::FindOrAddParameter(FNiagaraVariable Variable, const FNiagaraParameterPanelCategory Category) 
 {
 	TGuardValue<bool> AddParameterRefreshGuard(bIsAddingParameter, true);
 	bool bSuccess = false;
 	UNiagaraSystem& System = SystemViewModel->GetSystem();
 
+	TSharedPtr<FNiagaraScratchPadScriptViewModel> ScratchDocumentVM = SystemViewModel->GetDocumentViewModel()->GetActiveScratchPadViewModelIfSet();
+	if (ScratchDocumentVM.IsValid())
+	{
+		ScratchDocumentVM->GetParameterPanelViewModel()->FindOrAddParameter(Variable, Category);
+		return;
+	}
+
 	FScopedTransaction AddTransaction(LOCTEXT("FindOrAddAddSystemParameterTransaction", "Add parameter to script."));
 	const FGuid NamespaceId = FNiagaraEditorUtilities::GetNamespaceMetaDataForVariableName(Variable.GetName()).GetGuid();
 	if (NamespaceId == FNiagaraEditorGuids::UserNamespaceMetaDataGuid)
@@ -1216,36 +1317,100 @@ bool FNiagaraSystemToolkitParameterPanelViewModel::GetShowReferenceCounter() con
 	return (ActiveSectionIndex != CommonOnlyIdx);
 }
 
+
+void FNiagaraSystemToolkitParameterPanelViewModel::SetActiveSection(int32 InSection)
+{
+	INiagaraParameterPanelViewModel::SetActiveSection(InSection);
+
+	// We want to cache this off so that we can remember what you were last on when the primary document was active.
+	if (SystemViewModel->GetDocumentViewModel()->IsPrimaryDocumentActive())
+	{
+		LastActiveSystemSectionIdx = ActiveSectionIndex;
+	}
+	else
+	{
+		TSharedPtr<FNiagaraScratchPadScriptViewModel> ScratchDocumentVM = SystemViewModel->GetDocumentViewModel()->GetActiveScratchPadViewModelIfSet();
+		if (ScratchDocumentVM.IsValid())
+		{
+			TWeakPtr<INiagaraParameterPanelViewModel> WeakVM = StaticCastSharedRef< INiagaraParameterPanelViewModel>(AsShared());
+			ScratchDocumentVM->GetParameterPanelViewModel()->SetMasterParameterPanelViewModel(WeakVM);
+		}
+	}
+}
+
 bool FNiagaraSystemToolkitParameterPanelViewModel::GetCanAddParametersToCategory(FNiagaraParameterPanelCategory Category) const
 {
 	if (ActiveSectionIndex == CommonOnlyIdx)
+	{
 		return false;
+	}
+
+	// If in emitter edit mode, don't support anything other than Emitter level and below
+	if (SystemViewModel && SystemViewModel->GetEditMode() == ENiagaraSystemViewModelEditMode::EmitterAsset)
+	{
+		if (Category.NamespaceMetaData.GetGuid() == FNiagaraEditorGuids::SystemNamespaceMetaDataGuid || 
+			Category.NamespaceMetaData.GetGuid() == FNiagaraEditorGuids::UserNamespaceMetaDataGuid)
+		{
+			return false;
+		}
+	}
 	return GetEditableGraphsConst().Num() > 0 && Category.NamespaceMetaData.GetGuid() != FNiagaraEditorGuids::StaticSwitchNamespaceMetaDataGuid;
 }
 
-void FNiagaraSystemToolkitParameterPanelViewModel::DeleteParameters(const TArray<FNiagaraParameterPanelItem>& ItemsToDelete) const
+void FNiagaraSystemToolkitParameterPanelViewModel::DeleteParameters(const TArray<FNiagaraParameterPanelItem>& ItemsToDelete) 
 {
 	bool bAnyChange = false;
 	FScopedTransaction RemoveParameterTransaction(LOCTEXT("RemoveParameter", "Removed Parameter(s)"));
 
-	for(const FNiagaraParameterPanelItem& ItemToDelete : ItemsToDelete)
+	// Handle deleting parameters from the active script differently than from the system overview being active.
+	// This will allow us to behave as users expect.
+	TSharedPtr<FNiagaraScratchPadScriptViewModel> ScratchDocumentVM = SystemViewModel->GetDocumentViewModel()->GetActiveScratchPadViewModelIfSet();
+	if (ScratchDocumentVM.IsValid())
 	{
-		if (ItemToDelete.bExternallyReferenced)
+		for (const FNiagaraParameterPanelItem& ItemToDelete : ItemsToDelete)
 		{
-			continue;
+			if (ItemToDelete.bExternallyReferenced)
+			{
+				continue;
+			}
+
+			for (UNiagaraGraph* Graph : GetEditableGraphs())
+			{
+				UNiagaraScriptVariable* SelectedScriptVariable = Graph->GetScriptVariable(ItemToDelete.GetVariable());
+				if (SelectedScriptVariable && VariableObjectSelection->GetSelectedObjects().Contains(SelectedScriptVariable))
+				{
+					VariableObjectSelection->ClearSelectedObjects();
+				}
+				Graph->RemoveParameter(ItemToDelete.GetVariable());
+
+				if (!bAnyChange)
+				{
+					bAnyChange = true;
+				}
+			}
 		}
-		
-		UNiagaraSystem& System = SystemViewModel->GetSystem();
-		const FGuid& ScriptVarId = ItemToDelete.GetVariableMetaData().GetVariableGuid();
-		System.Modify();
-		System.GetExposedParameters().RemoveParameter(ItemToDelete.GetVariable());
-		UNiagaraEditorParametersAdapter* EditorParametersAdapter = SystemViewModel->GetEditorOnlyParametersAdapter();
-		EditorParametersAdapter->Modify();
-		EditorParametersAdapter->GetParameters().RemoveAll([&ScriptVarId](const UNiagaraScriptVariable* ScriptVar) { return ScriptVar->Metadata.GetVariableGuid() == ScriptVarId; });
+	}
+	else
+	{
+		for (const FNiagaraParameterPanelItem& ItemToDelete : ItemsToDelete)
+		{
+			if (ItemToDelete.bExternallyReferenced)
+			{
+				continue;
+			}
 
-		// Update anything that was referencing that parameter
-		System.HandleVariableRemoved(ItemToDelete.GetVariable(), true);
-		bAnyChange = true;
+			UNiagaraSystem& System = SystemViewModel->GetSystem();
+			const FGuid& ScriptVarId = ItemToDelete.GetVariableMetaData().GetVariableGuid();
+			System.Modify();
+			System.GetExposedParameters().RemoveParameter(ItemToDelete.GetVariable());
+			UNiagaraEditorParametersAdapter* EditorParametersAdapter = SystemViewModel->GetEditorOnlyParametersAdapter();
+			EditorParametersAdapter->Modify();
+			EditorParametersAdapter->GetParameters().RemoveAll([&ScriptVarId](const UNiagaraScriptVariable* ScriptVar) { return ScriptVar->Metadata.GetVariableGuid() == ScriptVarId; });
+
+			// Update anything that was referencing that parameter
+			System.HandleVariableRemoved(ItemToDelete.GetVariable(), true);
+			bAnyChange = true;
+		}
 	}
 
 	if(bAnyChange)
@@ -1258,7 +1423,7 @@ void FNiagaraSystemToolkitParameterPanelViewModel::DeleteParameters(const TArray
 	}
 }
 
-void FNiagaraSystemToolkitParameterPanelViewModel::RenameParameter(const FNiagaraParameterPanelItem& ItemToRename, const FName NewName) const
+void FNiagaraSystemToolkitParameterPanelViewModel::RenameParameter(const FNiagaraParameterPanelItem& ItemToRename, const FName NewName) 
 {
 	if (ensureMsgf(ItemToRename.bExternallyReferenced == false, TEXT("Can not modify an externally referenced parameter.")) == false)
 	{
@@ -1273,6 +1438,15 @@ void FNiagaraSystemToolkitParameterPanelViewModel::RenameParameter(const FNiagar
 	const FNiagaraVariable Parameter = ItemToRename.GetVariable();
 	const FGuid& ScriptVarId = ItemToRename.GetVariableMetaData().GetVariableGuid();
 	UNiagaraSystem& System = SystemViewModel->GetSystem();
+
+	// Forward renaming of script variables for the active script to the working parameter panel view model for just that script.
+	TSharedPtr<FNiagaraScratchPadScriptViewModel> ScratchDocumentVM = SystemViewModel->GetDocumentViewModel()->GetActiveScratchPadViewModelIfSet();
+	if (ScratchDocumentVM.IsValid())
+	{
+		ScratchDocumentVM->GetParameterPanelViewModel()->RenameParameter(ItemToRename, NewName);
+		return;
+	}
+
 	System.Modify();
 
 	bool bExposedParametersRename = false;
@@ -1361,7 +1535,7 @@ void FNiagaraSystemToolkitParameterPanelViewModel::RenameParameter(const FNiagar
 	}
 }
 
-void FNiagaraSystemToolkitParameterPanelViewModel::SetParameterIsSubscribedToLibrary(const UNiagaraScriptVariable* ScriptVarToModify, const bool bSubscribed) const
+void FNiagaraSystemToolkitParameterPanelViewModel::SetParameterIsSubscribedToLibrary(const UNiagaraScriptVariable* ScriptVarToModify, const bool bSubscribed) 
 {
 	const FText TransactionText = bSubscribed ? LOCTEXT("SubscribeParameter", "Subscribe parameter") : LOCTEXT("UnsubscribeParameter", "Unsubscribe parameter");
 	FScopedTransaction SubscribeTransaction(TransactionText);
@@ -1505,14 +1679,17 @@ void FNiagaraSystemToolkitParameterPanelViewModel::DebugParameters(const TArray<
 const TArray<FNiagaraParameterPanelCategory>& FNiagaraSystemToolkitParameterPanelViewModel::GetDefaultCategories() const
 {
 	// Show advanced in Active/Common case. Only show User when in user mode.
-	bool bForceAll = false;
+	bool bForceScript = false;
+	bool bForceSystem = false;
 	bool bUserOnly = false;
 	bool bCommonOnly = false;
 	if (ActiveSectionIndex == UserOnlyIdx && UserOnlyIdx != -1)
 		bUserOnly = true;
-	else if (ActiveSectionIndex == ActiveIdx && ActiveIdx != -1)
-		bForceAll = true;
-	else if (ActiveSectionIndex == CommonOnlyIdx && ActiveIdx != -1)
+	else if (ActiveSectionIndex == ActiveScriptIdx && ActiveScriptIdx != -1)
+		bForceScript= true;
+	else if (ActiveSectionIndex == ActiveSystemIdx && ActiveSystemIdx != -1)
+		bForceSystem = true;
+	else if (ActiveSectionIndex == CommonOnlyIdx && CommonOnlyIdx != -1)
 		bCommonOnly = true;
 
 
@@ -1558,8 +1735,15 @@ const TArray<FNiagaraParameterPanelCategory>& FNiagaraSystemToolkitParameterPane
 		}
 	}
 	else
-	{
-		CachedCurrentCategories = FNiagaraSystemToolkitParameterPanelViewModel::DefaultAdvancedCategories;
+	{	
+		if (bForceScript)
+		{
+			CachedCurrentCategories = FNiagaraSystemToolkitParameterPanelViewModel::DefaultAdvancedScriptCategories;
+		}
+		else
+		{
+			CachedCurrentCategories = FNiagaraSystemToolkitParameterPanelViewModel::DefaultAdvancedCategories;
+		}
 	}
 
 	
@@ -1567,7 +1751,7 @@ const TArray<FNiagaraParameterPanelCategory>& FNiagaraSystemToolkitParameterPane
 	return CachedCurrentCategories;
 }
 
-FMenuAndSearchBoxWidgets FNiagaraSystemToolkitParameterPanelViewModel::GetParameterMenu(FNiagaraParameterPanelCategory Category) const
+FMenuAndSearchBoxWidgets FNiagaraSystemToolkitParameterPanelViewModel::GetParameterMenu(FNiagaraParameterPanelCategory Category) 
 {
 	const bool bRequestRename = true;
 	const bool bSkipSubscribedLibraries = false;
@@ -1602,7 +1786,7 @@ FMenuAndSearchBoxWidgets FNiagaraSystemToolkitParameterPanelViewModel::GetParame
 	return MenuAndSearchBoxWidgets;
 }
 
-FReply FNiagaraSystemToolkitParameterPanelViewModel::HandleDragDropOperation(TSharedPtr<FDragDropOperation> DragDropOperation) const
+FReply FNiagaraSystemToolkitParameterPanelViewModel::HandleDragDropOperation(TSharedPtr<FDragDropOperation> DragDropOperation) 
 {
 	TSharedPtr<FNiagaraParameterGraphDragOperation> ParameterGraphDragDropOperation = StaticCastSharedPtr<FNiagaraParameterGraphDragOperation>(DragDropOperation);
 	if (ParameterGraphDragDropOperation.IsValid() == false)
@@ -1662,6 +1846,22 @@ bool FNiagaraSystemToolkitParameterPanelViewModel::GetCanHandleDragDropOperation
 	return true;
 }
 
+
+
+void FNiagaraSystemToolkitParameterPanelViewModel::RefreshDueToActiveDocumentChanged()
+{
+	// We want to recall the section you were last in when you swapped between Primary and Scratch documents, which
+	// will invoke a full parameter refresh.
+	if (SystemViewModel->GetDocumentViewModel()->IsPrimaryDocumentActive())
+	{
+		SetActiveSection(LastActiveSystemSectionIdx);
+	}
+	else
+	{
+		SetActiveSection(ActiveScriptIdx);
+	}
+}
+
 bool FNiagaraSystemToolkitParameterPanelViewModel::GetCanSetParameterNamespaceAndToolTip(const FNiagaraParameterPanelItem& ItemToModify, const FName NewNamespace, FText& OutCanSetParameterNamespaceToolTip) const
 {
 	return FNiagaraParameterPanelUtilities::GetCanSetParameterNamespaceAndToolTipForScriptOrSystem(ItemToModify, NewNamespace, OutCanSetParameterNamespaceToolTip);
@@ -1677,7 +1877,7 @@ bool FNiagaraSystemToolkitParameterPanelViewModel::GetCanSetParameterCustomNames
 	return FNiagaraParameterPanelUtilities::GetCanSetParameterCustomNamespaceModifierAndToolTipForScriptOrSystem(ItemToModify, bDuplicateParameter, OutCanSetParameterNamespaceModifierToolTip);
 }
 
-TSharedRef<SWidget> FNiagaraSystemToolkitParameterPanelViewModel::CreateAddParameterMenuForAssignmentNode(UNiagaraNodeAssignment* AssignmentNode, const TSharedPtr<SComboButton>& AddButton) const
+TSharedRef<SWidget> FNiagaraSystemToolkitParameterPanelViewModel::CreateAddParameterMenuForAssignmentNode(UNiagaraNodeAssignment* AssignmentNode, const TSharedPtr<SComboButton>& AddButton) 
 {
 	auto AddParameterLambda = [this, AssignmentNode](FNiagaraVariable& NewParameter) {
 		// If an assignment target is already setting the associated parameter, make a unique name here so that the assignment target is valid.
@@ -1864,24 +2064,94 @@ const TArray<UNiagaraParameterDefinitions*> FNiagaraSystemToolkitParameterPanelV
 	return SystemViewModel->GetAvailableParameterDefinitions(bSkipSubscribedParameterDefinitions);
 }
 
+void FNiagaraSystemToolkitParameterPanelViewModel::PreSectionChange(const TArray<FNiagaraParameterPanelCategory>& ExpandedItems)
+{
+	// Before we go to a different setup, cache the existing expanded states.
+	if (Sections.IsValidIndex(ActiveSectionIndex))
+	{
+		UNiagaraEditorSettings* Settings = GetMutableDefault<UNiagaraEditorSettings>();
+		bool bAdded = false;
+		FNiagaraParameterPanelSectionStorage& Storage = Settings->FindOrAddParameterPanelSectionStorage(Sections[ActiveSectionIndex].SectionId, bAdded);
+
+		TArray<FGuid> ExpandedCategories;
+
+		for (const FNiagaraParameterPanelCategory& Item : ExpandedItems)
+		{
+			if (Item.NamespaceMetaData.IsValid() && Item.NamespaceMetaData.GetGuid().IsValid())
+			{
+				ExpandedCategories.AddUnique(Item.NamespaceMetaData.GetGuid());
+			}
+		}
+		Storage.ExpandedCategories = ExpandedCategories;
+		Settings->SaveConfig();
+	}
+}
+
+bool FNiagaraSystemToolkitParameterPanelViewModel::IsCategoryExpandedByDefault(const FNiagaraParameterPanelCategory& Category) const 
+{
+	// Pull data from the settings for categories being expanded or collapsed
+	if (Sections.IsValidIndex(ActiveSectionIndex))
+	{
+		UNiagaraEditorSettings* Settings = GetMutableDefault<UNiagaraEditorSettings>();
+		bool bAdded = false;
+		FNiagaraParameterPanelSectionStorage& Storage = Settings->FindOrAddParameterPanelSectionStorage(Sections[ActiveSectionIndex].SectionId, bAdded);
+		if (Storage.ExpandedCategories.Contains(Category.NamespaceMetaData.GetGuid()))
+		{
+			return true;
+		}
+	}
+	return false;
+}
+
 const TArray<UNiagaraGraph*> FNiagaraSystemToolkitParameterPanelViewModel::GetAllGraphsConst() const
 {
-	return FNiagaraSystemToolkitParameterPanelUtilities::GetAllGraphs(SystemViewModel);
+	// If an active script, just get the graphs associated with that. Otherwise for system overview, get all graphs.
+	if (ActiveSectionIndex == ActiveScriptIdx && ActiveScriptIdx != -1)
+	{
+		return FNiagaraSystemToolkitParameterPanelUtilities::GetAllGraphs(SystemViewModel, true);
+	}
+	else
+	{
+		return FNiagaraSystemToolkitParameterPanelUtilities::GetAllGraphs(SystemViewModel, false);
+	}
 }
 
 TArray<UNiagaraGraph*> FNiagaraSystemToolkitParameterPanelViewModel::GetEditableGraphs() const
 {
-	return FNiagaraSystemToolkitParameterPanelUtilities::GetEditableGraphs(SystemViewModel, SystemGraphSelectionViewModelWeak);
+	// If an active script, just get the graphs associated with that. Otherwise for system overview, get all editable graphs.
+	if (ActiveSectionIndex == ActiveScriptIdx && ActiveScriptIdx != -1)
+	{
+		return FNiagaraSystemToolkitParameterPanelUtilities::GetEditableGraphs(SystemViewModel, SystemGraphSelectionViewModelWeak, true);
+	}
+	else
+	{
+		return FNiagaraSystemToolkitParameterPanelUtilities::GetEditableGraphs(SystemViewModel, SystemGraphSelectionViewModelWeak, false);
+	}
+}
+void INiagaraParameterPanelViewModel::SetActiveSection(int32 InSection)
+{
+	int32 OldActiveSectionIndex = ActiveSectionIndex;
+
+	if (Sections.IsValidIndex(InSection))
+	{
+		ActiveSectionIndex = InSection;
+	}
+
+	// Need to update UI if sections change
+	if (OldActiveSectionIndex != ActiveSectionIndex)
+	{
+		RefreshFullNextTick(true);
+	}
 }
 
 void INiagaraParameterPanelViewModel::SetActiveSection(FText & InSection)
 {
-	ActiveSectionIndex = -1;
 	for (int32 Idx = 0; Idx <Sections.Num(); Idx++)
 	{
-		if (Sections[Idx].EqualTo(InSection))
+		if (Sections[Idx].DisplayName.EqualTo(InSection))
 		{
-			ActiveSectionIndex = Idx;
+			SetActiveSection(Idx);
+			return;
 		}
 	}
 }
@@ -1890,23 +2160,26 @@ FText INiagaraParameterPanelViewModel::GetTooltipForSection(FText& InSection) co
 {
 	for (int32 i = 0; i < Sections.Num(); i++)
 	{
-		if (Sections[i].EqualTo(InSection))
-			return SectionDescriptions[i];
+		if (Sections[i].DisplayName.EqualTo(InSection))
+			return Sections[i].Description;
 	}
 
 	return FText::GetEmpty();
 }
 TArray<FNiagaraParameterPanelItem> FNiagaraSystemToolkitParameterPanelViewModel::GetViewedParameterItems() const
 {
-	bool bForceAll = false;
+	bool bForceScript = false;
+	bool bForceSystem = false;
 	bool bUserOnly = false;
 	bool bCommonOnly = false;
 	if (ActiveSectionIndex == UserOnlyIdx && UserOnlyIdx != -1)
 		bUserOnly = true;
+	else if (ActiveSectionIndex == ActiveScriptIdx && ActiveScriptIdx != -1)
+		bForceScript = true;
+	else if (ActiveSectionIndex == ActiveSystemIdx && ActiveSystemIdx != -1)
+		bForceSystem = true;
 	else if (ActiveSectionIndex == CommonOnlyIdx && CommonOnlyIdx != -1)
 		bCommonOnly = true;
-	else if (ActiveSectionIndex == ActiveIdx && ActiveIdx != -1)
-		bForceAll = true;
 
 	// On the first time opening the parameter panel view model we are not guaranteed to call GetDefaultCategories() before GetViewedParameterItems(). 
 	// We require CachedCurrentCategories being set as this is used to filter out parameter items that are being viewed. If CachedCurrentCategories 
@@ -1964,7 +2237,7 @@ TArray<FNiagaraParameterPanelItem> FNiagaraSystemToolkitParameterPanelViewModel:
 	};
 
 	// Collect user parameters from system.
-	if (bUserOnly || bForceAll)
+	if (bUserOnly || bForceSystem)
 	{
 		CollectParamStore(&SystemViewModel->GetSystem().GetExposedParameters());
 	}
@@ -1995,7 +2268,7 @@ TArray<FNiagaraParameterPanelItem> FNiagaraSystemToolkitParameterPanelViewModel:
 		}
 	}
 
-	if (bForceAll)
+	if (bForceSystem)
 	{
 		// Collect parameters added to the system asset.
 		for (UNiagaraScriptVariable* EditorOnlyScriptVar : SystemViewModel->GetEditorOnlyParametersAdapter()->GetParameters())
@@ -2015,7 +2288,10 @@ TArray<FNiagaraParameterPanelItem> FNiagaraSystemToolkitParameterPanelViewModel:
 
 			VisitedParameterToItemMap.Add(Var, Item);
 		}
+	}
 
+	if (bForceSystem)
+	{
 		// Collect parameters for all emitters.
 		TArray<FNiagaraVariable> VisitedInvalidParameters;
 
@@ -2139,14 +2415,18 @@ TArray<FNiagaraParameterPanelItem> FNiagaraSystemToolkitParameterPanelViewModel:
 						}
 
 						bool bVarOnlyInTopLevelGraph = true;
-						for (FModuleScopedPin& WritePin : Builder.Histories[0].PerVariableWriteHistory[VariableIndex])
+						if (!bForceScript)
 						{
-							UEdGraphNode* VariableOwningNode = WritePin.Pin->GetOwningNode();
-							bVarOnlyInTopLevelGraph &= AllGraphs.Contains(static_cast<const UNiagaraGraph*>(VariableOwningNode->GetGraph()));
-						}
-						for (FNiagaraParameterMapHistory::FReadHistory& ReadPins : Builder.Histories[0].PerVariableReadHistory[VariableIndex])
-						{
-							bVarOnlyInTopLevelGraph &= !AllGraphs.Contains(static_cast<const UNiagaraGraph*>(ReadPins.ReadPin.Pin->GetOwningNode()->GetGraph()));
+							for (FModuleScopedPin& WritePin : Builder.Histories[0].PerVariableWriteHistory[VariableIndex])
+							{
+								UEdGraphNode* VariableOwningNode = WritePin.Pin->GetOwningNode();
+								bVarOnlyInTopLevelGraph &= AllGraphs.Contains(static_cast<const UNiagaraGraph*>(VariableOwningNode->GetGraph()));
+							}
+							for (FNiagaraParameterMapHistory::FReadHistory& ReadPins : Builder.Histories[0].PerVariableReadHistory[VariableIndex])
+							{
+								UEdGraphNode* VariableOwningNode = ReadPins.ReadPin.Pin->GetOwningNode();
+								bVarOnlyInTopLevelGraph &= !AllGraphs.Contains(static_cast<const UNiagaraGraph*>(VariableOwningNode->GetGraph()));
+							}
 						}
 
 						FNiagaraParameterPanelItem Item = FNiagaraParameterPanelItem();
@@ -2177,6 +2457,19 @@ TArray<FNiagaraParameterPanelItem> FNiagaraSystemToolkitParameterPanelViewModel:
 			}
 		}
 	}
+	else if (bForceScript)
+	{
+		TSharedPtr<class FNiagaraScratchPadScriptViewModel> ScratchPadScriptVM = SystemViewModel->GetDocumentViewModel()->GetActiveScratchPadViewModelIfSet();
+		if (ScratchPadScriptVM.IsValid())
+		{
+			TSharedPtr<INiagaraParameterPanelViewModel> ParamPanelForScriptVM = ScratchPadScriptVM->GetParameterPanelViewModel();
+			if (ParamPanelForScriptVM.IsValid())
+			{
+				CachedViewedItems = ParamPanelForScriptVM->GetViewedParameterItems();
+				return CachedViewedItems;
+			}
+		}
+	}
 
 	// Refresh the CachedViewedItems and return that as the latest array of viewed items.
 	VisitedParameterToItemMap.GenerateValueArray(CachedViewedItems);
@@ -2213,7 +2506,7 @@ TArray<FNiagaraEmitterHandle*> FNiagaraSystemToolkitParameterPanelViewModel::Get
 	return EditableEmitterHandles;
 }
 
-void FNiagaraSystemToolkitParameterPanelViewModel::AddScriptVariable(const UNiagaraScriptVariable* NewScriptVar) const
+void FNiagaraSystemToolkitParameterPanelViewModel::AddScriptVariable(const UNiagaraScriptVariable* NewScriptVar) 
 {
 	if (NewScriptVar == nullptr)
 	{
@@ -2544,7 +2837,7 @@ void FNiagaraScriptToolkitParameterPanelViewModel::Init(const FScriptToolkitUICo
 	}
 }
 
-void FNiagaraScriptToolkitParameterPanelViewModel::AddParameter(FNiagaraVariable NewVariable, const FNiagaraParameterPanelCategory Category, const bool bRequestRename, const bool bMakeUniqueName) const
+void FNiagaraScriptToolkitParameterPanelViewModel::AddParameter(FNiagaraVariable NewVariable, const FNiagaraParameterPanelCategory Category, const bool bRequestRename, const bool bMakeUniqueName) 
 {
 	TGuardValue<bool> AddParameterRefreshGuard(bIsAddingParameter, true);
 	bool bSuccess = false;
@@ -2581,7 +2874,7 @@ void FNiagaraScriptToolkitParameterPanelViewModel::AddParameter(FNiagaraVariable
 	}
 }
 
-void FNiagaraScriptToolkitParameterPanelViewModel::FindOrAddParameter(FNiagaraVariable Variable, const FNiagaraParameterPanelCategory Category) const
+void FNiagaraScriptToolkitParameterPanelViewModel::FindOrAddParameter(FNiagaraVariable Variable, const FNiagaraParameterPanelCategory Category) 
 {
 	TGuardValue<bool> AddParameterRefreshGuard(bIsAddingParameter, true);
 	bool bSuccess = false;
@@ -2613,7 +2906,7 @@ bool FNiagaraScriptToolkitParameterPanelViewModel::GetCanAddParametersToCategory
 	return GetEditableGraphsConst().Num() > 0 && Category.NamespaceMetaData.GetGuid() != FNiagaraEditorGuids::StaticSwitchNamespaceMetaDataGuid;
 }
 
-void FNiagaraScriptToolkitParameterPanelViewModel::DeleteParameters(const TArray<FNiagaraParameterPanelItem>& ItemsToDelete) const
+void FNiagaraScriptToolkitParameterPanelViewModel::DeleteParameters(const TArray<FNiagaraParameterPanelItem>& ItemsToDelete) 
 {
 	bool bAnyChange = false;
 	FScopedTransaction RemoveParametersWithPins(LOCTEXT("RemoveParametersWithPins", "Removed parameter(s) and referenced pins"));
@@ -2651,7 +2944,7 @@ void FNiagaraScriptToolkitParameterPanelViewModel::DeleteParameters(const TArray
 	}
 }
 
-void FNiagaraScriptToolkitParameterPanelViewModel::RenameParameter(const FNiagaraParameterPanelItem& ItemToRename, const FName NewName) const
+void FNiagaraScriptToolkitParameterPanelViewModel::RenameParameter(const FNiagaraParameterPanelItem& ItemToRename, const FName NewName) 
 {
 	if (ensureMsgf(ItemToRename.bExternallyReferenced == false, TEXT("Can not modify an externally referenced parameter.")) == false)
 	{
@@ -2665,7 +2958,7 @@ void FNiagaraScriptToolkitParameterPanelViewModel::RenameParameter(const FNiagar
 	RenameParameter(ItemToRename.ScriptVariable, NewName);
 }
 
-void FNiagaraScriptToolkitParameterPanelViewModel::ChangeParameterType(const TArray<FNiagaraParameterPanelItem> ItemsToModify, const FNiagaraTypeDefinition NewType) const
+void FNiagaraScriptToolkitParameterPanelViewModel::ChangeParameterType(const TArray<FNiagaraParameterPanelItem> ItemsToModify, const FNiagaraTypeDefinition NewType) 
 {
 	FScopedTransaction Transaction(LOCTEXT("ChangeParameterTypeTransaction", "Change parameter type"));
 
@@ -2688,7 +2981,7 @@ void FNiagaraScriptToolkitParameterPanelViewModel::ChangeParameterType(const TAr
 	Refresh();
 }
 
-void FNiagaraScriptToolkitParameterPanelViewModel::RenameParameter(const UNiagaraScriptVariable* ScriptVarToRename, const FName NewName) const
+void FNiagaraScriptToolkitParameterPanelViewModel::RenameParameter(const UNiagaraScriptVariable* ScriptVarToRename, const FName NewName) 
 {
 	FScopedTransaction RenameTransaction(LOCTEXT("RenameParameterTransaction", "Rename parameter"));
 	ScriptViewModel->GetStandaloneScript().Script->Modify();
@@ -2696,12 +2989,8 @@ void FNiagaraScriptToolkitParameterPanelViewModel::RenameParameter(const UNiagar
 	const FNiagaraVariable& Parameter = ScriptVarToRename->Variable;
 	for (UNiagaraGraph* Graph : GetEditableGraphs())
 	{
-		const FNiagaraGraphParameterReferenceCollection* ReferenceCollection = Graph->GetParameterReferenceMap().Find(Parameter);
-		if (ensureMsgf(ReferenceCollection != nullptr, TEXT("Parameter in view which wasn't in the reference collection.")) == false)
-		{
-			// Can't handle parameters with no reference collections.
-			continue;
-		}
+		// Removed some prior code here around parameter usage collections and earlying out as it was causing "unreferenced" parameters to be
+		// unabled to have their namespace changed for no good reason.
 		Graph->Modify();
 		Graph->RenameParameter(Parameter, NewName);
 		bSuccess = true;
@@ -2718,7 +3007,7 @@ void FNiagaraScriptToolkitParameterPanelViewModel::RenameParameter(const UNiagar
 	}
 }
 
-void FNiagaraScriptToolkitParameterPanelViewModel::RenameParameter(const FNiagaraVariable& VariableToRename, const FName NewName) const
+void FNiagaraScriptToolkitParameterPanelViewModel::RenameParameter(const FNiagaraVariable& VariableToRename, const FName NewName) 
 {
 	for (UNiagaraGraph* Graph : GetEditableGraphs())
 	{
@@ -2730,7 +3019,7 @@ void FNiagaraScriptToolkitParameterPanelViewModel::RenameParameter(const FNiagar
 	}
 }
 
-void FNiagaraScriptToolkitParameterPanelViewModel::DuplicateParameters(const TArray<FNiagaraParameterPanelItem> ItemsToDuplicate) const
+void FNiagaraScriptToolkitParameterPanelViewModel::DuplicateParameters(const TArray<FNiagaraParameterPanelItem> ItemsToDuplicate) 
 {
 	FScopedTransaction Transaction(LOCTEXT("DuplicateParameterTransaction", "Duplicate parameter"));
 	TGuardValue<bool> AddParameterRefreshGuard(bIsAddingParameter, true);
@@ -2781,7 +3070,7 @@ void FNiagaraScriptToolkitParameterPanelViewModel::DuplicateParameters(const TAr
 	}
 }
 
-void FNiagaraScriptToolkitParameterPanelViewModel::SetParameterIsSubscribedToLibrary(const UNiagaraScriptVariable* ScriptVarToModify, const bool bSubscribed) const
+void FNiagaraScriptToolkitParameterPanelViewModel::SetParameterIsSubscribedToLibrary(const UNiagaraScriptVariable* ScriptVarToModify, const bool bSubscribed) 
 {
 	const FText TransactionText = bSubscribed ? LOCTEXT("SubscribeParameter", "Subscribe parameter") : LOCTEXT("UnsubscribeParameter", "Unsubscribe parameter");
 	FScopedTransaction SubscribeTransaction(TransactionText);
@@ -2975,7 +3264,7 @@ const TArray<FNiagaraParameterPanelCategory>& FNiagaraScriptToolkitParameterPane
 	return CachedCurrentCategories;
 }
 
-FMenuAndSearchBoxWidgets FNiagaraScriptToolkitParameterPanelViewModel::GetParameterMenu(FNiagaraParameterPanelCategory Category) const
+FMenuAndSearchBoxWidgets FNiagaraScriptToolkitParameterPanelViewModel::GetParameterMenu(FNiagaraParameterPanelCategory Category) 
 {
 	const bool bRequestRename = true;
 	const bool bSkipSubscribedLibraries = false;
@@ -3004,7 +3293,7 @@ FMenuAndSearchBoxWidgets FNiagaraScriptToolkitParameterPanelViewModel::GetParame
 	return MenuAndSearchBoxWidgets;
 }
 
-FReply FNiagaraScriptToolkitParameterPanelViewModel::HandleDragDropOperation(TSharedPtr<FDragDropOperation> DragDropOperation) const
+FReply FNiagaraScriptToolkitParameterPanelViewModel::HandleDragDropOperation(TSharedPtr<FDragDropOperation> DragDropOperation) 
 {
 	TSharedPtr<FNiagaraParameterGraphDragOperation> ParameterGraphDragDropOperation = StaticCastSharedPtr<FNiagaraParameterGraphDragOperation>(DragDropOperation);
 	if (ParameterGraphDragDropOperation.IsValid() == false)
@@ -3079,7 +3368,7 @@ bool FNiagaraScriptToolkitParameterPanelViewModel::GetCanSetParameterCustomNames
 	return FNiagaraParameterPanelUtilities::GetCanSetParameterCustomNamespaceModifierAndToolTipForScriptOrSystem(ItemToModify, bDuplicateParameter, OutCanSetParameterNamespaceModifierToolTip);
 }
 
-void FNiagaraScriptToolkitParameterPanelViewModel::SetParameterIsOverridingLibraryDefaultValue(const FNiagaraParameterPanelItem ItemToModify, const bool bOverriding) const
+void FNiagaraScriptToolkitParameterPanelViewModel::SetParameterIsOverridingLibraryDefaultValue(const FNiagaraParameterPanelItem ItemToModify, const bool bOverriding) 
 {
 	if (ensureMsgf(ItemToModify.bExternallyReferenced == false, TEXT("Cannot modify an externally referenced parameter.")) == false)
 	{
@@ -3244,7 +3533,7 @@ TArray<FNiagaraParameterPanelItem> FNiagaraScriptToolkitParameterPanelViewModel:
 	return CachedViewedItems;
 }
 
-void FNiagaraScriptToolkitParameterPanelViewModel::AddScriptVariable(const UNiagaraScriptVariable* NewScriptVar) const
+void FNiagaraScriptToolkitParameterPanelViewModel::AddScriptVariable(const UNiagaraScriptVariable* NewScriptVar) 
 {
 	bool bSuccess = false;
 	FScopedTransaction AddTransaction(LOCTEXT("AddScriptParameterTransaction", "Add parameter to script."));
@@ -3323,10 +3612,9 @@ void FNiagaraParameterDefinitionsToolkitParameterPanelViewModel::Cleanup()
 	}
 }
 
-void FNiagaraParameterDefinitionsToolkitParameterPanelViewModel::AddParameter(FNiagaraVariable NewVariable, const FNiagaraParameterPanelCategory Category, const bool bRequestRename, const bool bMakeUniqueName) const
+void FNiagaraParameterDefinitionsToolkitParameterPanelViewModel::AddParameter(FNiagaraVariable NewVariable, const FNiagaraParameterPanelCategory Category, const bool bRequestRename, const bool bMakeUniqueName) 
 {
 	TGuardValue<bool> AddParameterRefreshGuard(bIsAddingParameter, true);
-	bool bSuccess = false;
 
 	if (bMakeUniqueName)
 	{
@@ -3341,9 +3629,10 @@ void FNiagaraParameterDefinitionsToolkitParameterPanelViewModel::AddParameter(FN
 	
 	ParameterDefinitionsWeak.Get()->AddParameter(NewVariable);
 	SelectParameterItemByName(NewVariable.GetName(), bRequestRename);
+
 }
 
-void FNiagaraParameterDefinitionsToolkitParameterPanelViewModel::FindOrAddParameter(FNiagaraVariable Variable, const FNiagaraParameterPanelCategory Category) const
+void FNiagaraParameterDefinitionsToolkitParameterPanelViewModel::FindOrAddParameter(FNiagaraVariable Variable, const FNiagaraParameterPanelCategory Category) 
 {
 	TGuardValue<bool> AddParameterRefreshGuard(bIsAddingParameter, true);
 
@@ -3360,7 +3649,7 @@ bool FNiagaraParameterDefinitionsToolkitParameterPanelViewModel::GetCanAddParame
 	return true;
 }
 
-void FNiagaraParameterDefinitionsToolkitParameterPanelViewModel::DeleteParameters(const TArray<FNiagaraParameterPanelItem>& ItemsToDelete) const
+void FNiagaraParameterDefinitionsToolkitParameterPanelViewModel::DeleteParameters(const TArray<FNiagaraParameterPanelItem>& ItemsToDelete) 
 {
 	FScopedTransaction RemoveParameter(LOCTEXT("ParameterDefinitionsToolkitParameterPanelViewModel_RemoveParameter", "Removed Parameter(s)"));
 
@@ -3372,13 +3661,13 @@ void FNiagaraParameterDefinitionsToolkitParameterPanelViewModel::DeleteParameter
 	UIContext.RefreshSelectionDetailsViewPanel();
 }
 
-void FNiagaraParameterDefinitionsToolkitParameterPanelViewModel::RenameParameter(const FNiagaraParameterPanelItem& ItemToRename, const FName NewName) const
+void FNiagaraParameterDefinitionsToolkitParameterPanelViewModel::RenameParameter(const FNiagaraParameterPanelItem& ItemToRename, const FName NewName) 
 {
 	FScopedTransaction RenameParameter(LOCTEXT("ParameterDefinitionsToolkitParameterPanelViewModel_RenameParameter", "Rename Parameter"));
 	ParameterDefinitionsWeak.Get()->RenameParameter(ItemToRename.GetVariable(), NewName);
 }
 
-void FNiagaraParameterDefinitionsToolkitParameterPanelViewModel::SetParameterIsSubscribedToLibrary(const UNiagaraScriptVariable* ScriptVarToModify, const bool bSubscribed) const
+void FNiagaraParameterDefinitionsToolkitParameterPanelViewModel::SetParameterIsSubscribedToLibrary(const UNiagaraScriptVariable* ScriptVarToModify, const bool bSubscribed) 
 {
 	// Do nothing, parameter definitions parameters are always subscribed to their parent parameter definitions.
 	ensureMsgf(false, TEXT("Tried to set a parameter definitions defined parameter subscribing! This should not be reachable."));
@@ -3570,7 +3859,7 @@ const TArray<FNiagaraParameterPanelCategory>& FNiagaraParameterDefinitionsToolki
 	return FNiagaraParameterDefinitionsToolkitParameterPanelViewModel::DefaultCategories;
 }
 
-FMenuAndSearchBoxWidgets FNiagaraParameterDefinitionsToolkitParameterPanelViewModel::GetParameterMenu(FNiagaraParameterPanelCategory Category) const
+FMenuAndSearchBoxWidgets FNiagaraParameterDefinitionsToolkitParameterPanelViewModel::GetParameterMenu(FNiagaraParameterPanelCategory Category) 
 {
 	const bool bRequestRename = true;
 	const bool bMakeUniqueName = true;
@@ -3591,7 +3880,7 @@ FMenuAndSearchBoxWidgets FNiagaraParameterDefinitionsToolkitParameterPanelViewMo
 	return MenuAndSearchBoxWidgets;
 }
 
-FReply FNiagaraParameterDefinitionsToolkitParameterPanelViewModel::HandleDragDropOperation(TSharedPtr<FDragDropOperation> DropOperation) const
+FReply FNiagaraParameterDefinitionsToolkitParameterPanelViewModel::HandleDragDropOperation(TSharedPtr<FDragDropOperation> DropOperation) 
 {
 	ensureMsgf(false, TEXT("Tried to handle drag drop op in parameter definitions parameter panel viewmodel!"));
 	return FReply::Handled();
diff --git a/Engine/Plugins/FX/Niagara/Source/NiagaraEditor/Private/ViewModels/NiagaraScratchPadScriptViewModel.cpp b/Engine/Plugins/FX/Niagara/Source/NiagaraEditor/Private/ViewModels/NiagaraScratchPadScriptViewModel.cpp
--- a/Engine/Plugins/FX/Niagara/Source/NiagaraEditor/Private/ViewModels/NiagaraScratchPadScriptViewModel.cpp
+++ b/Engine/Plugins/FX/Niagara/Source/NiagaraEditor/Private/ViewModels/NiagaraScratchPadScriptViewModel.cpp
@@ -13,6 +13,7 @@
 #include "NiagaraParameterDefinitions.h"
 #include "ViewModels/NiagaraScratchPadUtilities.h"
 #include "ViewModels/NiagaraScriptGraphViewModel.h"
+#include "ViewModels/NiagaraSystemViewModel.h"
 
 #include "ScopedTransaction.h"
 #include "Framework/Commands/UICommandList.h"
@@ -47,17 +48,30 @@ FNiagaraScratchPadScriptViewModel::~FNiagaraScratchPadScriptViewModel()
 	}
 }
 
-void FNiagaraScratchPadScriptViewModel::Initialize(UNiagaraScript* Script)
+void FNiagaraScratchPadScriptViewModel::Initialize(UNiagaraScript* Script, UNiagaraScript* InEditScript, TWeakPtr<FNiagaraSystemViewModel> InMasterViewModel)
 {
 	OriginalScript = Script;
-	EditScript.Script = CastChecked<UNiagaraScript>(StaticDuplicateObject(Script, GetTransientPackage()));
-	SetScript(EditScript);
+	SystemViewModel = InMasterViewModel;
+
+	// Copy over the old edited graph, as this initialize might be incoming from inheritance changes and we want to maintain
+	// the edits to the old graph we were working on.
+	if (InEditScript)
+	{
+		EditScript.Script = InEditScript;
+		SetScript(EditScript);
+	}
+	else
+	{
+		EditScript.Script = CastChecked<UNiagaraScript>(StaticDuplicateObject(Script, GetTransientPackage()));
+		SetScript(EditScript);
+	}
 	UNiagaraScriptSource* EditScriptSource = CastChecked<UNiagaraScriptSource>(EditScript.Script->GetLatestSource());
 	OnGraphNeedsRecompileHandle = EditScriptSource->NodeGraph->AddOnGraphNeedsRecompileHandler(FOnGraphChanged::FDelegate::CreateSP(this, &FNiagaraScratchPadScriptViewModel::OnScriptGraphChanged));
 	EditScript.Script->OnPropertyChanged().AddSP(this, &FNiagaraScratchPadScriptViewModel::OnScriptPropertyChanged);
 	ParameterPanelCommands = MakeShared<FUICommandList>();
 
 	ParameterPaneViewModel = MakeShared<FNiagaraScriptToolkitParameterPanelViewModel>(this->AsShared());
+	ExternalSelectionChangedDelegate = ParameterPaneViewModel->GetOnExternalSelectionChangedDelegate().AddRaw(this, &FNiagaraScratchPadScriptViewModel::OnGraphSubObjectSelectionChanged);
 	FScriptToolkitUIContext UIContext = FScriptToolkitUIContext(
 		FSimpleDelegate::CreateSP(ParameterPaneViewModel.ToSharedRef(), &INiagaraImmutableParameterPanelViewModel::Refresh),
 		FSimpleDelegate(), //@todo(ng) skip binding refresh parameter definitions panel as scratchpad does not have a parameter definitions panel currently
@@ -66,11 +80,17 @@ void FNiagaraScratchPadScriptViewModel::Initialize(UNiagaraScript* Script)
 	ParameterPaneViewModel->Init(UIContext);
 }
 
+void FNiagaraScratchPadScriptViewModel::OnGraphSubObjectSelectionChanged(const UObject* Obj)
+{
+	OnGraphSelectionChanged().Broadcast(Obj);
+}
+
 void FNiagaraScratchPadScriptViewModel::Finalize()
 {
 	// This pointer needs to be reset manually here because there is a shared ref cycle.
 	if (ParameterPaneViewModel.IsValid())
 	{
+		ParameterPaneViewModel->GetOnExternalSelectionChangedDelegate().Remove(ExternalSelectionChangedDelegate);
 		ParameterPaneViewModel.Reset();
 	}
 }
@@ -164,6 +184,17 @@ void FNiagaraScratchPadScriptViewModel::SetIsPinned(bool bInIsPinned)
 	}
 }
 
+void FNiagaraScratchPadScriptViewModel::TransferFromOldWhenDoingApply(TSharedPtr< FNiagaraScratchPadScriptViewModel> InOldScriptVM)
+{
+	// Copy over has pending changes, as this might be incoming from inheritance changes and we want to maintain
+	// the edits to the old graph we were working on.
+	if (InOldScriptVM.IsValid())
+	{
+		bHasPendingChanges = InOldScriptVM->bHasPendingChanges;
+	}
+}
+
+
 float FNiagaraScratchPadScriptViewModel::GetEditorHeight() const
 {
 	return EditorHeight;
diff --git a/Engine/Plugins/FX/Niagara/Source/NiagaraEditor/Private/ViewModels/NiagaraScratchPadViewModel.cpp b/Engine/Plugins/FX/Niagara/Source/NiagaraEditor/Private/ViewModels/NiagaraScratchPadViewModel.cpp
--- a/Engine/Plugins/FX/Niagara/Source/NiagaraEditor/Private/ViewModels/NiagaraScratchPadViewModel.cpp
+++ b/Engine/Plugins/FX/Niagara/Source/NiagaraEditor/Private/ViewModels/NiagaraScratchPadViewModel.cpp
@@ -106,22 +106,64 @@ void UNiagaraScratchPadViewModel::RefreshScriptViewModels()
 	TArray<TObjectPtr<UNiagaraScript>>* TargetScripts;
 	GetOuterAndTargetScripts(GetSystemViewModel(), ScriptOuter, TargetScripts);
 
+	TSharedPtr<FNiagaraScratchPadScriptViewModel> ForceSetActiveDocumentAtEnd;
+
+	TArray<UNiagaraGraph*> GraphsToClose;
+	TArray<UNiagaraGraph*> GraphsToOpen;
 	if (ScriptOuter != nullptr && TargetScripts != nullptr)
 	{
-		for (UNiagaraScript* ScratchPadScript : (*TargetScripts))
+		TSharedPtr<FNiagaraScratchPadScriptViewModel> ActiveScratchViewModel = GetSystemViewModel()->GetDocumentViewModel()->GetActiveScratchPadViewModelIfSet();
+		for (int32 i = 0; i < (*TargetScripts).Num(); i++)
 		{
+			UNiagaraScript* ScratchPadScript = (*TargetScripts)[i];
 			TSharedPtr<FNiagaraScratchPadScriptViewModel> ScriptViewModel;
 
-			TSharedRef<FNiagaraScratchPadScriptViewModel>* OldScriptViewModel = OldScriptViewModels.FindByPredicate(
-				[ScratchPadScript](TSharedRef<FNiagaraScratchPadScriptViewModel> ScriptViewModel) { return ScriptViewModel->GetOriginalScript() == ScratchPadScript; });
-			if (OldScriptViewModel != nullptr)
+			if (GetSystemViewModel()->GetEditMode() == ENiagaraSystemViewModelEditMode::SystemAsset)
 			{
-				ScriptViewModel = *OldScriptViewModel;
-				OldScriptViewModels.Remove(ScriptViewModel.ToSharedRef());
+				// System scripts don't ever merge, so we can safely just take the scripts or reuse view models.
+				TSharedRef<FNiagaraScratchPadScriptViewModel>* OldScriptViewModel = OldScriptViewModels.FindByPredicate(
+					[ScratchPadScript](TSharedRef<FNiagaraScratchPadScriptViewModel> ScriptViewModel) { return ScriptViewModel->GetOriginalScript() == ScratchPadScript; });
+				if (OldScriptViewModel != nullptr)
+				{
+					ScriptViewModel = *OldScriptViewModel;
+					OldScriptViewModels.Remove(ScriptViewModel.ToSharedRef());
+				}
+				else
+				{
+					ScriptViewModel = CreateAndSetupScriptviewModel(ScratchPadScript, nullptr);
+					bViewModelsChanged = true;
+				}
 			}
 			else
 			{
-				ScriptViewModel = CreateAndSetupScriptviewModel(ScratchPadScript);
+				// Emitter scripts may already be opened for edit but not yet applied so mergers can refresh and would wipe out the base scripts. These need to be kept up to date.
+				// Because the merge could wipe out view models, we can't just reuse the old ones, but we can steal their graphs.
+				TSharedPtr<FNiagaraScratchPadScriptViewModel> OldScriptViewModel = OldScriptViewModels.IsValidIndex(i) ? OldScriptViewModels[i] : TSharedPtr<FNiagaraScratchPadScriptViewModel>();
+
+				if (OldScriptViewModel.IsValid())
+				{
+					// Pass the baton of the script to the new generation. Make sure to clear out references in the previous generation.
+					ScriptViewModel = CreateAndSetupScriptviewModel(ScratchPadScript, OldScriptViewModel->GetEditScript().Script);
+					// This will copy over apply state and anything else that would otherwise be lost in the handoff
+					ScriptViewModel->TransferFromOldWhenDoingApply(OldScriptViewModel);
+				
+					// We need to keep track of the active document and reset again after we have the full array built up.
+					if (ActiveScratchViewModel == OldScriptViewModel)
+					{
+						ForceSetActiveDocumentAtEnd = ScriptViewModel;
+					}
+
+					// We now need to make sure that all the UI/UX is transferred to the right objects for event handlers/etc.
+					GetSystemViewModel()->GetDocumentViewModel()->SwapEditableScripts(OldScriptViewModel, ScriptViewModel);
+					// Clear out the old to make sure it isn't touching anything
+					OldScriptViewModel->SetScript(FVersionedNiagaraScript());
+					OldScriptViewModels.Remove(OldScriptViewModel.ToSharedRef());
+					OldScriptViewModel.Reset();
+				}
+				else
+				{
+					ScriptViewModel = CreateAndSetupScriptviewModel(ScratchPadScript, nullptr);
+				}
 				bViewModelsChanged = true;
 			}
 			ScriptViewModels.Add(ScriptViewModel.ToSharedRef());
@@ -164,11 +206,26 @@ void UNiagaraScratchPadViewModel::RefreshScriptViewModels()
 	{
 		RefreshEditScriptViewModels();
 	}
+
+	{
+		for (int32 i = 0; i < ScriptViewModels.Num(); i++)
+		{
+			TArray<TSharedPtr<FNiagaraScriptViewModel>> ExistingViewModels;
+			TNiagaraViewModelManager<UNiagaraScript, FNiagaraScriptViewModel>::GetAllViewModelsForObject(ScriptViewModels[i]->GetEditScript().Script, ExistingViewModels);
+
+			(ensureMsgf(ExistingViewModels.Num() == 1, TEXT("Missing or duplicate script view models detected.  Can not create script details customization.")));
+		}
+	}
 	
 	if (bViewModelsChanged)
 	{
 		OnScriptViewModelsChangedDelegate.Broadcast();
 	}
+
+	if (ForceSetActiveDocumentAtEnd.IsValid())
+	{
+		FocusScratchPadScriptViewModel(ForceSetActiveDocumentAtEnd.ToSharedRef());
+	}
 }
 
 void UNiagaraScratchPadViewModel::ApplyScratchPadChanges()
@@ -206,6 +263,18 @@ TSharedPtr<FNiagaraScratchPadScriptViewModel> UNiagaraScratchPadViewModel::GetVi
 	return TSharedPtr<FNiagaraScratchPadScriptViewModel>();
 }
 
+TSharedPtr<FNiagaraScratchPadScriptViewModel> UNiagaraScratchPadViewModel::GetViewModelForScript(FName InScriptName)
+{
+	TSharedRef<FNiagaraScratchPadScriptViewModel>* ViewModelForScript = ScriptViewModels.FindByPredicate([InScriptName](TSharedRef<FNiagaraScratchPadScriptViewModel>& ScriptViewModel) { return ScriptViewModel->GetOriginalScript()->GetFName().IsEqual(InScriptName, ENameCase::IgnoreCase, false); });
+	if (ViewModelForScript != nullptr)
+	{
+		return *ViewModelForScript;
+	}
+	return TSharedPtr<FNiagaraScratchPadScriptViewModel>();
+}
+
+
+
 TSharedPtr<FNiagaraScratchPadScriptViewModel> UNiagaraScratchPadViewModel::GetViewModelForEditScript(UNiagaraScript* InEditScript)
 {
 	TSharedRef<FNiagaraScratchPadScriptViewModel>* ViewModelForEditScript = ScriptViewModels.FindByPredicate([InEditScript](TSharedRef<FNiagaraScratchPadScriptViewModel>& ScriptViewModel) { return ScriptViewModel->GetEditScript().Script == InEditScript; });
@@ -639,12 +708,13 @@ TSharedRef<FNiagaraSystemViewModel> UNiagaraScratchPadViewModel::GetSystemViewMo
 	return SystemViewModel.ToSharedRef();
 }
 
-TSharedRef<FNiagaraScratchPadScriptViewModel> UNiagaraScratchPadViewModel::CreateAndSetupScriptviewModel(UNiagaraScript* ScratchPadScript)
+TSharedRef<FNiagaraScratchPadScriptViewModel> UNiagaraScratchPadViewModel::CreateAndSetupScriptviewModel(UNiagaraScript* ScratchPadScript, UNiagaraScript* ScratchPadEditScript)
 {
 	TSharedRef<FNiagaraScratchPadScriptViewModel> ScriptViewModel = MakeShared<FNiagaraScratchPadScriptViewModel>(GetSystemViewModel()->GetIsForDataProcessingOnly());
-	ScriptViewModel->Initialize(ScratchPadScript);
+	ScriptViewModel->Initialize(ScratchPadScript, ScratchPadEditScript, SystemViewModelWeak);
 	ScriptViewModel->GetGraphViewModel()->GetNodeSelection()->OnSelectedObjectsChanged().AddUObject(this, &UNiagaraScratchPadViewModel::ScriptGraphNodeSelectionChanged, TWeakPtr<FNiagaraScratchPadScriptViewModel>(ScriptViewModel));
 	ScriptViewModel->OnRenamed().AddUObject(this, &UNiagaraScratchPadViewModel::ScriptViewModelScriptRenamed, TWeakPtr<FNiagaraScratchPadScriptViewModel>(ScriptViewModel));
+	ScriptViewModel->OnGraphSelectionChanged().AddUObject(this, &UNiagaraScratchPadViewModel::ScriptViewModelGraphSelectionChanged, TWeakPtr<FNiagaraScratchPadScriptViewModel>(ScriptViewModel));
 	ScriptViewModel->OnPinnedChanged().AddUObject(this, &UNiagaraScratchPadViewModel::ScriptViewModelPinnedChanged, TWeakPtr<FNiagaraScratchPadScriptViewModel>(ScriptViewModel));
 	ScriptViewModel->OnHasUnappliedChangesChanged().AddUObject(this, &UNiagaraScratchPadViewModel::ScriptViewModelHasUnappliedChangesChanged);
 	ScriptViewModel->OnChangesApplied().AddUObject(this, &UNiagaraScratchPadViewModel::ScriptViewModelChangesApplied);
@@ -657,6 +727,7 @@ void UNiagaraScratchPadViewModel::TearDownScriptViewModel(TSharedRef<FNiagaraScr
 {
 	InScriptViewModel->GetGraphViewModel()->GetNodeSelection()->OnSelectedObjectsChanged().RemoveAll(this);
 	InScriptViewModel->OnRenamed().RemoveAll(this);
+	InScriptViewModel->OnGraphSelectionChanged().RemoveAll(this);
 	InScriptViewModel->OnPinnedChanged().RemoveAll(this);
 	InScriptViewModel->OnHasUnappliedChangesChanged().RemoveAll(this);
 	InScriptViewModel->OnChangesApplied().RemoveAll(this);
@@ -715,6 +786,15 @@ void UNiagaraScratchPadViewModel::ScriptViewModelScriptRenamed(TWeakPtr<FNiagara
 	}
 }
 
+void UNiagaraScratchPadViewModel::ScriptViewModelGraphSelectionChanged(const UObject* Obj, TWeakPtr<FNiagaraScratchPadScriptViewModel> ScriptViewModelWeak)
+{
+	const UNiagaraScriptVariable* Var = Cast<UNiagaraScriptVariable>(Obj);
+	if (Var)
+	{
+		GetSystemViewModel()->GetParameterPanelViewModel()->SelectParameterItemByName(Var->Variable.GetName(), false);
+	}
+}
+
 void UNiagaraScratchPadViewModel::ScriptViewModelPinnedChanged(TWeakPtr<FNiagaraScratchPadScriptViewModel> ScriptViewModelWeak)
 {
 	TSharedPtr<FNiagaraScratchPadScriptViewModel> ScriptViewModel = ScriptViewModelWeak.Pin();
@@ -792,7 +872,7 @@ void UNiagaraScratchPadViewModel::ScriptViewModelRequestDiscardChanges(TWeakPtr<
 		int32 ViewModelIndex = ScriptViewModels.IndexOfByKey(ScriptViewModel.ToSharedRef());
 		if (ensureMsgf(ViewModelIndex != INDEX_NONE, TEXT("Active script view model wasn't in script view model collection!")))
 		{
-			TSharedRef<FNiagaraScratchPadScriptViewModel> NewScriptViewModel = CreateAndSetupScriptviewModel(ScriptViewModel->GetOriginalScript());
+			TSharedRef<FNiagaraScratchPadScriptViewModel> NewScriptViewModel = CreateAndSetupScriptviewModel(ScriptViewModel->GetOriginalScript(), nullptr);
 			ScriptViewModels[ViewModelIndex] = NewScriptViewModel;
 
 			int32 PinnedViewModelIndex = PinnedScriptViewModels.IndexOfByKey(ScriptViewModel.ToSharedRef());
diff --git a/Engine/Plugins/FX/Niagara/Source/NiagaraEditor/Private/ViewModels/NiagaraSystemEditorDocumentsViewModel.cpp b/Engine/Plugins/FX/Niagara/Source/NiagaraEditor/Private/ViewModels/NiagaraSystemEditorDocumentsViewModel.cpp
--- a/Engine/Plugins/FX/Niagara/Source/NiagaraEditor/Private/ViewModels/NiagaraSystemEditorDocumentsViewModel.cpp
+++ b/Engine/Plugins/FX/Niagara/Source/NiagaraEditor/Private/ViewModels/NiagaraSystemEditorDocumentsViewModel.cpp
@@ -6,6 +6,7 @@
 #include "ViewModels/NiagaraScratchPadViewModel.h"
 #include "ViewModels/NiagaraScratchPadScriptViewModel.h"
 #include "ViewModels/NiagaraScriptGraphViewModel.h"
+#include "ViewModels/NiagaraSystemScriptViewModel.h"
 #include "ViewModels/NiagaraParameterPanelViewModel.h"
 #include "ViewModels/Stack/NiagaraStackModuleItem.h"
 #include "ViewModels/Stack/NiagaraStackFunctionInput.h"
@@ -49,12 +50,126 @@ void UNiagaraSystemEditorDocumentsViewModel::Finalize()
 	ActiveDocumentTabScriptViewModel.Reset();
 }
 
-TArray<UNiagaraGraph*> UNiagaraSystemEditorDocumentsViewModel::GetEditableGraphsForActiveDocument()
+TArray<UNiagaraGraph*> UNiagaraSystemEditorDocumentsViewModel::GetAllGraphsForActiveScriptDocument()
 {
 	if (ActiveDocumentTabScriptViewModel.IsValid())
+	{
 		return ActiveDocumentTabScriptViewModel->GetEditableGraphs();
+	}
+	else if (ActiveDocumentTab.IsValid())
+	{
+		if (ActiveDocumentTab.Pin()->GetLayoutIdentifier().TabType != PrimaryDocumentTabId)
+		{
+			ActiveDocumentTabScriptViewModel = GetActiveScratchPadViewModelIfSet();
+
+			if (ActiveDocumentTabScriptViewModel.IsValid())
+			{
+				return ActiveDocumentTabScriptViewModel->GetEditableGraphs();
+			}
+		}
+	}
+	return TArray<UNiagaraGraph*>();
+}
+
+TArray<UNiagaraGraph*> UNiagaraSystemEditorDocumentsViewModel::GetEditableGraphsForActiveScriptDocument()
+{
+	if (ActiveDocumentTabScriptViewModel.IsValid())
+	{
+		return ActiveDocumentTabScriptViewModel->GetEditableGraphs();
+	}
+	else if (ActiveDocumentTab.IsValid())
+	{
+		if (ActiveDocumentTab.Pin()->GetLayoutIdentifier().TabType != PrimaryDocumentTabId)
+		{
+			ActiveDocumentTabScriptViewModel = GetActiveScratchPadViewModelIfSet();
+			
+			if (ActiveDocumentTabScriptViewModel.IsValid())
+			{
+				return ActiveDocumentTabScriptViewModel->GetEditableGraphs();
+			}
+		}	
+	}
+	
+	return TArray<UNiagaraGraph*>();
+}
+
+TArray<UNiagaraGraph*> UNiagaraSystemEditorDocumentsViewModel::GetAllGraphsForPrimaryDocument()
+{
+	TArray<UNiagaraGraph*> OutGraphs;
+
+	// Helper lambda to null check graph weak object ptrs before adding them to the retval array.
+	auto AddToOutGraphsChecked = [&OutGraphs](const TWeakObjectPtr<UNiagaraGraph>& WeakGraph) {
+		UNiagaraGraph* Graph = WeakGraph.Get();
+		if (Graph == nullptr)
+		{
+			ensureMsgf(false, TEXT("Encountered null graph when gathering editable graphs for system parameter panel viewmodel!"));
+			return;
+		}
+		OutGraphs.Add(Graph);
+	};
+	TSharedRef<FNiagaraSystemViewModel> SystemViewModel = GetSystemViewModel();
+	if (SystemViewModel->GetEditMode() == ENiagaraSystemViewModelEditMode::SystemAsset)
+	{
+		TArray<UNiagaraGraph*> Graphs;
+		AddToOutGraphsChecked(SystemViewModel->GetSystemScriptViewModel()->GetGraphViewModel()->GetGraph());
+
+		const TArray<TSharedRef<FNiagaraEmitterHandleViewModel>>& EmitterHandleViewModels = SystemViewModel->GetEmitterHandleViewModels();
+		for (const TSharedRef<FNiagaraEmitterHandleViewModel>& EmitterHandleViewModel : EmitterHandleViewModels)
+		{
+			FNiagaraEmitterHandle* EmitterHandle = EmitterHandleViewModel->GetEmitterHandle();
+			if (EmitterHandle == nullptr)
+			{
+				continue;
+			}
+			UNiagaraGraph* Graph = Cast<UNiagaraScriptSource>(EmitterHandle->GetEmitterData()->GraphSource)->NodeGraph;
+			if (Graph)
+			{
+				OutGraphs.Add(Graph);
+			}
+		}
+	}
+	else
+	{
+		FNiagaraEmitterHandle* EmitterHandle = SystemViewModel->GetEmitterHandleViewModels()[0]->GetEmitterHandle();
+		if (EmitterHandle != nullptr)
+		{
+			OutGraphs.Add(Cast<UNiagaraScriptSource>(EmitterHandle->GetEmitterData()->GraphSource)->NodeGraph);
+		}
+	}
+	return OutGraphs;
+}
+
+TArray<UNiagaraGraph*> UNiagaraSystemEditorDocumentsViewModel::GetEditableGraphsForPrimaryDocument()
+{
+	
+	TArray<UNiagaraGraph*> EditableGraphs;
+
+	// Helper lambda to null check graph weak object ptrs before adding them to the retval array.
+	auto AddToEditableGraphChecked = [&EditableGraphs](const TWeakObjectPtr<UNiagaraGraph>& WeakGraph) {
+		UNiagaraGraph* Graph = WeakGraph.Get();
+		if (Graph == nullptr)
+		{
+			ensureMsgf(false, TEXT("Encountered null graph when gathering editable graphs for system parameter panel viewmodel!"));
+			return;
+		}
+		EditableGraphs.Add(Graph);
+	};
+
+	TSharedRef<FNiagaraSystemViewModel> SystemViewModel = GetSystemViewModel();
+	if (SystemViewModel->GetEditMode() == ENiagaraSystemViewModelEditMode::SystemAsset && ensureMsgf(SystemGraphSelectionVMWeak.IsValid(), TEXT("SystemGraphSelectionViewModel was null for System edit mode!")))
+	{
+		for (const TWeakObjectPtr<UNiagaraGraph>& WeakGraph : SystemGraphSelectionVMWeak.Pin()->GetSelectedEmitterScriptGraphs())
+		{
+			AddToEditableGraphChecked(WeakGraph);
+		}
+		AddToEditableGraphChecked(SystemViewModel->GetSystemScriptViewModel()->GetGraphViewModel()->GetGraph());
+	}
 	else
-		return TArray<UNiagaraGraph*>();
+	{
+		EditableGraphs.Add(Cast<UNiagaraScriptSource>(SystemViewModel->GetEmitterHandleViewModels()[0]->GetEmitterHandle()->GetEmitterData()->GraphSource)->NodeGraph);
+	}
+
+	return EditableGraphs;
 }
 
 void UNiagaraSystemEditorDocumentsViewModel::OpenChildScript(UEdGraph* InGraph)
@@ -66,10 +181,45 @@ void UNiagaraSystemEditorDocumentsViewModel::OpenChildScript(UEdGraph* InGraph)
 	}
 	else if (Results.Num() > 0)
 	{
-		TabManager->DrawAttention(Results[0].ToSharedRef());
+		if (Results[0].IsValid())
+		{
+			TabManager->DrawAttention(Results[0].ToSharedRef());
+			Results[0]->ActivateInParent(ETabActivationCause::SetDirectly);
+			SetActiveDocumentTab(Results[0]);
+		}
 	}
 }
 
+void UNiagaraSystemEditorDocumentsViewModel::SwapEditableScripts(TSharedPtr < class FNiagaraScratchPadScriptViewModel> OldScriptViewModel, TSharedPtr < class FNiagaraScratchPadScriptViewModel> NewScriptViewModel)
+{
+	// In the event of an inheritance update, we might be called upon to do a hotswap of the currently editable scratch scripts,
+	// which are never overwritten by the parent. However, they were on graphs that need to be propagated to the 
+	// documents that are currently open if already in use for editing.
+	if (OldScriptViewModel.IsValid())
+	{
+		TArray< TSharedPtr<SDockTab> > Results;
+		UEdGraph* OldGraph = OldScriptViewModel->GetGraphViewModel()->GetGraph();
+		if (FindOpenTabsContainingDocument(OldGraph, Results))
+		{
+			for (TSharedPtr<SDockTab> Tab : Results)
+			{
+				if (Tab.IsValid() && Tab->GetLayoutIdentifier().TabType == TEXT("Document"))
+				{
+					TSharedRef<SNiagaraScratchPadScriptEditor> GraphEditor = StaticCastSharedRef<SNiagaraScratchPadScriptEditor>(Tab->GetContent());
+					GraphEditor->SetViewModel(NewScriptViewModel);
+
+					if (ActiveDocumentTabScriptViewModel == OldScriptViewModel)
+					{
+						ActiveDocumentTabScriptViewModel = NewScriptViewModel;
+					}
+				}
+			}
+		}
+	}
+
+}
+
+
 void UNiagaraSystemEditorDocumentsViewModel::CloseChildScript(UEdGraph* InGraph)
 {
 	CloseDocumentTab(InGraph);
@@ -106,6 +256,25 @@ void UNiagaraSystemEditorDocumentsViewModel::InitializePostTabManager(TSharedPtr
 {
 	TabManager = InToolkit->GetTabManager();
 	DocumentManager->SetTabManager(TabManager.ToSharedRef());
+	SystemGraphSelectionVMWeak = InToolkit->GetSystemGraphSelectionViewModel();
+}
+
+void UNiagaraSystemEditorDocumentsViewModel::DrawAttentionToPrimaryDocument()
+{
+	// Flash the system overivew in the UI
+	TSharedPtr<SDockTab>  ActiveTab = GetActiveDocumentTab().Pin();
+	if (ActiveTab.IsValid())
+	{
+		if (ActiveTab->GetLayoutIdentifier().TabType != PrimaryDocumentTabId)
+		{
+			TSharedPtr<SDockTab> FoundPrimaryDocumentTab = TabManager->FindExistingLiveTab(PrimaryDocumentTabId);
+			if (FoundPrimaryDocumentTab.IsValid())
+			{
+				FoundPrimaryDocumentTab->DrawAttention();
+				FoundPrimaryDocumentTab->ActivateInParent(ETabActivationCause::SetDirectly);
+			}
+		}
+	}
 }
 
 
@@ -166,12 +335,25 @@ void UNiagaraSystemEditorDocumentsViewModel::SetActiveDocumentTab(TSharedPtr<SDo
 
 	ActiveDocChangedDelegate.Broadcast(Tab);
 
-
 	// We need to update the parameter panel view model with new parameters potentially
-	if (GetSystemViewModel()->GetParameterPanelViewModel())
-		GetSystemViewModel()->GetParameterPanelViewModel()->RefreshNextTick();
+	INiagaraParameterPanelViewModel* PanelVM = GetSystemViewModel()->GetParameterPanelViewModel();
+	if (PanelVM)
+	{
+		PanelVM->RefreshNextTick();
+	}
 }
 
+bool UNiagaraSystemEditorDocumentsViewModel::IsPrimaryDocumentActive() const
+{
+	TSharedPtr<SDockTab> Tab = ActiveDocumentTab.Pin();
+	if (Tab.IsValid() && Tab->GetLayoutIdentifier().TabType == PrimaryDocumentTabId)
+	{
+		return true;
+	}
+	return false;
+}
+
+
 TSharedPtr<FNiagaraScratchPadScriptViewModel> UNiagaraSystemEditorDocumentsViewModel::GetScratchPadViewModelFromGraph(FNiagaraSystemViewModel* InSysViewModel, UEdGraph* InTargetGraph)
 {
 	if (InSysViewModel && InTargetGraph)
diff --git a/Engine/Plugins/FX/Niagara/Source/NiagaraEditor/Private/ViewModels/Stack/NiagaraStackFunctionInput.cpp b/Engine/Plugins/FX/Niagara/Source/NiagaraEditor/Private/ViewModels/Stack/NiagaraStackFunctionInput.cpp
--- a/Engine/Plugins/FX/Niagara/Source/NiagaraEditor/Private/ViewModels/Stack/NiagaraStackFunctionInput.cpp
+++ b/Engine/Plugins/FX/Niagara/Source/NiagaraEditor/Private/ViewModels/Stack/NiagaraStackFunctionInput.cpp
@@ -2679,9 +2679,16 @@ const UNiagaraStackFunctionInput::FCollectedUsageData& UNiagaraStackFunctionInpu
 		INiagaraParameterPanelViewModel* ParamVM = SystemVM->GetParameterPanelViewModel();
 		if (ParamVM)
 		{
+			bool bFoundOverride = false;
+			if (InputValues.Mode == EValueMode::Linked)
+			{
+				FNiagaraVariableBase Var(InputType, InputValues.LinkedHandle.GetParameterHandleString());
+				bFoundOverride = ParamVM->IsVariableSelected(Var);
+			}
 			FNiagaraVariableBase Var(InputType, AliasedInputParameterHandle.GetParameterHandleString());
-			CachedCollectedUsageData.GetValue().bHasReferencedParameterRead = ParamVM->IsVariableSelected(Var);
+			CachedCollectedUsageData.GetValue().bHasReferencedParameterRead = ParamVM->IsVariableSelected(Var) || bFoundOverride;
 		}
+
 	}
 
 	return CachedCollectedUsageData.GetValue();
@@ -3044,4 +3051,29 @@ void UNiagaraStackFunctionInput::RemoveOverridePin()
 	}
 }
 
+bool UNiagaraStackFunctionInput::OpenSourceAsset() const
+{
+	// Helper to open scratch script or function script in the right sub-editor.
+	UNiagaraNodeFunctionCall* DynamicInputNode = GetDynamicInputNode();
+	if (DynamicInputNode->FunctionScript != nullptr)
+	{
+		if (DynamicInputNode->FunctionScript->IsAsset())
+		{
+			DynamicInputNode->FunctionScript->VersionToOpenInEditor = DynamicInputNode->SelectedScriptVersion;
+			return GEditor->GetEditorSubsystem<UAssetEditorSubsystem>()->OpenEditorForAsset(DynamicInputNode->FunctionScript);
+		}
+		else
+		{
+			TSharedPtr<FNiagaraScratchPadScriptViewModel> ScratchPadScriptViewModel = GetSystemViewModel()->GetScriptScratchPadViewModel()->GetViewModelForScript(DynamicInputNode->FunctionScript);
+			if (ScratchPadScriptViewModel.IsValid())
+			{
+				GetSystemViewModel()->GetScriptScratchPadViewModel()->FocusScratchPadScriptViewModel(ScratchPadScriptViewModel.ToSharedRef());
+				return true;
+			}
+		}
+	}
+
+	return false;
+}
+
 #undef LOCTEXT_NAMESPACE
diff --git a/Engine/Plugins/FX/Niagara/Source/NiagaraEditor/Private/ViewModels/Stack/NiagaraStackModuleItem.cpp b/Engine/Plugins/FX/Niagara/Source/NiagaraEditor/Private/ViewModels/Stack/NiagaraStackModuleItem.cpp
--- a/Engine/Plugins/FX/Niagara/Source/NiagaraEditor/Private/ViewModels/Stack/NiagaraStackModuleItem.cpp
+++ b/Engine/Plugins/FX/Niagara/Source/NiagaraEditor/Private/ViewModels/Stack/NiagaraStackModuleItem.cpp
@@ -34,6 +34,8 @@
 #include "ViewModels/Stack/NiagaraStackGraphUtilities.h"
 #include "ViewModels/Stack/NiagaraStackModuleItemLinkedInputCollection.h"
 #include "ViewModels/Stack/NiagaraStackModuleItemOutputCollection.h"
+#include "NiagaraEditorModule.h"
+#include "Toolkits/NiagaraSystemToolkit.h"
 
 // TODO: Remove these
 #include "ViewModels/Stack/NiagaraStackViewModel.h"
@@ -1634,4 +1636,59 @@ bool UNiagaraStackModuleItem::CanDrag() const
 	return true;
 }
 
+bool UNiagaraStackModuleItem::OpenSourceAsset() const
+{
+	// Helper to open scratch script or function script in the appropriate sub editor based on type.
+	const UNiagaraNodeFunctionCall& ModuleFunctionCall = GetModuleNode();
+	if (ModuleFunctionCall.FunctionScript != nullptr)
+	{
+		if (ModuleFunctionCall.FunctionScript->IsAsset() || GbShowNiagaraDeveloperWindows > 0)
+		{
+			ModuleFunctionCall.FunctionScript->VersionToOpenInEditor = ModuleFunctionCall.SelectedScriptVersion;
+			return GEditor->GetEditorSubsystem<UAssetEditorSubsystem>()->OpenEditorForAsset(ToRawPtr(ModuleFunctionCall.FunctionScript));
+		}
+		else if (IsScratchModule())
+		{
+			TSharedPtr<FNiagaraScratchPadScriptViewModel> ScratchPadScriptViewModel = GetSystemViewModel()->GetScriptScratchPadViewModel()->GetViewModelForScript(ModuleFunctionCall.FunctionScript);
+			if (ScratchPadScriptViewModel.IsValid())
+			{
+				GetSystemViewModel()->GetScriptScratchPadViewModel()->FocusScratchPadScriptViewModel(ScratchPadScriptViewModel.ToSharedRef());
+				return true;
+			}
+		}
+		else if (GetIsInherited() && !ModuleFunctionCall.FunctionScript->IsAsset() && GetOutputNode() && GetEmitterViewModel())
+		{
+			TSharedRef<FNiagaraScriptMergeManager> MergeManager = FNiagaraScriptMergeManager::Get();
+
+			FVersionedNiagaraEmitter BaseEmitter = GetEmitterViewModel()->GetParentEmitter();
+
+			UNiagaraScript* FoundScript = nullptr;
+			FGuid FoundScriptVersion;
+			FVersionedNiagaraEmitter FoundBaseEmitter;
+			if (MergeManager->FindBaseModule(BaseEmitter, GetOutputNode()->GetUsage(), OutputNode->GetUsageId(), ModuleFunctionCall.NodeGuid, FoundScript, FoundScriptVersion, FoundBaseEmitter) && FoundScript && FoundBaseEmitter.Emitter)
+			{
+				if (GEditor->GetEditorSubsystem<UAssetEditorSubsystem>()->OpenEditorForAsset(FoundBaseEmitter.Emitter))
+				{
+					if (FNiagaraSystemToolkit* EditorInstance = static_cast<FNiagaraSystemToolkit*>(GEditor->GetEditorSubsystem<UAssetEditorSubsystem>()->FindEditorForAsset(FoundBaseEmitter.Emitter, true)))
+					{
+						EditorInstance->SwitchToVersion(FoundBaseEmitter.Version);
+						TSharedPtr<FNiagaraScratchPadScriptViewModel> ScratchPadScriptViewModel = EditorInstance->GetSystemViewModel()->GetScriptScratchPadViewModel()->GetViewModelForScript(FoundScript->GetFName());
+						if (ScratchPadScriptViewModel.IsValid())
+						{
+							EditorInstance->GetSystemViewModel()->GetScriptScratchPadViewModel()->FocusScratchPadScriptViewModel(ScratchPadScriptViewModel.ToSharedRef());
+							return true;
+						}
+					}
+					return true;
+				}
+			}
+			else
+			{
+				FNiagaraEditorUtilities::InfoWithToastAndLog(FText::Format(LOCTEXT("CannotAutoOpenBaseModule", "Cannot auto-open base module {0} as this was created before merging tracked this data. Newer assets should work. You should be able to open it manually however."), FText::FromString(ModuleFunctionCall.FunctionScript->GetName())));
+			}
+		}
+	}
+	return false;
+}
+
 #undef LOCTEXT_NAMESPACE
diff --git a/Engine/Plugins/FX/Niagara/Source/NiagaraEditor/Private/Widgets/NiagaraScratchScriptEditor.cpp b/Engine/Plugins/FX/Niagara/Source/NiagaraEditor/Private/Widgets/NiagaraScratchScriptEditor.cpp
--- a/Engine/Plugins/FX/Niagara/Source/NiagaraEditor/Private/Widgets/NiagaraScratchScriptEditor.cpp
+++ b/Engine/Plugins/FX/Niagara/Source/NiagaraEditor/Private/Widgets/NiagaraScratchScriptEditor.cpp
@@ -16,7 +16,6 @@
 
 void SNiagaraScratchPadScriptEditor::Construct(const FArguments& InArgs, TSharedRef<FNiagaraScratchPadScriptViewModel> InScriptViewModel)
 {
-	ScriptViewModel = InScriptViewModel;
 
 	ChildSlot
 		[
@@ -63,7 +62,7 @@ void SNiagaraScratchPadScriptEditor::Construct(const FArguments& InArgs, TShared
 					SNew(STextBlock)
 					.TextStyle(FNiagaraEditorStyle::Get(), "NiagaraEditor.ScratchPad.EditorHeaderText")
 					.Text(this, &SNiagaraScratchPadScriptEditor::GetNameText)
-					.ToolTipText(ScriptViewModel.ToSharedRef(), &FNiagaraScratchPadScriptViewModel::GetToolTip)
+					.ToolTipText(this, &SNiagaraScratchPadScriptEditor::GetNameToolTipText)
 				]
 				+ SHorizontalBox::Slot()
 				.VAlign(VAlign_Center)
@@ -74,19 +73,32 @@ void SNiagaraScratchPadScriptEditor::Construct(const FArguments& InArgs, TShared
 					.Visibility(this, &SNiagaraScratchPadScriptEditor::GetUnappliedChangesVisibility)
 					.TextStyle(FNiagaraEditorStyle::Get(), "NiagaraEditor.ScratchPad.EditorHeaderText")
 					.Text(FText::FromString(TEXT("*")))
-					.ToolTipText(ScriptViewModel.ToSharedRef(), &FNiagaraScratchPadScriptViewModel::GetToolTip)
+					.ToolTipText(this, &SNiagaraScratchPadScriptEditor::GetNameToolTipText)
 				]
 			]
 			+ SVerticalBox::Slot()
 			[
-				SAssignNew(Graph, SNiagaraScriptGraph, ScriptViewModel->GetGraphViewModel())
+				SAssignNew(Graph, SNiagaraScriptGraph, InScriptViewModel->GetGraphViewModel())
 				.ZoomToFitOnLoad(true)
 				.ShowHeader(false)
 			]
 		];
 
+	SetViewModel(InScriptViewModel);
+}
+
+void SNiagaraScratchPadScriptEditor::SetViewModel(TSharedPtr<FNiagaraScratchPadScriptViewModel> InScriptViewModel)
+{
+	ClearHandles();
+
+	ScriptViewModel = InScriptViewModel;
+
 	if (ScriptViewModel)
 	{
+		if (Graph.IsValid() && Graph->GetViewModel() != InScriptViewModel->GetGraphViewModel())
+		{
+			Graph->UpdateViewModel(InScriptViewModel->GetGraphViewModel());
+		}
 		NodeIDHandle = ScriptViewModel->OnNodeIDFocusRequested().AddLambda(
 			[this](FNiagaraScriptIDAndGraphFocusInfo* FocusInfo)
 			{
@@ -109,6 +121,7 @@ void SNiagaraScratchPadScriptEditor::Construct(const FArguments& InArgs, TShared
 	}
 }
 
+
 void FNiagaraGraphTabHistory::EvokeHistory(TSharedPtr<FTabInfo> InTabInfo, bool bPrevTabMatches)
 {
 	FWorkflowTabSpawnInfo SpawnInfo;
diff --git a/Engine/Plugins/FX/Niagara/Source/NiagaraEditor/Private/Widgets/NiagaraScratchScriptEditor.h b/Engine/Plugins/FX/Niagara/Source/NiagaraEditor/Private/Widgets/NiagaraScratchScriptEditor.h
--- a/Engine/Plugins/FX/Niagara/Source/NiagaraEditor/Private/Widgets/NiagaraScratchScriptEditor.h
+++ b/Engine/Plugins/FX/Niagara/Source/NiagaraEditor/Private/Widgets/NiagaraScratchScriptEditor.h
@@ -23,11 +23,7 @@ class SNiagaraScratchPadScriptEditor : public SCompoundWidget
 
 	~SNiagaraScratchPadScriptEditor()
 	{
-		if (ScriptViewModel)
-		{
-			ScriptViewModel->OnNodeIDFocusRequested().Remove(NodeIDHandle);
-			ScriptViewModel->OnPinIDFocusRequested().Remove(PinIDHandle);
-		}
+		ClearHandles();
 	}
 
 	TSharedPtr<SGraphEditor> GetGraphEditor()
@@ -43,13 +39,30 @@ class SNiagaraScratchPadScriptEditor : public SCompoundWidget
 		return ScriptViewModel;
 	}
 
+	void SetViewModel(TSharedPtr<FNiagaraScratchPadScriptViewModel> InViewModel);
+
 
 private:
+
+	void ClearHandles()
+	{
+		if (ScriptViewModel)
+		{
+			ScriptViewModel->OnNodeIDFocusRequested().Remove(NodeIDHandle);
+			ScriptViewModel->OnPinIDFocusRequested().Remove(PinIDHandle);
+		}
+	}
+
 	FText GetNameText() const
 	{
 		return ScriptViewModel->GetDisplayName();
 	}
 
+	FText GetNameToolTipText() const
+	{
+		return ScriptViewModel->GetToolTip();
+	}
+
 	EVisibility GetUnappliedChangesVisibility() const
 	{
 		return ScriptViewModel->HasUnappliedChanges() ? EVisibility::Visible : EVisibility::Collapsed;
diff --git a/Engine/Plugins/FX/Niagara/Source/NiagaraEditor/Private/Widgets/SNiagaraParameterDefinitionsPanel.cpp b/Engine/Plugins/FX/Niagara/Source/NiagaraEditor/Private/Widgets/SNiagaraParameterDefinitionsPanel.cpp
--- a/Engine/Plugins/FX/Niagara/Source/NiagaraEditor/Private/Widgets/SNiagaraParameterDefinitionsPanel.cpp
+++ b/Engine/Plugins/FX/Niagara/Source/NiagaraEditor/Private/Widgets/SNiagaraParameterDefinitionsPanel.cpp
@@ -258,7 +258,7 @@ FReply SNiagaraParameterDefinitionsPanel::OnParameterItemsDragged(const TArray<F
 	return ParameterDefinitionsPanelViewModel->OnParameterItemsDragged(DraggedItems, MouseEvent);
 }
 
-void SNiagaraParameterDefinitionsPanel::Refresh()
+void SNiagaraParameterDefinitionsPanel::Refresh(bool bExpandCategories)
 {
 	ItemSelector->RefreshItemsAndDefaultCategories(GetViewedParameterDefinitionsItems(), GetDefaultCategories()); 
 }
diff --git a/Engine/Plugins/FX/Niagara/Source/NiagaraEditor/Private/Widgets/SNiagaraParameterDefinitionsPanel.h b/Engine/Plugins/FX/Niagara/Source/NiagaraEditor/Private/Widgets/SNiagaraParameterDefinitionsPanel.h
--- a/Engine/Plugins/FX/Niagara/Source/NiagaraEditor/Private/Widgets/SNiagaraParameterDefinitionsPanel.h
+++ b/Engine/Plugins/FX/Niagara/Source/NiagaraEditor/Private/Widgets/SNiagaraParameterDefinitionsPanel.h
@@ -59,7 +59,7 @@ public:
 	void OnCategoryActivated(const FNiagaraParameterDefinitionsPanelCategory& ActivatedCategory) const;
 	//~ End SItemSelector
 
-	void Refresh();
+	void Refresh(bool bExpandCategories);
 	void AddParameterDefinitions(UNiagaraParameterDefinitions* NewParameterDefinitions) const;
 	void OnParameterItemSelected(const FNiagaraParameterDefinitionsPanelItem& SelectedItem, ESelectInfo::Type SelectInfo) const;
 	FReply OnParameterItemsDragged(const TArray<FNiagaraParameterDefinitionsPanelItem>& DraggedItems, const FPointerEvent& MouseEvent) const;
diff --git a/Engine/Plugins/FX/Niagara/Source/NiagaraEditor/Private/Widgets/SNiagaraParameterPanel.cpp b/Engine/Plugins/FX/Niagara/Source/NiagaraEditor/Private/Widgets/SNiagaraParameterPanel.cpp
--- a/Engine/Plugins/FX/Niagara/Source/NiagaraEditor/Private/Widgets/SNiagaraParameterPanel.cpp
+++ b/Engine/Plugins/FX/Niagara/Source/NiagaraEditor/Private/Widgets/SNiagaraParameterPanel.cpp
@@ -70,6 +70,12 @@ void SNiagaraParameterPanel::Construct(const FArguments& InArgs, const TSharedPt
 	ParameterPanelViewModel->GetOnNotifyParameterPendingNamespaceModifierRenameDelegate().BindSP(this, &SNiagaraParameterPanel::AddParameterPendingNamespaceModifierRename);
 	ParameterPanelViewModel->GetParametersWithNamespaceModifierRenamePendingDelegate().BindSP(this, &SNiagaraParameterPanel::GetParametersWithNamespaceModifierRenamePending);
 
+	if (ParameterPanelViewModel->UsesCategoryFilteringForInitialExpansion())
+	{
+		FilterCategoryExpandedDelegate.BindRaw(ParameterPanelViewModel.Get(), &INiagaraImmutableParameterPanelViewModel::IsCategoryExpandedByDefault);
+	}
+
+
 	SAssignNew(ItemSelector, SNiagaraParameterPanelSelector)
 	.PreserveSelectionOnRefresh(true)
 	.PreserveExpansionOnRefresh(true)
@@ -91,11 +97,13 @@ void SNiagaraParameterPanel::Construct(const FArguments& InArgs, const TSharedPt
 	.ClearSelectionOnClick(true)
 	.CategoryRowStyle(FNiagaraEditorStyle::Get(), "NiagaraEditor.Parameters.TableRow")
 	.OnGetCategoryBackgroundImage(this, &SNiagaraParameterPanel::GetCategoryBackgroundImage)
+	.OnItemExpandedInitially(FilterCategoryExpandedDelegate)
 	.CategoryBorderBackgroundColor(FLinearColor(.6, .6, .6, 1.0f))
 	.CategoryChildSlotPadding(FMargin(0.0f, 2.0f, 0.0f, 0.0f))
 	.CategoryBorderBackgroundPadding(FMargin(0.0f, 3.0f))
 	.OnGetKeyForItem(this, &SNiagaraParameterPanel::OnGetKeyForItem)
-	.OnGetKeyForCategory(this, &SNiagaraParameterPanel::OnGetKeyForCategory);
+	.OnGetKeyForCategory(this, &SNiagaraParameterPanel::OnGetKeyForCategory)
+	.ExpandInitially(false);
 
 	// Finalize the widget
 	ChildSlot
@@ -144,19 +152,20 @@ void SNiagaraParameterPanel::ConstructSectionButtons()
 	SectionSelectorBox->ClearChildren();
 
 	SectionSelectorBox->SetVisibility(ParameterPanelViewModel->GetShowSections() ? EVisibility::Visible : EVisibility::Collapsed);
-	for (FText Section : ParameterPanelViewModel->GetSections())
+	for (const INiagaraParameterPanelViewModel::FSectionDesc& Section : ParameterPanelViewModel->GetSections())
 	{
 		SectionSelectorBox->AddSlot()
 			[
 				SNew(SCheckBox)
 				.Style(FAppStyle::Get(), "DetailsView.SectionButton")
-				.OnCheckStateChanged(this, &SNiagaraParameterPanel::OnSectionChecked, Section)
-				.IsChecked(this, &SNiagaraParameterPanel::GetSectionCheckState, Section)
-				.ToolTipText(ParameterPanelViewModel->GetTooltipForSection(Section))
+				.OnCheckStateChanged(this, &SNiagaraParameterPanel::OnSectionChecked, Section.DisplayName)
+				.IsChecked(this, &SNiagaraParameterPanel::GetSectionCheckState, Section.DisplayName)
+				.IsEnabled(this, &SNiagaraParameterPanel::GetSectionEnabled, Section.DisplayName)
+				.ToolTipText(Section.Description)
 				[
 					SNew(STextBlock)
 					.TextStyle(FAppStyle::Get(), "SmallText")
-					.Text(Section)
+					.Text(Section.DisplayName)
 				]
 			];
 	}
@@ -168,10 +177,19 @@ ECheckBoxState SNiagaraParameterPanel::GetSectionCheckState(FText Section) const
 }
 
 
+bool SNiagaraParameterPanel::GetSectionEnabled(FText Section) const
+{
+	return ParameterPanelViewModel->GetSectionEnabled(Section);
+}
+
+
 void SNiagaraParameterPanel::OnSectionChecked(ECheckBoxState CheckState, FText Section)
 {
 	if (CheckState == ECheckBoxState::Checked)
 	{
+		TArray<FNiagaraParameterPanelCategory> ExpandedItems;
+		ItemSelector->GetExpandedCategoryItems(ExpandedItems);
+		ParameterPanelViewModel->PreSectionChange(ExpandedItems);
 		ParameterPanelViewModel->SetActiveSection(Section);
 		Refresh();
 	}
@@ -504,7 +522,8 @@ void SNiagaraParameterPanel::Tick(const FGeometry& AllottedGeometry, const doubl
 		if (bPendingRefresh)
 		{
 			bPendingRefresh = false;
-			Refresh();
+			Refresh(bRunCategoryExpansionFilter);
+			bRunCategoryExpansionFilter = false;
 		}
 		else if (bParameterItemsPendingChange)
 		{
@@ -643,16 +662,27 @@ bool SNiagaraParameterPanel::OnParameterNameTextVerifyChanged(const FText& InNew
 	return ParameterPanelViewModel->GetCanRenameParameterAndToolTip(ItemToBeRenamed, InNewText, bCheckEmptyNameText, OutErrorMessage);
 }
 
-void SNiagaraParameterPanel::Refresh()
+void SNiagaraParameterPanel::Refresh(bool bInRunCategoryExpansionFilter)
 {
 	ItemSelector->RefreshItemsAndDefaultCategories(ParameterPanelViewModel->GetViewedParameterItems(), ParameterPanelViewModel->GetDefaultCategories());
+	if (bInRunCategoryExpansionFilter && FilterCategoryExpandedDelegate.IsBound() && ParameterPanelViewModel->UsesCategoryFilteringForInitialExpansion())
+	{	
+		ItemSelector->ExpandTreeByFilter(FilterCategoryExpandedDelegate);
+	}
 }
 
-void SNiagaraParameterPanel::RefreshNextTick()
+void SNiagaraParameterPanel::RefreshNextTick(bool bInRunCategoryExpansionFilter)
 {
 	bPendingRefresh = true;
+
+	// We want to make sure that if anyone requested the expansion filter, it doesn't get cleared out, so only set if it isn't already true...
+	if (!bRunCategoryExpansionFilter)
+	{
+		bRunCategoryExpansionFilter = bInRunCategoryExpansionFilter;
+	}
 }
 
+
 void SNiagaraParameterPanel::SelectParameterItemByName(const FName ParameterName) const
 {
 	for (const FNiagaraParameterPanelItem& Item : ParameterPanelViewModel->GetCachedViewedParameterItems())
@@ -660,8 +690,9 @@ void SNiagaraParameterPanel::SelectParameterItemByName(const FName ParameterName
 		if (Item.GetVariable().GetName() == ParameterName)
 		{
 			const TArray<FNiagaraParameterPanelItem> ItemsToSelect = { Item };
-			ItemSelector->SetSelectedItems(ItemsToSelect);
+			ItemSelector->SetSelectedItems(ItemsToSelect, true);
 			ItemSelector->RequestScrollIntoView(Item);
+			
 			return;
 		}
 	}
diff --git a/Engine/Plugins/FX/Niagara/Source/NiagaraEditor/Private/Widgets/SNiagaraScriptGraph.cpp b/Engine/Plugins/FX/Niagara/Source/NiagaraEditor/Private/Widgets/SNiagaraScriptGraph.cpp
--- a/Engine/Plugins/FX/Niagara/Source/NiagaraEditor/Private/Widgets/SNiagaraScriptGraph.cpp
+++ b/Engine/Plugins/FX/Niagara/Source/NiagaraEditor/Private/Widgets/SNiagaraScriptGraph.cpp
@@ -68,12 +68,24 @@ SNiagaraScriptGraph::~SNiagaraScriptGraph()
 
 void SNiagaraScriptGraph::UpdateViewModel(TSharedRef<FNiagaraScriptGraphViewModel> InNewModel)
 {
+	bool bKeepOldGraphView = false;
+	FVector2D Location;
+	float ZoomAmount =1.0f;
+	FGuid BookmarkId = FGuid();
+
 	// remove old listeners
 	if (ViewModel)
 	{
 		ViewModel->GetNodeSelection()->OnSelectedObjectsChanged().RemoveAll(this);
 		ViewModel->OnNodesPasted().RemoveAll(this);
 		ViewModel->OnGraphChanged().RemoveAll(this);
+
+		if (GraphEditor.IsValid())
+		{
+			bKeepOldGraphView = true;
+			GraphEditor->GetViewLocation(Location, ZoomAmount);
+			GraphEditor->GetViewBookmark(BookmarkId);
+		}
 	}
 
 	// set model and listeners
@@ -81,6 +93,13 @@ void SNiagaraScriptGraph::UpdateViewModel(TSharedRef<FNiagaraScriptGraphViewMode
 	ViewModel->GetNodeSelection()->OnSelectedObjectsChanged().AddSP(this, &SNiagaraScriptGraph::ViewModelSelectedNodesChanged);
 	ViewModel->OnNodesPasted().AddSP(this, &SNiagaraScriptGraph::NodesPasted);
 	ViewModel->OnGraphChanged().AddSP(this, &SNiagaraScriptGraph::GraphChanged);
+
+	RecreateGraphWidget();
+
+	if (bKeepOldGraphView && GraphEditor.IsValid())
+	{
+		GraphEditor->SetViewLocation(Location, ZoomAmount, BookmarkId);
+	}
 }
 
 void SNiagaraScriptGraph::RecreateGraphWidget()
@@ -271,6 +290,7 @@ TSharedRef<SGraphEditor> SNiagaraScriptGraph::ConstructGraphEditor()
 	return CreatedGraphEditor;
 }
 
+
 void SNiagaraScriptGraph::ViewModelSelectedNodesChanged()
 {
 	if (FNiagaraEditorUtilities::SetsMatch(GraphEditor->GetSelectedNodes(), ViewModel->GetNodeSelection()->GetSelectedObjects()) == false)
diff --git a/Engine/Plugins/FX/Niagara/Source/NiagaraEditor/Public/NiagaraEditorSettings.h b/Engine/Plugins/FX/Niagara/Source/NiagaraEditor/Public/NiagaraEditorSettings.h
--- a/Engine/Plugins/FX/Niagara/Source/NiagaraEditor/Public/NiagaraEditorSettings.h
+++ b/Engine/Plugins/FX/Niagara/Source/NiagaraEditor/Public/NiagaraEditorSettings.h
@@ -212,6 +212,26 @@ struct NIAGARAEDITOR_API FNiagaraActionColors
 	FLinearColor DeveloperColor;
 };
 
+
+USTRUCT()
+struct NIAGARAEDITOR_API FNiagaraParameterPanelSectionStorage
+{
+	GENERATED_BODY()
+		
+	FNiagaraParameterPanelSectionStorage() {}
+
+	FNiagaraParameterPanelSectionStorage(const FGuid& InGuid) : ParamStorageId(InGuid) 
+	{
+		check(InGuid.IsValid());
+	}
+
+	UPROPERTY()
+	FGuid ParamStorageId;
+
+	UPROPERTY()
+	TArray<FGuid> ExpandedCategories;
+};
+
 FORCEINLINE uint32 GetTypeHash(const FNiagaraNamespaceMetadata& NamespaceMetaData)
 {
 	return GetTypeHash(NamespaceMetaData.GetGuid());
@@ -337,6 +357,9 @@ public:
 	const TMap<FString, FString>& GetHLSLKeywordReplacementsMap()const { return HLSLKeywordReplacements; }
 
 	FLinearColor GetSourceColor(EScriptSource Source) const;
+
+	FNiagaraParameterPanelSectionStorage& FindOrAddParameterPanelSectionStorage(FGuid PanelSectionId, bool& bOutAdded);
+
 private:
 	void SetupNamespaceMetadata();
 	void BuildCachedPlaybackSpeeds() const;
@@ -421,6 +444,10 @@ private:
 	UPROPERTY(config)
 	bool bShowGpuTickInformation;
 
+
+	UPROPERTY(config)
+	TArray<FNiagaraParameterPanelSectionStorage> SystemParameterPanelSectionData;
+
 public:
 	bool IsShowInstructionsCount() const;
 	void SetShowInstructionsCount(bool bShowInstructionsCount);
diff --git a/Engine/Plugins/FX/Niagara/Source/NiagaraEditor/Public/ViewModels/NiagaraParameterPanelViewModel.h b/Engine/Plugins/FX/Niagara/Source/NiagaraEditor/Public/ViewModels/NiagaraParameterPanelViewModel.h
--- a/Engine/Plugins/FX/Niagara/Source/NiagaraEditor/Public/ViewModels/NiagaraParameterPanelViewModel.h
+++ b/Engine/Plugins/FX/Niagara/Source/NiagaraEditor/Public/ViewModels/NiagaraParameterPanelViewModel.h
@@ -39,8 +39,10 @@ namespace FNiagaraParameterUtilities
 // NOTE: These utilities are not defined in the view model directly as they are shared between ParameterPanelViewModel and ParameterDefinitionsPanelViewModel.
 namespace FNiagaraSystemToolkitParameterPanelUtilities
 {
-	TArray<UNiagaraGraph*> GetAllGraphs(const TSharedPtr<FNiagaraSystemViewModel>& SystemViewModel);
-	TArray<UNiagaraGraph*> GetEditableGraphs(const TSharedPtr<FNiagaraSystemViewModel>& SystemViewModel, const TWeakPtr<FNiagaraSystemGraphSelectionViewModel>& SystemGraphSelectionViewModelWeak);
+	// Get all graphs for the active emitter selection or just the one for the active script mode
+	TArray<UNiagaraGraph*> GetAllGraphs(const TSharedPtr<FNiagaraSystemViewModel>& SystemViewModel, bool bActiveScriptMode);
+	// Get all editable graphs for the active emitter selection or just the one for the active script mode
+	TArray<UNiagaraGraph*> GetEditableGraphs(const TSharedPtr<FNiagaraSystemViewModel>& SystemViewModel, const TWeakPtr<FNiagaraSystemGraphSelectionViewModel>& SystemGraphSelectionViewModelWeak, bool bActiveScriptMode);
 	FReply CreateDragEventForParameterItem(
 		const FNiagaraParameterPanelItemBase& DraggedItem,
 		const FPointerEvent& MouseEvent,
@@ -80,8 +82,8 @@ class NIAGARAEDITOR_API INiagaraImmutableParameterPanelViewModel : public TShare
 {
 public:
 	/** Delegate to signal the view model's state has changed. */
-	DECLARE_DELEGATE(FOnRequestRefresh);
-	DECLARE_DELEGATE(FOnRequestRefreshNextTick);
+	DECLARE_DELEGATE_OneParam(FOnRequestRefresh, bool);
+	DECLARE_DELEGATE_OneParam(FOnRequestRefreshNextTick, bool);
 
 	virtual ~INiagaraImmutableParameterPanelViewModel() { }
 
@@ -109,13 +111,31 @@ public:
 
 	virtual bool GetCanCopyParameterMetaDataAndToolTip(const FNiagaraParameterPanelItemBase& ItemToCopy, FText& OutCanCopyToolTip) const;
 
-	virtual void Refresh() const;
+	// Refresh often encapsulates GUI rebuilding, so we have multiple API's for very specific use cases within the view model.
 
-	virtual void RefreshNextTick() const;
+	// Plain refresh is a relatively fast refresh
+	void Refresh() const { RefreshFull(false); };
+	
+	// RefreshFull allows for optional category tree expansion/collapse, which can be more expensive.
+	virtual void RefreshFull(bool bDoCategoryExpansion) const;
+
+	// Enqueue light refresh next tick
+	void RefreshNextTick() const { RefreshFullNextTick(false); }
+	
+	// Equeue a full refresh, parameterized by inputs next tick
+	virtual void RefreshFullNextTick(bool bDoCategoryExpansion) const;
+
+	// Do an exhaustive refresh as context may be swapped out entirely
+	virtual void RefreshDueToActiveDocumentChanged() {};
 
 	FOnRequestRefresh& GetOnRequestRefreshDelegate() { return OnRequestRefreshDelegate; };
 	FOnRequestRefreshNextTick& GetOnRequestRefreshNextTickDelegate() { return OnRequestRefreshNextTickDelegate; };
 
+	// API for interacting with category expand/collapse within the view model
+	virtual bool UsesCategoryFilteringForInitialExpansion() const { return false; }
+	virtual bool IsCategoryExpandedByDefault(const FNiagaraParameterPanelCategory& Category) const { return false; };
+	virtual void PreSectionChange(const TArray<FNiagaraParameterPanelCategory>& ExpandedItems) {};
+
 protected:
 	FOnRequestRefresh OnRequestRefreshDelegate;
 	FOnRequestRefreshNextTick OnRequestRefreshNextTickDelegate;
@@ -124,7 +144,17 @@ protected:
 /** Interface for view models to SiagaraParameterPanel. */
 class INiagaraParameterPanelViewModel : public INiagaraImmutableParameterPanelViewModel
 {
+	
 public:
+	/** Helper struct for storing information about a section in the parameter panel (if used)*/
+	struct FSectionDesc
+	{
+		FSectionDesc(const FGuid& InId, const FText& InName, const FText& InDesc) : SectionId(InId), DisplayName(InName), Description(InDesc) {}
+		FGuid SectionId;
+		FText DisplayName;
+		FText Description;
+	};
+
 	/** Delegate to handle responses to external selection changes (e.g. TNiagaraSelection changes.) */
 	DECLARE_MULTICAST_DELEGATE_OneParam(FOnParameterPanelViewModelExternalSelectionChanged, const UObject*);
 
@@ -154,17 +184,17 @@ public:
 	virtual const TArray<UNiagaraParameterDefinitions*> GetAvailableParameterDefinitions(bool bSkipSubscribedParameterDefinitions) const = 0;
 	//~ End INiagaraImmutableParameterPanelViewModel interface
 
-	virtual void AddParameter(FNiagaraVariable NewVariable, const FNiagaraParameterPanelCategory Category, const bool bRequestRename, const bool bMakeUniqueName) const = 0;
+	virtual void AddParameter(FNiagaraVariable NewVariable, const FNiagaraParameterPanelCategory Category, const bool bRequestRename, const bool bMakeUniqueName)  = 0;
 
-	virtual void FindOrAddParameter(FNiagaraVariable Variable, const FNiagaraParameterPanelCategory Category) const = 0;
+	virtual void FindOrAddParameter(FNiagaraVariable Variable, const FNiagaraParameterPanelCategory Category)  = 0;
 	
 	virtual bool GetCanAddParametersToCategory(FNiagaraParameterPanelCategory Category) const = 0;
 
-	virtual void DeleteParameters(const TArray<FNiagaraParameterPanelItem>& ItemsToDelete) const = 0;
+	virtual void DeleteParameters(const TArray<FNiagaraParameterPanelItem>& ItemsToDelete)  = 0;
 
-	virtual void RenameParameter(const FNiagaraParameterPanelItem& ItemToRename, const FName NewName) const = 0;
+	virtual void RenameParameter(const FNiagaraParameterPanelItem& ItemToRename, const FName NewName)  = 0;
 
-	virtual void SetParameterIsSubscribedToLibrary(const UNiagaraScriptVariable* ScriptVarToModify, const bool bSubscribed) const = 0;
+	virtual void SetParameterIsSubscribedToLibrary(const UNiagaraScriptVariable* ScriptVarToModify, const bool bSubscribed)  = 0;
 
 	virtual TSharedPtr<SWidget> CreateContextMenuForItems(const TArray<FNiagaraParameterPanelItem>& Items, const TSharedPtr<FUICommandList>& ToolkitCommands) = 0;
 
@@ -176,9 +206,11 @@ public:
 
 	virtual const TArray<FNiagaraParameterPanelCategory>& GetDefaultCategories() const = 0;
 
-	virtual FMenuAndSearchBoxWidgets GetParameterMenu(FNiagaraParameterPanelCategory Category) const = 0;
+	virtual FMenuAndSearchBoxWidgets GetParameterMenu(FNiagaraParameterPanelCategory Category)  = 0;
 
-	virtual FReply HandleDragDropOperation(TSharedPtr<FDragDropOperation> DropOperation) const = 0;
+	virtual FReply HandleDragDropOperation(TSharedPtr<FDragDropOperation> DropOperation)  = 0;
+
+	virtual bool GetSectionEnabled(FText Section) const { return true; }
 
 	virtual bool GetCanHandleDragDropOperation(TSharedPtr<FDragDropOperation> DragDropOperation) const = 0;
 
@@ -191,9 +223,9 @@ public:
 
 	virtual bool GetCanDeleteParameterAndToolTip(const FNiagaraParameterPanelItem& ItemToDelete, FText& OutCanDeleteParameterToolTip) const;
 
-	virtual void ChangeParameterType(const TArray<FNiagaraParameterPanelItem> ItemsToModify, const FNiagaraTypeDefinition NewType) const;
+	virtual void ChangeParameterType(const TArray<FNiagaraParameterPanelItem> ItemsToModify, const FNiagaraTypeDefinition NewType) ;
 	virtual bool GetCanChangeParameterType(const TArray<FNiagaraParameterPanelItem>& ItemToChange, FText& OutTooltip) const;
-	virtual void GetChangeTypeSubMenu(FMenuBuilder& MenuBuilder, TArray<FNiagaraParameterPanelItem> Item) const;
+	virtual void GetChangeTypeSubMenu(FMenuBuilder& MenuBuilder, TArray<FNiagaraParameterPanelItem> Item) ;
 
 	virtual void PasteParameterMetaData(const TArray<FNiagaraParameterPanelItem> SelectedItems);
 	virtual bool GetCanPasteParameterMetaDataAndToolTip(FText& OutCanPasteToolTip);
@@ -201,7 +233,7 @@ public:
 	virtual bool GetAllowMultiSelect() const { return false; }
 	virtual bool GetShowReferenceCounter() const { return true; }
 
-	virtual void DuplicateParameters(const TArray<FNiagaraParameterPanelItem> ItemToDuplicate) const;
+	virtual void DuplicateParameters(const TArray<FNiagaraParameterPanelItem> ItemToDuplicate) ;
 	virtual bool GetCanDuplicateParameterAndToolTip(const TArray<FNiagaraParameterPanelItem>& ItemsToDuplicate, FText& OutCanDuplicateParameterToolTip) const;
 	bool GetCanDebugParameters(const TArray<FNiagaraParameterPanelItem>& ItemsToDebug) const;
 
@@ -209,17 +241,17 @@ public:
 
 	virtual bool GetCanSubscribeParameterToLibraryAndToolTip(const FNiagaraParameterPanelItem& ItemToModify, const bool bSubscribing, FText& OutCanSubscribeParameterToolTip) const;
 
-	virtual void SetParameterIsSubscribedToLibrary(const FNiagaraParameterPanelItem ItemToModify, const bool bSubscribed) const;
+	virtual void SetParameterIsSubscribedToLibrary(const FNiagaraParameterPanelItem ItemToModify, const bool bSubscribed) ;
 
-	virtual void SetParameterNamespace(const FNiagaraParameterPanelItem ItemToModify, FNiagaraNamespaceMetadata NewNamespaceMetaData, bool bDuplicateParameter) const;
+	virtual void SetParameterNamespace(const FNiagaraParameterPanelItem ItemToModify, FNiagaraNamespaceMetadata NewNamespaceMetaData, bool bDuplicateParameter) ;
 
-	virtual void SetParameterNamespaceModifier(const FNiagaraParameterPanelItem ItemToModify, const FName NewNamespaceModifier, bool bDuplicateParameter) const;
+	virtual void SetParameterNamespaceModifier(const FNiagaraParameterPanelItem ItemToModify, const FName NewNamespaceModifier, bool bDuplicateParameter) ;
 
-	virtual void SetParameterCustomNamespaceModifier(const FNiagaraParameterPanelItem ItemToModify, bool bDuplicateParameter) const;
+	virtual void SetParameterCustomNamespaceModifier(const FNiagaraParameterPanelItem ItemToModify, bool bDuplicateParameter) ;
 
-	virtual void GetChangeNamespaceSubMenu(FMenuBuilder& MenuBuilder, bool bDuplicateParameter, FNiagaraParameterPanelItem Item) const;
+	virtual void GetChangeNamespaceSubMenu(FMenuBuilder& MenuBuilder, bool bDuplicateParameter, FNiagaraParameterPanelItem Item);
 
-	virtual void GetChangeNamespaceModifierSubMenu(FMenuBuilder& MenuBuilder, bool bDuplicateParameter, FNiagaraParameterPanelItem Item) const;
+	virtual void GetChangeNamespaceModifierSubMenu(FMenuBuilder& MenuBuilder, bool bDuplicateParameter, FNiagaraParameterPanelItem Item) ;
 
 	virtual void OnParameterItemSelected(const FNiagaraParameterPanelItem& SelectedItem, ESelectInfo::Type SelectInfo) const {};
 
@@ -229,10 +261,11 @@ public:
 
 	const TArray<FNiagaraParameterPanelItem>& GetCachedViewedParameterItems() const;
 
-	void SelectParameterItemByName(const FName ParameterName, const bool bRequestRename) const;
+	void SelectParameterItemByName(const FName ParameterName, const bool bRequestRename);
 
-	void SubscribeParameterToLibraryIfMatchingDefinition(const UNiagaraScriptVariable* ScriptVarToModify, const FName ScriptVarName) const;
+	void SubscribeParameterToLibraryIfMatchingDefinition(const UNiagaraScriptVariable* ScriptVarToModify, const FName ScriptVarName) ;
 
+	virtual void RefreshFull(bool bDoCategoryExpansion) const;
 
 	FOnParameterPanelViewModelExternalSelectionChanged& GetOnExternalSelectionChangedDelegate() { return OnParameterPanelViewModelExternalSelectionChangedDelegate; };
 
@@ -244,13 +277,16 @@ public:
 
 	FText GetActiveSection() const
 	{
-		return Sections.IsValidIndex(ActiveSectionIndex) ? Sections[ActiveSectionIndex] : FText::GetEmpty();
+		return Sections.IsValidIndex(ActiveSectionIndex) ? Sections[ActiveSectionIndex].DisplayName : FText::GetEmpty();
 	}
 	virtual bool GetShowSections() const { return false; }
 	virtual bool GetNamespaceActive(const FName& InNamespace) const { return true; }
-	virtual const TArray<FText>& GetSections() const { return Sections; }
+	virtual const TArray<FSectionDesc>& GetSections() const { return Sections; }
 	virtual FText GetTooltipForSection(FText& InSection) const;
-	virtual void SetActiveSection(FText& InSection);
+	void SetActiveSection(FText& InSection);
+	virtual void SetActiveSection(int32 InSection);
+	virtual void SetMasterParameterPanelViewModel(TWeakPtr<INiagaraParameterPanelViewModel> InMasterParameterPanelViewModel) { MasterParameterPanelViewModel = InMasterParameterPanelViewModel; }
+
 
 	virtual bool IsVariableSelected(FNiagaraVariableBase& InVar) const { return false; };
 
@@ -261,6 +297,7 @@ protected:
 
 
 protected:
+	TWeakPtr<INiagaraParameterPanelViewModel> MasterParameterPanelViewModel;
 	FOnParameterPanelViewModelExternalSelectionChanged OnParameterPanelViewModelExternalSelectionChangedDelegate;
 	FOnSelectParameterItemByName OnSelectParameterItemByNameDelegate;
 	FOnNotifyParameterPendingRename OnNotifyParameterPendingRenameDelegate;
@@ -280,8 +317,7 @@ protected:
 	/** Transient UNiagaraScriptVariables used to pass to new FNiagaraParameterPanelItems when the source FNiagaraVariable is not associated with a UNiagaraScriptVariable in a graph. */
 	mutable TMap<FNiagaraVariable, TObjectPtr<UNiagaraScriptVariable>> TransientParameterToScriptVarMap;
 
-	TArray<FText> Sections;
-	TArray<FText> SectionDescriptions;
+	TArray<FSectionDesc> Sections;
 	int32 ActiveSectionIndex;
 };
 
@@ -307,17 +343,17 @@ public:
 	//~ End INiagaraImmutableParameterPanelViewModel interface
 
 	//~ Begin INiagaraParameterPanelViewModel interface
-	virtual void AddParameter(FNiagaraVariable NewVariable, const FNiagaraParameterPanelCategory Category, const bool bRequestRename, const bool bMakeUniqueName) const override;
+	virtual void AddParameter(FNiagaraVariable NewVariable, const FNiagaraParameterPanelCategory Category, const bool bRequestRename, const bool bMakeUniqueName)  override;
 
-	virtual void FindOrAddParameter(FNiagaraVariable Variable, const FNiagaraParameterPanelCategory Category) const override;
+	virtual void FindOrAddParameter(FNiagaraVariable Variable, const FNiagaraParameterPanelCategory Category)  override;
 
 	virtual bool GetCanAddParametersToCategory(FNiagaraParameterPanelCategory Category) const override;
 
-	virtual void DeleteParameters(const TArray<FNiagaraParameterPanelItem>& ItemsToDelete) const override;
+	virtual void DeleteParameters(const TArray<FNiagaraParameterPanelItem>& ItemsToDelete)  override;
 
-	virtual void RenameParameter(const FNiagaraParameterPanelItem& ItemToRename, const FName NewName) const override;
+	virtual void RenameParameter(const FNiagaraParameterPanelItem& ItemToRename, const FName NewName)  override;
 
-	virtual void SetParameterIsSubscribedToLibrary(const UNiagaraScriptVariable* ScriptVarToModify, const bool bSubscribed) const override;
+	virtual void SetParameterIsSubscribedToLibrary(const UNiagaraScriptVariable* ScriptVarToModify, const bool bSubscribed)  override;
 
 	virtual FReply OnParameterItemsDragged(const TArray<FNiagaraParameterPanelItem>& DraggedItems, const FPointerEvent& MouseEvent) const override;
 
@@ -331,9 +367,9 @@ public:
 
 	virtual const TArray<FNiagaraParameterPanelCategory>& GetDefaultCategories() const override;
 
-	virtual FMenuAndSearchBoxWidgets GetParameterMenu(FNiagaraParameterPanelCategory Category) const override;
+	virtual FMenuAndSearchBoxWidgets GetParameterMenu(FNiagaraParameterPanelCategory Category)  override;
 
-	virtual FReply HandleDragDropOperation(TSharedPtr<FDragDropOperation> DragDropOperation) const override;
+	virtual FReply HandleDragDropOperation(TSharedPtr<FDragDropOperation> DragDropOperation)  override;
 
 	virtual bool GetCanHandleDragDropOperation(TSharedPtr<FDragDropOperation> DragDropOperation) const override;
 
@@ -344,12 +380,15 @@ public:
 	virtual bool GetCanSetParameterCustomNamespaceModifierAndToolTip(const FNiagaraParameterPanelItem& ItemToModify, bool bDuplicateParameter, FText& OutCanSetParameterNamespaceModifierToolTip) const override;
 	//~ End INiagaraParameterPanelViewModel interface
 
-	TSharedRef<SWidget> CreateAddParameterMenuForAssignmentNode(UNiagaraNodeAssignment* AssignmentNode, const TSharedPtr<SComboButton>& AddButton) const;
+	TSharedRef<SWidget> CreateAddParameterMenuForAssignmentNode(UNiagaraNodeAssignment* AssignmentNode, const TSharedPtr<SComboButton>& AddButton);
 
 	void DebugParameters(const TArray<FNiagaraParameterPanelItem> ItemToDuplicate) const;
 
 	virtual void OnParameterItemSelected(const FNiagaraParameterPanelItem& SelectedItem, ESelectInfo::Type SelectInfo) const override;
 	virtual bool IsVariableSelected(FNiagaraVariableBase& InVar) const override;
+	virtual void RefreshDueToActiveDocumentChanged() override;
+
+	virtual bool GetSectionEnabled(FText Section) const override;
 
 
 	virtual bool GetShowSections() const override { return true; }
@@ -357,7 +396,17 @@ public:
 
 	virtual TSharedPtr<FNiagaraObjectSelection> GetVariableObjectSelection() override { return VariableObjectSelection; };
 	virtual bool GetShowReferenceCounter() const override;
+
+	virtual void SetActiveSection(int32 InSection) override;
+
+	virtual bool IsCategoryExpandedByDefault(const FNiagaraParameterPanelCategory& Category) const override;
+	virtual bool UsesCategoryFilteringForInitialExpansion() const override { return true; }
+
+	virtual void PreSectionChange(const TArray<FNiagaraParameterPanelCategory>& ExpandedItems) override;
 private:
+
+	bool  ShouldRouteThroughScratchParameterMap(const FNiagaraParameterPanelCategory* Category = nullptr, const FNiagaraVariableBase* Variable = nullptr);
+
 	const TArray<UNiagaraGraph*> GetAllGraphsConst() const;
 
 	TArray<UNiagaraGraph*> GetEditableGraphs() const;
@@ -366,7 +415,7 @@ private:
 
 	TArray<FNiagaraEmitterHandle*> GetEditableEmitterHandles() const;
 
-	void AddScriptVariable(const UNiagaraScriptVariable* NewScriptVar) const;
+	void AddScriptVariable(const UNiagaraScriptVariable* NewScriptVar) ;
 
 	void AddParameterDefinitions(UNiagaraParameterDefinitions* NewParameterDefinitions) const;
 
@@ -396,11 +445,15 @@ private:
 	static TArray<FNiagaraParameterPanelCategory> DefaultCategories;
 	static TArray<FNiagaraParameterPanelCategory> UserCategories;
 	static TArray<FNiagaraParameterPanelCategory> DefaultAdvancedCategories;
+	static TArray<FNiagaraParameterPanelCategory> DefaultScriptCategories;
+	static TArray<FNiagaraParameterPanelCategory> DefaultAdvancedScriptCategories;
 
 	TSharedPtr<FNiagaraObjectSelection> VariableObjectSelection;
 	int32 UserOnlyIdx = -1;
 	int32 CommonOnlyIdx = -1;
-	int32 ActiveIdx = -1;
+	int32 ActiveScriptIdx = -1;
+	int32 ActiveSystemIdx = -1;
+	int32 LastActiveSystemSectionIdx = -1;
 
 	TNiagaraViewModelManager<UNiagaraSystem, FNiagaraSystemToolkitParameterPanelViewModel>::Handle RegisteredHandle;
 };
@@ -426,21 +479,21 @@ public:
 	//~ End INiagaraImmutableParameterPanelViewModel interface
 
 	//~ Begin INiagaraParameterPanelViewModel interface
-	virtual void AddParameter(FNiagaraVariable NewVariable, const FNiagaraParameterPanelCategory Category, const bool bRequestRename, const bool bMakeUniqueName) const override;
+	virtual void AddParameter(FNiagaraVariable NewVariable, const FNiagaraParameterPanelCategory Category, const bool bRequestRename, const bool bMakeUniqueName)  override;
 
-	virtual void FindOrAddParameter(FNiagaraVariable Variable, const FNiagaraParameterPanelCategory Category) const override;
+	virtual void FindOrAddParameter(FNiagaraVariable Variable, const FNiagaraParameterPanelCategory Category)  override;
 
 	virtual bool GetCanAddParametersToCategory(FNiagaraParameterPanelCategory Category) const override;
 
-	virtual void DeleteParameters(const TArray<FNiagaraParameterPanelItem>& ItemToDelete) const override;
+	virtual void DeleteParameters(const TArray<FNiagaraParameterPanelItem>& ItemToDelete)  override;
 
-	virtual void RenameParameter(const FNiagaraParameterPanelItem& ItemToRename, const FName NewName) const override;
+	virtual void RenameParameter(const FNiagaraParameterPanelItem& ItemToRename, const FName NewName)  override;
 	
-	virtual void ChangeParameterType(const TArray<FNiagaraParameterPanelItem> ItemsToModify, const FNiagaraTypeDefinition NewType) const override;
+	virtual void ChangeParameterType(const TArray<FNiagaraParameterPanelItem> ItemsToModify, const FNiagaraTypeDefinition NewType)  override;
 
-	virtual void DuplicateParameters(const TArray<FNiagaraParameterPanelItem> ItemsToDuplicate) const override;
+	virtual void DuplicateParameters(const TArray<FNiagaraParameterPanelItem> ItemsToDuplicate)  override;
 
-	virtual void SetParameterIsSubscribedToLibrary(const UNiagaraScriptVariable* ScriptVarToModify, const bool bSubscribed) const override;
+	virtual void SetParameterIsSubscribedToLibrary(const UNiagaraScriptVariable* ScriptVarToModify, const bool bSubscribed)  override;
 
 	virtual void OnParameterItemSelected(const FNiagaraParameterPanelItem& SelectedItem, ESelectInfo::Type SelectInfo) const override;
 
@@ -459,9 +512,9 @@ public:
 	virtual bool GetAllowMultiSelect() const override { return true; }
 
 
-	virtual FMenuAndSearchBoxWidgets GetParameterMenu(FNiagaraParameterPanelCategory Category) const override;
+	virtual FMenuAndSearchBoxWidgets GetParameterMenu(FNiagaraParameterPanelCategory Category)  override;
 
-	virtual FReply HandleDragDropOperation(TSharedPtr<FDragDropOperation> DragDropOperation) const override;
+	virtual FReply HandleDragDropOperation(TSharedPtr<FDragDropOperation> DragDropOperation)  override;
 
 	virtual bool GetCanHandleDragDropOperation(TSharedPtr<FDragDropOperation> DragDropOperation) const override;
 
@@ -472,16 +525,16 @@ public:
 	virtual bool GetCanSetParameterCustomNamespaceModifierAndToolTip(const FNiagaraParameterPanelItem& ItemToModify, bool bDuplicateParameter, FText& OutCanSetParameterNamespaceModifierToolTip) const override;
 	//~ End INiagaraParameterPanelViewModel interface
 
-	void RenameParameter(const UNiagaraScriptVariable* ScriptVarToRename, const FName NewName) const;
+	void RenameParameter(const UNiagaraScriptVariable* ScriptVarToRename, const FName NewName) ;
 
-	void RenameParameter(const FNiagaraVariable& VariableToRename, const FName NewName) const;
+	void RenameParameter(const FNiagaraVariable& VariableToRename, const FName NewName) ;
 
 private:
-	void SetParameterIsOverridingLibraryDefaultValue(const FNiagaraParameterPanelItem ItemToModify, const bool bOverriding) const;
+	void SetParameterIsOverridingLibraryDefaultValue(const FNiagaraParameterPanelItem ItemToModify, const bool bOverriding) ;
 
 	TArray<UNiagaraGraph*> GetEditableGraphs() const;
 
-	void AddScriptVariable(const UNiagaraScriptVariable* NewScriptVar) const;
+	void AddScriptVariable(const UNiagaraScriptVariable* NewScriptVar) ;
 
 	void AddParameterDefinitions(UNiagaraParameterDefinitions* NewParameterDefinitions) const;
 
@@ -531,17 +584,17 @@ public:
 	//~ End INiagaraImmutableParameterPanelViewModel interface
 
 	//~ Begin INiagaraParameterPanelViewModel interface
-	virtual void AddParameter(FNiagaraVariable NewVariable, const FNiagaraParameterPanelCategory Category, const bool bRequestRename, const bool bMakeUniqueName) const override;
+	virtual void AddParameter(FNiagaraVariable NewVariable, const FNiagaraParameterPanelCategory Category, const bool bRequestRename, const bool bMakeUniqueName)  override;
 
-	virtual void FindOrAddParameter(FNiagaraVariable Variable, const FNiagaraParameterPanelCategory Category) const override;
+	virtual void FindOrAddParameter(FNiagaraVariable Variable, const FNiagaraParameterPanelCategory Category)  override;
 
 	virtual bool GetCanAddParametersToCategory(FNiagaraParameterPanelCategory Category) const override;
 
-	virtual void DeleteParameters(const TArray<FNiagaraParameterPanelItem>& ItemToDelete) const override;
+	virtual void DeleteParameters(const TArray<FNiagaraParameterPanelItem>& ItemToDelete)  override;
 
-	virtual void RenameParameter(const FNiagaraParameterPanelItem& ItemToRename, const FName NewName) const override;
+	virtual void RenameParameter(const FNiagaraParameterPanelItem& ItemToRename, const FName NewName)  override;
 
-	virtual void SetParameterIsSubscribedToLibrary(const UNiagaraScriptVariable* ScriptVarToModify, const bool bSubscribed) const override;
+	virtual void SetParameterIsSubscribedToLibrary(const UNiagaraScriptVariable* ScriptVarToModify, const bool bSubscribed)  override;
 
 	virtual TSharedPtr<SWidget> CreateContextMenuForItems(const TArray<FNiagaraParameterPanelItem>& Items, const TSharedPtr<FUICommandList>& ToolkitCommands) override;
 
@@ -553,9 +606,9 @@ public:
 
 	virtual const TArray<FNiagaraParameterPanelCategory>& GetDefaultCategories() const override;
 
-	virtual FMenuAndSearchBoxWidgets GetParameterMenu(FNiagaraParameterPanelCategory Category) const override;
+	virtual FMenuAndSearchBoxWidgets GetParameterMenu(FNiagaraParameterPanelCategory Category)  override;
 
-	virtual FReply HandleDragDropOperation(TSharedPtr<FDragDropOperation> DropOperation) const override;
+	virtual FReply HandleDragDropOperation(TSharedPtr<FDragDropOperation> DropOperation)  override;
 
 	virtual bool GetCanHandleDragDropOperation(TSharedPtr<FDragDropOperation> DragDropOperation) const override;
 
diff --git a/Engine/Plugins/FX/Niagara/Source/NiagaraEditor/Public/ViewModels/NiagaraScratchPadScriptViewModel.h b/Engine/Plugins/FX/Niagara/Source/NiagaraEditor/Public/ViewModels/NiagaraScratchPadScriptViewModel.h
--- a/Engine/Plugins/FX/Niagara/Source/NiagaraEditor/Public/ViewModels/NiagaraScratchPadScriptViewModel.h
+++ b/Engine/Plugins/FX/Niagara/Source/NiagaraEditor/Public/ViewModels/NiagaraScratchPadScriptViewModel.h
@@ -19,12 +19,15 @@ public:
 	DECLARE_MULTICAST_DELEGATE(FOnChangesApplied);
 	DECLARE_MULTICAST_DELEGATE_OneParam(FOnNodeIDFocusRequested, FNiagaraScriptIDAndGraphFocusInfo*  /* FocusInfo */);
 	DECLARE_MULTICAST_DELEGATE_OneParam(FOnPinIDFocusRequested, FNiagaraScriptIDAndGraphFocusInfo*  /* FocusInfo */);
+	DECLARE_MULTICAST_DELEGATE_OneParam(FOnGraphSubObjectSelectionChanged, const UObject*);
 
 	FNiagaraScratchPadScriptViewModel(bool bInIsForDataProcessingOnly);
 
 	~FNiagaraScratchPadScriptViewModel();
 
-	void Initialize(UNiagaraScript* Script);
+	void Initialize(UNiagaraScript* Script, UNiagaraScript* InEditScript, TWeakPtr<class FNiagaraSystemViewModel> InSystemViewModel);
+
+	void TransferFromOldWhenDoingApply(TSharedPtr< FNiagaraScratchPadScriptViewModel> InOldScriptVM);
 
 	void Finalize();
 
@@ -39,6 +42,10 @@ protected:
 	virtual INiagaraParameterDefinitionsSubscriber* GetParameterDefinitionsSubscriber() override { return &EditScript; };
 	//~ End NiagaraParameterDefinitionsSubscriberViewModel Interface
 
+	void OnGraphSubObjectSelectionChanged(const UObject* Obj);
+
+	FDelegateHandle ExternalSelectionChangedDelegate;
+
 public:
 	virtual FString GetReferencerName() const override
 	{
@@ -88,6 +95,8 @@ public:
 	FOnNodeIDFocusRequested& OnNodeIDFocusRequested();
 	FOnPinIDFocusRequested& OnPinIDFocusRequested();
 
+	FOnGraphSubObjectSelectionChanged& OnGraphSelectionChanged() { return OnGraphSubObjectSelectionChangedDelegate;	};
+
 	void RaisePinFocusRequested(FNiagaraScriptIDAndGraphFocusInfo*  InFocusInfo);
 	void RaiseNodeFocusRequested(FNiagaraScriptIDAndGraphFocusInfo* InFocusInfo);
 
@@ -105,6 +114,8 @@ private:
 	UNiagaraScript* OriginalScript = nullptr;
 	FVersionedNiagaraScript EditScript;
 
+	TWeakPtr<FNiagaraSystemViewModel> SystemViewModel;
+
 	bool bHasPendingChanges;
 
 	TSharedPtr<FUICommandList> ParameterPanelCommands;
@@ -117,6 +128,7 @@ private:
 	FOnHasUnappliedChangesChanged OnHasUnappliedChangesChangedDelegate;
 	FOnChangesApplied OnChangesAppliedDelegate;
 	FSimpleDelegate	OnRequestDiscardChangesDelegate;
+	FOnGraphSubObjectSelectionChanged OnGraphSubObjectSelectionChangedDelegate;
 
 	FOnNodeIDFocusRequested OnNodeIDFocusRequestedDelegate;
 	FOnPinIDFocusRequested OnPinIDFocusRequestedDelegate;
diff --git a/Engine/Plugins/FX/Niagara/Source/NiagaraEditor/Public/ViewModels/NiagaraScratchPadViewModel.h b/Engine/Plugins/FX/Niagara/Source/NiagaraEditor/Public/ViewModels/NiagaraScratchPadViewModel.h
--- a/Engine/Plugins/FX/Niagara/Source/NiagaraEditor/Public/ViewModels/NiagaraScratchPadViewModel.h
+++ b/Engine/Plugins/FX/Niagara/Source/NiagaraEditor/Public/ViewModels/NiagaraScratchPadViewModel.h
@@ -36,6 +36,7 @@ public:
 	const TArray<TSharedRef<FNiagaraScratchPadScriptViewModel>>& GetEditScriptViewModels() const;
 
 	TSharedPtr<FNiagaraScratchPadScriptViewModel> GetViewModelForScript(UNiagaraScript* InScript);
+	TSharedPtr<FNiagaraScratchPadScriptViewModel> GetViewModelForScript(FName InScriptName);
 
 	TSharedPtr<FNiagaraScratchPadScriptViewModel> GetViewModelForEditScript(UNiagaraScript* InEditScript);
 
@@ -88,7 +89,7 @@ public:
 private:
 	TSharedRef<FNiagaraSystemViewModel> GetSystemViewModel();
 
-	TSharedRef<FNiagaraScratchPadScriptViewModel> CreateAndSetupScriptviewModel(UNiagaraScript* ScratchPadScript);
+	TSharedRef<FNiagaraScratchPadScriptViewModel> CreateAndSetupScriptviewModel(UNiagaraScript* ScratchPadScript, UNiagaraScript* ScratchPadEditScript);
 
 	void TearDownScriptViewModel(TSharedRef<FNiagaraScratchPadScriptViewModel> InScriptViewModel);
 
@@ -110,6 +111,9 @@ private:
 
 	void ScriptViewModelVariableSelectionChanged(TWeakPtr<FNiagaraScratchPadScriptViewModel> ScriptViewModelWeak);
 
+	void ScriptViewModelGraphSelectionChanged(const UObject* Obj, TWeakPtr<FNiagaraScratchPadScriptViewModel> ScriptViewModelWeak);
+
+
 private:
 	TSharedPtr<FNiagaraObjectSelection> ObjectSelection;
 
diff --git a/Engine/Plugins/FX/Niagara/Source/NiagaraEditor/Public/ViewModels/NiagaraSystemEditorDocumentsViewModel.h b/Engine/Plugins/FX/Niagara/Source/NiagaraEditor/Public/ViewModels/NiagaraSystemEditorDocumentsViewModel.h
--- a/Engine/Plugins/FX/Niagara/Source/NiagaraEditor/Public/ViewModels/NiagaraSystemEditorDocumentsViewModel.h
+++ b/Engine/Plugins/FX/Niagara/Source/NiagaraEditor/Public/ViewModels/NiagaraSystemEditorDocumentsViewModel.h
@@ -12,6 +12,7 @@
 
 class FNiagaraSystemViewModel;
 class FNiagaraScratchPadScriptViewModel;
+class FNiagaraSystemGraphSelectionViewModel;
 class FNiagaraObjectSelection;
 class UNiagaraScript;
 class UEdGraph;
@@ -35,8 +36,10 @@ public:
 	void OpenChildScript(UEdGraph* Graph);
 	void CloseChildScript(UEdGraph* Graph);
 
-	TArray<class UNiagaraGraph*> GetEditableGraphsForActiveDocument();
-	
+	TArray<class UNiagaraGraph*> GetEditableGraphsForActiveScriptDocument();
+	TArray<class UNiagaraGraph*> GetAllGraphsForActiveScriptDocument();
+	TArray<class UNiagaraGraph*> GetEditableGraphsForPrimaryDocument();
+	TArray<class UNiagaraGraph*> GetAllGraphsForPrimaryDocument();
 	void InitializePreTabManager(TSharedPtr<class FNiagaraSystemToolkit> InToolkit);
 	void InitializePostTabManager(TSharedPtr<class FNiagaraSystemToolkit> InToolkit);
 
@@ -46,9 +49,14 @@ public:
 		return ActiveDocChangedDelegate;
 	}
 
+	bool IsPrimaryDocumentActive() const;
 	TSharedPtr<class FNiagaraScratchPadScriptViewModel> GetActiveScratchPadViewModelIfSet();
 	static TSharedPtr < class FNiagaraScratchPadScriptViewModel> GetScratchPadViewModelFromGraph(FNiagaraSystemViewModel* InSysViewModel, UEdGraph* InTargetGraph);
 
+	void DrawAttentionToPrimaryDocument();
+	void SetPrimaryDocumentID(const FName& TabId) { PrimaryDocumentTabId = TabId; }
+	void SwapEditableScripts(TSharedPtr < class FNiagaraScratchPadScriptViewModel> OldScriptViewModel, TSharedPtr < class FNiagaraScratchPadScriptViewModel> NewScriptViewModel);
+
 protected:
 	TSharedPtr<SDockTab> OpenDocument(const UObject* DocumentID, FDocumentTracker::EOpenDocumentCause Cause);
 
@@ -59,6 +67,8 @@ protected:
 	// Finds any open tabs containing the specified document and adds them to the specified array; returns true if at least one is found
 	bool FindOpenTabsContainingDocument(const UObject* DocumentID, /*inout*/ TArray< TSharedPtr<SDockTab> >& Results);
 
+	FName PrimaryDocumentTabId;
+
 
 private:
 	TSharedRef<FNiagaraSystemViewModel> GetSystemViewModel();
@@ -66,6 +76,7 @@ private:
 	FScriptToolkitsActiveDocumentChanged ActiveDocChangedDelegate;
 
 	TWeakPtr<FNiagaraSystemViewModel> SystemViewModelWeak;
+	TWeakPtr< FNiagaraSystemGraphSelectionViewModel> SystemGraphSelectionVMWeak;
 
 	TSharedRef<class SNiagaraScratchPadScriptEditor> CreateGraphEditorWidget(TSharedRef<FTabInfo> InTabInfo, UEdGraph* InGraph);
 
diff --git a/Engine/Plugins/FX/Niagara/Source/NiagaraEditor/Public/ViewModels/Stack/NiagaraStackFunctionInput.h b/Engine/Plugins/FX/Niagara/Source/NiagaraEditor/Public/ViewModels/Stack/NiagaraStackFunctionInput.h
--- a/Engine/Plugins/FX/Niagara/Source/NiagaraEditor/Public/ViewModels/Stack/NiagaraStackFunctionInput.h
+++ b/Engine/Plugins/FX/Niagara/Source/NiagaraEditor/Public/ViewModels/Stack/NiagaraStackFunctionInput.h
@@ -282,6 +282,8 @@ public:
 
 
 	virtual const FCollectedUsageData& GetCollectedUsageData() const override;
+
+	bool OpenSourceAsset() const;
 protected:
 	//~ UNiagaraStackEntry interface
 	virtual void FinalizeInternal() override;
diff --git a/Engine/Plugins/FX/Niagara/Source/NiagaraEditor/Public/ViewModels/Stack/NiagaraStackModuleItem.h b/Engine/Plugins/FX/Niagara/Source/NiagaraEditor/Public/ViewModels/Stack/NiagaraStackModuleItem.h
--- a/Engine/Plugins/FX/Niagara/Source/NiagaraEditor/Public/ViewModels/Stack/NiagaraStackModuleItem.h
+++ b/Engine/Plugins/FX/Niagara/Source/NiagaraEditor/Public/ViewModels/Stack/NiagaraStackModuleItem.h
@@ -127,7 +127,8 @@ public:
 
 	FOnNoteModeSet& OnNoteModeSet() { return OnNoteModeSetDelegate;}
 
-	
+	bool OpenSourceAsset() const;
+
 protected:
 	FOnRequestDeprecationRecommended DeprecationDelegate;
 
diff --git a/Engine/Plugins/FX/Niagara/Source/NiagaraEditor/Public/Widgets/SItemSelector.h b/Engine/Plugins/FX/Niagara/Source/NiagaraEditor/Public/Widgets/SItemSelector.h
--- a/Engine/Plugins/FX/Niagara/Source/NiagaraEditor/Public/Widgets/SItemSelector.h
+++ b/Engine/Plugins/FX/Niagara/Source/NiagaraEditor/Public/Widgets/SItemSelector.h
@@ -75,6 +75,8 @@ public:
 	DECLARE_DELEGATE_RetVal_OneParam(const SectionKeyType&, FOnGetKeyForSection, const SectionType& /* Section */);
 	DECLARE_DELEGATE_OneParam(FOnSuggestionUpdated, int32 /* Suggestion Index*/)
 
+	DECLARE_DELEGATE_RetVal_OneParam(bool, FOnCategoryPassesFilter, const CategoryType& /* CategoryA */);
+
 public:
 	SLATE_BEGIN_ARGS(SItemSelector)
 		: _AllowMultiselect(false)
@@ -219,6 +221,9 @@ public:
 		/** Whether we want to expand the tree initially or not. */
 		SLATE_ARGUMENT(bool, ExpandInitially)
 
+		/** An optional delegate called when initially creating this and ExpandInitially is false. */
+		SLATE_EVENT(FOnCategoryPassesFilter, OnItemExpandedInitially)
+
 		/** Slot for additional widget content to go adjacent to right of the search box. */
 		SLATE_NAMED_SLOT(FArguments, SearchBoxAdjacentContent)
 	SLATE_END_ARGS();
@@ -1693,6 +1698,10 @@ public:
 		{
 			ExpandTree();
 		}
+		else if (InArgs._OnItemExpandedInitially.IsBound())
+		{
+			ExpandTreeByFilter(InArgs._OnItemExpandedInitially);
+		}
 
 		ExpandSections();
 
@@ -1733,7 +1742,7 @@ public:
 		return SelectedCategories;
 	}
 
-	void SetSelectedItems(const TArray<ItemType>& NewSelectedItems)
+	void SetSelectedItems(const TArray<ItemType>& NewSelectedItems, bool bExpandToShow = false)
 	{
 		checkf(OnCompareItemsForEquality.IsBound(), TEXT("OnCompareItemsForEquality event must be handled to use the SetSelectedItems function."));
 		TGuardValue<bool> SelectionGuard(bIsSettingSelection, true);
@@ -1748,6 +1757,16 @@ public:
 			for (TSharedRef<FItemSelectorItemViewModel> SelectedItemViewModel : SelectedItemViewModels)
 			{
 				ItemTree->SetSelection(SelectedItemViewModel);
+				if (bExpandToShow)
+				{
+					TSharedPtr<FRootViewModel> RootViewModel = ViewModelUtilities->GetRootViewModel();
+					if (RootViewModel.IsValid())
+					{
+						bool bFound = false;
+						TSharedRef < FItemSelectorItemViewModel> AsBase = RootViewModel.ToSharedRef();
+						FindAndExpandModelAndParents(bFound, AsBase, SelectedItemViewModel);
+					}
+				}
 			}
 		}
 		OnSelectionChanged.ExecuteIfBound();
@@ -1928,6 +1947,20 @@ public:
 		}
 	}
 
+	void ExpandTreeByFilter(const FOnCategoryPassesFilter& InFilter)
+	{
+		TArray<TSharedRef<FItemSelectorItemViewModel>> ItemsToProcess;
+		ItemsToProcess.Append(*ViewModelUtilities->GetRootItems());
+		while (ItemsToProcess.Num() > 0)
+		{
+			TSharedRef<FItemSelectorItemViewModel> ItemToProcess = ItemsToProcess[0];
+			ItemsToProcess.RemoveAtSwap(0);
+
+			TOptional<bool> bForce;
+			ExpandTreeByFilterRecursive(ItemToProcess, InFilter, bForce);
+		}
+	}
+
 	bool IsItemSelected(const ItemType Item)
 	{
 		const TArray<ItemType> ItemArr = {Item};
@@ -1945,12 +1978,79 @@ public:
 		return ItemTree->IsPendingRefresh();
 	}
 
+	void GetExpandedCategoryItems(TArray<CategoryType >& OutExpandedItems)
+	{
+		if (ItemTree.IsValid())
+		{
+			TSet<TSharedRef<FItemSelectorItemViewModel>> ExpandedItemViewModels;
+			ItemTree->GetExpandedItems(ExpandedItemViewModels);
+			for (const TSharedRef<FItemSelectorItemViewModel>& ExpandedItemViewModel : ExpandedItemViewModels)
+			{
+				if (ExpandedItemViewModel->GetType() == EItemSelectorItemViewModelType::Category)
+				{					
+					OutExpandedItems.AddUnique(StaticCastSharedRef<FItemSelectorItemCategoryViewModel>(ExpandedItemViewModel)->GetCategory());
+				}
+			}
+		}
+	}
+
+
 private:
+	void ExpandTreeByFilterRecursive(TSharedRef<FItemSelectorItemViewModel>& ItemToProcess, const FOnCategoryPassesFilter& InFilter, TOptional<bool> bForce)
+	{
+		if (!bForce.IsSet() && ItemToProcess->GetType() == EItemSelectorItemViewModelType::Category)
+		{
+			bool bExpand = InFilter.Execute(StaticCastSharedRef<FItemSelectorItemCategoryViewModel>(ItemToProcess)->GetCategory());
+			ItemTree->SetItemExpansion(ItemToProcess, bExpand);
+			bForce = bExpand;
+		}
+		else if (bForce.IsSet())
+		{
+			ItemTree->SetItemExpansion(ItemToProcess, bForce.GetValue());
+		}
+
+		TArray<TSharedRef<FItemSelectorItemViewModel>> ItemsToProcess;
+		ItemToProcess->GetChildren(ItemsToProcess);
+
+		for (TSharedRef<FItemSelectorItemViewModel>& Child : ItemsToProcess)
+		{
+			ExpandTreeByFilterRecursive(Child, InFilter, bForce);
+		}
+	}
+
 	EVisibility GetSearchBoxVisibility() const
 	{
 		return OnDoesItemMatchFilterText.IsBound() ? EVisibility::Visible : EVisibility::Collapsed;
 	}
 
+	void FindAndExpandModelAndParents(bool& bFound, const TSharedRef<FItemSelectorItemViewModel>& Current, const TSharedRef<FItemSelectorItemViewModel>& Target)
+	{
+		if (!bFound && Current == Target)
+		{
+			bFound = true;
+		}
+
+		if (!bFound)
+		{
+			TArray<TSharedRef<FItemSelectorItemViewModel>> Children;
+
+			Current->GetChildren(Children);
+			for (TSharedRef<FItemSelectorItemViewModel> Child : Children)
+			{
+				FindAndExpandModelAndParents(bFound, Child, Target);
+				if (bFound)
+				{
+					break;
+				}
+			}
+		}
+
+		if (bFound)
+		{
+			ItemTree->SetItemExpansion(Current, true);
+		}
+	};
+
 	void OnSearchTextChanged(const FText& SearchText)
 	{
 		if (ViewModelUtilities->GetFilterText().CompareTo(SearchText) != 0)
diff --git a/Engine/Plugins/FX/Niagara/Source/NiagaraEditor/Public/Widgets/SNiagaraParameterPanel.h b/Engine/Plugins/FX/Niagara/Source/NiagaraEditor/Public/Widgets/SNiagaraParameterPanel.h
--- a/Engine/Plugins/FX/Niagara/Source/NiagaraEditor/Public/Widgets/SNiagaraParameterPanel.h
+++ b/Engine/Plugins/FX/Niagara/Source/NiagaraEditor/Public/Widgets/SNiagaraParameterPanel.h
@@ -71,8 +71,8 @@ public:
 	NIAGARAEDITOR_API virtual void Tick(const FGeometry& AllottedGeometry, const double InCurrentTime, const float InDeltaTime) override;
 
 	/** Refreshes the items for the item selector. */
-	void Refresh();
-	void RefreshNextTick();
+	void Refresh(bool bRunCategoryExpansionFilter = false);
+	void RefreshNextTick(bool bRunCategoryExpansionFilter = false);
 
 	void OnParameterItemSelected(const FNiagaraParameterPanelItem& SelectedItem, ESelectInfo::Type SelectInfo) const;
 	FReply OnParameterItemsDragged(const TArray<FNiagaraParameterPanelItem>& DraggedItems, const FPointerEvent& MouseEvent) const;
@@ -115,6 +115,7 @@ private:
 	void HandleExternalSelectionChanged(const UObject* Obj);
 
 	const FSlateBrush* GetCategoryBackgroundImage(bool bIsCategoryHovered, bool bIsCategoryExpanded) const;
+	bool GetCategoryExpandedInitially(const FNiagaraParameterPanelCategory& Category) const;
 
 	TSharedRef<SWidget> GetViewOptionsMenu();
 
@@ -127,6 +128,7 @@ private:
 	void ConstructSectionButtons();
 
 	ECheckBoxState GetSectionCheckState(FText Section) const;
+	bool GetSectionEnabled(FText Section) const;
 
 	void OnSectionChecked(ECheckBoxState CheckState, FText Section);
 
@@ -138,6 +140,7 @@ private:
 	mutable bool bPendingRefresh;
 	mutable bool bPendingSelectionRestore;
 	mutable bool bParameterItemsPendingChange;
+	mutable bool bRunCategoryExpansionFilter = false;
 
 	/** Tracking list for parameters that are awaiting entering selection. */
 	TArray<FName> ParametersWithSelectionPending;
@@ -166,4 +169,7 @@ private:
 
 	/** Whether or not to display the reference counter for each parameter entry. */
 	bool bShowParameterReferenceCounter;
+
+	SNiagaraParameterPanelSelector::FOnCategoryPassesFilter FilterCategoryExpandedDelegate;
+
 };
diff --git a/Engine/Plugins/FX/Niagara/Source/NiagaraEditor/Public/Widgets/SNiagaraScriptGraph.h b/Engine/Plugins/FX/Niagara/Source/NiagaraEditor/Public/Widgets/SNiagaraScriptGraph.h
--- a/Engine/Plugins/FX/Niagara/Source/NiagaraEditor/Public/Widgets/SNiagaraScriptGraph.h
+++ b/Engine/Plugins/FX/Niagara/Source/NiagaraEditor/Public/Widgets/SNiagaraScriptGraph.h
@@ -44,6 +44,8 @@ public:
 
 	void RecreateGraphWidget();
 
+	TSharedPtr<FNiagaraScriptGraphViewModel> GetViewModel() { return ViewModel; };
+
 private:
 	/** Constructs the graph editor widget for the current graph. */
 	TSharedRef<SGraphEditor> ConstructGraphEditor();
@@ -115,7 +117,7 @@ private:
 	void OnDistributeNodesH();
 	void OnDistributeNodesV();
 
-
+	void RebuildCommands();
 private:
 	/** The combined commands used by the graph editor. */
 	TSharedPtr<FUICommandList> Commands;
diff --git a/Engine/Plugins/FX/Niagara/Source/NiagaraEditorWidgets/Private/SNiagaraCurveOverview.cpp b/Engine/Plugins/FX/Niagara/Source/NiagaraEditorWidgets/Private/SNiagaraCurveOverview.cpp
--- a/Engine/Plugins/FX/Niagara/Source/NiagaraEditorWidgets/Private/SNiagaraCurveOverview.cpp
+++ b/Engine/Plugins/FX/Niagara/Source/NiagaraEditorWidgets/Private/SNiagaraCurveOverview.cpp
@@ -19,6 +19,7 @@
 #include "ViewModels/NiagaraSystemSelectionViewModel.h"
 #include "ViewModels/NiagaraSystemViewModel.h"
 #include "Widgets/SNiagaraParameterName.h"
+#include "ViewModels/NiagaraSystemEditorDocumentsViewModel.h"
 
 #define LOCTEXT_NAMESPACE "NiagaraCurveEditor"
 
@@ -536,6 +537,8 @@ void SNiagaraCurveOverview::CurveTreeItemDoubleClicked(FCurveEditorTreeItemID Tr
 			{
 				SystemViewModel->GetSelectionViewModel()->EmptySelection();
 				SystemViewModel->GetSelectionViewModel()->AddEntryToSelectionByDisplayedObjectKeyDeferred(DisplayedObjectKey.GetValue());
+
+				SystemViewModel->GetDocumentViewModel()->DrawAttentionToPrimaryDocument();
 			}
 		}
 	}
diff --git a/Engine/Plugins/FX/Niagara/Source/NiagaraEditorWidgets/Private/SNiagaraOverviewStack.cpp b/Engine/Plugins/FX/Niagara/Source/NiagaraEditorWidgets/Private/SNiagaraOverviewStack.cpp
--- a/Engine/Plugins/FX/Niagara/Source/NiagaraEditorWidgets/Private/SNiagaraOverviewStack.cpp
+++ b/Engine/Plugins/FX/Niagara/Source/NiagaraEditorWidgets/Private/SNiagaraOverviewStack.cpp
@@ -290,6 +290,7 @@ class SNiagaraSystemOverviewEntryListRow : public STableRow<UNiagaraStackEntry*>
 	{
 		if (UNiagaraStackModuleItem* ModuleItem = Cast<UNiagaraStackModuleItem>(StackEntry))
 		{
+			
 			// if the user double-clicks a collapsed module we uncollapse it
 			UNiagaraStackEditorData& StackEditorData = ModuleItem->GetSystemViewModel()->GetEditorData().GetStackEditorData();
 			if (ModuleItem->GetIsEnabled() == false && StackEditorData.bHideDisabledModules)
@@ -298,25 +299,10 @@ class SNiagaraSystemOverviewEntryListRow : public STableRow<UNiagaraStackEntry*>
 				return FReply::Handled();
 			}
 			
-			const UNiagaraNodeFunctionCall& ModuleFunctionCall = ModuleItem->GetModuleNode();
-			if (ModuleFunctionCall.FunctionScript != nullptr)
+			if (ModuleItem->OpenSourceAsset())
 			{
-				if (ModuleFunctionCall.FunctionScript->IsAsset() || GbShowNiagaraDeveloperWindows > 0)
-				{
-					ModuleFunctionCall.FunctionScript->VersionToOpenInEditor = ModuleFunctionCall.SelectedScriptVersion;
-					GEditor->GetEditorSubsystem<UAssetEditorSubsystem>()->OpenEditorForAsset(const_cast<UNiagaraScript*>(ToRawPtr(ModuleFunctionCall.FunctionScript)));
-					return FReply::Handled();
-				}
-				else if (ModuleItem->IsScratchModule())
-				{
-					TSharedPtr<FNiagaraScratchPadScriptViewModel> ScratchPadScriptViewModel = ModuleItem->GetSystemViewModel()->GetScriptScratchPadViewModel()->GetViewModelForScript(ModuleFunctionCall.FunctionScript);
-					if (ScratchPadScriptViewModel.IsValid())
-					{
-						ModuleItem->GetSystemViewModel()->GetScriptScratchPadViewModel()->FocusScratchPadScriptViewModel(ScratchPadScriptViewModel.ToSharedRef());
-						return FReply::Handled();
-					}
-				}
-			}
+				return FReply::Handled();
+			}			
 		}
 		return FReply::Unhandled();
 	}
diff --git a/Engine/Plugins/FX/Niagara/Source/NiagaraEditorWidgets/Private/Stack/SNiagaraStackFunctionInputValue.cpp b/Engine/Plugins/FX/Niagara/Source/NiagaraEditorWidgets/Private/Stack/SNiagaraStackFunctionInputValue.cpp
--- a/Engine/Plugins/FX/Niagara/Source/NiagaraEditorWidgets/Private/Stack/SNiagaraStackFunctionInputValue.cpp
+++ b/Engine/Plugins/FX/Niagara/Source/NiagaraEditorWidgets/Private/Stack/SNiagaraStackFunctionInputValue.cpp
@@ -614,25 +614,11 @@ void SNiagaraStackFunctionInputValue::OnExpressionTextCommitted(const FText& Nam
 
 FReply SNiagaraStackFunctionInputValue::DynamicInputTextDoubleClicked(const FGeometry& MyGeometry, const FPointerEvent& PointerEvent)
 {
-	UNiagaraNodeFunctionCall* DynamicInputNode = FunctionInput->GetDynamicInputNode();
-	if (DynamicInputNode->FunctionScript != nullptr)
+	if (FunctionInput->OpenSourceAsset())
 	{
-		if (DynamicInputNode->FunctionScript->IsAsset())
-		{
-			DynamicInputNode->FunctionScript->VersionToOpenInEditor = DynamicInputNode->SelectedScriptVersion;
-			GEditor->GetEditorSubsystem<UAssetEditorSubsystem>()->OpenEditorForAsset(DynamicInputNode->FunctionScript);
-			return FReply::Handled();
-		}
-		else
-		{
-			TSharedPtr<FNiagaraScratchPadScriptViewModel> ScratchPadScriptViewModel = FunctionInput->GetSystemViewModel()->GetScriptScratchPadViewModel()->GetViewModelForScript(DynamicInputNode->FunctionScript);
-			if(ScratchPadScriptViewModel.IsValid())
-			{
-				FunctionInput->GetSystemViewModel()->GetScriptScratchPadViewModel()->FocusScratchPadScriptViewModel(ScratchPadScriptViewModel.ToSharedRef());
-				return FReply::Handled();
-			}
-		}
+		return FReply::Handled();
 	}
+	
 	return FReply::Unhandled();
 }
 
diff --git a/Engine/Plugins/FX/Niagara/Source/NiagaraEditorWidgets/Private/Stack/SNiagaraStackModuleItem.cpp b/Engine/Plugins/FX/Niagara/Source/NiagaraEditorWidgets/Private/Stack/SNiagaraStackModuleItem.cpp
--- a/Engine/Plugins/FX/Niagara/Source/NiagaraEditorWidgets/Private/Stack/SNiagaraStackModuleItem.cpp
+++ b/Engine/Plugins/FX/Niagara/Source/NiagaraEditorWidgets/Private/Stack/SNiagaraStackModuleItem.cpp
@@ -69,25 +69,11 @@ void SNiagaraStackModuleItem::FillRowContextMenu(FMenuBuilder& MenuBuilder)
 
 FReply SNiagaraStackModuleItem::OnMouseButtonDoubleClick(const FGeometry& InMyGeometry, const FPointerEvent& InMouseEvent)
 {
-	const UNiagaraNodeFunctionCall& ModuleFunctionCall = ModuleItem->GetModuleNode();
-	if (ModuleFunctionCall.FunctionScript != nullptr)
+	if (ModuleItem->OpenSourceAsset())
 	{
-		if (ModuleFunctionCall.FunctionScript->IsAsset() || GbShowNiagaraDeveloperWindows > 0)
-		{
-			ModuleFunctionCall.FunctionScript->VersionToOpenInEditor = ModuleFunctionCall.SelectedScriptVersion;
-			GEditor->GetEditorSubsystem<UAssetEditorSubsystem>()->OpenEditorForAsset(ToRawPtr(ModuleFunctionCall.FunctionScript));
-			return FReply::Handled();
-		}
-		else if (ModuleItem->IsScratchModule())
-		{
-			TSharedPtr<FNiagaraScratchPadScriptViewModel> ScratchPadScriptViewModel = ModuleItem->GetSystemViewModel()->GetScriptScratchPadViewModel()->GetViewModelForScript(ModuleFunctionCall.FunctionScript);
-			if (ScratchPadScriptViewModel.IsValid())
-			{
-				ModuleItem->GetSystemViewModel()->GetScriptScratchPadViewModel()->FocusScratchPadScriptViewModel(ScratchPadScriptViewModel.ToSharedRef());
-				return FReply::Handled();
-			}
-		}
+		return FReply::Handled();
 	}
+	
 	return FReply::Unhandled();
 }
 
