[KillPendingKill] Release references to source actor and component from NiagaraDataInterfaceSkeletalMesh.

diff --git a/Engine/Plugins/FX/Niagara/Source/Niagara/Classes/NiagaraDataInterfaceSkeletalMesh.h b/Engine/Plugins/FX/Niagara/Source/Niagara/Classes/NiagaraDataInterfaceSkeletalMesh.h
--- a/Engine/Plugins/FX/Niagara/Source/Niagara/Classes/NiagaraDataInterfaceSkeletalMesh.h
+++ b/Engine/Plugins/FX/Niagara/Source/Niagara/Classes/NiagaraDataInterfaceSkeletalMesh.h
@@ -726,6 +726,7 @@ public:
 	TSoftObjectPtr<USkeletalMesh> PreviewMesh;
 #endif
 
+protected:
 	/** The source actor from which to sample. Takes precedence over the direct mesh. Note that this can only be set when used as a user variable on a component in the world.*/
 	UPROPERTY(EditAnywhere, Category = "Mesh", meta = (DisplayName = "Source Actor"))
 	TSoftObjectPtr<AActor> SoftSourceActor;
@@ -739,14 +740,15 @@ public:
 	TObjectPtr<AActor> Source_DEPRECATED;
 #endif
 
-	/** Reference to a user parameter if we're reading one. */
-	UPROPERTY(EditAnywhere, Category = "Mesh")
-	FNiagaraUserParameterBinding MeshUserParameter;
-
 	/** The source component from which to sample. Takes precedence over the direct mesh. Not exposed to the user, only indirectly accessible from blueprints. */
 	UPROPERTY(Transient)
 	TObjectPtr<USkeletalMeshComponent> SourceComponent;
 
+public:
+	/** Reference to a user parameter if we're reading one. */
+	UPROPERTY(EditAnywhere, Category = "Mesh")
+	FNiagaraUserParameterBinding MeshUserParameter;
+
 	/** Selects which skinning mode to use, for most cases Skin On The Fly will cover your requirements, see individual tooltips for more information. */
 	UPROPERTY(EditAnywhere, Category="Mesh")
 	ENDISkeletalMesh_SkinningMode SkinningMode;
@@ -840,6 +842,9 @@ public:
 	/** Finds the skeletal mesh based on settings of the DI and the hierarchy of the object provided */
 	USkeletalMesh* GetSkeletalMesh(UNiagaraComponent* Component);
 
+	USkeletalMeshComponent* GetSourceComponent() const { return SourceComponent; }
+	AActor* GetSourceActor() const { return SoftSourceActor.Get(); }
+
 	int32 CalculateLODIndexAndSamplingRegions(USkeletalMesh* InMesh, TArray<int32>& OutSamplingRegionIndices, bool& OutAllRegionsAreAreaWeighting) const;
 
 	virtual void ProvidePerInstanceDataForRenderThread(void* DataForRenderThread, void* PerInstanceData, const FNiagaraSystemInstanceID& SystemInstance) override;
@@ -892,6 +897,13 @@ public:
 protected:
 	virtual bool CopyToInternal(UNiagaraDataInterface* Destination) const override;
 
+	// Bind/unbind delegates to release references to the source actor & component.
+	void UnbindSourceDelegates();
+	void BindSourceDelegates();
+
+	UFUNCTION()
+	void OnSourceEndPlay(AActor* InSource, EEndPlayReason::Type Reason);
+
 	//////////////////////////////////////////////////////////////////////////
 	//Triangle sampling
 	//Triangles are sampled a using MeshTriangleCoordinates which are composed of Triangle index and a bary centric coordinate on that triangle.
diff --git a/Engine/Plugins/FX/Niagara/Source/Niagara/Private/NiagaraDataInterfaceSkeletalMesh.cpp b/Engine/Plugins/FX/Niagara/Source/Niagara/Private/NiagaraDataInterfaceSkeletalMesh.cpp
--- a/Engine/Plugins/FX/Niagara/Source/Niagara/Private/NiagaraDataInterfaceSkeletalMesh.cpp
+++ b/Engine/Plugins/FX/Niagara/Source/Niagara/Private/NiagaraDataInterfaceSkeletalMesh.cpp
@@ -2112,10 +2112,10 @@ bool FNDISkeletalMesh_InstanceData::ResetRequired(UNiagaraDataInterfaceSkeletalM
 			return true;
 		}
 	}
-	else if (Interface->SourceComponent)
+	else if (Interface->GetSourceComponent())
 	{
 		// Reset if the source component changed (or there wasn't one and now there is)
-		if (Interface->SourceComponent != Comp)
+		if (Interface->GetSourceComponent() != Comp)
 		{
 			return true;
 		}
@@ -2415,6 +2415,7 @@ bool UNiagaraDataInterfaceSkeletalMesh::CopyToInternal(UNiagaraDataInterface* De
 	}
 
 	UNiagaraDataInterfaceSkeletalMesh* OtherTyped = CastChecked<UNiagaraDataInterfaceSkeletalMesh>(Destination);
+	OtherTyped->UnbindSourceDelegates();
 	OtherTyped->SourceMode = SourceMode;
 	OtherTyped->SoftSourceActor = SoftSourceActor;
 	OtherTyped->MeshUserParameter = MeshUserParameter;
@@ -2431,6 +2432,8 @@ bool UNiagaraDataInterfaceSkeletalMesh::CopyToInternal(UNiagaraDataInterface* De
 #if WITH_EDITORONLY_DATA
 	OtherTyped->PreviewMesh = PreviewMesh;
 #endif
+	OtherTyped->BindSourceDelegates();
+
 	return true;
 }
 
@@ -3430,12 +3433,44 @@ void UNiagaraDataInterfaceSkeletalMesh::SetShaderParameters(const FNiagaraDataIn
 	}
 }
 
+void UNiagaraDataInterfaceSkeletalMesh::BindSourceDelegates()
+{
+	if (AActor* Source = SoftSourceActor.Get())
+	{
+		Source->OnEndPlay.AddDynamic(this, &UNiagaraDataInterfaceSkeletalMesh::OnSourceEndPlay);
+	}
+	else if (SourceComponent)
+	{
+		UE_CLOG(!UObjectBaseUtility::IsPendingKillEnabled(), 
+			LogNiagara, Warning, TEXT("%s: Unable to bind OnEndPlay for actor-less source component %s, this may extend the lifetime of the component"), 
+			*GetFullName(), *SourceComponent->GetPathName());
+	}
+}
+
+void UNiagaraDataInterfaceSkeletalMesh::UnbindSourceDelegates()
+{
+	if (AActor* Source = SoftSourceActor.Get())
+	{
+		Source->OnEndPlay.RemoveAll(this);
+	}
+}
+
+void UNiagaraDataInterfaceSkeletalMesh::OnSourceEndPlay(AActor* InSource, EEndPlayReason::Type Reason)
+{
+	// Increment change id in case we're able to find a new source component 
+	ChangeId++;
+	SoftSourceActor = nullptr;
+	SourceComponent = nullptr;
+}
+
 void UNiagaraDataInterfaceSkeletalMesh::SetSourceComponentFromBlueprints(USkeletalMeshComponent* ComponentToUse)
 {
 	// NOTE: When ChangeId changes the next tick will be skipped and a reset of the per-instance data will be initiated.
 	ChangeId++;
+	UnbindSourceDelegates();
 	SourceComponent = ComponentToUse;
 	SoftSourceActor = ComponentToUse->GetOwner();
+	BindSourceDelegates();
 }
 
 void UNiagaraDataInterfaceSkeletalMesh::SetSamplingRegionsFromBlueprints(const TArray<FName>& InSamplingRegions)
