Fixing crash due ot missing material

diff --git a/Engine/Plugins/FX/Niagara/Source/NiagaraEditor/Private/NiagaraValidationRules.cpp b/Engine/Plugins/FX/Niagara/Source/NiagaraEditor/Private/NiagaraValidationRules.cpp
--- a/Engine/Plugins/FX/Niagara/Source/NiagaraEditor/Private/NiagaraValidationRules.cpp
+++ b/Engine/Plugins/FX/Niagara/Source/NiagaraEditor/Private/NiagaraValidationRules.cpp
@@ -458,6 +458,11 @@ void UNiagaraValidationRule_NoOpaqueRenderMaterial::CheckValidity(const FNiagara
 						RendererProperties->GetUsedMaterials(nullptr, OutMaterials);
 						for (UMaterialInterface* Material : OutMaterials)
 						{
+							if (!Material)
+							{
+								continue;
+							}
+							
 							if (Material->GetBlendMode() == BLEND_Opaque || Material->GetBlendMode() == BLEND_Masked)
 							{
 								FText Description = LOCTEXT("NoOpaqueRenderMaterialDescription", "This renderer uses a material with a masked or opaque blend mode, which writes to the depth buffer.\nThis will cause conflicts when the collision module also uses depth buffer collisions.");
