Niagara - Fix a regression in the merge manager where the emitter merge adapter cache used by the stack was being bypassed.

diff --git a/Engine/Plugins/FX/Niagara/Source/NiagaraEditor/Private/NiagaraScriptMergeManager.cpp b/Engine/Plugins/FX/Niagara/Source/NiagaraEditor/Private/NiagaraScriptMergeManager.cpp
--- a/Engine/Plugins/FX/Niagara/Source/NiagaraEditor/Private/NiagaraScriptMergeManager.cpp
+++ b/Engine/Plugins/FX/Niagara/Source/NiagaraEditor/Private/NiagaraScriptMergeManager.cpp
@@ -3785,7 +3785,7 @@ TSharedRef<FNiagaraEmitterMergeAdapter> FNiagaraScriptMergeManager::GetEmitterMe
 	FCachedMergeAdapter* CachedMergeAdapter = FindOrAddMergeAdapterCacheForEmitter(Emitter);
 
 	if (CachedMergeAdapter->EmitterMergeAdapter.IsValid() == false ||
-		CachedMergeAdapter->EmitterMergeAdapter->GetEditableEmitter().Emitter != nullptr)
+		CachedMergeAdapter->EmitterMergeAdapter->GetEditableEmitter().Emitter == nullptr)
 	{
 		CachedMergeAdapter->EmitterMergeAdapter = MakeShared<FNiagaraEmitterMergeAdapter>(Emitter);
 	}
