Removed overly aggressive ensure, which asserted that the actor must have a root component which, apparently isn't always true

diff --git a/Engine/Plugins/FX/Niagara/Source/Niagara/Private/NiagaraDataInterfaceSkeletalMesh.cpp b/Engine/Plugins/FX/Niagara/Source/Niagara/Private/NiagaraDataInterfaceSkeletalMesh.cpp
--- a/Engine/Plugins/FX/Niagara/Source/Niagara/Private/NiagaraDataInterfaceSkeletalMesh.cpp
+++ b/Engine/Plugins/FX/Niagara/Source/Niagara/Private/NiagaraDataInterfaceSkeletalMesh.cpp
@@ -1516,8 +1516,7 @@ USkeletalMesh* UNiagaraDataInterfaceSkeletalMesh::GetSkeletalMesh(FNiagaraSystem
 		// Fall back on any valid component on the actor
 		while (Actor)
 		{
-			USceneComponent* ActorRootComponent = Actor->GetRootComponent();
-			if (ensure(ActorRootComponent))
+			if (USceneComponent* ActorRootComponent = Actor->GetRootComponent())
 			{
 				TArray<USceneComponent*> ActorComponents;
 				ActorRootComponent->GetChildrenComponents(true /*bIncludeAllDescendants*/, ActorComponents);
