- Add multi-GPU transfers back in as RDG does not handle this currently

diff --git a/Engine/Plugins/FX/Niagara/Source/Niagara/Private/NiagaraDataInterfaceIntRenderTarget2D.cpp b/Engine/Plugins/FX/Niagara/Source/Niagara/Private/NiagaraDataInterfaceIntRenderTarget2D.cpp
--- a/Engine/Plugins/FX/Niagara/Source/Niagara/Private/NiagaraDataInterfaceIntRenderTarget2D.cpp
+++ b/Engine/Plugins/FX/Niagara/Source/Niagara/Private/NiagaraDataInterfaceIntRenderTarget2D.cpp
@@ -128,14 +128,13 @@ struct FNDIIntRenderTarget2DProxy : public FNiagaraDataInterfaceProxyRW
 	{
 		FNDIIntRenderTarget2DInstanceData_RenderThread& InstanceData = SystemInstancesToProxyData_RT.FindChecked(Context.GetSystemInstanceID());
 
-//-TODO:RDG:mGPU
-//		// We only need to transfer this frame if it was written to.
-//		// If also read then we need to notify that the texture is important for the simulation
-//		// We also assume the texture is important for rendering, without discovering renderer bindings we don't really know
-//		if (InstanceData->bWroteThisFrame)
-//		{
-//			Context.ComputeDispatchInterface->MultiGPUResourceModified(RHICmdList, InstanceData->TextureRHI, InstanceData->bReadThisFrame, true);
-//		}
+		// We only need to transfer this frame if it was written to.
+		// If also read then we need to notify that the texture is important for the simulation
+		// We also assume the texture is important for rendering, without discovering renderer bindings we don't really know
+		if (InstanceData.bWroteThisFrame)
+		{
+			Context.GetComputeDispatchInterface().MultiGPUResourceModified(Context.GetGraphBuilder(), InstanceData.TextureRHI, InstanceData.bReadThisFrame, true);
+		}
 
 		InstanceData.bWroteThisFrame = true;
 		InstanceData.bReadThisFrame = true;
diff --git a/Engine/Plugins/FX/Niagara/Source/Niagara/Private/NiagaraDataInterfaceRenderTarget2D.cpp b/Engine/Plugins/FX/Niagara/Source/Niagara/Private/NiagaraDataInterfaceRenderTarget2D.cpp
--- a/Engine/Plugins/FX/Niagara/Source/Niagara/Private/NiagaraDataInterfaceRenderTarget2D.cpp
+++ b/Engine/Plugins/FX/Niagara/Source/Niagara/Private/NiagaraDataInterfaceRenderTarget2D.cpp
@@ -874,7 +874,7 @@ void FNiagaraDataInterfaceProxyRenderTarget2DProxy::PostSimulate(const FNDIGpuCo
 	// We also assume the texture is important for rendering, without discovering renderer bindings we don't really know
 	if (ProxyData->bWroteThisFrame)
 	{
-		//-TODO:RDG:mGPU:Context.ComputeDispatchInterface->MultiGPUResourceModified(RHICmdList, ProxyData->TextureRHI, ProxyData->bReadThisFrame, true);
+		Context.GetComputeDispatchInterface().MultiGPUResourceModified(Context.GetGraphBuilder(), ProxyData->TextureRHI, ProxyData->bReadThisFrame, true);
 	}
 
 #if NIAGARA_COMPUTEDEBUG_ENABLED && WITH_EDITORONLY_DATA
diff --git a/Engine/Plugins/FX/Niagara/Source/Niagara/Private/NiagaraDataInterfaceRenderTarget2DArray.cpp b/Engine/Plugins/FX/Niagara/Source/Niagara/Private/NiagaraDataInterfaceRenderTarget2DArray.cpp
--- a/Engine/Plugins/FX/Niagara/Source/Niagara/Private/NiagaraDataInterfaceRenderTarget2DArray.cpp
+++ b/Engine/Plugins/FX/Niagara/Source/Niagara/Private/NiagaraDataInterfaceRenderTarget2DArray.cpp
@@ -702,7 +702,7 @@ void FNiagaraDataInterfaceProxyRenderTarget2DArrayProxy::PostSimulate(const FNDI
 	// We also assume the texture is important for rendering, without discovering renderer bindings we don't really know
 	if (ProxyData->bWroteThisFrame)
 	{
-		//-TODO:RDG:mGPU:Context.ComputeDispatchInterface->MultiGPUResourceModified(RHICmdList, ProxyData->TextureRHI, ProxyData->bReadThisFrame, true);
+		Context.GetComputeDispatchInterface().MultiGPUResourceModified(Context.GetGraphBuilder(), ProxyData->RenderTarget->GetRHI(), ProxyData->bReadThisFrame, true);
 	}
 
 #if NIAGARA_COMPUTEDEBUG_ENABLED && WITH_EDITORONLY_DATA
diff --git a/Engine/Plugins/FX/Niagara/Source/Niagara/Private/NiagaraDataInterfaceRenderTargetCube.cpp b/Engine/Plugins/FX/Niagara/Source/Niagara/Private/NiagaraDataInterfaceRenderTargetCube.cpp
--- a/Engine/Plugins/FX/Niagara/Source/Niagara/Private/NiagaraDataInterfaceRenderTargetCube.cpp
+++ b/Engine/Plugins/FX/Niagara/Source/Niagara/Private/NiagaraDataInterfaceRenderTargetCube.cpp
@@ -678,7 +678,7 @@ void FNiagaraDataInterfaceProxyRenderTargetCubeProxy::PostSimulate(const FNDIGpu
 	// We also assume the texture is important for rendering, without discovering renderer bindings we don't really know
 	if (ProxyData->bWroteThisFrame)
 	{
-		//-TODO:RDG:mGPU:Context.ComputeDispatchInterface->MultiGPUResourceModified(RHICmdList, ProxyData->TextureRHI, ProxyData->bReadThisFrame, true);
+		Context.GetComputeDispatchInterface().MultiGPUResourceModified(Context.GetGraphBuilder(), ProxyData->RenderTarget->GetRHI(), ProxyData->bReadThisFrame, true);
 	}
 
 #if NIAGARA_COMPUTEDEBUG_ENABLED && WITH_EDITORONLY_DATA
diff --git a/Engine/Plugins/FX/Niagara/Source/Niagara/Private/NiagaraDataInterfaceRenderTargetVolume.cpp b/Engine/Plugins/FX/Niagara/Source/Niagara/Private/NiagaraDataInterfaceRenderTargetVolume.cpp
--- a/Engine/Plugins/FX/Niagara/Source/Niagara/Private/NiagaraDataInterfaceRenderTargetVolume.cpp
+++ b/Engine/Plugins/FX/Niagara/Source/Niagara/Private/NiagaraDataInterfaceRenderTargetVolume.cpp
@@ -696,7 +696,7 @@ void FNiagaraDataInterfaceProxyRenderTargetVolumeProxy::PostSimulate(const FNDIG
 	// We also assume the texture is important for rendering, without discovering renderer bindings we don't really know
 	if (ProxyData->bWroteThisFrame)
 	{
-		//-TODO:RDG:mGPU:Context.ComputeDispatchInterface->MultiGPUResourceModified(RHICmdList, ProxyData->TextureRHI, ProxyData->bReadThisFrame, true);
+		Context.GetComputeDispatchInterface().MultiGPUResourceModified(Context.GetGraphBuilder(), ProxyData->RenderTarget->GetRHI(), ProxyData->bReadThisFrame, true);
 	}
 
 #if NIAGARA_COMPUTEDEBUG_ENABLED && WITH_EDITORONLY_DATA
diff --git a/Engine/Plugins/FX/Niagara/Source/Niagara/Private/NiagaraGpuComputeDispatch.cpp b/Engine/Plugins/FX/Niagara/Source/Niagara/Private/NiagaraGpuComputeDispatch.cpp
--- a/Engine/Plugins/FX/Niagara/Source/Niagara/Private/NiagaraGpuComputeDispatch.cpp
+++ b/Engine/Plugins/FX/Niagara/Source/Niagara/Private/NiagaraGpuComputeDispatch.cpp
@@ -1217,25 +1217,24 @@ void FNiagaraGpuComputeDispatch::ExecuteTicks(FRDGBuilder& GraphBuilder, TConstA
 				ComputeContext->SetTranslucentDataToRender(nullptr);
 				ComputeContext->SetDataToRender(CurrentData);
 
-//-TODO:RDG:mGPU:
-//#if WITH_MGPU
-//				if (bAFREnabled)
-//				{
-//					AddAFRBuffer(CurrentData->GetGPUBufferFloat().Buffer);
-//					AddAFRBuffer(CurrentData->GetGPUBufferHalf().Buffer);
-//					AddAFRBuffer(CurrentData->GetGPUBufferInt().Buffer);
-//					if (ComputeContext->MainDataSet->RequiresPersistentIDs())
-//					{
-//						AddAFRBuffer(ComputeContext->MainDataSet->GetGPUFreeIDs().Buffer);
-//					}
-//				}
-//				if (bCrossGPUTransferEnabled)
-//				{
-//					AddCrossGPUTransfer(RHICmdList, CurrentData->GetGPUBufferFloat().Buffer);
-//					AddCrossGPUTransfer(RHICmdList, CurrentData->GetGPUBufferHalf().Buffer);
-//					AddCrossGPUTransfer(RHICmdList, CurrentData->GetGPUBufferInt().Buffer);
-//				}
-//#endif // WITH_MGPU
+#if WITH_MGPU
+				if (bAFREnabled)
+				{
+					AddAFRBuffer(CurrentData->GetGPUBufferFloat().Buffer);
+					AddAFRBuffer(CurrentData->GetGPUBufferHalf().Buffer);
+					AddAFRBuffer(CurrentData->GetGPUBufferInt().Buffer);
+					if (ComputeContext->MainDataSet->RequiresPersistentIDs())
+					{
+						AddAFRBuffer(ComputeContext->MainDataSet->GetGPUFreeIDs().Buffer);
+					}
+				}
+				if (bCrossGPUTransferEnabled)
+				{
+					AddCrossGPUTransfer(GraphBuilder.RHICmdList, CurrentData->GetGPUBufferFloat().Buffer);
+					AddCrossGPUTransfer(GraphBuilder.RHICmdList, CurrentData->GetGPUBufferHalf().Buffer);
+					AddCrossGPUTransfer(GraphBuilder.RHICmdList, CurrentData->GetGPUBufferInt().Buffer);
+				}
+#endif // WITH_MGPU
 			}
 			// If this is not the final tick of the final stage we need set our temporary buffer for data interfaces, etc, that may snoop from CurrentData
 			else
@@ -1684,6 +1683,21 @@ void FNiagaraGpuComputeDispatch::PostInitViews(FRDGBuilder& GraphBuilder, TArray
 	if (bAllowGPUParticleUpdate && FNiagaraUtilities::AllowGPUParticles(GetShaderPlatform()))
 	{
 		ExecuteTicks(GraphBuilder, Views, ENiagaraGpuComputeTickStage::PostInitViews);
+
+	#if WITH_MGPU
+		// Queue a transfer request
+		if (StageToTransferGPUBuffers == ENiagaraGpuComputeTickStage::PostInitViews)
+		{
+			AddPass(
+				GraphBuilder,
+				RDG_EVENT_NAME("Niagara::TransferMultiGPUBufers"),
+				[this](FRHICommandListImmediate& RHICmdList)
+				{
+					TransferMultiGPUBufers(RHICmdList, ENiagaraGpuComputeTickStage::PostInitViews);
+				}
+			);
+		}
+	#endif // WITH_MGPU
 	}
 }
 
@@ -2286,6 +2300,16 @@ void FNiagaraGpuComputeDispatch::DrawSceneDebug_RenderThread(class FRDGBuilder&
 }
 
 #if WITH_MGPU
+void FNiagaraGpuComputeDispatch::MultiGPUResourceModified(FRDGBuilder& GraphBuilder, FRHIBuffer* Buffer, bool bRequiredForSimulation, bool bRequiredForRendering) const
+{
+	MultiGPUResourceModified(GraphBuilder.RHICmdList, Buffer, bRequiredForSimulation, bRequiredForRendering);
+}
+
+void FNiagaraGpuComputeDispatch::MultiGPUResourceModified(FRDGBuilder& GraphBuilder, FRHITexture* Texture, bool bRequiredForSimulation, bool bRequiredForRendering) const
+{
+	MultiGPUResourceModified(GraphBuilder.RHICmdList, Texture, bRequiredForSimulation, bRequiredForRendering);
+}
+
 void FNiagaraGpuComputeDispatch::MultiGPUResourceModified(FRHICommandList& RHICmdList, FRHIBuffer* Buffer, bool bRequiredForSimulation, bool bRequiredForRendering) const
 {
 	if (bAFREnabled && bRequiredForSimulation)
diff --git a/Engine/Plugins/FX/Niagara/Source/Niagara/Private/NiagaraGpuComputeDispatch.h b/Engine/Plugins/FX/Niagara/Source/Niagara/Private/NiagaraGpuComputeDispatch.h
--- a/Engine/Plugins/FX/Niagara/Source/Niagara/Private/NiagaraGpuComputeDispatch.h
+++ b/Engine/Plugins/FX/Niagara/Source/Niagara/Private/NiagaraGpuComputeDispatch.h
@@ -183,6 +183,9 @@ public:
 	virtual void AddDebugReadback(FNiagaraSystemInstanceID InstanceID, TSharedPtr<struct FNiagaraScriptDebuggerInfo, ESPMode::ThreadSafe> DebugInfo, FNiagaraComputeExecutionContext* Context) override;
 
 #if WITH_MGPU
+	virtual void MultiGPUResourceModified(FRDGBuilder& GraphBuilder, FRHIBuffer* Buffer, bool bRequiredForSimulation, bool bRequiredForRendering) const override;
+	virtual void MultiGPUResourceModified(FRDGBuilder& GraphBuilder, FRHITexture* Texture, bool bRequiredForSimulation, bool bRequiredForRendering) const override;
+
 	virtual void MultiGPUResourceModified(FRHICommandList& RHICmdList, FRHIBuffer* Buffer, bool bRequiredForSimulation, bool bRequiredForRendering) const override;
 	virtual void MultiGPUResourceModified(FRHICommandList& RHICmdList, FRHITexture* Texture, bool bRequiredForSimulation, bool bRequiredForRendering) const override;
 #endif
diff --git a/Engine/Plugins/FX/Niagara/Source/Niagara/Public/NiagaraGpuComputeDispatchInterface.h b/Engine/Plugins/FX/Niagara/Source/Niagara/Public/NiagaraGpuComputeDispatchInterface.h
--- a/Engine/Plugins/FX/Niagara/Source/Niagara/Public/NiagaraGpuComputeDispatchInterface.h
+++ b/Engine/Plugins/FX/Niagara/Source/Niagara/Public/NiagaraGpuComputeDispatchInterface.h
@@ -117,10 +117,20 @@ public:
 	bRequiredForSimulation	- When true we require this resource for simulation passes
 	bRequiredForRendering	- When true we require this resource for rendering passes
 	*/
+	virtual void MultiGPUResourceModified(FRDGBuilder& GraphBuilder, FRHIBuffer* Buffer, bool bRequiredForSimulation, bool bRequiredForRendering) const = 0;
+	virtual void MultiGPUResourceModified(FRDGBuilder& GraphBuilder, FRHITexture* Texture, bool bRequiredForSimulation, bool bRequiredForRendering) const = 0;
+
+	UE_DEPRECATED(5.1, "ImmediateMode is deprecated for Niagara please migrate to using a FRDGBuilder")
 	virtual void MultiGPUResourceModified(FRHICommandList& RHICmdList, FRHIBuffer* Buffer, bool bRequiredForSimulation, bool bRequiredForRendering) const = 0;
+	UE_DEPRECATED(5.1, "ImmediateMode is deprecated for Niagara please migrate to using a FRDGBuilder")
 	virtual void MultiGPUResourceModified(FRHICommandList& RHICmdList, FRHITexture* Texture, bool bRequiredForSimulation, bool bRequiredForRendering) const = 0;
 #else
+	FORCEINLINE void MultiGPUResourceModified(FRDGBuilder& GraphBuilder, FRHIBuffer* Buffer, bool bRequiredForSimulation, bool bRequiredForRendering) const {}
+	FORCEINLINE void MultiGPUResourceModified(FRDGBuilder& GraphBuilder, FRHITexture* Texture, bool bRequiredForSimulation, bool bRequiredForRendering) const {}
+
+	UE_DEPRECATED(5.1, "ImmediateMode is deprecated for Niagara please migrate to using a FRDGBuilder")
 	FORCEINLINE void MultiGPUResourceModified(FRHICommandList& RHICmdList, FRHIBuffer* Buffer, bool bRequiredForSimulation, bool bRequiredForRendering) const {}
+	UE_DEPRECATED(5.1, "ImmediateMode is deprecated for Niagara please migrate to using a FRDGBuilder")
 	FORCEINLINE void MultiGPUResourceModified(FRHICommandList& RHICmdList, FRHITexture* Texture, bool bRequiredForSimulation, bool bRequiredForRendering) const {}
 #endif
 
diff --git a/Engine/Plugins/Runtime/HairStrands/Source/HairStrandsCore/Private/Niagara/NiagaraDataInterfaceHairStrands.cpp b/Engine/Plugins/Runtime/HairStrands/Source/HairStrandsCore/Private/Niagara/NiagaraDataInterfaceHairStrands.cpp
--- a/Engine/Plugins/Runtime/HairStrands/Source/HairStrandsCore/Private/Niagara/NiagaraDataInterfaceHairStrands.cpp
+++ b/Engine/Plugins/Runtime/HairStrands/Source/HairStrandsCore/Private/Niagara/NiagaraDataInterfaceHairStrands.cpp
@@ -3571,14 +3571,12 @@ void FNDIHairStrandsProxy::PostSimulate(const FNDIGpuComputePostSimulateContext&
 	const bool bIsDeformedValid = bIsHairValid && ProxyData->HairStrandsBuffer->SourceDeformedResources && ProxyData->HairStrandsBuffer->SourceDeformedResources->IsInitialized();
 
 	// MGPU DeformedPositionBuffer copy after simulation
-	//-TODO:RDG:mGPU
-	//if (bIsDeformedValid)
-	//{
-	//	const FNDIHairStrandsBuffer* HairStrandsBuffer = ProxyData->HairStrandsBuffer;
-	//	FRHIBuffer* DeformedPositionBuffer = HairStrandsBuffer->SourceDeformedResources->DeformedPositionBuffer[HairStrandsBuffer->SourceDeformedResources->CurrentIndex].Buffer->GetRHI();
-
-	//	Context.ComputeDispatchInterface->MultiGPUResourceModified(RHICmdList, DeformedPositionBuffer, false, true);
-	//}
+	if (bIsDeformedValid)
+	{
+		const FNDIHairStrandsBuffer* HairStrandsBuffer = ProxyData->HairStrandsBuffer;
+		FRHIBuffer* DeformedPositionBuffer = HairStrandsBuffer->SourceDeformedResources->DeformedPositionBuffer[HairStrandsBuffer->SourceDeformedResources->CurrentIndex].Buffer->GetRHI();
+		Context.GetComputeDispatchInterface().MultiGPUResourceModified(Context.GetGraphBuilder(), DeformedPositionBuffer, false, true);
+	}
 
 	if (Context.IsFinalPostSimulate())
 	{
