- Allow Cube / 2DArray / Volume Sample DIs to select both the texture and render target types
- SetTextureObject function can also take either parameter for these now

diff --git a/Engine/Plugins/FX/Niagara/Source/Niagara/Classes/NiagaraDataInterface2DArrayTexture.h b/Engine/Plugins/FX/Niagara/Source/Niagara/Classes/NiagaraDataInterface2DArrayTexture.h
--- a/Engine/Plugins/FX/Niagara/Source/Niagara/Classes/NiagaraDataInterface2DArrayTexture.h
+++ b/Engine/Plugins/FX/Niagara/Source/Niagara/Classes/NiagaraDataInterface2DArrayTexture.h
@@ -22,8 +22,8 @@ class NIAGARA_API UNiagaraDataInterface2DArrayTexture : public UNiagaraDataInter
 	END_SHADER_PARAMETER_STRUCT()
 
 public:
-	UPROPERTY(EditAnywhere, Category = "Texture")
-	TObjectPtr<UTexture2DArray> Texture;
+	UPROPERTY(EditAnywhere, Category = "Texture", meta = (AllowedClasses = "/Script/Engine.Texture2DArray,/Script/Engine.TextureRenderTarget2DArray"))
+	TObjectPtr<UTexture> Texture;
 
 	UPROPERTY(EditAnywhere, Category = "Texture", meta = (ToolTip = "When valid the user parameter is used as the texture rather than the one on the data interface"))
 	FNiagaraUserParameterBinding TextureUserParameter;
@@ -59,7 +59,7 @@ public:
 	virtual void BuildShaderParameters(FNiagaraShaderParametersBuilder& ShaderParametersBuilder) const override;
 	virtual void SetShaderParameters(const FNiagaraDataInterfaceSetShaderParametersContext& Context) const override;
 
-	void SetTexture(UTexture2DArray* InTexture);
+	void SetTexture(UTexture* InTexture);
 
 protected:
 	virtual bool CopyToInternal(UNiagaraDataInterface* Destination) const override;
diff --git a/Engine/Plugins/FX/Niagara/Source/Niagara/Classes/NiagaraDataInterfaceCubeTexture.h b/Engine/Plugins/FX/Niagara/Source/Niagara/Classes/NiagaraDataInterfaceCubeTexture.h
--- a/Engine/Plugins/FX/Niagara/Source/Niagara/Classes/NiagaraDataInterfaceCubeTexture.h
+++ b/Engine/Plugins/FX/Niagara/Source/Niagara/Classes/NiagaraDataInterfaceCubeTexture.h
@@ -22,8 +22,8 @@ class NIAGARA_API UNiagaraDataInterfaceCubeTexture : public UNiagaraDataInterfac
 	END_SHADER_PARAMETER_STRUCT()
 
 public:
-	UPROPERTY(EditAnywhere, Category = "Texture")
-	TObjectPtr<UTextureCube> Texture;
+	UPROPERTY(EditAnywhere, Category = "Texture", meta=(AllowedClasses = "/Script/Engine.TextureCube,/Script/Engine.TextureRenderTargetCube"))
+	TObjectPtr<UTexture> Texture;
 
 	UPROPERTY(EditAnywhere, Category = "Texture", meta = (ToolTip = "When valid the user parameter is used as the texture rather than the one on the data interface"))
 	FNiagaraUserParameterBinding TextureUserParameter;
@@ -60,7 +60,7 @@ public:
 	virtual void BuildShaderParameters(FNiagaraShaderParametersBuilder& ShaderParametersBuilder) const override;
 	virtual void SetShaderParameters(const FNiagaraDataInterfaceSetShaderParametersContext& Context) const override;
 
-	void SetTexture(UTextureCube* InTexture);
+	void SetTexture(UTexture* InTexture);
 
 protected:
 	virtual bool CopyToInternal(UNiagaraDataInterface* Destination) const override;
diff --git a/Engine/Plugins/FX/Niagara/Source/Niagara/Classes/NiagaraDataInterfaceVolumeTexture.h b/Engine/Plugins/FX/Niagara/Source/Niagara/Classes/NiagaraDataInterfaceVolumeTexture.h
--- a/Engine/Plugins/FX/Niagara/Source/Niagara/Classes/NiagaraDataInterfaceVolumeTexture.h
+++ b/Engine/Plugins/FX/Niagara/Source/Niagara/Classes/NiagaraDataInterfaceVolumeTexture.h
@@ -22,8 +22,8 @@ class NIAGARA_API UNiagaraDataInterfaceVolumeTexture : public UNiagaraDataInterf
 	END_SHADER_PARAMETER_STRUCT()
 
 public:
-	UPROPERTY(EditAnywhere, Category = "Texture")
-	TObjectPtr<UVolumeTexture> Texture;
+	UPROPERTY(EditAnywhere, Category = "Texture", meta=(AllowedClasses = "/Script/Engine.VolumeTexture,/Script/Engine.TextureRenderTargetVolume"))
+	TObjectPtr<UTexture> Texture;
 
 	UPROPERTY(EditAnywhere, Category = "Texture", meta = (ToolTip = "When valid the user parameter is used as the texture rather than the one on the data interface"))
 	FNiagaraUserParameterBinding TextureUserParameter;
@@ -60,7 +60,7 @@ public:
 	virtual void BuildShaderParameters(FNiagaraShaderParametersBuilder& ShaderParametersBuilder) const override;
 	virtual void SetShaderParameters(const FNiagaraDataInterfaceSetShaderParametersContext& Context) const override;
 
-	void SetTexture(UVolumeTexture* InTexture);
+	void SetTexture(UTexture* InTexture);
 
 protected:
 	virtual bool CopyToInternal(UNiagaraDataInterface* Destination) const override;
diff --git a/Engine/Plugins/FX/Niagara/Source/Niagara/Private/NiagaraDataInterface2DArrayTexture.cpp b/Engine/Plugins/FX/Niagara/Source/Niagara/Private/NiagaraDataInterface2DArrayTexture.cpp
--- a/Engine/Plugins/FX/Niagara/Source/Niagara/Private/NiagaraDataInterface2DArrayTexture.cpp
+++ b/Engine/Plugins/FX/Niagara/Source/Niagara/Private/NiagaraDataInterface2DArrayTexture.cpp
@@ -290,7 +290,7 @@ void UNiagaraDataInterface2DArrayTexture::SetShaderParameters(const FNiagaraData
 	}
 }
 
-void UNiagaraDataInterface2DArrayTexture::SetTexture(UTexture2DArray* InTexture)
+void UNiagaraDataInterface2DArrayTexture::SetTexture(UTexture* InTexture)
 {
 	if (InTexture)
 	{
diff --git a/Engine/Plugins/FX/Niagara/Source/Niagara/Private/NiagaraDataInterfaceCubeTexture.cpp b/Engine/Plugins/FX/Niagara/Source/Niagara/Private/NiagaraDataInterfaceCubeTexture.cpp
--- a/Engine/Plugins/FX/Niagara/Source/Niagara/Private/NiagaraDataInterfaceCubeTexture.cpp
+++ b/Engine/Plugins/FX/Niagara/Source/Niagara/Private/NiagaraDataInterfaceCubeTexture.cpp
@@ -310,7 +310,7 @@ void UNiagaraDataInterfaceCubeTexture::SetShaderParameters(const FNiagaraDataInt
 	}
 }
 
-void UNiagaraDataInterfaceCubeTexture::SetTexture(UTextureCube* InTexture)
+void UNiagaraDataInterfaceCubeTexture::SetTexture(UTexture* InTexture)
 {
 	if (InTexture)
 	{
diff --git a/Engine/Plugins/FX/Niagara/Source/Niagara/Private/NiagaraDataInterfaceVolumeTexture.cpp b/Engine/Plugins/FX/Niagara/Source/Niagara/Private/NiagaraDataInterfaceVolumeTexture.cpp
--- a/Engine/Plugins/FX/Niagara/Source/Niagara/Private/NiagaraDataInterfaceVolumeTexture.cpp
+++ b/Engine/Plugins/FX/Niagara/Source/Niagara/Private/NiagaraDataInterfaceVolumeTexture.cpp
@@ -330,7 +330,7 @@ void UNiagaraDataInterfaceVolumeTexture::SetShaderParameters(const FNiagaraDataI
 	}
 }
 
-void UNiagaraDataInterfaceVolumeTexture::SetTexture(UVolumeTexture* InTexture)
+void UNiagaraDataInterfaceVolumeTexture::SetTexture(UTexture* InTexture)
 {
 	if (InTexture)
 	{
diff --git a/Engine/Plugins/FX/Niagara/Source/Niagara/Private/NiagaraFunctionLibrary.cpp b/Engine/Plugins/FX/Niagara/Source/Niagara/Private/NiagaraFunctionLibrary.cpp
--- a/Engine/Plugins/FX/Niagara/Source/Niagara/Private/NiagaraFunctionLibrary.cpp
+++ b/Engine/Plugins/FX/Niagara/Source/Niagara/Private/NiagaraFunctionLibrary.cpp
@@ -21,6 +21,9 @@
 #include "Engine/VolumeTexture.h"
 #include "Engine/Texture2DArray.h"
 #include "Engine/TextureCube.h"
+#include "Engine/TextureRenderTarget2DArray.h"
+#include "Engine/TextureRenderTargetVolume.h"
+#include "Engine/TextureRenderTargetCube.h"
 #include "NiagaraGpuComputeDispatchInterface.h"
 
 #define LOCTEXT_NAMESPACE "NiagaraFunctionLibrary"
@@ -508,7 +511,7 @@ void UNiagaraFunctionLibrary::SetTextureObject(UNiagaraComponent* NiagaraSystem,
 
 		TextureDI->SetTexture(Texture);
 	}
-	else if (UTexture2DArray* Texture2DArray = Cast<UTexture2DArray>(Texture))
+	else if (Texture->IsA<UTexture2DArray>() || Texture->IsA<UTextureRenderTarget2DArray>())
 	{
 		const FNiagaraVariable Variable(FNiagaraTypeDefinition(UNiagaraDataInterface2DArrayTexture::StaticClass()), *OverrideName);
 		const int32 Index = OverrideParameters.IndexOf(Variable);
@@ -525,9 +528,9 @@ void UNiagaraFunctionLibrary::SetTextureObject(UNiagaraComponent* NiagaraSystem,
 			return;
 		}
 
-		TextureDI->SetTexture(Texture2DArray);
+		TextureDI->SetTexture(Texture);
 	}
-	else if (UVolumeTexture* TextureVolume = Cast<UVolumeTexture>(Texture))
+	else if (Texture->IsA<UVolumeTexture>() || Texture->IsA<UTextureRenderTargetVolume>())
 	{
 		const FNiagaraVariable Variable(FNiagaraTypeDefinition(UNiagaraDataInterfaceVolumeTexture::StaticClass()), *OverrideName);
 		const int32 Index = OverrideParameters.IndexOf(Variable);
@@ -544,9 +547,9 @@ void UNiagaraFunctionLibrary::SetTextureObject(UNiagaraComponent* NiagaraSystem,
 			return;
 		}
 
-		TextureDI->SetTexture(TextureVolume);
+		TextureDI->SetTexture(Texture);
 	}
-	else if (UTextureCube* TextureCube = Cast<UTextureCube>(Texture))
+	else if (Texture->IsA<UTextureCube>() || Texture->IsA<UTextureRenderTargetCube>())
 	{
 		const FNiagaraVariable Variable(FNiagaraTypeDefinition(UNiagaraDataInterfaceCubeTexture::StaticClass()), *OverrideName);
 		const int32 Index = OverrideParameters.IndexOf(Variable);
@@ -563,7 +566,7 @@ void UNiagaraFunctionLibrary::SetTextureObject(UNiagaraComponent* NiagaraSystem,
 			return;
 		}
 
-		TextureDI->SetTexture(TextureCube);
+		TextureDI->SetTexture(Texture);
 	}
 	else
 	{
