- Update Array DI to ShaderParameters

diff --git a/Engine/Plugins/FX/Niagara/Shaders/Private/NiagaraDataInterfaceArrayRWTemplate.ush b/Engine/Plugins/FX/Niagara/Shaders/Private/NiagaraDataInterfaceArrayRWTemplate.ush
--- a/Engine/Plugins/FX/Niagara/Shaders/Private/NiagaraDataInterfaceArrayRWTemplate.ush
+++ b/Engine/Plugins/FX/Niagara/Shaders/Private/NiagaraDataInterfaceArrayRWTemplate.ush
@@ -5,11 +5,11 @@
 // RWBufferRead		- Hlsl code to read from the buffer, i.e Value = BUFFER_NAME[Index]
 // RWBufferWrite	- Hlsl code to write to the buffer, i.e BUFFER_NAME[Index] = Value
 
-RWBuffer<{RWBufferType}> ArrayBuffer_{ParameterName};
-int2 ArrayBufferParams_{ParameterName};				// Where x=CountOffset & y=BufferCapacity
+RWBuffer<{RWBufferType}>	{ParameterName}_ArrayRWBuffer;
+int2						{ParameterName}_ArrayBufferParams;				// Where x=CountOffset & y=BufferCapacity
 
-int GetCountOffset_{ParameterName}() { return ArrayBufferParams_{ParameterName}.x; }
-int GetCapacity_{ParameterName}() { return ArrayBufferParams_{ParameterName}.y; }
+int GetCountOffset_{ParameterName}() { return {ParameterName}_ArrayBufferParams.x; }
+int GetCapacity_{ParameterName}() { return {ParameterName}_ArrayBufferParams.y; }
 int GetLength_{ParameterName}() { return RWInstanceCounts[GetCountOffset_{ParameterName}()]; }
 
 void Length_{ParameterName}(out int ArrayLength)
@@ -31,7 +31,7 @@ void Get_{ParameterName}(int Index, out {VariableType} Value)
 {
 	int Length = GetLength_{ParameterName}();
 	Index = Length > 0 ? clamp(Index, 0, Length - 1) : GetCapacity_{ParameterName}();
-	#define BUFFER_NAME ArrayBuffer_{ParameterName}
+	#define BUFFER_NAME {ParameterName}_ArrayRWBuffer
 	{RWBufferRead}
 	#undef BUFFER_NAME
 }
@@ -42,7 +42,7 @@ void SetArrayElem_{ParameterName}_UEImpureCall(bool bSkip, int Index, {VariableT
 	IsValidIndex_{ParameterName}(Index, bIsValidIndex);
 	if ( !bSkip && bIsValidIndex )
 	{
-		#define BUFFER_NAME ArrayBuffer_{ParameterName}
+		#define BUFFER_NAME {ParameterName}_ArrayRWBuffer
 		{RWBufferWrite}
 		#undef BUFFER_NAME
 	}
@@ -59,7 +59,7 @@ void Add_{ParameterName}_UEImpureCall(bool bSkip, {VariableType} Value)
 	InterlockedAdd(RWInstanceCounts[GetCountOffset_{ParameterName}()], 1, Index);
 	if ( Index < GetCapacity_{ParameterName}() )
 	{
-		#define BUFFER_NAME ArrayBuffer_{ParameterName}
+		#define BUFFER_NAME {ParameterName}_ArrayRWBuffer
 		{RWBufferWrite}
 		#undef BUFFER_NAME
 	}
@@ -86,7 +86,7 @@ void RemoveLastElem_{ParameterName}_UEImpureCall(bool bSkip, out {VariableType}
 	}
 
 	Index = bOutValid ? Index : GetCapacity_{ParameterName}();
-	#define BUFFER_NAME ArrayBuffer_{ParameterName}
+	#define BUFFER_NAME {ParameterName}_ArrayRWBuffer
 	{RWBufferRead}
 	#undef BUFFER_NAME
 }
diff --git a/Engine/Plugins/FX/Niagara/Shaders/Private/NiagaraDataInterfaceArrayTemplate.ush b/Engine/Plugins/FX/Niagara/Shaders/Private/NiagaraDataInterfaceArrayTemplate.ush
--- a/Engine/Plugins/FX/Niagara/Shaders/Private/NiagaraDataInterfaceArrayTemplate.ush
+++ b/Engine/Plugins/FX/Niagara/Shaders/Private/NiagaraDataInterfaceArrayTemplate.ush
@@ -4,11 +4,11 @@
 // ReadBufferType	- Read Buffer Type, i.e. float, float2, etc
 // ReadBufferRead	- Hlsl code to read from the buffer, i.e Value = BUFFER_NAME[Index]
 
-Buffer<{ReadBufferType}> ArrayBuffer_{ParameterName};
-int2 ArrayBufferParams_{ParameterName};				// Where x=Length & y=max(Length-1,0)
+Buffer<{ReadBufferType}>	{ParameterName}_ArrayReadBuffer;
+int2						{ParameterName}_ArrayBufferParams;				// Where x=Length & y=max(Length-1,0)
 
-int GetLength_{ParameterName}() { return ArrayBufferParams_{ParameterName}.x; }
-int GetLengthMinusOne_{ParameterName}() { return ArrayBufferParams_{ParameterName}.y; }
+int GetLength_{ParameterName}() { return {ParameterName}_ArrayBufferParams.x; }
+int GetLengthMinusOne_{ParameterName}() { return {ParameterName}_ArrayBufferParams.y; }
 
 void Length_{ParameterName}(out int ArrayLength)
 {
@@ -28,7 +28,7 @@ void LastIndex_{ParameterName}(out int Index)
 void Get_{ParameterName}(int Index, out {VariableType} Value)
 {
 	Index = clamp(Index, 0, GetLengthMinusOne_{ParameterName}());
-	#define BUFFER_NAME ArrayBuffer_{ParameterName}
+	#define BUFFER_NAME {ParameterName}_ArrayReadBuffer
 	{ReadBufferRead}
 	#undef BUFFER_NAME
 }
diff --git a/Engine/Plugins/FX/Niagara/Source/Niagara/Classes/NiagaraDataInterfaceArray.h b/Engine/Plugins/FX/Niagara/Source/Niagara/Classes/NiagaraDataInterfaceArray.h
--- a/Engine/Plugins/FX/Niagara/Source/Niagara/Classes/NiagaraDataInterfaceArray.h
+++ b/Engine/Plugins/FX/Niagara/Source/Niagara/Classes/NiagaraDataInterfaceArray.h
@@ -13,6 +13,12 @@
 
 struct INDIArrayProxyBase : public FNiagaraDataInterfaceProxyRW
 {
+	BEGIN_SHADER_PARAMETER_STRUCT(FShaderParameters,)
+		SHADER_PARAMETER(FIntPoint,		ArrayBufferParams)
+		SHADER_PARAMETER_SRV(Buffer,	ArrayReadBuffer)
+		SHADER_PARAMETER_UAV(RWBuffer,	ArrayRWBuffer)
+	END_SHADER_PARAMETER_STRUCT()
+
 	virtual ~INDIArrayProxyBase() {}
 	virtual void GetFunctions(TArray<FNiagaraFunctionSignature>& OutFunctions) const = 0;
 	virtual void GetVMExternalFunction(const FVMExternalFunctionBindingInfo& BindingInfo, void* InstanceData, FVMExternalFunction &OutFunc) = 0;
@@ -31,11 +37,7 @@ struct INDIArrayProxyBase : public FNiagaraDataInterfaceProxyRW
 	virtual bool InitPerInstanceData(UNiagaraDataInterface* DataInterface, void* InPerInstanceData, FNiagaraSystemInstance* SystemInstance) = 0;
 	virtual void DestroyPerInstanceData(void* InPerInstanceData, FNiagaraSystemInstance* SystemInstance) = 0;
 	virtual void ProvidePerInstanceDataForRenderThread(void* DataForRenderThread, void* PerInstanceData, const FNiagaraSystemInstanceID& SystemInstance) = 0;
-	virtual FNiagaraDataInterfaceParametersCS* CreateComputeParameters() const = 0;
-	virtual const FTypeLayoutDesc* GetComputeParametersTypeDesc() const = 0;
-	virtual void BindParameters(FNiagaraDataInterfaceParametersCS* Base, const FNiagaraDataInterfaceGPUParamInfo& ParameterInfo, const class FShaderParameterMap& ParameterMap) = 0;
-	virtual void SetParameters(const FNiagaraDataInterfaceParametersCS* Base, FRHICommandList& RHICmdList, const FNiagaraDataInterfaceSetArgs& Context) const = 0;
-	virtual void UnsetParameters(const FNiagaraDataInterfaceParametersCS* Base, FRHICommandList& RHICmdList, const FNiagaraDataInterfaceSetArgs& Context) const = 0;
+	virtual void SetShaderParameters(FShaderParameters* ShaderParameters, FNiagaraSystemInstanceID SystemInstanceID) const = 0;
 };
 
 UCLASS(abstract, EditInlineNew)
@@ -58,7 +60,7 @@ public:
 #if WITH_EDITORONLY_DATA
 	virtual void GetParameterDefinitionHLSL(const FNiagaraDataInterfaceGPUParamInfo& ParamInfo, FString& OutHLSL) override { GetProxyAs<INDIArrayProxyBase>()->GetParameterDefinitionHLSL(ParamInfo, OutHLSL); }
 	virtual bool GetFunctionHLSL(const FNiagaraDataInterfaceGPUParamInfo& ParamInfo, const FNiagaraDataInterfaceGeneratedFunction& FunctionInfo, int FunctionInstanceIndex, FString& OutHLSL) override { return GetProxyAs<INDIArrayProxyBase>()->GetFunctionHLSL(ParamInfo, FunctionInfo, FunctionInstanceIndex, OutHLSL); }
-	virtual bool AppendCompileHash(FNiagaraCompileHashVisitor* InVisitor) const override { return GetProxyAs<INDIArrayProxyBase>()->AppendCompileHash(InVisitor); }
+	virtual bool AppendCompileHash(FNiagaraCompileHashVisitor* InVisitor) const override;
 	virtual bool UpgradeFunctionCall(FNiagaraFunctionSignature& FunctionSignature) override { return GetProxyAs<INDIArrayProxyBase>()->UpgradeFunctionCall(FunctionSignature); }
 #endif
 	virtual bool CanExecuteOnTarget(ENiagaraSimTarget Target) const override { return true; }
@@ -76,11 +78,9 @@ public:
 
 	virtual void ProvidePerInstanceDataForRenderThread(void* DataForRenderThread, void* PerInstanceData, const FNiagaraSystemInstanceID& SystemInstance) override { GetProxyAs<INDIArrayProxyBase>()->ProvidePerInstanceDataForRenderThread(DataForRenderThread, PerInstanceData, SystemInstance); }
 
-	virtual FNiagaraDataInterfaceParametersCS* CreateComputeParameters() const override { return GetProxyAs<INDIArrayProxyBase>()->CreateComputeParameters(); }
-	virtual const FTypeLayoutDesc* GetComputeParametersTypeDesc() const override { return GetProxyAs<INDIArrayProxyBase>()->GetComputeParametersTypeDesc(); }
-	virtual void BindParameters(FNiagaraDataInterfaceParametersCS* Base, const FNiagaraDataInterfaceGPUParamInfo& ParameterInfo, const class FShaderParameterMap& ParameterMap) override { GetProxyAs<INDIArrayProxyBase>()->BindParameters(Base, ParameterInfo, ParameterMap); }
-	virtual void SetParameters(const FNiagaraDataInterfaceParametersCS* Base, FRHICommandList& RHICmdList, const FNiagaraDataInterfaceSetArgs& Context) const override { return GetProxyAs<INDIArrayProxyBase>()->SetParameters(Base, RHICmdList, Context); }
-	virtual void UnsetParameters(const FNiagaraDataInterfaceParametersCS* Base, FRHICommandList& RHICmdList, const FNiagaraDataInterfaceSetArgs& Context) const override { return GetProxyAs<INDIArrayProxyBase>()->UnsetParameters(Base, RHICmdList, Context); }
+	virtual bool UseLegacyShaderBindings() const  override { return false; }
+	virtual void BuildShaderParameters(FNiagaraShaderParametersBuilder& ShaderParametersBuilder) const override;
+	virtual void SetShaderParameters(const FNiagaraDataInterfaceSetShaderParametersContext& Context) const override;
 	//UNiagaraDataInterface Interface
 
 	/** ReadWrite lock to ensure safe access to the underlying array. */
diff --git a/Engine/Plugins/FX/Niagara/Source/Niagara/Classes/NiagaraDataInterfaceArrayImpl.h b/Engine/Plugins/FX/Niagara/Source/Niagara/Classes/NiagaraDataInterfaceArrayImpl.h
--- a/Engine/Plugins/FX/Niagara/Source/Niagara/Classes/NiagaraDataInterfaceArrayImpl.h
+++ b/Engine/Plugins/FX/Niagara/Source/Niagara/Classes/NiagaraDataInterfaceArrayImpl.h
@@ -89,16 +89,6 @@ struct FNiagaraDataInterfaceArrayImplHelper
 	}
 };
 
-//////////////////////////////////////////////////////////////////////////
-// Compute Shader Impl
-struct FNiagaraDataInterfaceParametersCS_ArrayImpl : public FNiagaraDataInterfaceParametersCS
-{
-	DECLARE_TYPE_LAYOUT(FNiagaraDataInterfaceParametersCS_ArrayImpl, NonVirtual);
-
-	LAYOUT_FIELD(FShaderResourceParameter, BufferParam);
-	LAYOUT_FIELD(FShaderParameter, BufferParamsParam);
-};
-
 //////////////////////////////////////////////////////////////////////////
 // Instance Data, Proxy with Impl
 
@@ -924,75 +914,20 @@ struct FNDIArrayProxyImpl : public INDIArrayProxyBase
 		InstanceData_GT->~FNDIArrayInstanceData_GameThread<TArrayType>();
 	}
 
-	FNiagaraDataInterfaceParametersCS* CreateComputeParameters() const
-	{
-		if (FNDIArrayImplHelper<TArrayType>::bSupportsGPU)
-		{
-			return new FNiagaraDataInterfaceParametersCS_ArrayImpl();
-		}
-		return nullptr;
-	}
-
-	const FTypeLayoutDesc* GetComputeParametersTypeDesc() const
-	{
-		if (FNDIArrayImplHelper<TArrayType>::bSupportsGPU)
-		{
-			return &StaticGetTypeLayoutDesc<FNiagaraDataInterfaceParametersCS_ArrayImpl>();
-		}
-		return nullptr;
-	}
-
-	void BindParameters(FNiagaraDataInterfaceParametersCS* Base, const FNiagaraDataInterfaceGPUParamInfo& ParameterInfo, const class FShaderParameterMap& ParameterMap)
+	virtual void SetShaderParameters(FShaderParameters* ShaderParameters, FNiagaraSystemInstanceID SystemInstanceID) const override
 	{
-		check(FNDIArrayImplHelper<TArrayType>::bSupportsGPU);
-
-		FNiagaraDataInterfaceParametersCS_ArrayImpl* TypedCS = static_cast<FNiagaraDataInterfaceParametersCS_ArrayImpl*>(Base);
-		TypedCS->BufferParam.Bind(ParameterMap, *(TEXT("ArrayBuffer_") + ParameterInfo.DataInterfaceHLSLSymbol));
-		TypedCS->BufferParamsParam.Bind(ParameterMap, *(TEXT("ArrayBufferParams_") + ParameterInfo.DataInterfaceHLSLSymbol));
-	}
-
-	void SetParameters(const FNiagaraDataInterfaceParametersCS* Base, FRHICommandList& RHICmdList, const FNiagaraDataInterfaceSetArgs& Context) const
-	{
-		check(FNDIArrayImplHelper<TArrayType>::bSupportsGPU);
-
-		const FNiagaraDataInterfaceParametersCS_ArrayImpl* TypedCS = static_cast<const FNiagaraDataInterfaceParametersCS_ArrayImpl*>(Base);
-		FRHIComputeShader* ComputeShaderRHI = Context.Shader.GetComputeShader();
-
-		const auto* ArrayProxy = static_cast<const FNDIArrayProxyImpl<TArrayType, TOwnerType>*>(Context.DataInterface);
-		const auto* InstanceData_RT = &ArrayProxy->PerInstanceData_RenderThread.FindChecked(Context.SystemInstanceID);
-
+		const auto* InstanceData_RT = &PerInstanceData_RenderThread.FindChecked(SystemInstanceID);
 		if ( InstanceData_RT->IsReadOnly() )
 		{
-			SetSRVParameter(RHICmdList, ComputeShaderRHI, TypedCS->BufferParam, InstanceData_RT->ArraySRV);
-
-			// Where x=Length & y=max(Length-1,0)
-			const int32 BufferParams[] = { InstanceData_RT->NumElements, FMath::Max(0, InstanceData_RT->NumElements - 1) };
-			SetShaderValue(RHICmdList, ComputeShaderRHI, TypedCS->BufferParamsParam, BufferParams);
+			ShaderParameters->ArrayBufferParams.X	= InstanceData_RT->NumElements;
+			ShaderParameters->ArrayBufferParams.Y	= FMath::Max(0, InstanceData_RT->NumElements - 1);
+			ShaderParameters->ArrayReadBuffer		= InstanceData_RT->ArraySRV;
 		}
 		else
 		{
-			SetUAVParameter(RHICmdList, ComputeShaderRHI, TypedCS->BufferParam, InstanceData_RT->ArrayUAV);
-
-			// Where x=CountOffset & y=BufferCapacity
-			const int32 BufferParams[] = { (int32)InstanceData_RT->CountOffset, InstanceData_RT->NumElements };
-			SetShaderValue(RHICmdList, ComputeShaderRHI, TypedCS->BufferParamsParam, BufferParams);
-		}
-	}
-
-	void UnsetParameters(const FNiagaraDataInterfaceParametersCS* Base, FRHICommandList& RHICmdList, const FNiagaraDataInterfaceSetArgs& Context) const
-	{
-		check(FNDIArrayImplHelper<TArrayType>::bSupportsGPU);
-
-		const FNiagaraDataInterfaceParametersCS_ArrayImpl* TypedCS = static_cast<const FNiagaraDataInterfaceParametersCS_ArrayImpl*>(Base);
-		FRHIComputeShader* ComputeShaderRHI = Context.Shader.GetComputeShader();
-		const auto* ArrayProxy = static_cast<const FNDIArrayProxyImpl<TArrayType, TOwnerType>*>(Context.DataInterface);
-		const auto* InstanceData_RT = &ArrayProxy->PerInstanceData_RenderThread.FindChecked(Context.SystemInstanceID);
-
-		if ( !InstanceData_RT->IsReadOnly() )
-		{
-			SetUAVParameter(RHICmdList, ComputeShaderRHI, TypedCS->BufferParam, nullptr);
-			//-OPT: Should be batched
-			RHICmdList.Transition(FRHITransitionInfo(InstanceData_RT->ArrayUAV, ERHIAccess::UAVCompute, ERHIAccess::UAVCompute));
+			ShaderParameters->ArrayBufferParams.X	= (int32)InstanceData_RT->CountOffset;
+			ShaderParameters->ArrayBufferParams.Y	= InstanceData_RT->NumElements;
+			ShaderParameters->ArrayRWBuffer			= InstanceData_RT->ArrayUAV;
 		}
 	}
 
diff --git a/Engine/Plugins/FX/Niagara/Source/Niagara/Private/NiagaraDataInterfaceArray.cpp b/Engine/Plugins/FX/Niagara/Source/Niagara/Private/NiagaraDataInterfaceArray.cpp
--- a/Engine/Plugins/FX/Niagara/Source/Niagara/Private/NiagaraDataInterfaceArray.cpp
+++ b/Engine/Plugins/FX/Niagara/Source/Niagara/Private/NiagaraDataInterfaceArray.cpp
@@ -1,6 +1,7 @@
 // Copyright Epic Games, Inc. All Rights Reserved.
 
 #include "NiagaraDataInterfaceArray.h"
+#include "NiagaraShaderParametersBuilder.h"
 
 #define LOCTEXT_NAMESPACE "UNiagaraDataInterfaceArray"
 
@@ -66,4 +67,27 @@ bool UNiagaraDataInterfaceArray::Equals(const UNiagaraDataInterface* Other) cons
 	return GetProxyAs<INDIArrayProxyBase>()->Equals(OtherTyped->GetProxyAs<INDIArrayProxyBase>());
 }
 
+#if WITH_EDITORONLY_DATA
+bool UNiagaraDataInterfaceArray::AppendCompileHash(FNiagaraCompileHashVisitor* InVisitor) const
+{
+	bool bSuccess = Super::AppendCompileHash(InVisitor);
+	bSuccess &= InVisitor->UpdateShaderParameters<INDIArrayProxyBase::FShaderParameters>();
+	bSuccess &= GetProxyAs<INDIArrayProxyBase>()->AppendCompileHash(InVisitor);
+	return bSuccess;
+}
+#endif
+
+void UNiagaraDataInterfaceArray::BuildShaderParameters(FNiagaraShaderParametersBuilder& ShaderParametersBuilder) const
+{
+	ShaderParametersBuilder.AddNestedStruct<INDIArrayProxyBase::FShaderParameters>();
+}
+
+void UNiagaraDataInterfaceArray::SetShaderParameters(const FNiagaraDataInterfaceSetShaderParametersContext& Context) const
+{
+	INDIArrayProxyBase& ArrayProxy = Context.GetProxy<INDIArrayProxyBase>();
+	INDIArrayProxyBase::FShaderParameters* ShaderParameters = Context.GetParameterNestedStruct<INDIArrayProxyBase::FShaderParameters>();
+
+	ArrayProxy.SetShaderParameters(ShaderParameters, Context.GetSystemInstanceID());
+}
+
 #undef LOCTEXT_NAMESPACE
diff --git a/Engine/Plugins/FX/Niagara/Source/Niagara/Private/NiagaraDataInterfaceArrayImpl.cpp b/Engine/Plugins/FX/Niagara/Source/Niagara/Private/NiagaraDataInterfaceArrayImpl.cpp
--- a/Engine/Plugins/FX/Niagara/Source/Niagara/Private/NiagaraDataInterfaceArrayImpl.cpp
+++ b/Engine/Plugins/FX/Niagara/Source/Niagara/Private/NiagaraDataInterfaceArrayImpl.cpp
@@ -60,5 +60,3 @@ bool FNiagaraDataInterfaceArrayImplHelper::UpgradeFunctionCall(FNiagaraFunctionS
 	return true;
 }
 #endif
-
-IMPLEMENT_TYPE_LAYOUT(FNiagaraDataInterfaceParametersCS_ArrayImpl);
