- Fix for not disable backface culling on various ribbon types
- Rename CMDList -> RHICmdList for consistency with other parts of the code

#rnx

diff --git a/Engine/Plugins/FX/Niagara/Source/Niagara/Private/NiagaraRendererRibbons.cpp b/Engine/Plugins/FX/Niagara/Source/Niagara/Private/NiagaraRendererRibbons.cpp
--- a/Engine/Plugins/FX/Niagara/Source/Niagara/Private/NiagaraRendererRibbons.cpp
+++ b/Engine/Plugins/FX/Niagara/Source/Niagara/Private/NiagaraRendererRibbons.cpp
@@ -555,16 +555,16 @@ private:
 		}
 	}
 
-	void OnPostPreRender(FRHICommandListImmediate& CMDList)
+	void OnPostPreRender(FRHICommandListImmediate& RHICmdList)
 	{
-		GenerateAllGPUData(CMDList);
+		GenerateAllGPUData(RHICmdList);
 		
 		SortManager->PostPreRenderEvent.Remove(SortManagerEventRef);
 		SortManagerEventRef.Reset();
 		SortManager = nullptr;
 	}
 
-	void GenerateAllGPUData(FRHICommandListImmediate& CMDList);
+	void GenerateAllGPUData(FRHICommandListImmediate& RHICmdList);
 
 public:
 	static FNiagaraRibbonComputeDispatchManager& Get()
@@ -1052,7 +1052,7 @@ void FNiagaraRendererRibbons::GenerateShapeStateMultiPlane(FNiagaraRibbonShapeGe
 void FNiagaraRendererRibbons::GenerateShapeStateTube(FNiagaraRibbonShapeGeometryData& State, int32 TubeSubdivisions)
 {
 	State.Shape = ENiagaraRibbonShapeMode::Tube;
-	State.bDisableBackfaceCulling = false;
+	State.bDisableBackfaceCulling = true;
 	State.bShouldFlipNormalToView = true;
 	State.TrianglesPerSegment = 2 * TubeSubdivisions;
 	State.NumVerticesInSlice = TubeSubdivisions + 1;
@@ -1080,7 +1080,7 @@ void FNiagaraRendererRibbons::GenerateShapeStateTube(FNiagaraRibbonShapeGeometry
 void FNiagaraRendererRibbons::GenerateShapeStateCustom(FNiagaraRibbonShapeGeometryData& State, const TArray<FNiagaraRibbonShapeCustomVertex>& CustomVertices)
 {
 	State.Shape = ENiagaraRibbonShapeMode::Custom;
-	State.bDisableBackfaceCulling = false;
+	State.bDisableBackfaceCulling = true;
 	State.bShouldFlipNormalToView = true;
 	State.TrianglesPerSegment = 2 * CustomVertices.Num();
 	State.NumVerticesInSlice = CustomVertices.Num() + 1;
@@ -2004,7 +2004,7 @@ static inline TArray<Type> ReadBackBuffer(FRHIBuffer* Buffer, int32 NumValues)
 	return Test;
 }
 
-void FNiagaraRendererRibbons::InitializeViewIndexBuffersGPU(FRHICommandListImmediate& CMDList, FNiagaraGpuComputeDispatchInterface* ComputeDispatchInterface,
+void FNiagaraRendererRibbons::InitializeViewIndexBuffersGPU(FRHICommandListImmediate& RHICmdList, FNiagaraGpuComputeDispatchInterface* ComputeDispatchInterface,
 	const FNiagaraDataBuffer* SourceParticleData, const TSharedPtr<FNiagaraRibbonRenderingFrameViewResources>& RenderingViewResources) const
 {
 	SCOPE_CYCLE_COUNTER(STAT_NiagaraRenderRibbonsGenIndiciesGPU);
@@ -2016,7 +2016,7 @@ void FNiagaraRendererRibbons::InitializeViewIndexBuffersGPU(FRHICommandListImmed
 		return;
 	}
 
-	SCOPED_DRAW_EVENT(CMDList, NiagaraRenderRibbonsGenIndiciesGPU);
+	SCOPED_DRAW_EVENT(RHICmdList, NiagaraRenderRibbonsGenIndiciesGPU);
 	{
 		FNiagaraRibbonCreateIndexBufferParamsCS::FPermutationDomain PermutationVector;
 		PermutationVector.Set<FRibbonWantsAutomaticTessellation>(GenerationConfig.WantsAutomaticTessellation());
@@ -2055,12 +2055,12 @@ void FNiagaraRendererRibbons::InitializeViewIndexBuffersGPU(FRHICommandListImmed
 		Params.GNiagaraRibbonTessellationEnabled = GNiagaraRibbonTessellationEnabled ? 1 : 0;
 		Params.GNiagaraRibbonTessellationMinDisplacementError = GNiagaraRibbonTessellationMinDisplacementError;
 
-		CMDList.Transition(FRHITransitionInfo(RenderingViewResources->IndirectDrawBuffer.UAV, ERHIAccess::SRVMask | ERHIAccess::IndirectArgs, ERHIAccess::UAVCompute));
-		SetComputePipelineState(CMDList, ComputeShader.GetComputeShader());
-		SetShaderParameters(CMDList, ComputeShader, ComputeShader.GetComputeShader(), Params);
-		DispatchComputeShader(CMDList, ComputeShader, 1, 1, 1);
-		UnsetShaderUAVs(CMDList, ComputeShader, ComputeShader.GetComputeShader());
-		CMDList.Transition(FRHITransitionInfo(RenderingViewResources->IndirectDrawBuffer.UAV, ERHIAccess::UAVCompute, ERHIAccess::SRVMask | ERHIAccess::IndirectArgs));
+		RHICmdList.Transition(FRHITransitionInfo(RenderingViewResources->IndirectDrawBuffer.UAV, ERHIAccess::SRVMask | ERHIAccess::IndirectArgs, ERHIAccess::UAVCompute));
+		SetComputePipelineState(RHICmdList, ComputeShader.GetComputeShader());
+		SetShaderParameters(RHICmdList, ComputeShader, ComputeShader.GetComputeShader(), Params);
+		DispatchComputeShader(RHICmdList, ComputeShader, 1, 1, 1);
+		UnsetShaderUAVs(RHICmdList, ComputeShader, ComputeShader.GetComputeShader());
+		RHICmdList.Transition(FRHITransitionInfo(RenderingViewResources->IndirectDrawBuffer.UAV, ERHIAccess::UAVCompute, ERHIAccess::SRVMask | ERHIAccess::IndirectArgs));
 	}
 	
 	// Not possible to have a valid ribbon with less than 2 particles, abort!
@@ -2116,12 +2116,12 @@ void FNiagaraRendererRibbons::InitializeViewIndexBuffersGPU(FRHICommandListImmed
 		Params.SubSegmentBitShift = RenderingViewResources->IndexGenerationSettings.SubSegmentBitShift;
 		Params.SubSegmentBitMask = RenderingViewResources->IndexGenerationSettings.SubSegmentBitMask;
 		
-		CMDList.Transition(FRHITransitionInfo(RenderingViewResources->IndexBuffer.UAV, ERHIAccess::VertexOrIndexBuffer, ERHIAccess::UAVCompute));
-		SetComputePipelineState(CMDList, ComputeShader.GetComputeShader());
-		SetShaderParameters(CMDList, ComputeShader, ComputeShader.GetComputeShader(), Params);
-		DispatchIndirectComputeShader(CMDList, ComputeShader.GetShader(), RenderingViewResources->IndirectDrawBuffer.Buffer, IndirectDispatchArgsOffset);
-		UnsetShaderUAVs(CMDList, ComputeShader, ComputeShader.GetComputeShader());
-		CMDList.Transition(FRHITransitionInfo(RenderingViewResources->IndexBuffer.UAV, ERHIAccess::UAVCompute, ERHIAccess::VertexOrIndexBuffer));
+		RHICmdList.Transition(FRHITransitionInfo(RenderingViewResources->IndexBuffer.UAV, ERHIAccess::VertexOrIndexBuffer, ERHIAccess::UAVCompute));
+		SetComputePipelineState(RHICmdList, ComputeShader.GetComputeShader());
+		SetShaderParameters(RHICmdList, ComputeShader, ComputeShader.GetComputeShader(), Params);
+		DispatchIndirectComputeShader(RHICmdList, ComputeShader.GetShader(), RenderingViewResources->IndirectDrawBuffer.Buffer, IndirectDispatchArgsOffset);
+		UnsetShaderUAVs(RHICmdList, ComputeShader, ComputeShader.GetComputeShader());
+		RHICmdList.Transition(FRHITransitionInfo(RenderingViewResources->IndexBuffer.UAV, ERHIAccess::UAVCompute, ERHIAccess::VertexOrIndexBuffer));
 	}
 }
 
@@ -2272,7 +2272,7 @@ FRibbonComputeUniformParameters FNiagaraRendererRibbons::SetupComputeVertexGenPa
 	return CommonParams;
 }
 
-void FNiagaraRendererRibbons::InitializeVertexBuffersGPU(FRHICommandListImmediate& CMDList, FNiagaraGpuComputeDispatchInterface* ComputeDispatchInterface, const FNiagaraDataBuffer* SourceParticleData,
+void FNiagaraRendererRibbons::InitializeVertexBuffersGPU(FRHICommandListImmediate& RHICmdList, FNiagaraGpuComputeDispatchInterface* ComputeDispatchInterface, const FNiagaraDataBuffer* SourceParticleData,
 	FNiagaraRibbonGPUInitComputeBuffers& TempBuffers, const TSharedPtr<FNiagaraRibbonRenderingFrameResources>& RenderingResources) const
 {	
 	SCOPE_CYCLE_COUNTER(STAT_NiagaraRenderRibbonsGenVerticesGPU);
@@ -2286,9 +2286,9 @@ void FNiagaraRendererRibbons::InitializeVertexBuffersGPU(FRHICommandListImmediat
 	// Clear the command buffer if we just initialized it, or if the sim doesn't have enough data to run
 	if ((!bCanRun || VertexBuffers.bJustCreatedCommandBuffer) && VertexBuffers.GPUComputeCommandLength > 0)
 	{
-		CMDList.Transition(FRHITransitionInfo(VertexBuffers.GPUComputeCommandBuffer.Buffer, ERHIAccess::SRVMask | ERHIAccess::VertexOrIndexBuffer | ERHIAccess::IndirectArgs, ERHIAccess::UAVCompute));
-		CMDList.ClearUAVUint(VertexBuffers.GPUComputeCommandBuffer.UAV, FUintVector4(0));
-		CMDList.Transition(FRHITransitionInfo(VertexBuffers.GPUComputeCommandBuffer.Buffer, ERHIAccess::UAVCompute, ERHIAccess::SRVMask | ERHIAccess::VertexOrIndexBuffer | ERHIAccess::IndirectArgs));
+		RHICmdList.Transition(FRHITransitionInfo(VertexBuffers.GPUComputeCommandBuffer.Buffer, ERHIAccess::SRVMask | ERHIAccess::VertexOrIndexBuffer | ERHIAccess::IndirectArgs, ERHIAccess::UAVCompute));
+		RHICmdList.ClearUAVUint(VertexBuffers.GPUComputeCommandBuffer.UAV, FUintVector4(0));
+		RHICmdList.Transition(FRHITransitionInfo(VertexBuffers.GPUComputeCommandBuffer.Buffer, ERHIAccess::UAVCompute, ERHIAccess::SRVMask | ERHIAccess::VertexOrIndexBuffer | ERHIAccess::IndirectArgs));
 		VertexBuffers.bJustCreatedCommandBuffer = false;
 	}
 	
@@ -2299,7 +2299,7 @@ void FNiagaraRendererRibbons::InitializeVertexBuffersGPU(FRHICommandListImmediat
 	}
 
 	{
-		SCOPED_DRAW_EVENT(CMDList, NiagaraRenderRibbonsGenVerticesSortGPU);
+		SCOPED_DRAW_EVENT(RHICmdList, NiagaraRenderRibbonsGenVerticesSortGPU);
 		SCOPE_CYCLE_COUNTER(STAT_NiagaraRenderRibbonsGenVerticesSortGPU);
 		
 		FNiagaraRibbonSortPhase1CS::FPermutationDomain PermutationVector;
@@ -2337,21 +2337,21 @@ void FNiagaraRendererRibbons::InitializeVertexBuffersGPU(FRHICommandListImmediat
 		}
 
 		{			
-			SCOPED_DRAW_EVENT(CMDList, NiagaraRenderRibbonsGenVerticesInitialSortGPU);
+			SCOPED_DRAW_EVENT(RHICmdList, NiagaraRenderRibbonsGenVerticesInitialSortGPU);
 			SCOPE_CYCLE_COUNTER(STAT_NiagaraRenderRibbonsGenVerticesInitialSortGPU);
 			
 			// Initial sort, sets up the buffer, and runs a parallel bubble sort to create groups of BubbleSortGroupWidth size
-			CMDList.Transition(FRHITransitionInfo(SortParams.DestinationSortedIndices, ERHIAccess::SRVMask | ERHIAccess::VertexOrIndexBuffer, ERHIAccess::UAVCompute));
-			SetComputePipelineState(CMDList, BubbleSortShader.GetComputeShader());
+			RHICmdList.Transition(FRHITransitionInfo(SortParams.DestinationSortedIndices, ERHIAccess::SRVMask | ERHIAccess::VertexOrIndexBuffer, ERHIAccess::UAVCompute));
+			SetComputePipelineState(RHICmdList, BubbleSortShader.GetComputeShader());
 			ValidateShaderParameters(BubbleSortShader, SortParams);
-			SetShaderParameters(CMDList, BubbleSortShader, BubbleSortShader.GetComputeShader(), SortParams);
-			DispatchComputeShader(CMDList, BubbleSortShader, NumInitialThreadGroups, 1, 1);
-			UnsetShaderUAVs(CMDList, BubbleSortShader, BubbleSortShader.GetComputeShader());
-			CMDList.Transition(FRHITransitionInfo(SortParams.DestinationSortedIndices, ERHIAccess::UAVCompute, ERHIAccess::SRVMask | ERHIAccess::VertexOrIndexBuffer));
+			SetShaderParameters(RHICmdList, BubbleSortShader, BubbleSortShader.GetComputeShader(), SortParams);
+			DispatchComputeShader(RHICmdList, BubbleSortShader, NumInitialThreadGroups, 1, 1);
+			UnsetShaderUAVs(RHICmdList, BubbleSortShader, BubbleSortShader.GetComputeShader());
+			RHICmdList.Transition(FRHITransitionInfo(SortParams.DestinationSortedIndices, ERHIAccess::UAVCompute, ERHIAccess::SRVMask | ERHIAccess::VertexOrIndexBuffer));
 		}
 		
 		{
-			SCOPED_DRAW_EVENT(CMDList, NiagaraRenderRibbonsGenVerticesFinalSortGPU);
+			SCOPED_DRAW_EVENT(RHICmdList, NiagaraRenderRibbonsGenVerticesFinalSortGPU);
 			SCOPE_CYCLE_COUNTER(STAT_NiagaraRenderRibbonsGenVerticesFinalSortGPU);
 			
 			// Repeatedly runs a scatter based merge sort until we have the final buffer
@@ -2363,12 +2363,12 @@ void FNiagaraRendererRibbons::InitializeVertexBuffersGPU(FRHICommandListImmediat
 		
 				SwapBuffers();
 			
-				CMDList.Transition(FRHITransitionInfo(SortParams.DestinationSortedIndices, ERHIAccess::SRVMask | ERHIAccess::VertexOrIndexBuffer, ERHIAccess::UAVCompute));
-				SetComputePipelineState(CMDList, MergeSortShader.GetComputeShader());
-				SetShaderParameters(CMDList, MergeSortShader, MergeSortShader.GetComputeShader(), SortParams);
-				DispatchComputeShader(CMDList, MergeSortShader, NumMergeSortThreadGroups, 1, 1);
-				UnsetShaderUAVs(CMDList, MergeSortShader, MergeSortShader.GetComputeShader());
-				CMDList.Transition(FRHITransitionInfo(SortParams.DestinationSortedIndices, ERHIAccess::UAVCompute, ERHIAccess::SRVMask | ERHIAccess::VertexOrIndexBuffer));
+				RHICmdList.Transition(FRHITransitionInfo(SortParams.DestinationSortedIndices, ERHIAccess::SRVMask | ERHIAccess::VertexOrIndexBuffer, ERHIAccess::UAVCompute));
+				SetComputePipelineState(RHICmdList, MergeSortShader.GetComputeShader());
+				SetShaderParameters(RHICmdList, MergeSortShader, MergeSortShader.GetComputeShader(), SortParams);
+				DispatchComputeShader(RHICmdList, MergeSortShader, NumMergeSortThreadGroups, 1, 1);
+				UnsetShaderUAVs(RHICmdList, MergeSortShader, MergeSortShader.GetComputeShader());
+				RHICmdList.Transition(FRHITransitionInfo(SortParams.DestinationSortedIndices, ERHIAccess::UAVCompute, ERHIAccess::SRVMask | ERHIAccess::VertexOrIndexBuffer));
 				
 				SortGroupSize *= 2;
 			}			
@@ -2376,7 +2376,7 @@ void FNiagaraRendererRibbons::InitializeVertexBuffersGPU(FRHICommandListImmediat
 	}
 	
 	{
-		SCOPED_DRAW_EVENT(CMDList, NiagaraRenderRibbonsGenVerticesReductionPhase1GPU);
+		SCOPED_DRAW_EVENT(RHICmdList, NiagaraRenderRibbonsGenVerticesReductionPhase1GPU);
 		SCOPE_CYCLE_COUNTER(STAT_NiagaraRenderRibbonsGenVerticesReductionPhase1GPU);
 		
 		FNiagaraRibbonVertexReductionInitializationCS::FPermutationDomain PermutationVector;
@@ -2438,7 +2438,7 @@ void FNiagaraRendererRibbons::InitializeVertexBuffersGPU(FRHICommandListImmediat
 			SwapBuffers();
 		}		
 		
-		const auto TransitionOutputBuffers = [this, &CMDList, &Params, &TempBuffers](ERHIAccess Previous, ERHIAccess Next)
+		const auto TransitionOutputBuffers = [this, &RHICmdList, &Params, &TempBuffers](ERHIAccess Previous, ERHIAccess Next)
 		{
 			FRHITransitionInfo DataBufferTransitions[] =
 			{
@@ -2447,7 +2447,7 @@ void FNiagaraRendererRibbons::InitializeVertexBuffersGPU(FRHICommandListImmediat
 				FRHITransitionInfo(Params.OutputSegments, Previous, Next),
 				FRHITransitionInfo(Params.OutputTessellationStats, Previous, Next),
 			};
-			CMDList.Transition(MakeArrayView(DataBufferTransitions, UE_ARRAY_COUNT(DataBufferTransitions)));			
+			RHICmdList.Transition(MakeArrayView(DataBufferTransitions, UE_ARRAY_COUNT(DataBufferTransitions)));
 		};
 		
 	
@@ -2455,10 +2455,10 @@ void FNiagaraRendererRibbons::InitializeVertexBuffersGPU(FRHICommandListImmediat
 			
 			const uint32 NumThreadGroupsInitialization = FMath::DivideAndRoundUp<uint32>(NumExecutableInstances, FNiagaraRibbonComputeCommon::VertexGenReductionInitializationThreadSize);	
 			TransitionOutputBuffers(ERHIAccess::SRVMask | ERHIAccess::VertexOrIndexBuffer, ERHIAccess::UAVCompute);
-			SetComputePipelineState(CMDList, ReductionInitializationShader.GetComputeShader());
-			SetShaderParameters(CMDList, ReductionInitializationShader, ReductionInitializationShader.GetComputeShader(), Params);
-			DispatchComputeShader(CMDList, ReductionInitializationShader, NumThreadGroupsInitialization, 1, 1);
-			UnsetShaderUAVs(CMDList, ReductionInitializationShader, ReductionInitializationShader.GetComputeShader());
+			SetComputePipelineState(RHICmdList, ReductionInitializationShader.GetComputeShader());
+			SetShaderParameters(RHICmdList, ReductionInitializationShader, ReductionInitializationShader.GetComputeShader(), Params);
+			DispatchComputeShader(RHICmdList, ReductionInitializationShader, NumThreadGroupsInitialization, 1, 1);
+			UnsetShaderUAVs(RHICmdList, ReductionInitializationShader, ReductionInitializationShader.GetComputeShader());
 			TransitionOutputBuffers(ERHIAccess::UAVCompute, ERHIAccess::SRVMask | ERHIAccess::VertexOrIndexBuffer);
 		}
 		
@@ -2473,10 +2473,10 @@ void FNiagaraRendererRibbons::InitializeVertexBuffersGPU(FRHICommandListImmediat
 				
 				TransitionOutputBuffers(ERHIAccess::SRVMask | ERHIAccess::VertexOrIndexBuffer, ERHIAccess::UAVCompute);
 			
-				SetComputePipelineState(CMDList, ReductionPropgateShader.GetComputeShader());
-				SetShaderParameters(CMDList, ReductionPropgateShader, ReductionPropgateShader.GetComputeShader(), Params);
-				DispatchComputeShader(CMDList, ReductionPropgateShader, NumThreadGroups, 1, 1);
-				UnsetShaderUAVs(CMDList, ReductionPropgateShader, ReductionPropgateShader.GetComputeShader());
+				SetComputePipelineState(RHICmdList, ReductionPropgateShader.GetComputeShader());
+				SetShaderParameters(RHICmdList, ReductionPropgateShader, ReductionPropgateShader.GetComputeShader(), Params);
+				DispatchComputeShader(RHICmdList, ReductionPropgateShader, NumThreadGroups, 1, 1);
+				UnsetShaderUAVs(RHICmdList, ReductionPropgateShader, ReductionPropgateShader.GetComputeShader());
 				
 				TransitionOutputBuffers(ERHIAccess::UAVCompute, ERHIAccess::SRVMask | ERHIAccess::VertexOrIndexBuffer);
 			}
@@ -2488,7 +2488,7 @@ void FNiagaraRendererRibbons::InitializeVertexBuffersGPU(FRHICommandListImmediat
 	static constexpr int32 CommandBufferOffset = 0;
 		
 	{
-		SCOPED_DRAW_EVENT(CMDList, NiagaraRenderRibbonsGenVerticesReductionPhase2GPU);
+		SCOPED_DRAW_EVENT(RHICmdList, NiagaraRenderRibbonsGenVerticesReductionPhase2GPU);
 		SCOPE_CYCLE_COUNTER(STAT_NiagaraRenderRibbonsGenVerticesReductionPhase2GPU);
 		
 		FNiagaraRibbonVertexReductionFinalizationParameters FinalizationParams;
@@ -2520,24 +2520,24 @@ void FNiagaraRendererRibbons::InitializeVertexBuffersGPU(FRHICommandListImmediat
 			const uint32 NumThreadGroups = GenerationConfig.HasRibbonIDs() ?
 				FMath::DivideAndRoundUp<uint32>(NumExecutableInstances, FNiagaraRibbonComputeCommon::VertexGenReductionFinalizationThreadSize) :
 				1;
-			CMDList.Transition({
+			RHICmdList.Transition({
 				FRHITransitionInfo(VertexBuffers.RibbonLookupTableBuffer.Buffer, ERHIAccess::SRVMask | ERHIAccess::VertexOrIndexBuffer, ERHIAccess::UAVCompute),
 				FRHITransitionInfo(VertexBuffers.GPUComputeCommandBuffer.Buffer, ERHIAccess::SRVMask | ERHIAccess::VertexOrIndexBuffer | ERHIAccess::IndirectArgs, ERHIAccess::UAVCompute)});
 			
-			SetComputePipelineState(CMDList, ComputeShader.GetComputeShader());
-			SetShaderParameters(CMDList, ComputeShader, ComputeShader.GetComputeShader(), FinalizationParams);
-			DispatchComputeShader(CMDList, ComputeShader, NumThreadGroups, 1, 1);
-			UnsetShaderUAVs(CMDList, ComputeShader, ComputeShader.GetComputeShader());
+			SetComputePipelineState(RHICmdList, ComputeShader.GetComputeShader());
+			SetShaderParameters(RHICmdList, ComputeShader, ComputeShader.GetComputeShader(), FinalizationParams);
+			DispatchComputeShader(RHICmdList, ComputeShader, NumThreadGroups, 1, 1);
+			UnsetShaderUAVs(RHICmdList, ComputeShader, ComputeShader.GetComputeShader());
 			
 			// We don't need to transition RibbonLookupTableBuffer as it's still needed for the next shader
-			CMDList.Transition({
+			RHICmdList.Transition({
 				FRHITransitionInfo(VertexBuffers.RibbonLookupTableBuffer.Buffer, ERHIAccess::UAVCompute, ERHIAccess::UAVCompute),
 				FRHITransitionInfo(VertexBuffers.GPUComputeCommandBuffer.Buffer, ERHIAccess::UAVCompute, ERHIAccess::SRVMask | ERHIAccess::VertexOrIndexBuffer | ERHIAccess::IndirectArgs)});		
 		}		
 	}
 		
 	{
-		SCOPED_DRAW_EVENT(CMDList, NiagaraRenderRibbonsGenVerticesMultiRibbonInitGPU);
+		SCOPED_DRAW_EVENT(RHICmdList, NiagaraRenderRibbonsGenVerticesMultiRibbonInitGPU);
 		SCOPE_CYCLE_COUNTER(STAT_NiagaraRenderRibbonsGenVerticesMultiRibbonInitGPU);
 		
 		FNiagaraRibbonVertexFinalizationParameters FinalizeParams;
@@ -2577,23 +2577,23 @@ void FNiagaraRendererRibbons::InitializeVertexBuffersGPU(FRHICommandListImmediat
 			TShaderMapRef<FNiagaraRibbonUVParamCalculationCS> ComputeShader(GetGlobalShaderMap(GMaxRHIFeatureLevel), PermutationVector);
 
 			// We don't need to transition RibbonLookupTableBuffer as it's still setup for UAV from the last shader
-			CMDList.Transition({
+			RHICmdList.Transition({
 				FRHITransitionInfo(VertexBuffers.RibbonLookupTableBuffer.Buffer, ERHIAccess::UAVCompute, ERHIAccess::UAVCompute),
 				FRHITransitionInfo(VertexBuffers.TangentsAndDistancesBuffer.Buffer, ERHIAccess::SRVMask | ERHIAccess::VertexOrIndexBuffer, ERHIAccess::UAVCompute)});
 			
-			SetComputePipelineState(CMDList, ComputeShader.GetComputeShader());			
-			SetShaderParameters(CMDList, ComputeShader, ComputeShader.GetComputeShader(), FinalizeParams);
-			DispatchIndirectComputeShader(CMDList, ComputeShader.GetShader(), VertexBuffers.GPUComputeCommandBuffer.Buffer, CommandBufferOffset * FNiagaraRibbonCommandBufferLayout::NumElements);
-			UnsetShaderUAVs(CMDList, ComputeShader, ComputeShader.GetComputeShader());
+			SetComputePipelineState(RHICmdList, ComputeShader.GetComputeShader());
+			SetShaderParameters(RHICmdList, ComputeShader, ComputeShader.GetComputeShader(), FinalizeParams);
+			DispatchIndirectComputeShader(RHICmdList, ComputeShader.GetShader(), VertexBuffers.GPUComputeCommandBuffer.Buffer, CommandBufferOffset * FNiagaraRibbonCommandBufferLayout::NumElements);
+			UnsetShaderUAVs(RHICmdList, ComputeShader, ComputeShader.GetComputeShader());
 			
-			CMDList.Transition({
+			RHICmdList.Transition({
 				FRHITransitionInfo(VertexBuffers.TangentsAndDistancesBuffer.Buffer, ERHIAccess::UAVCompute, ERHIAccess::SRVMask | ERHIAccess::VertexOrIndexBuffer),
 				FRHITransitionInfo(VertexBuffers.RibbonLookupTableBuffer.Buffer, ERHIAccess::UAVCompute, ERHIAccess::SRVMask | ERHIAccess::VertexOrIndexBuffer)});
 		}		
 	}
 }
 
-void FNiagaraRibbonComputeDispatchManager::GenerateAllGPUData(FRHICommandListImmediate& CMDList)
+void FNiagaraRibbonComputeDispatchManager::GenerateAllGPUData(FRHICommandListImmediate& RHICmdList)
 {
 	// Handle all vertex gens first
 	for (int32 Index = 0; Index < RenderersToGenerate.Num(); Index++)
@@ -2611,7 +2611,7 @@ void FNiagaraRibbonComputeDispatchManager::GenerateAllGPUData(FRHICommandListImm
 				Params.Renderer->GenerationConfig.HasTwist());
 		
 		
-			Params.Renderer->InitializeVertexBuffersGPU(CMDList, Params.ComputeDispatchInterface, Params.SourceParticleData, ComputeBuffers, RenderingResources);
+			Params.Renderer->InitializeVertexBuffersGPU(RHICmdList, Params.ComputeDispatchInterface, Params.SourceParticleData, ComputeBuffers, RenderingResources);
 		}
 	}
 		
@@ -2624,7 +2624,7 @@ void FNiagaraRibbonComputeDispatchManager::GenerateAllGPUData(FRHICommandListImm
 			for (int32 Index = 0; Index < RenderingResources->ViewResources.Num(); Index++)
 			{
 				const auto& RenderingResourcesView = RenderingResources->ViewResources[Index];
-				RendererToGen.Renderer->InitializeViewIndexBuffersGPU(CMDList, RendererToGen.ComputeDispatchInterface, RendererToGen.SourceParticleData, RenderingResourcesView);
+				RendererToGen.Renderer->InitializeViewIndexBuffersGPU(RHICmdList, RendererToGen.ComputeDispatchInterface, RendererToGen.SourceParticleData, RenderingResourcesView);
 			}
 		}
 	}
diff --git a/Engine/Plugins/FX/Niagara/Source/Niagara/Public/NiagaraRendererRibbons.h b/Engine/Plugins/FX/Niagara/Source/Niagara/Public/NiagaraRendererRibbons.h
--- a/Engine/Plugins/FX/Niagara/Source/Niagara/Public/NiagaraRendererRibbons.h
+++ b/Engine/Plugins/FX/Niagara/Source/Niagara/Public/NiagaraRendererRibbons.h
@@ -408,13 +408,13 @@ protected:
 	                                               const TSharedPtr<FNiagaraRibbonRenderingFrameResources>& RenderingResources, const TSharedPtr<FNiagaraRibbonRenderingFrameViewResources>& RenderingViewResources, FMeshBatch& OutMeshBatch, bool bShouldUseGPUInitIndices) const;
 
 
-	void InitializeViewIndexBuffersGPU(FRHICommandListImmediate& CMDList, FNiagaraGpuComputeDispatchInterface* ComputeDispatchInterface, const FNiagaraDataBuffer* SourceParticleData,
+	void InitializeViewIndexBuffersGPU(FRHICommandListImmediate& RHICmdList, FNiagaraGpuComputeDispatchInterface* ComputeDispatchInterface, const FNiagaraDataBuffer* SourceParticleData,
 		const TSharedPtr<FNiagaraRibbonRenderingFrameViewResources>& RenderingViewResources) const;
 
 	void InitializeVertexBuffersResources(const FNiagaraDynamicDataRibbon* DynamicDataRibbon, FNiagaraDataBuffer* SourceParticleData,
 	                                      FGlobalDynamicReadBuffer& DynamicReadBuffer, const TSharedPtr<FNiagaraRibbonRenderingFrameResources>& RenderingResources, bool bShouldUseGPUInit) const;
 	
-	void InitializeVertexBuffersGPU(FRHICommandListImmediate& CMDList, FNiagaraGpuComputeDispatchInterface* ComputeDispatchInterface, const FNiagaraDataBuffer* SourceParticleData,
+	void InitializeVertexBuffersGPU(FRHICommandListImmediate& RHICmdList, FNiagaraGpuComputeDispatchInterface* ComputeDispatchInterface, const FNiagaraDataBuffer* SourceParticleData,
 		struct FNiagaraRibbonGPUInitComputeBuffers& TempBuffers, const TSharedPtr<FNiagaraRibbonRenderingFrameResources>& RenderingResources) const;
 
 	FRibbonComputeUniformParameters SetupComputeVertexGenParams(FNiagaraGpuComputeDispatchInterface* ComputeDispatchInterface, const TSharedPtr<FNiagaraRibbonRenderingFrameResources>& RenderingResources, const FNiagaraDataBuffer* SourceParticleData) const;
