Fix cooking warning about non localized pin name in niagara

diff --git a/Engine/Plugins/FX/Niagara/Source/NiagaraEditor/Private/NiagaraNodeUsageSelector.cpp b/Engine/Plugins/FX/Niagara/Source/NiagaraEditor/Private/NiagaraNodeUsageSelector.cpp
--- a/Engine/Plugins/FX/Niagara/Source/NiagaraEditor/Private/NiagaraNodeUsageSelector.cpp
+++ b/Engine/Plugins/FX/Niagara/Source/NiagaraEditor/Private/NiagaraNodeUsageSelector.cpp
@@ -112,7 +112,7 @@ TArray<int32> UNiagaraNodeUsageSelector::GetOptionValues() const
 
 FText UNiagaraNodeUsageSelector::GetOptionPinFriendlyName(const FNiagaraVariable& Variable) const
 {
-	return FText::FromString(Variable.GetName().ToString());
+	return FText::AsCultureInvariant(Variable.GetName().ToString());
 }
 
 FString UNiagaraNodeUsageSelector::GetInputCaseName(int32 Case) const
