Optimizations to rapid parameter generation - should cut cost by about 50%
-make sure that we don't worry about spending cycles finding hidden pins unless we're actually using it
-reduce heap creation of strings when dealing with names
-reduce number of traversals we need to process by finding multiple pins at once
-share module pins across both new and updated RI parameters

diff --git a/Engine/Plugins/FX/Niagara/Source/Niagara/Private/NiagaraCommon.cpp b/Engine/Plugins/FX/Niagara/Source/Niagara/Private/NiagaraCommon.cpp
--- a/Engine/Plugins/FX/Niagara/Source/Niagara/Private/NiagaraCommon.cpp
+++ b/Engine/Plugins/FX/Niagara/Source/Niagara/Private/NiagaraCommon.cpp
@@ -977,36 +977,45 @@ FName NIAGARA_API FNiagaraUtilities::GetUniqueName(FName CandidateName, const TS
 
 FString FNiagaraUtilities::CreateRapidIterationConstantName(FName InVariableName, const TCHAR* InEmitterName, ENiagaraScriptUsage InUsage)
 {
-	TArray<FString> SplitName;
-	InVariableName.ToString().ParseIntoArray(SplitName, TEXT("."));
-	int32 NumSlots = SplitName.Num();
+	FNameBuilder ConstantName;
+
+	ConstantName.Append(PARAM_MAP_RAPID_ITERATION_BASE_STR);
+	ConstantName.AppendChar(TCHAR('.'));
+
 	if (InEmitterName != nullptr)
 	{
-		for (int32 i = 0; i < NumSlots; i++)
-		{
-			if (SplitName[i] == TEXT("Emitter"))
-			{
-				SplitName[i] = InEmitterName;
-			}
-		}
+		FNameBuilder VariableSource(InVariableName);
+		FStringView VariableView(VariableSource);
+
+		constexpr TCHAR EmitterNamespace[] = TEXT("Emitter.");
+		constexpr int32 EmitterNamespaceLength = UE_ARRAY_COUNT(EmitterNamespace) - 1;
+
+		const int32 EmitterLocation = VariableView.Find(EmitterNamespace);
+		const bool HasEmitterNamespace = EmitterLocation != INDEX_NONE
+			&& (EmitterLocation == 0 || VariableView[EmitterLocation - 1] == TCHAR('.'))
+			&& (VariableView.Len() > (EmitterLocation + EmitterNamespaceLength));
+
+		ConstantName.Append(InEmitterName);
+		ConstantName.AppendChar(TCHAR('.'));
 
-		if (NumSlots >= 3 && SplitName[0] == InEmitterName)
+		if (HasEmitterNamespace)
 		{
-			// Do nothing
-			UE_LOG(LogNiagara, Log, TEXT("ConvertVariableToRapidIterationConstantName Got here!"));
+			ConstantName.Append(VariableView.Left(EmitterLocation));
+			ConstantName.Append(InEmitterName);
+			ConstantName.AppendChar(TCHAR('.'));
+			ConstantName.Append(VariableView.RightChop(EmitterLocation + EmitterNamespaceLength));
 		}
 		else
 		{
-			SplitName.Insert(InEmitterName, 0);
+			ConstantName.Append(VariableView);
 		}
-		SplitName.Insert(TEXT("Constants"), 0);
 	}
 	else
 	{
-		SplitName.Insert(TEXT("Constants"), 0);
+		InVariableName.AppendString(ConstantName);
 	}
 
-	return FString::Join(SplitName, TEXT("."));
+	return ConstantName.ToString();
 }
 
 FNiagaraVariable FNiagaraUtilities::ConvertVariableToRapidIterationConstantName(FNiagaraVariable InVar, const TCHAR* InEmitterName, ENiagaraScriptUsage InUsage)
diff --git a/Engine/Plugins/FX/Niagara/Source/NiagaraEditor/Private/NiagaraGraph.cpp b/Engine/Plugins/FX/Niagara/Source/NiagaraEditor/Private/NiagaraGraph.cpp
--- a/Engine/Plugins/FX/Niagara/Source/NiagaraEditor/Private/NiagaraGraph.cpp
+++ b/Engine/Plugins/FX/Niagara/Source/NiagaraEditor/Private/NiagaraGraph.cpp
@@ -1218,14 +1218,26 @@ void UNiagaraGraph::ForceBaseId(ENiagaraScriptUsage InUsage, const FGuid& InUsag
 }
 
 UEdGraphPin* UNiagaraGraph::FindParameterMapDefaultValuePin(const FName VariableName, ENiagaraScriptUsage InUsage, ENiagaraScriptUsage InParentUsage) const
+{
+	TArray<UEdGraphPin*> DefaultInputPin = { nullptr };
+
+	MultiFindParameterMapDefaultValuePins({ VariableName }, InUsage, InParentUsage, { DefaultInputPin });
+
+	return DefaultInputPin[0];
+}
+
+void UNiagaraGraph::MultiFindParameterMapDefaultValuePins(TConstArrayView<FName> VariableNames, ENiagaraScriptUsage InUsage, ENiagaraScriptUsage InParentUsage, TArrayView<UEdGraphPin*> DefaultPins) const
 {
 	TArray<UEdGraphPin*> MatchingDefaultPins;
-	
+
 	TArray<UNiagaraNode*> NodesTraversed;
 	BuildTraversal(NodesTraversed, InUsage, FGuid(), true);
 
-	UEdGraphPin* DefaultInputPin = nullptr;
 	FPinCollectorArray OutputPins;
+
+	const int32 VariableNameCount = VariableNames.Num();
+	int32 DefaultPinsFound = 0;
+
 	for (UNiagaraNode* Node : NodesTraversed)
 	{
 		if (UNiagaraNodeParameterMapGet* GetNode = Cast<UNiagaraNodeParameterMapGet>(Node))
@@ -1234,56 +1246,63 @@ UEdGraphPin* UNiagaraGraph::FindParameterMapDefaultValuePin(const FName Variable
 			GetNode->GetOutputPins(OutputPins);
 			for (UEdGraphPin* OutputPin : OutputPins)
 			{
-				if (VariableName == OutputPin->PinName)
+				int32 NameIndex = VariableNames.IndexOfByKey(OutputPin->PinName);
+				if ((NameIndex != INDEX_NONE) && (DefaultPins[NameIndex] == nullptr))
 				{
 					UEdGraphPin* Pin = GetNode->GetDefaultPin(OutputPin);
 					if (Pin)
 					{
-						DefaultInputPin = Pin;
-						break;
+						++DefaultPinsFound;
+						DefaultPins[NameIndex] = Pin;
+
+						if (DefaultPinsFound == VariableNameCount)
+						{
+							break;
+						}
 					}
 				}
 			}
-			if (DefaultInputPin != nullptr)
-			{
-				break;
-			}
 		}
 	}
 
-
-	// There are some pins 
-	if (DefaultInputPin && DefaultInputPin->LinkedTo.Num() != 0 && DefaultInputPin->LinkedTo[0] != nullptr)
+	auto TraceInputPin = [&](UEdGraphPin* Pin)
 	{
-		UNiagaraNode* Owner = Cast<UNiagaraNode>(DefaultInputPin->LinkedTo[0]->GetOwningNode());
-		UEdGraphPin* PreviousInput = DefaultInputPin;
-		int32 NumIters = 0;
-		while (Owner)
+		if (Pin && Pin->LinkedTo.Num() != 0 && Pin->LinkedTo[0] != nullptr)
 		{
-			// Check to see if there are any reroute or choose by usage nodes involved in this..
-			UEdGraphPin* InputPin = Owner->GetPassThroughPin(PreviousInput->LinkedTo[0], InParentUsage);
-			if (InputPin == nullptr)
+			UNiagaraNode* Owner = Cast<UNiagaraNode>(Pin->LinkedTo[0]->GetOwningNode());
+			UEdGraphPin* PreviousInput = Pin;
+			int32 NumIters = 0;
+			while (Owner)
 			{
-				return PreviousInput;
-			}
-			else if (InputPin->LinkedTo.Num() == 0)
-			{
-				return InputPin;
-			}
+				// Check to see if there are any reroute or choose by usage nodes involved in this..
+				UEdGraphPin* InputPin = Owner->GetPassThroughPin(PreviousInput->LinkedTo[0], InParentUsage);
+				if (InputPin == nullptr)
+				{
+					return PreviousInput;
+				}
+				else if (InputPin->LinkedTo.Num() == 0)
+				{
+					return InputPin;
+				}
 
-			check(InputPin->LinkedTo[0] != nullptr);
-			Owner = Cast<UNiagaraNode>(InputPin->LinkedTo[0]->GetOwningNode());
-			PreviousInput = InputPin;
-			++NumIters;
-			check(NumIters < Nodes.Num()); // If you hit this assert then we have a cycle in our graph somewhere.
+				check(InputPin->LinkedTo[0] != nullptr);
+				Owner = Cast<UNiagaraNode>(InputPin->LinkedTo[0]->GetOwningNode());
+				PreviousInput = InputPin;
+				++NumIters;
+				check(NumIters < Nodes.Num()); // If you hit this assert then we have a cycle in our graph somewhere.
+			}
 		}
-	}
-	else
+
+		return Pin;
+	};
+
+	if (DefaultPinsFound > 0)
 	{
-		return DefaultInputPin;
+		for (UEdGraphPin*& DefaultInputPin : DefaultPins)
+		{
+			DefaultInputPin = TraceInputPin(DefaultInputPin);
+		}
 	}
-
-	return nullptr;
 }
 
 void UNiagaraGraph::FindOutputNodes(TArray<UNiagaraNodeOutput*>& OutputNodes) const
diff --git a/Engine/Plugins/FX/Niagara/Source/NiagaraEditor/Private/NiagaraGraphDataCache.cpp b/Engine/Plugins/FX/Niagara/Source/NiagaraEditor/Private/NiagaraGraphDataCache.cpp
--- a/Engine/Plugins/FX/Niagara/Source/NiagaraEditor/Private/NiagaraGraphDataCache.cpp
+++ b/Engine/Plugins/FX/Niagara/Source/NiagaraEditor/Private/NiagaraGraphDataCache.cpp
@@ -4,32 +4,87 @@
 
 #include "NiagaraNodeFunctionCall.h"
 
-void FNiagaraGraphDataCache::GetStackFunctionInputPins(UNiagaraNodeFunctionCall& FunctionCallNode, TArray<const UEdGraphPin*>& OutInputPins, TSet<const UEdGraphPin*>& OutHiddenPins, FCompileConstantResolver ConstantResolver, FNiagaraStackGraphUtilities::ENiagaraGetStackFunctionInputPinsOptions Options, bool bIgnoreDisabled)
+void FNiagaraGraphDataCache::GetStackFunctionInputPins(UNiagaraNodeFunctionCall& FunctionCallNode, TConstArrayView<FNiagaraVariable> StaticVars, TArray<const UEdGraphPin*>& OutInputPins, FCompileConstantResolver ConstantResolver, FNiagaraStackGraphUtilities::ENiagaraGetStackFunctionInputPinsOptions Options, bool bIgnoreDisabled)
 {
-	FGetStackFunctionInputPinsKey Key(FunctionCallNode, ConstantResolver, Options, bIgnoreDisabled);
-
+	FGetStackFunctionInputPinsKey Key(FunctionCallNode, ConstantResolver, Options, bIgnoreDisabled, false);
 	FCachedStackFunctionInputPins& CachedResult = GetStackFunctionInputPinsCache.FindOrAdd(Key);
-	if (CachedResult.IsValid())
+	if (!CachedResult.IsValid())
+	{
+		FNiagaraStackGraphUtilities::GetStackFunctionInputPinsWithoutCache(
+			FunctionCallNode,
+			StaticVars,
+			CachedResult.InputPins,
+			ConstantResolver,
+			Options,
+			bIgnoreDisabled,
+			false /*bFilterForCompilation*/);
+	}
+
+	OutInputPins = CachedResult.InputPins;
+}
+
+void FNiagaraGraphDataCache::GetStackFunctionInputPins(UNiagaraNodeFunctionCall& FunctionCallNode, TConstArrayView<FNiagaraVariable> StaticVars, TArray<const UEdGraphPin*>& OutInputPins, TSet<const UEdGraphPin*>& OutHiddenPins, FCompileConstantResolver ConstantResolver, FNiagaraStackGraphUtilities::ENiagaraGetStackFunctionInputPinsOptions Options, bool bIgnoreDisabled)
+{
+	FGetStackFunctionInputPinsKey Key(FunctionCallNode, ConstantResolver, Options, bIgnoreDisabled, false);
+	FGetStackFunctionInputPinsKey CompilationKey(FunctionCallNode, ConstantResolver, Options, bIgnoreDisabled, true);
+
+	if (!GetStackFunctionInputPinsCache.Contains(Key))
+	{
+		GetStackFunctionInputPinsCache.Add(Key);
+	}
+
+	if (!GetStackFunctionInputPinsCache.Contains(CompilationKey))
 	{
-		OutInputPins = CachedResult.InputPins;
-		OutHiddenPins = CachedResult.HiddenPins;
-		return;
+		GetStackFunctionInputPinsCache.Add(CompilationKey);
 	}
 
-	OutInputPins.Empty();
-	OutHiddenPins.Empty();
-	FNiagaraStackGraphUtilities::GetStackFunctionInputPinsWithoutCache(FunctionCallNode, OutInputPins, OutHiddenPins, ConstantResolver, Options, bIgnoreDisabled);
-
-	CachedResult.SourceGraphWeak = FunctionCallNode.GetNiagaraGraph();
-	CachedResult.LastSourceGraphChangeId = FunctionCallNode.GetNiagaraGraph()->GetChangeID();
-	CachedResult.FunctionCallNodeWeak = &FunctionCallNode;
-	CachedResult.CalledGraphWeak = FunctionCallNode.GetCalledGraph();
-	CachedResult.LastCalledGraphChangeId = FunctionCallNode.GetCalledGraph() != nullptr ? FunctionCallNode.GetCalledGraph()->GetChangeID() : FGuid();
-	CachedResult.InputPins = OutInputPins;
-	CachedResult.HiddenPins = OutHiddenPins;
+	FCachedStackFunctionInputPins& CachedResult = GetStackFunctionInputPinsCache.FindChecked(Key);
+	FCachedStackFunctionInputPins& CompilationCachedResult = GetStackFunctionInputPinsCache.FindChecked(CompilationKey);
+
+	if (!CachedResult.IsValid())
+	{
+		FNiagaraStackGraphUtilities::GetStackFunctionInputPinsWithoutCache(
+			FunctionCallNode,
+			StaticVars,
+			CachedResult.InputPins,
+			ConstantResolver,
+			Options,
+			bIgnoreDisabled,
+			false /*bFilterForCompilation*/);
+	}
+
+	if (!CompilationCachedResult.IsValid())
+	{
+		FNiagaraStackGraphUtilities::GetStackFunctionInputPinsWithoutCache(
+			FunctionCallNode,
+			StaticVars,
+			CompilationCachedResult.InputPins,
+			ConstantResolver,
+			Options,
+			bIgnoreDisabled,
+			true /*bFilterForCompilation*/);
+	}
+
+	OutInputPins = CachedResult.InputPins;
+
+	// generate hidden pins
+	auto PinsMatch = [](const UEdGraphPin* Lhs, const UEdGraphPin* Rhs) -> bool
+	{
+		return Lhs->GetFName() == Rhs->GetFName()
+			&& Lhs->PinType.PinCategory == Rhs->PinType.PinCategory
+			&& Lhs->PinType.PinSubCategoryObject == Rhs->PinType.PinSubCategoryObject;
+	};
+
+	for (const UEdGraphPin* InputPin : CachedResult.InputPins)
+	{
+		if (!CompilationCachedResult.InputPins.ContainsByPredicate([&](const UEdGraphPin* CompilationPin) { return PinsMatch(InputPin, CompilationPin); }))
+		{
+			OutHiddenPins.Add(InputPin);
+		}
+	}
 }
 
-FNiagaraGraphDataCache::FGetStackFunctionInputPinsKey::FGetStackFunctionInputPinsKey(UNiagaraNodeFunctionCall& InFunctionCallNode, const FCompileConstantResolver& InConstantResolver, FNiagaraStackGraphUtilities::ENiagaraGetStackFunctionInputPinsOptions InOptions, bool bInIgnoreDisabled)
+FNiagaraGraphDataCache::FGetStackFunctionInputPinsKey::FGetStackFunctionInputPinsKey(UNiagaraNodeFunctionCall& InFunctionCallNode, const FCompileConstantResolver& InConstantResolver, FNiagaraStackGraphUtilities::ENiagaraGetStackFunctionInputPinsOptions InOptions, bool bInIgnoreDisabled, bool bInFilterForCompilation)
 {
 	FunctionCallNodeKey = FObjectKey(&InFunctionCallNode);
 
@@ -41,6 +96,7 @@ FNiagaraGraphDataCache::FGetStackFunctionInputPinsKey::FGetStackFunctionInputPin
 
 	Options = InOptions;
 	bIgnoreDisabled = bInIgnoreDisabled;
+	bFilterForCompilation = bInFilterForCompilation;
 
 	Hash = GetTypeHash(FunctionCallNodeKey);
 	Hash = HashCombine(Hash, GetTypeHash(ConstantResolverSystemKey));
@@ -51,6 +107,7 @@ FNiagaraGraphDataCache::FGetStackFunctionInputPinsKey::FGetStackFunctionInputPin
 
 	Hash = HashCombine(Hash, GetTypeHash(Options));
 	Hash = HashCombine(Hash, GetTypeHash(bIgnoreDisabled));
+	Hash = HashCombine(Hash, GetTypeHash(bFilterForCompilation));
 }
 
 bool FNiagaraGraphDataCache::FCachedStackFunctionInputPins::IsValid() const
diff --git a/Engine/Plugins/FX/Niagara/Source/NiagaraEditor/Private/NiagaraNodeFunctionCall.cpp b/Engine/Plugins/FX/Niagara/Source/NiagaraEditor/Private/NiagaraNodeFunctionCall.cpp
--- a/Engine/Plugins/FX/Niagara/Source/NiagaraEditor/Private/NiagaraNodeFunctionCall.cpp
+++ b/Engine/Plugins/FX/Niagara/Source/NiagaraEditor/Private/NiagaraNodeFunctionCall.cpp
@@ -1053,7 +1053,7 @@ void UNiagaraNodeFunctionCall::BuildParameterMapHistory(FNiagaraParameterMapHist
 		// Looks like this function call is not yet hooked up. Skip it to prevent cascading errors in the compilation
 		return;
 	}
-	
+
 	Super::BuildParameterMapHistory(OutHistory, bRecursive, bFilterForCompilation);
 	if (!IsNodeEnabled() && OutHistory.GetIgnoreDisabled())
 	{
@@ -1061,13 +1061,15 @@ void UNiagaraNodeFunctionCall::BuildParameterMapHistory(FNiagaraParameterMapHist
 		return;
 	}
 
-	if (HasValidScriptAndGraph() == false && Signature.IsValid() == false)
+	const bool ValidScriptAndGraph = HasValidScriptAndGraph();
+
+	if (ValidScriptAndGraph == false && Signature.IsValid() == false)
 	{
 		RouteParameterMapAroundMe(OutHistory, bRecursive);
 		return;
 	}
 
-	if (HasValidScriptAndGraph())
+	if (ValidScriptAndGraph)
 	{
 		UNiagaraGraph* FunctionGraph = GetCalledGraph();
 
@@ -1094,7 +1096,7 @@ void UNiagaraNodeFunctionCall::BuildParameterMapHistory(FNiagaraParameterMapHist
 
 		int32 ParamMapIdx = INDEX_NONE;
 		uint32 NodeIdx = INDEX_NONE;
-		
+
 		if (bHasParamMapPin && bRecursive)
 		{
 			ParamMapIdx = OutHistory.TraceParameterMapOutputPin(UNiagaraNode::TraceOutputPin(GetInputPin(0)->LinkedTo[0]));
@@ -1115,39 +1117,50 @@ void UNiagaraNodeFunctionCall::BuildParameterMapHistory(FNiagaraParameterMapHist
 		TArray<TPair<UEdGraphPin*, int32>, TInlineAllocator<16> > MatchedPairs;
 		TArray<TPair<UEdGraphPin*, int32>, TInlineAllocator<16> > MatchedConstants;
 		TArray<bool, TInlineAllocator<16> > OutputMatched;
-		OutputMatched.AddDefaulted(OutputPins.Num());
+		TArray<FNiagaraVariableBase, TInlineAllocator<16>> OutputVariables;
+
+		const int32 OutputPinCount = OutputPins.Num();
+
+		OutputMatched.AddDefaulted(OutputPinCount);
+		OutputVariables.Reserve(OutputPinCount);
+
+		for (int32 OutputIt = 0; OutputIt < OutputPinCount; ++OutputIt)
+		{
+			OutputVariables.Emplace(Schema->PinToNiagaraVariable(OutputPins[OutputIt]));
+		}
 
 		// Find the matches of names and types of the sub-graph output pins and this function call nodes' outputs.
 		for (UEdGraphPin* ChildOutputNodePin : OutputNode->GetAllPins())
 		{
-			FNiagaraVariable VarChild = Schema->PinToNiagaraVariable(ChildOutputNodePin);
+			if (!ChildOutputNodePin)
+			{
+				continue;
+			}
+
+			const FNiagaraVariableBase VarChild = Schema->PinToNiagaraVariable(ChildOutputNodePin);
 
 			if (ChildOutputNodePin->LinkedTo.Num() > 0 && VarChild.GetType() == FNiagaraTypeDefinition::GetParameterMapDef())
 			{
 				for (int32 i = 0; i < OutputPins.Num(); i++)
 				{
-					FNiagaraVariable OutputVar = Schema->PinToNiagaraVariable(OutputPins[i]);
-					if (OutputVar.IsEquivalent(VarChild))
+					if (OutputVariables[i].IsEquivalent(VarChild))
 					{
-						TPair<UEdGraphPin*, int32> Pair;
+						TPair<UEdGraphPin*, int32>& Pair = MatchedPairs.AddZeroed_GetRef();
 						Pair.Key = OutputPins[i];
 						Pair.Value = OutHistory.TraceParameterMapOutputPin(ChildOutputNodePin->LinkedTo[0]);
-						MatchedPairs.Add(Pair);
 						OutputMatched[i] = true;
 					}
 				}
 			}
-			else if (ChildOutputNodePin && VarChild.GetType().IsStatic())
+			else if (VarChild.GetType().IsStatic())
 			{
 				for (int32 i = 0; i < OutputPins.Num(); i++)
 				{
-					FNiagaraVariable OutputVar = Schema->PinToNiagaraVariable(OutputPins[i]);
-					if (OutputVar.IsEquivalent(VarChild) && OutputMatched[i] == false)
+					if (!OutputMatched[i] && OutputVariables[i].IsEquivalent(VarChild))
 					{
-						TPair<UEdGraphPin*, int32> Pair;
+						TPair<UEdGraphPin*, int32>& Pair = MatchedConstants.AddZeroed_GetRef();
 						Pair.Key = OutputPins[i];
 						Pair.Value = OutHistory.GetConstantFromInputPin(ChildOutputNodePin);
-						MatchedConstants.Add(Pair);
 						OutputMatched[i] = true;
 					}
 				}
@@ -1218,9 +1231,8 @@ void UNiagaraNodeFunctionCall::ChangeScriptVersion(FGuid NewScriptVersion, const
 			OverrideNode->Modify();
 			OverrideNode->GetInputPins(OverridePins);
 			
-			TSet<const UEdGraphPin*> HiddenModulePins;
 			TArray<const UEdGraphPin*> ModuleInputPins;
-			GetStackFunctionInputPins(*this, ModuleInputPins, HiddenModulePins, UpgradeContext.ConstantResolver, FNiagaraStackGraphUtilities::ENiagaraGetStackFunctionInputPinsOptions::ModuleInputsOnly);
+			FNiagaraStackGraphUtilities::GetStackFunctionInputPins(*this, ModuleInputPins, UpgradeContext.ConstantResolver, FNiagaraStackGraphUtilities::ENiagaraGetStackFunctionInputPinsOptions::ModuleInputsOnly);
 			for (const UEdGraphPin* InputPin : ModuleInputPins)
 			{
 				FunctionInputNames.Add(FNiagaraParameterHandle(InputPin->PinName).GetName(), InputPin->PinType);
@@ -1272,6 +1284,24 @@ UEdGraphPin* UNiagaraNodeFunctionCall::FindParameterMapDefaultValuePin(const FNa
 	return nullptr;
 }
 
+void UNiagaraNodeFunctionCall::FindParameterMapDefaultValuePins(TConstArrayView<FName> VariableNames, ENiagaraScriptUsage InParentUsage, const FCompileConstantResolver& ConstantResolver, TArrayView<UEdGraphPin*> DefaultPins)
+{
+	if (FunctionScript)
+	{
+		UNiagaraScriptSource* ScriptSource = GetFunctionScriptSource();
+		if (ScriptSource != nullptr && ScriptSource->NodeGraph != nullptr)
+		{
+			// Set the static switch values so we traverse the correct node paths
+			TArray<UEdGraphPin*> InputPins;
+			GetInputPins(InputPins);
+
+			FNiagaraEditorUtilities::SetStaticSwitchConstants(ScriptSource->NodeGraph, InputPins, ConstantResolver.WithDebugState(DebugState));
+
+			ScriptSource->NodeGraph->MultiFindParameterMapDefaultValuePins(VariableNames, FunctionScript->GetUsage(), InParentUsage, DefaultPins);
+		}
+	}
+}
+
 UEdGraphPin* UNiagaraNodeFunctionCall::FindStaticSwitchInputPin(const FName& VariableName) const
 {
 	UNiagaraGraph* CalledGraph = GetCalledGraph();
diff --git a/Engine/Plugins/FX/Niagara/Source/NiagaraEditor/Private/NiagaraParameterMapHistory.cpp b/Engine/Plugins/FX/Niagara/Source/NiagaraEditor/Private/NiagaraParameterMapHistory.cpp
--- a/Engine/Plugins/FX/Niagara/Source/NiagaraEditor/Private/NiagaraParameterMapHistory.cpp
+++ b/Engine/Plugins/FX/Niagara/Source/NiagaraEditor/Private/NiagaraParameterMapHistory.cpp
@@ -1806,14 +1806,14 @@ int32 FNiagaraParameterMapHistoryBuilder::HandleVariableRead(int32 ParamMapIdx,
 	return FoundIdx;
 }
 
-void FNiagaraParameterMapHistoryBuilder::RegisterEncounterableVariables(const TArray<FNiagaraVariable>& Variables)
+void FNiagaraParameterMapHistoryBuilder::RegisterEncounterableVariables(TConstArrayView<FNiagaraVariable> Variables)
 {
-	EncounterableExternalVariables.Append(Variables);
+	EncounterableExternalVariables.Append(Variables.GetData(), Variables.Num());
 }
 
-void FNiagaraParameterMapHistoryBuilder::RegisterExternalStaticVariables(const TArray<FNiagaraVariable>& Variables)
+void FNiagaraParameterMapHistoryBuilder::RegisterExternalStaticVariables(TConstArrayView<FNiagaraVariable> Variables)
 {
-	StaticVariables.Append(Variables);
+	StaticVariables.Append(Variables.GetData(), Variables.Num());
 
 	if (UNiagaraScript::LogCompileStaticVars > 0)
 	{
diff --git a/Engine/Plugins/FX/Niagara/Source/NiagaraEditor/Private/NiagaraScriptMergeManager.cpp b/Engine/Plugins/FX/Niagara/Source/NiagaraEditor/Private/NiagaraScriptMergeManager.cpp
--- a/Engine/Plugins/FX/Niagara/Source/NiagaraEditor/Private/NiagaraScriptMergeManager.cpp
+++ b/Engine/Plugins/FX/Niagara/Source/NiagaraEditor/Private/NiagaraScriptMergeManager.cpp
@@ -255,10 +255,9 @@ FNiagaraStackFunctionMergeAdapter::FNiagaraStackFunctionMergeAdapter(const FVers
 	const UEdGraphSchema_Niagara* NiagaraSchema = GetDefault<UEdGraphSchema_Niagara>();
 	if (InFunctionCallNode.FunctionScript != nullptr)
 	{
-		TSet<const UEdGraphPin*> HiddenPins;
 		FCompileConstantResolver Resolver(InOwningEmitter, ENiagaraScriptUsage::Function);
 		TArray<const UEdGraphPin*> FunctionInputPins;
-		GetStackFunctionInputPins(*FunctionCallNode, FunctionInputPins, HiddenPins, Resolver, FNiagaraStackGraphUtilities::ENiagaraGetStackFunctionInputPinsOptions::ModuleInputsOnly, false);
+		GetStackFunctionInputPins(*FunctionCallNode, FunctionInputPins, Resolver, FNiagaraStackGraphUtilities::ENiagaraGetStackFunctionInputPinsOptions::ModuleInputsOnly, false);
 		
 		for (const UEdGraphPin* FunctionInputPin : FunctionInputPins)
 		{
diff --git a/Engine/Plugins/FX/Niagara/Source/NiagaraEditor/Private/NiagaraScriptSource.cpp b/Engine/Plugins/FX/Niagara/Source/NiagaraEditor/Private/NiagaraScriptSource.cpp
--- a/Engine/Plugins/FX/Niagara/Source/NiagaraEditor/Private/NiagaraScriptSource.cpp
+++ b/Engine/Plugins/FX/Niagara/Source/NiagaraEditor/Private/NiagaraScriptSource.cpp
@@ -296,42 +296,53 @@ bool UNiagaraScriptSource::AddModuleIfMissing(FString ModulePath, ENiagaraScript
 	return false;
 }
 
-void UNiagaraScriptSource::FixupRenamedParameters(UNiagaraNode* Node, FNiagaraParameterStore& RapidIterationParameters, const TArray<FNiagaraVariable>& OldRapidIterationVariables, const FVersionedNiagaraEmitter& VersionedEmitter, ENiagaraScriptUsage ScriptUsage) const
+void UNiagaraScriptSource::FixupRenamedParameters(UNiagaraNodeFunctionCall* FunctionCallNode, TConstArrayView<const UEdGraphPin*> ModuleInputPins, FNiagaraParameterStore& RapidIterationParameters, const TArray<FNiagaraVariable>& OldRapidIterationVariables, const FVersionedNiagaraEmitter& VersionedEmitter, ENiagaraScriptUsage ScriptUsage) const
 {
-	UNiagaraNodeFunctionCall* FunctionCallNode = Cast<UNiagaraNodeFunctionCall>(Node);
-	if (FunctionCallNode == nullptr || FunctionCallNode->HasValidScriptAndGraph() == false)
-	{
-		return;
-	}
-
-	// find out which inputs the module offers
-	TSet<const UEdGraphPin*> HiddenModulePins;
-	TArray<const UEdGraphPin*> ModuleInputPins;
-	FCompileConstantResolver ConstantResolver;
-	GetStackFunctionInputPins(*FunctionCallNode, ModuleInputPins, HiddenModulePins, ConstantResolver, FNiagaraStackGraphUtilities::ENiagaraGetStackFunctionInputPinsOptions::ModuleInputsOnly);
-
 	// the rapid iteration parameters and the function input pins use different variable naming schemes, so most of this is just used to convert one name to the other 
 	UNiagaraGraph* Graph = FunctionCallNode->GetCalledGraph();
 	const UEdGraphSchema_Niagara* NiagaraSchema = Graph->GetNiagaraSchema();
 	const FString UniqueEmitterName = VersionedEmitter.Emitter ? VersionedEmitter.Emitter->GetUniqueEmitterName() : FString();
 
+	TArray<FNiagaraVariable, TInlineAllocator<16>> ModulePinVariables;
+	TArray<TOptional<FNiagaraVariableMetaData>, TInlineAllocator<16>> ModulePinMetaData;
+	TArray<FName, TInlineAllocator<16>> ModulePinAliasedName;
+	TArray<FString, TInlineAllocator<16>> ModulePinConstantName;
+	TArray<int32, TInlineAllocator<16>> ModulePinNameSpaceEnd;
+
+	const int32 ModuleInputPinCount = ModuleInputPins.Num();
+	ModulePinVariables.Reserve(ModuleInputPinCount);
+	ModulePinMetaData.Reserve(ModuleInputPinCount);
+	ModulePinAliasedName.Reserve(ModuleInputPinCount);
+	ModulePinConstantName.Reserve(ModuleInputPinCount);
+	ModulePinNameSpaceEnd.Reserve(ModuleInputPinCount);
+
+	for (const UEdGraphPin* ModuleInputPin : ModuleInputPins)
+	{
+		const FNiagaraVariable& InputVar = ModulePinVariables.Add_GetRef(NiagaraSchema->PinToNiagaraVariable(ModuleInputPin));
+		ModulePinMetaData.Add(Graph->GetMetaData(InputVar));
+
+		FNiagaraParameterHandle AliasedFunctionInputHandle = FNiagaraParameterHandle::CreateAliasedModuleParameterHandle(FNiagaraParameterHandle(ModuleInputPin->PinName), FunctionCallNode);
+		const FName& AliasedName = ModulePinAliasedName.Add_GetRef(AliasedFunctionInputHandle.GetParameterHandleString());
+		const FString& ConstantName = ModulePinConstantName.Add_GetRef(FNiagaraUtilities::CreateRapidIterationConstantName(AliasedName, *UniqueEmitterName, ScriptUsage));
+
+		ConstantName.FindLastChar(TEXT('.'), ModulePinNameSpaceEnd.AddDefaulted_GetRef());
+	}
+
 	// go through the existing rapid iteration params to see if they are either still valid or were renamed
-	for (FNiagaraVariable OldRapidIterationVar : OldRapidIterationVariables)
+	for (const FNiagaraVariable& OldRapidIterationVar : OldRapidIterationVariables)
 	{
 		if (FGuid* BoundGuid = RapidIterationParameters.ParameterGuidMapping.Find(OldRapidIterationVar))
 		{
-			for (const UEdGraphPin* ModulePin : ModuleInputPins)
+			for (int32 ModuleInputPinIt = 0; ModuleInputPinIt < ModuleInputPinCount; ++ModuleInputPinIt)
 			{
-				FNiagaraVariable InputVar = NiagaraSchema->PinToNiagaraVariable(ModulePin);
-				TOptional<FNiagaraVariableMetaData> VariableMetaData = Graph->GetMetaData(InputVar);
+				const FNiagaraVariable& InputVar = ModulePinVariables[ModuleInputPinIt];
+				const FStringView ConstantName(ModulePinConstantName[ModuleInputPinIt]);
+				const int32 NameSpaceEnd = ModulePinNameSpaceEnd[ModuleInputPinIt];
 
-				FNiagaraParameterHandle AliasedFunctionInputHandle = FNiagaraParameterHandle::CreateAliasedModuleParameterHandle(FNiagaraParameterHandle(ModulePin->PinName), FunctionCallNode);
-				const FString ConstantName = FNiagaraUtilities::CreateRapidIterationConstantName(AliasedFunctionInputHandle.GetParameterHandleString(), *UniqueEmitterName, ScriptUsage);
-				bool bNameMatches = OldRapidIterationVar.GetName() == *ConstantName;
+				bool bNameMatches = FStringView(FNameBuilder(OldRapidIterationVar.GetName())) == ConstantName;
 				
 				// move on if the namespaces don't match
-				int32 NameSpaceEnd;
-				if (ConstantName.FindLastChar(TEXT('.'), NameSpaceEnd))
+				if (NameSpaceEnd != INDEX_NONE)
 				{
 					if (!OldRapidIterationVar.IsInNameSpace(ConstantName.Left(NameSpaceEnd)))
 					{
@@ -339,6 +350,8 @@ void UNiagaraScriptSource::FixupRenamedParameters(UNiagaraNode* Node, FNiagaraPa
 					}
 				}
 				
+				const TOptional<FNiagaraVariableMetaData>& VariableMetaData = ModulePinMetaData[ModuleInputPinIt];
+
 				if (VariableMetaData.IsSet() && VariableMetaData->GetVariableGuid() == *BoundGuid)
 				{
 					// if the names match we have nothing to do
@@ -348,8 +361,8 @@ void UNiagaraScriptSource::FixupRenamedParameters(UNiagaraNode* Node, FNiagaraPa
 					}
 
 					// if the guid matches but the names differ then the parameter was renamed, so lets update the rapid iteration parameter
-					FNiagaraTypeDefinition InputType = NiagaraSchema->PinToTypeDefinition(ModulePin);
-					FNiagaraVariable NewRapidIterationVar = FNiagaraStackGraphUtilities::CreateRapidIterationParameter(UniqueEmitterName, ScriptUsage, AliasedFunctionInputHandle.GetParameterHandleString(), InputType);
+					const FName& AliasedName = ModulePinAliasedName[ModuleInputPinIt];
+					FNiagaraVariable NewRapidIterationVar = FNiagaraStackGraphUtilities::CreateRapidIterationParameter(UniqueEmitterName, ScriptUsage, AliasedName, InputVar.GetType());
 					RapidIterationParameters.RenameParameter(OldRapidIterationVar, NewRapidIterationVar.GetName());
 					break;
 				}
@@ -369,6 +382,7 @@ void UNiagaraScriptSource::FixupRenamedParameters(UNiagaraNode* Node, FNiagaraPa
 						FNiagaraVariable RapidIterationAlias(FNiagaraTypeDefinition::GetPositionDef(), OldRapidIterationVar.GetName());
 						RapidIterationParameters.RemoveParameter(RapidIterationAlias);
 						RapidIterationParameters.ChangeParameterType(OldRapidIterationVar, FNiagaraTypeDefinition::GetPositionDef());
+						break;
 					}
 				}
 			}
@@ -376,23 +390,27 @@ void UNiagaraScriptSource::FixupRenamedParameters(UNiagaraNode* Node, FNiagaraPa
 	}
 }
 
-void InitializeNewRapidIterationParametersForNode(const UEdGraphSchema_Niagara* Schema, UNiagaraNode* Node, const FVersionedNiagaraEmitter& Emitter, ENiagaraScriptUsage ScriptUsage, FNiagaraParameterStore& RapidIterationParameters, TSet<FName>& ValidRapidIterationParameterNames)
+void UNiagaraScriptSource::InitializeNewParameters(UNiagaraNodeFunctionCall* FunctionCallNode, TConstArrayView<const UEdGraphPin*> FunctionInputPins, FNiagaraParameterStore& RapidIterationParameters, const FVersionedNiagaraEmitter& VersionedEmitter, ENiagaraScriptUsage ScriptUsage, TSet<FName>& ValidRapidIterationParameterNames) const
 {
-	UNiagaraNodeFunctionCall* FunctionCallNode = Cast<UNiagaraNodeFunctionCall>(Node);
-	if (FunctionCallNode == nullptr)
-	{
-		return;
-	}
-	TSet<const UEdGraphPin*> HiddenPins;
-	FCompileConstantResolver Resolver(Emitter, ScriptUsage);
-	TArray<const UEdGraphPin*> FunctionInputPins;
+	const FString UniqueEmitterName = VersionedEmitter.Emitter ? VersionedEmitter.Emitter->GetUniqueEmitterName() : FString();
+	UNiagaraGraph* Graph = FunctionCallNode->GetCalledGraph();
+	const UEdGraphSchema_Niagara* NiagaraSchema = Graph->GetNiagaraSchema();
+
+	TArray<FName, TInlineAllocator<16>> PinNames;
+	TArray<UEdGraphPin*, TInlineAllocator<16>> DefaultPins;
+	TArray<FNiagaraParameterHandle, TInlineAllocator<16>> AliasHandles;
+	TArray<FNiagaraVariable, TInlineAllocator<16>> Parameters;
+
+	const int32 FunctionInputPinCount = FunctionInputPins.Num();
+	PinNames.Reserve(FunctionInputPinCount);
+	DefaultPins.Reserve(FunctionInputPinCount);
+	AliasHandles.Reserve(FunctionInputPinCount);
+	Parameters.Reserve(FunctionInputPinCount);
 
-	const FString UniqueEmitterName = Emitter.Emitter ? Emitter.Emitter->GetUniqueEmitterName() : FString();
 
-	FNiagaraStackGraphUtilities::GetStackFunctionInputPins(*FunctionCallNode, FunctionInputPins, HiddenPins, Resolver, FNiagaraStackGraphUtilities::ENiagaraGetStackFunctionInputPinsOptions::ModuleInputsOnly, false);
 	for (const UEdGraphPin* FunctionInputPin : FunctionInputPins)
 	{
-		FNiagaraTypeDefinition InputType = Schema->PinToTypeDefinition(FunctionInputPin);
+		FNiagaraTypeDefinition InputType = NiagaraSchema->PinToTypeDefinition(FunctionInputPin);
 		if (InputType.IsValid() == false)
 		{
 			UE_LOG(LogNiagaraEditor, Error, TEXT("Invalid input type found while attempting initialize new rapid iteration parameters. Function Node: %s %s Input Name: %s"),
@@ -410,20 +428,34 @@ void InitializeNewRapidIterationParametersForNode(const UEdGraphSchema_Niagara*
 			// Only set a value for the parameter if it's not already set.
 			if (ParameterIndex == INDEX_NONE)
 			{
-				FCompileConstantResolver ConstantResolver(Emitter, ScriptUsage);
-				UEdGraphPin* DefaultPin = FunctionCallNode->FindParameterMapDefaultValuePin(FunctionInputPin->PinName, ScriptUsage, ConstantResolver);
-				// Only set values for inputs which don't have a default wired in the script graph, since inputs with wired defaults can't currently use rapid iteration parameters.
-				if (DefaultPin != nullptr && DefaultPin->LinkedTo.Num() == 0)
+				PinNames.Add(FunctionInputPin->PinName);
+				AliasHandles.Add(AliasedFunctionInputHandle);
+				Parameters.Add(RapidIterationParameter);
+			}
+		}
+	}
+
+	if (!PinNames.IsEmpty())
+	{
+		const int32 DefaultPinCount = PinNames.Num();
+		DefaultPins.Init(nullptr, DefaultPinCount);
+
+		FCompileConstantResolver ConstantResolver(VersionedEmitter, ScriptUsage);
+		FunctionCallNode->FindParameterMapDefaultValuePins(PinNames, ScriptUsage, ConstantResolver, DefaultPins);
+
+		for (int32 DefaultPinIt = 0; DefaultPinIt < DefaultPinCount; ++DefaultPinIt)
+		{
+			UEdGraphPin* DefaultPin = DefaultPins[DefaultPinIt];
+			if (DefaultPin != nullptr && DefaultPin->LinkedTo.Num() == 0)
+			{
+				// Only set values for inputs without override pins, since and override pin means it's being read from a different value.
+				UEdGraphPin* OverridePin = FNiagaraStackGraphUtilities::GetStackFunctionInputOverridePin(*FunctionCallNode, AliasHandles[DefaultPinIt]);
+				if (OverridePin == nullptr)
 				{
-					// Only set values for inputs without override pins, since and override pin means it's being read from a different value.
-					UEdGraphPin* OverridePin = FNiagaraStackGraphUtilities::GetStackFunctionInputOverridePin(*FunctionCallNode, AliasedFunctionInputHandle);
-					if (OverridePin == nullptr)
-					{
-						FNiagaraVariable DefaultVariable = Schema->PinToNiagaraVariable(DefaultPin, true);
-						check(DefaultVariable.GetData() != nullptr);
-						bool bAddParameterIfMissing = true;
-						RapidIterationParameters.SetParameterData(DefaultVariable.GetData(), RapidIterationParameter, bAddParameterIfMissing);
-					}
+					FNiagaraVariable DefaultVariable = NiagaraSchema->PinToNiagaraVariable(DefaultPin, true);
+					check(DefaultVariable.GetData() != nullptr);
+					bool bAddParameterIfMissing = true;
+					RapidIterationParameters.SetParameterData(DefaultVariable.GetData(), Parameters[DefaultPinIt], bAddParameterIfMissing);
 				}
 			}
 		}
@@ -466,8 +498,19 @@ void UNiagaraScriptSource::CleanUpOldAndInitializeNewRapidIterationParameters(co
 
 			for(UNiagaraNode* Node : Nodes)
 			{
-				FixupRenamedParameters(Node, RapidIterationParameters, OldRapidIterationVariables, Emitter, ScriptUsage);
-				InitializeNewRapidIterationParametersForNode(Schema, Node, Emitter, OutputNode->GetUsage(), RapidIterationParameters, ValidRapidIterationParameterNames);
+				UNiagaraNodeFunctionCall* FunctionCallNode = Cast<UNiagaraNodeFunctionCall>(Node);
+				if (FunctionCallNode == nullptr || FunctionCallNode->HasValidScriptAndGraph() == false)
+				{
+					continue;
+				}
+
+				// find out which inputs the module offers
+				TArray<const UEdGraphPin*> ModuleInputPins;
+				FCompileConstantResolver ConstantResolver(Emitter, ScriptUsage);
+				GetStackFunctionInputPins(*FunctionCallNode, ModuleInputPins, ConstantResolver, FNiagaraStackGraphUtilities::ENiagaraGetStackFunctionInputPinsOptions::ModuleInputsOnly);
+
+				FixupRenamedParameters(FunctionCallNode, ModuleInputPins, RapidIterationParameters, OldRapidIterationVariables, Emitter, ScriptUsage);
+				InitializeNewParameters(FunctionCallNode, ModuleInputPins, RapidIterationParameters, Emitter, OutputNode->GetUsage(), ValidRapidIterationParameterNames);
 			}
 		}
 	}
diff --git a/Engine/Plugins/FX/Niagara/Source/NiagaraEditor/Private/ViewModels/Stack/NiagaraStackGraphUtilities.cpp b/Engine/Plugins/FX/Niagara/Source/NiagaraEditor/Private/ViewModels/Stack/NiagaraStackGraphUtilities.cpp
--- a/Engine/Plugins/FX/Niagara/Source/NiagaraEditor/Private/ViewModels/Stack/NiagaraStackGraphUtilities.cpp
+++ b/Engine/Plugins/FX/Niagara/Source/NiagaraEditor/Private/ViewModels/Stack/NiagaraStackGraphUtilities.cpp
@@ -49,6 +49,50 @@ DECLARE_CYCLE_STAT(TEXT("Niagara - StackGraphUtilities - RelayoutGraph"), STAT_N
 
 #define LOCTEXT_NAMESPACE "NiagaraStackGraphUtilities"
 
+namespace FNiagaraStackGraphUtilitiesImpl
+{
+
+static FVersionedNiagaraEmitter GetOuterEmitter(UNiagaraNode* Node)
+{
+	FVersionedNiagaraEmitter Result;
+	if (Node)
+	{
+		if (UNiagaraGraph* NiagaraGraph = Node->GetNiagaraGraph())
+		{
+			Result = NiagaraGraph->GetOwningEmitter();
+		}
+
+		if (!Result.Emitter)
+		{
+			Result.Emitter = Node->GetTypedOuter<UNiagaraEmitter>();
+		}
+	}
+
+	return Result;
+}
+
+static void GetFunctionStaticVariables(UNiagaraNodeFunctionCall& FunctionCallNode, TArray<FNiagaraVariable>& StaticVars)
+{
+	FVersionedNiagaraEmitter Emitter = GetOuterEmitter(&FunctionCallNode);
+	UNiagaraSystem* System = FunctionCallNode.GetTypedOuter<UNiagaraSystem>();
+
+	TSharedPtr<FNiagaraGraphCachedDataBase, ESPMode::ThreadSafe> CachedTraversalData;
+	if (Emitter.GetEmitterData())
+	{
+		CachedTraversalData = Emitter.Emitter->GetCachedTraversalData(Emitter.Version);
+	}
+	else if (System)
+	{
+		CachedTraversalData = System->GetCachedTraversalData();
+	}
+	if (CachedTraversalData.IsValid())
+	{
+		CachedTraversalData.Get()->GetStaticVariables(StaticVars);
+	}
+}
+
+} // FNiagaraStackGraphUtilitiesImpl::
+
 void FNiagaraStackGraphUtilities::MakeLinkTo(UEdGraphPin* PinA, UEdGraphPin* PinB)
 {
 	PinA->MakeLinkTo(PinB);
@@ -821,10 +865,12 @@ FString FNiagaraStackGraphUtilities::GenerateStackModuleEditorDataKey(UNiagaraNo
 
 void ExtractInputPinsFromHistory(FNiagaraParameterMapHistory& History, UEdGraph* FunctionGraph, FNiagaraStackGraphUtilities::ENiagaraGetStackFunctionInputPinsOptions Options, TArray<const UEdGraphPin*>& OutPins)
 {
+	const int32 ModuleNamespaceLength = FNiagaraConstants::ModuleNamespaceString.Len();
+
 	for (int32 i = 0; i < History.Variables.Num(); i++)
 	{
 		FNiagaraVariable& Variable = History.Variables[i];
-		const TArray< FNiagaraParameterMapHistory::FReadHistory>& ReadHistory = History.PerVariableReadHistory[i];
+		const TArray<FNiagaraParameterMapHistory::FReadHistory>& ReadHistory = History.PerVariableReadHistory[i];
 
 		// A read is only really exposed if it's the first read and it has no corresponding write.
 		if (ReadHistory.Num() > 0 && ReadHistory[0].PreviousWritePin.Pin == nullptr)
@@ -833,105 +879,71 @@ void ExtractInputPinsFromHistory(FNiagaraParameterMapHistory& History, UEdGraph*
 
 			// Make sure that the module input is from the called graph, and not a nested graph.
 			UEdGraph* NodeGraph = InputPin->GetOwningNode()->GetGraph();
-			if (NodeGraph == FunctionGraph &&
-				(Options == FNiagaraStackGraphUtilities::ENiagaraGetStackFunctionInputPinsOptions::AllInputs || FNiagaraParameterHandle(InputPin->PinName).IsModuleHandle()))
+			if (NodeGraph == FunctionGraph)
 			{
-				OutPins.Add(InputPin);
+				if (Options == FNiagaraStackGraphUtilities::ENiagaraGetStackFunctionInputPinsOptions::AllInputs)
+				{
+					OutPins.Add(InputPin);
+				}
+				else
+				{
+					FNameBuilder InputPinName(InputPin->PinName);
+					FStringView InputPinView(InputPinName);
+					if (InputPinView.StartsWith(FNiagaraConstants::ModuleNamespaceString)
+						&& InputPinView.Len() > ModuleNamespaceLength
+						&& InputPinView[ModuleNamespaceLength] == TCHAR('.'))
+					{
+						OutPins.Add(InputPin);
+					}
+				}
 			}
 		}
 	}
 }
 
-FVersionedNiagaraEmitter GetOuterEmitter(UNiagaraNode* Node)
+void FNiagaraStackGraphUtilities::GetStackFunctionInputPins(UNiagaraNodeFunctionCall& FunctionCallNode, TArray<const UEdGraphPin*>& OutInputPins, TSet<const UEdGraphPin*>& OutHiddenPins, FCompileConstantResolver ConstantResolver, ENiagaraGetStackFunctionInputPinsOptions Options, bool bIgnoreDisabled)
 {
-	FVersionedNiagaraEmitter Result;
-	if (Node)
-	{
-		if (UNiagaraGraph* NiagaraGraph = Node->GetNiagaraGraph())
-		{
-			Result = NiagaraGraph->GetOwningEmitter();
-		}
+	TArray<FNiagaraVariable> StaticVars;
+	FNiagaraStackGraphUtilitiesImpl::GetFunctionStaticVariables(FunctionCallNode, StaticVars);
 
-		if (!Result.Emitter)
-		{
-			Result.Emitter = Node->GetTypedOuter<UNiagaraEmitter>();
-		}
-	}
-	
-	return Result;
+	FNiagaraEditorModule::Get().GetGraphDataCache().GetStackFunctionInputPins(FunctionCallNode, StaticVars, OutInputPins, OutHiddenPins, ConstantResolver, Options, bIgnoreDisabled);
 }
 
-void FNiagaraStackGraphUtilities::GetStackFunctionInputPins(UNiagaraNodeFunctionCall& FunctionCallNode, TArray<const UEdGraphPin*>& OutInputPins, TSet<const UEdGraphPin*>& OutHiddenPins, FCompileConstantResolver ConstantResolver, ENiagaraGetStackFunctionInputPinsOptions Options, bool bIgnoreDisabled)
+void FNiagaraStackGraphUtilities::GetStackFunctionInputPins(UNiagaraNodeFunctionCall& FunctionCallNode, TArray<const UEdGraphPin*>& OutInputPins, FCompileConstantResolver ConstantResolver, ENiagaraGetStackFunctionInputPinsOptions Options, bool bIgnoreDisabled)
+{
+	TArray<FNiagaraVariable> StaticVars;
+	FNiagaraStackGraphUtilitiesImpl::GetFunctionStaticVariables(FunctionCallNode, StaticVars);
+
+	FNiagaraEditorModule::Get().GetGraphDataCache().GetStackFunctionInputPins(FunctionCallNode, StaticVars, OutInputPins, ConstantResolver, Options, bIgnoreDisabled);
+}
+
+void FNiagaraStackGraphUtilities::GetStackFunctionInputPins(UNiagaraNodeFunctionCall& FunctionCallNode, TArray<const UEdGraphPin*>& OutInputPins, ENiagaraGetStackFunctionInputPinsOptions Options /*= ENiagaraGetStackFunctionInputPinsOptions::AllInputs*/, bool bIgnoreDisabled /*= false*/)
 {
-	FNiagaraEditorModule::Get().GetGraphDataCache().GetStackFunctionInputPins(FunctionCallNode, OutInputPins, OutHiddenPins, ConstantResolver, Options, bIgnoreDisabled);
+	FCompileConstantResolver EmptyResolver;
+	GetStackFunctionInputPins(FunctionCallNode, OutInputPins, EmptyResolver, Options, bIgnoreDisabled);
 }
 
-void FNiagaraStackGraphUtilities::GetStackFunctionInputPinsWithoutCache(UNiagaraNodeFunctionCall& FunctionCallNode, TArray<const UEdGraphPin*>& OutInputPins, TSet<const UEdGraphPin*>& OutHiddenPins, FCompileConstantResolver ConstantResolver, ENiagaraGetStackFunctionInputPinsOptions Options, bool bIgnoreDisabled)
+void FNiagaraStackGraphUtilities::GetStackFunctionInputPinsWithoutCache(
+	UNiagaraNodeFunctionCall& FunctionCallNode,
+	TConstArrayView<FNiagaraVariable> StaticVars,
+	TArray<const UEdGraphPin*>& OutInputPins,
+	const FCompileConstantResolver& ConstantResolver,
+	ENiagaraGetStackFunctionInputPinsOptions Options,
+	bool bIgnoreDisabled,
+	bool bFilterForCompilation)
 {
 	FNiagaraParameterMapHistoryBuilder Builder;
 	Builder.SetIgnoreDisabled(bIgnoreDisabled);
 	Builder.ConstantResolver = ConstantResolver;
+	Builder.RegisterExternalStaticVariables(StaticVars);
 
-	FVersionedNiagaraEmitter Emitter = GetOuterEmitter(&FunctionCallNode);
-	UNiagaraSystem* System = FunctionCallNode.GetTypedOuter<UNiagaraSystem>();
+	FunctionCallNode.BuildParameterMapHistory(Builder, false, bFilterForCompilation);
 
-	TArray<FNiagaraVariable> StaticVars;
-	TSharedPtr<FNiagaraGraphCachedDataBase, ESPMode::ThreadSafe> CachedTraversalData;
-	if (Emitter.GetEmitterData())
-	{
-		CachedTraversalData = Emitter.Emitter->GetCachedTraversalData(Emitter.Version);
-	}
-	else if (System)
-	{
-		CachedTraversalData = System->GetCachedTraversalData();
-	}
-	if (CachedTraversalData.IsValid())
-	{
-		CachedTraversalData.Get()->GetStaticVariables(StaticVars);
-	}
-
-	Builder.RegisterExternalStaticVariables(StaticVars);
+	OutInputPins.Empty();
 
-	FunctionCallNode.BuildParameterMapHistory(Builder, false, false);
-	
 	if (Builder.Histories.Num() == 1)
 	{
 		ExtractInputPinsFromHistory(Builder.Histories[0], FunctionCallNode.GetCalledGraph(), Options, OutInputPins);
-
-		FNiagaraParameterMapHistoryBuilder BuilderCompiled;
-		BuilderCompiled.ConstantResolver = ConstantResolver;
-		BuilderCompiled.SetIgnoreDisabled(bIgnoreDisabled);
-		BuilderCompiled.RegisterExternalStaticVariables(StaticVars);
-
-		FunctionCallNode.BuildParameterMapHistory(BuilderCompiled, false, true);
-		TArray<const UEdGraphPin*> CompilationPins;
-		if (BuilderCompiled.Histories.Num() == 1)
-		{
-			ExtractInputPinsFromHistory(BuilderCompiled.Histories[0], FunctionCallNode.GetCalledGraph(), Options, CompilationPins);
-		}
-
-		//UE_LOG(LogNiagaraEditor, Log, TEXT("OutInputPins:"));
-		int32 Idx = 0;
-		for (const UEdGraphPin* Pin : OutInputPins)
-		{
-			bool bFoundPin = false;
-			for (const UEdGraphPin* CompiledPin : CompilationPins)
-			{
-				if (Pin->GetName() == CompiledPin->GetName() && 
-					Pin->PinType.PinCategory == CompiledPin->PinType.PinCategory && 
-					Pin->PinType.PinSubCategoryObject == CompiledPin->PinType.PinSubCategoryObject)
-				{
-					bFoundPin = true;
-				}
-			}
-			if (!bFoundPin)
-			{
-				OutHiddenPins.Add(Pin);
-			}
-
-			//UE_LOG(LogNiagaraEditor, Log, TEXT("[%d]:%s %s"), Idx, *Pin->PinName.ToString(), bFoundPin ? TEXT("") : TEXT("HIDDEN"));
-			++Idx;
-		}
 	}
 }
 
@@ -1017,13 +1029,6 @@ TArray<UEdGraphPin*> FNiagaraStackGraphUtilities::GetUnusedFunctionInputPins(UNi
 	return ResultPins;
 }
 
-void FNiagaraStackGraphUtilities::GetStackFunctionInputPins(UNiagaraNodeFunctionCall& FunctionCallNode, TArray<const UEdGraphPin*>& OutInputPins, ENiagaraGetStackFunctionInputPinsOptions Options /*= ENiagaraGetStackFunctionInputPinsOptions::AllInputs*/, bool bIgnoreDisabled /*= false*/)
-{
-	TSet<const UEdGraphPin*> HiddenPins;
-	FCompileConstantResolver EmptyResolver;
-	GetStackFunctionInputPins(FunctionCallNode, OutInputPins, HiddenPins, EmptyResolver, Options, bIgnoreDisabled);
-}
-
 void FNiagaraStackGraphUtilities::GetStackFunctionStaticSwitchPins(UNiagaraNodeFunctionCall& FunctionCallNode, TArray<UEdGraphPin*>& OutInputPins, TSet<UEdGraphPin*>& OutHiddenPins,
 	FCompileConstantResolver& ConstantResolver)
 {
@@ -1038,7 +1043,7 @@ void FNiagaraStackGraphUtilities::GetStackFunctionStaticSwitchPins(UNiagaraNodeF
 	FunctionCallNode.GetInputPins(InputPins);
 	FNiagaraEditorUtilities::SetStaticSwitchConstants(FunctionCallGraph, InputPins, ConstantResolver);
 
-	FVersionedNiagaraEmitter OuterEmitter = GetOuterEmitter(&FunctionCallNode);
+	FVersionedNiagaraEmitter OuterEmitter = FNiagaraStackGraphUtilitiesImpl::GetOuterEmitter(&FunctionCallNode);
 	UNiagaraSystem* System = FunctionCallNode.GetTypedOuter<UNiagaraSystem>();
 
 	TArray<FNiagaraVariable> StaticVars;
@@ -1106,7 +1111,7 @@ void FNiagaraStackGraphUtilities::GetStackFunctionOutputVariables(UNiagaraNodeFu
 	Builder.SetIgnoreDisabled(false);
 	Builder.ConstantResolver = ConstantResolver;
 
-	FVersionedNiagaraEmitter OuterEmitter = GetOuterEmitter(&FunctionCallNode);
+	FVersionedNiagaraEmitter OuterEmitter = FNiagaraStackGraphUtilitiesImpl::GetOuterEmitter(&FunctionCallNode);
 	UNiagaraSystem* System = FunctionCallNode.GetTypedOuter<UNiagaraSystem>();
 
 	TArray<FNiagaraVariable> StaticVars;
@@ -1159,7 +1164,7 @@ bool FNiagaraStackGraphUtilities::GetStackFunctionInputAndOutputVariables(UNiaga
 	Builder.SetIgnoreDisabled(false);
 	Builder.ConstantResolver = ConstantResolver;
 
-	FVersionedNiagaraEmitter OuterEmitter = GetOuterEmitter(&FunctionCallNode);
+	FVersionedNiagaraEmitter OuterEmitter = FNiagaraStackGraphUtilitiesImpl::GetOuterEmitter(&FunctionCallNode);
 	UNiagaraSystem* System = FunctionCallNode.GetTypedOuter<UNiagaraSystem>();
 
 	TArray<FNiagaraVariable> StaticVars;
diff --git a/Engine/Plugins/FX/Niagara/Source/NiagaraEditor/Public/NiagaraGraph.h b/Engine/Plugins/FX/Niagara/Source/NiagaraEditor/Public/NiagaraGraph.h
--- a/Engine/Plugins/FX/Niagara/Source/NiagaraEditor/Public/NiagaraGraph.h
+++ b/Engine/Plugins/FX/Niagara/Source/NiagaraEditor/Public/NiagaraGraph.h
@@ -264,6 +264,7 @@ class UNiagaraGraph : public UEdGraph
 
 	/** Walk through the graph for an ParameterMapGet nodes and see if any of them specify a default for VariableName.*/
 	UEdGraphPin* FindParameterMapDefaultValuePin(const FName VariableName, ENiagaraScriptUsage InUsage, ENiagaraScriptUsage InParentUsage) const;
+	void MultiFindParameterMapDefaultValuePins(TConstArrayView<FName> VariableNames, ENiagaraScriptUsage InUsage, ENiagaraScriptUsage InParentUsage, TArrayView<UEdGraphPin*> DefaultPins) const;
 
 	/** Walk through the graph for an ParameterMapGet nodes and find all matching default pins for VariableName, irrespective of usage. */
 	TArray<UEdGraphPin*> FindParameterMapDefaultValuePins(const FName VariableName) const;
diff --git a/Engine/Plugins/FX/Niagara/Source/NiagaraEditor/Public/NiagaraGraphDataCache.h b/Engine/Plugins/FX/Niagara/Source/NiagaraEditor/Public/NiagaraGraphDataCache.h
--- a/Engine/Plugins/FX/Niagara/Source/NiagaraEditor/Public/NiagaraGraphDataCache.h
+++ b/Engine/Plugins/FX/Niagara/Source/NiagaraEditor/Public/NiagaraGraphDataCache.h
@@ -16,12 +16,13 @@ class UNiagaraNodeFunctionCall;
 class FNiagaraGraphDataCache
 {
 public:
-	void GetStackFunctionInputPins(UNiagaraNodeFunctionCall& FunctionCallNode, TArray<const UEdGraphPin*>& OutInputPins, TSet<const UEdGraphPin*>& OutHiddenPins, FCompileConstantResolver ConstantResolver, FNiagaraStackGraphUtilities::ENiagaraGetStackFunctionInputPinsOptions Options, bool bIgnoreDisabled);
-	
+	void GetStackFunctionInputPins(UNiagaraNodeFunctionCall& FunctionCallNode, TConstArrayView<FNiagaraVariable> StaticVars, TArray<const UEdGraphPin*>& OutInputPins, FCompileConstantResolver ConstantResolver, FNiagaraStackGraphUtilities::ENiagaraGetStackFunctionInputPinsOptions Options, bool bIgnoreDisabled);
+	void GetStackFunctionInputPins(UNiagaraNodeFunctionCall& FunctionCallNode, TConstArrayView<FNiagaraVariable> StaticVars, TArray<const UEdGraphPin*>& OutInputPins, TSet<const UEdGraphPin*>& OutHiddenPins, FCompileConstantResolver ConstantResolver, FNiagaraStackGraphUtilities::ENiagaraGetStackFunctionInputPinsOptions Options, bool bIgnoreDisabled);
+
 private:
 	struct FGetStackFunctionInputPinsKey
 	{
-		FGetStackFunctionInputPinsKey(UNiagaraNodeFunctionCall& InFunctionCallNode, const FCompileConstantResolver& InConstantResolver, FNiagaraStackGraphUtilities::ENiagaraGetStackFunctionInputPinsOptions InOptions, bool bInIgnoreDisabled);
+		FGetStackFunctionInputPinsKey(UNiagaraNodeFunctionCall& InFunctionCallNode, const FCompileConstantResolver& InConstantResolver, FNiagaraStackGraphUtilities::ENiagaraGetStackFunctionInputPinsOptions InOptions, bool bInIgnoreDisabled, bool bFilterForCompilation);
 
 		FORCEINLINE bool operator == (const FGetStackFunctionInputPinsKey& Other) const
 		{
@@ -33,7 +34,8 @@ private:
 				ConstantResolverUsage == Other.ConstantResolverUsage &&
 				ConstantResolverDebugState == Other.ConstantResolverDebugState &&
 				Options == Other.Options &&
-				bIgnoreDisabled == Other.bIgnoreDisabled;
+				bIgnoreDisabled == Other.bIgnoreDisabled &&
+				bFilterForCompilation == Other.bFilterForCompilation;
 		}
 
 		friend uint32 GetTypeHash(const FGetStackFunctionInputPinsKey& Key)
@@ -52,6 +54,7 @@ private:
 
 		FNiagaraStackGraphUtilities::ENiagaraGetStackFunctionInputPinsOptions Options;
 		bool bIgnoreDisabled;
+		bool bFilterForCompilation;
 
 		int32 Hash;
 	};
@@ -64,7 +67,6 @@ private:
 		TWeakObjectPtr<UNiagaraGraph> CalledGraphWeak;
 		FGuid LastCalledGraphChangeId;
 		TArray<const UEdGraphPin*> InputPins;
-		TSet<const UEdGraphPin*> HiddenPins;
 		bool IsValid() const;
 	};
 
diff --git a/Engine/Plugins/FX/Niagara/Source/NiagaraEditor/Public/NiagaraNodeFunctionCall.h b/Engine/Plugins/FX/Niagara/Source/NiagaraEditor/Public/NiagaraNodeFunctionCall.h
--- a/Engine/Plugins/FX/Niagara/Source/NiagaraEditor/Public/NiagaraNodeFunctionCall.h
+++ b/Engine/Plugins/FX/Niagara/Source/NiagaraEditor/Public/NiagaraNodeFunctionCall.h
@@ -143,6 +143,7 @@ public:
 
 	/** Walk through the internal script graph for an ParameterMapGet nodes and see if any of them specify a default for VariableName.*/
 	UEdGraphPin* FindParameterMapDefaultValuePin(const FName VariableName, ENiagaraScriptUsage InParentUsage, FCompileConstantResolver ConstantResolver) const;
+	void FindParameterMapDefaultValuePins(TConstArrayView<FName> VariableNames, ENiagaraScriptUsage InParentUsage, const FCompileConstantResolver& ConstantResolver, TArrayView<UEdGraphPin*> DefaultPins);
 
 	/** Attempts to find the input pin for a static switch with the given name in the internal script graph. Returns nullptr if no such pin can be found. */
 	UEdGraphPin* FindStaticSwitchInputPin(const FName& VariableName) const;
diff --git a/Engine/Plugins/FX/Niagara/Source/NiagaraEditor/Public/NiagaraParameterMapHistory.h b/Engine/Plugins/FX/Niagara/Source/NiagaraEditor/Public/NiagaraParameterMapHistory.h
--- a/Engine/Plugins/FX/Niagara/Source/NiagaraEditor/Public/NiagaraParameterMapHistory.h
+++ b/Engine/Plugins/FX/Niagara/Source/NiagaraEditor/Public/NiagaraParameterMapHistory.h
@@ -502,10 +502,10 @@ public:
 	bool IsInEncounteredEmitterNamespace(FNiagaraVariable& InVar) const;
 
 	/** Register any user or other external variables that could possibly be encountered but may not be declared explicitly. */
-	void RegisterEncounterableVariables(const TArray<FNiagaraVariable>& Variables);
+	void RegisterEncounterableVariables(TConstArrayView<FNiagaraVariable> Variables);
 	const TArray<FNiagaraVariable>& GetEncounterableVariables() const {	return EncounterableExternalVariables;}
 
-	void RegisterExternalStaticVariables(const TArray<FNiagaraVariable>& Variables);
+	void RegisterExternalStaticVariables(TConstArrayView<FNiagaraVariable> Variables);
 
 	FCompileConstantResolver ConstantResolver;
 
diff --git a/Engine/Plugins/FX/Niagara/Source/NiagaraEditor/Public/NiagaraScriptSource.h b/Engine/Plugins/FX/Niagara/Source/NiagaraEditor/Public/NiagaraScriptSource.h
--- a/Engine/Plugins/FX/Niagara/Source/NiagaraEditor/Public/NiagaraScriptSource.h
+++ b/Engine/Plugins/FX/Niagara/Source/NiagaraEditor/Public/NiagaraScriptSource.h
@@ -50,7 +50,9 @@ class UNiagaraScriptSource : public UNiagaraScriptSourceBase
 
 	NIAGARAEDITOR_API virtual bool AddModuleIfMissing(FString ModulePath, ENiagaraScriptUsage Usage, bool& bOutFoundModule)override;
 
-	virtual void FixupRenamedParameters(UNiagaraNode* Node, FNiagaraParameterStore& RapidIterationParameters, const TArray<FNiagaraVariable>& OldRapidIterationVariables, const FVersionedNiagaraEmitter& Emitter, ENiagaraScriptUsage ScriptUsage) const;
+	virtual void FixupRenamedParameters(UNiagaraNodeFunctionCall* FunctionCallNode, TConstArrayView<const UEdGraphPin*> FunctionPins, FNiagaraParameterStore& RapidIterationParameters, const TArray<FNiagaraVariable>& OldRapidIterationVariables, const FVersionedNiagaraEmitter& Emitter, ENiagaraScriptUsage ScriptUsage) const;
+	virtual void InitializeNewParameters(UNiagaraNodeFunctionCall* FunctionCallNode, TConstArrayView<const UEdGraphPin*> FunctionPins, FNiagaraParameterStore& RapidIterationParameters, const FVersionedNiagaraEmitter& VersionedEmitter, ENiagaraScriptUsage ScriptUsage, TSet<FName>& ValidRapidIterationParameterNames) const;
+
 	virtual void CleanUpOldAndInitializeNewRapidIterationParameters(const FVersionedNiagaraEmitter& Emitter, ENiagaraScriptUsage ScriptUsage, FGuid ScriptUsageId, FNiagaraParameterStore& RapidIterationParameters) const override;
 	virtual void ForceGraphToRecompileOnNextCheck() override;
 	virtual void RefreshFromExternalChanges() override;
diff --git a/Engine/Plugins/FX/Niagara/Source/NiagaraEditor/Public/ViewModels/Stack/NiagaraStackGraphUtilities.h b/Engine/Plugins/FX/Niagara/Source/NiagaraEditor/Public/ViewModels/Stack/NiagaraStackGraphUtilities.h
--- a/Engine/Plugins/FX/Niagara/Source/NiagaraEditor/Public/ViewModels/Stack/NiagaraStackGraphUtilities.h
+++ b/Engine/Plugins/FX/Niagara/Source/NiagaraEditor/Public/ViewModels/Stack/NiagaraStackGraphUtilities.h
@@ -97,9 +97,17 @@ namespace FNiagaraStackGraphUtilities
 
 	/* Returns the input pins for the given function call node. */
 	void GetStackFunctionInputPins(UNiagaraNodeFunctionCall& FunctionCallNode, TArray<const UEdGraphPin*>& OutInputPins, TSet<const UEdGraphPin*>& OutHiddenPins, FCompileConstantResolver ConstantResolver, ENiagaraGetStackFunctionInputPinsOptions Options = ENiagaraGetStackFunctionInputPinsOptions::AllInputs, bool bIgnoreDisabled = false);
+	void GetStackFunctionInputPins(UNiagaraNodeFunctionCall& FunctionCallNode, TArray<const UEdGraphPin*>& OutInputPins, FCompileConstantResolver ConstantResolver, ENiagaraGetStackFunctionInputPinsOptions Options = ENiagaraGetStackFunctionInputPinsOptions::AllInputs, bool bIgnoreDisabled = false);
 
 	/* Returns the input pins for the given function call node.  Bypasses the module level caching of this data and generates it each call. */
-	void GetStackFunctionInputPinsWithoutCache(UNiagaraNodeFunctionCall& FunctionCallNode, TArray<const UEdGraphPin*>& OutInputPins, TSet<const UEdGraphPin*>& OutHiddenPins, FCompileConstantResolver ConstantResolver, ENiagaraGetStackFunctionInputPinsOptions Options, bool bIgnoreDisabled);
+	void GetStackFunctionInputPinsWithoutCache(
+		UNiagaraNodeFunctionCall& FunctionCallNode,
+		TConstArrayView<FNiagaraVariable> StaticVars,
+		TArray<const UEdGraphPin*>& OutInputPins,
+		const FCompileConstantResolver& ConstantResolver,
+		ENiagaraGetStackFunctionInputPinsOptions Options,
+		bool bIgnoreDisabled,
+		bool bFilterForCompilation);
 
 	/* Module script calls do not have direct inputs, but rely on the parameter map being initialized correctly. This utility function resolves which of the module's parameters are reachable during compilation and returns a list of pins on the parameter map node that do not have to be compiled. */
 	TArray<UEdGraphPin*> GetUnusedFunctionInputPins(UNiagaraNodeFunctionCall& FunctionCallNode, FCompileConstantResolver ConstantResolver);
