Fixed time dilation code leaving the solo component tick function enabled even for deactivated systems.

diff --git a/Engine/Plugins/FX/Niagara/Source/Niagara/Private/NiagaraComponent.cpp b/Engine/Plugins/FX/Niagara/Source/Niagara/Private/NiagaraComponent.cpp
--- a/Engine/Plugins/FX/Niagara/Source/Niagara/Private/NiagaraComponent.cpp
+++ b/Engine/Plugins/FX/Niagara/Source/Niagara/Private/NiagaraComponent.cpp
@@ -3654,7 +3654,11 @@ void UNiagaraComponent::UpdateInstanceSoloMode()
 	{
 		SystemInstanceController->SetForceSolo(bRequireSoloMode);
 	}
-	SetComponentTickEnabled(bRequireSoloMode);
+
+	if (IsActive())
+	{
+		SetComponentTickEnabled(bRequireSoloMode);
+	}
 }
 
 void UNiagaraComponent::SetGpuComputeDebug(bool bEnableDebug)
