Fixed order of user parameter removal in the user param store to first remove from redirects, then the actual store removal
This makes it so that accessing user params from the store using the "OnChanged" delegate will actually not have the parameter anymore

#rnx

diff --git a/Engine/Plugins/FX/Niagara/Source/Niagara/Private/NiagaraUserRedirectionParameterStore.cpp b/Engine/Plugins/FX/Niagara/Source/Niagara/Private/NiagaraUserRedirectionParameterStore.cpp
--- a/Engine/Plugins/FX/Niagara/Source/Niagara/Private/NiagaraUserRedirectionParameterStore.cpp
+++ b/Engine/Plugins/FX/Niagara/Source/Niagara/Private/NiagaraUserRedirectionParameterStore.cpp
@@ -118,11 +118,9 @@ bool FNiagaraUserRedirectionParameterStore::RemoveParameter(const FNiagaraVariab
 {
 	const FNiagaraVariable* Redirection = UserParameterRedirects.Find(InVar);
 	const FNiagaraVariable& ToRemove = Redirection ? *Redirection : InVar;
+	
+	UserParameterRedirects.Remove(GetUserRedirection(ToRemove));
 	bool Result = Super::RemoveParameter(ToRemove);
-	//if (Result)
-	{
-		UserParameterRedirects.Remove(GetUserRedirection(ToRemove));
-	}
 	return Result;
 }
 
