- Fix some incorrect ShouldCompilePermutation

#rnx

diff --git a/Engine/Plugins/FX/Niagara/Source/NiagaraShader/Private/NiagaraClearCounts.cpp b/Engine/Plugins/FX/Niagara/Source/NiagaraShader/Private/NiagaraClearCounts.cpp
--- a/Engine/Plugins/FX/Niagara/Source/NiagaraShader/Private/NiagaraClearCounts.cpp
+++ b/Engine/Plugins/FX/Niagara/Source/NiagaraShader/Private/NiagaraClearCounts.cpp
@@ -75,7 +75,7 @@ class FNiagaraClearCountsIntCS_Legacy : public FGlobalShader
 
 	static bool ShouldCompilePermutation(const FGlobalShaderPermutationParameters& Parameters)
 	{
-		return !IsMobilePlatform(Parameters.Platform);
+		return RHISupportsComputeShaders(Parameters.Platform);
 	}
 
 	static void ModifyCompilationEnvironment(const FGlobalShaderPermutationParameters& Parameters, FShaderCompilerEnvironment& OutEnvironment)
diff --git a/Engine/Plugins/FX/Niagara/Source/NiagaraShader/Private/NiagaraDebugShaders.cpp b/Engine/Plugins/FX/Niagara/Source/NiagaraShader/Private/NiagaraDebugShaders.cpp
--- a/Engine/Plugins/FX/Niagara/Source/NiagaraShader/Private/NiagaraDebugShaders.cpp
+++ b/Engine/Plugins/FX/Niagara/Source/NiagaraShader/Private/NiagaraDebugShaders.cpp
@@ -86,8 +86,6 @@ class NIAGARASHADER_API FNiagaraClearUAVCS : public FGlobalShader
 
 	static constexpr uint32 ThreadGroupSize = 32;
 
-	//class FUAVType : SHADER_PERMUTATION_INT("TEXTURE_TYPE", 3);
-
 	static bool ShouldCompilePermutation(const FGlobalShaderPermutationParameters& Parameters)
 	{
 		return RHISupportsComputeShaders(Parameters.Platform);
@@ -100,8 +98,6 @@ class NIAGARASHADER_API FNiagaraClearUAVCS : public FGlobalShader
 		OutEnvironment.SetDefine(TEXT("THREADGROUP_SIZE"), ThreadGroupSize);
 	}
 
-	//NIAGARA_DEBUGDRAW_CLEARUAV_UINT_CS
-
 	BEGIN_SHADER_PARAMETER_STRUCT(FParameters, )
 		SHADER_PARAMETER(FUintVector4,					ClearValue)
 		SHADER_PARAMETER(uint32,						ClearSize)
