- Convert Async GPU Trace DI to ShaderParams / RDG

diff --git a/Engine/Plugins/FX/Niagara/Shaders/Private/NiagaraDataInterfaceAsyncGpuTraceTemplate.ush b/Engine/Plugins/FX/Niagara/Shaders/Private/NiagaraDataInterfaceAsyncGpuTraceTemplate.ush
--- a/Engine/Plugins/FX/Niagara/Shaders/Private/NiagaraDataInterfaceAsyncGpuTraceTemplate.ush
+++ b/Engine/Plugins/FX/Niagara/Shaders/Private/NiagaraDataInterfaceAsyncGpuTraceTemplate.ush
@@ -1,31 +1,31 @@
 // Copyright Epic Games, Inc. All Rights Reserved.
 
-float3										SystemLWCTile_{ParameterName};
+float3											{ParameterName}_SystemLWCTile;
 /** Niagara ray tracing uses shared buffers between dispatches so each buffer has a buffer + and offset into that buffer for the current dispatch. */
-uint										MaxRayTraceCount_{ParameterName};
-RWStructuredBuffer<FNiagaraAsyncGpuTrace>	RWRayRequests_{ParameterName};
-uint										RayRequestsOffset_{ ParameterName };
-StructuredBuffer<FNiagaraAsyncGpuTraceResult>	IntersectionResults_{ParameterName};
-uint										IntersectionResultsOffset_{ ParameterName };
-RWBuffer<uint>								RWRayTraceCounts_{ParameterName};
-uint										RayTraceCountsOffset_{ ParameterName };
+uint											{ParameterName}_MaxRayTraceCount;
+RWStructuredBuffer<FNiagaraAsyncGpuTrace>		{ParameterName}_RWRayRequests;
+uint											{ParameterName}_RayRequestsOffset;
+StructuredBuffer<FNiagaraAsyncGpuTraceResult>	{ParameterName}_IntersectionResults;
+uint											{ParameterName}_IntersectionResultsOffset;
+RWBuffer<uint>									{ParameterName}_RWRayTraceCounts;
+uint											{ParameterName}_RayTraceCountsOffset;
 
 void IssueAsyncRayTraceGpu_{ParameterName}_UEImpureCall(int In_QueryID, float3 In_TraceStart, float3 In_TraceEnd, int In_CollisionGroup, out bool Out_IsQueryValid)
 {
-	NDIAsyncGpuTrace_IssueAsyncRayTrace(MaxRayTraceCount_{ParameterName}, RWRayRequests_{ParameterName}, RayRequestsOffset_{ ParameterName }, RWRayTraceCounts_{ParameterName}, RayTraceCountsOffset_{ ParameterName }, SystemLWCTile_{ParameterName}, In_QueryID, In_TraceStart, In_TraceEnd, In_CollisionGroup, Out_IsQueryValid);
+	NDIAsyncGpuTrace_IssueAsyncRayTrace({ParameterName}_MaxRayTraceCount, {ParameterName}_RWRayRequests, {ParameterName}_RayRequestsOffset, {ParameterName}_RWRayTraceCounts, {ParameterName}_RayTraceCountsOffset, {ParameterName}_SystemLWCTile, In_QueryID, In_TraceStart, In_TraceEnd, In_CollisionGroup, Out_IsQueryValid);
 }
 
 void ReserveRayTraceIndex_{ParameterName}_UEImpureCall(int In_TraceCount, out int Out_StartQueryID, out bool Out_IndicesValid)
 {
-	NDIAsyncGpuTrace_ReserveRayTraceIndex(MaxRayTraceCount_{ParameterName}, RWRayTraceCounts_{ParameterName}, RayTraceCountsOffset_{ ParameterName }, In_TraceCount, Out_StartQueryID, Out_IndicesValid);
+	NDIAsyncGpuTrace_ReserveRayTraceIndex({ParameterName}_MaxRayTraceCount, {ParameterName}_RWRayTraceCounts, {ParameterName}_RayTraceCountsOffset, In_TraceCount, Out_StartQueryID, Out_IndicesValid);
 }
 
 void CreateAsyncRayTraceGpu_{ParameterName}_UEImpureCall(float3 In_TraceStart, float3 In_TraceEnd, int In_CollisionGroup, out int Out_QueryID, out bool Out_IsQueryValid)
 {
-	NDIAsyncGpuTrace_CreateAsyncRayTrace(MaxRayTraceCount_{ParameterName}, RWRayRequests_{ParameterName}, RayRequestsOffset_{ ParameterName }, RWRayTraceCounts_{ParameterName}, RayTraceCountsOffset_{ ParameterName }, SystemLWCTile_{ParameterName}, In_TraceStart, In_TraceEnd, In_CollisionGroup, Out_QueryID, Out_IsQueryValid);
+	NDIAsyncGpuTrace_CreateAsyncRayTrace({ParameterName}_MaxRayTraceCount, {ParameterName}_RWRayRequests, {ParameterName}_RayRequestsOffset, {ParameterName}_RWRayTraceCounts, {ParameterName}_RayTraceCountsOffset, {ParameterName}_SystemLWCTile, In_TraceStart, In_TraceEnd, In_CollisionGroup, Out_QueryID, Out_IsQueryValid);
 }
 
 void ReadAsyncRayTraceGpu_{ParameterName}(int In_PreviousFrameQueryID, out bool Out_CollisionValid, out float Out_CollisionDistance, out float3 Out_CollisionPosWorld, out float3 Out_CollisionNormal)
 {
-	NDIAsyncGpuTrace_ReadAsyncRayTrace(MaxRayTraceCount_{ParameterName}, IntersectionResults_{ParameterName}, IntersectionResultsOffset_{ParameterName}, SystemLWCTile_{ParameterName}, In_PreviousFrameQueryID, Out_CollisionValid, Out_CollisionDistance, Out_CollisionPosWorld, Out_CollisionNormal);
+	NDIAsyncGpuTrace_ReadAsyncRayTrace({ParameterName}_MaxRayTraceCount, {ParameterName}_IntersectionResults, {ParameterName}_IntersectionResultsOffset, {ParameterName}_SystemLWCTile, In_PreviousFrameQueryID, Out_CollisionValid, Out_CollisionDistance, Out_CollisionPosWorld, Out_CollisionNormal);
 }
\ No newline at end of file
diff --git a/Engine/Plugins/FX/Niagara/Source/Niagara/Classes/NiagaraDataInterfaceAsyncGpuTrace.h b/Engine/Plugins/FX/Niagara/Source/Niagara/Classes/NiagaraDataInterfaceAsyncGpuTrace.h
--- a/Engine/Plugins/FX/Niagara/Source/Niagara/Classes/NiagaraDataInterfaceAsyncGpuTrace.h
+++ b/Engine/Plugins/FX/Niagara/Source/Niagara/Classes/NiagaraDataInterfaceAsyncGpuTrace.h
@@ -15,10 +15,8 @@ UCLASS(EditInlineNew, Category = "Collision", meta = (DisplayName = "Async Gpu T
 class NIAGARA_API UNiagaraDataInterfaceAsyncGpuTrace : public UNiagaraDataInterface
 {
 	GENERATED_UCLASS_BODY()
-public:
-
-	DECLARE_NIAGARA_DI_PARAMETER();
 
+public:
 	/** The maximum number of traces (per particle) that can be created per frame.  Defines the size of the preallocated 
 		buffer that is used to contain the traces. */
 	UPROPERTY(EditAnywhere, Category = "Async GPU Trace")
@@ -58,6 +56,9 @@ public:
 	virtual bool GetFunctionHLSL(const FNiagaraDataInterfaceGPUParamInfo& ParamInfo, const FNiagaraDataInterfaceGeneratedFunction& FunctionInfo, int FunctionInstanceIndex, FString& OutHLSL) override;
 	virtual bool UpgradeFunctionCall(FNiagaraFunctionSignature& FunctionSignature) override;
 #endif
+	virtual bool UseLegacyShaderBindings() const  override { return false; }
+	virtual void BuildShaderParameters(FNiagaraShaderParametersBuilder& ShaderParametersBuilder) const override;
+	virtual void SetShaderParameters(const FNiagaraDataInterfaceSetShaderParametersContext& Context) const override;
 
 private:
 	virtual bool Equals(const UNiagaraDataInterface* Other) const override;
diff --git a/Engine/Plugins/FX/Niagara/Source/Niagara/Private/DataInterface/NiagaraDataInterfaceAsyncGpuTrace.cpp b/Engine/Plugins/FX/Niagara/Source/Niagara/Private/DataInterface/NiagaraDataInterfaceAsyncGpuTrace.cpp
--- a/Engine/Plugins/FX/Niagara/Source/Niagara/Private/DataInterface/NiagaraDataInterfaceAsyncGpuTrace.cpp
+++ b/Engine/Plugins/FX/Niagara/Source/Niagara/Private/DataInterface/NiagaraDataInterfaceAsyncGpuTrace.cpp
@@ -9,6 +9,7 @@
 #include "NiagaraGpuComputeDispatchInterface.h"
 #include "NiagaraGpuComputeDispatch.h"
 #include "NiagaraSimStageData.h"
+#include "NiagaraShaderParametersBuilder.h"
 #include "NiagaraStats.h"
 #include "NiagaraTypes.h"
 #include "NiagaraWorldManager.h"
@@ -25,20 +26,22 @@ namespace NDIAsyncGpuTraceLocal
 	static const TCHAR* CommonShaderFile = TEXT("/Plugin/FX/Niagara/Private/NiagaraDataInterfaceAsyncGpuTrace.ush");
 	static const TCHAR* TemplateShaderFile = TEXT("/Plugin/FX/Niagara/Private/NiagaraDataInterfaceAsyncGpuTraceTemplate.ush");
 
+	BEGIN_SHADER_PARAMETER_STRUCT(FShaderParameters,)
+		SHADER_PARAMETER(FVector3f,											SystemLWCTile)
+		SHADER_PARAMETER(uint32,											MaxRayTraceCount)
+		SHADER_PARAMETER_UAV(RWStructuredBuffer<FNiagaraAsyncGpuTrace>,		RWRayRequests)
+		SHADER_PARAMETER(uint32,											RayRequestsOffset)
+		SHADER_PARAMETER_SRV(StructuredBuffer<FNiagaraAsyncGpuTraceResult>,	IntersectionResults)
+		SHADER_PARAMETER(uint32,											IntersectionResultsOffset)
+		SHADER_PARAMETER_UAV(RWBuffer<uint>,								RWRayTraceCounts)
+		SHADER_PARAMETER(uint32,											RayTraceCountsOffset)
+	END_SHADER_PARAMETER_STRUCT()
+
 	static const FName IssueAsyncRayTraceName(TEXT("IssueAsyncRayTraceGpu"));
 	static const FName CreateAsyncRayTraceName(TEXT("CreateAsyncRayTraceGpu"));
 	static const FName ReserveAsyncRayTraceName(TEXT("ReserveAsyncRayTraceGpu"));
 	static const FName ReadAsyncRayTraceName(TEXT("ReadAsyncRayTraceGpu"));
 
-	static const FString MaxRayTraceCountParamName(TEXT("MaxRayTraceCount_"));
-	static const FString RayRequestsParamName(TEXT("RayRequests_"));
-	static const FString RayRequestsOffsetParamName(TEXT("RayRequestsOffset_"));
-	static const FString IntersectionResultsParamName(TEXT("IntersectionResults_"));
-	static const FString IntersectionResultsOffsetParamName(TEXT("IntersectionResultsOffset_"));
-	static const FString RayTraceCountsParamName(TEXT("RayTraceCounts_"));
-	static const FString RayTraceCountsOffsetParamName(TEXT("RayTraceCountsOffset_"));
-	static const FString SystemLWCTileName(TEXT("SystemLWCTile_"));
-
 	struct FPerInstanceData
 	{
 		FNiagaraSystemInstanceID InstanceID;
@@ -79,22 +82,22 @@ struct FNiagaraDataIntefaceProxyAsyncGpuTrace : public FNiagaraDataInterfaceProx
 		return 0;
 	}
 
-	virtual void PreStage(FRHICommandList& RHICmdList, const FNiagaraDataInterfaceStageArgs& Context) override
+	virtual void PreStage(const FNDIGpuComputePreStageContext& Context) override
 	{
-		FNiagaraDataInterfaceProxy::PreStage(RHICmdList, Context);
+		FNiagaraDataInterfaceProxy::PreStage(Context);
 
-		if (const FInstanceData* InstanceData_RT = SystemInstancesToProxyData_RT.Find(Context.SystemInstanceID))
+		if (const FInstanceData* InstanceData_RT = SystemInstancesToProxyData_RT.Find(Context.GetSystemInstanceID()))
 		{
 			if (InstanceData_RT->MaxTracesPerParticle > 0)
 			{
 				//Accumulate the total ray requests for this DI for all dispatches in the stage.
-				const int32 RayRequests = InstanceData_RT->MaxTracesPerParticle * Context.SimStageData->DestinationNumInstances;
-				Context.ComputeDispatchInterface->GetAsyncGpuTraceHelper().AddToDispatch(
+				const int32 RayRequests = InstanceData_RT->MaxTracesPerParticle * Context.GetSimStageData().DestinationNumInstances;
+				Context.GetComputeDispatchInterface().GetAsyncGpuTraceHelper().AddToDispatch(
 					this,
 					RayRequests,
 					InstanceData_RT->MaxRetraces,
-					InstanceData_RT->ProviderType);
-
+					InstanceData_RT->ProviderType
+				);
 			}
 		}
 	}
@@ -104,16 +107,16 @@ struct FNiagaraDataIntefaceProxyAsyncGpuTrace : public FNiagaraDataInterfaceProx
 		return true;
 	}
 
-	virtual void FinalizePreStage(FRHICommandList& RHICmdList, const FNiagaraGpuComputeDispatchInterface* ComputeDispatchInterface) override
+	virtual void FinalizePreStage(FRDGBuilder& GraphBuilder, const FNiagaraGpuComputeDispatchInterface& ComputeDispatchInterface) override
 	{
-		FNiagaraAsyncGpuTraceHelper& TraceHelper = ComputeDispatchInterface->GetAsyncGpuTraceHelper();
+		FNiagaraAsyncGpuTraceHelper& TraceHelper = ComputeDispatchInterface.GetAsyncGpuTraceHelper();
 		if (SystemInstancesToProxyData_RT.Num())
 		{
-			TraceHelper.BuildDispatch(RHICmdList, this);
+			TraceHelper.BuildDispatch(GraphBuilder.RHICmdList, this);
 		}
 		else
 		{
-			TraceHelper.BuildDummyDispatch(RHICmdList);
+			TraceHelper.BuildDummyDispatch(GraphBuilder.RHICmdList);
 		}
 	}
 };
@@ -397,12 +400,76 @@ bool UNiagaraDataInterfaceAsyncGpuTrace::AppendCompileHash(FNiagaraCompileHashVi
 
 	InVisitor->UpdateString(TEXT("NDIAsyncGpuTraceCommonHLSLSource"), GetShaderFileHash(NDIAsyncGpuTraceLocal::CommonShaderFile, EShaderPlatform::SP_PCD3D_SM5).ToString());
 	InVisitor->UpdateString(TEXT("NDIAsyncGpuTraceTemplateHLSLSource"), GetShaderFileHash(NDIAsyncGpuTraceLocal::TemplateShaderFile, EShaderPlatform::SP_PCD3D_SM5).ToString());
+	InVisitor->UpdateShaderParameters<NDIAsyncGpuTraceLocal::FShaderParameters>();
 
 	return true;
 }
 
 #endif
 
+void UNiagaraDataInterfaceAsyncGpuTrace::BuildShaderParameters(FNiagaraShaderParametersBuilder& ShaderParametersBuilder) const
+{
+	ShaderParametersBuilder.AddNestedStruct<NDIAsyncGpuTraceLocal::FShaderParameters>();
+}
+
+void UNiagaraDataInterfaceAsyncGpuTrace::SetShaderParameters(const FNiagaraDataInterfaceSetShaderParametersContext& Context) const
+{
+	NDIAsyncGpuTraceLocal::FShaderParameters* ShaderParameters = Context.GetParameterNestedStruct<NDIAsyncGpuTraceLocal::FShaderParameters>();
+	ShaderParameters->SystemLWCTile = Context.GetSystemLWCTile();
+
+	const bool bRWRayRequestsBound = Context.IsResourceBound(&ShaderParameters->RWRayRequests);
+	const bool bIntersectionResultsBound = Context.IsResourceBound(&ShaderParameters->IntersectionResults);
+	const bool bRWRayTraceCountsBound = Context.IsResourceBound(&ShaderParameters->RWRayTraceCounts);
+	const bool HasRayTracingParametersBound = bRWRayRequestsBound || bIntersectionResultsBound || bRWRayTraceCountsBound;
+
+	if (HasRayTracingParametersBound)
+	{
+		FNiagaraDataIntefaceProxyAsyncGpuTrace& DIProxy = Context.GetProxy<FNiagaraDataIntefaceProxyAsyncGpuTrace>();
+		const FNiagaraDataIntefaceProxyAsyncGpuTrace::FInstanceData* InstanceData = DIProxy.SystemInstancesToProxyData_RT.Find(Context.GetSystemInstanceID());
+
+		FNiagaraAsyncGpuTraceHelper& TraceHelper = Context.GetComputeDispatchInterface().GetAsyncGpuTraceHelper();
+		const FNiagaraAsyncGpuTraceDispatchInfo* DispatchInfo = nullptr;
+		if (InstanceData && InstanceData->MaxTracesPerParticle > 0)
+		{
+			DispatchInfo = &TraceHelper.GetDispatch(&DIProxy);
+		}
+		else
+		{
+			DispatchInfo = &TraceHelper.GetDummyDispatch();
+		}
+
+		ShaderParameters->MaxRayTraceCount = DispatchInfo->MaxTraces;
+
+		if (bRWRayRequestsBound)
+		{
+			check(DispatchInfo->TraceRequests.IsValid());
+
+			ShaderParameters->RWRayRequests = DispatchInfo->TraceRequests.Buffer->UAV;
+			ShaderParameters->RayRequestsOffset = DispatchInfo->TraceRequests.Offset;
+		}
+
+		if (bIntersectionResultsBound)
+		{
+			check(DispatchInfo->LastFrameTraceResults.IsValid());
+
+			ShaderParameters->IntersectionResults = DispatchInfo->LastFrameTraceResults.Buffer->SRV;
+			ShaderParameters->IntersectionResultsOffset = DispatchInfo->LastFrameTraceResults.Offset;
+		}
+
+		if (bRWRayTraceCountsBound)
+		{
+			check(DispatchInfo->TraceCounts.IsValid());
+
+			ShaderParameters->RWRayTraceCounts = DispatchInfo->TraceCounts.Buffer->UAV;
+			ShaderParameters->RayTraceCountsOffset = DispatchInfo->TraceCounts.Offset;
+		}
+	}
+	else
+	{
+		ShaderParameters->MaxRayTraceCount = 0;
+	}
+}
+
 bool UNiagaraDataInterfaceAsyncGpuTrace::Equals(const UNiagaraDataInterface* Other) const
 {
 	if (!Super::Equals(Other))
@@ -479,111 +546,4 @@ void UNiagaraDataInterfaceAsyncGpuTrace::PostEditChangeProperty(struct FProperty
 }
 #endif
 
-//////////////////////////////////////////////////////////////////////////
-
-struct FNiagaraDataInterfaceParametersCS_AsyncGpuTrace : public FNiagaraDataInterfaceParametersCS
-{
-	DECLARE_TYPE_LAYOUT(FNiagaraDataInterfaceParametersCS_AsyncGpuTrace, NonVirtual);
-public:
-	void Bind(const FNiagaraDataInterfaceGPUParamInfo& ParameterInfo, const class FShaderParameterMap& ParameterMap)
-	{
-		SystemLWCTileParam.Bind(ParameterMap, *(NDIAsyncGpuTraceLocal::SystemLWCTileName + ParameterInfo.DataInterfaceHLSLSymbol));
-		MaxRayTraceCountParam.Bind(ParameterMap, *(NDIAsyncGpuTraceLocal::MaxRayTraceCountParamName + ParameterInfo.DataInterfaceHLSLSymbol));
-		RayRequestsParam.Bind(ParameterMap, *(NDIAsyncGpuTraceLocal::RayRequestsParamName + ParameterInfo.DataInterfaceHLSLSymbol));
-		RayRequestOffsetParam.Bind(ParameterMap, *(NDIAsyncGpuTraceLocal::RayRequestsOffsetParamName + ParameterInfo.DataInterfaceHLSLSymbol));
-		IntersectionResultsParam.Bind(ParameterMap, *(NDIAsyncGpuTraceLocal::IntersectionResultsParamName + ParameterInfo.DataInterfaceHLSLSymbol));
-		IntersectionResultOffsetParam.Bind(ParameterMap, *(NDIAsyncGpuTraceLocal::IntersectionResultsOffsetParamName + ParameterInfo.DataInterfaceHLSLSymbol));
-		RayTraceCountsParam.Bind(ParameterMap, *(NDIAsyncGpuTraceLocal::RayTraceCountsParamName + ParameterInfo.DataInterfaceHLSLSymbol));
-		RayTraceCountsOffsetParam.Bind(ParameterMap, *(NDIAsyncGpuTraceLocal::RayTraceCountsOffsetParamName + ParameterInfo.DataInterfaceHLSLSymbol));
-	}
-
-	void Set(FRHICommandList& RHICmdList, const FNiagaraDataInterfaceSetArgs& Context) const
-	{
-		check(IsInRenderingThread());
-
-		FNiagaraDataIntefaceProxyAsyncGpuTrace* QueryDI = (FNiagaraDataIntefaceProxyAsyncGpuTrace*)Context.DataInterface;
-
-		FRHIComputeShader* ComputeShaderRHI = Context.Shader.GetComputeShader();
-		SetShaderValue(RHICmdList, ComputeShaderRHI, SystemLWCTileParam, Context.SystemLWCTile);
-		
-		const bool HasRayTracingParametersBound = RayRequestsParam.IsUAVBound()
-			|| IntersectionResultsParam.IsBound()
-			|| RayTraceCountsParam.IsBound();
-
-		if (HasRayTracingParametersBound)
-		{
-			const FNiagaraDataIntefaceProxyAsyncGpuTrace::FInstanceData* InstanceData = QueryDI->SystemInstancesToProxyData_RT.Find(Context.SystemInstanceID);
-
-			FNiagaraAsyncGpuTraceHelper& TraceHelper = Context.ComputeDispatchInterface->GetAsyncGpuTraceHelper();
-			const FNiagaraAsyncGpuTraceDispatchInfo* DispatchInfo = nullptr;
-			if (InstanceData && InstanceData->MaxTracesPerParticle > 0)
-			{
-				DispatchInfo = &TraceHelper.GetDispatch(QueryDI);
-			}
-			else
-			{
-				DispatchInfo = &TraceHelper.GetDummyDispatch();
-			}
-
-			SetShaderValue(RHICmdList, ComputeShaderRHI, MaxRayTraceCountParam, DispatchInfo->MaxTraces);
-
-			if (RayRequestsParam.IsUAVBound())
-			{
-				check(DispatchInfo->TraceRequests.IsValid());
-				RHICmdList.SetUAVParameter(ComputeShaderRHI, RayRequestsParam.GetUAVIndex(), DispatchInfo->TraceRequests.Buffer->UAV);
-				SetShaderValue(RHICmdList, ComputeShaderRHI, RayRequestOffsetParam, DispatchInfo->TraceRequests.Offset);
-			}
-
-			if (IntersectionResultsParam.IsBound())
-			{
-				check(DispatchInfo->LastFrameTraceResults.IsValid());
-
-				SetSRVParameter(RHICmdList, ComputeShaderRHI, IntersectionResultsParam, DispatchInfo->LastFrameTraceResults.Buffer->SRV);
-				SetShaderValue(RHICmdList, ComputeShaderRHI, IntersectionResultOffsetParam, DispatchInfo->LastFrameTraceResults.Offset);
-			}
-
-			if (RayTraceCountsParam.IsUAVBound())
-			{
-				check(DispatchInfo->TraceCounts.IsValid());
-
-				RHICmdList.SetUAVParameter(ComputeShaderRHI, RayTraceCountsParam.GetUAVIndex(), DispatchInfo->TraceCounts.Buffer->UAV);
-				SetShaderValue(RHICmdList, ComputeShaderRHI, RayTraceCountsOffsetParam, DispatchInfo->TraceCounts.Offset);
-			}
-		}
-		else
-		{
-			SetShaderValue(RHICmdList, ComputeShaderRHI, MaxRayTraceCountParam, 0);
-		}
-	}
-
-	void Unset(FRHICommandList& RHICmdList, const FNiagaraDataInterfaceSetArgs& Context) const
-	{
-		FRHIComputeShader* ComputeShaderRHI = Context.Shader.GetComputeShader();
-		if (RayRequestsParam.IsUAVBound())
-		{
-			RayRequestsParam.UnsetUAV(RHICmdList, ComputeShaderRHI);
-		}
-
-		if (RayTraceCountsParam.IsUAVBound())
-		{
-			RayTraceCountsParam.UnsetUAV(RHICmdList, ComputeShaderRHI);
-		}
-	}
-
-private:
-	LAYOUT_FIELD(FShaderParameter, SystemLWCTileParam);
-
-	LAYOUT_FIELD(FShaderParameter, MaxRayTraceCountParam);
-	LAYOUT_FIELD(FRWShaderParameter, RayRequestsParam);
-	LAYOUT_FIELD(FShaderParameter, RayRequestOffsetParam);
-	LAYOUT_FIELD(FShaderResourceParameter, IntersectionResultsParam);
-	LAYOUT_FIELD(FShaderParameter, IntersectionResultOffsetParam);
-	LAYOUT_FIELD(FRWShaderParameter, RayTraceCountsParam);
-	LAYOUT_FIELD(FShaderParameter, RayTraceCountsOffsetParam);
-};
-
-IMPLEMENT_TYPE_LAYOUT(FNiagaraDataInterfaceParametersCS_AsyncGpuTrace);
-
-IMPLEMENT_NIAGARA_DI_PARAMETER(UNiagaraDataInterfaceAsyncGpuTrace, FNiagaraDataInterfaceParametersCS_AsyncGpuTrace);
-
 #undef LOCTEXT_NAMESPACE
diff --git a/Engine/Plugins/FX/Niagara/Source/Niagara/Private/NiagaraGpuComputeDispatch.cpp b/Engine/Plugins/FX/Niagara/Source/Niagara/Private/NiagaraGpuComputeDispatch.cpp
--- a/Engine/Plugins/FX/Niagara/Source/Niagara/Private/NiagaraGpuComputeDispatch.cpp
+++ b/Engine/Plugins/FX/Niagara/Source/Niagara/Private/NiagaraGpuComputeDispatch.cpp
@@ -2207,6 +2207,10 @@ void FNiagaraGpuComputeDispatch::PreStageInterface(FRDGBuilder& GraphBuilder, co
 			FNiagaraDataInterfaceProxy* DataInterfaceProxy = InstanceData.DataInterfaceProxies[iDataInterface];
 			Context.SetDataInterfaceProxy(DataInterfaceProxy);
 			DataInterfaceProxy->PreStage(Context);
+			if (DataInterfaceProxy->RequiresPreStageFinalize())
+			{
+				ProxiesToFinalize.Add(DataInterfaceProxy);
+			}
 		}
 	}
 }
@@ -2231,6 +2235,10 @@ void FNiagaraGpuComputeDispatch::PostStageInterface(FRDGBuilder& GraphBuilder, c
 			FNiagaraDataInterfaceProxy* DataInterfaceProxy = InstanceData.DataInterfaceProxies[iDataInterface];
 			Context.SetDataInterfaceProxy(DataInterfaceProxy);
 			DataInterfaceProxy->PostStage(Context);
+			if (DataInterfaceProxy->RequiresPostStageFinalize())
+			{
+				ProxiesToFinalize.Add(DataInterfaceProxy);
+			}
 		}
 	}
 }
