Low hanging fruit for optimizing UNiagaraGraph::PostLoad()

diff --git a/Engine/Plugins/FX/Niagara/Source/NiagaraEditor/Private/NiagaraGraph.cpp b/Engine/Plugins/FX/Niagara/Source/NiagaraEditor/Private/NiagaraGraph.cpp
--- a/Engine/Plugins/FX/Niagara/Source/NiagaraEditor/Private/NiagaraGraph.cpp
+++ b/Engine/Plugins/FX/Niagara/Source/NiagaraEditor/Private/NiagaraGraph.cpp
@@ -144,28 +144,33 @@ void UNiagaraGraph::PostLoad()
 
 	const int32 NiagaraVer = GetLinkerCustomVersion(FNiagaraCustomVersion::GUID);
 
-	for (auto It = VariableToScriptVariable.CreateIterator(); It; ++It)
+	if (!VariableToScriptVariable.IsEmpty())
 	{
-		FNiagaraVariable Var = It.Key();
-		TObjectPtr<UNiagaraScriptVariable>& ScriptVar = It.Value();
+		const TArray<FNiagaraVariable> StaticSwitchInputs = FindStaticSwitchInputs();
 
-		if (ScriptVar == nullptr)
+		for (auto It = VariableToScriptVariable.CreateIterator(); It; ++It)
 		{
-			ScriptVar = NewObject<UNiagaraScriptVariable>(const_cast<UNiagaraGraph*>(this));
-			ScriptVar->Init(Var, FNiagaraVariableMetaData());
-			UE_LOG(LogNiagaraEditor, Display, TEXT("Fixed null UNiagaraScriptVariable | variable %s | asset path %s"), *Var.GetName().ToString(), *GetPathName());
-		}
-		else
-		{
-			if(ScriptVar->GetName().Contains(TEXT(".")) || ScriptVar->GetName().Contains(TEXT(" ")))
+			FNiagaraVariable Var = It.Key();
+			TObjectPtr<UNiagaraScriptVariable>& ScriptVar = It.Value();
+
+			if (ScriptVar == nullptr)
 			{
-				ScriptVar->Rename(*MakeUniqueObjectName(this, UNiagaraScriptVariable::StaticClass(), NAME_None).ToString(), nullptr, REN_NonTransactional | REN_ForceNoResetLoaders);
+				ScriptVar = NewObject<UNiagaraScriptVariable>(const_cast<UNiagaraGraph*>(this));
+				ScriptVar->Init(Var, FNiagaraVariableMetaData());
+				UE_LOG(LogNiagaraEditor, Display, TEXT("Fixed null UNiagaraScriptVariable | variable %s | asset path %s"), *Var.GetName().ToString(), *GetPathName());
 			}
-			
-			// Conditional postload all ScriptVars to ensure static switch default values are allocated as these are required when postloading all graph nodes later.
-			ScriptVar->ConditionalPostLoad();
+			else
+			{
+				if (ScriptVar->GetName().Contains(TEXT(".")) || ScriptVar->GetName().Contains(TEXT(" ")))
+				{
+					ScriptVar->Rename(*MakeUniqueObjectName(this, UNiagaraScriptVariable::StaticClass(), NAME_None).ToString(), nullptr, REN_NonTransactional | REN_ForceNoResetLoaders);
+				}
+
+				// Conditional postload all ScriptVars to ensure static switch default values are allocated as these are required when postloading all graph nodes later.
+				ScriptVar->ConditionalPostLoad();
+			}
+			ScriptVar->SetIsStaticSwitch(StaticSwitchInputs.Contains(Var));
 		}
-		ScriptVar->SetIsStaticSwitch(FindStaticSwitchInputs().Contains(Var));
 	}
 
 	for (UEdGraphNode* Node : Nodes)
@@ -225,20 +230,15 @@ void UNiagaraGraph::PostLoad()
 	}
 
 	// Assume that all externally referenced assets have changed, so update to match. They will return true if they have changed.
-	TArray<UNiagaraNode*> NiagaraNodes;
-	GetNodesOfClass<UNiagaraNode>(NiagaraNodes);
-	for (UNiagaraNode* NiagaraNode : NiagaraNodes)
+	for (UEdGraphNode* Node : Nodes)
 	{
-		UObject* ReferencedAsset = NiagaraNode->GetReferencedAsset();
-		if (ReferencedAsset != nullptr)
-		{
-			ReferencedAsset->ConditionalPostLoad();
-			NiagaraNode->ConditionalPostLoad();
-			NiagaraNode->RefreshFromExternalChanges();
-		}
-		else
+		if (UNiagaraNode* NiagaraNode = Cast<UNiagaraNode>(Node))
 		{
-			NiagaraNode->ConditionalPostLoad();
+			if (UObject* ReferencedAsset = NiagaraNode->GetReferencedAsset())
+			{
+				ReferencedAsset->ConditionalPostLoad();
+				NiagaraNode->RefreshFromExternalChanges();
+			}
 		}
 	}
 
@@ -770,62 +770,122 @@ TArray<UEdGraphPin*> UNiagaraGraph::FindParameterMapDefaultValuePins(const FName
 void UNiagaraGraph::ValidateDefaultPins()
 {
 	FNiagaraEditorModule& EditorModule = FNiagaraEditorModule::Get();
-	for (auto& MetaData : GetAllMetaData()) {
-		FNiagaraVariable Variable = MetaData.Key;
-		UNiagaraScriptVariable* ScriptVariable = MetaData.Value;
-		if (!ScriptVariable || ScriptVariable->DefaultMode == ENiagaraDefaultMode::Custom || ScriptVariable->DefaultMode == ENiagaraDefaultMode::FailIfPreviouslyNotSet) {
-			// If the user selected custom mode or if previously unset they can basically do whatever they want
-			continue;
-		}
-		if (ScriptVariable->GetIsStaticSwitch()) {
-			// We ignore static switch variables as they handle default values differently
-			continue;
-		}
 
-		// Determine if the values set in the graph (which are used for the compilation) are consistent and if necessary update the metadata value
-		TArray<UEdGraphPin*> Pins = FindParameterMapDefaultValuePins(Variable.GetName());
-		bool IsCustom = false;
-		bool IsConsistent = true;
-		bool DefaultInitialized = false;
-		FNiagaraVariable DefaultData;
-		auto TypeUtilityValue = EditorModule.GetTypeUtilities(Variable.GetType());
-		for (UEdGraphPin* Pin : Pins) {
-			if (Pin->LinkedTo.Num() > 0) {
-				IsCustom = true;
-			}
-			else if (!Pin->bHidden && TypeUtilityValue) {
-				FNiagaraVariable ComparisonData = Variable;
-				TypeUtilityValue->SetValueFromPinDefaultString(Pin->DefaultValue, ComparisonData);
-				if (!DefaultInitialized) {
-					DefaultData = ComparisonData;
-					DefaultInitialized = true;
-				}
-				else if (!ComparisonData.HoldsSameData(DefaultData)) {
-					IsConsistent = false;
+	if (!VariableToScriptVariable.IsEmpty())
+	{
+		TArray<UNiagaraNodeParameterMapGet*> TraversedMapGetNodes;
+
+		{
+			TArray<UNiagaraNodeOutput*> OutputNodes;
+			GetNodesOfClass(OutputNodes);
+
+			TArray<UNiagaraNode*> NodesTraversed;
+			for (UNiagaraNodeOutput* OutputNode : OutputNodes)
+			{
+				NodesTraversed.Reset();
+				BuildTraversal(NodesTraversed, OutputNode);
+
+				for (UNiagaraNode* NiagaraNode : NodesTraversed)
+				{
+					if (UNiagaraNodeParameterMapGet* MapGetNode = Cast<UNiagaraNodeParameterMapGet>(NiagaraNode))
+					{
+						TraversedMapGetNodes.AddUnique(MapGetNode);
+					}
 				}
 			}
 		}
 
-		if (!IsCustom && !IsConsistent) {
-			UE_LOG(LogNiagaraEditor, Log, TEXT("Niagara graph %s: The default value declaration for the variable '%s' is not consistent between the graph and the metadata.\n  Either change the default mode to 'custom' or check the input pins in the parameter map get node in the graph."),
-				*GetFullName(), *Variable.GetName().ToString());
-			continue;
-		}
-		if (IsCustom) {
-			ScriptVariable->DefaultMode = ENiagaraDefaultMode::Custom;
-		}
-		else {
-			for (UEdGraphPin* Pin : Pins) {
-				Pin->bNotConnectable = true;
-				Pin->bDefaultValueIsReadOnly = true;
-				if (ScriptVariable->DefaultMode == ENiagaraDefaultMode::Binding || ScriptVariable->DefaultMode == ENiagaraDefaultMode::FailIfPreviouslyNotSet) {
-					Pin->bHidden = true;
+		for (auto& MetaData : GetAllMetaData())
+		{
+			const FNiagaraVariable& Variable = MetaData.Key;
+			UNiagaraScriptVariable* ScriptVariable = MetaData.Value;
+			if (!ScriptVariable || ScriptVariable->DefaultMode == ENiagaraDefaultMode::Custom || ScriptVariable->DefaultMode == ENiagaraDefaultMode::FailIfPreviouslyNotSet)
+			{
+				// If the user selected custom mode or if previously unset they can basically do whatever they want
+				continue;
+			}
+			if (ScriptVariable->GetIsStaticSwitch())
+			{
+				// We ignore static switch variables as they handle default values differently
+				continue;
+			}
+
+			// Determine if the values set in the graph (which are used for the compilation) are consistent and if necessary update the metadata value
+			FPinCollectorArray DefaultPins;
+
+			{
+				for (UNiagaraNodeParameterMapGet* MapGetNode : TraversedMapGetNodes)
+				{
+					FPinCollectorArray OutputPins;
+					MapGetNode->GetOutputPins(OutputPins);
+					for (UEdGraphPin* OutputPin : OutputPins)
+					{
+						if (Variable.GetName() != OutputPin->PinName || OutputPin->PinType.PinSubCategory != UNiagaraNodeParameterMapBase::ParameterPinSubCategory)
+						{
+							continue;
+						}
+						if (UEdGraphPin* Pin = MapGetNode->GetDefaultPin(OutputPin))
+						{
+							check(Pin->Direction == EEdGraphPinDirection::EGPD_Input);
+							DefaultPins.AddUnique(Pin);
+						}
+					}
+				}
+			}
+			bool IsCustom = false;
+			bool IsConsistent = true;
+			bool DefaultInitialized = false;
+			FNiagaraVariable DefaultData;
+			auto TypeUtilityValue = EditorModule.GetTypeUtilities(Variable.GetType());
+			for (UEdGraphPin* Pin : DefaultPins)
+			{
+				if (Pin->LinkedTo.Num() > 0)
+				{
+					IsCustom = true;
+				}
+				else if (!Pin->bHidden && TypeUtilityValue)
+				{
+					FNiagaraVariable ComparisonData = Variable;
+					TypeUtilityValue->SetValueFromPinDefaultString(Pin->DefaultValue, ComparisonData);
+					if (!DefaultInitialized)
+					{
+						DefaultData = ComparisonData;
+						DefaultInitialized = true;
+					}
+					else if (!ComparisonData.HoldsSameData(DefaultData))
+					{
+						IsConsistent = false;
+					}
 				}
+			}
 
+			if (!IsCustom && !IsConsistent)
+			{
+				UE_LOG(LogNiagaraEditor, Log, TEXT("Niagara graph %s: The default value declaration for the variable '%s' is not consistent between the graph and the metadata.\n  Either change the default mode to 'custom' or check the input pins in the parameter map get node in the graph."),
+					*GetFullName(), *Variable.GetName().ToString());
+				continue;
+			}
+			if (IsCustom)
+			{
+				ScriptVariable->DefaultMode = ENiagaraDefaultMode::Custom;
 			}
-			if (DefaultInitialized) {
-				// set default value from pins
-				ScriptVariable->Variable = DefaultData;
+			else
+			{
+				for (UEdGraphPin* Pin : DefaultPins)
+				{
+					Pin->bNotConnectable = true;
+					Pin->bDefaultValueIsReadOnly = true;
+					if (ScriptVariable->DefaultMode == ENiagaraDefaultMode::Binding || ScriptVariable->DefaultMode == ENiagaraDefaultMode::FailIfPreviouslyNotSet)
+					{
+						Pin->bHidden = true;
+					}
+
+				}
+				if (DefaultInitialized)
+				{
+					// set default value from pins
+					ScriptVariable->Variable = DefaultData;
+				}
 			}
 		}
 	}
