- Flush Niagara Pending Ticks in batches if we inject a large amount in a single frame
- RDG has a limit on the number of passes that can execute and we can blow this limit if we skip forward in the sequencer timeline (for example)

diff --git a/Engine/Plugins/FX/Niagara/Source/Niagara/Private/NiagaraGpuComputeDispatch.cpp b/Engine/Plugins/FX/Niagara/Source/Niagara/Private/NiagaraGpuComputeDispatch.cpp
--- a/Engine/Plugins/FX/Niagara/Source/Niagara/Private/NiagaraGpuComputeDispatch.cpp
+++ b/Engine/Plugins/FX/Niagara/Source/Niagara/Private/NiagaraGpuComputeDispatch.cpp
@@ -94,6 +94,15 @@ namespace FNiagaraGpuComputeDispatchLocal
 		ECVF_Default
 	);
 
+	int32 GTickFlushMaxPendingTicks = 10;
+	static FAutoConsoleVariableRef CVarNiagaraTickFlushMaxPendingTicks(
+		TEXT("fx.Niagara.Batcher.TickFlush.MaxPendingTicks"),
+		GTickFlushMaxPendingTicks,
+		TEXT("The maximum number of unprocess ticks before we process them.\n")
+		TEXT("The larger the number the more data we process in a single frame."),
+		ECVF_Default
+	);
+
 	int32 GTickFlushMode = 1;
 	static FAutoConsoleVariableRef CVarNiagaraTickFlushMode(
 		TEXT("fx.Niagara.Batcher.TickFlush.Mode"),
@@ -329,9 +338,23 @@ void FNiagaraGpuComputeDispatch::ProcessPendingTicksFlush(FRHICommandListImmedia
 
 	// We have pending ticks increment our counter, once we cross the threshold we will perform the appropriate operation
 	++FramesBeforeTickFlush;
+	int32 MaxPendingTicks = 0;
+	const int32 TickFlushMaxPendingTicks = FMath::Max(FNiagaraGpuComputeDispatchLocal::GTickFlushMaxPendingTicks, 1);
+
 	if (!bForceFlush && (FramesBeforeTickFlush < uint32(FNiagaraGpuComputeDispatchLocal::GTickFlushMaxQueuedFrames)) )
 	{
-		return;
+		// If we have a lot of queued up ticks we will need to flush in batches as RDG can only handle a limited number of passes in a single Graph.
+		for (int iTickStage = 0; iTickStage < ENiagaraGpuComputeTickStage::Max; ++iTickStage)
+		{
+			for (FNiagaraSystemGpuComputeProxy* Proxy : ProxiesPerStage[iTickStage])
+			{
+				MaxPendingTicks = FMath::Max<int32>(MaxPendingTicks, Proxy->PendingTicks.Num());
+			}
+		}
+		if (MaxPendingTicks < TickFlushMaxPendingTicks)
+		{
+			return;
+		}
 	}
 	FramesBeforeTickFlush = 0;
 
@@ -350,20 +373,8 @@ void FNiagaraGpuComputeDispatch::ProcessPendingTicksFlush(FRHICommandListImmedia
 		{
 			//UE_LOG(LogNiagara, Log, TEXT("FNiagaraGpuComputeDispatch: Queued ticks are being Processed due to not rendering.  This may result in undesirable simulation artifacts."));
 
-			// Make a pass to see if we have any pending ticks, if not we can early out here
-			bool bHasPendingTicks = false;
-			for (int iTickStage=0; !bHasPendingTicks && (iTickStage < ENiagaraGpuComputeTickStage::Max); ++iTickStage)
-			{
-				for ( FNiagaraSystemGpuComputeProxy* Proxy : ProxiesPerStage[iTickStage] )
-				{
-					bHasPendingTicks = Proxy->PendingTicks.Num() > 0;
-					if (bHasPendingTicks)
-					{
-						break;
-					}
-				}
-			}
-			if (bHasPendingTicks == false)
+			// Early out if we have no pending ticks to process
+			if (MaxPendingTicks == 0)
 			{
 				GpuReadbackManagerPtr->Tick();
 				return;
@@ -405,24 +416,29 @@ void FNiagaraGpuComputeDispatch::ProcessPendingTicksFlush(FRHICommandListImmedia
 			RHICmdList.BeginScene();
 
 			// Execute all ticks that we can support without invalid simulations
-			FRDGBuilder GraphBuilder(RHICmdList);
-			CreateSystemTextures(GraphBuilder);
-			PreInitViews(GraphBuilder, bAllowGPUParticleUpdate);
-			AddPass(GraphBuilder, RDG_EVENT_NAME("UpdateDrawIndirectBuffers - PreOpaque"),
-				[this](FRHICommandListImmediate& RHICmdList)
-				{
-					GPUInstanceCounterManager.UpdateDrawIndirectBuffers(this, RHICmdList, ENiagaraGPUCountUpdatePhase::PreOpaque);
-				}
-			);
-			PostInitViews(GraphBuilder, DummyViews, bAllowGPUParticleUpdate);
-			PostRenderOpaque(GraphBuilder, DummyViews, bAllowGPUParticleUpdate);
-			AddPass(GraphBuilder, RDG_EVENT_NAME("UpdateDrawIndirectBuffers - PostOpaque"),
-				[this](FRHICommandListImmediate& RHICmdList)
-				{
-					GPUInstanceCounterManager.UpdateDrawIndirectBuffers(this, RHICmdList, ENiagaraGPUCountUpdatePhase::PostOpaque);
-				}
-			);
-			GraphBuilder.Execute();
+			MaxTicksToFlush = TickFlushMaxPendingTicks;
+			for (int32 iTickBatch = 0; iTickBatch < MaxPendingTicks; iTickBatch+=MaxTicksToFlush)
+			{
+				FRDGBuilder GraphBuilder(RHICmdList);
+				CreateSystemTextures(GraphBuilder);
+				PreInitViews(GraphBuilder, bAllowGPUParticleUpdate);
+				AddPass(GraphBuilder, RDG_EVENT_NAME("UpdateDrawIndirectBuffers - PreOpaque"),
+					[this](FRHICommandListImmediate& RHICmdList)
+					{
+						GPUInstanceCounterManager.UpdateDrawIndirectBuffers(this, RHICmdList, ENiagaraGPUCountUpdatePhase::PreOpaque);
+					}
+				);
+				PostInitViews(GraphBuilder, DummyViews, bAllowGPUParticleUpdate);
+				PostRenderOpaque(GraphBuilder, DummyViews, bAllowGPUParticleUpdate);
+				AddPass(GraphBuilder, RDG_EVENT_NAME("UpdateDrawIndirectBuffers - PostOpaque"),
+					[this](FRHICommandListImmediate& RHICmdList)
+					{
+						GPUInstanceCounterManager.UpdateDrawIndirectBuffers(this, RHICmdList, ENiagaraGPUCountUpdatePhase::PostOpaque);
+					}
+				);
+				GraphBuilder.Execute();
+			}
+			MaxTicksToFlush = TNumericLimits<int32>::Max();
 
 			// We have completed flushing the commands
 			RHICmdList.EndScene();
@@ -450,7 +466,7 @@ void FNiagaraGpuComputeDispatch::FinishDispatches()
 	{
 		for (FNiagaraSystemGpuComputeProxy* ComputeProxy : ProxiesPerStage[iTickStage])
 		{
-			ComputeProxy->ReleaseTicks(GPUInstanceCounterManager);
+			ComputeProxy->ReleaseTicks(GPUInstanceCounterManager, MaxTicksToFlush);
 		}
 	}
 
@@ -618,8 +634,10 @@ void FNiagaraGpuComputeDispatch::UpdateInstanceCountManager(FRHICommandListImmed
 		{
 			for (FNiagaraSystemGpuComputeProxy* ComputeProxy : ProxiesPerStage[iTickStage])
 			{
-				for (FNiagaraGPUSystemTick& Tick : ComputeProxy->PendingTicks)
+				const int32 NumTicksToProcess = FMath::Min(ComputeProxy->PendingTicks.Num(), MaxTicksToFlush);
+				for (int32 iTickToProcess=0; iTickToProcess < NumTicksToProcess; ++iTickToProcess)
 				{
+					FNiagaraGPUSystemTick& Tick = ComputeProxy->PendingTicks[iTickToProcess];
 					TotalDispatchCount += (int32)Tick.TotalDispatches;
 
 					for ( FNiagaraComputeInstanceData& InstanceData : Tick.GetInstances() )
@@ -697,13 +715,16 @@ void FNiagaraGpuComputeDispatch::PrepareTicksForProxy(FRHICommandListImmediate&
 	const bool bEnqueueCountReadback = !GPUInstanceCounterManager.HasPendingGPUReadback();
 
 	// Set final tick flag
-	ComputeProxy->PendingTicks.Last().bIsFinalTick = true;
+	const int32 NumTicksToProcess = FMath::Min(ComputeProxy->PendingTicks.Num(), MaxTicksToFlush);
+	ComputeProxy->PendingTicks[NumTicksToProcess - 1].bIsFinalTick = true;
 
 	// Process ticks
 	int32 iTickStartDispatchGroup = 0;
 
-	for (FNiagaraGPUSystemTick& Tick : ComputeProxy->PendingTicks)
+	for ( int32 iTickToProcess=0; iTickToProcess < NumTicksToProcess; ++iTickToProcess )
 	{
+		FNiagaraGPUSystemTick& Tick = ComputeProxy->PendingTicks[iTickToProcess];
+
 		int32 iInstanceStartDispatchGroup = iTickStartDispatchGroup;
 		int32 iInstanceCurrDispatchGroup = iTickStartDispatchGroup;
 		bool bHasFreeIDUpdates = false;
diff --git a/Engine/Plugins/FX/Niagara/Source/Niagara/Private/NiagaraGpuComputeDispatch.h b/Engine/Plugins/FX/Niagara/Source/Niagara/Private/NiagaraGpuComputeDispatch.h
--- a/Engine/Plugins/FX/Niagara/Source/Niagara/Private/NiagaraGpuComputeDispatch.h
+++ b/Engine/Plugins/FX/Niagara/Source/Niagara/Private/NiagaraGpuComputeDispatch.h
@@ -250,6 +250,8 @@ private:
 
 	int32 TotalDispatchesThisFrame = 0;
 
+	int32 MaxTicksToFlush = TNumericLimits<int32>::Max();
+
 	bool bRequiresReadback = false;
 	TArray<FNiagaraSystemGpuComputeProxy*> ProxiesPerStage[ENiagaraGpuComputeTickStage::Max];
 
diff --git a/Engine/Plugins/FX/Niagara/Source/Niagara/Private/NiagaraSystemGpuComputeProxy.cpp b/Engine/Plugins/FX/Niagara/Source/Niagara/Private/NiagaraSystemGpuComputeProxy.cpp
--- a/Engine/Plugins/FX/Niagara/Source/Niagara/Private/NiagaraSystemGpuComputeProxy.cpp
+++ b/Engine/Plugins/FX/Niagara/Source/Niagara/Private/NiagaraSystemGpuComputeProxy.cpp
@@ -105,7 +105,7 @@ void FNiagaraSystemGpuComputeProxy::RemoveFromRenderThread(FNiagaraGpuComputeDis
 		[=](FRHICommandListImmediate& RHICmdList)
 		{
 			ComputeDispatchInterface->RemoveGpuComputeProxy(this);
-			ReleaseTicks(ComputeDispatchInterface->GetGPUInstanceCounterManager());
+			ReleaseTicks(ComputeDispatchInterface->GetGPUInstanceCounterManager(), TNumericLimits<int32>::Max());
 
 			for (FNiagaraComputeExecutionContext* ComputeContext : ComputeContexts)
 			{
@@ -137,7 +137,7 @@ void FNiagaraSystemGpuComputeProxy::ClearTicksFromRenderThread(FNiagaraGpuComput
 	ENQUEUE_RENDER_COMMAND(ClearTicksFromProxy)(
 		[=](FRHICommandListImmediate& RHICmdList)
 		{
-			ReleaseTicks(ComputeDispatchInterface->GetGPUInstanceCounterManager());
+			ReleaseTicks(ComputeDispatchInterface->GetGPUInstanceCounterManager(), TNumericLimits<int32>::Max());
 		}
 	);
 }
@@ -170,16 +170,18 @@ void FNiagaraSystemGpuComputeProxy::QueueTick(const FNiagaraGPUSystemTick& Tick)
 	}
 }
 
-void FNiagaraSystemGpuComputeProxy::ReleaseTicks(FNiagaraGPUInstanceCountManager& GPUInstanceCountManager)
+void FNiagaraSystemGpuComputeProxy::ReleaseTicks(FNiagaraGPUInstanceCountManager& GPUInstanceCountManager, int32 NumTicksToRelease)
 {
 	check(IsInRenderingThread());
 
 	// Release all the ticks
-	for (auto& Tick : PendingTicks)
+	NumTicksToRelease = FMath::Min(NumTicksToRelease, PendingTicks.Num());
+	for ( int32 iTick=0; iTick < NumTicksToRelease; ++iTick )
 	{
+		FNiagaraGPUSystemTick& Tick = PendingTicks[iTick];
 		Tick.Destroy();
 	}
-	PendingTicks.Empty();
+	PendingTicks.RemoveAt(0, NumTicksToRelease, NumTicksToRelease == PendingTicks.Num());
 
 	for (FNiagaraComputeExecutionContext* ComputeContext : ComputeContexts)
 	{
diff --git a/Engine/Plugins/FX/Niagara/Source/Niagara/Public/NiagaraSystemGpuComputeProxy.h b/Engine/Plugins/FX/Niagara/Source/Niagara/Public/NiagaraSystemGpuComputeProxy.h
--- a/Engine/Plugins/FX/Niagara/Source/Niagara/Public/NiagaraSystemGpuComputeProxy.h
+++ b/Engine/Plugins/FX/Niagara/Source/Niagara/Public/NiagaraSystemGpuComputeProxy.h
@@ -25,7 +25,7 @@ public:
 	FNiagaraSystemInstanceID GetSystemInstanceID() const { return SystemInstanceID; }
 	ENiagaraGpuComputeTickStage::Type GetComputeTickStage() const { return ComputeTickStage; }
 	void QueueTick(const FNiagaraGPUSystemTick& Tick);
-	void ReleaseTicks(FNiagaraGPUInstanceCountManager& GPUInstanceCountManager);
+	void ReleaseTicks(FNiagaraGPUInstanceCountManager& GPUInstanceCountManager, int32 NumTicksToRelease);
 
 	bool RequiresDistanceFieldData() const { return bRequiresDistanceFieldData; }
 	bool RequiresDepthBuffer() const { return bRequiresDepthBuffer; }
