- Disable RayTracing Ribbon path for GPU / GPU init as both TDR currently
- Add master control for GPU ribbons being enabled
- Adjust GPU init control into user, forced on, forced off

#rnx

diff --git a/Engine/Plugins/FX/Niagara/Source/Niagara/Private/NiagaraRendererRibbons.cpp b/Engine/Plugins/FX/Niagara/Source/Niagara/Private/NiagaraRendererRibbons.cpp
--- a/Engine/Plugins/FX/Niagara/Source/Niagara/Private/NiagaraRendererRibbons.cpp
+++ b/Engine/Plugins/FX/Niagara/Source/Niagara/Private/NiagaraRendererRibbons.cpp
@@ -104,28 +104,29 @@ static FAutoConsoleVariableRef CVarEnableNiagaraRibbonRendering(
 	ECVF_Default
 );
 
-static TAutoConsoleVariable<int32> CVarRayTracingNiagaraRibbons(
-	TEXT("r.RayTracing.Geometry.NiagaraRibbons"),
-	1,
-	TEXT("Include Niagara ribbons in ray tracing effects (default = 1 (Niagara ribbons enabled in ray tracing))"));
-
-
-static int32 GbEnableNiagaraRibbonComputeShaderGen = 1;
-static FAutoConsoleVariableRef CVarEnableNiagaraRibbonComputeShaderGen(
-	TEXT("fx.EnableNiagaraRibbonComputeShaderGen"),
-	GbEnableNiagaraRibbonComputeShaderGen,
-	TEXT("If == 1, Niagara Ribbon Renderer will use compute shaders if available to generate intermediate data. \n"),
-	ECVF_Default
+static int32 GNiagaraRibbonGpuEnabled = 1;
+static FAutoConsoleVariableRef CVarNiagaraRibbonGpuEnabled(
+	TEXT("Niagara.Ribbon.GpuEnabled"),
+	GNiagaraRibbonGpuEnabled,
+	TEXT("Enable any GPU ribbon related code (including GPU init)."),
+	ECVF_Scalability
 );
 
-static int32 GbForceNiagaraRibbonGPUInit = 0;
-static FAutoConsoleVariableRef CVarForceNiagaraRibbonGPUInit(
-	TEXT("fx.ForceNiagaraRibbonGPUInit"),
-	GbForceNiagaraRibbonGPUInit,
-	TEXT("If == 1, Niagara Ribbon Renderer will use compute shader based initialization for all CPU systems. \n"),
-	ECVF_Default
+static int32 GNiagaraRibbonGpuInitMode = 0;
+static FAutoConsoleVariableRef CVarNiagaraRibbonGpuInitMode(
+	TEXT("Niagara.Ribbon.GpuInitMode"),
+	GNiagaraRibbonGpuInitMode,
+	TEXT("Modifies the GPU initialization mode used, i.e. offloading CPU calculations to the GPU.\n")
+	TEXT("0 = Respect bUseGPUInit from properties (Default)\n")
+	TEXT("1 = Force enabled\n")
+	TEXT("2 = Force disabled"),
+	ECVF_Scalability
 );
 
+static TAutoConsoleVariable<int32> CVarRayTracingNiagaraRibbons(
+	TEXT("r.RayTracing.Geometry.NiagaraRibbons"),
+	1,
+	TEXT("Include Niagara ribbons in ray tracing effects (default = 1 (Niagara ribbons enabled in ray tracing))"));
 
 // max absolute error 9.0x10^-3
 // Eberly's polynomial degree 1 - respect bounds
@@ -788,6 +789,11 @@ FNiagaraDynamicDataBase* FNiagaraRendererRibbons::GenerateDynamicData(const FNia
 			return nullptr;
 		}
 
+		if ((SimTarget == ENiagaraSimTarget::GPUComputeSim) && GNiagaraRibbonGpuEnabled == 0)
+		{
+			return nullptr;
+		}
+
 		FNiagaraDataBuffer* DataToRender = Emitter->GetData().GetCurrentData();		
 		if(SimTarget == ENiagaraSimTarget::GPUComputeSim || (DataToRender != nullptr && DataToRender->GetNumInstances() > 1))
 		{
@@ -808,9 +814,8 @@ FNiagaraDynamicDataBase* FNiagaraRendererRibbons::GenerateDynamicData(const FNia
 			const bool bIsGPUSystem = SimTarget == ENiagaraSimTarget::GPUComputeSim;
 			
 			// We disable compute initialization when compute isn't available or they're CVar'd off
-			const bool bCanUseComputeGenForCPUSystems = (GbEnableNiagaraRibbonComputeShaderGen && FNiagaraUtilities::AllowComputeShaders(GShaderPlatformForFeatureLevel[FeatureLevel]));
-						
-			const bool bWantsGPUInit = (bCanUseComputeGenForCPUSystems && (GbForceNiagaraRibbonGPUInit || Properties->bUseGPUInit));
+			const bool bCanUseComputeGenForCPUSystems = FNiagaraUtilities::AllowComputeShaders(GShaderPlatformForFeatureLevel[FeatureLevel]) && (GNiagaraRibbonGpuInitMode != 2) && (GNiagaraRibbonGpuEnabled != 0);
+			const bool bWantsGPUInit = bCanUseComputeGenForCPUSystems && (Properties->bUseGPUInit || (GNiagaraRibbonGpuInitMode == 1));
 			
 			DynamicData->bUseGPUInit = bIsGPUSystem || bWantsGPUInit;
 			DynamicData->bIsGPUSystem = bIsGPUSystem;
@@ -891,12 +896,15 @@ void FNiagaraRendererRibbons::GetDynamicRayTracingInstances(FRayTracingMaterialG
 
 	if (DynamicDataRibbon->bIsGPUSystem)
 	{
+		//-TODO: RayTracing does not support GPU Ribbons currently
+		return;
+
 		// Bail if we don't have enough particle data to have a valid ribbon
 		// or if somehow the sim targets don't match
-		if (SimTarget != ENiagaraSimTarget::GPUComputeSim || SourceParticleData->GetNumInstancesAllocated() < 2)
-		{
-			return;
-		}
+		//if (SimTarget != ENiagaraSimTarget::GPUComputeSim || SourceParticleData->GetNumInstancesAllocated() < 2)
+		//{
+		//	return;
+		//}
 	}
 	else
 	{
@@ -908,6 +916,12 @@ void FNiagaraRendererRibbons::GetDynamicRayTracingInstances(FRayTracingMaterialG
 			return;
 		}
 
+		//-TODO: RayTracing does not support GPU Init for Ribbons
+		if (DynamicDataRibbon->bUseGPUInit)
+		{
+			return;
+		}
+
 		if (!DynamicDataRibbon->bUseGPUInit && (
 			!DynamicDataRibbon->GenerationOutput.IsValid() ||
 			DynamicDataRibbon->GenerationOutput->SegmentData.Num() <= 0))
