- Add some tooltips for array pins

#rnx

diff --git a/Engine/Plugins/FX/Niagara/Source/Niagara/Classes/NiagaraDataInterfaceArrayImpl.h b/Engine/Plugins/FX/Niagara/Source/Niagara/Classes/NiagaraDataInterfaceArrayImpl.h
--- a/Engine/Plugins/FX/Niagara/Source/Niagara/Classes/NiagaraDataInterfaceArrayImpl.h
+++ b/Engine/Plugins/FX/Niagara/Source/Niagara/Classes/NiagaraDataInterfaceArrayImpl.h
@@ -586,39 +586,39 @@ struct FNDIArrayProxyImpl : public INDIArrayProxyBase
 		{
 			FNiagaraFunctionSignature& Sig = OutFunctions.Add_GetRef(DefaultImmutableSig);
 			Sig.Name = FNiagaraDataInterfaceArrayImplHelper::Function_LengthName;
-			#if WITH_EDITORONLY_DATA
-				Sig.Description = NSLOCTEXT("Niagara", "Array_LengthDesc", "Gets the number of elements in the array.");
-			#endif
 			Sig.Outputs.Emplace(FNiagaraTypeDefinition::GetIntDef(), TEXT("Num"));
+		#if WITH_EDITORONLY_DATA
+			Sig.Description = NSLOCTEXT("Niagara", "Array_LengthDesc", "Gets the number of elements in the array.");
+		#endif
 		}
 
 		{
 			FNiagaraFunctionSignature& Sig = OutFunctions.Add_GetRef(DefaultImmutableSig);
 			Sig.Name = FNiagaraDataInterfaceArrayImplHelper::Function_IsValidIndexName;
-			#if WITH_EDITORONLY_DATA
-				Sig.Description = NSLOCTEXT("Niagara", "Array_IsValidIndexDesc", "Tests to see if the index is valid and exists in the array.");
-			#endif
 			Sig.Inputs.Emplace(FNiagaraTypeDefinition::GetIntDef(), TEXT("Index"));
 			Sig.Outputs.Emplace(FNiagaraTypeDefinition::GetBoolDef(), TEXT("Valid"));
+		#if WITH_EDITORONLY_DATA
+			Sig.Description = NSLOCTEXT("Niagara", "Array_IsValidIndexDesc", "Tests to see if the index is valid and exists in the array.");
+		#endif
 		}
 
 		{
 			FNiagaraFunctionSignature& Sig = OutFunctions.Add_GetRef(DefaultImmutableSig);
 			Sig.Name = FNiagaraDataInterfaceArrayImplHelper::Function_LastIndexName;
-#if WITH_EDITORONLY_DATA
-			Sig.Description = NSLOCTEXT("Niagara", "Array_LastIndexDesc", "Returns the last valid index in the array, will be -1 if no elements.");
-#endif
 			Sig.Outputs.Emplace(FNiagaraTypeDefinition::GetIntDef(), TEXT("Index"));
+		#if WITH_EDITORONLY_DATA
+			Sig.Description = NSLOCTEXT("Niagara", "Array_LastIndexDesc", "Returns the last valid index in the array, will be -1 if no elements.");
+		#endif
 		}
 
 		{
 			FNiagaraFunctionSignature& Sig = OutFunctions.Add_GetRef(DefaultImmutableSig);
 			Sig.Name = FNiagaraDataInterfaceArrayImplHelper::Function_GetName;
-			#if WITH_EDITORONLY_DATA
-				Sig.Description = NSLOCTEXT("Niagara", "Array_GetDesc", "Gets the value from the array at the given zero based index.");
-			#endif
 			Sig.Inputs.Emplace(FNiagaraTypeDefinition::GetIntDef(), TEXT("Index"));
 			Sig.Outputs.Emplace(FNDIArrayImplHelper<TArrayType>::GetTypeDefinition(), TEXT("Value"));
+		#if WITH_EDITORONLY_DATA
+			Sig.Description = NSLOCTEXT("Niagara", "Array_GetDesc", "Gets the value from the array at the given zero based index.");
+		#endif
 		}
 
 		// Mutable functions
@@ -628,52 +628,56 @@ struct FNDIArrayProxyImpl : public INDIArrayProxyBase
 		{
 			FNiagaraFunctionSignature& Sig = OutFunctions.Add_GetRef(DefaultMutableSig);
 			Sig.Name = FNiagaraDataInterfaceArrayImplHelper::Function_ClearName;
-			#if WITH_EDITORONLY_DATA
-				Sig.Description = NSLOCTEXT("Niagara", "Array_ClearDesc", "Clears the array, removing all elements");
-			#endif
 			Sig.ModuleUsageBitmask = ENiagaraScriptUsageMask::System | ENiagaraScriptUsageMask::Emitter;
+		#if WITH_EDITORONLY_DATA
+			Sig.Description = NSLOCTEXT("Niagara", "Array_ClearDesc", "Clears the array, removing all elements");
+		#endif
 		}
 
 		{
 			FNiagaraFunctionSignature& Sig = OutFunctions.Add_GetRef(DefaultMutableSig);
 			Sig.Name = FNiagaraDataInterfaceArrayImplHelper::Function_ResizeName;
-			#if WITH_EDITORONLY_DATA
-				Sig.Description = NSLOCTEXT("Niagara", "Array_ResizeDesc", "Resizes the array to the specified size, initializing new elements with the default value.");
-			#endif
 			Sig.ModuleUsageBitmask = ENiagaraScriptUsageMask::System | ENiagaraScriptUsageMask::Emitter;
 			Sig.Inputs.Emplace(FNiagaraTypeDefinition::GetIntDef(), TEXT("Num"));
+		#if WITH_EDITORONLY_DATA
+			Sig.Description = NSLOCTEXT("Niagara", "Array_ResizeDesc", "Resizes the array to the specified size, initializing new elements with the default value.");
+		#endif
 		}
 
 		{
 			FNiagaraFunctionSignature& Sig = OutFunctions.Add_GetRef(DefaultMutableSig);
 			Sig.Name = FNiagaraDataInterfaceArrayImplHelper::Function_SetArrayElemName;
-			#if WITH_EDITORONLY_DATA
-				Sig.Description = NSLOCTEXT("Niagara", "Array_SetArrayElemDesc", "Sets the value at the given zero based index (i.e the first element is 0).");
-			#endif
 			Sig.Inputs.Emplace(FNiagaraTypeDefinition::GetBoolDef(), TEXT("SkipSet"));
 			Sig.Inputs.Emplace(FNiagaraTypeDefinition::GetIntDef(), TEXT("Index"));
 			Sig.Inputs.Emplace(FNDIArrayImplHelper<TArrayType>::GetTypeDefinition(), TEXT("Value"));
+		#if WITH_EDITORONLY_DATA
+			Sig.Description = NSLOCTEXT("Niagara", "Array_SetArrayElemDesc", "Sets the value at the given zero based index (i.e the first element is 0).");
+			Sig.InputDescriptions.Add(Sig.Inputs[1], NSLOCTEXT("Niagara", "Array_SetArrayElemDesc_SkipSet", "When enabled will not set the array value."));
+		#endif
 		}
 
 		{
 			FNiagaraFunctionSignature& Sig = OutFunctions.Add_GetRef(DefaultMutableSig);
 			Sig.Name = FNiagaraDataInterfaceArrayImplHelper::Function_AddName;
-			#if WITH_EDITORONLY_DATA
-				Sig.Description = NSLOCTEXT("Niagara", "Array_AddDesc", "Optionally add a value onto the end of the array.");
-			#endif
 			Sig.Inputs.Emplace(FNiagaraTypeDefinition::GetBoolDef(), TEXT("SkipAdd"));
 			Sig.Inputs.Emplace(FNDIArrayImplHelper<TArrayType>::GetTypeDefinition(), TEXT("Value"));
+		#if WITH_EDITORONLY_DATA
+			Sig.Description = NSLOCTEXT("Niagara", "Array_AddDesc", "Optionally add a value onto the end of the array.");
+			Sig.InputDescriptions.Add(Sig.Inputs[1], NSLOCTEXT("Niagara", "Array_AddDesc_SkipAdd", "When enabled we will not add an element to the array."));
+		#endif
 		}
 
 		{
 			FNiagaraFunctionSignature& Sig = OutFunctions.Add_GetRef(DefaultMutableSig);
 			Sig.Name = FNiagaraDataInterfaceArrayImplHelper::Function_RemoveLastElemName;
-#if WITH_EDITORONLY_DATA
-			Sig.Description = NSLOCTEXT("Niagara", "Array_RemoveLastElemDesc", "Optionally remove the last element from the array.  Returns the default value if no elements are in the array or you skip the remove.");
-#endif
 			Sig.Inputs.Emplace(FNiagaraTypeDefinition::GetBoolDef(), TEXT("SkipRemove"));
 			Sig.Outputs.Emplace(FNDIArrayImplHelper<TArrayType>::GetTypeDefinition(), TEXT("Value"));
 			Sig.Outputs.Emplace(FNiagaraTypeDefinition::GetBoolDef(), TEXT("IsValid"));
+		#if WITH_EDITORONLY_DATA
+			Sig.Description = NSLOCTEXT("Niagara", "Array_RemoveLastElemDesc", "Optionally remove the last element from the array.  Returns the default value if no elements are in the array or you skip the remove.");
+			Sig.InputDescriptions.Add(Sig.Inputs[1], NSLOCTEXT("Niagara", "Array_RemoveLastElemDesc_SkipRemove", "When enabled will not remove a value from the array, the return value will therefore be invalid."));
+			Sig.OutputDescriptions.Add(Sig.Outputs[1], NSLOCTEXT("Niagara", "Array_RemoveLastElemDesc_IsValid", "True if we removed a value from the array, False if no entries or we skipped the remove."));
+		#endif
 		}
 	}
 
