Fix FNiagaraDistanceFieldHelper compile errors

diff --git a/Engine/Plugins/FX/Niagara/Source/NiagaraShader/Private/NiagaraDistanceFieldHelper.cpp b/Engine/Plugins/FX/Niagara/Source/NiagaraShader/Private/NiagaraDistanceFieldHelper.cpp
--- a/Engine/Plugins/FX/Niagara/Source/NiagaraShader/Private/NiagaraDistanceFieldHelper.cpp
+++ b/Engine/Plugins/FX/Niagara/Source/NiagaraShader/Private/NiagaraDistanceFieldHelper.cpp
@@ -30,9 +30,6 @@ void FNiagaraDistanceFieldHelper::SetGlobalDistanceFieldParameters(const FGlobal
 		ShaderParameters.GlobalVolumeTexelSize = 1.0f / OptionalParameterData->GlobalDFResolution;
 		ShaderParameters.MaxGlobalDFAOConeDistance = OptionalParameterData->MaxDFAOConeDistance;
 		ShaderParameters.NumGlobalSDFClipmaps = OptionalParameterData->NumGlobalSDFClipmaps;
-		ShaderParameters.FullyCoveredExpandSurfaceScale = 0.0f;//GLumenSceneGlobalSDFFullyCoveredExpandSurfaceScale;
-		ShaderParameters.UncoveredExpandSurfaceScale = 0.0f;//GLumenSceneGlobalSDFUncoveredExpandSurfaceScale;
-		ShaderParameters.UncoveredMinStepScale = 0.0f;//GLumenSceneGlobalSDFUncoveredMinStepScale;
 	}
 	else
 	{
@@ -58,13 +55,13 @@ void FNiagaraDistanceFieldHelper::SetGlobalDistanceFieldParameters(const FGlobal
 		ShaderParameters.GlobalVolumeTexelSize = 0.0f;
 		ShaderParameters.MaxGlobalDFAOConeDistance = 0.0f;
 		ShaderParameters.NumGlobalSDFClipmaps = 0;
-		ShaderParameters.FullyCoveredExpandSurfaceScale = 0.0f;
-		ShaderParameters.UncoveredExpandSurfaceScale = 0.0f;
-		ShaderParameters.UncoveredMinStepScale = 0.0f;
-		ShaderParameters.FullyCoveredExpandSurfaceScale = 0.0f;//GLumenSceneGlobalSDFFullyCoveredExpandSurfaceScale;
-		ShaderParameters.UncoveredExpandSurfaceScale = 0.0f;//GLumenSceneGlobalSDFUncoveredExpandSurfaceScale;
-		ShaderParameters.UncoveredMinStepScale = 0.0f;//GLumenSceneGlobalSDFUncoveredMinStepScale;
 	}
+
+	ShaderParameters.CoveredExpandSurfaceScale = 0.0f;
+	ShaderParameters.NotCoveredExpandSurfaceScale = 0.0f;
+	ShaderParameters.NotCoveredMinStepScale = 0.0f;
+	ShaderParameters.DitheredTransparencyStepThreshold = 0.0f;
+	ShaderParameters.DitheredTransparencyTraceThreshold = 0.0f;
 }
 
 void FNiagaraDistanceFieldHelper::SetMeshDistanceFieldParameters(FRDGBuilder& GraphBuilder, const FDistanceFieldSceneData* OptionalDistanceFieldData, FDistanceFieldObjectBufferParameters& ObjectShaderParameters, FDistanceFieldAtlasParameters& AtlasShaderParameters, FRHIShaderResourceView* DummyFloat4Buffer)
