Niagara simulation now uses the game thread local queue when waiting for async results.
This helps performance on platforms with a low core count

diff --git a/Engine/Plugins/FX/Niagara/Source/Niagara/Private/NiagaraSystemInstance.cpp b/Engine/Plugins/FX/Niagara/Source/Niagara/Private/NiagaraSystemInstance.cpp
--- a/Engine/Plugins/FX/Niagara/Source/Niagara/Private/NiagaraSystemInstance.cpp
+++ b/Engine/Plugins/FX/Niagara/Source/Niagara/Private/NiagaraSystemInstance.cpp
@@ -2282,7 +2282,7 @@ void FNiagaraSystemInstance::WaitForConcurrentTickDoNotFinalize(bool bEnsureComp
 		{
 			do
 			{
-				FTaskGraphInterface::Get().WaitUntilTaskCompletes(ConcurrentTickGraphEvent, ENamedThreads::GameThread);
+				FTaskGraphInterface::Get().WaitUntilTaskCompletes(ConcurrentTickGraphEvent, ENamedThreads::GameThread_Local);
 			} while (ConcurrentTickGraphEvent && !ConcurrentTickGraphEvent->IsComplete());
 		}
 
diff --git a/Engine/Plugins/FX/Niagara/Source/Niagara/Private/NiagaraSystemSimulation.cpp b/Engine/Plugins/FX/Niagara/Source/Niagara/Private/NiagaraSystemSimulation.cpp
--- a/Engine/Plugins/FX/Niagara/Source/Niagara/Private/NiagaraSystemSimulation.cpp
+++ b/Engine/Plugins/FX/Niagara/Source/Niagara/Private/NiagaraSystemSimulation.cpp
@@ -1617,7 +1617,7 @@ void FNiagaraSystemSimulation::WaitForConcurrentTickComplete(bool bEnsureComplet
 		}
 		else
 		{
-			FTaskGraphInterface::Get().WaitUntilTaskCompletes(ConcurrentTickGraphEvent, ENamedThreads::GameThread);
+			FTaskGraphInterface::Get().WaitUntilTaskCompletes(ConcurrentTickGraphEvent, ENamedThreads::GameThread_Local);
 		}
 	}
 	ConcurrentTickGraphEvent = nullptr;
