Move parameter rename logic from parameter panel viewmodel to script & system viewmodel respectively.

diff --git a/Engine/Plugins/FX/Niagara/Source/NiagaraEditor/Private/NiagaraEditorUtilities.cpp b/Engine/Plugins/FX/Niagara/Source/NiagaraEditor/Private/NiagaraEditorUtilities.cpp
--- a/Engine/Plugins/FX/Niagara/Source/NiagaraEditor/Private/NiagaraEditorUtilities.cpp
+++ b/Engine/Plugins/FX/Niagara/Source/NiagaraEditor/Private/NiagaraEditorUtilities.cpp
@@ -3965,4 +3965,47 @@ EParameterDefinitionMatchState FNiagaraParameterDefinitionsUtilities::GetDefinit
 	}
 }
 
+void FNiagaraParameterDefinitionsUtilities::TrySubscribeScriptVarToDefinitionByName(UNiagaraScriptVariable* ScriptVar, INiagaraParameterDefinitionsSubscriberViewModel* OwningDefinitionSubscriberViewModel)
+{
+	const FName& ScriptVarName = ScriptVar->Variable.GetName();
+	const TArray<const UNiagaraScriptVariable*> ReservedParametersForName = FNiagaraParameterDefinitionsUtilities::FindReservedParametersByName(ScriptVarName);
+
+	if (ReservedParametersForName.Num() == 0)
+	{
+		// Do not call SetParameterIsSubscribedToLibrary() unless ScriptVarToModify is already marked as SubscribedToParameterDefinitions.
+		if (ScriptVar->GetIsSubscribedToParameterDefinitions())
+		{
+			//SetParameterIsSubscribedToLibrary(ScriptVar, false);
+			OwningDefinitionSubscriberViewModel->SetParameterIsSubscribedToDefinitions(ScriptVar->Metadata.GetVariableGuid(), false);
+		}
+	}
+	else if (ReservedParametersForName.Num() == 1)
+	{
+		const UNiagaraScriptVariable* ReservedParameterDefinition = (ReservedParametersForName)[0];
+		if (ReservedParameterDefinition->Variable.GetType() != ScriptVar->Variable.GetType())
+		{
+			const FText TypeMismatchWarningTemplate = LOCTEXT("RenameParameter_DefinitionTypeMismatch", "Renamed parameter \"{0}\" with type {1}. Type does not match existing parameter definition \"{0}\" with type {2} from {3}! ");
+			FText TypeMismatchWarning = FText::Format(
+				TypeMismatchWarningTemplate,
+				FText::FromName(ScriptVarName),
+				ScriptVar->Variable.GetType().GetNameText(),
+				ReservedParameterDefinition->Variable.GetType().GetNameText(),
+				FText::FromString(ReservedParameterDefinition->GetTypedOuter<UNiagaraParameterDefinitions>()->GetName()));
+			FNiagaraEditorUtilities::WarnWithToastAndLog(TypeMismatchWarning);
+
+			// Do not call SetParameterIsSubscribedToLibrary() unless ScriptVarToModify is already marked as SubscribedToParameterDefinitions.
+			if (ScriptVar->GetIsSubscribedToParameterDefinitions())
+			{
+				OwningDefinitionSubscriberViewModel->SetParameterIsSubscribedToDefinitions(ScriptVar->Metadata.GetVariableGuid(), false);
+			}
+		}
+		else
+		{
+			// NOTE: It is possible we are calling SetParameterIsSubscribedToLibrary() on a UNiagaraScriptVariable that is already marked as SubscribedToParameterDefinitions;
+			// This is intended as SetParameterIsSubscribedToLibrary() will register a new link to a different parameter definition.
+			OwningDefinitionSubscriberViewModel->SetParameterIsSubscribedToDefinitions(ScriptVar->Metadata.GetVariableGuid(), true);
+		}
+	}
+}
+
 #undef LOCTEXT_NAMESPACE
diff --git a/Engine/Plugins/FX/Niagara/Source/NiagaraEditor/Private/NiagaraGraph.cpp b/Engine/Plugins/FX/Niagara/Source/NiagaraEditor/Private/NiagaraGraph.cpp
--- a/Engine/Plugins/FX/Niagara/Source/NiagaraEditor/Private/NiagaraGraph.cpp
+++ b/Engine/Plugins/FX/Niagara/Source/NiagaraEditor/Private/NiagaraGraph.cpp
@@ -33,7 +33,7 @@
 #include "Misc/SecureHash.h"
 #include "Modules/ModuleManager.h"
 #include "String/ParseTokens.h"
-#include "ViewModels/NiagaraParameterPanelViewModel.h"
+#include "ViewModels/NiagaraScriptViewModel.h"
 #include "ViewModels/TNiagaraViewModelManager.h"
 
 
@@ -2081,10 +2081,10 @@ bool UNiagaraGraph::RenameParameterFromPin(const FNiagaraVariable& Parameter, FN
 	}
 
 	UNiagaraScript* OuterScript = GetTypedOuter<UNiagaraScript>();
-	TSharedPtr<FNiagaraScriptToolkitParameterPanelViewModel> ParameterPanelViewModel = TNiagaraViewModelManager<UNiagaraScript, FNiagaraScriptToolkitParameterPanelViewModel>::GetExistingViewModelForObject(OuterScript);
-	if (ParameterPanelViewModel.IsValid() == false)
+	TSharedPtr<FNiagaraScriptViewModel> ScriptViewModel = TNiagaraViewModelManager<UNiagaraScript, FNiagaraScriptViewModel>::GetExistingViewModelForObject(OuterScript);
+	if (ScriptViewModel.IsValid() == false)
 	{
-		ensureMsgf(false, TEXT("Failed to get parameter panel view model when renaming parameter via pin in graph!"));
+		ensureMsgf(false, TEXT("Failed to get script viewmodel when renaming parameter via pin in graph!"));
 		return false;
 	}
 
@@ -2101,7 +2101,7 @@ bool UNiagaraGraph::RenameParameterFromPin(const FNiagaraVariable& Parameter, FN
 			{
 				if (TObjectPtr<UNiagaraScriptVariable> const* RenamedScriptVarPtr = VariableToScriptVariable.Find(FNiagaraVariable(Parameter.GetType(), NewName)))
 				{
-					ParameterPanelViewModel->SubscribeParameterToLibraryIfMatchingDefinition(*RenamedScriptVarPtr, NewName);
+					FNiagaraParameterDefinitionsUtilities::TrySubscribeScriptVarToDefinitionByName(*RenamedScriptVarPtr, ScriptViewModel.Get());
 				}
 
 				// Rename all the bindings that point to the old parameter 
@@ -2153,7 +2153,7 @@ bool UNiagaraGraph::RenameParameterFromPin(const FNiagaraVariable& Parameter, FN
 		NewScriptVariable->Metadata.CreateNewGuid();
 		NewScriptVariable->SetIsSubscribedToParameterDefinitions(false);
 		VariableToScriptVariable.Add(NewParameter, NewScriptVariable);
-		ParameterPanelViewModel->SubscribeParameterToLibraryIfMatchingDefinition(NewScriptVariable, NewName);
+		FNiagaraParameterDefinitionsUtilities::TrySubscribeScriptVarToDefinitionByName(NewScriptVariable, ScriptViewModel.Get());
 
 		const FNiagaraGraphParameterReferenceCollection* ReferenceCollection = GetParameterReferenceMap().Find(Parameter);
 		if (ReferenceCollection && ReferenceCollection->ParameterReferences.Num() < 1)
diff --git a/Engine/Plugins/FX/Niagara/Source/NiagaraEditor/Private/Toolkits/NiagaraSystemToolkit.cpp b/Engine/Plugins/FX/Niagara/Source/NiagaraEditor/Private/Toolkits/NiagaraSystemToolkit.cpp
--- a/Engine/Plugins/FX/Niagara/Source/NiagaraEditor/Private/Toolkits/NiagaraSystemToolkit.cpp
+++ b/Engine/Plugins/FX/Niagara/Source/NiagaraEditor/Private/Toolkits/NiagaraSystemToolkit.cpp
@@ -143,10 +143,8 @@
 
 	SystemViewModel = MakeShared<FNiagaraSystemViewModel>();
 	SystemViewModel->Initialize(*System, SystemOptions);
-	SystemGraphSelectionViewModel = MakeShared<FNiagaraSystemGraphSelectionViewModel>();
-	SystemGraphSelectionViewModel->Init(SystemViewModel);
-	ParameterPanelViewModel = MakeShared<FNiagaraSystemToolkitParameterPanelViewModel>(SystemViewModel, TWeakPtr<FNiagaraSystemGraphSelectionViewModel>(SystemGraphSelectionViewModel));
-	ParameterDefinitionsPanelViewModel = MakeShared<FNiagaraSystemToolkitParameterDefinitionsPanelViewModel>(SystemViewModel, SystemGraphSelectionViewModel);
+	ParameterPanelViewModel = MakeShared<FNiagaraSystemToolkitParameterPanelViewModel>(SystemViewModel);
+	ParameterDefinitionsPanelViewModel = MakeShared<FNiagaraSystemToolkitParameterDefinitionsPanelViewModel>(SystemViewModel);
 	FSystemToolkitUIContext UIContext = FSystemToolkitUIContext(
 		FSimpleDelegate::CreateSP(ParameterPanelViewModel.ToSharedRef(), &INiagaraImmutableParameterPanelViewModel::Refresh),
 		FSimpleDelegate::CreateSP(ParameterDefinitionsPanelViewModel.ToSharedRef(), &INiagaraImmutableParameterPanelViewModel::Refresh)
@@ -822,6 +820,11 @@
 	return SystemViewModel;
 }
 
+TSharedPtr<FNiagaraSystemGraphSelectionViewModel> FNiagaraSystemToolkit::GetSystemGraphSelectionViewModel()
+{
+	return SystemViewModel->GetSystemGraphSelectionViewModel();
+}
+
 void FNiagaraSystemToolkit::RegisterToolbarTab(const TSharedRef<FTabManager>& InTabManager)
 {
 	FAssetEditorToolkit::RegisterTabSpawners(InTabManager);
diff --git a/Engine/Plugins/FX/Niagara/Source/NiagaraEditor/Private/Toolkits/NiagaraSystemToolkit.h b/Engine/Plugins/FX/Niagara/Source/NiagaraEditor/Private/Toolkits/NiagaraSystemToolkit.h
--- a/Engine/Plugins/FX/Niagara/Source/NiagaraEditor/Private/Toolkits/NiagaraSystemToolkit.h
+++ b/Engine/Plugins/FX/Niagara/Source/NiagaraEditor/Private/Toolkits/NiagaraSystemToolkit.h
@@ -95,9 +95,7 @@ public:
 	void CompileSystem(bool bFullRebuild);
 
 	TSharedPtr<FNiagaraSystemViewModel> GetSystemViewModel();
-	TSharedPtr<FNiagaraSystemGraphSelectionViewModel> GetSystemGraphSelectionViewModel() {
-		return SystemGraphSelectionViewModel;
-	}
+	TSharedPtr<FNiagaraSystemGraphSelectionViewModel> GetSystemGraphSelectionViewModel();
 	
 	// @todo This is a hack for now until we reconcile the default toolbar with application modes [duplicated from counterpart in Blueprint Editor]
 	void RegisterToolbarTab(const TSharedRef<class FTabManager>& TabManager);
@@ -234,9 +232,6 @@ private:
 	/* The view model for the System being edited */
 	TSharedPtr<FNiagaraSystemViewModel> SystemViewModel;
 	
-	/* The view model for the selected Emitter Script graphs of the System being edited. */
-	TSharedPtr<FNiagaraSystemGraphSelectionViewModel> SystemGraphSelectionViewModel;
-
 	/** Message log, with the log listing that it reflects */
 	TSharedPtr<FNiagaraMessageLogViewModel> NiagaraMessageLogViewModel;
 	TSharedPtr<class SWidget> NiagaraMessageLog;
diff --git a/Engine/Plugins/FX/Niagara/Source/NiagaraEditor/Private/ViewModels/NiagaraParameterDefinitionsPanelViewModel.cpp b/Engine/Plugins/FX/Niagara/Source/NiagaraEditor/Private/ViewModels/NiagaraParameterDefinitionsPanelViewModel.cpp
--- a/Engine/Plugins/FX/Niagara/Source/NiagaraEditor/Private/ViewModels/NiagaraParameterDefinitionsPanelViewModel.cpp
+++ b/Engine/Plugins/FX/Niagara/Source/NiagaraEditor/Private/ViewModels/NiagaraParameterDefinitionsPanelViewModel.cpp
@@ -10,6 +10,7 @@
 #include "NiagaraScriptVariable.h"
 #include "ScopedTransaction.h"
 #include "ViewModels/NiagaraEmitterHandleViewModel.h"
+#include "ViewModels/NiagaraScriptGraphViewModel.h"
 #include "ViewModels/NiagaraStandaloneScriptViewModel.h"
 #include "ViewModels/NiagaraSystemViewModel.h"
 
@@ -84,7 +85,12 @@ void FNiagaraScriptToolkitParameterDefinitionsPanelViewModel::Init(const FScript
 
 const TArray<UNiagaraGraph*> FNiagaraScriptToolkitParameterDefinitionsPanelViewModel::GetEditableGraphsConst() const
 {
-	return FNiagaraScriptToolkitParameterPanelUtilities::GetEditableGraphs(ScriptViewModel);
+	TArray<UNiagaraGraph*> EditableGraphs;
+	if (UNiagaraGraph* Graph = ScriptViewModel->GetGraphViewModel()->GetGraph())
+	{
+		EditableGraphs.Add(Graph);
+	}
+	return EditableGraphs;
 }
 
 const TArray<UNiagaraParameterDefinitions*> FNiagaraScriptToolkitParameterDefinitionsPanelViewModel::GetAvailableParameterDefinitionsAssets(bool bSkipSubscribedParameterDefinitions) const
@@ -203,16 +209,12 @@ FReply FNiagaraScriptToolkitParameterDefinitionsPanelViewModel::OnParameterItems
 /// System Toolkit Parameter Definitions Panel View Model					///
 ///////////////////////////////////////////////////////////////////////////////
 
-FNiagaraSystemToolkitParameterDefinitionsPanelViewModel::FNiagaraSystemToolkitParameterDefinitionsPanelViewModel(const TSharedPtr<FNiagaraSystemViewModel>& InSystemViewModel, const TWeakPtr<FNiagaraSystemGraphSelectionViewModel>& InSystemGraphSelectionViewModelWeak)
+FNiagaraSystemToolkitParameterDefinitionsPanelViewModel::FNiagaraSystemToolkitParameterDefinitionsPanelViewModel(const TSharedPtr<FNiagaraSystemViewModel>& InSystemViewModel)
 {
 	SystemViewModel = InSystemViewModel;
-	SystemGraphSelectionViewModelWeak = InSystemGraphSelectionViewModelWeak;
+	SystemGraphSelectionViewModelWeak = SystemViewModel->GetSystemGraphSelectionViewModel();
 }
 
-FNiagaraSystemToolkitParameterDefinitionsPanelViewModel::FNiagaraSystemToolkitParameterDefinitionsPanelViewModel(const TSharedPtr<FNiagaraSystemViewModel>& InSystemViewModel)
-	: FNiagaraSystemToolkitParameterDefinitionsPanelViewModel(InSystemViewModel, nullptr)
-{}
-
 void FNiagaraSystemToolkitParameterDefinitionsPanelViewModel::Cleanup()
 {
 	SystemViewModel->GetOnSubscribedParameterDefinitionsChangedDelegate().RemoveAll(this);
@@ -227,7 +229,7 @@ void FNiagaraSystemToolkitParameterDefinitionsPanelViewModel::Init(const FSystem
 
 const TArray<UNiagaraGraph*> FNiagaraSystemToolkitParameterDefinitionsPanelViewModel::GetEditableGraphsConst() const
 {
-	return FNiagaraSystemToolkitParameterPanelUtilities::GetEditableGraphs(SystemViewModel, SystemGraphSelectionViewModelWeak, false);
+	return SystemViewModel->GetSelectedGraphs();
 }
 
 const TArray<UNiagaraParameterDefinitions*> FNiagaraSystemToolkitParameterDefinitionsPanelViewModel::GetAvailableParameterDefinitionsAssets(bool bSkipSubscribedParameterDefinitions) const
diff --git a/Engine/Plugins/FX/Niagara/Source/NiagaraEditor/Private/ViewModels/NiagaraParameterPanelViewModel.cpp b/Engine/Plugins/FX/Niagara/Source/NiagaraEditor/Private/ViewModels/NiagaraParameterPanelViewModel.cpp
--- a/Engine/Plugins/FX/Niagara/Source/NiagaraEditor/Private/ViewModels/NiagaraParameterPanelViewModel.cpp
+++ b/Engine/Plugins/FX/Niagara/Source/NiagaraEditor/Private/ViewModels/NiagaraParameterPanelViewModel.cpp
@@ -19,28 +19,30 @@
 #include "NiagaraObjectSelection.h"
 #include "NiagaraParameterDefinitions.h"
 #include "NiagaraScriptGraphViewModel.h"
+#include "NiagaraScriptVariable.h"
 #include "NiagaraSimulationStageBase.h"
 #include "NiagaraSystem.h"
+#include "NiagaraSystemScriptViewModel.h"
 #include "NiagaraTypes.h"
 #include "ViewModels/NiagaraEmitterHandleViewModel.h"
+#include "ViewModels/NiagaraScratchPadScriptViewModel.h"
 #include "ViewModels/NiagaraScriptViewModel.h"
+#include "ViewModels/NiagaraSystemEditorDocumentsViewModel.h"
+#include "ViewModels/NiagaraSystemGraphSelectionViewModel.h"
 #include "ViewModels/NiagaraSystemSelectionViewModel.h"
 #include "ViewModels/NiagaraSystemViewModel.h"
-#include "ViewModels/NiagaraSystemGraphSelectionViewModel.h"
-#include "NiagaraSystemScriptViewModel.h"
 #include "ViewModels/Stack/NiagaraStackGraphUtilities.h"
+#include "ViewModels/Stack/NiagaraStackViewModel.h"
+#include "Widgets/SNiagaraDebugger.h"
 #include "Widgets/SNiagaraParameterMenu.h"
 #include "Widgets/SNiagaraParameterPanel.h"
-#include "ViewModels/Stack/NiagaraStackViewModel.h"
-#include "ViewModels/NiagaraScratchPadScriptViewModel.h"
-#include "ViewModels/NiagaraSystemEditorDocumentsViewModel.h"
 
 #include "Framework/Commands/GenericCommands.h"
 #include "Framework/MultiBox/MultiBoxBuilder.h"
 #include "HAL/PlatformApplicationMisc.h"
 #include "ScopedTransaction.h"
 #include "Widgets/SNullWidget.h"
-#include "Widgets/SNiagaraDebugger.h"
+
 
 #define LOCTEXT_NAMESPACE "NiagaraParameterPanelViewModel"
 
@@ -64,34 +66,6 @@ TArray<FNiagaraParameterPanelCategory> FNiagaraParameterDefinitionsToolkitParame
 /// System Toolkit Parameter Panel Utilities								///
 ///////////////////////////////////////////////////////////////////////////////
 
-TArray<UNiagaraGraph*> FNiagaraSystemToolkitParameterPanelUtilities::GetAllGraphs(const TSharedPtr<FNiagaraSystemViewModel>& SystemViewModel, bool bActiveScriptMode)
-{
-	if (bActiveScriptMode)
-	{
-		TArray<UNiagaraGraph*> EdGraphs = SystemViewModel->GetDocumentViewModel()->GetAllGraphsForActiveScriptDocument();
-		return EdGraphs;
-	}
-	else
-	{
-		TArray<UNiagaraGraph*> EdGraphs = SystemViewModel->GetDocumentViewModel()->GetAllGraphsForPrimaryDocument();
-		return EdGraphs;
-	}
-}
-
-TArray<UNiagaraGraph*> FNiagaraSystemToolkitParameterPanelUtilities::GetEditableGraphs(const TSharedPtr<FNiagaraSystemViewModel>& SystemViewModel, const TWeakPtr<FNiagaraSystemGraphSelectionViewModel>& SystemGraphSelectionViewModelWeak, bool bActiveScriptMode)
-{
-	if (bActiveScriptMode)
-	{
-		TArray<UNiagaraGraph*> EdGraphs = SystemViewModel->GetDocumentViewModel()->GetEditableGraphsForActiveScriptDocument();
-		return EdGraphs;
-	}
-	else
-	{
-		TArray<UNiagaraGraph*> EdGraphs = SystemViewModel->GetDocumentViewModel()->GetEditableGraphsForPrimaryDocument();
-		return EdGraphs;
-	}
-}
-
 FReply FNiagaraSystemToolkitParameterPanelUtilities::CreateDragEventForParameterItem(const FNiagaraParameterPanelItemBase& DraggedItem, const FPointerEvent& MouseEvent, const TArray<FNiagaraGraphParameterReference>& GraphParameterReferencesForItem, const TSharedPtr<TArray<FName>>& ParametersWithNamespaceModifierRenamePending)
 {
 	const static FText TooltipFormat = LOCTEXT("Parameters", "Name: {0} \nType: {1}");
@@ -115,16 +89,6 @@ FReply FNiagaraSystemToolkitParameterPanelUtilities::CreateDragEventForParameter
 /// Script Toolkit Parameter Panel Utilities								///
 ///////////////////////////////////////////////////////////////////////////////
 
-TArray<UNiagaraGraph*> FNiagaraScriptToolkitParameterPanelUtilities::GetEditableGraphs(const TSharedPtr<FNiagaraScriptViewModel>& ScriptViewModel)
-{
-	TArray<UNiagaraGraph*> EditableGraphs;
-	if (UNiagaraGraph* Graph = ScriptViewModel->GetGraphViewModel()->GetGraph())
-	{
-		EditableGraphs.Add(Graph);
-	}
-	return EditableGraphs;
-}
-
 FReply FNiagaraScriptToolkitParameterPanelUtilities::CreateDragEventForParameterItem(const FNiagaraParameterPanelItemBase& DraggedItem, const FPointerEvent& MouseEvent, const TArray<FNiagaraGraphParameterReference>& GraphParameterReferencesForItem, const TSharedPtr<TArray<FName>>& ParametersWithNamespaceModifierRenamePending)
 {
 	const static FText TooltipFormat = LOCTEXT("Parameters", "Name: {0} \nType: {1}");
@@ -288,10 +252,8 @@ INiagaraParameterPanelViewModel::~INiagaraParameterPanelViewModel()
 		ScriptVar->RemoveFromRoot();
 		ScriptVar = nullptr;
 	}
-
 }
 
-
 void INiagaraParameterPanelViewModel::RefreshFull(bool bDoCategoryExpansion) const
 {
 	INiagaraImmutableParameterPanelViewModel::RefreshFull(bDoCategoryExpansion);
@@ -495,16 +457,6 @@ bool INiagaraParameterPanelViewModel::GetCanSubscribeParameterToLibraryAndToolTi
 	return false;
 }
 
-void INiagaraParameterPanelViewModel::SetParameterIsSubscribedToLibrary(const FNiagaraParameterPanelItem ItemToModify, const bool bSubscribed) 
-{
-	if (ensureMsgf(ItemToModify.bExternallyReferenced == false, TEXT("Cannot modify an externally referenced parameter.")) == false)
-	{
-		return;
-	}
-
-	SetParameterIsSubscribedToLibrary(ItemToModify.ScriptVariable, bSubscribed);
-}
-
 void INiagaraParameterPanelViewModel::SetParameterNamespace(const FNiagaraParameterPanelItem ItemToModify, FNiagaraNamespaceMetadata NewNamespaceMetaData, bool bDuplicateParameter) 
 {
 	FName NewName = FNiagaraParameterUtilities::ChangeNamespace(ItemToModify.GetVariable().GetName(), NewNamespaceMetaData);
@@ -815,47 +767,6 @@ void INiagaraParameterPanelViewModel::SelectParameterItemByName(const FName Para
 	}
 }
 
-void INiagaraParameterPanelViewModel::SubscribeParameterToLibraryIfMatchingDefinition(const UNiagaraScriptVariable* ScriptVarToModify, const FName ScriptVarName) 
-{
-	const TArray<const UNiagaraScriptVariable*> ReservedParametersForName = FNiagaraParameterDefinitionsUtilities::FindReservedParametersByName(ScriptVarName);
-
-	if (ReservedParametersForName.Num() == 0)
-	{
-		// Do not call SetParameterIsSubscribedToLibrary() unless ScriptVarToModify is already marked as SubscribedToParameterDefinitions.
-		if (ScriptVarToModify->GetIsSubscribedToParameterDefinitions())
-		{
-			SetParameterIsSubscribedToLibrary(ScriptVarToModify, false);
-		}
-	}
-	else if (ReservedParametersForName.Num() == 1)
-	{
-		const UNiagaraScriptVariable* ReservedParameterDefinition = (ReservedParametersForName)[0];
-		if (ReservedParameterDefinition->Variable.GetType() != ScriptVarToModify->Variable.GetType())
-		{
-			const FText TypeMismatchWarningTemplate = LOCTEXT("RenameParameter_DefinitionTypeMismatch", "Renamed parameter \"{0}\" with type {1}. Type does not match existing parameter definition \"{0}\" with type {2} from {3}! ");
-			FText TypeMismatchWarning = FText::Format(
-				TypeMismatchWarningTemplate,
-				FText::FromName(ScriptVarName),
-				ScriptVarToModify->Variable.GetType().GetNameText(),
-				ReservedParameterDefinition->Variable.GetType().GetNameText(),
-				FText::FromString(ReservedParameterDefinition->GetTypedOuter<UNiagaraParameterDefinitions>()->GetName()));
-			FNiagaraEditorUtilities::WarnWithToastAndLog(TypeMismatchWarning);
-
-			// Do not call SetParameterIsSubscribedToLibrary() unless ScriptVarToModify is already marked as SubscribedToParameterDefinitions.
-			if (ScriptVarToModify->GetIsSubscribedToParameterDefinitions())
-			{
-				SetParameterIsSubscribedToLibrary(ScriptVarToModify, false);
-			}
-		}
-		else
-		{
-			// NOTE: It is possible we are calling SetParameterIsSubscribedToLibrary() on a UNiagaraScriptVariable that is already marked as SubscribedToParameterDefinitions;
-			// This is intended as SetParameterIsSubscribedToLibrary() will register a new link to a different parameter definition.
-			SetParameterIsSubscribedToLibrary(ScriptVarToModify, true);
-		}
-	}
-}
-
 bool INiagaraParameterPanelViewModel::CanMakeNewParameterOfType(const FNiagaraTypeDefinition& InType)
 {
 	return InType != FNiagaraTypeDefinition::GetParameterMapDef()
@@ -868,17 +779,13 @@ bool INiagaraParameterPanelViewModel::CanMakeNewParameterOfType(const FNiagaraTy
 /// System Toolkit Parameter Panel View Model								///
 ///////////////////////////////////////////////////////////////////////////////
 
-FNiagaraSystemToolkitParameterPanelViewModel::FNiagaraSystemToolkitParameterPanelViewModel(const TSharedPtr<FNiagaraSystemViewModel>& InSystemViewModel, const TWeakPtr<FNiagaraSystemGraphSelectionViewModel>& InSystemGraphSelectionViewModelWeak)
+FNiagaraSystemToolkitParameterPanelViewModel::FNiagaraSystemToolkitParameterPanelViewModel(const TSharedPtr<FNiagaraSystemViewModel>& InSystemViewModel/*, const TWeakPtr<FNiagaraSystemGraphSelectionViewModel>& InSystemGraphSelectionViewModelWeak*/)
 {
 	SystemViewModel = InSystemViewModel;
-	SystemGraphSelectionViewModelWeak = InSystemGraphSelectionViewModelWeak;
+	SystemGraphSelectionViewModelWeak = InSystemViewModel->GetSystemGraphSelectionViewModel();
 	SystemScriptGraph = SystemViewModel->GetSystemScriptViewModel()->GetGraphViewModel()->GetGraph();
 }
 
-FNiagaraSystemToolkitParameterPanelViewModel::FNiagaraSystemToolkitParameterPanelViewModel(const TSharedPtr<FNiagaraSystemViewModel>& InSystemViewModel)
-	: FNiagaraSystemToolkitParameterPanelViewModel(InSystemViewModel, nullptr)
-{}
-
 void FNiagaraSystemToolkitParameterPanelViewModel::Cleanup()
 {
 	UnregisterViewModelWithMap(RegisteredHandle);
@@ -1254,7 +1161,7 @@ void FNiagaraSystemToolkitParameterPanelViewModel::AddParameter(FNiagaraVariable
 				bSuccess = true;
 
 				// Check if the new parameter has the same name and type as an existing parameter definition, and if so, link to the definition automatically.
-				SubscribeParameterToLibraryIfMatchingDefinition(NewScriptVar, NewScriptVar->Variable.GetName());
+				FNiagaraParameterDefinitionsUtilities::TrySubscribeScriptVarToDefinitionByName(NewScriptVar, SystemViewModel.Get());
 			}
 		}
 	}
@@ -1434,116 +1341,29 @@ void FNiagaraSystemToolkitParameterPanelViewModel::RenameParameter(const FNiagar
 		return;
 	}
 
+	bool bRenameSuccess = false;
 	FScopedTransaction RenameParameterTransaction(LOCTEXT("RenameParameter", "Rename parameter"));
-	const FNiagaraVariable Parameter = ItemToRename.GetVariable();
-	const FGuid& ScriptVarId = ItemToRename.GetVariableMetaData().GetVariableGuid();
-	UNiagaraSystem& System = SystemViewModel->GetSystem();
 
-	// Forward renaming of script variables for the active script to the working parameter panel view model for just that script.
+	// Forward renaming of script variables to the active script if it is selected.
 	TSharedPtr<FNiagaraScratchPadScriptViewModel> ScratchDocumentVM = SystemViewModel->GetDocumentViewModel()->GetActiveScratchPadViewModelIfSet();
 	if (ScratchDocumentVM.IsValid())
 	{
-		ScratchDocumentVM->GetParameterPanelViewModel()->RenameParameter(ItemToRename, NewName);
-		return;
-	}
-
-	System.Modify();
-
-	bool bExposedParametersRename = false;
-	bool bEditorOnlyParametersRename = false;
-	bool bAssignmentNodeRename = false;
-
-	// Rename the parameter if it's in the parameter store for user parameters.
-	if (System.GetExposedParameters().IndexOf(Parameter) != INDEX_NONE)
-	{
-		FNiagaraParameterStore* ExposedParametersStore = &System.GetExposedParameters();
-		TArray<FNiagaraVariable> OwningParameters;
-		ExposedParametersStore->GetParameters(OwningParameters);
-		if (OwningParameters.ContainsByPredicate([NewName](const FNiagaraVariable& Variable) { return Variable.GetName() == NewName; }))
-		{
-			// If the parameter store already has a parameter with this name, remove the old parameter to prevent collisions.
-			ExposedParametersStore->RemoveParameter(Parameter);
-		}
-		else
-		{
-			// Otherwise it's safe to rename.
-			ExposedParametersStore->RenameParameter(Parameter, NewName);
-		}
-		bExposedParametersRename = true;
-	}
-
-	// Look for set parameters nodes or linked inputs which reference this parameter and rename if so.
-	for (const FNiagaraGraphParameterReference& ParameterReference : GetGraphParameterReferencesForItem(ItemToRename))
-	{
-		UNiagaraNode* ReferenceNode = Cast<UNiagaraNode>(ParameterReference.Value);
-		if (ReferenceNode != nullptr)
-		{
-			UNiagaraNodeAssignment* OwningAssignmentNode = ReferenceNode->GetTypedOuter<UNiagaraNodeAssignment>();
-			if (OwningAssignmentNode != nullptr)
-			{
-				// If this is owned by a set variables node and it's not locked, update the assignment target on the assignment node.
-				bAssignmentNodeRename |= FNiagaraStackGraphUtilities::TryRenameAssignmentTarget(*OwningAssignmentNode, Parameter, NewName);
-			}
-			else
-			{
-				// Otherwise if the reference node is a get node it's for a linked input so we can just update pin name.
-				UNiagaraNodeParameterMapGet* ReferenceGetNode = Cast<UNiagaraNodeParameterMapGet>(ReferenceNode);
-				if (ReferenceGetNode != nullptr)
-				{
-					if (ReferenceGetNode->Pins.ContainsByPredicate([&ParameterReference](UEdGraphPin* Pin) { return Pin->PersistentGuid == ParameterReference.Key; }))
-					{
-						ReferenceGetNode->GetNiagaraGraph()->RenameParameter(Parameter, NewName, true);
-					}
-				}
-			}
-		}
+		bRenameSuccess = ScratchDocumentVM->RenameParameter(ItemToRename.GetVariable(), NewName);
 	}
-
-	// Rename the parameter if it is owned directly as an editor only parameter.
-	UNiagaraEditorParametersAdapter* EditorParametersAdapter = SystemViewModel->GetEditorOnlyParametersAdapter();
-	if (UNiagaraScriptVariable** ScriptVariablePtr = EditorParametersAdapter->GetParameters().FindByPredicate([&ScriptVarId](const UNiagaraScriptVariable* ScriptVariable) { return ScriptVariable->Metadata.GetVariableGuid() == ScriptVarId; }))
+	// Otherwise rename via the system.
+	else
 	{
-		EditorParametersAdapter->Modify();
-		UNiagaraScriptVariable* ScriptVariable = *ScriptVariablePtr;
-		ScriptVariable->Modify();
-		ScriptVariable->Variable.SetName(NewName);
-		ScriptVariable->UpdateChangeId();
-		bEditorOnlyParametersRename = true;
+		bRenameSuccess = SystemViewModel->RenameParameter(ItemToRename.GetVariable(), NewName, ENiagaraGetGraphParameterReferencesMode::SelectedGraphs);
 	}
-
-	// Check if the rename will give the same name and type as an existing parameter definition, and if so, link to the definition automatically.
-	SubscribeParameterToLibraryIfMatchingDefinition(ItemToRename.ScriptVariable, NewName);
-
-	// Handle renaming any renderer properties that might match.
-	if (bExposedParametersRename | bEditorOnlyParametersRename | bAssignmentNodeRename)
+	
+	if (bRenameSuccess)
 	{
-		if (Parameter.IsInNameSpace(FNiagaraConstants::ParticleAttributeNamespaceString) || Parameter.IsInNameSpace(FNiagaraConstants::EmitterNamespaceString))
-		{
-			for (const FNiagaraEmitterHandle* EmitterHandle : GetEditableEmitterHandles())
-			{
-				EmitterHandle->GetInstance().Emitter->HandleVariableRenamed(Parameter, FNiagaraVariableBase(Parameter.GetType(), NewName), true, EmitterHandle->GetInstance().Version);
-			}		
-		}
-		else
-		{
-			System.HandleVariableRenamed(Parameter, FNiagaraVariableBase(Parameter.GetType(), NewName), true);
-		}
-
 		Refresh();
 		const bool bRequestRename = false;
 		SelectParameterItemByName(NewName, bRequestRename);
 	}
 }
 
-void FNiagaraSystemToolkitParameterPanelViewModel::SetParameterIsSubscribedToLibrary(const UNiagaraScriptVariable* ScriptVarToModify, const bool bSubscribed) 
-{
-	const FText TransactionText = bSubscribed ? LOCTEXT("SubscribeParameter", "Subscribe parameter") : LOCTEXT("UnsubscribeParameter", "Unsubscribe parameter");
-	FScopedTransaction SubscribeTransaction(TransactionText);
-	SystemViewModel->SetParameterIsSubscribedToDefinitions(ScriptVarToModify->Metadata.GetVariableGuid(), bSubscribed);
-	Refresh();
-	UIContext.RefreshParameterDefinitionsPanel();
-}
-
 FReply FNiagaraSystemToolkitParameterPanelViewModel::OnParameterItemsDragged(const TArray<FNiagaraParameterPanelItem>& DraggedItems, const FPointerEvent& MouseEvent) const
 {
 	if (OnGetParametersWithNamespaceModifierRenamePendingDelegate.IsBound() == false)
@@ -1971,92 +1791,9 @@ const TArray<FNiagaraGraphParameterReference> FNiagaraSystemToolkitParameterPane
 {
 	const FNiagaraParameterPanelItem& ParameterPanelItem = static_cast<const FNiagaraParameterPanelItem&>(Item);
 	bool bGetReferencesAcrossAllGraphs = ParameterPanelItem.bExternallyReferenced == false;
-	const TArray<UNiagaraGraph*> TargetGraphs = bGetReferencesAcrossAllGraphs ? GetAllGraphsConst() : GetEditableGraphsConst();
-
-	// -For each selected graph perform a parameter map history traversal and collect all graph parameter references associated with the target FNiagaraParameterPanelItem.
-	TArray<FNiagaraGraphParameterReference> GraphParameterReferences;
-	for (const UNiagaraGraph* Graph : TargetGraphs)
-	{
-		TArray<UNiagaraNodeOutput*> OutputNodes;
-		Graph->GetNodesOfClass<UNiagaraNodeOutput>(OutputNodes);
-		for (UNiagaraNodeOutput* OutputNode : OutputNodes)
-		{
-			UNiagaraNode* NodeToTraverse = OutputNode;
-			if (OutputNode->GetUsage() == ENiagaraScriptUsage::SystemSpawnScript || OutputNode->GetUsage() == ENiagaraScriptUsage::SystemUpdateScript)
-			{
-				// Traverse past the emitter nodes, otherwise the system scripts will pick up all of the emitter and particle script parameters.
-				UEdGraphPin* InputPin = FNiagaraStackGraphUtilities::GetParameterMapInputPin(*NodeToTraverse);
-				while (NodeToTraverse != nullptr && InputPin != nullptr && InputPin->LinkedTo.Num() == 1 &&
-					(NodeToTraverse->IsA<UNiagaraNodeOutput>() || NodeToTraverse->IsA<UNiagaraNodeEmitter>()))
-				{
-					NodeToTraverse = Cast<UNiagaraNode>(InputPin->LinkedTo[0]->GetOwningNode());
-					InputPin = NodeToTraverse != nullptr ? FNiagaraStackGraphUtilities::GetParameterMapInputPin(*NodeToTraverse) : nullptr;
-				}
-			}
-
-			if (NodeToTraverse == nullptr)
-			{
-				continue;
-			}
-
-			bool bIgnoreDisabled = true;
-			FNiagaraParameterMapHistoryBuilder Builder;
-			FVersionedNiagaraEmitter GraphOwningEmitter = Graph->GetOwningEmitter();
-			FCompileConstantResolver ConstantResolver = GraphOwningEmitter.Emitter != nullptr
-				? FCompileConstantResolver(GraphOwningEmitter, ENiagaraScriptUsage::Function)
-				: FCompileConstantResolver();
-
-			Builder.SetIgnoreDisabled(bIgnoreDisabled);
-			Builder.ConstantResolver = ConstantResolver;
-			FName StageName;
-			ENiagaraScriptUsage StageUsage = OutputNode->GetUsage();
-			if (StageUsage == ENiagaraScriptUsage::ParticleSimulationStageScript && GraphOwningEmitter.Emitter)
-			{
-				UNiagaraSimulationStageBase* Base = GraphOwningEmitter.GetEmitterData()->GetSimulationStageById(OutputNode->GetUsageId());
-				if (Base)
-				{
-					StageName = Base->GetStackContextReplacementName();
-				}
-			}
-			Builder.BeginUsage(StageUsage, StageName);
-			NodeToTraverse->BuildParameterMapHistory(Builder, true, false);
-			Builder.EndUsage();
-
-			if (Builder.Histories.Num() != 1)
-			{
-				// We should only have traversed one emitter (have not visited more than one NiagaraNodeEmitter.)
-				ensureMsgf(false, TEXT("Encountered '%d' parameter map history when collecting parameters for system parameter panel view model, we expect only 1!"), Builder.Histories.Num());
-			}
-			if (Builder.Histories.Num() == 0)
-			{
-				continue;
-			}
-
-			const TArray<FName>& CustomIterationSourceNamespaces = Builder.Histories[0].IterationNamespaceOverridesEncountered;
-			for (int32 VariableIndex = 0; VariableIndex < Builder.Histories[0].Variables.Num(); VariableIndex++)
-			{
-				if (Item.GetVariable() == Builder.Histories[0].Variables[VariableIndex])
-				{
-					for (const FNiagaraParameterMapHistory::FReadHistory& ReadHistory : Builder.Histories[0].PerVariableReadHistory[VariableIndex])
-					{
-						if (ReadHistory.ReadPin.Pin->GetOwningNode() != nullptr)
-						{
-							GraphParameterReferences.Add(FNiagaraGraphParameterReference(ReadHistory.ReadPin.Pin->PersistentGuid, ReadHistory.ReadPin.Pin->GetOwningNode()));
-						}
-					}
 
-					for (const FModuleScopedPin& Write : Builder.Histories[0].PerVariableWriteHistory[VariableIndex])
-					{
-						if (Write.Pin->GetOwningNode() != nullptr)
-						{
-							GraphParameterReferences.Add(FNiagaraGraphParameterReference(Write.Pin->PersistentGuid, Write.Pin->GetOwningNode()));
-						}
-					}
-				}
-			}
-		}
-	}
-	return GraphParameterReferences;
+	ENiagaraGetGraphParameterReferencesMode GetGraphParameterReferencesMode = bGetReferencesAcrossAllGraphs ? ENiagaraGetGraphParameterReferencesMode::AllGraphs : ENiagaraGetGraphParameterReferencesMode::SelectedGraphs;
+	return SystemViewModel->GetGraphParameterReferences(Item.GetVariable(), GetGraphParameterReferencesMode);
 }
 
 const TArray<UNiagaraParameterDefinitions*> FNiagaraSystemToolkitParameterPanelViewModel::GetAvailableParameterDefinitions(bool bSkipSubscribedParameterDefinitions) const
@@ -2105,29 +1842,14 @@ bool FNiagaraSystemToolkitParameterPanelViewModel::IsCategoryExpandedByDefault(c
 
 const TArray<UNiagaraGraph*> FNiagaraSystemToolkitParameterPanelViewModel::GetAllGraphsConst() const
 {
-	// If an active script, just get the graphs associated with that. Otherwise for system overview, get all graphs.
-	if (ActiveSectionIndex == ActiveScriptIdx && ActiveScriptIdx != -1)
-	{
-		return FNiagaraSystemToolkitParameterPanelUtilities::GetAllGraphs(SystemViewModel, true);
-	}
-	else
-	{
-		return FNiagaraSystemToolkitParameterPanelUtilities::GetAllGraphs(SystemViewModel, false);
-	}
+	return SystemViewModel->GetAllGraphs();
 }
 
 TArray<UNiagaraGraph*> FNiagaraSystemToolkitParameterPanelViewModel::GetEditableGraphs() const
 {
-	// If an active script, just get the graphs associated with that. Otherwise for system overview, get all editable graphs.
-	if (ActiveSectionIndex == ActiveScriptIdx && ActiveScriptIdx != -1)
-	{
-		return FNiagaraSystemToolkitParameterPanelUtilities::GetEditableGraphs(SystemViewModel, SystemGraphSelectionViewModelWeak, true);
-	}
-	else
-	{
-		return FNiagaraSystemToolkitParameterPanelUtilities::GetEditableGraphs(SystemViewModel, SystemGraphSelectionViewModelWeak, false);
-	}
+	return SystemViewModel->GetSelectedGraphs();
 }
+
 void INiagaraParameterPanelViewModel::SetActiveSection(int32 InSection)
 {
 	int32 OldActiveSectionIndex = ActiveSectionIndex;
@@ -2860,11 +2582,11 @@ void FNiagaraScriptToolkitParameterPanelViewModel::AddParameter(FNiagaraVariable
 	for (UNiagaraGraph* Graph : GetEditableGraphs())
 	{
 		Graph->Modify();
-		const UNiagaraScriptVariable* NewScriptVar = Graph->AddParameter(NewVariable);
+		UNiagaraScriptVariable* NewScriptVar = Graph->AddParameter(NewVariable);
 		bSuccess = true;
 
 		// Check if the new parameter has the same name and type as an existing parameter definition, and if so, link to the definition automatically.
-		SubscribeParameterToLibraryIfMatchingDefinition(NewScriptVar, NewScriptVar->Variable.GetName());
+		FNiagaraParameterDefinitionsUtilities::TrySubscribeScriptVarToDefinitionByName(NewScriptVar, ScriptViewModel.Get());
 	}
 
 	if (bSuccess)
@@ -2883,11 +2605,11 @@ void FNiagaraScriptToolkitParameterPanelViewModel::FindOrAddParameter(FNiagaraVa
 	for (UNiagaraGraph* Graph : GetEditableGraphs())
 	{
 		Graph->Modify();
-		const UNiagaraScriptVariable* NewScriptVar = Graph->AddParameter(Variable);
+		UNiagaraScriptVariable* NewScriptVar = Graph->AddParameter(Variable);
 		bSuccess = true;
 
 		// Check if the new parameter has the same name and type as an existing parameter definition, and if so, link to the definition automatically.
-		SubscribeParameterToLibraryIfMatchingDefinition(NewScriptVar, NewScriptVar->Variable.GetName());
+		FNiagaraParameterDefinitionsUtilities::TrySubscribeScriptVarToDefinitionByName(NewScriptVar, ScriptViewModel.Get());
 	}
 
 	if (bSuccess)
@@ -2984,20 +2706,8 @@ void FNiagaraScriptToolkitParameterPanelViewModel::ChangeParameterType(const TAr
 void FNiagaraScriptToolkitParameterPanelViewModel::RenameParameter(const UNiagaraScriptVariable* ScriptVarToRename, const FName NewName) 
 {
 	FScopedTransaction RenameTransaction(LOCTEXT("RenameParameterTransaction", "Rename parameter"));
-	ScriptViewModel->GetStandaloneScript().Script->Modify();
-	bool bSuccess = false;
-	const FNiagaraVariable& Parameter = ScriptVarToRename->Variable;
-	for (UNiagaraGraph* Graph : GetEditableGraphs())
-	{
-		// Removed some prior code here around parameter usage collections and earlying out as it was causing "unreferenced" parameters to be
-		// unabled to have their namespace changed for no good reason.
-		Graph->Modify();
-		Graph->RenameParameter(Parameter, NewName);
-		bSuccess = true;
-	}
 
-	// Check if the rename will give the same name and type as an existing parameter definition, and if so, link to the definition automatically.
-	SubscribeParameterToLibraryIfMatchingDefinition(ScriptVarToRename, NewName);
+	bool bSuccess = ScriptViewModel->RenameParameter(ScriptVarToRename->Variable, NewName);
 
 	if (bSuccess)
 	{
@@ -3070,17 +2780,6 @@ void FNiagaraScriptToolkitParameterPanelViewModel::DuplicateParameters(const TAr
 	}
 }
 
-void FNiagaraScriptToolkitParameterPanelViewModel::SetParameterIsSubscribedToLibrary(const UNiagaraScriptVariable* ScriptVarToModify, const bool bSubscribed) 
-{
-	const FText TransactionText = bSubscribed ? LOCTEXT("SubscribeParameter", "Subscribe parameter") : LOCTEXT("UnsubscribeParameter", "Unsubscribe parameter");
-	FScopedTransaction SubscribeTransaction(TransactionText);
-	ScriptViewModel->GetStandaloneScript().Script->Modify();
-	ScriptViewModel->SetParameterIsSubscribedToDefinitions(ScriptVarToModify->Metadata.GetVariableGuid(), bSubscribed);
-	Refresh();
-	UIContext.RefreshParameterDefinitionsPanel();
-	UIContext.RefreshSelectionDetailsViewPanel();
-}
-
 FReply FNiagaraScriptToolkitParameterPanelViewModel::OnParameterItemsDragged(const TArray<FNiagaraParameterPanelItem>& DraggedItems, const FPointerEvent& MouseEvent) const
 {
 	if (OnGetParametersWithNamespaceModifierRenamePendingDelegate.IsBound() == false)
@@ -3386,7 +3085,12 @@ void FNiagaraScriptToolkitParameterPanelViewModel::SetParameterIsOverridingLibra
 
 TArray<UNiagaraGraph*> FNiagaraScriptToolkitParameterPanelViewModel::GetEditableGraphs() const
 {
-	return FNiagaraScriptToolkitParameterPanelUtilities::GetEditableGraphs(ScriptViewModel);
+	TArray<UNiagaraGraph*> EditableGraphs;
+	if (UNiagaraGraph* Graph = ScriptViewModel->GetGraphViewModel()->GetGraph())
+	{
+		EditableGraphs.Add(Graph);
+	}
+	return EditableGraphs;
 }
 
 const TArray<UNiagaraScriptVariable*> FNiagaraScriptToolkitParameterPanelViewModel::GetEditableScriptVariablesWithName(const FName ParameterName) const
@@ -3667,12 +3371,6 @@ void FNiagaraParameterDefinitionsToolkitParameterPanelViewModel::RenameParameter
 	ParameterDefinitionsWeak.Get()->RenameParameter(ItemToRename.GetVariable(), NewName);
 }
 
-void FNiagaraParameterDefinitionsToolkitParameterPanelViewModel::SetParameterIsSubscribedToLibrary(const UNiagaraScriptVariable* ScriptVarToModify, const bool bSubscribed) 
-{
-	// Do nothing, parameter definitions parameters are always subscribed to their parent parameter definitions.
-	ensureMsgf(false, TEXT("Tried to set a parameter definitions defined parameter subscribing! This should not be reachable."));
-}
-
 TSharedPtr<SWidget> FNiagaraParameterDefinitionsToolkitParameterPanelViewModel::CreateContextMenuForItems(const TArray<FNiagaraParameterPanelItem>& Items, const TSharedPtr<FUICommandList>& ToolkitCommands)
 {
 	// Only create context menus when a single item is selected.
diff --git a/Engine/Plugins/FX/Niagara/Source/NiagaraEditor/Private/ViewModels/NiagaraScriptViewModel.cpp b/Engine/Plugins/FX/Niagara/Source/NiagaraEditor/Private/ViewModels/NiagaraScriptViewModel.cpp
--- a/Engine/Plugins/FX/Niagara/Source/NiagaraEditor/Private/ViewModels/NiagaraScriptViewModel.cpp
+++ b/Engine/Plugins/FX/Niagara/Source/NiagaraEditor/Private/ViewModels/NiagaraScriptViewModel.cpp
@@ -11,6 +11,7 @@
 #include "NiagaraScriptInputCollectionViewModel.h"
 #include "NiagaraScriptOutputCollectionViewModel.h"
 #include "NiagaraScriptSource.h"
+#include "NiagaraScriptVariable.h"
 #include "ViewModels/TNiagaraViewModelManager.h"
 
 template<> TMap<UNiagaraScript*, TArray<FNiagaraScriptViewModel*>> TNiagaraViewModelManager<UNiagaraScript, FNiagaraScriptViewModel>::ObjectsToViewModels{};
@@ -312,6 +313,23 @@ FVersionedNiagaraScript FNiagaraScriptViewModel::GetStandaloneScript()
 	return FVersionedNiagaraScript();
 }
 
+bool FNiagaraScriptViewModel::RenameParameter(const FNiagaraVariable TargetParameter, const FName NewName)
+{
+	GetStandaloneScript().Script->Modify();
+	UNiagaraGraph* Graph = GetGraphViewModel()->GetGraph();
+	
+	Graph->Modify();
+	if (Graph->RenameParameter(TargetParameter, NewName))
+	{
+		UNiagaraScriptVariable* RenamedScriptVar = Graph->GetScriptVariable(NewName);
+
+		// Check if the rename will give the same name and type as an existing parameter definition, and if so, link to the definition automatically.
+		FNiagaraParameterDefinitionsUtilities::TrySubscribeScriptVarToDefinitionByName(RenamedScriptVar, this);
+	}
+
+	return true;
+}
+
 void FNiagaraScriptViewModel::UpdateCompileStatus(ENiagaraScriptCompileStatus InAggregateCompileStatus, const FString& InAggregateCompileErrorString,
 	const TArray<ENiagaraScriptCompileStatus>& InCompileStatuses, const TArray<FString>& InCompileErrors, const TArray<FString>& InCompilePaths,
 	const TArray<UNiagaraScript*>& InCompileSources)
diff --git a/Engine/Plugins/FX/Niagara/Source/NiagaraEditor/Private/ViewModels/NiagaraSystemGraphSelectionViewModel.cpp b/Engine/Plugins/FX/Niagara/Source/NiagaraEditor/Private/ViewModels/NiagaraSystemGraphSelectionViewModel.cpp
--- a/Engine/Plugins/FX/Niagara/Source/NiagaraEditor/Private/ViewModels/NiagaraSystemGraphSelectionViewModel.cpp
+++ b/Engine/Plugins/FX/Niagara/Source/NiagaraEditor/Private/ViewModels/NiagaraSystemGraphSelectionViewModel.cpp
@@ -12,10 +12,9 @@
 #include "ViewModels/NiagaraSystemSelectionViewModel.h"
 
 
-void FNiagaraSystemGraphSelectionViewModel::Init(TSharedPtr<FNiagaraSystemViewModel> InSystemViewModel)
+void FNiagaraSystemGraphSelectionViewModel::Initialize(TSharedRef<FNiagaraSystemViewModel> InSystemViewModel)
 {
 	SystemViewModelWeak = InSystemViewModel;
-
 	InSystemViewModel->GetSelectionViewModel()->OnEmitterHandleIdSelectionChanged().AddSP(this, &FNiagaraSystemGraphSelectionViewModel::RefreshSelectedEmitterScriptGraphs);
 	InSystemViewModel->OnEmitterHandleViewModelsChanged().AddSP(this, &FNiagaraSystemGraphSelectionViewModel::RefreshSelectedEmitterScriptGraphs);
 	RefreshSelectedEmitterScriptGraphs();
@@ -45,11 +44,10 @@ void FNiagaraSystemGraphSelectionViewModel::RefreshSelectedEmitterScriptGraphs()
 TSharedRef<FNiagaraSystemViewModel> FNiagaraSystemGraphSelectionViewModel::GetSystemViewModel()
 {
 	TSharedPtr<FNiagaraSystemViewModel> SystemViewModel = SystemViewModelWeak.Pin();
-	checkf(SystemViewModel.IsValid(), TEXT("Owning system view model destroyed before system overview view model."));
+	checkf(SystemViewModel.IsValid(), TEXT("Owning system viewmodel destroyed before system graph selection viewmodel."));
 	return SystemViewModel.ToSharedRef();
 }
 
-
 UNiagaraSystemSelectionViewModel* FNiagaraSystemGraphSelectionViewModel::GetSystemSelectionViewModel()
 {
 	return GetSystemViewModel()->GetSelectionViewModel();
diff --git a/Engine/Plugins/FX/Niagara/Source/NiagaraEditor/Private/ViewModels/NiagaraSystemViewModel.cpp b/Engine/Plugins/FX/Niagara/Source/NiagaraEditor/Private/ViewModels/NiagaraSystemViewModel.cpp
--- a/Engine/Plugins/FX/Niagara/Source/NiagaraEditor/Private/ViewModels/NiagaraSystemViewModel.cpp
+++ b/Engine/Plugins/FX/Niagara/Source/NiagaraEditor/Private/ViewModels/NiagaraSystemViewModel.cpp
@@ -1,62 +1,67 @@
 // Copyright Epic Games, Inc. All Rights Reserved.
 
 #include "ViewModels/NiagaraSystemViewModel.h"
-#include "NiagaraSystem.h"
-#include "NiagaraEmitterHandle.h"
-#include "NiagaraEmitter.h"
-#include "NiagaraEmitterEditorData.h"
-#include "NiagaraScriptSource.h"
-#include "NiagaraSimulationStageBase.h"
-#include "NiagaraEditorData.h"
-#include "NiagaraEditorUtilities.h"
-#include "NiagaraDataInterfaceCurveBase.h"
-#include "ViewModels/NiagaraEmitterHandleViewModel.h"
-#include "ViewModels/NiagaraEmitterViewModel.h"
-#include "ViewModels/NiagaraSystemSelectionViewModel.h"
-#include "ViewModels/NiagaraScratchPadViewModel.h"
-#include "ViewModels/NiagaraCurveSelectionViewModel.h"
-#include "ViewModels/Stack/NiagaraStackViewModel.h"
-#include "ViewModels/Stack/NiagaraStackGraphUtilities.h"
-#include "ViewModels/Stack/NiagaraStackEntry.h"
-#include "ViewModels/NiagaraPlaceholderDataInterfaceManager.h"
-#include "ViewModels/NiagaraSystemEditorDocumentsViewModel.h"
-#include "NiagaraSystemScriptViewModel.h"
-#include "NiagaraScriptGraphViewModel.h"
-#include "NiagaraSequence.h"
-#include "MovieSceneNiagaraEmitterTrack.h"
-#include "NiagaraSystemInstance.h"
-#include "NiagaraEditorSettings.h"
-#include "NiagaraSystemEditorData.h"
-#include "NiagaraGraph.h"
-#include "NiagaraNodeInput.h"
-#include "NiagaraNodeOutput.h"
-#include "NiagaraOverviewNode.h"
-#include "NiagaraEditorModule.h"
-#include "NiagaraNodeFunctionCall.h"
+
+#include "Channels/MovieSceneChannelProxy.h"
 #include "EdGraphSchema_NiagaraSystemOverview.h"
 #include "EdGraphUtilities.h"
-#include "ViewModels/NiagaraScratchPadUtilities.h"
-#include "NiagaraEmitterFactoryNew.h"
-
 #include "Editor.h"
-
-#include "ScopedTransaction.h"
-#include "MovieScene.h"
-#include "ISequencerModule.h"
 #include "EditorSupportDelegates.h"
-#include "Modules/ModuleManager.h"
-#include "ViewModels/TNiagaraViewModelManager.h"
+#include "Framework/Application/SlateApplication.h"
+#include "Framework/Commands/UICommandList.h"
+#include "ISequencerModule.h"
 #include "Math/NumericLimits.h"
-#include "NiagaraComponent.h"
+#include "Modules/ModuleManager.h"
+#include "MovieScene.h"
 #include "MovieSceneFolder.h"
-#include "Channels/MovieSceneChannelProxy.h"
-#include "Framework/Commands/UICommandList.h"
+#include "MovieSceneNiagaraEmitterTrack.h"
+#include "NiagaraComponent.h"
+#include "NiagaraDataInterfaceCurveBase.h"
+#include "NiagaraEditorData.h"
+#include "NiagaraEditorModule.h"
+#include "NiagaraEditorSettings.h"
+#include "NiagaraEditorUtilities.h"
+#include "NiagaraEmitter.h"
+#include "NiagaraEmitterEditorData.h"
+#include "NiagaraEmitterFactoryNew.h"
+#include "NiagaraEmitterHandle.h"
+#include "NiagaraGraph.h"
 #include "NiagaraMessageManager.h"
-#include "NiagaraMessages.h"
 #include "NiagaraMessageUtilities.h"
-#include "ViewModels/NiagaraOverviewGraphViewModel.h"
-#include "Framework/Application/SlateApplication.h"
+#include "NiagaraMessages.h"
+#include "NiagaraNodeAssignment.h"
+#include "NiagaraNodeEmitter.h"
+#include "NiagaraNodeFunctionCall.h"
+#include "NiagaraNodeInput.h"
+#include "NiagaraNodeOutput.h"
+#include "NiagaraNodeParameterMapGet.h"
+#include "NiagaraOverviewNode.h"
 #include "NiagaraParameterDefinitionsSubscriber.h"
+#include "NiagaraScriptGraphViewModel.h"
+#include "NiagaraScriptSource.h"
+#include "NiagaraScriptVariable.h"
+#include "NiagaraSequence.h"
+#include "NiagaraSimulationStageBase.h"
+#include "NiagaraSystem.h"
+#include "NiagaraSystemEditorData.h"
+#include "NiagaraSystemInstance.h"
+#include "NiagaraSystemScriptViewModel.h"
+#include "ScopedTransaction.h"
+#include "ViewModels/NiagaraCurveSelectionViewModel.h"
+#include "ViewModels/NiagaraEmitterHandleViewModel.h"
+#include "ViewModels/NiagaraEmitterViewModel.h"
+#include "ViewModels/NiagaraOverviewGraphViewModel.h"
+#include "ViewModels/NiagaraPlaceholderDataInterfaceManager.h"
+#include "ViewModels/NiagaraScratchPadUtilities.h"
+#include "ViewModels/NiagaraScratchPadViewModel.h"
+#include "ViewModels/NiagaraSystemEditorDocumentsViewModel.h"
+#include "ViewModels/NiagaraSystemGraphSelectionViewModel.h"
+#include "ViewModels/NiagaraSystemSelectionViewModel.h"
+#include "ViewModels/Stack/NiagaraStackEntry.h"
+#include "ViewModels/Stack/NiagaraStackGraphUtilities.h"
+#include "ViewModels/Stack/NiagaraStackViewModel.h"
+#include "ViewModels/TNiagaraViewModelManager.h"
 
+
 DECLARE_CYCLE_STAT(TEXT("Niagara - SystemViewModel - CompileSystem"), STAT_NiagaraEditor_SystemViewModel_CompileSystem, STATGROUP_NiagaraEditor);
 
@@ -143,6 +148,9 @@
 	ScalabilityViewModel->Initialize(this->AsShared());
 	
 	PlaceholderDataInterfaceManager = MakeShared<FNiagaraPlaceholderDataInterfaceManager>(this->AsShared());
+
+	SystemGraphSelectionViewModel = MakeShared<FNiagaraSystemGraphSelectionViewModel>();
+	SystemGraphSelectionViewModel->Initialize(this->AsShared());
 	
 	SetupPreviewComponentAndInstance();
 	SetupSequencer();
@@ -261,6 +269,11 @@
 		PlaceholderDataInterfaceManager.Reset();
 	}
 
+	if (SystemGraphSelectionViewModel != nullptr)
+	{
+		SystemGraphSelectionViewModel = nullptr;
+	}
+
 	System = nullptr;
 }
 
@@ -893,6 +906,21 @@
 	return StackModuleData;
 }
 
+TArray<TSharedRef<FNiagaraEmitterHandleViewModel>> FNiagaraSystemViewModel::GetSelectedEmitterHandleViewModels()
+{
+	TArray<TSharedRef<FNiagaraEmitterHandleViewModel>> SelectedEmitterHandleViewModels;
+	const TArray<FGuid>& SelectedEmitterHandleIds = GetSelectionViewModel()->GetSelectedEmitterHandleIds();
+
+	for (const TSharedRef<FNiagaraEmitterHandleViewModel>& EmitterHandleViewModel : EmitterHandleViewModels)
+	{
+		if (SelectedEmitterHandleIds.Contains(EmitterHandleViewModel->GetId()))
+		{
+			SelectedEmitterHandleViewModels.Add(EmitterHandleViewModel);
+		}
+	}
+	return SelectedEmitterHandleViewModels;
+}
+
 UNiagaraEditorParametersAdapter* FNiagaraSystemViewModel::GetEditorOnlyParametersAdapter() const
 {
 	if (EditMode == ENiagaraSystemViewModelEditMode::SystemAsset)
@@ -981,6 +1009,289 @@
 	return bIsForDataProcessingOnly;
 }
 
+bool FNiagaraSystemViewModel::RenameParameter(const FNiagaraVariable TargetParameter, const FName NewName, ENiagaraGetGraphParameterReferencesMode RenameScopeMode)
+{
+	FScopedTransaction RenameParameterTransaction(LOCTEXT("RenameParameter", "Rename parameter"));
+
+	System->Modify();
+
+	bool bExposedParametersRename = false;
+	bool bEditorOnlyParametersRename = false;
+	bool bAssignmentNodeRename = false;
+
+	// Rename the parameter if it's in the parameter store for user parameters.
+	if (System->GetExposedParameters().IndexOf(TargetParameter) != INDEX_NONE)
+	{
+		FNiagaraParameterStore* ExposedParametersStore = &System->GetExposedParameters();
+		TArray<FNiagaraVariable> OwningParameters;
+		ExposedParametersStore->GetParameters(OwningParameters);
+		if (OwningParameters.ContainsByPredicate([NewName](const FNiagaraVariable& Variable) { return Variable.GetName() == NewName; }))
+		{
+			// If the parameter store already has a parameter with this name, remove the old parameter to prevent collisions.
+			ExposedParametersStore->RemoveParameter(TargetParameter);
+		}
+		else
+		{
+			// Otherwise it's safe to rename.
+			ExposedParametersStore->RenameParameter(TargetParameter, NewName);
+		}
+		bExposedParametersRename = true;
+	}
+
+	// Look for set parameters nodes or linked inputs which reference this parameter and rename if so.
+	for (const FNiagaraGraphParameterReference& ParameterReference : GetGraphParameterReferences(TargetParameter, RenameScopeMode))
+	{
+		UNiagaraNode* ReferenceNode = Cast<UNiagaraNode>(ParameterReference.Value);
+		if (ReferenceNode != nullptr)
+		{
+			UNiagaraNodeAssignment* OwningAssignmentNode = ReferenceNode->GetTypedOuter<UNiagaraNodeAssignment>();
+			if (OwningAssignmentNode != nullptr)
+			{
+				// If this is owned by a set variables node and it's not locked, update the assignment target on the assignment node.
+				bAssignmentNodeRename |= FNiagaraStackGraphUtilities::TryRenameAssignmentTarget(*OwningAssignmentNode, TargetParameter, NewName);
+			}
+			else
+			{
+				// Otherwise if the reference node is a get node it's for a linked input so we can just update pin name.
+				UNiagaraNodeParameterMapGet* ReferenceGetNode = Cast<UNiagaraNodeParameterMapGet>(ReferenceNode);
+				if (ReferenceGetNode != nullptr)
+				{
+					if (ReferenceGetNode->Pins.ContainsByPredicate([&ParameterReference](UEdGraphPin* Pin) { return Pin->PersistentGuid == ParameterReference.Key; }))
+					{
+						ReferenceGetNode->GetNiagaraGraph()->RenameParameter(TargetParameter, NewName, true);
+					}
+				}
+			}
+		}
+	}
+
+	// Rename the parameter if it is owned directly as an editor only parameter.
+	UNiagaraEditorParametersAdapter* EditorParametersAdapter = GetEditorOnlyParametersAdapter();
+	if (UNiagaraScriptVariable** ScriptVariablePtr = EditorParametersAdapter->GetParameters().FindByPredicate([&TargetParameter](const UNiagaraScriptVariable* ScriptVariable) { return ScriptVariable->Variable.GetName() == TargetParameter.GetName(); }))
+	{
+		EditorParametersAdapter->Modify();
+		UNiagaraScriptVariable* ScriptVariable = *ScriptVariablePtr;
+		ScriptVariable->Modify();
+		ScriptVariable->Variable.SetName(NewName);
+		ScriptVariable->UpdateChangeId();
+		bEditorOnlyParametersRename = true;
+
+		// Check if the rename will give the same name and type as an existing parameter definition, and if so, link to the definition automatically.
+		FNiagaraParameterDefinitionsUtilities::TrySubscribeScriptVarToDefinitionByName(ScriptVariable, this);
+	}
+
+	// Handle renaming any renderer properties that might match.
+	if (bExposedParametersRename | bEditorOnlyParametersRename | bAssignmentNodeRename)
+	{
+		if (TargetParameter.IsInNameSpace(FNiagaraConstants::ParticleAttributeNamespaceString) || TargetParameter.IsInNameSpace(FNiagaraConstants::EmitterNamespaceString))
+		{
+			for (const TSharedRef<FNiagaraEmitterHandleViewModel>& EmitterHandleViewModel : GetSelectedEmitterHandleViewModels())
+			{
+				FNiagaraEmitterHandle* EmitterHandle = EmitterHandleViewModel->GetEmitterHandle();
+				EmitterHandle->GetInstance().Emitter->HandleVariableRenamed(TargetParameter, FNiagaraVariableBase(TargetParameter.GetType(), NewName), true, EmitterHandle->GetInstance().Version);
+			}
+		}
+		else
+		{
+			System->HandleVariableRenamed(TargetParameter, FNiagaraVariableBase(TargetParameter.GetType(), NewName), true);
+		}
+
+		return true;
+	}
+
+	return false;
+}
+
+TArray<UNiagaraGraph*> FNiagaraSystemViewModel::GetAllGraphs()
+{
+	TArray<UNiagaraGraph*> OutGraphs;
+
+	// Helper lambda to null check graph weak object ptrs before adding them to the retval array.
+	auto AddToOutGraphsChecked = [&OutGraphs](const TWeakObjectPtr<UNiagaraGraph>& WeakGraph) {
+		UNiagaraGraph* Graph = WeakGraph.Get();
+		if (Graph == nullptr)
+		{
+			ensureMsgf(false, TEXT("Encountered null graph when gathering editable graphs for system parameter panel viewmodel!"));
+			return;
+		}
+		OutGraphs.Add(Graph);
+	};
+
+	if (GetEditMode() == ENiagaraSystemViewModelEditMode::SystemAsset)
+	{
+		TArray<UNiagaraGraph*> Graphs;
+		AddToOutGraphsChecked(GetSystemScriptViewModel()->GetGraphViewModel()->GetGraph());
+
+		for (const TSharedRef<FNiagaraEmitterHandleViewModel>& EmitterHandleViewModel : EmitterHandleViewModels)
+		{
+			FNiagaraEmitterHandle* EmitterHandle = EmitterHandleViewModel->GetEmitterHandle();
+			if (EmitterHandle == nullptr)
+			{
+				continue;
+			}
+			UNiagaraGraph* Graph = Cast<UNiagaraScriptSource>(EmitterHandle->GetEmitterData()->GraphSource)->NodeGraph;
+			if (Graph)
+			{
+				OutGraphs.Add(Graph);
+			}
+		}
+	}
+	else
+	{
+		FNiagaraEmitterHandle* EmitterHandle = GetEmitterHandleViewModels()[0]->GetEmitterHandle();
+		if (EmitterHandle != nullptr)
+		{
+			OutGraphs.Add(Cast<UNiagaraScriptSource>(EmitterHandle->GetEmitterData()->GraphSource)->NodeGraph);
+		}
+	}
+
+	TArray<UNiagaraGraph*> EdGraphs = GetDocumentViewModel()->GetEditableGraphsForActiveScriptDocument();
+	for (UNiagaraGraph* Graph : EdGraphs)
+	{
+		OutGraphs.AddUnique(Graph);
+	}
+	return OutGraphs;
+}
+
+TArray<UNiagaraGraph*> FNiagaraSystemViewModel::GetSelectedGraphs()
+{
+	TArray<UNiagaraGraph*> SelectedGraphs;
+
+	// Helper lambda to null check graph weak object ptrs before adding them to the retval array.
+	auto AddToSelectedGraphsChecked = [&SelectedGraphs](const TWeakObjectPtr<UNiagaraGraph>& WeakGraph) {
+		UNiagaraGraph* Graph = WeakGraph.Get();
+		if (Graph == nullptr)
+		{
+			ensureMsgf(false, TEXT("Encountered null graph when gathering editable graphs for system parameter panel viewmodel!"));
+			return;
+		}
+		SelectedGraphs.Add(Graph);
+	};
+
+	if (GetEditMode() == ENiagaraSystemViewModelEditMode::SystemAsset)
+	{
+		for (const TWeakObjectPtr<UNiagaraGraph>& WeakGraph : SystemGraphSelectionViewModel->GetSelectedEmitterScriptGraphs())
+		{
+			AddToSelectedGraphsChecked(WeakGraph);
+		}
+		AddToSelectedGraphsChecked(GetSystemScriptViewModel()->GetGraphViewModel()->GetGraph());
+	}
+	else
+	{
+		SelectedGraphs.Add(Cast<UNiagaraScriptSource>(GetEmitterHandleViewModels()[0]->GetEmitterHandle()->GetEmitterData()->GraphSource)->NodeGraph);
+	}
+
+	TArray<UNiagaraGraph*> EdGraphs = GetDocumentViewModel()->GetEditableGraphsForActiveScriptDocument();
+	for (UNiagaraGraph* Graph : EdGraphs)
+	{
+		SelectedGraphs.AddUnique(Graph);
+	}
+	return SelectedGraphs;
+}
+
+TArray<FNiagaraGraphParameterReference> FNiagaraSystemViewModel::GetGraphParameterReferences(const FNiagaraVariable& Parameter, ENiagaraGetGraphParameterReferencesMode Mode)
+{
+	TArray<UNiagaraGraph*> TargetGraphs;
+	if (Mode == ENiagaraGetGraphParameterReferencesMode::AllGraphs)
+	{
+		TargetGraphs = GetAllGraphs();
+	}
+	else if (Mode == ENiagaraGetGraphParameterReferencesMode::SelectedGraphs)
+	{
+		TargetGraphs = GetSelectedGraphs();
+	}
+	else
+	{
+		ensureMsgf(false, TEXT("Encountered unexpected ENiagaraGetGraphParameterReferencesMode value!"));
+		return TArray<FNiagaraGraphParameterReference>();
+	}
+
+	// -For each selected graph perform a parameter map history traversal and collect all graph parameter references associated with the target FNiagaraParameterPanelItem.
+	TArray<FNiagaraGraphParameterReference> GraphParameterReferences;
+	for (const UNiagaraGraph* Graph : TargetGraphs)
+	{
+		TArray<UNiagaraNodeOutput*> OutputNodes;
+		Graph->GetNodesOfClass<UNiagaraNodeOutput>(OutputNodes);
+		for (UNiagaraNodeOutput* OutputNode : OutputNodes)
+		{
+			UNiagaraNode* NodeToTraverse = OutputNode;
+			if (OutputNode->GetUsage() == ENiagaraScriptUsage::SystemSpawnScript || OutputNode->GetUsage() == ENiagaraScriptUsage::SystemUpdateScript)
+			{
+				// Traverse past the emitter nodes, otherwise the system scripts will pick up all of the emitter and particle script parameters.
+				UEdGraphPin* InputPin = FNiagaraStackGraphUtilities::GetParameterMapInputPin(*NodeToTraverse);
+				while (NodeToTraverse != nullptr && InputPin != nullptr && InputPin->LinkedTo.Num() == 1 &&
+					(NodeToTraverse->IsA<UNiagaraNodeOutput>() || NodeToTraverse->IsA<UNiagaraNodeEmitter>()))
+				{
+					NodeToTraverse = Cast<UNiagaraNode>(InputPin->LinkedTo[0]->GetOwningNode());
+					InputPin = NodeToTraverse != nullptr ? FNiagaraStackGraphUtilities::GetParameterMapInputPin(*NodeToTraverse) : nullptr;
+				}
+			}
+
+			if (NodeToTraverse == nullptr)
+			{
+				continue;
+			}
+
+			bool bIgnoreDisabled = true;
+			FNiagaraParameterMapHistoryBuilder Builder;
+			FVersionedNiagaraEmitter GraphOwningEmitter = Graph->GetOwningEmitter();
+			FCompileConstantResolver ConstantResolver = GraphOwningEmitter.Emitter != nullptr
+				? FCompileConstantResolver(GraphOwningEmitter, ENiagaraScriptUsage::Function)
+				: FCompileConstantResolver();
+
+			Builder.SetIgnoreDisabled(bIgnoreDisabled);
+			Builder.ConstantResolver = ConstantResolver;
+			FName StageName;
+			ENiagaraScriptUsage StageUsage = OutputNode->GetUsage();
+			if (StageUsage == ENiagaraScriptUsage::ParticleSimulationStageScript && GraphOwningEmitter.Emitter)
+			{
+				UNiagaraSimulationStageBase* Base = GraphOwningEmitter.GetEmitterData()->GetSimulationStageById(OutputNode->GetUsageId());
+				if (Base)
+				{
+					StageName = Base->GetStackContextReplacementName();
+				}
+			}
+			Builder.BeginUsage(StageUsage, StageName);
+			NodeToTraverse->BuildParameterMapHistory(Builder, true, false);
+			Builder.EndUsage();
+
+			if (Builder.Histories.Num() != 1)
+			{
+				// We should only have traversed one emitter (have not visited more than one NiagaraNodeEmitter.)
+				ensureMsgf(false, TEXT("Encountered more than one parameter map history when collecting graph parameter reference collections for system parameter panel view model!"));
+			}
+			if (Builder.Histories.Num() == 0)
+			{
+				continue;
+			}
+
+			const TArray<FName>& CustomIterationSourceNamespaces = Builder.Histories[0].IterationNamespaceOverridesEncountered;
+			for (int32 VariableIndex = 0; VariableIndex < Builder.Histories[0].Variables.Num(); VariableIndex++)
+			{
+				if (Parameter == Builder.Histories[0].Variables[VariableIndex])
+				{
+					for (const FNiagaraParameterMapHistory::FReadHistory& ReadHistory : Builder.Histories[0].PerVariableReadHistory[VariableIndex])
+					{
+						if (ReadHistory.ReadPin.Pin->GetOwningNode() != nullptr)
+						{
+							GraphParameterReferences.Add(FNiagaraGraphParameterReference(ReadHistory.ReadPin.Pin->PersistentGuid, ReadHistory.ReadPin.Pin->GetOwningNode()));
+						}
+					}
+
+					for (const FModuleScopedPin& Write : Builder.Histories[0].PerVariableWriteHistory[VariableIndex])
+					{
+						if (Write.Pin->GetOwningNode() != nullptr)
+						{
+							GraphParameterReferences.Add(FNiagaraGraphParameterReference(Write.Pin->PersistentGuid, Write.Pin->GetOwningNode()));
+						}
+					}
+				}
+			}
+		}
+	}
+	return GraphParameterReferences;
+}
+
 void FNiagaraSystemViewModel::SendLastCompileMessageJobs() const
 {
 	if (SystemMessageLogGuidKey.IsSet() == false)
diff --git a/Engine/Plugins/FX/Niagara/Source/NiagaraEditor/Public/NiagaraEditorUtilities.h b/Engine/Plugins/FX/Niagara/Source/NiagaraEditor/Public/NiagaraEditorUtilities.h
--- a/Engine/Plugins/FX/Niagara/Source/NiagaraEditor/Public/NiagaraEditorUtilities.h
+++ b/Engine/Plugins/FX/Niagara/Source/NiagaraEditor/Public/NiagaraEditorUtilities.h
@@ -434,4 +434,5 @@ namespace FNiagaraParameterDefinitionsUtilities
 	TArray<const UNiagaraScriptVariable*>  FindReservedParametersByName(const FName ParameterName);
 	int32 GetNumParametersReservedForName(const FName ParameterName);
 	EParameterDefinitionMatchState GetDefinitionMatchStateForParameter(const FNiagaraVariableBase& Parameter);
+	void TrySubscribeScriptVarToDefinitionByName(UNiagaraScriptVariable* ScriptVar, INiagaraParameterDefinitionsSubscriberViewModel* OwningDefinitionSubscriberViewModel);
 };
diff --git a/Engine/Plugins/FX/Niagara/Source/NiagaraEditor/Public/ViewModels/NiagaraParameterDefinitionsPanelViewModel.h b/Engine/Plugins/FX/Niagara/Source/NiagaraEditor/Public/ViewModels/NiagaraParameterDefinitionsPanelViewModel.h
--- a/Engine/Plugins/FX/Niagara/Source/NiagaraEditor/Public/ViewModels/NiagaraParameterDefinitionsPanelViewModel.h
+++ b/Engine/Plugins/FX/Niagara/Source/NiagaraEditor/Public/ViewModels/NiagaraParameterDefinitionsPanelViewModel.h
@@ -80,7 +80,6 @@ private:
 class FNiagaraSystemToolkitParameterDefinitionsPanelViewModel : public INiagaraParameterDefinitionsPanelViewModel
 {
 public:
-	FNiagaraSystemToolkitParameterDefinitionsPanelViewModel(const TSharedPtr<FNiagaraSystemViewModel>& InSystemViewModel, const TWeakPtr<FNiagaraSystemGraphSelectionViewModel>& SystemGraphSelectionViewModelWeak);
 	FNiagaraSystemToolkitParameterDefinitionsPanelViewModel(const TSharedPtr<FNiagaraSystemViewModel>& InSystemViewModel);
 
 	void Init(const FSystemToolkitUIContext& InUIContext);
diff --git a/Engine/Plugins/FX/Niagara/Source/NiagaraEditor/Public/ViewModels/NiagaraParameterPanelViewModel.h b/Engine/Plugins/FX/Niagara/Source/NiagaraEditor/Public/ViewModels/NiagaraParameterPanelViewModel.h
--- a/Engine/Plugins/FX/Niagara/Source/NiagaraEditor/Public/ViewModels/NiagaraParameterPanelViewModel.h
+++ b/Engine/Plugins/FX/Niagara/Source/NiagaraEditor/Public/ViewModels/NiagaraParameterPanelViewModel.h
@@ -39,10 +39,6 @@ namespace FNiagaraParameterUtilities
 // NOTE: These utilities are not defined in the view model directly as they are shared between ParameterPanelViewModel and ParameterDefinitionsPanelViewModel.
 namespace FNiagaraSystemToolkitParameterPanelUtilities
 {
-	// Get all graphs for the active emitter selection or just the one for the active script mode
-	TArray<UNiagaraGraph*> GetAllGraphs(const TSharedPtr<FNiagaraSystemViewModel>& SystemViewModel, bool bActiveScriptMode);
-	// Get all editable graphs for the active emitter selection or just the one for the active script mode
-	TArray<UNiagaraGraph*> GetEditableGraphs(const TSharedPtr<FNiagaraSystemViewModel>& SystemViewModel, const TWeakPtr<FNiagaraSystemGraphSelectionViewModel>& SystemGraphSelectionViewModelWeak, bool bActiveScriptMode);
 	FReply CreateDragEventForParameterItem(
 		const FNiagaraParameterPanelItemBase& DraggedItem,
 		const FPointerEvent& MouseEvent,
@@ -53,7 +49,6 @@ namespace FNiagaraSystemToolkitParameterPanelUtilities
 
 namespace FNiagaraScriptToolkitParameterPanelUtilities
 {
-	TArray<UNiagaraGraph*> GetEditableGraphs(const TSharedPtr<FNiagaraScriptViewModel>& ScriptViewModel);
 	FReply CreateDragEventForParameterItem(
 		const FNiagaraParameterPanelItemBase& DraggedItem,
 		const FPointerEvent& MouseEvent,
@@ -173,6 +168,7 @@ public:
 	INiagaraParameterPanelViewModel();
 	virtual ~INiagaraParameterPanelViewModel() override;
 
+	void RefreshFull(bool bDoCategoryExpansion) const;
 	//~ Begin INiagaraImmutableParameterPanelViewModel interface
 	virtual const TArray<UNiagaraGraph*> GetEditableGraphsConst() const = 0;
 
@@ -194,8 +190,6 @@ public:
 
 	virtual void RenameParameter(const FNiagaraParameterPanelItem& ItemToRename, const FName NewName)  = 0;
 
-	virtual void SetParameterIsSubscribedToLibrary(const UNiagaraScriptVariable* ScriptVarToModify, const bool bSubscribed)  = 0;
-
 	virtual TSharedPtr<SWidget> CreateContextMenuForItems(const TArray<FNiagaraParameterPanelItem>& Items, const TSharedPtr<FUICommandList>& ToolkitCommands) = 0;
 
 	virtual FNiagaraParameterUtilities::EParameterContext GetParameterContext() const = 0;
@@ -241,8 +235,6 @@ public:
 
 	virtual bool GetCanSubscribeParameterToLibraryAndToolTip(const FNiagaraParameterPanelItem& ItemToModify, const bool bSubscribing, FText& OutCanSubscribeParameterToolTip) const;
 
-	virtual void SetParameterIsSubscribedToLibrary(const FNiagaraParameterPanelItem ItemToModify, const bool bSubscribed) ;
-
 	virtual void SetParameterNamespace(const FNiagaraParameterPanelItem ItemToModify, FNiagaraNamespaceMetadata NewNamespaceMetaData, bool bDuplicateParameter) ;
 
 	virtual void SetParameterNamespaceModifier(const FNiagaraParameterPanelItem ItemToModify, const FName NewNamespaceModifier, bool bDuplicateParameter) ;
@@ -263,9 +255,6 @@ public:
 
 	void SelectParameterItemByName(const FName ParameterName, const bool bRequestRename);
 
-	void SubscribeParameterToLibraryIfMatchingDefinition(const UNiagaraScriptVariable* ScriptVarToModify, const FName ScriptVarName) ;
-
-	virtual void RefreshFull(bool bDoCategoryExpansion) const;
 
 	FOnParameterPanelViewModelExternalSelectionChanged& GetOnExternalSelectionChangedDelegate() { return OnParameterPanelViewModelExternalSelectionChangedDelegate; };
 
@@ -326,7 +315,6 @@ class FNiagaraSystemToolkitParameterPanelViewModel : public INiagaraParameterPan
 public:
 	/** Construct a SystemToolkit Parameter Panel View Model from a System View Model and an optional SystemGraphSelectionViewModel. */
 	FNiagaraSystemToolkitParameterPanelViewModel(const TSharedPtr<FNiagaraSystemViewModel>& InSystemViewModel);
-	FNiagaraSystemToolkitParameterPanelViewModel(const TSharedPtr<FNiagaraSystemViewModel>& InSystemViewModel, const TWeakPtr<FNiagaraSystemGraphSelectionViewModel>& InSystemGraphSelectionViewModelWeak);
 
 	void Init(const FSystemToolkitUIContext& InUIContext);
 
@@ -353,8 +341,6 @@ public:
 
 	virtual void RenameParameter(const FNiagaraParameterPanelItem& ItemToRename, const FName NewName)  override;
 
-	virtual void SetParameterIsSubscribedToLibrary(const UNiagaraScriptVariable* ScriptVarToModify, const bool bSubscribed)  override;
-
 	virtual FReply OnParameterItemsDragged(const TArray<FNiagaraParameterPanelItem>& DraggedItems, const FPointerEvent& MouseEvent) const override;
 
 	virtual TSharedPtr<SWidget> CreateContextMenuForItems(const TArray<FNiagaraParameterPanelItem>& Items, const TSharedPtr<FUICommandList>& ToolkitCommands) override;
@@ -493,8 +479,6 @@ public:
 
 	virtual void DuplicateParameters(const TArray<FNiagaraParameterPanelItem> ItemsToDuplicate)  override;
 
-	virtual void SetParameterIsSubscribedToLibrary(const UNiagaraScriptVariable* ScriptVarToModify, const bool bSubscribed)  override;
-
 	virtual void OnParameterItemSelected(const FNiagaraParameterPanelItem& SelectedItem, ESelectInfo::Type SelectInfo) const override;
 
 	virtual FReply OnParameterItemsDragged(const TArray<FNiagaraParameterPanelItem>& DraggedItems, const FPointerEvent& MouseEvent) const override;
@@ -594,8 +578,6 @@ public:
 
 	virtual void RenameParameter(const FNiagaraParameterPanelItem& ItemToRename, const FName NewName)  override;
 
-	virtual void SetParameterIsSubscribedToLibrary(const UNiagaraScriptVariable* ScriptVarToModify, const bool bSubscribed)  override;
-
 	virtual TSharedPtr<SWidget> CreateContextMenuForItems(const TArray<FNiagaraParameterPanelItem>& Items, const TSharedPtr<FUICommandList>& ToolkitCommands) override;
 
 	virtual FNiagaraParameterUtilities::EParameterContext GetParameterContext() const override;
diff --git a/Engine/Plugins/FX/Niagara/Source/NiagaraEditor/Public/ViewModels/NiagaraScriptViewModel.h b/Engine/Plugins/FX/Niagara/Source/NiagaraEditor/Public/ViewModels/NiagaraScriptViewModel.h
--- a/Engine/Plugins/FX/Niagara/Source/NiagaraEditor/Public/ViewModels/NiagaraScriptViewModel.h
+++ b/Engine/Plugins/FX/Niagara/Source/NiagaraEditor/Public/ViewModels/NiagaraScriptViewModel.h
@@ -86,7 +86,8 @@ public:
 
 	/** If this is editing a standalone script, returns the script being edited.*/
 	virtual FVersionedNiagaraScript GetStandaloneScript();
-	//const UNiagaraScript* GetStandaloneScript() const;
+	
+	bool RenameParameter(const FNiagaraVariable TargetParameter, const FName NewName);
 
 private:
 	/** Handles the selection changing in the graph view model. */
diff --git a/Engine/Plugins/FX/Niagara/Source/NiagaraEditor/Public/ViewModels/NiagaraSystemGraphSelectionViewModel.h b/Engine/Plugins/FX/Niagara/Source/NiagaraEditor/Public/ViewModels/NiagaraSystemGraphSelectionViewModel.h
--- a/Engine/Plugins/FX/Niagara/Source/NiagaraEditor/Public/ViewModels/NiagaraSystemGraphSelectionViewModel.h
+++ b/Engine/Plugins/FX/Niagara/Source/NiagaraEditor/Public/ViewModels/NiagaraSystemGraphSelectionViewModel.h
@@ -18,7 +18,7 @@ public:
 
 	FNiagaraSystemGraphSelectionViewModel() = default;
 
-	void Init(TSharedPtr<FNiagaraSystemViewModel> InSystemViewModel);
+	void Initialize(TSharedRef<FNiagaraSystemViewModel> InSystemViewModel);
 
 	const TArray<TWeakObjectPtr<UNiagaraGraph>> GetSelectedEmitterScriptGraphs() const { return SelectedEmitterScriptGraphs; };
 
diff --git a/Engine/Plugins/FX/Niagara/Source/NiagaraEditor/Public/ViewModels/NiagaraSystemViewModel.h b/Engine/Plugins/FX/Niagara/Source/NiagaraEditor/Public/ViewModels/NiagaraSystemViewModel.h
--- a/Engine/Plugins/FX/Niagara/Source/NiagaraEditor/Public/ViewModels/NiagaraSystemViewModel.h
+++ b/Engine/Plugins/FX/Niagara/Source/NiagaraEditor/Public/ViewModels/NiagaraSystemViewModel.h
@@ -2,56 +2,56 @@
 
 #pragma once
 
+#include "EditorUndoClient.h"
 #include "IMovieScenePlayer.h"
-
+#include "ISequencer.h"
+#include "ISequencerModule.h"
 #include "NiagaraCommon.h"
 #include "NiagaraEditorCommon.h"
-#include "EditorUndoClient.h"
-#include "UObject/GCObject.h"
-#include "ViewModels/TNiagaraViewModelManager.h"
-#include "ViewModels/NiagaraParameterDefinitionsSubscriberViewModel.h"
 #include "NiagaraOverviewNode.h"
-#include "ISequencer.h"
-
-#include "TickableEditorObject.h"
-#include "ISequencerModule.h"
 #include "NiagaraSystemScalabilityViewModel.h"
+#include "TickableEditorObject.h"
+#include "UObject/GCObject.h"
 #include "UObject/ObjectKey.h"
+#include "ViewModels/TNiagaraViewModelManager.h"
+#include "ViewModels/NiagaraParameterDefinitionsSubscriberViewModel.h"
 
 
-class UNiagaraSystem;
-class UNiagaraComponent;
-class UNiagaraSequence;
-class UMovieSceneNiagaraEmitterTrack;
-struct FNiagaraVariable;
+struct FAssetData;
+struct FEdGraphEditAction;
 struct FNiagaraEmitterHandle;
-class UNiagaraOverviewNode;
 class FNiagaraEmitterHandleViewModel;
-class FNiagaraSystemScriptViewModel;
-class UNiagaraStackViewModel;
-class UNiagaraStackEntry;
-class UNiagaraSystemSelectionViewModel;
+class FNiagaraEmitterViewModel;
+struct FNiagaraGraphParameterReference;
+class FNiagaraOverviewGraphViewModel;
+struct FNiagaraParameterStore;
+class FNiagaraPlaceholderDataInterfaceManager;
+class FNiagaraSystemGraphSelectionViewModel;
 class FNiagaraSystemInstance;
+class FNiagaraSystemScriptViewModel;
+struct FNiagaraVariable;
+class INiagaraParameterDefinitionsSubscriber;
+class INiagaraParameterPanelViewModel;
 class ISequencer;
-struct FAssetData;
-class UNiagaraSystemEditorData;
+class UMovieSceneNiagaraEmitterTrack;
+class UNiagaraComponent;
+class UNiagaraCurveSelectionViewModel;
+class UNiagaraEditorParametersAdapter;
 class UNiagaraEditorSettings;
-struct FNiagaraParameterStore;
-struct FEdGraphEditAction;
+class UNiagaraGraph;
+class UNiagaraMessageData;
 class UNiagaraNodeFunctionCall;
-class FNiagaraEmitterViewModel;
-class FNiagaraOverviewGraphViewModel;
+class UNiagaraOverviewNode;
 class UNiagaraScratchPadViewModel;
-class UNiagaraCurveSelectionViewModel;
-class UNiagaraMessageData;
-class UNiagaraEditorParametersAdapter;
+class UNiagaraSequence;
+class UNiagaraStackEntry;
+class UNiagaraStackViewModel;
+class UNiagaraSystem;
+class UNiagaraSystemEditorData;
 class UNiagaraSystemEditorDocumentsViewModel;
+class UNiagaraSystemSelectionViewModel;
 
-class INiagaraParameterDefinitionsSubscriber;
 
-class FNiagaraPlaceholderDataInterfaceManager;
-class INiagaraParameterPanelViewModel;
-
 /** Defines different editing modes for this system view model. */
 enum class NIAGARAEDITOR_API ENiagaraSystemViewModelEditMode
 {
@@ -63,6 +63,12 @@
 	EmitterDuringMerge,
 };
 
+enum class ENiagaraGetGraphParameterReferencesMode
+{
+	AllGraphs,
+	SelectedGraphs,
+};
+
 DECLARE_DELEGATE_RetVal(FName, FOnGetWorkflowMode);
 
 /** Defines options for the niagara System view model */
@@ -426,6 +432,20 @@
 
 	UNiagaraSystemEditorDocumentsViewModel* GetDocumentViewModel() { return EditorDocumentsViewModel;};
 
+	/** Rename a parameter within the RenameScope and update associated tracking metadata such as Parameter Definition synchronization. */
+	bool RenameParameter(const FNiagaraVariable TargetParameter, const FName NewName, ENiagaraGetGraphParameterReferencesMode RenameScopeMode);
+
+	const TSharedPtr<FNiagaraSystemGraphSelectionViewModel>& GetSystemGraphSelectionViewModel() const { return SystemGraphSelectionViewModel; };
+
+	/** Utility method to gather all graphs. */
+	TArray<UNiagaraGraph*> GetAllGraphs();
+
+	/** Utility method to gather all graphs marked as selected via the system selection viewmodel. */
+	TArray<UNiagaraGraph*> GetSelectedGraphs();
+
+	/** Utility method to gather graph parameter references */
+	TArray<FNiagaraGraphParameterReference> GetGraphParameterReferences(const FNiagaraVariable& Parameter, ENiagaraGetGraphParameterReferencesMode Mode);
+
 private:
 	/** Sends message jobs to FNiagaraMessageManager for all compile events from the last compile. */
 	void SendLastCompileMessageJobs() const;
@@ -552,6 +572,9 @@
 
 	const TArray<FNiagaraStackModuleData>& BuildAndCacheStackModuleData(FGuid EmitterHandleId, const FVersionedNiagaraEmitter& Emitter);
 
+	/** Utility method to gather all emitter handle viewmodels marked as selected via the system selection viewmodel. */
+	TArray<TSharedRef<FNiagaraEmitterHandleViewModel>> GetSelectedEmitterHandleViewModels();
+
 private:
 	/** The System being viewed and edited by this view model. */
 	UNiagaraSystem* System;
@@ -711,4 +734,7 @@
 
 	/** Specifies that this view model is for data processing only and will not be displayed in the UI. */
 	bool bIsForDataProcessingOnly;
+
+	/** ViewModel for caching the selected graph state and notifying when it changes. */
+	TSharedPtr<FNiagaraSystemGraphSelectionViewModel> SystemGraphSelectionViewModel;
 };
