- Fix OOB reads from bubble sort
- Fix OOB reads from the Ribbon VF
- Fix incorrect triangle count from CPU simulations
- Remove some temporary array copies to improve CPU performance
- Change CPU buffers to BUF_Volatile to fix issues on some RHIs
- Fixed incorrect bShouldFlipNormalToView set on various ribbon types

diff --git a/Engine/Plugins/FX/Niagara/Shaders/Private/NiagaraRibbonVertexFactory.ush b/Engine/Plugins/FX/Niagara/Shaders/Private/NiagaraRibbonVertexFactory.ush
--- a/Engine/Plugins/FX/Niagara/Shaders/Private/NiagaraRibbonVertexFactory.ush
+++ b/Engine/Plugins/FX/Niagara/Shaders/Private/NiagaraRibbonVertexFactory.ush
@@ -286,11 +286,6 @@ void SetUV(inout FVertexFactoryInterpolantsVSToPS Interpolants, int UVIndex, flo
 
 #endif
 
-int GetParticleIdByInterpVtxId(uint InterpVtxId)
-{
-	return NiagaraRibbonVFLooseParameters.SortedIndices[UnpackParticleId(InterpVtxId)];
-}
-
 uint GetInterpCount()
 {
 	BRANCH
@@ -317,14 +312,29 @@ float GetOneOverInterpCount()
 	}
 }
 
-FInterpVtxData GetInterpVtxData(FVertexFactoryInput Input)
+void GetInterpVtxAndRibbonDataData(FVertexFactoryInput Input, out FInterpVtxData InterpVtxData, out FSegmentData SegmentData)
 {
-	FInterpVtxData InterpVtxData = (FInterpVtxData)0;
-
+	// Get First Particle Data
 	InterpVtxData.RawParticleId0 = UnpackParticleId(Input.InterpVtxId);
-	InterpVtxData.RawParticleId1 = InterpVtxData.RawParticleId0 + 1;
-
 	InterpVtxData.ParticleId0 = NiagaraRibbonVFLooseParameters.SortedIndices[InterpVtxData.RawParticleId0];
+
+	// Unpack Ribbon Data
+	const uint RibbonIndex = NiagaraRibbonVF.ShouldUseMultiRibbon == 1 ? NiagaraRibbonVFLooseParameters.MultiRibbonIndices[InterpVtxData.RawParticleId0] : 0;
+	const uint RibbonBufferIndex = RibbonIndex * 8;
+
+	SegmentData.UV0Scale = float2(asfloat(NiagaraRibbonVFLooseParameters.PackedPerRibbonDataByIndex[RibbonBufferIndex + 0]), NiagaraRibbonVF.PackedVData.x);
+	SegmentData.UV0Offset = float2(asfloat(NiagaraRibbonVFLooseParameters.PackedPerRibbonDataByIndex[RibbonBufferIndex + 1]), NiagaraRibbonVF.PackedVData.y);
+	SegmentData.U0DistributionScaler = asfloat(NiagaraRibbonVFLooseParameters.PackedPerRibbonDataByIndex[RibbonBufferIndex + 2]);
+	
+	SegmentData.UV1Scale = float2(asfloat(NiagaraRibbonVFLooseParameters.PackedPerRibbonDataByIndex[RibbonBufferIndex + 3]), NiagaraRibbonVF.PackedVData.z);
+	SegmentData.UV1Offset = float2(asfloat(NiagaraRibbonVFLooseParameters.PackedPerRibbonDataByIndex[RibbonBufferIndex + 4]), NiagaraRibbonVF.PackedVData.w);
+	SegmentData.U1DistributionScaler = asfloat(NiagaraRibbonVFLooseParameters.PackedPerRibbonDataByIndex[RibbonBufferIndex + 5]);
+
+	SegmentData.StartParticleId = NiagaraRibbonVFLooseParameters.PackedPerRibbonDataByIndex[RibbonBufferIndex + 6];
+	SegmentData.LastParticleId = NiagaraRibbonVFLooseParameters.PackedPerRibbonDataByIndex[RibbonBufferIndex + 7];
+
+	// Get Remaining Interp Data
+	InterpVtxData.RawParticleId1 = min(InterpVtxData.RawParticleId0 + 1, SegmentData.LastParticleId);
 	InterpVtxData.ParticleId1 = NiagaraRibbonVFLooseParameters.SortedIndices[InterpVtxData.RawParticleId1];
 
 #if USE_PARTICLE_INTERPOLATION
@@ -334,27 +344,6 @@ FInterpVtxData GetInterpVtxData(FVertexFactoryInput Input)
 #else
 	InterpVtxData.Alpha = 0;
 #endif
-
-	return InterpVtxData;
-}
-
-FSegmentData UnpackPerRibbonDataByIndex(int RibbonIndex)
-{
-	FSegmentData SegmentData = (FSegmentData)0;
-	const int Index = RibbonIndex * 8;
-
-	SegmentData.UV0Scale = float2(asfloat(NiagaraRibbonVFLooseParameters.PackedPerRibbonDataByIndex[Index + 0]), NiagaraRibbonVF.PackedVData.x);
-	SegmentData.UV0Offset = float2(asfloat(NiagaraRibbonVFLooseParameters.PackedPerRibbonDataByIndex[Index + 1]), NiagaraRibbonVF.PackedVData.y);
-	SegmentData.U0DistributionScaler = asfloat(NiagaraRibbonVFLooseParameters.PackedPerRibbonDataByIndex[Index + 2]);
-	
-	SegmentData.UV1Scale = float2(asfloat(NiagaraRibbonVFLooseParameters.PackedPerRibbonDataByIndex[Index + 3]), NiagaraRibbonVF.PackedVData.z);
-	SegmentData.UV1Offset = float2(asfloat(NiagaraRibbonVFLooseParameters.PackedPerRibbonDataByIndex[Index + 4]), NiagaraRibbonVF.PackedVData.w);
-	SegmentData.U1DistributionScaler = asfloat(NiagaraRibbonVFLooseParameters.PackedPerRibbonDataByIndex[Index + 5]);
-
-	SegmentData.StartParticleId = NiagaraRibbonVFLooseParameters.PackedPerRibbonDataByIndex[Index + 6];
-	SegmentData.LastParticleId = NiagaraRibbonVFLooseParameters.PackedPerRibbonDataByIndex[Index + 7];
-
-	return SegmentData;
 }
 
 /** Converts from vertex factory specific interpolants to a FMaterialPixelParameters, which is used by material inputs. */
@@ -623,13 +612,12 @@ FRibbonPositionIntermediates GetRibbonPositionIntermediates(
 FVertexFactoryIntermediates GetVertexFactoryIntermediates(FVertexFactoryInput Input)
 {
 	FVertexFactoryIntermediates Intermediates = (FVertexFactoryIntermediates)0;
-	FInterpVtxData InterpVtxData = GetInterpVtxData(Input);
+	FInterpVtxData InterpVtxData = (FInterpVtxData)0;
+	FSegmentData SegmentData = (FSegmentData)0;
+	GetInterpVtxAndRibbonDataData(Input, InterpVtxData, SegmentData);
 	
 	Intermediates.SceneData = VF_GPUSCENE_GET_INTERMEDIATES(Input);
 	
-	uint RibbonIndex = NiagaraRibbonVF.ShouldUseMultiRibbon == 1? NiagaraRibbonVFLooseParameters.MultiRibbonIndices[InterpVtxData.RawParticleId0] : 0;
-	FSegmentData SegmentData = UnpackPerRibbonDataByIndex(RibbonIndex);
-
 	const float2 SliceVertexNormals = UnpackParticleVertexNormal(Input.InterpVtxId);
 
 	FPrimitiveSceneData PrimitiveData = GetPrimitiveDataFromUniformBuffer();
diff --git a/Engine/Plugins/FX/Niagara/Shaders/Private/Ribbons/NiagaraRibbonSortParticles.usf b/Engine/Plugins/FX/Niagara/Shaders/Private/Ribbons/NiagaraRibbonSortParticles.usf
--- a/Engine/Plugins/FX/Niagara/Shaders/Private/Ribbons/NiagaraRibbonSortParticles.usf
+++ b/Engine/Plugins/FX/Niagara/Shaders/Private/Ribbons/NiagaraRibbonSortParticles.usf
@@ -23,8 +23,6 @@ struct FBubbleShortEntry
 
 groupshared FBubbleShortEntry BubbleSortTemp[THREADGROUP_SIZE];
 
-
-
 [numthreads(THREADGROUP_SIZE, 1, 1)]
 void InitialSortList(uint3 DispatchThreadId : SV_DispatchThreadID, uint InGroupIndex : SV_GroupIndex, uint3 InGroupID : SV_GroupID)
 {
@@ -38,7 +36,7 @@ void InitialSortList(uint3 DispatchThreadId : SV_DispatchThreadID, uint InGroupI
 	const int LocalIndexOffset = InGroupIndex;
 
 	FBubbleShortEntry NewEntry;
-	NewEntry.ParticleID = GlobalIndexOffset; //(TotalNumParticles - 1) - GlobalIndexOffset;
+	NewEntry.ParticleID = min(GlobalIndexOffset, int(TotalNumParticles - 1));
 	NewEntry.LinkID = GetRibbonLinkID(NewEntry.ParticleID);	
 #if HAS_RIBBON_ID
 	NewEntry.RibbonID = GetRibbonID(NewEntry.ParticleID);
@@ -83,11 +81,6 @@ void InitialSortList(uint3 DispatchThreadId : SV_DispatchThreadID, uint InGroupI
 }
 
 
-
-
-
-
-
 int MergeSortCompare(float SearchLinkOrder, float TestLinkOrder)
 {
 	if (SearchLinkOrder > TestLinkOrder)
@@ -122,9 +115,6 @@ int MergeSortCompare(FRibbonID Left, float LeftLinkOrder, FRibbonID Right, float
 }
 #endif
 
-
-
-
 // This is a scatter based mergesort kernel
 [numthreads(THREADGROUP_SIZE, 1, 1)]
 void MergeSort(uint3 DispatchThreadId : SV_DispatchThreadID)
@@ -182,7 +172,3 @@ void MergeSort(uint3 DispatchThreadId : SV_DispatchThreadID)
 		DestinationSortedIndices[FinalIndex] = ParticleID;
 	}
 }
-
-
-
-
diff --git a/Engine/Plugins/FX/Niagara/Source/Niagara/Private/NiagaraRendererRibbons.cpp b/Engine/Plugins/FX/Niagara/Source/Niagara/Private/NiagaraRendererRibbons.cpp
--- a/Engine/Plugins/FX/Niagara/Source/Niagara/Private/NiagaraRendererRibbons.cpp
+++ b/Engine/Plugins/FX/Niagara/Source/Niagara/Private/NiagaraRendererRibbons.cpp
@@ -357,54 +357,75 @@ public:
 	}
 };
 
-bool FNiagaraRibbonVertexBuffers::InitOrUpdateBuffer(bool bEnabled, FRWBuffer& Buffer, int32& CurrentLength, int32 NeededLength, int32 MaxLength, FRWBuffer(* InitFunction)(int32, ERHIAccess), ERHIAccess InitialAccessFlags)
+bool FNiagaraRibbonGpuBuffer::Allocate(uint32 NumElements, uint32 MaxElements, ERHIAccess InResourceState, bool bGpuReadOnly)
 {
-	static constexpr float UpsizeMultipler = 1.1f;
-	static constexpr float DownsizeMultiplier = 1.2f;
+	if (NumElements == 0)
+	{
+		Release();
+		return false;
+	}
+	else
+	{
+		static constexpr float UpsizeMultipler = 1.1f;
+		static constexpr float DownsizeMultiplier = 1.2f;
 
-	bEnabled &= NeededLength > 0;
+		check(NumElements <= MaxElements);
 
-	if (bEnabled)
-	{
-		check(NeededLength <= MaxLength);
-		
-		// we resize the buffer if it's too small or we're over the allowed free space.
-		if (CurrentLength < NeededLength || CurrentLength > (NeededLength * DownsizeMultiplier))
+		const bool bGpuUsageChanged = bGpuReadOnly ? UAV != nullptr : UAV == nullptr;
+
+		const uint32 CurrentElements = NumBytes / ElementBytes;
+		if (bGpuUsageChanged || CurrentElements < NumElements || CurrentElements > uint32(FMath::CeilToInt32(NumElements * DownsizeMultiplier)))
 		{
-			const int32 NewLength = FMath::Min(MaxLength, FMath::RoundToInt32(NeededLength * UpsizeMultipler));
+			NumElements = FMath::Min(MaxElements, uint32(FMath::RoundToInt32(NumElements * UpsizeMultipler)));
 
-			Buffer = InitFunction(NewLength, InitialAccessFlags);
-			CurrentLength = NewLength;
+			FRHIResourceCreateInfo CreateInfo(DebugName);
+			const EBufferUsageFlags UsageFlags = BUF_ShaderResource | (bGpuReadOnly ? BUF_Volatile : BUF_Static | BUF_UnorderedAccess);
 
+			NumBytes = ElementBytes * NumElements;
+			Buffer = RHICreateVertexBuffer(NumBytes, UsageFlags, InResourceState, CreateInfo);
+			SRV = RHICreateShaderResourceView(Buffer, ElementBytes, UE_PIXELFORMAT_TO_UINT8(PixelFormat));
+			UAV = bGpuReadOnly ? nullptr : RHICreateUnorderedAccessView(Buffer, UE_PIXELFORMAT_TO_UINT8(PixelFormat));
 			return true;
 		}
 	}
-	else if (CurrentLength != 0)
-	{
-		Buffer.Release();
-		CurrentLength = 0;
-		return true;
-	}
-
 	return false;
 }
 
+void FNiagaraRibbonGpuBuffer::Release()
+{
+	NumBytes = 0;
+	Buffer.SafeRelease();
+	UAV.SafeRelease();
+	SRV.SafeRelease();
+}
+
+FNiagaraRibbonVertexBuffers::FNiagaraRibbonVertexBuffers()
+	: SortedIndicesBuffer(TEXT("RibbonSortedIndices"), EPixelFormat::PF_R32_UINT, sizeof(uint32))
+	, TangentsAndDistancesBuffer(TEXT("TangentsAndDistancesBuffer"), EPixelFormat::PF_R32_FLOAT, sizeof(float))
+	, MultiRibbonIndicesBuffer(TEXT("MultiRibbonIndicesBuffer"), EPixelFormat::PF_R32_UINT, sizeof(uint32))
+	, RibbonLookupTableBuffer(TEXT("RibbonLookupTableBuffer"), EPixelFormat::PF_R32_UINT, sizeof(uint32))
+	, SegmentsBuffer(TEXT("SegmentsBuffer"), EPixelFormat::PF_R32_UINT, sizeof(uint32))
+	, GPUComputeCommandBuffer(TEXT("GPUComputeCommandBuffer"), EPixelFormat::PF_R32_UINT, sizeof(uint32))
+{
+}
+
 void FNiagaraRibbonVertexBuffers::InitializeOrUpdateBuffers(const FNiagaraRibbonGenerationConfig& GenerationConfig, const TSharedPtr<FNiagaraRibbonCPUGeneratedVertexData>& GeneratedGeometryData, const FNiagaraDataBuffer* SourceParticleData, int32 MaxAllocatedCount, bool bIsUsingGPUInit)
 {	
 	const uint32 MaxAllocatedRibbons = GenerationConfig.HasRibbonIDs()? (GenerationConfig.GetMaxNumRibbons() > 0? GenerationConfig.GetMaxNumRibbons() : MaxAllocatedCount / 2) : 1;
 	
 	constexpr ERHIAccess InitialBufferAccessFlags = ERHIAccess::SRVMask | ERHIAccess::VertexOrIndexBuffer;
-	
+
 	if (!bIsUsingGPUInit)
 	{
 		check(GeneratedGeometryData.IsValid());
 		
-		InitOrUpdateBuffer(true,									SortedIndicesBuffer, SortedIndicesLength, GeneratedGeometryData->SortedIndices.Num(), MaxAllocatedCount, &CreateSortedIndicesBuffer, InitialBufferAccessFlags);
-		InitOrUpdateBuffer(true,									TangentsAndDistancesBuffer, TangentsLength, GeneratedGeometryData->TangentAndDistances.Num(), MaxAllocatedCount, &CreateTangentsAndDistancesBuffer, InitialBufferAccessFlags);
-		InitOrUpdateBuffer(GenerationConfig.HasRibbonIDs(),		MultiRibbonIndicesBuffer, MultiRibbonIndexLength, GeneratedGeometryData->MultiRibbonIndices.Num(), MaxAllocatedCount, &CreateMultiRibbonIndicesBuffer, InitialBufferAccessFlags);
-		InitOrUpdateBuffer(true,									RibbonLookupTableBuffer, RibbonLookupTableLength, GeneratedGeometryData->RibbonInfoLookup.Num(), MaxAllocatedRibbons, &CreateRibbonLookupTableBuffer, InitialBufferAccessFlags);
-		InitOrUpdateBuffer(false,								SegmentsBuffer, SegmentLength, 0, 0, &CreateSegmentsBuffer);
-		bJustCreatedCommandBuffer |= InitOrUpdateBuffer(false,	GPUComputeCommandBuffer, GPUComputeCommandLength, 0, 0, &CreateCommandBuffer);
+		SortedIndicesBuffer.Allocate(GeneratedGeometryData->SortedIndices.Num(), MaxAllocatedCount, InitialBufferAccessFlags, true);
+		TangentsAndDistancesBuffer.Allocate(GeneratedGeometryData->TangentAndDistances.Num() * 4, MaxAllocatedCount * 4, InitialBufferAccessFlags, true);
+		MultiRibbonIndicesBuffer.Allocate(GenerationConfig.HasRibbonIDs() ? GeneratedGeometryData->MultiRibbonIndices.Num() : 0, MaxAllocatedCount, InitialBufferAccessFlags, true);
+		RibbonLookupTableBuffer.Allocate(GeneratedGeometryData->RibbonInfoLookup.Num() * FRibbonMultiRibbonInfoBufferEntry::NumElements, MaxAllocatedCount * FRibbonMultiRibbonInfoBufferEntry::NumElements, InitialBufferAccessFlags, true);
+		SegmentsBuffer.Release();
+		GPUComputeCommandBuffer.Release();
+		bJustCreatedCommandBuffer = false;
 	}
 	else
 	{		
@@ -413,58 +434,15 @@ void FNiagaraRibbonVertexBuffers::InitializeOrUpdateBuffers(const FNiagaraRibbon
 		// We assume to have at least 2 particles in each ribbon, therefor we can't have more than instances/2 ribbons
 		const int32 TotalRibbons = FMath::Clamp<int32>(TotalParticles / 2, 1, MaxAllocatedRibbons);
 		
-		InitOrUpdateBuffer(true,									SortedIndicesBuffer, SortedIndicesLength, TotalParticles, MaxAllocatedCount, &CreateSortedIndicesBuffer, InitialBufferAccessFlags);
-		InitOrUpdateBuffer(true,									TangentsAndDistancesBuffer, TangentsLength, TotalParticles, MaxAllocatedCount, &CreateTangentsAndDistancesBuffer, InitialBufferAccessFlags);
-		InitOrUpdateBuffer(GenerationConfig.HasRibbonIDs(),		MultiRibbonIndicesBuffer, MultiRibbonIndexLength, TotalParticles, MaxAllocatedCount, &CreateMultiRibbonIndicesBuffer, InitialBufferAccessFlags);
-		InitOrUpdateBuffer(true,									RibbonLookupTableBuffer, RibbonLookupTableLength, TotalRibbons, MaxAllocatedRibbons, &CreateRibbonLookupTableBuffer, InitialBufferAccessFlags);
-		InitOrUpdateBuffer(true,									SegmentsBuffer, SegmentLength, TotalParticles, MaxAllocatedCount, &CreateSegmentsBuffer, InitialBufferAccessFlags);
-		bJustCreatedCommandBuffer |= InitOrUpdateBuffer(true,	GPUComputeCommandBuffer, GPUComputeCommandLength, FNiagaraRibbonCommandBufferLayout::NumElements, FNiagaraRibbonCommandBufferLayout::NumElements, &CreateCommandBuffer, InitialBufferAccessFlags | ERHIAccess::IndirectArgs);
+		SortedIndicesBuffer.Allocate(TotalParticles, MaxAllocatedCount, InitialBufferAccessFlags, false);
+		TangentsAndDistancesBuffer.Allocate(TotalParticles * 4, MaxAllocatedCount * 4, InitialBufferAccessFlags, false);
+		MultiRibbonIndicesBuffer.Allocate(GenerationConfig.HasRibbonIDs() ? TotalParticles : 0, MaxAllocatedCount, InitialBufferAccessFlags, false);
+		RibbonLookupTableBuffer.Allocate(TotalRibbons * FRibbonMultiRibbonInfoBufferEntry::NumElements, MaxAllocatedRibbons * FRibbonMultiRibbonInfoBufferEntry::NumElements, InitialBufferAccessFlags, false);
+		SegmentsBuffer.Allocate(TotalParticles, MaxAllocatedCount, InitialBufferAccessFlags, false);
+		bJustCreatedCommandBuffer |= GPUComputeCommandBuffer.Allocate(FNiagaraRibbonCommandBufferLayout::NumElements, FNiagaraRibbonCommandBufferLayout::NumElements, InitialBufferAccessFlags | ERHIAccess::IndirectArgs, false);
 	}	
 }
 
-
-FRWBuffer FNiagaraRibbonVertexBuffers::CreateSortedIndicesBuffer(int32 Size, ERHIAccess InitialAccessFlags)
-{
-	FRWBuffer NewBuffer;
-	NewBuffer.Initialize(TEXT("RibbonSortedIndices"), sizeof(uint32), Size, EPixelFormat::PF_R32_UINT, InitialAccessFlags, BUF_Static);
-	return NewBuffer;
-}
-
-FRWBuffer FNiagaraRibbonVertexBuffers::CreateTangentsAndDistancesBuffer(int32 Size, ERHIAccess InitialAccessFlags)
-{
-	FRWBuffer NewBuffer;
-	NewBuffer.Initialize(TEXT("RibbonTangentsAndDistances"), sizeof(float), Size * 4, EPixelFormat::PF_R32_FLOAT, InitialAccessFlags, BUF_Static);
-	return NewBuffer;
-}
-
-FRWBuffer FNiagaraRibbonVertexBuffers::CreateMultiRibbonIndicesBuffer(int32 Size, ERHIAccess InitialAccessFlags)
-{
-	FRWBuffer NewBuffer;
-	NewBuffer.Initialize(TEXT("RibbonMultiRibbonIndices"), sizeof(uint32), Size, EPixelFormat::PF_R32_UINT, InitialAccessFlags, BUF_Static);
-	return NewBuffer;		
-}
-
-FRWBuffer FNiagaraRibbonVertexBuffers::CreateRibbonLookupTableBuffer(int32 Size, ERHIAccess InitialAccessFlags)
-{
-	FRWBuffer NewBuffer;
-	NewBuffer.Initialize(TEXT("RibbonDataLookupTable"), sizeof(uint32), Size * FRibbonMultiRibbonInfoBufferEntry::NumElements, EPixelFormat::PF_R32_UINT, InitialAccessFlags, BUF_Static);
-	return NewBuffer;		
-}
-
-FRWBuffer FNiagaraRibbonVertexBuffers::CreateSegmentsBuffer(int32 Size, ERHIAccess InitialAccessFlags)
-{
-	FRWBuffer NewBuffer;
-	NewBuffer.Initialize(TEXT("RibbonSegments"), sizeof(uint32), Size, EPixelFormat::PF_R32_UINT, InitialAccessFlags, BUF_Static);
-	return NewBuffer;		
-}
-
-FRWBuffer FNiagaraRibbonVertexBuffers::CreateCommandBuffer(int32 Size, ERHIAccess InitialAccessFlags)
-{
-	FRWBuffer NewBuffer;
-	NewBuffer.Initialize(TEXT("RibbonGPUInitCommandBuffer"), sizeof(uint32), FNiagaraRibbonCommandBufferLayout::NumElements, EPixelFormat::PF_R32_UINT, InitialAccessFlags, BUF_DrawIndirect | BUF_Static);
-	return NewBuffer;		
-}
-
 struct FNiagaraRibbonGPUInitParameters
 {
 	const FNiagaraRendererRibbons* Renderer;
@@ -496,33 +474,27 @@ struct FNiagaraRibbonGPUInitComputeBuffers
 	void InitOrUpdateBuffers(int32 NeededSize, bool bWantsMultiRibbon, bool bWantsTessellation, bool bWantsTessellationTwist)
 	{
 		// TODO: Downsize these when we haven't needed the size for a bit
-
 		constexpr ERHIAccess InitialAccess = ERHIAccess::SRVMask | ERHIAccess::VertexOrIndexBuffer;
 		
 		if (SortBuffer.NumBytes < NeededSize * sizeof(int32))
 		{
-			SortBuffer.Initialize(TEXT("NiagarGPUInit-SortedIndices"), sizeof(uint32), NeededSize,
-			                      EPixelFormat::PF_R32_UINT, InitialAccess);
+			SortBuffer.Initialize(TEXT("NiagarGPUInit-SortedIndices"), sizeof(uint32), NeededSize, EPixelFormat::PF_R32_UINT, InitialAccess);
 		}
 
 		if (TempSegments.NumBytes < NeededSize * sizeof(int32))
 		{
-			TempSegments.Initialize(TEXT("NiagaraGPUInit-Segments"), sizeof(uint32), NeededSize,
-			                        EPixelFormat::PF_R32_UINT, InitialAccess, BUF_Static);
+			TempSegments.Initialize(TEXT("NiagaraGPUInit-Segments"), sizeof(uint32), NeededSize, EPixelFormat::PF_R32_UINT, InitialAccess, BUF_Static);
 		}
 
 		if (TempDistances.NumBytes < NeededSize * sizeof(float) * 4)
 		{
-			TempDistances.Initialize(TEXT("NiagaraGPUInit-Distances"), sizeof(float), NeededSize * 4,
-			                         EPixelFormat::PF_R32_FLOAT, InitialAccess, BUF_Static);
+			TempDistances.Initialize(TEXT("NiagaraGPUInit-Distances"), sizeof(float), NeededSize * 4, EPixelFormat::PF_R32_FLOAT, InitialAccess, BUF_Static);
 		}
 
 		const uint32 MultiRibbonBufferSize = NeededSize * (bWantsMultiRibbon? 1 : 0);
 		if (TempMultiRibbon.NumBytes < MultiRibbonBufferSize * sizeof(int32))
 		{
-			TempMultiRibbon.Initialize(
-				TEXT("NiagaraGPUInit-MultiRibbon"), sizeof(uint32), MultiRibbonBufferSize, EPixelFormat::PF_R32_UINT,
-				InitialAccess, BUF_Static);
+			TempMultiRibbon.Initialize(TEXT("NiagaraGPUInit-MultiRibbon"), sizeof(uint32), MultiRibbonBufferSize, EPixelFormat::PF_R32_UINT, InitialAccess, BUF_Static);
 		}
 
 		const uint32 TessellationBufferSize = NeededSize * (bWantsTessellation? (bWantsTessellationTwist? FTessellationStatsEntry::NumElements : FTessellationStatsEntryNoTwist::NumElements) : 0);
@@ -620,7 +592,6 @@ void FNiagaraRendererRibbons::CreateRenderThreadResources()
 {
 	FNiagaraRenderer::CreateRenderThreadResources();
 
-
 	{
 		// Initialize the shape vertex buffer. This doesn't change frame-to-frame, so we can set it up once
 		const int32 NumElements = ShapeState.SliceTriangleToVertexIds.Num();
@@ -1067,7 +1038,7 @@ void FNiagaraRendererRibbons::GenerateShapeStateTube(FNiagaraRibbonShapeGeometry
 {
 	State.Shape = ENiagaraRibbonShapeMode::Tube;
 	State.bDisableBackfaceCulling = true;
-	State.bShouldFlipNormalToView = true;
+	State.bShouldFlipNormalToView = false;
 	State.TrianglesPerSegment = 2 * TubeSubdivisions;
 	State.NumVerticesInSlice = TubeSubdivisions + 1;
 	State.BitsNeededForShape = CalculateBitsForRange(State.NumVerticesInSlice);
@@ -1095,7 +1066,7 @@ void FNiagaraRendererRibbons::GenerateShapeStateCustom(FNiagaraRibbonShapeGeomet
 {
 	State.Shape = ENiagaraRibbonShapeMode::Custom;
 	State.bDisableBackfaceCulling = true;
-	State.bShouldFlipNormalToView = true;
+	State.bShouldFlipNormalToView = false;
 	State.TrianglesPerSegment = 2 * CustomVertices.Num();
 	State.NumVerticesInSlice = CustomVertices.Num() + 1;
 	State.BitsNeededForShape = CalculateBitsForRange(State.NumVerticesInSlice);
@@ -1134,7 +1105,7 @@ void FNiagaraRendererRibbons::GenerateShapeStatePlane(FNiagaraRibbonShapeGeometr
 {
 	State.Shape = ENiagaraRibbonShapeMode::Plane;
 	State.bDisableBackfaceCulling = true;
-	State.bShouldFlipNormalToView = true;
+	State.bShouldFlipNormalToView = false;
 	State.TrianglesPerSegment = 2 * WidthSegmentationCount;
 	State.NumVerticesInSlice = WidthSegmentationCount + 1;
 	State.BitsNeededForShape = CalculateBitsForRange(State.NumVerticesInSlice);
@@ -1305,6 +1276,9 @@ void FNiagaraRendererRibbons::GenerateVertexBufferForRibbonPart(const FNiagaraGe
 	float LastWidth = 0;
 	double TotalDistance = 0.0f;
 
+	OutputData.SortedIndices.Reserve(OutputData.SortedIndices.Num() + RibbonIndices.Num());
+	OutputData.TangentAndDistances.Reserve(OutputData.TangentAndDistances.Num() + RibbonIndices.Num());
+
 	// Find the first position with enough distance.
 	int32 CurrentIndex = 1;
 	while (CurrentIndex < RibbonIndices.Num())
@@ -1500,6 +1474,7 @@ void FNiagaraRendererRibbons::GenerateVertexBufferForRibbonPart(const FNiagaraGe
 template<typename IDType, typename ReaderType, bool bWantsTessellation, bool bHasTwist>
 void FNiagaraRendererRibbons::GenerateVertexBufferForMultiRibbonInternal(const FNiagaraGenerationInputDataCPUAccessors& CPUData, const ReaderType& IDReader, FNiagaraRibbonCPUGeneratedVertexData& OutputData) const
 {
+	//-OPT: Consider MemStack
 	TMap<IDType, TArray<uint32>> MultiRibbonSortedIndices;
 
 	for (uint32 i = 0; i < CPUData.TotalNumParticles; ++i)
@@ -1565,8 +1540,10 @@ void FNiagaraRendererRibbons::GenerateVertexBufferCPU(const FNiagaraGenerationIn
 		}		
 	}
 	else
-	{		
+	{
+		//-OPT: Consider MemStack
 		TArray<uint32> SortedIndices;
+		SortedIndices.Reserve(CPUData.TotalNumParticles + 1);
 		for (uint32 i = 0; i < CPUData.TotalNumParticles; ++i)
 		{
 			SortedIndices.Add(i);
@@ -1737,6 +1714,7 @@ FNiagaraIndexGenerationInput FNiagaraRendererRibbons::CalculateIndexBufferConfig
 	
 	IndexGenInput.TotalBitCount = NumSegmentBits + NumSubSegmentBits + ShapeState.BitsNeededForShape;
 	IndexGenInput.TotalNumIndices = IndexGenInput.MaxSegmentCount * IndexGenInput.SubSegmentCount * ShapeState.TrianglesPerSegment * 3;
+	IndexGenInput.CPUTriangleCount = 0;
 
 	return IndexGenInput;
 }
@@ -1752,7 +1730,8 @@ void FNiagaraRendererRibbons::GenerateIndexBufferForView(FNiagaraIndexGeneration
 			RenderingViewResources->IndirectDrawBuffer.Initialize(TEXT("RibbonIndirectDrawBuffer"), sizeof(uint32), FNiagaraRibbonIndirectDrawBufferLayout::NumElements, EPixelFormat::PF_R32_UINT, ERHIAccess::IndirectArgs | ERHIAccess::SRVMask, BUF_Static | BUF_DrawIndirect);			
 		}
 		
-		if (GeneratedData.TotalBitCount <= 16/* Number of bits in a ushort*/ && !bShouldUseGPUInitIndices)
+		// Can we pack into a uint16 vs a uint32?
+		if (GeneratedData.TotalBitCount <= 16 && !bShouldUseGPUInitIndices)
 		{
 			RenderingViewResources->IndexBuffer.Initialize(TEXT("NiagaraRibbonIndexBuffer"), sizeof(uint16), GeneratedData.TotalNumIndices, PF_R16_UINT, ERHIAccess::VertexOrIndexBuffer, BUF_Static);
 			if (!bShouldUseGPUInitIndices)
@@ -1782,9 +1761,8 @@ void FNiagaraRendererRibbons::GenerateIndexBufferCPU(FNiagaraIndexGenerationInpu
 	
 	const TSharedPtr<FNiagaraRibbonCPUGeneratedVertexData>& GeneratedGeometryData = DynamicDataRibbon->GenerationOutput;
 
-	TArray<TValue> IndexBufferData;
-	IndexBufferData.SetNum(GeneratedData.TotalNumIndices);
-	TValue* CurrentIndexBuffer = IndexBufferData.GetData();
+	TValue* StartIndexBuffer = reinterpret_cast<TValue*>(RHILockBuffer(RenderingViewResources->IndexBuffer.IndexBufferRHI, 0, GeneratedData.TotalNumIndices * sizeof(TValue), RLM_WriteOnly));
+	TValue* CurrentIndexBuffer = StartIndexBuffer;
 	if (IsTranslucentBlendMode(BlendMode) && GeneratedGeometryData->RibbonInfoLookup.Num())
 	{
 		for (const FRibbonMultiRibbonInfo& MultiRibbonInfo : GeneratedGeometryData->RibbonInfoLookup)
@@ -1798,15 +1776,11 @@ void FNiagaraRendererRibbons::GenerateIndexBufferCPU(FNiagaraIndexGenerationInpu
 		const TArrayView<uint32> CurrentSegmentData(GeneratedGeometryData->SegmentData.GetData(), GeneratedGeometryData->SegmentData.Num());
 		CurrentIndexBuffer = AppendToIndexBufferCPU<TValue>(CurrentIndexBuffer, GeneratedData, ShapeState, CurrentSegmentData, false);
 	}
-	
-	// Copy the indices over to the GPU
-	void* IndexBufferPtr = RHILockBuffer(RenderingViewResources->IndexBuffer.IndexBufferRHI, 0, IndexBufferData.Num() * sizeof(TValue), RLM_WriteOnly);
-	FMemory::Memcpy(IndexBufferPtr, IndexBufferData.GetData(), IndexBufferData.Num() * sizeof(TValue));
+	GeneratedData.CPUTriangleCount = (CurrentIndexBuffer - StartIndexBuffer) / 3;
+	check(CurrentIndexBuffer <= StartIndexBuffer + GeneratedData.TotalNumIndices);
 	RHIUnlockBuffer(RenderingViewResources->IndexBuffer.IndexBufferRHI);
 }
 
-
-
 template <typename TValue>
 TValue* FNiagaraRendererRibbons::AppendToIndexBufferCPU(TValue* OutIndices, const FNiagaraIndexGenerationInput& GeneratedData, const FNiagaraRibbonShapeGeometryData& ShapeState, const TArrayView<uint32>& SegmentData, bool bInvertOrder)
 {
@@ -1815,28 +1789,28 @@ TValue* FNiagaraRendererRibbons::AppendToIndexBufferCPU(TValue* OutIndices, cons
 		return OutIndices;
 	}
 
-	const int32 FirstSegmentDataIndex = bInvertOrder ? SegmentData.Num() - 1 : 0;
-	const int32 LastSegmentDataIndex = bInvertOrder ? -1 : SegmentData.Num();
-	const int32 SegmentDataIndexInc = bInvertOrder ? -1 : 1;
-	const int32 FlipGeometryIndex = FMath::Min(FMath::Max(ShapeState.SliceTriangleToVertexIds.Num() / 2, 2), ShapeState.SliceTriangleToVertexIds.Num());
+	const uint32 FirstSegmentDataIndex = bInvertOrder ? SegmentData.Num() - 1 : 0;
+	const uint32 LastSegmentDataIndex = bInvertOrder ? -1 : SegmentData.Num();
+	const uint32 SegmentDataIndexInc = bInvertOrder ? -1 : 1;
+	const uint32 FlipGeometryIndex = FMath::Min(FMath::Max(ShapeState.SliceTriangleToVertexIds.Num() / 2, 2), ShapeState.SliceTriangleToVertexIds.Num());
 	
-	for (int32 SegmentDataIndex = FirstSegmentDataIndex; SegmentDataIndex != LastSegmentDataIndex; SegmentDataIndex += SegmentDataIndexInc)
+	for (uint32 SegmentDataIndex = FirstSegmentDataIndex; SegmentDataIndex != LastSegmentDataIndex; SegmentDataIndex += SegmentDataIndexInc)
 	{
-		const int32 SegmentIndex = SegmentData[SegmentDataIndex];
+		const uint32 SegmentIndex = SegmentData[SegmentDataIndex];
 		for (uint32 SubSegmentIndex = 0; SubSegmentIndex < GeneratedData.SubSegmentCount; ++SubSegmentIndex)
 		{
 			const bool bIsFinalInterp = SubSegmentIndex == GeneratedData.SubSegmentCount - 1;
 
-			const int32 ThisSegmentOffset = SegmentIndex << GeneratedData.SegmentBitShift;
-			const int32 NextSegmentOffset = (SegmentIndex + (bIsFinalInterp ? 1 : 0)) << GeneratedData.SegmentBitShift;
+			const uint32 ThisSegmentOffset = SegmentIndex << GeneratedData.SegmentBitShift;
+			const uint32 NextSegmentOffset = (SegmentIndex + (bIsFinalInterp ? 1 : 0)) << GeneratedData.SegmentBitShift;
 
-			const int32 ThisSubSegmentOffset = SubSegmentIndex << GeneratedData.SubSegmentBitShift;
-			const int32 NextSubSegmentOffset = (bIsFinalInterp ? 0 : SubSegmentIndex + 1) << GeneratedData.SubSegmentBitShift;
+			const uint32 ThisSubSegmentOffset = SubSegmentIndex << GeneratedData.SubSegmentBitShift;
+			const uint32 NextSubSegmentOffset = (bIsFinalInterp ? 0 : SubSegmentIndex + 1) << GeneratedData.SubSegmentBitShift;
 
-			const int32 CurrSegment = ThisSegmentOffset | ThisSubSegmentOffset;
-			const int32 NextSegment = NextSegmentOffset | NextSubSegmentOffset;
+			const uint32 CurrSegment = ThisSegmentOffset | ThisSubSegmentOffset;
+			const uint32 NextSegment = NextSegmentOffset | NextSubSegmentOffset;
 
-			int32 TriangleId = 0;
+			uint32 TriangleId = 0;
 			
 			for (; TriangleId < FlipGeometryIndex; TriangleId += 2)
 			{
@@ -1852,10 +1826,10 @@ TValue* FNiagaraRendererRibbons::AppendToIndexBufferCPU(TValue* OutIndices, cons
 
 				OutIndices += 6;
 			}
-			for (; TriangleId < ShapeState.SliceTriangleToVertexIds.Num(); TriangleId += 2)
+			for (; TriangleId < uint32(ShapeState.SliceTriangleToVertexIds.Num()); TriangleId += 2)
 			{
-				const int32 FirstIndex = ShapeState.SliceTriangleToVertexIds[TriangleId];
-				const int32 SecondIndex = ShapeState.SliceTriangleToVertexIds[TriangleId + 1];
+				const uint32 FirstIndex = ShapeState.SliceTriangleToVertexIds[TriangleId];
+				const uint32 SecondIndex = ShapeState.SliceTriangleToVertexIds[TriangleId + 1];
 
 				OutIndices[0] = CurrSegment | FirstIndex;
 				OutIndices[1] = CurrSegment | SecondIndex;
@@ -1998,7 +1972,7 @@ inline void FNiagaraRendererRibbons::SetupMeshBatchAndCollectorResourceForView(c
 	}
 	else
 	{
-		MeshElement.NumPrimitives = GeneratedData.TotalNumIndices / 3; // 3 indices per triangle
+		MeshElement.NumPrimitives = GeneratedData.CPUTriangleCount;
 		MeshElement.MaxVertexIndex = SourceParticleData->GetNumInstances();
 		check(MeshElement.NumPrimitives > 0);
 	}	
@@ -2209,6 +2183,7 @@ void FNiagaraRendererRibbons::InitializeVertexBuffersResources(const FNiagaraDyn
 		}
 		
 		// Copy the packed u data for stable age based uv generation.
+		//-OPT: Remove copy, push straight into GPU Memory
 		TArray<uint32> PackedRibbonLookupTable;
 		PackedRibbonLookupTable.Reserve(GeneratedGeometryData.RibbonInfoLookup.Num() * FRibbonMultiRibbonInfoBufferEntry::NumElements);
 		for (int32 Index = 0; Index < GeneratedGeometryData.RibbonInfoLookup.Num(); Index++)
@@ -2298,7 +2273,7 @@ void FNiagaraRendererRibbons::InitializeVertexBuffersGPU(FRHICommandListImmediat
 	const bool bCanRun = NumExecutableInstances >= 2;
 	
 	// Clear the command buffer if we just initialized it, or if the sim doesn't have enough data to run
-	if ((!bCanRun || VertexBuffers.bJustCreatedCommandBuffer) && VertexBuffers.GPUComputeCommandLength > 0)
+	if ((!bCanRun || VertexBuffers.bJustCreatedCommandBuffer) && VertexBuffers.GPUComputeCommandBuffer.NumBytes > 0)
 	{
 		RHICmdList.Transition(FRHITransitionInfo(VertexBuffers.GPUComputeCommandBuffer.Buffer, ERHIAccess::SRVMask | ERHIAccess::VertexOrIndexBuffer | ERHIAccess::IndirectArgs, ERHIAccess::UAVCompute));
 		RHICmdList.ClearUAVUint(VertexBuffers.GPUComputeCommandBuffer.UAV, FUintVector4(0));
diff --git a/Engine/Plugins/FX/Niagara/Source/Niagara/Public/NiagaraRendererRibbons.h b/Engine/Plugins/FX/Niagara/Source/Niagara/Public/NiagaraRendererRibbons.h
--- a/Engine/Plugins/FX/Niagara/Source/Niagara/Public/NiagaraRendererRibbons.h
+++ b/Engine/Plugins/FX/Niagara/Source/Niagara/Public/NiagaraRendererRibbons.h
@@ -56,22 +56,24 @@ struct FNiagaraGenerationInputDataCPUAccessors
 
 struct FNiagaraIndexGenerationInput
 {
-	float ViewDistance;
-	int32 LODDistanceFactor;
+	float ViewDistance = 0.0f;
+	int32 LODDistanceFactor = 0;
 	
-	uint32 MaxSegmentCount;
-	uint32 SubSegmentCount;	
+	uint32 MaxSegmentCount = 0;
+	uint32 SubSegmentCount = 0;
 
-	uint32 SegmentBitShift;
-	uint32 SegmentBitMask;
+	uint32 SegmentBitShift = 0;
+	uint32 SegmentBitMask = 0;
 
-	uint32 SubSegmentBitShift;
-	uint32 SubSegmentBitMask;
+	uint32 SubSegmentBitShift = 0;
+	uint32 SubSegmentBitMask = 0;
 
-	uint32 ShapeBitMask;
+	uint32 ShapeBitMask = 0;
 	
-	uint32 TotalBitCount;
-	uint32 TotalNumIndices;
+	uint32 TotalBitCount = 0;
+	uint32 TotalNumIndices = 0;
+
+	uint32 CPUTriangleCount = 0;
 };
 
 struct FNiagaraRibbonGenerationConfig
@@ -240,40 +242,46 @@ struct FNiagaraRibbonTessellationSmoothingData
 	float TessellationTotalSegmentLength = 0;
 };
 
-struct FNiagaraRibbonVertexBuffers
+struct FNiagaraRibbonGpuBuffer
 {
-	FRWBuffer SortedIndicesBuffer;
-	FRWBuffer TangentsAndDistancesBuffer;
-	FRWBuffer MultiRibbonIndicesBuffer;
-	FRWBuffer RibbonLookupTableBuffer;
-	FRWBuffer SegmentsBuffer;
-
-	FRWBuffer GPUComputeCommandBuffer;
-
-	int32 SortedIndicesLength;
-	int32 TangentsLength;
-	int32 MultiRibbonIndexLength;
-	int32 RibbonLookupTableLength;
-	int32 SegmentLength;
-	int32 GPUComputeCommandLength;
-	bool bJustCreatedCommandBuffer;
-
-	FNiagaraRibbonVertexBuffers()
-		: SortedIndicesLength(0)
-		, TangentsLength(0)
-		, MultiRibbonIndexLength(0)
-		, RibbonLookupTableLength(0)
-		, SegmentLength(0)
-		, GPUComputeCommandLength(0)
-		, bJustCreatedCommandBuffer(false)
+	FNiagaraRibbonGpuBuffer(const TCHAR* InDebugName, EPixelFormat InPixelFormat, uint32 InElementBytes)
+		: DebugName(InDebugName)
+		, PixelFormat(InPixelFormat)
+		, ElementBytes(InElementBytes)
 	{
-		
 	}
 
-	static bool InitOrUpdateBuffer(bool bEnabled, FRWBuffer& Buffer, int32& CurrentLength, int32 NeededLength, int32 MaxLength, FRWBuffer (*InitFunction)(int32, ERHIAccess), ERHIAccess InitialAccessFlags = ERHIAccess::None);
+	~FNiagaraRibbonGpuBuffer()
+	{
+		Release();
+	}
+
+	bool Allocate(uint32 NumElements, uint32 MaxElements, ERHIAccess InResourceState, bool bGpuReadOnly);
+	void Release();
+
+	const TCHAR*				DebugName = nullptr;
+	const EPixelFormat			PixelFormat = PF_Unknown;
+	const uint32				ElementBytes = 0;
+
+	uint32						NumBytes = 0;
+	FBufferRHIRef				Buffer;
+	FUnorderedAccessViewRHIRef	UAV;
+	FShaderResourceViewRHIRef	SRV;
+};
+
+struct FNiagaraRibbonVertexBuffers
+{
+	FNiagaraRibbonVertexBuffers();
+
+	FNiagaraRibbonGpuBuffer SortedIndicesBuffer;
+	FNiagaraRibbonGpuBuffer TangentsAndDistancesBuffer;
+	FNiagaraRibbonGpuBuffer MultiRibbonIndicesBuffer;
+	FNiagaraRibbonGpuBuffer RibbonLookupTableBuffer;
+	FNiagaraRibbonGpuBuffer SegmentsBuffer;
+	FNiagaraRibbonGpuBuffer GPUComputeCommandBuffer;
+	bool bJustCreatedCommandBuffer = false;
 
-	void InitializeOrUpdateBuffers(const FNiagaraRibbonGenerationConfig& GenerationConfig, const TSharedPtr<FNiagaraRibbonCPUGeneratedVertexData>& GeneratedGeometryData,
-		const FNiagaraDataBuffer* SourceParticleData, int32 MaxAllocatedCount, bool bIsUsingGPUInit);
+	void InitializeOrUpdateBuffers(const FNiagaraRibbonGenerationConfig& GenerationConfig, const TSharedPtr<FNiagaraRibbonCPUGeneratedVertexData>& GeneratedGeometryData, const FNiagaraDataBuffer* SourceParticleData, int32 MaxAllocatedCount, bool bIsUsingGPUInit);
 
 	void Release()
 	{
@@ -283,21 +291,7 @@ struct FNiagaraRibbonVertexBuffers
 		RibbonLookupTableBuffer.Release();
 		SegmentsBuffer.Release();
 		GPUComputeCommandBuffer.Release();
-		
-		SortedIndicesLength = 0;
-		TangentsLength = 0;
-		MultiRibbonIndexLength = 0;
-		RibbonLookupTableLength = 0;
-		SegmentLength = 0;
-		GPUComputeCommandLength = 0;
 	}
-
-	static FRWBuffer CreateSortedIndicesBuffer(int32 Size, ERHIAccess InitialAccessFlags = ERHIAccess::None);
-	static FRWBuffer CreateTangentsAndDistancesBuffer(int32 Size, ERHIAccess InitialAccessFlags = ERHIAccess::None);
-	static FRWBuffer CreateMultiRibbonIndicesBuffer(int32 Size, ERHIAccess InitialAccessFlags = ERHIAccess::None);
-	static FRWBuffer CreateRibbonLookupTableBuffer(int32 Size, ERHIAccess InitialAccessFlags = ERHIAccess::None);
-	static FRWBuffer CreateSegmentsBuffer(int32 Size, ERHIAccess InitialAccessFlags = ERHIAccess::None);
-	static FRWBuffer CreateCommandBuffer(int32 Size, ERHIAccess InitialAccessFlags = ERHIAccess::None);
 };
 
 struct FNiagaraRibbonShapeGeometryData
@@ -438,10 +432,3 @@ protected:
 
 	friend class FNiagaraRibbonComputeDispatchManager;
 };
-
-
-
-
-
-
-
