- Remove event handler in FinishDestroy

#rnx

diff --git a/Engine/Plugins/FX/Niagara/Source/Niagara/Classes/NiagaraDataInterfaceSpriteRendererInfo.h b/Engine/Plugins/FX/Niagara/Source/Niagara/Classes/NiagaraDataInterfaceSpriteRendererInfo.h
--- a/Engine/Plugins/FX/Niagara/Source/Niagara/Classes/NiagaraDataInterfaceSpriteRendererInfo.h
+++ b/Engine/Plugins/FX/Niagara/Source/Niagara/Classes/NiagaraDataInterfaceSpriteRendererInfo.h
@@ -21,6 +21,7 @@ public:
 	//UObject Interface
 	virtual void PostInitProperties() override;
 	virtual void PostLoad() override;
+	virtual void BeginDestroy() override;
 #if WITH_EDITOR	
 	virtual void PreEditChange(FProperty* PropertyAboutToChange) override;
 	virtual void PostEditChangeProperty(struct FPropertyChangedEvent& PropertyChangedEvent) override;
diff --git a/Engine/Plugins/FX/Niagara/Source/Niagara/Private/NiagaraDataInterfaceMeshRendererInfo.cpp b/Engine/Plugins/FX/Niagara/Source/Niagara/Private/NiagaraDataInterfaceMeshRendererInfo.cpp
--- a/Engine/Plugins/FX/Niagara/Source/Niagara/Private/NiagaraDataInterfaceMeshRendererInfo.cpp
+++ b/Engine/Plugins/FX/Niagara/Source/Niagara/Private/NiagaraDataInterfaceMeshRendererInfo.cpp
@@ -84,9 +84,10 @@ void UNiagaraDataInterfaceMeshRendererInfo::PostLoad()
 
 void UNiagaraDataInterfaceMeshRendererInfo::BeginDestroy()
 {
-	Super::BeginDestroy();
-
+#if WITH_EDITOR
 	OnMeshRendererChanged(nullptr);
+#endif
+	Super::BeginDestroy();
 }
 
 #if WITH_EDITOR
diff --git a/Engine/Plugins/FX/Niagara/Source/Niagara/Private/NiagaraDataInterfaceSpriteRendererInfo.cpp b/Engine/Plugins/FX/Niagara/Source/Niagara/Private/NiagaraDataInterfaceSpriteRendererInfo.cpp
--- a/Engine/Plugins/FX/Niagara/Source/Niagara/Private/NiagaraDataInterfaceSpriteRendererInfo.cpp
+++ b/Engine/Plugins/FX/Niagara/Source/Niagara/Private/NiagaraDataInterfaceSpriteRendererInfo.cpp
@@ -94,6 +94,14 @@ void UNiagaraDataInterfaceSpriteRendererInfo::PostLoad()
 	}
 }
 
+void UNiagaraDataInterfaceSpriteRendererInfo::BeginDestroy()
+{
+#if WITH_EDITOR
+	OnSpriteRendererChanged(nullptr);
+#endif
+	Super::BeginDestroy();
+}
+
 #if WITH_EDITOR
 void UNiagaraDataInterfaceSpriteRendererInfo::PreEditChange(FProperty* PropertyAboutToChange)
 {
