- Ensure force flush captures MaxPendingTicks otherwise nothing processes

#rnx

diff --git a/Engine/Plugins/FX/Niagara/Source/Niagara/Private/NiagaraGpuComputeDispatch.cpp b/Engine/Plugins/FX/Niagara/Source/Niagara/Private/NiagaraGpuComputeDispatch.cpp
--- a/Engine/Plugins/FX/Niagara/Source/Niagara/Private/NiagaraGpuComputeDispatch.cpp
+++ b/Engine/Plugins/FX/Niagara/Source/Niagara/Private/NiagaraGpuComputeDispatch.cpp
@@ -341,7 +341,7 @@ void FNiagaraGpuComputeDispatch::ProcessPendingTicksFlush(FRHICommandListImmedia
 	int32 MaxPendingTicks = 0;
 	const int32 TickFlushMaxPendingTicks = FMath::Max(FNiagaraGpuComputeDispatchLocal::GTickFlushMaxPendingTicks, 1);
 
-	if (!bForceFlush && (FramesBeforeTickFlush < uint32(FNiagaraGpuComputeDispatchLocal::GTickFlushMaxQueuedFrames)) )
+	if (bForceFlush || (FramesBeforeTickFlush < uint32(FNiagaraGpuComputeDispatchLocal::GTickFlushMaxQueuedFrames)))
 	{
 		// If we have a lot of queued up ticks we will need to flush in batches as RDG can only handle a limited number of passes in a single Graph.
 		for (int iTickStage = 0; iTickStage < ENiagaraGpuComputeTickStage::Max; ++iTickStage)
