- Experimental VirtualTexture DataInterface

diff --git a/Engine/Plugins/FX/Niagara/Shaders/Private/NiagaraDataInterfaceVirtualTexture.ush b/Engine/Plugins/FX/Niagara/Shaders/Private/NiagaraDataInterfaceVirtualTexture.ush
new file mode 100644
--- /dev/null
+++ b/Engine/Plugins/FX/Niagara/Shaders/Private/NiagaraDataInterfaceVirtualTexture.ush
@@ -0,0 +1,20 @@
+// Copyright Epic Games, Inc. All Rights Reserved.
+
+#pragma once
+
+#ifndef NUM_VIRTUALTEXTURE_SAMPLES
+	#define NUM_VIRTUALTEXTURE_SAMPLES 1
+#endif
+
+#include "/Engine/Private/VirtualTextureCommon.ush"
+
+// Mirror of ERuntimeVirtualTextureMaterialType
+#define ERuntimeVirtualTextureMaterialType_BaseColor								0
+#define ERuntimeVirtualTextureMaterialType_BaseColor_Normal_DEPRECATED				1
+#define ERuntimeVirtualTextureMaterialType_BaseColor_Normal_Roughness				2
+#define ERuntimeVirtualTextureMaterialType_BaseColor_Normal_Specular				3
+#define ERuntimeVirtualTextureMaterialType_BaseColor_Normal_Specular_YCoCg			4
+#define ERuntimeVirtualTextureMaterialType_BaseColor_Normal_Specular_Mask_YCoCg		5
+#define ERuntimeVirtualTextureMaterialType_WorldHeight								6
+#define ERuntimeVirtualTextureMaterialType_Count									7
+
diff --git a/Engine/Plugins/FX/Niagara/Shaders/Private/NiagaraDataInterfaceVirtualTextureTemplate.ush b/Engine/Plugins/FX/Niagara/Shaders/Private/NiagaraDataInterfaceVirtualTextureTemplate.ush
new file mode 100644
--- /dev/null
+++ b/Engine/Plugins/FX/Niagara/Shaders/Private/NiagaraDataInterfaceVirtualTextureTemplate.ush
@@ -0,0 +1,148 @@
+// Copyright Epic Games, Inc. All Rights Reserved.
+
+Texture2D			{ParameterName}_VirtualTexture0;
+Texture2D<uint4>	{ParameterName}_VirtualTexture0PageTable;
+uint4				{ParameterName}_VirtualTexture0TextureUniforms;
+
+Texture2D			{ParameterName}_VirtualTexture1;
+Texture2D<uint4>	{ParameterName}_VirtualTexture1PageTable;
+uint4				{ParameterName}_VirtualTexture1TextureUniforms;
+
+Texture2D			{ParameterName}_VirtualTexture2;
+Texture2D<uint4>	{ParameterName}_VirtualTexture2PageTable;
+uint4				{ParameterName}_VirtualTexture2TextureUniforms;
+
+uint				{ParameterName}_ValidLayersMask;
+int					{ParameterName}_MaterialType;
+float4				{ParameterName}_UVUniforms[3];
+uint4				{ParameterName}_PageTableUniforms[2];
+float2				{ParameterName}_WorldHeightUnpack;
+SamplerState		{ParameterName}_SharedSampler;
+
+void GetAttributesValid_{ParameterName}(out bool BaseColorValid, out bool SpecularValid, out bool RoughnessValid, out bool NormalValid, out bool WorldHeightValid, out bool MaskValid)
+{
+	BaseColorValid = false;
+	SpecularValid = false;
+	RoughnessValid = false;
+	NormalValid = false;
+	WorldHeightValid = false;
+	MaskValid = false;
+
+	switch ( {ParameterName}_MaterialType )
+	{
+		case ERuntimeVirtualTextureMaterialType_BaseColor:
+			BaseColorValid = true;
+			break;
+
+		case ERuntimeVirtualTextureMaterialType_BaseColor_Normal_Roughness:
+			BaseColorValid = true;
+			RoughnessValid = true;
+			NormalValid = true;
+			break;
+
+		case ERuntimeVirtualTextureMaterialType_BaseColor_Normal_DEPRECATED:
+		case ERuntimeVirtualTextureMaterialType_BaseColor_Normal_Specular:
+			BaseColorValid = true;
+			SpecularValid = true;
+			RoughnessValid = true;
+			NormalValid = true;
+			break;
+		
+		case ERuntimeVirtualTextureMaterialType_BaseColor_Normal_Specular_YCoCg:
+			BaseColorValid = true;
+			SpecularValid = true;
+			RoughnessValid = true;
+			NormalValid = true;
+			break;
+		
+		case ERuntimeVirtualTextureMaterialType_BaseColor_Normal_Specular_Mask_YCoCg:
+			BaseColorValid = true;
+			SpecularValid = true;
+			RoughnessValid = true;
+			NormalValid = true;
+			MaskValid = true;
+			break;
+		
+		case ERuntimeVirtualTextureMaterialType_WorldHeight:
+			WorldHeightValid = true;
+			break;
+	}
+}
+
+float4 SampleRVTLayer_{ParameterName}(float2 SampleUV, Texture2D InTexture, Texture2D<uint4> InPageTable, uint4 InTextureUniforms)
+{
+	VTPageTableResult PageTable = TextureLoadVirtualPageTableLevel(InPageTable, VTPageTableUniform_Unpack({ParameterName}_PageTableUniforms[0], {ParameterName}_PageTableUniforms[1]), SampleUV, VTADDRESSMODE_CLAMP, VTADDRESSMODE_CLAMP, 0.0f);
+	return TextureVirtualSample(InTexture, {ParameterName}_SharedSampler, PageTable, 0, VTUniform_Unpack(InTextureUniforms));
+}
+
+void SampleRVT_{ParameterName}(in float3 WorldPosition, out float3 BaseColor, out float Specular, out float Roughness, out float3 Normal, out float WorldHeight, out float Mask)
+{
+	BaseColor = float3(0.0f, 0.0f, 0.0f);
+	Specular = 0.5f;
+	Roughness = 0.5f;
+	Normal = float3(0.0f, 0.0f, 1.0f);
+	WorldHeight = 0.0f;
+	Mask = 1.0f;
+
+	// Get Sample Location
+	float2 SampleUV = VirtualTextureWorldToUV(LWCPromote(WorldPosition), LWCPromote({ParameterName}_UVUniforms[0].xyz), {ParameterName}_UVUniforms[1].xyz, {ParameterName}_UVUniforms[2].xyz);
+
+	// Sample Textures
+	float4 LayerSample[3];
+	LayerSample[0] = ({ParameterName}_ValidLayersMask & 0x1) != 0 ? SampleRVTLayer_{ParameterName}(SampleUV, {ParameterName}_VirtualTexture0, {ParameterName}_VirtualTexture0PageTable, {ParameterName}_VirtualTexture0TextureUniforms) : 0;
+	LayerSample[1] = ({ParameterName}_ValidLayersMask & 0x2) != 0 ? SampleRVTLayer_{ParameterName}(SampleUV, {ParameterName}_VirtualTexture1, {ParameterName}_VirtualTexture1PageTable, {ParameterName}_VirtualTexture1TextureUniforms) : 0;
+	LayerSample[2] = ({ParameterName}_ValidLayersMask & 0x4) != 0 ? SampleRVTLayer_{ParameterName}(SampleUV, {ParameterName}_VirtualTexture2, {ParameterName}_VirtualTexture2PageTable, {ParameterName}_VirtualTexture2TextureUniforms) : 0;
+
+	// Sample Available Attributes
+	switch ( {ParameterName}_MaterialType )
+	{
+		case ERuntimeVirtualTextureMaterialType_BaseColor:
+		{
+			BaseColor = LayerSample[0].xyz;
+			break;
+		}
+
+		case ERuntimeVirtualTextureMaterialType_BaseColor_Normal_Roughness:
+		{
+			BaseColor = VirtualTextureUnpackBaseColorSRGB(LayerSample[0]);
+			Roughness = LayerSample[1].y;
+			Normal = VirtualTextureUnpackNormalBGR565(LayerSample[1]);
+			break;
+		}
+
+		case ERuntimeVirtualTextureMaterialType_BaseColor_Normal_DEPRECATED:
+		case ERuntimeVirtualTextureMaterialType_BaseColor_Normal_Specular:
+		{
+			BaseColor = LayerSample[0].xyz;
+			Specular = LayerSample[1].x;
+			Roughness = LayerSample[1].y;
+			Normal = VirtualTextureUnpackNormalBC3BC3(LayerSample[0], LayerSample[1]);
+			break;
+		}
+		
+		case ERuntimeVirtualTextureMaterialType_BaseColor_Normal_Specular_YCoCg:
+		{
+			BaseColor = VirtualTextureUnpackBaseColorYCoCg(LayerSample[0]);
+			Specular = LayerSample[2].x;
+			Roughness = LayerSample[2].y;
+			Normal = VirtualTextureUnpackNormalBC5BC1(LayerSample[1], LayerSample[2]);
+			break;
+		}
+		
+		case ERuntimeVirtualTextureMaterialType_BaseColor_Normal_Specular_Mask_YCoCg:
+		{
+			BaseColor = VirtualTextureUnpackBaseColorYCoCg(LayerSample[0]);
+			Specular = LayerSample[2].x;
+			Roughness = LayerSample[2].y;
+			Normal = VirtualTextureUnpackNormalBC5BC1(LayerSample[1], LayerSample[2]);
+			Mask = LayerSample[2].w;
+			break;
+		}
+		
+		case ERuntimeVirtualTextureMaterialType_WorldHeight:
+		{
+			WorldHeight = VirtualTextureUnpackHeight(LayerSample[0], {ParameterName}_WorldHeightUnpack);
+			break;
+		}
+	}
+}
diff --git a/Engine/Plugins/FX/Niagara/Source/Niagara/Private/DataInterface/NiagaraDataInterfaceVirtualTexture.cpp b/Engine/Plugins/FX/Niagara/Source/Niagara/Private/DataInterface/NiagaraDataInterfaceVirtualTexture.cpp
new file mode 100644
--- /dev/null
+++ b/Engine/Plugins/FX/Niagara/Source/Niagara/Private/DataInterface/NiagaraDataInterfaceVirtualTexture.cpp
@@ -0,0 +1,450 @@
+// Copyright Epic Games, Inc. All Rights Reserved.
+
+#include "NiagaraDataInterfaceVirtualTexture.h"
+#include "NiagaraComputeExecutionContext.h"
+#include "NiagaraCustomVersion.h"
+#include "NiagaraGpuComputeDispatchInterface.h"
+#include "NiagaraShader.h"
+#include "NiagaraShaderParametersBuilder.h"
+#include "NiagaraSystemInstance.h"
+
+#include "VT/RuntimeVirtualTexture.h"
+
+#define LOCTEXT_NAMESPACE "UNiagaraDataInterfaceVirtualTexture"
+
+namespace NDIVirtualTextureLocal
+{
+	BEGIN_SHADER_PARAMETER_STRUCT(FShaderParameters, )
+		SHADER_PARAMETER_TEXTURE(Texture2D,			VirtualTexture0)
+		SHADER_PARAMETER_TEXTURE(Texture2D<uint4>,	VirtualTexture0PageTable)
+		SHADER_PARAMETER(FUintVector4,				VirtualTexture0TextureUniforms)
+
+		SHADER_PARAMETER_TEXTURE(Texture2D,			VirtualTexture1)
+		SHADER_PARAMETER_TEXTURE(Texture2D<uint4>,	VirtualTexture1PageTable)
+		SHADER_PARAMETER(FUintVector4,				VirtualTexture1TextureUniforms)
+
+		SHADER_PARAMETER_TEXTURE(Texture2D,			VirtualTexture2)
+		SHADER_PARAMETER_TEXTURE(Texture2D<uint4>,	VirtualTexture2PageTable)
+		SHADER_PARAMETER(FUintVector4,				VirtualTexture2TextureUniforms)
+
+		SHADER_PARAMETER(uint32,					ValidLayersMask)
+		SHADER_PARAMETER(int32,						MaterialType)
+		SHADER_PARAMETER_ARRAY(FUintVector4,		PageTableUniforms, [2])
+		SHADER_PARAMETER_ARRAY(FVector4f,			UVUniforms, [3])
+		SHADER_PARAMETER(FVector2f,					WorldHeightUnpack)
+		SHADER_PARAMETER_SAMPLER(SamplerState,		SharedSampler)
+	END_SHADER_PARAMETER_STRUCT()
+
+	const TCHAR*	CommonShaderFilePath = TEXT("/Plugin/FX/Niagara/Private/NiagaraDataInterfaceVirtualTexture.ush");
+	const TCHAR*	TemplateShaderFilePath = TEXT("/Plugin/FX/Niagara/Private/NiagaraDataInterfaceVirtualTextureTemplate.ush");
+	const FName		GetAttributesValidName("GetAttributesValid");
+	const FName		SampleRVTName("SampleRVT");
+	constexpr int32	MaxRVTLayers = 3;
+
+	struct FNDIInstanceData_GameThread
+	{
+		TWeakObjectPtr<URuntimeVirtualTexture> CurrentTexture = nullptr;
+		FNiagaraParameterDirectBinding<UObject*> UserParamBinding;
+	};
+
+	struct FNDInstanceData_RenderThread
+	{
+		FNDInstanceData_RenderThread()
+		{
+			SharedSampler = TStaticSamplerState<SF_Bilinear, AM_Clamp, AM_Clamp, AM_Clamp>::GetRHI();
+		}
+
+		struct FLayerTexture
+		{
+			bool IsValid() const
+			{
+				return Texture.IsValid() && PageTable.IsValid();
+			}
+
+			FTextureRHIRef	Texture;
+			FTextureRHIRef	PageTable;
+			FUintVector4	TextureUniforms = FUintVector4(0, 0, 0, 0);
+		};
+
+		FLayerTexture		LayerTexture[MaxRVTLayers];
+		ERuntimeVirtualTextureMaterialType	MaterialType = ERuntimeVirtualTextureMaterialType::Count;
+		FUintVector4		PageTableUniforms[2] = { FUintVector4(0, 0, 0, 0), FUintVector4(0, 0, 0, 0) };
+		FVector4			UVUniforms[3] = { FVector4::Zero(), FVector4::Zero(), FVector4::Zero() };
+		FVector2f			WorldHeightUnpack = FVector2f::ZeroVector;
+		FSamplerStateRHIRef	SharedSampler;
+	};
+
+	struct FNDIProxy : public FNiagaraDataInterfaceProxy
+	{
+		virtual void ConsumePerInstanceDataFromGameThread(void* PerInstanceData, const FNiagaraSystemInstanceID& Instance) override { checkNoEntry(); }
+		virtual int32 PerInstanceDataPassedToRenderThreadSize() const override { return 0; }
+
+		TMap<FNiagaraSystemInstanceID, FNDInstanceData_RenderThread> InstanceData_RT;
+	};
+}
+
+UNiagaraDataInterfaceVirtualTexture::UNiagaraDataInterfaceVirtualTexture(FObjectInitializer const& ObjectInitializer)
+	: Super(ObjectInitializer)
+	, Texture(nullptr)
+{
+	using namespace NDIVirtualTextureLocal;
+
+	Proxy.Reset(new FNDIProxy());
+
+	FNiagaraTypeDefinition Def(URuntimeVirtualTexture::StaticClass());
+	TextureUserParameter.Parameter.SetType(Def);
+}
+
+void UNiagaraDataInterfaceVirtualTexture::PostInitProperties()
+{
+	Super::PostInitProperties();
+
+	if (HasAnyFlags(RF_ClassDefaultObject))
+	{
+		ENiagaraTypeRegistryFlags Flags = ENiagaraTypeRegistryFlags::AllowAnyVariable | ENiagaraTypeRegistryFlags::AllowParameter;
+		FNiagaraTypeRegistry::Register(FNiagaraTypeDefinition(GetClass()), Flags);
+	}
+}
+
+void UNiagaraDataInterfaceVirtualTexture::PostLoad()
+{
+	Super::PostLoad();
+#if WITH_EDITOR
+	const int32 NiagaraVer = GetLinkerCustomVersion(FNiagaraCustomVersion::GUID);
+	if (NiagaraVer < FNiagaraCustomVersion::TextureDataInterfaceUsesCustomSerialize)
+	{
+		if (Texture != nullptr)
+		{
+			Texture->ConditionalPostLoad();
+		}
+	}
+#endif
+}
+
+void UNiagaraDataInterfaceVirtualTexture::Serialize(FArchive& Ar)
+{
+	Super::Serialize(Ar);
+	if (Ar.IsLoading() == false || Ar.CustomVer(FNiagaraCustomVersion::GUID) >= FNiagaraCustomVersion::TextureDataInterfaceUsesCustomSerialize)
+	{
+		TArray<uint8> StreamData;
+		Ar << StreamData;
+	}
+	Ar.UsingCustomVersion(FNiagaraCustomVersion::GUID);
+}
+
+bool UNiagaraDataInterfaceVirtualTexture::CopyToInternal(UNiagaraDataInterface* Destination) const
+{
+	if (!Super::CopyToInternal(Destination))
+	{
+		return false;
+	}
+	UNiagaraDataInterfaceVirtualTexture* DestinationTexture = CastChecked<UNiagaraDataInterfaceVirtualTexture>(Destination);
+	DestinationTexture->Texture = Texture;
+	DestinationTexture->TextureUserParameter = TextureUserParameter;
+
+	return true;
+}
+
+bool UNiagaraDataInterfaceVirtualTexture::Equals(const UNiagaraDataInterface* Other) const
+{
+	if (!Super::Equals(Other))
+	{
+		return false;
+	}
+	const UNiagaraDataInterfaceVirtualTexture* OtherTexture = CastChecked<const UNiagaraDataInterfaceVirtualTexture>(Other);
+	return
+		OtherTexture->Texture == Texture &&
+		OtherTexture->TextureUserParameter == TextureUserParameter;
+}
+
+void UNiagaraDataInterfaceVirtualTexture::GetFunctions(TArray<FNiagaraFunctionSignature>& OutFunctions)
+{
+	using namespace NDIVirtualTextureLocal;
+
+	{
+		FNiagaraFunctionSignature & Sig = OutFunctions.AddDefaulted_GetRef();
+		Sig.Name = GetAttributesValidName;
+		Sig.bMemberFunction = true;
+		Sig.bRequiresContext = false;
+		Sig.bSupportsCPU = false;
+		Sig.bSupportsGPU = true;
+		Sig.Inputs.Emplace(FNiagaraTypeDefinition(GetClass()), TEXT("Texture"));
+		Sig.Outputs.Emplace(FNiagaraTypeDefinition::GetBoolDef(), TEXT("BaseColorValid"));
+		Sig.Outputs.Emplace(FNiagaraTypeDefinition::GetBoolDef(), TEXT("SpecularValid"));
+		Sig.Outputs.Emplace(FNiagaraTypeDefinition::GetBoolDef(), TEXT("RoughnessValid"));
+		Sig.Outputs.Emplace(FNiagaraTypeDefinition::GetBoolDef(), TEXT("NormalValid"));
+		Sig.Outputs.Emplace(FNiagaraTypeDefinition::GetBoolDef(), TEXT("WorldHeightValid"));
+		Sig.Outputs.Emplace(FNiagaraTypeDefinition::GetBoolDef(), TEXT("MaskValid"));
+	}
+
+	{
+		FNiagaraFunctionSignature& Sig = OutFunctions.AddDefaulted_GetRef();
+		Sig.Name = SampleRVTName;
+		Sig.bMemberFunction = true;
+		Sig.bRequiresContext = false;		
+		Sig.bSupportsCPU = false;
+		Sig.bSupportsGPU = true;
+		Sig.Inputs.Emplace(FNiagaraTypeDefinition(GetClass()), TEXT("Texture"));
+		Sig.Inputs.Emplace(FNiagaraTypeDefinition::GetVec3Def(), TEXT("WorldPosition"));
+		Sig.Outputs.Emplace(FNiagaraTypeDefinition::GetVec3Def(), TEXT("BaseColor"));
+		Sig.Outputs.Emplace(FNiagaraTypeDefinition::GetFloatDef(), TEXT("Specular"));
+		Sig.Outputs.Emplace(FNiagaraTypeDefinition::GetFloatDef(), TEXT("Roughness"));
+		Sig.Outputs.Emplace(FNiagaraTypeDefinition::GetVec3Def(), TEXT("Normal"));
+		Sig.Outputs.Emplace(FNiagaraTypeDefinition::GetFloatDef(), TEXT("WorldHeight"));
+		Sig.Outputs.Emplace(FNiagaraTypeDefinition::GetFloatDef(), TEXT("Mask"));
+	}
+}
+
+void UNiagaraDataInterfaceVirtualTexture::GetVMExternalFunction(const FVMExternalFunctionBindingInfo& BindingInfo, void* InstanceData, FVMExternalFunction &OutFunc)
+{
+	//if (BindingInfo.Name == SampleTexture2DName)
+	//{
+	//	check(BindingInfo.GetNumInputs() == 3 && BindingInfo.GetNumOutputs() == 4);
+	//	NDI_FUNC_BINDER(UNiagaraDataInterfaceVirtualTexture, SampleTexture)::Bind(this, OutFunc);
+	//}
+	//else if (BindingInfo.Name == SamplePseudoVolumeTextureName)
+	//{
+	//	check(BindingInfo.GetNumInputs() == 13 && BindingInfo.GetNumOutputs() == 4);
+	//	NDI_FUNC_BINDER(UNiagaraDataInterfaceVirtualTexture, SamplePseudoVolumeTexture)::Bind(this, OutFunc);
+	//}
+	//else if (BindingInfo.Name == TextureDimsName)
+	//{
+	//	check(BindingInfo.GetNumInputs() == 1 && BindingInfo.GetNumOutputs() == 2);
+	//	OutFunc = FVMExternalFunction::CreateUObject(this, &UNiagaraDataInterfaceVirtualTexture::GetTextureDimensions);
+	//}
+}
+
+int32 UNiagaraDataInterfaceVirtualTexture::PerInstanceDataSize() const
+{
+	using namespace NDIVirtualTextureLocal;
+	return sizeof(FNDIInstanceData_GameThread);
+}
+
+bool UNiagaraDataInterfaceVirtualTexture::InitPerInstanceData(void* PerInstanceData, FNiagaraSystemInstance* SystemInstance)
+{
+	using namespace NDIVirtualTextureLocal;
+	FNDIInstanceData_GameThread* InstanceData = new (PerInstanceData) FNDIInstanceData_GameThread();
+	InstanceData->UserParamBinding.Init(SystemInstance->GetInstanceParameters(), TextureUserParameter.Parameter);
+	return true;
+}
+
+void UNiagaraDataInterfaceVirtualTexture::DestroyPerInstanceData(void* PerInstanceData, FNiagaraSystemInstance* SystemInstance)
+{
+	using namespace NDIVirtualTextureLocal;
+
+	FNDIInstanceData_GameThread* InstanceData = static_cast<FNDIInstanceData_GameThread*>(PerInstanceData);
+	InstanceData->~FNDIInstanceData_GameThread();
+
+	ENQUEUE_RENDER_COMMAND(NDITexture_RemoveInstance)
+	(
+		[RT_Proxy=GetProxyAs<FNDIProxy>(), RT_InstanceID=SystemInstance->GetId()](FRHICommandListImmediate&)
+		{
+			RT_Proxy->InstanceData_RT.Remove(RT_InstanceID);
+		}
+	);
+}
+
+bool UNiagaraDataInterfaceVirtualTexture::PerInstanceTick(void* PerInstanceData, FNiagaraSystemInstance* SystemInstance, float DeltaSeconds)
+{
+	using namespace NDIVirtualTextureLocal;
+	FNDIInstanceData_GameThread* InstanceData = static_cast<FNDIInstanceData_GameThread*>(PerInstanceData);
+
+	URuntimeVirtualTexture* CurrentTexture = InstanceData->UserParamBinding.GetValueOrDefault<URuntimeVirtualTexture>(Texture);
+	//-OPT: We shouldn't need to send this data each frame, currently playing it safe
+	{
+		InstanceData->CurrentTexture = CurrentTexture;
+
+		ENQUEUE_RENDER_COMMAND(NDITexture_UpdateInstance)
+		(
+			[RT_Proxy=GetProxyAs<FNDIProxy>(), RT_InstanceID=SystemInstance->GetId(), RT_WeakTexture=TWeakObjectPtr<URuntimeVirtualTexture>(CurrentTexture)](FRHICommandListImmediate&)
+			{
+				URuntimeVirtualTexture* RT_Texture = RT_WeakTexture.Get();
+				IAllocatedVirtualTexture* VirtualTexture = RT_Texture ? RT_Texture->GetAllocatedVirtualTexture() : nullptr;
+				FNDInstanceData_RenderThread& InstanceData = RT_Proxy->InstanceData_RT.FindOrAdd(RT_InstanceID);
+
+				if (VirtualTexture != nullptr)
+				{
+					InstanceData.MaterialType = RT_Texture->GetMaterialType();
+					VirtualTexture->GetPackedPageTableUniform(InstanceData.PageTableUniforms);
+					InstanceData.UVUniforms[0] = RT_Texture->GetUniformParameter(ERuntimeVirtualTextureShaderUniform_WorldToUVTransform0);
+					InstanceData.UVUniforms[1] = RT_Texture->GetUniformParameter(ERuntimeVirtualTextureShaderUniform_WorldToUVTransform1);
+					InstanceData.UVUniforms[2] = RT_Texture->GetUniformParameter(ERuntimeVirtualTextureShaderUniform_WorldToUVTransform2);
+
+					const FVector4 WorldHeightUnpack = RT_Texture->GetUniformParameter(ERuntimeVirtualTextureShaderUniform_WorldHeightUnpack);
+					InstanceData.WorldHeightUnpack = FVector2f(WorldHeightUnpack.X, WorldHeightUnpack.Y);
+
+					for ( int32 iLayer=0; iLayer < MaxRVTLayers; ++iLayer)
+					{
+						FNDInstanceData_RenderThread::FLayerTexture& LayerTexture = InstanceData.LayerTexture[iLayer];
+						LayerTexture.Texture = VirtualTexture->GetPhysicalTexture(iLayer);
+						LayerTexture.PageTable = VirtualTexture->GetPageTableTexture(0);//-TODO:iLayer);
+
+						VirtualTexture->GetPackedUniform(&LayerTexture.TextureUniforms, iLayer);
+					}
+				}
+				else
+				{
+					InstanceData = FNDInstanceData_RenderThread();
+				}
+			}
+		);
+	}
+
+	return false;
+}
+
+//void UNiagaraDataInterfaceVirtualTexture::GetTextureDimensions(FVectorVMExternalFunctionContext& Context)
+//{
+//	VectorVM::FUserPtrHandler<NDIVirtualTextureLocal::FNDIInstanceData_GameThread> InstData(Context);
+//	FNDIOutputParam<float> OutWidth(Context);
+//	FNDIOutputParam<float> OutHeight(Context);
+//
+//	for (int32 i = 0; i < Context.GetNumInstances(); ++i)
+//	{
+//		OutWidth.SetAndAdvance(InstData->CurrentTextureSize.X);
+//		OutHeight.SetAndAdvance(InstData->CurrentTextureSize.Y);
+//	}
+//}
+
+#if WITH_EDITORONLY_DATA
+bool UNiagaraDataInterfaceVirtualTexture::AppendCompileHash(FNiagaraCompileHashVisitor* InVisitor) const
+{
+	using namespace NDIVirtualTextureLocal;
+
+	bool bSuccess = Super::AppendCompileHash(InVisitor);
+	InVisitor->UpdateString(TEXT("UNiagaraDataInterfaceVirtualTextureHLSLSource"), GetShaderFileHash(CommonShaderFilePath, EShaderPlatform::SP_PCD3D_SM5).ToString());
+	InVisitor->UpdateString(TEXT("UNiagaraDataInterfaceVirtualTextureHLSLSource"), GetShaderFileHash(TemplateShaderFilePath, EShaderPlatform::SP_PCD3D_SM5).ToString());
+	bSuccess &= InVisitor->UpdateShaderParameters<FShaderParameters>();
+	return bSuccess;
+}
+
+void UNiagaraDataInterfaceVirtualTexture::GetCommonHLSL(FString& OutHlsl)
+{
+	using namespace NDIVirtualTextureLocal;
+
+	OutHlsl.Appendf(TEXT("#include \"%s\"\n"), CommonShaderFilePath);
+}
+
+void UNiagaraDataInterfaceVirtualTexture::GetParameterDefinitionHLSL(const FNiagaraDataInterfaceGPUParamInfo& ParamInfo, FString& OutHLSL)
+{
+	using namespace NDIVirtualTextureLocal;
+	TMap<FString, FStringFormatArg> TemplateArgs =
+	{
+		{TEXT("ParameterName"),	ParamInfo.DataInterfaceHLSLSymbol},
+	};
+
+	FString TemplateFile;
+	LoadShaderSourceFile(TemplateShaderFilePath, EShaderPlatform::SP_PCD3D_SM5, &TemplateFile, nullptr);
+	OutHLSL += FString::Format(*TemplateFile, TemplateArgs);
+}
+
+bool UNiagaraDataInterfaceVirtualTexture::GetFunctionHLSL(const FNiagaraDataInterfaceGPUParamInfo& ParamInfo, const FNiagaraDataInterfaceGeneratedFunction& FunctionInfo, int FunctionInstanceIndex, FString& OutHLSL)
+{
+	using namespace NDIVirtualTextureLocal;
+
+	static const TSet<FName> ValidGpuFunctions =
+	{
+		GetAttributesValidName,
+		SampleRVTName,
+	};
+
+	return ValidGpuFunctions.Contains(FunctionInfo.DefinitionName);
+}
+#endif
+
+void UNiagaraDataInterfaceVirtualTexture::BuildShaderParameters(FNiagaraShaderParametersBuilder& ShaderParametersBuilder) const
+{
+	using namespace NDIVirtualTextureLocal;
+	ShaderParametersBuilder.AddNestedStruct<FShaderParameters>();
+}
+
+void UNiagaraDataInterfaceVirtualTexture::SetShaderParameters(const FNiagaraDataInterfaceSetShaderParametersContext& Context) const
+{
+	using namespace NDIVirtualTextureLocal;
+
+	FNDIProxy& TextureProxy = Context.GetProxy<FNDIProxy>();
+	FNDInstanceData_RenderThread* InstanceData = TextureProxy.InstanceData_RT.Find(Context.GetSystemInstanceID());
+
+	FShaderParameters* ShaderParameters = Context.GetParameterNestedStruct<FShaderParameters>();
+
+	bool bLayerValid[MaxRVTLayers];
+	FMemory::Memset(bLayerValid, 0);
+
+	if ( InstanceData )
+	{
+		bLayerValid[0] = InstanceData->LayerTexture[0].IsValid();
+		bLayerValid[1] = InstanceData->LayerTexture[1].IsValid();
+		bLayerValid[2] = InstanceData->LayerTexture[2].IsValid();
+
+		if ( bLayerValid[0] == true)
+		{
+			ShaderParameters->VirtualTexture0					= InstanceData->LayerTexture[0].Texture;
+			ShaderParameters->VirtualTexture0PageTable			= InstanceData->LayerTexture[0].PageTable;
+			ShaderParameters->VirtualTexture0TextureUniforms	= InstanceData->LayerTexture[0].TextureUniforms;
+		}
+
+		if ( bLayerValid[1] == true)
+		{
+			ShaderParameters->VirtualTexture1					= InstanceData->LayerTexture[1].Texture;
+			ShaderParameters->VirtualTexture1PageTable			= InstanceData->LayerTexture[1].PageTable;
+			ShaderParameters->VirtualTexture1TextureUniforms	= InstanceData->LayerTexture[1].TextureUniforms;
+		}
+
+		if ( bLayerValid[2] == true)
+		{
+			ShaderParameters->VirtualTexture2					= InstanceData->LayerTexture[2].Texture;
+			ShaderParameters->VirtualTexture2PageTable			= InstanceData->LayerTexture[2].PageTable;
+			ShaderParameters->VirtualTexture2TextureUniforms	= InstanceData->LayerTexture[2].TextureUniforms;
+		}
+
+		ShaderParameters->ValidLayersMask		 = bLayerValid[0] ? 1 : 0;
+		ShaderParameters->ValidLayersMask		|= bLayerValid[1] ? 2 : 0;
+		ShaderParameters->ValidLayersMask		|= bLayerValid[2] ? 4 : 0;
+		ShaderParameters->MaterialType			= int(InstanceData->MaterialType);
+		ShaderParameters->PageTableUniforms[0]	= InstanceData->PageTableUniforms[0];
+		ShaderParameters->PageTableUniforms[1]	= InstanceData->PageTableUniforms[1];
+		ShaderParameters->UVUniforms[0]			= FVector4f(InstanceData->UVUniforms[0]);	//-TODO: LWC
+		ShaderParameters->UVUniforms[1]			= FVector4f(InstanceData->UVUniforms[1]);
+		ShaderParameters->UVUniforms[2]			= FVector4f(InstanceData->UVUniforms[2]);
+		ShaderParameters->WorldHeightUnpack		= InstanceData->WorldHeightUnpack;
+		ShaderParameters->SharedSampler			= InstanceData->SharedSampler;
+	}
+	else
+	{
+		ShaderParameters->ValidLayersMask		= 0;
+		ShaderParameters->MaterialType			= int(ERuntimeVirtualTextureMaterialType::Count);
+		ShaderParameters->PageTableUniforms[0]	= FUintVector4(0, 0, 0, 0);
+		ShaderParameters->PageTableUniforms[1]	= FUintVector4(0, 0, 0, 0);
+		ShaderParameters->UVUniforms[0]			= FVector4f::Zero();
+		ShaderParameters->UVUniforms[1]			= FVector4f::Zero();
+		ShaderParameters->UVUniforms[2]			= FVector4f::Zero();
+		ShaderParameters->WorldHeightUnpack		= FVector2f::ZeroVector;
+		ShaderParameters->SharedSampler			= TStaticSamplerState<SF_Bilinear, AM_Clamp, AM_Clamp, AM_Clamp>::GetRHI();
+	}
+	
+	if ( bLayerValid[0] == false )
+	{
+		ShaderParameters->VirtualTexture0					= GBlackTexture->GetTextureRHI();
+		ShaderParameters->VirtualTexture0PageTable			= GBlackUintTexture->TextureRHI;
+		ShaderParameters->VirtualTexture0TextureUniforms	= FUintVector4(0, 0, 0, 0);
+	}
+	if ( bLayerValid[1] == false )
+	{
+		ShaderParameters->VirtualTexture1					= GBlackTexture->GetTextureRHI();
+		ShaderParameters->VirtualTexture1PageTable			= GBlackUintTexture->TextureRHI;
+		ShaderParameters->VirtualTexture1TextureUniforms	= FUintVector4(0, 0, 0, 0);
+	}
+	if ( bLayerValid[2] == false )
+	{
+		ShaderParameters->VirtualTexture2					= GBlackTexture->GetTextureRHI();
+		ShaderParameters->VirtualTexture2PageTable			= GBlackUintTexture->TextureRHI;
+		ShaderParameters->VirtualTexture2TextureUniforms	= FUintVector4(0, 0, 0, 0);
+	}
+}
+
+void UNiagaraDataInterfaceVirtualTexture::SetTexture(URuntimeVirtualTexture* InTexture)
+{
+	Texture = InTexture;
+}
+
+#undef LOCTEXT_NAMESPACE
diff --git a/Engine/Plugins/FX/Niagara/Source/Niagara/Private/DataInterface/NiagaraDataInterfaceVirtualTexture.h b/Engine/Plugins/FX/Niagara/Source/Niagara/Private/DataInterface/NiagaraDataInterfaceVirtualTexture.h
new file mode 100644
--- /dev/null
+++ b/Engine/Plugins/FX/Niagara/Source/Niagara/Private/DataInterface/NiagaraDataInterfaceVirtualTexture.h
@@ -0,0 +1,65 @@
+// Copyright Epic Games, Inc. All Rights Reserved.
+
+#pragma once
+
+#include "NiagaraDataInterface.h"
+#include "NiagaraCommon.h"
+#include "VectorVM.h"
+#include "NiagaraDataInterfaceVirtualTexture.generated.h"
+
+class URuntimeVirtualTexture;
+
+/** Data Interface allowing sampling of a texture */
+UCLASS(Experimental, EditInlineNew, Category = "Texture", meta = (DisplayName = "Runtime Virtual Texture Sample"))
+class NIAGARA_API UNiagaraDataInterfaceVirtualTexture : public UNiagaraDataInterface
+{
+	GENERATED_UCLASS_BODY()
+
+public:
+	UPROPERTY(EditAnywhere, Category = "Texture")//, meta = (AllowedClasses = "/Script/Engine.VirtualTexture2D"))
+	TObjectPtr<URuntimeVirtualTexture> Texture;
+
+	UPROPERTY(EditAnywhere, Category = "Texture", meta = (ToolTip = "When valid the user parameter is used as the texture rather than the one on the data interface"))
+	FNiagaraUserParameterBinding TextureUserParameter;
+
+	//UObject Interface
+	virtual void PostInitProperties() override;
+	virtual void PostLoad() override;
+	virtual void Serialize(FArchive& Ar) override;
+	//UObject Interface End
+
+	//UNiagaraDataInterface Interface
+	virtual void GetFunctions(TArray<FNiagaraFunctionSignature>& OutFunctions) override;
+	virtual void GetVMExternalFunction(const FVMExternalFunctionBindingInfo& BindingInfo, void* InstanceData, FVMExternalFunction &OutFunc) override;
+	virtual bool CanExecuteOnTarget(ENiagaraSimTarget Target) const override { return Target==ENiagaraSimTarget::GPUComputeSim; }
+
+	virtual int32 PerInstanceDataSize() const override;
+	virtual bool InitPerInstanceData(void* PerInstanceData, FNiagaraSystemInstance* SystemInstance) override;
+	virtual void DestroyPerInstanceData(void* PerInstanceData, FNiagaraSystemInstance* SystemInstance) override;
+
+	virtual bool HasPreSimulateTick() const override { return true; }
+	virtual bool PerInstanceTick(void* PerInstanceData, FNiagaraSystemInstance* SystemInstance, float DeltaSeconds) override;
+	//UNiagaraDataInterface Interface End
+
+	//void SampleTexture(FVectorVMExternalFunctionContext& Context);
+	//void GetTextureDimensions(FVectorVMExternalFunctionContext& Context);
+	//void SamplePseudoVolumeTexture(FVectorVMExternalFunctionContext& Context);
+
+	virtual bool Equals(const UNiagaraDataInterface* Other) const override;
+
+	// GPU sim functionality
+#if WITH_EDITORONLY_DATA
+	virtual bool AppendCompileHash(FNiagaraCompileHashVisitor* InVisitor) const override;
+	virtual void GetCommonHLSL(FString& OutHlsl) override;
+	virtual void GetParameterDefinitionHLSL(const FNiagaraDataInterfaceGPUParamInfo& ParamInfo, FString& OutHLSL) override;
+	virtual bool GetFunctionHLSL(const FNiagaraDataInterfaceGPUParamInfo& ParamInfo, const FNiagaraDataInterfaceGeneratedFunction& FunctionInfo, int FunctionInstanceIndex, FString& OutHLSL) override;
+#endif
+	virtual bool UseLegacyShaderBindings() const override { return false; }
+	virtual void BuildShaderParameters(FNiagaraShaderParametersBuilder& ShaderParametersBuilder) const override;
+	virtual void SetShaderParameters(const FNiagaraDataInterfaceSetShaderParametersContext& Context) const override;
+
+	void SetTexture(URuntimeVirtualTexture* InTexture);
+
+protected:
+	virtual bool CopyToInternal(UNiagaraDataInterface* Destination) const override;
+};
