- Change to Position type and handle LWC through the chain.
- Potential precision loss currently as we are rebasing the origin into simulation space, can be changed later if the need comes up

#rnx

diff --git a/Engine/Plugins/FX/Niagara/Shaders/Private/NiagaraDataInterfaceVirtualTextureTemplate.ush b/Engine/Plugins/FX/Niagara/Shaders/Private/NiagaraDataInterfaceVirtualTextureTemplate.ush
--- a/Engine/Plugins/FX/Niagara/Shaders/Private/NiagaraDataInterfaceVirtualTextureTemplate.ush
+++ b/Engine/Plugins/FX/Niagara/Shaders/Private/NiagaraDataInterfaceVirtualTextureTemplate.ush
@@ -17,6 +17,7 @@ int					{ParameterName}_MaterialType;
 float4				{ParameterName}_UVUniforms[3];
 uint4				{ParameterName}_PageTableUniforms[2];
 float2				{ParameterName}_WorldHeightUnpack;
+float3				{ParameterName}_SystemLWCTile;
 SamplerState		{ParameterName}_SharedSampler;
 
 void GetAttributesValid_{ParameterName}(out bool BaseColorValid, out bool SpecularValid, out bool RoughnessValid, out bool NormalValid, out bool WorldHeightValid, out bool MaskValid)
@@ -85,7 +86,10 @@ void SampleRVT_{ParameterName}(in float3 WorldPosition, out float3 BaseColor, ou
 	Mask = 1.0f;
 
 	// Get Sample Location
-	float2 SampleUV = VirtualTextureWorldToUV(LWCPromote(WorldPosition), LWCPromote({ParameterName}_UVUniforms[0].xyz), {ParameterName}_UVUniforms[1].xyz, {ParameterName}_UVUniforms[2].xyz);
+	FLWCVector3 LWCWorldPosition = MakeLWCVector3({ParameterName}_SystemLWCTile, WorldPosition);
+	FLWCVector3 LWCUVOrigin = MakeLWCVector3({ParameterName}_SystemLWCTile, {ParameterName}_UVUniforms[0].xyz);
+
+	float2 SampleUV = VirtualTextureWorldToUV(LWCWorldPosition, LWCUVOrigin, {ParameterName}_UVUniforms[1].xyz, {ParameterName}_UVUniforms[2].xyz);
 
 	// Sample Textures
 	float4 LayerSample[3];
diff --git a/Engine/Plugins/FX/Niagara/Source/Niagara/Private/DataInterface/NiagaraDataInterfaceVirtualTexture.cpp b/Engine/Plugins/FX/Niagara/Source/Niagara/Private/DataInterface/NiagaraDataInterfaceVirtualTexture.cpp
--- a/Engine/Plugins/FX/Niagara/Source/Niagara/Private/DataInterface/NiagaraDataInterfaceVirtualTexture.cpp
+++ b/Engine/Plugins/FX/Niagara/Source/Niagara/Private/DataInterface/NiagaraDataInterfaceVirtualTexture.cpp
@@ -8,6 +8,7 @@
 #include "NiagaraShaderParametersBuilder.h"
 #include "NiagaraSystemInstance.h"
 
+#include "Misc/LargeWorldRenderPosition.h"
 #include "VT/RuntimeVirtualTexture.h"
 
 #define LOCTEXT_NAMESPACE "UNiagaraDataInterfaceVirtualTexture"
@@ -32,6 +33,7 @@ namespace NDIVirtualTextureLocal
 		SHADER_PARAMETER_ARRAY(FUintVector4,		PageTableUniforms, [2])
 		SHADER_PARAMETER_ARRAY(FVector4f,			UVUniforms, [3])
 		SHADER_PARAMETER(FVector2f,					WorldHeightUnpack)
+		SHADER_PARAMETER(FVector3f,					SystemLWCTile)
 		SHADER_PARAMETER_SAMPLER(SamplerState,		SharedSampler)
 	END_SHADER_PARAMETER_STRUCT()
 
@@ -168,6 +170,7 @@ void UNiagaraDataInterfaceVirtualTexture::GetFunctions(TArray<FNiagaraFunctionSi
 		Sig.bRequiresContext = false;
 		Sig.bSupportsCPU = false;
 		Sig.bSupportsGPU = true;
+		Sig.bExperimental = true;
 		Sig.Inputs.Emplace(FNiagaraTypeDefinition(GetClass()), TEXT("Texture"));
 		Sig.Outputs.Emplace(FNiagaraTypeDefinition::GetBoolDef(), TEXT("BaseColorValid"));
 		Sig.Outputs.Emplace(FNiagaraTypeDefinition::GetBoolDef(), TEXT("SpecularValid"));
@@ -184,8 +187,9 @@ void UNiagaraDataInterfaceVirtualTexture::GetFunctions(TArray<FNiagaraFunctionSi
 		Sig.bRequiresContext = false;		
 		Sig.bSupportsCPU = false;
 		Sig.bSupportsGPU = true;
+		Sig.bExperimental = true;
 		Sig.Inputs.Emplace(FNiagaraTypeDefinition(GetClass()), TEXT("Texture"));
-		Sig.Inputs.Emplace(FNiagaraTypeDefinition::GetVec3Def(), TEXT("WorldPosition"));
+		Sig.Inputs.Emplace(FNiagaraTypeDefinition::GetPositionDef(), TEXT("WorldPosition"));
 		Sig.Outputs.Emplace(FNiagaraTypeDefinition::GetVec3Def(), TEXT("BaseColor"));
 		Sig.Outputs.Emplace(FNiagaraTypeDefinition::GetFloatDef(), TEXT("Specular"));
 		Sig.Outputs.Emplace(FNiagaraTypeDefinition::GetFloatDef(), TEXT("Roughness"));
@@ -197,21 +201,7 @@ void UNiagaraDataInterfaceVirtualTexture::GetFunctions(TArray<FNiagaraFunctionSi
 
 void UNiagaraDataInterfaceVirtualTexture::GetVMExternalFunction(const FVMExternalFunctionBindingInfo& BindingInfo, void* InstanceData, FVMExternalFunction &OutFunc)
 {
-	//if (BindingInfo.Name == SampleTexture2DName)
-	//{
-	//	check(BindingInfo.GetNumInputs() == 3 && BindingInfo.GetNumOutputs() == 4);
-	//	NDI_FUNC_BINDER(UNiagaraDataInterfaceVirtualTexture, SampleTexture)::Bind(this, OutFunc);
-	//}
-	//else if (BindingInfo.Name == SamplePseudoVolumeTextureName)
-	//{
-	//	check(BindingInfo.GetNumInputs() == 13 && BindingInfo.GetNumOutputs() == 4);
-	//	NDI_FUNC_BINDER(UNiagaraDataInterfaceVirtualTexture, SamplePseudoVolumeTexture)::Bind(this, OutFunc);
-	//}
-	//else if (BindingInfo.Name == TextureDimsName)
-	//{
-	//	check(BindingInfo.GetNumInputs() == 1 && BindingInfo.GetNumOutputs() == 2);
-	//	OutFunc = FVMExternalFunction::CreateUObject(this, &UNiagaraDataInterfaceVirtualTexture::GetTextureDimensions);
-	//}
+	// No CPU side functions yet
 }
 
 int32 UNiagaraDataInterfaceVirtualTexture::PerInstanceDataSize() const
@@ -293,19 +283,6 @@ bool UNiagaraDataInterfaceVirtualTexture::PerInstanceTick(void* PerInstanceData,
 	return false;
 }
 
-//void UNiagaraDataInterfaceVirtualTexture::GetTextureDimensions(FVectorVMExternalFunctionContext& Context)
-//{
-//	VectorVM::FUserPtrHandler<NDIVirtualTextureLocal::FNDIInstanceData_GameThread> InstData(Context);
-//	FNDIOutputParam<float> OutWidth(Context);
-//	FNDIOutputParam<float> OutHeight(Context);
-//
-//	for (int32 i = 0; i < Context.GetNumInstances(); ++i)
-//	{
-//		OutWidth.SetAndAdvance(InstData->CurrentTextureSize.X);
-//		OutHeight.SetAndAdvance(InstData->CurrentTextureSize.Y);
-//	}
-//}
-
 #if WITH_EDITORONLY_DATA
 bool UNiagaraDataInterfaceVirtualTexture::AppendCompileHash(FNiagaraCompileHashVisitor* InVisitor) const
 {
@@ -397,13 +374,16 @@ void UNiagaraDataInterfaceVirtualTexture::SetShaderParameters(const FNiagaraData
 			ShaderParameters->VirtualTexture2TextureUniforms	= InstanceData->LayerTexture[2].TextureUniforms;
 		}
 
+		const FNiagaraLWCConverter LWCConverter(FVector(Context.GetSystemLWCTile()) * FLargeWorldRenderScalar::GetTileSize());
+		const FVector3f UVOriginRebased = LWCConverter.ConvertWorldToSimulationVector(FVector(InstanceData->UVUniforms[0]));
+
 		ShaderParameters->ValidLayersMask		 = bLayerValid[0] ? 1 : 0;
 		ShaderParameters->ValidLayersMask		|= bLayerValid[1] ? 2 : 0;
 		ShaderParameters->ValidLayersMask		|= bLayerValid[2] ? 4 : 0;
 		ShaderParameters->MaterialType			= int(InstanceData->MaterialType);
 		ShaderParameters->PageTableUniforms[0]	= InstanceData->PageTableUniforms[0];
 		ShaderParameters->PageTableUniforms[1]	= InstanceData->PageTableUniforms[1];
-		ShaderParameters->UVUniforms[0]			= FVector4f(InstanceData->UVUniforms[0]);	//-TODO: LWC
+		ShaderParameters->UVUniforms[0]			= FVector4f(UVOriginRebased);
 		ShaderParameters->UVUniforms[1]			= FVector4f(InstanceData->UVUniforms[1]);
 		ShaderParameters->UVUniforms[2]			= FVector4f(InstanceData->UVUniforms[2]);
 		ShaderParameters->WorldHeightUnpack		= InstanceData->WorldHeightUnpack;
@@ -421,6 +401,7 @@ void UNiagaraDataInterfaceVirtualTexture::SetShaderParameters(const FNiagaraData
 		ShaderParameters->WorldHeightUnpack		= FVector2f::ZeroVector;
 		ShaderParameters->SharedSampler			= TStaticSamplerState<SF_Bilinear, AM_Clamp, AM_Clamp, AM_Clamp>::GetRHI();
 	}
+	ShaderParameters->SystemLWCTile	= Context.GetSystemLWCTile();
 	
 	if ( bLayerValid[0] == false )
 	{
