- Use FVector2D::Max to get a component wise max / min

#rnx

diff --git a/Engine/Plugins/FX/Niagara/Source/Niagara/Private/NiagaraUvQuadTree.cpp b/Engine/Plugins/FX/Niagara/Source/Niagara/Private/NiagaraUvQuadTree.cpp
--- a/Engine/Plugins/FX/Niagara/Source/Niagara/Private/NiagaraUvQuadTree.cpp
+++ b/Engine/Plugins/FX/Niagara/Source/Niagara/Private/NiagaraUvQuadTree.cpp
@@ -180,9 +180,9 @@ void FNiagaraUvQuadTree::Split(int32 SubTreeIndex)
 		// if the node we're trying to push down is larger than our threshold we keep it in the parent node
 		else
 		{
-			const FVector2D OverlapMin = FMath::Max(ContentNode.Coverage.Min, Parent.Coverage.Min);
-			const FVector2D OverlapMax = FMath::Min(ContentNode.Coverage.Max, Parent.Coverage.Max);
-			FBox2D OverlapRegion(OverlapMin, FMath::Max(OverlapMin, OverlapMax));
+			const FVector2D OverlapMin = FVector2D::Max(ContentNode.Coverage.Min, Parent.Coverage.Min);
+			const FVector2D OverlapMax = FVector2D::Min(ContentNode.Coverage.Max, Parent.Coverage.Max);
+			FBox2D OverlapRegion(OverlapMin, FVector2D::Max(OverlapMin, OverlapMax));
 			const float AreaRatio = OverlapRegion.GetArea() / Parent.Coverage.GetArea();
 
 			if (AreaRatio > GNiagaraUvQuadTreeDuplicateThreshold)
