clamped cubic interpolation

diff --git a/Engine/Plugins/FX/Niagara/Shaders/Private/NiagaraDataInterfaceGrid3DCollection.ush b/Engine/Plugins/FX/Niagara/Shaders/Private/NiagaraDataInterfaceGrid3DCollection.ush
--- a/Engine/Plugins/FX/Niagara/Shaders/Private/NiagaraDataInterfaceGrid3DCollection.ush
+++ b/Engine/Plugins/FX/Niagara/Shaders/Private/NiagaraDataInterfaceGrid3DCollection.ush
@@ -40,14 +40,12 @@ float SampleTriCubicLagrange_{ParameterName}(SamplerState Sampler, float3 UVW, i
 								
 	int3 LocalCell = GridCell % {NumCellsName};
 								
-	bool bBoundary = !any(LocalCell);
-	LocalCell -= ({NumCellsName} -int3(2,2,2));
-	bBoundary = bBoundary || !any(LocalCell);
-		
-	if (bBoundary)
+	const int3 MaxCell = {NumCellsName} - int3(2,2,2);
+
+	if (any((LocalCell <= int3(1,1,1)) || (LocalCell >= MaxCell)))
 	{
 		// revert to trilinear hardware sampling at the boundary cells.
-		return {GridName}.SampleLevel(Sampler, UVW, MipLevel);;
+		return {GridName}.SampleLevel(Sampler, UVW, MipLevel);
 	}
 	else 
 	{
@@ -60,6 +58,8 @@ float SampleTriCubicLagrange_{ParameterName}(SamplerState Sampler, float3 UVW, i
 		float v = 1. + delta.z;
 		
 		float4 DataInZ = float4(0.f, 0.f, 0.f, 0.f);
+		float minv = 3.402823466e+38;
+		float maxv = -3.402823466e+38;
 		for (int zk = 0; zk < 4; zk++)
 		{
 			int plane = GridCell[2] + zk -1;
@@ -75,6 +75,8 @@ float SampleTriCubicLagrange_{ParameterName}(SamplerState Sampler, float3 UVW, i
 				{
 					int col = GridCell[0] + xi - 1;
 					DataInX[xi] = {GridName}.Load(int4(col, row, plane, MipLevel)); 
+					minv = min(DataInX[xi], minv);
+					maxv = max(DataInX[xi], maxv);
 				}
 		
 				// coefficients
@@ -107,6 +109,6 @@ float SampleTriCubicLagrange_{ParameterName}(SamplerState Sampler, float3 UVW, i
 		// float BlendAmount = 0.0f;
 		// return lerp(TriCubicValue, TriLinearValue, BlendAmount);		
 
-		return TriCubicValue;		
+		return  min(max(TriCubicValue, minv), maxv);
 	}
 }
\ No newline at end of file
