Fixed FNiagaraRendererMeshes::GetDynamicRayTracingInstances fails to use Reserve, MoveTemp, or Emplace_GetRef to avoid copies and reallocates.

On behalf of Evan Hart @ Nvidia

diff --git a/Engine/Plugins/FX/Niagara/Source/Niagara/Private/NiagaraRendererMeshes.cpp b/Engine/Plugins/FX/Niagara/Source/Niagara/Private/NiagaraRendererMeshes.cpp
--- a/Engine/Plugins/FX/Niagara/Source/Niagara/Private/NiagaraRendererMeshes.cpp
+++ b/Engine/Plugins/FX/Niagara/Source/Niagara/Private/NiagaraRendererMeshes.cpp
@@ -1337,7 +1337,7 @@ void FNiagaraRendererMeshes::GetDynamicRayTracingInstances(FRayTracingMaterialGa
 	FParticleMeshRenderData ParticleMeshRenderData;
 	PrepareParticleMeshRenderData(ParticleMeshRenderData, *View->Family, Context.RayTracingMeshResourceCollector, DynamicDataRender, SceneProxy, true);
 
-	if (ParticleMeshRenderData.SourceParticleData == nullptr)
+	if (ParticleMeshRenderData.SourceParticleData == nullptr || Meshes.Num() == 0)
 	{
 		return;
 	}
@@ -1357,6 +1357,8 @@ void FNiagaraRendererMeshes::GetDynamicRayTracingInstances(FRayTracingMaterialGa
 		InitializeSortInfo(ParticleMeshRenderData, *SceneProxy, *View, ViewIndex, bIsInstancedStereo, SortInfo);
 	}
 
+	OutRayTracingInstances.Reserve(Meshes.Num());
+
 	for (int32 MeshIndex = 0; MeshIndex < Meshes.Num(); ++MeshIndex)
 	{
 		// No binding for mesh index we only render the first mesh not all of them
@@ -1413,6 +1415,8 @@ void FNiagaraRendererMeshes::GetDynamicRayTracingInstances(FRayTracingMaterialGa
 		VertexFactory.SetUniformBuffer(VFUniformBuffer);
 		CollectorResources->UniformBuffer = VFUniformBuffer;
 
+		RayTracingInstance.Materials.Reserve(LODModel.Sections.Num());
+
 		for (int32 SectionIndex = 0; SectionIndex < LODModel.Sections.Num(); SectionIndex++)
 		{
 			const FStaticMeshSection& Section = LODModel.Sections[SectionIndex];
@@ -1467,7 +1471,7 @@ void FNiagaraRendererMeshes::GetDynamicRayTracingInstances(FRayTracingMaterialGa
 			MeshBatch.Elements[0].VisualizeElementIndex = SectionIndex;
 #endif
 
-			RayTracingInstance.Materials.Add(MeshBatch);
+			RayTracingInstance.Materials.Add(MoveTemp(MeshBatch));
 		}
 
 		if (RayTracingInstance.Materials.Num() == 0 || LODModel.Sections.Num() != RayTracingInstance.Materials.Num())
@@ -1546,7 +1550,7 @@ void FNiagaraRendererMeshes::GetDynamicRayTracingInstances(FRayTracingMaterialGa
 
 			RayTracingInstance.InstanceTransforms.Add(InstanceTransform);
 		}
-		else 
+		else
 		{
 			TConstArrayView<FNiagaraRendererVariableInfo> VFVariables = ParticleMeshRenderData.RendererLayout->GetVFVariables_RenderThread();
 			if (SimTarget == ENiagaraSimTarget::CPUSim)
@@ -1699,7 +1703,7 @@ void FNiagaraRendererMeshes::GetDynamicRayTracingInstances(FRayTracingMaterialGa
 		}
 
 		RayTracingInstance.BuildInstanceMaskAndFlags(FeatureLevel);
-		OutRayTracingInstances.Add(RayTracingInstance);
+		OutRayTracingInstances.Add(MoveTemp(RayTracingInstance));
 	}
 }
 #endif
