Fixed a crash when playing audio from particle system on a platform that has audio disabled

diff --git a/Engine/Plugins/FX/Niagara/Source/Niagara/Private/NiagaraDataInterfaceAudioPlayer.cpp b/Engine/Plugins/FX/Niagara/Source/Niagara/Private/NiagaraDataInterfaceAudioPlayer.cpp
--- a/Engine/Plugins/FX/Niagara/Source/Niagara/Private/NiagaraDataInterfaceAudioPlayer.cpp
+++ b/Engine/Plugins/FX/Niagara/Source/Niagara/Private/NiagaraDataInterfaceAudioPlayer.cpp
@@ -841,6 +841,11 @@ void UNiagaraDataInterfaceAudioPlayer::PlayPersistentAudio(FVectorVMExternalFunc
 					{
 						bool bStopWithEffect = InstanceData->bStopWhenComponentIsDestroyed || Sound->IsLooping(); // we don't allow looping effects to outlive us because then they keep playing forever
 						UAudioComponent* AudioComponent = UGameplayStatics::SpawnSoundAttached(Sound.Get(), NiagaraComponent, NAME_None, Position, Rotation, EAttachLocation::KeepWorldPosition, bStopWithEffect, Volume, Pitch, StartTime, InstanceData->Attenuation.Get(), InstanceData->Concurrency.Get(), true);
+						if (AudioComponent == nullptr)
+						{
+							// looks like audio is disabled, so we'll skip adding a mapping
+							return;
+						}
 						if (FadeIn > 0.0)
 						{
 							AudioComponent->FadeIn(FadeIn, Volume, StartTime);
