Fixed preview component not resetting when chaning scalability preview mode or device.

diff --git a/Engine/Plugins/FX/Niagara/Source/NiagaraEditor/Private/ViewModels/NiagaraSystemScalabilityViewModel.cpp b/Engine/Plugins/FX/Niagara/Source/NiagaraEditor/Private/ViewModels/NiagaraSystemScalabilityViewModel.cpp
--- a/Engine/Plugins/FX/Niagara/Source/NiagaraEditor/Private/ViewModels/NiagaraSystemScalabilityViewModel.cpp
+++ b/Engine/Plugins/FX/Niagara/Source/NiagaraEditor/Private/ViewModels/NiagaraSystemScalabilityViewModel.cpp
@@ -48,6 +48,9 @@ void UNiagaraSystemScalabilityViewModel::UpdatePreviewDeviceProfile(UDeviceProfi
 
 	FNiagaraPlatformSet::InvalidateCachedData();
 	GetSystemViewModel().Pin()->GetSystem().UpdateScalability();
+
+	//Ensure we properly reset the preview system.
+	GetSystemViewModel().Pin()->ResetSystem(FNiagaraSystemViewModel::ETimeResetMode::AllowResetTime, FNiagaraSystemViewModel::EMultiResetMode::ResetThisInstance, FNiagaraSystemViewModel::EReinitMode::ReinitializeSystem);
 }
 
 void UNiagaraSystemScalabilityViewModel::UpdatePreviewQualityLevel(int32 QualityLevel)
@@ -57,6 +60,9 @@ void UNiagaraSystemScalabilityViewModel::UpdatePreviewQualityLevel(int32 Quality
 
 	FNiagaraPlatformSet::InvalidateCachedData();
 	GetSystemViewModel().Pin()->GetSystem().UpdateScalability();
+
+	//Ensure we properly reset the preview system.
+	GetSystemViewModel().Pin()->ResetSystem(FNiagaraSystemViewModel::ETimeResetMode::AllowResetTime, FNiagaraSystemViewModel::EMultiResetMode::ResetThisInstance, FNiagaraSystemViewModel::EReinitMode::ReinitializeSystem);
 }
 
 bool UNiagaraSystemScalabilityViewModel::IsPlatformActive(const FNiagaraPlatformSet& PlatformSet)
