Fixed a bug where emitter properties with edit conditions could not be edited

diff --git a/Engine/Plugins/FX/Niagara/Source/Niagara/Classes/NiagaraEmitter.h b/Engine/Plugins/FX/Niagara/Source/Niagara/Classes/NiagaraEmitter.h
--- a/Engine/Plugins/FX/Niagara/Source/Niagara/Classes/NiagaraEmitter.h
+++ b/Engine/Plugins/FX/Niagara/Source/Niagara/Classes/NiagaraEmitter.h
@@ -331,7 +331,7 @@ struct NIAGARA_API FVersionedNiagaraEmitterData
 	The emitter will allocate at least this many particles on it's first tick.
 	This can aid performance by avoiding many allocations as an emitter ramps up to it's max size.
 	*/
-	UPROPERTY(EditAnywhere, AdvancedDisplay, Category = "Emitter", meta = (EditCondition = "AllocationMode == EParticleAllocationMode::ManualEstimate"))
+	UPROPERTY(EditAnywhere, AdvancedDisplay, Category = "Emitter", meta = (EditCondition = "AllocationMode != EParticleAllocationMode::AutomaticEstimate"))
 	int32 PreAllocationCount = 0;
 
 	UPROPERTY()
diff --git a/Engine/Plugins/FX/Niagara/Source/NiagaraEditor/Private/NiagaraEmitterDetailsCustomization.cpp b/Engine/Plugins/FX/Niagara/Source/NiagaraEditor/Private/NiagaraEmitterDetailsCustomization.cpp
--- a/Engine/Plugins/FX/Niagara/Source/NiagaraEditor/Private/NiagaraEmitterDetailsCustomization.cpp
+++ b/Engine/Plugins/FX/Niagara/Source/NiagaraEditor/Private/NiagaraEmitterDetailsCustomization.cpp
@@ -15,6 +15,45 @@ TSharedRef<IDetailCustomization> FNiagaraEmitterDetails::MakeInstance()
 	return MakeShared<FNiagaraEmitterDetails>();
 }
 
+void CustomizeEmitterData(IDetailLayoutBuilder& InDetailLayout, UNiagaraEmitter* EmitterBeingCustomized, TMap<FName, TSet<FProperty*>> CategoryPropertyMap)
+{
+	TArray<TSharedPtr<FNiagaraEmitterViewModel>> ExistingViewModels;
+	TNiagaraViewModelManager<UNiagaraEmitter, FNiagaraEmitterViewModel>::GetAllViewModelsForObject(EmitterBeingCustomized, ExistingViewModels);
+	FVersionedNiagaraEmitter VersionedNiagaraEmitter = ExistingViewModels[0]->GetEmitter();
+	FVersionedNiagaraEmitterData* EmitterData = VersionedNiagaraEmitter.GetEmitterData();
+	TSharedPtr<FStructOnScope> StructData = MakeShareable(new FStructOnScope(FVersionedNiagaraEmitterData::StaticStruct(), reinterpret_cast<uint8*>(EmitterData)));
+	for (const auto& Entry : CategoryPropertyMap)
+	{
+		IDetailCategoryBuilder& CategoryBuilder = InDetailLayout.EditCategory(Entry.Key);
+		TArray<IDetailPropertyRow*> OutRows;
+			
+		// We want only one structonscope, otherwise edit conditions don't work correctly. Unfortuately this means we need to add each category separately and hide properties not assigned to the category.
+		TArray<TSharedPtr<IPropertyHandle>> Handles = CategoryBuilder.AddAllExternalStructureProperties(StructData.ToSharedRef(), EPropertyLocation::Default, &OutRows);
+			
+		for (IDetailPropertyRow* PropertyRow : OutRows)
+		{
+			TSharedPtr<IPropertyHandle> PropertyHandle = PropertyRow->GetPropertyHandle();
+			if (!Entry.Value.Contains(PropertyHandle->GetProperty()))
+			{
+				PropertyHandle->MarkHiddenByCustomization();
+				PropertyRow->Visibility(EVisibility::Collapsed);
+			}
+			else
+			{
+				FProperty* ChildProperty = PropertyHandle->GetProperty();
+				const auto& PostEditChangeLambda = [VersionedNiagaraEmitter, ChildProperty]
+				{
+					FPropertyChangedEvent ChangeEvent(ChildProperty);
+					VersionedNiagaraEmitter.Emitter->PostEditChangeVersionedProperty(ChangeEvent, VersionedNiagaraEmitter.Version);
+				};
+
+				PropertyRow->GetPropertyHandle()->SetOnChildPropertyValueChanged(FSimpleDelegate::CreateLambda(PostEditChangeLambda));
+				PropertyRow->GetPropertyHandle()->SetOnPropertyValueChanged(FSimpleDelegate::CreateLambda(PostEditChangeLambda));
+			}
+		}
+	}
+}
+
 void FNiagaraEmitterDetails::CustomizeDetails(IDetailLayoutBuilder& InDetailLayout)
 {
 	TArray<TWeakObjectPtr<UObject>> ObjectsBeingCustomized;
@@ -22,33 +61,23 @@ void FNiagaraEmitterDetails::CustomizeDetails(IDetailLayoutBuilder& InDetailLayo
 	
 	if (UNiagaraEmitter* EmitterBeingCustomized = CastChecked<UNiagaraEmitter>(ObjectsBeingCustomized[0]))
 	{
-		TArray<TSharedPtr<FNiagaraEmitterViewModel>> ExistingViewModels;
-		TNiagaraViewModelManager<UNiagaraEmitter, FNiagaraEmitterViewModel>::GetAllViewModelsForObject(EmitterBeingCustomized, ExistingViewModels);
-		FVersionedNiagaraEmitter VersionedNiagaraEmitter = ExistingViewModels[0]->GetEmitter();
-		FVersionedNiagaraEmitterData* EmitterData = VersionedNiagaraEmitter.GetEmitterData();
-
+		TMap<FName, TSet<FProperty*>> CategoryPropertyMap;
 		for (FProperty* ChildProperty : TFieldRange<FProperty>(FVersionedNiagaraEmitterData::StaticStruct()))
 		{
 			if (ChildProperty->HasAllPropertyFlags(CPF_Edit))
 			{
 				FName Category = FName(ChildProperty->GetMetaData(TEXT("Category")));
-				IDetailCategoryBuilder& CategoryBuilder = InDetailLayout.EditCategory(Category);
-				
-				TSharedPtr<FStructOnScope> StructData = MakeShareable(new FStructOnScope(FVersionedNiagaraEmitterData::StaticStruct(), reinterpret_cast<uint8*>(EmitterData)));
-				if (IDetailPropertyRow* PropertyRow = CategoryBuilder.AddExternalStructureProperty(StructData, ChildProperty->GetFName()))
+
+				// we display the scalability category within scalability mode, which is why we hide it here
+				if (Category != FName("Scalability"))
 				{
-					PropertyRow->GetPropertyHandle()->SetOnPropertyValueChanged(FSimpleDelegate::CreateLambda([VersionedNiagaraEmitter, ChildProperty]()
-					{
-						FPropertyChangedEvent ChangeEvent(ChildProperty);
-						VersionedNiagaraEmitter.Emitter->PostEditChangeVersionedProperty(ChangeEvent, VersionedNiagaraEmitter.Version);
-					}));
+					CategoryPropertyMap.FindOrAdd(Category).Add(ChildProperty);
 				}
 			}
 		}
-	}
 
-	// we display the scalability category within scalability mode, which is why we hide it here
-	InDetailLayout.HideCategory("Scalability");
+		CustomizeEmitterData(InDetailLayout, EmitterBeingCustomized, CategoryPropertyMap);
+	}
 }
 
 TSharedRef<IDetailCustomization> FNiagaraEmitterScalabilityDetails::MakeInstance()
@@ -62,32 +91,17 @@ void FNiagaraEmitterScalabilityDetails::CustomizeDetails(IDetailLayoutBuilder& I
 	InDetailLayout.GetObjectsBeingCustomized(ObjectsBeingCustomized);
 	
 	if (UNiagaraEmitter* EmitterBeingCustomized = CastChecked<UNiagaraEmitter>(ObjectsBeingCustomized[0]))
-	{
-		TArray<TSharedPtr<FNiagaraEmitterViewModel>> ExistingViewModels;
-		TNiagaraViewModelManager<UNiagaraEmitter, FNiagaraEmitterViewModel>::GetAllViewModelsForObject(EmitterBeingCustomized, ExistingViewModels);
-		FVersionedNiagaraEmitter VersionedNiagaraEmitter = ExistingViewModels[0]->GetEmitter();
-		FVersionedNiagaraEmitterData* EmitterData = VersionedNiagaraEmitter.GetEmitterData();
-
+	{		
+		TMap<FName, TSet<FProperty*>> CategoryPropertyMap;
 		for (FProperty* ChildProperty : TFieldRange<FProperty>(FVersionedNiagaraEmitterData::StaticStruct()))
 		{
 			if (ChildProperty->HasAllPropertyFlags(CPF_Edit) && ChildProperty->HasMetaData(TEXT("DisplayInScalabilityContext")))
 			{
 				FName Category = FName(ChildProperty->GetMetaData(TEXT("Category")));
-				IDetailCategoryBuilder& CategoryBuilder = InDetailLayout.EditCategory(FName(Category));
-				
-				TSharedPtr<FStructOnScope> StructData = MakeShareable(new FStructOnScope(FVersionedNiagaraEmitterData::StaticStruct(), reinterpret_cast<uint8*>(EmitterData)));
-				if (IDetailPropertyRow* PropertyRow = CategoryBuilder.AddExternalStructureProperty(StructData, ChildProperty->GetFName()))
-				{
-					const auto& PostEditChangeLambda = [VersionedNiagaraEmitter, ChildProperty]
-					{
-						FPropertyChangedEvent ChangeEvent(ChildProperty);
-						VersionedNiagaraEmitter.Emitter->PostEditChangeVersionedProperty(ChangeEvent, VersionedNiagaraEmitter.Version);
-					};
-
-					PropertyRow->GetPropertyHandle()->SetOnChildPropertyValueChanged(FSimpleDelegate::CreateLambda(PostEditChangeLambda));
-					PropertyRow->GetPropertyHandle()->SetOnPropertyValueChanged(FSimpleDelegate::CreateLambda(PostEditChangeLambda));
-				}
+				CategoryPropertyMap.FindOrAdd(Category).Add(ChildProperty);
 			}
 		}
+
+		CustomizeEmitterData(InDetailLayout, EmitterBeingCustomized, CategoryPropertyMap);
 	}
 }
