Fixed crash in validation commandlet trying to initialize sequencer with no slate application.

#rnx

diff --git a/Engine/Plugins/FX/Niagara/Source/NiagaraEditor/Private/ViewModels/NiagaraSystemViewModel.cpp b/Engine/Plugins/FX/Niagara/Source/NiagaraEditor/Private/ViewModels/NiagaraSystemViewModel.cpp
--- a/Engine/Plugins/FX/Niagara/Source/NiagaraEditor/Private/ViewModels/NiagaraSystemViewModel.cpp
+++ b/Engine/Plugins/FX/Niagara/Source/NiagaraEditor/Private/ViewModels/NiagaraSystemViewModel.cpp
@@ -1851,9 +1851,9 @@ UMovieSceneNiagaraEmitterTrack* FNiagaraSystemViewModel::GetTrackForHandleViewMo
 
 void FNiagaraSystemViewModel::SetupSequencer()
 {
-	if (EditMode == ENiagaraSystemViewModelEditMode::EmitterDuringMerge)
+	if (EditMode == ENiagaraSystemViewModelEditMode::EmitterDuringMerge || !FSlateApplicationBase::IsInitialized())
 	{
-		// we don't need a sequencer when merging emitters
+		// we don't need a sequencer when merging emitters or if we're in a commandlet with no slate application
 		return;
 	}
 	NiagaraSequence = NewObject<UNiagaraSequence>(GetTransientPackage());
