Stu's fix for crash when deleting renderers from a running system.

#rnx

diff --git a/Engine/Plugins/FX/Niagara/Source/Niagara/Private/NiagaraSystemInstance.cpp b/Engine/Plugins/FX/Niagara/Source/Niagara/Private/NiagaraSystemInstance.cpp
--- a/Engine/Plugins/FX/Niagara/Source/Niagara/Private/NiagaraSystemInstance.cpp
+++ b/Engine/Plugins/FX/Niagara/Source/Niagara/Private/NiagaraSystemInstance.cpp
@@ -2541,6 +2541,10 @@ void FNiagaraSystemInstance::OnSimulationDestroyed()
 	ensureMsgf(!IsSolo(), TEXT("OnSimulationDestroyed should only happen for systems referencing a simulation from the world manager"));
 	if (SystemSimulation.IsValid())
 	{
+		if (SystemInstanceIndex != INDEX_NONE)
+		{
+			SystemSimulation->SetInstanceState(this, ENiagaraSystemInstanceState::None);
+		}
 		UnbindParameters();
 		SystemSimulation = nullptr;
 	}
