Scrapped the num materials function so it returns 0 now. This fixes an issue with the CustomPrimitivateData customization forcing a refresh on the details panel of the Niagara Component

#rnx

diff --git a/Engine/Plugins/FX/Niagara/Source/Niagara/Private/NiagaraComponent.cpp b/Engine/Plugins/FX/Niagara/Source/Niagara/Private/NiagaraComponent.cpp
--- a/Engine/Plugins/FX/Niagara/Source/Niagara/Private/NiagaraComponent.cpp
+++ b/Engine/Plugins/FX/Niagara/Source/Niagara/Private/NiagaraComponent.cpp
@@ -1928,15 +1928,11 @@ void UNiagaraComponent::SendRenderDynamicData_Concurrent()
 
 int32 UNiagaraComponent::GetNumMaterials() const
 {
-	if (SystemInstanceController)
-	{
-		return SystemInstanceController->GetNumMaterials();
-	}
-
+	// we don't really care for the materials used by the component here
+	// technically should be the number of materials used by each enabled renderer of each enabled emitter
 	return 0;
 }
 
-
 FBoxSphereBounds UNiagaraComponent::CalcBounds(const FTransform& LocalToWorld) const
 {
 	FScopeCycleCounter SystemStatCounter(Asset ? Asset->GetStatID(true, false) : TStatId());
diff --git a/Engine/Plugins/FX/Niagara/Source/Niagara/Private/NiagaraSystemInstanceController.cpp b/Engine/Plugins/FX/Niagara/Source/Niagara/Private/NiagaraSystemInstanceController.cpp
--- a/Engine/Plugins/FX/Niagara/Source/Niagara/Private/NiagaraSystemInstanceController.cpp
+++ b/Engine/Plugins/FX/Niagara/Source/Niagara/Private/NiagaraSystemInstanceController.cpp
@@ -96,27 +96,6 @@ void FNiagaraSystemInstanceController::NotifyRenderersComplete(FNiagaraSystemRen
 	}
 }
 
-int32 FNiagaraSystemInstanceController::GetNumMaterials() const
-{
-	TArray<UMaterialInterface*> UsedMaterials;
-	if (SystemInstance)
-	{
-		for (TSharedPtr<FNiagaraEmitterInstance, ESPMode::ThreadSafe> EmitterInst : SystemInstance->GetEmitters())
-		{
-			if (FVersionedNiagaraEmitterData* EmitterData = EmitterInst->GetCachedEmitterData())
-			{
-				EmitterData->ForEachEnabledRenderer([&](UNiagaraRendererProperties* Properties)
-					{
-						Properties->GetUsedMaterials(EmitterInst.Get(), UsedMaterials);
-					}
-				);
-			}
-		}
-	}
-
-	return UsedMaterials.Num();
-}
-
 void FNiagaraSystemInstanceController::GetUsedMaterials(TArray<UMaterialInterface*>& OutMaterials)
 {
 	if (!SystemInstance.IsValid())
diff --git a/Engine/Plugins/FX/Niagara/Source/Niagara/Public/NiagaraSystemInstanceController.h b/Engine/Plugins/FX/Niagara/Source/Niagara/Public/NiagaraSystemInstanceController.h
--- a/Engine/Plugins/FX/Niagara/Source/Niagara/Public/NiagaraSystemInstanceController.h
+++ b/Engine/Plugins/FX/Niagara/Source/Niagara/Public/NiagaraSystemInstanceController.h
@@ -76,8 +76,7 @@ public:
 	FNiagaraSystemSimulationPtr GetSoloSystemSimulation() const { return ensure(IsValid() && SystemInstance->IsSolo()) ? SystemInstance->GetSystemSimulation() : nullptr; }
 
 	FNiagaraSystemInstanceID GetSystemInstanceID() const { ensure(IsValid()); return SystemInstance->GetId(); }
-
-	int32 GetNumMaterials() const;
+	
 	void GetUsedMaterials(TArray<UMaterialInterface*>& OutMaterials);
 	UMaterialInterface* GetMaterialOverride(const UNiagaraRendererProperties* InProps, int32 InMaterialSubIndex) const;
 	void SetOnMaterialsUpdated(const FOnMaterialsUpdated& Delegate) { OnMaterialsUpdatedDelegate = Delegate; }
