Niagara - Changes to get the niagara plugin content working in a cooked editor.
+ Add a cvar to enable preloading selectable assets so that their asset registry tags are available in a cooked editor.  The preload is used when enumerating emitters and systems for the UI.
+ Add some editor only uproperties to scripts and curve data interfaces to cache data that gets conditionally stripped during cooking.
+ Added a check so that we don't try to compile cooked systems.
+ Added a way to supress the on changed delegate on parameter stores so that we can avoid updating the change id of owning graphs and emitters when the script is making its temporary changes during serialization optimization of the rapid iteration parameters.  This prevents every cooked emitter from trying to merge on load.

diff --git a/Engine/Plugins/FX/Niagara/Source/Niagara/Classes/NiagaraDataInterfaceColorCurve.h b/Engine/Plugins/FX/Niagara/Source/Niagara/Classes/NiagaraDataInterfaceColorCurve.h
--- a/Engine/Plugins/FX/Niagara/Source/Niagara/Classes/NiagaraDataInterfaceColorCurve.h
+++ b/Engine/Plugins/FX/Niagara/Source/Niagara/Classes/NiagaraDataInterfaceColorCurve.h
@@ -34,6 +34,20 @@ public:
 	UPROPERTY(EditAnywhere, Category = "Curve")
 	FRichCurve AlphaCurve;
 
+
+#if WITH_EDITORONLY_DATA
+private:
+	UPROPERTY()
+	FRichCurve RedCurveCookedEditorCache;
+	UPROPERTY()
+	FRichCurve GreenCurveCookedEditorCache;
+	UPROPERTY()
+	FRichCurve BlueCurveCookedEditorCache;
+	UPROPERTY()
+	FRichCurve AlphaCurveCookedEditorCache;
+public:
+#endif
+
 	//UObject Interface
 	virtual void PostInitProperties() override;
 	virtual void Serialize(FArchive& Ar) override;
diff --git a/Engine/Plugins/FX/Niagara/Source/Niagara/Classes/NiagaraDataInterfaceCurve.h b/Engine/Plugins/FX/Niagara/Source/Niagara/Classes/NiagaraDataInterfaceCurve.h
--- a/Engine/Plugins/FX/Niagara/Source/Niagara/Classes/NiagaraDataInterfaceCurve.h
+++ b/Engine/Plugins/FX/Niagara/Source/Niagara/Classes/NiagaraDataInterfaceCurve.h
@@ -20,6 +20,13 @@ public:
 	UPROPERTY(EditAnywhere, Category = "Curve")
 	FRichCurve Curve; 
 
+#if WITH_EDITORONLY_DATA
+private:
+	UPROPERTY()
+	FRichCurve CurveCookedEditorCache;
+public:
+#endif
+
 	enum
 	{
 		CurveLUTNumElems = 1,
diff --git a/Engine/Plugins/FX/Niagara/Source/Niagara/Classes/NiagaraDataInterfaceVector2DCurve.h b/Engine/Plugins/FX/Niagara/Source/Niagara/Classes/NiagaraDataInterfaceVector2DCurve.h
--- a/Engine/Plugins/FX/Niagara/Source/Niagara/Classes/NiagaraDataInterfaceVector2DCurve.h
+++ b/Engine/Plugins/FX/Niagara/Source/Niagara/Classes/NiagaraDataInterfaceVector2DCurve.h
@@ -23,6 +23,15 @@ public:
 	UPROPERTY(EditAnywhere, Category = "Curve")
 	FRichCurve YCurve;
 
+#if WITH_EDITORONLY_DATA
+private:
+	UPROPERTY()
+	FRichCurve XCurveCookedEditorCache;
+	UPROPERTY()
+	FRichCurve YCurveCookedEditorCache;
+public:
+#endif
+
 	enum
 	{
 		CurveLUTNumElems = 2,
diff --git a/Engine/Plugins/FX/Niagara/Source/Niagara/Classes/NiagaraDataInterfaceVector4Curve.h b/Engine/Plugins/FX/Niagara/Source/Niagara/Classes/NiagaraDataInterfaceVector4Curve.h
--- a/Engine/Plugins/FX/Niagara/Source/Niagara/Classes/NiagaraDataInterfaceVector4Curve.h
+++ b/Engine/Plugins/FX/Niagara/Source/Niagara/Classes/NiagaraDataInterfaceVector4Curve.h
@@ -33,6 +33,19 @@ public:
 
 	UPROPERTY(EditAnywhere, Category = "Curve")
 	FRichCurve WCurve;
+
+#if WITH_EDITORONLY_DATA
+private:
+	UPROPERTY()
+	FRichCurve XCurveCookedEditorCache;
+	UPROPERTY()
+	FRichCurve YCurveCookedEditorCache;
+	UPROPERTY()
+	FRichCurve ZCurveCookedEditorCache;
+	UPROPERTY()
+	FRichCurve WCurveCookedEditorCache;
+public:
+#endif
 	
 	enum
 	{
diff --git a/Engine/Plugins/FX/Niagara/Source/Niagara/Classes/NiagaraDataInterfaceVectorCurve.h b/Engine/Plugins/FX/Niagara/Source/Niagara/Classes/NiagaraDataInterfaceVectorCurve.h
--- a/Engine/Plugins/FX/Niagara/Source/Niagara/Classes/NiagaraDataInterfaceVectorCurve.h
+++ b/Engine/Plugins/FX/Niagara/Source/Niagara/Classes/NiagaraDataInterfaceVectorCurve.h
@@ -31,6 +31,17 @@ public:
 	UPROPERTY(EditAnywhere, Category = "Curve")
 	FRichCurve ZCurve;
 
+#if WITH_EDITORONLY_DATA
+private:
+	UPROPERTY()
+	FRichCurve XCurveCookedEditorCache;
+	UPROPERTY()
+	FRichCurve YCurveCookedEditorCache;
+	UPROPERTY()
+	FRichCurve ZCurveCookedEditorCache;
+public:
+#endif
+
 	enum
 	{
 		CurveLUTNumElems = 3,
diff --git a/Engine/Plugins/FX/Niagara/Source/Niagara/Classes/NiagaraScript.h b/Engine/Plugins/FX/Niagara/Source/Niagara/Classes/NiagaraScript.h
--- a/Engine/Plugins/FX/Niagara/Source/Niagara/Classes/NiagaraScript.h
+++ b/Engine/Plugins/FX/Niagara/Source/Niagara/Classes/NiagaraScript.h
@@ -799,6 +799,9 @@ public:
 	FNiagaraParameterStore RapidIterationParameters;
 	
 #if WITH_EDITORONLY_DATA
+	UPROPERTY()
+	FNiagaraParameterStore RapidIterationParametersCookedEditorCache;
+
 	/** This is used as a transient value to open a specific version in the graph editor */
 	UPROPERTY(Transient)
 	FGuid VersionToOpenInEditor;
diff --git a/Engine/Plugins/FX/Niagara/Source/Niagara/Private/NiagaraDataInterfaceColorCurve.cpp b/Engine/Plugins/FX/Niagara/Source/Niagara/Private/NiagaraDataInterfaceColorCurve.cpp
--- a/Engine/Plugins/FX/Niagara/Source/Niagara/Private/NiagaraDataInterfaceColorCurve.cpp
+++ b/Engine/Plugins/FX/Niagara/Source/Niagara/Private/NiagaraDataInterfaceColorCurve.cpp
@@ -46,21 +46,26 @@ void UNiagaraDataInterfaceColorCurve::Serialize(FArchive& Ar)
 	{
 		UpdateLUT(true);
 
-		FRichCurve TempRCurve;
-		FRichCurve TempGCurve;
-		FRichCurve TempBCurve;
-		FRichCurve TempACurve;
-		Exchange(RedCurve, TempRCurve);
-		Exchange(GreenCurve, TempGCurve);
-		Exchange(BlueCurve, TempBCurve);
-		Exchange(AlphaCurve, TempACurve);
+		Exchange(RedCurve, RedCurveCookedEditorCache);
+		Exchange(GreenCurve, GreenCurveCookedEditorCache);
+		Exchange(BlueCurve, BlueCurveCookedEditorCache);
+		Exchange(AlphaCurve, AlphaCurveCookedEditorCache);
 
 		Super::Serialize(Ar);
 
-		Exchange(RedCurve, TempRCurve);
-		Exchange(GreenCurve, TempGCurve);
-		Exchange(BlueCurve, TempBCurve);
-		Exchange(AlphaCurve, TempACurve);
+		Exchange(RedCurve, RedCurveCookedEditorCache);
+		Exchange(GreenCurve, GreenCurveCookedEditorCache);
+		Exchange(BlueCurve, BlueCurveCookedEditorCache);
+		Exchange(AlphaCurve, AlphaCurveCookedEditorCache);
+	}
+	else if (bUseLUT && Ar.IsLoading() && GetOutermost()->HasAnyPackageFlags(PKG_Cooked))
+	{
+		Super::Serialize(Ar);
+
+		Exchange(RedCurve, RedCurveCookedEditorCache);
+		Exchange(GreenCurve, GreenCurveCookedEditorCache);
+		Exchange(BlueCurve, BlueCurveCookedEditorCache);
+		Exchange(AlphaCurve, AlphaCurveCookedEditorCache);
 	}
 	else
 #endif
diff --git a/Engine/Plugins/FX/Niagara/Source/Niagara/Private/NiagaraDataInterfaceCurve.cpp b/Engine/Plugins/FX/Niagara/Source/Niagara/Private/NiagaraDataInterfaceCurve.cpp
--- a/Engine/Plugins/FX/Niagara/Source/Niagara/Private/NiagaraDataInterfaceCurve.cpp
+++ b/Engine/Plugins/FX/Niagara/Source/Niagara/Private/NiagaraDataInterfaceCurve.cpp
@@ -42,12 +42,17 @@ void UNiagaraDataInterfaceCurve::Serialize(FArchive& Ar)
 	{
 		UpdateLUT(true);
 
-		FRichCurve TempCurve;
-		Exchange(Curve, TempCurve);
+		Exchange(Curve, CurveCookedEditorCache);
 
 		Super::Serialize(Ar);
 
-		Exchange(Curve, TempCurve);
+		Exchange(Curve, CurveCookedEditorCache);
+	}
+	else if (bUseLUT && Ar.IsLoading() && GetOutermost()->HasAnyPackageFlags(PKG_Cooked))
+	{
+		Super::Serialize(Ar);
+
+		Exchange(Curve, CurveCookedEditorCache);
 	}
 	else
 #endif
diff --git a/Engine/Plugins/FX/Niagara/Source/Niagara/Private/NiagaraDataInterfaceVector2DCurve.cpp b/Engine/Plugins/FX/Niagara/Source/Niagara/Private/NiagaraDataInterfaceVector2DCurve.cpp
--- a/Engine/Plugins/FX/Niagara/Source/Niagara/Private/NiagaraDataInterfaceVector2DCurve.cpp
+++ b/Engine/Plugins/FX/Niagara/Source/Niagara/Private/NiagaraDataInterfaceVector2DCurve.cpp
@@ -45,21 +45,25 @@ void UNiagaraDataInterfaceVector2DCurve::Serialize(FArchive& Ar)
 	{
 		UpdateLUT(true);
 
-		FRichCurve TempXCurve;
-		FRichCurve TempYCurve;
-		Exchange(XCurve, TempXCurve);
-		Exchange(YCurve, TempYCurve);
+		Exchange(XCurve, XCurveCookedEditorCache);
+		Exchange(YCurve, YCurveCookedEditorCache);
 
 		Super::Serialize(Ar);
 
-		Exchange(XCurve, TempXCurve);
-		Exchange(YCurve, TempYCurve);
+		Exchange(XCurve, XCurveCookedEditorCache);
+		Exchange(YCurve, YCurveCookedEditorCache);
 	}
-	else
-#endif
+	else if (bUseLUT && Ar.IsLoading() && GetOutermost()->HasAnyPackageFlags(PKG_Cooked))
 	{
 		Super::Serialize(Ar);
+
+		Exchange(XCurve, XCurveCookedEditorCache);
+		Exchange(YCurve, YCurveCookedEditorCache);
 	}
+#endif
+	{
+		Super::Serialize(Ar);
+	} 
 }
 
 void UNiagaraDataInterfaceVector2DCurve::UpdateTimeRanges()
diff --git a/Engine/Plugins/FX/Niagara/Source/Niagara/Private/NiagaraDataInterfaceVector4Curve.cpp b/Engine/Plugins/FX/Niagara/Source/Niagara/Private/NiagaraDataInterfaceVector4Curve.cpp
--- a/Engine/Plugins/FX/Niagara/Source/Niagara/Private/NiagaraDataInterfaceVector4Curve.cpp
+++ b/Engine/Plugins/FX/Niagara/Source/Niagara/Private/NiagaraDataInterfaceVector4Curve.cpp
@@ -46,21 +46,26 @@ void UNiagaraDataInterfaceVector4Curve::Serialize(FArchive& Ar)
 	{
 		UpdateLUT(true);
 
-		FRichCurve TempXCurve;
-		FRichCurve TempYCurve;
-		FRichCurve TempZCurve;
-		FRichCurve TempWCurve;
-		Exchange(XCurve, TempXCurve);
-		Exchange(YCurve, TempYCurve);
-		Exchange(ZCurve, TempZCurve);
-		Exchange(WCurve, TempWCurve);
+		Exchange(XCurve, XCurveCookedEditorCache);
+		Exchange(YCurve, YCurveCookedEditorCache);
+		Exchange(ZCurve, ZCurveCookedEditorCache);
+		Exchange(WCurve, WCurveCookedEditorCache);
 
 		Super::Serialize(Ar);
 
-		Exchange(XCurve, TempXCurve);
-		Exchange(YCurve, TempYCurve);
-		Exchange(ZCurve, TempZCurve);
-		Exchange(WCurve, TempWCurve);
+		Exchange(XCurve, XCurveCookedEditorCache);
+		Exchange(YCurve, YCurveCookedEditorCache);
+		Exchange(ZCurve, ZCurveCookedEditorCache);
+		Exchange(WCurve, WCurveCookedEditorCache);
+	}
+	else if (bUseLUT && Ar.IsLoading() && GetOutermost()->HasAnyPackageFlags(PKG_Cooked))
+	{
+		Super::Serialize(Ar);
+
+		Exchange(XCurve, XCurveCookedEditorCache);
+		Exchange(YCurve, YCurveCookedEditorCache);
+		Exchange(ZCurve, ZCurveCookedEditorCache);
+		Exchange(WCurve, WCurveCookedEditorCache);
 	}
 	else
 #endif
diff --git a/Engine/Plugins/FX/Niagara/Source/Niagara/Private/NiagaraDataInterfaceVectorCurve.cpp b/Engine/Plugins/FX/Niagara/Source/Niagara/Private/NiagaraDataInterfaceVectorCurve.cpp
--- a/Engine/Plugins/FX/Niagara/Source/Niagara/Private/NiagaraDataInterfaceVectorCurve.cpp
+++ b/Engine/Plugins/FX/Niagara/Source/Niagara/Private/NiagaraDataInterfaceVectorCurve.cpp
@@ -46,18 +46,23 @@ void UNiagaraDataInterfaceVectorCurve::Serialize(FArchive& Ar)
 	{
 		UpdateLUT(true);
 
-		FRichCurve TempXCurve;
-		FRichCurve TempYCurve;
-		FRichCurve TempZCurve;
-		Exchange(XCurve, TempXCurve);
-		Exchange(YCurve, TempYCurve);
-		Exchange(ZCurve, TempZCurve);
+		Exchange(XCurve, XCurveCookedEditorCache);
+		Exchange(YCurve, YCurveCookedEditorCache);
+		Exchange(ZCurve, ZCurveCookedEditorCache);
 
 		Super::Serialize(Ar);
 
-		Exchange(XCurve, TempXCurve);
-		Exchange(YCurve, TempYCurve);
-		Exchange(ZCurve, TempZCurve);
+		Exchange(XCurve, XCurveCookedEditorCache);
+		Exchange(YCurve, YCurveCookedEditorCache);
+		Exchange(ZCurve, ZCurveCookedEditorCache);
+	}
+	else if (bUseLUT && Ar.IsLoading() && GetOutermost()->HasAnyPackageFlags(PKG_Cooked))
+	{
+		Super::Serialize(Ar);
+
+		Exchange(XCurve, XCurveCookedEditorCache);
+		Exchange(YCurve, YCurveCookedEditorCache);
+		Exchange(ZCurve, ZCurveCookedEditorCache);
 	}
 	else
 #endif
diff --git a/Engine/Plugins/FX/Niagara/Source/Niagara/Private/NiagaraParameterStore.cpp b/Engine/Plugins/FX/Niagara/Source/Niagara/Private/NiagaraParameterStore.cpp
--- a/Engine/Plugins/FX/Niagara/Source/Niagara/Private/NiagaraParameterStore.cpp
+++ b/Engine/Plugins/FX/Niagara/Source/Niagara/Private/NiagaraParameterStore.cpp
@@ -215,7 +215,10 @@ FNiagaraParameterStore& FNiagaraParameterStore::operator=(const FNiagaraParamete
 	SetOriginalPositionData(Other.OriginalPositionData);
 	++LayoutVersion;
 #if WITH_EDITOR
-	OnChangedDelegate.Broadcast();
+	if (bSuppressOnChanged == false)
+	{
+		OnChangedDelegate.Broadcast();
+	}
 #endif
 	//Don't copy bindings. We just want the data.
 	return *this;
@@ -1223,7 +1226,10 @@ void FNiagaraParameterStore::OnLayoutChange()
 	++LayoutVersion;
 
 #if WITH_EDITOR
-	OnChangedDelegate.Broadcast();
+	if (bSuppressOnChanged == false)
+	{
+		OnChangedDelegate.Broadcast();
+	}
 #endif
 }
 
diff --git a/Engine/Plugins/FX/Niagara/Source/Niagara/Private/NiagaraScript.cpp b/Engine/Plugins/FX/Niagara/Source/Niagara/Private/NiagaraScript.cpp
--- a/Engine/Plugins/FX/Niagara/Source/Niagara/Private/NiagaraScript.cpp
+++ b/Engine/Plugins/FX/Niagara/Source/Niagara/Private/NiagaraScript.cpp
@@ -1611,6 +1611,10 @@ void UNiagaraScript::Serialize(FArchive& Ar)
 	Ar.UsingCustomVersion(FNiagaraCustomVersion::GUID);		// only changes version if not loading
 	const int32 NiagaraVer = Ar.CustomVer(FNiagaraCustomVersion::GUID);
 
+#if WITH_EDITORONLY_DATA
+	FNiagaraParameterStore::FScopedSuppressOnChanged SupressOnChanged(RapidIterationParameters);
+#endif
+
 	FNiagaraParameterStore TemporaryStore;
 	int32 NumRemoved = 0;
 	if (Ar.IsCooking())
@@ -1629,6 +1633,11 @@ void UNiagaraScript::Serialize(FArchive& Ar)
 			// Copy off the parameter store for now..
 			TemporaryStore = RapidIterationParameters;
 
+#if WITH_EDITORONLY_DATA
+			// Cache the unmodified rapid iteration parameters in an editor only property so that they can be restored in a cooked editor.
+			RapidIterationParametersCookedEditorCache = RapidIterationParameters;
+#endif
+
 			auto ParameterVariables = TemporaryStore.ReadParameterVariables();
 
 			// Get the active parameters
@@ -1679,12 +1688,26 @@ void UNiagaraScript::Serialize(FArchive& Ar)
 	Super::Serialize(Ar);
 
 	// Restore after serialize
-	if (Ar.IsCooking() && NumRemoved > 0)
+	if (Ar.IsCooking())
 	{
-		RapidIterationParameters = TemporaryStore;
+		if (NumRemoved > 0)
+		{
+			RapidIterationParameters = TemporaryStore;
+		}
+#if WITH_EDITORONLY_DATA
+		RapidIterationParametersCookedEditorCache.Empty();
+#endif
 	}
 
 #if WITH_EDITORONLY_DATA
+	if (Ar.IsLoading() && RapidIterationParametersCookedEditorCache.Num() > 0)
+	{
+		// If we're loading and the cooked editor cache is available, re-populate the parameters from the cache after they're
+		// deserialized.
+		RapidIterationParameters = RapidIterationParametersCookedEditorCache;
+		RapidIterationParametersCookedEditorCache.Empty();
+	}
+
 	if (Ar.IsCooking() && Ar.IsSaving())
 	{
 		if (!HasValidParameterBindings())
@@ -3388,6 +3411,12 @@ bool UNiagaraScript::SynchronizeExecutablesWithMaster(const UNiagaraScript* Scri
 {
 	check(IsInGameThread());
 
+	if (Script->CachedScriptVM.IsValid() == false)
+	{
+		// If the master script has no valid VM then there's nothing to synchronize.
+		return false;
+	}
+
 	FNiagaraVMExecutableDataId Id;
 	ComputeVMCompilationId(Id, FGuid());
 
diff --git a/Engine/Plugins/FX/Niagara/Source/Niagara/Private/NiagaraSystem.cpp b/Engine/Plugins/FX/Niagara/Source/Niagara/Private/NiagaraSystem.cpp
--- a/Engine/Plugins/FX/Niagara/Source/Niagara/Private/NiagaraSystem.cpp
+++ b/Engine/Plugins/FX/Niagara/Source/Niagara/Private/NiagaraSystem.cpp
@@ -549,6 +549,7 @@ void UNiagaraSystem::UpdateSystemAfterLoad()
 
 #if WITH_EDITORONLY_DATA
 	// check if the system needs to be compiled and start a compile task if necessary
+	if (GetOutermost()->HasAnyPackageFlags(EPackageFlags::PKG_Cooked) == false)
 	{
 		bool bSystemScriptsAreSynchronized = true;
 		for (UNiagaraScript* SystemScript : AllSystemScripts)
diff --git a/Engine/Plugins/FX/Niagara/Source/Niagara/Public/NiagaraParameterStore.h b/Engine/Plugins/FX/Niagara/Source/Niagara/Public/NiagaraParameterStore.h
--- a/Engine/Plugins/FX/Niagara/Source/Niagara/Public/NiagaraParameterStore.h
+++ b/Engine/Plugins/FX/Niagara/Source/Niagara/Public/NiagaraParameterStore.h
@@ -167,6 +167,13 @@ struct NIAGARA_API FNiagaraParameterStore
 {
 #if WITH_EDITOR
 	DECLARE_MULTICAST_DELEGATE(FOnChanged);
+	struct FScopedSuppressOnChanged : TGuardValue<bool>
+	{
+		FScopedSuppressOnChanged(FNiagaraParameterStore& TargetStore)
+			: TGuardValue(TargetStore.bSuppressOnChanged, true)
+		{
+		}
+	};
 #endif
 
 	GENERATED_USTRUCT_BODY()
@@ -229,6 +236,7 @@ private:
 
 #if WITH_EDITOR
 	FOnChanged OnChangedDelegate;
+	bool bSuppressOnChanged = false;
 #endif
 
 	void SetPositionData(const FName& Name, const FVector& Position);
@@ -610,7 +618,10 @@ public:
 	{ 
 		bParametersDirty = true;
 #if WITH_EDITOR
-		OnChangedDelegate.Broadcast();
+		if (bSuppressOnChanged == false)
+		{
+			OnChangedDelegate.Broadcast();
+		}
 #endif
 	}
 
@@ -618,7 +629,10 @@ public:
 	{ 
 		bInterfacesDirty = true;
 #if WITH_EDITOR
-		OnChangedDelegate.Broadcast();
+		if (bSuppressOnChanged == false)
+		{
+			OnChangedDelegate.Broadcast();
+		}
 #endif
 	}
 
@@ -626,7 +640,10 @@ public:
 	{
 		bUObjectsDirty = true;
 #if WITH_EDITOR
-		OnChangedDelegate.Broadcast();
+		if (bSuppressOnChanged == false)
+		{
+			OnChangedDelegate.Broadcast();
+		}
 #endif
 	}
 
@@ -636,7 +653,10 @@ public:
 		bInterfacesDirty = true;
 		bParametersDirty = true;
 #if WITH_EDITOR
-		OnChangedDelegate.Broadcast();
+		if (bSuppressOnChanged == false)
+		{
+			OnChangedDelegate.Broadcast();
+		}
 #endif
 	}
 
diff --git a/Engine/Plugins/FX/Niagara/Source/NiagaraEditor/Private/NiagaraEditorModule.cpp b/Engine/Plugins/FX/Niagara/Source/NiagaraEditor/Private/NiagaraEditorModule.cpp
--- a/Engine/Plugins/FX/Niagara/Source/NiagaraEditor/Private/NiagaraEditorModule.cpp
+++ b/Engine/Plugins/FX/Niagara/Source/NiagaraEditor/Private/NiagaraEditorModule.cpp
@@ -168,6 +168,13 @@ static FAutoConsoleVariableRef CVarShowNiagaraDeveloperWindows(
 	ECVF_Default
 	);
 
+int32 GbPreloadSelectablePluginAssetsOnDemand = 1;
+static FAutoConsoleVariableRef CVarPreloadSelectablePluginAssetsOnDemand(
+	TEXT("fx.Niagara.PreloadSelectablePluginAssetsOnDemand"),
+	GbPreloadSelectablePluginAssetsOnDemand,
+	TEXT("If > 0 then niagara system, emitter, and script assets provided by the niagara plugin will be preloaded when a dialog is opened to select them. This is a temoporary workaround for asset registry issues in cooked editor builds.\n"),
+	ECVF_Default);
+
 // this is required for gpu script compilation ticks
 static FNiagaraShaderQueueTickable NiagaraShaderQueueProcessor;
 
@@ -1903,6 +1910,24 @@ void FNiagaraEditorModule::OnPerfBaselineWindowClosed(const TSharedRef<SWindow>&
 	ClosedWindow->SetContent(SNullWidget::NullWidget);
 	BaselineViewport.Reset();
 }
+void FNiagaraEditorModule::PreloadSelectablePluginAssetsByClass(UClass* InClass)
+{
+	if (GbPreloadSelectablePluginAssetsOnDemand && PluginAssetClassesPreloaded.Contains(InClass) == false)
+	{
+		FAssetRegistryModule& AssetRegistryModule = FModuleManager::LoadModuleChecked<FAssetRegistryModule>(TEXT("AssetRegistry"));
+		TArray<FAssetData> AllClassAssets;
+		AssetRegistryModule.Get().GetAssetsByClass(InClass->GetClassPathName(), AllClassAssets);
+		for (FAssetData& ClassAsset : AllClassAssets)
+		{
+			if (ClassAsset.HasAnyPackageFlags(PKG_Cooked) && ClassAsset.IsAssetLoaded() == false && FNiagaraEditorUtilities::IsEnginePluginAsset(ClassAsset))
+			{
+				ClassAsset.GetAsset();
+			}
+		}
+		PluginAssetClassesPreloaded.Add(InClass);
+	}
+}
+
 #endif
 
 #undef LOCTEXT_NAMESPACE
diff --git a/Engine/Plugins/FX/Niagara/Source/NiagaraEditor/Private/NiagaraEditorUtilities.cpp b/Engine/Plugins/FX/Niagara/Source/NiagaraEditor/Private/NiagaraEditorUtilities.cpp
--- a/Engine/Plugins/FX/Niagara/Source/NiagaraEditor/Private/NiagaraEditorUtilities.cpp
+++ b/Engine/Plugins/FX/Niagara/Source/NiagaraEditor/Private/NiagaraEditorUtilities.cpp
@@ -1373,6 +1373,8 @@ void FNiagaraEditorUtilities::PreprocessFunctionGraph(const UEdGraphSchema_Niaga
 
 void FNiagaraEditorUtilities::GetFilteredScriptAssets(FGetFilteredScriptAssetsOptions InFilter, TArray<FAssetData>& OutFilteredScriptAssets)
 {
+	FNiagaraEditorModule::Get().PreloadSelectablePluginAssetsByClass(UNiagaraScript::StaticClass());
+
 	FARFilter ScriptFilter;
 	ScriptFilter.ClassPaths.Add(UNiagaraScript::StaticClass()->GetClassPathName());
 
@@ -1550,7 +1552,7 @@ ENiagaraScriptLibraryVisibility FNiagaraEditorUtilities::GetScriptAssetVisibilit
 	{
 		if (ScriptAssetData.IsAssetLoaded())
 		{
-			UNiagaraScript* Script = static_cast<UNiagaraScript*>(ScriptAssetData.GetAsset());
+			UNiagaraScript* Script = Cast<UNiagaraScript>(ScriptAssetData.GetAsset());
 			if (Script != nullptr)
 			{
 				ScriptVisibility = Script->GetLatestScriptData()->LibraryVisibility;
@@ -1610,6 +1612,14 @@ bool FNiagaraEditorUtilities::IsScriptAssetInLibrary(const FAssetData& ScriptAss
 	return GetScriptAssetVisibility(ScriptAssetData) == ENiagaraScriptLibraryVisibility::Library;
 }
 
+bool FNiagaraEditorUtilities::IsEnginePluginAsset(const FAssetData& InAssetData)
+{
+	FString PackagePathLocal = "";
+	return 
+		FPackageName::TryConvertGameRelativePackagePathToLocalPath(InAssetData.PackagePath.ToString(), PackagePathLocal) &&
+		FPaths::IsUnderDirectory(PackagePathLocal, FPaths::EnginePluginsDir() + "FX/Niagara");
+}
+
 int32 FNiagaraEditorUtilities::GetWeightForItem(const TSharedPtr<FNiagaraMenuAction_Generic>& InCurrentAction, const TArray<FString>& InFilterTerms)
 {
 	// The overall 'weight'
diff --git a/Engine/Plugins/FX/Niagara/Source/NiagaraEditor/Private/SNewSystemDialog.cpp b/Engine/Plugins/FX/Niagara/Source/NiagaraEditor/Private/SNewSystemDialog.cpp
--- a/Engine/Plugins/FX/Niagara/Source/NiagaraEditor/Private/SNewSystemDialog.cpp
+++ b/Engine/Plugins/FX/Niagara/Source/NiagaraEditor/Private/SNewSystemDialog.cpp
@@ -2,6 +2,7 @@
 
 #include "SNewSystemDialog.h"
 #include "NiagaraSystem.h"
+#include "NiagaraEditorModule.h"
 #include "NiagaraEditorStyle.h"
 #include "NiagaraEditor/Private/SNiagaraAssetPickerList.h"
 
@@ -39,12 +40,14 @@ void SNewSystemDialog::Construct(const FArguments& InArgs)
 	TemplateAndBehaviorsOnlyTabOptions.ChangeTabState(ENiagaraScriptTemplateSpecification::Template, true);
 	TemplateAndBehaviorsOnlyTabOptions.ChangeTabState(ENiagaraScriptTemplateSpecification::Behavior, true);
 
+	FNiagaraEditorModule::Get().PreloadSelectablePluginAssetsByClass(UNiagaraEmitter::StaticClass());
 	SAssignNew(EmitterAssetPicker, SNiagaraAssetPickerList, UNiagaraEmitter::StaticClass())
 		.OnTemplateAssetActivated(this, &SNewSystemDialog::OnEmitterAssetsActivated)
 		.ViewOptions(DisplayAllViewOptions)
 		.TabOptions(AllTabOptions)
 		.bAllowMultiSelect(true);
 
+	FNiagaraEditorModule::Get().PreloadSelectablePluginAssetsByClass(UNiagaraSystem::StaticClass());
 	SAssignNew(TemplateBehaviorAssetPicker, SNiagaraAssetPickerList, UNiagaraSystem::StaticClass())
 		.OnTemplateAssetActivated(this, &SNewSystemDialog::ConfirmSelection)
 		.ViewOptions(DisplayAllViewOptions)
diff --git a/Engine/Plugins/FX/Niagara/Source/NiagaraEditor/Public/NiagaraEditorModule.h b/Engine/Plugins/FX/Niagara/Source/NiagaraEditor/Public/NiagaraEditorModule.h
--- a/Engine/Plugins/FX/Niagara/Source/NiagaraEditor/Public/NiagaraEditorModule.h
+++ b/Engine/Plugins/FX/Niagara/Source/NiagaraEditor/Public/NiagaraEditorModule.h
@@ -48,6 +48,7 @@ class UNiagaraParameterCollection;
 DECLARE_STATS_GROUP(TEXT("Niagara Editor"), STATGROUP_NiagaraEditor, STATCAT_Advanced);
 
 extern NIAGARAEDITOR_API int32 GbShowNiagaraDeveloperWindows;
+extern NIAGARAEDITOR_API int32 GbPreloadSelectablePluginAssetsOnDemand;
 
 /* Defines methods for allowing external modules to supply widgets to the core editor module. */
 class NIAGARAEDITOR_API INiagaraEditorWidgetProvider
@@ -211,6 +212,8 @@ public:
 
 	NIAGARAEDITOR_API UNiagaraParameterCollection* FindCollectionForVariable(const FString& VariableName);
 
+	void PreloadSelectablePluginAssetsByClass(UClass* InClass);
+
 private:
 	class FDeferredDestructionContainerBase
 	{
@@ -387,4 +390,6 @@ private:
 
 	TAssetPreloadCache<UNiagaraParameterCollection> ParameterCollectionAssetCache;
 	TAssetPreloadCache<UNiagaraParameterDefinitions> ParameterDefinitionsAssetCache;
+
+	TArray<UClass*> PluginAssetClassesPreloaded;
 };
diff --git a/Engine/Plugins/FX/Niagara/Source/NiagaraEditor/Public/NiagaraEditorUtilities.h b/Engine/Plugins/FX/Niagara/Source/NiagaraEditor/Public/NiagaraEditorUtilities.h
--- a/Engine/Plugins/FX/Niagara/Source/NiagaraEditor/Public/NiagaraEditorUtilities.h
+++ b/Engine/Plugins/FX/Niagara/Source/NiagaraEditor/Public/NiagaraEditorUtilities.h
@@ -233,6 +233,8 @@ namespace FNiagaraEditorUtilities
 
 	NIAGARAEDITOR_API bool IsScriptAssetInLibrary(const FAssetData& ScriptAssetData);
 
+	NIAGARAEDITOR_API bool IsEnginePluginAsset(const FAssetData& InAssetData);
+
 	NIAGARAEDITOR_API int32 GetWeightForItem(const TSharedPtr<FNiagaraMenuAction_Generic>& Item, const TArray<FString>& FilterTerms);
 
 	NIAGARAEDITOR_API bool DoesItemMatchFilterText(const FText& FilterText, const TSharedPtr<FNiagaraMenuAction_Generic>& Item);
