Improved some wording in comments and property descriptions

diff --git a/Engine/Plugins/FX/Niagara/Source/Niagara/Private/NiagaraGPUProfiler.cpp b/Engine/Plugins/FX/Niagara/Source/Niagara/Private/NiagaraGPUProfiler.cpp
--- a/Engine/Plugins/FX/Niagara/Source/Niagara/Private/NiagaraGPUProfiler.cpp
+++ b/Engine/Plugins/FX/Niagara/Source/Niagara/Private/NiagaraGPUProfiler.cpp
@@ -9,7 +9,7 @@ bool GNiagaraGpuProfilingEnabled = true;
 static FAutoConsoleVariableRef CVarNiagaraGpuProfilingEnabled(
 	TEXT("fx.Niagara.GpuProfiling.Enabled"),
 	GNiagaraGpuProfilingEnabled,
-	TEXT("Master control to allow Niagara to use GPU profiling or not.\n"),
+	TEXT("Primary control to allow Niagara to use GPU profiling or not.\n"),
 	ECVF_Default
 );
 
diff --git a/Engine/Plugins/FX/Niagara/Source/Niagara/Private/NiagaraRendererMeshes.cpp b/Engine/Plugins/FX/Niagara/Source/Niagara/Private/NiagaraRendererMeshes.cpp
--- a/Engine/Plugins/FX/Niagara/Source/Niagara/Private/NiagaraRendererMeshes.cpp
+++ b/Engine/Plugins/FX/Niagara/Source/Niagara/Private/NiagaraRendererMeshes.cpp
@@ -188,7 +188,7 @@ void FNiagaraRendererMeshes::Initialize(const UNiagaraRendererProperties* InProp
 			MeshData.Rotation = FQuat4f(MeshProperties.Rotation.Quaternion());
 			MeshData.LocalBounds = Mesh->GetExtendedBounds().GetBox();
 
-			// Create an index remap from mesh material index to it's index in the master material list
+			// Create an index remap from mesh material index to it's index in the base material list
 			TArray<UMaterialInterface*> MeshMaterials;
 			Properties->GetUsedMeshMaterials(SourceMeshIndex, Emitter, MeshMaterials);
 			for (auto MeshMaterial : MeshMaterials)
diff --git a/Engine/Plugins/FX/Niagara/Source/Niagara/Public/NiagaraDebuggerCommon.h b/Engine/Plugins/FX/Niagara/Source/Niagara/Public/NiagaraDebuggerCommon.h
--- a/Engine/Plugins/FX/Niagara/Source/Niagara/Public/NiagaraDebuggerCommon.h
+++ b/Engine/Plugins/FX/Niagara/Source/Niagara/Public/NiagaraDebuggerCommon.h
@@ -428,11 +428,11 @@ struct NIAGARA_API FNiagaraDebugHUDSettingsData
 	bool bWidgetEnabled = false;
 #endif
 
-	/** Master control for all HUD features. */
+	/** Primary control for all HUD features. */
 	UPROPERTY(EditAnywhere, Category = "Debug General", meta = (DisplayName = "Debug HUD Enabled"))
 	bool bHudEnabled = true;
 
-	/** Master control for HUD rendering. */
+	/** Primary control for HUD rendering. */
 	UPROPERTY(EditAnywhere, Category = "Debug General", meta = (DisplayName = "Debug HUD Rendering Enabled"))
 	bool bHudRenderingEnabled = true;
 
