Fix parameter definition synchronization handling when programmatically renaming graph pins

diff --git a/Engine/Plugins/FX/Niagara/Source/NiagaraEditor/Private/NiagaraGraph.cpp b/Engine/Plugins/FX/Niagara/Source/NiagaraEditor/Private/NiagaraGraph.cpp
--- a/Engine/Plugins/FX/Niagara/Source/NiagaraEditor/Private/NiagaraGraph.cpp
+++ b/Engine/Plugins/FX/Niagara/Source/NiagaraEditor/Private/NiagaraGraph.cpp
@@ -2075,19 +2075,19 @@ void CopyScriptVariableDataForRename(const UNiagaraScriptVariable& OldScriptVari
 
 bool UNiagaraGraph::RenameParameterFromPin(const FNiagaraVariable& Parameter, FName NewName, UEdGraphPin* InPin)
 {
+	auto TrySubscribeScriptVarToDefinitionByName = [this](UNiagaraScriptVariable* TargetScriptVar) {
+		UNiagaraScript* OuterScript = GetTypedOuter<UNiagaraScript>();
+		bool bForDataProcessingOnly = true;
+		TSharedPtr<FNiagaraScriptViewModel> ScriptViewModel = MakeShared<FNiagaraScriptViewModel>(FText(), ENiagaraParameterEditMode::EditAll, bForDataProcessingOnly);
+		ScriptViewModel->SetScript(OuterScript);
+		FNiagaraParameterDefinitionsUtilities::TrySubscribeScriptVarToDefinitionByName(TargetScriptVar, ScriptViewModel.Get());
+	};
+
 	if (Parameter.GetName() == NewName)
 	{
 		return true;
 	}
 
-	UNiagaraScript* OuterScript = GetTypedOuter<UNiagaraScript>();
-	TSharedPtr<FNiagaraScriptViewModel> ScriptViewModel = TNiagaraViewModelManager<UNiagaraScript, FNiagaraScriptViewModel>::GetExistingViewModelForObject(OuterScript);
-	if (ScriptViewModel.IsValid() == false)
-	{
-		ensureMsgf(false, TEXT("Failed to get script viewmodel when renaming parameter via pin in graph!"));
-		return false;
-	}
-
 	Modify();
 	if (FNiagaraGraphParameterReferenceCollection* ReferenceCollection = ParameterToReferencesMap.Find(Parameter))
 	{
@@ -2101,7 +2101,7 @@ bool UNiagaraGraph::RenameParameterFromPin(const FNiagaraVariable& Parameter, FN
 			{
 				if (TObjectPtr<UNiagaraScriptVariable> const* RenamedScriptVarPtr = VariableToScriptVariable.Find(FNiagaraVariable(Parameter.GetType(), NewName)))
 				{
-					FNiagaraParameterDefinitionsUtilities::TrySubscribeScriptVarToDefinitionByName(*RenamedScriptVarPtr, ScriptViewModel.Get());
+					TrySubscribeScriptVarToDefinitionByName(*RenamedScriptVarPtr);
 				}
 
 				// Rename all the bindings that point to the old parameter 
@@ -2153,7 +2153,7 @@ bool UNiagaraGraph::RenameParameterFromPin(const FNiagaraVariable& Parameter, FN
 		NewScriptVariable->Metadata.CreateNewGuid();
 		NewScriptVariable->SetIsSubscribedToParameterDefinitions(false);
 		VariableToScriptVariable.Add(NewParameter, NewScriptVariable);
-		FNiagaraParameterDefinitionsUtilities::TrySubscribeScriptVarToDefinitionByName(NewScriptVariable, ScriptViewModel.Get());
+		TrySubscribeScriptVarToDefinitionByName(NewScriptVariable);
 
 		const FNiagaraGraphParameterReferenceCollection* ReferenceCollection = GetParameterReferenceMap().Find(Parameter);
 		if (ReferenceCollection && ReferenceCollection->ParameterReferences.Num() < 1)
diff --git a/Engine/Plugins/FX/Niagara/Source/NiagaraEditor/Private/ViewModels/NiagaraScriptViewModel.cpp b/Engine/Plugins/FX/Niagara/Source/NiagaraEditor/Private/ViewModels/NiagaraScriptViewModel.cpp
--- a/Engine/Plugins/FX/Niagara/Source/NiagaraEditor/Private/ViewModels/NiagaraScriptViewModel.cpp
+++ b/Engine/Plugins/FX/Niagara/Source/NiagaraEditor/Private/ViewModels/NiagaraScriptViewModel.cpp
@@ -124,6 +124,13 @@ FNiagaraScriptViewModel::~FNiagaraScriptViewModel()
 
 }
 
+INiagaraParameterDefinitionsSubscriber* FNiagaraScriptViewModel::GetParameterDefinitionsSubscriber()
+{
+	checkf(bIsForDataProcessingOnly, TEXT("Tried to get parameter definitions subscriber for scriptviewmodel not being used strictly for data processing! When not for data processing, this must be a standalone or scratch scriptviewmodel!"));
+	checkf(Scripts.Num() == 1, TEXT("Tried to get parameter definitions subscriber for scriptviewmodel but there was more than one script!"));
+	return &Scripts[0];
+}
+
 FText FNiagaraScriptViewModel::GetDisplayName() const
 {
 	return GraphViewModel->GetDisplayName();
diff --git a/Engine/Plugins/FX/Niagara/Source/NiagaraEditor/Public/ViewModels/NiagaraScriptViewModel.h b/Engine/Plugins/FX/Niagara/Source/NiagaraEditor/Public/ViewModels/NiagaraScriptViewModel.h
--- a/Engine/Plugins/FX/Niagara/Source/NiagaraEditor/Public/ViewModels/NiagaraScriptViewModel.h
+++ b/Engine/Plugins/FX/Niagara/Source/NiagaraEditor/Public/ViewModels/NiagaraScriptViewModel.h
@@ -35,7 +35,7 @@ public:
 
 	//~ Begin NiagaraParameterDefinitionsSubscriberViewModel Interface
 protected:
-	virtual INiagaraParameterDefinitionsSubscriber* GetParameterDefinitionsSubscriber() override { checkf(false, TEXT("Tried to modify parameter libraries for script viewmodel; this is only a valid operation for standalone script viewmodel and scratchpad scriptviewmodel!")); return nullptr; };
+	virtual INiagaraParameterDefinitionsSubscriber* GetParameterDefinitionsSubscriber() override;
 	//~ End NiagaraParameterDefinitionsSubscriberViewModel Interface
 
 public:
