- Texture2DArray DI was using incorrect output type for TextureDimensions

#rnx
#9401

diff --git a/Engine/Plugins/FX/Niagara/Shaders/Private/NiagaraDataInterfaceTexture2DArrayTemplate.ush b/Engine/Plugins/FX/Niagara/Shaders/Private/NiagaraDataInterfaceTexture2DArrayTemplate.ush
--- a/Engine/Plugins/FX/Niagara/Shaders/Private/NiagaraDataInterfaceTexture2DArrayTemplate.ush
+++ b/Engine/Plugins/FX/Niagara/Shaders/Private/NiagaraDataInterfaceTexture2DArrayTemplate.ush
@@ -9,7 +9,7 @@ void SampleTexture_{ParameterName}(in float3 UVW, in float MipLevel, out float4
 	OutValue = {ParameterName}_Texture.SampleLevel({ParameterName}_TextureSampler, UVW, MipLevel);
 }
 
-void TextureDimensions_{ParameterName}(out float2 OutValue)
+void TextureDimensions_{ParameterName}(out float3 OutValue)
 {
 	OutValue = {ParameterName}_TextureSize;
 }
