- Fix typo on UVW by using input / output template param handlers

#rnx
#9403

diff --git a/Engine/Plugins/FX/Niagara/Source/Niagara/Private/NiagaraDataInterfaceVolumeTexture.cpp b/Engine/Plugins/FX/Niagara/Source/Niagara/Private/NiagaraDataInterfaceVolumeTexture.cpp
--- a/Engine/Plugins/FX/Niagara/Source/Niagara/Private/NiagaraDataInterfaceVolumeTexture.cpp
+++ b/Engine/Plugins/FX/Niagara/Source/Niagara/Private/NiagaraDataInterfaceVolumeTexture.cpp
@@ -232,27 +232,17 @@ void UNiagaraDataInterfaceVolumeTexture::GetTextureDimensions(FVectorVMExternalF
 void UNiagaraDataInterfaceVolumeTexture::SampleVolumeTexture(FVectorVMExternalFunctionContext& Context)
 {
 	VectorVM::FUserPtrHandler<FNDIVolumeTextureInstanceData_GameThread> InstData(Context);
-	VectorVM::FExternalFuncInputHandler<float> XParam(Context);
-	VectorVM::FExternalFuncInputHandler<float> YParam(Context);
-	VectorVM::FExternalFuncInputHandler<float> ZParam(Context);
-	VectorVM::FExternalFuncInputHandler<float> MipLevelParam(Context);
-	VectorVM::FExternalFuncRegisterHandler<float> OutSampleR(Context);
-	VectorVM::FExternalFuncRegisterHandler<float> OutSampleG(Context);
-	VectorVM::FExternalFuncRegisterHandler<float> OutSampleB(Context);
-	VectorVM::FExternalFuncRegisterHandler<float> OutSampleA(Context);
+	FNDIInputParam<FVector3f> UVWParam(Context);
+	FNDIInputParam<float> MipLevelParam(Context);
+	FNDIOutputParam<FVector4f> OutSample(Context);
 
+	const FVector4f DefaultSample(1.0f, 0.0f, 1.0f, 1.0f);
 	for (int32 i = 0; i < Context.GetNumInstances(); ++i)
 	{
-		float X = XParam.GetAndAdvance();
-		float Y = YParam.GetAndAdvance();
-		float Z = YParam.GetAndAdvance();
-		float Mip = MipLevelParam.GetAndAdvance();
-		*OutSampleR.GetDestAndAdvance() = 1.0;
-		*OutSampleG.GetDestAndAdvance() = 0.0;
-		*OutSampleB.GetDestAndAdvance() = 1.0;
-		*OutSampleA.GetDestAndAdvance() = 1.0;
+		const FVector3f UVW = UVWParam.GetAndAdvance();
+		const float Mip = MipLevelParam.GetAndAdvance();
+		OutSample.SetAndAdvance(DefaultSample);
 	}
-
 }
 
 #if WITH_EDITORONLY_DATA
