- Make GPU sampling tolerant to missing sampling info
- This can occur with runtime merging of skeletal meshes

#rnx

diff --git a/Engine/Plugins/FX/Niagara/Source/Niagara/Classes/NiagaraDataInterfaceSkeletalMesh.h b/Engine/Plugins/FX/Niagara/Source/Niagara/Classes/NiagaraDataInterfaceSkeletalMesh.h
--- a/Engine/Plugins/FX/Niagara/Source/Niagara/Classes/NiagaraDataInterfaceSkeletalMesh.h
+++ b/Engine/Plugins/FX/Niagara/Source/Niagara/Classes/NiagaraDataInterfaceSkeletalMesh.h
@@ -435,7 +435,7 @@ public:
 
 	virtual ~FSkeletalMeshGpuSpawnStaticBuffers() override;
 
-	FORCEINLINE_DEBUGGABLE void Initialise(struct FNDISkeletalMesh_InstanceData* InstData, const FSkeletalMeshLODRenderData& SkeletalMeshLODRenderData,const FSkeletalMeshSamplingLODBuiltData& SkeletalMeshSamplingLODBuiltData, FNiagaraSystemInstance* SystemInstance);
+	FORCEINLINE_DEBUGGABLE void Initialise(struct FNDISkeletalMesh_InstanceData* InstData, const FSkeletalMeshLODRenderData& SkeletalMeshLODRenderData,const FSkeletalMeshSamplingLODBuiltData* SkeletalMeshSamplingLODBuiltData, FNiagaraSystemInstance* SystemInstance);
 
 	virtual void InitRHI() override;
 	virtual void ReleaseRHI() override;
diff --git a/Engine/Plugins/FX/Niagara/Source/Niagara/Private/NiagaraDataInterfaceSkeletalMesh.cpp b/Engine/Plugins/FX/Niagara/Source/Niagara/Private/NiagaraDataInterfaceSkeletalMesh.cpp
--- a/Engine/Plugins/FX/Niagara/Source/Niagara/Private/NiagaraDataInterfaceSkeletalMesh.cpp
+++ b/Engine/Plugins/FX/Niagara/Source/Niagara/Private/NiagaraDataInterfaceSkeletalMesh.cpp
@@ -768,7 +768,7 @@ FSkeletalMeshGpuSpawnStaticBuffers::~FSkeletalMeshGpuSpawnStaticBuffers()
 	//ValidSections.Empty();
 }
 
-void FSkeletalMeshGpuSpawnStaticBuffers::Initialise(FNDISkeletalMesh_InstanceData* InstData, const FSkeletalMeshLODRenderData& SkeletalMeshLODRenderData, const FSkeletalMeshSamplingLODBuiltData& MeshSamplingLODBuiltData, FNiagaraSystemInstance* SystemInstance)
+void FSkeletalMeshGpuSpawnStaticBuffers::Initialise(FNDISkeletalMesh_InstanceData* InstData, const FSkeletalMeshLODRenderData& SkeletalMeshLODRenderData, const FSkeletalMeshSamplingLODBuiltData* MeshSamplingLODBuiltData, FNiagaraSystemInstance* SystemInstance)
 {
 	SkeletalMeshSamplingLODBuiltData = nullptr;
 	bUseGpuUniformlyDistributedSampling = false;
@@ -785,7 +785,7 @@ void FSkeletalMeshGpuSpawnStaticBuffers::Initialise(FNDISkeletalMesh_InstanceDat
 
 	if (InstData)
 	{
-		SkeletalMeshSamplingLODBuiltData = &MeshSamplingLODBuiltData;
+		SkeletalMeshSamplingLODBuiltData = MeshSamplingLODBuiltData;
 		bUseGpuUniformlyDistributedSampling = InstData->bIsGpuUniformlyDistributedSampling;
 
 		LODRenderData = &SkeletalMeshLODRenderData;
@@ -804,6 +804,7 @@ void FSkeletalMeshGpuSpawnStaticBuffers::Initialise(FNDISkeletalMesh_InstanceDat
 
 		if (bUseGpuUniformlyDistributedSampling)
 		{
+			check(SkeletalMeshSamplingLODBuiltData != nullptr);
 			const int32 NumAreaSamples = SkeletalMeshSamplingLODBuiltData->AreaWeightedTriangleSampler.GetNumEntries();
 			if (NumAreaSamples != TriangleCount)
 			{
@@ -2074,8 +2075,11 @@ bool FNDISkeletalMesh_InstanceData::Init(UNiagaraDataInterfaceSkeletalMesh* Inte
 			}
 
 			const FSkeletalMeshSamplingInfo& SamplingInfo = Mesh->GetSamplingInfo();
+			const FSkeletalMeshSamplingLODBuiltData* MeshSamplingBuiltData = SamplingInfo.GetBuiltData().WholeMeshBuiltData.IsValidIndex(CachedLODIdx) ? &SamplingInfo.GetBuiltData().WholeMeshBuiltData[CachedLODIdx] : nullptr;
+			bIsGpuUniformlyDistributedSampling &= MeshSamplingBuiltData != nullptr;
+
 			MeshGpuSpawnStaticBuffers = new FSkeletalMeshGpuSpawnStaticBuffers();
-			MeshGpuSpawnStaticBuffers->Initialise(this, *CachedLODData, SamplingInfo.GetBuiltData().WholeMeshBuiltData[CachedLODIdx], SystemInstance);
+			MeshGpuSpawnStaticBuffers->Initialise(this, *CachedLODData, MeshSamplingBuiltData, SystemInstance);
 			BeginInitResource(MeshGpuSpawnStaticBuffers);
 
 			MeshGpuSpawnDynamicBuffers = new FSkeletalMeshGpuDynamicBufferProxy();
