- Fix for always including the system origin in bounds calculation

#rnx

diff --git a/Engine/Plugins/FX/Niagara/Source/Niagara/Private/NiagaraSystemInstance.cpp b/Engine/Plugins/FX/Niagara/Source/Niagara/Private/NiagaraSystemInstance.cpp
--- a/Engine/Plugins/FX/Niagara/Source/Niagara/Private/NiagaraSystemInstance.cpp
+++ b/Engine/Plugins/FX/Niagara/Source/Niagara/Private/NiagaraSystemInstance.cpp
@@ -2501,7 +2501,7 @@ void FNiagaraSystemInstance::Tick_Concurrent(bool bEnqueueGPUTickIfNeeded)
 			}
 		}
 
-		LocalBounds = FBox(FVector::ZeroVector, FVector::ZeroVector);
+		LocalBounds = FBox(ForceInit);
 		if (NewDynamicBounds.IsValid)
 		{
 			FVector Center = NewDynamicBounds.GetCenter();
@@ -2521,6 +2521,12 @@ void FNiagaraSystemInstance::Tick_Concurrent(bool bEnqueueGPUTickIfNeeded)
 			const FVector Extent = NewFixedBounds.GetExtent() * GNiagaraEmitterBoundsFixedExpandMultiplier;
 			LocalBounds += FBox(Center - Extent, Center + Extent);
 		}
+
+		// In the case that we have no bounds initialize to a sensible default value to avoid NaNs later on
+		if (LocalBounds.IsValid == false)
+		{
+			LocalBounds = FBox(FVector::ZeroVector, FVector::ZeroVector);
+		}
 	}
 
 	//Enqueue a GPU tick for this sim if we're allowed to do so from a concurrent thread.
