Static Mesh DI respects per component painted vertex colors.

TODO: Some callback from the Vertex Painting / Static Mesh to update Niagara when there are changes to the vertex overrides or they're removed etc.
TODO: Verify lifetime of OverrideVertexColors passed to GPU sims. Unclear atm if this is safe. It is safe for an RT command to use this as the destruction is done RT safe. However, holding onto the ref beyond a single RT command as Niagara does seems unsafe. Though the skeleltal mesh DI also does this so I may be missing some safety mechanism here.

diff --git a/Engine/Plugins/FX/Niagara/Source/Niagara/Private/DataInterface/NiagaraDataInterfaceStaticMesh.cpp b/Engine/Plugins/FX/Niagara/Source/Niagara/Private/DataInterface/NiagaraDataInterfaceStaticMesh.cpp
--- a/Engine/Plugins/FX/Niagara/Source/Niagara/Private/DataInterface/NiagaraDataInterfaceStaticMesh.cpp
+++ b/Engine/Plugins/FX/Niagara/Source/Niagara/Private/DataInterface/NiagaraDataInterfaceStaticMesh.cpp
@@ -310,6 +310,8 @@ namespace NDIStaticMeshLocal
 		FNiagaraSystemInstanceID			SystemInstanceID;
 		TRefCountPtr<const FStaticMeshLODResources>		LODResource = nullptr;
 
+		FColorVertexBuffer*					OverrideColorBuffer = nullptr;
+
 #if DO_CHECK
 		FName								SystemFName;
 		FName								StaticMeshFName;
@@ -393,7 +395,7 @@ namespace NDIStaticMeshLocal
 				MeshPositionBufferSRV = GpuInitializeData.LODResource->VertexBuffers.PositionVertexBuffer.GetSRV();
 				MeshTangentBufferSRV = GpuInitializeData.LODResource->VertexBuffers.StaticMeshVertexBuffer.GetTangentsSRV();
 				MeshUVBufferSRV = GpuInitializeData.LODResource->VertexBuffers.StaticMeshVertexBuffer.GetTexCoordsSRV();
-				MeshColorBufferSRV = GpuInitializeData.LODResource->VertexBuffers.ColorVertexBuffer.GetColorComponentsSRV();
+				MeshColorBufferSRV = GpuInitializeData.OverrideColorBuffer ? GpuInitializeData.OverrideColorBuffer->GetColorComponentsSRV() : GpuInitializeData.LODResource->VertexBuffers.ColorVertexBuffer.GetColorComponentsSRV();
 
 				NumTriangles.X = GpuInitializeData.LODResource->IndexBuffer.GetNumIndices() / 3;
 				NumTriangles.Y = GpuInitializeData.NumFilteredTriangles;
@@ -909,6 +911,31 @@ namespace NDIStaticMeshLocal
 			return nullptr;
 		}
 
+		const FStaticMeshLODResources* GetCurrentFirstLODWithVertexColorOverrides(FColorVertexBuffer*& OutVertexColorOverrides)
+		{
+			//-OPT: Perhaps we could cache this during the tick function?
+			UStaticMesh* StaticMesh = StaticMeshWeakPtr.Get();
+			OutVertexColorOverrides = nullptr;
+			if (bMeshValid && StaticMesh)
+			{
+				if(UStaticMeshComponent* SMComp = Cast<UStaticMeshComponent>(SceneComponentWeakPtr.Get()))
+				{
+					int32 LODIdx = StaticMesh->GetRenderData()->GetCurrentFirstLODIdx(MinLOD);
+					if(SMComp->LODData.IsValidIndex(LODIdx))
+					{
+						OutVertexColorOverrides = SMComp->LODData[LODIdx].OverrideVertexColors;
+					}
+				}
+
+				if (const FStaticMeshLODResources* LODResource = StaticMesh->GetRenderData()->GetCurrentFirstLOD(MinLOD))
+				{
+					return LODResource;
+				}
+
+			}
+			return nullptr;
+		}
+
 		TConstArrayView<int32> GetFilteredSections() const { return MakeArrayView(FilteredAndUnfilteredSections.GetData(), NumFilteredSections); }
 		TConstArrayView<int32> GetUnfilteredSections() const { return MakeArrayView(FilteredAndUnfilteredSections.GetData() + NumFilteredSections, NumUnfilteredSections); }
 
@@ -951,7 +978,7 @@ namespace NDIStaticMeshLocal
 		FORCEINLINE FStaticMeshCpuHelper(FVectorVMExternalFunctionContext& Context)
 			: InstanceData(Context)
 		{
-			LODResource = InstanceData->GetCurrentFirstLOD();
+			LODResource = InstanceData->GetCurrentFirstLODWithVertexColorOverrides(OverrideVertexColors);
 		}
 
 
@@ -1023,6 +1050,10 @@ namespace NDIStaticMeshLocal
 
 		FORCEINLINE int32 GetNumColorVertices() const
 		{
+			if(OverrideVertexColors)
+			{
+				return OverrideVertexColors->GetAllowCPUAccess() ? OverrideVertexColors->GetNumVertices() : 0; 
+			}
 			return LODResource && LODResource->VertexBuffers.ColorVertexBuffer.GetVertexData() && LODResource->VertexBuffers.ColorVertexBuffer.GetAllowCPUAccess() ? LODResource->VertexBuffers.ColorVertexBuffer.GetNumVertices() : 0;
 		}
 
@@ -1075,9 +1106,19 @@ namespace NDIStaticMeshLocal
 		FORCEINLINE FLinearColor GetTriangleColor(const FVector3f& BaryCoord, int32 Index0, int32 Index1, int32 Index2) const
 		{
 			FLinearColor Color;
-			Color  = FLinearColor(LODResource->VertexBuffers.ColorVertexBuffer.VertexColor(Index0)) * BaryCoord.X;
-			Color += FLinearColor(LODResource->VertexBuffers.ColorVertexBuffer.VertexColor(Index1)) * BaryCoord.Y;
-			Color += FLinearColor(LODResource->VertexBuffers.ColorVertexBuffer.VertexColor(Index2)) * BaryCoord.Z;
+
+			if (OverrideVertexColors)
+			{
+				Color = FLinearColor(OverrideVertexColors->VertexColor(Index0)) * BaryCoord.X;
+				Color += FLinearColor(OverrideVertexColors->VertexColor(Index1)) * BaryCoord.Y;
+				Color += FLinearColor(OverrideVertexColors->VertexColor(Index2)) * BaryCoord.Z;
+			}
+			else
+			{
+				Color  = FLinearColor(LODResource->VertexBuffers.ColorVertexBuffer.VertexColor(Index0)) * BaryCoord.X;
+				Color += FLinearColor(LODResource->VertexBuffers.ColorVertexBuffer.VertexColor(Index1)) * BaryCoord.Y;
+				Color += FLinearColor(LODResource->VertexBuffers.ColorVertexBuffer.VertexColor(Index2)) * BaryCoord.Z;			
+			}
 			return Color;
 		}
 
@@ -1239,6 +1280,7 @@ namespace NDIStaticMeshLocal
 
 		VectorVM::FUserPtrHandler<FInstanceData_GameThread> InstanceData;
 		TRefCountPtr<const FStaticMeshLODResources> LODResource;
+		FColorVertexBuffer* OverrideVertexColors;
 		TTransformHandler TransformHandler;
 	};
 }
@@ -1345,7 +1387,11 @@ bool UNiagaraDataInterfaceStaticMesh::InitPerInstanceData(void* PerInstanceData,
 		TUniquePtr<NDIStaticMeshLocal::FGpuInitializeData> GpuInitializeData = MakeUnique<NDIStaticMeshLocal::FGpuInitializeData>();
 		GpuInitializeData->RenderProxy = GetProxyAs<NDIStaticMeshLocal::FRenderProxy>();
 		GpuInitializeData->SystemInstanceID = SystemInstance->GetId();
-		GpuInitializeData->LODResource = InstanceData->GetCurrentFirstLOD();
+
+		//This is safe to ref on the RT as it's freed on the RT in FStaticMeshComponentLODInfo::BeginReleaseOverrideVertexColors()
+		//However, it's seems unsafe to reference in Niagara's instance data this way as there looks to be a window between this RT command and it's actual use where the data could have been freed.
+		GpuInitializeData->LODResource = InstanceData->GetCurrentFirstLODWithVertexColorOverrides(GpuInitializeData->OverrideColorBuffer);
+
 		if ( GpuInitializeData->LODResource )
 		{
 		#if DO_CHECK
