- Initial pass for Data Interface SimCache storage
- Currently it's all or nothing we will want different modes in the future so the API is expected to evolve

#rnx

diff --git a/Engine/Plugins/FX/Niagara/Source/Niagara/Classes/NiagaraDataInterface.h b/Engine/Plugins/FX/Niagara/Source/Niagara/Classes/NiagaraDataInterface.h
--- a/Engine/Plugins/FX/Niagara/Source/Niagara/Classes/NiagaraDataInterface.h
+++ b/Engine/Plugins/FX/Niagara/Source/Niagara/Classes/NiagaraDataInterface.h
@@ -499,7 +499,18 @@ public:
 	virtual bool PerInstanceTick(void* PerInstanceData, FNiagaraSystemInstance* SystemInstance, float DeltaSeconds) { return false; }
 	virtual bool PerInstanceTickPostSimulate(void* PerInstanceData, FNiagaraSystemInstance* SystemInstance, float DeltaSeconds) { return false; }
 	
-	/** Update internal datas after reading the cache. */
+	/** Begin writing data for a simulation cache, returning a nullptr means the data interface does not store data into the simulation cache. */
+	virtual UObject* SimCacheBeginWrite(UObject* SimCache, FNiagaraSystemInstance* NiagaraSystemInstance, const void* OptionalPerInstanceData) { return nullptr; }
+	/** Write a new frame of data for the simulation cache.  This is always in sequence, i.e. 0, 1, 2, etc, we will never jump around frames. */
+	virtual bool SimCacheWriteFrame(UObject* StorageObject, int FrameIndex, FNiagaraSystemInstance* SystemInstance, const void* OptionalPerInstanceData) { return true; }
+	/** End writing data for a simulation cache.  Note this is called on the CDO not the instance the object was created from. */
+	virtual bool SimCacheEndWrite(UObject* StorageObject) { return true; }
+	/** Read a frame of data from the simulation cache. */
+	virtual bool SimCacheReadFrame(UObject* StorageObject, int FrameA, int FrameB, float Interp, FNiagaraSystemInstance* SystemInstance, void* OptionalPerInstanceData) { return true; }
+	/**
+	Called when the simulation cache has finished reading a frame.
+	Only DataInterfaces with PerInstanceData are currently supported.
+	*/
 	virtual void SimCachePostReadFrame(void* OptionalPerInstanceData, FNiagaraSystemInstance* SystemInstance) {}
 
 #if WITH_EDITORONLY_DATA
diff --git a/Engine/Plugins/FX/Niagara/Source/Niagara/Classes/NiagaraSimCache.h b/Engine/Plugins/FX/Niagara/Source/Niagara/Classes/NiagaraSimCache.h
--- a/Engine/Plugins/FX/Niagara/Source/Niagara/Classes/NiagaraSimCache.h
+++ b/Engine/Plugins/FX/Niagara/Source/Niagara/Classes/NiagaraSimCache.h
@@ -70,8 +70,6 @@ struct FNiagaraSimCacheFrame
 
 	UPROPERTY()
 	TArray<FNiagaraSimCacheEmitterFrame> EmitterData;
-
-	//-TODO: DataInterface cache of data
 };
 
 USTRUCT()
@@ -171,6 +169,9 @@ class NIAGARA_API UNiagaraSimCache : public UObject
 	UPROPERTY()
 	TArray<FNiagaraSimCacheFrame> CacheFrames;
 
+	UPROPERTY()
+	TMap<FNiagaraVariableBase, TObjectPtr<UObject>> DataInterfaceStorage;
+
 	// UObject Interface
 	virtual bool IsReadyForFinishDestroy() override;
 	// UObject Interface
@@ -183,7 +184,7 @@ class NIAGARA_API UNiagaraSimCache : public UObject
 
 	bool CanRead(UNiagaraSystem* NiagaraSystem);
 	bool Read(float TimeSeconds, FNiagaraSystemInstance* SystemInstance) const;
-	bool ReadFrame(int32 FrameIndex, FNiagaraSystemInstance* SystemInstance) const;
+	bool ReadFrame(int32 FrameIndex, float FrameFraction, FNiagaraSystemInstance* SystemInstance) const;
 
 	/*
 	Create a simulation cache from the component's current data.
diff --git a/Engine/Plugins/FX/Niagara/Source/Niagara/Private/DataInterface/NiagaraDataInterfaceUtilities.cpp b/Engine/Plugins/FX/Niagara/Source/Niagara/Private/DataInterface/NiagaraDataInterfaceUtilities.cpp
--- a/Engine/Plugins/FX/Niagara/Source/Niagara/Private/DataInterface/NiagaraDataInterfaceUtilities.cpp
+++ b/Engine/Plugins/FX/Niagara/Source/Niagara/Private/DataInterface/NiagaraDataInterfaceUtilities.cpp
@@ -145,8 +145,25 @@ void ForEachGpuFunctionEqualsImpl(class UNiagaraDataInterface* DataInterface, cl
 	);
 }
 
+void ForEachDataInterface(const FNiagaraParameterStore& ParameterStore, TFunction<bool(const FNiagaraVariableBase Variable, UNiagaraDataInterface* DataInterface)> Action)
+{
+	const TArray<UNiagaraDataInterface*>& DataInterfaces = ParameterStore.GetDataInterfaces();
+	for (const FNiagaraVariableWithOffset& Variable : ParameterStore.ReadParameterVariables())
+	{
+		if (Variable.IsDataInterface() == false)
+		{
+			continue;
+		}
+
+		if ( UNiagaraDataInterface* DataInterface = DataInterfaces[Variable.Offset] )
+		{
+			Action(Variable, DataInterface);
+		}
+	}
 }
 
+} //namespace FNiagaraDataInterfaceUtilities
+
 void FNiagaraDataInterfaceUtilities::ForEachVMFunctionEquals(UNiagaraDataInterface* DataInterface, UNiagaraSystem* NiagaraSystem, TFunction<bool(const FVMExternalFunctionBindingInfo&)> Action)
 {
 	ForEachVMFunctionEqualsImpl(DataInterface, NiagaraSystem, NiagaraSystem->GetExposedParameters(), Action);
@@ -305,3 +322,43 @@ void FNiagaraDataInterfaceUtilities::ForEachGpuFunction(class UNiagaraDataInterf
 		}
 	}
 }
+
+void FNiagaraDataInterfaceUtilities::ForEachDataInterface(FNiagaraSystemInstance* SystemInstance, TFunction<bool(const FNiagaraVariableBase Variable, UNiagaraDataInterface* DataInterface)> Action)
+{
+	if ( SystemInstance == nullptr )
+	{
+		return;
+	}
+
+	FNiagaraSystemSimulationPtr SystemSimulation = SystemInstance->GetSystemSimulation();
+	if (SystemSimulation == nullptr || !SystemSimulation->IsValid())
+	{
+		return;
+	}
+
+	ForEachDataInterface(*SystemInstance->GetOverrideParameters(), Action);
+
+	ForEachDataInterface(SystemSimulation->GetSpawnExecutionContext()->Parameters, Action);
+	ForEachDataInterface(SystemSimulation->GetUpdateExecutionContext()->Parameters, Action);
+
+	for (const TSharedRef<FNiagaraEmitterInstance, ESPMode::ThreadSafe>& Emitter : SystemInstance->GetEmitters())
+	{
+		if ( Emitter->IsDisabled() )
+		{
+			continue;
+		}
+
+		ForEachDataInterface(Emitter->GetSpawnExecutionContext().Parameters, Action);
+		ForEachDataInterface(Emitter->GetUpdateExecutionContext().Parameters, Action);
+		for (int32 i=0; i < Emitter->GetEventExecutionContexts().Num(); i++)
+		{
+			ForEachDataInterface(Emitter->GetEventExecutionContexts()[i].Parameters, Action);
+		}
+
+		FVersionedNiagaraEmitterData* EmitterData = Emitter->GetCachedEmitterData();
+		if (EmitterData && EmitterData->SimTarget == ENiagaraSimTarget::GPUComputeSim && Emitter->GetGPUContext())
+		{
+			ForEachDataInterface(Emitter->GetGPUContext()->CombinedParamStore, Action);
+		}
+	}
+}
diff --git a/Engine/Plugins/FX/Niagara/Source/Niagara/Private/NiagaraSimCache.cpp b/Engine/Plugins/FX/Niagara/Source/Niagara/Private/NiagaraSimCache.cpp
--- a/Engine/Plugins/FX/Niagara/Source/Niagara/Private/NiagaraSimCache.cpp
+++ b/Engine/Plugins/FX/Niagara/Source/Niagara/Private/NiagaraSimCache.cpp
@@ -5,6 +5,7 @@
 #include "NiagaraClearCounts.h"
 #include "NiagaraComponent.h"
 #include "NiagaraComputeExecutionContext.h"
+#include "NiagaraDataInterfaceUtilities.h"
 #include "NiagaraDataSetReadback.h"
 #include "NiagaraEmitterInstance.h"
 #include "NiagaraGPUInstanceCountManager.h"
@@ -460,6 +461,20 @@ void UNiagaraSimCache::BeginWrite(UNiagaraComponent* NiagaraComponent)
 		const FNiagaraEmitterCompiledData& EmitterCompiledData = Helper.NiagaraSystem->GetEmitterCompiledData()[i].Get();
 		Helper.BuildCacheLayout(CacheLayout.EmitterLayouts[i], EmitterCompiledData.DataSetCompiledData, Helper.NiagaraSystem->GetEmitterHandle(i).GetName());
 	}
+
+	// Find data interfaces we may want to cache
+	FNiagaraDataInterfaceUtilities::ForEachDataInterface(
+		Helper.SystemInstance,
+		[&](const FNiagaraVariableBase& Variable, UNiagaraDataInterface* DataInterface)
+		{
+			const void* PerInstanceData = Helper.SystemInstance->FindDataInterfaceInstanceData(DataInterface);
+			if (UObject* DICacheStorage = DataInterface->SimCacheBeginWrite(this, Helper.SystemInstance, PerInstanceData))
+			{
+				DataInterfaceStorage.FindOrAdd(Variable) = DICacheStorage;
+			}
+			return true;
+		}
+	);
 }
 
 void UNiagaraSimCache::WriteFrame(UNiagaraComponent* NiagaraComponent)
@@ -534,6 +549,33 @@ void UNiagaraSimCache::WriteFrame(UNiagaraComponent* NiagaraComponent)
 			Helper.WriteDataBuffer(*EmitterCurrentData, CacheLayout.EmitterLayouts[i], CacheEmitterFrame.ParticleDataBuffers, 0, EmitterCurrentData->GetNumInstances());
 		}
 	}
+
+	// Store data interface data
+	//-OPT: We shouldn't need to search all the time here
+	if (DataInterfaceStorage.IsEmpty() == false)
+	{
+		const int FrameIndex = CacheFrames.Num() - 1;
+		bool bDataInterfacesSucess = true;
+
+		FNiagaraDataInterfaceUtilities::ForEachDataInterface(
+			Helper.SystemInstance,
+			[&](const FNiagaraVariableBase& Variable, UNiagaraDataInterface* DataInterface)
+			{
+				if ( UObject* StorageObject = DataInterfaceStorage.FindRef(Variable) )
+				{
+					const void* PerInstanceData = Helper.SystemInstance->FindDataInterfaceInstanceData(DataInterface);
+					bDataInterfacesSucess &= DataInterface->SimCacheWriteFrame(StorageObject, FrameIndex, Helper.SystemInstance, PerInstanceData);
+				}
+				return true;
+			}
+		);
+
+		// A data interface failed to write information
+		if (bDataInterfacesSucess == false)
+		{
+			SoftNiagaraSystem.Reset();
+		}
+	}
 }
 
 void UNiagaraSimCache::EndWrite()
@@ -543,6 +585,23 @@ void UNiagaraSimCache::EndWrite()
 	{
 		SoftNiagaraSystem.Reset();
 	}
+
+	if (DataInterfaceStorage.IsEmpty() == false)
+	{
+		bool bDataInterfacesSucess = true;
+		for ( auto it=DataInterfaceStorage.CreateIterator(); it; ++it )
+		{
+			UClass* DataInterfaceClass = it.Key().GetType().GetClass();
+			check(DataInterfaceClass != nullptr);
+			UNiagaraDataInterface* DataInterface = CastChecked<UNiagaraDataInterface>(DataInterfaceClass->GetDefaultObject());
+			bDataInterfacesSucess &= DataInterface->SimCacheEndWrite(it.Value());
+		}
+
+		if (bDataInterfacesSucess == false)
+		{
+			SoftNiagaraSystem.Reset();
+		}
+	}
 }
 
 bool UNiagaraSimCache::CanRead(UNiagaraSystem* NiagaraSystem)
@@ -630,14 +689,13 @@ bool UNiagaraSimCache::Read(float TimeSeconds, FNiagaraSystemInstance* SystemIns
 	}
 
 	const float FrameTime		= (RelativeTime / DurationSeconds) * float(CacheFrames.Num() - 1);
-	const float FrameIndexA		= FMath::Floor(FrameTime);
-	const float FrameIndexB		= FMath::Min(FrameIndexA + 1 , CacheFrames.Num() - 1);
-	const float FrameFraction	= FrameTime - float(FrameIndexA);
+	const float FrameIndex		= FMath::Floor(FrameTime);
+	const float FrameFraction	= FrameTime - float(FrameIndex);
 
-	return ReadFrame(FrameIndexA, SystemInstance);
+	return ReadFrame(FrameIndex, FrameFraction, SystemInstance);
 }
 
-bool UNiagaraSimCache::ReadFrame(int32 FrameIndex, FNiagaraSystemInstance* SystemInstance) const
+bool UNiagaraSimCache::ReadFrame(int32 FrameIndex, float FrameFraction, FNiagaraSystemInstance* SystemInstance) const
 {
 	FNiagaraSimCacheHelper Helper(SystemInstance);
 	if ( !Helper.HasValidSimulation() )
@@ -665,7 +723,36 @@ bool UNiagaraSimCache::ReadFrame(int32 FrameIndex, FNiagaraSystemInstance* Syste
 			Helper.ReadDataBuffer(CacheLayout.EmitterLayouts[i], CacheEmitterFrame.ParticleDataBuffers, EmitterInstance.GetData());
 		}
 	}
-	for(TPair<TWeakObjectPtr<UNiagaraDataInterface>, int32>& DataInterfacePair : SystemInstance->DataInterfaceInstanceDataOffsets)
+
+	// Store data interface data
+	//-OPT: We shouldn't need to search all the time here
+	if (DataInterfaceStorage.IsEmpty() == false)
+	{
+		const int NextFrameIndex = FMath::Min(FrameIndex + 1, CacheFrames.Num() - 1);
+		bool bDataInterfacesSucess = true;
+
+		FNiagaraDataInterfaceUtilities::ForEachDataInterface(
+			Helper.SystemInstance,
+			[&](const FNiagaraVariableBase& Variable, UNiagaraDataInterface* DataInterface)
+			{
+				if (UObject* StorageObject = DataInterfaceStorage.FindRef(Variable))
+				{
+					void* PerInstanceData = Helper.SystemInstance->FindDataInterfaceInstanceData(DataInterface);
+					bDataInterfacesSucess &= DataInterface->SimCacheReadFrame(StorageObject, FrameIndex, NextFrameIndex, FrameFraction, Helper.SystemInstance, PerInstanceData);
+				}
+				return true;
+			}
+		);
+
+
+		if (bDataInterfacesSucess == false)
+		{
+			return false;
+		}
+	}
+
+	//-TODO: This should loop over all DataInterfaces that register not just ones with instance data
+	for (TPair<TWeakObjectPtr<UNiagaraDataInterface>, int32>& DataInterfacePair : SystemInstance->DataInterfaceInstanceDataOffsets)
 	{
 		if (UNiagaraDataInterface* Interface = DataInterfacePair.Key.Get())
 		{
diff --git a/Engine/Plugins/FX/Niagara/Source/Niagara/Public/NiagaraDataInterfaceUtilities.h b/Engine/Plugins/FX/Niagara/Source/Niagara/Public/NiagaraDataInterfaceUtilities.h
--- a/Engine/Plugins/FX/Niagara/Source/Niagara/Public/NiagaraDataInterfaceUtilities.h
+++ b/Engine/Plugins/FX/Niagara/Source/Niagara/Public/NiagaraDataInterfaceUtilities.h
@@ -33,4 +33,8 @@ namespace FNiagaraDataInterfaceUtilities
 	// The same function call made be made multiple times across different scripts so you may see the same function multiple times
 	// The action should return True to continue iteration or False to stop
 	NIAGARA_API void ForEachGpuFunction(UNiagaraDataInterface* DataInterface, FNiagaraSystemInstance* SystemInstance, TFunction<bool(const FNiagaraDataInterfaceGeneratedFunction&)> Action);
+
+	// Loops over all data interfaces inside the SystemInstance
+	// The action should return True to continue iteration or False to stop
+	void ForEachDataInterface(FNiagaraSystemInstance* SystemInstance, TFunction<bool(const FNiagaraVariableBase Variable, UNiagaraDataInterface* DataInterface)> Action);
 }
