- Remove cvar to hide sim cache baking

#rnx

diff --git a/Engine/Plugins/FX/Niagara/Source/NiagaraEditor/Private/NiagaraBakerOutputRegistry.cpp b/Engine/Plugins/FX/Niagara/Source/NiagaraEditor/Private/NiagaraBakerOutputRegistry.cpp
--- a/Engine/Plugins/FX/Niagara/Source/NiagaraEditor/Private/NiagaraBakerOutputRegistry.cpp
+++ b/Engine/Plugins/FX/Niagara/Source/NiagaraEditor/Private/NiagaraBakerOutputRegistry.cpp
@@ -10,13 +10,6 @@
 #include "NiagaraBakerRendererOutputVolumeTexture.h"
 #include "Customizations/NiagaraBakerOutputCustomization.h"
 
-static TAutoConsoleVariable<bool> CVarNiagaraBakerSimCacheEnabled(
-	TEXT("fx.Niagara.Baker.SimCache.Enabled"),
-	false,
-	TEXT("SimCache is experimental and will undergo high fluctuation in the data, therefore it is meant to be used for exerpimentation only."),
-	ECVF_Default
-);
-
 FNiagaraBakerOutputRegistry& FNiagaraBakerOutputRegistry::Get()
 {
 	static TUniquePtr<FNiagaraBakerOutputRegistry> Singleton;
@@ -28,7 +21,6 @@ FNiagaraBakerOutputRegistry& FNiagaraBakerOutputRegistry::Get()
 		{
 			FRegistryEntry& Entry = Singleton->Registry.AddDefaulted_GetRef();
 			Entry.WeakClass				= UNiagaraBakerOutputSimCache::StaticClass();
-			Entry.ShouldShowCVar		= IConsoleManager::Get().FindConsoleVariable(TEXT("fx.Niagara.Baker.SimCache.Enabled"));
 			Entry.CreateRenderer		= []() -> FNiagaraBakerOutputRenderer* { return new FNiagaraBakerRendererOutputSimCache(); };
 			Entry.CreateCustomization	= FOnGetDetailCustomizationInstance::CreateStatic(&FNiagaraBakerOutputSimCacheDetails::MakeInstance);
 		}
