- Fix ribbon crash for incorrectly using Twist when we have no Size data

#rnx

diff --git a/Engine/Plugins/FX/Niagara/Source/Niagara/Public/NiagaraRendererRibbons.h b/Engine/Plugins/FX/Niagara/Source/Niagara/Public/NiagaraRendererRibbons.h
--- a/Engine/Plugins/FX/Niagara/Source/Niagara/Public/NiagaraRendererRibbons.h
+++ b/Engine/Plugins/FX/Niagara/Source/Niagara/Public/NiagaraRendererRibbons.h
@@ -86,7 +86,7 @@ public:
 		, bHasFullRibbonIDs(Properties->RibbonFullIDDataSetAccessor.IsValid())
 		, bHasSimpleRibbonIDs(Properties->RibbonIdDataSetAccessor.IsValid())
 		, bHasCustomLinkOrder(Properties->RibbonLinkOrderDataSetAccessor.IsValid())
-		, bHasTwist(Properties->TwistDataSetAccessor.IsValid())
+		, bHasTwist(Properties->TwistDataSetAccessor.IsValid() && Properties->SizeDataSetAccessor.IsValid())
 		, bHasCustomU0Data(Properties->UV0Settings.bEnablePerParticleUOverride && Properties->U0OverrideIsBound)
 		, bHasCustomU1Data(Properties->UV1Settings.bEnablePerParticleUOverride && Properties->U1OverrideIsBound)
 		, bWantsConstantTessellation(Properties->TessellationMode == ENiagaraRibbonTessellationMode::Custom && Properties->bUseConstantFactor)
