Fixing emitters not having proper memory of active parameter panel selection

diff --git a/Engine/Plugins/FX/Niagara/Source/NiagaraEditor/Private/Toolkits/NiagaraSystemToolkit.cpp b/Engine/Plugins/FX/Niagara/Source/NiagaraEditor/Private/Toolkits/NiagaraSystemToolkit.cpp
--- a/Engine/Plugins/FX/Niagara/Source/NiagaraEditor/Private/Toolkits/NiagaraSystemToolkit.cpp
+++ b/Engine/Plugins/FX/Niagara/Source/NiagaraEditor/Private/Toolkits/NiagaraSystemToolkit.cpp
@@ -251,6 +251,9 @@ void FNiagaraSystemToolkit::InitializeWithEmitter(const EToolkitMode::Type Mode,
 	// id because the copy's change id may have been updated from the original as part of post load and we use this id to detect if the editable 
 	// emitter has been changed.
 	LastSyncedEmitterChangeId = SystemViewModel->GetEmitterHandleViewModels()[0]->GetEmitterViewModel()->GetEmitter().Emitter->GetChangeId();
+
+	// Mirror the system setup above so that the parameter panel updates appropriately
+	SystemViewModel->SetParameterPanelViewModel(ParameterPanelViewModel);
 	SystemToolkitMode = ESystemToolkitMode::Emitter;
 
 	if (GbLogNiagaraSystemChanges > 0)
diff --git a/Engine/Plugins/FX/Niagara/Source/NiagaraEditor/Private/ViewModels/NiagaraSystemEditorDocumentsViewModel.cpp b/Engine/Plugins/FX/Niagara/Source/NiagaraEditor/Private/ViewModels/NiagaraSystemEditorDocumentsViewModel.cpp
--- a/Engine/Plugins/FX/Niagara/Source/NiagaraEditor/Private/ViewModels/NiagaraSystemEditorDocumentsViewModel.cpp
+++ b/Engine/Plugins/FX/Niagara/Source/NiagaraEditor/Private/ViewModels/NiagaraSystemEditorDocumentsViewModel.cpp
@@ -346,7 +346,7 @@ void UNiagaraSystemEditorDocumentsViewModel::SetActiveDocumentTab(TSharedPtr<SDo
 bool UNiagaraSystemEditorDocumentsViewModel::IsPrimaryDocumentActive() const
 {
 	TSharedPtr<SDockTab> Tab = ActiveDocumentTab.Pin();
-	if (Tab.IsValid() && Tab->GetLayoutIdentifier().TabType == PrimaryDocumentTabId)
+	if (Tab.IsValid() && (Tab->GetLayoutIdentifier().TabType == PrimaryDocumentTabId || PrimaryDocumentTabId.IsNone()))
 	{
 		return true;
 	}
