Fixing highlighting system overview  from parameter panel selection of user and linked variables

UE-157183 - Parameter read/write System Overview icons are not displayed on modules in Emitter assets
UE-157184 - Parameter Read/Write overview icons aren't displayed on the Set Parameters module
UE-157185 - Parameter read/write overview icons do not update when the selected parameter is unreferenced or being set by a Set Parameters module

diff --git a/Engine/Plugins/FX/Niagara/Source/NiagaraEditor/Private/ViewModels/Stack/NiagaraStackFunctionInput.cpp b/Engine/Plugins/FX/Niagara/Source/NiagaraEditor/Private/ViewModels/Stack/NiagaraStackFunctionInput.cpp
--- a/Engine/Plugins/FX/Niagara/Source/NiagaraEditor/Private/ViewModels/Stack/NiagaraStackFunctionInput.cpp
+++ b/Engine/Plugins/FX/Niagara/Source/NiagaraEditor/Private/ViewModels/Stack/NiagaraStackFunctionInput.cpp
@@ -2674,7 +2674,6 @@ const UNiagaraStackFunctionInput::FCollectedUsageData& UNiagaraStackFunctionInpu
 {
 	if (CachedCollectedUsageData.IsSet() == false)
 	{
-		CachedCollectedUsageData = FCollectedUsageData();
 		TSharedRef<FNiagaraSystemViewModel> SystemVM = GetSystemViewModel();
 		INiagaraParameterPanelViewModel* ParamVM = SystemVM->GetParameterPanelViewModel();
 		if (ParamVM)
@@ -2685,9 +2684,19 @@ const UNiagaraStackFunctionInput::FCollectedUsageData& UNiagaraStackFunctionInpu
 				FNiagaraVariableBase Var(InputType, InputValues.LinkedHandle.GetParameterHandleString());
 				bFoundOverride = ParamVM->IsVariableSelected(Var);
 			}
+			else if (InputValues.Mode == EValueMode::Dynamic)
+			{
+				CachedCollectedUsageData = UNiagaraStackEntry::GetCollectedUsageData();
+				return CachedCollectedUsageData.GetValue();
+			}
+			CachedCollectedUsageData = FCollectedUsageData();
 			FNiagaraVariableBase Var(InputType, AliasedInputParameterHandle.GetParameterHandleString());
 			CachedCollectedUsageData.GetValue().bHasReferencedParameterRead = ParamVM->IsVariableSelected(Var) || bFoundOverride;
 		}
+		else
+		{
+			CachedCollectedUsageData = FCollectedUsageData();
+		}
 
 	}
 
diff --git a/Engine/Plugins/FX/Niagara/Source/NiagaraEditor/Private/ViewModels/Stack/NiagaraStackModuleItem.cpp b/Engine/Plugins/FX/Niagara/Source/NiagaraEditor/Private/ViewModels/Stack/NiagaraStackModuleItem.cpp
--- a/Engine/Plugins/FX/Niagara/Source/NiagaraEditor/Private/ViewModels/Stack/NiagaraStackModuleItem.cpp
+++ b/Engine/Plugins/FX/Niagara/Source/NiagaraEditor/Private/ViewModels/Stack/NiagaraStackModuleItem.cpp
@@ -29,6 +29,7 @@
 #include "ViewModels/NiagaraEmitterHandleViewModel.h"
 #include "ViewModels/NiagaraEmitterViewModel.h"
 #include "ViewModels/NiagaraScratchPadViewModel.h"
+#include "ViewModels/NiagaraParameterPanelViewModel.h"
 #include "ViewModels/NiagaraSystemViewModel.h"
 #include "ViewModels/Stack/NiagaraStackFunctionInputCollection.h"
 #include "ViewModels/Stack/NiagaraStackGraphUtilities.h"
@@ -225,6 +226,30 @@ const UNiagaraStackModuleItem::FCollectedUsageData& UNiagaraStackModuleItem::Get
 	if (CachedCollectedUsageData.IsSet() == false)
 	{
 		CachedCollectedUsageData = FCollectedUsageData();
+		if (FunctionCallNode && FunctionCallNode->IsA<UNiagaraNodeAssignment>())
+		{
+			UNiagaraNodeAssignment* AssignmentNode = CastChecked<UNiagaraNodeAssignment>(FunctionCallNode);
+			TSharedRef<FNiagaraSystemViewModel> SystemVM = GetSystemViewModel();
+			INiagaraParameterPanelViewModel* ParamVM = SystemVM->GetParameterPanelViewModel();
+
+			CachedCollectedUsageData.GetValue().bHasReferencedParameterWrite = false;
+			CachedCollectedUsageData.GetValue().bHasReferencedParameterRead = false;
+
+			if (ParamVM)
+			{
+				const TArray<FNiagaraVariable>& TargetVars = AssignmentNode->GetAssignmentTargets();
+				for (const FNiagaraVariable& Var : TargetVars)
+				{
+					FNiagaraVariableBase TestVar = Var;
+					bool bHandled = ParamVM->IsVariableSelected(TestVar);
+					if (bHandled)
+					{
+						CachedCollectedUsageData.GetValue().bHasReferencedParameterWrite = true;
+					}
+				}
+			}
+		}
+
 		if (!LinkedInputCollection.IsNull())
 		{
 			if (LinkedInputCollection->GetCollectedUsageData().bHasReferencedParameterRead)
