Niagara - Refactor the serialize function on niagara script for clarity and to prevent static analysis warnings on non-editor builds.

diff --git a/Engine/Plugins/FX/Niagara/Source/Niagara/Private/NiagaraScript.cpp b/Engine/Plugins/FX/Niagara/Source/Niagara/Private/NiagaraScript.cpp
--- a/Engine/Plugins/FX/Niagara/Source/Niagara/Private/NiagaraScript.cpp
+++ b/Engine/Plugins/FX/Niagara/Source/Niagara/Private/NiagaraScript.cpp
@@ -1612,102 +1612,89 @@ void UNiagaraScript::Serialize(FArchive& Ar)
 	const int32 NiagaraVer = Ar.CustomVer(FNiagaraCustomVersion::GUID);
 
 #if WITH_EDITORONLY_DATA
-	FNiagaraParameterStore::FScopedSuppressOnChanged SupressOnChanged(RapidIterationParameters);
-#endif
+	IsCooked = Ar.IsLoading() && Ar.IsFilterEditorOnly();
 
-	FNiagaraParameterStore TemporaryStore;
-	int32 NumRemoved = 0;
 	if (Ar.IsCooking())
 	{
-		bool bUsesRapidIterationParams = true;
-
-#if WITH_EDITORONLY_DATA
-		if (UNiagaraSystem* System = GetTypedOuter<UNiagaraSystem>())
-		{
-			bUsesRapidIterationParams = System->ShouldUseRapidIterationParameters();
-		}
-#endif
+		// Cook optimization can temporarily modify the rapid iteration parameters for the serialize call.  We don't want this to broadcast
+		// on changed messages because it will modify the change id on owning objects so we supress those messages starting here.
+		FNiagaraParameterStore::FScopedSuppressOnChanged SupressOnChanged(RapidIterationParameters);
 
+		UNiagaraSystem* System = GetTypedOuter<UNiagaraSystem>();
+		bool bUsesRapidIterationParams = System != nullptr && System->ShouldUseRapidIterationParameters();
 		if (!bUsesRapidIterationParams)
 		{
-			// Copy off the parameter store for now..
-			TemporaryStore = RapidIterationParameters;
-
-#if WITH_EDITORONLY_DATA
-			// Cache the unmodified rapid iteration parameters in an editor only property so that they can be restored in a cooked editor.
+			// Cache the unmodified rapid iteration parameters in an editor only property so that they can be restored in a cooked editor
+			// and so they can be restored after this optimization.
 			RapidIterationParametersCookedEditorCache = RapidIterationParameters;
-#endif
-
-			auto ParameterVariables = TemporaryStore.ReadParameterVariables();
 
 			// Get the active parameters
 			// Remove all parameters that aren't data interfaces or uobjects
+			int32 NumRemoved = 0;
+			auto ParameterVariables = RapidIterationParameters.ReadParameterVariables();
 			for (const FNiagaraVariableBase& Var : ParameterVariables)
 			{
 				if (Var.IsDataInterface() || Var.IsUObject())
+				{
 					continue;
+				}
 				RapidIterationParameters.RemoveParameter(Var);
 				NumRemoved++;
 			}
 
 			UE_LOG(LogNiagara, Verbose, TEXT("Pruned %d/%d parameters from script %s"), NumRemoved, ParameterVariables.Num(), *GetFullName());
 		}
-	}
-
-#if WITH_EDITORONLY_DATA
-	if (Ar.IsCooking() && Ar.IsSaving())
-	{
-		auto& ExecutableData = GetVMExecutableData();
 
-		if (Usage != ENiagaraScriptUsage::ParticleGPUComputeScript)
+		if (Ar.IsSaving())
 		{
-			ExecutableData.BakeScriptLiterals(ExecutableData.ScriptLiterals);
+			auto& ExecutableData = GetVMExecutableData();
 
-			// we only need the padding info for when we're dealing with GPU scripts (for
-			// FNiagaraScriptInstanceParameterStore::CopyParameterDataToPaddedBuffer())
-			ScriptExecutionParamStore.PaddingInfo.Empty();
-		}
-		else
-		{
-			ExecutableData.ScriptLiterals.Empty();
-			ScriptExecutionParamStore.CoalescePaddingInfo();
+			if (Usage != ENiagaraScriptUsage::ParticleGPUComputeScript)
+			{
+				ExecutableData.BakeScriptLiterals(ExecutableData.ScriptLiterals);
+
+				// we only need the padding info for when we're dealing with GPU scripts (for
+				// FNiagaraScriptInstanceParameterStore::CopyParameterDataToPaddedBuffer())
+				ScriptExecutionParamStore.PaddingInfo.Empty();
+			}
+			else
+			{
+				ExecutableData.ScriptLiterals.Empty();
+				ScriptExecutionParamStore.CoalescePaddingInfo();
+			}
+
+			if (GNiagaraCompressScriptByteCode)
+			{
+				ExecutableData.ByteCode.Compress();
+			}
 		}
 
-		if (GNiagaraCompressScriptByteCode)
+		Super::Serialize(Ar);
+
+		// If we cached the rapid iteration parameters before serialize, restore them here.
+		if (RapidIterationParametersCookedEditorCache.Num() != 0)
 		{
-			ExecutableData.ByteCode.Compress();
+			RapidIterationParameters = RapidIterationParametersCookedEditorCache;
+			RapidIterationParametersCookedEditorCache.Empty();
 		}
 	}
-
-	if (Ar.IsLoading())
-	{
-		IsCooked = Ar.IsFilterEditorOnly();
-	}
-#endif
-
-	Super::Serialize(Ar);
-
-	// Restore after serialize
-	if (Ar.IsCooking())
+	else if (Ar.IsLoading() && GetOutermost()->HasAnyPackageFlags(PKG_Cooked))
 	{
-		if (NumRemoved > 0)
+		Super::Serialize(Ar);
+		if (RapidIterationParametersCookedEditorCache.Num() > 0)
 		{
-			RapidIterationParameters = TemporaryStore;
+			// Restore the rapid iteration parameters from the cache for cooked editor builds.
+			RapidIterationParameters = RapidIterationParametersCookedEditorCache;
+			RapidIterationParametersCookedEditorCache.Empty();
 		}
-#if WITH_EDITORONLY_DATA
-		RapidIterationParametersCookedEditorCache.Empty();
-#endif
 	}
-
-#if WITH_EDITORONLY_DATA
-	if (Ar.IsLoading() && RapidIterationParametersCookedEditorCache.Num() > 0)
+	else
+#endif
 	{
-		// If we're loading and the cooked editor cache is available, re-populate the parameters from the cache after they're
-		// deserialized.
-		RapidIterationParameters = RapidIterationParametersCookedEditorCache;
-		RapidIterationParametersCookedEditorCache.Empty();
+		Super::Serialize(Ar);
 	}
 
+#if WITH_EDITORONLY_DATA
 	if (Ar.IsCooking() && Ar.IsSaving())
 	{
 		if (!HasValidParameterBindings())
