- Fix for accessing GPU scene delegates on the GameThread vs RenderThread

#rnx

diff --git a/Engine/Plugins/FX/Niagara/Source/Niagara/Private/NiagaraAsyncGpuTraceHelper.cpp b/Engine/Plugins/FX/Niagara/Source/Niagara/Private/NiagaraAsyncGpuTraceHelper.cpp
--- a/Engine/Plugins/FX/Niagara/Source/Niagara/Private/NiagaraAsyncGpuTraceHelper.cpp
+++ b/Engine/Plugins/FX/Niagara/Source/Niagara/Private/NiagaraAsyncGpuTraceHelper.cpp
@@ -31,8 +31,15 @@ FNiagaraAsyncGpuTraceHelper::FNiagaraAsyncGpuTraceHelper(EShaderPlatform InShade
 	InitProviders(Dispatcher);
 
 #if NIAGARA_ASYNC_GPU_TRACE_COLLISION_GROUPS
-	FPrimitiveSceneInfo::OnGPUSceneInstancesAllocated.AddRaw(this, &FNiagaraAsyncGpuTraceHelper::OnPrimitiveGPUSceneInstancesDirtied);
-	FPrimitiveSceneInfo::OnGPUSceneInstancesFreed.AddRaw(this, &FNiagaraAsyncGpuTraceHelper::OnPrimitiveGPUSceneInstancesDirtied);
+	// The owner of the trace helper is constructed on the GameThread so we need to enqueue a command to set the delegates
+	ENQUEUE_RENDER_COMMAND(NiagaraAsyncGpuTraceHelperSetupGPUSceneDelegates)
+	(
+		[RT_Helper=this](FRHICommandListImmediate&)
+		{
+			FPrimitiveSceneInfo::OnGPUSceneInstancesAllocated.AddRaw(RT_Helper, &FNiagaraAsyncGpuTraceHelper::OnPrimitiveGPUSceneInstancesDirtied);
+			FPrimitiveSceneInfo::OnGPUSceneInstancesFreed.AddRaw(RT_Helper, &FNiagaraAsyncGpuTraceHelper::OnPrimitiveGPUSceneInstancesDirtied);
+		}
+	);
 #endif
 }
 
diff --git a/Engine/Plugins/FX/Niagara/Source/Niagara/Private/NiagaraGpuComputeDispatch.cpp b/Engine/Plugins/FX/Niagara/Source/Niagara/Private/NiagaraGpuComputeDispatch.cpp
--- a/Engine/Plugins/FX/Niagara/Source/Niagara/Private/NiagaraGpuComputeDispatch.cpp
+++ b/Engine/Plugins/FX/Niagara/Source/Niagara/Private/NiagaraGpuComputeDispatch.cpp
@@ -219,9 +219,6 @@ FNiagaraGpuComputeDispatch::~FNiagaraGpuComputeDispatch()
 	FinishDispatches();
 
 	AsyncGpuTraceHelper->Reset();
-
-	FPrimitiveSceneInfo::OnGPUSceneInstancesAllocated.RemoveAll(this);
-	FPrimitiveSceneInfo::OnGPUSceneInstancesFreed.RemoveAll(this);
 }
 
 void FNiagaraGpuComputeDispatch::AddGpuComputeProxy(FNiagaraSystemGpuComputeProxy* ComputeProxy)
