#9405: Improve Niagara simulation stage extendibility (Contributed by [REDACTED_PERSON])

#rnx

diff --git a/Engine/Plugins/FX/Niagara/Source/Niagara/Private/NiagaraSimulationStageBase.cpp b/Engine/Plugins/FX/Niagara/Source/Niagara/Private/NiagaraSimulationStageBase.cpp
--- a/Engine/Plugins/FX/Niagara/Source/Niagara/Private/NiagaraSimulationStageBase.cpp
+++ b/Engine/Plugins/FX/Niagara/Source/Niagara/Private/NiagaraSimulationStageBase.cpp
@@ -90,6 +90,37 @@ bool UNiagaraSimulationStageGeneric::AppendCompileHash(FNiagaraCompileHashVisito
 	return true;
 }
 
+#if WITH_EDITORONLY_DATA
+bool UNiagaraSimulationStageGeneric::FillCompilationData(TArray<FNiagaraSimulationStageCompilationData>& CompilationSimStageData) const
+{
+	FNiagaraSimulationStageCompilationData& SimStageData = CompilationSimStageData.AddDefaulted_GetRef();
+	SimStageData.StageGuid = Script->GetUsageId();
+	SimStageData.StageName = SimulationStageName;
+	SimStageData.EnabledBinding = EnabledBinding.GetName();
+	SimStageData.ElementCountBinding = ElementCountBinding.GetName();
+	SimStageData.NumIterations = Iterations;
+	SimStageData.NumIterationsBinding = NumIterationsBinding.GetName();
+	SimStageData.IterationSource = IterationSource == ENiagaraIterationSource::DataInterface ? DataInterface.BoundVariable.GetName() : FName();
+	SimStageData.ExecuteBehavior = ExecuteBehavior;
+	SimStageData.PartialParticleUpdate = bDisablePartialParticleUpdate == false;
+	SimStageData.bParticleIterationStateEnabled = bParticleIterationStateEnabled;
+	SimStageData.ParticleIterationStateRange = ParticleIterationStateRange;
+	SimStageData.bGpuDispatchForceLinear = bGpuDispatchForceLinear;
+	SimStageData.bOverrideGpuDispatchNumThreads = bOverrideGpuDispatchNumThreads;
+	SimStageData.OverrideGpuDispatchNumThreads = OverrideGpuDispatchNumThreads;
+
+	if (SimStageData.bParticleIterationStateEnabled)
+	{
+		FString AttributeName = ParticleIterationStateBinding.GetName().ToString();
+		if (ensureMsgf(AttributeName.RemoveFromStart(TEXT("Particles.")), TEXT("Attribute '%s' is not in particles namespace"), *AttributeName))
+		{
+			SimStageData.ParticleIterationStateBinding = FName(AttributeName);
+		}
+	}
+	return true;
+}
+#endif //WITH_EDITORONLY_DATA
+
 void UNiagaraSimulationStageGeneric::PostInitProperties()
 {
 	Super::PostInitProperties();
diff --git a/Engine/Plugins/FX/Niagara/Source/Niagara/Public/NiagaraSimulationStageBase.h b/Engine/Plugins/FX/Niagara/Source/Niagara/Public/NiagaraSimulationStageBase.h
--- a/Engine/Plugins/FX/Niagara/Source/Niagara/Public/NiagaraSimulationStageBase.h
+++ b/Engine/Plugins/FX/Niagara/Source/Niagara/Public/NiagaraSimulationStageBase.h
@@ -7,6 +7,7 @@
 #include "NiagaraCommon.h"
 #include "NiagaraTypes.h"
 #include "NiagaraConstants.h"
+#include "NiagaraSimulationStageCompileData.h"
 #include "NiagaraSimulationStageBase.generated.h"
 
 class UNiagaraScript;
@@ -14,7 +15,7 @@ class UNiagaraScript;
 /**
 * A base class for niagara simulation stages.  This class should be derived to add stage specific information.
 */
-UCLASS()
+UCLASS(abstract)
 class NIAGARA_API UNiagaraSimulationStageBase : public UNiagaraMergeable
 {
 	GENERATED_BODY()
@@ -38,6 +39,7 @@ public:
 
 	virtual bool AppendCompileHash(FNiagaraCompileHashVisitor* InVisitor) const;
 #if WITH_EDITORONLY_DATA
+	virtual bool FillCompilationData(TArray<FNiagaraSimulationStageCompilationData>& CompilationSimStageData) const PURE_VIRTUAL(UNiagaraSimulationStageBase::FillCompileSimStageData, return false;);
 	/** Return the FName to use in place of the default for the location in the stack context. If this would be the default, return NAME_None.*/
 	virtual FName GetStackContextReplacementName() const { return NAME_None; }
 	void SetEnabled(bool bEnabled);
@@ -140,6 +142,9 @@ public:
 	virtual void PostInitProperties() override;
 	virtual void PostLoad() override;
 	virtual bool AppendCompileHash(FNiagaraCompileHashVisitor* InVisitor) const override;
+#if WITH_EDITORONLY_DATA
+	virtual bool FillCompilationData(TArray<FNiagaraSimulationStageCompilationData>& CompilationSimStageData) const override;
+#endif
 #if WITH_EDITOR
 	virtual FName GetStackContextReplacementName() const override; 
 
diff --git a/Engine/Plugins/FX/Niagara/Source/Niagara/Public/NiagaraSimulationStageCompileData.h b/Engine/Plugins/FX/Niagara/Source/Niagara/Public/NiagaraSimulationStageCompileData.h
new file mode 100644
--- /dev/null
+++ b/Engine/Plugins/FX/Niagara/Source/Niagara/Public/NiagaraSimulationStageCompileData.h
@@ -0,0 +1,24 @@
+// Copyright Epic Games, Inc. All Rights Reserved.
+
+#pragma once
+
+#include "NiagaraScriptBase.h"
+
+struct FNiagaraSimulationStageCompilationData
+{
+	FGuid                           StageGuid;
+	FName                           StageName;
+	FName                           EnabledBinding;
+	FName                           ElementCountBinding;
+	uint32                          NumIterations = 1;
+	FName                           NumIterationsBinding;
+	FName                           IterationSource;
+	ENiagaraSimStageExecuteBehavior ExecuteBehavior = ENiagaraSimStageExecuteBehavior::Always;
+	mutable bool                    PartialParticleUpdate = false;
+	bool                            bParticleIterationStateEnabled = false;
+	FName                           ParticleIterationStateBinding;
+	FIntPoint                       ParticleIterationStateRange = FIntPoint::ZeroValue;
+	bool                            bGpuDispatchForceLinear = false;
+	bool                            bOverrideGpuDispatchNumThreads = false;
+	FIntVector                      OverrideGpuDispatchNumThreads = FIntVector(1, 1, 1);
+};
diff --git a/Engine/Plugins/FX/Niagara/Source/NiagaraEditor/Private/NiagaraCompiler.cpp b/Engine/Plugins/FX/Niagara/Source/NiagaraEditor/Private/NiagaraCompiler.cpp
--- a/Engine/Plugins/FX/Niagara/Source/NiagaraEditor/Private/NiagaraCompiler.cpp
+++ b/Engine/Plugins/FX/Niagara/Source/NiagaraEditor/Private/NiagaraCompiler.cpp
@@ -712,33 +712,8 @@ void FNiagaraCompileRequestData::FinishPrecompile(const TArray<FNiagaraVariable>
 					continue;
 				}
 
-				if ( UNiagaraSimulationStageGeneric* GenericStage = Cast<UNiagaraSimulationStageGeneric>(SimStage) )
+				if ( SimStage->FillCompilationData(CompileSimStageData) )
 				{
-					FCompileSimStageData& SimStageData	= CompileSimStageData.AddDefaulted_GetRef();
-					SimStageData.StageGuid				= GenericStage->Script->GetUsageId();
-					SimStageData.StageName				= GenericStage->SimulationStageName;
-					SimStageData.EnabledBinding			= GenericStage->EnabledBinding.GetName();
-					SimStageData.ElementCountBinding	= GenericStage->ElementCountBinding.GetName();
-					SimStageData.NumIterations			= GenericStage->Iterations;
-					SimStageData.NumIterationsBinding	= GenericStage->NumIterationsBinding.GetName();
-					SimStageData.IterationSource		= GenericStage->IterationSource == ENiagaraIterationSource::DataInterface ? GenericStage->DataInterface.BoundVariable.GetName() : FName();
-					SimStageData.ExecuteBehavior		= GenericStage->ExecuteBehavior;
-					SimStageData.PartialParticleUpdate	= GenericStage->bDisablePartialParticleUpdate == false;
-					SimStageData.bParticleIterationStateEnabled	= GenericStage->bParticleIterationStateEnabled;
-					SimStageData.ParticleIterationStateRange	= GenericStage->ParticleIterationStateRange;
-					SimStageData.bGpuDispatchForceLinear		= GenericStage->bGpuDispatchForceLinear;
-					SimStageData.bOverrideGpuDispatchNumThreads	= GenericStage->bOverrideGpuDispatchNumThreads;
-					SimStageData.OverrideGpuDispatchNumThreads	= GenericStage->OverrideGpuDispatchNumThreads;
-
-					if (SimStageData.bParticleIterationStateEnabled)
-					{
-						FString AttributeName = GenericStage->ParticleIterationStateBinding.GetName().ToString();
-						if ( ensureMsgf(AttributeName.RemoveFromStart(TEXT("Particles.")), TEXT("Attribute '%s' is not in particles namespace"), *AttributeName) )
-						{
-							SimStageData.ParticleIterationStateBinding = FName(AttributeName);
-						}
-					}
-
 					++ActiveStageCount;
 				}
 			}
diff --git a/Engine/Plugins/FX/Niagara/Source/NiagaraEditor/Private/NiagaraHlslTranslator.cpp b/Engine/Plugins/FX/Niagara/Source/NiagaraEditor/Private/NiagaraHlslTranslator.cpp
--- a/Engine/Plugins/FX/Niagara/Source/NiagaraEditor/Private/NiagaraHlslTranslator.cpp
+++ b/Engine/Plugins/FX/Niagara/Source/NiagaraEditor/Private/NiagaraHlslTranslator.cpp
@@ -1593,7 +1593,9 @@ const FNiagaraTranslateResults &FHlslNiagaraTranslator::Translate(const FNiagara
 				AddBodyComment(FString::Printf(TEXT("//Begin Stage Script: %s!"), *TranslationStages[i].PassNamespace));
 				//Now we compile the simulation stage and read from the temp values written above.
 				CurrentParamMapIndices.Empty();
-				CurrentParamMapIndices.Add(i);				
+				CurrentParamMapIndices.Add(i);
+				PinToCodeChunks.Empty();
+				PinToCodeChunks.AddDefaulted(1);				
 				FName IterSource = TranslationStages[i].IterationSource;
 				ActiveHistoryForFunctionCalls.BeginUsage(TranslationStages[i].ScriptUsage, IterSource);
 				TranslationStages[i].OutputNode->Compile(ThisTranslator, OutputChunks);
diff --git a/Engine/Plugins/FX/Niagara/Source/NiagaraEditor/Private/NiagaraHlslTranslator.h b/Engine/Plugins/FX/Niagara/Source/NiagaraEditor/Private/NiagaraHlslTranslator.h
--- a/Engine/Plugins/FX/Niagara/Source/NiagaraEditor/Private/NiagaraHlslTranslator.h
+++ b/Engine/Plugins/FX/Niagara/Source/NiagaraEditor/Private/NiagaraHlslTranslator.h
@@ -9,6 +9,7 @@
 #include "NiagaraShaderCompilationManager.h"
 #include "NiagaraDataInterface.h"
 #include "NiagaraScriptSource.h"
+#include "NiagaraSimulationStageCompileData.h"
 
 class Error;
 class UNiagaraGraph;
@@ -97,25 +98,7 @@ public:
 	virtual bool GetDisableDebugSwitches() const override { return bDisableDebugSwitches; }
 
 	// Simulation Stage Variables. Sim stage of 0 is always Spawn/Update
-	struct FCompileSimStageData
-	{
-		FGuid							StageGuid;
-		FName							StageName;
-		FName							EnabledBinding;
-		FName							ElementCountBinding;
-		uint32							NumIterations = 1;
-		FName							NumIterationsBinding;
-		FName							IterationSource;
-		ENiagaraSimStageExecuteBehavior	ExecuteBehavior = ENiagaraSimStageExecuteBehavior::Always;
-		mutable bool					PartialParticleUpdate = false;
-		bool							bParticleIterationStateEnabled = false;
-		FName							ParticleIterationStateBinding;
-		FIntPoint						ParticleIterationStateRange = FIntPoint::ZeroValue;
-		bool							bGpuDispatchForceLinear = false;
-		bool							bOverrideGpuDispatchNumThreads = false;
-		FIntVector						OverrideGpuDispatchNumThreads = FIntVector(1, 1, 1);
-	};
-	TArray<FCompileSimStageData> CompileSimStageData;
+	TArray<FNiagaraSimulationStageCompilationData> CompileSimStageData;
 
 	struct FCompileDataInterfaceData
 	{
