- Add TextureAddress options to 2D bake
- Change default to Wrap as the default Flipbook material node generates UVs outside of 0-1

#rnx

diff --git a/Engine/Plugins/FX/Niagara/Source/Niagara/Classes/NiagaraBakerOutputTexture2D.h b/Engine/Plugins/FX/Niagara/Source/Niagara/Classes/NiagaraBakerOutputTexture2D.h
--- a/Engine/Plugins/FX/Niagara/Source/Niagara/Classes/NiagaraBakerOutputTexture2D.h
+++ b/Engine/Plugins/FX/Niagara/Source/Niagara/Classes/NiagaraBakerOutputTexture2D.h
@@ -2,6 +2,7 @@
 
 #pragma once
 
+#include "Engine/Texture.h"
 #include "NiagaraBakerOutput.h"
 #include "NiagaraBakerOutputTexture2D.generated.h"
 
@@ -13,6 +14,8 @@ class NIAGARA_API UNiagaraBakerOutputTexture2D : public UNiagaraBakerOutput
 public:
 	UNiagaraBakerOutputTexture2D(const FObjectInitializer& Init)
 		: bGenerateAtlas(true)
+		, bSetTextureAddressX(true)
+		, bSetTextureAddressY(true)
 	{
 	}
 
@@ -29,6 +32,12 @@ public:
 	UPROPERTY(EditAnywhere, Category = "Settings", meta = (PinHiddenByDefault, InlineEditConditionToggle))
 	uint8 bExportFrames : 1;
 
+	UPROPERTY(EditAnywhere, Category = "Settings", meta = (PinHiddenByDefault, InlineEditConditionToggle))
+	uint8 bSetTextureAddressX : 1;
+
+	UPROPERTY(EditAnywhere, Category = "Settings", meta = (PinHiddenByDefault, InlineEditConditionToggle))
+	uint8 bSetTextureAddressY : 1;
+
 	/** Size of each frame we generate. */
 	UPROPERTY(EditAnywhere, Category = "Settings")
 	FIntPoint FrameSize = FIntPoint(128, 128);
@@ -40,6 +49,14 @@ public:
 	/** When either exporting or generating individual frames this is the size of the texture. */
 	FIntPoint SequenceFrameSize = FIntPoint(128, 128);
 
+	/** After baking sets the texture address mode to use on the X axis. */
+	UPROPERTY(EditAnywhere, Category="Settings", meta = (EditCondition = "bSetTextureAddressX"))
+	TEnumAsByte<enum TextureAddress> TextureAddressX = TextureAddress::TA_Wrap;
+
+	/** After baking sets the texture address mode to use on the Y axis. */
+	UPROPERTY(EditAnywhere, Category = "Settings", meta = (EditCondition = "bSetTextureAddressY"))
+	TEnumAsByte<enum TextureAddress> TextureAddressY = TextureAddress::TA_Wrap;
+
 	/**
 	When enabled a texture atlas is created
 	*/
diff --git a/Engine/Plugins/FX/Niagara/Source/NiagaraEditor/Private/NiagaraBakerRendererOutputTexture2D.cpp b/Engine/Plugins/FX/Niagara/Source/NiagaraEditor/Private/NiagaraBakerRendererOutputTexture2D.cpp
--- a/Engine/Plugins/FX/Niagara/Source/NiagaraEditor/Private/NiagaraBakerRendererOutputTexture2D.cpp
+++ b/Engine/Plugins/FX/Niagara/Source/NiagaraEditor/Private/NiagaraBakerRendererOutputTexture2D.cpp
@@ -183,8 +183,8 @@ void FNiagaraBakerRendererOutputTexture2D::BakeFrame(UNiagaraBakerOutput* InBake
 			OutputTexture->Source.Init(BakerOutput->FrameSize.X, BakerOutput->FrameSize.Y, 1, 1, TSF_RGBA16F, (const uint8*)(FrameResults.GetData()));
 			OutputTexture->PowerOfTwoMode = ETexturePowerOfTwoSetting::None;
 			OutputTexture->MipGenSettings = bIsPoT ? TextureMipGenSettings::TMGS_FromTextureGroup : TextureMipGenSettings::TMGS_NoMipmaps;
-			OutputTexture->AddressX = TextureAddress::TA_Clamp;
-			OutputTexture->AddressY = TextureAddress::TA_Clamp;
+			OutputTexture->AddressX = BakerOutput->bSetTextureAddressX ? BakerOutput->TextureAddressX : OutputTexture->AddressX;
+			OutputTexture->AddressY = BakerOutput->bSetTextureAddressY ? BakerOutput->TextureAddressY : OutputTexture->AddressY;
 			OutputTexture->UpdateResource();
 			OutputTexture->PostEditChange();
 			OutputTexture->MarkPackageDirty();
@@ -219,8 +219,8 @@ void FNiagaraBakerRendererOutputTexture2D::EndBake(UNiagaraBakerOutput* InBakerO
 			OutputTexture->Source.Init(BakerOutput->AtlasTextureSize.X, BakerOutput->AtlasTextureSize.Y, 1, 1, TSF_RGBA16F, (const uint8*)(BakeAtlasTextureData.GetData()));
 			OutputTexture->PowerOfTwoMode = ETexturePowerOfTwoSetting::None;
 			OutputTexture->MipGenSettings = bIsPoT ? TextureMipGenSettings::TMGS_FromTextureGroup : TextureMipGenSettings::TMGS_NoMipmaps;
-			OutputTexture->AddressX = TextureAddress::TA_Clamp;
-			OutputTexture->AddressY = TextureAddress::TA_Clamp;
+			OutputTexture->AddressX = BakerOutput->bSetTextureAddressX ? BakerOutput->TextureAddressX : OutputTexture->AddressX;
+			OutputTexture->AddressY = BakerOutput->bSetTextureAddressY ? BakerOutput->TextureAddressY : OutputTexture->AddressY;
 			OutputTexture->UpdateResource();
 			OutputTexture->PostEditChange();
 			OutputTexture->MarkPackageDirty();
