Couple of changes for handling overridden painted vertex colors for static meshes
-include override support for GetColor function (vertex sampling)
-in editor, change access to override buffer to be done per frame so that we are ensured to be using the most up to date buffer (potential issues remain where the buffer could be manipulated while things are in flight in niagara)

diff --git a/Engine/Plugins/FX/Niagara/Source/Niagara/Private/DataInterface/NiagaraDataInterfaceStaticMesh.cpp b/Engine/Plugins/FX/Niagara/Source/Niagara/Private/DataInterface/NiagaraDataInterfaceStaticMesh.cpp
--- a/Engine/Plugins/FX/Niagara/Source/Niagara/Private/DataInterface/NiagaraDataInterfaceStaticMesh.cpp
+++ b/Engine/Plugins/FX/Niagara/Source/Niagara/Private/DataInterface/NiagaraDataInterfaceStaticMesh.cpp
@@ -310,7 +310,9 @@ namespace NDIStaticMeshLocal
 		FNiagaraSystemInstanceID			SystemInstanceID;
 		TRefCountPtr<const FStaticMeshLODResources>		LODResource = nullptr;
 
+#if !WITH_EDITOR
 		FColorVertexBuffer*					OverrideColorBuffer = nullptr;
+#endif
 
 #if DO_CHECK
 		FName								SystemFName;
@@ -357,6 +359,12 @@ namespace NDIStaticMeshLocal
 		FShaderResourceViewRHIRef MeshUVBufferSRV;
 		FShaderResourceViewRHIRef MeshColorBufferSRV;
 
+#if WITH_EDITOR
+		// in the editor, there's the potential for dynamically updating (including getting cleared) vertex buffers.
+		// At runtime we can just handle things during initialization of the instance data
+		FShaderResourceViewRHIRef MeshOverrideColorBufferSRV;
+#endif
+
 		FShaderResourceViewRHIRef MeshUniformSamplingTriangleSRV;
 
 		FIntVector	SectionCounts = FIntVector::ZeroValue;	// X=NumSections, y = Num Filtered, Z=NumUnfiltered
@@ -395,7 +403,17 @@ namespace NDIStaticMeshLocal
 				MeshPositionBufferSRV = GpuInitializeData.LODResource->VertexBuffers.PositionVertexBuffer.GetSRV();
 				MeshTangentBufferSRV = GpuInitializeData.LODResource->VertexBuffers.StaticMeshVertexBuffer.GetTangentsSRV();
 				MeshUVBufferSRV = GpuInitializeData.LODResource->VertexBuffers.StaticMeshVertexBuffer.GetTexCoordsSRV();
-				MeshColorBufferSRV = GpuInitializeData.OverrideColorBuffer ? GpuInitializeData.OverrideColorBuffer->GetColorComponentsSRV() : GpuInitializeData.LODResource->VertexBuffers.ColorVertexBuffer.GetColorComponentsSRV();
+
+#if !WITH_EDITOR
+				if (GpuInitializeData.OverrideColorBuffer)
+				{
+					MeshColorBufferSRV = GpuInitializeData.OverrideColorBuffer->GetColorComponentsSRV();
+				}
+				else
+#endif
+				{
+					MeshColorBufferSRV = GpuInitializeData.LODResource->VertexBuffers.ColorVertexBuffer.GetColorComponentsSRV();
+				}
 
 				NumTriangles.X = GpuInitializeData.LODResource->IndexBuffer.GetNumIndices() / 3;
 				NumTriangles.Y = GpuInitializeData.NumFilteredTriangles;
@@ -484,6 +502,9 @@ namespace NDIStaticMeshLocal
 		FQuat4f					PrevRotation = FQuat4f::Identity;
 		float					DeltaSeconds = 0.0f;
 		FPrimitiveComponentId	DistanceFieldPrimitiveId;
+#if WITH_EDITOR
+		FColorVertexBuffer*		OverrideVertexColors = nullptr;
+#endif
 	};
 
 	struct FInstanceData_GameThread
@@ -959,6 +980,17 @@ namespace NDIStaticMeshLocal
 			InstanceData->DeltaSeconds				= FromGameThread->DeltaSeconds;
 			InstanceData->DistanceFieldPrimitiveId	= FromGameThread->DistanceFieldPrimitiveId;
 
+#if WITH_EDITOR
+			if (FromGameThread->OverrideVertexColors)
+			{
+				InstanceData->MeshOverrideColorBufferSRV = FromGameThread->OverrideVertexColors->GetColorComponentsSRV();
+			}
+			else
+			{
+				InstanceData->MeshOverrideColorBufferSRV.SafeRelease();
+			}
+#endif
+
 			FromGameThread->~FInstanceData_FromGameThread();
 		}
 
@@ -1166,6 +1198,10 @@ namespace NDIStaticMeshLocal
 
 		FORCEINLINE FLinearColor GetColor(int32 Vertex) const
 		{
+			if (OverrideVertexColors)
+			{
+				return OverrideVertexColors->VertexColor(Vertex);
+			}
 			return LODResource->VertexBuffers.ColorVertexBuffer.VertexColor(Vertex);
 		}
 
@@ -1390,7 +1426,11 @@ bool UNiagaraDataInterfaceStaticMesh::InitPerInstanceData(void* PerInstanceData,
 
 		//This is safe to ref on the RT as it's freed on the RT in FStaticMeshComponentLODInfo::BeginReleaseOverrideVertexColors()
 		//However, it's seems unsafe to reference in Niagara's instance data this way as there looks to be a window between this RT command and it's actual use where the data could have been freed.
+#if WITH_EDITOR
+		GpuInitializeData->LODResource = InstanceData->GetCurrentFirstLOD();
+#else
 		GpuInitializeData->LODResource = InstanceData->GetCurrentFirstLODWithVertexColorOverrides(GpuInitializeData->OverrideColorBuffer);
+#endif
 
 		if ( GpuInitializeData->LODResource )
 		{
@@ -1524,6 +1564,10 @@ void UNiagaraDataInterfaceStaticMesh::ProvidePerInstanceDataForRenderThread(void
 	DataFromGT->PrevRotation				= InstanceData->PrevRotation;
 	DataFromGT->DeltaSeconds				= InstanceData->DeltaSeconds;
 	DataFromGT->DistanceFieldPrimitiveId	= FPrimitiveComponentId();
+#if WITH_EDITOR
+	InstanceData->GetCurrentFirstLODWithVertexColorOverrides(DataFromGT->OverrideVertexColors);
+#endif
+
 	if ( UPrimitiveComponent* PrimitiveComponent = Cast<UPrimitiveComponent>(InstanceData->SceneComponentWeakPtr) )
 	{
 		DataFromGT->DistanceFieldPrimitiveId = PrimitiveComponent->ComponentId;
@@ -2460,7 +2504,15 @@ void UNiagaraDataInterfaceStaticMesh::SetShaderParameters(const FNiagaraDataInte
 	// Set mesh sampling data
 	if (InstanceData.bIsValid)
 	{
-		const bool bHasColors = InstanceData.MeshColorBufferSRV.IsValid();
+#if WITH_EDITOR
+		const FShaderResourceViewRHIRef& ColorSRV = InstanceData.MeshOverrideColorBufferSRV.IsValid()
+			? InstanceData.MeshOverrideColorBufferSRV
+			: InstanceData.MeshColorBufferSRV;
+#else
+		const FShaderResourceViewRHIRef& ColorSRV = InstanceData.MeshColorBufferSRV;
+#endif
+
+		const bool bHasColors = ColorSRV.IsValid();
 		ShaderParameters->NumTriangles		= InstanceData.NumTriangles;
 		ShaderParameters->NumVertices		= InstanceData.NumVertices;
 		ShaderParameters->NumUVs			= InstanceData.NumUVs;
@@ -2469,7 +2521,7 @@ void UNiagaraDataInterfaceStaticMesh::SetShaderParameters(const FNiagaraDataInte
 		ShaderParameters->PositionBuffer	= InstanceData.MeshPositionBufferSRV;
 		ShaderParameters->TangentBuffer		= InstanceData.MeshTangentBufferSRV;
 		ShaderParameters->UVBuffer			= InstanceData.MeshUVBufferSRV.IsValid() ? InstanceData.MeshUVBufferSRV.GetReference() : FNiagaraRenderer::GetDummyFloat2Buffer();
-		ShaderParameters->ColorBuffer		= bHasColors ? InstanceData.MeshColorBufferSRV.GetReference() : FNiagaraRenderer::GetDummyWhiteColorBuffer();
+		ShaderParameters->ColorBuffer		= bHasColors ? ColorSRV.GetReference() : FNiagaraRenderer::GetDummyWhiteColorBuffer();
 		ShaderParameters->HasUniformSampling = InstanceData.bGpuUniformDistribution ? 1 : 0;
 		ShaderParameters->UniformSamplingTriangles = InstanceData.bGpuUniformDistribution ? InstanceData.MeshUniformSamplingTriangleSRV.GetReference() : FNiagaraRenderer::GetDummyUInt2Buffer();
 	}
