Fixing invalid opened asset when working with inherited emitters and scratch pad scripts.

diff --git a/Engine/Plugins/FX/Niagara/Source/NiagaraEditor/Private/ViewModels/NiagaraScratchPadViewModel.cpp b/Engine/Plugins/FX/Niagara/Source/NiagaraEditor/Private/ViewModels/NiagaraScratchPadViewModel.cpp
--- a/Engine/Plugins/FX/Niagara/Source/NiagaraEditor/Private/ViewModels/NiagaraScratchPadViewModel.cpp
+++ b/Engine/Plugins/FX/Niagara/Source/NiagaraEditor/Private/ViewModels/NiagaraScratchPadViewModel.cpp
@@ -96,7 +96,7 @@ void UpdateChangeId(TSharedRef<FNiagaraSystemViewModel> SystemViewModel)
 
 void UNiagaraScratchPadViewModel::RefreshScriptViewModels()
 {
-	TArray<TSharedRef<FNiagaraScratchPadScriptViewModel>> OldScriptViewModels = ScriptViewModels;
+	TArray<TSharedRef<FNiagaraScratchPadScriptViewModel>> OldScriptViewModelsRemaining = ScriptViewModels;
 	ScriptViewModels.Empty();
 	bHasUnappliedChangesCache.Reset();
 
@@ -121,12 +121,12 @@ void UNiagaraScratchPadViewModel::RefreshScriptViewModels()
 			if (GetSystemViewModel()->GetEditMode() == ENiagaraSystemViewModelEditMode::SystemAsset)
 			{
 				// System scripts don't ever merge, so we can safely just take the scripts or reuse view models.
-				TSharedRef<FNiagaraScratchPadScriptViewModel>* OldScriptViewModel = OldScriptViewModels.FindByPredicate(
+				TSharedRef<FNiagaraScratchPadScriptViewModel>* OldScriptViewModel = OldScriptViewModelsRemaining.FindByPredicate(
 					[ScratchPadScript](TSharedRef<FNiagaraScratchPadScriptViewModel> ScriptViewModel) { return ScriptViewModel->GetOriginalScript() == ScratchPadScript; });
 				if (OldScriptViewModel != nullptr)
 				{
 					ScriptViewModel = *OldScriptViewModel;
-					OldScriptViewModels.Remove(ScriptViewModel.ToSharedRef());
+					OldScriptViewModelsRemaining.Remove(ScriptViewModel.ToSharedRef());
 				}
 				else
 				{
@@ -138,10 +138,14 @@ void UNiagaraScratchPadViewModel::RefreshScriptViewModels()
 			{
 				// Emitter scripts may already be opened for edit but not yet applied so mergers can refresh and would wipe out the base scripts. These need to be kept up to date.
 				// Because the merge could wipe out view models, we can't just reuse the old ones, but we can steal their graphs.
-				TSharedPtr<FNiagaraScratchPadScriptViewModel> OldScriptViewModel = OldScriptViewModels.IsValidIndex(i) ? OldScriptViewModels[i] : TSharedPtr<FNiagaraScratchPadScriptViewModel>();
+				// We need to match by name as indices are not stable.
+				TSharedRef<FNiagaraScratchPadScriptViewModel>* OldScriptViewModelRef = OldScriptViewModelsRemaining.FindByPredicate(
+					[ScratchPadScript](TSharedRef<FNiagaraScratchPadScriptViewModel> ScriptViewModel) { return ScriptViewModel->GetOriginalScript()->GetFName() == ScratchPadScript->GetFName(); });
 
-				if (OldScriptViewModel.IsValid())
+
+				if (OldScriptViewModelRef != nullptr)
 				{
+					TSharedPtr<FNiagaraScratchPadScriptViewModel> OldScriptViewModel = (*OldScriptViewModelRef);
 					// Pass the baton of the script to the new generation. Make sure to clear out references in the previous generation.
 					ScriptViewModel = CreateAndSetupScriptviewModel(ScratchPadScript, OldScriptViewModel->GetEditScript().Script);
 					// This will copy over apply state and anything else that would otherwise be lost in the handoff
@@ -157,7 +161,7 @@ void UNiagaraScratchPadViewModel::RefreshScriptViewModels()
 					GetSystemViewModel()->GetDocumentViewModel()->SwapEditableScripts(OldScriptViewModel, ScriptViewModel);
 					// Clear out the old to make sure it isn't touching anything
 					OldScriptViewModel->SetScript(FVersionedNiagaraScript());
-					OldScriptViewModels.Remove(OldScriptViewModel.ToSharedRef());
+					OldScriptViewModelsRemaining.Remove(OldScriptViewModel.ToSharedRef());
 					OldScriptViewModel.Reset();
 				}
 				else
@@ -170,9 +174,9 @@ void UNiagaraScratchPadViewModel::RefreshScriptViewModels()
 		}
 	}
 
-	if (OldScriptViewModels.Num() > 0)
+	if (OldScriptViewModelsRemaining.Num() > 0)
 	{
-		for (TSharedRef<FNiagaraScratchPadScriptViewModel> OldScriptViewModel : OldScriptViewModels)
+		for (TSharedRef<FNiagaraScratchPadScriptViewModel> OldScriptViewModel : OldScriptViewModelsRemaining)
 		{
 			TearDownScriptViewModel(OldScriptViewModel);
 		}
