- Add BP method to modify filltered socket / bones for skel mesh sampling
- This would ideally be done via user parameters, but until we have FNames / Arrays types in Niagara this is a stop gap solution

diff --git a/Engine/Plugins/FX/Niagara/Source/Niagara/Classes/NiagaraDataInterfaceSkeletalMesh.h b/Engine/Plugins/FX/Niagara/Source/Niagara/Classes/NiagaraDataInterfaceSkeletalMesh.h
--- a/Engine/Plugins/FX/Niagara/Source/Niagara/Classes/NiagaraDataInterfaceSkeletalMesh.h
+++ b/Engine/Plugins/FX/Niagara/Source/Niagara/Classes/NiagaraDataInterfaceSkeletalMesh.h
@@ -1046,6 +1046,8 @@ public:
 
 	void SetSourceComponentFromBlueprints(USkeletalMeshComponent* ComponentToUse);
 	void SetSamplingRegionsFromBlueprints(const TArray<FName>& InSamplingRegions);
+	void SetFilteredBonesFromBlueprints(const TArray<FName>& InFilteredBones);
+	void SetFilteredSocketsFromBlueprints(const TArray<FName>& InFilteredSockets);
 	void SetWholeMeshLODFromBlueprints(int32 MeshLODLevel);
 };
 
diff --git a/Engine/Plugins/FX/Niagara/Source/Niagara/Private/NiagaraDataInterfaceSkeletalMesh.cpp b/Engine/Plugins/FX/Niagara/Source/Niagara/Private/NiagaraDataInterfaceSkeletalMesh.cpp
--- a/Engine/Plugins/FX/Niagara/Source/Niagara/Private/NiagaraDataInterfaceSkeletalMesh.cpp
+++ b/Engine/Plugins/FX/Niagara/Source/Niagara/Private/NiagaraDataInterfaceSkeletalMesh.cpp
@@ -3461,7 +3461,7 @@ void UNiagaraDataInterfaceSkeletalMesh::UnbindSourceDelegates()
 void UNiagaraDataInterfaceSkeletalMesh::OnSourceEndPlay(AActor* InSource, EEndPlayReason::Type Reason)
 {
 	// Increment change id in case we're able to find a new source component 
-	ChangeId++;
+	++ChangeId;
 	SoftSourceActor = nullptr;
 	SourceComponent = nullptr;
 }
@@ -3469,7 +3469,7 @@ void UNiagaraDataInterfaceSkeletalMesh::OnSourceEndPlay(AActor* InSource, EEndPl
 void UNiagaraDataInterfaceSkeletalMesh::SetSourceComponentFromBlueprints(USkeletalMeshComponent* ComponentToUse)
 {
 	// NOTE: When ChangeId changes the next tick will be skipped and a reset of the per-instance data will be initiated.
-	ChangeId++;
+	++ChangeId;
 	UnbindSourceDelegates();
 	SourceComponent = ComponentToUse;
 	SoftSourceActor = ComponentToUse->GetOwner();
@@ -3479,10 +3479,24 @@ void UNiagaraDataInterfaceSkeletalMesh::SetSourceComponentFromBlueprints(USkelet
 void UNiagaraDataInterfaceSkeletalMesh::SetSamplingRegionsFromBlueprints(const TArray<FName>& InSamplingRegions)
 {
 	// NOTE: When ChangeId changes the next tick will be skipped and a reset of the per-instance data will be initiated.
-	ChangeId++;
+	++ChangeId;
 	SamplingRegions = InSamplingRegions;
 }
 
+void UNiagaraDataInterfaceSkeletalMesh::SetFilteredBonesFromBlueprints(const TArray<FName>& InFilteredBones)
+{
+	// NOTE: When ChangeId changes the next tick will be skipped and a reset of the per-instance data will be initiated.
+	++ChangeId;
+	FilteredBones = InFilteredBones;
+}
+
+void UNiagaraDataInterfaceSkeletalMesh::SetFilteredSocketsFromBlueprints(const TArray<FName>& InFilteredSockets)
+{
+	// NOTE: When ChangeId changes the next tick will be skipped and a reset of the per-instance data will be initiated.
+	++ChangeId;
+	FilteredSockets = InFilteredSockets;
+}
+
 void UNiagaraDataInterfaceSkeletalMesh::SetWholeMeshLODFromBlueprints(int32 InWholeMeshLOD)
 {
 	// NOTE: When ChangeId changes the next tick will be skipped and a reset of the per-instance data will be initiated.
diff --git a/Engine/Plugins/FX/Niagara/Source/Niagara/Private/NiagaraFunctionLibrary.cpp b/Engine/Plugins/FX/Niagara/Source/Niagara/Private/NiagaraFunctionLibrary.cpp
--- a/Engine/Plugins/FX/Niagara/Source/Niagara/Private/NiagaraFunctionLibrary.cpp
+++ b/Engine/Plugins/FX/Niagara/Source/Niagara/Private/NiagaraFunctionLibrary.cpp
@@ -473,6 +473,49 @@ void UNiagaraFunctionLibrary::SetSkeletalMeshDataInterfaceSamplingRegions(UNiaga
 	SkeletalMeshInterface->SetSamplingRegionsFromBlueprints(SamplingRegions);
 }
 
+
+void UNiagaraFunctionLibrary::SetSkeletalMeshDataInterfaceFilteredBones(UNiagaraComponent* NiagaraSystem, const FString& OverrideName, const TArray<FName>& FilteredBones)
+{
+	if (!NiagaraSystem)
+	{
+		if (FNiagaraUtilities::LogVerboseWarnings())
+		{
+			UE_LOG(LogNiagara, Warning, TEXT("NiagaraSystem in \"Set Skeletal Mesh Data Interface Filtered Bones\" is NULL, OverrideName \"%s\", skipping."), *OverrideName);
+		}
+		return;
+	}
+
+	UNiagaraDataInterfaceSkeletalMesh* SkeletalMeshInterface = GetSkeletalMeshDataInterface(NiagaraSystem, OverrideName);
+	if (!SkeletalMeshInterface)
+	{
+		UE_LOG(LogNiagara, Warning, TEXT("UNiagaraFunctionLibrary::SetSkeletalMeshDataInterfaceFilteredBones: Did not find a matching Skeletal Mesh Data Interface variable named \"%s\" in the User variables of NiagaraSystem \"%s\" ."), *OverrideName, *GetFullNameSafe(NiagaraSystem));
+		return;
+	}
+
+	SkeletalMeshInterface->SetFilteredBonesFromBlueprints(FilteredBones);
+}
+
+void UNiagaraFunctionLibrary::SetSkeletalMeshDataInterfaceFilteredSockets(UNiagaraComponent* NiagaraSystem, const FString& OverrideName, const TArray<FName>& FilteredSockets)
+{
+	if (!NiagaraSystem)
+	{
+		if (FNiagaraUtilities::LogVerboseWarnings())
+		{
+			UE_LOG(LogNiagara, Warning, TEXT("NiagaraSystem in \"Set Skeletal Mesh Data Interface Filtered Sockets\" is NULL, OverrideName \"%s\", skipping."), *OverrideName);
+		}
+		return;
+	}
+
+	UNiagaraDataInterfaceSkeletalMesh* SkeletalMeshInterface = GetSkeletalMeshDataInterface(NiagaraSystem, OverrideName);
+	if (!SkeletalMeshInterface)
+	{
+		UE_LOG(LogNiagara, Warning, TEXT("UNiagaraFunctionLibrary::SetSkeletalMeshDataInterfaceFilteredSockets: Did not find a matching Skeletal Mesh Data Interface variable named \"%s\" in the User variables of NiagaraSystem \"%s\" ."), *OverrideName, *GetFullNameSafe(NiagaraSystem));
+		return;
+	}
+
+	SkeletalMeshInterface->SetFilteredSocketsFromBlueprints(FilteredSockets);
+}
+
 void UNiagaraFunctionLibrary::SetTextureObject(UNiagaraComponent* NiagaraSystem, const FString& OverrideName, UTexture* Texture)
 {
 	if (!NiagaraSystem)
diff --git a/Engine/Plugins/FX/Niagara/Source/Niagara/Public/NiagaraFunctionLibrary.h b/Engine/Plugins/FX/Niagara/Source/Niagara/Public/NiagaraFunctionLibrary.h
--- a/Engine/Plugins/FX/Niagara/Source/Niagara/Public/NiagaraFunctionLibrary.h
+++ b/Engine/Plugins/FX/Niagara/Source/Niagara/Public/NiagaraFunctionLibrary.h
@@ -64,6 +64,14 @@ public:
 	UFUNCTION(BlueprintCallable, Category = Niagara)
 	static void SetSkeletalMeshDataInterfaceSamplingRegions(UNiagaraComponent* NiagaraSystem, const FString& OverrideName, const TArray<FName>& SamplingRegions);
 
+	/** Sets the Filtered Bones to use on the skeletal mesh data interface, this is destructive as it modifies the data interface. */
+	UFUNCTION(BlueprintCallable, Category = Niagara)
+	static void SetSkeletalMeshDataInterfaceFilteredBones(UNiagaraComponent* NiagaraSystem, const FString& OverrideName, const TArray<FName>& FilteredBones);
+
+	/** Sets the Filtered Sockets to use on the skeletal mesh data interface, this is destructive as it modifies the data interface. */
+	UFUNCTION(BlueprintCallable, Category = Niagara)
+	static void SetSkeletalMeshDataInterfaceFilteredSockets(UNiagaraComponent* NiagaraSystem, const FString& OverrideName, const TArray<FName>& FilteredSockets);
+
 	/** Overrides the Texture Object for a Niagara Texture Data Interface User Parameter.*/
 	UFUNCTION(BlueprintCallable, Category = Niagara)
 	static void SetTextureObject(UNiagaraComponent* NiagaraSystem, const FString& OverrideName, UTexture* Texture);
