- Add info text back in for Live vs Baked sim to Baker

#rnx

diff --git a/Engine/Plugins/FX/Niagara/Source/NiagaraEditor/Private/Widgets/SNiagaraBakerViewport.cpp b/Engine/Plugins/FX/Niagara/Source/NiagaraEditor/Private/Widgets/SNiagaraBakerViewport.cpp
--- a/Engine/Plugins/FX/Niagara/Source/NiagaraEditor/Private/Widgets/SNiagaraBakerViewport.cpp
+++ b/Engine/Plugins/FX/Niagara/Source/NiagaraEditor/Private/Widgets/SNiagaraBakerViewport.cpp
@@ -383,10 +383,7 @@ public:
 				ErrorString = TEXT("No output selected or the preview is invalid.");
 			}
 
-			if ( ErrorString.IsSet() )
-			{
-				DrawErrorString(Canvas, PreviewViewRect, ErrorString.GetValue());
-			}
+			DrawViewportText(Canvas, PreviewViewRect, TEXT("Live Sim"), ErrorString.Get(FString()));
 
 			DrawRectBorder(Canvas, PreviewViewRect);
 		}
@@ -420,10 +417,7 @@ public:
 				ErrorString = TEXT("No output or output not generated.\nPlease bake to generate the output");
 			}
 
-			if (ErrorString.IsSet())
-			{
-				DrawErrorString(Canvas, GeneratedViewRect, ErrorString.GetValue());
-			}
+			DrawViewportText(Canvas, GeneratedViewRect, TEXT("Baked Sim"), ErrorString.Get(FString()));
 
 			DrawRectBorder(Canvas, GeneratedViewRect);
 		}
@@ -494,14 +488,33 @@ public:
 		return MaxSize;
 	}
 
-	void DrawErrorString(FCanvas* Canvas, FIntRect ViewportRect, FStringView ErrorString)
+	void DrawViewportText(FCanvas* Canvas, FIntRect ViewportRect, FStringView InfoString, FStringView ErrorString)
 	{
+		// Anything to render?
+		const bool bShowInfoString = WeakViewModel.Pin()->ShowInfoText() && !InfoString.IsEmpty();
+		const bool bShowErrorString = !ErrorString.IsEmpty();
+		if ( !bShowInfoString && !bShowErrorString )
+		{
+			return;
+		}
+
+		// We need to flush to ensure previous data is gone from scene capture rendering
 		Canvas->Flush_GameThread();
 
 		UFont* Font = GetFont();
-		const FVector2f StringSize = GetStringSize(Font, ErrorString.GetData());
-		const FIntPoint TextCenter(ViewportRect.Min.X + (ViewportRect.Width() >> 1), ViewportRect.Min.Y + (ViewportRect.Height() >> 1));
-		Canvas->DrawShadowedString(TextCenter.X - int32(StringSize.X * 0.5f), TextCenter.Y - int32(StringSize.Y * 0.5f), ErrorString.GetData(), Font, FLinearColor::White);
+
+		if (bShowInfoString)
+		{
+			const FIntPoint TextPosition(ViewportRect.Min.X + 3, ViewportRect.Min.Y + 3);
+			Canvas->DrawShadowedString(TextPosition.X, TextPosition.Y, InfoString.GetData(), Font, FLinearColor::White);
+		}
+
+		if (bShowErrorString)
+		{
+			const FVector2f StringSize = GetStringSize(Font, ErrorString.GetData());
+			const FIntPoint TextCenter(ViewportRect.Min.X + (ViewportRect.Width() >> 1), ViewportRect.Min.Y + (ViewportRect.Height() >> 1));
+			Canvas->DrawShadowedString(TextCenter.X - int32(StringSize.X * 0.5f), TextCenter.Y - int32(StringSize.Y * 0.5f), ErrorString.GetData(), Font, FLinearColor::White);
+		}
 
 		Canvas->Flush_GameThread();
 		Canvas->SetRenderTargetScissorRect(FIntRect(0, 0, 0, 0));
