- Handle baking case where asset exists and are in use
- If the package also exists but is not the expected type warn the user and don't overwrite the data

#rnx

diff --git a/Engine/Plugins/FX/Niagara/Source/Niagara/Classes/NiagaraBakerOutput.h b/Engine/Plugins/FX/Niagara/Source/Niagara/Classes/NiagaraBakerOutput.h
--- a/Engine/Plugins/FX/Niagara/Source/Niagara/Classes/NiagaraBakerOutput.h
+++ b/Engine/Plugins/FX/Niagara/Source/Niagara/Classes/NiagaraBakerOutput.h
@@ -56,6 +56,15 @@ public:
 		const FString AssetPath = AssetFullName + "." + AssetName;
 		return TSoftObjectPtr<TObjectClass>(AssetPath).LoadSynchronous();
 	}
+
+	static UObject* GetOrCreateAsset(const FString& PackagePath, UClass* ObjectClass, UClass* FactoryClass);
+
+	template<typename TObjectClass, typename TFactoryClass>
+	static TObjectClass* GetOrCreateAsset(const FString& PackagePath)
+	{
+		UObject* FoundObject = GetOrCreateAsset(PackagePath, TObjectClass::StaticClass(), TFactoryClass::StaticClass());
+		return FoundObject ? CastChecked<TObjectClass>(FoundObject) : nullptr;
+	}
 #endif
 
 	virtual void PostInitProperties() override;
diff --git a/Engine/Plugins/FX/Niagara/Source/Niagara/Private/NiagaraBakerOutput.cpp b/Engine/Plugins/FX/Niagara/Source/Niagara/Private/NiagaraBakerOutput.cpp
--- a/Engine/Plugins/FX/Niagara/Source/Niagara/Private/NiagaraBakerOutput.cpp
+++ b/Engine/Plugins/FX/Niagara/Source/Niagara/Private/NiagaraBakerOutput.cpp
@@ -3,6 +3,9 @@
 #include "NiagaraBakerOutput.h"
 #include "NiagaraSystem.h"
 
+#include "AssetRegistry/IAssetRegistry.h"
+#include "AssetToolsModule.h"
+
 FString UNiagaraBakerOutput::MakeOutputName() const
 {
 	return SanitizeOutputName(GetFName().ToString());
@@ -82,6 +85,37 @@ FString UNiagaraBakerOutput::GetExportFolder(FString PathFormat, int32 FrameInde
 	const FString AssetPath = GetExportPath(PathFormat, FrameIndex);
 	return FString(FPathViews::GetPath(AssetPath));
 }
+
+UObject* UNiagaraBakerOutput::GetOrCreateAsset(const FString& PackagePath, UClass* ObjectClass, UClass* FactoryClass)
+{
+	// Find existing
+	IAssetRegistry* AssetRegistry = IAssetRegistry::Get();
+	TArray<FAssetData> FoundAssets;
+	if (AssetRegistry->GetAssetsByPackageName(FName(PackagePath), FoundAssets))
+	{
+		if (FoundAssets.Num() > 0)
+		{
+			if (UObject* ExistingOject = StaticLoadObject(ObjectClass, nullptr, *PackagePath))
+			{
+				return ExistingOject;
+			}
+
+			// If the above failed then the asset is not the right type, warn the user
+			FText ErrorMessage = FText::Format(
+				NSLOCTEXT("NiagarBaker", "GetOrCreateAsset_PackageExistsOfWrongType", "Could not bake asset '{0}' as package exists but is not a {1}.\nPlease delete the asset or the output to a different location."),
+				FText::FromString(PackagePath),
+				FText::FromName(ObjectClass->GetFName())
+			);
+			FMessageDialog::Open(EAppMsgType::Ok, ErrorMessage);
+			return nullptr;
+		}
+	}
+
+	// Create new
+	IAssetTools& AssetTools = FModuleManager::Get().LoadModuleChecked<FAssetToolsModule>("AssetTools").Get();
+	UFactory* NewFactory = NewObject<UFactory>(GetTransientPackage(), FactoryClass);
+	return AssetTools.CreateAsset(FString(FPathViews::GetCleanFilename(PackagePath)), FString(FPathViews::GetPath(PackagePath)), ObjectClass, NewFactory);
+}
 #endif
 
 void UNiagaraBakerOutput::PostInitProperties()
diff --git a/Engine/Plugins/FX/Niagara/Source/NiagaraEditor/Private/NiagaraBakerRendererOutputSimCache.cpp b/Engine/Plugins/FX/Niagara/Source/NiagaraEditor/Private/NiagaraBakerRendererOutputSimCache.cpp
--- a/Engine/Plugins/FX/Niagara/Source/NiagaraEditor/Private/NiagaraBakerRendererOutputSimCache.cpp
+++ b/Engine/Plugins/FX/Niagara/Source/NiagaraEditor/Private/NiagaraBakerRendererOutputSimCache.cpp
@@ -8,8 +8,6 @@
 #include "NiagaraSimCacheFactoryNew.h"
 #include "NiagaraSystemInstance.h"
 
-#include "AssetToolsModule.h"
-
 TArray<FNiagaraBakerOutputBinding> FNiagaraBakerRendererOutputSimCache::GetRendererBindings(UNiagaraBakerOutput* InBakerOutput) const
 {
 	return TArray<FNiagaraBakerOutputBinding>();
@@ -54,12 +52,7 @@ bool FNiagaraBakerRendererOutputSimCache::BeginBake(UNiagaraBakerOutput* InBaker
 	check(BakerSettings && BakerOutput);
 
 	const FString AssetFullName = BakerOutput->GetAssetPath(BakerOutput->SimCacheAssetPathFormat, 0);
-	const FString AssetName = FString(FPathViews::GetCleanFilename(AssetFullName));
-	const FString PackagePath = FString(FPathViews::GetPath(AssetFullName));
-
-	IAssetTools& AssetTools = FModuleManager::Get().LoadModuleChecked<FAssetToolsModule>("AssetTools").Get();
-	UNiagaraSimCacheFactoryNew* SimCacheFactory = NewObject<UNiagaraSimCacheFactoryNew>();
-	BakeSimCache = Cast<UNiagaraSimCache>(AssetTools.CreateAsset(AssetName, PackagePath, UNiagaraSimCache::StaticClass(), SimCacheFactory));
+	BakeSimCache = UNiagaraBakerOutput::GetOrCreateAsset<UNiagaraSimCache, UNiagaraSimCacheFactoryNew>(AssetFullName);
 	if (BakeSimCache == nullptr)
 	{
 		return false;
diff --git a/Engine/Plugins/FX/Niagara/Source/NiagaraEditor/Private/NiagaraBakerRendererOutputTexture2D.cpp b/Engine/Plugins/FX/Niagara/Source/NiagaraEditor/Private/NiagaraBakerRendererOutputTexture2D.cpp
--- a/Engine/Plugins/FX/Niagara/Source/NiagaraEditor/Private/NiagaraBakerRendererOutputTexture2D.cpp
+++ b/Engine/Plugins/FX/Niagara/Source/NiagaraEditor/Private/NiagaraBakerRendererOutputTexture2D.cpp
@@ -5,7 +5,6 @@
 
 #include "Engine/Texture2D.h"
 #include "Factories/Texture2dFactoryNew.h"
-#include "AssetToolsModule.h"
 #include "CanvasTypes.h"
 #include "CanvasItem.h"
 
@@ -170,13 +169,7 @@ void FNiagaraBakerRendererOutputTexture2D::BakeFrame(UNiagaraBakerOutput* InBake
 	if ( BakerOutput->bGenerateFrames )
 	{
 		const FString AssetFullName = BakerOutput->GetAssetPath(BakerOutput->FramesAssetPathFormat, FrameIndex);
-		const FString AssetName = FString(FPathViews::GetCleanFilename(AssetFullName));
-		const FString PackagePath = FString(FPathViews::GetPath(AssetFullName));
-
-		IAssetTools& AssetTools = FModuleManager::Get().LoadModuleChecked<FAssetToolsModule>("AssetTools").Get();
-		UTexture2DFactoryNew* TextureFactory = NewObject<UTexture2DFactoryNew>();
-		UTexture2D* OutputTexture = Cast<UTexture2D>(AssetTools.CreateAsset(AssetName, PackagePath, UTexture2D::StaticClass(), TextureFactory));
-		if (OutputTexture)
+		if (UTexture2D* OutputTexture = UNiagaraBakerOutput::GetOrCreateAsset<UTexture2D, UTexture2DFactoryNew>(AssetFullName))
 		{
 			const bool bIsPoT = FMath::IsPowerOfTwo(BakerOutput->FrameSize.X) && FMath::IsPowerOfTwo(BakerOutput->FrameSize.Y);
 
@@ -206,13 +199,7 @@ void FNiagaraBakerRendererOutputTexture2D::EndBake(UNiagaraBakerOutput* InBakerO
 	if ( BakerOutput->bGenerateAtlas )
 	{
 		const FString AssetFullName = BakerOutput->GetAssetPath(BakerOutput->AtlasAssetPathFormat, 0);
-		const FString AssetName = FString(FPathViews::GetCleanFilename(AssetFullName));
-		const FString PackagePath = FString(FPathViews::GetPath(AssetFullName));
-
-		IAssetTools& AssetTools = FModuleManager::Get().LoadModuleChecked<FAssetToolsModule>("AssetTools").Get();
-		UTexture2DFactoryNew* TextureFactory = NewObject<UTexture2DFactoryNew>();
-		UTexture2D* OutputTexture = Cast<UTexture2D>(AssetTools.CreateAsset(AssetName, PackagePath, UTexture2D::StaticClass(), TextureFactory));
-		if (OutputTexture)
+		if (UTexture2D* OutputTexture = UNiagaraBakerOutput::GetOrCreateAsset<UTexture2D, UTexture2DFactoryNew>(AssetFullName))
 		{
 			const bool bIsPoT = FMath::IsPowerOfTwo(BakerOutput->AtlasTextureSize.X) && FMath::IsPowerOfTwo(BakerOutput->AtlasTextureSize.Y);
 
diff --git a/Engine/Plugins/FX/Niagara/Source/NiagaraEditor/Private/NiagaraBakerRendererOutputVolumeTexture.cpp b/Engine/Plugins/FX/Niagara/Source/NiagaraEditor/Private/NiagaraBakerRendererOutputVolumeTexture.cpp
--- a/Engine/Plugins/FX/Niagara/Source/NiagaraEditor/Private/NiagaraBakerRendererOutputVolumeTexture.cpp
+++ b/Engine/Plugins/FX/Niagara/Source/NiagaraEditor/Private/NiagaraBakerRendererOutputVolumeTexture.cpp
@@ -16,8 +16,6 @@
 #include "Misc/PathViews.h"
 #include "Modules/ModuleManager.h"
 #include "Factories/VolumeTextureFactory.h"
-#include "IAssetTools.h"
-#include "AssetToolsModule.h"
 #include "UObject/UObjectGlobals.h"
 
 #include "VolumeCache.h"
@@ -558,15 +556,7 @@ void FNiagaraBakerRendererOutputVolumeTexture::BakeFrame(UNiagaraBakerOutput* In
 		if (OutputVolumeTexture->bGenerateFrames)
 		{
 			const FString AssetFullName	= OutputVolumeTexture->GetAssetPath(OutputVolumeTexture->FramesAssetPathFormat, FrameIndex);
-			const FString AssetName		= FString(FPathViews::GetCleanFilename(AssetFullName));
-			const FString PackagePath	= FString(FPathViews::GetPath(AssetFullName));
-
-			//-TODO: Handle package existing
-
-			IAssetTools& AssetTools = FModuleManager::Get().LoadModuleChecked<FAssetToolsModule>("AssetTools").Get();
-			UVolumeTextureFactory* VolumeTextureFactory = NewObject<UVolumeTextureFactory>();
-			UVolumeTexture* OutputTexture = Cast<UVolumeTexture>(AssetTools.CreateAsset(AssetName, PackagePath, UVolumeTexture::StaticClass(), VolumeTextureFactory));
-			if (OutputTexture)
+			if (UVolumeTexture* OutputTexture = UNiagaraBakerOutput::GetOrCreateAsset<UVolumeTexture, UVolumeTextureFactory>(AssetFullName))
 			{
 				OutputTexture->Source.Init(TextureSize.X, TextureSize.Y, TextureSize.Z, 1, TSF_RGBA16F, (const uint8*)(TextureData.GetData()));
 				OutputTexture->PowerOfTwoMode = ETexturePowerOfTwoSetting::None;
@@ -617,15 +607,7 @@ void FNiagaraBakerRendererOutputVolumeTexture::EndBake(UNiagaraBakerOutput* InBa
 	if (BakeAtlasTextureData.Num() > 0)
 	{
 		const FString AssetFullName = OutputVolumeTexture->GetAssetPath(OutputVolumeTexture->AtlasAssetPathFormat, 0);
-		const FString AssetName = FString(FPathViews::GetCleanFilename(AssetFullName));
-		const FString PackagePath = FString(FPathViews::GetPath(AssetFullName));
-
-		//-TODO: Handle package existing
-
-		IAssetTools& AssetTools = FModuleManager::Get().LoadModuleChecked<FAssetToolsModule>("AssetTools").Get();
-		UVolumeTextureFactory* VolumeTextureFactory = NewObject<UVolumeTextureFactory>();
-		UVolumeTexture* OutputTexture = Cast<UVolumeTexture>(AssetTools.CreateAsset(AssetName, PackagePath, UVolumeTexture::StaticClass(), VolumeTextureFactory));
-		if (OutputTexture)
+		if (UVolumeTexture* OutputTexture = UNiagaraBakerOutput::GetOrCreateAsset<UVolumeTexture, UVolumeTextureFactory>(AssetFullName))
 		{
 			OutputTexture->Source.Init(BakeAtlasTextureSize.X, BakeAtlasTextureSize.Y, BakeAtlasTextureSize.Z, 1, TSF_RGBA16F, (const uint8*)(BakeAtlasTextureData.GetData()));
 			OutputTexture->PowerOfTwoMode = ETexturePowerOfTwoSetting::None;
@@ -644,61 +626,49 @@ void FNiagaraBakerRendererOutputVolumeTexture::EndBake(UNiagaraBakerOutput* InBa
 	if (OutputVolumeTexture->bExportFrames)
 	{
 		const FString AssetFullName = OutputVolumeTexture->GetAssetPath(OutputVolumeTexture->AtlasAssetPathFormat, 0);
-		const FString AssetName = FString(FPathViews::GetCleanFilename(AssetFullName));
-		const FString PackagePath = FString(FPathViews::GetPath(AssetFullName));
-		IAssetTools& AssetTools = FModuleManager::Get().LoadModuleChecked<FAssetToolsModule>("AssetTools").Get();
-		
-		// see if the asset has been created
-		UPackage* Pkg = CreatePackage(*AssetFullName);		
-		UVolumeCache* OutputCache = Cast<UVolumeCache>(StaticFindObject(UObject::StaticClass(), Pkg, *AssetName));
-		
-		if (OutputCache == nullptr)
+		if (UVolumeCache* OutputCache = UNiagaraBakerOutput::GetOrCreateAsset<UVolumeCache, UVolumeCacheFactory>(AssetFullName))
 		{
-			UVolumeCacheFactory* VolumeCacheFactory = NewObject<UVolumeCacheFactory>();
-			OutputCache = Cast<UVolumeCache>(AssetTools.CreateAsset(AssetName, PackagePath, UVolumeCache::StaticClass(), VolumeCacheFactory));
-		}
+			FString FullPath = OutputVolumeTexture->GetExportPath(OutputVolumeTexture->FramesExportPathFormat, 0);
+			FullPath = FString(FPathViews::GetPath(FullPath));
 
-		FString FullPath = OutputVolumeTexture->GetExportPath(OutputVolumeTexture->FramesExportPathFormat, 0);
-		FullPath = FString(FPathViews::GetPath(FullPath));
+			FString FullFileName = FString(FPathViews::GetCleanFilename(OutputVolumeTexture->FramesExportPathFormat));
 
-		FString FullFileName = FString(FPathViews::GetCleanFilename(OutputVolumeTexture->FramesExportPathFormat));
+			FullPath = FullPath + "/" + FullFileName;
 
-		FullPath = FullPath + "/" + FullFileName;
-
-		OutputCache->FilePath = FullPath;
+			OutputCache->FilePath = FullPath;
 		
-		OutputCache->FrameRangeStart = BakedFrameRange.X;
-		OutputCache->FrameRangeEnd = BakedFrameRange.Y;
-		OutputCache->Resolution = BakedFrameSize;
-		OutputCache->CacheType = EVolumeCacheType::OpenVDB;
+			OutputCache->FrameRangeStart = BakedFrameRange.X;
+			OutputCache->FrameRangeEnd = BakedFrameRange.Y;
+			OutputCache->Resolution = BakedFrameSize;
+			OutputCache->CacheType = EVolumeCacheType::OpenVDB;
 
-		OutputCache->Modify();
-		OutputCache->PostEditChange();		
+			OutputCache->Modify();
+			OutputCache->PostEditChange();		
 
-		if (BakedFrameRange.Y - BakedFrameRange.X > 0)
-		{
-			if (BakedFrameSize.X + BakedFrameSize.Y + BakedFrameSize.Z > 0)
+			if (BakedFrameRange.Y - BakedFrameRange.X > 0)
 			{
-
-				OutputCache->InitData();
+				if (BakedFrameSize.X + BakedFrameSize.Y + BakedFrameSize.Z > 0)
+				{
+					OutputCache->InitData();
 		
-				// @todo: could be cool to have an option where the whole cache is read into memory after writing it so you can preview results super quick
-				// bool LoadedCache = OutputCache->LoadRange();
+					// @todo: could be cool to have an option where the whole cache is read into memory after writing it so you can preview results super quick
+					// bool LoadedCache = OutputCache->LoadRange();
 
-				// if (!LoadedCache)
-				// {
-				// 	UE_LOG(LogNiagaraBaker, Warning, TEXT("Could not load cache into memory: %s"), *FullPath);
-				// }
+					// if (!LoadedCache)
+					// {
+					// 	UE_LOG(LogNiagaraBaker, Warning, TEXT("Could not load cache into memory: %s"), *FullPath);
+					// }
+				}
+				else
+				{
+					UE_LOG(LogNiagaraBaker, Warning, TEXT("Output cache has no voxels : %s"), *FullPath);
+				}
 			}
 			else
 			{
-				UE_LOG(LogNiagaraBaker, Warning, TEXT("Output cache has no voxels : %s"), *FullPath);
+				UE_LOG(LogNiagaraBaker, Warning, TEXT("No frames written : %s"), *FullPath);
 			}
 		}
-		else
-		{
-			UE_LOG(LogNiagaraBaker, Warning, TEXT("No frames written : %s"), *FullPath);
-		}
 	}
 }
 
