- Wrap editor includes with define

#rnx

diff --git a/Engine/Plugins/FX/Niagara/Source/Niagara/Private/NiagaraBakerOutput.cpp b/Engine/Plugins/FX/Niagara/Source/Niagara/Private/NiagaraBakerOutput.cpp
--- a/Engine/Plugins/FX/Niagara/Source/Niagara/Private/NiagaraBakerOutput.cpp
+++ b/Engine/Plugins/FX/Niagara/Source/Niagara/Private/NiagaraBakerOutput.cpp
@@ -3,8 +3,12 @@
 #include "NiagaraBakerOutput.h"
 #include "NiagaraSystem.h"
 
+#if WITH_EDITOR
 #include "AssetRegistry/IAssetRegistry.h"
+#include "Factories/Factory.h"
+#include "Misc/MessageDialog.h"
 #include "AssetToolsModule.h"
+#endif //WITH_EDITOR
 
 FString UNiagaraBakerOutput::MakeOutputName() const
 {
