Sim Cache UI Refactor

Moving to MVVM for the sim cache spreadsheet view.
Removing the details customization and creating an editor with a toolkit and asset type actions.
Scrolling on the spreadsheet now behaves correctly.

diff --git a/Engine/Plugins/FX/Niagara/Source/NiagaraEditor/Private/AssetTypeActions/AssetTypeActions_NiagaraSimCache.cpp b/Engine/Plugins/FX/Niagara/Source/NiagaraEditor/Private/AssetTypeActions/AssetTypeActions_NiagaraSimCache.cpp
new file mode 100644
--- /dev/null
+++ b/Engine/Plugins/FX/Niagara/Source/NiagaraEditor/Private/AssetTypeActions/AssetTypeActions_NiagaraSimCache.cpp
@@ -0,0 +1,32 @@
+// Copyright Epic Games, Inc. All Rights Reserved.
+
+#include "AssetTypeActions_NiagaraSimCache.h"
+
+#include "NiagaraEditorStyle.h"
+#include "NiagaraSimCacheToolkit.h"
+
+FColor FAssetTypeActions_NiagaraSimCache::GetTypeColor() const
+{
+	return FNiagaraEditorStyle::Get().GetColor("NiagaraEditor.AssetColors.SimCache").ToFColor(true);
+}
+
+UClass* FAssetTypeActions_NiagaraSimCache::GetSupportedClass() const
+{
+	return UNiagaraSimCache::StaticClass();
+}
+
+void FAssetTypeActions_NiagaraSimCache::OpenAssetEditor(const TArray<UObject*>& InObjects,
+	TSharedPtr<IToolkitHost> EditWithinLevelEditor)
+{
+	const EToolkitMode::Type Mode = EditWithinLevelEditor.IsValid() ? EToolkitMode::WorldCentric : EToolkitMode::Standalone;
+
+	for (auto ObjIt = InObjects.CreateConstIterator(); ObjIt; ++ObjIt)
+	{
+		const auto SimCache = Cast<UNiagaraSimCache>(*ObjIt);
+		if (SimCache != nullptr)
+		{
+			const TSharedRef< FNiagaraSimCacheToolkit > NewNiagaraSimCacheToolkit(new FNiagaraSimCacheToolkit());
+			NewNiagaraSimCacheToolkit->Initialize(Mode, EditWithinLevelEditor, SimCache);
+		}
+	}
+}
diff --git a/Engine/Plugins/FX/Niagara/Source/NiagaraEditor/Private/AssetTypeActions/AssetTypeActions_NiagaraSimCache.h b/Engine/Plugins/FX/Niagara/Source/NiagaraEditor/Private/AssetTypeActions/AssetTypeActions_NiagaraSimCache.h
new file mode 100644
--- /dev/null
+++ b/Engine/Plugins/FX/Niagara/Source/NiagaraEditor/Private/AssetTypeActions/AssetTypeActions_NiagaraSimCache.h
@@ -0,0 +1,19 @@
+// Copyright Epic Games, Inc. All Rights Reserved.
+
+#pragma once
+
+#include "CoreMinimal.h"
+#include "Toolkits/IToolkitHost.h"
+#include "AssetTypeActions_Base.h"
+#include "NiagaraEditorModule.h"
+
+class FAssetTypeActions_NiagaraSimCache : public FAssetTypeActions_Base
+{
+public:
+	// IAssetTypeActions Implementation
+	virtual FText GetName() const override { return NSLOCTEXT("AssetTypeActions", "AssetTypeActions_NiagaraSimCache", "Niagara Simulation Cache"); }
+	virtual FColor GetTypeColor() const override;
+	virtual UClass* GetSupportedClass() const override;
+	virtual void OpenAssetEditor(const TArray<UObject*>& InObjects, TSharedPtr<class IToolkitHost> EditWithinLevelEditor = TSharedPtr<IToolkitHost>()) override;
+	virtual uint32 GetCategories() override { return FNiagaraEditorModule::GetAssetCategory(); }
+};
\ No newline at end of file
diff --git a/Engine/Plugins/FX/Niagara/Source/NiagaraEditor/Private/NiagaraEditorModule.cpp b/Engine/Plugins/FX/Niagara/Source/NiagaraEditor/Private/NiagaraEditorModule.cpp
--- a/Engine/Plugins/FX/Niagara/Source/NiagaraEditor/Private/NiagaraEditorModule.cpp
+++ b/Engine/Plugins/FX/Niagara/Source/NiagaraEditor/Private/NiagaraEditorModule.cpp
@@ -21,6 +21,7 @@
 #include "AssetTypeActions/AssetTypeActions_NiagaraParameterCollection.h"
 #include "AssetTypeActions/AssetTypeActions_NiagaraEffectType.h"
 #include "AssetTypeActions/AssetTypeActions_NiagaraParameterDefinitions.h"
+#include "AssetTypeActions/AssetTypeActions_NiagaraSimCache.h"
 
 #include "EdGraphUtilities.h"
 #include "SGraphPin.h"
@@ -847,6 +848,7 @@ void FNiagaraEditorModule::StartupModule()
 	RegisterAssetTypeAction(AssetTools, MakeShareable(new FAssetTypeActions_NiagaraParameterCollectionInstance()));
 	RegisterAssetTypeAction(AssetTools, MakeShareable(new FAssetTypeActions_NiagaraParameterDefinitions()));
 	RegisterAssetTypeAction(AssetTools, MakeShareable(new FAssetTypeActions_NiagaraEffectType()));
+	RegisterAssetTypeAction(AssetTools, MakeShareable(new FAssetTypeActions_NiagaraSimCache()));
 
 	// Preload all parameter definition & collection assets so that they will be postloaded before postload calls to scripts/emitters/systems that rely on them.
 	{
diff --git a/Engine/Plugins/FX/Niagara/Source/NiagaraEditor/Private/NiagaraEditorStyle.cpp b/Engine/Plugins/FX/Niagara/Source/NiagaraEditor/Private/NiagaraEditorStyle.cpp
--- a/Engine/Plugins/FX/Niagara/Source/NiagaraEditor/Private/NiagaraEditorStyle.cpp
+++ b/Engine/Plugins/FX/Niagara/Source/NiagaraEditor/Private/NiagaraEditorStyle.cpp
@@ -419,6 +419,7 @@ void FNiagaraEditorStyle::InitAssetColors()
 	Set("NiagaraEditor.AssetColors.ParameterCollectionInstance", FLinearColor(1.0f, 1.0f, 0.7f));
 	Set("NiagaraEditor.AssetColors.ParameterDefinitions", FLinearColor(0.57f, 0.82f, 0.06f));
 	Set("NiagaraEditor.AssetColors.EffectType", FLinearColor(1.f, 1.f, 1.f));
+	Set("NiagaraEditor.AssetColors.SimCache", FLinearColor(0.9f, 0.3f, 0.25f));
 }
 
 void FNiagaraEditorStyle::InitThumbnails()
diff --git a/Engine/Plugins/FX/Niagara/Source/NiagaraEditor/Private/Toolkits/NiagaraSimCacheToolkit.cpp b/Engine/Plugins/FX/Niagara/Source/NiagaraEditor/Private/Toolkits/NiagaraSimCacheToolkit.cpp
new file mode 100644
--- /dev/null
+++ b/Engine/Plugins/FX/Niagara/Source/NiagaraEditor/Private/Toolkits/NiagaraSimCacheToolkit.cpp
@@ -0,0 +1,119 @@
+// Copyright Epic Games, Inc. All Rights Reserved.
+
+#include "NiagaraSimCacheToolkit.h"
+#include "NiagaraEditorModule.h"
+
+// Widgets & ViewModels
+#include "Widgets/Docking/SDockTab.h"
+#include "NiagaraSimCacheViewModel.h"
+
+#define LOCTEXT_NAMESPACE "NiagaraSimCacheToolkit"
+
+const FName FNiagaraSimCacheToolkit::NiagaraSimCacheSpreadsheetTabId(TEXT("NiagaraSimCacheEditor_Spreadsheet"));
+
+//TODO Timeline, & Preview
+
+FNiagaraSimCacheToolkit::FNiagaraSimCacheToolkit()
+{
+		
+}
+
+FNiagaraSimCacheToolkit::~FNiagaraSimCacheToolkit()
+{
+
+}
+
+void FNiagaraSimCacheToolkit::RegisterTabSpawners(const TSharedRef<class FTabManager>& InTabManager)
+{
+	WorkspaceMenuCategory = InTabManager->AddLocalWorkspaceMenuCategory(LOCTEXT("WorkspaceMenu_NiagaraSimCacheEditor", "Niagara"));
+	
+	FAssetEditorToolkit::RegisterTabSpawners(InTabManager);
+
+	InTabManager->RegisterTabSpawner(NiagaraSimCacheSpreadsheetTabId, FOnSpawnTab::CreateSP(this, &FNiagaraSimCacheToolkit::SpawnTab_SimCacheSpreadsheet))
+		.SetDisplayName(LOCTEXT("SpreadsheetTab", "Cache Spreadsheet"))
+		.SetGroup(WorkspaceMenuCategory.ToSharedRef());
+}
+
+void FNiagaraSimCacheToolkit::UnregisterTabSpawners(const TSharedRef<class FTabManager>& InTabManager)
+{
+	FAssetEditorToolkit::UnregisterTabSpawners(InTabManager);
+
+	InTabManager->UnregisterTabSpawner(NiagaraSimCacheSpreadsheetTabId);
+}
+
+void FNiagaraSimCacheToolkit::Initialize(const EToolkitMode::Type Mode, const TSharedPtr< class IToolkitHost >& InitToolkitHost, UNiagaraSimCache* InSimCache)
+{
+	SimCache = InSimCache;
+	if (SimCache.IsValid())
+	{
+		SimCacheViewModel = MakeShared<FNiagaraSimCacheViewModel>();
+		
+		SimCacheViewModel->Initialize(SimCache);
+		
+
+		const TSharedRef<FTabManager::FLayout> StandaloneDefaultLayout = FTabManager::NewLayout("Standalone_Niagara_SimCache_Layout_V1")
+			->AddArea
+			(
+				FTabManager::NewPrimaryArea()->SetOrientation(Orient_Vertical)
+				->Split
+				(
+					FTabManager::NewStack()
+					->SetSizeCoefficient(0.1f)
+					->AddTab(NiagaraSimCacheSpreadsheetTabId, ETabState::OpenedTab)
+				)
+			);
+
+		constexpr bool bCreateDefaultStandaloneMenu = true;
+		constexpr bool bCreateDefaultToolbar = true;
+		FAssetEditorToolkit::InitAssetEditor(Mode, InitToolkitHost, FNiagaraEditorModule::NiagaraEditorAppIdentifier, StandaloneDefaultLayout, bCreateDefaultStandaloneMenu, bCreateDefaultToolbar, SimCache.Get());
+
+		RegenerateMenusAndToolbars();
+	}
+}
+
+FName FNiagaraSimCacheToolkit::GetToolkitFName() const
+{
+	return FName("Niagara");
+}
+
+FText FNiagaraSimCacheToolkit::GetBaseToolkitName() const
+{
+	return LOCTEXT("AppLabel", "Niagara");
+}
+
+FString FNiagaraSimCacheToolkit::GetWorldCentricTabPrefix() const
+{
+	return LOCTEXT("WorldCentricTabPrefix", "Niagara ").ToString();
+}
+
+FLinearColor FNiagaraSimCacheToolkit::GetWorldCentricTabColorScale() const
+{
+	return FNiagaraEditorModule::WorldCentricTabColorScale;
+}
+
+TSharedRef<SDockTab> FNiagaraSimCacheToolkit::SpawnTab_SimCacheSpreadsheet(const FSpawnTabArgs& Args)
+{
+	check(Args.GetTabId().TabType == NiagaraSimCacheSpreadsheetTabId);
+
+	SimCacheSpreadsheetView =
+		SNew(SNiagaraSimCacheView)
+		.SimCacheViewModel(SimCacheViewModel);
+
+	const TSharedRef<SVerticalBox> Contents = SNew(SVerticalBox);
+
+	Contents->AddSlot()
+		.FillHeight(1.0f)
+		[
+			SimCacheSpreadsheetView->AsShared()
+		];
+
+	TSharedRef<SDockTab> SpawnedTab =
+		SNew(SDockTab)
+		[
+			Contents
+		];
+	
+	return SpawnedTab;
+}
+
+#undef LOCTEXT_NAMESPACE /** NiagaraSimCacheToolkit **/
\ No newline at end of file
diff --git a/Engine/Plugins/FX/Niagara/Source/NiagaraEditor/Private/Toolkits/NiagaraSimCacheToolkit.h b/Engine/Plugins/FX/Niagara/Source/NiagaraEditor/Private/Toolkits/NiagaraSimCacheToolkit.h
new file mode 100644
--- /dev/null
+++ b/Engine/Plugins/FX/Niagara/Source/NiagaraEditor/Private/Toolkits/NiagaraSimCacheToolkit.h
@@ -0,0 +1,51 @@
+// Copyright Epic Games, Inc. All Rights Reserved.
+
+#pragma once
+
+#include "CoreMinimal.h"
+#include "SNiagaraSimCacheView.h"
+#include "Toolkits/IToolkitHost.h"
+#include "Toolkits/AssetEditorToolkit.h"
+
+class FNiagaraSimCacheViewModel;
+class UNiagaraSimCache;
+
+// Editor for Niagara Simulation Caches
+class FNiagaraSimCacheToolkit : public FAssetEditorToolkit
+{
+public:
+	FNiagaraSimCacheToolkit();
+	~FNiagaraSimCacheToolkit();
+
+	virtual void RegisterTabSpawners(const TSharedRef<class FTabManager>& InTabManager) override;
+	virtual void UnregisterTabSpawners(const TSharedRef<class FTabManager>& InTabManager) override;
+
+	// Initialize with a sim cache to view. 
+	void Initialize(const EToolkitMode::Type Mode, const TSharedPtr< class IToolkitHost >& InitToolkitHost, UNiagaraSimCache* InSimCache);
+
+	//~ Begin IToolkit Interface
+	virtual FName GetToolkitFName() const override;
+	virtual FText GetBaseToolkitName() const override;
+	virtual FString GetWorldCentricTabPrefix() const override;
+	virtual FLinearColor GetWorldCentricTabColorScale() const override;
+	//~ End IToolkit Interface
+
+protected:
+	/** The workspace menu category of this toolkit */
+	TSharedPtr<FWorkspaceItem> WorkspaceMenuCategory;
+	
+
+private:
+	TSharedRef<SDockTab> SpawnTab_SimCacheSpreadsheet(const FSpawnTabArgs& Args);
+
+	// The sim cache being viewed.
+	TWeakObjectPtr <UNiagaraSimCache> SimCache;
+
+	// The view model for the sim cache
+	TSharedPtr<FNiagaraSimCacheViewModel> SimCacheViewModel;
+
+	// Spreadsheet widget
+	TSharedPtr<SNiagaraSimCacheView> SimCacheSpreadsheetView;
+
+	static const FName NiagaraSimCacheSpreadsheetTabId;
+};
\ No newline at end of file
diff --git a/Engine/Plugins/FX/Niagara/Source/NiagaraEditor/Private/ViewModels/NiagaraSimCacheViewModel.cpp b/Engine/Plugins/FX/Niagara/Source/NiagaraEditor/Private/ViewModels/NiagaraSimCacheViewModel.cpp
new file mode 100644
--- /dev/null
+++ b/Engine/Plugins/FX/Niagara/Source/NiagaraEditor/Private/ViewModels/NiagaraSimCacheViewModel.cpp
@@ -0,0 +1,304 @@
+// Copyright Epic Games, Inc. All Rights Reserved.
+
+#include "ViewModels/NiagaraSimCacheViewModel.h"
+#include "NiagaraSimCache.h"
+#include "NiagaraEditorCommon.h"
+#include "NiagaraDebuggerCommon.h"
+#include "Serialization/ObjectAndNameAsStringProxyArchive.h"
+
+#define LOCTEXT_NAMESPACE "NiagaraSimCacheViewModel"
+
+struct FSimCacheBufferReader
+{
+
+	explicit FSimCacheBufferReader(TWeakObjectPtr<UNiagaraSimCache> SimCache, uint32 InFrameIndex = 0, uint32 InEmitterIndex = INDEX_NONE)
+		: WeakSimCache(SimCache)
+	{
+		UpdateLayout(InFrameIndex, InEmitterIndex);
+	}
+
+	void UpdateLayout(const uint32 FrameIndex, const uint32 EmitterIndex)
+	{
+		ComponentInfos.Empty();
+		FoundFloats = 0;
+		FoundHalfs = 0;
+		FoundInt32s = 0;
+
+		if (const FNiagaraSimCacheDataBuffersLayout* SimCacheLayout = GetSimCacheBufferLayout(FrameIndex, EmitterIndex))
+		{
+			for (const FNiagaraSimCacheVariable& Variable : SimCacheLayout->Variables)
+			{
+				const FNiagaraTypeDefinition& TypeDef = Variable.Variable.GetType();
+				if (TypeDef.IsEnum())
+				{
+					FNiagaraSimCacheComponentInfo& ComponentInfo = ComponentInfos.AddDefaulted_GetRef();
+					ComponentInfo.Name = Variable.Variable.GetName();
+					ComponentInfo.ComponentOffset = FoundInt32s++;
+					ComponentInfo.bIsInt32 = true;
+					ComponentInfo.Enum = TypeDef.GetEnum();
+				}
+				else
+				{
+					Build(Variable.Variable.GetName(), TypeDef.GetScriptStruct());
+				}
+			}
+
+			if ((FoundFloats != SimCacheLayout->FloatCount) ||
+				(FoundHalfs != SimCacheLayout->HalfCount) ||
+				(FoundInt32s != SimCacheLayout->Int32Count))
+			{
+				UE_LOG(LogNiagaraEditor, Warning, TEXT("SimCache Layout doesn't appear to match iterating the variables"));
+				ComponentInfos.Empty();
+			}
+		}
+	}
+
+	TConstArrayView<FNiagaraSimCacheComponentInfo> GetComponentInfos() const
+	{
+		return MakeArrayView(ComponentInfos);
+	}
+	
+	UNiagaraSimCache* GetSimCache() const
+	{
+		return WeakSimCache.Get();
+	}
+
+	const FNiagaraSimCacheDataBuffersLayout* GetSimCacheBufferLayout(const uint32 FrameIndex, const uint32 EmitterIndex) const
+	{
+		if (UNiagaraSimCache* const SimCache = WeakSimCache.Get())
+		{
+			if (SimCache->CacheFrames.IsValidIndex(FrameIndex))
+			{
+				if (EmitterIndex == INDEX_NONE)
+				{
+					return &SimCache->CacheLayout.SystemLayout;
+				}
+				else if (SimCache->CacheLayout.EmitterLayouts.IsValidIndex(EmitterIndex))
+				{
+					return &SimCache->CacheLayout.EmitterLayouts[EmitterIndex];
+				}
+			}
+		}
+		return nullptr;
+	}
+
+	const FNiagaraSimCacheDataBuffers* GetSimCacheDataBuffers(const uint32 FrameIndex, const uint32 EmitterIndex) const
+	{
+		if (UNiagaraSimCache* SimCache = WeakSimCache.Get())
+		{
+			if (SimCache->CacheFrames.IsValidIndex(FrameIndex))
+			{
+				if (EmitterIndex == INDEX_NONE)
+				{
+					return &SimCache->CacheFrames[FrameIndex].SystemData.SystemDataBuffers;
+				}
+				else if (SimCache->CacheFrames[FrameIndex].EmitterData.IsValidIndex(EmitterIndex))
+				{
+					return &SimCache->CacheFrames[FrameIndex].EmitterData[EmitterIndex].ParticleDataBuffers;
+				}
+			}
+		}
+		return nullptr;
+	}
+
+	FText GetComponentText(FName ComponentName, int32 InstanceIndex, uint32 FrameIndex, uint32 EmitterIndex) const
+	{
+		const FNiagaraSimCacheComponentInfo* ComponentInfo = ComponentInfos.FindByPredicate([ComponentName](const FNiagaraSimCacheComponentInfo& Info) { return Info.Name == ComponentName; });
+		const FNiagaraSimCacheDataBuffers* SimCacheDataBuffers = GetSimCacheDataBuffers(FrameIndex, EmitterIndex);
+	
+		if (ComponentInfo && SimCacheDataBuffers)
+		{
+			const int32 NumInstances = SimCacheDataBuffers->NumInstances;
+			if (InstanceIndex >= 0 && InstanceIndex < NumInstances)
+			{
+				if (ComponentInfo->bIsFloat)
+				{
+					const float* Value = reinterpret_cast<const float*>(SimCacheDataBuffers->FloatData.GetData()) + (ComponentInfo->ComponentOffset * NumInstances) + InstanceIndex;
+					return FText::AsNumber(*Value);
+				}
+				else if (ComponentInfo->bIsHalf)
+				{
+					const FFloat16* Value = reinterpret_cast<const FFloat16*>(SimCacheDataBuffers->HalfData.GetData()) + (ComponentInfo->ComponentOffset * NumInstances) + InstanceIndex;
+					return FText::AsNumber(Value->GetFloat());
+				}
+				else if (ComponentInfo->bIsInt32)
+				{
+					const int32* Value = reinterpret_cast<const int32*>(SimCacheDataBuffers->Int32Data.GetData()) + (ComponentInfo->ComponentOffset * NumInstances) + InstanceIndex;
+					if (ComponentInfo->bShowAsBool)
+					{
+						return *Value == 0 ? LOCTEXT("False", "False") : LOCTEXT("True", "True");
+					}
+					else if (ComponentInfo->Enum != nullptr)
+					{
+						return ComponentInfo->Enum->GetDisplayNameTextByValue(*Value);
+					}
+					else
+					{
+						return FText::AsNumber(*Value);
+					}
+				}
+			}
+		}
+		return LOCTEXT("Error", "Error");
+	}
+
+protected:
+	TWeakObjectPtr<UNiagaraSimCache>	WeakSimCache;
+
+	TArray<FNiagaraSimCacheComponentInfo>				ComponentInfos;
+	uint32								FoundFloats = 0;
+	uint32								FoundHalfs = 0;
+	uint32								FoundInt32s = 0;
+
+	void Build(const FName Name, const UScriptStruct* Struct)
+	{
+		int32 NumProperties = 0;
+		for (TFieldIterator<FProperty> PropertyIt(Struct, EFieldIteratorFlags::IncludeSuper); PropertyIt; ++PropertyIt)
+		{
+			++NumProperties;
+		}
+
+		for (TFieldIterator<FProperty> PropertyIt(Struct, EFieldIteratorFlags::IncludeSuper); PropertyIt; ++PropertyIt)
+		{
+			FProperty* Property = *PropertyIt;
+			const FName PropertyName = NumProperties > 1 ? FName(*FString::Printf(TEXT("%s.%s"), *Name.ToString(), *Property->GetName())) : Name;
+			if (Property->IsA(FFloatProperty::StaticClass()))
+			{
+				FNiagaraSimCacheComponentInfo& ComponentInfo = ComponentInfos.AddDefaulted_GetRef();
+				ComponentInfo.Name = PropertyName;
+				ComponentInfo.ComponentOffset = FoundFloats++;
+				ComponentInfo.bIsFloat = true;
+			}
+			else if (Property->IsA(FUInt16Property::StaticClass()))
+			{
+				FNiagaraSimCacheComponentInfo& ComponentInfo = ComponentInfos.AddDefaulted_GetRef();
+				ComponentInfo.Name = PropertyName;
+				ComponentInfo.ComponentOffset = FoundHalfs++;
+				ComponentInfo.bIsHalf = true;
+			}
+			else if (Property->IsA(FIntProperty::StaticClass()))
+			{
+				FNiagaraSimCacheComponentInfo& ComponentInfo = ComponentInfos.AddDefaulted_GetRef();
+				ComponentInfo.Name = PropertyName;
+				ComponentInfo.ComponentOffset = FoundInt32s++;
+				ComponentInfo.bIsInt32 = true;
+				ComponentInfo.bShowAsBool = (NumProperties == 1) && (Struct == FNiagaraTypeDefinition::GetBoolStruct());
+			}
+			else if (Property->IsA(FBoolProperty::StaticClass()))
+			{
+				FNiagaraSimCacheComponentInfo& ComponentInfo = ComponentInfos.AddDefaulted_GetRef();
+				ComponentInfo.Name = PropertyName;
+				ComponentInfo.ComponentOffset = FoundInt32s++;
+				ComponentInfo.bIsInt32 = true;
+				ComponentInfo.bShowAsBool = true;
+			}
+			else if (Property->IsA(FEnumProperty::StaticClass()))
+			{
+				FNiagaraSimCacheComponentInfo& ComponentInfo = ComponentInfos.AddDefaulted_GetRef();
+				ComponentInfo.Name = PropertyName;
+				ComponentInfo.ComponentOffset = FoundInt32s++;
+				ComponentInfo.bIsInt32 = true;
+				ComponentInfo.Enum = CastFieldChecked<FEnumProperty>(Property)->GetEnum();
+			}
+			else if (Property->IsA(FStructProperty::StaticClass()))
+			{
+				const FStructProperty* StructProperty = CastFieldChecked<FStructProperty>(Property);
+				Build(PropertyName, FNiagaraTypeHelper::FindNiagaraFriendlyTopLevelStruct(StructProperty->Struct, ENiagaraStructConversion::Simulation));
+			}
+			else
+			{
+				// Fail
+			}
+		}
+	}
+};
+
+TSharedPtr<FSimCacheBufferReader> BufferReader;
+
+FNiagaraSimCacheViewModel::FNiagaraSimCacheViewModel()
+{
+	
+}
+
+FNiagaraSimCacheViewModel::~FNiagaraSimCacheViewModel()
+{
+
+}
+
+void FNiagaraSimCacheViewModel::Initialize(TWeakObjectPtr<UNiagaraSimCache> SimCache) 
+{
+	BufferReader = MakeShared<FSimCacheBufferReader>(SimCache);
+	BufferReader->UpdateLayout(FrameIndex, EmitterIndex);
+}
+
+void FNiagaraSimCacheViewModel::UpdateSimCache(const FNiagaraSystemSimCacheCaptureReply& Reply)
+{
+	UNiagaraSimCache* SimCache = nullptr;
+	
+	if (Reply.SimCacheData.Num() > 0)
+	{
+		SimCache = NewObject<UNiagaraSimCache>();
+
+		FMemoryReader ArReader(Reply.SimCacheData);
+		FObjectAndNameAsStringProxyArchive ProxyArReader(ArReader, false);
+		SimCache->Serialize(ProxyArReader);
+	}
+	else
+	{
+		UE_LOG(LogNiagaraEditor, Warning, TEXT("Debug Spreadsheet recieved empty sim cache data."));
+	}
+	Initialize(SimCache);
+}
+
+FText FNiagaraSimCacheViewModel::GetComponentText(const FName ComponentName, const int32 InstanceIndex) const
+{
+	return BufferReader ? BufferReader->GetComponentText(ComponentName, InstanceIndex, FrameIndex, EmitterIndex) : FText();
+}
+
+TConstArrayView<FNiagaraSimCacheComponentInfo> FNiagaraSimCacheViewModel::GetComponentInfos() const
+{
+	return BufferReader ? BufferReader->GetComponentInfos() : TConstArrayView<FNiagaraSimCacheComponentInfo>();
+}
+
+int32 FNiagaraSimCacheViewModel::GetNumInstances() const
+{
+	const FNiagaraSimCacheDataBuffers* DataBuffer = BufferReader ? BufferReader->GetSimCacheDataBuffers(FrameIndex, EmitterIndex) : nullptr;
+	return DataBuffer ? DataBuffer->NumInstances : 0;
+}
+
+int32 FNiagaraSimCacheViewModel::GetNumFrames() const
+{
+	const UNiagaraSimCache* SimCache = BufferReader ? BufferReader->GetSimCache() : nullptr;
+	return SimCache ? SimCache->CacheFrames.Num() : 0;
+}
+
+
+void FNiagaraSimCacheViewModel::SetEmitterIndex(const int32 InEmitterIndex)
+{
+	EmitterIndex = InEmitterIndex;
+	BufferReader->UpdateLayout(FrameIndex, EmitterIndex);
+}
+
+const FNiagaraSimCacheDataBuffersLayout* FNiagaraSimCacheViewModel::GetSimCacheBufferLayout() const
+{
+	return BufferReader ? BufferReader->GetSimCacheBufferLayout(FrameIndex, EmitterIndex) : nullptr;
+}
+
+bool FNiagaraSimCacheViewModel::IsCacheValid()
+{
+	return BufferReader ? BufferReader->GetSimCache()->IsCacheValid() : false;
+}
+
+int32 FNiagaraSimCacheViewModel::GetNumEmitterLayouts()
+{
+	const UNiagaraSimCache* SimCache = BufferReader ? BufferReader->GetSimCache() : nullptr;
+	return SimCache ? SimCache->CacheLayout.EmitterLayouts.Num() : 0;
+}
+
+FName FNiagaraSimCacheViewModel::GetEmitterLayoutName(const int32 Index)
+{
+	const UNiagaraSimCache* SimCache = BufferReader ? BufferReader->GetSimCache() : nullptr;
+	return SimCache ? SimCache->CacheLayout.EmitterLayouts[Index].LayoutName : NAME_None;
+}
+
+#undef LOCTEXT_NAMESPACE
\ No newline at end of file
diff --git a/Engine/Plugins/FX/Niagara/Source/NiagaraEditor/Private/ViewModels/NiagaraSimCacheViewModel.h b/Engine/Plugins/FX/Niagara/Source/NiagaraEditor/Private/ViewModels/NiagaraSimCacheViewModel.h
new file mode 100644
--- /dev/null
+++ b/Engine/Plugins/FX/Niagara/Source/NiagaraEditor/Private/ViewModels/NiagaraSimCacheViewModel.h
@@ -0,0 +1,59 @@
+// Copyright Epic Games, Inc. All Rights Reserved.
+
+#pragma once
+
+#include "CoreMinimal.h"
+
+class UNiagaraSimCache;
+struct FNiagaraSystemSimCacheCaptureReply;
+struct FNiagaraSimCacheDataBuffersLayout;
+
+struct FNiagaraSimCacheComponentInfo
+{
+	FName Name = NAME_None;
+	uint32 ComponentOffset = INDEX_NONE;
+	bool bIsFloat = false;
+	bool bIsHalf = false;
+	bool bIsInt32 = false;
+	bool bShowAsBool = false;
+	UEnum* Enum = nullptr;
+};
+
+class NIAGARAEDITOR_API FNiagaraSimCacheViewModel : public TSharedFromThis<FNiagaraSimCacheViewModel>
+{
+public:
+	FNiagaraSimCacheViewModel();
+	~FNiagaraSimCacheViewModel();
+
+	void Initialize(TWeakObjectPtr<UNiagaraSimCache> SimCache);
+
+	void UpdateSimCache(const FNiagaraSystemSimCacheCaptureReply& Reply);
+
+	TConstArrayView<FNiagaraSimCacheComponentInfo> GetComponentInfos() const;
+
+	int32 GetNumInstances() const;
+
+	int32 GetNumFrames() const;
+
+	int32 GetFrameIndex() const { return FrameIndex; };
+
+	void SetFrameIndex(const int32 InFrameIndex) { FrameIndex = InFrameIndex; };
+
+	int32 GetEmitterIndex() const { return EmitterIndex; };
+
+	void SetEmitterIndex(int32 InEmitterIndex);
+
+	const FNiagaraSimCacheDataBuffersLayout* GetSimCacheBufferLayout() const;
+
+	FText GetComponentText(FName ComponentName, int32 InstanceIndex) const;
+
+	bool IsCacheValid();
+
+	int32 GetNumEmitterLayouts();
+
+	FName GetEmitterLayoutName(int32 Index);
+
+private:
+	uint32 FrameIndex = 0;
+	uint32 EmitterIndex = INDEX_NONE;
+};
\ No newline at end of file
diff --git a/Engine/Plugins/FX/Niagara/Source/NiagaraEditor/Private/Widgets/SNiagaraSimCacheView.cpp b/Engine/Plugins/FX/Niagara/Source/NiagaraEditor/Private/Widgets/SNiagaraSimCacheView.cpp
new file mode 100644
--- /dev/null
+++ b/Engine/Plugins/FX/Niagara/Source/NiagaraEditor/Private/Widgets/SNiagaraSimCacheView.cpp
@@ -0,0 +1,361 @@
+// Copyright Epic Games, Inc. All Rights Reserved.
+
+#include "SNiagaraSimCacheView.h"
+#include "ViewModels/NiagaraSimCacheViewModel.h"
+#include "NiagaraSimCache.h"
+
+#include "CoreMinimal.h"
+
+#include "Widgets/Input/SComboBox.h"
+#include "Widgets/Input/SEditableTextBox.h"
+#include "Widgets/Input/SSpinBox.h"
+#include "Widgets/Layout/SScrollBox.h"
+#include "Widgets/Views/SListView.h"
+#include "Widgets/Views/STableRow.h"
+#include "Widgets/Views/SHeaderRow.h"
+
+#include "IContentBrowserSingleton.h"
+#include "ContentBrowserModule.h"
+
+#define LOCTEXT_NAMESPACE "NiagaraSimCacheView"
+
+static const FName NAME_Instance("Instance");
+
+class SSimCacheDataBufferRowWidget : public SMultiColumnTableRow<TSharedPtr<int32>>
+{
+public:
+	SLATE_BEGIN_ARGS(SSimCacheDataBufferRowWidget) {}
+		SLATE_ARGUMENT(TSharedPtr<int32>, RowIndexPtr)
+		SLATE_ARGUMENT(TSharedPtr<FNiagaraSimCacheViewModel>, SimCacheViewModel)
+	SLATE_END_ARGS()
+
+	
+
+	/** Construct function for this widget */
+	void Construct(const FArguments& InArgs, const TSharedRef<STableViewBase>& InOwnerTableView)
+	{
+		RowIndexPtr = InArgs._RowIndexPtr;
+		SimCacheViewModel = InArgs._SimCacheViewModel;
+
+		SMultiColumnTableRow<TSharedPtr<int32>>::Construct(
+			FSuperRowType::FArguments()
+			.Style(FAppStyle::Get(), "DataTableEditor.CellListViewRow"),
+			InOwnerTableView
+		);
+	}
+
+	virtual TSharedRef<SWidget> GenerateWidgetForColumn(const FName& InColumnName) override
+	{
+		const int32 InstanceIndex = *RowIndexPtr;
+		if (InColumnName == NAME_Instance)
+		{
+			return SNew(STextBlock)
+				.Text(FText::AsNumber(InstanceIndex));
+		}
+		else
+		{
+			return SNew(STextBlock)
+				.Text(SimCacheViewModel->GetComponentText(InColumnName, InstanceIndex));
+		}
+	}
+
+	TSharedPtr<int32>						RowIndexPtr;
+	TSharedPtr<FNiagaraSimCacheViewModel>   SimCacheViewModel;
+};
+
+void SNiagaraSimCacheView::Construct(const FArguments& InArgs)
+{
+	SimCacheViewModel = InArgs._SimCacheViewModel;
+
+	HeaderRowWidget = SNew(SHeaderRow);
+
+	UpdateColumns(false);
+	UpdateRows(false);
+	UpdateBufferSelectionList();
+
+	TSharedRef<SScrollBar> HorizontalScrollBar =
+		SNew(SScrollBar)
+		.AlwaysShowScrollbar(true)
+		.Thickness(12.0f)
+		.Orientation(Orient_Horizontal);
+	TSharedRef<SScrollBar> VerticalScrollBar =
+		SNew(SScrollBar)
+		.AlwaysShowScrollbar(true)
+		.Thickness(12.0f)
+		.Orientation(Orient_Vertical);
+	
+	// Widget
+	ChildSlot
+	[
+		SNew(SVerticalBox)
+		+ SVerticalBox::Slot()
+		.AutoHeight()
+		[
+			// Select cache buffer
+			SNew(SHorizontalBox)
+			+ SHorizontalBox::Slot()
+			.AutoWidth()
+			.Padding(10.0f)
+			.VAlign(VAlign_Center)
+			.HAlign(HAlign_Right)
+			[
+				SNew(STextBlock)
+				.Text(LOCTEXT("CacheBufferSelection", "Cache Buffer Selection"))
+			]
+			+ SHorizontalBox::Slot()
+			.AutoWidth()
+			.VAlign(VAlign_Center)
+			.HAlign(HAlign_Left)
+			[
+				SNew(SComboBox<TSharedPtr<FBufferSelectionInfo>>)
+				.OptionsSource(&BufferSelectionList)
+				.OnGenerateWidget(this, &SNiagaraSimCacheView::BufferSelectionGenerateWidget)
+				.OnSelectionChanged(this, &SNiagaraSimCacheView::BufferSelectionChanged)
+				.InitiallySelectedItem(BufferSelectionList[0])
+				[
+					SNew(STextBlock)
+					.Text(this, &SNiagaraSimCacheView::GetBufferSelectionText)
+				]
+			]
+			// Select Cache Frame
+			+ SHorizontalBox::Slot()
+			.AutoWidth()
+			.Padding(10.0f)
+			.VAlign(VAlign_Center)
+			.HAlign(HAlign_Right)
+			[
+				SNew(STextBlock)
+				.Text(LOCTEXT("FrameIndex", "Frame Index"))
+			]
+			+ SHorizontalBox::Slot()
+			.AutoWidth()
+			.VAlign(VAlign_Center)
+			.HAlign(HAlign_Left)
+			[
+				SNew(SSpinBox<int32>)
+				.MinValue(0)
+				.MaxValue(this, &SNiagaraSimCacheView::GetNumFrames)
+				.Value(this, &SNiagaraSimCacheView::GetFrameIndex)
+				.OnValueChanged(this, &SNiagaraSimCacheView::SetFrameIndex)
+			]
+			// Component Filter
+			+ SHorizontalBox::Slot()
+			.AutoWidth()
+			.Padding(10.0f)
+			.VAlign(VAlign_Center)
+			.HAlign(HAlign_Right)
+			[
+				SNew(STextBlock)
+				.Text(LOCTEXT("ComponentFilter", "Component Filter"))
+			]
+			+ SHorizontalBox::Slot()
+			.AutoWidth()
+			.VAlign(VAlign_Center)
+			.HAlign(HAlign_Left)
+			[
+				SNew(SEditableTextBox)
+				.OnTextChanged(this, &SNiagaraSimCacheView::OnComponentFilterChange)
+				.MinDesiredWidth(100)
+			]
+		]
+		//// Main Spreadsheet View
+		+ SVerticalBox::Slot()
+		.FillHeight(1.0f)
+		[
+			SNew(SVerticalBox)
+			+ SVerticalBox::Slot()
+			.FillHeight(1.0f)
+			[
+				SNew(SHorizontalBox)
+				+ SHorizontalBox::Slot()
+				.FillWidth(1.0f)
+				[
+					SNew(SScrollBox)
+					.Orientation(Orient_Horizontal)
+					.ExternalScrollbar(HorizontalScrollBar)
+					+ SScrollBox::Slot()
+					[
+						SAssignNew(ListViewWidget, SListView<TSharedPtr<int32>>)
+						.ListItemsSource(&RowItems)
+						.OnGenerateRow(this, &SNiagaraSimCacheView::MakeRowWidget)
+						.Visibility(EVisibility::Visible)
+						.SelectionMode(ESelectionMode::Single)
+						.ExternalScrollbar(VerticalScrollBar)
+						.HeaderRow(HeaderRowWidget)
+					]
+				]
+				+ SHorizontalBox::Slot()
+				.AutoWidth()
+				[
+					VerticalScrollBar
+				]
+			]
+		]
+		+ SVerticalBox::Slot()
+		.AutoHeight()
+		[
+			HorizontalScrollBar
+		]
+	];
+};
+
+TSharedRef<ITableRow> SNiagaraSimCacheView::MakeRowWidget(const TSharedPtr<int32> RowIndexPtr, const TSharedRef<STableViewBase>& OwnerTable) const
+{
+	return
+		SNew(SSimCacheDataBufferRowWidget, OwnerTable)
+		.RowIndexPtr(RowIndexPtr)
+		.SimCacheViewModel(SimCacheViewModel);
+}
+
+void SNiagaraSimCacheView::UpdateColumns(const bool bRefresh)
+{
+	// TODO: Hide columns rather than regenerating everything?
+	HeaderRowWidget->ClearColumns();
+
+	HeaderRowWidget->AddColumn(
+		SHeaderRow::Column(NAME_Instance)
+		.DefaultLabel(FText::FromName(NAME_Instance))
+		.HAlignHeader(EHorizontalAlignment::HAlign_Center)
+		.VAlignHeader(EVerticalAlignment::VAlign_Fill)
+		.HAlignCell(EHorizontalAlignment::HAlign_Center)
+		.VAlignCell(EVerticalAlignment::VAlign_Fill)
+		.SortMode(EColumnSortMode::None)
+	);
+
+
+
+	const bool bFilterActive = ComponentFilterArray.Num() > 0;
+	for (const FNiagaraSimCacheComponentInfo& ComponentInfo : SimCacheViewModel->GetComponentInfos())
+	{
+		if (bFilterActive)
+		{
+			const FString ComponentInfoString = ComponentInfo.Name.ToString();
+			const bool bPassedFilter = ComponentFilterArray.ContainsByPredicate(
+				[ComponentInfoString = ComponentInfo.Name.ToString()](const FString& ComponentFilter)
+			{
+				return ComponentInfoString.Contains(ComponentFilter);
+			}
+			);
+			if (bPassedFilter == false)
+			{
+				continue;
+			}
+		}
+
+		HeaderRowWidget->AddColumn(
+			SHeaderRow::Column(ComponentInfo.Name)
+			.DefaultLabel(FText::FromName(ComponentInfo.Name))
+			.HAlignHeader(EHorizontalAlignment::HAlign_Center)
+			.VAlignHeader(EVerticalAlignment::VAlign_Fill)
+			.HAlignCell(EHorizontalAlignment::HAlign_Center)
+			.VAlignCell(EVerticalAlignment::VAlign_Fill)
+			.SortMode(EColumnSortMode::None)
+		);
+	}
+
+	HeaderRowWidget->ResetColumnWidths();
+	HeaderRowWidget->RefreshColumns();
+	if (bRefresh && ListViewWidget)
+	{
+		ListViewWidget->RequestListRefresh();
+	}
+}
+
+void SNiagaraSimCacheView::UpdateRows(const bool bRefresh)
+{
+	RowItems.Reset(SimCacheViewModel->GetNumInstances());
+	for (int32 i = 0; i < SimCacheViewModel->GetNumInstances(); ++i)
+	{
+		RowItems.Emplace(MakeShared<int32>(i));
+	}
+
+	if (bRefresh && ListViewWidget)
+	{
+		ListViewWidget->RequestListRefresh();
+	}
+}
+
+void SNiagaraSimCacheView::UpdateBufferSelectionList()
+{
+	BufferSelectionList.Empty();
+
+	if (SimCacheViewModel->IsCacheValid())
+	{
+		BufferSelectionList.Emplace(MakeShared<FBufferSelectionInfo>(INDEX_NONE, LOCTEXT("SystemInstance", "System Instance")));
+		for (int32 i = 0; i < SimCacheViewModel->GetNumEmitterLayouts(); ++i)
+		{
+			BufferSelectionList.Emplace(
+				MakeShared<FBufferSelectionInfo>(i, FText::Format(LOCTEXT("EmitterFormat", "Emitter - {0}"), FText::FromName(SimCacheViewModel->GetEmitterLayoutName(i))))
+			);
+		}
+	}
+	else
+	{
+		BufferSelectionList.Emplace(MakeShared<FBufferSelectionInfo>(INDEX_NONE, LOCTEXT("InvalidCache", "Invalid Cache")));
+	}
+}
+
+TOptional<int32> SNiagaraSimCacheView::GetNumFrames() const
+{
+	return SimCacheViewModel->GetNumFrames();
+}
+
+int32 SNiagaraSimCacheView::GetFrameIndex() const
+{
+	return SimCacheViewModel->GetFrameIndex();
+}
+
+void SNiagaraSimCacheView::SetFrameIndex(int32 Value)
+{
+	SimCacheViewModel->SetFrameIndex(Value);
+	UpdateRows(true);
+}
+
+TSharedRef<SWidget> SNiagaraSimCacheView::BufferSelectionGenerateWidget(TSharedPtr<FBufferSelectionInfo> InItem)
+{
+	return
+		SNew(STextBlock)
+		.Text(InItem->Value);
+}
+
+void SNiagaraSimCacheView::BufferSelectionChanged(TSharedPtr<FBufferSelectionInfo> NewSelection, ESelectInfo::Type SelectInfo)
+{
+	SimCacheViewModel->SetEmitterIndex(NewSelection.IsValid() ? NewSelection->Key : INDEX_NONE);
+	UpdateColumns(false);
+	UpdateRows(true);
+}
+
+FText SNiagaraSimCacheView::GetBufferSelectionText() const
+{
+	const int32 EmitterIndex = SimCacheViewModel->GetEmitterIndex();
+	for (const TSharedPtr<FBufferSelectionInfo>& SelectionInfo : BufferSelectionList)
+	{
+		if (SelectionInfo->Key == EmitterIndex)
+		{
+			return SelectionInfo->Value;
+		}
+	}
+	return FText::GetEmpty();
+}
+
+void SNiagaraSimCacheView::OnComponentFilterChange(const FText& InFilter)
+{
+	ComponentFilterArray.Empty();
+	InFilter.ToString().ParseIntoArray(ComponentFilterArray, TEXT(","));
+	UpdateColumns(true);
+}
+
+void SNiagaraSimCacheView::OnSimCacheChanged(const FAssetData& InAsset)
+{
+	FSlateApplication::Get().DismissAllMenus();
+
+	if (UNiagaraSimCache* SimCache = Cast<UNiagaraSimCache>(InAsset.GetAsset()))
+	{
+		SimCacheViewModel->Initialize(SimCache);
+		UpdateRows(true);
+		UpdateColumns(true);
+		UpdateBufferSelectionList();
+	}
+}
+
+#undef LOCTEXT_NAMESPACE
\ No newline at end of file
diff --git a/Engine/Plugins/FX/Niagara/Source/NiagaraEditor/Private/Widgets/SNiagaraSimCacheView.h b/Engine/Plugins/FX/Niagara/Source/NiagaraEditor/Private/Widgets/SNiagaraSimCacheView.h
new file mode 100644
--- /dev/null
+++ b/Engine/Plugins/FX/Niagara/Source/NiagaraEditor/Private/Widgets/SNiagaraSimCacheView.h
@@ -0,0 +1,58 @@
+// Copyright Epic Games, Inc. All Rights Reserved.
+
+#pragma once
+
+#include "Widgets/SCompoundWidget.h"
+#include "Widgets/Views/SListView.h"
+
+class FNiagaraSimCacheViewModel;
+class ITableRow;
+class STableViewBase;
+class SHeaderRow;
+
+class SNiagaraSimCacheView : public SCompoundWidget
+{
+public:
+	SLATE_BEGIN_ARGS(SNiagaraSimCacheView) {}
+		SLATE_ARGUMENT(TSharedPtr<FNiagaraSimCacheViewModel>, SimCacheViewModel)
+	SLATE_END_ARGS()
+
+	using FBufferSelectionInfo = TPair<int32, FText>;
+
+	void Construct(const FArguments& InArgs);
+
+	TSharedRef<ITableRow> MakeRowWidget(const TSharedPtr<int32> RowIndexPtr, const TSharedRef<STableViewBase>& OwnerTable) const;
+
+	void UpdateColumns(const bool bRefresh);
+
+	void UpdateRows(const bool bRefresh);
+
+	void UpdateBufferSelectionList();
+
+	TOptional<int32> GetNumFrames() const;
+	int32 GetFrameIndex() const;
+
+	void SetFrameIndex(int32);
+
+	TSharedRef<SWidget> BufferSelectionGenerateWidget(TSharedPtr<FBufferSelectionInfo> InItem);
+
+	void BufferSelectionChanged(TSharedPtr<FBufferSelectionInfo> NewSelection, ESelectInfo::Type SelectInfo);
+
+	FText GetBufferSelectionText() const;
+
+	void OnComponentFilterChange(const FText& InFilter);
+
+	void OnSimCacheChanged(const FAssetData& InAsset);
+
+	TArray<TSharedPtr<int32>>					RowItems;
+	TSharedPtr<FNiagaraSimCacheViewModel>		SimCacheViewModel;
+
+	TArray<TSharedPtr<FBufferSelectionInfo>>	BufferSelectionList;
+
+	int32										FrameIndex = 0;
+
+	TSharedPtr<SHeaderRow>						HeaderRowWidget;
+	TSharedPtr<SListView<TSharedPtr<int32>>>	ListViewWidget;
+
+	TArray<FString>								ComponentFilterArray;
+};
\ No newline at end of file
diff --git a/Engine/Plugins/FX/Niagara/Source/NiagaraEditorWidgets/Private/DetailCustomizations/NiagaraSimCacheDetailsCustomization.cpp b/Engine/Plugins/FX/Niagara/Source/NiagaraEditorWidgets/Private/DetailCustomizations/NiagaraSimCacheDetailsCustomization.cpp
deleted file mode 100644
diff --git a/Engine/Plugins/FX/Niagara/Source/NiagaraEditorWidgets/Private/DetailCustomizations/NiagaraSimCacheDetailsCustomization.h b/Engine/Plugins/FX/Niagara/Source/NiagaraEditorWidgets/Private/DetailCustomizations/NiagaraSimCacheDetailsCustomization.h
deleted file mode 100644
diff --git a/Engine/Plugins/FX/Niagara/Source/NiagaraEditorWidgets/Private/NiagaraEditorWidgetsModule.cpp b/Engine/Plugins/FX/Niagara/Source/NiagaraEditorWidgets/Private/NiagaraEditorWidgetsModule.cpp
--- a/Engine/Plugins/FX/Niagara/Source/NiagaraEditorWidgets/Private/NiagaraEditorWidgetsModule.cpp
+++ b/Engine/Plugins/FX/Niagara/Source/NiagaraEditorWidgets/Private/NiagaraEditorWidgetsModule.cpp
@@ -13,7 +13,6 @@
 #include "DetailCustomizations/NiagaraDataInterfaceMeshRendererInfoDetails.h"
 #include "DetailCustomizations/NiagaraDataInterfaceSpriteRendererInfoDetails.h"
 #include "DetailCustomizations/NiagaraMeshRendererDetails.h"
-#include "DetailCustomizations/NiagaraSimCacheDetailsCustomization.h"
 #include "ViewModels/NiagaraSystemViewModel.h"
 #include "SNiagaraOverviewGraph.h"
 #include "NiagaraEditorWidgetsUtilities.h"
@@ -133,7 +132,6 @@ void FNiagaraEditorWidgetsModule::StartupModule()
 	PropertyModule.RegisterCustomClassLayout("NiagaraDataInterfaceMeshRendererInfo", FOnGetDetailCustomizationInstance::CreateStatic(&FNiagaraDataInterfaceMeshRendererInfoDetails::MakeInstance));
 	PropertyModule.RegisterCustomClassLayout("NiagaraDataInterfaceSpriteRendererInfo", FOnGetDetailCustomizationInstance::CreateStatic(&FNiagaraDataInterfaceSpriteRendererInfoDetails::MakeInstance));
 	PropertyModule.RegisterCustomClassLayout("NiagaraMeshRendererProperties", FOnGetDetailCustomizationInstance::CreateStatic(&FNiagaraMeshRendererDetails::MakeInstance));
-	PropertyModule.RegisterCustomClassLayout("NiagaraSimCache", FOnGetDetailCustomizationInstance::CreateStatic(&FNiagaraSimCacheDetailsCustomization::MakeInstance));
 
 	ReinitializeStyleCommand = IConsoleManager::Get().RegisterConsoleCommand(
 		TEXT("fx.NiagaraEditorWidgets.ReinitializeStyle"),
