- Run component destruction immediatly if on GameThread to avoid lingering components in some contexts

#rnx

diff --git a/Engine/Plugins/FX/Niagara/Source/Niagara/Private/NiagaraRendererComponents.cpp b/Engine/Plugins/FX/Niagara/Source/Niagara/Private/NiagaraRendererComponents.cpp
--- a/Engine/Plugins/FX/Niagara/Source/Niagara/Private/NiagaraRendererComponents.cpp
+++ b/Engine/Plugins/FX/Niagara/Source/Niagara/Private/NiagaraRendererComponents.cpp
@@ -288,27 +288,48 @@ void FNiagaraRendererComponents::DestroyRenderState_Concurrent()
 	}
 #endif
 
-	// Rendering resources are being destroyed, but the component pool and their owner actor must be destroyed on the game thread
-	AsyncTask(
-		ENamedThreads::GameThread,
-		[Pool_GT = MoveTemp(ComponentPool), Owner_GT = MoveTemp(SpawnedOwner)]()
-		{
-			// we do not reset ParticlesWithComponents here because it's possible the render state is destroyed without destroying the renderer. In this case we want to know which particles
-			// had spawned some components previously
-			for (auto& PoolEntry : Pool_GT)
+	// If we are on the GameThread immediately execute
+	// Failure to do so can cause issue with construction scripts rerunning as the components can linger when they should not
+	if (IsInGameThread())
+	{
+		for (auto& PoolEntry : ComponentPool)
+		{
+			if (PoolEntry.Component.IsValid())
 			{
-				if (PoolEntry.Component.IsValid())
-				{
-					PoolEntry.Component->DestroyComponent();
-				}
+				PoolEntry.Component->DestroyComponent();
 			}
+		}
+		ComponentPool.Empty();
 
-			if (AActor* OwnerActor = Owner_GT.Get())
+		if (AActor* OwnerActor = SpawnedOwner.Get())
+		{
+			OwnerActor->Destroy();
+		}
+	}
+	else
+	{
+		// Rendering resources are being destroyed, but the component pool and their owner actor must be destroyed on the game thread
+		AsyncTask(
+			ENamedThreads::GameThread,
+			[Pool_GT = MoveTemp(ComponentPool), Owner_GT = MoveTemp(SpawnedOwner)]()
 			{
-				OwnerActor->Destroy();
+				// we do not reset ParticlesWithComponents here because it's possible the render state is destroyed without destroying the renderer. In this case we want to know which particles
+				// had spawned some components previously
+				for (auto& PoolEntry : Pool_GT)
+				{
+					if (PoolEntry.Component.IsValid())
+					{
+						PoolEntry.Component->DestroyComponent();
+					}
+				}
+
+				if (AActor* OwnerActor = Owner_GT.Get())
+				{
+					OwnerActor->Destroy();
+				}
 			}
-		}
-	);
+		);
+	}
 	SpawnedOwner.Reset();
 }
 
