Couple of fixes for handling experimental and legacy byte code side by side
-introduces a new cvar (fx.NiagaraScript.StripByteCodeOnLoad) which can be used to strip out one of the byte codes, saving memory, -1 -> strips experimental, 1 -> strips legacy
-fixes identification of whether an Inputhandle

diff --git a/Engine/Plugins/FX/Niagara/Source/Niagara/Classes/NiagaraScript.h b/Engine/Plugins/FX/Niagara/Source/Niagara/Classes/NiagaraScript.h
--- a/Engine/Plugins/FX/Niagara/Source/Niagara/Classes/NiagaraScript.h
+++ b/Engine/Plugins/FX/Niagara/Source/Niagara/Classes/NiagaraScript.h
@@ -563,10 +563,23 @@ public:
 	void BakeScriptLiterals(TArray<uint8>& OutLiterals) const;
 #endif
 
+	bool HasByteCode() const;
+
 #if VECTORVM_SUPPORTS_EXPERIMENTAL
 	bool SupportsExperimentalVM() const;
 	FVectorVMOptimizeContext BuildExperimentalContext() const;
 #endif //VECTORVM_SUPPORTS_EXPERIMENTAL
+
+	void PostSerialize(const FArchive& Ar);
+};
+
+template<>
+struct TStructOpsTypeTraits<FNiagaraVMExecutableData> : public TStructOpsTypeTraitsBase2<FNiagaraVMExecutableData>
+{
+	enum
+	{
+		WithPostSerialize = true,
+	};
 };
 
 struct NIAGARA_API FNiagaraGraphCachedDataBase
diff --git a/Engine/Plugins/FX/Niagara/Source/Niagara/Private/NiagaraScript.cpp b/Engine/Plugins/FX/Niagara/Source/Niagara/Private/NiagaraScript.cpp
--- a/Engine/Plugins/FX/Niagara/Source/Niagara/Private/NiagaraScript.cpp
+++ b/Engine/Plugins/FX/Niagara/Source/Niagara/Private/NiagaraScript.cpp
@@ -98,6 +98,14 @@ static FAutoConsoleVariableRef CVarLogCompileStaticVars(
 	ECVF_Default
 );
 
+int32 GNiagaraScriptStripByteCodeOnLoad = 0;
+static FAutoConsoleVariableRef CVarGNiagaraScriptStripByteCodeOnLoad(
+	TEXT("fx.NiagaraScript.StripByteCodeOnLoad"),
+	GNiagaraScriptStripByteCodeOnLoad,
+	TEXT("If > 0 all scripts will have their legacy byte code stripped on load.  If < 0 all scripts will have their experimental data stripped on load. \n"),
+	ECVF_Default
+);
+
 FNiagaraScriptDebuggerInfo::FNiagaraScriptDebuggerInfo() : bWaitForGPU(false), Usage(ENiagaraScriptUsage::Function), FrameLastWriteId(-1), bWritten(false)
 {
 }
@@ -313,6 +321,41 @@ FVectorVMOptimizeContext FNiagaraVMExecutableData::BuildExperimentalContext() co
 }
 #endif
 
+bool FNiagaraVMExecutableData::HasByteCode() const
+{
+#if VECTORVM_SUPPORTS_LEGACY
+	if (ByteCode.HasByteCode() || OptimizedByteCode.HasByteCode())
+	{
+		return true;
+	}
+#endif
+
+#if VECTORVM_SUPPORTS_EXPERIMENTAL
+	if (!ExperimentalContextData.IsEmpty())
+	{
+		return true;
+	}
+#endif
+
+	return false;
+}
+
+void FNiagaraVMExecutableData::PostSerialize(const FArchive& Ar)
+{
+#if VECTORVM_SUPPORTS_EXPERIMENTAL && VECTORVM_SUPPORTS_LEGACY && !WITH_EDITOR
+	if (Ar.IsLoading())
+	{
+		if (GNiagaraScriptStripByteCodeOnLoad > 0)
+		{
+			ByteCode.Reset();
+		}
+		else if (GNiagaraScriptStripByteCodeOnLoad < 0)
+		{
+			ExperimentalContextData.Empty();
+		}
+	}
+#endif
+}
 
 UNiagaraScript::UNiagaraScript()
 {
@@ -3550,7 +3593,10 @@ NIAGARA_API bool UNiagaraScript::IsScriptCompilationPending(bool bGPUScript) con
 	}
 	else if (CachedScriptVM.IsValid())
 	{
-		return !CachedScriptVM.ByteCode.HasByteCode() && !CachedScriptVM.OptimizedByteCode.HasByteCode() && (CachedScriptVM.LastCompileStatus == ENiagaraScriptCompileStatus::NCS_BeingCreated || CachedScriptVM.LastCompileStatus == ENiagaraScriptCompileStatus::NCS_Unknown);
+		if (CachedScriptVM.LastCompileStatus == ENiagaraScriptCompileStatus::NCS_BeingCreated || CachedScriptVM.LastCompileStatus == ENiagaraScriptCompileStatus::NCS_Unknown)
+		{
+			return !CachedScriptVM.HasByteCode();
+		}
 	}
 	return false;
 }
@@ -3583,7 +3629,7 @@ NIAGARA_API bool UNiagaraScript::DidScriptCompilationSucceed(bool bGPUScript) co
 	}
 	else if (CachedScriptVM.IsValid())
 	{
-		return (CachedScriptVM.ByteCode.HasByteCode()) || (CachedScriptVM.OptimizedByteCode.HasByteCode());
+		return (CachedScriptVM.HasByteCode());
 	}
 
 	return false;
