- Add DebugHud method so Static Mesh DI displays info about assigned static mesh
- Add missing Equals check for PreviewMesh

#rnx

diff --git a/Engine/Plugins/FX/Niagara/Source/Niagara/Private/DataInterface/NiagaraDataInterfaceStaticMesh.cpp b/Engine/Plugins/FX/Niagara/Source/Niagara/Private/DataInterface/NiagaraDataInterfaceStaticMesh.cpp
--- a/Engine/Plugins/FX/Niagara/Source/Niagara/Private/DataInterface/NiagaraDataInterfaceStaticMesh.cpp
+++ b/Engine/Plugins/FX/Niagara/Source/Niagara/Private/DataInterface/NiagaraDataInterfaceStaticMesh.cpp
@@ -2601,10 +2601,14 @@ bool UNiagaraDataInterfaceStaticMesh::Equals(const UNiagaraDataInterface* Other)
 		return false;
 	}
 	const UNiagaraDataInterfaceStaticMesh* OtherTyped = CastChecked<const UNiagaraDataInterfaceStaticMesh>(Other);
-	return OtherTyped->SourceMode == SourceMode &&
+	return
+		OtherTyped->SourceMode == SourceMode &&
+#if WITH_EDITORONLY_DATA
+		OtherTyped->PreviewMesh == PreviewMesh &&
+#endif
+		OtherTyped->DefaultMesh == DefaultMesh &&
 		OtherTyped->SoftSourceActor == SoftSourceActor &&
 		OtherTyped->SourceComponent == SourceComponent &&
-		OtherTyped->DefaultMesh == DefaultMesh &&
 		OtherTyped->SectionFilter.AllowedMaterialSlots == SectionFilter.AllowedMaterialSlots &&
 		OtherTyped->bUsePhysicsBodyVelocity == bUsePhysicsBodyVelocity &&
 		OtherTyped->FilteredSockets == FilteredSockets;
@@ -2619,18 +2623,33 @@ bool UNiagaraDataInterfaceStaticMesh::CopyToInternal(UNiagaraDataInterface* Dest
 
 	UNiagaraDataInterfaceStaticMesh* OtherTyped = CastChecked<UNiagaraDataInterfaceStaticMesh>(Destination);
 	OtherTyped->SourceMode = SourceMode;
-	OtherTyped->SoftSourceActor = SoftSourceActor;
-	OtherTyped->SourceComponent = SourceComponent;
-	OtherTyped->DefaultMesh = DefaultMesh;
 #if WITH_EDITORONLY_DATA
 	OtherTyped->PreviewMesh = PreviewMesh;
 #endif
+	OtherTyped->DefaultMesh = DefaultMesh;
+	OtherTyped->SoftSourceActor = SoftSourceActor;
+	OtherTyped->SourceComponent = SourceComponent;
 	OtherTyped->SectionFilter = SectionFilter;
 	OtherTyped->bUsePhysicsBodyVelocity = bUsePhysicsBodyVelocity;
 	OtherTyped->FilteredSockets = FilteredSockets;
 	return true;
 }
 
+#if WITH_NIAGARA_DEBUGGER
+void UNiagaraDataInterfaceStaticMesh::DrawDebugHud(UCanvas* Canvas, FNiagaraSystemInstance* SystemInstance, FString& VariableDataString, bool bVerbose) const
+{
+	NDIStaticMeshLocal::FInstanceData_GameThread* InstanceData_GT = SystemInstance->FindTypedDataInterfaceInstanceData<NDIStaticMeshLocal::FInstanceData_GameThread>(this);
+	if (InstanceData_GT == nullptr)
+	{
+		return;
+	}
+
+	USceneComponent* SceneComponent = InstanceData_GT->SceneComponentWeakPtr.Get();
+	UStaticMesh* SkeletalMesh = InstanceData_GT->StaticMeshWeakPtr.Get();
+	VariableDataString = FString::Printf(TEXT("StaticMesh(%s) StaticMeshComp(%s)"), *GetNameSafe(SkeletalMesh), *GetNameSafe(SceneComponent));
+}
+#endif
+
 #if WITH_EDITOR
 void UNiagaraDataInterfaceStaticMesh::GetFeedback(UNiagaraSystem* Asset, UNiagaraComponent* Component, TArray<FNiagaraDataInterfaceError>& OutErrors, TArray<FNiagaraDataInterfaceFeedback>& OutWarnings, TArray<FNiagaraDataInterfaceFeedback>& OutInfo)
 {
diff --git a/Engine/Plugins/FX/Niagara/Source/Niagara/Private/DataInterface/NiagaraDataInterfaceStaticMesh.h b/Engine/Plugins/FX/Niagara/Source/Niagara/Private/DataInterface/NiagaraDataInterfaceStaticMesh.h
--- a/Engine/Plugins/FX/Niagara/Source/Niagara/Private/DataInterface/NiagaraDataInterfaceStaticMesh.h
+++ b/Engine/Plugins/FX/Niagara/Source/Niagara/Private/DataInterface/NiagaraDataInterfaceStaticMesh.h
@@ -139,6 +139,9 @@ public:
 protected:
 	virtual bool CopyToInternal(UNiagaraDataInterface* Destination) const override;
 public:
+#if WITH_NIAGARA_DEBUGGER
+	virtual void DrawDebugHud(UCanvas* Canvas, FNiagaraSystemInstance* SystemInstance, FString& VariableDataString, bool bVerbose) const override;
+#endif
 #if WITH_EDITOR
 	virtual void GetFeedback(UNiagaraSystem* Asset, UNiagaraComponent* Component, TArray<FNiagaraDataInterfaceError>& OutErrors,
 		TArray<FNiagaraDataInterfaceFeedback>& OutWarnings, TArray<FNiagaraDataInterfaceFeedback>& OutInfo) override;
