Skeletal Mesh - Replace the newly created GetSkeletalMesh function by GetSkeletalMeshAsset so it matches the setter SetSkeletalMeshAsset function (which itself couldn't be named SetSkeletalMesh due to the function already existing and doing something else).

diff --git a/Engine/Plugins/Animation/AnimToTexture/Source/AnimToTextureEditor/Private/AnimToTextureBPLibrary.cpp b/Engine/Plugins/Animation/AnimToTexture/Source/AnimToTextureEditor/Private/AnimToTextureBPLibrary.cpp
--- a/Engine/Plugins/Animation/AnimToTexture/Source/AnimToTextureEditor/Private/AnimToTextureBPLibrary.cpp
+++ b/Engine/Plugins/Animation/AnimToTexture/Source/AnimToTextureEditor/Private/AnimToTextureBPLibrary.cpp
@@ -192,15 +192,15 @@ void UAnimToTextureBPLibrary::AnimationToTexture(UAnimToTextureDataAsset* DataAs
 	if (DataAsset->Mode == EAnimToTextureMode::Bone)
 	{
 		// Get Number of RawBones (no virtual)
-		NumBones = AnimToTexture_Private::GetNumBones(BaseMeshComponent->GetSkeletalMesh());
+		NumBones = AnimToTexture_Private::GetNumBones(BaseMeshComponent->GetSkeletalMeshAsset());
 
 		// Get Raw Ref Bone (no virtual)
 		TArray<FTransform> RefBoneTransforms;
-		AnimToTexture_Private::GetRefBoneTransforms(BaseMeshComponent->GetSkeletalMesh(), RefBoneTransforms);
+		AnimToTexture_Private::GetRefBoneTransforms(BaseMeshComponent->GetSkeletalMeshAsset(), RefBoneTransforms);
 		AnimToTexture_Private::DecomposeTransformations(RefBoneTransforms, BoneRefPositions, BoneRefRotations);
 
 		// Get Bone Names (no virtual)
-		AnimToTexture_Private::GetBoneNames(BaseMeshComponent->GetSkeletalMesh(), BoneNames);
+		AnimToTexture_Private::GetBoneNames(BaseMeshComponent->GetSkeletalMeshAsset(), BoneNames);
 
 		// Make sure array sizes are correct.
 		check(RefBoneTransforms.Num() == NumBones && BoneNames.Num() == NumBones);
@@ -238,9 +238,9 @@ void UAnimToTextureBPLibrary::AnimationToTexture(UAnimToTextureDataAsset* DataAs
 		}
 		
 		// Make sure SkelMesh and AnimSequence use same Skeleton
-		if (AnimSequence->GetSkeleton() != BaseMeshComponent->GetSkeletalMesh()->GetSkeleton())
+		if (AnimSequence->GetSkeleton() != BaseMeshComponent->GetSkeletalMeshAsset()->GetSkeleton())
 		{
-			UE_LOG(LogTemp, Warning, TEXT("Invalid AnimSequence: %s for given SkeletalMesh: %s"), *AnimSequence->GetFName().ToString(), *BaseMeshComponent->GetSkeletalMesh()->GetFName().ToString());
+			UE_LOG(LogTemp, Warning, TEXT("Invalid AnimSequence: %s for given SkeletalMesh: %s"), *AnimSequence->GetFName().ToString(), *BaseMeshComponent->GetSkeletalMeshAsset()->GetFName().ToString());
 			continue;
 		}
 		// Set AnimSequence
diff --git a/Engine/Plugins/Animation/AnimToTexture/Source/AnimToTextureEditor/Private/AnimToTextureSkeletalMesh.cpp b/Engine/Plugins/Animation/AnimToTexture/Source/AnimToTextureEditor/Private/AnimToTextureSkeletalMesh.cpp
--- a/Engine/Plugins/Animation/AnimToTexture/Source/AnimToTextureEditor/Private/AnimToTextureSkeletalMesh.cpp
+++ b/Engine/Plugins/Animation/AnimToTexture/Source/AnimToTextureEditor/Private/AnimToTextureSkeletalMesh.cpp
@@ -190,7 +190,7 @@ void GetSkinnedVertices(USkeletalMeshComponent* MeshComponent, const int32 LODIn
 	}
 
 	// Get SkeletalMesh
-	USkeletalMesh* SkeletalMesh = MeshComponent->GetSkeletalMesh();
+	USkeletalMesh* SkeletalMesh = MeshComponent->GetSkeletalMeshAsset();
 	if (!SkeletalMesh)
 	{
 		return;
diff --git a/Engine/Plugins/Animation/ControlRig/Source/ControlRig/Private/AnimNode_ControlRig.cpp b/Engine/Plugins/Animation/ControlRig/Source/ControlRig/Private/AnimNode_ControlRig.cpp
--- a/Engine/Plugins/Animation/ControlRig/Source/ControlRig/Private/AnimNode_ControlRig.cpp
+++ b/Engine/Plugins/Animation/ControlRig/Source/ControlRig/Private/AnimNode_ControlRig.cpp
@@ -293,7 +293,7 @@ void FAnimNode_ControlRig::UpdateControlRigRefPoseIfNeeded(const FAnimInstancePr
 
 	if(InProxy->GetSkelMeshComponent())
 	{
-		ExpectedHash = HashCombine(ExpectedHash, (int32)reinterpret_cast<uintptr_t>(InProxy->GetSkelMeshComponent()->GetSkeletalMesh()));
+		ExpectedHash = HashCombine(ExpectedHash, (int32)reinterpret_cast<uintptr_t>(InProxy->GetSkelMeshComponent()->GetSkeletalMeshAsset()));
 	}
 
 	if(bIncludePoseInHash)
diff --git a/Engine/Plugins/Animation/ControlRig/Source/ControlRig/Private/AnimNode_ControlRigBase.cpp b/Engine/Plugins/Animation/ControlRig/Source/ControlRig/Private/AnimNode_ControlRigBase.cpp
--- a/Engine/Plugins/Animation/ControlRig/Source/ControlRig/Private/AnimNode_ControlRigBase.cpp
+++ b/Engine/Plugins/Animation/ControlRig/Source/ControlRig/Private/AnimNode_ControlRigBase.cpp
@@ -39,7 +39,7 @@ void FAnimNode_ControlRigBase::OnInitializeAnimInstance(const FAnimInstanceProxy
 
 	USkeletalMeshComponent* Component = InAnimInstance->GetOwningComponent();
 	UControlRig* ControlRig = GetControlRig();
-	if (Component && Component->GetSkeletalMesh() && ControlRig)
+	if (Component && Component->GetSkeletalMeshAsset() && ControlRig)
 	{
 #if WITH_EDITORONLY_DATA
 		UBlueprintGeneratedClass* BlueprintClass = Cast<UBlueprintGeneratedClass>(ControlRig->GetClass());
@@ -47,7 +47,7 @@ void FAnimNode_ControlRigBase::OnInitializeAnimInstance(const FAnimInstanceProxy
 		{
 			UBlueprint* Blueprint = Cast<UBlueprint>(BlueprintClass->ClassGeneratedBy);
 			// node mapping container will be saved on the initialization part
-			NodeMappingContainer = Component->GetSkeletalMesh()->GetNodeMappingContainer(Blueprint);
+			NodeMappingContainer = Component->GetSkeletalMeshAsset()->GetNodeMappingContainer(Blueprint);
 		}
 #endif
 
diff --git a/Engine/Plugins/Animation/ControlRig/Source/ControlRig/Private/ControlRig.cpp b/Engine/Plugins/Animation/ControlRig/Source/ControlRig/Private/ControlRig.cpp
--- a/Engine/Plugins/Animation/ControlRig/Source/ControlRig/Private/ControlRig.cpp
+++ b/Engine/Plugins/Animation/ControlRig/Source/ControlRig/Private/ControlRig.cpp
@@ -3138,7 +3138,7 @@ void UControlRig::SetBoneInitialTransformsFromSkeletalMeshComponent(USkeletalMes
 	}
 	else
 	{
-		SetBoneInitialTransformsFromSkeletalMesh(InSkelMeshComp->GetSkeletalMesh());
+		SetBoneInitialTransformsFromSkeletalMesh(InSkelMeshComp->GetSkeletalMeshAsset());
 	}
 }
 
diff --git a/Engine/Plugins/Animation/ControlRig/Source/ControlRig/Private/ControlRigComponent.cpp b/Engine/Plugins/Animation/ControlRig/Source/ControlRig/Private/ControlRigComponent.cpp
--- a/Engine/Plugins/Animation/ControlRig/Source/ControlRig/Private/ControlRigComponent.cpp
+++ b/Engine/Plugins/Animation/ControlRig/Source/ControlRig/Private/ControlRigComponent.cpp
@@ -569,7 +569,7 @@ void UControlRigComponent::AddMappedSkeletalMesh(USkeletalMeshComponent* Skeleta
 	TArray<FControlRigComponentMappedBone> BonesToMap = Bones;
 	if (BonesToMap.Num() == 0)
 	{
-		if (const USkeletalMesh* SkeletalMesh = SkeletalMeshComponent->GetSkeletalMesh())
+		if (const USkeletalMesh* SkeletalMesh = SkeletalMeshComponent->GetSkeletalMeshAsset())
 		{
 			if (const USkeleton* Skeleton = SkeletalMesh->GetSkeleton())
 			{
@@ -595,7 +595,7 @@ void UControlRigComponent::AddMappedSkeletalMesh(USkeletalMeshComponent* Skeleta
 	TArray<FControlRigComponentMappedCurve> CurvesToMap = Curves;
 	if (CurvesToMap.Num() == 0)
 	{
-		if (USkeletalMesh* SkeletalMesh = SkeletalMeshComponent->GetSkeletalMesh())
+		if (USkeletalMesh* SkeletalMesh = SkeletalMeshComponent->GetSkeletalMeshAsset())
 		{
 			if (USkeleton* Skeleton = SkeletalMesh->GetSkeleton())
 			{
@@ -1278,7 +1278,7 @@ void UControlRigComponent::ValidateMappingData()
 				}
 				else if (!MappedElement.TransformName.IsNone())
 				{
-					if (USkeletalMesh* SkeletalMesh = SkeletalMeshComponent->GetSkeletalMesh())
+					if (USkeletalMesh* SkeletalMesh = SkeletalMeshComponent->GetSkeletalMeshAsset())
 					{
 						if (USkeleton* Skeleton = SkeletalMesh->GetSkeleton())
 						{
diff --git a/Engine/Plugins/Animation/ControlRig/Source/ControlRig/Private/Rigs/FKControlRig.cpp b/Engine/Plugins/Animation/ControlRig/Source/ControlRig/Private/Rigs/FKControlRig.cpp
--- a/Engine/Plugins/Animation/ControlRig/Source/ControlRig/Private/Rigs/FKControlRig.cpp
+++ b/Engine/Plugins/Animation/ControlRig/Source/ControlRig/Private/Rigs/FKControlRig.cpp
@@ -264,7 +264,7 @@ void UFKControlRig::Initialize(bool bInitRigUnits /*= true*/)
 	// create hierarchy from the incoming skeleton
 	if (USkeletalMeshComponent* SkeletalMeshComponent = Cast<USkeletalMeshComponent>(GetObjectBinding()->GetBoundObject()))
 	{
-		CreateRigElements(SkeletalMeshComponent->GetSkeletalMesh());
+		CreateRigElements(SkeletalMeshComponent->GetSkeletalMeshAsset());
 	}
 	else if (USkeleton* Skeleton = Cast<USkeleton>(GetObjectBinding()->GetBoundObject()))
 	{
diff --git a/Engine/Plugins/Animation/ControlRig/Source/ControlRig/Private/Sequencer/ControlRigLayerInstance.cpp b/Engine/Plugins/Animation/ControlRig/Source/ControlRig/Private/Sequencer/ControlRigLayerInstance.cpp
--- a/Engine/Plugins/Animation/ControlRig/Source/ControlRig/Private/Sequencer/ControlRigLayerInstance.cpp
+++ b/Engine/Plugins/Animation/ControlRig/Source/ControlRig/Private/Sequencer/ControlRigLayerInstance.cpp
@@ -121,7 +121,7 @@ void UControlRigLayerInstance::SavePose()
 {
 	if (USkeletalMeshComponent* SkeletalMeshComponent = GetSkelMeshComponent())
 	{
-		if (SkeletalMeshComponent->GetSkeletalMesh() && SkeletalMeshComponent->GetComponentSpaceTransforms().Num() > 0)
+		if (SkeletalMeshComponent->GetSkeletalMeshAsset() && SkeletalMeshComponent->GetComponentSpaceTransforms().Num() > 0)
 		{
 			SavePoseSnapshot(UControlRigLayerInstance::SequencerPoseName);
 		}
diff --git a/Engine/Plugins/Animation/ControlRig/Source/ControlRig/Private/Sequencer/MovieSceneControlRigParameterSection.cpp b/Engine/Plugins/Animation/ControlRig/Source/ControlRig/Private/Sequencer/MovieSceneControlRigParameterSection.cpp
--- a/Engine/Plugins/Animation/ControlRig/Source/ControlRig/Private/Sequencer/MovieSceneControlRigParameterSection.cpp
+++ b/Engine/Plugins/Animation/ControlRig/Source/ControlRig/Private/Sequencer/MovieSceneControlRigParameterSection.cpp
@@ -2632,12 +2632,12 @@ bool UMovieSceneControlRigParameterSection::LoadAnimSequenceIntoThisSection(UAni
 {
 	USkeletalMeshComponent* SkelMeshComp = Cast<USkeletalMeshComponent>(BoundObject);
 	
-	if (SkelMeshComp != nullptr && (SkelMeshComp->GetSkeletalMesh() == nullptr || SkelMeshComp->GetSkeletalMesh()->GetSkeleton() == nullptr))
+	if (SkelMeshComp != nullptr && (SkelMeshComp->GetSkeletalMeshAsset() == nullptr || SkelMeshComp->GetSkeletalMeshAsset()->GetSkeleton() == nullptr))
 	{
 		return false;
 	}
 	
-	USkeleton* Skeleton = SkelMeshComp ? SkelMeshComp->GetSkeletalMesh()->GetSkeleton() : Cast<USkeleton>(BoundObject);
+	USkeleton* Skeleton = SkelMeshComp ? SkelMeshComp->GetSkeletalMeshAsset()->GetSkeleton() : Cast<USkeleton>(BoundObject);
 	if (Skeleton == nullptr)
 	{
 		return false;
@@ -2741,7 +2741,7 @@ bool UMovieSceneControlRigParameterSection::LoadAnimSequenceIntoThisSection(UAni
 		// retrieve the pose using the services that persona and sequencer rely on
 		// rather than accessing the low level raw tracks.
 		FAnimPoseEvaluationOptions EvaluationOptions;
-		EvaluationOptions.OptionalSkeletalMesh = SkelMeshComp ? SkelMeshComp->GetSkeletalMesh() : nullptr;
+		EvaluationOptions.OptionalSkeletalMesh = SkelMeshComp ? SkelMeshComp->GetSkeletalMeshAsset() : nullptr;
 		EvaluationOptions.bShouldRetarget = false;
 		EvaluationOptions.EvaluationType = EAnimDataEvalType::Raw;
 
diff --git a/Engine/Plugins/Animation/ControlRig/Source/ControlRig/Private/Sequencer/MovieSceneControlRigParameterTemplate.cpp b/Engine/Plugins/Animation/ControlRig/Source/ControlRig/Private/Sequencer/MovieSceneControlRigParameterTemplate.cpp
--- a/Engine/Plugins/Animation/ControlRig/Source/ControlRig/Private/Sequencer/MovieSceneControlRigParameterTemplate.cpp
+++ b/Engine/Plugins/Animation/ControlRig/Source/ControlRig/Private/Sequencer/MovieSceneControlRigParameterTemplate.cpp
@@ -588,7 +588,7 @@ void FControlRigBindingHelper::BindToSequencerInstance(UControlRig* ControlRig)
 	check(ControlRig);
 	if (USkeletalMeshComponent* SkeletalMeshComponent = Cast<USkeletalMeshComponent>(ControlRig->GetObjectBinding()->GetBoundObject()))
 	{
-		if(SkeletalMeshComponent->GetSkeletalMesh())
+		if(SkeletalMeshComponent->GetSkeletalMeshAsset())
 		{
 			bool bWasCreated = false;
 			if (UControlRigLayerInstance* AnimInstance = FAnimCustomInstanceHelper::BindToSkeletalMeshComponent<UControlRigLayerInstance>(SkeletalMeshComponent, bWasCreated))
diff --git a/Engine/Plugins/Animation/ControlRig/Source/ControlRig/Private/Units/Animation/RigUnit_AnimAttribute.cpp b/Engine/Plugins/Animation/ControlRig/Source/ControlRig/Private/Units/Animation/RigUnit_AnimAttribute.cpp
--- a/Engine/Plugins/Animation/ControlRig/Source/ControlRig/Private/Units/Animation/RigUnit_AnimAttribute.cpp
+++ b/Engine/Plugins/Animation/ControlRig/Source/ControlRig/Private/Units/Animation/RigUnit_AnimAttribute.cpp
@@ -67,7 +67,7 @@ FORCEINLINE_DEBUGGABLE static T* GetAnimAttributeValue(
 	const USkeletalMeshComponent* OwningComponent = Cast<USkeletalMeshComponent>(Context.OwningComponent);
 
 	if (!OwningComponent ||
-		!OwningComponent->GetSkeletalMesh())
+		!OwningComponent->GetSkeletalMeshAsset())
 	{
 		return nullptr;
 	}
@@ -82,7 +82,7 @@ FORCEINLINE_DEBUGGABLE static T* GetAnimAttributeValue(
 		// Invalidate cache if input changed
 		if (CachedBoneName != BoneName)
 		{
-			CachedBoneIndex = OwningComponent->GetSkeletalMesh()->GetRefSkeleton().FindBoneIndex(BoneName);
+			CachedBoneIndex = OwningComponent->GetSkeletalMeshAsset()->GetRefSkeleton().FindBoneIndex(BoneName);
 		}
 	}
 	
diff --git a/Engine/Plugins/Animation/ControlRig/Source/ControlRig/Private/Units/Execution/RigUnit_DynamicHierarchy.cpp b/Engine/Plugins/Animation/ControlRig/Source/ControlRig/Private/Units/Execution/RigUnit_DynamicHierarchy.cpp
--- a/Engine/Plugins/Animation/ControlRig/Source/ControlRig/Private/Units/Execution/RigUnit_DynamicHierarchy.cpp
+++ b/Engine/Plugins/Animation/ControlRig/Source/ControlRig/Private/Units/Execution/RigUnit_DynamicHierarchy.cpp
@@ -291,13 +291,13 @@ FRigUnit_HierarchyImportFromSkeleton_Execute()
 		
 		if(const USkeletalMeshComponent* SkelMeshComponent = Context.DataSourceRegistry->RequestSource<USkeletalMeshComponent>(UControlRig::OwnerComponent))
 		{
-			if(SkelMeshComponent->GetSkeletalMesh())
+			if(SkelMeshComponent->GetSkeletalMeshAsset())
 			{
-				const FReferenceSkeleton& ReferenceSkeleton = SkelMeshComponent->GetSkeletalMesh()->GetRefSkeleton();
+				const FReferenceSkeleton& ReferenceSkeleton = SkelMeshComponent->GetSkeletalMeshAsset()->GetRefSkeleton();
 				Items = Controller->ImportBones(ReferenceSkeleton, NameSpace, false, false, false, false);
 				if(bIncludeCurves)
 				{
-					if(USkeleton* Skeleton = SkelMeshComponent->GetSkeletalMesh()->GetSkeleton())
+					if(USkeleton* Skeleton = SkelMeshComponent->GetSkeletalMeshAsset()->GetSkeleton())
 					{
 						Items.Append(Controller->ImportCurves(Skeleton, NameSpace, false, false, false));
 					}
diff --git a/Engine/Plugins/Animation/ControlRig/Source/ControlRig/Public/Tools/ControlRigToolAsset.h b/Engine/Plugins/Animation/ControlRig/Source/ControlRig/Public/Tools/ControlRigToolAsset.h
--- a/Engine/Plugins/Animation/ControlRig/Source/ControlRig/Public/Tools/ControlRigToolAsset.h
+++ b/Engine/Plugins/Animation/ControlRig/Source/ControlRig/Public/Tools/ControlRigToolAsset.h
@@ -59,7 +59,7 @@ public:
 			USkeletalMesh* SkelMesh = nullptr;
 			if (USkeletalMeshComponent* RigMeshComp = Cast<USkeletalMeshComponent>(InControlRig->GetObjectBinding()->GetBoundObject()))
 			{
-				SkelMesh = RigMeshComp->GetSkeletalMesh();
+				SkelMesh = RigMeshComp->GetSkeletalMeshAsset();
 			}
 
 			Asset->SavePose(InControlRig, bUseAllControls);
diff --git a/Engine/Plugins/Animation/ControlRig/Source/ControlRigEditor/Private/ControlRigBlueprintActions.cpp b/Engine/Plugins/Animation/ControlRig/Source/ControlRigEditor/Private/ControlRigBlueprintActions.cpp
--- a/Engine/Plugins/Animation/ControlRig/Source/ControlRigEditor/Private/ControlRigBlueprintActions.cpp
+++ b/Engine/Plugins/Animation/ControlRig/Source/ControlRigEditor/Private/ControlRigBlueprintActions.cpp
@@ -384,7 +384,7 @@ void FControlRigBlueprintActions::OnSpawnedSkeletalMeshActorChanged(UObject* InO
 
 			if (!ControlRig || ControlRig->GetClass() != ControlRigClass)
 			{		
-				USkeletalMesh* SkeletalMesh = MeshActor->GetSkeletalMeshComponent()->GetSkeletalMesh();
+				USkeletalMesh* SkeletalMesh = MeshActor->GetSkeletalMeshComponent()->GetSkeletalMeshAsset();
 				USkeleton* Skeleton = SkeletalMesh->GetSkeleton();
 				
 				ObjectName.RemoveFromEnd(TEXT("_C"));
diff --git a/Engine/Plugins/Animation/ControlRig/Source/ControlRigEditor/Private/ControlRigEditorModule.cpp b/Engine/Plugins/Animation/ControlRig/Source/ControlRigEditor/Private/ControlRigEditorModule.cpp
--- a/Engine/Plugins/Animation/ControlRig/Source/ControlRigEditor/Private/ControlRigEditorModule.cpp
+++ b/Engine/Plugins/Animation/ControlRig/Source/ControlRigEditor/Private/ControlRigEditorModule.cpp
@@ -710,7 +710,7 @@ void FControlRigEditorModule::BakeToControlRig(UClass* ControlRigClass, UAnimSeq
 			if (Track)
 			{
 				USkeletalMeshComponent* SkelMeshComp = MeshActor->GetSkeletalMeshComponent();
-				USkeletalMesh* SkeletalMesh = SkelMeshComp->GetSkeletalMesh();
+				USkeletalMesh* SkeletalMesh = SkelMeshComp->GetSkeletalMeshAsset();
 
 				FString ObjectName = (ControlRigClass->GetName());
 				ObjectName.RemoveFromEnd(TEXT("_C"));
diff --git a/Engine/Plugins/Animation/ControlRig/Source/ControlRigEditor/Private/ControlRigSequencerEditorLibrary.cpp b/Engine/Plugins/Animation/ControlRig/Source/ControlRigEditor/Private/ControlRigSequencerEditorLibrary.cpp
--- a/Engine/Plugins/Animation/ControlRig/Source/ControlRigEditor/Private/ControlRigSequencerEditorLibrary.cpp
+++ b/Engine/Plugins/Animation/ControlRig/Source/ControlRigEditor/Private/ControlRigSequencerEditorLibrary.cpp
@@ -116,9 +116,9 @@ static void AcquireSkeletonAndSkelMeshCompFromObject(UObject* BoundObject, USkel
 			if (SkeletalMeshComp)
 			{
 				*OutSkeletalMeshComponent = SkeletalMeshComp;
-				if (SkeletalMeshComp->GetSkeletalMesh() && SkeletalMeshComp->GetSkeletalMesh()->GetSkeleton())
+				if (SkeletalMeshComp->GetSkeletalMeshAsset() && SkeletalMeshComp->GetSkeletalMeshAsset()->GetSkeleton())
 				{
-					*OutSkeleton = SkeletalMeshComp->GetSkeletalMesh()->GetSkeleton();
+					*OutSkeleton = SkeletalMeshComp->GetSkeletalMeshAsset()->GetSkeleton();
 				}
 				return;
 			}
@@ -133,9 +133,9 @@ static void AcquireSkeletonAndSkelMeshCompFromObject(UObject* BoundObject, USkel
 				if (SkeletalMeshComp)
 				{
 					*OutSkeletalMeshComponent = SkeletalMeshComp;
-					if (SkeletalMeshComp->GetSkeletalMesh() && SkeletalMeshComp->GetSkeletalMesh()->GetSkeleton())
+					if (SkeletalMeshComp->GetSkeletalMeshAsset() && SkeletalMeshComp->GetSkeletalMeshAsset()->GetSkeleton())
 					{
-						*OutSkeleton = SkeletalMeshComp->GetSkeletalMesh()->GetSkeleton();
+						*OutSkeleton = SkeletalMeshComp->GetSkeletalMeshAsset()->GetSkeleton();
 					}
 					return;
 				}
@@ -155,9 +155,9 @@ static void AcquireSkeletonAndSkelMeshCompFromObject(UObject* BoundObject, USkel
 					if (SkeletalMeshComp)
 					{
 						*OutSkeletalMeshComponent = SkeletalMeshComp;
-						if (SkeletalMeshComp->GetSkeletalMesh() && SkeletalMeshComp->GetSkeletalMesh()->GetSkeleton())
+						if (SkeletalMeshComp->GetSkeletalMeshAsset() && SkeletalMeshComp->GetSkeletalMeshAsset()->GetSkeleton())
 						{
-							*OutSkeleton = SkeletalMeshComp->GetSkeletalMesh()->GetSkeleton();
+							*OutSkeleton = SkeletalMeshComp->GetSkeletalMeshAsset()->GetSkeleton();
 						}
 					}
 				}
@@ -167,9 +167,9 @@ static void AcquireSkeletonAndSkelMeshCompFromObject(UObject* BoundObject, USkel
 	else if (USkeletalMeshComponent* SkeletalMeshComponent = Cast<USkeletalMeshComponent>(BoundObject))
 	{
 		*OutSkeletalMeshComponent = SkeletalMeshComponent;
-		if (SkeletalMeshComponent->GetSkeletalMesh() && SkeletalMeshComponent->GetSkeletalMesh()->GetSkeleton())
+		if (SkeletalMeshComponent->GetSkeletalMeshAsset() && SkeletalMeshComponent->GetSkeletalMeshAsset()->GetSkeleton())
 		{
-			*OutSkeleton = SkeletalMeshComponent->GetSkeletalMesh()->GetSkeleton();
+			*OutSkeleton = SkeletalMeshComponent->GetSkeletalMeshAsset()->GetSkeleton();
 		}
 	}
 }
@@ -669,7 +669,7 @@ bool UControlRigSequencerEditorLibrary::BakeToControlRig(UWorld* World, ULevelSe
 			USkeleton* Skeleton = nullptr;
 			USkeletalMeshComponent* SkeletalMeshComp = nullptr;
 			AcquireSkeletonAndSkelMeshCompFromObject(BoundObject, &Skeleton, &SkeletalMeshComp);
-			if (SkeletalMeshComp && SkeletalMeshComp->GetSkeletalMesh() && SkeletalMeshComp->GetSkeletalMesh()->GetSkeleton())
+			if (SkeletalMeshComp && SkeletalMeshComp->GetSkeletalMeshAsset() && SkeletalMeshComp->GetSkeletalMeshAsset()->GetSkeleton())
 			{
 				UAnimSequence* TempAnimSequence = NewObject<UAnimSequence>(GetTransientPackage(), NAME_None);
 				TempAnimSequence->SetSkeleton(Skeleton);
diff --git a/Engine/Plugins/Animation/ControlRig/Source/ControlRigEditor/Private/Sequencer/ControlRigParameterTrackEditor.cpp b/Engine/Plugins/Animation/ControlRig/Source/ControlRigEditor/Private/Sequencer/ControlRigParameterTrackEditor.cpp
--- a/Engine/Plugins/Animation/ControlRig/Source/ControlRigEditor/Private/Sequencer/ControlRigParameterTrackEditor.cpp
+++ b/Engine/Plugins/Animation/ControlRig/Source/ControlRigEditor/Private/Sequencer/ControlRigParameterTrackEditor.cpp
@@ -103,7 +103,7 @@ static USkeletalMeshComponent* AcquireSkeletalMeshFromObject(UObject* BoundObjec
 	}
 	else if (USkeletalMeshComponent* SkeletalMeshComponent = Cast<USkeletalMeshComponent>(BoundObject))
 	{
-		if (SkeletalMeshComponent->GetSkeletalMesh())
+		if (SkeletalMeshComponent->GetSkeletalMeshAsset())
 		{
 			return SkeletalMeshComponent;
 		}
@@ -116,10 +116,10 @@ static USkeletalMeshComponent* AcquireSkeletalMeshFromObject(UObject* BoundObjec
 static USkeleton* GetSkeletonFromComponent(UActorComponent* InComponent)
 {
 	USkeletalMeshComponent* SkeletalMeshComp = Cast<USkeletalMeshComponent>(InComponent);
-	if (SkeletalMeshComp && SkeletalMeshComp->GetSkeletalMesh() && SkeletalMeshComp->GetSkeletalMesh()->GetSkeleton())
+	if (SkeletalMeshComp && SkeletalMeshComp->GetSkeletalMeshAsset() && SkeletalMeshComp->GetSkeletalMeshAsset()->GetSkeleton())
 	{
 		// @todo Multiple actors, multiple components
-		return SkeletalMeshComp->GetSkeletalMesh()->GetSkeleton();
+		return SkeletalMeshComp->GetSkeletalMeshAsset()->GetSkeleton();
 	}
 
 	return nullptr;
diff --git a/Engine/Plugins/Animation/ControlRig/Source/ControlRigEditor/Private/Tools/ControlRigPoseThumbnailRenderer.cpp b/Engine/Plugins/Animation/ControlRig/Source/ControlRigEditor/Private/Tools/ControlRigPoseThumbnailRenderer.cpp
--- a/Engine/Plugins/Animation/ControlRig/Source/ControlRigEditor/Private/Tools/ControlRigPoseThumbnailRenderer.cpp
+++ b/Engine/Plugins/Animation/ControlRig/Source/ControlRigEditor/Private/Tools/ControlRigPoseThumbnailRenderer.cpp
@@ -62,7 +62,7 @@ bool FControlRigPoseThumbnailScene::SetControlRigPoseAsset(UControlRigPoseAsset*
 
 	if (PoseAsset)
 	{
-		USkeletalMesh* SkeletalMesh = nullptr; // PoseAsset->GetSkeletalMesh();
+		USkeletalMesh* SkeletalMesh = nullptr; // PoseAsset->GetSkeletalMeshAsset();
 		UControlRig* ControlRig = nullptr; // PoseAsset->GetControlRig();
 		PreviewActor->GetSkeletalMeshComponent()->OverrideMaterials.Empty();
 
@@ -149,7 +149,7 @@ void FControlRigPoseThumbnailScene::CleanupComponentChildren(USceneComponent* Co
 void FControlRigPoseThumbnailScene::GetViewMatrixParameters(const float InFOVDegrees, FVector& OutOrigin, float& OutOrbitPitch, float& OutOrbitYaw, float& OutOrbitZoom) const
 {
 	check(PreviewActor->GetSkeletalMeshComponent());
-	check(PreviewActor->GetSkeletalMeshComponent()->GetSkeletalMesh());
+	check(PreviewActor->GetSkeletalMeshComponent()->GetSkeletalMeshAsset());
 
 	const float HalfFOVRadians = FMath::DegreesToRadians<float>(InFOVDegrees) * 0.5f;
 	// No need to add extra size to view slightly outside of the sphere to compensate for perspective since skeletal meshes already buffer bounds.
diff --git a/Engine/Plugins/Animation/GameplayInsights/Source/GameplayInsightsEditor/Private/InsightsSkeletalMeshComponent.cpp b/Engine/Plugins/Animation/GameplayInsights/Source/GameplayInsightsEditor/Private/InsightsSkeletalMeshComponent.cpp
--- a/Engine/Plugins/Animation/GameplayInsights/Source/GameplayInsightsEditor/Private/InsightsSkeletalMeshComponent.cpp
+++ b/Engine/Plugins/Animation/GameplayInsights/Source/GameplayInsightsEditor/Private/InsightsSkeletalMeshComponent.cpp
@@ -29,9 +29,9 @@ void UInsightsSkeletalMeshComponent::SetPoseFromProvider(const IAnimationProvide
 
 void UInsightsSkeletalMeshComponent::InitAnim(bool bForceReInit)
 {
-	if(GetSkeletalMesh())
+	if(GetSkeletalMeshAsset())
 	{
-		const FReferenceSkeleton& SkeletalMeshRefSkeleton = GetSkeletalMesh()->GetRefSkeleton();
+		const FReferenceSkeleton& SkeletalMeshRefSkeleton = GetSkeletalMeshAsset()->GetRefSkeleton();
 		// set up bone visibility states as this gets skipped since we allocate the component array before registration
 		for (int32 BaseIndex = 0; BaseIndex < 2; ++BaseIndex)
 		{
@@ -43,6 +43,6 @@ void UInsightsSkeletalMeshComponent::InitAnim(bool bForceReInit)
 			}
 		}
 
-		MorphTargetWeights.SetNumZeroed(GetSkeletalMesh()->GetMorphTargets().Num());
+		MorphTargetWeights.SetNumZeroed(GetSkeletalMeshAsset()->GetMorphTargets().Num());
 	}
 }
\ No newline at end of file
diff --git a/Engine/Plugins/Animation/IKRig/Source/IKRig/Private/AnimNodes/AnimNode_IKRig.cpp b/Engine/Plugins/Animation/IKRig/Source/IKRig/Private/AnimNodes/AnimNode_IKRig.cpp
--- a/Engine/Plugins/Animation/IKRig/Source/IKRig/Private/AnimNodes/AnimNode_IKRig.cpp
+++ b/Engine/Plugins/Animation/IKRig/Source/IKRig/Private/AnimNodes/AnimNode_IKRig.cpp
@@ -225,7 +225,7 @@ void FAnimNode_IKRig::PreUpdate(const UAnimInstance* InAnimInstance)
 	// initialize the IK Rig (will only try once on the current version of the rig asset)
 	if (!IKRigProcessor->IsInitialized())
 	{
- 		IKRigProcessor->Initialize(RigDefinitionAsset, InAnimInstance->GetSkelMeshComponent()->GetSkeletalMesh());
+ 		IKRigProcessor->Initialize(RigDefinitionAsset, InAnimInstance->GetSkelMeshComponent()->GetSkeletalMeshAsset());
 	}
 	
 	// cache list of goal creator components on the actor
diff --git a/Engine/Plugins/Animation/IKRig/Source/IKRig/Private/AnimNodes/AnimNode_RetargetPoseFromMesh.cpp b/Engine/Plugins/Animation/IKRig/Source/IKRig/Private/AnimNodes/AnimNode_RetargetPoseFromMesh.cpp
--- a/Engine/Plugins/Animation/IKRig/Source/IKRig/Private/AnimNodes/AnimNode_RetargetPoseFromMesh.cpp
+++ b/Engine/Plugins/Animation/IKRig/Source/IKRig/Private/AnimNodes/AnimNode_RetargetPoseFromMesh.cpp
@@ -28,7 +28,7 @@ void FAnimNode_RetargetPoseFromMesh::CacheBones_AnyThread(const FAnimationCacheB
 		return;
 	}
 	
-	if (!Context.AnimInstanceProxy->GetSkelMeshComponent()->GetSkeletalMesh())
+	if (!Context.AnimInstanceProxy->GetSkelMeshComponent()->GetSkeletalMeshAsset())
 	{
 		return;
 	}
@@ -37,7 +37,7 @@ void FAnimNode_RetargetPoseFromMesh::CacheBones_AnyThread(const FAnimationCacheB
 	RequiredToTargetBoneMapping.Reset();
 
 	const FReferenceSkeleton& RefSkeleton = RequiredBones.GetReferenceSkeleton();
-	const FReferenceSkeleton& TargetSkeleton = Context.AnimInstanceProxy->GetSkelMeshComponent()->GetSkeletalMesh()->GetRefSkeleton();
+	const FReferenceSkeleton& TargetSkeleton = Context.AnimInstanceProxy->GetSkelMeshComponent()->GetSkeletalMeshAsset()->GetRefSkeleton();
 	const TArray<FBoneIndexType>& RequiredBonesArray = RequiredBones.GetBoneIndicesArray();
 	for (int32 Index = 0; Index < RequiredBonesArray.Num(); ++Index)
 	{
@@ -83,8 +83,8 @@ void FAnimNode_RetargetPoseFromMesh::Evaluate_AnyThread(FPoseContext& Output)
 
 	// ensure processor was initialized with the currently used assets (source/target meshes and retarget asset)
 	// if processor is not ready this tick, it will be next tick as this state will trigger re-initialization
-	const TObjectPtr<USkeletalMesh> SourceMesh = SourceMeshComponent->GetSkeletalMesh();
-	const TObjectPtr<USkeletalMesh> TargetMesh = Output.AnimInstanceProxy->GetSkelMeshComponent()->GetSkeletalMesh();
+	const TObjectPtr<USkeletalMesh> SourceMesh = SourceMeshComponent->GetSkeletalMeshAsset();
+	const TObjectPtr<USkeletalMesh> TargetMesh = Output.AnimInstanceProxy->GetSkelMeshComponent()->GetSkeletalMeshAsset();
 	const bool bIsProcessorReady = Processor->WasInitializedWithTheseAssets(SourceMesh, TargetMesh, IKRetargeterAsset);
 
 	// if not ready to run, skip retarget and output the ref pose
@@ -239,8 +239,8 @@ bool FAnimNode_RetargetPoseFromMesh::EnsureProcessorIsInitialized(const TObjectP
 	}
 
 	// check that both a source and target mesh exist
-	const TObjectPtr<USkeletalMesh> SourceMesh = SourceMeshComponent->GetSkeletalMesh();
-	const TObjectPtr<USkeletalMesh> TargetMesh = TargetMeshComponent->GetSkeletalMesh();
+	const TObjectPtr<USkeletalMesh> SourceMesh = SourceMeshComponent->GetSkeletalMeshAsset();
+	const TObjectPtr<USkeletalMesh> TargetMesh = TargetMeshComponent->GetSkeletalMeshAsset();
 	if (SourceMesh.IsNull() || TargetMesh.IsNull())
 	{
 		return false; // cannot initialize if components are missing skeletal mesh references
diff --git a/Engine/Plugins/Animation/IKRig/Source/IKRigEditor/Private/RetargetEditor/IKRetargetAnimInstance.cpp b/Engine/Plugins/Animation/IKRig/Source/IKRigEditor/Private/RetargetEditor/IKRetargetAnimInstance.cpp
--- a/Engine/Plugins/Animation/IKRig/Source/IKRigEditor/Private/RetargetEditor/IKRetargetAnimInstance.cpp
+++ b/Engine/Plugins/Animation/IKRig/Source/IKRigEditor/Private/RetargetEditor/IKRetargetAnimInstance.cpp
@@ -15,7 +15,7 @@ void FAnimNode_PreviewRetargetPose::CacheBones_AnyThread(const FAnimationCacheBo
 		return;
 	}
 	
-	if (!Context.AnimInstanceProxy->GetSkelMeshComponent()->GetSkeletalMesh())
+	if (!Context.AnimInstanceProxy->GetSkelMeshComponent()->GetSkeletalMeshAsset())
 	{
 		return;
 	}
@@ -24,7 +24,7 @@ void FAnimNode_PreviewRetargetPose::CacheBones_AnyThread(const FAnimationCacheBo
 	RequiredBoneToMeshBoneMap.Reset();
 
 	const FReferenceSkeleton& RefSkeleton = RequiredBones.GetReferenceSkeleton();
-	const FReferenceSkeleton& TargetSkeleton = Context.AnimInstanceProxy->GetSkelMeshComponent()->GetSkeletalMesh()->GetRefSkeleton();
+	const FReferenceSkeleton& TargetSkeleton = Context.AnimInstanceProxy->GetSkelMeshComponent()->GetSkeletalMeshAsset()->GetRefSkeleton();
 	const TArray<FBoneIndexType>& RequiredBonesArray = RequiredBones.GetBoneIndicesArray();
 	for (int32 RequiredBoneIndex = 0; RequiredBoneIndex < RequiredBonesArray.Num(); ++RequiredBoneIndex)
 	{
@@ -65,7 +65,7 @@ void FAnimNode_PreviewRetargetPose::Evaluate_AnyThread(FPoseContext& Output)
 	
 	// generate full (no LOD) LOCAL retarget pose by applying the local rotation offsets from the stored retarget pose
 	// these poses are read back by the editor and need to be complete (not culled)
-	const FReferenceSkeleton& RefSkeleton = Output.AnimInstanceProxy->GetSkelMeshComponent()->GetSkeletalMesh()->GetRefSkeleton();
+	const FReferenceSkeleton& RefSkeleton = Output.AnimInstanceProxy->GetSkelMeshComponent()->GetSkeletalMeshAsset()->GetRefSkeleton();
 	const TArray<FTransform>& RefPose = RefSkeleton.GetRefBonePose(); 
 	RetargetLocalPose = RefPose;
 	for (const TTuple<FName, FQuat>& BoneDelta : RetargetPose->GetAllDeltaRotations())
diff --git a/Engine/Plugins/Animation/IKRig/Source/IKRigEditor/Private/RetargetEditor/IKRetargetDetails.cpp b/Engine/Plugins/Animation/IKRig/Source/IKRigEditor/Private/RetargetEditor/IKRetargetDetails.cpp
--- a/Engine/Plugins/Animation/IKRig/Source/IKRigEditor/Private/RetargetEditor/IKRetargetDetails.cpp
+++ b/Engine/Plugins/Animation/IKRig/Source/IKRigEditor/Private/RetargetEditor/IKRetargetDetails.cpp
@@ -35,13 +35,13 @@ FEulerTransform UIKRetargetBoneDetails::GetTransform(EIKRetargetTransformType Tr
 	// ensure we have a valid skeletal mesh
 	const bool bEditingSource = EditorController->GetSourceOrTarget() == ERetargetSourceOrTarget::Source;
 	UDebugSkelMeshComponent* Mesh = bEditingSource ? Controller->SourceSkelMeshComponent : Controller->TargetSkelMeshComponent;
-	if (!(Mesh && Mesh->GetSkeletalMesh()))
+	if (!(Mesh && Mesh->GetSkeletalMeshAsset()))
 	{
 		return FEulerTransform::Identity;
 	}
 
 	// ensure bone is valid
-	const FReferenceSkeleton& RefSkeleton = Mesh->GetSkeletalMesh()->GetRefSkeleton();
+	const FReferenceSkeleton& RefSkeleton = Mesh->GetSkeletalMeshAsset()->GetRefSkeleton();
 	const int32 BoneIndex = RefSkeleton.FindBoneIndex(SelectedBone);
 	if (BoneIndex == INDEX_NONE)
 	{
@@ -383,13 +383,13 @@ void UIKRetargetBoneDetails::OnNumericValueCommitted(
 	// ensure we have a valid skeletal mesh
 	const bool bEditingSource = EditorController->GetSourceOrTarget() == ERetargetSourceOrTarget::Source;
 	UDebugSkelMeshComponent* Mesh = bEditingSource ? EditorController->SourceSkelMeshComponent : EditorController->TargetSkelMeshComponent;
-	if (!(Mesh && Mesh->GetSkeletalMesh()))
+	if (!(Mesh && Mesh->GetSkeletalMeshAsset()))
 	{
 		return;
 	}
 
 	// ensure bone is valid
-	const FReferenceSkeleton& RefSkeleton = Mesh->GetSkeletalMesh()->GetRefSkeleton();
+	const FReferenceSkeleton& RefSkeleton = Mesh->GetSkeletalMeshAsset()->GetRefSkeleton();
 	const int32 BoneIndex = RefSkeleton.FindBoneIndex(SelectedBone);
 	if (BoneIndex == INDEX_NONE)
 	{
diff --git a/Engine/Plugins/Animation/IKRig/Source/IKRigEditor/Private/RetargetEditor/IKRetargetEditorController.cpp b/Engine/Plugins/Animation/IKRig/Source/IKRigEditor/Private/RetargetEditor/IKRetargetEditorController.cpp
--- a/Engine/Plugins/Animation/IKRig/Source/IKRigEditor/Private/RetargetEditor/IKRetargetEditorController.cpp
+++ b/Engine/Plugins/Animation/IKRig/Source/IKRigEditor/Private/RetargetEditor/IKRetargetEditorController.cpp
@@ -519,7 +519,7 @@ void FIKRetargetEditorController::EditBoneSelection(
 {
 	// must have a skeletal mesh
 	UDebugSkelMeshComponent* DebugComponent = GetEditedSkeletalMesh();
-	if (!DebugComponent->GetSkeletalMesh())
+	if (!DebugComponent->GetSkeletalMeshAsset())
 	{
 		return;
 	}
@@ -553,7 +553,7 @@ void FIKRetargetEditorController::EditBoneSelection(
 	}
 
 	// convert to bone indices
-	const FReferenceSkeleton& RefSkeleton = DebugComponent->GetSkeletalMesh()->GetRefSkeleton();
+	const FReferenceSkeleton& RefSkeleton = DebugComponent->GetSkeletalMeshAsset()->GetRefSkeleton();
 	TArray<int32> SelectedBoneIndices;
 	for (const FName& Bone : SelectedBones)
 	{
diff --git a/Engine/Plugins/Animation/IKRig/Source/IKRigEditor/Private/RigEditor/IKRigEditorController.cpp b/Engine/Plugins/Animation/IKRig/Source/IKRigEditor/Private/RigEditor/IKRigEditorController.cpp
--- a/Engine/Plugins/Animation/IKRig/Source/IKRigEditor/Private/RigEditor/IKRigEditorController.cpp
+++ b/Engine/Plugins/Animation/IKRig/Source/IKRigEditor/Private/RigEditor/IKRigEditorController.cpp
@@ -52,7 +52,7 @@ TOptional<FTransform> UIKRigBoneDetails::GetTransform(EIKRigTransformType::Type
 			
 			USkeletalMeshComponent* SkeletalMeshComponent = AnimInstancePtr->GetSkelMeshComponent();
 			const bool IsSkelMeshValid = SkeletalMeshComponent != nullptr &&
-										SkeletalMeshComponent->GetSkeletalMesh() != nullptr;
+										SkeletalMeshComponent->GetSkeletalMeshAsset() != nullptr;
 			if (IsSkelMeshValid)
 			{
 				GlobalTransform = SkeletalMeshComponent->GetBoneTransform(BoneIndex);
diff --git a/Engine/Plugins/Animation/RigLogic/Source/RigLogicModule/Private/AnimNode_RigLogic.cpp b/Engine/Plugins/Animation/RigLogic/Source/RigLogicModule/Private/AnimNode_RigLogic.cpp
--- a/Engine/Plugins/Animation/RigLogic/Source/RigLogicModule/Private/AnimNode_RigLogic.cpp
+++ b/Engine/Plugins/Animation/RigLogic/Source/RigLogicModule/Private/AnimNode_RigLogic.cpp
@@ -39,7 +39,7 @@ void FAnimNode_RigLogic::Initialize_AnyThread(const FAnimationInitializeContext&
 	FAnimNode_Base::Initialize_AnyThread(Context);
 	AnimSequence.Initialize(Context);
 	USkeletalMeshComponent* SkeletalMeshComponent = Context.AnimInstanceProxy->GetSkelMeshComponent();
-	USkeletalMesh* SkeletalMesh = SkeletalMeshComponent->GetSkeletalMesh();
+	USkeletalMesh* SkeletalMesh = SkeletalMeshComponent->GetSkeletalMeshAsset();
 
 	static FRWLock GlobalCreationLock;
 	FRWScopeLock ScopeLock(GlobalCreationLock, SLT_ReadOnly);
@@ -66,7 +66,7 @@ void FAnimNode_RigLogic::CacheBones_AnyThread(const FAnimationCacheBonesContext&
 
 	// Initialize things that depend on the skeleton of Anim BP
 	USkeletalMeshComponent* SkeletalMeshComponent = Context.AnimInstanceProxy->GetSkelMeshComponent();
-	USkeletalMesh* SkeletalMesh = SkeletalMeshComponent->GetSkeletalMesh();
+	USkeletalMesh* SkeletalMesh = SkeletalMeshComponent->GetSkeletalMeshAsset();
 	USkeleton* Skeleton = Context.AnimInstanceProxy->GetSkeleton();
 	UDNAAsset* DNAAsset = Cast<UDNAAsset>(SkeletalMesh->GetAssetUserDataOfClass(UDNAAsset::StaticClass()));
 	UDNAIndexMapping* DNAIndexMappingContainer = Cast<UDNAIndexMapping>(SkeletalMesh->GetAssetUserDataOfClass(UDNAIndexMapping::StaticClass()));
diff --git a/Engine/Plugins/Animation/RigLogic/Source/RigLogicModule/Private/RigUnit_RigLogic.cpp b/Engine/Plugins/Animation/RigLogic/Source/RigLogicModule/Private/RigUnit_RigLogic.cpp
--- a/Engine/Plugins/Animation/RigLogic/Source/RigLogicModule/Private/RigUnit_RigLogic.cpp
+++ b/Engine/Plugins/Animation/RigLogic/Source/RigLogicModule/Private/RigUnit_RigLogic.cpp
@@ -345,14 +345,14 @@ FRigUnit_RigLogic_Execute()
 					// however, during unit testing we cannot fetch it from DataSourceRegistry 
 					// in that case, a mock version will be inserted into Data by unit test beforehand
 				}
-				if (!Data.SkelMeshComponent.IsValid() || Data.SkelMeshComponent->GetSkeletalMesh() == nullptr)
+				if (!Data.SkelMeshComponent.IsValid() || Data.SkelMeshComponent->GetSkeletalMeshAsset() == nullptr)
 				{
 					return;
 				}
 				Data.CurrentLOD = Data.SkelMeshComponent->GetPredictedLODLevel();
 
 				// Fetch shared runtime context of rig from DNAAsset
-				TSharedPtr<FSharedRigRuntimeContext> RigRuntimeContext = GetSharedRigRuntimeContext(Data.SkelMeshComponent->GetSkeletalMesh());
+				TSharedPtr<FSharedRigRuntimeContext> RigRuntimeContext = GetSharedRigRuntimeContext(Data.SkelMeshComponent->GetSkeletalMeshAsset());
 				// Context is initialized with a BehaviorReader, which can be imported into SkeletalMesh from DNA file
 				// or overwritten by GeneSplicer when making a new character
 				Data.InitializeRigLogic(Hierarchy, RigRuntimeContext);
diff --git a/Engine/Plugins/Animation/RigLogic/Source/RigLogicModule/Private/RigUnit_RigLogic_Test.cpp b/Engine/Plugins/Animation/RigLogic/Source/RigLogicModule/Private/RigUnit_RigLogic_Test.cpp
--- a/Engine/Plugins/Animation/RigLogic/Source/RigLogicModule/Private/RigUnit_RigLogic_Test.cpp
+++ b/Engine/Plugins/Animation/RigLogic/Source/RigLogicModule/Private/RigUnit_RigLogic_Test.cpp
@@ -277,7 +277,7 @@ IMPLEMENT_RIGUNIT_AUTOMATION_TEST(FRigUnit_RigLogic)
 
 	//Test
 	InitAndExecute();
-	TSharedPtr<FSharedRigRuntimeContext> SharedRigRuntimeContext = Test.GetSharedRigRuntimeContext(Test.GetData()->SkelMeshComponent->GetSkeletalMesh());
+	TSharedPtr<FSharedRigRuntimeContext> SharedRigRuntimeContext = Test.GetSharedRigRuntimeContext(Test.GetData()->SkelMeshComponent->GetSkeletalMeshAsset());
 	//Assert
 	AddErrorIfFalse(
 		// Check rig logic initialized
diff --git a/Engine/Plugins/Developer/AnimationSharing/Source/AnimationSharing/Private/AnimationSharingManager.cpp b/Engine/Plugins/Developer/AnimationSharing/Source/AnimationSharing/Private/AnimationSharingManager.cpp
--- a/Engine/Plugins/Developer/AnimationSharing/Source/AnimationSharing/Private/AnimationSharingManager.cpp
+++ b/Engine/Plugins/Developer/AnimationSharing/Source/AnimationSharing/Private/AnimationSharingManager.cpp
@@ -366,7 +366,7 @@ void UAnimationSharingManager::RegisterActor(AActor* InActor, FUpdateActorHandle
 			const USkeleton* CurrentSkeleton = nullptr;
 			for (USkeletalMeshComponent* SkeletalMeshComponent : OwnedComponents)
 			{
-				const USkeletalMesh* Mesh = SkeletalMeshComponent->GetSkeletalMesh();
+				const USkeletalMesh* Mesh = SkeletalMeshComponent->GetSkeletalMeshAsset();
 				const USkeleton* Skeleton = Mesh->GetSkeleton();
 
 				if (CurrentSkeleton == nullptr)
diff --git a/Engine/Plugins/Editor/SequencerAnimTools/Source/SequencerAnimTools/Private/Sequencer/AnimationBoneTrail.cpp b/Engine/Plugins/Editor/SequencerAnimTools/Source/SequencerAnimTools/Private/Sequencer/AnimationBoneTrail.cpp
--- a/Engine/Plugins/Editor/SequencerAnimTools/Source/SequencerAnimTools/Private/Sequencer/AnimationBoneTrail.cpp
+++ b/Engine/Plugins/Editor/SequencerAnimTools/Source/SequencerAnimTools/Private/Sequencer/AnimationBoneTrail.cpp
@@ -29,14 +29,14 @@ void FAnimTrajectoryCache::Evaluate(FTrajectoryCache* ParentTransformCache)
 	TSharedPtr<ISequencer> Sequencer = WeakSequencer.Pin();
 	check(Sequencer);
 
-	if (!SkeletalMeshComponent.IsValid() || !SkeletalMeshComponent->GetSkeletalMesh() || !SkeletalMeshComponent->GetSkeletalMesh()->GetSkeleton())
+	if (!SkeletalMeshComponent.IsValid() || !SkeletalMeshComponent->GetSkeletalMeshAsset() || !SkeletalMeshComponent->GetSkeletalMeshAsset()->GetSkeleton())
 	{
 		return;
 	}
 
 	SkeletalMeshComponent->LeaderPoseComponent = nullptr;
 
-	CachedAnimSequence->SetSkeleton(SkeletalMeshComponent->GetSkeletalMesh()->GetSkeleton());
+	CachedAnimSequence->SetSkeleton(SkeletalMeshComponent->GetSkeletalMeshAsset()->GetSkeleton());
 
 	// TODO: for some reason SkeletalMeshComponent becomes invalid sometimes when evaluating, no clue why yet
 	FMovieSceneSequenceTransform MovieSceneSequenceTransform;
diff --git a/Engine/Plugins/Editor/SequencerAnimTools/Source/SequencerAnimTools/Private/Sequencer/AnimationBoneTrail.h b/Engine/Plugins/Editor/SequencerAnimTools/Source/SequencerAnimTools/Private/Sequencer/AnimationBoneTrail.h
--- a/Engine/Plugins/Editor/SequencerAnimTools/Source/SequencerAnimTools/Private/Sequencer/AnimationBoneTrail.h
+++ b/Engine/Plugins/Editor/SequencerAnimTools/Source/SequencerAnimTools/Private/Sequencer/AnimationBoneTrail.h
@@ -28,7 +28,7 @@ public:
 		, Spacing()
 		, bDirty(true)
 	{
-		CachedAnimSequence->SetSkeleton(InSkeletalMeshComponent->GetSkeletalMesh()->GetSkeleton());
+		CachedAnimSequence->SetSkeleton(InSkeletalMeshComponent->GetSkeletalMeshAsset()->GetSkeleton());
 	}
 
 	void Evaluate(FTrajectoryCache* ParentTrajectoryCache);
diff --git a/Engine/Plugins/Editor/SequencerAnimTools/Source/SequencerAnimTools/Private/Sequencer/SequencerTrailHierarchy.cpp b/Engine/Plugins/Editor/SequencerAnimTools/Source/SequencerAnimTools/Private/Sequencer/SequencerTrailHierarchy.cpp
--- a/Engine/Plugins/Editor/SequencerAnimTools/Source/SequencerAnimTools/Private/Sequencer/SequencerTrailHierarchy.cpp
+++ b/Engine/Plugins/Editor/SequencerAnimTools/Source/SequencerAnimTools/Private/Sequencer/SequencerTrailHierarchy.cpp
@@ -412,7 +412,7 @@ void FSequencerTrailHierarchy::OnBoneVisibilityChanged(class USkeleton* Skeleton
 				BoundComponent = BoundActor->FindComponentByClass<USkeletalMeshComponent>();
 			}
 
-			if (!BoundComponent || !BoundComponent->GetSkeletalMesh() || !BoundComponent->GetSkeletalMesh()->GetSkeleton() || !(BoundComponent->GetSkeletalMesh()->GetSkeleton() == Skeleton) || !BonesTracked.Contains(BoundComponent))
+			if (!BoundComponent || !BoundComponent->GetSkeletalMeshAsset() || !BoundComponent->GetSkeletalMeshAsset()->GetSkeleton() || !(BoundComponent->GetSkeletalMeshAsset()->GetSkeleton() == Skeleton) || !BonesTracked.Contains(BoundComponent))
 			{
 				continue;
 			}
@@ -534,7 +534,7 @@ void FSequencerTrailHierarchy::UpdateSequencerBindings(const TArray<FGuid>& Sequ
 					BoundComponent = BoundActor->FindComponentByClass<USkeletalMeshComponent>();
 				}
 
-				if (!BoundComponent || !BoundComponent->GetSkeletalMesh() || !BoundComponent->GetSkeletalMesh()->GetSkeleton())
+				if (!BoundComponent || !BoundComponent->GetSkeletalMeshAsset() || !BoundComponent->GetSkeletalMeshAsset()->GetSkeleton())
 				{
 					continue;
 				}
@@ -586,7 +586,7 @@ void FSequencerTrailHierarchy::UpdateSequencerBindings(const TArray<FGuid>& Sequ
 					BoundComponent = BoundActor->FindComponentByClass<USkeletalMeshComponent>();
 				}
 
-				if (!BoundComponent || !BoundComponent->GetSkeletalMesh() || !BoundComponent->GetSkeletalMesh()->GetSkeleton())
+				if (!BoundComponent || !BoundComponent->GetSkeletalMeshAsset() || !BoundComponent->GetSkeletalMeshAsset()->GetSkeleton())
 				{
 					continue;
 				};
@@ -676,7 +676,7 @@ void FSequencerTrailHierarchy::AddSkeletonToHierarchy(class USkeletalMeshCompone
 	TSharedPtr<FAnimTrajectoryCache> AnimTrajectoryCache = MakeShared<FAnimTrajectoryCache>(CompToAdd, Sequencer);
 	TMap<FName, FGuid>& BoneMap = BonesTracked.Add(CompToAdd, TMap<FName, FGuid>());
 	
-	USkeleton* MySkeleton = CompToAdd->GetSkeletalMesh()->GetSkeleton();
+	USkeleton* MySkeleton = CompToAdd->GetSkeletalMeshAsset()->GetSkeleton();
 	const int32 NumBones = MySkeleton->GetReferenceSkeleton().GetNum();
 	for (int32 BoneIdx = 0; BoneIdx < NumBones; BoneIdx++)
 	{
diff --git a/Engine/Plugins/Enterprise/DataprepEditor/Source/DataprepEditor/Private/DataprepStats.cpp b/Engine/Plugins/Enterprise/DataprepEditor/Source/DataprepEditor/Private/DataprepStats.cpp
--- a/Engine/Plugins/Enterprise/DataprepEditor/Source/DataprepEditor/Private/DataprepStats.cpp
+++ b/Engine/Plugins/Enterprise/DataprepEditor/Source/DataprepEditor/Private/DataprepStats.cpp
@@ -296,7 +296,7 @@ private:
 			}
 			else if( SkeletalMeshComponent )
 			{
-				USkeletalMesh* SkeletalMesh = SkeletalMeshComponent->GetSkeletalMesh();
+				USkeletalMesh* SkeletalMesh = SkeletalMeshComponent->GetSkeletalMeshAsset();
 				if( SkeletalMesh )
 				{
 					SkeletalMeshes.Add(SkeletalMesh);
diff --git a/Engine/Plugins/Experimental/Animation/CurveExpression/Source/Runtime/Private/AnimNode_RemapCurvesFromMesh.cpp b/Engine/Plugins/Experimental/Animation/CurveExpression/Source/Runtime/Private/AnimNode_RemapCurvesFromMesh.cpp
--- a/Engine/Plugins/Experimental/Animation/CurveExpression/Source/Runtime/Private/AnimNode_RemapCurvesFromMesh.cpp
+++ b/Engine/Plugins/Experimental/Animation/CurveExpression/Source/Runtime/Private/AnimNode_RemapCurvesFromMesh.cpp
@@ -136,7 +136,7 @@ void FAnimNode_RemapCurvesFromMesh::GatherDebugData(
 	DECLARE_SCOPE_HIERARCHICAL_COUNTER_ANIMNODE(GatherDebugData)
 	FString DebugLine = DebugData.GetNodeName(this);
 
-	DebugLine += FString::Printf(TEXT("('%s')"), *GetNameSafe(CurrentlyUsedSourceMeshComponent.IsValid() ? CurrentlyUsedSourceMeshComponent.Get()->GetSkeletalMesh() : nullptr));
+	DebugLine += FString::Printf(TEXT("('%s')"), *GetNameSafe(CurrentlyUsedSourceMeshComponent.IsValid() ? CurrentlyUsedSourceMeshComponent.Get()->GetSkeletalMeshAsset() : nullptr));
 	DebugData.AddDebugItem(DebugLine, true);
 }
 
@@ -153,7 +153,7 @@ void FAnimNode_RemapCurvesFromMesh::PreUpdate(
 
 	const USkeletalMeshComponent* CurrentMeshComponent = CurrentlyUsedSourceMeshComponent.IsValid() ? CurrentlyUsedSourceMeshComponent.Get() : nullptr;
 
-	if (CurrentMeshComponent && CurrentMeshComponent->GetSkeletalMesh() && CurrentMeshComponent->IsRegistered())
+	if (CurrentMeshComponent && CurrentMeshComponent->GetSkeletalMeshAsset() && CurrentMeshComponent->IsRegistered())
 	{
 		// If our source is running under leader-pose, then get bone data from there
 		if(USkeletalMeshComponent* LeaderPoseComponent = Cast<USkeletalMeshComponent>(CurrentMeshComponent->LeaderPoseComponent.Get()))
@@ -162,7 +162,7 @@ void FAnimNode_RemapCurvesFromMesh::PreUpdate(
 		}
 
 		// re-check mesh component validity as it may have changed to leader
-		if(CurrentMeshComponent->GetSkeletalMesh() && CurrentMeshComponent->IsRegistered())
+		if(CurrentMeshComponent->GetSkeletalMeshAsset() && CurrentMeshComponent->IsRegistered())
 		{
 			if (ExpressionEngine.IsSet())
 			{
@@ -197,10 +197,10 @@ void FAnimNode_RemapCurvesFromMesh::ReinitializeMeshComponent(
 	CurveNameToUIDMap.Reset();
 	CachedExpressions.Reset();
 
-	if (InTargetMeshComponent && IsValid(InNewSkeletalMeshComponent) && InNewSkeletalMeshComponent->GetSkeletalMesh())
+	if (InTargetMeshComponent && IsValid(InNewSkeletalMeshComponent) && InNewSkeletalMeshComponent->GetSkeletalMeshAsset())
 	{
-		USkeletalMesh* SourceSkelMesh = InNewSkeletalMeshComponent->GetSkeletalMesh();
-		USkeletalMesh* TargetSkelMesh = InTargetMeshComponent->GetSkeletalMesh();
+		USkeletalMesh* SourceSkelMesh = InNewSkeletalMeshComponent->GetSkeletalMeshAsset();
+		USkeletalMesh* TargetSkelMesh = InTargetMeshComponent->GetSkeletalMeshAsset();
 		
 		if (IsValid(SourceSkelMesh) && !SourceSkelMesh->HasAnyFlags(RF_NeedPostLoad) &&
 			IsValid(TargetSkelMesh) && !TargetSkelMesh->HasAnyFlags(RF_NeedPostLoad))
@@ -265,14 +265,14 @@ void FAnimNode_RemapCurvesFromMesh::RefreshMeshComponent(
 				ReinitializeMeshComponent(InMeshComponent, InTargetMeshComponent);
 			}
 			// if component is still same but mesh has been changed, we have to reinitialize
-			else if (CurrentMeshComponent->GetSkeletalMesh() != CurrentlyUsedSourceMesh.Get())
+			else if (CurrentMeshComponent->GetSkeletalMeshAsset() != CurrentlyUsedSourceMesh.Get())
 			{
 				ReinitializeMeshComponent(InMeshComponent, InTargetMeshComponent);
 			}
 			else if (InTargetMeshComponent)
 			{
 				// see if target mesh has changed
-				if (InTargetMeshComponent->GetSkeletalMesh() != CurrentlyUsedTargetMesh.Get())
+				if (InTargetMeshComponent->GetSkeletalMeshAsset() != CurrentlyUsedTargetMesh.Get())
 				{
 					ReinitializeMeshComponent(InMeshComponent, InTargetMeshComponent);
 				}
diff --git a/Engine/Plugins/Experimental/ChaosCloth/Source/Chaos/Private/ChaosCloth/ChaosClothConfig.cpp b/Engine/Plugins/Experimental/ChaosCloth/Source/Chaos/Private/ChaosCloth/ChaosClothConfig.cpp
--- a/Engine/Plugins/Experimental/ChaosCloth/Source/Chaos/Private/ChaosCloth/ChaosClothConfig.cpp
+++ b/Engine/Plugins/Experimental/ChaosCloth/Source/Chaos/Private/ChaosCloth/ChaosClothConfig.cpp
@@ -261,7 +261,7 @@ void UChaosClothSharedSimConfig::PostEditChangeChainProperty(FPropertyChangedCha
 			{
 				if (const USkeletalMeshComponent* const Component = *It)
 				{
-					if (Component->GetSkeletalMesh() == OwnerMesh)
+					if (Component->GetSkeletalMeshAsset() == OwnerMesh)
 					{
 						if (UClothingSimulationInteractor* const CurInteractor = Component->GetClothingSimulationInteractor())
 						{
diff --git a/Engine/Plugins/Experimental/GeometryCollectionPlugin/Source/GeometryCollectionEditor/Private/GeometryCollection/GeometryCollectionConversion.cpp b/Engine/Plugins/Experimental/GeometryCollectionPlugin/Source/GeometryCollectionEditor/Private/GeometryCollection/GeometryCollectionConversion.cpp
--- a/Engine/Plugins/Experimental/GeometryCollectionPlugin/Source/GeometryCollectionEditor/Private/GeometryCollection/GeometryCollectionConversion.cpp
+++ b/Engine/Plugins/Experimental/GeometryCollectionPlugin/Source/GeometryCollectionEditor/Private/GeometryCollection/GeometryCollectionConversion.cpp
@@ -146,14 +146,14 @@ void FGeometryCollectionConversion::CreateFromSelectedActorsCommand(UWorld * Wor
 			{
 				for (int Index = 0; Index < SkeletalMeshComponents.Num(); Index++)
 				{
-					if (SkeletalMeshComponents[Index]->GetSkeletalMesh())
+					if (SkeletalMeshComponents[Index]->GetSkeletalMeshAsset())
 					{
 						if (!GeometryCollectionAsset.IsCreated())
 						{
 							GeometryCollectionAsset.Create();
 						}
 						FGeometryCollectionEngineConversion::AppendSkeletalMesh(
-							SkeletalMeshComponents[Index]->GetSkeletalMesh(),
+							SkeletalMeshComponents[Index]->GetSkeletalMeshAsset(),
 							SkeletalMeshComponents[Index],
 							Actor->GetTransform(),
 							GeometryCollectionAsset.GeometryCollection);
diff --git a/Engine/Plugins/Experimental/GeometryCollectionPlugin/Source/GeometryCollectionEditor/Private/GeometryCollection/GeometryCollectionFactory.cpp b/Engine/Plugins/Experimental/GeometryCollectionPlugin/Source/GeometryCollectionEditor/Private/GeometryCollection/GeometryCollectionFactory.cpp
--- a/Engine/Plugins/Experimental/GeometryCollectionPlugin/Source/GeometryCollectionEditor/Private/GeometryCollection/GeometryCollectionFactory.cpp
+++ b/Engine/Plugins/Experimental/GeometryCollectionPlugin/Source/GeometryCollectionEditor/Private/GeometryCollection/GeometryCollectionFactory.cpp
@@ -66,10 +66,10 @@ void AppendActorComponentsRecursive(
 	{
 		for (int Index = 0; Index < SkeletalMeshComponents.Num(); Index++)
 		{
-			if (SkeletalMeshComponents[Index]->GetSkeletalMesh())
+			if (SkeletalMeshComponents[Index]->GetSkeletalMeshAsset())
 			{
 				SkeletalMeshList.Add(GeometryCollectionSkeletalMeshConversionTuple(
-					SkeletalMeshComponents[Index]->GetSkeletalMesh(),
+					SkeletalMeshComponents[Index]->GetSkeletalMeshAsset(),
 					SkeletalMeshComponents[Index],
 					SkeletalMeshComponents[Index]->GetComponentTransform()));
 			}
diff --git a/Engine/Plugins/Experimental/MLDeformer/MLDeformerFramework/Source/MLDeformerFramework/Private/MLDeformerModelInstance.cpp b/Engine/Plugins/Experimental/MLDeformer/MLDeformerFramework/Source/MLDeformerFramework/Private/MLDeformerModelInstance.cpp
--- a/Engine/Plugins/Experimental/MLDeformer/MLDeformerFramework/Source/MLDeformerFramework/Private/MLDeformerModelInstance.cpp
+++ b/Engine/Plugins/Experimental/MLDeformer/MLDeformerFramework/Source/MLDeformerFramework/Private/MLDeformerModelInstance.cpp
@@ -36,7 +36,7 @@ void UMLDeformerModelInstance::Init(USkeletalMeshComponent* SkelMeshComponent)
 		return;
 	}
 
-	USkeletalMesh* SkelMesh = SkelMeshComponent->GetSkeletalMesh();
+	USkeletalMesh* SkelMesh = SkelMeshComponent->GetSkeletalMeshAsset();
 	if (SkelMesh)
 	{
 		// Init the bone mapping table.
@@ -61,7 +61,7 @@ void UMLDeformerModelInstance::Init(USkeletalMeshComponent* SkelMeshComponent)
 
 void UMLDeformerModelInstance::UpdateCompatibilityStatus()
 {
-	bIsCompatible = SkeletalMeshComponent->GetSkeletalMesh() && CheckCompatibility(SkeletalMeshComponent, true).IsEmpty();
+	bIsCompatible = SkeletalMeshComponent->GetSkeletalMeshAsset() && CheckCompatibility(SkeletalMeshComponent, true).IsEmpty();
 }
 
 FString UMLDeformerModelInstance::CheckCompatibility(USkeletalMeshComponent* InSkelMeshComponent, bool LogIssues)
@@ -69,7 +69,7 @@ FString UMLDeformerModelInstance::CheckCompatibility(USkeletalMeshComponent* InS
 	ErrorText = FString();
 
 	// If we're not compatible, generate a compatibility string.
-	USkeletalMesh* SkelMesh = InSkelMeshComponent ? InSkelMeshComponent->GetSkeletalMesh() : nullptr;
+	USkeletalMesh* SkelMesh = InSkelMeshComponent ? InSkelMeshComponent->GetSkeletalMeshAsset() : nullptr;
 	UMLDeformerInputInfo* InputInfo = Model->GetInputInfo();
 	if (SkelMesh && !InputInfo->IsCompatible(SkelMesh) && Model->GetDeformerAsset())
 	{
@@ -296,7 +296,7 @@ bool UMLDeformerModelInstance::SetupNeuralNetworkForFrame()
 	// Some safety checks.
 	if (Model == nullptr ||
 		SkeletalMeshComponent == nullptr ||
-		SkeletalMeshComponent->GetSkeletalMesh() == nullptr ||
+		SkeletalMeshComponent->GetSkeletalMeshAsset() == nullptr ||
 		!bIsCompatible)
 	{
 		return false;
diff --git a/Engine/Plugins/Experimental/MLDeformer/MLDeformerFramework/Source/MLDeformerFrameworkEditor/Private/MLDeformerEditorActor.cpp b/Engine/Plugins/Experimental/MLDeformer/MLDeformerFramework/Source/MLDeformerFrameworkEditor/Private/MLDeformerEditorActor.cpp
--- a/Engine/Plugins/Experimental/MLDeformer/MLDeformerFramework/Source/MLDeformerFrameworkEditor/Private/MLDeformerEditorActor.cpp
+++ b/Engine/Plugins/Experimental/MLDeformer/MLDeformerFramework/Source/MLDeformerFrameworkEditor/Private/MLDeformerEditorActor.cpp
@@ -52,7 +52,7 @@ namespace UE::MLDeformer
 
 	bool FMLDeformerEditorActor::HasVisualMesh() const
 	{
-		if (SkeletalMeshComponent && SkeletalMeshComponent->GetSkeletalMesh())
+		if (SkeletalMeshComponent && SkeletalMeshComponent->GetSkeletalMeshAsset())
 		{
 			return true;
 		}
diff --git a/Engine/Plugins/Experimental/MLDeformer/MLDeformerFramework/Source/MLDeformerFrameworkEditor/Private/MLDeformerEditorModel.cpp b/Engine/Plugins/Experimental/MLDeformer/MLDeformerFramework/Source/MLDeformerFrameworkEditor/Private/MLDeformerEditorModel.cpp
--- a/Engine/Plugins/Experimental/MLDeformer/MLDeformerFramework/Source/MLDeformerFrameworkEditor/Private/MLDeformerEditorModel.cpp
+++ b/Engine/Plugins/Experimental/MLDeformer/MLDeformerFramework/Source/MLDeformerFrameworkEditor/Private/MLDeformerEditorModel.cpp
@@ -1379,7 +1379,7 @@ namespace UE::MLDeformer
 			const UMLDeformerModelInstance* ModelInstance = DeformerComponent->GetModelInstance();
 			if (ModelInstance &&
 				ModelInstance->GetSkeletalMeshComponent() && 
-				ModelInstance->GetSkeletalMeshComponent()->GetSkeletalMesh() &&
+				ModelInstance->GetSkeletalMeshComponent()->GetSkeletalMeshAsset() &&
 				!ModelInstance->IsCompatible() )
 			{
 				return FText::FromString( ModelInstance->GetCompatibilityErrorText() );
diff --git a/Engine/Plugins/Experimental/MLDeformer/MLDeformerFramework/Source/MLDeformerFrameworkEditor/Private/MLDeformerGeomCacheSampler.cpp b/Engine/Plugins/Experimental/MLDeformer/MLDeformerFramework/Source/MLDeformerFrameworkEditor/Private/MLDeformerGeomCacheSampler.cpp
--- a/Engine/Plugins/Experimental/MLDeformer/MLDeformerFramework/Source/MLDeformerFrameworkEditor/Private/MLDeformerGeomCacheSampler.cpp
+++ b/Engine/Plugins/Experimental/MLDeformer/MLDeformerFramework/Source/MLDeformerFrameworkEditor/Private/MLDeformerGeomCacheSampler.cpp
@@ -58,7 +58,7 @@ namespace UE::MLDeformer
 		// This will also calculate the skinned positions if the delta space is set to PostSkinning.
 		FMLDeformerSampler::Sample(InAnimFrameIndex);
 
-		USkeletalMesh* SkeletalMesh = SkeletalMeshComponent.Get() ? SkeletalMeshComponent->GetSkeletalMesh() : nullptr;
+		USkeletalMesh* SkeletalMesh = SkeletalMeshComponent.Get() ? SkeletalMeshComponent->GetSkeletalMeshAsset() : nullptr;
 		UGeometryCache* GeometryCache = GeometryCacheComponent->GetGeometryCache();
 		if (SkeletalMeshComponent && SkeletalMesh && GeometryCacheComponent && GeometryCache)
 		{
diff --git a/Engine/Plugins/Experimental/MLDeformer/MLDeformerFramework/Source/MLDeformerFrameworkEditor/Private/MLDeformerSampler.cpp b/Engine/Plugins/Experimental/MLDeformer/MLDeformerFramework/Source/MLDeformerFrameworkEditor/Private/MLDeformerSampler.cpp
--- a/Engine/Plugins/Experimental/MLDeformer/MLDeformerFramework/Source/MLDeformerFrameworkEditor/Private/MLDeformerSampler.cpp
+++ b/Engine/Plugins/Experimental/MLDeformer/MLDeformerFramework/Source/MLDeformerFrameworkEditor/Private/MLDeformerSampler.cpp
@@ -62,7 +62,7 @@ namespace UE::MLDeformer
 
 	void FMLDeformerSampler::UpdateSkeletalMeshComponent()
 	{
-		if (SkeletalMeshComponent && SkeletalMeshComponent->GetSkeletalMesh())
+		if (SkeletalMeshComponent && SkeletalMeshComponent->GetSkeletalMeshAsset())
 		{
 			// Sample the transforms at the frame time.
 			SkeletalMeshComponent->SetPosition(SampleTime);
@@ -95,7 +95,7 @@ namespace UE::MLDeformer
 		AnimFrameIndex = InAnimFrameIndex;
 		SampleTime = GetTimeAtFrame(InAnimFrameIndex);
 
-		const USkeletalMesh* SkeletalMesh = SkeletalMeshComponent->GetSkeletalMesh();
+		const USkeletalMesh* SkeletalMesh = SkeletalMeshComponent->GetSkeletalMeshAsset();
 		if (SkeletalMeshComponent && SkeletalMesh)
 		{
 			UpdateSkeletalMeshComponent();
@@ -117,7 +117,7 @@ namespace UE::MLDeformer
 	FMatrix44f FMLDeformerSampler::CalcInverseSkinningTransform(int32 VertexIndex, const FSkeletalMeshLODRenderData& SkelMeshLODData, const FSkinWeightVertexBuffer& SkinWeightBuffer) const
 	{
 		check(SkeletalMeshComponent);
-		const USkeletalMesh* Mesh = SkeletalMeshComponent->GetSkeletalMesh();
+		const USkeletalMesh* Mesh = SkeletalMeshComponent->GetSkeletalMeshAsset();
 		check(Mesh);
 
 		// Find the render section, which we need to find the right bone index.
@@ -158,7 +158,7 @@ namespace UE::MLDeformer
 			return;
 		}
 
-		USkeletalMesh* Mesh = SkeletalMeshComponent->GetSkeletalMesh();
+		USkeletalMesh* Mesh = SkeletalMeshComponent->GetSkeletalMeshAsset();
 		if (Mesh == nullptr)
 		{
 			return;
@@ -193,7 +193,7 @@ namespace UE::MLDeformer
 			return;
 		}
 
-		USkeletalMesh* Mesh = SkeletalMeshComponent->GetSkeletalMesh();
+		USkeletalMesh* Mesh = SkeletalMeshComponent->GetSkeletalMeshAsset();
 		if (Mesh == nullptr)
 		{
 			return;
diff --git a/Engine/Plugins/Experimental/MLDeformer/NeuralMorphModel/Source/NeuralMorphModel/Private/NeuralMorphModel.cpp b/Engine/Plugins/Experimental/MLDeformer/NeuralMorphModel/Source/NeuralMorphModel/Private/NeuralMorphModel.cpp
--- a/Engine/Plugins/Experimental/MLDeformer/NeuralMorphModel/Source/NeuralMorphModel/Private/NeuralMorphModel.cpp
+++ b/Engine/Plugins/Experimental/MLDeformer/NeuralMorphModel/Source/NeuralMorphModel/Private/NeuralMorphModel.cpp
@@ -106,7 +106,7 @@ void UNeuralMorphModel::PostMLDeformerComponentInit(UMLDeformerModelInstance* Mo
 
 	// Register the external morph targets buffer to the render data of the skeletal mesh.
 	USkeletalMeshComponent* SkelMeshComponent = ModelInstance->GetSkeletalMeshComponent();
-	if (SkelMeshComponent && SkelMeshComponent->GetSkeletalMesh())
+	if (SkelMeshComponent && SkelMeshComponent->GetSkeletalMeshAsset())
 	{
 		// Register the morph set. This overwrites the existing one for this model, if it already exists.
 		// Only add to LOD 0 for now.
diff --git a/Engine/Plugins/Experimental/MLDeformer/NeuralMorphModel/Source/NeuralMorphModelEditor/Private/NeuralMorphEditorModelActor.cpp b/Engine/Plugins/Experimental/MLDeformer/NeuralMorphModel/Source/NeuralMorphModelEditor/Private/NeuralMorphEditorModelActor.cpp
--- a/Engine/Plugins/Experimental/MLDeformer/NeuralMorphModel/Source/NeuralMorphModelEditor/Private/NeuralMorphEditorModelActor.cpp
+++ b/Engine/Plugins/Experimental/MLDeformer/NeuralMorphModel/Source/NeuralMorphModelEditor/Private/NeuralMorphEditorModelActor.cpp
@@ -48,7 +48,7 @@ namespace UE::NeuralMorphModel
 
 	bool FNeuralMorphEditorModelActor::HasVisualMesh() const
 	{
-		if (SkeletalMeshComponent && SkeletalMeshComponent->GetSkeletalMesh())
+		if (SkeletalMeshComponent && SkeletalMeshComponent->GetSkeletalMeshAsset())
 		{
 			return true;
 		}
diff --git a/Engine/Plugins/Experimental/MeshModelingToolsetExp/Source/MeshModelingToolsExp/Private/SkinWeightsBindingTool.cpp b/Engine/Plugins/Experimental/MeshModelingToolsetExp/Source/MeshModelingToolsExp/Private/SkinWeightsBindingTool.cpp
--- a/Engine/Plugins/Experimental/MeshModelingToolsetExp/Source/MeshModelingToolsExp/Private/SkinWeightsBindingTool.cpp
+++ b/Engine/Plugins/Experimental/MeshModelingToolsetExp/Source/MeshModelingToolsExp/Private/SkinWeightsBindingTool.cpp
@@ -770,9 +770,9 @@ void USkinWeightsBindingTool::Setup()
 	
 	const USkeletalMeshComponent* SkelMeshComponent = Cast<USkeletalMeshComponent>(UE::ToolTarget::GetTargetComponent(Targets[0]));
 	
-	if (SkelMeshComponent && SkelMeshComponent->GetSkeletalMesh())
+	if (SkelMeshComponent && SkelMeshComponent->GetSkeletalMeshAsset())
 	{
-		USkeletalMesh* SkeletalMesh = SkelMeshComponent->GetSkeletalMesh();
+		USkeletalMesh* SkeletalMesh = SkelMeshComponent->GetSkeletalMeshAsset();
 
 		// Initialize the bone browser
 		FCurveEvaluationOption CurveEvalOption(
diff --git a/Engine/Plugins/Experimental/MeshModelingToolsetExp/Source/MeshModelingToolsExp/Private/SkinWeightsPaintTool.cpp b/Engine/Plugins/Experimental/MeshModelingToolsetExp/Source/MeshModelingToolsExp/Private/SkinWeightsPaintTool.cpp
--- a/Engine/Plugins/Experimental/MeshModelingToolsetExp/Source/MeshModelingToolsExp/Private/SkinWeightsPaintTool.cpp
+++ b/Engine/Plugins/Experimental/MeshModelingToolsetExp/Source/MeshModelingToolsExp/Private/SkinWeightsPaintTool.cpp
@@ -68,11 +68,11 @@ void USkinWeightsPaintTool::Setup()
 
 	ToolProps = NewObject<USkinWeightsPaintToolProperties>(this);
 	ToolProps->RestoreProperties(this);
-	if (Component && Component->GetSkeletalMesh())
+	if (Component && Component->GetSkeletalMeshAsset())
 	{
 		// Get all non-virtual bones
 		// TArray<int32> BoneIndices;
-		USkeletalMesh& SkeletalMesh = *Component->GetSkeletalMesh();
+		USkeletalMesh& SkeletalMesh = *Component->GetSkeletalMeshAsset();
 		// SkeletalMesh.RefSkeleton.GetRawRefBoneInfo();
 		
 
@@ -84,11 +84,11 @@ void USkinWeightsPaintTool::Setup()
 				);
 		BoneContainer.InitializeTo(Component->RequiredBones, CurveEvalOption, SkeletalMesh);
 
-		ToolProps->SkeletalMesh = Component->GetSkeletalMesh();
+		ToolProps->SkeletalMesh = Component->GetSkeletalMeshAsset();
 		ToolProps->CurrentBone.Initialize(BoneContainer);
 
 		// Pick the first root bone as the initial selection.
-		PendingCurrentBone = Component->GetSkeletalMesh()->GetRefSkeleton().GetBoneName(0);
+		PendingCurrentBone = Component->GetSkeletalMeshAsset()->GetRefSkeleton().GetBoneName(0);
 		ToolProps->CurrentBone.BoneName = PendingCurrentBone.GetValue();
 
 		// Update the skeleton drawing information from the original bind pose
diff --git a/Engine/Plugins/FX/Niagara/Source/Niagara/Classes/NiagaraDataInterfaceSkeletalMesh.h b/Engine/Plugins/FX/Niagara/Source/Niagara/Classes/NiagaraDataInterfaceSkeletalMesh.h
--- a/Engine/Plugins/FX/Niagara/Source/Niagara/Classes/NiagaraDataInterfaceSkeletalMesh.h
+++ b/Engine/Plugins/FX/Niagara/Source/Niagara/Classes/NiagaraDataInterfaceSkeletalMesh.h
@@ -694,7 +694,7 @@ struct FNDISkeletalMesh_InstanceData
 	FORCEINLINE_DEBUGGABLE const FSkinWeightVertexBuffer* GetSkinWeights()
 	{
 		USkeletalMeshComponent* SkelComp = Cast<USkeletalMeshComponent>(SceneComponent.Get());
-		if (SkelComp != nullptr && SkelComp->GetSkeletalMesh() != nullptr)
+		if (SkelComp != nullptr && SkelComp->GetSkeletalMeshAsset() != nullptr)
 		{
 			return SkelComp->GetSkinWeightBuffer(CachedLODIdx);
 		}
diff --git a/Engine/Plugins/FX/Niagara/Source/Niagara/Private/DataInterface/NiagaraDataInterfacePhysicsAsset.cpp b/Engine/Plugins/FX/Niagara/Source/Niagara/Private/DataInterface/NiagaraDataInterfacePhysicsAsset.cpp
--- a/Engine/Plugins/FX/Niagara/Source/Niagara/Private/DataInterface/NiagaraDataInterfacePhysicsAsset.cpp
+++ b/Engine/Plugins/FX/Niagara/Source/Niagara/Private/DataInterface/NiagaraDataInterfacePhysicsAsset.cpp
@@ -99,7 +99,7 @@ void GetNumPrimitives(const TArray<TWeakObjectPtr<UPhysicsAsset>>& PhysicsAssets
 		TWeakObjectPtr<UPhysicsAsset> PhysicsAsset = PhysicsAssets[ComponentIndex];
 		if (PhysicsAsset.IsValid() && PhysicsAsset.Get() != nullptr)
 		{
-			USkeletalMesh* SkelMesh = (SkeletalMeshs[ComponentIndex].Get() && SkeletalMeshs[ComponentIndex]->GetSkeletalMesh()) ? ToRawPtr(SkeletalMeshs[ComponentIndex]->GetSkeletalMesh()) : PhysicsAsset->GetPreviewMesh();
+			USkeletalMesh* SkelMesh = (SkeletalMeshs[ComponentIndex].Get() && SkeletalMeshs[ComponentIndex]->GetSkeletalMeshAsset()) ? ToRawPtr(SkeletalMeshs[ComponentIndex]->GetSkeletalMeshAsset()) : PhysicsAsset->GetPreviewMesh();
 			if (!SkelMesh)
 			{
 				continue;
@@ -209,7 +209,7 @@ void CreateInternalArrays(const TArray<TWeakObjectPtr<UPhysicsAsset>>& PhysicsAs
 				TWeakObjectPtr<UPhysicsAsset> PhysicsAsset = PhysicsAssets[ComponentIndex];
 				if (PhysicsAsset.IsValid() && PhysicsAsset.Get() != nullptr)
 				{
-					USkeletalMesh* SkelMesh = (SkeletalMeshs[ComponentIndex].Get() && SkeletalMeshs[ComponentIndex]->GetSkeletalMesh()) ? ToRawPtr(SkeletalMeshs[ComponentIndex]->GetSkeletalMesh()) : PhysicsAsset->GetPreviewMesh();
+					USkeletalMesh* SkelMesh = (SkeletalMeshs[ComponentIndex].Get() && SkeletalMeshs[ComponentIndex]->GetSkeletalMeshAsset()) ? ToRawPtr(SkeletalMeshs[ComponentIndex]->GetSkeletalMeshAsset()) : PhysicsAsset->GetPreviewMesh();
 					if (!SkelMesh)
 					{
 						continue;
@@ -330,7 +330,7 @@ void UpdateInternalArrays(const TArray<TWeakObjectPtr<UPhysicsAsset>>& PhysicsAs
 				TWeakObjectPtr<UPhysicsAsset> PhysicsAsset = PhysicsAssets[ComponentIndex];
 				if (PhysicsAsset.IsValid() && PhysicsAsset.Get() != nullptr)
 				{
-					USkeletalMesh* SkelMesh = (SkeletalMeshs[ComponentIndex].Get() && SkeletalMeshs[ComponentIndex]->GetSkeletalMesh()) ? ToRawPtr(SkeletalMeshs[ComponentIndex]->GetSkeletalMesh()) : PhysicsAsset->GetPreviewMesh();
+					USkeletalMesh* SkelMesh = (SkeletalMeshs[ComponentIndex].Get() && SkeletalMeshs[ComponentIndex]->GetSkeletalMeshAsset()) ? ToRawPtr(SkeletalMeshs[ComponentIndex]->GetSkeletalMeshAsset()) : PhysicsAsset->GetPreviewMesh();
 					if (!SkelMesh)
 					{
 						continue;
@@ -474,8 +474,8 @@ void FNDIPhysicsAssetData::Init(UNiagaraDataInterfacePhysicsAsset* Interface, FN
 			TWeakObjectPtr<UPhysicsAsset> PhysicsAsset = Interface->PhysicsAssets[ComponentIndex];
 			if (PhysicsAsset.IsValid() && PhysicsAsset.Get() != nullptr)
 			{
-				USkeletalMesh* SkelMesh = (Interface->SourceComponents[ComponentIndex].IsValid() && Interface->SourceComponents[ComponentIndex].Get() && Interface->SourceComponents[ComponentIndex]->GetSkeletalMesh()) ?
-					ToRawPtr(Interface->SourceComponents[ComponentIndex]->GetSkeletalMesh()) : PhysicsAsset->GetPreviewMesh();
+				USkeletalMesh* SkelMesh = (Interface->SourceComponents[ComponentIndex].IsValid() && Interface->SourceComponents[ComponentIndex].Get() && Interface->SourceComponents[ComponentIndex]->GetSkeletalMeshAsset()) ?
+					ToRawPtr(Interface->SourceComponents[ComponentIndex]->GetSkeletalMeshAsset()) : PhysicsAsset->GetPreviewMesh();
 				if (SkelMesh)
 				{
 					BoundingBox += SkelMesh->GetImportedBounds().GetBox();
@@ -593,7 +593,7 @@ void UNiagaraDataInterfacePhysicsAsset::ExtractSourceComponent(FNiagaraSystemIns
 			for (UActorComponent* ActorComp : Actor->GetComponents())
 			{
 				USkeletalMeshComponent* Comp = Cast<USkeletalMeshComponent>(ActorComp);
-				if (IsValid(Comp) && Comp->GetSkeletalMesh() != nullptr)
+				if (IsValid(Comp) && Comp->GetSkeletalMeshAsset() != nullptr)
 				{
 					return Comp;
 				}
@@ -669,7 +669,7 @@ void UNiagaraDataInterfacePhysicsAsset::ExtractSourceComponent(FNiagaraSystemIns
 		for (USceneComponent* Curr = AttachComponent; Curr; Curr = Curr->GetAttachParent())
 		{
 			USkeletalMeshComponent* SkelMeshComp = Cast<USkeletalMeshComponent>(Curr);
-			if (SkelMeshComp && SkelMeshComp->GetSkeletalMesh())
+			if (SkelMeshComp && SkelMeshComp->GetSkeletalMeshAsset())
 			{
 				SourceComponent = SkelMeshComp;
 				break;
diff --git a/Engine/Plugins/FX/Niagara/Source/Niagara/Private/NiagaraDataInterfaceRigidMeshCollisionQuery.cpp b/Engine/Plugins/FX/Niagara/Source/Niagara/Private/NiagaraDataInterfaceRigidMeshCollisionQuery.cpp
--- a/Engine/Plugins/FX/Niagara/Source/Niagara/Private/NiagaraDataInterfaceRigidMeshCollisionQuery.cpp
+++ b/Engine/Plugins/FX/Niagara/Source/Niagara/Private/NiagaraDataInterfaceRigidMeshCollisionQuery.cpp
@@ -200,7 +200,7 @@ void CollectComponents(TConstArrayView<AActor*> Actors, TConstArrayView<FName> C
 	{
 		if (UPhysicsAsset* PhysicsAsset = Component->GetPhysicsAsset())
 		{
-			USkeletalMesh* MeshAsset = Component->GetSkeletalMesh() ? Component->GetSkeletalMesh() : PhysicsAsset->GetPreviewMesh();
+			USkeletalMesh* MeshAsset = Component->GetSkeletalMeshAsset() ? Component->GetSkeletalMeshAsset() : PhysicsAsset->GetPreviewMesh();
 			if (!MeshAsset || !MeshAsset->GetRefSkeleton().GetNum())
 			{
 				return false;
@@ -220,7 +220,7 @@ void ForEachBodySetup(USkeletalMeshComponent* Component, TBodySetupPredicate Pre
 {
 	if (UPhysicsAsset* PhysicsAsset = Component->GetPhysicsAsset())
 	{
-		USkeletalMesh* SkeletalMesh = Component->GetSkeletalMesh() ? Component->GetSkeletalMesh() : PhysicsAsset->GetPreviewMesh();
+		USkeletalMesh* SkeletalMesh = Component->GetSkeletalMeshAsset() ? Component->GetSkeletalMeshAsset() : PhysicsAsset->GetPreviewMesh();
 		const FReferenceSkeleton& RefSkeleton = SkeletalMesh->GetRefSkeleton();
 
 		for (const UBodySetup* BodySetup : PhysicsAsset->SkeletalBodySetups)
@@ -314,7 +314,7 @@ FTransform CreateElementTransform(const TComponentType* Component, const UBodySe
 template<>
 FTransform CreateElementTransform<USkeletalMeshComponent>(const USkeletalMeshComponent* Component, const UBodySetup* BodySetup)
 {
-	if (USkeletalMesh* SkeletalMesh = Component->GetSkeletalMesh())
+	if (USkeletalMesh* SkeletalMesh = Component->GetSkeletalMeshAsset())
 	{
 		const FReferenceSkeleton& RefSkeleton = SkeletalMesh->GetRefSkeleton();
 		const int32 BoneCount = RefSkeleton.GetNum();
diff --git a/Engine/Plugins/FX/Niagara/Source/Niagara/Private/NiagaraDataInterfaceSkeletalMesh.cpp b/Engine/Plugins/FX/Niagara/Source/Niagara/Private/NiagaraDataInterfaceSkeletalMesh.cpp
--- a/Engine/Plugins/FX/Niagara/Source/Niagara/Private/NiagaraDataInterfaceSkeletalMesh.cpp
+++ b/Engine/Plugins/FX/Niagara/Source/Niagara/Private/NiagaraDataInterfaceSkeletalMesh.cpp
@@ -291,7 +291,7 @@ void FSkeletalMeshSkinningData::RegisterUser(FSkeletalMeshSkinningDataUsage Usag
 	USkeletalMeshComponent* SkelComp = MeshComp.Get();
 	check(SkelComp);
 
-	USkeletalMesh* SkelMesh = SkelComp->GetSkeletalMesh();
+	USkeletalMesh* SkelMesh = SkelComp->GetSkeletalMeshAsset();
 	int32 LODIndex = 0;
 	int32 NumLODInfo = 1;
 
@@ -368,7 +368,7 @@ void FSkeletalMeshSkinningData::UnregisterUser(FSkeletalMeshSkinningDataUsage Us
 	int32 LODIndex = 0;
 
 	USkeletalMeshComponent* SkelComp = MeshComp.Get();
-	if (SkelComp && SkelComp->GetSkeletalMesh())
+	if (SkelComp && SkelComp->GetSkeletalMeshAsset())
 	{
 		LODIndex = Usage.GetLODIndex();
 	}
@@ -410,7 +410,7 @@ void FSkeletalMeshSkinningData::UpdateBoneTransforms()
 	USkeletalMeshComponent* SkelComp = MeshComp.Get();
 	check(SkelComp);
 
-	const USkeletalMesh* SkelMesh = SkelComp->GetSkeletalMesh();
+	const USkeletalMesh* SkelMesh = SkelComp->GetSkeletalMeshAsset();
 	if (SkelMesh == nullptr)
 	{
 		return;
@@ -505,9 +505,9 @@ bool FSkeletalMeshSkinningData::Tick(float InDeltaSeconds, bool bRequirePreskin)
 		PrevComponentTransforms() = CurrComponentTransforms();
 	}
 
-	if (bRequirePreskin && SkelComp->GetSkeletalMesh() != nullptr)
+	if (bRequirePreskin && SkelComp->GetSkeletalMeshAsset() != nullptr)
 	{
-		const USkeletalMesh* SkeletalMesh = SkelComp->GetSkeletalMesh();
+		const USkeletalMesh* SkeletalMesh = SkelComp->GetSkeletalMeshAsset();
 		const FSkeletalMeshRenderData* RenderData = SkeletalMesh->GetResourceForRendering();
 		check(RenderData);
 
@@ -1151,7 +1151,7 @@ void FSkeletalMeshGpuDynamicBufferProxy::NewFrame(const FNDISkeletalMesh_Instanc
 		SkelComp = Cast<USkeletalMeshComponent>(InstanceData->SceneComponent.Get());
 		if ( SkelComp != nullptr )
 		{
-			SkelMesh = SkelComp->GetSkeletalMesh();
+			SkelMesh = SkelComp->GetSkeletalMeshAsset();
 		}
 		if (SkelMesh == nullptr)
 		{
@@ -1497,7 +1497,7 @@ USkeletalMesh* UNiagaraDataInterfaceSkeletalMesh::GetSkeletalMesh(FNiagaraSystem
 {
 	auto IsValidComponent = [&](const USkeletalMeshComponent* Component)
 	{
-		return IsValid(Component) && IsValid(Component->GetSkeletalMesh()) &&
+		return IsValid(Component) && IsValid(Component->GetSkeletalMeshAsset()) &&
 			(ComponentTags.IsEmpty()
 			|| ComponentTags.ContainsByPredicate([&](const FName& Tag) { return Tag == NAME_None || Component->ComponentHasTag(Tag); }));
 	};
@@ -1625,7 +1625,7 @@ USkeletalMesh* UNiagaraDataInterfaceSkeletalMesh::GetSkeletalMesh(FNiagaraSystem
 	SceneComponent = nullptr;
 	if (FoundSkelComp)
 	{
-		Mesh = FoundSkelComp->GetSkeletalMesh();
+		Mesh = FoundSkelComp->GetSkeletalMeshAsset();
 		SceneComponent = FoundSkelComp;
 	}
 #if WITH_EDITORONLY_DATA
@@ -2156,7 +2156,7 @@ bool FNDISkeletalMesh_InstanceData::ResetRequired(UNiagaraDataInterfaceSkeletalM
 	// Reset if the skeletal mesh on the cached skeletal mesh component changed.
 	if (USkeletalMeshComponent* SkelComp = Cast<USkeletalMeshComponent>(Comp))
 	{
-		if (SkelComp->GetSkeletalMesh() != SkelMesh)
+		if (SkelComp->GetSkeletalMeshAsset() != SkelMesh)
 		{
 			if (SkinningData.SkinningData.IsValid())
 			{
diff --git a/Engine/Plugins/FX/Niagara/Source/Niagara/Private/NiagaraFunctionLibrary.cpp b/Engine/Plugins/FX/Niagara/Source/Niagara/Private/NiagaraFunctionLibrary.cpp
--- a/Engine/Plugins/FX/Niagara/Source/Niagara/Private/NiagaraFunctionLibrary.cpp
+++ b/Engine/Plugins/FX/Niagara/Source/Niagara/Private/NiagaraFunctionLibrary.cpp
@@ -431,7 +431,7 @@ void UNiagaraFunctionLibrary::OverrideSystemUserVariableSkeletalMeshComponent(UN
 	{
 		if (FNiagaraUtilities::LogVerboseWarnings())
 		{
-			UE_LOG(LogNiagara, Warning, TEXT("NiagaraSystem in \"Set Niagara Skeletal Mesh Component\" is NULL, OverrideName \"%s\" and SkeletalMeshComponent \"%s\" SkeletalMesh \"%s\", skipping."), *OverrideName, *GetFullNameSafe(SkeletalMeshComponent), *GetFullNameSafe(SkeletalMeshComponent ? SkeletalMeshComponent->GetSkeletalMesh() : nullptr));
+			UE_LOG(LogNiagara, Warning, TEXT("NiagaraSystem in \"Set Niagara Skeletal Mesh Component\" is NULL, OverrideName \"%s\" and SkeletalMeshComponent \"%s\" SkeletalMesh \"%s\", skipping."), *OverrideName, *GetFullNameSafe(SkeletalMeshComponent), *GetFullNameSafe(SkeletalMeshComponent ? SkeletalMeshComponent->GetSkeletalMeshAsset() : nullptr));
 		}
 		return;
 	}
diff --git a/Engine/Plugins/Importers/USDImporter/Source/USDSchemas/Private/USDSkelRootTranslator.cpp b/Engine/Plugins/Importers/USDImporter/Source/USDSchemas/Private/USDSkelRootTranslator.cpp
--- a/Engine/Plugins/Importers/USDImporter/Source/USDSchemas/Private/USDSkelRootTranslator.cpp
+++ b/Engine/Plugins/Importers/USDImporter/Source/USDSchemas/Private/USDSkelRootTranslator.cpp
@@ -278,17 +278,17 @@ namespace UsdSkelRootTranslatorImpl
 
 	void SetMorphTargetWeight( USkeletalMeshComponent& SkeletalMeshComponent, const FString& MorphTargetName, float Weight )
 	{
-		USkeletalMesh* SkeletalMesh = SkeletalMeshComponent.GetSkeletalMesh();
+		USkeletalMesh* SkeletalMesh = SkeletalMeshComponent.GetSkeletalMeshAsset();
 
 		// We try keeping a perfect correspondence between SkeletalMesh->GetMorphTargets() and SkeletalMeshComponent.ActiveMorphTargets
 		int32 IndexInSkeletalMesh = INDEX_NONE;
-		SkeletalMeshComponent.GetSkeletalMesh()->FindMorphTargetAndIndex( *MorphTargetName, IndexInSkeletalMesh );
+		SkeletalMeshComponent.GetSkeletalMeshAsset()->FindMorphTargetAndIndex( *MorphTargetName, IndexInSkeletalMesh );
 		if ( IndexInSkeletalMesh == INDEX_NONE )
 		{
 			return;
 		}
 
-		UMorphTarget* MorphTarget = SkeletalMeshComponent.GetSkeletalMesh()->GetMorphTargets()[ IndexInSkeletalMesh ];
+		UMorphTarget* MorphTarget = SkeletalMeshComponent.GetSkeletalMeshAsset()->GetMorphTargets()[ IndexInSkeletalMesh ];
 		if ( !MorphTarget )
 		{
 			return;
@@ -960,14 +960,14 @@ namespace UsdSkelRootTranslatorImpl
 
 	void UpdateLiveLinkProperties( const FUsdSchemaTranslationContext& Context, USkeletalMeshComponent* Component, const pxr::UsdPrim& Prim )
 	{
-		if ( !Component || !Component->GetSkeletalMesh() || !Prim )
+		if ( !Component || !Component->GetSkeletalMeshAsset() || !Prim )
 		{
 			return;
 		}
 
 		FString PrimName = UsdToUnreal::ConvertString( Prim.GetName() );
 
-		USkeleton* Skeleton = Component->GetSkeletalMesh()->GetSkeleton();
+		USkeleton* Skeleton = Component->GetSkeletalMeshAsset()->GetSkeleton();
 		if ( !Skeleton )
 		{
 			return;
@@ -1308,7 +1308,7 @@ void FUsdSkelRootTranslator::UpdateComponents( USceneComponent* SceneComponent )
 #if WITH_EDITOR
 	// Re-set the skeletal mesh if we created a new one (maybe the hash changed, a skinned UsdGeomMesh was hidden, etc.)
 	USkeletalMesh* TargetSkeletalMesh = Cast< USkeletalMesh >( Context->AssetCache->GetAssetForPrim( PrimPath.GetString() ) );
-	if ( SkeletalMeshComponent->GetSkeletalMesh() != TargetSkeletalMesh )
+	if ( SkeletalMeshComponent->GetSkeletalMeshAsset() != TargetSkeletalMesh )
 	{
 		SkeletalMeshComponent->SetSkeletalMesh(TargetSkeletalMesh);
 
@@ -1342,7 +1342,7 @@ void FUsdSkelRootTranslator::UpdateComponents( USceneComponent* SceneComponent )
 	// Update the animation state
 	// Don't try animating ourselves if the sequencer is animating as it will just overwrite the animation state on next
 	// tick anyway, and all this would do is lead to flickering and other issues
-	if ( !Context->bSequencerIsAnimating && SkeletalMeshComponent->GetSkeletalMesh() && !bPrimHasLiveLinkEnabled )
+	if ( !Context->bSequencerIsAnimating && SkeletalMeshComponent->GetSkeletalMeshAsset() && !bPrimHasLiveLinkEnabled )
 	{
 		if ( UAnimSequence* AnimSequence = Cast<UAnimSequence>( SkeletalMeshComponent->AnimationData.AnimToPlay.Get() ) )
 		{
diff --git a/Engine/Plugins/Importers/USDImporter/Source/USDStage/Private/USDLevelSequenceHelper.cpp b/Engine/Plugins/Importers/USDImporter/Source/USDStage/Private/USDLevelSequenceHelper.cpp
--- a/Engine/Plugins/Importers/USDImporter/Source/USDStage/Private/USDLevelSequenceHelper.cpp
+++ b/Engine/Plugins/Importers/USDImporter/Source/USDStage/Private/USDLevelSequenceHelper.cpp
@@ -256,7 +256,7 @@ namespace UsdLevelSequenceHelperImpl
 	)
 	{
 		UMovieScene* MovieScene = LevelSequence->GetMovieScene();
-		if ( !MovieScene || !SkeletalMeshComp || !SkeletalMeshComp->GetSkeletalMesh() || !SkeletalMeshComp->GetSkeletalMesh()->GetSkeleton())
+		if ( !MovieScene || !SkeletalMeshComp || !SkeletalMeshComp->GetSkeletalMeshAsset() || !SkeletalMeshComp->GetSkeletalMeshAsset()->GetSkeleton())
 		{
 			return false;
 		}
@@ -295,7 +295,7 @@ namespace UsdLevelSequenceHelperImpl
 		{
 			bCreatedTempSequence = true;
 			AnimSequence = NewObject<UAnimSequence>();
-			AnimSequence->SetSkeleton( SkeletalMeshComp->GetSkeletalMesh()->GetSkeleton());
+			AnimSequence->SetSkeleton( SkeletalMeshComp->GetSkeletalMeshAsset()->GetSkeleton());
 
 			FMovieSceneSequenceIDRef Template = MovieSceneSequenceID::Root;
 			FMovieSceneSequenceTransform RootToLocalTransform;
@@ -2925,7 +2925,7 @@ void FUsdLevelSequenceHelperImpl::HandleControlRigSectionChange( UMovieSceneCont
 		return;
 	}
 
-	USkeleton* Skeleton = BoundComponent->GetSkeletalMesh() ? BoundComponent->GetSkeletalMesh()->GetSkeleton() : nullptr;
+	USkeleton* Skeleton = BoundComponent->GetSkeletalMeshAsset() ? BoundComponent->GetSkeletalMeshAsset()->GetSkeleton() : nullptr;
 	if ( !Skeleton )
 	{
 		return;
diff --git a/Engine/Plugins/Importers/USDImporter/Source/USDStageImporter/Private/USDStageImporter.cpp b/Engine/Plugins/Importers/USDImporter/Source/USDStageImporter/Private/USDStageImporter.cpp
--- a/Engine/Plugins/Importers/USDImporter/Source/USDStageImporter/Private/USDStageImporter.cpp
+++ b/Engine/Plugins/Importers/USDImporter/Source/USDStageImporter/Private/USDStageImporter.cpp
@@ -558,7 +558,7 @@ namespace UsdStageImporterImpl
 			// We do not need to recreate the anim instance afterwards as InitAnim is called on component registration anyway.
 			for ( TObjectIterator<USkeletalMeshComponent> It; It; ++It )
 			{
-				if ( It->GetSkeletalMesh() == ExistingAsset)
+				if ( It->GetSkeletalMeshAsset() == ExistingAsset)
 				{
 					It->ClearAnimScriptInstance();
 				}
diff --git a/Engine/Plugins/Importers/USDImporter/Source/USDUtilities/Private/USDPrimConversion.cpp b/Engine/Plugins/Importers/USDImporter/Source/USDUtilities/Private/USDPrimConversion.cpp
--- a/Engine/Plugins/Importers/USDImporter/Source/USDUtilities/Private/USDPrimConversion.cpp
+++ b/Engine/Plugins/Importers/USDImporter/Source/USDUtilities/Private/USDPrimConversion.cpp
@@ -2840,7 +2840,7 @@ bool UnrealToUsd::CreateComponentPropertyBaker( UE::FUsdPrim& Prim, const UScene
 bool UnrealToUsd::CreateSkeletalAnimationBaker( UE::FUsdPrim& SkelRoot, UE::FUsdPrim& SkelAnimation, USkeletalMeshComponent& Component, FComponentBaker& OutBaker )
 {
 #if WITH_EDITOR
-	USkeletalMesh* SkeletalMesh = Component.GetSkeletalMesh();
+	USkeletalMesh* SkeletalMesh = Component.GetSkeletalMeshAsset();
 	if ( !SkeletalMesh )
 	{
 		return false;
diff --git a/Engine/Plugins/MeshPainting/Source/MeshPaintEditorMode/Private/MeshPaintMode.cpp b/Engine/Plugins/MeshPainting/Source/MeshPaintEditorMode/Private/MeshPaintMode.cpp
--- a/Engine/Plugins/MeshPainting/Source/MeshPaintEditorMode/Private/MeshPaintMode.cpp
+++ b/Engine/Plugins/MeshPainting/Source/MeshPaintEditorMode/Private/MeshPaintMode.cpp
@@ -593,9 +593,9 @@ void UMeshPaintMode::SavePaintedAssets()
 
 	for (USkeletalMeshComponent* SkeletalMeshComponent : SkeletalMeshComponents)
 	{
-		if (SkeletalMeshComponent && SkeletalMeshComponent->GetSkeletalMesh())
+		if (SkeletalMeshComponent && SkeletalMeshComponent->GetSkeletalMeshAsset())
 		{
-			ObjectsToSave.Add(SkeletalMeshComponent->GetSkeletalMesh());
+			ObjectsToSave.Add(SkeletalMeshComponent->GetSkeletalMeshAsset());
 		}
 	}
 
@@ -621,7 +621,7 @@ bool UMeshPaintMode::CanSaveMeshPackages() const
 		}
 		else if (USkeletalMeshComponent* SkeletalMeshComponent = Cast<USkeletalMeshComponent>(Component))
 		{
-			Object = SkeletalMeshComponent->GetSkeletalMesh();
+			Object = SkeletalMeshComponent->GetSkeletalMeshAsset();
 		}
 
 		if (Object != nullptr && Object->GetOutermost()->IsDirty())
diff --git a/Engine/Plugins/MeshPainting/Source/MeshPaintEditorMode/Private/MeshPaintModeHelpers.cpp b/Engine/Plugins/MeshPainting/Source/MeshPaintEditorMode/Private/MeshPaintModeHelpers.cpp
--- a/Engine/Plugins/MeshPainting/Source/MeshPaintEditorMode/Private/MeshPaintModeHelpers.cpp
+++ b/Engine/Plugins/MeshPainting/Source/MeshPaintEditorMode/Private/MeshPaintModeHelpers.cpp
@@ -233,10 +233,10 @@ void UMeshPaintModeSubsystem::ImportVertexColorsFromTexture(UMeshComponent* Mesh
 				{
 					USkeletalMeshComponent* SkeletalMeshComponent = Cast<USkeletalMeshComponent>(MeshComponent);
 
-					if (SkeletalMeshComponent->GetSkeletalMesh())
+					if (SkeletalMeshComponent->GetSkeletalMeshAsset())
 					{
 						// Import colors to skeletal mesh
-						ImportVertexColorsToSkeletalMesh(SkeletalMeshComponent->GetSkeletalMesh(), Options, ColorTexture);
+						ImportVertexColorsToSkeletalMesh(SkeletalMeshComponent->GetSkeletalMeshAsset(), Options, ColorTexture);
 					}
 				}
 			}
diff --git a/Engine/Plugins/MeshPainting/Source/MeshPaintingToolset/Private/MeshPaintHelpers.cpp b/Engine/Plugins/MeshPainting/Source/MeshPaintingToolset/Private/MeshPaintHelpers.cpp
--- a/Engine/Plugins/MeshPainting/Source/MeshPaintingToolset/Private/MeshPaintHelpers.cpp
+++ b/Engine/Plugins/MeshPainting/Source/MeshPaintingToolset/Private/MeshPaintHelpers.cpp
@@ -784,7 +784,7 @@ void UMeshPaintingSubsystem::FillStaticMeshVertexColors(UStaticMeshComponent* Me
 void UMeshPaintingSubsystem::FillSkeletalMeshVertexColors(USkeletalMeshComponent* MeshComponent, int32 LODIndex, const FColor FillColor, const FColor MaskColor)
 {
 	TUniquePtr< FSkinnedMeshComponentRecreateRenderStateContext > RecreateRenderStateContext;
-	USkeletalMesh* Mesh = MeshComponent->GetSkeletalMesh();
+	USkeletalMesh* Mesh = MeshComponent->GetSkeletalMeshAsset();
 	if (Mesh)
 	{
 		// Dirty the mesh
@@ -1005,7 +1005,7 @@ bool UMeshPaintingSubsystem::TryGetNumberOfLODs(const UMeshComponent* MeshCompon
 	}
 	else if (const USkeletalMeshComponent* SkeletalMeshComponent = Cast<USkeletalMeshComponent>(MeshComponent))
 	{
-		const USkeletalMesh* SkeletalMesh = SkeletalMeshComponent->GetSkeletalMesh();
+		const USkeletalMesh* SkeletalMesh = SkeletalMeshComponent->GetSkeletalMeshAsset();
 		if (SkeletalMesh != nullptr)
 		{
 			OutNumLODs = SkeletalMesh->GetLODNum();
@@ -1037,7 +1037,7 @@ int32 UMeshPaintingSubsystem::GetNumberOfUVs(const UMeshComponent* MeshComponent
 	}
 	else if (const USkeletalMeshComponent* SkeletalMeshComponent = Cast<USkeletalMeshComponent>(MeshComponent))
 	{
-		const USkeletalMesh* SkeletalMesh = SkeletalMeshComponent->GetSkeletalMesh();
+		const USkeletalMesh* SkeletalMesh = SkeletalMeshComponent->GetSkeletalMeshAsset();
 		if (SkeletalMesh != nullptr && SkeletalMesh->GetResourceForRendering() && SkeletalMesh->GetResourceForRendering()->LODRenderData.IsValidIndex(LODIndex))
 		{
 			NumUVs = SkeletalMesh->GetResourceForRendering()->LODRenderData[LODIndex].GetNumTexCoords();
@@ -1075,7 +1075,7 @@ bool UMeshPaintingSubsystem::DoesMeshComponentContainPerLODColors(const UMeshCom
 	}
 	else if (const USkeletalMeshComponent* SkeletalMeshComponent = Cast<USkeletalMeshComponent>(MeshComponent))
 	{
-		USkeletalMesh* SkeletalMesh = SkeletalMeshComponent->GetSkeletalMesh();
+		USkeletalMesh* SkeletalMesh = SkeletalMeshComponent->GetSkeletalMeshAsset();
 		if (SkeletalMesh)
 		{
 			const TArray<FSkeletalMeshLODInfo>& LODInfo = SkeletalMesh->GetLODInfoArray();
@@ -1266,7 +1266,7 @@ typedef TOctree2<FPaintedMeshVertex, FVertexColorPropogationOctreeSemantics> TVe
 void UMeshPaintingSubsystem::ApplyVertexColorsToAllLODs(IMeshPaintComponentAdapter& GeometryInfo, USkeletalMeshComponent* SkeletalMeshComponent)
 {
 	checkf(SkeletalMeshComponent != nullptr, TEXT("Invalid Skeletal Mesh Component"));
-	USkeletalMesh* Mesh = SkeletalMeshComponent->GetSkeletalMesh();
+	USkeletalMesh* Mesh = SkeletalMeshComponent->GetSkeletalMeshAsset();
 	if (Mesh)
 	{
 		FSkeletalMeshRenderData* Resource = Mesh->GetResourceForRendering();
diff --git a/Engine/Plugins/MeshPainting/Source/MeshPaintingToolset/Private/MeshPaintSkeletalMeshAdapter.cpp b/Engine/Plugins/MeshPainting/Source/MeshPaintingToolset/Private/MeshPaintSkeletalMeshAdapter.cpp
--- a/Engine/Plugins/MeshPainting/Source/MeshPaintingToolset/Private/MeshPaintSkeletalMeshAdapter.cpp
+++ b/Engine/Plugins/MeshPainting/Source/MeshPaintingToolset/Private/MeshPaintSkeletalMeshAdapter.cpp
@@ -111,9 +111,9 @@ bool FMeshPaintSkeletalMeshComponentAdapter::Construct(UMeshComponent* InCompone
 	{
 		SkeletalMeshChangedHandle = SkeletalMeshComponent->RegisterOnSkeletalMeshPropertyChanged(USkeletalMeshComponent::FOnSkeletalMeshPropertyChanged::CreateSP(this, &FMeshPaintSkeletalMeshComponentAdapter::OnSkeletalMeshChanged));
 
-		if (SkeletalMeshComponent->GetSkeletalMesh() != nullptr)
+		if (SkeletalMeshComponent->GetSkeletalMeshAsset() != nullptr)
 		{
-			ReferencedSkeletalMesh = SkeletalMeshComponent->GetSkeletalMesh();
+			ReferencedSkeletalMesh = SkeletalMeshComponent->GetSkeletalMeshAsset();
 			MeshLODIndex = InMeshLODIndex;
 			const bool bSuccess = Initialize();
 			return bSuccess;
@@ -140,8 +140,8 @@ void FMeshPaintSkeletalMeshComponentAdapter::OnSkeletalMeshChanged()
 	OnRemoved();
 	if (SkeletalMeshComponent.IsValid())
 	{
-		ReferencedSkeletalMesh = SkeletalMeshComponent->GetSkeletalMesh();
-		if (SkeletalMeshComponent->GetSkeletalMesh() != nullptr)
+		ReferencedSkeletalMesh = SkeletalMeshComponent->GetSkeletalMeshAsset();
+		if (SkeletalMeshComponent->GetSkeletalMeshAsset() != nullptr)
 		{
 			Initialize();
 			OnAdded();
@@ -163,7 +163,7 @@ bool FMeshPaintSkeletalMeshComponentAdapter::Initialize()
 
 	if (SkeletalMeshComponent.IsValid())
 	{
-		check(ReferencedSkeletalMesh == SkeletalMeshComponent->GetSkeletalMesh());
+		check(ReferencedSkeletalMesh == SkeletalMeshComponent->GetSkeletalMeshAsset());
 
 		MeshResource = ReferencedSkeletalMesh->GetResourceForRendering();
 		if (MeshResource != nullptr)
@@ -230,7 +230,7 @@ void FMeshPaintSkeletalMeshComponentAdapter::OnAdded()
 	}
 
 	checkf(ReferencedSkeletalMesh, TEXT("Invalid reference to Skeletal Mesh"));
-	checkf(ReferencedSkeletalMesh == SkeletalMeshComponent->GetSkeletalMesh(), TEXT("Referenced Skeletal Mesh does not match one in Component"));
+	checkf(ReferencedSkeletalMesh == SkeletalMeshComponent->GetSkeletalMeshAsset(), TEXT("Referenced Skeletal Mesh does not match one in Component"));
 
 	SkeletalMeshComponent->bUseRefPoseOnInitAnim = true;
 	SkeletalMeshComponent->InitAnim(true);
@@ -476,7 +476,7 @@ TSharedPtr<IMeshPaintComponentAdapter> FMeshPaintSkeletalMeshComponentAdapterFac
 {
 	if (USkeletalMeshComponent* SkeletalMeshComponent = Cast<USkeletalMeshComponent>(InComponent))
 	{
-		if (SkeletalMeshComponent->GetSkeletalMesh() != nullptr)
+		if (SkeletalMeshComponent->GetSkeletalMeshAsset() != nullptr)
 		{
 			TSharedRef<FMeshPaintSkeletalMeshComponentAdapter> Result = MakeShareable(new FMeshPaintSkeletalMeshComponentAdapter());
 			if (Result->Construct(InComponent, InMeshLODIndex))
diff --git a/Engine/Plugins/MeshPainting/Source/MeshPaintingToolset/Private/TexturePaintToolset.cpp b/Engine/Plugins/MeshPainting/Source/MeshPaintingToolset/Private/TexturePaintToolset.cpp
--- a/Engine/Plugins/MeshPainting/Source/MeshPaintingToolset/Private/TexturePaintToolset.cpp
+++ b/Engine/Plugins/MeshPainting/Source/MeshPaintingToolset/Private/TexturePaintToolset.cpp
@@ -450,7 +450,7 @@ void UTexturePaintToolset::RetrieveMeshSectionsForMaterialIndices(const UMeshCom
 	}
 	else if (const USkeletalMeshComponent* SkeletalMeshComponent = Cast<const USkeletalMeshComponent>(MeshComponent))
 	{
-		const USkeletalMesh* SkeletalMesh = SkeletalMeshComponent->GetSkeletalMesh();
+		const USkeletalMesh* SkeletalMesh = SkeletalMeshComponent->GetSkeletalMeshAsset();
 		if (SkeletalMesh)
 		{
 			const FSkeletalMeshRenderData* Resource = SkeletalMesh->GetResourceForRendering();
diff --git a/Engine/Plugins/MeshPainting/Source/MeshPaintingToolset/Public/MeshPaintSkeletalMeshAdapter.h b/Engine/Plugins/MeshPainting/Source/MeshPaintingToolset/Public/MeshPaintSkeletalMeshAdapter.h
--- a/Engine/Plugins/MeshPainting/Source/MeshPaintingToolset/Public/MeshPaintSkeletalMeshAdapter.h
+++ b/Engine/Plugins/MeshPainting/Source/MeshPaintingToolset/Public/MeshPaintSkeletalMeshAdapter.h
@@ -31,7 +31,7 @@ public:
 	virtual bool Initialize() override;
 	virtual void OnAdded() override;
 	virtual void OnRemoved() override;
-	virtual bool IsValid() const override { return SkeletalMeshComponent.IsValid() && ReferencedSkeletalMesh && SkeletalMeshComponent->GetSkeletalMesh() == ReferencedSkeletalMesh; }
+	virtual bool IsValid() const override { return SkeletalMeshComponent.IsValid() && ReferencedSkeletalMesh && SkeletalMeshComponent->GetSkeletalMeshAsset() == ReferencedSkeletalMesh; }
 	virtual bool SupportsTexturePaint() const override { return true; }
 	virtual bool SupportsVertexPaint() const override { return SkeletalMeshComponent.IsValid(); }
 	virtual bool LineTraceComponent(struct FHitResult& OutHit, const FVector Start, const FVector End, const struct FCollisionQueryParams& Params) const override;
diff --git a/Engine/Plugins/MovieScene/SequencerScripting/Source/SequencerScriptingEditor/Private/SequencerTools.cpp b/Engine/Plugins/MovieScene/SequencerScripting/Source/SequencerScriptingEditor/Private/SequencerTools.cpp
--- a/Engine/Plugins/MovieScene/SequencerScripting/Source/SequencerScriptingEditor/Private/SequencerTools.cpp
+++ b/Engine/Plugins/MovieScene/SequencerScripting/Source/SequencerScriptingEditor/Private/SequencerTools.cpp
@@ -278,7 +278,7 @@ static USkeletalMeshComponent* GetSkelMeshComponent(IMovieScenePlayer* Player, c
 		}
 		else if (USkeletalMeshComponent* SkeletalMeshComponent = Cast<USkeletalMeshComponent>(RuntimeObject.Get()))
 		{
-			if (SkeletalMeshComponent->GetSkeletalMesh())
+			if (SkeletalMeshComponent->GetSkeletalMeshAsset())
 			{
 				return SkeletalMeshComponent;
 			}
@@ -316,9 +316,9 @@ bool USequencerToolsFunctionLibrary::ExportAnimSequence(UWorld* World, ULevelSeq
 		}
  
 		USkeletalMeshComponent* SkeletalMeshComp =  GetSkelMeshComponent(Player, Binding);
-		if (SkeletalMeshComp && SkeletalMeshComp->GetSkeletalMesh() && SkeletalMeshComp->GetSkeletalMesh()->GetSkeleton())
+		if (SkeletalMeshComp && SkeletalMeshComp->GetSkeletalMeshAsset() && SkeletalMeshComp->GetSkeletalMeshAsset()->GetSkeleton())
 		{
-			AnimSequence->SetSkeleton(SkeletalMeshComp->GetSkeletalMesh()->GetSkeleton());
+			AnimSequence->SetSkeleton(SkeletalMeshComp->GetSkeletalMeshAsset()->GetSkeleton());
 			bResult = MovieSceneToolHelpers::ExportToAnimSequence(AnimSequence,ExportOptions, MovieScene, Player, SkeletalMeshComp, Template, RootToLocalTransform);
 		}
 	}
