- Expose Compute Debugging via an Interface outside of the Niagara

#rnx

diff --git a/Engine/Plugins/FX/Niagara/Source/Niagara/Private/DataInterface/NiagaraDataInterfaceDebugDraw.cpp b/Engine/Plugins/FX/Niagara/Source/Niagara/Private/DataInterface/NiagaraDataInterfaceDebugDraw.cpp
--- a/Engine/Plugins/FX/Niagara/Source/Niagara/Private/DataInterface/NiagaraDataInterfaceDebugDraw.cpp
+++ b/Engine/Plugins/FX/Niagara/Source/Niagara/Private/DataInterface/NiagaraDataInterfaceDebugDraw.cpp
@@ -3107,7 +3107,7 @@ bool UNiagaraDataInterfaceDebugDraw::InitPerInstanceData(void* PerInstanceData,
 		{
 			check(!RT_Proxy->SystemInstancesToProxyData_RT.Contains(RT_InstanceID));
 			FNDIDebugDrawInstanceData_RenderThread* RT_InstanceData = &RT_Proxy->SystemInstancesToProxyData_RT.Add(RT_InstanceID);
-			RT_InstanceData->GpuComputeDebug = RT_ComputeDispatchInterface->GetGpuComputeDebug();
+			RT_InstanceData->GpuComputeDebug = RT_ComputeDispatchInterface->GetGpuComputeDebugPrivate();
 			RT_InstanceData->OverrideMaxLineInstances = RT_OverrideMaxLineInstances;
 		}
 	);
diff --git a/Engine/Plugins/FX/Niagara/Source/Niagara/Private/NiagaraDataInterfaceGrid2DCollection.cpp b/Engine/Plugins/FX/Niagara/Source/Niagara/Private/NiagaraDataInterfaceGrid2DCollection.cpp
--- a/Engine/Plugins/FX/Niagara/Source/Niagara/Private/NiagaraDataInterfaceGrid2DCollection.cpp
+++ b/Engine/Plugins/FX/Niagara/Source/Niagara/Private/NiagaraDataInterfaceGrid2DCollection.cpp
@@ -2,7 +2,7 @@
 #include "NiagaraDataInterfaceGrid2DCollection.h"
 #include "NiagaraBatchedElements.h"
 #include "NiagaraConstants.h"
-#include "NiagaraGpuComputeDebug.h"
+#include "NiagaraGpuComputeDebugInterface.h"
 #include "NiagaraGpuComputeDispatchInterface.h"
 #include "NiagaraRenderer.h"
 #include "NiagaraSettings.h"
@@ -2624,17 +2624,15 @@ void FNiagaraDataInterfaceProxyGrid2DCollectionProxy::PostSimulate(const FNDIGpu
 #if NIAGARA_COMPUTEDEBUG_ENABLED && WITH_EDITORONLY_DATA
 	if (ProxyData->bPreviewGrid && ProxyData->CurrentData && ProxyData->CurrentData->IsValid())
 	{
-		if (FNiagaraGpuComputeDebug* GpuComputeDebug = Context.GetComputeDispatchInterface().GetGpuComputeDebug())
+		FNiagaraGpuComputeDebugInterface GpuComputeDebugInterface = Context.GetComputeDispatchInterface().GetGpuComputeDebugInterface();
+		FRDGBuilder& GraphBuilder = Context.GetGraphBuilder();
+		if (ProxyData->PreviewAttribute[0] != INDEX_NONE)
 		{
-			FRDGBuilder& GraphBuilder = Context.GetGraphBuilder();
-			if (ProxyData->PreviewAttribute[0] != INDEX_NONE)
-			{
-				GpuComputeDebug->AddAttributeTexture(GraphBuilder, Context.GetSystemInstanceID(), SourceDIName, ProxyData->CurrentData->GetOrCreateTexture(GraphBuilder), FIntPoint::ZeroValue, ProxyData->PreviewAttribute);
-			}
-			else
-			{
-				GpuComputeDebug->AddTexture(GraphBuilder, Context.GetSystemInstanceID(), SourceDIName, ProxyData->CurrentData->GetOrCreateTexture(GraphBuilder));
-			}
+			GpuComputeDebugInterface.AddAttributeTexture(GraphBuilder, Context.GetSystemInstanceID(), SourceDIName, ProxyData->CurrentData->GetOrCreateTexture(GraphBuilder), FIntPoint::ZeroValue, ProxyData->PreviewAttribute);
+		}
+		else
+		{
+			GpuComputeDebugInterface.AddTexture(GraphBuilder, Context.GetSystemInstanceID(), SourceDIName, ProxyData->CurrentData->GetOrCreateTexture(GraphBuilder));
 		}
 	}
 #endif
diff --git a/Engine/Plugins/FX/Niagara/Source/Niagara/Private/NiagaraDataInterfaceGrid2DCollectionReader.cpp b/Engine/Plugins/FX/Niagara/Source/Niagara/Private/NiagaraDataInterfaceGrid2DCollectionReader.cpp
--- a/Engine/Plugins/FX/Niagara/Source/Niagara/Private/NiagaraDataInterfaceGrid2DCollectionReader.cpp
+++ b/Engine/Plugins/FX/Niagara/Source/Niagara/Private/NiagaraDataInterfaceGrid2DCollectionReader.cpp
@@ -10,9 +10,6 @@
 #include "NiagaraSettings.h"
 #include "NiagaraConstants.h"
 #include "NiagaraComputeExecutionContext.h"
-#if WITH_EDITOR
-#include "NiagaraGpuComputeDebug.h"
-#endif
 
 #define LOCTEXT_NAMESPACE "NiagaraDataInterfaceGrid2DCollectionReader"
 
diff --git a/Engine/Plugins/FX/Niagara/Source/Niagara/Private/NiagaraDataInterfaceGrid3DCollection.cpp b/Engine/Plugins/FX/Niagara/Source/Niagara/Private/NiagaraDataInterfaceGrid3DCollection.cpp
--- a/Engine/Plugins/FX/Niagara/Source/Niagara/Private/NiagaraDataInterfaceGrid3DCollection.cpp
+++ b/Engine/Plugins/FX/Niagara/Source/Niagara/Private/NiagaraDataInterfaceGrid3DCollection.cpp
@@ -8,9 +8,7 @@
 #include "NiagaraShader.h"
 #include "NiagaraShaderParametersBuilder.h"
 #include "NiagaraSystemInstance.h"
-#if WITH_EDITOR
-#include "NiagaraGpuComputeDebug.h"
-#endif
+#include "NiagaraGpuComputeDebugInterface.h"
 
 #include "Engine/TextureRenderTargetVolume.h"
 #include "Engine/VolumeTexture.h"
@@ -3055,18 +3053,16 @@ void FNiagaraDataInterfaceProxyGrid3DCollectionProxy::PostSimulate(const FNDIGpu
 #if WITH_EDITOR
 	if (ProxyData->bPreviewGrid && ProxyData->CurrentData && ProxyData->CurrentData->IsValid())
 	{
-		if (FNiagaraGpuComputeDebug* GpuComputeDebug = Context.GetComputeDispatchInterface().GetGpuComputeDebug())
+		FNiagaraGpuComputeDebugInterface GpuComputeDebugInterface = Context.GetComputeDispatchInterface().GetGpuComputeDebugInterface();
+		FRDGBuilder& GraphBuilder = Context.GetGraphBuilder();
+		if (ProxyData->PreviewAttribute[0] != INDEX_NONE)
 		{
-			FRDGBuilder& GraphBuilder = Context.GetGraphBuilder();
-			if (ProxyData->PreviewAttribute[0] != INDEX_NONE)
-			{
-				const FIntVector4 TotalNumAttributeVector = FIntVector4(ProxyData->NumTiles.X, ProxyData->NumTiles.Y, ProxyData->NumTiles.Z, 0);
-				GpuComputeDebug->AddAttributeTexture(GraphBuilder, Context.GetSystemInstanceID(), SourceDIName, ProxyData->CurrentData->GetOrCreateTexture(GraphBuilder), TotalNumAttributeVector, ProxyData->PreviewAttribute);
-			}
-			else
-			{
-				GpuComputeDebug->AddTexture(GraphBuilder, Context.GetSystemInstanceID(), SourceDIName, ProxyData->CurrentData->GetOrCreateTexture(GraphBuilder));
-			}
+			const FIntVector4 TotalNumAttributeVector = FIntVector4(ProxyData->NumTiles.X, ProxyData->NumTiles.Y, ProxyData->NumTiles.Z, 0);
+			GpuComputeDebugInterface.AddAttributeTexture(GraphBuilder, Context.GetSystemInstanceID(), SourceDIName, ProxyData->CurrentData->GetOrCreateTexture(GraphBuilder), TotalNumAttributeVector, ProxyData->PreviewAttribute);
+		}
+		else
+		{
+			GpuComputeDebugInterface.AddTexture(GraphBuilder, Context.GetSystemInstanceID(), SourceDIName, ProxyData->CurrentData->GetOrCreateTexture(GraphBuilder));
 		}
 	}
 #endif
diff --git a/Engine/Plugins/FX/Niagara/Source/Niagara/Private/NiagaraDataInterfaceGrid3DCollectionReader.cpp b/Engine/Plugins/FX/Niagara/Source/Niagara/Private/NiagaraDataInterfaceGrid3DCollectionReader.cpp
--- a/Engine/Plugins/FX/Niagara/Source/Niagara/Private/NiagaraDataInterfaceGrid3DCollectionReader.cpp
+++ b/Engine/Plugins/FX/Niagara/Source/Niagara/Private/NiagaraDataInterfaceGrid3DCollectionReader.cpp
@@ -8,9 +8,6 @@
 #include "NiagaraSettings.h"
 #include "NiagaraConstants.h"
 #include "NiagaraComputeExecutionContext.h"
-#if WITH_EDITOR
-#include "NiagaraGpuComputeDebug.h"
-#endif
 
 #define LOCTEXT_NAMESPACE "NiagaraDataInterfaceGrid3DCollectionReader"
 
diff --git a/Engine/Plugins/FX/Niagara/Source/Niagara/Private/NiagaraDataInterfaceIntRenderTarget2D.cpp b/Engine/Plugins/FX/Niagara/Source/Niagara/Private/NiagaraDataInterfaceIntRenderTarget2D.cpp
--- a/Engine/Plugins/FX/Niagara/Source/Niagara/Private/NiagaraDataInterfaceIntRenderTarget2D.cpp
+++ b/Engine/Plugins/FX/Niagara/Source/Niagara/Private/NiagaraDataInterfaceIntRenderTarget2D.cpp
@@ -4,9 +4,7 @@
 #include "NiagaraSettings.h"
 #include "NiagaraSystemInstance.h"
 #include "NiagaraRenderer.h"
-#if NIAGARA_COMPUTEDEBUG_ENABLED
-#include "NiagaraGpuComputeDebug.h"
-#endif
+#include "NiagaraGpuComputeDebugInterface.h"
 #include "NiagaraGpuComputeDispatchInterface.h"
 #include "NiagaraShader.h"
 #include "NiagaraShaderParametersBuilder.h"
@@ -142,10 +140,8 @@ struct FNDIIntRenderTarget2DProxy : public FNiagaraDataInterfaceProxyRW
 #if NIAGARA_COMPUTEDEBUG_ENABLED && WITH_EDITORONLY_DATA
 		if (InstanceData.bPreviewRenderTarget && InstanceData.TransientRDGTexture)
 		{
-			if (FNiagaraGpuComputeDebug* GpuComputeDebug = Context.GetComputeDispatchInterface().GetGpuComputeDebug())
-			{
-				GpuComputeDebug->AddTexture(Context.GetGraphBuilder(), Context.GetSystemInstanceID(), SourceDIName, InstanceData.TransientRDGTexture);
-			}
+			FNiagaraGpuComputeDebugInterface GpuComputeDebugInterface = Context.GetComputeDispatchInterface().GetGpuComputeDebugInterface();
+			GpuComputeDebugInterface.AddTexture(Context.GetGraphBuilder(), Context.GetSystemInstanceID(), SourceDIName, InstanceData.TransientRDGTexture);
 		}
 #endif
 		if (Context.IsFinalPostSimulate())
diff --git a/Engine/Plugins/FX/Niagara/Source/Niagara/Private/NiagaraDataInterfaceRenderTarget2D.cpp b/Engine/Plugins/FX/Niagara/Source/Niagara/Private/NiagaraDataInterfaceRenderTarget2D.cpp
--- a/Engine/Plugins/FX/Niagara/Source/Niagara/Private/NiagaraDataInterfaceRenderTarget2D.cpp
+++ b/Engine/Plugins/FX/Niagara/Source/Niagara/Private/NiagaraDataInterfaceRenderTarget2D.cpp
@@ -11,7 +11,7 @@
 #include "NiagaraSettings.h"
 #include "NiagaraSystemInstance.h"
 #include "NiagaraRenderer.h"
-#include "NiagaraGpuComputeDebug.h"
+#include "NiagaraGpuComputeDebugInterface.h"
 #include "NiagaraGpuComputeDispatchInterface.h"
 #include "NiagaraGPUProfilerInterface.h"
 #include "NiagaraShader.h"
@@ -880,10 +880,8 @@ void FNiagaraDataInterfaceProxyRenderTarget2DProxy::PostSimulate(const FNDIGpuCo
 #if NIAGARA_COMPUTEDEBUG_ENABLED && WITH_EDITORONLY_DATA
 	if (ProxyData->bPreviewTexture && ProxyData->TransientRDGTexture)
 	{
-		if (FNiagaraGpuComputeDebug* GpuComputeDebug = Context.GetComputeDispatchInterface().GetGpuComputeDebug())
-		{
-			GpuComputeDebug->AddTexture(Context.GetGraphBuilder(), Context.GetSystemInstanceID(), SourceDIName, ProxyData->TransientRDGTexture);
-		}
+		FNiagaraGpuComputeDebugInterface GpuComputeDebugInterface = Context.GetComputeDispatchInterface().GetGpuComputeDebugInterface();
+		GpuComputeDebugInterface.AddTexture(Context.GetGraphBuilder(), Context.GetSystemInstanceID(), SourceDIName, ProxyData->TransientRDGTexture);
 	}
 #endif
 
diff --git a/Engine/Plugins/FX/Niagara/Source/Niagara/Private/NiagaraDataInterfaceRenderTarget2DArray.cpp b/Engine/Plugins/FX/Niagara/Source/Niagara/Private/NiagaraDataInterfaceRenderTarget2DArray.cpp
--- a/Engine/Plugins/FX/Niagara/Source/Niagara/Private/NiagaraDataInterfaceRenderTarget2DArray.cpp
+++ b/Engine/Plugins/FX/Niagara/Source/Niagara/Private/NiagaraDataInterfaceRenderTarget2DArray.cpp
@@ -13,7 +13,7 @@
 #include "NiagaraSettings.h"
 #include "NiagaraShader.h"
 #include "NiagaraShaderParametersBuilder.h"
-#include "NiagaraGpuComputeDebug.h"
+#include "NiagaraGpuComputeDebugInterface.h"
 #include "NiagaraGpuComputeDispatchInterface.h"
 
 #define LOCTEXT_NAMESPACE "NiagaraDataInterfaceRenderTarget2DArray"
@@ -708,10 +708,8 @@ void FNiagaraDataInterfaceProxyRenderTarget2DArrayProxy::PostSimulate(const FNDI
 #if NIAGARA_COMPUTEDEBUG_ENABLED && WITH_EDITORONLY_DATA
 	if (ProxyData->bPreviewTexture && ProxyData->TransientRDGTexture)
 	{
-		if (FNiagaraGpuComputeDebug* GpuComputeDebug = Context.GetComputeDispatchInterface().GetGpuComputeDebug())
-		{
-			GpuComputeDebug->AddTexture(Context.GetGraphBuilder(), Context.GetSystemInstanceID(), SourceDIName, ProxyData->TransientRDGTexture);
-		}
+		FNiagaraGpuComputeDebugInterface GpuComputeDebugInterface = Context.GetComputeDispatchInterface().GetGpuComputeDebugInterface();
+		GpuComputeDebugInterface.AddTexture(Context.GetGraphBuilder(), Context.GetSystemInstanceID(), SourceDIName, ProxyData->TransientRDGTexture);
 	}
 #endif
 
diff --git a/Engine/Plugins/FX/Niagara/Source/Niagara/Private/NiagaraDataInterfaceRenderTargetCube.cpp b/Engine/Plugins/FX/Niagara/Source/Niagara/Private/NiagaraDataInterfaceRenderTargetCube.cpp
--- a/Engine/Plugins/FX/Niagara/Source/Niagara/Private/NiagaraDataInterfaceRenderTargetCube.cpp
+++ b/Engine/Plugins/FX/Niagara/Source/Niagara/Private/NiagaraDataInterfaceRenderTargetCube.cpp
@@ -11,7 +11,7 @@
 #include "NiagaraSettings.h"
 #include "NiagaraShader.h"
 #include "NiagaraShaderParametersBuilder.h"
-#include "NiagaraGpuComputeDebug.h"
+#include "NiagaraGpuComputeDebugInterface.h"
 #include "NiagaraGpuComputeDispatchInterface.h"
 
 #define LOCTEXT_NAMESPACE "NiagaraDataInterfaceRenderTargetCube"
@@ -684,10 +684,8 @@ void FNiagaraDataInterfaceProxyRenderTargetCubeProxy::PostSimulate(const FNDIGpu
 #if NIAGARA_COMPUTEDEBUG_ENABLED && WITH_EDITORONLY_DATA
 	if (ProxyData->bPreviewTexture && ProxyData->TransientRDGTexture)
 	{
-		if (FNiagaraGpuComputeDebug* GpuComputeDebug = Context.GetComputeDispatchInterface().GetGpuComputeDebug())
-		{
-			GpuComputeDebug->AddTexture(Context.GetGraphBuilder(), Context.GetSystemInstanceID(), SourceDIName, ProxyData->TransientRDGTexture);
-		}
+		FNiagaraGpuComputeDebugInterface GpuComputeDebugInterface = Context.GetComputeDispatchInterface().GetGpuComputeDebugInterface();
+		GpuComputeDebugInterface.AddTexture(Context.GetGraphBuilder(), Context.GetSystemInstanceID(), SourceDIName, ProxyData->TransientRDGTexture);
 	}
 #endif
 
diff --git a/Engine/Plugins/FX/Niagara/Source/Niagara/Private/NiagaraDataInterfaceRenderTargetVolume.cpp b/Engine/Plugins/FX/Niagara/Source/Niagara/Private/NiagaraDataInterfaceRenderTargetVolume.cpp
--- a/Engine/Plugins/FX/Niagara/Source/Niagara/Private/NiagaraDataInterfaceRenderTargetVolume.cpp
+++ b/Engine/Plugins/FX/Niagara/Source/Niagara/Private/NiagaraDataInterfaceRenderTargetVolume.cpp
@@ -13,7 +13,7 @@
 #include "NiagaraSettings.h"
 #include "NiagaraShader.h"
 #include "NiagaraShaderParametersBuilder.h"
-#include "NiagaraGpuComputeDebug.h"
+#include "NiagaraGpuComputeDebugInterface.h"
 
 #define LOCTEXT_NAMESPACE "NiagaraDataInterfaceRenderTargetVolume"
 
@@ -702,10 +702,8 @@ void FNiagaraDataInterfaceProxyRenderTargetVolumeProxy::PostSimulate(const FNDIG
 #if NIAGARA_COMPUTEDEBUG_ENABLED && WITH_EDITORONLY_DATA
 	if (ProxyData->bPreviewTexture && ProxyData->TransientRDGTexture)
 	{
-		if (FNiagaraGpuComputeDebug* GpuComputeDebug = Context.GetComputeDispatchInterface().GetGpuComputeDebug())
-		{
-			GpuComputeDebug->AddTexture(Context.GetGraphBuilder(), Context.GetSystemInstanceID(), SourceDIName, ProxyData->TransientRDGTexture);
-		}
+		FNiagaraGpuComputeDebugInterface GpuComputeDebugInterface = Context.GetComputeDispatchInterface().GetGpuComputeDebugInterface();
+		GpuComputeDebugInterface.AddTexture(Context.GetGraphBuilder(), Context.GetSystemInstanceID(), SourceDIName, ProxyData->TransientRDGTexture);
 	}
 #endif
 
diff --git a/Engine/Plugins/FX/Niagara/Source/Niagara/Private/NiagaraDataInterfaceVolumeCache.cpp b/Engine/Plugins/FX/Niagara/Source/Niagara/Private/NiagaraDataInterfaceVolumeCache.cpp
--- a/Engine/Plugins/FX/Niagara/Source/Niagara/Private/NiagaraDataInterfaceVolumeCache.cpp
+++ b/Engine/Plugins/FX/Niagara/Source/Niagara/Private/NiagaraDataInterfaceVolumeCache.cpp
@@ -12,9 +12,6 @@
 #include "NiagaraConstants.h"
 #include "NiagaraComputeExecutionContext.h"
 #include "NiagaraShaderParametersBuilder.h"
-#if WITH_EDITOR
-#include "NiagaraGpuComputeDebug.h"
-#endif
 
 #define LOCTEXT_NAMESPACE "NiagaraDataInterfaceVolumeCache"
 
diff --git a/Engine/Plugins/FX/Niagara/Source/Niagara/Private/NiagaraGpuComputeDebugInterface.cpp b/Engine/Plugins/FX/Niagara/Source/Niagara/Private/NiagaraGpuComputeDebugInterface.cpp
new file mode 100644
--- /dev/null
+++ b/Engine/Plugins/FX/Niagara/Source/Niagara/Private/NiagaraGpuComputeDebugInterface.cpp
@@ -0,0 +1,32 @@
+// Copyright Epic Games, Inc. All Rights Reserved.
+
+#include "NiagaraGpuComputeDebugInterface.h"
+#include "NiagaraGpuComputeDebug.h"
+
+#if NIAGARA_COMPUTEDEBUG_ENABLED
+
+void FNiagaraGpuComputeDebugInterface::AddTexture(FRDGBuilder& GraphBuilder, FNiagaraSystemInstanceID SystemInstanceID, FName SourceName, FRDGTextureRef Texture, FVector2D PreviewDisplayRange)
+{
+	if (ComputeDebug != nullptr)
+	{
+		ComputeDebug->AddTexture(GraphBuilder, SystemInstanceID, SourceName, Texture, PreviewDisplayRange);
+	}
+}
+
+void FNiagaraGpuComputeDebugInterface::AddAttributeTexture(FRDGBuilder& GraphBuilder, FNiagaraSystemInstanceID SystemInstanceID, FName SourceName, FRDGTextureRef Texture, FIntPoint NumTextureAttributes, FIntVector4 AttributeIndices, FVector2D PreviewDisplayRange)
+{
+	if (ComputeDebug != nullptr)
+	{
+		ComputeDebug->AddAttributeTexture(GraphBuilder, SystemInstanceID, SourceName, Texture, NumTextureAttributes, AttributeIndices, PreviewDisplayRange);
+	}
+}
+
+void FNiagaraGpuComputeDebugInterface::AddAttributeTexture(FRDGBuilder& GraphBuilder, FNiagaraSystemInstanceID SystemInstanceID, FName SourceName, FRDGTextureRef Texture, FIntVector4 NumTextureAttributes, FIntVector4 AttributeIndices, FVector2D PreviewDisplayRange)
+{
+	if (ComputeDebug != nullptr )
+	{
+		ComputeDebug->AddAttributeTexture(GraphBuilder, SystemInstanceID, SourceName, Texture, NumTextureAttributes, AttributeIndices, PreviewDisplayRange);
+	}
+}
+
+#endif //NIAGARA_COMPUTEDEBUG_ENABLED
diff --git a/Engine/Plugins/FX/Niagara/Source/Niagara/Private/NiagaraGpuComputeDispatch.cpp b/Engine/Plugins/FX/Niagara/Source/Niagara/Private/NiagaraGpuComputeDispatch.cpp
--- a/Engine/Plugins/FX/Niagara/Source/Niagara/Private/NiagaraGpuComputeDispatch.cpp
+++ b/Engine/Plugins/FX/Niagara/Source/Niagara/Private/NiagaraGpuComputeDispatch.cpp
@@ -1641,7 +1641,7 @@ void FNiagaraGpuComputeDispatch::PreInitViews(FRDGBuilder& GraphBuilder, bool bA
 
 	GpuReadbackManagerPtr->Tick();
 #if NIAGARA_COMPUTEDEBUG_ENABLED
-	if ( FNiagaraGpuComputeDebug* GpuComputeDebug = GetGpuComputeDebug() )
+	if ( FNiagaraGpuComputeDebug* GpuComputeDebug = GpuComputeDebugPtr.Get() )
 	{
 		GpuComputeDebug->Tick(GraphBuilder);
 	}
diff --git a/Engine/Plugins/FX/Niagara/Source/Niagara/Private/NiagaraGpuComputeDispatchInterface.cpp b/Engine/Plugins/FX/Niagara/Source/Niagara/Private/NiagaraGpuComputeDispatchInterface.cpp
--- a/Engine/Plugins/FX/Niagara/Source/Niagara/Private/NiagaraGpuComputeDispatchInterface.cpp
+++ b/Engine/Plugins/FX/Niagara/Source/Niagara/Private/NiagaraGpuComputeDispatchInterface.cpp
@@ -3,6 +3,7 @@
 #include "NiagaraGpuComputeDispatchInterface.h"
 #include "NiagaraGpuComputeDispatch.h"
 #include "NiagaraGpuComputeDebug.h"
+#include "NiagaraGpuComputeDebugInterface.h"
 #include "NiagaraGpuReadbackManager.h"
 
 #include "Engine/World.h"
@@ -37,6 +38,13 @@ FNiagaraGpuComputeDispatchInterface* FNiagaraGpuComputeDispatchInterface::Get(FF
 	return FXSceneInterface ? static_cast<FNiagaraGpuComputeDispatchInterface*>(FXSceneInterface->GetInterface(FNiagaraGpuComputeDispatch::Name)) : nullptr;
 }
 
+#if NIAGARA_COMPUTEDEBUG_ENABLED
+FNiagaraGpuComputeDebugInterface FNiagaraGpuComputeDispatchInterface::GetGpuComputeDebugInterface() const
+{
+	return FNiagaraGpuComputeDebugInterface(GpuComputeDebugPtr.Get());
+}
+#endif
+
 FRDGTextureRef FNiagaraGpuComputeDispatchInterface::GetBlackTexture(FRDGBuilder& GraphBuilder, ETextureDimension TextureDimension) const
 {
 	const FRDGSystemTextures& SystemTextures = FRDGSystemTextures::Get(GraphBuilder);
diff --git a/Engine/Plugins/FX/Niagara/Source/Niagara/Private/NiagaraSystemInstance.cpp b/Engine/Plugins/FX/Niagara/Source/Niagara/Private/NiagaraSystemInstance.cpp
--- a/Engine/Plugins/FX/Niagara/Source/Niagara/Private/NiagaraSystemInstance.cpp
+++ b/Engine/Plugins/FX/Niagara/Source/Niagara/Private/NiagaraSystemInstance.cpp
@@ -575,7 +575,7 @@ void FNiagaraSystemInstance::SetGpuComputeDebug(bool bEnableDebug)
 		(
 			[RT_ComputeDispatchInterface=ComputeDispatchInterface, RT_InstanceID=GetId(), RT_SystemName=SystemName](FRHICommandListImmediate& RHICmdList)
 			{
-				if (FNiagaraGpuComputeDebug* GpuComputeDebug = RT_ComputeDispatchInterface->GetGpuComputeDebug())
+				if (FNiagaraGpuComputeDebug* GpuComputeDebug = RT_ComputeDispatchInterface->GetGpuComputeDebugPrivate())
 				{
 					GpuComputeDebug->AddSystemInstance(RT_InstanceID, RT_SystemName);
 				}
@@ -588,7 +588,7 @@ void FNiagaraSystemInstance::SetGpuComputeDebug(bool bEnableDebug)
 		(
 			[RT_ComputeDispatchInterface=ComputeDispatchInterface, RT_InstanceID=GetId()](FRHICommandListImmediate& RHICmdList)
 			{
-				if (FNiagaraGpuComputeDebug* GpuComputeDebug = RT_ComputeDispatchInterface->GetGpuComputeDebug())
+				if (FNiagaraGpuComputeDebug* GpuComputeDebug = RT_ComputeDispatchInterface->GetGpuComputeDebugPrivate())
 				{
 					GpuComputeDebug->RemoveSystemInstance(RT_InstanceID);
 				}
diff --git a/Engine/Plugins/FX/Niagara/Source/Niagara/Public/NiagaraGpuComputeDebugInterface.h b/Engine/Plugins/FX/Niagara/Source/Niagara/Public/NiagaraGpuComputeDebugInterface.h
new file mode 100644
--- /dev/null
+++ b/Engine/Plugins/FX/Niagara/Source/Niagara/Public/NiagaraGpuComputeDebugInterface.h
@@ -0,0 +1,45 @@
+// Copyright Epic Games, Inc. All Rights Reserved.
+
+/*==============================================================================
+Class used help realtime debug Gpu Compute simulations
+==============================================================================*/
+#pragma once
+
+#include "CoreMinimal.h"
+#include "NiagaraCommon.h"
+#include "NiagaraRenderGraphUtils.h"
+#include "RenderGraphDefinitions.h"
+
+#if NIAGARA_COMPUTEDEBUG_ENABLED
+
+class FNiagaraGpuComputeDebug;
+
+class NIAGARA_API FNiagaraGpuComputeDebugInterface
+{
+	friend class FNiagaraGpuComputeDispatchInterface;
+
+private:
+	explicit FNiagaraGpuComputeDebugInterface(FNiagaraGpuComputeDebug* InComputeDebug)
+		: ComputeDebug(InComputeDebug)
+	{
+	}
+
+public:
+	// Add a texture to visualize
+	void AddTexture(FRDGBuilder& GraphBuilder, FNiagaraSystemInstanceID SystemInstanceID, FName SourceName, FRDGTextureRef Texture, FVector2D PreviewDisplayRange = FVector2D::ZeroVector);
+
+	// Add a texture to visualize that contains a number of attributes and select which attributes to push into RGBA where -1 means ignore that channel
+	// The first -1 in the attribute indices list will also limit the number of attributes we attempt to read.
+	// NumTextureAttributes in this version is meant for a 2D atlas
+	void AddAttributeTexture(FRDGBuilder& GraphBuilder, FNiagaraSystemInstanceID SystemInstanceID, FName SourceName, FRDGTextureRef Texture, FIntPoint NumTextureAttributes, FIntVector4 AttributeIndices, FVector2D PreviewDisplayRange = FVector2D::ZeroVector);
+
+	// Add a texture to visualize that contains a number of attributes and select which attributes to push into RGBA where -1 means ignore that channel
+	// The first -1 in the attribute indices list will also limit the number of attributes we attempt to read
+	// NumTextureAttributes in this version is meant for a 3D atlas
+	void AddAttributeTexture(FRDGBuilder& GraphBuilder, FNiagaraSystemInstanceID SystemInstanceID, FName SourceName, FRDGTextureRef Texture, FIntVector4 NumTextureAttributes, FIntVector4 AttributeIndices, FVector2D PreviewDisplayRange = FVector2D::ZeroVector);
+
+private:
+	FNiagaraGpuComputeDebug* ComputeDebug = nullptr;
+};
+
+#endif //NIAGARA_COMPUTEDEBUG_ENABLED
diff --git a/Engine/Plugins/FX/Niagara/Source/Niagara/Public/NiagaraGpuComputeDispatchInterface.h b/Engine/Plugins/FX/Niagara/Source/Niagara/Public/NiagaraGpuComputeDispatchInterface.h
--- a/Engine/Plugins/FX/Niagara/Source/Niagara/Public/NiagaraGpuComputeDispatchInterface.h
+++ b/Engine/Plugins/FX/Niagara/Source/Niagara/Public/NiagaraGpuComputeDispatchInterface.h
@@ -13,6 +13,7 @@ class FRDGExternalAccessQueue;
 class FNiagaraAsyncGpuTraceHelper;
 struct FNiagaraComputeExecutionContext;
 class FNiagaraGpuComputeDebug;
+class FNiagaraGpuComputeDebugInterface;
 class FNiagaraGPUInstanceCountManager;
 class FNiagaraGpuReadbackManager;
 class FNiagaraRayTracingHelper;
@@ -56,8 +57,11 @@ public:
 	FORCEINLINE const FNiagaraGPUInstanceCountManager& GetGPUInstanceCounterManager() const { check(IsInRenderingThread()); return GPUInstanceCounterManager; }
 
 #if NIAGARA_COMPUTEDEBUG_ENABLED
-	/** Get access to Niagara's GpuComputeDebug */
-	FNiagaraGpuComputeDebug* GetGpuComputeDebug() const { return GpuComputeDebugPtr.Get(); }
+	/** Public interface to Niagara compute debugging. */
+	FNiagaraGpuComputeDebugInterface GetGpuComputeDebugInterface() const;
+
+	/** Get access to Niagara's GpuComputeDebug this is for internal use */
+	FNiagaraGpuComputeDebug* GetGpuComputeDebugPrivate() const { return GpuComputeDebugPtr.Get(); }
 #endif
 
 #if WITH_NIAGARA_GPU_PROFILER
