PR #9480: UE-NiagaraRendererRibbon_VertexBufferMaxIndex_NumVertices, MaxVertexIndex != NumInstances (Contributed by [REDACTED_PERSON])

#rnx

diff --git a/Engine/Plugins/FX/Niagara/Source/Niagara/Private/NiagaraRendererRibbons.cpp b/Engine/Plugins/FX/Niagara/Source/Niagara/Private/NiagaraRendererRibbons.cpp
--- a/Engine/Plugins/FX/Niagara/Source/Niagara/Private/NiagaraRendererRibbons.cpp
+++ b/Engine/Plugins/FX/Niagara/Source/Niagara/Private/NiagaraRendererRibbons.cpp
@@ -1973,7 +1973,7 @@ inline void FNiagaraRendererRibbons::SetupMeshBatchAndCollectorResourceForView(c
 	else
 	{
 		MeshElement.NumPrimitives = GeneratedData.CPUTriangleCount;
-		MeshElement.MaxVertexIndex = SourceParticleData->GetNumInstances();
+		MeshElement.MaxVertexIndex = SourceParticleData->GetNumInstances() == 0 ? 0 : SourceParticleData->GetNumInstances() - 1;
 		check(MeshElement.NumPrimitives > 0);
 	}	
 	
@@ -2619,4 +2619,4 @@ void FNiagaraRibbonComputeDispatchManager::GenerateAllGPUData(FRHICommandListImm
 	}
 	
 	RenderersToGenerate.Empty();
-}
\ No newline at end of file
+}
