- When force flushing we should execute it always

#rnx

diff --git a/Engine/Plugins/FX/Niagara/Source/Niagara/Private/NiagaraGpuComputeDispatch.cpp b/Engine/Plugins/FX/Niagara/Source/Niagara/Private/NiagaraGpuComputeDispatch.cpp
--- a/Engine/Plugins/FX/Niagara/Source/Niagara/Private/NiagaraGpuComputeDispatch.cpp
+++ b/Engine/Plugins/FX/Niagara/Source/Niagara/Private/NiagaraGpuComputeDispatch.cpp
@@ -348,7 +348,7 @@ void FNiagaraGpuComputeDispatch::ProcessPendingTicksFlush(FRHICommandListImmedia
 				MaxPendingTicks = FMath::Max<int32>(MaxPendingTicks, Proxy->PendingTicks.Num());
 			}
 		}
-		if (MaxPendingTicks < TickFlushMaxPendingTicks)
+		if (!bForceFlush && (MaxPendingTicks < TickFlushMaxPendingTicks))
 		{
 			return;
 		}
