- For for potential nullptr access if the component is autodestroy

#rnx

diff --git a/Engine/Plugins/FX/Niagara/Source/Niagara/Private/NiagaraComponent.cpp b/Engine/Plugins/FX/Niagara/Source/Niagara/Private/NiagaraComponent.cpp
--- a/Engine/Plugins/FX/Niagara/Source/Niagara/Private/NiagaraComponent.cpp
+++ b/Engine/Plugins/FX/Niagara/Source/Niagara/Private/NiagaraComponent.cpp
@@ -1210,20 +1210,21 @@ void UNiagaraComponent::ActivateInternal(bool bReset /* = false */, bool bIsScal
 	// NOTE: This call can cause SystemInstance itself to get destroyed with auto destroy systems
 	SystemInstanceController->Activate(ResetMode);
 
-	if (SystemInstanceController->IsComplete())
-	{
-		OnSystemComplete(true);
-		return;
-	}
-
-	if (SystemInstanceController && SystemInstanceController->IsSolo())
+	if (SystemInstanceController.IsValid())
 	{
-		const ETickingGroup SoloTickGroup = SystemInstanceController->CalculateTickGroup();
-		PrimaryComponentTick.TickGroup = FMath::Max(GNiagaraSoloTickEarly ? TG_PrePhysics : TG_DuringPhysics, SoloTickGroup);
-		PrimaryComponentTick.EndTickGroup = GNiagaraSoloAllowAsyncWorkToEndOfFrame ? TG_LastDemotable : ETickingGroup(PrimaryComponentTick.TickGroup);
+		if (SystemInstanceController->IsComplete())
+		{
+			OnSystemComplete(true);
+		}
+		else if (SystemInstanceController->IsSolo())
+		{
+			const ETickingGroup SoloTickGroup = SystemInstanceController->CalculateTickGroup();
+			PrimaryComponentTick.TickGroup = FMath::Max(GNiagaraSoloTickEarly ? TG_PrePhysics : TG_DuringPhysics, SoloTickGroup);
+			PrimaryComponentTick.EndTickGroup = GNiagaraSoloAllowAsyncWorkToEndOfFrame ? TG_LastDemotable : ETickingGroup(PrimaryComponentTick.TickGroup);
 
-		// Solo instances require the component tick to be enabled
-		ScopedComponentTickEnabled.bTickEnabled = true;
+			// Solo instances require the component tick to be enabled
+			ScopedComponentTickEnabled.bTickEnabled = true;
+		}
 	}
 }
 
