- Fix SimCache layout generation being incorrect with mixed types

#rnx

diff --git a/Engine/Plugins/FX/Niagara/Source/Niagara/Private/NiagaraSimCache.cpp b/Engine/Plugins/FX/Niagara/Source/Niagara/Private/NiagaraSimCache.cpp
--- a/Engine/Plugins/FX/Niagara/Source/Niagara/Private/NiagaraSimCache.cpp
+++ b/Engine/Plugins/FX/Niagara/Source/Niagara/Private/NiagaraSimCache.cpp
@@ -76,9 +76,6 @@ struct FNiagaraSimCacheHelper
 		CacheLayout.ComponentMappingsFromDataBuffer.Empty(CacheTotalComponents);
 		CacheLayout.ComponentMappingsFromDataBuffer.AddDefaulted(CacheTotalComponents);
 
-		int32 FloatOffset = 0;
-		int32 HalfOffset = CacheLayout.FloatCount;
-		int32 Int32Offset = HalfOffset + CacheLayout.HalfCount;
 		for ( int32 iVariable=0; iVariable < NumVariables; ++iVariable)
 		{
 			const FNiagaraVariableLayoutInfo& DataSetVariableLayout = CompiledData.VariableLayouts[iVariable];
@@ -95,6 +92,15 @@ struct FNiagaraSimCacheHelper
 			CacheLayout.FloatCount += DataSetVariableLayout.GetNumFloatComponents();
 			CacheLayout.HalfCount += DataSetVariableLayout.GetNumHalfComponents();
 			CacheLayout.Int32Count += DataSetVariableLayout.GetNumInt32Components();
+		}
+
+		int32 FloatOffset = 0;
+		int32 HalfOffset = CacheLayout.FloatCount;
+		int32 Int32Offset = HalfOffset + CacheLayout.HalfCount;
+		for (int32 iVariable = 0; iVariable < NumVariables; ++iVariable)
+		{
+			const FNiagaraVariableLayoutInfo& DataSetVariableLayout = CompiledData.VariableLayouts[iVariable];
+			FNiagaraSimCacheVariable& CacheVariable = CacheLayout.Variables[iVariable];
 
 			// Build DataBuffer To / From mappings
 			for (int32 iComponent=0; iComponent < CacheVariable.FloatCount; ++iComponent)
