- Give the system simulation an opportunity to react to sim cache updates
- This allows the renderer bindings to work correctly for example

#rnx

diff --git a/Engine/Plugins/FX/Niagara/Source/Niagara/Private/NiagaraSystemInstance.cpp b/Engine/Plugins/FX/Niagara/Source/Niagara/Private/NiagaraSystemInstance.cpp
--- a/Engine/Plugins/FX/Niagara/Source/Niagara/Private/NiagaraSystemInstance.cpp
+++ b/Engine/Plugins/FX/Niagara/Source/Niagara/Private/NiagaraSystemInstance.cpp
@@ -2221,6 +2221,8 @@ void FNiagaraSystemInstance::SimCacheTick_GameThread(UNiagaraSimCache* SimCache,
 		SystemSimulation->SetInstanceState(this, ENiagaraSystemInstanceState::Running);
 	}
 
+	SystemSimulation->SimCachePostTick_Concurrent(DeltaSeconds, MyCompletionGraphEvent);
+
 	if (OnPostTickDelegate.IsBound())
 	{
 		OnPostTickDelegate.Execute();
diff --git a/Engine/Plugins/FX/Niagara/Source/Niagara/Private/NiagaraSystemSimulation.cpp b/Engine/Plugins/FX/Niagara/Source/Niagara/Private/NiagaraSystemSimulation.cpp
--- a/Engine/Plugins/FX/Niagara/Source/Niagara/Private/NiagaraSystemSimulation.cpp
+++ b/Engine/Plugins/FX/Niagara/Source/Niagara/Private/NiagaraSystemSimulation.cpp
@@ -1574,6 +1574,48 @@ void FNiagaraSystemSimulation::Spawn_Concurrent(FNiagaraSystemSimulationTickCont
 	bInSpawnPhase = false;
 }
 
+void FNiagaraSystemSimulation::SimCachePostTick_Concurrent(float DeltaSeconds, const FGraphEventRef& MyCompletionGraphEvent)
+{
+	check(GetIsSolo());
+	check(System != nullptr);
+
+	TArray<FNiagaraSystemInstance*>& SystemInstances = GetSystemInstances(ENiagaraSystemInstanceState::Running);
+	if (SystemInstances.Num() == 0)
+	{
+		return;
+	}
+
+	if (!bBindingsInitialized)
+	{
+		InitParameterDataSetBindings(SystemInstances[0]);
+	}
+
+	FNiagaraDataSetReaderInt32<ENiagaraExecutionState> SystemExecutionStateAccessor = System->GetSystemExecutionStateAccessor().GetReader(MainDataSet);
+	TConstArrayView<FNiagaraDataSetAccessor<ENiagaraExecutionState>> EmitterExecutionStateAccessors = System->GetEmitterExecutionStateAccessors();
+
+	for ( int32 iSystemInstance=0; iSystemInstance < SystemInstances.Num(); ++iSystemInstance )
+	{
+		FNiagaraSystemInstance& SystemInstance = *SystemInstances[iSystemInstance];
+
+		TArray<TSharedRef<FNiagaraEmitterInstance, ESPMode::ThreadSafe>>& Emitters = SystemInstance.GetEmitters();
+		for (int32 iEmitter=0; iEmitter < Emitters.Num(); ++iEmitter)
+		{
+			FNiagaraEmitterInstance& EmitterInstance = Emitters[iEmitter].Get();
+			if ( EmitterInstance.IsComplete() )
+			{
+				continue;
+			}
+
+			ENiagaraExecutionState State = EmitterExecutionStateAccessors[iEmitter].GetReader(MainDataSet).GetSafe(iSystemInstance, ENiagaraExecutionState::Disabled);
+			EmitterInstance.SetExecutionState(State);
+
+			//DataSetToEmitterSpawnParameters[iEmitter].DataSetToParameterStore(EmitterInstance.GetSpawnExecutionContext().Parameters, MainDataSet, iSystemInstance);
+			//DataSetToEmitterUpdateParameters[iEmitter].DataSetToParameterStore(EmitterInstance.GetUpdateExecutionContext().Parameters, MainDataSet, iSystemInstance);
+			DataSetToEmitterRendererParameters[iEmitter].DataSetToParameterStore(EmitterInstance.GetRendererBoundVariables(), MainDataSet, iSystemInstance);
+		}
+	}
+}
+
 void FNiagaraSystemSimulation::DumpStalledInfo()
 {
 	TStringBuilder<128> Builder;
@@ -2064,7 +2106,6 @@ void FNiagaraSystemSimulation::TransferSystemSimResults(FNiagaraSystemSimulation
 				}
 
 				DataSetToEmitterRendererParameters[EmitterIdx].DataSetToParameterStore(EmitterInst.GetRendererBoundVariables(), Context.DataSet, SystemIndex);
-
 			}
 		}
 	}
diff --git a/Engine/Plugins/FX/Niagara/Source/Niagara/Public/NiagaraSystemSimulation.h b/Engine/Plugins/FX/Niagara/Source/Niagara/Public/NiagaraSystemSimulation.h
--- a/Engine/Plugins/FX/Niagara/Source/Niagara/Public/NiagaraSystemSimulation.h
+++ b/Engine/Plugins/FX/Niagara/Source/Niagara/Public/NiagaraSystemSimulation.h
@@ -268,7 +268,8 @@ public:
 	/** Spawn any pending instances */
 	void Spawn_Concurrent(FNiagaraSystemSimulationTickContext& Context);
 
-	/** Promote instances that have ticked during */
+	/** Called after the sim cache has been read. */
+	void SimCachePostTick_Concurrent(float DeltaSeconds, const FGraphEventRef& MyCompletionGraphEvent);
 
 	/** Wait for system simulation concurrent tick to complete.  If bEnsureComplete is true we will trigger an ensure if it is not complete. */
 	void WaitForConcurrentTickComplete(bool bEnsureComplete = false);
