- Use Viewport SizeXY vs Canvas ViewRect

#rnx

diff --git a/Engine/Plugins/FX/Niagara/Source/NiagaraEditor/Private/Widgets/SNiagaraBakerViewport.cpp b/Engine/Plugins/FX/Niagara/Source/NiagaraEditor/Private/Widgets/SNiagaraBakerViewport.cpp
--- a/Engine/Plugins/FX/Niagara/Source/NiagaraEditor/Private/Widgets/SNiagaraBakerViewport.cpp
+++ b/Engine/Plugins/FX/Niagara/Source/NiagaraEditor/Private/Widgets/SNiagaraBakerViewport.cpp
@@ -269,7 +269,8 @@ public:
 			}
 		}
 
-		const FIntRect ViewRect(Canvas->GetViewRect().Min + FIntPoint(2, 2), Canvas->GetViewRect().Max - FIntPoint(2, 2));
+		//const FIntRect ViewRect(Canvas->GetViewRect().Min + FIntPoint(2, 2), Canvas->GetViewRect().Max - FIntPoint(2, 2));
+		const FIntRect ViewRect(FIntPoint(2, 2), InViewport->GetSizeXY() - FIntPoint(2, 2));
 		if ( ViewRect.Width() <= 0 || ViewRect.Height() <= 0 )
 		{
 			return;
