Fixed multiple bindings not resetting properly, which was due to CDO properties not being set
This still leaves us with visible "reset to default" buttons due to the property node having the previous default value cached.

To get rid of that:
Added option to add custom reset behavior overrides for any properrty, currently used for "Renderer Enabled" and "Simulation Enabled" bindings, which allows us to specify custom "is visible" overrides & functionality

Removed custom reset button for attribute bindings as it should use the default one for the stack

diff --git a/Engine/Plugins/FX/Niagara/Source/Niagara/Private/NiagaraComponentRendererProperties.cpp b/Engine/Plugins/FX/Niagara/Source/Niagara/Private/NiagaraComponentRendererProperties.cpp
--- a/Engine/Plugins/FX/Niagara/Source/Niagara/Private/NiagaraComponentRendererProperties.cpp
+++ b/Engine/Plugins/FX/Niagara/Source/Niagara/Private/NiagaraComponentRendererProperties.cpp
@@ -451,11 +451,17 @@ void UNiagaraComponentRendererProperties::PostDuplicate(bool bDuplicateForPIE)
 	TemplateComponent = DuplicateObject(TemplateComponent, this);
 }
 
-void UNiagaraComponentRendererProperties::InitCDOPropertiesAfterModuleStartup()
+void UNiagaraComponentRendererProperties::InitBindings()
 {
-	UNiagaraComponentRendererProperties* CDO = CastChecked<UNiagaraComponentRendererProperties>(UNiagaraComponentRendererProperties::StaticClass()->GetDefaultObject());
-	CDO->EnabledBinding = FNiagaraConstants::GetAttributeDefaultBinding(SYS_PARAM_PARTICLES_COMPONENTS_ENABLED);
-	CDO->RendererVisibilityTagBinding = FNiagaraConstants::GetAttributeDefaultBinding(SYS_PARAM_PARTICLES_VISIBILITY_TAG);
+	Super::InitBindings();
+
+	EnabledBinding = FNiagaraConstants::GetAttributeDefaultBinding(SYS_PARAM_PARTICLES_COMPONENTS_ENABLED);
+	RendererVisibilityTagBinding = FNiagaraConstants::GetAttributeDefaultBinding(SYS_PARAM_PARTICLES_VISIBILITY_TAG);
+}
+
+void UNiagaraComponentRendererProperties::InitCDOPropertiesAfterModuleStartupInternal()
+{
+	Super::InitCDOPropertiesAfterModuleStartupInternal();
 
 	for (TWeakObjectPtr<UNiagaraComponentRendererProperties>& WeakComponentRendererProperties : ComponentRendererPropertiesToDeferredInit)
 	{
diff --git a/Engine/Plugins/FX/Niagara/Source/Niagara/Private/NiagaraMeshRendererProperties.cpp b/Engine/Plugins/FX/Niagara/Source/Niagara/Private/NiagaraMeshRendererProperties.cpp
--- a/Engine/Plugins/FX/Niagara/Source/Niagara/Private/NiagaraMeshRendererProperties.cpp
+++ b/Engine/Plugins/FX/Niagara/Source/Niagara/Private/NiagaraMeshRendererProperties.cpp
@@ -244,11 +244,10 @@ void UNiagaraMeshRendererProperties::Serialize(FArchive& Ar)
 }
 
 /** The bindings depend on variables that are created during the NiagaraModule startup. However, the CDO's are build prior to this being initialized, so we defer setting these values until later.*/
-void UNiagaraMeshRendererProperties::InitCDOPropertiesAfterModuleStartup()
+void UNiagaraMeshRendererProperties::InitCDOPropertiesAfterModuleStartupInternal()
 {
-	UNiagaraMeshRendererProperties* CDO = CastChecked<UNiagaraMeshRendererProperties>(UNiagaraMeshRendererProperties::StaticClass()->GetDefaultObject());
-	CDO->InitBindings();
-
+	Super::InitCDOPropertiesAfterModuleStartupInternal();
+	
 	for (TWeakObjectPtr<UNiagaraMeshRendererProperties>& WeakMeshRendererProperties : MeshRendererPropertiesToDeferredInit)
 	{
 		if (WeakMeshRendererProperties.Get())
@@ -260,6 +259,8 @@ void UNiagaraMeshRendererProperties::InitCDOPropertiesAfterModuleStartup()
 
 void UNiagaraMeshRendererProperties::InitBindings()
 {
+	Super::InitBindings();
+	
 	if (!PositionBinding.IsValid())
 	{
 		PositionBinding = FNiagaraConstants::GetAttributeDefaultBinding(SYS_PARAM_PARTICLES_POSITION);
diff --git a/Engine/Plugins/FX/Niagara/Source/Niagara/Private/NiagaraModule.cpp b/Engine/Plugins/FX/Niagara/Source/Niagara/Private/NiagaraModule.cpp
--- a/Engine/Plugins/FX/Niagara/Source/Niagara/Private/NiagaraModule.cpp
+++ b/Engine/Plugins/FX/Niagara/Source/Niagara/Private/NiagaraModule.cpp
@@ -25,6 +25,7 @@
 #include "NiagaraGpuComputeDispatch.h"
 #include "Misc/CoreDelegates.h"
 #include "NiagaraDebuggerClient.h"
+#include "NiagaraSimulationStageBase.h"
 #include "Particles/FXBudget.h"
 #include "Engine/StaticMesh.h"
 #include "UObject/UObjectGlobals.h"
@@ -341,12 +342,14 @@ void INiagaraModule::StartupModule()
 	Translator_CallID = FNiagaraVariable(FNiagaraTypeDefinition::GetIntDef(), TEXT("Translator.CallID"));
 
 	FNiagaraConstants::Init();
-	UNiagaraLightRendererProperties::InitCDOPropertiesAfterModuleStartup();
-	UNiagaraSpriteRendererProperties::InitCDOPropertiesAfterModuleStartup();
-	UNiagaraRibbonRendererProperties::InitCDOPropertiesAfterModuleStartup();
-	UNiagaraMeshRendererProperties::InitCDOPropertiesAfterModuleStartup();
-	UNiagaraComponentRendererProperties::InitCDOPropertiesAfterModuleStartup();
+	UNiagaraRendererProperties::InitCDOPropertiesAfterModuleStartup(UNiagaraLightRendererProperties::StaticClass());
+	UNiagaraRendererProperties::InitCDOPropertiesAfterModuleStartup(UNiagaraSpriteRendererProperties::StaticClass());
+	UNiagaraRendererProperties::InitCDOPropertiesAfterModuleStartup(UNiagaraRibbonRendererProperties::StaticClass());
+	UNiagaraRendererProperties::InitCDOPropertiesAfterModuleStartup(UNiagaraMeshRendererProperties::StaticClass());
+	UNiagaraRendererProperties::InitCDOPropertiesAfterModuleStartup(UNiagaraComponentRendererProperties::StaticClass());
 
+	UNiagaraSimulationStageGeneric::InitCDOPropertiesAfterModuleStartup();
+	
 	// Register the data interface CDO finder with the shader module..
 	INiagaraShaderModule& NiagaraShaderModule = FModuleManager::LoadModuleChecked<INiagaraShaderModule>("NiagaraShader");
 	NiagaraShaderModule.SetOnRequestDefaultDataInterfaceHandler(INiagaraShaderModule::FOnRequestDefaultDataInterface::CreateLambda([](const FString& DIClassName) -> UNiagaraDataInterfaceBase*
diff --git a/Engine/Plugins/FX/Niagara/Source/Niagara/Private/NiagaraRendererProperties.cpp b/Engine/Plugins/FX/Niagara/Source/Niagara/Private/NiagaraRendererProperties.cpp
--- a/Engine/Plugins/FX/Niagara/Source/Niagara/Private/NiagaraRendererProperties.cpp
+++ b/Engine/Plugins/FX/Niagara/Source/Niagara/Private/NiagaraRendererProperties.cpp
@@ -492,10 +492,28 @@ void UNiagaraRendererProperties::PostLoad()
 	{
 		MotionVectorSetting = ENiagaraRendererMotionVectorSetting::Disable;
 	}
+
+	// it was possible to revert the enabled binding's type to "none" when the revert arrow is used
+	if(!RendererEnabledBinding.GetType().IsValid())
+	{
+		FNiagaraVariableBase EnabledDefaultVariable(FNiagaraTypeDefinition::GetBoolDef(), NAME_None);
+		RendererEnabledBinding.Setup(EnabledDefaultVariable, EnabledDefaultVariable, ENiagaraRendererSourceDataMode::Emitter);
+	}
 }
 
-#if WITH_EDITORONLY_DATA
+void UNiagaraRendererProperties::InitBindings()
+{
+	FNiagaraVariableBase EnabledDefaultVariable(FNiagaraTypeDefinition::GetBoolDef(), NAME_None);
+	RendererEnabledBinding.Setup(EnabledDefaultVariable, EnabledDefaultVariable, ENiagaraRendererSourceDataMode::Emitter);
+}
+
+void UNiagaraRendererProperties::InitCDOPropertiesAfterModuleStartup(TSubclassOf<UNiagaraRendererProperties> Class)
+{
+	UNiagaraRendererProperties* CDO = CastChecked<UNiagaraRendererProperties>(Class->GetDefaultObject());
+	CDO->InitCDOPropertiesAfterModuleStartupInternal();
+}
 
+#if WITH_EDITORONLY_DATA
 void UNiagaraRendererProperties::PostEditChangeProperty(struct FPropertyChangedEvent& PropertyChangedEvent)
 {
 	Super::PostEditChangeProperty(PropertyChangedEvent);
diff --git a/Engine/Plugins/FX/Niagara/Source/Niagara/Private/NiagaraRibbonRendererProperties.cpp b/Engine/Plugins/FX/Niagara/Source/Niagara/Private/NiagaraRibbonRendererProperties.cpp
--- a/Engine/Plugins/FX/Niagara/Source/Niagara/Private/NiagaraRibbonRendererProperties.cpp
+++ b/Engine/Plugins/FX/Niagara/Source/Niagara/Private/NiagaraRibbonRendererProperties.cpp
@@ -227,11 +227,10 @@ void UNiagaraRibbonRendererProperties::PostInitProperties()
 }
 
 /** The bindings depend on variables that are created during the NiagaraModule startup. However, the CDO's are build prior to this being initialized, so we defer setting these values until later.*/
-void UNiagaraRibbonRendererProperties::InitCDOPropertiesAfterModuleStartup()
+void UNiagaraRibbonRendererProperties::InitCDOPropertiesAfterModuleStartupInternal()
 {
-	UNiagaraRibbonRendererProperties* CDO = CastChecked<UNiagaraRibbonRendererProperties>(UNiagaraRibbonRendererProperties::StaticClass()->GetDefaultObject());
-	CDO->InitBindings();
-
+	Super::InitCDOPropertiesAfterModuleStartupInternal();
+	
 	for (TWeakObjectPtr<UNiagaraRibbonRendererProperties>& WeakRibbonRendererProperties : RibbonRendererPropertiesToDeferredInit)
 	{
 		if (WeakRibbonRendererProperties.Get())
@@ -243,6 +242,8 @@ void UNiagaraRibbonRendererProperties::InitCDOPropertiesAfterModuleStartup()
 
 void UNiagaraRibbonRendererProperties::InitBindings()
 {
+	Super::InitBindings();
+
 	if (!PositionBinding.IsValid())
 	{
 		PositionBinding = FNiagaraConstants::GetAttributeDefaultBinding(SYS_PARAM_PARTICLES_POSITION);
diff --git a/Engine/Plugins/FX/Niagara/Source/Niagara/Private/NiagaraSimulationStageBase.cpp b/Engine/Plugins/FX/Niagara/Source/Niagara/Private/NiagaraSimulationStageBase.cpp
--- a/Engine/Plugins/FX/Niagara/Source/Niagara/Private/NiagaraSimulationStageBase.cpp
+++ b/Engine/Plugins/FX/Niagara/Source/Niagara/Private/NiagaraSimulationStageBase.cpp
@@ -121,17 +121,29 @@ bool UNiagaraSimulationStageGeneric::FillCompilationData(TArray<FNiagaraSimulati
 }
 #endif //WITH_EDITORONLY_DATA
 
+void UNiagaraSimulationStageGeneric::InitBindings()
+{
+	using namespace NiagaraSimulationStageLocal;
+	EnabledBinding.Setup(GetDefaultEnabledBinding(), GetDefaultEnabledBinding(), ENiagaraRendererSourceDataMode::Emitter);
+	ElementCountBinding.Setup(GetDefaultElementCountBinding(), GetDefaultElementCountBinding(), ENiagaraRendererSourceDataMode::Emitter);
+	NumIterationsBinding.Setup(GetDefaultNumIterationsBinding(), GetDefaultNumIterationsBinding(), ENiagaraRendererSourceDataMode::Emitter);
+	ParticleIterationStateBinding.Setup(GetDefaultParticleStateBinding(), GetDefaultParticleStateBinding(), ENiagaraRendererSourceDataMode::Particles);
+}
+
+void UNiagaraSimulationStageGeneric::InitCDOPropertiesAfterModuleStartup()
+{
+	UNiagaraSimulationStageGeneric* CDO = GetMutableDefault<UNiagaraSimulationStageGeneric>();
+	CDO->InitBindings();
+}
+
 void UNiagaraSimulationStageGeneric::PostInitProperties()
 {
 	Super::PostInitProperties();
 
-	if (HasAnyFlags(RF_ClassDefaultObject) == false)
+	// the type registry isn't live at this point for the CDO so we defer the bindings initialization for it
+	if(!HasAnyFlags(RF_ClassDefaultObject))
 	{
-		using namespace NiagaraSimulationStageLocal;
-		EnabledBinding.Setup(GetDefaultEnabledBinding(), GetDefaultEnabledBinding(), ENiagaraRendererSourceDataMode::Emitter);
-		ElementCountBinding.Setup(GetDefaultElementCountBinding(), GetDefaultElementCountBinding(), ENiagaraRendererSourceDataMode::Emitter);
-		NumIterationsBinding.Setup(GetDefaultNumIterationsBinding(), GetDefaultNumIterationsBinding(), ENiagaraRendererSourceDataMode::Emitter);
-		ParticleIterationStateBinding.Setup(GetDefaultParticleStateBinding(), GetDefaultParticleStateBinding(), ENiagaraRendererSourceDataMode::Particles);
+		InitBindings();
 	}
 }
 
diff --git a/Engine/Plugins/FX/Niagara/Source/Niagara/Private/NiagaraSpriteRendererProperties.cpp b/Engine/Plugins/FX/Niagara/Source/Niagara/Private/NiagaraSpriteRendererProperties.cpp
--- a/Engine/Plugins/FX/Niagara/Source/Niagara/Private/NiagaraSpriteRendererProperties.cpp
+++ b/Engine/Plugins/FX/Niagara/Source/Niagara/Private/NiagaraSpriteRendererProperties.cpp
@@ -200,11 +200,10 @@ void UNiagaraSpriteRendererProperties::Serialize(FStructuredArchive::FRecord Rec
 }
 
 /** The bindings depend on variables that are created during the NiagaraModule startup. However, the CDO's are build prior to this being initialized, so we defer setting these values until later.*/
-void UNiagaraSpriteRendererProperties::InitCDOPropertiesAfterModuleStartup()
+void UNiagaraSpriteRendererProperties::InitCDOPropertiesAfterModuleStartupInternal()
 {
-	UNiagaraSpriteRendererProperties* CDO = CastChecked<UNiagaraSpriteRendererProperties>(UNiagaraSpriteRendererProperties::StaticClass()->GetDefaultObject());
-	CDO->InitBindings();
-
+	Super::InitCDOPropertiesAfterModuleStartupInternal();
+	
 	for (TWeakObjectPtr<UNiagaraSpriteRendererProperties>& WeakSpriteRendererProperties : SpriteRendererPropertiesToDeferredInit)
 	{
 		if (WeakSpriteRendererProperties.Get())
@@ -216,6 +215,8 @@ void UNiagaraSpriteRendererProperties::InitCDOPropertiesAfterModuleStartup()
 
 void UNiagaraSpriteRendererProperties::InitBindings()
 {
+	Super::InitBindings();
+	
 	if (PositionBinding.GetParamMapBindableVariable().GetName() == NAME_None)
 	{
 		PositionBinding = FNiagaraConstants::GetAttributeDefaultBinding(SYS_PARAM_PARTICLES_POSITION);
diff --git a/Engine/Plugins/FX/Niagara/Source/Niagara/Public/NiagaraComponentRendererProperties.h b/Engine/Plugins/FX/Niagara/Source/Niagara/Public/NiagaraComponentRendererProperties.h
--- a/Engine/Plugins/FX/Niagara/Source/Niagara/Public/NiagaraComponentRendererProperties.h
+++ b/Engine/Plugins/FX/Niagara/Source/Niagara/Public/NiagaraComponentRendererProperties.h
@@ -59,9 +59,7 @@ public:
 	virtual void PostInitProperties() override;
 	virtual void PostDuplicate(bool bDuplicateForPIE) override;
 	//UObject Interface END
-
-	static void InitCDOPropertiesAfterModuleStartup();
-
+	
 	//~ UNiagaraRendererProperties interface
 	virtual FNiagaraRenderer* CreateEmitterRenderer(ERHIFeatureLevel::Type FeatureLevel, const FNiagaraEmitterInstance* Emitter, const FNiagaraSystemInstanceController& InController) override;
 	virtual class FNiagaraBoundsCalculator* CreateBoundsCalculator() override { return nullptr; }
@@ -160,6 +158,8 @@ public:
 	virtual bool NeedsSystemCompletion() const override { return true; }
 
 protected:
+	virtual void InitBindings() override;
+	virtual void InitCDOPropertiesAfterModuleStartupInternal() override;
 	virtual void UpdateSourceModeDerivates(ENiagaraRendererSourceDataMode InSourceMode, bool bFromPropertyEdit = false) override;
 
 private:
diff --git a/Engine/Plugins/FX/Niagara/Source/Niagara/Public/NiagaraMeshRendererProperties.h b/Engine/Plugins/FX/Niagara/Source/Niagara/Public/NiagaraMeshRendererProperties.h
--- a/Engine/Plugins/FX/Niagara/Source/Niagara/Public/NiagaraMeshRendererProperties.h
+++ b/Engine/Plugins/FX/Niagara/Source/Niagara/Public/NiagaraMeshRendererProperties.h
@@ -170,8 +170,6 @@ public:
 #endif// WITH_EDITORONLY_DATA
 	//UObject Interface END
 
-	static void InitCDOPropertiesAfterModuleStartup();
-
 	//~ UNiagaraRendererProperties interface
 	virtual FNiagaraRenderer* CreateEmitterRenderer(ERHIFeatureLevel::Type FeatureLevel, const FNiagaraEmitterInstance* Emitter, const FNiagaraSystemInstanceController& InController) override;
 	virtual class FNiagaraBoundsCalculator* CreateBoundsCalculator() override;
@@ -411,8 +409,9 @@ public:
 	FNiagaraRendererLayout RendererLayoutWithoutCustomSorting;
 
 protected:
-	void InitBindings();
+	virtual void InitBindings() override;
 	void SetPreviousBindings(const FVersionedNiagaraEmitter& SrcEmitter, ENiagaraRendererSourceDataMode InSourceMode);
+	virtual void InitCDOPropertiesAfterModuleStartupInternal() override;
 	virtual void UpdateSourceModeDerivates(ENiagaraRendererSourceDataMode InSourceMode, bool bFromPropertyEdit = false) override;
 	virtual bool NeedsMIDsForMaterials() const override { return MaterialParameters.HasAnyBindings(); }
 #if WITH_EDITORONLY_DATA
diff --git a/Engine/Plugins/FX/Niagara/Source/Niagara/Public/NiagaraRendererProperties.h b/Engine/Plugins/FX/Niagara/Source/Niagara/Public/NiagaraRendererProperties.h
--- a/Engine/Plugins/FX/Niagara/Source/Niagara/Public/NiagaraRendererProperties.h
+++ b/Engine/Plugins/FX/Niagara/Source/Niagara/Public/NiagaraRendererProperties.h
@@ -249,11 +249,15 @@ public:
 		, bAllowInCullProxies(true)
 		, bMotionBlurEnabled_DEPRECATED(true)
 	{
+		AttributeBindings.Add(&RendererEnabledBinding);
 	}
 
 	//UObject Interface Begin
 	virtual void PostInitProperties() override;
 	virtual void PostLoad() override;
+	virtual void InitBindings();
+
+	static void InitCDOPropertiesAfterModuleStartup(TSubclassOf<UNiagaraRendererProperties> Class);
 #if WITH_EDITORONLY_DATA
 	virtual void PostEditChangeProperty(struct FPropertyChangedEvent& PropertyChangedEvent) override;
 #endif
@@ -377,7 +381,7 @@ protected:
 
 	virtual void PostLoadBindings(ENiagaraRendererSourceDataMode InSourceMode);
 	virtual void UpdateSourceModeDerivates(ENiagaraRendererSourceDataMode InSourceMode, bool bFromPropertyEdit = false);
-
+	virtual void InitCDOPropertiesAfterModuleStartupInternal() { check(HasAnyFlags(RF_ClassDefaultObject)); InitBindings(); }
 #if WITH_EDITORONLY_DATA
 	/** returns the variable associated with the supplied binding if it should be bound given the current settings of the RendererProperties. */
 	virtual FNiagaraVariable GetBoundAttribute(const FNiagaraVariableAttributeBinding* Binding) const;
diff --git a/Engine/Plugins/FX/Niagara/Source/Niagara/Public/NiagaraRibbonRendererProperties.h b/Engine/Plugins/FX/Niagara/Source/Niagara/Public/NiagaraRibbonRendererProperties.h
--- a/Engine/Plugins/FX/Niagara/Source/Niagara/Public/NiagaraRibbonRendererProperties.h
+++ b/Engine/Plugins/FX/Niagara/Source/Niagara/Public/NiagaraRibbonRendererProperties.h
@@ -207,8 +207,6 @@ public:
 #endif
 	//UObject Interface END
 
-	static void InitCDOPropertiesAfterModuleStartup();
-
 	//UNiagaraRendererProperties Interface
 	virtual FNiagaraRenderer* CreateEmitterRenderer(ERHIFeatureLevel::Type FeatureLevel, const FNiagaraEmitterInstance* Emitter, const FNiagaraSystemInstanceController& InController) override;
 	virtual class FNiagaraBoundsCalculator* CreateBoundsCalculator() override;
@@ -481,9 +479,10 @@ public:
 	FNiagaraRendererLayout RendererLayout;
 
 protected:
-	void InitBindings();
+	virtual void InitBindings() override;
 	void SetPreviousBindings(const FVersionedNiagaraEmitter& SrcEmitter);
 
+	virtual void InitCDOPropertiesAfterModuleStartupInternal() override;
 	void UpdateSourceModeDerivates(ENiagaraRendererSourceDataMode InSourceMode, bool bFromPropertyEdit);
 
 	virtual bool NeedsMIDsForMaterials() const { return MaterialParameters.HasAnyBindings(); }
diff --git a/Engine/Plugins/FX/Niagara/Source/Niagara/Public/NiagaraSimulationStageBase.h b/Engine/Plugins/FX/Niagara/Source/Niagara/Public/NiagaraSimulationStageBase.h
--- a/Engine/Plugins/FX/Niagara/Source/Niagara/Public/NiagaraSimulationStageBase.h
+++ b/Engine/Plugins/FX/Niagara/Source/Niagara/Public/NiagaraSimulationStageBase.h
@@ -139,6 +139,9 @@ public:
 	UPROPERTY(EditAnywhere, Category = "Simulation Stage")
 	FIntVector OverrideGpuDispatchNumThreads = FIntVector(64, 1, 1);
 
+	void InitBindings();
+	static void InitCDOPropertiesAfterModuleStartup();
+
 	virtual void PostInitProperties() override;
 	virtual void PostLoad() override;
 	virtual bool AppendCompileHash(FNiagaraCompileHashVisitor* InVisitor) const override;
diff --git a/Engine/Plugins/FX/Niagara/Source/Niagara/Public/NiagaraSpriteRendererProperties.h b/Engine/Plugins/FX/Niagara/Source/Niagara/Public/NiagaraSpriteRendererProperties.h
--- a/Engine/Plugins/FX/Niagara/Source/Niagara/Public/NiagaraSpriteRendererProperties.h
+++ b/Engine/Plugins/FX/Niagara/Source/Niagara/Public/NiagaraSpriteRendererProperties.h
@@ -116,9 +116,7 @@ public:
 
 #endif // WITH_EDITORONLY_DATA
 	//UObject Interface END
-
-	static void InitCDOPropertiesAfterModuleStartup();
-
+	
 	//UNiagaraRendererProperties interface
 	virtual FNiagaraRenderer* CreateEmitterRenderer(ERHIFeatureLevel::Type FeatureLevel, const FNiagaraEmitterInstance* Emitter, const FNiagaraSystemInstanceController& InController) override;
 	virtual class FNiagaraBoundsCalculator* CreateBoundsCalculator() override;
@@ -391,8 +389,9 @@ public:
 	uint32 MaterialParamValidMask = 0;
 
 protected:
-	void InitBindings();
+	virtual void InitBindings() override;
 	void SetPreviousBindings(const FVersionedNiagaraEmitter& SrcEmitter, ENiagaraRendererSourceDataMode InSourceMode);
+	virtual void InitCDOPropertiesAfterModuleStartupInternal() override;
 	virtual void UpdateSourceModeDerivates(ENiagaraRendererSourceDataMode InSourceMode, bool bFromPropertyEdit = false) override;
 
 #if WITH_EDITORONLY_DATA
diff --git a/Engine/Plugins/FX/Niagara/Source/NiagaraEditor/Private/Customizations/NiagaraTypeCustomizations.cpp b/Engine/Plugins/FX/Niagara/Source/NiagaraEditor/Private/Customizations/NiagaraTypeCustomizations.cpp
--- a/Engine/Plugins/FX/Niagara/Source/NiagaraEditor/Private/Customizations/NiagaraTypeCustomizations.cpp
+++ b/Engine/Plugins/FX/Niagara/Source/NiagaraEditor/Private/Customizations/NiagaraTypeCustomizations.cpp
@@ -368,12 +368,7 @@ void FNiagaraVariableAttributeBindingCustomization::ChangeSource(FName InVarName
 	PropertyHandle->NotifyFinishedChangingProperties();
 }
 
-void FNiagaraVariableAttributeBindingCustomization::ResetToDefault()
-{
-	UE_LOG(LogNiagaraEditor, Warning, TEXT("Reset to default!"));
-}
-
-EVisibility FNiagaraVariableAttributeBindingCustomization::IsResetToDefaultsVisible() const
+bool FNiagaraVariableAttributeBindingCustomization::IsResetToDefaultsVisible(TSharedPtr<IPropertyHandle> InPropertyHandle) const
 {
 	check(BaseEmitter.Emitter);
 	check(RenderProps || SimulationStage);
@@ -381,10 +376,10 @@ EVisibility FNiagaraVariableAttributeBindingCustomization::IsResetToDefaultsVisi
 	check(DefaultVariableBinding);
 
 	const ENiagaraRendererSourceDataMode BindingSourceMode = RenderProps ? RenderProps->GetCurrentSourceMode() : ENiagaraRendererSourceDataMode::Emitter;
-	return TargetVariableBinding->MatchesDefault(*DefaultVariableBinding, BindingSourceMode) ? EVisibility::Hidden : EVisibility::Visible;
+	return !TargetVariableBinding->MatchesDefault(*DefaultVariableBinding, BindingSourceMode);
 }
 
-FReply FNiagaraVariableAttributeBindingCustomization::OnResetToDefaultsClicked()
+void FNiagaraVariableAttributeBindingCustomization::OnResetToDefaultsClicked(TSharedPtr<IPropertyHandle> InPropertyHandle)
 {
 	FScopedTransaction Transaction(LOCTEXT("ResetBindingParam", "Reset binding"));
 	TArray<UObject*> Objects;
@@ -403,7 +398,6 @@ FReply FNiagaraVariableAttributeBindingCustomization::OnResetToDefaultsClicked()
 	TargetVariableBinding->ResetToDefault(*DefaultVariableBinding, BaseEmitter, BindingSourceMode);
 	PropertyHandle->NotifyPostChange(EPropertyChangeType::ValueSet);
 	PropertyHandle->NotifyFinishedChangingProperties();
-	return FReply::Handled();
 }
 
 void FNiagaraVariableAttributeBindingCustomization::CustomizeHeader(TSharedRef<IPropertyHandle> InPropertyHandle, FDetailWidgetRow& HeaderRow, IPropertyTypeCustomizationUtils& CustomizationUtils)
@@ -415,17 +409,6 @@ void FNiagaraVariableAttributeBindingCustomization::CustomizeHeader(TSharedRef<I
 	TArray<UObject*> Objects;
 	PropertyHandle->GetOuterObjects(Objects);
 	bool bAddDefault = true;
-
-	InPropertyHandle->SetOnPropertyResetToDefault(FSimpleDelegate::CreateLambda([this]() { ResetToDefault(); }));
-	//InPropertyHandle->ExecuteCustomResetToDefault
-
-	/*FResetToDefaultOverride ResetOverride = FResetToDefaultOverride::Create(
-		FIsResetToDefaultVisible::CreateSP(this, &FMotionControllerDetails::IsSourceValueModified),
-		FResetToDefaultHandler::CreateSP(this, &FMotionControllerDetails::OnResetSourceValue)
-	);
-
-	PropertyRow.OverrideResetToDefault(ResetOverride); */
-	InPropertyHandle->MarkResetToDefaultCustomized(true);
 	
 	if (Objects.Num() == 1)
 	{
@@ -445,7 +428,15 @@ void FNiagaraVariableAttributeBindingCustomization::CustomizeHeader(TSharedRef<I
 		{
 			TargetVariableBinding = (FNiagaraVariableAttributeBinding*)PropertyHandle->GetValueBaseAddress((uint8*)Objects[0]);
 			DefaultVariableBinding = (FNiagaraVariableAttributeBinding*)PropertyHandle->GetValueBaseAddress((uint8*)Objects[0]->GetClass()->GetDefaultObject());
-				
+
+			FResetToDefaultOverride OverrideReset = FResetToDefaultOverride::Create
+			(
+				FIsResetToDefaultVisible::CreateSP(this, &FNiagaraVariableAttributeBindingCustomization::IsResetToDefaultsVisible),
+				FResetToDefaultHandler::CreateSP(this, &FNiagaraVariableAttributeBindingCustomization::OnResetToDefaultsClicked)
+			);
+			
+			HeaderRow.OverrideResetToDefault(OverrideReset);
+			
 			HeaderRow
 				.NameContent()
 				[
@@ -463,34 +454,10 @@ void FNiagaraVariableAttributeBindingCustomization::CustomizeHeader(TSharedRef<I
 						.ContentPadding(1)
 						.ToolTipText(this, &FNiagaraVariableAttributeBindingCustomization::GetTooltipText)
 						.ButtonContent()
-						[
-							/*SNew(STextBlock)
-							.Text(this, &FNiagaraVariableAttributeBindingCustomization::GetCurrentText)
-							.Font(IDetailLayoutBuilder::GetDetailFont())*/
+						[							
 							SNew(SNiagaraParameterName)
 							.ParameterName(this, &FNiagaraVariableAttributeBindingCustomization::GetVariableName)
 							.IsReadOnly(true)
-							//SNew(STextBlock)
-							//.Text(InCreateData->Action->GetMenuDescription())
-						]
-					]
-					+ SHorizontalBox::Slot()
-					.AutoWidth()
-					.VAlign(VAlign_Center)
-					.Padding(2, 1)
-					[
-						SNew(SButton)
-						.IsFocusable(false)
-						.ToolTipText(LOCTEXT("ResetToDefaultToolTip", "Reset to Default"))
-						.ButtonStyle(FAppStyle::Get(), "NoBorder")
-						.ContentPadding(0)
-						.Visibility(this, &FNiagaraVariableAttributeBindingCustomization::IsResetToDefaultsVisible)
-						.OnClicked(this, &FNiagaraVariableAttributeBindingCustomization::OnResetToDefaultsClicked)
-						.Content()
-						[
-							SNew(SImage)
-							.Image(FAppStyle::GetBrush("PropertyWindow.DiffersFromDefault"))
-							.ColorAndOpacity(FSlateColor::UseForeground())
 						]
 					]
 				];
diff --git a/Engine/Plugins/FX/Niagara/Source/NiagaraEditor/Private/Customizations/NiagaraTypeCustomizations.h b/Engine/Plugins/FX/Niagara/Source/NiagaraEditor/Private/Customizations/NiagaraTypeCustomizations.h
--- a/Engine/Plugins/FX/Niagara/Source/NiagaraEditor/Private/Customizations/NiagaraTypeCustomizations.h
+++ b/Engine/Plugins/FX/Niagara/Source/NiagaraEditor/Private/Customizations/NiagaraTypeCustomizations.h
@@ -168,9 +168,8 @@ public:
 	virtual void CustomizeChildren(TSharedRef<IPropertyHandle> StructPropertyHandle, class IDetailChildrenBuilder& ChildBuilder, IPropertyTypeCustomizationUtils& StructCustomizationUtils) override {};
 	/** IPropertyTypeCustomization interface end */
 private:
-	EVisibility IsResetToDefaultsVisible() const;
-	FReply OnResetToDefaultsClicked();
-	void ResetToDefault();
+	bool IsResetToDefaultsVisible(TSharedPtr<IPropertyHandle> InPropertyHandle) const;
+	void OnResetToDefaultsClicked(TSharedPtr<IPropertyHandle> InPropertyHandle);
 	FName GetVariableName() const;
 	FText GetCurrentText() const;
 	FText GetTooltipText() const;
diff --git a/Engine/Plugins/FX/Niagara/Source/NiagaraEditorWidgets/Private/SNiagaraStack.cpp b/Engine/Plugins/FX/Niagara/Source/NiagaraEditorWidgets/Private/SNiagaraStack.cpp
--- a/Engine/Plugins/FX/Niagara/Source/NiagaraEditorWidgets/Private/SNiagaraStack.cpp
+++ b/Engine/Plugins/FX/Niagara/Source/NiagaraEditorWidgets/Private/SNiagaraStack.cpp
@@ -27,6 +27,7 @@
 #include "Framework/Commands/GenericCommands.h"
 #include "IContentBrowserSingleton.h"
 #include "ContentBrowserModule.h"
+#include "DetailWidgetRow.h"
 #include "Widgets/Input/SButton.h"
 #include "Widgets/Images/SImage.h"
 #include "Widgets/Input/SComboButton.h"
@@ -955,78 +956,90 @@ SNiagaraStack::FRowWidgets SNiagaraStack::ConstructNameAndValueWidgetsForItem(UN
 		UNiagaraStackPropertyRow* PropertyRow = CastChecked<UNiagaraStackPropertyRow>(Item);
 		FNodeWidgets PropertyRowWidgets = PropertyRow->GetDetailTreeNode()->CreateNodeWidgets();
 
-		TAttribute<bool> IsEnabled;
-		IsEnabled.BindUObject(Item, &UNiagaraStackEntry::GetIsEnabledAndOwnerIsEnabled);
+		TSharedPtr<SWidget> TmpNameWidget;
+		TSharedPtr<SWidget> TmpValueWidget;
+		FDetailWidgetRow WidgetRow;
 
-		Container->AddFillRowContextMenuHandler(FNiagaraStackPropertyRowUtilities::CreateOnFillRowContextMenu(PropertyRow->GetDetailTreeNode()->CreatePropertyHandle(), PropertyRowWidgets.Actions));
-
-		if (PropertyRowWidgets.WholeRowWidget.IsValid())
+		TOptional<FResetToDefaultOverride> ResetToDefaultOverride;
+		if(TSharedPtr<IDetailPropertyRow> Row = PropertyRow->GetDetailTreeNode()->GetRow())
 		{
-			Container->SetOverrideNameWidth(PropertyRowWidgets.WholeRowWidgetLayoutData.MinWidth, PropertyRowWidgets.WholeRowWidgetLayoutData.MaxWidth);
-			Container->SetOverrideNameAlignment(PropertyRowWidgets.WholeRowWidgetLayoutData.HorizontalAlignment, PropertyRowWidgets.WholeRowWidgetLayoutData.VerticalAlignment);
-			PropertyRowWidgets.WholeRowWidget->SetEnabled(IsEnabled);
-			TSharedRef<SHorizontalBox> RowBox = SNew(SHorizontalBox);
-			if (PropertyRowWidgets.EditConditionWidget.IsValid())
-			{
-				RowBox->AddSlot()
-				.AutoWidth()
-				.Padding(0, 0, 3, 0)
-				[
-					PropertyRowWidgets.EditConditionWidget.ToSharedRef()
-				];
-			}
-			RowBox->AddSlot()
-			[
-				PropertyRowWidgets.WholeRowWidget.ToSharedRef()
-			];
-			RowBox->AddSlot()
-			.AutoWidth()
-			.Padding(3, 0, 0, 0)
-			[
-				SNew(SResetToDefaultPropertyEditor, PropertyRow->GetDetailTreeNode()->CreatePropertyHandle())
-			];
-			return FRowWidgets(RowBox); 
+			// we override the name & value widgets here as the previous ones won't be valid anymore but we require the widget row to access additional data like the custom reset override.
+			PropertyRow->GetDetailTreeNode()->GetRow()->GetDefaultWidgets(PropertyRowWidgets.NameWidget, PropertyRowWidgets.ValueWidget, WidgetRow);
+			ResetToDefaultOverride = WidgetRow.CustomResetToDefault;
 		}
-		else
-		{
-			Container->SetOverrideNameWidth(PropertyRowWidgets.NameWidgetLayoutData.MinWidth, PropertyRowWidgets.NameWidgetLayoutData.MaxWidth);
-			Container->SetOverrideValueWidth(PropertyRowWidgets.ValueWidgetLayoutData.MinWidth, PropertyRowWidgets.ValueWidgetLayoutData.MaxWidth);
-			Container->SetOverrideNameAlignment(HAlign_Left, VAlign_Center);
-			PropertyRowWidgets.NameWidget->SetEnabled(IsEnabled);
-			PropertyRowWidgets.ValueWidget->SetEnabled(IsEnabled);
-
-			TSharedPtr<SWidget> NameWidget;
-			if (PropertyRowWidgets.EditConditionWidget.IsValid())
-			{
-				NameWidget = SNew(SHorizontalBox)
-				+ SHorizontalBox::Slot()
-				.AutoWidth()
-				.Padding(0, 0, 3, 0)
-				[
-					PropertyRowWidgets.EditConditionWidget.ToSharedRef()
-				]
-				+ SHorizontalBox::Slot()
-				[
-					PropertyRowWidgets.NameWidget.ToSharedRef()
-				];
-			}
-			else
-			{
-				NameWidget = PropertyRowWidgets.NameWidget;
-			}
 
-			TSharedRef<SWidget> ValueWidget = SNew(SHorizontalBox)
-			+ SHorizontalBox::Slot()
-			[
-				PropertyRowWidgets.ValueWidget.ToSharedRef()
-			]
-			+ SHorizontalBox::Slot()
-			.AutoWidth()
-			.VAlign(VAlign_Center)
-			.Padding(3, 0, 0, 0)
-			[
-				SNew(SResetToDefaultPropertyEditor, PropertyRow->GetDetailTreeNode()->CreatePropertyHandle())
-			];
+		TAttribute<bool> IsEnabled;
+		IsEnabled.BindUObject(Item, &UNiagaraStackEntry::GetIsEnabledAndOwnerIsEnabled);
+
+		Container->AddFillRowContextMenuHandler(FNiagaraStackPropertyRowUtilities::CreateOnFillRowContextMenu(PropertyRow->GetDetailTreeNode()->CreatePropertyHandle(), PropertyRowWidgets.Actions));
+		
+		 if (PropertyRowWidgets.WholeRowWidget.IsValid())
+		 {
+		 	Container->SetOverrideNameWidth(PropertyRowWidgets.WholeRowWidgetLayoutData.MinWidth, PropertyRowWidgets.WholeRowWidgetLayoutData.MaxWidth);
+		 	Container->SetOverrideNameAlignment(PropertyRowWidgets.WholeRowWidgetLayoutData.HorizontalAlignment, PropertyRowWidgets.WholeRowWidgetLayoutData.VerticalAlignment);
+		 	PropertyRowWidgets.WholeRowWidget->SetEnabled(IsEnabled);
+		 	TSharedRef<SHorizontalBox> RowBox = SNew(SHorizontalBox);
+		 	if (PropertyRowWidgets.EditConditionWidget.IsValid())
+		 	{
+		 		RowBox->AddSlot()
+		 		.AutoWidth()
+		 		.Padding(0, 0, 3, 0)
+		 		[
+		 			PropertyRowWidgets.EditConditionWidget.ToSharedRef()
+		 		];
+		 	}
+		 	RowBox->AddSlot()
+		 	[
+		 		PropertyRowWidgets.WholeRowWidget.ToSharedRef()
+		 	];
+		 	RowBox->AddSlot()
+		 	.AutoWidth()
+		 	.Padding(3, 0, 0, 0)
+		 	[
+		 		SNew(SResetToDefaultPropertyEditor, PropertyRow->GetDetailTreeNode()->CreatePropertyHandle()).CustomResetToDefault(ResetToDefaultOverride)
+		 	];
+		 	return FRowWidgets(RowBox); 
+		 }
+		 else
+		 {
+		 	Container->SetOverrideNameWidth(PropertyRowWidgets.NameWidgetLayoutData.MinWidth, PropertyRowWidgets.NameWidgetLayoutData.MaxWidth);
+		 	Container->SetOverrideValueWidth(PropertyRowWidgets.ValueWidgetLayoutData.MinWidth, PropertyRowWidgets.ValueWidgetLayoutData.MaxWidth);
+		 	Container->SetOverrideNameAlignment(HAlign_Left, VAlign_Center);
+		 	PropertyRowWidgets.NameWidget->SetEnabled(IsEnabled);
+		 	PropertyRowWidgets.ValueWidget->SetEnabled(IsEnabled);
+		
+		 	TSharedPtr<SWidget> NameWidget;
+		 	if (PropertyRowWidgets.EditConditionWidget.IsValid())
+		 	{
+		 		NameWidget = SNew(SHorizontalBox)
+		 		+ SHorizontalBox::Slot()
+		 		.AutoWidth()
+		 		.Padding(0, 0, 3, 0)
+		 		[
+		 			PropertyRowWidgets.EditConditionWidget.ToSharedRef()
+		 		]
+		 		+ SHorizontalBox::Slot()
+		 		[
+		 			PropertyRowWidgets.NameWidget.ToSharedRef()
+		 		];
+		 	}
+		 	else
+		 	{
+		 		NameWidget = PropertyRowWidgets.NameWidget;
+		 	}
+		
+		 	TSharedRef<SWidget> ValueWidget = SNew(SHorizontalBox)
+		 	+ SHorizontalBox::Slot()
+		 	[
+		 		PropertyRowWidgets.ValueWidget.ToSharedRef()
+		 	]
+		 	+ SHorizontalBox::Slot()
+		 	.AutoWidth()
+		 	.VAlign(VAlign_Center)
+		 	.Padding(3, 0, 0, 0)
+		 	[
+		 		SNew(SResetToDefaultPropertyEditor, PropertyRow->GetDetailTreeNode()->CreatePropertyHandle()).CustomResetToDefault(ResetToDefaultOverride)
+		 	];
 
 			return FRowWidgets(NameWidget.ToSharedRef(), ValueWidget);
 		}
diff --git a/Engine/Plugins/Runtime/GeometryCache/Source/GeometryCache/Classes/NiagaraGeometryCacheRendererProperties.h b/Engine/Plugins/Runtime/GeometryCache/Source/GeometryCache/Classes/NiagaraGeometryCacheRendererProperties.h
--- a/Engine/Plugins/Runtime/GeometryCache/Source/GeometryCache/Classes/NiagaraGeometryCacheRendererProperties.h
+++ b/Engine/Plugins/Runtime/GeometryCache/Source/GeometryCache/Classes/NiagaraGeometryCacheRendererProperties.h
@@ -44,8 +44,6 @@ public:
 	virtual void PostInitProperties() override;
 	//UObject Interface END
 
-	static void InitCDOPropertiesAfterModuleStartup();
-
 	//~ UNiagaraRendererProperties interface
 	virtual FNiagaraRenderer* CreateEmitterRenderer(ERHIFeatureLevel::Type FeatureLevel, const FNiagaraEmitterInstance* Emitter, const FNiagaraSystemInstanceController& InController) override;
 	virtual class FNiagaraBoundsCalculator* CreateBoundsCalculator() override { return nullptr; }
@@ -114,7 +112,8 @@ public:
 	virtual bool NeedsSystemCompletion() const override { return true; }
 
 protected:
-	void InitBindings();
+	virtual void InitBindings() override;
+	virtual	void InitCDOPropertiesAfterModuleStartupInternal() override;
 
 	static FNiagaraVariable Particles_GeoCacheRotation;
 	static FNiagaraVariable Particles_Age;
diff --git a/Engine/Plugins/Runtime/GeometryCache/Source/GeometryCache/Private/GeometryCacheModule.cpp b/Engine/Plugins/Runtime/GeometryCache/Source/GeometryCache/Private/GeometryCacheModule.cpp
--- a/Engine/Plugins/Runtime/GeometryCache/Source/GeometryCache/Private/GeometryCacheModule.cpp
+++ b/Engine/Plugins/Runtime/GeometryCache/Source/GeometryCache/Private/GeometryCacheModule.cpp
@@ -20,7 +20,7 @@ void FGeometryCacheModule::StartupModule()
 
 	FCodecV1Decoder::InitLUT();
 
-	UNiagaraGeometryCacheRendererProperties::InitCDOPropertiesAfterModuleStartup();
+	UNiagaraRendererProperties::InitCDOPropertiesAfterModuleStartup(UNiagaraGeometryCacheRendererProperties::StaticClass());
 }
 
 void FGeometryCacheModule::ShutdownModule()
diff --git a/Engine/Plugins/Runtime/GeometryCache/Source/GeometryCache/Private/NiagaraGeometryCacheRendererProperties.cpp b/Engine/Plugins/Runtime/GeometryCache/Source/GeometryCache/Private/NiagaraGeometryCacheRendererProperties.cpp
--- a/Engine/Plugins/Runtime/GeometryCache/Source/GeometryCache/Private/NiagaraGeometryCacheRendererProperties.cpp
+++ b/Engine/Plugins/Runtime/GeometryCache/Source/GeometryCache/Private/NiagaraGeometryCacheRendererProperties.cpp
@@ -63,6 +63,8 @@ FNiagaraVariableAttributeBinding CreateDefaultBinding(FNiagaraVariable DefaultVa
 
 void UNiagaraGeometryCacheRendererProperties::InitBindings()
 {
+	Super::InitBindings();
+
 	if (!PositionBinding.IsValid())
 	{
 		PositionBinding = FNiagaraConstants::GetAttributeDefaultBinding(SYS_PARAM_PARTICLES_POSITION);
@@ -142,18 +144,19 @@ void UNiagaraGeometryCacheRendererProperties::PostInitProperties()
 	}
 }
 
-void UNiagaraGeometryCacheRendererProperties::InitCDOPropertiesAfterModuleStartup()
+void UNiagaraGeometryCacheRendererProperties::InitCDOPropertiesAfterModuleStartupInternal()
 {
 	InitDefaultAttributes();
 	
-	UNiagaraGeometryCacheRendererProperties* CDO = CastChecked<UNiagaraGeometryCacheRendererProperties>(StaticClass()->GetDefaultObject());
-	CDO->PositionBinding = FNiagaraConstants::GetAttributeDefaultBinding(SYS_PARAM_PARTICLES_POSITION);
-	CDO->ScaleBinding = FNiagaraConstants::GetAttributeDefaultBinding(SYS_PARAM_PARTICLES_SCALE);
-	CDO->RendererVisibilityTagBinding = FNiagaraConstants::GetAttributeDefaultBinding(SYS_PARAM_PARTICLES_VISIBILITY_TAG);
-	CDO->RotationBinding = CreateDefaultBinding(Particles_GeoCacheRotation, FVector3f(0, 0, 0));
-	CDO->ElapsedTimeBinding = CreateDefaultBinding(Particles_Age, 0.0f);
-	CDO->EnabledBinding = CreateDefaultBinding(Particles_GeoCacheIsEnabled, false);
-	CDO->ArrayIndexBinding = CreateDefaultBinding(Particles_ArrayIndex, 0);
+	Super::InitCDOPropertiesAfterModuleStartupInternal();
+	
+	PositionBinding = FNiagaraConstants::GetAttributeDefaultBinding(SYS_PARAM_PARTICLES_POSITION);
+	ScaleBinding = FNiagaraConstants::GetAttributeDefaultBinding(SYS_PARAM_PARTICLES_SCALE);
+	RendererVisibilityTagBinding = FNiagaraConstants::GetAttributeDefaultBinding(SYS_PARAM_PARTICLES_VISIBILITY_TAG);
+	RotationBinding = CreateDefaultBinding(Particles_GeoCacheRotation, FVector3f(0, 0, 0));
+	ElapsedTimeBinding = CreateDefaultBinding(Particles_Age, 0.0f);
+	EnabledBinding = CreateDefaultBinding(Particles_GeoCacheIsEnabled, false);
+	ArrayIndexBinding = CreateDefaultBinding(Particles_ArrayIndex, 0);
 
 	for (TWeakObjectPtr<UNiagaraGeometryCacheRendererProperties>& WeakGeometryCacheRendererProperties : GeometryCacheRendererPropertiesToDeferredInit)
 	{
