[Backout] - CL21404928
Original CL Desc
-----------------------------------------------------------------
Fixed multiple bindings not resetting properly, which was due to CDO properties not being set
This still leaves us with visible "reset to default" buttons due to the property node having the previous default value cached.

To get rid of that:
Added option to add custom reset behavior overrides for any properrty, currently used for "Renderer Enabled" and "Simulation Enabled" bindings, which allows us to specify custom "is visible" overrides & functionality

Removed custom reset button for attribute bindings as it should use the default one for the stack

diff --git a/Engine/Plugins/FX/Niagara/Source/Niagara/Private/NiagaraComponentRendererProperties.cpp b/Engine/Plugins/FX/Niagara/Source/Niagara/Private/NiagaraComponentRendererProperties.cpp
--- a/Engine/Plugins/FX/Niagara/Source/Niagara/Private/NiagaraComponentRendererProperties.cpp
+++ b/Engine/Plugins/FX/Niagara/Source/Niagara/Private/NiagaraComponentRendererProperties.cpp
@@ -451,17 +451,11 @@ void UNiagaraComponentRendererProperties::PostDuplicate(bool bDuplicateForPIE)
 	TemplateComponent = DuplicateObject(TemplateComponent, this);
 }
 
-void UNiagaraComponentRendererProperties::InitBindings()
+void UNiagaraComponentRendererProperties::InitCDOPropertiesAfterModuleStartup()
 {
-	Super::InitBindings();
-
-	EnabledBinding = FNiagaraConstants::GetAttributeDefaultBinding(SYS_PARAM_PARTICLES_COMPONENTS_ENABLED);
-	RendererVisibilityTagBinding = FNiagaraConstants::GetAttributeDefaultBinding(SYS_PARAM_PARTICLES_VISIBILITY_TAG);
-}
-
-void UNiagaraComponentRendererProperties::InitCDOPropertiesAfterModuleStartupInternal()
-{
-	Super::InitCDOPropertiesAfterModuleStartupInternal();
+	UNiagaraComponentRendererProperties* CDO = CastChecked<UNiagaraComponentRendererProperties>(UNiagaraComponentRendererProperties::StaticClass()->GetDefaultObject());
+	CDO->EnabledBinding = FNiagaraConstants::GetAttributeDefaultBinding(SYS_PARAM_PARTICLES_COMPONENTS_ENABLED);
+	CDO->RendererVisibilityTagBinding = FNiagaraConstants::GetAttributeDefaultBinding(SYS_PARAM_PARTICLES_VISIBILITY_TAG);
 
 	for (TWeakObjectPtr<UNiagaraComponentRendererProperties>& WeakComponentRendererProperties : ComponentRendererPropertiesToDeferredInit)
 	{
diff --git a/Engine/Plugins/FX/Niagara/Source/Niagara/Private/NiagaraMeshRendererProperties.cpp b/Engine/Plugins/FX/Niagara/Source/Niagara/Private/NiagaraMeshRendererProperties.cpp
--- a/Engine/Plugins/FX/Niagara/Source/Niagara/Private/NiagaraMeshRendererProperties.cpp
+++ b/Engine/Plugins/FX/Niagara/Source/Niagara/Private/NiagaraMeshRendererProperties.cpp
@@ -244,10 +244,11 @@ void UNiagaraMeshRendererProperties::Serialize(FArchive& Ar)
 }
 
 /** The bindings depend on variables that are created during the NiagaraModule startup. However, the CDO's are build prior to this being initialized, so we defer setting these values until later.*/
-void UNiagaraMeshRendererProperties::InitCDOPropertiesAfterModuleStartupInternal()
+void UNiagaraMeshRendererProperties::InitCDOPropertiesAfterModuleStartup()
 {
-	Super::InitCDOPropertiesAfterModuleStartupInternal();
-	
+	UNiagaraMeshRendererProperties* CDO = CastChecked<UNiagaraMeshRendererProperties>(UNiagaraMeshRendererProperties::StaticClass()->GetDefaultObject());
+	CDO->InitBindings();
+
 	for (TWeakObjectPtr<UNiagaraMeshRendererProperties>& WeakMeshRendererProperties : MeshRendererPropertiesToDeferredInit)
 	{
 		if (WeakMeshRendererProperties.Get())
@@ -259,8 +260,6 @@ void UNiagaraMeshRendererProperties::InitCDOPropertiesAfterModuleStartupInternal
 
 void UNiagaraMeshRendererProperties::InitBindings()
 {
-	Super::InitBindings();
-	
 	if (!PositionBinding.IsValid())
 	{
 		PositionBinding = FNiagaraConstants::GetAttributeDefaultBinding(SYS_PARAM_PARTICLES_POSITION);
diff --git a/Engine/Plugins/FX/Niagara/Source/Niagara/Private/NiagaraModule.cpp b/Engine/Plugins/FX/Niagara/Source/Niagara/Private/NiagaraModule.cpp
--- a/Engine/Plugins/FX/Niagara/Source/Niagara/Private/NiagaraModule.cpp
+++ b/Engine/Plugins/FX/Niagara/Source/Niagara/Private/NiagaraModule.cpp
@@ -25,7 +25,6 @@
 #include "NiagaraGpuComputeDispatch.h"
 #include "Misc/CoreDelegates.h"
 #include "NiagaraDebuggerClient.h"
-#include "NiagaraSimulationStageBase.h"
 #include "Particles/FXBudget.h"
 #include "Engine/StaticMesh.h"
 #include "UObject/UObjectGlobals.h"
@@ -342,14 +341,12 @@ void INiagaraModule::StartupModule()
 	Translator_CallID = FNiagaraVariable(FNiagaraTypeDefinition::GetIntDef(), TEXT("Translator.CallID"));
 
 	FNiagaraConstants::Init();
-	UNiagaraRendererProperties::InitCDOPropertiesAfterModuleStartup(UNiagaraLightRendererProperties::StaticClass());
-	UNiagaraRendererProperties::InitCDOPropertiesAfterModuleStartup(UNiagaraSpriteRendererProperties::StaticClass());
-	UNiagaraRendererProperties::InitCDOPropertiesAfterModuleStartup(UNiagaraRibbonRendererProperties::StaticClass());
-	UNiagaraRendererProperties::InitCDOPropertiesAfterModuleStartup(UNiagaraMeshRendererProperties::StaticClass());
-	UNiagaraRendererProperties::InitCDOPropertiesAfterModuleStartup(UNiagaraComponentRendererProperties::StaticClass());
+	UNiagaraLightRendererProperties::InitCDOPropertiesAfterModuleStartup();
+	UNiagaraSpriteRendererProperties::InitCDOPropertiesAfterModuleStartup();
+	UNiagaraRibbonRendererProperties::InitCDOPropertiesAfterModuleStartup();
+	UNiagaraMeshRendererProperties::InitCDOPropertiesAfterModuleStartup();
+	UNiagaraComponentRendererProperties::InitCDOPropertiesAfterModuleStartup();
 
-	UNiagaraSimulationStageGeneric::InitCDOPropertiesAfterModuleStartup();
-	
 	// Register the data interface CDO finder with the shader module..
 	INiagaraShaderModule& NiagaraShaderModule = FModuleManager::LoadModuleChecked<INiagaraShaderModule>("NiagaraShader");
 	NiagaraShaderModule.SetOnRequestDefaultDataInterfaceHandler(INiagaraShaderModule::FOnRequestDefaultDataInterface::CreateLambda([](const FString& DIClassName) -> UNiagaraDataInterfaceBase*
diff --git a/Engine/Plugins/FX/Niagara/Source/Niagara/Private/NiagaraRendererProperties.cpp b/Engine/Plugins/FX/Niagara/Source/Niagara/Private/NiagaraRendererProperties.cpp
--- a/Engine/Plugins/FX/Niagara/Source/Niagara/Private/NiagaraRendererProperties.cpp
+++ b/Engine/Plugins/FX/Niagara/Source/Niagara/Private/NiagaraRendererProperties.cpp
@@ -492,28 +492,10 @@ void UNiagaraRendererProperties::PostLoad()
 	{
 		MotionVectorSetting = ENiagaraRendererMotionVectorSetting::Disable;
 	}
-
-	// it was possible to revert the enabled binding's type to "none" when the revert arrow is used
-	if(!RendererEnabledBinding.GetType().IsValid())
-	{
-		FNiagaraVariableBase EnabledDefaultVariable(FNiagaraTypeDefinition::GetBoolDef(), NAME_None);
-		RendererEnabledBinding.Setup(EnabledDefaultVariable, EnabledDefaultVariable, ENiagaraRendererSourceDataMode::Emitter);
-	}
-}
-
-void UNiagaraRendererProperties::InitBindings()
-{
-	FNiagaraVariableBase EnabledDefaultVariable(FNiagaraTypeDefinition::GetBoolDef(), NAME_None);
-	RendererEnabledBinding.Setup(EnabledDefaultVariable, EnabledDefaultVariable, ENiagaraRendererSourceDataMode::Emitter);
-}
-
-void UNiagaraRendererProperties::InitCDOPropertiesAfterModuleStartup(TSubclassOf<UNiagaraRendererProperties> Class)
-{
-	UNiagaraRendererProperties* CDO = CastChecked<UNiagaraRendererProperties>(Class->GetDefaultObject());
-	CDO->InitCDOPropertiesAfterModuleStartupInternal();
 }
 
 #if WITH_EDITORONLY_DATA
+
 void UNiagaraRendererProperties::PostEditChangeProperty(struct FPropertyChangedEvent& PropertyChangedEvent)
 {
 	Super::PostEditChangeProperty(PropertyChangedEvent);
diff --git a/Engine/Plugins/FX/Niagara/Source/Niagara/Private/NiagaraRibbonRendererProperties.cpp b/Engine/Plugins/FX/Niagara/Source/Niagara/Private/NiagaraRibbonRendererProperties.cpp
--- a/Engine/Plugins/FX/Niagara/Source/Niagara/Private/NiagaraRibbonRendererProperties.cpp
+++ b/Engine/Plugins/FX/Niagara/Source/Niagara/Private/NiagaraRibbonRendererProperties.cpp
@@ -227,10 +227,11 @@ void UNiagaraRibbonRendererProperties::PostInitProperties()
 }
 
 /** The bindings depend on variables that are created during the NiagaraModule startup. However, the CDO's are build prior to this being initialized, so we defer setting these values until later.*/
-void UNiagaraRibbonRendererProperties::InitCDOPropertiesAfterModuleStartupInternal()
+void UNiagaraRibbonRendererProperties::InitCDOPropertiesAfterModuleStartup()
 {
-	Super::InitCDOPropertiesAfterModuleStartupInternal();
-	
+	UNiagaraRibbonRendererProperties* CDO = CastChecked<UNiagaraRibbonRendererProperties>(UNiagaraRibbonRendererProperties::StaticClass()->GetDefaultObject());
+	CDO->InitBindings();
+
 	for (TWeakObjectPtr<UNiagaraRibbonRendererProperties>& WeakRibbonRendererProperties : RibbonRendererPropertiesToDeferredInit)
 	{
 		if (WeakRibbonRendererProperties.Get())
@@ -242,8 +243,6 @@ void UNiagaraRibbonRendererProperties::InitCDOPropertiesAfterModuleStartupIntern
 
 void UNiagaraRibbonRendererProperties::InitBindings()
 {
-	Super::InitBindings();
-
 	if (!PositionBinding.IsValid())
 	{
 		PositionBinding = FNiagaraConstants::GetAttributeDefaultBinding(SYS_PARAM_PARTICLES_POSITION);
diff --git a/Engine/Plugins/FX/Niagara/Source/Niagara/Private/NiagaraSimulationStageBase.cpp b/Engine/Plugins/FX/Niagara/Source/Niagara/Private/NiagaraSimulationStageBase.cpp
--- a/Engine/Plugins/FX/Niagara/Source/Niagara/Private/NiagaraSimulationStageBase.cpp
+++ b/Engine/Plugins/FX/Niagara/Source/Niagara/Private/NiagaraSimulationStageBase.cpp
@@ -121,29 +121,17 @@ bool UNiagaraSimulationStageGeneric::FillCompilationData(TArray<FNiagaraSimulati
 }
 #endif //WITH_EDITORONLY_DATA
 
-void UNiagaraSimulationStageGeneric::InitBindings()
-{
-	using namespace NiagaraSimulationStageLocal;
-	EnabledBinding.Setup(GetDefaultEnabledBinding(), GetDefaultEnabledBinding(), ENiagaraRendererSourceDataMode::Emitter);
-	ElementCountBinding.Setup(GetDefaultElementCountBinding(), GetDefaultElementCountBinding(), ENiagaraRendererSourceDataMode::Emitter);
-	NumIterationsBinding.Setup(GetDefaultNumIterationsBinding(), GetDefaultNumIterationsBinding(), ENiagaraRendererSourceDataMode::Emitter);
-	ParticleIterationStateBinding.Setup(GetDefaultParticleStateBinding(), GetDefaultParticleStateBinding(), ENiagaraRendererSourceDataMode::Particles);
-}
-
-void UNiagaraSimulationStageGeneric::InitCDOPropertiesAfterModuleStartup()
-{
-	UNiagaraSimulationStageGeneric* CDO = GetMutableDefault<UNiagaraSimulationStageGeneric>();
-	CDO->InitBindings();
-}
-
 void UNiagaraSimulationStageGeneric::PostInitProperties()
 {
 	Super::PostInitProperties();
 
-	// the type registry isn't live at this point for the CDO so we defer the bindings initialization for it
-	if(!HasAnyFlags(RF_ClassDefaultObject))
+	if (HasAnyFlags(RF_ClassDefaultObject) == false)
 	{
-		InitBindings();
+		using namespace NiagaraSimulationStageLocal;
+		EnabledBinding.Setup(GetDefaultEnabledBinding(), GetDefaultEnabledBinding(), ENiagaraRendererSourceDataMode::Emitter);
+		ElementCountBinding.Setup(GetDefaultElementCountBinding(), GetDefaultElementCountBinding(), ENiagaraRendererSourceDataMode::Emitter);
+		NumIterationsBinding.Setup(GetDefaultNumIterationsBinding(), GetDefaultNumIterationsBinding(), ENiagaraRendererSourceDataMode::Emitter);
+		ParticleIterationStateBinding.Setup(GetDefaultParticleStateBinding(), GetDefaultParticleStateBinding(), ENiagaraRendererSourceDataMode::Particles);
 	}
 }
 
diff --git a/Engine/Plugins/FX/Niagara/Source/Niagara/Private/NiagaraSpriteRendererProperties.cpp b/Engine/Plugins/FX/Niagara/Source/Niagara/Private/NiagaraSpriteRendererProperties.cpp
--- a/Engine/Plugins/FX/Niagara/Source/Niagara/Private/NiagaraSpriteRendererProperties.cpp
+++ b/Engine/Plugins/FX/Niagara/Source/Niagara/Private/NiagaraSpriteRendererProperties.cpp
@@ -200,10 +200,11 @@ void UNiagaraSpriteRendererProperties::Serialize(FStructuredArchive::FRecord Rec
 }
 
 /** The bindings depend on variables that are created during the NiagaraModule startup. However, the CDO's are build prior to this being initialized, so we defer setting these values until later.*/
-void UNiagaraSpriteRendererProperties::InitCDOPropertiesAfterModuleStartupInternal()
+void UNiagaraSpriteRendererProperties::InitCDOPropertiesAfterModuleStartup()
 {
-	Super::InitCDOPropertiesAfterModuleStartupInternal();
-	
+	UNiagaraSpriteRendererProperties* CDO = CastChecked<UNiagaraSpriteRendererProperties>(UNiagaraSpriteRendererProperties::StaticClass()->GetDefaultObject());
+	CDO->InitBindings();
+
 	for (TWeakObjectPtr<UNiagaraSpriteRendererProperties>& WeakSpriteRendererProperties : SpriteRendererPropertiesToDeferredInit)
 	{
 		if (WeakSpriteRendererProperties.Get())
@@ -215,8 +216,6 @@ void UNiagaraSpriteRendererProperties::InitCDOPropertiesAfterModuleStartupIntern
 
 void UNiagaraSpriteRendererProperties::InitBindings()
 {
-	Super::InitBindings();
-	
 	if (PositionBinding.GetParamMapBindableVariable().GetName() == NAME_None)
 	{
 		PositionBinding = FNiagaraConstants::GetAttributeDefaultBinding(SYS_PARAM_PARTICLES_POSITION);
diff --git a/Engine/Plugins/FX/Niagara/Source/Niagara/Public/NiagaraComponentRendererProperties.h b/Engine/Plugins/FX/Niagara/Source/Niagara/Public/NiagaraComponentRendererProperties.h
--- a/Engine/Plugins/FX/Niagara/Source/Niagara/Public/NiagaraComponentRendererProperties.h
+++ b/Engine/Plugins/FX/Niagara/Source/Niagara/Public/NiagaraComponentRendererProperties.h
@@ -59,7 +59,9 @@ public:
 	virtual void PostInitProperties() override;
 	virtual void PostDuplicate(bool bDuplicateForPIE) override;
 	//UObject Interface END
-	
+
+	static void InitCDOPropertiesAfterModuleStartup();
+
 	//~ UNiagaraRendererProperties interface
 	virtual FNiagaraRenderer* CreateEmitterRenderer(ERHIFeatureLevel::Type FeatureLevel, const FNiagaraEmitterInstance* Emitter, const FNiagaraSystemInstanceController& InController) override;
 	virtual class FNiagaraBoundsCalculator* CreateBoundsCalculator() override { return nullptr; }
@@ -158,8 +160,6 @@ public:
 	virtual bool NeedsSystemCompletion() const override { return true; }
 
 protected:
-	virtual void InitBindings() override;
-	virtual void InitCDOPropertiesAfterModuleStartupInternal() override;
 	virtual void UpdateSourceModeDerivates(ENiagaraRendererSourceDataMode InSourceMode, bool bFromPropertyEdit = false) override;
 
 private:
diff --git a/Engine/Plugins/FX/Niagara/Source/Niagara/Public/NiagaraMeshRendererProperties.h b/Engine/Plugins/FX/Niagara/Source/Niagara/Public/NiagaraMeshRendererProperties.h
--- a/Engine/Plugins/FX/Niagara/Source/Niagara/Public/NiagaraMeshRendererProperties.h
+++ b/Engine/Plugins/FX/Niagara/Source/Niagara/Public/NiagaraMeshRendererProperties.h
@@ -170,6 +170,8 @@ public:
 #endif// WITH_EDITORONLY_DATA
 	//UObject Interface END
 
+	static void InitCDOPropertiesAfterModuleStartup();
+
 	//~ UNiagaraRendererProperties interface
 	virtual FNiagaraRenderer* CreateEmitterRenderer(ERHIFeatureLevel::Type FeatureLevel, const FNiagaraEmitterInstance* Emitter, const FNiagaraSystemInstanceController& InController) override;
 	virtual class FNiagaraBoundsCalculator* CreateBoundsCalculator() override;
@@ -409,9 +411,8 @@ public:
 	FNiagaraRendererLayout RendererLayoutWithoutCustomSorting;
 
 protected:
-	virtual void InitBindings() override;
+	void InitBindings();
 	void SetPreviousBindings(const FVersionedNiagaraEmitter& SrcEmitter, ENiagaraRendererSourceDataMode InSourceMode);
-	virtual void InitCDOPropertiesAfterModuleStartupInternal() override;
 	virtual void UpdateSourceModeDerivates(ENiagaraRendererSourceDataMode InSourceMode, bool bFromPropertyEdit = false) override;
 	virtual bool NeedsMIDsForMaterials() const override { return MaterialParameters.HasAnyBindings(); }
 #if WITH_EDITORONLY_DATA
diff --git a/Engine/Plugins/FX/Niagara/Source/Niagara/Public/NiagaraRendererProperties.h b/Engine/Plugins/FX/Niagara/Source/Niagara/Public/NiagaraRendererProperties.h
--- a/Engine/Plugins/FX/Niagara/Source/Niagara/Public/NiagaraRendererProperties.h
+++ b/Engine/Plugins/FX/Niagara/Source/Niagara/Public/NiagaraRendererProperties.h
@@ -249,15 +249,11 @@ public:
 		, bAllowInCullProxies(true)
 		, bMotionBlurEnabled_DEPRECATED(true)
 	{
-		AttributeBindings.Add(&RendererEnabledBinding);
 	}
 
 	//UObject Interface Begin
 	virtual void PostInitProperties() override;
 	virtual void PostLoad() override;
-	virtual void InitBindings();
-
-	static void InitCDOPropertiesAfterModuleStartup(TSubclassOf<UNiagaraRendererProperties> Class);
 #if WITH_EDITORONLY_DATA
 	virtual void PostEditChangeProperty(struct FPropertyChangedEvent& PropertyChangedEvent) override;
 #endif
@@ -381,7 +377,7 @@ protected:
 
 	virtual void PostLoadBindings(ENiagaraRendererSourceDataMode InSourceMode);
 	virtual void UpdateSourceModeDerivates(ENiagaraRendererSourceDataMode InSourceMode, bool bFromPropertyEdit = false);
-	virtual void InitCDOPropertiesAfterModuleStartupInternal() { check(HasAnyFlags(RF_ClassDefaultObject)); InitBindings(); }
+
 #if WITH_EDITORONLY_DATA
 	/** returns the variable associated with the supplied binding if it should be bound given the current settings of the RendererProperties. */
 	virtual FNiagaraVariable GetBoundAttribute(const FNiagaraVariableAttributeBinding* Binding) const;
diff --git a/Engine/Plugins/FX/Niagara/Source/Niagara/Public/NiagaraRibbonRendererProperties.h b/Engine/Plugins/FX/Niagara/Source/Niagara/Public/NiagaraRibbonRendererProperties.h
--- a/Engine/Plugins/FX/Niagara/Source/Niagara/Public/NiagaraRibbonRendererProperties.h
+++ b/Engine/Plugins/FX/Niagara/Source/Niagara/Public/NiagaraRibbonRendererProperties.h
@@ -207,6 +207,8 @@ public:
 #endif
 	//UObject Interface END
 
+	static void InitCDOPropertiesAfterModuleStartup();
+
 	//UNiagaraRendererProperties Interface
 	virtual FNiagaraRenderer* CreateEmitterRenderer(ERHIFeatureLevel::Type FeatureLevel, const FNiagaraEmitterInstance* Emitter, const FNiagaraSystemInstanceController& InController) override;
 	virtual class FNiagaraBoundsCalculator* CreateBoundsCalculator() override;
@@ -479,10 +481,9 @@ public:
 	FNiagaraRendererLayout RendererLayout;
 
 protected:
-	virtual void InitBindings() override;
+	void InitBindings();
 	void SetPreviousBindings(const FVersionedNiagaraEmitter& SrcEmitter);
 
-	virtual void InitCDOPropertiesAfterModuleStartupInternal() override;
 	void UpdateSourceModeDerivates(ENiagaraRendererSourceDataMode InSourceMode, bool bFromPropertyEdit);
 
 	virtual bool NeedsMIDsForMaterials() const { return MaterialParameters.HasAnyBindings(); }
diff --git a/Engine/Plugins/FX/Niagara/Source/Niagara/Public/NiagaraSimulationStageBase.h b/Engine/Plugins/FX/Niagara/Source/Niagara/Public/NiagaraSimulationStageBase.h
--- a/Engine/Plugins/FX/Niagara/Source/Niagara/Public/NiagaraSimulationStageBase.h
+++ b/Engine/Plugins/FX/Niagara/Source/Niagara/Public/NiagaraSimulationStageBase.h
@@ -139,9 +139,6 @@ public:
 	UPROPERTY(EditAnywhere, Category = "Simulation Stage")
 	FIntVector OverrideGpuDispatchNumThreads = FIntVector(64, 1, 1);
 
-	void InitBindings();
-	static void InitCDOPropertiesAfterModuleStartup();
-
 	virtual void PostInitProperties() override;
 	virtual void PostLoad() override;
 	virtual bool AppendCompileHash(FNiagaraCompileHashVisitor* InVisitor) const override;
diff --git a/Engine/Plugins/FX/Niagara/Source/Niagara/Public/NiagaraSpriteRendererProperties.h b/Engine/Plugins/FX/Niagara/Source/Niagara/Public/NiagaraSpriteRendererProperties.h
--- a/Engine/Plugins/FX/Niagara/Source/Niagara/Public/NiagaraSpriteRendererProperties.h
+++ b/Engine/Plugins/FX/Niagara/Source/Niagara/Public/NiagaraSpriteRendererProperties.h
@@ -116,7 +116,9 @@ public:
 
 #endif // WITH_EDITORONLY_DATA
 	//UObject Interface END
-	
+
+	static void InitCDOPropertiesAfterModuleStartup();
+
 	//UNiagaraRendererProperties interface
 	virtual FNiagaraRenderer* CreateEmitterRenderer(ERHIFeatureLevel::Type FeatureLevel, const FNiagaraEmitterInstance* Emitter, const FNiagaraSystemInstanceController& InController) override;
 	virtual class FNiagaraBoundsCalculator* CreateBoundsCalculator() override;
@@ -389,9 +391,8 @@ public:
 	uint32 MaterialParamValidMask = 0;
 
 protected:
-	virtual void InitBindings() override;
+	void InitBindings();
 	void SetPreviousBindings(const FVersionedNiagaraEmitter& SrcEmitter, ENiagaraRendererSourceDataMode InSourceMode);
-	virtual void InitCDOPropertiesAfterModuleStartupInternal() override;
 	virtual void UpdateSourceModeDerivates(ENiagaraRendererSourceDataMode InSourceMode, bool bFromPropertyEdit = false) override;
 
 #if WITH_EDITORONLY_DATA
diff --git a/Engine/Plugins/FX/Niagara/Source/NiagaraEditor/Private/Customizations/NiagaraTypeCustomizations.cpp b/Engine/Plugins/FX/Niagara/Source/NiagaraEditor/Private/Customizations/NiagaraTypeCustomizations.cpp
--- a/Engine/Plugins/FX/Niagara/Source/NiagaraEditor/Private/Customizations/NiagaraTypeCustomizations.cpp
+++ b/Engine/Plugins/FX/Niagara/Source/NiagaraEditor/Private/Customizations/NiagaraTypeCustomizations.cpp
@@ -368,7 +368,12 @@ void FNiagaraVariableAttributeBindingCustomization::ChangeSource(FName InVarName
 	PropertyHandle->NotifyFinishedChangingProperties();
 }
 
-bool FNiagaraVariableAttributeBindingCustomization::IsResetToDefaultsVisible(TSharedPtr<IPropertyHandle> InPropertyHandle) const
+void FNiagaraVariableAttributeBindingCustomization::ResetToDefault()
+{
+	UE_LOG(LogNiagaraEditor, Warning, TEXT("Reset to default!"));
+}
+
+EVisibility FNiagaraVariableAttributeBindingCustomization::IsResetToDefaultsVisible() const
 {
 	check(BaseEmitter.Emitter);
 	check(RenderProps || SimulationStage);
@@ -376,10 +381,10 @@ bool FNiagaraVariableAttributeBindingCustomization::IsResetToDefaultsVisible(TSh
 	check(DefaultVariableBinding);
 
 	const ENiagaraRendererSourceDataMode BindingSourceMode = RenderProps ? RenderProps->GetCurrentSourceMode() : ENiagaraRendererSourceDataMode::Emitter;
-	return !TargetVariableBinding->MatchesDefault(*DefaultVariableBinding, BindingSourceMode);
+	return TargetVariableBinding->MatchesDefault(*DefaultVariableBinding, BindingSourceMode) ? EVisibility::Hidden : EVisibility::Visible;
 }
 
-void FNiagaraVariableAttributeBindingCustomization::OnResetToDefaultsClicked(TSharedPtr<IPropertyHandle> InPropertyHandle)
+FReply FNiagaraVariableAttributeBindingCustomization::OnResetToDefaultsClicked()
 {
 	FScopedTransaction Transaction(LOCTEXT("ResetBindingParam", "Reset binding"));
 	TArray<UObject*> Objects;
@@ -398,6 +403,7 @@ void FNiagaraVariableAttributeBindingCustomization::OnResetToDefaultsClicked(TSh
 	TargetVariableBinding->ResetToDefault(*DefaultVariableBinding, BaseEmitter, BindingSourceMode);
 	PropertyHandle->NotifyPostChange(EPropertyChangeType::ValueSet);
 	PropertyHandle->NotifyFinishedChangingProperties();
+	return FReply::Handled();
 }
 
 void FNiagaraVariableAttributeBindingCustomization::CustomizeHeader(TSharedRef<IPropertyHandle> InPropertyHandle, FDetailWidgetRow& HeaderRow, IPropertyTypeCustomizationUtils& CustomizationUtils)
@@ -409,6 +415,17 @@ void FNiagaraVariableAttributeBindingCustomization::CustomizeHeader(TSharedRef<I
 	TArray<UObject*> Objects;
 	PropertyHandle->GetOuterObjects(Objects);
 	bool bAddDefault = true;
+
+	InPropertyHandle->SetOnPropertyResetToDefault(FSimpleDelegate::CreateLambda([this]() { ResetToDefault(); }));
+	//InPropertyHandle->ExecuteCustomResetToDefault
+
+	/*FResetToDefaultOverride ResetOverride = FResetToDefaultOverride::Create(
+		FIsResetToDefaultVisible::CreateSP(this, &FMotionControllerDetails::IsSourceValueModified),
+		FResetToDefaultHandler::CreateSP(this, &FMotionControllerDetails::OnResetSourceValue)
+	);
+
+	PropertyRow.OverrideResetToDefault(ResetOverride); */
+	InPropertyHandle->MarkResetToDefaultCustomized(true);
 	
 	if (Objects.Num() == 1)
 	{
@@ -428,15 +445,7 @@ void FNiagaraVariableAttributeBindingCustomization::CustomizeHeader(TSharedRef<I
 		{
 			TargetVariableBinding = (FNiagaraVariableAttributeBinding*)PropertyHandle->GetValueBaseAddress((uint8*)Objects[0]);
 			DefaultVariableBinding = (FNiagaraVariableAttributeBinding*)PropertyHandle->GetValueBaseAddress((uint8*)Objects[0]->GetClass()->GetDefaultObject());
-
-			FResetToDefaultOverride OverrideReset = FResetToDefaultOverride::Create
-			(
-				FIsResetToDefaultVisible::CreateSP(this, &FNiagaraVariableAttributeBindingCustomization::IsResetToDefaultsVisible),
-				FResetToDefaultHandler::CreateSP(this, &FNiagaraVariableAttributeBindingCustomization::OnResetToDefaultsClicked)
-			);
-			
-			HeaderRow.OverrideResetToDefault(OverrideReset);
-			
+				
 			HeaderRow
 				.NameContent()
 				[
@@ -454,10 +463,34 @@ void FNiagaraVariableAttributeBindingCustomization::CustomizeHeader(TSharedRef<I
 						.ContentPadding(1)
 						.ToolTipText(this, &FNiagaraVariableAttributeBindingCustomization::GetTooltipText)
 						.ButtonContent()
-						[							
+						[
+							/*SNew(STextBlock)
+							.Text(this, &FNiagaraVariableAttributeBindingCustomization::GetCurrentText)
+							.Font(IDetailLayoutBuilder::GetDetailFont())*/
 							SNew(SNiagaraParameterName)
 							.ParameterName(this, &FNiagaraVariableAttributeBindingCustomization::GetVariableName)
 							.IsReadOnly(true)
+							//SNew(STextBlock)
+							//.Text(InCreateData->Action->GetMenuDescription())
+						]
+					]
+					+ SHorizontalBox::Slot()
+					.AutoWidth()
+					.VAlign(VAlign_Center)
+					.Padding(2, 1)
+					[
+						SNew(SButton)
+						.IsFocusable(false)
+						.ToolTipText(LOCTEXT("ResetToDefaultToolTip", "Reset to Default"))
+						.ButtonStyle(FAppStyle::Get(), "NoBorder")
+						.ContentPadding(0)
+						.Visibility(this, &FNiagaraVariableAttributeBindingCustomization::IsResetToDefaultsVisible)
+						.OnClicked(this, &FNiagaraVariableAttributeBindingCustomization::OnResetToDefaultsClicked)
+						.Content()
+						[
+							SNew(SImage)
+							.Image(FAppStyle::GetBrush("PropertyWindow.DiffersFromDefault"))
+							.ColorAndOpacity(FSlateColor::UseForeground())
 						]
 					]
 				];
diff --git a/Engine/Plugins/FX/Niagara/Source/NiagaraEditor/Private/Customizations/NiagaraTypeCustomizations.h b/Engine/Plugins/FX/Niagara/Source/NiagaraEditor/Private/Customizations/NiagaraTypeCustomizations.h
--- a/Engine/Plugins/FX/Niagara/Source/NiagaraEditor/Private/Customizations/NiagaraTypeCustomizations.h
+++ b/Engine/Plugins/FX/Niagara/Source/NiagaraEditor/Private/Customizations/NiagaraTypeCustomizations.h
@@ -168,8 +168,9 @@ public:
 	virtual void CustomizeChildren(TSharedRef<IPropertyHandle> StructPropertyHandle, class IDetailChildrenBuilder& ChildBuilder, IPropertyTypeCustomizationUtils& StructCustomizationUtils) override {};
 	/** IPropertyTypeCustomization interface end */
 private:
-	bool IsResetToDefaultsVisible(TSharedPtr<IPropertyHandle> InPropertyHandle) const;
-	void OnResetToDefaultsClicked(TSharedPtr<IPropertyHandle> InPropertyHandle);
+	EVisibility IsResetToDefaultsVisible() const;
+	FReply OnResetToDefaultsClicked();
+	void ResetToDefault();
 	FName GetVariableName() const;
 	FText GetCurrentText() const;
 	FText GetTooltipText() const;
diff --git a/Engine/Plugins/FX/Niagara/Source/NiagaraEditorWidgets/Private/SNiagaraStack.cpp b/Engine/Plugins/FX/Niagara/Source/NiagaraEditorWidgets/Private/SNiagaraStack.cpp
--- a/Engine/Plugins/FX/Niagara/Source/NiagaraEditorWidgets/Private/SNiagaraStack.cpp
+++ b/Engine/Plugins/FX/Niagara/Source/NiagaraEditorWidgets/Private/SNiagaraStack.cpp
@@ -27,7 +27,6 @@
 #include "Framework/Commands/GenericCommands.h"
 #include "IContentBrowserSingleton.h"
 #include "ContentBrowserModule.h"
-#include "DetailWidgetRow.h"
 #include "Widgets/Input/SButton.h"
 #include "Widgets/Images/SImage.h"
 #include "Widgets/Input/SComboButton.h"
@@ -956,90 +955,78 @@ SNiagaraStack::FRowWidgets SNiagaraStack::ConstructNameAndValueWidgetsForItem(UN
 		UNiagaraStackPropertyRow* PropertyRow = CastChecked<UNiagaraStackPropertyRow>(Item);
 		FNodeWidgets PropertyRowWidgets = PropertyRow->GetDetailTreeNode()->CreateNodeWidgets();
 
-		TSharedPtr<SWidget> TmpNameWidget;
-		TSharedPtr<SWidget> TmpValueWidget;
-		FDetailWidgetRow WidgetRow;
-
-		TOptional<FResetToDefaultOverride> ResetToDefaultOverride;
-		if(TSharedPtr<IDetailPropertyRow> Row = PropertyRow->GetDetailTreeNode()->GetRow())
-		{
-			// we override the name & value widgets here as the previous ones won't be valid anymore but we require the widget row to access additional data like the custom reset override.
-			PropertyRow->GetDetailTreeNode()->GetRow()->GetDefaultWidgets(PropertyRowWidgets.NameWidget, PropertyRowWidgets.ValueWidget, WidgetRow);
-			ResetToDefaultOverride = WidgetRow.CustomResetToDefault;
-		}
-
 		TAttribute<bool> IsEnabled;
 		IsEnabled.BindUObject(Item, &UNiagaraStackEntry::GetIsEnabledAndOwnerIsEnabled);
 
 		Container->AddFillRowContextMenuHandler(FNiagaraStackPropertyRowUtilities::CreateOnFillRowContextMenu(PropertyRow->GetDetailTreeNode()->CreatePropertyHandle(), PropertyRowWidgets.Actions));
-		
-		 if (PropertyRowWidgets.WholeRowWidget.IsValid())
-		 {
-		 	Container->SetOverrideNameWidth(PropertyRowWidgets.WholeRowWidgetLayoutData.MinWidth, PropertyRowWidgets.WholeRowWidgetLayoutData.MaxWidth);
-		 	Container->SetOverrideNameAlignment(PropertyRowWidgets.WholeRowWidgetLayoutData.HorizontalAlignment, PropertyRowWidgets.WholeRowWidgetLayoutData.VerticalAlignment);
-		 	PropertyRowWidgets.WholeRowWidget->SetEnabled(IsEnabled);
-		 	TSharedRef<SHorizontalBox> RowBox = SNew(SHorizontalBox);
-		 	if (PropertyRowWidgets.EditConditionWidget.IsValid())
-		 	{
-		 		RowBox->AddSlot()
-		 		.AutoWidth()
-		 		.Padding(0, 0, 3, 0)
-		 		[
-		 			PropertyRowWidgets.EditConditionWidget.ToSharedRef()
-		 		];
-		 	}
-		 	RowBox->AddSlot()
-		 	[
-		 		PropertyRowWidgets.WholeRowWidget.ToSharedRef()
-		 	];
-		 	RowBox->AddSlot()
-		 	.AutoWidth()
-		 	.Padding(3, 0, 0, 0)
-		 	[
-		 		SNew(SResetToDefaultPropertyEditor, PropertyRow->GetDetailTreeNode()->CreatePropertyHandle()).CustomResetToDefault(ResetToDefaultOverride)
-		 	];
-		 	return FRowWidgets(RowBox); 
-		 }
-		 else
-		 {
-		 	Container->SetOverrideNameWidth(PropertyRowWidgets.NameWidgetLayoutData.MinWidth, PropertyRowWidgets.NameWidgetLayoutData.MaxWidth);
-		 	Container->SetOverrideValueWidth(PropertyRowWidgets.ValueWidgetLayoutData.MinWidth, PropertyRowWidgets.ValueWidgetLayoutData.MaxWidth);
-		 	Container->SetOverrideNameAlignment(HAlign_Left, VAlign_Center);
-		 	PropertyRowWidgets.NameWidget->SetEnabled(IsEnabled);
-		 	PropertyRowWidgets.ValueWidget->SetEnabled(IsEnabled);
-		
-		 	TSharedPtr<SWidget> NameWidget;
-		 	if (PropertyRowWidgets.EditConditionWidget.IsValid())
-		 	{
-		 		NameWidget = SNew(SHorizontalBox)
-		 		+ SHorizontalBox::Slot()
-		 		.AutoWidth()
-		 		.Padding(0, 0, 3, 0)
-		 		[
-		 			PropertyRowWidgets.EditConditionWidget.ToSharedRef()
-		 		]
-		 		+ SHorizontalBox::Slot()
-		 		[
-		 			PropertyRowWidgets.NameWidget.ToSharedRef()
-		 		];
-		 	}
-		 	else
-		 	{
-		 		NameWidget = PropertyRowWidgets.NameWidget;
-		 	}
-		
-		 	TSharedRef<SWidget> ValueWidget = SNew(SHorizontalBox)
-		 	+ SHorizontalBox::Slot()
-		 	[
-		 		PropertyRowWidgets.ValueWidget.ToSharedRef()
-		 	]
-		 	+ SHorizontalBox::Slot()
-		 	.AutoWidth()
-		 	.VAlign(VAlign_Center)
-		 	.Padding(3, 0, 0, 0)
-		 	[
-		 		SNew(SResetToDefaultPropertyEditor, PropertyRow->GetDetailTreeNode()->CreatePropertyHandle()).CustomResetToDefault(ResetToDefaultOverride)
-		 	];
+
+		if (PropertyRowWidgets.WholeRowWidget.IsValid())
+		{
+			Container->SetOverrideNameWidth(PropertyRowWidgets.WholeRowWidgetLayoutData.MinWidth, PropertyRowWidgets.WholeRowWidgetLayoutData.MaxWidth);
+			Container->SetOverrideNameAlignment(PropertyRowWidgets.WholeRowWidgetLayoutData.HorizontalAlignment, PropertyRowWidgets.WholeRowWidgetLayoutData.VerticalAlignment);
+			PropertyRowWidgets.WholeRowWidget->SetEnabled(IsEnabled);
+			TSharedRef<SHorizontalBox> RowBox = SNew(SHorizontalBox);
+			if (PropertyRowWidgets.EditConditionWidget.IsValid())
+			{
+				RowBox->AddSlot()
+				.AutoWidth()
+				.Padding(0, 0, 3, 0)
+				[
+					PropertyRowWidgets.EditConditionWidget.ToSharedRef()
+				];
+			}
+			RowBox->AddSlot()
+			[
+				PropertyRowWidgets.WholeRowWidget.ToSharedRef()
+			];
+			RowBox->AddSlot()
+			.AutoWidth()
+			.Padding(3, 0, 0, 0)
+			[
+				SNew(SResetToDefaultPropertyEditor, PropertyRow->GetDetailTreeNode()->CreatePropertyHandle())
+			];
+			return FRowWidgets(RowBox); 
+		}
+		else
+		{
+			Container->SetOverrideNameWidth(PropertyRowWidgets.NameWidgetLayoutData.MinWidth, PropertyRowWidgets.NameWidgetLayoutData.MaxWidth);
+			Container->SetOverrideValueWidth(PropertyRowWidgets.ValueWidgetLayoutData.MinWidth, PropertyRowWidgets.ValueWidgetLayoutData.MaxWidth);
+			Container->SetOverrideNameAlignment(HAlign_Left, VAlign_Center);
+			PropertyRowWidgets.NameWidget->SetEnabled(IsEnabled);
+			PropertyRowWidgets.ValueWidget->SetEnabled(IsEnabled);
+
+			TSharedPtr<SWidget> NameWidget;
+			if (PropertyRowWidgets.EditConditionWidget.IsValid())
+			{
+				NameWidget = SNew(SHorizontalBox)
+				+ SHorizontalBox::Slot()
+				.AutoWidth()
+				.Padding(0, 0, 3, 0)
+				[
+					PropertyRowWidgets.EditConditionWidget.ToSharedRef()
+				]
+				+ SHorizontalBox::Slot()
+				[
+					PropertyRowWidgets.NameWidget.ToSharedRef()
+				];
+			}
+			else
+			{
+				NameWidget = PropertyRowWidgets.NameWidget;
+			}
+
+			TSharedRef<SWidget> ValueWidget = SNew(SHorizontalBox)
+			+ SHorizontalBox::Slot()
+			[
+				PropertyRowWidgets.ValueWidget.ToSharedRef()
+			]
+			+ SHorizontalBox::Slot()
+			.AutoWidth()
+			.VAlign(VAlign_Center)
+			.Padding(3, 0, 0, 0)
+			[
+				SNew(SResetToDefaultPropertyEditor, PropertyRow->GetDetailTreeNode()->CreatePropertyHandle())
+			];
 
 			return FRowWidgets(NameWidget.ToSharedRef(), ValueWidget);
 		}
diff --git a/Engine/Plugins/Runtime/GeometryCache/Source/GeometryCache/Classes/NiagaraGeometryCacheRendererProperties.h b/Engine/Plugins/Runtime/GeometryCache/Source/GeometryCache/Classes/NiagaraGeometryCacheRendererProperties.h
--- a/Engine/Plugins/Runtime/GeometryCache/Source/GeometryCache/Classes/NiagaraGeometryCacheRendererProperties.h
+++ b/Engine/Plugins/Runtime/GeometryCache/Source/GeometryCache/Classes/NiagaraGeometryCacheRendererProperties.h
@@ -44,6 +44,8 @@ public:
 	virtual void PostInitProperties() override;
 	//UObject Interface END
 
+	static void InitCDOPropertiesAfterModuleStartup();
+
 	//~ UNiagaraRendererProperties interface
 	virtual FNiagaraRenderer* CreateEmitterRenderer(ERHIFeatureLevel::Type FeatureLevel, const FNiagaraEmitterInstance* Emitter, const FNiagaraSystemInstanceController& InController) override;
 	virtual class FNiagaraBoundsCalculator* CreateBoundsCalculator() override { return nullptr; }
@@ -112,8 +114,7 @@ public:
 	virtual bool NeedsSystemCompletion() const override { return true; }
 
 protected:
-	virtual void InitBindings() override;
-	virtual	void InitCDOPropertiesAfterModuleStartupInternal() override;
+	void InitBindings();
 
 	static FNiagaraVariable Particles_GeoCacheRotation;
 	static FNiagaraVariable Particles_Age;
diff --git a/Engine/Plugins/Runtime/GeometryCache/Source/GeometryCache/Private/GeometryCacheModule.cpp b/Engine/Plugins/Runtime/GeometryCache/Source/GeometryCache/Private/GeometryCacheModule.cpp
--- a/Engine/Plugins/Runtime/GeometryCache/Source/GeometryCache/Private/GeometryCacheModule.cpp
+++ b/Engine/Plugins/Runtime/GeometryCache/Source/GeometryCache/Private/GeometryCacheModule.cpp
@@ -20,7 +20,7 @@ void FGeometryCacheModule::StartupModule()
 
 	FCodecV1Decoder::InitLUT();
 
-	UNiagaraRendererProperties::InitCDOPropertiesAfterModuleStartup(UNiagaraGeometryCacheRendererProperties::StaticClass());
+	UNiagaraGeometryCacheRendererProperties::InitCDOPropertiesAfterModuleStartup();
 }
 
 void FGeometryCacheModule::ShutdownModule()
diff --git a/Engine/Plugins/Runtime/GeometryCache/Source/GeometryCache/Private/NiagaraGeometryCacheRendererProperties.cpp b/Engine/Plugins/Runtime/GeometryCache/Source/GeometryCache/Private/NiagaraGeometryCacheRendererProperties.cpp
--- a/Engine/Plugins/Runtime/GeometryCache/Source/GeometryCache/Private/NiagaraGeometryCacheRendererProperties.cpp
+++ b/Engine/Plugins/Runtime/GeometryCache/Source/GeometryCache/Private/NiagaraGeometryCacheRendererProperties.cpp
@@ -63,8 +63,6 @@ FNiagaraVariableAttributeBinding CreateDefaultBinding(FNiagaraVariable DefaultVa
 
 void UNiagaraGeometryCacheRendererProperties::InitBindings()
 {
-	Super::InitBindings();
-
 	if (!PositionBinding.IsValid())
 	{
 		PositionBinding = FNiagaraConstants::GetAttributeDefaultBinding(SYS_PARAM_PARTICLES_POSITION);
@@ -144,19 +142,18 @@ void UNiagaraGeometryCacheRendererProperties::PostInitProperties()
 	}
 }
 
-void UNiagaraGeometryCacheRendererProperties::InitCDOPropertiesAfterModuleStartupInternal()
+void UNiagaraGeometryCacheRendererProperties::InitCDOPropertiesAfterModuleStartup()
 {
 	InitDefaultAttributes();
 	
-	Super::InitCDOPropertiesAfterModuleStartupInternal();
-	
-	PositionBinding = FNiagaraConstants::GetAttributeDefaultBinding(SYS_PARAM_PARTICLES_POSITION);
-	ScaleBinding = FNiagaraConstants::GetAttributeDefaultBinding(SYS_PARAM_PARTICLES_SCALE);
-	RendererVisibilityTagBinding = FNiagaraConstants::GetAttributeDefaultBinding(SYS_PARAM_PARTICLES_VISIBILITY_TAG);
-	RotationBinding = CreateDefaultBinding(Particles_GeoCacheRotation, FVector3f(0, 0, 0));
-	ElapsedTimeBinding = CreateDefaultBinding(Particles_Age, 0.0f);
-	EnabledBinding = CreateDefaultBinding(Particles_GeoCacheIsEnabled, false);
-	ArrayIndexBinding = CreateDefaultBinding(Particles_ArrayIndex, 0);
+	UNiagaraGeometryCacheRendererProperties* CDO = CastChecked<UNiagaraGeometryCacheRendererProperties>(StaticClass()->GetDefaultObject());
+	CDO->PositionBinding = FNiagaraConstants::GetAttributeDefaultBinding(SYS_PARAM_PARTICLES_POSITION);
+	CDO->ScaleBinding = FNiagaraConstants::GetAttributeDefaultBinding(SYS_PARAM_PARTICLES_SCALE);
+	CDO->RendererVisibilityTagBinding = FNiagaraConstants::GetAttributeDefaultBinding(SYS_PARAM_PARTICLES_VISIBILITY_TAG);
+	CDO->RotationBinding = CreateDefaultBinding(Particles_GeoCacheRotation, FVector3f(0, 0, 0));
+	CDO->ElapsedTimeBinding = CreateDefaultBinding(Particles_Age, 0.0f);
+	CDO->EnabledBinding = CreateDefaultBinding(Particles_GeoCacheIsEnabled, false);
+	CDO->ArrayIndexBinding = CreateDefaultBinding(Particles_ArrayIndex, 0);
 
 	for (TWeakObjectPtr<UNiagaraGeometryCacheRendererProperties>& WeakGeometryCacheRendererProperties : GeometryCacheRendererPropertiesToDeferredInit)
 	{
