- Optional rebasing for simulation caches
- Looks for all Position types attributes and will automatically rebase them
- Attributes bounds to renderers like Quaternion / Matrix would also get rebased
- Creation options added to control rebasing, like include / exclude attribute list.  Also controls data interface caching, and in the future if we want a restart or renderer only cache.

#rnx

diff --git a/Engine/Plugins/FX/Niagara/Source/Niagara/Classes/NiagaraBakerOutputSimCache.h b/Engine/Plugins/FX/Niagara/Source/Niagara/Classes/NiagaraBakerOutputSimCache.h
--- a/Engine/Plugins/FX/Niagara/Source/Niagara/Classes/NiagaraBakerOutputSimCache.h
+++ b/Engine/Plugins/FX/Niagara/Source/Niagara/Classes/NiagaraBakerOutputSimCache.h
@@ -19,6 +19,9 @@ public:
 	UPROPERTY(EditAnywhere, Category = "Settings")
 	FString SimCacheAssetPathFormat = TEXT("{AssetFolder}/{AssetName}_SimCache_{OutputName}");
 
+	UPROPERTY(EditAnywhere, Category = "Settings")
+	FNiagaraSimCacheCreateParameters CreateParameters;
+
 	virtual bool Equals(const UNiagaraBakerOutput& Other) const override;
 
 #if WITH_EDITOR
diff --git a/Engine/Plugins/FX/Niagara/Source/Niagara/Classes/NiagaraSimCache.h b/Engine/Plugins/FX/Niagara/Source/Niagara/Classes/NiagaraSimCache.h
--- a/Engine/Plugins/FX/Niagara/Source/Niagara/Classes/NiagaraSimCache.h
+++ b/Engine/Plugins/FX/Niagara/Source/Niagara/Classes/NiagaraSimCache.h
@@ -8,6 +8,54 @@
 
 class UNiagaraComponent;
 
+USTRUCT()
+struct FNiagaraSimCacheCreateParameters
+{
+	GENERATED_BODY()
+
+	FNiagaraSimCacheCreateParameters()
+		: bAllowRebasing(true)
+		, bAllowDataInterfaceCaching(true)
+	{
+	}
+
+	/**
+	When enabled allows the SimCache to be re-based.
+	i.e. World space emitters can be moved to the new component's location
+	*/
+	UPROPERTY(EditAnywhere, Category="SimCache")
+	uint32 bAllowRebasing : 1;
+
+	/**
+	When enabled Data Interface data will be stored in the SimCache.
+	This can result in a large increase to the cache size, depending on what Data Interfaces are used
+	*/
+	UPROPERTY(EditAnywhere, Category="SimCache")
+	uint32 bAllowDataInterfaceCaching : 1;
+
+	/**
+	When enabled the SimCache will only be useful for rendering a replay, it can not be used to restart
+	the simulation from as only attributes & data interface that impact rendering will be stored.
+	This should result in a much smaller caches.
+	*/
+	//UPROPERTY(EditAnywhere, Category="SimCache")
+	//uint32 bRenderOnly : 1;
+
+	/**
+	List of Attributes to force include in the SimCache rebase, they should be the full path to the attribute
+	For example, MyEmitter.Particles.MyQuat would force the particle attribute MyQuat to be included for MyEmitter
+	*/
+	UPROPERTY(EditAnywhere, Category = "SimCache")
+	TArray<FName> RebaseIncludeList;
+
+	/**
+	List of Attributes to force exclude from the SimCache rebase, they should be the full path to the attribute
+	For example, MyEmitter.Particles.MyQuat would force the particle attribute MyQuat to be included for MyEmitter
+	*/
+	UPROPERTY(EditAnywhere, Category = "SimCache")
+	TArray<FName> RebaseExcludeList;
+};
+
 USTRUCT()
 struct FNiagaraSimCacheDataBuffers
 {
@@ -65,6 +113,9 @@ struct FNiagaraSimCacheFrame
 {
 	GENERATED_BODY()
 
+	UPROPERTY()
+	FTransform LocalToWorld;
+
 	UPROPERTY()
 	FNiagaraSimCacheSystemFrame SystemData;
 
@@ -104,6 +155,25 @@ struct FNiagaraSimCacheDataBuffersLayout
 {
 	GENERATED_BODY()
 
+	// Copy Function parameters are
+	// Dest, DestStride, Source, SourceStride, NumInstances, RebasedTransform
+	typedef void (*FVariableCopyFunction)(uint8*, uint32, const uint8*, uint32, uint32, const FTransform&);
+
+	struct FVariableCopyInfo
+	{
+		FVariableCopyInfo() = default;
+		explicit FVariableCopyInfo(uint16 InComponentFrom, uint16 InComponentTo, FVariableCopyFunction InCopyFunc)
+			: ComponentFrom(InComponentFrom)
+			, ComponentTo(InComponentTo)
+			, CopyFunc(InCopyFunc)
+		{
+		}
+
+		uint16					ComponentFrom = 0;
+		uint16					ComponentTo = 0;
+		FVariableCopyFunction	CopyFunc;
+	};
+
 	UPROPERTY()
 	FName LayoutName;
 
@@ -122,10 +192,12 @@ struct FNiagaraSimCacheDataBuffersLayout
 	UPROPERTY()
 	uint16 Int32Count = 0;
 
-	UPROPERTY(transient)
+	UPROPERTY()
+	TArray<FName> RebaseVariableNames;
+
 	TArray<uint16> ComponentMappingsToDataBuffer;
+	TArray<FVariableCopyInfo> VariableMappingsToDataBuffer;
 
-	UPROPERTY(transient)
 	TArray<uint16> ComponentMappingsFromDataBuffer;
 };
 
@@ -155,6 +227,9 @@ class NIAGARA_API UNiagaraSimCache : public UObject
 	UPROPERTY(VisibleAnywhere, Category=SimCache)
 	float DurationSeconds = 0.0f;
 
+	UPROPERTY()
+	FNiagaraSimCacheCreateParameters CreateParameters;
+
 	UPROPERTY(transient)
 	bool bNeedsReadComponentMappingRecache = true;
 
@@ -178,7 +253,7 @@ class NIAGARA_API UNiagaraSimCache : public UObject
 
 	bool IsCacheValid() const { return SoftNiagaraSystem.IsNull() == false; }
 
-	void BeginWrite(UNiagaraComponent* NiagaraComponent);
+	void BeginWrite(FNiagaraSimCacheCreateParameters InCreateParameters, UNiagaraComponent* NiagaraComponent);
 	void WriteFrame(UNiagaraComponent* NiagaraComponent);
 	void EndWrite();
 
@@ -186,12 +261,6 @@ class NIAGARA_API UNiagaraSimCache : public UObject
 	bool Read(float TimeSeconds, FNiagaraSystemInstance* SystemInstance) const;
 	bool ReadFrame(int32 FrameIndex, float FrameFraction, FNiagaraSystemInstance* SystemInstance) const;
 
-	/*
-	Create a simulation cache from the component's current data.
-	The cache can be applied to a component to warmup the system, for example, or can be used to debug the system.
-	*/
-	static UNiagaraSimCache* CreateSingleFrame(UObject* OuterObject, UNiagaraComponent* NiagaraComponent);
-
 private:
 	mutable std::atomic<int32> PendingCommandsInFlight;
 };
diff --git a/Engine/Plugins/FX/Niagara/Source/Niagara/Private/NiagaraDebuggerClient.cpp b/Engine/Plugins/FX/Niagara/Source/Niagara/Private/NiagaraDebuggerClient.cpp
--- a/Engine/Plugins/FX/Niagara/Source/Niagara/Private/NiagaraDebuggerClient.cpp
+++ b/Engine/Plugins/FX/Niagara/Source/Niagara/Private/NiagaraDebuggerClient.cpp
@@ -510,7 +510,7 @@ bool FNiagaraSimCacheCaptureInfo::Process()
 		if (ProcessedFrames == Request.CaptureDelayFrames)
 		{	
 			SimCache.Reset(NewObject<UNiagaraSimCache>(GetTransientPackage()));
-			SimCache->BeginWrite(Comp);
+			SimCache->BeginWrite(FNiagaraSimCacheCreateParameters(), Comp);
 		}
 
 		SimCache->WriteFrame(Comp);
diff --git a/Engine/Plugins/FX/Niagara/Source/Niagara/Private/NiagaraSimCache.cpp b/Engine/Plugins/FX/Niagara/Source/Niagara/Private/NiagaraSimCache.cpp
--- a/Engine/Plugins/FX/Niagara/Source/Niagara/Private/NiagaraSimCache.cpp
+++ b/Engine/Plugins/FX/Niagara/Source/Niagara/Private/NiagaraSimCache.cpp
@@ -3,6 +3,7 @@
 #include "NiagaraSimCache.h"
 
 #include "NiagaraClearCounts.h"
+#include "NiagaraConstants.h"
 #include "NiagaraComponent.h"
 #include "NiagaraComputeExecutionContext.h"
 #include "NiagaraDataInterfaceUtilities.h"
@@ -62,7 +63,7 @@ struct FNiagaraSimCacheHelper
 	bool HasValidSimulation() const { return SystemSimulation != nullptr; }
 	bool HasValidSimulationData() const { return SystemSimulationDataBuffer != nullptr; }
 
-	void BuildCacheLayout(FNiagaraSimCacheDataBuffersLayout& CacheLayout, const FNiagaraDataSetCompiledData& CompiledData, FName LayoutName) const
+	void BuildCacheLayout(FNiagaraSimCacheDataBuffersLayout& CacheLayout, const FNiagaraDataSetCompiledData& CompiledData, FName LayoutName, TArray<FName> InRebaseVariableNames) const
 	{
 		CacheLayout.LayoutName = LayoutName;
 		CacheLayout.SimTarget = CompiledData.SimTarget;
@@ -71,10 +72,9 @@ struct FNiagaraSimCacheHelper
 		CacheLayout.Variables.AddDefaulted(NumVariables);
 
 		const int32 CacheTotalComponents = CompiledData.TotalFloatComponents + CompiledData.TotalFloatComponents + CompiledData.TotalInt32Components;
-		CacheLayout.ComponentMappingsToDataBuffer.Empty(CacheTotalComponents);
-		CacheLayout.ComponentMappingsToDataBuffer.AddDefaulted(CacheTotalComponents);
 		CacheLayout.ComponentMappingsFromDataBuffer.Empty(CacheTotalComponents);
 		CacheLayout.ComponentMappingsFromDataBuffer.AddDefaulted(CacheTotalComponents);
+		CacheLayout.RebaseVariableNames = MoveTemp(InRebaseVariableNames);
 
 		for ( int32 iVariable=0; iVariable < NumVariables; ++iVariable)
 		{
@@ -94,6 +94,7 @@ struct FNiagaraSimCacheHelper
 			CacheLayout.Int32Count += DataSetVariableLayout.GetNumInt32Components();
 		}
 
+		// Build write mappings we will build read mappings in a separate path
 		int32 FloatOffset = 0;
 		int32 HalfOffset = CacheLayout.FloatCount;
 		int32 Int32Offset = HalfOffset + CacheLayout.HalfCount;
@@ -102,28 +103,152 @@ struct FNiagaraSimCacheHelper
 			const FNiagaraVariableLayoutInfo& DataSetVariableLayout = CompiledData.VariableLayouts[iVariable];
 			FNiagaraSimCacheVariable& CacheVariable = CacheLayout.Variables[iVariable];
 
-			// Build DataBuffer To / From mappings
 			for (int32 iComponent=0; iComponent < CacheVariable.FloatCount; ++iComponent)
 			{
-				CacheLayout.ComponentMappingsToDataBuffer[FloatOffset] = DataSetVariableLayout.FloatComponentStart + iComponent;
 				CacheLayout.ComponentMappingsFromDataBuffer[FloatOffset] = CacheVariable.FloatOffset + iComponent;
 				++FloatOffset;
 			}
 
 			for (int32 iComponent=0; iComponent < CacheVariable.HalfCount; ++iComponent)
 			{
-				CacheLayout.ComponentMappingsToDataBuffer[HalfOffset] = DataSetVariableLayout.HalfComponentStart + iComponent;
 				CacheLayout.ComponentMappingsFromDataBuffer[HalfOffset] = CacheVariable.HalfOffset + iComponent;
 				++HalfOffset;
 			}
 
 			for (int32 iComponent=0; iComponent < CacheVariable.Int32Count; ++iComponent)
 			{
-				CacheLayout.ComponentMappingsToDataBuffer[Int32Offset] = DataSetVariableLayout.Int32ComponentStart + iComponent;
 				CacheLayout.ComponentMappingsFromDataBuffer[Int32Offset] = CacheVariable.Int32Offset + iComponent;
 				++Int32Offset;
 			}
 		}
+
+		// Slightly inefficient but we can share the code between the paths
+		BuildCacheReadMappings(CacheLayout, CompiledData);
+	}
+
+	void BuildCacheLayoutForSystem(const FNiagaraSimCacheCreateParameters& CreateParmaeters, FNiagaraSimCacheDataBuffersLayout& CacheLayout)
+	{
+		const FNiagaraDataSetCompiledData& SystemCompileData = NiagaraSystem->GetSystemCompiledData().DataSetCompiledData;
+
+		TArray<FName> RebaseVariableNames;
+		if ( CreateParmaeters.bAllowRebasing )
+		{
+			TArray<FString, TInlineAllocator<8>> LocalSpaceEmitters;
+			for ( int32 i=0; i < NiagaraSystem->GetNumEmitters(); ++i )
+			{
+				const FNiagaraEmitterHandle& EmitterHandle = NiagaraSystem->GetEmitterHandle(i);
+				if ( EmitterHandle.GetInstance().GetEmitterData()->bLocalSpace )
+				{
+					LocalSpaceEmitters.Add(EmitterHandle.GetUniqueInstanceName());
+				}
+			}
+
+			for (const FNiagaraVariable& Variable : SystemCompileData.Variables)
+			{
+				if (Variable.GetType() == FNiagaraTypeDefinition::GetPositionDef())
+				{
+					// If this is an emitter variable we need to check if it's local space or not
+					bool bIsLocalSpace = false;
+					for ( const FString& LocalSpaceEmitter : LocalSpaceEmitters )
+					{
+						if ( Variable.IsInNameSpace(LocalSpaceEmitter) )
+						{
+							bIsLocalSpace = true;
+							break;
+						}
+					}
+
+					if ( bIsLocalSpace == false && CreateParmaeters.RebaseExcludeList.Contains(Variable.GetName()) == false )
+					{
+						RebaseVariableNames.AddUnique(Variable.GetName());
+					}
+				}
+				else if ( CanRebaseVariable(Variable) && CreateParmaeters.RebaseIncludeList.Contains(Variable.GetName()) )
+				{
+					RebaseVariableNames.AddUnique(Variable.GetName());
+				}
+			}
+		}
+
+		BuildCacheLayout(CacheLayout, SystemCompileData, NiagaraSystem->GetFName(), MoveTemp(RebaseVariableNames));
+	}
+
+	void BuildCacheLayoutForEmitter(const FNiagaraSimCacheCreateParameters& CreateParmaeters, FNiagaraSimCacheDataBuffersLayout& CacheLayout, int EmitterIndex)
+	{
+		const FNiagaraEmitterHandle& EmitterHandle = NiagaraSystem->GetEmitterHandle(EmitterIndex);
+		const FNiagaraEmitterCompiledData& EmitterCompiledData = NiagaraSystem->GetEmitterCompiledData()[EmitterIndex].Get();
+
+		// Find potential candidates for re-basing
+		TArray<FName> RebaseVariableNames;
+		if ( CreateParmaeters.bAllowRebasing && EmitterHandle.GetInstance().GetEmitterData()->bLocalSpace == false )
+		{
+			// Build list of include / exclude names
+			TArray<FName> ForceIncludeNames;
+			TArray<FName> ForceExcludeNames;
+			if ( CreateParmaeters.RebaseIncludeList.Num() > 0 || CreateParmaeters.RebaseExcludeList.Num() > 0 )
+			{
+				const FString EmitterName = EmitterHandle.GetUniqueInstanceName();
+				for (FName RebaseName : CreateParmaeters.RebaseIncludeList)
+				{
+					FNiagaraVariableBase BaseVar(FNiagaraTypeDefinition::GetFloatDef(), RebaseName);
+					if (BaseVar.RemoveRootNamespace(EmitterName))
+					{
+						ForceIncludeNames.Add(BaseVar.GetName());
+					}
+				}
+
+				for (FName RebaseName : CreateParmaeters.RebaseExcludeList)
+				{
+					FNiagaraVariableBase BaseVar(FNiagaraTypeDefinition::GetFloatDef(), RebaseName);
+					if (BaseVar.RemoveRootNamespace(EmitterName))
+					{
+						ForceExcludeNames.Add(BaseVar.GetName());
+					}
+				}
+			}
+
+		#if WITH_EDITORONLY_DATA
+			// Look for renderer attributes bound to Quat / Matrix types are we will want to rebase those
+			// We will add all Position types after this so no need to add them here
+			EmitterHandle.GetInstance().GetEmitterData()->ForEachEnabledRenderer(
+				[&](UNiagaraRendererProperties* RenderProperties)
+				{
+					for (FNiagaraVariableBase BoundAttribute : RenderProperties->GetBoundAttributes())
+					{
+						if ( (BoundAttribute.GetType() == FNiagaraTypeDefinition::GetQuatDef()) ||
+							 (BoundAttribute.GetType() == FNiagaraTypeDefinition::GetMatrix4Def()) )
+						{
+							if (BoundAttribute.RemoveRootNamespace(FNiagaraConstants::ParticleAttributeNamespaceString))
+							{
+								if (EmitterCompiledData.DataSetCompiledData.Variables.Contains(BoundAttribute) && ForceExcludeNames.Contains(BoundAttribute.GetName()) == false )
+								{
+									RebaseVariableNames.AddUnique(BoundAttribute.GetName());
+								}
+							}
+						}
+					}
+				}
+			);
+		#endif
+
+			// Look for regular attributes that we are forcing to rebase or can rebase like positions
+			for (const FNiagaraVariable& Variable : EmitterCompiledData.DataSetCompiledData.Variables)
+			{
+				if ( Variable.GetType() == FNiagaraTypeDefinition::GetPositionDef() )
+				{
+					if ( ForceExcludeNames.Contains(Variable.GetName()) == false )
+					{
+						RebaseVariableNames.AddUnique(Variable.GetName());
+					}
+				}
+				else if ( ForceIncludeNames.Contains(Variable.GetName()) && CanRebaseVariable(Variable) )
+				{
+					RebaseVariableNames.AddUnique(Variable.GetName());
+				}
+			}
+		}
+
+		BuildCacheLayout(CacheLayout, EmitterCompiledData.DataSetCompiledData, EmitterHandle.GetName(), MoveTemp(RebaseVariableNames));
 	}
 
 	static bool BuildCacheReadMappings(FNiagaraSimCacheDataBuffersLayout& CacheLayout, const FNiagaraDataSetCompiledData& CompiledData)
@@ -131,6 +256,7 @@ struct FNiagaraSimCacheHelper
 		const int32 CacheTotalComponents = CacheLayout.FloatCount + CacheLayout.HalfCount + CacheLayout.Int32Count;
 		CacheLayout.ComponentMappingsToDataBuffer.Empty(CacheTotalComponents);
 		CacheLayout.ComponentMappingsToDataBuffer.AddDefaulted(CacheTotalComponents);
+		CacheLayout.VariableMappingsToDataBuffer.Empty(0);
 
 		int32 FloatOffset = 0;
 		int32 HalfOffset = CacheLayout.FloatCount;
@@ -145,25 +271,51 @@ struct FNiagaraSimCacheHelper
 				DestVariableLayout = &CompiledData.VariableLayouts[DataSetVariableIndex];
 
 				// If the variable exists but types not match the cache is invalid
-				if ( (DestVariableLayout->GetNumFloatComponents() != SourceVariable.FloatCount) ||
-					 (DestVariableLayout->GetNumHalfComponents() != SourceVariable.HalfCount) ||
-					 (DestVariableLayout->GetNumInt32Components() != SourceVariable.Int32Count) )
+				if ((DestVariableLayout->GetNumFloatComponents() != SourceVariable.FloatCount) ||
+					(DestVariableLayout->GetNumHalfComponents() != SourceVariable.HalfCount) ||
+					(DestVariableLayout->GetNumInt32Components() != SourceVariable.Int32Count))
 				{
 					return false;
 				}
 			}
 
-			for ( int32 i=0; i < SourceVariable.FloatCount; ++i )
+			// Is this a type that requires conversion / re-basing?
+			if (DestVariableLayout != nullptr)
+			{
+				if ( CacheLayout.RebaseVariableNames.Contains(SourceVariable.Variable.GetName()) )
+				{
+					if (SourceVariable.Variable.GetType() == FNiagaraTypeDefinition::GetPositionDef())
+					{
+						check(SourceVariable.FloatCount == 3);
+						CacheLayout.VariableMappingsToDataBuffer.Emplace(FloatOffset, DestVariableLayout->FloatComponentStart, &FNiagaraSimCacheHelper::CopyPositions);
+						DestVariableLayout = nullptr;
+					}
+					else if (SourceVariable.Variable.GetType() == FNiagaraTypeDefinition::GetQuatDef())
+					{
+						check(SourceVariable.FloatCount == 4);
+						CacheLayout.VariableMappingsToDataBuffer.Emplace(FloatOffset, DestVariableLayout->FloatComponentStart, &FNiagaraSimCacheHelper::CopyQuaternions);
+						DestVariableLayout = nullptr;
+					}
+					else if (SourceVariable.Variable.GetType() == FNiagaraTypeDefinition::GetMatrix4Def())
+					{
+						check(SourceVariable.FloatCount == 16);
+						CacheLayout.VariableMappingsToDataBuffer.Emplace(FloatOffset, DestVariableLayout->FloatComponentStart, &FNiagaraSimCacheHelper::CopyMatrices);
+						DestVariableLayout = nullptr;
+					}
+				}
+			}
+
+			for (int32 i = 0; i < SourceVariable.FloatCount; ++i)
 			{
 				CacheLayout.ComponentMappingsToDataBuffer[FloatOffset++] = DestVariableLayout ? DestVariableLayout->FloatComponentStart + i : INDEX_NONE;
 			}
 
-			for ( int32 i=0; i < SourceVariable.HalfCount; ++i)
+			for (int32 i = 0; i < SourceVariable.HalfCount; ++i)
 			{
 				CacheLayout.ComponentMappingsToDataBuffer[HalfOffset++] = DestVariableLayout ? DestVariableLayout->HalfComponentStart + i : INDEX_NONE;
 			}
 
-			for ( int32 i=0; i < SourceVariable.Int32Count; ++i)
+			for (int32 i = 0; i < SourceVariable.Int32Count; ++i)
 			{
 				CacheLayout.ComponentMappingsToDataBuffer[Int32Offset++] = DestVariableLayout ? DestVariableLayout->Int32ComponentStart + i : INDEX_NONE;
 			}
@@ -239,7 +391,7 @@ struct FNiagaraSimCacheHelper
 		}
 	}
 
-	void ReadDataBuffer(const FNiagaraSimCacheDataBuffersLayout& CacheLayout, const FNiagaraSimCacheDataBuffers& CacheBuffer, FNiagaraDataSet& DataSet)
+	void ReadDataBuffer(const FTransform& RebaseTransform, const FNiagaraSimCacheDataBuffersLayout& CacheLayout, const FNiagaraSimCacheDataBuffers& CacheBuffer, FNiagaraDataSet& DataSet)
 	{
 		FNiagaraDataBuffer& DataBuffer = DataSet.BeginSimulate();
 		DataBuffer.Allocate(CacheBuffer.NumInstances);
@@ -290,6 +442,15 @@ struct FNiagaraSimCacheHelper
 				CheckedMemcpy(DataBuffer.GetInt32Buffer(), Dest, CacheBuffer.Int32Data, Source, sizeof(int32) * NumInstances);
 				//FMemory::Memcpy(Dest, Source, sizeof(int32) * NumInstances);
 			}
+
+			// Copy variables that require processing
+			for ( const FNiagaraSimCacheDataBuffersLayout::FVariableCopyInfo& VariableCopyInfo : CacheLayout.VariableMappingsToDataBuffer )
+			{
+				const uint32 SrcStride = uint32(NumInstances) * sizeof(float);
+				const uint8* Src = CacheBuffer.FloatData.GetData() + (VariableCopyInfo.ComponentFrom * SrcStride);
+				uint8* Dst = DataBuffer.GetComponentPtrFloat(VariableCopyInfo.ComponentTo);
+				VariableCopyInfo.CopyFunc(Dst, DataBuffer.GetFloatStride(), Src, SrcStride, uint32(NumInstances), RebaseTransform);
+			}
 		}
 
 		//-TODO:DestinationDataBuffer.SetIDTable(CacheBuffer.IDToIndexTable);
@@ -298,7 +459,7 @@ struct FNiagaraSimCacheHelper
 		DataSet.EndSimulate();
 	}
 
-	void ReadDataBufferGPU(FNiagaraEmitterInstance& EmitterInstance, const FNiagaraSimCacheDataBuffersLayout& InCacheLayout, const FNiagaraSimCacheDataBuffers& InCacheBuffer, FNiagaraDataSet& InDataSet, std::atomic<int32>& InPendingCommandsCounter)
+	void ReadDataBufferGPU(const FTransform& InRebaseTransform, FNiagaraEmitterInstance& EmitterInstance, const FNiagaraSimCacheDataBuffersLayout& InCacheLayout, const FNiagaraSimCacheDataBuffers& InCacheBuffer, FNiagaraDataSet& InDataSet, std::atomic<int32>& InPendingCommandsCounter)
 	{
 		if (EmitterInstance.IsDisabled())
 		{
@@ -311,7 +472,7 @@ struct FNiagaraSimCacheHelper
 
 		FNiagaraGpuComputeDispatchInterface* DispathInterface = EmitterInstance.GetParentSystemInstance()->GetComputeDispatchInterface();
 		ENQUEUE_RENDER_COMMAND(NiagaraSimCacheGpuReadFrame)(
-			[DispathInterface, GPUExecContext=EmitterInstance.GetGPUContext(), CacheLayout=&InCacheLayout, CacheBuffer=&InCacheBuffer, DataSet=&InDataSet, PendingCommandsCounter=&InPendingCommandsCounter](FRHICommandListImmediate& RHICmdList)
+			[DispathInterface, GPUExecContext=EmitterInstance.GetGPUContext(), RebaseTransform=InRebaseTransform, CacheLayout=&InCacheLayout, CacheBuffer=&InCacheBuffer, DataSet=&InDataSet, PendingCommandsCounter=&InPendingCommandsCounter](FRHICommandListImmediate& RHICmdList)
 			{
 				const int32 NumInstances = CacheBuffer->NumInstances;
 
@@ -361,6 +522,15 @@ struct FNiagaraSimCacheHelper
 							//FMemory::Memcpy(Dest, Source, sizeof(float) * NumInstances);
 						}
 
+						// Copy variables that require processing
+						for (const FNiagaraSimCacheDataBuffersLayout::FVariableCopyInfo& VariableCopyInfo : CacheLayout->VariableMappingsToDataBuffer)
+						{
+							const uint32 SrcStride = uint32(NumInstances) * sizeof(float);
+							const uint8* Src = CacheBuffer->FloatData.GetData() + (VariableCopyInfo.ComponentFrom * SrcStride);
+							uint8* Dst = RWBufferMemory + (uint32(VariableCopyInfo.ComponentTo) * RWComponentStride);
+							VariableCopyInfo.CopyFunc(Dst, DataBuffer.GetFloatStride(), Src, SrcStride, uint32(NumInstances), RebaseTransform);
+						}
+
 						RHIUnlockBuffer(RWBuffer.Buffer);
 					}
 
@@ -421,6 +591,86 @@ struct FNiagaraSimCacheHelper
 		);
 	}
 
+	static bool CanRebaseVariable(const FNiagaraVariableBase& Variable)
+	{
+		return	
+			(Variable.GetType() == FNiagaraTypeDefinition::GetQuatDef()) ||
+			(Variable.GetType() == FNiagaraTypeDefinition::GetMatrix4Def()) ||
+			(Variable.GetType() == FNiagaraTypeDefinition::GetPositionDef());
+	}
+
+	static void CopyPositions(uint8* Dst, uint32 DstStride, const uint8* Src, uint32 SrcStride, uint32 NumInstances, const FTransform& RebaseTransform)
+	{
+		float* DstFloats = reinterpret_cast<float*>(Dst);
+		DstStride = DstStride >> 2;
+
+		const float* SrcFloats = reinterpret_cast<const float*>(Src);
+		SrcStride = SrcStride >> 2;
+
+		for (uint32 i = 0; i < NumInstances; ++i)
+		{
+			const FVector CachePosition(
+				SrcFloats[i + (SrcStride * 0)],
+				SrcFloats[i + (SrcStride * 1)],
+				SrcFloats[i + (SrcStride * 2)]
+			);
+			const FVector RebasedPosition = RebaseTransform.TransformPosition(CachePosition);
+			DstFloats[i + (DstStride * 0)] = RebasedPosition.X;
+			DstFloats[i + (DstStride * 1)] = RebasedPosition.Y;
+			DstFloats[i + (DstStride * 2)] = RebasedPosition.Z;
+		}
+	}
+
+	static void CopyQuaternions(uint8* Dst, uint32 DstStride, const uint8* Src, uint32 SrcStride, uint32 NumInstances, const FTransform& RebaseTransform)
+	{
+		float* DstFloats = reinterpret_cast<float*>(Dst);
+		DstStride = DstStride >> 2;
+
+		const float* SrcFloats = reinterpret_cast<const float*>(Src);
+		SrcStride = SrcStride >> 2;
+
+		for (uint32 i = 0; i < NumInstances; ++i)
+		{
+			const FQuat4f CacheRotation(
+				SrcFloats[i + (SrcStride * 0)],
+				SrcFloats[i + (SrcStride * 1)],
+				SrcFloats[i + (SrcStride * 2)],
+				SrcFloats[i + (SrcStride * 3)]
+			);
+			const FQuat4f RebasedQuat = CacheRotation * FQuat4f(RebaseTransform.GetRotation());
+			DstFloats[i + (DstStride * 0)] = RebasedQuat.X;
+			DstFloats[i + (DstStride * 1)] = RebasedQuat.Y;
+			DstFloats[i + (DstStride * 2)] = RebasedQuat.Z;
+			DstFloats[i + (DstStride * 3)] = RebasedQuat.W;
+		}
+	}
+
+	static void CopyMatrices(uint8* Dst, uint32 DstStride, const uint8* Src, uint32 SrcStride, uint32 NumInstances, const FTransform& RebaseTransform)
+	{
+		float* DstFloats = reinterpret_cast<float*>(Dst);
+		DstStride = DstStride >> 2;
+
+		const float* SrcFloats = reinterpret_cast<const float*>(Src);
+		SrcStride = SrcStride >> 2;
+
+		const FMatrix44d RebaseMatrix = RebaseTransform.ToMatrixWithScale();
+		for (uint32 i = 0; i < NumInstances; ++i)
+		{
+			FMatrix44d CacheMatrix;
+			for (int32 j = 0; j < 16; ++j)
+			{
+				CacheMatrix.M[j >> 2][j & 0x3] = double(SrcFloats[i + (SrcStride * j)]);
+			}
+
+			CacheMatrix = CacheMatrix * RebaseMatrix;
+
+			for (int32 j = 0; j < 16; ++j)
+			{
+				DstFloats[i + (DstStride * j)] = float(CacheMatrix.M[j >> 2][j & 0x3]);
+			}
+		}
+	}
+
 	UNiagaraSystem*						NiagaraSystem = nullptr;
 	FNiagaraSystemInstance*				SystemInstance = nullptr;
 	FNiagaraSystemSimulationPtr			SystemSimulation = nullptr;
@@ -440,7 +690,7 @@ bool UNiagaraSimCache::IsReadyForFinishDestroy()
 	return PendingCommandsInFlight == 0;
 }
 
-void UNiagaraSimCache::BeginWrite(UNiagaraComponent* NiagaraComponent)
+void UNiagaraSimCache::BeginWrite(FNiagaraSimCacheCreateParameters InCreateParameters, UNiagaraComponent* NiagaraComponent)
 {
 	check(PendingCommandsInFlight == 0);
 
@@ -450,37 +700,42 @@ void UNiagaraSimCache::BeginWrite(UNiagaraComponent* NiagaraComponent)
 		return;
 	}
 
+	Modify();
+
 	// Reset to defaults
 	SoftNiagaraSystem = Helper.NiagaraSystem;
+	CreateParameters = InCreateParameters;
 	StartSeconds = 0.0f;
 	DurationSeconds = 0.0f;
 	CacheLayout = FNiagaraSimCacheLayout();
 	CacheFrames.Empty();
 
 	// Build new layout for system / emitters
-	Helper.BuildCacheLayout(CacheLayout.SystemLayout, Helper.NiagaraSystem->GetSystemCompiledData().DataSetCompiledData, Helper.NiagaraSystem->GetFName());
+	Helper.BuildCacheLayoutForSystem(CreateParameters, CacheLayout.SystemLayout);
 
 	const int32 NumEmitters = Helper.NiagaraSystem->GetEmitterHandles().Num();
 	CacheLayout.EmitterLayouts.AddDefaulted(NumEmitters);
 	for ( int32 i=0; i < NumEmitters; ++i )
 	{
-		const FNiagaraEmitterCompiledData& EmitterCompiledData = Helper.NiagaraSystem->GetEmitterCompiledData()[i].Get();
-		Helper.BuildCacheLayout(CacheLayout.EmitterLayouts[i], EmitterCompiledData.DataSetCompiledData, Helper.NiagaraSystem->GetEmitterHandle(i).GetName());
+		Helper.BuildCacheLayoutForEmitter(CreateParameters, CacheLayout.EmitterLayouts[i], i);
 	}
 
 	// Find data interfaces we may want to cache
-	FNiagaraDataInterfaceUtilities::ForEachDataInterface(
-		Helper.SystemInstance,
-		[&](const FNiagaraVariableBase& Variable, UNiagaraDataInterface* DataInterface)
-		{
-			const void* PerInstanceData = Helper.SystemInstance->FindDataInterfaceInstanceData(DataInterface);
-			if (UObject* DICacheStorage = DataInterface->SimCacheBeginWrite(this, Helper.SystemInstance, PerInstanceData))
+	if ( CreateParameters.bAllowDataInterfaceCaching )
+	{
+		FNiagaraDataInterfaceUtilities::ForEachDataInterface(
+			Helper.SystemInstance,
+			[&](const FNiagaraVariableBase& Variable, UNiagaraDataInterface* DataInterface)
 			{
-				DataInterfaceStorage.FindOrAdd(Variable) = DICacheStorage;
+				const void* PerInstanceData = Helper.SystemInstance->FindDataInterfaceInstanceData(DataInterface);
+				if (UObject* DICacheStorage = DataInterface->SimCacheBeginWrite(this, Helper.SystemInstance, PerInstanceData))
+				{
+					DataInterfaceStorage.FindOrAdd(Variable) = DICacheStorage;
+				}
+				return true;
 			}
-			return true;
-		}
-	);
+		);
+	}
 }
 
 void UNiagaraSimCache::WriteFrame(UNiagaraComponent* NiagaraComponent)
@@ -527,6 +782,8 @@ void UNiagaraSimCache::WriteFrame(UNiagaraComponent* NiagaraComponent)
 
 	// Cache frame
 	FNiagaraSimCacheFrame& CacheFrame = CacheFrames.AddDefaulted_GetRef();
+	CacheFrame.LocalToWorld = Helper.SystemInstance->GatheredInstanceParameters.ComponentTrans;
+
 	CacheFrame.SystemData.LocalBounds = Helper.SystemInstance->GetLocalBounds();
 
 	const int32 NumEmitters = CacheLayout.EmitterLayouts.Num();
@@ -710,8 +967,17 @@ bool UNiagaraSimCache::ReadFrame(int32 FrameIndex, float FrameFraction, FNiagara
 	}
 
 	const FNiagaraSimCacheFrame& CacheFrame = CacheFrames[FrameIndex];
+
+	FTransform RebaseTransform = FTransform::Identity;
+	if ( USceneComponent* AttachComponent = SystemInstance->GetAttachComponent() )
+	{
+		RebaseTransform = AttachComponent->GetComponentToWorld();
+		RebaseTransform.AddToTranslation(FVector(SystemInstance->GetLWCTile()) * -FLargeWorldRenderScalar::GetTileSize());
+		RebaseTransform = CacheFrame.LocalToWorld * RebaseTransform;
+	}
+
 	Helper.SystemInstance->LocalBounds = CacheFrame.SystemData.LocalBounds;
-	Helper.ReadDataBuffer(CacheLayout.SystemLayout, CacheFrame.SystemData.SystemDataBuffers, Helper.GetSystemSimulationDataSet());
+	Helper.ReadDataBuffer(RebaseTransform, CacheLayout.SystemLayout, CacheFrame.SystemData.SystemDataBuffers, Helper.GetSystemSimulationDataSet());
 
 	const int32 NumEmitters = CacheLayout.EmitterLayouts.Num();
 	for (int32 i=0; i < NumEmitters; ++i)
@@ -720,13 +986,14 @@ bool UNiagaraSimCache::ReadFrame(int32 FrameIndex, float FrameFraction, FNiagara
 		FNiagaraEmitterInstance& EmitterInstance = Helper.SystemInstance->GetEmitters()[i].Get();
 		EmitterInstance.CachedBounds = CacheEmitterFrame.LocalBounds;
 		EmitterInstance.TotalSpawnedParticles = CacheEmitterFrame.TotalSpawnedParticles;
+
 		if (CacheLayout.EmitterLayouts[i].SimTarget == ENiagaraSimTarget::GPUComputeSim)
 		{
-			Helper.ReadDataBufferGPU(EmitterInstance, CacheLayout.EmitterLayouts[i], CacheEmitterFrame.ParticleDataBuffers, EmitterInstance.GetData(), PendingCommandsInFlight);
+			Helper.ReadDataBufferGPU(RebaseTransform, EmitterInstance, CacheLayout.EmitterLayouts[i], CacheEmitterFrame.ParticleDataBuffers, EmitterInstance.GetData(), PendingCommandsInFlight);
 		}
 		else
 		{
-			Helper.ReadDataBuffer(CacheLayout.EmitterLayouts[i], CacheEmitterFrame.ParticleDataBuffers, EmitterInstance.GetData());
+			Helper.ReadDataBuffer(RebaseTransform, CacheLayout.EmitterLayouts[i], CacheEmitterFrame.ParticleDataBuffers, EmitterInstance.GetData());
 		}
 	}
 
@@ -767,12 +1034,3 @@ bool UNiagaraSimCache::ReadFrame(int32 FrameIndex, float FrameFraction, FNiagara
 	}
 	return true;
 }
-
-UNiagaraSimCache* UNiagaraSimCache::CreateSingleFrame(UObject* OuterObject, UNiagaraComponent* NiagaraComponent)
-{
-	UNiagaraSimCache* SimCache = NewObject<UNiagaraSimCache>(OuterObject);
-	SimCache->BeginWrite(NiagaraComponent);
-	SimCache->WriteFrame(NiagaraComponent);
-	SimCache->EndWrite();
-	return SimCache;
-}
diff --git a/Engine/Plugins/FX/Niagara/Source/Niagara/Private/NiagaraTypes.cpp b/Engine/Plugins/FX/Niagara/Source/Niagara/Private/NiagaraTypes.cpp
--- a/Engine/Plugins/FX/Niagara/Source/Niagara/Private/NiagaraTypes.cpp
+++ b/Engine/Plugins/FX/Niagara/Source/Niagara/Private/NiagaraTypes.cpp
@@ -19,6 +19,39 @@ void FNiagaraVariableBase::SetNamespacedName(const FString& InNamespace, FName I
 	Name = FName(NameBuilder.ToString());
 }
 
+bool FNiagaraVariableBase::RemoveRootNamespace(const FStringView& ExpectedNamespace)
+{
+	FNameBuilder NameString;
+	Name.ToString(NameString);
+
+	FStringView NameStringView = NameString.ToView();
+	if ( (NameStringView.Len() > ExpectedNamespace.Len() + 1) && (NameStringView[ExpectedNamespace.Len()] == '.') && NameStringView.StartsWith(ExpectedNamespace) )
+	{
+		FNameBuilder NewNameString;
+		NewNameString.Append(NameStringView.Mid(ExpectedNamespace.Len() + 1));
+		Name = FName(FStringView(NewNameString));
+		return true;
+	}
+	return false;
+}
+
+bool FNiagaraVariableBase::ReplaceRootNamespace(const FStringView& ExpectedNamespace, const FStringView& NewNamespace)
+{
+	FNameBuilder NameString;
+	Name.ToString(NameString);
+
+	FStringView NameStringView = NameString.ToView();
+	if ( (NameStringView.Len() > ExpectedNamespace.Len() + 1) && (NameStringView[ExpectedNamespace.Len()] == '.') && NameStringView.StartsWith(ExpectedNamespace) )
+	{
+		FNameBuilder NewNameString;
+		NewNameString.Append(NewNamespace);
+		NewNameString.Append(NameStringView.Mid(ExpectedNamespace.Len() + 1));
+		Name = FName(FStringView(NewNameString));
+		return true;
+	}
+	return false;
+}
+
 void FNiagaraVariableMetaData::CopyUserEditableMetaData(const FNiagaraVariableMetaData& OtherMetaData)
 {
 	for (const FProperty* ChildProperty : TFieldRange<FProperty>(StaticStruct()))
diff --git a/Engine/Plugins/FX/Niagara/Source/Niagara/Public/NiagaraTypes.h b/Engine/Plugins/FX/Niagara/Source/Niagara/Public/NiagaraTypes.h
--- a/Engine/Plugins/FX/Niagara/Source/Niagara/Public/NiagaraTypes.h
+++ b/Engine/Plugins/FX/Niagara/Source/Niagara/Public/NiagaraTypes.h
@@ -1743,6 +1743,20 @@ struct FNiagaraVariableBase
 		return Name != NAME_None && TypeDefHandle->IsValid();
 	}
 
+	/**
+	Removes the root namespace from the variable if it matches the expected namespace.
+	i.e. MyNamespace.VarName would become VarName but Other.MyNamespace.VarName would not be modified.
+	For more complex namespace manipulation refer to FNiagaraUtilities::ResolveAliases
+	*/
+	bool RemoveRootNamespace(const FStringView& ExpectedNamespace);
+
+	/**
+	Replaces the root namespace from the variable if it matches the expected namespace with the NewNamespace.
+	i.e. Prev.VarName would become Next.VarName but Other.Prev.VarName would not be modified.
+	For more complex namespace manipulation refer to FNiagaraUtilities::ResolveAliases
+	*/
+	bool ReplaceRootNamespace(const FStringView& ExpectedNamespace, const FStringView& NewNamespace);
+
 	FORCEINLINE bool IsInNameSpace(const FStringView& Namespace) const
 	{
 		TStringBuilder<128> NameString;
diff --git a/Engine/Plugins/FX/Niagara/Source/NiagaraEditor/Private/NiagaraBakerRendererOutputSimCache.cpp b/Engine/Plugins/FX/Niagara/Source/NiagaraEditor/Private/NiagaraBakerRendererOutputSimCache.cpp
--- a/Engine/Plugins/FX/Niagara/Source/NiagaraEditor/Private/NiagaraBakerRendererOutputSimCache.cpp
+++ b/Engine/Plugins/FX/Niagara/Source/NiagaraEditor/Private/NiagaraBakerRendererOutputSimCache.cpp
@@ -73,7 +73,7 @@ void FNiagaraBakerRendererOutputSimCache::BakeFrame(UNiagaraBakerOutput* InBaker
 
 	if ( FrameIndex == 0 )
 	{
-		BakeSimCache->BeginWrite(BakerRenderer.GetPreviewComponent());
+		BakeSimCache->BeginWrite(BakerOutput->CreateParameters, BakerRenderer.GetPreviewComponent());
 	}
 	BakeSimCache->WriteFrame(BakerRenderer.GetPreviewComponent());
 }
