- Take DPI scale into account for ViewRect calculation

#rnx

diff --git a/Engine/Plugins/FX/Niagara/Source/NiagaraEditor/Private/Widgets/SNiagaraBakerViewport.cpp b/Engine/Plugins/FX/Niagara/Source/NiagaraEditor/Private/Widgets/SNiagaraBakerViewport.cpp
--- a/Engine/Plugins/FX/Niagara/Source/NiagaraEditor/Private/Widgets/SNiagaraBakerViewport.cpp
+++ b/Engine/Plugins/FX/Niagara/Source/NiagaraEditor/Private/Widgets/SNiagaraBakerViewport.cpp
@@ -269,8 +269,13 @@ public:
 			}
 		}
 
-		//const FIntRect ViewRect(Canvas->GetViewRect().Min + FIntPoint(2, 2), Canvas->GetViewRect().Max - FIntPoint(2, 2));
-		const FIntRect ViewRect(FIntPoint(2, 2), InViewport->GetSizeXY() - FIntPoint(2, 2));
+		const float DPIScaleFactor = ShouldDPIScaleSceneCanvas() ? GetDPIScale() : 1.0f;
+		const FIntRect ViewRect(
+			FMath::Floor(Canvas->GetViewRect().Min.X / DPIScaleFactor) + 2,
+			FMath::Floor(Canvas->GetViewRect().Min.Y / DPIScaleFactor) + 2,
+			FMath::Floor(Canvas->GetViewRect().Max.X / DPIScaleFactor) - 2,
+			FMath::Floor(Canvas->GetViewRect().Max.Y / DPIScaleFactor) - 2
+		);
 		if ( ViewRect.Width() <= 0 || ViewRect.Height() <= 0 )
 		{
 			return;
