Adding Timeline for Sim Cache UI

diff --git a/Engine/Plugins/FX/Niagara/Source/NiagaraEditor/Private/Toolkits/NiagaraSimCacheToolkit.cpp b/Engine/Plugins/FX/Niagara/Source/NiagaraEditor/Private/Toolkits/NiagaraSimCacheToolkit.cpp
--- a/Engine/Plugins/FX/Niagara/Source/NiagaraEditor/Private/Toolkits/NiagaraSimCacheToolkit.cpp
+++ b/Engine/Plugins/FX/Niagara/Source/NiagaraEditor/Private/Toolkits/NiagaraSimCacheToolkit.cpp
@@ -6,12 +6,14 @@
 // Widgets & ViewModels
 #include "Widgets/Docking/SDockTab.h"
 #include "NiagaraSimCacheViewModel.h"
+#include "SNiagaraSimCacheViewTimeline.h"
+#include "SNiagaraSimCacheViewTransportControls.h"
 
 #define LOCTEXT_NAMESPACE "NiagaraSimCacheToolkit"
 
 const FName FNiagaraSimCacheToolkit::NiagaraSimCacheSpreadsheetTabId(TEXT("NiagaraSimCacheEditor_Spreadsheet"));
 
-//TODO Timeline, & Preview
+//TODO Preview
 
 FNiagaraSimCacheToolkit::FNiagaraSimCacheToolkit()
 {
@@ -104,7 +106,31 @@ TSharedRef<SDockTab> FNiagaraSimCacheToolkit::SpawnTab_SimCacheSpreadsheet(const
 	Contents->AddSlot()
 		.FillHeight(1.0f)
 		[
-			SimCacheSpreadsheetView->AsShared()
+			SNew(SVerticalBox)
+			+SVerticalBox::Slot()
+			[
+				SimCacheSpreadsheetView->AsShared()
+			]
+			+SVerticalBox::Slot()
+			.AutoHeight()
+			[
+				SNew(SHorizontalBox)
+				+SHorizontalBox::Slot()
+				.AutoWidth()
+				.Padding(2.0f)
+				[
+					SNew(SNiagaraSimCacheViewTransportControls)
+					.WeakViewModel(SimCacheViewModel->AsWeak())
+				]
+				+SHorizontalBox::Slot()
+				.Padding(4.0f)
+				[
+					SNew(SNiagaraSimCacheViewTimeline)
+					.WeakViewModel(SimCacheViewModel->AsWeak())
+				]
+				
+			]
+			
 		];
 
 	TSharedRef<SDockTab> SpawnedTab =
diff --git a/Engine/Plugins/FX/Niagara/Source/NiagaraEditor/Private/ViewModels/NiagaraSimCacheViewModel.cpp b/Engine/Plugins/FX/Niagara/Source/NiagaraEditor/Private/ViewModels/NiagaraSimCacheViewModel.cpp
--- a/Engine/Plugins/FX/Niagara/Source/NiagaraEditor/Private/ViewModels/NiagaraSimCacheViewModel.cpp
+++ b/Engine/Plugins/FX/Niagara/Source/NiagaraEditor/Private/ViewModels/NiagaraSimCacheViewModel.cpp
@@ -272,6 +272,12 @@ int32 FNiagaraSimCacheViewModel::GetNumFrames() const
 	return SimCache ? SimCache->CacheFrames.Num() : 0;
 }
 
+void FNiagaraSimCacheViewModel::SetFrameIndex(const int32 InFrameIndex)
+{
+	FrameIndex = InFrameIndex;
+	OnFrameUpdatedDelegate.Broadcast(true);
+}
+
 
 void FNiagaraSimCacheViewModel::SetEmitterIndex(const int32 InEmitterIndex)
 {
@@ -301,4 +307,9 @@ FName FNiagaraSimCacheViewModel::GetEmitterLayoutName(const int32 Index)
 	return SimCache ? SimCache->CacheLayout.EmitterLayouts[Index].LayoutName : NAME_None;
 }
 
-#undef LOCTEXT_NAMESPACE
\ No newline at end of file
+FNiagaraSimCacheViewModel::FOnFrameUpdated& FNiagaraSimCacheViewModel::OnFrameUpdated()
+{
+	return OnFrameUpdatedDelegate;
+}
+
+#undef LOCTEXT_NAMESPACE
diff --git a/Engine/Plugins/FX/Niagara/Source/NiagaraEditor/Private/ViewModels/NiagaraSimCacheViewModel.h b/Engine/Plugins/FX/Niagara/Source/NiagaraEditor/Private/ViewModels/NiagaraSimCacheViewModel.h
--- a/Engine/Plugins/FX/Niagara/Source/NiagaraEditor/Private/ViewModels/NiagaraSimCacheViewModel.h
+++ b/Engine/Plugins/FX/Niagara/Source/NiagaraEditor/Private/ViewModels/NiagaraSimCacheViewModel.h
@@ -22,6 +22,8 @@ struct FNiagaraSimCacheComponentInfo
 class NIAGARAEDITOR_API FNiagaraSimCacheViewModel : public TSharedFromThis<FNiagaraSimCacheViewModel>
 {
 public:
+	DECLARE_MULTICAST_DELEGATE_OneParam(FOnFrameUpdated, bool)
+	
 	FNiagaraSimCacheViewModel();
 	~FNiagaraSimCacheViewModel();
 
@@ -37,7 +39,7 @@ public:
 
 	int32 GetFrameIndex() const { return FrameIndex; };
 
-	void SetFrameIndex(const int32 InFrameIndex) { FrameIndex = InFrameIndex; };
+	void SetFrameIndex(const int32 InFrameIndex);;
 
 	int32 GetEmitterIndex() const { return EmitterIndex; };
 
@@ -53,7 +55,11 @@ public:
 
 	FName GetEmitterLayoutName(int32 Index);
 
+	FOnFrameUpdated& OnFrameUpdated();
+
 private:
 	uint32 FrameIndex = 0;
 	uint32 EmitterIndex = INDEX_NONE;
+
+	FOnFrameUpdated OnFrameUpdatedDelegate;
 };
\ No newline at end of file
diff --git a/Engine/Plugins/FX/Niagara/Source/NiagaraEditor/Private/Widgets/SNiagaraSimCacheView.cpp b/Engine/Plugins/FX/Niagara/Source/NiagaraEditor/Private/Widgets/SNiagaraSimCacheView.cpp
--- a/Engine/Plugins/FX/Niagara/Source/NiagaraEditor/Private/Widgets/SNiagaraSimCacheView.cpp
+++ b/Engine/Plugins/FX/Niagara/Source/NiagaraEditor/Private/Widgets/SNiagaraSimCacheView.cpp
@@ -67,6 +67,8 @@ void SNiagaraSimCacheView::Construct(const FArguments& InArgs)
 {
 	SimCacheViewModel = InArgs._SimCacheViewModel;
 
+	SimCacheViewModel->OnFrameUpdated().AddSP(this, &SNiagaraSimCacheView::OnSimCacheFrameChanged);
+
 	HeaderRowWidget = SNew(SHeaderRow);
 
 	UpdateColumns(false);
@@ -117,27 +119,6 @@ void SNiagaraSimCacheView::Construct(const FArguments& InArgs)
 					.Text(this, &SNiagaraSimCacheView::GetBufferSelectionText)
 				]
 			]
-			// Select Cache Frame
-			+ SHorizontalBox::Slot()
-			.AutoWidth()
-			.Padding(10.0f)
-			.VAlign(VAlign_Center)
-			.HAlign(HAlign_Right)
-			[
-				SNew(STextBlock)
-				.Text(LOCTEXT("FrameIndex", "Frame Index"))
-			]
-			+ SHorizontalBox::Slot()
-			.AutoWidth()
-			.VAlign(VAlign_Center)
-			.HAlign(HAlign_Left)
-			[
-				SNew(SSpinBox<int32>)
-				.MinValue(0)
-				.MaxValue(this, &SNiagaraSimCacheView::GetNumFrames)
-				.Value(this, &SNiagaraSimCacheView::GetFrameIndex)
-				.OnValueChanged(this, &SNiagaraSimCacheView::SetFrameIndex)
-			]
 			// Component Filter
 			+ SHorizontalBox::Slot()
 			.AutoWidth()
@@ -295,22 +276,6 @@ void SNiagaraSimCacheView::UpdateBufferSelectionList()
 	}
 }
 
-TOptional<int32> SNiagaraSimCacheView::GetNumFrames() const
-{
-	return SimCacheViewModel->GetNumFrames();
-}
-
-int32 SNiagaraSimCacheView::GetFrameIndex() const
-{
-	return SimCacheViewModel->GetFrameIndex();
-}
-
-void SNiagaraSimCacheView::SetFrameIndex(int32 Value)
-{
-	SimCacheViewModel->SetFrameIndex(Value);
-	UpdateRows(true);
-}
-
 TSharedRef<SWidget> SNiagaraSimCacheView::BufferSelectionGenerateWidget(TSharedPtr<FBufferSelectionInfo> InItem)
 {
 	return
@@ -358,4 +323,9 @@ void SNiagaraSimCacheView::OnSimCacheChanged(const FAssetData& InAsset)
 	}
 }
 
-#undef LOCTEXT_NAMESPACE
\ No newline at end of file
+void SNiagaraSimCacheView::OnSimCacheFrameChanged(const bool bRefresh)
+{
+	UpdateRows(bRefresh);
+}
+
+#undef LOCTEXT_NAMESPACE
diff --git a/Engine/Plugins/FX/Niagara/Source/NiagaraEditor/Private/Widgets/SNiagaraSimCacheView.h b/Engine/Plugins/FX/Niagara/Source/NiagaraEditor/Private/Widgets/SNiagaraSimCacheView.h
--- a/Engine/Plugins/FX/Niagara/Source/NiagaraEditor/Private/Widgets/SNiagaraSimCacheView.h
+++ b/Engine/Plugins/FX/Niagara/Source/NiagaraEditor/Private/Widgets/SNiagaraSimCacheView.h
@@ -28,12 +28,7 @@ public:
 	void UpdateRows(const bool bRefresh);
 
 	void UpdateBufferSelectionList();
-
-	TOptional<int32> GetNumFrames() const;
-	int32 GetFrameIndex() const;
-
-	void SetFrameIndex(int32);
-
+	
 	TSharedRef<SWidget> BufferSelectionGenerateWidget(TSharedPtr<FBufferSelectionInfo> InItem);
 
 	void BufferSelectionChanged(TSharedPtr<FBufferSelectionInfo> NewSelection, ESelectInfo::Type SelectInfo);
@@ -44,6 +39,8 @@ public:
 
 	void OnSimCacheChanged(const FAssetData& InAsset);
 
+	void OnSimCacheFrameChanged(const bool bRefresh);
+
 	TArray<TSharedPtr<int32>>					RowItems;
 	TSharedPtr<FNiagaraSimCacheViewModel>		SimCacheViewModel;
 
diff --git a/Engine/Plugins/FX/Niagara/Source/NiagaraEditor/Private/Widgets/SNiagaraSimCacheViewTimeline.cpp b/Engine/Plugins/FX/Niagara/Source/NiagaraEditor/Private/Widgets/SNiagaraSimCacheViewTimeline.cpp
new file mode 100644
--- /dev/null
+++ b/Engine/Plugins/FX/Niagara/Source/NiagaraEditor/Private/Widgets/SNiagaraSimCacheViewTimeline.cpp
@@ -0,0 +1,112 @@
+﻿// Copyright Epic Games, Inc. All Rights Reserved.
+
+#include "SNiagaraSimCacheViewTimeline.h"
+
+#include "ViewModels/NiagaraSimCacheViewModel.h"
+#include "Styling/AppStyle.h"
+#include "EditorWidgetsModule.h"
+
+
+#define LOCTEXT_NAMESPACE "NigaraSimCacheViewTimeline"
+
+void SNiagaraSimCacheViewTimeline::Construct(const FArguments& InArgs)
+{
+	WeakViewModel = InArgs._WeakViewModel;
+}
+
+int32 SNiagaraSimCacheViewTimeline::OnPaint(const FPaintArgs& Args, const FGeometry& AllottedGeometry,
+	const FSlateRect& MyCullingRect, FSlateWindowElementList& OutDrawElements, int32 LayerId,
+	const FWidgetStyle& InWidgetStyle, bool bParentEnabled) const
+{
+	FNiagaraSimCacheViewModel* ViewModel = WeakViewModel.Pin().Get();
+
+	if( !ViewModel )
+	{
+		return LayerId;
+	}
+
+	const int NumFrames = ViewModel->GetNumFrames();
+	const int CurrentFrame = ViewModel->GetFrameIndex();
+	const float NormalizedTime = NumFrames > 0 ? float(CurrentFrame) / float(NumFrames) : 0;
+
+	//////////////////////////////////////////////////////////////////////////
+	// Draw boxes to show frames
+	{
+		const FSlateBrush* BoxBrush = FAppStyle::GetBrush("Sequencer.Section.BackgroundTint");
+		const FLinearColor BoxTints[2] = { FLinearColor(0.5f, 0.5f, 0.5f, 1.0f), FLinearColor(0.3f, 0.3f, 0.3f, 1.0f) };
+		const FLinearColor CurrentTint(0.5f, 0.5f, 1.0f, 1.0f);
+		if (NumFrames > 0)
+		{
+			const float UStep = AllottedGeometry.Size.X / float(NumFrames);
+
+			for (int32 i = 0; i < NumFrames; ++i)
+			{
+				const FVector2D BoxLocation(UStep * float(i), 0.0f);
+				const FVector2D BoxSize(UStep, AllottedGeometry.Size.Y);
+				const FLinearColor& Tint = (i == CurrentFrame) ? CurrentTint : BoxTints[i & 1];
+				FSlateDrawElement::MakeBox(OutDrawElements, LayerId++, AllottedGeometry.ToPaintGeometry(BoxLocation, BoxSize), BoxBrush, ESlateDrawEffect::None, Tint);
+			}
+		}
+		else
+		{
+			FSlateDrawElement::MakeBox(OutDrawElements, LayerId++, AllottedGeometry.ToPaintGeometry(), BoxBrush, ESlateDrawEffect::None, BoxTints[0]);
+		}
+	}
+
+	//////////////////////////////////////////////////////////////////////////
+	// Draw a line to show the current time
+	{
+		TArray<FVector2D> LinePoints;
+		LinePoints.AddUninitialized(2);
+
+		const float LineXPos = AllottedGeometry.Size.X * NormalizedTime;
+		LinePoints[0] = FVector2D(LineXPos, 0.0f);
+		LinePoints[1] = FVector2D(LineXPos, AllottedGeometry.Size.Y);
+
+		FSlateDrawElement::MakeLines(
+			OutDrawElements,
+			LayerId++,
+			AllottedGeometry.ToPaintGeometry(),
+			LinePoints,
+			ESlateDrawEffect::None,
+			FLinearColor::White
+		);
+	}
+
+	return LayerId;
+}
+
+FReply SNiagaraSimCacheViewTimeline::OnMouseButtonDown(const FGeometry& MyGeometry, const FPointerEvent& MouseEvent)
+{
+	return OnTimelineScrubbed(MyGeometry, MouseEvent);
+}
+
+FReply SNiagaraSimCacheViewTimeline::OnMouseButtonUp(const FGeometry& MyGeometry, const FPointerEvent& MouseEvent)
+{
+	return OnTimelineScrubbed(MyGeometry, MouseEvent);
+}
+
+FReply SNiagaraSimCacheViewTimeline::OnMouseMove(const FGeometry& MyGeometry, const FPointerEvent& MouseEvent)
+{
+	// TODO Implement dragging on timeline
+	return FReply::Unhandled();
+}
+
+FReply SNiagaraSimCacheViewTimeline::OnTimelineScrubbed(const FGeometry& MyGeometry, const FPointerEvent& MouseEvent)
+{
+	if ( FNiagaraSimCacheViewModel* ViewModel = WeakViewModel.Pin().Get() )
+	{
+		const FVector2D LocalLocation = MyGeometry.AbsoluteToLocal(MouseEvent.GetScreenSpacePosition());
+		const float NormalizedTime = LocalLocation.X / MyGeometry.Size.X;
+
+		const int32 NewFrameIndex = FMath::RoundToInt(NormalizedTime * float(ViewModel->GetNumFrames() - 1));
+
+		ViewModel->SetFrameIndex(NewFrameIndex);
+
+		return FReply::Handled();
+	}
+
+	return FReply::Unhandled();
+}
+
+#undef LOCTEXT_NAMESPACE
diff --git a/Engine/Plugins/FX/Niagara/Source/NiagaraEditor/Private/Widgets/SNiagaraSimCacheViewTimeline.h b/Engine/Plugins/FX/Niagara/Source/NiagaraEditor/Private/Widgets/SNiagaraSimCacheViewTimeline.h
new file mode 100644
--- /dev/null
+++ b/Engine/Plugins/FX/Niagara/Source/NiagaraEditor/Private/Widgets/SNiagaraSimCacheViewTimeline.h
@@ -0,0 +1,30 @@
+﻿// Copyright Epic Games, Inc. All Rights Reserved.
+
+#pragma once
+
+#include "Widgets/DeclarativeSyntaxSupport.h"
+#include "Widgets/SCompoundWidget.h"
+
+class SNiagaraSimCacheViewTimeline : public SCompoundWidget
+{
+public:
+	SLATE_BEGIN_ARGS(SNiagaraSimCacheViewTimeline) {}
+		SLATE_ARGUMENT(TWeakPtr<class FNiagaraSimCacheViewModel>, WeakViewModel)
+	SLATE_END_ARGS()
+
+	void Construct(const FArguments& InArgs);
+
+	// SWidget interface
+	virtual int32 OnPaint(const FPaintArgs& Args, const FGeometry& AllottedGeometry, const FSlateRect& MyCullingRect, FSlateWindowElementList& OutDrawElements, int32 LayerId, const FWidgetStyle& InWidgetStyle, bool bParentEnabled) const override;
+	virtual FReply OnMouseButtonDown(const FGeometry& MyGeometry, const FPointerEvent& MouseEvent) override;
+	virtual FReply OnMouseButtonUp(const FGeometry& MyGeometry, const FPointerEvent& MouseEvent) override;
+	virtual FReply OnMouseMove(const FGeometry& MyGeometry, const FPointerEvent& MouseEvent) override;
+	// End of SWidget interface
+
+	FReply OnTimelineScrubbed(const FGeometry& MyGeometry, const FPointerEvent& MouseEvent);
+
+	
+private:
+	TWeakPtr<FNiagaraSimCacheViewModel> WeakViewModel;
+
+};
diff --git a/Engine/Plugins/FX/Niagara/Source/NiagaraEditor/Private/Widgets/SNiagaraSimCacheViewTransportControls.cpp b/Engine/Plugins/FX/Niagara/Source/NiagaraEditor/Private/Widgets/SNiagaraSimCacheViewTransportControls.cpp
new file mode 100644
--- /dev/null
+++ b/Engine/Plugins/FX/Niagara/Source/NiagaraEditor/Private/Widgets/SNiagaraSimCacheViewTransportControls.cpp
@@ -0,0 +1,132 @@
+﻿// Fill out your copyright notice in the Description page of Project Settings.
+
+
+#include "SNiagaraSimCacheViewTransportControls.h"
+
+#include "Modules/ModuleManager.h"
+#include "EditorWidgetsModule.h"
+#include "SlateOptMacros.h"
+#include "ITransportControl.h"
+#include "NiagaraSimCacheViewModel.h"
+#include "Widgets/Input/SNumericEntryBox.h"
+
+void SNiagaraSimCacheViewTransportControls::Construct(const FArguments& InArgs)
+{
+	WeakViewModel = InArgs._WeakViewModel;
+	FEditorWidgetsModule& EditorWidgetsModule = FModuleManager::Get().LoadModuleChecked<FEditorWidgetsModule>("EditorWidgets");
+
+	TransportControls = MakeTransportControlsWidget(EditorWidgetsModule);
+
+	ChildSlot
+	[
+		TransportControls.ToSharedRef()
+	];
+}
+
+TSharedRef<SWidget> SNiagaraSimCacheViewTransportControls::MakeTransportControlsWidget(
+	FEditorWidgetsModule& EditorWidgetsModule)
+{
+	FTransportControlArgs TransportControlArgs;
+	TransportControlArgs.OnBackwardEnd.BindSP(this, &SNiagaraSimCacheViewTransportControls::OnTransportBackwardEnd);
+	TransportControlArgs.OnBackwardStep.BindSP(this, &SNiagaraSimCacheViewTransportControls::OnTransportBackwardStep);
+	TransportControlArgs.OnForwardStep.BindSP(this, &SNiagaraSimCacheViewTransportControls::OnTransportForwardStep);
+	TransportControlArgs.OnForwardEnd.BindSP(this, &SNiagaraSimCacheViewTransportControls::OnTransportForwardEnd);
+	
+
+	TransportControlArgs.WidgetsToCreate.Add(FTransportControlWidget(ETransportControlWidgetType::BackwardEnd));
+	TransportControlArgs.WidgetsToCreate.Add(FTransportControlWidget(ETransportControlWidgetType::BackwardStep));
+	TransportControlArgs.WidgetsToCreate.Add(FTransportControlWidget(FOnMakeTransportWidget::CreateSP(this, &SNiagaraSimCacheViewTransportControls::CreateFrameWidget)));
+	TransportControlArgs.WidgetsToCreate.Add(FTransportControlWidget(ETransportControlWidgetType::ForwardStep));
+	TransportControlArgs.WidgetsToCreate.Add(FTransportControlWidget(ETransportControlWidgetType::ForwardEnd));
+	
+	TransportControlArgs.bAreButtonsFocusable = false;
+
+	return EditorWidgetsModule.CreateTransportControl(TransportControlArgs);
+}
+
+TOptional<int32> SNiagaraSimCacheViewTransportControls::GetCurrentFrame() const
+{
+	if ( FNiagaraSimCacheViewModel* ViewModel = WeakViewModel.Pin().Get() )
+	{
+		return ViewModel->GetFrameIndex();
+	}
+
+	return 0;
+}
+
+void SNiagaraSimCacheViewTransportControls::OnFrameIndexChanged(const int32 InFrameIndex, ETextCommit::Type Arg)
+{
+	if (FNiagaraSimCacheViewModel* ViewModel = WeakViewModel.Pin().Get())
+	{
+		int32 NewFrameIndex = FMath::Max(0, InFrameIndex);
+		NewFrameIndex = FMath::Min(NewFrameIndex, ViewModel->GetNumFrames() - 1);
+		ViewModel->SetFrameIndex(NewFrameIndex);
+	}
+}
+
+TOptional<int32> SNiagaraSimCacheViewTransportControls::GetMaxFrameIndex() const
+{
+	if (FNiagaraSimCacheViewModel* ViewModel = WeakViewModel.Pin().Get())
+	{
+		return ViewModel->GetNumFrames() - 1;
+	}
+
+	return 0;
+}
+
+TSharedRef<SWidget> SNiagaraSimCacheViewTransportControls::CreateFrameWidget()
+{
+	return SNew(SHorizontalBox)
+		+SHorizontalBox::Slot()
+		.Padding(0.5f)
+		[
+			SNew(SNumericEntryBox<int32>)
+			.AllowSpin(false)
+			.MinValue(0)
+			.MaxValue(this, &SNiagaraSimCacheViewTransportControls::GetMaxFrameIndex)
+			.Value(this, &SNiagaraSimCacheViewTransportControls::GetCurrentFrame)
+			.OnValueCommitted(this, &SNiagaraSimCacheViewTransportControls::OnFrameIndexChanged)
+		];
+	
+}
+
+FReply SNiagaraSimCacheViewTransportControls::OnTransportBackwardEnd()
+{
+	if ( FNiagaraSimCacheViewModel* ViewModel = WeakViewModel.Pin().Get() )
+	{
+		ViewModel->SetFrameIndex(0);
+		
+		return FReply::Handled();
+	}
+	return FReply::Unhandled();
+}
+
+FReply SNiagaraSimCacheViewTransportControls::OnTransportBackwardStep()
+{
+	if ( FNiagaraSimCacheViewModel* ViewModel = WeakViewModel.Pin().Get() )
+	{
+		ViewModel->SetFrameIndex(FMath::Max(ViewModel->GetFrameIndex() -1, 0));
+		return FReply::Handled();
+	}
+	return FReply::Unhandled();
+}
+
+FReply SNiagaraSimCacheViewTransportControls::OnTransportForwardStep()
+{
+	if ( FNiagaraSimCacheViewModel* ViewModel = WeakViewModel.Pin().Get() )
+	{
+		ViewModel->SetFrameIndex(FMath::Min(ViewModel->GetFrameIndex() + 1, ViewModel->GetNumFrames() - 1));
+		return FReply::Handled();
+	}
+	return FReply::Unhandled();
+}
+
+FReply SNiagaraSimCacheViewTransportControls::OnTransportForwardEnd()
+{
+	if ( FNiagaraSimCacheViewModel* ViewModel = WeakViewModel.Pin().Get() )
+	{
+		ViewModel->SetFrameIndex(ViewModel->GetNumFrames() - 1);
+		return FReply::Handled();
+	}
+	return FReply::Unhandled();
+}
\ No newline at end of file
diff --git a/Engine/Plugins/FX/Niagara/Source/NiagaraEditor/Private/Widgets/SNiagaraSimCacheViewTransportControls.h b/Engine/Plugins/FX/Niagara/Source/NiagaraEditor/Private/Widgets/SNiagaraSimCacheViewTransportControls.h
new file mode 100644
--- /dev/null
+++ b/Engine/Plugins/FX/Niagara/Source/NiagaraEditor/Private/Widgets/SNiagaraSimCacheViewTransportControls.h
@@ -0,0 +1,35 @@
+﻿// Fill out your copyright notice in the Description page of Project Settings.
+
+#pragma once
+
+#include "CoreMinimal.h"
+#include "Widgets/SCompoundWidget.h"
+
+class NIAGARAEDITOR_API SNiagaraSimCacheViewTransportControls : public SCompoundWidget
+{
+public:
+	SLATE_BEGIN_ARGS(SNiagaraSimCacheViewTransportControls) {}
+		SLATE_ARGUMENT(TWeakPtr<class FNiagaraSimCacheViewModel>, WeakViewModel)
+	SLATE_END_ARGS()
+
+	/** Constructs this widget with InArgs */
+	void Construct(const FArguments& InArgs);
+
+	
+	// Transport Controls
+	TSharedRef<SWidget> MakeTransportControlsWidget(class FEditorWidgetsModule& EditorWidgetsModule);
+	TOptional<int32> GetCurrentFrame() const;
+	void OnFrameIndexChanged(const int32 InFrameIndex, ETextCommit::Type Arg);
+	TOptional<int32> GetMaxFrameIndex() const;
+	// Custom transport widget for frame selection
+	TSharedRef<SWidget> CreateFrameWidget();
+	
+	FReply OnTransportBackwardEnd();
+	FReply OnTransportBackwardStep();
+	FReply OnTransportForwardStep();
+	FReply OnTransportForwardEnd();
+
+private:
+	TWeakPtr<FNiagaraSimCacheViewModel> WeakViewModel;
+	TSharedPtr<SWidget>	TransportControls;
+};
