Adding Missing Copyright for Niagara Sim Cache Transport Controls

diff --git a/Engine/Plugins/FX/Niagara/Source/NiagaraEditor/Private/Widgets/SNiagaraSimCacheViewTransportControls.cpp b/Engine/Plugins/FX/Niagara/Source/NiagaraEditor/Private/Widgets/SNiagaraSimCacheViewTransportControls.cpp
--- a/Engine/Plugins/FX/Niagara/Source/NiagaraEditor/Private/Widgets/SNiagaraSimCacheViewTransportControls.cpp
+++ b/Engine/Plugins/FX/Niagara/Source/NiagaraEditor/Private/Widgets/SNiagaraSimCacheViewTransportControls.cpp
@@ -1,4 +1,4 @@
-﻿// Fill out your copyright notice in the Description page of Project Settings.
+﻿// Copyright Epic Games, Inc. All Rights Reserved.
 
 
 #include "SNiagaraSimCacheViewTransportControls.h"
diff --git a/Engine/Plugins/FX/Niagara/Source/NiagaraEditor/Private/Widgets/SNiagaraSimCacheViewTransportControls.h b/Engine/Plugins/FX/Niagara/Source/NiagaraEditor/Private/Widgets/SNiagaraSimCacheViewTransportControls.h
--- a/Engine/Plugins/FX/Niagara/Source/NiagaraEditor/Private/Widgets/SNiagaraSimCacheViewTransportControls.h
+++ b/Engine/Plugins/FX/Niagara/Source/NiagaraEditor/Private/Widgets/SNiagaraSimCacheViewTransportControls.h
@@ -1,4 +1,4 @@
-﻿// Fill out your copyright notice in the Description page of Project Settings.
+﻿// Copyright Epic Games, Inc. All Rights Reserved.
 
 #pragma once
 
