- Add Bake Quality Level option to Niagara, defaults to cinematic

diff --git a/Engine/Plugins/FX/Niagara/Config/BaseNiagara.ini b/Engine/Plugins/FX/Niagara/Config/BaseNiagara.ini
--- a/Engine/Plugins/FX/Niagara/Config/BaseNiagara.ini
+++ b/Engine/Plugins/FX/Niagara/Config/BaseNiagara.ini
@@ -47,6 +47,9 @@
 +StructRedirects=(OldName="NiagaraFlipbookTextureSettings", NewName="NiagaraBakerTextureSettings")
 +StructRedirects=(OldName="NiagaraFlipbookTextureSource", NewName="NiagaraBakerTextureSource")
 
+[/Script/Niagara.NiagaraBakerSettings]
+BakeQualityLevel=Cinematic
+
 [/Script/Niagara.NiagaraSettings]
 +AdditionalParameterEnums=/Niagara/Enums/ENiagaraCoordinateSpace.ENiagaraCoordinateSpace
 +AdditionalParameterEnums=/Niagara/Enums/ENiagaraOrientationAxis.ENiagaraOrientationAxis
diff --git a/Engine/Plugins/FX/Niagara/Source/Niagara/Classes/NiagaraBakerSettings.h b/Engine/Plugins/FX/Niagara/Source/Niagara/Classes/NiagaraBakerSettings.h
--- a/Engine/Plugins/FX/Niagara/Source/Niagara/Classes/NiagaraBakerSettings.h
+++ b/Engine/Plugins/FX/Niagara/Source/Niagara/Classes/NiagaraBakerSettings.h
@@ -138,7 +138,7 @@ struct FNiagaraBakerOutputFrameIndices
 	float	Interp = 0.0f;
 };
 
-UCLASS()
+UCLASS(config=Niagara, defaultconfig)
 class NIAGARA_API UNiagaraBakerSettings : public UObject
 {
 	GENERATED_BODY()
@@ -185,6 +185,10 @@ public:
 	UPROPERTY(EditAnywhere, Category = "Settings")
 	int32 CurrentCameraIndex = 0;
 
+	/** What quality level to use when baking the simulation, where None means use the current quality level. */
+	UPROPERTY(EditAnywhere, Category = "Settings", config)
+	FName BakeQualityLevel;
+
 	/** Should we render just the component or the whole scene. */
 	UPROPERTY(EditAnywhere, Category = "Settings")
 	uint8 bRenderComponentOnly : 1;
diff --git a/Engine/Plugins/FX/Niagara/Source/NiagaraEditor/Private/ViewModels/NiagaraBakerViewModel.cpp b/Engine/Plugins/FX/Niagara/Source/NiagaraEditor/Private/ViewModels/NiagaraBakerViewModel.cpp
--- a/Engine/Plugins/FX/Niagara/Source/NiagaraEditor/Private/ViewModels/NiagaraBakerViewModel.cpp
+++ b/Engine/Plugins/FX/Niagara/Source/NiagaraEditor/Private/ViewModels/NiagaraBakerViewModel.cpp
@@ -6,6 +6,8 @@
 #include "NiagaraBakerRenderer.h"
 #include "NiagaraGpuComputeDispatchInterface.h"
 #include "NiagaraComponent.h"
+#include "NiagaraPlatformSet.h"
+#include "NiagaraSettings.h"
 #include "NiagaraSystem.h"
 
 #include "Niagara/Classes/NiagaraDataInterfaceRenderTarget2D.h"
@@ -348,6 +350,22 @@ void FNiagaraBakerViewModel::ResetCurrentCamera()
 	}
 }
 
+bool FNiagaraBakerViewModel::IsBakeQualityLevel(FName QualityLevel) const
+{
+	UNiagaraBakerSettings* BakerSettings = GetBakerSettings();
+	return BakerSettings ? BakerSettings->BakeQualityLevel == QualityLevel : false;
+}
+
+void FNiagaraBakerViewModel::SetBakeQualityLevel(FName QualityLevel)
+{
+	if (UNiagaraBakerSettings* BakerSettings = GetBakerSettings())
+	{
+		const FScopedTransaction Transaction(LOCTEXT("BakeQualityLevel", "Set Bake Quality Level"));
+		BakerSettings->Modify();
+		BakerSettings->BakeQualityLevel = QualityLevel;
+	}
+}
+
 bool FNiagaraBakerViewModel::IsSimTickRate(int TickRate) const
 {
 	UNiagaraBakerSettings* BakerSettings = GetBakerSettings();
@@ -611,6 +629,20 @@ void FNiagaraBakerViewModel::RenderBaker()
 		FScopedSlowTask SlowTask(TotalFrames);
 		SlowTask.MakeDialog();
 
+		const int32 ExistingQualityLevel = FNiagaraPlatformSet::GetQualityLevel();
+		if (BakerSettings->BakeQualityLevel.IsNone() == false)
+		{
+			const UNiagaraSettings* NiagaraSettings = GetDefault<UNiagaraSettings>();
+			for (int32 i=0; i < NiagaraSettings->QualityLevels.Num(); ++i)
+			{
+				if (FName(NiagaraSettings->QualityLevels[i].ToString()) == BakerSettings->BakeQualityLevel)
+				{
+					SetGNiagaraQualityLevel(i);
+					break;
+				}
+			}
+		}
+
 		for ( int32 iFrame=0; iFrame < TotalFrames; ++iFrame )
 		{
 			SlowTask.EnterProgressFrame(1, FText::Format(LOCTEXT("BakingFormat", "Baking Frame ({0} / {1})"), iFrame + 1, TotalFrames));
@@ -628,6 +660,11 @@ void FNiagaraBakerViewModel::RenderBaker()
 				}
 			}
 		}
+
+		if (BakerSettings->BakeQualityLevel.IsNone() == false)
+		{
+			SetGNiagaraQualityLevel(ExistingQualityLevel);
+		}
 	}
 
 	// Ensure we flush everything from the render thread
diff --git a/Engine/Plugins/FX/Niagara/Source/NiagaraEditor/Private/ViewModels/NiagaraBakerViewModel.h b/Engine/Plugins/FX/Niagara/Source/NiagaraEditor/Private/ViewModels/NiagaraBakerViewModel.h
--- a/Engine/Plugins/FX/Niagara/Source/NiagaraEditor/Private/ViewModels/NiagaraBakerViewModel.h
+++ b/Engine/Plugins/FX/Niagara/Source/NiagaraEditor/Private/ViewModels/NiagaraBakerViewModel.h
@@ -103,6 +103,9 @@ public:
 
 	void ResetCurrentCamera();
 
+	bool IsBakeQualityLevel(FName QualityLevel) const;
+	void SetBakeQualityLevel(FName QualityLevel);
+
 	bool IsSimTickRate(int TickRate) const;
 	int GetSimTickRate() const;
 	void SetSimTickRate(int TickRate);
diff --git a/Engine/Plugins/FX/Niagara/Source/NiagaraEditor/Private/Widgets/SNiagaraBakerWidget.cpp b/Engine/Plugins/FX/Niagara/Source/NiagaraEditor/Private/Widgets/SNiagaraBakerWidget.cpp
--- a/Engine/Plugins/FX/Niagara/Source/NiagaraEditor/Private/Widgets/SNiagaraBakerWidget.cpp
+++ b/Engine/Plugins/FX/Niagara/Source/NiagaraEditor/Private/Widgets/SNiagaraBakerWidget.cpp
@@ -7,6 +7,7 @@
 #include "NiagaraEditorCommon.h"
 #include "NiagaraEditorStyle.h"
 #include "NiagaraComponent.h"
+#include "NiagaraSettings.h"
 #include "NiagaraSystem.h"
 #include "ViewModels/NiagaraBakerViewModel.h"
 
@@ -640,6 +641,12 @@ TSharedRef<SWidget> SNiagaraBakerWidget::MakeSettingsMenu()
 		EUserInterfaceActionType::ToggleButton
 	);
 
+	MenuBuilder.AddSubMenu(
+		LOCTEXT("BakeQualityLevel", "Bake Quality Level"),
+		LOCTEXT("BakeQualityLevelToolTip", "Allows you to optionally override the quality level when baking, i.e. use cinematic quality while baking"),
+		FNewMenuDelegate::CreateSP(this, &SNiagaraBakerWidget::MakeBakeQualityLevelMenu)
+	);
+
 	MenuBuilder.AddSubMenu(
 		LOCTEXT("SimulationTickRate", "Simulation Tick Rate"),
 		LOCTEXT("SimulationTickRateToolTip", "The rate at which the simulation will tick, i.e. 120fps, 60fps, 30fps, etc."),
@@ -777,6 +784,48 @@ TSharedRef<SWidget> SNiagaraBakerWidget::MakeAddOutputMenu()
 	return MenuBuilder.MakeWidget();
 }
 
+void SNiagaraBakerWidget::MakeBakeQualityLevelMenu(FMenuBuilder& MenuBuilder) const
+{
+	FNiagaraBakerViewModel* ViewModel = WeakViewModel.Pin().Get();
+	check(ViewModel != nullptr);
+
+	// Special entry for 'leave as default'
+	MenuBuilder.AddMenuEntry(
+		LOCTEXT("CurrentQualityLevel" ,"Current Quality Level"),
+		LOCTEXT("CurrentQualityLevelTooltip", "Uses the current quality level that is set in the editor when baking."),
+		FSlateIcon(),
+		FUIAction(
+			FExecuteAction::CreateSP(ViewModel, &FNiagaraBakerViewModel::SetBakeQualityLevel, FName()),
+			FCanExecuteAction(),
+			FIsActionChecked::CreateSP(ViewModel, &FNiagaraBakerViewModel::IsBakeQualityLevel, FName())
+		),
+		NAME_None,
+		EUserInterfaceActionType::ToggleButton
+	);
+
+	// Engine driven quality levels
+	MenuBuilder.BeginSection("QualityLevel", LOCTEXT("QualityLevel", "QualityLevel"));
+	{
+		for (const FText& QualityLevelText : GetDefault<UNiagaraSettings>()->QualityLevels)
+		{
+			const FName QualityLevel(*QualityLevelText.ToString());
+
+			MenuBuilder.AddMenuEntry(
+				QualityLevelText,
+				FText::GetEmpty(),
+				FSlateIcon(),
+				FUIAction(
+					FExecuteAction::CreateSP(ViewModel, &FNiagaraBakerViewModel::SetBakeQualityLevel, QualityLevel),
+					FCanExecuteAction(),
+					FIsActionChecked::CreateSP(ViewModel, &FNiagaraBakerViewModel::IsBakeQualityLevel, QualityLevel)
+				),
+				NAME_None,
+				EUserInterfaceActionType::ToggleButton
+			);
+		}
+	}
+}
+
 void SNiagaraBakerWidget::MakeSimTickRateMenu(FMenuBuilder& MenuBuilder) const
 {
 	using namespace NiagaraBakerWidgetLocal;
diff --git a/Engine/Plugins/FX/Niagara/Source/NiagaraEditor/Private/Widgets/SNiagaraBakerWidget.h b/Engine/Plugins/FX/Niagara/Source/NiagaraEditor/Private/Widgets/SNiagaraBakerWidget.h
--- a/Engine/Plugins/FX/Niagara/Source/NiagaraEditor/Private/Widgets/SNiagaraBakerWidget.h
+++ b/Engine/Plugins/FX/Niagara/Source/NiagaraEditor/Private/Widgets/SNiagaraBakerWidget.h
@@ -33,6 +33,7 @@ public:
 	TSharedRef<SWidget> MakeOutputSelectMenu();
 	TSharedRef<SWidget> MakeAddOutputMenu();
 
+	void MakeBakeQualityLevelMenu(FMenuBuilder& MenuBuilder) const;
 	void MakeSimTickRateMenu(FMenuBuilder& MenuBuilder) const;
 
 	void FindWarnings();
