Added exception to the creation of user parameters when a user var linked handle is set, when:
- the input type is position and a vector type user var of the same name exists
- the input type is vector and a position type user var of the same name exists

diff --git a/Engine/Plugins/FX/Niagara/Source/NiagaraEditor/Private/ViewModels/Stack/NiagaraStackFunctionInput.cpp b/Engine/Plugins/FX/Niagara/Source/NiagaraEditor/Private/ViewModels/Stack/NiagaraStackFunctionInput.cpp
--- a/Engine/Plugins/FX/Niagara/Source/NiagaraEditor/Private/ViewModels/Stack/NiagaraStackFunctionInput.cpp
+++ b/Engine/Plugins/FX/Niagara/Source/NiagaraEditor/Private/ViewModels/Stack/NiagaraStackFunctionInput.cpp
@@ -1249,7 +1249,28 @@ void UNiagaraStackFunctionInput::SetLinkedValueHandle(const FNiagaraParameterHan
 			// rather than waiting on the compile results so that it's immediately available.
 			FNiagaraVariable LinkedParameterVariable = FNiagaraVariable(InputType, InParameterHandle.GetParameterHandleString()); 
 			FNiagaraParameterStore& UserParameters = GetSystemViewModel()->GetSystem().GetExposedParameters();
-			if (UserParameters.IndexOf(LinkedParameterVariable) == INDEX_NONE)
+
+			// we only know the input type, and the parameter handle. We don't know the assigned parameter's real type (a vector parameter can be linked to a position input).
+			// therefore, assuming we can FindOrAdd a parameter based on just input type & handle is not possible.
+			// We special case the position<->vector assignment and skip the creation of the parameter if one of the same name but the other type already exists.
+			bool bSkipCreation = false;
+			if(InputType == FNiagaraTypeDefinition::GetPositionDef())
+			{
+				FNiagaraVariable ExistingVectorParameter(FNiagaraTypeDefinition::GetVec3Def(), InParameterHandle.GetParameterHandleString());
+				if(UserParameters.IndexOf(ExistingVectorParameter) != INDEX_NONE)
+				{
+					bSkipCreation = true;
+				}
+			}
+			else if(InputType == FNiagaraTypeDefinition::GetVec3Def())
+			{
+				FNiagaraVariable ExistingPositionParameter(FNiagaraTypeDefinition::GetPositionDef(), InParameterHandle.GetParameterHandleString());
+				if(UserParameters.IndexOf(ExistingPositionParameter) != INDEX_NONE)
+				{
+					bSkipCreation = true;
+				}
+			}
+			if (UserParameters.IndexOf(LinkedParameterVariable) == INDEX_NONE && bSkipCreation == false)
 			{
 				if (InputType.IsDataInterface())
 				{
