Fix for FQuadTreeQueryHelper::FindFirstTriangle() which was early outing too early (without neceesarily finding the first valid triangle).

diff --git a/Engine/Plugins/FX/Niagara/Source/Niagara/Private/NiagaraDataInterfaceSkeletalMeshUvMapping.cpp b/Engine/Plugins/FX/Niagara/Source/Niagara/Private/NiagaraDataInterfaceSkeletalMeshUvMapping.cpp
--- a/Engine/Plugins/FX/Niagara/Source/Niagara/Private/NiagaraDataInterfaceSkeletalMeshUvMapping.cpp
+++ b/Engine/Plugins/FX/Niagara/Source/Niagara/Private/NiagaraDataInterfaceSkeletalMeshUvMapping.cpp
@@ -74,7 +74,7 @@ struct FQuadTreeQueryHelper
 				return false;
 			}
 
-			return false;
+			return true;
 		});
 
 		return FoundTriangleIndex;
