Apply SystemLWCTile when evaluating Mesh SDFs in Niagara.

diff --git a/Engine/Plugins/FX/Niagara/Shaders/Private/NiagaraDataInterfaceRigidMeshCollisionQuery.ush b/Engine/Plugins/FX/Niagara/Shaders/Private/NiagaraDataInterfaceRigidMeshCollisionQuery.ush
--- a/Engine/Plugins/FX/Niagara/Shaders/Private/NiagaraDataInterfaceRigidMeshCollisionQuery.ush
+++ b/Engine/Plugins/FX/Niagara/Shaders/Private/NiagaraDataInterfaceRigidMeshCollisionQuery.ush
@@ -22,6 +22,7 @@ Buffer<float4>	{ParameterName}_InverseTransformBuffer;
 Buffer<float4>	{ParameterName}_ElementExtentBuffer;
 Buffer<uint>	{ParameterName}_PhysicsTypeBuffer;
 Buffer<uint>	{ParameterName}_DFIndexBuffer;
+float3			{ParameterName}_SystemLWCTile;
 
 /* -----------------------------------------------------------------
  * Element accessors
@@ -361,7 +362,9 @@ void GetClosestDistance_{ParameterName}(in float3 WorldPosition, in float TimeFr
 
 // Given a world space position (WorldPosition) compute the static mesh closest point (position,normal,velocity)
 void GetClosestPointMeshDistanceField_{ParameterName}(in float3 WorldPosition, in float DeltaTime, in float TimeFraction, in float MaxDistance, out float OutClosestDistance, out float3 OutClosestPosition, out float3 OutClosestNormal, out float3 OutClosestVelocity, out bool NormalIsValid)
-{	
+{
+	const FLWCVector3 LWCWorldPosition = MakeLWCVector3({ParameterName}_SystemLWCTile, WorldPosition);
+
 	float3 CollisionPosition = float3(0,0,0);
 	float3 CollisionNormal = float3(0,0,0);
 	OutClosestDistance = MAX_DISTANCE;	
@@ -384,17 +387,18 @@ void GetClosestPointMeshDistanceField_{ParameterName}(in float3 WorldPosition, i
 			// #todo(dmp): this does a dynamic branch based on intersecting the bbox.  Maybe we can factor that out due to the broadphase here?
 			// #todo(dmp): Check the distance from the physics asset for a speedup?
 			MaxDistance = max(.001, MaxDistance);
-			OutClosestDistance = DistanceToNearestSurfaceForObject(DFIndex, LWCPromote(WorldPosition), MaxDistance); // LWC_TODO
+			OutClosestDistance = DistanceToNearestSurfaceForObject(DFIndex, LWCWorldPosition, MaxDistance);
 
 			FDFObjectData DFObjectData = LoadDFObjectData(DFIndex);
-			float4x4 WorldToVolume = LWCHackToFloat(DFObjectData.WorldToVolume);
 			uint NumMips = LoadDFAssetData(DFObjectData.AssetIndex, 0).NumMips;
 			FDFAssetData DFAssetData = LoadDFAssetData(DFObjectData.AssetIndex, NumMips - 1);
-			float3 VolumePosition = mul(float4(WorldPosition, 1), WorldToVolume).xyz;		
+
+			float3 VolumePosition = LWCMultiply(LWCWorldPosition, DFObjectData.WorldToVolume);
 			float3 VolumeGradient = CalculateMeshSDFGradient(VolumePosition, DFAssetData);
 			float VolumeGradientLength = length(VolumeGradient);
 			float3 VolumeNormal = VolumeGradientLength > 0.0f ? VolumeGradient / VolumeGradientLength : 0;
 			// Transform by transposed inverse to handle non-uniform scaling
+			float4x4 WorldToVolume = LWCHackToFloat(DFObjectData.WorldToVolume);
 			float3 WorldGradient = mul(VolumeNormal, transpose((float3x3)WorldToVolume));
 			float WorldGradientLength = length(WorldGradient);
 			OutClosestNormal = WorldGradientLength > 0.0f ? WorldGradient / WorldGradientLength : 0;
@@ -416,7 +420,9 @@ void GetClosestPointMeshDistanceField_{ParameterName}(in float3 WorldPosition, i
 
 // Given a world space position (WorldPosition) compute the static mesh closest point (position,normal,velocity) by querying all overlapping mesh distance fields
 void GetClosestPointMeshDistanceFieldAccurate_{ParameterName}(in float3 WorldPosition, in float DeltaTime, in float TimeFraction, in float MaxDistance, out float OutClosestDistance, out float3 OutClosestPosition, out float3 OutClosestNormal, out float3 OutClosestVelocity, out bool NormalIsValid)
-{	
+{
+	const FLWCVector3 LWCWorldPosition = MakeLWCVector3({ParameterName}_SystemLWCTile, WorldPosition);
+
 	float3 CollisionPosition = float3(0,0,0);
 	float3 CollisionNormal = float3(0,0,0);
 	OutClosestDistance = MAX_DISTANCE;	
@@ -434,20 +440,21 @@ void GetClosestPointMeshDistanceFieldAccurate_{ParameterName}(in float3 WorldPos
 		// #todo(dmp): this does a dynamic branch based on intersecting the bbox.  Maybe we can factor that out due to the broadphase here?
 		// #todo(dmp): Check the distance from the physics asset for a speedup?
 		MaxDistance = max(.001, MaxDistance);
-		float OutClosestDistanceTmp = DistanceToNearestSurfaceForObject(DFIndex, LWCPromote(WorldPosition), MaxDistance); // LWC_TODO
+		float OutClosestDistanceTmp = DistanceToNearestSurfaceForObject(DFIndex, LWCWorldPosition, MaxDistance);
 
 		if (OutClosestDistanceTmp < OutClosestDistance)
 		{
 			OutClosestDistance = OutClosestDistanceTmp;
 			FDFObjectData DFObjectData = LoadDFObjectData(DFIndex);
-			float4x4 WorldToVolume = LWCHackToFloat(DFObjectData.WorldToVolume);
 			uint NumMips = LoadDFAssetData(DFObjectData.AssetIndex, 0).NumMips;
 			FDFAssetData DFAssetData = LoadDFAssetData(DFObjectData.AssetIndex, NumMips - 1);
-			float3 VolumePosition = mul(float4(WorldPosition, 1), WorldToVolume).xyz;		
+
+			float3 VolumePosition = LWCMultiply(LWCWorldPosition, DFObjectData.WorldToVolume);
 			float3 VolumeGradient = CalculateMeshSDFGradient(VolumePosition, DFAssetData);
 			float VolumeGradientLength = length(VolumeGradient);
 			float3 VolumeNormal = VolumeGradientLength > 0.0f ? VolumeGradient / VolumeGradientLength : 0;
 			// Transform by transposed inverse to handle non-uniform scaling
+			float4x4 WorldToVolume = LWCHackToFloat(DFObjectData.WorldToVolume);
 			float3 WorldGradient = mul(VolumeNormal, transpose((float3x3)WorldToVolume));
 			float WorldGradientLength = length(WorldGradient);
 			OutClosestNormal = WorldGradientLength > 0.0f ? WorldGradient / WorldGradientLength : 0;
@@ -469,7 +476,9 @@ void GetClosestPointMeshDistanceFieldAccurate_{ParameterName}(in float3 WorldPos
 
 // Given a world space position (WorldPosition) compute the static mesh closest point (position,normal,velocity)
 void GetClosestPointMeshDistanceFieldNoNormal_{ParameterName}(in float3 WorldPosition, in float DeltaTime, in float TimeFraction, in float MaxDistance, out float OutClosestDistance, out float3 OutClosestPosition, out float3 OutClosestVelocity)
-{	
+{
+	const FLWCVector3 LWCWorldPosition = MakeLWCVector3({ParameterName}_SystemLWCTile, WorldPosition);
+
 	float3 CollisionPosition = float3(0,0,0);
 	float3 CollisionNormal = float3(0,0,0);
 	OutClosestDistance = MAX_DISTANCE;	
@@ -487,7 +496,7 @@ void GetClosestPointMeshDistanceFieldNoNormal_{ParameterName}(in float3 WorldPos
 		{
 			// #todo(dmp): this does a dynamic branch based on intersecting the bbox.  Maybe we can factor that out due to the broadphase here?
 			MaxDistance = max(.001, MaxDistance);
-			OutClosestDistance = DistanceToNearestSurfaceForObject(DFIndex, LWCPromote(WorldPosition), MaxDistance); // LWC_TODO
+			OutClosestDistance = DistanceToNearestSurfaceForObject(DFIndex, LWCWorldPosition, MaxDistance);
 		}
 		
 		const float3 PreviousPosition = mul(GetPreviousTransform_{ParameterName}(ElementIndex), float4(CollisionPosition,1.0)).xyz;
diff --git a/Engine/Plugins/FX/Niagara/Shaders/Private/NiagaraDataInterfaceStaticMeshTemplate.ush b/Engine/Plugins/FX/Niagara/Shaders/Private/NiagaraDataInterfaceStaticMeshTemplate.ush
--- a/Engine/Plugins/FX/Niagara/Shaders/Private/NiagaraDataInterfaceStaticMeshTemplate.ush
+++ b/Engine/Plugins/FX/Niagara/Shaders/Private/NiagaraDataInterfaceStaticMeshTemplate.ush
@@ -30,6 +30,8 @@ float4x4		{ParameterName}_InstancePreviousTransformInverseTransposed;
 float4			{ParameterName}_InstancePreviousRotation;
 float3			{ParameterName}_InstanceWorldVelocity;
 
+float3			{ParameterName}_SystemLWCTile;
+
 #if DISTATICMESH_ALLOWDISTANCEFIELD
 int				{ParameterName}_InstanceDistanceFieldIndex;
 #endif
@@ -566,6 +568,8 @@ void GetWorldVelocity_{ParameterName}(out float3 Velocity)
 // Experimental Distance Field Functions
 void QueryDistanceField_{ParameterName}(bool bExecute, float3 WorldPosition, bool bUseMaxDistance, float MaxDistance, out bool OutIsValid, out float OutDistance)
 {
+	const FLWCVector3 LWCWorldPosition = MakeLWCVector3({ParameterName}_SystemLWCTile, WorldPosition);
+
 	if ( !bUseMaxDistance )
 	{
 		MaxDistance = 1e38;
@@ -578,7 +582,7 @@ void QueryDistanceField_{ParameterName}(bool bExecute, float3 WorldPosition, boo
 	{
 		// Needs to be a smidge over 0.0 for DistanceToNearestSurfaceForObject to query the SDF
 		MaxDistance = max(0.001, MaxDistance);
-		OutDistance = DistanceToNearestSurfaceForObject({ParameterName}_InstanceDistanceFieldIndex, LWCPromote(WorldPosition), MaxDistance); // LWC_TODO
+		OutDistance = DistanceToNearestSurfaceForObject({ParameterName}_InstanceDistanceFieldIndex, LWCWorldPosition, MaxDistance);
 	}
 #else
 	OutIsValid = false;
diff --git a/Engine/Plugins/FX/Niagara/Source/Niagara/Private/DataInterface/NiagaraDataInterfaceStaticMesh.cpp b/Engine/Plugins/FX/Niagara/Source/Niagara/Private/DataInterface/NiagaraDataInterfaceStaticMesh.cpp
--- a/Engine/Plugins/FX/Niagara/Source/Niagara/Private/DataInterface/NiagaraDataInterfaceStaticMesh.cpp
+++ b/Engine/Plugins/FX/Niagara/Source/Niagara/Private/DataInterface/NiagaraDataInterfaceStaticMesh.cpp
@@ -70,6 +70,8 @@ namespace NDIStaticMeshLocal
 		SHADER_PARAMETER(FVector3f,				InstanceWorldVelocity)
 
 		SHADER_PARAMETER(int,					InstanceDistanceFieldIndex)
+
+		SHADER_PARAMETER(FVector3f,				SystemLWCTile)
 	END_SHADER_PARAMETER_STRUCT()
 
 	static const TCHAR* TemplateShaderFile = TEXT("/Plugin/FX/Niagara/Private/NiagaraDataInterfaceStaticMeshTemplate.ush");
@@ -2592,6 +2594,8 @@ void UNiagaraDataInterfaceStaticMesh::SetShaderParameters(const FNiagaraDataInte
 	{
 		FNiagaraDistanceFieldHelper::SetMeshDistanceFieldParameters(Context.GetGraphBuilder(), DistanceFieldSceneData, *ShaderDistanceFieldObjectParameters, *ShaderDistanceFieldAtlasParameters, FNiagaraRenderer::GetDummyFloat4Buffer());
 	}
+
+	ShaderParameters->SystemLWCTile = Context.GetSystemLWCTile();
 }
 
 bool UNiagaraDataInterfaceStaticMesh::Equals(const UNiagaraDataInterface* Other) const
diff --git a/Engine/Plugins/FX/Niagara/Source/Niagara/Private/NiagaraDataInterfaceRigidMeshCollisionQuery.cpp b/Engine/Plugins/FX/Niagara/Source/Niagara/Private/NiagaraDataInterfaceRigidMeshCollisionQuery.cpp
--- a/Engine/Plugins/FX/Niagara/Source/Niagara/Private/NiagaraDataInterfaceRigidMeshCollisionQuery.cpp
+++ b/Engine/Plugins/FX/Niagara/Source/Niagara/Private/NiagaraDataInterfaceRigidMeshCollisionQuery.cpp
@@ -62,6 +62,7 @@ BEGIN_SHADER_PARAMETER_STRUCT(FShaderParameters,)
 	SHADER_PARAMETER_SRV(Buffer<float4>,	ElementExtentBuffer)
 	SHADER_PARAMETER_SRV(Buffer<uint32>,	PhysicsTypeBuffer)
 	SHADER_PARAMETER_SRV(Buffer<uint32>,	DFIndexBuffer)
+	SHADER_PARAMETER(FVector3f,				SystemLWCTile)
 END_SHADER_PARAMETER_STRUCT()
 
 static const TCHAR* TemplateShaderFile = TEXT("/Plugin/FX/Niagara/Private/NiagaraDataInterfaceRigidMeshCollisionQuery.ush");
@@ -1421,6 +1422,8 @@ void UNiagaraDataInterfaceRigidMeshCollisionQuery::SetShaderParameters(const FNi
 	{
 		FNiagaraDistanceFieldHelper::SetMeshDistanceFieldParameters(Context.GetGraphBuilder(), DistanceFieldSceneData, *ShaderDistanceFieldObjectParameters, *ShaderDistanceFieldAtlasParameters, FNiagaraRenderer::GetDummyFloat4Buffer());
 	}
+
+	ShaderParameters->SystemLWCTile = Context.GetSystemLWCTile();
 }
 
 #if WITH_EDITOR
