- Blueprint library for capturing Niagara SimCache data

diff --git a/Engine/Plugins/FX/Niagara/Source/Niagara/Classes/NiagaraSimCache.h b/Engine/Plugins/FX/Niagara/Source/Niagara/Classes/NiagaraSimCache.h
--- a/Engine/Plugins/FX/Niagara/Source/Niagara/Classes/NiagaraSimCache.h
+++ b/Engine/Plugins/FX/Niagara/Source/Niagara/Classes/NiagaraSimCache.h
@@ -8,7 +8,7 @@
 
 class UNiagaraComponent;
 
-USTRUCT()
+USTRUCT(BlueprintType)
 struct FNiagaraSimCacheCreateParameters
 {
 	GENERATED_BODY()
@@ -23,14 +23,14 @@ struct FNiagaraSimCacheCreateParameters
 	When enabled allows the SimCache to be re-based.
 	i.e. World space emitters can be moved to the new component's location
 	*/
-	UPROPERTY(EditAnywhere, Category="SimCache")
+	UPROPERTY(EditAnywhere, BlueprintReadWrite, Category="SimCache")
 	uint32 bAllowRebasing : 1;
 
 	/**
 	When enabled Data Interface data will be stored in the SimCache.
 	This can result in a large increase to the cache size, depending on what Data Interfaces are used
 	*/
-	UPROPERTY(EditAnywhere, Category="SimCache")
+	UPROPERTY(EditAnywhere, BlueprintReadWrite, Category="SimCache")
 	uint32 bAllowDataInterfaceCaching : 1;
 
 	/**
@@ -45,14 +45,14 @@ struct FNiagaraSimCacheCreateParameters
 	List of Attributes to force include in the SimCache rebase, they should be the full path to the attribute
 	For example, MyEmitter.Particles.MyQuat would force the particle attribute MyQuat to be included for MyEmitter
 	*/
-	UPROPERTY(EditAnywhere, Category = "SimCache")
+	UPROPERTY(EditAnywhere, BlueprintReadWrite, Category = "SimCache")
 	TArray<FName> RebaseIncludeList;
 
 	/**
 	List of Attributes to force exclude from the SimCache rebase, they should be the full path to the attribute
 	For example, MyEmitter.Particles.MyQuat would force the particle attribute MyQuat to be included for MyEmitter
 	*/
-	UPROPERTY(EditAnywhere, Category = "SimCache")
+	UPROPERTY(EditAnywhere, BlueprintReadWrite, Category = "SimCache")
 	TArray<FName> RebaseExcludeList;
 };
 
@@ -213,7 +213,7 @@ struct FNiagaraSimCacheLayout
 	TArray<FNiagaraSimCacheDataBuffersLayout> EmitterLayouts;
 };
 
-UCLASS(Experimental)
+UCLASS(Experimental, BlueprintType)
 class NIAGARA_API UNiagaraSimCache : public UObject
 {
 	GENERATED_UCLASS_BODY()
@@ -251,6 +251,8 @@ class NIAGARA_API UNiagaraSimCache : public UObject
 	virtual bool IsReadyForFinishDestroy() override;
 	// UObject Interface
 
+	/** Is the simulation cache valid or not. */
+	UFUNCTION(BlueprintCallable, Category=NiagaraSimCache)
 	bool IsCacheValid() const { return SoftNiagaraSystem.IsNull() == false; }
 
 	void BeginWrite(FNiagaraSimCacheCreateParameters InCreateParameters, UNiagaraComponent* NiagaraComponent);
diff --git a/Engine/Plugins/FX/Niagara/Source/Niagara/Classes/NiagaraSimCacheFunctionLibrary.h b/Engine/Plugins/FX/Niagara/Source/Niagara/Classes/NiagaraSimCacheFunctionLibrary.h
new file mode 100644
--- /dev/null
+++ b/Engine/Plugins/FX/Niagara/Source/Niagara/Classes/NiagaraSimCacheFunctionLibrary.h
@@ -0,0 +1,73 @@
+// Copyright Epic Games, Inc. All Rights Reserved.
+
+#include "NiagaraSimCache.h"
+
+#include "Containers/Ticker.h"
+#include "Kismet/BlueprintFunctionLibrary.h"
+#include "Kismet/BlueprintAsyncActionBase.h"
+
+#include "NiagaraSimCacheFunctionLibrary.generated.h"
+
+UCLASS()
+class UAsyncNiagaraCaptureSimCache : public UBlueprintAsyncActionBase
+{
+	GENERATED_BODY()
+
+public:
+	DECLARE_DYNAMIC_MULTICAST_DELEGATE(FOnCaptureComplete);
+
+	UPROPERTY()
+	TObjectPtr<UNiagaraSimCache> CaptureSimCache;
+
+	UPROPERTY()
+	TObjectPtr<UNiagaraComponent> CaptureComponent;
+
+	UPROPERTY()
+	int32 CaptureNumFrames = 0;
+
+	UPROPERTY()
+	int32 CaptureFrameRate = 0;
+
+	UPROPERTY()
+	int32 CaptureFrameCounter = 0;
+
+	UPROPERTY(BlueprintAssignable)
+	FOnCaptureComplete CaptureComplete;
+
+	FTSTicker::FDelegateHandle TickerHandle;
+
+	virtual void Activate() override;
+	virtual void SetReadyToDestroy() override;
+	bool OnFrameTick(float DeltaTime);
+
+	/**
+	Capture multiple frames from the provided simulation into a SimCache until the simulation becomes inactive, completes or we hit the NumFrames limit.
+	Capture occurs at the end of each frame with the first frame being this frame.
+	CaptureRate allows you to reduce the rate of capture, i.e. a rate of 2 would capture frames 0, 2, 4, etc.
+	*/
+	UFUNCTION(BlueprintCallable, meta=(BlueprintInternalUseOnly="true", Category=NiagaraSimCache))
+	static UAsyncNiagaraCaptureSimCache* CaptureNiagaraSimCacheMultiFrame(UNiagaraSimCache* SimCache, FNiagaraSimCacheCreateParameters CreateParameters, UNiagaraComponent* NiagaraComponent, int32 NumFrames = 16, int32 CaptureRate = 1);
+
+	/**
+	Capture frames from the provided simulation into a SimCache until the simulation becomes inactive or completes.
+	Capture occurs at the end of each frame with the first frame being this frame.
+	CaptureRate allows you to reduce the rate of capture, i.e. a rate of 2 would capture frames 0, 2, 4, etc.
+	*/
+	UFUNCTION(BlueprintCallable, meta=(BlueprintInternalUseOnly="true", Category=NiagaraSimCache))
+	static UAsyncNiagaraCaptureSimCache* CaptureNiagaraSimCacheUntilComplete(UNiagaraSimCache* SimCache, FNiagaraSimCacheCreateParameters CreateParameters, UNiagaraComponent* NiagaraComponent, int32 CaptureRate = 1);
+};
+
+UCLASS()
+class NIAGARA_API UNiagaraSimCacheFunctionLibrary : public UBlueprintFunctionLibrary
+{
+	GENERATED_UCLASS_BODY()
+
+public:
+	/**
+	Captures the simulations current frame data into the SimCache.
+	This happens immediately so you may need to be careful with tick order of the component you are capturing from.
+	The return can be invalid if the component can not be captured for some reason (i.e. not active).
+	*/
+	UFUNCTION(BlueprintCallable, Category=NiagaraSimCache)
+	static bool CaptureNiagaraSimCacheImmediate(UNiagaraSimCache* SimCache, FNiagaraSimCacheCreateParameters CreateParameters, UNiagaraComponent* NiagaraComponent);
+};
diff --git a/Engine/Plugins/FX/Niagara/Source/Niagara/Private/NiagaraSimCache.cpp b/Engine/Plugins/FX/Niagara/Source/Niagara/Private/NiagaraSimCache.cpp
--- a/Engine/Plugins/FX/Niagara/Source/Niagara/Private/NiagaraSimCache.cpp
+++ b/Engine/Plugins/FX/Niagara/Source/Niagara/Private/NiagaraSimCache.cpp
@@ -27,6 +27,11 @@ struct FNiagaraSimCacheHelper
 
 	explicit FNiagaraSimCacheHelper(UNiagaraComponent* NiagaraComponent)
 	{
+		if ( NiagaraComponent == nullptr )
+		{
+			return;
+		}
+
 		NiagaraSystem = NiagaraComponent->GetAsset();
 		if ( NiagaraSystem == nullptr )
 		{
@@ -951,7 +956,7 @@ bool UNiagaraSimCache::Read(float TimeSeconds, FNiagaraSystemInstance* SystemIns
 		return false;
 	}
 
-	const float FrameTime		= (RelativeTime / DurationSeconds) * float(CacheFrames.Num() - 1);
+	const float FrameTime		= DurationSeconds > 0.0f ? (RelativeTime / DurationSeconds) * float(CacheFrames.Num() - 1) : 0.0f;
 	const float FrameIndex		= FMath::Floor(FrameTime);
 	const float FrameFraction	= FrameTime - float(FrameIndex);
 
@@ -973,7 +978,7 @@ bool UNiagaraSimCache::ReadFrame(int32 FrameIndex, float FrameFraction, FNiagara
 	{
 		RebaseTransform = AttachComponent->GetComponentToWorld();
 		RebaseTransform.AddToTranslation(FVector(SystemInstance->GetLWCTile()) * -FLargeWorldRenderScalar::GetTileSize());
-		RebaseTransform = CacheFrame.LocalToWorld * RebaseTransform;
+		RebaseTransform = RebaseTransform * CacheFrame.LocalToWorld.Inverse();
 	}
 
 	Helper.SystemInstance->LocalBounds = CacheFrame.SystemData.LocalBounds;
diff --git a/Engine/Plugins/FX/Niagara/Source/Niagara/Private/NiagaraSimCacheFunctionLibrary.cpp b/Engine/Plugins/FX/Niagara/Source/Niagara/Private/NiagaraSimCacheFunctionLibrary.cpp
new file mode 100644
--- /dev/null
+++ b/Engine/Plugins/FX/Niagara/Source/Niagara/Private/NiagaraSimCacheFunctionLibrary.cpp
@@ -0,0 +1,120 @@
+// Copyright Epic Games, Inc. All Rights Reserved.
+
+#include "NiagaraSimCacheFunctionLibrary.h"
+#include "NiagaraComponent.h"
+
+void UAsyncNiagaraCaptureSimCache::Activate()
+{
+	Super::Activate();
+
+	if ( CaptureComponent != nullptr )
+	{
+		if (UWorld* OwnerWorld = CaptureComponent->GetWorld())
+		{
+			TickerHandle = FTSTicker::GetCoreTicker().AddTicker(FTickerDelegate::CreateUObject(this, &UAsyncNiagaraCaptureSimCache::OnFrameTick), 0);
+		}
+	}
+
+	if ( CaptureSimCache == nullptr || CaptureComponent == nullptr || TickerHandle.IsValid() == false )
+	{
+		SetReadyToDestroy();
+	}
+}
+
+void UAsyncNiagaraCaptureSimCache::SetReadyToDestroy()
+{
+	Super::SetReadyToDestroy();
+
+	if (TickerHandle.IsValid())
+	{
+		FTSTicker::GetCoreTicker().RemoveTicker(TickerHandle);
+	}
+
+	if ( CaptureSimCache != nullptr )
+	{
+		CaptureSimCache->EndWrite();
+	}
+	CaptureComplete.Broadcast();
+}
+
+bool UAsyncNiagaraCaptureSimCache::OnFrameTick(float DeltaTime)
+{
+	// Component invalid or not active?  If so complete the cache recording
+	if ( CaptureComponent == nullptr || !CaptureComponent->IsActive() )
+	{
+		SetReadyToDestroy();
+		return true;
+	}
+
+	// Should we record this frame?
+	if ( (CaptureFrameCounter % CaptureFrameRate) == 0 )
+	{
+		if ( CaptureSimCache != nullptr )
+		{
+			CaptureSimCache->WriteFrame(CaptureComponent);
+		}
+	}
+	++CaptureFrameCounter;
+
+	// Have we recorded all the frames we need?
+	// Note: the -1 is because T0 was the initial frame
+	if ( (CaptureNumFrames > 0) && (CaptureFrameCounter > (CaptureFrameRate * (CaptureNumFrames - 1))) )
+	{
+		SetReadyToDestroy();
+		return true;
+	}
+
+	return true;
+}
+
+UAsyncNiagaraCaptureSimCache* UAsyncNiagaraCaptureSimCache::CaptureNiagaraSimCacheMultiFrame(UNiagaraSimCache* SimCache, FNiagaraSimCacheCreateParameters CreateParameters, UNiagaraComponent* NiagaraComponent, int32 NumFrames, int32 CaptureRate)
+{
+	UAsyncNiagaraCaptureSimCache* CaptureAction = NewObject<UAsyncNiagaraCaptureSimCache>();
+	CaptureAction->CaptureSimCache = SimCache;
+	CaptureAction->CaptureComponent = NiagaraComponent;
+	CaptureAction->CaptureNumFrames = FMath::Max(1, NumFrames);
+	CaptureAction->CaptureFrameRate = FMath::Max(1, CaptureRate);
+	CaptureAction->CaptureFrameCounter = 0;
+
+	if (SimCache != nullptr)
+	{
+		SimCache->BeginWrite(CreateParameters, NiagaraComponent);
+	}
+
+	return CaptureAction;
+}
+
+UAsyncNiagaraCaptureSimCache* UAsyncNiagaraCaptureSimCache::CaptureNiagaraSimCacheUntilComplete(UNiagaraSimCache* SimCache, FNiagaraSimCacheCreateParameters CreateParameters, UNiagaraComponent* NiagaraComponent, int32 CaptureRate)
+{
+	UAsyncNiagaraCaptureSimCache* CaptureAction = NewObject<UAsyncNiagaraCaptureSimCache>();
+	CaptureAction->CaptureSimCache = SimCache;
+	CaptureAction->CaptureComponent = NiagaraComponent;
+	CaptureAction->CaptureNumFrames = 0;
+	CaptureAction->CaptureFrameRate = FMath::Max(1, CaptureRate);
+	CaptureAction->CaptureFrameCounter = 0;
+
+	if (SimCache != nullptr)
+	{
+		SimCache->BeginWrite(CreateParameters, NiagaraComponent);
+	}
+
+	return CaptureAction;
+}
+
+UNiagaraSimCacheFunctionLibrary::UNiagaraSimCacheFunctionLibrary(const FObjectInitializer& ObjectInitializer)
+	: Super(ObjectInitializer)
+{
+}
+
+bool UNiagaraSimCacheFunctionLibrary::CaptureNiagaraSimCacheImmediate(UNiagaraSimCache* SimCache, FNiagaraSimCacheCreateParameters CreateParameters, UNiagaraComponent* NiagaraComponent)
+{
+	if ( SimCache == nullptr || NiagaraComponent == nullptr )
+	{
+		return false;
+	}
+
+	SimCache->BeginWrite(CreateParameters, NiagaraComponent);
+	SimCache->WriteFrame(NiagaraComponent);
+	SimCache->EndWrite();
+	return SimCache->IsCacheValid();
+}
diff --git a/Engine/Plugins/FX/Niagara/Source/NiagaraEditor/Private/AssetTypeActions/AssetTypeActions_NiagaraSimCache.h b/Engine/Plugins/FX/Niagara/Source/NiagaraEditor/Private/AssetTypeActions/AssetTypeActions_NiagaraSimCache.h
--- a/Engine/Plugins/FX/Niagara/Source/NiagaraEditor/Private/AssetTypeActions/AssetTypeActions_NiagaraSimCache.h
+++ b/Engine/Plugins/FX/Niagara/Source/NiagaraEditor/Private/AssetTypeActions/AssetTypeActions_NiagaraSimCache.h
@@ -15,5 +15,6 @@ public:
 	virtual FColor GetTypeColor() const override;
 	virtual UClass* GetSupportedClass() const override;
 	virtual void OpenAssetEditor(const TArray<UObject*>& InObjects, TSharedPtr<class IToolkitHost> EditWithinLevelEditor = TSharedPtr<IToolkitHost>()) override;
+	virtual bool CanFilter() override { return true; }
 	virtual uint32 GetCategories() override { return FNiagaraEditorModule::GetAssetCategory(); }
-};
\ No newline at end of file
+};
diff --git a/Engine/Plugins/FX/Niagara/Source/NiagaraEditor/Private/NiagaraSimCacheFactoryNew.cpp b/Engine/Plugins/FX/Niagara/Source/NiagaraEditor/Private/NiagaraSimCacheFactoryNew.cpp
--- a/Engine/Plugins/FX/Niagara/Source/NiagaraEditor/Private/NiagaraSimCacheFactoryNew.cpp
+++ b/Engine/Plugins/FX/Niagara/Source/NiagaraEditor/Private/NiagaraSimCacheFactoryNew.cpp
@@ -12,7 +12,7 @@ UNiagaraSimCacheFactoryNew::UNiagaraSimCacheFactoryNew(const FObjectInitializer&
 
 bool UNiagaraSimCacheFactoryNew::ShouldShowInNewMenu() const
 {
-	return false;
+	return true;
 }
 
 UObject* UNiagaraSimCacheFactoryNew::FactoryCreateNew(UClass* InClass, UObject* InParent, FName InName, EObjectFlags Flags, UObject* Context, FFeedbackContext* Warn)
