- Check UntypedDI is not nullptr

#rnx

diff --git a/Engine/Plugins/FX/Niagara/Source/Niagara/Private/NiagaraFunctionLibrary.cpp b/Engine/Plugins/FX/Niagara/Source/Niagara/Private/NiagaraFunctionLibrary.cpp
--- a/Engine/Plugins/FX/Niagara/Source/Niagara/Private/NiagaraFunctionLibrary.cpp
+++ b/Engine/Plugins/FX/Niagara/Source/Niagara/Private/NiagaraFunctionLibrary.cpp
@@ -649,7 +649,7 @@ UNiagaraDataInterface* UNiagaraFunctionLibrary::GetDataInterface(UClass* DIClass
 	}
 
 	UNiagaraDataInterface* UntypedDI = OverrideParameters.GetDataInterface(*Index);
-	if (UntypedDI->IsA(DIClass))
+	if (UntypedDI != nullptr && UntypedDI->IsA(DIClass))
 	{
 		return UntypedDI;
 	}
