Trying to resolve issues where the AlignParticlesWithCollisionPlane and ConstrainPositionToPlane were mapped to inputs with different names and different guids in previous versions. This breaks any content that was relying on that. Added the ability for an input to define multiple aliases to resolve this. Also added a commandlet to output xml about the various modules so that we can safely compare across versions of UE. Updated the two modules that were most affected by this.

diff --git a/Engine/Plugins/FX/Niagara/Source/Niagara/Private/NiagaraCommon.cpp b/Engine/Plugins/FX/Niagara/Source/Niagara/Private/NiagaraCommon.cpp
--- a/Engine/Plugins/FX/Niagara/Source/Niagara/Private/NiagaraCommon.cpp
+++ b/Engine/Plugins/FX/Niagara/Source/Niagara/Private/NiagaraCommon.cpp
@@ -1314,7 +1314,9 @@ void AliasStandardParameter(const FNiagaraAliasContext& InContext, TArray<FStrin
 	// If there are more than 2 parts in the parameter, and it's not an assignment node, then
 	// it may be a module specific dataset value, so the 2nd position must be checked for the
 	// module mapping.
-	if (InOutSplitName.Num() > 1 && ModuleMapping.IsSet() &&
+	// Examples to match: Particles.Module.CustomOutput, Transient.Module.PhysicsVar
+	// Examples *not* to match: Module.SpawnRate, where SpawnRate is also the name of the module.
+	if (InOutSplitName.Num() > 2 && ModuleMapping.IsSet() &&
 		InOutSplitName[1] == ModuleMapping.GetValue().Key)
 	{
 		InOutSplitName[1] = ModuleMapping.GetValue().Value;
diff --git a/Engine/Plugins/FX/Niagara/Source/Niagara/Public/NiagaraTypes.h b/Engine/Plugins/FX/Niagara/Source/Niagara/Public/NiagaraTypes.h
--- a/Engine/Plugins/FX/Niagara/Source/Niagara/Public/NiagaraTypes.h
+++ b/Engine/Plugins/FX/Niagara/Source/Niagara/Public/NiagaraTypes.h
@@ -819,6 +819,9 @@ struct NIAGARA_API FNiagaraVariableMetaData
 	UPROPERTY(EditAnywhere, Category = "Variable", meta = (ToolTip = "If set, this attribute is visually displayed as a child under the given parent attribute. Currently, only static switches are supported as parent attributes!", SkipForCompileHash = "true"))
 	FName ParentAttribute;
 
+	UPROPERTY(EditAnywhere, Category = "Variable", DisplayName = "Alternate Aliases For Variable", AdvancedDisplay, meta = (ToolTip = "List of alternate/previous names for this variable. Note that this is not normally needed if you rename through the UX. However, if you delete and then add a different variable, intending for it to match, you will likely want to add the prior name here.\n\nYou may need to restart and reload assets after making this change to have it take effect on already loaded assets."))
+	TArray<FName> AlternateAliases;
+
 	bool GetIsStaticSwitch_DEPRECATED() const { return bIsStaticSwitch_DEPRECATED; };
 
 	int32 GetStaticSwitchDefaultValue_DEPRECATED() const { return StaticSwitchDefaultValue_DEPRECATED; };
diff --git a/Engine/Plugins/FX/Niagara/Source/NiagaraEditor/Private/Commandlets/NiagaraDumpModuleInfoCommandlet.cpp b/Engine/Plugins/FX/Niagara/Source/NiagaraEditor/Private/Commandlets/NiagaraDumpModuleInfoCommandlet.cpp
new file mode 100644
--- /dev/null
+++ b/Engine/Plugins/FX/Niagara/Source/NiagaraEditor/Private/Commandlets/NiagaraDumpModuleInfoCommandlet.cpp
@@ -0,0 +1,392 @@
+// Copyright Epic Games, Inc. All Rights Reserved.
+
+#include "Commandlets/NiagaraDumpModuleInfoCommandlet.h"
+
+#include "AssetRegistry/ARFilter.h"
+#include "AssetRegistry/AssetData.h"
+#include "AssetRegistry/AssetRegistryModule.h"
+#include "CollectionManagerModule.h"
+#include "CollectionManagerTypes.h"
+#include "HAL/FileManager.h"
+#include "ICollectionManager.h"
+#include "Misc/OutputDeviceArchiveWrapper.h"
+#include "Misc/PackageName.h"
+#include "Misc/Paths.h"
+#include "Modules/ModuleManager.h"
+#include "NiagaraSystem.h"
+#include "NiagaraScriptSource.h"
+#include "NiagaraGraph.h"
+#include "NiagaraScriptVariable.h"
+#include "NiagaraParameterDefinitions.h"
+
+DEFINE_LOG_CATEGORY_STATIC(LogNiagaraDumpModuleInfoCommandlet, Log, All);
+
+UNiagaraDumpModuleInfoCommandlet::UNiagaraDumpModuleInfoCommandlet(const FObjectInitializer& ObjectInitializer)
+	: Super(ObjectInitializer)
+{
+}
+
+int32 UNiagaraDumpModuleInfoCommandlet::Main(const FString& Params)
+{
+	TArray<FString> Tokens;
+	TArray<FString> Switches;
+
+	ParseCommandLine(*Params, Tokens, Switches);
+
+	if (!FParse::Value(*Params, TEXT("AuditOutputFolder="), AuditOutputFolder))
+	{
+		// No output folder specified. Use the default folder.
+		AuditOutputFolder = FPaths::ProjectSavedDir() / TEXT("Audit");
+	}
+
+	// Add a timestamp to the folder
+	AuditOutputFolder /= FDateTime::Now().ToString();
+
+	FParse::Value(*Params, TEXT("FilterCollection="), FilterCollection);
+
+	// Package Paths
+	FString PackagePathsString;
+	if (FParse::Value(*Params, TEXT("PackagePaths="), PackagePathsString, false))
+	{
+		TArray<FString> PackagePathsStrings;
+		PackagePathsString.ParseIntoArray(PackagePathsStrings, TEXT(","));
+		for (const FString& v : PackagePathsStrings)
+		{
+			PackagePaths.Add(FName(v));
+		}
+	}
+
+	if (PackagePaths.Num() == 0)
+	{
+		//PackagePaths.Add(FName(TEXT("/Game")));
+		PackagePaths.Add(FName(TEXT("/Plugins")));
+		PackagePaths.Add(FName(TEXT("/Niagara")));
+		PackagePaths.Add(FName(TEXT("/NiagaraFluids")));
+		PackagePaths.Add(FName(TEXT("/Engine")));
+	}
+
+	ProcessNiagaraScripts();
+	ProcessNiagaraParameterDefinitions();
+
+	return 0;
+}
+
+void UNiagaraDumpModuleInfoCommandlet::ProcessNiagaraParameterDefinitions()
+{
+	FAssetRegistryModule& AssetRegistryModule = FModuleManager::LoadModuleChecked<FAssetRegistryModule>(TEXT("AssetRegistry"));
+	IAssetRegistry& AssetRegistry = AssetRegistryModule.Get();
+	AssetRegistry.SearchAllAssets(true);
+
+	FARFilter Filter;
+	Filter.PackagePaths = PackagePaths;
+	Filter.bRecursivePaths = true;
+
+	Filter.ClassPaths.Add(UNiagaraParameterDefinitions::StaticClass()->GetClassPathName());
+	if (!FilterCollection.IsEmpty())
+	{
+		FCollectionManagerModule& CollectionManagerModule = FCollectionManagerModule::GetModule();
+		CollectionManagerModule.Get().GetObjectsInCollection(FName(*FilterCollection), ECollectionShareType::CST_All, Filter.ObjectPaths, ECollectionRecursionFlags::SelfAndChildren);
+	}
+
+	TArray<FAssetData> AssetList;
+	AssetRegistry.GetAssets(Filter, AssetList);
+
+	//  Iterate over all scripts
+	const FString DevelopersFolder = FPackageName::FilenameToLongPackageName(FPaths::GameDevelopersDir().LeftChop(1));
+	FString LastPackageName = TEXT("");
+	UPackage* CurrentPackage = nullptr;
+	for (const FAssetData& AssetIt : AssetList)
+	{
+		const FString SystemName = AssetIt.ObjectPath.ToString();
+		const FString PackageName = AssetIt.PackageName.ToString();
+
+		if (PackageName.StartsWith(DevelopersFolder))
+		{
+			// Skip developer folders
+			continue;
+		}
+
+		if (PackageName != LastPackageName)
+		{
+			UPackage* Package = ::LoadPackage(nullptr, *PackageName, LOAD_None);
+			if (Package != nullptr)
+			{
+				LastPackageName = PackageName;
+				Package->FullyLoad();
+				CurrentPackage = Package;
+			}
+			else
+			{
+				UE_LOG(LogNiagaraDumpModuleInfoCommandlet, Warning, TEXT("Failed to load package %s processing %s"), *PackageName, *SystemName);
+				CurrentPackage = nullptr;
+			}
+		}
+
+		const FString ShorterSystemName = AssetIt.AssetName.ToString();
+		UNiagaraParameterDefinitions* NiagaraParameterDefinitions = FindObject<UNiagaraParameterDefinitions>(CurrentPackage, *ShorterSystemName);
+		if (NiagaraParameterDefinitions == nullptr)
+		{
+			UE_LOG(LogNiagaraDumpModuleInfoCommandlet, Warning, TEXT("Failed to load NiagaraParameterDefinitions %s"), *SystemName);
+			continue;
+		}
+
+		IFileManager::Get().MakeDirectory(*(AuditOutputFolder / TEXT("_Definitions_") + ShorterSystemName));
+
+		const FString MetaDataFileName = AuditOutputFolder / TEXT("_Definitions_") + ShorterSystemName / TEXT("Definitions.xml");
+		TUniquePtr<FArchive> FileArchive(IFileManager::Get().CreateDebugFileWriter(*MetaDataFileName));
+		TUniquePtr<FOutputDeviceArchiveWrapper> OutputStream(new FOutputDeviceArchiveWrapper(FileArchive.Get()));
+
+		const TArray<UNiagaraScriptVariable*>& ScriptPtrVars =  NiagaraParameterDefinitions->GetParametersConst();
+		TArray< UNiagaraScriptVariable*> ScriptVars;
+
+		for (TObjectPtr<UNiagaraScriptVariable> Var : ScriptPtrVars)
+		{
+			if (Var.IsNull())
+				continue;
+			ScriptVars.Add(Var);
+		}
+
+
+		ScriptVars.StableSort([](const UNiagaraScriptVariable& Lhs, const UNiagaraScriptVariable& Rhs)
+			{
+				return Lhs.Variable.GetName().ToString().Compare(Rhs.Variable.GetName().ToString(), ESearchCase::IgnoreCase) < 0;
+			});
+
+		OutputStream->Log(TEXT("<?xml version='1.0' ?>"));
+		OutputStream->Log(TEXT("<Info>"));
+		OutputStream->Logf(TEXT("\t<Path>%s</Path>"), *GetPathNameSafe(NiagaraParameterDefinitions));
+
+		for (UNiagaraScriptVariable*& ScriptVar : ScriptVars)
+		{
+			OutputStream->Log(TEXT("\t<ScriptVar>"));
+			OutputStream->Logf(TEXT("\t\t<Name>%s</Name>"), *ScriptVar->Variable.GetName().ToString());
+			OutputStream->Logf(TEXT("\t\t<Description>%s</Description>"), *ScriptVar->Metadata.Description.ToString());			
+			OutputStream->Logf(TEXT("\t<Guid>%s</Guid>"), *ScriptVar->Metadata.GetVariableGuid().ToString());
+			OutputStream->Log(TEXT("\t</ScriptVar>"));
+		}
+
+		OutputStream->Log(TEXT("</Info>"));
+	}
+
+	// Probably don't need to do this, but just in case we have any 'hanging' packages 
+	// and more processing steps are added later, let's clean up everything...
+	::CollectGarbage(RF_NoFlags);
+}
+
+void UNiagaraDumpModuleInfoCommandlet::ProcessNiagaraScripts()
+{
+	FAssetRegistryModule& AssetRegistryModule = FModuleManager::LoadModuleChecked<FAssetRegistryModule>(TEXT("AssetRegistry"));
+	IAssetRegistry& AssetRegistry = AssetRegistryModule.Get();
+	AssetRegistry.SearchAllAssets(true);
+
+	FARFilter Filter;
+	Filter.PackagePaths = PackagePaths;
+	Filter.bRecursivePaths = true;
+
+	Filter.ClassPaths.Add(UNiagaraScript::StaticClass()->GetClassPathName());
+	if (!FilterCollection.IsEmpty())
+	{
+		FCollectionManagerModule& CollectionManagerModule = FCollectionManagerModule::GetModule();
+		CollectionManagerModule.Get().GetObjectsInCollection(FName(*FilterCollection), ECollectionShareType::CST_All, Filter.ObjectPaths, ECollectionRecursionFlags::SelfAndChildren);
+	}
+
+	TArray<FAssetData> AssetList;
+	AssetRegistry.GetAssets(Filter, AssetList);
+
+	const static UEnum* UsageEnum = StaticEnum<ENiagaraScriptUsage>();
+
+	const double StartProcessNiagaraSystemsTime = FPlatformTime::Seconds();
+
+	//  Iterate over all scripts
+	const FString DevelopersFolder = FPackageName::FilenameToLongPackageName(FPaths::GameDevelopersDir().LeftChop(1));
+	FString LastPackageName = TEXT("");
+	UPackage* CurrentPackage = nullptr;
+	for (const FAssetData& AssetIt : AssetList)
+	{
+		const FString SystemName = AssetIt.ObjectPath.ToString();
+		const FString PackageName = AssetIt.PackageName.ToString();
+
+		if (PackageName.StartsWith(DevelopersFolder))
+		{
+			// Skip developer folders
+			continue;
+		}
+
+		if (PackageName != LastPackageName)
+		{
+			UPackage* Package = ::LoadPackage(nullptr, *PackageName, LOAD_None);
+			if (Package != nullptr)
+			{
+				LastPackageName = PackageName;
+				Package->FullyLoad();
+				CurrentPackage = Package;
+			}
+			else
+			{
+				UE_LOG(LogNiagaraDumpModuleInfoCommandlet, Warning, TEXT("Failed to load package %s processing %s"), *PackageName, *SystemName);
+				CurrentPackage = nullptr;
+			}
+		}
+
+		const FString ShorterSystemName = AssetIt.AssetName.ToString();
+		UNiagaraScript* NiagaraScript = FindObject<UNiagaraScript>(CurrentPackage, *ShorterSystemName);
+		if (NiagaraScript == nullptr)
+		{
+			UE_LOG(LogNiagaraDumpModuleInfoCommandlet, Warning, TEXT("Failed to load Niagara script %s"), *SystemName);
+			continue;
+		}
+
+	
+		IFileManager::Get().MakeDirectory(*(AuditOutputFolder / ShorterSystemName));
+
+		TArray<FNiagaraAssetVersion> Versions = NiagaraScript->GetAllAvailableVersions();
+		for (FNiagaraAssetVersion Version : Versions)
+		{
+			const UNiagaraScriptSource* ScriptSrc = Cast<UNiagaraScriptSource>(NiagaraScript->GetSource(Version.VersionGuid));
+			TSharedPtr<FNiagaraVersionDataAccessor> VersionDataAccessor = NiagaraScript->GetVersionDataAccessor(Version.VersionGuid);
+			if (ScriptSrc)
+			{
+				UNiagaraGraph* FunctionGraph = ScriptSrc->NodeGraph;
+				FText VersionDesc;
+				if (VersionDataAccessor)
+					VersionDesc = VersionDataAccessor->GetVersionChangeDescription();
+
+				if (FunctionGraph)
+				{
+					const FString MetaDataFileName = AuditOutputFolder / ShorterSystemName / Version.VersionGuid.ToString() + TEXT(".xml");
+					TUniquePtr<FArchive> FileArchive(IFileManager::Get().CreateDebugFileWriter(*MetaDataFileName));
+					TUniquePtr<FOutputDeviceArchiveWrapper> OutputStream(new FOutputDeviceArchiveWrapper(FileArchive.Get()));
+
+					const TMap<FNiagaraVariable, TObjectPtr<UNiagaraScriptVariable>>& AllVars = FunctionGraph->GetAllMetaData();
+					TArray< TObjectPtr<UNiagaraScriptVariable>> ScriptPtrVars;
+					AllVars.GenerateValueArray(ScriptPtrVars);
+					TArray< UNiagaraScriptVariable*> ScriptVars;
+
+					for (TObjectPtr<UNiagaraScriptVariable> Var : ScriptPtrVars)
+					{
+						if (Var.IsNull())
+							continue;
+						ScriptVars.Add(Var);
+					}
+					
+
+					ScriptVars.StableSort([](const UNiagaraScriptVariable& Lhs, const UNiagaraScriptVariable& Rhs)
+					{
+							return Lhs.Variable.GetName().ToString().Compare(Rhs.Variable.GetName().ToString(), ESearchCase::IgnoreCase) < 0;
+					});
+
+					OutputStream->Log(TEXT("<?xml version='1.0' ?>"));
+					OutputStream->Log(TEXT("<Info>"));
+					OutputStream->Logf(TEXT("\t<VersionNumber>%d . %d</VersionNumber>"), Version.MajorVersion, Version.MinorVersion);
+					OutputStream->Logf(TEXT("\t<VersionDesc>%s</VersionDesc>"), *VersionDesc.ToString());
+					OutputStream->Logf(TEXT("\t<Path>%s</Path>"), *GetPathNameSafe(FunctionGraph));
+
+					for (UNiagaraScriptVariable*& ScriptVar : ScriptVars)
+					{
+						OutputStream->Log(TEXT("\t<ScriptVar>"));
+						OutputStream->Logf(TEXT("\t\t<Name>%s</Name>"), *ScriptVar->Variable.GetName().ToString());
+						OutputStream->Logf(TEXT("\t\t<Type>%s</Type>"), *ScriptVar->Variable.GetType().GetName());
+						OutputStream->Logf(TEXT("\t<Guid>%s</Guid>"), *ScriptVar->Metadata.GetVariableGuid().ToString());
+						OutputStream->Logf(TEXT("\t\t<Description>%s</Description>"), *ScriptVar->Metadata.Description.ToString());
+						OutputStream->Logf(TEXT("\t\t<Category>%s</Category>"), *ScriptVar->Metadata.CategoryName.ToString());
+
+						OutputStream->Logf(TEXT("\t\t<Advanced>%s</Advanced>"), ScriptVar->Metadata.bAdvancedDisplay ? TEXT("true") : TEXT("false"));
+						OutputStream->Logf(TEXT("\t\t<DisplayInOverviewStack>%s</DisplayInOverviewStack>"), ScriptVar->Metadata.bDisplayInOverviewStack ? TEXT("true") : TEXT("false"));
+						OutputStream->Logf(TEXT("\t\t<EnableBoolOverride>%s</EnableBoolOverride>"), ScriptVar->Metadata.bEnableBoolOverride ? TEXT("true") : TEXT("false"));
+						OutputStream->Logf(TEXT("\t\t<InlineEditConditionToggle>%s</InlineEditConditionToggle>"), ScriptVar->Metadata.bInlineEditConditionToggle ? TEXT("true") : TEXT("false"));
+						for (const FName& Alias : ScriptVar->Metadata.AlternateAliases)
+							OutputStream->Logf(TEXT("\t\t<Alias>%s</Alias>"), *Alias.ToString());
+						OutputStream->Log(TEXT("\t</ScriptVar>"));
+					}
+					
+					OutputStream->Log(TEXT("</Info>"));
+				}
+			}
+		}
+		
+		/*
+		* // sort the data alphabetically (based on MetaData.FullName
+	ScriptMetaData.StableSort([](const FScriptMetaData& Lhs, const FScriptMetaData& Rhs)
+	{
+		return Lhs.FullName.Compare(Rhs.FullName, ESearchCase::IgnoreCase) < 0;
+	});
+		*/
+	}
+
+	
+	// Probably don't need to do this, but just in case we have any 'hanging' packages 
+	// and more processing steps are added later, let's clean up everything...
+	::CollectGarbage(RF_NoFlags);
+
+	double ProcessNiagaraSystemsTime = FPlatformTime::Seconds() - StartProcessNiagaraSystemsTime;
+	UE_LOG(LogNiagaraDumpModuleInfoCommandlet, Log, TEXT("Took %5.3f seconds to process referenced Niagara scripts..."), ProcessNiagaraSystemsTime);
+}
+
+/*void UNiagaraDumpModuleInfoCommandlet::DumpByteCode(const UNiagaraScript* Script, const FString& FilePath)
+{
+	if (!Script)
+	{
+		return;
+	}
+
+	const auto& ExecData = Script->GetVMExecutableData();
+
+	FScriptMetaData& MetaData = ScriptMetaData.AddZeroed_GetRef();
+	MetaData.FullName = FilePath;
+	MetaData.RegisterCount = ExecData.NumTempRegisters;
+	MetaData.OpCount = ExecData.LastOpCount;
+	MetaData.ConstantCount = ExecData.InternalParameters.GetTableSize() / 4;
+	MetaData.AttributeCount = 0;
+
+	for (const FNiagaraVariableBase& Var : ExecData.Attributes)
+	{
+		MetaData.AttributeCount += Var.GetType().GetSize() / 4;
+	}
+
+	const static UEnum* VmOpEnum = StaticEnum<EVectorVMOp>();
+
+	if (Script)
+	{
+		const FString FullFilePath = AuditOutputFolder / FilePath;
+
+		TUniquePtr<FArchive> FileArchive(IFileManager::Get().CreateDebugFileWriter(*FullFilePath));
+		if (!FileArchive)
+		{
+			UE_LOG(LogNiagaraDumpModuleInfoCommandlet, Warning, TEXT("Failed to create output stream %s"), *FullFilePath);
+		}
+
+		TUniquePtr<FOutputDeviceArchiveWrapper> OutputStream(new FOutputDeviceArchiveWrapper(FileArchive.Get()));
+
+		const auto& VMData = Script->GetVMExecutableData();
+
+		TArray<FString> AssemblyLines;
+		VMData.LastAssemblyTranslation.ParseIntoArrayLines(AssemblyLines, false);
+
+		// we want to translate all instances of OP_X into the string from the VMVector enum
+		for (FString& CurrentLine : AssemblyLines)
+		{
+			int OpStartIndex = CurrentLine.Find(TEXT("OP_"), ESearchCase::CaseSensitive);
+			if (OpStartIndex != INDEX_NONE)
+			{
+				const int32 OpEndStartSearch = OpStartIndex + 3;
+				int OpEndIndex = CurrentLine.Find(TEXT("("), ESearchCase::IgnoreCase, ESearchDir::FromStart, OpEndStartSearch);
+				if (OpEndIndex == INDEX_NONE)
+				{
+					OpEndIndex = CurrentLine.Find(TEXT(";"), ESearchCase::IgnoreCase, ESearchDir::FromStart, OpEndStartSearch);
+				}
+
+				if (OpEndIndex != INDEX_NONE)
+				{
+					const FString OpIndexString = CurrentLine.Mid(OpEndStartSearch, OpEndIndex - OpEndStartSearch);
+					int OpIndexValue = FCString::Atoi(*OpIndexString);
+
+					FString LinePrefix = CurrentLine.Left(OpStartIndex);
+					FString LineSuffix = CurrentLine.RightChop(OpEndIndex);
+
+					CurrentLine = LinePrefix + TEXT(":") + *VmOpEnum->GetNameStringByValue(OpIndexValue) + LineSuffix;
+				}
+			}
+			OutputStream->Log(CurrentLine);
+		}
+	}
+}*/
\ No newline at end of file
diff --git a/Engine/Plugins/FX/Niagara/Source/NiagaraEditor/Private/ViewModels/Stack/NiagaraStackGraphUtilities.cpp b/Engine/Plugins/FX/Niagara/Source/NiagaraEditor/Private/ViewModels/Stack/NiagaraStackGraphUtilities.cpp
--- a/Engine/Plugins/FX/Niagara/Source/NiagaraEditor/Private/ViewModels/Stack/NiagaraStackGraphUtilities.cpp
+++ b/Engine/Plugins/FX/Niagara/Source/NiagaraEditor/Private/ViewModels/Stack/NiagaraStackGraphUtilities.cpp
@@ -3194,6 +3194,11 @@ TSharedRef<TMap<FName, FGuid>> GetVariableNameToVariableIdMap(UNiagaraNodeFuncti
 			if (VariableGuid.IsValid())
 			{
 				VariableNameToVariableIdMap->FindOrAdd(VariableMetadataPair.Key.GetName()) = VariableGuid;
+                // Add in alternate names as well so that we can map them cleanly in the second pass.
+				for (const FName& AltName : VariableMetadataPair.Value->Metadata.AlternateAliases)
+				{
+					VariableNameToVariableIdMap->FindOrAdd(AltName) = VariableGuid;
+				}
 			}
 		}
 	}
diff --git a/Engine/Plugins/FX/Niagara/Source/NiagaraEditor/Public/Commandlets/NiagaraDumpModuleInfoCommandlet.h b/Engine/Plugins/FX/Niagara/Source/NiagaraEditor/Public/Commandlets/NiagaraDumpModuleInfoCommandlet.h
new file mode 100644
--- /dev/null
+++ b/Engine/Plugins/FX/Niagara/Source/NiagaraEditor/Public/Commandlets/NiagaraDumpModuleInfoCommandlet.h
@@ -0,0 +1,31 @@
+// Copyright Epic Games, Inc. All Rights Reserved.
+
+#pragma once
+
+#include "CoreMinimal.h"
+#include "UObject/ObjectMacros.h"
+#include "Commandlets/Commandlet.h"
+#include "NiagaraDumpModuleInfoCommandlet.generated.h"
+
+class UNiagaraScript;
+
+UCLASS(config=Editor)
+class UNiagaraDumpModuleInfoCommandlet : public UCommandlet
+{
+	GENERATED_UCLASS_BODY()
+
+	int32 Main(const FString& Params) override;
+
+	/** The folder in which the commandlet's output files will be stored */
+	FString AuditOutputFolder;
+
+	/** Only assets in this collection will be considered. If this is left blank, no assets will be filtered by collection */
+	FString FilterCollection;
+
+	/** Package paths to include */
+	TArray<FName> PackagePaths;
+
+private:
+	void ProcessNiagaraScripts();
+	void ProcessNiagaraParameterDefinitions();
+};
\ No newline at end of file
