Release reference to SourceComponent when source actor ends play.

diff --git a/Engine/Plugins/FX/Niagara/Source/Niagara/Private/DataInterface/NiagaraDataInterfaceStaticMesh.cpp b/Engine/Plugins/FX/Niagara/Source/Niagara/Private/DataInterface/NiagaraDataInterfaceStaticMesh.cpp
--- a/Engine/Plugins/FX/Niagara/Source/Niagara/Private/DataInterface/NiagaraDataInterfaceStaticMesh.cpp
+++ b/Engine/Plugins/FX/Niagara/Source/Niagara/Private/DataInterface/NiagaraDataInterfaceStaticMesh.cpp
@@ -2626,6 +2626,7 @@ bool UNiagaraDataInterfaceStaticMesh::CopyToInternal(UNiagaraDataInterface* Dest
 	}
 
 	UNiagaraDataInterfaceStaticMesh* OtherTyped = CastChecked<UNiagaraDataInterfaceStaticMesh>(Destination);
+	OtherTyped->UnbindSourceDelegates();
 	OtherTyped->SourceMode = SourceMode;
 #if WITH_EDITORONLY_DATA
 	OtherTyped->PreviewMesh = PreviewMesh;
@@ -2636,6 +2637,7 @@ bool UNiagaraDataInterfaceStaticMesh::CopyToInternal(UNiagaraDataInterface* Dest
 	OtherTyped->SectionFilter = SectionFilter;
 	OtherTyped->bUsePhysicsBodyVelocity = bUsePhysicsBodyVelocity;
 	OtherTyped->FilteredSockets = FilteredSockets;
+	OtherTyped->BindSourceDelegates();
 	return true;
 }
 
@@ -2841,8 +2843,40 @@ void UNiagaraDataInterfaceStaticMesh::SetSourceComponentFromBlueprints(UStaticMe
 {
 	// NOTE: When ChangeId changes the next tick will be skipped and a reset of the per-instance data will be initiated. 
 	++ChangeId;
+	UnbindSourceDelegates();
 	SourceComponent = ComponentToUse;
 	SoftSourceActor = ComponentToUse->GetOwner();
+	BindSourceDelegates();
+}
+
+void UNiagaraDataInterfaceStaticMesh::BindSourceDelegates()
+{
+	if (AActor* Source = SoftSourceActor.Get())
+	{
+		Source->OnEndPlay.AddDynamic(this, &UNiagaraDataInterfaceStaticMesh::OnSourceEndPlay);
+	}
+	else if (SourceComponent)
+	{
+		UE_CLOG(!UObjectBaseUtility::IsPendingKillEnabled(),
+			LogNiagara, Warning, TEXT("%s: Unable to bind OnEndPlay for actor-less source component %s, this may extend the lifetime of the component"),
+			*GetFullName(), *SourceComponent->GetPathName());
+	}
+}
+
+void UNiagaraDataInterfaceStaticMesh::UnbindSourceDelegates()
+{
+	if (AActor* Source = SoftSourceActor.Get())
+	{
+		Source->OnEndPlay.RemoveAll(this);
+	}
+}
+
+void UNiagaraDataInterfaceStaticMesh::OnSourceEndPlay(AActor* InSource, EEndPlayReason::Type Reason)
+{
+	// Increment change id in case we're able to find a new source component 
+	++ChangeId;
+	SoftSourceActor = nullptr;
+	SourceComponent = nullptr;
 }
 
 void UNiagaraDataInterfaceStaticMesh::SetDefaultMeshFromBlueprints(UStaticMesh* MeshToUse)
diff --git a/Engine/Plugins/FX/Niagara/Source/Niagara/Private/DataInterface/NiagaraDataInterfaceStaticMesh.h b/Engine/Plugins/FX/Niagara/Source/Niagara/Private/DataInterface/NiagaraDataInterfaceStaticMesh.h
--- a/Engine/Plugins/FX/Niagara/Source/Niagara/Private/DataInterface/NiagaraDataInterfaceStaticMesh.h
+++ b/Engine/Plugins/FX/Niagara/Source/Niagara/Private/DataInterface/NiagaraDataInterfaceStaticMesh.h
@@ -70,6 +70,7 @@ public:
 	UPROPERTY(EditAnywhere, Category = "Mesh")
 	TObjectPtr<UStaticMesh> DefaultMesh;
 
+protected:
 	/** The source actor from which to sample. Takes precedence over the direct mesh. Note that this can only be set when used as a user variable on a component in the world. */
 	UPROPERTY(EditAnywhere, Category = "Mesh", meta = (DisplayName = "Source Actor"))
 	TSoftObjectPtr<AActor> SoftSourceActor;
@@ -82,6 +83,7 @@ public:
 	/** The source component from which to sample. Takes precedence over the direct mesh. Not exposed to the user, only indirectly accessible from blueprints. */
 	UPROPERTY(Transient)
 	TObjectPtr<UStaticMeshComponent> SourceComponent;
+public:
 
 	/** Array of filters the can be used to limit sampling to certain sections of the mesh. */
 	UPROPERTY(EditAnywhere, Category = "Mesh")
@@ -153,6 +155,13 @@ public:
 	void SetDefaultMeshFromBlueprints(UStaticMesh* MeshToUse);
 
 protected:
+	// Bind/unbind delegates to release references to the source actor & component.
+	void UnbindSourceDelegates();
+	void BindSourceDelegates();
+
+	UFUNCTION()
+	void OnSourceEndPlay(AActor* InSource, EEndPlayReason::Type Reason);
+
 	// VM Vertex Sampling
 	void VMIsValidVertex(FVectorVMExternalFunctionContext& Context);
 	void VMRandomVertex(FVectorVMExternalFunctionContext& Context);
