- Exclude internal attributes from parameter store traversal

#rnx

diff --git a/Engine/Plugins/FX/Niagara/Source/Niagara/Classes/NiagaraConstants.h b/Engine/Plugins/FX/Niagara/Source/Niagara/Classes/NiagaraConstants.h
--- a/Engine/Plugins/FX/Niagara/Source/Niagara/Classes/NiagaraConstants.h
+++ b/Engine/Plugins/FX/Niagara/Source/Niagara/Classes/NiagaraConstants.h
@@ -236,6 +236,7 @@ struct NIAGARA_API FNiagaraConstants
 	static const FString StackContextNamespaceString;
 	static const FString SystemNamespaceString;
 	static const FString UserNamespaceString;
+	static const FString InternalNamespaceString;
 
 	static const int32 MaxCategoryNameLength;
 	static const int32 MaxParameterLength;
diff --git a/Engine/Plugins/FX/Niagara/Source/Niagara/Private/DataInterface/NiagaraDataInterfaceUtilities.cpp b/Engine/Plugins/FX/Niagara/Source/Niagara/Private/DataInterface/NiagaraDataInterfaceUtilities.cpp
--- a/Engine/Plugins/FX/Niagara/Source/Niagara/Private/DataInterface/NiagaraDataInterfaceUtilities.cpp
+++ b/Engine/Plugins/FX/Niagara/Source/Niagara/Private/DataInterface/NiagaraDataInterfaceUtilities.cpp
@@ -3,6 +3,7 @@
 #include "NiagaraDataInterfaceUtilities.h"
 #include "NiagaraComponent.h"
 #include "NiagaraComputeExecutionContext.h"
+#include "NiagaraConstants.h"
 #include "NiagaraDataInterface.h"
 #include "NiagaraSystem.h"
 
@@ -155,6 +156,11 @@ void ForEachDataInterface(const FNiagaraParameterStore& ParameterStore, TFunctio
 			continue;
 		}
 
+		if (Variable.IsInNameSpace(FNiagaraConstants::InternalNamespaceString))
+		{
+			continue;
+		}
+
 		if ( UNiagaraDataInterface* DataInterface = DataInterfaces[Variable.Offset] )
 		{
 			Action(Variable, DataInterface);
diff --git a/Engine/Plugins/FX/Niagara/Source/Niagara/Private/NiagaraConstants.cpp b/Engine/Plugins/FX/Niagara/Source/Niagara/Private/NiagaraConstants.cpp
--- a/Engine/Plugins/FX/Niagara/Source/Niagara/Private/NiagaraConstants.cpp
+++ b/Engine/Plugins/FX/Niagara/Source/Niagara/Private/NiagaraConstants.cpp
@@ -59,6 +59,7 @@ const FString FNiagaraConstants::RapidIterationParametersNamespaceString(TEXT("C
 const FString FNiagaraConstants::StackContextNamespaceString(TEXT("StackContext"));
 const FString FNiagaraConstants::SystemNamespaceString(TEXT("System"));
 const FString FNiagaraConstants::UserNamespaceString(TEXT("User"));
+const FString FNiagaraConstants::InternalNamespaceString(TEXT("__INTERNAL__"));
 
 const FName FNiagaraConstants::InputScopeName(TEXT("Input"));
 const FName FNiagaraConstants::OutputScopeName(TEXT("Output"));
diff --git a/Engine/Plugins/FX/Niagara/Source/Niagara/Private/NiagaraScriptExecutionParameterStore.cpp b/Engine/Plugins/FX/Niagara/Source/Niagara/Private/NiagaraScriptExecutionParameterStore.cpp
--- a/Engine/Plugins/FX/Niagara/Source/Niagara/Private/NiagaraScriptExecutionParameterStore.cpp
+++ b/Engine/Plugins/FX/Niagara/Source/Niagara/Private/NiagaraScriptExecutionParameterStore.cpp
@@ -1,6 +1,7 @@
 // Copyright Epic Games, Inc. All Rights Reserved.
 
 #include "NiagaraScriptExecutionParameterStore.h"
+#include "NiagaraConstants.h"
 #include "NiagaraStats.h"
 #include "NiagaraDataInterface.h"
 #include "NiagaraSystemInstance.h"
@@ -320,7 +321,11 @@ void FNiagaraScriptExecutionParameterStore::AddScriptParams(UNiagaraScript* Scri
 		{
 			// If the data interface wasn't used in a parameter map, mangle the name so that it doesn't accidentally bind to
 			// a valid parameter.
-			ParameterName = *(TEXT("__INTERNAL__.") + Info.Name.ToString());
+			FNameBuilder NameBuilder;
+			NameBuilder.Append(FNiagaraConstants::InternalNamespaceString);
+			NameBuilder.AppendChar(TEXT('.'));
+			Info.Name.AppendString(NameBuilder);
+			ParameterName = FName(NameBuilder.ToString()); ;
 		}
 
 		FNiagaraVariable Var(Info.Type, ParameterName);
