Handle TypeDefinitions with transient structs (SWC) during serialization
-uses the mapping of SWC structs created during the session to serialize out the source LWC version, if one is available

diff --git a/Engine/Plugins/FX/Niagara/Source/Niagara/Private/NiagaraModule.cpp b/Engine/Plugins/FX/Niagara/Source/Niagara/Private/NiagaraModule.cpp
--- a/Engine/Plugins/FX/Niagara/Source/Niagara/Private/NiagaraModule.cpp
+++ b/Engine/Plugins/FX/Niagara/Source/Niagara/Private/NiagaraModule.cpp
@@ -377,6 +377,7 @@ void INiagaraModule::StartupModule()
 	FWorldDelegates::OnWorldTickStart.AddRaw(this, &INiagaraModule::OnWorldTickStart);
 
 	FCoreDelegates::OnBeginFrame.AddRaw(this, &INiagaraModule::OnBeginFrame);
+	FCoreUObjectDelegates::GetPostGarbageCollect().AddRaw(this, &INiagaraModule::OnPostGarbageCollect);
 }
 
 void INiagaraModule::OnPostEngineInit()
@@ -444,6 +445,11 @@ void INiagaraModule::OnWorldTickStart(UWorld* World, ELevelTick TickType, float
 #endif
 }
 
+void INiagaraModule::OnPostGarbageCollect()
+{
+	FNiagaraTypeHelper::TickTypeRemap();
+}
+
 void INiagaraModule::OnBeginFrame()
 {
 	FNiagaraPlatformSet::RefreshScalability();
@@ -1272,40 +1278,83 @@ FNiagaraTypeDefinition FNiagaraTypeDefinition::GetNumericOutputType(const TArray
 bool FNiagaraTypeDefinition::Serialize(FArchive& Ar)
 {
 	Ar.UsingCustomVersion(FNiagaraCustomVersion::GUID);
-	return false;
-}
 
-void FNiagaraTypeDefinition::PostSerialize(const FArchive& Ar)
-{
-#if WITH_EDITORONLY_DATA
-	if (Ar.IsLoading() && Ar.CustomVer(FNiagaraCustomVersion::GUID) < FNiagaraCustomVersion::MemorySaving)
+	UScriptStruct* StructToResetOnExit = nullptr;
+	ON_SCOPE_EXIT
 	{
-		if (Enum_DEPRECATED != nullptr)
-		{
-			UnderlyingType = UT_Enum;
-			ClassStructOrEnum = Enum_DEPRECATED;
-			Enum_DEPRECATED = nullptr;
-		}
-		else if (Struct_DEPRECATED != nullptr)
+		if (StructToResetOnExit)
 		{
-			UnderlyingType = Struct_DEPRECATED->IsA<UClass>() ? UT_Class : UT_Struct;
-			ClassStructOrEnum = Struct_DEPRECATED;
-			Struct_DEPRECATED = nullptr;
+			ClassStructOrEnum = StructToResetOnExit;
+			Flags &= ~TF_SerializedAsLWC;
 		}
-		else
+	};
+
+	if (Ar.IsSaving())
+	{
+		if (UScriptStruct* StructDef = GetScriptStruct())
 		{
-			UnderlyingType = UT_None;
-			ClassStructOrEnum = nullptr;
+			// if the struct is a transient package then we look for the source struct that was used
+			if (StructDef->GetOutermost() == GetTransientPackage())
+			{
+				Flags |= TF_SerializedAsLWC;
+
+				StructToResetOnExit = StructDef;
+				ClassStructOrEnum = FNiagaraTypeHelper::GetLWCStruct(StructDef);
+				check(ClassStructOrEnum);
+				check(ClassStructOrEnum->GetOutermost() != GetTransientPackage());
+			}
 		}
 	}
-	if (Ar.IsLoading() && ClassStructOrEnum != nullptr)
+
+	if (Ar.IsLoading() || Ar.IsSaving())
 	{
-		if (ClassStructOrEnum.GetClass()->IsChildOf(UScriptStruct::StaticClass()))
+		UScriptStruct* Struct = FNiagaraTypeDefinition::StaticStruct();
+		Struct->SerializeTaggedProperties(Ar, (uint8*)this, Struct, nullptr);
+	}
+
+	if (Ar.IsLoading())
+	{
+#if WITH_EDITORONLY_DATA
+		if (Ar.CustomVer(FNiagaraCustomVersion::GUID) < FNiagaraCustomVersion::MemorySaving)
 		{
-			ClassStructOrEnum = FNiagaraTypeHelper::FindNiagaraFriendlyTopLevelStruct(static_cast<UScriptStruct*>(ClassStructOrEnum), ENiagaraStructConversion::UserFacing);
+			if (Enum_DEPRECATED != nullptr)
+			{
+				UnderlyingType = UT_Enum;
+				ClassStructOrEnum = Enum_DEPRECATED;
+				Enum_DEPRECATED = nullptr;
+			}
+			else if (Struct_DEPRECATED != nullptr)
+			{
+				UnderlyingType = Struct_DEPRECATED->IsA<UClass>() ? UT_Class : UT_Struct;
+				ClassStructOrEnum = Struct_DEPRECATED;
+				Struct_DEPRECATED = nullptr;
+			}
+			else
+			{
+				UnderlyingType = UT_None;
+				ClassStructOrEnum = nullptr;
+			}
 		}
-	}
 #endif
+
+		if (UScriptStruct* StructDef = GetScriptStruct())
+		{
+			if (Flags & TF_SerializedAsLWC)
+			{
+				ClassStructOrEnum = FNiagaraTypeHelper::GetSWCStruct(StructDef);
+
+				Flags &= ~TF_SerializedAsLWC;
+			}
+#if WITH_EDITORONLY_DATA
+			else
+			{
+				ClassStructOrEnum = FNiagaraTypeHelper::FindNiagaraFriendlyTopLevelStruct(static_cast<UScriptStruct*>(ClassStructOrEnum), ENiagaraStructConversion::UserFacing);
+			}
+#endif
+		}
+	}
+
+	return true;
 }
 
 //////////////////////////////////////////////////////////////////////////
diff --git a/Engine/Plugins/FX/Niagara/Source/Niagara/Private/NiagaraScript.cpp b/Engine/Plugins/FX/Niagara/Source/Niagara/Private/NiagaraScript.cpp
--- a/Engine/Plugins/FX/Niagara/Source/Niagara/Private/NiagaraScript.cpp
+++ b/Engine/Plugins/FX/Niagara/Source/Niagara/Private/NiagaraScript.cpp
@@ -2461,51 +2461,6 @@ static bool ValidateExecData(const UNiagaraScript* Script, const FNiagaraVMExecu
 	return IsValid;
 }
 
-struct FObjectAndNameAsStringProxyArchive_NiagaraExecData : public FObjectAndNameAsStringProxyArchive
-{
-	FObjectAndNameAsStringProxyArchive_NiagaraExecData(FArchive& InInnerArchive, bool bInLoadIfFindFails)
-		: FObjectAndNameAsStringProxyArchive(InInnerArchive, bInLoadIfFindFails)
-	{ }
-
-	virtual FArchive& operator<<(UObject*& Obj) override
-	{
-		if (IsLoading())
-		{
-			// load the path name to the object
-			FString LoadedString;
-			InnerArchive << LoadedString;
-			// look up the object by fully qualified pathname
-			Obj = FindObject<UObject>(nullptr, *LoadedString, false);
-			// If we couldn't find it, and we want to load it, do that
-			if (!Obj)
-			{
-				if (bLoadIfFindFails)
-				{
-					Obj = LoadObject<UObject>(nullptr, *LoadedString);
-				}
-				if (!Obj)
-				{
-					if (LoadedString.EndsWith(TEXT("_SWC")))
-					{
-						FString LWCShortName = FPackageName::ObjectPathToObjectName(LoadedString).LeftChop(4);
-						if (UScriptStruct* LWCStruct = FindFirstObject<UScriptStruct>(*LWCShortName))
-						{
-							Obj = FNiagaraTypeHelper::GetSWCStruct(LWCStruct);
-						}
-					}
-				}
-			}
-		}
-		else
-		{
-			// save out the fully qualified object name
-			FString SavedString(Obj->GetPathName());
-			InnerArchive << SavedString;
-		}
-		return *this;
-	}
-};
-
 bool UNiagaraScript::BinaryToExecData(const UNiagaraScript* Script, const TArray<uint8>& InBinaryData, FNiagaraVMExecutableData& OutExecData)
 {
 	check(IsInGameThread());
@@ -2515,7 +2470,7 @@ bool UNiagaraScript::BinaryToExecData(const UNiagaraScript* Script, const TArray
 	}
 
 	FMemoryReader Ar(InBinaryData, true);
-	FObjectAndNameAsStringProxyArchive_NiagaraExecData SafeAr(Ar, false);
+	FObjectAndNameAsStringProxyArchive SafeAr(Ar, false);
 	OutExecData.SerializeData(SafeAr, true);
 
 	FString ValidationErrors;
diff --git a/Engine/Plugins/FX/Niagara/Source/Niagara/Private/NiagaraTypes.cpp b/Engine/Plugins/FX/Niagara/Source/Niagara/Private/NiagaraTypes.cpp
--- a/Engine/Plugins/FX/Niagara/Source/Niagara/Private/NiagaraTypes.cpp
+++ b/Engine/Plugins/FX/Niagara/Source/Niagara/Private/NiagaraTypes.cpp
@@ -243,6 +243,10 @@ void FNiagaraLwcStructConverter::AddConversionStep(int32 InSourceBytes, int32 In
 
 //////////////////////////////////////////////////////////////////////////
 
+FRWLock FNiagaraTypeHelper::RemapTableLock;
+TMap<TWeakObjectPtr<UScriptStruct>, FNiagaraTypeHelper::FRemapEntry> FNiagaraTypeHelper::RemapTable;
+std::atomic<bool> FNiagaraTypeHelper::RemapTableDirty;
+
 FString FNiagaraTypeHelper::ToString(const uint8* ValueData, const UObject* StructOrEnum)
 {
 	FString Ret;
@@ -440,30 +444,32 @@ bool FNiagaraTypeHelper::IsLWCType(const FNiagaraTypeDefinition& InType)
 	return InType.IsValid() && !InType.IsUObject() && IsLWCStructure(InType.GetStruct());
 }
 
-UScriptStruct* FNiagaraTypeHelper::GetSWCStruct(UScriptStruct* LWCStruct)
+void FNiagaraTypeHelper::TickTypeRemap()
 {
-	struct FRemapEntry
+	if (RemapTableDirty)
 	{
-		UScriptStruct* Get(UScriptStruct* InStruct) const
+		FWriteScopeLock WriteLock(RemapTableLock);
+
+		for (auto it = RemapTable.CreateIterator(); it; ++it)
 		{
-		#if WITH_EDITORONLY_DATA
-			return SerialNumber == InStruct->FieldPathSerialNumber ? Struct.Get() : nullptr;
-		#else
-			return Struct.Get();
-		#endif
+			UScriptStruct* SrcScript = it->Key.Get();
+			if (SrcScript == nullptr || it->Value.Get(SrcScript) == nullptr)
+			{
+				UScriptStruct* DstStruct = it->Value.Struct.Get();
+				if (DstStruct && SrcScript != DstStruct)
+				{
+					DstStruct->RemoveFromRoot();
+				}
+				it.RemoveCurrent();
+			}
 		}
 
-		TWeakObjectPtr<UScriptStruct>	Struct;
-	#if WITH_EDITORONLY_DATA
-		int32 SerialNumber = 0;
-	#endif
-	};
-
-	static FRWLock RemapTableLock;
-	static TMap<TWeakObjectPtr<UScriptStruct>, FRemapEntry> RemapTable;
-
-	bool bRemoveDeadRemaps = false;
+		RemapTableDirty = false;
+	}
+}
 
+UScriptStruct* FNiagaraTypeHelper::GetSWCStruct(UScriptStruct* LWCStruct)
+{
 	// Attempt to find existing struct
 	UScriptStruct* SWCStruct = nullptr;
 	{
@@ -471,7 +477,7 @@ UScriptStruct* FNiagaraTypeHelper::GetSWCStruct(UScriptStruct* LWCStruct)
 		if ( FRemapEntry* RemapEntry = RemapTable.Find(LWCStruct) )
 		{
 			SWCStruct = RemapEntry->Get(LWCStruct);
-			bRemoveDeadRemaps = true;
+			RemapTableDirty = true;
 		}
 	}
 
@@ -483,33 +489,13 @@ UScriptStruct* FNiagaraTypeHelper::GetSWCStruct(UScriptStruct* LWCStruct)
 		if (FRemapEntry* RemapEntry = RemapTable.Find(LWCStruct))
 		{
 			SWCStruct = RemapEntry->Get(LWCStruct);
-			bRemoveDeadRemaps = true;
+			RemapTableDirty = true;
 		}
 
 		if (SWCStruct == nullptr)
 		{
 			QUICK_SCOPE_CYCLE_COUNTER(STAT_Niagara_GetSWCStruct);
 
-			// Do we need to remove dead remaps from the list?  
-			//-OPT: We should hook into post GC and prune the map at that point
-			bRemoveDeadRemaps = true;
-			if ( bRemoveDeadRemaps )
-			{
-				for (auto it = RemapTable.CreateIterator(); it; ++it)
-				{
-					UScriptStruct* SrcScript = it->Key.Get();
-					if (SrcScript == nullptr || it->Value.Get(SrcScript) == nullptr)
-					{
-						UScriptStruct* DstStruct = it->Value.Struct.Get();
-						if (DstStruct && SrcScript != DstStruct)
-						{
-							DstStruct->RemoveFromRoot();
-						}
-						it.RemoveCurrent();
-					}
-				}
-			}
-
 			// If this is a LWC structure we need to built a new structure now
 			if (IsLWCStructure(LWCStruct))
 			{
@@ -535,6 +521,27 @@ UScriptStruct* FNiagaraTypeHelper::GetSWCStruct(UScriptStruct* LWCStruct)
 	return SWCStruct;
 }
 
+UScriptStruct* FNiagaraTypeHelper::GetLWCStruct(UScriptStruct* SWCStruct)
+{
+	if (SWCStruct->GetOutermost() != GetTransientPackage())
+	{
+		return SWCStruct;
+	}
+
+	FReadScopeLock ReadLock(RemapTableLock);
+
+	for (auto It = RemapTable.CreateConstIterator(); It; ++It)
+	{
+		const FRemapEntry& RemapEntry = It.Value();
+		if (RemapEntry.Get(SWCStruct) == SWCStruct)
+		{
+			return It.Key().Get();
+		}
+	}
+
+	return nullptr;
+}
+
 UScriptStruct* FNiagaraTypeHelper::FindNiagaraFriendlyTopLevelStruct(UScriptStruct* InStruct, ENiagaraStructConversion StructConversion)
 {
 	if (!InStruct)
diff --git a/Engine/Plugins/FX/Niagara/Source/Niagara/Public/NiagaraModule.h b/Engine/Plugins/FX/Niagara/Source/Niagara/Public/NiagaraModule.h
--- a/Engine/Plugins/FX/Niagara/Source/Niagara/Public/NiagaraModule.h
+++ b/Engine/Plugins/FX/Niagara/Source/Niagara/Public/NiagaraModule.h
@@ -55,6 +55,7 @@ public:
 
 	void OnWorldTickStart(UWorld* World, ELevelTick TickType, float DeltaSeconds);
 	void OnBeginFrame();
+	void OnPostGarbageCollect();
 	void OnWorldBeginTearDown(UWorld* World);
 
 	FDelegateHandle SetOnProcessShaderCompilationQueue(FOnProcessQueue InOnProcessQueue);
diff --git a/Engine/Plugins/FX/Niagara/Source/Niagara/Public/NiagaraTypes.h b/Engine/Plugins/FX/Niagara/Source/Niagara/Public/NiagaraTypes.h
--- a/Engine/Plugins/FX/Niagara/Source/Niagara/Public/NiagaraTypes.h
+++ b/Engine/Plugins/FX/Niagara/Source/Niagara/Public/NiagaraTypes.h
@@ -373,6 +373,30 @@ public:
 	static UScriptStruct* FindNiagaraFriendlyTopLevelStruct(UScriptStruct* InStruct, ENiagaraStructConversion StructConversion);
 	static bool IsNiagaraFriendlyTopLevelStruct(UScriptStruct* InStruct, ENiagaraStructConversion StructConversion);
 	static UScriptStruct* GetSWCStruct(UScriptStruct* LWCStruct);
+	static UScriptStruct* GetLWCStruct(UScriptStruct* LWCStruct);
+	static void TickTypeRemap();
+
+private:
+	struct FRemapEntry
+	{
+		UScriptStruct* Get(UScriptStruct* InStruct) const
+		{
+#if WITH_EDITORONLY_DATA
+			return SerialNumber == InStruct->FieldPathSerialNumber ? Struct.Get() : nullptr;
+#else
+			return Struct.Get();
+#endif
+		}
+
+		TWeakObjectPtr<UScriptStruct>	Struct;
+#if WITH_EDITORONLY_DATA
+		int32 SerialNumber = 0;
+#endif
+	};
+
+	static FRWLock RemapTableLock;
+	static TMap<TWeakObjectPtr<UScriptStruct>, FRemapEntry> RemapTable;
+	static std::atomic<bool> RemapTableDirty;
 };
 
 /** Information about how this type should be laid out in an FNiagaraDataSet */
@@ -869,8 +893,12 @@ struct NIAGARA_API FNiagaraTypeDefinition
 
 	enum FTypeFlags
 	{
-		TF_None = 0,
-		TF_Static
+		TF_None			= 0x0000,
+		TF_Static		= 0x0001,
+
+		/// indicates that the ClassStructOrEnum property has been serialized as the LWC struct (see FNiagaraTypeHelper)
+		/// instead of the Transient SWC version of the struct
+		TF_SerializedAsLWC	= 0x0002,
 	};
 
 public:
@@ -1126,7 +1154,6 @@ public:
 	uint16 UnderlyingType;
 
 	bool Serialize(FArchive& Ar);
-	void PostSerialize(const FArchive& Ar);
 
 private:
 	UPROPERTY(EditAnywhere, Category = Type)
@@ -1318,7 +1345,6 @@ struct TStructOpsTypeTraits<FNiagaraTypeDefinition> : public TStructOpsTypeTrait
 	enum
 	{
 		WithSerializer = true,
-		WithPostSerialize = true,
 	};
 };
 
diff --git a/Engine/Source/Runtime/Core/Private/UObject/FortniteMainBranchObjectVersion.cpp b/Engine/Source/Runtime/Core/Private/UObject/FortniteMainBranchObjectVersion.cpp
--- a/Engine/Source/Runtime/Core/Private/UObject/FortniteMainBranchObjectVersion.cpp
+++ b/Engine/Source/Runtime/Core/Private/UObject/FortniteMainBranchObjectVersion.cpp
@@ -11,7 +11,7 @@ TMap<FGuid, FGuid> FFortniteMainBranchObjectVersion::GetSystemGuids()
 	SystemGuids.Add(DevGuids.MATERIALSHADERMAP_DERIVEDDATA_VER, FGuid("1DA85DA5733C4A489A1BE117CDC5ACCC"));
 	SystemGuids.Add(DevGuids.NANITE_DERIVEDDATA_VER, FGuid("4C051B8C8FC74BA3943CD60A62D1CB73"));
 	SystemGuids.Add(DevGuids.NIAGARASHADERMAP_DERIVEDDATA_VER, FGuid("2289D5116CF94BC0AFEAEC541468E645"));
-	SystemGuids.Add(DevGuids.Niagara_LatestScriptCompileVersion, FGuid("E93049C042704A758D6B65386328DC9F"));
+	SystemGuids.Add(DevGuids.Niagara_LatestScriptCompileVersion, FGuid("23DE02026DAF4B60B289CAB316D28680"));
 	SystemGuids.Add(DevGuids.SkeletalMeshDerivedDataVersion, FGuid("ACF593EAAE354FCBB34CF44F5AA1BFD2"));
 	SystemGuids.Add(DevGuids.STATICMESH_DERIVEDDATA_VER, FGuid("3006d21f867648a68cf4edee185446b7"));
 
