- SimCache functions to access data from BP
- Split SimCache code into different files and hide details we may want to refactor

diff --git a/Engine/Plugins/FX/Niagara/Source/Niagara/Classes/NiagaraSimCache.h b/Engine/Plugins/FX/Niagara/Source/Niagara/Classes/NiagaraSimCache.h
--- a/Engine/Plugins/FX/Niagara/Source/Niagara/Classes/NiagaraSimCache.h
+++ b/Engine/Plugins/FX/Niagara/Source/Niagara/Classes/NiagaraSimCache.h
@@ -114,7 +114,10 @@ struct FNiagaraSimCacheFrame
 	GENERATED_BODY()
 
 	UPROPERTY()
-	FTransform LocalToWorld;
+	FTransform LocalToWorld = FTransform::Identity;
+
+	UPROPERTY()
+	FVector3f LWCTile = FVector3f::ZeroVector;
 
 	UPROPERTY()
 	FNiagaraSimCacheSystemFrame SystemData;
@@ -192,6 +195,9 @@ struct FNiagaraSimCacheDataBuffersLayout
 	UPROPERTY()
 	uint16 Int32Count = 0;
 
+	UPROPERTY()
+	bool bLocalSpace = false;
+
 	UPROPERTY()
 	TArray<FName> RebaseVariableNames;
 
@@ -216,8 +222,129 @@ struct FNiagaraSimCacheLayout
 UCLASS(Experimental, BlueprintType)
 class NIAGARA_API UNiagaraSimCache : public UObject
 {
+	friend struct FNiagaraSimCacheAttributeReaderHelper;
+	friend struct FNiagaraSimCacheHelper;
+	friend struct FSimCacheBufferReader;	//-TODO: Refactor to use API
+
 	GENERATED_UCLASS_BODY()
 
+public:
+	// UObject Interface
+	virtual bool IsReadyForFinishDestroy() override;
+	// UObject Interface
+
+	/** Is the simulation cache valid or not. */
+	UFUNCTION(BlueprintCallable, Category=NiagaraSimCache)
+	bool IsCacheValid() const { return SoftNiagaraSystem.IsNull() == false; }
+
+	void BeginWrite(FNiagaraSimCacheCreateParameters InCreateParameters, UNiagaraComponent* NiagaraComponent);
+	void WriteFrame(UNiagaraComponent* NiagaraComponent);
+	void EndWrite();
+
+	bool CanRead(UNiagaraSystem* NiagaraSystem);
+	bool Read(float TimeSeconds, FNiagaraSystemInstance* SystemInstance) const;
+	bool ReadFrame(int32 FrameIndex, float FrameFraction, FNiagaraSystemInstance* SystemInstance) const;
+
+	/** Get the time the simulation was at when recorded. */
+	UFUNCTION(BlueprintCallable, Category = NiagaraSimCache)
+	float GetStartSeconds() const { return StartSeconds; }
+
+	/** Get the duration the cache was recorded over. */
+	float GetDurationSeconds() const { return DurationSeconds; }
+
+	/** Get number of frames stored in the cache. */
+	UFUNCTION(BlueprintCallable, Category = NiagaraSimCache)
+	int GetNumFrames() const { return CacheFrames.Num(); }
+
+	/** Get number of emitters stored inside the cache. */
+	UFUNCTION(BlueprintCallable, Category = NiagaraSimCache)
+	int GetNumEmitters() const { return CacheLayout.EmitterLayouts.Num(); }
+
+	/** Get the emitter name at the provided index. */
+	UFUNCTION(BlueprintCallable, Category = NiagaraSimCache)
+	FName GetEmitterName(int32 EmitterIndex) const { return CacheLayout.EmitterLayouts.IsValidIndex(EmitterIndex) ? CacheLayout.EmitterLayouts[EmitterIndex].LayoutName : NAME_None; }
+
+	/** Returns a list of emitters we have captured in the SimCache. */
+	UFUNCTION(BlueprintCallable, Category = NiagaraSimCache)
+	TArray<FName> GetEmitterNames() const;
+
+	/**
+	Reads Niagara int attributes by name from the cache frame and appends them into the OutValues array.
+	EmitterName - If left blank will return the system simulation attributes.
+	*/
+	UFUNCTION(BlueprintCallable, Category = NiagaraSimCache)
+	void ReadIntAttribute(TArray<int32>& OutValues, FName AttributeName, FName EmitterName, int FrameIndex = 0) const;
+
+	/**
+	Reads Niagara float attributes by name from the cache frame and appends them into the OutValues array.
+	EmitterName - If left blank will return the system simulation attributes.
+	*/
+	UFUNCTION(BlueprintCallable, Category = NiagaraSimCache)
+	void ReadFloatAttribute(TArray<float>& OutValues, FName AttributeName, FName EmitterName, int FrameIndex = 0) const;
+
+	/**
+	Reads Niagara Vec2 attributes by name from the cache frame and appends them into the OutValues array.
+	EmitterName - If left blank will return the system simulation attributes.
+	*/
+	UFUNCTION(BlueprintCallable, Category = NiagaraSimCache)
+	void ReadVector2Attribute(TArray<FVector2D>& OutValues, FName AttributeName, FName EmitterName, int FrameIndex = 0) const;
+
+	/**
+	Reads Niagara Vec3 attributes by name from the cache frame and appends them into the OutValues array.
+	EmitterName - If left blank will return the system simulation attributes.
+	*/
+	UFUNCTION(BlueprintCallable, Category = NiagaraSimCache)
+	void ReadVectorAttribute(TArray<FVector>& OutValues, FName AttributeName, FName EmitterName, int FrameIndex = 0) const;
+
+	/**
+	Reads Niagara Vec4 attributes by name from the cache frame and appends them into the OutValues array.
+	EmitterName - If left blank will return the system simulation attributes.
+	*/
+	UFUNCTION(BlueprintCallable, Category = NiagaraSimCache)
+	void ReadVector4Attribute(TArray<FVector4>& OutValues, FName AttributeName, FName EmitterName, int FrameIndex = 0) const;
+
+	/**
+	Reads Niagara Color attributes by name from the cache frame and appends them into the OutValues array.
+	EmitterName - If left blank will return the system simulation attributes.
+	*/
+	UFUNCTION(BlueprintCallable, Category = NiagaraSimCache)
+	void ReadColorAttribute(TArray<FLinearColor>& OutValues, FName AttributeName, FName EmitterName, int FrameIndex = 0) const;
+
+	/**
+	Reads Niagara Position attributes by name from the cache frame and appends them into the OutValues array.
+	Local space emitters provide data at local locations unless bLocalSpaceToWorld is true.
+	EmitterName - If left blank will return the system simulation attributes.
+	*/
+	UFUNCTION(BlueprintCallable, Category = NiagaraSimCache)
+	void ReadPositionAttribute(TArray<FVector>& OutValues, FName AttributeName, FName EmitterName, bool bLocalSpaceToWorld = true, int FrameIndex = 0) const;
+
+	/**
+	Reads Niagara Position attributes by name from the cache frame and appends them into the OutValues array.
+	All attributes read with this method will be re-based from the captured space into the provided transform space,
+	this will occur even if the cache was not captured with re-basing enabled.
+	EmitterName - If left blank will return the system simulation attributes.
+	*/
+	UFUNCTION(BlueprintCallable, Category = NiagaraSimCache)
+	void ReadPositionAttributeWithRebase(TArray<FVector>& OutValues, FName AttributeName, FName EmitterName, FTransform Transform, int FrameIndex = 0) const;
+
+	/**
+	Reads Niagara Quaternion attributes by name from the cache frame and appends them into the OutValues array.
+	Local space emitters provide data at local rotation unless bLocalSpaceToWorld is true.
+	EmitterName - If left blank will return the system simulation attributes.
+	*/
+	UFUNCTION(BlueprintCallable, Category = NiagaraSimCache)
+	void ReadQuatAttribute(TArray<FQuat>& OutValues, FName AttributeName, FName EmitterName, bool bLocalSpaceToWorld = true, int FrameIndex = 0) const;
+
+	/**
+	Reads Niagara Quaternion attributes by name from the cache frame and appends them into the OutValues array.
+	Only attributes that in the rebase list will be transform into the provided Quat space.  Therefore the cache
+	must be captured with rebasing enabled to have any impact.
+	EmitterName - If left blank will return the system simulation attributes.
+	*/
+	UFUNCTION(BlueprintCallable, Category = NiagaraSimCache)
+	void ReadQuatAttributeWithRebase(TArray<FQuat>& OutValues, FName AttributeName, FName EmitterName, FQuat Quat, int FrameIndex = 0) const;
+
+private:
 	UPROPERTY(VisibleAnywhere, Category=SimCache, meta=(DisplayName="Niagara System"))
 	TSoftObjectPtr<UNiagaraSystem> SoftNiagaraSystem;
 
@@ -247,22 +374,5 @@ class NIAGARA_API UNiagaraSimCache : public UObject
 	UPROPERTY()
 	TMap<FNiagaraVariableBase, TObjectPtr<UObject>> DataInterfaceStorage;
 
-	// UObject Interface
-	virtual bool IsReadyForFinishDestroy() override;
-	// UObject Interface
-
-	/** Is the simulation cache valid or not. */
-	UFUNCTION(BlueprintCallable, Category=NiagaraSimCache)
-	bool IsCacheValid() const { return SoftNiagaraSystem.IsNull() == false; }
-
-	void BeginWrite(FNiagaraSimCacheCreateParameters InCreateParameters, UNiagaraComponent* NiagaraComponent);
-	void WriteFrame(UNiagaraComponent* NiagaraComponent);
-	void EndWrite();
-
-	bool CanRead(UNiagaraSystem* NiagaraSystem);
-	bool Read(float TimeSeconds, FNiagaraSystemInstance* SystemInstance) const;
-	bool ReadFrame(int32 FrameIndex, float FrameFraction, FNiagaraSystemInstance* SystemInstance) const;
-
-private:
 	mutable std::atomic<int32> PendingCommandsInFlight;
 };
diff --git a/Engine/Plugins/FX/Niagara/Source/Niagara/Private/NiagaraSimCache.cpp b/Engine/Plugins/FX/Niagara/Source/Niagara/Private/NiagaraSimCache.cpp
--- a/Engine/Plugins/FX/Niagara/Source/Niagara/Private/NiagaraSimCache.cpp
+++ b/Engine/Plugins/FX/Niagara/Source/Niagara/Private/NiagaraSimCache.cpp
@@ -11,681 +11,11 @@
 #include "NiagaraEmitterInstance.h"
 #include "NiagaraGPUInstanceCountManager.h"
 #include "NiagaraGpuComputeDispatchInterface.h"
+#include "NiagaraSimCacheAttributeReaderHelper.h"
+#include "NiagaraSimCacheHelper.h"
 #include "NiagaraSystem.h"
 #include "NiagaraSystemInstance.h"
 
-struct FNiagaraSimCacheHelper
-{
-	explicit FNiagaraSimCacheHelper(FNiagaraSystemInstance* InSystemInstance)
-	{
-		SystemInstance = InSystemInstance;
-		SystemSimulation = SystemInstance->GetSystemSimulation();
-		check(SystemSimulation);
-		SystemSimulationDataBuffer = SystemSimulation->MainDataSet.GetCurrentData();
-		NiagaraSystem = SystemSimulation->GetSystem();
-	}
-
-	explicit FNiagaraSimCacheHelper(UNiagaraComponent* NiagaraComponent)
-	{
-		if ( NiagaraComponent == nullptr )
-		{
-			return;
-		}
-
-		NiagaraSystem = NiagaraComponent->GetAsset();
-		if ( NiagaraSystem == nullptr )
-		{
-			return;
-		}
-
-		FNiagaraSystemInstanceControllerPtr SystemInstanceController = NiagaraComponent->GetSystemInstanceController();
-		if (SystemInstanceController.IsValid() == false)
-		{
-			return;
-		}
-
-		SystemInstance = SystemInstanceController->GetSystemInstance_Unsafe();
-		if (SystemInstance == nullptr)
-		{
-			return;
-		}
-
-		SystemSimulation = SystemInstance->GetSystemSimulation();
-		if (SystemSimulation == nullptr)
-		{
-			return;
-		}
-
-		SystemSimulationDataBuffer = SystemSimulation->MainDataSet.GetCurrentData();
-		if (SystemSimulationDataBuffer == nullptr)
-		{
-			return;
-		}
-	}
-
-	FNiagaraDataSet& GetSystemSimulationDataSet() { return SystemSimulation->MainDataSet; }
-
-	bool HasValidSimulation() const { return SystemSimulation != nullptr; }
-	bool HasValidSimulationData() const { return SystemSimulationDataBuffer != nullptr; }
-
-	void BuildCacheLayout(FNiagaraSimCacheDataBuffersLayout& CacheLayout, const FNiagaraDataSetCompiledData& CompiledData, FName LayoutName, TArray<FName> InRebaseVariableNames) const
-	{
-		CacheLayout.LayoutName = LayoutName;
-		CacheLayout.SimTarget = CompiledData.SimTarget;
-
-		const int32 NumVariables = CompiledData.Variables.Num();
-		CacheLayout.Variables.AddDefaulted(NumVariables);
-
-		const int32 CacheTotalComponents = CompiledData.TotalFloatComponents + CompiledData.TotalFloatComponents + CompiledData.TotalInt32Components;
-		CacheLayout.ComponentMappingsFromDataBuffer.Empty(CacheTotalComponents);
-		CacheLayout.ComponentMappingsFromDataBuffer.AddDefaulted(CacheTotalComponents);
-		CacheLayout.RebaseVariableNames = MoveTemp(InRebaseVariableNames);
-
-		for ( int32 iVariable=0; iVariable < NumVariables; ++iVariable)
-		{
-			const FNiagaraVariableLayoutInfo& DataSetVariableLayout = CompiledData.VariableLayouts[iVariable];
-			FNiagaraSimCacheVariable& CacheVariable = CacheLayout.Variables[iVariable];
-
-			CacheVariable.Variable = CompiledData.Variables[iVariable];
-			CacheVariable.FloatOffset = DataSetVariableLayout.GetNumFloatComponents() > 0 ? CacheLayout.FloatCount : INDEX_NONE;
-			CacheVariable.FloatCount = DataSetVariableLayout.GetNumFloatComponents();
-			CacheVariable.HalfOffset = DataSetVariableLayout.GetNumHalfComponents() > 0 ? CacheLayout.HalfCount : INDEX_NONE;
-			CacheVariable.HalfCount = DataSetVariableLayout.GetNumHalfComponents();
-			CacheVariable.Int32Offset = DataSetVariableLayout.GetNumInt32Components() > 0 ? CacheLayout.Int32Count : INDEX_NONE;
-			CacheVariable.Int32Count = DataSetVariableLayout.GetNumInt32Components();
-
-			CacheLayout.FloatCount += DataSetVariableLayout.GetNumFloatComponents();
-			CacheLayout.HalfCount += DataSetVariableLayout.GetNumHalfComponents();
-			CacheLayout.Int32Count += DataSetVariableLayout.GetNumInt32Components();
-		}
-
-		// Build write mappings we will build read mappings in a separate path
-		int32 FloatOffset = 0;
-		int32 HalfOffset = CacheLayout.FloatCount;
-		int32 Int32Offset = HalfOffset + CacheLayout.HalfCount;
-		for (int32 iVariable = 0; iVariable < NumVariables; ++iVariable)
-		{
-			const FNiagaraVariableLayoutInfo& DataSetVariableLayout = CompiledData.VariableLayouts[iVariable];
-			FNiagaraSimCacheVariable& CacheVariable = CacheLayout.Variables[iVariable];
-
-			for (int32 iComponent=0; iComponent < CacheVariable.FloatCount; ++iComponent)
-			{
-				CacheLayout.ComponentMappingsFromDataBuffer[FloatOffset] = CacheVariable.FloatOffset + iComponent;
-				++FloatOffset;
-			}
-
-			for (int32 iComponent=0; iComponent < CacheVariable.HalfCount; ++iComponent)
-			{
-				CacheLayout.ComponentMappingsFromDataBuffer[HalfOffset] = CacheVariable.HalfOffset + iComponent;
-				++HalfOffset;
-			}
-
-			for (int32 iComponent=0; iComponent < CacheVariable.Int32Count; ++iComponent)
-			{
-				CacheLayout.ComponentMappingsFromDataBuffer[Int32Offset] = CacheVariable.Int32Offset + iComponent;
-				++Int32Offset;
-			}
-		}
-
-		// Slightly inefficient but we can share the code between the paths
-		BuildCacheReadMappings(CacheLayout, CompiledData);
-	}
-
-	void BuildCacheLayoutForSystem(const FNiagaraSimCacheCreateParameters& CreateParmaeters, FNiagaraSimCacheDataBuffersLayout& CacheLayout)
-	{
-		const FNiagaraDataSetCompiledData& SystemCompileData = NiagaraSystem->GetSystemCompiledData().DataSetCompiledData;
-
-		TArray<FName> RebaseVariableNames;
-		if ( CreateParmaeters.bAllowRebasing )
-		{
-			TArray<FString, TInlineAllocator<8>> LocalSpaceEmitters;
-			for ( int32 i=0; i < NiagaraSystem->GetNumEmitters(); ++i )
-			{
-				const FNiagaraEmitterHandle& EmitterHandle = NiagaraSystem->GetEmitterHandle(i);
-				if ( EmitterHandle.GetInstance().GetEmitterData()->bLocalSpace )
-				{
-					LocalSpaceEmitters.Add(EmitterHandle.GetUniqueInstanceName());
-				}
-			}
-
-			for (const FNiagaraVariable& Variable : SystemCompileData.Variables)
-			{
-				if (Variable.GetType() == FNiagaraTypeDefinition::GetPositionDef())
-				{
-					// If this is an emitter variable we need to check if it's local space or not
-					bool bIsLocalSpace = false;
-					for ( const FString& LocalSpaceEmitter : LocalSpaceEmitters )
-					{
-						if ( Variable.IsInNameSpace(LocalSpaceEmitter) )
-						{
-							bIsLocalSpace = true;
-							break;
-						}
-					}
-
-					if ( bIsLocalSpace == false && CreateParmaeters.RebaseExcludeList.Contains(Variable.GetName()) == false )
-					{
-						RebaseVariableNames.AddUnique(Variable.GetName());
-					}
-				}
-				else if ( CanRebaseVariable(Variable) && CreateParmaeters.RebaseIncludeList.Contains(Variable.GetName()) )
-				{
-					RebaseVariableNames.AddUnique(Variable.GetName());
-				}
-			}
-		}
-
-		BuildCacheLayout(CacheLayout, SystemCompileData, NiagaraSystem->GetFName(), MoveTemp(RebaseVariableNames));
-	}
-
-	void BuildCacheLayoutForEmitter(const FNiagaraSimCacheCreateParameters& CreateParmaeters, FNiagaraSimCacheDataBuffersLayout& CacheLayout, int EmitterIndex)
-	{
-		const FNiagaraEmitterHandle& EmitterHandle = NiagaraSystem->GetEmitterHandle(EmitterIndex);
-		const FNiagaraEmitterCompiledData& EmitterCompiledData = NiagaraSystem->GetEmitterCompiledData()[EmitterIndex].Get();
-
-		// Find potential candidates for re-basing
-		TArray<FName> RebaseVariableNames;
-		if ( CreateParmaeters.bAllowRebasing && EmitterHandle.GetInstance().GetEmitterData()->bLocalSpace == false )
-		{
-			// Build list of include / exclude names
-			TArray<FName> ForceIncludeNames;
-			TArray<FName> ForceExcludeNames;
-			if ( CreateParmaeters.RebaseIncludeList.Num() > 0 || CreateParmaeters.RebaseExcludeList.Num() > 0 )
-			{
-				const FString EmitterName = EmitterHandle.GetUniqueInstanceName();
-				for (FName RebaseName : CreateParmaeters.RebaseIncludeList)
-				{
-					FNiagaraVariableBase BaseVar(FNiagaraTypeDefinition::GetFloatDef(), RebaseName);
-					if (BaseVar.RemoveRootNamespace(EmitterName))
-					{
-						ForceIncludeNames.Add(BaseVar.GetName());
-					}
-				}
-
-				for (FName RebaseName : CreateParmaeters.RebaseExcludeList)
-				{
-					FNiagaraVariableBase BaseVar(FNiagaraTypeDefinition::GetFloatDef(), RebaseName);
-					if (BaseVar.RemoveRootNamespace(EmitterName))
-					{
-						ForceExcludeNames.Add(BaseVar.GetName());
-					}
-				}
-			}
-
-		#if WITH_EDITORONLY_DATA
-			// Look for renderer attributes bound to Quat / Matrix types are we will want to rebase those
-			// We will add all Position types after this so no need to add them here
-			EmitterHandle.GetInstance().GetEmitterData()->ForEachEnabledRenderer(
-				[&](UNiagaraRendererProperties* RenderProperties)
-				{
-					for (FNiagaraVariableBase BoundAttribute : RenderProperties->GetBoundAttributes())
-					{
-						if ( (BoundAttribute.GetType() == FNiagaraTypeDefinition::GetQuatDef()) ||
-							 (BoundAttribute.GetType() == FNiagaraTypeDefinition::GetMatrix4Def()) )
-						{
-							if (BoundAttribute.RemoveRootNamespace(FNiagaraConstants::ParticleAttributeNamespaceString))
-							{
-								if (EmitterCompiledData.DataSetCompiledData.Variables.Contains(BoundAttribute) && ForceExcludeNames.Contains(BoundAttribute.GetName()) == false )
-								{
-									RebaseVariableNames.AddUnique(BoundAttribute.GetName());
-								}
-							}
-						}
-					}
-				}
-			);
-		#endif
-
-			// Look for regular attributes that we are forcing to rebase or can rebase like positions
-			for (const FNiagaraVariable& Variable : EmitterCompiledData.DataSetCompiledData.Variables)
-			{
-				if ( Variable.GetType() == FNiagaraTypeDefinition::GetPositionDef() )
-				{
-					if ( ForceExcludeNames.Contains(Variable.GetName()) == false )
-					{
-						RebaseVariableNames.AddUnique(Variable.GetName());
-					}
-				}
-				else if ( ForceIncludeNames.Contains(Variable.GetName()) && CanRebaseVariable(Variable) )
-				{
-					RebaseVariableNames.AddUnique(Variable.GetName());
-				}
-			}
-		}
-
-		BuildCacheLayout(CacheLayout, EmitterCompiledData.DataSetCompiledData, EmitterHandle.GetName(), MoveTemp(RebaseVariableNames));
-	}
-
-	static bool BuildCacheReadMappings(FNiagaraSimCacheDataBuffersLayout& CacheLayout, const FNiagaraDataSetCompiledData& CompiledData)
-	{
-		const int32 CacheTotalComponents = CacheLayout.FloatCount + CacheLayout.HalfCount + CacheLayout.Int32Count;
-		CacheLayout.ComponentMappingsToDataBuffer.Empty(CacheTotalComponents);
-		CacheLayout.ComponentMappingsToDataBuffer.AddDefaulted(CacheTotalComponents);
-		CacheLayout.VariableMappingsToDataBuffer.Empty(0);
-
-		int32 FloatOffset = 0;
-		int32 HalfOffset = CacheLayout.FloatCount;
-		int32 Int32Offset = HalfOffset + CacheLayout.HalfCount;
-		for (const FNiagaraSimCacheVariable& SourceVariable : CacheLayout.Variables)
-		{
-			// Find variable, if it doesn't exist that's ok as the cache contains more data than is required
-			const int32 DataSetVariableIndex = CompiledData.Variables.IndexOfByPredicate(FNiagaraVariableMatch(SourceVariable.Variable.GetType(), SourceVariable.Variable.GetName()));
-			const FNiagaraVariableLayoutInfo* DestVariableLayout = nullptr;
-			if (DataSetVariableIndex != INDEX_NONE)
-			{
-				DestVariableLayout = &CompiledData.VariableLayouts[DataSetVariableIndex];
-
-				// If the variable exists but types not match the cache is invalid
-				if ((DestVariableLayout->GetNumFloatComponents() != SourceVariable.FloatCount) ||
-					(DestVariableLayout->GetNumHalfComponents() != SourceVariable.HalfCount) ||
-					(DestVariableLayout->GetNumInt32Components() != SourceVariable.Int32Count))
-				{
-					return false;
-				}
-			}
-
-			// Is this a type that requires conversion / re-basing?
-			if (DestVariableLayout != nullptr)
-			{
-				if ( CacheLayout.RebaseVariableNames.Contains(SourceVariable.Variable.GetName()) )
-				{
-					if (SourceVariable.Variable.GetType() == FNiagaraTypeDefinition::GetPositionDef())
-					{
-						check(SourceVariable.FloatCount == 3);
-						CacheLayout.VariableMappingsToDataBuffer.Emplace(FloatOffset, DestVariableLayout->FloatComponentStart, &FNiagaraSimCacheHelper::CopyPositions);
-						DestVariableLayout = nullptr;
-					}
-					else if (SourceVariable.Variable.GetType() == FNiagaraTypeDefinition::GetQuatDef())
-					{
-						check(SourceVariable.FloatCount == 4);
-						CacheLayout.VariableMappingsToDataBuffer.Emplace(FloatOffset, DestVariableLayout->FloatComponentStart, &FNiagaraSimCacheHelper::CopyQuaternions);
-						DestVariableLayout = nullptr;
-					}
-					else if (SourceVariable.Variable.GetType() == FNiagaraTypeDefinition::GetMatrix4Def())
-					{
-						check(SourceVariable.FloatCount == 16);
-						CacheLayout.VariableMappingsToDataBuffer.Emplace(FloatOffset, DestVariableLayout->FloatComponentStart, &FNiagaraSimCacheHelper::CopyMatrices);
-						DestVariableLayout = nullptr;
-					}
-				}
-			}
-
-			for (int32 i = 0; i < SourceVariable.FloatCount; ++i)
-			{
-				CacheLayout.ComponentMappingsToDataBuffer[FloatOffset++] = DestVariableLayout ? DestVariableLayout->FloatComponentStart + i : INDEX_NONE;
-			}
-
-			for (int32 i = 0; i < SourceVariable.HalfCount; ++i)
-			{
-				CacheLayout.ComponentMappingsToDataBuffer[HalfOffset++] = DestVariableLayout ? DestVariableLayout->HalfComponentStart + i : INDEX_NONE;
-			}
-
-			for (int32 i = 0; i < SourceVariable.Int32Count; ++i)
-			{
-				CacheLayout.ComponentMappingsToDataBuffer[Int32Offset++] = DestVariableLayout ? DestVariableLayout->Int32ComponentStart + i : INDEX_NONE;
-			}
-		}
-
-		return true;
-	}
-
-	static void CheckedMemcpy(TConstArrayView<uint8> DstArray, uint8* Dst, TConstArrayView<uint8> SrcArray, const uint8* Src, uint32 Size)
-	{
-		checkf(Src >= SrcArray.GetData() && Src + Size <= SrcArray.GetData() + SrcArray.Num(), TEXT("Source %p-%p is out of bounds, start %p end %p"), Src, Src + Size, SrcArray.GetData(), SrcArray.GetData() + SrcArray.Num());
-		checkf(Dst >= DstArray.GetData() && Dst + Size <= DstArray.GetData() + DstArray.Num(), TEXT("Dest %p-%p is out of bounds, start %p end %p"), Dst, Dst + Size, DstArray.GetData(), DstArray.GetData() + DstArray.Num());
-		FMemory::Memcpy(Dst, Src, Size);
-	}
-
-	void WriteDataBuffer(const FNiagaraDataBuffer& DataBuffer, const FNiagaraSimCacheDataBuffersLayout& CacheLayout, FNiagaraSimCacheDataBuffers& CacheBuffer, int32 FirstInstance, int32 NumInstances)
-	{
-		if ( NumInstances == 0 )
-		{
-			return;
-		}
-
-		CacheBuffer.NumInstances = NumInstances;
-
-		int32 iComponent = 0;
-
-		// Copy Float
-		CacheBuffer.FloatData.AddDefaulted(CacheLayout.FloatCount * NumInstances * sizeof(float));
-		for ( uint32 i=0; i < CacheLayout.FloatCount; ++i )
-		{
-			const uint32 Component = CacheLayout.ComponentMappingsFromDataBuffer[iComponent++];
-			const uint8* Source = DataBuffer.GetComponentPtrFloat(Component) + (FirstInstance * sizeof(float));
-			uint8* Dest = CacheBuffer.FloatData.GetData() + (i * NumInstances * sizeof(float));
-			CheckedMemcpy(CacheBuffer.FloatData, Dest, DataBuffer.GetFloatBuffer(), Source, sizeof(float) * NumInstances);
-			//FMemory::Memcpy(Dest, Source, sizeof(float) * NumInstances);
-		}
-		
-		// Copy Half
-		CacheBuffer.HalfData.AddDefaulted(CacheLayout.HalfCount * NumInstances * sizeof(FFloat16));
-		for (uint32 i = 0; i < CacheLayout.HalfCount; ++i)
-		{
-			const uint32 Component = CacheLayout.ComponentMappingsFromDataBuffer[iComponent++];
-			const uint8* Source = DataBuffer.GetComponentPtrHalf(Component) + (FirstInstance * sizeof(FFloat16));
-			uint8* Dest = CacheBuffer.HalfData.GetData() + (i * NumInstances * sizeof(FFloat16));
-			CheckedMemcpy(CacheBuffer.HalfData, Dest, DataBuffer.GetHalfBuffer(), Source, sizeof(FFloat16) * NumInstances);
-			//FMemory::Memcpy(Dest, Source, sizeof(FFloat16) * NumInstances);
-		}
-
-		// Copy Int32
-		CacheBuffer.Int32Data.AddDefaulted(CacheLayout.Int32Count * NumInstances * sizeof(int32));
-		for (uint32 i = 0; i < CacheLayout.Int32Count; ++i)
-		{
-			const uint32 Component = CacheLayout.ComponentMappingsFromDataBuffer[iComponent++];
-			const uint8* Source = DataBuffer.GetComponentPtrInt32(Component) + (FirstInstance * sizeof(int32));
-			uint8* Dest = CacheBuffer.Int32Data.GetData() + (i * NumInstances * sizeof(int32));
-			CheckedMemcpy(CacheBuffer.Int32Data, Dest, DataBuffer.GetInt32Buffer(), Source, sizeof(int32) * NumInstances);
-			//FMemory::Memcpy(Dest, Source, sizeof(int32) * NumInstances);
-		}
-
-		// Copy ID to Index Table
-		CacheBuffer.IDToIndexTable = DataBuffer.GetIDTable();
-		CacheBuffer.IDAcquireTag = DataBuffer.GetIDAcquireTag();
-	}
-
-	void WriteDataBufferGPU(FNiagaraEmitterInstance& EmitterInstance, const FNiagaraDataBuffer& DataBuffer, const FNiagaraSimCacheDataBuffersLayout& CacheLayout, FNiagaraSimCacheDataBuffers& CacheBuffer)
-	{
-		//-TODO: Make async
-		TSharedRef<FNiagaraDataSetReadback> ReadbackRequest = MakeShared<FNiagaraDataSetReadback>();
-		ReadbackRequest->ImmediateReadback(&EmitterInstance);
-		if ( FNiagaraDataBuffer* CurrentData = ReadbackRequest->GetDataSet().GetCurrentData() )
-		{
-			WriteDataBuffer(*CurrentData, CacheLayout, CacheBuffer, 0, CurrentData->GetNumInstances());
-		}
-	}
-
-	void ReadDataBuffer(const FTransform& RebaseTransform, const FNiagaraSimCacheDataBuffersLayout& CacheLayout, const FNiagaraSimCacheDataBuffers& CacheBuffer, FNiagaraDataSet& DataSet)
-	{
-		FNiagaraDataBuffer& DataBuffer = DataSet.BeginSimulate();
-		DataBuffer.Allocate(CacheBuffer.NumInstances);
-		DataBuffer.SetNumInstances(CacheBuffer.NumInstances);
-		if ( CacheBuffer.NumInstances > 0 )
-		{
-			int32 iComponent = 0;
-			const int32 NumInstances = CacheBuffer.NumInstances;
-
-			// Copy Float
-			for (uint32 i=0; i < CacheLayout.FloatCount; ++i)
-			{
-				const uint32 Component = CacheLayout.ComponentMappingsToDataBuffer[iComponent++];
-				if (Component == InvalidComponent)
-				{
-					continue;
-				}
-				const uint8* Source = CacheBuffer.FloatData.GetData() + (i * NumInstances * sizeof(float));
-				uint8* Dest = DataBuffer.GetComponentPtrFloat(Component);
-				CheckedMemcpy(DataBuffer.GetFloatBuffer(), Dest, CacheBuffer.FloatData, Source, sizeof(float) * NumInstances);
-				//FMemory::Memcpy(Dest, Source, sizeof(float) * NumInstances);
-			}
-
-			// Copy Half
-			for (uint32 i = 0; i < CacheLayout.HalfCount; ++i)
-			{
-				const uint32 Component = CacheLayout.ComponentMappingsToDataBuffer[iComponent++];
-				if (Component == InvalidComponent)
-				{
-					continue;
-				}
-				const uint8* Source = CacheBuffer.HalfData.GetData() + (i * NumInstances * sizeof(FFloat16));
-				uint8* Dest = DataBuffer.GetComponentPtrHalf(Component);
-				CheckedMemcpy(DataBuffer.GetHalfBuffer(), Dest, CacheBuffer.HalfData, Source, sizeof(FFloat16) * NumInstances);
-				//FMemory::Memcpy(Dest, Source, sizeof(FFloat16) * NumInstances);
-			}
-
-			// Copy Int32
-			for (uint32 i = 0; i < CacheLayout.Int32Count; ++i)
-			{
-				const uint32 Component = CacheLayout.ComponentMappingsToDataBuffer[iComponent++];
-				if (Component == InvalidComponent)
-				{
-					continue;
-				}
-				const uint8* Source = CacheBuffer.Int32Data.GetData() + (i * NumInstances * sizeof(int32));
-				uint8* Dest = DataBuffer.GetComponentPtrInt32(Component);
-				CheckedMemcpy(DataBuffer.GetInt32Buffer(), Dest, CacheBuffer.Int32Data, Source, sizeof(int32) * NumInstances);
-				//FMemory::Memcpy(Dest, Source, sizeof(int32) * NumInstances);
-			}
-
-			// Copy variables that require processing
-			for ( const FNiagaraSimCacheDataBuffersLayout::FVariableCopyInfo& VariableCopyInfo : CacheLayout.VariableMappingsToDataBuffer )
-			{
-				const uint32 SrcStride = uint32(NumInstances) * sizeof(float);
-				const uint8* Src = CacheBuffer.FloatData.GetData() + (VariableCopyInfo.ComponentFrom * SrcStride);
-				uint8* Dst = DataBuffer.GetComponentPtrFloat(VariableCopyInfo.ComponentTo);
-				VariableCopyInfo.CopyFunc(Dst, DataBuffer.GetFloatStride(), Src, SrcStride, uint32(NumInstances), RebaseTransform);
-			}
-		}
-
-		//-TODO:DestinationDataBuffer.SetIDTable(CacheBuffer.IDToIndexTable);
-		DataBuffer.SetIDAcquireTag(CacheBuffer.IDAcquireTag);
-
-		DataSet.EndSimulate();
-	}
-
-	void ReadDataBufferGPU(const FTransform& InRebaseTransform, FNiagaraEmitterInstance& EmitterInstance, const FNiagaraSimCacheDataBuffersLayout& InCacheLayout, const FNiagaraSimCacheDataBuffers& InCacheBuffer, FNiagaraDataSet& InDataSet, std::atomic<int32>& InPendingCommandsCounter)
-	{
-		if (EmitterInstance.IsDisabled())
-		{
-			return;
-		}
-
-		++InPendingCommandsCounter;
-
-		check(EmitterInstance.GetGPUContext());
-
-		FNiagaraGpuComputeDispatchInterface* DispathInterface = EmitterInstance.GetParentSystemInstance()->GetComputeDispatchInterface();
-		ENQUEUE_RENDER_COMMAND(NiagaraSimCacheGpuReadFrame)(
-			[DispathInterface, GPUExecContext=EmitterInstance.GetGPUContext(), RebaseTransform=InRebaseTransform, CacheLayout=&InCacheLayout, CacheBuffer=&InCacheBuffer, DataSet=&InDataSet, PendingCommandsCounter=&InPendingCommandsCounter](FRHICommandListImmediate& RHICmdList)
-			{
-				const int32 NumInstances = CacheBuffer->NumInstances;
-
-				// Set Instance Count
-				{
-					FNiagaraGPUInstanceCountManager& CountManager = DispathInterface->GetGPUInstanceCounterManager();
-					if (GPUExecContext->CountOffset_RT == INDEX_NONE)
-					{
-						GPUExecContext->CountOffset_RT = CountManager.AcquireOrAllocateEntry(RHICmdList);
-					}
-
-					const FRWBuffer& CountBuffer = CountManager.GetInstanceCountBuffer();
-					const TPair<uint32, int32> DataToSet(GPUExecContext->CountOffset_RT, NumInstances);
-					RHICmdList.Transition(FRHITransitionInfo(CountBuffer.UAV, FNiagaraGPUInstanceCountManager::kCountBufferDefaultState, ERHIAccess::UAVCompute));
-					NiagaraClearCounts::ClearCountsInt(RHICmdList, CountBuffer.UAV, MakeArrayView(&DataToSet, 1));
-					RHICmdList.Transition(FRHITransitionInfo(CountBuffer.UAV, ERHIAccess::UAVCompute, FNiagaraGPUInstanceCountManager::kCountBufferDefaultState));
-				}
-
-				// Copy instance counts
-				FNiagaraDataBuffer& DataBuffer = DataSet->GetCurrentDataChecked();
-				DataBuffer.AllocateGPU(RHICmdList, NumInstances, DispathInterface->GetFeatureLevel(), TEXT("NiagaraSimCache"));
-				DataBuffer.SetNumInstances(NumInstances);
-				DataBuffer.SetGPUDataReadyStage(ENiagaraGpuComputeTickStage::PreInitViews);
-				GPUExecContext->SetDataToRender(&DataBuffer);
-
-				if (CacheBuffer->NumInstances > 0 )
-				{
-					int32 iComponent = 0;
-
-					// Copy Float
-					if ( CacheLayout->FloatCount > 0 )
-					{
-						FRWBuffer& RWBuffer = DataBuffer.GetGPUBufferFloat();
-						const int32 RWComponentStride = DataBuffer.GetFloatStride();
-						uint8* RWBufferMemory = reinterpret_cast<uint8*>(RHILockBuffer(RWBuffer.Buffer, 0, RWBuffer.NumBytes, RLM_WriteOnly));
-
-						for (uint32 i=0; i < CacheLayout->FloatCount; ++i)
-						{
-							const uint32 Component = CacheLayout->ComponentMappingsToDataBuffer[iComponent++];
-							if (Component == InvalidComponent)
-							{
-								continue;
-							}
-							const uint8* Source = CacheBuffer->FloatData.GetData() + (i * NumInstances * sizeof(float));
-							uint8* Dest = RWBufferMemory + (uint32(Component) * RWComponentStride);
-							CheckedMemcpy(MakeArrayView(RWBufferMemory, RWBuffer.NumBytes), Dest, CacheBuffer->FloatData, Source, sizeof(float) * NumInstances);
-							//FMemory::Memcpy(Dest, Source, sizeof(float) * NumInstances);
-						}
-
-						// Copy variables that require processing
-						for (const FNiagaraSimCacheDataBuffersLayout::FVariableCopyInfo& VariableCopyInfo : CacheLayout->VariableMappingsToDataBuffer)
-						{
-							const uint32 SrcStride = uint32(NumInstances) * sizeof(float);
-							const uint8* Src = CacheBuffer->FloatData.GetData() + (VariableCopyInfo.ComponentFrom * SrcStride);
-							uint8* Dst = RWBufferMemory + (uint32(VariableCopyInfo.ComponentTo) * RWComponentStride);
-							VariableCopyInfo.CopyFunc(Dst, DataBuffer.GetFloatStride(), Src, SrcStride, uint32(NumInstances), RebaseTransform);
-						}
-
-						RHIUnlockBuffer(RWBuffer.Buffer);
-					}
-
-					// Copy Half
-					if (CacheLayout->HalfCount > 0)
-					{
-						FRWBuffer& RWBuffer = DataBuffer.GetGPUBufferHalf();
-						const int32 RWComponentStride = DataBuffer.GetHalfStride();
-						uint8* RWBufferMemory = reinterpret_cast<uint8*>(RHILockBuffer(RWBuffer.Buffer, 0, RWBuffer.NumBytes, RLM_WriteOnly));
-
-						for (uint32 i = 0; i < CacheLayout->HalfCount; ++i)
-						{
-							const uint32 Component = CacheLayout->ComponentMappingsToDataBuffer[iComponent++];
-							if (Component == InvalidComponent)
-							{
-								continue;
-							}
-							const uint8* Source = CacheBuffer->HalfData.GetData() + (i * NumInstances * sizeof(FFloat16));
-							uint8* Dest = RWBufferMemory + (uint32(Component) * RWComponentStride);
-							CheckedMemcpy(MakeArrayView(RWBufferMemory, RWBuffer.NumBytes), Dest, CacheBuffer->HalfData, Source, sizeof(FFloat16) * NumInstances);
-							//FMemory::Memcpy(Dest, Source, sizeof(FFloat16) * NumInstances);
-						}
-
-						RHIUnlockBuffer(RWBuffer.Buffer);
-					}
-
-
-					// Copy Int32
-					if (CacheLayout->Int32Count > 0)
-					{
-						FRWBuffer& RWBuffer = DataBuffer.GetGPUBufferInt();
-						const int32 RWComponentStride = DataBuffer.GetInt32Stride();
-						uint8* RWBufferMemory = reinterpret_cast<uint8*>(RHILockBuffer(RWBuffer.Buffer, 0, RWBuffer.NumBytes, RLM_WriteOnly));
-
-						for (uint32 i = 0; i < CacheLayout->Int32Count; ++i)
-						{
-							const uint32 Component = CacheLayout->ComponentMappingsToDataBuffer[iComponent++];
-							if (Component == InvalidComponent)
-							{
-								continue;
-							}
-							const uint8* Source = CacheBuffer->Int32Data.GetData() + (i * NumInstances * sizeof(int32));
-							uint8* Dest = RWBufferMemory + (uint32(Component) * RWComponentStride);
-							CheckedMemcpy(MakeArrayView(RWBufferMemory, RWBuffer.NumBytes), Dest, CacheBuffer->Int32Data, Source, sizeof(int32) * NumInstances);
-							//FMemory::Memcpy(Dest, Source, sizeof(int32) * NumInstances);
-						}
-
-						RHIUnlockBuffer(RWBuffer.Buffer);
-					}
-				}
-
-				//-TODO:DestinationDataBuffer.SetIDTable(CacheBuffer.IDToIndexTable);
-				DataBuffer.SetIDAcquireTag(CacheBuffer->IDAcquireTag);
-
-				// Ensure we decrement our counter so the GameThread knows the state of things
-				--(*PendingCommandsCounter);
-			}
-		);
-	}
-
-	static bool CanRebaseVariable(const FNiagaraVariableBase& Variable)
-	{
-		return	
-			(Variable.GetType() == FNiagaraTypeDefinition::GetQuatDef()) ||
-			(Variable.GetType() == FNiagaraTypeDefinition::GetMatrix4Def()) ||
-			(Variable.GetType() == FNiagaraTypeDefinition::GetPositionDef());
-	}
-
-	static void CopyPositions(uint8* Dst, uint32 DstStride, const uint8* Src, uint32 SrcStride, uint32 NumInstances, const FTransform& RebaseTransform)
-	{
-		float* DstFloats = reinterpret_cast<float*>(Dst);
-		DstStride = DstStride >> 2;
-
-		const float* SrcFloats = reinterpret_cast<const float*>(Src);
-		SrcStride = SrcStride >> 2;
-
-		for (uint32 i = 0; i < NumInstances; ++i)
-		{
-			const FVector CachePosition(
-				SrcFloats[i + (SrcStride * 0)],
-				SrcFloats[i + (SrcStride * 1)],
-				SrcFloats[i + (SrcStride * 2)]
-			);
-			const FVector RebasedPosition = RebaseTransform.TransformPosition(CachePosition);
-			DstFloats[i + (DstStride * 0)] = RebasedPosition.X;
-			DstFloats[i + (DstStride * 1)] = RebasedPosition.Y;
-			DstFloats[i + (DstStride * 2)] = RebasedPosition.Z;
-		}
-	}
-
-	static void CopyQuaternions(uint8* Dst, uint32 DstStride, const uint8* Src, uint32 SrcStride, uint32 NumInstances, const FTransform& RebaseTransform)
-	{
-		float* DstFloats = reinterpret_cast<float*>(Dst);
-		DstStride = DstStride >> 2;
-
-		const float* SrcFloats = reinterpret_cast<const float*>(Src);
-		SrcStride = SrcStride >> 2;
-
-		for (uint32 i = 0; i < NumInstances; ++i)
-		{
-			const FQuat4f CacheRotation(
-				SrcFloats[i + (SrcStride * 0)],
-				SrcFloats[i + (SrcStride * 1)],
-				SrcFloats[i + (SrcStride * 2)],
-				SrcFloats[i + (SrcStride * 3)]
-			);
-			const FQuat4f RebasedQuat = CacheRotation * FQuat4f(RebaseTransform.GetRotation());
-			DstFloats[i + (DstStride * 0)] = RebasedQuat.X;
-			DstFloats[i + (DstStride * 1)] = RebasedQuat.Y;
-			DstFloats[i + (DstStride * 2)] = RebasedQuat.Z;
-			DstFloats[i + (DstStride * 3)] = RebasedQuat.W;
-		}
-	}
-
-	static void CopyMatrices(uint8* Dst, uint32 DstStride, const uint8* Src, uint32 SrcStride, uint32 NumInstances, const FTransform& RebaseTransform)
-	{
-		float* DstFloats = reinterpret_cast<float*>(Dst);
-		DstStride = DstStride >> 2;
-
-		const float* SrcFloats = reinterpret_cast<const float*>(Src);
-		SrcStride = SrcStride >> 2;
-
-		const FMatrix44d RebaseMatrix = RebaseTransform.ToMatrixWithScale();
-		for (uint32 i = 0; i < NumInstances; ++i)
-		{
-			FMatrix44d CacheMatrix;
-			for (int32 j = 0; j < 16; ++j)
-			{
-				CacheMatrix.M[j >> 2][j & 0x3] = double(SrcFloats[i + (SrcStride * j)]);
-			}
-
-			CacheMatrix = CacheMatrix * RebaseMatrix;
-
-			for (int32 j = 0; j < 16; ++j)
-			{
-				DstFloats[i + (DstStride * j)] = float(CacheMatrix.M[j >> 2][j & 0x3]);
-			}
-		}
-	}
-
-	UNiagaraSystem*						NiagaraSystem = nullptr;
-	FNiagaraSystemInstance*				SystemInstance = nullptr;
-	FNiagaraSystemSimulationPtr			SystemSimulation = nullptr;
-	FNiagaraDataBuffer*					SystemSimulationDataBuffer = nullptr;
-
-	static constexpr uint16				InvalidComponent = INDEX_NONE;
-};
-
-//////////////////////////////////////////////////////////////////////////
-
 UNiagaraSimCache::UNiagaraSimCache(const FObjectInitializer& ObjectInitializer)
 {
 }
@@ -788,6 +118,7 @@ void UNiagaraSimCache::WriteFrame(UNiagaraComponent* NiagaraComponent)
 	// Cache frame
 	FNiagaraSimCacheFrame& CacheFrame = CacheFrames.AddDefaulted_GetRef();
 	CacheFrame.LocalToWorld = Helper.SystemInstance->GatheredInstanceParameters.ComponentTrans;
+	CacheFrame.LWCTile = Helper.SystemInstance->GetLWCTile();
 
 	CacheFrame.SystemData.LocalBounds = Helper.SystemInstance->GetLocalBounds();
 
@@ -1039,3 +370,178 @@ bool UNiagaraSimCache::ReadFrame(int32 FrameIndex, float FrameFraction, FNiagara
 	}
 	return true;
 }
+
+TArray<FName> UNiagaraSimCache::GetEmitterNames() const
+{
+	TArray<FName> EmitterNames;
+	if ( IsCacheValid() )
+	{
+		EmitterNames.Reserve(CacheLayout.EmitterLayouts.Num());
+		for (const FNiagaraSimCacheDataBuffersLayout& EmitterLayout : CacheLayout.EmitterLayouts)
+		{
+			EmitterNames.Add(EmitterLayout.LayoutName);
+		}
+	}
+
+	return EmitterNames;
+}
+
+void UNiagaraSimCache::ReadIntAttribute(TArray<int32>& OutValues, FName AttributeName, FName EmitterName, int FrameIndex) const
+{
+	FNiagaraSimCacheAttributeReaderHelper AttributeReader(this, EmitterName, AttributeName, FrameIndex);
+	if (AttributeReader.IsValid() && AttributeReader.Variable->Variable.GetType() == FNiagaraTypeDefinition::GetIntDef())
+	{
+		const int32 OutValueOffset = OutValues.AddUninitialized(AttributeReader.GetNumInstances());
+		for (int32 i = 0; i < AttributeReader.GetNumInstances(); ++i)
+		{
+			OutValues[OutValueOffset + i] = AttributeReader.ReadInt(i);
+		}
+	}
+}
+
+void UNiagaraSimCache::ReadFloatAttribute(TArray<float>& OutValues, FName AttributeName, FName EmitterName, int FrameIndex) const
+{
+	FNiagaraSimCacheAttributeReaderHelper AttributeReader(this, EmitterName, AttributeName, FrameIndex);
+	if (AttributeReader.IsValid() && AttributeReader.Variable->Variable.GetType() == FNiagaraTypeDefinition::GetFloatDef())
+	{
+		const int32 OutValueOffset = OutValues.AddUninitialized(AttributeReader.GetNumInstances());
+		for (int32 i = 0; i < AttributeReader.GetNumInstances(); ++i)
+		{
+			OutValues[OutValueOffset + i] = AttributeReader.ReadFloat(i);
+		}
+	}
+}
+
+void UNiagaraSimCache::ReadVector2Attribute(TArray<FVector2D>& OutValues, FName AttributeName, FName EmitterName, int FrameIndex) const
+{
+	FNiagaraSimCacheAttributeReaderHelper AttributeReader(this, EmitterName, AttributeName, FrameIndex);
+	if (AttributeReader.IsValid() && AttributeReader.Variable->Variable.GetType() == FNiagaraTypeDefinition::GetVec2Def())
+	{
+		const int32 OutValueOffset = OutValues.AddUninitialized(AttributeReader.GetNumInstances());
+		for (int32 i = 0; i < AttributeReader.GetNumInstances(); ++i)
+		{
+			OutValues[OutValueOffset + i] = AttributeReader.ReadFloat2(i);
+		}
+	}
+}
+
+void UNiagaraSimCache::ReadVectorAttribute(TArray<FVector>& OutValues, FName AttributeName, FName EmitterName, int FrameIndex) const
+{
+	FNiagaraSimCacheAttributeReaderHelper AttributeReader(this, EmitterName, AttributeName, FrameIndex);
+	if (AttributeReader.IsValid() && AttributeReader.Variable->Variable.GetType() == FNiagaraTypeDefinition::GetVec3Def())
+	{
+		const int32 OutValueOffset = OutValues.AddUninitialized(AttributeReader.GetNumInstances());
+		for (int32 i=0; i < AttributeReader.GetNumInstances(); ++i)
+		{
+			OutValues[OutValueOffset + i] = AttributeReader.ReadFloat3(i);
+		}
+	}
+}
+
+void UNiagaraSimCache::ReadVector4Attribute(TArray<FVector4>& OutValues, FName AttributeName, FName EmitterName, int FrameIndex) const
+{
+	FNiagaraSimCacheAttributeReaderHelper AttributeReader(this, EmitterName, AttributeName, FrameIndex);
+	if (AttributeReader.IsValid() && AttributeReader.Variable->Variable.GetType() == FNiagaraTypeDefinition::GetVec4Def())
+	{
+		const int32 OutValueOffset = OutValues.AddUninitialized(AttributeReader.GetNumInstances());
+		for (int32 i = 0; i < AttributeReader.GetNumInstances(); ++i)
+		{
+			OutValues[OutValueOffset + i] = AttributeReader.ReadFloat4(i);
+		}
+	}
+}
+
+void UNiagaraSimCache::ReadColorAttribute(TArray<FLinearColor>& OutValues, FName AttributeName, FName EmitterName, int FrameIndex) const
+{
+	FNiagaraSimCacheAttributeReaderHelper AttributeReader(this, EmitterName, AttributeName, FrameIndex);
+	if (AttributeReader.IsValid() && AttributeReader.Variable->Variable.GetType() == FNiagaraTypeDefinition::GetColorDef())
+	{
+		const int32 OutValueOffset = OutValues.AddUninitialized(AttributeReader.GetNumInstances());
+		for (int32 i = 0; i < AttributeReader.GetNumInstances(); ++i)
+		{
+			OutValues[OutValueOffset + i] = AttributeReader.ReadColor(i);
+		}
+	}
+}
+
+void UNiagaraSimCache::ReadPositionAttribute(TArray<FVector>& OutValues, FName AttributeName, FName EmitterName, bool bLocalSpaceToWorld, int FrameIndex) const
+{
+	FNiagaraSimCacheAttributeReaderHelper AttributeReader(this, EmitterName, AttributeName, FrameIndex);
+	if (AttributeReader.IsValid() && AttributeReader.Variable->Variable.GetType() == FNiagaraTypeDefinition::GetPositionDef())
+	{
+		FTransform LocalToWorld;
+		if ( bLocalSpaceToWorld && AttributeReader.DataBuffersLayout->bLocalSpace )
+		{
+			LocalToWorld = AttributeReader.CacheFrame->LocalToWorld;
+		}
+		else
+		{
+			LocalToWorld = FTransform::Identity;
+			LocalToWorld.SetTranslation(FVector(AttributeReader.CacheFrame->LWCTile) * FLargeWorldRenderScalar::GetTileSize());
+		}
+
+		const int32 OutValueOffset = OutValues.AddUninitialized(AttributeReader.GetNumInstances());
+		for (int32 i = 0; i < AttributeReader.GetNumInstances(); ++i)
+		{
+			const FVector Position = AttributeReader.ReadFloat3(i);
+			OutValues[OutValueOffset + i] = LocalToWorld.TransformPosition(Position);
+		}
+	}
+}
+
+void UNiagaraSimCache::ReadPositionAttributeWithRebase(TArray<FVector>& OutValues, FName AttributeName, FName EmitterName, FTransform Transform, int FrameIndex) const
+{
+	FNiagaraSimCacheAttributeReaderHelper AttributeReader(this, EmitterName, AttributeName, FrameIndex);
+	if (AttributeReader.IsValid() && AttributeReader.Variable->Variable.GetType() == FNiagaraTypeDefinition::GetPositionDef())
+	{
+		const FTransform LocalToWorld = AttributeReader.DataBuffersLayout->bLocalSpace ? Transform : Transform * AttributeReader.CacheFrame->LocalToWorld.Inverse();
+		const int32 OutValueOffset = OutValues.AddUninitialized(AttributeReader.GetNumInstances());
+		for (int32 i = 0; i < AttributeReader.GetNumInstances(); ++i)
+		{
+			const FVector Position = AttributeReader.ReadFloat3(i);
+			OutValues[OutValueOffset + i] = LocalToWorld.TransformPosition(Position);
+		}
+	}
+}
+
+void UNiagaraSimCache::ReadQuatAttribute(TArray<FQuat>& OutValues, FName AttributeName, FName EmitterName, bool bLocalSpaceToWorld, int FrameIndex) const
+{
+	FNiagaraSimCacheAttributeReaderHelper AttributeReader(this, EmitterName, AttributeName, FrameIndex);
+	if (AttributeReader.IsValid() && AttributeReader.Variable->Variable.GetType() == FNiagaraTypeDefinition::GetQuatDef())
+	{
+		const int32 OutValueOffset = OutValues.AddUninitialized(AttributeReader.GetNumInstances());
+		if (bLocalSpaceToWorld && AttributeReader.DataBuffersLayout->bLocalSpace)
+		{
+			for (int32 i = 0; i < AttributeReader.GetNumInstances(); ++i)
+			{
+				const FQuat Quat = AttributeReader.ReadQuat(i);
+				OutValues[OutValueOffset + i] = Quat * AttributeReader.CacheFrame->LocalToWorld.GetRotation();
+			}
+		}
+		else
+		{
+			for (int32 i = 0; i < AttributeReader.GetNumInstances(); ++i)
+			{
+				const FQuat Quat = AttributeReader.ReadQuat(i);
+				OutValues[OutValueOffset + i] = Quat;
+			}
+		}
+	}
+}
+
+void UNiagaraSimCache::ReadQuatAttributeWithRebase(TArray<FQuat>& OutValues, FName AttributeName, FName EmitterName, FQuat Quat, int FrameIndex) const
+{
+	FNiagaraSimCacheAttributeReaderHelper AttributeReader(this, EmitterName, AttributeName, FrameIndex);
+	if (AttributeReader.IsValid() && AttributeReader.Variable->Variable.GetType() == FNiagaraTypeDefinition::GetQuatDef())
+	{
+		const bool bRebase = AttributeReader.DataBuffersLayout->bLocalSpace == false && AttributeReader.DataBuffersLayout->RebaseVariableNames.Contains(AttributeName);
+
+		const FQuat LocalToWorld = bRebase ? Quat * AttributeReader.CacheFrame->LocalToWorld.GetRotation().Inverse() : Quat;
+		const int32 OutValueOffset = OutValues.AddUninitialized(AttributeReader.GetNumInstances());
+		for (int32 i = 0; i < AttributeReader.GetNumInstances(); ++i)
+		{
+			const FQuat Rotation = AttributeReader.ReadQuat(i);
+			OutValues[OutValueOffset + i] = Rotation * LocalToWorld;
+		}
+	}
+}
diff --git a/Engine/Plugins/FX/Niagara/Source/Niagara/Private/NiagaraSimCacheAttributeReaderHelper.h b/Engine/Plugins/FX/Niagara/Source/Niagara/Private/NiagaraSimCacheAttributeReaderHelper.h
new file mode 100644
--- /dev/null
+++ b/Engine/Plugins/FX/Niagara/Source/Niagara/Private/NiagaraSimCacheAttributeReaderHelper.h
@@ -0,0 +1,120 @@
+// Copyright Epic Games, Inc. All Rights Reserved.
+
+#pragma once
+
+#include "NiagaraSimCache.h"
+
+struct FNiagaraSimCacheAttributeReaderHelper
+{
+	explicit FNiagaraSimCacheAttributeReaderHelper(const UNiagaraSimCache* SimCache, FName EmitterName, FName AttributeName, int FrameIndex)
+	{
+		if (SimCache->IsCacheValid() == false)
+		{
+			return;
+		}
+
+		if (SimCache->CacheFrames.IsValidIndex(FrameIndex) == false)
+		{
+			return;
+		}
+		CacheFrame = &SimCache->CacheFrames[FrameIndex];
+
+		if ( EmitterName.IsNone() == false )
+		{
+			const int32 EmitterIndex = SimCache->CacheLayout.EmitterLayouts.IndexOfByPredicate([&EmitterName](const FNiagaraSimCacheDataBuffersLayout& FoundLayout) { return FoundLayout.LayoutName == EmitterName; });
+			if (EmitterIndex == INDEX_NONE)
+			{
+				return;
+			}
+
+			DataBuffers = &CacheFrame->EmitterData[EmitterIndex].ParticleDataBuffers;
+			DataBuffersLayout = &SimCache->CacheLayout.EmitterLayouts[EmitterIndex];
+		}
+		else
+		{
+			DataBuffers = &CacheFrame->SystemData.SystemDataBuffers;
+			DataBuffersLayout = &SimCache->CacheLayout.SystemLayout;
+		}
+
+		if (DataBuffers->NumInstances == 0)
+		{
+			return;
+		}
+
+		const int32 VariableIndex = DataBuffersLayout->Variables.IndexOfByPredicate([&AttributeName](const FNiagaraSimCacheVariable& FoundVariable) { return FoundVariable.Variable.GetName() == AttributeName; });
+		if (VariableIndex == INDEX_NONE)
+		{
+			return;
+		}
+
+		Variable = &DataBuffersLayout->Variables[VariableIndex];
+	}
+
+	bool IsValid() const { return Variable != nullptr; }
+
+	int32 GetNumInstances() const { check(IsValid()); return DataBuffers->NumInstances; }
+
+	float ReadInt(int32 Instance) const
+	{
+		check(IsValid());
+		check(Variable->Int32Offset != INDEX_NONE && Variable->Int32Count == 1);
+
+		const int32 Int32Offset = Instance + (Variable->Int32Offset * DataBuffers->NumInstances);
+
+		int32 Value;
+		FMemory::Memcpy(&Value, &DataBuffers->Int32Data[Int32Offset * sizeof(int32)], sizeof(int32));
+		return Value;
+	}
+
+	float InternalReadFloat(int32 Offset, int32 Instance) const
+	{
+		check(IsValid());
+		
+		const int32 FloatOffset = Instance + ((Variable->FloatOffset + Offset) * DataBuffers->NumInstances);
+
+		float Value;
+		FMemory::Memcpy(&Value, &DataBuffers->FloatData[FloatOffset * sizeof(float)], sizeof(float));
+		return Value;
+	}
+
+	float ReadFloat(int32 Instance) const
+	{
+		check(Variable->FloatOffset != INDEX_NONE && Variable->FloatCount == 1);
+		return InternalReadFloat(0, Instance);
+	}
+
+	FVector2D ReadFloat2(int32 Instance) const
+	{
+		check(Variable->FloatOffset != INDEX_NONE && Variable->FloatCount == 2);
+		return FVector2D(InternalReadFloat(0, Instance), InternalReadFloat(1, Instance));
+	}
+
+	FVector ReadFloat3(int32 Instance) const
+	{
+		check(Variable->FloatOffset != INDEX_NONE && Variable->FloatCount == 3);
+		return FVector(InternalReadFloat(0, Instance), InternalReadFloat(1, Instance), InternalReadFloat(2, Instance));
+	}
+
+	FVector4 ReadFloat4(int32 Instance) const
+	{
+		check(Variable->FloatOffset != INDEX_NONE && Variable->FloatCount == 4);
+		return FVector4(InternalReadFloat(0, Instance), InternalReadFloat(1, Instance), InternalReadFloat(2, Instance), InternalReadFloat(3, Instance));
+	}
+
+	FLinearColor ReadColor(int32 Instance) const
+	{
+		check(Variable->FloatOffset != INDEX_NONE && Variable->FloatCount == 4);
+		return FLinearColor(InternalReadFloat(0, Instance), InternalReadFloat(1, Instance), InternalReadFloat(2, Instance), InternalReadFloat(3, Instance));
+	}
+
+	FQuat ReadQuat(int32 Instance) const
+	{
+		check(Variable->FloatOffset != INDEX_NONE && Variable->FloatCount == 4);
+		return FQuat(InternalReadFloat(0, Instance), InternalReadFloat(1, Instance), InternalReadFloat(2, Instance), InternalReadFloat(3, Instance));
+	}
+	
+	const FNiagaraSimCacheFrame*				CacheFrame = nullptr;
+	const FNiagaraSimCacheDataBuffers*			DataBuffers = nullptr;
+	const FNiagaraSimCacheDataBuffersLayout*	DataBuffersLayout = nullptr;
+	const FNiagaraSimCacheVariable*				Variable = nullptr;
+};
diff --git a/Engine/Plugins/FX/Niagara/Source/Niagara/Private/NiagaraSimCacheHelper.h b/Engine/Plugins/FX/Niagara/Source/Niagara/Private/NiagaraSimCacheHelper.h
new file mode 100644
--- /dev/null
+++ b/Engine/Plugins/FX/Niagara/Source/Niagara/Private/NiagaraSimCacheHelper.h
@@ -0,0 +1,677 @@
+// Copyright Epic Games, Inc. All Rights Reserved.
+
+#pragma  once
+
+#include "NiagaraSimCache.h"
+
+struct FNiagaraSimCacheHelper
+{
+	explicit FNiagaraSimCacheHelper(FNiagaraSystemInstance* InSystemInstance)
+	{
+		SystemInstance = InSystemInstance;
+		SystemSimulation = SystemInstance->GetSystemSimulation();
+		check(SystemSimulation);
+		SystemSimulationDataBuffer = SystemSimulation->MainDataSet.GetCurrentData();
+		NiagaraSystem = SystemSimulation->GetSystem();
+	}
+
+	explicit FNiagaraSimCacheHelper(UNiagaraComponent* NiagaraComponent)
+	{
+		if ( NiagaraComponent == nullptr )
+		{
+			return;
+		}
+
+		NiagaraSystem = NiagaraComponent->GetAsset();
+		if ( NiagaraSystem == nullptr )
+		{
+			return;
+		}
+
+		FNiagaraSystemInstanceControllerPtr SystemInstanceController = NiagaraComponent->GetSystemInstanceController();
+		if (SystemInstanceController.IsValid() == false)
+		{
+			return;
+		}
+
+		SystemInstance = SystemInstanceController->GetSystemInstance_Unsafe();
+		if (SystemInstance == nullptr)
+		{
+			return;
+		}
+
+		SystemSimulation = SystemInstance->GetSystemSimulation();
+		if (SystemSimulation == nullptr)
+		{
+			return;
+		}
+
+		SystemSimulationDataBuffer = SystemSimulation->MainDataSet.GetCurrentData();
+		if (SystemSimulationDataBuffer == nullptr)
+		{
+			return;
+		}
+	}
+
+	FNiagaraDataSet& GetSystemSimulationDataSet() { return SystemSimulation->MainDataSet; }
+
+	bool HasValidSimulation() const { return SystemSimulation != nullptr; }
+	bool HasValidSimulationData() const { return SystemSimulationDataBuffer != nullptr; }
+
+	void BuildCacheLayout(FNiagaraSimCacheDataBuffersLayout& CacheLayout, const FNiagaraDataSetCompiledData& CompiledData, FName LayoutName, TArray<FName> InRebaseVariableNames) const
+	{
+		CacheLayout.LayoutName = LayoutName;
+		CacheLayout.SimTarget = CompiledData.SimTarget;
+
+		const int32 NumVariables = CompiledData.Variables.Num();
+		CacheLayout.Variables.AddDefaulted(NumVariables);
+
+		const int32 CacheTotalComponents = CompiledData.TotalFloatComponents + CompiledData.TotalFloatComponents + CompiledData.TotalInt32Components;
+		CacheLayout.ComponentMappingsFromDataBuffer.Empty(CacheTotalComponents);
+		CacheLayout.ComponentMappingsFromDataBuffer.AddDefaulted(CacheTotalComponents);
+		CacheLayout.RebaseVariableNames = MoveTemp(InRebaseVariableNames);
+
+		for ( int32 iVariable=0; iVariable < NumVariables; ++iVariable)
+		{
+			const FNiagaraVariableLayoutInfo& DataSetVariableLayout = CompiledData.VariableLayouts[iVariable];
+			FNiagaraSimCacheVariable& CacheVariable = CacheLayout.Variables[iVariable];
+
+			CacheVariable.Variable = CompiledData.Variables[iVariable];
+			CacheVariable.FloatOffset = DataSetVariableLayout.GetNumFloatComponents() > 0 ? CacheLayout.FloatCount : INDEX_NONE;
+			CacheVariable.FloatCount = DataSetVariableLayout.GetNumFloatComponents();
+			CacheVariable.HalfOffset = DataSetVariableLayout.GetNumHalfComponents() > 0 ? CacheLayout.HalfCount : INDEX_NONE;
+			CacheVariable.HalfCount = DataSetVariableLayout.GetNumHalfComponents();
+			CacheVariable.Int32Offset = DataSetVariableLayout.GetNumInt32Components() > 0 ? CacheLayout.Int32Count : INDEX_NONE;
+			CacheVariable.Int32Count = DataSetVariableLayout.GetNumInt32Components();
+
+			CacheLayout.FloatCount += DataSetVariableLayout.GetNumFloatComponents();
+			CacheLayout.HalfCount += DataSetVariableLayout.GetNumHalfComponents();
+			CacheLayout.Int32Count += DataSetVariableLayout.GetNumInt32Components();
+		}
+
+		// Build write mappings we will build read mappings in a separate path
+		int32 FloatOffset = 0;
+		int32 HalfOffset = CacheLayout.FloatCount;
+		int32 Int32Offset = HalfOffset + CacheLayout.HalfCount;
+		for (int32 iVariable = 0; iVariable < NumVariables; ++iVariable)
+		{
+			const FNiagaraVariableLayoutInfo& DataSetVariableLayout = CompiledData.VariableLayouts[iVariable];
+			FNiagaraSimCacheVariable& CacheVariable = CacheLayout.Variables[iVariable];
+
+			for (int32 iComponent=0; iComponent < CacheVariable.FloatCount; ++iComponent)
+			{
+				CacheLayout.ComponentMappingsFromDataBuffer[FloatOffset] = CacheVariable.FloatOffset + iComponent;
+				++FloatOffset;
+			}
+
+			for (int32 iComponent=0; iComponent < CacheVariable.HalfCount; ++iComponent)
+			{
+				CacheLayout.ComponentMappingsFromDataBuffer[HalfOffset] = CacheVariable.HalfOffset + iComponent;
+				++HalfOffset;
+			}
+
+			for (int32 iComponent=0; iComponent < CacheVariable.Int32Count; ++iComponent)
+			{
+				CacheLayout.ComponentMappingsFromDataBuffer[Int32Offset] = CacheVariable.Int32Offset + iComponent;
+				++Int32Offset;
+			}
+		}
+
+		// Slightly inefficient but we can share the code between the paths
+		BuildCacheReadMappings(CacheLayout, CompiledData);
+	}
+
+	void BuildCacheLayoutForSystem(const FNiagaraSimCacheCreateParameters& CreateParmaeters, FNiagaraSimCacheDataBuffersLayout& CacheLayout)
+	{
+		const FNiagaraDataSetCompiledData& SystemCompileData = NiagaraSystem->GetSystemCompiledData().DataSetCompiledData;
+
+		TArray<FName> RebaseVariableNames;
+		if ( CreateParmaeters.bAllowRebasing )
+		{
+			TArray<FString, TInlineAllocator<8>> LocalSpaceEmitters;
+			for ( int32 i=0; i < NiagaraSystem->GetNumEmitters(); ++i )
+			{
+				const FNiagaraEmitterHandle& EmitterHandle = NiagaraSystem->GetEmitterHandle(i);
+				if ( EmitterHandle.GetInstance().GetEmitterData()->bLocalSpace )
+				{
+					LocalSpaceEmitters.Add(EmitterHandle.GetUniqueInstanceName());
+				}
+			}
+
+			for (const FNiagaraVariable& Variable : SystemCompileData.Variables)
+			{
+				if (Variable.GetType() == FNiagaraTypeDefinition::GetPositionDef())
+				{
+					// If this is an emitter variable we need to check if it's local space or not
+					bool bIsLocalSpace = false;
+					for ( const FString& LocalSpaceEmitter : LocalSpaceEmitters )
+					{
+						if ( Variable.IsInNameSpace(LocalSpaceEmitter) )
+						{
+							bIsLocalSpace = true;
+							break;
+						}
+					}
+
+					if ( bIsLocalSpace == false && CreateParmaeters.RebaseExcludeList.Contains(Variable.GetName()) == false )
+					{
+						RebaseVariableNames.AddUnique(Variable.GetName());
+					}
+				}
+				else if ( CanRebaseVariable(Variable) && CreateParmaeters.RebaseIncludeList.Contains(Variable.GetName()) )
+				{
+					RebaseVariableNames.AddUnique(Variable.GetName());
+				}
+			}
+		}
+
+		BuildCacheLayout(CacheLayout, SystemCompileData, NiagaraSystem->GetFName(), MoveTemp(RebaseVariableNames));
+	}
+
+	void BuildCacheLayoutForEmitter(const FNiagaraSimCacheCreateParameters& CreateParmaeters, FNiagaraSimCacheDataBuffersLayout& CacheLayout, int EmitterIndex)
+	{
+		const FNiagaraEmitterHandle& EmitterHandle = NiagaraSystem->GetEmitterHandle(EmitterIndex);
+		const FNiagaraEmitterCompiledData& EmitterCompiledData = NiagaraSystem->GetEmitterCompiledData()[EmitterIndex].Get();
+
+		// Find potential candidates for re-basing
+		CacheLayout.bLocalSpace = EmitterHandle.GetInstance().GetEmitterData()->bLocalSpace;
+
+		TArray<FName> RebaseVariableNames;
+		if ( CreateParmaeters.bAllowRebasing && CacheLayout.bLocalSpace == false )
+		{
+			// Build list of include / exclude names
+			TArray<FName> ForceIncludeNames;
+			TArray<FName> ForceExcludeNames;
+			if ( CreateParmaeters.RebaseIncludeList.Num() > 0 || CreateParmaeters.RebaseExcludeList.Num() > 0 )
+			{
+				const FString EmitterName = EmitterHandle.GetUniqueInstanceName();
+				for (FName RebaseName : CreateParmaeters.RebaseIncludeList)
+				{
+					FNiagaraVariableBase BaseVar(FNiagaraTypeDefinition::GetFloatDef(), RebaseName);
+					if (BaseVar.RemoveRootNamespace(EmitterName))
+					{
+						ForceIncludeNames.Add(BaseVar.GetName());
+					}
+				}
+
+				for (FName RebaseName : CreateParmaeters.RebaseExcludeList)
+				{
+					FNiagaraVariableBase BaseVar(FNiagaraTypeDefinition::GetFloatDef(), RebaseName);
+					if (BaseVar.RemoveRootNamespace(EmitterName))
+					{
+						ForceExcludeNames.Add(BaseVar.GetName());
+					}
+				}
+			}
+
+		#if WITH_EDITORONLY_DATA
+			// Look for renderer attributes bound to Quat / Matrix types are we will want to rebase those
+			// We will add all Position types after this so no need to add them here
+			EmitterHandle.GetInstance().GetEmitterData()->ForEachEnabledRenderer(
+				[&](UNiagaraRendererProperties* RenderProperties)
+				{
+					for (FNiagaraVariableBase BoundAttribute : RenderProperties->GetBoundAttributes())
+					{
+						if ( (BoundAttribute.GetType() == FNiagaraTypeDefinition::GetQuatDef()) ||
+							 (BoundAttribute.GetType() == FNiagaraTypeDefinition::GetMatrix4Def()) )
+						{
+							if (BoundAttribute.RemoveRootNamespace(FNiagaraConstants::ParticleAttributeNamespaceString))
+							{
+								if (EmitterCompiledData.DataSetCompiledData.Variables.Contains(BoundAttribute) && ForceExcludeNames.Contains(BoundAttribute.GetName()) == false )
+								{
+									RebaseVariableNames.AddUnique(BoundAttribute.GetName());
+								}
+							}
+						}
+					}
+				}
+			);
+		#endif
+
+			// Look for regular attributes that we are forcing to rebase or can rebase like positions
+			for (const FNiagaraVariable& Variable : EmitterCompiledData.DataSetCompiledData.Variables)
+			{
+				if ( Variable.GetType() == FNiagaraTypeDefinition::GetPositionDef() )
+				{
+					if ( ForceExcludeNames.Contains(Variable.GetName()) == false )
+					{
+						RebaseVariableNames.AddUnique(Variable.GetName());
+					}
+				}
+				else if ( ForceIncludeNames.Contains(Variable.GetName()) && CanRebaseVariable(Variable) )
+				{
+					RebaseVariableNames.AddUnique(Variable.GetName());
+				}
+			}
+		}
+
+		BuildCacheLayout(CacheLayout, EmitterCompiledData.DataSetCompiledData, EmitterHandle.GetName(), MoveTemp(RebaseVariableNames));
+	}
+
+	static bool BuildCacheReadMappings(FNiagaraSimCacheDataBuffersLayout& CacheLayout, const FNiagaraDataSetCompiledData& CompiledData)
+	{
+		const int32 CacheTotalComponents = CacheLayout.FloatCount + CacheLayout.HalfCount + CacheLayout.Int32Count;
+		CacheLayout.ComponentMappingsToDataBuffer.Empty(CacheTotalComponents);
+		CacheLayout.ComponentMappingsToDataBuffer.AddDefaulted(CacheTotalComponents);
+		CacheLayout.VariableMappingsToDataBuffer.Empty(0);
+
+		int32 FloatOffset = 0;
+		int32 HalfOffset = CacheLayout.FloatCount;
+		int32 Int32Offset = HalfOffset + CacheLayout.HalfCount;
+		for (const FNiagaraSimCacheVariable& SourceVariable : CacheLayout.Variables)
+		{
+			// Find variable, if it doesn't exist that's ok as the cache contains more data than is required
+			const int32 DataSetVariableIndex = CompiledData.Variables.IndexOfByPredicate([&](const FNiagaraVariable& DataSetVariable) { return DataSetVariable == SourceVariable.Variable; });
+			const FNiagaraVariableLayoutInfo* DestVariableLayout = nullptr;
+			if (DataSetVariableIndex != INDEX_NONE)
+			{
+				DestVariableLayout = &CompiledData.VariableLayouts[DataSetVariableIndex];
+
+				// If the variable exists but types not match the cache is invalid
+				if ((DestVariableLayout->GetNumFloatComponents() != SourceVariable.FloatCount) ||
+					(DestVariableLayout->GetNumHalfComponents() != SourceVariable.HalfCount) ||
+					(DestVariableLayout->GetNumInt32Components() != SourceVariable.Int32Count))
+				{
+					return false;
+				}
+			}
+
+			// Is this a type that requires conversion / re-basing?
+			if (DestVariableLayout != nullptr)
+			{
+				if ( CacheLayout.RebaseVariableNames.Contains(SourceVariable.Variable.GetName()) )
+				{
+					if (SourceVariable.Variable.GetType() == FNiagaraTypeDefinition::GetPositionDef())
+					{
+						check(SourceVariable.FloatCount == 3);
+						CacheLayout.VariableMappingsToDataBuffer.Emplace(FloatOffset, DestVariableLayout->FloatComponentStart, &FNiagaraSimCacheHelper::CopyPositions);
+						DestVariableLayout = nullptr;
+					}
+					else if (SourceVariable.Variable.GetType() == FNiagaraTypeDefinition::GetQuatDef())
+					{
+						check(SourceVariable.FloatCount == 4);
+						CacheLayout.VariableMappingsToDataBuffer.Emplace(FloatOffset, DestVariableLayout->FloatComponentStart, &FNiagaraSimCacheHelper::CopyQuaternions);
+						DestVariableLayout = nullptr;
+					}
+					else if (SourceVariable.Variable.GetType() == FNiagaraTypeDefinition::GetMatrix4Def())
+					{
+						check(SourceVariable.FloatCount == 16);
+						CacheLayout.VariableMappingsToDataBuffer.Emplace(FloatOffset, DestVariableLayout->FloatComponentStart, &FNiagaraSimCacheHelper::CopyMatrices);
+						DestVariableLayout = nullptr;
+					}
+				}
+			}
+
+			for (int32 i = 0; i < SourceVariable.FloatCount; ++i)
+			{
+				CacheLayout.ComponentMappingsToDataBuffer[FloatOffset++] = DestVariableLayout ? DestVariableLayout->FloatComponentStart + i : INDEX_NONE;
+			}
+
+			for (int32 i = 0; i < SourceVariable.HalfCount; ++i)
+			{
+				CacheLayout.ComponentMappingsToDataBuffer[HalfOffset++] = DestVariableLayout ? DestVariableLayout->HalfComponentStart + i : INDEX_NONE;
+			}
+
+			for (int32 i = 0; i < SourceVariable.Int32Count; ++i)
+			{
+				CacheLayout.ComponentMappingsToDataBuffer[Int32Offset++] = DestVariableLayout ? DestVariableLayout->Int32ComponentStart + i : INDEX_NONE;
+			}
+		}
+
+		return true;
+	}
+
+	static void CheckedMemcpy(TConstArrayView<uint8> DstArray, uint8* Dst, TConstArrayView<uint8> SrcArray, const uint8* Src, uint32 Size)
+	{
+		checkf(Src >= SrcArray.GetData() && Src + Size <= SrcArray.GetData() + SrcArray.Num(), TEXT("Source %p-%p is out of bounds, start %p end %p"), Src, Src + Size, SrcArray.GetData(), SrcArray.GetData() + SrcArray.Num());
+		checkf(Dst >= DstArray.GetData() && Dst + Size <= DstArray.GetData() + DstArray.Num(), TEXT("Dest %p-%p is out of bounds, start %p end %p"), Dst, Dst + Size, DstArray.GetData(), DstArray.GetData() + DstArray.Num());
+		FMemory::Memcpy(Dst, Src, Size);
+	}
+
+	void WriteDataBuffer(const FNiagaraDataBuffer& DataBuffer, const FNiagaraSimCacheDataBuffersLayout& CacheLayout, FNiagaraSimCacheDataBuffers& CacheBuffer, int32 FirstInstance, int32 NumInstances)
+	{
+		if ( NumInstances == 0 )
+		{
+			return;
+		}
+
+		CacheBuffer.NumInstances = NumInstances;
+
+		int32 iComponent = 0;
+
+		// Copy Float
+		CacheBuffer.FloatData.AddDefaulted(CacheLayout.FloatCount * NumInstances * sizeof(float));
+		for ( uint32 i=0; i < CacheLayout.FloatCount; ++i )
+		{
+			const uint32 Component = CacheLayout.ComponentMappingsFromDataBuffer[iComponent++];
+			const uint8* Source = DataBuffer.GetComponentPtrFloat(Component) + (FirstInstance * sizeof(float));
+			uint8* Dest = CacheBuffer.FloatData.GetData() + (i * NumInstances * sizeof(float));
+			CheckedMemcpy(CacheBuffer.FloatData, Dest, DataBuffer.GetFloatBuffer(), Source, sizeof(float) * NumInstances);
+			//FMemory::Memcpy(Dest, Source, sizeof(float) * NumInstances);
+		}
+		
+		// Copy Half
+		CacheBuffer.HalfData.AddDefaulted(CacheLayout.HalfCount * NumInstances * sizeof(FFloat16));
+		for (uint32 i = 0; i < CacheLayout.HalfCount; ++i)
+		{
+			const uint32 Component = CacheLayout.ComponentMappingsFromDataBuffer[iComponent++];
+			const uint8* Source = DataBuffer.GetComponentPtrHalf(Component) + (FirstInstance * sizeof(FFloat16));
+			uint8* Dest = CacheBuffer.HalfData.GetData() + (i * NumInstances * sizeof(FFloat16));
+			CheckedMemcpy(CacheBuffer.HalfData, Dest, DataBuffer.GetHalfBuffer(), Source, sizeof(FFloat16) * NumInstances);
+			//FMemory::Memcpy(Dest, Source, sizeof(FFloat16) * NumInstances);
+		}
+
+		// Copy Int32
+		CacheBuffer.Int32Data.AddDefaulted(CacheLayout.Int32Count * NumInstances * sizeof(int32));
+		for (uint32 i = 0; i < CacheLayout.Int32Count; ++i)
+		{
+			const uint32 Component = CacheLayout.ComponentMappingsFromDataBuffer[iComponent++];
+			const uint8* Source = DataBuffer.GetComponentPtrInt32(Component) + (FirstInstance * sizeof(int32));
+			uint8* Dest = CacheBuffer.Int32Data.GetData() + (i * NumInstances * sizeof(int32));
+			CheckedMemcpy(CacheBuffer.Int32Data, Dest, DataBuffer.GetInt32Buffer(), Source, sizeof(int32) * NumInstances);
+			//FMemory::Memcpy(Dest, Source, sizeof(int32) * NumInstances);
+		}
+
+		// Copy ID to Index Table
+		CacheBuffer.IDToIndexTable = DataBuffer.GetIDTable();
+		CacheBuffer.IDAcquireTag = DataBuffer.GetIDAcquireTag();
+	}
+
+	void WriteDataBufferGPU(FNiagaraEmitterInstance& EmitterInstance, const FNiagaraDataBuffer& DataBuffer, const FNiagaraSimCacheDataBuffersLayout& CacheLayout, FNiagaraSimCacheDataBuffers& CacheBuffer)
+	{
+		//-TODO: Make async
+		TSharedRef<FNiagaraDataSetReadback> ReadbackRequest = MakeShared<FNiagaraDataSetReadback>();
+		ReadbackRequest->ImmediateReadback(&EmitterInstance);
+		if ( FNiagaraDataBuffer* CurrentData = ReadbackRequest->GetDataSet().GetCurrentData() )
+		{
+			WriteDataBuffer(*CurrentData, CacheLayout, CacheBuffer, 0, CurrentData->GetNumInstances());
+		}
+	}
+
+	void ReadDataBuffer(const FTransform& RebaseTransform, const FNiagaraSimCacheDataBuffersLayout& CacheLayout, const FNiagaraSimCacheDataBuffers& CacheBuffer, FNiagaraDataSet& DataSet)
+	{
+		FNiagaraDataBuffer& DataBuffer = DataSet.BeginSimulate();
+		DataBuffer.Allocate(CacheBuffer.NumInstances);
+		DataBuffer.SetNumInstances(CacheBuffer.NumInstances);
+		if ( CacheBuffer.NumInstances > 0 )
+		{
+			int32 iComponent = 0;
+			const int32 NumInstances = CacheBuffer.NumInstances;
+
+			// Copy Float
+			for (uint32 i=0; i < CacheLayout.FloatCount; ++i)
+			{
+				const uint32 Component = CacheLayout.ComponentMappingsToDataBuffer[iComponent++];
+				if (Component == InvalidComponent)
+				{
+					continue;
+				}
+				const uint8* Source = CacheBuffer.FloatData.GetData() + (i * NumInstances * sizeof(float));
+				uint8* Dest = DataBuffer.GetComponentPtrFloat(Component);
+				CheckedMemcpy(DataBuffer.GetFloatBuffer(), Dest, CacheBuffer.FloatData, Source, sizeof(float) * NumInstances);
+				//FMemory::Memcpy(Dest, Source, sizeof(float) * NumInstances);
+			}
+
+			// Copy Half
+			for (uint32 i = 0; i < CacheLayout.HalfCount; ++i)
+			{
+				const uint32 Component = CacheLayout.ComponentMappingsToDataBuffer[iComponent++];
+				if (Component == InvalidComponent)
+				{
+					continue;
+				}
+				const uint8* Source = CacheBuffer.HalfData.GetData() + (i * NumInstances * sizeof(FFloat16));
+				uint8* Dest = DataBuffer.GetComponentPtrHalf(Component);
+				CheckedMemcpy(DataBuffer.GetHalfBuffer(), Dest, CacheBuffer.HalfData, Source, sizeof(FFloat16) * NumInstances);
+				//FMemory::Memcpy(Dest, Source, sizeof(FFloat16) * NumInstances);
+			}
+
+			// Copy Int32
+			for (uint32 i = 0; i < CacheLayout.Int32Count; ++i)
+			{
+				const uint32 Component = CacheLayout.ComponentMappingsToDataBuffer[iComponent++];
+				if (Component == InvalidComponent)
+				{
+					continue;
+				}
+				const uint8* Source = CacheBuffer.Int32Data.GetData() + (i * NumInstances * sizeof(int32));
+				uint8* Dest = DataBuffer.GetComponentPtrInt32(Component);
+				CheckedMemcpy(DataBuffer.GetInt32Buffer(), Dest, CacheBuffer.Int32Data, Source, sizeof(int32) * NumInstances);
+				//FMemory::Memcpy(Dest, Source, sizeof(int32) * NumInstances);
+			}
+
+			// Copy variables that require processing
+			for ( const FNiagaraSimCacheDataBuffersLayout::FVariableCopyInfo& VariableCopyInfo : CacheLayout.VariableMappingsToDataBuffer )
+			{
+				const uint32 SrcStride = uint32(NumInstances) * sizeof(float);
+				const uint8* Src = CacheBuffer.FloatData.GetData() + (VariableCopyInfo.ComponentFrom * SrcStride);
+				uint8* Dst = DataBuffer.GetComponentPtrFloat(VariableCopyInfo.ComponentTo);
+				VariableCopyInfo.CopyFunc(Dst, DataBuffer.GetFloatStride(), Src, SrcStride, uint32(NumInstances), RebaseTransform);
+			}
+		}
+
+		//-TODO:DestinationDataBuffer.SetIDTable(CacheBuffer.IDToIndexTable);
+		DataBuffer.SetIDAcquireTag(CacheBuffer.IDAcquireTag);
+
+		DataSet.EndSimulate();
+	}
+
+	void ReadDataBufferGPU(const FTransform& InRebaseTransform, FNiagaraEmitterInstance& EmitterInstance, const FNiagaraSimCacheDataBuffersLayout& InCacheLayout, const FNiagaraSimCacheDataBuffers& InCacheBuffer, FNiagaraDataSet& InDataSet, std::atomic<int32>& InPendingCommandsCounter)
+	{
+		if (EmitterInstance.IsDisabled())
+		{
+			return;
+		}
+
+		++InPendingCommandsCounter;
+
+		check(EmitterInstance.GetGPUContext());
+
+		FNiagaraGpuComputeDispatchInterface* DispathInterface = EmitterInstance.GetParentSystemInstance()->GetComputeDispatchInterface();
+		ENQUEUE_RENDER_COMMAND(NiagaraSimCacheGpuReadFrame)(
+			[DispathInterface, GPUExecContext=EmitterInstance.GetGPUContext(), RebaseTransform=InRebaseTransform, CacheLayout=&InCacheLayout, CacheBuffer=&InCacheBuffer, DataSet=&InDataSet, PendingCommandsCounter=&InPendingCommandsCounter](FRHICommandListImmediate& RHICmdList)
+			{
+				const int32 NumInstances = CacheBuffer->NumInstances;
+
+				// Set Instance Count
+				{
+					FNiagaraGPUInstanceCountManager& CountManager = DispathInterface->GetGPUInstanceCounterManager();
+					if (GPUExecContext->CountOffset_RT == INDEX_NONE)
+					{
+						GPUExecContext->CountOffset_RT = CountManager.AcquireOrAllocateEntry(RHICmdList);
+					}
+
+					const FRWBuffer& CountBuffer = CountManager.GetInstanceCountBuffer();
+					const TPair<uint32, int32> DataToSet(GPUExecContext->CountOffset_RT, NumInstances);
+					RHICmdList.Transition(FRHITransitionInfo(CountBuffer.UAV, FNiagaraGPUInstanceCountManager::kCountBufferDefaultState, ERHIAccess::UAVCompute));
+					NiagaraClearCounts::ClearCountsInt(RHICmdList, CountBuffer.UAV, MakeArrayView(&DataToSet, 1));
+					RHICmdList.Transition(FRHITransitionInfo(CountBuffer.UAV, ERHIAccess::UAVCompute, FNiagaraGPUInstanceCountManager::kCountBufferDefaultState));
+				}
+
+				// Copy instance counts
+				FNiagaraDataBuffer& DataBuffer = DataSet->GetCurrentDataChecked();
+				DataBuffer.AllocateGPU(RHICmdList, NumInstances, DispathInterface->GetFeatureLevel(), TEXT("NiagaraSimCache"));
+				DataBuffer.SetNumInstances(NumInstances);
+				DataBuffer.SetGPUDataReadyStage(ENiagaraGpuComputeTickStage::PreInitViews);
+				GPUExecContext->SetDataToRender(&DataBuffer);
+
+				if (CacheBuffer->NumInstances > 0 )
+				{
+					int32 iComponent = 0;
+
+					// Copy Float
+					if ( CacheLayout->FloatCount > 0 )
+					{
+						FRWBuffer& RWBuffer = DataBuffer.GetGPUBufferFloat();
+						const int32 RWComponentStride = DataBuffer.GetFloatStride();
+						uint8* RWBufferMemory = reinterpret_cast<uint8*>(RHILockBuffer(RWBuffer.Buffer, 0, RWBuffer.NumBytes, RLM_WriteOnly));
+
+						for (uint32 i=0; i < CacheLayout->FloatCount; ++i)
+						{
+							const uint32 Component = CacheLayout->ComponentMappingsToDataBuffer[iComponent++];
+							if (Component == InvalidComponent)
+							{
+								continue;
+							}
+							const uint8* Source = CacheBuffer->FloatData.GetData() + (i * NumInstances * sizeof(float));
+							uint8* Dest = RWBufferMemory + (uint32(Component) * RWComponentStride);
+							CheckedMemcpy(MakeArrayView(RWBufferMemory, RWBuffer.NumBytes), Dest, CacheBuffer->FloatData, Source, sizeof(float) * NumInstances);
+							//FMemory::Memcpy(Dest, Source, sizeof(float) * NumInstances);
+						}
+
+						// Copy variables that require processing
+						for (const FNiagaraSimCacheDataBuffersLayout::FVariableCopyInfo& VariableCopyInfo : CacheLayout->VariableMappingsToDataBuffer)
+						{
+							const uint32 SrcStride = uint32(NumInstances) * sizeof(float);
+							const uint8* Src = CacheBuffer->FloatData.GetData() + (VariableCopyInfo.ComponentFrom * SrcStride);
+							uint8* Dst = RWBufferMemory + (uint32(VariableCopyInfo.ComponentTo) * RWComponentStride);
+							VariableCopyInfo.CopyFunc(Dst, DataBuffer.GetFloatStride(), Src, SrcStride, uint32(NumInstances), RebaseTransform);
+						}
+
+						RHIUnlockBuffer(RWBuffer.Buffer);
+					}
+
+					// Copy Half
+					if (CacheLayout->HalfCount > 0)
+					{
+						FRWBuffer& RWBuffer = DataBuffer.GetGPUBufferHalf();
+						const int32 RWComponentStride = DataBuffer.GetHalfStride();
+						uint8* RWBufferMemory = reinterpret_cast<uint8*>(RHILockBuffer(RWBuffer.Buffer, 0, RWBuffer.NumBytes, RLM_WriteOnly));
+
+						for (uint32 i = 0; i < CacheLayout->HalfCount; ++i)
+						{
+							const uint32 Component = CacheLayout->ComponentMappingsToDataBuffer[iComponent++];
+							if (Component == InvalidComponent)
+							{
+								continue;
+							}
+							const uint8* Source = CacheBuffer->HalfData.GetData() + (i * NumInstances * sizeof(FFloat16));
+							uint8* Dest = RWBufferMemory + (uint32(Component) * RWComponentStride);
+							CheckedMemcpy(MakeArrayView(RWBufferMemory, RWBuffer.NumBytes), Dest, CacheBuffer->HalfData, Source, sizeof(FFloat16) * NumInstances);
+							//FMemory::Memcpy(Dest, Source, sizeof(FFloat16) * NumInstances);
+						}
+
+						RHIUnlockBuffer(RWBuffer.Buffer);
+					}
+
+
+					// Copy Int32
+					if (CacheLayout->Int32Count > 0)
+					{
+						FRWBuffer& RWBuffer = DataBuffer.GetGPUBufferInt();
+						const int32 RWComponentStride = DataBuffer.GetInt32Stride();
+						uint8* RWBufferMemory = reinterpret_cast<uint8*>(RHILockBuffer(RWBuffer.Buffer, 0, RWBuffer.NumBytes, RLM_WriteOnly));
+
+						for (uint32 i = 0; i < CacheLayout->Int32Count; ++i)
+						{
+							const uint32 Component = CacheLayout->ComponentMappingsToDataBuffer[iComponent++];
+							if (Component == InvalidComponent)
+							{
+								continue;
+							}
+							const uint8* Source = CacheBuffer->Int32Data.GetData() + (i * NumInstances * sizeof(int32));
+							uint8* Dest = RWBufferMemory + (uint32(Component) * RWComponentStride);
+							CheckedMemcpy(MakeArrayView(RWBufferMemory, RWBuffer.NumBytes), Dest, CacheBuffer->Int32Data, Source, sizeof(int32) * NumInstances);
+							//FMemory::Memcpy(Dest, Source, sizeof(int32) * NumInstances);
+						}
+
+						RHIUnlockBuffer(RWBuffer.Buffer);
+					}
+				}
+
+				//-TODO:DestinationDataBuffer.SetIDTable(CacheBuffer.IDToIndexTable);
+				DataBuffer.SetIDAcquireTag(CacheBuffer->IDAcquireTag);
+
+				// Ensure we decrement our counter so the GameThread knows the state of things
+				--(*PendingCommandsCounter);
+			}
+		);
+	}
+
+	static bool CanRebaseVariable(const FNiagaraVariableBase& Variable)
+	{
+		return	
+			(Variable.GetType() == FNiagaraTypeDefinition::GetQuatDef()) ||
+			(Variable.GetType() == FNiagaraTypeDefinition::GetMatrix4Def()) ||
+			(Variable.GetType() == FNiagaraTypeDefinition::GetPositionDef());
+	}
+
+	static void CopyPositions(uint8* Dst, uint32 DstStride, const uint8* Src, uint32 SrcStride, uint32 NumInstances, const FTransform& RebaseTransform)
+	{
+		float* DstFloats = reinterpret_cast<float*>(Dst);
+		DstStride = DstStride >> 2;
+
+		const float* SrcFloats = reinterpret_cast<const float*>(Src);
+		SrcStride = SrcStride >> 2;
+
+		for (uint32 i = 0; i < NumInstances; ++i)
+		{
+			const FVector CachePosition(
+				SrcFloats[i + (SrcStride * 0)],
+				SrcFloats[i + (SrcStride * 1)],
+				SrcFloats[i + (SrcStride * 2)]
+			);
+			const FVector RebasedPosition = RebaseTransform.TransformPosition(CachePosition);
+			DstFloats[i + (DstStride * 0)] = RebasedPosition.X;
+			DstFloats[i + (DstStride * 1)] = RebasedPosition.Y;
+			DstFloats[i + (DstStride * 2)] = RebasedPosition.Z;
+		}
+	}
+
+	static void CopyQuaternions(uint8* Dst, uint32 DstStride, const uint8* Src, uint32 SrcStride, uint32 NumInstances, const FTransform& RebaseTransform)
+	{
+		float* DstFloats = reinterpret_cast<float*>(Dst);
+		DstStride = DstStride >> 2;
+
+		const float* SrcFloats = reinterpret_cast<const float*>(Src);
+		SrcStride = SrcStride >> 2;
+
+		for (uint32 i = 0; i < NumInstances; ++i)
+		{
+			const FQuat4f CacheRotation(
+				SrcFloats[i + (SrcStride * 0)],
+				SrcFloats[i + (SrcStride * 1)],
+				SrcFloats[i + (SrcStride * 2)],
+				SrcFloats[i + (SrcStride * 3)]
+			);
+			const FQuat4f RebasedQuat = CacheRotation * FQuat4f(RebaseTransform.GetRotation());
+			DstFloats[i + (DstStride * 0)] = RebasedQuat.X;
+			DstFloats[i + (DstStride * 1)] = RebasedQuat.Y;
+			DstFloats[i + (DstStride * 2)] = RebasedQuat.Z;
+			DstFloats[i + (DstStride * 3)] = RebasedQuat.W;
+		}
+	}
+
+	static void CopyMatrices(uint8* Dst, uint32 DstStride, const uint8* Src, uint32 SrcStride, uint32 NumInstances, const FTransform& RebaseTransform)
+	{
+		float* DstFloats = reinterpret_cast<float*>(Dst);
+		DstStride = DstStride >> 2;
+
+		const float* SrcFloats = reinterpret_cast<const float*>(Src);
+		SrcStride = SrcStride >> 2;
+
+		const FMatrix44d RebaseMatrix = RebaseTransform.ToMatrixWithScale();
+		for (uint32 i = 0; i < NumInstances; ++i)
+		{
+			FMatrix44d CacheMatrix;
+			for (int32 j = 0; j < 16; ++j)
+			{
+				CacheMatrix.M[j >> 2][j & 0x3] = double(SrcFloats[i + (SrcStride * j)]);
+			}
+
+			CacheMatrix = CacheMatrix * RebaseMatrix;
+
+			for (int32 j = 0; j < 16; ++j)
+			{
+				DstFloats[i + (DstStride * j)] = float(CacheMatrix.M[j >> 2][j & 0x3]);
+			}
+		}
+	}
+
+	UNiagaraSystem*						NiagaraSystem = nullptr;
+	FNiagaraSystemInstance*				SystemInstance = nullptr;
+	FNiagaraSystemSimulationPtr			SystemSimulation = nullptr;
+	FNiagaraDataBuffer*					SystemSimulationDataBuffer = nullptr;
+
+	static constexpr uint16				InvalidComponent = INDEX_NONE;
+};
diff --git a/Engine/Plugins/FX/Niagara/Source/Niagara/Private/NiagaraSystemInstance.cpp b/Engine/Plugins/FX/Niagara/Source/Niagara/Private/NiagaraSystemInstance.cpp
--- a/Engine/Plugins/FX/Niagara/Source/Niagara/Private/NiagaraSystemInstance.cpp
+++ b/Engine/Plugins/FX/Niagara/Source/Niagara/Private/NiagaraSystemInstance.cpp
@@ -2174,7 +2174,7 @@ void FNiagaraSystemInstance::SimCacheTick_GameThread(UNiagaraSimCache* SimCache,
 		return;
 	}
 
-	if ( DesiredAge >= SimCache->StartSeconds + SimCache->DurationSeconds )
+	if ( DesiredAge >= SimCache->GetStartSeconds() + SimCache->GetDurationSeconds() )
 	{
 		Complete(false);
 		return;
diff --git a/Engine/Plugins/FX/Niagara/Source/NiagaraEditor/Private/ViewModels/NiagaraSimCacheViewModel.cpp b/Engine/Plugins/FX/Niagara/Source/NiagaraEditor/Private/ViewModels/NiagaraSimCacheViewModel.cpp
--- a/Engine/Plugins/FX/Niagara/Source/NiagaraEditor/Private/ViewModels/NiagaraSimCacheViewModel.cpp
+++ b/Engine/Plugins/FX/Niagara/Source/NiagaraEditor/Private/ViewModels/NiagaraSimCacheViewModel.cpp
@@ -269,7 +269,7 @@ int32 FNiagaraSimCacheViewModel::GetNumInstances() const
 int32 FNiagaraSimCacheViewModel::GetNumFrames() const
 {
 	const UNiagaraSimCache* SimCache = BufferReader ? BufferReader->GetSimCache() : nullptr;
-	return SimCache ? SimCache->CacheFrames.Num() : 0;
+	return SimCache ? SimCache->GetNumFrames() : 0;
 }
 
 void FNiagaraSimCacheViewModel::SetFrameIndex(const int32 InFrameIndex)
@@ -298,13 +298,13 @@ bool FNiagaraSimCacheViewModel::IsCacheValid()
 int32 FNiagaraSimCacheViewModel::GetNumEmitterLayouts()
 {
 	const UNiagaraSimCache* SimCache = BufferReader ? BufferReader->GetSimCache() : nullptr;
-	return SimCache ? SimCache->CacheLayout.EmitterLayouts.Num() : 0;
+	return SimCache ? SimCache->GetNumEmitters() : 0;
 }
 
 FName FNiagaraSimCacheViewModel::GetEmitterLayoutName(const int32 Index)
 {
 	const UNiagaraSimCache* SimCache = BufferReader ? BufferReader->GetSimCache() : nullptr;
-	return SimCache ? SimCache->CacheLayout.EmitterLayouts[Index].LayoutName : NAME_None;
+	return SimCache ? SimCache->GetEmitterName(Index) : NAME_None;
 }
 
 FNiagaraSimCacheViewModel::FOnFrameUpdated& FNiagaraSimCacheViewModel::OnFrameUpdated()
