- Fix for wrapped dispatch count being higher than required
- Convert to use RHI dispatch information rather than hard coded one

#rnx

diff --git a/Engine/Plugins/FX/Niagara/Source/Niagara/Private/NiagaraGpuComputeDispatch.cpp b/Engine/Plugins/FX/Niagara/Source/Niagara/Private/NiagaraGpuComputeDispatch.cpp
--- a/Engine/Plugins/FX/Niagara/Source/Niagara/Private/NiagaraGpuComputeDispatch.cpp
+++ b/Engine/Plugins/FX/Niagara/Source/Niagara/Private/NiagaraGpuComputeDispatch.cpp
@@ -1540,9 +1540,13 @@ void FNiagaraGpuComputeDispatch::DispatchStage(FRDGBuilder& GraphBuilder, const
 		// In the OneD case we can use the Y dimension to get higher particle counts
 		if (DispatchType == ENiagaraGpuDispatchType::OneD)
 		{
-			const int32 TotalThreadGroups = FMath::DivideAndRoundUp(DispatchCount.X, DispatchNumThreads.X);
-			DispatchCount.Y = FMath::DivideAndRoundUp<int32>(TotalThreadGroups, NiagaraMaxThreadGroupCountPerDimension);
-			DispatchCount.X = FMath::DivideAndRoundUp(DispatchCount.X, DispatchNumThreads.Y);
+			const int32 GroupCount = FMath::DivideAndRoundUp(DispatchCount.X, DispatchNumThreads.X);
+			if (GroupCount > GRHIMaxDispatchThreadGroupsPerDimension.X)
+			{
+				DispatchCount.Y = FMath::DivideAndRoundUp(GroupCount, GRHIMaxDispatchThreadGroupsPerDimension.X);
+				DispatchCount.X = FMath::DivideAndRoundUp(GroupCount, DispatchCount.Y) * DispatchNumThreads.X;
+				ensure(DispatchCount.Y <= GRHIMaxDispatchThreadGroupsPerDimension.Y);
+			}
 		}
 
 		FIntVector ThreadGroupCount;
diff --git a/Engine/Plugins/FX/Niagara/Source/Niagara/Public/NiagaraCommon.h b/Engine/Plugins/FX/Niagara/Source/Niagara/Public/NiagaraCommon.h
--- a/Engine/Plugins/FX/Niagara/Source/Niagara/Public/NiagaraCommon.h
+++ b/Engine/Plugins/FX/Niagara/Source/Niagara/Public/NiagaraCommon.h
@@ -50,9 +50,6 @@ struct FNiagaraParameterStore;
 /** Defines The maximum ThreadGroup size we allow in Niagara.  This is important for how memory is allocated as we always need to round this and the final instance is used to avoid overflowing the buffer. */
 constexpr uint32 NiagaraComputeMaxThreadGroupSize = 64;
 
-/** Defines The maximum Thread Group Count Per Dimension. */
-constexpr uint32 NiagaraMaxThreadGroupCountPerDimension = 65535;
-
 /** The maximum number of spawn infos we can run on the GPU, modifying this will require a version update as it is used in the shader compiler  */
 constexpr uint32 NIAGARA_MAX_GPU_SPAWN_INFOS = 8;
 
