- SimCache write functions converted to be const as the DI should not be modified in the process

#rnx

diff --git a/Engine/Plugins/FX/Niagara/Source/Niagara/Classes/NiagaraDataInterface.h b/Engine/Plugins/FX/Niagara/Source/Niagara/Classes/NiagaraDataInterface.h
--- a/Engine/Plugins/FX/Niagara/Source/Niagara/Classes/NiagaraDataInterface.h
+++ b/Engine/Plugins/FX/Niagara/Source/Niagara/Classes/NiagaraDataInterface.h
@@ -500,11 +500,11 @@ public:
 	virtual bool PerInstanceTickPostSimulate(void* PerInstanceData, FNiagaraSystemInstance* SystemInstance, float DeltaSeconds) { return false; }
 	
 	/** Begin writing data for a simulation cache, returning a nullptr means the data interface does not store data into the simulation cache. */
-	virtual UObject* SimCacheBeginWrite(UObject* SimCache, FNiagaraSystemInstance* NiagaraSystemInstance, const void* OptionalPerInstanceData) { return nullptr; }
+	virtual UObject* SimCacheBeginWrite(UObject* SimCache, FNiagaraSystemInstance* NiagaraSystemInstance, const void* OptionalPerInstanceData) const { return nullptr; }
 	/** Write a new frame of data for the simulation cache.  This is always in sequence, i.e. 0, 1, 2, etc, we will never jump around frames. */
-	virtual bool SimCacheWriteFrame(UObject* StorageObject, int FrameIndex, FNiagaraSystemInstance* SystemInstance, const void* OptionalPerInstanceData) { return true; }
+	virtual bool SimCacheWriteFrame(UObject* StorageObject, int FrameIndex, FNiagaraSystemInstance* SystemInstance, const void* OptionalPerInstanceData) const { return true; }
 	/** End writing data for a simulation cache.  Note this is called on the CDO not the instance the object was created from. */
-	virtual bool SimCacheEndWrite(UObject* StorageObject) { return true; }
+	virtual bool SimCacheEndWrite(UObject* StorageObject) const { return true; }
 	/** Read a frame of data from the simulation cache. */
 	virtual bool SimCacheReadFrame(UObject* StorageObject, int FrameA, int FrameB, float Interp, FNiagaraSystemInstance* SystemInstance, void* OptionalPerInstanceData) { return true; }
 	/**
