Hierarchy Editor & User Param Revamp Part 1. This adds an "Edit Hierarchy" button to the user parameters in a system and lets the user rearrange order & add tooltips, categories, sections. These are mirrored in the details panel.
Part 2 will fix various smaller issues and have a separate User Parameters tab which will include the "Edit Hierarchy" button instead of the Parameters tab

diff --git a/Engine/Plugins/FX/Niagara/Source/Niagara/Public/NiagaraUserRedirectionParameterStore.h b/Engine/Plugins/FX/Niagara/Source/Niagara/Public/NiagaraUserRedirectionParameterStore.h
--- a/Engine/Plugins/FX/Niagara/Source/Niagara/Public/NiagaraUserRedirectionParameterStore.h
+++ b/Engine/Plugins/FX/Niagara/Source/Niagara/Public/NiagaraUserRedirectionParameterStore.h
@@ -28,7 +28,7 @@ public:
 	/** If necessary it will replace the supplied variable with the fully qualified namespace (User.) appropriate for a user variable.
 		Will return false if the variable wasn't able to be converted into a valid user namespaced variable.
 	*/
-	bool RedirectUserVariable(FNiagaraVariableBase& UserVar) const;
+	NIAGARA_API bool RedirectUserVariable(FNiagaraVariableBase& UserVar) const;
 
 	/** Get the list of FNiagaraVariables that are exposed to the user. Note that the values will be stale and are not to be trusted directly. Get the Values using the offset specified by IndexOf or GetParameterValue.*/
 	FORCEINLINE void GetUserParameters(TArray<FNiagaraVariable>& OutParameters) const { return UserParameterRedirects.GenerateKeyArray(OutParameters); }
diff --git a/Engine/Plugins/FX/Niagara/Source/NiagaraEditor/Private/Customizations/NiagaraComponentDetails.cpp b/Engine/Plugins/FX/Niagara/Source/NiagaraEditor/Private/Customizations/NiagaraComponentDetails.cpp
--- a/Engine/Plugins/FX/Niagara/Source/NiagaraEditor/Private/Customizations/NiagaraComponentDetails.cpp
+++ b/Engine/Plugins/FX/Niagara/Source/NiagaraEditor/Private/Customizations/NiagaraComponentDetails.cpp
@@ -39,17 +39,24 @@
 #include "GameDelegates.h"
 #include "NiagaraEditorStyle.h"
 #include "IDetailChildrenBuilder.h"
+#include "IDetailGroup.h"
 #include "Framework/MultiBox/MultiBoxBuilder.h"
 #include "Widgets/Input/SComboButton.h"
 #include "Widgets/Images/SImage.h"
 #include "NiagaraEditorModule.h"
 #include "Modules/ModuleManager.h"
 #include "INiagaraEditorTypeUtilities.h"
+#include "NiagaraConstants.h"
+#include "NiagaraScriptVariable.h"
+#include "NiagaraSystemEditorData.h"
 #include "Widgets/Layout/SBox.h"
 #include "UObject/WeakObjectPtr.h"
 #include "NiagaraUserRedirectionParameterStore.h"
 #include "Widgets/Layout/SUniformGridPanel.h"
 #include "Widgets/SNiagaraDebugger.h"
+#include "Widgets/SNiagaraParameterName.h"
+#include "ViewModels/HierarchyEditor/NiagaraUserParametersHierarchyViewModel.h"
+#include "Widgets/Input/SCheckBox.h"
 
 #define LOCTEXT_NAMESPACE "NiagaraComponentDetails"
 
@@ -180,8 +187,8 @@ public:
 			}
 		}
 	}
-
-	void OnParameterChanged()
+	
+	void OnParameterChanged(const FPropertyChangedEvent& PropertyChangedEvent)
 	{
 		if (bResettingToDefault)
 		{
@@ -195,7 +202,7 @@ public:
 
 			for (TSharedPtr<IPropertyHandle> PropertyHandle : PropertyHandles)
 			{
-				PropertyHandle->NotifyPostChange(EPropertyChangeType::ValueSet);
+				PropertyHandle->NotifyPostChange(PropertyChangedEvent.ChangeType);
 			}
 		}
 	}
@@ -251,6 +258,7 @@ public:
 	const FNiagaraVariableBase& Key() const { return ParameterKey; }
 	FNiagaraVariant& Value() { return ParameterValue; }
 
+	TWeakObjectPtr<UNiagaraComponent> GetComponent() const { return Component; }
 
 private:
 	TWeakObjectPtr<UNiagaraComponent> Component;
@@ -261,6 +269,254 @@ private:
 	bool bResettingToDefault;
 };
 
+void AddRowForUserParameter(TSharedPtr<FNiagaraParameterProxy> ParameterProxy, IDetailChildrenBuilder& ChildrenBuilder, TMap<FName, TWeakPtr<FStructOnScope>>& ParameterNameToDisplayStruct)
+{
+	FNiagaraVariable Parameter = ParameterProxy->Key();
+	IDetailPropertyRow* Row = nullptr;
+
+ 	TSharedPtr<SWidget> CustomValueWidget;
+			
+	if (Parameter.IsDataInterface())
+	{
+		// if no changes are made, then it'll just be the same as the asset
+		TArray<UObject*> Objects { ParameterProxy->Value().GetDataInterface() };
+	
+		FAddPropertyParams Params = FAddPropertyParams()
+			.UniqueId(Parameter.GetName())
+			.AllowChildren(true)
+			.CreateCategoryNodes(true);
+	
+		Row = ChildrenBuilder.AddExternalObjectProperty(Objects, NAME_None, Params);
+
+		CustomValueWidget =
+			SNew(STextBlock)
+			.TextStyle(FNiagaraEditorStyle::Get(), "NiagaraEditor.ParameterText")
+			.Text(FText::FromString(FName::NameToDisplayString(Parameter.GetType().GetClass()->GetName(), false)));
+	}
+	else if(Parameter.IsUObject())
+	{
+		TArray<UObject*> Objects { ParameterProxy->Value().GetUObject() };
+
+		// How do I set this up so I can have this pick actors from the level?
+
+		FAddPropertyParams Params = FAddPropertyParams()
+			.UniqueId(Parameter.GetName())
+			.AllowChildren(false) // Don't show the material's properties
+			.CreateCategoryNodes(false);
+
+		Row = ChildrenBuilder.AddExternalObjectProperty(Objects, NAME_None, Params);
+
+		CustomValueWidget = SNew(SObjectPropertyEntryBox)
+			.ObjectPath(ParameterProxy.ToSharedRef(), &FNiagaraParameterProxy::GetCurrentAssetPath)
+			.AllowedClass(Parameter.GetType().GetClass())
+			.OnObjectChanged(ParameterProxy.ToSharedRef(), &FNiagaraParameterProxy::OnAssetSelectedFromPicker)
+			.AllowClear(false)
+			.DisplayUseSelected(true)
+			.DisplayBrowse(true)
+			.DisplayThumbnail(true)
+			.NewAssetFactories(TArray<UFactory*>());
+	}
+	else
+	{
+		FNiagaraTypeDefinition Type = Parameter.GetType();
+		UStruct* Struct = Type.GetStruct();
+			
+		if (Type == FNiagaraTypeDefinition::GetPositionDef())
+		{
+			static UPackage* CoreUObjectPkg = FindObjectChecked<UPackage>(nullptr, TEXT("/Script/CoreUObject"));
+			static UScriptStruct* VectorStruct = FindObjectChecked<UScriptStruct>(CoreUObjectPkg, TEXT("Vector"));
+			Struct = VectorStruct;
+		}
+			
+		TSharedPtr<FStructOnScope> StructOnScope = MakeShareable(new FStructOnScope(Struct, ParameterProxy->Value().GetBytes()));
+
+		FAddPropertyParams Params = FAddPropertyParams()
+			.UniqueId(Parameter.GetName());
+
+		Row = ChildrenBuilder.AddExternalStructureProperty(StructOnScope.ToSharedRef(), NAME_None, Params);
+		FNiagaraVariable UserParameter = Parameter;
+		FNiagaraUserRedirectionParameterStore::MakeUserVariable(UserParameter);
+		ParameterNameToDisplayStruct.Add(UserParameter.GetName(), TWeakPtr<FStructOnScope>(StructOnScope));
+	}
+
+	TSharedPtr<SWidget> DefaultNameWidget;
+	TSharedPtr<SWidget> DefaultValueWidget;
+	
+	if(Parameter.IsInNameSpace(FNiagaraConstants::UserNamespaceString))
+	{
+		Parameter.SetName(FName(Parameter.GetName().ToString().RightChop(5)));
+	}
+	
+	Row->DisplayName(FText::FromName(Parameter.GetName()));
+
+	FDetailWidgetRow& CustomWidget = Row->CustomWidget(true);
+
+	FNiagaraVariable UserParameter = Parameter;
+	ParameterProxy->GetComponent()->GetAsset()->GetExposedParameters().RedirectUserVariable(UserParameter);
+	UNiagaraScriptVariable* ScriptVariable = FNiagaraEditorUtilities::GetScriptVariableForUserParameter(UserParameter, *ParameterProxy->GetComponent()->GetAsset());
+		
+	FText Tooltip = FText::GetEmpty();
+	if(ScriptVariable && !ScriptVariable->Metadata.Description.IsEmpty())
+	{
+		Tooltip = ScriptVariable->Metadata.Description;
+	}
+		
+	Row->GetDefaultWidgets(DefaultNameWidget, DefaultValueWidget, CustomWidget);
+
+	Row->GetPropertyHandle()->SetOnPropertyValuePreChange(FSimpleDelegate::CreateSP(ParameterProxy.ToSharedRef(), &FNiagaraParameterProxy::OnParameterPreChange));
+	Row->GetPropertyHandle()->SetOnChildPropertyValuePreChange(FSimpleDelegate::CreateSP(ParameterProxy.ToSharedRef(), &FNiagaraParameterProxy::OnParameterPreChange));
+	Row->GetPropertyHandle()->SetOnPropertyValueChangedWithData(TDelegate<void(const FPropertyChangedEvent&)>::CreateSP(ParameterProxy.ToSharedRef(), &FNiagaraParameterProxy::OnParameterChanged));
+	Row->GetPropertyHandle()->SetOnChildPropertyValueChangedWithData(TDelegate<void(const FPropertyChangedEvent&)>::CreateSP(ParameterProxy.ToSharedRef(), &FNiagaraParameterProxy::OnParameterChanged));
+	Row->GetPropertyHandle()->SetOnPropertyResetToDefault(FSimpleDelegate::CreateSP(ParameterProxy.ToSharedRef(), &FNiagaraParameterProxy::OnResetToDefault));
+	
+	CustomWidget
+		.NameContent()
+		[
+			SNew(SBox)
+			.Padding(FMargin(0.0f, 2.0f))
+			[
+				SNew(STextBlock)
+				.TextStyle(FNiagaraEditorStyle::Get(), "NiagaraEditor.ParameterText")
+				.Text(FText::FromName(Parameter.GetName()))
+				.ToolTipText(Tooltip)
+			]
+		];
+
+	TSharedPtr<SWidget> ValueWidget = DefaultValueWidget;
+	if (CustomValueWidget.IsValid())
+	{
+		ValueWidget = CustomValueWidget;
+	}
+	
+	CustomWidget
+	.ValueContent()
+	[
+		SNew(SHorizontalBox)
+		+ SHorizontalBox::Slot()
+		.HAlign(HAlign_Fill)
+		.Padding(4.0f)
+		[
+			// Add in the parameter editor factoried above.
+			ValueWidget.ToSharedRef()
+		]
+		+ SHorizontalBox::Slot()
+		.VAlign(VAlign_Center)
+		.AutoWidth()
+		[
+			// Add in the "reset to default" buttons
+			SNew(SButton)
+			.OnClicked(ParameterProxy.ToSharedRef(), &FNiagaraParameterProxy::OnResetToDefaultClicked)
+			.Visibility(ParameterProxy.ToSharedRef(), &FNiagaraParameterProxy::GetResetToDefaultVisibility)
+			.ContentPadding(FMargin(5.f, 0.f))
+			.ToolTipText(LOCTEXT("ResetToDefaultToolTip", "Reset to Default"))
+			.ButtonStyle(FAppStyle::Get(), "NoBorder")
+			.Content()
+			[
+				SNew(SImage)
+				.Image(FAppStyle::GetBrush("PropertyWindow.DiffersFromDefault"))
+			]
+		]
+	];
+}
+
+class FNiagaraUserParameterCategoryBuilder : public IDetailCustomNodeBuilder, public TSharedFromThis<FNiagaraUserParameterCategoryBuilder>
+{
+public:
+	FNiagaraUserParameterCategoryBuilder(UNiagaraComponent* InComponent, const UNiagaraHierarchyCategory* InNiagaraHierarchyCategory, FName InCustomBuilderRowName,
+		TArray<TSharedPtr<FNiagaraParameterProxy>>& InParameterProxies, TMap<FName, TWeakPtr<FStructOnScope>>& InParameterNameToDisplayStruct,
+		TArray<TSharedPtr<IPropertyHandle>> InOverridePropertyHandles) : ParameterProxies(InParameterProxies), ParameterNameToDisplayStruct(InParameterNameToDisplayStruct)
+	{
+		HierarchyCategory = InNiagaraHierarchyCategory;
+		OverridePropertyHandles = InOverridePropertyHandles;
+		CustomBuilderRowName = InCustomBuilderRowName;
+
+		Component = InComponent;
+		bDelegatesInitialized = false;
+	}
+	
+	virtual void GenerateHeaderRowContent(FDetailWidgetRow& NodeRow) override
+	{
+		NodeRow.NameContent()
+		[
+			SNew(STextBlock)
+			.TextStyle(FNiagaraEditorStyle::Get(), "NiagaraEditor.ParameterText")
+			.Text(FText::FromName(CustomBuilderRowName))
+			.ToolTipText(HierarchyCategory->GetTooltip())
+		];
+
+		NodeRow.FilterString(FText::FromName(CustomBuilderRowName));
+	}
+	virtual void Tick(float DeltaTime) override {}
+	virtual bool RequiresTick() const override { return false; }
+	virtual bool InitiallyCollapsed() const override { return false; }
+
+	virtual FName GetName() const  override
+	{
+		return CustomBuilderRowName;
+	}
+	
+	virtual void GenerateChildContent(IDetailChildrenBuilder& ChildrenBuilder) override
+	{
+		TArray<const UNiagaraHierarchyUserParameter*> ContainedParameters;
+		TArray<const UNiagaraHierarchyCategory*> ContainedCategories;
+
+		for(const UNiagaraHierarchyItemBase* Child : HierarchyCategory->GetChildren())
+		{
+			if(const UNiagaraHierarchyUserParameter* UserParameter = Cast<UNiagaraHierarchyUserParameter>(Child))
+			{
+				ContainedParameters.Add(UserParameter);
+			}
+		}
+
+		for(UNiagaraHierarchyItemBase* Child : HierarchyCategory->GetChildren())
+		{
+			if(UNiagaraHierarchyCategory* Category = Cast<UNiagaraHierarchyCategory>(Child))
+			{
+				ContainedCategories.Add(Category);
+			}
+		}
+
+		for(const UNiagaraHierarchyCategory* Category : ContainedCategories)
+		{
+			if(Category->DoesOneChildExist<UNiagaraHierarchyUserParameter>())
+			{
+				ChildrenBuilder.AddCustomBuilder(MakeShared<FNiagaraUserParameterCategoryBuilder>(
+					Component.Get(),Category, Category->GetCategoryName(),
+					ParameterProxies, ParameterNameToDisplayStruct,
+					OverridePropertyHandles));
+			}
+		}
+		
+		for(const UNiagaraHierarchyUserParameter* HierarchyUserParameter : ContainedParameters)
+		{
+			FNiagaraVariable Parameter = HierarchyUserParameter->GetUserParameter();
+			FNiagaraVariant ParameterValue = GetCurrentParameterValue(Parameter, Component.Get());
+			if (!ParameterValue.IsValid())
+			{
+				continue;
+			}
+		
+			if (Parameter.IsDataInterface())
+			{
+				ParameterValue = FNiagaraVariant(DuplicateObject(ParameterValue.GetDataInterface(), Component.Get()));
+			}
+			
+			TSharedPtr<FNiagaraParameterProxy> ParameterProxy = ParameterProxies.Add_GetRef(MakeShareable(new FNiagaraParameterProxy(Component, Parameter, ParameterValue, OnRebuildChildren, OverridePropertyHandles)));
+			AddRowForUserParameter(ParameterProxy, ChildrenBuilder, ParameterNameToDisplayStruct);
+		}
+	}
+
+private:
+	bool bDelegatesInitialized;
+	TWeakObjectPtr<UNiagaraComponent> Component;
+	TWeakObjectPtr<const UNiagaraHierarchyCategory> HierarchyCategory;
+	TArray<TSharedPtr<IPropertyHandle>> OverridePropertyHandles;
+	TArray<TSharedPtr<FNiagaraParameterProxy>>& ParameterProxies;
+	TMap<FName, TWeakPtr<FStructOnScope>>& ParameterNameToDisplayStruct;
+	FSimpleDelegate OnRebuildChildren;
+	FName CustomBuilderRowName;
+};
+
 class FNiagaraComponentNodeBuilder : public IDetailCustomNodeBuilder, public TSharedFromThis<FNiagaraComponentNodeBuilder>
 {
 public:
@@ -291,7 +547,7 @@ public:
 		OnRebuildChildren = InOnRegenerateChildren;
 	}
 
-	virtual void GenerateHeaderRowContent(FDetailWidgetRow& NodeRow) {}
+	virtual void GenerateHeaderRowContent(FDetailWidgetRow& NodeRow) override;
 	virtual void Tick(float DeltaTime) override {}
 	virtual bool RequiresTick() const override { return false; }
 	virtual bool InitiallyCollapsed() const { return false; }
@@ -314,171 +570,164 @@ public:
 
 		ParameterProxies.Reset();
 
-		FNiagaraEditorModule& NiagaraEditorModule = FModuleManager::GetModuleChecked<FNiagaraEditorModule>("NiagaraEditor");
-
 		UNiagaraSystem* SystemAsset = Component->GetAsset();
 		if (SystemAsset == nullptr)
 		{
 			return;
 		}
-
+		
 		TArray<FNiagaraVariable> UserParameters;
 		SystemAsset->GetExposedParameters().GetUserParameters(UserParameters);
-
+		
 		ParameterProxies.Reserve(UserParameters.Num());
 
 		ParameterNameToDisplayStruct.Empty();
-		for (const FNiagaraVariable& Parameter : UserParameters)
-		{
-			TSharedPtr<SWidget> NameWidget;
+		TMap<const UNiagaraHierarchySection*, TArray<const UNiagaraHierarchyCategory*>> SectionCategoryMap;
 
-			NameWidget =
-				SNew(STextBlock)
-				.TextStyle(FNiagaraEditorStyle::Get(), "NiagaraEditor.ParameterText")
-				.Text(FText::FromName(Parameter.GetName()));
-
-			IDetailPropertyRow* Row = nullptr;
+		UNiagaraHierarchyRoot* Root = Cast<UNiagaraSystemEditorData>(SystemAsset->GetEditorData())->UserParameterHierarchy;
+		for(const UNiagaraHierarchyItemBase* Child : Root->GetChildren())
+		{
+			if(const UNiagaraHierarchyCategory* HierarchyCategory = Cast<UNiagaraHierarchyCategory>(Child))
+			{
+				if(HierarchyCategory->GetSection() != nullptr)
+				{
+					SectionCategoryMap.FindOrAdd(HierarchyCategory->GetSection()).Add(HierarchyCategory);	
+				}
+			}
+		}
 
-			TSharedPtr<SWidget> CustomValueWidget;
+		// we create a custom row for user param sections here, but only if we have at least one section specified
+		if(SectionCategoryMap.Num() > 0)
+		{
+			TSharedRef<SHorizontalBox> SectionsBox = SNew(SHorizontalBox);
 			
-			FNiagaraVariant ParameterValue = GetCurrentParameterValue(Parameter, Component.Get());
-			if (!ParameterValue.IsValid())
+			for(auto& SectionCategories : SectionCategoryMap)
 			{
-				continue;
+				const UNiagaraHierarchySection* Section = SectionCategories.Key;
+				TArray<const UNiagaraHierarchyCategory*>& HierarchyCategories = SectionCategories.Value;
+				
+				if(HierarchyCategories.Num() > 0)
+				{
+					bool bDoesUserParamExist = false;
+					for(const UNiagaraHierarchyCategory* Category : HierarchyCategories)
+					{
+						bDoesUserParamExist |= Category->DoesOneChildExist<UNiagaraHierarchyUserParameter>(true);
+
+						if(bDoesUserParamExist)
+						{
+							break;
+						}
+					}
+
+					if(bDoesUserParamExist)
+					{
+						SectionsBox->AddSlot()
+						.AutoWidth()
+						.Padding(FMargin(2.f))
+						[
+							SNew(SBox)
+							.Padding(FMargin(0.f, 4.f, 0.f, 0.f))
+							[
+								SNew(SCheckBox)
+								.Style(FAppStyle::Get(), "DetailsView.SectionButton")
+								.OnCheckStateChanged(this, &FNiagaraComponentNodeBuilder::OnActiveSectionChanged, Section)
+								.IsChecked(this, &FNiagaraComponentNodeBuilder::IsSectionActive, Section)
+								[
+									SNew(STextBlock)
+									.TextStyle(FAppStyle::Get(), "SmallText")
+									.Text(Section->GetSectionNameAsText())
+									.ToolTipText(Section->GetTooltip())
+								]
+							]
+						];
+					}
+				}
 			}
 
-			if (Parameter.IsDataInterface())
+			// if we have at least one custom section, we add a default "All" section
+			if(SectionsBox->GetChildren()->Num() > 0)
 			{
-				ParameterValue = FNiagaraVariant(DuplicateObject(ParameterValue.GetDataInterface(), Component.Get()));
+				SectionsBox->AddSlot()
+				.AutoWidth()
+				.Padding(FMargin(2.f))
+				[
+					SNew(SBox)
+					.Padding(FMargin(0.f, 4.f, 0.f, 0.f))
+					[
+						SNew(SCheckBox)
+						.Style(FAppStyle::Get(), "DetailsView.SectionButton")
+						.OnCheckStateChanged(this, &FNiagaraComponentNodeBuilder::OnActiveSectionChanged, (const UNiagaraHierarchySection*) nullptr)
+						.IsChecked(this, &FNiagaraComponentNodeBuilder::IsSectionActive, (const UNiagaraHierarchySection*) nullptr)
+						[
+							SNew(STextBlock)
+							.TextStyle(FAppStyle::Get(), "SmallText")
+							.Text(FText::FromString("All"))
+						]
+					]
+				];
+				
+				FDetailWidgetRow& Row = ChildrenBuilder.AddCustomRow(FText::FromString("Parameter Sections"));
+				Row.WholeRowContent()
+				[
+					SectionsBox
+				];
 			}
+		}
 
-			TSharedPtr<FNiagaraParameterProxy> ParameterProxy = ParameterProxies.Add_GetRef(MakeShareable(new FNiagaraParameterProxy(Component, Parameter, ParameterValue, OnRebuildChildren, OverridePropertyHandles)));
-				
-			if (Parameter.IsDataInterface())
+		// we add the categories now if the active section is set to all (nullptr) or if the section it belongs to is active
+		for(const UNiagaraHierarchyItemBase* Child : Root->GetChildren())
+		{
+			if(const UNiagaraHierarchyCategory* HierarchyCategory = Cast<UNiagaraHierarchyCategory>(Child))
 			{
-				// duplicate the DI here so that if it is changed, the component's SetParameterOverride will override the value
-				// if no changes are made, then it'll just be the same as the asset
-				TArray<UObject*> Objects { ParameterProxy->Value().GetDataInterface() };
+				if(HierarchyCategory->DoesOneChildExist<UNiagaraHierarchyUserParameter>(true) && (ActiveSection == nullptr || HierarchyCategory->GetSection() == ActiveSection))
+				{
+					ChildrenBuilder.AddCustomBuilder(MakeShared<FNiagaraUserParameterCategoryBuilder>(
+						Component.Get(), HierarchyCategory, HierarchyCategory->GetCategoryName(),
+						ParameterProxies, ParameterNameToDisplayStruct,
+						OverridePropertyHandles));
+				}
+			}
+		}
 
-				FAddPropertyParams Params = FAddPropertyParams()
-					.UniqueId(Parameter.GetName())
-					.AllowChildren(true)
-					.CreateCategoryNodes(false);
+		// at last we add all parameters that haven't been setup at all, if the active section is set to "All"
+		if(ActiveSection == nullptr)
+		{
+			TArray<FNiagaraVariable> LeftoverUserParameters;
+			SystemAsset->GetExposedParameters().GetUserParameters(LeftoverUserParameters);
 
-				Row = ChildrenBuilder.AddExternalObjectProperty(Objects, NAME_None, Params); 
+			TArray<UNiagaraHierarchyUserParameter*> HierarchyUserParameters;
+			Root->GetChildrenOfType<>(HierarchyUserParameters, true);
 
-				CustomValueWidget =
-					SNew(STextBlock)
-					.TextStyle(FNiagaraEditorStyle::Get(), "NiagaraEditor.ParameterText")
-					.Text(FText::FromString(FName::NameToDisplayString(Parameter.GetType().GetClass()->GetName(), false)));
-			}
-			else if (Parameter.IsUObject())
+			TArray<FNiagaraVariable> RedirectedUserParameters;
+			for(FNiagaraVariable& Parameter : LeftoverUserParameters)
 			{
-				TArray<UObject*> Objects { ParameterProxy->Value().GetUObject() };
-
-				// How do I set this up so I can have this pick actors from the level?
-
-				FAddPropertyParams Params = FAddPropertyParams()
-					.UniqueId(Parameter.GetName())
-					.AllowChildren(false) // Don't show the material's properties
-					.CreateCategoryNodes(false);
-
-				Row = ChildrenBuilder.AddExternalObjectProperty(Objects, NAME_None, Params);
-
-				CustomValueWidget = SNew(SObjectPropertyEntryBox)
-					.ObjectPath(ParameterProxy.ToSharedRef(), &FNiagaraParameterProxy::GetCurrentAssetPath)
-					.AllowedClass(Parameter.GetType().GetClass())
-					.OnObjectChanged(ParameterProxy.ToSharedRef(), &FNiagaraParameterProxy::OnAssetSelectedFromPicker)
-					.AllowClear(false)
-					.DisplayUseSelected(true)
-					.DisplayBrowse(true)
-					.DisplayThumbnail(true)
-					.NewAssetFactories(TArray<UFactory*>());
+				SystemAsset->GetExposedParameters().RedirectUserVariable(Parameter);
 			}
-			else
+			
+			for(UNiagaraHierarchyUserParameter* HierarchyUserParameter : HierarchyUserParameters)
 			{
-				FNiagaraTypeDefinition Type = Parameter.GetType();
-				UStruct* Struct = Type.GetStruct();
-				if (Type == FNiagaraTypeDefinition::GetPositionDef())
+				if(LeftoverUserParameters.Contains(HierarchyUserParameter->GetUserParameter()))
 				{
-					static UPackage* CoreUObjectPkg = FindObjectChecked<UPackage>(nullptr, TEXT("/Script/CoreUObject"));
-					static UScriptStruct* VectorStruct = FindObjectChecked<UScriptStruct>(CoreUObjectPkg, TEXT("Vector"));
-					Struct = VectorStruct;
+					LeftoverUserParameters.Remove(HierarchyUserParameter->GetUserParameter());
 				}
-				TSharedPtr<FStructOnScope> StructOnScope = MakeShareable(new FStructOnScope(Struct, ParameterProxy->Value().GetBytes()));
-
-				FAddPropertyParams Params = FAddPropertyParams()
-					.UniqueId(Parameter.GetName());
-
-				Row = ChildrenBuilder.AddExternalStructureProperty(StructOnScope.ToSharedRef(), NAME_None, Params);
-
-				FNiagaraVariable UserParameter = Parameter;
-				FNiagaraUserRedirectionParameterStore::MakeUserVariable(UserParameter);
-				ParameterNameToDisplayStruct.Add(UserParameter.GetName(), TWeakPtr<FStructOnScope>(StructOnScope));
 			}
 
-			check(Row && ParameterProxy.IsValid() && ParameterProxy->Value().IsValid());
-
-			TSharedPtr<SWidget> DefaultNameWidget;
-			TSharedPtr<SWidget> DefaultValueWidget;
-
-			Row->DisplayName(FText::FromName(Parameter.GetName()));
-
-			FDetailWidgetRow& CustomWidget = Row->CustomWidget(true);
-
-			Row->GetDefaultWidgets(DefaultNameWidget, DefaultValueWidget, CustomWidget);
-
-			Row->GetPropertyHandle()->SetOnPropertyValuePreChange(FSimpleDelegate::CreateSP(ParameterProxy.ToSharedRef(), &FNiagaraParameterProxy::OnParameterPreChange));
-			Row->GetPropertyHandle()->SetOnChildPropertyValuePreChange(FSimpleDelegate::CreateSP(ParameterProxy.ToSharedRef(), &FNiagaraParameterProxy::OnParameterPreChange));
-			Row->GetPropertyHandle()->SetOnPropertyValueChanged(FSimpleDelegate::CreateSP(ParameterProxy.ToSharedRef(), &FNiagaraParameterProxy::OnParameterChanged));
-			Row->GetPropertyHandle()->SetOnChildPropertyValueChanged(FSimpleDelegate::CreateSP(ParameterProxy.ToSharedRef(), &FNiagaraParameterProxy::OnParameterChanged));
-			Row->GetPropertyHandle()->SetOnPropertyResetToDefault(FSimpleDelegate::CreateSP(ParameterProxy.ToSharedRef(), &FNiagaraParameterProxy::OnResetToDefault));
-
-			CustomWidget
-				.NameContent()
-				[
-					SNew(SBox)
-					.Padding(FMargin(0.0f, 2.0f))
-					[
-						NameWidget.ToSharedRef()
-					]
-				];
-
-			TSharedPtr<SWidget> ValueWidget = DefaultValueWidget;
-			if (CustomValueWidget.IsValid())
+			for(FNiagaraVariable& LeftoverUserParameter : LeftoverUserParameters)
 			{
-				ValueWidget = CustomValueWidget;
+				FNiagaraVariant ParameterValue = GetCurrentParameterValue(LeftoverUserParameter, Component.Get());
+				if (!ParameterValue.IsValid())
+				{
+					continue;
+				}
+			
+				if (LeftoverUserParameter.IsDataInterface())
+				{
+					ParameterValue = FNiagaraVariant(DuplicateObject(ParameterValue.GetDataInterface(), Component.Get()));
+				}
+				
+				TSharedPtr<FNiagaraParameterProxy> ParameterProxy = ParameterProxies.Add_GetRef(MakeShareable(new FNiagaraParameterProxy(Component, LeftoverUserParameter, ParameterValue, OnRebuildChildren, OverridePropertyHandles)));
+				AddRowForUserParameter(ParameterProxy, ChildrenBuilder, ParameterNameToDisplayStruct);
 			}
-
-			CustomWidget
-				.ValueContent()
-				[
-					SNew(SHorizontalBox)
-					+ SHorizontalBox::Slot()
-					.HAlign(HAlign_Fill)
-					.Padding(4.0f)
-					[
-						// Add in the parameter editor factoried above.
-						ValueWidget.ToSharedRef()
-					]
-					+ SHorizontalBox::Slot()
-					.VAlign(VAlign_Center)
-					.AutoWidth()
-					[
-						// Add in the "reset to default" buttons
-						SNew(SButton)
-						.OnClicked(ParameterProxy.ToSharedRef(), &FNiagaraParameterProxy::OnResetToDefaultClicked)
-						.Visibility(ParameterProxy.ToSharedRef(), &FNiagaraParameterProxy::GetResetToDefaultVisibility)
-						.ContentPadding(FMargin(5.f, 0.f))
-						.ToolTipText(LOCTEXT("ResetToDefaultToolTip", "Reset to Default"))
-						.ButtonStyle(FAppStyle::Get(), "NoBorder")
-						.Content()
-						[
-							SNew(SImage)
-							.Image(FAppStyle::GetBrush("PropertyWindow.DiffersFromDefault"))
-						]
-					]
-				];
 		}
 	}
 
@@ -518,6 +767,17 @@ private:
 		}
 	}
 
+	void OnActiveSectionChanged(ECheckBoxState State, const UNiagaraHierarchySection* NewSelection)
+	{
+		ActiveSection = NewSelection;
+		OnRebuildChildren.Execute();
+	}
+
+	ECheckBoxState IsSectionActive(const UNiagaraHierarchySection* Section) const
+	{
+		return ActiveSection == Section ? ECheckBoxState::Checked : ECheckBoxState::Unchecked;
+	}
+
 private:
 	bool bDelegatesInitialized;
 	TWeakObjectPtr<UNiagaraComponent> Component;
@@ -526,8 +786,14 @@ private:
 	TArray<TSharedPtr<FNiagaraParameterProxy>> ParameterProxies;
 	TMap<FName, TWeakPtr<FStructOnScope>> ParameterNameToDisplayStruct;
 	FName CustomBuilderRowName;
+	const UNiagaraHierarchySection* ActiveSection = nullptr;
+	TArray<UNiagaraHierarchySection*> AvailableSections;
 };
 
+void FNiagaraComponentNodeBuilder::GenerateHeaderRowContent(FDetailWidgetRow& NodeRow)
+{
+}
+
 TSharedRef<IDetailCustomization> FNiagaraComponentDetails::MakeInstance()
 {
 	return MakeShareable(new FNiagaraComponentDetails);
@@ -580,7 +846,6 @@ void FNiagaraComponentDetails::OnWorldDestroyed(class UWorld* InWorld)
 
 void FNiagaraComponentDetails::CustomizeDetails(IDetailLayoutBuilder& DetailBuilder)
 {
-
 	FPropertyEditorModule& PropertyModule = FModuleManager::GetModuleChecked<FPropertyEditorModule>("PropertyEditor");
 
 	Builder = &DetailBuilder;
@@ -640,13 +905,13 @@ void FNiagaraComponentDetails::CustomizeDetails(IDetailLayoutBuilder& DetailBuil
 
 		FGameDelegates::Get().GetEndPlayMapDelegate().AddRaw(this, &FNiagaraComponentDetails::OnPiEEnd);
 			
-		IDetailCategoryBuilder& InputParamCategory = DetailBuilder.EditCategory(ParamCategoryName, LOCTEXT("ParamCategoryName", "Override Parameters"), ECategoryPriority::Important);
+		IDetailCategoryBuilder& InputParamCategory = DetailBuilder.EditCategory(ParamCategoryName, LOCTEXT("ParamCategoryName", "User Parameters"), ECategoryPriority::Important);
 		InputParamCategory.AddCustomBuilder(MakeShared<FNiagaraComponentNodeBuilder>(Component.Get(), PropertyHandles, ParamCategoryName));
 	}
 	else if (ObjectsCustomized.Num() > 1)
 	{
-		IDetailCategoryBuilder& InputParamCategory = DetailBuilder.EditCategory(ParamCategoryName, LOCTEXT("ParamCategoryName", "Override Parameters"));
-		InputParamCategory.AddCustomRow(LOCTEXT("ParamCategoryName", "Override Parameters"))
+		IDetailCategoryBuilder& InputParamCategory = DetailBuilder.EditCategory(ParamCategoryName, LOCTEXT("ParamCategoryName", "User Parameters"));
+		InputParamCategory.AddCustomRow(LOCTEXT("ParamCategoryName", "User Parameters"))
 			.WholeRowContent()
 			[
 				SNew(STextBlock)
@@ -770,4 +1035,4 @@ FReply FNiagaraComponentDetails::OnDebugSelectedSystem()
 
 	return FReply::Handled();
 }
-#undef LOCTEXT_NAMESPACE
+#undef LOCTEXT_NAMESPACE
\ No newline at end of file
diff --git a/Engine/Plugins/FX/Niagara/Source/NiagaraEditor/Private/Customizations/NiagaraScriptVariableCustomization.cpp b/Engine/Plugins/FX/Niagara/Source/NiagaraEditor/Private/Customizations/NiagaraScriptVariableCustomization.cpp
--- a/Engine/Plugins/FX/Niagara/Source/NiagaraEditor/Private/Customizations/NiagaraScriptVariableCustomization.cpp
+++ b/Engine/Plugins/FX/Niagara/Source/NiagaraEditor/Private/Customizations/NiagaraScriptVariableCustomization.cpp
@@ -62,7 +62,7 @@ TArray<UEdGraphPin*> FNiagaraScriptVariableDetails::GetDefaultPins()
 	}
 	return TArray<UEdGraphPin*>();
 }
- 
+
 void FNiagaraScriptVariableDetails::Refresh()
 {
 	IDetailLayoutBuilder* DetailBuilderPtr = nullptr;
@@ -691,4 +691,57 @@ void FNiagaraScriptVariableDetails::OnLibrarySynchronizedDefaultModeChanged(int3
 	}
 }
 
+TSharedRef<IDetailCustomization> FNiagaraScriptVariableHierarchyDetails::MakeInstance()
+{
+	return MakeShared<FNiagaraScriptVariableHierarchyDetails>();
+}
+
+FNiagaraScriptVariableHierarchyDetails::FNiagaraScriptVariableHierarchyDetails()
+{
+}
+
+void FNiagaraScriptVariableHierarchyDetails::CustomizeDetails(IDetailLayoutBuilder& DetailBuilder)
+{
+	TArray<TWeakObjectPtr<UObject>> ObjectsCustomized;
+	DetailBuilder.GetObjectsBeingCustomized(ObjectsCustomized);
+
+	if (ObjectsCustomized.Num() != 1)
+	{
+		// Only allow selecting one UNiagaraScriptVariable at a time.
+		return;
+	}
+	if (!ObjectsCustomized[0]->IsA<UNiagaraScriptVariable>())
+	{
+		return;
+	}
+
+	UNiagaraScriptVariable* Variable = Cast<UNiagaraScriptVariable>(ObjectsCustomized[0].Get());
+	if (Variable == nullptr)
+	{
+		return;
+	}
+
+	DetailBuilder.HideProperty(GET_MEMBER_NAME_CHECKED(UNiagaraScriptVariable, DefaultMode));
+	DetailBuilder.HideProperty(GET_MEMBER_NAME_CHECKED(UNiagaraScriptVariable, DefaultBinding));
+	// member isn't public so we can't use GET_MEMBER_NAME_CHECKED
+	DetailBuilder.HideProperty("DefaultValueVariant", UNiagaraScriptVariable::StaticClass());
+
+	DetailBuilder.HideProperty(GET_MEMBER_NAME_CHECKED(UNiagaraScriptVariable, Metadata.CategoryName));
+	DetailBuilder.HideProperty(GET_MEMBER_NAME_CHECKED(UNiagaraScriptVariable, Metadata.bAdvancedDisplay));
+	DetailBuilder.HideProperty(GET_MEMBER_NAME_CHECKED(UNiagaraScriptVariable, Metadata.bOverrideColor));
+	DetailBuilder.HideProperty(GET_MEMBER_NAME_CHECKED(UNiagaraScriptVariable, Metadata.bEnableBoolOverride));
+	DetailBuilder.HideProperty(GET_MEMBER_NAME_CHECKED(UNiagaraScriptVariable, Metadata.bInlineEditConditionToggle));
+	DetailBuilder.HideProperty(GET_MEMBER_NAME_CHECKED(UNiagaraScriptVariable, Metadata.EditCondition));
+	DetailBuilder.HideProperty(GET_MEMBER_NAME_CHECKED(UNiagaraScriptVariable, Metadata.VisibleCondition));
+	DetailBuilder.HideProperty(GET_MEMBER_NAME_CHECKED(UNiagaraScriptVariable, Metadata.EditorSortPriority));
+	DetailBuilder.HideProperty(GET_MEMBER_NAME_CHECKED(UNiagaraScriptVariable, Metadata.PropertyMetaData));
+	DetailBuilder.HideProperty(GET_MEMBER_NAME_CHECKED(UNiagaraScriptVariable, Metadata.ParentAttribute));
+	DetailBuilder.HideProperty(GET_MEMBER_NAME_CHECKED(UNiagaraScriptVariable, Metadata.bDisplayInOverviewStack));
+	DetailBuilder.HideProperty(GET_MEMBER_NAME_CHECKED(UNiagaraScriptVariable, Metadata.InlineParameterSortPriority));
+	DetailBuilder.HideProperty(GET_MEMBER_NAME_CHECKED(UNiagaraScriptVariable, Metadata.InlineParameterEnumOverrides));
+	DetailBuilder.HideProperty(GET_MEMBER_NAME_CHECKED(UNiagaraScriptVariable, Metadata.InlineParameterColorOverride));
+	DetailBuilder.HideProperty(GET_MEMBER_NAME_CHECKED(UNiagaraScriptVariable, Metadata.InlineParameterBoolOverride));
+
+}
+
 #undef LOCTEXT_NAMESPACE
diff --git a/Engine/Plugins/FX/Niagara/Source/NiagaraEditor/Private/Customizations/NiagaraScriptVariableCustomization.h b/Engine/Plugins/FX/Niagara/Source/NiagaraEditor/Private/Customizations/NiagaraScriptVariableCustomization.h
--- a/Engine/Plugins/FX/Niagara/Source/NiagaraEditor/Private/Customizations/NiagaraScriptVariableCustomization.h
+++ b/Engine/Plugins/FX/Niagara/Source/NiagaraEditor/Private/Customizations/NiagaraScriptVariableCustomization.h
@@ -112,3 +112,19 @@ private:
 
 	static const FName DefaultValueCategoryName;
 };
+
+/** This customization sets up a custom details panel for script variables found in the hierarchy editor. */
+class FNiagaraScriptVariableHierarchyDetails : public IDetailCustomization
+{
+public:
+ 
+	/** Makes a new instance of this detail layout class for a specific detail view requesting it */
+	static TSharedRef<IDetailCustomization> MakeInstance();
+ 
+	FNiagaraScriptVariableHierarchyDetails();
+	virtual ~FNiagaraScriptVariableHierarchyDetails() override {}
+ 
+	/** IDetailCustomization interface */
+	virtual void CustomizeDetails( IDetailLayoutBuilder& DetailBuilder ) override;
+};
+
diff --git a/Engine/Plugins/FX/Niagara/Source/NiagaraEditor/Private/NiagaraActions.cpp b/Engine/Plugins/FX/Niagara/Source/NiagaraEditor/Private/NiagaraActions.cpp
--- a/Engine/Plugins/FX/Niagara/Source/NiagaraEditor/Private/NiagaraActions.cpp
+++ b/Engine/Plugins/FX/Niagara/Source/NiagaraEditor/Private/NiagaraActions.cpp
@@ -630,7 +630,7 @@ TSharedPtr<SWidget> FNiagaraParameterDragOperation::GetDefaultDecorator() const
 		+ SVerticalBox::Slot()
 		.AutoHeight()
 		[
-			FNiagaraParameterUtilities::GetParameterWidget(StaticCastSharedPtr<FNiagaraParameterAction>(SourceAction)->GetParameter(), false)
+			FNiagaraParameterUtilities::GetParameterWidget(StaticCastSharedPtr<FNiagaraParameterAction>(SourceAction)->GetParameter(), true, false)
 		]
 		+ SVerticalBox::Slot()
 		.AutoHeight()
diff --git a/Engine/Plugins/FX/Niagara/Source/NiagaraEditor/Private/NiagaraEditorStyle.cpp b/Engine/Plugins/FX/Niagara/Source/NiagaraEditor/Private/NiagaraEditorStyle.cpp
--- a/Engine/Plugins/FX/Niagara/Source/NiagaraEditor/Private/NiagaraEditorStyle.cpp
+++ b/Engine/Plugins/FX/Niagara/Source/NiagaraEditor/Private/NiagaraEditorStyle.cpp
@@ -681,6 +681,30 @@ void FNiagaraEditorStyle::InitScalabilityIcons()
 	Set("NiagaraEditor.Scalability.Preview.ResetPlatform", new IMAGE_BRUSH("Icons/Scalability/ResetPreviewPlatform", Icon20x20));
 }
 
+void FNiagaraEditorStyle::InitHierarchyEditor()
+{
+	const FTextBlockStyle NormalText = FAppStyle::GetWidgetStyle<FTextBlockStyle>("NormalText");
+	FSlateFontInfo CategoryFont = FAppStyle::Get().GetFontStyle(TEXT("DetailsView.CategoryFontStyle"));
+	CategoryFont.Size = 11;
+	// const FEditableTextBoxStyle NormalEditableTextBox = FStarshipCoreStyle::GetCoreStyle().GetWidgetStyle<FEditableTextBoxStyle>("NormalEditableTextBox");
+	const FEditableTextBoxStyle NormalEditableTextBox = FAppStyle::Get().GetWidgetStyle<FEditableTextBoxStyle>("NormalEditableTextBox");
+	
+	FEditableTextBoxStyle CategoryEditableText = FEditableTextBoxStyle(NormalEditableTextBox)
+		.SetFont(CategoryFont)
+		.SetForegroundColor(FStyleColors::AccentWhite);
+	
+	FTextBlockStyle CategoryText = FTextBlockStyle(NormalText)
+		.SetFont(CategoryFont);
+	FInlineEditableTextBlockStyle HierarchyCategoryTextStyle = FInlineEditableTextBlockStyle()
+		.SetTextStyle(CategoryText)
+		.SetEditableTextBoxStyle(CategoryEditableText);
+	
+	Set("NiagaraEditor.HierarchyEditor.Category", HierarchyCategoryTextStyle);
+
+	Set("NiagaraEditor.Stack.DropTarget.BorderVertical", new IMAGE_BRUSH("Icons/StackDropTargetBorder_Vertical", FVector2D(2, 8), FLinearColor::White, ESlateBrushTileType::Vertical));
+	Set("NiagaraEditor.Stack.DropTarget.BorderHorizontal", new IMAGE_BRUSH("Icons/StackDropTargetBorder_Horizontal", FVector2D(8, 2), FLinearColor::White, ESlateBrushTileType::Horizontal));
+}
+
 FNiagaraEditorStyle::FNiagaraEditorStyle() : FSlateStyleSet("NiagaraEditorStyle")
 {
 	SetContentRoot(FPaths::EnginePluginsDir() / TEXT("FX/Niagara/Content/Slate"));
@@ -715,7 +739,9 @@ FNiagaraEditorStyle::FNiagaraEditorStyle() : FSlateStyleSet("NiagaraEditorStyle"
 	InitScalabilityIcons();
 	InitViewportStyle();
 	InitScratchStyle();
+	InitHierarchyEditor();
 }
+
 void FNiagaraEditorStyle::InitScratchStyle()
 {
 	const FTextBlockStyle NormalText = FAppStyle::Get().GetWidgetStyle<FTextBlockStyle>("NormalText");
diff --git a/Engine/Plugins/FX/Niagara/Source/NiagaraEditor/Private/NiagaraEditorUtilities.cpp b/Engine/Plugins/FX/Niagara/Source/NiagaraEditor/Private/NiagaraEditorUtilities.cpp
--- a/Engine/Plugins/FX/Niagara/Source/NiagaraEditor/Private/NiagaraEditorUtilities.cpp
+++ b/Engine/Plugins/FX/Niagara/Source/NiagaraEditor/Private/NiagaraEditorUtilities.cpp
@@ -69,6 +69,7 @@
 #include "Widgets/SNiagaraPinTypeSelector.h"
 #include "ViewModels/NiagaraEmitterHandleViewModel.h"
 #include "ViewModels/NiagaraEmitterViewModel.h"
+#include "ViewModels/NiagaraParameterPanelViewModel.h"
 
 #define LOCTEXT_NAMESPACE "FNiagaraEditorUtilities"
 
@@ -2190,6 +2191,21 @@ bool FNiagaraEditorUtilities::AddParameter(FNiagaraVariable& NewParameterVariabl
 	return bSuccess;
 }
 
+TObjectPtr<UNiagaraScriptVariable> FNiagaraEditorUtilities::GetScriptVariableForUserParameter(const FNiagaraVariable& UserParameter, TSharedPtr<FNiagaraSystemViewModel> SystemViewModel)
+{
+	return Cast<UNiagaraSystemEditorData>(SystemViewModel->GetSystem().GetEditorData())->FindOrAddUserScriptVariable(UserParameter, SystemViewModel->GetSystem());
+}
+
+TObjectPtr<UNiagaraScriptVariable> FNiagaraEditorUtilities::GetScriptVariableForUserParameter(const FNiagaraVariable& UserParameter, UNiagaraSystem& System)
+{
+	return Cast<UNiagaraSystemEditorData>(System.GetEditorData())->FindOrAddUserScriptVariable(UserParameter, System);
+}
+
+TObjectPtr<UNiagaraScriptVariable> FNiagaraEditorUtilities::FindScriptVariableForUserParameter(const FGuid& UserParameterGuid, UNiagaraSystem& System)
+{
+	return Cast<UNiagaraSystemEditorData>(System.GetEditorData())->FindUserScriptVariable(UserParameterGuid);
+}
+
 void FNiagaraEditorUtilities::ShowParentEmitterInContentBrowser(TSharedRef<FNiagaraEmitterViewModel> Emitter)
 {
 	FContentBrowserModule& ContentBrowserModule = FModuleManager::Get().LoadModuleChecked<FContentBrowserModule>(TEXT("ContentBrowser"));
@@ -3708,20 +3724,15 @@ bool FNiagaraParameterUtilities::TestCanRenameWithMessage(FName ParameterName, F
 	return false;
 }
 
-TSharedRef<SWidget> FNiagaraParameterUtilities::GetParameterWidget(FNiagaraVariable Variable, bool bShowValue)
+TSharedRef<SWidget> FNiagaraParameterUtilities::GetParameterWidget(FNiagaraVariable Variable, bool bAddTypeIcon, bool bShowValue)
 {
 	FNiagaraTypeDefinition Type = Variable.GetType();
 	const FLinearColor TypeColor = UEdGraphSchema_Niagara::GetTypeColor(Type);
 	TSharedPtr<INiagaraEditorTypeUtilities> TypeUtilities = FNiagaraEditorModule::Get().GetTypeUtilities(Type);
 
-	FText			   IconToolTip = FText::GetEmpty();
-	FSlateBrush const* IconBrush = FAppStyle::Get().GetBrush(TEXT("Kismet.AllClasses.VariableIcon"));
-	FSlateColor        IconColor = FSlateColor(TypeColor);
-	FString			   IconDocLink, IconDocExcerpt;
-	FSlateBrush const* SecondaryIconBrush = FAppStyle::GetBrush(TEXT("NoBrush"));
-	FSlateColor        SecondaryIconColor = IconColor;
-
 	TSharedPtr<SHorizontalBox> ParameterContainer;
+	
+	SHorizontalBox::FSlot* TypeIconSlot;
 	TSharedRef<SVerticalBox> ParameterWidget = SNew(SVerticalBox)
 		+ SVerticalBox::Slot()
 		.AutoHeight()
@@ -3731,16 +3742,7 @@ TSharedRef<SWidget> FNiagaraParameterUtilities::GetParameterWidget(FNiagaraVaria
 			.AutoWidth()
 			.HAlign(HAlign_Center)
 			.VAlign(VAlign_Center)
-			[
-				SNew(SNiagaraIconWidget)
-				.IconToolTip(IconToolTip)
-				.IconBrush(IconBrush)
-				.IconColor(IconColor)
-				.DocLink(IconDocLink)
-				.DocExcerpt(IconDocExcerpt)
-				.SecondaryIconBrush(SecondaryIconBrush) 
-				.SecondaryIconColor(SecondaryIconColor)
-			]
+			.Expose(TypeIconSlot)
 			+ SHorizontalBox::Slot()
 			.AutoWidth()
 			.HAlign(HAlign_Center)
@@ -3753,6 +3755,26 @@ TSharedRef<SWidget> FNiagaraParameterUtilities::GetParameterWidget(FNiagaraVaria
 			]
 		];
 
+	if(bAddTypeIcon)
+	{
+		FText			   IconToolTip = FText::GetEmpty();
+		FSlateBrush const* IconBrush = FAppStyle::Get().GetBrush(TEXT("Kismet.AllClasses.VariableIcon"));
+		FSlateColor        IconColor = FSlateColor(TypeColor);
+		FString			   IconDocLink, IconDocExcerpt;
+		FSlateBrush const* SecondaryIconBrush = FAppStyle::GetBrush(TEXT("NoBrush"));
+		FSlateColor        SecondaryIconColor = IconColor;
+		
+		TypeIconSlot->AttachWidget(
+			SNew(SNiagaraIconWidget)
+			   .IconToolTip(IconToolTip)
+			   .IconBrush(IconBrush)
+			   .IconColor(IconColor)
+			   .DocLink(IconDocLink)
+			   .DocExcerpt(IconDocExcerpt)
+			   .SecondaryIconBrush(SecondaryIconBrush) 
+			   .SecondaryIconColor(SecondaryIconColor));
+	}
+	
 	if(bShowValue && Variable.IsDataAllocated())
 	{
 		const FText ValueText = TypeUtilities->GetStackDisplayText(Variable);
@@ -3799,7 +3821,7 @@ TSharedRef<SToolTip> FNiagaraParameterUtilities::GetTooltipWidget(FNiagaraVariab
 	FSlateColor        SecondaryIconColor = IconColor;
 
 	TSharedPtr<SHorizontalBox> ParameterContainer;
-	TSharedRef<SWidget> TooltipContent = GetParameterWidget(Variable, bShowValue);
+	TSharedRef<SWidget> TooltipContent = GetParameterWidget(Variable, true, bShowValue);
 
 	TSharedPtr<SVerticalBox> InnerContainer;
 	// we construct a tooltip widget that shows the parameter the value is associated with
diff --git a/Engine/Plugins/FX/Niagara/Source/NiagaraEditor/Private/NiagaraSystemEditorData.cpp b/Engine/Plugins/FX/Niagara/Source/NiagaraEditor/Private/NiagaraSystemEditorData.cpp
--- a/Engine/Plugins/FX/Niagara/Source/NiagaraEditor/Private/NiagaraSystemEditorData.cpp
+++ b/Engine/Plugins/FX/Niagara/Source/NiagaraEditor/Private/NiagaraSystemEditorData.cpp
@@ -10,6 +10,8 @@
 #include "NiagaraOverviewNode.h"
 #include "NiagaraNode.h"
 #include "EdGraphSchema_NiagaraSystemOverview.h"
+#include "NiagaraConstants.h"
+#include "NiagaraScriptVariable.h"
 #include "EdGraph/EdGraph.h"
 
 static constexpr float SystemOverviewNodePadding = 250.0f;
@@ -62,6 +64,7 @@ UNiagaraSystemEditorData::UNiagaraSystemEditorData(const FObjectInitializer& Obj
 {
 	RootFolder = ObjectInitializer.CreateDefaultSubobject<UNiagaraSystemEditorFolder>(this, TEXT("RootFolder"));
 	StackEditorData = ObjectInitializer.CreateDefaultSubobject<UNiagaraStackEditorData>(this, TEXT("StackEditorData"));
+	UserParameterHierarchy = CreateDefaultSubobject<UNiagaraHierarchyRoot>(TEXT("UserParameterHierarchyRoot"));
 	OwnerTransform.SetLocation(FVector(0.0f, 0.0f, 0.0f));
 	PlaybackRangeMin = 0;
 	PlaybackRangeMax = 10;
@@ -90,6 +93,21 @@ void UNiagaraSystemEditorData::PostLoadFromOwner(UObject* InOwner)
 	{
 		StackEditorData = NewObject<UNiagaraStackEditorData>(this, TEXT("StackEditorData"), RF_Transactional);
 	}
+	if (UserParameterHierarchy == nullptr)
+	{
+		UserParameterHierarchy = NewObject<UNiagaraHierarchyRoot>(this, TEXT("UserParameterHierarchyRoot"), RF_Transactional);
+	}
+
+	TArray<FNiagaraVariable> UserParameters;
+	OwningSystem->GetExposedParameters().GetUserParameters(UserParameters);
+	for(FNiagaraVariable& UserParameter : UserParameters)
+	{
+		FNiagaraVariable RedirectedUserParameter = UserParameter;
+		OwningSystem->GetExposedParameters().RedirectUserVariable(RedirectedUserParameter);
+		FindOrAddUserScriptVariable(RedirectedUserParameter, *OwningSystem);
+	}
+
+	OwningSystem->GetExposedParameters().AddOnChangedHandler(FNiagaraParameterStore::FOnChanged::FDelegate::CreateUObject(this, &UNiagaraSystemEditorData::SyncUserScriptVariables, OwningSystem));
 
 	const int32 NiagaraVer = GetLinkerCustomVersion(FNiagaraCustomVersion::GUID);
 
@@ -312,3 +330,93 @@ void UNiagaraSystemEditorData::SynchronizeOverviewGraphWithSystem(UNiagaraSystem
 	// Dispatch an empty graph changed message here so that any graph UIs which are visible refresh their node's widgets.
 	SystemOverviewGraph->NotifyGraphChanged();
 }
+
+void UNiagaraSystemEditorData::SyncUserScriptVariables(UNiagaraSystem* System)
+{
+	TArray<FNiagaraVariable> UserParameters;
+	System->GetExposedParameters().GetUserParameters(UserParameters);
+
+	TArray<UNiagaraScriptVariable*> CurrentVariables = UserParameterMetaData;
+
+	TSet<UNiagaraScriptVariable*> UpdatedScriptVariables;
+	for(FNiagaraVariable UserParameter : UserParameters)
+	{
+		UpdatedScriptVariables.Add(FindOrAddUserScriptVariable(UserParameter, *System));
+	}
+	
+	for(UNiagaraScriptVariable* CurrentVariable : CurrentVariables)
+	{
+		if(!UpdatedScriptVariables.Contains(CurrentVariable))
+		{
+			UserParameterMetaData.Remove(CurrentVariable);
+		}
+	}
+
+	OnUserParameterScriptVariablesSyncedDelegate.Broadcast();
+}
+
+UNiagaraScriptVariable* UNiagaraSystemEditorData::FindOrAddUserScriptVariable(FNiagaraVariable UserParameter, UNiagaraSystem& System)
+{
+	bool bSuccess = System.GetExposedParameters().RedirectUserVariable(UserParameter);
+	if(bSuccess)
+	{
+		TObjectPtr<UNiagaraScriptVariable>* ExistingScriptVariable = UserParameterMetaData.FindByPredicate([UserParameter](const UNiagaraScriptVariable* ScriptVariable)
+			{
+				return ScriptVariable->Variable == UserParameter;
+			});
+
+		if(ExistingScriptVariable == nullptr)
+		{
+			TObjectPtr<UNiagaraScriptVariable> NewScriptVariable = NewObject<UNiagaraScriptVariable>(this);
+			FNiagaraVariableMetaData MetaData;
+			NewScriptVariable->Init(UserParameter, MetaData);
+			UserParameterMetaData.Add(NewScriptVariable);
+			ExistingScriptVariable = &NewScriptVariable;
+		}
+		
+		return *ExistingScriptVariable;
+	}
+
+	return nullptr;
+}
+
+UNiagaraScriptVariable* UNiagaraSystemEditorData::FindUserScriptVariable(FGuid UserParameterGuid)
+{
+	TObjectPtr<UNiagaraScriptVariable>* ExistingScriptVariable = UserParameterMetaData.FindByPredicate([UserParameterGuid](const UNiagaraScriptVariable* ScriptVariable)
+		{
+			return ScriptVariable->Metadata.GetVariableGuid() == UserParameterGuid;
+		});
+
+	if(ExistingScriptVariable == nullptr)
+	{
+		return nullptr;
+	}
+
+	return *ExistingScriptVariable;
+}
+
+bool UNiagaraSystemEditorData::RenameUserScriptVariable(FNiagaraVariable OldVariable, FName NewName)
+{
+	TObjectPtr<UNiagaraScriptVariable>* ExistingScriptVariable = UserParameterMetaData.FindByPredicate([OldVariable](const UNiagaraScriptVariable* ScriptVariable)
+		{
+			return ScriptVariable->Variable == OldVariable;
+		});
+
+	if(ExistingScriptVariable == nullptr)
+	{
+		return false;
+	}
+
+	(*ExistingScriptVariable)->Modify();
+	(*ExistingScriptVariable)->Variable.SetName(NewName);
+	(*ExistingScriptVariable)->UpdateChangeId();
+	return true;
+}
+
+bool UNiagaraSystemEditorData::RemoveUserScriptVariable(FNiagaraVariable Variable)
+{
+	return UserParameterMetaData.RemoveAll([Variable](const UNiagaraScriptVariable* ScriptVariable)
+		{
+			return ScriptVariable->Variable == Variable;
+		}) > 0;
+}
diff --git a/Engine/Plugins/FX/Niagara/Source/NiagaraEditor/Private/Toolkits/SystemToolkitModes/NiagaraSystemToolkitModeBase.cpp b/Engine/Plugins/FX/Niagara/Source/NiagaraEditor/Private/Toolkits/SystemToolkitModes/NiagaraSystemToolkitModeBase.cpp
--- a/Engine/Plugins/FX/Niagara/Source/NiagaraEditor/Private/Toolkits/SystemToolkitModes/NiagaraSystemToolkitModeBase.cpp
+++ b/Engine/Plugins/FX/Niagara/Source/NiagaraEditor/Private/Toolkits/SystemToolkitModes/NiagaraSystemToolkitModeBase.cpp
@@ -25,11 +25,14 @@
 #include "Widgets/SNiagaraSelectedObjectsDetails.h"
 #include "Widgets/SNiagaraSpreadsheetView.h"
 #include "Widgets/SNiagaraGeneratedCodeView.h"
+#include "Widgets/SNiagaraHierarchy.h"
 #include "Widgets/Layout/SWidgetSwitcher.h"
 #include "ViewModels/NiagaraScratchPadScriptViewModel.h"
 #include "Widgets/SNiagaraSelectedObjectsDetails.h"
 #include "NiagaraObjectSelection.h"
+#include "Styling/StyleColors.h"
 #include "ViewModels/NiagaraSystemEditorDocumentsViewModel.h"
+#include "ViewModels/HierarchyEditor/NiagaraUserParametersHierarchyViewModel.h"
 
 #define LOCTEXT_NAMESPACE "NiagaraSystemToolkitModeBase"
 
@@ -51,7 +54,7 @@ const FName FNiagaraSystemToolkitModeBase::ScriptStatsTabID(TEXT("NiagaraSystemE
 const FName FNiagaraSystemToolkitModeBase::BakerTabID(TEXT("NiagaraSystemEditor_Baker"));
 const FName FNiagaraSystemToolkitModeBase::VersioningTabID(TEXT("NiagaraSystemEditor_Versioning"));
 const FName FNiagaraSystemToolkitModeBase::ScratchPadScriptsTabID(TEXT("NiagaraSystemEditor_ScratchPadScripts"));
-
+const FName FNiagaraSystemToolkitModeBase::UserParametersHierarchyTabID(TEXT("NiagaraSystemEditor_UserParameters"));
 
 FNiagaraSystemToolkitModeBase::FNiagaraSystemToolkitModeBase(FName InModeName, TWeakPtr<FNiagaraSystemToolkit> InSystemToolkit) : FApplicationMode(InModeName), SystemToolkit(InSystemToolkit), SwitcherIdx(0)
 {
@@ -286,6 +289,13 @@ void FNiagaraSystemToolkitModeBase::RegisterTabFactories(TSharedPtr<FTabManager>
 		.SetDisplayName(LOCTEXT("VersioningTab", "Versioning"))
 		.SetGroup(WorkspaceMenuCategory.ToSharedRef())
 		.SetIcon(FSlateIcon(FAppStyle::GetAppStyleSetName(), "Versions"));
+
+	if(SystemToolkit.Pin()->GetSystemViewModel()->GetEditMode() == ENiagaraSystemViewModelEditMode::SystemAsset)
+	{
+		InTabManager->RegisterTabSpawner(UserParametersHierarchyTabID, FOnSpawnTab::CreateSP(this, &FNiagaraSystemToolkitModeBase::SpawnTab_UserParametersHierarchyEditor))
+			.SetDisplayName(LOCTEXT("UserParametersHierarchyTab", "User Parameters Hierarchy"))
+			.SetGroup(WorkspaceMenuCategory.ToSharedRef());
+	}
 }
 
 int FNiagaraSystemToolkitModeBase::GetActiveSelectionDetailsIndex() const
@@ -758,4 +768,21 @@ TSharedRef<SDockTab> FNiagaraSystemToolkitModeBase::SpawnTab_Versioning(const FS
 	return SpawnedTab;
 }
 
-#undef LOCTEXT_NAMESPACE
\ No newline at end of file
+TSharedRef<SDockTab> FNiagaraSystemToolkitModeBase::SpawnTab_UserParametersHierarchyEditor(const FSpawnTabArgs& Args)
+{
+	check(Args.GetTabId().TabType == UserParametersHierarchyTabID);
+	
+	TSharedRef<SDockTab> SpawnedTab = SNew(SDockTab)
+		.Label(LOCTEXT("UserParametersHierarchyTitle", "User Parameters Hierarchy"))
+		[
+			SNew(SBox)
+			.AddMetaData<FTagMetaData>(FTagMetaData(TEXT("UserParameters")))
+			[
+				SNew(SNiagaraHierarchy, SystemToolkit.Pin()->GetSystemViewModel()->GetUserParametersHierarchyViewModel())
+			]
+		];
+	
+	return SpawnedTab;
+}
+
+#undef LOCTEXT_NAMESPACE
diff --git a/Engine/Plugins/FX/Niagara/Source/NiagaraEditor/Private/Toolkits/SystemToolkitModes/NiagaraSystemToolkitModeBase.h b/Engine/Plugins/FX/Niagara/Source/NiagaraEditor/Private/Toolkits/SystemToolkitModes/NiagaraSystemToolkitModeBase.h
--- a/Engine/Plugins/FX/Niagara/Source/NiagaraEditor/Private/Toolkits/SystemToolkitModes/NiagaraSystemToolkitModeBase.h
+++ b/Engine/Plugins/FX/Niagara/Source/NiagaraEditor/Private/Toolkits/SystemToolkitModes/NiagaraSystemToolkitModeBase.h
@@ -41,6 +41,7 @@ private:
 	TSharedRef<SDockTab> SpawnTab_ScriptStats(const FSpawnTabArgs& Args);
 	TSharedRef<SDockTab> SpawnTab_Baker(const FSpawnTabArgs& Args);
 	TSharedRef<SDockTab> SpawnTab_Versioning(const FSpawnTabArgs& Args);
+	TSharedRef<SDockTab> SpawnTab_UserParametersHierarchyEditor(const FSpawnTabArgs& Args);
 
 protected:
 	TWeakPtr<FNiagaraSystemToolkit> SystemToolkit;
@@ -81,5 +82,6 @@ public:
 	static const FName ScriptStatsTabID;
 	static const FName BakerTabID;
 	static const FName VersioningTabID;
+	static const FName UserParametersHierarchyTabID;
 };
 
diff --git a/Engine/Plugins/FX/Niagara/Source/NiagaraEditor/Private/ViewModels/HierarchyEditor/NiagaraHierarchyViewModelBase.cpp b/Engine/Plugins/FX/Niagara/Source/NiagaraEditor/Private/ViewModels/HierarchyEditor/NiagaraHierarchyViewModelBase.cpp
new file mode 100644
--- /dev/null
+++ b/Engine/Plugins/FX/Niagara/Source/NiagaraEditor/Private/ViewModels/HierarchyEditor/NiagaraHierarchyViewModelBase.cpp
@@ -0,0 +1,949 @@
+// Copyright Epic Games, Inc. All Rights Reserved.
+
+#include "ViewModels/HierarchyEditor/NiagaraHierarchyViewModelBase.h"
+#include "NiagaraEditorUtilities.h"
+#include "NiagaraEmitterEditorData.h"
+#include "NiagaraNodeFunctionCall.h"
+#include "NiagaraScriptVariable.h"
+#include "NiagaraSystemEditorData.h"
+#include "SDropTarget.h"
+#include "ViewModels/NiagaraSystemViewModel.h"
+#include "ViewModels/Stack/NiagaraStackViewModel.h"
+#include "Widgets/Views/STableRow.h"
+#include "Widgets/SNiagaraHierarchy.h"
+#include "Widgets/Text/SInlineEditableTextBlock.h"
+#include "Widgets/SToolTip.h"
+#include "Widgets/Input/SCheckBox.h"
+#include "Editor.h"
+#include "ScopedTransaction.h"
+
+#define LOCTEXT_NAMESPACE "NiagaraHierarchyEditor"
+
+const TArray<UNiagaraHierarchyItemBase*>& UNiagaraHierarchyItemBase::GetChildren() const
+{
+	if(CanHaveChildren())
+	{
+		return Children;
+	}
+
+	static TArray<UNiagaraHierarchyItemBase*> EmptyChildren;
+	return EmptyChildren;
+}
+
+void UNiagaraHierarchyItemBase::Refresh()
+{
+	TArray<UNiagaraHierarchyItemBase*> OldChildren = Children;
+
+	RefreshDataInternal();
+
+	Children.Empty();
+	for(UNiagaraHierarchyItemBase* Item : OldChildren)
+	{
+		if(Item->IsFinalized())
+		{
+			continue;
+		}
+
+		Item->Refresh();
+		Children.Add(Item);
+	}
+}
+
+void UNiagaraHierarchyItemBase::Finalize()
+{	
+	for(UNiagaraHierarchyItemBase* Child : Children)
+	{
+		Child->Finalize();
+	}
+	
+	bFinalized = true;
+}
+
+bool UNiagaraHierarchyItemBase::Modify(bool bAlwaysMarkDirty)
+{
+	bool bSavedToTransactionBuffer = true;
+
+	for(UNiagaraHierarchyItemBase* Child : Children)
+	{
+		bSavedToTransactionBuffer &= Child->Modify(bAlwaysMarkDirty);
+	}
+	
+	bSavedToTransactionBuffer &= UObject::Modify(bAlwaysMarkDirty);
+
+	return bSavedToTransactionBuffer;
+}
+
+void UNiagaraHierarchyRoot::RefreshDataInternal()
+{
+	// in addition, we update our sections
+	TArray<UNiagaraHierarchySection*> OldSections = Sections;
+
+	Sections.Empty();
+	for(UNiagaraHierarchySection* Section : OldSections)
+	{
+		// if a section was marked finalized, we don't add it back
+		if(Section->IsFinalized())
+		{
+			continue;
+		}
+
+		Sections.Add(Section);
+	}
+}
+
+UNiagaraHierarchySection* UNiagaraHierarchyRoot::AddSection(FText InNewSectionName)
+{
+	TSet<FName> ExistingSectionNames;
+	
+	for(FName& SectionName : GetSections())
+	{
+		ExistingSectionNames.Add(SectionName);
+	}
+	
+	FName NewName = FNiagaraUtilities::GetUniqueName(FName(InNewSectionName.ToString()), ExistingSectionNames);
+	UNiagaraHierarchySection* NewSectionItem = NewObject<UNiagaraHierarchySection>(this);
+	NewSectionItem->SetSectionName(NewName);
+	NewSectionItem->SetFlags(RF_Transactional);
+	Sections.Add(NewSectionItem);
+	return NewSectionItem;
+}
+
+void UNiagaraHierarchyRoot::RemoveSection(FText SectionName)
+{
+	if(ensure(Sections.ContainsByPredicate([SectionName](UNiagaraHierarchySection* Section)
+	{
+		return Section->GetSectionNameAsText().EqualTo(SectionName);
+	})))
+	{
+		Sections.RemoveAll([SectionName](UNiagaraHierarchySection* Section)
+		{
+			return Section->GetSectionNameAsText().EqualTo(SectionName);
+		});
+	}
+}
+
+TSet<FName> UNiagaraHierarchyRoot::GetSections() const
+{
+	TSet<FName> OutSections;
+	for(UNiagaraHierarchySection* Section : Sections)
+	{
+		OutSections.Add(Section->GetSectionName());
+	}
+
+	return OutSections;
+}
+
+int32 UNiagaraHierarchyRoot::GetSectionIndex(UNiagaraHierarchySection* Section) const
+{
+	return Sections.Find(Section);
+}
+
+bool UNiagaraHierarchyRoot::Modify(bool bAlwaysMarkDirty)
+{
+	bool bSavedToTransactionBuffer = true;
+	
+	for(UNiagaraHierarchySection* Section : Sections)
+	{
+		bSavedToTransactionBuffer &= Section->Modify();
+	}
+	
+	bSavedToTransactionBuffer &= Super::Modify(bAlwaysMarkDirty);	
+
+	return bSavedToTransactionBuffer;
+}
+
+bool FNiagaraHierarchyCategoryViewModel::IsTopCategoryActive() const
+{
+	const UNiagaraHierarchyCategory* Result = Cast<UNiagaraHierarchyCategory>(GetDataMutable());
+	const UNiagaraHierarchyCategory* TopLevelCategory = Result;
+	for (; TopLevelCategory != nullptr; TopLevelCategory = TopLevelCategory->GetTypedOuter<UNiagaraHierarchyCategory>() )
+	{
+		if(TopLevelCategory != nullptr)
+		{
+			Result = TopLevelCategory;
+		}
+	}
+	
+	return HierarchyViewModel->IsSectionActive(Result->GetSection());
+}
+
+void UNiagaraHierarchySection::SetSectionNameAsText(const FText& Text)
+{
+	Section = FName(Text.ToString());
+}
+
+UNiagaraHierarchyViewModelBase::UNiagaraHierarchyViewModelBase()
+{
+	SourceRoot = NewObject<UNiagaraHierarchyRoot>(this, FName("HierarchySourceRoot"));
+	SourceViewModelRoot = MakeShared<FNiagaraHierarchyRootViewModel>(SourceRoot.Get(), this);
+	Commands = MakeShared<FUICommandList>();
+}
+
+UNiagaraHierarchyViewModelBase::~UNiagaraHierarchyViewModelBase()
+{
+	RefreshSourceViewDelegate.Unbind();
+	RefreshHierarchyWidgetDelegate.Unbind();
+	RefreshSectionsWidgetDelegate.Unbind();
+	IsItemSelectedDelegate.Unbind();
+	SelectObjectInDetailsPanelDelegate.Unbind();
+}
+
+void UNiagaraHierarchyViewModelBase::ForceFullRefresh()
+{
+	PrepareSourceItems();
+	HierarchyViewModelRoot->SyncToData();
+	RefreshSourceView(true);
+	RefreshHierarchyView(true);
+	RefreshSectionsWidget();
+}
+
+TSharedRef<SWidget> UNiagaraHierarchyViewModelBase::GenerateRowContentWidget(TSharedRef<FNiagaraHierarchyItemViewModelBase> ItemBaseViewModel) const
+{
+	if(UNiagaraHierarchyCategory* Category = Cast<UNiagaraHierarchyCategory>(ItemBaseViewModel->GetDataMutable()))
+	{
+		TSharedRef<FNiagaraHierarchyCategoryViewModel> TreeViewCategory = StaticCastSharedRef<FNiagaraHierarchyCategoryViewModel>(ItemBaseViewModel);
+		return SNew(SNiagaraHierarchyCategory, TreeViewCategory)
+		.IsSelected(FIsSelected::CreateLambda([this, ItemBaseViewModel]
+		{
+			return this->IsItemSelected(ItemBaseViewModel);
+		}));
+	}
+
+	return SNullWidget::NullWidget;
+}
+
+void UNiagaraHierarchyViewModelBase::RefreshSourceView(bool bFullRefresh) const
+{
+	if(RefreshSourceViewDelegate.IsBound())
+	{
+		RefreshSourceViewDelegate.Execute(bFullRefresh);
+	}
+}
+
+void UNiagaraHierarchyViewModelBase::RefreshHierarchyView(bool bFullRefresh) const
+{
+	if(RefreshHierarchyWidgetDelegate.IsBound())
+	{
+		RefreshHierarchyWidgetDelegate.Execute(bFullRefresh);
+	}
+}
+
+void UNiagaraHierarchyViewModelBase::RefreshSectionsWidget() const
+{
+	if(RefreshSectionsWidgetDelegate.IsBound())
+	{
+		RefreshSectionsWidgetDelegate.Execute();
+	}
+}
+
+void UNiagaraHierarchyViewModelBase::SelectItemInDetailsPanel(UObject* Object) const
+{
+	if(SelectObjectInDetailsPanelDelegate.IsBound())
+	{
+		SelectObjectInDetailsPanelDelegate.Execute(Object);
+	}
+}
+
+bool UNiagaraHierarchyViewModelBase::IsItemSelected(TSharedPtr<FNiagaraHierarchyItemViewModelBase> Item) const
+{
+	if(IsItemSelectedDelegate.IsBound())
+	{
+		return IsItemSelectedDelegate.Execute(Item);
+	}
+
+	return false;
+}
+
+void UNiagaraHierarchyViewModelBase::PostUndo(bool bSuccess)
+{
+	ForceFullRefresh();
+}
+
+void UNiagaraHierarchyViewModelBase::PostRedo(bool bSuccess)
+{
+	PostUndo(bSuccess);
+}
+
+void UNiagaraHierarchyViewModelBase::Initialize(TSharedPtr<FNiagaraSystemViewModel> InSystemViewModel)
+{
+	SystemViewModel = InSystemViewModel;
+	
+	if(GEditor)
+	{
+		GEditor->RegisterForUndo(this);
+	}
+	
+	HierarchyRoot = GetHierarchyDataRoot();
+	HierarchyRoot->SetFlags(RF_Transactional);
+
+	TArray<UNiagaraHierarchyItemBase*> AllItems;
+	HierarchyRoot->GetChildrenOfType<UNiagaraHierarchyItemBase>(AllItems, true);
+	for(UNiagaraHierarchyItemBase* Item : AllItems)
+	{
+		Item->SetFlags(RF_Transactional);
+	}
+
+	for(UNiagaraHierarchySection* Section : HierarchyRoot->GetSectionDataMutable())
+	{
+		Section->SetFlags(RF_Transactional);
+	}
+
+	SetupCommands();
+
+	HierarchyViewModelRoot = MakeShared<FNiagaraHierarchyRootViewModel>(HierarchyRoot.Get(), this);
+	HierarchyViewModelRoot->SyncToData();
+
+	SetActiveSection(nullptr);
+
+	PrepareSourceItems();
+
+	InitializeInternal();
+}
+
+void UNiagaraHierarchyViewModelBase::Finalize()
+{
+	if(GEditor)
+	{
+		GEditor->UnregisterForUndo(this);
+	}
+	
+	SourceViewModelRoot->Finalize();
+	SourceViewModelRoot.Reset();
+	HierarchyViewModelRoot.Reset();
+	SourceRoot = nullptr;
+	HierarchyRoot = nullptr;
+	
+	FinalizeInternal();
+}
+
+const TArray<TSharedPtr<FNiagaraHierarchyItemViewModelBase>>& UNiagaraHierarchyViewModelBase::GetSourceItems() const
+{
+	return SourceViewModelRoot->GetChildren();
+}
+
+const TArray<TSharedPtr<FNiagaraHierarchyItemViewModelBase>>& UNiagaraHierarchyViewModelBase::GetHierarchyItems() const
+{
+	return HierarchyViewModelRoot->GetFilteredChildren();
+}
+
+void UNiagaraHierarchyViewModelBase::OnGetChildren(TSharedPtr<FNiagaraHierarchyItemViewModelBase> Item, TArray<TSharedPtr<FNiagaraHierarchyItemViewModelBase>>& OutChildren) const
+{
+	OutChildren.Append(Item->GetFilteredChildren());
+}
+
+void UNiagaraHierarchyViewModelBase::SetActiveSection(TSharedPtr<FNiagaraHierarchySectionViewModel> Section)
+{
+	ActiveSection = Section;
+	if(Section.IsValid())
+	{
+		SelectItemInDetailsPanel(Section->GetDataMutable());
+	}
+	else
+	{
+		SelectItemInDetailsPanel(nullptr);
+	}
+	
+	RefreshHierarchyView(true);
+}
+
+TSharedPtr<FNiagaraHierarchySectionViewModel> UNiagaraHierarchyViewModelBase::GetActiveSection() const
+{
+	return ActiveSection.IsValid() ? ActiveSection.Pin() : nullptr;
+}
+
+UNiagaraHierarchySection* UNiagaraHierarchyViewModelBase::GetActiveSectionData() const
+{
+	return ActiveSection.IsValid() ? Cast<UNiagaraHierarchySection>(ActiveSection.Pin()->GetDataMutable()) : nullptr;
+}
+
+bool UNiagaraHierarchyViewModelBase::IsSectionActive(const UNiagaraHierarchySection* Section) const
+{
+	return ActiveSection == nullptr || ActiveSection.Pin()->GetData() == Section;
+}
+
+FString UNiagaraHierarchyViewModelBase::OnItemToStringDebug(TSharedPtr<FNiagaraHierarchyItemViewModelBase> ItemBaseViewModel) const
+{
+	return ItemBaseViewModel->ToString();	
+}
+
+void FNiagaraHierarchyItemViewModelBase::SyncToData()
+{
+	// this will refresh all underlying data recursively
+	ItemBase->Refresh();
+
+	// now that the data is refreshed, we can sync to the data by recycling view models & creating new ones
+	TArray<TSharedPtr<FNiagaraHierarchyItemViewModelBase>> NewChildren;
+	for(UNiagaraHierarchyItemBase* Child : ItemBase->GetChildren())
+	{		
+		int32 ViewModelIndex = FindIndexOfChild(Child);
+		// if we couldn't find a view model for a data child, we create it here
+		if(ViewModelIndex == INDEX_NONE)
+		{
+			if(UNiagaraHierarchyItem* Item = Cast<UNiagaraHierarchyItem>(Child))
+			{
+				TSharedPtr<FNiagaraHierarchyItemViewModel> NewItem = MakeShared<FNiagaraHierarchyItemViewModel>(Item, AsShared(), HierarchyViewModel);
+				NewItem->SyncToData();
+				NewChildren.Add(NewItem);
+			}
+			else if(UNiagaraHierarchyCategory* Category = Cast<UNiagaraHierarchyCategory>(Child))
+			{
+				TSharedPtr<FNiagaraHierarchyCategoryViewModel> NewCategory = MakeShared<FNiagaraHierarchyCategoryViewModel>(Category, AsShared(), HierarchyViewModel);
+				NewCategory->SyncToData();
+				NewChildren.Add(NewCategory);
+			}
+		}
+		// if we could find it, we refresh its contained view models and readd it
+		else
+		{
+			Children[ViewModelIndex]->SyncToData();
+			NewChildren.Add(Children[ViewModelIndex]);
+		}
+	}
+
+	/** Give the view models a chance to further customize the children sync process. */
+	SyncChildrenViewModelsInternal();
+	
+	Children.Empty();
+	Children.Append(NewChildren);
+
+	// first we sort the data. Categories before items.
+	GetDataMutable()->GetChildrenMutable().StableSort([](const UNiagaraHierarchyItemBase& ItemA, const UNiagaraHierarchyItemBase& ItemB)
+	{
+		return ItemA.IsA<UNiagaraHierarchyCategory>() && ItemB.IsA<UNiagaraHierarchyItem>();
+	});
+
+	// then we sort the view models according to the data order as this is what will determine widget order created from the view models
+	Children.Sort([this](const TSharedPtr<FNiagaraHierarchyItemViewModelBase>& ItemA, const TSharedPtr<FNiagaraHierarchyItemViewModelBase>& ItemB)
+		{
+			return FindIndexOfDataChild(ItemA) < FindIndexOfDataChild(ItemB);
+		});
+	
+	// we refresh the filtered children here as well
+	GetFilteredChildren();
+
+	OnSyncedDelegate.Broadcast();
+}
+
+const TArray<TSharedPtr<FNiagaraHierarchyItemViewModelBase>>& FNiagaraHierarchyItemViewModelBase::GetFilteredChildren() const
+{
+	FilteredChildren.Empty();
+
+	for(TSharedPtr<FNiagaraHierarchyItemViewModelBase> Child : Children)
+	{
+		if(const UNiagaraHierarchyCategory* Category = Cast<UNiagaraHierarchyCategory>(Child->GetData()))
+		{
+			if(StaticCastSharedPtr<FNiagaraHierarchyCategoryViewModel>(Child)->IsTopCategoryActive()) //HierarchyViewModel->IsSectionActive(Category->GetSection()))
+			{
+				FilteredChildren.Add(Child);
+			}
+		}
+		else
+		{
+			FilteredChildren.Add(Child);
+		}
+	}
+
+	return FilteredChildren;
+}
+
+bool FNiagaraHierarchyItemViewModelBase::HasParent(TSharedPtr<FNiagaraHierarchyItemViewModelBase> ParentCandidate, bool bRecursive)
+{
+	if(Parent.IsValid())
+	{
+		if(Parent == ParentCandidate)
+		{
+			return true;
+		}
+		else if(bRecursive)
+		{
+			Parent.Pin()->HasParent(ParentCandidate, bRecursive);
+		}
+	}
+
+	return false;
+}
+
+void FNiagaraHierarchyItemViewModelBase::AddChild(TSharedPtr<FNiagaraHierarchyItemViewModelBase> Item)
+{
+	if(Item->Parent.IsValid())
+	{
+		ensure(Item->Parent.Pin() == AsShared());
+	}
+	
+	Children.Add(Item);
+}
+
+UNiagaraHierarchyItemBase* FNiagaraHierarchyItemViewModelBase::AddNewItem(TSubclassOf<UNiagaraHierarchyItemBase> NewItemClass)
+{
+	FText TransactionText = FText::FormatOrdered(LOCTEXT("Transaction_AddedItem", "Added new {0} to hierarchy"), FText::FromString(NewItemClass->GetName()));
+	FScopedTransaction Transaction(TransactionText);
+	HierarchyViewModel->GetHierarchyDataRoot()->Modify();
+	
+	UNiagaraHierarchyItemBase* NewItem = NewObject<UNiagaraHierarchyItemBase>(GetDataMutable(), NewItemClass, NAME_None, RF_Transactional);
+	NewItem->Modify();
+	GetDataMutable()->GetChildrenMutable().Add(NewItem);
+	SyncToData();
+
+	return NewItem;
+}
+
+void FNiagaraHierarchyItemViewModelBase::DuplicateToThis(TSharedPtr<FNiagaraHierarchyItemViewModelBase> ItemToDuplicate, int32 InsertIndex)
+{
+	if(InsertIndex == INDEX_NONE)
+	{
+		GetDataMutable()->GetChildrenMutable().Add(Cast<UNiagaraHierarchyItemBase>(StaticDuplicateObject(ItemToDuplicate->GetData(), GetDataMutable())));
+	}
+	else
+	{
+		GetDataMutable()->GetChildrenMutable().Insert(Cast<UNiagaraHierarchyItemBase>(StaticDuplicateObject(ItemToDuplicate->GetData(), GetDataMutable())), InsertIndex);
+	}
+	
+	SyncToData();
+}
+
+void FNiagaraHierarchyItemViewModelBase::ReparentToThis(TSharedPtr<FNiagaraHierarchyItemViewModelBase> ItemToMove, int32 InsertIndex)
+{
+	if(InsertIndex == INDEX_NONE)
+	{
+		GetDataMutable()->GetChildrenMutable().Add(Cast<UNiagaraHierarchyItemBase>(StaticDuplicateObject(ItemToMove->GetDataMutable(), GetDataMutable())));
+	}
+	else
+	{
+		GetDataMutable()->GetChildrenMutable().Insert(Cast<UNiagaraHierarchyItemBase>(StaticDuplicateObject(ItemToMove->GetData(), GetDataMutable())), InsertIndex);
+	}
+	
+	ItemToMove->Finalize();
+	ItemToMove->Parent.Pin()->SyncToData();
+	SyncToData();
+}
+
+int32 FNiagaraHierarchyItemViewModelBase::FindIndexOfChild(TSharedPtr<FNiagaraHierarchyItemViewModelBase> Child) const
+{
+	return Children.Find(Child);
+}
+
+int32 FNiagaraHierarchyItemViewModelBase::FindIndexOfChild(UNiagaraHierarchyItemBase* Child) const
+{
+	return Children.FindLastByPredicate([Child](TSharedPtr<FNiagaraHierarchyItemViewModelBase> Item)
+	{
+		return Item->GetData() == Child;
+	});
+}
+
+int32 FNiagaraHierarchyItemViewModelBase::FindIndexOfDataChild(TSharedPtr<FNiagaraHierarchyItemViewModelBase> Child) const
+{
+	return GetData()->GetChildren().Find(Child->GetDataMutable());
+}
+
+int32 FNiagaraHierarchyItemViewModelBase::FindIndexOfDataChild(UNiagaraHierarchyItemBase* Child) const
+{
+	return GetData()->GetChildren().Find(Child);
+}
+
+void FNiagaraHierarchyItemViewModelBase::Finalize()
+{	
+	for(TSharedPtr<FNiagaraHierarchyItemViewModelBase> ChildViewModel : Children)
+	{
+		ChildViewModel->Finalize();
+	}	
+
+	GetDataMutable()->Finalize();
+	FinalizeInternal();
+	SyncToData();
+}
+
+TOptional<EItemDropZone> FNiagaraHierarchyItemViewModelBase::OnCanAcceptDrop(const FDragDropEvent& DragDropEvent, EItemDropZone ItemDropZone, TSharedPtr<FNiagaraHierarchyItemViewModelBase> Item)
+{
+	if(TSharedPtr<FNiagaraHierarchyDragDropOp> HierarchyDragDropOp = DragDropEvent.GetOperationAs<FNiagaraHierarchyDragDropOp>())
+	{
+		return HierarchyViewModel->CanDropOn(HierarchyDragDropOp->GetDraggedItem().Pin().ToSharedRef(), AsShared(), ItemDropZone);
+	}
+
+	return TOptional<EItemDropZone>();	
+}
+
+TSharedRef<FNiagaraHierarchyDragDropOp> FNiagaraHierarchyItemViewModelBase::CreateDragDropOp()
+{
+	check(CanDrag() == true);
+		
+	return HierarchyViewModel->CreateDragDropOp(AsShared());
+}
+
+FReply FNiagaraHierarchyItemViewModelBase::OnDragDetected(const FGeometry& Geometry, const FPointerEvent& PointerEvent,	bool bIsSource)
+{
+	if(CanDrag())
+	{
+		// if the drag is coming from source, we check if any of the hierarchy data already contains that item and we don't start a drag drop in that case
+		if(bIsSource)
+		{
+			TArray<UNiagaraHierarchyItemBase*> AllItems;
+			GetHierarchyViewModel()->GetHierarchyDataRoot()->GetChildrenOfType<UNiagaraHierarchyItemBase>(AllItems, true);
+
+			for(UNiagaraHierarchyItemBase* Item : AllItems)
+			{
+				if(GetData()->GetPersistentIdentity() == Item->GetPersistentIdentity())
+				{
+					return FReply::Unhandled();
+				}
+			}
+		}
+		
+		TSharedRef<FNiagaraHierarchyDragDropOp> HierarchyDragDropOp = CreateDragDropOp();
+		HierarchyDragDropOp->SetFromSourceList(bIsSource);
+
+		return FReply::Handled().BeginDragDrop(HierarchyDragDropOp);			
+	}
+		
+	return FReply::Unhandled();
+}
+
+TSharedPtr<FNiagaraHierarchySectionViewModel> FNiagaraHierarchyRootViewModel::AddNewSection()
+{
+	FScopedTransaction ScopedTransaction(LOCTEXT("NewSectionAdded","Added Section"));
+	HierarchyViewModel->GetHierarchyDataRoot()->Modify();
+	
+	UNiagaraHierarchySection* SectionData = Cast<UNiagaraHierarchyRoot>(ItemBase)->AddSection(LOCTEXT("NiagaraHierarchyEditorDefaultNewSectionName", "New Section"));
+	SectionData->Modify();
+	TSharedPtr<FNiagaraHierarchySectionViewModel> NewSectionViewModel = MakeShared<FNiagaraHierarchySectionViewModel>(SectionData, StaticCastSharedRef<FNiagaraHierarchyRootViewModel>(AsShared()), HierarchyViewModel);
+	SectionViewModels.Add(NewSectionViewModel);
+	SyncToData();
+	HierarchyViewModel->SetActiveSection(NewSectionViewModel);
+	NewSectionViewModel->RequestRename();
+
+	return NewSectionViewModel;
+}
+
+void FNiagaraHierarchyRootViewModel::SyncChildrenViewModelsInternal()
+{
+	const UNiagaraHierarchyRoot* RootData = Cast<UNiagaraHierarchyRoot>(GetData());
+
+	TArray<TSharedPtr<FNiagaraHierarchySectionViewModel>> SectionsToRemove;
+	for(TSharedPtr<FNiagaraHierarchySectionViewModel> SectionViewModel : SectionViewModels)
+	{
+		if(SectionViewModel->GetDataMutable()->IsFinalized() || !RootData->GetSectionData().Contains(SectionViewModel->GetData()))
+		{
+			SectionsToRemove.Add(SectionViewModel);
+		}
+	}
+
+	for (TSharedPtr<FNiagaraHierarchySectionViewModel> SectionViewModel : SectionsToRemove)
+	{
+		SectionViewModel->Finalize();
+		SectionViewModels.Remove(SectionViewModel);
+	}
+	
+	for(UNiagaraHierarchySection* Section : RootData->GetSectionData())
+	{
+		if(Section->IsFinalized())
+		{
+			continue;
+		}
+		
+		if(!SectionViewModels.ContainsByPredicate([Section](TSharedPtr<FNiagaraHierarchySectionViewModel> SectionViewModel)
+		{
+			return SectionViewModel->GetData() == Section;
+		}))
+		{
+			TSharedPtr<FNiagaraHierarchySectionViewModel> NewSectionViewModel = MakeShared<FNiagaraHierarchySectionViewModel>(Section, StaticCastSharedRef<FNiagaraHierarchyRootViewModel>(AsShared()), HierarchyViewModel); 
+			NewSectionViewModel->SyncToData();
+			SectionViewModels.Add(NewSectionViewModel);
+		}
+	}
+
+	SectionViewModels.Sort([this](const TSharedPtr<FNiagaraHierarchySectionViewModel>& ItemA, const TSharedPtr<FNiagaraHierarchySectionViewModel>& ItemB)
+		{
+			return Cast<UNiagaraHierarchyRoot>(GetDataMutable())->GetSectionData().Find(Cast<UNiagaraHierarchySection>(ItemA->GetDataMutable()))
+						<
+					Cast<UNiagaraHierarchyRoot>(GetDataMutable())->GetSectionData().Find(Cast<UNiagaraHierarchySection>(ItemB->GetDataMutable())); 
+		});
+}
+
+FReply FNiagaraHierarchySectionViewModel::OnDroppedOn(const FDragDropEvent& DragDropEvent, EItemDropZone ItemDropZone, TSharedPtr<FNiagaraHierarchyItemViewModelBase> Item)
+{
+	if(DragDropEvent.GetOperation()->IsOfType<FNiagaraSectionDragDropOp>())
+	{
+		TSharedPtr<FNiagaraSectionDragDropOp> SectionDragDropOp = DragDropEvent.GetOperationAs<FNiagaraSectionDragDropOp>();
+		TWeakPtr<FNiagaraHierarchySectionViewModel> DraggedItem = SectionDragDropOp->GetDraggedSection();
+		UNiagaraHierarchySection* DraggedSectionData = Cast<UNiagaraHierarchySection>(DraggedItem.Pin()->GetDataMutable());
+
+		int32 IndexOfThis = HierarchyViewModel->GetHierarchyDataRoot()->GetSectionData().Find(Cast<UNiagaraHierarchySection>(GetDataMutable()));
+		int32 DraggedSectionIndex = HierarchyViewModel->GetHierarchyDataRoot()->GetSectionData().Find(Cast<UNiagaraHierarchySection>(DraggedItem.Pin()->GetDataMutable()));
+
+		TArray<UNiagaraHierarchySection*>& SectionData = HierarchyViewModel->GetHierarchyDataRoot()->GetSectionDataMutable();
+		int32 Count = SectionData.Num();
+		
+		// above constitutes to the left here
+		if(ItemDropZone == EItemDropZone::AboveItem)
+		{
+			SectionData.RemoveAt(DraggedSectionIndex);
+			SectionData.Insert(DraggedSectionData, FMath::Max(IndexOfThis, 0));
+
+			HierarchyViewModel->ForceFullRefresh();
+			return FReply::Handled();
+		}
+		else if(ItemDropZone == EItemDropZone::BelowItem)
+		{
+			SectionData.RemoveAt(DraggedSectionIndex);
+
+			if(IndexOfThis + 1 > SectionData.Num())
+			{
+				SectionData.Add(DraggedSectionData);
+			}
+			else
+			{
+				SectionData.Insert(DraggedSectionData, FMath::Min(IndexOfThis+1, Count));
+			}
+
+			HierarchyViewModel->ForceFullRefresh();
+			return FReply::Handled();
+		}
+
+		return FReply::Unhandled();
+	}
+	else if(DragDropEvent.GetOperation()->IsOfType<FNiagaraHierarchyDragDropOp>())
+	{		
+		TSharedPtr<FNiagaraHierarchyDragDropOp> HierarchyDragDropOp = DragDropEvent.GetOperationAs<FNiagaraHierarchyDragDropOp>();
+		TWeakPtr<FNiagaraHierarchyItemViewModelBase> DraggedItem = HierarchyDragDropOp->GetDraggedItem();
+
+		if(UNiagaraHierarchyCategory* HierarchyCategory = Cast<UNiagaraHierarchyCategory>(DraggedItem.Pin()->GetDataMutable()))
+		{
+			FScopedTransaction Transaction(LOCTEXT("Transaction_OnSectionDrop", "Moved category to section"));
+			HierarchyViewModel->GetHierarchyDataRoot()->Modify();
+			
+			HierarchyCategory->SetSection(Cast<UNiagaraHierarchySection>(GetDataMutable()));
+
+			// we null out any sections for all contained categories
+			TArray<UNiagaraHierarchyCategory*> AllChildCategories;
+			HierarchyCategory->GetChildrenOfType<UNiagaraHierarchyCategory>(AllChildCategories, true);
+			for(UNiagaraHierarchyCategory* ChildCategory : AllChildCategories)
+			{
+				ChildCategory->SetSection(nullptr);
+			}
+
+			// we only need to reparent if the parent isn't already the root. This stops unnecessary reordering
+			if(DraggedItem.Pin()->GetParent() != HierarchyViewModel->GetHierarchyViewModelRoot())
+			{
+				HierarchyViewModel->GetHierarchyViewModelRoot()->ReparentToThis(DraggedItem.Pin());
+			}
+			
+			HierarchyViewModel->RefreshHierarchyView();
+			
+			return FReply::Handled();
+		}
+	}
+
+	return FReply::Unhandled();
+}
+
+void FNiagaraHierarchySectionViewModel::FinalizeInternal()
+{
+	if(HierarchyViewModel->GetActiveSection() == AsShared())
+	{
+		HierarchyViewModel->SetActiveSection(nullptr);
+	}
+
+	// we make sure to reset all categories' section entry that were referencing this section
+	TArray<UNiagaraHierarchyCategory*> AllCategories;
+	HierarchyViewModel->GetHierarchyDataRoot()->GetChildrenOfType<UNiagaraHierarchyCategory>(AllCategories, true);
+
+	for(UNiagaraHierarchyCategory* Category : AllCategories)
+	{
+		if(Category->GetSection() == GetData())
+		{
+			Category->SetSection(nullptr);
+		}
+	}
+}
+
+FReply FNiagaraHierarchyItemViewModel::OnDroppedOn(const FDragDropEvent& DragDropEvent, EItemDropZone ItemDropZone, TSharedPtr<FNiagaraHierarchyItemViewModelBase> Item)
+{
+	FScopedTransaction Transaction(LOCTEXT("Transaction_MovedItem", "Moved an item in the hierarchy"));
+	HierarchyViewModel->GetHierarchyDataRoot()->Modify();
+	
+	if(DragDropEvent.GetOperation()->IsOfType<FNiagaraHierarchyDragDropOp>())
+	{
+		TSharedPtr<FNiagaraHierarchyDragDropOp> HierarchyDragDropOp = DragDropEvent.GetOperationAs<FNiagaraHierarchyDragDropOp>();
+		if(ItemDropZone == EItemDropZone::AboveItem)
+		{
+			int32 IndexOfThis = Parent.Pin()->FindIndexOfDataChild(AsShared());
+
+			if(HierarchyDragDropOp->GetIsFromSourceList())
+			{
+				Parent.Pin()->DuplicateToThis(HierarchyDragDropOp->GetDraggedItem().Pin(), FMath::Max(IndexOfThis, 0));
+			}
+			else
+			{
+				Parent.Pin()->ReparentToThis(HierarchyDragDropOp->GetDraggedItem().Pin(), FMath::Max(IndexOfThis, 0));
+			}
+			
+			HierarchyViewModel->RefreshHierarchyView();
+			HierarchyViewModel->RefreshSourceView();
+			return FReply::Handled();
+		}
+		else if(ItemDropZone == EItemDropZone::BelowItem)
+		{
+			int32 IndexOfThis = Parent.Pin()->FindIndexOfDataChild(AsShared());
+			Parent.Pin()->ReparentToThis(HierarchyDragDropOp->GetDraggedItem().Pin(), FMath::Min(IndexOfThis+1, Parent.Pin()->GetChildren().Num()));
+
+			HierarchyViewModel->RefreshHierarchyView();
+			HierarchyViewModel->RefreshSourceView();
+			return FReply::Handled();
+		}
+	}
+
+	return FReply::Unhandled();
+}
+
+FReply FNiagaraHierarchyCategoryViewModel::OnDroppedOn(const FDragDropEvent& DragDropEvent, EItemDropZone ItemDropZone, TSharedPtr<FNiagaraHierarchyItemViewModelBase> Item)
+{
+	if(DragDropEvent.GetOperation()->IsOfType<FNiagaraHierarchyDragDropOp>())
+	{
+		FScopedTransaction Transaction(LOCTEXT("Transaction_OnCategoryDrop", "Dropped item on/above/below category"));
+		HierarchyViewModel->GetHierarchyDataRoot()->Modify();
+		
+		TSharedPtr<FNiagaraHierarchyDragDropOp> HierarchyDragDropOp = DragDropEvent.GetOperationAs<FNiagaraHierarchyDragDropOp>();
+		TSharedPtr<FNiagaraHierarchyItemViewModelBase> DraggedItem = HierarchyDragDropOp->GetDraggedItem().Pin();
+		
+		if(ItemDropZone == EItemDropZone::OntoItem)
+		{
+			// if we are dragging a category onto another category, we null out its section
+			if(UNiagaraHierarchyCategory* Category = Cast<UNiagaraHierarchyCategory>(DraggedItem->GetDataMutable()))
+			{
+				Category->SetSection(nullptr);
+
+				// we null out any sections for all contained categories
+				TArray<UNiagaraHierarchyCategory*> AllChildCategories;
+				Category->GetChildrenOfType<UNiagaraHierarchyCategory>(AllChildCategories, true);
+				for(UNiagaraHierarchyCategory* ChildCategory : AllChildCategories)
+				{
+					ChildCategory->SetSection(nullptr);
+				}
+			}
+
+			if(HierarchyDragDropOp->GetIsFromSourceList())
+			{
+				DuplicateToThis(HierarchyDragDropOp->GetDraggedItem().Pin());
+			}
+			else
+			{
+				ReparentToThis(HierarchyDragDropOp->GetDraggedItem().Pin());
+			}
+		}
+		else if(ItemDropZone == EItemDropZone::AboveItem)
+		{
+			int32 IndexOfThis = Parent.Pin()->FindIndexOfDataChild(AsShared());
+			if(HierarchyDragDropOp->GetIsFromSourceList())
+			{
+				Parent.Pin()->DuplicateToThis(HierarchyDragDropOp->GetDraggedItem().Pin(), FMath::Max(IndexOfThis, 0));
+			}
+			else
+			{
+				Parent.Pin()->ReparentToThis(HierarchyDragDropOp->GetDraggedItem().Pin(), FMath::Max(IndexOfThis, 0));
+			}
+		}
+		else if(ItemDropZone == EItemDropZone::BelowItem)
+		{
+			int32 IndexOfThis = Parent.Pin()->FindIndexOfDataChild(AsShared());
+			if(HierarchyDragDropOp->GetIsFromSourceList())
+			{
+				Parent.Pin()->DuplicateToThis(HierarchyDragDropOp->GetDraggedItem().Pin(), FMath::Min(IndexOfThis+1, Parent.Pin()->GetChildren().Num()));
+			}
+			else
+			{
+				Parent.Pin()->ReparentToThis(HierarchyDragDropOp->GetDraggedItem().Pin(), FMath::Min(IndexOfThis+1, Parent.Pin()->GetChildren().Num()));
+			}
+		}
+		
+		// todo (me)
+		// refreshing both views whenever we drag something seems wasteful
+		HierarchyViewModel->RefreshHierarchyView();
+		HierarchyViewModel->RefreshSourceView();
+		return FReply::Handled();
+	}
+
+	return FReply::Unhandled();
+}
+
+void UNiagaraHierarchyViewModelBase::AddCategory() const
+{
+	UNiagaraHierarchyCategory* Category = Cast<UNiagaraHierarchyCategory>(HierarchyViewModelRoot->AddNewItem(UNiagaraHierarchyCategory::StaticClass()));
+	TArray<UNiagaraHierarchyCategory*> SiblingCategories;
+	Category->GetTypedOuter<UNiagaraHierarchyItemBase>()->GetChildrenOfType<UNiagaraHierarchyCategory>(SiblingCategories);
+	
+	TSet<FName> CategoryNames;
+	for(const auto& SiblingCategory : SiblingCategories)
+	{
+		CategoryNames.Add(SiblingCategory->GetCategoryName());
+	}
+
+	Category->SetCategoryName(FNiagaraUtilities::GetUniqueName(FName("New Category"), CategoryNames));
+	
+	// we only set the section property if the current section isn't set to "All"
+	Category->SetSection(GetActiveSectionData());
+	
+	RefreshHierarchyView();
+}
+
+void UNiagaraHierarchyViewModelBase::AddSection() const
+{
+	HierarchyViewModelRoot->AddNewSection();
+}
+
+TSharedPtr<SWidget> FNiagaraHierarchyDragDropOp::GetDefaultDecorator() const
+{
+	TSharedRef<SWidget> CustomDecorator = CreateCustomDecorator();
+
+	SVerticalBox::FSlot* CustomSlot;
+	TSharedPtr<SWidget> Decorator = SNew(SToolTip)
+	[
+		SNew(SVerticalBox)
+		+ SVerticalBox::Slot()
+		.Expose(CustomSlot).
+		AutoHeight()
+		+ SVerticalBox::Slot()
+		.AutoHeight()
+		.Padding(2.f)
+		[
+			SNew(STextBlock)
+			.Text(this, &FNiagaraHierarchyDragDropOp::GetAdditionalLabel)
+			.TextStyle(&FAppStyle::GetWidgetStyle<FTextBlockStyle>("NormalText.Important"))
+			.Visibility_Lambda([this]()
+			{
+				return GetAdditionalLabel().IsEmpty() ? EVisibility::Collapsed : EVisibility::Visible;
+			})
+		]
+		+ SVerticalBox::Slot()
+		.AutoHeight()
+		.Padding(2.f)
+		[
+			SNew(STextBlock)
+			.Text(this, &FNiagaraHierarchyDragDropOp::GetDescription)
+			.TextStyle(&FAppStyle::GetWidgetStyle<FTextBlockStyle>("NormalText"))
+			.Visibility_Lambda([this]()
+			{
+				return GetDescription().IsEmpty() ? EVisibility::Collapsed : EVisibility::Visible;
+			})
+		]
+	];
+
+	if(CustomDecorator != SNullWidget::NullWidget)
+	{
+		CustomSlot->AttachWidget(CustomDecorator);
+	}
+
+	return Decorator;
+}
+
+TSharedRef<SWidget> FNiagaraSectionDragDropOp::CreateCustomDecorator() const
+{
+	return SNew(SCheckBox)
+		.Visibility(EVisibility::HitTestInvisible)
+		.Style(FAppStyle::Get(), "DetailsView.SectionButton")
+		.IsChecked(ECheckBoxState::Unchecked)
+		[
+			SNew(SInlineEditableTextBlock)
+			.Text(GetDraggedSection().Pin()->GetSectionNameAsText())
+		];
+}
+
+#undef LOCTEXT_NAMESPACE
diff --git a/Engine/Plugins/FX/Niagara/Source/NiagaraEditor/Private/ViewModels/HierarchyEditor/NiagaraUserParametersHierarchyViewModel.cpp b/Engine/Plugins/FX/Niagara/Source/NiagaraEditor/Private/ViewModels/HierarchyEditor/NiagaraUserParametersHierarchyViewModel.cpp
new file mode 100644
--- /dev/null
+++ b/Engine/Plugins/FX/Niagara/Source/NiagaraEditor/Private/ViewModels/HierarchyEditor/NiagaraUserParametersHierarchyViewModel.cpp
@@ -0,0 +1,193 @@
+// Copyright Epic Games, Inc. All Rights Reserved.
+
+#include "ViewModels/HierarchyEditor/NiagaraUserParametersHierarchyViewModel.h"
+#include "NiagaraEditorUtilities.h"
+#include "NiagaraEmitterEditorData.h"
+#include "NiagaraNodeFunctionCall.h"
+#include "NiagaraScriptVariable.h"
+#include "NiagaraSystem.h"
+#include "NiagaraSystemEditorData.h"
+#include "SDropTarget.h"
+#include "Customizations/NiagaraScriptVariableCustomization.h"
+#include "ViewModels/NiagaraSystemViewModel.h"
+#include "ViewModels/Stack/NiagaraStackViewModel.h"
+#include "Widgets/SNiagaraParameterName.h"
+#include "Widgets/Layout/SWrapBox.h"
+#include "Widgets/Views/STableRow.h"
+#include "Widgets/SToolTip.h"
+
+#define LOCTEXT_NAMESPACE "NiagaraUserParametersHierarchyEditor"
+
+void UNiagaraHierarchyUserParameter::Initialize(UNiagaraScriptVariable& InUserParameterScriptVariable, UNiagaraSystem& InSystem)
+{
+	UserParameterScriptVariable = &InUserParameterScriptVariable;
+	SetGuid(InUserParameterScriptVariable.Metadata.GetVariableGuid());
+	System = &InSystem;
+}
+
+void UNiagaraHierarchyUserParameter::RefreshDataInternal()
+{	
+	if(System == nullptr || FNiagaraEditorUtilities::FindScriptVariableForUserParameter(GetPersistentIdentity(), *System) == nullptr)
+	{
+		Finalize();
+	}
+}
+
+void UNiagaraUserParametersHierarchyViewModel::InitializeInternal()
+{
+	UNiagaraSystemEditorData* SystemEditorData = Cast<UNiagaraSystemEditorData>(GetSystemViewModel()->GetSystem().GetEditorData());
+	SystemEditorData->OnUserParameterScriptVariablesSynced().AddUObject(this, &UNiagaraUserParametersHierarchyViewModel::ForceFullRefresh);
+}
+
+void UNiagaraUserParametersHierarchyViewModel::FinalizeInternal()
+{
+	if(SystemViewModel.IsValid())
+	{
+		GetSystemViewModel()->GetSystem().GetExposedParameters().RemoveAllOnChangedHandlers(this);
+		UNiagaraSystemEditorData* SystemEditorData = Cast<UNiagaraSystemEditorData>(GetSystemViewModel()->GetSystem().GetEditorData());
+		SystemEditorData->OnUserParameterScriptVariablesSynced().RemoveAll(this);
+	}
+}
+
+TSharedRef<SWidget> FNiagaraUserParameterDragDropOp::CreateCustomDecorator() const
+{
+	return FNiagaraParameterUtilities::GetParameterWidget(GetUserParameter(), true, false);
+}
+
+UNiagaraHierarchyRoot* UNiagaraUserParametersHierarchyViewModel::GetHierarchyDataRoot() const
+{
+	UNiagaraHierarchyRoot* RootItem = GetSystemViewModel()->GetEditorData().UserParameterHierarchy;
+
+	ensure(RootItem != nullptr);
+	return RootItem;
+}
+
+void UNiagaraUserParametersHierarchyViewModel::PrepareSourceItems()
+{	
+	TArray<FNiagaraVariable> UserParameters;
+	GetSystemViewModel()->GetSystem().GetExposedParameters().GetUserParameters(UserParameters);
+
+	TArray<UNiagaraHierarchyItemBase*> OldChildren = SourceRoot->GetChildrenMutable();
+	SourceRoot->GetChildrenMutable().Empty();
+
+	// we ensure we have one data child per user parameter, and we get rid of deleted ones
+	for(FNiagaraVariable& UserParameter : UserParameters)
+	{
+		FNiagaraVariable RedirectedUserParameter = UserParameter;
+		GetSystemViewModel()->GetSystem().GetExposedParameters().RedirectUserVariable(RedirectedUserParameter);
+		TObjectPtr<UNiagaraScriptVariable> ScriptVariable = FNiagaraEditorUtilities::GetScriptVariableForUserParameter(RedirectedUserParameter, GetSystemViewModel());
+		
+		if(UNiagaraHierarchyItemBase** ExistingChild = OldChildren.FindByPredicate([ScriptVariable](UNiagaraHierarchyItemBase* ItemBase)
+		{
+			return ItemBase->GetPersistentIdentity() == ScriptVariable->Metadata.GetVariableGuid();
+		}))
+		{
+			SourceRoot->GetChildrenMutable().Add(*ExistingChild);
+			continue;
+		}
+		
+		// since the source items are transient we need to create them here and keep them around until the end of the tool's lifetime
+		UNiagaraHierarchyUserParameter* UserParameterHierarchyObject = NewObject<UNiagaraHierarchyUserParameter>(SourceRoot);
+		UserParameterHierarchyObject->Initialize(*ScriptVariable.Get(), GetSystemViewModel()->GetSystem());
+		SourceRoot->GetChildrenMutable().Add(UserParameterHierarchyObject);
+	}
+
+	// after we synced the data, we sync their view models that we display in the UI
+	SourceViewModelRoot->SyncToData();
+}
+
+void UNiagaraUserParametersHierarchyViewModel::SetupCommands()
+{
+	// no custom commands yet
+}
+
+TOptional<EItemDropZone> UNiagaraUserParametersHierarchyViewModel::CanDropOn(TSharedPtr<FNiagaraHierarchyItemViewModelBase> SourceDropItem, TSharedPtr<FNiagaraHierarchyItemViewModelBase> TargetDropItem, EItemDropZone DropZone)
+{
+	bool bAllowDropOn =
+		SourceDropItem->GetData() != TargetDropItem->GetData() &&
+			(!SourceDropItem->HasParent(TargetDropItem, false) || DropZone != EItemDropZone::OntoItem)  &&
+			!TargetDropItem->HasParent(SourceDropItem, true)
+		&&
+		(
+			(SourceDropItem->GetData()->IsA<UNiagaraHierarchyUserParameter>() && TargetDropItem->GetData()->IsA<UNiagaraHierarchyCategory>() && DropZone == EItemDropZone::OntoItem)
+				||
+			(SourceDropItem->GetData()->IsA<UNiagaraHierarchyUserParameter>() && TargetDropItem->GetData()->IsA<UNiagaraHierarchyUserParameter>() && DropZone != EItemDropZone::OntoItem)
+				||
+			(SourceDropItem->GetData()->IsA<UNiagaraHierarchyCategory>() && TargetDropItem->GetData()->IsA<UNiagaraHierarchyCategory>()
+				)
+		);
+	
+	return bAllowDropOn ? DropZone : TOptional<EItemDropZone>();
+}
+
+TSharedRef<FNiagaraHierarchyDragDropOp> UNiagaraUserParametersHierarchyViewModel::CreateDragDropOp(TSharedRef<FNiagaraHierarchyItemViewModelBase> Item)
+{
+	if(const UNiagaraHierarchyUserParameter* UserParameter = Cast<UNiagaraHierarchyUserParameter>(Item->GetData()))
+	{
+		TSharedRef<FNiagaraUserParameterDragDropOp> ParameterDragDropOp = MakeShared<FNiagaraUserParameterDragDropOp>(Item);
+		ParameterDragDropOp->Construct();
+		return ParameterDragDropOp;
+	}
+
+	if(UNiagaraHierarchyCategory* HierarchyCategory = Cast<UNiagaraHierarchyCategory>(Item->GetDataMutable()))
+	{
+		TSharedRef<FNiagaraHierarchyDragDropOp> CategoryDragDropOp = MakeShared<FNiagaraHierarchyDragDropOp>(Item);
+		CategoryDragDropOp->SetAdditionalLabel(FText::FromName(HierarchyCategory->GetCategoryName()));
+		CategoryDragDropOp->Construct();
+		return CategoryDragDropOp;
+	}
+
+	return MakeShared<FNiagaraHierarchyDragDropOp>(nullptr);
+}
+
+void UNiagaraUserParametersHierarchyViewModel::OnSelectionChanged(TSharedPtr<FNiagaraHierarchyItemViewModelBase> HierarchyItem)
+{
+	if(!HierarchyItem.IsValid())
+	{
+		SelectObjectInDetailsPanelDelegate.Execute(nullptr);
+		return;
+	}
+
+	UObject* ObjectToSelect = nullptr;
+	if(HierarchyItem->GetData()->IsA<UNiagaraHierarchyUserParameter>())
+	{
+		UNiagaraHierarchyUserParameter* HierarchyUserParameter = Cast<UNiagaraHierarchyUserParameter>(HierarchyItem->GetDataMutable());
+		FNiagaraVariable ContainedVariable = HierarchyUserParameter->GetUserParameter();
+		GetSystemViewModel()->GetSystem().GetExposedParameters().RedirectUserVariable(ContainedVariable);
+		TObjectPtr<UNiagaraScriptVariable> ScriptVariable = FNiagaraEditorUtilities::GetScriptVariableForUserParameter(ContainedVariable, GetSystemViewModel());
+
+		if(ScriptVariable)
+		{
+			ObjectToSelect = ScriptVariable.Get();
+		}
+	}
+	else
+	{
+		ObjectToSelect = HierarchyItem->GetDataMutable();
+	}
+
+	SelectObjectInDetailsPanelDelegate.Execute(ObjectToSelect);
+}
+
+TArray<TTuple<UClass*, FOnGetDetailCustomizationInstance>> UNiagaraUserParametersHierarchyViewModel::GetInstanceCustomizations()
+{
+	return {{UNiagaraScriptVariable::StaticClass(), FOnGetDetailCustomizationInstance::CreateStatic(&FNiagaraScriptVariableHierarchyDetails::MakeInstance)}};
+}
+
+TSharedRef<SWidget> UNiagaraUserParametersHierarchyViewModel::GenerateRowContentWidget(TSharedRef<FNiagaraHierarchyItemViewModelBase> ItemViewModelBase) const
+{
+	TSharedRef<SWidget> RowWidget = Super::GenerateRowContentWidget(ItemViewModelBase);
+
+	if(RowWidget == SNullWidget::NullWidget)
+	{
+		if(const UNiagaraHierarchyUserParameter* UserParameter = Cast<UNiagaraHierarchyUserParameter>(ItemViewModelBase->GetData()))
+		{
+			RowWidget = FNiagaraParameterUtilities::GetParameterWidget(UserParameter->GetUserParameter(), true, false);
+		}
+	}	
+
+	check(RowWidget != SNullWidget::NullWidget);
+	return RowWidget;
+}
+
+#undef LOCTEXT_NAMESPACE
diff --git a/Engine/Plugins/FX/Niagara/Source/NiagaraEditor/Private/ViewModels/NiagaraParameterPanelViewModel.cpp b/Engine/Plugins/FX/Niagara/Source/NiagaraEditor/Private/ViewModels/NiagaraParameterPanelViewModel.cpp
--- a/Engine/Plugins/FX/Niagara/Source/NiagaraEditor/Private/ViewModels/NiagaraParameterPanelViewModel.cpp
+++ b/Engine/Plugins/FX/Niagara/Source/NiagaraEditor/Private/ViewModels/NiagaraParameterPanelViewModel.cpp
@@ -23,6 +23,7 @@
 #include "NiagaraSimulationStageBase.h"
 #include "NiagaraSystem.h"
 #include "NiagaraSystemScriptViewModel.h"
+#include "NiagaraSystemEditorData.h"
 #include "NiagaraTypes.h"
 #include "ViewModels/NiagaraEmitterHandleViewModel.h"
 #include "ViewModels/NiagaraScratchPadScriptViewModel.h"
@@ -33,9 +34,11 @@
 #include "ViewModels/NiagaraSystemViewModel.h"
 #include "ViewModels/Stack/NiagaraStackGraphUtilities.h"
 #include "ViewModels/Stack/NiagaraStackViewModel.h"
+#include "ViewModels/HierarchyEditor/NiagaraUserParametersHierarchyViewModel.h"
 #include "Widgets/SNiagaraDebugger.h"
 #include "Widgets/SNiagaraParameterMenu.h"
 #include "Widgets/SNiagaraParameterPanel.h"
+#include "Widgets/Input/SButton.h"
 
 #include "Framework/Commands/GenericCommands.h"
 #include "Framework/MultiBox/MultiBoxBuilder.h"
@@ -43,7 +46,6 @@
 #include "ScopedTransaction.h"
 #include "Widgets/SNullWidget.h"
 
-
 #define LOCTEXT_NAMESPACE "NiagaraParameterPanelViewModel"
 
 template<> TMap<UNiagaraScript*, TArray<FNiagaraScriptToolkitParameterPanelViewModel*>> TNiagaraViewModelManager<UNiagaraScript, FNiagaraScriptToolkitParameterPanelViewModel>::ObjectsToViewModels{};
@@ -1467,6 +1469,28 @@ TSharedPtr<SWidget> FNiagaraSystemToolkitParameterPanelViewModel::CreateContextM
 	return SNullWidget::NullWidget;
 }
 
+TSharedRef<SWidget> FNiagaraSystemToolkitParameterPanelViewModel::GenerateAdjacentWidget()
+{
+	if(SystemViewModel->GetEditMode() == ENiagaraSystemViewModelEditMode::SystemAsset)
+	{
+		return SNew(SButton)
+		.Visibility_Lambda([this]()
+		{
+			return IsUserSectionActive() ? EVisibility::Visible : EVisibility::Collapsed;
+		})
+		.OnClicked(FOnClicked::CreateLambda([this]()
+		{
+			SummonHierarchyView();
+			return FReply::Handled();
+		}))
+		[
+			SNew(STextBlock).Text(FText::FromString("Edit Hierarchy"))
+		];
+	}
+
+	return SNullWidget::NullWidget;
+}
+
 FNiagaraParameterUtilities::EParameterContext FNiagaraSystemToolkitParameterPanelViewModel::GetParameterContext() const
 {
 	return FNiagaraParameterUtilities::EParameterContext::System;
@@ -1777,6 +1801,11 @@ TSharedRef<SWidget> FNiagaraSystemToolkitParameterPanelViewModel::CreateAddParam
 	return MenuWidget;
 }
 
+void FNiagaraSystemToolkitParameterPanelViewModel::SummonHierarchyView() const
+{
+	SystemViewModel->FocusTab("NiagaraSystemEditor_UserParameters");
+}
+
 TArray<FNiagaraVariable> FNiagaraSystemToolkitParameterPanelViewModel::GetEditableStaticSwitchParameters() const
 {
 	TArray<FNiagaraVariable> OutStaticSwitchParameters;
@@ -1930,19 +1959,7 @@ TArray<FNiagaraParameterPanelItem> FNiagaraSystemToolkitParameterPanelViewModel:
 		ParamStore->GetParameters(Vars);
 		for (const FNiagaraVariable& Var : Vars)
 		{
-			UNiagaraScriptVariable* ScriptVar;
-			if (TObjectPtr<UNiagaraScriptVariable> const* ScriptVarPtr = ParameterToScriptVariableMap.Find(Var))
-			{
-				ScriptVar = *ScriptVarPtr;
-			}
-			else
-			{
-				// Create a new UNiagaraScriptVariable to represent this parameter for the lifetime of the ParameterPanelViewModel.
-				ScriptVar = NewObject<UNiagaraScriptVariable>(GetTransientPackage());
-				ScriptVar->AddToRoot();
-				ScriptVar->Init(Var, FNiagaraVariableMetaData());
-				TransientParameterToScriptVarMap.Add(Var, ScriptVar);
-			}
+			UNiagaraScriptVariable* ScriptVar = FNiagaraEditorUtilities::GetScriptVariableForUserParameter(Var, SystemViewModel);
 
 			FNiagaraParameterPanelItem Item = FNiagaraParameterPanelItem();
 			Item.ScriptVariable = ScriptVar;
@@ -2420,7 +2437,7 @@ void FNiagaraSystemToolkitParameterPanelViewModel::OnParameterRemovedExternally(
 	// User and System need to be checked for all graphs as they could be used anywhere.
 	else if (InOldVar.IsInNameSpace(FNiagaraConstants::UserNamespaceString) ||
 		InOldVar.IsInNameSpace(FNiagaraConstants::SystemNamespaceString))
-	{
+	{		
 		TArray<UNiagaraGraph*> Graphs;
 		Graphs.Add(SystemScriptGraph.Get());
 
diff --git a/Engine/Plugins/FX/Niagara/Source/NiagaraEditor/Private/ViewModels/NiagaraSystemViewModel.cpp b/Engine/Plugins/FX/Niagara/Source/NiagaraEditor/Private/ViewModels/NiagaraSystemViewModel.cpp
--- a/Engine/Plugins/FX/Niagara/Source/NiagaraEditor/Private/ViewModels/NiagaraSystemViewModel.cpp
+++ b/Engine/Plugins/FX/Niagara/Source/NiagaraEditor/Private/ViewModels/NiagaraSystemViewModel.cpp
@@ -61,6 +61,7 @@
 #include "ViewModels/Stack/NiagaraStackGraphUtilities.h"
 #include "ViewModels/Stack/NiagaraStackViewModel.h"
 #include "ViewModels/TNiagaraViewModelManager.h"
+#include "ViewModels/HierarchyEditor/NiagaraUserParametersHierarchyViewModel.h"
 
 
 DECLARE_CYCLE_STAT(TEXT("Niagara - SystemViewModel - CompileSystem"), STAT_NiagaraEditor_SystemViewModel_CompileSystem, STATGROUP_NiagaraEditor);
@@ -146,6 +147,9 @@ void FNiagaraSystemViewModel::Initialize(UNiagaraSystem& InSystem, FNiagaraSyste
 
 	ScalabilityViewModel = NewObject<UNiagaraSystemScalabilityViewModel>(GetTransientPackage());
 	ScalabilityViewModel->Initialize(this->AsShared());
+
+	UserParametersHierarchyViewModel = NewObject<UNiagaraUserParametersHierarchyViewModel>(GetTransientPackage());
+	UserParametersHierarchyViewModel->Initialize(this->AsShared());
 	
 	PlaceholderDataInterfaceManager = MakeShared<FNiagaraPlaceholderDataInterfaceManager>(this->AsShared());
 
@@ -263,6 +267,12 @@ void FNiagaraSystemViewModel::Cleanup()
 	{
 		ScalabilityViewModel = nullptr;
 	}
+	
+	if(UserParametersHierarchyViewModel != nullptr)
+	{
+		UserParametersHierarchyViewModel->Finalize();
+		UserParametersHierarchyViewModel = nullptr;
+	}
 
 	if (PlaceholderDataInterfaceManager.IsValid())
 	{
@@ -731,6 +741,10 @@ void FNiagaraSystemViewModel::AddReferencedObjects(FReferenceCollector& Collecto
 	{
 		Collector.AddReferencedObject(ScalabilityViewModel);
 	}
+	if (UserParametersHierarchyViewModel != nullptr)
+	{
+		Collector.AddReferencedObject(UserParametersHierarchyViewModel);
+	}
 }
 
 void FNiagaraSystemViewModel::PostUndo(bool bSuccess)
@@ -989,6 +1003,11 @@ UNiagaraSystemScalabilityViewModel* FNiagaraSystemViewModel::GetScalabilityViewM
 	return ScalabilityViewModel;
 }
 
+UNiagaraUserParametersHierarchyViewModel* FNiagaraSystemViewModel::GetUserParametersHierarchyViewModel()
+{
+	return UserParametersHierarchyViewModel;
+}
+
 TArray<float> FNiagaraSystemViewModel::OnGetPlaybackSpeeds() const
 {
 	return GetDefault<UNiagaraEditorSettings>()->GetPlaybackSpeeds();
@@ -1033,8 +1052,10 @@ bool FNiagaraSystemViewModel::RenameParameter(const FNiagaraVariable TargetParam
 		else
 		{
 			// Otherwise it's safe to rename.
+			GetEditorData().RenameUserScriptVariable(TargetParameter, NewName);
 			ExposedParametersStore->RenameParameter(TargetParameter, NewName);
 		}
+		
 		bExposedParametersRename = true;
 	}
 
diff --git a/Engine/Plugins/FX/Niagara/Source/NiagaraEditor/Private/Widgets/SNiagaraHierarchy.cpp b/Engine/Plugins/FX/Niagara/Source/NiagaraEditor/Private/Widgets/SNiagaraHierarchy.cpp
new file mode 100644
--- /dev/null
+++ b/Engine/Plugins/FX/Niagara/Source/NiagaraEditor/Private/Widgets/SNiagaraHierarchy.cpp
@@ -0,0 +1,749 @@
+// Copyright Epic Games, Inc. All Rights Reserved.
+
+#include "Widgets/SNiagaraHierarchy.h"
+#include "Framework/Commands/GenericCommands.h"
+#include "Styling/StyleColors.h"
+#include "Styling/AppStyle.h"
+#include "ViewModels/HierarchyEditor/NiagaraHierarchyViewModelBase.h"
+#include "Widgets/SNiagaraParameterName.h"
+#include "Widgets/Layout/SSeparator.h"
+#include "Widgets/Layout/SWrapBox.h"
+#include "Widgets/Input/SCheckBox.h"
+#include "Widgets/Text/SInlineEditableTextBlock.h"
+#include "Widgets/Text/STextBlock.h"
+#include "Framework/MultiBox/MultiBoxBuilder.h"
+#include "ScopedTransaction.h"
+#include "Modules/ModuleManager.h"
+#include "PropertyEditorModule.h"
+
+#define LOCTEXT_NAMESPACE "NiagaraHierarchy"
+
+void SNiagaraHierarchyCategory::Construct(const FArguments& InArgs,	TSharedPtr<FNiagaraHierarchyCategoryViewModel> InCategory)
+{
+	Category = InCategory;
+
+	InCategory->GetOnRequestRename().BindSP(this, &SNiagaraHierarchyCategory::EnterEditingMode);
+	
+	ChildSlot
+	[
+		SAssignNew(InlineEditableTextBlock, SInlineEditableTextBlock)
+		.Style(&FNiagaraEditorStyle::Get().GetWidgetStyle<FInlineEditableTextBlockStyle>("NiagaraEditor.HierarchyEditor.Category"))
+		.Text(this, &SNiagaraHierarchyCategory::GetCategoryText)
+		.OnTextCommitted(this, &SNiagaraHierarchyCategory::OnRenameCategory)
+		.OnVerifyTextChanged(this, &SNiagaraHierarchyCategory::OnVerifyCategoryRename)
+		.IsSelected(InArgs._IsSelected)
+	];
+}
+
+void SNiagaraHierarchyCategory::EnterEditingMode() const
+{
+	InlineEditableTextBlock->EnterEditingMode();
+}
+
+bool SNiagaraHierarchyCategory::OnVerifyCategoryRename(const FText& NewName, FText& OutTooltip) const
+{
+	TArray<TSharedPtr<FNiagaraHierarchyCategoryViewModel>> Categories;
+	Category.Pin()->GetParent().Pin()->GetChildrenViewModelsForType<UNiagaraHierarchyCategory, FNiagaraHierarchyCategoryViewModel>(Categories);
+
+	if(GetCategoryText().ToString() != NewName.ToString())
+	{
+		TSet<FString> CategoryNames;
+		for(const auto& CategoryViewModel : Categories)
+		{
+			CategoryNames.Add(Cast<UNiagaraHierarchyCategory>(CategoryViewModel->GetDataMutable())->GetCategoryName().ToString());
+		}
+
+		if(CategoryNames.Contains(NewName.ToString()))
+		{
+			OutTooltip = LOCTEXT("HierarchyCategoryCantRename_DuplicateOnLayer", "Another category of the same name already exists on the same hierarchy level!");
+			return false;
+		}
+	}
+
+	return true;
+}
+
+FText SNiagaraHierarchyCategory::GetCategoryText() const
+{
+	return FText::FromName(Cast<UNiagaraHierarchyCategory>(Category.Pin()->GetData())->GetCategoryName());
+}
+
+void SNiagaraHierarchyCategory::OnRenameCategory(const FText& NewText, ETextCommit::Type) const
+{
+	FScopedTransaction Transaction(LOCTEXT("Transaction_Rename_Category", "Renamed hierarchy category"));
+	Category.Pin()->GetHierarchyViewModel()->GetHierarchyDataRoot()->Modify();
+	
+	Cast<UNiagaraHierarchyCategory>(Category.Pin()->GetDataMutable())->SetCategoryName(FName(NewText.ToString()));
+}
+
+void SNiagaraHierarchySection::Construct(const FArguments& InArgs, TSharedPtr<FNiagaraHierarchySectionViewModel> InSection, TWeakObjectPtr<UNiagaraHierarchyViewModelBase> InHierarchyViewModel)
+{
+	SectionViewModel = InSection;
+	HierarchyViewModel = InHierarchyViewModel;
+
+	if(SectionViewModel != nullptr)
+	{
+		SectionViewModel->GetOnRequestRename().BindSP(this, &SNiagaraHierarchySection::EnterEditingMode);
+
+		SDropTarget::FArguments LeftDropTargetArgs;
+		SDropTarget::FArguments OntoDropTargetArgs;
+		SDropTarget::FArguments RightDropTargetArgs;
+
+		LeftDropTargetArgs
+			.OnAllowDrop(this, &SNiagaraHierarchySection::OnCanAcceptSectionDrop, EItemDropZone::AboveItem)
+			.OnDropped(this, &SNiagaraHierarchySection::OnDroppedOn, EItemDropZone::AboveItem)
+			.VerticalImage(FAppStyle::GetNoBrush())
+			.HorizontalImage(FNiagaraEditorStyle::Get().GetBrush("NiagaraEditor.Stack.DropTarget.BorderHorizontal"));
+
+		OntoDropTargetArgs
+			.OnAllowDrop(this, &SNiagaraHierarchySection::OnCanAcceptDrop)
+			.OnDropped(this, &SNiagaraHierarchySection::OnDroppedOn, EItemDropZone::OntoItem);
+		
+		RightDropTargetArgs
+			.OnAllowDrop(this, &SNiagaraHierarchySection::OnCanAcceptSectionDrop, EItemDropZone::BelowItem)
+			.OnDropped(this, &SNiagaraHierarchySection::OnDroppedOn, EItemDropZone::BelowItem)
+			.VerticalImage(FAppStyle::GetNoBrush())
+			.HorizontalImage(FNiagaraEditorStyle::Get().GetBrush("NiagaraEditor.Stack.DropTarget.BorderHorizontal"));
+		
+		ChildSlot
+		[
+			SNew(SHorizontalBox)
+			+ SHorizontalBox::Slot()
+			.AutoWidth()
+			[
+				SNew(SNiagaraSectionDragDropTarget)
+				.DropTargetArgs(LeftDropTargetArgs)
+			]
+			+ SHorizontalBox::Slot()
+			[
+				SNew(SNiagaraSectionDragDropTarget)
+				.DropTargetArgs(OntoDropTargetArgs
+				.Content()
+				[
+					SNew(SCheckBox)
+					.Visibility(EVisibility::HitTestInvisible)
+					.Style(FAppStyle::Get(), "DetailsView.SectionButton")
+					.OnCheckStateChanged(this, &SNiagaraHierarchySection::OnSectionCheckChanged)
+					.IsChecked(this, &SNiagaraHierarchySection::GetSectionCheckState)
+					.OnGetMenuContent(this, &SNiagaraHierarchySection::OnGetMenuContent)
+					[
+						SAssignNew(InlineEditableTextBlock, SInlineEditableTextBlock)
+						.IsSelected(InArgs._IsSelected)
+						.Text(this, &SNiagaraHierarchySection::GetText)
+						.OnTextCommitted(this, &SNiagaraHierarchySection::OnRenameSection)
+						.OnVerifyTextChanged(this, &SNiagaraHierarchySection::OnVerifySectionRename)
+					]
+				])
+			]
+			+ SHorizontalBox::Slot()
+			.AutoWidth()
+			[
+				SNew(SNiagaraSectionDragDropTarget)
+				.DropTargetArgs(RightDropTargetArgs)				
+			]
+		];
+	}
+	// if this section doesn't represent data, it's the "All" widget
+	else
+	{
+		ChildSlot
+		[
+			SNew(SDropTarget)
+			.OnAllowDrop(this, &SNiagaraHierarchySection::OnCanAcceptDrop)
+			.OnDropped(this, &SNiagaraHierarchySection::OnDroppedOn, EItemDropZone::OntoItem)
+			[
+				SNew(SCheckBox)
+				.Style(FAppStyle::Get(), "DetailsView.SectionButton")
+				.OnCheckStateChanged(this, &SNiagaraHierarchySection::OnSectionCheckChanged)
+				.IsChecked(this, &SNiagaraHierarchySection::GetSectionCheckState)
+				[
+					SNew(STextBlock)
+					.Text(LOCTEXT("AllSection", "All"))
+				]
+			]			
+		];
+	}
+}
+
+SNiagaraHierarchySection::~SNiagaraHierarchySection()
+{
+	SectionViewModel.Reset();
+}
+
+void SNiagaraHierarchySection::EnterEditingMode() const
+{
+	InlineEditableTextBlock->EnterEditingMode();
+}
+
+bool SNiagaraHierarchySection::OnCanAcceptDrop(TSharedPtr<FDragDropOperation> DragDropOperation) const
+{
+	if(DragDropOperation->IsOfType<FNiagaraHierarchyDragDropOp>())
+	{
+		TSharedPtr<FNiagaraHierarchyDragDropOp> HierarchyDragDropOp = StaticCastSharedPtr<FNiagaraHierarchyDragDropOp>(DragDropOperation);
+		TWeakPtr<FNiagaraHierarchyItemViewModelBase> DraggedItem = HierarchyDragDropOp->GetDraggedItem();
+
+		if(const UNiagaraHierarchyCategory* Category = Cast<UNiagaraHierarchyCategory>(DraggedItem.Pin()->GetData()))
+		{
+			return TryGetSectionData() != Category->GetSection();
+		}
+	}
+
+	return false;
+}
+
+bool SNiagaraHierarchySection::OnCanAcceptSectionDrop(TSharedPtr<FDragDropOperation> DragDropOperation, EItemDropZone ItemDropZone) const
+{
+	if(DragDropOperation->IsOfType<FNiagaraSectionDragDropOp>() && ItemDropZone != EItemDropZone::OntoItem)
+	{
+		TSharedPtr<FNiagaraSectionDragDropOp> SectionDragDropOp = StaticCastSharedPtr<FNiagaraSectionDragDropOp>(DragDropOperation);
+		TWeakPtr<FNiagaraHierarchyItemViewModelBase> DraggedItem = SectionDragDropOp->GetDraggedItem();
+		if(UNiagaraHierarchySection* DraggedSection = Cast<UNiagaraHierarchySection>(DraggedItem.Pin()->GetDataMutable()))
+		{
+			int32 DraggedItemIndex = SectionViewModel->GetHierarchyViewModel()->GetHierarchyDataRoot()->GetSectionIndex(DraggedSection);
+			bool bSameSection = TryGetSectionData() == DraggedSection;
+
+			int32 InsertionIndex = INDEX_NONE;
+			if(SectionViewModel.IsValid())
+			{
+				InsertionIndex = SectionViewModel->GetHierarchyViewModel()->GetHierarchyDataRoot()->GetSectionIndex(TryGetSectionData());
+
+				// we add 1 to the insertion index if it's below an item because we either want to insert at the current index to place the item above, or at current+1 for below
+				InsertionIndex += ItemDropZone == EItemDropZone::AboveItem ? -1 : 1;
+			}
+
+			bool bCanDrop = !bSameSection && DraggedItemIndex != InsertionIndex;			
+			return bCanDrop;
+		}
+	}
+
+	return false;
+}
+
+FReply SNiagaraHierarchySection::OnDroppedOn(const FGeometry&, const FDragDropEvent& DragDropEvent, EItemDropZone DropZone) const
+{
+	if(SectionViewModel.IsValid())
+	{
+		return SectionViewModel->OnDroppedOn(DragDropEvent, DropZone, SectionViewModel);
+	}
+	else
+	{
+		if(TSharedPtr<FNiagaraHierarchyDragDropOp> HierarchyDragDropOp = DragDropEvent.GetOperationAs<FNiagaraHierarchyDragDropOp>())
+		{
+			if(UNiagaraHierarchyCategory* Category = Cast<UNiagaraHierarchyCategory>(HierarchyDragDropOp->GetDraggedItem().Pin()->GetDataMutable()))
+			{
+				Category->SetSection(nullptr);
+				TArray<UNiagaraHierarchyCategory*> ChildrenCategories;
+				Category->GetChildrenOfType<>(ChildrenCategories, true);
+
+				for(UNiagaraHierarchyCategory* ChildCategory : ChildrenCategories)
+				{
+					ChildCategory->SetSection(nullptr);
+				}
+
+				// we only need to reparent if the parent isn't already the root. This stops unnecessary reordering
+				if(HierarchyDragDropOp->GetDraggedItem().Pin()->GetParent() != HierarchyViewModel->GetHierarchyViewModelRoot())
+				{
+					HierarchyViewModel->GetHierarchyViewModelRoot()->ReparentToThis(HierarchyDragDropOp->GetDraggedItem().Pin());
+				}
+			
+				HierarchyViewModel->RefreshHierarchyView();
+
+				return FReply::Handled();
+			}
+		}
+	}
+
+	return FReply::Unhandled();
+}
+
+FReply SNiagaraHierarchySection::OnMouseButtonDown(const FGeometry& MyGeometry, const FPointerEvent& MouseEvent)
+{
+	if(SectionViewModel.IsValid())
+	{
+		if(MouseEvent.IsMouseButtonDown(EKeys::RightMouseButton))
+		{
+			OnSectionCheckChanged(ECheckBoxState::Checked);
+			return FReply::Handled();
+		}
+		else if(MouseEvent.IsMouseButtonDown(EKeys::LeftMouseButton))
+		{
+			OnSectionCheckChanged(ECheckBoxState::Checked);
+			return FReply::Handled().DetectDrag(AsShared(), EKeys::LeftMouseButton);	
+		}
+	}
+
+	return FReply::Unhandled();
+}
+
+FReply SNiagaraHierarchySection::OnDragDetected(const FGeometry& MyGeometry, const FPointerEvent& MouseEvent)
+{
+	if(MouseEvent.IsMouseButtonDown(EKeys::LeftMouseButton) && SectionViewModel.IsValid())
+	{
+		TSharedRef<FNiagaraSectionDragDropOp> SectionDragDropOp = MakeShared<FNiagaraSectionDragDropOp>(SectionViewModel);
+		SectionDragDropOp->Construct();
+		return FReply::Handled().BeginDragDrop(SectionDragDropOp);
+	}
+	
+	return FReply::Unhandled();
+}
+
+void SNiagaraHierarchySection::OnDragLeave(const FDragDropEvent& DragDropEvent)
+{
+	if(TSharedPtr<FNiagaraHierarchyDragDropOp> HierarchyDragDropOp = DragDropEvent.GetOperationAs<FNiagaraHierarchyDragDropOp>())
+	{
+		HierarchyDragDropOp->SetDescription(FText::GetEmpty());
+	}
+}
+
+TSharedRef<SWidget> SNiagaraHierarchySection::OnGetMenuContent() const
+{
+	FMenuBuilder MenuBuilder(true, HierarchyViewModel->GetCommands());
+	
+	MenuBuilder.AddMenuEntry(FGenericCommands::Get().Rename);
+	MenuBuilder.AddMenuEntry(FGenericCommands::Get().Delete);
+
+	return MenuBuilder.MakeWidget();
+}
+
+UNiagaraHierarchySection* SNiagaraHierarchySection::TryGetSectionData() const
+{
+	return SectionViewModel.IsValid() ? Cast<UNiagaraHierarchySection>(SectionViewModel->GetDataMutable()) : nullptr;
+}
+
+FText SNiagaraHierarchySection::GetText() const
+{
+	return SectionViewModel->GetSectionNameAsText();
+}
+
+void SNiagaraHierarchySection::OnRenameSection(const FText& Text, ETextCommit::Type CommitType) const
+{
+	FScopedTransaction Transaction(LOCTEXT("Transaction_Rename_Section", "Renamed hierarchy section"));
+	HierarchyViewModel->GetHierarchyDataRoot()->Modify();
+	
+	SectionViewModel->SetSectionNameAsText(Text);
+}
+
+bool SNiagaraHierarchySection::OnVerifySectionRename(const FText& NewName, FText& OutTooltip) const
+{
+	// this function shouldn't be used in case the section isn't valid but we'll make sure regardless
+	if(!SectionViewModel.IsValid())
+	{
+		return false;
+	}
+
+	if(SectionViewModel->GetSectionName().ToString() != NewName.ToString())
+	{
+		TArray<FString> SectionNames;
+
+		SectionNames.Add("All");
+		for(auto& Section : HierarchyViewModel->GetHierarchyViewModelRoot()->GetSectionViewModels())
+		{
+			SectionNames.Add(Section->GetSectionName().ToString());
+		}
+
+		if(SectionNames.Contains(NewName.ToString()))
+		{
+			OutTooltip = LOCTEXT("HierarchySectionCantRename_Duplicate", "A section with that name already exists!");
+			return false;
+		}
+	}
+
+	return true;
+}
+
+ECheckBoxState SNiagaraHierarchySection::GetSectionCheckState() const
+{
+	return HierarchyViewModel->GetActiveSection() == SectionViewModel ? ECheckBoxState::Checked : ECheckBoxState::Unchecked;
+}
+
+void SNiagaraHierarchySection::OnSectionCheckChanged(ECheckBoxState NewState)
+{
+	if(NewState == ECheckBoxState::Checked)
+	{
+		HierarchyViewModel->SetActiveSection(SectionViewModel);
+	}
+
+	if(SectionViewModel.IsValid())
+	{
+		HierarchyViewModel->SelectItemInDetailsPanel(SectionViewModel->GetDataMutable());
+	}
+}
+
+void SNiagaraHierarchy::Construct(const FArguments& InArgs, TObjectPtr<UNiagaraHierarchyViewModelBase> InHierarchyViewModel)
+{
+	HierarchyViewModel = InHierarchyViewModel;
+	
+	HierarchyViewModel->OnRefreshSourceView().BindSP(this, &SNiagaraHierarchy::RefreshSourceView);
+	HierarchyViewModel->OnRefreshHierarchyView().BindSP(this, &SNiagaraHierarchy::RefreshHierarchyView);
+	HierarchyViewModel->OnRefreshSections().BindSP(this, &SNiagaraHierarchy::RefreshSectionsWidget);
+	HierarchyViewModel->OnSelectObjectInDetailsPanel().BindSP(this, &SNiagaraHierarchy::SelectObjectInDetailsPanel);
+	
+	SHorizontalBox::FSlot* DetailsPanelSlot = nullptr;
+	
+	ChildSlot
+	[
+		SNew(SHorizontalBox)
+		+ SHorizontalBox::Slot()
+		[
+			SAssignNew(SourceTreeView, STreeView<TSharedPtr<FNiagaraHierarchyItemViewModelBase>>)
+			.TreeItemsSource(&InHierarchyViewModel->GetSourceItems())
+			.OnSelectionChanged(this, &SNiagaraHierarchy::OnSelectionChanged)
+			.OnGenerateRow(this, &SNiagaraHierarchy::GenerateSourceItemRow)
+			.OnGetChildren_UObject(InHierarchyViewModel.Get(), &UNiagaraHierarchyViewModelBase::OnGetChildren)
+			.OnItemToString_Debug_UObject(InHierarchyViewModel.Get(), &UNiagaraHierarchyViewModelBase::OnItemToStringDebug)
+		]
+		+ SHorizontalBox::Slot()
+		.AutoWidth()
+		[
+			SNew(SSeparator).Orientation(EOrientation::Orient_Vertical)
+		]
+		+ SHorizontalBox::Slot()
+		[
+			SNew(SVerticalBox)
+			+ SVerticalBox::Slot()
+			.AutoHeight()
+			[
+				SNew(SVerticalBox)
+				+ SVerticalBox::Slot()
+				.AutoHeight()
+				[
+					SNew(SBox)
+					[
+						CreateHierarchyButtonWidgets()
+					]
+				]
+				+ SVerticalBox::Slot().Expose(SectionsSlot)
+				.AutoHeight()
+			]
+			+ SVerticalBox::Slot()
+			[
+				SAssignNew(HierarchyTreeView, STreeView<TSharedPtr<FNiagaraHierarchyItemViewModelBase>>)
+				.TreeItemsSource(&InHierarchyViewModel->GetHierarchyItems())
+				.OnSelectionChanged(this, &SNiagaraHierarchy::OnSelectionChanged)
+				.OnGenerateRow(this, &SNiagaraHierarchy::GenerateHierarchyItemRow)
+				.OnGetChildren_UObject(InHierarchyViewModel.Get(), &UNiagaraHierarchyViewModelBase::OnGetChildren)
+				.OnItemToString_Debug_UObject(InHierarchyViewModel.Get(), &UNiagaraHierarchyViewModelBase::OnItemToStringDebug)
+			]
+		]
+		+ SHorizontalBox::Slot().Expose(DetailsPanelSlot)		
+	];
+
+	if(InHierarchyViewModel->SupportsDetailsPanel())
+	{
+		FPropertyEditorModule& PropertyEditorModule = FModuleManager::Get().LoadModuleChecked<FPropertyEditorModule>("PropertyEditor");
+
+		FDetailsViewArgs DetailsViewArgs;
+		DetailsViewArgs.NameAreaSettings = FDetailsViewArgs::ENameAreaSettings::ObjectsUseNameArea;
+		DetailsViewArgs.bShowObjectLabel = true;
+		DetailsPanel = PropertyEditorModule.CreateDetailView(DetailsViewArgs);
+
+		for(auto& Customizations : HierarchyViewModel->GetInstanceCustomizations())
+		{
+			DetailsPanel->RegisterInstancedCustomPropertyLayout(Customizations.Key, Customizations.Value);
+		}
+		
+		DetailsPanelSlot->AttachWidget(DetailsPanel.ToSharedRef());
+	}
+	
+	RefreshSectionsWidget();
+	HierarchyViewModel->GetHierarchyViewModelRoot()->GetOnSynced().AddSP(this, &SNiagaraHierarchy::RefreshSectionsWidget);
+	
+	HierarchyViewModel->GetCommands()->MapAction(FGenericCommands::Get().Rename,
+		FExecuteAction::CreateSP(this, &SNiagaraHierarchy::RequestRenameSelectedItem),
+		FCanExecuteAction::CreateSP(this, &SNiagaraHierarchy::CanRequestRenameSelectedItem));
+
+	HierarchyViewModel->GetCommands()->MapAction(FGenericCommands::Get().Delete,
+		FExecuteAction::CreateSP(this, &SNiagaraHierarchy::DeleteSelectedItems),
+		FCanExecuteAction::CreateSP(this, &SNiagaraHierarchy::CanDeleteSelectedItems));
+}
+
+SNiagaraHierarchy::~SNiagaraHierarchy()
+{
+	HierarchyViewModel->GetHierarchyViewModelRoot()->GetOnSynced().RemoveAll(this);
+
+	HierarchyViewModel->OnRefreshSourceView().Unbind();
+	HierarchyViewModel->OnRefreshHierarchyView().Unbind();
+	HierarchyViewModel->OnRefreshSections().Unbind();
+	HierarchyViewModel->OnSelectObjectInDetailsPanel().Unbind();
+
+	HierarchyViewModel->GetCommands()->UnmapAction(FGenericCommands::Get().Delete);
+	HierarchyViewModel->GetCommands()->UnmapAction(FGenericCommands::Get().Rename);
+}
+
+FReply SNiagaraHierarchy::OnKeyDown(const FGeometry& MyGeometry, const FKeyEvent& InKeyEvent)
+{
+	return HierarchyViewModel->GetCommands()->ProcessCommandBindings(InKeyEvent) ? FReply::Handled() : FReply::Unhandled();
+}
+
+void SNiagaraHierarchy::RefreshSourceView(bool bFullRefresh) const
+{
+	if(bFullRefresh)
+	{
+		SourceTreeView->RebuildList();
+	}
+	else
+	{
+		SourceTreeView->RequestTreeRefresh();
+	}
+}
+
+void SNiagaraHierarchy::RefreshHierarchyView(bool bFullRefresh) const
+{
+	// the top layer objects might have changed due to filtering. We need to refresh these too.
+	HierarchyTreeView->SetTreeItemsSource(&HierarchyViewModel->GetHierarchyItems());
+	if(bFullRefresh)
+	{
+		HierarchyTreeView->RebuildList();
+	}
+	else
+	{
+		HierarchyTreeView->RequestTreeRefresh();
+	}
+}
+
+void SNiagaraHierarchy::RefreshSectionsWidget()
+{
+	SectionsSlot->AttachWidget(CreateSectionWidgets());
+}
+
+bool SNiagaraHierarchy::IsItemSelected(TSharedPtr<FNiagaraHierarchyItemViewModelBase> Item) const
+{
+	return HierarchyTreeView->IsItemSelected(Item);
+}
+
+void SNiagaraHierarchy::SelectObjectInDetailsPanel(UObject* Object) const
+{
+	if(DetailsPanel.IsValid())
+	{
+		DetailsPanel->SetObject(Object);
+	}
+}
+
+TSharedRef<SWidget> SNiagaraHierarchy::CreateHierarchyButtonWidgets() const
+{
+	FSlimHorizontalToolBarBuilder ToolBarBuilder(nullptr, FMultiBoxCustomization::None);
+
+	// Add Category
+	{
+		FUIAction AddCategoryAction;
+		AddCategoryAction.ExecuteAction = FExecuteAction::CreateUObject(HierarchyViewModel.Get(), &UNiagaraHierarchyViewModelBase::AddCategory);
+		ToolBarBuilder.AddToolBarButton(AddCategoryAction, NAME_None,
+			FText::FromString("Add Category")
+			);
+	}
+
+	// Add Section
+	{
+		FUIAction AddSectionAction;
+		AddSectionAction.ExecuteAction = FExecuteAction::CreateUObject(HierarchyViewModel.Get(), &UNiagaraHierarchyViewModelBase::AddSection);
+		ToolBarBuilder.AddToolBarButton(AddSectionAction, NAME_None,
+			FText::FromString("Add Section")
+			);
+	}
+
+	return ToolBarBuilder.MakeWidget();
+}
+
+TSharedRef<SWidget> SNiagaraHierarchy::CreateSectionWidgets()
+{
+	SectionsWidgetMap.Empty();
+	
+	TSharedRef<SWrapBox> SectionsBox = SNew(SWrapBox)
+		.UseAllottedWidth(true);
+	
+	for (TSharedPtr<FNiagaraHierarchySectionViewModel>& Section : HierarchyViewModel->GetHierarchyViewModelRoot()->GetSectionViewModels())
+	{
+		TSharedPtr<SNiagaraHierarchySection> SectionWidget = SNew(SNiagaraHierarchySection, Section, HierarchyViewModel);
+		SectionsWidgetMap.Add(Section, SectionWidget);
+		
+		SectionsBox->AddSlot()
+		[
+			SectionWidget.ToSharedRef()
+		];
+	}
+	
+	TSharedPtr<SNiagaraHierarchySection> DefaultSectionWidget = SNew(SNiagaraHierarchySection, nullptr, HierarchyViewModel);
+	SectionsWidgetMap.Add(nullptr, DefaultSectionWidget);
+
+	SectionsBox->AddSlot()
+	[
+		DefaultSectionWidget.ToSharedRef()
+	];
+	
+	return SNew(SBorder).BorderBackgroundColor(FStyleColors::Recessed)
+	[
+		SectionsBox
+	];
+}
+
+TSharedRef<ITableRow> SNiagaraHierarchy::GenerateSourceItemRow(TSharedPtr<FNiagaraHierarchyItemViewModelBase> HierarchyItem, const TSharedRef<STableViewBase>& TableViewBase)
+{
+	return SNew(STableRow<TSharedPtr<FNiagaraHierarchyItemViewModelBase>>, TableViewBase)
+	.OnDragDetected(HierarchyItem.ToSharedRef(), &FNiagaraHierarchyItemViewModelBase::OnDragDetected, true)
+	[
+		HierarchyViewModel->GenerateRowContentWidget(HierarchyItem.ToSharedRef())
+	];
+}
+
+TSharedRef<ITableRow> SNiagaraHierarchy::GenerateHierarchyItemRow(TSharedPtr<FNiagaraHierarchyItemViewModelBase> HierarchyItem, const TSharedRef<STableViewBase>& TableViewBase)
+{
+	return SNew(STableRow<TSharedPtr<FNiagaraHierarchyItemViewModelBase>>, TableViewBase)
+	.OnAcceptDrop(HierarchyItem.ToSharedRef(), &FNiagaraHierarchyItemViewModelBase::OnDroppedOn)
+	.OnCanAcceptDrop(HierarchyItem.ToSharedRef(), &FNiagaraHierarchyItemViewModelBase::OnCanAcceptDrop)
+	.OnDragDetected(HierarchyItem.ToSharedRef(), &FNiagaraHierarchyItemViewModelBase::OnDragDetected, false)
+	[
+		HierarchyViewModel->GenerateRowContentWidget(HierarchyItem.ToSharedRef())
+	];
+}
+
+void SNiagaraHierarchy::RequestRenameSelectedItem()
+{
+	TArray<TSharedPtr<FNiagaraHierarchyItemViewModelBase>> SelectedItems = HierarchyTreeView->GetSelectedItems();
+
+	if(SelectedItems.Num() == 0)
+	{
+		TSharedPtr<FNiagaraHierarchySectionViewModel> ActiveSection = HierarchyViewModel->GetActiveSection();
+		if(ActiveSection)
+		{
+			SelectedItems = { ActiveSection };
+		}
+	}
+	
+	if(SelectedItems.Num() == 1)
+	{
+		SelectedItems[0]->RequestRename();
+	}
+}
+
+bool SNiagaraHierarchy::CanRequestRenameSelectedItem() const
+{
+	TArray<TSharedPtr<FNiagaraHierarchyItemViewModelBase>> SelectedItems = HierarchyTreeView->GetSelectedItems();
+
+	if(SelectedItems.Num() == 0)
+	{
+		TSharedPtr<FNiagaraHierarchySectionViewModel> ActiveSection = HierarchyViewModel->GetActiveSection();
+		if(ActiveSection)
+		{
+			SelectedItems = { ActiveSection };
+		}
+	}
+	
+	if(SelectedItems.Num() == 1)
+	{
+		return SelectedItems[0]->CanRename();
+	}
+
+	return false;
+}
+
+void SNiagaraHierarchy::DeleteSelectedItems()
+{
+	FScopedTransaction Transaction(LOCTEXT("Transaction_DeleteSelectedItems", "Deleted selected hierarchy items"));
+	HierarchyViewModel->GetHierarchyDataRoot()->Modify();
+	
+	TArray<TSharedPtr<FNiagaraHierarchyItemViewModelBase>> SelectedItems = HierarchyTreeView->GetSelectedItems();
+
+	if(SelectedItems.Num() == 0)
+	{
+		TSharedPtr<FNiagaraHierarchySectionViewModel> ActiveSection = HierarchyViewModel->GetActiveSection();
+		if(ActiveSection)
+		{
+			SelectedItems = { ActiveSection };
+		}
+	}
+	
+	for(TSharedPtr<FNiagaraHierarchyItemViewModelBase>& SelectedItem : SelectedItems)
+	{
+		SelectedItem->Finalize();
+		if(SelectedItem->GetParent().IsValid())
+		{
+			SelectedItem->GetParent().Pin()->SyncToData();
+		}
+	}
+
+	RefreshSourceView();
+	RefreshHierarchyView();
+}
+
+bool SNiagaraHierarchy::CanDeleteSelectedItems() const
+{
+	TArray<TSharedPtr<FNiagaraHierarchyItemViewModelBase>> SelectedItems = HierarchyTreeView->GetSelectedItems();
+
+	if(SelectedItems.Num() == 0)
+	{
+		TSharedPtr<FNiagaraHierarchySectionViewModel> ActiveSection = HierarchyViewModel->GetActiveSection();
+		if(ActiveSection)
+		{
+			SelectedItems = { ActiveSection };
+		}
+	}
+
+	if(SelectedItems.Num() > 0)
+	{
+		bool bCanDelete = true;
+		for(TSharedPtr<FNiagaraHierarchyItemViewModelBase> SelectedItem : SelectedItems)
+		{
+			bCanDelete &= SelectedItems[0]->CanDelete();
+		}
+
+		return bCanDelete;
+	}
+
+	return false;
+}
+
+void SNiagaraHierarchy::OnSelectionChanged(TSharedPtr<FNiagaraHierarchyItemViewModelBase> HierarchyItem, ESelectInfo::Type Type) const
+{
+	HierarchyViewModel->OnSelectionChanged(HierarchyItem);
+}
+
+void SNiagaraSectionDragDropTarget::Construct(const FArguments& InArgs)
+{
+	SDropTarget::Construct(InArgs._DropTargetArgs);
+}
+
+FReply SNiagaraSectionDragDropTarget::OnDragOver(const FGeometry& MyGeometry, const FDragDropEvent& DragDropEvent)
+{
+	TSharedPtr<FNiagaraHierarchyDragDropOp> HierarchyDragDropOp = DragDropEvent.GetOperationAs<FNiagaraHierarchyDragDropOp>();
+
+	if(HierarchyDragDropOp.IsValid())
+	{
+		if(AllowDrop(DragDropEvent.GetOperation()))
+		{
+			if(const UNiagaraHierarchySection* Section = Cast<UNiagaraHierarchySection>(HierarchyDragDropOp->GetDraggedItem().Pin()->GetData()))
+			{
+				FText Text = LOCTEXT("MoveSectionHere", "The section will be moved here.");
+				HierarchyDragDropOp->SetDescription(Text);
+			}
+
+			if(const UNiagaraHierarchyCategory* Category = Cast<UNiagaraHierarchyCategory>(HierarchyDragDropOp->GetDraggedItem().Pin()->GetData()))
+			{
+				FText Text = LOCTEXT("DropCategoryOnSection", "Move category here");
+				HierarchyDragDropOp->SetDescription(Text);
+			}
+		}
+		else
+		{
+			if(const UNiagaraHierarchySection* Section = Cast<UNiagaraHierarchySection>(HierarchyDragDropOp->GetDraggedItem().Pin()->GetData()))
+			{
+				HierarchyDragDropOp->SetDescription(LOCTEXT("SectionOnSectionNotAllowed", "Can't drop section on a section."));
+			}
+		}
+	}
+
+	return SDropTarget::OnDragOver(MyGeometry, DragDropEvent);
+}
+
+void SNiagaraSectionDragDropTarget::OnDragLeave(const FDragDropEvent& DragDropEvent)
+{
+	TSharedPtr<FNiagaraHierarchyDragDropOp> HierarchyDragDropOp = DragDropEvent.GetOperationAs<FNiagaraHierarchyDragDropOp>();
+
+	if(HierarchyDragDropOp.IsValid())
+	{
+		HierarchyDragDropOp->SetDescription(FText::GetEmpty());
+	}
+	
+	SDropTarget::OnDragLeave(DragDropEvent);
+}
+
+#undef LOCTEXT_NAMESPACE
diff --git a/Engine/Plugins/FX/Niagara/Source/NiagaraEditor/Private/Widgets/SNiagaraParameterPanel.cpp b/Engine/Plugins/FX/Niagara/Source/NiagaraEditor/Private/Widgets/SNiagaraParameterPanel.cpp
--- a/Engine/Plugins/FX/Niagara/Source/NiagaraEditor/Private/Widgets/SNiagaraParameterPanel.cpp
+++ b/Engine/Plugins/FX/Niagara/Source/NiagaraEditor/Private/Widgets/SNiagaraParameterPanel.cpp
@@ -75,7 +75,6 @@ void SNiagaraParameterPanel::Construct(const FArguments& InArgs, const TSharedPt
 		FilterCategoryExpandedDelegate.BindRaw(ParameterPanelViewModel.Get(), &INiagaraImmutableParameterPanelViewModel::IsCategoryExpandedByDefault);
 	}
 
-
 	SAssignNew(ItemSelector, SNiagaraParameterPanelSelector)
 	.PreserveSelectionOnRefresh(true)
 	.PreserveExpansionOnRefresh(true)
@@ -103,7 +102,11 @@ void SNiagaraParameterPanel::Construct(const FArguments& InArgs, const TSharedPt
 	.CategoryBorderBackgroundPadding(FMargin(0.0f, 3.0f))
 	.OnGetKeyForItem(this, &SNiagaraParameterPanel::OnGetKeyForItem)
 	.OnGetKeyForCategory(this, &SNiagaraParameterPanel::OnGetKeyForCategory)
-	.ExpandInitially(false);
+	.ExpandInitially(false)
+	.SearchBoxAdjacentContent()
+	[
+		ParameterPanelViewModel->GenerateAdjacentWidget()
+	];
 
 	// Finalize the widget
 	ChildSlot
@@ -404,7 +407,10 @@ TSharedRef<SWidget> SNiagaraParameterPanel::OnGenerateWidgetForItem(const FNiaga
 	TSharedPtr<SNiagaraParameterNameTextBlock> ParameterNameTextBlock = SNew(SNiagaraParameterNameTextBlock)
 	.ParameterText(FText::FromName(Item.GetVariable().GetName()))
 	.HighlightText(ItemSelector.ToSharedRef(), &SNiagaraParameterPanelSelector::GetFilterTextNoRef)
-	.ToolTipText(Item.ScriptVariable->Metadata.Description)
+	.ToolTipText_Lambda([Item]
+	{
+		return Item.ScriptVariable->Metadata.Description;
+	})
 	.OnTextCommitted(this, &SNiagaraParameterPanel::OnParameterNameTextCommitted, Item)
 	.OnVerifyTextChanged(this, &SNiagaraParameterPanel::OnParameterNameTextVerifyChanged, Item)
 	.IsSelected(ItemSelector.ToSharedRef(), &SNiagaraParameterPanelSelector::IsItemSelected, Item)
diff --git a/Engine/Plugins/FX/Niagara/Source/NiagaraEditor/Public/NiagaraEditorStyle.h b/Engine/Plugins/FX/Niagara/Source/NiagaraEditor/Public/NiagaraEditorStyle.h
--- a/Engine/Plugins/FX/Niagara/Source/NiagaraEditor/Public/NiagaraEditorStyle.h
+++ b/Engine/Plugins/FX/Niagara/Source/NiagaraEditor/Public/NiagaraEditorStyle.h
@@ -51,6 +51,7 @@ private:
 	void InitScalabilityColors();
 	void InitScalabilityIcons();
 	void InitScratchStyle();
+	void InitHierarchyEditor();
 
 	static TSharedPtr<FNiagaraEditorStyle> NiagaraEditorStyle;
 };
diff --git a/Engine/Plugins/FX/Niagara/Source/NiagaraEditor/Public/NiagaraEditorUtilities.h b/Engine/Plugins/FX/Niagara/Source/NiagaraEditor/Public/NiagaraEditorUtilities.h
--- a/Engine/Plugins/FX/Niagara/Source/NiagaraEditor/Public/NiagaraEditorUtilities.h
+++ b/Engine/Plugins/FX/Niagara/Source/NiagaraEditor/Public/NiagaraEditorUtilities.h
@@ -12,6 +12,7 @@
 #include "NiagaraEditorSettings.h"
 #include "UpgradeNiagaraScriptResults.h"
 #include "ViewModels/NiagaraSystemScalabilityViewModel.h"
+#include "ViewModels/NiagaraSystemViewModel.h"
 
 class UNiagaraNodeInput;
 class UNiagaraNodeOutput;
@@ -275,6 +276,10 @@ namespace FNiagaraEditorUtilities
 	 */
 	bool AddParameter(FNiagaraVariable& NewParameterVariable, FNiagaraParameterStore& TargetParameterStore, UObject& ParameterStoreOwner, UNiagaraStackEditorData* StackEditorData);
 
+	NIAGARAEDITOR_API TObjectPtr<UNiagaraScriptVariable> GetScriptVariableForUserParameter(const FNiagaraVariable& UserParameter, TSharedPtr<FNiagaraSystemViewModel> SystemViewModel);
+	NIAGARAEDITOR_API TObjectPtr<UNiagaraScriptVariable> GetScriptVariableForUserParameter(const FNiagaraVariable& UserParameter, UNiagaraSystem& System);
+	NIAGARAEDITOR_API TObjectPtr<UNiagaraScriptVariable> FindScriptVariableForUserParameter(const FGuid& UserParameterGuid, UNiagaraSystem& System);
+	
 	NIAGARAEDITOR_API bool AddEmitterContextMenuActions(FMenuBuilder& MenuBuilder, const TSharedPtr<FNiagaraEmitterHandleViewModel>& EmitterHandleViewModel);
 
 	void ShowParentEmitterInContentBrowser(TSharedRef<FNiagaraEmitterViewModel> EmitterViewModel);
@@ -423,7 +428,7 @@ namespace FNiagaraParameterUtilities
 
 	NIAGARAEDITOR_API bool TestCanRenameWithMessage(FName ParameterName, FText& OutMessage);
 
-	NIAGARAEDITOR_API TSharedRef<SWidget> GetParameterWidget(FNiagaraVariable Variable, bool bShowValue);
+	NIAGARAEDITOR_API TSharedRef<SWidget> GetParameterWidget(FNiagaraVariable Variable, bool bAddTypeIcon, bool bShowValue);
 	
 	/** Creates a tooltip based on a parameter. Also shows the value, if allocated and enabled. */
 	NIAGARAEDITOR_API TSharedRef<SToolTip> GetTooltipWidget(FNiagaraVariable Variable, bool bShowValue = true, TSharedPtr<SWidget> AdditionalVerticalWidget = nullptr,  TSharedPtr<SWidget> AdditionalHorizontalWidget = nullptr);
diff --git a/Engine/Plugins/FX/Niagara/Source/NiagaraEditor/Public/NiagaraSystemEditorData.h b/Engine/Plugins/FX/Niagara/Source/NiagaraEditor/Public/NiagaraSystemEditorData.h
--- a/Engine/Plugins/FX/Niagara/Source/NiagaraEditor/Public/NiagaraSystemEditorData.h
+++ b/Engine/Plugins/FX/Niagara/Source/NiagaraEditor/Public/NiagaraSystemEditorData.h
@@ -4,6 +4,7 @@
 
 #include "NiagaraEditorDataBase.h"
 #include "NiagaraEditorData.h"
+#include "ViewModels/HierarchyEditor/NiagaraHierarchyViewModelBase.h"
 #include "NiagaraSystemEditorData.generated.h"
 
 class UNiagaraStackEditorData;
@@ -50,6 +51,7 @@ class NIAGARAEDITOR_API UNiagaraSystemEditorData : public UNiagaraEditorDataBase
 {
 	GENERATED_BODY()
 
+	DECLARE_MULTICAST_DELEGATE(FOnUserParameterScriptVariablesSynced)
 public:
 	UNiagaraSystemEditorData(const FObjectInitializer& ObjectInitializer);
 
@@ -87,13 +89,25 @@ public:
 
 	void SynchronizeOverviewGraphWithSystem(UNiagaraSystem& OwnerSystem);
 
+	void SyncUserScriptVariables(UNiagaraSystem* System);
+	FOnUserParameterScriptVariablesSynced& OnUserParameterScriptVariablesSynced() { return OnUserParameterScriptVariablesSyncedDelegate; }
+	
+	UNiagaraScriptVariable* FindOrAddUserScriptVariable(FNiagaraVariable UserParameter, UNiagaraSystem& System);
+	UNiagaraScriptVariable* FindUserScriptVariable(FGuid UserParameterGuid);
+	bool RenameUserScriptVariable(FNiagaraVariable OldVariable, FName NewName);
+	bool RemoveUserScriptVariable(FNiagaraVariable Variable);
+
 	// If true then the preview viewport's orbit setting is saved in the asset data
 	UPROPERTY()
 	bool bSetOrbitModeByAsset = false;
 
 	UPROPERTY()
 	bool bSystemViewportInOrbitMode = true;
-
+	
+	/** Contains the root ids for organizing user parameters. */
+	UPROPERTY()
+	TObjectPtr<UNiagaraHierarchyRoot> UserParameterHierarchy;
+	
 private:
 	void UpdatePlaybackRangeFromEmitters(UNiagaraSystem& OwnerSystem);
 
@@ -122,4 +136,9 @@ private:
 
 	UPROPERTY()
 	bool bSystemIsPlaceholder;
+
+	UPROPERTY()
+	TArray<TObjectPtr<UNiagaraScriptVariable>> UserParameterMetaData;
+	
+	FOnUserParameterScriptVariablesSynced OnUserParameterScriptVariablesSyncedDelegate;
 };
diff --git a/Engine/Plugins/FX/Niagara/Source/NiagaraEditor/Public/ViewModels/HierarchyEditor/NiagaraHierarchyViewModelBase.h b/Engine/Plugins/FX/Niagara/Source/NiagaraEditor/Public/ViewModels/HierarchyEditor/NiagaraHierarchyViewModelBase.h
new file mode 100644
--- /dev/null
+++ b/Engine/Plugins/FX/Niagara/Source/NiagaraEditor/Public/ViewModels/HierarchyEditor/NiagaraHierarchyViewModelBase.h
@@ -0,0 +1,524 @@
+// Copyright Epic Games, Inc. All Rights Reserved.
+
+#pragma once
+
+#include "CoreMinimal.h"
+#include "NiagaraEditorUtilities.h"
+#include "ViewModels/Stack/NiagaraStackFunctionInput.h"
+#include "NiagaraHierarchyViewModelBase.generated.h"
+
+UCLASS()
+class UNiagaraHierarchyItemBase : public UObject
+{
+	GENERATED_BODY()
+
+public:
+	UNiagaraHierarchyItemBase() { Guid = FGuid::NewGuid(); }
+	virtual ~UNiagaraHierarchyItemBase() override {}
+
+	TArray<UNiagaraHierarchyItemBase*>& GetChildrenMutable() { return Children; }
+	const TArray<UNiagaraHierarchyItemBase*>& GetChildren() const;
+
+	template<class ChildClass>
+	bool DoesOneChildExist(bool bRecursive = false) const;
+
+	template<class ChildClass>
+	TArray<ChildClass*> GetChildrenOfType(TArray<ChildClass*>& Out, bool bRecursive = false) const;
+
+	/** Used to clean out unneeded data using the bFinalized flag. */
+	void Refresh();
+
+	virtual FString ToString() const { return GetName(); }
+
+	/** Finalize an item to mark it is no longer used and should be deleted. */
+	void Finalize();
+	bool IsFinalized() const { return bFinalized; }
+
+	/** A guid can be optionally set to create a mapping from previously existing guids to hierarchy items that represent them. */
+	void SetGuid(FGuid InGuid) { Guid = InGuid; }
+	virtual FGuid GetPersistentIdentity() const { return Guid; }
+	
+	virtual bool CanHaveChildren() const { return false; }
+
+	/** Overridden modify method to also mark all children as modified */
+	virtual bool Modify(bool bAlwaysMarkDirty = true) override;
+
+protected:
+	/** Called by the public Refresh function. Can be overridden to further customize the refresh process, i.e. refreshing section data in the root. */
+	virtual void RefreshDataInternal() { }
+
+	UPROPERTY()
+	TArray<TObjectPtr<UNiagaraHierarchyItemBase>> Children;
+	
+	/** An optional guid; can be used if hierarchy items represent outside items */
+	UPROPERTY()
+	FGuid Guid;
+
+	UPROPERTY(Transient)
+	bool bFinalized = false;
+};
+
+template <class ChildClass>
+bool UNiagaraHierarchyItemBase::DoesOneChildExist(bool bRecursive) const
+{
+	for(UNiagaraHierarchyItemBase* ItemBase : Children)
+	{
+		if(ItemBase->IsA<ChildClass>())
+		{
+			return true;
+		}
+	}
+
+	if(bRecursive)
+	{
+		for(UNiagaraHierarchyItemBase* ItemBase : Children)
+		{
+			if(ItemBase->DoesOneChildExist<ChildClass>(bRecursive))
+			{
+				return true;
+			}
+		}
+	}
+
+	return false;
+}
+
+template <class ChildClass>
+TArray<ChildClass*> UNiagaraHierarchyItemBase::GetChildrenOfType(TArray<ChildClass*>& Out, bool bRecursive) const
+{
+	for(UNiagaraHierarchyItemBase* ItemBase : Children)
+	{
+		if(ItemBase->IsA<ChildClass>())
+		{
+			Out.Add(Cast<ChildClass>(ItemBase));
+		}
+	}
+
+	if(bRecursive)
+	{
+		for(UNiagaraHierarchyItemBase* ItemBase : Children)
+		{
+			ItemBase->GetChildrenOfType<ChildClass>(Out, bRecursive);
+		}
+	}
+
+	return Out;
+}
+
+class UNiagaraHierarchySection;
+
+UCLASS()
+class UNiagaraHierarchyRoot : public UNiagaraHierarchyItemBase
+{
+	GENERATED_BODY()
+public:
+	UNiagaraHierarchyRoot() {}
+	virtual ~UNiagaraHierarchyRoot() override {}
+	
+	virtual void RefreshDataInternal() override;
+	
+	virtual bool CanHaveChildren() const override { return true; }
+	
+	class UNiagaraHierarchySection* AddSection(FText InNewSectionName);
+	void RemoveSection(FText SectionName);
+	TSet<FName> GetSections() const;
+	const TArray<UNiagaraHierarchySection*>& GetSectionData() const { return Sections; }
+	TArray<UNiagaraHierarchySection*>& GetSectionDataMutable() { return Sections; }
+	int32 GetSectionIndex(UNiagaraHierarchySection* Section) const;
+	virtual bool Modify(bool bAlwaysMarkDirty = true) override;
+protected:
+
+	UPROPERTY()
+	TArray<TObjectPtr<UNiagaraHierarchySection>> Sections;
+};
+
+UCLASS()
+class UNiagaraHierarchyItem : public UNiagaraHierarchyItemBase
+{
+	GENERATED_BODY()
+public:
+	UNiagaraHierarchyItem() {}
+	virtual ~UNiagaraHierarchyItem() override {}
+
+	/** An item is assumed to not have any children. */
+	virtual bool CanHaveChildren() const override { return false; }
+};
+
+UCLASS()
+class UNiagaraHierarchyCategory : public UNiagaraHierarchyItemBase
+{
+	GENERATED_BODY()
+public:
+	UNiagaraHierarchyCategory() {}
+	UNiagaraHierarchyCategory(FName InCategory) : Category(InCategory) {}
+	
+	void SetCategoryName(FName NewCategory) { Category = NewCategory; }
+	FName GetCategoryName() const { return Category; }
+
+	FText GetTooltip() const { return Tooltip; }
+
+	void SetSection(UNiagaraHierarchySection* InSection) { Section = InSection; }
+	const UNiagaraHierarchySection* GetSection() const { return Section; }
+
+	virtual bool CanHaveChildren() const override { return true; }
+
+	virtual FString ToString() const override { return Category.ToString(); }
+
+private:
+	UPROPERTY()
+	FName Category;
+
+	/** The tooltip used when the user is hovering this category */
+	UPROPERTY(EditAnywhere, Category = "Category")
+	FText Tooltip;
+
+	UPROPERTY()
+	TObjectPtr<UNiagaraHierarchySection> Section = nullptr;
+};
+
+UCLASS()
+class UNiagaraHierarchySection : public UNiagaraHierarchyItemBase
+{
+	GENERATED_BODY()
+
+public:
+	UNiagaraHierarchySection() {}
+
+	void SetSectionName(FName InSectionName) { Section = InSectionName; }
+	FName GetSectionName() const { return Section; }
+	
+	void SetSectionNameAsText(const FText& Text);
+	FText GetSectionNameAsText() const { return FText::FromName(Section); }
+
+	FText GetTooltip() const { return Tooltip; }
+
+private:
+	UPROPERTY()
+	FName Section;
+
+	/** The tooltip used when the user is hovering this section */
+	UPROPERTY(EditAnywhere, Category = "Section")
+	FText Tooltip;
+};
+
+struct FNiagaraHierarchyItemViewModelBase;
+
+class FNiagaraHierarchyDragDropOp : public FDragDropOperation
+{
+public:
+	DRAG_DROP_OPERATOR_TYPE(FNiagaraHierarchyDragDropOp, FDragDropOperation)
+
+	FNiagaraHierarchyDragDropOp(TSharedPtr<FNiagaraHierarchyItemViewModelBase> ItemBase) : DraggedItem(ItemBase) {}
+
+	virtual void Construct() override { FDragDropOperation::Construct(); }
+	virtual TSharedPtr<SWidget> GetDefaultDecorator() const override final;
+	
+	/** Override this custom decorator function to provide custom widget visuals */
+	virtual TSharedRef<SWidget> CreateCustomDecorator() const { return SNullWidget::NullWidget; }
+	
+	TWeakPtr<FNiagaraHierarchyItemViewModelBase> GetDraggedItem() { return DraggedItem; }
+	void SetAdditionalLabel(FText InText) { AdditionalLabel = InText; }
+	FText GetAdditionalLabel() const { return AdditionalLabel; }
+
+	void SetDescription(FText InText) { Description = InText; }
+	FText GetDescription() const { return Description; }
+
+	void SetFromSourceList(bool bInFromSourceList) { bFromSourceList = bInFromSourceList; }
+	bool GetIsFromSourceList() const { return bFromSourceList; }
+protected:
+	TWeakPtr<FNiagaraHierarchyItemViewModelBase> DraggedItem;
+	/** Additional label will be displayed below the custom decorator. Useful for runtime tweaking of the tooltip based on what we are hovering. */
+	FText AdditionalLabel;
+	FText Description;
+	/** If the drag drop op is from the source list, we can further customize the actions */
+	bool bFromSourceList = false;
+};
+
+UCLASS(Abstract)
+class UNiagaraHierarchyViewModelBase : public UObject, public FEditorUndoClient
+{
+public:
+	DECLARE_DELEGATE_OneParam(FRefreshTreeWidget, bool bFullRefresh)
+	DECLARE_DELEGATE_RetVal_OneParam(bool, FIsItemSelected, TSharedPtr<FNiagaraHierarchyItemViewModelBase> Item)
+	DECLARE_DELEGATE_OneParam(FSelectObjectInDetailsPanel, UObject* Object)
+
+	GENERATED_BODY()
+	
+	UNiagaraHierarchyViewModelBase();
+	virtual ~UNiagaraHierarchyViewModelBase() override;
+
+	void Initialize(TSharedPtr<FNiagaraSystemViewModel> InSystemViewModel);
+	void Finalize();
+
+	TSharedPtr<FNiagaraSystemViewModel> GetSystemViewModel() const { return SystemViewModel; }
+
+	void AddCategory() const;
+	void AddSection() const;
+
+	/** Refreshes all data and widgets */
+	void ForceFullRefresh();
+
+	/** The hierarchy root the widget is editing. This should point to persistent data stored somewhere else as the serialized root of the hierarchy. */
+	virtual UNiagaraHierarchyRoot* GetHierarchyDataRoot() const PURE_VIRTUAL(UNiagaraHierarchyViewModelBase::GetHierarchyDataRoot, return nullptr;);
+	TSharedPtr<struct FNiagaraHierarchyRootViewModel> GetHierarchyViewModelRoot() const { return HierarchyViewModelRoot; }
+	
+	/** Source items reflect the base, unedited status of items to edit into a hierarchy */
+	const TArray<TSharedPtr<FNiagaraHierarchyItemViewModelBase>>& GetSourceItems() const;
+	/** Hierarchy items reflect the already edited hierarchy. This should generally be constructed from persistent serialized data. */
+	const TArray<TSharedPtr<FNiagaraHierarchyItemViewModelBase>>& GetHierarchyItems() const;
+	
+	/** Prepares the items we want to create a hierarchy for. The idea is to Drag & Drop these into the hierarchy or click on them to edit details.  */
+	virtual void PrepareSourceItems() PURE_VIRTUAL(UNiagaraHierarchyViewModelBase::PrepareSourceItems,);
+
+	/** Additional commands can be specified overriding the Commands function. */
+	virtual void SetupCommands() {}
+	TSharedRef<FUICommandList> GetCommands() const { return Commands.ToSharedRef(); }
+	
+	/** Central function to determine whether an item can be dropped on, above or below another item. */
+	virtual TOptional<EItemDropZone> CanDropOn(TSharedPtr<FNiagaraHierarchyItemViewModelBase> SourceDropItem, TSharedPtr<FNiagaraHierarchyItemViewModelBase> TargetDropItem, EItemDropZone DropZone = EItemDropZone::OntoItem) PURE_VIRTUAL(UNiagaraHierarchyViewModelBase::CanDropOn, return TOptional<EItemDropZone>(););
+	virtual TSharedRef<FNiagaraHierarchyDragDropOp> CreateDragDropOp(TSharedRef<FNiagaraHierarchyItemViewModelBase> Item) PURE_VIRTUAL(UNiagaraHierarchyViewModelBase::CreateDragDropOp, return MakeShared<FNiagaraHierarchyDragDropOp>(nullptr););
+	
+	virtual void OnSelectionChanged(TSharedPtr<FNiagaraHierarchyItemViewModelBase> HierarchyItem) {}
+	
+	void OnGetChildren(TSharedPtr<FNiagaraHierarchyItemViewModelBase> Item, TArray<TSharedPtr<FNiagaraHierarchyItemViewModelBase>>& OutChildren) const;
+	
+	virtual bool SupportsDetailsPanel() { return false; }
+	// Overriding this will give the details panel instance customizations for specific UClasses
+	virtual TArray<TTuple<UClass*, FOnGetDetailCustomizationInstance>> GetInstanceCustomizations() { return {}; }
+	// Child view models need to override this to handle their own row generation. By default, we provide category widgets
+	virtual TSharedRef<SWidget> GenerateRowContentWidget(TSharedRef<FNiagaraHierarchyItemViewModelBase>) const;
+	
+	void RefreshSourceView(bool bFullRefresh = false) const;
+	void RefreshHierarchyView(bool bFullRefresh = false) const;
+	void RefreshSectionsWidget() const;
+	void SelectItemInDetailsPanel(UObject* Object) const;
+	bool IsItemSelected(TSharedPtr<FNiagaraHierarchyItemViewModelBase> Item) const;
+	
+	// Delegate that call functions from SNiagaraHierarchy
+	FRefreshTreeWidget& OnRefreshSourceView() { return RefreshSourceViewDelegate; }
+	FRefreshTreeWidget& OnRefreshHierarchyView() { return RefreshHierarchyWidgetDelegate; }
+	FSimpleDelegate& OnRefreshSections() { return RefreshSectionsWidgetDelegate; }
+	FSelectObjectInDetailsPanel& OnSelectObjectInDetailsPanel() { return SelectObjectInDetailsPanelDelegate; }
+	FIsItemSelected& IsItemSelected() { return IsItemSelectedDelegate; }
+
+	// Sections
+	void SetActiveSection(TSharedPtr<struct FNiagaraHierarchySectionViewModel>);
+	TSharedPtr<FNiagaraHierarchySectionViewModel> GetActiveSection() const;
+	UNiagaraHierarchySection* GetActiveSectionData() const;
+	bool IsSectionActive(const UNiagaraHierarchySection* Section) const;
+	
+	FString OnItemToStringDebug(TSharedPtr<FNiagaraHierarchyItemViewModelBase> ItemBaseViewModel) const; 
+
+protected:
+	virtual void InitializeInternal() {}
+	virtual void FinalizeInternal() {}
+
+	virtual void PostUndo(bool bSuccess) override;
+	virtual void PostRedo(bool bSuccess) override;
+protected:
+	UPROPERTY()
+	TObjectPtr<UNiagaraHierarchyRoot> SourceRoot;
+
+	UPROPERTY()
+	TObjectPtr<UNiagaraHierarchyRoot> HierarchyRoot;
+
+	TSharedPtr<struct FNiagaraHierarchyRootViewModel> SourceViewModelRoot;
+	TSharedPtr<struct FNiagaraHierarchyRootViewModel> HierarchyViewModelRoot;
+
+	TWeakPtr<struct FNiagaraHierarchySectionViewModel> ActiveSection;
+	TSharedPtr<FUICommandList> Commands;
+
+	TSharedPtr<FNiagaraSystemViewModel> SystemViewModel;
+
+protected:
+	// delegate collection to call UI functions
+	FRefreshTreeWidget RefreshSourceViewDelegate;
+	FRefreshTreeWidget RefreshHierarchyWidgetDelegate;
+	FSimpleDelegate RefreshSectionsWidgetDelegate;
+	FIsItemSelected IsItemSelectedDelegate;
+	FSelectObjectInDetailsPanel SelectObjectInDetailsPanelDelegate;
+};
+
+struct FNiagaraHierarchyItemViewModelBase : TSharedFromThis<FNiagaraHierarchyItemViewModelBase>
+{
+	DECLARE_MULTICAST_DELEGATE(FOnSynced)
+	
+	FNiagaraHierarchyItemViewModelBase(UNiagaraHierarchyItemBase* InItemBase, TSharedPtr<FNiagaraHierarchyItemViewModelBase> InParent, TWeakObjectPtr<UNiagaraHierarchyViewModelBase> InHierarchyViewModel)
+		: ItemBase(InItemBase)
+		, Parent(InParent)
+		, HierarchyViewModel(InHierarchyViewModel)
+	{		
+		
+	}
+	
+	virtual ~FNiagaraHierarchyItemViewModelBase() {}
+
+	UNiagaraHierarchyItemBase* GetDataMutable() const { return ItemBase; }
+	const UNiagaraHierarchyItemBase* GetData() const { return ItemBase; }
+
+	FString ToString() const { return ItemBase->ToString(); }
+	
+	void SyncToData();
+	const TArray<TSharedPtr<FNiagaraHierarchyItemViewModelBase>>& GetChildren() const { return Children; }
+	const TArray<TSharedPtr<FNiagaraHierarchyItemViewModelBase>>& GetFilteredChildren() const;
+
+	template<class DataClass, class ViewModelChildClass>
+	void GetChildrenViewModelsForType(TArray<TSharedPtr<ViewModelChildClass>>& OutChildren, bool bRecursive = false);
+
+	bool HasParent(TSharedPtr<FNiagaraHierarchyItemViewModelBase> ParentCandidate, bool bRecursive = false);
+
+	void AddChild(TSharedPtr<FNiagaraHierarchyItemViewModelBase> Item);
+	UNiagaraHierarchyItemBase* AddNewItem(TSubclassOf<UNiagaraHierarchyItemBase> NewItemClass);
+	void DuplicateToThis(TSharedPtr<FNiagaraHierarchyItemViewModelBase> ItemToDuplicate, int32 InsertIndex = INDEX_NONE);
+	void ReparentToThis(TSharedPtr<FNiagaraHierarchyItemViewModelBase> ItemToMove, int32 InsertIndex = INDEX_NONE);
+	int32 FindIndexOfChild(TSharedPtr<FNiagaraHierarchyItemViewModelBase> Child) const;
+	int32 FindIndexOfChild(UNiagaraHierarchyItemBase* Child) const;
+	int32 FindIndexOfDataChild(TSharedPtr<FNiagaraHierarchyItemViewModelBase> Child) const;
+	int32 FindIndexOfDataChild(UNiagaraHierarchyItemBase* Child) const;
+
+	/** Finalizing the view model will finalize the data & the children as well as handle cleanup. */
+	void Finalize();
+
+	TWeakObjectPtr<UNiagaraHierarchyViewModelBase> GetHierarchyViewModel() const { return HierarchyViewModel; }
+	
+	virtual FReply OnDroppedOn(const FDragDropEvent& DragDropEvent, EItemDropZone ItemDropZone, TSharedPtr<FNiagaraHierarchyItemViewModelBase> Item)
+	{
+		return FReply::Unhandled();
+	}
+	
+	TOptional<EItemDropZone> OnCanAcceptDrop(const FDragDropEvent& DragDropEvent, EItemDropZone ItemDropZone, TSharedPtr<FNiagaraHierarchyItemViewModelBase> Item);
+
+	/** Should return true if draggable */
+	virtual bool CanDrag() { return false; }
+
+	/** Should return true if renamable */
+	virtual bool CanRename() { return false; }
+
+	virtual bool CanDelete() { return true; }
+
+	void RequestRename()
+	{
+		if(CanRename())
+		{
+			OnRequestRenameDelegate.Execute();
+		}
+	}
+	
+	/** Used to create customized drag drop ops. */
+	TSharedRef<class FNiagaraHierarchyDragDropOp> CreateDragDropOp();
+	
+	FReply OnDragDetected(const FGeometry& Geometry, const FPointerEvent& PointerEvent, bool bIsSource);
+
+	FSimpleDelegate& GetOnRequestRename() { return OnRequestRenameDelegate; }
+	FOnSynced& GetOnSynced() { return OnSyncedDelegate; }
+
+	TWeakPtr<FNiagaraHierarchyItemViewModelBase> GetParent() { return Parent; }
+
+private:
+	virtual void SyncChildrenViewModelsInternal() {}
+	virtual void FinalizeInternal() {}
+protected:
+	UNiagaraHierarchyItemBase* const ItemBase;
+	/** Parent should be valid for all instances of this struct except for root objects */
+	TWeakPtr<FNiagaraHierarchyItemViewModelBase> Parent;
+	TArray<TSharedPtr<FNiagaraHierarchyItemViewModelBase>> Children;
+	mutable TArray<TSharedPtr<FNiagaraHierarchyItemViewModelBase>> FilteredChildren;
+	TWeakObjectPtr<UNiagaraHierarchyViewModelBase> HierarchyViewModel;
+	FSimpleDelegate OnRequestRenameDelegate;
+	FOnSynced OnSyncedDelegate;
+};
+
+template <class DataClass, class ViewModelClass>
+void FNiagaraHierarchyItemViewModelBase::GetChildrenViewModelsForType(TArray<TSharedPtr<ViewModelClass>>& OutChildren, bool bRecursive)
+{
+	for(auto& Child : Children)
+	{
+		if(Child->GetData()->IsA<DataClass>())
+		{
+			OutChildren.Add(StaticCastSharedPtr<ViewModelClass>(Child));
+		}
+	}
+
+	if(bRecursive)
+	{
+		for(auto& Child : Children)
+		{
+			Child->GetChildrenViewModelsForType<DataClass, ViewModelClass>(OutChildren, bRecursive);
+		}
+	}
+}
+
+struct FNiagaraHierarchyRootViewModel : FNiagaraHierarchyItemViewModelBase
+{
+	FNiagaraHierarchyRootViewModel(UNiagaraHierarchyItemBase* InItem, TWeakObjectPtr<UNiagaraHierarchyViewModelBase> InHierarchyViewModel) : FNiagaraHierarchyItemViewModelBase(InItem, nullptr, InHierarchyViewModel) {}
+	
+	virtual ~FNiagaraHierarchyRootViewModel() override {}
+
+	virtual bool CanDrag() override { return false; }
+
+	TSharedPtr<struct FNiagaraHierarchySectionViewModel> AddNewSection();
+	TArray<TSharedPtr<struct FNiagaraHierarchySectionViewModel>>& GetSectionViewModels() { return SectionViewModels; }
+
+private:
+	virtual void SyncChildrenViewModelsInternal() override;
+	TArray<TSharedPtr<struct FNiagaraHierarchySectionViewModel>> SectionViewModels;
+};
+
+struct FNiagaraHierarchySectionViewModel : FNiagaraHierarchyItemViewModelBase
+{
+	FNiagaraHierarchySectionViewModel(UNiagaraHierarchySection* InItem, TSharedPtr<FNiagaraHierarchyRootViewModel> InParent, TWeakObjectPtr<UNiagaraHierarchyViewModelBase> InHierarchyViewModel) : FNiagaraHierarchyItemViewModelBase(InItem, InParent, InHierarchyViewModel) {}
+	
+	virtual ~FNiagaraHierarchySectionViewModel() override {}
+
+	void SetSectionName(FName InSectionName) const { Cast<UNiagaraHierarchySection>(ItemBase)->SetSectionName(InSectionName); }
+	FName GetSectionName() const { return Cast<UNiagaraHierarchySection>(ItemBase)->GetSectionName(); }
+	
+	void SetSectionNameAsText(const FText& Text) const { Cast<UNiagaraHierarchySection>(ItemBase)->SetSectionNameAsText(Text); }
+	FText GetSectionNameAsText() const { return Cast<UNiagaraHierarchySection>(ItemBase)->GetSectionNameAsText(); }
+	
+	virtual bool CanDrag() override { return false; }
+	virtual bool CanRename() override { return true; }
+	virtual bool CanDelete() override { return true; }
+
+	virtual FReply OnDroppedOn(const FDragDropEvent& DragDropEvent, EItemDropZone ItemDropZone, TSharedPtr<FNiagaraHierarchyItemViewModelBase> Item) override;
+
+	virtual void FinalizeInternal() override;
+};
+
+struct FNiagaraHierarchyItemViewModel : FNiagaraHierarchyItemViewModelBase
+{
+	FNiagaraHierarchyItemViewModel(UNiagaraHierarchyItem* InItem, TSharedPtr<FNiagaraHierarchyItemViewModelBase> InParent, TWeakObjectPtr<UNiagaraHierarchyViewModelBase> InHierarchyViewModel) : FNiagaraHierarchyItemViewModelBase(InItem, InParent, InHierarchyViewModel) {}
+	
+	virtual ~FNiagaraHierarchyItemViewModel() override {}
+
+	virtual bool CanDrag() override { return true; }
+	
+	virtual FReply OnDroppedOn(const FDragDropEvent& DragDropEvent, EItemDropZone ItemDropZone, TSharedPtr<FNiagaraHierarchyItemViewModelBase> Item) override;
+};
+
+struct FNiagaraHierarchyCategoryViewModel : FNiagaraHierarchyItemViewModelBase
+{
+	FNiagaraHierarchyCategoryViewModel(UNiagaraHierarchyCategory* InCategory, TSharedPtr<FNiagaraHierarchyItemViewModelBase> InParent, TWeakObjectPtr<UNiagaraHierarchyViewModelBase> InHierarchyViewModel) : FNiagaraHierarchyItemViewModelBase(InCategory, InParent, InHierarchyViewModel) {}
+	virtual ~FNiagaraHierarchyCategoryViewModel() override{}
+
+	virtual bool CanDrag() override { return true; }
+	virtual bool CanRename() override { return true; }
+
+	const UNiagaraHierarchySection* GetSection() const { return Cast<UNiagaraHierarchyCategory>(ItemBase)->GetSection(); }
+
+	bool IsTopCategoryActive() const;
+
+	virtual FReply OnDroppedOn(const FDragDropEvent& DragDropEvent, EItemDropZone ItemDropZone, TSharedPtr<FNiagaraHierarchyItemViewModelBase> Item) override;
+};
+
+class FNiagaraSectionDragDropOp : public FNiagaraHierarchyDragDropOp
+{
+public:
+	DRAG_DROP_OPERATOR_TYPE(FNiagaraSectionDragDropOp, FNiagaraHierarchyDragDropOp)
+
+	FNiagaraSectionDragDropOp(TSharedPtr<FNiagaraHierarchySectionViewModel> ItemBase) : FNiagaraHierarchyDragDropOp(ItemBase) {}
+	
+	TWeakPtr<FNiagaraHierarchySectionViewModel> GetDraggedSection() const { return StaticCastSharedPtr<FNiagaraHierarchySectionViewModel>(DraggedItem.Pin()); }
+private:
+	virtual TSharedRef<SWidget> CreateCustomDecorator() const override;
+};
diff --git a/Engine/Plugins/FX/Niagara/Source/NiagaraEditor/Public/ViewModels/HierarchyEditor/NiagaraUserParametersHierarchyViewModel.h b/Engine/Plugins/FX/Niagara/Source/NiagaraEditor/Public/ViewModels/HierarchyEditor/NiagaraUserParametersHierarchyViewModel.h
new file mode 100644
--- /dev/null
+++ b/Engine/Plugins/FX/Niagara/Source/NiagaraEditor/Public/ViewModels/HierarchyEditor/NiagaraUserParametersHierarchyViewModel.h
@@ -0,0 +1,78 @@
+// Copyright Epic Games, Inc. All Rights Reserved.
+
+#pragma once
+
+#include "ViewModels/HierarchyEditor/NiagaraHierarchyViewModelBase.h"
+#include "NiagaraScriptVariable.h"
+#include "NiagaraUserParametersHierarchyViewModel.generated.h"
+
+UCLASS()
+class UNiagaraHierarchyUserParameter : public UNiagaraHierarchyItem
+{
+	GENERATED_BODY()
+
+public:
+	UNiagaraHierarchyUserParameter() {}
+	virtual ~UNiagaraHierarchyUserParameter() override {}
+	
+	void Initialize(UNiagaraScriptVariable& InUserParameterScriptVariable, UNiagaraSystem& System);
+
+	virtual void RefreshDataInternal() override;
+	
+	const FNiagaraVariable& GetUserParameter() const { return UserParameterScriptVariable->Variable; }
+
+	virtual FString ToString() const override { return GetUserParameter().GetName().ToString(); }
+private:
+	UPROPERTY()
+	TObjectPtr<UNiagaraScriptVariable> UserParameterScriptVariable;
+
+	UPROPERTY()
+	TObjectPtr<UNiagaraSystem> System;
+};
+
+UCLASS()
+class UNiagaraUserParametersHierarchyViewModel : public UNiagaraHierarchyViewModelBase
+{
+	GENERATED_BODY()
+public:
+	UNiagaraUserParametersHierarchyViewModel() {}
+	virtual ~UNiagaraUserParametersHierarchyViewModel() override
+	{
+	}
+
+	virtual void InitializeInternal() override;
+	
+	virtual UNiagaraHierarchyRoot* GetHierarchyDataRoot() const override;
+	
+	virtual void PrepareSourceItems() override;
+	virtual void SetupCommands() override;
+
+	virtual TOptional<EItemDropZone> CanDropOn(TSharedPtr<FNiagaraHierarchyItemViewModelBase> SourceDropItem, TSharedPtr<FNiagaraHierarchyItemViewModelBase> TargetDropItem, EItemDropZone DropZone = EItemDropZone::OntoItem) override;
+	virtual TSharedRef<FNiagaraHierarchyDragDropOp> CreateDragDropOp(TSharedRef<FNiagaraHierarchyItemViewModelBase> Item) override;
+
+	virtual void OnSelectionChanged(TSharedPtr<FNiagaraHierarchyItemViewModelBase> HierarchyItem) override;
+	
+	virtual bool SupportsDetailsPanel() override { return true; }
+	virtual TArray<TTuple<UClass*, FOnGetDetailCustomizationInstance>> GetInstanceCustomizations() override;
+	virtual TSharedRef<SWidget> GenerateRowContentWidget(TSharedRef<FNiagaraHierarchyItemViewModelBase>) const override;
+protected:
+	virtual void FinalizeInternal() override;
+private:
+	TWeakPtr<FNiagaraSystemViewModel> SystemViewModel;
+};
+
+class FNiagaraUserParameterDragDropOp : public FNiagaraHierarchyDragDropOp
+{
+public:
+	DRAG_DROP_OPERATOR_TYPE(FNiagaraUserParameterDragDropOp, FNiagaraHierarchyDragDropOp)
+
+	FNiagaraUserParameterDragDropOp(TSharedPtr<FNiagaraHierarchyItemViewModelBase> UserParameterItem) : FNiagaraHierarchyDragDropOp(UserParameterItem) {}
+
+	FNiagaraVariable GetUserParameter() const
+	{
+		const UNiagaraHierarchyUserParameter* HierarchyUserParameter = CastChecked<UNiagaraHierarchyUserParameter>(DraggedItem.Pin()->GetData());
+		return HierarchyUserParameter->GetUserParameter();
+	}
+	
+	virtual TSharedRef<SWidget> CreateCustomDecorator() const override;
+};
diff --git a/Engine/Plugins/FX/Niagara/Source/NiagaraEditor/Public/ViewModels/NiagaraParameterPanelViewModel.h b/Engine/Plugins/FX/Niagara/Source/NiagaraEditor/Public/ViewModels/NiagaraParameterPanelViewModel.h
--- a/Engine/Plugins/FX/Niagara/Source/NiagaraEditor/Public/ViewModels/NiagaraParameterPanelViewModel.h
+++ b/Engine/Plugins/FX/Niagara/Source/NiagaraEditor/Public/ViewModels/NiagaraParameterPanelViewModel.h
@@ -253,9 +253,15 @@ public:
 
 	const TArray<FNiagaraParameterPanelItem>& GetCachedViewedParameterItems() const;
 
-	void SelectParameterItemByName(const FName ParameterName, const bool bRequestRename);
+	virtual bool CanSummonHierarchyView() const { return false; };
 
+	virtual void SummonHierarchyView() const {}
+	
+	void SelectParameterItemByName(const FName ParameterName, const bool bRequestRename);
 
+	// TODO (me) Remove again when the user parameter tab goes live, since this is only used for the "Edit Hierarchy" button
+	virtual TSharedRef<SWidget> GenerateAdjacentWidget() { return SNullWidget::NullWidget; }
+	
 	FOnParameterPanelViewModelExternalSelectionChanged& GetOnExternalSelectionChangedDelegate() { return OnParameterPanelViewModelExternalSelectionChangedDelegate; };
 
 	FOnSelectParameterItemByName& GetOnSelectParameterItemByNameDelegate() { return OnSelectParameterItemByNameDelegate; };
@@ -345,6 +351,8 @@ public:
 
 	virtual TSharedPtr<SWidget> CreateContextMenuForItems(const TArray<FNiagaraParameterPanelItem>& Items, const TSharedPtr<FUICommandList>& ToolkitCommands) override;
 
+	virtual TSharedRef<SWidget> GenerateAdjacentWidget() override;
+	
 	virtual FNiagaraParameterUtilities::EParameterContext GetParameterContext() const override;
 
 	virtual TArray<FNiagaraVariable> GetEditableStaticSwitchParameters() const override;
@@ -376,6 +384,7 @@ public:
 
 	virtual bool GetSectionEnabled(FText Section) const override;
 
+	bool IsUserSectionActive() const { return ActiveSectionIndex == UserOnlyIdx && UserOnlyIdx != -1; }
 
 	virtual bool GetShowSections() const override { return true; }
 	virtual bool GetNamespaceActive(const FName& InNamespace) const override { return true; }
@@ -389,6 +398,9 @@ public:
 	virtual bool UsesCategoryFilteringForInitialExpansion() const override { return true; }
 
 	virtual void PreSectionChange(const TArray<FNiagaraParameterPanelCategory>& ExpandedItems) override;
+	virtual bool CanSummonHierarchyView() const override { return true; }
+
+	virtual void SummonHierarchyView() const override;
 private:
 
 	bool  ShouldRouteThroughScratchParameterMap(const FNiagaraParameterPanelCategory* Category = nullptr, const FNiagaraVariableBase* Variable = nullptr);
diff --git a/Engine/Plugins/FX/Niagara/Source/NiagaraEditor/Public/ViewModels/NiagaraSystemViewModel.h b/Engine/Plugins/FX/Niagara/Source/NiagaraEditor/Public/ViewModels/NiagaraSystemViewModel.h
--- a/Engine/Plugins/FX/Niagara/Source/NiagaraEditor/Public/ViewModels/NiagaraSystemViewModel.h
+++ b/Engine/Plugins/FX/Niagara/Source/NiagaraEditor/Public/ViewModels/NiagaraSystemViewModel.h
@@ -16,7 +16,6 @@
 #include "ViewModels/TNiagaraViewModelManager.h"
 #include "ViewModels/NiagaraParameterDefinitionsSubscriberViewModel.h"
 
-
 struct FAssetData;
 struct FEdGraphEditAction;
 struct FNiagaraEmitterHandle;
@@ -50,6 +49,7 @@ class UNiagaraSystem;
 class UNiagaraSystemEditorData;
 class UNiagaraSystemEditorDocumentsViewModel;
 class UNiagaraSystemSelectionViewModel;
+class UNiagaraUserParametersHierarchyViewModel;
 
 
 /** Defines different editing modes for this system view model. */
@@ -385,6 +385,8 @@ public:
 
 	NIAGARAEDITOR_API UNiagaraSystemScalabilityViewModel* GetScalabilityViewModel();
 
+	NIAGARAEDITOR_API UNiagaraUserParametersHierarchyViewModel* GetUserParametersHierarchyViewModel();
+	
 	TArray<float> OnGetPlaybackSpeeds() const;
 	
 	/** Duplicates a set of emitters and refreshes everything.*/
@@ -715,6 +717,8 @@ private:
 
 	UNiagaraSystemScalabilityViewModel* ScalabilityViewModel;
 
+	UNiagaraUserParametersHierarchyViewModel* UserParametersHierarchyViewModel;
+	
 	TWeakPtr<INiagaraParameterPanelViewModel> ParameterPanelViewModel;
 
 	TSharedPtr<FNiagaraPlaceholderDataInterfaceManager> PlaceholderDataInterfaceManager;
diff --git a/Engine/Plugins/FX/Niagara/Source/NiagaraEditor/Public/Widgets/SNiagaraHierarchy.h b/Engine/Plugins/FX/Niagara/Source/NiagaraEditor/Public/Widgets/SNiagaraHierarchy.h
new file mode 100644
--- /dev/null
+++ b/Engine/Plugins/FX/Niagara/Source/NiagaraEditor/Public/Widgets/SNiagaraHierarchy.h
@@ -0,0 +1,131 @@
+// Copyright Epic Games, Inc. All Rights Reserved.
+#pragma once
+
+#include "CoreMinimal.h"
+#include "SDropTarget.h"
+#include "IDetailsView.h"
+#include "Widgets/Views/STreeView.h"
+#include "Widgets/Views/STableRow.h"
+#include "Widgets/Text/SInlineEditableTextBlock.h"
+#include "Widgets/SBoxPanel.h"
+#include "Styling/SlateTypes.h"
+
+class UNiagaraHierarchyViewModelBase;
+class UNiagaraHierarchySection;
+struct FNiagaraHierarchyItemViewModelBase;
+struct FNiagaraHierarchySectionViewModel;
+
+class SNiagaraHierarchy : public SCompoundWidget
+{
+	SLATE_BEGIN_ARGS(SNiagaraHierarchy)
+		: _bReadOnly(true)
+	{}
+		SLATE_ARGUMENT(bool, bReadOnly)
+	SLATE_END_ARGS()
+
+public:
+	void Construct(const FArguments& InArgs, TObjectPtr<UNiagaraHierarchyViewModelBase> InHierarchyViewModel);
+	virtual ~SNiagaraHierarchy() override;
+	
+	virtual FReply OnKeyDown(const FGeometry& MyGeometry, const FKeyEvent& InKeyEvent) override;
+
+	void RefreshSourceView(bool bFullRefresh = false) const;
+	void RefreshHierarchyView(bool bFullRefresh = false) const;
+	void RefreshSectionsWidget();
+	bool IsItemSelected(TSharedPtr<FNiagaraHierarchyItemViewModelBase> Item) const;
+	void SelectObjectInDetailsPanel(UObject* Object) const;
+	
+private:
+	TSharedRef<SWidget> CreateHierarchyButtonWidgets() const;
+	TSharedRef<SWidget> CreateSectionWidgets();
+	TSharedRef<ITableRow> GenerateSourceItemRow(TSharedPtr<FNiagaraHierarchyItemViewModelBase> HierarchyItem, const TSharedRef<STableViewBase>& TableViewBase);
+	TSharedRef<ITableRow> GenerateHierarchyItemRow(TSharedPtr<FNiagaraHierarchyItemViewModelBase> HierarchyItem, const TSharedRef<STableViewBase>& TableViewBase);
+private:	
+	void RequestRenameSelectedItem();
+	bool CanRequestRenameSelectedItem() const;
+
+	void DeleteSelectedItems();
+	bool CanDeleteSelectedItems() const;
+
+	void OnSelectionChanged(TSharedPtr<FNiagaraHierarchyItemViewModelBase> HierarchyItem, ESelectInfo::Type Type) const;
+private:
+	TWeakObjectPtr<UNiagaraHierarchyViewModelBase> HierarchyViewModel;
+	TSharedPtr<STreeView<TSharedPtr<FNiagaraHierarchyItemViewModelBase>>> SourceTreeView;
+	TSharedPtr<STreeView<TSharedPtr<FNiagaraHierarchyItemViewModelBase>>> HierarchyTreeView;
+	TMap<TSharedPtr<FNiagaraHierarchySectionViewModel>, TSharedPtr<class SNiagaraHierarchySection>> SectionsWidgetMap;
+	SVerticalBox::FSlot* SectionsSlot = nullptr;
+	TSharedPtr<IDetailsView> DetailsPanel;
+};
+
+class SNiagaraHierarchySection : public SCompoundWidget
+{
+	SLATE_BEGIN_ARGS(SNiagaraHierarchySection)
+	{}
+		SLATE_EVENT(FIsSelected, IsSelected)
+	SLATE_END_ARGS()
+
+public:
+	void Construct(const FArguments& InArgs, TSharedPtr<struct FNiagaraHierarchySectionViewModel> InSection, TWeakObjectPtr<UNiagaraHierarchyViewModelBase> HierarchyViewModel);
+	virtual ~SNiagaraHierarchySection() override;
+	
+
+	bool OnCanAcceptDrop(TSharedPtr<FDragDropOperation> DragDropOperation) const;
+	bool OnCanAcceptSectionDrop(TSharedPtr<FDragDropOperation> DragDropOperation, EItemDropZone ItemDropZone) const;
+	FReply OnDroppedOn(const FGeometry&, const FDragDropEvent& DragDropEvent, EItemDropZone DropZone) const;
+
+	void EnterEditingMode() const;
+	virtual FReply OnMouseButtonDown(const FGeometry& MyGeometry, const FPointerEvent& MouseEvent) override;
+	virtual FReply OnDragDetected(const FGeometry& MyGeometry, const FPointerEvent& MouseEvent) override;
+	virtual void OnDragLeave(const FDragDropEvent& DragDropEvent) override;
+	
+	TSharedRef<SWidget> OnGetMenuContent() const;
+
+	UNiagaraHierarchySection* TryGetSectionData() const;
+
+	FText GetText() const;
+	void OnRenameSection(const FText& Text, ETextCommit::Type CommitType) const;
+	bool OnVerifySectionRename(const FText& NewName, FText& OutTooltip) const;
+
+	ECheckBoxState GetSectionCheckState() const;
+	void OnSectionCheckChanged(ECheckBoxState NewState);
+private:
+	TSharedPtr<SInlineEditableTextBlock> InlineEditableTextBlock;
+	TWeakObjectPtr<UNiagaraHierarchyViewModelBase> HierarchyViewModel;
+	TSharedPtr<FNiagaraHierarchySectionViewModel> SectionViewModel;
+};
+
+class SNiagaraHierarchyCategory : public SCompoundWidget
+{
+	SLATE_BEGIN_ARGS(SNiagaraHierarchyCategory)
+	{}
+		SLATE_EVENT(FIsSelected, IsSelected)
+	SLATE_END_ARGS()
+
+public:
+	void Construct(const FArguments& InArgs, TSharedPtr<struct FNiagaraHierarchyCategoryViewModel> InCategory);
+
+	void EnterEditingMode() const;
+	bool OnVerifyCategoryRename(const FText& NewName, FText& OutTooltip) const;
+	
+	FText GetCategoryText() const;
+	void OnRenameCategory(const FText& NewText, ETextCommit::Type) const;
+	
+private:
+	TWeakPtr<struct FNiagaraHierarchyCategoryViewModel> Category;
+	TSharedPtr<SInlineEditableTextBlock> InlineEditableTextBlock;
+};
+
+class SNiagaraSectionDragDropTarget : public SDropTarget
+{
+public:
+	SLATE_BEGIN_ARGS(SNiagaraSectionDragDropTarget)
+		: _DropTargetArgs(SDropTarget::FArguments())
+	{}
+		SLATE_ARGUMENT(SDropTarget::FArguments, DropTargetArgs)
+	SLATE_END_ARGS()
+
+	void Construct(const FArguments& InArgs);
+	
+	virtual FReply OnDragOver(const FGeometry& MyGeometry, const FDragDropEvent& DragDropEvent) override;
+	virtual void OnDragLeave(const FDragDropEvent& DragDropEvent) override;
+};
diff --git a/Engine/Plugins/FX/Niagara/Source/NiagaraEditorWidgets/Private/SNiagaraParameterDropTarget.h b/Engine/Plugins/FX/Niagara/Source/NiagaraEditorWidgets/Private/SNiagaraParameterDropTarget.h
--- a/Engine/Plugins/FX/Niagara/Source/NiagaraEditorWidgets/Private/SNiagaraParameterDropTarget.h
+++ b/Engine/Plugins/FX/Niagara/Source/NiagaraEditorWidgets/Private/SNiagaraParameterDropTarget.h
@@ -16,7 +16,6 @@
 class NIAGARAEDITORWIDGETS_API SNiagaraParameterDropTarget : public SDropTarget
 {
 public:
-    DECLARE_DELEGATE_OneParam(FOnDragOperationDragged, TSharedPtr<FDragDropOperation>)
     SLATE_BEGIN_ARGS(SNiagaraParameterDropTarget)
 	    : _DropTargetArgs(SDropTarget::FArguments())
 		, _TargetParameter(TOptional<FNiagaraVariable>())
diff --git a/Engine/Plugins/FX/Niagara/Source/NiagaraEditorWidgets/Private/Stack/SNiagaraStackFunctionInputValue.cpp b/Engine/Plugins/FX/Niagara/Source/NiagaraEditorWidgets/Private/Stack/SNiagaraStackFunctionInputValue.cpp
--- a/Engine/Plugins/FX/Niagara/Source/NiagaraEditorWidgets/Private/Stack/SNiagaraStackFunctionInputValue.cpp
+++ b/Engine/Plugins/FX/Niagara/Source/NiagaraEditorWidgets/Private/Stack/SNiagaraStackFunctionInputValue.cpp
@@ -15,6 +15,7 @@
 #include "NiagaraEditorWidgetsUtilities.h"
 #include "NiagaraNodeCustomHlsl.h"
 #include "NiagaraNodeFunctionCall.h"
+#include "NiagaraScriptVariable.h"
 #include "NiagaraParameterCollection.h"
 #include "NiagaraSystem.h"
 #include "ScopedTransaction.h"
@@ -50,6 +51,7 @@
 #include "Widgets/Input/SNumericEntryBox.h"
 #include "NiagaraNodeAssignment.h"
 #include "NiagaraCommon.h"
+#include "NiagaraScriptVariable.h"
 #include "NiagaraSettings.h"
 #include "SNiagaraParameterDropTarget.h"
 #include "Stack/SNiagaraStackIndent.h"
@@ -251,10 +253,32 @@ TSharedRef<SWidget> SNiagaraStackFunctionInputValue::ConstructValueWidgets()
 	}
 	case UNiagaraStackFunctionInput::EValueMode::Linked:
 	{
-		return SNew(SNiagaraParameterName)
+		TSharedRef<SWidget> ParameterWidget = SNew(SNiagaraParameterName)
 			.ReadOnlyTextStyle(FNiagaraEditorStyle::Get(), "NiagaraEditor.ParameterText")
 			.ParameterName(this, &SNiagaraStackFunctionInputValue::GetLinkedValueHandleName)
 			.OnDoubleClicked(this, &SNiagaraStackFunctionInputValue::OnLinkedInputDoubleClicked);
+
+		if(FunctionInput->GetLinkedValueHandle().IsUserHandle())
+		{
+			FNiagaraParameterHandle ParameterHandle = FunctionInput->GetLinkedValueHandle();
+			TArray<FNiagaraVariable> UserParameters;
+			FunctionInput->GetSystemViewModel()->GetSystem().GetExposedParameters().GetUserParameters(UserParameters);
+			FNiagaraVariable* MatchingVariable = UserParameters.FindByPredicate([ParameterHandle](const FNiagaraVariable& Variable)
+			{
+				return Variable.GetName().ToString() == ParameterHandle.GetName().ToString();
+			});
+
+			if(MatchingVariable)
+			{
+				FText Tooltip = FNiagaraEditorUtilities::GetScriptVariableForUserParameter(*MatchingVariable, FunctionInput->GetSystemViewModel())->Metadata.Description;
+				if(!Tooltip.IsEmpty())
+				{
+					ParameterWidget->SetToolTipText(Tooltip);
+				}
+			}
+		}
+			
+		return ParameterWidget; 
 	}
 	case UNiagaraStackFunctionInput::EValueMode::Data:
 	{
