Fixed the context menu on hierarchy sections not working

diff --git a/Engine/Plugins/FX/Niagara/Source/NiagaraEditor/Private/Widgets/SNiagaraHierarchy.cpp b/Engine/Plugins/FX/Niagara/Source/NiagaraEditor/Private/Widgets/SNiagaraHierarchy.cpp
--- a/Engine/Plugins/FX/Niagara/Source/NiagaraEditor/Private/Widgets/SNiagaraHierarchy.cpp
+++ b/Engine/Plugins/FX/Niagara/Source/NiagaraEditor/Private/Widgets/SNiagaraHierarchy.cpp
@@ -120,18 +120,21 @@ void SNiagaraHierarchySection::Construct(const FArguments& InArgs, TSharedPtr<FN
 				.DropTargetArgs(OntoDropTargetArgs
 				.Content()
 				[
-					SNew(SCheckBox)
-					.Visibility(EVisibility::HitTestInvisible)
-					.Style(FAppStyle::Get(), "DetailsView.SectionButton")
-					.OnCheckStateChanged(this, &SNiagaraHierarchySection::OnSectionCheckChanged)
-					.IsChecked(this, &SNiagaraHierarchySection::GetSectionCheckState)
+					SAssignNew(MenuAnchor, SMenuAnchor)
 					.OnGetMenuContent(this, &SNiagaraHierarchySection::OnGetMenuContent)
 					[
-						SAssignNew(InlineEditableTextBlock, SInlineEditableTextBlock)
-						.IsSelected(InArgs._IsSelected)
-						.Text(this, &SNiagaraHierarchySection::GetText)
-						.OnTextCommitted(this, &SNiagaraHierarchySection::OnRenameSection)
-						.OnVerifyTextChanged(this, &SNiagaraHierarchySection::OnVerifySectionRename)
+						SNew(SCheckBox)
+						.Visibility(EVisibility::HitTestInvisible)
+						.Style(FAppStyle::Get(), "DetailsView.SectionButton")
+						.OnCheckStateChanged(this, &SNiagaraHierarchySection::OnSectionCheckChanged)
+						.IsChecked(this, &SNiagaraHierarchySection::GetSectionCheckState)
+						[
+							SAssignNew(InlineEditableTextBlock, SInlineEditableTextBlock)
+							.IsSelected(InArgs._IsSelected)
+							.Text(this, &SNiagaraHierarchySection::GetText)
+							.OnTextCommitted(this, &SNiagaraHierarchySection::OnRenameSection)
+							.OnVerifyTextChanged(this, &SNiagaraHierarchySection::OnVerifySectionRename)
+						]
 					]
 				])
 			]
@@ -262,6 +265,7 @@ FReply SNiagaraHierarchySection::OnMouseButtonDown(const FGeometry& MyGeometry,
 	{
 		if(MouseEvent.IsMouseButtonDown(EKeys::RightMouseButton))
 		{
+			MenuAnchor->SetIsOpen(true);
 			OnSectionCheckChanged(ECheckBoxState::Checked);
 			return FReply::Handled();
 		}
diff --git a/Engine/Plugins/FX/Niagara/Source/NiagaraEditor/Public/Widgets/SNiagaraHierarchy.h b/Engine/Plugins/FX/Niagara/Source/NiagaraEditor/Public/Widgets/SNiagaraHierarchy.h
--- a/Engine/Plugins/FX/Niagara/Source/NiagaraEditor/Public/Widgets/SNiagaraHierarchy.h
+++ b/Engine/Plugins/FX/Niagara/Source/NiagaraEditor/Public/Widgets/SNiagaraHierarchy.h
@@ -8,6 +8,7 @@
 #include "Widgets/Views/STableRow.h"
 #include "Widgets/Text/SInlineEditableTextBlock.h"
 #include "Widgets/SBoxPanel.h"
+#include "Widgets/Input/SMenuAnchor.h"
 #include "Styling/SlateTypes.h"
 
 class UNiagaraHierarchyViewModelBase;
@@ -89,6 +90,7 @@ public:
 	ECheckBoxState GetSectionCheckState() const;
 	void OnSectionCheckChanged(ECheckBoxState NewState);
 private:
+	TSharedPtr<SMenuAnchor> MenuAnchor;
 	TSharedPtr<SInlineEditableTextBlock> InlineEditableTextBlock;
 	TWeakObjectPtr<UNiagaraHierarchyViewModelBase> HierarchyViewModel;
 	TSharedPtr<FNiagaraHierarchySectionViewModel> SectionViewModel;
