Fix for SA issue: potential (but not effective) nullptr dereferencing

diff --git a/Engine/Plugins/FX/Niagara/Source/NiagaraEditor/Private/ViewModels/HierarchyEditor/NiagaraHierarchyViewModelBase.cpp b/Engine/Plugins/FX/Niagara/Source/NiagaraEditor/Private/ViewModels/HierarchyEditor/NiagaraHierarchyViewModelBase.cpp
--- a/Engine/Plugins/FX/Niagara/Source/NiagaraEditor/Private/ViewModels/HierarchyEditor/NiagaraHierarchyViewModelBase.cpp
+++ b/Engine/Plugins/FX/Niagara/Source/NiagaraEditor/Private/ViewModels/HierarchyEditor/NiagaraHierarchyViewModelBase.cpp
@@ -154,17 +154,23 @@ bool UNiagaraHierarchyRoot::Modify(bool bAlwaysMarkDirty)
 
 bool FNiagaraHierarchyCategoryViewModel::IsTopCategoryActive() const
 {
-	const UNiagaraHierarchyCategory* Result = Cast<UNiagaraHierarchyCategory>(GetDataMutable());
-	const UNiagaraHierarchyCategory* TopLevelCategory = Result;
-	for (; TopLevelCategory != nullptr; TopLevelCategory = TopLevelCategory->GetTypedOuter<UNiagaraHierarchyCategory>() )
+	if(UNiagaraHierarchyCategory* Category = Cast<UNiagaraHierarchyCategory>(GetDataMutable()))
 	{
-		if(TopLevelCategory != nullptr)
+		const UNiagaraHierarchyCategory* Result = Category;
+		const UNiagaraHierarchyCategory* TopLevelCategory = Result;
+		
+		for (; TopLevelCategory != nullptr; TopLevelCategory = TopLevelCategory->GetTypedOuter<UNiagaraHierarchyCategory>() )
 		{
-			Result = TopLevelCategory;
+			if(TopLevelCategory != nullptr)
+			{
+				Result = TopLevelCategory;
+			}
 		}
+		
+		return HierarchyViewModel->IsSectionActive(Result->GetSection());
 	}
-	
-	return HierarchyViewModel->IsSectionActive(Result->GetSection());
+
+	return false;	
 }
 
 void UNiagaraHierarchySection::SetSectionNameAsText(const FText& Text)
