- Static analysis fix

#rnx

diff --git a/Engine/Plugins/FX/Niagara/Source/Niagara/Private/NiagaraFunctionLibrary.cpp b/Engine/Plugins/FX/Niagara/Source/Niagara/Private/NiagaraFunctionLibrary.cpp
--- a/Engine/Plugins/FX/Niagara/Source/Niagara/Private/NiagaraFunctionLibrary.cpp
+++ b/Engine/Plugins/FX/Niagara/Source/Niagara/Private/NiagaraFunctionLibrary.cpp
@@ -649,7 +649,13 @@ UNiagaraDataInterface* UNiagaraFunctionLibrary::GetDataInterface(UClass* DIClass
 	}
 
 	UNiagaraDataInterface* UntypedDI = OverrideParameters.GetDataInterface(*Index);
-	if (UntypedDI != nullptr && UntypedDI->IsA(DIClass))
+	if (UntypedDI == nullptr)
+	{
+		UE_LOG(LogNiagara, Warning, TEXT("OverrideParameter\"%s\" DataInterface is nullptr."), *OverrideName.ToString());
+		return nullptr;
+	}
+
+	if (UntypedDI->IsA(DIClass))
 	{
 		return UntypedDI;
 	}
