Cleanup natvis for niagara types
-moves types into its own niagara.natvis file
-allows resolving FNiagaraTypeDefinitionHandle through global pointers so that FNiagaraVariable can now be properly visualized

diff --git a/Engine/Plugins/FX/Niagara/Source/Niagara/Private/NiagaraModule.cpp b/Engine/Plugins/FX/Niagara/Source/Niagara/Private/NiagaraModule.cpp
--- a/Engine/Plugins/FX/Niagara/Source/Niagara/Private/NiagaraModule.cpp
+++ b/Engine/Plugins/FX/Niagara/Source/Niagara/Private/NiagaraModule.cpp
@@ -4,6 +4,7 @@
 #include "Misc/LazySingleton.h"
 #include "Modules/ModuleManager.h"
 #include "NiagaraTypes.h"
+#include "NiagaraDebugVis.h"
 #include "NiagaraEvents.h"
 #include "NiagaraSettings.h"
 #include "NiagaraDataInterfaceCurlNoise.h"
@@ -91,6 +92,17 @@ void INiagaraModule::OnUseGlobalFXBudgetChanged(IConsoleVariable* Variable)
 	}
 }
 
+// these two globals are intended to help ensure that a global variable is accessible to natvis (as defined in Niagara.natvis)
+// while debugging.  GCoreTypeRegistrySingletonPtr is the pointer to the actual data stored within the TLazySingleton<FNiagaraTypeRegistry>
+// while GTypeRegistrySingletonPtr can be declared in each module to ensure that it can be accessed while debugging any Niagara
+// module.  See UE4_VISUALIZERS_HELPERS.  Note that currently it seems like we can effectively debug NiagaraEditor/NiagaraShader without
+// further declarations, which is nice, so will be leaving it as it is.
+namespace NiagaraDebugVisHelper
+{
+	const FNiagaraTypeRegistry* GCoreTypeRegistrySingletonPtr = nullptr;
+	const FNiagaraTypeRegistry*& GTypeRegistrySingletonPtr = GCoreTypeRegistrySingletonPtr;
+}
+
 FNiagaraVariable INiagaraModule::Engine_DeltaTime;
 FNiagaraVariable INiagaraModule::Engine_InvDeltaTime;
 FNiagaraVariable INiagaraModule::Engine_Time;
@@ -1420,6 +1432,17 @@ FString FNiagaraTypeRegistry::GetReferencerName() const
 	return TEXT("FNiagaraTypeRegistry");
 }
 
+FNiagaraTypeRegistry::FNiagaraTypeRegistry()
+{
+	NiagaraDebugVisHelper::GCoreTypeRegistrySingletonPtr = this;
+
+}
+
+FNiagaraTypeRegistry::~FNiagaraTypeRegistry()
+{
+	NiagaraDebugVisHelper::GCoreTypeRegistrySingletonPtr = nullptr;
+}
+
 FNiagaraTypeRegistry& FNiagaraTypeRegistry::Get()
 {
 	return TLazySingleton<FNiagaraTypeRegistry>::Get();
diff --git a/Engine/Plugins/FX/Niagara/Source/Niagara/Public/NiagaraDebugVis.h b/Engine/Plugins/FX/Niagara/Source/Niagara/Public/NiagaraDebugVis.h
new file mode 100644
--- /dev/null
+++ b/Engine/Plugins/FX/Niagara/Source/Niagara/Public/NiagaraDebugVis.h
@@ -0,0 +1,12 @@
+// Copyright Epic Games, Inc. All Rights Reserved.
+
+#pragma once
+
+class FNiagaraTypeRegistry;
+
+#if !IS_MONOLITHIC
+namespace NiagaraDebugVisHelper
+{
+	extern NIAGARA_API const FNiagaraTypeRegistry* GCoreTypeRegistrySingletonPtr;
+}
+#endif
\ No newline at end of file
diff --git a/Engine/Plugins/FX/Niagara/Source/Niagara/Public/NiagaraTypes.h b/Engine/Plugins/FX/Niagara/Source/Niagara/Public/NiagaraTypes.h
--- a/Engine/Plugins/FX/Niagara/Source/Niagara/Public/NiagaraTypes.h
+++ b/Engine/Plugins/FX/Niagara/Source/Niagara/Public/NiagaraTypes.h
@@ -1626,6 +1626,11 @@ public:
 	virtual FString GetReferencerName() const;
 
 private:
+	friend class FLazySingleton;
+
+	FNiagaraTypeRegistry();
+	virtual ~FNiagaraTypeRegistry();
+
 	RegisteredTypesArray RegisteredTypes;
 
 	TArray<FNiagaraTypeDefinition> RegisteredUserVariableTypes;
