Prevent caching shaders for scripts which are not associated with a system.  This prevents compilation of standalone scripts and scripts from standalone emitters.
-stop gap measure since making the scripts themselves editor only doesn't seem to prevent the compile requests from being triggered during cooks

diff --git a/Engine/Plugins/FX/Niagara/Source/Niagara/Private/NiagaraScript.cpp b/Engine/Plugins/FX/Niagara/Source/Niagara/Private/NiagaraScript.cpp
--- a/Engine/Plugins/FX/Niagara/Source/Niagara/Private/NiagaraScript.cpp
+++ b/Engine/Plugins/FX/Niagara/Source/Niagara/Private/NiagaraScript.cpp
@@ -2166,6 +2166,13 @@ bool UNiagaraScript::ShouldCacheShadersForCooking(const ITargetPlatform* TargetP
 			}
 		}
 
+		// currently only scripts that are associated with a system are appropriate for caching.  Standalone scripts or scripts
+		// associated with standalone emitters are not required on cooked platforms.
+		if (GetTypedOuter<UNiagaraSystem>() == nullptr)
+		{
+			return false;
+		}
+
 		FVersionedNiagaraEmitter OuterEmitter = GetOuterEmitter();
 		FVersionedNiagaraEmitterData* EmitterData = OuterEmitter.GetEmitterData();
 		if (EmitterData != nullptr && EmitterData->SimTarget == ENiagaraSimTarget::GPUComputeSim && OuterEmitter.Emitter->NeedsLoadForTargetPlatform(TargetPlatform))
