- Ensure we don't kill the system if the RT DI is not actually used with a GPU emitter
- Move some common code into a common location

#rnx

diff --git a/Engine/Plugins/FX/Niagara/Source/Niagara/Private/NiagaraDataInterfaceIntRenderTarget2D.cpp b/Engine/Plugins/FX/Niagara/Source/Niagara/Private/NiagaraDataInterfaceIntRenderTarget2D.cpp
--- a/Engine/Plugins/FX/Niagara/Source/Niagara/Private/NiagaraDataInterfaceIntRenderTarget2D.cpp
+++ b/Engine/Plugins/FX/Niagara/Source/Niagara/Private/NiagaraDataInterfaceIntRenderTarget2D.cpp
@@ -1,5 +1,6 @@
 // Copyright Epic Games, Inc. All Rights Reserved.
 #include "NiagaraDataInterfaceIntRenderTarget2D.h"
+#include "NiagaraDataInterfaceRenderTargetCommon.h"
 #include "NiagaraShader.h"
 #include "NiagaraSettings.h"
 #include "NiagaraSystemInstance.h"
@@ -611,10 +612,21 @@ bool UNiagaraDataInterfaceIntRenderTarget2D::InitPerInstanceData(void* PerInstan
 {
 	check(Proxy);
 
-	extern float GNiagaraRenderTargetResolutionMultiplier;
 	FNDIIntRenderTarget2DInstanceData_GameThread* InstanceData = new (PerInstanceData) FNDIIntRenderTarget2DInstanceData_GameThread();
-	InstanceData->Size.X = FMath::Clamp<int>(int(float(Size.X) * GNiagaraRenderTargetResolutionMultiplier), 1, GMaxTextureDimensions);
-	InstanceData->Size.Y = FMath::Clamp<int>(int(float(Size.Y) * GNiagaraRenderTargetResolutionMultiplier), 1, GMaxTextureDimensions);
+
+	if (NiagaraDataInterfaceRenderTargetCommon::GIgnoreCookedOut && !IsUsedWithGPUEmitter())
+	{
+		return true;
+	}
+
+	// Note: If the Format can ever change we will need to replicate this elsewhere also
+	if (RHIIsTypedUAVStoreSupported(InstanceData->Format) == false)
+	{
+		return false;
+	}
+
+	InstanceData->Size.X = FMath::Clamp<int>(int(float(Size.X) * NiagaraDataInterfaceRenderTargetCommon::GResolutionMultiplier), 1, GMaxTextureDimensions);
+	InstanceData->Size.Y = FMath::Clamp<int>(int(float(Size.Y) * NiagaraDataInterfaceRenderTargetCommon::GResolutionMultiplier), 1, GMaxTextureDimensions);
 #if WITH_EDITORONLY_DATA
 	InstanceData->bPreviewRenderTarget = bPreviewRenderTarget;
 	InstanceData->PreviewDisplayRange = PreviewDisplayRange;
@@ -663,9 +675,8 @@ void UNiagaraDataInterfaceIntRenderTarget2D::DestroyPerInstanceData(void* PerIns
 	);
 
 	// Make sure to clear out the reference to the render target if we created one.
-	extern int32 GNiagaraReleaseResourceOnRemove;
 	decltype(ManagedRenderTargets)::ValueType ExistingRenderTarget = nullptr;
-	if ( ManagedRenderTargets.RemoveAndCopyValue(SystemInstance->GetId(), ExistingRenderTarget) && GNiagaraReleaseResourceOnRemove)
+	if ( ManagedRenderTargets.RemoveAndCopyValue(SystemInstance->GetId(), ExistingRenderTarget) && NiagaraDataInterfaceRenderTargetCommon::GReleaseResourceOnRemove )
 	{
 		ExistingRenderTarget->ReleaseResource();
 	}
@@ -695,6 +706,12 @@ int32 UNiagaraDataInterfaceIntRenderTarget2D::PerInstanceDataSize() const
 
 bool UNiagaraDataInterfaceIntRenderTarget2D::PerInstanceTickPostSimulate(void* PerInstanceData, FNiagaraSystemInstance* SystemInstance, float DeltaSeconds)
 {
+	//-TEMP: Until we prune data interface on cook this will avoid consuming memory
+	if (NiagaraDataInterfaceRenderTargetCommon::GIgnoreCookedOut && !IsUsedWithGPUEmitter())
+	{
+		return false;
+	}
+
 	FNDIIntRenderTarget2DInstanceData_GameThread* InstanceData = static_cast<FNDIIntRenderTarget2DInstanceData_GameThread*>(PerInstanceData);
 
 	{
@@ -747,9 +764,8 @@ bool UNiagaraDataInterfaceIntRenderTarget2D::UpdateInstanceTexture(FNiagaraSyste
 			{
 				InstanceData->TargetTexture = UserTargetTexture;
 
-				extern int32 GNiagaraReleaseResourceOnRemove;
 				decltype(ManagedRenderTargets)::ValueType ExistingRenderTarget = nullptr;
-				if (ManagedRenderTargets.RemoveAndCopyValue(SystemInstance->GetId(), ExistingRenderTarget) && GNiagaraReleaseResourceOnRemove)
+				if (ManagedRenderTargets.RemoveAndCopyValue(SystemInstance->GetId(), ExistingRenderTarget) && NiagaraDataInterfaceRenderTargetCommon::GReleaseResourceOnRemove)
 				{
 					ExistingRenderTarget->ReleaseResource();
 				}
@@ -820,7 +836,6 @@ void UNiagaraDataInterfaceIntRenderTarget2D::VMSetSize(FVectorVMExternalFunction
 	FNDIInputParam<int32> InSizeY(Context);
 	FNDIOutputParam<FNiagaraBool> OutSuccess(Context);
 
-	extern float GNiagaraRenderTargetResolutionMultiplier;
 	for (int32 i=0; i < Context.GetNumInstances(); ++i)
 	{
 		const int SizeX = InSizeX.GetAndAdvance();
@@ -829,8 +844,8 @@ void UNiagaraDataInterfaceIntRenderTarget2D::VMSetSize(FVectorVMExternalFunction
 		OutSuccess.SetAndAdvance(bSuccess);
 		if (bSuccess)
 		{
-			InstData->Size.X = FMath::Clamp<int>(int(float(SizeX) * GNiagaraRenderTargetResolutionMultiplier), 1, GMaxTextureDimensions);
-			InstData->Size.Y = FMath::Clamp<int>(int(float(SizeY) * GNiagaraRenderTargetResolutionMultiplier), 1, GMaxTextureDimensions);
+			InstData->Size.X = FMath::Clamp<int>(int(float(SizeX) * NiagaraDataInterfaceRenderTargetCommon::GResolutionMultiplier), 1, GMaxTextureDimensions);
+			InstData->Size.Y = FMath::Clamp<int>(int(float(SizeY) * NiagaraDataInterfaceRenderTargetCommon::GResolutionMultiplier), 1, GMaxTextureDimensions);
 		}
 	}
 }
diff --git a/Engine/Plugins/FX/Niagara/Source/Niagara/Private/NiagaraDataInterfaceRenderTarget2D.cpp b/Engine/Plugins/FX/Niagara/Source/Niagara/Private/NiagaraDataInterfaceRenderTarget2D.cpp
--- a/Engine/Plugins/FX/Niagara/Source/Niagara/Private/NiagaraDataInterfaceRenderTarget2D.cpp
+++ b/Engine/Plugins/FX/Niagara/Source/Niagara/Private/NiagaraDataInterfaceRenderTarget2D.cpp
@@ -8,6 +8,7 @@
 #include "Engine/Texture2D.h"
 #include "Engine/TextureRenderTarget2D.h"
 
+#include "NiagaraDataInterfaceRenderTargetCommon.h"
 #include "NiagaraSettings.h"
 #include "NiagaraSystemInstance.h"
 #include "NiagaraRenderer.h"
@@ -43,119 +44,6 @@ namespace NDIRenderTarget2DLocal
 
 FNiagaraVariableBase UNiagaraDataInterfaceRenderTarget2D::ExposedRTVar;
 
-int32 GNiagaraReleaseResourceOnRemove = false;
-static FAutoConsoleVariableRef CVarNiagaraReleaseResourceOnRemove(
-	TEXT("fx.Niagara.RenderTarget.ReleaseResourceOnRemove"),
-	GNiagaraReleaseResourceOnRemove,
-	TEXT("Releases the render target resource once it is removed from the manager list rather than waiting for a GC."),
-	ECVF_Default
-);
-
-//-TEMP: Until we prune data interface on cook this will avoid consuming memory
-int32 GNiagaraRenderTargetIgnoreCookedOut = true;
-static FAutoConsoleVariableRef CVarNiagaraRenderTargetIgnoreCookedOut(
-	TEXT("fx.Niagara.RenderTarget.IgnoreCookedOut"),
-	GNiagaraRenderTargetIgnoreCookedOut,
-	TEXT("Ignores create render targets for cooked out emitter, i.e. ones that are not used by any GPU emitter."),
-	ECVF_Default
-);
-
-float GNiagaraRenderTargetResolutionMultiplier = 1.0f;
-static FAutoConsoleVariableRef CVarNiagaraRenderTargetResolutionMultiplier(
-	TEXT("fx.Niagara.RenderTarget.ResolutionMultiplier"),
-	GNiagaraRenderTargetResolutionMultiplier,
-	TEXT("Optional global modifier to Niagara render target resolution."),
-	ECVF_Default
-);
-
-int GNiagaraRenderTargetAllowReads = 0;
-static FAutoConsoleVariableRef CVarNiagaraRenderTargetAllowReads(
-	TEXT("fx.Niagara.RenderTarget.AllowReads"),
-	GNiagaraRenderTargetAllowReads,
-	TEXT("Enables read operations to be visible in the UI, very experimental."),
-	ECVF_Default
-);
-
-/*--------------------------------------------------------------------------------------------------------------------------*/
-
-static bool GNiagaraRenderTargetOverrideFormatEnabled = false;
-static ETextureRenderTargetFormat GNiagaraRenderTargetOverrideFormat = ETextureRenderTargetFormat::RTF_RGBA32f;
-
-static FAutoConsoleCommandWithWorldAndArgs GCommandNiagaraRenderTargetOverrideFormat(
-	TEXT("fx.Niagara.RenderTarget.OverrideFormat"),
-	TEXT("Optional global format override for all Niagara render targets"),
-	FConsoleCommandWithWorldAndArgsDelegate::CreateLambda(
-		[](const TArray<FString>& Args, UWorld*)
-		{
-			static UEnum* TextureRenderTargetFormatEnum = StaticEnum<ETextureRenderTargetFormat>();
-			if ( ensure(TextureRenderTargetFormatEnum) )
-			{
-				if (Args.Num() == 1)
-				{
-					const int32 EnumIndex = TextureRenderTargetFormatEnum->GetIndexByNameString(Args[0]);
-					if (EnumIndex != INDEX_NONE)
-					{
-						GNiagaraRenderTargetOverrideFormatEnabled = true;
-						GNiagaraRenderTargetOverrideFormat = ETextureRenderTargetFormat(EnumIndex);
-					}
-					else
-					{
-						GNiagaraRenderTargetOverrideFormatEnabled = false;
-					}
-				}
-				UE_LOG(LogNiagara, Log, TEXT("Niagara RenderTarget Override is '%s' with format '%s'."), GNiagaraRenderTargetOverrideFormatEnabled ? TEXT("Enabled") : TEXT("Disabled"), *TextureRenderTargetFormatEnum->GetNameStringByIndex(int32(GNiagaraRenderTargetOverrideFormat)));
-			}
-		}
-	)
-);
-
-bool GetRenderTargetFormat(bool bOverrideFormat, ETextureRenderTargetFormat OverrideFormat, ETextureRenderTargetFormat& OutRenderTargetFormat)
-{
-	OutRenderTargetFormat = bOverrideFormat ? OverrideFormat : GetDefault<UNiagaraSettings>()->DefaultRenderTargetFormat.GetValue();
-	EPixelFormat PixelFormat = GetPixelFormatFromRenderTargetFormat(OutRenderTargetFormat);
-	if (GNiagaraRenderTargetOverrideFormatEnabled)
-	{
-		OutRenderTargetFormat = GNiagaraRenderTargetOverrideFormat;
-	}
-
-	// If the format does not support typed store we need to find one that will
-	while ( RHIIsTypedUAVStoreSupported(GetPixelFormatFromRenderTargetFormat(OutRenderTargetFormat)) == false )
-	{
-		static const TPair<ETextureRenderTargetFormat, ETextureRenderTargetFormat> FormatRemapTable[] =
-		{
-			TPairInitializer(RTF_R8,			RTF_R16f),
-			TPairInitializer(RTF_RG8,			RTF_RG16f),
-			TPairInitializer(RTF_RGBA8,			RTF_RGBA16f),
-			TPairInitializer(RTF_RGBA8_SRGB,	RTF_RGBA16f),
-			TPairInitializer(RTF_R16f,			RTF_R32f),
-			TPairInitializer(RTF_RG16f,			RTF_RG32f),
-			TPairInitializer(RTF_RGBA16f,		RTF_RGBA32f),
-			TPairInitializer(RTF_R32f,			RTF_RGBA32f),
-			TPairInitializer(RTF_RG32f,			RTF_RGBA32f),
-			TPairInitializer(RTF_RGBA32f,		RTF_RGBA32f),
-			TPairInitializer(RTF_RGB10A2,		RTF_RGBA32f),
-		};
-
-		const ETextureRenderTargetFormat PreviousFormat = OutRenderTargetFormat;
-		for ( const auto& FormatRemap : FormatRemapTable)
-		{
-			if (FormatRemap.Key == OutRenderTargetFormat)
-			{
-				OutRenderTargetFormat = FormatRemap.Value;
-				break;
-			}
-		}
-		if ( PreviousFormat == OutRenderTargetFormat)
-		{
-			// This is fatal as we failed to find any format that supports typed UAV stores
-			UE_LOG(LogNiagara, Warning, TEXT("Failed to find a render target format that supports UAV store"));
-			return false;
-		}
-	}
-
-	return true;
-}
-
 /*--------------------------------------------------------------------------------------------------------------------------*/
 
 struct FNDIRenderTarget2DFunctionVersion
@@ -285,7 +173,7 @@ void UNiagaraDataInterfaceRenderTarget2D::GetFunctions(TArray<FNiagaraFunctionSi
 		Sig.Inputs.Add(FNiagaraVariable(FNiagaraTypeDefinition::GetIntDef(), TEXT("IndexY")));
 		Sig.Outputs.Add(FNiagaraVariable(FNiagaraTypeDefinition::GetColorDef(), TEXT("Value")));
 
-		Sig.bHidden = GNiagaraRenderTargetAllowReads != 1;
+		Sig.bHidden = NiagaraDataInterfaceRenderTargetCommon::GAllowReads != 1;
 		Sig.bExperimental = true;
 		Sig.bMemberFunction = true;
 		Sig.bRequiresContext = false;
@@ -304,7 +192,7 @@ void UNiagaraDataInterfaceRenderTarget2D::GetFunctions(TArray<FNiagaraFunctionSi
 		Sig.Inputs.Add(FNiagaraVariable(FNiagaraTypeDefinition::GetVec2Def(), TEXT("UV")));
 		Sig.Outputs.Add(FNiagaraVariable(FNiagaraTypeDefinition::GetColorDef(), TEXT("Value")));
 
-		Sig.bHidden = GNiagaraRenderTargetAllowReads != 1;
+		Sig.bHidden = NiagaraDataInterfaceRenderTargetCommon::GAllowReads != 1;
 		Sig.bExperimental = true;
 		Sig.bMemberFunction = true;
 		Sig.bRequiresContext = false;
@@ -566,18 +454,22 @@ bool UNiagaraDataInterfaceRenderTarget2D::InitPerInstanceData(void* PerInstanceD
 {
 	check(Proxy);
 
-	extern float GNiagaraRenderTargetResolutionMultiplier;
 	FRenderTarget2DRWInstanceData_GameThread* InstanceData = new (PerInstanceData) FRenderTarget2DRWInstanceData_GameThread();
 	SystemInstancesToProxyData_GT.Emplace(SystemInstance->GetId(), InstanceData);
 
+	if (NiagaraDataInterfaceRenderTargetCommon::GIgnoreCookedOut && !IsUsedWithGPUEmitter())
+	{
+		return true;
+	}
+
 	ETextureRenderTargetFormat RenderTargetFormat;
-	if ( GetRenderTargetFormat(bOverrideFormat, OverrideRenderTargetFormat, RenderTargetFormat) == false )
+	if ( NiagaraDataInterfaceRenderTargetCommon::GetRenderTargetFormat(bOverrideFormat, OverrideRenderTargetFormat, RenderTargetFormat) == false )
 	{
 		return false;
 	}
 
-	InstanceData->Size.X = FMath::Clamp<int>(int(float(Size.X) * GNiagaraRenderTargetResolutionMultiplier), 1, GMaxTextureDimensions);
-	InstanceData->Size.Y = FMath::Clamp<int>(int(float(Size.Y) * GNiagaraRenderTargetResolutionMultiplier), 1, GMaxTextureDimensions);
+	InstanceData->Size.X = FMath::Clamp<int>(int(float(Size.X) * NiagaraDataInterfaceRenderTargetCommon::GResolutionMultiplier), 1, GMaxTextureDimensions);
+	InstanceData->Size.Y = FMath::Clamp<int>(int(float(Size.Y) * NiagaraDataInterfaceRenderTargetCommon::GResolutionMultiplier), 1, GMaxTextureDimensions);
 	InstanceData->MipMapGeneration = MipMapGeneration;
 	InstanceData->MipMapGenerationType = MipMapGenerationType;
 	InstanceData->Format = RenderTargetFormat;
@@ -616,7 +508,7 @@ void UNiagaraDataInterfaceRenderTarget2D::DestroyPerInstanceData(void* PerInstan
 	// Make sure to clear out the reference to the render target if we created one.
 	using RenderTargetType = decltype(decltype(ManagedRenderTargets)::ElementType::Value);
 	RenderTargetType ExistingRenderTarget = nullptr;
-	if ( ManagedRenderTargets.RemoveAndCopyValue(SystemInstance->GetId(), ExistingRenderTarget) && GNiagaraReleaseResourceOnRemove)
+	if ( ManagedRenderTargets.RemoveAndCopyValue(SystemInstance->GetId(), ExistingRenderTarget) && NiagaraDataInterfaceRenderTargetCommon::GReleaseResourceOnRemove )
 	{
 		ExistingRenderTarget->ReleaseResource();
 	}
@@ -683,7 +575,6 @@ void UNiagaraDataInterfaceRenderTarget2D::VMSetSize(FVectorVMExternalFunctionCon
 	FNDIInputParam<int> InSizeY(Context);
 	FNDIOutputParam<FNiagaraBool> OutSuccess(Context);
 
-	extern float GNiagaraRenderTargetResolutionMultiplier;
 	for (int32 InstanceIdx = 0; InstanceIdx < Context.GetNumInstances(); ++InstanceIdx)
 	{
 		const int SizeX = InSizeX.GetAndAdvance();
@@ -692,8 +583,8 @@ void UNiagaraDataInterfaceRenderTarget2D::VMSetSize(FVectorVMExternalFunctionCon
 		OutSuccess.SetAndAdvance(bSuccess);
 		if (bSuccess)
 		{
-			InstData->Size.X = FMath::Clamp<int>(int(float(SizeX) * GNiagaraRenderTargetResolutionMultiplier), 1, GMaxTextureDimensions);
-			InstData->Size.Y = FMath::Clamp<int>(int(float(SizeY) * GNiagaraRenderTargetResolutionMultiplier), 1, GMaxTextureDimensions);
+			InstData->Size.X = FMath::Clamp<int>(int(float(SizeX) * NiagaraDataInterfaceRenderTargetCommon::GResolutionMultiplier), 1, GMaxTextureDimensions);
+			InstData->Size.Y = FMath::Clamp<int>(int(float(SizeY) * NiagaraDataInterfaceRenderTargetCommon::GResolutionMultiplier), 1, GMaxTextureDimensions);
 		}
 	}
 }
@@ -723,7 +614,7 @@ bool UNiagaraDataInterfaceRenderTarget2D::PerInstanceTick(void* PerInstanceData,
 		InstanceData->TargetTexture = UserTargetTexture;
 
 		TObjectPtr<UTextureRenderTarget2D> ExistingRenderTarget = nullptr;
-		if (ManagedRenderTargets.RemoveAndCopyValue(SystemInstance->GetId(), ExistingRenderTarget) && GNiagaraReleaseResourceOnRemove)
+		if (ManagedRenderTargets.RemoveAndCopyValue(SystemInstance->GetId(), ExistingRenderTarget) && NiagaraDataInterfaceRenderTargetCommon::GReleaseResourceOnRemove)
 		{
 			ExistingRenderTarget->ReleaseResource();
 		}
@@ -767,7 +658,7 @@ bool UNiagaraDataInterfaceRenderTarget2D::PerInstanceTickPostSimulate(void* PerI
 #endif
 
 	//-TEMP: Until we prune data interface on cook this will avoid consuming memory
-	if (GNiagaraRenderTargetIgnoreCookedOut && !IsUsedWithGPUEmitter())
+	if (NiagaraDataInterfaceRenderTargetCommon::GIgnoreCookedOut && !IsUsedWithGPUEmitter())
 	{
 		return false;
 	}
diff --git a/Engine/Plugins/FX/Niagara/Source/Niagara/Private/NiagaraDataInterfaceRenderTarget2DArray.cpp b/Engine/Plugins/FX/Niagara/Source/Niagara/Private/NiagaraDataInterfaceRenderTarget2DArray.cpp
--- a/Engine/Plugins/FX/Niagara/Source/Niagara/Private/NiagaraDataInterfaceRenderTarget2DArray.cpp
+++ b/Engine/Plugins/FX/Niagara/Source/Niagara/Private/NiagaraDataInterfaceRenderTarget2DArray.cpp
@@ -7,6 +7,7 @@
 #include "Engine/Texture2DArray.h"
 #include "Engine/TextureRenderTarget2DArray.h"
 
+#include "NiagaraDataInterfaceRenderTargetCommon.h"
 #include "NiagaraSystemInstance.h"
 #include "NiagaraStats.h"
 #include "NiagaraRenderer.h"
@@ -164,7 +165,6 @@ void UNiagaraDataInterfaceRenderTarget2DArray::GetFunctions(TArray<FNiagaraFunct
 	#endif
 	}
 
-	extern int GNiagaraRenderTargetAllowReads;
 	{
 		FNiagaraFunctionSignature& Sig = OutFunctions.AddDefaulted_GetRef();
 		Sig.Name = GetValueFunctionName;
@@ -174,7 +174,7 @@ void UNiagaraDataInterfaceRenderTarget2DArray::GetFunctions(TArray<FNiagaraFunct
 		Sig.Inputs.Add(FNiagaraVariable(FNiagaraTypeDefinition::GetIntDef(), TEXT("IndexZ")));
 		Sig.Outputs.Add(FNiagaraVariable(FNiagaraTypeDefinition::GetColorDef(), TEXT("Value")));
 
-		Sig.bHidden = GNiagaraRenderTargetAllowReads != 1;
+		Sig.bHidden = NiagaraDataInterfaceRenderTargetCommon::GAllowReads != 1;
 		Sig.bExperimental = true;
 		Sig.bMemberFunction = true;
 		Sig.bRequiresContext = false;
@@ -194,7 +194,7 @@ void UNiagaraDataInterfaceRenderTarget2DArray::GetFunctions(TArray<FNiagaraFunct
 		Sig.Inputs.Add(FNiagaraVariable(FNiagaraTypeDefinition::GetIntDef(), TEXT("Slice")));
 		Sig.Outputs.Add(FNiagaraVariable(FNiagaraTypeDefinition::GetColorDef(), TEXT("Value")));
 
-		Sig.bHidden = GNiagaraRenderTargetAllowReads != 1;
+		Sig.bHidden = NiagaraDataInterfaceRenderTargetCommon::GAllowReads != 1;
 		Sig.bExperimental = true;
 		Sig.bMemberFunction = true;
 		Sig.bRequiresContext = false;
@@ -447,18 +447,21 @@ bool UNiagaraDataInterfaceRenderTarget2DArray::InitPerInstanceData(void* PerInst
 {
 	check(Proxy);
 
-	extern bool GetRenderTargetFormat(bool bOverrideFormat, ETextureRenderTargetFormat OverrideFormat, ETextureRenderTargetFormat & OutRenderTargetFormat);
-	extern float GNiagaraRenderTargetResolutionMultiplier;
+	FRenderTarget2DArrayRWInstanceData_GameThread* InstanceData = new (PerInstanceData) FRenderTarget2DArrayRWInstanceData_GameThread();
+
+	if (NiagaraDataInterfaceRenderTargetCommon::GIgnoreCookedOut && !IsUsedWithGPUEmitter())
+	{
+		return true;
+	}
 
 	ETextureRenderTargetFormat RenderTargetFormat;
-	if (GetRenderTargetFormat(bOverrideFormat, OverrideRenderTargetFormat, RenderTargetFormat) == false)
+	if (NiagaraDataInterfaceRenderTargetCommon::GetRenderTargetFormat(bOverrideFormat, OverrideRenderTargetFormat, RenderTargetFormat) == false)
 	{
 		return false;
 	}
 
-	FRenderTarget2DArrayRWInstanceData_GameThread* InstanceData = new (PerInstanceData) FRenderTarget2DArrayRWInstanceData_GameThread();
-	InstanceData->Size.X = FMath::Clamp<int>(int(float(Size.X) * GNiagaraRenderTargetResolutionMultiplier), 1, GMaxTextureDimensions);
-	InstanceData->Size.Y = FMath::Clamp<int>(int(float(Size.Y) * GNiagaraRenderTargetResolutionMultiplier), 1, GMaxTextureDimensions);
+	InstanceData->Size.X = FMath::Clamp<int>(int(float(Size.X) * NiagaraDataInterfaceRenderTargetCommon::GResolutionMultiplier), 1, GMaxTextureDimensions);
+	InstanceData->Size.Y = FMath::Clamp<int>(int(float(Size.Y) * NiagaraDataInterfaceRenderTargetCommon::GResolutionMultiplier), 1, GMaxTextureDimensions);
 	InstanceData->Size.Z = FMath::Clamp<int>(Size.Z, 1, GMaxTextureDimensions);
 	InstanceData->Format = GetPixelFormatFromRenderTargetFormat(RenderTargetFormat);
 	InstanceData->RTUserParamBinding.Init(SystemInstance->GetInstanceParameters(), RenderTargetUserParameter.Parameter);
@@ -492,10 +495,9 @@ void UNiagaraDataInterfaceRenderTarget2DArray::DestroyPerInstanceData(void* PerI
 	);
 
 	// Make sure to clear out the reference to the render target if we created one.
-	extern int32 GNiagaraReleaseResourceOnRemove;
 	using RenderTargetType = decltype(decltype(ManagedRenderTargets)::ElementType::Value);
 	RenderTargetType ExistingRenderTarget = nullptr;
-	if (ManagedRenderTargets.RemoveAndCopyValue(SystemInstance->GetId(), ExistingRenderTarget) && GNiagaraReleaseResourceOnRemove)
+	if (ManagedRenderTargets.RemoveAndCopyValue(SystemInstance->GetId(), ExistingRenderTarget) && NiagaraDataInterfaceRenderTargetCommon::GReleaseResourceOnRemove)
 	{
 		ExistingRenderTarget->ReleaseResource();
 	}
@@ -528,7 +530,6 @@ void UNiagaraDataInterfaceRenderTarget2DArray::VMSetSize(FVectorVMExternalFuncti
 	FNDIInputParam<int> InSlices(Context);
 	FNDIOutputParam<FNiagaraBool> OutSuccess(Context);
 
-	extern float GNiagaraRenderTargetResolutionMultiplier;
 	for (int32 InstanceIdx = 0; InstanceIdx < Context.GetNumInstances(); ++InstanceIdx)
 	{
 		const int SizeX = InSizeX.GetAndAdvance();
@@ -538,8 +539,8 @@ void UNiagaraDataInterfaceRenderTarget2DArray::VMSetSize(FVectorVMExternalFuncti
 		OutSuccess.SetAndAdvance(bSuccess);
 		if (bSuccess)
 		{
-			InstData->Size.X = FMath::Clamp<int>(int(float(SizeX) * GNiagaraRenderTargetResolutionMultiplier), 1, GMaxTextureDimensions);
-			InstData->Size.Y = FMath::Clamp<int>(int(float(SizeY) * GNiagaraRenderTargetResolutionMultiplier), 1, GMaxTextureDimensions);
+			InstData->Size.X = FMath::Clamp<int>(int(float(SizeX) * NiagaraDataInterfaceRenderTargetCommon::GResolutionMultiplier), 1, GMaxTextureDimensions);
+			InstData->Size.Y = FMath::Clamp<int>(int(float(SizeY) * NiagaraDataInterfaceRenderTargetCommon::GResolutionMultiplier), 1, GMaxTextureDimensions);
 			InstData->Size.Z = FMath::Clamp<int>(Slices, 1, GMaxTextureDimensions);
 		}
 	}
@@ -570,9 +571,8 @@ bool UNiagaraDataInterfaceRenderTarget2DArray::PerInstanceTick(void* PerInstance
 	{
 		InstanceData->TargetTexture = UserTargetTexture;
 
-		extern int32 GNiagaraReleaseResourceOnRemove;
 		TObjectPtr<UTextureRenderTarget2DArray> ExistingRenderTarget = nullptr;
-		if (ManagedRenderTargets.RemoveAndCopyValue(SystemInstance->GetId(), ExistingRenderTarget) && GNiagaraReleaseResourceOnRemove)
+		if (ManagedRenderTargets.RemoveAndCopyValue(SystemInstance->GetId(), ExistingRenderTarget) && NiagaraDataInterfaceRenderTargetCommon::GReleaseResourceOnRemove)
 		{
 			ExistingRenderTarget->ReleaseResource();
 		}
@@ -615,12 +615,9 @@ bool UNiagaraDataInterfaceRenderTarget2DArray::PerInstanceTickPostSimulate(void*
 #endif
 
 	//-TEMP: Until we prune data interface on cook this will avoid consuming memory
+	if (NiagaraDataInterfaceRenderTargetCommon::GIgnoreCookedOut && !IsUsedWithGPUEmitter())
 	{
-		extern int32 GNiagaraRenderTargetIgnoreCookedOut;
-		if (GNiagaraRenderTargetIgnoreCookedOut && !IsUsedWithGPUEmitter())
-		{
-			return false;
-		}
+		return false;
 	}
 
 	if (!bInheritUserParameterSettings && (InstanceData->TargetTexture == nullptr))
diff --git a/Engine/Plugins/FX/Niagara/Source/Niagara/Private/NiagaraDataInterfaceRenderTargetCommon.cpp b/Engine/Plugins/FX/Niagara/Source/Niagara/Private/NiagaraDataInterfaceRenderTargetCommon.cpp
new file mode 100644
--- /dev/null
+++ b/Engine/Plugins/FX/Niagara/Source/Niagara/Private/NiagaraDataInterfaceRenderTargetCommon.cpp
@@ -0,0 +1,120 @@
+// Copyright Epic Games, Inc. All Rights Reserved.
+
+#include "NiagaraDataInterfaceRenderTargetCommon.h"
+
+#include "NiagaraCommon.h"
+#include "NiagaraSettings.h"
+
+namespace NiagaraDataInterfaceRenderTargetCommon
+{
+	int32 GReleaseResourceOnRemove = false;
+	static FAutoConsoleVariableRef CVarNiagaraReleaseResourceOnRemove(
+		TEXT("fx.Niagara.RenderTarget.ReleaseResourceOnRemove"),
+		GReleaseResourceOnRemove,
+		TEXT("Releases the render target resource once it is removed from the manager list rather than waiting for a GC."),
+		ECVF_Default
+	);
+
+	//-TEMP: Until we prune data interface on cook this will avoid consuming memory
+	int32 GIgnoreCookedOut = true;
+	static FAutoConsoleVariableRef CVarNiagaraRenderTargetIgnoreCookedOut(
+		TEXT("fx.Niagara.RenderTarget.IgnoreCookedOut"),
+		GIgnoreCookedOut,
+		TEXT("Ignores create render targets for cooked out emitter, i.e. ones that are not used by any GPU emitter."),
+		ECVF_Default
+	);
+
+	float GResolutionMultiplier = 1.0f;
+	static FAutoConsoleVariableRef CVarNiagaraRenderTargetResolutionMultiplier(
+		TEXT("fx.Niagara.RenderTarget.ResolutionMultiplier"),
+		GResolutionMultiplier,
+		TEXT("Optional global modifier to Niagara render target resolution."),
+		ECVF_Default
+	);
+
+	int GAllowReads = 0;
+	static FAutoConsoleVariableRef CVarNiagaraRenderTargetAllowReads(
+		TEXT("fx.Niagara.RenderTarget.AllowReads"),
+		GAllowReads,
+		TEXT("Enables read operations to be visible in the UI, very experimental."),
+		ECVF_Default
+	);
+
+	static bool GNiagaraRenderTargetOverrideFormatEnabled = false;
+	static ETextureRenderTargetFormat GNiagaraRenderTargetOverrideFormat = ETextureRenderTargetFormat::RTF_RGBA32f;
+
+	static FAutoConsoleCommandWithWorldAndArgs GCommandNiagaraRenderTargetOverrideFormat(
+		TEXT("fx.Niagara.RenderTarget.OverrideFormat"),
+		TEXT("Optional global format override for all Niagara render targets"),
+		FConsoleCommandWithWorldAndArgsDelegate::CreateLambda(
+			[](const TArray<FString>& Args, UWorld*)
+			{
+				static UEnum* TextureRenderTargetFormatEnum = StaticEnum<ETextureRenderTargetFormat>();
+				if ( ensure(TextureRenderTargetFormatEnum) )
+				{
+					if (Args.Num() == 1)
+					{
+						const int32 EnumIndex = TextureRenderTargetFormatEnum->GetIndexByNameString(Args[0]);
+						if (EnumIndex != INDEX_NONE)
+						{
+							GNiagaraRenderTargetOverrideFormatEnabled = true;
+							GNiagaraRenderTargetOverrideFormat = ETextureRenderTargetFormat(EnumIndex);
+						}
+						else
+						{
+							GNiagaraRenderTargetOverrideFormatEnabled = false;
+						}
+					}
+					UE_LOG(LogNiagara, Log, TEXT("Niagara RenderTarget Override is '%s' with format '%s'."), GNiagaraRenderTargetOverrideFormatEnabled ? TEXT("Enabled") : TEXT("Disabled"), *TextureRenderTargetFormatEnum->GetNameStringByIndex(int32(GNiagaraRenderTargetOverrideFormat)));
+				}
+			}
+		)
+	);
+
+	bool GetRenderTargetFormat(bool bOverrideFormat, ETextureRenderTargetFormat OverrideFormat, ETextureRenderTargetFormat& OutRenderTargetFormat)
+	{
+		OutRenderTargetFormat = bOverrideFormat ? OverrideFormat : GetDefault<UNiagaraSettings>()->DefaultRenderTargetFormat.GetValue();
+		EPixelFormat PixelFormat = GetPixelFormatFromRenderTargetFormat(OutRenderTargetFormat);
+		if (GNiagaraRenderTargetOverrideFormatEnabled)
+		{
+			OutRenderTargetFormat = GNiagaraRenderTargetOverrideFormat;
+		}
+
+		// If the format does not support typed store we need to find one that will
+		while (RHIIsTypedUAVStoreSupported(GetPixelFormatFromRenderTargetFormat(OutRenderTargetFormat)) == false)
+		{
+			static const TPair<ETextureRenderTargetFormat, ETextureRenderTargetFormat> FormatRemapTable[] =
+			{
+				TPairInitializer(RTF_R8,			RTF_R16f),
+				TPairInitializer(RTF_RG8,			RTF_RG16f),
+				TPairInitializer(RTF_RGBA8,			RTF_RGBA16f),
+				TPairInitializer(RTF_RGBA8_SRGB,	RTF_RGBA16f),
+				TPairInitializer(RTF_R16f,			RTF_R32f),
+				TPairInitializer(RTF_RG16f,			RTF_RG32f),
+				TPairInitializer(RTF_RGBA16f,		RTF_RGBA32f),
+				TPairInitializer(RTF_R32f,			RTF_RGBA32f),
+				TPairInitializer(RTF_RG32f,			RTF_RGBA32f),
+				TPairInitializer(RTF_RGBA32f,		RTF_RGBA32f),
+				TPairInitializer(RTF_RGB10A2,		RTF_RGBA32f),
+			};
+
+			const ETextureRenderTargetFormat PreviousFormat = OutRenderTargetFormat;
+			for (const auto& FormatRemap : FormatRemapTable)
+			{
+				if (FormatRemap.Key == OutRenderTargetFormat)
+				{
+					OutRenderTargetFormat = FormatRemap.Value;
+					break;
+				}
+			}
+			if (PreviousFormat == OutRenderTargetFormat)
+			{
+				// This is fatal as we failed to find any format that supports typed UAV stores
+				UE_LOG(LogNiagara, Warning, TEXT("Failed to find a render target format that supports UAV store"));
+				return false;
+			}
+		}
+
+		return true;
+	}
+}
diff --git a/Engine/Plugins/FX/Niagara/Source/Niagara/Private/NiagaraDataInterfaceRenderTargetCommon.h b/Engine/Plugins/FX/Niagara/Source/Niagara/Private/NiagaraDataInterfaceRenderTargetCommon.h
new file mode 100644
--- /dev/null
+++ b/Engine/Plugins/FX/Niagara/Source/Niagara/Private/NiagaraDataInterfaceRenderTargetCommon.h
@@ -0,0 +1,15 @@
+// Copyright Epic Games, Inc. All Rights Reserved.
+
+#pragma once
+
+#include "Engine/TextureRenderTarget2D.h"
+
+namespace NiagaraDataInterfaceRenderTargetCommon
+{
+	extern int32 GReleaseResourceOnRemove;
+	extern int32 GIgnoreCookedOut;
+	extern float GResolutionMultiplier;
+	extern int GAllowReads;
+
+	extern bool GetRenderTargetFormat(bool bOverrideFormat, ETextureRenderTargetFormat OverrideFormat, ETextureRenderTargetFormat& OutRenderTargetFormat);
+}
diff --git a/Engine/Plugins/FX/Niagara/Source/Niagara/Private/NiagaraDataInterfaceRenderTargetCube.cpp b/Engine/Plugins/FX/Niagara/Source/Niagara/Private/NiagaraDataInterfaceRenderTargetCube.cpp
--- a/Engine/Plugins/FX/Niagara/Source/Niagara/Private/NiagaraDataInterfaceRenderTargetCube.cpp
+++ b/Engine/Plugins/FX/Niagara/Source/Niagara/Private/NiagaraDataInterfaceRenderTargetCube.cpp
@@ -5,6 +5,7 @@
 #include "TextureResource.h"
 #include "Engine/TextureRenderTargetCube.h"
 
+#include "NiagaraDataInterfaceRenderTargetCommon.h"
 #include "NiagaraSystemInstance.h"
 #include "NiagaraStats.h"
 #include "NiagaraRenderer.h"
@@ -157,7 +158,6 @@ void UNiagaraDataInterfaceRenderTargetCube::GetFunctions(TArray<FNiagaraFunction
 	#endif
 	}
 
-	extern int GNiagaraRenderTargetAllowReads;
 // CubeTexture.Load() is not available
 //	{
 //		FNiagaraFunctionSignature& Sig = OutFunctions.AddDefaulted_GetRef();
@@ -168,7 +168,7 @@ void UNiagaraDataInterfaceRenderTargetCube::GetFunctions(TArray<FNiagaraFunction
 //		Sig.Inputs.Add(FNiagaraVariable(FNiagaraTypeDefinition::GetIntDef(), TEXT("Face")));
 //		Sig.Outputs.Add(FNiagaraVariable(FNiagaraTypeDefinition::GetColorDef(), TEXT("Value")));
 //
-//		Sig.bHidden = GNiagaraRenderTargetAllowReads != 1;
+//		Sig.bHidden = NiagaraDataInterfaceRenderTargetCommon::GAllowReads != 1;
 //		Sig.bExperimental = true;
 //		Sig.bMemberFunction = true;
 //		Sig.bRequiresContext = false;
@@ -187,7 +187,7 @@ void UNiagaraDataInterfaceRenderTargetCube::GetFunctions(TArray<FNiagaraFunction
 		Sig.Inputs.Add(FNiagaraVariable(FNiagaraTypeDefinition::GetVec3Def(), TEXT("UVW")));
 		Sig.Outputs.Add(FNiagaraVariable(FNiagaraTypeDefinition::GetColorDef(), TEXT("Value")));
 
-		Sig.bHidden = GNiagaraRenderTargetAllowReads != 1;
+		Sig.bHidden = NiagaraDataInterfaceRenderTargetCommon::GAllowReads != 1;
 		Sig.bExperimental = true;
 		Sig.bMemberFunction = true;
 		Sig.bRequiresContext = false;
@@ -435,17 +435,20 @@ bool UNiagaraDataInterfaceRenderTargetCube::InitPerInstanceData(void* PerInstanc
 {
 	check(Proxy);
 
-	extern bool GetRenderTargetFormat(bool bOverrideFormat, ETextureRenderTargetFormat OverrideFormat, ETextureRenderTargetFormat & OutRenderTargetFormat);
-	extern float GNiagaraRenderTargetResolutionMultiplier;
+	FRenderTargetCubeRWInstanceData_GameThread* InstanceData = new (PerInstanceData) FRenderTargetCubeRWInstanceData_GameThread();
+
+	if (NiagaraDataInterfaceRenderTargetCommon::GIgnoreCookedOut && !IsUsedWithGPUEmitter())
+	{
+		return true;
+	}
 
 	ETextureRenderTargetFormat RenderTargetFormat;
-	if (GetRenderTargetFormat(bOverrideFormat, OverrideRenderTargetFormat, RenderTargetFormat) == false)
+	if (NiagaraDataInterfaceRenderTargetCommon::GetRenderTargetFormat(bOverrideFormat, OverrideRenderTargetFormat, RenderTargetFormat) == false)
 	{
 		return false;
 	}
 
-	FRenderTargetCubeRWInstanceData_GameThread* InstanceData = new (PerInstanceData) FRenderTargetCubeRWInstanceData_GameThread();
-	InstanceData->Size = FMath::Clamp<int>(int(float(Size) * GNiagaraRenderTargetResolutionMultiplier), 1, GMaxCubeTextureDimensions);
+	InstanceData->Size = FMath::Clamp<int>(int(float(Size) * NiagaraDataInterfaceRenderTargetCommon::GResolutionMultiplier), 1, GMaxCubeTextureDimensions);
 	InstanceData->Format = GetPixelFormatFromRenderTargetFormat(RenderTargetFormat);
 	InstanceData->RTUserParamBinding.Init(SystemInstance->GetInstanceParameters(), RenderTargetUserParameter.Parameter);
 #if WITH_EDITORONLY_DATA
@@ -478,9 +481,8 @@ void UNiagaraDataInterfaceRenderTargetCube::DestroyPerInstanceData(void* PerInst
 	);
 
 	// Make sure to clear out the reference to the render target if we created one.
-	extern int32 GNiagaraReleaseResourceOnRemove;
 	decltype(ManagedRenderTargets)::ValueType ExistingRenderTarget = nullptr;
-	if (ManagedRenderTargets.RemoveAndCopyValue(SystemInstance->GetId(), ExistingRenderTarget) && GNiagaraReleaseResourceOnRemove)
+	if (ManagedRenderTargets.RemoveAndCopyValue(SystemInstance->GetId(), ExistingRenderTarget) && NiagaraDataInterfaceRenderTargetCommon::GReleaseResourceOnRemove)
 	{
 		ExistingRenderTarget->ReleaseResource();
 	}
@@ -511,7 +513,6 @@ void UNiagaraDataInterfaceRenderTargetCube::VMSetSize(FVectorVMExternalFunctionC
 	FNDIInputParam<int> InSize(Context);
 	FNDIOutputParam<FNiagaraBool> OutSuccess(Context);
 
-	extern float GNiagaraRenderTargetResolutionMultiplier;
 	for (int32 InstanceIdx = 0; InstanceIdx < Context.GetNumInstances(); ++InstanceIdx)
 	{
 		const int NewSize = InSize.GetAndAdvance();
@@ -519,7 +520,7 @@ void UNiagaraDataInterfaceRenderTargetCube::VMSetSize(FVectorVMExternalFunctionC
 		OutSuccess.SetAndAdvance(bSuccess);
 		if (bSuccess)
 		{
-			InstData->Size = FMath::Clamp<int>(int(float(NewSize) * GNiagaraRenderTargetResolutionMultiplier), 1, GMaxCubeTextureDimensions);
+			InstData->Size = FMath::Clamp<int>(int(float(NewSize) * NiagaraDataInterfaceRenderTargetCommon::GResolutionMultiplier), 1, GMaxCubeTextureDimensions);
 		}
 	}
 }
@@ -546,9 +547,8 @@ bool UNiagaraDataInterfaceRenderTargetCube::PerInstanceTick(void* PerInstanceDat
 	{
 		InstanceData->TargetTexture = UserTargetTexture;
 
-		extern int32 GNiagaraReleaseResourceOnRemove;
 		decltype(ManagedRenderTargets)::ValueType ExistingRenderTarget = nullptr;
-		if (ManagedRenderTargets.RemoveAndCopyValue(SystemInstance->GetId(), ExistingRenderTarget) && GNiagaraReleaseResourceOnRemove)
+		if (ManagedRenderTargets.RemoveAndCopyValue(SystemInstance->GetId(), ExistingRenderTarget) && NiagaraDataInterfaceRenderTargetCommon::GReleaseResourceOnRemove)
 		{
 			ExistingRenderTarget->ReleaseResource();
 		}
@@ -589,12 +589,9 @@ bool UNiagaraDataInterfaceRenderTargetCube::PerInstanceTickPostSimulate(void* Pe
 #endif
 
 	//-TEMP: Until we prune data interface on cook this will avoid consuming memory
+	if (NiagaraDataInterfaceRenderTargetCommon::GIgnoreCookedOut && !IsUsedWithGPUEmitter())
 	{
-		extern int32 GNiagaraRenderTargetIgnoreCookedOut;
-		if (GNiagaraRenderTargetIgnoreCookedOut && !IsUsedWithGPUEmitter())
-		{
-			return false;
-		}
+		return false;
 	}
 
 	// Do we need to create a new texture?
diff --git a/Engine/Plugins/FX/Niagara/Source/Niagara/Private/NiagaraDataInterfaceRenderTargetVolume.cpp b/Engine/Plugins/FX/Niagara/Source/Niagara/Private/NiagaraDataInterfaceRenderTargetVolume.cpp
--- a/Engine/Plugins/FX/Niagara/Source/Niagara/Private/NiagaraDataInterfaceRenderTargetVolume.cpp
+++ b/Engine/Plugins/FX/Niagara/Source/Niagara/Private/NiagaraDataInterfaceRenderTargetVolume.cpp
@@ -6,6 +6,7 @@
 #include "TextureRenderTargetVolumeResource.h"
 #include "Engine/TextureRenderTargetVolume.h"
 
+#include "NiagaraDataInterfaceRenderTargetCommon.h"
 #include "NiagaraGpuComputeDispatchInterface.h"
 #include "NiagaraSystemInstance.h"
 #include "NiagaraStats.h"
@@ -162,7 +163,6 @@ void UNiagaraDataInterfaceRenderTargetVolume::GetFunctions(TArray<FNiagaraFuncti
 	#endif
 	}
 
-	extern int GNiagaraRenderTargetAllowReads;
 	{
 		FNiagaraFunctionSignature& Sig = OutFunctions.AddDefaulted_GetRef();
 		Sig.Name = GetValueFunctionName;
@@ -172,7 +172,7 @@ void UNiagaraDataInterfaceRenderTargetVolume::GetFunctions(TArray<FNiagaraFuncti
 		Sig.Inputs.Add(FNiagaraVariable(FNiagaraTypeDefinition::GetIntDef(), TEXT("IndexZ")));
 		Sig.Outputs.Add(FNiagaraVariable(FNiagaraTypeDefinition::GetColorDef(), TEXT("Value")));
 
-		Sig.bHidden = GNiagaraRenderTargetAllowReads != 1;
+		Sig.bHidden = NiagaraDataInterfaceRenderTargetCommon::GAllowReads != 1;
 		Sig.bExperimental = true;
 		Sig.bMemberFunction = true;
 		Sig.bRequiresContext = false;
@@ -191,7 +191,7 @@ void UNiagaraDataInterfaceRenderTargetVolume::GetFunctions(TArray<FNiagaraFuncti
 		Sig.Inputs.Add(FNiagaraVariable(FNiagaraTypeDefinition::GetVec3Def(), TEXT("UVW")));
 		Sig.Outputs.Add(FNiagaraVariable(FNiagaraTypeDefinition::GetColorDef(), TEXT("Value")));
 
-		Sig.bHidden = GNiagaraRenderTargetAllowReads != 1;
+		Sig.bHidden = NiagaraDataInterfaceRenderTargetCommon::GAllowReads != 1;
 		Sig.bExperimental = true;
 		Sig.bMemberFunction = true;
 		Sig.bRequiresContext = false;
@@ -439,19 +439,22 @@ bool UNiagaraDataInterfaceRenderTargetVolume::InitPerInstanceData(void* PerInsta
 {
 	check(Proxy);
 
-	extern bool GetRenderTargetFormat(bool bOverrideFormat, ETextureRenderTargetFormat OverrideFormat, ETextureRenderTargetFormat & OutRenderTargetFormat);
-	extern float GNiagaraRenderTargetResolutionMultiplier;
+	FRenderTargetVolumeRWInstanceData_GameThread* InstanceData = new (PerInstanceData) FRenderTargetVolumeRWInstanceData_GameThread();
+
+	if (NiagaraDataInterfaceRenderTargetCommon::GIgnoreCookedOut && !IsUsedWithGPUEmitter())
+	{
+		return true;
+	}
 
 	ETextureRenderTargetFormat RenderTargetFormat;
-	if (GetRenderTargetFormat(bOverrideFormat, OverrideRenderTargetFormat, RenderTargetFormat) == false)
+	if (NiagaraDataInterfaceRenderTargetCommon::GetRenderTargetFormat(bOverrideFormat, OverrideRenderTargetFormat, RenderTargetFormat) == false)
 	{
 		return false;
 	}
 
-	FRenderTargetVolumeRWInstanceData_GameThread* InstanceData = new (PerInstanceData) FRenderTargetVolumeRWInstanceData_GameThread();
-	InstanceData->Size.X = FMath::Clamp<int>(int(float(Size.X) * GNiagaraRenderTargetResolutionMultiplier), 1, GMaxVolumeTextureDimensions);
-	InstanceData->Size.Y = FMath::Clamp<int>(int(float(Size.Y) * GNiagaraRenderTargetResolutionMultiplier), 1, GMaxVolumeTextureDimensions);
-	InstanceData->Size.Z = FMath::Clamp<int>(int(float(Size.Z) * GNiagaraRenderTargetResolutionMultiplier), 1, GMaxVolumeTextureDimensions);
+	InstanceData->Size.X = FMath::Clamp<int>(int(float(Size.X) * NiagaraDataInterfaceRenderTargetCommon::GResolutionMultiplier), 1, GMaxVolumeTextureDimensions);
+	InstanceData->Size.Y = FMath::Clamp<int>(int(float(Size.Y) * NiagaraDataInterfaceRenderTargetCommon::GResolutionMultiplier), 1, GMaxVolumeTextureDimensions);
+	InstanceData->Size.Z = FMath::Clamp<int>(int(float(Size.Z) * NiagaraDataInterfaceRenderTargetCommon::GResolutionMultiplier), 1, GMaxVolumeTextureDimensions);
 	InstanceData->Format = GetPixelFormatFromRenderTargetFormat(RenderTargetFormat);
 	InstanceData->RTUserParamBinding.Init(SystemInstance->GetInstanceParameters(), RenderTargetUserParameter.Parameter);
 #if WITH_EDITORONLY_DATA
@@ -484,10 +487,9 @@ void UNiagaraDataInterfaceRenderTargetVolume::DestroyPerInstanceData(void* PerIn
 	);
 
 	// Make sure to clear out the reference to the render target if we created one.
-	extern int32 GNiagaraReleaseResourceOnRemove;
 	using RenderTargetType = decltype(decltype(ManagedRenderTargets)::ElementType::Value);
 	RenderTargetType ExistingRenderTarget = nullptr;
-	if (ManagedRenderTargets.RemoveAndCopyValue(SystemInstance->GetId(), ExistingRenderTarget) && GNiagaraReleaseResourceOnRemove)
+	if (ManagedRenderTargets.RemoveAndCopyValue(SystemInstance->GetId(), ExistingRenderTarget) && NiagaraDataInterfaceRenderTargetCommon::GReleaseResourceOnRemove)
 	{
 		ExistingRenderTarget->ReleaseResource();
 	}
@@ -520,7 +522,6 @@ void UNiagaraDataInterfaceRenderTargetVolume::VMSetSize(FVectorVMExternalFunctio
 	FNDIInputParam<int> InSizeZ(Context);
 	FNDIOutputParam<FNiagaraBool> OutSuccess(Context);
 
-	extern float GNiagaraRenderTargetResolutionMultiplier;
 	for (int32 InstanceIdx = 0; InstanceIdx < Context.GetNumInstances(); ++InstanceIdx)
 	{
 		const int SizeX = InSizeX.GetAndAdvance();
@@ -530,9 +531,9 @@ void UNiagaraDataInterfaceRenderTargetVolume::VMSetSize(FVectorVMExternalFunctio
 		OutSuccess.SetAndAdvance(bSuccess);
 		if (bSuccess)
 		{
-			InstData->Size.X = FMath::Clamp<int>(int(float(SizeX) * GNiagaraRenderTargetResolutionMultiplier), 1, GMaxVolumeTextureDimensions);
-			InstData->Size.Y = FMath::Clamp<int>(int(float(SizeY) * GNiagaraRenderTargetResolutionMultiplier), 1, GMaxVolumeTextureDimensions);
-			InstData->Size.Z = FMath::Clamp<int>(int(float(SizeZ) * GNiagaraRenderTargetResolutionMultiplier), 1, GMaxVolumeTextureDimensions);
+			InstData->Size.X = FMath::Clamp<int>(int(float(SizeX) * NiagaraDataInterfaceRenderTargetCommon::GResolutionMultiplier), 1, GMaxVolumeTextureDimensions);
+			InstData->Size.Y = FMath::Clamp<int>(int(float(SizeY) * NiagaraDataInterfaceRenderTargetCommon::GResolutionMultiplier), 1, GMaxVolumeTextureDimensions);
+			InstData->Size.Z = FMath::Clamp<int>(int(float(SizeZ) * NiagaraDataInterfaceRenderTargetCommon::GResolutionMultiplier), 1, GMaxVolumeTextureDimensions);
 		}
 	}
 }
@@ -563,9 +564,8 @@ bool UNiagaraDataInterfaceRenderTargetVolume::PerInstanceTick(void* PerInstanceD
 	{
 		InstanceData->TargetTexture = UserTargetTexture;
 
-		extern int32 GNiagaraReleaseResourceOnRemove;
 		TObjectPtr<UTextureRenderTargetVolume> ExistingRenderTarget = nullptr;
-		if (ManagedRenderTargets.RemoveAndCopyValue(SystemInstance->GetId(), ExistingRenderTarget) && GNiagaraReleaseResourceOnRemove)
+		if (ManagedRenderTargets.RemoveAndCopyValue(SystemInstance->GetId(), ExistingRenderTarget) && NiagaraDataInterfaceRenderTargetCommon::GReleaseResourceOnRemove)
 		{
 			ExistingRenderTarget->ReleaseResource();
 		}
@@ -608,12 +608,9 @@ bool UNiagaraDataInterfaceRenderTargetVolume::PerInstanceTickPostSimulate(void*
 #endif
 
 	//-TEMP: Until we prune data interface on cook this will avoid consuming memory
+	if (NiagaraDataInterfaceRenderTargetCommon::GIgnoreCookedOut && !IsUsedWithGPUEmitter())
 	{
-		extern int32 GNiagaraRenderTargetIgnoreCookedOut;
-		if (GNiagaraRenderTargetIgnoreCookedOut && !IsUsedWithGPUEmitter())
-		{
-			return false;
-		}
+		return false;
 	}
 
 	// Do we need to create a new texture?
