Ensure InstanceOverrideParameter DIs are post loaded before they're applied to the component's owner parameters.
Without this, post load fixup in instance DIs is not applied before it's used and so data can be lost.

Current example is the skeletal mesh DI post loading it's Soruce actor pointer into a soft object ptr.

diff --git a/Engine/Plugins/FX/Niagara/Source/Niagara/Private/NiagaraComponent.cpp b/Engine/Plugins/FX/Niagara/Source/Niagara/Private/NiagaraComponent.cpp
--- a/Engine/Plugins/FX/Niagara/Source/Niagara/Private/NiagaraComponent.cpp
+++ b/Engine/Plugins/FX/Niagara/Source/Niagara/Private/NiagaraComponent.cpp
@@ -2574,6 +2574,24 @@ void UNiagaraComponent::PostLoad()
 {
 	Super::PostLoad();
 
+#if WITH_EDITORONLY_DATA
+	auto PostLoadVariantDIs = [](auto& VariantMap)
+	{
+		for (const auto& Pair : VariantMap)
+		{
+			if (Pair.Value.GetMode() == ENiagaraVariantMode::DataInterface)
+			{
+				if (UNiagaraDataInterface* DI = Pair.Value.GetDataInterface())
+				{
+					DI->ConditionalPostLoad();
+				}
+			}
+		}
+	};
+	PostLoadVariantDIs(TemplateParameterOverrides);
+	PostLoadVariantDIs(InstanceParameterOverrides);
+#endif
+
 	OverrideParameters.PostLoad();
 	OverrideParameters.SanityCheckData();
 
