- Fix duplicate cvar name

#rnx

diff --git a/Engine/Plugins/FX/Niagara/Source/NiagaraEditor/Private/Toolkits/NiagaraSystemToolkit.cpp b/Engine/Plugins/FX/Niagara/Source/NiagaraEditor/Private/Toolkits/NiagaraSystemToolkit.cpp
--- a/Engine/Plugins/FX/Niagara/Source/NiagaraEditor/Private/Toolkits/NiagaraSystemToolkit.cpp
+++ b/Engine/Plugins/FX/Niagara/Source/NiagaraEditor/Private/Toolkits/NiagaraSystemToolkit.cpp
@@ -72,7 +72,7 @@ const FName FNiagaraSystemToolkit::ScalabilityModeName(TEXT("Scalability"));
 IConsoleVariable* FNiagaraSystemToolkit::VmStatEnabledVar = IConsoleManager::Get().FindConsoleVariable(TEXT("vm.DetailedVMScriptStats"));
 
 static int32 GbLogNiagaraSystemChanges = 0;
-static FAutoConsoleVariableRef CVarSuppressNiagaraSystems(
+static FAutoConsoleVariableRef CVarLogNiagaraSystemChanges(
 	TEXT("fx.LogNiagaraSystemChanges"),
 	GbLogNiagaraSystemChanges,
 	TEXT("If > 0 Niagara Systems will be written to a text format when opened and closed in the editor. \n"),
