- When in EOF updates we must use an async task for the component destroy

#rnx

diff --git a/Engine/Plugins/FX/Niagara/Source/Niagara/Private/NiagaraRendererComponents.cpp b/Engine/Plugins/FX/Niagara/Source/Niagara/Private/NiagaraRendererComponents.cpp
--- a/Engine/Plugins/FX/Niagara/Source/Niagara/Private/NiagaraRendererComponents.cpp
+++ b/Engine/Plugins/FX/Niagara/Source/Niagara/Private/NiagaraRendererComponents.cpp
@@ -288,38 +288,54 @@ void FNiagaraRendererComponents::DestroyRenderState_Concurrent()
 	}
 #endif
 
-	// If we are on the GameThread immediately execute
+	// If we are on the GameThread immediately execute, but only if not inside EOF updates
 	// Failure to do so can cause issue with construction scripts rerunning as the components can linger when they should not
+	bool bRequiresAyncTask = true;
 	if (IsInGameThread())
 	{
-		for (auto& PoolEntry : ComponentPool)
+		UWorld* FoundWorld = nullptr;
+		if (AActor* OwnerActor = SpawnedOwner.Get())
 		{
-			if (PoolEntry.Component.IsValid())
+			FoundWorld = OwnerActor->GetWorld();
+			if (FoundWorld != nullptr)
 			{
-				PoolEntry.Component->DestroyComponent();
+				bRequiresAyncTask = FoundWorld->bPostTickComponentUpdate;
 			}
 		}
-		ComponentPool.Empty();
 
-		if (AActor* OwnerActor = SpawnedOwner.Get())
+		if ( FoundWorld == nullptr )
 		{
-			OwnerActor->Destroy();
+			for (auto& PoolEntry : ComponentPool)
+			{
+				if (USceneComponent* Component = PoolEntry.Component.Get())
+				{
+					FoundWorld = Component->GetWorld();
+					if (FoundWorld != nullptr)
+					{
+						bRequiresAyncTask = FoundWorld->bPostTickComponentUpdate;
+						break;
+					}
+				}
+			}
+
 		}
+		ComponentPool.Empty();
 	}
-	else
+
+	if (bRequiresAyncTask)
 	{
 		// Rendering resources are being destroyed, but the component pool and their owner actor must be destroyed on the game thread
 		AsyncTask(
 			ENamedThreads::GameThread,
-			[Pool_GT = MoveTemp(ComponentPool), Owner_GT = MoveTemp(SpawnedOwner)]()
+			[Pool_GT=MoveTemp(ComponentPool), Owner_GT=MoveTemp(SpawnedOwner)]()
 			{
 				// we do not reset ParticlesWithComponents here because it's possible the render state is destroyed without destroying the renderer. In this case we want to know which particles
 				// had spawned some components previously
 				for (auto& PoolEntry : Pool_GT)
 				{
-					if (PoolEntry.Component.IsValid())
+					if (USceneComponent* Component = PoolEntry.Component.Get())
 					{
-						PoolEntry.Component->DestroyComponent();
+						Component->DestroyComponent();
 					}
 				}
 
@@ -330,6 +346,22 @@ void FNiagaraRendererComponents::DestroyRenderState_Concurrent()
 			}
 		);
 	}
+	else
+	{
+		if (AActor* OwnerActor = SpawnedOwner.Get())
+		{
+			OwnerActor->Destroy();
+		}
+
+		for (auto& PoolEntry : ComponentPool)
+		{
+			if (USceneComponent* Component = PoolEntry.Component.Get())
+			{
+				Component->DestroyComponent();
+			}
+		}
+		ComponentPool.Empty();
+	}
 	SpawnedOwner.Reset();
 }
 
