Niagara - Refactor hierarchy view model initialization to remove circular shared ptr reference between the system view model and user parameter hierarchy view models.  This was causing the view model to not get destructed and subsequently the stack view models and entries were never properly finalized.

diff --git a/Engine/Plugins/FX/Niagara/Source/NiagaraEditor/Private/ViewModels/HierarchyEditor/NiagaraHierarchyViewModelBase.cpp b/Engine/Plugins/FX/Niagara/Source/NiagaraEditor/Private/ViewModels/HierarchyEditor/NiagaraHierarchyViewModelBase.cpp
--- a/Engine/Plugins/FX/Niagara/Source/NiagaraEditor/Private/ViewModels/HierarchyEditor/NiagaraHierarchyViewModelBase.cpp
+++ b/Engine/Plugins/FX/Niagara/Source/NiagaraEditor/Private/ViewModels/HierarchyEditor/NiagaraHierarchyViewModelBase.cpp
@@ -270,10 +270,8 @@ void UNiagaraHierarchyViewModelBase::PostRedo(bool bSuccess)
 	PostUndo(bSuccess);
 }
 
-void UNiagaraHierarchyViewModelBase::Initialize(TSharedPtr<FNiagaraSystemViewModel> InSystemViewModel)
+void UNiagaraHierarchyViewModelBase::Initialize()
 {
-	SystemViewModel = InSystemViewModel;
-	
 	if(GEditor)
 	{
 		GEditor->RegisterForUndo(this);
diff --git a/Engine/Plugins/FX/Niagara/Source/NiagaraEditor/Private/ViewModels/HierarchyEditor/NiagaraUserParametersHierarchyViewModel.cpp b/Engine/Plugins/FX/Niagara/Source/NiagaraEditor/Private/ViewModels/HierarchyEditor/NiagaraUserParametersHierarchyViewModel.cpp
--- a/Engine/Plugins/FX/Niagara/Source/NiagaraEditor/Private/ViewModels/HierarchyEditor/NiagaraUserParametersHierarchyViewModel.cpp
+++ b/Engine/Plugins/FX/Niagara/Source/NiagaraEditor/Private/ViewModels/HierarchyEditor/NiagaraUserParametersHierarchyViewModel.cpp
@@ -33,14 +33,25 @@ void UNiagaraHierarchyUserParameter::RefreshDataInternal()
 	}
 }
 
-void UNiagaraUserParametersHierarchyViewModel::InitializeInternal()
+TSharedRef<FNiagaraSystemViewModel> UNiagaraUserParametersHierarchyViewModel::GetSystemViewModel() const
 {
-	UNiagaraSystemEditorData* SystemEditorData = Cast<UNiagaraSystemEditorData>(GetSystemViewModel()->GetSystem().GetEditorData());
-	SystemEditorData->OnUserParameterScriptVariablesSynced().AddUObject(this, &UNiagaraUserParametersHierarchyViewModel::ForceFullRefresh);
+	TSharedPtr<FNiagaraSystemViewModel> SystemViewModelPinned = SystemViewModelWeak.Pin();
+	checkf(SystemViewModelPinned.IsValid(), TEXT("System view model destroyed before user parameters hierarchy view model."));
+	return SystemViewModelPinned.ToSharedRef();
+}
+
+void UNiagaraUserParametersHierarchyViewModel::Initialize(TSharedRef<FNiagaraSystemViewModel> InSystemViewModel)
+{
+	SystemViewModelWeak = InSystemViewModel;
+	UNiagaraHierarchyViewModelBase::Initialize();
+
+	UNiagaraSystemEditorData& SystemEditorData = InSystemViewModel->GetEditorData();
+	SystemEditorData.OnUserParameterScriptVariablesSynced().AddUObject(this, &UNiagaraUserParametersHierarchyViewModel::ForceFullRefresh);
 }
 
 void UNiagaraUserParametersHierarchyViewModel::FinalizeInternal()
 {
+	TSharedPtr<FNiagaraSystemViewModel> SystemViewModel = SystemViewModelWeak.Pin();
 	if(SystemViewModel.IsValid())
 	{
 		GetSystemViewModel()->GetSystem().GetExposedParameters().RemoveAllOnChangedHandlers(this);
diff --git a/Engine/Plugins/FX/Niagara/Source/NiagaraEditor/Public/ViewModels/HierarchyEditor/NiagaraHierarchyViewModelBase.h b/Engine/Plugins/FX/Niagara/Source/NiagaraEditor/Public/ViewModels/HierarchyEditor/NiagaraHierarchyViewModelBase.h
--- a/Engine/Plugins/FX/Niagara/Source/NiagaraEditor/Public/ViewModels/HierarchyEditor/NiagaraHierarchyViewModelBase.h
+++ b/Engine/Plugins/FX/Niagara/Source/NiagaraEditor/Public/ViewModels/HierarchyEditor/NiagaraHierarchyViewModelBase.h
@@ -247,11 +247,9 @@ public:
 	UNiagaraHierarchyViewModelBase();
 	virtual ~UNiagaraHierarchyViewModelBase() override;
 
-	void Initialize(TSharedPtr<FNiagaraSystemViewModel> InSystemViewModel);
+	void Initialize();
 	void Finalize();
 
-	TSharedPtr<FNiagaraSystemViewModel> GetSystemViewModel() const { return SystemViewModel; }
-
 	void AddCategory() const;
 	void AddSection() const;
 
@@ -328,8 +326,6 @@ protected:
 	TWeakPtr<struct FNiagaraHierarchySectionViewModel> ActiveSection;
 	TSharedPtr<FUICommandList> Commands;
 
-	TSharedPtr<FNiagaraSystemViewModel> SystemViewModel;
-
 protected:
 	// delegate collection to call UI functions
 	FRefreshTreeWidget RefreshSourceViewDelegate;
diff --git a/Engine/Plugins/FX/Niagara/Source/NiagaraEditor/Public/ViewModels/HierarchyEditor/NiagaraUserParametersHierarchyViewModel.h b/Engine/Plugins/FX/Niagara/Source/NiagaraEditor/Public/ViewModels/HierarchyEditor/NiagaraUserParametersHierarchyViewModel.h
--- a/Engine/Plugins/FX/Niagara/Source/NiagaraEditor/Public/ViewModels/HierarchyEditor/NiagaraUserParametersHierarchyViewModel.h
+++ b/Engine/Plugins/FX/Niagara/Source/NiagaraEditor/Public/ViewModels/HierarchyEditor/NiagaraUserParametersHierarchyViewModel.h
@@ -40,7 +40,9 @@ public:
 	{
 	}
 
-	virtual void InitializeInternal() override;
+	TSharedRef<FNiagaraSystemViewModel> GetSystemViewModel() const;
+
+	void Initialize(TSharedRef<FNiagaraSystemViewModel> InSystemViewModel);
 	
 	virtual UNiagaraHierarchyRoot* GetHierarchyDataRoot() const override;
 	
@@ -58,7 +60,7 @@ public:
 protected:
 	virtual void FinalizeInternal() override;
 private:
-	TWeakPtr<FNiagaraSystemViewModel> SystemViewModel;
+	TWeakPtr<FNiagaraSystemViewModel> SystemViewModelWeak;
 };
 
 class FNiagaraUserParameterDragDropOp : public FNiagaraHierarchyDragDropOp
