Fixed SA issue

diff --git a/Engine/Plugins/FX/Niagara/Source/NiagaraEditor/Private/NiagaraNodeParameterMapGet.cpp b/Engine/Plugins/FX/Niagara/Source/NiagaraEditor/Private/NiagaraNodeParameterMapGet.cpp
--- a/Engine/Plugins/FX/Niagara/Source/NiagaraEditor/Private/NiagaraNodeParameterMapGet.cpp
+++ b/Engine/Plugins/FX/Niagara/Source/NiagaraEditor/Private/NiagaraNodeParameterMapGet.cpp
@@ -424,7 +424,7 @@ FText UNiagaraNodeParameterMapGet::GetNodeTitle(ENodeTitleType::Type TitleType)
 
 void UNiagaraNodeParameterMapGet::AddOrphanedPinPairGuids(UEdGraphPin* OutputPin, UEdGraphPin* DefaultPin)
 {
-	ensure(OutputPin->bOrphanedPin && DefaultPin->bOrphanedPin && OutputPin->PersistentGuid.IsValid() && DefaultPin->PersistentGuid.IsValid());
+	ensure(OutputPin && DefaultPin && OutputPin->bOrphanedPin && DefaultPin->bOrphanedPin && OutputPin->PersistentGuid.IsValid() && DefaultPin->PersistentGuid.IsValid());
 
 	PinOutputToPinDefaultPersistentId.Add(OutputPin->PersistentGuid, DefaultPin->PersistentGuid);
 }
