Generalize the UvMapping feature so that it can be shared by both static and skeletal meshes
-takes a different appraoch with specifying the uv index for the static mesh via a system script level VM function

diff --git a/Engine/Plugins/FX/Niagara/Shaders/Private/Experimental/NiagaraUvMappingUtils.ush b/Engine/Plugins/FX/Niagara/Shaders/Private/Experimental/NiagaraUvMappingUtils.ush
new file mode 100644
--- /dev/null
+++ b/Engine/Plugins/FX/Niagara/Shaders/Private/Experimental/NiagaraUvMappingUtils.ush
@@ -0,0 +1,266 @@
+// Copyright Epic Games, Inc. All Rights Reserved.
+
+#include "/Plugin/FX/Niagara/Private/NiagaraBaryCentricUtils.ush"
+
+struct UvMappingParameters
+{
+	int				NumUVs;
+	Buffer<uint>	IndexBuffer;
+	Buffer<float2>	UVBuffer;
+	Buffer<int>		UvMappingBuffer;
+	uint			UvMappingBufferLength;
+	uint			UvMappingSet;
+};
+
+#define UvMapping_ElementCountOffset 4
+#define UvMapping_TriangleIndexOffset 5
+#define UvMapping_MinValidUvMappingBufferSize 24 // minimum size (in bytes) for a quadtree to have valid data
+
+bool UvMapping_NormalizedAabbTriangleIntersection(float2 A, float2 B, float2 C)
+{
+	float2 TriAabbMin = float2(min3(A.x, B.x, C.x), min3(A.y, B.y, C.y));
+	float2 TriAabbMax = float2(max3(A.x, B.x, C.x), max3(A.y, B.y, C.y));
+
+	if (any(TriAabbMin > 1.0f) || any(TriAabbMax < 0.0f))
+	{
+		return false;
+	}
+
+	#define AXIS_COUNT 3
+	float2 TriangleEdges[AXIS_COUNT] = { C - B, A - C, B - A };
+
+	for (int i = 0; i < AXIS_COUNT; ++i)
+	{
+		float2 Axis = TriangleEdges[i].yx * (float2(-1.0f, 0.0f));
+		float AabbSegmentMin = min(0.0f, min3(Axis.x, Axis.y, Axis.x + Axis.y));
+		float AabbSegmentMax = max(0.0f, max3(Axis.x, Axis.y, Axis.x + Axis.y));
+		float TriangleSegmentMin = min3(dot(A, Axis), dot(B, Axis), dot(C, Axis));
+		float TriangleSegmentMax = max3(dot(A, Axis), dot(B, Axis), dot(C, Axis));
+		if (AabbSegmentMin > TriangleSegmentMax || AabbSegmentMax < TriangleSegmentMin)
+		{
+			return false;
+		}
+	}
+
+	return true;
+}
+
+// Evaluates QuadTree for a single point
+bool UvMapping_QuadTreeTriangleOverlap(
+	UvMappingParameters MappingParam,
+	float2 Uv,
+	float Tolerance,
+	int TriangleIndex,
+	out float3 OutBaryCoord)
+{
+	if (MappingParam.NumUVs > 0)
+	{
+        uint IndexBufferOffset = TriangleIndex * 3;
+		uint VertexIndex0 = MappingParam.IndexBuffer[IndexBufferOffset  ];
+		uint VertexIndex1 = MappingParam.IndexBuffer[IndexBufferOffset+1];
+		uint VertexIndex2 = MappingParam.IndexBuffer[IndexBufferOffset+2];
+
+		uint Stride = MappingParam.NumUVs;
+		float2 UV0 = MappingParam.UVBuffer[VertexIndex0 * Stride + MappingParam.UvMappingSet];
+		float2 UV1 = MappingParam.UVBuffer[VertexIndex1 * Stride + MappingParam.UvMappingSet];
+		float2 UV2 = MappingParam.UVBuffer[VertexIndex2 * Stride + MappingParam.UvMappingSet];
+
+		OutBaryCoord = GetBaryCentric2D(Uv, UV0, UV1, UV2);
+
+		if (all(OutBaryCoord > -Tolerance) && all(OutBaryCoord < (1.0f + Tolerance)))
+		{
+			return true;
+		}
+	}
+
+	return false;
+}
+
+bool UvMapping_QuadTreeNodeContains(
+	UvMappingParameters MappingParam,
+	int NodeBufferOffset,
+	float2 Uv,
+	float Tolerance,
+	out int OutTriangleIndex,
+	out float3 OutBaryCoord)
+{
+	if (NodeBufferOffset == -1)
+	{
+		return false;
+	}
+
+	const int TriangleCount = MappingParam.UvMappingBuffer[NodeBufferOffset + UvMapping_ElementCountOffset];
+	for (int TriangleIt = 0; TriangleIt < TriangleCount; ++TriangleIt)
+	{
+		int TriangleIndex = MappingParam.UvMappingBuffer[NodeBufferOffset + UvMapping_TriangleIndexOffset + TriangleIt];
+		if (UvMapping_QuadTreeTriangleOverlap(MappingParam, Uv, Tolerance, TriangleIndex, OutBaryCoord))
+		{
+			OutTriangleIndex = TriangleIndex;
+			return true;
+		}
+	}
+
+	return false;
+}
+
+void UvMapping_GetTriangleCoordAtUV(
+	UvMappingParameters MappingParam,
+	bool Enabled,
+	float2 InUv,
+	float InTolerance,
+	out int OutTriangle,
+	out float3 OutBaryCoord,
+	out bool OutIsValid)
+{
+	OutTriangle = -1;
+	OutBaryCoord = float3(0.0f, 0.0f, 0.0f);
+	OutIsValid = false;
+
+	if (!Enabled || MappingParam.UvMappingBufferLength < UvMapping_MinValidUvMappingBufferSize)
+	{
+		return;
+	}
+
+	int NodeBufferOffset = 0;
+
+	float2 MidPoint = float2(0.5f, 0.5f);
+	float2 Extent = float2(0.5f, 0.5f);
+
+	while (NodeBufferOffset != -1)
+	{
+		if (UvMapping_QuadTreeNodeContains(MappingParam, NodeBufferOffset, InUv, InTolerance, OutTriangle, OutBaryCoord))
+		{
+			OutIsValid = true;
+			return;
+		}
+
+		Extent *= 0.5f;
+
+		const bool LessX = InUv.x < MidPoint.x;
+		const bool LessY = InUv.y < MidPoint.y;
+
+		MidPoint += Extent * float2(LessX ? -1.0f : 1.0f, LessY ? -1.0f : 1.0f);
+		int ChildIndex = LessX ? (LessY ? 2 : 0) : (LessY ? 3 : 1);
+
+		NodeBufferOffset = MappingParam.UvMappingBuffer[NodeBufferOffset + ChildIndex];
+	}
+}
+
+// Evaluates QuadTree for an AABB region
+bool UvMapping_QuadTreeTriangleAabbOverlap(
+	UvMappingParameters MappingParam,
+	float2 UvRef,
+	float2 UvExtentScale,
+	float2 UvExtentBias,
+	int TriangleIndex,
+	out float3 OutBaryCoord)
+{
+	if (MappingParam.NumUVs > 0)
+	{
+        uint IndexBufferOffset = TriangleIndex * 3;
+		uint VertexIndex0 = MappingParam.IndexBuffer[IndexBufferOffset  ];
+		uint VertexIndex1 = MappingParam.IndexBuffer[IndexBufferOffset+1];
+		uint VertexIndex2 = MappingParam.IndexBuffer[IndexBufferOffset+2];
+
+		uint Stride = MappingParam.NumUVs;
+		float2 UV0 = MappingParam.UVBuffer[VertexIndex0 * Stride + MappingParam.UvMappingSet];
+		float2 UV1 = MappingParam.UVBuffer[VertexIndex1 * Stride + MappingParam.UvMappingSet];
+		float2 UV2 = MappingParam.UVBuffer[VertexIndex2 * Stride + MappingParam.UvMappingSet];
+
+		// evaluate whether our triangle intersects with our AABB, if it does, then generate the barycentric coordinate for the center of the AABB
+		const bool Intersection = UvMapping_NormalizedAabbTriangleIntersection(
+				UV0 * UvExtentScale + UvExtentBias,
+				UV1 * UvExtentScale + UvExtentBias,
+				UV2 * UvExtentScale + UvExtentBias);
+
+		if (Intersection)
+		{
+			OutBaryCoord = GetClosestBaryCentric2D(UvRef, UV0, UV1, UV2);
+			return true;
+		}
+	}
+
+	return false;
+}
+
+bool UvMapping_QuadTreeNodeContainsAabb(
+	UvMappingParameters MappingParam,
+	int NodeBufferOffset,
+	float2 UvMin,
+	float2 UvMax,
+	out int OutTriangleIndex,
+	out float3 OutBaryCoord)
+{
+	if (NodeBufferOffset == -1)
+	{
+		return false;
+	}
+
+	const int TriangleCount = MappingParam.UvMappingBuffer[NodeBufferOffset + UvMapping_ElementCountOffset];
+	if (TriangleCount == 0)
+	{
+		return false;
+	}
+
+	const float2 UvRef = 0.5f * (UvMin + UvMax);
+	const float2 UvExtentScale = 1.0f / (UvMax - UvMin);
+	const float2 UvExtentBias = 1.0f - UvMax * UvExtentScale;
+
+	for (int TriangleIt = 0; TriangleIt < TriangleCount; ++TriangleIt)
+	{
+		float3 BaryCoord;
+		int TriangleIndex = MappingParam.UvMappingBuffer[NodeBufferOffset + UvMapping_TriangleIndexOffset + TriangleIt];
+		if (UvMapping_QuadTreeTriangleAabbOverlap(MappingParam, UvRef, UvExtentScale, UvExtentBias, TriangleIndex, BaryCoord))
+		{
+			OutTriangleIndex = TriangleIndex;
+			OutBaryCoord = BaryCoord;
+			return true;
+		}
+	}
+
+	return false;
+}
+
+void UvMapping_GetTriangleCoordInAabb(
+	UvMappingParameters MappingParam,
+	bool Enabled,
+	float2 InUvMin,
+	float2 InUvMax,
+	out int OutTriangle,
+	out float3 OutBaryCoord,
+	out bool OutIsValid)
+{
+	OutTriangle = -1;
+	OutBaryCoord = float3(0.0f, 0.0f, 0.0f);
+	OutIsValid = false;
+
+	if (!Enabled || MappingParam.UvMappingBufferLength < UvMapping_MinValidUvMappingBufferSize)
+	{
+		return;
+	}
+
+	int NodeBufferOffset = 0;
+
+	const float2 UvRegionCenter = 0.5f * (InUvMin + InUvMax);
+	float2 MidPoint = float2(0.5f, 0.5f);
+	float2 Extent = float2(0.5f, 0.5f);
+
+	while (NodeBufferOffset != -1)
+	{
+		if (UvMapping_QuadTreeNodeContainsAabb(MappingParam, NodeBufferOffset, InUvMin, InUvMax, OutTriangle, OutBaryCoord))
+		{
+			OutIsValid = true;
+			return;
+		}
+
+		Extent *= 0.5f;
+
+		const bool LessX = UvRegionCenter.x < MidPoint.x;
+		const bool LessY = UvRegionCenter.y < MidPoint.y;
+
+		MidPoint += Extent * float2(LessX ? -1.0f : 1.0f, LessY ? -1.0f : 1.0f);
+		int ChildIndex = LessX ? (LessY ? 2 : 0) : (LessY ? 3 : 1);
+
+		NodeBufferOffset = MappingParam.UvMappingBuffer[NodeBufferOffset + ChildIndex];
+	}
+}
\ No newline at end of file
diff --git a/Engine/Plugins/FX/Niagara/Shaders/Private/NiagaraDataInterfaceSkeletalMeshTemplate.ush b/Engine/Plugins/FX/Niagara/Shaders/Private/NiagaraDataInterfaceSkeletalMeshTemplate.ush
--- a/Engine/Plugins/FX/Niagara/Shaders/Private/NiagaraDataInterfaceSkeletalMeshTemplate.ush
+++ b/Engine/Plugins/FX/Niagara/Shaders/Private/NiagaraDataInterfaceSkeletalMeshTemplate.ush
@@ -511,188 +511,30 @@ void GetFilteredTriangleAt_{ParameterName}(int FilteredIndex, out int OutTriangl
 
 ////////////////////////////////////////////////////////////////////////////////////////////////////////////////////////////////////////////////////////////////
 // Uv Mapping
-
-// Evaluates QuadTree for a single point
-bool QuadTreeTriangleOverlap_{ParameterName}(float2 Uv, float Tolerance, int TriangleIndex, out float3 OutBaryCoord)
-{
-	if ({ParameterName}_MeshNumTexCoord > 0)
-	{
-        uint IndexBufferOffset = TriangleIndex * 3;
-		uint VertexIndex0 = {ParameterName}_MeshIndexBuffer[IndexBufferOffset  ];
-		uint VertexIndex1 = {ParameterName}_MeshIndexBuffer[IndexBufferOffset+1];
-		uint VertexIndex2 = {ParameterName}_MeshIndexBuffer[IndexBufferOffset+2];
-
-		uint Stride = {ParameterName}_MeshNumTexCoord;
-		float2 UV0 = {ParameterName}_MeshTexCoordBuffer[VertexIndex0 * Stride + {ParameterName}_UvMappingSet];
-		float2 UV1 = {ParameterName}_MeshTexCoordBuffer[VertexIndex1 * Stride + {ParameterName}_UvMappingSet];
-		float2 UV2 = {ParameterName}_MeshTexCoordBuffer[VertexIndex2 * Stride + {ParameterName}_UvMappingSet];
-
-		OutBaryCoord = GetBaryCentric2D(Uv, UV0, UV1, UV2);
-
-		if (all(OutBaryCoord > -Tolerance) && all(OutBaryCoord < (1.0f + Tolerance)))
-		{
-			return true;
-		}
-	}
-
-	return false;
-}
-
-bool QuadTreeNodeContains_{ParameterName}(int NodeBufferOffset, float2 Uv, float Tolerance, out int OutTriangleIndex, out float3 OutBaryCoord)
-{
-	if (NodeBufferOffset == -1)
-	{
-		return false;
-	}
-
-	const int TriangleCount = {ParameterName}_UvMappingBuffer[NodeBufferOffset + DISkelMesh_ElementCountOffset];
-	for (int TriangleIt = 0; TriangleIt < TriangleCount; ++TriangleIt)
-	{
-		int TriangleIndex = {ParameterName}_UvMappingBuffer[NodeBufferOffset + DISkelMesh_TriangleIndexOffset + TriangleIt];
-		if (QuadTreeTriangleOverlap_{ParameterName}(Uv, Tolerance, TriangleIndex, OutBaryCoord))
-		{
-			OutTriangleIndex = TriangleIndex;
-			return true;
-		}
-	}
-
-	return false;
-}
-
-// Evaluates QuadTree for an AABB region
-bool QuadTreeTriangleAabbOverlap_{ParameterName}(float2 UvRef, float2 UvExtentScale, float2 UvExtentBias, int TriangleIndex, out float3 OutBaryCoord)
-{
-	if ({ParameterName}_MeshNumTexCoord > 0)
-	{
-        uint IndexBufferOffset = TriangleIndex * 3;
-		uint VertexIndex0 = {ParameterName}_MeshIndexBuffer[IndexBufferOffset  ];
-		uint VertexIndex1 = {ParameterName}_MeshIndexBuffer[IndexBufferOffset+1];
-		uint VertexIndex2 = {ParameterName}_MeshIndexBuffer[IndexBufferOffset+2];
-
-		uint Stride = {ParameterName}_MeshNumTexCoord;
-		float2 UV0 = {ParameterName}_MeshTexCoordBuffer[VertexIndex0 * Stride + {ParameterName}_UvMappingSet];
-		float2 UV1 = {ParameterName}_MeshTexCoordBuffer[VertexIndex1 * Stride + {ParameterName}_UvMappingSet];
-		float2 UV2 = {ParameterName}_MeshTexCoordBuffer[VertexIndex2 * Stride + {ParameterName}_UvMappingSet];
-
-		// evaluate whether our triangle intersects with our AABB, if it does, then generate the barycentric coordinate for the center of the AABB
-		const bool Intersection = DISkelMesh_NormalizedAabbTriangleIntersection(
-				UV0 * UvExtentScale + UvExtentBias,
-				UV1 * UvExtentScale + UvExtentBias,
-				UV2 * UvExtentScale + UvExtentBias);
-
-		if (Intersection)
-		{
-			OutBaryCoord = GetClosestBaryCentric2D(UvRef, UV0, UV1, UV2);
-			return true;
-		}
-	}
-
-	return false;
-}
-
-bool QuadTreeNodeContains_{ParameterName}Aabb(int NodeBufferOffset, float2 UvMin, float2 UvMax, out int OutTriangleIndex, out float3 OutBaryCoord)
-{
-	if (NodeBufferOffset == -1)
-	{
-		return false;
-	}
-
-	const int TriangleCount = {ParameterName}_UvMappingBuffer[NodeBufferOffset + DISkelMesh_ElementCountOffset];
-	if (TriangleCount == 0)
-	{
-		return false;
-	}
-
-	const float2 UvRef = 0.5f * (UvMin + UvMax);
-	const float2 UvExtentScale = 1.0f / (UvMax - UvMin);
-	const float2 UvExtentBias = 1.0f - UvMax * UvExtentScale;
-
-	for (int TriangleIt = 0; TriangleIt < TriangleCount; ++TriangleIt)
-	{
-		float3 BaryCoord;
-		int TriangleIndex = {ParameterName}_UvMappingBuffer[NodeBufferOffset + DISkelMesh_TriangleIndexOffset + TriangleIt];
-		if (QuadTreeTriangleAabbOverlap_{ParameterName}(UvRef, UvExtentScale, UvExtentBias, TriangleIndex, BaryCoord))
-		{
-			OutTriangleIndex = TriangleIndex;
-			OutBaryCoord = BaryCoord;
-			return true;
-		}
-	}
-
-	return false;
-}
-
 void GetTriangleCoordAtUV_{ParameterName}(bool Enabled, float2 InUv, float InTolerance, out int OutTriangle, out float3 OutBaryCoord, out bool OutIsValid)
 {
-	OutTriangle = -1;
-	OutBaryCoord = float3(0.0f, 0.0f, 0.0f);
-	OutIsValid = false;
-
-	if (!Enabled || {ParameterName}_UvMappingBufferLength < DISkelMesh_MinValidUvMappingBufferSize)
-	{
-		return;
-	}
-
-	int NodeBufferOffset = 0;
+	UvMappingParameters MappingParam;
+	MappingParam.NumUVs = {ParameterName}_MeshNumTexCoord;
+	MappingParam.IndexBuffer = {ParameterName}_MeshIndexBuffer;
+	MappingParam.UVBuffer = {ParameterName}_MeshTexCoordBuffer;
+	MappingParam.UvMappingBuffer = {ParameterName}_UvMappingBuffer;
+	MappingParam.UvMappingBufferLength = {ParameterName}_UvMappingBufferLength;
+	MappingParam.UvMappingSet = {ParameterName}_UvMappingSet;
 
-	float2 MidPoint = float2(0.5f, 0.5f);
-	float2 Extent = float2(0.5f, 0.5f);
-
-	while (NodeBufferOffset != -1)
-	{
-		if (QuadTreeNodeContains_{ParameterName}(NodeBufferOffset, InUv, InTolerance, OutTriangle, OutBaryCoord))
-		{
-			OutIsValid = true;
-			return;
-		}
-
-		Extent *= 0.5f;
-
-		const bool LessX = InUv.x < MidPoint.x;
-		const bool LessY = InUv.y < MidPoint.y;
-
-		MidPoint += Extent * float2(LessX ? -1.0f : 1.0f, LessY ? -1.0f : 1.0f);
-		int ChildIndex = LessX ? (LessY ? 2 : 0) : (LessY ? 3 : 1);
-
-		NodeBufferOffset = {ParameterName}_UvMappingBuffer[NodeBufferOffset + ChildIndex];
-	}
+	UvMapping_GetTriangleCoordAtUV(MappingParam, Enabled, InUv, InTolerance, OutTriangle, OutBaryCoord, OutIsValid);
 }
 
 void GetTriangleCoordInAabb_{ParameterName}(bool Enabled, float2 InUvMin, float2 InUvMax, out int OutTriangle, out float3 OutBaryCoord, out bool OutIsValid)
 {
-	OutTriangle = -1;
-	OutBaryCoord = float3(0.0f, 0.0f, 0.0f);
-	OutIsValid = false;
+	UvMappingParameters MappingParam;
+	MappingParam.NumUVs = {ParameterName}_MeshNumTexCoord;
+	MappingParam.IndexBuffer = {ParameterName}_MeshIndexBuffer;
+	MappingParam.UVBuffer = {ParameterName}_MeshTexCoordBuffer;
+	MappingParam.UvMappingBuffer = {ParameterName}_UvMappingBuffer;
+	MappingParam.UvMappingBufferLength = {ParameterName}_UvMappingBufferLength;
+	MappingParam.UvMappingSet = {ParameterName}_UvMappingSet;
 
-	if (!Enabled || {ParameterName}_UvMappingBufferLength < DISkelMesh_MinValidUvMappingBufferSize)
-	{
-		return;
-	}
-
-	int NodeBufferOffset = 0;
-
-	const float2 UvRegionCenter = 0.5f * (InUvMin + InUvMax);
-	float2 MidPoint = float2(0.5f, 0.5f);
-	float2 Extent = float2(0.5f, 0.5f);
-
-	while (NodeBufferOffset != -1)
-	{
-		if (QuadTreeNodeContains_{ParameterName}Aabb(NodeBufferOffset, InUvMin, InUvMax, OutTriangle, OutBaryCoord))
-		{
-			OutIsValid = true;
-			return;
-		}
-
-		Extent *= 0.5f;
-
-		const bool LessX = UvRegionCenter.x < MidPoint.x;
-		const bool LessY = UvRegionCenter.y < MidPoint.y;
-
-		MidPoint += Extent * float2(LessX ? -1.0f : 1.0f, LessY ? -1.0f : 1.0f);
-		int ChildIndex = LessX ? (LessY ? 2 : 0) : (LessY ? 3 : 1);
-
-		NodeBufferOffset = {ParameterName}_UvMappingBuffer[NodeBufferOffset + ChildIndex];
-	}
+	UvMapping_GetTriangleCoordInAabb(MappingParam, Enabled, InUvMin, InUvMax, OutTriangle, OutBaryCoord, OutIsValid);
 }
 
 ////////////////////////////////////////////////////////////////////////////////////////////////////////////////////////////////////////////////////////////////
diff --git a/Engine/Plugins/FX/Niagara/Shaders/Private/NiagaraDataInterfaceStaticMeshTemplate.ush b/Engine/Plugins/FX/Niagara/Shaders/Private/NiagaraDataInterfaceStaticMeshTemplate.ush
--- a/Engine/Plugins/FX/Niagara/Shaders/Private/NiagaraDataInterfaceStaticMeshTemplate.ush
+++ b/Engine/Plugins/FX/Niagara/Shaders/Private/NiagaraDataInterfaceStaticMeshTemplate.ush
@@ -36,6 +36,10 @@ float3			{ParameterName}_SystemLWCTile;
 int				{ParameterName}_InstanceDistanceFieldIndex;
 #endif
 
+Buffer<int>		{ParameterName}_UvMappingBuffer;
+uint			{ParameterName}_UvMappingBufferLength;
+uint			{ParameterName}_UvMappingSet;
+
 //////////////////////////////////////////////////////////////////////////
 // Fwd declarations
 void RandomSection_{ParameterName}(NiagaraRandInfo RandInfo, out int Section);
@@ -589,3 +593,32 @@ void QueryDistanceField_{ParameterName}(bool bExecute, float3 WorldPosition, boo
 	OutDistance = MaxDistance;
 #endif
 }
+
+
+////////////////////////////////////////////////////////////////////////////////////////////////////////////////////////////////////////////////////////////////
+// Uv Mapping
+void GetTriangleCoordAtUV_{ParameterName}(bool Enabled, float2 InUv, float InTolerance, out int OutTriangle, out float3 OutBaryCoord, out bool OutIsValid)
+{
+	UvMappingParameters MappingParam;
+	MappingParam.NumUVs = {ParameterName}_NumUVs;
+	MappingParam.IndexBuffer = {ParameterName}_IndexBuffer;
+	MappingParam.UVBuffer = {ParameterName}_UVBuffer;
+	MappingParam.UvMappingBuffer = {ParameterName}_UvMappingBuffer;
+	MappingParam.UvMappingBufferLength = {ParameterName}_UvMappingBufferLength;
+	MappingParam.UvMappingSet = {ParameterName}_UvMappingSet;
+
+	UvMapping_GetTriangleCoordAtUV(MappingParam, Enabled, InUv, InTolerance, OutTriangle, OutBaryCoord, OutIsValid);
+}
+
+void GetTriangleCoordInAabb_{ParameterName}(bool Enabled, float2 InUvMin, float2 InUvMax, out int OutTriangle, out float3 OutBaryCoord, out bool OutIsValid)
+{
+	UvMappingParameters MappingParam;
+	MappingParam.NumUVs = {ParameterName}_NumUVs;
+	MappingParam.IndexBuffer = {ParameterName}_IndexBuffer;
+	MappingParam.UVBuffer = {ParameterName}_UVBuffer;
+	MappingParam.UvMappingBuffer = {ParameterName}_UvMappingBuffer;
+	MappingParam.UvMappingBufferLength = {ParameterName}_UvMappingBufferLength;
+	MappingParam.UvMappingSet = {ParameterName}_UvMappingSet;
+
+	UvMapping_GetTriangleCoordInAabb(MappingParam, Enabled, InUvMin, InUvMax, OutTriangle, OutBaryCoord, OutIsValid);
+}
\ No newline at end of file
diff --git a/Engine/Plugins/FX/Niagara/Source/Niagara/Classes/NiagaraDataInterfaceSkeletalMesh.h b/Engine/Plugins/FX/Niagara/Source/Niagara/Classes/NiagaraDataInterfaceSkeletalMesh.h
--- a/Engine/Plugins/FX/Niagara/Source/Niagara/Classes/NiagaraDataInterfaceSkeletalMesh.h
+++ b/Engine/Plugins/FX/Niagara/Source/Niagara/Classes/NiagaraDataInterfaceSkeletalMesh.h
@@ -4,6 +4,7 @@
 #include "Components/SkeletalMeshComponent.h"
 #include "Containers/Array.h"
 #include "Engine/SkeletalMesh.h"
+#include "Experimental/NiagaraMeshUvMappingHandle.h"
 #include "NiagaraDataInterfaceMeshCommon.h"
 #include "NiagaraParameterStore.h"
 #include "Rendering/SkeletalMeshRenderData.h"
@@ -17,7 +18,7 @@ struct FSkeletalMeshSkinningData;
 struct FSkeletalMeshConnectivity;
 class FSkeletalMeshConnectivityProxy;
 struct FSkeletalMeshUvMapping;
-class FSkeletalMeshUvMappingBufferProxy;
+class FMeshUvMappingBufferProxy;
 struct FNDISkeletalMesh_InstanceData;
 class FSkinWeightVertexBuffer;
 struct FSkeletalMeshSamplingRegion;
@@ -241,47 +242,6 @@ private:
 	bool bForceDataRefresh;
 };
 
-struct FSkeletalMeshUvMappingUsage
-{
-	FSkeletalMeshUvMappingUsage() = default;
-	FSkeletalMeshUvMappingUsage(bool InRequiresCpuAccess, bool InrequiresGpuAccess)
-		: RequiresCpuAccess(InRequiresCpuAccess)
-		, RequiresGpuAccess(InrequiresGpuAccess)
-	{}
-
-	bool IsValid() const { return RequiresCpuAccess || RequiresGpuAccess; }
-
-	bool RequiresCpuAccess = false;
-	bool RequiresGpuAccess = false;
-};
-
-struct NIAGARA_API FSkeletalMeshUvMappingHandle
-{
-	FSkeletalMeshUvMappingHandle();
-	FSkeletalMeshUvMappingHandle(FSkeletalMeshUvMappingUsage InUsage, const TSharedPtr<struct FSkeletalMeshUvMapping>& InUvMappingData, bool bNeedsDataImmediately);
-	FSkeletalMeshUvMappingHandle(const FSkeletalMeshUvMappingHandle& Other) = delete;
-	FSkeletalMeshUvMappingHandle(FSkeletalMeshUvMappingHandle&& Other) noexcept;
-	~FSkeletalMeshUvMappingHandle();
-
-	FSkeletalMeshUvMappingHandle& operator=(const FSkeletalMeshUvMappingHandle& Other) = delete;
-	FSkeletalMeshUvMappingHandle& operator=(FSkeletalMeshUvMappingHandle&& Other) noexcept;
-	explicit operator bool() const;
-
-	FSkeletalMeshUvMappingUsage Usage;
-
-	void FindOverlappingTriangles(const FVector2D& InUv, float Tolerance, TArray<int32>& TriangleIndices) const;
-	int32 FindFirstTriangle(const FVector2D& InUv, float Tolerance, FVector3f& BarycentricCoord) const;
-	int32 FindFirstTriangle(const FBox2D& InUvBox, FVector3f& BarycentricCoord) const;
-	const FSkeletalMeshUvMappingBufferProxy* GetQuadTreeProxy() const;
-	int32 GetUvSetIndex() const;
-	int32 GetLodIndex() const;
-
-	void PinAndInvalidateHandle();
-
-private:
-	TSharedPtr<FSkeletalMeshUvMapping> UvMappingData;
-};
-
 struct FSkeletalMeshConnectivityUsage
 {
 	FSkeletalMeshConnectivityUsage() = default;
@@ -333,7 +293,7 @@ class NIAGARA_API FNDI_SkeletalMesh_GeneratedData : public FNDI_GeneratedData
 
 public:
 	FSkeletalMeshSkinningDataHandle GetCachedSkinningData(TWeakObjectPtr<USkeletalMeshComponent>& InComponent, FSkeletalMeshSkinningDataUsage Usage, bool bNeedsDataImmediately);
-	FSkeletalMeshUvMappingHandle GetCachedUvMapping(TWeakObjectPtr<USkeletalMesh>& InMeshObject, int32 InLodIndex, int32 InUvSetIndex, FSkeletalMeshUvMappingUsage Usage, bool bNeedsDataImmediately);
+	FSkeletalMeshUvMappingHandle GetCachedUvMapping(TWeakObjectPtr<USkeletalMesh>& InMeshObject, int32 InLodIndex, int32 InUvSetIndex, FMeshUvMappingUsage Usage, bool bNeedsDataImmediately);
 	FSkeletalMeshConnectivityHandle GetCachedConnectivity(TWeakObjectPtr<USkeletalMesh>& InMeshObject, int32 InLodIndex, FSkeletalMeshConnectivityUsage Usage, bool bNeedsDataImmediately);
 
 	virtual void Tick(ETickingGroup TickGroup, float DeltaSeconds) override;
@@ -1130,7 +1090,7 @@ struct FNiagaraDISkeletalMeshPassedDataToRT
 	FSkeletalMeshGpuDynamicBufferProxy* DynamicBuffer = nullptr;
 	const FSkinWeightDataVertexBuffer* MeshSkinWeightBuffer = nullptr;
 	const FSkinWeightLookupVertexBuffer* MeshSkinWeightLookupBuffer = nullptr;
-	const FSkeletalMeshUvMappingBufferProxy* UvMappingBuffer = nullptr;
+	const FMeshUvMappingBufferProxy* UvMappingBuffer = nullptr;
 	const FSkeletalMeshConnectivityProxy* ConnectivityBuffer = nullptr;
 
 	bool bIsGpuUniformlyDistributedSampling = false;
diff --git a/Engine/Plugins/FX/Niagara/Source/Niagara/Private/DataInterface/NiagaraDataInterfaceStaticMesh.cpp b/Engine/Plugins/FX/Niagara/Source/Niagara/Private/DataInterface/NiagaraDataInterfaceStaticMesh.cpp
--- a/Engine/Plugins/FX/Niagara/Source/Niagara/Private/DataInterface/NiagaraDataInterfaceStaticMesh.cpp
+++ b/Engine/Plugins/FX/Niagara/Source/Niagara/Private/DataInterface/NiagaraDataInterfaceStaticMesh.cpp
@@ -21,6 +21,8 @@
 
 #include "Engine/StaticMeshActor.h"
 #include "Engine/StaticMeshSocket.h"
+#include "NiagaraDataInterfaceStaticMeshUvMapping.h"
+#include "Experimental/NiagaraMeshUvMappingHandle.h"
 #include "Internationalization/Internationalization.h"
 #include "ShaderCompilerCore.h"
 
@@ -72,6 +74,10 @@ namespace NDIStaticMeshLocal
 		SHADER_PARAMETER(int,					InstanceDistanceFieldIndex)
 
 		SHADER_PARAMETER(FVector3f,				SystemLWCTile)
+
+		SHADER_PARAMETER_SRV(Buffer<int>,		UvMappingBuffer)
+		SHADER_PARAMETER(uint32,				UvMappingBufferLength)
+		SHADER_PARAMETER(uint32,				UvMappingSet)
 	END_SHADER_PARAMETER_STRUCT()
 
 	static const TCHAR* TemplateShaderFile = TEXT("/Plugin/FX/Niagara/Private/NiagaraDataInterfaceStaticMeshTemplate.ush");
@@ -179,6 +185,13 @@ namespace NDIStaticMeshLocal
 	static const FName	GetLocalToWorldInverseTransposedName("GetLocalToWorldInverseTransposed");
 	static const FName	GetWorldVelocityName("GetWorldVelocity");
 
+	//////////////////////////////////////////////////////////////////////////
+	// VM UV mapping functions
+	static const FName	GetTriangleCoordAtUVName("GetTriangleCoordAtUV");
+	static const FName	GetTriangleCoordInAabbName("GetTriangleCoordInAabb");
+
+	static const FName	BuildUvMappingName("BuildUvMapping");
+
 	//////////////////////////////////////////////////////////////////////////
 	// Distance Field Functions
 	static const FName	QueryDistanceFieldName("QueryDistanceField");
@@ -381,6 +394,9 @@ namespace NDIStaticMeshLocal
 		int64 GPUMemoryUsage = 0;
 	#endif
 
+		const FMeshUvMappingBufferProxy* UvMappingBuffer = nullptr;
+		uint32 UvMappingSet = 0;
+
 		~FInstanceData_RenderThread()
 		{
 			Release();
@@ -504,6 +520,10 @@ namespace NDIStaticMeshLocal
 		FQuat4f					PrevRotation = FQuat4f::Identity;
 		float					DeltaSeconds = 0.0f;
 		FPrimitiveComponentId	DistanceFieldPrimitiveId;
+
+		const FMeshUvMappingBufferProxy* UvMappingBuffer = nullptr;
+		uint32					UvMappingSet = 0;
+
 #if WITH_EDITOR
 		FColorVertexBuffer*		OverrideVertexColors = nullptr;
 #endif
@@ -517,6 +537,9 @@ namespace NDIStaticMeshLocal
 		/** Cached ptr to the mesh so that we can make sure that we haven't been deleted. */
 		TWeakObjectPtr<UStaticMesh> StaticMeshWeakPtr;
 
+		/** Handle to our uv mapping data. */
+		FStaticMeshUvMappingHandle UvMapping;
+
 		/** Cached ComponentToWorld. (Falls back to WorldTransform of the system instance) */
 		FMatrix Transform;
 		/** InverseTranspose of above for transforming normals/tangents. */
@@ -581,6 +604,11 @@ namespace NDIStaticMeshLocal
 		/** Filter sockets followed by unfiltered sockets */
 		TArray<uint16> FilteredAndUnfilteredSockets;
 
+		/** The UV index to use for UV mapping, -1 indicates the use of the feature is pending the call to BuildUvMapping() */
+		int32 UvMappingIndexSet = -1;
+
+		FMeshUvMappingUsage UvMappingUsage;
+
 #if WITH_EDITOR
 		FDelegateHandle OnMeshChanged;
 		FDelegateHandle OnPostMeshBuild;
@@ -823,11 +851,47 @@ namespace NDIStaticMeshLocal
 						}
 					}
 				}
+
+			}
+
+			// support for UV mapping
+			{
+				bool UsedByCpuUvMapping = false;
+				bool UsedByGpuUvMapping = false;
+
+				SystemInstance->EvaluateBoundFunction(GetTriangleCoordAtUVName, UsedByCpuUvMapping, UsedByGpuUvMapping);
+				SystemInstance->EvaluateBoundFunction(GetTriangleCoordInAabbName, UsedByCpuUvMapping, UsedByGpuUvMapping);
+
+				const bool MeshValid = IsValid(StaticMesh);
+				const bool SupportUvMappingCpu = UsedByCpuUvMapping && MeshValid;
+				const bool SupportUvMappingGpu = UsedByGpuUvMapping && MeshValid && Interface->IsUsedWithGPUEmitter();
+
+				UvMappingUsage = FMeshUvMappingUsage(SupportUvMappingCpu, SupportUvMappingGpu);
+
+				if (UvMappingIndexSet != INDEX_NONE)
+				{
+					UpdateUvMapping(SystemInstance);
+				}
 			}
 
 			return true;
 		}
 
+		void UpdateUvMapping(FNiagaraSystemInstance* SystemInstance)
+		{
+			if (UvMappingUsage.IsValid())
+			{
+				const bool bNeedsDataImmediately = true;
+
+				FNDI_StaticMesh_GeneratedData& GeneratedData = SystemInstance->GetWorldManager()->EditGeneratedData<FNDI_StaticMesh_GeneratedData>();
+				UvMapping = GeneratedData.GetCachedUvMapping(StaticMeshWeakPtr, CachedLODIdx, UvMappingIndexSet, UvMappingUsage, bNeedsDataImmediately);
+			}
+			else
+			{
+				UvMapping = FStaticMeshUvMappingHandle(UvMappingUsage, nullptr, false);
+			}
+		}
+
 		bool Tick(UNiagaraDataInterfaceStaticMesh* Interface, FNiagaraSystemInstance* SystemInstance, float InDeltaSeconds)
 		{
 			if (ResetRequired(Interface))
@@ -861,6 +925,11 @@ namespace NDIStaticMeshLocal
 				}
 			}
 
+			if (UvMappingIndexSet != INDEX_NONE && !UvMapping)
+			{
+				UpdateUvMapping(SystemInstance);
+			}
+
 			return false;
 		}
 
@@ -982,6 +1051,9 @@ namespace NDIStaticMeshLocal
 			InstanceData->DeltaSeconds				= FromGameThread->DeltaSeconds;
 			InstanceData->DistanceFieldPrimitiveId	= FromGameThread->DistanceFieldPrimitiveId;
 
+			InstanceData->UvMappingBuffer			= FromGameThread->UvMappingBuffer;
+			InstanceData->UvMappingSet				= FromGameThread->UvMappingSet;
+
 #if WITH_EDITOR
 			if (FromGameThread->OverrideVertexColors)
 			{
@@ -1566,6 +1638,9 @@ void UNiagaraDataInterfaceStaticMesh::ProvidePerInstanceDataForRenderThread(void
 	DataFromGT->PrevRotation				= InstanceData->PrevRotation;
 	DataFromGT->DeltaSeconds				= InstanceData->DeltaSeconds;
 	DataFromGT->DistanceFieldPrimitiveId	= FPrimitiveComponentId();
+	DataFromGT->UvMappingBuffer				= InstanceData->UvMapping.GetQuadTreeProxy();
+	DataFromGT->UvMappingSet				= InstanceData->UvMapping.GetUvSetIndex();
+
 #if WITH_EDITOR
 	InstanceData->GetCurrentFirstLODWithVertexColorOverrides(DataFromGT->OverrideVertexColors);
 #endif
@@ -1879,6 +1954,48 @@ void UNiagaraDataInterfaceStaticMesh::GetFunctions(TArray<FNiagaraFunctionSignat
 		Sig.Outputs.Emplace(FNiagaraTypeDefinition::GetVec3Def(), TEXT("Velocity"));
 	}
 
+	//////////////////////////////////////////////////////////////////////////
+	// VM UV mapping functions
+	{
+		FNiagaraFunctionSignature& Sig = OutFunctions.Add_GetRef(BaseSignature);
+		Sig.Name = GetTriangleCoordAtUVName;
+
+		Sig.Inputs.Emplace_GetRef(FNiagaraVariable(FNiagaraTypeDefinition::GetBoolDef(), TEXT("Enabled"))).SetValue(true);
+		Sig.Inputs.Emplace(FNiagaraVariable(FNiagaraTypeDefinition::GetVec2Def(), TEXT("UV")));
+		Sig.Inputs.Emplace_GetRef(FNiagaraVariable(FNiagaraTypeDefinition::GetFloatDef(), TEXT("Tolerance"))).SetValue(KINDA_SMALL_NUMBER);
+
+		Sig.Outputs.Add(FNiagaraVariable(FNiagaraTypeDefinition::GetIntDef(), TEXT("Triangle")));
+		Sig.Outputs.Add(FNiagaraVariable(FNiagaraTypeDefinition::GetVec3Def(), TEXT("BaryCoord")));
+		Sig.Outputs.Add(FNiagaraVariable(FNiagaraTypeDefinition::GetBoolDef(), TEXT("IsValid")));
+		Sig.bExperimental = true;
+	}
+
+	{
+		FNiagaraFunctionSignature& Sig = OutFunctions.Add_GetRef(BaseSignature);
+		Sig.Name = GetTriangleCoordInAabbName;
+
+		Sig.Inputs.Emplace_GetRef(FNiagaraVariable(FNiagaraTypeDefinition::GetBoolDef(), TEXT("Enabled"))).SetValue(true);
+		Sig.Inputs.Emplace(FNiagaraVariable(FNiagaraTypeDefinition::GetVec2Def(), TEXT("UvMin")));
+		Sig.Inputs.Emplace(FNiagaraVariable(FNiagaraTypeDefinition::GetVec2Def(), TEXT("UvMax")));
+
+		Sig.Outputs.Add(FNiagaraVariable(FNiagaraTypeDefinition::GetIntDef(), TEXT("Triangle")));
+		Sig.Outputs.Add(FNiagaraVariable(FNiagaraTypeDefinition::GetVec3Def(), TEXT("BaryCoord")));
+		Sig.Outputs.Add(FNiagaraVariable(FNiagaraTypeDefinition::GetBoolDef(), TEXT("IsValid")));
+		Sig.bExperimental = true;
+	}
+
+	{
+		FNiagaraFunctionSignature& Sig = OutFunctions.Add_GetRef(BaseSignature);
+		Sig.Name = BuildUvMappingName;
+		Sig.ModuleUsageBitmask = ENiagaraScriptUsageMask::System;
+		Sig.bSupportsGPU = false;
+		Sig.bRequiresExecPin = true;
+
+		Sig.Inputs.Emplace(FNiagaraVariable(FNiagaraTypeDefinition::GetIntDef(), TEXT("UvSet")));
+
+		Sig.bExperimental = true;
+	}
+
 	//////////////////////////////////////////////////////////////////////////
 	// Distance Field Functions
 	{
@@ -2069,6 +2186,8 @@ void UNiagaraDataInterfaceStaticMesh::GetCommonHLSL(FString& OutHLSL)
 		OutHLSL += TEXT("#include \"/Engine/Private/DistanceFieldLightingShared.ush\"\n");
 		OutHLSL += TEXT("#include \"/Engine/Private/MeshDistanceFieldCommon.ush\"\n");
 	}
+
+	OutHLSL += TEXT("#include \"/Plugin/FX/Niagara/Private/Experimental/NiagaraUvMappingUtils.ush\"\n");
 }
 #endif //WITH_EDITORONLY_DATA
 
@@ -2299,6 +2418,21 @@ void UNiagaraDataInterfaceStaticMesh::GetVMExternalFunction(const FVMExternalFun
 		OutFunc = FVMExternalFunction::CreateLambda([this](FVectorVMExternalFunctionContext& Context) { VMGetWorldVelocity(Context); });
 	}
 
+	//////////////////////////////////////////////////////////////////////////
+	// VM UV mapping functions
+	else if (BindingInfo.Name == GetTriangleCoordAtUVName)
+	{
+		OutFunc = FVMExternalFunction::CreateLambda([this](FVectorVMExternalFunctionContext& Context) { VMGetTriangleCoordAtUV(Context); });
+	}
+	else if (BindingInfo.Name == GetTriangleCoordInAabbName)
+	{
+		OutFunc = FVMExternalFunction::CreateLambda([this](FVectorVMExternalFunctionContext& Context) { VMGetTriangleCoordInAabb(Context); });
+	}
+	else if (BindingInfo.Name == BuildUvMappingName)
+	{
+		OutFunc = FVMExternalFunction::CreateLambda([this](FVectorVMExternalFunctionContext& Context) { VMBuildUvMapping(Context); });
+	}
+
 	//////////////////////////////////////////////////////////////////////////
 	// Deprecated Functions
 	else if (BindingInfo.Name == Deprecated_GetTriColorName)
@@ -2596,6 +2730,19 @@ void UNiagaraDataInterfaceStaticMesh::SetShaderParameters(const FNiagaraDataInte
 	}
 
 	ShaderParameters->SystemLWCTile = Context.GetSystemLWCTile();
+
+	if (InstanceData.UvMappingBuffer)
+	{
+		ShaderParameters->UvMappingBuffer = InstanceData.UvMappingBuffer->GetSrv();
+		ShaderParameters->UvMappingBufferLength = InstanceData.UvMappingBuffer->GetBufferSize();
+		ShaderParameters->UvMappingSet = InstanceData.UvMappingSet;
+	}
+	else
+	{
+		ShaderParameters->UvMappingBuffer = FNiagaraRenderer::GetDummyIntBuffer();
+		ShaderParameters->UvMappingBufferLength = 0;
+		ShaderParameters->UvMappingSet = 0;
+	}
 }
 
 bool UNiagaraDataInterfaceStaticMesh::Equals(const UNiagaraDataInterface* Other) const
@@ -3919,6 +4066,126 @@ void UNiagaraDataInterfaceStaticMesh::VMGetWorldVelocity(FVectorVMExternalFuncti
 	}
 }
 
+//////////////////////////////////////////////////////////////////////////
+// VM UV mapping functions
+void UNiagaraDataInterfaceStaticMesh::VMGetTriangleCoordAtUV(FVectorVMExternalFunctionContext& Context)
+{
+	VectorVM::FUserPtrHandler<NDIStaticMeshLocal::FInstanceData_GameThread> InstanceData(Context);
+
+	FNDIInputParam<bool> InEnabled(Context);
+	FNDIInputParam<FVector2f> InUV(Context);
+	FNDIInputParam<float> InTolerance(Context);
+
+	FNDIOutputParam<int32> OutTriangleIndex(Context);
+	FNDIOutputParam<FVector3f> OutBaryCoord(Context);
+	FNDIOutputParam<FNiagaraBool> OutIsValid(Context);
+
+	const FStaticMeshUvMapping* UvMapping = InstanceData->UvMapping ? InstanceData->UvMapping.GetMappingData() : nullptr;
+
+	if (UvMapping)
+	{
+		for (int32 i = 0; i < Context.GetNumInstances(); ++i)
+		{
+			const bool Enabled = InEnabled.GetAndAdvance();
+			const FVector2D SourceUv = FVector2D(InUV.GetAndAdvance());
+			const float Tolerance = InTolerance.GetAndAdvance();
+
+			FVector3f BaryCoord(ForceInitToZero);
+			int32 TriangleIndex = INDEX_NONE;
+
+			if (Enabled)
+			{
+				FVector BaryCoordD;
+				TriangleIndex = UvMapping->FindFirstTriangle(SourceUv, Tolerance, BaryCoordD);
+				BaryCoord = FVector3f(BaryCoordD);
+			}
+
+			OutTriangleIndex.SetAndAdvance(TriangleIndex);
+			OutBaryCoord.SetAndAdvance(BaryCoord);
+			OutIsValid.SetAndAdvance(TriangleIndex != INDEX_NONE);
+		}
+	}
+	else
+	{
+		for (int32 i = 0; i < Context.GetNumInstances(); ++i)
+		{
+			OutTriangleIndex.SetAndAdvance(INDEX_NONE);
+			OutBaryCoord.SetAndAdvance(FVector3f::ZeroVector);
+			OutIsValid.SetAndAdvance(false);
+		}
+	}
+}
+
+void UNiagaraDataInterfaceStaticMesh::VMGetTriangleCoordInAabb(FVectorVMExternalFunctionContext& Context)
+{
+	VectorVM::FUserPtrHandler<NDIStaticMeshLocal::FInstanceData_GameThread> InstanceData(Context);
+
+	FNDIInputParam<bool> InEnabled(Context);
+	FNDIInputParam<FVector2f> InUvMin(Context);
+	FNDIInputParam<FVector2f> InUvMax(Context);
+
+	FNDIOutputParam<int32> OutTriangleIndex(Context);
+	FNDIOutputParam<FVector3f> OutBaryCoord(Context);
+	FNDIOutputParam<FNiagaraBool> OutIsValid(Context);
+
+	const FStaticMeshUvMapping* UvMapping = InstanceData->UvMapping ? InstanceData->UvMapping.GetMappingData() : nullptr;
+
+	if (UvMapping)
+	{
+		for (int32 i = 0; i < Context.GetNumInstances(); ++i)
+		{
+			const bool Enabled = InEnabled.GetAndAdvance();
+			const FVector2D SourceUvMin = FVector2D(InUvMin.GetAndAdvance());
+			const FVector2D SourceUvMax = FVector2D(InUvMax.GetAndAdvance());
+
+			FVector3f BaryCoord(ForceInitToZero);
+			int32 TriangleIndex = INDEX_NONE;
+
+			if (Enabled)
+			{
+				FVector BaryCoordD;
+				TriangleIndex = UvMapping->FindFirstTriangle(FBox2D(SourceUvMin, SourceUvMax), BaryCoordD);
+				BaryCoord = FVector3f(BaryCoordD);
+			}
+
+			OutTriangleIndex.SetAndAdvance(TriangleIndex);
+			OutBaryCoord.SetAndAdvance(BaryCoord);
+			OutIsValid.SetAndAdvance(TriangleIndex != INDEX_NONE);
+		}
+	}
+	else
+	{
+		for (int32 i = 0; i < Context.GetNumInstances(); ++i)
+		{
+			OutTriangleIndex.SetAndAdvance(INDEX_NONE);
+			OutBaryCoord.SetAndAdvance(FVector3f::ZeroVector);
+			OutIsValid.SetAndAdvance(false);
+		}
+	}
+}
+
+void UNiagaraDataInterfaceStaticMesh::VMBuildUvMapping(FVectorVMExternalFunctionContext& Context)
+{
+	VectorVM::FUserPtrHandler<NDIStaticMeshLocal::FInstanceData_GameThread> InstanceData(Context);
+
+	FNDIInputParam<int32> InUvSet(Context);
+
+	const int32 InstanceCount = Context.GetNumInstances();
+	if (InstanceCount > 0)
+	{
+		const int32 UvMappingIndexSet = InUvSet.GetAndAdvance();
+
+		// for sanity, verify that we don't have multiple instances with conflicted reports for the uv set to use
+		for (int32 i = 1; i < InstanceCount; ++i)
+		{
+			check(InUvSet.GetAndAdvance() == UvMappingIndexSet);
+		}
+
+		InstanceData->UvMappingIndexSet = UvMappingIndexSet;
+	}
+}
+
+
 //////////////////////////////////////////////////////////////////////////
 // Deprecated VM Functions
 template<typename TTransformHandler>
@@ -4086,4 +4353,68 @@ void UNiagaraDataInterfaceStaticMesh::VMGetTriangleNormal_Deprecated(FVectorVMEx
 	}
 }
 
+
+void FNDI_StaticMesh_GeneratedData::Tick(ETickingGroup TickGroup, float DeltaSeconds)
+{
+	check(IsInGameThread());
+
+	{ // handle any changes to the UV mappings
+		FRWScopeLock UvMappingWriteLock(CachedUvMappingGuard, SLT_Write);
+
+		TArray<int32, TInlineAllocator<32>> MappingsToRemove;
+
+		const int32 MappingCount = CachedUvMapping.Num();
+
+		for (int32 MappingIt = 0; MappingIt < MappingCount; ++MappingIt)
+		{
+			const TSharedPtr<FStaticMeshUvMapping>& UvMappingData = CachedUvMapping[MappingIt];;
+
+			if (UvMappingData->CanBeDestroyed())
+			{
+				MappingsToRemove.Add(MappingIt);
+			}
+		}
+
+		while (MappingsToRemove.Num())
+		{
+			CachedUvMapping.RemoveAtSwap(MappingsToRemove.Pop(false));
+		}
+	}
+}
+
+FStaticMeshUvMappingHandle FNDI_StaticMesh_GeneratedData::GetCachedUvMapping(TWeakObjectPtr<UStaticMesh>& MeshObject, int32 InLodIndex, int32 InUvSetIndex, FMeshUvMappingUsage Usage, bool bNeedsDataImmediately)
+{
+	check(MeshObject.Get() != nullptr);
+
+	if (!FStaticMeshUvMapping::IsValidMeshObject(MeshObject, InLodIndex, InUvSetIndex))
+	{
+		return FStaticMeshUvMappingHandle();
+	}
+
+	// Attempt to Find data
+	auto MappingMatchPredicate = [&](const TSharedPtr<FStaticMeshUvMapping>& UvMapping)
+	{
+		return UvMapping->Matches(MeshObject, InLodIndex, InUvSetIndex);
+	};
+
+	{
+		FRWScopeLock ReadLock(CachedUvMappingGuard, SLT_ReadOnly);
+		if (TSharedPtr<FStaticMeshUvMapping>* Existing = CachedUvMapping.FindByPredicate(MappingMatchPredicate))
+		{
+			return FStaticMeshUvMappingHandle(Usage, *Existing, bNeedsDataImmediately);
+		}
+	}
+
+	// We need to add
+	FRWScopeLock WriteLock(CachedUvMappingGuard, SLT_Write);
+	if (TSharedPtr<FStaticMeshUvMapping>* Existing = CachedUvMapping.FindByPredicate(MappingMatchPredicate))
+	{
+		return FStaticMeshUvMappingHandle(Usage, *Existing, bNeedsDataImmediately);
+	}
+	return FStaticMeshUvMappingHandle(
+		Usage,
+		CachedUvMapping.Add_GetRef(MakeShared<FStaticMeshUvMapping>(MeshObject, InLodIndex, InUvSetIndex)),
+		bNeedsDataImmediately);
+}
+
 #undef LOCTEXT_NAMESPACE
diff --git a/Engine/Plugins/FX/Niagara/Source/Niagara/Private/DataInterface/NiagaraDataInterfaceStaticMesh.h b/Engine/Plugins/FX/Niagara/Source/Niagara/Private/DataInterface/NiagaraDataInterfaceStaticMesh.h
--- a/Engine/Plugins/FX/Niagara/Source/Niagara/Private/DataInterface/NiagaraDataInterfaceStaticMesh.h
+++ b/Engine/Plugins/FX/Niagara/Source/Niagara/Private/DataInterface/NiagaraDataInterfaceStaticMesh.h
@@ -4,6 +4,7 @@
 #include "NiagaraDataInterfaceMeshCommon.h"
 #include "StaticMeshResources.h"
 #include "Engine/StaticMesh.h"
+#include "Experimental/NiagaraMeshUvMappingHandle.h"
 #include "NiagaraDataInterfaceStaticMesh.generated.h"
 
 UENUM()
@@ -233,6 +234,11 @@ protected:
 	void VMGetLocalToWorldInverseTransposed(FVectorVMExternalFunctionContext& Context);
 	void VMGetWorldVelocity(FVectorVMExternalFunctionContext& Context);
 
+	// VM UV mapping functions
+	void VMGetTriangleCoordAtUV(FVectorVMExternalFunctionContext& Context);
+	void VMGetTriangleCoordInAabb(FVectorVMExternalFunctionContext& Context);
+	void VMBuildUvMapping(FVectorVMExternalFunctionContext& Context);;
+
 	// Deprecated VM Functions
 	template<typename TTransformHandler>
 	void VMGetVertexPosition_Deprecated(FVectorVMExternalFunctionContext& Context);
@@ -245,3 +251,20 @@ protected:
 	template<typename TTransformHandler>
 	void VMGetTriangleNormal_Deprecated(FVectorVMExternalFunctionContext& Context);
 };
+
+class NIAGARA_API FNDI_StaticMesh_GeneratedData : public FNDI_GeneratedData
+{
+	FRWLock CachedUvMappingGuard;
+	TArray<TSharedPtr<FStaticMeshUvMapping>> CachedUvMapping;
+
+public:
+	FStaticMeshUvMappingHandle GetCachedUvMapping(TWeakObjectPtr<UStaticMesh>& InMeshObject, int32 InLodIndex, int32 InUvSetIndex, FMeshUvMappingUsage Usage, bool bNeedsDataImmediately);
+
+	virtual void Tick(ETickingGroup TickGroup, float DeltaSeconds) override;
+
+	static TypeHash GetTypeHash()
+	{
+		static const TypeHash Hash = ::GetTypeHash(TEXT("FNDI_StaticMesh_GeneratedData"));
+		return Hash;
+	}
+};
diff --git a/Engine/Plugins/FX/Niagara/Source/Niagara/Private/Experimental/NiagaraDataInterfaceSkeletalMeshUvMapping.cpp b/Engine/Plugins/FX/Niagara/Source/Niagara/Private/Experimental/NiagaraDataInterfaceSkeletalMeshUvMapping.cpp
new file mode 100644
--- /dev/null
+++ b/Engine/Plugins/FX/Niagara/Source/Niagara/Private/Experimental/NiagaraDataInterfaceSkeletalMeshUvMapping.cpp
@@ -0,0 +1,185 @@
+// Copyright Epic Games, Inc. All Rights Reserved.
+
+#include "NiagaraDataInterfaceSkeletalMeshUvMapping.h"
+
+#include "Experimental/NiagaraMeshUvMapping.h"
+#include "NDISkeletalMeshCommon.h"
+#include "NiagaraResourceArrayWriter.h"
+#include "NiagaraStats.h"
+
+template<bool UseFullPrecisionUv>
+struct FSkelMeshVertexAccessorHelper
+{
+public:
+	FSkelMeshVertexAccessorHelper(const FSkeletalMeshLODRenderData* InLodRenderData)
+		: LodRenderData(InLodRenderData)
+		, IndexBuffer(*InLodRenderData->MultiSizeIndexContainer.GetIndexBuffer())
+	{}
+
+	FORCEINLINE FVector2D GetVertexUV(int32 VertexIndex, int32 UvChannel) const
+	{
+		return MeshVertexAccessor.GetVertexUV(LodRenderData, IndexBuffer.Get(VertexIndex), UvChannel);
+	}
+
+	FORCEINLINE int32 GetTriangleCount() const
+	{
+		return IndexBuffer.Num() / 3;
+	}
+
+private:
+	const FSkeletalMeshLODRenderData* LodRenderData;
+	const FRawStaticIndexBuffer16or32Interface& IndexBuffer;
+	const FSkelMeshVertexAccessor<UseFullPrecisionUv> MeshVertexAccessor;
+};
+
+FSkeletalMeshUvMapping::FSkeletalMeshUvMapping(TWeakObjectPtr<USkeletalMesh> InMeshObject, int32 InLodIndex, int32 InUvSetIndex)
+	: FMeshUvMapping(InLodIndex, InUvSetIndex)
+	, MeshObject(InMeshObject)
+{
+}
+
+bool FSkeletalMeshUvMapping::IsValidMeshObject(const TWeakObjectPtr<USkeletalMesh>& MeshObject, int32 InLodIndex, int32 InUvSetIndex)
+{
+	USkeletalMesh* Mesh = MeshObject.Get();
+
+	if (!Mesh)
+	{
+		return false;
+	}
+
+	const FSkeletalMeshLODInfo* LodInfo = Mesh->GetLODInfo(InLodIndex);
+	if (!LodInfo)
+	{
+		// invalid Lod index
+		return false;
+	}
+
+	if (!LodInfo->bAllowCPUAccess)
+	{
+		// we need CPU access to buffers in order to generate our UV mapping quad tree
+		return false;
+	}
+
+	if (!GetLodRenderData(*Mesh, InLodIndex, InUvSetIndex))
+	{
+		// no render data available
+		return false;
+	}
+
+	return true;
+}
+
+bool FSkeletalMeshUvMapping::Matches(const TWeakObjectPtr<USkeletalMesh>& InMeshObject, int32 InLodIndex, int32 InUvSetIndex) const
+{
+	return LodIndex == InLodIndex && MeshObject == InMeshObject && UvSetIndex == InUvSetIndex;
+}
+
+const FSkeletalMeshLODRenderData* FSkeletalMeshUvMapping::GetLodRenderData() const
+{
+	if (const USkeletalMesh* Mesh = MeshObject.Get())
+	{
+		return GetLodRenderData(*Mesh, LodIndex, UvSetIndex);
+	}
+	return nullptr;
+}
+
+void FSkeletalMeshUvMapping::FindOverlappingTriangles(const FVector2D& InUv, float Tolerance, TArray<int32>& TriangleIndices) const
+{
+	TriangleIndices.Empty();
+
+	if (const FSkeletalMeshLODRenderData* LodRenderData = GetLodRenderData())
+	{
+		if (LodRenderData->StaticVertexBuffers.StaticMeshVertexBuffer.GetUseFullPrecisionUVs())
+		{
+			FSkelMeshVertexAccessorHelper<true> VertexHelper(LodRenderData);
+			FQuadTreeQueryHelper QueryHelper(TriangleIndexQuadTree, VertexHelper, UvSetIndex);
+			QueryHelper.FindOverlappingTriangle(InUv, Tolerance, TriangleIndices);
+		}
+		else
+		{
+			FSkelMeshVertexAccessorHelper<false> VertexHelper(LodRenderData);
+			FQuadTreeQueryHelper QueryHelper(TriangleIndexQuadTree, VertexHelper, UvSetIndex);
+			QueryHelper.FindOverlappingTriangle(InUv, Tolerance, TriangleIndices);
+		}
+	}
+}
+
+int32 FSkeletalMeshUvMapping::FindFirstTriangle(const FVector2D& InUv, float Tolerance, FVector& BarycentricCoord) const
+{
+	if (const FSkeletalMeshLODRenderData* LodRenderData = GetLodRenderData())
+	{
+		if (LodRenderData->StaticVertexBuffers.StaticMeshVertexBuffer.GetUseFullPrecisionUVs())
+		{
+			FSkelMeshVertexAccessorHelper<true> VertexHelper(LodRenderData);
+			FQuadTreeQueryHelper QueryHelper(TriangleIndexQuadTree, VertexHelper, UvSetIndex);
+			return QueryHelper.FindFirstTriangle(InUv, Tolerance, BarycentricCoord);
+		}
+		else
+		{
+			FSkelMeshVertexAccessorHelper<false> VertexHelper(LodRenderData);
+			FQuadTreeQueryHelper QueryHelper(TriangleIndexQuadTree, VertexHelper, UvSetIndex);
+			return QueryHelper.FindFirstTriangle(InUv, Tolerance, BarycentricCoord);
+		}
+	}
+
+	return INDEX_NONE;
+}
+
+int32 FSkeletalMeshUvMapping::FindFirstTriangle(const FBox2D& InUvBox, FVector& BarycentricCoord) const
+{
+	if (const FSkeletalMeshLODRenderData* LodRenderData = GetLodRenderData())
+	{
+		if (LodRenderData->StaticVertexBuffers.StaticMeshVertexBuffer.GetUseFullPrecisionUVs())
+		{
+			FSkelMeshVertexAccessorHelper<true> VertexHelper(LodRenderData);
+			FQuadTreeQueryHelper QueryHelper(TriangleIndexQuadTree, VertexHelper, UvSetIndex);
+			return QueryHelper.FindFirstTriangle(InUvBox, BarycentricCoord);
+		}
+		else
+		{
+			FSkelMeshVertexAccessorHelper<false> VertexHelper(LodRenderData);
+			FQuadTreeQueryHelper QueryHelper(TriangleIndexQuadTree, VertexHelper, UvSetIndex);
+			return QueryHelper.FindFirstTriangle(InUvBox, BarycentricCoord);
+		}
+	}
+
+	return INDEX_NONE;
+}
+
+const FSkeletalMeshLODRenderData* FSkeletalMeshUvMapping::GetLodRenderData(const USkeletalMesh& Mesh, int32 LodIndex, int32 UvSetIndex)
+{
+	if (const FSkeletalMeshRenderData* RenderData = Mesh.GetResourceForRendering())
+	{
+		if (RenderData->LODRenderData.IsValidIndex(LodIndex))
+		{
+			const FSkeletalMeshLODRenderData& LodRenderData = RenderData->LODRenderData[LodIndex];
+			const int32 NumTexCoords = LodRenderData.GetNumTexCoords();
+
+			if (NumTexCoords > UvSetIndex)
+			{
+				return &LodRenderData;
+			}
+		}
+	}
+
+	return nullptr;
+}
+
+void FSkeletalMeshUvMapping::BuildQuadTree()
+{
+	if (const FSkeletalMeshLODRenderData* LodRenderData = GetLodRenderData())
+	{
+		if (LodRenderData->StaticVertexBuffers.StaticMeshVertexBuffer.GetUseFullPrecisionUVs())
+		{
+			FSkelMeshVertexAccessorHelper<true> VertexHelper(LodRenderData);
+			FQuadTreeBuildHelper BuildHelper(VertexHelper, UvSetIndex, TriangleIndexQuadTree);
+			BuildHelper.Build();
+		}
+		else
+		{
+			FSkelMeshVertexAccessorHelper<false> VertexHelper(LodRenderData);
+			FQuadTreeBuildHelper BuildHelper(VertexHelper, UvSetIndex, TriangleIndexQuadTree);
+			BuildHelper.Build();
+		}
+	}
+}
\ No newline at end of file
diff --git a/Engine/Plugins/FX/Niagara/Source/Niagara/Private/Experimental/NiagaraDataInterfaceSkeletalMeshUvMapping.h b/Engine/Plugins/FX/Niagara/Source/Niagara/Private/Experimental/NiagaraDataInterfaceSkeletalMeshUvMapping.h
new file mode 100644
--- /dev/null
+++ b/Engine/Plugins/FX/Niagara/Source/Niagara/Private/Experimental/NiagaraDataInterfaceSkeletalMeshUvMapping.h
@@ -0,0 +1,30 @@
+// Copyright Epic Games, Inc. All Rights Reserved.
+#pragma once
+
+#include "Experimental/NiagaraMeshUvMapping.h"
+#include "NiagaraDataInterfaceSkeletalMesh.h"
+#include "NiagaraUvQuadTree.h"
+#include "RenderResource.h"
+#include "UObject/WeakObjectPtrTemplates.h"
+
+struct FSkeletalMeshUvMapping : public FMeshUvMapping
+{
+	FSkeletalMeshUvMapping(TWeakObjectPtr<USkeletalMesh> InMeshObject, int32 InLodIndex, int32 InUvSetIndex);
+	virtual ~FSkeletalMeshUvMapping() = default;
+
+	static bool IsValidMeshObject(const TWeakObjectPtr<USkeletalMesh>& MeshObject, int32 InLodIndex, int32 InUvSetIndex);
+
+	bool Matches(const TWeakObjectPtr<USkeletalMesh>& InMeshObject, int32 InLodIndex, int32 InUvSetIndex) const;
+
+	const FSkeletalMeshLODRenderData* GetLodRenderData() const;
+
+	void FindOverlappingTriangles(const FVector2D& InUv, float Tolerance, TArray<int32>& TriangleIndices) const;
+	int32 FindFirstTriangle(const FVector2D& InUv, float Tolerance, FVector& BarycentricCoord) const;
+	int32 FindFirstTriangle(const FBox2D& InUvBox, FVector& BarycentricCoord) const;
+
+private:
+	static const FSkeletalMeshLODRenderData* GetLodRenderData(const USkeletalMesh& Mesh, int32 LodIndex, int32 UvSetIndex);
+	virtual void BuildQuadTree() override;
+
+	TWeakObjectPtr<USkeletalMesh> MeshObject;
+};
\ No newline at end of file
diff --git a/Engine/Plugins/FX/Niagara/Source/Niagara/Private/Experimental/NiagaraMeshUvMapping.cpp b/Engine/Plugins/FX/Niagara/Source/Niagara/Private/Experimental/NiagaraMeshUvMapping.cpp
new file mode 100644
--- /dev/null
+++ b/Engine/Plugins/FX/Niagara/Source/Niagara/Private/Experimental/NiagaraMeshUvMapping.cpp
@@ -0,0 +1,223 @@
+// Copyright Epic Games, Inc. All Rights Reserved.
+
+#include "NiagaraMeshUvMapping.h"
+
+#include "Experimental/NiagaraMeshUvMappingHandle.h"
+#include "NiagaraResourceArrayWriter.h"
+#include "NiagaraStats.h"
+
+DECLARE_CYCLE_STAT(TEXT("Niagara - UvQuadTree Cpu"), STAT_Niagara_UvMapping_Cpu, STATGROUP_Niagara);
+DECLARE_CYCLE_STAT(TEXT("Niagara - UvQuadTree Gpu"), STAT_Niagara_UvMapping_Gpu, STATGROUP_Niagara);
+
+FMeshUvMapping::FMeshUvMapping(int32 InLodIndex, int32 InUvSetIndex)
+: LodIndex(InLodIndex)
+, UvSetIndex(InUvSetIndex)
+, TriangleIndexQuadTree(8 /* Internal node capacity */, 8 /* Maximum tree depth */)
+{}
+
+FMeshUvMapping::~FMeshUvMapping()
+{
+	ReleaseQuadTree();
+	ReleaseGpuQuadTree();
+}
+
+void FMeshUvMapping::FreezeQuadTree(TResourceArray<uint8>& OutQuadTree) const
+{
+	FNiagaraResourceArrayWriter Ar(OutQuadTree);
+	TriangleIndexQuadTree.Freeze(Ar);
+}
+
+void FMeshUvMapping::ReleaseQuadTree()
+{
+	TriangleIndexQuadTree.Empty();
+}
+
+void FMeshUvMapping::BuildGpuQuadTree()
+{
+	check(FrozenQuadTreeProxy == nullptr);
+	FrozenQuadTreeProxy.Reset(new FMeshUvMappingBufferProxy());
+	FrozenQuadTreeProxy->Initialize(*this);
+	BeginInitResource(FrozenQuadTreeProxy.Get());
+}
+
+void FMeshUvMapping::ReleaseGpuQuadTree()
+{
+	if (FMeshUvMappingBufferProxy* ProxyPtr = FrozenQuadTreeProxy.Release())
+	{
+		ENQUEUE_RENDER_COMMAND(BeginDestroyCommand)([RT_Proxy = ProxyPtr](FRHICommandListImmediate& RHICmdList)
+			{
+				RT_Proxy->ReleaseResource();
+				delete RT_Proxy;
+			});
+	}
+}
+
+bool FMeshUvMapping::IsUsed() const
+{
+	return (CpuQuadTreeUserCount > 0)
+		|| (GpuQuadTreeUserCount > 0);
+}
+
+bool FMeshUvMapping::CanBeDestroyed() const
+{
+	return !IsUsed();
+}
+
+void FMeshUvMapping::RegisterUser(FMeshUvMappingUsage Usage, bool bNeedsDataImmediately)
+{
+	if (Usage.RequiresCpuAccess || Usage.RequiresGpuAccess)
+	{
+		if (CpuQuadTreeUserCount++ == 0)
+		{
+			SCOPE_CYCLE_COUNTER(STAT_Niagara_UvMapping_Cpu);
+			BuildQuadTree();
+		}
+	}
+
+	if (Usage.RequiresGpuAccess)
+	{
+		if (GpuQuadTreeUserCount++ == 0)
+		{
+			SCOPE_CYCLE_COUNTER(STAT_Niagara_UvMapping_Gpu);
+			BuildGpuQuadTree();
+		}
+	}
+}
+
+void FMeshUvMapping::UnregisterUser(FMeshUvMappingUsage Usage)
+{
+	if (Usage.RequiresCpuAccess || Usage.RequiresGpuAccess)
+	{
+		if (--CpuQuadTreeUserCount == 0)
+		{
+			ReleaseQuadTree();
+		}
+	}
+
+	if (Usage.RequiresGpuAccess)
+	{
+		if (--GpuQuadTreeUserCount == 0)
+		{
+			ReleaseGpuQuadTree();
+		}
+	}
+}
+
+const FMeshUvMappingBufferProxy* FMeshUvMapping::GetQuadTreeProxy() const
+{
+	return FrozenQuadTreeProxy.Get();
+}
+
+void
+FMeshUvMappingBufferProxy::Initialize(const FMeshUvMapping& UvMapping)
+{
+	UvMapping.FreezeQuadTree(FrozenQuadTree);
+}
+
+void
+FMeshUvMappingBufferProxy::InitRHI()
+{
+	FRHIResourceCreateInfo CreateInfo(TEXT("UvMappingBuffer"));
+	CreateInfo.ResourceArray = &FrozenQuadTree;
+
+	const int32 BufferSize = FrozenQuadTree.Num();
+
+	UvMappingBuffer = RHICreateVertexBuffer(BufferSize, BUF_ShaderResource | BUF_Static, CreateInfo);
+	UvMappingSrv = RHICreateShaderResourceView(UvMappingBuffer, sizeof(int32), PF_R32_SINT);
+
+#if STATS
+	check(GpuMemoryUsage == 0);
+	GpuMemoryUsage = BufferSize;
+#endif
+
+	INC_MEMORY_STAT_BY(STAT_NiagaraGPUDataInterfaceMemory, GpuMemoryUsage);
+}
+
+void
+FMeshUvMappingBufferProxy::ReleaseRHI()
+{
+	UvMappingBuffer.SafeRelease();
+	UvMappingSrv.SafeRelease();
+
+	DEC_MEMORY_STAT_BY(STAT_NiagaraGPUDataInterfaceMemory, GpuMemoryUsage);
+
+#if STATS
+	GpuMemoryUsage = 0;
+#endif
+}
+
+FMeshUvMappingHandleBase::FMeshUvMappingHandleBase(FMeshUvMappingUsage InUsage, const TSharedPtr<FMeshUvMapping>&InUvMappingData, bool bNeedsDataImmediately)
+	: Usage(InUsage)
+	, UvMappingData(InUvMappingData)
+{
+	if (FMeshUvMapping* MappingData = UvMappingData.Get())
+	{
+		MappingData->RegisterUser(Usage, bNeedsDataImmediately);
+	}
+}
+
+FMeshUvMappingHandleBase::FMeshUvMappingHandleBase(FMeshUvMappingHandleBase&& Other) noexcept
+{
+	Usage = Other.Usage;
+	UvMappingData = Other.UvMappingData;
+	Other.UvMappingData = nullptr;
+}
+
+FMeshUvMappingHandleBase::~FMeshUvMappingHandleBase()
+{
+	if (FMeshUvMapping* MappingData = UvMappingData.Get())
+	{
+		MappingData->UnregisterUser(Usage);
+	}
+}
+
+FMeshUvMappingHandleBase& FMeshUvMappingHandleBase::operator=(FMeshUvMappingHandleBase&& Other) noexcept
+{
+	if (this != &Other)
+	{
+		Usage = Other.Usage;
+		UvMappingData = Other.UvMappingData;
+		Other.UvMappingData = nullptr;
+	}
+	return *this;
+}
+
+FMeshUvMappingHandleBase::operator bool() const
+{
+	return UvMappingData.IsValid();
+}
+
+const FMeshUvMappingBufferProxy* FMeshUvMappingHandleBase::GetQuadTreeProxy() const
+{
+	if (UvMappingData)
+	{
+		return UvMappingData->GetQuadTreeProxy();
+	}
+
+	return nullptr;
+}
+
+int32 FMeshUvMappingHandleBase::GetUvSetIndex() const
+{
+	if (UvMappingData)
+	{
+		return UvMappingData->UvSetIndex;
+	}
+
+	return 0;
+}
+
+int32 FMeshUvMappingHandleBase::GetLodIndex() const
+{
+	if (UvMappingData)
+	{
+		return UvMappingData->LodIndex;
+	}
+
+	return 0;
+}
+
+void FMeshUvMappingHandleBase::PinAndInvalidateHandle()
+{
+	UvMappingData.Reset();
+}
\ No newline at end of file
diff --git a/Engine/Plugins/FX/Niagara/Source/Niagara/Private/Experimental/NiagaraMeshUvMapping.h b/Engine/Plugins/FX/Niagara/Source/Niagara/Private/Experimental/NiagaraMeshUvMapping.h
new file mode 100644
--- /dev/null
+++ b/Engine/Plugins/FX/Niagara/Source/Niagara/Private/Experimental/NiagaraMeshUvMapping.h
@@ -0,0 +1,228 @@
+// Copyright Epic Games, Inc. All Rights Reserved.
+
+#pragma once
+
+#include "CoreMinimal.h"
+
+#include "Experimental/NiagaraMeshUvMappingHandle.h"
+#include "NiagaraUvQuadTree.h"
+#include "RenderResource.h"
+
+struct FMeshUvMapping;
+class FNiagaraUvQuadTree;
+
+template<typename MeshAccessorType>
+struct FQuadTreeBuildHelper
+{
+	const MeshAccessorType& MeshAccessor;
+	const int32 UvSetIndex;
+	FNiagaraUvQuadTree& QuadTree;
+
+	FQuadTreeBuildHelper(const MeshAccessorType& InMeshAccessor, int32 InUvSetIndex, FNiagaraUvQuadTree& InQuadTree)
+		: MeshAccessor(InMeshAccessor)
+		, UvSetIndex(InUvSetIndex)
+		, QuadTree(InQuadTree)
+	{
+	}
+
+	void Build()
+	{
+		const int32 TriangleCount = MeshAccessor.GetTriangleCount();
+
+		for (int32 TriangleIt = 0; TriangleIt < TriangleCount; ++TriangleIt)
+		{
+			const FVector2D UVs[] =
+			{
+				MeshAccessor.GetVertexUV(TriangleIt * 3 + 0, UvSetIndex),
+				MeshAccessor.GetVertexUV(TriangleIt * 3 + 1, UvSetIndex),
+				MeshAccessor.GetVertexUV(TriangleIt * 3 + 2, UvSetIndex),
+			};
+
+			// we want to skip degenerate triangles
+			if (FMath::Abs((UVs[1] - UVs[0]) ^ (UVs[2] - UVs[0])) < SMALL_NUMBER)
+			{
+				continue;
+			}
+
+			QuadTree.Insert(TriangleIt, FBox2D(UVs, UE_ARRAY_COUNT(UVs)));
+		}
+	}
+};
+
+template<typename MeshAccessorType>
+struct FQuadTreeQueryHelper
+{
+	const FNiagaraUvQuadTree& QuadTree;
+	const MeshAccessorType& MeshAccessor;
+	const int32 UvSetIndex;
+
+	FQuadTreeQueryHelper(const FNiagaraUvQuadTree& InQuadTree, const MeshAccessorType& InMeshAccessor, int32 InUvSetIndex)
+	: QuadTree(InQuadTree)
+	, MeshAccessor(InMeshAccessor)
+	, UvSetIndex(InUvSetIndex)
+	{
+	}
+
+	FVector BuildTriangleCoordinate(const FVector2D& InUv, int32 TriangleIndex) const
+	{
+		const FVector VertexUvs[] =
+		{
+			FVector(MeshAccessor.GetVertexUV(TriangleIndex * 3 + 0, UvSetIndex), 0.0f),
+			FVector(MeshAccessor.GetVertexUV(TriangleIndex * 3 + 1, UvSetIndex), 0.0f),
+			FVector(MeshAccessor.GetVertexUV(TriangleIndex * 3 + 2, UvSetIndex), 0.0f),
+		};
+
+		return FMath::GetBaryCentric2D(FVector(InUv, 0.0f), VertexUvs[0], VertexUvs[1], VertexUvs[2]);
+	}
+
+	void FindOverlappingTriangle(const FVector2D& InUv, float Tolerance, TArray<int32>& TriangleIndices) const
+	{
+		FBox2D UvBox(InUv, InUv);
+
+		TArray<int32, TInlineAllocator<32>> Elements;
+		QuadTree.GetElements(UvBox, Elements);
+
+		for (int32 TriangleIndex : Elements)
+		{
+			// generate the barycentric coordinates using the UVs of the triangle, the result may not be in the (0,1) range
+			const FVector BarycentricCoord = BuildTriangleCoordinate(InUv, TriangleIndex);
+
+			if ((BarycentricCoord.GetMin() > -Tolerance) && (BarycentricCoord.GetMax() < (1.0f + Tolerance)))
+			{
+				TriangleIndices.Add(TriangleIndex);
+			}
+		}
+	}
+
+	int32 FindFirstTriangle(const FVector2D& InUv, float Tolerance, FVector& BarycentricCoord) const
+	{
+		int32 FoundTriangleIndex = INDEX_NONE;
+
+		QuadTree.VisitElements(FBox2D(InUv, InUv), [&](int32 TriangleIndex)
+		{
+			// generate the barycentric coordinates using the UVs of the triangle, the result may not be in the (0,1) range
+			FVector ElementCoord = BuildTriangleCoordinate(InUv, TriangleIndex);
+
+			if ((ElementCoord.GetMin() > -Tolerance) && (ElementCoord.GetMax() < (1.0f + Tolerance)))
+			{
+				FoundTriangleIndex = TriangleIndex;
+				BarycentricCoord = ElementCoord;
+				return false;
+			}
+
+			return true;
+		});
+
+		return FoundTriangleIndex;
+	}
+
+	static bool NormalizedAabbTriangleOverlap(const FVector2D& A, const FVector2D& B, const FVector2D& C)
+	{
+		const FVector2D TriAabbMin = FVector2D(FMath::Min3(A.X, B.X, C.X), FMath::Min3(A.Y, B.Y, C.Y));
+		const FVector2D TriAabbMax = FVector2D(FMath::Max3(A.X, B.X, C.X), FMath::Max3(A.Y, B.Y, C.Y));
+
+		if (TriAabbMin.GetMax() > 1.0f || TriAabbMax.GetMin() < 0.0f)
+		{
+			return false;
+		}
+
+		const FVector2D TriangleEdges[] = { C - B, A - C, B - A };
+
+		for (int32 i = 0; i < UE_ARRAY_COUNT(TriangleEdges); ++i)
+		{
+			const FVector2D SeparatingAxis(-TriangleEdges[i].Y, TriangleEdges[i].X);
+			float AabbSegmentMin = FMath::Min(0.0f, FMath::Min3(SeparatingAxis.X, SeparatingAxis.Y, SeparatingAxis.X + SeparatingAxis.Y));
+			float AabbSegmentMax = FMath::Max(0.0f, FMath::Max3(SeparatingAxis.X, SeparatingAxis.Y, SeparatingAxis.X + SeparatingAxis.Y));
+			float TriangleSegmentMin = FMath::Min3(FVector2D::DotProduct(A, SeparatingAxis), FVector2D::DotProduct(B, SeparatingAxis), FVector2D::DotProduct(C, SeparatingAxis));
+			float TriangleSegmentMax = FMath::Max3(FVector2D::DotProduct(A, SeparatingAxis), FVector2D::DotProduct(B, SeparatingAxis), FVector2D::DotProduct(C, SeparatingAxis));
+
+			if (AabbSegmentMin > TriangleSegmentMax || AabbSegmentMax < TriangleSegmentMin)
+			{
+				return false;
+			}
+		}
+
+		return true;
+	}
+
+	int32 FindFirstTriangle(const FBox2D& InUvBox, FVector& BarycentricCoord) const
+	{
+		int32 FoundTriangleIndex = INDEX_NONE;
+		const FVector2D NormalizeScale = FVector2D(1.0f, 1.0f) / (InUvBox.Max - InUvBox.Min);
+		const FVector2D NormalizeBias = FVector2D(1.0f, 1.0f) - InUvBox.Max * NormalizeScale;
+		const FVector UvRef = FVector(InUvBox.GetCenter(), 0.0f);
+
+		QuadTree.VisitElements(InUvBox, [&](int32 TriangleIndex)
+		{
+			const FVector2D A = MeshAccessor.GetVertexUV(TriangleIndex * 3 + 0, UvSetIndex);
+			const FVector2D B = MeshAccessor.GetVertexUV(TriangleIndex * 3 + 1, UvSetIndex);
+			const FVector2D C = MeshAccessor.GetVertexUV(TriangleIndex * 3 + 2, UvSetIndex);
+
+			// evaluate if the triangle overlaps with the InUvBox
+			if (!NormalizedAabbTriangleOverlap(NormalizeScale * A + NormalizeBias, NormalizeScale * B + NormalizeBias, NormalizeScale * C + NormalizeBias))
+			{
+				return true;
+			}
+
+			BarycentricCoord = FMath::GetBaryCentric2D(UvRef, FVector(A, 0.0f), FVector(B, 0.0f), FVector(C, 0.0f));
+			FoundTriangleIndex = TriangleIndex;
+			return false;
+		});
+
+		return FoundTriangleIndex;
+	}
+};
+
+class FMeshUvMappingBufferProxy : public FRenderResource
+{
+public:
+	void Initialize(const FMeshUvMapping& UvMappingData);
+
+	virtual void InitRHI() override;
+	virtual void ReleaseRHI() override;
+
+	FShaderResourceViewRHIRef GetSrv() const { return UvMappingSrv; }
+	uint32 GetBufferSize() const { return UvMappingBuffer ? UvMappingBuffer->GetSize() : 0; }
+
+private:
+	TResourceArray<uint8> FrozenQuadTree;
+
+	FBufferRHIRef UvMappingBuffer;
+	FShaderResourceViewRHIRef UvMappingSrv;
+
+#if STATS
+	int64 GpuMemoryUsage = 0;
+#endif
+};
+
+struct FMeshUvMapping
+{
+	FMeshUvMapping(int32 InLodIndex, int32 InUvSetIndex);
+	FMeshUvMapping() = delete;
+	FMeshUvMapping(const FMeshUvMapping&) = delete;
+	virtual ~FMeshUvMapping();
+
+	bool IsUsed() const;
+	bool CanBeDestroyed() const;
+	void RegisterUser(FMeshUvMappingUsage Usage, bool bNeedsDataImmediately);
+	void UnregisterUser(FMeshUvMappingUsage Usage);
+
+	void FreezeQuadTree(TResourceArray<uint8>& OutQuadTree) const;
+	const FMeshUvMappingBufferProxy* GetQuadTreeProxy() const;
+
+	const int32 LodIndex;
+	const int32 UvSetIndex;
+
+protected:
+	virtual void BuildQuadTree() = 0;
+	void ReleaseQuadTree();
+
+	void BuildGpuQuadTree();
+	void ReleaseGpuQuadTree();
+
+	FNiagaraUvQuadTree TriangleIndexQuadTree;
+	TUniquePtr<FMeshUvMappingBufferProxy> FrozenQuadTreeProxy;
+
+	std::atomic<int32> CpuQuadTreeUserCount = {0};
+	std::atomic<int32> GpuQuadTreeUserCount = {0};
+};
\ No newline at end of file
diff --git a/Engine/Plugins/FX/Niagara/Source/Niagara/Private/NDISkeletalMesh_TriangleSampling.cpp b/Engine/Plugins/FX/Niagara/Source/Niagara/Private/NDISkeletalMesh_TriangleSampling.cpp
--- a/Engine/Plugins/FX/Niagara/Source/Niagara/Private/NDISkeletalMesh_TriangleSampling.cpp
+++ b/Engine/Plugins/FX/Niagara/Source/Niagara/Private/NDISkeletalMesh_TriangleSampling.cpp
@@ -7,6 +7,7 @@
 #include "NDISkeletalMeshCommon.h"
 #include "NiagaraComponent.h"
 #include "NiagaraDataInterfaceArrayInt.h"
+#include "Experimental/NiagaraDataInterfaceSkeletalMeshUvMapping.h"
 #include "NiagaraStats.h"
 #include "SkeletalMeshTypes.h"
 #include "Templates/AlignmentTemplates.h"
@@ -1078,7 +1079,9 @@ void UNiagaraDataInterfaceSkeletalMesh::GetTriangleCoordAtUV(FVectorVMExternalFu
 	checkf(InstData.Get(), TEXT("Skeletal Mesh Interface has invalid instance data. %s"), *GetPathName());
 	checkf(InstData->bMeshValid, TEXT("Skeletal Mesh Interface has invalid mesh. %s"), *GetPathName());
 
-	if (InstData->UvMapping)
+	const FSkeletalMeshUvMapping* UvMapping = InstData->UvMapping ? InstData->UvMapping.GetMappingData() : nullptr;
+
+	if (UvMapping)
 	{
 		for (int32 i = 0; i < Context.GetNumInstances(); ++i)
 		{
@@ -1091,7 +1094,9 @@ void UNiagaraDataInterfaceSkeletalMesh::GetTriangleCoordAtUV(FVectorVMExternalFu
 
 			if (Enabled)
 			{
-				TriangleIndex = InstData->UvMapping.FindFirstTriangle(SourceUv, Tolerance, BaryCoord);
+				FVector BaryCoordD;
+				TriangleIndex = UvMapping->FindFirstTriangle(SourceUv, Tolerance, BaryCoordD);
+				BaryCoord = FVector3f(BaryCoordD);
 			}
 
 			OutTriangleIndex.SetAndAdvance(TriangleIndex);
@@ -1148,7 +1153,9 @@ void UNiagaraDataInterfaceSkeletalMesh::GetTriangleCoordInAabb(FVectorVMExternal
 	checkf(InstData.Get(), TEXT("Skeletal Mesh Interface has invalid instance data. %s"), *GetPathName());
 	checkf(InstData->bMeshValid, TEXT("Skeletal Mesh Interface has invalid mesh. %s"), *GetPathName());
 
-	if (InstData->UvMapping)
+	const FSkeletalMeshUvMapping* UvMapping = InstData->UvMapping ? InstData->UvMapping.GetMappingData() : nullptr;
+
+	if (UvMapping)
 	{
 		for (int32 i = 0; i < Context.GetNumInstances(); ++i)
 		{
@@ -1160,7 +1167,9 @@ void UNiagaraDataInterfaceSkeletalMesh::GetTriangleCoordInAabb(FVectorVMExternal
 			int32 TriangleIndex = INDEX_NONE;
 			if (Enabled)
 			{
-				TriangleIndex = InstData->UvMapping.FindFirstTriangle(FBox2D(MinExtent, MaxExtent), BaryCoord);
+				FVector BaryCoordD;
+				TriangleIndex = UvMapping->FindFirstTriangle(FBox2D(MinExtent, MaxExtent), BaryCoordD);
+				BaryCoord = FVector3f(BaryCoordD);
 			}
 
 			OutTriangleIndex.SetAndAdvance(TriangleIndex);
diff --git a/Engine/Plugins/FX/Niagara/Source/Niagara/Private/NiagaraDataInterfaceSkeletalMesh.cpp b/Engine/Plugins/FX/Niagara/Source/Niagara/Private/NiagaraDataInterfaceSkeletalMesh.cpp
--- a/Engine/Plugins/FX/Niagara/Source/Niagara/Private/NiagaraDataInterfaceSkeletalMesh.cpp
+++ b/Engine/Plugins/FX/Niagara/Source/Niagara/Private/NiagaraDataInterfaceSkeletalMesh.cpp
@@ -12,7 +12,7 @@
 #include "NiagaraEmitterInstance.h"
 #include "NiagaraComponent.h"
 #include "NiagaraDataInterfaceSkeletalMeshConnectivity.h"
-#include "NiagaraDataInterfaceSkeletalMeshUvMapping.h"
+#include "Experimental/NiagaraDataInterfaceSkeletalMeshUvMapping.h"
 #include "NiagaraSettings.h"
 #include "NiagaraStats.h"
 #include "NiagaraSystemInstance.h"
@@ -696,7 +696,7 @@ void FNDI_SkeletalMesh_GeneratedData::Tick(ETickingGroup TickGroup, float DeltaS
 	}
 }
 
-FSkeletalMeshUvMappingHandle FNDI_SkeletalMesh_GeneratedData::GetCachedUvMapping(TWeakObjectPtr<USkeletalMesh>& MeshObject, int32 InLodIndex, int32 InUvSetIndex, FSkeletalMeshUvMappingUsage Usage, bool bNeedsDataImmediately)
+FSkeletalMeshUvMappingHandle FNDI_SkeletalMesh_GeneratedData::GetCachedUvMapping(TWeakObjectPtr<USkeletalMesh>& MeshObject, int32 InLodIndex, int32 InUvSetIndex, FMeshUvMappingUsage Usage, bool bNeedsDataImmediately)
 {
 	check(MeshObject.Get() != nullptr);
 
@@ -706,14 +706,14 @@ FSkeletalMeshUvMappingHandle FNDI_SkeletalMesh_GeneratedData::GetCachedUvMapping
 	}
 
 	// Attempt to Find data
+	auto MappingMatchPredicate = [&](const TSharedPtr<FSkeletalMeshUvMapping>& UvMapping)
 	{
-		FRWScopeLock ReadLock(CachedUvMappingGuard, SLT_ReadOnly);
-		TSharedPtr<FSkeletalMeshUvMapping>* Existing = CachedUvMapping.FindByPredicate([&](const TSharedPtr<FSkeletalMeshUvMapping>& UvMapping)
-			{
-				return UvMapping->Matches(MeshObject, InLodIndex, InUvSetIndex);
-			});
+		return UvMapping->Matches(MeshObject, InLodIndex, InUvSetIndex);
+	};
 
-		if (Existing)
+	{
+		FRWScopeLock ReadLock(CachedUvMappingGuard, SLT_ReadOnly);
+		if (TSharedPtr<FSkeletalMeshUvMapping>* Existing = CachedUvMapping.FindByPredicate(MappingMatchPredicate))
 		{
 			return FSkeletalMeshUvMappingHandle(Usage, *Existing, bNeedsDataImmediately);
 		}
@@ -721,6 +721,11 @@ FSkeletalMeshUvMappingHandle FNDI_SkeletalMesh_GeneratedData::GetCachedUvMapping
 
 	// We need to add
 	FRWScopeLock WriteLock(CachedUvMappingGuard, SLT_Write);
+	if (TSharedPtr<FSkeletalMeshUvMapping>* Existing = CachedUvMapping.FindByPredicate(MappingMatchPredicate))
+	{
+		return FSkeletalMeshUvMappingHandle(Usage, *Existing, bNeedsDataImmediately);
+	}
+
 	return FSkeletalMeshUvMappingHandle(
 		Usage,
 		CachedUvMapping.Add_GetRef(MakeShared<FSkeletalMeshUvMapping>(MeshObject, InLodIndex, InUvSetIndex)),
@@ -738,14 +743,15 @@ FSkeletalMeshConnectivityHandle FNDI_SkeletalMesh_GeneratedData::GetCachedConnec
 	}
 
 	// Attempt to Find data
+	auto ConnectivityMatchPredicate = [&](const TSharedPtr<FSkeletalMeshConnectivity>& Connectivity)
 	{
-		FRWScopeLock ReadLock(CachedConnectivityGuard, SLT_ReadOnly);
-		TSharedPtr<FSkeletalMeshConnectivity>* Existing = CachedConnectivity.FindByPredicate([&](const TSharedPtr<FSkeletalMeshConnectivity>& Connectivity)
-		{
-			return Connectivity->CanBeUsed(MeshObject, InLodIndex);
-		});
+		return Connectivity->CanBeUsed(MeshObject, InLodIndex);
+	};
 
-		if (Existing)
+
+	{
+		FRWScopeLock ReadLock(CachedConnectivityGuard, SLT_ReadOnly);
+		if (TSharedPtr<FSkeletalMeshConnectivity>* Existing = CachedConnectivity.FindByPredicate(ConnectivityMatchPredicate))
 		{
 			return FSkeletalMeshConnectivityHandle(Usage, *Existing, bNeedsDataImmediately);
 		}
@@ -753,6 +759,11 @@ FSkeletalMeshConnectivityHandle FNDI_SkeletalMesh_GeneratedData::GetCachedConnec
 
 	// We need to add
 	FRWScopeLock WriteLock(CachedConnectivityGuard, SLT_Write);
+	if (TSharedPtr<FSkeletalMeshConnectivity>* Existing = CachedConnectivity.FindByPredicate(ConnectivityMatchPredicate))
+	{
+		return FSkeletalMeshConnectivityHandle(Usage, *Existing, bNeedsDataImmediately);
+	}
+
 	return FSkeletalMeshConnectivityHandle(
 		Usage,
 		CachedConnectivity.Add_GetRef(MakeShared<FSkeletalMeshConnectivity>(MeshObject, InLodIndex)),
@@ -1339,104 +1350,6 @@ void FSkeletalMeshGpuDynamicBufferProxy::NewFrame(const FNDISkeletalMesh_Instanc
 	);
 }
 
-//////////////////////////////////////////////////////////////////////////
-//FNiagaraDataInterfaceParametersCS_SkeletalMesh
-struct FNDISkeletalMeshParametersName
-{
-	FString MeshIndexBufferName;
-	FString MeshVertexBufferName;
-	FString MeshSkinWeightBufferName;
-	FString MeshSkinWeightLookupBufferName;
-	FString MeshCurrBonesBufferName;
-	FString MeshPrevBonesBufferName;
-	FString MeshCurrSamplingBonesBufferName;
-	FString MeshPrevSamplingBonesBufferName;
-	FString MeshTangentBufferName;
-	FString MeshTexCoordBufferName;
-	FString MeshColorBufferName;
-	FString MeshTriangleSamplerProbAliasBufferName;
-	FString MeshNumSamplingRegionTrianglesName;
-	FString MeshNumSamplingRegionVerticesName;
-	FString MeshSamplingRegionsProbAliasBufferName;
-	FString MeshSampleRegionsTriangleIndicesName;
-	FString MeshSampleRegionsVerticesName;
-	FString MeshTriangleMatricesOffsetBufferName;
-	FString MeshTriangleCountName;
-	FString MeshVertexCountName;
-	FString MeshWeightStrideName;
-	FString MeshSkinWeightIndexSizeName;
-	FString MeshNumTexCoordName;
-	FString MeshNumWeightsName;
-	FString NumBonesName;
-	FString NumFilteredBonesName;
-	FString NumUnfilteredBonesName;
-	FString RandomMaxBoneName;
-	FString ExcludeBoneIndexName;
-	FString FilteredAndUnfilteredBonesName;
-	FString NumFilteredSocketsName;
-	FString FilteredSocketBoneOffsetName;
-	FString UvMappingBufferName;
-	FString UvMappingBufferLengthName;
-	FString UvMappingSetName;
-	FString ConnectivityBufferName;
-	FString ConnectivityBufferLengthName;
-	FString ConnectivityMaxAdjacentPerVertexName;
-	FString InstanceTransformName;
-	FString InstancePrevTransformName;
-	FString InstanceRotationName;
-	FString InstancePrevRotationName;
-	FString InstanceInvDeltaTimeName;
-	FString EnabledFeaturesName;
-};
-
-static void GetNiagaraDataInterfaceParametersName(FNDISkeletalMeshParametersName& Names, const FString& Suffix)
-{
-	Names.MeshIndexBufferName = UNiagaraDataInterfaceSkeletalMesh::MeshIndexBufferName + Suffix;
-	Names.MeshVertexBufferName = UNiagaraDataInterfaceSkeletalMesh::MeshVertexBufferName + Suffix;
-	Names.MeshSkinWeightBufferName = UNiagaraDataInterfaceSkeletalMesh::MeshSkinWeightBufferName + Suffix;
-	Names.MeshSkinWeightLookupBufferName = UNiagaraDataInterfaceSkeletalMesh::MeshSkinWeightLookupBufferName + Suffix;
-	Names.MeshCurrBonesBufferName = UNiagaraDataInterfaceSkeletalMesh::MeshCurrBonesBufferName + Suffix;
-	Names.MeshPrevBonesBufferName = UNiagaraDataInterfaceSkeletalMesh::MeshPrevBonesBufferName + Suffix;
-	Names.MeshCurrSamplingBonesBufferName = UNiagaraDataInterfaceSkeletalMesh::MeshCurrSamplingBonesBufferName + Suffix;
-	Names.MeshPrevSamplingBonesBufferName = UNiagaraDataInterfaceSkeletalMesh::MeshPrevSamplingBonesBufferName + Suffix;
-	Names.MeshTangentBufferName = UNiagaraDataInterfaceSkeletalMesh::MeshTangentBufferName + Suffix;
-	Names.MeshTexCoordBufferName = UNiagaraDataInterfaceSkeletalMesh::MeshTexCoordBufferName + Suffix;
-	Names.MeshColorBufferName = UNiagaraDataInterfaceSkeletalMesh::MeshColorBufferName + Suffix;
-	Names.MeshTriangleSamplerProbAliasBufferName = UNiagaraDataInterfaceSkeletalMesh::MeshTriangleSamplerProbAliasBufferName + Suffix;
-	Names.MeshNumSamplingRegionTrianglesName = UNiagaraDataInterfaceSkeletalMesh::MeshNumSamplingRegionTrianglesName + Suffix;
-	Names.MeshNumSamplingRegionVerticesName = UNiagaraDataInterfaceSkeletalMesh::MeshNumSamplingRegionVerticesName + Suffix;
-	Names.MeshSamplingRegionsProbAliasBufferName = UNiagaraDataInterfaceSkeletalMesh::MeshSamplingRegionsProbAliasBufferName + Suffix;
-	Names.MeshSampleRegionsTriangleIndicesName = UNiagaraDataInterfaceSkeletalMesh::MeshSampleRegionsTriangleIndicesName + Suffix;
-	Names.MeshSampleRegionsVerticesName = UNiagaraDataInterfaceSkeletalMesh::MeshSampleRegionsVerticesName + Suffix;
-	Names.MeshTriangleMatricesOffsetBufferName = UNiagaraDataInterfaceSkeletalMesh::MeshTriangleMatricesOffsetBufferName + Suffix;
-	Names.MeshTriangleCountName = UNiagaraDataInterfaceSkeletalMesh::MeshTriangleCountName + Suffix;
-	Names.MeshVertexCountName = UNiagaraDataInterfaceSkeletalMesh::MeshVertexCountName + Suffix;
-	Names.MeshWeightStrideName = UNiagaraDataInterfaceSkeletalMesh::MeshWeightStrideName + Suffix;
-	Names.MeshSkinWeightIndexSizeName = UNiagaraDataInterfaceSkeletalMesh::MeshSkinWeightIndexSizeName + Suffix;
-	Names.MeshNumTexCoordName = UNiagaraDataInterfaceSkeletalMesh::MeshNumTexCoordName + Suffix;
-	Names.MeshNumWeightsName = UNiagaraDataInterfaceSkeletalMesh::MeshNumWeightsName + Suffix;
-	Names.NumBonesName = UNiagaraDataInterfaceSkeletalMesh::NumBonesName + Suffix;
-	Names.NumFilteredBonesName = UNiagaraDataInterfaceSkeletalMesh::NumFilteredBonesName + Suffix;
-	Names.NumUnfilteredBonesName = UNiagaraDataInterfaceSkeletalMesh::NumUnfilteredBonesName + Suffix;
-	Names.RandomMaxBoneName = UNiagaraDataInterfaceSkeletalMesh::RandomMaxBoneName + Suffix;
-	Names.ExcludeBoneIndexName = UNiagaraDataInterfaceSkeletalMesh::ExcludeBoneIndexName + Suffix;
-	Names.FilteredAndUnfilteredBonesName = UNiagaraDataInterfaceSkeletalMesh::FilteredAndUnfilteredBonesName + Suffix;
-	Names.NumFilteredSocketsName = UNiagaraDataInterfaceSkeletalMesh::NumFilteredSocketsName + Suffix;
-	Names.FilteredSocketBoneOffsetName = UNiagaraDataInterfaceSkeletalMesh::FilteredSocketBoneOffsetName + Suffix;
-	Names.UvMappingBufferName = UNiagaraDataInterfaceSkeletalMesh::UvMappingBufferName + Suffix;
-	Names.UvMappingBufferLengthName = UNiagaraDataInterfaceSkeletalMesh::UvMappingBufferLengthName + Suffix;
-	Names.UvMappingSetName = UNiagaraDataInterfaceSkeletalMesh::UvMappingSetName + Suffix;
-	Names.ConnectivityBufferName = UNiagaraDataInterfaceSkeletalMesh::ConnectivityBufferName + Suffix;
-	Names.ConnectivityBufferLengthName = UNiagaraDataInterfaceSkeletalMesh::ConnectivityBufferLengthName + Suffix;
-	Names.ConnectivityMaxAdjacentPerVertexName = UNiagaraDataInterfaceSkeletalMesh::ConnectivityMaxAdjacentPerVertexName + Suffix;
-	Names.InstanceTransformName = UNiagaraDataInterfaceSkeletalMesh::InstanceTransformName + Suffix;
-	Names.InstancePrevTransformName = UNiagaraDataInterfaceSkeletalMesh::InstancePrevTransformName + Suffix;
-	Names.InstanceRotationName = UNiagaraDataInterfaceSkeletalMesh::InstanceRotationName + Suffix;
-	Names.InstancePrevRotationName = UNiagaraDataInterfaceSkeletalMesh::InstancePrevRotationName + Suffix;
-	Names.InstanceInvDeltaTimeName = UNiagaraDataInterfaceSkeletalMesh::InstanceInvDeltaTimeName + Suffix;
-	Names.EnabledFeaturesName = UNiagaraDataInterfaceSkeletalMesh::EnabledFeaturesName + Suffix;
-}
-
 //////////////////////////////////////////////////////////////////////////
 
 void FNiagaraDataInterfaceProxySkeletalMesh::ConsumePerInstanceDataFromGameThread(void* PerInstanceData, const FNiagaraSystemInstanceID& Instance)
@@ -1814,7 +1727,7 @@ bool FNDISkeletalMesh_InstanceData::Init(UNiagaraDataInterfaceSkeletalMesh* Inte
 		const bool SupportUvMappingCpu = UsedByCpuUvMapping && MeshValid;
 		const bool SupportUvMappingGpu = UsedByGpuUvMapping && MeshValid && Interface->IsUsedWithGPUEmitter();
 
-		FSkeletalMeshUvMappingUsage UvMappingUsage(SupportUvMappingCpu, SupportUvMappingGpu);
+		FMeshUvMappingUsage UvMappingUsage(SupportUvMappingCpu, SupportUvMappingGpu);
 
 		if (UvMappingUsage.IsValid())
 		{
@@ -2890,6 +2803,7 @@ void UNiagaraDataInterfaceSkeletalMesh::ModifyCompilationEnvironment(EShaderPlat
 void UNiagaraDataInterfaceSkeletalMesh::GetCommonHLSL(FString& OutHLSL)
 {
 	OutHLSL.Appendf(TEXT("#include \"%s\"\n"), NDISkelMeshLocal::CommonShaderFile);
+	OutHLSL.Append(TEXT("#include \"/Plugin/FX/Niagara/Private/Experimental/NiagaraUvMappingUtils.ush\"\n"));
 }
 
 void UNiagaraDataInterfaceSkeletalMesh::GetParameterDefinitionHLSL(const FNiagaraDataInterfaceGPUParamInfo& ParamInfo, FString& OutHLSL)
diff --git a/Engine/Plugins/FX/Niagara/Source/Niagara/Private/NiagaraDataInterfaceSkeletalMeshUvMapping.cpp b/Engine/Plugins/FX/Niagara/Source/Niagara/Private/NiagaraDataInterfaceSkeletalMeshUvMapping.cpp
deleted file mode 100644
diff --git a/Engine/Plugins/FX/Niagara/Source/Niagara/Private/NiagaraDataInterfaceSkeletalMeshUvMapping.h b/Engine/Plugins/FX/Niagara/Source/Niagara/Private/NiagaraDataInterfaceSkeletalMeshUvMapping.h
deleted file mode 100644
diff --git a/Engine/Plugins/FX/Niagara/Source/Niagara/Private/NiagaraDataInterfaceStaticMeshUvMapping.cpp b/Engine/Plugins/FX/Niagara/Source/Niagara/Private/NiagaraDataInterfaceStaticMeshUvMapping.cpp
new file mode 100644
--- /dev/null
+++ b/Engine/Plugins/FX/Niagara/Source/Niagara/Private/NiagaraDataInterfaceStaticMeshUvMapping.cpp
@@ -0,0 +1,136 @@
+// Copyright Epic Games, Inc. All Rights Reserved.
+
+#include "NiagaraDataInterfaceStaticMeshUvMapping.h"
+
+#include "Experimental/NiagaraMeshUvMapping.h"
+#include "NiagaraResourceArrayWriter.h"
+#include "NiagaraStats.h"
+#include "RawIndexBuffer.h"
+#include "Engine/StaticMesh.h"
+#include "StaticMeshResources.h"
+
+struct FStaticMeshVertexAccessorHelper
+{
+public:
+	FStaticMeshVertexAccessorHelper(const FStaticMeshLODResources* InLodRenderData)
+		: LodRenderData(InLodRenderData)
+		, IndexBuffer(LodRenderData->IndexBuffer.GetArrayView())
+	{}
+
+	FORCEINLINE FVector2D GetVertexUV(int32 VertexIndex, int32 UvChannel) const
+	{
+		return FVector2D(LodRenderData->VertexBuffers.StaticMeshVertexBuffer.GetVertexUV(IndexBuffer[VertexIndex], UvChannel));
+	}
+
+	FORCEINLINE int32 GetTriangleCount() const
+	{
+		return LodRenderData->IndexBuffer.GetNumIndices() / 3;
+	}
+
+private:
+	const FStaticMeshLODResources* LodRenderData;
+	const FIndexArrayView IndexBuffer;
+};
+
+FStaticMeshUvMapping::FStaticMeshUvMapping(TWeakObjectPtr<UStaticMesh> InMeshObject, int32 InLodIndex, int32 InUvSetIndex)
+	: FMeshUvMapping(InLodIndex, InUvSetIndex)
+	, MeshObject(InMeshObject)
+{
+}
+
+bool FStaticMeshUvMapping::IsValidMeshObject(const TWeakObjectPtr<UStaticMesh>& MeshObject, int32 InLodIndex, int32 InUvSetIndex)
+{
+	UStaticMesh* Mesh = MeshObject.Get();
+
+	if (!Mesh || !Mesh->bAllowCPUAccess)
+	{
+		return false;
+	}
+
+	if (!GetLodRenderData(*Mesh, InLodIndex, InUvSetIndex))
+	{
+		// no render data available
+		return false;
+	}
+
+	return true;
+}
+
+bool FStaticMeshUvMapping::Matches(const TWeakObjectPtr<UStaticMesh>& InMeshObject, int32 InLodIndex, int32 InUvSetIndex) const
+{
+	return LodIndex == InLodIndex && MeshObject == InMeshObject && UvSetIndex == InUvSetIndex;
+}
+
+const FStaticMeshLODResources* FStaticMeshUvMapping::GetLodRenderData() const
+{
+	if (const UStaticMesh* Mesh = MeshObject.Get())
+	{
+		return GetLodRenderData(*Mesh, LodIndex, UvSetIndex);
+	}
+	return nullptr;
+}
+
+void FStaticMeshUvMapping::FindOverlappingTriangles(const FVector2D& InUv, float Tolerance, TArray<int32>& TriangleIndices) const
+{
+	TriangleIndices.Empty();
+
+	if (const FStaticMeshLODResources* LodRenderData = GetLodRenderData())
+	{
+		FStaticMeshVertexAccessorHelper VertexHelper(LodRenderData);
+		FQuadTreeQueryHelper QueryHelper(TriangleIndexQuadTree, VertexHelper, UvSetIndex);
+		return QueryHelper.FindOverlappingTriangle(InUv, Tolerance, TriangleIndices);
+	}
+}
+
+int32 FStaticMeshUvMapping::FindFirstTriangle(const FVector2D& InUv, float Tolerance, FVector& BarycentricCoord) const
+{
+	if (const FStaticMeshLODResources* LodRenderData = GetLodRenderData())
+	{
+		FStaticMeshVertexAccessorHelper VertexHelper(LodRenderData);
+		FQuadTreeQueryHelper QueryHelper(TriangleIndexQuadTree, VertexHelper, UvSetIndex);
+		return QueryHelper.FindFirstTriangle(InUv, Tolerance, BarycentricCoord);
+	}
+
+	return INDEX_NONE;
+}
+
+int32 FStaticMeshUvMapping::FindFirstTriangle(const FBox2D& InUvBox, FVector& BarycentricCoord) const
+{
+	if (const FStaticMeshLODResources* LodRenderData = GetLodRenderData())
+	{
+		FStaticMeshVertexAccessorHelper VertexHelper(LodRenderData);
+		FQuadTreeQueryHelper QueryHelper(TriangleIndexQuadTree, VertexHelper, UvSetIndex);
+		return QueryHelper.FindFirstTriangle(InUvBox, BarycentricCoord);
+	}
+
+	return INDEX_NONE;
+}
+
+const FStaticMeshLODResources* FStaticMeshUvMapping::GetLodRenderData(const UStaticMesh& Mesh, int32 LodIndex, int32 UvSetIndex)
+{
+	if (const FStaticMeshRenderData* RenderData = Mesh.GetRenderData())
+	{
+		if (RenderData->LODResources.IsValidIndex(LodIndex))
+		{
+			const FStaticMeshLODResources& LodRenderData = RenderData->LODResources[LodIndex];
+			const int32 NumTexCoords = LodRenderData.GetNumTexCoords();
+
+			if (NumTexCoords > UvSetIndex)
+			{
+				return &LodRenderData;
+			}
+		}
+	}
+
+	return nullptr;
+}
+
+void FStaticMeshUvMapping::BuildQuadTree()
+{
+	if (const FStaticMeshLODResources* LodRenderData = GetLodRenderData())
+	{
+		FStaticMeshVertexAccessorHelper VertexHelper(LodRenderData);
+		FQuadTreeBuildHelper BuildHelper(VertexHelper, UvSetIndex, TriangleIndexQuadTree);
+		BuildHelper.Build();
+	}
+}
\ No newline at end of file
diff --git a/Engine/Plugins/FX/Niagara/Source/Niagara/Private/NiagaraDataInterfaceStaticMeshUvMapping.h b/Engine/Plugins/FX/Niagara/Source/Niagara/Private/NiagaraDataInterfaceStaticMeshUvMapping.h
new file mode 100644
--- /dev/null
+++ b/Engine/Plugins/FX/Niagara/Source/Niagara/Private/NiagaraDataInterfaceStaticMeshUvMapping.h
@@ -0,0 +1,33 @@
+// Copyright Epic Games, Inc. All Rights Reserved.
+#pragma once
+
+#include "Engine/StaticMesh.h"
+#include "Experimental/NiagaraMeshUvMapping.h"
+#include "NiagaraUvQuadTree.h"
+#include "RenderResource.h"
+#include "UObject/WeakObjectPtrTemplates.h"
+
+class UStaticMesh;
+struct FStaticMeshLODResources;
+
+struct FStaticMeshUvMapping : public FMeshUvMapping
+{
+	FStaticMeshUvMapping(TWeakObjectPtr<UStaticMesh> InMeshObject, int32 InLodIndex, int32 InUvSetIndex);
+	virtual ~FStaticMeshUvMapping() = default;
+
+	static bool IsValidMeshObject(const TWeakObjectPtr<UStaticMesh>& MeshObject, int32 InLodIndex, int32 InUvSetIndex);
+
+	bool Matches(const TWeakObjectPtr<UStaticMesh>& InMeshObject, int32 InLodIndex, int32 InUvSetIndex) const;
+
+	const FStaticMeshLODResources* GetLodRenderData() const;
+
+	void FindOverlappingTriangles(const FVector2D& InUv, float Tolerance, TArray<int32>& TriangleIndices) const;
+	int32 FindFirstTriangle(const FVector2D& InUv, float Tolerance, FVector& BarycentricCoord) const;
+	int32 FindFirstTriangle(const FBox2D& InUvBox, FVector& BarycentricCoord) const;
+
+private:
+	static const FStaticMeshLODResources* GetLodRenderData(const UStaticMesh& Mesh, int32 LodIndex, int32 UvSetIndex);
+	virtual void BuildQuadTree() override;
+
+	TWeakObjectPtr<UStaticMesh> MeshObject;
+};
\ No newline at end of file
diff --git a/Engine/Plugins/FX/Niagara/Source/Niagara/Public/Experimental/NiagaraMeshUvMappingHandle.h b/Engine/Plugins/FX/Niagara/Source/Niagara/Public/Experimental/NiagaraMeshUvMappingHandle.h
new file mode 100644
--- /dev/null
+++ b/Engine/Plugins/FX/Niagara/Source/Niagara/Public/Experimental/NiagaraMeshUvMappingHandle.h
@@ -0,0 +1,77 @@
+// Copyright Epic Games, Inc. All Rights Reserved.
+#pragma once
+
+#include "CoreMinimal.h"
+
+class FMeshUvMappingBufferProxy;
+struct FMeshUvMapping;
+struct FSkeletalMeshUvMapping;
+struct FStaticMeshUvMapping;
+
+struct FMeshUvMappingUsage
+{
+	FMeshUvMappingUsage() = default;
+	FMeshUvMappingUsage(bool InRequiresCpuAccess, bool InrequiresGpuAccess)
+		: RequiresCpuAccess(InRequiresCpuAccess)
+		, RequiresGpuAccess(InrequiresGpuAccess)
+	{}
+
+	bool IsValid() const { return RequiresCpuAccess || RequiresGpuAccess; }
+
+	bool RequiresCpuAccess = false;
+	bool RequiresGpuAccess = false;
+};
+
+struct NIAGARA_API FMeshUvMappingHandleBase
+{
+	FMeshUvMappingHandleBase() = default;
+	FMeshUvMappingHandleBase(const FMeshUvMappingHandleBase& Other) = delete;
+	~FMeshUvMappingHandleBase();
+	FMeshUvMappingHandleBase& operator=(const FMeshUvMappingHandleBase& Other) = delete;
+	explicit operator bool() const;
+
+	const FMeshUvMappingBufferProxy* GetQuadTreeProxy() const;
+	int32 GetUvSetIndex() const;
+	int32 GetLodIndex() const;
+	void PinAndInvalidateHandle();
+
+	FMeshUvMappingUsage Usage;
+
+protected:
+	FMeshUvMappingHandleBase(FMeshUvMappingUsage InUsage, const TSharedPtr<FMeshUvMapping>& InUvMappingData, bool bNeedsDataImmediately);
+	FMeshUvMappingHandleBase(FMeshUvMappingHandleBase&& Other) noexcept;
+	FMeshUvMappingHandleBase& operator=(FMeshUvMappingHandleBase&& Other) noexcept;
+
+	TSharedPtr<FMeshUvMapping> UvMappingData;
+};
+
+
+template<typename MappingType>
+struct NIAGARA_API TTypedMeshUvMappingHandle : public FMeshUvMappingHandleBase
+{
+	TTypedMeshUvMappingHandle() = default;
+	TTypedMeshUvMappingHandle(FMeshUvMappingUsage InUsage, const TSharedPtr<MappingType>& InUvMappingData, bool bNeedsDataImmediately)
+		: FMeshUvMappingHandleBase(InUsage, InUvMappingData, bNeedsDataImmediately) {}
+
+	TTypedMeshUvMappingHandle(TTypedMeshUvMappingHandle<MappingType>&& Other) noexcept
+		: FMeshUvMappingHandleBase(Other) {}
+
+	TTypedMeshUvMappingHandle<MappingType>& operator=(TTypedMeshUvMappingHandle<MappingType>&& Other) noexcept
+	{
+		FMeshUvMappingHandleBase::operator=(MoveTemp(Other));
+		return *this;
+	}
+
+	const MappingType* GetMappingData() const
+	{
+		if (UvMappingData)
+		{
+			return reinterpret_cast<const MappingType*>(UvMappingData.Get());
+		}
+
+		return nullptr;
+	}
+};
+
+using FSkeletalMeshUvMappingHandle = TTypedMeshUvMappingHandle<FSkeletalMeshUvMapping>;
+using FStaticMeshUvMappingHandle = TTypedMeshUvMappingHandle<FStaticMeshUvMapping>;
\ No newline at end of file
