- Add Attribute Capture Mode options to Niagara SimCache
- This allows us to capture all, renderer only, or user supplied attributes only into the cache
- Update the UI to handle cache changing while the UI is open and to use the new API to read the data

diff --git a/Engine/Plugins/FX/Niagara/Source/Niagara/Classes/NiagaraSimCache.h b/Engine/Plugins/FX/Niagara/Source/Niagara/Classes/NiagaraSimCache.h
--- a/Engine/Plugins/FX/Niagara/Source/Niagara/Classes/NiagaraSimCache.h
+++ b/Engine/Plugins/FX/Niagara/Source/Niagara/Classes/NiagaraSimCache.h
@@ -8,6 +8,31 @@
 
 class UNiagaraComponent;
 
+UENUM(BlueprintType)
+enum class ENiagaraSimCacheAttributeCaptureMode : uint8
+{
+	/**
+	Captures all attributes available.
+	This kind of cache will be useful for restarting a simulation from or debugging.
+	*/
+	All UMETA(DisplayName = "Capture All Attributes"),
+
+	/**
+	Captures attributes that are required to render the system only.
+	This kind of cache is useful for rendering only and should have a much smaller
+	size than capturing all attributes.
+	*/
+	RenderingOnly UMETA(DisplayName = "Capture Attributes Needed For Rendering"),
+
+	/**
+	Captures only attributes that match the 'ExplicitCaptureAttributes' list provided by the user.
+	This kind of cache is useful to keep the size down if you need to debug a very
+	specific attribute, or you want to do some additional process on the attributes
+	i.e. capture MyEmitter.Particles.Position and place static meshes in those locations.
+	*/
+	ExplicitAttributes UMETA(DisplayName = "Capture Explicit Attributes Only"),
+};
+
 USTRUCT(BlueprintType)
 struct FNiagaraSimCacheCreateParameters
 {
@@ -19,6 +44,13 @@ struct FNiagaraSimCacheCreateParameters
 	{
 	}
 
+	/**
+	How do we want to capture attributes for the simulation cache.
+	The mode selected depends on what situations the cache can be used in.
+	*/
+	UPROPERTY(EditAnywhere, BlueprintReadWrite, Category = "SimCache")
+	ENiagaraSimCacheAttributeCaptureMode AttributeCaptureMode = ENiagaraSimCacheAttributeCaptureMode::All;
+
 	/**
 	When enabled allows the SimCache to be re-based.
 	i.e. World space emitters can be moved to the new component's location
@@ -33,27 +65,26 @@ struct FNiagaraSimCacheCreateParameters
 	UPROPERTY(EditAnywhere, BlueprintReadWrite, Category="SimCache")
 	uint32 bAllowDataInterfaceCaching : 1;
 
-	/**
-	When enabled the SimCache will only be useful for rendering a replay, it can not be used to restart
-	the simulation from as only attributes & data interface that impact rendering will be stored.
-	This should result in a much smaller caches.
-	*/
-	//UPROPERTY(EditAnywhere, Category="SimCache")
-	//uint32 bRenderOnly : 1;
-
 	/**
 	List of Attributes to force include in the SimCache rebase, they should be the full path to the attribute
 	For example, MyEmitter.Particles.MyQuat would force the particle attribute MyQuat to be included for MyEmitter
 	*/
 	UPROPERTY(EditAnywhere, BlueprintReadWrite, Category = "SimCache")
-	TArray<FName> RebaseIncludeList;
+	TArray<FName> RebaseIncludeAttributes;
 
 	/**
 	List of Attributes to force exclude from the SimCache rebase, they should be the full path to the attribute
 	For example, MyEmitter.Particles.MyQuat would force the particle attribute MyQuat to be included for MyEmitter
 	*/
 	UPROPERTY(EditAnywhere, BlueprintReadWrite, Category = "SimCache")
-	TArray<FName> RebaseExcludeList;
+	TArray<FName> RebaseExcludeAttributes;
+
+	/**
+	List of attributes to capture when the capture attribute capture mode is set to explicit.
+	For example, adding MyEmitter.Particles.Position will only gather that attribute inside the cache.
+	*/
+	UPROPERTY(EditAnywhere, BlueprintReadWrite, Category = "SimCache", meta=(EditCondition="AttributeCaptureMode == ENiagaraSimCacheAttributeCaptureMode::ExplicitAttributes"))
+	TArray<FName> ExplicitCaptureAttributes;
 };
 
 USTRUCT()
@@ -224,11 +255,14 @@ class NIAGARA_API UNiagaraSimCache : public UObject
 {
 	friend struct FNiagaraSimCacheAttributeReaderHelper;
 	friend struct FNiagaraSimCacheHelper;
-	friend struct FSimCacheBufferReader;	//-TODO: Refactor to use API
+	friend class FNiagaraSimCacheViewModel;	//-TODO: Refactor to use API
 
 	GENERATED_UCLASS_BODY()
 
 public:
+	DECLARE_MULTICAST_DELEGATE_OneParam(FOnCacheBeginWrite, UNiagaraSimCache*)
+	DECLARE_MULTICAST_DELEGATE_OneParam(FOnCacheEndWrite, UNiagaraSimCache*)
+
 	// UObject Interface
 	virtual bool IsReadyForFinishDestroy() override;
 	// UObject Interface
@@ -268,6 +302,13 @@ public:
 	UFUNCTION(BlueprintCallable, Category = NiagaraSimCache)
 	TArray<FName> GetEmitterNames() const;
 
+	/**
+	Reads Niagara attributes by name from the cache frame and appends them into the relevant arrays.
+	When reading using this method the attributes for a FVector would come out as AoS rather than SoA so each component is wrote one by one
+	EmitterName - If left blank will return the system simulation attributes.
+	*/
+	void ReadAttribute(TArray<float>& OutFloats, TArray<FFloat16>& OutHalfs, TArray<int32>& OutInts, FName AttributeName, FName EmitterName, int FrameIndex = 0) const;
+
 	/**
 	Reads Niagara int attributes by name from the cache frame and appends them into the OutValues array.
 	EmitterName - If left blank will return the system simulation attributes.
@@ -375,4 +416,8 @@ private:
 	TMap<FNiagaraVariableBase, TObjectPtr<UObject>> DataInterfaceStorage;
 
 	mutable std::atomic<int32> PendingCommandsInFlight;
+
+public:
+	static FOnCacheBeginWrite	OnCacheBeginWrite;
+	static FOnCacheEndWrite		OnCacheEndWrite;
 };
diff --git a/Engine/Plugins/FX/Niagara/Source/Niagara/Private/NiagaraSimCache.cpp b/Engine/Plugins/FX/Niagara/Source/Niagara/Private/NiagaraSimCache.cpp
--- a/Engine/Plugins/FX/Niagara/Source/Niagara/Private/NiagaraSimCache.cpp
+++ b/Engine/Plugins/FX/Niagara/Source/Niagara/Private/NiagaraSimCache.cpp
@@ -16,6 +16,9 @@
 #include "NiagaraSystem.h"
 #include "NiagaraSystemInstance.h"
 
+UNiagaraSimCache::FOnCacheBeginWrite	UNiagaraSimCache::OnCacheBeginWrite;
+UNiagaraSimCache::FOnCacheEndWrite		UNiagaraSimCache::OnCacheEndWrite;
+
 UNiagaraSimCache::UNiagaraSimCache(const FObjectInitializer& ObjectInitializer)
 {
 }
@@ -29,6 +32,8 @@ void UNiagaraSimCache::BeginWrite(FNiagaraSimCacheCreateParameters InCreateParam
 {
 	check(PendingCommandsInFlight == 0);
 
+	OnCacheBeginWrite.Broadcast(this);
+
 	FNiagaraSimCacheHelper Helper(NiagaraComponent);
 	if (Helper.HasValidSimulation() == false)
 	{
@@ -36,6 +41,9 @@ void UNiagaraSimCache::BeginWrite(FNiagaraSimCacheCreateParameters InCreateParam
 	}
 
 	Modify();
+#if WITH_EDITOR
+	PostEditChange();
+#endif
 
 	// Reset to defaults
 	SoftNiagaraSystem = Helper.NiagaraSystem;
@@ -45,6 +53,59 @@ void UNiagaraSimCache::BeginWrite(FNiagaraSimCacheCreateParameters InCreateParam
 	CacheLayout = FNiagaraSimCacheLayout();
 	CacheFrames.Empty();
 
+	for ( auto it=DataInterfaceStorage.CreateIterator(); it; ++it )
+	{
+		it->Value->MarkAsGarbage();
+	}
+	DataInterfaceStorage.Empty();
+
+	// Not explicit capture mode?  Empty the list as internally we reuse this list when in rendering only mode
+	if (CreateParameters.AttributeCaptureMode == ENiagaraSimCacheAttributeCaptureMode::ExplicitAttributes)
+	{
+		// Invalid as the user specified nothing to actual capture
+		if (CreateParameters.ExplicitCaptureAttributes.IsEmpty())
+		{
+			SoftNiagaraSystem.Reset();
+			return;
+		}
+	}
+	else
+	{
+		CreateParameters.ExplicitCaptureAttributes.Empty();
+	}
+
+	// When in rendering only mode ask all renderers which attributes we need to capture
+	if ( CreateParameters.AttributeCaptureMode == ENiagaraSimCacheAttributeCaptureMode::RenderingOnly )
+	{
+	#if WITH_EDITORONLY_DATA
+		for ( const FNiagaraEmitterHandle& EmitterHandle : Helper.NiagaraSystem->GetEmitterHandles() )
+		{
+			EmitterHandle.GetInstance().GetEmitterData()->ForEachEnabledRenderer(
+				[&](UNiagaraRendererProperties* RenderProperties)
+				{
+					for (FNiagaraVariableBase BoundAttribute : RenderProperties->GetBoundAttributes())
+					{
+						FNameBuilder NameBuilder;
+						NameBuilder.Append(EmitterHandle.GetUniqueInstanceName());
+						NameBuilder.AppendChar(TEXT('.'));
+						BoundAttribute.GetName().AppendString(NameBuilder);
+
+						CreateParameters.ExplicitCaptureAttributes.Emplace(NameBuilder.ToString());
+					}
+				}
+			);
+		}
+	#endif
+		for (const TSharedRef<FNiagaraEmitterInstance, ESPMode::ThreadSafe>& EmitterInstance : Helper.SystemInstance->GetEmitters())
+		{
+			const FNiagaraParameterStore& RendererBindings = EmitterInstance->GetRendererBoundVariables();
+			for ( const FNiagaraVariableWithOffset& Variable : RendererBindings.ReadParameterVariables() )
+			{
+				CreateParameters.ExplicitCaptureAttributes.Add(Variable.GetName());
+			}
+		}
+	}
+
 	// Build new layout for system / emitters
 	Helper.BuildCacheLayoutForSystem(CreateParameters, CacheLayout.SystemLayout);
 
@@ -62,10 +123,13 @@ void UNiagaraSimCache::BeginWrite(FNiagaraSimCacheCreateParameters InCreateParam
 			Helper.SystemInstance,
 			[&](const FNiagaraVariableBase& Variable, UNiagaraDataInterface* DataInterface)
 			{
-				const void* PerInstanceData = Helper.SystemInstance->FindDataInterfaceInstanceData(DataInterface);
-				if (UObject* DICacheStorage = DataInterface->SimCacheBeginWrite(this, Helper.SystemInstance, PerInstanceData))
+				if ( (CreateParameters.ExplicitCaptureAttributes.Num() == 0) || CreateParameters.ExplicitCaptureAttributes.Contains(Variable.GetName()) )
 				{
-					DataInterfaceStorage.FindOrAdd(Variable) = DICacheStorage;
+					const void* PerInstanceData = Helper.SystemInstance->FindDataInterfaceInstanceData(DataInterface);
+					if (UObject* DICacheStorage = DataInterface->SimCacheBeginWrite(this, Helper.SystemInstance, PerInstanceData))
+					{
+						DataInterfaceStorage.FindOrAdd(Variable) = DICacheStorage;
+					}
 				}
 				return true;
 			}
@@ -201,6 +265,8 @@ void UNiagaraSimCache::EndWrite()
 			SoftNiagaraSystem.Reset();
 		}
 	}
+
+	OnCacheEndWrite.Broadcast(this);
 }
 
 bool UNiagaraSimCache::CanRead(UNiagaraSystem* NiagaraSystem)
@@ -386,6 +452,31 @@ TArray<FName> UNiagaraSimCache::GetEmitterNames() const
 	return EmitterNames;
 }
 
+void UNiagaraSimCache::ReadAttribute(TArray<float>& OutFloats, TArray<FFloat16>& OutHalfs, TArray<int32>& OutInts, FName AttributeName, FName EmitterName, int FrameIndex) const
+{
+	FNiagaraSimCacheAttributeReaderHelper AttributeReader(this, EmitterName, AttributeName, FrameIndex);
+	if ( AttributeReader.IsValid() )
+	{
+		const int32 OutFloatOffset = OutFloats.AddUninitialized(AttributeReader.Variable->FloatCount * AttributeReader.GetNumInstances());
+		for ( int32 i=0; i < AttributeReader.Variable->FloatCount; ++i )
+		{
+			AttributeReader.CopyComponentFloats(i, &OutFloats[OutFloatOffset + (i * AttributeReader.GetNumInstances())]);
+		}
+
+		const int32 OutHalfOffset = OutHalfs.AddUninitialized(AttributeReader.Variable->HalfCount * AttributeReader.GetNumInstances());
+		for (int32 i = 0; i < AttributeReader.Variable->HalfCount; ++i)
+		{
+			AttributeReader.CopyComponentHalfs(i, &OutHalfs[OutHalfOffset + (i * AttributeReader.GetNumInstances())]);
+		}
+
+		const int32 OutIntOffset = OutInts.AddUninitialized(AttributeReader.Variable->Int32Count * AttributeReader.GetNumInstances());
+		for (int32 i = 0; i < AttributeReader.Variable->Int32Count; ++i)
+		{
+			AttributeReader.CopyComponentInts(i, &OutInts[OutIntOffset + (i * AttributeReader.GetNumInstances())]);
+		}
+	}
+}
+
 void UNiagaraSimCache::ReadIntAttribute(TArray<int32>& OutValues, FName AttributeName, FName EmitterName, int FrameIndex) const
 {
 	FNiagaraSimCacheAttributeReaderHelper AttributeReader(this, EmitterName, AttributeName, FrameIndex);
diff --git a/Engine/Plugins/FX/Niagara/Source/Niagara/Private/NiagaraSimCacheAttributeReaderHelper.h b/Engine/Plugins/FX/Niagara/Source/Niagara/Private/NiagaraSimCacheAttributeReaderHelper.h
--- a/Engine/Plugins/FX/Niagara/Source/Niagara/Private/NiagaraSimCacheAttributeReaderHelper.h
+++ b/Engine/Plugins/FX/Niagara/Source/Niagara/Private/NiagaraSimCacheAttributeReaderHelper.h
@@ -54,6 +54,24 @@ struct FNiagaraSimCacheAttributeReaderHelper
 
 	int32 GetNumInstances() const { check(IsValid()); return DataBuffers->NumInstances; }
 
+	void CopyComponentFloats(int32 Component, float* OutBuffer)
+	{
+		check(Component < Variable->FloatCount);
+		FMemory::Memcpy(OutBuffer, &DataBuffers->FloatData[(Variable->FloatOffset + Component) * DataBuffers->NumInstances * sizeof(float)], sizeof(float) * DataBuffers->NumInstances);
+	}
+
+	void CopyComponentHalfs(int32 Component, FFloat16* OutBuffer)
+	{
+		check(Component < Variable->HalfCount);
+		FMemory::Memcpy(OutBuffer, &DataBuffers->HalfData[(Variable->HalfOffset + Component) * DataBuffers->NumInstances * sizeof(FFloat16)], sizeof(FFloat16) * DataBuffers->NumInstances);
+	}
+
+	void CopyComponentInts(int32 Component, int32* OutBuffer)
+	{
+		check(Component < Variable->Int32Count);
+		FMemory::Memcpy(OutBuffer, &DataBuffers->Int32Data[(Variable->Int32Offset + Component) * DataBuffers->NumInstances * sizeof(int32)], sizeof(int32) * DataBuffers->NumInstances);
+	}
+
 	float ReadInt(int32 Instance) const
 	{
 		check(IsValid());
diff --git a/Engine/Plugins/FX/Niagara/Source/Niagara/Private/NiagaraSimCacheHelper.h b/Engine/Plugins/FX/Niagara/Source/Niagara/Private/NiagaraSimCacheHelper.h
--- a/Engine/Plugins/FX/Niagara/Source/Niagara/Private/NiagaraSimCacheHelper.h
+++ b/Engine/Plugins/FX/Niagara/Source/Niagara/Private/NiagaraSimCacheHelper.h
@@ -58,25 +58,41 @@ struct FNiagaraSimCacheHelper
 	bool HasValidSimulation() const { return SystemSimulation != nullptr; }
 	bool HasValidSimulationData() const { return SystemSimulationDataBuffer != nullptr; }
 
-	void BuildCacheLayout(FNiagaraSimCacheDataBuffersLayout& CacheLayout, const FNiagaraDataSetCompiledData& CompiledData, FName LayoutName, TArray<FName> InRebaseVariableNames) const
+	void BuildCacheLayout(FNiagaraSimCacheDataBuffersLayout& CacheLayout, const FNiagaraDataSetCompiledData& CompiledData, FName LayoutName, TArray<FName> InRebaseVariableNames, TConstArrayView<FName> ExplicitCaptureAttributes) const
 	{
 		CacheLayout.LayoutName = LayoutName;
 		CacheLayout.SimTarget = CompiledData.SimTarget;
 
-		const int32 NumVariables = CompiledData.Variables.Num();
-		CacheLayout.Variables.AddDefaulted(NumVariables);
+		int32 TotalCacheComponents = 0;
+		TArray<int32> CacheToDataSetVariables;
 
-		const int32 CacheTotalComponents = CompiledData.TotalFloatComponents + CompiledData.TotalFloatComponents + CompiledData.TotalInt32Components;
-		CacheLayout.ComponentMappingsFromDataBuffer.Empty(CacheTotalComponents);
-		CacheLayout.ComponentMappingsFromDataBuffer.AddDefaulted(CacheTotalComponents);
+		CacheToDataSetVariables.Reserve(CompiledData.Variables.Num());
+		for (int32 i = 0; i < CompiledData.Variables.Num(); ++i)
+		{
+			const FNiagaraVariableLayoutInfo& DataSetVariableLayout = CompiledData.VariableLayouts[i];
+			const FNiagaraVariable& DataSetVariable = CompiledData.Variables[i];
+			if (ExplicitCaptureAttributes.Num() == 0 || ExplicitCaptureAttributes.Contains(DataSetVariable.GetName()))
+			{
+				CacheToDataSetVariables.Add(i);
+				TotalCacheComponents += DataSetVariableLayout.GetNumFloatComponents() + DataSetVariableLayout.GetNumHalfComponents() + DataSetVariableLayout.GetNumInt32Components();
+			}
+		}
+
+		const int32 NumCacheVariables = CacheToDataSetVariables.Num();
+
+		CacheLayout.Variables.AddDefaulted(NumCacheVariables);
+
+		CacheLayout.ComponentMappingsFromDataBuffer.Empty(TotalCacheComponents);
+		CacheLayout.ComponentMappingsFromDataBuffer.AddDefaulted(TotalCacheComponents);
 		CacheLayout.RebaseVariableNames = MoveTemp(InRebaseVariableNames);
 
-		for ( int32 iVariable=0; iVariable < NumVariables; ++iVariable)
+		for ( int32 iCacheVariable=0; iCacheVariable < NumCacheVariables; ++iCacheVariable)
 		{
-			const FNiagaraVariableLayoutInfo& DataSetVariableLayout = CompiledData.VariableLayouts[iVariable];
-			FNiagaraSimCacheVariable& CacheVariable = CacheLayout.Variables[iVariable];
+			const int32 iDataSetVariable = CacheToDataSetVariables[iCacheVariable];
+			const FNiagaraVariableLayoutInfo& DataSetVariableLayout = CompiledData.VariableLayouts[iDataSetVariable];
+			FNiagaraSimCacheVariable& CacheVariable = CacheLayout.Variables[iCacheVariable];
 
-			CacheVariable.Variable = CompiledData.Variables[iVariable];
+			CacheVariable.Variable = CompiledData.Variables[iDataSetVariable];
 			CacheVariable.FloatOffset = DataSetVariableLayout.GetNumFloatComponents() > 0 ? CacheLayout.FloatCount : INDEX_NONE;
 			CacheVariable.FloatCount = DataSetVariableLayout.GetNumFloatComponents();
 			CacheVariable.HalfOffset = DataSetVariableLayout.GetNumHalfComponents() > 0 ? CacheLayout.HalfCount : INDEX_NONE;
@@ -93,26 +109,27 @@ struct FNiagaraSimCacheHelper
 		int32 FloatOffset = 0;
 		int32 HalfOffset = CacheLayout.FloatCount;
 		int32 Int32Offset = HalfOffset + CacheLayout.HalfCount;
-		for (int32 iVariable = 0; iVariable < NumVariables; ++iVariable)
+		for (int32 iCacheVariable = 0; iCacheVariable < NumCacheVariables; ++iCacheVariable)
 		{
-			const FNiagaraVariableLayoutInfo& DataSetVariableLayout = CompiledData.VariableLayouts[iVariable];
-			FNiagaraSimCacheVariable& CacheVariable = CacheLayout.Variables[iVariable];
+			const int32 iDataSetVariable = CacheToDataSetVariables[iCacheVariable];
+			const FNiagaraVariableLayoutInfo& DataSetVariableLayout = CompiledData.VariableLayouts[iDataSetVariable];
+			FNiagaraSimCacheVariable& CacheVariable = CacheLayout.Variables[iCacheVariable];
 
 			for (int32 iComponent=0; iComponent < CacheVariable.FloatCount; ++iComponent)
 			{
-				CacheLayout.ComponentMappingsFromDataBuffer[FloatOffset] = CacheVariable.FloatOffset + iComponent;
+				CacheLayout.ComponentMappingsFromDataBuffer[FloatOffset] = DataSetVariableLayout.FloatComponentStart + iComponent;
 				++FloatOffset;
 			}
 
 			for (int32 iComponent=0; iComponent < CacheVariable.HalfCount; ++iComponent)
 			{
-				CacheLayout.ComponentMappingsFromDataBuffer[HalfOffset] = CacheVariable.HalfOffset + iComponent;
+				CacheLayout.ComponentMappingsFromDataBuffer[HalfOffset] = DataSetVariableLayout.HalfComponentStart + iComponent;
 				++HalfOffset;
 			}
 
 			for (int32 iComponent=0; iComponent < CacheVariable.Int32Count; ++iComponent)
 			{
-				CacheLayout.ComponentMappingsFromDataBuffer[Int32Offset] = CacheVariable.Int32Offset + iComponent;
+				CacheLayout.ComponentMappingsFromDataBuffer[Int32Offset] = DataSetVariableLayout.Int32ComponentStart + iComponent;
 				++Int32Offset;
 			}
 		}
@@ -121,12 +138,12 @@ struct FNiagaraSimCacheHelper
 		BuildCacheReadMappings(CacheLayout, CompiledData);
 	}
 
-	void BuildCacheLayoutForSystem(const FNiagaraSimCacheCreateParameters& CreateParmaeters, FNiagaraSimCacheDataBuffersLayout& CacheLayout)
+	void BuildCacheLayoutForSystem(const FNiagaraSimCacheCreateParameters& CreateParameters, FNiagaraSimCacheDataBuffersLayout& CacheLayout)
 	{
 		const FNiagaraDataSetCompiledData& SystemCompileData = NiagaraSystem->GetSystemCompiledData().DataSetCompiledData;
 
 		TArray<FName> RebaseVariableNames;
-		if ( CreateParmaeters.bAllowRebasing )
+		if ( CreateParameters.bAllowRebasing )
 		{
 			TArray<FString, TInlineAllocator<8>> LocalSpaceEmitters;
 			for ( int32 i=0; i < NiagaraSystem->GetNumEmitters(); ++i )
@@ -153,22 +170,22 @@ struct FNiagaraSimCacheHelper
 						}
 					}
 
-					if ( bIsLocalSpace == false && CreateParmaeters.RebaseExcludeList.Contains(Variable.GetName()) == false )
+					if ( bIsLocalSpace == false && CreateParameters.RebaseExcludeAttributes.Contains(Variable.GetName()) == false )
 					{
 						RebaseVariableNames.AddUnique(Variable.GetName());
 					}
 				}
-				else if ( CanRebaseVariable(Variable) && CreateParmaeters.RebaseIncludeList.Contains(Variable.GetName()) )
+				else if ( CanRebaseVariable(Variable) && CreateParameters.RebaseIncludeAttributes.Contains(Variable.GetName()) )
 				{
 					RebaseVariableNames.AddUnique(Variable.GetName());
 				}
 			}
 		}
 
-		BuildCacheLayout(CacheLayout, SystemCompileData, NiagaraSystem->GetFName(), MoveTemp(RebaseVariableNames));
+		BuildCacheLayout(CacheLayout, SystemCompileData, NiagaraSystem->GetFName(), MoveTemp(RebaseVariableNames), CreateParameters.ExplicitCaptureAttributes);
 	}
 
-	void BuildCacheLayoutForEmitter(const FNiagaraSimCacheCreateParameters& CreateParmaeters, FNiagaraSimCacheDataBuffersLayout& CacheLayout, int EmitterIndex)
+	void BuildCacheLayoutForEmitter(const FNiagaraSimCacheCreateParameters& CreateParameters, FNiagaraSimCacheDataBuffersLayout& CacheLayout, int EmitterIndex)
 	{
 		const FNiagaraEmitterHandle& EmitterHandle = NiagaraSystem->GetEmitterHandle(EmitterIndex);
 		const FNiagaraEmitterCompiledData& EmitterCompiledData = NiagaraSystem->GetEmitterCompiledData()[EmitterIndex].Get();
@@ -177,15 +194,15 @@ struct FNiagaraSimCacheHelper
 		CacheLayout.bLocalSpace = EmitterHandle.GetInstance().GetEmitterData()->bLocalSpace;
 
 		TArray<FName> RebaseVariableNames;
-		if ( CreateParmaeters.bAllowRebasing && CacheLayout.bLocalSpace == false )
+		if ( CreateParameters.bAllowRebasing && CacheLayout.bLocalSpace == false )
 		{
 			// Build list of include / exclude names
 			TArray<FName> ForceIncludeNames;
 			TArray<FName> ForceExcludeNames;
-			if ( CreateParmaeters.RebaseIncludeList.Num() > 0 || CreateParmaeters.RebaseExcludeList.Num() > 0 )
+			if ( CreateParameters.RebaseIncludeAttributes.Num() > 0 || CreateParameters.RebaseExcludeAttributes.Num() > 0 )
 			{
 				const FString EmitterName = EmitterHandle.GetUniqueInstanceName();
-				for (FName RebaseName : CreateParmaeters.RebaseIncludeList)
+				for (FName RebaseName : CreateParameters.RebaseIncludeAttributes)
 				{
 					FNiagaraVariableBase BaseVar(FNiagaraTypeDefinition::GetFloatDef(), RebaseName);
 					if (BaseVar.RemoveRootNamespace(EmitterName))
@@ -194,7 +211,7 @@ struct FNiagaraSimCacheHelper
 					}
 				}
 
-				for (FName RebaseName : CreateParmaeters.RebaseExcludeList)
+				for (FName RebaseName : CreateParameters.RebaseExcludeAttributes)
 				{
 					FNiagaraVariableBase BaseVar(FNiagaraTypeDefinition::GetFloatDef(), RebaseName);
 					if (BaseVar.RemoveRootNamespace(EmitterName))
@@ -244,8 +261,24 @@ struct FNiagaraSimCacheHelper
 				}
 			}
 		}
+		TArray<FName> ExplicitCaptureAttributes;
+		if ( CreateParameters.ExplicitCaptureAttributes.Num() > 0 )
+		{
+			const FString EmitterName = EmitterHandle.GetUniqueInstanceName();
+			for ( FName AttributeName : CreateParameters.ExplicitCaptureAttributes)
+			{
+				FNiagaraVariableBase AttributeVar(FNiagaraTypeDefinition::GetFloatDef(), AttributeName);
+				if (AttributeVar.RemoveRootNamespace(EmitterName))
+				{
+					if (AttributeVar.RemoveRootNamespace(FNiagaraConstants::ParticleAttributeNamespaceString))
+					{
+						ExplicitCaptureAttributes.Add(AttributeVar.GetName());
+					}
+				}
+			}
+		}
 
-		BuildCacheLayout(CacheLayout, EmitterCompiledData.DataSetCompiledData, EmitterHandle.GetName(), MoveTemp(RebaseVariableNames));
+		BuildCacheLayout(CacheLayout, EmitterCompiledData.DataSetCompiledData, EmitterHandle.GetName(), MoveTemp(RebaseVariableNames), ExplicitCaptureAttributes);
 	}
 
 	static bool BuildCacheReadMappings(FNiagaraSimCacheDataBuffersLayout& CacheLayout, const FNiagaraDataSetCompiledData& CompiledData)
diff --git a/Engine/Plugins/FX/Niagara/Source/NiagaraEditor/Private/ViewModels/NiagaraSimCacheViewModel.cpp b/Engine/Plugins/FX/Niagara/Source/NiagaraEditor/Private/ViewModels/NiagaraSimCacheViewModel.cpp
--- a/Engine/Plugins/FX/Niagara/Source/NiagaraEditor/Private/ViewModels/NiagaraSimCacheViewModel.cpp
+++ b/Engine/Plugins/FX/Niagara/Source/NiagaraEditor/Private/ViewModels/NiagaraSimCacheViewModel.cpp
@@ -4,231 +4,32 @@
 #include "NiagaraSimCache.h"
 #include "NiagaraEditorCommon.h"
 #include "NiagaraDebuggerCommon.h"
+#include "Misc/ScopeExit.h"
 #include "Serialization/ObjectAndNameAsStringProxyArchive.h"
 
 #define LOCTEXT_NAMESPACE "NiagaraSimCacheViewModel"
 
-struct FSimCacheBufferReader
-{
-
-	explicit FSimCacheBufferReader(TWeakObjectPtr<UNiagaraSimCache> SimCache, uint32 InFrameIndex = 0, uint32 InEmitterIndex = INDEX_NONE)
-		: WeakSimCache(SimCache)
-	{
-		UpdateLayout(InFrameIndex, InEmitterIndex);
-	}
-
-	void UpdateLayout(const uint32 FrameIndex, const uint32 EmitterIndex)
-	{
-		ComponentInfos.Empty();
-		FoundFloats = 0;
-		FoundHalfs = 0;
-		FoundInt32s = 0;
-
-		if (const FNiagaraSimCacheDataBuffersLayout* SimCacheLayout = GetSimCacheBufferLayout(FrameIndex, EmitterIndex))
-		{
-			for (const FNiagaraSimCacheVariable& Variable : SimCacheLayout->Variables)
-			{
-				const FNiagaraTypeDefinition& TypeDef = Variable.Variable.GetType();
-				if (TypeDef.IsEnum())
-				{
-					FNiagaraSimCacheComponentInfo& ComponentInfo = ComponentInfos.AddDefaulted_GetRef();
-					ComponentInfo.Name = Variable.Variable.GetName();
-					ComponentInfo.ComponentOffset = FoundInt32s++;
-					ComponentInfo.bIsInt32 = true;
-					ComponentInfo.Enum = TypeDef.GetEnum();
-				}
-				else
-				{
-					Build(Variable.Variable.GetName(), TypeDef.GetScriptStruct());
-				}
-			}
-
-			if ((FoundFloats != SimCacheLayout->FloatCount) ||
-				(FoundHalfs != SimCacheLayout->HalfCount) ||
-				(FoundInt32s != SimCacheLayout->Int32Count))
-			{
-				UE_LOG(LogNiagaraEditor, Warning, TEXT("SimCache Layout doesn't appear to match iterating the variables"));
-				ComponentInfos.Empty();
-			}
-		}
-	}
-
-	TConstArrayView<FNiagaraSimCacheComponentInfo> GetComponentInfos() const
-	{
-		return MakeArrayView(ComponentInfos);
-	}
-	
-	UNiagaraSimCache* GetSimCache() const
-	{
-		return WeakSimCache.Get();
-	}
-
-	const FNiagaraSimCacheDataBuffersLayout* GetSimCacheBufferLayout(const uint32 FrameIndex, const uint32 EmitterIndex) const
-	{
-		if (UNiagaraSimCache* const SimCache = WeakSimCache.Get())
-		{
-			if (SimCache->CacheFrames.IsValidIndex(FrameIndex))
-			{
-				if (EmitterIndex == INDEX_NONE)
-				{
-					return &SimCache->CacheLayout.SystemLayout;
-				}
-				else if (SimCache->CacheLayout.EmitterLayouts.IsValidIndex(EmitterIndex))
-				{
-					return &SimCache->CacheLayout.EmitterLayouts[EmitterIndex];
-				}
-			}
-		}
-		return nullptr;
-	}
-
-	const FNiagaraSimCacheDataBuffers* GetSimCacheDataBuffers(const uint32 FrameIndex, const uint32 EmitterIndex) const
-	{
-		if (UNiagaraSimCache* SimCache = WeakSimCache.Get())
-		{
-			if (SimCache->CacheFrames.IsValidIndex(FrameIndex))
-			{
-				if (EmitterIndex == INDEX_NONE)
-				{
-					return &SimCache->CacheFrames[FrameIndex].SystemData.SystemDataBuffers;
-				}
-				else if (SimCache->CacheFrames[FrameIndex].EmitterData.IsValidIndex(EmitterIndex))
-				{
-					return &SimCache->CacheFrames[FrameIndex].EmitterData[EmitterIndex].ParticleDataBuffers;
-				}
-			}
-		}
-		return nullptr;
-	}
-
-	FText GetComponentText(FName ComponentName, int32 InstanceIndex, uint32 FrameIndex, uint32 EmitterIndex) const
-	{
-		const FNiagaraSimCacheComponentInfo* ComponentInfo = ComponentInfos.FindByPredicate([ComponentName](const FNiagaraSimCacheComponentInfo& Info) { return Info.Name == ComponentName; });
-		const FNiagaraSimCacheDataBuffers* SimCacheDataBuffers = GetSimCacheDataBuffers(FrameIndex, EmitterIndex);
-	
-		if (ComponentInfo && SimCacheDataBuffers)
-		{
-			const int32 NumInstances = SimCacheDataBuffers->NumInstances;
-			if (InstanceIndex >= 0 && InstanceIndex < NumInstances)
-			{
-				if (ComponentInfo->bIsFloat)
-				{
-					const float* Value = reinterpret_cast<const float*>(SimCacheDataBuffers->FloatData.GetData()) + (ComponentInfo->ComponentOffset * NumInstances) + InstanceIndex;
-					return FText::AsNumber(*Value);
-				}
-				else if (ComponentInfo->bIsHalf)
-				{
-					const FFloat16* Value = reinterpret_cast<const FFloat16*>(SimCacheDataBuffers->HalfData.GetData()) + (ComponentInfo->ComponentOffset * NumInstances) + InstanceIndex;
-					return FText::AsNumber(Value->GetFloat());
-				}
-				else if (ComponentInfo->bIsInt32)
-				{
-					const int32* Value = reinterpret_cast<const int32*>(SimCacheDataBuffers->Int32Data.GetData()) + (ComponentInfo->ComponentOffset * NumInstances) + InstanceIndex;
-					if (ComponentInfo->bShowAsBool)
-					{
-						return *Value == 0 ? LOCTEXT("False", "False") : LOCTEXT("True", "True");
-					}
-					else if (ComponentInfo->Enum != nullptr)
-					{
-						return ComponentInfo->Enum->GetDisplayNameTextByValue(*Value);
-					}
-					else
-					{
-						return FText::AsNumber(*Value);
-					}
-				}
-			}
-		}
-		return LOCTEXT("Error", "Error");
-	}
-
-protected:
-	TWeakObjectPtr<UNiagaraSimCache>	WeakSimCache;
-
-	TArray<FNiagaraSimCacheComponentInfo>				ComponentInfos;
-	uint32								FoundFloats = 0;
-	uint32								FoundHalfs = 0;
-	uint32								FoundInt32s = 0;
-
-	void Build(const FName Name, const UScriptStruct* Struct)
-	{
-		int32 NumProperties = 0;
-		for (TFieldIterator<FProperty> PropertyIt(Struct, EFieldIteratorFlags::IncludeSuper); PropertyIt; ++PropertyIt)
-		{
-			++NumProperties;
-		}
-
-		for (TFieldIterator<FProperty> PropertyIt(Struct, EFieldIteratorFlags::IncludeSuper); PropertyIt; ++PropertyIt)
-		{
-			FProperty* Property = *PropertyIt;
-			const FName PropertyName = NumProperties > 1 ? FName(*FString::Printf(TEXT("%s.%s"), *Name.ToString(), *Property->GetName())) : Name;
-			if (Property->IsA(FFloatProperty::StaticClass()))
-			{
-				FNiagaraSimCacheComponentInfo& ComponentInfo = ComponentInfos.AddDefaulted_GetRef();
-				ComponentInfo.Name = PropertyName;
-				ComponentInfo.ComponentOffset = FoundFloats++;
-				ComponentInfo.bIsFloat = true;
-			}
-			else if (Property->IsA(FUInt16Property::StaticClass()))
-			{
-				FNiagaraSimCacheComponentInfo& ComponentInfo = ComponentInfos.AddDefaulted_GetRef();
-				ComponentInfo.Name = PropertyName;
-				ComponentInfo.ComponentOffset = FoundHalfs++;
-				ComponentInfo.bIsHalf = true;
-			}
-			else if (Property->IsA(FIntProperty::StaticClass()))
-			{
-				FNiagaraSimCacheComponentInfo& ComponentInfo = ComponentInfos.AddDefaulted_GetRef();
-				ComponentInfo.Name = PropertyName;
-				ComponentInfo.ComponentOffset = FoundInt32s++;
-				ComponentInfo.bIsInt32 = true;
-				ComponentInfo.bShowAsBool = (NumProperties == 1) && (Struct == FNiagaraTypeDefinition::GetBoolStruct());
-			}
-			else if (Property->IsA(FBoolProperty::StaticClass()))
-			{
-				FNiagaraSimCacheComponentInfo& ComponentInfo = ComponentInfos.AddDefaulted_GetRef();
-				ComponentInfo.Name = PropertyName;
-				ComponentInfo.ComponentOffset = FoundInt32s++;
-				ComponentInfo.bIsInt32 = true;
-				ComponentInfo.bShowAsBool = true;
-			}
-			else if (Property->IsA(FEnumProperty::StaticClass()))
-			{
-				FNiagaraSimCacheComponentInfo& ComponentInfo = ComponentInfos.AddDefaulted_GetRef();
-				ComponentInfo.Name = PropertyName;
-				ComponentInfo.ComponentOffset = FoundInt32s++;
-				ComponentInfo.bIsInt32 = true;
-				ComponentInfo.Enum = CastFieldChecked<FEnumProperty>(Property)->GetEnum();
-			}
-			else if (Property->IsA(FStructProperty::StaticClass()))
-			{
-				const FStructProperty* StructProperty = CastFieldChecked<FStructProperty>(Property);
-				Build(PropertyName, FNiagaraTypeHelper::FindNiagaraFriendlyTopLevelStruct(StructProperty->Struct, ENiagaraStructConversion::Simulation));
-			}
-			else
-			{
-				// Fail
-			}
-		}
-	}
-};
-
-TSharedPtr<FSimCacheBufferReader> BufferReader;
-
 FNiagaraSimCacheViewModel::FNiagaraSimCacheViewModel()
 {
-	
 }
 
 FNiagaraSimCacheViewModel::~FNiagaraSimCacheViewModel()
 {
-
+	UNiagaraSimCache::OnCacheEndWrite.RemoveAll(this);
+	bDelegatesAdded = false;
 }
 
 void FNiagaraSimCacheViewModel::Initialize(TWeakObjectPtr<UNiagaraSimCache> SimCache) 
 {
-	BufferReader = MakeShared<FSimCacheBufferReader>(SimCache);
-	BufferReader->UpdateLayout(FrameIndex, EmitterIndex);
+	if (bDelegatesAdded == false)
+	{
+		bDelegatesAdded = true;
+		UNiagaraSimCache::OnCacheEndWrite.AddSP(this, &FNiagaraSimCacheViewModel::OnCacheModified);
+	}
+
+	WeakSimCache = SimCache;
+	UpdateCachedFrame();
+	OnViewDataChangedDelegate.Broadcast(true);
 }
 
 void FNiagaraSimCacheViewModel::UpdateSimCache(const FNiagaraSystemSimCacheCaptureReply& Reply)
@@ -252,64 +53,211 @@ void FNiagaraSimCacheViewModel::UpdateSimCache(const FNiagaraSystemSimCacheCaptu
 
 FText FNiagaraSimCacheViewModel::GetComponentText(const FName ComponentName, const int32 InstanceIndex) const
 {
-	return BufferReader ? BufferReader->GetComponentText(ComponentName, InstanceIndex, FrameIndex, EmitterIndex) : FText();
-}
+	const FComponentInfo* ComponentInfo = ComponentInfos.FindByPredicate([ComponentName](const FComponentInfo& FoundInfo) { return FoundInfo.Name == ComponentName; });
 
-TConstArrayView<FNiagaraSimCacheComponentInfo> FNiagaraSimCacheViewModel::GetComponentInfos() const
-{
-	return BufferReader ? BufferReader->GetComponentInfos() : TConstArrayView<FNiagaraSimCacheComponentInfo>();
-}
-
-int32 FNiagaraSimCacheViewModel::GetNumInstances() const
-{
-	const FNiagaraSimCacheDataBuffers* DataBuffer = BufferReader ? BufferReader->GetSimCacheDataBuffers(FrameIndex, EmitterIndex) : nullptr;
-	return DataBuffer ? DataBuffer->NumInstances : 0;
+	if (ComponentInfo)
+	{
+		if (InstanceIndex >= 0 && InstanceIndex < NumInstances)
+		{
+			if (ComponentInfo->bIsFloat)
+			{
+				const float Value = FloatComponents[(ComponentInfo->ComponentOffset * NumInstances) + InstanceIndex];
+				return FText::AsNumber(Value);
+			}
+			else if (ComponentInfo->bIsHalf)
+			{
+				const FFloat16 Value = HalfComponents[(ComponentInfo->ComponentOffset * NumInstances) + InstanceIndex];
+				return FText::AsNumber(Value.GetFloat());
+			}
+			else if (ComponentInfo->bIsInt32)
+			{
+				const int32 Value = Int32Components[(ComponentInfo->ComponentOffset * NumInstances) + InstanceIndex];
+				if (ComponentInfo->bShowAsBool)
+				{
+					return Value == 0 ? LOCTEXT("False", "False") : LOCTEXT("True", "True");
+				}
+				else if (ComponentInfo->Enum != nullptr)
+				{
+					return ComponentInfo->Enum->GetDisplayNameTextByValue(Value);
+				}
+				else
+				{
+					return FText::AsNumber(Value);
+				}
+			}
+		}
+	}
+	return LOCTEXT("Error", "Error");
 }
 
 int32 FNiagaraSimCacheViewModel::GetNumFrames() const
 {
-	const UNiagaraSimCache* SimCache = BufferReader ? BufferReader->GetSimCache() : nullptr;
+	UNiagaraSimCache* SimCache = WeakSimCache.Get();
 	return SimCache ? SimCache->GetNumFrames() : 0;
 }
 
 void FNiagaraSimCacheViewModel::SetFrameIndex(const int32 InFrameIndex)
 {
 	FrameIndex = InFrameIndex;
-	OnFrameUpdatedDelegate.Broadcast(true);
+	UpdateCachedFrame();
+	OnViewDataChangedDelegate.Broadcast(false);
 }
 
-
 void FNiagaraSimCacheViewModel::SetEmitterIndex(const int32 InEmitterIndex)
 {
 	EmitterIndex = InEmitterIndex;
-	BufferReader->UpdateLayout(FrameIndex, EmitterIndex);
-}
-
-const FNiagaraSimCacheDataBuffersLayout* FNiagaraSimCacheViewModel::GetSimCacheBufferLayout() const
-{
-	return BufferReader ? BufferReader->GetSimCacheBufferLayout(FrameIndex, EmitterIndex) : nullptr;
+	UpdateCachedFrame();
+	OnViewDataChangedDelegate.Broadcast(true);
 }
 
 bool FNiagaraSimCacheViewModel::IsCacheValid()
 {
-	return BufferReader ? BufferReader->GetSimCache()->IsCacheValid() : false;
+	UNiagaraSimCache* SimCache = WeakSimCache.Get();
+	return SimCache ? SimCache->IsCacheValid() : false;
 }
 
 int32 FNiagaraSimCacheViewModel::GetNumEmitterLayouts()
 {
-	const UNiagaraSimCache* SimCache = BufferReader ? BufferReader->GetSimCache() : nullptr;
+	UNiagaraSimCache* SimCache = WeakSimCache.Get();
 	return SimCache ? SimCache->GetNumEmitters() : 0;
 }
 
 FName FNiagaraSimCacheViewModel::GetEmitterLayoutName(const int32 Index)
 {
-	const UNiagaraSimCache* SimCache = BufferReader ? BufferReader->GetSimCache() : nullptr;
+	UNiagaraSimCache* SimCache = WeakSimCache.Get();
 	return SimCache ? SimCache->GetEmitterName(Index) : NAME_None;
 }
 
-FNiagaraSimCacheViewModel::FOnFrameUpdated& FNiagaraSimCacheViewModel::OnFrameUpdated()
+FNiagaraSimCacheViewModel::FOnViewDataChanged& FNiagaraSimCacheViewModel::OnViewDataChanged()
+{
+	return OnViewDataChangedDelegate;
+}
+
+void FNiagaraSimCacheViewModel::OnCacheModified(UNiagaraSimCache* SimCache)
+{
+	if ( UNiagaraSimCache* ThisSimCache = WeakSimCache.Get() )
+	{
+		if ( ThisSimCache == SimCache )
+		{
+			UpdateCachedFrame();
+			OnViewDataChangedDelegate.Broadcast(true);
+		}
+	}
+}
+
+void FNiagaraSimCacheViewModel::UpdateCachedFrame()
+{
+	NumInstances = 0;
+	ComponentInfos.Empty();
+	FloatComponents.Empty();
+	HalfComponents.Empty();
+	Int32Components.Empty();
+
+	UNiagaraSimCache* SimCache = WeakSimCache.Get();
+	if (SimCache == nullptr)
+	{
+		return;
+	}
+
+	if ( FrameIndex < 0 || FrameIndex >= SimCache->GetNumFrames() )
+	{
+		return;
+	}
+
+	if ( EmitterIndex != INDEX_NONE && EmitterIndex >= SimCache->GetNumEmitters() )
+	{
+		return;
+	}
+
+	FoundFloatComponents = 0;
+	FoundHalfComponents = 0;
+	FoundInt32Components = 0;
+
+	NumInstances = EmitterIndex == INDEX_NONE ? 1 : SimCache->CacheFrames[FrameIndex].EmitterData[EmitterIndex].ParticleDataBuffers.NumInstances;
+	const FName EmitterName = EmitterIndex == INDEX_NONE ? NAME_None : SimCache->GetEmitterName(EmitterIndex);
+	const FNiagaraSimCacheDataBuffersLayout* DataBufferLayout = EmitterIndex == INDEX_NONE ? &SimCache->CacheLayout.SystemLayout : &SimCache->CacheLayout.EmitterLayouts[EmitterIndex];
+	for ( const FNiagaraSimCacheVariable& Variable : DataBufferLayout->Variables )
+	{
+		// Build component info
+		const FNiagaraTypeDefinition& TypeDef = Variable.Variable.GetType();
+		if (TypeDef.IsEnum())
+		{
+			FComponentInfo& ComponentInfo = ComponentInfos.AddDefaulted_GetRef();
+			ComponentInfo.Name = Variable.Variable.GetName();
+			ComponentInfo.ComponentOffset = FoundInt32Components++;
+			ComponentInfo.bIsInt32 = true;
+			ComponentInfo.Enum = TypeDef.GetEnum();
+		}
+		else
+		{
+			BuildComponentInfos(Variable.Variable.GetName(), TypeDef.GetScriptStruct());
+		}
+
+		// Pull in data
+		SimCache->ReadAttribute(FloatComponents, HalfComponents, Int32Components, Variable.Variable.GetName(), EmitterName, FrameIndex);
+	}
+}
+
+void FNiagaraSimCacheViewModel::BuildComponentInfos(const FName Name, const UScriptStruct* Struct)
 {
-	return OnFrameUpdatedDelegate;
+	int32 NumProperties = 0;
+	for (TFieldIterator<FProperty> PropertyIt(Struct, EFieldIteratorFlags::IncludeSuper); PropertyIt; ++PropertyIt)
+	{
+		++NumProperties;
+	}
+
+	for (TFieldIterator<FProperty> PropertyIt(Struct, EFieldIteratorFlags::IncludeSuper); PropertyIt; ++PropertyIt)
+	{
+		FProperty* Property = *PropertyIt;
+		const FName PropertyName = NumProperties > 1 ? FName(*FString::Printf(TEXT("%s.%s"), *Name.ToString(), *Property->GetName())) : Name;
+		if (Property->IsA(FFloatProperty::StaticClass()))
+		{
+			FComponentInfo& ComponentInfo = ComponentInfos.AddDefaulted_GetRef();
+			ComponentInfo.Name = PropertyName;
+			ComponentInfo.ComponentOffset = FoundFloatComponents++;
+			ComponentInfo.bIsFloat = true;
+		}
+		else if (Property->IsA(FUInt16Property::StaticClass()))
+		{
+			FComponentInfo& ComponentInfo = ComponentInfos.AddDefaulted_GetRef();
+			ComponentInfo.Name = PropertyName;
+			ComponentInfo.ComponentOffset = FoundHalfComponents++;
+			ComponentInfo.bIsHalf = true;
+		}
+		else if (Property->IsA(FIntProperty::StaticClass()))
+		{
+			FComponentInfo& ComponentInfo = ComponentInfos.AddDefaulted_GetRef();
+			ComponentInfo.Name = PropertyName;
+			ComponentInfo.ComponentOffset = FoundInt32Components++;
+			ComponentInfo.bIsInt32 = true;
+			ComponentInfo.bShowAsBool = (NumProperties == 1) && (Struct == FNiagaraTypeDefinition::GetBoolStruct());
+		}
+		else if (Property->IsA(FBoolProperty::StaticClass()))
+		{
+			FComponentInfo& ComponentInfo = ComponentInfos.AddDefaulted_GetRef();
+			ComponentInfo.Name = PropertyName;
+			ComponentInfo.ComponentOffset = FoundInt32Components++;
+			ComponentInfo.bIsInt32 = true;
+			ComponentInfo.bShowAsBool = true;
+		}
+		else if (Property->IsA(FEnumProperty::StaticClass()))
+		{
+			FComponentInfo& ComponentInfo = ComponentInfos.AddDefaulted_GetRef();
+			ComponentInfo.Name = PropertyName;
+			ComponentInfo.ComponentOffset = FoundInt32Components++;
+			ComponentInfo.bIsInt32 = true;
+			ComponentInfo.Enum = CastFieldChecked<FEnumProperty>(Property)->GetEnum();
+		}
+		else if (Property->IsA(FStructProperty::StaticClass()))
+		{
+			const FStructProperty* StructProperty = CastFieldChecked<FStructProperty>(Property);
+			BuildComponentInfos(PropertyName, FNiagaraTypeHelper::FindNiagaraFriendlyTopLevelStruct(StructProperty->Struct, ENiagaraStructConversion::Simulation));
+		}
+		else
+		{
+			// Fail
+		}
+	}
 }
 
 #undef LOCTEXT_NAMESPACE
diff --git a/Engine/Plugins/FX/Niagara/Source/NiagaraEditor/Private/ViewModels/NiagaraSimCacheViewModel.h b/Engine/Plugins/FX/Niagara/Source/NiagaraEditor/Private/ViewModels/NiagaraSimCacheViewModel.h
--- a/Engine/Plugins/FX/Niagara/Source/NiagaraEditor/Private/ViewModels/NiagaraSimCacheViewModel.h
+++ b/Engine/Plugins/FX/Niagara/Source/NiagaraEditor/Private/ViewModels/NiagaraSimCacheViewModel.h
@@ -8,21 +8,21 @@ class UNiagaraSimCache;
 struct FNiagaraSystemSimCacheCaptureReply;
 struct FNiagaraSimCacheDataBuffersLayout;
 
-struct FNiagaraSimCacheComponentInfo
-{
-	FName Name = NAME_None;
-	uint32 ComponentOffset = INDEX_NONE;
-	bool bIsFloat = false;
-	bool bIsHalf = false;
-	bool bIsInt32 = false;
-	bool bShowAsBool = false;
-	UEnum* Enum = nullptr;
-};
-
 class NIAGARAEDITOR_API FNiagaraSimCacheViewModel : public TSharedFromThis<FNiagaraSimCacheViewModel>
 {
 public:
-	DECLARE_MULTICAST_DELEGATE_OneParam(FOnFrameUpdated, bool)
+	struct FComponentInfo
+	{
+		FName Name = NAME_None;
+		uint32 ComponentOffset = INDEX_NONE;
+		bool bIsFloat = false;
+		bool bIsHalf = false;
+		bool bIsInt32 = false;
+		bool bShowAsBool = false;
+		UEnum* Enum = nullptr;
+	};
+
+	DECLARE_MULTICAST_DELEGATE_OneParam(FOnViewDataChanged, bool)
 	
 	FNiagaraSimCacheViewModel();
 	~FNiagaraSimCacheViewModel();
@@ -31,22 +31,20 @@ public:
 
 	void UpdateSimCache(const FNiagaraSystemSimCacheCaptureReply& Reply);
 
-	TConstArrayView<FNiagaraSimCacheComponentInfo> GetComponentInfos() const;
+	TConstArrayView<FComponentInfo> GetComponentInfos() const { return ComponentInfos; }
 
-	int32 GetNumInstances() const;
+	int32 GetNumInstances() const { return NumInstances; }
 
 	int32 GetNumFrames() const;
 
 	int32 GetFrameIndex() const { return FrameIndex; };
 
-	void SetFrameIndex(const int32 InFrameIndex);;
+	void SetFrameIndex(const int32 InFrameIndex);
 
 	int32 GetEmitterIndex() const { return EmitterIndex; };
 
 	void SetEmitterIndex(int32 InEmitterIndex);
 
-	const FNiagaraSimCacheDataBuffersLayout* GetSimCacheBufferLayout() const;
-
 	FText GetComponentText(FName ComponentName, int32 InstanceIndex) const;
 
 	bool IsCacheValid();
@@ -55,11 +53,29 @@ public:
 
 	FName GetEmitterLayoutName(int32 Index);
 
-	FOnFrameUpdated& OnFrameUpdated();
+	FOnViewDataChanged& OnViewDataChanged();
+
+	void OnCacheModified(UNiagaraSimCache* SimCache);
+
+	void UpdateCachedFrame();
 
 private:
-	uint32 FrameIndex = 0;
-	uint32 EmitterIndex = INDEX_NONE;
+	void BuildComponentInfos(const FName Name, const UScriptStruct* Struct);
 
-	FOnFrameUpdated OnFrameUpdatedDelegate;
-};
\ No newline at end of file
+private:
+	TWeakObjectPtr<UNiagaraSimCache>	WeakSimCache;
+	int32								FrameIndex = 0;
+	int32								EmitterIndex = INDEX_NONE;
+	int32								NumInstances = 0;
+	TArray<FComponentInfo>				ComponentInfos;
+	TArray<float>						FloatComponents;
+	TArray<FFloat16>					HalfComponents;
+	TArray<int32>						Int32Components;
+
+	int32								FoundFloatComponents = 0;
+	int32								FoundHalfComponents = 0;
+	int32								FoundInt32Components = 0;
+
+	bool								bDelegatesAdded = false;
+	FOnViewDataChanged					OnViewDataChangedDelegate;
+};
diff --git a/Engine/Plugins/FX/Niagara/Source/NiagaraEditor/Private/Widgets/SNiagaraSimCacheView.cpp b/Engine/Plugins/FX/Niagara/Source/NiagaraEditor/Private/Widgets/SNiagaraSimCacheView.cpp
--- a/Engine/Plugins/FX/Niagara/Source/NiagaraEditor/Private/Widgets/SNiagaraSimCacheView.cpp
+++ b/Engine/Plugins/FX/Niagara/Source/NiagaraEditor/Private/Widgets/SNiagaraSimCacheView.cpp
@@ -67,7 +67,7 @@ void SNiagaraSimCacheView::Construct(const FArguments& InArgs)
 {
 	SimCacheViewModel = InArgs._SimCacheViewModel;
 
-	SimCacheViewModel->OnFrameUpdated().AddSP(this, &SNiagaraSimCacheView::OnSimCacheFrameChanged);
+	SimCacheViewModel->OnViewDataChanged().AddSP(this, &SNiagaraSimCacheView::OnViewDataChanged);
 
 	HeaderRowWidget = SNew(SHeaderRow);
 
@@ -206,7 +206,7 @@ void SNiagaraSimCacheView::UpdateColumns(const bool bRefresh)
 
 
 	const bool bFilterActive = ComponentFilterArray.Num() > 0;
-	for (const FNiagaraSimCacheComponentInfo& ComponentInfo : SimCacheViewModel->GetComponentInfos())
+	for (const FNiagaraSimCacheViewModel::FComponentInfo& ComponentInfo : SimCacheViewModel->GetComponentInfos())
 	{
 		if (bFilterActive)
 		{
@@ -323,9 +323,14 @@ void SNiagaraSimCacheView::OnSimCacheChanged(const FAssetData& InAsset)
 	}
 }
 
-void SNiagaraSimCacheView::OnSimCacheFrameChanged(const bool bRefresh)
+void SNiagaraSimCacheView::OnViewDataChanged(const bool bFullRefresh)
 {
-	UpdateRows(bRefresh);
+	UpdateRows(true);
+	if (bFullRefresh)
+	{
+		UpdateColumns(true);
+		UpdateBufferSelectionList();
+	}
 }
 
 #undef LOCTEXT_NAMESPACE
diff --git a/Engine/Plugins/FX/Niagara/Source/NiagaraEditor/Private/Widgets/SNiagaraSimCacheView.h b/Engine/Plugins/FX/Niagara/Source/NiagaraEditor/Private/Widgets/SNiagaraSimCacheView.h
--- a/Engine/Plugins/FX/Niagara/Source/NiagaraEditor/Private/Widgets/SNiagaraSimCacheView.h
+++ b/Engine/Plugins/FX/Niagara/Source/NiagaraEditor/Private/Widgets/SNiagaraSimCacheView.h
@@ -39,7 +39,7 @@ public:
 
 	void OnSimCacheChanged(const FAssetData& InAsset);
 
-	void OnSimCacheFrameChanged(const bool bRefresh);
+	void OnViewDataChanged(const bool bFullRefresh);
 
 	TArray<TSharedPtr<int32>>					RowItems;
 	TSharedPtr<FNiagaraSimCacheViewModel>		SimCacheViewModel;
