- IWYU fix for nonunity

#rnx

diff --git a/Engine/Plugins/FX/Niagara/Source/NiagaraEditor/Private/ViewModels/NiagaraSimCacheViewModel.h b/Engine/Plugins/FX/Niagara/Source/NiagaraEditor/Private/ViewModels/NiagaraSimCacheViewModel.h
--- a/Engine/Plugins/FX/Niagara/Source/NiagaraEditor/Private/ViewModels/NiagaraSimCacheViewModel.h
+++ b/Engine/Plugins/FX/Niagara/Source/NiagaraEditor/Private/ViewModels/NiagaraSimCacheViewModel.h
@@ -3,8 +3,8 @@
 #pragma once
 
 #include "CoreMinimal.h"
+#include "NiagaraSimCache.h"
 
-class UNiagaraSimCache;
 struct FNiagaraSystemSimCacheCaptureReply;
 struct FNiagaraSimCacheDataBuffersLayout;
 
