- Fix IsReady not triggering if the readback fails

#rnx

diff --git a/Engine/Plugins/FX/Niagara/Source/Niagara/Private/NiagaraDataSetReadback.cpp b/Engine/Plugins/FX/Niagara/Source/Niagara/Private/NiagaraDataSetReadback.cpp
--- a/Engine/Plugins/FX/Niagara/Source/Niagara/Private/NiagaraDataSetReadback.cpp
+++ b/Engine/Plugins/FX/Niagara/Source/Niagara/Private/NiagaraDataSetReadback.cpp
@@ -73,6 +73,25 @@ void FNiagaraDataSetReadback::ImmediateReadback(FNiagaraEmitterInstance* Emitter
 	}
 }
 
+void FNiagaraDataSetReadback::ReadbackCompleteInternal()
+{
+	if ( OnReadbackReady.IsBound() )
+	{
+		AsyncTask(
+			ENamedThreads::GameThread,
+			[Readback=this]()
+			{
+				Readback->PendingReadbacks--;
+				Readback->OnReadbackReady.Execute(*Readback);
+			}
+		);
+	}
+	else
+	{
+		PendingReadbacks--;
+	}
+}
+
 void FNiagaraDataSetReadback::GPUReadbackInternal(FRHICommandListImmediate& RHICmdList, FNiagaraGpuComputeDispatchInterface* DispatchInterface, FNiagaraComputeExecutionContext* GPUContext)
 {
 	FNiagaraGpuReadbackManager* ReadbackManager = DispatchInterface->GetGpuReadbackManager();
@@ -80,6 +99,7 @@ void FNiagaraDataSetReadback::GPUReadbackInternal(FRHICommandListImmediate& RHIC
 	if (CurrentDataBuffer == nullptr || ReadbackManager == nullptr)
 	{
 		DataSet.CopyFromGPUReadback(nullptr, nullptr, nullptr, 0, 0, 0, 0, 0);
+		ReadbackCompleteInternal();
 		return;
 	}
 
@@ -87,6 +107,7 @@ void FNiagaraDataSetReadback::GPUReadbackInternal(FRHICommandListImmediate& RHIC
 	if (CountOffset == INDEX_NONE)
 	{
 		DataSet.CopyFromGPUReadback(nullptr, nullptr, nullptr, 0, 0, 0, 0, 0);
+		ReadbackCompleteInternal();
 		return;
 	}
 
@@ -135,21 +156,8 @@ void FNiagaraDataSetReadback::GPUReadbackInternal(FRHICommandListImmediate& RHIC
 				Readback->IDToIndexTable.Empty();
 			}
 
-			if ( Readback->OnReadbackReady.IsBound() )
-			{
-				AsyncTask(
-					ENamedThreads::GameThread,
-					[Readback]()
-					{
-						Readback->PendingReadbacks--;
-						Readback->OnReadbackReady.Execute(*Readback);
-					}
-				);
-			}
-			else
-			{
-				Readback->PendingReadbacks--;
-			}
+
+			Readback->ReadbackCompleteInternal();
 		}
 	);
 }
diff --git a/Engine/Plugins/FX/Niagara/Source/Niagara/Public/NiagaraDataSetReadback.h b/Engine/Plugins/FX/Niagara/Source/Niagara/Public/NiagaraDataSetReadback.h
--- a/Engine/Plugins/FX/Niagara/Source/Niagara/Public/NiagaraDataSetReadback.h
+++ b/Engine/Plugins/FX/Niagara/Source/Niagara/Public/NiagaraDataSetReadback.h
@@ -28,6 +28,7 @@ public:
 	void ImmediateReadback(FNiagaraEmitterInstance* EmitterInstance);
 
 private:
+	void ReadbackCompleteInternal();
 	void GPUReadbackInternal(FRHICommandListImmediate& RHICmdList, FNiagaraGpuComputeDispatchInterface* DispatchInterface, FNiagaraComputeExecutionContext* GPUContext);
 
 private:
