- Handle nullptr for GPUComputeScript

#rnx

diff --git a/Engine/Plugins/FX/Niagara/Source/Niagara/Classes/NiagaraEmitter.h b/Engine/Plugins/FX/Niagara/Source/Niagara/Classes/NiagaraEmitter.h
--- a/Engine/Plugins/FX/Niagara/Source/Niagara/Classes/NiagaraEmitter.h
+++ b/Engine/Plugins/FX/Niagara/Source/Niagara/Classes/NiagaraEmitter.h
@@ -1002,7 +1002,10 @@ void FVersionedNiagaraEmitterData::ForEachScript(TAction Func) const
 	Func(SpawnScriptProps.Script);
 	Func(UpdateScriptProps.Script);
 
-	Func(GPUComputeScript);
+	if (GPUComputeScript)
+	{
+		Func(GPUComputeScript);
+	}
 
 	for (auto& EventScriptProps : EventHandlerScriptProps)
 	{
