Fix for compiler warning

diff --git a/Engine/Plugins/FX/Niagara/Source/Niagara/Public/Experimental/NiagaraMeshUvMappingHandle.h b/Engine/Plugins/FX/Niagara/Source/Niagara/Public/Experimental/NiagaraMeshUvMappingHandle.h
--- a/Engine/Plugins/FX/Niagara/Source/Niagara/Public/Experimental/NiagaraMeshUvMappingHandle.h
+++ b/Engine/Plugins/FX/Niagara/Source/Niagara/Public/Experimental/NiagaraMeshUvMappingHandle.h
@@ -66,7 +66,7 @@ struct NIAGARA_API TTypedMeshUvMappingHandle : public FMeshUvMappingHandleBase
 	{
 		if (UvMappingData)
 		{
-			return reinterpret_cast<const MappingType*>(UvMappingData.Get());
+			return static_cast<const MappingType*>(UvMappingData.Get());
 		}
 
 		return nullptr;
