- Handle cooked out & disabled emitters when generating the cache

#rnx

diff --git a/Engine/Plugins/FX/Niagara/Source/Niagara/Private/NiagaraSimCacheHelper.h b/Engine/Plugins/FX/Niagara/Source/Niagara/Private/NiagaraSimCacheHelper.h
--- a/Engine/Plugins/FX/Niagara/Source/Niagara/Private/NiagaraSimCacheHelper.h
+++ b/Engine/Plugins/FX/Niagara/Source/Niagara/Private/NiagaraSimCacheHelper.h
@@ -149,9 +149,13 @@ struct FNiagaraSimCacheHelper
 			for ( int32 i=0; i < NiagaraSystem->GetNumEmitters(); ++i )
 			{
 				const FNiagaraEmitterHandle& EmitterHandle = NiagaraSystem->GetEmitterHandle(i);
-				if ( EmitterHandle.GetInstance().GetEmitterData()->bLocalSpace )
+				if (EmitterHandle.GetIsEnabled() )
 				{
-					LocalSpaceEmitters.Add(EmitterHandle.GetUniqueInstanceName());
+					const FVersionedNiagaraEmitterData* EmitterData = EmitterHandle.GetInstance().GetEmitterData();
+					if (EmitterData && EmitterData->bLocalSpace)
+					{
+						LocalSpaceEmitters.Add(EmitterHandle.GetUniqueInstanceName());
+					}
 				}
 			}
 
@@ -189,9 +193,14 @@ struct FNiagaraSimCacheHelper
 	{
 		const FNiagaraEmitterHandle& EmitterHandle = NiagaraSystem->GetEmitterHandle(EmitterIndex);
 		const FNiagaraEmitterCompiledData& EmitterCompiledData = NiagaraSystem->GetEmitterCompiledData()[EmitterIndex].Get();
+		const FVersionedNiagaraEmitterData* EmitterData = EmitterHandle.GetInstance().GetEmitterData();
+		if (EmitterHandle.GetIsEnabled() == false || EmitterData == nullptr)
+		{
+			return;
+		}
 
 		// Find potential candidates for re-basing
-		CacheLayout.bLocalSpace = EmitterHandle.GetInstance().GetEmitterData()->bLocalSpace;
+		CacheLayout.bLocalSpace = EmitterData->bLocalSpace;
 
 		TArray<FName> RebaseVariableNames;
 		if ( CreateParameters.bAllowRebasing && CacheLayout.bLocalSpace == false )
@@ -224,7 +233,7 @@ struct FNiagaraSimCacheHelper
 		#if WITH_EDITORONLY_DATA
 			// Look for renderer attributes bound to Quat / Matrix types are we will want to rebase those
 			// We will add all Position types after this so no need to add them here
-			EmitterHandle.GetInstance().GetEmitterData()->ForEachEnabledRenderer(
+			EmitterData->ForEachEnabledRenderer(
 				[&](UNiagaraRendererProperties* RenderProperties)
 				{
 					for (FNiagaraVariableBase BoundAttribute : RenderProperties->GetBoundAttributes())
