- Use tick count when capturing the sim cache, it's more robust than using Age
- Add time out to multi frame / until complete capture types to ensure they complete if we pause the system

diff --git a/Engine/Plugins/FX/Niagara/Source/Niagara/Classes/NiagaraSimCache.h b/Engine/Plugins/FX/Niagara/Source/Niagara/Classes/NiagaraSimCache.h
--- a/Engine/Plugins/FX/Niagara/Source/Niagara/Classes/NiagaraSimCache.h
+++ b/Engine/Plugins/FX/Niagara/Source/Niagara/Classes/NiagaraSimCache.h
@@ -271,9 +271,9 @@ public:
 	UFUNCTION(BlueprintCallable, Category=NiagaraSimCache)
 	bool IsCacheValid() const { return SoftNiagaraSystem.IsNull() == false; }
 
-	void BeginWrite(FNiagaraSimCacheCreateParameters InCreateParameters, UNiagaraComponent* NiagaraComponent);
-	void WriteFrame(UNiagaraComponent* NiagaraComponent);
-	void EndWrite();
+	bool BeginWrite(FNiagaraSimCacheCreateParameters InCreateParameters, UNiagaraComponent* NiagaraComponent);
+	bool WriteFrame(UNiagaraComponent* NiagaraComponent);
+	bool EndWrite();
 
 	bool CanRead(UNiagaraSystem* NiagaraSystem);
 	bool Read(float TimeSeconds, FNiagaraSystemInstance* SystemInstance) const;
@@ -415,6 +415,8 @@ private:
 	UPROPERTY()
 	TMap<FNiagaraVariableBase, TObjectPtr<UObject>> DataInterfaceStorage;
 
+	int32 CaptureTickCount = INDEX_NONE;
+
 	mutable std::atomic<int32> PendingCommandsInFlight;
 
 public:
diff --git a/Engine/Plugins/FX/Niagara/Source/Niagara/Classes/NiagaraSimCacheFunctionLibrary.h b/Engine/Plugins/FX/Niagara/Source/Niagara/Classes/NiagaraSimCacheFunctionLibrary.h
--- a/Engine/Plugins/FX/Niagara/Source/Niagara/Classes/NiagaraSimCacheFunctionLibrary.h
+++ b/Engine/Plugins/FX/Niagara/Source/Niagara/Classes/NiagaraSimCacheFunctionLibrary.h
@@ -33,6 +33,9 @@ public:
 	UPROPERTY()
 	int32 CaptureFrameCounter = 0;
 
+	UPROPERTY()
+	int32 TimeOutCounter = 0;
+
 	UPROPERTY(BlueprintAssignable)
 	FOnCaptureComplete CaptureComplete;
 
diff --git a/Engine/Plugins/FX/Niagara/Source/Niagara/Private/NiagaraSimCache.cpp b/Engine/Plugins/FX/Niagara/Source/Niagara/Private/NiagaraSimCache.cpp
--- a/Engine/Plugins/FX/Niagara/Source/Niagara/Private/NiagaraSimCache.cpp
+++ b/Engine/Plugins/FX/Niagara/Source/Niagara/Private/NiagaraSimCache.cpp
@@ -28,7 +28,7 @@ bool UNiagaraSimCache::IsReadyForFinishDestroy()
 	return PendingCommandsInFlight == 0;
 }
 
-void UNiagaraSimCache::BeginWrite(FNiagaraSimCacheCreateParameters InCreateParameters, UNiagaraComponent* NiagaraComponent)
+bool UNiagaraSimCache::BeginWrite(FNiagaraSimCacheCreateParameters InCreateParameters, UNiagaraComponent* NiagaraComponent)
 {
 	check(PendingCommandsInFlight == 0);
 
@@ -37,7 +37,7 @@ void UNiagaraSimCache::BeginWrite(FNiagaraSimCacheCreateParameters InCreateParam
 	FNiagaraSimCacheHelper Helper(NiagaraComponent);
 	if (Helper.HasValidSimulation() == false)
 	{
-		return;
+		return false;
 	}
 
 	Modify();
@@ -52,6 +52,7 @@ void UNiagaraSimCache::BeginWrite(FNiagaraSimCacheCreateParameters InCreateParam
 	DurationSeconds = 0.0f;
 	CacheLayout = FNiagaraSimCacheLayout();
 	CacheFrames.Empty();
+	CaptureTickCount = INDEX_NONE;
 
 	for ( auto it=DataInterfaceStorage.CreateIterator(); it; ++it )
 	{
@@ -66,7 +67,7 @@ void UNiagaraSimCache::BeginWrite(FNiagaraSimCacheCreateParameters InCreateParam
 		if (CreateParameters.ExplicitCaptureAttributes.IsEmpty())
 		{
 			SoftNiagaraSystem.Reset();
-			return;
+			return false;
 		}
 	}
 	else
@@ -135,46 +136,53 @@ void UNiagaraSimCache::BeginWrite(FNiagaraSimCacheCreateParameters InCreateParam
 			}
 		);
 	}
+	return true;
 }
 
-void UNiagaraSimCache::WriteFrame(UNiagaraComponent* NiagaraComponent)
+bool UNiagaraSimCache::WriteFrame(UNiagaraComponent* NiagaraComponent)
 {
 	FNiagaraSimCacheHelper Helper(NiagaraComponent);
 	if (Helper.HasValidSimulationData() == false)
 	{
 		SoftNiagaraSystem.Reset();
-		return;
+		return false;
 	}
 
 	if ( SoftNiagaraSystem.Get() != Helper.NiagaraSystem )
 	{
 		SoftNiagaraSystem.Reset();
-		return;
+		return false;
 	}
 
 	// Simulation is complete nothing to cache
 	if ( Helper.SystemInstance->IsComplete() )
 	{
-		return;
+		return false;
 	}
 
 	// Is the simulation running?  If not nothing to cache yet
 	if ( Helper.SystemInstance->SystemInstanceState != ENiagaraSystemInstanceState::Running )
 	{
-		return;
+		return false;
 	}
 
 	// First frame we are about to cache?
-	if ( CacheFrames.Num() == 0 )
+	if (CaptureTickCount == INDEX_NONE)
 	{
 		StartSeconds = Helper.SystemInstance->GetAge();
+		CaptureTickCount = Helper.SystemInstance->GetTickCount();
+	}
+	// If our tick counter hasn't moved then we won't capture a frame as there's nothing new to process
+	else if (CaptureTickCount == Helper.SystemInstance->GetTickCount())
+	{
+		return false;
 	}
 
-	// Invalid we have reset for some reason
-	if ( Helper.SystemInstance->GetAge() < StartSeconds + DurationSeconds )
+	// If the tick counter is lower than the previous value then the system was reset so we won't capture
+	if ( Helper.SystemInstance->GetTickCount() < CaptureTickCount )
 	{
 		SoftNiagaraSystem.Reset();
-		return;
+		return false;
 	}
 
 	DurationSeconds = Helper.SystemInstance->GetAge() - StartSeconds;
@@ -237,11 +245,13 @@ void UNiagaraSimCache::WriteFrame(UNiagaraComponent* NiagaraComponent)
 		if (bDataInterfacesSucess == false)
 		{
 			SoftNiagaraSystem.Reset();
+			return false;
 		}
 	}
+	return true;
 }
 
-void UNiagaraSimCache::EndWrite()
+bool UNiagaraSimCache::EndWrite()
 {
 	check(PendingCommandsInFlight == 0);
 	if ( CacheFrames.Num() == 0 )
@@ -267,6 +277,7 @@ void UNiagaraSimCache::EndWrite()
 	}
 
 	OnCacheEndWrite.Broadcast(this);
+	return IsCacheValid();
 }
 
 bool UNiagaraSimCache::CanRead(UNiagaraSystem* NiagaraSystem)
diff --git a/Engine/Plugins/FX/Niagara/Source/Niagara/Private/NiagaraSimCacheFunctionLibrary.cpp b/Engine/Plugins/FX/Niagara/Source/Niagara/Private/NiagaraSimCacheFunctionLibrary.cpp
--- a/Engine/Plugins/FX/Niagara/Source/Niagara/Private/NiagaraSimCacheFunctionLibrary.cpp
+++ b/Engine/Plugins/FX/Niagara/Source/Niagara/Private/NiagaraSimCacheFunctionLibrary.cpp
@@ -40,7 +40,7 @@ void UAsyncNiagaraCaptureSimCache::SetReadyToDestroy()
 bool UAsyncNiagaraCaptureSimCache::OnFrameTick(float DeltaTime)
 {
 	// Component invalid or not active?  If so complete the cache recording
-	if ( CaptureComponent == nullptr || !CaptureComponent->IsActive() )
+	if ( CaptureComponent == nullptr || !CaptureComponent->IsActive() || CaptureSimCache == nullptr )
 	{
 		SetReadyToDestroy();
 		return true;
@@ -49,11 +49,26 @@ bool UAsyncNiagaraCaptureSimCache::OnFrameTick(float DeltaTime)
 	// Should we record this frame?
 	if ( (CaptureFrameCounter % CaptureFrameRate) == 0 )
 	{
-		if ( CaptureSimCache != nullptr )
+		// If we fail to capture the frame it might be because things became invalid
+		// Or it might be because the simulation was not ticked since the last capture in which case don't advance the counter
+		if ( CaptureSimCache->WriteFrame(CaptureComponent) == false )
 		{
-			CaptureSimCache->WriteFrame(CaptureComponent);
+			if ( CaptureSimCache->IsCacheValid() == false )
+			{
+				SetReadyToDestroy();
+			}
+
+			// Make sure we don't keep this alive forever, if we didn't managed to capture anything in 10 ticks something has probably gone wrong so bail
+			if (TimeOutCounter++ > 10)
+			{
+				UE_LOG(LogNiagara, Warning, TEXT("SimCache Write has failed too many times, abandoning capturing for (%s)"), *GetFullNameSafe(CaptureSimCache));
+				SetReadyToDestroy();
+			}
+			return true;
 		}
 	}
+
+	TimeOutCounter = 0;
 	++CaptureFrameCounter;
 
 	// Have we recorded all the frames we need?
