Changing description for persistent id's in UI.

diff --git a/Engine/Plugins/FX/Niagara/Source/Niagara/Classes/NiagaraEmitter.h b/Engine/Plugins/FX/Niagara/Source/Niagara/Classes/NiagaraEmitter.h
--- a/Engine/Plugins/FX/Niagara/Source/Niagara/Classes/NiagaraEmitter.h
+++ b/Engine/Plugins/FX/Niagara/Source/Niagara/Classes/NiagaraEmitter.h
@@ -290,7 +290,7 @@ struct NIAGARA_API FVersionedNiagaraEmitterData
 	UPROPERTY(EditAnywhere, BlueprintReadWrite, Category = "Emitter", meta = (EditConditionHides, EditCondition = "CalculateBoundsMode == ENiagaraEmitterCalculateBoundMode::Fixed"))
 	FBox FixedBounds;
 
-	/** Do particles in this emitter require a persistent ID? */
+	/** Creates a stable Identifier (Particles.ID) which does not vary from frame to frame. This comes at a small memory and performance cost. This allows external objects to track the same particle over multiple frames. Particle arrays are tightly packed and a particle's actual index in the array may change from frame to frame. This optionally lets you use a lookup table to track a particle by index in the lookup table. */
 	UPROPERTY(EditAnywhere, BlueprintReadWrite, Category = "Emitter")
 	uint32 bRequiresPersistentIDs : 1;
 
