Fixed potential crash in Select node when the underlying type of a pin was deleted

diff --git a/Engine/Plugins/FX/Niagara/Source/NiagaraEditor/Private/NiagaraNodeSelect.cpp b/Engine/Plugins/FX/Niagara/Source/NiagaraEditor/Private/NiagaraNodeSelect.cpp
--- a/Engine/Plugins/FX/Niagara/Source/NiagaraEditor/Private/NiagaraNodeSelect.cpp
+++ b/Engine/Plugins/FX/Niagara/Source/NiagaraEditor/Private/NiagaraNodeSelect.cpp
@@ -197,6 +197,13 @@ void UNiagaraNodeSelect::AllocateDefaultPins()
 {
 	const UEdGraphSchema_Niagara* Schema = GetDefault<UEdGraphSchema_Niagara>();
 
+	// in case der selector type became invalid, we change it back to wildcard here.
+	// this can happen if the underlying ustruct, for example an enum, has been deleted
+	if(!SelectorPinType.IsValid())
+	{
+		SelectorPinType = FNiagaraTypeDefinition::GetWildcardDef();
+	}
+	
 	TArray<int32> OptionValues = GetOptionValues();
 	NumOptionsPerVariable = OptionValues.Num();
 
@@ -207,7 +214,7 @@ void UNiagaraNodeSelect::AllocateDefaultPins()
 			AddOptionPin(Variable, OptionValues[OptionIndex]);
 		}
 	}
-
+	
 	// create the selector pin
 	UEdGraphPin* SelectorPin = CreatePin(EGPD_Input, Schema->TypeDefinitionToPinType(SelectorPinType), GetSelectorPinName());
 	SelectorPin->PersistentGuid = SelectorPinGuid;
@@ -300,6 +307,20 @@ void UNiagaraNodeSelect::GetNodeContextMenuActions(UToolMenu* Menu, UGraphNodeCo
 void UNiagaraNodeSelect::Compile(FHlslNiagaraTranslator* Translator, TArray<int32>& Outputs)
 {	
 	const UEdGraphPin* SelectorPin = GetSelectorPin();
+
+	if(!SelectorPinType.IsValid())
+	{
+		Translator->Warning(LOCTEXT("SelectNodePinSelectorTypeInvalid", "Select node selector pin should have a valid type."), this, nullptr);
+	}
+	
+	for(FNiagaraVariable& Variable : OutputVars)
+	{
+		if(!Variable.GetType().IsValid())
+		{
+			Translator->Warning(FText::Format(LOCTEXT("SelectNodePinOutputTypeInvalid", "Select node output pin should have a valid type. {0} is invalid."), *Variable.GetName().ToString()), this, nullptr);
+		}
+	}
+	
 	int32 Selection = Translator->CompilePin(SelectorPin);
 
 	// a map from selector value to compiled option pins (i.e.: for selector value 0 all pins that should be case "if 0" get their compiled index added under key 0)
@@ -673,7 +694,7 @@ TArray<int32> UNiagaraNodeSelect::GetOptionValues() const
 			SelectorValues.Add(Index);
 		}
 	}
-	else if(SelectorPinType.IsEnum() && SelectorPinType.GetEnum())
+	else if(SelectorPinType.IsEnum() && SelectorPinType.IsValid() && SelectorPinType.GetEnum())
 	{
 		UEnum* Enum = SelectorPinType.GetEnum();
 		const int32 EnumEntryCount = Enum->NumEnums();
