Fixed compile issue from previous CL

diff --git a/Engine/Plugins/FX/Niagara/Source/NiagaraEditor/Private/NiagaraNodeSelect.cpp b/Engine/Plugins/FX/Niagara/Source/NiagaraEditor/Private/NiagaraNodeSelect.cpp
--- a/Engine/Plugins/FX/Niagara/Source/NiagaraEditor/Private/NiagaraNodeSelect.cpp
+++ b/Engine/Plugins/FX/Niagara/Source/NiagaraEditor/Private/NiagaraNodeSelect.cpp
@@ -317,7 +317,7 @@ void UNiagaraNodeSelect::Compile(FHlslNiagaraTranslator* Translator, TArray<int3
 	{
 		if(!Variable.GetType().IsValid())
 		{
-			Translator->Warning(FText::Format(LOCTEXT("SelectNodePinOutputTypeInvalid", "Select node output pin should have a valid type. {0} is invalid."), *Variable.GetName().ToString()), this, nullptr);
+			Translator->Warning(FText::Format(LOCTEXT("SelectNodePinOutputTypeInvalid", "Select node output pin should have a valid type. {0} is invalid."), FText::FromName(Variable.GetName())), this, nullptr);
 		}
 	}
 	
