Fixed issue with taking a temp vars address that is no longer in scope when used

diff --git a/Engine/Plugins/FX/Niagara/Source/NiagaraEditor/Private/NiagaraSystemEditorData.cpp b/Engine/Plugins/FX/Niagara/Source/NiagaraEditor/Private/NiagaraSystemEditorData.cpp
--- a/Engine/Plugins/FX/Niagara/Source/NiagaraEditor/Private/NiagaraSystemEditorData.cpp
+++ b/Engine/Plugins/FX/Niagara/Source/NiagaraEditor/Private/NiagaraSystemEditorData.cpp
@@ -371,7 +371,7 @@ UNiagaraScriptVariable* UNiagaraSystemEditorData::FindOrAddUserScriptVariable(FN
 			FNiagaraVariableMetaData MetaData;
 			NewScriptVariable->Init(UserParameter, MetaData);
 			UserParameterMetaData.Add(NewScriptVariable);
-			ExistingScriptVariable = &NewScriptVariable;
+			return NewScriptVariable;
 		}
 		
 		return *ExistingScriptVariable;
@@ -387,11 +387,6 @@ UNiagaraScriptVariable* UNiagaraSystemEditorData::FindUserScriptVariable(FGuid U
 			return ScriptVariable->Metadata.GetVariableGuid() == UserParameterGuid;
 		});
 
-	if(ExistingScriptVariable == nullptr)
-	{
-		return nullptr;
-	}
-
 	return *ExistingScriptVariable;
 }
 
