Added missing #pragma once; fixing for multiple defined UvMappingPrameters struct

diff --git a/Engine/Plugins/FX/Niagara/Shaders/Private/Experimental/NiagaraUvMappingUtils.ush b/Engine/Plugins/FX/Niagara/Shaders/Private/Experimental/NiagaraUvMappingUtils.ush
--- a/Engine/Plugins/FX/Niagara/Shaders/Private/Experimental/NiagaraUvMappingUtils.ush
+++ b/Engine/Plugins/FX/Niagara/Shaders/Private/Experimental/NiagaraUvMappingUtils.ush
@@ -1,5 +1,7 @@
 // Copyright Epic Games, Inc. All Rights Reserved.
 
+#pragma once
+
 #include "/Plugin/FX/Niagara/Private/NiagaraBaryCentricUtils.ush"
 
 struct UvMappingParameters
