Fixing User variable Use count as well as adding renderer bindings to the count.

diff --git a/Engine/Plugins/FX/Niagara/Source/NiagaraEditor/Private/ViewModels/NiagaraParameterPanelViewModel.cpp b/Engine/Plugins/FX/Niagara/Source/NiagaraEditor/Private/ViewModels/NiagaraParameterPanelViewModel.cpp
--- a/Engine/Plugins/FX/Niagara/Source/NiagaraEditor/Private/ViewModels/NiagaraParameterPanelViewModel.cpp
+++ b/Engine/Plugins/FX/Niagara/Source/NiagaraEditor/Private/ViewModels/NiagaraParameterPanelViewModel.cpp
@@ -2028,7 +2028,7 @@ TArray<FNiagaraParameterPanelItem> FNiagaraSystemToolkitParameterPanelViewModel:
 		}
 	}
 
-	if (bForceSystem)
+	if (bForceSystem || bUserOnly)
 	{
 		// Collect parameters for all emitters.
 		TArray<FNiagaraVariable> VisitedInvalidParameters;
@@ -2125,7 +2125,7 @@ TArray<FNiagaraParameterPanelItem> FNiagaraSystemToolkitParameterPanelViewModel:
 						// This variable has already been registered, increment the reference count.
 						ItemPtr->ReferenceCount += Builder.Histories[0].PerVariableReadHistory[VariableIndex].Num() + Builder.Histories[0].PerVariableWriteHistory[VariableIndex].Num();
 					}
-					else
+					else if (!bUserOnly) // Add newly found variables, unless in user-only mode in which we use the parameter store defined ones only
 					{
 						// This variable has not been registered, prepare the FNiagaraParameterPanelItem.
 						// -First make sure the variable namespace is in a valid category. If not, skip it.
@@ -2194,6 +2194,33 @@ TArray<FNiagaraParameterPanelItem> FNiagaraSystemToolkitParameterPanelViewModel:
 				}
 			}
 		}
+
+		// Add active renderers usage variables to the counts
+		for (TSharedRef<FNiagaraEmitterHandleViewModel> EmitterVMS : SystemViewModel->GetEmitterHandleViewModels())
+		{
+			if (EmitterVMS.Get().IsValid() && EmitterVMS.Get().GetIsEnabled() && EmitterVMS.Get().GetEmitterHandle())
+			{
+				FVersionedNiagaraEmitterData* ED = EmitterVMS.Get().GetEmitterHandle()->GetEmitterData();
+				if (ED)
+				{
+					ED->ForEachEnabledRenderer(
+						[&](UNiagaraRendererProperties* RenderProperties)
+						{
+							for (FNiagaraVariableBase BoundAttribute : RenderProperties->GetBoundAttributes())
+							{
+								if (FNiagaraParameterPanelItem* ItemPtr = VisitedParameterToItemMap.Find(BoundAttribute))
+								{
+									// This variable has already been registered, increment the reference count. Otherwise, it is 
+									// not a live binding and we can skip.
+									ItemPtr->ReferenceCount++;
+								}
+							}
+						}
+					);
+				}
+			}
+		}
+
 	}
 	else if (bForceScript)
 	{
