- Change the default to avoid crash if we find no attribute channels on Grid3D

#rnx

diff --git a/Engine/Plugins/FX/Niagara/Source/Niagara/Classes/NiagaraDataInterfaceGrid3DCollection.h b/Engine/Plugins/FX/Niagara/Source/Niagara/Classes/NiagaraDataInterfaceGrid3DCollection.h
--- a/Engine/Plugins/FX/Niagara/Source/Niagara/Classes/NiagaraDataInterfaceGrid3DCollection.h
+++ b/Engine/Plugins/FX/Niagara/Source/Niagara/Classes/NiagaraDataInterfaceGrid3DCollection.h
@@ -53,8 +53,8 @@ struct FGrid3DCollectionRWInstanceData_RenderThread
 {
 	FName SourceDIName;
 
-	FIntVector NumCells = FIntVector::ZeroValue;
-	FIntVector NumTiles = FIntVector::ZeroValue;
+	FIntVector NumCells = FIntVector(1, 1, 1);
+	FIntVector NumTiles = FIntVector(1, 1, 1);
 	int32 TotalNumAttributes = 0;
 	int32 TotalNumNamedAttributes = 0;
 	FVector CellSize = FVector::ZeroVector;
