- Modify Niagara Debug Data Accessor used by DebugHud
- For enum types display the enum text
- For all other types ensure we iterate over all variables, they may not be in order for structs but it's useful data to have

#rnx

diff --git a/Engine/Plugins/FX/Niagara/Source/Niagara/Classes/NiagaraDataSetDebugAccessor.h b/Engine/Plugins/FX/Niagara/Source/Niagara/Classes/NiagaraDataSetDebugAccessor.h
--- a/Engine/Plugins/FX/Niagara/Source/Niagara/Classes/NiagaraDataSetDebugAccessor.h
+++ b/Engine/Plugins/FX/Niagara/Source/Niagara/Classes/NiagaraDataSetDebugAccessor.h
@@ -11,57 +11,62 @@ struct FNiagaraDataSetDebugAccessor
 	bool Init(const FNiagaraDataSetCompiledData& CompiledData, FName InVariableName);
 
 	float ReadFloat(const FNiagaraDataBuffer* DataBuffer, uint32 Instance, uint32 Component) const;
+	float ReadHalf(const FNiagaraDataBuffer* DataBuffer, uint32 Instance, uint32 Component) const;
 	int32 ReadInt(const FNiagaraDataBuffer* DataBuffer, uint32 Instance, uint32 Component) const;
+	
 
 	template<typename TString>
 	void StringAppend(TString& StringType, const FNiagaraDataBuffer* DataBuffer, uint32 Instance) const
 	{
-		if (IsFloat() || IsHalf())
+		if (bIsEnum)
 		{
-			for (uint32 iComponent = 0; iComponent < NumComponents; ++iComponent)
+			const int32 Value = ReadInt(DataBuffer, Instance, 0);
+			StringType.Appendf(TEXT("%d (%s)"), Value, *NiagaraType.GetEnum()->GetDisplayNameTextByValue(Value).ToString());
+		}
+		else if (bIsBool)
+		{
+			const TCHAR* TrueText = TEXT("true");
+			const TCHAR* FalseText = TEXT("false");
+			StringType.Append(ReadInt(DataBuffer, Instance, 0) == FNiagaraBool::True ? TrueText : FalseText);
+		}
+		else
+		{
+			bool bNeedsComma = false;
+			for (uint32 iComponent=0; iComponent < NumComponentsInt32; ++iComponent)
 			{
-				if (iComponent != 0)
+				if (bNeedsComma)
 				{
 					StringType.Append(TEXT(", "));
 				}
-				StringType.Appendf(TEXT("%.2f"), ReadFloat(DataBuffer, Instance, iComponent));
+				bNeedsComma = true;
+				StringType.Appendf(TEXT("%d"), ReadInt(DataBuffer, Instance, iComponent));
 			}
-		}
-		else if (IsInt())
-		{
-			if (NiagaraType.IsSameBaseDefinition(FNiagaraTypeDefinition::GetBoolDef()))
-			{
-				const TCHAR* TrueText = TEXT("true");
-				const TCHAR* FalseText = TEXT("false");
 
-				for (uint32 iComponent = 0; iComponent < NumComponents; ++iComponent)
+			for (uint32 iComponent=0; iComponent < NumComponentsFloat; ++iComponent)
+			{
+				if (bNeedsComma)
 				{
-					if (iComponent != 0)
-					{
-						StringType.Append(TEXT(", "));
-					}
-					StringType.Append(ReadInt(DataBuffer, Instance, iComponent) == FNiagaraBool::True ? TrueText : FalseText);
+					StringType.Append(TEXT(", "));
 				}
+				bNeedsComma = true;
+				StringType.Appendf(TEXT("%.2f"), ReadFloat(DataBuffer, Instance, iComponent));
 			}
-			else
+
+			for (uint32 iComponent=0; iComponent < NumComponentsHalf; ++iComponent)
 			{
-				for (uint32 iComponent = 0; iComponent < NumComponents; ++iComponent)
+				if (bNeedsComma)
 				{
-					if (iComponent != 0)
-					{
-						StringType.Append(TEXT(", "));
-					}
-					StringType.Appendf(TEXT("%d"), ReadInt(DataBuffer, Instance, iComponent));
+					StringType.Append(TEXT(", "));
 				}
+				bNeedsComma = true;
+				StringType.Appendf(TEXT("%.2f"), ReadHalf(DataBuffer, Instance, iComponent));
 			}
 		}
 	}
 
 	FName GetName() const { return VariableName; }
-	bool IsFloat() const { return bIsFloat; }
-	bool IsHalf() const { return bIsHalf; }
-	bool IsInt() const { return bIsInt; }
-	uint32 GetNumComponents() const { return NumComponents; }
+	bool IsEnum() const { return bIsEnum; }
+	bool IsBool() const { return bIsBool; }
 
 	static bool ValidateDataBuffer(const FNiagaraDataSetCompiledData& CompiledData, const FNiagaraDataBuffer* DataBuffer, uint32 iInstance, TFunction<void(const FNiagaraVariable&, int32)> ErrorCallback);
 	static bool ValidateDataBuffer(const FNiagaraDataSetCompiledData& CompiledData, const FNiagaraDataBuffer* DataBuffer, TFunction<void(const FNiagaraVariable&, uint32, int32)> ErrorCallback);
@@ -69,9 +74,12 @@ struct FNiagaraDataSetDebugAccessor
 private:
 	FName					VariableName;
 	FNiagaraTypeDefinition	NiagaraType;
-	bool					bIsFloat = false;
-	bool					bIsHalf = false;
-	bool					bIsInt = false;
-	uint32					NumComponents = 0;
-	uint32					ComponentIndex = INDEX_NONE;
+	bool					bIsEnum = false;
+	bool					bIsBool = false;
+	uint32					NumComponentsInt32 = 0;
+	uint32					NumComponentsFloat = 0;
+	uint32					NumComponentsHalf = 0;
+	uint32					ComponentIndexInt32 = INDEX_NONE;
+	uint32					ComponentIndexFloat = INDEX_NONE;
+	uint32					ComponentIndexHalf = INDEX_NONE;
 };
diff --git a/Engine/Plugins/FX/Niagara/Source/Niagara/Private/NiagaraDataSetDebugAccessor.cpp b/Engine/Plugins/FX/Niagara/Source/Niagara/Private/NiagaraDataSetDebugAccessor.cpp
--- a/Engine/Plugins/FX/Niagara/Source/Niagara/Private/NiagaraDataSetDebugAccessor.cpp
+++ b/Engine/Plugins/FX/Niagara/Source/Niagara/Private/NiagaraDataSetDebugAccessor.cpp
@@ -5,68 +5,64 @@
 bool FNiagaraDataSetDebugAccessor::Init(const FNiagaraDataSetCompiledData& CompiledData, FName InVariableName)
 {
 	VariableName = InVariableName;
-	bIsFloat = false;
-	bIsHalf = false;
-	bIsInt = false;
-	NumComponents = 0;
-	ComponentIndex = INDEX_NONE;
+	bIsEnum = false;
+	bIsBool = false;
+	NumComponentsInt32 = 0;
+	NumComponentsFloat = 0;
+	NumComponentsHalf = 0;
+	ComponentIndexInt32 = INDEX_NONE;
+	ComponentIndexFloat = INDEX_NONE;
+	ComponentIndexHalf = INDEX_NONE;
 
 	for (int32 i = 0; i < CompiledData.Variables.Num(); ++i)
 	{
-		if (CompiledData.Variables[i].GetName() != VariableName)
+		const FNiagaraVariable& Variable = CompiledData.Variables[i];
+		if (Variable.GetName() != VariableName)
 		{
 			continue;
 		}
 
-		if (CompiledData.VariableLayouts[i].GetNumFloatComponents() > 0)
-		{
-			bIsFloat = true;
-			ComponentIndex = CompiledData.VariableLayouts[i].FloatComponentStart;
-			NumComponents = CompiledData.VariableLayouts[i].GetNumFloatComponents();
-			NiagaraType = CompiledData.Variables[i].GetType();
-		}
-		else if (CompiledData.VariableLayouts[i].GetNumHalfComponents() > 0)
-		{
-			bIsHalf = true;
-			ComponentIndex = CompiledData.VariableLayouts[i].HalfComponentStart;
-			NumComponents = CompiledData.VariableLayouts[i].GetNumHalfComponents();
-			NiagaraType = CompiledData.Variables[i].GetType();
-		}
-		else if (CompiledData.VariableLayouts[i].GetNumInt32Components() > 0)
-		{
-			bIsInt = true;
-			ComponentIndex = CompiledData.VariableLayouts[i].Int32ComponentStart;
-			NumComponents = CompiledData.VariableLayouts[i].GetNumInt32Components();
-			NiagaraType = CompiledData.Variables[i].GetType();
-		}
-		return NumComponents > 0;
+		const FNiagaraVariableLayoutInfo& VariableLayout = CompiledData.VariableLayouts[i];
+		NumComponentsInt32 = VariableLayout.GetNumInt32Components();
+		NumComponentsFloat = VariableLayout.GetNumFloatComponents();
+		NumComponentsHalf = VariableLayout.GetNumHalfComponents();
+		ComponentIndexInt32 = VariableLayout.Int32ComponentStart; 
+		ComponentIndexFloat = VariableLayout.FloatComponentStart;
+		ComponentIndexHalf = VariableLayout.HalfComponentStart;
+		NiagaraType = Variable.GetType();
+
+		bIsBool = Variable.GetType().IsSameBaseDefinition(FNiagaraTypeDefinition::GetBoolDef());
+		bIsEnum = Variable.GetType().IsEnum();
+		return NumComponentsInt32 + NumComponentsFloat + NumComponentsHalf > 0;
 	}
 	return false;
 }
 
 float FNiagaraDataSetDebugAccessor::ReadFloat(const FNiagaraDataBuffer* DataBuffer, uint32 Instance, uint32 Component) const
 {
-	if (DataBuffer != nullptr && ComponentIndex != INDEX_NONE && Instance < DataBuffer->GetNumInstances() && Component < NumComponents)
+	if (DataBuffer != nullptr && ComponentIndexInt32 != INDEX_NONE && Instance < DataBuffer->GetNumInstances() && Component < NumComponentsInt32)
 	{
-		if (bIsFloat)
-		{
-			const float* FloatData = reinterpret_cast<const float*>(DataBuffer->GetComponentPtrFloat(ComponentIndex + Component));
-			return FloatData[Instance];
-		}
-		else if (bIsHalf)
-		{
-			const FFloat16* HalfData = reinterpret_cast<const FFloat16*>(DataBuffer->GetComponentPtrHalf(ComponentIndex + Component));
-			return HalfData[Instance];
-		}
+		const float* FloatData = reinterpret_cast<const float*>(DataBuffer->GetComponentPtrFloat(ComponentIndexInt32 + Component));
+		return FloatData[Instance];
+	}
+	return 0.0f;
+}
+
+float FNiagaraDataSetDebugAccessor::ReadHalf(const FNiagaraDataBuffer* DataBuffer, uint32 Instance, uint32 Component) const
+{
+	if (DataBuffer != nullptr && ComponentIndexHalf != INDEX_NONE && Instance < DataBuffer->GetNumInstances() && Component < NumComponentsHalf)
+	{
+		const FFloat16* HalfData = reinterpret_cast<const FFloat16*>(DataBuffer->GetComponentPtrHalf(NumComponentsHalf + Component));
+		return HalfData[Instance];
 	}
 	return 0.0f;
 }
 
 int32 FNiagaraDataSetDebugAccessor::ReadInt(const FNiagaraDataBuffer* DataBuffer, uint32 Instance, uint32 Component) const
 {
-	if (DataBuffer != nullptr && ComponentIndex != INDEX_NONE && Instance < DataBuffer->GetNumInstances() && Component < NumComponents)
+	if (DataBuffer != nullptr && ComponentIndexInt32 != INDEX_NONE && Instance < DataBuffer->GetNumInstances() && Component < NumComponentsInt32)
 	{
-		const int32* IntData = reinterpret_cast<const int32*>(DataBuffer->GetComponentPtrInt32(ComponentIndex + Component));
+		const int32* IntData = reinterpret_cast<const int32*>(DataBuffer->GetComponentPtrInt32(ComponentIndexInt32 + Component));
 		return IntData[Instance];
 	}
 	return 0;
