Add symbol for natvis to NiagaraEditor to ensure that natvis works when the NiagaraEditor module is on the stack

diff --git a/Engine/Plugins/FX/Niagara/Source/NiagaraEditor/Private/NiagaraEditorModule.cpp b/Engine/Plugins/FX/Niagara/Source/NiagaraEditor/Private/NiagaraEditorModule.cpp
--- a/Engine/Plugins/FX/Niagara/Source/NiagaraEditor/Private/NiagaraEditorModule.cpp
+++ b/Engine/Plugins/FX/Niagara/Source/NiagaraEditor/Private/NiagaraEditorModule.cpp
@@ -147,8 +147,8 @@
 
 #include "SNiagaraDebugger.h"
 
+#include "NiagaraDebugVis.h"
 #include "NiagaraPerfBaseline.h"
-
 #include "NiagaraGraphDataCache.h"
 #include "Misc/ScopedSlowTask.h"
 
@@ -179,6 +179,13 @@ static FAutoConsoleVariableRef CVarPreloadSelectablePluginAssetsOnDemand(
 // this is required for gpu script compilation ticks
 static FNiagaraShaderQueueTickable NiagaraShaderQueueProcessor;
 
+#if !IS_MONOLITHIC
+namespace NiagaraDebugVisHelper
+{
+	const FNiagaraTypeRegistry*& GTypeRegistrySingletonPtr = GCoreTypeRegistrySingletonPtr;
+}
+#endif
+
 //////////////////////////////////////////////////////////////////////////
 
 class FNiagaraEditorOnlyDataUtilities : public INiagaraEditorOnlyDataUtilities
