Stage 1 of de-objectifying UNiagaraScriptVariable
with EngineTest and a cold DDC, this reduces teh number of UNiagaraScriptVariable instances created by ~30% (260k -> 185k).  No change on a hot ddc.

diff --git a/Engine/Plugins/FX/Niagara/Source/NiagaraEditor/Private/NiagaraAttributeTrimmer.cpp b/Engine/Plugins/FX/Niagara/Source/NiagaraEditor/Private/NiagaraAttributeTrimmer.cpp
--- a/Engine/Plugins/FX/Niagara/Source/NiagaraEditor/Private/NiagaraAttributeTrimmer.cpp
+++ b/Engine/Plugins/FX/Niagara/Source/NiagaraEditor/Private/NiagaraAttributeTrimmer.cpp
@@ -377,27 +377,27 @@ namespace NiagaraAttributeTrimming
 	{
 		if (const UNiagaraGraph* OwnerGraph = Cast<const UNiagaraGraph>(ReadPin.Pin->GetOwningNode()->GetGraph()))
 		{
-			FString ParamName = ParamMap.Variables[VariableIndex].GetName().ToString();
+			FNiagaraVariable ParamVariable = ParamMap.Variables[VariableIndex];
+			FString ParamName = ParamVariable.GetName().ToString();
 
 			// now we need to replace the module name with 'module'
 			ParamName.ReplaceInline(*ReadPin.ModuleName.ToString(), *FNiagaraConstants::ModuleNamespace.ToString());
 
-			const FName VariableName = *ParamName;
+			ParamVariable.SetName(*ParamName);
 
-			if (const UNiagaraScriptVariable* ScriptVariable = OwnerGraph->GetScriptVariable(VariableName))
+			FNiagaraScriptVariableBinding DefaultBinding;
+			TOptional<ENiagaraDefaultMode> DefaultMode = OwnerGraph->GetDefaultMode(ParamVariable, &DefaultBinding);
+			if (DefaultMode.IsSet() && *DefaultMode == ENiagaraDefaultMode::Binding && DefaultBinding.IsValid())
 			{
-				if (ScriptVariable->DefaultMode == ENiagaraDefaultMode::Binding && ScriptVariable->DefaultBinding.IsValid())
-				{
-					int32 BoundVariableIndex = ParamMap.FindVariableByName(ScriptVariable->DefaultBinding.GetName());
+				int32 BoundVariableIndex = ParamMap.FindVariableByName(DefaultBinding.GetName());
 
-					if (BoundVariableIndex != INDEX_NONE)
+				if (BoundVariableIndex != INDEX_NONE)
+				{
+					for (const FNiagaraParameterMapHistory::FReadHistory& ReadHistoryEntry : ParamMap.PerVariableReadHistory[BoundVariableIndex])
 					{
-						for (const FNiagaraParameterMapHistory::FReadHistory& ReadHistoryEntry : ParamMap.PerVariableReadHistory[BoundVariableIndex])
-						{
-							// for now adding the first pin to the dependencies works, but this needs to be revisited
-							ResolvedDependencies.Pins.Add(ReadHistoryEntry.ReadPin);
-							break;
-						}
+						// for now adding the first pin to the dependencies works, but this needs to be revisited
+						ResolvedDependencies.Pins.Add(ReadHistoryEntry.ReadPin);
+						break;
 					}
 				}
 			}
diff --git a/Engine/Plugins/FX/Niagara/Source/NiagaraEditor/Private/NiagaraGraph.cpp b/Engine/Plugins/FX/Niagara/Source/NiagaraEditor/Private/NiagaraGraph.cpp
--- a/Engine/Plugins/FX/Niagara/Source/NiagaraEditor/Private/NiagaraGraph.cpp
+++ b/Engine/Plugins/FX/Niagara/Source/NiagaraEditor/Private/NiagaraGraph.cpp
@@ -2,6 +2,8 @@
 
 #include "NiagaraGraph.h"
 
+#include "Algo/RemoveIf.h"
+#include "Algo/StableSort.h"
 #include "EdGraphSchema_Niagara.h"
 #include "GraphEditAction.h"
 #include "INiagaraEditorTypeUtilities.h"
@@ -138,6 +140,8 @@ void UNiagaraGraph::NotifyGraphChanged()
 
 void UNiagaraGraph::PostLoad()
 {
+	check(!bIsForCompilationOnly);
+
 	Super::PostLoad();
 
 	NIAGARA_SCOPE_CYCLE_COUNTER(STAT_NiagaraEditor_Graph_PostLoad);
@@ -470,6 +474,8 @@ void UNiagaraGraph::PostLoad()
 
 void UNiagaraGraph::ChangeParameterType(const TArray<FNiagaraVariable>& ParametersToChange, const FNiagaraTypeDefinition& NewType, bool bAllowOrphanedPins)
 {
+	check(!bIsForCompilationOnly);
+
 	const UEdGraphSchema_Niagara* NiagaraSchema = GetDefault<UEdGraphSchema_Niagara>();
 	TArray<UNiagaraNodeParameterMapBase*> NiagaraParameterNodes;
 	GetNodesOfClass<UNiagaraNodeParameterMapBase>(NiagaraParameterNodes);
@@ -769,6 +775,8 @@ TArray<UEdGraphPin*> UNiagaraGraph::FindParameterMapDefaultValuePins(const FName
 
 void UNiagaraGraph::ValidateDefaultPins()
 {
+	check(!bIsForCompilationOnly);
+
 	FNiagaraEditorModule& EditorModule = FNiagaraEditorModule::Get();
 
 	if (!VariableToScriptVariable.IsEmpty())
@@ -1093,6 +1101,8 @@ FName UNiagaraGraph::StandardizeName(FName Name, ENiagaraScriptUsage Usage, bool
 
 void UNiagaraGraph::StandardizeParameterNames()
 {
+	check(!bIsForCompilationOnly);
+
 	TMap<FName, FName> OldNameToStandardizedNameMap;
 	bool bAnyNamesModified = false;
 	const UEdGraphSchema_Niagara* NiagaraSchema = GetDefault<UEdGraphSchema_Niagara>();
@@ -1213,6 +1223,26 @@ void UNiagaraGraph::StandardizeParameterNames()
 	}
 }
 
+bool UNiagaraGraph::VariableLess(const FNiagaraVariable& Lhs, const FNiagaraVariable& Rhs)
+{
+	return Lhs.GetName().LexicalLess(Rhs.GetName());
+}
+
+TOptional<FNiagaraScriptVariableData> UNiagaraGraph::GetScriptVariableData(const FNiagaraVariable& Variable) const
+{
+	// for now this is only supported on the compilation copy
+	if (bIsForCompilationOnly)
+	{
+		const int32 ScriptVariableIndex = Algo::BinarySearchBy(CompilationScriptVariables, Variable, &FNiagaraScriptVariableData::Variable, VariableLess);
+		if (CompilationScriptVariables.IsValidIndex(ScriptVariableIndex))
+		{
+			return CompilationScriptVariables[ScriptVariableIndex];
+		}
+	}
+
+	return TOptional<FNiagaraScriptVariableData>();
+}
+
 void UNiagaraGraph::PostEditChangeProperty(FPropertyChangedEvent& PropertyChangedEvent)
 {
 	NotifyGraphChanged();
@@ -1450,11 +1480,18 @@ UNiagaraGraph* UNiagaraGraph::CreateCompilationCopy(const TArray<ENiagaraScriptU
 	}
 	Result->bIsForCompilationOnly = true;
 
-	// get new script variables from the pool
-	for (auto& It : Result->VariableToScriptVariable)
+	Result->CompilationScriptVariables.Reserve(Result->VariableToScriptVariable.Num());
+	for (const auto& It : Result->VariableToScriptVariable)
 	{
-		It.Value = CastChecked<UNiagaraScriptVariable>(NiagaraEditorModule.GetPooledDuplicateObject(It.Value));
+		if (const UNiagaraScriptVariable* ScriptVariable = It.Value)
+		{
+			FNiagaraScriptVariableData& VariableData = Result->CompilationScriptVariables.AddDefaulted_GetRef();
+			VariableData.InitFrom(*ScriptVariable, false);
+		}
 	}
+	Algo::StableSortBy(Result->CompilationScriptVariables, &FNiagaraScriptVariableData::Variable, VariableLess);
+
+	Result->VariableToScriptVariable.Empty();
 
 	// Only add references to nodes which match a compile usage.
 	TMap<UEdGraphNode*, UEdGraphNode*> DuplicationMapping;
@@ -1560,14 +1597,6 @@ void UNiagaraGraph::ReleaseCompilationCopy()
 	{
 		return;
 	}
-	FNiagaraEditorModule& NiagaraEditorModule = FModuleManager::GetModuleChecked<FNiagaraEditorModule>("NiagaraEditor");
-	for (auto It : VariableToScriptVariable)
-	{
-		NiagaraEditorModule.ReleaseObjectToPool(It.Value);
-	}
-
-	// clear script variables and kill this object to surface any invalid access after it's released
-	VariableToScriptVariable.Empty();
 	MarkAsGarbage();
 }
 
@@ -1860,11 +1889,15 @@ void UNiagaraGraph::GetParameters(TArray<FNiagaraVariable>& Inputs, TArray<FNiag
 
 const TMap<FNiagaraVariable, TObjectPtr<UNiagaraScriptVariable>>& UNiagaraGraph::GetAllMetaData() const
 {
+	check(!bIsForCompilationOnly);
+
 	return VariableToScriptVariable;
 }
 
 TMap<FNiagaraVariable, TObjectPtr<UNiagaraScriptVariable>>& UNiagaraGraph::GetAllMetaData()
 {
+	check(!bIsForCompilationOnly);
+
 	return VariableToScriptVariable;
 }
 
@@ -1877,8 +1910,10 @@ const TMap<FNiagaraVariable, FNiagaraGraphParameterReferenceCollection>& UNiagar
 	return ParameterToReferencesMap;
 }
 
-UNiagaraScriptVariable* UNiagaraGraph::GetScriptVariable(FNiagaraVariable Parameter, bool bUpdateIfPending) const
+UNiagaraScriptVariable* UNiagaraGraph::GetScriptVariable(FNiagaraVariable Parameter, bool bUpdateIfPending)
 {
+	check(!bIsForCompilationOnly);
+
 	if (bUpdateIfPending && bParameterReferenceRefreshPending)
 	{
 		RefreshParameterReferences();
@@ -1890,8 +1925,10 @@ UNiagaraScriptVariable* UNiagaraGraph::GetScriptVariable(FNiagaraVariable Parame
 	return nullptr;
 }
 
-UNiagaraScriptVariable* UNiagaraGraph::GetScriptVariable(FName ParameterName, bool bUpdateIfPending) const
+UNiagaraScriptVariable* UNiagaraGraph::GetScriptVariable(FName ParameterName, bool bUpdateIfPending)
 {
+	check(!bIsForCompilationOnly);
+
 	if (bUpdateIfPending && bParameterReferenceRefreshPending)
 	{
 		RefreshParameterReferences();
@@ -1906,8 +1943,35 @@ UNiagaraScriptVariable* UNiagaraGraph::GetScriptVariable(FName ParameterName, bo
 	return nullptr;
 }
 
+UNiagaraScriptVariable* UNiagaraGraph::GetScriptVariable(FNiagaraVariable Parameter) const
+{
+	check(!bIsForCompilationOnly);
+
+	if (TObjectPtr<UNiagaraScriptVariable>* FoundScriptVariable = VariableToScriptVariable.Find(Parameter))
+	{
+		return *FoundScriptVariable;
+	}
+	return nullptr;
+}
+
+UNiagaraScriptVariable* UNiagaraGraph::GetScriptVariable(FName ParameterName) const
+{
+	check(!bIsForCompilationOnly);
+
+	for (auto& VariableToScriptVariableItem : VariableToScriptVariable)
+	{
+		if (VariableToScriptVariableItem.Key.GetName() == ParameterName)
+		{
+			return VariableToScriptVariableItem.Value;
+		}
+	}
+	return nullptr;
+}
+
 UNiagaraScriptVariable* UNiagaraGraph::AddParameter(const FNiagaraVariable& Parameter, bool bIsStaticSwitch /*= false*/)
 {
+	check(!bIsForCompilationOnly);
+
 	// Delay the NotifyGraphChanged() call until the static switch flag is set on the UNiagaraScriptVariable so that ParameterPanel displays correctly.
 	const bool bNotifyChanged = false;
 	UNiagaraScriptVariable* NewScriptVar = AddParameter(Parameter, FNiagaraVariableMetaData(), bIsStaticSwitch, bNotifyChanged);
@@ -1917,6 +1981,8 @@ UNiagaraScriptVariable* UNiagaraGraph::AddParameter(const FNiagaraVariable& Para
 
 UNiagaraScriptVariable* UNiagaraGraph::AddParameter(const FNiagaraVariable& Parameter, const FNiagaraVariableMetaData& ParameterMetaData, bool bIsStaticSwitch, bool bNotifyChanged)
 {
+	check(!bIsForCompilationOnly);
+
 	FNiagaraGraphParameterReferenceCollection* FoundParameterReferenceCollection = ParameterToReferencesMap.Find(Parameter);
 	if (!FoundParameterReferenceCollection)
 	{
@@ -1959,6 +2025,8 @@ UNiagaraScriptVariable* UNiagaraGraph::AddParameter(const FNiagaraVariable& Para
 
 UNiagaraScriptVariable* UNiagaraGraph::AddParameter(const UNiagaraScriptVariable* InScriptVar)
 {
+	check(!bIsForCompilationOnly);
+
 	TObjectPtr<UNiagaraScriptVariable>* FoundScriptVariable = VariableToScriptVariable.Find(InScriptVar->Variable);
 	if (!FoundScriptVariable)
 	{
@@ -2018,6 +2086,8 @@ void UNiagaraGraph::AddParameterReference(const FNiagaraVariable& Parameter, FNi
 
 void UNiagaraGraph::RemoveParameter(const FNiagaraVariable& Parameter, bool bAllowDeleteStaticSwitch /*= false*/)
 {
+	check(!bIsForCompilationOnly);
+
 	FNiagaraGraphParameterReferenceCollection* ReferenceCollection = ParameterToReferencesMap.Find(Parameter);
 	if (ReferenceCollection)
 	{
@@ -2075,6 +2145,8 @@ void CopyScriptVariableDataForRename(const UNiagaraScriptVariable& OldScriptVari
 
 bool UNiagaraGraph::RenameParameterFromPin(const FNiagaraVariable& Parameter, FName NewName, UEdGraphPin* InPin)
 {
+	check(!bIsForCompilationOnly);
+
 	auto TrySubscribeScriptVarToDefinitionByName = [this](UNiagaraScriptVariable* TargetScriptVar) {
 		UNiagaraScript* OuterScript = GetTypedOuter<UNiagaraScript>();
 		bool bForDataProcessingOnly = true;
@@ -2185,6 +2257,8 @@ bool UNiagaraGraph::RenameParameterFromPin(const FNiagaraVariable& Parameter, FN
 
 bool UNiagaraGraph::RenameParameter(const FNiagaraVariable& Parameter, FName NewName, bool bRenameRequestedFromStaticSwitch, bool* bMerged, bool bSuppressEvents)
 {
+	check(!bIsForCompilationOnly);
+
 	// Initialize the merger state if requested
 	if (bMerged)
 		*bMerged = false;
@@ -2313,6 +2387,8 @@ bool UNiagaraGraph::RenameParameter(const FNiagaraVariable& Parameter, FName New
 
 void UNiagaraGraph::ScriptVariableChanged(FNiagaraVariable Variable)
 {
+	check(!bIsForCompilationOnly);
+
 	TObjectPtr<UNiagaraScriptVariable>* ScriptVariable = GetAllMetaData().Find(Variable);
 	if (!ScriptVariable || !*ScriptVariable || (*ScriptVariable)->GetIsStaticSwitch())
 	{
@@ -2376,6 +2452,8 @@ FVersionedNiagaraEmitter UNiagaraGraph::GetOwningEmitter() const
 
 bool UNiagaraGraph::SynchronizeScriptVariable(const UNiagaraScriptVariable* SourceScriptVar, UNiagaraScriptVariable* DestScriptVar /*= nullptr*/, bool bIgnoreChangeId /*= false*/)
 {
+	check(!bIsForCompilationOnly);
+
 	if (DestScriptVar == nullptr)
 	{
 		const FGuid& SourceScriptVarId = SourceScriptVar->Metadata.GetVariableGuid();
@@ -2437,6 +2515,8 @@ bool UNiagaraGraph::SynchronizeScriptVariable(const UNiagaraScriptVariable* Sour
 
 bool UNiagaraGraph::SynchronizeParameterDefinitionsScriptVariableRemoved(const FGuid RemovedScriptVarId)
 {
+	check(!bIsForCompilationOnly);
+
 	TArray<TObjectPtr<UNiagaraScriptVariable>> ScriptVariables;
 	VariableToScriptVariable.GenerateValueArray(ScriptVariables);
 	for (UNiagaraScriptVariable* ScriptVar : ScriptVariables)
@@ -2458,6 +2538,8 @@ void UNiagaraGraph::SynchronizeParametersWithParameterDefinitions(
 	INiagaraParameterDefinitionsSubscriber* Subscriber,
 	FSynchronizeWithParameterDefinitionsArgs Args)
 {
+	check(!bIsForCompilationOnly);
+
 	bool bMarkRequiresSync = false;
 	TArray<TObjectPtr<UNiagaraScriptVariable>> ScriptVariables;
 	TArray<UNiagaraScriptVariable*> TargetScriptVariables;
@@ -2645,6 +2727,15 @@ bool UNiagaraGraph::HasParameterMapParameters()const
 
 bool UNiagaraGraph::GetPropertyMetadata(FName PropertyName, FString& OutValue) const
 {
+	for (const FNiagaraScriptVariableData& ScriptVariableData : CompilationScriptVariables)
+	{
+		if (const FString* FoundMatch = ScriptVariableData.Metadata.PropertyMetaData.Find(PropertyName))
+		{
+			OutValue = *FoundMatch;
+			return true;
+		}
+	}
+
 	const TMap<FNiagaraVariable, TObjectPtr<UNiagaraScriptVariable>>& MetaDataMap = GetAllMetaData();
 	auto Iter = MetaDataMap.CreateConstIterator();
 	while (Iter)
@@ -2728,6 +2819,8 @@ FNiagaraTypeDefinition UNiagaraGraph::GetCachedNumericConversion(class UEdGraphP
 
 bool UNiagaraGraph::AppendCompileHash(FNiagaraCompileHashVisitor* InVisitor, const TArray<UNiagaraNode*>& InTraversal) const
 {
+	check(!bIsForCompilationOnly);
+
 #if WITH_EDITORONLY_DATA
 	int32 Index = InVisitor->Values.AddDefaulted();
 	InVisitor->Values[Index].Object = FString::Printf(TEXT("Class: \"%s\"  Name: \"%s\""), *this->GetClass()->GetName(),  *this->GetName());
@@ -2917,6 +3010,8 @@ void DiffProperties(const FNiagaraCompileHashVisitorDebugInfo& A, const FNiagara
 
 void UNiagaraGraph::RebuildCachedCompileIds(bool bForce)
 {
+	check(!bIsForCompilationOnly);
+
 	// If the graph hasn't changed since last rebuild, then do nothing.
 	if (!bForce && ChangeId == LastBuiltTraversalDataChangeId && LastBuiltTraversalDataChangeId.IsValid())
 	{
@@ -3286,6 +3381,11 @@ void UNiagaraGraph::MarkGraphRequiresSynchronization(FString Reason)
 
 TOptional<FNiagaraVariableMetaData> UNiagaraGraph::GetMetaData(const FNiagaraVariable& InVar) const
 {
+	if (TOptional<FNiagaraScriptVariableData> VariableData = GetScriptVariableData(InVar))
+	{
+		return VariableData->Metadata;
+	}
+
 	if (TObjectPtr<UNiagaraScriptVariable>* MetaData = VariableToScriptVariable.Find(InVar))
 	{
 		if (*MetaData)
@@ -3298,6 +3398,8 @@ TOptional<FNiagaraVariableMetaData> UNiagaraGraph::GetMetaData(const FNiagaraVar
 
 void UNiagaraGraph::SetMetaData(const FNiagaraVariable& InVar, const FNiagaraVariableMetaData& InMetaData)
 {
+	check(!bIsForCompilationOnly);
+
 	if (TObjectPtr<UNiagaraScriptVariable>* FoundMetaData = VariableToScriptVariable.Find(InVar))
 	{
 		if (*FoundMetaData)
@@ -3331,7 +3433,7 @@ UNiagaraGraph::FOnSubObjectSelectionChanged& UNiagaraGraph::OnSubObjectSelection
 	return OnSelectedSubObjectChanged;
 }
 
-void UNiagaraGraph::RefreshParameterReferences() const 
+void UNiagaraGraph::RefreshParameterReferences() const
 {
 	// A set of variables to track which parameters are used so that unused parameters can be removed after the reference tracking.
 	TSet<FNiagaraVariable> CandidateUnreferencedParametersToRemove;
@@ -3450,46 +3552,77 @@ void UNiagaraGraph::RefreshParameterReferences() const
 		}
 	}
 
-	// Add reference to all variables that are default bound to
-	for (auto It = VariableToScriptVariable.CreateConstIterator(); It; ++It)
+	if (bIsForCompilationOnly)
 	{
-		UNiagaraScriptVariable* Variable = It.Value();
-		if (!Variable || (Variable->DefaultMode != ENiagaraDefaultMode::Binding || !Variable->DefaultBinding.IsValid()))
+		// Add reference to all variables that are default bound to
+		for (const FNiagaraScriptVariableData& Variable : CompilationScriptVariables)
 		{
-			continue;
+			FNiagaraVariable LinkedVariable(Variable.Variable.GetType(), Variable.DefaultBinding.GetName());
+
+			if ((Variable.DefaultMode != ENiagaraDefaultMode::Binding || !Variable.DefaultBinding.IsValid()))
+			{
+				continue;
+			}
+
+			if (FNiagaraConstants::GetOldPositionTypeVariables().Contains(LinkedVariable))
+			{
+				// it is not uncommon that old assets have vector inputs that default bind to what is now a position type. If we detect that, we change the type to prevent a compiler error.
+				LinkedVariable.SetType(FNiagaraTypeDefinition::GetPositionDef());
+			}
+
+			AddBindingParameterReference(LinkedVariable);
 		}
 
-		FNiagaraTypeDefinition LinkedType = Variable->Variable.GetType();
-		FName BindingName = Variable->DefaultBinding.GetName();
-		if (FNiagaraConstants::GetOldPositionTypeVariables().Contains(FNiagaraVariable(LinkedType, BindingName)))
+		check(CandidateUnreferencedParametersToRemove.IsEmpty());
+
+		CompilationScriptVariables.SetNum(Algo::StableRemoveIf(CompilationScriptVariables, [&](const FNiagaraScriptVariableData& VariableData)
 		{
-			// it is not uncommon that old assets have vector inputs that default bind to what is now a position type. If we detect that, we change the type to prevent a compiler error.
-			LinkedType = FNiagaraTypeDefinition::GetPositionDef();
-		}
-		AddBindingParameterReference(FNiagaraVariable(LinkedType, BindingName));
+			return !ParameterToReferencesMap.Contains(VariableData.Variable);
+		}));
 	}
-
-	// If there were any previous parameters which didn't have any references added, remove them here.
-	for (const FNiagaraVariable& UnreferencedParameterToRemove : CandidateUnreferencedParametersToRemove)
+	else
 	{
-		ParameterToReferencesMap.Remove(UnreferencedParameterToRemove);
-		VariableToScriptVariable.Remove(UnreferencedParameterToRemove);
-	}
+		// Add reference to all variables that are default bound to
+		for (auto It = VariableToScriptVariable.CreateConstIterator(); It; ++It)
+		{
+			UNiagaraScriptVariable* Variable = It.Value();
+			if (!Variable || (Variable->DefaultMode != ENiagaraDefaultMode::Binding || !Variable->DefaultBinding.IsValid()))
+			{
+				continue;
+			}
 
-	// Remove any script variables 
-	TArray<FNiagaraVariable> UnreferencedScriptVariables;
-	for (auto It : VariableToScriptVariable)
-	{
-		FNiagaraGraphParameterReferenceCollection* ReferenceCollection = ParameterToReferencesMap.Find(It.Key);
-		if (ReferenceCollection == nullptr)
+			FNiagaraTypeDefinition LinkedType = Variable->Variable.GetType();
+			FName BindingName = Variable->DefaultBinding.GetName();
+			if (FNiagaraConstants::GetOldPositionTypeVariables().Contains(FNiagaraVariable(LinkedType, BindingName)))
+			{
+				// it is not uncommon that old assets have vector inputs that default bind to what is now a position type. If we detect that, we change the type to prevent a compiler error.
+				LinkedType = FNiagaraTypeDefinition::GetPositionDef();
+			}
+			AddBindingParameterReference(FNiagaraVariable(LinkedType, BindingName));
+		}
+
+		// If there were any previous parameters which didn't have any references added, remove them here.
+		for (const FNiagaraVariable& UnreferencedParameterToRemove : CandidateUnreferencedParametersToRemove)
 		{
-			UnreferencedScriptVariables.Add(It.Key);
+			ParameterToReferencesMap.Remove(UnreferencedParameterToRemove);
+			VariableToScriptVariable.Remove(UnreferencedParameterToRemove);
 		}
-	}
 
-	for (auto Variable : UnreferencedScriptVariables)
-	{
-		VariableToScriptVariable.Remove(Variable);
+		// Remove any script variables 
+		TArray<FNiagaraVariable> UnreferencedScriptVariables;
+		for (auto It : VariableToScriptVariable)
+		{
+			FNiagaraGraphParameterReferenceCollection* ReferenceCollection = ParameterToReferencesMap.Find(It.Key);
+			if (ReferenceCollection == nullptr)
+			{
+				UnreferencedScriptVariables.Add(It.Key);
+			}
+		}
+
+		for (auto Variable : UnreferencedScriptVariables)
+		{
+			VariableToScriptVariable.Remove(Variable);
+		}
 	}
 
 	bParameterReferenceRefreshPending = false;
@@ -3549,5 +3682,151 @@ const TMap<FNiagaraVariable, FInputPinsAndOutputPins> UNiagaraGraph::CollectVars
 	return VarToPinsMap;
 }
 
+void UNiagaraGraph::GetAllScriptVariableGuids(TArray<FGuid>& VariableGuids) const
+{
+	if (bIsForCompilationOnly)
+	{
+		VariableGuids.Reserve(CompilationScriptVariables.Num());
+		Algo::Transform(CompilationScriptVariables, VariableGuids, [&](const FNiagaraScriptVariableData& VariableData) { return VariableData.Metadata.GetVariableGuid(); });
+		return;
+	}
+
+	for (auto It = VariableToScriptVariable.CreateConstIterator(); It; ++It)
+	{
+		VariableGuids.Add(It.Value()->Metadata.GetVariableGuid());
+	}
+}
+
+void UNiagaraGraph::GetAllVariables(TArray<FNiagaraVariable>& Variables) const
+{
+	if (bIsForCompilationOnly)
+	{
+		Variables.Reserve(CompilationScriptVariables.Num());
+		Algo::Transform(CompilationScriptVariables, Variables, [&](const FNiagaraScriptVariableData& VariableData) { return VariableData.Variable; });
+		return;
+	}
+
+	VariableToScriptVariable.GetKeys(Variables);
+}
+
+bool UNiagaraGraph::HasVariable(const FNiagaraVariable& Variable) const
+{
+	if (TOptional<FNiagaraScriptVariableData> VariableData = GetScriptVariableData(Variable))
+	{
+		return true;
+	}
+
+	return VariableToScriptVariable.Contains(Variable);
+}
+
+TOptional<bool> UNiagaraGraph::IsStaticSwitch(const FNiagaraVariable& Variable) const
+{
+	if (TOptional<FNiagaraScriptVariableData> VariableData = GetScriptVariableData(Variable))
+	{
+		return VariableData->GetIsStaticSwitch();
+	}
+
+	if (TObjectPtr<UNiagaraScriptVariable>* FoundScriptVariable = VariableToScriptVariable.Find(Variable))
+	{
+		return (*FoundScriptVariable)->GetIsStaticSwitch();
+	}
+
+	return TOptional<bool>();
+}
+
+TOptional<int> UNiagaraGraph::GetStaticSwitchDefaultValue(const FNiagaraVariable& Variable) const
+{
+	if (TOptional<FNiagaraScriptVariableData> VariableData = GetScriptVariableData(Variable))
+	{
+		return VariableData->GetStaticSwitchDefaultValue();
+	}
+
+	if (TObjectPtr<UNiagaraScriptVariable>* FoundScriptVariable = VariableToScriptVariable.Find(Variable))
+	{
+		return (*FoundScriptVariable)->GetStaticSwitchDefaultValue();
+	}
+
+	return TOptional<int32>();
+}
+
+TOptional<ENiagaraDefaultMode> UNiagaraGraph::GetDefaultMode(const FNiagaraVariable& Variable, FNiagaraScriptVariableBinding* Binding) const
+{
+	if (TOptional<FNiagaraScriptVariableData> VariableData = GetScriptVariableData(Variable))
+	{
+		if (VariableData->DefaultMode == ENiagaraDefaultMode::Binding && VariableData->DefaultBinding.IsValid())
+		{
+			if (Binding)
+			{
+				*Binding = VariableData->DefaultBinding;
+			}
+			return { ENiagaraDefaultMode::Binding };
+		}
+		else
+		{
+			return VariableData->DefaultMode;
+		}
+	}
+
+	if (TObjectPtr<UNiagaraScriptVariable>* FoundScriptVariable = VariableToScriptVariable.Find(Variable))
+	{
+		const UNiagaraScriptVariable* ScriptVariable = *FoundScriptVariable;
+		if (ScriptVariable->DefaultMode == ENiagaraDefaultMode::Binding && ScriptVariable->DefaultBinding.IsValid())
+		{
+			if (Binding)
+			{
+				*Binding = ScriptVariable->DefaultBinding;
+			}
+			return {ENiagaraDefaultMode::Binding};
+		}
+		else
+		{
+			return ScriptVariable->DefaultMode;
+		}
+	}
+
+	return TOptional<ENiagaraDefaultMode>();
+}
+
+TOptional<FGuid> UNiagaraGraph::GetScriptVariableGuid(const FNiagaraVariable& Variable) const
+{
+	if (TOptional<FNiagaraScriptVariableData> VariableData = GetScriptVariableData(Variable))
+	{
+		return VariableData->Metadata.GetVariableGuid();
+	}
+
+	if (TObjectPtr<UNiagaraScriptVariable>* FoundScriptVariable = VariableToScriptVariable.Find(Variable))
+	{
+		return (*FoundScriptVariable)->Metadata.GetVariableGuid();
+	}
+
+	return TOptional<FGuid>();
+}
+
+TOptional<FNiagaraVariable> UNiagaraGraph::GetVariable(const FNiagaraVariable& Variable) const
+{
+	if (TOptional<FNiagaraScriptVariableData> VariableData = GetScriptVariableData(Variable))
+	{
+		return VariableData->Variable;
+	}
+
+	if (TObjectPtr<UNiagaraScriptVariable>* FoundScriptVariable = VariableToScriptVariable.Find(Variable))
+	{
+		return (*FoundScriptVariable)->Variable;
+	}
+
+	return TOptional<FNiagaraVariable>();
+}
+
+void UNiagaraGraph::SetIsStaticSwitch(const FNiagaraVariable& Variable, bool InValue)
+{
+	check(!bIsForCompilationOnly);
+
+	if (TObjectPtr<UNiagaraScriptVariable>* FoundScriptVariable = VariableToScriptVariable.Find(Variable))
+	{
+		(*FoundScriptVariable)->SetIsStaticSwitch(InValue);
+	}
+}
+
+
 #undef NIAGARA_SCOPE_CYCLE_COUNTER
 #undef LOCTEXT_NAMESPACE
diff --git a/Engine/Plugins/FX/Niagara/Source/NiagaraEditor/Private/NiagaraHlslTranslator.cpp b/Engine/Plugins/FX/Niagara/Source/NiagaraEditor/Private/NiagaraHlslTranslator.cpp
--- a/Engine/Plugins/FX/Niagara/Source/NiagaraEditor/Private/NiagaraHlslTranslator.cpp
+++ b/Engine/Plugins/FX/Niagara/Source/NiagaraEditor/Private/NiagaraHlslTranslator.cpp
@@ -636,16 +636,18 @@ void FHlslNiagaraTranslator::BuildMissingDefaults()
 			bool bWriteToParamMapEntries = UniqueVarToWriteToParamMap.FindChecked(Var);
 			int32 OutputChunkId = INDEX_NONE;
 
-			UNiagaraScriptVariable* ScriptVariable = nullptr;
+			TOptional<ENiagaraDefaultMode> DefaultMode;
+			FNiagaraScriptVariableBinding DefaultBinding;
+
 			if (DefaultPin) 
 			{
 				if (UNiagaraGraph* DefaultPinGraph = CastChecked<UNiagaraGraph>(DefaultPin->GetOwningNode()->GetGraph())) 
 				{
-					ScriptVariable = DefaultPinGraph->GetScriptVariable(Var);
+					DefaultMode = DefaultPinGraph->GetDefaultMode(Var, &DefaultBinding);
 				}
 			}
 
-			HandleParameterRead(ActiveStageIdx, Var, DefaultPin, DefaultPin != nullptr ? Cast<UNiagaraNode>(DefaultPin->GetOwningNode()) : nullptr, OutputChunkId, ScriptVariable, !bWriteToParamMapEntries, true);
+			HandleParameterRead(ActiveStageIdx, Var, DefaultPin, DefaultPin != nullptr ? Cast<UNiagaraNode>(DefaultPin->GetOwningNode()) : nullptr, OutputChunkId, DefaultMode, DefaultBinding, !bWriteToParamMapEntries, true);
 		}
 
 		DeferredVariablesMissingDefault.Empty();
@@ -4912,19 +4914,20 @@ void FHlslNiagaraTranslator::InitializeParameterMapDefaults(int32 ParamMapHistor
 			bool bWriteToParamMapEntries = UniqueVarToWriteToParamMap.FindChecked(Var);
 			int32 OutputChunkId = INDEX_NONE;
 			
-			UNiagaraScriptVariable* ScriptVariable = nullptr;
+			TOptional<ENiagaraDefaultMode> DefaultMode;
+			FNiagaraScriptVariableBinding DefaultBinding;
 			if (DefaultPin) 
 			{
 				if (UNiagaraGraph* DefaultPinGraph = CastChecked<UNiagaraGraph>(DefaultPin->GetOwningNode()->GetGraph())) 
 				{
-					ScriptVariable = DefaultPinGraph->GetScriptVariable(Var);
+					DefaultMode = DefaultPinGraph->GetDefaultMode(Var, &DefaultBinding);
 				}
 			}
 
 			// During the initial pass, only support constants for the default pin and non-bound variables
-			if (!FNiagaraParameterMapHistory::IsInitialValue(Var) && (DefaultPin == nullptr || DefaultPin->LinkedTo.Num() == 0) && !(ScriptVariable && (ScriptVariable->DefaultMode == ENiagaraDefaultMode::Binding || ScriptVariable->DefaultMode == ENiagaraDefaultMode::FailIfPreviouslyNotSet)))
+			if (!FNiagaraParameterMapHistory::IsInitialValue(Var) && (DefaultPin == nullptr || DefaultPin->LinkedTo.Num() == 0) && !(DefaultMode.IsSet() && (*DefaultMode == ENiagaraDefaultMode::Binding || *DefaultMode == ENiagaraDefaultMode::FailIfPreviouslyNotSet)))
 			{
-				HandleParameterRead(ParamMapHistoryIdx, Var, DefaultPin, DefaultPin != nullptr ? Cast<UNiagaraNode>(DefaultPin->GetOwningNode()) : nullptr, OutputChunkId, nullptr, !bWriteToParamMapEntries);
+				HandleParameterRead(ParamMapHistoryIdx, Var, DefaultPin, DefaultPin != nullptr ? Cast<UNiagaraNode>(DefaultPin->GetOwningNode()) : nullptr, OutputChunkId, TOptional<ENiagaraDefaultMode>(), TOptional<FNiagaraScriptVariableBinding>(), !bWriteToParamMapEntries);
 				UniqueVarToChunk.Add(Var, OutputChunkId);
 			}
 			else if (FNiagaraParameterMapHistory::IsInitialValue(Var))
@@ -6032,12 +6035,13 @@ void FHlslNiagaraTranslator::ParameterMapGet(UNiagaraNodeParameterMapGet* GetNod
 				OutputTypeDefinition.IsDataInterface() == false;
 			FNiagaraVariable Var = Schema->PinToNiagaraVariable(OutputPins[i], bNeedsValue, ENiagaraStructConversion::Simulation);
 
-			UNiagaraScriptVariable* Variable = GetNode->GetNiagaraGraph()->GetScriptVariable(Var);
+			FNiagaraScriptVariableBinding DefaultBinding;
+			TOptional<ENiagaraDefaultMode> DefaultMode = GetNode->GetNiagaraGraph()->GetDefaultMode(Var, &DefaultBinding);
 			if (Var.GetType().IsStatic())
 			{
 				if (FNiagaraParameterMapHistory::IsExternalConstantNamespace(Var, CompileOptions.TargetUsage, CompileOptions.GetTargetUsageBitmask()))
 				{
-					if (Variable && Variable->DefaultMode == ENiagaraDefaultMode::FailIfPreviouslyNotSet && !Variable->Variable.IsInNameSpace(FNiagaraConstants::UserNamespaceString))
+					if (DefaultMode.IsSet() && *DefaultMode == ENiagaraDefaultMode::FailIfPreviouslyNotSet && !Var.IsInNameSpace(FNiagaraConstants::UserNamespaceString))
 					{
 						RegisterCompileDependency(Var, FText::Format(LOCTEXT("UsedBeforeSet", "Variable {0} was read before being set. It's default mode is \"Fail If Previously Not Set\", so this isn't allowed."), FText::FromName(Var.GetName())), GetNode, OutputPins[i], true, ParamMapHistoryIdx);
 					}
@@ -6047,7 +6051,7 @@ void FHlslNiagaraTranslator::ParameterMapGet(UNiagaraNodeParameterMapGet* GetNod
 			}
 			else
 			{
-				HandleParameterRead(ParamMapHistoryIdx, Var, GetNode->GetDefaultPin(OutputPins[i]), GetNode, Outputs[i], Variable);
+				HandleParameterRead(ParamMapHistoryIdx, Var, GetNode->GetDefaultPin(OutputPins[i]), GetNode, Outputs[i], DefaultMode, DefaultBinding);
 			}
 		}
 	}
@@ -6063,7 +6067,7 @@ int32 FHlslNiagaraTranslator::MakeStaticVariableDirect(const UEdGraphPin* InDefa
 	return GetConstantDirect(Constant);
 }
 
-void FHlslNiagaraTranslator::HandleParameterRead(int32 ParamMapHistoryIdx, const FNiagaraVariable& InVar, const UEdGraphPin* DefaultPin, UNiagaraNode* ErrorNode, int32& OutputChunkId, UNiagaraScriptVariable* Variable, bool bTreatAsUnknownParameterMap, bool bIgnoreDefaultSetFirst)
+void FHlslNiagaraTranslator::HandleParameterRead(int32 ParamMapHistoryIdx, const FNiagaraVariable& InVar, const UEdGraphPin* DefaultPin, UNiagaraNode* ErrorNode, int32& OutputChunkId, TOptional<ENiagaraDefaultMode> DefaultMode, TOptional<FNiagaraScriptVariableBinding> DefaultBinding, bool bTreatAsUnknownParameterMap, bool bIgnoreDefaultSetFirst)
 {
 	FString ParameterMapInstanceName = GetParameterMapInstanceName(ParamMapHistoryIdx);
 	FNiagaraVariable Var = ConvertToSimulationVariable(InVar);
@@ -6108,7 +6112,7 @@ void FHlslNiagaraTranslator::HandleParameterRead(int32 ParamMapHistoryIdx, const
 	
 	if (FNiagaraParameterMapHistory::IsExternalConstantNamespace(Var, CompileOptions.TargetUsage, CompileOptions.GetTargetUsageBitmask()))
 	{
-		if (Variable && Variable->DefaultMode == ENiagaraDefaultMode::FailIfPreviouslyNotSet && !bIgnoreDefaultSetFirst)
+		if (DefaultMode.IsSet() && *DefaultMode == ENiagaraDefaultMode::FailIfPreviouslyNotSet && !bIgnoreDefaultSetFirst)
 		{
 			RegisterCompileDependency(Var, FText::Format(LOCTEXT("UsedBeforeSet", "Variable {0} was read before being set. It's default mode is \"Fail If Previously Not Set\", so this isn't allowed."), FText::FromName(Var.GetName())), ErrorNode, nullptr, true, ParamMapHistoryIdx);
 		}
@@ -6124,7 +6128,7 @@ void FHlslNiagaraTranslator::HandleParameterRead(int32 ParamMapHistoryIdx, const
 	}
 	else if (FNiagaraParameterMapHistory::IsAliasedModuleParameter(Var) && ActiveHistoryForFunctionCalls.InTopLevelFunctionCall(CompileOptions.TargetUsage))
 	{
-		if (Variable && Variable->DefaultMode == ENiagaraDefaultMode::Binding && Variable->DefaultBinding.IsValid())
+		if (DefaultMode.IsSet() && *DefaultMode == ENiagaraDefaultMode::Binding && DefaultBinding.IsSet() && DefaultBinding->IsValid())
 		{
 			// Skip the case where the below condition is met, but it's overridden by a binding.
 			bIsCandidateForRapidIteration = false;
@@ -6166,9 +6170,9 @@ void FHlslNiagaraTranslator::HandleParameterRead(int32 ParamMapHistoryIdx, const
 
 	bool bFailIfPreviouslyNotSetSentinel = false;
 	bool bValidateFailIfPreviouslyNotSet = false;
-	if (Variable != nullptr)
+	if (DefaultMode.IsSet())
 	{
-		bValidateFailIfPreviouslyNotSet = Variable->DefaultMode == ENiagaraDefaultMode::FailIfPreviouslyNotSet;
+		bValidateFailIfPreviouslyNotSet = *DefaultMode == ENiagaraDefaultMode::FailIfPreviouslyNotSet;
 	}
 
 	if (bValidateFailIfPreviouslyNotSet)
@@ -6387,9 +6391,9 @@ void FHlslNiagaraTranslator::HandleParameterRead(int32 ParamMapHistoryIdx, const
 
 		if (LastSetChunkIdx == INDEX_NONE && !bIgnoreDefaultValue)
 		{
-			if (Variable && Variable->DefaultMode == ENiagaraDefaultMode::Binding && Variable->DefaultBinding.IsValid())
+			if (DefaultMode.IsSet() && *DefaultMode == ENiagaraDefaultMode::Binding && DefaultBinding.IsSet() && DefaultBinding->IsValid())
 			{
-				FNiagaraScriptVariableBinding Bind = Variable->DefaultBinding;
+				FNiagaraScriptVariableBinding Bind = *DefaultBinding;
 
 				int32 Out = INDEX_NONE;
 				FNiagaraVariable BindVar = FNiagaraVariable(InVar.GetType(), Bind.GetName());
@@ -6398,7 +6402,7 @@ void FHlslNiagaraTranslator::HandleParameterRead(int32 ParamMapHistoryIdx, const
 					// Old assets often have vector inputs that default bind to what is now a position type. If we detect that, we change the type to prevent a compiler error.
 					BindVar.SetType(FNiagaraTypeDefinition::GetPositionDef());
 				}
-				HandleParameterRead(ActiveStageIdx, BindVar, nullptr, ErrorNode, Out, nullptr);
+				HandleParameterRead(ActiveStageIdx, BindVar, nullptr, ErrorNode, Out);
 
 				if (Out != INDEX_NONE)
 				{
@@ -7896,12 +7900,13 @@ void FHlslNiagaraTranslator::RegisterFunctionCall(ENiagaraScriptUsage ScriptUsag
 									if (LastSetChunkIdx == INDEX_NONE)
 									{
 										const UEdGraphPin* DefaultPin = History.GetDefaultValuePin(VarIdx);
-										UNiagaraScriptVariable* ScriptVariable = SourceGraph->GetScriptVariable(AliasedVar);
+										FNiagaraScriptVariableBinding DefaultBinding;
+										TOptional<ENiagaraDefaultMode> DefaultMode = SourceGraph->GetDefaultMode(AliasedVar, &DefaultBinding);
 
 										// Do not error on defaults for parameter reads here; we may be entering a SetVariable function call which is setting the first default for a parameter.
 										const bool bTreatAsUnknownParameterMap = false;
 										const bool bIgnoreDefaultSetFirst = true;
-										HandleParameterRead(ActiveStageIdx, AliasedVar, DefaultPin, ParamNode, LastSetChunkIdx, ScriptVariable, bTreatAsUnknownParameterMap, bIgnoreDefaultSetFirst);
+										HandleParameterRead(ActiveStageIdx, AliasedVar, DefaultPin, ParamNode, LastSetChunkIdx, DefaultMode, DefaultBinding, bTreatAsUnknownParameterMap, bIgnoreDefaultSetFirst);
 										
 										// If this variable was in the pending defaults list, go ahead and remove it
 										// as we added it before first use...
diff --git a/Engine/Plugins/FX/Niagara/Source/NiagaraEditor/Private/NiagaraHlslTranslator.h b/Engine/Plugins/FX/Niagara/Source/NiagaraEditor/Private/NiagaraHlslTranslator.h
--- a/Engine/Plugins/FX/Niagara/Source/NiagaraEditor/Private/NiagaraHlslTranslator.h
+++ b/Engine/Plugins/FX/Niagara/Source/NiagaraEditor/Private/NiagaraHlslTranslator.h
@@ -677,7 +677,7 @@ public:
 
 private:
 	void InitializeParameterMapDefaults(int32 ParamMapHistoryIdx);
-	void HandleParameterRead(int32 ParamMapHistoryIdx, const FNiagaraVariable& Var, const UEdGraphPin* DefaultPin, UNiagaraNode* ErrorNode, int32& OutputChunkId, UNiagaraScriptVariable* Variable, bool bTreatAsUnknownParameterMap = false, bool bIgnoreDefaultSetFirst = false);
+	void HandleParameterRead(int32 ParamMapHistoryIdx, const FNiagaraVariable& Var, const UEdGraphPin* DefaultPin, UNiagaraNode* ErrorNode, int32& OutputChunkId, TOptional<ENiagaraDefaultMode> DefaultMode = TOptional<ENiagaraDefaultMode>(), TOptional<FNiagaraScriptVariableBinding> DefaultBinding = TOptional<FNiagaraScriptVariableBinding>(), bool bTreatAsUnknownParameterMap = false, bool bIgnoreDefaultSetFirst = false);
 	bool ShouldConsiderTargetParameterMap(ENiagaraScriptUsage InUsage) const;
 	FString BuildParameterMapHlslDefinitions(TArray<FNiagaraVariable>& PrimaryDataSetOutputEntries);
 	void BuildMissingDefaults();
diff --git a/Engine/Plugins/FX/Niagara/Source/NiagaraEditor/Private/NiagaraNodeFunctionCall.cpp b/Engine/Plugins/FX/Niagara/Source/NiagaraEditor/Private/NiagaraNodeFunctionCall.cpp
--- a/Engine/Plugins/FX/Niagara/Source/NiagaraEditor/Private/NiagaraNodeFunctionCall.cpp
+++ b/Engine/Plugins/FX/Niagara/Source/NiagaraEditor/Private/NiagaraNodeFunctionCall.cpp
@@ -226,13 +226,13 @@ UEdGraphPin* UNiagaraNodeFunctionCall::AddStaticSwitchInputPin(FNiagaraVariable
 	NewPin->SetSavePinIfOrphaned(true);
 
 	FGuid PinPersistentGuid;
-	UNiagaraScriptVariable* ScriptVar = Graph->GetScriptVariable(Input);
-	if (ScriptVar)
+	TOptional<int32> SwitchDefaultValue = Graph->GetStaticSwitchDefaultValue(Input);
+	TOptional<FGuid> ScriptVariableGuid = Graph->GetScriptVariableGuid(Input);
+	if (SwitchDefaultValue.IsSet() && ScriptVariableGuid.IsSet())
 	{
-		int32 DefaultValue = ScriptVar->GetStaticSwitchDefaultValue();
 		Input.AllocateData();
-		Input.SetValue<FNiagaraInt32>({ DefaultValue });
-		PinPersistentGuid = ScriptVar->Metadata.GetVariableGuid();
+		Input.SetValue<FNiagaraInt32>( {*SwitchDefaultValue} );
+		PinPersistentGuid = *ScriptVariableGuid;
 	}
 	else
 	{
@@ -945,10 +945,10 @@ void UNiagaraNodeFunctionCall::UpdateStaticSwitchPinsWithPersistentGuids()
 			FNiagaraVariable InputVariable = GetDefault<UEdGraphSchema_Niagara>()->PinToNiagaraVariable(InputPin);
 			if (StaticSwitchVariables.Contains(InputVariable))
 			{
-				UNiagaraScriptVariable* InputScriptVariable = CalledGraph->GetScriptVariable(InputPin->PinName);
-				if (InputScriptVariable != nullptr)
+				TOptional<FGuid> VariableGuid = CalledGraph->GetScriptVariableGuid(InputVariable);
+				if (VariableGuid.IsSet())
 				{
-					InputPin->PersistentGuid = InputScriptVariable->Metadata.GetVariableGuid();
+					InputPin->PersistentGuid = *VariableGuid;
 				}
 			}
 		}
diff --git a/Engine/Plugins/FX/Niagara/Source/NiagaraEditor/Private/NiagaraNodeParameterMapGet.cpp b/Engine/Plugins/FX/Niagara/Source/NiagaraEditor/Private/NiagaraNodeParameterMapGet.cpp
--- a/Engine/Plugins/FX/Niagara/Source/NiagaraEditor/Private/NiagaraNodeParameterMapGet.cpp
+++ b/Engine/Plugins/FX/Niagara/Source/NiagaraEditor/Private/NiagaraNodeParameterMapGet.cpp
@@ -119,14 +119,20 @@ UEdGraphPin* UNiagaraNodeParameterMapGet::CreateDefaultPin(UEdGraphPin* OutputPi
 	}
 
 	// If the variable of the new default pin is already in use in the graph we use the configured default value
-	if (UNiagaraGraph* Graph = Cast<UNiagaraGraph>(GetGraph())) {
-		if (UNiagaraScriptVariable* ScriptVar = Graph->GetScriptVariable(Var.GetName())) {
-			if (ScriptVar->DefaultMode == ENiagaraDefaultMode::Value && ScriptVar->Variable.IsValid() && ScriptVar->Variable.IsDataAllocated()) {
+	if (UNiagaraGraph* Graph = Cast<UNiagaraGraph>(GetGraph())) 
+	{
+		TOptional<ENiagaraDefaultMode> DefaultMode = Graph->GetDefaultMode(Var);
+		TOptional<FNiagaraVariable> ScriptVariable = Graph->GetVariable(Var);
+
+		if (DefaultMode && ScriptVariable)
+		{
+			if (*DefaultMode == ENiagaraDefaultMode::Value && ScriptVariable->IsValid() && ScriptVariable->IsDataAllocated())
+			{
 				FNiagaraEditorModule& NiagaraEditorModule = FModuleManager::GetModuleChecked<FNiagaraEditorModule>("NiagaraEditor");
-				TSharedPtr<INiagaraEditorTypeUtilities, ESPMode::ThreadSafe> TypeEditorUtilities = NiagaraEditorModule.GetTypeUtilities(ScriptVar->Variable.GetType());
+				TSharedPtr<INiagaraEditorTypeUtilities, ESPMode::ThreadSafe> TypeEditorUtilities = NiagaraEditorModule.GetTypeUtilities(ScriptVariable->GetType());
 				if (bNeedsValue && TypeEditorUtilities.IsValid() && TypeEditorUtilities->CanHandlePinDefaults())
 				{
-					DefaultPin->DefaultValue = TypeEditorUtilities->GetPinDefaultStringFromValue(ScriptVar->Variable);
+					DefaultPin->DefaultValue = TypeEditorUtilities->GetPinDefaultStringFromValue(*ScriptVariable);
 				}
 			}
 		}
@@ -146,16 +152,6 @@ UEdGraphPin* UNiagaraNodeParameterMapGet::CreateDefaultPin(UEdGraphPin* OutputPi
 	return DefaultPin;
 }
 
-UNiagaraScriptVariable* UNiagaraNodeParameterMapGet::GetScriptVariable(FName VariableName) const
-{
-	if (UNiagaraGraph* Graph = Cast<UNiagaraGraph>(GetGraph())) {
-		if (UNiagaraScriptVariable* ScriptVar = Graph->GetScriptVariable(VariableName)) {
-			return ScriptVar;
-		}
-	}
-	return nullptr;
-}
-
 void UNiagaraNodeParameterMapGet::OnPinRenamed(UEdGraphPin* RenamedPin, const FString& OldName)
 {
 	UNiagaraNodeParameterMapBase::OnPinRenamed(RenamedPin, OldName);
@@ -244,8 +240,8 @@ void UNiagaraNodeParameterMapGet::OnNewTypedPinAdded(UEdGraphPin*& NewPin)
 	constexpr bool bNeedsValue = false;
 	FNiagaraVariable PinVariable = Schema->PinToNiagaraVariable(NewPin, bNeedsValue);
 	constexpr bool bIsStaticSwitch = false;
-	UNiagaraScriptVariable* ScriptVar = GetNiagaraGraph()->AddParameter(PinVariable, bIsStaticSwitch);
-	NewPin->PinName = ScriptVar->Variable.GetName();
+	GetNiagaraGraph()->AddParameter(PinVariable, bIsStaticSwitch);
+	NewPin->PinName = PinVariable.GetName();
 
 	if (NewPin->Direction == EEdGraphPinDirection::EGPD_Output)
 	{
@@ -381,7 +377,6 @@ void UNiagaraNodeParameterMapGet::SynchronizeDefaultInputPin(UEdGraphPin* Defaul
 		return;
 	}
 	
-	UNiagaraScriptVariable* ScriptVar = GetScriptVariable(OutputPin->PinName);
 	FNiagaraVariable Var = Schema->PinToNiagaraVariable(OutputPin);
 	if (FNiagaraParameterMapHistory::IsEngineParameter(Var))
 	{
@@ -397,26 +392,35 @@ void UNiagaraNodeParameterMapGet::SynchronizeDefaultInputPin(UEdGraphPin* Defaul
 	}
 
 	// Sync pin visibility with the configured default mode
-	if (ScriptVar) {
-		if (ScriptVar->DefaultMode == ENiagaraDefaultMode::Value)
-		{
-			DefaultPin->bNotConnectable = true;
-			DefaultPin->bDefaultValueIsReadOnly = true;
-		}
-		else if (ScriptVar->DefaultMode == ENiagaraDefaultMode::Custom) 
-		{
-			DefaultPin->bNotConnectable = false;
-			DefaultPin->bDefaultValueIsReadOnly = false;
-		}
-		else if (ScriptVar->DefaultMode == ENiagaraDefaultMode::Binding ||
-			ScriptVar->DefaultMode == ENiagaraDefaultMode::FailIfPreviouslyNotSet)
+	if (UNiagaraGraph* Graph = Cast<UNiagaraGraph>(GetGraph()))
+	{
+		if (TOptional<ENiagaraDefaultMode> DefaultMode = Graph->GetDefaultMode(Var))
 		{
-			DefaultPin->bHidden = true;
+			switch (*DefaultMode)
+			{
+				case ENiagaraDefaultMode::Value:
+				{
+					DefaultPin->bNotConnectable = true;
+					DefaultPin->bDefaultValueIsReadOnly = true;
+					break;
+				}
+				case ENiagaraDefaultMode::Custom:
+				{
+					DefaultPin->bNotConnectable = false;
+					DefaultPin->bDefaultValueIsReadOnly = false;
+					break;
+				}
+				case ENiagaraDefaultMode::Binding:
+				case ENiagaraDefaultMode::FailIfPreviouslyNotSet:
+				{
+					DefaultPin->bHidden = true;
+					break;
+				}
+			}
 		}
 	}
 }
 
-
 FText UNiagaraNodeParameterMapGet::GetNodeTitle(ENodeTitleType::Type TitleType) const
 {
 	return LOCTEXT("UNiagaraNodeParameterMapGetName", "Map Get");
diff --git a/Engine/Plugins/FX/Niagara/Source/NiagaraEditor/Private/NiagaraNodeParameterMapGet.h b/Engine/Plugins/FX/Niagara/Source/NiagaraEditor/Private/NiagaraNodeParameterMapGet.h
--- a/Engine/Plugins/FX/Niagara/Source/NiagaraEditor/Private/NiagaraNodeParameterMapGet.h
+++ b/Engine/Plugins/FX/Niagara/Source/NiagaraEditor/Private/NiagaraNodeParameterMapGet.h
@@ -69,9 +69,6 @@ protected:
 	/** Properly set up the default input pin for an output pin.*/
 	UEdGraphPin* CreateDefaultPin(UEdGraphPin* OutputPin);
 
-	/** Returns the metadata for the given variable or nullptr if no metadata is available */
-	UNiagaraScriptVariable* GetScriptVariable(FName VariableName) const;
-
 	UPROPERTY(meta = (SkipForCompileHash="true"))
 	TMap<FGuid, FGuid> PinOutputToPinDefaultPersistentId;
 };
diff --git a/Engine/Plugins/FX/Niagara/Source/NiagaraEditor/Private/NiagaraNodeParameterMapSet.cpp b/Engine/Plugins/FX/Niagara/Source/NiagaraEditor/Private/NiagaraNodeParameterMapSet.cpp
--- a/Engine/Plugins/FX/Niagara/Source/NiagaraEditor/Private/NiagaraNodeParameterMapSet.cpp
+++ b/Engine/Plugins/FX/Niagara/Source/NiagaraEditor/Private/NiagaraNodeParameterMapSet.cpp
@@ -132,8 +132,8 @@ void UNiagaraNodeParameterMapSet::OnNewTypedPinAdded(UEdGraphPin*& NewPin)
 			constexpr bool bNeedsValue = false;
 			FNiagaraVariable PinVariable = Schema->PinToNiagaraVariable(NewPin, bNeedsValue);
 			constexpr bool bIsStaticSwitch = false;
-			UNiagaraScriptVariable* ScriptVar = GetNiagaraGraph()->AddParameter(PinVariable, bIsStaticSwitch);
-			NewPin->PinName = ScriptVar->Variable.GetName();
+			GetNiagaraGraph()->AddParameter(PinVariable, bIsStaticSwitch);
+			NewPin->PinName = PinVariable.GetName();
 		}
 	}
 
diff --git a/Engine/Plugins/FX/Niagara/Source/NiagaraEditor/Private/NiagaraNodeStaticSwitch.cpp b/Engine/Plugins/FX/Niagara/Source/NiagaraEditor/Private/NiagaraNodeStaticSwitch.cpp
--- a/Engine/Plugins/FX/Niagara/Source/NiagaraEditor/Private/NiagaraNodeStaticSwitch.cpp
+++ b/Engine/Plugins/FX/Niagara/Source/NiagaraEditor/Private/NiagaraNodeStaticSwitch.cpp
@@ -728,12 +728,23 @@ void UNiagaraNodeStaticSwitch::PostLoad()
 	// Make sure that we are added to the static switch list.
 	if (GetInputType().IsValid() && InputParameterName.IsValid())
 	{
-		UNiagaraScriptVariable* Var = GetNiagaraGraph()->GetScriptVariable(InputParameterName);
-		if (Var != nullptr && Var->Variable.GetType() == GetInputType() && Var->GetIsStaticSwitch() == false)
+		if (InputParameterName.ToString() == TEXT("Enable Flattened Endcaps"))
 		{
-			UE_LOG(LogNiagaraEditor, Log, TEXT("Static switch constant \"%s\" in \"%s\" didn't have static switch meta-data conversion set properly. Fixing now."), *InputParameterName.ToString(), *GetPathName())
-			Var->SetIsStaticSwitch(true);
-			MarkNodeRequiresSynchronization(TEXT("Static switch metadata updated"), true);
+			int asdf = 0;
+			++asdf;
+		}
+
+		if (UNiagaraGraph* NiagaraGraph = GetNiagaraGraph())
+		{
+			const FNiagaraVariable Variable(GetInputType(), InputParameterName);
+			TOptional<bool> IsStaticSwitch = NiagaraGraph->IsStaticSwitch(Variable);
+
+			if (IsStaticSwitch.IsSet() && *IsStaticSwitch == false)
+			{
+				UE_LOG(LogNiagaraEditor, Log, TEXT("Static switch constant \"%s\" in \"%s\" didn't have static switch meta-data conversion set properly. Fixing now."), *InputParameterName.ToString(), *GetPathName())
+				NiagaraGraph->SetIsStaticSwitch(Variable, true);
+				MarkNodeRequiresSynchronization(TEXT("Static switch metadata updated"), true);
+			}
 		}
 	}
 
diff --git a/Engine/Plugins/FX/Niagara/Source/NiagaraEditor/Private/NiagaraParameterMapHistory.cpp b/Engine/Plugins/FX/Niagara/Source/NiagaraEditor/Private/NiagaraParameterMapHistory.cpp
--- a/Engine/Plugins/FX/Niagara/Source/NiagaraEditor/Private/NiagaraParameterMapHistory.cpp
+++ b/Engine/Plugins/FX/Niagara/Source/NiagaraEditor/Private/NiagaraParameterMapHistory.cpp
@@ -1660,10 +1660,11 @@ int32 FNiagaraParameterMapHistoryBuilder::HandleVariableRead(int32 ParamMapIdx,
 			// Add the default binding as well to the parameter history, if used.
 			if (UNiagaraGraph* Graph = Cast<UNiagaraGraph>(InPin->GetOwningNode()->GetGraph()))
 			{
-				UNiagaraScriptVariable* Variable = Graph->GetScriptVariable(AliasedVar);
-				if (Variable && Variable->DefaultMode == ENiagaraDefaultMode::Binding && Variable->DefaultBinding.IsValid())
+				FNiagaraScriptVariableBinding VariableBinding;
+				TOptional<ENiagaraDefaultMode> DefaultMode = Graph->GetDefaultMode(AliasedVar, &VariableBinding);
+				if (DefaultMode.IsSet() && *DefaultMode == ENiagaraDefaultMode::Binding && VariableBinding.IsValid())
 				{
-					FNiagaraVariable TempVar = FNiagaraVariable(Var.GetType(), Variable->DefaultBinding.GetName());
+					FNiagaraVariable TempVar = FNiagaraVariable(Var.GetType(), VariableBinding.GetName());
 					if (FNiagaraConstants::GetOldPositionTypeVariables().Contains(TempVar))
 					{
 						// Old assets often have vector inputs that default bind to what is now a position type. If we detect that, we change the type to prevent a compiler error.
diff --git a/Engine/Plugins/FX/Niagara/Source/NiagaraEditor/Private/NiagaraScriptVariable.cpp b/Engine/Plugins/FX/Niagara/Source/NiagaraEditor/Private/NiagaraScriptVariable.cpp
--- a/Engine/Plugins/FX/Niagara/Source/NiagaraEditor/Private/NiagaraScriptVariable.cpp
+++ b/Engine/Plugins/FX/Niagara/Source/NiagaraEditor/Private/NiagaraScriptVariable.cpp
@@ -178,3 +178,71 @@ FGuid UNiagaraScriptVariable::GenerateStableGuid(const UNiagaraScriptVariable* S
 	// Assign the guid components from the hash.
 	return FGuid(HashBuffer[0], HashBuffer[1], HashBuffer[2], HashBuffer[3]);
 }
+
+void FNiagaraScriptVariableData::Init(const FNiagaraVariable& InVariable, const FNiagaraVariableMetaData& InMetadata)
+{
+	Variable = InVariable;
+	Metadata = InMetadata;
+
+	AllocateData();
+
+	if (!Metadata.GetVariableGuid().IsValid())
+	{
+		Metadata.CreateNewGuid();
+	}
+	if (ChangeId.IsValid() == false)
+	{
+		UpdateChangeId();
+	}
+}
+
+void FNiagaraScriptVariableData::InitFrom(const UNiagaraScriptVariable& ScriptVariable, bool bCreateNewGuid)
+{
+	Init(ScriptVariable.Variable, ScriptVariable.Metadata);
+
+	DefaultMode = ScriptVariable.DefaultMode;
+	DefaultBinding = ScriptVariable.DefaultBinding;
+	bSubscribedToParameterDefinitions = ScriptVariable.GetIsSubscribedToParameterDefinitions();
+	bOverrideParameterDefinitionsDefaultValue = ScriptVariable.GetIsOverridingParameterDefinitionsDefaultValue();
+
+	DefaultValueVariant = ScriptVariable.GetDefaultValueVariant();
+	StaticSwitchDefaultValue = ScriptVariable.GetStaticSwitchDefaultValue();
+	bIsStaticSwitch = ScriptVariable.GetIsStaticSwitch();
+	ChangeId = ScriptVariable.GetChangeId();
+
+	// we generally want to create a new metadata guid as that is used to identify variables for renames, unless we specifically want to copy the entire variable
+	if (bCreateNewGuid)
+	{
+		Metadata.CreateNewGuid();
+	}
+}
+
+void FNiagaraScriptVariableData::InitFrom(const FNiagaraScriptVariableData& Source, bool bCreateNewGuid)
+{
+	Init(Source.Variable, Source.Metadata);
+
+	DefaultMode = Source.DefaultMode;
+	DefaultBinding = Source.DefaultBinding;
+	bSubscribedToParameterDefinitions = Source.GetIsSubscribedToParameterDefinitions();
+	bOverrideParameterDefinitionsDefaultValue = Source.GetIsOverridingParameterDefinitionsDefaultValue();
+
+	DefaultValueVariant = Source.DefaultValueVariant;
+	StaticSwitchDefaultValue = Source.StaticSwitchDefaultValue;
+	bIsStaticSwitch = Source.bIsStaticSwitch;
+	ChangeId = Source.ChangeId;
+
+	// we generally want to create a new metadata guid as that is used to identify variables for renames, unless we specifically want to copy the entire variable
+	if (bCreateNewGuid)
+	{
+		Metadata.CreateNewGuid();
+	}
+}
+
+bool FNiagaraScriptVariableData::AppendCompileHash(FNiagaraCompileHashVisitor* InVisitor) const
+{
+	if (!FNiagaraEditorUtilities::NestedPropertiesAppendCompileHash(static_cast<const void*>(this), StaticStruct(), EFieldIteratorFlags::ExcludeSuper, StaticStruct()->GetName(), InVisitor))
+	{
+		return false;
+	}
+	return true;
+}
diff --git a/Engine/Plugins/FX/Niagara/Source/NiagaraEditor/Private/ViewModels/Stack/NiagaraStackFunctionInput.cpp b/Engine/Plugins/FX/Niagara/Source/NiagaraEditor/Private/ViewModels/Stack/NiagaraStackFunctionInput.cpp
--- a/Engine/Plugins/FX/Niagara/Source/NiagaraEditor/Private/ViewModels/Stack/NiagaraStackFunctionInput.cpp
+++ b/Engine/Plugins/FX/Niagara/Source/NiagaraEditor/Private/ViewModels/Stack/NiagaraStackFunctionInput.cpp
@@ -2351,7 +2351,7 @@ TArray<UNiagaraScript*> UNiagaraStackFunctionInput::GetPossibleConversionScripts
     	{
     		// checking via metadata is not really correct, but it's super fast and good enough for the prefiltered list of scripts
     		TArray<FNiagaraVariable> AvailableVars;
-    		NodeGraph->GetAllMetaData().GetKeys(AvailableVars);
+		NodeGraph->GetAllVariables(AvailableVars);
     		int MatchingVars = 0;
     		for (const FNiagaraVariable& Var : AvailableVars)
     		{
@@ -2915,19 +2915,27 @@ void UNiagaraStackFunctionInput::GetDefaultLinkedHandleOrLinkedFunctionFromDefau
 void UNiagaraStackFunctionInput::UpdateValuesFromScriptDefaults(FInputValues& InInputValues) const
 {
 	// Get the script variable first since it's used to determine static switch and bound input values.
-	UNiagaraScriptVariable* InputScriptVariable = nullptr;
+	TOptional<ENiagaraDefaultMode> DefaultMode;
+	TOptional<int32> StaticSwitchDefaultValue;
+	FNiagaraScriptVariableBinding DefaultBinding;
+
+	FNiagaraVariable InputVariable(InputType, InputParameterHandle.GetParameterHandleString());
+
 	if (OwningFunctionCallNode->FunctionScript != nullptr)
 	{
-		UNiagaraGraph* FunctionGraph = CastChecked<UNiagaraScriptSource>(OwningFunctionCallNode->GetFunctionScriptSource())->NodeGraph;
-		InputScriptVariable = FunctionGraph->GetScriptVariable(InputParameterHandle.GetParameterHandleString());
+		if (UNiagaraGraph* FunctionGraph = CastChecked<UNiagaraScriptSource>(OwningFunctionCallNode->GetFunctionScriptSource())->NodeGraph)
+		{
+			DefaultMode = FunctionGraph->GetDefaultMode(InputVariable, &DefaultBinding);
+			StaticSwitchDefaultValue = FunctionGraph->GetStaticSwitchDefaultValue(InputVariable);
+		}
 	}
 
 	if (IsStaticParameter())
 	{
 		// Static switch parameters are always locally set values.
-		if (InputScriptVariable != nullptr)
+		if (StaticSwitchDefaultValue.IsSet())
 		{
-			TSharedPtr<FStructOnScope> StaticSwitchLocalStruct = FNiagaraEditorUtilities::StaticSwitchDefaultIntToStructOnScope(InputScriptVariable->GetStaticSwitchDefaultValue(), InputType);
+			TSharedPtr<FStructOnScope> StaticSwitchLocalStruct = FNiagaraEditorUtilities::StaticSwitchDefaultIntToStructOnScope(*StaticSwitchDefaultValue, InputType);
 			if(ensureMsgf(StaticSwitchLocalStruct.IsValid(), TEXT("Unsupported static struct default value.")))
 			{ 
 				InInputValues.Mode = EValueMode::Local;
@@ -2937,11 +2945,11 @@ void UNiagaraStackFunctionInput::UpdateValuesFromScriptDefaults(FInputValues& In
 	}
 	else
 	{
-		if (InputScriptVariable != nullptr && InputScriptVariable->DefaultMode == ENiagaraDefaultMode::Binding && InputScriptVariable->DefaultBinding.IsValid())
+		if (DefaultMode.IsSet() && *DefaultMode == ENiagaraDefaultMode::Binding && DefaultBinding.IsValid())
 		{
 			// The next highest precedence value is a linked value from a variable binding so check that.
 			InInputValues.Mode = EValueMode::Linked;
-			InInputValues.LinkedHandle = InputScriptVariable->DefaultBinding.GetName();
+			InInputValues.LinkedHandle = DefaultBinding.GetName();
 		}
 		else if (SourceScript.IsValid())
 		{
diff --git a/Engine/Plugins/FX/Niagara/Source/NiagaraEditor/Private/ViewModels/Stack/NiagaraStackGraphUtilities.cpp b/Engine/Plugins/FX/Niagara/Source/NiagaraEditor/Private/ViewModels/Stack/NiagaraStackGraphUtilities.cpp
--- a/Engine/Plugins/FX/Niagara/Source/NiagaraEditor/Private/ViewModels/Stack/NiagaraStackGraphUtilities.cpp
+++ b/Engine/Plugins/FX/Niagara/Source/NiagaraEditor/Private/ViewModels/Stack/NiagaraStackGraphUtilities.cpp
@@ -3102,7 +3102,7 @@ void FNiagaraStackGraphUtilities::AddNewVariableToParameterMapNode(UNiagaraNodeP
 	}
 
 	// First check that the new variable exists on the UNiagaraGraph. If not, add the new variable as a parameter.
-	if (Graph->GetScriptVariable(NewVariable.GetName()) == nullptr)
+	if (!Graph->HasVariable(NewVariable))
 	{
 		Graph->Modify();
 		Graph->AddParameter(NewVariable);
@@ -3126,7 +3126,7 @@ void FNiagaraStackGraphUtilities::AddNewVariableToParameterMapNode(UNiagaraNodeP
 	}
 
 	// First check that the new variable exists on the UNiagaraGraph. If not, add the new variable as a parameter.
-	if (Graph->GetScriptVariable(NewVariable.GetName()) == nullptr)
+	if (!Graph->HasVariable(NewVariable))
 	{
 		Graph->Modify();
 		Graph->AddParameter(NewScriptVar);
diff --git a/Engine/Plugins/FX/Niagara/Source/NiagaraEditor/Private/ViewModels/Stack/NiagaraStackModuleItem.cpp b/Engine/Plugins/FX/Niagara/Source/NiagaraEditor/Private/ViewModels/Stack/NiagaraStackModuleItem.cpp
--- a/Engine/Plugins/FX/Niagara/Source/NiagaraEditor/Private/ViewModels/Stack/NiagaraStackModuleItem.cpp
+++ b/Engine/Plugins/FX/Niagara/Source/NiagaraEditor/Private/ViewModels/Stack/NiagaraStackModuleItem.cpp
@@ -153,9 +153,6 @@ void UNiagaraStackModuleItem::RefreshChildrenInternal(const TArray<UNiagaraStack
 		// Determine if meta-data requires that we add our own refresh button here.
 		if (FunctionCallNode->HasValidScriptAndGraph())
 		{
-			UNiagaraGraph* Graph = FunctionCallNode->GetCalledGraph();
-			const TMap<FNiagaraVariable, TObjectPtr<UNiagaraScriptVariable>>& MetaDataMap = Graph->GetAllMetaData();
-			auto Iter = MetaDataMap.CreateConstIterator();
 			bCanRefresh = true;
 		}
 
diff --git a/Engine/Plugins/FX/Niagara/Source/NiagaraEditor/Public/NiagaraGraph.h b/Engine/Plugins/FX/Niagara/Source/NiagaraEditor/Public/NiagaraGraph.h
--- a/Engine/Plugins/FX/Niagara/Source/NiagaraEditor/Public/NiagaraGraph.h
+++ b/Engine/Plugins/FX/Niagara/Source/NiagaraEditor/Public/NiagaraGraph.h
@@ -11,6 +11,7 @@
 class UNiagaraParameterDefinitions;
 class UNiagaraScriptVariable;
 struct FSynchronizeWithParameterDefinitionsArgs;
+struct FNiagaraScriptVariableData;
 
 
 /** This is the type of action that occurred on a given Niagara graph. Note that this should follow from EEdGraphActionType, leaving some slop for growth. */
@@ -275,13 +276,17 @@ class UNiagaraGraph : public UEdGraph
 	/** Sets the meta-data associated with this variable. Creates a new UNiagaraScriptVariable if the target variable cannot be found. Illegal to call on FNiagaraVariables that are Niagara Constants. */
 	void SetMetaData(const FNiagaraVariable& InVar, const FNiagaraVariableMetaData& MetaData);
 
+	const TMap<FNiagaraVariable, FNiagaraGraphParameterReferenceCollection>& GetParameterReferenceMap() const; // NOTE: The const is a lie! (This indirectly calls RefreshParameterReferences, which can recreate the entire map)
+
+	// These functions are not supported for compilation copies
 	NIAGARAEDITOR_API const TMap<FNiagaraVariable, TObjectPtr<UNiagaraScriptVariable>>& GetAllMetaData() const;
 	NIAGARAEDITOR_API TMap<FNiagaraVariable, TObjectPtr<UNiagaraScriptVariable>>& GetAllMetaData();
 
-	const TMap<FNiagaraVariable, FNiagaraGraphParameterReferenceCollection>& GetParameterReferenceMap() const; // NOTE: The const is a lie! (This indirectly calls RefreshParameterReferences, which can recreate the entire map)
+	UNiagaraScriptVariable* GetScriptVariable(FNiagaraVariable Parameter, bool bUpdateIfPending = false);
+	NIAGARAEDITOR_API UNiagaraScriptVariable* GetScriptVariable(FName ParameterName, bool bUpdateIfPending = false);
 
-	UNiagaraScriptVariable* GetScriptVariable(FNiagaraVariable Parameter, bool bUpdateIfPending = false) const;
-	NIAGARAEDITOR_API UNiagaraScriptVariable* GetScriptVariable(FName ParameterName, bool bUpdateIfPending = false) const;
+	UNiagaraScriptVariable* GetScriptVariable(FNiagaraVariable Parameter) const;
+	NIAGARAEDITOR_API UNiagaraScriptVariable* GetScriptVariable(FName ParameterName) const;
 
 	/** Adds parameter to the VariableToScriptVariable map.*/
 	UNiagaraScriptVariable* AddParameter(const FNiagaraVariable& Parameter, bool bIsStaticSwitch = false);
@@ -393,6 +398,18 @@ class UNiagaraGraph : public UEdGraph
 
 	/** Helper to get a map of variables to all input/output pins with the same name. */
 	const TMap<FNiagaraVariable, FInputPinsAndOutputPins> NIAGARAEDITOR_API CollectVarsToInOutPinsMap() const;
+
+	void GetAllScriptVariableGuids(TArray<FGuid>& VariableGuids) const;
+	void GetAllVariables(TArray<FNiagaraVariable>& Variables) const;
+	TOptional<bool> IsStaticSwitch(const FNiagaraVariable& Variable) const;
+	TOptional<int32> GetStaticSwitchDefaultValue(const FNiagaraVariable& Variable) const;
+	TOptional<ENiagaraDefaultMode> GetDefaultMode(const FNiagaraVariable& Variable, FNiagaraScriptVariableBinding* Binding = nullptr) const;
+	TOptional<FGuid> GetScriptVariableGuid(const FNiagaraVariable& Variable) const;
+	TOptional<FNiagaraVariable> GetVariable(const FNiagaraVariable& Variable) const;
+	bool HasVariable(const FNiagaraVariable& Variable) const;
+
+	void SetIsStaticSwitch(const FNiagaraVariable& Variable, bool InValue);
+
 protected:
 	void RebuildNumericCache();
 	bool bNeedNumericCacheRebuilt;
@@ -420,6 +437,9 @@ private:
 
 	void StandardizeParameterNames();
 
+	static bool VariableLess(const FNiagaraVariable& Lhs, const FNiagaraVariable& Rhs);
+	TOptional<FNiagaraScriptVariableData> GetScriptVariableData(const FNiagaraVariable& Variable) const;
+
 private:
 	/** The current change identifier for this graph overall. Used to sync status with UNiagaraScripts.*/
 	UPROPERTY()
@@ -442,11 +462,14 @@ private:
 	/** Storage of variables defined for use with this graph.*/
 	UPROPERTY()
 	mutable TMap<FNiagaraVariable, TObjectPtr<UNiagaraScriptVariable>> VariableToScriptVariable;
-	
+
 	/** A map of parameters in the graph to their referencers. */
 	UPROPERTY(Transient)
 	mutable TMap<FNiagaraVariable, FNiagaraGraphParameterReferenceCollection> ParameterToReferencesMap;
 
+	UPROPERTY(Transient)
+	mutable TArray<FNiagaraScriptVariableData> CompilationScriptVariables;
+
 	FOnDataInterfaceChanged OnDataInterfaceChangedDelegate;
 	FOnSubObjectSelectionChanged OnSelectedSubObjectChanged;
 	FOnGetParameterDefinitionsForDetailsCustomization OnGetParameterDefinitionsForDetailsCustomizationDelegate;
diff --git a/Engine/Plugins/FX/Niagara/Source/NiagaraEditor/Public/NiagaraScriptVariable.h b/Engine/Plugins/FX/Niagara/Source/NiagaraEditor/Public/NiagaraScriptVariable.h
--- a/Engine/Plugins/FX/Niagara/Source/NiagaraEditor/Public/NiagaraScriptVariable.h
+++ b/Engine/Plugins/FX/Niagara/Source/NiagaraEditor/Public/NiagaraScriptVariable.h
@@ -7,6 +7,129 @@
 #include "NiagaraVariant.h"
 #include "NiagaraScriptVariable.generated.h"
 
+class UNiagaraScriptVariable;
+
+USTRUCT()
+struct NIAGARAEDITOR_API FNiagaraScriptVariableData
+{
+	GENERATED_USTRUCT_BODY()
+
+	FNiagaraScriptVariableData() = default;
+
+	void Init(const FNiagaraVariable& InVar, const FNiagaraVariableMetaData& InVarMetaData);
+	void InitFrom(const UNiagaraScriptVariable& Source, bool bCreateNewGuid = true);
+	void InitFrom(const FNiagaraScriptVariableData& Source, bool bCreateNewGuid = true);
+
+	/** The default mode. Can be Value, Binding or Custom. */
+	UPROPERTY(EditAnywhere, Category = "Default Value")
+	ENiagaraDefaultMode DefaultMode; 
+
+	/** The default binding. Only used if DefaultMode == ENiagaraDefaultMode::Binding. */
+	UPROPERTY(EditAnywhere, Category = "Default Value")
+	FNiagaraScriptVariableBinding DefaultBinding; 
+
+	/** Variable type, name and data. The data is not persistent, but used as a buffer when interfacing elsewhere. */
+	UPROPERTY()
+	FNiagaraVariable Variable;
+
+	/** The metadata associated with this script variable. */
+	UPROPERTY(EditAnywhere, Category = "Variable", meta=(ShowOnlyInnerProperties))
+	FNiagaraVariableMetaData Metadata;
+
+	/** Entry point for generating the compile hash.*/
+	bool AppendCompileHash(FNiagaraCompileHashVisitor* InVisitor) const;
+
+	bool GetIsStaticSwitch() const { return bIsStaticSwitch; };
+	void SetIsStaticSwitch(bool bInIsStaticSwitch) { bIsStaticSwitch = bInIsStaticSwitch; };
+
+	bool GetIsSubscribedToParameterDefinitions() const { return bSubscribedToParameterDefinitions; };
+	void SetIsSubscribedToParameterDefinitions(bool bInSubscribedToParameterDefinitions);
+
+	bool GetIsOverridingParameterDefinitionsDefaultValue() const { return bOverrideParameterDefinitionsDefaultValue; };
+	void SetIsOverridingParameterDefinitionsDefaultValue(bool bInOverridingParameterDefinitionsDefaultValue);
+
+	const FGuid& GetChangeId() const { return ChangeId; };
+	void SetChangeId(const FGuid& NewId) { ChangeId = NewId; };
+	void UpdateChangeId() { ChangeId = FGuid::NewGuid(); };
+
+	const FNiagaraVariant& GetDefaultValueVariant() const { return DefaultValueVariant; };
+
+	void SetDefaultValueData(const uint8* Data)
+	{
+		check(Data);
+		AllocateData();
+		DefaultValueVariant.SetBytes(Data, DefaultValueVariant.GetNumBytes());
+		Variable.SetData(GetDefaultValueData());
+	}
+
+	const uint8* GetDefaultValueData() const
+	{
+		if (DefaultValueVariant.GetMode() == ENiagaraVariantMode::Bytes)
+		{
+			return DefaultValueVariant.GetBytes();
+		}
+		return nullptr;
+	}
+
+	void CopyDefaultValueDataTo(uint8* Dest) const
+	{
+		check(Variable.GetSizeInBytes() == DefaultValueVariant.GetNumBytes());
+		check(IsDataAllocated());
+		FMemory::Memcpy(Dest, DefaultValueVariant.GetBytes(), DefaultValueVariant.GetNumBytes());
+	}
+
+	void SetStaticSwitchDefaultValue(const int32 Value)
+	{
+		StaticSwitchDefaultValue = Value;
+	}
+
+	int32 GetStaticSwitchDefaultValue() const
+	{
+		return StaticSwitchDefaultValue;
+	}
+
+private:
+	void AllocateData()
+	{
+		//It is possible the underlying variable does not have a valid type. Do not allocate until this is valid.
+		if (Variable.GetType().IsValid())
+		{
+			if (DefaultValueVariant.GetNumBytes() != Variable.GetSizeInBytes())
+			{
+				// Allocate the correct size in bytes, and set the variant mode to bytes by invoking the ctor.
+				TArray<uint8> InitialAllocationBytes;
+				InitialAllocationBytes.SetNumZeroed(Variable.GetSizeInBytes());
+				DefaultValueVariant = FNiagaraVariant(InitialAllocationBytes);
+			}
+			Variable.AllocateData();
+		}
+	}
+
+	bool IsDataAllocated() const
+	{
+		return DefaultValueVariant.GetNumBytes() > 0 && DefaultValueVariant.GetNumBytes() == Variable.GetSizeInBytes() && Variable.IsDataAllocated();
+	}
+
+private:
+	UPROPERTY(EditAnywhere, Category = "Hidden")
+	FNiagaraVariant DefaultValueVariant;
+
+	UPROPERTY()
+	int32 StaticSwitchDefaultValue;
+
+	UPROPERTY()
+	bool bIsStaticSwitch;
+
+	UPROPERTY(meta = (SkipForCompileHash = "true"))
+	bool bSubscribedToParameterDefinitions;
+
+	UPROPERTY(meta = (SkipForCompileHash = "true"))
+	FGuid ChangeId;
+
+	UPROPERTY(meta = (SkipForCompileHash = "true"))
+	bool bOverrideParameterDefinitionsDefaultValue;
+};
+
 /*
 * Used to store variable data and metadata per graph. 
 */
diff --git a/Engine/Plugins/FX/Niagara/Source/NiagaraEditor/Public/Widgets/SNiagaraParameterMenu.cpp b/Engine/Plugins/FX/Niagara/Source/NiagaraEditor/Public/Widgets/SNiagaraParameterMenu.cpp
--- a/Engine/Plugins/FX/Niagara/Source/NiagaraEditor/Public/Widgets/SNiagaraParameterMenu.cpp
+++ b/Engine/Plugins/FX/Niagara/Source/NiagaraEditor/Public/Widgets/SNiagaraParameterMenu.cpp
@@ -377,11 +377,9 @@ void SNiagaraAddParameterFromPanelMenu::CollectAllActions(FGraphActionListBuilde
 		// If culling parameter actions that match existing parameters in the graph, collect all IDs for parameters visited in the graph.
 		if(bCullParameterActionsAlreadyInGraph)
 		{ 
-			const TMap<FNiagaraVariable, TObjectPtr<UNiagaraScriptVariable>>& VariableToScriptVariableMap = Graph->GetAllMetaData();
-			for (auto It = VariableToScriptVariableMap.CreateConstIterator(); It; ++It)
-			{
-				ExistingGraphParameterIds.Add(It.Value()->Metadata.GetVariableGuid());
-			}
+			TArray<FGuid> VariableGuids;
+			Graph->GetAllScriptVariableGuids(VariableGuids);
+			ExistingGraphParameterIds.Append(VariableGuids);
 		}
 	}
 
@@ -440,10 +438,11 @@ void SNiagaraAddParameterFromPanelMenu::CollectAllActions(FGraphActionListBuilde
 				const FNiagaraVariable& Parameter = ParameterEntry.Key;
 
 				// Check if the graph owns the parameter (has a script variable for the parameter.)
-				if (const UNiagaraScriptVariable* ScriptVar = Graph->GetScriptVariable(Parameter))
+				if (Graph->HasVariable(Parameter))
 				{
 					// The graph owns the parameter. Skip if it is a static switch.
-					if (ScriptVar->GetIsStaticSwitch())
+					TOptional<bool> IsStaticSwitch = Graph->IsStaticSwitch(Parameter);
+					if (IsStaticSwitch.IsSet() && *IsStaticSwitch)
 					{
 						continue;
 					}
