- Fix for potential crash if attempting to allocate a zero sized neighbor grid

#rnx

diff --git a/Engine/Plugins/FX/Niagara/Source/Niagara/Private/NiagaraDataInterfaceNeighborGrid3D.cpp b/Engine/Plugins/FX/Niagara/Source/Niagara/Private/NiagaraDataInterfaceNeighborGrid3D.cpp
--- a/Engine/Plugins/FX/Niagara/Source/Niagara/Private/NiagaraDataInterfaceNeighborGrid3D.cpp
+++ b/Engine/Plugins/FX/Niagara/Source/Niagara/Private/NiagaraDataInterfaceNeighborGrid3D.cpp
@@ -51,8 +51,8 @@ void NeighborGrid3DRWInstanceData::ResizeBuffers(FRDGBuilder& GraphBuilder)
 	NeighborhoodCountBuffer.Release();
 	NeighborhoodBuffer.Release();
 
-	NeighborhoodCountBuffer.Initialize(GraphBuilder, TEXT("NiagaraNeighborGrid3D::NeighborCount"), EPixelFormat::PF_R32_SINT, sizeof(int32), NumTotalCells, BUF_Static);
-	NeighborhoodBuffer.Initialize(GraphBuilder, TEXT("NiagaraNeighborGrid3D::NeighborsGrid"), EPixelFormat::PF_R32_SINT, sizeof(int32), NumIntsInGridBuffer, BUF_Static);
+	NeighborhoodCountBuffer.Initialize(GraphBuilder, TEXT("NiagaraNeighborGrid3D::NeighborCount"), EPixelFormat::PF_R32_SINT, sizeof(int32), FMath::Max(NumTotalCells, 1u), BUF_Static);
+	NeighborhoodBuffer.Initialize(GraphBuilder, TEXT("NiagaraNeighborGrid3D::NeighborsGrid"), EPixelFormat::PF_R32_SINT, sizeof(int32), FMath::Max(NumIntsInGridBuffer, 1u), BUF_Static);
 }
 
 //////////////////////////////////////////////////////////////////////////
