Handle the usage indicators for Set Specific Variables stack nodes. Update the indicators when variables are added/removed.
#oforgithub public

diff --git a/Engine/Plugins/FX/Niagara/Source/NiagaraEditor/Private/ViewModels/NiagaraParameterPanelViewModel.cpp b/Engine/Plugins/FX/Niagara/Source/NiagaraEditor/Private/ViewModels/NiagaraParameterPanelViewModel.cpp
--- a/Engine/Plugins/FX/Niagara/Source/NiagaraEditor/Private/ViewModels/NiagaraParameterPanelViewModel.cpp
+++ b/Engine/Plugins/FX/Niagara/Source/NiagaraEditor/Private/ViewModels/NiagaraParameterPanelViewModel.cpp
@@ -2309,6 +2309,8 @@ void FNiagaraSystemToolkitParameterPanelViewModel::AddScriptVariable(const UNiag
 
 	if (bSuccess)
 	{
+		if (SystemViewModel->GetSystemStackViewModel())
+			SystemViewModel->GetSystemStackViewModel()->InvalidateCachedParameterUsage();
 		Refresh();
 		const bool bRequestRename = false;
 		SelectParameterItemByName(NewScriptVar->Variable.GetName(), bRequestRename);
diff --git a/Engine/Plugins/FX/Niagara/Source/NiagaraEditor/Private/ViewModels/Stack/NiagaraStackFunctionInput.cpp b/Engine/Plugins/FX/Niagara/Source/NiagaraEditor/Private/ViewModels/Stack/NiagaraStackFunctionInput.cpp
--- a/Engine/Plugins/FX/Niagara/Source/NiagaraEditor/Private/ViewModels/Stack/NiagaraStackFunctionInput.cpp
+++ b/Engine/Plugins/FX/Niagara/Source/NiagaraEditor/Private/ViewModels/Stack/NiagaraStackFunctionInput.cpp
@@ -44,6 +44,7 @@
 #include "ViewModels/NiagaraScratchPadScriptViewModel.h"
 #include "ViewModels/NiagaraScratchPadViewModel.h"
 #include "ViewModels/NiagaraSystemViewModel.h"
+#include "ViewModels/Stack/NiagaraStackViewModel.h"
 #include "ViewModels/Stack/NiagaraParameterHandle.h"
 #include "ViewModels/Stack/NiagaraStackFunctionInputCollection.h"
 #include "ViewModels/Stack/NiagaraStackGraphUtilities.h"
@@ -535,6 +536,10 @@ UNiagaraStackEntry::FStackIssueFixDelegate UNiagaraStackFunctionInput::GetUpgrad
 		FNiagaraScriptVersionUpgradeContext UpgradeContext;
 		UpgradeContext.CreateClipboardCallback = [this](UNiagaraClipboardContent* ClipboardContent)
 		{
+			TSharedRef<FNiagaraSystemViewModel> CachedSysViewModel = GetSystemViewModel();
+			if (CachedSysViewModel->GetSystemStackViewModel())
+				CachedSysViewModel->GetSystemStackViewModel()->InvalidateCachedParameterUsage();
+
 			RefreshChildren();
 			Copy(ClipboardContent);
 			if (ClipboardContent->Functions.Num() > 0)
@@ -1655,6 +1660,10 @@ void UNiagaraStackFunctionInput::SetDynamicInput(UNiagaraScript* DynamicInput, F
 	FNiagaraStackGraphUtilities::InitializeStackFunctionInputs(GetSystemViewModel(), GetEmitterViewModel(), GetStackEditorData(), *OwningModuleNode, *FunctionCallNode);
 	FNiagaraStackGraphUtilities::RelayoutGraph(*OwningFunctionCallNode->GetGraph());
 
+	TSharedRef<FNiagaraSystemViewModel> CachedSysViewModel = GetSystemViewModel();
+	if (CachedSysViewModel->GetSystemStackViewModel())
+		CachedSysViewModel->GetSystemStackViewModel()->InvalidateCachedParameterUsage();
+
 	RefreshChildren();
 }
 
@@ -1678,7 +1687,9 @@ void UNiagaraStackFunctionInput::SetCustomExpression(const FString& InCustomExpr
 	FNiagaraStackGraphUtilities::SetCustomExpressionForFunctionInput(OverridePin, InCustomExpression, FunctionCallNode);
 	FNiagaraStackGraphUtilities::InitializeStackFunctionInputs(GetSystemViewModel(), GetEmitterViewModel(), GetStackEditorData(), *OwningModuleNode, *FunctionCallNode);
 	FNiagaraStackGraphUtilities::RelayoutGraph(*OwningFunctionCallNode->GetGraph());
-
+	TSharedRef<FNiagaraSystemViewModel> CachedSysViewModel = GetSystemViewModel();
+	if (CachedSysViewModel->GetSystemStackViewModel())
+		CachedSysViewModel->GetSystemStackViewModel()->InvalidateCachedParameterUsage();
 	RefreshChildren();
 }
 
@@ -1808,7 +1819,9 @@ void UNiagaraStackFunctionInput::SetLocalValue(TSharedRef<FStructOnScope> InLoca
 	{
 		FNiagaraStackGraphUtilities::RelayoutGraph(*OwningFunctionCallNode->GetNiagaraGraph());
 	}
-
+	TSharedRef<FNiagaraSystemViewModel> CachedSysViewModel = GetSystemViewModel();
+	if (CachedSysViewModel->GetSystemStackViewModel())
+		CachedSysViewModel->GetSystemStackViewModel()->InvalidateCachedParameterUsage();
 	RefreshChildren();
 	RefreshValues();
 }
@@ -1931,6 +1944,10 @@ void UNiagaraStackFunctionInput::Reset()
 		{
 			ensureMsgf(false, TEXT("Attempted to reset a function input to default without a valid default."));
 		}
+		
+		TSharedRef<FNiagaraSystemViewModel> CachedSysViewModel = GetSystemViewModel();
+		if (CachedSysViewModel->GetSystemStackViewModel())
+			CachedSysViewModel->GetSystemStackViewModel()->InvalidateCachedParameterUsage();
 
 		RefreshChildren();
 		if (bBroadcastDataObjectChanged && InputValues.DataObject.IsValid())
@@ -2030,6 +2047,11 @@ void UNiagaraStackFunctionInput::ResetToBase()
 				}
 			}
 		}
+
+		TSharedRef<FNiagaraSystemViewModel> CachedSysViewModel = GetSystemViewModel();
+		if (CachedSysViewModel->GetSystemStackViewModel())
+			CachedSysViewModel->GetSystemStackViewModel()->InvalidateCachedParameterUsage();
+
 		RefreshChildren();
 	}
 }
@@ -2096,6 +2118,10 @@ void UNiagaraStackFunctionInput::OnRenamed(FText NewNameText)
 			NewName);
 		ensureMsgf(IsFinalized(), TEXT("Input not finalized when renamed."));
 
+
+		if (CachedSysViewModel->GetSystemStackViewModel())
+			CachedSysViewModel->GetSystemStackViewModel()->InvalidateCachedParameterUsage();
+
 		CachedSysViewModel->NotifyParameterRenamedExternally(OldVar, NewVar, Emitter.Emitter);
 
 	}
@@ -2134,7 +2160,8 @@ void UNiagaraStackFunctionInput::DeleteInput()
 		FNiagaraVariable Var = FNiagaraVariable(GetInputType(), GetInputParameterHandle().GetName());
 		NodeAssignment->Modify();
 		NodeAssignment->RemoveParameter(Var);
-
+		if (CachedSysViewModel->GetSystemStackViewModel())
+			CachedSysViewModel->GetSystemStackViewModel()->InvalidateCachedParameterUsage();
 		CachedSysViewModel->NotifyParameterRemovedExternally(Var, VersionedEmitter.Emitter);
 	}
 }
@@ -2296,6 +2323,9 @@ void UNiagaraStackFunctionInput::ReassignDynamicInputScript(UNiagaraScript* Dyna
 		FNiagaraStackGraphUtilities::RenameReferencingParameters(&System, Emitter, *InputValues.DynamicNode.Get(), OldName, NewName);
 
 		InputValues.DynamicNode->RefreshFromExternalChanges();
+		TSharedRef<FNiagaraSystemViewModel> CachedSysViewModel = GetSystemViewModel();
+		if (CachedSysViewModel->GetSystemStackViewModel())
+			CachedSysViewModel->GetSystemStackViewModel()->InvalidateCachedParameterUsage();
 
 		InputValues.DynamicNode->MarkNodeRequiresSynchronization(TEXT("Dynamic input script reassigned."), true);
 		RefreshChildren();
@@ -2475,6 +2505,9 @@ void UNiagaraStackFunctionInput::ChangeScriptVersion(FGuid NewScriptVersion)
       FCompileConstantResolver(GetEmitterViewModel()->GetEmitter(), FNiagaraStackGraphUtilities::GetOutputNodeUsage(*GetDynamicInputNode())) :
       FCompileConstantResolver(&GetSystemViewModel()->GetSystem(), FNiagaraStackGraphUtilities::GetOutputNodeUsage(*GetDynamicInputNode()));
 	GetDynamicInputNode()->ChangeScriptVersion(NewScriptVersion, UpgradeContext, true);
+	TSharedRef<FNiagaraSystemViewModel> CachedSysViewModel = GetSystemViewModel();
+	if (CachedSysViewModel->GetSystemStackViewModel())
+		CachedSysViewModel->GetSystemStackViewModel()->InvalidateCachedParameterUsage();
 	RefreshChildren();
 }
 
@@ -2600,6 +2633,10 @@ void UNiagaraStackFunctionInput::SetValueFromClipboardFunctionInput(const UNiaga
 	{
 		SetEditConditionEnabled(ClipboardFunctionInput.bEditConditionValue);
 	}
+
+	TSharedRef<FNiagaraSystemViewModel> CachedSysViewModel = GetSystemViewModel();
+	if (CachedSysViewModel->GetSystemStackViewModel())
+		CachedSysViewModel->GetSystemStackViewModel()->InvalidateCachedParameterUsage();
 }
 
 bool UNiagaraStackFunctionInput::IsScratchDynamicInput() const
