Add support for numeric type coercion in Add pins for nodes that don't support Numerics, like Parameter Map Get/Set nodes. This uses the type inference we are already doing to pick the correct type for the pin.
#oforgithub public

diff --git a/Engine/Plugins/FX/Niagara/Source/NiagaraEditor/Private/EdGraphSchema_Niagara.cpp b/Engine/Plugins/FX/Niagara/Source/NiagaraEditor/Private/EdGraphSchema_Niagara.cpp
--- a/Engine/Plugins/FX/Niagara/Source/NiagaraEditor/Private/EdGraphSchema_Niagara.cpp
+++ b/Engine/Plugins/FX/Niagara/Source/NiagaraEditor/Private/EdGraphSchema_Niagara.cpp
@@ -1115,29 +1115,55 @@ const FPinConnectionResponse UEdGraphSchema_Niagara::CanCreateConnection(const U
 			UNiagaraNodeWithDynamicPins* NodeA = Cast<UNiagaraNodeWithDynamicPins>(PinA->GetOwningNode());
 			UNiagaraNodeWithDynamicPins* NodeB = Cast<UNiagaraNodeWithDynamicPins>(PinB->GetOwningNode());
 			
-			// TODO: This shouldn't be handled explicitly here.
-			bool bPinAIsAddAndAcceptsPinB =
-				PinA->PinType.PinCategory == PinCategoryMisc && PinA->PinType.PinSubCategory == UNiagaraNodeWithDynamicPins::AddPinSubCategory &&
-				(
-					PinB->PinType.PinCategory == PinCategoryType &&
-					NodeA && PinToTypeDefinition(PinB) == FNiagaraTypeDefinition::GetGenericNumericDef() ?
-					NodeA->AllowNiagaraTypeForAddPin(PinToTypeDefinition(PinB)) : PinToTypeDefinition(PinB) != FNiagaraTypeDefinition::GetGenericNumericDef()
-				)
-				&& PinToTypeDefinition(PinB) != FNiagaraTypeDefinition::GetParameterMapDef();
-
-			bool bPinBIsAddAndAcceptsPinA =
-				PinB->PinType.PinCategory == PinCategoryMisc && PinB->PinType.PinSubCategory == UNiagaraNodeWithDynamicPins::AddPinSubCategory &&
-				(
-                    PinA->PinType.PinCategory == PinCategoryType &&
-                    NodeB && PinToTypeDefinition(PinA) == FNiagaraTypeDefinition::GetGenericNumericDef() ?
-                    NodeB->AllowNiagaraTypeForAddPin(PinToTypeDefinition(PinA)) : PinToTypeDefinition(PinA) != FNiagaraTypeDefinition::GetGenericNumericDef()
-                )
-				&& PinToTypeDefinition(PinA) != FNiagaraTypeDefinition::GetParameterMapDef();
-
-			if (bPinAIsAddAndAcceptsPinB == false && bPinBIsAddAndAcceptsPinA == false)
+			// Handle Direct connection or Inferred connection between numerics
+			FNiagaraTypeDefinition PinAType = PinToTypeDefinition(PinA);
+			FNiagaraTypeDefinition PinBType = PinToTypeDefinition(PinB);
+			FNiagaraTypeDefinition PinATypeAlt = PinAType;
+			FNiagaraTypeDefinition PinBTypeAlt = PinBType;
+
+			if (PinA->GetOwningNode())
+			{
+				UNiagaraGraph* Graph = Cast< UNiagaraGraph>(PinA->GetOwningNode()->GetGraph());
+				if (Graph)
+				{
+					PinATypeAlt = Graph->GetCachedNumericConversion(PinA);
+				}
+			}
+
+			if (PinB->GetOwningNode())
+			{
+				UNiagaraGraph* Graph = Cast< UNiagaraGraph>(PinB->GetOwningNode()->GetGraph());
+				if (Graph)
+				{
+					PinBTypeAlt = Graph->GetCachedNumericConversion(PinB);
+				}
+			}
+
+
+			auto CanConnect = [](const UEdGraphPin* SrcPin, const UEdGraphPin* DestPin, const FNiagaraTypeDefinition& DestPinType, UNiagaraNodeWithDynamicPins* SrcNode)->bool{
+				return SrcPin->PinType.PinCategory == PinCategoryMisc && SrcPin->PinType.PinSubCategory == UNiagaraNodeWithDynamicPins::AddPinSubCategory &&
+					(
+						DestPin->PinType.PinCategory == PinCategoryType &&
+						SrcNode && DestPinType == FNiagaraTypeDefinition::GetGenericNumericDef() ?
+						SrcNode->AllowNiagaraTypeForAddPin(DestPinType) : DestPinType != FNiagaraTypeDefinition::GetGenericNumericDef()
+						)
+					&& PinToTypeDefinition(DestPin) != FNiagaraTypeDefinition::GetParameterMapDef();
+			};
+
+			bool bPinAIsAddAndAcceptsPinB = CanConnect(PinA, PinB, PinBType, NodeA);
+			bool bPinAIsAddAndAcceptsPinBAlt = CanConnect(PinA, PinB, PinBTypeAlt, NodeA);
+
+			bool bPinBIsAddAndAcceptsPinA = CanConnect(PinB, PinA, PinAType, NodeB);
+			bool bPinBIsAddAndAcceptsPinAAlt = CanConnect(PinB, PinA, PinATypeAlt, NodeB);
+
+			// Disallow only if both paths are invalid. I.e. a Numeric that doesn't have an inferred type doesn't match, 
+			// but one that does have an inferred type can match.
+			if (bPinAIsAddAndAcceptsPinB == false && bPinBIsAddAndAcceptsPinA == false && 
+				bPinAIsAddAndAcceptsPinBAlt == false && bPinBIsAddAndAcceptsPinAAlt == false)
 			{
 				return FPinConnectionResponse(CONNECT_RESPONSE_DISALLOW, TypesAreNotCompatibleText);
 			}
+
 		}
 		else
 		{
diff --git a/Engine/Plugins/FX/Niagara/Source/NiagaraEditor/Private/NiagaraGraph.cpp b/Engine/Plugins/FX/Niagara/Source/NiagaraEditor/Private/NiagaraGraph.cpp
--- a/Engine/Plugins/FX/Niagara/Source/NiagaraEditor/Private/NiagaraGraph.cpp
+++ b/Engine/Plugins/FX/Niagara/Source/NiagaraEditor/Private/NiagaraGraph.cpp
@@ -2798,7 +2798,7 @@ void UNiagaraGraph::NotifyGraphDataInterfaceChanged()
 	OnDataInterfaceChangedDelegate.Broadcast();
 }
 
-FNiagaraTypeDefinition UNiagaraGraph::GetCachedNumericConversion(class UEdGraphPin* InPin)
+FNiagaraTypeDefinition UNiagaraGraph::GetCachedNumericConversion(const class UEdGraphPin* InPin)
 {
 	if (bNeedNumericCacheRebuilt)
 	{
diff --git a/Engine/Plugins/FX/Niagara/Source/NiagaraEditor/Private/NiagaraNodeWithDynamicPins.cpp b/Engine/Plugins/FX/Niagara/Source/NiagaraEditor/Private/NiagaraNodeWithDynamicPins.cpp
--- a/Engine/Plugins/FX/Niagara/Source/NiagaraEditor/Private/NiagaraNodeWithDynamicPins.cpp
+++ b/Engine/Plugins/FX/Niagara/Source/NiagaraEditor/Private/NiagaraNodeWithDynamicPins.cpp
@@ -29,6 +29,13 @@ void UNiagaraNodeWithDynamicPins::PinConnectionListChanged(UEdGraphPin* Pin)
 	if (Pin->PinType.PinCategory == UEdGraphSchema_Niagara::PinCategoryMisc && Pin->PinType.PinSubCategory == AddPinSubCategory && Pin->LinkedTo.Num() > 0)
 	{
 		FNiagaraTypeDefinition LinkedPinType = Schema->PinToTypeDefinition(Pin->LinkedTo[0]);
+
+		// Handle converting numerics to their inferred type in the case where preserving the numeric isn't a valid option, like Parameter Map Get Nodes
+		if (LinkedPinType == FNiagaraTypeDefinition::GetGenericNumericDef() && Pin->PinType.PinCategory == UEdGraphSchema_Niagara::PinCategoryMisc && 
+			Pin->PinType.PinSubCategory == UNiagaraNodeWithDynamicPins::AddPinSubCategory && !AllowNiagaraTypeForAddPin(LinkedPinType))
+		{
+			LinkedPinType = GetNiagaraGraph()->GetCachedNumericConversion(Pin->LinkedTo[0]);
+		}
 		Pin->PinType = Schema->TypeDefinitionToPinType(LinkedPinType);
 
 		FName NewPinName;
diff --git a/Engine/Plugins/FX/Niagara/Source/NiagaraEditor/Public/NiagaraGraph.h b/Engine/Plugins/FX/Niagara/Source/NiagaraEditor/Public/NiagaraGraph.h
--- a/Engine/Plugins/FX/Niagara/Source/NiagaraEditor/Public/NiagaraGraph.h
+++ b/Engine/Plugins/FX/Niagara/Source/NiagaraEditor/Public/NiagaraGraph.h
@@ -337,7 +337,7 @@ class UNiagaraGraph : public UEdGraph
 	/** Remove a listener for OnGraphNeedsRecompile events */
 	void RemoveOnGraphNeedsRecompileHandler(FDelegateHandle Handle);
 
-	FNiagaraTypeDefinition GetCachedNumericConversion(class UEdGraphPin* InPin);
+	FNiagaraTypeDefinition GetCachedNumericConversion(const class UEdGraphPin* InPin);
 
 	const class UEdGraphSchema_Niagara* GetNiagaraSchema() const;
 
