Fix for struct initialization of FNiagaraScriptVariableData

diff --git a/Engine/Plugins/FX/Niagara/Source/NiagaraEditor/Public/NiagaraScriptVariable.h b/Engine/Plugins/FX/Niagara/Source/NiagaraEditor/Public/NiagaraScriptVariable.h
--- a/Engine/Plugins/FX/Niagara/Source/NiagaraEditor/Public/NiagaraScriptVariable.h
+++ b/Engine/Plugins/FX/Niagara/Source/NiagaraEditor/Public/NiagaraScriptVariable.h
@@ -22,11 +22,11 @@ struct NIAGARAEDITOR_API FNiagaraScriptVariableData
 
 	/** The default mode. Can be Value, Binding or Custom. */
 	UPROPERTY(EditAnywhere, Category = "Default Value")
-	ENiagaraDefaultMode DefaultMode; 
+	ENiagaraDefaultMode DefaultMode = ENiagaraDefaultMode::Value;
 
 	/** The default binding. Only used if DefaultMode == ENiagaraDefaultMode::Binding. */
 	UPROPERTY(EditAnywhere, Category = "Default Value")
-	FNiagaraScriptVariableBinding DefaultBinding; 
+	FNiagaraScriptVariableBinding DefaultBinding;
 
 	/** Variable type, name and data. The data is not persistent, but used as a buffer when interfacing elsewhere. */
 	UPROPERTY()
@@ -115,19 +115,19 @@ private:
 	FNiagaraVariant DefaultValueVariant;
 
 	UPROPERTY()
-	int32 StaticSwitchDefaultValue;
+	int32 StaticSwitchDefaultValue = 0;
 
 	UPROPERTY()
-	bool bIsStaticSwitch;
+	bool bIsStaticSwitch = false;
 
 	UPROPERTY(meta = (SkipForCompileHash = "true"))
-	bool bSubscribedToParameterDefinitions;
+	bool bSubscribedToParameterDefinitions = false;
 
 	UPROPERTY(meta = (SkipForCompileHash = "true"))
 	FGuid ChangeId;
 
 	UPROPERTY(meta = (SkipForCompileHash = "true"))
-	bool bOverrideParameterDefinitionsDefaultValue;
+	bool bOverrideParameterDefinitionsDefaultValue = false;
 };
 
 /*
