- Presize VersionedScriptAdapters to fix issue with callback pointer to this

#rnx

diff --git a/Engine/Plugins/FX/Niagara/Source/Niagara/Classes/NiagaraParameterDefinitionsSubscriber.h b/Engine/Plugins/FX/Niagara/Source/Niagara/Classes/NiagaraParameterDefinitionsSubscriber.h
--- a/Engine/Plugins/FX/Niagara/Source/Niagara/Classes/NiagaraParameterDefinitionsSubscriber.h
+++ b/Engine/Plugins/FX/Niagara/Source/Niagara/Classes/NiagaraParameterDefinitionsSubscriber.h
@@ -43,7 +43,7 @@ class NIAGARA_API INiagaraParameterDefinitionsSubscriber
 {
 #if WITH_EDITORONLY_DATA
 public:
-	virtual ~INiagaraParameterDefinitionsSubscriber() = default;
+	virtual ~INiagaraParameterDefinitionsSubscriber();
 
 	void PostLoadDefinitionsSubscriptions();
 	void CleanupDefinitionsSubscriptions();
diff --git a/Engine/Plugins/FX/Niagara/Source/Niagara/Private/NiagaraParameterDefinitionsSubscriber.cpp b/Engine/Plugins/FX/Niagara/Source/Niagara/Private/NiagaraParameterDefinitionsSubscriber.cpp
--- a/Engine/Plugins/FX/Niagara/Source/Niagara/Private/NiagaraParameterDefinitionsSubscriber.cpp
+++ b/Engine/Plugins/FX/Niagara/Source/Niagara/Private/NiagaraParameterDefinitionsSubscriber.cpp
@@ -12,6 +12,11 @@
 
 #if WITH_EDITORONLY_DATA
 
+INiagaraParameterDefinitionsSubscriber::~INiagaraParameterDefinitionsSubscriber()
+{
+	check(OnDeferredSyncAllNameMatchParametersHandle.IsValid() == false);
+}
+
 void INiagaraParameterDefinitionsSubscriber::PostLoadDefinitionsSubscriptions()
 {
 	// First remove any subscriptions that are pointing to null parameter definitions assets. 
@@ -46,6 +51,7 @@ void INiagaraParameterDefinitionsSubscriber::PostLoadDefinitionsSubscriptions()
 	if (bAllParameterDefinitionsDiscovered == false)
 	{
 		SynchronizeWithParameterDefinitions();
+		check(OnDeferredSyncAllNameMatchParametersHandle.IsValid() == false);
 		OnDeferredSyncAllNameMatchParametersHandle = AssetRegistryModule.Get().OnFilesLoaded().AddLambda(SyncAllNameMatchParameters);
 	}
 	// Else the asset registry has discovered all parameter definitions assets;
@@ -66,6 +72,7 @@ void INiagaraParameterDefinitionsSubscriber::CleanupDefinitionsSubscriptions()
 			if (AssetRegistry)
 			{
 				AssetRegistry->OnFilesLoaded().Remove(OnDeferredSyncAllNameMatchParametersHandle);
+				OnDeferredSyncAllNameMatchParametersHandle.Reset();
 			}
 		}
 	}
diff --git a/Engine/Plugins/FX/Niagara/Source/Niagara/Private/NiagaraScript.cpp b/Engine/Plugins/FX/Niagara/Source/Niagara/Private/NiagaraScript.cpp
--- a/Engine/Plugins/FX/Niagara/Source/Niagara/Private/NiagaraScript.cpp
+++ b/Engine/Plugins/FX/Niagara/Source/Niagara/Private/NiagaraScript.cpp
@@ -2035,6 +2035,7 @@ void UNiagaraScript::PostLoad()
 		}
 	}
 
+	VersionedScriptAdapters.Reserve(VersionData.Num());
 	for (FVersionedNiagaraScriptData& Data : VersionData)
 	{
 		UNiagaraScriptSourceBase* Source = Data.Source;
@@ -2964,9 +2965,9 @@ void UNiagaraScript::BeginDestroy()
 {
 	Super::BeginDestroy();
 
-	if (!HasAnyFlags(RF_ClassDefaultObject) && ScriptResource)
+	if (!HasAnyFlags(RF_ClassDefaultObject))
 	{
-		if (!ScriptResource->QueueForRelease(ReleasedByRT))
+		if (ScriptResource == nullptr || !ScriptResource->QueueForRelease(ReleasedByRT))
 		{
 			// if there was nothing to release, then we don't need to wait for anything
 			ReleasedByRT = true;
