Implementing ApplyWorldOffset for UNiagaraComponent. In our case we will reset the simulation.

#offorgithub public

diff --git a/Engine/Plugins/FX/Niagara/Source/Niagara/Private/NiagaraComponent.cpp b/Engine/Plugins/FX/Niagara/Source/Niagara/Private/NiagaraComponent.cpp
--- a/Engine/Plugins/FX/Niagara/Source/Niagara/Private/NiagaraComponent.cpp
+++ b/Engine/Plugins/FX/Niagara/Source/Niagara/Private/NiagaraComponent.cpp
@@ -1763,6 +1763,19 @@ void UNiagaraComponent::OnUnregister()
 	}
 }
 
+void UNiagaraComponent::ApplyWorldOffset(const FVector& InOffset, bool bWorldShift)
+{
+	Super::ApplyWorldOffset(InOffset, bWorldShift);
+
+	// If this system is already active, assume that we want to dump the existing data and reset 
+	if (IsActive())
+	{
+	
+		Activate(true);
+	}
+}
+
+
 void UNiagaraComponent::BeginDestroy()
 {
 	//UE_LOG(LogNiagara, Log, TEXT("UNiagaraComponent::BeginDestroy(): %p - %d - %s\n"), this, ScalabilityManagerHandle, *GetAsset()->GetFullName());
diff --git a/Engine/Plugins/FX/Niagara/Source/Niagara/Public/NiagaraComponent.h b/Engine/Plugins/FX/Niagara/Source/Niagara/Public/NiagaraComponent.h
--- a/Engine/Plugins/FX/Niagara/Source/Niagara/Public/NiagaraComponent.h
+++ b/Engine/Plugins/FX/Niagara/Source/Niagara/Public/NiagaraComponent.h
@@ -141,7 +141,8 @@ private:
 	//~ Begin UActorComponent Interface.
 protected:
 	virtual void OnRegister() override;
-	virtual void OnUnregister() override;
+	virtual void OnUnregister() override; 
+	virtual void ApplyWorldOffset(const FVector& InOffset, bool bWorldShift) override;
 	virtual void OnEndOfFrameUpdateDuringTick() override;
 	virtual void CreateRenderState_Concurrent(FRegisterComponentContext* Context) override;
 	virtual void DestroyRenderState_Concurrent() override;
