- Add BP CreateNiagaraSimCache function so we can keep it all in BP without assets

#rnx

diff --git a/Engine/Plugins/FX/Niagara/Source/Niagara/Classes/NiagaraSimCacheFunctionLibrary.h b/Engine/Plugins/FX/Niagara/Source/Niagara/Classes/NiagaraSimCacheFunctionLibrary.h
--- a/Engine/Plugins/FX/Niagara/Source/Niagara/Classes/NiagaraSimCacheFunctionLibrary.h
+++ b/Engine/Plugins/FX/Niagara/Source/Niagara/Classes/NiagaraSimCacheFunctionLibrary.h
@@ -75,4 +75,10 @@ public:
 	*/
 	UFUNCTION(BlueprintCallable, Category=NiagaraSimCache)
 	static bool CaptureNiagaraSimCacheImmediate(UNiagaraSimCache* SimCache, FNiagaraSimCacheCreateParameters CreateParameters, UNiagaraComponent* NiagaraComponent);
+
+	/**
+	Captures the simulation cache object that can be captured into using the various API calls.
+	*/
+	UFUNCTION(BlueprintCallable, Category=NiagaraSimCache, meta=(WorldContext="WorldContextObject"))
+	static UNiagaraSimCache* CreateNiagaraSimCache(UObject* WorldContextObject);
 };
diff --git a/Engine/Plugins/FX/Niagara/Source/Niagara/Private/NiagaraSimCacheFunctionLibrary.cpp b/Engine/Plugins/FX/Niagara/Source/Niagara/Private/NiagaraSimCacheFunctionLibrary.cpp
--- a/Engine/Plugins/FX/Niagara/Source/Niagara/Private/NiagaraSimCacheFunctionLibrary.cpp
+++ b/Engine/Plugins/FX/Niagara/Source/Niagara/Private/NiagaraSimCacheFunctionLibrary.cpp
@@ -133,3 +133,13 @@ bool UNiagaraSimCacheFunctionLibrary::CaptureNiagaraSimCacheImmediate(UNiagaraSi
 	SimCache->EndWrite();
 	return SimCache->IsCacheValid();
 }
+
+UNiagaraSimCache* UNiagaraSimCacheFunctionLibrary::CreateNiagaraSimCache(UObject* WorldContextObject)
+{
+	if ( UWorld* World = GEngine->GetWorldFromContextObject(WorldContextObject, EGetWorldErrorMode::LogAndReturnNull) )
+	{
+		return NewObject<UNiagaraSimCache>(WorldContextObject);
+	}
+
+	return nullptr;
+}
