- Improve various BP functions for SimCache
- Add tooltips to improve clarity
- Add a ClearSimCache to the NiagaraComponent and also handle attempting to restart a sim from a cache that won't work

#rnx

diff --git a/Engine/Plugins/FX/Niagara/Source/Niagara/Classes/NiagaraSimCache.h b/Engine/Plugins/FX/Niagara/Source/Niagara/Classes/NiagaraSimCache.h
--- a/Engine/Plugins/FX/Niagara/Source/Niagara/Classes/NiagaraSimCache.h
+++ b/Engine/Plugins/FX/Niagara/Source/Niagara/Classes/NiagaraSimCache.h
@@ -267,10 +267,18 @@ public:
 	virtual bool IsReadyForFinishDestroy() override;
 	// UObject Interface
 
-	/** Is the simulation cache valid or not. */
-	UFUNCTION(BlueprintCallable, Category=NiagaraSimCache)
+	/** A valid cache is one that contains at least 1 frames worth of data. */
+	UFUNCTION(BlueprintCallable, Category=NiagaraSimCache, meta=(DisplayName="IsValid"))
 	bool IsCacheValid() const { return SoftNiagaraSystem.IsNull() == false; }
 
+	/** An empty cache contains no frame data and can not be used */
+	UFUNCTION(BlueprintCallable, Category = NiagaraSimCache)
+	bool IsEmpty() const { return !IsCacheValid(); }
+
+	/**	How were the attributes captured for this sim cache. */
+	UFUNCTION(BlueprintCallable, Category=NiagaraSimCache)
+	ENiagaraSimCacheAttributeCaptureMode GetAttributeCaptureMode() const { return CreateParameters.AttributeCaptureMode; }
+
 	bool BeginWrite(FNiagaraSimCacheCreateParameters InCreateParameters, UNiagaraComponent* NiagaraComponent);
 	bool WriteFrame(UNiagaraComponent* NiagaraComponent);
 	bool EndWrite();
@@ -349,15 +357,17 @@ public:
 	EmitterName - If left blank will return the system simulation attributes.
 	*/
 	UFUNCTION(BlueprintCallable, Category = NiagaraSimCache)
-	void ReadColorAttribute(TArray<FLinearColor>& OutValues, FName AttributeName, FName EmitterName, int FrameIndex = 0) const;
+	void ReadColorAttribute(TArray<FLinearColor>& OutValues, FName AttributeName = FName("Color"), FName EmitterName = NAME_None, int FrameIndex = 0) const;
 
 	/**
 	Reads Niagara Position attributes by name from the cache frame and appends them into the OutValues array.
 	Local space emitters provide data at local locations unless bLocalSpaceToWorld is true.
 	EmitterName - If left blank will return the system simulation attributes.
+	LocalSpaceToWorld - Caches are always stored in the emitters space, i.e. local or world space.  You can set this to false if you want the local position rather than the world position.
 	*/
-	UFUNCTION(BlueprintCallable, Category = NiagaraSimCache)
-	void ReadPositionAttribute(TArray<FVector>& OutValues, FName AttributeName, FName EmitterName, bool bLocalSpaceToWorld = true, int FrameIndex = 0) const;
+	UFUNCTION(BlueprintCallable, Category = NiagaraSimCache, meta=(AdvancedDisplay="bLocalSpaceToWorld"))
+	void ReadPositionAttribute(TArray<FVector>& OutValues, FName AttributeName = FName("Position"), FName EmitterName = NAME_None, bool bLocalSpaceToWorld = true, int FrameIndex = 0) const;
+	
 
 	/**
 	Reads Niagara Position attributes by name from the cache frame and appends them into the OutValues array.
@@ -366,15 +376,16 @@ public:
 	EmitterName - If left blank will return the system simulation attributes.
 	*/
 	UFUNCTION(BlueprintCallable, Category = NiagaraSimCache)
-	void ReadPositionAttributeWithRebase(TArray<FVector>& OutValues, FName AttributeName, FName EmitterName, FTransform Transform, int FrameIndex = 0) const;
+	void ReadPositionAttributeWithRebase(TArray<FVector>& OutValues, FTransform Transform, FName AttributeName = FName("Position"), FName EmitterName = NAME_None, int FrameIndex = 0) const;
 
 	/**
 	Reads Niagara Quaternion attributes by name from the cache frame and appends them into the OutValues array.
 	Local space emitters provide data at local rotation unless bLocalSpaceToWorld is true.
 	EmitterName - If left blank will return the system simulation attributes.
+	LocalSpaceToWorld - Caches are always stored in the emitters space, i.e. local or world space.  You can set this to false if you want the local Quat rather than the world Quat.
 	*/
-	UFUNCTION(BlueprintCallable, Category = NiagaraSimCache)
-	void ReadQuatAttribute(TArray<FQuat>& OutValues, FName AttributeName, FName EmitterName, bool bLocalSpaceToWorld = true, int FrameIndex = 0) const;
+	UFUNCTION(BlueprintCallable, Category = NiagaraSimCache, meta=(AdvancedDisplay="bLocalSpaceToWorld"))
+	void ReadQuatAttribute(TArray<FQuat>& OutValues, FName AttributeName = FName("MeshOrientation"), FName EmitterName = NAME_None, bool bLocalSpaceToWorld = true, int FrameIndex = 0) const;
 
 	/**
 	Reads Niagara Quaternion attributes by name from the cache frame and appends them into the OutValues array.
@@ -383,7 +394,7 @@ public:
 	EmitterName - If left blank will return the system simulation attributes.
 	*/
 	UFUNCTION(BlueprintCallable, Category = NiagaraSimCache)
-	void ReadQuatAttributeWithRebase(TArray<FQuat>& OutValues, FName AttributeName, FName EmitterName, FQuat Quat, int FrameIndex = 0) const;
+	void ReadQuatAttributeWithRebase(TArray<FQuat>& OutValues, FQuat Quat, FName AttributeName = FName("MeshOrientation"), FName EmitterName = NAME_None, int FrameIndex = 0) const;
 
 private:
 	UPROPERTY(VisibleAnywhere, Category=SimCache, meta=(DisplayName="Niagara System"))
diff --git a/Engine/Plugins/FX/Niagara/Source/Niagara/Classes/NiagaraSimCacheFunctionLibrary.h b/Engine/Plugins/FX/Niagara/Source/Niagara/Classes/NiagaraSimCacheFunctionLibrary.h
--- a/Engine/Plugins/FX/Niagara/Source/Niagara/Classes/NiagaraSimCacheFunctionLibrary.h
+++ b/Engine/Plugins/FX/Niagara/Source/Niagara/Classes/NiagaraSimCacheFunctionLibrary.h
@@ -16,7 +16,7 @@ class UAsyncNiagaraCaptureSimCache : public UBlueprintAsyncActionBase
 	GENERATED_BODY()
 
 public:
-	DECLARE_DYNAMIC_MULTICAST_DELEGATE(FOnCaptureComplete);
+	DECLARE_DYNAMIC_MULTICAST_DELEGATE_OneParam(FOnCaptureComplete, bool, bSuccess);
 
 	UPROPERTY()
 	TObjectPtr<UNiagaraSimCache> CaptureSimCache;
@@ -51,7 +51,7 @@ public:
 	CaptureRate allows you to reduce the rate of capture, i.e. a rate of 2 would capture frames 0, 2, 4, etc.
 	*/
 	UFUNCTION(BlueprintCallable, meta=(BlueprintInternalUseOnly="true", Category=NiagaraSimCache))
-	static UAsyncNiagaraCaptureSimCache* CaptureNiagaraSimCacheMultiFrame(UNiagaraSimCache* SimCache, FNiagaraSimCacheCreateParameters CreateParameters, UNiagaraComponent* NiagaraComponent, int32 NumFrames = 16, int32 CaptureRate = 1);
+	static UAsyncNiagaraCaptureSimCache* CaptureNiagaraSimCacheMultiFrame(UNiagaraSimCache* SimCache, FNiagaraSimCacheCreateParameters CreateParameters, UNiagaraComponent* NiagaraComponent, UNiagaraSimCache*& OutSimCache, int32 NumFrames = 16, int32 CaptureRate = 1);
 
 	/**
 	Capture frames from the provided simulation into a SimCache until the simulation becomes inactive or completes.
@@ -59,7 +59,7 @@ public:
 	CaptureRate allows you to reduce the rate of capture, i.e. a rate of 2 would capture frames 0, 2, 4, etc.
 	*/
 	UFUNCTION(BlueprintCallable, meta=(BlueprintInternalUseOnly="true", Category=NiagaraSimCache))
-	static UAsyncNiagaraCaptureSimCache* CaptureNiagaraSimCacheUntilComplete(UNiagaraSimCache* SimCache, FNiagaraSimCacheCreateParameters CreateParameters, UNiagaraComponent* NiagaraComponent, int32 CaptureRate = 1);
+	static UAsyncNiagaraCaptureSimCache* CaptureNiagaraSimCacheUntilComplete(UNiagaraSimCache* SimCache, FNiagaraSimCacheCreateParameters CreateParameters, UNiagaraComponent* NiagaraComponent, UNiagaraSimCache*& OutSimCache, int32 CaptureRate = 1);
 };
 
 UCLASS()
@@ -73,12 +73,12 @@ public:
 	This happens immediately so you may need to be careful with tick order of the component you are capturing from.
 	The return can be invalid if the component can not be captured for some reason (i.e. not active).
 	*/
-	UFUNCTION(BlueprintCallable, Category=NiagaraSimCache)
-	static bool CaptureNiagaraSimCacheImmediate(UNiagaraSimCache* SimCache, FNiagaraSimCacheCreateParameters CreateParameters, UNiagaraComponent* NiagaraComponent);
+	UFUNCTION(BlueprintCallable, Category=NiagaraSimCache, meta=(ReturnDisplayName="Success"))
+	static bool CaptureNiagaraSimCacheImmediate(UNiagaraSimCache* SimCache, FNiagaraSimCacheCreateParameters CreateParameters, UNiagaraComponent* NiagaraComponent, UNiagaraSimCache*& OutSimCache);
 
 	/**
 	Captures the simulation cache object that can be captured into using the various API calls.
 	*/
-	UFUNCTION(BlueprintCallable, Category=NiagaraSimCache, meta=(WorldContext="WorldContextObject"))
+	UFUNCTION(BlueprintCallable, Category=NiagaraSimCache, meta=(WorldContext="WorldContextObject", ReturnDisplayName="SimCache"))
 	static UNiagaraSimCache* CreateNiagaraSimCache(UObject* WorldContextObject);
 };
diff --git a/Engine/Plugins/FX/Niagara/Source/Niagara/Private/NiagaraComponent.cpp b/Engine/Plugins/FX/Niagara/Source/Niagara/Private/NiagaraComponent.cpp
--- a/Engine/Plugins/FX/Niagara/Source/Niagara/Private/NiagaraComponent.cpp
+++ b/Engine/Plugins/FX/Niagara/Source/Niagara/Private/NiagaraComponent.cpp
@@ -3525,10 +3525,17 @@ void UNiagaraComponent::SetPreviewLODDistance(bool bInEnablePreviewLODDistance,
 	}
 }
 
-void UNiagaraComponent::SetSimCache(UNiagaraSimCache* InSimCache)
+void UNiagaraComponent::SetSimCache(UNiagaraSimCache* InSimCache, bool bResetSystem)
 {
+	const bool bForceReset = SimCache ? SimCache->GetAttributeCaptureMode() != ENiagaraSimCacheAttributeCaptureMode::All : false;
+
 	SimCache = InSimCache;
 	UpdateInstanceSoloMode();
+
+	if (bResetSystem || bForceReset)
+	{
+		ReinitializeSystem();
+	}
 }
 
 UNiagaraSimCache* UNiagaraComponent::GetSimCache() const
diff --git a/Engine/Plugins/FX/Niagara/Source/Niagara/Private/NiagaraSimCache.cpp b/Engine/Plugins/FX/Niagara/Source/Niagara/Private/NiagaraSimCache.cpp
--- a/Engine/Plugins/FX/Niagara/Source/Niagara/Private/NiagaraSimCache.cpp
+++ b/Engine/Plugins/FX/Niagara/Source/Niagara/Private/NiagaraSimCache.cpp
@@ -591,7 +591,7 @@ void UNiagaraSimCache::ReadPositionAttribute(TArray<FVector>& OutValues, FName A
 	}
 }
 
-void UNiagaraSimCache::ReadPositionAttributeWithRebase(TArray<FVector>& OutValues, FName AttributeName, FName EmitterName, FTransform Transform, int FrameIndex) const
+void UNiagaraSimCache::ReadPositionAttributeWithRebase(TArray<FVector>& OutValues, FTransform Transform, FName AttributeName, FName EmitterName, int FrameIndex) const
 {
 	FNiagaraSimCacheAttributeReaderHelper AttributeReader(this, EmitterName, AttributeName, FrameIndex);
 	if (AttributeReader.IsValid() && AttributeReader.Variable->Variable.GetType() == FNiagaraTypeDefinition::GetPositionDef())
@@ -631,7 +631,7 @@ void UNiagaraSimCache::ReadQuatAttribute(TArray<FQuat>& OutValues, FName Attribu
 	}
 }
 
-void UNiagaraSimCache::ReadQuatAttributeWithRebase(TArray<FQuat>& OutValues, FName AttributeName, FName EmitterName, FQuat Quat, int FrameIndex) const
+void UNiagaraSimCache::ReadQuatAttributeWithRebase(TArray<FQuat>& OutValues, FQuat Quat, FName AttributeName, FName EmitterName, int FrameIndex) const
 {
 	FNiagaraSimCacheAttributeReaderHelper AttributeReader(this, EmitterName, AttributeName, FrameIndex);
 	if (AttributeReader.IsValid() && AttributeReader.Variable->Variable.GetType() == FNiagaraTypeDefinition::GetQuatDef())
diff --git a/Engine/Plugins/FX/Niagara/Source/Niagara/Private/NiagaraSimCacheFunctionLibrary.cpp b/Engine/Plugins/FX/Niagara/Source/Niagara/Private/NiagaraSimCacheFunctionLibrary.cpp
--- a/Engine/Plugins/FX/Niagara/Source/Niagara/Private/NiagaraSimCacheFunctionLibrary.cpp
+++ b/Engine/Plugins/FX/Niagara/Source/Niagara/Private/NiagaraSimCacheFunctionLibrary.cpp
@@ -34,7 +34,7 @@ void UAsyncNiagaraCaptureSimCache::SetReadyToDestroy()
 	{
 		CaptureSimCache->EndWrite();
 	}
-	CaptureComplete.Broadcast();
+	CaptureComplete.Broadcast(CaptureSimCache ? CaptureSimCache->IsCacheValid() : false);
 }
 
 bool UAsyncNiagaraCaptureSimCache::OnFrameTick(float DeltaTime)
@@ -82,7 +82,7 @@ bool UAsyncNiagaraCaptureSimCache::OnFrameTick(float DeltaTime)
 	return true;
 }
 
-UAsyncNiagaraCaptureSimCache* UAsyncNiagaraCaptureSimCache::CaptureNiagaraSimCacheMultiFrame(UNiagaraSimCache* SimCache, FNiagaraSimCacheCreateParameters CreateParameters, UNiagaraComponent* NiagaraComponent, int32 NumFrames, int32 CaptureRate)
+UAsyncNiagaraCaptureSimCache* UAsyncNiagaraCaptureSimCache::CaptureNiagaraSimCacheMultiFrame(UNiagaraSimCache* SimCache, FNiagaraSimCacheCreateParameters CreateParameters, UNiagaraComponent* NiagaraComponent, UNiagaraSimCache*& OutSimCache, int32 NumFrames, int32 CaptureRate)
 {
 	UAsyncNiagaraCaptureSimCache* CaptureAction = NewObject<UAsyncNiagaraCaptureSimCache>();
 	CaptureAction->CaptureSimCache = SimCache;
@@ -95,11 +95,12 @@ UAsyncNiagaraCaptureSimCache* UAsyncNiagaraCaptureSimCache::CaptureNiagaraSimCac
 	{
 		SimCache->BeginWrite(CreateParameters, NiagaraComponent);
 	}
+	OutSimCache = SimCache;
 
 	return CaptureAction;
 }
 
-UAsyncNiagaraCaptureSimCache* UAsyncNiagaraCaptureSimCache::CaptureNiagaraSimCacheUntilComplete(UNiagaraSimCache* SimCache, FNiagaraSimCacheCreateParameters CreateParameters, UNiagaraComponent* NiagaraComponent, int32 CaptureRate)
+UAsyncNiagaraCaptureSimCache* UAsyncNiagaraCaptureSimCache::CaptureNiagaraSimCacheUntilComplete(UNiagaraSimCache* SimCache, FNiagaraSimCacheCreateParameters CreateParameters, UNiagaraComponent* NiagaraComponent, UNiagaraSimCache*& OutSimCache, int32 CaptureRate)
 {
 	UAsyncNiagaraCaptureSimCache* CaptureAction = NewObject<UAsyncNiagaraCaptureSimCache>();
 	CaptureAction->CaptureSimCache = SimCache;
@@ -112,6 +113,7 @@ UAsyncNiagaraCaptureSimCache* UAsyncNiagaraCaptureSimCache::CaptureNiagaraSimCac
 	{
 		SimCache->BeginWrite(CreateParameters, NiagaraComponent);
 	}
+	OutSimCache = SimCache;
 
 	return CaptureAction;
 }
@@ -121,17 +123,18 @@ UNiagaraSimCacheFunctionLibrary::UNiagaraSimCacheFunctionLibrary(const FObjectIn
 {
 }
 
-bool UNiagaraSimCacheFunctionLibrary::CaptureNiagaraSimCacheImmediate(UNiagaraSimCache* SimCache, FNiagaraSimCacheCreateParameters CreateParameters, UNiagaraComponent* NiagaraComponent)
+bool UNiagaraSimCacheFunctionLibrary::CaptureNiagaraSimCacheImmediate(UNiagaraSimCache* SimCache, FNiagaraSimCacheCreateParameters CreateParameters, UNiagaraComponent* NiagaraComponent, UNiagaraSimCache*& OutSimCache)
 {
-	if ( SimCache == nullptr || NiagaraComponent == nullptr )
+	bool bSuccess = false;
+	if (SimCache != nullptr && NiagaraComponent != nullptr)
 	{
-		return false;
+		SimCache->BeginWrite(CreateParameters, NiagaraComponent);
+		SimCache->WriteFrame(NiagaraComponent);
+		SimCache->EndWrite();
+		bSuccess = SimCache->IsCacheValid();
 	}
-
-	SimCache->BeginWrite(CreateParameters, NiagaraComponent);
-	SimCache->WriteFrame(NiagaraComponent);
-	SimCache->EndWrite();
-	return SimCache->IsCacheValid();
+	OutSimCache = SimCache;
+	return bSuccess;
 }
 
 UNiagaraSimCache* UNiagaraSimCacheFunctionLibrary::CreateNiagaraSimCache(UObject* WorldContextObject)
diff --git a/Engine/Plugins/FX/Niagara/Source/Niagara/Public/NiagaraComponent.h b/Engine/Plugins/FX/Niagara/Source/Niagara/Public/NiagaraComponent.h
--- a/Engine/Plugins/FX/Niagara/Source/Niagara/Public/NiagaraComponent.h
+++ b/Engine/Plugins/FX/Niagara/Source/Niagara/Public/NiagaraComponent.h
@@ -705,12 +705,31 @@ public:
 	uint32 bWaitForCompilationOnActivate : 1;
 #endif
 
+	/**
+	Sets the simulation cache to use for the component.
+	A null SimCache parameter will clear the active simulation cache.
+	When clearing a simulation cache that has been running you may wish to reset or continue, this option is only
+	valid when using a full simulation cache.  A renderer only cache will always reset as we can not continue the
+	simulation due to missing simulation data.
+	*/
 	UFUNCTION(BlueprintCallable, Category = SimCache)
-	void SetSimCache(UNiagaraSimCache* InSimCache);
+	void SetSimCache(UNiagaraSimCache* SimCache, bool bResetSystem = false);
 
+	/**
+	Get the active simulation cache, will return null if we do not have an active one.
+	*/
 	UFUNCTION(BlueprintCallable, Category = SimCache)
 	UNiagaraSimCache* GetSimCache() const;
 
+	/**
+	Clear any active simulation cache.
+	When clearing a simulation cache that has been running you may wish to reset or continue, this option is only
+	valid when using a full simulation cache.  A renderer only cache will always reset as we can not continue the
+	simulation due to missing simulation data.
+	*/
+	UFUNCTION(BlueprintCallable, Category = SimCache)
+	void ClearSimCache(bool bResetSystem = false) { SetSimCache(nullptr, bResetSystem); }
+
 	/** Set whether this component is allowed to perform scalability checks and potentially be culled etc. Occasionally it is useful to disable this for specific components. E.g. Effects on the local player. */
 	UFUNCTION(BlueprintSetter, Category = Scalability, meta = (Keywords = "LOD scalability"))
 	void SetAllowScalability(bool bAllow);
