- Niagara SimCache Reader DataInterface

diff --git a/Engine/Plugins/FX/Niagara/Shaders/Private/NiagaraDataInterfaceSimCacheReaderTemplate.ush b/Engine/Plugins/FX/Niagara/Shaders/Private/NiagaraDataInterfaceSimCacheReaderTemplate.ush
new file mode 100644
--- /dev/null
+++ b/Engine/Plugins/FX/Niagara/Shaders/Private/NiagaraDataInterfaceSimCacheReaderTemplate.ush
@@ -0,0 +1,119 @@
+// Copyright Epic Games, Inc. All Rights Reserved.
+
+int				{ParameterName}_NumFrames;
+int				{ParameterName}_NumEmitters;
+int				{ParameterName}_EmitterIndex;
+Buffer<uint>	{ParameterName}_CacheData;
+
+struct {ParameterName}_EmitterFrameData
+{
+	uint	NumInstances;
+	uint	Int32DataOffset;
+	uint	FloatDataOffset;
+	uint	HalfDataOffset;
+};
+
+void GetNumFrames_{ParameterName}(out int NumFrames)
+{
+	NumFrames = {ParameterName}_NumFrames;
+}
+
+void GetNumEmitters_{ParameterName}(out int NumEmitters)
+{
+	NumEmitters = {ParameterName}_NumEmitters;
+}
+
+void GetEmitterIndex_{ParameterName}(out int EmitterIndex)
+{
+	EmitterIndex = {ParameterName}_EmitterIndex;
+}
+
+{ParameterName}_EmitterFrameData GetEmitterFrameData_{ParameterName}(int FrameIndex, int EmitterIndex)
+{
+	{ParameterName}_EmitterFrameData FrameData;
+	FrameData.NumInstances		= 0;
+	FrameData.Int32DataOffset	= 0;
+	FrameData.FloatDataOffset	= 0;
+	FrameData.HalfDataOffset	= 0;
+
+	if ( FrameIndex >= 0 && FrameIndex < {ParameterName}_NumFrames && EmitterIndex >= 0 && EmitterIndex < {ParameterName}_NumEmitters )
+	{
+		uint FrameDataOffset		= (FrameIndex * 4 * {ParameterName}_NumEmitters) + EmitterIndex;
+		FrameData.NumInstances		= {ParameterName}_CacheData[FrameDataOffset + 0];
+		FrameData.Int32DataOffset	= {ParameterName}_CacheData[FrameDataOffset + 1];
+		FrameData.FloatDataOffset	= {ParameterName}_CacheData[FrameDataOffset + 2];
+		FrameData.HalfDataOffset	= {ParameterName}_CacheData[FrameDataOffset + 3];
+	}
+	return FrameData;
+}
+
+void GetInstanceInt_{ParameterName}(int FrameIndex, int EmitterIndex, int AttributeIndex, int InstanceIndex, out bool bValid, out int Value)
+{
+	bValid = false;
+	Value = 0;
+
+	{ParameterName}_EmitterFrameData FrameData = GetEmitterFrameData_{ParameterName}(FrameIndex, EmitterIndex);
+	if ( InstanceIndex >= 0 && InstanceIndex < FrameData.NumInstances )
+	{
+		bValid = true;
+		Value = asint({ParameterName}_CacheData[FrameData.Int32DataOffset + (AttributeIndex * FrameData.NumInstances) + InstanceIndex]);
+	}
+}
+
+void GetInstanceFloat_{ParameterName}(int FrameIndex, int EmitterIndex, int AttributeIndex, int InstanceIndex, out bool bValid, out float Value)
+{
+	bValid = false;
+	Value = 0.0f;
+
+	{ParameterName}_EmitterFrameData FrameData = GetEmitterFrameData_{ParameterName}(FrameIndex, EmitterIndex);
+	if ( InstanceIndex >= 0 && InstanceIndex < FrameData.NumInstances )
+	{
+		bValid = true;
+		Value = asfloat({ParameterName}_CacheData[FrameData.FloatDataOffset + (AttributeIndex * FrameData.NumInstances) + InstanceIndex]);
+	}
+}
+
+void GetInstanceFloat2_{ParameterName}(int FrameIndex, int EmitterIndex, int AttributeIndex, int InstanceIndex, out bool bValid, out float2 Value)
+{
+	bValid = false;
+	Value = 0.0f;
+
+	{ParameterName}_EmitterFrameData FrameData = GetEmitterFrameData_{ParameterName}(FrameIndex, EmitterIndex);
+	if ( InstanceIndex >= 0 && InstanceIndex < FrameData.NumInstances )
+	{
+		bValid = true;
+		Value.x = asfloat({ParameterName}_CacheData[FrameData.FloatDataOffset + ((AttributeIndex + 0) * FrameData.NumInstances) + InstanceIndex]);
+		Value.y = asfloat({ParameterName}_CacheData[FrameData.FloatDataOffset + ((AttributeIndex + 1) * FrameData.NumInstances) + InstanceIndex]);
+	}
+}
+
+void GetInstanceFloat3_{ParameterName}(int FrameIndex, int EmitterIndex, int AttributeIndex, int InstanceIndex, out bool bValid, out float3 Value)
+{
+	bValid = false;
+	Value = 0.0f;
+
+	{ParameterName}_EmitterFrameData FrameData = GetEmitterFrameData_{ParameterName}(FrameIndex, EmitterIndex);
+	if ( InstanceIndex >= 0 && InstanceIndex < FrameData.NumInstances )
+	{
+		bValid = true;
+		Value.x = asfloat({ParameterName}_CacheData[FrameData.FloatDataOffset + ((AttributeIndex + 0) * FrameData.NumInstances) + InstanceIndex]);
+		Value.y = asfloat({ParameterName}_CacheData[FrameData.FloatDataOffset + ((AttributeIndex + 1) * FrameData.NumInstances) + InstanceIndex]);
+		Value.z = asfloat({ParameterName}_CacheData[FrameData.FloatDataOffset + ((AttributeIndex + 2) * FrameData.NumInstances) + InstanceIndex]);
+	}
+}
+
+void GetInstanceFloat4_{ParameterName}(int FrameIndex, int EmitterIndex, int AttributeIndex, int InstanceIndex, out bool bValid, out float4 Value)
+{
+	bValid = false;
+	Value = 0.0f;
+
+	{ParameterName}_EmitterFrameData FrameData = GetEmitterFrameData_{ParameterName}(FrameIndex, EmitterIndex);
+	if ( InstanceIndex >= 0 && InstanceIndex < FrameData.NumInstances )
+	{
+		bValid = true;
+		Value.x = asfloat({ParameterName}_CacheData[FrameData.FloatDataOffset + ((AttributeIndex + 0) * FrameData.NumInstances) + InstanceIndex]);
+		Value.y = asfloat({ParameterName}_CacheData[FrameData.FloatDataOffset + ((AttributeIndex + 1) * FrameData.NumInstances) + InstanceIndex]);
+		Value.z = asfloat({ParameterName}_CacheData[FrameData.FloatDataOffset + ((AttributeIndex + 2) * FrameData.NumInstances) + InstanceIndex]);
+		Value.w = asfloat({ParameterName}_CacheData[FrameData.FloatDataOffset + ((AttributeIndex + 3) * FrameData.NumInstances) + InstanceIndex]);
+	}
+}
diff --git a/Engine/Plugins/FX/Niagara/Source/Niagara/Classes/NiagaraSimCache.h b/Engine/Plugins/FX/Niagara/Source/Niagara/Classes/NiagaraSimCache.h
--- a/Engine/Plugins/FX/Niagara/Source/Niagara/Classes/NiagaraSimCache.h
+++ b/Engine/Plugins/FX/Niagara/Source/Niagara/Classes/NiagaraSimCache.h
@@ -255,7 +255,7 @@ class NIAGARA_API UNiagaraSimCache : public UObject
 {
 	friend struct FNiagaraSimCacheAttributeReaderHelper;
 	friend struct FNiagaraSimCacheHelper;
-	friend class FNiagaraSimCacheViewModel;	//-TODO: Refactor to use API
+	friend struct FNiagaraSimCacheGpuResource;
 
 	GENERATED_UCLASS_BODY()
 
@@ -302,6 +302,9 @@ public:
 	UFUNCTION(BlueprintCallable, Category = NiagaraSimCache)
 	int GetNumEmitters() const { return CacheLayout.EmitterLayouts.Num(); }
 
+	/** Get the emitter index from a name */
+	int GetEmitterIndex(FName EmitterName) const;
+
 	/** Get the emitter name at the provided index. */
 	UFUNCTION(BlueprintCallable, Category = NiagaraSimCache)
 	FName GetEmitterName(int32 EmitterIndex) const { return CacheLayout.EmitterLayouts.IsValidIndex(EmitterIndex) ? CacheLayout.EmitterLayouts[EmitterIndex].LayoutName : NAME_None; }
@@ -310,6 +313,19 @@ public:
 	UFUNCTION(BlueprintCallable, Category = NiagaraSimCache)
 	TArray<FName> GetEmitterNames() const;
 
+	/**
+	Get number of active instances for the emitter at the given frame.
+	An EmitterIndex or INDEX_NONE will return information about the system instance.
+	*/
+	int GetEmitterNumInstances(int32 EmitterIndex, int32 FrameIndex) const;
+
+	/**
+	Runs a function for each attribute for the provided emitter index.
+	An EmitterIndex or INDEX_NONE will return information about the system instance.
+	Return true to continue iterating or false to stop
+	*/
+	void ForEachEmitterAttribute(int32 EmitterIndex, TFunction<bool(const FNiagaraSimCacheVariable&)> Function) const;
+
 	/**
 	Reads Niagara attributes by name from the cache frame and appends them into the relevant arrays.
 	When reading using this method the attributes for a FVector would come out as AoS rather than SoA so each component is wrote one by one
diff --git a/Engine/Plugins/FX/Niagara/Source/Niagara/Private/DataInterface/NiagaraDataInterfaceSimCacheReader.cpp b/Engine/Plugins/FX/Niagara/Source/Niagara/Private/DataInterface/NiagaraDataInterfaceSimCacheReader.cpp
new file mode 100644
--- /dev/null
+++ b/Engine/Plugins/FX/Niagara/Source/Niagara/Private/DataInterface/NiagaraDataInterfaceSimCacheReader.cpp
@@ -0,0 +1,607 @@
+// Copyright Epic Games, Inc. All Rights Reserved.
+
+#include "NiagaraDataInterfaceSimCacheReader.h"
+#include "NiagaraDataInterfaceUtilities.h"
+#include "NiagaraRenderer.h"
+#include "NiagaraSimCache.h"
+#include "NiagaraSimCacheAttributeReaderHelper.h"
+#include "NiagaraSimCacheGpuResource.h"
+#include "NiagaraShaderParametersBuilder.h"
+#include "NiagaraSystemInstance.h"
+
+#include "Internationalization/Internationalization.h"
+#include "ShaderParameterUtils.h"
+
+#define LOCTEXT_NAMESPACE "NiagaraDataInterfaceSimCacheReader"
+
+//////////////////////////////////////////////////////////////////////////
+
+namespace NDISimCacheReaderLocal
+{
+	static const TCHAR* TemplateShaderFile = TEXT("/Plugin/FX/Niagara/Private/NiagaraDataInterfaceSimCacheReaderTemplate.ush");
+	static const FName GetNumFrames("GetNumFrames");
+	static const FName GetNumEmitters("GetNumEmitters");
+	static const FName GetEmitterIndex("GetEmitterIndex");
+	static const FName GetInstanceInt("GetInstanceInt");
+	static const FName GetInstanceFloat("GetInstanceFloat");
+	static const FName GetInstanceVector2("GetInstanceVector2");
+	static const FName GetInstanceVector3("GetInstanceVector3");
+	static const FName GetInstanceVector4("GetInstanceVector4");
+	static const FName GetInstanceColor("GetInstanceColor");
+	static const FName GetInstanceQuat("GetInstanceQuat");
+	static const FName GetInstancePosition("GetInstancePosition");
+
+	static const FName NAME_Attribute("Attribute");
+
+	struct FInstanceData_GameThread
+	{
+		FNiagaraParameterDirectBinding<UObject*>	SimCacheBinding;
+		TWeakObjectPtr<UNiagaraSimCache>			WeakSimCache;
+		int32										EmitterIndex = INDEX_NONE;
+	};
+
+	struct FInstanceData_RenderThread
+	{
+		int32							EmitterIndex = INDEX_NONE;
+		FNiagaraSimCacheGpuResourcePtr	SimCacheResource;
+		TArray<int32>					ComponentOffsets;
+	};
+
+	struct FNDIProxy : public FNiagaraDataInterfaceProxy
+	{
+		virtual int32 PerInstanceDataPassedToRenderThreadSize() const { return 0; }
+		virtual void ConsumePerInstanceDataFromGameThread(void* PerInstanceData, const FNiagaraSystemInstanceID& Instance) {}
+
+		TMap<FNiagaraSystemInstanceID, FInstanceData_RenderThread>	PerInstanceData_RenderThread;
+	};
+
+	struct FGetFunctionInfo
+	{
+		explicit FGetFunctionInfo(FName InFunctionName, FNiagaraTypeDefinition InTypeDef, const TCHAR* InHlslType, const TCHAR* InHlslFunction)
+			: FunctionName(InFunctionName)
+			, TypeDef(InTypeDef)
+			, HlslType(InHlslType)
+			, HlslFunction(InHlslFunction)
+		{
+		}
+
+		FName					FunctionName;
+		FNiagaraTypeDefinition	TypeDef;
+		const TCHAR*			HlslType = nullptr;
+		const TCHAR*			HlslFunction = nullptr;
+	};
+
+	TConstArrayView<FGetFunctionInfo> GetFunctionInfos()
+	{
+		static const FGetFunctionInfo Infos[] =
+		{
+			FGetFunctionInfo(GetInstanceInt,		FNiagaraTypeDefinition::GetIntDef(),		TEXT("int"),	TEXT("GetInstanceInt")),
+			FGetFunctionInfo(GetInstanceFloat,		FNiagaraTypeDefinition::GetFloatDef(),		TEXT("float"),	TEXT("GetInstanceFloat")),
+			FGetFunctionInfo(GetInstanceVector2,	FNiagaraTypeDefinition::GetVec2Def(),		TEXT("float2"),	TEXT("GetInstanceFloat2")),
+			FGetFunctionInfo(GetInstanceVector3,	FNiagaraTypeDefinition::GetVec3Def(),		TEXT("float3"),	TEXT("GetInstanceFloat3")),
+			FGetFunctionInfo(GetInstanceVector4,	FNiagaraTypeDefinition::GetVec4Def(),		TEXT("float4"),	TEXT("GetInstanceFloat4")),
+			FGetFunctionInfo(GetInstanceColor,		FNiagaraTypeDefinition::GetColorDef(),		TEXT("float4"),	TEXT("GetInstanceFloat4")),
+			FGetFunctionInfo(GetInstanceQuat,		FNiagaraTypeDefinition::GetQuatDef(),		TEXT("float4"),	TEXT("GetInstanceFloat4")),
+			FGetFunctionInfo(GetInstancePosition,	FNiagaraTypeDefinition::GetPositionDef(),	TEXT("float3"),	TEXT("GetInstanceFloat3")),
+		};
+		return MakeArrayView(Infos);
+	}
+
+	static FNiagaraTypeDefinition FunctionNameToTypeDef(FName FunctionName)
+	{
+		for ( const FGetFunctionInfo& GetFuncInfo : GetFunctionInfos())
+		{
+			if (GetFuncInfo.FunctionName == FunctionName)
+			{
+				return GetFuncInfo.TypeDef;
+			}
+		}
+
+		check(false);
+		return FNiagaraTypeDefinition::GetFloatDef();
+	}
+
+	static int AttributeIndexFromVariable(UNiagaraSimCache* SimCache, FName FunctionName, FName AttributeName, int EmitterIndex)
+	{
+		const FNiagaraVariableBase Variable(FunctionNameToTypeDef(FunctionName), AttributeName);
+		return FNiagaraSimCacheAttributeReaderHelper::FindVariableIndex(SimCache, Variable, EmitterIndex);
+	}
+
+	static int ComponentIndexFromVariable(UNiagaraSimCache* SimCache, FName FunctionName, FName AttributeName, int EmitterIndex)
+	{
+		const FNiagaraVariableBase Variable(FunctionNameToTypeDef(FunctionName), AttributeName);
+		if ( const FNiagaraSimCacheVariable* CacheVariable = FNiagaraSimCacheAttributeReaderHelper::FindVariable(SimCache, Variable, EmitterIndex) )
+		{
+			if (Variable.GetType() == FNiagaraTypeDefinition::GetIntDef())
+			{
+				return CacheVariable->Int32Offset;
+			}
+			else
+			{
+				return CacheVariable->FloatOffset;
+			}
+		}
+		return INDEX_NONE;
+	}
+
+	template<typename TType>
+	static void VMInstanceGetFunction(FVectorVMExternalFunctionContext& Context, const TType DefaultValue, int32 AttributeIndex, TFunction<TType(FNiagaraSimCacheAttributeReaderHelper&, int32)> ReadFunction)
+	{
+		VectorVM::FUserPtrHandler<FInstanceData_GameThread> InstanceData_GT(Context);
+		FNDIInputParam<int32> InFrameIndex(Context);
+		FNDIInputParam<int32> InInstanceIndex(Context);
+		FNDIOutputParam<bool> OutValid(Context);
+		FNDIOutputParam<TType> OutValue(Context);
+
+		if (UNiagaraSimCache* InstanceSimCache = InstanceData_GT->WeakSimCache.Get())
+		{
+			FNiagaraSimCacheAttributeReaderHelper AttributeReader;
+			int CachedFrameIndex = INDEX_NONE;
+			
+			for (int32 i = 0; i < Context.GetNumInstances(); ++i)
+			{
+				const int32 FrameIndex = InFrameIndex.GetAndAdvance();
+				const int32 InstanceIndex = InInstanceIndex.GetAndAdvance();
+
+				if (CachedFrameIndex != FrameIndex)
+				{
+					CachedFrameIndex = FrameIndex;
+					AttributeReader.Initialize(InstanceSimCache, InstanceData_GT->EmitterIndex, AttributeIndex, FrameIndex);
+				}
+
+				if ( AttributeReader.IsValid() && InstanceIndex >= 0 && InstanceIndex < int32(AttributeReader.DataBuffers->NumInstances) )
+				{
+					OutValid.SetAndAdvance(true);
+					OutValue.SetAndAdvance(ReadFunction(AttributeReader, InstanceIndex));
+				}
+				else
+				{
+					OutValid.SetAndAdvance(false);
+					OutValue.SetAndAdvance(DefaultValue);
+				}
+			}
+		}
+		else
+		{
+			for (int32 i = 0; i < Context.GetNumInstances(); ++i)
+			{
+				OutValid.SetAndAdvance(false);
+				OutValue.SetAndAdvance(DefaultValue);
+			}
+		}
+	}
+}
+
+//////////////////////////////////////////////////////////////////////////
+
+UNiagaraDataInterfaceSimCacheReader::UNiagaraDataInterfaceSimCacheReader(FObjectInitializer const& ObjectInitializer)
+	: Super(ObjectInitializer)
+{
+	using namespace NDISimCacheReaderLocal;
+
+	SimCacheBinding.Parameter.SetType(FNiagaraTypeDefinition::GetSimCacheClassDef());
+
+	Proxy.Reset(new FNDIProxy());
+}
+
+void UNiagaraDataInterfaceSimCacheReader::PostInitProperties()
+{
+	Super::PostInitProperties();
+
+	//Can we register data interfaces as regular types and fold them into the FNiagaraVariable framework for UI and function calls etc?
+	if (HasAnyFlags(RF_ClassDefaultObject))
+	{
+		ENiagaraTypeRegistryFlags Flags = ENiagaraTypeRegistryFlags::AllowAnyVariable | ENiagaraTypeRegistryFlags::AllowParameter;
+		FNiagaraTypeRegistry::Register(FNiagaraTypeDefinition(GetClass()), Flags);
+	}
+}
+
+bool UNiagaraDataInterfaceSimCacheReader::Equals(const UNiagaraDataInterface* Other) const
+{
+	if (!Super::Equals(Other))
+	{
+		return false;
+	}
+
+	auto* OtherTyped = CastChecked<const UNiagaraDataInterfaceSimCacheReader>(Other);
+	return
+		OtherTyped->SimCacheBinding == SimCacheBinding &&
+		OtherTyped->SimCache == SimCache &&
+		OtherTyped->EmitterBinding == EmitterBinding;
+}
+
+bool UNiagaraDataInterfaceSimCacheReader::CopyToInternal(UNiagaraDataInterface* Destination) const
+{
+	if (!Super::CopyToInternal(Destination))
+	{
+		return false;
+	}
+
+	auto* DestinationTyped = CastChecked<UNiagaraDataInterfaceSimCacheReader>(Destination);
+	DestinationTyped->SimCacheBinding = SimCacheBinding;
+	DestinationTyped->SimCache = SimCache;
+	DestinationTyped->EmitterBinding = EmitterBinding;
+
+	return true;
+}
+
+bool UNiagaraDataInterfaceSimCacheReader::InitPerInstanceData(void* PerInstanceData, FNiagaraSystemInstance* SystemInstance)
+{
+	using namespace NDISimCacheReaderLocal;
+
+	FInstanceData_GameThread* InstanceData_GT = new(PerInstanceData) FInstanceData_GameThread();
+	InstanceData_GT->SimCacheBinding.Init(SystemInstance->GetInstanceParameters(), SimCacheBinding.Parameter);
+
+	return true;
+}
+
+void UNiagaraDataInterfaceSimCacheReader::DestroyPerInstanceData(void* PerInstanceData, FNiagaraSystemInstance* SystemInstance)
+{
+	using namespace NDISimCacheReaderLocal;
+
+	FInstanceData_GameThread* InstanceData_GT = reinterpret_cast<FInstanceData_GameThread*>(PerInstanceData);
+	InstanceData_GT->~FInstanceData_GameThread();
+
+	if ( IsUsedByGPUEmitter() )
+	{
+		ENQUEUE_RENDER_COMMAND(NDISimCacheReader_DestroyRT)
+		(
+			[Proxy_RT=GetProxyAs<FNDIProxy>(), InstanceID=SystemInstance->GetId()](FRHICommandList& CmdList)
+			{
+				Proxy_RT->PerInstanceData_RenderThread.Remove(InstanceID);
+			}
+		);
+	}
+}
+
+int32 UNiagaraDataInterfaceSimCacheReader::PerInstanceDataSize() const
+{
+	using namespace NDISimCacheReaderLocal;
+	return sizeof(FInstanceData_GameThread);
+}
+
+bool UNiagaraDataInterfaceSimCacheReader::PerInstanceTick(void* PerInstanceData, FNiagaraSystemInstance* SystemInstance, float DeltaSeconds)
+{
+	using namespace NDISimCacheReaderLocal;
+
+	check(PerInstanceData && SystemInstance);
+
+	FInstanceData_GameThread* InstanceData_GT = reinterpret_cast<FInstanceData_GameThread*>(PerInstanceData);
+
+	// Has the cache changed at all?
+	UNiagaraSimCache* InstanceSimCache = InstanceData_GT->SimCacheBinding.GetValueOrDefault<UNiagaraSimCache>(SimCache.Get());
+	if (InstanceData_GT->WeakSimCache.Get() != InstanceSimCache)
+	{
+		InstanceData_GT->EmitterIndex = InstanceSimCache->GetEmitterIndex(EmitterBinding);
+
+		TArray<int32> ComponentOffsets;
+		FNiagaraDataInterfaceUtilities::ForEachGpuFunction(
+			this, SystemInstance,
+			[&](const FNiagaraDataInterfaceGeneratedFunction& FunctionInfo)
+			{
+				if (FunctionInfo.Specifiers.Num() != 1)
+				{
+					return true;
+				}
+
+				check(FunctionInfo.Specifiers[0].Key == NAME_Attribute);
+
+				const int32 ComponentIndex = ComponentIndexFromVariable(InstanceSimCache, FunctionInfo.DefinitionName, FunctionInfo.Specifiers[0].Value, InstanceData_GT->EmitterIndex);
+				ComponentOffsets.Add(ComponentIndex);
+
+				return true;
+			}
+		);
+
+		const int32 AlignedAttributes = ComponentOffsets.Num() == 0 ? 4 : Align(ComponentOffsets.Num(), 4);
+		for ( int32 i= ComponentOffsets.Num(); i < AlignedAttributes; ++i )
+		{
+			ComponentOffsets.Add(INDEX_NONE);
+		}
+
+		FNiagaraSimCacheGpuResourcePtr SimCacheResource = FNiagaraSimCacheGpuResource::CreateResource(InstanceSimCache);
+
+		ENQUEUE_RENDER_COMMAND(NDISimCacheReader_InitRT)
+		(
+			[Proxy_RT=GetProxyAs<FNDIProxy>(), InstanceID_RT=SystemInstance->GetId(), EmitterIndex_RT=InstanceData_GT->EmitterIndex, SimCacheResource_RT=MoveTemp(SimCacheResource), ComponentOffsets_RT=MoveTemp(ComponentOffsets)](FRHICommandList& CmdList)
+			{
+				FInstanceData_RenderThread* InstanceData_RT = &Proxy_RT->PerInstanceData_RenderThread.FindOrAdd(InstanceID_RT);
+				InstanceData_RT->SimCacheResource = SimCacheResource_RT;
+				InstanceData_RT->EmitterIndex = EmitterIndex_RT;
+				InstanceData_RT->ComponentOffsets = ComponentOffsets_RT;
+			}
+		);
+	}
+	else if (InstanceSimCache == nullptr)
+	{
+		InstanceData_GT->EmitterIndex = INDEX_NONE;
+		if ( IsUsedByGPUEmitter() )
+		{
+			ENQUEUE_RENDER_COMMAND(NDISimCacheReader_DestroyRT)
+			(
+				[Proxy_RT=GetProxyAs<FNDIProxy>(), InstanceID=SystemInstance->GetId()](FRHICommandList& CmdList)
+				{
+					Proxy_RT->PerInstanceData_RenderThread.Remove(InstanceID);
+				}
+			);
+		}
+	}
+	InstanceData_GT->WeakSimCache = InstanceSimCache;
+
+	return false;
+}
+
+void UNiagaraDataInterfaceSimCacheReader::GetFunctions(TArray<FNiagaraFunctionSignature>& OutFunctions)
+{
+	using namespace NDISimCacheReaderLocal;
+
+	// Reserve space
+	OutFunctions.Reserve(OutFunctions.Num() + 3 + GetFunctionInfos().Num());
+
+	// Build default signature
+	FNiagaraFunctionSignature DefaultSignature;
+	DefaultSignature.bMemberFunction = true;
+	DefaultSignature.bRequiresContext = false;
+	DefaultSignature.bSupportsCPU = true;
+	DefaultSignature.bSupportsGPU = true;
+	DefaultSignature.Inputs.Emplace(FNiagaraTypeDefinition(GetClass()), TEXT("SimCacheReader"));
+
+	// Utility functions
+	{
+		FNiagaraFunctionSignature& Sig = OutFunctions.Add_GetRef(DefaultSignature);
+		Sig.Name = GetNumFrames;
+		Sig.Outputs.Emplace(FNiagaraTypeDefinition::GetIntDef(), TEXT("NumFrames"));
+		Sig.SetDescription(LOCTEXT("GetNumFramesDesc", "Returns the number of frames stored inside the cache"));
+	}
+	{
+		FNiagaraFunctionSignature& Sig = OutFunctions.Add_GetRef(DefaultSignature);
+		Sig.Name = GetNumEmitters;
+		Sig.Outputs.Emplace(FNiagaraTypeDefinition::GetIntDef(), TEXT("NumEmitters"));
+		Sig.SetDescription(LOCTEXT("GetNumEmittersDesc", "Returns the number of emitters stored inside the cache"));
+	}
+	{
+		FNiagaraFunctionSignature& Sig = OutFunctions.Add_GetRef(DefaultSignature);
+		Sig.Name = GetEmitterIndex;
+		Sig.Outputs.Emplace(FNiagaraTypeDefinition::GetIntDef(), TEXT("EmitterIndex"));
+		Sig.SetDescription(LOCTEXT("GetEmitterIndexDesc", "Returns the emitter index the data interface is bound to"));		
+	}
+
+	// Attribute Reading functions
+	FNiagaraFunctionSignature AttributeFunctionSignature = DefaultSignature;
+	AttributeFunctionSignature.Inputs.Emplace(FNiagaraTypeDefinition::GetIntDef(), TEXT("FrameIndex"));
+	AttributeFunctionSignature.Inputs.Emplace(FNiagaraTypeDefinition::GetIntDef(), TEXT("InstanceIndex"));
+	AttributeFunctionSignature.Outputs.Emplace(FNiagaraTypeDefinition::GetBoolDef(), TEXT("Valid"));
+	AttributeFunctionSignature.FunctionSpecifiers.Emplace(NAME_Attribute);
+	AttributeFunctionSignature.SetDescription(LOCTEXT("GetInstanceDesc", "Returns the instance value, valid will be false if the instance / frame / emitter were not valid or out of range."));
+
+	for (const FGetFunctionInfo& GenFuncInfo : GetFunctionInfos())
+	{
+		FNiagaraFunctionSignature & Sig = OutFunctions.Add_GetRef(AttributeFunctionSignature);
+		Sig.Name = GenFuncInfo.FunctionName;
+		Sig.Outputs.Emplace(GenFuncInfo.TypeDef, TEXT("Value"));
+	}
+}
+
+void UNiagaraDataInterfaceSimCacheReader::GetVMExternalFunction(const FVMExternalFunctionBindingInfo& BindingInfo, void* PerInstanceData, FVMExternalFunction& OutFunc)
+{
+	using namespace NDISimCacheReaderLocal;
+
+	FInstanceData_GameThread* InstanceData = reinterpret_cast<FInstanceData_GameThread*>(PerInstanceData);
+	UNiagaraSimCache* InstanceSimCache = InstanceData->WeakSimCache.Get();
+	if ( BindingInfo.Name == GetNumFrames )
+	{
+		OutFunc = FVMExternalFunction::CreateUObject(this, &UNiagaraDataInterfaceSimCacheReader::VMGetNumFrames);
+	}
+	else if (BindingInfo.Name == GetNumEmitters)
+	{
+		OutFunc = FVMExternalFunction::CreateUObject(this, &UNiagaraDataInterfaceSimCacheReader::VMGetNumEmitters);
+	}
+	else if (BindingInfo.Name == GetInstanceInt)
+	{
+		const int32 AttributeIndex = AttributeIndexFromVariable(InstanceSimCache, BindingInfo.Name, BindingInfo.FunctionSpecifiers[0].Value, InstanceData->EmitterIndex);
+		OutFunc = FVMExternalFunction::CreateUObject(this, &UNiagaraDataInterfaceSimCacheReader::VMGetInstanceInt, AttributeIndex);
+	}
+	else if (BindingInfo.Name == GetInstanceFloat)
+	{
+		const int32 AttributeIndex = AttributeIndexFromVariable(InstanceSimCache, BindingInfo.Name, BindingInfo.FunctionSpecifiers[0].Value, InstanceData->EmitterIndex);
+		OutFunc = FVMExternalFunction::CreateUObject(this, &UNiagaraDataInterfaceSimCacheReader::VMGetInstanceFloat, AttributeIndex);
+	}
+	else if (BindingInfo.Name == GetInstanceVector2)
+	{
+		const int32 AttributeIndex = AttributeIndexFromVariable(InstanceSimCache, BindingInfo.Name, BindingInfo.FunctionSpecifiers[0].Value, InstanceData->EmitterIndex);
+		OutFunc = FVMExternalFunction::CreateUObject(this, &UNiagaraDataInterfaceSimCacheReader::VMGetInstanceVector2, AttributeIndex);
+	}
+	else if (BindingInfo.Name == GetInstanceVector3)
+	{
+		const int32 AttributeIndex = AttributeIndexFromVariable(InstanceSimCache, BindingInfo.Name, BindingInfo.FunctionSpecifiers[0].Value, InstanceData->EmitterIndex);
+		OutFunc = FVMExternalFunction::CreateUObject(this, &UNiagaraDataInterfaceSimCacheReader::VMGetInstanceVector3, AttributeIndex);
+	}
+	else if (BindingInfo.Name == GetInstanceVector2)
+	{
+		const int32 AttributeIndex = AttributeIndexFromVariable(InstanceSimCache, BindingInfo.Name, BindingInfo.FunctionSpecifiers[0].Value, InstanceData->EmitterIndex);
+		OutFunc = FVMExternalFunction::CreateUObject(this, &UNiagaraDataInterfaceSimCacheReader::VMGetInstanceVector4, AttributeIndex);
+	}
+	else if (BindingInfo.Name == GetInstanceColor)
+	{
+		const int32 AttributeIndex = AttributeIndexFromVariable(InstanceSimCache, BindingInfo.Name, BindingInfo.FunctionSpecifiers[0].Value, InstanceData->EmitterIndex);
+		OutFunc = FVMExternalFunction::CreateUObject(this, &UNiagaraDataInterfaceSimCacheReader::VMGetInstanceColor, AttributeIndex);
+	}
+	else if (BindingInfo.Name == GetInstanceQuat)
+	{
+		const int32 AttributeIndex = AttributeIndexFromVariable(InstanceSimCache, BindingInfo.Name, BindingInfo.FunctionSpecifiers[0].Value, InstanceData->EmitterIndex);
+		OutFunc = FVMExternalFunction::CreateUObject(this, &UNiagaraDataInterfaceSimCacheReader::VMGetInstanceQuat, AttributeIndex, false);
+	}
+	else if (BindingInfo.Name == GetInstancePosition)
+	{
+		const int32 AttributeIndex = AttributeIndexFromVariable(InstanceSimCache, BindingInfo.Name, BindingInfo.FunctionSpecifiers[0].Value, InstanceData->EmitterIndex);
+		OutFunc = FVMExternalFunction::CreateUObject(this, &UNiagaraDataInterfaceSimCacheReader::VMGetInstancePosition, AttributeIndex, false);
+	}
+}
+
+#if WITH_EDITORONLY_DATA
+bool UNiagaraDataInterfaceSimCacheReader::AppendCompileHash(FNiagaraCompileHashVisitor* InVisitor) const
+{
+	bool bSuccess = Super::AppendCompileHash(InVisitor);
+	FSHAHash Hash = GetShaderFileHash(NDISimCacheReaderLocal::TemplateShaderFile, EShaderPlatform::SP_PCD3D_SM5);
+	InVisitor->UpdateString(TEXT("NDISimCacheReaderTemplateShaderFile"), Hash.ToString());
+	InVisitor->UpdateShaderParameters<FShaderParameters>();
+	return bSuccess;
+}
+
+void UNiagaraDataInterfaceSimCacheReader::GetParameterDefinitionHLSL(const FNiagaraDataInterfaceGPUParamInfo& ParameterInfo, FString& OutHLSL)
+{
+	TMap<FString, FStringFormatArg> TemplateArgs =
+	{
+		{TEXT("ParameterName"),	ParameterInfo.DataInterfaceHLSLSymbol},
+	};
+
+	OutHLSL.Appendf(TEXT("int4 %s_ComponentOffsets[%d];\n"), *ParameterInfo.DataInterfaceHLSLSymbol, FMath::DivideAndRoundUp(ParameterInfo.GeneratedFunctions.Num(), 4));
+
+	FString TemplateFile;
+	LoadShaderSourceFile(NDISimCacheReaderLocal::TemplateShaderFile, EShaderPlatform::SP_PCD3D_SM5, &TemplateFile, nullptr);
+	OutHLSL += FString::Format(*TemplateFile, TemplateArgs);
+}
+
+bool UNiagaraDataInterfaceSimCacheReader::GetFunctionHLSL(const FNiagaraDataInterfaceGPUParamInfo& ParamInfo, const FNiagaraDataInterfaceGeneratedFunction& FunctionInfo, int FunctionInstanceIndex, FString& OutHLSL)
+{
+	using namespace NDISimCacheReaderLocal;
+
+	const int32 Float4Index = FunctionInstanceIndex / 4;
+	const int32 ElementIndex = FunctionInstanceIndex % 4;
+
+	for (const FGetFunctionInfo& GenFuncInfo : GetFunctionInfos())
+	{
+		if ( GenFuncInfo.FunctionName == FunctionInfo.DefinitionName )
+		{
+			OutHLSL.Appendf(TEXT("void %s(int FrameIndex, int InstanceIndex, out bool bValid, out %s Value)\n"), *FunctionInfo.InstanceName, GenFuncInfo.HlslType);
+			OutHLSL.Append(TEXT("{\n"));
+			OutHLSL.Appendf(TEXT("	int AttributeIndex = %s_ComponentOffsets[%d][%d];\n"), *ParamInfo.DataInterfaceHLSLSymbol, Float4Index, ElementIndex);
+			OutHLSL.Appendf(TEXT("	%s_%s(FrameIndex, %s_EmitterIndex, AttributeIndex, InstanceIndex, bValid, Value);\n"), GenFuncInfo.HlslFunction, *ParamInfo.DataInterfaceHLSLSymbol, *ParamInfo.DataInterfaceHLSLSymbol);
+			OutHLSL.Append(TEXT("}\n"));
+			return true;
+		}
+	}
+
+	return true;
+}
+#endif
+
+void UNiagaraDataInterfaceSimCacheReader::BuildShaderParameters(FNiagaraShaderParametersBuilder& ShaderParametersBuilder) const
+{
+	ShaderParametersBuilder.AddNestedStruct<FShaderParameters>();
+
+	const int32 ComponentOffsetsSize = FMath::Max(FMath::DivideAndRoundUp(ShaderParametersBuilder.GetGeneratedFunctions().Num(), 4), 1);
+	ShaderParametersBuilder.AddLooseParamArray<FUintVector4>(TEXT("ComponentOffsets"), ComponentOffsetsSize);
+}
+
+void UNiagaraDataInterfaceSimCacheReader::SetShaderParameters(const FNiagaraDataInterfaceSetShaderParametersContext& Context) const
+{
+	using namespace NDISimCacheReaderLocal;
+
+	const FNDIProxy& DIProxy = Context.GetProxy<FNDIProxy>();
+	const FInstanceData_RenderThread* InstanceData_RT = DIProxy.PerInstanceData_RenderThread.Find(Context.GetSystemInstanceID());
+
+	FShaderParameters* ShaderParameters = Context.GetParameterNestedStruct<FShaderParameters>();
+	if (InstanceData_RT && InstanceData_RT->SimCacheResource->GetNumFrames() > 0)
+	{
+		ShaderParameters->NumFrames		= InstanceData_RT->SimCacheResource->GetNumFrames();
+		ShaderParameters->NumEmitters	= InstanceData_RT->SimCacheResource->GetNumEmitters();
+		ShaderParameters->EmitterIndex	= InstanceData_RT->EmitterIndex;
+		ShaderParameters->CacheData		= InstanceData_RT->SimCacheResource->GetBufferSRV();
+
+		check((InstanceData_RT->ComponentOffsets.Num() % 4) == 0 && InstanceData_RT->ComponentOffsets.Num() > 0);
+		const int32 ComponentOffsetsSize = InstanceData_RT->ComponentOffsets.Num() / 4;
+		TArrayView<FIntVector4> ComponentOffsets = Context.GetParameterLooseArray<FIntVector4>(ComponentOffsetsSize);
+		FMemory::Memcpy(ComponentOffsets.GetData(), InstanceData_RT->ComponentOffsets.GetData(), InstanceData_RT->ComponentOffsets.Num() * InstanceData_RT->ComponentOffsets.GetTypeSize());
+	}
+	else
+	{
+		ShaderParameters->NumFrames		= 0;
+		ShaderParameters->NumEmitters	= 0;
+		ShaderParameters->EmitterIndex	= INDEX_NONE;
+		ShaderParameters->CacheData		= FNiagaraRenderer::GetDummyUIntBuffer();
+
+		TArrayView<FIntVector4> ComponentOffsets = Context.GetParameterLooseArray<FIntVector4>(1);
+		FMemory::Memset(ComponentOffsets.GetData(), 0, sizeof(FIntVector4));
+	}
+}
+
+void UNiagaraDataInterfaceSimCacheReader::VMGetNumFrames(FVectorVMExternalFunctionContext& Context)
+{
+	using namespace NDISimCacheReaderLocal;
+
+	VectorVM::FUserPtrHandler<FInstanceData_GameThread> InstanceData_GT(Context);
+	FNDIOutputParam<int32> OutNumFrames(Context);
+	
+	UNiagaraSimCache* InstanceSimCache = InstanceData_GT->WeakSimCache.Get();
+	const uint32 NumFrames = InstanceSimCache ? InstanceSimCache->GetNumFrames() : 0;
+	for (int32 i = 0; i < Context.GetNumInstances(); ++i)
+	{
+		OutNumFrames.SetAndAdvance(NumFrames);
+	}
+}
+
+void UNiagaraDataInterfaceSimCacheReader::VMGetNumEmitters(FVectorVMExternalFunctionContext& Context)
+{
+	using namespace NDISimCacheReaderLocal;
+
+	VectorVM::FUserPtrHandler<FInstanceData_GameThread> InstanceData_GT(Context);
+	FNDIOutputParam<int32> OutNumEmitters(Context);
+
+	UNiagaraSimCache* InstanceSimCache = InstanceData_GT->WeakSimCache.Get();
+	const uint32 NumEmitters = InstanceSimCache ? InstanceSimCache->GetNumEmitters() : 0;
+	for (int32 i = 0; i < Context.GetNumInstances(); ++i)
+	{
+		OutNumEmitters.SetAndAdvance(NumEmitters);
+	}
+}
+
+void UNiagaraDataInterfaceSimCacheReader::VMGetInstanceInt(FVectorVMExternalFunctionContext& Context, int AttributeIndex)
+{
+	using namespace NDISimCacheReaderLocal;
+	VMInstanceGetFunction<int32>(Context, 0, AttributeIndex, [](FNiagaraSimCacheAttributeReaderHelper& AttributeReader, int32 InstanceIndex) { return AttributeReader.ReadInt(InstanceIndex); });
+}
+
+void UNiagaraDataInterfaceSimCacheReader::VMGetInstanceFloat(FVectorVMExternalFunctionContext& Context, int AttributeIndex)
+{
+	using namespace NDISimCacheReaderLocal;
+	VMInstanceGetFunction<float>(Context, 0.0f, AttributeIndex, [](FNiagaraSimCacheAttributeReaderHelper& AttributeReader, int32 InstanceIndex) { return AttributeReader.ReadFloat(InstanceIndex); });
+}
+
+void UNiagaraDataInterfaceSimCacheReader::VMGetInstanceVector2(FVectorVMExternalFunctionContext& Context, int AttributeIndex)
+{
+	using namespace NDISimCacheReaderLocal;
+	VMInstanceGetFunction<FVector2f>(Context, FVector2f::ZeroVector, AttributeIndex, [](FNiagaraSimCacheAttributeReaderHelper& AttributeReader, int32 InstanceIndex) { return AttributeReader.ReadFloat2f(InstanceIndex); });
+}
+
+void UNiagaraDataInterfaceSimCacheReader::VMGetInstanceVector3(FVectorVMExternalFunctionContext& Context, int AttributeIndex)
+{
+	using namespace NDISimCacheReaderLocal;
+	VMInstanceGetFunction<FVector3f>(Context, FVector3f::ZeroVector, AttributeIndex, [](FNiagaraSimCacheAttributeReaderHelper& AttributeReader, int32 InstanceIndex) { return AttributeReader.ReadFloat3f(InstanceIndex); });
+}
+
+void UNiagaraDataInterfaceSimCacheReader::VMGetInstanceVector4(FVectorVMExternalFunctionContext& Context, int AttributeIndex)
+{
+	using namespace NDISimCacheReaderLocal;
+	VMInstanceGetFunction<FVector4f>(Context, FVector4f::Zero(), AttributeIndex, [](FNiagaraSimCacheAttributeReaderHelper& AttributeReader, int32 InstanceIndex) { return AttributeReader.ReadFloat4f(InstanceIndex); });
+}
+
+void UNiagaraDataInterfaceSimCacheReader::VMGetInstanceColor(FVectorVMExternalFunctionContext& Context, int AttributeIndex)
+{
+	using namespace NDISimCacheReaderLocal;
+	VMInstanceGetFunction<FLinearColor>(Context, FLinearColor::White, AttributeIndex, [](FNiagaraSimCacheAttributeReaderHelper& AttributeReader, int32 InstanceIndex) { return AttributeReader.ReadColor(InstanceIndex); });
+}
+
+void UNiagaraDataInterfaceSimCacheReader::VMGetInstanceQuat(FVectorVMExternalFunctionContext& Context, int AttributeIndex, bool bRebase)
+{
+	using namespace NDISimCacheReaderLocal;
+	VMInstanceGetFunction<FQuat4f>(Context, FQuat4f::Identity, AttributeIndex, [](FNiagaraSimCacheAttributeReaderHelper& AttributeReader, int32 InstanceIndex) { return AttributeReader.ReadQuat4f(InstanceIndex); });
+}
+
+void UNiagaraDataInterfaceSimCacheReader::VMGetInstancePosition(FVectorVMExternalFunctionContext& Context, int AttributeIndex, bool bRebase)
+{
+	using namespace NDISimCacheReaderLocal;
+	VMInstanceGetFunction<FVector3f>(Context, FVector3f::ZeroVector, AttributeIndex, [](FNiagaraSimCacheAttributeReaderHelper& AttributeReader, int32 InstanceIndex) { return AttributeReader.ReadFloat3f(InstanceIndex); });
+}
+
+#undef LOCTEXT_NAMESPACE
diff --git a/Engine/Plugins/FX/Niagara/Source/Niagara/Private/DataInterface/NiagaraDataInterfaceSimCacheReader.h b/Engine/Plugins/FX/Niagara/Source/Niagara/Private/DataInterface/NiagaraDataInterfaceSimCacheReader.h
new file mode 100644
--- /dev/null
+++ b/Engine/Plugins/FX/Niagara/Source/Niagara/Private/DataInterface/NiagaraDataInterfaceSimCacheReader.h
@@ -0,0 +1,80 @@
+// Copyright Epic Games, Inc. All Rights Reserved.
+#pragma once
+
+#include "NiagaraCommon.h"
+#include "NiagaraShared.h"
+#include "NiagaraDataInterface.h"
+#include "NiagaraDataInterfaceSimCacheReader.generated.h"
+
+class UNiagaraSimCache;
+
+/**
+Data interface to read properties from a Niagara Simulation Cache
+*/
+UCLASS(EditInlineNew, Category = "DataInterface", meta=(DisplayName="SimCache Reader", Experimental))
+class UNiagaraDataInterfaceSimCacheReader : public UNiagaraDataInterface
+{
+	GENERATED_UCLASS_BODY()
+
+	BEGIN_SHADER_PARAMETER_STRUCT(FShaderParameters, )
+		SHADER_PARAMETER(int32,				NumFrames)
+		SHADER_PARAMETER(int32,				NumEmitters)
+		SHADER_PARAMETER(int32,				EmitterIndex)
+		SHADER_PARAMETER_SRV(Buffer<uint>,	CacheData)
+	END_SHADER_PARAMETER_STRUCT()
+
+public:
+	/** User parameter Object binding if this is not a valid sim cache the default one will be used instead. */
+	UPROPERTY(EditAnywhere, Category = "SimCacheReader")
+	FNiagaraUserParameterBinding SimCacheBinding;
+
+	/** Optional source SimCache to use, if the user parameter binding is valid this will be ignored. */
+	UPROPERTY(EditAnywhere, Category = "SimCacheReader")
+	TObjectPtr<UNiagaraSimCache> SimCache;
+
+	/** Which Emitter we should read from within the SimCache. */
+	UPROPERTY(EditAnywhere, Category = "Emitter")
+	FName EmitterBinding;
+
+	//UNiagaraDataInterface Interface
+	virtual void PostInitProperties() override;
+	// UObject Interface End
+
+	// UNiagaraDataInterface Interface Begin
+	virtual bool CanExecuteOnTarget(ENiagaraSimTarget Target) const override { return true; }
+
+	virtual bool Equals(const UNiagaraDataInterface* Other) const override;
+	virtual bool CopyToInternal(UNiagaraDataInterface* Destination) const override;
+
+	virtual bool InitPerInstanceData(void* PerInstanceData, FNiagaraSystemInstance* SystemInstance) override;
+	virtual void DestroyPerInstanceData(void* PerInstanceData, FNiagaraSystemInstance* SystemInstance) override;
+	virtual int32 PerInstanceDataSize() const override;
+
+	virtual bool HasPreSimulateTick() const override { return true; }
+	virtual bool PerInstanceTick(void* PerInstanceData, FNiagaraSystemInstance* SystemInstance, float DeltaSeconds) override;
+
+	virtual void GetFunctions(TArray<FNiagaraFunctionSignature>& OutFunctions) override;
+	virtual void GetVMExternalFunction(const FVMExternalFunctionBindingInfo& BindingInfo, void* InstanceData, FVMExternalFunction& OutFunc) override;
+#if WITH_EDITORONLY_DATA
+	virtual bool AppendCompileHash(FNiagaraCompileHashVisitor* InVisitor) const override;
+	virtual void GetParameterDefinitionHLSL(const FNiagaraDataInterfaceGPUParamInfo& ParamInfo, FString& OutHLSL) override;
+	virtual bool GetFunctionHLSL(const FNiagaraDataInterfaceGPUParamInfo& ParamInfo, const FNiagaraDataInterfaceGeneratedFunction& FunctionInfo, int FunctionInstanceIndex, FString& OutHLSL) override;
+#endif
+	virtual bool UseLegacyShaderBindings() const override { return false; }
+	virtual void BuildShaderParameters(FNiagaraShaderParametersBuilder& ShaderParametersBuilder) const override;
+	virtual void SetShaderParameters(const FNiagaraDataInterfaceSetShaderParametersContext& Context) const override;
+	//UNiagaraDataInterface Interface
+
+private:
+	void VMGetNumFrames(FVectorVMExternalFunctionContext& Context);
+	void VMGetNumEmitters(FVectorVMExternalFunctionContext& Context);
+	void VMGetEmitterIndex(FVectorVMExternalFunctionContext& Context);
+	void VMGetInstanceInt(FVectorVMExternalFunctionContext& Context, int AttributeIndex);
+	void VMGetInstanceFloat(FVectorVMExternalFunctionContext& Context, int AttributeIndex);
+	void VMGetInstanceVector2(FVectorVMExternalFunctionContext& Context, int AttributeIndex);
+	void VMGetInstanceVector3(FVectorVMExternalFunctionContext& Context, int AttributeIndex);
+	void VMGetInstanceVector4(FVectorVMExternalFunctionContext& Context, int AttributeIndex);
+	void VMGetInstanceColor(FVectorVMExternalFunctionContext& Context, int AttributeIndex);
+	void VMGetInstanceQuat(FVectorVMExternalFunctionContext& Context, int AttributeIndex, bool bRebase);
+	void VMGetInstancePosition(FVectorVMExternalFunctionContext& Context, int AttributeIndex, bool bRebase);
+};
diff --git a/Engine/Plugins/FX/Niagara/Source/Niagara/Private/NiagaraModule.cpp b/Engine/Plugins/FX/Niagara/Source/Niagara/Private/NiagaraModule.cpp
--- a/Engine/Plugins/FX/Niagara/Source/Niagara/Private/NiagaraModule.cpp
+++ b/Engine/Plugins/FX/Niagara/Source/Niagara/Private/NiagaraModule.cpp
@@ -746,6 +746,7 @@ FNiagaraTypeDefinition FNiagaraTypeDefinition::UMaterialDef;
 FNiagaraTypeDefinition FNiagaraTypeDefinition::UTextureDef;
 FNiagaraTypeDefinition FNiagaraTypeDefinition::UTextureRenderTargetDef;
 FNiagaraTypeDefinition FNiagaraTypeDefinition::UStaticMeshDef;
+FNiagaraTypeDefinition FNiagaraTypeDefinition::USimCacheClassDef;
 FNiagaraTypeDefinition FNiagaraTypeDefinition::WildcardDef;
 
 TSet<UScriptStruct*> FNiagaraTypeDefinition::NumericStructs;
@@ -822,6 +823,7 @@ void FNiagaraTypeDefinition::Init()
 	UTextureDef = FNiagaraTypeDefinition(UTextureClass);
 	UTextureRenderTargetDef = FNiagaraTypeDefinition(UTextureRenderTargetClass);
 	UStaticMeshDef = FNiagaraTypeDefinition(UStaticMesh::StaticClass());
+	USimCacheClassDef = FNiagaraTypeDefinition(UNiagaraSimCache::StaticClass());
 
 	CollisionEventDef = FNiagaraTypeDefinition(FNiagaraCollisionEventPayload::StaticStruct());
 	NumericStructs.Add(NumericStruct);
@@ -1041,6 +1043,7 @@ void FNiagaraTypeDefinition::RecreateUserDefinedTypeRegistry()
 	FNiagaraTypeRegistry::Register(UTextureDef, VarFlags);
 	FNiagaraTypeRegistry::Register(UTextureRenderTargetDef, VarFlags);
 	FNiagaraTypeRegistry::Register(UStaticMeshDef, VarFlags);
+	FNiagaraTypeRegistry::Register(USimCacheClassDef, VarFlags);
 	FNiagaraTypeRegistry::Register(StaticEnum<ENiagaraLegacyTrailWidthMode>(), ParamFlags);
 
 	
diff --git a/Engine/Plugins/FX/Niagara/Source/Niagara/Private/NiagaraSimCache.cpp b/Engine/Plugins/FX/Niagara/Source/Niagara/Private/NiagaraSimCache.cpp
--- a/Engine/Plugins/FX/Niagara/Source/Niagara/Private/NiagaraSimCache.cpp
+++ b/Engine/Plugins/FX/Niagara/Source/Niagara/Private/NiagaraSimCache.cpp
@@ -330,7 +330,9 @@ bool UNiagaraSimCache::CanRead(UNiagaraSystem* NiagaraSystem)
 		for ( int32 i=0; i < NumEmitters; ++i )
 		{
 			const FNiagaraEmitterCompiledData& EmitterCompiledData = NiagaraSystem->GetEmitterCompiledData()[i].Get();
-			bCacheValid &= FNiagaraSimCacheHelper::BuildCacheReadMappings(CacheLayout.EmitterLayouts[i], EmitterCompiledData.DataSetCompiledData);
+			FNiagaraSimCacheDataBuffersLayout& EmitterLayout = CacheLayout.EmitterLayouts[i];
+			bCacheValid &= FNiagaraSimCacheHelper::BuildCacheReadMappings(EmitterLayout, EmitterCompiledData.DataSetCompiledData);
+			bCacheValid &= EmitterLayout.SimTarget == EmitterCompiledData.DataSetCompiledData.SimTarget;
 		}
 
 		if (bCacheValid == false)
@@ -448,6 +450,21 @@ bool UNiagaraSimCache::ReadFrame(int32 FrameIndex, float FrameFraction, FNiagara
 	return true;
 }
 
+int UNiagaraSimCache::GetEmitterIndex(FName EmitterName) const
+{
+	if ( IsCacheValid() )
+	{
+		for ( int i=0; i < CacheLayout.EmitterLayouts.Num(); ++i )
+		{
+			if ( CacheLayout.EmitterLayouts[i].LayoutName == EmitterName )
+			{
+				return i;
+			}
+		}
+	}
+	return INDEX_NONE;
+}
+
 TArray<FName> UNiagaraSimCache::GetEmitterNames() const
 {
 	TArray<FName> EmitterNames;
@@ -463,6 +480,46 @@ TArray<FName> UNiagaraSimCache::GetEmitterNames() const
 	return EmitterNames;
 }
 
+int UNiagaraSimCache::GetEmitterNumInstances(int32 EmitterIndex, int32 FrameIndex) const
+{
+	if ( CacheFrames.IsValidIndex(FrameIndex) )
+	{
+		if ( EmitterIndex == INDEX_NONE )
+		{
+			return CacheFrames[FrameIndex].SystemData.SystemDataBuffers.NumInstances;
+		}
+		else if ( CacheFrames[FrameIndex].EmitterData.IsValidIndex(EmitterIndex) )
+		{
+			return CacheFrames[FrameIndex].EmitterData[EmitterIndex].ParticleDataBuffers.NumInstances;
+		}
+	}
+	return 0;
+}
+
+void UNiagaraSimCache::ForEachEmitterAttribute(int32 EmitterIndex, TFunction<bool(const FNiagaraSimCacheVariable&)> Function) const
+{
+	if (EmitterIndex == INDEX_NONE)
+	{
+		for (const FNiagaraSimCacheVariable& CacheVariable : CacheLayout.SystemLayout.Variables)
+		{
+			if ( Function(CacheVariable) == false )
+			{
+				break;
+			}
+		}
+	}
+	else if (CacheLayout.EmitterLayouts.IsValidIndex(EmitterIndex))
+	{
+		for (const FNiagaraSimCacheVariable& CacheVariable : CacheLayout.EmitterLayouts[EmitterIndex].Variables)
+		{
+			if (Function(CacheVariable) == false)
+			{
+				break;
+			}
+		}
+	}
+}
+
 void UNiagaraSimCache::ReadAttribute(TArray<float>& OutFloats, TArray<FFloat16>& OutHalfs, TArray<int32>& OutInts, FName AttributeName, FName EmitterName, int FrameIndex) const
 {
 	FNiagaraSimCacheAttributeReaderHelper AttributeReader(this, EmitterName, AttributeName, FrameIndex);
@@ -522,7 +579,7 @@ void UNiagaraSimCache::ReadVector2Attribute(TArray<FVector2D>& OutValues, FName
 		const int32 OutValueOffset = OutValues.AddUninitialized(AttributeReader.GetNumInstances());
 		for (int32 i = 0; i < AttributeReader.GetNumInstances(); ++i)
 		{
-			OutValues[OutValueOffset + i] = AttributeReader.ReadFloat2(i);
+			OutValues[OutValueOffset + i] = FVector2D(AttributeReader.ReadFloat2f(i));
 		}
 	}
 }
@@ -535,7 +592,7 @@ void UNiagaraSimCache::ReadVectorAttribute(TArray<FVector>& OutValues, FName Att
 		const int32 OutValueOffset = OutValues.AddUninitialized(AttributeReader.GetNumInstances());
 		for (int32 i=0; i < AttributeReader.GetNumInstances(); ++i)
 		{
-			OutValues[OutValueOffset + i] = AttributeReader.ReadFloat3(i);
+			OutValues[OutValueOffset + i] = FVector(AttributeReader.ReadFloat3f(i));
 		}
 	}
 }
@@ -548,7 +605,7 @@ void UNiagaraSimCache::ReadVector4Attribute(TArray<FVector4>& OutValues, FName A
 		const int32 OutValueOffset = OutValues.AddUninitialized(AttributeReader.GetNumInstances());
 		for (int32 i = 0; i < AttributeReader.GetNumInstances(); ++i)
 		{
-			OutValues[OutValueOffset + i] = AttributeReader.ReadFloat4(i);
+			OutValues[OutValueOffset + i] = FVector4(AttributeReader.ReadFloat4f(i));
 		}
 	}
 }
@@ -585,7 +642,7 @@ void UNiagaraSimCache::ReadPositionAttribute(TArray<FVector>& OutValues, FName A
 		const int32 OutValueOffset = OutValues.AddUninitialized(AttributeReader.GetNumInstances());
 		for (int32 i = 0; i < AttributeReader.GetNumInstances(); ++i)
 		{
-			const FVector Position = AttributeReader.ReadFloat3(i);
+			const FVector Position = FVector(AttributeReader.ReadFloat3f(i));
 			OutValues[OutValueOffset + i] = LocalToWorld.TransformPosition(Position);
 		}
 	}
@@ -600,7 +657,7 @@ void UNiagaraSimCache::ReadPositionAttributeWithRebase(TArray<FVector>& OutValue
 		const int32 OutValueOffset = OutValues.AddUninitialized(AttributeReader.GetNumInstances());
 		for (int32 i = 0; i < AttributeReader.GetNumInstances(); ++i)
 		{
-			const FVector Position = AttributeReader.ReadFloat3(i);
+			const FVector Position = FVector(AttributeReader.ReadFloat3f(i));
 			OutValues[OutValueOffset + i] = LocalToWorld.TransformPosition(Position);
 		}
 	}
@@ -616,7 +673,7 @@ void UNiagaraSimCache::ReadQuatAttribute(TArray<FQuat>& OutValues, FName Attribu
 		{
 			for (int32 i = 0; i < AttributeReader.GetNumInstances(); ++i)
 			{
-				const FQuat Quat = AttributeReader.ReadQuat(i);
+				const FQuat Quat = FQuat(AttributeReader.ReadQuat4f(i));
 				OutValues[OutValueOffset + i] = Quat * AttributeReader.CacheFrame->LocalToWorld.GetRotation();
 			}
 		}
@@ -624,7 +681,7 @@ void UNiagaraSimCache::ReadQuatAttribute(TArray<FQuat>& OutValues, FName Attribu
 		{
 			for (int32 i = 0; i < AttributeReader.GetNumInstances(); ++i)
 			{
-				const FQuat Quat = AttributeReader.ReadQuat(i);
+				const FQuat Quat = FQuat(AttributeReader.ReadQuat4f(i));
 				OutValues[OutValueOffset + i] = Quat;
 			}
 		}
@@ -642,7 +699,7 @@ void UNiagaraSimCache::ReadQuatAttributeWithRebase(TArray<FQuat>& OutValues, FQu
 		const int32 OutValueOffset = OutValues.AddUninitialized(AttributeReader.GetNumInstances());
 		for (int32 i = 0; i < AttributeReader.GetNumInstances(); ++i)
 		{
-			const FQuat Rotation = AttributeReader.ReadQuat(i);
+			const FQuat Rotation = FQuat(AttributeReader.ReadQuat4f(i));
 			OutValues[OutValueOffset + i] = Rotation * LocalToWorld;
 		}
 	}
diff --git a/Engine/Plugins/FX/Niagara/Source/Niagara/Private/NiagaraSimCacheAttributeReaderHelper.h b/Engine/Plugins/FX/Niagara/Source/Niagara/Private/NiagaraSimCacheAttributeReaderHelper.h
--- a/Engine/Plugins/FX/Niagara/Source/Niagara/Private/NiagaraSimCacheAttributeReaderHelper.h
+++ b/Engine/Plugins/FX/Niagara/Source/Niagara/Private/NiagaraSimCacheAttributeReaderHelper.h
@@ -6,8 +6,20 @@
 
 struct FNiagaraSimCacheAttributeReaderHelper
 {
+	FNiagaraSimCacheAttributeReaderHelper() = default;
+
 	explicit FNiagaraSimCacheAttributeReaderHelper(const UNiagaraSimCache* SimCache, FName EmitterName, FName AttributeName, int FrameIndex)
 	{
+		Initialize(SimCache, EmitterName, AttributeName, FrameIndex);
+	}
+
+	void Initialize(const UNiagaraSimCache* SimCache, FName EmitterName, FName AttributeName, int FrameIndex)
+	{
+		CacheFrame = nullptr;
+		DataBuffers = nullptr;
+		DataBuffersLayout = nullptr;
+		Variable = nullptr;
+
 		if (SimCache->IsCacheValid() == false)
 		{
 			return;
@@ -50,23 +62,94 @@ struct FNiagaraSimCacheAttributeReaderHelper
 		Variable = &DataBuffersLayout->Variables[VariableIndex];
 	}
 
+	void Initialize(const UNiagaraSimCache* SimCache, int EmitterIndex, int AttributeIndex, int FrameIndex)
+	{
+		CacheFrame = nullptr;
+		DataBuffers = nullptr;
+		DataBuffersLayout = nullptr;
+		Variable = nullptr;
+
+		if (SimCache->IsCacheValid() == false)
+		{
+			return;
+		}
+
+		if (SimCache->CacheFrames.IsValidIndex(FrameIndex) == false)
+		{
+			return;
+		}
+		CacheFrame = &SimCache->CacheFrames[FrameIndex];
+
+		if ( EmitterIndex != INDEX_NONE )
+		{
+			if (SimCache->CacheLayout.EmitterLayouts.IsValidIndex(EmitterIndex) == false)
+			{
+				return;
+			}
+
+			DataBuffers = &CacheFrame->EmitterData[EmitterIndex].ParticleDataBuffers;
+			DataBuffersLayout = &SimCache->CacheLayout.EmitterLayouts[EmitterIndex];
+		}
+		else
+		{
+			DataBuffers = &CacheFrame->SystemData.SystemDataBuffers;
+			DataBuffersLayout = &SimCache->CacheLayout.SystemLayout;
+		}
+
+		if (DataBuffers->NumInstances == 0)
+		{
+			return;
+		}
+
+		if (DataBuffersLayout->Variables.IsValidIndex(AttributeIndex) == false)
+		{
+			return;
+		}
+
+		Variable = &DataBuffersLayout->Variables[AttributeIndex];
+	}
+
+	static int32 FindVariableIndex(UNiagaraSimCache* SimCache, FNiagaraVariableBase Variable, int EmitterIndex)
+	{
+		if (SimCache->CacheLayout.EmitterLayouts.IsValidIndex(EmitterIndex))
+		{
+			const FNiagaraSimCacheDataBuffersLayout& EmitterLayout = SimCache->CacheLayout.EmitterLayouts[EmitterIndex];
+			const int32 VariableIndex = EmitterLayout.Variables.IndexOfByPredicate([Variable](const FNiagaraSimCacheVariable& FoundVariable) { return FoundVariable.Variable == Variable; });
+			return VariableIndex;
+		}
+
+		return INDEX_NONE;
+	}
+
+	static const FNiagaraSimCacheVariable* FindVariable(UNiagaraSimCache* SimCache, FNiagaraVariableBase Variable, int EmitterIndex)
+	{
+		const int32 VariableIndex = FindVariableIndex(SimCache, Variable, EmitterIndex);
+		if (VariableIndex != INDEX_NONE)
+		{
+			const FNiagaraSimCacheDataBuffersLayout& EmitterLayout = SimCache->CacheLayout.EmitterLayouts[EmitterIndex];
+			return &EmitterLayout.Variables[VariableIndex];
+		}
+
+		return nullptr;
+	}
+
 	bool IsValid() const { return Variable != nullptr; }
 
 	int32 GetNumInstances() const { check(IsValid()); return DataBuffers->NumInstances; }
 
-	void CopyComponentFloats(int32 Component, float* OutBuffer)
+	void CopyComponentFloats(int32 Component, float* OutBuffer) const
 	{
 		check(Component < Variable->FloatCount);
 		FMemory::Memcpy(OutBuffer, &DataBuffers->FloatData[(Variable->FloatOffset + Component) * DataBuffers->NumInstances * sizeof(float)], sizeof(float) * DataBuffers->NumInstances);
 	}
 
-	void CopyComponentHalfs(int32 Component, FFloat16* OutBuffer)
+	void CopyComponentHalfs(int32 Component, FFloat16* OutBuffer) const
 	{
 		check(Component < Variable->HalfCount);
 		FMemory::Memcpy(OutBuffer, &DataBuffers->HalfData[(Variable->HalfOffset + Component) * DataBuffers->NumInstances * sizeof(FFloat16)], sizeof(FFloat16) * DataBuffers->NumInstances);
 	}
 
-	void CopyComponentInts(int32 Component, int32* OutBuffer)
+	void CopyComponentInts(int32 Component, int32* OutBuffer) const
 	{
 		check(Component < Variable->Int32Count);
 		FMemory::Memcpy(OutBuffer, &DataBuffers->Int32Data[(Variable->Int32Offset + Component) * DataBuffers->NumInstances * sizeof(int32)], sizeof(int32) * DataBuffers->NumInstances);
@@ -101,36 +184,36 @@ struct FNiagaraSimCacheAttributeReaderHelper
 		return InternalReadFloat(0, Instance);
 	}
 
-	FVector2D ReadFloat2(int32 Instance) const
+	FVector2f ReadFloat2f(int32 Instance) const
 	{
 		check(Variable->FloatOffset != INDEX_NONE && Variable->FloatCount == 2);
-		return FVector2D(InternalReadFloat(0, Instance), InternalReadFloat(1, Instance));
+		return FVector2f(InternalReadFloat(0, Instance), InternalReadFloat(1, Instance));
 	}
 
-	FVector ReadFloat3(int32 Instance) const
+	FVector3f ReadFloat3f(int32 Instance) const
 	{
 		check(Variable->FloatOffset != INDEX_NONE && Variable->FloatCount == 3);
-		return FVector(InternalReadFloat(0, Instance), InternalReadFloat(1, Instance), InternalReadFloat(2, Instance));
+		return FVector3f(InternalReadFloat(0, Instance), InternalReadFloat(1, Instance), InternalReadFloat(2, Instance));
 	}
 
-	FVector4 ReadFloat4(int32 Instance) const
+	FVector4f ReadFloat4f(int32 Instance) const
 	{
 		check(Variable->FloatOffset != INDEX_NONE && Variable->FloatCount == 4);
-		return FVector4(InternalReadFloat(0, Instance), InternalReadFloat(1, Instance), InternalReadFloat(2, Instance), InternalReadFloat(3, Instance));
+		return FVector4f(InternalReadFloat(0, Instance), InternalReadFloat(1, Instance), InternalReadFloat(2, Instance), InternalReadFloat(3, Instance));
 	}
 
-	FLinearColor ReadColor(int32 Instance) const
+	FQuat4f ReadQuat4f(int32 Instance) const
 	{
 		check(Variable->FloatOffset != INDEX_NONE && Variable->FloatCount == 4);
-		return FLinearColor(InternalReadFloat(0, Instance), InternalReadFloat(1, Instance), InternalReadFloat(2, Instance), InternalReadFloat(3, Instance));
+		return FQuat4f(InternalReadFloat(0, Instance), InternalReadFloat(1, Instance), InternalReadFloat(2, Instance), InternalReadFloat(3, Instance));
 	}
 
-	FQuat ReadQuat(int32 Instance) const
+	FLinearColor ReadColor(int32 Instance) const
 	{
 		check(Variable->FloatOffset != INDEX_NONE && Variable->FloatCount == 4);
-		return FQuat(InternalReadFloat(0, Instance), InternalReadFloat(1, Instance), InternalReadFloat(2, Instance), InternalReadFloat(3, Instance));
+		return FLinearColor(InternalReadFloat(0, Instance), InternalReadFloat(1, Instance), InternalReadFloat(2, Instance), InternalReadFloat(3, Instance));
 	}
-	
+
 	const FNiagaraSimCacheFrame*				CacheFrame = nullptr;
 	const FNiagaraSimCacheDataBuffers*			DataBuffers = nullptr;
 	const FNiagaraSimCacheDataBuffersLayout*	DataBuffersLayout = nullptr;
diff --git a/Engine/Plugins/FX/Niagara/Source/Niagara/Private/NiagaraSimCacheGpuResource.cpp b/Engine/Plugins/FX/Niagara/Source/Niagara/Private/NiagaraSimCacheGpuResource.cpp
new file mode 100644
--- /dev/null
+++ b/Engine/Plugins/FX/Niagara/Source/Niagara/Private/NiagaraSimCacheGpuResource.cpp
@@ -0,0 +1,108 @@
+// Copyright Epic Games, Inc. All Rights Reserved.
+
+#include "NiagaraSimCacheGpuResource.h"
+
+void FNiagaraSimCacheGpuResource::FDeletor::operator()(FNiagaraSimCacheGpuResource* Ptr)
+{
+	ENQUEUE_RENDER_COMMAND(FNiagaraSimCacheGpuResource_FDeletor)(
+		[RT_Ptr=Ptr](FRHICommandListImmediate& RHICmdList)
+		{
+			delete RT_Ptr;
+		}
+	);
+}
+
+FNiagaraSimCacheGpuResource::FNiagaraSimCacheGpuResource(UNiagaraSimCache* SimCache)
+	: WeakSimCache(SimCache)
+{
+	OnSimCacheChangedHandle = UNiagaraSimCache::OnCacheEndWrite.AddRaw(this, &FNiagaraSimCacheGpuResource::BuildResource);
+	BuildResource(SimCache);
+}
+
+FNiagaraSimCacheGpuResource::~FNiagaraSimCacheGpuResource()
+{
+	check(IsInRenderingThread());
+	if ( OnSimCacheChangedHandle.IsValid() )
+	{
+		UNiagaraSimCache::OnCacheEndWrite.Remove(OnSimCacheChangedHandle);
+	}
+}
+
+FNiagaraSimCacheGpuResourcePtr FNiagaraSimCacheGpuResource::CreateResource(UNiagaraSimCache* SimCache)
+{
+	return MakeShareable(new FNiagaraSimCacheGpuResource(SimCache), FDeletor());
+}
+
+void FNiagaraSimCacheGpuResource::BuildResource(UNiagaraSimCache* SimCache)
+{
+	check(SimCache != nullptr);
+
+	NumFrames = SimCache->GetNumFrames();
+	NumEmitters = SimCache->GetNumEmitters();
+
+	const uint32 FrameDataSizeInInts = 4;	//-TODO: Match with shader, this is NumInstances, IntDataOffset, FloatDataOffset, HalfDataOffset
+	uint32 CurrentFrameOffset = 0;
+	uint32 CurrentDataOffset = NumFrames * NumEmitters * FrameDataSizeInInts;
+	TArray<uint32> GpuCacheData;
+	{
+		uint32 RequiredElements = CurrentDataOffset;
+		for (const FNiagaraSimCacheFrame& CacheFrame : SimCache->CacheFrames)
+		{
+			for (const FNiagaraSimCacheEmitterFrame& EmitterFrame : CacheFrame.EmitterData)
+			{
+				RequiredElements += FMath::DivideAndRoundUp<uint32>(EmitterFrame.ParticleDataBuffers.Int32Data.Num(), sizeof(uint32));
+				RequiredElements += FMath::DivideAndRoundUp<uint32>(EmitterFrame.ParticleDataBuffers.FloatData.Num(), sizeof(uint32));
+				RequiredElements += FMath::DivideAndRoundUp<uint32>(EmitterFrame.ParticleDataBuffers.HalfData.Num(), sizeof(uint32));
+			}
+		}
+
+		GpuCacheData.AddUninitialized(RequiredElements);
+	}
+
+	for (const FNiagaraSimCacheFrame& CacheFrame : SimCache->CacheFrames)
+	{
+		for (const FNiagaraSimCacheEmitterFrame& EmitterFrame : CacheFrame.EmitterData)
+		{
+			GpuCacheData[CurrentFrameOffset++] = EmitterFrame.ParticleDataBuffers.NumInstances;
+
+			GpuCacheData[CurrentFrameOffset++] = CurrentDataOffset;
+			if (EmitterFrame.ParticleDataBuffers.Int32Data.Num() > 0)
+			{
+				FMemory::Memcpy(&GpuCacheData[CurrentDataOffset], EmitterFrame.ParticleDataBuffers.Int32Data.GetData(), EmitterFrame.ParticleDataBuffers.Int32Data.Num());
+				CurrentDataOffset += FMath::DivideAndRoundUp<uint32>(EmitterFrame.ParticleDataBuffers.Int32Data.Num(), sizeof(uint32));
+			}
+
+			GpuCacheData[CurrentFrameOffset++] = CurrentDataOffset;
+			if (EmitterFrame.ParticleDataBuffers.FloatData.Num() > 0)
+			{
+				FMemory::Memcpy(&GpuCacheData[CurrentDataOffset], EmitterFrame.ParticleDataBuffers.FloatData.GetData(), EmitterFrame.ParticleDataBuffers.FloatData.Num());
+				CurrentDataOffset += FMath::DivideAndRoundUp<uint32>(EmitterFrame.ParticleDataBuffers.FloatData.Num(), sizeof(uint32));
+			}
+
+			GpuCacheData[CurrentFrameOffset++] = CurrentDataOffset;
+			if (EmitterFrame.ParticleDataBuffers.HalfData.Num() > 0)
+			{
+				FMemory::Memcpy(&GpuCacheData[CurrentDataOffset], EmitterFrame.ParticleDataBuffers.HalfData.GetData(), EmitterFrame.ParticleDataBuffers.HalfData.Num());
+				CurrentDataOffset += FMath::DivideAndRoundUp<uint32>(EmitterFrame.ParticleDataBuffers.HalfData.Num(), sizeof(uint32));
+			}
+		}
+	}
+
+	ENQUEUE_RENDER_COMMAND(FNiagaraSimCacheGpuResource_BuildResource)
+	(
+		[this, NumFrames_RT=NumFrames, NumEmitters_RT=NumEmitters, GpuCacheData_RT=MoveTemp(GpuCacheData)](FRHICommandList& CmdList)
+		{
+			NumFrames = NumFrames_RT;
+			NumEmitters = NumEmitters_RT;
+			SimCacheBuffer.Release();
+			if (GpuCacheData_RT.Num() > 0)
+			{
+				SimCacheBuffer.Initialize(TEXT("NiagaraSimCache"), sizeof(uint32), GpuCacheData_RT.Num(), EPixelFormat::PF_R32_UINT, BUF_Static);
+
+				void* GpuMemory = RHILockBuffer(SimCacheBuffer.Buffer, 0, SimCacheBuffer.NumBytes, RLM_WriteOnly);
+				FMemory::Memcpy(GpuMemory, GpuCacheData_RT.GetData(), SimCacheBuffer.NumBytes);
+				RHIUnlockBuffer(SimCacheBuffer.Buffer);
+			}
+		}
+	);
+}
diff --git a/Engine/Plugins/FX/Niagara/Source/Niagara/Private/NiagaraSimCacheGpuResource.h b/Engine/Plugins/FX/Niagara/Source/Niagara/Private/NiagaraSimCacheGpuResource.h
new file mode 100644
--- /dev/null
+++ b/Engine/Plugins/FX/Niagara/Source/Niagara/Private/NiagaraSimCacheGpuResource.h
@@ -0,0 +1,46 @@
+// Copyright Epic Games, Inc. All Rights Reserved.
+
+#pragma  once
+
+#include "Templates/SharedPointer.h"
+
+#include "NiagaraSimCache.h"
+#include "NiagaraRenderer.h"
+
+struct FNiagaraSimCacheGpuResource
+{
+	UE_NONCOPYABLE(FNiagaraSimCacheGpuResource)
+	friend struct FDeletor;
+
+private:
+	struct FDeletor
+	{
+		void operator()(FNiagaraSimCacheGpuResource* Pointer);
+	};
+
+	FNiagaraSimCacheGpuResource(UNiagaraSimCache* SimCache);
+	virtual ~FNiagaraSimCacheGpuResource();
+
+public:
+	static TSharedPtr<FNiagaraSimCacheGpuResource> CreateResource(UNiagaraSimCache* SimCache);
+
+	int32 GetNumFrames() const { return NumFrames; }
+	int32 GetNumEmitters() const { return NumEmitters; }
+	int32 GetEmitterIndex() const { return EmitterIndex; }
+	FRHIShaderResourceView* GetBufferSRV() const { return SimCacheBuffer.SRV.IsValid() ? SimCacheBuffer.SRV.GetReference() : FNiagaraRenderer::GetDummyUIntBuffer(); }
+
+private:
+	void BuildResource(UNiagaraSimCache* SimCache);
+
+private:
+	TWeakObjectPtr<UNiagaraSimCache>	WeakSimCache;
+
+	int32								NumFrames = 0;
+	int32								NumEmitters = 0;
+	int32								EmitterIndex = INDEX_NONE;
+	FReadBuffer							SimCacheBuffer;
+
+	FDelegateHandle						OnSimCacheChangedHandle;
+};
+
+typedef TSharedPtr<FNiagaraSimCacheGpuResource> FNiagaraSimCacheGpuResourcePtr;
diff --git a/Engine/Plugins/FX/Niagara/Source/Niagara/Public/NiagaraTypes.h b/Engine/Plugins/FX/Niagara/Source/Niagara/Public/NiagaraTypes.h
--- a/Engine/Plugins/FX/Niagara/Source/Niagara/Public/NiagaraTypes.h
+++ b/Engine/Plugins/FX/Niagara/Source/Niagara/Public/NiagaraTypes.h
@@ -1194,6 +1194,7 @@ public:
 	static const FNiagaraTypeDefinition& GetUStaticMeshDef() { return UStaticMeshDef; }
 	static const FNiagaraTypeDefinition& GetUTextureDef() { return UTextureDef; }
 	static const FNiagaraTypeDefinition& GetUTextureRenderTargetDef() { return UTextureRenderTargetDef; }
+	static const FNiagaraTypeDefinition& GetSimCacheClassDef() { return USimCacheClassDef; }
 	static const FNiagaraTypeDefinition& GetWildcardDef() { return WildcardDef; }
 	static const FNiagaraTypeDefinition& GetHalfDef() { return HalfDef; }
 	static const FNiagaraTypeDefinition& GetHalfVec2Def() { return HalfVec2Def; }
@@ -1280,6 +1281,7 @@ private:
 	static FNiagaraTypeDefinition UTextureDef;
 	static FNiagaraTypeDefinition UTextureRenderTargetDef;
 	static FNiagaraTypeDefinition UStaticMeshDef;
+	static FNiagaraTypeDefinition USimCacheClassDef;
 	static FNiagaraTypeDefinition WildcardDef;
 
 	static FNiagaraTypeDefinition HalfDef;
diff --git a/Engine/Plugins/FX/Niagara/Source/NiagaraEditor/Private/ViewModels/NiagaraSimCacheViewModel.cpp b/Engine/Plugins/FX/Niagara/Source/NiagaraEditor/Private/ViewModels/NiagaraSimCacheViewModel.cpp
--- a/Engine/Plugins/FX/Niagara/Source/NiagaraEditor/Private/ViewModels/NiagaraSimCacheViewModel.cpp
+++ b/Engine/Plugins/FX/Niagara/Source/NiagaraEditor/Private/ViewModels/NiagaraSimCacheViewModel.cpp
@@ -173,29 +173,33 @@ void FNiagaraSimCacheViewModel::UpdateCachedFrame()
 	FoundHalfComponents = 0;
 	FoundInt32Components = 0;
 
-	NumInstances = EmitterIndex == INDEX_NONE ? 1 : SimCache->CacheFrames[FrameIndex].EmitterData[EmitterIndex].ParticleDataBuffers.NumInstances;
+	NumInstances = EmitterIndex == INDEX_NONE ? 1 : SimCache->GetEmitterNumInstances(EmitterIndex, FrameIndex);
 	const FName EmitterName = EmitterIndex == INDEX_NONE ? NAME_None : SimCache->GetEmitterName(EmitterIndex);
-	const FNiagaraSimCacheDataBuffersLayout* DataBufferLayout = EmitterIndex == INDEX_NONE ? &SimCache->CacheLayout.SystemLayout : &SimCache->CacheLayout.EmitterLayouts[EmitterIndex];
-	for ( const FNiagaraSimCacheVariable& Variable : DataBufferLayout->Variables )
-	{
-		// Build component info
-		const FNiagaraTypeDefinition& TypeDef = Variable.Variable.GetType();
-		if (TypeDef.IsEnum())
-		{
-			FComponentInfo& ComponentInfo = ComponentInfos.AddDefaulted_GetRef();
-			ComponentInfo.Name = Variable.Variable.GetName();
-			ComponentInfo.ComponentOffset = FoundInt32Components++;
-			ComponentInfo.bIsInt32 = true;
-			ComponentInfo.Enum = TypeDef.GetEnum();
-		}
-		else
+
+	SimCache->ForEachEmitterAttribute(EmitterIndex,
+		[&](const FNiagaraSimCacheVariable& Variable)
 		{
-			BuildComponentInfos(Variable.Variable.GetName(), TypeDef.GetScriptStruct());
-		}
+			// Build component info
+			const FNiagaraTypeDefinition& TypeDef = Variable.Variable.GetType();
+			if (TypeDef.IsEnum())
+			{
+				FComponentInfo& ComponentInfo = ComponentInfos.AddDefaulted_GetRef();
+				ComponentInfo.Name = Variable.Variable.GetName();
+				ComponentInfo.ComponentOffset = FoundInt32Components++;
+				ComponentInfo.bIsInt32 = true;
+				ComponentInfo.Enum = TypeDef.GetEnum();
+			}
+			else
+			{
+				BuildComponentInfos(Variable.Variable.GetName(), TypeDef.GetScriptStruct());
+			}
 
-		// Pull in data
-		SimCache->ReadAttribute(FloatComponents, HalfComponents, Int32Components, Variable.Variable.GetName(), EmitterName, FrameIndex);
-	}
+			// Pull in data
+			SimCache->ReadAttribute(FloatComponents, HalfComponents, Int32Components, Variable.Variable.GetName(), EmitterName, FrameIndex);
+
+			return true;
+		}
+	);
 }
 
 void FNiagaraSimCacheViewModel::BuildComponentInfos(const FName Name, const UScriptStruct* Struct)
