Missing merge file

diff --git a/Engine/Plugins/FX/Niagara/Source/NiagaraShader/Private/NiagaraDistanceFieldHelper.cpp b/Engine/Plugins/FX/Niagara/Source/NiagaraShader/Private/NiagaraDistanceFieldHelper.cpp
--- a/Engine/Plugins/FX/Niagara/Source/NiagaraShader/Private/NiagaraDistanceFieldHelper.cpp
+++ b/Engine/Plugins/FX/Niagara/Source/NiagaraShader/Private/NiagaraDistanceFieldHelper.cpp
@@ -76,8 +76,8 @@ void FNiagaraDistanceFieldHelper::SetMeshDistanceFieldParameters(FRDGBuilder& Gr
 		FRDGBufferSRVRef DefaultVector4 = GraphBuilder.CreateSRV(FRDGBufferSRVDesc(GSystemTextures.GetDefaultStructuredBuffer(GraphBuilder, sizeof(FVector4f))));
 		FRDGBufferSRVRef DefaultUInt32 = GraphBuilder.CreateSRV(FRDGBufferSRVDesc(GSystemTextures.GetDefaultStructuredBuffer(GraphBuilder, sizeof(uint32))));
 
-		ObjectShaderParameters.SceneObjectBounds = DummyFloat4Buffer;
-		ObjectShaderParameters.SceneObjectData = DummyFloat4Buffer;
+		ObjectShaderParameters.SceneObjectBounds = DefaultVector4;
+		ObjectShaderParameters.SceneObjectData = DefaultVector4;
 		ObjectShaderParameters.NumSceneObjects = 0;
 		ObjectShaderParameters.SceneHeightfieldObjectBounds = DefaultVector4;
 		ObjectShaderParameters.SceneHeightfieldObjectData = DefaultVector4;
