- Added support to override Sim dispatch type (Contributed by [REDACTED_PERSON])

#9536

diff --git a/Engine/Plugins/FX/Niagara/Source/Niagara/Classes/NiagaraGPUSystemTick.h b/Engine/Plugins/FX/Niagara/Source/Niagara/Classes/NiagaraGPUSystemTick.h
--- a/Engine/Plugins/FX/Niagara/Source/Niagara/Classes/NiagaraGPUSystemTick.h
+++ b/Engine/Plugins/FX/Niagara/Source/Niagara/Classes/NiagaraGPUSystemTick.h
@@ -23,12 +23,12 @@ struct FNiagaraComputeInstanceData
 	struct FPerStageInfo
 	{
 		FPerStageInfo() {}
-		FPerStageInfo(int32 InNumIterations, int32 InUserElementCount) : NumIterations(InNumIterations), UserElementCount(InUserElementCount) {}
+		FPerStageInfo(int32 InNumIterations, FIntVector InElementCountXYZ) : NumIterations(InNumIterations), ElementCountXYZ(InElementCountXYZ) {}
 
-		int32	NumIterations = 0;
-		int32	UserElementCount = -1;
+		int32 NumIterations = 0;
+		FIntVector ElementCountXYZ = FIntVector::NoneValue;
 
-		bool ShouldRunStage() const { return NumIterations > 0 && (UserElementCount != 0); }
+		bool ShouldRunStage() const { return NumIterations > 0 && (ElementCountXYZ != FIntVector::ZeroValue); }
 	};
 
 	FNiagaraComputeInstanceData()
diff --git a/Engine/Plugins/FX/Niagara/Source/Niagara/Private/NiagaraEmitterInstance.cpp b/Engine/Plugins/FX/Niagara/Source/Niagara/Private/NiagaraEmitterInstance.cpp
--- a/Engine/Plugins/FX/Niagara/Source/Niagara/Private/NiagaraEmitterInstance.cpp
+++ b/Engine/Plugins/FX/Niagara/Source/Niagara/Private/NiagaraEmitterInstance.cpp
@@ -491,9 +491,20 @@ void FNiagaraEmitterInstance::Init(int32 InEmitterIdx, FNiagaraSystemInstanceID
 				{
 					bAnyRendererBindingsAdded |= RendererBindings.AddParameter(FNiagaraVariable(FNiagaraTypeDefinition::GetBoolDef(), SimStageMetaData.EnabledBinding), false);
 				}
-				if (!SimStageMetaData.ElementCountBinding.IsNone())
+				if (SimStageMetaData.bOverrideElementCount)
 				{
-					bAnyRendererBindingsAdded |= RendererBindings.AddParameter(FNiagaraVariable(FNiagaraTypeDefinition::GetIntDef(), SimStageMetaData.ElementCountBinding), false);
+					if (!SimStageMetaData.ElementCountXBinding.IsNone())
+					{
+						bAnyRendererBindingsAdded |= RendererBindings.AddParameter(FNiagaraVariable(FNiagaraTypeDefinition::GetIntDef(), SimStageMetaData.ElementCountXBinding), false);
+					}
+					if (!SimStageMetaData.ElementCountYBinding.IsNone())
+					{
+						bAnyRendererBindingsAdded |= RendererBindings.AddParameter(FNiagaraVariable(FNiagaraTypeDefinition::GetIntDef(), SimStageMetaData.ElementCountYBinding), false);
+					}
+					if (!SimStageMetaData.ElementCountZBinding.IsNone())
+					{
+						bAnyRendererBindingsAdded |= RendererBindings.AddParameter(FNiagaraVariable(FNiagaraTypeDefinition::GetIntDef(), SimStageMetaData.ElementCountZBinding), false);
+					}
 				}
 				if ( !SimStageMetaData.NumIterationsBinding.IsNone() )
 				{
diff --git a/Engine/Plugins/FX/Niagara/Source/Niagara/Private/NiagaraGPUSystemTick.cpp b/Engine/Plugins/FX/Niagara/Source/Niagara/Private/NiagaraGPUSystemTick.cpp
--- a/Engine/Plugins/FX/Niagara/Source/Niagara/Private/NiagaraGPUSystemTick.cpp
+++ b/Engine/Plugins/FX/Niagara/Source/Niagara/Private/NiagaraGPUSystemTick.cpp
@@ -200,7 +200,7 @@ void FNiagaraGPUSystemTick::Init(FNiagaraSystemInstance* InSystemInstance)
 			for ( FSimulationStageMetaData& SimStageMetaData : GPUContext->SimStageInfo )
 			{
 				int32 NumIterations = SimStageMetaData.NumIterations;
-				int32 UserElementCount = -1;
+				FIntVector ElementCountXYZ = FIntVector::NoneValue;
 				if (SimStageMetaData.ShouldRunStage(InstanceData->bResetData))
 				{
 					InstanceData->bHasMultipleStages = true;
@@ -222,13 +222,37 @@ void FNiagaraGPUSystemTick::Init(FNiagaraSystemInstance* InSystemInstance)
 							NumIterations = StageEnabled.GetValue() ? NumIterations : 0;
 						}
 					}
-					if (!SimStageMetaData.ElementCountBinding.IsNone())
+					if (SimStageMetaData.bOverrideElementCount)
 					{
-						FNiagaraParameterStore& BoundParamStore = EmitterInstance->GetRendererBoundVariables();
-						if (const uint8* ParameterData = BoundParamStore.GetParameterData(FNiagaraVariable(FNiagaraTypeDefinition::GetIntDef(), SimStageMetaData.ElementCountBinding)))
+						if (!SimStageMetaData.ElementCountXBinding.IsNone() && (SimStageMetaData.GpuDispatchType >= ENiagaraGpuDispatchType::OneD))
+						{
+							FNiagaraParameterStore& BoundParamStore = EmitterInstance->GetRendererBoundVariables();
+							if (const uint8* ParameterData = BoundParamStore.GetParameterData(FNiagaraVariable(FNiagaraTypeDefinition::GetIntDef(), SimStageMetaData.ElementCountXBinding)))
+							{
+								ElementCountXYZ.X = *reinterpret_cast<const int32*>(ParameterData);
+							}
+						}
+						if (!SimStageMetaData.ElementCountYBinding.IsNone() && (SimStageMetaData.GpuDispatchType >= ENiagaraGpuDispatchType::TwoD))
 						{
-							UserElementCount = FMath::Max(*reinterpret_cast<const int32*>(ParameterData), 0);
+							FNiagaraParameterStore& BoundParamStore = EmitterInstance->GetRendererBoundVariables();
+							if (const uint8* ParameterData = BoundParamStore.GetParameterData(FNiagaraVariable(FNiagaraTypeDefinition::GetIntDef(), SimStageMetaData.ElementCountYBinding)))
+							{
+								ElementCountXYZ.Y = *reinterpret_cast<const int32*>(ParameterData);
+							}
 						}
+						if (!SimStageMetaData.ElementCountZBinding.IsNone() && (SimStageMetaData.GpuDispatchType >= ENiagaraGpuDispatchType::ThreeD))
+						{
+							FNiagaraParameterStore& BoundParamStore = EmitterInstance->GetRendererBoundVariables();
+							if (const uint8* ParameterData = BoundParamStore.GetParameterData(FNiagaraVariable(FNiagaraTypeDefinition::GetIntDef(), SimStageMetaData.ElementCountZBinding)))
+							{
+								ElementCountXYZ.Z = *reinterpret_cast<const int32*>(ParameterData);
+							}
+						}
+
+						// make sure we don't have negatives as we use the values directly for dispatch
+						ElementCountXYZ.X = FMath::Max(0, ElementCountXYZ.X);
+						ElementCountXYZ.Y = FMath::Max(0, ElementCountXYZ.Y);
+						ElementCountXYZ.Z = FMath::Max(0, ElementCountXYZ.Z);
 					}
 				}
 				else
@@ -237,7 +261,7 @@ void FNiagaraGPUSystemTick::Init(FNiagaraSystemInstance* InSystemInstance)
 				}
 
 				InstanceData->TotalDispatches += NumIterations;
-				InstanceData->PerStageInfo.Emplace(NumIterations, UserElementCount);
+				InstanceData->PerStageInfo.Emplace(NumIterations, ElementCountXYZ);
 			}
 			TotalDispatches += InstanceData->TotalDispatches;
 		}
diff --git a/Engine/Plugins/FX/Niagara/Source/Niagara/Private/NiagaraGpuComputeDispatch.cpp b/Engine/Plugins/FX/Niagara/Source/Niagara/Private/NiagaraGpuComputeDispatch.cpp
--- a/Engine/Plugins/FX/Niagara/Source/Niagara/Private/NiagaraGpuComputeDispatch.cpp
+++ b/Engine/Plugins/FX/Niagara/Source/Niagara/Private/NiagaraGpuComputeDispatch.cpp
@@ -575,9 +575,9 @@ void FNiagaraGpuComputeDispatch::DumpDebugFrame()
 				Builder.Appendf(TEXT("Source(%p 0x%08x %d) "), SimStageData.Source, SimStageData.SourceCountOffset, SimStageData.SourceNumInstances);
 				Builder.Appendf(TEXT("Destination(%p 0x%08x %d) "), SimStageData.Destination, SimStageData.DestinationCountOffset, SimStageData.DestinationNumInstances);
 				Builder.Appendf(TEXT("Iteration(%d | %s) "), SimStageData.IterationIndex, SimStageData.AlternateIterationSource ? *SimStageData.AlternateIterationSource->SourceDIName.ToString() : TEXT("Particles"));
-				if (SimStageData.UserElementCount != -1)
+				if (SimStageData.StageMetaData->bOverrideElementCount)
 				{
-					Builder.Appendf(TEXT("UserElementCount(%d) "), SimStageData.UserElementCount);
+					Builder.Appendf(TEXT("ElementCountXYZ(%d, %d, %d) "), SimStageData.ElementCountXYZ.X, SimStageData.ElementCountXYZ.Y, SimStageData.ElementCountXYZ.Z);
 				}
 				UE_LOG(LogNiagara, Warning, TEXT("%s"), Builder.ToString());
 			}
@@ -832,7 +832,7 @@ void FNiagaraGpuComputeDispatch::PrepareTicksForProxy(FRHICommandListImmediate&
 					SimStageData.bFirstStage = bFirstStage;
 					SimStageData.StageIndex = SimStageIndex;
 					SimStageData.IterationIndex = IterationIndex;
-					SimStageData.UserElementCount = InstanceData.PerStageInfo[SimStageIndex].UserElementCount;
+					SimStageData.ElementCountXYZ = InstanceData.PerStageInfo[SimStageIndex].ElementCountXYZ;
 					SimStageData.StageMetaData = &SimStageMetaData;
 					SimStageData.AlternateIterationSource = IterationInterface;
 
@@ -1408,10 +1408,20 @@ void FNiagaraGpuComputeDispatch::DispatchStage(FRDGBuilder& GraphBuilder, const
 		DispatchNumThreads = FNiagaraShader::GetDefaultThreadGroupSize(ENiagaraGpuDispatchType::OneD);
 	}
 
-	// User override element count
-	if (SimStageData.UserElementCount != -1)
+	if (SimStageData.StageMetaData->bOverrideElementCount)
 	{
-		DispatchCount = FIntVector(SimStageData.UserElementCount, 1, 1);
+		if (DispatchType == ENiagaraGpuDispatchType::OneD)
+		{
+			DispatchCount = FIntVector(SimStageData.ElementCountXYZ.X, 1, 1);
+		}
+		else if (DispatchType == ENiagaraGpuDispatchType::TwoD)
+		{
+			DispatchCount = FIntVector(SimStageData.ElementCountXYZ.X, SimStageData.ElementCountXYZ.Y, 1);
+		}
+		else if (DispatchType == ENiagaraGpuDispatchType::ThreeD)
+		{
+			DispatchCount = SimStageData.ElementCountXYZ;
+		}
 	}
 
 	const int32 TotalDispatchCount = DispatchCount.X * DispatchCount.Y * DispatchCount.Z;
diff --git a/Engine/Plugins/FX/Niagara/Source/Niagara/Private/NiagaraScript.cpp b/Engine/Plugins/FX/Niagara/Source/Niagara/Private/NiagaraScript.cpp
--- a/Engine/Plugins/FX/Niagara/Source/Niagara/Private/NiagaraScript.cpp
+++ b/Engine/Plugins/FX/Niagara/Source/Niagara/Private/NiagaraScript.cpp
@@ -1028,7 +1028,12 @@ void UNiagaraScript::ComputeVMCompilationId(FNiagaraVMExecutableDataId& Id, FGui
 					if (const UNiagaraSimulationStageGeneric* SimStageGeneric = Cast<const UNiagaraSimulationStageGeneric>(SimStageBase))
 					{
 						AddAttributeToPreserve(SimStageGeneric->EnabledBinding.GetParamMapBindableVariable().GetName());
-						AddAttributeToPreserve(SimStageGeneric->ElementCountBinding.GetParamMapBindableVariable().GetName());
+						if (SimStageGeneric->bOverrideGpuDispatchType)
+						{
+							AddAttributeToPreserve(SimStageGeneric->ElementCountXBinding.GetParamMapBindableVariable().GetName());
+							AddAttributeToPreserve(SimStageGeneric->ElementCountYBinding.GetParamMapBindableVariable().GetName());
+							AddAttributeToPreserve(SimStageGeneric->ElementCountZBinding.GetParamMapBindableVariable().GetName());
+						}
 						AddAttributeToPreserve(SimStageGeneric->NumIterationsBinding.GetParamMapBindableVariable().GetName());
 					}
 				}
@@ -3352,10 +3357,23 @@ void UNiagaraScript::SyncAliases(const FNiagaraAliasContext& ResolveAliasesConte
 			SimStageMetaData.EnabledBinding = FNiagaraUtilities::ResolveAliases(Var, ResolveAliasesContext).GetName();
 		}
 
-		if (!SimStageMetaData.ElementCountBinding.IsNone())
+		if (SimStageMetaData.bOverrideElementCount)
 		{
-			FNiagaraVariable Var(FNiagaraTypeDefinition::GetIntDef(), SimStageMetaData.ElementCountBinding);
-			SimStageMetaData.ElementCountBinding = FNiagaraUtilities::ResolveAliases(Var, ResolveAliasesContext).GetName();
+			if (!SimStageMetaData.ElementCountXBinding.IsNone())
+			{
+				FNiagaraVariable Var(FNiagaraTypeDefinition::GetIntDef(), SimStageMetaData.ElementCountXBinding);
+				SimStageMetaData.ElementCountXBinding = FNiagaraUtilities::ResolveAliases(Var, ResolveAliasesContext).GetName();
+			}
+			if (!SimStageMetaData.ElementCountYBinding.IsNone())
+			{
+				FNiagaraVariable Var(FNiagaraTypeDefinition::GetIntDef(), SimStageMetaData.ElementCountYBinding);
+				SimStageMetaData.ElementCountYBinding = FNiagaraUtilities::ResolveAliases(Var, ResolveAliasesContext).GetName();
+			}
+			if (!SimStageMetaData.ElementCountZBinding.IsNone())
+			{
+				FNiagaraVariable Var(FNiagaraTypeDefinition::GetIntDef(), SimStageMetaData.ElementCountZBinding);
+				SimStageMetaData.ElementCountZBinding = FNiagaraUtilities::ResolveAliases(Var, ResolveAliasesContext).GetName();
+			}
 		}
 
 		if (!SimStageMetaData.NumIterationsBinding.IsNone())
diff --git a/Engine/Plugins/FX/Niagara/Source/Niagara/Private/NiagaraSimulationStageBase.cpp b/Engine/Plugins/FX/Niagara/Source/Niagara/Private/NiagaraSimulationStageBase.cpp
--- a/Engine/Plugins/FX/Niagara/Source/Niagara/Private/NiagaraSimulationStageBase.cpp
+++ b/Engine/Plugins/FX/Niagara/Source/Niagara/Private/NiagaraSimulationStageBase.cpp
@@ -72,7 +72,9 @@ bool UNiagaraSimulationStageGeneric::AppendCompileHash(FNiagaraCompileHashVisito
 	Super::AppendCompileHash(InVisitor);
 
 	InVisitor->UpdateString(TEXT("EnabledBinding"), EnabledBinding.GetDataSetBindableVariable().GetName().ToString());
-	InVisitor->UpdateString(TEXT("ElementCountBinding"), ElementCountBinding.GetDataSetBindableVariable().GetName().ToString());
+	InVisitor->UpdateString(TEXT("ElementCountXBinding"), ElementCountXBinding.GetDataSetBindableVariable().GetName().ToString());
+	InVisitor->UpdateString(TEXT("ElementCountYBinding"), ElementCountYBinding.GetDataSetBindableVariable().GetName().ToString());
+	InVisitor->UpdateString(TEXT("ElementCountZBinding"), ElementCountZBinding.GetDataSetBindableVariable().GetName().ToString());
 	InVisitor->UpdatePOD(TEXT("Iterations"), Iterations);
 	InVisitor->UpdateString(TEXT("NumIterationsBinding"), NumIterationsBinding.GetDataSetBindableVariable().GetName().ToString());
 	InVisitor->UpdatePOD(TEXT("IterationSource"), (int32)IterationSource);
@@ -84,6 +86,8 @@ bool UNiagaraSimulationStageGeneric::AppendCompileHash(FNiagaraCompileHashVisito
 	InVisitor->UpdateString(TEXT("ParticleIterationStateBinding"), ParticleIterationStateBinding.GetDataSetBindableVariable().GetName().ToString());
 	InVisitor->UpdateString(TEXT("ParticleIterationStateRange"), FString::Printf(TEXT("%d,%d"), ParticleIterationStateRange.X, ParticleIterationStateRange.Y));
 	InVisitor->UpdatePOD(TEXT("bGpuDispatchForceLinear"), bGpuDispatchForceLinear ? 1 : 0);
+	InVisitor->UpdatePOD(TEXT("bOverrideGpuDispatchType"), bOverrideGpuDispatchType ? 1 : 0);
+	InVisitor->UpdatePOD(TEXT("OverrideGpuDispatchType"), (int32)OverrideGpuDispatchType);
 	InVisitor->UpdatePOD(TEXT("bOverrideGpuDispatchNumThreads"), bOverrideGpuDispatchNumThreads ? 1 : 0);
 	InVisitor->UpdateString(TEXT("OverrideGpuDispatchNumThreads"), FString::Printf(TEXT("%d,%d,%d"), OverrideGpuDispatchNumThreads.X, OverrideGpuDispatchNumThreads.Y, OverrideGpuDispatchNumThreads.Z));
 
@@ -97,7 +101,9 @@ bool UNiagaraSimulationStageGeneric::FillCompilationData(TArray<FNiagaraSimulati
 	SimStageData.StageGuid = Script->GetUsageId();
 	SimStageData.StageName = SimulationStageName;
 	SimStageData.EnabledBinding = EnabledBinding.GetName();
-	SimStageData.ElementCountBinding = ElementCountBinding.GetName();
+	SimStageData.ElementCountXBinding = ElementCountXBinding.GetName();
+	SimStageData.ElementCountYBinding = ElementCountYBinding.GetName();
+	SimStageData.ElementCountZBinding = ElementCountZBinding.GetName();
 	SimStageData.NumIterations = Iterations;
 	SimStageData.NumIterationsBinding = NumIterationsBinding.GetName();
 	SimStageData.IterationSource = IterationSource == ENiagaraIterationSource::DataInterface ? DataInterface.BoundVariable.GetName() : FName();
@@ -106,6 +112,8 @@ bool UNiagaraSimulationStageGeneric::FillCompilationData(TArray<FNiagaraSimulati
 	SimStageData.bParticleIterationStateEnabled = bParticleIterationStateEnabled;
 	SimStageData.ParticleIterationStateRange = ParticleIterationStateRange;
 	SimStageData.bGpuDispatchForceLinear = bGpuDispatchForceLinear;
+	SimStageData.bOverrideGpuDispatchType = bOverrideGpuDispatchType;
+	SimStageData.OverrideGpuDispatchType = OverrideGpuDispatchType;
 	SimStageData.bOverrideGpuDispatchNumThreads = bOverrideGpuDispatchNumThreads;
 	SimStageData.OverrideGpuDispatchNumThreads = OverrideGpuDispatchNumThreads;
 
@@ -129,7 +137,10 @@ void UNiagaraSimulationStageGeneric::PostInitProperties()
 	{
 		using namespace NiagaraSimulationStageLocal;
 		EnabledBinding.Setup(GetDefaultEnabledBinding(), GetDefaultEnabledBinding(), ENiagaraRendererSourceDataMode::Emitter);
-		ElementCountBinding.Setup(GetDefaultElementCountBinding(), GetDefaultElementCountBinding(), ENiagaraRendererSourceDataMode::Emitter);
+		ElementCountBinding_DEPRECATED.Setup(GetDefaultElementCountBinding(), GetDefaultElementCountBinding(), ENiagaraRendererSourceDataMode::Emitter);
+		ElementCountXBinding.Setup(GetDefaultElementCountBinding(), GetDefaultElementCountBinding(), ENiagaraRendererSourceDataMode::Emitter);
+		ElementCountYBinding.Setup(GetDefaultElementCountBinding(), GetDefaultElementCountBinding(), ENiagaraRendererSourceDataMode::Emitter);
+		ElementCountZBinding.Setup(GetDefaultElementCountBinding(), GetDefaultElementCountBinding(), ENiagaraRendererSourceDataMode::Emitter);
 		NumIterationsBinding.Setup(GetDefaultNumIterationsBinding(), GetDefaultNumIterationsBinding(), ENiagaraRendererSourceDataMode::Emitter);
 		ParticleIterationStateBinding.Setup(GetDefaultParticleStateBinding(), GetDefaultParticleStateBinding(), ENiagaraRendererSourceDataMode::Particles);
 	}
@@ -146,9 +157,17 @@ void UNiagaraSimulationStageGeneric::PostLoad()
 	{
 		EnabledBinding.Setup(GetDefaultEnabledBinding(), GetDefaultEnabledBinding(), ENiagaraRendererSourceDataMode::Emitter);
 	}
-	if (ElementCountBinding.GetType() != GetDefaultElementCountBinding().GetType())
+	if (ElementCountXBinding.GetType() != GetDefaultElementCountBinding().GetType())
 	{
-		ElementCountBinding.Setup(GetDefaultElementCountBinding(), GetDefaultElementCountBinding(), ENiagaraRendererSourceDataMode::Emitter);
+		ElementCountXBinding.Setup(GetDefaultElementCountBinding(), GetDefaultElementCountBinding(), ENiagaraRendererSourceDataMode::Emitter);
+	}
+	if (ElementCountYBinding.GetType() != GetDefaultElementCountBinding().GetType())
+	{
+		ElementCountYBinding.Setup(GetDefaultElementCountBinding(), GetDefaultElementCountBinding(), ENiagaraRendererSourceDataMode::Emitter);
+	}
+	if (ElementCountZBinding.GetType() != GetDefaultElementCountBinding().GetType())
+	{
+		ElementCountZBinding.Setup(GetDefaultElementCountBinding(), GetDefaultElementCountBinding(), ENiagaraRendererSourceDataMode::Emitter);
 	}
 	if (NumIterationsBinding.GetType() != GetDefaultNumIterationsBinding().GetType())
 	{
@@ -158,6 +177,13 @@ void UNiagaraSimulationStageGeneric::PostLoad()
 	{
 		ParticleIterationStateBinding.Setup(GetDefaultParticleStateBinding(), GetDefaultParticleStateBinding(), ENiagaraRendererSourceDataMode::Particles);
 	}
+
+	if (ElementCountBinding_DEPRECATED.IsValid())
+	{
+		bOverrideGpuDispatchType = true;
+		OverrideGpuDispatchType = ENiagaraGpuDispatchType::OneD;
+		ElementCountXBinding = ElementCountBinding_DEPRECATED;
+	}
 #endif
 }
 
@@ -182,7 +208,15 @@ void UNiagaraSimulationStageGeneric::PostEditChangeProperty(struct FPropertyChan
 	{
 		bNeedsRecompile = true;
 	}
-	else if (PropertyName == GET_MEMBER_NAME_CHECKED(UNiagaraSimulationStageGeneric, ElementCountBinding))
+	else if (PropertyName == GET_MEMBER_NAME_CHECKED(UNiagaraSimulationStageGeneric, ElementCountXBinding))
+	{
+		bNeedsRecompile = true;
+	}
+	else if (PropertyName == GET_MEMBER_NAME_CHECKED(UNiagaraSimulationStageGeneric, ElementCountYBinding))
+	{
+		bNeedsRecompile = true;
+	}
+	else if (PropertyName == GET_MEMBER_NAME_CHECKED(UNiagaraSimulationStageGeneric, ElementCountZBinding))
 	{
 		bNeedsRecompile = true;
 	}
@@ -230,6 +264,14 @@ void UNiagaraSimulationStageGeneric::PostEditChangeProperty(struct FPropertyChan
 	{
 		bNeedsRecompile = true;
 	}
+	else if (PropertyName == GET_MEMBER_NAME_CHECKED(UNiagaraSimulationStageGeneric, bOverrideGpuDispatchType))
+	{
+		bNeedsRecompile = true;
+	}
+	else if (PropertyName == GET_MEMBER_NAME_CHECKED(UNiagaraSimulationStageGeneric, OverrideGpuDispatchType))
+	{
+		bNeedsRecompile = true;
+	}
 	else if (PropertyName == GET_MEMBER_NAME_CHECKED(UNiagaraSimulationStageGeneric, bOverrideGpuDispatchNumThreads))
 	{
 		bNeedsRecompile = true;
diff --git a/Engine/Plugins/FX/Niagara/Source/Niagara/Public/NiagaraSimStageData.h b/Engine/Plugins/FX/Niagara/Source/Niagara/Public/NiagaraSimStageData.h
--- a/Engine/Plugins/FX/Niagara/Source/Niagara/Public/NiagaraSimStageData.h
+++ b/Engine/Plugins/FX/Niagara/Source/Niagara/Public/NiagaraSimStageData.h
@@ -20,7 +20,7 @@ struct FNiagaraSimStageData
 
 	uint32 StageIndex = INDEX_NONE;
 	uint32 IterationIndex = 0;
-	int32 UserElementCount = -1;
+	FIntVector ElementCountXYZ = FIntVector::NoneValue;
 
 	class FNiagaraDataBuffer* Source = nullptr;
 	uint32 SourceCountOffset = INDEX_NONE;
diff --git a/Engine/Plugins/FX/Niagara/Source/Niagara/Public/NiagaraSimulationStageBase.h b/Engine/Plugins/FX/Niagara/Source/Niagara/Public/NiagaraSimulationStageBase.h
--- a/Engine/Plugins/FX/Niagara/Source/Niagara/Public/NiagaraSimulationStageBase.h
+++ b/Engine/Plugins/FX/Niagara/Source/Niagara/Public/NiagaraSimulationStageBase.h
@@ -64,13 +64,32 @@ public:
 	UPROPERTY(EditAnywhere, Category = "Simulation Stage")
 	FNiagaraVariableAttributeBinding EnabledBinding;
 
+	UPROPERTY()
+	FNiagaraVariableAttributeBinding ElementCountBinding_DEPRECATED;
+
 	/**
-	Optional integer binding to override the number of elements the stage will execute on per dispatch.
+	Integer binding to override the number of elements the stage will execute along X.
 	For example, if you want to iterate over a custom source such as triangles on a mesh you can
 	set an int to the triangle count in an emitter script and bind that as the element count.
 	*/
-	UPROPERTY(EditAnywhere, Category = "Simulation Stage")
-	FNiagaraVariableAttributeBinding ElementCountBinding;
+	UPROPERTY(EditAnywhere, Category = "Simulation Stage", meta = (DisplayName = "ElementCount X"))
+	FNiagaraVariableAttributeBinding ElementCountXBinding;
+
+	/**
+	Integer binding to override the number of elements the stage will execute along Y.
+	For example, if you want to iterate over a 2D texture you can set an int to the texture height
+	in an emitter script and bind that as the element count.
+	*/
+	UPROPERTY(EditAnywhere, Category = "Simulation Stage", meta = (DisplayName = "ElementCount Y"))
+	FNiagaraVariableAttributeBinding ElementCountYBinding;
+
+	/**
+	Integer binding to override the number of elements the stage will execute along Z.
+	For example, if you want to iterate over a 3D texture you can set an int to the texture depth
+	in an emitter script and bind that as the element count.
+	*/
+	UPROPERTY(EditAnywhere, Category = "Simulation Stage", meta = (DisplayName = "ElementCount Z"))
+	FNiagaraVariableAttributeBinding ElementCountZBinding;
 	
 	/**
 	Select what we should be iterating over, particles or data interfaces.
@@ -128,10 +147,24 @@ public:
 	UPROPERTY(EditAnywhere, Category = "Simulation Stage", meta = (editcondition = "IterationSource == ENiagaraIterationSource::Particles"))
 	FIntPoint ParticleIterationStateRange = FIntPoint(0, 0);
 
-	/** When enabled we force the dispatch to be linear (i.e. one dimension is used). */
-	UPROPERTY(EditAnywhere, Category = "Simulation Stage")
+	/**
+	When enabled we force the dispatch to be linear (i.e. one dimension is used).
+	NOTE: Cannot be used with OverrideGpuDispatchType.
+	*/
+	UPROPERTY(EditAnywhere, Category = "Simulation Stage", meta = (editcondition = "!bOverrideGpuDispatchType"))
 	uint32 bGpuDispatchForceLinear : 1;
 
+	/**
+	When enabled we use a custom number of element count for the dispatch.
+	NOTE: Cannot be used with GpuDispatchForceLinear.
+	*/
+	UPROPERTY(EditAnywhere, Category = "Simulation Stage", meta = (editcondition = "!bGpuDispatchForceLinear"))
+	uint32 bOverrideGpuDispatchType : 1;
+
+	/** Dimensions to use for dispatch. */
+	UPROPERTY(EditAnywhere, Category = "Simulation Stage", meta = (DisplayName = "Gpu Dispatch Type"))
+	ENiagaraGpuDispatchType OverrideGpuDispatchType = ENiagaraGpuDispatchType::OneD;
+
 	/** When enabled we use a custom number of threads for the dispatch. */
 	UPROPERTY(EditAnywhere, Category = "Simulation Stage")
 	uint32 bOverrideGpuDispatchNumThreads : 1;
diff --git a/Engine/Plugins/FX/Niagara/Source/Niagara/Public/NiagaraSimulationStageCompileData.h b/Engine/Plugins/FX/Niagara/Source/Niagara/Public/NiagaraSimulationStageCompileData.h
--- a/Engine/Plugins/FX/Niagara/Source/Niagara/Public/NiagaraSimulationStageCompileData.h
+++ b/Engine/Plugins/FX/Niagara/Source/Niagara/Public/NiagaraSimulationStageCompileData.h
@@ -9,7 +9,9 @@ struct FNiagaraSimulationStageCompilationData
 	FGuid                           StageGuid;
 	FName                           StageName;
 	FName                           EnabledBinding;
-	FName                           ElementCountBinding;
+	FName                           ElementCountXBinding;
+	FName                           ElementCountYBinding;
+	FName                           ElementCountZBinding;
 	uint32                          NumIterations = 1;
 	FName                           NumIterationsBinding;
 	FName                           IterationSource;
@@ -19,6 +21,8 @@ struct FNiagaraSimulationStageCompilationData
 	FName                           ParticleIterationStateBinding;
 	FIntPoint                       ParticleIterationStateRange = FIntPoint::ZeroValue;
 	bool                            bGpuDispatchForceLinear = false;
+	bool                            bOverrideGpuDispatchType = false;
+	ENiagaraGpuDispatchType         OverrideGpuDispatchType = ENiagaraGpuDispatchType::OneD;
 	bool                            bOverrideGpuDispatchNumThreads = false;
 	FIntVector                      OverrideGpuDispatchNumThreads = FIntVector(1, 1, 1);
 };
diff --git a/Engine/Plugins/FX/Niagara/Source/NiagaraEditor/Private/Customizations/NiagaraSimulationStageCustomization.cpp b/Engine/Plugins/FX/Niagara/Source/NiagaraEditor/Private/Customizations/NiagaraSimulationStageCustomization.cpp
--- a/Engine/Plugins/FX/Niagara/Source/NiagaraEditor/Private/Customizations/NiagaraSimulationStageCustomization.cpp
+++ b/Engine/Plugins/FX/Niagara/Source/NiagaraEditor/Private/Customizations/NiagaraSimulationStageCustomization.cpp
@@ -66,7 +66,9 @@ void FNiagaraSimulationStageGenericCustomization::CustomizeDetails(class IDetail
 
 	// Hide all
 	DetailBuilder.HideProperty(GET_MEMBER_NAME_CHECKED(UNiagaraSimulationStageGeneric, EnabledBinding));
-	DetailBuilder.HideProperty(GET_MEMBER_NAME_CHECKED(UNiagaraSimulationStageGeneric, ElementCountBinding));
+	DetailBuilder.HideProperty(GET_MEMBER_NAME_CHECKED(UNiagaraSimulationStageGeneric, ElementCountXBinding));
+	DetailBuilder.HideProperty(GET_MEMBER_NAME_CHECKED(UNiagaraSimulationStageGeneric, ElementCountYBinding));
+	DetailBuilder.HideProperty(GET_MEMBER_NAME_CHECKED(UNiagaraSimulationStageGeneric, ElementCountZBinding));
 	DetailBuilder.HideProperty(GET_MEMBER_NAME_CHECKED(UNiagaraSimulationStageGeneric, IterationSource));
 	DetailBuilder.HideProperty(GET_MEMBER_NAME_CHECKED(UNiagaraSimulationStageGeneric, Iterations));
 	DetailBuilder.HideProperty(GET_MEMBER_NAME_CHECKED(UNiagaraSimulationStageGeneric, NumIterationsBinding));
@@ -77,15 +79,35 @@ void FNiagaraSimulationStageGenericCustomization::CustomizeDetails(class IDetail
 	DetailBuilder.HideProperty(GET_MEMBER_NAME_CHECKED(UNiagaraSimulationStageGeneric, ParticleIterationStateBinding));
 	DetailBuilder.HideProperty(GET_MEMBER_NAME_CHECKED(UNiagaraSimulationStageGeneric, ParticleIterationStateRange));
 	DetailBuilder.HideProperty(GET_MEMBER_NAME_CHECKED(UNiagaraSimulationStageGeneric, bGpuDispatchForceLinear));
+	DetailBuilder.HideProperty(GET_MEMBER_NAME_CHECKED(UNiagaraSimulationStageGeneric, bOverrideGpuDispatchType));
+	DetailBuilder.HideProperty(GET_MEMBER_NAME_CHECKED(UNiagaraSimulationStageGeneric, OverrideGpuDispatchType));
 	DetailBuilder.HideProperty(GET_MEMBER_NAME_CHECKED(UNiagaraSimulationStageGeneric, bOverrideGpuDispatchNumThreads));
 	DetailBuilder.HideProperty(GET_MEMBER_NAME_CHECKED(UNiagaraSimulationStageGeneric, OverrideGpuDispatchNumThreads));
 
 	// Show properties in the order we want them to appear
 	SimStageCategory.AddProperty(DetailBuilder.GetProperty(GET_MEMBER_NAME_CHECKED(UNiagaraSimulationStageGeneric, EnabledBinding)));
-	SimStageCategory.AddProperty(DetailBuilder.GetProperty(GET_MEMBER_NAME_CHECKED(UNiagaraSimulationStageGeneric, ElementCountBinding)));
 	SimStageCategory.AddProperty(DetailBuilder.GetProperty(GET_MEMBER_NAME_CHECKED(UNiagaraSimulationStageGeneric, Iterations)));
 	SimStageCategory.AddProperty(DetailBuilder.GetProperty(GET_MEMBER_NAME_CHECKED(UNiagaraSimulationStageGeneric, NumIterationsBinding)));
 	SimStageCategory.AddProperty(DetailBuilder.GetProperty(GET_MEMBER_NAME_CHECKED(UNiagaraSimulationStageGeneric, IterationSource)));
+
+	DataInterfaceStageCategory.AddProperty(DetailBuilder.GetProperty(GET_MEMBER_NAME_CHECKED(UNiagaraSimulationStageGeneric, bOverrideGpuDispatchType)));
+	if (SimStage->bOverrideGpuDispatchType)
+	{
+		DataInterfaceStageCategory.AddProperty(DetailBuilder.GetProperty(GET_MEMBER_NAME_CHECKED(UNiagaraSimulationStageGeneric, OverrideGpuDispatchType)));
+		if (SimStage->OverrideGpuDispatchType >= ENiagaraGpuDispatchType::OneD)
+		{
+			DataInterfaceStageCategory.AddProperty(DetailBuilder.GetProperty(GET_MEMBER_NAME_CHECKED(UNiagaraSimulationStageGeneric, ElementCountXBinding)));
+		}
+		if (SimStage->OverrideGpuDispatchType >= ENiagaraGpuDispatchType::TwoD)
+		{
+			DataInterfaceStageCategory.AddProperty(DetailBuilder.GetProperty(GET_MEMBER_NAME_CHECKED(UNiagaraSimulationStageGeneric, ElementCountYBinding)));
+		}
+		if (SimStage->OverrideGpuDispatchType >= ENiagaraGpuDispatchType::ThreeD)
+		{
+			DataInterfaceStageCategory.AddProperty(DetailBuilder.GetProperty(GET_MEMBER_NAME_CHECKED(UNiagaraSimulationStageGeneric, ElementCountZBinding)));
+		}
+	}
+
 	if ( SimStage->IterationSource == ENiagaraIterationSource::Particles )
 	{
 		ParticleStageCategory.AddProperty(DetailBuilder.GetProperty(GET_MEMBER_NAME_CHECKED(UNiagaraSimulationStageGeneric, bDisablePartialParticleUpdate)));
@@ -102,6 +124,7 @@ void FNiagaraSimulationStageGenericCustomization::CustomizeDetails(class IDetail
 		DataInterfaceStageCategory.AddProperty(DetailBuilder.GetProperty(GET_MEMBER_NAME_CHECKED(UNiagaraSimulationStageGeneric, ExecuteBehavior)));
 
 		DataInterfaceStageCategory.AddProperty(DetailBuilder.GetProperty(GET_MEMBER_NAME_CHECKED(UNiagaraSimulationStageGeneric, bGpuDispatchForceLinear)));
+
 		DataInterfaceStageCategory.AddProperty(DetailBuilder.GetProperty(GET_MEMBER_NAME_CHECKED(UNiagaraSimulationStageGeneric, bOverrideGpuDispatchNumThreads)));
 		if (SimStage->bOverrideGpuDispatchNumThreads)
 		{
diff --git a/Engine/Plugins/FX/Niagara/Source/NiagaraEditor/Private/NiagaraHlslTranslator.cpp b/Engine/Plugins/FX/Niagara/Source/NiagaraEditor/Private/NiagaraHlslTranslator.cpp
--- a/Engine/Plugins/FX/Niagara/Source/NiagaraEditor/Private/NiagaraHlslTranslator.cpp
+++ b/Engine/Plugins/FX/Niagara/Source/NiagaraEditor/Private/NiagaraHlslTranslator.cpp
@@ -1255,7 +1255,9 @@ const FNiagaraTranslateResults &FHlslNiagaraTranslator::Translate(const FNiagara
 					TranslationStages[Index].bCopyPreviousParams = false;
 					TranslationStages[Index].SimulationStageIndex = DestSimStageIndex;					
 					TranslationStages[Index].EnabledBinding = CompileSimStageData.EnabledBinding;
-					TranslationStages[Index].ElementCountBinding = CompileSimStageData.ElementCountBinding;
+					TranslationStages[Index].ElementCountXBinding = CompileSimStageData.ElementCountXBinding;
+					TranslationStages[Index].ElementCountYBinding = CompileSimStageData.ElementCountYBinding;
+					TranslationStages[Index].ElementCountZBinding = CompileSimStageData.ElementCountZBinding;
 					TranslationStages[Index].NumIterations = CompileSimStageData.NumIterations;
 					TranslationStages[Index].ExecuteBehavior = CompileSimStageData.ExecuteBehavior;
 					TranslationStages[Index].bPartialParticleUpdate = CompileSimStageData.PartialParticleUpdate;
@@ -1265,6 +1267,8 @@ const FNiagaraTranslateResults &FHlslNiagaraTranslator::Translate(const FNiagara
 					TranslationStages[Index].ParticleIterationStateBinding = CompileSimStageData.ParticleIterationStateBinding;
 					TranslationStages[Index].ParticleIterationStateRange = CompileSimStageData.ParticleIterationStateRange;
 					TranslationStages[Index].bGpuDispatchForceLinear = CompileSimStageData.bGpuDispatchForceLinear;
+					TranslationStages[Index].bOverrideGpuDispatchType = CompileSimStageData.bOverrideGpuDispatchType;
+					TranslationStages[Index].OverrideGpuDispatchType = CompileSimStageData.OverrideGpuDispatchType;
 					TranslationStages[Index].bOverrideGpuDispatchNumThreads = CompileSimStageData.bOverrideGpuDispatchNumThreads;
 					TranslationStages[Index].OverrideGpuDispatchNumThreads = CompileSimStageData.OverrideGpuDispatchNumThreads;
 
@@ -1325,7 +1329,9 @@ const FNiagaraTranslateResults &FHlslNiagaraTranslator::Translate(const FNiagara
 					FSimulationStageMetaData& SimulationStageMetaData = CompilationOutput.ScriptData.SimulationStageMetaData.AddDefaulted_GetRef();
 					SimulationStageMetaData.SimulationStageName = InCompileData->CompileSimStageData[SourceSimStageIndex].StageName;
 					SimulationStageMetaData.EnabledBinding = InCompileData->CompileSimStageData[SourceSimStageIndex].EnabledBinding;
-					SimulationStageMetaData.ElementCountBinding = InCompileData->CompileSimStageData[SourceSimStageIndex].ElementCountBinding;
+					SimulationStageMetaData.ElementCountXBinding = InCompileData->CompileSimStageData[SourceSimStageIndex].ElementCountXBinding;
+					SimulationStageMetaData.ElementCountYBinding = InCompileData->CompileSimStageData[SourceSimStageIndex].ElementCountYBinding;
+					SimulationStageMetaData.ElementCountZBinding = InCompileData->CompileSimStageData[SourceSimStageIndex].ElementCountZBinding;
 					SimulationStageMetaData.ExecuteBehavior = TranslationStages[Index].ExecuteBehavior;
 					SimulationStageMetaData.IterationSource = InCompileData->CompileSimStageData[SourceSimStageIndex].IterationSource;
 					SimulationStageMetaData.NumIterationsBinding = InCompileData->CompileSimStageData[SourceSimStageIndex].NumIterationsBinding;
@@ -1335,6 +1341,7 @@ const FNiagaraTranslateResults &FHlslNiagaraTranslator::Translate(const FNiagara
 					SimulationStageMetaData.bParticleIterationStateEnabled = TranslationStages[Index].bParticleIterationStateEnabled;
 					SimulationStageMetaData.ParticleIterationStateBinding = TranslationStages[Index].ParticleIterationStateBinding;
 					SimulationStageMetaData.ParticleIterationStateRange = TranslationStages[Index].ParticleIterationStateRange;
+					SimulationStageMetaData.bOverrideElementCount = TranslationStages[Index].bOverrideGpuDispatchType;
 
 					// Determine dispatch information from iteration source (if we have one)
 					SimulationStageMetaData.GpuDispatchType = ENiagaraGpuDispatchType::OneD;
@@ -1345,15 +1352,22 @@ const FNiagaraTranslateResults &FHlslNiagaraTranslator::Translate(const FNiagara
 						{
 							if ( UNiagaraDataInterface* IteratinoSourceCDO = CompileDuplicateData->GetDuplicatedDataInterfaceCDOForClass(IterationSourceVar->GetType().GetClass()) )
 							{
-								if ( TranslationStages[Index].bOverrideGpuDispatchNumThreads )
+								if (TranslationStages[Index].bGpuDispatchForceLinear)
+								{
+									SimulationStageMetaData.GpuDispatchType = ENiagaraGpuDispatchType::OneD;
+								}
+								else
+								{
+									SimulationStageMetaData.GpuDispatchType = TranslationStages[Index].bOverrideGpuDispatchType ? TranslationStages[Index].OverrideGpuDispatchType : IteratinoSourceCDO->GetGpuDispatchType();
+								}
+
+								if (TranslationStages[Index].bOverrideGpuDispatchNumThreads)
 								{
-									SimulationStageMetaData.GpuDispatchType = TranslationStages[Index].bGpuDispatchForceLinear ? ENiagaraGpuDispatchType::OneD : ENiagaraGpuDispatchType::Custom;
 									SimulationStageMetaData.GpuDispatchNumThreads = TranslationStages[Index].OverrideGpuDispatchNumThreads;
 								}
 								else
 								{
-									SimulationStageMetaData.GpuDispatchType = TranslationStages[Index].bGpuDispatchForceLinear ? ENiagaraGpuDispatchType::OneD : IteratinoSourceCDO->GetGpuDispatchType();
-									SimulationStageMetaData.GpuDispatchNumThreads = SimulationStageMetaData.GpuDispatchType == ENiagaraGpuDispatchType::Custom ? IteratinoSourceCDO->GetGpuDispatchNumThreads() : FNiagaraShader::GetDefaultThreadGroupSize(SimulationStageMetaData.GpuDispatchType);
+									SimulationStageMetaData.GpuDispatchNumThreads = (SimulationStageMetaData.GpuDispatchType == ENiagaraGpuDispatchType::Custom) ? IteratinoSourceCDO->GetGpuDispatchNumThreads() : FNiagaraShader::GetDefaultThreadGroupSize(SimulationStageMetaData.GpuDispatchType);
 								}
 							}
 						}
diff --git a/Engine/Plugins/FX/Niagara/Source/NiagaraEditor/Private/NiagaraHlslTranslator.h b/Engine/Plugins/FX/Niagara/Source/NiagaraEditor/Private/NiagaraHlslTranslator.h
--- a/Engine/Plugins/FX/Niagara/Source/NiagaraEditor/Private/NiagaraHlslTranslator.h
+++ b/Engine/Plugins/FX/Niagara/Source/NiagaraEditor/Private/NiagaraHlslTranslator.h
@@ -368,7 +368,9 @@ public:
 	FName IterationSource;
 	int32 SimulationStageIndex = -1;
 	FName EnabledBinding;
-	FName ElementCountBinding;
+	FName ElementCountXBinding;
+	FName ElementCountYBinding;
+	FName ElementCountZBinding;
 	int32 NumIterations = 1;
 	FName NumIterationsBinding;
 	ENiagaraSimStageExecuteBehavior ExecuteBehavior = ENiagaraSimStageExecuteBehavior::Always;
@@ -380,6 +382,8 @@ public:
 	bool bWritesParticles = false;
 	bool bPartialParticleUpdate = false;
 	bool bGpuDispatchForceLinear = false;
+	bool bOverrideGpuDispatchType = false;
+	ENiagaraGpuDispatchType OverrideGpuDispatchType = ENiagaraGpuDispatchType::OneD;
 	bool bOverrideGpuDispatchNumThreads = false;
 	FIntVector OverrideGpuDispatchNumThreads = FIntVector(1, 1, 1);
 	TArray<FNiagaraVariable> SetParticleAttributes;
diff --git a/Engine/Plugins/FX/Niagara/Source/NiagaraShader/Private/NiagaraScriptBase.cpp b/Engine/Plugins/FX/Niagara/Source/NiagaraShader/Private/NiagaraScriptBase.cpp
--- a/Engine/Plugins/FX/Niagara/Source/NiagaraShader/Private/NiagaraScriptBase.cpp
+++ b/Engine/Plugins/FX/Niagara/Source/NiagaraShader/Private/NiagaraScriptBase.cpp
@@ -11,6 +11,7 @@ FSimulationStageMetaData::FSimulationStageMetaData()
 	: bWritesParticles(0)
 	, bPartialParticleUpdate(0)
 	, bParticleIterationStateEnabled(0)
+	, bOverrideElementCount(0)
 	, GpuDispatchType(ENiagaraGpuDispatchType::OneD)
 	, GpuDispatchNumThreads(0, 0, 0)
 {
diff --git a/Engine/Plugins/FX/Niagara/Source/NiagaraShader/Public/NiagaraScriptBase.h b/Engine/Plugins/FX/Niagara/Source/NiagaraShader/Public/NiagaraScriptBase.h
--- a/Engine/Plugins/FX/Niagara/Source/NiagaraShader/Public/NiagaraScriptBase.h
+++ b/Engine/Plugins/FX/Niagara/Source/NiagaraShader/Public/NiagaraScriptBase.h
@@ -17,7 +17,7 @@ enum class ENiagaraGpuDispatchType : uint8
 	/* Instances will distribute along X, Y & Z using platforms specific thread counts. */
 	ThreeD,
 	/* NumThreads will be determined manually. */
-	Custom,
+	Custom			UMETA(Hidden),
 };
 
 UENUM()
@@ -47,7 +47,15 @@ public:
 
 	/** Optional binding to manually specify the element count. */
 	UPROPERTY()
-	FName ElementCountBinding;
+	FName ElementCountXBinding;
+
+	/** Optional binding to manually specify the element count. */
+	UPROPERTY()
+	FName ElementCountYBinding;
+
+	/** Optional binding to manually specify the element count. */
+	UPROPERTY()
+	FName ElementCountZBinding;
 
 	/** The Data Interface that we iterate over for this stage. If None, then use particles.*/
 	UPROPERTY()
@@ -68,6 +76,10 @@ public:
 	UPROPERTY()
 	uint32 bParticleIterationStateEnabled : 1;
 
+	/** When enabled the simulation stage uses element count provided by user to dispatch work. */
+	UPROPERTY()
+	uint32 bOverrideElementCount : 1;
+
 	/** When the value is not none this is the binding used for particle state iteration stages. */
 	UPROPERTY()
 	FName ParticleIterationStateBinding;
