- Fix static analysis issue

#rnx

diff --git a/Engine/Plugins/FX/Niagara/Source/NiagaraEditor/Private/Customizations/NiagaraSimulationStageCustomization.cpp b/Engine/Plugins/FX/Niagara/Source/NiagaraEditor/Private/Customizations/NiagaraSimulationStageCustomization.cpp
--- a/Engine/Plugins/FX/Niagara/Source/NiagaraEditor/Private/Customizations/NiagaraSimulationStageCustomization.cpp
+++ b/Engine/Plugins/FX/Niagara/Source/NiagaraEditor/Private/Customizations/NiagaraSimulationStageCustomization.cpp
@@ -94,7 +94,8 @@ void FNiagaraSimulationStageGenericCustomization::CustomizeDetails(class IDetail
 	if (SimStage->bOverrideGpuDispatchType)
 	{
 		DataInterfaceStageCategory.AddProperty(DetailBuilder.GetProperty(GET_MEMBER_NAME_CHECKED(UNiagaraSimulationStageGeneric, OverrideGpuDispatchType)));
-		if (SimStage->OverrideGpuDispatchType >= ENiagaraGpuDispatchType::OneD)
+		// Always true as we always dispatch across at least 1 dimension
+		//if (SimStage->OverrideGpuDispatchType >= ENiagaraGpuDispatchType::OneD)
 		{
 			DataInterfaceStageCategory.AddProperty(DetailBuilder.GetProperty(GET_MEMBER_NAME_CHECKED(UNiagaraSimulationStageGeneric, ElementCountXBinding)));
 		}
