Removed temp debug logging

diff --git a/Engine/Plugins/FX/Niagara/Source/NiagaraEditor/Private/NiagaraScriptMergeManager.cpp b/Engine/Plugins/FX/Niagara/Source/NiagaraEditor/Private/NiagaraScriptMergeManager.cpp
--- a/Engine/Plugins/FX/Niagara/Source/NiagaraEditor/Private/NiagaraScriptMergeManager.cpp
+++ b/Engine/Plugins/FX/Niagara/Source/NiagaraEditor/Private/NiagaraScriptMergeManager.cpp
@@ -1406,9 +1406,6 @@ void FNiagaraScriptMergeManager::UpdateModuleVersions(const FVersionedNiagaraEmi
 	UNiagaraSystem* System = nullptr;
 	if (ChangedVersionOtherModules.Num() > 0)
 	{
-		// TODO (mga) REMOVE - debug log to find cook problems
-		UE_LOG(LogNiagaraEditor, Log, TEXT("Upgrading module version for emitter %s"), *GetPathNameSafe(Instance.Emitter));
-
 		TSharedPtr<FNiagaraSystemViewModel> SystemViewModel;
 
 		if (!FSlateApplication::IsInitialized())
