Fixed a crash when binding an invalid object to a niagara component track in sequencer

diff --git a/Engine/Plugins/FX/Niagara/Source/Niagara/Private/MovieScene/MovieSceneNiagaraSystemTrackTemplate.cpp b/Engine/Plugins/FX/Niagara/Source/Niagara/Private/MovieScene/MovieSceneNiagaraSystemTrackTemplate.cpp
--- a/Engine/Plugins/FX/Niagara/Source/Niagara/Private/MovieScene/MovieSceneNiagaraSystemTrackTemplate.cpp
+++ b/Engine/Plugins/FX/Niagara/Source/Niagara/Private/MovieScene/MovieSceneNiagaraSystemTrackTemplate.cpp
@@ -119,8 +119,11 @@ struct FNiagaraSystemUpdateDesiredAgeExecutionToken : IMovieSceneExecutionToken
 	{
 		for (TWeakObjectPtr<> Object : Player.FindBoundObjects(Operand))
 		{
-			UObject* ObjectPtr = Object.Get();
-			UNiagaraComponent* NiagaraComponent = Cast<UNiagaraComponent>(ObjectPtr);
+			UNiagaraComponent* NiagaraComponent = Cast<UNiagaraComponent>(Object);
+			if (!NiagaraComponent)
+			{
+				return;
+			}
 
 			{
 				static FMovieSceneAnimTypeID TypeID = TMovieSceneAnimTypeID<FNiagaraSystemUpdateDesiredAgeExecutionToken, 0>();
