Fixed a bug where newly created component renderers used the wrong type for the position attribute

diff --git a/Engine/Plugins/FX/Niagara/Source/Niagara/Private/NiagaraComponentRendererProperties.cpp b/Engine/Plugins/FX/Niagara/Source/Niagara/Private/NiagaraComponentRendererProperties.cpp
--- a/Engine/Plugins/FX/Niagara/Source/Niagara/Private/NiagaraComponentRendererProperties.cpp
+++ b/Engine/Plugins/FX/Niagara/Source/Niagara/Private/NiagaraComponentRendererProperties.cpp
@@ -213,14 +213,14 @@ void UNiagaraComponentRendererProperties::PostLoad()
 	{
 		Binding.AttributeBinding.PostLoad(InSourceMode);
 
-		// Move old bindings over to new position type 
+		// Move old bindings over to new position type
 		for (FNiagaraVariable OldVarType : OldTypes)
 		{
-			if (Binding.AttributeBinding.GetParamMapBindableVariable() == static_cast<const FNiagaraVariableBase&>(OldVarType))
+			if (Binding.AttributeBinding.GetParamMapBindableVariable() == OldVarType)
 			{
 				FNiagaraVariable NewVarType(FNiagaraTypeDefinition::GetPositionDef(), OldVarType.GetName());
 				Binding.AttributeBinding.Setup(NewVarType, NewVarType, InSourceMode);
-				Binding.PropertyType = FNiagaraTypeDefinition(FNiagaraTypeDefinition::GetVec3Struct());
+				Binding.PropertyType = GetFVectorDef();
 				break;
 			}
 		}
@@ -547,7 +547,7 @@ void UNiagaraComponentRendererProperties::PostEditChangeProperty(struct FPropert
 			FNiagaraComponentPropertyBinding PositionBinding;
 			PositionBinding.AttributeBinding.Setup(SYS_PARAM_PARTICLES_POSITION, SYS_PARAM_PARTICLES_POSITION);
 			PositionBinding.PropertyName = FName("RelativeLocation");
-			PositionBinding.PropertyType = FNiagaraTypeDefinition(FNiagaraTypeDefinition::GetVec3Struct());
+			PositionBinding.PropertyType = GetFVectorDef();
 			PropertyBindings.Add(PositionBinding);
 
 			FNiagaraComponentPropertyBinding ScaleBinding;
diff --git a/Engine/Plugins/FX/Niagara/Source/Niagara/Private/NiagaraRendererComponents.cpp b/Engine/Plugins/FX/Niagara/Source/Niagara/Private/NiagaraRendererComponents.cpp
--- a/Engine/Plugins/FX/Niagara/Source/Niagara/Private/NiagaraRendererComponents.cpp
+++ b/Engine/Plugins/FX/Niagara/Source/Niagara/Private/NiagaraRendererComponents.cpp
@@ -114,7 +114,7 @@ void ConvertVariableToType(const FNiagaraVariable& SourceVariable, const FNiagar
 	}
 	else if (SourceType == FNiagaraTypeDefinition::GetPositionDef() && TargetType == UNiagaraComponentRendererProperties::GetFVector3fDef())
 	{
-		const FVector3f Data = (FVector3f)LwcConverter.ConvertSimulationPositionToWorld(SourceVariable.GetValue<FNiagaraPosition>());
+		const FVector3f Data = FVector3f(LwcConverter.ConvertSimulationPositionToWorld(SourceVariable.GetValue<FNiagaraPosition>()));
 		TargetData.AddUninitialized(sizeof(FVector3f));
 		*reinterpret_cast<FVector3f*>(TargetData.GetData()) = Data;
 	}
