- Flush EOF updates / GPU updates when capturing from a GPU emitter

#rnx

diff --git a/Engine/Plugins/FX/Niagara/Source/Niagara/Private/NiagaraSimCache.cpp b/Engine/Plugins/FX/Niagara/Source/Niagara/Private/NiagaraSimCache.cpp
--- a/Engine/Plugins/FX/Niagara/Source/Niagara/Private/NiagaraSimCache.cpp
+++ b/Engine/Plugins/FX/Niagara/Source/Niagara/Private/NiagaraSimCache.cpp
@@ -199,6 +199,7 @@ bool UNiagaraSimCache::WriteFrame(UNiagaraComponent* NiagaraComponent)
 
 	Helper.WriteDataBuffer(*Helper.SystemSimulationDataBuffer, CacheLayout.SystemLayout, CacheFrame.SystemData.SystemDataBuffers, Helper.SystemInstance->GetSystemInstanceIndex(), 1);
 
+	bool bNeedsGpuTickFlush = true;
 	for ( int32 i=0; i < NumEmitters; ++i )
 	{
 		FNiagaraSimCacheEmitterFrame& CacheEmitterFrame = CacheFrame.EmitterData[i];
@@ -213,6 +214,24 @@ bool UNiagaraSimCache::WriteFrame(UNiagaraComponent* NiagaraComponent)
 		CacheEmitterFrame.TotalSpawnedParticles = EmitterInstance.GetTotalSpawnedParticles();
 		if (CacheLayout.EmitterLayouts[i].SimTarget == ENiagaraSimTarget::GPUComputeSim)
 		{
+			// First time we encounter a GPU emitter we need to process we also need to flush EOF updates / the compute process to ensure we get latest data
+			if (bNeedsGpuTickFlush)
+			{
+				UWorld* World = NiagaraComponent->GetWorld();
+				if (ensure(World))
+				{
+					World->SendAllEndOfFrameUpdates();
+
+					FNiagaraGpuComputeDispatchInterface* ComputeDispatchInterface = FNiagaraGpuComputeDispatchInterface::Get(World);
+					if (ensureMsgf(ComputeDispatchInterface, TEXT("Compute dispatch interface was invalid when flushing commands")))
+					{
+						ComputeDispatchInterface->FlushPendingTicks_GameThread();
+					}
+				}
+
+				bNeedsGpuTickFlush = false;
+			}
+
 			Helper.WriteDataBufferGPU(EmitterInstance, *EmitterCurrentData, CacheLayout.EmitterLayouts[i], CacheEmitterFrame.ParticleDataBuffers);
 		}
 		else
