- Make sure we mark the object modified when baking

diff --git a/Engine/Plugins/FX/Niagara/Source/Niagara/Classes/NiagaraSystem.h b/Engine/Plugins/FX/Niagara/Source/Niagara/Classes/NiagaraSystem.h
--- a/Engine/Plugins/FX/Niagara/Source/Niagara/Classes/NiagaraSystem.h
+++ b/Engine/Plugins/FX/Niagara/Source/Niagara/Classes/NiagaraSystem.h
@@ -713,8 +713,7 @@ public:
 #if WITH_EDITORONLY_DATA
 	UNiagaraBakerSettings* GetBakerSettings();
 	const UNiagaraBakerSettings* GetBakerGeneratedSettings() const { return BakerGeneratedSettings; }
-	void SetBakerGeneratedSettings(UNiagaraBakerSettings* Settings) { BakerGeneratedSettings = Settings; }
-
+	void SetBakerGeneratedSettings(UNiagaraBakerSettings* Settings);
 
 	/** Get the cached parameter map traversal for this emitter.  */
 	const TSharedPtr<FNiagaraGraphCachedDataBase, ESPMode::ThreadSafe>& GetCachedTraversalData() const;
diff --git a/Engine/Plugins/FX/Niagara/Source/Niagara/Private/NiagaraSystem.cpp b/Engine/Plugins/FX/Niagara/Source/Niagara/Private/NiagaraSystem.cpp
--- a/Engine/Plugins/FX/Niagara/Source/Niagara/Private/NiagaraSystem.cpp
+++ b/Engine/Plugins/FX/Niagara/Source/Niagara/Private/NiagaraSystem.cpp
@@ -3740,10 +3740,23 @@ UNiagaraBakerSettings* UNiagaraSystem::GetBakerSettings()
 {
 	if ( BakerSettings == nullptr )
 	{
+		Modify();
 		BakerSettings = NewObject<UNiagaraBakerSettings>(this, "BakerSettings", RF_Transactional);
+		PostEditChange();
 	}
 	return BakerSettings;
 }
+
+void UNiagaraSystem::SetBakerGeneratedSettings(UNiagaraBakerSettings* Settings)
+{
+	if (BakerGeneratedSettings != Settings)
+	{
+		Modify();
+		BakerGeneratedSettings = Settings;
+		PostEditChange();
+	}
+}
+
 #endif
 
 #undef LOCTEXT_NAMESPACE // NiagaraSystem
