- Test data interface sim caching with Volume RT

#rnx

diff --git a/Engine/Plugins/FX/Niagara/Source/Niagara/Classes/NiagaraDataInterfaceRenderTargetVolume.h b/Engine/Plugins/FX/Niagara/Source/Niagara/Classes/NiagaraDataInterfaceRenderTargetVolume.h
--- a/Engine/Plugins/FX/Niagara/Source/Niagara/Classes/NiagaraDataInterfaceRenderTargetVolume.h
+++ b/Engine/Plugins/FX/Niagara/Source/Niagara/Classes/NiagaraDataInterfaceRenderTargetVolume.h
@@ -9,7 +9,6 @@
 class FNiagaraSystemInstance;
 class UTextureRenderTargetVolume;
 
-
 struct FRenderTargetVolumeRWInstanceData_GameThread
 {
 	FRenderTargetVolumeRWInstanceData_GameThread()
@@ -116,6 +115,11 @@ public:
 	virtual bool CanExposeVariables() const override { return true;}
 	virtual void GetExposedVariables(TArray<FNiagaraVariableBase>& OutVariables) const override;
 	virtual bool GetExposedVariableValue(const FNiagaraVariableBase& InVariable, void* InPerInstanceData, FNiagaraSystemInstance* InSystemInstance, void* OutData) const override;
+
+	virtual UObject* SimCacheBeginWrite(UObject* SimCache, FNiagaraSystemInstance* NiagaraSystemInstance, const void* OptionalPerInstanceData) const override;
+	virtual bool SimCacheWriteFrame(UObject* StorageObject, int FrameIndex, FNiagaraSystemInstance* SystemInstance, const void* OptionalPerInstanceData) const override;
+	virtual bool SimCacheEndWrite(UObject* StorageObject) const override;
+	virtual bool SimCacheReadFrame(UObject* StorageObject, int FrameA, int FrameB, float Interp, FNiagaraSystemInstance* SystemInstance, void* OptionalPerInstanceData) override;
 	//~ UNiagaraDataInterface interface END
 	
 	void VMGetSize(FVectorVMExternalFunctionContext& Context); 
diff --git a/Engine/Plugins/FX/Niagara/Source/Niagara/Private/NDIRenderTargetVolumeSimCacheData.h b/Engine/Plugins/FX/Niagara/Source/Niagara/Private/NDIRenderTargetVolumeSimCacheData.h
new file mode 100644
--- /dev/null
+++ b/Engine/Plugins/FX/Niagara/Source/Niagara/Private/NDIRenderTargetVolumeSimCacheData.h
@@ -0,0 +1,88 @@
+// Copyright Epic Games, Inc. All Rights Reserved.
+#pragma once
+
+#include "NDIRenderTargetVolumeSimCacheData.generated.h"
+
+USTRUCT()
+struct FNDIRenderTargetVolumeSimCacheFrame
+{
+	GENERATED_BODY();
+
+	~FNDIRenderTargetVolumeSimCacheFrame()
+	{
+		ReleasePixelData();
+	}
+
+	uint8* GetPixelData()
+	{
+		if (PixelData == nullptr && CompressedSize > 0)
+		{
+			BulkData.GetCopy(reinterpret_cast<void**>(&PixelData), true);
+		}
+		return PixelData;
+	}
+
+	void ReleasePixelData()
+	{
+		if (PixelData != nullptr)
+		{
+			FMemory::Free(PixelData);
+			PixelData = nullptr;
+		}
+	}
+
+	UPROPERTY()
+	FIntVector Size = FIntVector(EForceInit::ForceInitToZero);
+
+	UPROPERTY()
+	TEnumAsByte<EPixelFormat> Format = EPixelFormat::PF_A16B16G16R16;
+
+	UPROPERTY()
+	int32 UncompressedSize = 0;
+
+	UPROPERTY()
+	int32 CompressedSize = 0;
+
+	FByteBulkData BulkData;
+	uint8* PixelData = nullptr;
+};
+
+UCLASS()
+class UNDIRenderTargetVolumeSimCacheData : public UObject
+{
+	GENERATED_BODY()
+
+public:
+	// Begin UObject Interface
+	virtual void BeginDestroy() override
+	{
+		Super::BeginDestroy();
+
+		for (FNDIRenderTargetVolumeSimCacheFrame& Frame : Frames)
+		{
+			Frame.ReleasePixelData();
+		}
+		Frames.Empty();
+	}
+
+	virtual void Serialize(FArchive& Ar) override
+	{
+		Super::Serialize(Ar);
+
+		for (FNDIRenderTargetVolumeSimCacheFrame& Frame : Frames)
+		{
+			if (Frame.UncompressedSize == 0)
+			{
+				continue;
+			}
+
+			Frame.BulkData.Serialize(Ar, this);
+		}
+	}
+
+	UPROPERTY()
+	FName CompressionType;
+
+	UPROPERTY()
+	TArray<FNDIRenderTargetVolumeSimCacheFrame> Frames;
+};
diff --git a/Engine/Plugins/FX/Niagara/Source/Niagara/Private/NiagaraDataInterfaceRenderTargetVolume.cpp b/Engine/Plugins/FX/Niagara/Source/Niagara/Private/NiagaraDataInterfaceRenderTargetVolume.cpp
--- a/Engine/Plugins/FX/Niagara/Source/Niagara/Private/NiagaraDataInterfaceRenderTargetVolume.cpp
+++ b/Engine/Plugins/FX/Niagara/Source/Niagara/Private/NiagaraDataInterfaceRenderTargetVolume.cpp
@@ -6,6 +6,7 @@
 #include "TextureRenderTargetVolumeResource.h"
 #include "Engine/TextureRenderTargetVolume.h"
 
+#include "NDIRenderTargetVolumeSimCacheData.h"
 #include "NiagaraDataInterfaceRenderTargetCommon.h"
 #include "NiagaraGpuComputeDispatchInterface.h"
 #include "NiagaraSystemInstance.h"
@@ -49,6 +50,23 @@ namespace NDIRenderTargetVolumeLocal
 			LatestVersion = VersionPlusOne - 1
 		};
 	};
+
+	int32 GSimCacheEnabled = true;
+	static FAutoConsoleVariableRef CVarSimCacheEnabled(
+		TEXT("fx.Niagara.RenderTargetVolume.SimCacheEnabled"),
+		GSimCacheEnabled,
+		TEXT("When enabled we can write data into the simulation cache."),
+		ECVF_Default
+	);
+
+	int32 GSimCacheCompressed = true;
+	static FAutoConsoleVariableRef CVarSimCacheCompressed(
+		TEXT("fx.Niagara.RenderTargetVolume.SimCacheCompressed"),
+		GSimCacheCompressed,
+		TEXT("When enabled compression is used for the sim cache data."),
+		ECVF_Default
+	);
+	static const FName GetSimCacheCompressionType() { return GSimCacheCompressed ? NAME_Oodle : NAME_None; }
 }
 
 FNiagaraVariableBase UNiagaraDataInterfaceRenderTargetVolume::ExposedRTVar;
@@ -513,6 +531,167 @@ bool UNiagaraDataInterfaceRenderTargetVolume::GetExposedVariableValue(const FNia
 	return false;
 }
 
+UObject* UNiagaraDataInterfaceRenderTargetVolume::SimCacheBeginWrite(UObject* SimCache, FNiagaraSystemInstance* NiagaraSystemInstance, const void* OptionalPerInstanceData) const
+{
+	UNDIRenderTargetVolumeSimCacheData* SimCacheData = nullptr;
+	if (NDIRenderTargetVolumeLocal::GSimCacheEnabled)
+	{
+		SimCacheData = NewObject<UNDIRenderTargetVolumeSimCacheData>(SimCache);
+		SimCacheData->CompressionType = NDIRenderTargetVolumeLocal::GetSimCacheCompressionType();
+	}
+	return SimCacheData;
+}
+
+bool UNiagaraDataInterfaceRenderTargetVolume::SimCacheWriteFrame(UObject* StorageObject, int FrameIndex, FNiagaraSystemInstance* SystemInstance, const void* OptionalPerInstanceData) const
+{
+	check(OptionalPerInstanceData);
+
+	UNDIRenderTargetVolumeSimCacheData* SimCacheData = CastChecked<UNDIRenderTargetVolumeSimCacheData>(StorageObject);
+	SimCacheData->Frames.SetNum(FMath::Max(SimCacheData->Frames.Num(), FrameIndex + 1));
+
+	FNDIRenderTargetVolumeSimCacheFrame* CacheFrame = &SimCacheData->Frames[FrameIndex];
+
+	const FRenderTargetVolumeRWInstanceData_GameThread* InstanceData_GT = reinterpret_cast<const FRenderTargetVolumeRWInstanceData_GameThread*>(OptionalPerInstanceData);
+	CacheFrame->Size = InstanceData_GT->Size;
+	CacheFrame->Format = EPixelFormat::PF_FloatRGBA;
+
+	if (InstanceData_GT->TargetTexture)
+	{
+		const FNiagaraDataInterfaceProxyRenderTargetVolumeProxy* RT_Proxy = GetProxyAs<FNiagaraDataInterfaceProxyRenderTargetVolumeProxy>();
+		const FTextureRenderTargetResource* RT_TargetTexture = InstanceData_GT->TargetTexture->GameThread_GetRenderTargetResource();
+
+		//-OPT: Currently we are flushing rendering commands.  Do not remove this until making access to the frame data safe across threads.
+		TArray<FFloat16Color> TextureData;
+		ENQUEUE_RENDER_COMMAND(NDIRenderTargetVolume_CacheFrame)
+		(
+			[RT_Proxy, RT_InstanceID=SystemInstance->GetId(), RT_TargetTexture, RT_TextureData=&TextureData](FRHICommandListImmediate& RHICmdList)
+			{
+				if (const FRenderTargetVolumeRWInstanceData_RenderThread* InstanceData_RT = RT_Proxy->SystemInstancesToProxyData_RT.Find(RT_InstanceID))
+				{
+					// Readback TextureData
+					RHICmdList.Read3DSurfaceFloatData(
+						InstanceData_RT->RenderTarget->GetRHI(),
+						FIntRect(0, 0, InstanceData_RT->Size.X, InstanceData_RT->Size.Y),
+						FIntPoint(0, InstanceData_RT->Size.Z),
+						*RT_TextureData
+					);
+				}
+			}
+		);
+		FlushRenderingCommands();
+
+		if (TextureData.Num() > 0)
+		{
+			const FName CompressionType = SimCacheData->CompressionType;
+			const bool bUseCompression = CompressionType.IsNone() == false;
+			const int TextureSizeBytes = TextureData.Num() * TextureData.GetTypeSize();
+			int CompressedSize = bUseCompression ? FCompression::CompressMemoryBound(CompressionType, TextureSizeBytes) : TextureSizeBytes;
+
+			CacheFrame->UncompressedSize = TextureSizeBytes;
+			CacheFrame->CompressedSize = 0;
+			uint8* FinalPixelData = reinterpret_cast<uint8*>(FMemory::Malloc(CompressedSize));
+			if (bUseCompression)
+			{
+				if (FCompression::CompressMemory(CompressionType, FinalPixelData, CompressedSize, TextureData.GetData(), TextureSizeBytes, COMPRESS_BiasMemory))
+				{
+					CacheFrame->CompressedSize = CompressedSize;
+				}
+			}
+
+			if (CacheFrame->CompressedSize == 0)
+			{
+				FMemory::Memcpy(FinalPixelData, TextureData.GetData(), TextureSizeBytes);
+			}
+
+			// Update bulk data
+			CacheFrame->BulkData.Lock(LOCK_READ_WRITE);
+			{
+				const int32 FinalByteSize = CacheFrame->CompressedSize > 0 ? CacheFrame->CompressedSize : CacheFrame->UncompressedSize;
+				uint8* BulkDataPtr = reinterpret_cast<uint8*>(CacheFrame->BulkData.Realloc(FinalByteSize));
+				FMemory::Memcpy(BulkDataPtr, FinalPixelData, FinalByteSize);
+			}
+			CacheFrame->BulkData.Unlock();
+
+			FMemory::Free(FinalPixelData);
+		}
+	}
+
+	return true;
+}
+
+bool UNiagaraDataInterfaceRenderTargetVolume::SimCacheEndWrite(UObject* StorageObject) const
+{
+	UNDIRenderTargetVolumeSimCacheData* SimCacheData = CastChecked<UNDIRenderTargetVolumeSimCacheData>(StorageObject);
+
+	return true;
+}
+
+bool UNiagaraDataInterfaceRenderTargetVolume::SimCacheReadFrame(UObject* StorageObject, int FrameA, int FrameB, float Interp, FNiagaraSystemInstance* SystemInstance, void* OptionalPerInstanceData)
+{
+	UNDIRenderTargetVolumeSimCacheData* SimCacheData = CastChecked<UNDIRenderTargetVolumeSimCacheData>(StorageObject);
+
+	const int FrameIndex = Interp >= 0.5f ? FrameB : FrameA;
+	if (!SimCacheData->Frames.IsValidIndex(FrameIndex))
+	{
+		return false;
+	}
+
+	FNDIRenderTargetVolumeSimCacheFrame* CacheFrame = &SimCacheData->Frames[FrameIndex];
+
+	FRenderTargetVolumeRWInstanceData_GameThread* InstanceData_GT = reinterpret_cast<FRenderTargetVolumeRWInstanceData_GameThread*>(OptionalPerInstanceData);
+	InstanceData_GT->Size = CacheFrame->Size;
+	InstanceData_GT->Format = CacheFrame->Format;
+
+	PerInstanceTick(InstanceData_GT, SystemInstance, 0.0f);
+	PerInstanceTickPostSimulate(InstanceData_GT, SystemInstance, 0.0f);
+
+	if (InstanceData_GT->TargetTexture && CacheFrame->GetPixelData() != nullptr)
+	{
+		FNiagaraDataInterfaceProxyRenderTargetVolumeProxy* RT_Proxy = GetProxyAs<FNiagaraDataInterfaceProxyRenderTargetVolumeProxy>();
+		FTextureRenderTargetResource* RT_TargetTexture = InstanceData_GT->TargetTexture->GameThread_GetRenderTargetResource();
+		ENQUEUE_RENDER_COMMAND(NDIRenderTargetVolumeUpdate)
+		(
+			[RT_Proxy, RT_InstanceID=SystemInstance->GetId(), RT_TargetTexture, RT_CacheFrame=CacheFrame, RT_CompressionType=SimCacheData->CompressionType](FRHICommandListImmediate& RHICmdList)
+			{
+				if (FRenderTargetVolumeRWInstanceData_RenderThread* InstanceData_RT = RT_Proxy->SystemInstancesToProxyData_RT.Find(RT_InstanceID))
+				{
+					FUpdateTextureRegion3D UpdateRegion(FIntVector::ZeroValue, FIntVector::ZeroValue, InstanceData_RT->Size);
+
+					if (RT_CacheFrame->CompressedSize > 0)
+					{
+						TArray<uint8> Decompressed;
+						Decompressed.AddUninitialized(sizeof(FFloat16Color) * InstanceData_RT->Size.X * InstanceData_RT->Size.Y * InstanceData_RT->Size.Z);
+
+						FCompression::UncompressMemory(RT_CompressionType, Decompressed.GetData(), Decompressed.Num(), RT_CacheFrame->GetPixelData(), RT_CacheFrame->CompressedSize);
+
+						RHICmdList.UpdateTexture3D(
+							InstanceData_RT->RenderTarget->GetRHI(),
+							0,
+							UpdateRegion,
+							InstanceData_RT->Size.X * sizeof(FFloat16Color),
+							InstanceData_RT->Size.X * InstanceData_RT->Size.Y * sizeof(FFloat16Color),
+							Decompressed.GetData()
+						);
+					}
+					else
+					{
+						RHICmdList.UpdateTexture3D(
+							InstanceData_RT->RenderTarget->GetRHI(),
+							0,
+							UpdateRegion,
+							InstanceData_RT->Size.X * sizeof(FFloat16Color),
+							InstanceData_RT->Size.X * InstanceData_RT->Size.Y * sizeof(FFloat16Color),
+							RT_CacheFrame->GetPixelData()
+						);
+					}
+				}
+			}
+		);
+	}
+
+	return true;
+}
+
 void UNiagaraDataInterfaceRenderTargetVolume::VMSetSize(FVectorVMExternalFunctionContext& Context)
 {
 	// This should only be called from a system or emitter script due to a need for only setting up initially.
