- Fix for float attribute display on debugger

#rnx

diff --git a/Engine/Plugins/FX/Niagara/Source/Niagara/Private/NiagaraDataSetDebugAccessor.cpp b/Engine/Plugins/FX/Niagara/Source/Niagara/Private/NiagaraDataSetDebugAccessor.cpp
--- a/Engine/Plugins/FX/Niagara/Source/Niagara/Private/NiagaraDataSetDebugAccessor.cpp
+++ b/Engine/Plugins/FX/Niagara/Source/Niagara/Private/NiagaraDataSetDebugAccessor.cpp
@@ -40,9 +40,9 @@ bool FNiagaraDataSetDebugAccessor::Init(const FNiagaraDataSetCompiledData& Compi
 
 float FNiagaraDataSetDebugAccessor::ReadFloat(const FNiagaraDataBuffer* DataBuffer, uint32 Instance, uint32 Component) const
 {
-	if (DataBuffer != nullptr && ComponentIndexInt32 != INDEX_NONE && Instance < DataBuffer->GetNumInstances() && Component < NumComponentsInt32)
+	if (DataBuffer != nullptr && ComponentIndexFloat != INDEX_NONE && Instance < DataBuffer->GetNumInstances() && Component < NumComponentsFloat)
 	{
-		const float* FloatData = reinterpret_cast<const float*>(DataBuffer->GetComponentPtrFloat(ComponentIndexInt32 + Component));
+		const float* FloatData = reinterpret_cast<const float*>(DataBuffer->GetComponentPtrFloat(ComponentIndexFloat + Component));
 		return FloatData[Instance];
 	}
 	return 0.0f;
