fixes for dump byte code commandlet
-systems that shared the same name (but different paths) would overwrite results

diff --git a/Engine/Plugins/FX/Niagara/Source/NiagaraEditor/Private/Commandlets/NiagaraDumpBytecodeCommandlet.cpp b/Engine/Plugins/FX/Niagara/Source/NiagaraEditor/Private/Commandlets/NiagaraDumpBytecodeCommandlet.cpp
--- a/Engine/Plugins/FX/Niagara/Source/NiagaraEditor/Private/Commandlets/NiagaraDumpBytecodeCommandlet.cpp
+++ b/Engine/Plugins/FX/Niagara/Source/NiagaraEditor/Private/Commandlets/NiagaraDumpBytecodeCommandlet.cpp
@@ -163,9 +163,12 @@ void UNiagaraDumpByteCodeCommandlet::ProcessNiagaraScripts()
 			UE_LOG(LogNiagaraDumpBytecodeCommandlet, Warning, TEXT("Loaded system was Invalid! %s"), *SystemName);
 		}
 
-		IFileManager::Get().MakeDirectory(*(AuditOutputFolder / ShorterSystemName));
-		DumpByteCode(NiagaraSystem->GetSystemSpawnScript(), ShorterSystemName / TEXT("SystemSpawnScript.bin"));
-		DumpByteCode(NiagaraSystem->GetSystemUpdateScript(), ShorterSystemName / TEXT("SystemUpdateScript.bin"));
+		const FString SystemPathName = NiagaraSystem->GetPathName();
+		const FString HashedPathName = FString::Printf(TEXT("%08x"), GetTypeHash(SystemPathName));
+
+		IFileManager::Get().MakeDirectory(*(AuditOutputFolder / HashedPathName));
+		DumpByteCode(NiagaraSystem->GetSystemSpawnScript(), SystemPathName, HashedPathName, TEXT("SystemSpawnScript.txt"));
+		DumpByteCode(NiagaraSystem->GetSystemUpdateScript(), SystemPathName, HashedPathName, TEXT("SystemUpdateScript.txt"));
 		
 		for (const auto& EmitterHandle : NiagaraSystem->GetEmitterHandles())
 		{
@@ -176,16 +179,19 @@ void UNiagaraDumpByteCodeCommandlet::ProcessNiagaraScripts()
 
 			if (FVersionedNiagaraEmitterData* Emitter = EmitterHandle.GetEmitterData())
 			{
-				const FString EmitterName = EmitterHandle.GetUniqueInstanceName();
+				if (Emitter->SimTarget == ENiagaraSimTarget::CPUSim)
+				{
+					const FString EmitterName = EmitterHandle.GetUniqueInstanceName();
 
-				TArray<UNiagaraScript*> EmitterScripts;
-				Emitter->GetScripts(EmitterScripts);
+					TArray<UNiagaraScript*> EmitterScripts;
+					Emitter->GetScripts(EmitterScripts);
 
-				IFileManager::Get().MakeDirectory(*(ShorterSystemName / EmitterName));
+					IFileManager::Get().MakeDirectory(*(HashedPathName / EmitterName));
 
-				for (const auto* EmitterScript : EmitterScripts)
-				{
-					DumpByteCode(EmitterScript, ShorterSystemName / EmitterName / UsageEnum->GetNameStringByValue(static_cast<int64>(EmitterScript->GetUsage())) + TEXT(".bin"));
+					for (const auto* EmitterScript : EmitterScripts)
+					{
+						DumpByteCode(EmitterScript, SystemPathName, HashedPathName, EmitterName / UsageEnum->GetNameStringByValue(static_cast<int64>(EmitterScript->GetUsage())) + TEXT(".txt"));
+					}
 				}
 			}
 		}
@@ -208,11 +214,12 @@ void UNiagaraDumpByteCodeCommandlet::ProcessNiagaraScripts()
 		for (const auto& MetaData : ScriptMetaData)
 		{
 			OutputStream->Log(TEXT("\t<Script>"));
+			OutputStream->Logf(TEXT("\t\t<Hash>%s</Hash>"), *MetaData.SystemHash);
 			OutputStream->Logf(TEXT("\t\t<Name>%s</Name>"), *MetaData.FullName);
-			OutputStream->Logf(TEXT("\t<OpCount>%d</OpCount>"), MetaData.OpCount);
-			OutputStream->Logf(TEXT("\t<RegisterCount>%d</RegisterCount>"), MetaData.RegisterCount);
-			OutputStream->Logf(TEXT("\t<ConstantCount>%d</ConstantCount>"), MetaData.ConstantCount);
-			OutputStream->Logf(TEXT("\t<AttributeCount>%d</AttributeCount>"), MetaData.AttributeCount);
+			OutputStream->Logf(TEXT("\t\t<OpCount>%d</OpCount>"), MetaData.OpCount);
+			OutputStream->Logf(TEXT("\t\t<RegisterCount>%d</RegisterCount>"), MetaData.RegisterCount);
+			OutputStream->Logf(TEXT("\t\t<ConstantCount>%d</ConstantCount>"), MetaData.ConstantCount);
+			OutputStream->Logf(TEXT("\t\t<AttributeCount>%d</AttributeCount>"), MetaData.AttributeCount);
 			OutputStream->Log(TEXT("\t</Script>"));
 		}
 		OutputStream->Log(TEXT("</Scripts>"));
@@ -224,10 +231,11 @@ void UNiagaraDumpByteCodeCommandlet::ProcessNiagaraScripts()
 		TUniquePtr<FArchive> FileArchive(IFileManager::Get().CreateDebugFileWriter(*MetaDataFileName));
 		TUniquePtr<FOutputDeviceArchiveWrapper> OutputStream(new FOutputDeviceArchiveWrapper(FileArchive.Get()));
 
-		OutputStream->Log(TEXT("Name, OpCount, RegisteredCount, ConstantCount, AttributeCount"));
+		OutputStream->Log(TEXT("Hash, Name, OpCount, RegisteredCount, ConstantCount, AttributeCount"));
 		for (const auto& MetaData : ScriptMetaData)
 		{
-			OutputStream->Logf(TEXT("%s, %d, %d, %d, %d"),
+			OutputStream->Logf(TEXT("%s, %s, %d, %d, %d, %d"),
+				*MetaData.SystemHash,
 				*MetaData.FullName,
 				MetaData.OpCount,
 				MetaData.RegisterCount,
@@ -244,7 +252,7 @@ void UNiagaraDumpByteCodeCommandlet::ProcessNiagaraScripts()
 	UE_LOG(LogNiagaraDumpBytecodeCommandlet, Log, TEXT("Took %5.3f seconds to process referenced Niagara systems..."), ProcessNiagaraSystemsTime);
 }
 
-void UNiagaraDumpByteCodeCommandlet::DumpByteCode(const UNiagaraScript* Script, const FString& FilePath)
+void UNiagaraDumpByteCodeCommandlet::DumpByteCode(const UNiagaraScript* Script, const FString& PathName, const FString& HashName, const FString& FilePath)
 {
 	if (!Script)
 	{
@@ -254,7 +262,8 @@ void UNiagaraDumpByteCodeCommandlet::DumpByteCode(const UNiagaraScript* Script,
 	const auto& ExecData = Script->GetVMExecutableData();
 
 	FScriptMetaData& MetaData = ScriptMetaData.AddZeroed_GetRef();
-	MetaData.FullName = FilePath;
+	MetaData.SystemHash = HashName;
+	MetaData.FullName = PathName / FilePath;
 	MetaData.RegisterCount = ExecData.NumTempRegisters;
 	MetaData.OpCount = ExecData.LastOpCount;
 	MetaData.ConstantCount = ExecData.InternalParameters.GetTableSize() / 4;
@@ -269,7 +278,7 @@ void UNiagaraDumpByteCodeCommandlet::DumpByteCode(const UNiagaraScript* Script,
 
 	if (Script)
 	{
-		const FString FullFilePath = AuditOutputFolder / FilePath;
+		const FString FullFilePath = AuditOutputFolder / HashName / FilePath;
 
 		TUniquePtr<FArchive> FileArchive(IFileManager::Get().CreateDebugFileWriter(*FullFilePath));
 		if (!FileArchive)
diff --git a/Engine/Plugins/FX/Niagara/Source/NiagaraEditor/Public/Commandlets/NiagaraDumpBytecodeCommandlet.h b/Engine/Plugins/FX/Niagara/Source/NiagaraEditor/Public/Commandlets/NiagaraDumpBytecodeCommandlet.h
--- a/Engine/Plugins/FX/Niagara/Source/NiagaraEditor/Public/Commandlets/NiagaraDumpBytecodeCommandlet.h
+++ b/Engine/Plugins/FX/Niagara/Source/NiagaraEditor/Public/Commandlets/NiagaraDumpBytecodeCommandlet.h
@@ -33,10 +33,11 @@ class UNiagaraDumpByteCodeCommandlet : public UCommandlet
 
 private:
 	void ProcessNiagaraScripts();
-	void DumpByteCode(const UNiagaraScript* Script, const FString& FilePath);
+	void DumpByteCode(const UNiagaraScript* Script, const FString& PathName, const FString& HashName, const FString& FilePath);
 
 	struct FScriptMetaData
 	{
+		FString SystemHash;
 		FString FullName;
 		int32 RegisterCount;
 		int32 OpCount;
