Couple of fixes for attribute trimming with niagara particle scripts
-parameters to an impure function that were using the StackContext namespace weren't being properly included in the dependency chain
-custom HLSL nodes which internally used impure DI functions wouldn't necessarily have their inputs marked as dependents (in the case where a dummy output variable was being used, which itself got trimmed)

diff --git a/Engine/Plugins/FX/Niagara/Source/NiagaraEditor/Private/NiagaraAttributeTrimmer.cpp b/Engine/Plugins/FX/Niagara/Source/NiagaraEditor/Private/NiagaraAttributeTrimmer.cpp
--- a/Engine/Plugins/FX/Niagara/Source/NiagaraEditor/Private/NiagaraAttributeTrimmer.cpp
+++ b/Engine/Plugins/FX/Niagara/Source/NiagaraEditor/Private/NiagaraAttributeTrimmer.cpp
@@ -4,6 +4,7 @@
 
 #include "Algo/ForEach.h"
 #include "Algo/RemoveIf.h"
+#include "EdGraphSchema_Niagara.h"
 #include "NiagaraConstants.h"
 #include "NiagaraModule.h"
 #include "NiagaraNodeCustomHlsl.h"
@@ -22,12 +23,18 @@ uint32 GetTypeHash(const FModuleScopedPin& ScopedPin)
 namespace NiagaraAttributeTrimming
 {
 	using FDependencySet = TSet<FModuleScopedPin>;
-	using FCustomHlslNodeSet = TSet<const UNiagaraNodeCustomHlsl*>;
+	struct FCustomHlslNodeInfo
+	{
+		bool bHasDataInterfaceInputs = false;
+		bool bHasImpureFunctionText = false;
+	};
+
+	using FCustomHlslNodeMap = TMap<const UNiagaraNodeCustomHlsl*, FCustomHlslNodeInfo>;
 
 	struct FDependencyChain
 	{
 		FDependencySet Pins;
-		FCustomHlslNodeSet CustomNodes;
+		FCustomHlslNodeMap CustomNodes;
 	};
 
 	using FDependencyMap = TMap<FModuleScopedPin, FDependencyChain>;
@@ -166,9 +173,10 @@ namespace NiagaraAttributeTrimming
 	{
 	public:
 
-		FExpressionBuilder(const FNiagaraParameterMapHistory& InParamMap, const FFunctionInputResolver& InInputResolver)
+		FExpressionBuilder(const FNiagaraParameterMapHistory& InParamMap, const FFunctionInputResolver& InInputResolver, const UEdGraphSchema_Niagara* InNiagaraSchema)
 		: ParamMap(InParamMap)
 		, InputResolver(InInputResolver)
+		, NiagaraSchema(InNiagaraSchema)
 		{}
 
 		// generates a set of dependencies for a specific pin
@@ -309,7 +317,11 @@ namespace NiagaraAttributeTrimming
 
 							if (const UNiagaraNodeCustomHlsl* CustomHlslNode = Cast<const UNiagaraNodeCustomHlsl>(CurrentNode))
 							{
-								Dependencies.CustomNodes.Add(CustomHlslNode);
+								if (!Dependencies.CustomNodes.Contains(CustomHlslNode))
+								{
+									FCustomHlslNodeInfo& NodeInfo = Dependencies.CustomNodes.Add(CustomHlslNode);
+									BuildCustomHlslNodeInfo(CustomHlslNode, NodeInfo);
+								}
 							}
 						}
 					}
@@ -326,8 +338,68 @@ namespace NiagaraAttributeTrimming
 			}
 		}
 
+		void BuildCustomHlslNodeInfo(const UNiagaraNodeCustomHlsl* CustomNode, FCustomHlslNodeInfo& NodeInfo)
+		{
+			check(CustomNode);
+
+			NodeInfo.bHasDataInterfaceInputs = false;
+			NodeInfo.bHasImpureFunctionText = false;
+
+			TArray<FString> ImpureFunctionNames;
+
+			FPinCollectorArray InputPins;
+			CustomNode->GetInputPins(InputPins);
+			for (const UEdGraphPin* InputPin : InputPins)
+			{
+				FNiagaraTypeDefinition NiagaraType = NiagaraSchema->PinToTypeDefinition(InputPin);
+				if (NiagaraType.IsDataInterface())
+				{
+					NodeInfo.bHasDataInterfaceInputs = true;
+
+					if (UNiagaraDataInterface* DataInterfaceClass = CastChecked<UNiagaraDataInterface>(NiagaraType.GetClass()->ClassDefaultObject))
+					{
+						TArray<FNiagaraFunctionSignature> FunctionSignatures;
+						DataInterfaceClass->GetFunctions(FunctionSignatures);
+
+						for (const FNiagaraFunctionSignature& FunctionSignature : FunctionSignatures)
+						{
+							if (FunctionSignature.bRequiresExecPin)
+							{
+								TStringBuilder<256> Builder;
+								InputPin->GetFName().AppendString(Builder);
+								Builder.AppendChar(TCHAR('.'));
+								FunctionSignature.Name.AppendString(Builder);
+
+								ImpureFunctionNames.AddUnique(Builder.ToString());
+							}
+						}
+					}
+				}
+			}
+
+			if (!ImpureFunctionNames.IsEmpty())
+			{
+				TArray<FString> StringTokens;
+				UNiagaraNodeCustomHlsl::GetTokensFromString(CustomNode->GetCustomHlsl(), StringTokens, false, false);
+
+				for (const FString& FunctionCall : ImpureFunctionNames)
+				{
+					for (const FString& Token : StringTokens)
+					{
+						if (Token.Equals(FunctionCall))
+						{
+							NodeInfo.bHasImpureFunctionText = true;
+							return;
+						}
+					}
+				}
+			}
+		}
+
 		const FNiagaraParameterMapHistory& ParamMap;
 		const FFunctionInputResolver& InputResolver;
+		const UEdGraphSchema_Niagara* NiagaraSchema = nullptr;
+
 		const bool EvaluateStaticSwitches = false;
 	};
 
@@ -541,7 +613,7 @@ static void TrimAttributes_Aggressive(const FNiagaraCompileRequestDuplicateData*
 
 	FDependencyMap PerVariableDependencySets;
 
-	TMap<FName, TArray<FDependencyChain>> DependencySets;
+	const UEdGraphSchema_Niagara* NiagaraSchema = GetDefault<UEdGraphSchema_Niagara>();
 
 	for (const FNiagaraParameterMapHistory* ParamMap : LocalParamHistories)
 	{
@@ -557,10 +629,16 @@ static void TrimAttributes_Aggressive(const FNiagaraCompileRequestDuplicateData*
 			{
 				FDependencyChain& Dependencies = PerVariableDependencySets.Add(WritePin);
 
-				FExpressionBuilder Builder(*ParamMap, InputResolver);
+				FExpressionBuilder Builder(*ParamMap, InputResolver, NiagaraSchema);
 				Builder.FindDependencies(WritePin, Dependencies);
 
-				DependencySets.FindOrAdd(Var.GetName()).Add(Dependencies);
+				for (FCustomHlslNodeMap::TConstIterator CustomNodeIt = Dependencies.CustomNodes.CreateConstIterator(); CustomNodeIt; ++CustomNodeIt)
+				{
+					if (CustomNodeIt.Value().bHasImpureFunctionText)
+					{
+						AttributesToPreserve.Add(Var.GetName());
+					}
+				}
 			}
 		}
 
@@ -569,7 +647,7 @@ static void TrimAttributes_Aggressive(const FNiagaraCompileRequestDuplicateData*
 		{
 			FDependencyChain Dependencies;
 
-			FExpressionBuilder Builder(*ParamMap, InputResolver);
+			FExpressionBuilder Builder(*ParamMap, InputResolver, NiagaraSchema);
 			Builder.FindDependencies(RequiredFunctionInput, Dependencies);
 
 			for (const FModuleScopedPin& DependentPin : Dependencies.Pins)
@@ -583,6 +661,19 @@ static void TrimAttributes_Aggressive(const FNiagaraCompileRequestDuplicateData*
 
 				const FName VariableName = *VariableNameString;
 
+				if (VariableNameString.StartsWith(FNiagaraConstants::StackContextNamespaceString))
+				{
+					const int32 AliasedIndex = ParamMap->VariablesWithOriginalAliasesIntact.IndexOfByPredicate([&](const FNiagaraVariable& Variable)
+					{
+						return Variable.GetName() == VariableName;
+					});
+
+					if (ParamMap->Variables.IsValidIndex(AliasedIndex))
+					{
+						AttributesToPreserve.Add(ParamMap->Variables[AliasedIndex].GetName());
+					}
+				}
+				else
 				if (ParamMap->FindVariableByName(VariableName) != INDEX_NONE)
 				{
 					AttributesToPreserve.Add(VariableName);
@@ -640,7 +731,7 @@ static void TrimAttributes_Aggressive(const FNiagaraCompileRequestDuplicateData*
 				}
 			}
 
-			for (const UNiagaraNodeCustomHlsl* CustomHlsl : ResolvedChain.CustomNodes)
+			for (const auto& CustomHlslInfo : ResolvedChain.CustomNodes)
 			{
 				// go through all of the variables in the attributes to see if they are referenced by any encountered CustomNodes
 				for (const FNiagaraVariableBase& Variable : UnifiedVariables)
@@ -651,7 +742,7 @@ static void TrimAttributes_Aggressive(const FNiagaraCompileRequestDuplicateData*
 						continue;
 					}
 
-					if (CustomHlsl->ReferencesVariable(Variable))
+					if (CustomHlslInfo.Key->ReferencesVariable(Variable))
 					{
 						ConditionalAddAttribute(Variable.GetName());
 					}
diff --git a/Engine/Plugins/FX/Niagara/Source/NiagaraEditor/Private/NiagaraHlslTranslator.cpp b/Engine/Plugins/FX/Niagara/Source/NiagaraEditor/Private/NiagaraHlslTranslator.cpp
--- a/Engine/Plugins/FX/Niagara/Source/NiagaraEditor/Private/NiagaraHlslTranslator.cpp
+++ b/Engine/Plugins/FX/Niagara/Source/NiagaraEditor/Private/NiagaraHlslTranslator.cpp
@@ -7444,10 +7444,10 @@ void FHlslNiagaraTranslator::ProcessCustomHlsl(const FString& InCustomHlsl, ENia
 		FNiagaraVariable Input = OutSignature.Inputs[i];
 		if (Input.GetType() == FNiagaraTypeDefinition::GetParameterMapDef())
 		{
-			FString ParameterMapInstanceName = GetParameterMapInstanceName(0);
-			FString ReplaceSrc = Input.GetName().ToString();
-			FString ReplaceDest = ParameterMapInstanceName;
+			FNameBuilder ReplaceSrc(Input.GetName());
+			FString ReplaceDest = GetParameterMapInstanceName(0);
 			UNiagaraNodeCustomHlsl::ReplaceExactMatchTokens(Tokens, ReplaceSrc, ReplaceDest, true);
+
 			SigInputs.Add(Input);
 			OutSignature.bRequiresContext = true;
 			ParamMapHistoryIdx = Inputs[i];
@@ -7587,8 +7587,11 @@ void FHlslNiagaraTranslator::ProcessCustomHlsl(const FString& InCustomHlsl, ENia
 		}
 		else
 		{
-			FString ReplaceSrc = Input.GetName().ToString();
-			FString ReplaceDest = TEXT("In_") + ReplaceSrc;
+			FNameBuilder ReplaceSrc(Input.GetName());
+			TStringBuilder<128> ReplaceDest;
+			ReplaceDest.Append(TEXT("In_"));
+			ReplaceDest.Append(ReplaceSrc);
+
 			UNiagaraNodeCustomHlsl::ReplaceExactMatchTokens(Tokens, ReplaceSrc, ReplaceDest, true);
 			SigInputs.Add(Input);
 		}
@@ -7601,9 +7604,9 @@ void FHlslNiagaraTranslator::ProcessCustomHlsl(const FString& InCustomHlsl, ENia
 	{
 		if (Output.GetType() == FNiagaraTypeDefinition::GetParameterMapDef())
 		{
-			FString ParameterMapInstanceName = GetParameterMapInstanceName(0);
-			FString ReplaceSrc = Output.GetName().ToString();
-			FString ReplaceDest = ParameterMapInstanceName;
+			FNameBuilder ReplaceSrc(Output.GetName());
+			FString ReplaceDest = GetParameterMapInstanceName(0);
+
 			UNiagaraNodeCustomHlsl::ReplaceExactMatchTokens(Tokens, ReplaceSrc, ReplaceDest, true);
 			SigOutputs.Add(Output);
 			OutSignature.bRequiresContext = true;
@@ -7611,8 +7614,11 @@ void FHlslNiagaraTranslator::ProcessCustomHlsl(const FString& InCustomHlsl, ENia
 		}
 		else
 		{
-			FString ReplaceSrc = Output.GetName().ToString();
-			FString ReplaceDest = TEXT("Out_") + ReplaceSrc;
+			FNameBuilder ReplaceSrc(Output.GetName());
+			TStringBuilder<128> ReplaceDest;
+			ReplaceDest.Append(TEXT("Out_"));
+			ReplaceDest.Append(ReplaceSrc);
+
 			UNiagaraNodeCustomHlsl::ReplaceExactMatchTokens(Tokens, ReplaceSrc, ReplaceDest, true);
 			SigOutputs.Add(Output);
 		}
diff --git a/Engine/Plugins/FX/Niagara/Source/NiagaraEditor/Private/NiagaraNodeCustomHlsl.cpp b/Engine/Plugins/FX/Niagara/Source/NiagaraEditor/Private/NiagaraNodeCustomHlsl.cpp
--- a/Engine/Plugins/FX/Niagara/Source/NiagaraEditor/Private/NiagaraNodeCustomHlsl.cpp
+++ b/Engine/Plugins/FX/Niagara/Source/NiagaraEditor/Private/NiagaraNodeCustomHlsl.cpp
@@ -87,27 +87,48 @@ FLinearColor UNiagaraNodeCustomHlsl::GetNodeTitleColor() const
 	return UEdGraphSchema_Niagara::NodeTitleColor_CustomHlsl;
 }
 
-bool UNiagaraNodeCustomHlsl::GetTokensFromString(const FString& InHlsl, TArray<FString>& OutTokens, bool IncludeComments)
+bool UNiagaraNodeCustomHlsl::GetTokensFromString(const FString& InHlsl, TArray<FString>& OutTokens, bool IncludeComments, bool IncludeWhitespace)
 {
 	if (InHlsl.Len() == 0)
 	{
 		return false;
 	}
 
-	FString Separators = TEXT(";/*+-=)(?:, []<>\"\t\n");
+	FString Separators = TEXT(";/*+-=)(?:, []<>\"\t\r\n");
+	const int32 TargetLength = InHlsl.Len();
+
 	int32 TokenStart = 0;
-	int32 TargetLength = InHlsl.Len();
+	bool TokenIsWhitespace = true;
+
+	auto AddToken = [&](bool DoAdd, FString&& TokenString)
+	{
+		if (DoAdd && (!TokenIsWhitespace || IncludeWhitespace))
+		{
+			OutTokens.Add(MoveTemp(TokenString));
+		}
+		
+		// reset the meta data about the token
+		TokenIsWhitespace = true;
+	};
+
 	for (int32 i = 0; i < TargetLength; )
 	{
 		int32 Index = INDEX_NONE;
 
+		const bool bWhitespace = FChar::IsWhitespace(InHlsl[i]);
+
 		// Determine if we are a splitter character or a regular character.
 		if (Separators.FindChar(InHlsl[i], Index) && Index != INDEX_NONE)
 		{
 			// Commit the current token, if any.
 			if (i > TokenStart)
 			{
-				OutTokens.Add(InHlsl.Mid(TokenStart, i - TokenStart));
+				AddToken(true, InHlsl.Mid(TokenStart, i - TokenStart));
+			}
+
+			if (!bWhitespace)
+			{
+				TokenIsWhitespace = false;
 			}
 
 			if (InHlsl[i] == '/' && (i + 1 != TargetLength) && InHlsl[i + 1] == '/')
@@ -119,10 +140,7 @@ bool UNiagaraNodeCustomHlsl::GetTokensFromString(const FString& InHlsl, TArray<F
 					FoundEndIdx = TargetLength - 1;
 				}
 
-				if (IncludeComments)
-				{
-					OutTokens.Add(InHlsl.Mid(i, FoundEndIdx - i + 1));
-				}
+				AddToken(IncludeComments, InHlsl.Mid(i, FoundEndIdx - i + 1));
 				i = FoundEndIdx + 1;
 			}
 			else if (InHlsl[i] == '/' && (i + 1 != TargetLength) && InHlsl[i + 1] == '*')
@@ -140,10 +158,7 @@ bool UNiagaraNodeCustomHlsl::GetTokensFromString(const FString& InHlsl, TArray<F
 					FoundEndIdx = TargetLength - 1;
 				}
 
-				if (IncludeComments)
-				{
-					OutTokens.Add(InHlsl.Mid(i, FoundEndIdx - i + 1));
-				}
+				AddToken(IncludeComments, InHlsl.Mid(i, FoundEndIdx - i + 1));
 				i = FoundEndIdx + 1;
 			}
 			else if (InHlsl[i] == '"')
@@ -159,14 +174,13 @@ bool UNiagaraNodeCustomHlsl::GetTokensFromString(const FString& InHlsl, TArray<F
 					FoundEndIdx = TargetLength - 1;
 				}
 
-				OutTokens.Add(InHlsl.Mid(i, FoundEndIdx - i + 1));
+				AddToken(true, InHlsl.Mid(i, FoundEndIdx - i + 1));
 				i = FoundEndIdx + 1;
 
 			}
 			else
 			{
-				// The separator will be its own token.
-				OutTokens.Add(FString(1, &InHlsl[i]));
+				AddToken(true, FString(1, &InHlsl[i]));
 				i++;
 			}
 
@@ -175,6 +189,11 @@ bool UNiagaraNodeCustomHlsl::GetTokensFromString(const FString& InHlsl, TArray<F
 		}
 		else
 		{
+			if (!bWhitespace)
+			{
+				TokenIsWhitespace = false;
+			}
+
 			// This character is part of a token, continue scanning.
 			i++;
 		}
@@ -183,15 +202,14 @@ bool UNiagaraNodeCustomHlsl::GetTokensFromString(const FString& InHlsl, TArray<F
 	// We may need to pull in the last chars from the end.
 	if (TokenStart < TargetLength)
 	{
-		OutTokens.Add(InHlsl.Mid(TokenStart));
+		AddToken(true, InHlsl.Mid(TokenStart));
 	}
 	return true;
 }
 
-bool UNiagaraNodeCustomHlsl::GetTokens(TArray<FString>& OutTokens, bool IncludeComments) const
+bool UNiagaraNodeCustomHlsl::GetTokens(TArray<FString>& OutTokens, bool IncludeComments, bool IncludeWhitespace) const
 {
-	FString HlslData = *CustomHlsl;
-	return GetTokensFromString(HlslData, OutTokens, IncludeComments);
+	return GetTokensFromString(CustomHlsl, OutTokens, IncludeComments, IncludeWhitespace);
 }
 
 void UNiagaraNodeCustomHlsl::PostEditChangeProperty(FPropertyChangedEvent& PropertyChangedEvent)
@@ -362,7 +380,7 @@ void UNiagaraNodeCustomHlsl::BuildParameterMapHistory(FNiagaraParameterMapHistor
 	}
 
 	TArray<FString> Tokens;
-	GetTokens(Tokens);
+	GetTokens(Tokens, false, false);
 
 	FPinCollectorArray InputPins;
 	GetInputPins(InputPins);
@@ -386,8 +404,6 @@ void UNiagaraNodeCustomHlsl::BuildParameterMapHistory(FNiagaraParameterMapHistor
 			if (Input.GetType() == FNiagaraTypeDefinition::GetParameterMapDef())
 			{
 				bHasParamMapInput = true;
-				FString ReplaceSrc = Input.GetName().ToString() + TEXT(".");
-				ReplaceExactMatchTokens(Tokens, ReplaceSrc, TEXT(""), false);
 				if (InputPins[i]->LinkedTo.Num() != 0)
 				{
 					ParamMapIdx = OutHistory.TraceParameterMapOutputPin(InputPins[i]->LinkedTo[0]);
@@ -408,8 +424,6 @@ void UNiagaraNodeCustomHlsl::BuildParameterMapHistory(FNiagaraParameterMapHistor
 			if (Output.GetType() == FNiagaraTypeDefinition::GetParameterMapDef())
 			{
 				bHasParamMapOutput = true;
-				FString ReplaceSrc = Output.GetName().ToString() + TEXT(".");
-				ReplaceExactMatchTokens(Tokens, ReplaceSrc, TEXT(""), false);
 				OutHistory.RegisterParameterMapPin(ParamMapIdx, OutputPins[i]);
 			}
 			else
@@ -493,22 +507,29 @@ void UNiagaraNodeCustomHlsl::GatherExternalDependencyData(ENiagaraScriptUsage In
 }
 
 // Replace items in the tokens array if they start with the src string or optionally src string and a namespace delimiter
-uint32 UNiagaraNodeCustomHlsl::ReplaceExactMatchTokens(TArray<FString>& Tokens, const FString& SrcString, const FString& ReplaceString, bool bAllowNamespaceSeparation)
+uint32 UNiagaraNodeCustomHlsl::ReplaceExactMatchTokens(TArray<FString>& Tokens, FStringView SrcString, FStringView ReplaceString, bool bAllowNamespaceSeparation)
 {
+	const int32 SrcLength = SrcString.Len();
+
 	uint32 Count = 0;
 	for (int32 i = 0; i < Tokens.Num(); i++)
 	{
-		if (Tokens[i].StartsWith(SrcString, ESearchCase::CaseSensitive))
+		if (FStringView(Tokens[i]).StartsWith(SrcString, ESearchCase::CaseSensitive))
 		{
-			if (Tokens[i].Len() > SrcString.Len() && Tokens[i][SrcString.Len()] == '.' && bAllowNamespaceSeparation)
+			const int32 TokenLength = Tokens[i].Len();
+
+			if (TokenLength > SrcLength)
 			{
-				Tokens[i] = ReplaceString + Tokens[i].Mid(SrcString.Len());
-				Count++;
+				if (bAllowNamespaceSeparation && Tokens[i][SrcLength] == TCHAR('.'))
+				{
+					Tokens[i] = ReplaceString + Tokens[i].Mid(SrcLength);
+					++Count;
+				}
 			}
-			else if (Tokens[i].Len() == SrcString.Len())
+			else
 			{
 				Tokens[i] = ReplaceString;
-				Count++;
+				++Count;
 			}
 		}
 	}
@@ -540,7 +561,7 @@ bool UNiagaraNodeCustomHlsl::ReferencesVariable(const FNiagaraVariableBase& InVa
 	// todo - all variable references in custom code should be explicit and typed
 	TArray<FString> Tokens;
 
-	GetTokens(Tokens, false);
+	GetTokens(Tokens, false, false);
 
 	const FString VariableName = InVar.GetName().ToString();
 
diff --git a/Engine/Plugins/FX/Niagara/Source/NiagaraEditor/Public/NiagaraNodeCustomHlsl.h b/Engine/Plugins/FX/Niagara/Source/NiagaraEditor/Public/NiagaraNodeCustomHlsl.h
--- a/Engine/Plugins/FX/Niagara/Source/NiagaraEditor/Public/NiagaraNodeCustomHlsl.h
+++ b/Engine/Plugins/FX/Niagara/Source/NiagaraEditor/Public/NiagaraNodeCustomHlsl.h
@@ -31,19 +31,20 @@ public:
 	FText GetHlslText() const;
 	void OnCustomHlslTextCommitted(const FText& InText, ETextCommit::Type InType);
 
-	bool GetTokens(TArray<FString>& OutTokens, bool IncludeComments = true) const;
+	bool GetTokens(TArray<FString>& OutTokens, bool IncludeComments = true, bool IncludeWhitespace = true) const;
 
 	virtual void BuildParameterMapHistory(FNiagaraParameterMapHistoryBuilder& OutHistory, bool bRecursive = true, bool bFilterForCompilation = true) const override;
 	virtual void GatherExternalDependencyData(ENiagaraScriptUsage InMasterUsage, const FGuid& InMasterUsageId, TArray<FNiagaraCompileHash>& InReferencedCompileHashes, TArray<FString>& InReferencedObjs) const override;
 
 	// Replace items in the tokens array if they start with the src string or optionally src string and a namespace delimiter
-	static uint32 ReplaceExactMatchTokens(TArray<FString>& Tokens, const FString& SrcString, const FString& ReplaceString, bool bAllowNamespaceSeparation);
+	static uint32 ReplaceExactMatchTokens(TArray<FString>& Tokens, FStringView SrcString, FStringView ReplaceString, bool bAllowNamespaceSeparation);
+
 	virtual bool AllowNiagaraTypeForAddPin(const FNiagaraTypeDefinition& InType) const override;
 	virtual bool AllowNiagaraTypeForAddPin(const FNiagaraTypeDefinition& InType, EEdGraphPinDirection InDirection) const override;
 
 	virtual bool ReferencesVariable(const FNiagaraVariableBase& InVar) const;
 
-	static bool GetTokensFromString(const FString& InHlsl, TArray<FString>& OutTokens, bool IncludeComments = true);
+	static bool GetTokensFromString(const FString& InHlsl, TArray<FString>& OutTokens, bool IncludeComments = true, bool IncludeWhitespace = true);
 
 	virtual void PostEditChangeProperty(struct FPropertyChangedEvent& PropertyChangedEvent) override;
 
diff --git a/Engine/Source/Runtime/Core/Private/UObject/FortniteMainBranchObjectVersion.cpp b/Engine/Source/Runtime/Core/Private/UObject/FortniteMainBranchObjectVersion.cpp
--- a/Engine/Source/Runtime/Core/Private/UObject/FortniteMainBranchObjectVersion.cpp
+++ b/Engine/Source/Runtime/Core/Private/UObject/FortniteMainBranchObjectVersion.cpp
@@ -11,7 +11,7 @@ TMap<FGuid, FGuid> FFortniteMainBranchObjectVersion::GetSystemGuids()
 	SystemGuids.Add(DevGuids.MATERIALSHADERMAP_DERIVEDDATA_VER, FGuid("1DA85DA5733C4A489A1BE117CDC5ACCC"));
 	SystemGuids.Add(DevGuids.NANITE_DERIVEDDATA_VER, FGuid("4C051B8C8FC74BA3943CD60A62D1CB73"));
 	SystemGuids.Add(DevGuids.NIAGARASHADERMAP_DERIVEDDATA_VER, FGuid("2289D5116CF94BC0AFEAEC541468E645"));
-	SystemGuids.Add(DevGuids.Niagara_LatestScriptCompileVersion, FGuid("23DE02026DAF4B60B289CAB316D28680"));
+	SystemGuids.Add(DevGuids.Niagara_LatestScriptCompileVersion, FGuid("FDA3A752D5534180B6A6C27166FDFB7E"));
 	SystemGuids.Add(DevGuids.SkeletalMeshDerivedDataVersion, FGuid("ACF593EAAE354FCBB34CF44F5AA1BFD2"));
 	SystemGuids.Add(DevGuids.STATICMESH_DERIVEDDATA_VER, FGuid("4696D39D9BC743AFBFA1DAB359AD7A7D"));
 
