FNames containing asset paths are deprecated. FSoftObjectPath or FTopLevelAssetPath should be used instead.
Fixups for Niagara.

#release 5.1

diff --git a/Engine/Plugins/FX/Niagara/Source/Niagara/Private/NiagaraModule.cpp b/Engine/Plugins/FX/Niagara/Source/Niagara/Private/NiagaraModule.cpp
--- a/Engine/Plugins/FX/Niagara/Source/Niagara/Private/NiagaraModule.cpp
+++ b/Engine/Plugins/FX/Niagara/Source/Niagara/Private/NiagaraModule.cpp
@@ -1049,7 +1049,7 @@ void FNiagaraTypeDefinition::RecreateUserDefinedTypeRegistry()
 	
 	if (!IsRunningCommandlet())
 	{
-		TArray<FString> DenyList;
+		TArray<FSoftObjectPath> DenyList;
 		DenyList.Emplace(TEXT("/Niagara/Enums/ENiagaraCoordinateSpace.ENiagaraCoordinateSpace"));
 		DenyList.Emplace(TEXT("/Niagara/Enums/ENiagaraOrientationAxis.ENiagaraOrientationAxis"));
 		
@@ -1068,9 +1068,9 @@ void FNiagaraTypeDefinition::RecreateUserDefinedTypeRegistry()
 
 		for (FSoftObjectPath AssetRef : TotalStructAssets)
 		{
-			FName AssetRefPathNamePreResolve = AssetRef.GetAssetPathName();
+			FSoftObjectPath AssetRefPathNamePreResolve = AssetRef;
 
-			if (DenyList.Contains(AssetRefPathNamePreResolve.ToString()))
+			if (DenyList.Contains(AssetRefPathNamePreResolve))
 			{
 				continue;
 			}
@@ -1117,10 +1117,10 @@ void FNiagaraTypeDefinition::RecreateUserDefinedTypeRegistry()
 		
 		for (FSoftObjectPath AssetRef : Settings->AdditionalParameterEnums)
 		{
-			FName AssetRefPathNamePreResolve = AssetRef.GetAssetPathName();
+			FSoftObjectPath AssetRefPathNamePreResolve = AssetRef;
 			UObject* Obj = AssetRef.ResolveObject();
 
-			if (DenyList.Contains(AssetRefPathNamePreResolve.ToString()))
+			if (DenyList.Contains(AssetRefPathNamePreResolve))
 			{
 				continue;
 			}
diff --git a/Engine/Plugins/FX/Niagara/Source/NiagaraEditor/Private/Commandlets/NiagaraDumpBytecodeCommandlet.cpp b/Engine/Plugins/FX/Niagara/Source/NiagaraEditor/Private/Commandlets/NiagaraDumpBytecodeCommandlet.cpp
--- a/Engine/Plugins/FX/Niagara/Source/NiagaraEditor/Private/Commandlets/NiagaraDumpBytecodeCommandlet.cpp
+++ b/Engine/Plugins/FX/Niagara/Source/NiagaraEditor/Private/Commandlets/NiagaraDumpBytecodeCommandlet.cpp
@@ -92,7 +92,9 @@ void UNiagaraDumpByteCodeCommandlet::ProcessNiagaraScripts()
 	if (!FilterCollection.IsEmpty())
 	{
 		FCollectionManagerModule& CollectionManagerModule = FCollectionManagerModule::GetModule();
+PRAGMA_DISABLE_DEPRECATION_WARNINGS
 		CollectionManagerModule.Get().GetObjectsInCollection(FName(*FilterCollection), ECollectionShareType::CST_All, Filter.ObjectPaths, ECollectionRecursionFlags::SelfAndChildren);
+PRAGMA_DISABLE_DEPRECATION_WARNINGS
 	}
 
 	TArray<FAssetData> AssetList;
@@ -108,7 +110,7 @@ void UNiagaraDumpByteCodeCommandlet::ProcessNiagaraScripts()
 	UPackage* CurrentPackage = nullptr;
 	for (const FAssetData& AssetIt : AssetList)
 	{
-		const FString SystemName = AssetIt.ObjectPath.ToString();
+		const FString SystemName = AssetIt.GetObjectPathString();
 		const FString PackageName = AssetIt.PackageName.ToString();
 
 		if (PackageName.StartsWith(DevelopersFolder))
diff --git a/Engine/Plugins/FX/Niagara/Source/NiagaraEditor/Private/Commandlets/NiagaraDumpModuleInfoCommandlet.cpp b/Engine/Plugins/FX/Niagara/Source/NiagaraEditor/Private/Commandlets/NiagaraDumpModuleInfoCommandlet.cpp
--- a/Engine/Plugins/FX/Niagara/Source/NiagaraEditor/Private/Commandlets/NiagaraDumpModuleInfoCommandlet.cpp
+++ b/Engine/Plugins/FX/Niagara/Source/NiagaraEditor/Private/Commandlets/NiagaraDumpModuleInfoCommandlet.cpp
@@ -85,7 +85,9 @@ void UNiagaraDumpModuleInfoCommandlet::ProcessNiagaraParameterDefinitions()
 	if (!FilterCollection.IsEmpty())
 	{
 		FCollectionManagerModule& CollectionManagerModule = FCollectionManagerModule::GetModule();
+PRAGMA_DISABLE_DEPRECATION_WARNINGS
 		CollectionManagerModule.Get().GetObjectsInCollection(FName(*FilterCollection), ECollectionShareType::CST_All, Filter.ObjectPaths, ECollectionRecursionFlags::SelfAndChildren);
+PRAGMA_ENABLE_DEPRECATION_WARNINGS
 	}
 
 	TArray<FAssetData> AssetList;
@@ -97,7 +99,7 @@ void UNiagaraDumpModuleInfoCommandlet::ProcessNiagaraParameterDefinitions()
 	UPackage* CurrentPackage = nullptr;
 	for (const FAssetData& AssetIt : AssetList)
 	{
-		const FString SystemName = AssetIt.ObjectPath.ToString();
+		const FString SystemName = AssetIt.GetObjectPathString();
 		const FString PackageName = AssetIt.PackageName.ToString();
 
 		if (PackageName.StartsWith(DevelopersFolder))
@@ -187,7 +189,9 @@ void UNiagaraDumpModuleInfoCommandlet::ProcessNiagaraScripts()
 	if (!FilterCollection.IsEmpty())
 	{
 		FCollectionManagerModule& CollectionManagerModule = FCollectionManagerModule::GetModule();
+PRAGMA_DISABLE_DEPRECATION_WARNINGS
 		CollectionManagerModule.Get().GetObjectsInCollection(FName(*FilterCollection), ECollectionShareType::CST_All, Filter.ObjectPaths, ECollectionRecursionFlags::SelfAndChildren);
+PRAGMA_ENABLE_DEPRECATION_WARNINGS
 	}
 
 	TArray<FAssetData> AssetList;
@@ -203,7 +207,7 @@ void UNiagaraDumpModuleInfoCommandlet::ProcessNiagaraScripts()
 	UPackage* CurrentPackage = nullptr;
 	for (const FAssetData& AssetIt : AssetList)
 	{
-		const FString SystemName = AssetIt.ObjectPath.ToString();
+		const FString SystemName = AssetIt.GetObjectPathString();
 		const FString PackageName = AssetIt.PackageName.ToString();
 
 		if (PackageName.StartsWith(DevelopersFolder))
diff --git a/Engine/Plugins/FX/Niagara/Source/NiagaraEditor/Private/Commandlets/NiagaraSystemAuditCommandlet.cpp b/Engine/Plugins/FX/Niagara/Source/NiagaraEditor/Private/Commandlets/NiagaraSystemAuditCommandlet.cpp
--- a/Engine/Plugins/FX/Niagara/Source/NiagaraEditor/Private/Commandlets/NiagaraSystemAuditCommandlet.cpp
+++ b/Engine/Plugins/FX/Niagara/Source/NiagaraEditor/Private/Commandlets/NiagaraSystemAuditCommandlet.cpp
@@ -126,7 +126,9 @@ bool UNiagaraSystemAuditCommandlet::ProcessNiagaraSystems()
 	if (!FilterCollection.IsEmpty())
 	{
 		FCollectionManagerModule& CollectionManagerModule = FCollectionManagerModule::GetModule();
+PRAGMA_DISABLE_DEPRECATION_WARNINGS
 		CollectionManagerModule.Get().GetObjectsInCollection(FName(*FilterCollection), ECollectionShareType::CST_All, Filter.ObjectPaths, ECollectionRecursionFlags::SelfAndChildren);
+PRAGMA_ENABLE_DEPRECATION_WARNINGS
 	}
 
 	TArray<FAssetData> AssetList;
@@ -145,7 +147,7 @@ bool UNiagaraSystemAuditCommandlet::ProcessNiagaraSystems()
 
 	for (const FAssetData& AssetIt : AssetList)
 	{
-		const FString SystemName = AssetIt.ObjectPath.ToString();
+		const FString SystemName = AssetIt.GetObjectPathString();
 		const FString PackageName = AssetIt.PackageName.ToString();
 
 		if (PackageName.StartsWith(DevelopersFolder))
diff --git a/Engine/Plugins/FX/Niagara/Source/NiagaraEditor/Private/EdGraphSchema_Niagara.cpp b/Engine/Plugins/FX/Niagara/Source/NiagaraEditor/Private/EdGraphSchema_Niagara.cpp
--- a/Engine/Plugins/FX/Niagara/Source/NiagaraEditor/Private/EdGraphSchema_Niagara.cpp
+++ b/Engine/Plugins/FX/Niagara/Source/NiagaraEditor/Private/EdGraphSchema_Niagara.cpp
@@ -498,14 +498,16 @@ TArray<TSharedPtr<FNiagaraAction_NewNode>> UEdGraphSchema_Niagara::GetGraphActio
 
 		bool bIsInLibrary = FNiagaraEditorUtilities::IsScriptAssetInLibrary(ScriptAsset);
 		const FText DisplayName = FNiagaraEditorUtilities::FormatScriptName(ScriptAsset.AssetName, bIsInLibrary);
-		const FText TooltipDesc = FNiagaraEditorUtilities::FormatScriptDescription(AssetDesc, ScriptAsset.ObjectPath, bIsInLibrary);
+		const FText TooltipDesc = FNiagaraEditorUtilities::FormatScriptDescription(AssetDesc, ScriptAsset.GetSoftObjectPath(), bIsInLibrary);
 		const TTuple<EScriptSource, FText> Source = FNiagaraEditorUtilities::GetScriptSource(ScriptAsset);
 		FNiagaraActionSourceData SourceData(Source.Key, Source.Value, true);
 		
 		const ENiagaraMenuSections Section = bSuggested ? ENiagaraMenuSections::Suggested: ENiagaraMenuSections::General;
 
 		UNiagaraNodeFunctionCall* FunctionCallNode = NewObject<UNiagaraNodeFunctionCall>(OwnerOfTemporaries);
-		FunctionCallNode->FunctionScriptAssetObjectPath = ScriptAsset.ObjectPath;
+PRAGMA_DISABLE_DEPRECATION_WARNINGS;
+		FunctionCallNode->FunctionScriptAssetObjectPath = ScriptAsset.GetSoftObjectPath().ToFName();
+PRAGMA_DISABLE_DEPRECATION_WARNINGS;
 
 		TSharedPtr<FNiagaraAction_NewNode> Action = AddNewNodeMenuAction(NewActions, FunctionCallNode, DisplayName, Section, Categories, TooltipDesc, Keywords, SourceData);
 		Action->bIsInLibrary = bIsInLibrary;
diff --git a/Engine/Plugins/FX/Niagara/Source/NiagaraEditor/Private/NiagaraEditorModule.cpp b/Engine/Plugins/FX/Niagara/Source/NiagaraEditor/Private/NiagaraEditorModule.cpp
--- a/Engine/Plugins/FX/Niagara/Source/NiagaraEditor/Private/NiagaraEditorModule.cpp
+++ b/Engine/Plugins/FX/Niagara/Source/NiagaraEditor/Private/NiagaraEditorModule.cpp
@@ -358,7 +358,7 @@ void DumpRapidIterationParamersForAsset(const TArray<FString>& Arguments)
 	if (Arguments.Num() == 1)
 	{
 		const FAssetRegistryModule& AssetRegistryModule = FModuleManager::LoadModuleChecked<FAssetRegistryModule>(TEXT("AssetRegistry"));
-		const FAssetData AssetData = AssetRegistryModule.Get().GetAssetByObjectPath(*Arguments[0]);
+		const FAssetData AssetData = AssetRegistryModule.Get().GetAssetByObjectPath(FSoftObjectPath(Arguments[0]));
 		UObject* Asset = AssetData.GetAsset();
 		if (Asset != nullptr)
 		{
@@ -446,7 +446,7 @@ void PreventSystemRecompile(const TArray<FString>& Arguments)
 	if (Arguments.Num() > 0)
 	{
 		const FAssetRegistryModule& AssetRegistryModule = FModuleManager::LoadModuleChecked<FAssetRegistryModule>(TEXT("AssetRegistry"));
-		FAssetData SystemAsset = AssetRegistryModule.Get().GetAssetByObjectPath(*Arguments[0]);
+		FAssetData SystemAsset = AssetRegistryModule.Get().GetAssetByObjectPath(FSoftObjectPath(Arguments[0]));
 		if (SystemAsset.IsValid() == false)
 		{
 			TArray<FAssetData> AssetsInPackage;
@@ -618,7 +618,7 @@ void DumpCompileIdDataForAsset(const TArray<FString>& Arguments)
 	if (Arguments.Num() > 0)
 	{
 		const FAssetRegistryModule& AssetRegistryModule = FModuleManager::LoadModuleChecked<FAssetRegistryModule>(TEXT("AssetRegistry"));
-		FAssetData SystemAsset = AssetRegistryModule.Get().GetAssetByObjectPath(*Arguments[0]);
+		FAssetData SystemAsset = AssetRegistryModule.Get().GetAssetByObjectPath(FSoftObjectPath(Arguments[0]));
 		if (SystemAsset.IsValid() == false)
 		{
 			TArray<FAssetData> AssetsInPackage;
diff --git a/Engine/Plugins/FX/Niagara/Source/NiagaraEditor/Private/NiagaraEditorUtilities.cpp b/Engine/Plugins/FX/Niagara/Source/NiagaraEditor/Private/NiagaraEditorUtilities.cpp
--- a/Engine/Plugins/FX/Niagara/Source/NiagaraEditor/Private/NiagaraEditorUtilities.cpp
+++ b/Engine/Plugins/FX/Niagara/Source/NiagaraEditor/Private/NiagaraEditorUtilities.cpp
@@ -1963,15 +1963,15 @@ NIAGARAEDITOR_API FText FNiagaraEditorUtilities::FormatScriptName(FName Name, bo
 	return FText::FromString(FName::NameToDisplayString(Name.ToString(), false) + (bIsInLibrary ? TEXT("") : TEXT("*")));
 }
 
-FText FNiagaraEditorUtilities::FormatScriptDescription(FText Description, FName Path, bool bIsInLibrary)
+FText FNiagaraEditorUtilities::FormatScriptDescription(FText Description, const FSoftObjectPath& Path, bool bIsInLibrary)
 {
 	FText LibrarySuffix = !bIsInLibrary
 		? LOCTEXT("LibrarySuffix", "\n* Script is not exposed to the library.")
 		: FText();
 
 	return Description.IsEmptyOrWhitespace()
-		? FText::Format(LOCTEXT("ScriptAssetDescriptionFormatPathOnly", "Path: {0}{1}"), FText::FromName(Path), LibrarySuffix)
-		: FText::Format(LOCTEXT("ScriptAssetDescriptionFormat", "{1}\nPath: {0}{2}"), FText::FromName(Path), Description, LibrarySuffix);
+		? FText::Format(LOCTEXT("ScriptAssetDescriptionFormatPathOnly", "Path: {0}{1}"), FText::FromString(Path.ToString()), LibrarySuffix)
+		: FText::Format(LOCTEXT("ScriptAssetDescriptionFormat", "{1}\nPath: {0}{2}"), FText::FromString(Path.ToString()), Description, LibrarySuffix);
 }
 
 FText FNiagaraEditorUtilities::FormatVariableDescription(FText Description, FText Name, FText Type)
diff --git a/Engine/Plugins/FX/Niagara/Source/NiagaraEditor/Private/NiagaraMessageManager.cpp b/Engine/Plugins/FX/Niagara/Source/NiagaraEditor/Private/NiagaraMessageManager.cpp
--- a/Engine/Plugins/FX/Niagara/Source/NiagaraEditor/Private/NiagaraMessageManager.cpp
+++ b/Engine/Plugins/FX/Niagara/Source/NiagaraEditor/Private/NiagaraMessageManager.cpp
@@ -344,7 +344,7 @@ void FNiagaraCompileEventToken::OpenScriptAssetByPathAndFocusNodeOrPinIfSet(
 {
 	FAssetRegistryModule& AssetRegistryModule = FModuleManager::LoadModuleChecked<FAssetRegistryModule>("AssetRegistry");
 	IAssetRegistry& AssetRegistry = AssetRegistryModule.Get();
-	FAssetData AssetData = AssetRegistry.GetAssetByObjectPath(*InScriptAssetPath);
+	FAssetData AssetData = AssetRegistry.GetAssetByObjectPath(FSoftObjectPath(InScriptAssetPath));
 	if (AssetData.IsValid())
 	{
 		UNiagaraScript* ScriptAsset = Cast<UNiagaraScript>(AssetData.GetAsset());
diff --git a/Engine/Plugins/FX/Niagara/Source/NiagaraEditor/Private/NiagaraNodeFunctionCall.cpp b/Engine/Plugins/FX/Niagara/Source/NiagaraEditor/Private/NiagaraNodeFunctionCall.cpp
--- a/Engine/Plugins/FX/Niagara/Source/NiagaraEditor/Private/NiagaraNodeFunctionCall.cpp
+++ b/Engine/Plugins/FX/Niagara/Source/NiagaraEditor/Private/NiagaraNodeFunctionCall.cpp
@@ -268,7 +268,9 @@ void UNiagaraNodeFunctionCall::AllocateDefaultPins()
 	if (FunctionScriptAssetObjectPath != NAME_None && FunctionScript == nullptr)
 	{
 		FAssetRegistryModule& AssetRegistryModule = FModuleManager::LoadModuleChecked<FAssetRegistryModule>(TEXT("AssetRegistry"));
+PRAGMA_DISABLE_DEPRECATION_WARNINGS;
 		FAssetData ScriptAssetData = AssetRegistryModule.Get().GetAssetByObjectPath(FunctionScriptAssetObjectPath);
+PRAGMA_ENABLE_DEPRECATION_WARNINGS;
 		if (ScriptAssetData.IsValid())
 		{
 			FunctionScript = Cast<UNiagaraScript>(ScriptAssetData.GetAsset());
@@ -468,7 +470,9 @@ bool UNiagaraNodeFunctionCall::CanAddToGraph(UNiagaraGraph* TargetGraph, FString
 	if (FunctionScriptAssetObjectPath != NAME_None && FunctionScript == nullptr)
 	{
 		FAssetRegistryModule& AssetRegistryModule = FModuleManager::LoadModuleChecked<FAssetRegistryModule>(TEXT("AssetRegistry"));
+PRAGMA_DISABLE_DEPRECATION_WARNINGS;
 		FAssetData ScriptAssetData = AssetRegistryModule.Get().GetAssetByObjectPath(FunctionScriptAssetObjectPath);
+PRAGMA_ENABLE_DEPRECATION_WARNINGS;
 		if (ScriptAssetData.IsValid())
 		{
 			SpawningFunctionScript = Cast<UNiagaraScript>(ScriptAssetData.GetAsset());
diff --git a/Engine/Plugins/FX/Niagara/Source/NiagaraEditor/Private/NiagaraScriptSource.cpp b/Engine/Plugins/FX/Niagara/Source/NiagaraEditor/Private/NiagaraScriptSource.cpp
--- a/Engine/Plugins/FX/Niagara/Source/NiagaraEditor/Private/NiagaraScriptSource.cpp
+++ b/Engine/Plugins/FX/Niagara/Source/NiagaraEditor/Private/NiagaraScriptSource.cpp
@@ -276,7 +276,8 @@ bool UNiagaraScriptSource::AddModuleIfMissing(FString ModulePath, ENiagaraScript
 {
 	FSoftObjectPath SystemUpdateScriptRef(ModulePath);
 	FAssetData ModuleScriptAsset;
-	ModuleScriptAsset.ObjectPath = SystemUpdateScriptRef.GetAssetPathName();
+	ModuleScriptAsset.PackageName = SystemUpdateScriptRef.GetAssetPath().GetPackageName();
+	ModuleScriptAsset.AssetName = SystemUpdateScriptRef.GetAssetPath().GetAssetName();
 	bOutFoundModule = false;
 
 	if (ModuleScriptAsset.IsValid())
diff --git a/Engine/Plugins/FX/Niagara/Source/NiagaraEditor/Private/SNiagaraAssetPickerList.cpp b/Engine/Plugins/FX/Niagara/Source/NiagaraEditor/Private/SNiagaraAssetPickerList.cpp
--- a/Engine/Plugins/FX/Niagara/Source/NiagaraEditor/Private/SNiagaraAssetPickerList.cpp
+++ b/Engine/Plugins/FX/Niagara/Source/NiagaraEditor/Private/SNiagaraAssetPickerList.cpp
@@ -195,7 +195,7 @@ TArray<FText> SNiagaraAssetPickerList::OnGetCategoriesForItem(const FAssetData&
 
 	auto AddAssetPathCategory = [&Categories, &Item]() {
 		TArray<FString> AssetPathParts;
-		Item.ObjectPath.ToString().ParseIntoArray(AssetPathParts, TEXT("/"));
+		Item.GetObjectPathString().ParseIntoArray(AssetPathParts, TEXT("/"));
 		if (AssetPathParts.Num() > 0)
 		{
 			if (AssetPathParts[0] == TEXT("Niagara"))
diff --git a/Engine/Plugins/FX/Niagara/Source/NiagaraEditor/Private/Tests/NiagaraEditorTestUtilities.cpp b/Engine/Plugins/FX/Niagara/Source/NiagaraEditor/Private/Tests/NiagaraEditorTestUtilities.cpp
--- a/Engine/Plugins/FX/Niagara/Source/NiagaraEditor/Private/Tests/NiagaraEditorTestUtilities.cpp
+++ b/Engine/Plugins/FX/Niagara/Source/NiagaraEditor/Private/Tests/NiagaraEditorTestUtilities.cpp
@@ -59,6 +59,6 @@ TSharedRef<FNiagaraSystemViewModel> FNiagaraEditorTestUtilities::CreateTestSyste
 UNiagaraEmitter* FNiagaraEditorTestUtilities::LoadEmitter(FString AssetPath)
 {
 	const FAssetRegistryModule& AssetRegistryModule = FModuleManager::LoadModuleChecked<FAssetRegistryModule>(TEXT("AssetRegistry"));
-	const FAssetData AssetData = AssetRegistryModule.Get().GetAssetByObjectPath(*AssetPath);
+	const FAssetData AssetData = AssetRegistryModule.Get().GetAssetByObjectPath(FSoftObjectPath(AssetPath));
 	return Cast<UNiagaraEmitter>(AssetData.GetAsset());
 }
\ No newline at end of file
diff --git a/Engine/Plugins/FX/Niagara/Source/NiagaraEditor/Private/ViewModels/NiagaraScratchPadViewModel.cpp b/Engine/Plugins/FX/Niagara/Source/NiagaraEditor/Private/ViewModels/NiagaraScratchPadViewModel.cpp
--- a/Engine/Plugins/FX/Niagara/Source/NiagaraEditor/Private/ViewModels/NiagaraScratchPadViewModel.cpp
+++ b/Engine/Plugins/FX/Niagara/Source/NiagaraEditor/Private/ViewModels/NiagaraScratchPadViewModel.cpp
@@ -387,7 +387,7 @@ void GetScriptsFromClipboard(TArray<const UNiagaraScript*>& OutScripts, const TA
 		if (ClipboardString.IsEmpty() == false && ClipboardString.Len() < NAME_SIZE && AssetRegistry != nullptr)
 		{
 			TArray<FAssetData> Assets;
-			FAssetData AssetFoundByObjectPath = AssetRegistry->GetAssetByObjectPath(*ClipboardString);
+			FAssetData AssetFoundByObjectPath = AssetRegistry->GetAssetByObjectPath(FSoftObjectPath(ClipboardString));
 			if (AssetFoundByObjectPath.IsValid())
 			{
 				Assets.Add(AssetFoundByObjectPath);
diff --git a/Engine/Plugins/FX/Niagara/Source/NiagaraEditor/Private/ViewModels/Stack/NiagaraStackFunctionInput.cpp b/Engine/Plugins/FX/Niagara/Source/NiagaraEditor/Private/ViewModels/Stack/NiagaraStackFunctionInput.cpp
--- a/Engine/Plugins/FX/Niagara/Source/NiagaraEditor/Private/ViewModels/Stack/NiagaraStackFunctionInput.cpp
+++ b/Engine/Plugins/FX/Niagara/Source/NiagaraEditor/Private/ViewModels/Stack/NiagaraStackFunctionInput.cpp
@@ -1579,7 +1579,7 @@ void UNiagaraStackFunctionInput::GetAvailableParameterHandles(TArray<FNiagaraPar
 		}
 		else
 		{
-			UE_LOG(LogNiagaraEditor, Warning, TEXT("Failed to load NiagaraParameterCollection '%s'"), *CollectionAsset.ObjectPath.ToString());
+			UE_LOG(LogNiagaraEditor, Warning, TEXT("Failed to load NiagaraParameterCollection '%s'"), *CollectionAsset.GetObjectPathString());
 		}
 	}
 }
diff --git a/Engine/Plugins/FX/Niagara/Source/NiagaraEditor/Private/ViewModels/Stack/NiagaraStackGraphUtilities.cpp b/Engine/Plugins/FX/Niagara/Source/NiagaraEditor/Private/ViewModels/Stack/NiagaraStackGraphUtilities.cpp
--- a/Engine/Plugins/FX/Niagara/Source/NiagaraEditor/Private/ViewModels/Stack/NiagaraStackGraphUtilities.cpp
+++ b/Engine/Plugins/FX/Niagara/Source/NiagaraEditor/Private/ViewModels/Stack/NiagaraStackGraphUtilities.cpp
@@ -1761,12 +1761,14 @@ bool FNiagaraStackGraphUtilities::FindScriptModulesInStack(FAssetData ModuleScri
 	Graph->BuildTraversal(Nodes, &TargetOutputNode);
 
 	OutFunctionCalls.Empty();
-	FString ModuleObjectName = ModuleScriptAsset.ObjectPath.ToString();
+	FString ModuleObjectName = ModuleScriptAsset.GetObjectPathString();
 	for (UEdGraphNode* Node : Nodes)
 	{
 		if (UNiagaraNodeFunctionCall* FunctionCallNode = Cast<UNiagaraNodeFunctionCall>(Node))
 		{
-			if (FunctionCallNode->FunctionScriptAssetObjectPath == ModuleScriptAsset.ObjectPath ||
+PRAGMA_DISABLE_DEPRECATION_WARNINGS
+			if (FunctionCallNode->FunctionScriptAssetObjectPath == ModuleScriptAsset.GetSoftObjectPath().ToFName() ||
+PRAGMA_ENABLE_DEPRECATION_WARNINGS
 				(FunctionCallNode->FunctionScript != nullptr && FunctionCallNode->FunctionScript->GetPathName() == ModuleObjectName))
 			{
 				OutFunctionCalls.Add(FunctionCallNode);
@@ -1871,7 +1873,9 @@ UNiagaraNodeFunctionCall* FNiagaraStackGraphUtilities::AddScriptModuleToStack(FA
 
 	FGraphNodeCreator<UNiagaraNodeFunctionCall> ModuleNodeCreator(*Graph);
 	UNiagaraNodeFunctionCall* NewModuleNode = ModuleNodeCreator.CreateNode();
-	NewModuleNode->FunctionScriptAssetObjectPath = ModuleScriptAsset.ObjectPath;
+PRAGMA_DISABLE_DEPRECATION_WARNINGS
+	NewModuleNode->FunctionScriptAssetObjectPath = ModuleScriptAsset.GetSoftObjectPath().ToFName();
+PRAGMA_ENABLE_DEPRECATION_WARNINGS
 	ModuleNodeCreator.Finalize();
 
 	if (NewModuleNode->HasValidScriptAndGraph() == false)
diff --git a/Engine/Plugins/FX/Niagara/Source/NiagaraEditor/Private/ViewModels/Stack/NiagaraStackScriptItemGroup.cpp b/Engine/Plugins/FX/Niagara/Source/NiagaraEditor/Private/ViewModels/Stack/NiagaraStackScriptItemGroup.cpp
--- a/Engine/Plugins/FX/Niagara/Source/NiagaraEditor/Private/ViewModels/Stack/NiagaraStackScriptItemGroup.cpp
+++ b/Engine/Plugins/FX/Niagara/Source/NiagaraEditor/Private/ViewModels/Stack/NiagaraStackScriptItemGroup.cpp
@@ -52,7 +52,7 @@ public:
 
 		FText AssetDescription;
 		AssetData.GetTagValue(GET_MEMBER_NAME_CHECKED(FVersionedNiagaraScriptData, Description), AssetDescription);
-		FText Description = FNiagaraEditorUtilities::FormatScriptDescription(AssetDescription, AssetData.ObjectPath, bIsInLibrary);
+		FText Description = FNiagaraEditorUtilities::FormatScriptDescription(AssetDescription, AssetData.GetSoftObjectPath(), bIsInLibrary);
 
 		FText Keywords;
 		AssetData.GetTagValue(GET_MEMBER_NAME_CHECKED(FVersionedNiagaraScriptData, Keywords), Keywords);
@@ -75,7 +75,7 @@ public:
 		// we assume scratch pads should always be displayed, so we act like it's a library action so it's not accidentally filtered out
 		bool bIsInLibrary = true;
 		FText DisplayName = FNiagaraEditorUtilities::FormatScriptName(ModuleScript->GetFName(), bIsInLibrary);
-		FText Description = FNiagaraEditorUtilities::FormatScriptDescription(ScriptData->Description, *ModuleScript->GetPathName(), bIsInLibrary);
+		FText Description = FNiagaraEditorUtilities::FormatScriptDescription(ScriptData->Description, FSoftObjectPath(ModuleScript), bIsInLibrary);
 		FText Keywords = ScriptData->Keywords;
 		bool bSuggested = ScriptData->bSuggested;
 
@@ -615,7 +615,7 @@ void UNiagaraStackScriptItemGroup::RefreshIssues(TArray<FStackIssue>& NewIssues)
 			// The factor ensures this, but older assets may not have it or it may have been removed accidentally.
 			// For now, treat this as an error and allow them to resolve.
 			const FAssetRegistryModule& AssetRegistryModule = FModuleManager::LoadModuleChecked<FAssetRegistryModule>(TEXT("AssetRegistry"));
-			const FAssetData ModuleScriptAsset = AssetRegistryModule.Get().GetAssetByObjectPath(GetDefault<UNiagaraEditorSettings>()->RequiredSystemUpdateScript.GetAssetPathName());
+			const FAssetData ModuleScriptAsset = AssetRegistryModule.Get().GetAssetByObjectPath(GetDefault<UNiagaraEditorSettings>()->RequiredSystemUpdateScript);
 
 			TArray<UNiagaraNodeFunctionCall*> FoundCalls;
 			UNiagaraNodeOutput* MatchingOutputNode = Graph->FindOutputNode(ScriptUsage, ScriptUsageId);
diff --git a/Engine/Plugins/FX/Niagara/Source/NiagaraEditor/Public/NiagaraEditorUtilities.h b/Engine/Plugins/FX/Niagara/Source/NiagaraEditor/Public/NiagaraEditorUtilities.h
--- a/Engine/Plugins/FX/Niagara/Source/NiagaraEditor/Public/NiagaraEditorUtilities.h
+++ b/Engine/Plugins/FX/Niagara/Source/NiagaraEditor/Public/NiagaraEditorUtilities.h
@@ -246,7 +246,7 @@ namespace FNiagaraEditorUtilities
 
 	NIAGARAEDITOR_API FText FormatScriptName(FName Name, bool bIsInLibrary);
 
-	NIAGARAEDITOR_API FText FormatScriptDescription(FText Description, FName Path, bool bIsInLibrary);
+	NIAGARAEDITOR_API FText FormatScriptDescription(FText Description, const FSoftObjectPath& Path, bool bIsInLibrary);
 
 	NIAGARAEDITOR_API FText FormatVariableDescription(FText Description, FText Name, FText Type);
 
diff --git a/Engine/Plugins/FX/Niagara/Source/NiagaraEditorWidgets/Private/Stack/SNiagaraStackFunctionInputValue.cpp b/Engine/Plugins/FX/Niagara/Source/NiagaraEditorWidgets/Private/Stack/SNiagaraStackFunctionInputValue.cpp
--- a/Engine/Plugins/FX/Niagara/Source/NiagaraEditorWidgets/Private/Stack/SNiagaraStackFunctionInputValue.cpp
+++ b/Engine/Plugins/FX/Niagara/Source/NiagaraEditorWidgets/Private/Stack/SNiagaraStackFunctionInputValue.cpp
@@ -929,7 +929,7 @@ TArray<TSharedPtr<FNiagaraMenuAction_Generic>> SNiagaraStackFunctionInputValue::
 		FVersionedNiagaraScriptData* ScriptData = DynamicInputScript->GetLatestScriptData();
 		bool bIsInLibrary = ScriptData->LibraryVisibility == ENiagaraScriptLibraryVisibility::Library;
 		const FText DisplayName = FNiagaraEditorUtilities::FormatScriptName(DynamicInputScript->GetFName(), bIsInLibrary);
-		const FText Tooltip = FNiagaraEditorUtilities::FormatScriptDescription(ScriptData->Description, *DynamicInputScript->GetPathName(), bIsInLibrary);
+		const FText Tooltip = FNiagaraEditorUtilities::FormatScriptDescription(ScriptData->Description, FSoftObjectPath(DynamicInputScript), bIsInLibrary);
 		TTuple<EScriptSource, FText> Source = FNiagaraEditorUtilities::GetScriptSource(FAssetData(DynamicInputScript));
 
 		// scratch pad dynamic inputs are always considered to be in the library and will have Niagara as the source
@@ -1082,7 +1082,7 @@ TArray<TSharedPtr<FNiagaraMenuAction_Generic>> SNiagaraStackFunctionInputValue::
 			FVersionedNiagaraScriptData* ScriptData = DynamicInputScript->GetLatestScriptData();
 			bool bIsInLibrary = ScriptData->LibraryVisibility == ENiagaraScriptLibraryVisibility::Library;
 			const FText DisplayName = FNiagaraEditorUtilities::FormatScriptName(DynamicInputScript->GetFName(), bIsInLibrary);
-			const FText Tooltip = FNiagaraEditorUtilities::FormatScriptDescription(ScriptData->Description, *DynamicInputScript->GetPathName(), bIsInLibrary);
+			const FText Tooltip = FNiagaraEditorUtilities::FormatScriptDescription(ScriptData->Description, FSoftObjectPath(DynamicInputScript), bIsInLibrary);
 
 			// scratch pad dynamic inputs are always considered to be in the library and will have Niagara as the source
 			if(ScratchPadDynamicInputs.Contains(DynamicInputScript))
diff --git a/Engine/Plugins/FX/Niagara/Source/NiagaraEditorWidgets/Private/Stack/SNiagaraStackModuleItem.cpp b/Engine/Plugins/FX/Niagara/Source/NiagaraEditorWidgets/Private/Stack/SNiagaraStackModuleItem.cpp
--- a/Engine/Plugins/FX/Niagara/Source/NiagaraEditorWidgets/Private/Stack/SNiagaraStackModuleItem.cpp
+++ b/Engine/Plugins/FX/Niagara/Source/NiagaraEditorWidgets/Private/Stack/SNiagaraStackModuleItem.cpp
@@ -545,7 +545,7 @@ void SNiagaraStackModuleItem::CollectModuleActions(FGraphActionListBuilderBase&
 
 		FText AssetDescription;
 		ModuleAsset.GetTagValue(GET_MEMBER_NAME_CHECKED(FVersionedNiagaraScriptData, Description), AssetDescription);
-		FText Description = FNiagaraEditorUtilities::FormatScriptDescription(AssetDescription, ModuleAsset.ObjectPath, bIsInLibrary);
+		FText Description = FNiagaraEditorUtilities::FormatScriptDescription(AssetDescription, ModuleAsset.GetSoftObjectPath(), bIsInLibrary);
 
 		FText Keywords;
 		ModuleAsset.GetTagValue(GET_MEMBER_NAME_CHECKED(FVersionedNiagaraScriptData, Keywords), Keywords);
