- IWYU fix

#rnx

diff --git a/Engine/Plugins/FX/Niagara/Source/Niagara/Private/NDIRenderTargetVolumeSimCacheData.h b/Engine/Plugins/FX/Niagara/Source/Niagara/Private/NDIRenderTargetVolumeSimCacheData.h
--- a/Engine/Plugins/FX/Niagara/Source/Niagara/Private/NDIRenderTargetVolumeSimCacheData.h
+++ b/Engine/Plugins/FX/Niagara/Source/Niagara/Private/NDIRenderTargetVolumeSimCacheData.h
@@ -1,6 +1,9 @@
 // Copyright Epic Games, Inc. All Rights Reserved.
 #pragma once
 
+#include "Serialization/BulkData.h"
+#include "PixelFormat.h"
+
 #include "NDIRenderTargetVolumeSimCacheData.generated.h"
 
 USTRUCT()
