Fix mismatched deprecation macros

diff --git a/Engine/Plugins/FX/Niagara/Source/NiagaraEditor/Private/Commandlets/NiagaraDumpBytecodeCommandlet.cpp b/Engine/Plugins/FX/Niagara/Source/NiagaraEditor/Private/Commandlets/NiagaraDumpBytecodeCommandlet.cpp
--- a/Engine/Plugins/FX/Niagara/Source/NiagaraEditor/Private/Commandlets/NiagaraDumpBytecodeCommandlet.cpp
+++ b/Engine/Plugins/FX/Niagara/Source/NiagaraEditor/Private/Commandlets/NiagaraDumpBytecodeCommandlet.cpp
@@ -94,7 +94,7 @@ void UNiagaraDumpByteCodeCommandlet::ProcessNiagaraScripts()
 		FCollectionManagerModule& CollectionManagerModule = FCollectionManagerModule::GetModule();
 PRAGMA_DISABLE_DEPRECATION_WARNINGS
 		CollectionManagerModule.Get().GetObjectsInCollection(FName(*FilterCollection), ECollectionShareType::CST_All, Filter.ObjectPaths, ECollectionRecursionFlags::SelfAndChildren);
-PRAGMA_DISABLE_DEPRECATION_WARNINGS
+PRAGMA_ENABLE_DEPRECATION_WARNINGS
 	}
 
 	TArray<FAssetData> AssetList;
diff --git a/Engine/Plugins/FX/Niagara/Source/NiagaraEditor/Private/EdGraphSchema_Niagara.cpp b/Engine/Plugins/FX/Niagara/Source/NiagaraEditor/Private/EdGraphSchema_Niagara.cpp
--- a/Engine/Plugins/FX/Niagara/Source/NiagaraEditor/Private/EdGraphSchema_Niagara.cpp
+++ b/Engine/Plugins/FX/Niagara/Source/NiagaraEditor/Private/EdGraphSchema_Niagara.cpp
@@ -505,9 +505,9 @@ TArray<TSharedPtr<FNiagaraAction_NewNode>> UEdGraphSchema_Niagara::GetGraphActio
 		const ENiagaraMenuSections Section = bSuggested ? ENiagaraMenuSections::Suggested: ENiagaraMenuSections::General;
 
 		UNiagaraNodeFunctionCall* FunctionCallNode = NewObject<UNiagaraNodeFunctionCall>(OwnerOfTemporaries);
-PRAGMA_DISABLE_DEPRECATION_WARNINGS;
+PRAGMA_DISABLE_DEPRECATION_WARNINGS
 		FunctionCallNode->FunctionScriptAssetObjectPath = ScriptAsset.GetSoftObjectPath().ToFName();
-PRAGMA_DISABLE_DEPRECATION_WARNINGS;
+PRAGMA_ENABLE_DEPRECATION_WARNINGS
 
 		TSharedPtr<FNiagaraAction_NewNode> Action = AddNewNodeMenuAction(NewActions, FunctionCallNode, DisplayName, Section, Categories, TooltipDesc, Keywords, SourceData);
 		Action->bIsInLibrary = bIsInLibrary;
