Niagara - Make the emitter handles null out the emitter thumbnail pointers on load to prevent them interfering with the processing of cooked editor content.  This is a temporary workaround.

diff --git a/Engine/Plugins/FX/Niagara/Source/Niagara/Private/NiagaraEmitterHandle.cpp b/Engine/Plugins/FX/Niagara/Source/Niagara/Private/NiagaraEmitterHandle.cpp
--- a/Engine/Plugins/FX/Niagara/Source/Niagara/Private/NiagaraEmitterHandle.cpp
+++ b/Engine/Plugins/FX/Niagara/Source/Niagara/Private/NiagaraEmitterHandle.cpp
@@ -225,6 +225,9 @@ void FNiagaraEmitterHandle::ConditionalPostLoad(int32 NiagaraCustomVersion)
 				SetName(Name, *OwningSystem);
 			}
 		}
+
+		// Temporary Workaround - These thumbnails aren't used by the system so remove them here to fix issues with cooked editor data.
+		InstanceEmitter->ThumbnailImage = nullptr;
 	}
 }
 
