Add an error the niagara hlsl translator catching cases where conflicting symbol names would be created for unique parameters

diff --git a/Engine/Plugins/FX/Niagara/Source/NiagaraEditor/Private/NiagaraHlslTranslator.cpp b/Engine/Plugins/FX/Niagara/Source/NiagaraEditor/Private/NiagaraHlslTranslator.cpp
--- a/Engine/Plugins/FX/Niagara/Source/NiagaraEditor/Private/NiagaraHlslTranslator.cpp
+++ b/Engine/Plugins/FX/Niagara/Source/NiagaraEditor/Private/NiagaraHlslTranslator.cpp
@@ -3924,18 +3924,18 @@ void FHlslNiagaraTranslator::Init()
 {
 }
 
-FString FHlslNiagaraTranslator::GetSanitizedSymbolName(FString SymbolName, bool bCollapsNamespaces)
+FString FHlslNiagaraTranslator::GetSanitizedSymbolName(FStringView SymbolName, bool bCollapsNamespaces)
 {
 	if (SymbolName.Len() == 0)
 	{
-		return SymbolName;
+		return FString(SymbolName);
 	}
 
 	const UNiagaraEditorSettings* Settings = GetDefault<UNiagaraEditorSettings>();
 	check(Settings);
 	const TMap<FString, FString>& ReplacementsForInvalid = Settings->GetHLSLKeywordReplacementsMap();
 
-	FString Ret = SymbolName;
+	FString Ret = FString(SymbolName);
 
 	// Split up into individual namespaces...
 	TArray<FString> SplitName;
@@ -5567,7 +5567,31 @@ bool FHlslNiagaraTranslator::ParameterMapRegisterExternalConstantNamespaceVariab
 
 				if ( IsVariableInUniformBuffer(InVariable) )
 				{
-					CompilationOutput.ScriptData.Parameters.SetOrAdd(InVariable);
+					// we must ensure that there's a one to one relationship between symbol name and parameter.  The generated VM only
+					// knows about the symbols while the parameter stores knows about the parameters, if these mismatch, then we're going
+					// to be incorrectly addressing the constant table
+					if (!CompilationOutput.ScriptData.Parameters.FindParameter(InVariable))
+					{
+						bool AddParameter = true;
+
+						// add the parameter, but first evaluate whether any of the symbols for existing parameters would conflict
+						for (const FNiagaraVariable& ExistingParameter : CompilationOutput.ScriptData.Parameters.Parameters)
+						{
+							FNameBuilder ExistingParameterName(ExistingParameter.GetName());
+							if (GetSanitizedSymbolName(ExistingParameterName).Equals(SymbolName))
+							{
+								Error(FText::Format(LOCTEXT("NonUniqueSymbolNames", "Parameters ('{0}' and '{1}') found which resolve to the same HLSL symbol name '{2}'.  These should be disambiguated."),
+									FText::FromName(InVariable.GetName()), FText::FromName(ExistingParameter.GetName()), FText::FromString(SymbolName)), InNodeForErrorReporting, InDefaultPin);
+
+								AddParameter = false;
+							}
+						}
+
+						if (AddParameter)
+						{
+							CompilationOutput.ScriptData.Parameters.Parameters.Add(InVariable);
+						}
+					}
 				}
 
 				UniformChunk = AddUniformChunk(SymbolNameDefined, InVariable, ENiagaraCodeChunkMode::Uniform, UNiagaraScript::IsGPUScript(CompileOptions.TargetUsage));
diff --git a/Engine/Plugins/FX/Niagara/Source/NiagaraEditor/Private/NiagaraHlslTranslator.h b/Engine/Plugins/FX/Niagara/Source/NiagaraEditor/Private/NiagaraHlslTranslator.h
--- a/Engine/Plugins/FX/Niagara/Source/NiagaraEditor/Private/NiagaraHlslTranslator.h
+++ b/Engine/Plugins/FX/Niagara/Source/NiagaraEditor/Private/NiagaraHlslTranslator.h
@@ -643,7 +643,7 @@ public:
 	static TArray<FName> ConditionPropertyPath(const FNiagaraTypeDefinition& Type, const TArray<FName>& InPath);
 
 
-	static FString GetSanitizedSymbolName(FString SymbolName, bool bCollapseNamespaces=false);
+	static FString GetSanitizedSymbolName(FStringView SymbolName, bool bCollapseNamespaces=false);
 	static FString GetSanitizedDIFunctionName(const FString& FunctionName);
 	static FString GetSanitizedFunctionNameSuffix(FString Name);
 
diff --git a/Engine/Source/Runtime/Core/Private/UObject/FortniteMainBranchObjectVersion.cpp b/Engine/Source/Runtime/Core/Private/UObject/FortniteMainBranchObjectVersion.cpp
--- a/Engine/Source/Runtime/Core/Private/UObject/FortniteMainBranchObjectVersion.cpp
+++ b/Engine/Source/Runtime/Core/Private/UObject/FortniteMainBranchObjectVersion.cpp
@@ -11,7 +11,7 @@ TMap<FGuid, FGuid> FFortniteMainBranchObjectVersion::GetSystemGuids()
 	SystemGuids.Add(DevGuids.MATERIALSHADERMAP_DERIVEDDATA_VER, FGuid("1DA85DA5733C4A489A1BE117CDC5ACCC"));
 	SystemGuids.Add(DevGuids.NANITE_DERIVEDDATA_VER, FGuid("93379C2F29624E74AFB14FB13E327FA2"));
 	SystemGuids.Add(DevGuids.NIAGARASHADERMAP_DERIVEDDATA_VER, FGuid("2289D5116CF94BC0AFEAEC541468E645"));
-	SystemGuids.Add(DevGuids.Niagara_LatestScriptCompileVersion, FGuid("FDA3A752D5534180B6A6C27166FDFB7E"));
+	SystemGuids.Add(DevGuids.Niagara_LatestScriptCompileVersion, FGuid("D6A13E2AE4D2453EB1A44529AAAB8305"));
 	SystemGuids.Add(DevGuids.SkeletalMeshDerivedDataVersion, FGuid("ACF593EAAE354FCBB34CF44F5AA1BFD2"));
 	SystemGuids.Add(DevGuids.STATICMESH_DERIVEDDATA_VER, FGuid("4696D39D9BC743AFBFA1DAB359AD7A7D"));
 
