Fix for when you change renderer bindings, to ensure the parameter panel counts update.

diff --git a/Engine/Plugins/FX/Niagara/Source/NiagaraEditor/Private/ViewModels/Stack/NiagaraStackRendererItem.cpp b/Engine/Plugins/FX/Niagara/Source/NiagaraEditor/Private/ViewModels/Stack/NiagaraStackRendererItem.cpp
--- a/Engine/Plugins/FX/Niagara/Source/NiagaraEditor/Private/ViewModels/Stack/NiagaraStackRendererItem.cpp
+++ b/Engine/Plugins/FX/Niagara/Source/NiagaraEditor/Private/ViewModels/Stack/NiagaraStackRendererItem.cpp
@@ -23,6 +23,7 @@
 #include "NiagaraScriptMergeManager.h"
 #include "NiagaraClipboard.h"
 #include "NiagaraEmitterEditorData.h"
+#include "ViewModels/NiagaraParameterPanelViewModel.h"
 
 #include "ViewModels/NiagaraEmitterHandleViewModel.h"
 #include "ViewModels/Stack/NiagaraStackViewModel.h"
@@ -483,7 +484,15 @@ void UNiagaraStackRendererItem::RendererChanged()
 		bCanResetToBaseCache.Reset();
 		RefreshChildren();
 
-		
+		if (GetSystemViewModel()->GetSystemStackViewModel())
+		{
+			GetSystemViewModel()->GetSystemStackViewModel()->InvalidateCachedParameterUsage();
+		}
+		if (GetSystemViewModel()->GetParameterPanelViewModel())
+		{
+			GetSystemViewModel()->GetParameterPanelViewModel()->RefreshNextTick();
+		}
+
 		if (GetEmitterViewModel().IsValid())
 		{
 			GetSystemViewModel()->GetEmitterHandleViewModelForEmitter(GetEmitterViewModel()->GetEmitter()).Get()->GetEmitterStackViewModel()->RequestValidationUpdate();
