Fixed a bug when saving an emitter over an existing asset

diff --git a/Engine/Plugins/FX/Niagara/Source/NiagaraEditor/Private/NiagaraEditorUtilities.cpp b/Engine/Plugins/FX/Niagara/Source/NiagaraEditor/Private/NiagaraEditorUtilities.cpp
--- a/Engine/Plugins/FX/Niagara/Source/NiagaraEditor/Private/NiagaraEditorUtilities.cpp
+++ b/Engine/Plugins/FX/Niagara/Source/NiagaraEditor/Private/NiagaraEditorUtilities.cpp
@@ -2344,6 +2344,7 @@ void FNiagaraEditorUtilities::CreateAssetFromEmitter(TSharedRef<FNiagaraEmitterH
 	
 	// First duplicate the asset so that fixes can be made on the duplicate without modifying the system and before it's saved.
 	UNiagaraEmitter* DuplicateEmitter = CastChecked<UNiagaraEmitter>(StaticDuplicateObject(EmitterToCopy, GetTransientPackage()));
+	DuplicateEmitter->AddToRoot(); // needed because the asset save dialog can call garbage collection
 	for (const FNiagaraAssetVersion& Version : DuplicateEmitter->GetAllAvailableVersions())
 	{
 		FNiagaraScratchPadUtilities::FixExternalScratchPadScriptsForEmitter(SystemViewModel->GetSystem(), FVersionedNiagaraEmitter(DuplicateEmitter, Version.VersionGuid));
@@ -2351,6 +2352,7 @@ void FNiagaraEditorUtilities::CreateAssetFromEmitter(TSharedRef<FNiagaraEmitterH
 
 	// Save the duplicated emitter.
 	UNiagaraEmitter* CreatedAsset = Cast<UNiagaraEmitter>(AssetToolsModule.Get().DuplicateAssetWithDialogAndTitle(EmitterName.GetPlainNameString(), PackagePath, DuplicateEmitter, LOCTEXT("CreateEmitterAssetDialogTitle", "Create Emitter As")));
+	DuplicateEmitter->RemoveFromRoot();
 	if (CreatedAsset != nullptr)
 	{
 		CreatedAsset->SetUniqueEmitterName(CreatedAsset->GetName());
