- Added new widget "SNiagaraSystemUserParameters" for display & editing of user parameters in a system
- The new widget is contained in a new User Parameters tab
-Added "OnStructureChanged" delegate to user param store, which isn't called when a param value changes. This is used to refresh the above widget more granularly

diff --git a/Engine/Plugins/FX/Niagara/Source/Niagara/Private/NiagaraParameterStore.cpp b/Engine/Plugins/FX/Niagara/Source/Niagara/Private/NiagaraParameterStore.cpp
--- a/Engine/Plugins/FX/Niagara/Source/Niagara/Private/NiagaraParameterStore.cpp
+++ b/Engine/Plugins/FX/Niagara/Source/Niagara/Private/NiagaraParameterStore.cpp
@@ -1229,6 +1229,7 @@ void FNiagaraParameterStore::OnLayoutChange()
 	if (bSuppressOnChanged == false)
 	{
 		OnChangedDelegate.Broadcast();
+		OnStructureChangedDelegate.Broadcast();
 	}
 #endif
 }
diff --git a/Engine/Plugins/FX/Niagara/Source/Niagara/Public/NiagaraParameterStore.h b/Engine/Plugins/FX/Niagara/Source/Niagara/Public/NiagaraParameterStore.h
--- a/Engine/Plugins/FX/Niagara/Source/Niagara/Public/NiagaraParameterStore.h
+++ b/Engine/Plugins/FX/Niagara/Source/Niagara/Public/NiagaraParameterStore.h
@@ -174,6 +174,8 @@ struct NIAGARA_API FNiagaraParameterStore
 		{
 		}
 	};
+	
+	DECLARE_MULTICAST_DELEGATE(FOnStructureChanged)
 #endif
 
 	GENERATED_USTRUCT_BODY()
@@ -237,6 +239,7 @@ private:
 #if WITH_EDITOR
 	FOnChanged OnChangedDelegate;
 	bool bSuppressOnChanged = false;
+	FOnStructureChanged OnStructureChangedDelegate;
 #endif
 
 	void SetPositionData(const FName& Name, const FVector& Position);
@@ -664,6 +667,8 @@ public:
 	FDelegateHandle AddOnChangedHandler(FOnChanged::FDelegate InOnChanged);
 	void RemoveOnChangedHandler(FDelegateHandle DelegateHandle);
 	void RemoveAllOnChangedHandlers(const void* InUserObject);
+
+	FOnStructureChanged& OnStructureChanged() { return OnStructureChangedDelegate; }
 #endif
 
 	void TriggerOnLayoutChanged() { OnLayoutChange(); }
diff --git a/Engine/Plugins/FX/Niagara/Source/NiagaraEditor/Private/Customizations/NiagaraComponentDetails.cpp b/Engine/Plugins/FX/Niagara/Source/NiagaraEditor/Private/Customizations/NiagaraComponentDetails.cpp
--- a/Engine/Plugins/FX/Niagara/Source/NiagaraEditor/Private/Customizations/NiagaraComponentDetails.cpp
+++ b/Engine/Plugins/FX/Niagara/Source/NiagaraEditor/Private/Customizations/NiagaraComponentDetails.cpp
@@ -6,51 +6,31 @@
 #include "DetailCategoryBuilder.h"
 #include "NiagaraComponent.h"
 #include "NiagaraDataInterface.h"
-#include "DetailLayoutBuilder.h"
 #include "IDetailPropertyRow.h"
 #include "DetailWidgetRow.h"
-#include "Materials/Material.h"
 #include "Widgets/Text/SInlineEditableTextBlock.h"
-#include "Widgets/SToolTip.h"
 #include "PropertyCustomizationHelpers.h"
 #include "Widgets/Text/STextBlock.h"
 #include "Widgets/Images/SImage.h"
 #include "Widgets/Layout/SBox.h"
 #include "Styling/AppStyle.h"
-#include "NiagaraSystemInstance.h"
 #include "ViewModels/NiagaraParameterViewModel.h"
 #include "NiagaraEditorUtilities.h"
 #include "INiagaraEditorTypeUtilities.h"
 #include "NiagaraTypes.h"
 #include "ScopedTransaction.h"
-#include "SNiagaraParameterEditor.h"
 #include "NiagaraEditorModule.h"
 #include "PropertyEditorModule.h"
 #include "Modules/ModuleManager.h"
-#include "IStructureDetailsView.h"
-#include "IDetailsView.h"
 #include "Widgets/Input/SButton.h"
-#include "Widgets/Input/SEditableText.h"
-#include "Editor.h"
-#include "NiagaraParameterCollectionViewModel.h"
 #include "NiagaraScriptSource.h"
 #include "NiagaraNodeInput.h"
-#include "NiagaraGraph.h"
 #include "GameDelegates.h"
 #include "NiagaraEditorStyle.h"
 #include "IDetailChildrenBuilder.h"
-#include "IDetailGroup.h"
-#include "Framework/MultiBox/MultiBoxBuilder.h"
-#include "Widgets/Input/SComboButton.h"
-#include "Widgets/Images/SImage.h"
-#include "NiagaraEditorModule.h"
-#include "Modules/ModuleManager.h"
-#include "INiagaraEditorTypeUtilities.h"
 #include "NiagaraConstants.h"
 #include "NiagaraScriptVariable.h"
 #include "NiagaraSystemEditorData.h"
-#include "Widgets/Layout/SBox.h"
-#include "UObject/WeakObjectPtr.h"
 #include "NiagaraUserRedirectionParameterStore.h"
 #include "Widgets/Layout/SUniformGridPanel.h"
 #include "Widgets/SNiagaraDebugger.h"
@@ -60,9 +40,9 @@
 
 #define LOCTEXT_NAMESPACE "NiagaraComponentDetails"
 
-static FNiagaraVariant GetParameterValueFromAsset(const FNiagaraVariableBase& Parameter, const UNiagaraComponent* Component)
+static FNiagaraVariant GetParameterValueFromSystem(const FNiagaraVariableBase& Parameter, const UNiagaraSystem* System)
 {
-	FNiagaraUserRedirectionParameterStore& UserParameterStore = Component->GetAsset()->GetExposedParameters();
+	const FNiagaraUserRedirectionParameterStore& UserParameterStore = System->GetExposedParameters();
 	if (Parameter.IsDataInterface())
 	{
 		int32 Index = UserParameterStore.IndexOf(Parameter);
@@ -96,25 +76,35 @@ static FNiagaraVariant GetParameterValueFromAsset(const FNiagaraVariableBase& Pa
 	return FNiagaraVariant();
 }
 
-static FNiagaraVariant GetCurrentParameterValue(const FNiagaraVariableBase& Parameter, const UNiagaraComponent* Component)
+static FNiagaraVariant GetCurrentParameterValue(const FNiagaraVariableBase& Parameter, const UObject* Owner)
 {
-	FNiagaraVariant CurrentValue = Component->GetCurrentParameterValue(Parameter);
-	if (CurrentValue.IsValid())
+	if(const UNiagaraComponent* Component = Cast<UNiagaraComponent>(Owner))
 	{
-		return CurrentValue;
+		FNiagaraVariant CurrentValue = Component->GetCurrentParameterValue(Parameter);
+		if (CurrentValue.IsValid())
+		{
+			return CurrentValue;
+		}
+		
+		return GetParameterValueFromSystem(Parameter, Component->GetAsset());
 	}
-	
-	return GetParameterValueFromAsset(Parameter, Component);
+	else if(const UNiagaraSystem* System = Cast<UNiagaraSystem>(Owner))
+	{
+		return GetParameterValueFromSystem(Parameter, System);
+	}
+
+	ensure(false);
+	return FNiagaraVariant();
 }
 
 // Proxy class to allow us to override values on the component that are not yet overridden.
 class FNiagaraParameterProxy
 {
 public:
-	FNiagaraParameterProxy(TWeakObjectPtr<UNiagaraComponent> InComponent, const FNiagaraVariableBase& InKey, const FNiagaraVariant& InValue, const FSimpleDelegate& InOnRebuild, TArray<TSharedPtr<IPropertyHandle>> InPropertyHandles)
+	FNiagaraParameterProxy(TWeakObjectPtr<UObject> InOwner, const FNiagaraVariableBase& InKey, const FNiagaraVariant& InValue, const FSimpleDelegate& InOnRebuild, TArray<TSharedPtr<IPropertyHandle>> InPropertyHandles)
 	{
 		bResettingToDefault = false;
-		Component = InComponent;
+		Owner = InOwner;
 		ParameterKey = InKey;
 		ParameterValue = InValue;
 		OnRebuild = InOnRebuild;
@@ -129,35 +119,38 @@ public:
 
 	void OnResetToDefault()
 	{
-		UNiagaraComponent* RawComponent = Component.Get();
-		if (RawComponent != nullptr)
+		if(Owner->IsA<UNiagaraComponent>())
 		{
-			FScopedTransaction ScopedTransaction(NSLOCTEXT("UnrealEd", "PropertyWindowResetToDefault", "Reset to Default"));
-			RawComponent->Modify();
+			TWeakObjectPtr<UNiagaraComponent> RawComponent = GetComponent();
+			if (RawComponent != nullptr)
+			{
+				FScopedTransaction ScopedTransaction(NSLOCTEXT("UnrealEd", "PropertyWindowResetToDefault", "Reset to Default"));
+				RawComponent->Modify();
 
-			bResettingToDefault = true;
+				bResettingToDefault = true;
 
-			for (TSharedPtr<IPropertyHandle> PropertyHandle : PropertyHandles)
-			{
-				PropertyHandle->NotifyPreChange();
-			}
+				for (TSharedPtr<IPropertyHandle> PropertyHandle : PropertyHandles)
+				{
+					PropertyHandle->NotifyPreChange();
+				}
 
-			RawComponent->RemoveParameterOverride(ParameterKey);
+				RawComponent->RemoveParameterOverride(ParameterKey);
 
-			for (TSharedPtr<IPropertyHandle> PropertyHandle : PropertyHandles)
-			{
-				PropertyHandle->NotifyPostChange(EPropertyChangeType::ValueSet);
-			}
+				for (TSharedPtr<IPropertyHandle> PropertyHandle : PropertyHandles)
+				{
+					PropertyHandle->NotifyPostChange(EPropertyChangeType::ValueSet);
+				}
 
-			OnRebuild.ExecuteIfBound();
-			bResettingToDefault = false;
+				OnRebuild.ExecuteIfBound();
+				bResettingToDefault = false;
+			}
 		}
 	}
 
 	EVisibility GetResetToDefaultVisibility() const 
 	{
-		UNiagaraComponent* RawComponent = Component.Get();
-		if (RawComponent != nullptr)
+		TWeakObjectPtr<UNiagaraComponent> RawComponent = GetComponent();
+		if (RawComponent.IsValid())
 		{
 			return RawComponent->HasParameterOverride(ParameterKey) ? EVisibility::Visible : EVisibility::Hidden;
 		}
@@ -166,8 +159,8 @@ public:
 
 	FNiagaraVariant FindExistingOverride() const
 	{
-		UNiagaraComponent* RawComponent = Component.Get();
-		if (RawComponent != nullptr)
+		TWeakObjectPtr<UNiagaraComponent> RawComponent = GetComponent();
+		if (RawComponent.IsValid())
 		{
 			return RawComponent->FindParameterOverride(ParameterKey);
 		}
@@ -176,8 +169,8 @@ public:
 
 	void OnParameterPreChange()
 	{
-		UNiagaraComponent* RawComponent = Component.Get();
-		if (RawComponent != nullptr)
+		TWeakObjectPtr<UNiagaraComponent> RawComponent = GetComponent();
+		if (RawComponent.IsValid())
 		{
 			RawComponent->Modify();
 
@@ -190,38 +183,89 @@ public:
 	
 	void OnParameterChanged(const FPropertyChangedEvent& PropertyChangedEvent)
 	{
-		if (bResettingToDefault)
+		if(Owner->IsA<UNiagaraComponent>())
 		{
-			return;
-		}
+			if (bResettingToDefault)
+			{
+				return;
+			}
 
-		UNiagaraComponent* RawComponent = Component.Get();
-		if (RawComponent != nullptr)
-		{
-			RawComponent->SetParameterOverride(ParameterKey, ParameterValue);
+			TWeakObjectPtr<UNiagaraComponent> RawComponent = GetComponent();
+			if(RawComponent.IsValid())
+			{
+				RawComponent->SetParameterOverride(ParameterKey, ParameterValue);
 
-			for (TSharedPtr<IPropertyHandle> PropertyHandle : PropertyHandles)
+				for (TSharedPtr<IPropertyHandle> PropertyHandle : PropertyHandles)
+				{
+					PropertyHandle->NotifyPostChange(PropertyChangedEvent.ChangeType);
+				}
+			}
+		}
+		else if(Owner->IsA<UNiagaraSystem>())
+		{
+			TWeakObjectPtr<UNiagaraSystem> System = GetSystem();
+			if(System.IsValid())
 			{
-				PropertyHandle->NotifyPostChange(PropertyChangedEvent.ChangeType);
+				if(ParameterValue.GetMode() == ENiagaraVariantMode::Bytes)
+				{
+					System->GetExposedParameters().SetParameterData(ParameterValue.GetBytes(), ParameterKey);
+				}
+				else if(ParameterValue.GetMode() == ENiagaraVariantMode::DataInterface)
+				{
+					System->GetExposedParameters().OnInterfaceChange();
+				}
+				else if(ParameterValue.GetMode() == ENiagaraVariantMode::Object)
+				{
+					System->GetExposedParameters().OnUObjectChange();
+				}
+
+				for (TSharedPtr<IPropertyHandle> PropertyHandle : PropertyHandles)
+				{
+					PropertyHandle->NotifyPostChange(PropertyChangedEvent.ChangeType);
+				}			
 			}
 		}
 	}
 
 	void OnAssetSelectedFromPicker(const FAssetData& InAssetData)
 	{
-		UNiagaraComponent* RawComponent = Component.Get();
-		if (RawComponent != nullptr)
+		if(Owner->IsA<UNiagaraComponent>())
 		{
-			UObject* Asset = InAssetData.GetAsset();
-			if (Asset == nullptr || Asset->GetClass()->IsChildOf(ParameterKey.GetType().GetClass()))
+			TWeakObjectPtr<UNiagaraComponent> RawComponent = GetComponent();
+
+			if(RawComponent.IsValid())
 			{
-				FScopedTransaction ScopedTransaction(LOCTEXT("ChangeAsset", "Change asset"));
-				RawComponent->Modify();
-				RawComponent->SetParameterOverride(ParameterKey, FNiagaraVariant(Asset));
+				UObject* Asset = InAssetData.GetAsset();
+				if (Asset == nullptr || Asset->GetClass()->IsChildOf(ParameterKey.GetType().GetClass()))
+				{
+					FScopedTransaction ScopedTransaction(LOCTEXT("ChangeAsset", "Change asset"));
+					RawComponent->Modify();
+					RawComponent->SetParameterOverride(ParameterKey, FNiagaraVariant(Asset));
 
-				for (TSharedPtr<IPropertyHandle> PropertyHandle : PropertyHandles)
+					for (TSharedPtr<IPropertyHandle> PropertyHandle : PropertyHandles)
+					{
+						PropertyHandle->NotifyPostChange(EPropertyChangeType::ValueSet);
+					}
+				}
+			}
+		}
+		else if(Owner->IsA<UNiagaraSystem>())
+		{
+			TWeakObjectPtr<UNiagaraSystem> System = GetSystem();
+
+			if(System.IsValid())
+			{
+				UObject* Asset = InAssetData.GetAsset();
+				if (Asset == nullptr || Asset->GetClass()->IsChildOf(ParameterKey.GetType().GetClass()))
 				{
-					PropertyHandle->NotifyPostChange(EPropertyChangeType::ValueSet);
+					FScopedTransaction ScopedTransaction(LOCTEXT("ChangeAsset", "Change asset"));
+					System->Modify();
+					System->GetExposedParameters().SetUObject(Asset, ParameterKey);
+
+					for (TSharedPtr<IPropertyHandle> PropertyHandle : PropertyHandles)
+					{
+						PropertyHandle->NotifyPostChange(EPropertyChangeType::ValueSet);
+					}
 				}
 			}
 		}
@@ -231,23 +275,37 @@ public:
 	{
 		UObject* CurrentObject = nullptr;
 
-		UNiagaraComponent* RawComponent = Component.Get();
-		if (RawComponent != nullptr)
+		if(Owner->IsA<UNiagaraComponent>())
 		{
-			FNiagaraVariant CurrentValue = FindExistingOverride();
-			if (CurrentValue.IsValid())
-			{
-				CurrentObject = CurrentValue.GetUObject();
-			}
-			else
+			TWeakObjectPtr<UNiagaraComponent> RawComponent = GetComponent();
+
+			if(RawComponent.IsValid())
 			{
-				// fetch from asset
-				UNiagaraSystem* System = RawComponent->GetAsset();
-				if (System != nullptr)
+				FNiagaraVariant CurrentValue = FindExistingOverride();
+				if (CurrentValue.IsValid())
 				{
-					FNiagaraUserRedirectionParameterStore& AssetParamStore = System->GetExposedParameters();
-					CurrentObject = AssetParamStore.GetUObject(ParameterKey);
+					CurrentObject = CurrentValue.GetUObject();
 				}
+				else
+				{
+					// fetch from asset
+					UNiagaraSystem* System = RawComponent->GetAsset();
+					if (System != nullptr)
+					{
+						FNiagaraUserRedirectionParameterStore& AssetParamStore = System->GetExposedParameters();
+						CurrentObject = AssetParamStore.GetUObject(ParameterKey);
+					}
+				}
+			}
+		}
+		else if(Owner->IsA<UNiagaraSystem>())
+		{			
+			TWeakObjectPtr<UNiagaraSystem> System = GetSystem();
+
+			if(System.IsValid())
+			{
+				FNiagaraUserRedirectionParameterStore& AssetParamStore = System->GetExposedParameters();
+				CurrentObject = AssetParamStore.GetUObject(ParameterKey);			
 			}
 		}
 
@@ -258,10 +316,11 @@ public:
 	const FNiagaraVariableBase& Key() const { return ParameterKey; }
 	FNiagaraVariant& Value() { return ParameterValue; }
 
-	TWeakObjectPtr<UNiagaraComponent> GetComponent() const { return Component; }
+	TWeakObjectPtr<UNiagaraComponent> GetComponent() const { return Cast<UNiagaraComponent>(Owner); }
+	TWeakObjectPtr<UNiagaraSystem> GetSystem() const { return Cast<UNiagaraSystem>(Owner); }
 
 private:
-	TWeakObjectPtr<UNiagaraComponent> Component;
+	TWeakObjectPtr<UObject> Owner;
 	TArray<TSharedPtr<IPropertyHandle>> PropertyHandles;
 	FNiagaraVariableBase ParameterKey;
 	FNiagaraVariant ParameterValue;
@@ -297,8 +356,6 @@ void AddRowForUserParameter(TSharedPtr<FNiagaraParameterProxy> ParameterProxy, I
 	{
 		TArray<UObject*> Objects { ParameterProxy->Value().GetUObject() };
 
-		// How do I set this up so I can have this pick actors from the level?
-
 		FAddPropertyParams Params = FAddPropertyParams()
 			.UniqueId(Parameter.GetName())
 			.AllowChildren(false) // Don't show the material's properties
@@ -352,13 +409,25 @@ void AddRowForUserParameter(TSharedPtr<FNiagaraParameterProxy> ParameterProxy, I
 	FDetailWidgetRow& CustomWidget = Row->CustomWidget(true);
 
 	FNiagaraVariable UserParameter = Parameter;
-	ParameterProxy->GetComponent()->GetAsset()->GetExposedParameters().RedirectUserVariable(UserParameter);
-	UNiagaraScriptVariable* ScriptVariable = FNiagaraEditorUtilities::GetScriptVariableForUserParameter(UserParameter, *ParameterProxy->GetComponent()->GetAsset());
+	UNiagaraSystem* System = nullptr;
+	if(ParameterProxy->GetSystem().IsValid())
+	{
+		System = ParameterProxy->GetSystem().Get();
+	}
+	else if(ParameterProxy->GetComponent().IsValid())
+	{
+		System = ParameterProxy->GetComponent()->GetAsset();
+	}
+	System->GetExposedParameters().RedirectUserVariable(UserParameter);
+	UNiagaraScriptVariable* ScriptVariable = FNiagaraEditorUtilities::GetScriptVariableForUserParameter(UserParameter, *System);
 		
-	FText Tooltip = FText::GetEmpty();
-	if(ScriptVariable && !ScriptVariable->Metadata.Description.IsEmpty())
+	TAttribute<FText> Tooltip = FText::GetEmpty();
+	if(ScriptVariable)
 	{
-		Tooltip = ScriptVariable->Metadata.Description;
+		Tooltip = TAttribute<FText>::CreateLambda([ScriptVariable]
+		{
+			return ScriptVariable->Metadata.Description;
+		});
 	}
 		
 	Row->GetDefaultWidgets(DefaultNameWidget, DefaultValueWidget, CustomWidget);
@@ -422,7 +491,7 @@ void AddRowForUserParameter(TSharedPtr<FNiagaraParameterProxy> ParameterProxy, I
 class FNiagaraUserParameterCategoryBuilder : public IDetailCustomNodeBuilder, public TSharedFromThis<FNiagaraUserParameterCategoryBuilder>
 {
 public:
-	FNiagaraUserParameterCategoryBuilder(UNiagaraComponent* InComponent, const UNiagaraHierarchyCategory* InNiagaraHierarchyCategory, FName InCustomBuilderRowName,
+	FNiagaraUserParameterCategoryBuilder(UObject* InOwner, const UNiagaraHierarchyCategory* InNiagaraHierarchyCategory, FName InCustomBuilderRowName,
 		TArray<TSharedPtr<FNiagaraParameterProxy>>& InParameterProxies, TMap<FName, TWeakPtr<FStructOnScope>>& InParameterNameToDisplayStruct,
 		TArray<TSharedPtr<IPropertyHandle>> InOverridePropertyHandles) : ParameterProxies(InParameterProxies), ParameterNameToDisplayStruct(InParameterNameToDisplayStruct)
 	{
@@ -430,8 +499,7 @@ public:
 		OverridePropertyHandles = InOverridePropertyHandles;
 		CustomBuilderRowName = InCustomBuilderRowName;
 
-		Component = InComponent;
-		bDelegatesInitialized = false;
+		Owner = InOwner;
 	}
 	
 	virtual void GenerateHeaderRowContent(FDetailWidgetRow& NodeRow) override
@@ -441,7 +509,7 @@ public:
 			SNew(STextBlock)
 			.TextStyle(FNiagaraEditorStyle::Get(), "NiagaraEditor.ParameterText")
 			.Text(FText::FromName(CustomBuilderRowName))
-			.ToolTipText(HierarchyCategory->GetTooltip())
+			.ToolTipText_UObject(HierarchyCategory.Get(), &UNiagaraHierarchyCategory::GetTooltip)
 		];
 
 		NodeRow.FilterString(FText::FromName(CustomBuilderRowName));
@@ -481,7 +549,7 @@ public:
 			if(Category->DoesOneChildExist<UNiagaraHierarchyUserParameter>())
 			{
 				ChildrenBuilder.AddCustomBuilder(MakeShared<FNiagaraUserParameterCategoryBuilder>(
-					Component.Get(),Category, Category->GetCategoryName(),
+					Owner.Get(),Category, Category->GetCategoryName(),
 					ParameterProxies, ParameterNameToDisplayStruct,
 					OverridePropertyHandles));
 			}
@@ -490,25 +558,35 @@ public:
 		for(const UNiagaraHierarchyUserParameter* HierarchyUserParameter : ContainedParameters)
 		{
 			FNiagaraVariable Parameter = HierarchyUserParameter->GetUserParameter();
-			FNiagaraVariant ParameterValue = GetCurrentParameterValue(Parameter, Component.Get());
+			FNiagaraVariant ParameterValue = GetCurrentParameterValue(Parameter, Owner.Get());
 			if (!ParameterValue.IsValid())
 			{
 				continue;
 			}
-		
+
+			// the parameter value should operate on a duplicate for components and on the contained DI for a system
 			if (Parameter.IsDataInterface())
 			{
-				ParameterValue = FNiagaraVariant(DuplicateObject(ParameterValue.GetDataInterface(), Component.Get()));
+				if(Owner->IsA<UNiagaraComponent>())
+				{
+					ParameterValue = FNiagaraVariant(DuplicateObject(ParameterValue.GetDataInterface(), GetComponent().Get()));
+				}
+				else if(Owner->IsA<UNiagaraSystem>())
+				{
+					ParameterValue = FNiagaraVariant(ParameterValue.GetDataInterface());
+				}
 			}
 			
-			TSharedPtr<FNiagaraParameterProxy> ParameterProxy = ParameterProxies.Add_GetRef(MakeShareable(new FNiagaraParameterProxy(Component, Parameter, ParameterValue, OnRebuildChildren, OverridePropertyHandles)));
+			TSharedPtr<FNiagaraParameterProxy> ParameterProxy = ParameterProxies.Add_GetRef(MakeShareable(new FNiagaraParameterProxy(Owner, Parameter, ParameterValue, OnRebuildChildren, OverridePropertyHandles)));
 			AddRowForUserParameter(ParameterProxy, ChildrenBuilder, ParameterNameToDisplayStruct);
 		}
 	}
 
+	TWeakObjectPtr<UNiagaraComponent> GetComponent() const { return Cast<UNiagaraComponent>(Owner); }
+	TWeakObjectPtr<UNiagaraSystem> GetSystem() const { return Cast<UNiagaraSystem>(Owner); }
+
 private:
-	bool bDelegatesInitialized;
-	TWeakObjectPtr<UNiagaraComponent> Component;
+	TWeakObjectPtr<UObject> Owner;
 	TWeakObjectPtr<const UNiagaraHierarchyCategory> HierarchyCategory;
 	TArray<TSharedPtr<IPropertyHandle>> OverridePropertyHandles;
 	TArray<TSharedPtr<FNiagaraParameterProxy>>& ParameterProxies;
@@ -517,6 +595,188 @@ private:
 	FName CustomBuilderRowName;
 };
 
+void GenerateUserParameterRows(IDetailChildrenBuilder& ChildrenBuilder, UNiagaraHierarchyRoot& UserParameterHierarchyRoot, UObject* Owner, const UNiagaraHierarchySection*& ActiveSection,
+	TArray<TSharedPtr<FNiagaraParameterProxy>>& ParameterProxies, TMap<FName, TWeakPtr<FStructOnScope>>& ParameterNameToDisplayStruct, TArray<TSharedPtr<IPropertyHandle>> OverridePropertyHandles, FSimpleDelegate OnRebuildChildren)
+{
+	UNiagaraSystem* SystemAsset = nullptr;
+
+	if(UNiagaraSystem* System = Cast<UNiagaraSystem>(Owner))
+	{
+		SystemAsset = System;
+	}
+	else if(UNiagaraComponent* Component = Cast<UNiagaraComponent>(Owner))
+	{
+		SystemAsset = Component->GetAsset();
+	}
+
+	if(SystemAsset == nullptr)
+	{
+		return;
+	}
+	
+	TMap<const UNiagaraHierarchySection*, TArray<const UNiagaraHierarchyCategory*>> SectionCategoryMap;
+
+	UNiagaraHierarchyRoot* Root = &UserParameterHierarchyRoot;
+	for(const UNiagaraHierarchyItemBase* Child : Root->GetChildren())
+	{
+		if(const UNiagaraHierarchyCategory* HierarchyCategory = Cast<UNiagaraHierarchyCategory>(Child))
+		{
+			if(HierarchyCategory->GetSection() != nullptr)
+			{
+				SectionCategoryMap.FindOrAdd(HierarchyCategory->GetSection()).Add(HierarchyCategory);	
+			}
+		}
+	}
+
+	// we create a custom row for user param sections here, but only if we have at least one section specified
+	if(SectionCategoryMap.Num() > 0)
+	{
+		TSharedRef<SHorizontalBox> SectionsBox = SNew(SHorizontalBox);
+		
+		for(auto& SectionCategories : SectionCategoryMap)
+		{
+			const UNiagaraHierarchySection* Section = SectionCategories.Key;
+			TArray<const UNiagaraHierarchyCategory*>& HierarchyCategories = SectionCategories.Value;
+			
+			if(HierarchyCategories.Num() > 0)
+			{
+				bool bDoesUserParamExist = false;
+				for(const UNiagaraHierarchyCategory* Category : HierarchyCategories)
+				{
+					bDoesUserParamExist |= Category->DoesOneChildExist<UNiagaraHierarchyUserParameter>(true);
+
+					if(bDoesUserParamExist)
+					{
+						break;
+					}
+				}
+
+				if(bDoesUserParamExist)
+				{
+					SectionsBox->AddSlot()
+					.AutoWidth()
+					.Padding(FMargin(2.f))
+					[
+						SNew(SBox)
+						.Padding(FMargin(0.f, 4.f, 0.f, 0.f))
+						[
+							SNew(SCheckBox)
+							.Style(FAppStyle::Get(), "DetailsView.SectionButton")
+							.OnCheckStateChanged_Lambda([OnRebuildChildren, &ActiveSection, Section](ECheckBoxState NewState)
+							{
+								ActiveSection = Section;
+								OnRebuildChildren.ExecuteIfBound();
+							})
+							.IsChecked_Lambda([=]()
+							{
+								return ActiveSection == Section ? ECheckBoxState::Checked : ECheckBoxState::Unchecked;
+							})
+							[
+								SNew(STextBlock)
+								.TextStyle(FAppStyle::Get(), "SmallText")
+								.Text_UObject(Section, &UNiagaraHierarchySection::GetSectionNameAsText)
+								.ToolTipText_UObject(Section, &UNiagaraHierarchySection::GetTooltip)
+							]
+						]
+					];
+				}
+			}
+		}
+
+		// if we have at least one custom section, we add a default "All" section
+		if(SectionsBox->GetChildren()->Num() > 0)
+		{
+			SectionsBox->AddSlot()
+			.AutoWidth()
+			.Padding(FMargin(2.f))
+			[
+				SNew(SBox)
+				.Padding(FMargin(0.f, 4.f, 0.f, 0.f))
+				[
+					SNew(SCheckBox)
+					.Style(FAppStyle::Get(), "DetailsView.SectionButton")
+					.OnCheckStateChanged_Lambda([&ActiveSection, OnRebuildChildren](ECheckBoxState NewState)
+					{
+						ActiveSection = nullptr;
+						OnRebuildChildren.ExecuteIfBound();
+					})
+					.IsChecked_Lambda([=]()
+					{
+						return ActiveSection == nullptr ? ECheckBoxState::Checked : ECheckBoxState::Unchecked;
+					})
+					[
+						SNew(STextBlock)
+						.TextStyle(FAppStyle::Get(), "SmallText")
+						.Text(FText::FromString("All"))
+					]
+				]
+			];
+			
+			FDetailWidgetRow& Row = ChildrenBuilder.AddCustomRow(FText::FromString("Parameter Sections"));
+			Row.WholeRowContent()
+			[
+				SectionsBox
+			];
+		}
+	}
+
+	// we add the categories now if the active section is set to all (nullptr) or if the section it belongs to is active
+	for(const UNiagaraHierarchyItemBase* Child : Root->GetChildren())
+	{
+		if(const UNiagaraHierarchyCategory* HierarchyCategory = Cast<UNiagaraHierarchyCategory>(Child))
+		{
+			if(HierarchyCategory->DoesOneChildExist<UNiagaraHierarchyUserParameter>(true) && (ActiveSection == nullptr || HierarchyCategory->GetSection() == ActiveSection))
+			{
+				ChildrenBuilder.AddCustomBuilder(MakeShared<FNiagaraUserParameterCategoryBuilder>(
+					Owner, HierarchyCategory, HierarchyCategory->GetCategoryName(),
+					ParameterProxies, ParameterNameToDisplayStruct,
+					OverridePropertyHandles));
+			}
+		}
+	}
+
+	// at last we add all parameters that haven't been setup at all, if the active section is set to "All"
+	if(ActiveSection == nullptr)
+	{
+		TArray<FNiagaraVariable> LeftoverUserParameters;
+		SystemAsset->GetExposedParameters().GetUserParameters(LeftoverUserParameters);
+
+		TArray<UNiagaraHierarchyUserParameter*> HierarchyUserParameters;
+		Root->GetChildrenOfType<>(HierarchyUserParameters, true);
+
+		TArray<FNiagaraVariable> RedirectedUserParameters;
+		for(FNiagaraVariable& Parameter : LeftoverUserParameters)
+		{
+			SystemAsset->GetExposedParameters().RedirectUserVariable(Parameter);
+		}
+		
+		for(UNiagaraHierarchyUserParameter* HierarchyUserParameter : HierarchyUserParameters)
+		{
+			if(LeftoverUserParameters.Contains(HierarchyUserParameter->GetUserParameter()))
+			{
+				LeftoverUserParameters.Remove(HierarchyUserParameter->GetUserParameter());
+			}
+		}
+
+		for(FNiagaraVariable& LeftoverUserParameter : LeftoverUserParameters)
+		{
+			FNiagaraVariant ParameterValue = GetCurrentParameterValue(LeftoverUserParameter, Owner);
+			if (!ParameterValue.IsValid())
+			{
+				continue;
+			}
+		
+			if (LeftoverUserParameter.IsDataInterface())
+			{
+				ParameterValue = FNiagaraVariant(DuplicateObject(ParameterValue.GetDataInterface(), Owner));
+			}
+			
+			TSharedPtr<FNiagaraParameterProxy> ParameterProxy = ParameterProxies.Add_GetRef(MakeShareable(new FNiagaraParameterProxy(Owner, LeftoverUserParameter, ParameterValue, OnRebuildChildren, OverridePropertyHandles)));
+			AddRowForUserParameter(ParameterProxy, ChildrenBuilder, ParameterNameToDisplayStruct);
+		}
+	}
+}
+
 class FNiagaraComponentNodeBuilder : public IDetailCustomNodeBuilder, public TSharedFromThis<FNiagaraComponentNodeBuilder>
 {
 public:
@@ -527,19 +787,24 @@ public:
 
 		Component = InComponent;
 		bDelegatesInitialized = false;
-
-		//UE_LOG(LogNiagaraEditor, Log, TEXT("FNiagaraComponentNodeBuilder %p Component %p"), this, Component.Get());
 	}
 
-	~FNiagaraComponentNodeBuilder()
+	virtual ~FNiagaraComponentNodeBuilder() override
 	{
 		if (Component.IsValid() && bDelegatesInitialized)
 		{
 			Component->OnSynchronizedWithAssetParameters().RemoveAll(this);
 			Component->GetOverrideParameters().RemoveAllOnChangedHandlers(this);
-		}
 
-		//UE_LOG(LogNiagaraEditor, Log, TEXT("~FNiagaraComponentNodeBuilder %p Component %p"), this, Component.Get());
+			if(UNiagaraSystem* System = Component->GetAsset())
+			{
+				TSharedPtr<FNiagaraSystemViewModel> SystemViewModel = TNiagaraViewModelManager<UNiagaraSystem, FNiagaraSystemViewModel>::GetExistingViewModelForObject(System);
+				if(SystemViewModel.IsValid())
+				{
+					SystemViewModel->GetUserParametersHierarchyViewModel()->OnHierarchyChanged().RemoveAll(this);
+				}
+			}
+		}
 	}
 
 	virtual void SetOnRebuildChildren(FSimpleDelegate InOnRegenerateChildren) override
@@ -547,10 +812,10 @@ public:
 		OnRebuildChildren = InOnRegenerateChildren;
 	}
 
-	virtual void GenerateHeaderRowContent(FDetailWidgetRow& NodeRow) override;
+	virtual void GenerateHeaderRowContent(FDetailWidgetRow& NodeRow) override {}
 	virtual void Tick(float DeltaTime) override {}
 	virtual bool RequiresTick() const override { return false; }
-	virtual bool InitiallyCollapsed() const { return false; }
+	virtual bool InitiallyCollapsed() const override { return false; }
 	virtual FName GetName() const  override
 	{
 		return CustomBuilderRowName;
@@ -560,9 +825,11 @@ public:
 	{
 		if (bDelegatesInitialized == false)
 		{
-			Component->OnSynchronizedWithAssetParameters().AddSP(this, &FNiagaraComponentNodeBuilder::ComponentSynchronizedWithAssetParameters);
+			Component->OnSynchronizedWithAssetParameters().AddSP(this, &FNiagaraComponentNodeBuilder::Rebuild);
 			Component->GetOverrideParameters().AddOnChangedHandler(
 				FNiagaraParameterStore::FOnChanged::FDelegate::CreateSP(this, &FNiagaraComponentNodeBuilder::ParameterValueChanged));
+
+			RegisterRebuildOnHierarchyChanged();
 			bDelegatesInitialized = true;
 		}
 
@@ -582,153 +849,150 @@ public:
 		ParameterProxies.Reserve(UserParameters.Num());
 
 		ParameterNameToDisplayStruct.Empty();
-		TMap<const UNiagaraHierarchySection*, TArray<const UNiagaraHierarchyCategory*>> SectionCategoryMap;
 
-		UNiagaraHierarchyRoot* Root = Cast<UNiagaraSystemEditorData>(SystemAsset->GetEditorData())->UserParameterHierarchy;
-		for(const UNiagaraHierarchyItemBase* Child : Root->GetChildren())
+		GenerateUserParameterRows(ChildrenBuilder, *Cast<UNiagaraSystemEditorData>(SystemAsset->GetEditorData())->UserParameterHierarchy.Get(),
+			Component.Get(), ActiveSection, ParameterProxies, ParameterNameToDisplayStruct, OverridePropertyHandles, OnRebuildChildren);
+	}
+
+private:
+
+	void Rebuild() const
+	{
+		OnRebuildChildren.ExecuteIfBound();
+	}
+
+	void RegisterRebuildOnHierarchyChanged()
+	{
+		if(UNiagaraSystem* Asset = Component->GetAsset())
 		{
-			if(const UNiagaraHierarchyCategory* HierarchyCategory = Cast<UNiagaraHierarchyCategory>(Child))
+			TSharedPtr<FNiagaraSystemViewModel> SystemViewModel = TNiagaraViewModelManager<UNiagaraSystem, FNiagaraSystemViewModel>::GetExistingViewModelForObject(Asset);
+			if(SystemViewModel.IsValid())
 			{
-				if(HierarchyCategory->GetSection() != nullptr)
-				{
-					SectionCategoryMap.FindOrAdd(HierarchyCategory->GetSection()).Add(HierarchyCategory);	
-				}
+				SystemViewModel->GetUserParametersHierarchyViewModel()->OnHierarchyChanged().RemoveAll(this);
+				SystemViewModel->GetUserParametersHierarchyViewModel()->OnHierarchyChanged().Add(UNiagaraHierarchyViewModelBase::FOnHierarchyChanged::FDelegate::CreateSP(this, &FNiagaraComponentNodeBuilder::Rebuild));
 			}
 		}
-
-		// we create a custom row for user param sections here, but only if we have at least one section specified
-		if(SectionCategoryMap.Num() > 0)
+	}
+	
+	void ParameterValueChanged()
+	{
+		if (Component.IsValid())
 		{
-			TSharedRef<SHorizontalBox> SectionsBox = SNew(SHorizontalBox);
-			
-			for(auto& SectionCategories : SectionCategoryMap)
+			const FNiagaraParameterStore& OverrideParameters = Component->GetOverrideParameters();
+			TArray<FNiagaraVariable> UserParameters;
+			OverrideParameters.GetParameters(UserParameters);
+			for (const FNiagaraVariable& UserParameter : UserParameters)
 			{
-				const UNiagaraHierarchySection* Section = SectionCategories.Key;
-				TArray<const UNiagaraHierarchyCategory*>& HierarchyCategories = SectionCategories.Value;
-				
-				if(HierarchyCategories.Num() > 0)
+				if (UserParameter.IsUObject() == false)
 				{
-					bool bDoesUserParamExist = false;
-					for(const UNiagaraHierarchyCategory* Category : HierarchyCategories)
+					TWeakPtr<FStructOnScope>* DisplayStructPtr = ParameterNameToDisplayStruct.Find(UserParameter.GetName());
+					if (DisplayStructPtr != nullptr && DisplayStructPtr->IsValid())
 					{
-						bDoesUserParamExist |= Category->DoesOneChildExist<UNiagaraHierarchyUserParameter>(true);
-
-						if(bDoesUserParamExist)
+						TSharedPtr<FStructOnScope> DisplayStruct = DisplayStructPtr->Pin();
+						if (UserParameter.GetType() == FNiagaraTypeDefinition::GetPositionDef())
 						{
-							break;
+							FMemory::Memcpy(DisplayStruct->GetStructMemory(), OverrideParameters.GetPositionParameterValue(UserParameter.GetName()), UserParameter.GetSizeInBytes());
+						}
+						else
+						{
+							OverrideParameters.CopyParameterData(UserParameter, DisplayStruct->GetStructMemory());
 						}
-					}
-
-					if(bDoesUserParamExist)
-					{
-						SectionsBox->AddSlot()
-						.AutoWidth()
-						.Padding(FMargin(2.f))
-						[
-							SNew(SBox)
-							.Padding(FMargin(0.f, 4.f, 0.f, 0.f))
-							[
-								SNew(SCheckBox)
-								.Style(FAppStyle::Get(), "DetailsView.SectionButton")
-								.OnCheckStateChanged(this, &FNiagaraComponentNodeBuilder::OnActiveSectionChanged, Section)
-								.IsChecked(this, &FNiagaraComponentNodeBuilder::IsSectionActive, Section)
-								[
-									SNew(STextBlock)
-									.TextStyle(FAppStyle::Get(), "SmallText")
-									.Text(Section->GetSectionNameAsText())
-									.ToolTipText(Section->GetTooltip())
-								]
-							]
-						];
 					}
 				}
 			}
-
-			// if we have at least one custom section, we add a default "All" section
-			if(SectionsBox->GetChildren()->Num() > 0)
-			{
-				SectionsBox->AddSlot()
-				.AutoWidth()
-				.Padding(FMargin(2.f))
-				[
-					SNew(SBox)
-					.Padding(FMargin(0.f, 4.f, 0.f, 0.f))
-					[
-						SNew(SCheckBox)
-						.Style(FAppStyle::Get(), "DetailsView.SectionButton")
-						.OnCheckStateChanged(this, &FNiagaraComponentNodeBuilder::OnActiveSectionChanged, (const UNiagaraHierarchySection*) nullptr)
-						.IsChecked(this, &FNiagaraComponentNodeBuilder::IsSectionActive, (const UNiagaraHierarchySection*) nullptr)
-						[
-							SNew(STextBlock)
-							.TextStyle(FAppStyle::Get(), "SmallText")
-							.Text(FText::FromString("All"))
-						]
-					]
-				];
-				
-				FDetailWidgetRow& Row = ChildrenBuilder.AddCustomRow(FText::FromString("Parameter Sections"));
-				Row.WholeRowContent()
-				[
-					SectionsBox
-				];
-			}
 		}
+	}
 
-		// we add the categories now if the active section is set to all (nullptr) or if the section it belongs to is active
-		for(const UNiagaraHierarchyItemBase* Child : Root->GetChildren())
+	void OnActiveSectionChanged(ECheckBoxState State, const UNiagaraHierarchySection* NewSelection)
+	{
+		ActiveSection = NewSelection;
+		OnRebuildChildren.Execute();
+	}
+
+	ECheckBoxState IsSectionActive(const UNiagaraHierarchySection* Section) const
+	{
+		return ActiveSection == Section ? ECheckBoxState::Checked : ECheckBoxState::Unchecked;
+	}
+
+private:
+	bool bDelegatesInitialized;
+	TWeakObjectPtr<UNiagaraComponent> Component;
+	TArray<TSharedPtr<IPropertyHandle>> OverridePropertyHandles;
+	FSimpleDelegate OnRebuildChildren;
+	TArray<TSharedPtr<FNiagaraParameterProxy>> ParameterProxies;
+	TMap<FName, TWeakPtr<FStructOnScope>> ParameterNameToDisplayStruct;
+	FName CustomBuilderRowName;
+	const UNiagaraHierarchySection* ActiveSection = nullptr;
+	TArray<UNiagaraHierarchySection*> AvailableSections;
+};
+
+class FNiagaraSystemUserParameterBuilder : public IDetailCustomNodeBuilder, public TSharedFromThis<FNiagaraSystemUserParameterBuilder>
+{
+public:
+	FNiagaraSystemUserParameterBuilder(UNiagaraSystem* InSystem, FName InCustomBuilderRowName) 
+	{
+		CustomBuilderRowName = InCustomBuilderRowName;
+
+		System = InSystem;
+		bDelegatesInitialized = false;
+	}
+
+	virtual ~FNiagaraSystemUserParameterBuilder() override
+	{
+		if (System.IsValid() && bDelegatesInitialized)
 		{
-			if(const UNiagaraHierarchyCategory* HierarchyCategory = Cast<UNiagaraHierarchyCategory>(Child))
+			System->GetExposedParameters().OnStructureChanged().RemoveAll(this);
+
+			TSharedPtr<FNiagaraSystemViewModel> SystemViewModel = TNiagaraViewModelManager<UNiagaraSystem, FNiagaraSystemViewModel>::GetExistingViewModelForObject(System.Get());
+			if(SystemViewModel.IsValid())
 			{
-				if(HierarchyCategory->DoesOneChildExist<UNiagaraHierarchyUserParameter>(true) && (ActiveSection == nullptr || HierarchyCategory->GetSection() == ActiveSection))
-				{
-					ChildrenBuilder.AddCustomBuilder(MakeShared<FNiagaraUserParameterCategoryBuilder>(
-						Component.Get(), HierarchyCategory, HierarchyCategory->GetCategoryName(),
-						ParameterProxies, ParameterNameToDisplayStruct,
-						OverridePropertyHandles));
-				}
+				SystemViewModel->GetUserParametersHierarchyViewModel()->OnHierarchyChanged().RemoveAll(this);
 			}
 		}
+	}
 
-		// at last we add all parameters that haven't been setup at all, if the active section is set to "All"
-		if(ActiveSection == nullptr)
-		{
-			TArray<FNiagaraVariable> LeftoverUserParameters;
-			SystemAsset->GetExposedParameters().GetUserParameters(LeftoverUserParameters);
+	void Rebuild() const
+	{
+		OnRebuildChildren.ExecuteIfBound();
+	}
+	
+	virtual void SetOnRebuildChildren(FSimpleDelegate InOnRegenerateChildren) override
+	{
+		OnRebuildChildren = InOnRegenerateChildren;
+	}
 
-			TArray<UNiagaraHierarchyUserParameter*> HierarchyUserParameters;
-			Root->GetChildrenOfType<>(HierarchyUserParameters, true);
+	virtual void GenerateHeaderRowContent(FDetailWidgetRow& NodeRow) override {}
+	virtual void Tick(float DeltaTime) override {}
+	virtual bool RequiresTick() const override { return false; }
+	virtual bool InitiallyCollapsed() const override { return false; }
+	virtual FName GetName() const  override
+	{
+		return CustomBuilderRowName;
+	}
 
-			TArray<FNiagaraVariable> RedirectedUserParameters;
-			for(FNiagaraVariable& Parameter : LeftoverUserParameters)
-			{
-				SystemAsset->GetExposedParameters().RedirectUserVariable(Parameter);
-			}
-			
-			for(UNiagaraHierarchyUserParameter* HierarchyUserParameter : HierarchyUserParameters)
-			{
-				if(LeftoverUserParameters.Contains(HierarchyUserParameter->GetUserParameter()))
-				{
-					LeftoverUserParameters.Remove(HierarchyUserParameter->GetUserParameter());
-				}
-			}
+	virtual void GenerateChildContent(IDetailChildrenBuilder& ChildrenBuilder) override
+	{
+		if (bDelegatesInitialized == false)
+		{
+			System->GetExposedParameters().OnStructureChanged().Add(FNiagaraParameterStore::FOnStructureChanged::FDelegate::CreateSP(this, &FNiagaraSystemUserParameterBuilder::Rebuild));
+			TSharedPtr<FNiagaraSystemViewModel> SystemViewModel = TNiagaraViewModelManager<UNiagaraSystem, FNiagaraSystemViewModel>::GetExistingViewModelForObject(System.Get());
+			SystemViewModel->GetUserParametersHierarchyViewModel()->OnHierarchyChanged().Add(UNiagaraHierarchyViewModelBase::FOnHierarchyChanged::FDelegate::CreateSP(this, &FNiagaraSystemUserParameterBuilder::Rebuild));
 
-			for(FNiagaraVariable& LeftoverUserParameter : LeftoverUserParameters)
-			{
-				FNiagaraVariant ParameterValue = GetCurrentParameterValue(LeftoverUserParameter, Component.Get());
-				if (!ParameterValue.IsValid())
-				{
-					continue;
-				}
-			
-				if (LeftoverUserParameter.IsDataInterface())
-				{
-					ParameterValue = FNiagaraVariant(DuplicateObject(ParameterValue.GetDataInterface(), Component.Get()));
-				}
-				
-				TSharedPtr<FNiagaraParameterProxy> ParameterProxy = ParameterProxies.Add_GetRef(MakeShareable(new FNiagaraParameterProxy(Component, LeftoverUserParameter, ParameterValue, OnRebuildChildren, OverridePropertyHandles)));
-				AddRowForUserParameter(ParameterProxy, ChildrenBuilder, ParameterNameToDisplayStruct);
-			}
+			bDelegatesInitialized = true;
 		}
+
+		check(System.IsValid());
+
+		ParameterProxies.Reset();
+		
+		TArray<FNiagaraVariable> UserParameters;
+		System->GetExposedParameters().GetUserParameters(UserParameters);
+		ParameterProxies.Reserve(UserParameters.Num());
+
+		ParameterNameToDisplayStruct.Empty();
+
+		UNiagaraHierarchyRoot* Root = Cast<UNiagaraSystemEditorData>(System->GetEditorData())->UserParameterHierarchy;
+		GenerateUserParameterRows(ChildrenBuilder, *Root, System.Get(), ActiveSection, ParameterProxies, ParameterNameToDisplayStruct, {}, OnRebuildChildren);
 	}
 
 private:
@@ -740,11 +1004,11 @@ private:
 
 	void ParameterValueChanged()
 	{
-		if (Component.IsValid())
+		if (System.IsValid())
 		{
-			const FNiagaraParameterStore& OverrideParameters = Component->GetOverrideParameters();
+			const FNiagaraParameterStore& ExposedParameters = System->GetExposedParameters();
 			TArray<FNiagaraVariable> UserParameters;
-			OverrideParameters.GetParameters(UserParameters);
+			ExposedParameters.GetParameters(UserParameters);
 			for (const FNiagaraVariable& UserParameter : UserParameters)
 			{
 				if (UserParameter.IsUObject() == false)
@@ -755,11 +1019,11 @@ private:
 						TSharedPtr<FStructOnScope> DisplayStruct = DisplayStructPtr->Pin();
 						if (UserParameter.GetType() == FNiagaraTypeDefinition::GetPositionDef())
 						{
-							FMemory::Memcpy(DisplayStruct->GetStructMemory(), OverrideParameters.GetPositionParameterValue(UserParameter.GetName()), UserParameter.GetSizeInBytes());
+							FMemory::Memcpy(DisplayStruct->GetStructMemory(), ExposedParameters.GetPositionParameterValue(UserParameter.GetName()), UserParameter.GetSizeInBytes());
 						}
 						else
 						{
-							OverrideParameters.CopyParameterData(UserParameter, DisplayStruct->GetStructMemory());
+							ExposedParameters.CopyParameterData(UserParameter, DisplayStruct->GetStructMemory());
 						}
 					}
 				}
@@ -780,8 +1044,7 @@ private:
 
 private:
 	bool bDelegatesInitialized;
-	TWeakObjectPtr<UNiagaraComponent> Component;
-	TArray<TSharedPtr<IPropertyHandle>> OverridePropertyHandles;
+	TWeakObjectPtr<UNiagaraSystem> System;
 	FSimpleDelegate OnRebuildChildren;
 	TArray<TSharedPtr<FNiagaraParameterProxy>> ParameterProxies;
 	TMap<FName, TWeakPtr<FStructOnScope>> ParameterNameToDisplayStruct;
@@ -790,19 +1053,11 @@ private:
 	TArray<UNiagaraHierarchySection*> AvailableSections;
 };
 
-void FNiagaraComponentNodeBuilder::GenerateHeaderRowContent(FDetailWidgetRow& NodeRow)
-{
-}
-
 TSharedRef<IDetailCustomization> FNiagaraComponentDetails::MakeInstance()
 {
 	return MakeShareable(new FNiagaraComponentDetails);
 }
 
-FNiagaraComponentDetails::FNiagaraComponentDetails() : Builder(nullptr)
-{
-}
-
 FNiagaraComponentDetails::~FNiagaraComponentDetails()
 {
 	if (GEngine)
@@ -919,8 +1174,7 @@ void FNiagaraComponentDetails::CustomizeDetails(IDetailLayoutBuilder& DetailBuil
 				.Text(LOCTEXT("OverrideParameterMultiselectionUnsupported", "Multiple override parameter sets cannot be edited simultaneously."))
 			];
 	}
-
-
+	
 	IDetailCategoryBuilder& CustomCategory = DetailBuilder.EditCategory(ParamUtilitiesName, LOCTEXT("ParamUtilsCategoryName", "Niagara Utilities"), ECategoryPriority::Important);
 
 	CustomCategory.AddCustomRow(FText::GetEmpty())
@@ -1035,4 +1289,33 @@ FReply FNiagaraComponentDetails::OnDebugSelectedSystem()
 
 	return FReply::Handled();
 }
-#undef LOCTEXT_NAMESPACE
\ No newline at end of file
+
+TSharedRef<IDetailCustomization> FNiagaraSystemUserParameterDetails::MakeInstance()
+{
+	return MakeShared<FNiagaraSystemUserParameterDetails>();
+}
+
+void FNiagaraSystemUserParameterDetails::CustomizeDetails(IDetailLayoutBuilder& DetailBuilder)
+{
+	TArray<TWeakObjectPtr<UObject>> CustomizedObjects;
+	DetailBuilder.GetObjectsBeingCustomized(CustomizedObjects);
+	
+	TArray<FName> Categories;
+	DetailBuilder.GetCategoryNames(Categories);
+
+	for(FName& CategoryName : Categories)
+	{
+		DetailBuilder.HideCategory(CategoryName);
+	}
+	
+	static const FName ParamCategoryName = TEXT("NiagaraSystem_UserParameters");
+
+	if(CustomizedObjects.Num() == 1 && CustomizedObjects[0]->IsA<UNiagaraSystem>())
+	{
+		System = Cast<UNiagaraSystem>(CustomizedObjects[0]);
+		IDetailCategoryBuilder& InputParamCategory = DetailBuilder.EditCategory(ParamCategoryName, LOCTEXT("ParamCategoryName", "User Parameters"), ECategoryPriority::Important);
+		InputParamCategory.AddCustomBuilder(MakeShared<FNiagaraSystemUserParameterBuilder>(System.Get(), ParamCategoryName));
+	}
+}
+
+#undef LOCTEXT_NAMESPACE
diff --git a/Engine/Plugins/FX/Niagara/Source/NiagaraEditor/Private/Customizations/NiagaraComponentDetails.h b/Engine/Plugins/FX/Niagara/Source/NiagaraEditor/Private/Customizations/NiagaraComponentDetails.h
--- a/Engine/Plugins/FX/Niagara/Source/NiagaraEditor/Private/Customizations/NiagaraComponentDetails.h
+++ b/Engine/Plugins/FX/Niagara/Source/NiagaraEditor/Private/Customizations/NiagaraComponentDetails.h
@@ -4,11 +4,7 @@
 
 #include "CoreMinimal.h"
 #include "IDetailCustomization.h"
-#include "EditorUndoClient.h"
-#include "TickableEditorObject.h"
 #include "Input/Reply.h"
-#include "Layout/Visibility.h"
-#include "Misc/NotifyHook.h"
 #include "NiagaraComponent.h"
 
 class IDetailLayoutBuilder;
@@ -22,7 +18,7 @@ class UNiagaraSystem;
 class FNiagaraComponentDetails : public IDetailCustomization
 {
 public:
-	virtual ~FNiagaraComponentDetails();
+	virtual ~FNiagaraComponentDetails() override;
 
 	/** Makes a new instance of this detail layout class for a specific detail view requesting it */
 	static TSharedRef<IDetailCustomization> MakeInstance();
@@ -34,11 +30,20 @@ protected:
 	void OnWorldDestroyed(class UWorld* InWorld);
 	void OnPiEEnd();
 
-	FNiagaraComponentDetails();
-
 	FReply OnResetSelectedSystem();
 	FReply OnDebugSelectedSystem();
 private:
 	TWeakObjectPtr<UNiagaraComponent> Component;
-	IDetailLayoutBuilder* Builder;
+	IDetailLayoutBuilder* Builder = nullptr;
+};
+
+class FNiagaraSystemUserParameterDetails : public IDetailCustomization
+{
+public:
+	static TSharedRef<IDetailCustomization> MakeInstance();
+
+	virtual void CustomizeDetails( IDetailLayoutBuilder& DetailBuilder ) override;
+
+protected:
+	TWeakObjectPtr<UNiagaraSystem> System;	
 };
diff --git a/Engine/Plugins/FX/Niagara/Source/NiagaraEditor/Private/Toolkits/SystemToolkitModes/NiagaraSystemToolkitModeBase.cpp b/Engine/Plugins/FX/Niagara/Source/NiagaraEditor/Private/Toolkits/SystemToolkitModes/NiagaraSystemToolkitModeBase.cpp
--- a/Engine/Plugins/FX/Niagara/Source/NiagaraEditor/Private/Toolkits/SystemToolkitModes/NiagaraSystemToolkitModeBase.cpp
+++ b/Engine/Plugins/FX/Niagara/Source/NiagaraEditor/Private/Toolkits/SystemToolkitModes/NiagaraSystemToolkitModeBase.cpp
@@ -3,6 +3,7 @@
 #include "NiagaraSystemToolkitModeBase.h"
 
 #include "AdvancedPreviewSceneModule.h"
+#include "NiagaraConstants.h"
 #include "NiagaraEditorModule.h"
 #include "NiagaraEditorSettings.h"
 #include "NiagaraEditorStyle.h"
@@ -27,9 +28,12 @@
 #include "Widgets/SNiagaraGeneratedCodeView.h"
 #include "Widgets/SNiagaraHierarchy.h"
 #include "Widgets/Layout/SWidgetSwitcher.h"
+#include "Widgets/Input/SButton.h"
 #include "ViewModels/NiagaraScratchPadScriptViewModel.h"
 #include "Widgets/SNiagaraSelectedObjectsDetails.h"
 #include "NiagaraObjectSelection.h"
+#include "SNiagaraSystemUserParameters.h"
+#include "Customizations/NiagaraComponentDetails.h"
 #include "Styling/StyleColors.h"
 #include "ViewModels/NiagaraSystemEditorDocumentsViewModel.h"
 #include "ViewModels/HierarchyEditor/NiagaraUserParametersHierarchyViewModel.h"
@@ -54,7 +58,8 @@ const FName FNiagaraSystemToolkitModeBase::ScriptStatsTabID(TEXT("NiagaraSystemE
 const FName FNiagaraSystemToolkitModeBase::BakerTabID(TEXT("NiagaraSystemEditor_Baker"));
 const FName FNiagaraSystemToolkitModeBase::VersioningTabID(TEXT("NiagaraSystemEditor_Versioning"));
 const FName FNiagaraSystemToolkitModeBase::ScratchPadScriptsTabID(TEXT("NiagaraSystemEditor_ScratchPadScripts"));
-const FName FNiagaraSystemToolkitModeBase::UserParametersHierarchyTabID(TEXT("NiagaraSystemEditor_UserParameters"));
+const FName FNiagaraSystemToolkitModeBase::UserParametersTabID(TEXT("NiagaraSystemEditor_UserParameters"));
+const FName FNiagaraSystemToolkitModeBase::UserParametersHierarchyTabID(TEXT("NiagaraSystemEditor_UserParametersHierarchy"));
 
 FNiagaraSystemToolkitModeBase::FNiagaraSystemToolkitModeBase(FName InModeName, TWeakPtr<FNiagaraSystemToolkit> InSystemToolkit) : FApplicationMode(InModeName), SystemToolkit(InSystemToolkit), SwitcherIdx(0)
 {
@@ -292,6 +297,10 @@ void FNiagaraSystemToolkitModeBase::RegisterTabFactories(TSharedPtr<FTabManager>
 
 	if(SystemToolkit.Pin()->GetSystemViewModel()->GetEditMode() == ENiagaraSystemViewModelEditMode::SystemAsset)
 	{
+		InTabManager->RegisterTabSpawner(UserParametersTabID, FOnSpawnTab::CreateSP(this, &FNiagaraSystemToolkitModeBase::SpawnTab_UserParameters))
+			.SetDisplayName(LOCTEXT("UserParametersTab", "User Parameters"))
+			.SetGroup(WorkspaceMenuCategory.ToSharedRef());
+		
 		InTabManager->RegisterTabSpawner(UserParametersHierarchyTabID, FOnSpawnTab::CreateSP(this, &FNiagaraSystemToolkitModeBase::SpawnTab_UserParametersHierarchyEditor))
 			.SetDisplayName(LOCTEXT("UserParametersHierarchyTab", "User Parameters Hierarchy"))
 			.SetGroup(WorkspaceMenuCategory.ToSharedRef());
@@ -768,6 +777,19 @@ TSharedRef<SDockTab> FNiagaraSystemToolkitModeBase::SpawnTab_Versioning(const FS
 	return SpawnedTab;
 }
 
+TSharedRef<SDockTab> FNiagaraSystemToolkitModeBase::SpawnTab_UserParameters(const FSpawnTabArgs& Args)
+{
+	check(Args.GetTabId().TabType == UserParametersTabID);
+	
+	TSharedRef<SDockTab> SpawnedTab = SNew(SDockTab)
+		.Label(LOCTEXT("UserParametersTabTitle", "User Parameters"))
+		[
+			SNew(SNiagaraSystemUserParameters, SystemToolkit.Pin()->GetSystemViewModel())
+		];
+
+	return SpawnedTab;
+}
+
 TSharedRef<SDockTab> FNiagaraSystemToolkitModeBase::SpawnTab_UserParametersHierarchyEditor(const FSpawnTabArgs& Args)
 {
 	check(Args.GetTabId().TabType == UserParametersHierarchyTabID);
diff --git a/Engine/Plugins/FX/Niagara/Source/NiagaraEditor/Private/Toolkits/SystemToolkitModes/NiagaraSystemToolkitModeBase.h b/Engine/Plugins/FX/Niagara/Source/NiagaraEditor/Private/Toolkits/SystemToolkitModes/NiagaraSystemToolkitModeBase.h
--- a/Engine/Plugins/FX/Niagara/Source/NiagaraEditor/Private/Toolkits/SystemToolkitModes/NiagaraSystemToolkitModeBase.h
+++ b/Engine/Plugins/FX/Niagara/Source/NiagaraEditor/Private/Toolkits/SystemToolkitModes/NiagaraSystemToolkitModeBase.h
@@ -41,6 +41,7 @@ private:
 	TSharedRef<SDockTab> SpawnTab_ScriptStats(const FSpawnTabArgs& Args);
 	TSharedRef<SDockTab> SpawnTab_Baker(const FSpawnTabArgs& Args);
 	TSharedRef<SDockTab> SpawnTab_Versioning(const FSpawnTabArgs& Args);
+	TSharedRef<SDockTab> SpawnTab_UserParameters(const FSpawnTabArgs& Args);
 	TSharedRef<SDockTab> SpawnTab_UserParametersHierarchyEditor(const FSpawnTabArgs& Args);
 
 protected:
@@ -82,6 +83,7 @@ public:
 	static const FName ScriptStatsTabID;
 	static const FName BakerTabID;
 	static const FName VersioningTabID;
+	static const FName UserParametersTabID;
 	static const FName UserParametersHierarchyTabID;
 };
 
diff --git a/Engine/Plugins/FX/Niagara/Source/NiagaraEditor/Private/Toolkits/SystemToolkitModes/NiagaraSystemToolkitMode_Default.cpp b/Engine/Plugins/FX/Niagara/Source/NiagaraEditor/Private/Toolkits/SystemToolkitModes/NiagaraSystemToolkitMode_Default.cpp
--- a/Engine/Plugins/FX/Niagara/Source/NiagaraEditor/Private/Toolkits/SystemToolkitModes/NiagaraSystemToolkitMode_Default.cpp
+++ b/Engine/Plugins/FX/Niagara/Source/NiagaraEditor/Private/Toolkits/SystemToolkitModes/NiagaraSystemToolkitMode_Default.cpp
@@ -17,7 +17,7 @@
 
 FNiagaraSystemToolkitMode_Default::FNiagaraSystemToolkitMode_Default(TWeakPtr<FNiagaraSystemToolkit> InSystemToolkit) : FNiagaraSystemToolkitModeBase(FNiagaraSystemToolkit::DefaultModeName, InSystemToolkit)
 {
-	TabLayout = FTabManager::NewLayout("Standalone_Niagara_System_Layout_v27")
+	TabLayout = FTabManager::NewLayout("Standalone_Niagara_System_Layout_v28")
 		->AddArea
 		(
 			FTabManager::NewPrimaryArea()->SetOrientation(Orient_Vertical)
@@ -38,6 +38,7 @@ FNiagaraSystemToolkitMode_Default::FNiagaraSystemToolkitMode_Default(TWeakPtr<FN
 						// Left sidebar
 						FTabManager::NewStack()
 						->AddTab(SystemParametersTabID, ETabState::OpenedTab)
+						->AddTab(UserParametersTabID, ETabState::OpenedTab)
 						->AddTab(ScratchPadScriptsTabID, ETabState::OpenedTab)
 						->SetForegroundTab(SystemParametersTabID)
 					)
diff --git a/Engine/Plugins/FX/Niagara/Source/NiagaraEditor/Private/ViewModels/HierarchyEditor/NiagaraHierarchyViewModelBase.cpp b/Engine/Plugins/FX/Niagara/Source/NiagaraEditor/Private/ViewModels/HierarchyEditor/NiagaraHierarchyViewModelBase.cpp
--- a/Engine/Plugins/FX/Niagara/Source/NiagaraEditor/Private/ViewModels/HierarchyEditor/NiagaraHierarchyViewModelBase.cpp
+++ b/Engine/Plugins/FX/Niagara/Source/NiagaraEditor/Private/ViewModels/HierarchyEditor/NiagaraHierarchyViewModelBase.cpp
@@ -503,6 +503,7 @@ void FNiagaraHierarchyItemViewModelBase::DuplicateToThis(TSharedPtr<FNiagaraHier
 	}
 	
 	SyncToData();
+	HierarchyViewModel->OnHierarchyChanged().Broadcast();
 }
 
 void FNiagaraHierarchyItemViewModelBase::ReparentToThis(TSharedPtr<FNiagaraHierarchyItemViewModelBase> ItemToMove, int32 InsertIndex)
@@ -519,6 +520,7 @@ void FNiagaraHierarchyItemViewModelBase::ReparentToThis(TSharedPtr<FNiagaraHiera
 	ItemToMove->Finalize();
 	ItemToMove->Parent.Pin()->SyncToData();
 	SyncToData();
+	HierarchyViewModel->OnHierarchyChanged().Broadcast();
 }
 
 int32 FNiagaraHierarchyItemViewModelBase::FindIndexOfChild(TSharedPtr<FNiagaraHierarchyItemViewModelBase> Child) const
@@ -613,6 +615,7 @@ TSharedPtr<FNiagaraHierarchySectionViewModel> FNiagaraHierarchyRootViewModel::Ad
 	SyncToData();
 	HierarchyViewModel->SetActiveSection(NewSectionViewModel);
 	NewSectionViewModel->RequestRename();
+	
 
 	return NewSectionViewModel;
 }
@@ -675,15 +678,15 @@ FReply FNiagaraHierarchySectionViewModel::OnDroppedOn(const FDragDropEvent& Drag
 
 		TArray<UNiagaraHierarchySection*>& SectionData = HierarchyViewModel->GetHierarchyDataRoot()->GetSectionDataMutable();
 		int32 Count = SectionData.Num();
-		
+
+		bool bDropSucceeded = false;
 		// above constitutes to the left here
 		if(ItemDropZone == EItemDropZone::AboveItem)
 		{
 			SectionData.RemoveAt(DraggedSectionIndex);
 			SectionData.Insert(DraggedSectionData, FMath::Max(IndexOfThis, 0));
 
-			HierarchyViewModel->ForceFullRefresh();
-			return FReply::Handled();
+			bDropSucceeded = true;
 		}
 		else if(ItemDropZone == EItemDropZone::BelowItem)
 		{
@@ -698,7 +701,14 @@ FReply FNiagaraHierarchySectionViewModel::OnDroppedOn(const FDragDropEvent& Drag
 				SectionData.Insert(DraggedSectionData, FMath::Min(IndexOfThis+1, Count));
 			}
 
+			bDropSucceeded = true;
+
+		}
+
+		if(bDropSucceeded)
+		{
 			HierarchyViewModel->ForceFullRefresh();
+			HierarchyViewModel->OnHierarchyChanged().Broadcast();
 			return FReply::Handled();
 		}
 
@@ -731,7 +741,8 @@ FReply FNiagaraHierarchySectionViewModel::OnDroppedOn(const FDragDropEvent& Drag
 			}
 			
 			HierarchyViewModel->RefreshHierarchyView();
-			
+			HierarchyViewModel->OnHierarchyChanged().Broadcast();
+
 			return FReply::Handled();
 		}
 	}
@@ -767,6 +778,7 @@ FReply FNiagaraHierarchyItemViewModel::OnDroppedOn(const FDragDropEvent& DragDro
 	if(DragDropEvent.GetOperation()->IsOfType<FNiagaraHierarchyDragDropOp>())
 	{
 		TSharedPtr<FNiagaraHierarchyDragDropOp> HierarchyDragDropOp = DragDropEvent.GetOperationAs<FNiagaraHierarchyDragDropOp>();
+		bool bDropSucceeded = false;
 		if(ItemDropZone == EItemDropZone::AboveItem)
 		{
 			int32 IndexOfThis = Parent.Pin()->FindIndexOfDataChild(AsShared());
@@ -779,16 +791,19 @@ FReply FNiagaraHierarchyItemViewModel::OnDroppedOn(const FDragDropEvent& DragDro
 			{
 				Parent.Pin()->ReparentToThis(HierarchyDragDropOp->GetDraggedItem().Pin(), FMath::Max(IndexOfThis, 0));
 			}
-			
-			HierarchyViewModel->RefreshHierarchyView();
-			HierarchyViewModel->RefreshSourceView();
-			return FReply::Handled();
+
+			bDropSucceeded = true;
 		}
 		else if(ItemDropZone == EItemDropZone::BelowItem)
 		{
 			int32 IndexOfThis = Parent.Pin()->FindIndexOfDataChild(AsShared());
 			Parent.Pin()->ReparentToThis(HierarchyDragDropOp->GetDraggedItem().Pin(), FMath::Min(IndexOfThis+1, Parent.Pin()->GetChildren().Num()));
 
+			bDropSucceeded = true;
+		}
+
+		if(bDropSucceeded)
+		{
 			HierarchyViewModel->RefreshHierarchyView();
 			HierarchyViewModel->RefreshSourceView();
 			return FReply::Handled();
@@ -886,11 +901,14 @@ void UNiagaraHierarchyViewModelBase::AddCategory() const
 	Category->SetSection(GetActiveSectionData());
 	
 	RefreshHierarchyView();
+
+	OnHierarchyChangedDelegate.Broadcast();
 }
 
 void UNiagaraHierarchyViewModelBase::AddSection() const
 {
 	HierarchyViewModelRoot->AddNewSection();
+	OnHierarchyChangedDelegate.Broadcast();
 }
 
 TSharedPtr<SWidget> FNiagaraHierarchyDragDropOp::GetDefaultDecorator() const
diff --git a/Engine/Plugins/FX/Niagara/Source/NiagaraEditor/Private/ViewModels/HierarchyEditor/NiagaraUserParametersHierarchyViewModel.cpp b/Engine/Plugins/FX/Niagara/Source/NiagaraEditor/Private/ViewModels/HierarchyEditor/NiagaraUserParametersHierarchyViewModel.cpp
--- a/Engine/Plugins/FX/Niagara/Source/NiagaraEditor/Private/ViewModels/HierarchyEditor/NiagaraUserParametersHierarchyViewModel.cpp
+++ b/Engine/Plugins/FX/Niagara/Source/NiagaraEditor/Private/ViewModels/HierarchyEditor/NiagaraUserParametersHierarchyViewModel.cpp
@@ -60,7 +60,7 @@ void UNiagaraUserParametersHierarchyViewModel::FinalizeInternal()
 	}
 }
 
-TSharedRef<SWidget> FNiagaraUserParameterDragDropOp::CreateCustomDecorator() const
+TSharedRef<SWidget> FNiagaraUserParameterHierarchyDragDropOp::CreateCustomDecorator() const
 {
 	return FNiagaraParameterUtilities::GetParameterWidget(GetUserParameter(), true, false);
 }
@@ -135,7 +135,7 @@ TSharedRef<FNiagaraHierarchyDragDropOp> UNiagaraUserParametersHierarchyViewModel
 {
 	if(const UNiagaraHierarchyUserParameter* UserParameter = Cast<UNiagaraHierarchyUserParameter>(Item->GetData()))
 	{
-		TSharedRef<FNiagaraUserParameterDragDropOp> ParameterDragDropOp = MakeShared<FNiagaraUserParameterDragDropOp>(Item);
+		TSharedRef<FNiagaraUserParameterHierarchyDragDropOp> ParameterDragDropOp = MakeShared<FNiagaraUserParameterHierarchyDragDropOp>(Item);
 		ParameterDragDropOp->Construct();
 		return ParameterDragDropOp;
 	}
diff --git a/Engine/Plugins/FX/Niagara/Source/NiagaraEditor/Private/ViewModels/NiagaraParameterPanelViewModel.cpp b/Engine/Plugins/FX/Niagara/Source/NiagaraEditor/Private/ViewModels/NiagaraParameterPanelViewModel.cpp
--- a/Engine/Plugins/FX/Niagara/Source/NiagaraEditor/Private/ViewModels/NiagaraParameterPanelViewModel.cpp
+++ b/Engine/Plugins/FX/Niagara/Source/NiagaraEditor/Private/ViewModels/NiagaraParameterPanelViewModel.cpp
@@ -1468,28 +1468,6 @@ TSharedPtr<SWidget> FNiagaraSystemToolkitParameterPanelViewModel::CreateContextM
 	return SNullWidget::NullWidget;
 }
 
-TSharedRef<SWidget> FNiagaraSystemToolkitParameterPanelViewModel::GenerateAdjacentWidget()
-{
-	if(SystemViewModel->GetEditMode() == ENiagaraSystemViewModelEditMode::SystemAsset)
-	{
-		return SNew(SButton)
-		.Visibility_Lambda([this]()
-		{
-			return IsUserSectionActive() ? EVisibility::Visible : EVisibility::Collapsed;
-		})
-		.OnClicked(FOnClicked::CreateLambda([this]()
-		{
-			SummonHierarchyView();
-			return FReply::Handled();
-		}))
-		[
-			SNew(STextBlock).Text(FText::FromString("Edit Hierarchy"))
-		];
-	}
-
-	return SNullWidget::NullWidget;
-}
-
 FNiagaraParameterUtilities::EParameterContext FNiagaraSystemToolkitParameterPanelViewModel::GetParameterContext() const
 {
 	return FNiagaraParameterUtilities::EParameterContext::System;
@@ -1800,11 +1778,6 @@ TSharedRef<SWidget> FNiagaraSystemToolkitParameterPanelViewModel::CreateAddParam
 	return MenuWidget;
 }
 
-void FNiagaraSystemToolkitParameterPanelViewModel::SummonHierarchyView() const
-{
-	SystemViewModel->FocusTab("NiagaraSystemEditor_UserParameters");
-}
-
 TArray<FNiagaraVariable> FNiagaraSystemToolkitParameterPanelViewModel::GetEditableStaticSwitchParameters() const
 {
 	TArray<FNiagaraVariable> OutStaticSwitchParameters;
diff --git a/Engine/Plugins/FX/Niagara/Source/NiagaraEditor/Private/Widgets/SNiagaraHierarchy.cpp b/Engine/Plugins/FX/Niagara/Source/NiagaraEditor/Private/Widgets/SNiagaraHierarchy.cpp
--- a/Engine/Plugins/FX/Niagara/Source/NiagaraEditor/Private/Widgets/SNiagaraHierarchy.cpp
+++ b/Engine/Plugins/FX/Niagara/Source/NiagaraEditor/Private/Widgets/SNiagaraHierarchy.cpp
@@ -667,6 +667,7 @@ void SNiagaraHierarchy::DeleteSelectedItems()
 
 	RefreshSourceView();
 	RefreshHierarchyView();
+	HierarchyViewModel->OnHierarchyChanged().Broadcast();
 }
 
 bool SNiagaraHierarchy::CanDeleteSelectedItems() const
diff --git a/Engine/Plugins/FX/Niagara/Source/NiagaraEditor/Private/Widgets/SNiagaraSystemUserParameters.cpp b/Engine/Plugins/FX/Niagara/Source/NiagaraEditor/Private/Widgets/SNiagaraSystemUserParameters.cpp
new file mode 100644
--- /dev/null
+++ b/Engine/Plugins/FX/Niagara/Source/NiagaraEditor/Private/Widgets/SNiagaraSystemUserParameters.cpp
@@ -0,0 +1,108 @@
+// Copyright Epic Games, Inc. All Rights Reserved.
+
+#include "Widgets/SNiagaraSystemUserParameters.h"
+
+#include "NiagaraConstants.h"
+#include "NiagaraEditorData.h"
+#include "NiagaraEditorUtilities.h"
+#include "Customizations/NiagaraComponentDetails.h"
+#include "Styling/StyleColors.h"
+#include "SystemToolkitModes/NiagaraSystemToolkitModeBase.h"
+#include "Modules/ModuleManager.h"
+#include "ViewModels/NiagaraParameterPanelViewModel.h"
+#include "Widgets/SNiagaraParameterMenu.h"
+#include "Widgets/Input/SComboBox.h"
+#include "Widgets/Input/SButton.h"
+#include "Widgets/Input/SEditableTextBox.h"
+
+#define LOCTEXT_NAMESPACE "NiagaraSystemUserParameters"
+
+void SNiagaraSystemUserParameters::Construct(const FArguments& InArgs, TSharedPtr<FNiagaraSystemViewModel> InSystemViewModel)
+{
+	SystemViewModel = InSystemViewModel;
+	
+	FPropertyEditorModule& PropertyEditorModule = FModuleManager::Get().LoadModuleChecked<FPropertyEditorModule>("PropertyEditor");
+
+	FDetailsViewArgs DetailsViewArgs;
+	DetailsViewArgs.NameAreaSettings = FDetailsViewArgs::ENameAreaSettings::HideNameArea;
+	TSharedRef<IDetailsView> ObjectDetailsView = PropertyEditorModule.CreateDetailView(DetailsViewArgs);
+	ObjectDetailsView->RegisterInstancedCustomPropertyLayout(UNiagaraSystem::StaticClass(), FOnGetDetailCustomizationInstance::CreateStatic(&FNiagaraSystemUserParameterDetails::MakeInstance));
+	ObjectDetailsView->SetObject(&SystemViewModel->GetSystem());
+	
+	ChildSlot
+	[
+		SNew(SVerticalBox)
+		+ SVerticalBox::Slot()
+		.AutoHeight()
+		.HAlign(HAlign_Right)
+		.Padding(2.f)
+		[
+			SNew(SHorizontalBox)
+			+ SHorizontalBox::Slot()
+			.AutoWidth()
+			[
+				SAssignNew(AddParameterButton, SComboButton)
+				.OnGetMenuContent(this, &SNiagaraSystemUserParameters::GetParameterMenu)
+				.OnComboBoxOpened_Lambda([this]
+				{
+					AddParameterButton->SetMenuContentWidgetToFocus(ParameterPanel->GetSearchBox());
+				})
+				.ButtonContent()
+				[
+					SNew(STextBlock)
+					.Text(LOCTEXT("AddUserParameterLabel", "Add Parameter"))
+				]
+			]
+			+ SHorizontalBox::Slot()
+			.AutoWidth()
+			[
+				SNew(SButton)
+				.OnClicked(this, &SNiagaraSystemUserParameters::SummonHierarchyEditor)
+				.Text(LOCTEXT("SummonUserParametersHierarchyButtonLabel", "Edit Hierarchy"))
+			]
+		]
+		+ SVerticalBox::Slot()
+		[
+			SNew(SBox)
+			.AddMetaData<FTagMetaData>(FTagMetaData(TEXT("UserParameters")))
+			[
+				ObjectDetailsView
+			]
+		]
+	];
+}
+
+TSharedRef<SWidget> SNiagaraSystemUserParameters::GetParameterMenu() 
+{
+	FNiagaraParameterPanelCategory UserCategory = GetDefault<UNiagaraEditorSettings>()->GetMetaDataForNamespaces({FNiagaraConstants::UserNamespace});
+	ParameterPanel = SNew(SNiagaraAddParameterFromPanelMenu)
+		.Graphs(SystemViewModel->GetParameterPanelViewModel()->GetEditableGraphsConst())
+		.OnNewParameterRequested(this, &SNiagaraSystemUserParameters::AddParameter)
+		.OnAllowMakeType(this, &SNiagaraSystemUserParameters::CanMakeNewParameterOfType)
+		.NamespaceId(UserCategory.NamespaceMetaData.GetGuid())
+		.ShowNamespaceCategory(false)
+		.ShowGraphParameters(false)
+		.AutoExpandMenu(false);
+	
+	return ParameterPanel.ToSharedRef();
+}
+
+void SNiagaraSystemUserParameters::AddParameter(FNiagaraVariable NewParameter) const
+{
+	FNiagaraParameterPanelCategory UserCategory = GetDefault<UNiagaraEditorSettings>()->GetMetaDataForNamespaces({FNiagaraConstants::UserNamespace});
+	// TODO (ME) Change request rename to true when it's supported in the UI
+	SystemViewModel->GetParameterPanelViewModel()->AddParameter(NewParameter, UserCategory, false, true);
+}
+
+bool SNiagaraSystemUserParameters::CanMakeNewParameterOfType(const FNiagaraTypeDefinition& InType) const
+{
+	return SystemViewModel->GetParameterPanelViewModel()->CanMakeNewParameterOfType(InType);
+}
+
+FReply SNiagaraSystemUserParameters::SummonHierarchyEditor()
+{
+	SystemViewModel->FocusTab(FNiagaraSystemToolkitModeBase::UserParametersHierarchyTabID);
+	return FReply::Handled();
+}
+
+#undef LOCTEXT_NAMESPACE
diff --git a/Engine/Plugins/FX/Niagara/Source/NiagaraEditor/Private/Widgets/SNiagaraSystemUserParameters.h b/Engine/Plugins/FX/Niagara/Source/NiagaraEditor/Private/Widgets/SNiagaraSystemUserParameters.h
new file mode 100644
--- /dev/null
+++ b/Engine/Plugins/FX/Niagara/Source/NiagaraEditor/Private/Widgets/SNiagaraSystemUserParameters.h
@@ -0,0 +1,38 @@
+// Copyright Epic Games, Inc. All Rights Reserved.
+
+#pragma once
+
+#include "CoreMinimal.h"
+#include "ViewModels/NiagaraParameterPanelViewModel.h"
+#include "ViewModels/NiagaraSystemViewModel.h"
+#include "Widgets/SCompoundWidget.h"
+#include "Widgets/DeclarativeSyntaxSupport.h"
+#include "Widgets/SNiagaraParameterMenu.h"
+
+class SNiagaraSystemUserParameters : public SCompoundWidget
+{
+public:
+	DECLARE_DELEGATE_OneParam(FOnThumbnailCaptured, UTexture2D*);
+
+public:
+	SLATE_BEGIN_ARGS( SNiagaraSystemUserParameters ){}
+	SLATE_END_ARGS()
+	
+	void Construct(const FArguments& InArgs, TSharedPtr<FNiagaraSystemViewModel> InSystemViewModel);
+	virtual ~SNiagaraSystemUserParameters() override {}
+	
+	virtual void Tick( const FGeometry& AllottedGeometry, const double InCurrentTime, const float InDeltaTime ) override {}
+
+private:
+	TSharedRef<SWidget> GetParameterMenu();
+	
+	void AddParameter(FNiagaraVariable NewParameter) const;
+	bool CanMakeNewParameterOfType(const FNiagaraTypeDefinition& InType) const;
+
+	FReply SummonHierarchyEditor();
+private:
+	TSharedPtr<FNiagaraSystemViewModel> SystemViewModel;
+	
+	TSharedPtr<SComboButton> AddParameterButton;
+	TSharedPtr<SNiagaraAddParameterFromPanelMenu> ParameterPanel;
+};
diff --git a/Engine/Plugins/FX/Niagara/Source/NiagaraEditor/Public/ViewModels/HierarchyEditor/NiagaraHierarchyViewModelBase.h b/Engine/Plugins/FX/Niagara/Source/NiagaraEditor/Public/ViewModels/HierarchyEditor/NiagaraHierarchyViewModelBase.h
--- a/Engine/Plugins/FX/Niagara/Source/NiagaraEditor/Public/ViewModels/HierarchyEditor/NiagaraHierarchyViewModelBase.h
+++ b/Engine/Plugins/FX/Niagara/Source/NiagaraEditor/Public/ViewModels/HierarchyEditor/NiagaraHierarchyViewModelBase.h
@@ -238,6 +238,7 @@ UCLASS(Abstract)
 class UNiagaraHierarchyViewModelBase : public UObject, public FEditorUndoClient
 {
 public:
+	DECLARE_MULTICAST_DELEGATE(FOnHierarchyChanged)
 	DECLARE_DELEGATE_OneParam(FRefreshTreeWidget, bool bFullRefresh)
 	DECLARE_DELEGATE_RetVal_OneParam(bool, FIsItemSelected, TSharedPtr<FNiagaraHierarchyItemViewModelBase> Item)
 	DECLARE_DELEGATE_OneParam(FSelectObjectInDetailsPanel, UObject* Object)
@@ -299,6 +300,9 @@ public:
 	FSelectObjectInDetailsPanel& OnSelectObjectInDetailsPanel() { return SelectObjectInDetailsPanelDelegate; }
 	FIsItemSelected& IsItemSelected() { return IsItemSelectedDelegate; }
 
+	// Delegates for external systems
+	FOnHierarchyChanged& OnHierarchyChanged() { return OnHierarchyChangedDelegate; } 
+	
 	// Sections
 	void SetActiveSection(TSharedPtr<struct FNiagaraHierarchySectionViewModel>);
 	TSharedPtr<FNiagaraHierarchySectionViewModel> GetActiveSection() const;
@@ -333,6 +337,8 @@ protected:
 	FSimpleDelegate RefreshSectionsWidgetDelegate;
 	FIsItemSelected IsItemSelectedDelegate;
 	FSelectObjectInDetailsPanel SelectObjectInDetailsPanelDelegate;
+
+	FOnHierarchyChanged OnHierarchyChangedDelegate;
 };
 
 struct FNiagaraHierarchyItemViewModelBase : TSharedFromThis<FNiagaraHierarchyItemViewModelBase>
diff --git a/Engine/Plugins/FX/Niagara/Source/NiagaraEditor/Public/ViewModels/HierarchyEditor/NiagaraUserParametersHierarchyViewModel.h b/Engine/Plugins/FX/Niagara/Source/NiagaraEditor/Public/ViewModels/HierarchyEditor/NiagaraUserParametersHierarchyViewModel.h
--- a/Engine/Plugins/FX/Niagara/Source/NiagaraEditor/Public/ViewModels/HierarchyEditor/NiagaraUserParametersHierarchyViewModel.h
+++ b/Engine/Plugins/FX/Niagara/Source/NiagaraEditor/Public/ViewModels/HierarchyEditor/NiagaraUserParametersHierarchyViewModel.h
@@ -63,12 +63,12 @@ private:
 	TWeakPtr<FNiagaraSystemViewModel> SystemViewModelWeak;
 };
 
-class FNiagaraUserParameterDragDropOp : public FNiagaraHierarchyDragDropOp
+class FNiagaraUserParameterHierarchyDragDropOp : public FNiagaraHierarchyDragDropOp
 {
 public:
 	DRAG_DROP_OPERATOR_TYPE(FNiagaraUserParameterDragDropOp, FNiagaraHierarchyDragDropOp)
 
-	FNiagaraUserParameterDragDropOp(TSharedPtr<FNiagaraHierarchyItemViewModelBase> UserParameterItem) : FNiagaraHierarchyDragDropOp(UserParameterItem) {}
+	FNiagaraUserParameterHierarchyDragDropOp(TSharedPtr<FNiagaraHierarchyItemViewModelBase> UserParameterItem) : FNiagaraHierarchyDragDropOp(UserParameterItem) {}
 
 	FNiagaraVariable GetUserParameter() const
 	{
diff --git a/Engine/Plugins/FX/Niagara/Source/NiagaraEditor/Public/ViewModels/NiagaraParameterPanelViewModel.h b/Engine/Plugins/FX/Niagara/Source/NiagaraEditor/Public/ViewModels/NiagaraParameterPanelViewModel.h
--- a/Engine/Plugins/FX/Niagara/Source/NiagaraEditor/Public/ViewModels/NiagaraParameterPanelViewModel.h
+++ b/Engine/Plugins/FX/Niagara/Source/NiagaraEditor/Public/ViewModels/NiagaraParameterPanelViewModel.h
@@ -252,10 +252,6 @@ public:
 	virtual void OnParameterItemActivated(const FNiagaraParameterPanelItem& ActivatedItem) const;
 
 	const TArray<FNiagaraParameterPanelItem>& GetCachedViewedParameterItems() const;
-
-	virtual bool CanSummonHierarchyView() const { return false; };
-
-	virtual void SummonHierarchyView() const {}
 	
 	void SelectParameterItemByName(const FName ParameterName, const bool bRequestRename);
 
@@ -287,10 +283,8 @@ public:
 
 	virtual TSharedPtr<FNiagaraObjectSelection> GetVariableObjectSelection() { return TSharedPtr<FNiagaraObjectSelection>(); };
 
-protected:
 	static bool CanMakeNewParameterOfType(const FNiagaraTypeDefinition& InType);
-
-
+	
 protected:
 	TWeakPtr<INiagaraParameterPanelViewModel> MasterParameterPanelViewModel;
 	FOnParameterPanelViewModelExternalSelectionChanged OnParameterPanelViewModelExternalSelectionChangedDelegate;
@@ -350,8 +344,6 @@ public:
 	virtual FReply OnParameterItemsDragged(const TArray<FNiagaraParameterPanelItem>& DraggedItems, const FPointerEvent& MouseEvent) const override;
 
 	virtual TSharedPtr<SWidget> CreateContextMenuForItems(const TArray<FNiagaraParameterPanelItem>& Items, const TSharedPtr<FUICommandList>& ToolkitCommands) override;
-
-	virtual TSharedRef<SWidget> GenerateAdjacentWidget() override;
 	
 	virtual FNiagaraParameterUtilities::EParameterContext GetParameterContext() const override;
 
@@ -398,9 +390,6 @@ public:
 	virtual bool UsesCategoryFilteringForInitialExpansion() const override { return true; }
 
 	virtual void PreSectionChange(const TArray<FNiagaraParameterPanelCategory>& ExpandedItems) override;
-	virtual bool CanSummonHierarchyView() const override { return true; }
-
-	virtual void SummonHierarchyView() const override;
 private:
 
 	bool  ShouldRouteThroughScratchParameterMap(const FNiagaraParameterPanelCategory* Category = nullptr, const FNiagaraVariableBase* Variable = nullptr);
