- Remove PostEditChange to fix issue with callbacks triggering

#rnx

diff --git a/Engine/Plugins/FX/Niagara/Source/Niagara/Private/NiagaraSystem.cpp b/Engine/Plugins/FX/Niagara/Source/Niagara/Private/NiagaraSystem.cpp
--- a/Engine/Plugins/FX/Niagara/Source/Niagara/Private/NiagaraSystem.cpp
+++ b/Engine/Plugins/FX/Niagara/Source/Niagara/Private/NiagaraSystem.cpp
@@ -3742,7 +3742,7 @@ UNiagaraBakerSettings* UNiagaraSystem::GetBakerSettings()
 	{
 		Modify();
 		BakerSettings = NewObject<UNiagaraBakerSettings>(this, "BakerSettings", RF_Transactional);
-		PostEditChange();
+		//PostEditChange();
 	}
 	return BakerSettings;
 }
