- Preskinned local bounds for static & skeletal mesh DI

diff --git a/Engine/Plugins/FX/Niagara/Shaders/Private/NiagaraDataInterfaceSkeletalMeshTemplate.ush b/Engine/Plugins/FX/Niagara/Shaders/Private/NiagaraDataInterfaceSkeletalMeshTemplate.ush
--- a/Engine/Plugins/FX/Niagara/Shaders/Private/NiagaraDataInterfaceSkeletalMeshTemplate.ush
+++ b/Engine/Plugins/FX/Niagara/Shaders/Private/NiagaraDataInterfaceSkeletalMeshTemplate.ush
@@ -46,6 +46,8 @@ float4x4		{ParameterName}_InstancePrevTransform;
 float4			{ParameterName}_InstanceRotation;
 float4			{ParameterName}_InstancePrevRotation;
 float			{ParameterName}_InstanceInvDeltaTime;
+float3			{ParameterName}_PreSkinnedLocalBoundsCenter;
+float3			{ParameterName}_PreSkinnedLocalBoundsExtents;
 uint			{ParameterName}_EnabledFeatures;
 
 ////////////////////////////////////////////////////////////////////////////////////////////////////////////////////////////////////////////////////////////////
@@ -55,6 +57,17 @@ bool UniformSamplingRegionEnabled_{ParameterName}() { return ({ParameterName}_En
 bool UnlimitedBoneInfluences_{ParameterName}() { return ({ParameterName}_EnabledFeatures & 0x4) == 4; }
 bool HasMeshColors_{ParameterName}() { return ({ParameterName}_EnabledFeatures & 0x8) == 8; }
 
+////////////////////////////////////////////////////////////////////////////////////////////////////////////////////////////////////////////////////////////////
+// Misc functions
+void GetPreSkinnedLocalBounds_{ParameterName}(out float3 OutCenter, out float3 OutExtentsMin, out float3 OutExtentsMax, out float3 OutExtents, out float3 OutHalfExtents)
+{
+	OutCenter		= {ParameterName}_PreSkinnedLocalBoundsCenter;
+	OutExtents		= {ParameterName}_PreSkinnedLocalBoundsExtents;
+	OutHalfExtents	= OutExtents * 0.5f;
+	OutExtentsMin	= OutCenter - OutHalfExtents;
+	OutExtentsMax	= OutCenter + OutHalfExtents;
+}
+
 ////////////////////////////////////////////////////////////////////////////////////////////////////////////////////////////////////////////////////////////////
 // Triangle Sampling
 float3x4 GetPrevBoneSkinningMatrix_{ParameterName}(uint Bone)
diff --git a/Engine/Plugins/FX/Niagara/Shaders/Private/NiagaraDataInterfaceStaticMeshTemplate.ush b/Engine/Plugins/FX/Niagara/Shaders/Private/NiagaraDataInterfaceStaticMeshTemplate.ush
--- a/Engine/Plugins/FX/Niagara/Shaders/Private/NiagaraDataInterfaceStaticMeshTemplate.ush
+++ b/Engine/Plugins/FX/Niagara/Shaders/Private/NiagaraDataInterfaceStaticMeshTemplate.ush
@@ -30,6 +30,8 @@ float4x4		{ParameterName}_InstancePreviousTransformInverseTransposed;
 float4			{ParameterName}_InstancePreviousRotation;
 float3			{ParameterName}_InstanceWorldVelocity;
 
+float3			{ParameterName}_PreSkinnedLocalBoundsCenter;
+float3			{ParameterName}_PreSkinnedLocalBoundsExtents;
 float3			{ParameterName}_SystemLWCTile;
 
 #if DISTATICMESH_ALLOWDISTANCEFIELD
@@ -47,7 +49,8 @@ void RandomFilteredSection_{ParameterName}(NiagaraRandInfo RandInfo, out int Sec
 void RandomUnfilteredSection_{ParameterName}(NiagaraRandInfo RandInfo, out int Section);
 void RandomSectionTriangle_{ParameterName}(NiagaraRandInfo RandInfo, int Section, out int Tri, out float3 BaryCoord);
 
-//////////////////////////////////////////////////////////////////////////
+////////////////////////////////////////////////////////////////////////////////////////////////////////////////////////////////////////////////////////////////
+// Misc functions
 int UniformTriangle_{ParameterName}(NiagaraRandInfo RandInfo, int Tri, int SectionOffset)
 {
 	uint2 ProbAlias = {ParameterName}_UniformSamplingTriangles[Tri];
@@ -56,6 +59,15 @@ int UniformTriangle_{ParameterName}(NiagaraRandInfo RandInfo, int Tri, int Secti
 	return NiagaraRandomFloat(RandInfo) < Prob ? Tri : Alias + SectionOffset;
 }
 
+void GetPreSkinnedLocalBounds_{ParameterName}(out float3 OutCenter, out float3 OutExtentsMin, out float3 OutExtentsMax, out float3 OutExtents, out float3 OutHalfExtents)
+{
+	OutCenter		= {ParameterName}_PreSkinnedLocalBoundsCenter;
+	OutExtents		= {ParameterName}_PreSkinnedLocalBoundsExtents;
+	OutHalfExtents	= OutExtents * 0.5f;
+	OutExtentsMin	= OutCenter - OutHalfExtents;
+	OutExtentsMax	= OutCenter + OutHalfExtents;
+}
+
 //////////////////////////////////////////////////////////////////////////
 // Vertex Sampling
 void IsValidVertex_{ParameterName}(out bool IsValid)
diff --git a/Engine/Plugins/FX/Niagara/Source/Niagara/Classes/NiagaraDataInterfaceSkeletalMesh.h b/Engine/Plugins/FX/Niagara/Source/Niagara/Classes/NiagaraDataInterfaceSkeletalMesh.h
--- a/Engine/Plugins/FX/Niagara/Source/Niagara/Classes/NiagaraDataInterfaceSkeletalMesh.h
+++ b/Engine/Plugins/FX/Niagara/Source/Niagara/Classes/NiagaraDataInterfaceSkeletalMesh.h
@@ -584,6 +584,10 @@ struct FNDISkeletalMesh_InstanceData
 	/** Time separating Transform and PrevTransform. */
 	float DeltaSeconds;
 
+	/** Preskinned local bounds pulled from attached skeletal mesh */
+	FVector3f PreSkinnedLocalBoundsCenter = FVector3f::ZeroVector;
+	FVector3f PreSkinnedLocalBoundsExtents = FVector3f::ZeroVector;
+
 	/* Excluded bone for some specific functions, generally the root bone which you don't want to include when picking a random bone. */
 	int32 ExcludedBoneIndex = INDEX_NONE;
 
@@ -809,51 +813,6 @@ public:
 
 	virtual void ProvidePerInstanceDataForRenderThread(void* DataForRenderThread, void* PerInstanceData, const FNiagaraSystemInstanceID& SystemInstance) override;
 
-	static const FString MeshIndexBufferName;
-	static const FString MeshVertexBufferName;
-	static const FString MeshSkinWeightBufferName;
-	static const FString MeshSkinWeightLookupBufferName;
-	static const FString MeshCurrBonesBufferName;
-	static const FString MeshPrevBonesBufferName;
-	static const FString MeshCurrSamplingBonesBufferName;
-	static const FString MeshPrevSamplingBonesBufferName;
-	static const FString MeshTangentBufferName;
-	static const FString MeshTexCoordBufferName;
-	static const FString MeshColorBufferName;
-	static const FString MeshTriangleSamplerProbAliasBufferName;
-	static const FString MeshNumSamplingRegionTrianglesName;
-	static const FString MeshNumSamplingRegionVerticesName;
-	static const FString MeshSamplingRegionsProbAliasBufferName;
-	static const FString MeshSampleRegionsTriangleIndicesName;
-	static const FString MeshSampleRegionsVerticesName;
-	static const FString MeshTriangleMatricesOffsetBufferName;
-	static const FString MeshTriangleCountName;
-	static const FString MeshVertexCountName;
-	static const FString MeshWeightStrideName;
-	static const FString MeshSkinWeightIndexSizeName;
-	static const FString MeshNumTexCoordName;
-	static const FString MeshNumWeightsName;
-	static const FString NumBonesName;
-	static const FString NumFilteredBonesName;
-	static const FString NumUnfilteredBonesName;
-	static const FString RandomMaxBoneName;
-	static const FString ExcludeBoneIndexName;
-	static const FString FilteredAndUnfilteredBonesName;
-	static const FString NumFilteredSocketsName;
-	static const FString FilteredSocketBoneOffsetName;
-	static const FString UvMappingBufferName;
-	static const FString UvMappingBufferLengthName;
-	static const FString UvMappingSetName;
-	static const FString ConnectivityBufferName;
-	static const FString ConnectivityBufferLengthName;
-	static const FString ConnectivityMaxAdjacentPerVertexName;
-	static const FString InstanceTransformName;
-	static const FString InstancePrevTransformName;
-	static const FString InstanceRotationName;
-	static const FString InstancePrevRotationName;
-	static const FString InstanceInvDeltaTimeName;
-	static const FString EnabledFeaturesName;
-
 protected:
 	virtual bool CopyToInternal(UNiagaraDataInterface* Destination) const override;
 
@@ -864,6 +823,10 @@ protected:
 	UFUNCTION()
 	void OnSourceEndPlay(AActor* InSource, EEndPlayReason::Type Reason);
 
+	//////////////////////////////////////////////////////////////////////////
+	// Misc Functions
+	void VMGetPreSkinnedLocalBounds(FVectorVMExternalFunctionContext& Context);
+
 	//////////////////////////////////////////////////////////////////////////
 	//Triangle sampling
 	//Triangles are sampled a using MeshTriangleCoordinates which are composed of Triangle index and a bary centric coordinate on that triangle.
@@ -1102,6 +1065,8 @@ struct FNiagaraDISkeletalMeshPassedDataToRT
 	FMatrix44f PrevTransform = FMatrix44f::Identity;
 	float DeltaSeconds = 0.0f;
 	uint32 UvMappingSet = 0;
+	FVector3f PreSkinnedLocalBoundsCenter = FVector3f::ZeroVector;
+	FVector3f PreSkinnedLocalBoundsExtents = FVector3f::ZeroVector;
 };
 
 typedef FNiagaraDISkeletalMeshPassedDataToRT FNiagaraDataInterfaceProxySkeletalMeshData;
diff --git a/Engine/Plugins/FX/Niagara/Source/Niagara/Private/DataInterface/NiagaraDataInterfaceStaticMesh.cpp b/Engine/Plugins/FX/Niagara/Source/Niagara/Private/DataInterface/NiagaraDataInterfaceStaticMesh.cpp
--- a/Engine/Plugins/FX/Niagara/Source/Niagara/Private/DataInterface/NiagaraDataInterfaceStaticMesh.cpp
+++ b/Engine/Plugins/FX/Niagara/Source/Niagara/Private/DataInterface/NiagaraDataInterfaceStaticMesh.cpp
@@ -73,6 +73,8 @@ namespace NDIStaticMeshLocal
 
 		SHADER_PARAMETER(int,					InstanceDistanceFieldIndex)
 
+		SHADER_PARAMETER(FVector3f,				PreSkinnedLocalBoundsCenter)
+		SHADER_PARAMETER(FVector3f,				PreSkinnedLocalBoundsExtents)
 		SHADER_PARAMETER(FVector3f,				SystemLWCTile)
 
 		SHADER_PARAMETER_SRV(Buffer<int>,		UvMappingBuffer)
@@ -181,6 +183,8 @@ namespace NDIStaticMeshLocal
 	// VM Misc Functions
 	static const FName	IsValidName("IsValid");
 
+	static const FName	GetPreSkinnedLocalBoundsName("GetPreSkinnedLocalBounds");
+
 	static const FName	GetLocalToWorldName("GetLocalToWorld");
 	static const FName	GetLocalToWorldInverseTransposedName("GetLocalToWorldInverseTransposed");
 	static const FName	GetWorldVelocityName("GetWorldVelocity");
@@ -362,7 +366,8 @@ namespace NDIStaticMeshLocal
 		FQuat4f		Rotation = FQuat4f::Identity;
 		FQuat4f		PrevRotation = FQuat4f::Identity;
 		float		DeltaSeconds = 0.0f;
-
+		FVector3f	PreSkinnedLocalBoundsCenter = FVector3f::ZeroVector;
+		FVector3f	PreSkinnedLocalBoundsExtents = FVector3f::ZeroVector;
 		FPrimitiveComponentId	DistanceFieldPrimitiveId;
 
 		FIntVector NumTriangles = FIntVector::ZeroValue;
@@ -519,6 +524,8 @@ namespace NDIStaticMeshLocal
 		FQuat4f					Rotation = FQuat4f::Identity;
 		FQuat4f					PrevRotation = FQuat4f::Identity;
 		float					DeltaSeconds = 0.0f;
+		FVector3f				PreSkinnedLocalBoundsCenter = FVector3f::ZeroVector;
+		FVector3f				PreSkinnedLocalBoundsExtents = FVector3f::ZeroVector;
 		FPrimitiveComponentId	DistanceFieldPrimitiveId;
 
 		const FMeshUvMappingBufferProxy* UvMappingBuffer = nullptr;
@@ -558,6 +565,10 @@ namespace NDIStaticMeshLocal
 		/** Time separating Transform and PrevTransform. */
 		float DeltaSeconds;
 
+		/** Preskinned local bounds pulled from attached static mesh */
+		FVector3f PreSkinnedLocalBoundsCenter = FVector3f::ZeroVector;
+		FVector3f PreSkinnedLocalBoundsExtents = FVector3f::ZeroVector;
+
 		/** Velocity set by the physics body of the mesh component */
 		FVector PhysicsVelocity;
 		/** True if velocity should not be calculated via the transforms, but rather read the physics data from the mesh component */
@@ -650,6 +661,8 @@ namespace NDIStaticMeshLocal
 			Rotation = FQuat4f::Identity;
 			PrevRotation = FQuat4f::Identity;
 			DeltaSeconds = 0.0f;
+			PreSkinnedLocalBoundsCenter = FVector3f::ZeroVector;
+			PreSkinnedLocalBoundsExtents = FVector3f::ZeroVector;
 			PhysicsVelocity = FVector::ZeroVector;
 			bUsePhysicsVelocity = Interface->bUsePhysicsBodyVelocity;
 			bComponentValid = false;
@@ -852,6 +865,12 @@ namespace NDIStaticMeshLocal
 					}
 				}
 
+				// Preskinned local bounds
+				{
+					const FBoxSphereBounds LocalBounds = StaticMesh->GetBounds();
+					PreSkinnedLocalBoundsCenter = FVector3f(LocalBounds.Origin);
+					PreSkinnedLocalBoundsExtents = FVector3f(LocalBounds.BoxExtent);
+				}
 			}
 
 			// support for UV mapping
@@ -1049,6 +1068,8 @@ namespace NDIStaticMeshLocal
 			InstanceData->Rotation					= FromGameThread->Rotation;
 			InstanceData->PrevRotation				= FromGameThread->PrevRotation;
 			InstanceData->DeltaSeconds				= FromGameThread->DeltaSeconds;
+			InstanceData->PreSkinnedLocalBoundsCenter = FromGameThread->PreSkinnedLocalBoundsCenter;
+			InstanceData->PreSkinnedLocalBoundsExtents = FromGameThread->PreSkinnedLocalBoundsExtents;
 			InstanceData->DistanceFieldPrimitiveId	= FromGameThread->DistanceFieldPrimitiveId;
 
 			InstanceData->UvMappingBuffer			= FromGameThread->UvMappingBuffer;
@@ -1637,6 +1658,8 @@ void UNiagaraDataInterfaceStaticMesh::ProvidePerInstanceDataForRenderThread(void
 	DataFromGT->Rotation					= InstanceData->Rotation;
 	DataFromGT->PrevRotation				= InstanceData->PrevRotation;
 	DataFromGT->DeltaSeconds				= InstanceData->DeltaSeconds;
+	DataFromGT->PreSkinnedLocalBoundsCenter = InstanceData->PreSkinnedLocalBoundsCenter;
+	DataFromGT->PreSkinnedLocalBoundsExtents = InstanceData->PreSkinnedLocalBoundsExtents;
 	DataFromGT->DistanceFieldPrimitiveId	= FPrimitiveComponentId();
 	DataFromGT->UvMappingBuffer				= InstanceData->UvMapping.GetQuadTreeProxy();
 	DataFromGT->UvMappingSet				= InstanceData->UvMapping.GetUvSetIndex();
@@ -1938,6 +1961,16 @@ void UNiagaraDataInterfaceStaticMesh::GetFunctions(TArray<FNiagaraFunctionSignat
 		Sig.Name = IsValidName;
 		Sig.Outputs.Emplace(FNiagaraTypeDefinition::GetBoolDef(), TEXT("Valid"));
 	}
+	{
+		FNiagaraFunctionSignature& Sig = OutFunctions.Add_GetRef(BaseSignature);
+		Sig.Name = GetPreSkinnedLocalBoundsName;
+		Sig.Outputs.Emplace(FNiagaraTypeDefinition::GetVec3Def(), TEXT("Center"));
+		Sig.Outputs.Emplace(FNiagaraTypeDefinition::GetVec3Def(), TEXT("ExtentsMin"));
+		Sig.Outputs.Emplace(FNiagaraTypeDefinition::GetVec3Def(), TEXT("ExtentsMax"));
+		Sig.Outputs.Emplace(FNiagaraTypeDefinition::GetVec3Def(), TEXT("Extents"));
+		Sig.Outputs.Emplace(FNiagaraTypeDefinition::GetVec3Def(), TEXT("Half Extents"));
+		Sig.SetDescription(LOCTEXT("GetPreSkinnedLocalBoundsDesc", "Returns the local bounds of the static mesh"));
+	}
 	{
 		FNiagaraFunctionSignature& Sig = OutFunctions.Add_GetRef(BaseSignature);
 		Sig.Name = GetLocalToWorldName;
@@ -2405,6 +2438,10 @@ void UNiagaraDataInterfaceStaticMesh::GetVMExternalFunction(const FVMExternalFun
 	{
 		OutFunc = FVMExternalFunction::CreateLambda([this](FVectorVMExternalFunctionContext& Context) { VMIsValid(Context); });
 	}
+	else if (BindingInfo.Name == GetPreSkinnedLocalBoundsName)
+	{
+	OutFunc = FVMExternalFunction::CreateUObject(this, &UNiagaraDataInterfaceStaticMesh::VMGetPreSkinnedLocalBounds);
+	}
 	else if (BindingInfo.Name == GetLocalToWorldName)
 	{
 		OutFunc = FVMExternalFunction::CreateLambda([this](FVectorVMExternalFunctionContext& Context) { VMGetLocalToWorld(Context); });
@@ -2729,6 +2766,9 @@ void UNiagaraDataInterfaceStaticMesh::SetShaderParameters(const FNiagaraDataInte
 		FNiagaraDistanceFieldHelper::SetMeshDistanceFieldParameters(Context.GetGraphBuilder(), DistanceFieldSceneData, *ShaderDistanceFieldObjectParameters, *ShaderDistanceFieldAtlasParameters, FNiagaraRenderer::GetDummyFloat4Buffer());
 	}
 
+	ShaderParameters->PreSkinnedLocalBoundsCenter = InstanceData.PreSkinnedLocalBoundsCenter;
+	ShaderParameters->PreSkinnedLocalBoundsExtents = InstanceData.PreSkinnedLocalBoundsExtents;
+
 	ShaderParameters->SystemLWCTile = Context.GetSystemLWCTile();
 
 	if (InstanceData.UvMappingBuffer)
@@ -3035,6 +3075,33 @@ void UNiagaraDataInterfaceStaticMesh::SetDefaultMeshFromBlueprints(UStaticMesh*
 	DefaultMesh = MeshToUse;
 }
 
+//////////////////////////////////////////////////////////////////////////
+// Misc Functions
+void UNiagaraDataInterfaceStaticMesh::VMGetPreSkinnedLocalBounds(FVectorVMExternalFunctionContext& Context)
+{
+	NDIStaticMeshLocal::FStaticMeshCpuHelper<FNDITransformHandlerNoop> StaticMeshHelper(Context);
+	FNDIOutputParam<FVector3f>	OutCenter(Context);
+	FNDIOutputParam<FVector3f>	OutExtentsMin(Context);
+	FNDIOutputParam<FVector3f>	OutExtentsMax(Context);
+	FNDIOutputParam<FVector3f>	OutExtents(Context);
+	FNDIOutputParam<FVector3f>	OutHalfExtents(Context);
+
+	const FVector3f Center = StaticMeshHelper.InstanceData->PreSkinnedLocalBoundsCenter;
+	const FVector3f Extents = StaticMeshHelper.InstanceData->PreSkinnedLocalBoundsExtents;
+	const FVector3f HalfExtents = Extents * 0.5f;
+	const FVector3f ExtentsMin = Center - HalfExtents;
+	const FVector3f ExtentsMax = Center + HalfExtents;
+
+	for (int32 i = 0; i < Context.GetNumInstances(); ++i)
+	{
+		OutCenter.SetAndAdvance(Center);
+		OutExtentsMin.SetAndAdvance(ExtentsMin);
+		OutExtentsMax.SetAndAdvance(ExtentsMax);
+		OutExtents.SetAndAdvance(Extents);
+		OutHalfExtents.SetAndAdvance(HalfExtents);
+	}
+}
+
 //////////////////////////////////////////////////////////////////////////
 // VM Vertex Sampling
 void UNiagaraDataInterfaceStaticMesh::VMIsValidVertex(FVectorVMExternalFunctionContext& Context)
diff --git a/Engine/Plugins/FX/Niagara/Source/Niagara/Private/DataInterface/NiagaraDataInterfaceStaticMesh.h b/Engine/Plugins/FX/Niagara/Source/Niagara/Private/DataInterface/NiagaraDataInterfaceStaticMesh.h
--- a/Engine/Plugins/FX/Niagara/Source/Niagara/Private/DataInterface/NiagaraDataInterfaceStaticMesh.h
+++ b/Engine/Plugins/FX/Niagara/Source/Niagara/Private/DataInterface/NiagaraDataInterfaceStaticMesh.h
@@ -230,6 +230,8 @@ protected:
 	// VM Misc Functions
 	void VMIsValid(FVectorVMExternalFunctionContext& Context);
 
+	void VMGetPreSkinnedLocalBounds(FVectorVMExternalFunctionContext& Context);
+
 	void VMGetLocalToWorld(FVectorVMExternalFunctionContext& Context);
 	void VMGetLocalToWorldInverseTransposed(FVectorVMExternalFunctionContext& Context);
 	void VMGetWorldVelocity(FVectorVMExternalFunctionContext& Context);
diff --git a/Engine/Plugins/FX/Niagara/Source/Niagara/Private/NiagaraDataInterfaceSkeletalMesh.cpp b/Engine/Plugins/FX/Niagara/Source/Niagara/Private/NiagaraDataInterfaceSkeletalMesh.cpp
--- a/Engine/Plugins/FX/Niagara/Source/Niagara/Private/NiagaraDataInterfaceSkeletalMesh.cpp
+++ b/Engine/Plugins/FX/Niagara/Source/Niagara/Private/NiagaraDataInterfaceSkeletalMesh.cpp
@@ -52,6 +52,8 @@ struct FNiagaraSkelMeshDIFunctionVersion
 
 namespace NDISkelMeshLocal
 {
+	static FName NAME_GetPreSkinnedLocalBounds("GetPreSkinnedLocalBounds");
+
 	static const TCHAR* CommonShaderFile = TEXT("/Plugin/FX/Niagara/Private/NiagaraDataInterfaceSkeletalMesh.ush");
 	static const TCHAR* TemplateShaderFile = TEXT("/Plugin/FX/Niagara/Private/NiagaraDataInterfaceSkeletalMeshTemplate.ush");
 
@@ -99,6 +101,8 @@ namespace NDISkelMeshLocal
 		SHADER_PARAMETER(FQuat4f,				InstanceRotation)
 		SHADER_PARAMETER(FQuat4f,				InstancePrevRotation)
 		SHADER_PARAMETER(float,					InstanceInvDeltaTime)
+		SHADER_PARAMETER(FVector3f,				PreSkinnedLocalBoundsCenter)
+		SHADER_PARAMETER(FVector3f,				PreSkinnedLocalBoundsExtents)
 		SHADER_PARAMETER(uint32,				EnabledFeatures)
 	END_SHADER_PARAMETER_STRUCT()
 
@@ -1367,6 +1371,8 @@ void FNiagaraDataInterfaceProxySkeletalMesh::ConsumePerInstanceDataFromGameThrea
 	Data.PrevTransform = SourceData->PrevTransform;
 	Data.StaticBuffers = SourceData->StaticBuffers;
 	Data.Transform = SourceData->Transform;
+	Data.PreSkinnedLocalBoundsCenter = SourceData->PreSkinnedLocalBoundsCenter;
+	Data.PreSkinnedLocalBoundsExtents = SourceData->PreSkinnedLocalBoundsExtents;
 
 	Data.MeshSkinWeightBuffer = SourceData->MeshSkinWeightBuffer;
 	Data.MeshSkinWeightLookupBuffer = SourceData->MeshSkinWeightLookupBuffer;
@@ -1396,6 +1402,8 @@ void UNiagaraDataInterfaceSkeletalMesh::ProvidePerInstanceDataForRenderThread(vo
 	Data->PrevTransform = FMatrix44f(SourceData->PrevTransform);	// LWC_TODO: Precision loss
 	Data->StaticBuffers = SourceData->MeshGpuSpawnStaticBuffers;
 	Data->Transform = FMatrix44f(SourceData->Transform);			// LWC_TODO: Precision loss
+	Data->PreSkinnedLocalBoundsCenter = SourceData->PreSkinnedLocalBoundsCenter;
+	Data->PreSkinnedLocalBoundsExtents = SourceData->PreSkinnedLocalBoundsExtents;
 
 	Data->MeshSkinWeightBuffer = SourceData->MeshSkinWeightBuffer;
 	Data->MeshSkinWeightLookupBuffer = SourceData->MeshSkinWeightLookupBuffer;
@@ -1582,6 +1590,8 @@ bool FNDISkeletalMesh_InstanceData::Init(UNiagaraDataInterfaceSkeletalMesh* Inte
 	MeshGpuSpawnStaticBuffers = nullptr;
 	MeshGpuSpawnDynamicBuffers = nullptr;
 	bAllowCPUMeshDataAccess = false;
+	PreSkinnedLocalBoundsCenter = FVector3f::ZeroVector;
+	PreSkinnedLocalBoundsExtents = FVector3f::ZeroVector;
 
 	// Get skel mesh and confirm have valid data
 	USkeletalMeshComponent* NewSkelComp = nullptr;
@@ -1686,6 +1696,13 @@ bool FNDISkeletalMesh_InstanceData::Init(UNiagaraDataInterfaceSkeletalMesh* Inte
 			}
 		}
 #endif
+
+		if ( Mesh != nullptr )
+		{
+			const FBoxSphereBounds LocalBounds = Mesh->GetBounds();
+			PreSkinnedLocalBoundsCenter = FVector3f(LocalBounds.Origin);
+			PreSkinnedLocalBoundsExtents = FVector3f(LocalBounds.BoxExtent);
+		}
 	}
 
 	check(CachedLODIdx >= 0);
@@ -2268,6 +2285,20 @@ void UNiagaraDataInterfaceSkeletalMesh::GetFunctions(TArray<FNiagaraFunctionSign
 {
 	const int32 FirstFunction = OutFunctions.Num();
 
+	{
+		FNiagaraFunctionSignature& Sig = OutFunctions.AddDefaulted_GetRef();
+		Sig.Name = NDISkelMeshLocal::NAME_GetPreSkinnedLocalBounds;
+		Sig.Inputs.Emplace(FNiagaraTypeDefinition(GetClass()), TEXT("SkeletalMesh"));
+		Sig.Outputs.Emplace(FNiagaraTypeDefinition::GetVec3Def(), TEXT("Center"));
+		Sig.Outputs.Emplace(FNiagaraTypeDefinition::GetVec3Def(), TEXT("ExtentsMin"));
+		Sig.Outputs.Emplace(FNiagaraTypeDefinition::GetVec3Def(), TEXT("ExtentsMax"));
+		Sig.Outputs.Emplace(FNiagaraTypeDefinition::GetVec3Def(), TEXT("Extents"));
+		Sig.Outputs.Emplace(FNiagaraTypeDefinition::GetVec3Def(), TEXT("Half Extents"));
+		Sig.SetDescription(LOCTEXT("GetPreSkinnedLocalBoundsDesc", "Returns the local bounds of the Skeletal Mesh"));
+		Sig.bMemberFunction = true;
+		Sig.bRequiresContext = false;
+	}
+
 	GetTriangleSamplingFunctions(OutFunctions);
 	GetVertexSamplingFunctions(OutFunctions);
 	GetSkeletonSamplingFunctions(OutFunctions);
@@ -2290,6 +2321,13 @@ void UNiagaraDataInterfaceSkeletalMesh::GetVMExternalFunction(const FVMExternalF
 		return;
 	}
 
+	// Bind misc functions
+	if (BindingInfo.Name == NDISkelMeshLocal::NAME_GetPreSkinnedLocalBounds)
+	{
+		OutFunc = FVMExternalFunction::CreateUObject(this, &UNiagaraDataInterfaceSkeletalMesh::VMGetPreSkinnedLocalBounds);
+		return;
+	}
+
 	// Bind skeleton sampling function
 	BindSkeletonSamplingFunction(BindingInfo, InstData, OutFunc);
 	if (OutFunc.IsBound())
@@ -2726,51 +2764,6 @@ void UNiagaraDataInterfaceSkeletalMesh::ValidateFunction(const FNiagaraFunctionS
 
 #endif
 
-const FString UNiagaraDataInterfaceSkeletalMesh::MeshIndexBufferName(TEXT("MeshIndexBuffer_"));
-const FString UNiagaraDataInterfaceSkeletalMesh::MeshVertexBufferName(TEXT("MeshVertexBuffer_"));
-const FString UNiagaraDataInterfaceSkeletalMesh::MeshSkinWeightBufferName(TEXT("MeshSkinWeightBuffer_"));
-const FString UNiagaraDataInterfaceSkeletalMesh::MeshSkinWeightLookupBufferName(TEXT("MeshSkinWeightLookupBuffer_"));
-const FString UNiagaraDataInterfaceSkeletalMesh::MeshCurrBonesBufferName(TEXT("MeshCurrBonesBuffer_"));
-const FString UNiagaraDataInterfaceSkeletalMesh::MeshPrevBonesBufferName(TEXT("MeshPrevBonesBuffer_"));
-const FString UNiagaraDataInterfaceSkeletalMesh::MeshCurrSamplingBonesBufferName(TEXT("MeshCurrSamplingBonesBuffer_"));
-const FString UNiagaraDataInterfaceSkeletalMesh::MeshPrevSamplingBonesBufferName(TEXT("MeshPrevSamplingBonesBuffer_"));
-const FString UNiagaraDataInterfaceSkeletalMesh::MeshTangentBufferName(TEXT("MeshTangentBuffer_"));
-const FString UNiagaraDataInterfaceSkeletalMesh::MeshTexCoordBufferName(TEXT("MeshTexCoordBuffer_"));
-const FString UNiagaraDataInterfaceSkeletalMesh::MeshColorBufferName(TEXT("MeshColorBuffer_"));
-const FString UNiagaraDataInterfaceSkeletalMesh::MeshTriangleSamplerProbAliasBufferName(TEXT("MeshTriangleSamplerProbAliasBuffer_"));
-const FString UNiagaraDataInterfaceSkeletalMesh::MeshNumSamplingRegionTrianglesName(TEXT("MeshNumSamplingRegionTriangles_"));
-const FString UNiagaraDataInterfaceSkeletalMesh::MeshNumSamplingRegionVerticesName(TEXT("MeshNumSamplingRegionVertices_"));
-const FString UNiagaraDataInterfaceSkeletalMesh::MeshSamplingRegionsProbAliasBufferName(TEXT("MeshSamplingRegionsProbAliasBuffer_"));
-const FString UNiagaraDataInterfaceSkeletalMesh::MeshSampleRegionsTriangleIndicesName(TEXT("MeshSampleRegionsTriangleIndices_"));
-const FString UNiagaraDataInterfaceSkeletalMesh::MeshSampleRegionsVerticesName(TEXT("MeshSampleRegionsVertices_"));
-const FString UNiagaraDataInterfaceSkeletalMesh::MeshTriangleMatricesOffsetBufferName(TEXT("MeshTriangleMatricesOffsetBuffer_"));
-const FString UNiagaraDataInterfaceSkeletalMesh::MeshTriangleCountName(TEXT("MeshTriangleCount_"));
-const FString UNiagaraDataInterfaceSkeletalMesh::MeshVertexCountName(TEXT("MeshVertexCount_"));
-const FString UNiagaraDataInterfaceSkeletalMesh::MeshWeightStrideName(TEXT("MeshWeightStride_"));
-const FString UNiagaraDataInterfaceSkeletalMesh::MeshSkinWeightIndexSizeName(TEXT("MeshSkinWeightIndexSize_"));
-const FString UNiagaraDataInterfaceSkeletalMesh::MeshNumTexCoordName(TEXT("MeshNumTexCoord_"));
-const FString UNiagaraDataInterfaceSkeletalMesh::MeshNumWeightsName(TEXT("MeshNumWeights_"));
-const FString UNiagaraDataInterfaceSkeletalMesh::NumBonesName(TEXT("NumBones_"));
-const FString UNiagaraDataInterfaceSkeletalMesh::NumFilteredBonesName(TEXT("NumFilteredBones_"));
-const FString UNiagaraDataInterfaceSkeletalMesh::NumUnfilteredBonesName(TEXT("NumUnfilteredBones_"));
-const FString UNiagaraDataInterfaceSkeletalMesh::RandomMaxBoneName(TEXT("RandomMaxBone_"));
-const FString UNiagaraDataInterfaceSkeletalMesh::ExcludeBoneIndexName(TEXT("ExcludeBoneIndex_"));
-const FString UNiagaraDataInterfaceSkeletalMesh::FilteredAndUnfilteredBonesName(TEXT("FilteredAndUnfilteredBones_"));
-const FString UNiagaraDataInterfaceSkeletalMesh::NumFilteredSocketsName(TEXT("NumFilteredSockets_"));
-const FString UNiagaraDataInterfaceSkeletalMesh::FilteredSocketBoneOffsetName(TEXT("FilteredSocketBoneOffset_"));
-const FString UNiagaraDataInterfaceSkeletalMesh::UvMappingBufferName(TEXT("UvMappingBuffer_"));
-const FString UNiagaraDataInterfaceSkeletalMesh::UvMappingBufferLengthName(TEXT("UvMappingBufferLength_"));
-const FString UNiagaraDataInterfaceSkeletalMesh::UvMappingSetName(TEXT("UvMappingSet_"));
-const FString UNiagaraDataInterfaceSkeletalMesh::ConnectivityBufferName(TEXT("ConnectivityBuffer_"));
-const FString UNiagaraDataInterfaceSkeletalMesh::ConnectivityBufferLengthName(TEXT("ConnectivityBufferLength_"));
-const FString UNiagaraDataInterfaceSkeletalMesh::ConnectivityMaxAdjacentPerVertexName(TEXT("ConnectivityMaxAdjacentPerVertex_"));
-const FString UNiagaraDataInterfaceSkeletalMesh::InstanceTransformName(TEXT("InstanceTransform_"));
-const FString UNiagaraDataInterfaceSkeletalMesh::InstancePrevTransformName(TEXT("InstancePrevTransform_"));
-const FString UNiagaraDataInterfaceSkeletalMesh::InstanceRotationName(TEXT("InstanceRotation_"));
-const FString UNiagaraDataInterfaceSkeletalMesh::InstancePrevRotationName(TEXT("InstancePrevRotation_"));
-const FString UNiagaraDataInterfaceSkeletalMesh::InstanceInvDeltaTimeName(TEXT("InstanceInvDeltaTime_"));
-const FString UNiagaraDataInterfaceSkeletalMesh::EnabledFeaturesName(TEXT("EnabledFeatures_"));
-
 #if WITH_EDITORONLY_DATA
 bool UNiagaraDataInterfaceSkeletalMesh::AppendCompileHash(FNiagaraCompileHashVisitor* InVisitor) const
 {
@@ -3288,6 +3281,9 @@ void UNiagaraDataInterfaceSkeletalMesh::SetShaderParameters(const FNiagaraDataIn
 		ShaderParameters->InstancePrevRotation = InstanceData->PrevTransform.GetMatrixWithoutScale().ToQuat();
 		ShaderParameters->InstanceInvDeltaTime = 1.0f / InstanceData->DeltaSeconds;
 
+		ShaderParameters->PreSkinnedLocalBoundsCenter = InstanceData->PreSkinnedLocalBoundsCenter;
+		ShaderParameters->PreSkinnedLocalBoundsExtents = InstanceData->PreSkinnedLocalBoundsExtents;
+
 		ShaderParameters->EnabledFeatures = EnabledFeaturesBits;
 	}
 	else
@@ -3346,6 +3342,9 @@ void UNiagaraDataInterfaceSkeletalMesh::SetShaderParameters(const FNiagaraDataIn
 		ShaderParameters->InstancePrevRotation = FQuat4f::Identity;
 		ShaderParameters->InstanceInvDeltaTime = 0.0f;
 
+		ShaderParameters->PreSkinnedLocalBoundsCenter = FVector3f::ZeroVector;
+		ShaderParameters->PreSkinnedLocalBoundsExtents = FVector3f::ZeroVector;
+
 		ShaderParameters->EnabledFeatures = 0;
 	}
 }
@@ -3507,6 +3506,30 @@ int32 UNiagaraDataInterfaceSkeletalMesh::CalculateLODIndexAndSamplingRegions(USk
 	}
 }
 
+void UNiagaraDataInterfaceSkeletalMesh::VMGetPreSkinnedLocalBounds(FVectorVMExternalFunctionContext& Context)
+{
+	VectorVM::FUserPtrHandler<FNDISkeletalMesh_InstanceData> InstData(Context);
+	FNDIOutputParam<FVector3f>	OutCenter(Context);
+	FNDIOutputParam<FVector3f>	OutExtentsMin(Context);
+	FNDIOutputParam<FVector3f>	OutExtentsMax(Context);
+	FNDIOutputParam<FVector3f>	OutExtents(Context);
+	FNDIOutputParam<FVector3f>	OutHalfExtents(Context);
+
+	const FVector3f Center = InstData->PreSkinnedLocalBoundsCenter;
+	const FVector3f Extents = InstData->PreSkinnedLocalBoundsExtents;
+	const FVector3f HalfExtents = Extents * 0.5f;
+	const FVector3f ExtentsMin = Center - HalfExtents;
+	const FVector3f ExtentsMax = Center + HalfExtents;
+
+	for (int32 i = 0; i < Context.GetNumInstances(); ++i)
+	{
+		OutCenter.SetAndAdvance(Center);
+		OutExtentsMin.SetAndAdvance(ExtentsMin);
+		OutExtentsMax.SetAndAdvance(ExtentsMax);
+		OutExtents.SetAndAdvance(Extents);
+		OutHalfExtents.SetAndAdvance(HalfExtents);
+	}
+}
 
 //UNiagaraDataInterfaceSkeletalMesh END
 //////////////////////////////////////////////////////////////////////////
