- Enable deprecation warnings for legacy data interfaces

diff --git a/Engine/Plugins/FX/Niagara/Source/NiagaraCore/Public/NiagaraDataInterfaceBase.h b/Engine/Plugins/FX/Niagara/Source/NiagaraCore/Public/NiagaraDataInterfaceBase.h
--- a/Engine/Plugins/FX/Niagara/Source/NiagaraCore/Public/NiagaraDataInterfaceBase.h
+++ b/Engine/Plugins/FX/Niagara/Source/NiagaraCore/Public/NiagaraDataInterfaceBase.h
@@ -107,9 +107,6 @@ class NIAGARACORE_API UNiagaraDataInterfaceBase : public UNiagaraMergeable
 	GENERATED_UCLASS_BODY()
 
 public:
-	//UE_DEPRECATED(5.1, "Support for legacy shaders bindings will be removed in a future release, please upgrade your data interface.")
-	virtual bool UseLegacyShaderBindings() const { return true; }
-
 	/**
 	Override this method to provide parameters to the GPU (SRV / UAV / Constants / etc)
 	The most common usage will be to provide a single structure which is nested with other parameters.
@@ -141,27 +138,25 @@ public:
 	*/
 	virtual const FTypeLayoutDesc* GetShaderStorageType() const { return nullptr; }
 
+	/** Returns true if the DI (owned by OwnerEmitter) reads any attributes from the Provider emitter */
+	virtual bool HasInternalAttributeReads(const UNiagaraEmitter* OwnerEmitter, const UNiagaraEmitter* Provider) const { return false; }
+
+	//////////////////////////////////////////////////////////////////////////
+	/* 5.1 Deprecated methods will be removed in a future version. */
+	virtual bool UseLegacyShaderBindings() const { return true; }
+
 	/** Constructs the correct CS parameter type for this DI (if any). The object type returned by this can only vary by class and not per object data. */
-	//UE_DEPRECATED(5.1, "Is deprecated please use CreateShaderStorage instead.")
 	virtual FNiagaraDataInterfaceParametersCS* CreateComputeParameters() const { return nullptr; }
-	//UE_DEPRECATED(5.1, "Is deprecated please use GetShaderStorageType instead.")
 	virtual const FTypeLayoutDesc* GetComputeParametersTypeDesc() const { return nullptr; }
 
 	/** Methods that operate on an instance of type FNiagaraDataInterfaceParametersCS*, created by the above CreateComputeParameters() method */
-	//UE_DEPRECATED(5.1, "Is deprecated please convert to use BuildShaderParameters.")
 	virtual void BindParameters(FNiagaraDataInterfaceParametersCS* Base, const FNiagaraDataInterfaceGPUParamInfo& ParameterInfo, const class FShaderParameterMap& ParameterMap) {}
-	//UE_DEPRECATED(5.1, "Is deprecated please convert to use UNiagaraDataInterface::SetShaderParameters.")
 	virtual void SetParameters(const FNiagaraDataInterfaceParametersCS* Base, FRHICommandList& RHICmdList, const FNiagaraDataInterfaceSetArgs& Context) const {}
-	//UE_DEPRECATED(5.1, "Is deprecated and will be removed as it is no longer required.")
 	virtual void UnsetParameters(const FNiagaraDataInterfaceParametersCS* Base, FRHICommandList& RHICmdList, const FNiagaraDataInterfaceSetArgs& Context) const {}
-
-	/** Returns true if the DI (owned by OwnerEmitter) reads any attributes from the Provider emitter */
-	virtual bool HasInternalAttributeReads(const UNiagaraEmitter* OwnerEmitter, const UNiagaraEmitter* Provider) const { return false; }
 };
 
 /** This goes in class declaration for UNiagaraDataInterfaceBase-derived types, that need custom parameter type */
 #define DECLARE_NIAGARA_DI_PARAMETER() \
-	/*UE_STATIC_DEPRECATE(5.1, true, "DECLARE_NIAGARA_DI_PARAMETER is deprecate please update to use the new BuildShaderParameters / SetShaderParameters path.");*/ \
 	virtual FNiagaraDataInterfaceParametersCS* CreateComputeParameters() const override; \
 	virtual const FTypeLayoutDesc* GetComputeParametersTypeDesc() const override; \
 	virtual void BindParameters(FNiagaraDataInterfaceParametersCS* Base, const FNiagaraDataInterfaceGPUParamInfo& ParameterInfo, const class FShaderParameterMap& ParameterMap) override; \
@@ -170,7 +165,7 @@ public:
 
 /** This goes in cpp file matched with class declartion using DECLARE_NIAGARA_DI_PARAMETER() */
 #define IMPLEMENT_NIAGARA_DI_PARAMETER(T, ParameterType) \
-	/*UE_STATIC_DEPRECATE(5.1, true, "IMPLEMENT_NIAGARA_DI_PARAMETER is deprecate please update to use the new BuildShaderParameters / SetShaderParameters path.");*/ \
+	UE_STATIC_DEPRECATE(5.1, true, "IMPLEMENT_NIAGARA_DI_PARAMETER is deprecate please update to use the new BuildShaderParameters / SetShaderParameters path."); \
 	static_assert(ParameterType::InterfaceType == ETypeLayoutInterface::NonVirtual, "DI ParameterType must be non-virtual"); \
 	FNiagaraDataInterfaceParametersCS* T::CreateComputeParameters() const { return new ParameterType(); } \
 	const FTypeLayoutDesc* T::GetComputeParametersTypeDesc() const { return &StaticGetTypeLayoutDesc<ParameterType>(); } \
