- Add option to advance the simulation when capturing
- This allows for a one shot function in BP to do all the work rather than having to tick the world

#rnx

diff --git a/Engine/Plugins/FX/Niagara/Source/Niagara/Classes/NiagaraSimCacheFunctionLibrary.h b/Engine/Plugins/FX/Niagara/Source/Niagara/Classes/NiagaraSimCacheFunctionLibrary.h
--- a/Engine/Plugins/FX/Niagara/Source/Niagara/Classes/NiagaraSimCacheFunctionLibrary.h
+++ b/Engine/Plugins/FX/Niagara/Source/Niagara/Classes/NiagaraSimCacheFunctionLibrary.h
@@ -36,6 +36,12 @@ public:
 	UPROPERTY()
 	int32 TimeOutCounter = 0;
 
+	UPROPERTY()
+	bool bAdvanceSimulation = false;
+
+	UPROPERTY()
+	float AdvanceDeltaTime = 1.0f / 60.0f;
+
 	UPROPERTY(BlueprintAssignable)
 	FOnCaptureComplete CaptureComplete;
 
@@ -49,17 +55,19 @@ public:
 	Capture multiple frames from the provided simulation into a SimCache until the simulation becomes inactive, completes or we hit the NumFrames limit.
 	Capture occurs at the end of each frame with the first frame being this frame.
 	CaptureRate allows you to reduce the rate of capture, i.e. a rate of 2 would capture frames 0, 2, 4, etc.
+	When AdvanceSimulation is true we will manually advance the simulation in a loop until we have captured the number of frames request, rather than reading from the component each frame.
 	*/
 	UFUNCTION(BlueprintCallable, meta=(BlueprintInternalUseOnly="true", Category=NiagaraSimCache))
-	static UAsyncNiagaraCaptureSimCache* CaptureNiagaraSimCacheMultiFrame(UNiagaraSimCache* SimCache, FNiagaraSimCacheCreateParameters CreateParameters, UNiagaraComponent* NiagaraComponent, UNiagaraSimCache*& OutSimCache, int32 NumFrames = 16, int32 CaptureRate = 1);
+	static UAsyncNiagaraCaptureSimCache* CaptureNiagaraSimCacheMultiFrame(UNiagaraSimCache* SimCache, FNiagaraSimCacheCreateParameters CreateParameters, UNiagaraComponent* NiagaraComponent, UNiagaraSimCache*& OutSimCache, int32 NumFrames = 16, int32 CaptureRate = 1, bool bAdvanceSimulation=false, float AdvanceDeltaTime=0.01666f);
 
 	/**
 	Capture frames from the provided simulation into a SimCache until the simulation becomes inactive or completes.
 	Capture occurs at the end of each frame with the first frame being this frame.
 	CaptureRate allows you to reduce the rate of capture, i.e. a rate of 2 would capture frames 0, 2, 4, etc.
+	When AdvanceSimulation is true we will manually advance the simulation until the capture is complete inside a loop, rather than reading from the component each frame.
 	*/
 	UFUNCTION(BlueprintCallable, meta=(BlueprintInternalUseOnly="true", Category=NiagaraSimCache))
-	static UAsyncNiagaraCaptureSimCache* CaptureNiagaraSimCacheUntilComplete(UNiagaraSimCache* SimCache, FNiagaraSimCacheCreateParameters CreateParameters, UNiagaraComponent* NiagaraComponent, UNiagaraSimCache*& OutSimCache, int32 CaptureRate = 1);
+	static UAsyncNiagaraCaptureSimCache* CaptureNiagaraSimCacheUntilComplete(UNiagaraSimCache* SimCache, FNiagaraSimCacheCreateParameters CreateParameters, UNiagaraComponent* NiagaraComponent, UNiagaraSimCache*& OutSimCache, int32 CaptureRate = 1, bool bAdvanceSimulation=false, float AdvanceDeltaTime=0.01666f);
 };
 
 UCLASS()
@@ -72,9 +80,10 @@ public:
 	Captures the simulations current frame data into the SimCache.
 	This happens immediately so you may need to be careful with tick order of the component you are capturing from.
 	The return can be invalid if the component can not be captured for some reason (i.e. not active).
+	When AdvanceSimulation is true we will manually advance the simulation one frame using the provided AdvanceDeltaTime before capturing.
 	*/
 	UFUNCTION(BlueprintCallable, Category=NiagaraSimCache, meta=(ReturnDisplayName="Success"))
-	static bool CaptureNiagaraSimCacheImmediate(UNiagaraSimCache* SimCache, FNiagaraSimCacheCreateParameters CreateParameters, UNiagaraComponent* NiagaraComponent, UNiagaraSimCache*& OutSimCache);
+	static bool CaptureNiagaraSimCacheImmediate(UNiagaraSimCache* SimCache, FNiagaraSimCacheCreateParameters CreateParameters, UNiagaraComponent* NiagaraComponent, UNiagaraSimCache*& OutSimCache, bool bAdvanceSimulation=false, float AdvanceDeltaTime=0.01666f);
 
 	/**
 	Captures the simulation cache object that can be captured into using the various API calls.
diff --git a/Engine/Plugins/FX/Niagara/Source/Niagara/Private/NiagaraSimCacheFunctionLibrary.cpp b/Engine/Plugins/FX/Niagara/Source/Niagara/Private/NiagaraSimCacheFunctionLibrary.cpp
--- a/Engine/Plugins/FX/Niagara/Source/Niagara/Private/NiagaraSimCacheFunctionLibrary.cpp
+++ b/Engine/Plugins/FX/Niagara/Source/Niagara/Private/NiagaraSimCacheFunctionLibrary.cpp
@@ -7,17 +7,30 @@ void UAsyncNiagaraCaptureSimCache::Activate()
 {
 	Super::Activate();
 
-	if ( CaptureComponent != nullptr )
+	// In this path we are going to manually advance the simulation until we complete
+	if (bAdvanceSimulation)
 	{
-		if (UWorld* OwnerWorld = CaptureComponent->GetWorld())
+		while (CaptureComponent != nullptr)
 		{
-			TickerHandle = FTSTicker::GetCoreTicker().AddTicker(FTickerDelegate::CreateUObject(this, &UAsyncNiagaraCaptureSimCache::OnFrameTick), 0);
+			CaptureComponent->AdvanceSimulation(1, AdvanceDeltaTime);
+			OnFrameTick(AdvanceDeltaTime);
 		}
 	}
-
-	if ( CaptureSimCache == nullptr || CaptureComponent == nullptr || TickerHandle.IsValid() == false )
+	// In this path we are going to monitor the simulation and capture each frame
+	else
 	{
-		SetReadyToDestroy();
+		if (CaptureComponent != nullptr)
+		{
+			if (UWorld* OwnerWorld = CaptureComponent->GetWorld())
+			{
+				TickerHandle = FTSTicker::GetCoreTicker().AddTicker(FTickerDelegate::CreateUObject(this, &UAsyncNiagaraCaptureSimCache::OnFrameTick), 0);
+			}
+		}
+
+		if (CaptureSimCache == nullptr || CaptureComponent == nullptr || TickerHandle.IsValid() == false)
+		{
+			SetReadyToDestroy();
+		}
 	}
 }
 
@@ -35,6 +48,8 @@ void UAsyncNiagaraCaptureSimCache::SetReadyToDestroy()
 		CaptureSimCache->EndWrite();
 	}
 	CaptureComplete.Broadcast(CaptureSimCache ? CaptureSimCache->IsCacheValid() : false);
+
+	CaptureComponent = nullptr;
 }
 
 bool UAsyncNiagaraCaptureSimCache::OnFrameTick(float DeltaTime)
@@ -82,7 +97,7 @@ bool UAsyncNiagaraCaptureSimCache::OnFrameTick(float DeltaTime)
 	return true;
 }
 
-UAsyncNiagaraCaptureSimCache* UAsyncNiagaraCaptureSimCache::CaptureNiagaraSimCacheMultiFrame(UNiagaraSimCache* SimCache, FNiagaraSimCacheCreateParameters CreateParameters, UNiagaraComponent* NiagaraComponent, UNiagaraSimCache*& OutSimCache, int32 NumFrames, int32 CaptureRate)
+UAsyncNiagaraCaptureSimCache* UAsyncNiagaraCaptureSimCache::CaptureNiagaraSimCacheMultiFrame(UNiagaraSimCache* SimCache, FNiagaraSimCacheCreateParameters CreateParameters, UNiagaraComponent* NiagaraComponent, UNiagaraSimCache*& OutSimCache, int32 NumFrames, int32 CaptureRate, bool bAdvanceSimulation, float AdvanceDeltaTime)
 {
 	UAsyncNiagaraCaptureSimCache* CaptureAction = NewObject<UAsyncNiagaraCaptureSimCache>();
 	CaptureAction->CaptureSimCache = SimCache;
@@ -90,6 +105,8 @@ UAsyncNiagaraCaptureSimCache* UAsyncNiagaraCaptureSimCache::CaptureNiagaraSimCac
 	CaptureAction->CaptureNumFrames = FMath::Max(1, NumFrames);
 	CaptureAction->CaptureFrameRate = FMath::Max(1, CaptureRate);
 	CaptureAction->CaptureFrameCounter = 0;
+	CaptureAction->bAdvanceSimulation = bAdvanceSimulation;
+	CaptureAction->AdvanceDeltaTime = AdvanceDeltaTime;
 
 	if (SimCache != nullptr)
 	{
@@ -100,7 +117,7 @@ UAsyncNiagaraCaptureSimCache* UAsyncNiagaraCaptureSimCache::CaptureNiagaraSimCac
 	return CaptureAction;
 }
 
-UAsyncNiagaraCaptureSimCache* UAsyncNiagaraCaptureSimCache::CaptureNiagaraSimCacheUntilComplete(UNiagaraSimCache* SimCache, FNiagaraSimCacheCreateParameters CreateParameters, UNiagaraComponent* NiagaraComponent, UNiagaraSimCache*& OutSimCache, int32 CaptureRate)
+UAsyncNiagaraCaptureSimCache* UAsyncNiagaraCaptureSimCache::CaptureNiagaraSimCacheUntilComplete(UNiagaraSimCache* SimCache, FNiagaraSimCacheCreateParameters CreateParameters, UNiagaraComponent* NiagaraComponent, UNiagaraSimCache*& OutSimCache, int32 CaptureRate, bool bAdvanceSimulation, float AdvanceDeltaTime)
 {
 	UAsyncNiagaraCaptureSimCache* CaptureAction = NewObject<UAsyncNiagaraCaptureSimCache>();
 	CaptureAction->CaptureSimCache = SimCache;
@@ -108,6 +125,8 @@ UAsyncNiagaraCaptureSimCache* UAsyncNiagaraCaptureSimCache::CaptureNiagaraSimCac
 	CaptureAction->CaptureNumFrames = 0;
 	CaptureAction->CaptureFrameRate = FMath::Max(1, CaptureRate);
 	CaptureAction->CaptureFrameCounter = 0;
+	CaptureAction->bAdvanceSimulation = bAdvanceSimulation;
+	CaptureAction->AdvanceDeltaTime = AdvanceDeltaTime;
 
 	if (SimCache != nullptr)
 	{
@@ -123,11 +142,16 @@ UNiagaraSimCacheFunctionLibrary::UNiagaraSimCacheFunctionLibrary(const FObjectIn
 {
 }
 
-bool UNiagaraSimCacheFunctionLibrary::CaptureNiagaraSimCacheImmediate(UNiagaraSimCache* SimCache, FNiagaraSimCacheCreateParameters CreateParameters, UNiagaraComponent* NiagaraComponent, UNiagaraSimCache*& OutSimCache)
+bool UNiagaraSimCacheFunctionLibrary::CaptureNiagaraSimCacheImmediate(UNiagaraSimCache* SimCache, FNiagaraSimCacheCreateParameters CreateParameters, UNiagaraComponent* NiagaraComponent, UNiagaraSimCache*& OutSimCache, bool bAdvanceSimulation, float AdvanceDeltaTime)
 {
 	bool bSuccess = false;
 	if (SimCache != nullptr && NiagaraComponent != nullptr)
 	{
+		if ( bAdvanceSimulation )
+		{
+			NiagaraComponent->AdvanceSimulation(1, AdvanceDeltaTime);
+		}
+
 		SimCache->BeginWrite(CreateParameters, NiagaraComponent);
 		SimCache->WriteFrame(NiagaraComponent);
 		SimCache->EndWrite();
