- Fix incorrect MaxVertexIndex for raytracing

#rnx

diff --git a/Engine/Plugins/FX/Niagara/Source/Niagara/Private/NiagaraRendererRibbons.cpp b/Engine/Plugins/FX/Niagara/Source/Niagara/Private/NiagaraRendererRibbons.cpp
--- a/Engine/Plugins/FX/Niagara/Source/Niagara/Private/NiagaraRendererRibbons.cpp
+++ b/Engine/Plugins/FX/Niagara/Source/Niagara/Private/NiagaraRendererRibbons.cpp
@@ -943,8 +943,7 @@ void FNiagaraRendererRibbons::GetDynamicRayTracingInstances(FRayTracingMaterialG
 	// Use the internal vertex buffer only when initialized otherwise used the shared vertex buffer - needs to be updated every frame
 	FRWBuffer* VertexBuffer = RayTracingDynamicVertexBuffer.NumBytes > 0 ? &RayTracingDynamicVertexBuffer : nullptr;
 	
-	const uint32 VertexCount = DynamicDataRibbon->bUseGPUInit? SourceParticleData->GetNumInstances() : DynamicDataRibbon->GenerationOutput->SortedIndices.Num();
-
+	const uint32 VertexCount = DynamicDataRibbon->bUseGPUInit ? SourceParticleData->GetNumInstances() : DynamicDataRibbon->GenerationOutput->SortedIndices.Num();
 	
 	const int32 MaxTriangleCount = RenderingViewResources->IndexGenerationSettings.MaxSegmentCount * RenderingViewResources->IndexGenerationSettings.SubSegmentCount * ShapeState.TrianglesPerSegment;
 	
@@ -1932,8 +1931,8 @@ inline void FNiagaraRendererRibbons::SetupMeshBatchAndCollectorResourceForView(c
 	VFLooseParams.Shape = static_cast<uint32>(ShapeState.Shape);
 	VFLooseParams.NeedsPreciseMotionVectors = GenerationConfig.NeedsPreciseMotionVectors();
 
-	VFLooseParams.IndirectDrawOutput = bShouldUseGPUInitIndices? (FRHIShaderResourceView*)RenderingViewResources->IndirectDrawBuffer.SRV : GetDummyUIntBuffer();
-	VFLooseParams.IndirectDrawOutputOffset = bShouldUseGPUInitIndices? 0 : -1;
+	VFLooseParams.IndirectDrawOutput = bShouldUseGPUInitIndices ? (FRHIShaderResourceView*)RenderingViewResources->IndirectDrawBuffer.SRV : GetDummyUIntBuffer();
+	VFLooseParams.IndirectDrawOutputOffset = bShouldUseGPUInitIndices ? 0 : -1;
 
 	// Collector.AllocateOneFrameResource uses default ctor, initialize the vertex factory
 	RenderingViewResources->VertexFactory.SetParticleFactoryType(NVFT_Ribbon);
@@ -1973,7 +1972,10 @@ inline void FNiagaraRendererRibbons::SetupMeshBatchAndCollectorResourceForView(c
 	else
 	{
 		MeshElement.NumPrimitives = GeneratedData.CPUTriangleCount;
-		MeshElement.MaxVertexIndex = SourceParticleData->GetNumInstances() == 0 ? 0 : SourceParticleData->GetNumInstances() - 1;
+
+		const uint32 MaxVertexCount = bShouldUseGPUInitIndices ? SourceParticleData->GetNumInstances() : DynamicDataRibbon->GenerationOutput->SortedIndices.Num();
+		MeshElement.MaxVertexIndex = MaxVertexCount > 0 ? MaxVertexCount - 1 : 0;
+
 		check(MeshElement.NumPrimitives > 0);
 	}	
 	
