Fixed SA issue with potential nullptr dereferencing (should never happen anyways)

diff --git a/Engine/Plugins/FX/Niagara/Source/NiagaraEditor/Private/Customizations/NiagaraComponentDetails.cpp b/Engine/Plugins/FX/Niagara/Source/NiagaraEditor/Private/Customizations/NiagaraComponentDetails.cpp
--- a/Engine/Plugins/FX/Niagara/Source/NiagaraEditor/Private/Customizations/NiagaraComponentDetails.cpp
+++ b/Engine/Plugins/FX/Niagara/Source/NiagaraEditor/Private/Customizations/NiagaraComponentDetails.cpp
@@ -418,74 +418,78 @@ void AddRowForUserParameter(TSharedPtr<FNiagaraParameterProxy> ParameterProxy, I
 	{
 		System = ParameterProxy->GetComponent()->GetAsset();
 	}
-	System->GetExposedParameters().RedirectUserVariable(UserParameter);
-	UNiagaraScriptVariable* ScriptVariable = FNiagaraEditorUtilities::GetScriptVariableForUserParameter(UserParameter, *System);
-		
-	TAttribute<FText> Tooltip = FText::GetEmpty();
-	if(ScriptVariable)
+
+	if(System)
 	{
-		Tooltip = TAttribute<FText>::CreateLambda([ScriptVariable]
+		System->GetExposedParameters().RedirectUserVariable(UserParameter);
+		UNiagaraScriptVariable* ScriptVariable = FNiagaraEditorUtilities::GetScriptVariableForUserParameter(UserParameter, *System);
+			
+		TAttribute<FText> Tooltip = FText::GetEmpty();
+		if(ScriptVariable)
 		{
-			return ScriptVariable->Metadata.Description;
-		});
-	}
+			Tooltip = TAttribute<FText>::CreateLambda([ScriptVariable]
+			{
+				return ScriptVariable->Metadata.Description;
+			});
+		}
+			
+		Row->GetDefaultWidgets(DefaultNameWidget, DefaultValueWidget, CustomWidget);
+
+		Row->GetPropertyHandle()->SetOnPropertyValuePreChange(FSimpleDelegate::CreateSP(ParameterProxy.ToSharedRef(), &FNiagaraParameterProxy::OnParameterPreChange));
+		Row->GetPropertyHandle()->SetOnChildPropertyValuePreChange(FSimpleDelegate::CreateSP(ParameterProxy.ToSharedRef(), &FNiagaraParameterProxy::OnParameterPreChange));
+		Row->GetPropertyHandle()->SetOnPropertyValueChangedWithData(TDelegate<void(const FPropertyChangedEvent&)>::CreateSP(ParameterProxy.ToSharedRef(), &FNiagaraParameterProxy::OnParameterChanged));
+		Row->GetPropertyHandle()->SetOnChildPropertyValueChangedWithData(TDelegate<void(const FPropertyChangedEvent&)>::CreateSP(ParameterProxy.ToSharedRef(), &FNiagaraParameterProxy::OnParameterChanged));
+		Row->GetPropertyHandle()->SetOnPropertyResetToDefault(FSimpleDelegate::CreateSP(ParameterProxy.ToSharedRef(), &FNiagaraParameterProxy::OnResetToDefault));
 		
-	Row->GetDefaultWidgets(DefaultNameWidget, DefaultValueWidget, CustomWidget);
+		CustomWidget
+			.NameContent()
+			[
+				SNew(SBox)
+				.Padding(FMargin(0.0f, 2.0f))
+				[
+					SNew(STextBlock)
+					.TextStyle(FNiagaraEditorStyle::Get(), "NiagaraEditor.ParameterText")
+					.Text(FText::FromName(Parameter.GetName()))
+					.ToolTipText(Tooltip)
+				]
+			];
 
-	Row->GetPropertyHandle()->SetOnPropertyValuePreChange(FSimpleDelegate::CreateSP(ParameterProxy.ToSharedRef(), &FNiagaraParameterProxy::OnParameterPreChange));
-	Row->GetPropertyHandle()->SetOnChildPropertyValuePreChange(FSimpleDelegate::CreateSP(ParameterProxy.ToSharedRef(), &FNiagaraParameterProxy::OnParameterPreChange));
-	Row->GetPropertyHandle()->SetOnPropertyValueChangedWithData(TDelegate<void(const FPropertyChangedEvent&)>::CreateSP(ParameterProxy.ToSharedRef(), &FNiagaraParameterProxy::OnParameterChanged));
-	Row->GetPropertyHandle()->SetOnChildPropertyValueChangedWithData(TDelegate<void(const FPropertyChangedEvent&)>::CreateSP(ParameterProxy.ToSharedRef(), &FNiagaraParameterProxy::OnParameterChanged));
-	Row->GetPropertyHandle()->SetOnPropertyResetToDefault(FSimpleDelegate::CreateSP(ParameterProxy.ToSharedRef(), &FNiagaraParameterProxy::OnResetToDefault));
-	
-	CustomWidget
-		.NameContent()
+		TSharedPtr<SWidget> ValueWidget = DefaultValueWidget;
+		if (CustomValueWidget.IsValid())
+		{
+			ValueWidget = CustomValueWidget;
+		}
+		
+		CustomWidget
+		.ValueContent()
 		[
-			SNew(SBox)
-			.Padding(FMargin(0.0f, 2.0f))
+			SNew(SHorizontalBox)
+			+ SHorizontalBox::Slot()
+			.HAlign(HAlign_Fill)
+			.Padding(4.0f)
 			[
-				SNew(STextBlock)
-				.TextStyle(FNiagaraEditorStyle::Get(), "NiagaraEditor.ParameterText")
-				.Text(FText::FromName(Parameter.GetName()))
-				.ToolTipText(Tooltip)
+				// Add in the parameter editor factoried above.
+				ValueWidget.ToSharedRef()
 			]
-		];
-
-	TSharedPtr<SWidget> ValueWidget = DefaultValueWidget;
-	if (CustomValueWidget.IsValid())
-	{
-		ValueWidget = CustomValueWidget;
-	}
-	
-	CustomWidget
-	.ValueContent()
-	[
-		SNew(SHorizontalBox)
-		+ SHorizontalBox::Slot()
-		.HAlign(HAlign_Fill)
-		.Padding(4.0f)
-		[
-			// Add in the parameter editor factoried above.
-			ValueWidget.ToSharedRef()
-		]
-		+ SHorizontalBox::Slot()
-		.VAlign(VAlign_Center)
-		.AutoWidth()
-		[
-			// Add in the "reset to default" buttons
-			SNew(SButton)
-			.OnClicked(ParameterProxy.ToSharedRef(), &FNiagaraParameterProxy::OnResetToDefaultClicked)
-			.Visibility(ParameterProxy.ToSharedRef(), &FNiagaraParameterProxy::GetResetToDefaultVisibility)
-			.ContentPadding(FMargin(5.f, 0.f))
-			.ToolTipText(LOCTEXT("ResetToDefaultToolTip", "Reset to Default"))
-			.ButtonStyle(FAppStyle::Get(), "NoBorder")
-			.Content()
+			+ SHorizontalBox::Slot()
+			.VAlign(VAlign_Center)
+			.AutoWidth()
 			[
-				SNew(SImage)
-				.Image(FAppStyle::GetBrush("PropertyWindow.DiffersFromDefault"))
+				// Add in the "reset to default" buttons
+				SNew(SButton)
+				.OnClicked(ParameterProxy.ToSharedRef(), &FNiagaraParameterProxy::OnResetToDefaultClicked)
+				.Visibility(ParameterProxy.ToSharedRef(), &FNiagaraParameterProxy::GetResetToDefaultVisibility)
+				.ContentPadding(FMargin(5.f, 0.f))
+				.ToolTipText(LOCTEXT("ResetToDefaultToolTip", "Reset to Default"))
+				.ButtonStyle(FAppStyle::Get(), "NoBorder")
+				.Content()
+				[
+					SNew(SImage)
+					.Image(FAppStyle::GetBrush("PropertyWindow.DiffersFromDefault"))
+				]
 			]
-		]
-	];
+		];
+	}
 }
 
 class FNiagaraUserParameterCategoryBuilder : public IDetailCustomNodeBuilder, public TSharedFromThis<FNiagaraUserParameterCategoryBuilder>
