Speculative fix of  SA issue with potential nullptr dereferencing (should never happen anyways)

diff --git a/Engine/Plugins/FX/Niagara/Source/NiagaraEditor/Private/Customizations/NiagaraComponentDetails.cpp b/Engine/Plugins/FX/Niagara/Source/NiagaraEditor/Private/Customizations/NiagaraComponentDetails.cpp
--- a/Engine/Plugins/FX/Niagara/Source/NiagaraEditor/Private/Customizations/NiagaraComponentDetails.cpp
+++ b/Engine/Plugins/FX/Niagara/Source/NiagaraEditor/Private/Customizations/NiagaraComponentDetails.cpp
@@ -976,27 +976,30 @@ public:
 
 	virtual void GenerateChildContent(IDetailChildrenBuilder& ChildrenBuilder) override
 	{
-		if (bDelegatesInitialized == false)
-		{
-			System->GetExposedParameters().OnStructureChanged().Add(FNiagaraParameterStore::FOnStructureChanged::FDelegate::CreateSP(this, &FNiagaraSystemUserParameterBuilder::Rebuild));
-			TSharedPtr<FNiagaraSystemViewModel> SystemViewModel = TNiagaraViewModelManager<UNiagaraSystem, FNiagaraSystemViewModel>::GetExistingViewModelForObject(System.Get());
-			SystemViewModel->GetUserParametersHierarchyViewModel()->OnHierarchyChanged().Add(UNiagaraHierarchyViewModelBase::FOnHierarchyChanged::FDelegate::CreateSP(this, &FNiagaraSystemUserParameterBuilder::Rebuild));
-
-			bDelegatesInitialized = true;
-		}
-
 		check(System.IsValid());
 
-		ParameterProxies.Reset();
-		
-		TArray<FNiagaraVariable> UserParameters;
-		System->GetExposedParameters().GetUserParameters(UserParameters);
-		ParameterProxies.Reserve(UserParameters.Num());
+		if (System.IsValid())
+		{
+			if(bDelegatesInitialized == false)
+			{
+				System->GetExposedParameters().OnStructureChanged().Add(FNiagaraParameterStore::FOnStructureChanged::FDelegate::CreateSP(this, &FNiagaraSystemUserParameterBuilder::Rebuild));
+				TSharedPtr<FNiagaraSystemViewModel> SystemViewModel = TNiagaraViewModelManager<UNiagaraSystem, FNiagaraSystemViewModel>::GetExistingViewModelForObject(System.Get());
+				SystemViewModel->GetUserParametersHierarchyViewModel()->OnHierarchyChanged().Add(UNiagaraHierarchyViewModelBase::FOnHierarchyChanged::FDelegate::CreateSP(this, &FNiagaraSystemUserParameterBuilder::Rebuild));
 
-		ParameterNameToDisplayStruct.Empty();
+				bDelegatesInitialized = true;
+			}
+			
+			ParameterProxies.Reset();
+			
+			TArray<FNiagaraVariable> UserParameters;
+			System->GetExposedParameters().GetUserParameters(UserParameters);
+			ParameterProxies.Reserve(UserParameters.Num());
 
-		UNiagaraHierarchyRoot* Root = Cast<UNiagaraSystemEditorData>(System->GetEditorData())->UserParameterHierarchy;
-		GenerateUserParameterRows(ChildrenBuilder, *Root, System.Get(), ActiveSection, ParameterProxies, ParameterNameToDisplayStruct, {}, OnRebuildChildren);
+			ParameterNameToDisplayStruct.Empty();
+
+			UNiagaraHierarchyRoot* Root = Cast<UNiagaraSystemEditorData>(System->GetEditorData())->UserParameterHierarchy;
+			GenerateUserParameterRows(ChildrenBuilder, *Root, System.Get(), ActiveSection, ParameterProxies, ParameterNameToDisplayStruct, {}, OnRebuildChildren);
+		}
 	}
 
 private:
