- Niagara Arrays now store data in the format ready for VVM / GPU
- BP functions go through a conversion in and out process when dealing with double types
- Native functions can set directly and avoid the conversion to LWC types
- Added UInt8 array type for more compact storage

diff --git a/Engine/Plugins/FX/Niagara/Source/Niagara/Classes/NiagaraDataInterfaceArrayFloat.h b/Engine/Plugins/FX/Niagara/Source/Niagara/Classes/NiagaraDataInterfaceArrayFloat.h
--- a/Engine/Plugins/FX/Niagara/Source/Niagara/Classes/NiagaraDataInterfaceArrayFloat.h
+++ b/Engine/Plugins/FX/Niagara/Source/Niagara/Classes/NiagaraDataInterfaceArrayFloat.h
@@ -40,27 +40,25 @@ struct FNDIArrayImplHelper<FVector2f> : public FNDIArrayImplHelperBase<FVector2f
 
 	static const FNiagaraTypeDefinition& GetTypeDefinition() { return FNiagaraTypeDefinition::GetVec2Def(); }
 	static const FVector2f GetDefaultValue() { return FVector2f::ZeroVector; }
-};
 
-// LWC_TODO: This is represented as an FVector2f internally (array is converted to floats during PushToRenderThread)
-template<>
-struct FNDIArrayImplHelper<FVector2d> : public FNDIArrayImplHelper<FVector2f>
-{
-	static void CopyToGpuMemory(void* Dest, const FVector2d* Src, int32 NumElements)
+	static void CopyCpuToCpuMemory(FVector2f* Dest, const FVector2f* Src, int32 NumElements)
+	{
+		FMemory::Memcpy(Dest, Src, NumElements * sizeof(FVector2f));
+	}
+
+	static void CopyCpuToCpuMemory(FVector2f* Dest, const FVector2D* Src, int32 NumElements)
 	{
-		FVector2f* TypedDest = reinterpret_cast<FVector2f*>(Dest);
-		while (NumElements--)
+		for ( int32 i=0; i < NumElements; ++i )
 		{
-			*TypedDest++ = FVector2f(*Src++);
+			Dest[i] = FVector2f(Src[i]);
 		}
 	}
 
-	static void CopyToCpuMemory(void* Dest, const FVector2f* Src, int32 NumElements)
+	static void CopyCpuToCpuMemory(FVector2D* Dest, const FVector2f* Src, int32 NumElements)
 	{
-		FVector2d* TypedDest = reinterpret_cast<FVector2d*>(Dest);
-		while (NumElements--)
+		for (int32 i = 0; i < NumElements; ++i)
 		{
-			*TypedDest++ = FVector2d(*Src++);
+			Dest[i] = FVector2D(Src[i]);
 		}
 	}
 };
@@ -81,28 +79,25 @@ struct FNDIArrayImplHelper<FVector3f> : public FNDIArrayImplHelperBase<FVector3f
 
 	static const FNiagaraTypeDefinition& GetTypeDefinition() { return FNiagaraTypeDefinition::GetVec3Def(); }
 	static const FVector3f GetDefaultValue() { return FVector3f::ZeroVector; }
-};
 
-// LWC_TODO: This is represented as an FVector3f internally (array is converted to floats during PushToRenderThread)
-template<>
-struct FNDIArrayImplHelper<FVector3d> : public FNDIArrayImplHelper<FVector3f>
-{
-	static void CopyToGpuMemory(void* Dest, const FVector3d* Src, int32 NumElements)
+	static void CopyCpuToCpuMemory(FVector3f* Dest, const FVector3f* Src, int32 NumElements)
+	{
+		FMemory::Memcpy(Dest, Src, NumElements * sizeof(FVector3f));
+	}
+
+	static void CopyCpuToCpuMemory(FVector3f* Dest, const FVector* Src, int32 NumElements)
 	{
-		FVector3f* TypedDest = reinterpret_cast<FVector3f*>(Dest);
-		while (NumElements--)
+		for (int32 i = 0; i < NumElements; ++i)
 		{
-			*TypedDest++ = FVector3f(*Src++);
+			Dest[i] = FVector3f(Src[i]);
 		}
 	}
 
-
-	static void CopyToCpuMemory(void* Dest, const FVector3f* Src, int32 NumElements)
+	static void CopyCpuToCpuMemory(FVector* Dest, const FVector3f* Src, int32 NumElements)
 	{
-		FVector3d* TypedDest = reinterpret_cast<FVector3d*>(Dest);
-		while (NumElements--)
+		for (int32 i = 0; i < NumElements; ++i)
 		{
-			*TypedDest++ = FVector3d(*Src++);
+			Dest[i] = FVector(Src[i]);
 		}
 	}
 };
@@ -114,6 +109,14 @@ struct FNDIArrayImplHelper<FNiagaraPosition> : public FNDIArrayImplHelper<FVecto
 
 	static const FNiagaraTypeDefinition& GetTypeDefinition() { return FNiagaraTypeDefinition::GetPositionDef(); }
 	static const FNiagaraPosition GetDefaultValue() { return FVector3f::ZeroVector; }
+
+	static void CopyCpuToCpuMemory(FVector3f* Dest, const FNiagaraPosition* Src, int32 NumElements)
+	{
+		for (int32 i = 0; i < NumElements; ++i)
+		{
+			Dest[i] = Src[i];
+		}
+	}
 };
 
 template<>
@@ -132,27 +135,25 @@ struct FNDIArrayImplHelper<FVector4f> : public FNDIArrayImplHelperBase<FVector4f
 
 	static const FNiagaraTypeDefinition& GetTypeDefinition() { return FNiagaraTypeDefinition::GetVec4Def(); }
 	static const FVector4f GetDefaultValue() { return FVector4f(ForceInitToZero); }
-};
 
-// LWC_TODO: This is represented as an FVector4f internally (array is converted to floats during PushToRenderThread)
-template<>
-struct FNDIArrayImplHelper<FVector4d> : public FNDIArrayImplHelper<FVector4f>
-{
-	static void CopyToGpuMemory(void* Dest, const FVector4d* Src, int32 NumElements)
+	static void CopyCpuToCpuMemory(FVector4f* Dest, const FVector4f* Src, int32 NumElements)
 	{
-		FVector4f* TypedDest = reinterpret_cast<FVector4f*>(Dest);
-		while (NumElements--)
+		FMemory::Memcpy(Dest, Src, NumElements * sizeof(FVector4f));
+	}
+
+	static void CopyCpuToCpuMemory(FVector4f* Dest, const FVector4* Src, int32 NumElements)
+	{
+		for (int32 i = 0; i < NumElements; ++i)
 		{
-			*TypedDest++ = FVector4f(*Src++);
+			Dest[i] = FVector4f(Src[i]);
 		}
 	}
 
-	static void CopyToCpuMemory(void* Dest, const FVector4f* Src, int32 NumElements)
+	static void CopyCpuToCpuMemory(FVector4* Dest, const FVector4f* Src, int32 NumElements)
 	{
-		FVector4d* TypedDest = reinterpret_cast<FVector4d*>(Dest);
-		while (NumElements--)
+		for (int32 i = 0; i < NumElements; ++i)
 		{
-			*TypedDest++ = FVector4d(*Src++);
+			Dest[i] = FVector4(Src[i]);
 		}
 	}
 };
@@ -191,104 +192,122 @@ struct FNDIArrayImplHelper<FQuat4f> : public FNDIArrayImplHelperBase<FQuat4f>
 
 	static const FNiagaraTypeDefinition& GetTypeDefinition() { return FNiagaraTypeDefinition::GetQuatDef(); }
 	static const FQuat4f GetDefaultValue() { return FQuat4f::Identity; }
-};
 
-// LWC_TODO: This is represented as an FQuat4f internally (array is converted to floats during PushToRenderThread)
-template<>
-struct FNDIArrayImplHelper<FQuat4d> : public FNDIArrayImplHelper<FQuat4f>
-{
-	static void CopyToGpuMemory(void* Dest, const FQuat4d* Src, int32 NumElements)
+	static void CopyCpuToCpuMemory(FQuat4f* Dest, const FQuat4f* Src, int32 NumElements)
+	{
+		FMemory::Memcpy(Dest, Src, NumElements * sizeof(FQuat4f));
+	}
+
+	static void CopyCpuToCpuMemory(FQuat4f* Dest, const FQuat* Src, int32 NumElements)
 	{
-		FQuat4f* TypedDest = reinterpret_cast<FQuat4f*>(Dest);
-		while (NumElements--)
+		for (int32 i = 0; i < NumElements; ++i)
 		{
-			*TypedDest++ = FQuat4f(*Src++);
+			Dest[i] = FQuat4f(Src[i]);
 		}
 	}
 
-	static void CopyToCpuMemory(void* Dest, const FQuat4f* Src, int32 NumElements)
+	static void CopyCpuToCpuMemory(FQuat* Dest, const FQuat4f* Src, int32 NumElements)
 	{
-		FQuat4d* TypedDest = reinterpret_cast<FQuat4d*>(Dest);
-		while (NumElements--)
+		for (int32 i = 0; i < NumElements; ++i)
 		{
-			*TypedDest++ = FQuat4d(*Src++);
+			Dest[i] = FQuat(Src[i]);
 		}
 	}
 };
 
-UCLASS(EditInlineNew, Category = "Array", meta = (DisplayName = "Float Array", Experimental), Blueprintable, BlueprintType)
+UCLASS(EditInlineNew, Category = "Array", meta = (DisplayName = "Float Array"), Blueprintable, BlueprintType)
 class NIAGARA_API UNiagaraDataInterfaceArrayFloat : public UNiagaraDataInterfaceArray
 {
-	GENERATED_UCLASS_BODY()
 public:
+	GENERATED_BODY()
+	NDIARRAY_GENERATE_BODY(UNiagaraDataInterfaceArrayFloat, float, FloatData)
+
 	UPROPERTY(EditAnywhere, BlueprintReadWrite, Category = "Array")
 	TArray<float> FloatData;
-
-	TArray<float>& GetArrayReference() { return FloatData; }
 };
 
-UCLASS(EditInlineNew, Category = "Array", meta = (DisplayName = "Vector 2D Array", Experimental), Blueprintable, BlueprintType)
+UCLASS(EditInlineNew, Category = "Array", meta = (DisplayName = "Vector 2D Array"), Blueprintable, BlueprintType)
 class NIAGARA_API UNiagaraDataInterfaceArrayFloat2 : public UNiagaraDataInterfaceArray
 {
-	GENERATED_UCLASS_BODY()
 public:
+	GENERATED_BODY()
+	NDIARRAY_GENERATE_BODY_LWC(UNiagaraDataInterfaceArrayFloat2, FVector2f, FloatData)
+
+#if WITH_EDITORONLY_DATA
 	UPROPERTY(EditAnywhere, BlueprintReadWrite, Category = "Array")
-	TArray<FVector2D> FloatData;	// LWC_TODO: Should be FVector2f, but only FVector2D is blueprint accessible
+	TArray<FVector2D> FloatData;
+#endif
 
-	TArray<FVector2D>& GetArrayReference() { return FloatData; }
+	UPROPERTY()
+	TArray<FVector2f> InternalFloatData;
 };
 
-UCLASS(EditInlineNew, Category = "Array", meta = (DisplayName = "Vector Array", Experimental), Blueprintable, BlueprintType)
+UCLASS(EditInlineNew, Category = "Array", meta = (DisplayName = "Vector Array"), Blueprintable, BlueprintType)
 class NIAGARA_API UNiagaraDataInterfaceArrayFloat3 : public UNiagaraDataInterfaceArray
 {
-	GENERATED_UCLASS_BODY()
 public:
+	GENERATED_BODY()
+	NDIARRAY_GENERATE_BODY_LWC(UNiagaraDataInterfaceArrayFloat3, FVector3f, FloatData)
+
+#if WITH_EDITORONLY_DATA
 	UPROPERTY(EditAnywhere, BlueprintReadWrite, Category = "Array")
-	TArray<FVector> FloatData;		// LWC_TODO: Should be FVector3f, but only FVector is blueprint accessible
+	TArray<FVector> FloatData;
+#endif
 
-	TArray<FVector>& GetArrayReference() { return FloatData; }
+	UPROPERTY()
+	TArray<FVector3f> InternalFloatData;
 };
 
-UCLASS(EditInlineNew, Category = "Array", meta = (DisplayName = "Position Array", Experimental), Blueprintable, BlueprintType)
+UCLASS(EditInlineNew, Category = "Array", meta = (DisplayName = "Position Array"), Blueprintable, BlueprintType)
 class NIAGARA_API UNiagaraDataInterfaceArrayPosition : public UNiagaraDataInterfaceArray
 {
-	GENERATED_UCLASS_BODY()
 public:
+	GENERATED_BODY()
+	NDIARRAY_GENERATE_BODY(UNiagaraDataInterfaceArrayPosition, FNiagaraPosition, PositionData)
+
 	UPROPERTY(EditAnywhere, BlueprintReadWrite, Category = "Array")
 	TArray<FNiagaraPosition> PositionData;
-
-	TArray<FNiagaraPosition>& GetArrayReference() { return PositionData; }
 };
 
-UCLASS(EditInlineNew, Category = "Array", meta = (DisplayName = "Vector 4 Array", Experimental), Blueprintable, BlueprintType)
+UCLASS(EditInlineNew, Category = "Array", meta = (DisplayName = "Vector 4 Array"), Blueprintable, BlueprintType)
 class NIAGARA_API UNiagaraDataInterfaceArrayFloat4 : public UNiagaraDataInterfaceArray
 {
-	GENERATED_UCLASS_BODY()
 public:
+	GENERATED_BODY()
+	NDIARRAY_GENERATE_BODY_LWC(UNiagaraDataInterfaceArrayFloat4, FVector4f, FloatData)
+
+#if WITH_EDITORONLY_DATA
 	UPROPERTY(EditAnywhere, BlueprintReadWrite, Category = "Array")
-	TArray<FVector4> FloatData;		// LWC_TODO: Should be FVector4f, but only FVector4 is blueprint accessible
+	TArray<FVector4> FloatData;
+#endif
 
-	TArray<FVector4>& GetArrayReference() { return FloatData; }
+	UPROPERTY()
+	TArray<FVector4f> InternalFloatData;
 };
 
-UCLASS(EditInlineNew, Category = "Array", meta = (DisplayName = "Color Array", Experimental), Blueprintable, BlueprintType)
+UCLASS(EditInlineNew, Category = "Array", meta = (DisplayName = "Color Array"), Blueprintable, BlueprintType)
 class NIAGARA_API UNiagaraDataInterfaceArrayColor : public UNiagaraDataInterfaceArray
 {
-	GENERATED_UCLASS_BODY()
 public:
+	GENERATED_BODY()
+	NDIARRAY_GENERATE_BODY(UNiagaraDataInterfaceArrayColor, FLinearColor, ColorData)
+
 	UPROPERTY(EditAnywhere, BlueprintReadWrite, Category = "Array")
 	TArray<FLinearColor> ColorData;
-
-	TArray<FLinearColor>& GetArrayReference() { return ColorData; }
 };
 
-UCLASS(EditInlineNew, Category = "Array", meta = (DisplayName = "Quaternion Array", Experimental), Blueprintable, BlueprintType)
+UCLASS(EditInlineNew, Category = "Array", meta = (DisplayName = "Quaternion Array"), Blueprintable, BlueprintType)
 class NIAGARA_API UNiagaraDataInterfaceArrayQuat : public UNiagaraDataInterfaceArray
 {
-	GENERATED_UCLASS_BODY()
 public:
+	GENERATED_BODY()
+	NDIARRAY_GENERATE_BODY_LWC(UNiagaraDataInterfaceArrayQuat, FQuat4f, QuatData)
+
+#if WITH_EDITORONLY_DATA
 	UPROPERTY(EditAnywhere, BlueprintReadWrite, Category = "Array")
 	TArray<FQuat> QuatData;
+#endif
 
-	TArray<FQuat>& GetArrayReference() { return QuatData; }
+	UPROPERTY()
+	TArray<FQuat4f> InternalQuatData;
 };
diff --git a/Engine/Plugins/FX/Niagara/Source/Niagara/Classes/NiagaraDataInterfaceArrayFunctionLibrary.h b/Engine/Plugins/FX/Niagara/Source/Niagara/Classes/NiagaraDataInterfaceArrayFunctionLibrary.h
--- a/Engine/Plugins/FX/Niagara/Source/Niagara/Classes/NiagaraDataInterfaceArrayFunctionLibrary.h
+++ b/Engine/Plugins/FX/Niagara/Source/Niagara/Classes/NiagaraDataInterfaceArrayFunctionLibrary.h
@@ -40,6 +40,9 @@ public:
 	/** Sets Niagara Array Int32 Data. */
 	UFUNCTION(BlueprintCallable, Category = Niagara, meta = (DisplayName = "Niagara Set Int32 Array"))
 	static void SetNiagaraArrayInt32(UNiagaraComponent* NiagaraSystem, FName OverrideName, const TArray<int32>& ArrayData);
+	/** Sets Niagara Array UInt8 Data. */
+	UFUNCTION(BlueprintCallable, Category = Niagara, meta = (DisplayName = "Niagara Set UInt8 Array"))
+	static void SetNiagaraArrayUInt8(UNiagaraComponent* NiagaraSystem, FName OverrideName, const TArray<int32>& ArrayData);
 	/** Sets Niagara Array Bool Data. */
 	UFUNCTION(BlueprintCallable, Category = Niagara, meta = (DisplayName = "Niagara Set Bool Array"))
 	static void SetNiagaraArrayBool(UNiagaraComponent* NiagaraSystem, FName OverrideName, const TArray<bool>& ArrayData);
@@ -68,6 +71,9 @@ public:
 	/** Gets a copy of Niagara Int32 Data. */
 	UFUNCTION(BlueprintCallable, Category = Niagara, meta = (DisplayName = "Niagara Get Int32 Array"))
 	static TArray<int32> GetNiagaraArrayInt32(UNiagaraComponent* NiagaraSystem, FName OverrideName);
+	/** Gets a copy of Niagara UInt8 Data. */
+	UFUNCTION(BlueprintCallable, Category = Niagara, meta = (DisplayName = "Niagara Get UInt8 Array"))
+	static TArray<int32> GetNiagaraArrayUInt8(UNiagaraComponent* NiagaraSystem, FName OverrideName);
 	/** Gets a copy of Niagara Bool Data. */
 	UFUNCTION(BlueprintCallable, Category = Niagara, meta = (DisplayName = "Niagara Get Bool Array"))
 	static TArray<bool> GetNiagaraArrayBool(UNiagaraComponent* NiagaraSystem, FName OverrideName);
@@ -95,6 +101,9 @@ public:
 	/** Sets a single value within a Niagara Array Int32. */
 	UFUNCTION(BlueprintCallable, Category = Niagara, meta = (DisplayName = "Niagara Set Int32 Array Value"))
 	static void SetNiagaraArrayInt32Value(UNiagaraComponent* NiagaraSystem, FName OverrideName, int Index, int32 Value, bool bSizeToFit);
+	/** Sets a single value within a Niagara Array UInt8. */
+	UFUNCTION(BlueprintCallable, Category = Niagara, meta = (DisplayName = "Niagara Set UInt8 Array Value"))
+	static void SetNiagaraArrayUInt8Value(UNiagaraComponent* NiagaraSystem, FName OverrideName, int Index, int32 Value, bool bSizeToFit);
 	/** Sets a single value within a Niagara Array Bool. */
 	UFUNCTION(BlueprintCallable, Category = Niagara, meta = (DisplayName = "Niagara Set Bool Array Value"))
 	static void SetNiagaraArrayBoolValue(UNiagaraComponent* NiagaraSystem, FName OverrideName, int Index, const bool& Value, bool bSizeToFit);
@@ -123,7 +132,18 @@ public:
 	/** Gets a single value within a Niagara Array Int32. */
 	UFUNCTION(BlueprintCallable, Category = Niagara, meta = (DisplayName = "Niagara Get Int32 Array Value"))
 	static int32 GetNiagaraArrayInt32Value(UNiagaraComponent* NiagaraSystem, FName OverrideName, int Index);
+	/** Gets a single value within a Niagara Array UInt8. */
+	UFUNCTION(BlueprintCallable, Category = Niagara, meta = (DisplayName = "Niagara Get UInt8 Array Value"))
+	static int32 GetNiagaraArrayUInt8Value(UNiagaraComponent* NiagaraSystem, FName OverrideName, int Index);
 	/** Gets a single value within a Niagara Array Bool. */
 	UFUNCTION(BlueprintCallable, Category = Niagara, meta = (DisplayName = "Niagara Get Bool Array Value"))
 	static bool GetNiagaraArrayBoolValue(UNiagaraComponent* NiagaraSystem, FName OverrideName, int Index);
+
+	//////////////////////////////////////////////////////////////////////////////////////////////////////////////////
+	// None BP compatable set functions
+	static void SetNiagaraArrayVector2D(UNiagaraComponent* NiagaraSystem, FName OverrideName, TConstArrayView<FVector2f> ArrayData);
+	static void SetNiagaraArrayVector(UNiagaraComponent* NiagaraSystem, FName OverrideName, TConstArrayView<FVector3f> ArrayData);
+	static void SetNiagaraArrayVector4(UNiagaraComponent* NiagaraSystem, FName OverrideName, TConstArrayView<FVector4f> ArrayData);
+	static void SetNiagaraArrayQuat(UNiagaraComponent* NiagaraSystem, FName OverrideName, TConstArrayView<FQuat4f> ArrayData);
+	static void SetNiagaraArrayUInt8(UNiagaraComponent* NiagaraSystem, FName OverrideName, TConstArrayView<uint8> ArrayData);
 };
diff --git a/Engine/Plugins/FX/Niagara/Source/Niagara/Classes/NiagaraDataInterfaceArrayImpl.h b/Engine/Plugins/FX/Niagara/Source/Niagara/Classes/NiagaraDataInterfaceArrayImpl.h
--- a/Engine/Plugins/FX/Niagara/Source/Niagara/Classes/NiagaraDataInterfaceArrayImpl.h
+++ b/Engine/Plugins/FX/Niagara/Source/Niagara/Classes/NiagaraDataInterfaceArrayImpl.h
@@ -11,12 +11,45 @@
 #include "NiagaraSystemInstance.h"
 
 #include "Async/Async.h"
+#include "Containers/ArrayView.h"
 #include "ShaderParameterUtils.h"
 #include "ShaderCompilerCore.h"
 
 //////////////////////////////////////////////////////////////////////////
 // Helpers
 
+#define NDIARRAY_GENERATE_BODY(CLASSNAME, TYPENAME, MEMBERNAME) \
+	using FProxyType = FNDIArrayProxyImpl<TYPENAME, CLASSNAME>; \
+	virtual void PostInitProperties() override \
+	{ \
+		Proxy.Reset(new FProxyType(this)); \
+		Super::PostInitProperties(); \
+	} \
+	TArray<TYPENAME>& GetArrayReference() { return MEMBERNAME; }
+
+#if WITH_EDITORONLY_DATA
+	#define NDIARRAY_GENERATE_BODY_LWC(CLASSNAME, TYPENAME, MEMBERNAME) \
+		using FProxyType = FNDIArrayProxyImpl<TYPENAME, CLASSNAME>; \
+		virtual void PostInitProperties() override \
+		{ \
+			Super::PostInitProperties(); \
+			Proxy.Reset(new FProxyType(this)); \
+		} \
+		virtual void PostLoad() override \
+		{ \
+			Super::PostLoad(); \
+			GetProxyAs<FProxyType>()->template SetArrayData<decltype(MEMBERNAME)::ElementType>(MakeArrayView(MEMBERNAME)); \
+		} \
+		virtual void PostEditChangeProperty(struct FPropertyChangedEvent& PropertyChangedEvent) override \
+		{ \
+			Super::PostEditChangeProperty(PropertyChangedEvent); \
+			GetProxyAs<FProxyType>()->template SetArrayData<decltype(MEMBERNAME)::ElementType>(MakeArrayView(MEMBERNAME)); \
+		} \
+		TArray<TYPENAME>& GetArrayReference() { return Internal##MEMBERNAME; }
+#else
+	#define NDIARRAY_GENERATE_BODY_LWC(CLASSNAME, TYPENAME, MEMBERNAME) NDIARRAY_GENERATE_BODY(CLASSNAME, TYPENAME, Internal##MEMBERNAME)
+#endif
+
 template<typename TArrayType>
 struct FNDIArrayImplHelperBase
 {
@@ -36,12 +69,17 @@ struct FNDIArrayImplHelperBase
 	//static const FNiagaraTypeDefinition& GetTypeDefinition() { return FNiagaraTypeDefinition::GetFloatDef(); }
 	//static const TArrayType GetDefaultValue();
 
-	static void CopyToGpuMemory(void* Dest, const TArrayType* Src, int32 NumElements)
+	static void CopyCpuToCpuMemory(TArrayType* Dest, const TArrayType* Src, int32 NumElements)
 	{
 		FMemory::Memcpy(Dest, Src, NumElements * sizeof(TArrayType));
 	}
 
-	static void CopyToCpuMemory(void* Dest, const void* Src, int32 NumElements)
+	static void CopyCpuToGpuMemory(void* Dest, const TArrayType* Src, int32 NumElements)
+	{
+		FMemory::Memcpy(Dest, Src, NumElements * sizeof(TArrayType));
+	}
+
+	static void CopyGpuToCpuMemory(void* Dest, const void* Src, int32 NumElements)
 	{
 		FMemory::Memcpy(Dest, Src, NumElements * sizeof(TArrayType));
 	}
@@ -170,11 +208,11 @@ struct FNDIArrayInstanceData_RenderThread
 			uint8* GPUMemory = reinterpret_cast<uint8*>(RHILockBuffer(ArrayBuffer, 0, ArrayNumBytes, RLM_WriteOnly));
 			if (InArrayData.Num() > 0)
 			{
-				T::CopyToGpuMemory(GPUMemory, InArrayData.GetData(), InArrayData.Num());
+				T::CopyCpuToGpuMemory(GPUMemory, InArrayData.GetData(), InArrayData.Num());
 			}
 
 			const TArrayType DefaultValue = TArrayType(FNDIArrayImplHelper<TArrayType>::GetDefaultValue());
-			T::CopyToGpuMemory(GPUMemory + (sizeof(TVMArrayType) * NumElements), &DefaultValue, 1);
+			T::CopyCpuToGpuMemory(GPUMemory + (sizeof(TVMArrayType) * NumElements), &DefaultValue, 1);
 
 			RHIUnlockBuffer(ArrayBuffer);
 		}
@@ -427,7 +465,7 @@ struct FNDIArrayProxyImpl : public INDIArrayProxyBase
 				if ( NumElements > 0 )
 				{
 					ArrayData.AddUninitialized(NumElements);
-					FNDIArrayImplHelper<TArrayType>::CopyToCpuMemory(ArrayData.GetData(), reinterpret_cast<const TVMArrayType*>(ReadbackData[1].Key), NumElements);
+					FNDIArrayImplHelper<TArrayType>::CopyGpuToCpuMemory(ArrayData.GetData(), reinterpret_cast<const TVMArrayType*>(ReadbackData[1].Key), NumElements);
 				}
 
 				AsyncTask(
@@ -440,7 +478,7 @@ struct FNDIArrayProxyImpl : public INDIArrayProxyBase
 							return;
 						}
 
-						Proxy->SetArrayData(SystemInstanceID, ArrayData);
+						Proxy->SetInstanceArrayData(SystemInstanceID, ArrayData);
 					}
 				);
 			}
@@ -450,11 +488,13 @@ struct FNDIArrayProxyImpl : public INDIArrayProxyBase
 
 	//////////////////////////////////////////////////////////////////////////
 	// BP user parameter accessors, should remove if we every start to share the object between instances
-	void SetArrayData(const TArray<TArrayType>& InArrayData)
+	template<typename TFromArrayType>
+	void SetArrayData(TConstArrayView<TFromArrayType> InArrayData)
 	{
 		if (PerInstanceData_GameThread.IsEmpty())
 		{
-			Owner->GetArrayReference() = InArrayData;
+			Owner->GetArrayReference().SetNumUninitialized(InArrayData.Num());
+			FNDIArrayImplHelper<TArrayType>::CopyCpuToCpuMemory(Owner->GetArrayReference().GetData(), InArrayData.GetData(), InArrayData.Num());
 		}
 		else
 		{
@@ -464,19 +504,25 @@ struct FNDIArrayProxyImpl : public INDIArrayProxyBase
 			InstanceData->bIsModified = false;
 			InstanceData->bIsRenderDirty = bShouldSyncToGpu;
 			InstanceData->ArrayData.Empty();
-			Owner->GetArrayReference() = InArrayData;
+			Owner->GetArrayReference().SetNum(InArrayData.Num());
+			FNDIArrayImplHelper<TArrayType>::CopyCpuToCpuMemory(Owner->GetArrayReference().GetData(), InArrayData.GetData(), InArrayData.Num());
 		}
 	}
 
-	TArray<TArrayType> GetArrayData()
+	template<typename TToArrayType>
+	TArray<TToArrayType> GetArrayDataCopy()
 	{
 		ensure(PerInstanceData_GameThread.Num() <= 1);
 		FNDIArrayInstanceData_GameThread<TArrayType>* InstanceData = PerInstanceData_GameThread.IsEmpty() ? nullptr : PerInstanceData_GameThread.CreateConstIterator().Value();
 		FReadArrayRef ArrayRef(Owner, InstanceData);
-		return ArrayRef.GetArray();
+		TArray<TToArrayType> OutArray;
+		OutArray.SetNum(ArrayRef.GetArray().Num());
+		FNDIArrayImplHelper<TArrayType>::CopyCpuToCpuMemory(OutArray.GetData(), ArrayRef.GetArray().GetData(), ArrayRef.GetArray().Num());
+		return OutArray;
 	}
 
-	void SetArrayValue(int Index, const TArrayType& Value, bool bSizeToFit)
+	template<typename TFromArrayType>
+	void SetArrayValue(int Index, const TFromArrayType& Value, bool bSizeToFit)
 	{
 		ensure(PerInstanceData_GameThread.Num() <= 1);
 		FNDIArrayInstanceData_GameThread<TArrayType>* InstanceData = PerInstanceData_GameThread.IsEmpty() ? nullptr : PerInstanceData_GameThread.CreateConstIterator().Value();
@@ -490,10 +536,11 @@ struct FNDIArrayProxyImpl : public INDIArrayProxyBase
 			ArrayRef.GetArray().AddDefaulted(Index + 1 - ArrayRef.GetArray().Num());
 		}
 
-		ArrayRef.GetArray()[Index] = Value;
+		FNDIArrayImplHelper<TArrayType>::CopyCpuToCpuMemory(ArrayRef.GetArray().GetData() + Index, &Value, 1);
 	}
 
-	TArrayType GetArrayValue(int Index)
+	template<typename TToArrayType>
+	TToArrayType GetArrayValue(int Index)
 	{
 		TArrayType ValueOut = TArrayType(FNDIArrayImplHelper<TArrayType>::GetDefaultValue());
 
@@ -506,12 +553,14 @@ struct FNDIArrayProxyImpl : public INDIArrayProxyBase
 			ValueOut = ArrayRef.GetArray()[Index];
 		}
 
-		return ValueOut;
+		TToArrayType ToValueOut;
+		FNDIArrayImplHelper<TArrayType>::CopyCpuToCpuMemory(&ToValueOut, &ValueOut, 1);
+		return ToValueOut;
 	}
 
 	//////////////////////////////////////////////////////////////////////////
 	// VM accessors to ensure we maintain per correctness for shared data interfaces
-	void SetArrayData(FNiagaraSystemInstanceID InstanceID, const TArray<TArrayType>& InArrayData)
+	void SetInstanceArrayData(FNiagaraSystemInstanceID InstanceID, const TArray<TArrayType>& InArrayData)
 	{
 		if ( FNDIArrayInstanceData_GameThread<TArrayType>* InstanceData = PerInstanceData_GameThread.FindRef(InstanceID) )
 		{
diff --git a/Engine/Plugins/FX/Niagara/Source/Niagara/Classes/NiagaraDataInterfaceArrayInt.h b/Engine/Plugins/FX/Niagara/Source/Niagara/Classes/NiagaraDataInterfaceArrayInt.h
--- a/Engine/Plugins/FX/Niagara/Source/Niagara/Classes/NiagaraDataInterfaceArrayInt.h
+++ b/Engine/Plugins/FX/Niagara/Source/Niagara/Classes/NiagaraDataInterfaceArrayInt.h
@@ -24,6 +24,45 @@ struct FNDIArrayImplHelper<int32> : public FNDIArrayImplHelperBase<int32>
 	static const int32 GetDefaultValue() { return 0; }
 };
 
+template<>
+struct FNDIArrayImplHelper<uint8> : public FNDIArrayImplHelperBase<uint8>
+{
+	typedef int32 TVMArrayType;
+
+	static constexpr TCHAR const* HLSLVariableType = TEXT("int");
+	static constexpr EPixelFormat ReadPixelFormat = PF_R8_UINT;
+	static constexpr TCHAR const* ReadHLSLBufferType = TEXT("int");
+	static constexpr TCHAR const* ReadHLSLBufferRead = TEXT("Value = BUFFER_NAME[Index];");
+	static constexpr EPixelFormat RWPixelFormat = PF_R8_UINT;
+	static constexpr TCHAR const* RWHLSLBufferType = TEXT("int");
+	static constexpr TCHAR const* RWHLSLBufferRead = TEXT("Value = BUFFER_NAME[Index];");
+	static constexpr TCHAR const* RWHLSLBufferWrite = TEXT("BUFFER_NAME[Index] = Value;");
+
+	static const FNiagaraTypeDefinition& GetTypeDefinition() { return FNiagaraTypeDefinition::GetIntDef(); }
+	static const int32 GetDefaultValue() { return 0; }
+
+	static void CopyCpuToCpuMemory(uint8* Dest, const uint8* Src, int32 NumElements)
+	{
+		FMemory::Memcpy(Dest, Src, NumElements * sizeof(uint8));
+	}
+
+	static void CopyCpuToCpuMemory(uint8* Dest, const int32* Src, int32 NumElements)
+	{
+		for (int32 i = 0; i < NumElements; ++i)
+		{
+			Dest[i] = uint8(Src[i]);
+		}
+	}
+
+	static void CopyCpuToCpuMemory(int32* Dest, const uint8* Src, int32 NumElements)
+	{
+		for (int32 i = 0; i < NumElements; ++i)
+		{
+			Dest[i] = int32(Src[i]);
+		}
+	}
+};
+
 template<>
 struct FNDIArrayImplHelper<bool> : public FNDIArrayImplHelperBase<bool>
 {
@@ -41,7 +80,7 @@ struct FNDIArrayImplHelper<bool> : public FNDIArrayImplHelperBase<bool>
 	static const FNiagaraTypeDefinition& GetTypeDefinition() { return FNiagaraTypeDefinition::GetBoolDef(); }
 	static const FNiagaraBool GetDefaultValue() { return FNiagaraBool::False; }
 
-	static void CopyToGpuMemory(void* Dest, const bool* Src, int32 NumElements)
+	static void CopyCpuToGpuMemory(void* Dest, const bool* Src, int32 NumElements)
 	{
 		uint8* TypedDest = reinterpret_cast<uint8*>(Dest);
 		while (NumElements--)
@@ -50,7 +89,7 @@ struct FNDIArrayImplHelper<bool> : public FNDIArrayImplHelperBase<bool>
 		}
 	}
 
-	static void CopyToCpuMemory(void* Dest, const void* Src, int32 NumElements)
+	static void CopyGpuToCpuMemory(void* Dest, const void* Src, int32 NumElements)
 	{
 		FNiagaraBool* TypedDest = reinterpret_cast<FNiagaraBool*>(Dest);
 		const uint8* TypedSrc = reinterpret_cast<const uint8*>(Src);
@@ -61,24 +100,40 @@ struct FNDIArrayImplHelper<bool> : public FNDIArrayImplHelperBase<bool>
 	}
 };
 
-UCLASS(EditInlineNew, Category = "Array", meta = (DisplayName = "Int32 Array", Experimental), Blueprintable, BlueprintType)
+UCLASS(EditInlineNew, Category = "Array", meta = (DisplayName = "Int32 Array"), Blueprintable, BlueprintType)
 class NIAGARA_API UNiagaraDataInterfaceArrayInt32 : public UNiagaraDataInterfaceArray
 {
-	GENERATED_UCLASS_BODY()
 public:
+	GENERATED_BODY()
+	NDIARRAY_GENERATE_BODY(UNiagaraDataInterfaceArrayInt32, int32, IntData)
+
+	UPROPERTY(EditAnywhere, BlueprintReadWrite, Category = "Array")
+	TArray<int32> IntData;
+};
+
+UCLASS(EditInlineNew, Category = "Array", meta = (DisplayName = "UInt8 Array"), Blueprintable, BlueprintType)
+class NIAGARA_API UNiagaraDataInterfaceArrayUInt8 : public UNiagaraDataInterfaceArray
+{
+public:
+	GENERATED_BODY()
+	NDIARRAY_GENERATE_BODY_LWC(UNiagaraDataInterfaceArrayUInt8, uint8, IntData)
+
+#if WITH_EDITORONLY_DATA
 	UPROPERTY(EditAnywhere, BlueprintReadWrite, Category = "Array")
 	TArray<int32> IntData;
+#endif
 
-	TArray<int32>& GetArrayReference() { return IntData; }
+	UPROPERTY()
+	TArray<uint8> InternalIntData;
 };
 
-UCLASS(EditInlineNew, Category = "Array", meta = (DisplayName = "Bool Array", Experimental), Blueprintable, BlueprintType)
+UCLASS(EditInlineNew, Category = "Array", meta = (DisplayName = "Bool Array"), Blueprintable, BlueprintType)
 class NIAGARA_API UNiagaraDataInterfaceArrayBool : public UNiagaraDataInterfaceArray
 {
-	GENERATED_UCLASS_BODY()
 public:
+	GENERATED_BODY()
+	NDIARRAY_GENERATE_BODY(UNiagaraDataInterfaceArrayBool, bool, BoolData)
+
 	UPROPERTY(EditAnywhere, BlueprintReadWrite, Category = "Array")
 	TArray<bool> BoolData;
-
-	TArray<bool>& GetArrayReference() { return BoolData; }
 };
diff --git a/Engine/Plugins/FX/Niagara/Source/Niagara/Classes/NiagaraDataInterfaceArrayNiagaraID.h b/Engine/Plugins/FX/Niagara/Source/Niagara/Classes/NiagaraDataInterfaceArrayNiagaraID.h
--- a/Engine/Plugins/FX/Niagara/Source/Niagara/Classes/NiagaraDataInterfaceArrayNiagaraID.h
+++ b/Engine/Plugins/FX/Niagara/Source/Niagara/Classes/NiagaraDataInterfaceArrayNiagaraID.h
@@ -3,17 +3,35 @@
 #pragma once
 
 #include "NiagaraDataInterfaceArray.h"
+#include "NiagaraDataInterfaceArrayImpl.h"
 #include "NiagaraDataInterfaceArrayNiagaraID.generated.h"
 
+template<>
+struct FNDIArrayImplHelper<FNiagaraID> : public FNDIArrayImplHelperBase<FNiagaraID>
+{
+	typedef FNiagaraID TVMArrayType;
+
+	static constexpr TCHAR const* HLSLVariableType = TEXT("NiagaraID");
+	static constexpr EPixelFormat ReadPixelFormat = PF_R32_SINT;		// Should be int2 but we don't have it
+	static constexpr TCHAR const* ReadHLSLBufferType = TEXT("int");
+	static constexpr TCHAR const* ReadHLSLBufferRead = TEXT("Value.Index = BUFFER_NAME[Index * 2 + 0]; Value.AcquireTag = BUFFER_NAME[Index * 2 + 1];");
+	static constexpr EPixelFormat RWPixelFormat = PF_R32_SINT;
+	static constexpr TCHAR const* RWHLSLBufferType = TEXT("int");
+	static constexpr TCHAR const* RWHLSLBufferRead = TEXT("Value.Index = BUFFER_NAME[Index * 2 + 0]; Value.AcquireTag = BUFFER_NAME[Index * 2 + 1];");
+	static constexpr TCHAR const* RWHLSLBufferWrite = TEXT("BUFFER_NAME[Index * 2 + 0] = Value.Index; BUFFER_NAME[Index * 2 + 1] = Value.AcquireTag;");
+
+	static const FNiagaraTypeDefinition GetTypeDefinition() { return FNiagaraTypeDefinition(FNiagaraID::StaticStruct()); }
+	static const FNiagaraID GetDefaultValue() { return FNiagaraID(); }
+};
+
 UCLASS(EditInlineNew, Category = "Array", meta = (DisplayName = "NiagaraID Array", Experimental), Blueprintable, BlueprintType)
 class NIAGARA_API UNiagaraDataInterfaceArrayNiagaraID : public UNiagaraDataInterfaceArray
 {
-	GENERATED_UCLASS_BODY()
 public:
+	GENERATED_BODY()
+	NDIARRAY_GENERATE_BODY(UNiagaraDataInterfaceArrayNiagaraID, FNiagaraID, IntData)
 
 	UPROPERTY(EditAnywhere, BlueprintReadWrite, Category = "Array")
 	TArray<FNiagaraID> IntData;
-
-	TArray<FNiagaraID>& GetArrayReference() { return IntData; }
 };
 
diff --git a/Engine/Plugins/FX/Niagara/Source/Niagara/Private/NiagaraDataInterfaceArrayFloat.cpp b/Engine/Plugins/FX/Niagara/Source/Niagara/Private/NiagaraDataInterfaceArrayFloat.cpp
--- a/Engine/Plugins/FX/Niagara/Source/Niagara/Private/NiagaraDataInterfaceArrayFloat.cpp
+++ b/Engine/Plugins/FX/Niagara/Source/Niagara/Private/NiagaraDataInterfaceArrayFloat.cpp
@@ -2,45 +2,3 @@
 
 #include "NiagaraDataInterfaceArrayFloat.h"
 #include "NiagaraRenderer.h"
-
-UNiagaraDataInterfaceArrayFloat::UNiagaraDataInterfaceArrayFloat(FObjectInitializer const& ObjectInitializer)
-	: UNiagaraDataInterfaceArray(ObjectInitializer)
-{
-	Proxy.Reset(new FNDIArrayProxyImpl<float, UNiagaraDataInterfaceArrayFloat>(this));
-}
-
-UNiagaraDataInterfaceArrayFloat2::UNiagaraDataInterfaceArrayFloat2(FObjectInitializer const& ObjectInitializer)
-	: UNiagaraDataInterfaceArray(ObjectInitializer)
-{
-	Proxy.Reset(new FNDIArrayProxyImpl<FVector2D, UNiagaraDataInterfaceArrayFloat2>(this));
-}
-
-UNiagaraDataInterfaceArrayFloat3::UNiagaraDataInterfaceArrayFloat3(FObjectInitializer const& ObjectInitializer)
-	: UNiagaraDataInterfaceArray(ObjectInitializer)
-{
-	Proxy.Reset(new FNDIArrayProxyImpl<FVector, UNiagaraDataInterfaceArrayFloat3>(this));
-}
-
-UNiagaraDataInterfaceArrayPosition::UNiagaraDataInterfaceArrayPosition(FObjectInitializer const& ObjectInitializer)
-	: UNiagaraDataInterfaceArray(ObjectInitializer)
-{
-	Proxy.Reset(new FNDIArrayProxyImpl<FNiagaraPosition, UNiagaraDataInterfaceArrayPosition>(this));
-}
-
-UNiagaraDataInterfaceArrayFloat4::UNiagaraDataInterfaceArrayFloat4(FObjectInitializer const& ObjectInitializer)
-	: UNiagaraDataInterfaceArray(ObjectInitializer)
-{
-	Proxy.Reset(new FNDIArrayProxyImpl<FVector4, UNiagaraDataInterfaceArrayFloat4>(this));
-}
-
-UNiagaraDataInterfaceArrayColor::UNiagaraDataInterfaceArrayColor(FObjectInitializer const& ObjectInitializer)
-	: UNiagaraDataInterfaceArray(ObjectInitializer)
-{
-	Proxy.Reset(new FNDIArrayProxyImpl<FLinearColor, UNiagaraDataInterfaceArrayColor>(this));
-}
-
-UNiagaraDataInterfaceArrayQuat::UNiagaraDataInterfaceArrayQuat(FObjectInitializer const& ObjectInitializer)
-	: UNiagaraDataInterfaceArray(ObjectInitializer)
-{
-	Proxy.Reset(new FNDIArrayProxyImpl<FQuat, UNiagaraDataInterfaceArrayQuat>(this));
-}
diff --git a/Engine/Plugins/FX/Niagara/Source/Niagara/Private/NiagaraDataInterfaceArrayFunctionLibrary.cpp b/Engine/Plugins/FX/Niagara/Source/Niagara/Private/NiagaraDataInterfaceArrayFunctionLibrary.cpp
--- a/Engine/Plugins/FX/Niagara/Source/Niagara/Private/NiagaraDataInterfaceArrayFunctionLibrary.cpp
+++ b/Engine/Plugins/FX/Niagara/Source/Niagara/Private/NiagaraDataInterfaceArrayFunctionLibrary.cpp
@@ -23,19 +23,21 @@ void SetNiagaraVariantArray(UNiagaraComponent* NiagaraComponent, FName OverrideN
 	}
 
 	TDataInterace* VariantDataInterface = CastChecked<TDataInterace>(DuplicateObject(ExistingDataInterface, NiagaraComponent));
-	auto* ExistingProxy = static_cast<FNDIArrayProxyImpl<TArrayType, TDataInterace>*>(ExistingDataInterface->GetProxy());
-	VariantDataInterface->GetArrayReference() = ExistingProxy->GetArrayData();
+	auto* ExistingProxy = static_cast<typename TDataInterace::FProxyType*>(ExistingDataInterface->GetProxy());
+	auto* VariantProxy = static_cast<typename TDataInterace::FProxyType*>(VariantDataInterface->GetProxy());
+	TArray<TArrayType> TempArrayCopy = ExistingProxy->template GetArrayDataCopy<TArrayType>();
+	VariantProxy->template SetArrayData<TArrayType>(MakeArrayView<TArrayType>(TempArrayCopy.GetData(), TempArrayCopy.Num()));
 	NiagaraComponent->SetParameterOverride(FNiagaraVariableBase(FNiagaraTypeDefinition(TDataInterace::StaticClass()), OverrideName), FNiagaraVariant(VariantDataInterface));
 #endif
 }
 
 // If / when we share user parameter UObjects we will need to make this per instance which introduces some tricky things about allocating before the instance is active
 template<typename TArrayType, typename TDataInterace>
-void SetNiagaraArray(UNiagaraComponent* NiagaraComponent, FName OverrideName, const TArray<TArrayType>& InArray)
+void SetNiagaraArray(UNiagaraComponent* NiagaraComponent, FName OverrideName, TConstArrayView<TArrayType> InArray)
 {
 	if (TDataInterace* ArrayDI = UNiagaraFunctionLibrary::GetDataInterface<TDataInterace>(NiagaraComponent, OverrideName))
 	{
-		auto* ArrayProxy = static_cast<FNDIArrayProxyImpl<TArrayType, TDataInterace>*>(ArrayDI->GetProxy());
+		auto* ArrayProxy = static_cast<typename TDataInterace::FProxyType*>(ArrayDI->GetProxy());
 		ArrayProxy->SetArrayData(InArray);
 		SetNiagaraVariantArray<TArrayType, TDataInterace>(NiagaraComponent, OverrideName, ArrayDI);
 	}
@@ -46,8 +48,8 @@ TArray<TArrayType> GetNiagaraArray(UNiagaraComponent* NiagaraComponent, FName Ov
 {
 	if (TDataInterace* ArrayDI = UNiagaraFunctionLibrary::GetDataInterface<TDataInterace>(NiagaraComponent, OverrideName))
 	{
-		auto* ArrayProxy = static_cast<FNDIArrayProxyImpl<TArrayType, TDataInterace>*>(ArrayDI->GetProxy());
-		return ArrayProxy->GetArrayData();
+		auto* ArrayProxy = static_cast<typename TDataInterace::FProxyType*>(ArrayDI->GetProxy());
+		return ArrayProxy->template GetArrayDataCopy<TArrayType>();
 	}
 	return TArray<TArrayType>();
 }
@@ -57,7 +59,7 @@ void SetNiagaraArrayValue(UNiagaraComponent* NiagaraComponent, FName OverrideNam
 {
 	if (TDataInterace* ArrayDI = UNiagaraFunctionLibrary::GetDataInterface<TDataInterace>(NiagaraComponent, OverrideName))
 	{
-		auto* ArrayProxy = static_cast<FNDIArrayProxyImpl<TArrayType, TDataInterace>*>(ArrayDI->GetProxy());
+		auto* ArrayProxy = static_cast<typename TDataInterace::FProxyType*>(ArrayDI->GetProxy());
 		ArrayProxy->SetArrayValue(Index, Value, bSizeToFit);
 		SetNiagaraVariantArray<TArrayType, TDataInterace>(NiagaraComponent, OverrideName, ArrayDI);
 	}
@@ -68,8 +70,8 @@ TArrayType GetNiagaraArrayValue(UNiagaraComponent* NiagaraComponent, FName Overr
 {
 	if (TDataInterace* ArrayDI = UNiagaraFunctionLibrary::GetDataInterface<TDataInterace>(NiagaraComponent, OverrideName))
 	{
-		auto* ArrayProxy = static_cast<FNDIArrayProxyImpl<TArrayType, TDataInterace>*>(ArrayDI->GetProxy());
-		return ArrayProxy->GetArrayValue(Index);
+		auto* ArrayProxy = static_cast<typename TDataInterace::FProxyType*>(ArrayDI->GetProxy());
+		return ArrayProxy->template GetArrayValue<TArrayType>(Index);
 	}
 	//-TODO: Should be DefaultValue
 	return TArrayType();
@@ -136,6 +138,11 @@ void UNiagaraDataInterfaceArrayFunctionLibrary::SetNiagaraArrayInt32(UNiagaraCom
 	SetNiagaraArray<int32, UNiagaraDataInterfaceArrayInt32>(NiagaraComponent, OverrideName, ArrayData);
 }
 
+void UNiagaraDataInterfaceArrayFunctionLibrary::SetNiagaraArrayUInt8(UNiagaraComponent* NiagaraComponent, FName OverrideName, const TArray<int32>& ArrayData)
+{
+	SetNiagaraArray<int32, UNiagaraDataInterfaceArrayUInt8>(NiagaraComponent, OverrideName, ArrayData);
+}
+
 void UNiagaraDataInterfaceArrayFunctionLibrary::SetNiagaraArrayBool(UNiagaraComponent* NiagaraComponent, FName OverrideName, const TArray<bool>& ArrayData)
 {
 	SetNiagaraArray<bool, UNiagaraDataInterfaceArrayBool>(NiagaraComponent, OverrideName, ArrayData);
@@ -193,6 +200,11 @@ TArray<int32> UNiagaraDataInterfaceArrayFunctionLibrary::GetNiagaraArrayInt32(UN
 	return GetNiagaraArray<int32, UNiagaraDataInterfaceArrayInt32>(NiagaraComponent, OverrideName);
 }
 
+TArray<int32> UNiagaraDataInterfaceArrayFunctionLibrary::GetNiagaraArrayUInt8(UNiagaraComponent* NiagaraComponent, FName OverrideName)
+{
+	return GetNiagaraArray<int32, UNiagaraDataInterfaceArrayUInt8>(NiagaraComponent, OverrideName);
+}
+
 TArray<bool> UNiagaraDataInterfaceArrayFunctionLibrary::GetNiagaraArrayBool(UNiagaraComponent* NiagaraComponent, FName OverrideName)
 {
 	return GetNiagaraArray<bool, UNiagaraDataInterfaceArrayBool>(NiagaraComponent, OverrideName);
@@ -242,6 +254,11 @@ void UNiagaraDataInterfaceArrayFunctionLibrary::SetNiagaraArrayInt32Value(UNiaga
 	SetNiagaraArrayValue<int32, UNiagaraDataInterfaceArrayInt32>(NiagaraComponent, OverrideName, Index, Value, bSizeToFit);
 }
 
+void UNiagaraDataInterfaceArrayFunctionLibrary::SetNiagaraArrayUInt8Value(UNiagaraComponent* NiagaraComponent, FName OverrideName, int Index, int32 Value, bool bSizeToFit)
+{
+	SetNiagaraArrayValue<int32, UNiagaraDataInterfaceArrayUInt8>(NiagaraComponent, OverrideName, Index, Value, bSizeToFit);
+}
+
 void UNiagaraDataInterfaceArrayFunctionLibrary::SetNiagaraArrayBoolValue(UNiagaraComponent* NiagaraComponent, FName OverrideName, int Index, const bool& Value, bool bSizeToFit)
 {
 	SetNiagaraArrayValue<bool, UNiagaraDataInterfaceArrayBool>(NiagaraComponent, OverrideName, Index, Value, bSizeToFit);
@@ -291,7 +308,39 @@ int32 UNiagaraDataInterfaceArrayFunctionLibrary::GetNiagaraArrayInt32Value(UNiag
 	return GetNiagaraArrayValue<int32, UNiagaraDataInterfaceArrayInt32>(NiagaraComponent, OverrideName, Index);
 }
 
+int32 UNiagaraDataInterfaceArrayFunctionLibrary::GetNiagaraArrayUInt8Value(UNiagaraComponent* NiagaraComponent, FName OverrideName, int Index)
+{
+	return GetNiagaraArrayValue<int32, UNiagaraDataInterfaceArrayUInt8>(NiagaraComponent, OverrideName, Index);
+}
+
 bool UNiagaraDataInterfaceArrayFunctionLibrary::GetNiagaraArrayBoolValue(UNiagaraComponent* NiagaraComponent, FName OverrideName, int Index)
 {
 	return GetNiagaraArrayValue<bool, UNiagaraDataInterfaceArrayBool>(NiagaraComponent, OverrideName, Index);
 }
+
+//////////////////////////////////////////////////////////////////////////////////////////////////////////////////
+// None BP compatable set functions
+void UNiagaraDataInterfaceArrayFunctionLibrary::SetNiagaraArrayVector2D(UNiagaraComponent* NiagaraSystem, FName OverrideName, TConstArrayView<FVector2f> ArrayData)
+{
+	SetNiagaraArray<FVector2f, UNiagaraDataInterfaceArrayFloat2>(NiagaraSystem, OverrideName, ArrayData);
+}
+
+void UNiagaraDataInterfaceArrayFunctionLibrary::SetNiagaraArrayVector(UNiagaraComponent* NiagaraSystem, FName OverrideName, TConstArrayView<FVector3f> ArrayData)
+{
+	SetNiagaraArray<FVector3f, UNiagaraDataInterfaceArrayFloat3>(NiagaraSystem, OverrideName, ArrayData);
+}
+
+void UNiagaraDataInterfaceArrayFunctionLibrary::SetNiagaraArrayVector4(UNiagaraComponent* NiagaraSystem, FName OverrideName, TConstArrayView<FVector4f> ArrayData)
+{
+	SetNiagaraArray<FVector4f, UNiagaraDataInterfaceArrayFloat4>(NiagaraSystem, OverrideName, ArrayData);
+}
+
+void UNiagaraDataInterfaceArrayFunctionLibrary::SetNiagaraArrayQuat(UNiagaraComponent* NiagaraSystem, FName OverrideName, TConstArrayView<FQuat4f> ArrayData)
+{
+	SetNiagaraArray<FQuat4f, UNiagaraDataInterfaceArrayQuat>(NiagaraSystem, OverrideName, ArrayData);
+}
+
+void UNiagaraDataInterfaceArrayFunctionLibrary::SetNiagaraArrayUInt8(UNiagaraComponent* NiagaraSystem, FName OverrideName, TConstArrayView<uint8> ArrayData)
+{
+	SetNiagaraArray<uint8, UNiagaraDataInterfaceArrayUInt8>(NiagaraSystem, OverrideName, ArrayData);
+}
diff --git a/Engine/Plugins/FX/Niagara/Source/Niagara/Private/NiagaraDataInterfaceArrayInt.cpp b/Engine/Plugins/FX/Niagara/Source/Niagara/Private/NiagaraDataInterfaceArrayInt.cpp
--- a/Engine/Plugins/FX/Niagara/Source/Niagara/Private/NiagaraDataInterfaceArrayInt.cpp
+++ b/Engine/Plugins/FX/Niagara/Source/Niagara/Private/NiagaraDataInterfaceArrayInt.cpp
@@ -3,16 +3,4 @@
 #include "NiagaraDataInterfaceArrayInt.h"
 #include "NiagaraRenderer.h"
 
-UNiagaraDataInterfaceArrayInt32::UNiagaraDataInterfaceArrayInt32(FObjectInitializer const& ObjectInitializer)
-	: UNiagaraDataInterfaceArray(ObjectInitializer)
-{
-	Proxy.Reset(new FNDIArrayProxyImpl<int32, UNiagaraDataInterfaceArrayInt32>(this));
-}
-
-UNiagaraDataInterfaceArrayBool::UNiagaraDataInterfaceArrayBool(FObjectInitializer const& ObjectInitializer)
-	: UNiagaraDataInterfaceArray(ObjectInitializer)
-{
-	static_assert(sizeof(bool) == sizeof(uint8), "Bool != 1 byte this will mean the GPU array does not match in size");
-
-	Proxy.Reset(new FNDIArrayProxyImpl<bool, UNiagaraDataInterfaceArrayBool>(this));
-}
+static_assert(sizeof(bool) == sizeof(uint8), "Bool != 1 byte this will mean the GPU array does not match in size");
diff --git a/Engine/Plugins/FX/Niagara/Source/Niagara/Private/NiagaraDataInterfaceArrayNiagaraID.cpp b/Engine/Plugins/FX/Niagara/Source/Niagara/Private/NiagaraDataInterfaceArrayNiagaraID.cpp
--- a/Engine/Plugins/FX/Niagara/Source/Niagara/Private/NiagaraDataInterfaceArrayNiagaraID.cpp
+++ b/Engine/Plugins/FX/Niagara/Source/Niagara/Private/NiagaraDataInterfaceArrayNiagaraID.cpp
@@ -1,30 +1,5 @@
 // Copyright Epic Games, Inc. All Rights Reserved.
 
 #include "NiagaraDataInterfaceArrayNiagaraID.h"
-#include "NiagaraDataInterfaceArrayImpl.h"
 #include "NiagaraRenderer.h"
 
-template<>
-struct FNDIArrayImplHelper<FNiagaraID> : public FNDIArrayImplHelperBase<FNiagaraID>
-{
-	typedef FNiagaraID TVMArrayType;
-
-	static constexpr TCHAR const* HLSLVariableType		= TEXT("NiagaraID");
-	static constexpr EPixelFormat ReadPixelFormat		= PF_R32_SINT;		// Should be int2 but we don't have it
-	static constexpr TCHAR const* ReadHLSLBufferType	= TEXT("int");
-	static constexpr TCHAR const* ReadHLSLBufferRead	= TEXT("Value.Index = BUFFER_NAME[Index * 2 + 0]; Value.AcquireTag = BUFFER_NAME[Index * 2 + 1];");
-	static constexpr EPixelFormat RWPixelFormat			= PF_R32_SINT;
-	static constexpr TCHAR const* RWHLSLBufferType		= TEXT("int");
-	static constexpr TCHAR const* RWHLSLBufferRead		= TEXT("Value.Index = BUFFER_NAME[Index * 2 + 0]; Value.AcquireTag = BUFFER_NAME[Index * 2 + 1];");
-	static constexpr TCHAR const* RWHLSLBufferWrite		= TEXT("BUFFER_NAME[Index * 2 + 0] = Value.Index; BUFFER_NAME[Index * 2 + 1] = Value.AcquireTag;");
-
-	static const FNiagaraTypeDefinition GetTypeDefinition() { return FNiagaraTypeDefinition(FNiagaraID::StaticStruct()); }
-	static const FNiagaraID GetDefaultValue() { return FNiagaraID(); }
-};
-
-UNiagaraDataInterfaceArrayNiagaraID::UNiagaraDataInterfaceArrayNiagaraID(FObjectInitializer const& ObjectInitializer)
-	: UNiagaraDataInterfaceArray(ObjectInitializer)
-{
-	Proxy.Reset(new FNDIArrayProxyImpl<FNiagaraID, UNiagaraDataInterfaceArrayNiagaraID>(this));
-}
-
