Event spawn scripts can now modify the Initial values for particle attributes so modules dependant on this can function correctly when spawning from events. e.g. ScaleColor.

This is optional per even handler but defaults to enabled. Existing content has this disabled to maintain existing behavior.

diff --git a/Engine/Plugins/FX/Niagara/Source/Niagara/Classes/NiagaraEmitter.h b/Engine/Plugins/FX/Niagara/Source/Niagara/Classes/NiagaraEmitter.h
--- a/Engine/Plugins/FX/Niagara/Source/Niagara/Classes/NiagaraEmitter.h
+++ b/Engine/Plugins/FX/Niagara/Source/Niagara/Classes/NiagaraEmitter.h
@@ -165,6 +165,7 @@ struct FNiagaraEventScriptProperties : public FNiagaraEmitterScriptProperties
 		MaxEventsPerFrame = 0;
 		bRandomSpawnNumber = false;
 		MinSpawnNumber = 0;
+		UpdateAttributeInitialValues = true;
 	}
 	
 	/** Controls which particles have the event script run on them.*/
@@ -194,6 +195,10 @@ struct FNiagaraEventScriptProperties : public FNiagaraEmitterScriptProperties
 	/** The minimum spawn number when random spawn is used. Spawn Number is used as the maximum range. */
 	UPROPERTY(EditAnywhere, Category = "Event Handler Options")
 	uint32 MinSpawnNumber;
+	
+	/** Should Event Spawn Scripts modify the Initial values for particle attributes they modify. */
+	UPROPERTY(EditAnywhere, Category = "Event Handler Options", meta = (EditCondition="ExecutionMode==EScriptExecutionMode::SpawnedParticles"))
+	bool UpdateAttributeInitialValues;
 };
 
 /** Legacy struct for spawn count scale overrides. This is now done in FNiagaraEmitterScalabilityOverrides*/
@@ -344,7 +349,7 @@ struct NIAGARA_API FVersionedNiagaraEmitterData
 	FNiagaraParameterStore RendererBindings;
 
 	void CopyFrom(const FVersionedNiagaraEmitterData& Source);
-	void PostLoad(UNiagaraEmitter& Emitter, bool bIsCooked);
+	void PostLoad(UNiagaraEmitter& Emitter, bool bIsCooked, int32 NiagaraVer);
 	void PostInitProperties(UNiagaraEmitter* Outer);
 	bool UsesCollection(const UNiagaraParameterCollection* Collection) const;
 	bool UsesScript(const UNiagaraScript* Script) const;
diff --git a/Engine/Plugins/FX/Niagara/Source/Niagara/Classes/NiagaraScriptSourceBase.h b/Engine/Plugins/FX/Niagara/Source/Niagara/Classes/NiagaraScriptSourceBase.h
--- a/Engine/Plugins/FX/Niagara/Source/Niagara/Classes/NiagaraScriptSourceBase.h
+++ b/Engine/Plugins/FX/Niagara/Source/Niagara/Classes/NiagaraScriptSourceBase.h
@@ -60,6 +60,7 @@ public:
 	NIAGARA_API static const FString CpuScriptDefine;
 	NIAGARA_API static const FString GpuScriptDefine;
 	NIAGARA_API static const FString EventSpawnDefine;
+	NIAGARA_API static const FString EventSpawnInitialAttribWritesDefine;
 
 	FNiagaraCompileOptions() : TargetUsage(ENiagaraScriptUsage::Function), TargetUsageBitmask(0)
 	{
diff --git a/Engine/Plugins/FX/Niagara/Source/Niagara/Private/NiagaraCommon.cpp b/Engine/Plugins/FX/Niagara/Source/Niagara/Private/NiagaraCommon.cpp
--- a/Engine/Plugins/FX/Niagara/Source/Niagara/Private/NiagaraCommon.cpp
+++ b/Engine/Plugins/FX/Niagara/Source/Niagara/Private/NiagaraCommon.cpp
@@ -1522,6 +1522,7 @@ bool FVMExternalFunctionBindingInfo::Serialize(FArchive& Ar)
 const FString FNiagaraCompileOptions::CpuScriptDefine = TEXT("CPUSim");
 const FString FNiagaraCompileOptions::GpuScriptDefine = TEXT("GPUComputeSim");
 const FString FNiagaraCompileOptions::EventSpawnDefine = TEXT("EventSpawn");
+const FString FNiagaraCompileOptions::EventSpawnInitialAttribWritesDefine = TEXT("EventSpawnInitialAttribWrites");
 
 FSynchronizeWithParameterDefinitionsArgs::FSynchronizeWithParameterDefinitionsArgs()
 	: SpecificDefinitionsUniqueIds(TArray<FGuid>())
diff --git a/Engine/Plugins/FX/Niagara/Source/Niagara/Private/NiagaraEmitter.cpp b/Engine/Plugins/FX/Niagara/Source/Niagara/Private/NiagaraEmitter.cpp
--- a/Engine/Plugins/FX/Niagara/Source/Niagara/Private/NiagaraEmitter.cpp
+++ b/Engine/Plugins/FX/Niagara/Source/Niagara/Private/NiagaraEmitter.cpp
@@ -624,7 +624,7 @@ void UNiagaraEmitter::PostLoad()
 	for (FVersionedNiagaraEmitterData& Data : VersionData)
 	{
 #if WITH_EDITORONLY_DATA
-		Data.PostLoad(*this, IsCooked);
+		Data.PostLoad(*this, IsCooked, NiagaraVer);
 		Data.GPUComputeScript->OnGPUScriptCompiled().RemoveAll(this);
 		Data.GPUComputeScript->OnGPUScriptCompiled().AddUObject(this, &UNiagaraEmitter::RaiseOnEmitterGPUCompiled);
 		if (UE5MainVer < FUE5MainStreamObjectVersion::FixGpuAlwaysRunningUpdateScriptNoneInterpolated)
@@ -639,7 +639,7 @@ void UNiagaraEmitter::PostLoad()
 			}
 		}
 #else
-		Data.PostLoad(*this, true);
+		Data.PostLoad(*this, true, NiagaraVer);
 #endif
 	}
 
@@ -650,13 +650,25 @@ void UNiagaraEmitter::PostLoad()
 #endif
 }
 
-void FVersionedNiagaraEmitterData::PostLoad(UNiagaraEmitter& Emitter, bool bIsCooked)
+void FVersionedNiagaraEmitterData::PostLoad(UNiagaraEmitter& Emitter, bool bIsCooked, int32 NiagaraVer)
 {
 #if STATS
 	StatDatabase.Init();
 #endif
 	RuntimeEstimation.Init();
 
+#if WITH_EDITORONLY_DATA
+	//Force old data to not update Initial values to maintain existing behavior.
+	//New data and anything resaved will update Initial. values with writes from Event scripts. 
+	if (NiagaraVer < FNiagaraCustomVersion::EventSpawnsUpdateInitialAttributeValues)
+	{
+		for (FNiagaraEventScriptProperties& EventScript : EventHandlerScriptProps)
+		{
+			EventScript.UpdateAttributeInitialValues = false;
+		}
+	}
+#endif
+
 #if WITH_EDITORONLY_DATA && WITH_EDITOR
 	if (EditorParameters == nullptr)
 	{
diff --git a/Engine/Plugins/FX/Niagara/Source/Niagara/Private/NiagaraScript.cpp b/Engine/Plugins/FX/Niagara/Source/Niagara/Private/NiagaraScript.cpp
--- a/Engine/Plugins/FX/Niagara/Source/Niagara/Private/NiagaraScript.cpp
+++ b/Engine/Plugins/FX/Niagara/Source/Niagara/Private/NiagaraScript.cpp
@@ -107,6 +107,14 @@ static FAutoConsoleVariableRef CVarGNiagaraScriptStripByteCodeOnLoad(
 	ECVF_Default
 );
 
+static int32 GbNiagaraEventSpawnsUpdateInitialAttributeValues = 1;
+static FAutoConsoleVariableRef CVarNiagaraEventSpawnsUpdateInitialAttributeValues(
+	TEXT("fx.Niagara.EventSpawnsUpdateAttributeInitialValues"),
+	GbNiagaraEventSpawnsUpdateInitialAttributeValues,
+	TEXT("If > 0 Niagara Event Spawn Scripts will update the Initial.* values for particle attributes. \n"),
+	ECVF_Default
+);
+
 FNiagaraScriptDebuggerInfo::FNiagaraScriptDebuggerInfo() : bWaitForGPU(false), Usage(ENiagaraScriptUsage::Function), FrameLastWriteId(-1), bWritten(false)
 {
 }
@@ -1114,6 +1122,11 @@ void UNiagaraScript::ComputeVMCompilationId(FNiagaraVMExecutableDataId& Id, FGui
 				if (EventScriptProps->ExecutionMode == EScriptExecutionMode::SpawnedParticles)
 				{
 					Id.AdditionalDefines.AddUnique(FNiagaraCompileOptions::EventSpawnDefine);
+
+					if (EventScriptProps->UpdateAttributeInitialValues && GbNiagaraEventSpawnsUpdateInitialAttributeValues)
+					{
+						Id.AdditionalDefines.AddUnique(FNiagaraCompileOptions::EventSpawnInitialAttribWritesDefine);
+					}
 				}
 			}
 		}
@@ -3102,10 +3115,10 @@ void UNiagaraScript::BeginCacheForCookedPlatformData(const ITargetPlatform *Targ
 		for (int32 FormatIndex = 0; FormatIndex < DesiredShaderFormats.Num(); FormatIndex++)
 		{
 			const EShaderPlatform LegacyShaderPlatform = ShaderFormatToLegacyShaderPlatform(DesiredShaderFormats[FormatIndex]);
-			CacheResourceShadersForCooking(LegacyShaderPlatform, CachedScriptResourcesForPlatform, TargetPlatform);
+				CacheResourceShadersForCooking(LegacyShaderPlatform, CachedScriptResourcesForPlatform, TargetPlatform);
+			}
 		}
 	}
-}
 
 bool UNiagaraScript::IsCachedCookedPlatformDataLoaded(const ITargetPlatform* TargetPlatform)
 {
@@ -3259,9 +3272,9 @@ void UNiagaraScript::CacheResourceShadersForRendering(bool bRegenerateId, bool b
 
 			ScriptResource->BuildScriptParametersMetadata(CachedScriptVM.ShaderScriptParametersMetadata);
 
-			CacheShadersForResources(ScriptResource.Get(), true);
-			ScriptResourcesByFeatureLevel[CacheFeatureLevel] = ScriptResource.Get();
-		}
+				CacheShadersForResources(ScriptResource.Get(), true);
+				ScriptResourcesByFeatureLevel[CacheFeatureLevel] = ScriptResource.Get();
+			}
 		else
 		{
 			ScriptResource->Invalidate();
diff --git a/Engine/Plugins/FX/Niagara/Source/NiagaraCore/Public/NiagaraCustomVersion.h b/Engine/Plugins/FX/Niagara/Source/NiagaraCore/Public/NiagaraCustomVersion.h
--- a/Engine/Plugins/FX/Niagara/Source/NiagaraCore/Public/NiagaraCustomVersion.h
+++ b/Engine/Plugins/FX/Niagara/Source/NiagaraCore/Public/NiagaraCustomVersion.h
@@ -180,6 +180,8 @@ struct FNiagaraCustomVersion
 
 		RepopulateFunctionCallNodePinNameBindings, // Previous repopulate didn't handle module attributes like Particles.Module.Name so they need to be repopulated for renaming to work correctly.
 
+		EventSpawnsUpdateInitialAttributeValues, // Event spawns now optionally update Initial. attribute values. New default is true but old data is kept false to maintain existing behavior.
+
 		// DO NOT ADD A NEW VERSION UNLESS YOU HAVE TALKED TO THE NIAGARA LEAD. Mismanagement of these versions can lead to data loss if it is adjusted in multiple streams simultaneously.
 		// -----<new versions can be added above this line>  -------------------------------------------------
 		VersionPlusOne,
diff --git a/Engine/Plugins/FX/Niagara/Source/NiagaraEditor/Private/Customizations/NiagaraEventScriptPropertiesCustomization.cpp b/Engine/Plugins/FX/Niagara/Source/NiagaraEditor/Private/Customizations/NiagaraEventScriptPropertiesCustomization.cpp
--- a/Engine/Plugins/FX/Niagara/Source/NiagaraEditor/Private/Customizations/NiagaraEventScriptPropertiesCustomization.cpp
+++ b/Engine/Plugins/FX/Niagara/Source/NiagaraEditor/Private/Customizations/NiagaraEventScriptPropertiesCustomization.cpp
@@ -56,6 +56,7 @@ void FNiagaraEventScriptPropertiesCustomization::CustomizeChildren(TSharedRef<IP
 	HandleMaxEvents = StructPropertyHandle->GetChildHandle(GET_MEMBER_NAME_CHECKED(FNiagaraEventScriptProperties, MaxEventsPerFrame));
 	HandleUseRandomSpawnNumber = StructPropertyHandle->GetChildHandle(GET_MEMBER_NAME_CHECKED(FNiagaraEventScriptProperties, bRandomSpawnNumber));
 	HandleMinSpawnNumber = StructPropertyHandle->GetChildHandle(GET_MEMBER_NAME_CHECKED(FNiagaraEventScriptProperties, MinSpawnNumber));
+	HandleUpdateInitialValues = StructPropertyHandle->GetChildHandle(GET_MEMBER_NAME_CHECKED(FNiagaraEventScriptProperties, UpdateAttributeInitialValues));
 
 	ResolveEmitterName();
 
@@ -180,6 +181,27 @@ void FNiagaraEventScriptPropertiesCustomization::CustomizeChildren(TSharedRef<IP
 				SpawnRow.ValueWidget.Widget = SpawnValueWidget
 			];
 	}
+
+	// Add bool saying if event spawn scripts should update the initial value of particle attributes.
+	{
+		FText SrcTxt = LOCTEXT("UpdateAttributeInitialValues", "Update Initial Values");
+		FDetailWidgetRow& InitialValuseRow = ChildBuilder.AddCustomRow(SrcTxt);
+		TSharedRef<SWidget> InitialValuseWidget = HandleUpdateInitialValues->CreatePropertyValueWidget();
+		TAttribute<bool> EnabledAttribute(this, &FNiagaraEventScriptPropertiesCustomization::GetUpdateInitialValuesEnabled);
+		InitialValuseRow.IsEnabled(EnabledAttribute);
+
+		FSimpleDelegate OnUpdateInitialValuesChangedDelegate = FSimpleDelegate::CreateSP(this, &FNiagaraEventScriptPropertiesCustomization::OnUpdateInitialValuesChanged);
+		HandleUpdateInitialValues->SetOnPropertyValueChanged(OnUpdateInitialValuesChangedDelegate);
+
+		InitialValuseRow.NameWidget
+			[
+				HandleUpdateInitialValues->CreatePropertyNameWidget()
+			];
+		InitialValuseRow.ValueWidget
+			[
+				InitialValuseRow.ValueWidget.Widget = InitialValuseWidget
+			];
+	}
 }
 
 void FNiagaraEventScriptPropertiesCustomization::PostUndo(bool bSuccess)
@@ -283,6 +305,21 @@ EVisibility FNiagaraEventScriptPropertiesCustomization::GetMinSpawnNumberVisible
 	return GetUseRandomSpawnNumber() ? EVisibility::Visible : EVisibility::Collapsed;
 }
 
+bool FNiagaraEventScriptPropertiesCustomization::GetUpdateInitialValuesEnabled() const
+{
+	uint8 Value = 255;
+	if (HandleExecutionMode->GetValue(Value) && Value == (uint8)EScriptExecutionMode::SpawnedParticles)
+	{
+		return true;
+	}
+	return false;
+}
+
+void FNiagaraEventScriptPropertiesCustomization::OnUpdateInitialValuesChanged()
+{
+	UNiagaraSystem::RequestCompileForEmitter(Emitter.ResolveWeakPtr());
+}
+
 void FNiagaraEventScriptPropertiesCustomization::ResolveEmitterName()
 {
 	FString EventSrcIdStr;
diff --git a/Engine/Plugins/FX/Niagara/Source/NiagaraEditor/Private/NiagaraHlslTranslator.cpp b/Engine/Plugins/FX/Niagara/Source/NiagaraEditor/Private/NiagaraHlslTranslator.cpp
--- a/Engine/Plugins/FX/Niagara/Source/NiagaraEditor/Private/NiagaraHlslTranslator.cpp
+++ b/Engine/Plugins/FX/Niagara/Source/NiagaraEditor/Private/NiagaraHlslTranslator.cpp
@@ -626,7 +626,7 @@ void FHlslNiagaraTranslator::BuildMissingDefaults()
 {
 	AddBodyComment(TEXT("// Begin HandleMissingDefaultValues"));
 
-	if (TranslationStages[ActiveStageIdx].ShouldDoSpawnOnlyLogic())
+	if (TranslationStages[ActiveStageIdx].ShouldDoSpawnOnlyLogic() || TranslationStages[ActiveStageIdx].bShouldUpdateInitialAttributeValues)
 	{
 		// First go through all the variables that we did not write the defaults for yet. For spawn scripts, this usually
 		// means variables that reference other variables but are not themselves used within spawn.
@@ -652,21 +652,24 @@ void FHlslNiagaraTranslator::BuildMissingDefaults()
 
 		DeferredVariablesMissingDefault.Empty();
 
-		// Now go through and initialize any "Particles.Initial." variables
-		for (FNiagaraVariable& Var : InitialNamespaceVariablesMissingDefault)
+		if (TranslationStages[ActiveStageIdx].bShouldUpdateInitialAttributeValues)
 		{
-			if (FNiagaraParameterMapHistory::IsInitialValue(Var))
+			// Now go through and initialize any "Particles.Initial." variables
+			for (FNiagaraVariable& Var : InitialNamespaceVariablesMissingDefault)
 			{
-				FNiagaraVariable SourceForInitialValue = FNiagaraParameterMapHistory::GetSourceForInitialValue(Var);
-				FString ParameterMapInstanceName = GetParameterMapInstanceName(0);
-				FString Value = FString::Printf(TEXT("%s.%s = %s.%s;\n"), *ParameterMapInstanceName, *GetSanitizedSymbolName(Var.GetName().ToString()),
-					*ParameterMapInstanceName, *GetSanitizedSymbolName(SourceForInitialValue.GetName().ToString()));
-				AddBodyChunk(Value);
-				continue;
+				if (FNiagaraParameterMapHistory::IsInitialValue(Var))
+				{
+					FNiagaraVariable SourceForInitialValue = FNiagaraParameterMapHistory::GetSourceForInitialValue(Var);
+					FString ParameterMapInstanceName = GetParameterMapInstanceName(0);
+					FString Value = FString::Printf(TEXT("%s.%s = %s.%s;\n"), *ParameterMapInstanceName, *GetSanitizedSymbolName(Var.GetName().ToString()),
+						*ParameterMapInstanceName, *GetSanitizedSymbolName(SourceForInitialValue.GetName().ToString()));
+					AddBodyChunk(Value);
+					continue;
+				}
 			}
+			InitialNamespaceVariablesMissingDefault.Empty();
 		}
 
-		InitialNamespaceVariablesMissingDefault.Empty();
 	}
 
 	AddBodyComment(TEXT("// End HandleMissingDefaultValues\n\n"));
@@ -1168,6 +1171,7 @@ const FNiagaraTranslateResults &FHlslNiagaraTranslator::Translate(const FNiagara
 		TranslationStages[1].NumIterations = 1;
 		TranslationStages[0].bWritesParticles = true;
 		TranslationStages[1].bWritesParticles = true;
+		TranslationStages[0].bShouldUpdateInitialAttributeValues = true;
 		ParamMapHistories.AddDefaulted(2);
 		ParamMapSetVariablesToChunks.AddDefaulted(2);
 		ParamMapHistoriesSourceInOtherHistories.AddDefaulted(2);
@@ -1188,6 +1192,7 @@ const FNiagaraTranslateResults &FHlslNiagaraTranslator::Translate(const FNiagara
 		TranslationStages[1].NumIterations = 1;
 		TranslationStages[0].bWritesParticles = true;
 		TranslationStages[1].bWritesParticles = true;
+		TranslationStages[0].bShouldUpdateInitialAttributeValues = true;
 		ParamMapHistories.AddDefaulted(2);
 		ParamMapHistoriesSourceInOtherHistories.AddDefaulted(2);
 		ParamMapSetVariablesToChunks.AddDefaulted(2);
@@ -1390,6 +1395,7 @@ const FNiagaraTranslateResults &FHlslNiagaraTranslator::Translate(const FNiagara
 		TranslationStages[0].SimulationStageIndex = 0;
 		TranslationStages[0].NumIterations = 1;
 		TranslationStages[0].bWritesParticles = true;
+		TranslationStages[0].bShouldUpdateInitialAttributeValues = TranslationStages[0].ShouldDoSpawnOnlyLogic() || (IsEventSpawnScript() && CompileOptions.AdditionalDefines.Contains(FNiagaraCompileOptions::EventSpawnInitialAttribWritesDefine));
 
 		if (CompileOptions.TargetUsage == ENiagaraScriptUsage::ParticleSimulationStageScript)
 		{
@@ -1637,7 +1643,8 @@ const FNiagaraTranslateResults &FHlslNiagaraTranslator::Translate(const FNiagara
 		TranslationStages[0].OutputNode->Compile(ThisTranslator, OutputChunks);
 		ActiveHistoryForFunctionCalls.EndUsage();
 
-		if (IsSpawnScript())
+		bool bIsEventSpawn = IsEventSpawnScript();
+		if (IsSpawnScript() || TranslationStages[0].bShouldUpdateInitialAttributeValues)
 		{
 			BuildMissingDefaults();
 		}
@@ -4825,8 +4832,12 @@ bool FHlslNiagaraTranslationStage::ShouldDoSpawnOnlyLogic() const
 	{
 		return true;
 	}
-	
-	return (ScriptUsage == ENiagaraScriptUsage::ParticleSimulationStageScript) && (ExecuteBehavior == ENiagaraSimStageExecuteBehavior::OnSimulationReset);
+	if((ScriptUsage == ENiagaraScriptUsage::ParticleSimulationStageScript) && (ExecuteBehavior == ENiagaraSimStageExecuteBehavior::OnSimulationReset))
+	{
+		return true;
+	}
+
+	return false;
 }
 
 bool FHlslNiagaraTranslationStage::IsExternalConstantNamespace(const FNiagaraVariable& InVar, ENiagaraScriptUsage InTargetUsage, uint32 InTargetBitmask)
@@ -4843,7 +4854,7 @@ bool FHlslNiagaraTranslationStage::IsExternalConstantNamespace(const FNiagaraVar
 
 bool FHlslNiagaraTranslationStage::IsRelevantToSpawnForStage(const FNiagaraParameterMapHistory& InHistory, const FNiagaraVariable& InAliasedVar, const FNiagaraVariable& InVar) const
 {
-	if (InHistory.IsPrimaryDataSetOutput(InAliasedVar, ScriptUsage) && UNiagaraScript::IsSpawnScript(ScriptUsage))
+	if (InHistory.IsPrimaryDataSetOutput(InAliasedVar, ScriptUsage) && (UNiagaraScript::IsSpawnScript(ScriptUsage) || bShouldUpdateInitialAttributeValues))
 	{
 		return true;
 	}
@@ -4882,7 +4893,7 @@ void FHlslNiagaraTranslator::InitializeParameterMapDefaults(int32 ParamMapHistor
 			const FNiagaraVariable& Var = History.Variables[i];
 			const FNiagaraVariable& AliasedVar = History.VariablesWithOriginalAliasesIntact[i];
 			// Only add primary data set outputs at the top of the script if in a spawn script, otherwise they should be left alone.
-			if (TranslationStages[ActiveStageIdx].ShouldDoSpawnOnlyLogic())
+			if (TranslationStages[ActiveStageIdx].ShouldDoSpawnOnlyLogic() || TranslationStages[ActiveStageIdx].bShouldUpdateInitialAttributeValues)
 			{
 				if (TranslationStages[ActiveStageIdx].IsRelevantToSpawnForStage(History, AliasedVar, Var) &&
 					!UniqueVars.Contains(Var))
@@ -4898,7 +4909,7 @@ void FHlslNiagaraTranslator::InitializeParameterMapDefaults(int32 ParamMapHistor
 
 	// Only add primary data set outputs at the top of the script if in a spawn script, otherwise they should be left alone.
 	// Above we added all the known from the spawn script, now let's add for all the others.
-	if (TranslationStages[ActiveStageIdx].ShouldDoSpawnOnlyLogic())
+	if (TranslationStages[ActiveStageIdx].ShouldDoSpawnOnlyLogic() || TranslationStages[ActiveStageIdx].bShouldUpdateInitialAttributeValues)
 	{
 		// Go through all referenced parameter maps and pull in any variables that are 
 		// in the primary data set output namespaces.
@@ -5386,6 +5397,10 @@ bool FHlslNiagaraTranslator::IsSpawnScript() const
 	return false;
 }
 
+bool FHlslNiagaraTranslator::IsEventSpawnScript()const 
+{
+	return UNiagaraScript::IsParticleEventScript(CompileOptions.TargetUsage) && CompileOptions.AdditionalDefines.Contains(FNiagaraCompileOptions::EventSpawnDefine);
+}
 
 bool FHlslNiagaraTranslator::RequiresInterpolation() const
 {
diff --git a/Engine/Plugins/FX/Niagara/Source/NiagaraEditor/Private/NiagaraHlslTranslator.h b/Engine/Plugins/FX/Niagara/Source/NiagaraEditor/Private/NiagaraHlslTranslator.h
--- a/Engine/Plugins/FX/Niagara/Source/NiagaraEditor/Private/NiagaraHlslTranslator.h
+++ b/Engine/Plugins/FX/Niagara/Source/NiagaraEditor/Private/NiagaraHlslTranslator.h
@@ -361,7 +361,7 @@ public:
 	ENiagaraScriptUsage ScriptUsage;
 	FGuid UsageId;
 	UNiagaraNodeOutput* OutputNode = nullptr;
-	FString PassNamespace ;
+	FString PassNamespace;
 	bool bInterpolatePreviousParams = false;
 	bool bCopyPreviousParams = true;
 	ENiagaraCodeChunkMode ChunkModeIndex = (ENiagaraCodeChunkMode)-1;
@@ -385,6 +385,7 @@ public:
 	bool bOverrideGpuDispatchType = false;
 	ENiagaraGpuDispatchType OverrideGpuDispatchType = ENiagaraGpuDispatchType::OneD;
 	bool bOverrideGpuDispatchNumThreads = false;
+	bool bShouldUpdateInitialAttributeValues = false;
 	FIntVector OverrideGpuDispatchNumThreads = FIntVector(1, 1, 1);
 	TArray<FNiagaraVariable> SetParticleAttributes;
 	FString CustomReadFunction;
@@ -678,6 +679,7 @@ public:
 
 	int32 MakeStaticVariableDirect(const UEdGraphPin* InDefaultPin);
 
+	bool IsEventSpawnScript()const;
 
 private:
 	void InitializeParameterMapDefaults(int32 ParamMapHistoryIdx);
diff --git a/Engine/Plugins/FX/Niagara/Source/NiagaraEditor/Public/Customizations/NiagaraEventScriptPropertiesCustomization.h b/Engine/Plugins/FX/Niagara/Source/NiagaraEditor/Public/Customizations/NiagaraEventScriptPropertiesCustomization.h
--- a/Engine/Plugins/FX/Niagara/Source/NiagaraEditor/Public/Customizations/NiagaraEventScriptPropertiesCustomization.h
+++ b/Engine/Plugins/FX/Niagara/Source/NiagaraEditor/Public/Customizations/NiagaraEventScriptPropertiesCustomization.h
@@ -72,6 +72,7 @@ protected:
 	bool GetSpawnNumberEnabled() const;
 	bool GetUseRandomSpawnNumber() const;
 	EVisibility GetMinSpawnNumberVisible() const;
+	bool GetUpdateInitialValuesEnabled() const;
 	void ResolveEmitterName();
 	void ComputeErrorVisibility();
 	EVisibility GetErrorVisibility() const;
@@ -79,6 +80,9 @@ protected:
 	FText GetErrorTextTooltip() const;
 
 	TArray<FName> GetEventNames(const FVersionedNiagaraEmitter& Emitter) const;
+
+	void OnUpdateInitialValuesChanged();
+
 private:
 	TSharedPtr<IPropertyHandle> HandleSrcID;
 	TSharedPtr<IPropertyHandle> HandleEventName;
@@ -87,6 +91,7 @@ private:
 	TSharedPtr<IPropertyHandle> HandleMaxEvents;
 	TSharedPtr<IPropertyHandle> HandleUseRandomSpawnNumber;
 	TSharedPtr<IPropertyHandle> HandleMinSpawnNumber;
+	TSharedPtr<IPropertyHandle> HandleUpdateInitialValues;
 
 	TWeakObjectPtr<UNiagaraSystem> System;
 	FVersionedNiagaraEmitterWeakPtr Emitter;
diff --git a/Engine/Source/Runtime/Core/Private/UObject/FortniteMainBranchObjectVersion.cpp b/Engine/Source/Runtime/Core/Private/UObject/FortniteMainBranchObjectVersion.cpp
--- a/Engine/Source/Runtime/Core/Private/UObject/FortniteMainBranchObjectVersion.cpp
+++ b/Engine/Source/Runtime/Core/Private/UObject/FortniteMainBranchObjectVersion.cpp
@@ -11,7 +11,7 @@ TMap<FGuid, FGuid> FFortniteMainBranchObjectVersion::GetSystemGuids()
 	SystemGuids.Add(DevGuids.MATERIALSHADERMAP_DERIVEDDATA_VER, FGuid("1DA85DA5733C4A489A1BE117CDC5ACCC"));
 	SystemGuids.Add(DevGuids.NANITE_DERIVEDDATA_VER, FGuid("93379C2F29624E74AFB14FB13E327FA2"));
 	SystemGuids.Add(DevGuids.NIAGARASHADERMAP_DERIVEDDATA_VER, FGuid("2289D5116CF94BC0AFEAEC541468E645"));
-	SystemGuids.Add(DevGuids.Niagara_LatestScriptCompileVersion, FGuid("D6A13E2AE4D2453EB1A44529AAAB8305"));
+	SystemGuids.Add(DevGuids.Niagara_LatestScriptCompileVersion, FGuid("04A14D478F6449BC92BD142968DCB2BB"));
 	SystemGuids.Add(DevGuids.SkeletalMeshDerivedDataVersion, FGuid("ACF593EAAE354FCBB34CF44F5AA1BFD2"));
 	SystemGuids.Add(DevGuids.STATICMESH_DERIVEDDATA_VER, FGuid("9A48CE6B46204777968322F98F82D11F"));
 
