Don't use GetFeatureLevel() because it could still have an invalid feature level set in the editor

diff --git a/Engine/Plugins/FX/Niagara/Source/NiagaraVertexFactories/Private/NiagaraMeshVertexFactory.cpp b/Engine/Plugins/FX/Niagara/Source/NiagaraVertexFactories/Private/NiagaraMeshVertexFactory.cpp
--- a/Engine/Plugins/FX/Niagara/Source/NiagaraVertexFactories/Private/NiagaraMeshVertexFactory.cpp
+++ b/Engine/Plugins/FX/Niagara/Source/NiagaraVertexFactories/Private/NiagaraMeshVertexFactory.cpp
@@ -75,7 +75,7 @@ void FNiagaraMeshVertexFactory::InitRHI()
 		Elements.Add(AccessStreamComponent(Data.PositionComponent, 0));
 	}
 
-	const bool bUseManualVertexFetch = SupportsManualVertexFetch(GetFeatureLevel());
+	const bool bUseManualVertexFetch = SupportsManualVertexFetch(GMaxRHIFeatureLevel);
 	if (!bUseManualVertexFetch)
 	{
 		// only tangent,normal are used by the stream. the binormal is derived in the shader
