Use GetFeatureLevel() when building the vertex factory declarations when it's valid - fixes issues with ES3.1 preview in the editor

diff --git a/Engine/Plugins/FX/Niagara/Source/NiagaraVertexFactories/Private/NiagaraMeshVertexFactory.cpp b/Engine/Plugins/FX/Niagara/Source/NiagaraVertexFactories/Private/NiagaraMeshVertexFactory.cpp
--- a/Engine/Plugins/FX/Niagara/Source/NiagaraVertexFactories/Private/NiagaraMeshVertexFactory.cpp
+++ b/Engine/Plugins/FX/Niagara/Source/NiagaraVertexFactories/Private/NiagaraMeshVertexFactory.cpp
@@ -75,8 +75,11 @@ void FNiagaraMeshVertexFactory::InitRHI()
 		Elements.Add(AccessStreamComponent(Data.PositionComponent, 0));
 	}
 
-	const bool bUseManualVertexFetch = SupportsManualVertexFetch(GMaxRHIFeatureLevel);
-	if (!bUseManualVertexFetch)
+	// Can't use GetFeatureLevel() on FNiagaraMeshVertexFactory because it's never set during creation
+	// This needs to be fixed and then the unused stream components don't have to be added when
+	// manual vertex fetch is used - can't use GMaxRHIFeatureLevel to check for support because then ES3.1 preview won't work anymore
+	//const bool bUseManualVertexFetch = SupportsManualVertexFetch(GetFeatureLevel());
+	//if (!bUseManualVertexFetch)
 	{
 		// only tangent,normal are used by the stream. the binormal is derived in the shader
 		uint8 TangentBasisAttributes[2] = { 1, 2 };
diff --git a/Engine/Source/Runtime/Engine/Private/LocalVertexFactory.cpp b/Engine/Source/Runtime/Engine/Private/LocalVertexFactory.cpp
--- a/Engine/Source/Runtime/Engine/Private/LocalVertexFactory.cpp
+++ b/Engine/Source/Runtime/Engine/Private/LocalVertexFactory.cpp
@@ -417,8 +417,8 @@ void FLocalVertexFactory::InitRHI()
 
 	// VertexFactory needs to be able to support max possible shader platform and feature level
 	// in case if we switch feature level at runtime.
-	const bool bCanUseGPUScene = UseGPUScene(GMaxRHIShaderPlatform, GMaxRHIFeatureLevel);
-	const bool bUseManualVertexFetch = SupportsManualVertexFetch(GMaxRHIFeatureLevel);
+	const bool bCanUseGPUScene = UseGPUScene(GMaxRHIShaderPlatform, GetFeatureLevel());
+	const bool bUseManualVertexFetch = SupportsManualVertexFetch(GetFeatureLevel());
 
 	// If the vertex buffer containing position is not the same vertex buffer containing the rest of the data,
 	// then initialize PositionStream and PositionDeclaration.
