Readded ability to select an emitter or system's scalability properties in the stack by clicking on the node. Clicking on Properties still works as before.

diff --git a/Engine/Plugins/FX/Niagara/Source/NiagaraEditor/Private/Widgets/Scalability/SNiagaraScalabilityContext.cpp b/Engine/Plugins/FX/Niagara/Source/NiagaraEditor/Private/Widgets/Scalability/SNiagaraScalabilityContext.cpp
--- a/Engine/Plugins/FX/Niagara/Source/NiagaraEditor/Private/Widgets/Scalability/SNiagaraScalabilityContext.cpp
+++ b/Engine/Plugins/FX/Niagara/Source/NiagaraEditor/Private/Widgets/Scalability/SNiagaraScalabilityContext.cpp
@@ -15,6 +15,7 @@
 #include "PropertyEditorModule.h"
 #include "ViewModels/NiagaraEmitterViewModel.h"
 #include "ViewModels/Stack/NiagaraStackEmitterPropertiesGroup.h"
+#include "ViewModels/Stack/NiagaraStackRoot.h"
 #include "ViewModels/Stack/NiagaraStackSystemSettingsGroup.h"
 #include "Widgets/Layout/SWidgetSwitcher.h"
 
@@ -25,6 +26,7 @@ void SNiagaraScalabilityContext::Construct(const FArguments& InArgs, UNiagaraSys
 	ScalabilityViewModel = &InScalabilityViewModel;
 	
 	ScalabilityViewModel->GetSystemViewModel().Pin()->GetSelectionViewModel()->OnEntrySelectionChanged().AddSP(this, &SNiagaraScalabilityContext::UpdateScalabilityContent);
+	ScalabilityViewModel->GetSystemViewModel().Pin()->GetOverviewGraphViewModel()->GetNodeSelection()->OnSelectedObjectsChanged().AddSP(this, &SNiagaraScalabilityContext::UpdateScalabilityContent);
 
 	FPropertyEditorModule& PropertyEditorModule = FModuleManager::Get().LoadModuleChecked<FPropertyEditorModule>("PropertyEditor");
 	
@@ -40,6 +42,7 @@ void SNiagaraScalabilityContext::Construct(const FArguments& InArgs, UNiagaraSys
 	DetailsView->RegisterInstancedCustomPropertyLayout(UNiagaraEmitter::StaticClass(), FOnGetDetailCustomizationInstance::CreateStatic(&FNiagaraEmitterScalabilityDetails::MakeInstance));
 	
 	DetailsView->SetIsPropertyVisibleDelegate(FIsPropertyVisible::CreateSP(this, &SNiagaraScalabilityContext::FilterScalabilityProperties));
+	
 	UpdateScalabilityContent();
 
 	ChildSlot
@@ -66,6 +69,22 @@ void SNiagaraScalabilityContext::Construct(const FArguments& InArgs, UNiagaraSys
 	];
 }
 
+SNiagaraScalabilityContext::~SNiagaraScalabilityContext()
+{
+	if(ScalabilityViewModel.IsValid() && ScalabilityViewModel->GetSystemViewModel().IsValid())
+	{
+		if(ScalabilityViewModel->GetSystemViewModel().Pin()->GetSelectionViewModel())
+		{
+			ScalabilityViewModel->GetSystemViewModel().Pin()->GetSelectionViewModel()->OnEntrySelectionChanged().RemoveAll(this);
+		}
+
+		if(ScalabilityViewModel->GetSystemViewModel().Pin()->GetOverviewGraphViewModel())
+		{
+			ScalabilityViewModel->GetSystemViewModel().Pin()->GetOverviewGraphViewModel()->GetNodeSelection()->OnSelectedObjectsChanged().RemoveAll(this);
+		}
+	}
+}
+
 void SNiagaraScalabilityContext::SetObject(UObject* Object)
 {
 	if(Object == nullptr)
@@ -89,11 +108,14 @@ void SNiagaraScalabilityContext::SetObject(UObject* Object)
 void SNiagaraScalabilityContext::UpdateScalabilityContent()
 {
 	UObject* NewSelection = nullptr;
-	
+
 	TArray<UNiagaraStackEntry*> StackEntries;
 	ScalabilityViewModel->GetSystemViewModel().Pin()->GetSelectionViewModel()->GetSelectedEntries(StackEntries);
-	
-	if(StackEntries.Num() == 1)
+
+	TSet<UObject*> SelectedNodes = ScalabilityViewModel->GetSystemViewModel().Pin()->GetOverviewGraphViewModel()->GetNodeSelection()->GetSelectedObjects();
+
+	// if we select a node, the stack entries will be updated to a niagara stack root, which we want to ignore
+	if(StackEntries.Num() == 1 && !StackEntries[0]->IsA<UNiagaraStackRoot>())
 	{
 		UNiagaraStackEntry* StackEntry = StackEntries[0];
 
@@ -112,6 +134,22 @@ void SNiagaraScalabilityContext::UpdateScalabilityContent()
 			NewSelection = &System;
 		}
 	}
+	else if(SelectedNodes.Num() == 1)
+	{
+		UObject* SelectedNode = SelectedNodes.Array()[0];
+
+		if(UNiagaraOverviewNode* OverviewNode = Cast<UNiagaraOverviewNode>(SelectedNode))
+		{
+			if(FNiagaraEmitterHandle* EmitterHandle = OverviewNode->TryGetEmitterHandle())
+			{
+				NewSelection = EmitterHandle->GetInstance().Emitter;
+			}
+			else if(UNiagaraSystem* System = OverviewNode->GetOwningSystem())
+			{
+				NewSelection = System;
+			}
+		}
+	}	
 	
 	SetObject(NewSelection);
 }
diff --git a/Engine/Plugins/FX/Niagara/Source/NiagaraEditor/Private/Widgets/Scalability/SNiagaraScalabilityContext.h b/Engine/Plugins/FX/Niagara/Source/NiagaraEditor/Private/Widgets/Scalability/SNiagaraScalabilityContext.h
--- a/Engine/Plugins/FX/Niagara/Source/NiagaraEditor/Private/Widgets/Scalability/SNiagaraScalabilityContext.h
+++ b/Engine/Plugins/FX/Niagara/Source/NiagaraEditor/Private/Widgets/Scalability/SNiagaraScalabilityContext.h
@@ -17,11 +17,11 @@ public:
 	SLATE_END_ARGS()
 
 	void Construct(const FArguments& InArgs, UNiagaraSystemScalabilityViewModel& InScalabilityViewModel);
-
+	virtual ~SNiagaraScalabilityContext() override;
+	
 	void SetObject(UObject* Object);
 
 	void UpdateScalabilityContent();
-
 private:
 	bool FilterScalabilityProperties(const FPropertyAndParent& InPropertyAndParent) const;
 	TWeakObjectPtr<UNiagaraSystemScalabilityViewModel> ScalabilityViewModel;
