NiagaraParameterMap pins are editable upon creation (Contributed by [REDACTED_PERSON])

#9539

diff --git a/Engine/Plugins/FX/Niagara/Source/NiagaraEditor/Private/TNiagaraGraphPinEditableName.h b/Engine/Plugins/FX/Niagara/Source/NiagaraEditor/Private/TNiagaraGraphPinEditableName.h
--- a/Engine/Plugins/FX/Niagara/Source/NiagaraEditor/Private/TNiagaraGraphPinEditableName.h
+++ b/Engine/Plugins/FX/Niagara/Source/NiagaraEditor/Private/TNiagaraGraphPinEditableName.h
@@ -106,11 +106,6 @@ protected:
 		}
 		
 		auto CreateRenamableLabelTextBlock = [&]()->TSharedRef<SWidget> {
-			if (ParentNode->IsPinNameEditableUponCreation(this->GraphPinObj))
-			{
-				bPendingRename = true;
-			}
-			
 			CreatedTextBlock = SNew(SInlineEditableTextBlock)
 				.Style(&FAppStyle::Get().GetWidgetStyle<FInlineEditableTextBlockStyle>("Graph.Node.InlineEditablePinName"))
 				.Text(this, &TNiagaraGraphPinEditableName<BaseClass>::GetParentPinLabel)
@@ -123,6 +118,12 @@ protected:
 		
 		const bool bAllowPinTypeChanges = ParentNode->AllowExternalPinTypeChanges(this->GraphPinObj) && this->GraphPinObj->bOrphanedPin == false;
 		const bool bIsPinEditable = ParentNode->IsPinNameEditable(this->GraphPinObj);
+		
+		if (bIsPinEditable && ParentNode->IsPinNameEditableUponCreation(this->GraphPinObj))
+		{
+			bPendingRename = true;
+		}
+		
 		if (ParentNode->IsA<UNiagaraNodeParameterMapBase>())
 		{
 			return SNew(SBox)
