- pack grid 3d collections into rgba textures instead of tiled attributes in a single channel float texture
  * much faster for retrieving vector attributes from a grid (ie: velocity, color, etc)
  * far less overhead for common cases where grids have 1 attribute
  * many future use cases and improvements are possible
- usage
  * controlled via cvar (fx.Niagara.Grid3D.UseRGBAGrid 1) to enable
  * only used for grids that have 1 attribute with 1 to 4 channels
  * doesn't work with unnamed grid attributes
- many new "AtIndex" DI functions that give a faster path for various grid operations
- warn about unnamed attributes as a deprecated feature
- integrated changes into many grid3d related modules (25% faster across the board)

diff --git a/Engine/Plugins/FX/Niagara/Shaders/Private/NiagaraDataInterfaceGrid3DCollection.ush b/Engine/Plugins/FX/Niagara/Shaders/Private/NiagaraDataInterfaceGrid3DCollection.ush
--- a/Engine/Plugins/FX/Niagara/Shaders/Private/NiagaraDataInterfaceGrid3DCollection.ush
+++ b/Engine/Plugins/FX/Niagara/Shaders/Private/NiagaraDataInterfaceGrid3DCollection.ush
@@ -30,7 +30,7 @@ float SampleExplicitTrilinear_{ParameterName}(SamplerState Sampler, float3 UVW,
 	return TV * delta.z + BV * (1.-delta.z);	
 }
 
-float SampleTriCubicLagrange_{ParameterName}(SamplerState Sampler, float3 UVW, int MipLevel)
+float{NumChannels} SampleTriCubicLagrange_{ParameterName}(SamplerState Sampler, float3 UVW, int MipLevel)
 {		    
 	int3 TotalDim = {NumCellsName} * {NumTiles};	
 	float3 GridPos = float3(TotalDim[0] * UVW[0] - 0.5,  TotalDim[1] * UVW[1] - 0.5, TotalDim[2] * UVW[2] - 0.5f);
@@ -57,17 +57,17 @@ float SampleTriCubicLagrange_{ParameterName}(SamplerState Sampler, float3 UVW, i
 		float u = 1. + delta.y;
 		float v = 1. + delta.z;
 		
-		float4 DataInZ = float4(0.f, 0.f, 0.f, 0.f);
-		float minv = 3.402823466e+38;
-		float maxv = -3.402823466e+38;
+		float4x{NumChannels} DataInZ;
+		float{NumChannels} minv = 3.402823466e+38;
+		float{NumChannels} maxv = -3.402823466e+38;
 		for (int zk = 0; zk < 4; zk++)
 		{
 			int plane = GridCell[2] + zk -1;
 		
-			float4 DataInY = float4(0.f, 0.f, 0.f, 0.f);
+			float4x{NumChannels}  DataInY;
 			for (int yj = 0; yj < 4; yj++)
 			{
-				float4 DataInX = float4(0.f, 0.f, 0.f, 0.f);
+				float4x{NumChannels} DataInX;
 								
 				int row = GridCell[1] + yj - 1;
 		
@@ -80,30 +80,30 @@ float SampleTriCubicLagrange_{ParameterName}(SamplerState Sampler, float3 UVW, i
 				}
 		
 				// coefficients
-				float Cx0 =   6.f * DataInX[0];
-				float Cx1 = -11.f * DataInX[0] + 18.f * DataInX[1] - 9.f * DataInX[2] + 2.f * DataInX[3];
-				float Cx2 =   6.f * DataInX[0] - 15.f * DataInX[1] +12.f * DataInX[2] - 3.f * DataInX[3];
-				float Cx3 =        -DataInX[0] +  3.f * DataInX[1] - 3.f * DataInX[2] +       DataInX[3];
+				float{NumChannels} Cx0 =   6.f * DataInX[0];
+				float{NumChannels} Cx1 = -11.f * DataInX[0] + 18.f * DataInX[1] - 9.f * DataInX[2] + 2.f * DataInX[3];
+				float{NumChannels} Cx2 =   6.f * DataInX[0] - 15.f * DataInX[1] +12.f * DataInX[2] - 3.f * DataInX[3];
+				float{NumChannels} Cx3 =        -DataInX[0] +  3.f * DataInX[1] - 3.f * DataInX[2] +       DataInX[3];
 		
 				DataInY[yj] = (1.f/6.f) * (Cx0 + t * ( Cx1 +  t * ( Cx2 + t * Cx3 ) ));
 		
 			}
 									
-			float Cy0 =   6.f * DataInY[0];
-			float Cy1 = -11.f * DataInY[0] + 18.f * DataInY[1] - 9.f * DataInY[2] + 2.f * DataInY[3];
-			float Cy2 =   6.f * DataInY[0] - 15.f * DataInY[1] +12.f * DataInY[2] - 3.f * DataInY[3];
-			float Cy3 =        -DataInY[0] +  3.f * DataInY[1] - 3.f * DataInY[2] +       DataInY[3];
+			float{NumChannels} Cy0 =   6.f * DataInY[0];
+			float{NumChannels} Cy1 = -11.f * DataInY[0] + 18.f * DataInY[1] - 9.f * DataInY[2] + 2.f * DataInY[3];
+			float{NumChannels} Cy2 =   6.f * DataInY[0] - 15.f * DataInY[1] +12.f * DataInY[2] - 3.f * DataInY[3];
+			float{NumChannels} Cy3 =        -DataInY[0] +  3.f * DataInY[1] - 3.f * DataInY[2] +       DataInY[3];
 		
 			DataInZ[zk] = (1.f/6.f) * (Cy0 + u * ( Cy1 +  u * ( Cy2 + u * Cy3 ) ));
 									
 		}
 		
-		float Cz0 =   6.f * DataInZ[0];
-		float Cz1 = -11.f * DataInZ[0] + 18.f * DataInZ[1] - 9.f * DataInZ[2] + 2.f * DataInZ[3];
-		float Cz2 =   6.f * DataInZ[0] - 15.f * DataInZ[1] +12.f * DataInZ[2] - 3.f * DataInZ[3];
-		float Cz3 =        -DataInZ[0] +  3.f * DataInZ[1] - 3.f * DataInZ[2] +       DataInZ[3];
+		float{NumChannels} Cz0 =   6.f * DataInZ[0];
+		float{NumChannels} Cz1 = -11.f * DataInZ[0] + 18.f * DataInZ[1] - 9.f * DataInZ[2] + 2.f * DataInZ[3];
+		float{NumChannels} Cz2 =   6.f * DataInZ[0] - 15.f * DataInZ[1] +12.f * DataInZ[2] - 3.f * DataInZ[3];
+		float{NumChannels} Cz3 =        -DataInZ[0] +  3.f * DataInZ[1] - 3.f * DataInZ[2] +       DataInZ[3];
 								
-		float TriCubicValue = (1.f/6.f) * (Cz0 + v * ( Cz1 +  v * ( Cz2 + v * Cz3 ) ));
+		float{NumChannels} TriCubicValue = (1.f/6.f) * (Cz0 + v * ( Cz1 +  v * ( Cz2 + v * Cz3 ) ));
 									
 		// for the fluid sims, a nice look can be achieved with BlendAmount = 0.25f;
 		// float BlendAmount = 0.0f;
diff --git a/Engine/Plugins/FX/Niagara/Source/Niagara/Classes/NiagaraDataInterface.h b/Engine/Plugins/FX/Niagara/Source/Niagara/Classes/NiagaraDataInterface.h
--- a/Engine/Plugins/FX/Niagara/Source/Niagara/Classes/NiagaraDataInterface.h
+++ b/Engine/Plugins/FX/Niagara/Source/Niagara/Classes/NiagaraDataInterface.h
@@ -475,7 +475,7 @@ public:
 	/** Generate the necessary plumbing HLSL at the beginning of the stage where this is used as a sim stage iteration source. Note that this should inject other internal calls using the CustomHLSL node syntax. See GridCollection2D for an example.*/
 	virtual bool GenerateIterationSourceNamespaceReadAttributesHLSL(FNiagaraDataInterfaceGPUParamInfo& DIInstanceInfo, const FNiagaraVariable& InIterationSourceVariable, TConstArrayView<FNiagaraVariable> InArguments, TConstArrayView<FNiagaraVariable> InAttributes, TConstArrayView<FString> InAttributeHLSLNames, bool bInSetToDefaults, bool bSpawnOnly, bool bPartialWrites, TArray<FText>& OutErrors, FString& OutHLSL) const { return false; };
 	/** Generate the necessary plumbing HLSL at the end of the stage where this is used as a sim stage iteration source. Note that this should inject other internal calls using the CustomHLSL node syntax. See GridCollection2D for an example.*/
-	virtual bool GenerateIterationSourceNamespaceWriteAttributesHLSL(FNiagaraDataInterfaceGPUParamInfo& DIInstanceInfo, const FNiagaraVariable& InIterationSourceVariable, TConstArrayView<FNiagaraVariable> InArguments, TConstArrayView<FNiagaraVariable> InAttributes, TConstArrayView<FString> InAttributeHLSLNames, bool bSpawnOnly, bool bPartialWrites, TArray<FText>& OutErrors, FString& OutHLSL) const { return false; };
+	virtual bool GenerateIterationSourceNamespaceWriteAttributesHLSL(FNiagaraDataInterfaceGPUParamInfo& DIInstanceInfo, const FNiagaraVariable& InIterationSourceVariable, TConstArrayView<FNiagaraVariable> InArguments, TConstArrayView<FNiagaraVariable> InAttributes, TConstArrayView<FString> InAttributeHLSLNames, TConstArrayView<FNiagaraVariable> InAllAttributes, bool bSpawnOnly, bool bPartialWrites, TArray<FText>& OutErrors, FString& OutHLSL) const { return false; };
 	/** Used by the translator when dealing with signatures that turn into compiler tags to figure out the precise compiler tag. */
 	virtual bool GenerateCompilerTagPrefix(const FNiagaraFunctionSignature& InSignature, FString& OutPrefix) const  { return false; }
 
diff --git a/Engine/Plugins/FX/Niagara/Source/Niagara/Classes/NiagaraDataInterfaceGrid2DCollection.h b/Engine/Plugins/FX/Niagara/Source/Niagara/Classes/NiagaraDataInterfaceGrid2DCollection.h
--- a/Engine/Plugins/FX/Niagara/Source/Niagara/Classes/NiagaraDataInterfaceGrid2DCollection.h
+++ b/Engine/Plugins/FX/Niagara/Source/Niagara/Classes/NiagaraDataInterfaceGrid2DCollection.h
@@ -266,7 +266,7 @@ public:
 	virtual bool GenerateTeardownHLSL(FNiagaraDataInterfaceGPUParamInfo& DIInstanceInfo, TConstArrayView<FNiagaraVariable> InArguments, bool bSpawnOnly, bool bPartialWrites, TArray<FText>& OutErrors, FString& OutHLSL) const;
 	virtual bool SupportsIterationSourceNamespaceAttributesHLSL() const override { return true; }
 	virtual bool GenerateIterationSourceNamespaceReadAttributesHLSL(FNiagaraDataInterfaceGPUParamInfo& DIInstanceInfo, const FNiagaraVariable& IterationSourceVar, TConstArrayView<FNiagaraVariable> InArguments, TConstArrayView<FNiagaraVariable> InAttributes, TConstArrayView<FString> InAttributeHLSLNames, bool bInSetToDefaults, bool bSpawnOnly, bool bPartialWrites, TArray<FText>& OutErrors, FString& OutHLSL) const override;
-	virtual bool GenerateIterationSourceNamespaceWriteAttributesHLSL(FNiagaraDataInterfaceGPUParamInfo& DIInstanceInfo, const FNiagaraVariable& IterationSourceVar, TConstArrayView<FNiagaraVariable> InArguments, TConstArrayView<FNiagaraVariable> InAttributes, TConstArrayView<FString> InAttributeHLSLNames, bool bSpawnOnly, bool bPartialWrites, TArray<FText>& OutErrors, FString& OutHLSL) const override;
+	virtual bool GenerateIterationSourceNamespaceWriteAttributesHLSL(FNiagaraDataInterfaceGPUParamInfo& DIInstanceInfo, const FNiagaraVariable& IterationSourceVar, TConstArrayView<FNiagaraVariable> InArguments, TConstArrayView<FNiagaraVariable> InAttributes, TConstArrayView<FString> InAttributeHLSLNames, TConstArrayView<FNiagaraVariable> InAllAttributes, bool bSpawnOnly, bool bPartialWrites, TArray<FText>& OutErrors, FString& OutHLSL) const override;
 #endif
 
 	static int32 GetComponentCountFromFuncName(const FName& FuncName);
diff --git a/Engine/Plugins/FX/Niagara/Source/Niagara/Classes/NiagaraDataInterfaceGrid3DCollection.h b/Engine/Plugins/FX/Niagara/Source/Niagara/Classes/NiagaraDataInterfaceGrid3DCollection.h
--- a/Engine/Plugins/FX/Niagara/Source/Niagara/Classes/NiagaraDataInterfaceGrid3DCollection.h
+++ b/Engine/Plugins/FX/Niagara/Source/Niagara/Classes/NiagaraDataInterfaceGrid3DCollection.h
@@ -25,6 +25,7 @@ struct FGrid3DCollectionRWInstanceData_GameThread
 	FVector CellSize = FVector::ZeroVector;
 	FVector WorldBBoxSize = FVector::ZeroVector;
 	EPixelFormat PixelFormat = EPixelFormat::PF_R32_FLOAT;
+	bool UseRGBATexture = false;
 #if WITH_EDITOR
 	bool bPreviewGrid = false;
 	FIntVector4 PreviewAttribute = FIntVector4(INDEX_NONE, INDEX_NONE, INDEX_NONE, INDEX_NONE);
@@ -60,6 +61,8 @@ struct FGrid3DCollectionRWInstanceData_RenderThread
 	FVector CellSize = FVector::ZeroVector;
 	FVector WorldBBoxSize = FVector::ZeroVector;
 	EPixelFormat PixelFormat = EPixelFormat::PF_R32_FLOAT;
+	bool UseRGBATexture = false;
+
 	TArray<int32> AttributeIndices;
 
 	TArray<FName> Vars;
@@ -227,6 +230,7 @@ public:
 	static const FString UnitClampMinName;
 	static const FString UnitClampMaxName;
 
+	static const FName NAME_Attribute;
 	static const FString GridName;
 	static const FString OutputGridName;
 	static const FString SamplerName;
@@ -254,6 +258,26 @@ public:
 	static const FName SamplePreviousGridAtIndexFunctionName;
 	static const FName CubicSamplePreviousGridAtIndexFunctionName;
 
+	static const FName SetVector4ValueAtIndexFunctionName;
+	static const FName SetVector3ValueAtIndexFunctionName;
+	static const FName SetVector2ValueAtIndexFunctionName;	
+	static const FName SetFloatValueAtIndexFunctionName;
+
+	static const FName GetPreviousVector4ValueAtIndexFunctionName;
+	static const FName GetPreviousVector3ValueAtIndexFunctionName;
+	static const FName GetPreviousVector2ValueAtIndexFunctionName;
+	static const FName GetPreviousFloatValueAtIndexFunctionName;
+
+	static const FName SamplePreviousGridVector4AtIndexFunctionName;
+	static const FName SamplePreviousGridVector3AtIndexFunctionName;
+	static const FName SamplePreviousGridVector2AtIndexFunctionName;
+	static const FName SamplePreviousGridFloatAtIndexFunctionName;
+
+	static const FName CubicSamplePreviousGridVector4AtIndexFunctionName;
+	static const FName CubicSamplePreviousGridVector3AtIndexFunctionName;
+	static const FName CubicSamplePreviousGridVector2AtIndexFunctionName;
+	static const FName CubicSamplePreviousGridFloatAtIndexFunctionName;
+
 	static const FName GetPreviousVector4ValueFunctionName;
 	static const FName SamplePreviousGridVector4FunctionName;
 	static const FName CubicSamplePreviousGridVector4FunctionName;
@@ -289,7 +313,7 @@ public:
 	virtual bool GenerateTeardownHLSL(FNiagaraDataInterfaceGPUParamInfo& DIInstanceInfo, TConstArrayView<FNiagaraVariable> InArguments, bool bSpawnOnly, bool bPartialWrites, TArray<FText>& OutErrors, FString& OutHLSL) const;
 	virtual bool SupportsIterationSourceNamespaceAttributesHLSL() const override { return true; }
 	virtual bool GenerateIterationSourceNamespaceReadAttributesHLSL(FNiagaraDataInterfaceGPUParamInfo& DIInstanceInfo, const FNiagaraVariable& IterationSourceVar, TConstArrayView<FNiagaraVariable> InArguments, TConstArrayView<FNiagaraVariable> InAttributes, TConstArrayView<FString> InAttributeHLSLNames, bool bInSetToDefaults, bool bSpawnOnly, bool bPartialWrites, TArray<FText>& OutErrors, FString& OutHLSL) const override;
-	virtual bool GenerateIterationSourceNamespaceWriteAttributesHLSL(FNiagaraDataInterfaceGPUParamInfo& DIInstanceInfo, const FNiagaraVariable& IterationSourceVar, TConstArrayView<FNiagaraVariable> InArguments, TConstArrayView<FNiagaraVariable> InAttributes, TConstArrayView<FString> InAttributeHLSLNames, bool bSpawnOnly, bool bPartialWrites, TArray<FText>& OutErrors, FString& OutHLSL) const override;
+	virtual bool GenerateIterationSourceNamespaceWriteAttributesHLSL(FNiagaraDataInterfaceGPUParamInfo& DIInstanceInfo, const FNiagaraVariable& IterationSourceVar, TConstArrayView<FNiagaraVariable> InArguments, TConstArrayView<FNiagaraVariable> InAttributes, TConstArrayView<FString> InAttributeHLSLNames, TConstArrayView<FNiagaraVariable> InAllAttributes, bool bSpawnOnly, bool bPartialWrites, TArray<FText>& OutErrors, FString& OutHLSL) const override;
 #endif
 
 	static int32 GetComponentCountFromFuncName(const FName& FuncName);
@@ -300,8 +324,11 @@ public:
 protected:
 #if WITH_EDITORONLY_DATA
 	void WriteSetHLSL(const FNiagaraDataInterfaceGPUParamInfo& ParamInfo, const FNiagaraDataInterfaceGeneratedFunction& FunctionInfo, int FunctionInstanceIndex, int32 InNumChannels, FString& OutHLSL);
+	void WriteSetAtIndexHLSL(const FNiagaraDataInterfaceGPUParamInfo& ParamInfo, const FNiagaraDataInterfaceGeneratedFunction& FunctionInfo, int32 InNumChannels, FString& OutHLSL);
 	void WriteGetHLSL(const FNiagaraDataInterfaceGPUParamInfo& ParamInfo, const FNiagaraDataInterfaceGeneratedFunction& FunctionInfo, int FunctionInstanceIndex, int32 InNumChannels, FString& OutHLSL);
+	void WriteGetAtIndexHLSL(const FNiagaraDataInterfaceGPUParamInfo& ParamInfo, const FNiagaraDataInterfaceGeneratedFunction& FunctionInfo, int32 InNumChannels, FString& OutHLSL);
 	void WriteSampleHLSL(const FNiagaraDataInterfaceGPUParamInfo& ParamInfo, const FNiagaraDataInterfaceGeneratedFunction& FunctionInfo, int FunctionInstanceIndex, int32 InNumChannels, FString SampleFunction, FString& OutHLSL);
+	void WriteSampleAtIndexHLSL(const FNiagaraDataInterfaceGPUParamInfo& ParamInfo, const FNiagaraDataInterfaceGeneratedFunction& FunctionInfo, int32 InNumChannels, FString SampleFunction, FString& OutHLSL);
 	void WriteAttributeGetIndexHLSL(const FNiagaraDataInterfaceGPUParamInfo& ParamInfo, const FNiagaraDataInterfaceGeneratedFunction& FunctionInfo, int FunctionInstanceIndex, int32 InNumChannels, FString& OutHLSL);
 
 	const TCHAR* TypeDefinitionToHLSLTypeString(const FNiagaraTypeDefinition& InDef) const;
diff --git a/Engine/Plugins/FX/Niagara/Source/Niagara/Private/NiagaraDataInterfaceGrid2DCollection.cpp b/Engine/Plugins/FX/Niagara/Source/Niagara/Private/NiagaraDataInterfaceGrid2DCollection.cpp
--- a/Engine/Plugins/FX/Niagara/Source/Niagara/Private/NiagaraDataInterfaceGrid2DCollection.cpp
+++ b/Engine/Plugins/FX/Niagara/Source/Niagara/Private/NiagaraDataInterfaceGrid2DCollection.cpp
@@ -1636,7 +1636,7 @@ bool UNiagaraDataInterfaceGrid2DCollection::GenerateIterationSourceNamespaceRead
 	}
 	return true;
 }
-bool UNiagaraDataInterfaceGrid2DCollection::GenerateIterationSourceNamespaceWriteAttributesHLSL(FNiagaraDataInterfaceGPUParamInfo& DIInstanceInfo, const FNiagaraVariable& IterationSourceVar, TConstArrayView<FNiagaraVariable> InArguments, TConstArrayView<FNiagaraVariable> InAttributes, TConstArrayView<FString> InAttributeHLSLNames, bool bSpawnOnly, bool bPartialWrites, TArray<FText>& OutErrors, FString& OutHLSL) const
+bool UNiagaraDataInterfaceGrid2DCollection::GenerateIterationSourceNamespaceWriteAttributesHLSL(FNiagaraDataInterfaceGPUParamInfo& DIInstanceInfo, const FNiagaraVariable& IterationSourceVar, TConstArrayView<FNiagaraVariable> InArguments, TConstArrayView<FNiagaraVariable> InAttributes, TConstArrayView<FString> InAttributeHLSLNames, TConstArrayView<FNiagaraVariable> InAllAttributes, bool bSpawnOnly, bool bPartialWrites, TArray<FText>& OutErrors, FString& OutHLSL) const
 {
 	FString DIVarName;
 	OutHLSL += TEXT("\t//Generated by UNiagaraDataInterfaceGrid2DCollection::GenerateIterationSourceNamespaceWriteAttributesHLSL\n");
diff --git a/Engine/Plugins/FX/Niagara/Source/Niagara/Private/NiagaraDataInterfaceGrid3DCollection.cpp b/Engine/Plugins/FX/Niagara/Source/Niagara/Private/NiagaraDataInterfaceGrid3DCollection.cpp
--- a/Engine/Plugins/FX/Niagara/Source/Niagara/Private/NiagaraDataInterfaceGrid3DCollection.cpp
+++ b/Engine/Plugins/FX/Niagara/Source/Niagara/Private/NiagaraDataInterfaceGrid3DCollection.cpp
@@ -9,6 +9,7 @@
 #include "NiagaraShaderParametersBuilder.h"
 #include "NiagaraSystemInstance.h"
 #include "NiagaraGpuComputeDebugInterface.h"
+#include "NiagaraDataInterfaceGrid3DCollectionUtils.h"
 
 #include "Engine/TextureRenderTargetVolume.h"
 #include "Engine/VolumeTexture.h"
@@ -38,11 +39,28 @@ BEGIN_SHADER_PARAMETER_STRUCT(FNDIGrid3DShaderParameters, )
 	SHADER_PARAMETER_SRV(Buffer<float4>,					PerAttributeData)
 END_SHADER_PARAMETER_STRUCT()
 
+static float GNiagaraGrid3DResolutionMultiplier = 1.0f;
+static FAutoConsoleVariableRef CVarNiagaraGrid3DResolutionMultiplier(
+	TEXT("fx.Niagara.Grid3D.ResolutionMultiplier"),
+	GNiagaraGrid3DResolutionMultiplier,
+	TEXT("Optional global modifier to grid resolution\n"),
+	ECVF_Default
+);
+
+static int32 GNiagaraGrid3DOverrideFormat = -1;
+static FAutoConsoleVariableRef CVarNiagaraGrid3DOverrideFormat(
+	TEXT("fx.Niagara.Grid3D.OverrideFormat"),
+	GNiagaraGrid3DOverrideFormat,
+	TEXT("Optional override for all grids to use this format.\n"),
+	ECVF_Default
+);
+
 const FString UNiagaraDataInterfaceGrid3DCollection::NumTilesName(TEXT("_NumTiles"));
 const FString UNiagaraDataInterfaceGrid3DCollection::OneOverNumTilesName(TEXT("_OneOverNumTiles"));
 const FString UNiagaraDataInterfaceGrid3DCollection::UnitClampMinName(TEXT("_UnitClampMin"));
 const FString UNiagaraDataInterfaceGrid3DCollection::UnitClampMaxName(TEXT("_UnitClampMax"));
 
+const FName UNiagaraDataInterfaceGrid3DCollection::NAME_Attribute("Attribute");
 const FString UNiagaraDataInterfaceGrid3DCollection::GridName(TEXT("_Grid"));
 const FString UNiagaraDataInterfaceGrid3DCollection::OutputGridName(TEXT("_OutputGrid"));
 const FString UNiagaraDataInterfaceGrid3DCollection::SamplerName(TEXT("_GridSampler"));
@@ -71,6 +89,26 @@ const FName UNiagaraDataInterfaceGrid3DCollection::GetPreviousValueAtIndexFuncti
 const FName UNiagaraDataInterfaceGrid3DCollection::SamplePreviousGridAtIndexFunctionName("SamplePreviousGridAtIndex");
 const FName UNiagaraDataInterfaceGrid3DCollection::CubicSamplePreviousGridAtIndexFunctionName("CubicSamplePreviousGridAtIndex");
 
+const FName UNiagaraDataInterfaceGrid3DCollection::SetVector4ValueAtIndexFunctionName("SetVector4ValueAtIndex");
+const FName UNiagaraDataInterfaceGrid3DCollection::SetVector3ValueAtIndexFunctionName("SetVector3ValueAtIndex");
+const FName UNiagaraDataInterfaceGrid3DCollection::SetVector2ValueAtIndexFunctionName("SetVector2ValueAtIndex");
+const FName UNiagaraDataInterfaceGrid3DCollection::SetFloatValueAtIndexFunctionName("SetFloatValueAtIndex");
+
+const FName UNiagaraDataInterfaceGrid3DCollection::GetPreviousVector4ValueAtIndexFunctionName("GetPreviousVector4ValueAtIndex");
+const FName UNiagaraDataInterfaceGrid3DCollection::GetPreviousVector3ValueAtIndexFunctionName("GetPreviousVector3ValueAtIndex");
+const FName UNiagaraDataInterfaceGrid3DCollection::GetPreviousVector2ValueAtIndexFunctionName("GetPreviousVector2ValueAtIndex");
+const FName UNiagaraDataInterfaceGrid3DCollection::GetPreviousFloatValueAtIndexFunctionName("GetPreviousFloatValueAtIndex");
+
+const FName UNiagaraDataInterfaceGrid3DCollection::SamplePreviousGridVector4AtIndexFunctionName("SamplePreviousGridVector4AtIndex");
+const FName UNiagaraDataInterfaceGrid3DCollection::SamplePreviousGridVector3AtIndexFunctionName("SamplePreviousGridVector3AtIndex");
+const FName UNiagaraDataInterfaceGrid3DCollection::SamplePreviousGridVector2AtIndexFunctionName("SamplePreviousGridVector2AtIndex");
+const FName UNiagaraDataInterfaceGrid3DCollection::SamplePreviousGridFloatAtIndexFunctionName("SamplePreviousGridFloatAtIndex");
+
+const FName UNiagaraDataInterfaceGrid3DCollection::CubicSamplePreviousGridVector4AtIndexFunctionName("CubicSamplePreviousGridVector4AtIndex");
+const FName UNiagaraDataInterfaceGrid3DCollection::CubicSamplePreviousGridVector3AtIndexFunctionName("CubicSamplePreviousGridVector3AtIndex");
+const FName UNiagaraDataInterfaceGrid3DCollection::CubicSamplePreviousGridVector2AtIndexFunctionName("CubicSamplePreviousGridVector2AtIndex");
+const FName UNiagaraDataInterfaceGrid3DCollection::CubicSamplePreviousGridFloatAtIndexFunctionName("CubicSamplePreviousGridFloatAtIndex");
+
 const FName UNiagaraDataInterfaceGrid3DCollection::GetPreviousVector4ValueFunctionName("GetPreviousVector4Value");
 const FName UNiagaraDataInterfaceGrid3DCollection::SamplePreviousGridVector4FunctionName("SamplePreviousGridVector4Value");
 const FName UNiagaraDataInterfaceGrid3DCollection::CubicSamplePreviousGridVector4FunctionName("CubicSamplePreviousGridVector4Value");
@@ -115,28 +153,98 @@ struct FNiagaraGridCollection3DDIFunctionVersion
 	};
 };
 
+
 bool UNiagaraDataInterfaceGrid3DCollection::CanCreateVarFromFuncName(const FName& FuncName)
 {
-	if (FuncName == UNiagaraDataInterfaceGrid3DCollection::SetVector4ValueFunctionName	|| FuncName == UNiagaraDataInterfaceGrid3DCollection::GetPreviousVector4ValueFunctionName || FuncName == UNiagaraDataInterfaceGrid3DCollection::SamplePreviousGridVector4FunctionName || FuncName == UNiagaraDataInterfaceGrid3DCollection::CubicSamplePreviousGridVector4FunctionName)
+	if (
+		FuncName == UNiagaraDataInterfaceGrid3DCollection::SetVector4ValueFunctionName ||
+		FuncName == UNiagaraDataInterfaceGrid3DCollection::GetPreviousVector4ValueFunctionName ||
+		FuncName == UNiagaraDataInterfaceGrid3DCollection::SamplePreviousGridVector4FunctionName ||
+		FuncName == UNiagaraDataInterfaceGrid3DCollection::CubicSamplePreviousGridVector4FunctionName ||
+		FuncName == UNiagaraDataInterfaceGrid3DCollection::SetVector4ValueAtIndexFunctionName ||
+		FuncName == UNiagaraDataInterfaceGrid3DCollection::GetPreviousVector4ValueAtIndexFunctionName ||
+		FuncName == UNiagaraDataInterfaceGrid3DCollection::SamplePreviousGridVector4AtIndexFunctionName ||
+		FuncName == UNiagaraDataInterfaceGrid3DCollection::CubicSamplePreviousGridVector4AtIndexFunctionName)
+
 		return true;
-	else if (FuncName == UNiagaraDataInterfaceGrid3DCollection::SetVector3ValueFunctionName || FuncName == UNiagaraDataInterfaceGrid3DCollection::SetVectorValueFunctionName ||  FuncName == UNiagaraDataInterfaceGrid3DCollection::GetPreviousVectorValueFunctionName || FuncName == UNiagaraDataInterfaceGrid3DCollection::SamplePreviousGridVectorFunctionName || FuncName == UNiagaraDataInterfaceGrid3DCollection::CubicSamplePreviousGridVectorFunctionName)
+	else if (
+		FuncName == UNiagaraDataInterfaceGrid3DCollection::SetVector3ValueFunctionName ||
+		FuncName == UNiagaraDataInterfaceGrid3DCollection::SetVectorValueFunctionName ||
+		FuncName == UNiagaraDataInterfaceGrid3DCollection::GetPreviousVectorValueFunctionName ||
+		FuncName == UNiagaraDataInterfaceGrid3DCollection::SamplePreviousGridVectorFunctionName ||
+		FuncName == UNiagaraDataInterfaceGrid3DCollection::CubicSamplePreviousGridVectorFunctionName ||
+		FuncName == UNiagaraDataInterfaceGrid3DCollection::SetVector3ValueAtIndexFunctionName ||
+		FuncName == UNiagaraDataInterfaceGrid3DCollection::GetPreviousVector3ValueAtIndexFunctionName ||
+		FuncName == UNiagaraDataInterfaceGrid3DCollection::SamplePreviousGridVector3AtIndexFunctionName ||
+		FuncName == UNiagaraDataInterfaceGrid3DCollection::CubicSamplePreviousGridVector3AtIndexFunctionName)
 		return true;
-	else if (FuncName == UNiagaraDataInterfaceGrid3DCollection::SetVector2ValueFunctionName || FuncName == UNiagaraDataInterfaceGrid3DCollection::SetVector2DValueFunctionName || FuncName == UNiagaraDataInterfaceGrid3DCollection::GetPreviousVector2DValueFunctionName || FuncName == UNiagaraDataInterfaceGrid3DCollection::SamplePreviousGridVector2DFunctionName || FuncName == UNiagaraDataInterfaceGrid3DCollection::CubicSamplePreviousGridVector2DFunctionName)
+	else if (
+		FuncName == UNiagaraDataInterfaceGrid3DCollection::SetVector2ValueFunctionName ||
+		FuncName == UNiagaraDataInterfaceGrid3DCollection::SetVector2DValueFunctionName ||
+		FuncName == UNiagaraDataInterfaceGrid3DCollection::GetPreviousVector2DValueFunctionName ||
+		FuncName == UNiagaraDataInterfaceGrid3DCollection::SamplePreviousGridVector2DFunctionName ||
+		FuncName == UNiagaraDataInterfaceGrid3DCollection::CubicSamplePreviousGridVector2DFunctionName ||
+		FuncName == UNiagaraDataInterfaceGrid3DCollection::SetVector2ValueAtIndexFunctionName ||
+		FuncName == UNiagaraDataInterfaceGrid3DCollection::GetPreviousVector2ValueAtIndexFunctionName ||
+		FuncName == UNiagaraDataInterfaceGrid3DCollection::SamplePreviousGridVector2AtIndexFunctionName ||
+		FuncName == UNiagaraDataInterfaceGrid3DCollection::CubicSamplePreviousGridVector2AtIndexFunctionName)
 		return true;
-	else if (FuncName == UNiagaraDataInterfaceGrid3DCollection::SetFloatValueFunctionName       || FuncName == UNiagaraDataInterfaceGrid3DCollection::GetPreviousFloatValueFunctionName || FuncName == UNiagaraDataInterfaceGrid3DCollection::SamplePreviousGridFloatFunctionName || FuncName == UNiagaraDataInterfaceGrid3DCollection::CubicSamplePreviousGridFloatFunctionName)
+	else if (
+		FuncName == UNiagaraDataInterfaceGrid3DCollection::SetFloatValueFunctionName ||
+		FuncName == UNiagaraDataInterfaceGrid3DCollection::GetPreviousFloatValueFunctionName ||
+		FuncName == UNiagaraDataInterfaceGrid3DCollection::SamplePreviousGridFloatFunctionName ||
+		FuncName == UNiagaraDataInterfaceGrid3DCollection::CubicSamplePreviousGridFloatFunctionName ||
+		FuncName == UNiagaraDataInterfaceGrid3DCollection::SetFloatValueAtIndexFunctionName ||
+		FuncName == UNiagaraDataInterfaceGrid3DCollection::GetPreviousFloatValueAtIndexFunctionName ||
+		FuncName == UNiagaraDataInterfaceGrid3DCollection::SamplePreviousGridFloatAtIndexFunctionName ||
+		FuncName == UNiagaraDataInterfaceGrid3DCollection::CubicSamplePreviousGridFloatAtIndexFunctionName)
 		return true;
 	return false;
 }
 
 FNiagaraTypeDefinition UNiagaraDataInterfaceGrid3DCollection::GetValueTypeFromFuncName(const FName& FuncName)
 {
-	if (FuncName == UNiagaraDataInterfaceGrid3DCollection::SetVector4ValueFunctionName || FuncName == UNiagaraDataInterfaceGrid3DCollection::GetPreviousVector4ValueFunctionName || FuncName == UNiagaraDataInterfaceGrid3DCollection::SamplePreviousGridVector4FunctionName || FuncName == UNiagaraDataInterfaceGrid3DCollection::CubicSamplePreviousGridVector4FunctionName)
+	if (
+		FuncName == UNiagaraDataInterfaceGrid3DCollection::SetVector4ValueFunctionName ||
+		FuncName == UNiagaraDataInterfaceGrid3DCollection::GetPreviousVector4ValueFunctionName ||
+		FuncName == UNiagaraDataInterfaceGrid3DCollection::SamplePreviousGridVector4FunctionName ||
+		FuncName == UNiagaraDataInterfaceGrid3DCollection::CubicSamplePreviousGridVector4FunctionName ||
+		FuncName == UNiagaraDataInterfaceGrid3DCollection::SetVector4ValueAtIndexFunctionName ||
+		FuncName == UNiagaraDataInterfaceGrid3DCollection::GetPreviousVector4ValueAtIndexFunctionName ||
+		FuncName == UNiagaraDataInterfaceGrid3DCollection::SamplePreviousGridVector4AtIndexFunctionName ||
+		FuncName == UNiagaraDataInterfaceGrid3DCollection::CubicSamplePreviousGridVector4AtIndexFunctionName)
 		return FNiagaraTypeDefinition::GetVec4Def();
-	else if (FuncName == UNiagaraDataInterfaceGrid3DCollection::SetVector3ValueFunctionName || FuncName == UNiagaraDataInterfaceGrid3DCollection::SetVectorValueFunctionName || FuncName == UNiagaraDataInterfaceGrid3DCollection::GetPreviousVectorValueFunctionName || FuncName == UNiagaraDataInterfaceGrid3DCollection::SamplePreviousGridVectorFunctionName || FuncName == UNiagaraDataInterfaceGrid3DCollection::CubicSamplePreviousGridVectorFunctionName)
+	else if (
+		FuncName == UNiagaraDataInterfaceGrid3DCollection::SetVector3ValueFunctionName ||
+		FuncName == UNiagaraDataInterfaceGrid3DCollection::SetVectorValueFunctionName ||
+		FuncName == UNiagaraDataInterfaceGrid3DCollection::GetPreviousVectorValueFunctionName ||
+		FuncName == UNiagaraDataInterfaceGrid3DCollection::SamplePreviousGridVectorFunctionName ||
+		FuncName == UNiagaraDataInterfaceGrid3DCollection::CubicSamplePreviousGridVectorFunctionName ||
+		FuncName == UNiagaraDataInterfaceGrid3DCollection::SetVector3ValueAtIndexFunctionName ||
+		FuncName == UNiagaraDataInterfaceGrid3DCollection::GetPreviousVector3ValueAtIndexFunctionName ||
+		FuncName == UNiagaraDataInterfaceGrid3DCollection::SamplePreviousGridVector3AtIndexFunctionName ||
+		FuncName == UNiagaraDataInterfaceGrid3DCollection::CubicSamplePreviousGridVector3AtIndexFunctionName)
 		return FNiagaraTypeDefinition::GetVec3Def();
-	else if (FuncName == UNiagaraDataInterfaceGrid3DCollection::SetVector2ValueFunctionName || FuncName == UNiagaraDataInterfaceGrid3DCollection::SetVector2DValueFunctionName ||  FuncName == UNiagaraDataInterfaceGrid3DCollection::SamplePreviousGridVector2DFunctionName || FuncName == UNiagaraDataInterfaceGrid3DCollection::CubicSamplePreviousGridVector2DFunctionName)
+	else if (
+		FuncName == UNiagaraDataInterfaceGrid3DCollection::SetVector2ValueFunctionName ||
+		FuncName == UNiagaraDataInterfaceGrid3DCollection::SetVector2DValueFunctionName ||
+		FuncName == UNiagaraDataInterfaceGrid3DCollection::GetPreviousVector2DValueFunctionName ||
+		FuncName == UNiagaraDataInterfaceGrid3DCollection::SamplePreviousGridVector2DFunctionName ||
+		FuncName == UNiagaraDataInterfaceGrid3DCollection::CubicSamplePreviousGridVector2DFunctionName ||
+		FuncName == UNiagaraDataInterfaceGrid3DCollection::SetVector2ValueAtIndexFunctionName ||
+		FuncName == UNiagaraDataInterfaceGrid3DCollection::GetPreviousVector2ValueAtIndexFunctionName ||
+		FuncName == UNiagaraDataInterfaceGrid3DCollection::SamplePreviousGridVector2AtIndexFunctionName ||
+		FuncName == UNiagaraDataInterfaceGrid3DCollection::CubicSamplePreviousGridVector2AtIndexFunctionName)
 		return FNiagaraTypeDefinition::GetVec2Def();
-	else if (FuncName == UNiagaraDataInterfaceGrid3DCollection::SetFloatValueFunctionName || FuncName == UNiagaraDataInterfaceGrid3DCollection::GetPreviousFloatValueFunctionName || FuncName == UNiagaraDataInterfaceGrid3DCollection::SamplePreviousGridFloatFunctionName || FuncName == UNiagaraDataInterfaceGrid3DCollection::CubicSamplePreviousGridFloatFunctionName)
+	else if (
+		FuncName == UNiagaraDataInterfaceGrid3DCollection::SetFloatValueFunctionName ||
+		FuncName == UNiagaraDataInterfaceGrid3DCollection::GetPreviousFloatValueFunctionName ||
+		FuncName == UNiagaraDataInterfaceGrid3DCollection::SamplePreviousGridFloatFunctionName ||
+		FuncName == UNiagaraDataInterfaceGrid3DCollection::CubicSamplePreviousGridFloatFunctionName ||
+		FuncName == UNiagaraDataInterfaceGrid3DCollection::SetFloatValueAtIndexFunctionName ||
+		FuncName == UNiagaraDataInterfaceGrid3DCollection::GetPreviousFloatValueAtIndexFunctionName ||
+		FuncName == UNiagaraDataInterfaceGrid3DCollection::SamplePreviousGridFloatAtIndexFunctionName ||
+		FuncName == UNiagaraDataInterfaceGrid3DCollection::CubicSamplePreviousGridFloatAtIndexFunctionName)
 		return FNiagaraTypeDefinition::GetFloatDef();
 
 	return FNiagaraTypeDefinition();
@@ -144,36 +252,59 @@ FNiagaraTypeDefinition UNiagaraDataInterfaceGrid3DCollection::GetValueTypeFromFu
 
 int32 UNiagaraDataInterfaceGrid3DCollection::GetComponentCountFromFuncName(const FName& FuncName)
 {
-	if (FuncName == UNiagaraDataInterfaceGrid3DCollection::SetVector4ValueFunctionName || FuncName == UNiagaraDataInterfaceGrid3DCollection::GetVector4AttributeIndexFunctionName || FuncName == UNiagaraDataInterfaceGrid3DCollection::GetPreviousVector4ValueFunctionName || FuncName == UNiagaraDataInterfaceGrid3DCollection::SamplePreviousGridVector4FunctionName || FuncName == UNiagaraDataInterfaceGrid3DCollection::CubicSamplePreviousGridVector4FunctionName)
+	if (
+		FuncName == UNiagaraDataInterfaceGrid3DCollection::SetVector4ValueFunctionName ||
+		FuncName == UNiagaraDataInterfaceGrid3DCollection::GetPreviousVector4ValueFunctionName ||
+		FuncName == UNiagaraDataInterfaceGrid3DCollection::SamplePreviousGridVector4FunctionName ||
+		FuncName == UNiagaraDataInterfaceGrid3DCollection::CubicSamplePreviousGridVector4FunctionName ||
+		FuncName == UNiagaraDataInterfaceGrid3DCollection::SetVector4ValueAtIndexFunctionName ||
+		FuncName == UNiagaraDataInterfaceGrid3DCollection::GetPreviousVector4ValueAtIndexFunctionName ||
+		FuncName == UNiagaraDataInterfaceGrid3DCollection::SamplePreviousGridVector4AtIndexFunctionName ||
+		FuncName == UNiagaraDataInterfaceGrid3DCollection::CubicSamplePreviousGridVector4AtIndexFunctionName || 
+		FuncName == UNiagaraDataInterfaceGrid3DCollection::GetVector4AttributeIndexFunctionName)
 		return 4;
-	else if (FuncName == UNiagaraDataInterfaceGrid3DCollection::SetVector3ValueFunctionName || FuncName == UNiagaraDataInterfaceGrid3DCollection::SetVectorValueFunctionName ||  FuncName == UNiagaraDataInterfaceGrid3DCollection::GetVectorAttributeIndexFunctionName || FuncName == UNiagaraDataInterfaceGrid3DCollection::GetPreviousVectorValueFunctionName || FuncName == UNiagaraDataInterfaceGrid3DCollection::SamplePreviousGridVectorFunctionName || FuncName == UNiagaraDataInterfaceGrid3DCollection::CubicSamplePreviousGridVectorFunctionName)
+	else if (
+		FuncName == UNiagaraDataInterfaceGrid3DCollection::SetVector3ValueFunctionName ||
+		FuncName == UNiagaraDataInterfaceGrid3DCollection::SetVectorValueFunctionName ||
+		FuncName == UNiagaraDataInterfaceGrid3DCollection::GetPreviousVectorValueFunctionName ||
+		FuncName == UNiagaraDataInterfaceGrid3DCollection::SamplePreviousGridVectorFunctionName ||
+		FuncName == UNiagaraDataInterfaceGrid3DCollection::CubicSamplePreviousGridVectorFunctionName ||
+		FuncName == UNiagaraDataInterfaceGrid3DCollection::SetVector3ValueAtIndexFunctionName ||
+		FuncName == UNiagaraDataInterfaceGrid3DCollection::GetPreviousVector3ValueAtIndexFunctionName ||
+		FuncName == UNiagaraDataInterfaceGrid3DCollection::SamplePreviousGridVector3AtIndexFunctionName ||
+		FuncName == UNiagaraDataInterfaceGrid3DCollection::CubicSamplePreviousGridVector3AtIndexFunctionName ||
+		FuncName == UNiagaraDataInterfaceGrid3DCollection::GetVectorAttributeIndexFunctionName)
 		return 3;
-	else if (FuncName == UNiagaraDataInterfaceGrid3DCollection::SetVector2ValueFunctionName || FuncName == UNiagaraDataInterfaceGrid3DCollection::SetVector2DValueFunctionName ||  FuncName == UNiagaraDataInterfaceGrid3DCollection::GetVector2DAttributeIndexFunctionName|| FuncName == UNiagaraDataInterfaceGrid3DCollection::GetPreviousVector2DValueFunctionName || FuncName == UNiagaraDataInterfaceGrid3DCollection::SamplePreviousGridVector2DFunctionName || FuncName == UNiagaraDataInterfaceGrid3DCollection::CubicSamplePreviousGridVector2DFunctionName)
+	else if (
+		FuncName == UNiagaraDataInterfaceGrid3DCollection::SetVector2ValueFunctionName ||
+		FuncName == UNiagaraDataInterfaceGrid3DCollection::SetVector2DValueFunctionName ||
+		FuncName == UNiagaraDataInterfaceGrid3DCollection::GetPreviousVector2DValueFunctionName ||
+		FuncName == UNiagaraDataInterfaceGrid3DCollection::SamplePreviousGridVector2DFunctionName ||
+		FuncName == UNiagaraDataInterfaceGrid3DCollection::CubicSamplePreviousGridVector2DFunctionName ||
+		FuncName == UNiagaraDataInterfaceGrid3DCollection::SetVector2ValueAtIndexFunctionName ||
+		FuncName == UNiagaraDataInterfaceGrid3DCollection::GetPreviousVector2ValueAtIndexFunctionName ||
+		FuncName == UNiagaraDataInterfaceGrid3DCollection::SamplePreviousGridVector2AtIndexFunctionName ||
+		FuncName == UNiagaraDataInterfaceGrid3DCollection::CubicSamplePreviousGridVector2AtIndexFunctionName ||
+		FuncName == UNiagaraDataInterfaceGrid3DCollection::GetVector2DAttributeIndexFunctionName)
 		return 2;
-	else if (FuncName == UNiagaraDataInterfaceGrid3DCollection::SetFloatValueFunctionName || FuncName == UNiagaraDataInterfaceGrid3DCollection::GetPreviousFloatValueFunctionName || FuncName == UNiagaraDataInterfaceGrid3DCollection::SamplePreviousGridFloatFunctionName|| FuncName == UNiagaraDataInterfaceGrid3DCollection::GetFloatAttributeIndexFunctionName || FuncName == UNiagaraDataInterfaceGrid3DCollection::CubicSamplePreviousGridFloatFunctionName)
+	else if (
+		FuncName == UNiagaraDataInterfaceGrid3DCollection::SetFloatValueFunctionName ||
+		FuncName == UNiagaraDataInterfaceGrid3DCollection::GetPreviousFloatValueFunctionName ||
+		FuncName == UNiagaraDataInterfaceGrid3DCollection::SamplePreviousGridFloatFunctionName ||
+		FuncName == UNiagaraDataInterfaceGrid3DCollection::CubicSamplePreviousGridFloatFunctionName ||
+		FuncName == UNiagaraDataInterfaceGrid3DCollection::SetFloatValueAtIndexFunctionName ||
+		FuncName == UNiagaraDataInterfaceGrid3DCollection::GetPreviousFloatValueAtIndexFunctionName ||
+		FuncName == UNiagaraDataInterfaceGrid3DCollection::SamplePreviousGridFloatAtIndexFunctionName ||
+		FuncName == UNiagaraDataInterfaceGrid3DCollection::CubicSamplePreviousGridFloatAtIndexFunctionName ||
+		FuncName == UNiagaraDataInterfaceGrid3DCollection::GetFloatAttributeIndexFunctionName)
 		return 1;
 
 	return INDEX_NONE;
 }
-static float GNiagaraGrid3DResolutionMultiplier = 1.0f;
-static FAutoConsoleVariableRef CVarNiagaraGrid3DResolutionMultiplier(
-	TEXT("fx.Niagara.Grid3D.ResolutionMultiplier"),
-	GNiagaraGrid3DResolutionMultiplier,
-	TEXT("Optional global modifier to grid resolution\n"),
-	ECVF_Default
-);
-
-static int32 GNiagaraGrid3DOverrideFormat = -1;
-static FAutoConsoleVariableRef CVarNiagaraGrid3DOverrideFormat(
-	TEXT("fx.Niagara.Grid3D.OverrideFormat"),
-	GNiagaraGrid3DOverrideFormat,
-	TEXT("Optional override for all grids to use this format.\n"),
-	ECVF_Default
-);
 
 UNiagaraDataInterfaceGrid3DCollection::UNiagaraDataInterfaceGrid3DCollection(FObjectInitializer const& ObjectInitializer)
 	: Super(ObjectInitializer)
-	, NumAttributes(1)
+	, NumAttributes(0)
 	, OverrideBufferFormat(ENiagaraGpuBufferFormat::Float)
 	, bOverrideFormat(false)
 #if WITH_EDITORONLY_DATA
@@ -207,7 +338,7 @@ bool UNiagaraDataInterfaceGrid3DCollection::AppendCompileHash(FNiagaraCompileHas
 
 
 	InVisitor->UpdatePOD(TEXT("UNiagaraDataInterfaceGrid3DCollectionVersion"), (int32)FNiagaraGridCollection3DDIFunctionVersion::LatestVersion);
-	
+
 	//InVisitor->UpdateString(TEXT("UNiagaraDataInterfaceVolumeTextureHLSLSource"), GetShaderFileHash(TemplateShaderFilePath, EShaderPlatform::SP_PCD3D_SM5).ToString());
 	InVisitor->UpdateShaderParameters<FNDIGrid3DShaderParameters>();
 
@@ -335,7 +466,7 @@ void UNiagaraDataInterfaceGrid3DCollection::GetFunctions(TArray<FNiagaraFunction
 		Sig.Inputs.Add(FNiagaraVariable(FNiagaraTypeDefinition(GetClass()), TEXT("Grid")));
 		Sig.Inputs.Add(FNiagaraVariable(FNiagaraTypeDefinition::GetIntDef(), TEXT("IndexX")));
 		Sig.Inputs.Add(FNiagaraVariable(FNiagaraTypeDefinition::GetIntDef(), TEXT("IndexY")));
-		Sig.Inputs.Add(FNiagaraVariable(FNiagaraTypeDefinition::GetIntDef(), TEXT("IndexZ")));		
+		Sig.Inputs.Add(FNiagaraVariable(FNiagaraTypeDefinition::GetIntDef(), TEXT("IndexZ")));
 		Sig.Outputs.Add(FNiagaraVariable(FNiagaraTypeDefinition::GetFloatDef(), TEXT("Value")));
 
 		Sig.bExperimental = true;
@@ -345,6 +476,8 @@ void UNiagaraDataInterfaceGrid3DCollection::GetFunctions(TArray<FNiagaraFunction
 		Sig.bSupportsCPU = false;
 		Sig.bSupportsGPU = true;
 		Sig.bReadFunction = true;
+		Sig.bHidden = true;
+		Sig.bSoftDeprecatedFunction = true;
 #if WITH_EDITORONLY_DATA
 		Sig.Description = NSLOCTEXT("Niagara", "NiagaraDataInterfaceGridColl2D_PreviousValueAtIndexFunction", "Get the value at a specific index.");
 #endif
@@ -383,7 +516,7 @@ void UNiagaraDataInterfaceGrid3DCollection::GetFunctions(TArray<FNiagaraFunction
 		Sig.Inputs.Add(FNiagaraVariable(FNiagaraTypeDefinition(GetClass()), TEXT("Grid")));
 		Sig.Inputs.Add(FNiagaraVariable(FNiagaraTypeDefinition::GetIntDef(), TEXT("IndexX")));
 		Sig.Inputs.Add(FNiagaraVariable(FNiagaraTypeDefinition::GetIntDef(), TEXT("IndexY")));
-		Sig.Inputs.Add(FNiagaraVariable(FNiagaraTypeDefinition::GetIntDef(), TEXT("IndexZ")));	
+		Sig.Inputs.Add(FNiagaraVariable(FNiagaraTypeDefinition::GetIntDef(), TEXT("IndexZ")));
 		Sig.Inputs.Add(FNiagaraVariable(FNiagaraTypeDefinition::GetFloatDef(), TEXT("Value")));
 		Sig.Outputs.Add(FNiagaraVariable(FNiagaraTypeDefinition::GetIntDef(), TEXT("IGNORE")));
 
@@ -395,6 +528,7 @@ void UNiagaraDataInterfaceGrid3DCollection::GetFunctions(TArray<FNiagaraFunction
 		Sig.bSoftDeprecatedFunction = true;
 		Sig.bSupportsCPU = false;
 		Sig.bSupportsGPU = true;
+		Sig.bHidden = true;
 #if WITH_EDITORONLY_DATA
 		Sig.Description = NSLOCTEXT("Niagara", "NiagaraDataInterfaceGridColl2D_SetValueFunction", "Set the value at a specific index. Note that this is an older way of working with Grids. Consider using the SetFloat or other typed, named functions or parameter map variables with StackContext namespace instead.");
 #endif
@@ -452,7 +586,7 @@ void UNiagaraDataInterfaceGrid3DCollection::GetFunctions(TArray<FNiagaraFunction
 		Sig.Inputs.Add(FNiagaraVariable(FNiagaraTypeDefinition(GetClass()), TEXT("Grid")));
 		Sig.Inputs.Add(FNiagaraVariable(FNiagaraTypeDefinition::GetIntDef(), TEXT("IndexX")));
 		Sig.Inputs.Add(FNiagaraVariable(FNiagaraTypeDefinition::GetIntDef(), TEXT("IndexY")));
-		Sig.Inputs.Add(FNiagaraVariable(FNiagaraTypeDefinition::GetIntDef(), TEXT("IndexZ")));		
+		Sig.Inputs.Add(FNiagaraVariable(FNiagaraTypeDefinition::GetIntDef(), TEXT("IndexZ")));
 		Sig.Inputs.Add(FNiagaraVariable(FNiagaraTypeDefinition::GetIntDef(), TEXT("NumAttributesSet")));
 		Sig.Inputs.Add(FNiagaraVariable(FNiagaraTypeDefinition::GetIntDef(), TEXT("AttributeMask")));
 
@@ -479,7 +613,7 @@ void UNiagaraDataInterfaceGrid3DCollection::GetFunctions(TArray<FNiagaraFunction
 		Sig.Inputs.Add(FNiagaraVariable(FNiagaraTypeDefinition::GetIntDef(), TEXT("IndexY")));
 		Sig.Inputs.Add(FNiagaraVariable(FNiagaraTypeDefinition::GetIntDef(), TEXT("IndexZ")));
 		Sig.Inputs.Add(FNiagaraVariable(FNiagaraTypeDefinition::GetVec4Def(), TEXT("Value")));
-		Sig.FunctionSpecifiers.Add(FName("Attribute"));
+		Sig.FunctionSpecifiers.Add(FName(NAME_Attribute));
 		Sig.bMemberFunction = true;
 		Sig.bRequiresContext = false;
 		Sig.bExperimental = true;
@@ -502,7 +636,7 @@ void UNiagaraDataInterfaceGrid3DCollection::GetFunctions(TArray<FNiagaraFunction
 		Sig.Inputs.Add(FNiagaraVariable(FNiagaraTypeDefinition::GetIntDef(), TEXT("IndexY")));
 		Sig.Inputs.Add(FNiagaraVariable(FNiagaraTypeDefinition::GetIntDef(), TEXT("IndexZ")));
 		Sig.Outputs.Add(FNiagaraVariable(FNiagaraTypeDefinition::GetVec4Def(), TEXT("Value")));
-		Sig.FunctionSpecifiers.Add(FName("Attribute"));
+		Sig.FunctionSpecifiers.Add(FName(NAME_Attribute));
 		Sig.bMemberFunction = true;
 		Sig.bRequiresContext = false;
 		Sig.bExperimental = true;
@@ -523,7 +657,7 @@ void UNiagaraDataInterfaceGrid3DCollection::GetFunctions(TArray<FNiagaraFunction
 		Sig.Inputs.Add(FNiagaraVariable(FNiagaraTypeDefinition(GetClass()), TEXT("Grid")));
 		Sig.Inputs.Add(FNiagaraVariable(FNiagaraTypeDefinition::GetVec3Def(), TEXT("Unit")));
 		Sig.Outputs.Add(FNiagaraVariable(FNiagaraTypeDefinition::GetVec4Def(), TEXT("Value")));
-		Sig.FunctionSpecifiers.Add(FName("Attribute"));
+		Sig.FunctionSpecifiers.Add(FName(NAME_Attribute));
 		Sig.bMemberFunction = true;
 		Sig.bRequiresContext = false;
 		Sig.bExperimental = true;
@@ -544,7 +678,7 @@ void UNiagaraDataInterfaceGrid3DCollection::GetFunctions(TArray<FNiagaraFunction
 		Sig.Inputs.Add(FNiagaraVariable(FNiagaraTypeDefinition(GetClass()), TEXT("Grid")));
 		Sig.Inputs.Add(FNiagaraVariable(FNiagaraTypeDefinition::GetVec3Def(), TEXT("Unit")));
 		Sig.Outputs.Add(FNiagaraVariable(FNiagaraTypeDefinition::GetVec4Def(), TEXT("Value")));
-		Sig.FunctionSpecifiers.Add(FName("Attribute"));
+		Sig.FunctionSpecifiers.Add(FName(NAME_Attribute));
 		Sig.bMemberFunction = true;
 		Sig.bRequiresContext = false;
 		Sig.bExperimental = true;
@@ -567,7 +701,7 @@ void UNiagaraDataInterfaceGrid3DCollection::GetFunctions(TArray<FNiagaraFunction
 		Sig.Inputs.Add(FNiagaraVariable(FNiagaraTypeDefinition::GetIntDef(), TEXT("IndexY")));
 		Sig.Inputs.Add(FNiagaraVariable(FNiagaraTypeDefinition::GetIntDef(), TEXT("IndexZ")));
 		Sig.Inputs.Add(FNiagaraVariable(FNiagaraTypeDefinition::GetVec3Def(), TEXT("Value")));
-		Sig.FunctionSpecifiers.Add(FName("Attribute"));
+		Sig.FunctionSpecifiers.Add(FName(NAME_Attribute));
 		Sig.bMemberFunction = true;
 		Sig.bRequiresContext = false;
 		Sig.bExperimental = true;
@@ -591,7 +725,7 @@ void UNiagaraDataInterfaceGrid3DCollection::GetFunctions(TArray<FNiagaraFunction
 		Sig.Inputs.Add(FNiagaraVariable(FNiagaraTypeDefinition::GetIntDef(), TEXT("IndexY")));
 		Sig.Inputs.Add(FNiagaraVariable(FNiagaraTypeDefinition::GetIntDef(), TEXT("IndexZ")));
 		Sig.Outputs.Add(FNiagaraVariable(FNiagaraTypeDefinition::GetVec3Def(), TEXT("Value")));
-		Sig.FunctionSpecifiers.Add(FName("Attribute"));
+		Sig.FunctionSpecifiers.Add(FName(NAME_Attribute));
 		Sig.bMemberFunction = true;
 		Sig.bRequiresContext = false;
 		Sig.bExperimental = true;
@@ -612,7 +746,7 @@ void UNiagaraDataInterfaceGrid3DCollection::GetFunctions(TArray<FNiagaraFunction
 		Sig.Inputs.Add(FNiagaraVariable(FNiagaraTypeDefinition(GetClass()), TEXT("Grid")));
 		Sig.Inputs.Add(FNiagaraVariable(FNiagaraTypeDefinition::GetVec3Def(), TEXT("Unit")));
 		Sig.Outputs.Add(FNiagaraVariable(FNiagaraTypeDefinition::GetVec3Def(), TEXT("Value")));
-		Sig.FunctionSpecifiers.Add(FName("Attribute"));
+		Sig.FunctionSpecifiers.Add(FName(NAME_Attribute));
 		Sig.bMemberFunction = true;
 		Sig.bRequiresContext = false;
 		Sig.bExperimental = true;
@@ -633,7 +767,7 @@ void UNiagaraDataInterfaceGrid3DCollection::GetFunctions(TArray<FNiagaraFunction
 		Sig.Inputs.Add(FNiagaraVariable(FNiagaraTypeDefinition(GetClass()), TEXT("Grid")));
 		Sig.Inputs.Add(FNiagaraVariable(FNiagaraTypeDefinition::GetVec3Def(), TEXT("Unit")));
 		Sig.Outputs.Add(FNiagaraVariable(FNiagaraTypeDefinition::GetVec3Def(), TEXT("Value")));
-		Sig.FunctionSpecifiers.Add(FName("Attribute"));
+		Sig.FunctionSpecifiers.Add(FName(NAME_Attribute));
 		Sig.bMemberFunction = true;
 		Sig.bRequiresContext = false;
 		Sig.bExperimental = true;
@@ -656,7 +790,7 @@ void UNiagaraDataInterfaceGrid3DCollection::GetFunctions(TArray<FNiagaraFunction
 		Sig.Inputs.Add(FNiagaraVariable(FNiagaraTypeDefinition::GetIntDef(), TEXT("IndexY")));
 		Sig.Inputs.Add(FNiagaraVariable(FNiagaraTypeDefinition::GetIntDef(), TEXT("IndexZ")));
 		Sig.Inputs.Add(FNiagaraVariable(FNiagaraTypeDefinition::GetVec2Def(), TEXT("Value")));
-		Sig.FunctionSpecifiers.Add(FName("Attribute"));
+		Sig.FunctionSpecifiers.Add(FName(NAME_Attribute));
 		Sig.bMemberFunction = true;
 		Sig.bRequiresContext = false;
 		Sig.bExperimental = true;
@@ -680,7 +814,7 @@ void UNiagaraDataInterfaceGrid3DCollection::GetFunctions(TArray<FNiagaraFunction
 		Sig.Inputs.Add(FNiagaraVariable(FNiagaraTypeDefinition::GetIntDef(), TEXT("IndexY")));
 		Sig.Inputs.Add(FNiagaraVariable(FNiagaraTypeDefinition::GetIntDef(), TEXT("IndexZ")));
 		Sig.Outputs.Add(FNiagaraVariable(FNiagaraTypeDefinition::GetVec2Def(), TEXT("Value")));
-		Sig.FunctionSpecifiers.Add(FName("Attribute"));
+		Sig.FunctionSpecifiers.Add(FName(NAME_Attribute));
 		Sig.bMemberFunction = true;
 		Sig.bRequiresContext = false;
 		Sig.bExperimental = true;
@@ -701,7 +835,7 @@ void UNiagaraDataInterfaceGrid3DCollection::GetFunctions(TArray<FNiagaraFunction
 		Sig.Inputs.Add(FNiagaraVariable(FNiagaraTypeDefinition(GetClass()), TEXT("Grid")));
 		Sig.Inputs.Add(FNiagaraVariable(FNiagaraTypeDefinition::GetVec3Def(), TEXT("Unit")));
 		Sig.Outputs.Add(FNiagaraVariable(FNiagaraTypeDefinition::GetVec2Def(), TEXT("Value")));
-		Sig.FunctionSpecifiers.Add(FName("Attribute"));
+		Sig.FunctionSpecifiers.Add(FName(NAME_Attribute));
 		Sig.bMemberFunction = true;
 		Sig.bRequiresContext = false;
 		Sig.bExperimental = true;
@@ -722,7 +856,7 @@ void UNiagaraDataInterfaceGrid3DCollection::GetFunctions(TArray<FNiagaraFunction
 		Sig.Inputs.Add(FNiagaraVariable(FNiagaraTypeDefinition(GetClass()), TEXT("Grid")));
 		Sig.Inputs.Add(FNiagaraVariable(FNiagaraTypeDefinition::GetVec3Def(), TEXT("Unit")));
 		Sig.Outputs.Add(FNiagaraVariable(FNiagaraTypeDefinition::GetVec2Def(), TEXT("Value")));
-		Sig.FunctionSpecifiers.Add(FName("Attribute"));
+		Sig.FunctionSpecifiers.Add(FName(NAME_Attribute));
 		Sig.bMemberFunction = true;
 		Sig.bRequiresContext = false;
 		Sig.bExperimental = true;
@@ -745,7 +879,7 @@ void UNiagaraDataInterfaceGrid3DCollection::GetFunctions(TArray<FNiagaraFunction
 		Sig.Inputs.Add(FNiagaraVariable(FNiagaraTypeDefinition::GetIntDef(), TEXT("IndexY")));
 		Sig.Inputs.Add(FNiagaraVariable(FNiagaraTypeDefinition::GetIntDef(), TEXT("IndexZ")));
 		Sig.Inputs.Add(FNiagaraVariable(FNiagaraTypeDefinition::GetFloatDef(), TEXT("Value")));
-		Sig.FunctionSpecifiers.Add(FName("Attribute"));
+		Sig.FunctionSpecifiers.Add(FName(NAME_Attribute));
 		Sig.bMemberFunction = true;
 		Sig.bRequiresContext = false;
 		Sig.bExperimental = true;
@@ -768,7 +902,7 @@ void UNiagaraDataInterfaceGrid3DCollection::GetFunctions(TArray<FNiagaraFunction
 		Sig.Inputs.Add(FNiagaraVariable(FNiagaraTypeDefinition::GetIntDef(), TEXT("IndexY")));
 		Sig.Inputs.Add(FNiagaraVariable(FNiagaraTypeDefinition::GetIntDef(), TEXT("IndexZ")));
 		Sig.Outputs.Add(FNiagaraVariable(FNiagaraTypeDefinition::GetFloatDef(), TEXT("Value")));
-		Sig.FunctionSpecifiers.Add(FName("Attribute"));
+		Sig.FunctionSpecifiers.Add(FName(NAME_Attribute));
 		Sig.bMemberFunction = true;
 		Sig.bRequiresContext = false;
 		Sig.bExperimental = true;
@@ -789,7 +923,7 @@ void UNiagaraDataInterfaceGrid3DCollection::GetFunctions(TArray<FNiagaraFunction
 		Sig.Inputs.Add(FNiagaraVariable(FNiagaraTypeDefinition(GetClass()), TEXT("Grid")));
 		Sig.Inputs.Add(FNiagaraVariable(FNiagaraTypeDefinition::GetVec3Def(), TEXT("Unit")));
 		Sig.Outputs.Add(FNiagaraVariable(FNiagaraTypeDefinition::GetFloatDef(), TEXT("Value")));
-		Sig.FunctionSpecifiers.Add(FName("Attribute"));
+		Sig.FunctionSpecifiers.Add(FName(NAME_Attribute));
 		Sig.bMemberFunction = true;
 		Sig.bRequiresContext = false;
 		Sig.bExperimental = true;
@@ -810,7 +944,7 @@ void UNiagaraDataInterfaceGrid3DCollection::GetFunctions(TArray<FNiagaraFunction
 		Sig.Inputs.Add(FNiagaraVariable(FNiagaraTypeDefinition(GetClass()), TEXT("Grid")));
 		Sig.Inputs.Add(FNiagaraVariable(FNiagaraTypeDefinition::GetVec3Def(), TEXT("Unit")));
 		Sig.Outputs.Add(FNiagaraVariable(FNiagaraTypeDefinition::GetFloatDef(), TEXT("Value")));
-		Sig.FunctionSpecifiers.Add(FName("Attribute"));
+		Sig.FunctionSpecifiers.Add(FName(NAME_Attribute));
 		Sig.bMemberFunction = true;
 		Sig.bRequiresContext = false;
 		Sig.bExperimental = true;
@@ -911,7 +1045,7 @@ void UNiagaraDataInterfaceGrid3DCollection::GetFunctions(TArray<FNiagaraFunction
 		FNiagaraFunctionSignature Sig;
 		Sig.Name = SamplePreviousFullGridFunctionName;
 		Sig.Inputs.Add(FNiagaraVariable(FNiagaraTypeDefinition(GetClass()), TEXT("Grid")));
-		Sig.Inputs.Add(FNiagaraVariable(FNiagaraTypeDefinition::GetVec3Def(), TEXT("Unit")));		
+		Sig.Inputs.Add(FNiagaraVariable(FNiagaraTypeDefinition::GetVec3Def(), TEXT("Unit")));
 		Sig.Outputs.Add(FNiagaraVariable(FNiagaraTypeDefinition::GetFloatDef(), TEXT("Value")));
 
 		Sig.bExperimental = true;
@@ -921,6 +1055,8 @@ void UNiagaraDataInterfaceGrid3DCollection::GetFunctions(TArray<FNiagaraFunction
 		Sig.bSupportsCPU = false;
 		Sig.bSupportsGPU = true;
 		Sig.bReadFunction = true;
+		Sig.bHidden = true;
+		Sig.bSoftDeprecatedFunction = true;
 		OutFunctions.Add(Sig);
 	}
 
@@ -938,6 +1074,8 @@ void UNiagaraDataInterfaceGrid3DCollection::GetFunctions(TArray<FNiagaraFunction
 		Sig.bSupportsCPU = false;
 		Sig.bSupportsGPU = true;
 		Sig.bReadFunction = true;
+		Sig.bHidden = true;
+		Sig.bSoftDeprecatedFunction = true;
 		OutFunctions.Add(Sig);
 	}
 
@@ -946,7 +1084,7 @@ void UNiagaraDataInterfaceGrid3DCollection::GetFunctions(TArray<FNiagaraFunction
 		Sig.Name = GetVector4AttributeIndexFunctionName;
 		Sig.Inputs.Add(FNiagaraVariable(FNiagaraTypeDefinition(GetClass()), TEXT("Grid")));
 		Sig.Outputs.Add(FNiagaraVariable(FNiagaraTypeDefinition::GetIntDef(), TEXT("AttributeIndex")));
-		Sig.FunctionSpecifiers.Add(FName("Attribute"));
+		Sig.FunctionSpecifiers.Add(FName(NAME_Attribute));
 		Sig.bMemberFunction = true;
 		Sig.bRequiresContext = false;
 		Sig.bExperimental = true;
@@ -961,7 +1099,7 @@ void UNiagaraDataInterfaceGrid3DCollection::GetFunctions(TArray<FNiagaraFunction
 		Sig.Name = GetVectorAttributeIndexFunctionName;
 		Sig.Inputs.Add(FNiagaraVariable(FNiagaraTypeDefinition(GetClass()), TEXT("Grid")));
 		Sig.Outputs.Add(FNiagaraVariable(FNiagaraTypeDefinition::GetIntDef(), TEXT("AttributeIndex")));
-		Sig.FunctionSpecifiers.Add(FName("Attribute"));
+		Sig.FunctionSpecifiers.Add(FName(NAME_Attribute));
 		Sig.bMemberFunction = true;
 		Sig.bRequiresContext = false;
 		Sig.bExperimental = true;
@@ -976,7 +1114,7 @@ void UNiagaraDataInterfaceGrid3DCollection::GetFunctions(TArray<FNiagaraFunction
 		Sig.Name = GetVector2DAttributeIndexFunctionName;
 		Sig.Inputs.Add(FNiagaraVariable(FNiagaraTypeDefinition(GetClass()), TEXT("Grid")));
 		Sig.Outputs.Add(FNiagaraVariable(FNiagaraTypeDefinition::GetIntDef(), TEXT("AttributeIndex")));
-		Sig.FunctionSpecifiers.Add(FName("Attribute"));
+		Sig.FunctionSpecifiers.Add(FName(NAME_Attribute));
 		Sig.bMemberFunction = true;
 		Sig.bRequiresContext = false;
 		Sig.bExperimental = true;
@@ -991,7 +1129,7 @@ void UNiagaraDataInterfaceGrid3DCollection::GetFunctions(TArray<FNiagaraFunction
 		Sig.Name = GetFloatAttributeIndexFunctionName;
 		Sig.Inputs.Add(FNiagaraVariable(FNiagaraTypeDefinition(GetClass()), TEXT("Grid")));
 		Sig.Outputs.Add(FNiagaraVariable(FNiagaraTypeDefinition::GetIntDef(), TEXT("AttributeIndex")));
-		Sig.FunctionSpecifiers.Add(FName("Attribute"));
+		Sig.FunctionSpecifiers.Add(FName(NAME_Attribute));
 		Sig.bMemberFunction = true;
 		Sig.bRequiresContext = false;
 		Sig.bExperimental = true;
@@ -1001,165 +1139,495 @@ void UNiagaraDataInterfaceGrid3DCollection::GetFunctions(TArray<FNiagaraFunction
 		OutFunctions.Add(Sig);
 	}
 
-#if WITH_EDITORONLY_DATA
-	for (int32 i = StartIndex; i < OutFunctions.Num(); i++)
 	{
-		FNiagaraFunctionSignature& Function = OutFunctions[i];
-		Function.FunctionVersion = FNiagaraGridCollection3DDIFunctionVersion::LatestVersion;
-	}
-#endif
-}
-
-// #todo(dmp): expose more CPU functionality
-// #todo(dmp): ideally these would be exposed on the parent class, but we can't bind functions of parent classes but need to work on the interface
-// for sharing an instance data object with the super class
-DEFINE_NDI_DIRECT_FUNC_BINDER(UNiagaraDataInterfaceGrid3DCollection, VMGetWorldBBoxSize);
-DEFINE_NDI_DIRECT_FUNC_BINDER(UNiagaraDataInterfaceGrid3DCollection, VMGetCellSize);
-DEFINE_NDI_DIRECT_FUNC_BINDER(UNiagaraDataInterfaceGrid3DCollection, VMSetNumCells);
-DEFINE_NDI_DIRECT_FUNC_BINDER(UNiagaraDataInterfaceGrid3DCollection, VMGetNumCells);
-DEFINE_NDI_DIRECT_FUNC_BINDER(UNiagaraDataInterfaceGrid3DCollection, VMUnitToFloatIndex);
-DEFINE_NDI_DIRECT_FUNC_BINDER_WITH_PAYLOAD(UNiagaraDataInterfaceGrid3DCollection, VMGetAttributeIndex);
-void UNiagaraDataInterfaceGrid3DCollection::GetVMExternalFunction(const FVMExternalFunctionBindingInfo& BindingInfo, void* InstanceData, FVMExternalFunction &OutFunc)
-{
-	Super::GetVMExternalFunction(BindingInfo, InstanceData, OutFunc);
+		FNiagaraFunctionSignature Sig;
+		Sig.Name = GetPreviousFloatValueAtIndexFunctionName;
+		Sig.Inputs.Add(FNiagaraVariable(FNiagaraTypeDefinition(GetClass()), TEXT("Grid")));
+		Sig.Inputs.Add(FNiagaraVariable(FNiagaraTypeDefinition::GetIntDef(), TEXT("IndexX")));
+		Sig.Inputs.Add(FNiagaraVariable(FNiagaraTypeDefinition::GetIntDef(), TEXT("IndexY")));
+		Sig.Inputs.Add(FNiagaraVariable(FNiagaraTypeDefinition::GetIntDef(), TEXT("IndexZ")));
+		Sig.Inputs.Add(FNiagaraVariable(FNiagaraTypeDefinition::GetIntDef(), TEXT("AttributeIndex")));
+		Sig.Outputs.Add(FNiagaraVariable(FNiagaraTypeDefinition::GetFloatDef(), TEXT("Value")));
 
-	static const FName NAME_Attribute("Attribute");
+		Sig.bExperimental = true;
+		Sig.bMemberFunction = true;
+		Sig.bRequiresContext = false;
+		Sig.ModuleUsageBitmask = ENiagaraScriptUsageMask::Particle;
+		Sig.bSupportsCPU = false;
+		Sig.bSupportsGPU = true;
+		Sig.bReadFunction = true;
+#if WITH_EDITORONLY_DATA
+		Sig.Description = NSLOCTEXT("Niagara", "NiagaraDataInterfaceGridColl2D_PreviousValueAtIndexFunction", "Get the value at a specific index.");
+#endif
 
-	if (BindingInfo.Name == UNiagaraDataInterfaceRWBase::WorldBBoxSizeFunctionName)
-	{
-		check(BindingInfo.GetNumInputs() == 1 && BindingInfo.GetNumOutputs() == 3);
-		NDI_FUNC_BINDER(UNiagaraDataInterfaceGrid3DCollection, VMGetWorldBBoxSize)::Bind(this, OutFunc);
-	}
-	else if (BindingInfo.Name == UNiagaraDataInterfaceRWBase::CellSizeFunctionName)
-	{
-		// #todo(dmp): this will override the base class definition for GetCellSize because the data interface instance data computes cell size
-		// it would be nice to refactor this so it can be part of the super class
-		check(BindingInfo.GetNumInputs() == 1 && BindingInfo.GetNumOutputs() == 3);
-		NDI_FUNC_BINDER(UNiagaraDataInterfaceGrid3DCollection, VMGetCellSize)::Bind(this, OutFunc);
-	}
-	else if (BindingInfo.Name == SetNumCellsFunctionName)
-	{
-		check(BindingInfo.GetNumInputs() == 4 && BindingInfo.GetNumOutputs() == 1);
-		NDI_FUNC_BINDER(UNiagaraDataInterfaceGrid3DCollection, VMSetNumCells)::Bind(this, OutFunc);
-	}
-	else if (BindingInfo.Name == NumCellsFunctionName)
-	{
-		check(BindingInfo.GetNumInputs() == 1 && BindingInfo.GetNumOutputs() == 3);
-		NDI_FUNC_BINDER(UNiagaraDataInterfaceGrid3DCollection, VMGetNumCells)::Bind(this, OutFunc);
-	}
-	else if (BindingInfo.Name == UnitToFloatIndexFunctionName)
-	{		
-		check(BindingInfo.GetNumInputs() == 4 && BindingInfo.GetNumOutputs() == 3);
-		NDI_FUNC_BINDER(UNiagaraDataInterfaceGrid3DCollection, VMUnitToFloatIndex)::Bind(this, OutFunc);
-	}
-	else if (BindingInfo.Name == GetVector4AttributeIndexFunctionName)
-	{
-		FName AttributeName = BindingInfo.FindSpecifier(NAME_Attribute)->Value;
-		check(BindingInfo.GetNumInputs() == 1 && BindingInfo.GetNumOutputs() == 1);
-		NDI_FUNC_BINDER(UNiagaraDataInterfaceGrid3DCollection, VMGetAttributeIndex)::Bind(this, OutFunc, AttributeName, 4);
-	}
-	else if (BindingInfo.Name == GetVectorAttributeIndexFunctionName)
-	{
-		FName AttributeName = BindingInfo.FindSpecifier(NAME_Attribute)->Value;
-		check(BindingInfo.GetNumInputs() == 1 && BindingInfo.GetNumOutputs() == 1);
-		NDI_FUNC_BINDER(UNiagaraDataInterfaceGrid3DCollection, VMGetAttributeIndex)::Bind(this, OutFunc, AttributeName, 3);
-	}
-	else if (BindingInfo.Name == GetVector2DAttributeIndexFunctionName)
-	{
-		FName AttributeName = BindingInfo.FindSpecifier(NAME_Attribute)->Value;
-		check(BindingInfo.GetNumInputs() == 1 && BindingInfo.GetNumOutputs() == 1);
-		NDI_FUNC_BINDER(UNiagaraDataInterfaceGrid3DCollection, VMGetAttributeIndex)::Bind(this, OutFunc, AttributeName, 2);
-	}
-	else if (BindingInfo.Name == GetFloatAttributeIndexFunctionName)
-	{
-		FName AttributeName = BindingInfo.FindSpecifier(NAME_Attribute)->Value;
-		check(BindingInfo.GetNumInputs() == 1 && BindingInfo.GetNumOutputs() == 1);
-		NDI_FUNC_BINDER(UNiagaraDataInterfaceGrid3DCollection, VMGetAttributeIndex)::Bind(this, OutFunc, AttributeName, 1);
+		OutFunctions.Add(Sig);
 	}
-}
-
-void UNiagaraDataInterfaceGrid3DCollection::VMGetNumCells(FVectorVMExternalFunctionContext& Context)
-{
-	VectorVM::FUserPtrHandler<FGrid3DCollectionRWInstanceData_GameThread> InstData(Context);
-
-	FNDIOutputParam<int32> NumCellsX(Context);
-	FNDIOutputParam<int32> NumCellsY(Context);
-	FNDIOutputParam<int32> NumCellsZ(Context);
 
-	int32 TmpNumCellsX = InstData->NumCells.X;
-	int32 TmpNumCellsY = InstData->NumCells.Y;
-	int32 TmpNumCellsZ = InstData->NumCells.Z;
-	if (InstData->OtherInstanceData != nullptr)
 	{
-		TmpNumCellsX = InstData->OtherInstanceData->NumCells.X;
-		TmpNumCellsY = InstData->OtherInstanceData->NumCells.Y;
-		TmpNumCellsZ = InstData->OtherInstanceData->NumCells.Z;
-	}
+		FNiagaraFunctionSignature Sig;
+		Sig.Name = GetPreviousVector2ValueAtIndexFunctionName;
+		Sig.Inputs.Add(FNiagaraVariable(FNiagaraTypeDefinition(GetClass()), TEXT("Grid")));
+		Sig.Inputs.Add(FNiagaraVariable(FNiagaraTypeDefinition::GetIntDef(), TEXT("IndexX")));
+		Sig.Inputs.Add(FNiagaraVariable(FNiagaraTypeDefinition::GetIntDef(), TEXT("IndexY")));
+		Sig.Inputs.Add(FNiagaraVariable(FNiagaraTypeDefinition::GetIntDef(), TEXT("IndexZ")));
+		Sig.Inputs.Add(FNiagaraVariable(FNiagaraTypeDefinition::GetIntDef(), TEXT("AttributeIndex")));
+		Sig.Outputs.Add(FNiagaraVariable(FNiagaraTypeDefinition::GetVec2Def(), TEXT("Value")));
 
-	for (int32 InstanceIdx = 0; InstanceIdx < Context.GetNumInstances(); ++InstanceIdx)
-	{
-		NumCellsX.SetAndAdvance(TmpNumCellsX);
-		NumCellsY.SetAndAdvance(TmpNumCellsY);
-		NumCellsZ.SetAndAdvance(TmpNumCellsZ);
+		Sig.bExperimental = true;
+		Sig.bMemberFunction = true;
+		Sig.bRequiresContext = false;
+		Sig.ModuleUsageBitmask = ENiagaraScriptUsageMask::Particle;
+		Sig.bSupportsCPU = false;
+		Sig.bSupportsGPU = true;
+		Sig.bReadFunction = true;
+#if WITH_EDITORONLY_DATA
+		Sig.Description = NSLOCTEXT("Niagara", "NiagaraDataInterfaceGridColl2D_PreviousValueAtIndexFunction", "Get the value at a specific index.");
+#endif
+
+		OutFunctions.Add(Sig);
 	}
-}
 
-bool UNiagaraDataInterfaceGrid3DCollection::Equals(const UNiagaraDataInterface* Other) const
-{
-	if (!Super::Equals(Other))
 	{
-		return false;
-	}
-	const UNiagaraDataInterfaceGrid3DCollection* OtherTyped = CastChecked<const UNiagaraDataInterfaceGrid3DCollection>(Other);
+		FNiagaraFunctionSignature Sig;
+		Sig.Name = GetPreviousVector3ValueAtIndexFunctionName;
+		Sig.Inputs.Add(FNiagaraVariable(FNiagaraTypeDefinition(GetClass()), TEXT("Grid")));
+		Sig.Inputs.Add(FNiagaraVariable(FNiagaraTypeDefinition::GetIntDef(), TEXT("IndexX")));
+		Sig.Inputs.Add(FNiagaraVariable(FNiagaraTypeDefinition::GetIntDef(), TEXT("IndexY")));
+		Sig.Inputs.Add(FNiagaraVariable(FNiagaraTypeDefinition::GetIntDef(), TEXT("IndexZ")));
+		Sig.Inputs.Add(FNiagaraVariable(FNiagaraTypeDefinition::GetIntDef(), TEXT("AttributeIndex")));
+		Sig.Outputs.Add(FNiagaraVariable(FNiagaraTypeDefinition::GetVec3Def(), TEXT("Value")));
 
-	return OtherTyped != nullptr &&
-		OtherTyped->NumAttributes == NumAttributes &&
-		OtherTyped->RenderTargetUserParameter == RenderTargetUserParameter &&
-		OtherTyped->OverrideBufferFormat == OverrideBufferFormat &&
+		Sig.bExperimental = true;
+		Sig.bMemberFunction = true;
+		Sig.bRequiresContext = false;
+		Sig.ModuleUsageBitmask = ENiagaraScriptUsageMask::Particle;
+		Sig.bSupportsCPU = false;
+		Sig.bSupportsGPU = true;
+		Sig.bReadFunction = true;
 #if WITH_EDITORONLY_DATA
-		OtherTyped->bPreviewGrid == bPreviewGrid &&
-		OtherTyped->PreviewAttribute == PreviewAttribute &&
+		Sig.Description = NSLOCTEXT("Niagara", "NiagaraDataInterfaceGridColl2D_PreviousValueAtIndexFunction", "Get the value at a specific index.");
 #endif
-		OtherTyped->bOverrideFormat == bOverrideFormat;
-}
 
-bool UNiagaraDataInterfaceGrid3DCollection::CopyToInternal(UNiagaraDataInterface* Destination) const
-{
-	if (!Super::CopyToInternal(Destination))
-	{
-		return false;
+		OutFunctions.Add(Sig);
 	}
 
-	UNiagaraDataInterfaceGrid3DCollection* OtherTyped = CastChecked<UNiagaraDataInterfaceGrid3DCollection>(Destination);
-	OtherTyped->NumAttributes = NumAttributes;
-	OtherTyped->RenderTargetUserParameter = RenderTargetUserParameter;
-	OtherTyped->OverrideBufferFormat = OverrideBufferFormat;
-	OtherTyped->bOverrideFormat = bOverrideFormat;
+	{
+		FNiagaraFunctionSignature Sig;
+		Sig.Name = GetPreviousVector4ValueAtIndexFunctionName;
+		Sig.Inputs.Add(FNiagaraVariable(FNiagaraTypeDefinition(GetClass()), TEXT("Grid")));
+		Sig.Inputs.Add(FNiagaraVariable(FNiagaraTypeDefinition::GetIntDef(), TEXT("IndexX")));
+		Sig.Inputs.Add(FNiagaraVariable(FNiagaraTypeDefinition::GetIntDef(), TEXT("IndexY")));
+		Sig.Inputs.Add(FNiagaraVariable(FNiagaraTypeDefinition::GetIntDef(), TEXT("IndexZ")));
+		Sig.Inputs.Add(FNiagaraVariable(FNiagaraTypeDefinition::GetIntDef(), TEXT("AttributeIndex")));
+		Sig.Outputs.Add(FNiagaraVariable(FNiagaraTypeDefinition::GetVec4Def(), TEXT("Value")));
+
+		Sig.bExperimental = true;
+		Sig.bMemberFunction = true;
+		Sig.bRequiresContext = false;
+		Sig.ModuleUsageBitmask = ENiagaraScriptUsageMask::Particle;
+		Sig.bSupportsCPU = false;
+		Sig.bSupportsGPU = true;
+		Sig.bReadFunction = true;
 #if WITH_EDITORONLY_DATA
-	OtherTyped->bPreviewGrid = bPreviewGrid;
-	OtherTyped->PreviewAttribute = PreviewAttribute;
+		Sig.Description = NSLOCTEXT("Niagara", "NiagaraDataInterfaceGridColl2D_PreviousValueAtIndexFunction", "Get the value at a specific index.");
 #endif
 
-	return true;
-}
+		OutFunctions.Add(Sig);
+	}
 
-#if WITH_EDITORONLY_DATA
-void UNiagaraDataInterfaceGrid3DCollection::GetParameterDefinitionHLSL(const FNiagaraDataInterfaceGPUParamInfo& ParamInfo, FString& OutHLSL)
-{
-	Super::GetParameterDefinitionHLSL(ParamInfo, OutHLSL);
+	{
+		FNiagaraFunctionSignature Sig;
+		Sig.Name = SamplePreviousGridFloatAtIndexFunctionName;
+		Sig.Inputs.Add(FNiagaraVariable(FNiagaraTypeDefinition(GetClass()), TEXT("Grid")));
+		Sig.Inputs.Add(FNiagaraVariable(FNiagaraTypeDefinition::GetVec3Def(), TEXT("Unit")));
+		Sig.Inputs.Add(FNiagaraVariable(FNiagaraTypeDefinition::GetIntDef(), TEXT("AttributeIndex")));
+		Sig.Outputs.Add(FNiagaraVariable(FNiagaraTypeDefinition::GetFloatDef(), TEXT("Value")));
 
-	static const TCHAR *FormatDeclarations = TEXT(R"(
-		Texture3D<float> {GridName};
-		RWTexture3D<float> {OutputGridName};
-		int3 {NumTiles};
-		float3 {OneOverNumTiles};
-		float3 {UnitClampMin};
-		float3 {UnitClampMax};
-		SamplerState {SamplerName};
-		int4 {AttributeIndicesName}[{AttributeInt4Count}];
-		Buffer<float4> {PerAttributeDataName};
-		int {NumAttributesName};
-		int {NumNamedAttributesName};
+		Sig.bExperimental = true;
+		Sig.bMemberFunction = true;
+		Sig.bRequiresContext = false;
+		Sig.ModuleUsageBitmask = ENiagaraScriptUsageMask::Particle;
+		Sig.bSupportsCPU = false;
+		Sig.bSupportsGPU = true;
+		Sig.bReadFunction = true;
+		OutFunctions.Add(Sig);
+	}
+
+	{
+		FNiagaraFunctionSignature Sig;
+		Sig.Name = SamplePreviousGridVector2AtIndexFunctionName;
+		Sig.Inputs.Add(FNiagaraVariable(FNiagaraTypeDefinition(GetClass()), TEXT("Grid")));
+		Sig.Inputs.Add(FNiagaraVariable(FNiagaraTypeDefinition::GetVec3Def(), TEXT("Unit")));
+		Sig.Inputs.Add(FNiagaraVariable(FNiagaraTypeDefinition::GetIntDef(), TEXT("AttributeIndex")));
+		Sig.Outputs.Add(FNiagaraVariable(FNiagaraTypeDefinition::GetVec2Def(), TEXT("Value")));
+
+		Sig.bExperimental = true;
+		Sig.bMemberFunction = true;
+		Sig.bRequiresContext = false;
+		Sig.ModuleUsageBitmask = ENiagaraScriptUsageMask::Particle;
+		Sig.bSupportsCPU = false;
+		Sig.bSupportsGPU = true;
+		Sig.bReadFunction = true;
+		OutFunctions.Add(Sig);
+	}
+
+	{
+		FNiagaraFunctionSignature Sig;
+		Sig.Name = SamplePreviousGridVector3AtIndexFunctionName;
+		Sig.Inputs.Add(FNiagaraVariable(FNiagaraTypeDefinition(GetClass()), TEXT("Grid")));
+		Sig.Inputs.Add(FNiagaraVariable(FNiagaraTypeDefinition::GetVec3Def(), TEXT("Unit")));
+		Sig.Inputs.Add(FNiagaraVariable(FNiagaraTypeDefinition::GetIntDef(), TEXT("AttributeIndex")));
+		Sig.Outputs.Add(FNiagaraVariable(FNiagaraTypeDefinition::GetVec3Def(), TEXT("Value")));
+
+		Sig.bExperimental = true;
+		Sig.bMemberFunction = true;
+		Sig.bRequiresContext = false;
+		Sig.ModuleUsageBitmask = ENiagaraScriptUsageMask::Particle;
+		Sig.bSupportsCPU = false;
+		Sig.bSupportsGPU = true;
+		Sig.bReadFunction = true;
+		OutFunctions.Add(Sig);
+	}
+
+	{
+		FNiagaraFunctionSignature Sig;
+		Sig.Name = SamplePreviousGridVector4AtIndexFunctionName;
+		Sig.Inputs.Add(FNiagaraVariable(FNiagaraTypeDefinition(GetClass()), TEXT("Grid")));
+		Sig.Inputs.Add(FNiagaraVariable(FNiagaraTypeDefinition::GetVec3Def(), TEXT("Unit")));
+		Sig.Inputs.Add(FNiagaraVariable(FNiagaraTypeDefinition::GetIntDef(), TEXT("AttributeIndex")));
+		Sig.Outputs.Add(FNiagaraVariable(FNiagaraTypeDefinition::GetVec4Def(), TEXT("Value")));
+
+		Sig.bExperimental = true;
+		Sig.bMemberFunction = true;
+		Sig.bRequiresContext = false;
+		Sig.ModuleUsageBitmask = ENiagaraScriptUsageMask::Particle;
+		Sig.bSupportsCPU = false;
+		Sig.bSupportsGPU = true;
+		Sig.bReadFunction = true;
+		OutFunctions.Add(Sig);
+	}
+
+	{
+		FNiagaraFunctionSignature Sig;
+		Sig.Name = CubicSamplePreviousGridFloatAtIndexFunctionName;
+		Sig.Inputs.Add(FNiagaraVariable(FNiagaraTypeDefinition(GetClass()), TEXT("Grid")));
+		Sig.Inputs.Add(FNiagaraVariable(FNiagaraTypeDefinition::GetVec3Def(), TEXT("Unit")));
+		Sig.Inputs.Add(FNiagaraVariable(FNiagaraTypeDefinition::GetIntDef(), TEXT("AttributeIndex")));
+		Sig.Outputs.Add(FNiagaraVariable(FNiagaraTypeDefinition::GetFloatDef(), TEXT("Value")));
+
+		Sig.bExperimental = true;
+		Sig.bMemberFunction = true;
+		Sig.bRequiresContext = false;
+		Sig.ModuleUsageBitmask = ENiagaraScriptUsageMask::Particle;
+		Sig.bSupportsCPU = false;
+		Sig.bSupportsGPU = true;
+		Sig.bReadFunction = true;
+		OutFunctions.Add(Sig);
+	}
+
+	{
+		FNiagaraFunctionSignature Sig;
+		Sig.Name = CubicSamplePreviousGridVector2AtIndexFunctionName;
+		Sig.Inputs.Add(FNiagaraVariable(FNiagaraTypeDefinition(GetClass()), TEXT("Grid")));
+		Sig.Inputs.Add(FNiagaraVariable(FNiagaraTypeDefinition::GetVec3Def(), TEXT("Unit")));
+		Sig.Inputs.Add(FNiagaraVariable(FNiagaraTypeDefinition::GetIntDef(), TEXT("AttributeIndex")));
+		Sig.Outputs.Add(FNiagaraVariable(FNiagaraTypeDefinition::GetVec2Def(), TEXT("Value")));
+
+		Sig.bExperimental = true;
+		Sig.bMemberFunction = true;
+		Sig.bRequiresContext = false;
+		Sig.ModuleUsageBitmask = ENiagaraScriptUsageMask::Particle;
+		Sig.bSupportsCPU = false;
+		Sig.bSupportsGPU = true;
+		Sig.bReadFunction = true;
+		OutFunctions.Add(Sig);
+	}
+
+	{
+		FNiagaraFunctionSignature Sig;
+		Sig.Name = CubicSamplePreviousGridVector3AtIndexFunctionName;
+		Sig.Inputs.Add(FNiagaraVariable(FNiagaraTypeDefinition(GetClass()), TEXT("Grid")));
+		Sig.Inputs.Add(FNiagaraVariable(FNiagaraTypeDefinition::GetVec3Def(), TEXT("Unit")));
+		Sig.Inputs.Add(FNiagaraVariable(FNiagaraTypeDefinition::GetIntDef(), TEXT("AttributeIndex")));
+		Sig.Outputs.Add(FNiagaraVariable(FNiagaraTypeDefinition::GetVec3Def(), TEXT("Value")));
+
+		Sig.bExperimental = true;
+		Sig.bMemberFunction = true;
+		Sig.bRequiresContext = false;
+		Sig.ModuleUsageBitmask = ENiagaraScriptUsageMask::Particle;
+		Sig.bSupportsCPU = false;
+		Sig.bSupportsGPU = true;
+		Sig.bReadFunction = true;
+		OutFunctions.Add(Sig);
+	}
+
+	{
+		FNiagaraFunctionSignature Sig;
+		Sig.Name = CubicSamplePreviousGridVector4AtIndexFunctionName;
+		Sig.Inputs.Add(FNiagaraVariable(FNiagaraTypeDefinition(GetClass()), TEXT("Grid")));
+		Sig.Inputs.Add(FNiagaraVariable(FNiagaraTypeDefinition::GetVec3Def(), TEXT("Unit")));
+		Sig.Inputs.Add(FNiagaraVariable(FNiagaraTypeDefinition::GetIntDef(), TEXT("AttributeIndex")));
+		Sig.Outputs.Add(FNiagaraVariable(FNiagaraTypeDefinition::GetVec4Def(), TEXT("Value")));
+
+		Sig.bExperimental = true;
+		Sig.bMemberFunction = true;
+		Sig.bRequiresContext = false;
+		Sig.ModuleUsageBitmask = ENiagaraScriptUsageMask::Particle;
+		Sig.bSupportsCPU = false;
+		Sig.bSupportsGPU = true;
+		Sig.bReadFunction = true;
+		OutFunctions.Add(Sig);
+	}
+
+	{
+		FNiagaraFunctionSignature Sig;
+		Sig.Name = SetFloatValueAtIndexFunctionName;
+		Sig.Inputs.Add(FNiagaraVariable(FNiagaraTypeDefinition(GetClass()), TEXT("Grid")));
+		Sig.Inputs.Add(FNiagaraVariable(FNiagaraTypeDefinition::GetIntDef(), TEXT("IndexX")));
+		Sig.Inputs.Add(FNiagaraVariable(FNiagaraTypeDefinition::GetIntDef(), TEXT("IndexY")));
+		Sig.Inputs.Add(FNiagaraVariable(FNiagaraTypeDefinition::GetIntDef(), TEXT("IndexZ")));
+		Sig.Inputs.Add(FNiagaraVariable(FNiagaraTypeDefinition::GetIntDef(), TEXT("Index")));
+		Sig.Inputs.Add(FNiagaraVariable(FNiagaraTypeDefinition::GetFloatDef(), TEXT("Value")));
+		Sig.bMemberFunction = true;
+		Sig.bRequiresContext = false;
+		Sig.bExperimental = true;
+		Sig.bRequiresExecPin = true;
+		Sig.bWriteFunction = true;
+		Sig.ModuleUsageBitmask = ENiagaraScriptUsageMask::Particle;
+		Sig.bSupportsCPU = false;
+		Sig.bSupportsGPU = true;
+#if WITH_EDITORONLY_DATA
+		Sig.Description = NSLOCTEXT("Niagara", "NiagaraDataInterfaceGridColl3D_SetFloat", "Sets a float value on the Grid by Attribute name.");
+#endif
+		OutFunctions.Add(Sig);
+	}
+
+	{
+		FNiagaraFunctionSignature Sig;
+		Sig.Name = SetVector2ValueAtIndexFunctionName;
+		Sig.Inputs.Add(FNiagaraVariable(FNiagaraTypeDefinition(GetClass()), TEXT("Grid")));
+		Sig.Inputs.Add(FNiagaraVariable(FNiagaraTypeDefinition::GetIntDef(), TEXT("IndexX")));
+		Sig.Inputs.Add(FNiagaraVariable(FNiagaraTypeDefinition::GetIntDef(), TEXT("IndexY")));
+		Sig.Inputs.Add(FNiagaraVariable(FNiagaraTypeDefinition::GetIntDef(), TEXT("IndexZ")));
+		Sig.Inputs.Add(FNiagaraVariable(FNiagaraTypeDefinition::GetIntDef(), TEXT("Index")));
+		Sig.Inputs.Add(FNiagaraVariable(FNiagaraTypeDefinition::GetVec2Def(), TEXT("Value")));
+		Sig.bMemberFunction = true;
+		Sig.bRequiresContext = false;
+		Sig.bExperimental = true;
+		Sig.bRequiresExecPin = true;
+		Sig.bWriteFunction = true;
+		Sig.ModuleUsageBitmask = ENiagaraScriptUsageMask::Particle;
+		Sig.bSupportsCPU = false;
+		Sig.bSupportsGPU = true;
+#if WITH_EDITORONLY_DATA
+		Sig.Description = NSLOCTEXT("Niagara", "NiagaraDataInterfaceGridColl3D_SetFloat", "Sets a float value on the Grid by Attribute name.");
+#endif
+		OutFunctions.Add(Sig);
+	}
+
+	{
+		FNiagaraFunctionSignature Sig;
+		Sig.Name = SetVector3ValueAtIndexFunctionName;
+		Sig.Inputs.Add(FNiagaraVariable(FNiagaraTypeDefinition(GetClass()), TEXT("Grid")));
+		Sig.Inputs.Add(FNiagaraVariable(FNiagaraTypeDefinition::GetIntDef(), TEXT("IndexX")));
+		Sig.Inputs.Add(FNiagaraVariable(FNiagaraTypeDefinition::GetIntDef(), TEXT("IndexY")));
+		Sig.Inputs.Add(FNiagaraVariable(FNiagaraTypeDefinition::GetIntDef(), TEXT("IndexZ")));
+		Sig.Inputs.Add(FNiagaraVariable(FNiagaraTypeDefinition::GetIntDef(), TEXT("Index")));
+		Sig.Inputs.Add(FNiagaraVariable(FNiagaraTypeDefinition::GetVec3Def(), TEXT("Value")));
+		Sig.bMemberFunction = true;
+		Sig.bRequiresContext = false;
+		Sig.bExperimental = true;
+		Sig.bRequiresExecPin = true;
+		Sig.bWriteFunction = true;
+		Sig.ModuleUsageBitmask = ENiagaraScriptUsageMask::Particle;
+		Sig.bSupportsCPU = false;
+		Sig.bSupportsGPU = true;
+#if WITH_EDITORONLY_DATA
+		Sig.Description = NSLOCTEXT("Niagara", "NiagaraDataInterfaceGridColl3D_SetFloat", "Sets a float value on the Grid by Attribute name.");
+#endif
+		OutFunctions.Add(Sig);
+	}
+
+	{
+		FNiagaraFunctionSignature Sig;
+		Sig.Name = SetVector4ValueAtIndexFunctionName;
+		Sig.Inputs.Add(FNiagaraVariable(FNiagaraTypeDefinition(GetClass()), TEXT("Grid")));
+		Sig.Inputs.Add(FNiagaraVariable(FNiagaraTypeDefinition::GetIntDef(), TEXT("IndexX")));
+		Sig.Inputs.Add(FNiagaraVariable(FNiagaraTypeDefinition::GetIntDef(), TEXT("IndexY")));
+		Sig.Inputs.Add(FNiagaraVariable(FNiagaraTypeDefinition::GetIntDef(), TEXT("IndexZ")));
+		Sig.Inputs.Add(FNiagaraVariable(FNiagaraTypeDefinition::GetIntDef(), TEXT("Index")));
+		Sig.Inputs.Add(FNiagaraVariable(FNiagaraTypeDefinition::GetVec4Def(), TEXT("Value")));
+		Sig.bMemberFunction = true;
+		Sig.bRequiresContext = false;
+		Sig.bExperimental = true;
+		Sig.bRequiresExecPin = true;
+		Sig.bWriteFunction = true;
+		Sig.ModuleUsageBitmask = ENiagaraScriptUsageMask::Particle;
+		Sig.bSupportsCPU = false;
+		Sig.bSupportsGPU = true;
+#if WITH_EDITORONLY_DATA
+		Sig.Description = NSLOCTEXT("Niagara", "NiagaraDataInterfaceGridColl3D_SetFloat", "Sets a float value on the Grid by Attribute name.");
+#endif
+		OutFunctions.Add(Sig);
+	}
+
+#if WITH_EDITORONLY_DATA
+	for (int32 i = StartIndex; i < OutFunctions.Num(); i++)
+	{
+		FNiagaraFunctionSignature& Function = OutFunctions[i];
+		Function.FunctionVersion = FNiagaraGridCollection3DDIFunctionVersion::LatestVersion;
+	}
+#endif
+}
+
+// #todo(dmp): expose more CPU functionality
+// #todo(dmp): ideally these would be exposed on the parent class, but we can't bind functions of parent classes but need to work on the interface
+// for sharing an instance data object with the super class
+DEFINE_NDI_DIRECT_FUNC_BINDER(UNiagaraDataInterfaceGrid3DCollection, VMGetWorldBBoxSize);
+DEFINE_NDI_DIRECT_FUNC_BINDER(UNiagaraDataInterfaceGrid3DCollection, VMGetCellSize);
+DEFINE_NDI_DIRECT_FUNC_BINDER(UNiagaraDataInterfaceGrid3DCollection, VMSetNumCells);
+DEFINE_NDI_DIRECT_FUNC_BINDER(UNiagaraDataInterfaceGrid3DCollection, VMGetNumCells);
+DEFINE_NDI_DIRECT_FUNC_BINDER(UNiagaraDataInterfaceGrid3DCollection, VMUnitToFloatIndex);
+DEFINE_NDI_DIRECT_FUNC_BINDER_WITH_PAYLOAD(UNiagaraDataInterfaceGrid3DCollection, VMGetAttributeIndex);
+void UNiagaraDataInterfaceGrid3DCollection::GetVMExternalFunction(const FVMExternalFunctionBindingInfo& BindingInfo, void* InstanceData, FVMExternalFunction& OutFunc)
+{
+	Super::GetVMExternalFunction(BindingInfo, InstanceData, OutFunc);
+
+	if (BindingInfo.Name == UNiagaraDataInterfaceRWBase::WorldBBoxSizeFunctionName)
+	{
+		check(BindingInfo.GetNumInputs() == 1 && BindingInfo.GetNumOutputs() == 3);
+		NDI_FUNC_BINDER(UNiagaraDataInterfaceGrid3DCollection, VMGetWorldBBoxSize)::Bind(this, OutFunc);
+	}
+	else if (BindingInfo.Name == UNiagaraDataInterfaceRWBase::CellSizeFunctionName)
+	{
+		// #todo(dmp): this will override the base class definition for GetCellSize because the data interface instance data computes cell size
+		// it would be nice to refactor this so it can be part of the super class
+		check(BindingInfo.GetNumInputs() == 1 && BindingInfo.GetNumOutputs() == 3);
+		NDI_FUNC_BINDER(UNiagaraDataInterfaceGrid3DCollection, VMGetCellSize)::Bind(this, OutFunc);
+	}
+	else if (BindingInfo.Name == SetNumCellsFunctionName)
+	{
+		check(BindingInfo.GetNumInputs() == 4 && BindingInfo.GetNumOutputs() == 1);
+		NDI_FUNC_BINDER(UNiagaraDataInterfaceGrid3DCollection, VMSetNumCells)::Bind(this, OutFunc);
+	}
+	else if (BindingInfo.Name == NumCellsFunctionName)
+	{
+		check(BindingInfo.GetNumInputs() == 1 && BindingInfo.GetNumOutputs() == 3);
+		NDI_FUNC_BINDER(UNiagaraDataInterfaceGrid3DCollection, VMGetNumCells)::Bind(this, OutFunc);
+	}
+	else if (BindingInfo.Name == UnitToFloatIndexFunctionName)
+	{
+		check(BindingInfo.GetNumInputs() == 4 && BindingInfo.GetNumOutputs() == 3);
+		NDI_FUNC_BINDER(UNiagaraDataInterfaceGrid3DCollection, VMUnitToFloatIndex)::Bind(this, OutFunc);
+	}
+	else if (BindingInfo.Name == GetVector4AttributeIndexFunctionName)
+	{
+		FName AttributeName = BindingInfo.FindSpecifier(NAME_Attribute)->Value;
+		check(BindingInfo.GetNumInputs() == 1 && BindingInfo.GetNumOutputs() == 1);
+		NDI_FUNC_BINDER(UNiagaraDataInterfaceGrid3DCollection, VMGetAttributeIndex)::Bind(this, OutFunc, AttributeName, 4);
+	}
+	else if (BindingInfo.Name == GetVectorAttributeIndexFunctionName)
+	{
+		FName AttributeName = BindingInfo.FindSpecifier(NAME_Attribute)->Value;
+		check(BindingInfo.GetNumInputs() == 1 && BindingInfo.GetNumOutputs() == 1);
+		NDI_FUNC_BINDER(UNiagaraDataInterfaceGrid3DCollection, VMGetAttributeIndex)::Bind(this, OutFunc, AttributeName, 3);
+	}
+	else if (BindingInfo.Name == GetVector2DAttributeIndexFunctionName)
+	{
+		FName AttributeName = BindingInfo.FindSpecifier(NAME_Attribute)->Value;
+		check(BindingInfo.GetNumInputs() == 1 && BindingInfo.GetNumOutputs() == 1);
+		NDI_FUNC_BINDER(UNiagaraDataInterfaceGrid3DCollection, VMGetAttributeIndex)::Bind(this, OutFunc, AttributeName, 2);
+	}
+	else if (BindingInfo.Name == GetFloatAttributeIndexFunctionName)
+	{
+		FName AttributeName = BindingInfo.FindSpecifier(NAME_Attribute)->Value;
+		check(BindingInfo.GetNumInputs() == 1 && BindingInfo.GetNumOutputs() == 1);
+		NDI_FUNC_BINDER(UNiagaraDataInterfaceGrid3DCollection, VMGetAttributeIndex)::Bind(this, OutFunc, AttributeName, 1);
+	}
+}
+
+void UNiagaraDataInterfaceGrid3DCollection::VMGetNumCells(FVectorVMExternalFunctionContext& Context)
+{
+	VectorVM::FUserPtrHandler<FGrid3DCollectionRWInstanceData_GameThread> InstData(Context);
+
+	FNDIOutputParam<int32> NumCellsX(Context);
+	FNDIOutputParam<int32> NumCellsY(Context);
+	FNDIOutputParam<int32> NumCellsZ(Context);
+
+	int32 TmpNumCellsX = InstData->NumCells.X;
+	int32 TmpNumCellsY = InstData->NumCells.Y;
+	int32 TmpNumCellsZ = InstData->NumCells.Z;
+	if (InstData->OtherInstanceData != nullptr)
+	{
+		TmpNumCellsX = InstData->OtherInstanceData->NumCells.X;
+		TmpNumCellsY = InstData->OtherInstanceData->NumCells.Y;
+		TmpNumCellsZ = InstData->OtherInstanceData->NumCells.Z;
+	}
+
+	for (int32 InstanceIdx = 0; InstanceIdx < Context.GetNumInstances(); ++InstanceIdx)
+	{
+		NumCellsX.SetAndAdvance(TmpNumCellsX);
+		NumCellsY.SetAndAdvance(TmpNumCellsY);
+		NumCellsZ.SetAndAdvance(TmpNumCellsZ);
+	}
+}
+
+bool UNiagaraDataInterfaceGrid3DCollection::Equals(const UNiagaraDataInterface* Other) const
+{
+	if (!Super::Equals(Other))
+	{
+		return false;
+	}
+	const UNiagaraDataInterfaceGrid3DCollection* OtherTyped = CastChecked<const UNiagaraDataInterfaceGrid3DCollection>(Other);
+
+	return OtherTyped != nullptr &&
+		OtherTyped->NumAttributes == NumAttributes &&
+		OtherTyped->RenderTargetUserParameter == RenderTargetUserParameter &&
+		OtherTyped->OverrideBufferFormat == OverrideBufferFormat &&
+#if WITH_EDITORONLY_DATA
+		OtherTyped->bPreviewGrid == bPreviewGrid &&
+		OtherTyped->PreviewAttribute == PreviewAttribute &&
+#endif
+		OtherTyped->bOverrideFormat == bOverrideFormat;
+}
+
+bool UNiagaraDataInterfaceGrid3DCollection::CopyToInternal(UNiagaraDataInterface* Destination) const
+{
+	if (!Super::CopyToInternal(Destination))
+	{
+		return false;
+	}
+
+	UNiagaraDataInterfaceGrid3DCollection* OtherTyped = CastChecked<UNiagaraDataInterfaceGrid3DCollection>(Destination);
+	OtherTyped->NumAttributes = NumAttributes;
+	OtherTyped->RenderTargetUserParameter = RenderTargetUserParameter;
+	OtherTyped->OverrideBufferFormat = OverrideBufferFormat;
+	OtherTyped->bOverrideFormat = bOverrideFormat;
+#if WITH_EDITORONLY_DATA
+	OtherTyped->bPreviewGrid = bPreviewGrid;
+	OtherTyped->PreviewAttribute = PreviewAttribute;
+#endif
+
+	return true;
+}
+
+#if WITH_EDITORONLY_DATA
+void UNiagaraDataInterfaceGrid3DCollection::GetParameterDefinitionHLSL(const FNiagaraDataInterfaceGPUParamInfo& ParamInfo, FString& OutHLSL)
+{
+	Super::GetParameterDefinitionHLSL(ParamInfo, OutHLSL);
+
+	static const TCHAR* FormatDeclarations = TEXT(R"(
+		Texture3D<float{FloatBufferSuffix}> {GridName};
+		RWTexture3D<float{FloatBufferSuffix}> {OutputGridName};
+		int3 {NumTiles};
+		float3 {OneOverNumTiles};
+		float3 {UnitClampMin};
+		float3 {UnitClampMax};
+		SamplerState {SamplerName};
+		int4 {AttributeIndicesName}[{AttributeInt4Count}];
+		Buffer<float4> {PerAttributeDataName};
+		int {NumAttributesName};
+		int {NumNamedAttributesName};
 	)");
 
 	// If we use an int array for the attribute indices, the shader compiler will actually use int4 due to the packing rules,
@@ -1169,11 +1637,21 @@ void UNiagaraDataInterfaceGrid3DCollection::GetParameterDefinitionHLSL(const FNi
 	const int AttributeCount = ParamInfo.GeneratedFunctions.Num();
 	const int AttributeInt4Count = FMath::Max(1, FMath::DivideAndRoundUp(AttributeCount, 4));
 
-	TMap<FString, FStringFormatArg> ArgsDeclarations = {	
+	// format of the grid is defined by number of attributes
+	FString FloatBufferSuffix = "";
+	FGrid3DCollectionAttributeHelper Helper(ParamInfo);
+	if (Helper.UseRGBAGrid() && Helper.TotalChannels > 1)
+	{
+		FloatBufferSuffix = FString::FromInt(Helper.TotalChannels);
+	}
+
+	TMap<FString, FStringFormatArg> ArgsDeclarations = {
 		{ TEXT("ParameterName"),	ParamInfo.DataInterfaceHLSLSymbol},
 		{ TEXT("GridName"),			ParamInfo.DataInterfaceHLSLSymbol + GridName},
 		{ TEXT("SamplerName"),		ParamInfo.DataInterfaceHLSLSymbol + SamplerName},
 		{ TEXT("OutputGridName"),	ParamInfo.DataInterfaceHLSLSymbol + OutputGridName},
+		{ TEXT("FloatBufferSuffix"),FloatBufferSuffix},
+		{ TEXT("NumChannels"),		Helper.UseRGBAGrid() ? Helper.TotalChannels : 1}, // float grids only have a single channel for the underlying buffer
 		{ TEXT("NumTiles"),			ParamInfo.DataInterfaceHLSLSymbol + NumTilesName},
 		{ TEXT("OneOverNumTiles"),	ParamInfo.DataInterfaceHLSLSymbol + OneOverNumTilesName},
 		{ TEXT("UnitClampMin"),		ParamInfo.DataInterfaceHLSLSymbol + UnitClampMinName},
@@ -1244,6 +1722,94 @@ void UNiagaraDataInterfaceGrid3DCollection::WriteSetHLSL(const FNiagaraDataInter
 	OutHLSL += FString::Format(*FormatBounds, ArgsBounds);
 }
 
+// code generation for reading grid attributes at a specific index - only considering tiled grid attributes
+void UNiagaraDataInterfaceGrid3DCollection::WriteGetAtIndexHLSL(const FNiagaraDataInterfaceGPUParamInfo& ParamInfo, const FNiagaraDataInterfaceGeneratedFunction& FunctionInfo, int32 InNumChannels, FString& OutHLSL)
+{
+
+	FString FormatBounds = TEXT(R"(
+			void {FunctionName}(int In_IndexX, int In_IndexY, int In_IndexZ, int AttributeIndex, out float{NumChannelsVariableSuffix} Out_Val)
+			{				
+			    for (int i = 0; i < {NumChannels}; i++)
+				{
+					int CurAttributeIndex = AttributeIndex + i;
+					int3 TileOffset = {PerAttributeDataName}[CurAttributeIndex].xyz;
+					float Val = {Grid}.Load(int4(In_IndexX + TileOffset.x, In_IndexY + TileOffset.y, In_IndexZ + TileOffset.z, 0));
+					)");
+	if (InNumChannels == 1)
+	{
+		FormatBounds += TEXT("					Out_Val = Val;\n");
+	}
+	else
+	{
+		FormatBounds += TEXT("					Out_Val[i] = Val;\n");
+	}
+	FormatBounds += TEXT(R"(
+				}
+			}
+		)");
+	TMap<FString, FStringFormatArg> ArgsBounds = {
+		{TEXT("ParameterName"),				ParamInfo.DataInterfaceHLSLSymbol},
+		{TEXT("FunctionName"),				FunctionInfo.InstanceName},
+		{TEXT("OutputGrid"),				ParamInfo.DataInterfaceHLSLSymbol + OutputGridName},
+		{TEXT("Grid"),						ParamInfo.DataInterfaceHLSLSymbol + GridName},
+		{TEXT("NumCellsName"),				ParamInfo.DataInterfaceHLSLSymbol + UNiagaraDataInterfaceRWBase::NumCellsName},
+		{TEXT("UnitToUVName"),				ParamInfo.DataInterfaceHLSLSymbol + UNiagaraDataInterfaceRWBase::UnitToUVName},
+		{TEXT("PerAttributeDataName"),		ParamInfo.DataInterfaceHLSLSymbol + TEXT("_") + PerAttributeDataName},
+		{TEXT("NumChannels"),				FString::FromInt(InNumChannels)},
+		{TEXT("NumChannelsVariableSuffix"),	InNumChannels > 1 ? FString::FromInt(InNumChannels) : TEXT("")},
+		{TEXT("NumTiles"),					ParamInfo.DataInterfaceHLSLSymbol + NumTilesName},
+		{TEXT("NumAttributesName"),			ParamInfo.DataInterfaceHLSLSymbol + UNiagaraDataInterfaceRWBase::NumAttributesName},
+	};
+	OutHLSL += FString::Format(*FormatBounds, ArgsBounds);
+}
+
+// code generation for writing grid attributes at a specific index - only considering tiled grid attributes
+void UNiagaraDataInterfaceGrid3DCollection::WriteSetAtIndexHLSL(const FNiagaraDataInterfaceGPUParamInfo& ParamInfo, const FNiagaraDataInterfaceGeneratedFunction& FunctionInfo, int32 InNumChannels, FString& OutHLSL)
+{
+
+	FString FormatBounds = TEXT(R"(
+			void {FunctionName}(int In_IndexX, int In_IndexY, int In_IndexZ, int AttributeIndex, float{NumChannelsVariableSuffix} In_Value)
+			{				
+			    for (int i = 0; i < {NumChannels}; i++)
+				{
+				)");
+	if (InNumChannels == 1)
+	{
+		FormatBounds += TEXT(R"(
+					float Val = In_Value;
+		)");
+	}
+	else
+	{
+		FormatBounds += TEXT(R"(
+					float Val = In_Value[i];
+		)");
+	}
+
+	FormatBounds += TEXT(R"(	
+					int CurAttributeIndex = AttributeIndex + i;
+					int3 TileOffset = {PerAttributeDataName}[CurAttributeIndex].xyz;
+					{OutputGrid}[int3(In_IndexX + TileOffset.x, In_IndexY + TileOffset.y, In_IndexZ + TileOffset.z)] = Val;
+				}
+			}
+		)");
+	TMap<FString, FStringFormatArg> ArgsBounds = {
+		{TEXT("ParameterName"),				ParamInfo.DataInterfaceHLSLSymbol},
+		{TEXT("FunctionName"),				FunctionInfo.InstanceName},
+		{TEXT("OutputGrid"),				ParamInfo.DataInterfaceHLSLSymbol + OutputGridName},
+		{TEXT("NumCellsName"),				ParamInfo.DataInterfaceHLSLSymbol + UNiagaraDataInterfaceRWBase::NumCellsName},
+		{TEXT("UnitToUVName"),				ParamInfo.DataInterfaceHLSLSymbol + UNiagaraDataInterfaceRWBase::UnitToUVName},
+		{TEXT("PerAttributeDataName"),		ParamInfo.DataInterfaceHLSLSymbol + TEXT("_") + PerAttributeDataName},
+		{TEXT("NumChannels"),				FString::FromInt(InNumChannels)},
+		{TEXT("NumChannelsVariableSuffix"),	InNumChannels > 1 ? FString::FromInt(InNumChannels) : TEXT("")},
+		{TEXT("NumTiles"),					ParamInfo.DataInterfaceHLSLSymbol + NumTilesName},
+		{TEXT("UnitToUVName"),				ParamInfo.DataInterfaceHLSLSymbol + UNiagaraDataInterfaceRWBase::UnitToUVName},
+		{TEXT("NumCellsName"),				ParamInfo.DataInterfaceHLSLSymbol + UNiagaraDataInterfaceRWBase::NumCellsName},
+		{TEXT("NumAttributesName"),			ParamInfo.DataInterfaceHLSLSymbol + UNiagaraDataInterfaceRWBase::NumAttributesName},
+	};
+	OutHLSL += FString::Format(*FormatBounds, ArgsBounds);
+}
+
 void UNiagaraDataInterfaceGrid3DCollection::WriteGetHLSL(const FNiagaraDataInterfaceGPUParamInfo& ParamInfo, const FNiagaraDataInterfaceGeneratedFunction& FunctionInfo, int FunctionInstanceIndex, int32 InNumChannels, FString& OutHLSL)
 {
 
@@ -1344,6 +1910,57 @@ void UNiagaraDataInterfaceGrid3DCollection::WriteSampleHLSL(const FNiagaraDataIn
 	OutHLSL += FString::Format(*FormatBounds, ArgsBounds);
 }
 
+// code generation for sampling grid attributes at a specific index - only considering tiled grid attributes
+void UNiagaraDataInterfaceGrid3DCollection::WriteSampleAtIndexHLSL(const FNiagaraDataInterfaceGPUParamInfo& ParamInfo, const FNiagaraDataInterfaceGeneratedFunction& FunctionInfo, int32 InNumChannels, FString SampleFunction, FString& OutHLSL)
+{
+	SampleFunction.ReplaceInline(TEXT("{Grid}"), *(ParamInfo.DataInterfaceHLSLSymbol + GridName));
+	SampleFunction.ReplaceInline(TEXT("{ParameterName}"), *ParamInfo.DataInterfaceHLSLSymbol);
+
+	FString FormatBounds = TEXT(R"(
+			void {FunctionName}(float3 In_Unit, int AttributeIndex, out float{NumChannelsVariableSuffix} Out_Val)
+			{			
+				float3 TileUV = clamp(In_Unit, {UnitClampMin}, {UnitClampMax}) * {OneOverNumTiles};
+			    for (int i = 0; i < {NumChannels}; i++)
+				{
+					int CurAttributeIndex = AttributeIndex + i;
+					float3 UVW = {PerAttributeDataName}[CurAttributeIndex + {NumAttributesName}].xyz + TileUV;
+
+					float Val = {SampleFunction}({SamplerName}, UVW, 0);
+					)");
+	if (InNumChannels == 1)
+	{
+		FormatBounds += TEXT("					Out_Val = Val;\n");
+	}
+	else
+	{
+		FormatBounds += TEXT("					Out_Val[i] = Val;\n");
+	}
+
+	FormatBounds += TEXT(R"(
+				}
+			}
+		)");
+
+	TMap<FString, FStringFormatArg> ArgsBounds = {
+		{TEXT("ParameterName"),				ParamInfo.DataInterfaceHLSLSymbol},
+		{TEXT("FunctionName"),				FunctionInfo.InstanceName},
+		{TEXT("Grid"),						ParamInfo.DataInterfaceHLSLSymbol + GridName},
+		{TEXT("SamplerName"),				ParamInfo.DataInterfaceHLSLSymbol + SamplerName },
+		{TEXT("SampleFunction"),			SampleFunction },
+		{TEXT("NumCellsName"),				ParamInfo.DataInterfaceHLSLSymbol + UNiagaraDataInterfaceRWBase::NumCellsName},
+		{TEXT("UnitToUVName"),				ParamInfo.DataInterfaceHLSLSymbol + UNiagaraDataInterfaceRWBase::UnitToUVName},
+		{TEXT("NumChannels"),				FString::FromInt(InNumChannels)},
+		{TEXT("NumChannelsVariableSuffix"),	InNumChannels > 1 ? FString::FromInt(InNumChannels) : TEXT("")},
+		{TEXT("PerAttributeDataName"),		ParamInfo.DataInterfaceHLSLSymbol + TEXT("_") + PerAttributeDataName},
+		{TEXT("NumTiles"),					ParamInfo.DataInterfaceHLSLSymbol + NumTilesName},
+		{TEXT("OneOverNumTiles"),			ParamInfo.DataInterfaceHLSLSymbol + OneOverNumTilesName},
+		{TEXT("UnitClampMin"),				ParamInfo.DataInterfaceHLSLSymbol + UnitClampMinName},
+		{TEXT("UnitClampMax"),				ParamInfo.DataInterfaceHLSLSymbol + UnitClampMaxName},
+		{TEXT("NumAttributesName"),			ParamInfo.DataInterfaceHLSLSymbol + UNiagaraDataInterfaceRWBase::NumAttributesName},
+	};
+	OutHLSL += FString::Format(*FormatBounds, ArgsBounds);
+}
+
 void UNiagaraDataInterfaceGrid3DCollection::WriteAttributeGetIndexHLSL(const FNiagaraDataInterfaceGPUParamInfo& ParamInfo, const FNiagaraDataInterfaceGeneratedFunction& FunctionInfo, int FunctionInstanceIndex, int32 InNumChannels, FString& OutHLSL)
 {
 	FString FormatBounds = TEXT(R"(
@@ -1423,9 +2040,12 @@ bool UNiagaraDataInterfaceGrid3DCollection::GetFunctionHLSL(const FNiagaraDataIn
 		return true;
 	}
 
+	FGrid3DCollectionAttributeHelper ParamHelper(ParamInfo);
+	FString GridChannelString = ParamHelper.GetGridChannelString();
+
 	TMap<FString, FStringFormatArg> ArgsBounds =
 	{
-		{TEXT("ParameterName"),		ParamInfo.DataInterfaceHLSLSymbol},
+		{TEXT("ParameterName"),			ParamInfo.DataInterfaceHLSLSymbol},
 		{TEXT("FunctionName"),			FunctionInfo.InstanceName},
 		{TEXT("Grid"),					ParamInfo.DataInterfaceHLSLSymbol + GridName},
 		{TEXT("OutputGrid"),			ParamInfo.DataInterfaceHLSLSymbol + OutputGridName},
@@ -1441,22 +2061,50 @@ bool UNiagaraDataInterfaceGrid3DCollection::GetFunctionHLSL(const FNiagaraDataIn
 		{TEXT("NumCellsName"),			ParamInfo.DataInterfaceHLSLSymbol + UNiagaraDataInterfaceRWBase::NumCellsName},
 		{TEXT("PerAttributeDataName"),	ParamInfo.DataInterfaceHLSLSymbol + TEXT("_") + PerAttributeDataName},
 		{TEXT("NumAttributesName"),		ParamInfo.DataInterfaceHLSLSymbol + UNiagaraDataInterfaceRWBase::NumAttributesName},
+		{TEXT("GridChannels"),			GridChannelString},
 	};
 
 	if (FunctionInfo.DefinitionName == GetValueFunctionName || FunctionInfo.DefinitionName == GetPreviousValueAtIndexFunctionName)
 	{
-		static const TCHAR *FormatBounds = TEXT(R"(
-			void {FunctionName}(int In_IndexX, int In_IndexY, int In_IndexZ, int In_AttributeIndex, out float Out_Val)
+		// if we have an rgba grid, then we must query for the rgba value at the current point, and get the specific channel requested
+		if (ParamHelper.UseRGBAGrid())
+		{			
+			FString RGBAHlsl;
+			RGBAHlsl.Appendf(TEXT("void {FunctionName}(int In_IndexX, int In_IndexY, int In_IndexZ, int In_AttributeIndex, out float Out_Val)\n"));
+			RGBAHlsl.Appendf(TEXT("{\n"));
+			RGBAHlsl.Appendf(TEXT("    Out_Val = 0;\n"));
+			RGBAHlsl.Appendf(TEXT("    const float{GridChannels} TmpVal = {Grid}.Load(int4(In_IndexX, In_IndexY, In_IndexZ, 0));\n"));
+			RGBAHlsl.Appendf(TEXT("    if (In_AttributeIndex == 0)\n"));
+			RGBAHlsl.Appendf(TEXT("    {\n"));
+			RGBAHlsl.Appendf(TEXT("        Out_Val = TmpVal.r;\n"));
+			RGBAHlsl.Appendf(TEXT("    }\n"));
+
+			for (int i = 0; i < ParamHelper.TotalChannels - 1; ++i)
 			{
-				Out_Val = 0;
-				if ( In_AttributeIndex < {NumAttributesName} )
+				RGBAHlsl.Appendf(TEXT("    else if (In_AttributeIndex == %d)\n"), i + 1);
+				RGBAHlsl.Appendf(TEXT("    {\n"));
+				RGBAHlsl.Appendf(TEXT("        Out_Val = TmpVal.%s;\n"), *FGrid3DCollectionAttributeHelper::Channels[i + 1]);
+				RGBAHlsl.Appendf(TEXT("    }\n"));
+			}
+
+			RGBAHlsl.Appendf(TEXT("}\n"));
+			OutHLSL += FString::Format(*RGBAHlsl, ArgsBounds);
+		}
+		else
+		{
+			static const TCHAR* FormatBounds = TEXT(R"(
+				void {FunctionName}(int In_IndexX, int In_IndexY, int In_IndexZ, int In_AttributeIndex, out float Out_Val)
 				{
-					int3 TileOffset = {PerAttributeDataName}[In_AttributeIndex].xyz;
-					Out_Val = {Grid}.Load(int4(In_IndexX + TileOffset.x, In_IndexY + TileOffset.y, In_IndexZ + TileOffset.z, 0));
+					Out_Val = 0;
+					if ( In_AttributeIndex < {NumAttributesName} )
+					{
+						int3 TileOffset = {PerAttributeDataName}[In_AttributeIndex].xyz;
+						Out_Val = {Grid}.Load(int4(In_IndexX + TileOffset.x, In_IndexY + TileOffset.y, In_IndexZ + TileOffset.z, 0));
+					}
 				}
-			}
-		)");
-		OutHLSL += FString::Format(FormatBounds, ArgsBounds);
+			)");
+			OutHLSL += FString::Format(FormatBounds, ArgsBounds);
+		}
 		return true;
 	}
 	else if (FunctionInfo.DefinitionName == GetFullGridPreviousValueFunctionName)
@@ -1469,21 +2117,54 @@ bool UNiagaraDataInterfaceGrid3DCollection::GetFunctionHLSL(const FNiagaraDataIn
 		)");
 		OutHLSL += FString::Format(FormatBounds, ArgsBounds);
 		return true;
-	}
-	else if (FunctionInfo.DefinitionName == SetValueFunctionName )
-	{
-		static const TCHAR *FormatBounds = TEXT(R"(
-			void {FunctionName}(int In_IndexX, int In_IndexY, int In_IndexZ, int In_AttributeIndex, float In_Value, out int val)
-			{			
-				val = 0;
-				if ( In_AttributeIndex < {NumAttributesName} )
-				{
-					int3 TileOffset = {PerAttributeDataName}[In_AttributeIndex].xyz;
-					{OutputGrid}[int3(In_IndexX + TileOffset.x, In_IndexY + TileOffset.y, In_IndexZ + TileOffset.z)] = In_Value;
-				}
+	} 
+	else if (FunctionInfo.DefinitionName == SetValueFunctionName)
+	{
+
+		// if we are using an rgba grid, then we must query the float4 value from the uav, then set the channel we are requesting
+		// note that the dynamic branching seems necessary otherwise it didn't work as expected.  Might be due to code optimization
+		if (ParamHelper.UseRGBAGrid())
+		{			
+			FString RGBAHlsl;
+			RGBAHlsl.Appendf(TEXT("void {FunctionName}(int In_IndexX, int In_IndexY, int In_IndexZ, int In_AttributeIndex, float In_Value, out int val)\n"));
+			RGBAHlsl.Appendf(TEXT("{\n"));
+			RGBAHlsl.Appendf(TEXT("    val = 0;\n"));
+			RGBAHlsl.Appendf(TEXT("    [branch] if (In_AttributeIndex == 0)\n"));
+			RGBAHlsl.Appendf(TEXT("    {\n"));
+			RGBAHlsl.Appendf(TEXT("		   float{GridChannels} TmpVal = {OutputGrid}.Load(int4(In_IndexX, In_IndexY, In_IndexZ, 0));\n"));
+			RGBAHlsl.Appendf(TEXT("        TmpVal.r = In_Value;\n"));
+			RGBAHlsl.Appendf(TEXT("        {OutputGrid}[int3(In_IndexX, In_IndexY, In_IndexZ)] = TmpVal;\n"));
+			RGBAHlsl.Appendf(TEXT("    }\n"));
+
+			for (int i = 0; i < ParamHelper.TotalChannels - 1; ++i)
+			{
+				RGBAHlsl.Appendf(TEXT("    else if (In_AttributeIndex == %d)\n"), i + 1);
+				RGBAHlsl.Appendf(TEXT("    {\n"));
+				RGBAHlsl.Appendf(TEXT("		   float{GridChannels} TmpVal = {OutputGrid}.Load(int4(In_IndexX, In_IndexY, In_IndexZ, 0));\n"));
+				RGBAHlsl.Appendf(TEXT("        TmpVal.%s = In_Value;\n"), *FGrid3DCollectionAttributeHelper::Channels[i + 1]);
+				RGBAHlsl.Appendf(TEXT("        {OutputGrid}[int3(In_IndexX, In_IndexY, In_IndexZ)] = TmpVal;\n"));
+				RGBAHlsl.Appendf(TEXT("    }\n"));
 			}
-		)");
-		OutHLSL += FString::Format(FormatBounds, ArgsBounds);
+
+			RGBAHlsl.Appendf(TEXT("}\n"));
+			OutHLSL += FString::Format(*RGBAHlsl, ArgsBounds);
+		}
+		else
+		{
+			static const TCHAR* FormatBounds = TEXT(R"(
+				void {FunctionName}(int In_IndexX, int In_IndexY, int In_IndexZ, int In_AttributeIndex, float In_Value, out int val)
+				{			
+					val = 0;
+					if ( In_AttributeIndex < {NumAttributesName} )
+					{
+						int3 TileOffset = {PerAttributeDataName}[In_AttributeIndex].xyz;
+						{OutputGrid}[int3(In_IndexX + TileOffset.x, In_IndexY + TileOffset.y, In_IndexZ + TileOffset.z)] = In_Value;
+					}
+				}
+			)");
+			OutHLSL += FString::Format(FormatBounds, ArgsBounds);
+		}
+
 		return true;
 	}
 	else if (FunctionInfo.DefinitionName == SetFullGridValueFunctionName)
@@ -1500,67 +2181,95 @@ bool UNiagaraDataInterfaceGrid3DCollection::GetFunctionHLSL(const FNiagaraDataIn
 	}
 	else if (FunctionInfo.DefinitionName == CopyPreviousToCurrentForCellFunctionName)
 	{
-		static const TCHAR* FormatBounds = TEXT(R"(
-			void {FunctionName}(int In_IndexX, int In_IndexY, int In_IndexZ)
-			{
-				for (int AttributeIndex = 0; AttributeIndex < {NumNamedAttributes}.x; AttributeIndex++)
-				{		
-					int3 TileOffset = {PerAttributeDataName}[AttributeIndex].xyz;
+		if (ParamHelper.UseRGBAGrid())
+		{
+			static const TCHAR* FormatBounds = TEXT(R"(
+				void {FunctionName}(int In_IndexX, int In_IndexY, int In_IndexZ)
+				{					
+					const int3 Index = int3(In_IndexX, In_IndexY, In_IndexZ);
+					{OutputGrid}[Index] = {Grid}.Load(Index, 0));
+				}
+			)");
+			OutHLSL += FString::Format(FormatBounds, ArgsBounds);
+		}
+		else
+		{
+			static const TCHAR* FormatBounds = TEXT(R"(
+				void {FunctionName}(int In_IndexX, int In_IndexY, int In_IndexZ)
+				{
+					for (int AttributeIndex = 0; AttributeIndex < {NumNamedAttributes}.x; AttributeIndex++)
+					{		
+						int3 TileOffset = {PerAttributeDataName}[AttributeIndex].xyz;
 
-					float Val = {Grid}.Load(int4(In_IndexX + TileOffset.x, In_IndexY + TileOffset.y, In_IndexZ + TileOffset.z, 0));	
-					{OutputGrid}[int3(In_IndexX + TileOffset.x, In_IndexY + TileOffset.y, In_IndexZ + TileOffset.z)] = Val;
+						float Val = {Grid}.Load(int4(In_IndexX + TileOffset.x, In_IndexY + TileOffset.y, In_IndexZ + TileOffset.z, 0));	
+						{OutputGrid}[int3(In_IndexX + TileOffset.x, In_IndexY + TileOffset.y, In_IndexZ + TileOffset.z)] = Val;
+					}
 				}
-			}
-		)");
-		OutHLSL += FString::Format(FormatBounds, ArgsBounds);
+			)");
+			OutHLSL += FString::Format(FormatBounds, ArgsBounds);
+		}
 		return true;
 	}
 	else if (FunctionInfo.DefinitionName == CopyMaskedPreviousToCurrentForCellFunctionName)
 	{
-		static const TCHAR* FormatBounds = TEXT(R"(
-			void {FunctionName}(int In_IndexX, int In_IndexY, int In_IndexZ, int NumAttributesSet, int AttributeMask)
-			{	
-				// early out if we set all the attributes that exist on the DI
-				[branch]
-				if (NumAttributesSet == {NumAttributes}.x)
-				{
-					return;					
+		if (ParamHelper.UseRGBAGrid())
+		{			
+			static const TCHAR* FormatBounds = TEXT(R"(
+				void {FunctionName}(int In_IndexX, int In_IndexY, int In_IndexZ, int NumAttributesSet, int AttributeMask)
+				{		
+					const int3 Index = int3(In_IndexX, In_IndexY, In_IndexZ);									
+					{OutputGrid}[Index] = {Grid}.Load(int4(Index, 0));					
 				}
-
-				for (int AttributeIndex = 0; AttributeIndex < {NumAttributes}.x; AttributeIndex++)
-				{					
-					// check the attribute index in the attribute mask to see if it has been set
-					// we automatically pass through attributes higher than 31 since we can only
-					// store 32 attributes in the mask
+			)");
+			OutHLSL += FString::Format(FormatBounds, ArgsBounds);
+		}
+		else
+		{
+			static const TCHAR* FormatBounds = TEXT(R"(
+				void {FunctionName}(int In_IndexX, int In_IndexY, int In_IndexZ, int NumAttributesSet, int AttributeMask)
+				{	
+					// early out if we set all the attributes that exist on the DI
 					[branch]
-					if ((AttributeMask & (1l << AttributeIndex)) == 0 || AttributeIndex >= 32)
-					{						
-						int3 TileOffset = {PerAttributeDataName}[AttributeIndex].xyz;
+					if (NumAttributesSet == {NumAttributes}.x)
+					{
+						return;					
+					}
 
-						float Val = {Grid}.Load(int4(In_IndexX + TileOffset.x, In_IndexY + TileOffset.y, In_IndexZ + TileOffset.z, 0));	
-						{OutputGrid}[int3(In_IndexX + TileOffset.x, In_IndexY + TileOffset.y, In_IndexZ + TileOffset.z)] = Val;
+					for (int AttributeIndex = 0; AttributeIndex < {NumAttributes}.x; AttributeIndex++)
+					{					
+						// check the attribute index in the attribute mask to see if it has been set
+						// we automatically pass through attributes higher than 31 since we can only
+						// store 32 attributes in the mask
+						[branch]
+						if ((AttributeMask & (1l << AttributeIndex)) == 0 || AttributeIndex >= 32)
+						{						
+							int3 TileOffset = {PerAttributeDataName}[AttributeIndex].xyz;
+
+							float Val = {Grid}.Load(int4(In_IndexX + TileOffset.x, In_IndexY + TileOffset.y, In_IndexZ + TileOffset.z, 0));	
+							{OutputGrid}[int3(In_IndexX + TileOffset.x, In_IndexY + TileOffset.y, In_IndexZ + TileOffset.z)] = Val;
+						}
 					}
-				}
 
-	// alternative implementation - look
-	/*
-				int AttributeIndex = -1;
-				[loop]			
-				for (int x = 0; x < {NumAttributes}.x - NumAttributesSet; ++x)
-				{											
-					do
-					{
-						AttributeIndex++;																		
-					} while(AttributeIndex < 32 && (AttributeMask & (1 << AttributeIndex)) == 1);
+		// alternative implementation - look
+		/*
+					int AttributeIndex = -1;
+					[loop]			
+					for (int x = 0; x < {NumAttributes}.x - NumAttributesSet; ++x)
+					{											
+						do
+						{
+							AttributeIndex++;																		
+						} while(AttributeIndex < 32 && (AttributeMask & (1 << AttributeIndex)) == 1);
 					
-					int3 TileOffset = {PerAttributeDataName}[AttributeIndex].xyz;
-					float Val = {Grid}.Load(int4(In_IndexX + TileOffset.x, In_IndexY + TileOffset.y, In_IndexZ + TileOffset.z, 0));	
-					{OutputGrid}[int3(In_IndexX + TileOffset.x, In_IndexY + TileOffset.y, In_IndexZ + TileOffset.z)] = Val;
+						int3 TileOffset = {PerAttributeDataName}[AttributeIndex].xyz;
+						float Val = {Grid}.Load(int4(In_IndexX + TileOffset.x, In_IndexY + TileOffset.y, In_IndexZ + TileOffset.z, 0));	
+						{OutputGrid}[int3(In_IndexX + TileOffset.x, In_IndexY + TileOffset.y, In_IndexZ + TileOffset.z)] = Val;
+					}
+	*/
 				}
-*/
-			}
-		)");
-		OutHLSL += FString::Format(FormatBounds, ArgsBounds);
+			)");
+			OutHLSL += FString::Format(FormatBounds, ArgsBounds);
+		}
 		return true;
 	}
 	else if (FunctionInfo.DefinitionName == ClearCellFunctionName)
@@ -1583,82 +2292,200 @@ bool UNiagaraDataInterfaceGrid3DCollection::GetFunctionHLSL(const FNiagaraDataIn
 	}
 	else if (FunctionInfo.DefinitionName == SetVector4ValueFunctionName)
 	{
-		WriteSetHLSL(ParamInfo, FunctionInfo, FunctionInstanceIndex, 4, OutHLSL);
+		if (ParamHelper.UseRGBAGrid())
+		{
+			ParamHelper.WriteSetHLSL(FGrid3DCollectionAttributeHelper::AttributeRetrievalMode::RGBAGrid, FunctionInfo, OutHLSL);
+		}
+		else
+		{
+			WriteSetHLSL(ParamInfo, FunctionInfo, FunctionInstanceIndex, 4, OutHLSL);
+		}
 		return true;
 	}
 	else if (FunctionInfo.DefinitionName == GetPreviousVector4ValueFunctionName)
 	{
-		WriteGetHLSL(ParamInfo, FunctionInfo, FunctionInstanceIndex, 4, OutHLSL);
+		if (ParamHelper.UseRGBAGrid())
+		{
+			ParamHelper.WriteGetHLSL(FGrid3DCollectionAttributeHelper::AttributeRetrievalMode::RGBAGrid, FunctionInfo, OutHLSL);
+		}
+		else
+		{
+			WriteGetHLSL(ParamInfo, FunctionInfo, FunctionInstanceIndex, 4, OutHLSL);
+		}
+
 		return true;
 	}
 	else if (FunctionInfo.DefinitionName == SamplePreviousGridVector4FunctionName)
 	{
-		WriteSampleHLSL(ParamInfo, FunctionInfo, FunctionInstanceIndex, 4, "{Grid}.SampleLevel", OutHLSL);
+		if (ParamHelper.UseRGBAGrid())
+		{
+			ParamHelper.WriteSampleHLSL(FGrid3DCollectionAttributeHelper::AttributeRetrievalMode::RGBAGrid, FunctionInfo, false, OutHLSL);
+		}
+		else
+		{
+			WriteSampleHLSL(ParamInfo, FunctionInfo, FunctionInstanceIndex, 4, "{Grid}.SampleLevel", OutHLSL);
+		}
 		return true;
 	}
 	else if (FunctionInfo.DefinitionName == CubicSamplePreviousGridVector4FunctionName)
 	{
-		WriteSampleHLSL(ParamInfo, FunctionInfo, FunctionInstanceIndex, 4, "SampleTriCubicLagrange_{ParameterName}", OutHLSL);
+		if (ParamHelper.UseRGBAGrid())
+		{
+			ParamHelper.WriteSampleHLSL(FGrid3DCollectionAttributeHelper::AttributeRetrievalMode::RGBAGrid, FunctionInfo, true, OutHLSL);
+		}
+		else
+		{
+			WriteSampleHLSL(ParamInfo, FunctionInfo, FunctionInstanceIndex, 4, "SampleTriCubicLagrange_{ParameterName}", OutHLSL);
+		}
 		return true;
 	}
 	else if (FunctionInfo.DefinitionName == SetVector3ValueFunctionName || FunctionInfo.DefinitionName == SetVectorValueFunctionName)
 	{
-		WriteSetHLSL(ParamInfo, FunctionInfo, FunctionInstanceIndex, 3, OutHLSL);
+		if (ParamHelper.UseRGBAGrid())
+		{
+			ParamHelper.WriteSetHLSL(FGrid3DCollectionAttributeHelper::AttributeRetrievalMode::RGBAGrid, FunctionInfo, OutHLSL);
+		}
+		else
+		{
+			WriteSetHLSL(ParamInfo, FunctionInfo, FunctionInstanceIndex, 3, OutHLSL);
+		}
 		return true;
 	}
-	else if ( FunctionInfo.DefinitionName == GetPreviousVectorValueFunctionName)
+	else if (FunctionInfo.DefinitionName == GetPreviousVectorValueFunctionName)
 	{
-		WriteGetHLSL(ParamInfo, FunctionInfo, FunctionInstanceIndex, 3, OutHLSL);
+		if (ParamHelper.UseRGBAGrid())
+		{
+			ParamHelper.WriteGetHLSL(FGrid3DCollectionAttributeHelper::AttributeRetrievalMode::RGBAGrid, FunctionInfo, OutHLSL);
+		}
+		else
+		{
+			WriteGetHLSL(ParamInfo, FunctionInfo, FunctionInstanceIndex, 3, OutHLSL);
+		}
+
 		return true;
 	}
-	else if ( FunctionInfo.DefinitionName == SamplePreviousGridVectorFunctionName)
+	else if (FunctionInfo.DefinitionName == SamplePreviousGridVectorFunctionName)
 	{
-		WriteSampleHLSL(ParamInfo, FunctionInfo, FunctionInstanceIndex, 3, "{Grid}.SampleLevel", OutHLSL);
+		if (ParamHelper.UseRGBAGrid())
+		{
+			ParamHelper.WriteSampleHLSL(FGrid3DCollectionAttributeHelper::AttributeRetrievalMode::RGBAGrid, FunctionInfo, false, OutHLSL);
+		}
+		else
+		{
+			WriteSampleHLSL(ParamInfo, FunctionInfo, FunctionInstanceIndex, 3, "{Grid}.SampleLevel", OutHLSL);
+		}
+
 		return true;
 	}
 	else if (FunctionInfo.DefinitionName == CubicSamplePreviousGridVectorFunctionName)
 	{
-		WriteSampleHLSL(ParamInfo, FunctionInfo, FunctionInstanceIndex, 3, "SampleTriCubicLagrange_{ParameterName}", OutHLSL);
-		return true;
+		if (ParamHelper.UseRGBAGrid())
+		{
+			ParamHelper.WriteSampleHLSL(FGrid3DCollectionAttributeHelper::AttributeRetrievalMode::RGBAGrid, FunctionInfo, true, OutHLSL);
+		}
+		else
+		{
+			WriteSampleHLSL(ParamInfo, FunctionInfo, FunctionInstanceIndex, 3, "SampleTriCubicLagrange_{ParameterName}", OutHLSL);
+		}
 	}
 	else if (FunctionInfo.DefinitionName == SetVector2ValueFunctionName || FunctionInfo.DefinitionName == SetVector2DValueFunctionName)
 	{
-		WriteSetHLSL(ParamInfo, FunctionInfo, FunctionInstanceIndex, 2, OutHLSL);
+		if (ParamHelper.UseRGBAGrid())
+		{
+			ParamHelper.WriteSetHLSL(FGrid3DCollectionAttributeHelper::AttributeRetrievalMode::RGBAGrid, FunctionInfo, OutHLSL);
+		}
+		else
+		{
+			WriteSetHLSL(ParamInfo, FunctionInfo, FunctionInstanceIndex, 2, OutHLSL);
+		}
+
 		return true;
 	}
-	else if ( FunctionInfo.DefinitionName == GetPreviousVector2DValueFunctionName)
+	else if (FunctionInfo.DefinitionName == GetPreviousVector2DValueFunctionName)
 	{
-		WriteGetHLSL(ParamInfo, FunctionInfo, FunctionInstanceIndex, 2, OutHLSL);
+		if (ParamHelper.UseRGBAGrid())
+		{
+			ParamHelper.WriteGetHLSL(FGrid3DCollectionAttributeHelper::AttributeRetrievalMode::RGBAGrid, FunctionInfo, OutHLSL);
+		}
+		else
+		{
+			WriteGetHLSL(ParamInfo, FunctionInfo, FunctionInstanceIndex, 2, OutHLSL);
+		}
+
 		return true;
 	}
-	else if ( FunctionInfo.DefinitionName == SamplePreviousGridVector2DFunctionName)
+	else if (FunctionInfo.DefinitionName == SamplePreviousGridVector2DFunctionName)
 	{
-		WriteSampleHLSL(ParamInfo, FunctionInfo, FunctionInstanceIndex, 2, "{Grid}.SampleLevel", OutHLSL);
+		if (ParamHelper.UseRGBAGrid())
+		{
+			ParamHelper.WriteSampleHLSL(FGrid3DCollectionAttributeHelper::AttributeRetrievalMode::RGBAGrid, FunctionInfo, false, OutHLSL);
+		}
+		else
+		{
+			WriteSampleHLSL(ParamInfo, FunctionInfo, FunctionInstanceIndex, 2, "{Grid}.SampleLevel", OutHLSL);
+		}
 		return true;
 	}
 	else if (FunctionInfo.DefinitionName == CubicSamplePreviousGridVector2DFunctionName)
 	{
-		WriteSampleHLSL(ParamInfo, FunctionInfo, FunctionInstanceIndex, 2, "SampleTriCubicLagrange_{ParameterName}", OutHLSL);
-		return true;
+		if (ParamHelper.UseRGBAGrid())
+		{
+			ParamHelper.WriteSampleHLSL(FGrid3DCollectionAttributeHelper::AttributeRetrievalMode::RGBAGrid, FunctionInfo, true, OutHLSL);
+		}
+		else
+		{
+			WriteSampleHLSL(ParamInfo, FunctionInfo, FunctionInstanceIndex, 2, "SampleTriCubicLagrange_{ParameterName}", OutHLSL);
+		}
 	}
 	else if (FunctionInfo.DefinitionName == SetFloatValueFunctionName)
 	{
-		WriteSetHLSL(ParamInfo, FunctionInfo, FunctionInstanceIndex, 1, OutHLSL);
+		if (ParamHelper.UseRGBAGrid())
+		{
+			ParamHelper.WriteSetHLSL(FGrid3DCollectionAttributeHelper::AttributeRetrievalMode::RGBAGrid, FunctionInfo, OutHLSL);
+		}
+		else
+		{
+			WriteSetHLSL(ParamInfo, FunctionInfo, FunctionInstanceIndex, 1, OutHLSL);
+		}
+
 		return true;
 	}
-	else if ( FunctionInfo.DefinitionName == GetPreviousFloatValueFunctionName)
+	else if (FunctionInfo.DefinitionName == GetPreviousFloatValueFunctionName)
 	{
-		WriteGetHLSL(ParamInfo, FunctionInfo, FunctionInstanceIndex, 1, OutHLSL);
+		if (ParamHelper.UseRGBAGrid())
+		{
+			ParamHelper.WriteGetHLSL(FGrid3DCollectionAttributeHelper::AttributeRetrievalMode::RGBAGrid, FunctionInfo, OutHLSL);
+		}
+		else
+		{
+			WriteGetHLSL(ParamInfo, FunctionInfo, FunctionInstanceIndex, 1, OutHLSL);
+		}
+
 		return true;
 	}
 	else if (FunctionInfo.DefinitionName == SamplePreviousGridFloatFunctionName)
 	{
-		WriteSampleHLSL(ParamInfo, FunctionInfo, FunctionInstanceIndex, 1, "{Grid}.SampleLevel", OutHLSL);
+		if (ParamHelper.UseRGBAGrid())
+		{
+			ParamHelper.WriteSampleHLSL(FGrid3DCollectionAttributeHelper::AttributeRetrievalMode::RGBAGrid, FunctionInfo, false, OutHLSL);
+		}
+		else
+		{
+			WriteSampleHLSL(ParamInfo, FunctionInfo, FunctionInstanceIndex, 1, "{Grid}.SampleLevel", OutHLSL);
+		}
+
 		return true;
 	}
 	else if (FunctionInfo.DefinitionName == CubicSamplePreviousGridFloatFunctionName)
 	{
-		WriteSampleHLSL(ParamInfo, FunctionInfo, FunctionInstanceIndex, 1, "SampleTriCubicLagrange_{ParameterName}", OutHLSL);
+		if (ParamHelper.UseRGBAGrid())
+		{
+			ParamHelper.WriteSampleHLSL(FGrid3DCollectionAttributeHelper::AttributeRetrievalMode::RGBAGrid, FunctionInfo, true, OutHLSL);
+		}
+		else
+		{
+			WriteSampleHLSL(ParamInfo, FunctionInfo, FunctionInstanceIndex, 1, "SampleTriCubicLagrange_{ParameterName}", OutHLSL);
+		}
 		return true;
 	}
 	else if (FunctionInfo.DefinitionName == GetVector4AttributeIndexFunctionName)
@@ -1683,64 +2510,315 @@ bool UNiagaraDataInterfaceGrid3DCollection::GetFunctionHLSL(const FNiagaraDataIn
 	}
 	else if (FunctionInfo.DefinitionName == SampleGridFunctionName || FunctionInfo.DefinitionName == SamplePreviousGridAtIndexFunctionName)
 	{
-		static const TCHAR* FormatBounds = TEXT(R"(
-				void {FunctionName}(float In_UnitX, float In_UnitY, float In_UnitZ, int In_AttributeIndex, out float Out_Val)
-				{
-					Out_Val = 0;
-					if ( In_AttributeIndex < {NumAttributesName} )
+		
+	// if we have an rgba grid, then we must sample for the rgba value at the current point, and get the specific channel requested
+		if (ParamHelper.UseRGBAGrid())
+		{			
+			FString RGBAHlsl;
+			RGBAHlsl.Appendf(TEXT("void {FunctionName}(float In_UnitX, float In_UnitY, float In_UnitZ, int In_AttributeIndex, out float Out_Val)\n"));
+			RGBAHlsl.Appendf(TEXT("{\n"));
+			RGBAHlsl.Appendf(TEXT("    Out_Val = 0;\n"));
+			RGBAHlsl.Appendf(TEXT("    const float{GridChannels} TmpVal = {Grid}.SampleLevel({SamplerName}, float3(In_UnitX, In_UnitY, In_UnitZ), 0);\n"));
+
+			RGBAHlsl.Appendf(TEXT("    [branch] if (In_AttributeIndex == 0)\n"));
+			RGBAHlsl.Appendf(TEXT("    {\n"));
+			RGBAHlsl.Appendf(TEXT("        Out_Val = TmpVal.r;\n"));
+			RGBAHlsl.Appendf(TEXT("    }\n"));
+
+			for (int i = 0; i < ParamHelper.TotalChannels - 1; ++i)
+			{
+				RGBAHlsl.Appendf(TEXT("    else if (In_AttributeIndex == %d)\n"), i + 1);
+				RGBAHlsl.Appendf(TEXT("    {\n"));
+				RGBAHlsl.Appendf(TEXT("        Out_Val = TmpVal.%s;\n"), *FGrid3DCollectionAttributeHelper::Channels[i + 1]);
+				RGBAHlsl.Appendf(TEXT("    }\n"));
+			}
+
+			RGBAHlsl.Appendf(TEXT("}\n"));
+			OutHLSL += FString::Format(*RGBAHlsl, ArgsBounds);
+		}
+		else
+		{
+			static const TCHAR* FormatBounds = TEXT(R"(
+					void {FunctionName}(float In_UnitX, float In_UnitY, float In_UnitZ, int In_AttributeIndex, out float Out_Val)
 					{
-						float3 Unit = float3(In_UnitX, In_UnitY, In_UnitZ);
-						Unit = clamp(Unit, {UnitClampMin}, {UnitClampMax}) * {OneOverNumTiles};
+						Out_Val = 0;
+						if ( In_AttributeIndex < {NumAttributesName} )
+						{
+							float3 Unit = float3(In_UnitX, In_UnitY, In_UnitZ);
+							Unit = clamp(Unit, {UnitClampMin}, {UnitClampMax}) * {OneOverNumTiles};
 
-						float3 UVW = {PerAttributeDataName}[In_AttributeIndex + {NumAttributesName}].xyz + Unit;
-						Out_Val = {Grid}.SampleLevel({SamplerName}, UVW, 0);
+							float3 UVW = {PerAttributeDataName}[In_AttributeIndex + {NumAttributesName}].xyz + Unit;
+							Out_Val = {Grid}.SampleLevel({SamplerName}, UVW, 0);
+						}
 					}
-				}
-			)");		
-		OutHLSL += FString::Format(FormatBounds, ArgsBounds);
+				)");
+			OutHLSL += FString::Format(FormatBounds, ArgsBounds);
+		}
+
+		return true;
+	}
+	else if (FunctionInfo.DefinitionName == CubicSampleGridFunctionName || FunctionInfo.DefinitionName == CubicSamplePreviousGridAtIndexFunctionName)
+	{
+		if (ParamHelper.UseRGBAGrid())
+		{
+			//#todo(dmp): maybe fix this at some point or deprecate?
+			UE_LOG(LogNiagara, Error, TEXT("RGBA grids are not supported for this function"));
+			return false;
+		}
+		else
+		{
+			static const TCHAR* FormatBounds = TEXT(R"(
+					void {FunctionName}(float In_UnitX, float In_UnitY, float In_UnitZ, int In_AttributeIndex, out float Out_Val)
+					{
+						Out_Val = 0;
+						if ( In_AttributeIndex < {NumAttributesName} )
+						{
+							float3 Unit = float3(In_UnitX, In_UnitY, In_UnitZ);
+							Unit = clamp(Unit, {UnitClampMin}, {UnitClampMax}) * {OneOverNumTiles};
+
+							float3 UVW = {PerAttributeDataName}[In_AttributeIndex + {NumAttributesName}].xyz + Unit;
+							Out_Val = SampleTriCubicLagrange_{ParameterName}({SamplerName}, UVW, 0);
+						}
+					}
+				)");
+			OutHLSL += FString::Format(FormatBounds, ArgsBounds);
+		}
+		return true;
+	}
+	else if (FunctionInfo.DefinitionName == SamplePreviousFullGridFunctionName)
+	{
+		if (ParamHelper.UseRGBAGrid())
+		{
+			//#todo(dmp): maybe fix this at some point or deprecate?
+			UE_LOG(LogNiagara, Error, TEXT("RGBA grids are not supported for this function"));
+			return false;
+		}
+		else
+		{
+			static const TCHAR* FormatBounds = TEXT(R"(
+						void {FunctionName}(float3 In_Unit, out float Out_Val)
+						{
+							Out_Val = {Grid}.SampleLevel({SamplerName}, In_Unit, 0);
+						}
+					)");
+			OutHLSL += FString::Format(FormatBounds, ArgsBounds);
+		}
+
+		return true;
+	}
+	else if (FunctionInfo.DefinitionName == CubicSamplePreviousFullGridFunctionName)
+	{
+		if (ParamHelper.UseRGBAGrid())
+		{
+			//#todo(dmp): maybe fix this at some point or deprecate?
+			UE_LOG(LogNiagara, Error, TEXT("RGBA grids are not supported for this function"));
+			return false;
+		}
+		else
+		{
+			static const TCHAR* FormatBounds = TEXT(R"(
+						void {FunctionName}(float3 In_Unit, out float Out_Val)
+						{
+							Out_Val = SampleTriCubicLagrange_{ParameterName}({SamplerName}, In_Unit, 0);
+						}
+					)");
+			OutHLSL += FString::Format(FormatBounds, ArgsBounds);
+		}
+		return true;
+	}
+	else if (FunctionInfo.DefinitionName == SetVector4ValueAtIndexFunctionName)
+	{
+		if (ParamHelper.UseRGBAGrid())
+		{
+			ParamHelper.WriteSetAtIndexHLSL(FGrid3DCollectionAttributeHelper::AttributeRetrievalMode::RGBAGrid, FunctionInfo, 4, OutHLSL);
+		}
+		else
+		{
+			WriteSetAtIndexHLSL(ParamInfo, FunctionInfo, 4, OutHLSL);
+		}
+		return true;
+	}
+	else if (FunctionInfo.DefinitionName == SetVector3ValueAtIndexFunctionName)
+	{
+		if (ParamHelper.UseRGBAGrid())
+		{
+			ParamHelper.WriteSetAtIndexHLSL(FGrid3DCollectionAttributeHelper::AttributeRetrievalMode::RGBAGrid, FunctionInfo, 3, OutHLSL);
+		}
+		else
+		{
+			WriteSetAtIndexHLSL(ParamInfo, FunctionInfo, 3, OutHLSL);
+		}
+		return true;
+	}
+	else if (FunctionInfo.DefinitionName == SetVector2ValueAtIndexFunctionName)
+	{
+		if (ParamHelper.UseRGBAGrid())
+		{
+			ParamHelper.WriteSetAtIndexHLSL(FGrid3DCollectionAttributeHelper::AttributeRetrievalMode::RGBAGrid, FunctionInfo, 2, OutHLSL);
+		}
+		else
+		{
+			WriteSetAtIndexHLSL(ParamInfo, FunctionInfo, 2, OutHLSL);
+		}
+		return true;
+	}
+	else if (FunctionInfo.DefinitionName == SetFloatValueAtIndexFunctionName)
+	{
+		if (ParamHelper.UseRGBAGrid())
+		{
+			ParamHelper.WriteSetAtIndexHLSL(FGrid3DCollectionAttributeHelper::AttributeRetrievalMode::RGBAGrid, FunctionInfo, 1, OutHLSL);
+		}
+		else
+		{
+			WriteSetAtIndexHLSL(ParamInfo, FunctionInfo, 1, OutHLSL);
+		}
+		return true;
+	}
+	else if (FunctionInfo.DefinitionName == GetPreviousVector4ValueAtIndexFunctionName)
+	{
+		if (ParamHelper.UseRGBAGrid())
+		{
+			ParamHelper.WriteGetAtIndexHLSL(FGrid3DCollectionAttributeHelper::AttributeRetrievalMode::RGBAGrid, FunctionInfo, 4, OutHLSL);
+		}
+		else
+		{
+			WriteGetAtIndexHLSL(ParamInfo, FunctionInfo, 4, OutHLSL);
+		}
+		return true;
+	}
+	else if (FunctionInfo.DefinitionName == GetPreviousVector3ValueAtIndexFunctionName)
+	{
+		if (ParamHelper.UseRGBAGrid())
+		{
+			ParamHelper.WriteGetAtIndexHLSL(FGrid3DCollectionAttributeHelper::AttributeRetrievalMode::RGBAGrid, FunctionInfo, 3, OutHLSL);
+		}
+		else
+		{
+			WriteGetAtIndexHLSL(ParamInfo, FunctionInfo, 3, OutHLSL);
+		}
+		return true;
+	}
+	else if (FunctionInfo.DefinitionName == GetPreviousVector2ValueAtIndexFunctionName)
+	{
+		if (ParamHelper.UseRGBAGrid())
+		{
+			ParamHelper.WriteGetAtIndexHLSL(FGrid3DCollectionAttributeHelper::AttributeRetrievalMode::RGBAGrid, FunctionInfo, 2, OutHLSL);
+		}
+		else
+		{
+			WriteGetAtIndexHLSL(ParamInfo, FunctionInfo, 2, OutHLSL);
+		}
+		return true;
+	}
+	else if (FunctionInfo.DefinitionName == GetPreviousFloatValueAtIndexFunctionName)
+	{
+		if (ParamHelper.UseRGBAGrid())
+		{
+			ParamHelper.WriteGetAtIndexHLSL(FGrid3DCollectionAttributeHelper::AttributeRetrievalMode::RGBAGrid, FunctionInfo, 1, OutHLSL);
+		}
+		else
+		{
+			WriteGetAtIndexHLSL(ParamInfo, FunctionInfo, 1, OutHLSL);
+		}
+		return true;
+	}
+	else if (FunctionInfo.DefinitionName == SamplePreviousGridVector4AtIndexFunctionName)
+	{
+		if (ParamHelper.UseRGBAGrid())
+		{
+			ParamHelper.WriteSampleAtIndexHLSL(FGrid3DCollectionAttributeHelper::AttributeRetrievalMode::RGBAGrid, FunctionInfo, 4, false, OutHLSL);
+		}
+		else
+		{
+			WriteSampleAtIndexHLSL(ParamInfo, FunctionInfo, 4, "{Grid}.SampleLevel", OutHLSL);
+		}
+		return true;
+	}
+	else if (FunctionInfo.DefinitionName == SamplePreviousGridVector3AtIndexFunctionName)
+	{
+		if (ParamHelper.UseRGBAGrid())
+		{
+			ParamHelper.WriteSampleAtIndexHLSL(FGrid3DCollectionAttributeHelper::AttributeRetrievalMode::RGBAGrid, FunctionInfo, 3, false, OutHLSL);
+		}
+		else
+		{
+			WriteSampleAtIndexHLSL(ParamInfo, FunctionInfo, 3, "{Grid}.SampleLevel", OutHLSL);
+		}
+		return true;
+	}
+	else if (FunctionInfo.DefinitionName == SamplePreviousGridVector2AtIndexFunctionName)
+	{
+		if (ParamHelper.UseRGBAGrid())
+		{
+			ParamHelper.WriteSampleAtIndexHLSL(FGrid3DCollectionAttributeHelper::AttributeRetrievalMode::RGBAGrid, FunctionInfo, 2, false, OutHLSL);
+		}
+		else
+		{
+			WriteSampleAtIndexHLSL(ParamInfo, FunctionInfo, 2, "{Grid}.SampleLevel", OutHLSL);
+		}
+		return true;
+	}
+	else if (FunctionInfo.DefinitionName == SamplePreviousGridFloatAtIndexFunctionName)
+	{
+		if (ParamHelper.UseRGBAGrid())
+		{
+			ParamHelper.WriteSampleAtIndexHLSL(FGrid3DCollectionAttributeHelper::AttributeRetrievalMode::RGBAGrid, FunctionInfo, 1, false, OutHLSL);
+		}
+		else
+		{
+			WriteSampleAtIndexHLSL(ParamInfo, FunctionInfo, 1, "{Grid}.SampleLevel", OutHLSL);
+		}
+		return true;
+	}
+	else if (FunctionInfo.DefinitionName == CubicSamplePreviousGridVector4AtIndexFunctionName)
+	{
+		if (ParamHelper.UseRGBAGrid())
+		{
+			ParamHelper.WriteSampleAtIndexHLSL(FGrid3DCollectionAttributeHelper::AttributeRetrievalMode::RGBAGrid, FunctionInfo, 4, true, OutHLSL);
+		}
+		else
+		{
+			WriteSampleAtIndexHLSL(ParamInfo, FunctionInfo, 4, "SampleTriCubicLagrange_{ParameterName}", OutHLSL);
+		}
 		return true;
 	}
-	else if (FunctionInfo.DefinitionName == CubicSampleGridFunctionName || FunctionInfo.DefinitionName == CubicSamplePreviousGridAtIndexFunctionName)
+	else if (FunctionInfo.DefinitionName == CubicSamplePreviousGridVector3AtIndexFunctionName)
 	{
-		static const TCHAR* FormatBounds = TEXT(R"(
-				void {FunctionName}(float In_UnitX, float In_UnitY, float In_UnitZ, int In_AttributeIndex, out float Out_Val)
-				{
-					Out_Val = 0;
-					if ( In_AttributeIndex < {NumAttributesName} )
-					{
-						float3 Unit = float3(In_UnitX, In_UnitY, In_UnitZ);
-						Unit = clamp(Unit, {UnitClampMin}, {UnitClampMax}) * {OneOverNumTiles};
-
-						float3 UVW = {PerAttributeDataName}[In_AttributeIndex + {NumAttributesName}].xyz + Unit;
-						Out_Val = SampleTriCubicLagrange_{ParameterName}({SamplerName}, UVW, 0);
-					}
-				}
-			)");
-		OutHLSL += FString::Format(FormatBounds, ArgsBounds);
+		if (ParamHelper.UseRGBAGrid())
+		{
+			ParamHelper.WriteSampleAtIndexHLSL(FGrid3DCollectionAttributeHelper::AttributeRetrievalMode::RGBAGrid, FunctionInfo, 3, true, OutHLSL);
+		}
+		else
+		{
+			WriteSampleAtIndexHLSL(ParamInfo, FunctionInfo, 3, "SampleTriCubicLagrange_{ParameterName}", OutHLSL);
+		}
 		return true;
 	}
-	else if (FunctionInfo.DefinitionName == SamplePreviousFullGridFunctionName)
+	else if (FunctionInfo.DefinitionName == CubicSamplePreviousGridVector2AtIndexFunctionName)
 	{
-		static const TCHAR* FormatBounds = TEXT(R"(
-					void {FunctionName}(float3 In_Unit, out float Out_Val)
-					{
-						Out_Val = {Grid}.SampleLevel({SamplerName}, In_Unit, 0);
-					}
-				)");
-		OutHLSL += FString::Format(FormatBounds, ArgsBounds);
+		if (ParamHelper.UseRGBAGrid())
+		{
+			ParamHelper.WriteSampleAtIndexHLSL(FGrid3DCollectionAttributeHelper::AttributeRetrievalMode::RGBAGrid, FunctionInfo, 2, true, OutHLSL);
+		}
+		else
+		{
+			WriteSampleAtIndexHLSL(ParamInfo, FunctionInfo, 2, "SampleTriCubicLagrange_{ParameterName}", OutHLSL);
+		}
 		return true;
 	}
-	else if (FunctionInfo.DefinitionName == CubicSamplePreviousFullGridFunctionName)
+	else if (FunctionInfo.DefinitionName == CubicSamplePreviousGridFloatAtIndexFunctionName)
 	{
-		static const TCHAR* FormatBounds = TEXT(R"(
-					void {FunctionName}(float3 In_Unit, out float Out_Val)
-					{
-						Out_Val = SampleTriCubicLagrange_{ParameterName}({SamplerName}, In_Unit, 0);
-					}
-				)");
-		OutHLSL += FString::Format(FormatBounds, ArgsBounds);
+		if (ParamHelper.UseRGBAGrid())
+		{
+			ParamHelper.WriteSampleAtIndexHLSL(FGrid3DCollectionAttributeHelper::AttributeRetrievalMode::RGBAGrid, FunctionInfo, 1, true, OutHLSL);
+		}
+		else
+		{
+			WriteSampleAtIndexHLSL(ParamInfo, FunctionInfo, 1, "SampleTriCubicLagrange_{ParameterName}", OutHLSL);
+		}
 		return true;
 	}
+
 	return false;
 }
 #endif
@@ -1822,16 +2900,16 @@ void UNiagaraDataInterfaceGrid3DCollection::SetShaderParameters(const FNiagaraDa
 	FRDGBuilder& GraphBuilder = Context.GetGraphBuilder();
 
 	FNDIGrid3DShaderParameters* Parameters = Context.GetParameterNestedStruct<FNDIGrid3DShaderParameters>();
-	Parameters->NumAttributes		= ProxyData->TotalNumAttributes;
-	Parameters->NumNamedAttributes	= ProxyData->TotalNumNamedAttributes;
-	Parameters->UnitToUV			= FVector3f(1.0f) / FVector3f(ProxyData->NumCells);
-	Parameters->NumCells			= ProxyData->NumCells;
-	Parameters->CellSize			= FVector3f(ProxyData->CellSize);
-	Parameters->NumTiles			= ProxyData->NumTiles;
-	Parameters->OneOverNumTiles		= FVector3f(1.0f) / FVector3f(ProxyData->NumTiles);
-	Parameters->UnitClampMin		= HalfPixelOffset;
-	Parameters->UnitClampMax		= FVector3f::OneVector - HalfPixelOffset;
-	Parameters->WorldBBoxSize		= FVector3f(ProxyData->WorldBBoxSize);
+	Parameters->NumAttributes = ProxyData->TotalNumAttributes;
+	Parameters->NumNamedAttributes = ProxyData->TotalNumNamedAttributes;
+	Parameters->UnitToUV = FVector3f(1.0f) / FVector3f(ProxyData->NumCells);
+	Parameters->NumCells = ProxyData->NumCells;
+	Parameters->CellSize = FVector3f(ProxyData->CellSize);
+	Parameters->NumTiles = ProxyData->NumTiles;
+	Parameters->OneOverNumTiles = FVector3f(1.0f) / FVector3f(ProxyData->NumTiles);
+	Parameters->UnitClampMin = HalfPixelOffset;
+	Parameters->UnitClampMax = FVector3f::OneVector - HalfPixelOffset;
+	Parameters->WorldBBoxSize = FVector3f(ProxyData->WorldBBoxSize);
 
 	if (Context.IsResourceBound(&Parameters->Grid))
 	{
@@ -1878,10 +2956,10 @@ FNiagaraDataInterfaceParametersCS* UNiagaraDataInterfaceGrid3DCollection::Create
 	ShaderStorage->AttributeNames.Reserve(NumFuncs);
 	ShaderStorage->AttributeChannelCount.Reserve(NumFuncs);
 
-	for (int32 FuncIdx=0; FuncIdx < NumFuncs; ++FuncIdx)
+	for (int32 FuncIdx = 0; FuncIdx < NumFuncs; ++FuncIdx)
 	{
 		const FNiagaraDataInterfaceGeneratedFunction& Func = ParameterInfo.GeneratedFunctions[FuncIdx];
-		static const FName NAME_Attribute("Attribute");
+
 		const FName* AttributeName = Func.FindSpecifierValue(NAME_Attribute);
 		if (AttributeName != nullptr)
 		{
@@ -1923,7 +3001,6 @@ bool UNiagaraDataInterfaceGrid3DCollection::GenerateIterationSourceNamespaceRead
 		}
 	}
 
-
 	if (InAttributes.Num() != InAttributeHLSLNames.Num())
 		return false;
 
@@ -1937,7 +3014,7 @@ bool UNiagaraDataInterfaceGrid3DCollection::GenerateIterationSourceNamespaceRead
 
 	for (int32 i = 0; i < InAttributes.Num(); i++)
 	{
-		OutHLSL += FString::Printf(TEXT("\t// Variable Name \"%s\" Type \"%s\" Var \"%s\"\n" ), *InAttributes[i].GetName().ToString(), *InAttributes[i].GetType().GetName(), *InAttributeHLSLNames[i]);
+		OutHLSL += FString::Printf(TEXT("\t// Variable Name \"%s\" Type \"%s\" Var \"%s\"\n"), *InAttributes[i].GetName().ToString(), *InAttributes[i].GetType().GetName(), *InAttributeHLSLNames[i]);
 
 		TArray<FString> OutArray;
 		if (InAttributes[i].GetName().ToString().ParseIntoArray(OutArray, TEXT(".")) > 0)
@@ -1951,7 +3028,7 @@ bool UNiagaraDataInterfaceGrid3DCollection::GenerateIterationSourceNamespaceRead
 
 			// Clear out the shared namespace with the root variable...
 			FString AttributeName;
-			for (int32 NamespaceIdx = 0;  NamespaceIdx < OutArray.Num(); NamespaceIdx++)
+			for (int32 NamespaceIdx = 0; NamespaceIdx < OutArray.Num(); NamespaceIdx++)
 			{
 				if (NamespaceIdx < RootArray.Num() && RootArray[NamespaceIdx] == OutArray[NamespaceIdx])
 					continue;
@@ -1970,8 +3047,10 @@ bool UNiagaraDataInterfaceGrid3DCollection::GenerateIterationSourceNamespaceRead
 	}
 	return true;
 }
-bool UNiagaraDataInterfaceGrid3DCollection::GenerateIterationSourceNamespaceWriteAttributesHLSL(FNiagaraDataInterfaceGPUParamInfo& DIInstanceInfo, const FNiagaraVariable& IterationSourceVar, TConstArrayView<FNiagaraVariable> InArguments, TConstArrayView<FNiagaraVariable> InAttributes, TConstArrayView<FString> InAttributeHLSLNames, bool bSpawnOnly, bool bPartialWrites, TArray<FText>& OutErrors, FString& OutHLSL) const
+bool UNiagaraDataInterfaceGrid3DCollection::GenerateIterationSourceNamespaceWriteAttributesHLSL(FNiagaraDataInterfaceGPUParamInfo& DIInstanceInfo, const FNiagaraVariable& IterationSourceVar, TConstArrayView<FNiagaraVariable> InArguments, TConstArrayView<FNiagaraVariable> InAttributes, TConstArrayView<FString> InAttributeHLSLNames, TConstArrayView<FNiagaraVariable> InAllAttributes, bool bSpawnOnly, bool bPartialWrites, TArray<FText>& OutErrors, FString& OutHLSL) const
 {
+	FGrid3DCollectionAttributeHelper ParamHelper(DIInstanceInfo);
+
 	FString DIVarName;
 	OutHLSL += TEXT("\t//Generated by UNiagaraDataInterfaceGrid3DCollection::GenerateIterationSourceNamespaceWriteAttributesHLSL\n");
 	for (int32 i = 0; i < InArguments.Num(); i++)
@@ -1985,19 +3064,47 @@ bool UNiagaraDataInterfaceGrid3DCollection::GenerateIterationSourceNamespaceWrit
 	if (InAttributes.Num() != InAttributeHLSLNames.Num())
 		return false;
 
+	// determine the total number of channels in the grid
+	int TotalChannels = 0;
+	for (int32 i = 0; i < InAllAttributes.Num(); i++)
+	{
+		int ComponentCount = GetComponentCountFromFuncName(TypeDefinitionToSetFunctionName(InAllAttributes[i].GetType()));
+		TotalChannels += ComponentCount;
+	}
+	bool UseRGBATexture = FGrid3DCollectionAttributeHelper::ShouldUseRGBAGrid(TotalChannels, InAttributes.Num());
+
 	TArray<FString> RootArray;
 	IterationSourceVar.GetName().ToString().ParseIntoArray(RootArray, TEXT("."));
 
-	OutHLSL += FString("\tint AttributeIsSetMask = 0;\n");		
+	// if we aren't using an rgba texture, then we want to maintain the attribute is set mask
+	if (!UseRGBATexture)
+	{
+		OutHLSL += FString("\tint AttributeIsSetMask = 0;\n");
+	}
+
+	// generate xyz execution index for the current grid
 	OutHLSL += FString("\tint CurrAttributeIndex;\n");
 	OutHLSL += FString("\tint X, Y, Z;\n");
 	OutHLSL += FString::Printf(TEXT("\t%s.ExecutionIndexToGridIndex(X, Y, Z);\n"), *DIVarName);
-	
-	
+
+	// only copy attributes from prev->next if the number of attributes written is less than the total number of attributes for the DI. 
+	if (UseRGBATexture && InAttributes.Num() > 0 && InAttributes.Num() < InAllAttributes.Num() && !bSpawnOnly && !bPartialWrites)
+	{
+		static const TCHAR* FormatBounds = TEXT(R"(			
+			{Grid}.CopyMaskedPreviousToCurrentForCell(X,Y,Z,0,0);
+		)");
+		TMap<FString, FStringFormatArg> ArgsBounds = {
+			{TEXT("Grid"), DIVarName},
+		};
+		OutHLSL += FString::Format(FormatBounds, ArgsBounds);
+	}
+
+	// generate code for setting values from parameter map to grid attributes
 	int NumAttributesSet = 0;
 	for (int32 i = 0; i < InAttributes.Num(); i++)
 	{
 		OutHLSL += FString::Printf(TEXT("\t// Name \"%s\" Type \"%s\" Var \"%s\"\n"), *InAttributes[i].GetName().ToString(), *InAttributes[i].GetType().GetName(), *InAttributeHLSLNames[i]);
+		int ComponentCount = GetComponentCountFromFuncName(TypeDefinitionToSetFunctionName(InAttributes[i].GetType()));
 
 		TArray<FString> OutArray;
 		if (InAttributes[i].GetName().ToString().ParseIntoArray(OutArray, TEXT(".")) > 0)
@@ -2024,24 +3131,26 @@ bool UNiagaraDataInterfaceGrid3DCollection::GenerateIterationSourceNamespaceWrit
 				if (AttributeName.Len() != 0)
 					AttributeName += TEXT(".");
 				AttributeName += OutArray[NamespaceIdx];
-			}						
-
-			OutHLSL += FString::Printf(TEXT("\t%s.%s<Attribute=\"%s\">(CurrAttributeIndex);\n"), *DIVarName, *TypeDefinitionToAttributeIndexFunctionName(InAttributes[i].GetType()).ToString(), *AttributeName);
-			
-			int ComponentCount = GetComponentCountFromFuncName(TypeDefinitionToSetFunctionName(InAttributes[i].GetType()));
+			}
 
-			for (int j = 0; j < ComponentCount; ++j)
+			// maintain the attribute set mask if we are not using an rgba texture
+			if (!UseRGBATexture)
 			{
-				OutHLSL += FString::Printf(TEXT("\tAttributeIsSetMask |= 1l << (CurrAttributeIndex+%d);\n"), j);				
-				NumAttributesSet++;				
-			}			
+				OutHLSL += FString::Printf(TEXT("\t%s.%s<Attribute=\"%s\">(CurrAttributeIndex);\n"), *DIVarName, *TypeDefinitionToAttributeIndexFunctionName(InAttributes[i].GetType()).ToString(), *AttributeName);
+
+				for (int j = 0; j < ComponentCount; ++j)
+				{
+					OutHLSL += FString::Printf(TEXT("\tAttributeIsSetMask |= 1l << (CurrAttributeIndex+%d);\n"), j);
+					NumAttributesSet++;
+				}
+			}
+
 			OutHLSL += FString::Printf(TEXT("\t%s.%s<Attribute=\"%s\">(X, Y, Z,  %s);\n"), *DIVarName, *TypeDefinitionToSetFunctionName(InAttributes[i].GetType()).ToString(), *AttributeName, *InAttributeHLSLNames[i]);
 		}
 	}
-
-	// First we need to copy all the data over from the input buffer, because we can't assume that this function will know all the attributes held within the grid. Instead, we copy all of them
-// over AND THEN overlay the local changes. Hopefully the optimizer will know enough to fix this up.
-	if (InAttributes.Num() > 0 && !bSpawnOnly && !bPartialWrites)
+	
+	// if we aren't using an rgba texture then we need to copy the grid attributes that were not set via the parameter map to make sure we don't lose values from previous stages
+	if (!UseRGBATexture && InAttributes.Num() > 0 && InAttributes.Num() < InAllAttributes.Num() && !bSpawnOnly && !bPartialWrites)
 	{
 
 		static const TCHAR* FormatBounds = TEXT(R"(			
@@ -2099,6 +3208,12 @@ bool UNiagaraDataInterfaceGrid3DCollection::InitPerInstanceData(void* PerInstanc
 	InstanceData->TotalNumAttributes = NumAttribChannelsFound;
 	InstanceData->TotalNumNamedAttributes = NumNamedAttribChannelsFound;
 
+	// give an error if there are any unnamed attributes
+	if (NumAttribChannelsFound - NumNamedAttribChannelsFound > 0)
+	{
+		UE_LOG(LogNiagara, Warning, TEXT("%s: Unnamed attributes should not be used with named.  This is a deprecated workflow"), *Proxy->SourceDIName.ToString());
+	}
+
 	if (SetResolutionMethod == ESetResolutionMethod::Independent)
 	{
 		InstanceData->NumCells.X = NumCells.X;
@@ -2110,20 +3225,13 @@ bool UNiagaraDataInterfaceGrid3DCollection::InitPerInstanceData(void* PerInstanc
 	}
 	else if (SetResolutionMethod == ESetResolutionMethod::MaxAxis)
 	{
-		InstanceData->CellSize = FVector(FMath::Max<float>(FMath::Max(WorldBBoxSize.X, WorldBBoxSize.Y), WorldBBoxSize.Z) / ((float) NumCellsMaxAxis));
+		InstanceData->CellSize = FVector(FMath::Max<float>(FMath::Max(WorldBBoxSize.X, WorldBBoxSize.Y), WorldBBoxSize.Z) / ((float)NumCellsMaxAxis));
 	}
 	else if (SetResolutionMethod == ESetResolutionMethod::CellSize)
 	{
 		InstanceData->CellSize = FVector(CellSize);
 	}
 
-	ENiagaraGpuBufferFormat BufferFormat = bOverrideFormat ? OverrideBufferFormat : GetDefault<UNiagaraSettings>()->DefaultGridFormat;
-	if (GNiagaraGrid3DOverrideFormat >= int32(ENiagaraGpuBufferFormat::Float) && (GNiagaraGrid3DOverrideFormat < int32(ENiagaraGpuBufferFormat::Max)))
-	{
-		BufferFormat = ENiagaraGpuBufferFormat(GNiagaraGrid3DOverrideFormat);
-	}
-	InstanceData->PixelFormat = FNiagaraUtilities::BufferFormatToPixelFormat(BufferFormat);
-
 	// compute world bounds and padding based on cell size
 	if (SetResolutionMethod == ESetResolutionMethod::MaxAxis || SetResolutionMethod == ESetResolutionMethod::CellSize)
 	{
@@ -2187,58 +3295,127 @@ bool UNiagaraDataInterfaceGrid3DCollection::InitPerInstanceData(void* PerInstanc
 		InstanceData->NumCells.Z = FMath::Max(1, int32(float(InstanceData->NumCells.Z) * GNiagaraGrid3DResolutionMultiplier));
 	}
 
-	// Compute number of tiles based on resolution of individual attributes
-	// #todo(dmp): refactor
-	int32 MaxDim = 2048;
-	int32 MaxTilesX = floor(MaxDim / InstanceData->NumCells.X);
-	int32 MaxTilesY = floor(MaxDim / InstanceData->NumCells.Y);
-	int32 MaxTilesZ = floor(MaxDim / InstanceData->NumCells.Z);
-	int32 MaxAttributes = MaxTilesX * MaxTilesY * MaxTilesZ;
-
-	if (InstanceData->NumCells.X > MaxDim || InstanceData->NumCells.Y > MaxDim || InstanceData->NumCells.Z > MaxDim)
+	if (NumAttribChannelsFound == 0)
 	{
-		UE_LOG(LogNiagara, Error, TEXT("Resolution is too high on %s... max is 2048, defined as %d, %d %d"), *Proxy->SourceDIName.ToString(), InstanceData->NumCells.X, InstanceData->NumCells.Y, InstanceData->NumCells.Z);
-		return false;
+		UE_LOG(LogNiagara, Warning, TEXT("Zero attributes defined on %s"), *Proxy->SourceDIName.ToString());
+
+		// Push Updates to Proxy.
+		FNiagaraDataInterfaceProxyGrid3DCollectionProxy* RT_Proxy = GetProxyAs<FNiagaraDataInterfaceProxyGrid3DCollectionProxy>();
+		ENQUEUE_RENDER_COMMAND(FUpdateData)(
+			[RT_Proxy, InstanceID = SystemInstance->GetId()](FRHICommandListImmediate& RHICmdList)
+			{
+				check(!RT_Proxy->SystemInstancesToProxyData_RT.Contains(InstanceID));
+				FGrid3DCollectionRWInstanceData_RenderThread* TargetData = &RT_Proxy->SystemInstancesToProxyData_RT.Add(InstanceID);
+
+			});
+
+		return true;
 	}
 
-	if (NumAttribChannelsFound > MaxAttributes && MaxAttributes > 0)
+	// Initialize target texture
+	ENiagaraGpuBufferFormat BufferFormat = bOverrideFormat ? OverrideBufferFormat : GetDefault<UNiagaraSettings>()->DefaultGridFormat;
+	if (GNiagaraGrid3DOverrideFormat >= int32(ENiagaraGpuBufferFormat::Float) && (GNiagaraGrid3DOverrideFormat < int32(ENiagaraGpuBufferFormat::Max)))
 	{
-		UE_LOG(LogNiagara, Error, TEXT("Invalid number of attributes defined on %s... max is %d, num defined is %d"), *Proxy->SourceDIName.ToString(), MaxAttributes, NumAttribChannelsFound);
-		return false;
+		BufferFormat = ENiagaraGpuBufferFormat(GNiagaraGrid3DOverrideFormat);
 	}
 
-	if (NumAttribChannelsFound == 0)
+	// if we have 4 or less attributes, pack data into rgba textures rather than a single large float texture with attributes tiled spatially
+	bool UseRGBATexture = FGrid3DCollectionAttributeHelper::ShouldUseRGBAGrid(NumAttribChannelsFound, InstanceData->Vars.Num());
+
+	if (UseRGBATexture)
 	{
-		UE_LOG(LogNiagara, Warning, TEXT("Invalid number of attributes defined on %s... max is %d, num defined is %d"), *Proxy->SourceDIName.ToString(), MaxAttributes, NumAttribChannelsFound);
-		
 
-		// Push Updates to Proxy.
-		FNiagaraDataInterfaceProxyGrid3DCollectionProxy* RT_Proxy = GetProxyAs<FNiagaraDataInterfaceProxyGrid3DCollectionProxy>();
-		ENQUEUE_RENDER_COMMAND(FUpdateData)(
-			[RT_Proxy, InstanceID = SystemInstance->GetId()](FRHICommandListImmediate& RHICmdList)
+		if (NumAttribChannelsFound - NumNamedAttribChannelsFound > 0)
 		{
-			check(!RT_Proxy->SystemInstancesToProxyData_RT.Contains(InstanceID));
-			FGrid3DCollectionRWInstanceData_RenderThread* TargetData = &RT_Proxy->SystemInstancesToProxyData_RT.Add(InstanceID);
+			UE_LOG(LogNiagara, Error, TEXT("%s: Unnamed attributes do not work with rgba textures"), *Proxy->SourceDIName.ToString());
+			return false;
+		}
 
-		});
+		InstanceData->NumTiles.X = 1;
+		InstanceData->NumTiles.Y = 1;
+		InstanceData->NumTiles.Z = 1;
 
-		return true;
+		switch (BufferFormat)
+		{
+		case ENiagaraGpuBufferFormat::Float:
+		{
+			switch (NumAttribChannelsFound)
+			{
+			case 1: InstanceData->PixelFormat = EPixelFormat::PF_R32_FLOAT; break;
+			case 2:	InstanceData->PixelFormat = EPixelFormat::PF_G32R32F; break;
+			case 3: InstanceData->PixelFormat = EPixelFormat::PF_R32G32B32F; break;
+			case 4:	InstanceData->PixelFormat = EPixelFormat::PF_A32B32G32R32F; break;
+			}
+			break;
+		}
+		/** 16-bit per channel floating point, range [-65504, 65504] */
+		case ENiagaraGpuBufferFormat::HalfFloat:
+		{
+			switch (NumAttribChannelsFound)
+			{
+			case 1: InstanceData->PixelFormat = EPixelFormat::PF_R16F; break;
+			case 2:	InstanceData->PixelFormat = EPixelFormat::PF_G16R16F; break;
+			case 3: InstanceData->PixelFormat = EPixelFormat::PF_FloatRGBA; break; // #todo(dmp): PF_FloatRGB translates to r11g11b10
+			case 4:	InstanceData->PixelFormat = EPixelFormat::PF_FloatRGBA; break;
+			}
+			break;
+		}
+		/** 8-bit per channel fixed point, range [0, 1]. */
+		case ENiagaraGpuBufferFormat::UnsignedNormalizedByte:
+		{
+			switch (NumAttribChannelsFound)
+			{
+			case 1: InstanceData->PixelFormat = EPixelFormat::PF_R8; break;
+			case 2:	InstanceData->PixelFormat = EPixelFormat::PF_R8G8; break;
+			case 3: InstanceData->PixelFormat = EPixelFormat::PF_R8G8B8A8; break; // #todo(dmp): need a 3 channel variant
+			case 4:	InstanceData->PixelFormat = EPixelFormat::PF_R8G8B8A8; break;
+			}
+
+			break;
+		}
+		}
+		InstanceData->UseRGBATexture = true;
 	}
+	else
+	{
+		// Compute number of tiles based on resolution of individual attributes
+
+		// #todo(dmp): refactor
+		int32 MaxDim = 2048;
+		int32 MaxTilesX = floor(MaxDim / InstanceData->NumCells.X);
+		int32 MaxTilesY = floor(MaxDim / InstanceData->NumCells.Y);
+		int32 MaxTilesZ = floor(MaxDim / InstanceData->NumCells.Z);
+		int32 MaxAttributes = MaxTilesX * MaxTilesY * MaxTilesZ;
+
+		if (InstanceData->NumCells.X > MaxDim || InstanceData->NumCells.Y > MaxDim || InstanceData->NumCells.Z > MaxDim)
+		{
+			UE_LOG(LogNiagara, Error, TEXT("Resolution is too high on %s... max is 2048, defined as %d, %d %d"), *Proxy->SourceDIName.ToString(), InstanceData->NumCells.X, InstanceData->NumCells.Y, InstanceData->NumCells.Z);
+			return false;
+		}
+
+		if (NumAttribChannelsFound > MaxAttributes && MaxAttributes > 0)
+		{
+			UE_LOG(LogNiagara, Error, TEXT("Invalid number of attributes defined on %s... max is %d, num defined is %d"), *Proxy->SourceDIName.ToString(), MaxAttributes, NumAttribChannelsFound);
+			return false;
+		}
 
-	// need to determine number of tiles in x and y based on number of attributes and max dimension size
-	int32 NumTilesX = FMath::Min<int32>(MaxTilesX, NumAttribChannelsFound);
-	int32 NumTilesY = FMath::Min<int32>(MaxTilesY, ceil(1.0 * NumAttribChannelsFound / NumTilesX));
-	int32 NumTilesZ = FMath::Min<int32>(MaxTilesZ, ceil(1.0 * NumAttribChannelsFound / (NumTilesX * NumTilesY)));
+		// need to determine number of tiles in x and y based on number of attributes and max dimension size
+		int32 NumTilesX = FMath::Min<int32>(MaxTilesX, NumAttribChannelsFound);
+		int32 NumTilesY = FMath::Min<int32>(MaxTilesY, ceil(1.0 * NumAttribChannelsFound / NumTilesX));
+		int32 NumTilesZ = FMath::Min<int32>(MaxTilesZ, ceil(1.0 * NumAttribChannelsFound / (NumTilesX * NumTilesY)));
 
-	InstanceData->NumTiles.X = NumTilesX;
-	InstanceData->NumTiles.Y = NumTilesY;
-	InstanceData->NumTiles.Z = NumTilesZ;
+		InstanceData->NumTiles.X = NumTilesX;
+		InstanceData->NumTiles.Y = NumTilesY;
+		InstanceData->NumTiles.Z = NumTilesZ;
 
-	check(InstanceData->NumTiles.X > 0);
-	check(InstanceData->NumTiles.Y > 0);
-	check(InstanceData->NumTiles.Z > 0);
+		check(InstanceData->NumTiles.X > 0);
+		check(InstanceData->NumTiles.Y > 0);
+		check(InstanceData->NumTiles.Z > 0);
+
+		InstanceData->UseRGBATexture = false;
+		InstanceData->PixelFormat = FNiagaraUtilities::BufferFormatToPixelFormat(BufferFormat);
+	}
 
-	// Initialize target texture
 	InstanceData->TargetTexture = nullptr;
 	InstanceData->RTUserParamBinding.Init(SystemInstance->GetInstanceParameters(), RenderTargetUserParameter.Parameter);
 	InstanceData->UpdateTargetTexture(BufferFormat);
@@ -2262,7 +3439,7 @@ bool UNiagaraDataInterfaceGrid3DCollection::InitPerInstanceData(void* PerInstanc
 			}
 		}
 		// Look for anonymous attributes
-		else if ( NumAttributes > 0 )
+		else if (NumAttributes > 0)
 		{
 			const FString PreviewAttributeString = PreviewAttribute.ToString();
 			if (PreviewAttributeString.StartsWith(AnonymousAttributeString))
@@ -2281,40 +3458,41 @@ bool UNiagaraDataInterfaceGrid3DCollection::InitPerInstanceData(void* PerInstanc
 	// Push Updates to Proxy.
 	FNiagaraDataInterfaceProxyGrid3DCollectionProxy* RT_Proxy = GetProxyAs<FNiagaraDataInterfaceProxyGrid3DCollectionProxy>();
 	ENQUEUE_RENDER_COMMAND(FUpdateData)(
-		[RT_Proxy, RT_Resource=InstanceData->TargetTexture ? InstanceData->TargetTexture->GetResource() : nullptr, InstanceID = SystemInstance->GetId(), RT_InstanceData=*InstanceData](FRHICommandListImmediate& RHICmdList)
-	{
-		check(!RT_Proxy->SystemInstancesToProxyData_RT.Contains(InstanceID));
-		FGrid3DCollectionRWInstanceData_RenderThread* TargetData = &RT_Proxy->SystemInstancesToProxyData_RT.Add(InstanceID);
-
-		TargetData->SourceDIName = RT_Proxy->SourceDIName;
-		TargetData->NumCells = RT_InstanceData.NumCells;
-		TargetData->NumTiles = RT_InstanceData.NumTiles;
-		TargetData->TotalNumAttributes = RT_InstanceData.TotalNumAttributes;
-		TargetData->TotalNumNamedAttributes = RT_InstanceData.TotalNumNamedAttributes;
-		TargetData->CellSize = RT_InstanceData.CellSize;
-		TargetData->WorldBBoxSize = RT_InstanceData.WorldBBoxSize;
-		TargetData->PixelFormat = RT_InstanceData.PixelFormat;
-		TargetData->Offsets = RT_InstanceData.Offsets;
-		TargetData->Vars.Reserve(RT_InstanceData.Vars.Num());
-		for (int32 i = 0; i < RT_InstanceData.Vars.Num(); i++)
-		{
-			TargetData->Vars.Emplace(RT_InstanceData.Vars[i].GetName());
-			TargetData->VarComponents.Emplace(RT_InstanceData.Vars[i].GetType().GetSize() / sizeof(float));
-		}
+		[RT_Proxy, RT_Resource = InstanceData->TargetTexture ? InstanceData->TargetTexture->GetResource() : nullptr, InstanceID = SystemInstance->GetId(), RT_InstanceData = *InstanceData](FRHICommandListImmediate& RHICmdList)
+		{
+			check(!RT_Proxy->SystemInstancesToProxyData_RT.Contains(InstanceID));
+			FGrid3DCollectionRWInstanceData_RenderThread* TargetData = &RT_Proxy->SystemInstancesToProxyData_RT.Add(InstanceID);
+
+			TargetData->SourceDIName = RT_Proxy->SourceDIName;
+			TargetData->NumCells = RT_InstanceData.NumCells;
+			TargetData->NumTiles = RT_InstanceData.NumTiles;
+			TargetData->TotalNumAttributes = RT_InstanceData.TotalNumAttributes;
+			TargetData->TotalNumNamedAttributes = RT_InstanceData.TotalNumNamedAttributes;
+			TargetData->CellSize = RT_InstanceData.CellSize;
+			TargetData->WorldBBoxSize = RT_InstanceData.WorldBBoxSize;
+			TargetData->PixelFormat = RT_InstanceData.PixelFormat;
+			TargetData->UseRGBATexture = RT_InstanceData.UseRGBATexture;
+			TargetData->Offsets = RT_InstanceData.Offsets;
+			TargetData->Vars.Reserve(RT_InstanceData.Vars.Num());
+			for (int32 i = 0; i < RT_InstanceData.Vars.Num(); i++)
+			{
+				TargetData->Vars.Emplace(RT_InstanceData.Vars[i].GetName());
+				TargetData->VarComponents.Emplace(RT_InstanceData.Vars[i].GetType().GetSize() / sizeof(float));
+			}
 #if WITH_EDITOR
-		TargetData->bPreviewGrid = RT_InstanceData.bPreviewGrid;
-		TargetData->PreviewAttribute = RT_InstanceData.PreviewAttribute;
+			TargetData->bPreviewGrid = RT_InstanceData.bPreviewGrid;
+			TargetData->PreviewAttribute = RT_InstanceData.PreviewAttribute;
 #endif
 
-		if (RT_Resource && RT_Resource->TextureRHI.IsValid())
-		{
-			TargetData->RenderTargetToCopyTo = RT_Resource->TextureRHI;
-		}
-		else
-		{
-			TargetData->RenderTargetToCopyTo = nullptr;
-		}
-	});
+			if (RT_Resource && RT_Resource->TextureRHI.IsValid())
+			{
+				TargetData->RenderTargetToCopyTo = RT_Resource->TextureRHI;
+			}
+			else
+			{
+				TargetData->RenderTargetToCopyTo = nullptr;
+			}
+		});
 
 	return true;
 }
@@ -2330,7 +3508,7 @@ void UNiagaraDataInterfaceGrid3DCollection::DestroyPerInstanceData(void* PerInst
 
 	FNiagaraDataInterfaceProxyGrid3DCollectionProxy* RT_Proxy = GetProxyAs<FNiagaraDataInterfaceProxyGrid3DCollectionProxy>();
 	ENQUEUE_RENDER_COMMAND(FNiagaraDIDestroyInstanceData) (
-		[RT_Proxy, InstanceID=SystemInstance->GetId()](FRHICommandListImmediate& CmdList)
+		[RT_Proxy, InstanceID = SystemInstance->GetId()](FRHICommandListImmediate& CmdList)
 		{
 			//check(ThisProxy->SystemInstancesToProxyData.Contains(InstanceID));
 			RT_Proxy->SystemInstancesToProxyData_RT.Remove(InstanceID);
@@ -2352,20 +3530,20 @@ bool UNiagaraDataInterfaceGrid3DCollection::PerInstanceTick(void* PerInstanceDat
 
 	FNiagaraDataInterfaceProxyGrid3DCollectionProxy* RT_Proxy = GetProxyAs<FNiagaraDataInterfaceProxyGrid3DCollectionProxy>();
 	ENQUEUE_RENDER_COMMAND(FUpdateData)(
-		[RT_Resource= InstanceData->TargetTexture ? InstanceData->TargetTexture->GetResource() : nullptr, RT_Proxy, InstanceID = SystemInstance->GetId()](FRHICommandListImmediate& RHICmdList)
-	{
-		FGrid3DCollectionRWInstanceData_RenderThread* TargetData = RT_Proxy->SystemInstancesToProxyData_RT.Find(InstanceID);
-
-		if (RT_Resource && RT_Resource->TextureRHI.IsValid())
+		[RT_Resource = InstanceData->TargetTexture ? InstanceData->TargetTexture->GetResource() : nullptr, RT_Proxy, InstanceID = SystemInstance->GetId()](FRHICommandListImmediate& RHICmdList)
 		{
-			TargetData->RenderTargetToCopyTo = RT_Resource->TextureRHI;
-		}
-		else
-		{
-			TargetData->RenderTargetToCopyTo = nullptr;
-		}
+			FGrid3DCollectionRWInstanceData_RenderThread* TargetData = RT_Proxy->SystemInstancesToProxyData_RT.Find(InstanceID);
 
-	});
+			if (RT_Resource && RT_Resource->TextureRHI.IsValid())
+			{
+				TargetData->RenderTargetToCopyTo = RT_Resource->TextureRHI;
+			}
+			else
+			{
+				TargetData->RenderTargetToCopyTo = nullptr;
+			}
+
+		});
 
 	return false;
 }
@@ -2393,14 +3571,12 @@ void UNiagaraDataInterfaceGrid3DCollection::CollectAttributesForScript(UNiagaraS
 	{
 		const FNiagaraVariableBase DataInterfaceVariable(FNiagaraTypeDefinition(UNiagaraDataInterfaceGrid3DCollection::StaticClass()), VariableName);
 
-		const int32 * IndexOfDataInterface = ParameterStore->FindParameterOffset(DataInterfaceVariable);
+		const int32* IndexOfDataInterface = ParameterStore->FindParameterOffset(DataInterfaceVariable);
 		if (IndexOfDataInterface != nullptr)
 		{
 			TConstArrayView<FNiagaraDataInterfaceGPUParamInfo> ParamInfoArray = Script->GetDataInterfaceGPUParamInfos();
 			for (const FNiagaraDataInterfaceGeneratedFunction& Func : ParamInfoArray[*IndexOfDataInterface].GeneratedFunctions)
 			{
-				static const FName NAME_Attribute("Attribute");
-
 				if (const FName* AttributeName = Func.FindSpecifierValue(NAME_Attribute))
 				{
 					FNiagaraVariableBase NewVar(UNiagaraDataInterfaceGrid3DCollection::GetValueTypeFromFuncName(Func.DefinitionName), *AttributeName);
@@ -2473,29 +3649,29 @@ void UNiagaraDataInterfaceGrid3DCollection::FindAttributes(TArray<FNiagaraVariab
 			// Search scripts for this data interface so we get the variable name
 			auto FindDataInterfaceVariable =
 				[&OwnerSystem, &EmitterData](const UNiagaraDataInterface* DataInterface) -> FName
+			{
+				UNiagaraScript* Scripts[] =
 				{
-					UNiagaraScript* Scripts[] =
-					{
-						OwnerSystem->GetSystemSpawnScript(),
-						OwnerSystem->GetSystemUpdateScript(),
-						EmitterData->GetGPUComputeScript(),
-					};
+					OwnerSystem->GetSystemSpawnScript(),
+					OwnerSystem->GetSystemUpdateScript(),
+					EmitterData->GetGPUComputeScript(),
+				};
 
-					for (UNiagaraScript* Script : Scripts)
+				for (UNiagaraScript* Script : Scripts)
+				{
+					for (FNiagaraScriptDataInterfaceInfo& DataInterfaceInfo : Script->GetCachedDefaultDataInterfaces())
 					{
-						for (FNiagaraScriptDataInterfaceInfo& DataInterfaceInfo : Script->GetCachedDefaultDataInterfaces())
+						if (DataInterfaceInfo.DataInterface == DataInterface)
 						{
-							if (DataInterfaceInfo.DataInterface == DataInterface)
-							{
-								return DataInterfaceInfo.RegisteredParameterMapRead.IsNone() ? DataInterfaceInfo.RegisteredParameterMapWrite : DataInterfaceInfo.RegisteredParameterMapRead;
-							}
+							return DataInterfaceInfo.RegisteredParameterMapRead.IsNone() ? DataInterfaceInfo.RegisteredParameterMapWrite : DataInterfaceInfo.RegisteredParameterMapRead;
 						}
 					}
-					return NAME_None;
-				};
+				}
+				return NAME_None;
+			};
 
 			const FName VariableName = FindDataInterfaceVariable(this);
-			if (!VariableName.IsNone() )
+			if (!VariableName.IsNone())
 			{
 				CollectAttributesForScript(EmitterData->GetGPUComputeScript(), VariableName, OutVariables, OutVariableOffsets, TotalAttributes, OutWarnings);
 			}
@@ -2504,7 +3680,7 @@ void UNiagaraDataInterfaceGrid3DCollection::FindAttributes(TArray<FNiagaraVariab
 	OutNumAttribChannelsFound = TotalAttributes - NumAttributes;
 }
 
-bool UNiagaraDataInterfaceGrid3DCollection::FillVolumeTexture(const UNiagaraComponent *Component, UVolumeTexture*Dest, int AttributeIndex)
+bool UNiagaraDataInterfaceGrid3DCollection::FillVolumeTexture(const UNiagaraComponent* Component, UVolumeTexture* Dest, int AttributeIndex)
 {
 	/*
 	#todo(dmp): we need to get a UVolumeTextureRenderTarget for any of this to work
@@ -2564,7 +3740,7 @@ bool UNiagaraDataInterfaceGrid3DCollection::FillVolumeTexture(const UNiagaraComp
 	return false;
 }
 
-bool UNiagaraDataInterfaceGrid3DCollection::FillRawVolumeTexture(const UNiagaraComponent *Component, UVolumeTexture *Dest, int &TilesX, int &TilesY, int&TilesZ)
+bool UNiagaraDataInterfaceGrid3DCollection::FillRawVolumeTexture(const UNiagaraComponent* Component, UVolumeTexture* Dest, int& TilesX, int& TilesY, int& TilesZ)
 {
 	/*
 	#todo(dmp): we need to get a UVolumeTextureRenderTarget for any of this to work
@@ -2632,7 +3808,7 @@ bool UNiagaraDataInterfaceGrid3DCollection::FillRawVolumeTexture(const UNiagaraC
 	return false;
 }
 
-void UNiagaraDataInterfaceGrid3DCollection::GetRawTextureSize(const UNiagaraComponent *Component, int &SizeX, int &SizeY, int &SizeZ)
+void UNiagaraDataInterfaceGrid3DCollection::GetRawTextureSize(const UNiagaraComponent* Component, int& SizeX, int& SizeY, int& SizeZ)
 {
 	if (!Component)
 	{
@@ -2666,7 +3842,7 @@ void UNiagaraDataInterfaceGrid3DCollection::GetRawTextureSize(const UNiagaraComp
 	SizeZ = Grid3DInstanceData->NumCells.Z * Grid3DInstanceData->NumTiles.Z;
 }
 
-void UNiagaraDataInterfaceGrid3DCollection::GetTextureSize(const UNiagaraComponent *Component, int &SizeX, int &SizeY, int &SizeZ)
+void UNiagaraDataInterfaceGrid3DCollection::GetTextureSize(const UNiagaraComponent* Component, int& SizeX, int& SizeY, int& SizeZ)
 {
 	if (!Component)
 	{
@@ -2762,7 +3938,7 @@ void UNiagaraDataInterfaceGrid3DCollection::VMUnitToFloatIndex(FVectorVMExternal
 
 	VectorVM::FExternalFuncRegisterHandler<float> OutFloatX(Context);
 	VectorVM::FExternalFuncRegisterHandler<float> OutFloatY(Context);
-	VectorVM::FExternalFuncRegisterHandler<float> OutFloatZ(Context);	
+	VectorVM::FExternalFuncRegisterHandler<float> OutFloatZ(Context);
 
 	for (int32 InstanceIdx = 0; InstanceIdx < Context.GetNumInstances(); ++InstanceIdx)
 	{
@@ -2791,66 +3967,75 @@ bool UNiagaraDataInterfaceGrid3DCollection::PerInstanceTickPostSimulate(void* Pe
 
 		FNiagaraDataInterfaceProxyGrid3DCollectionProxy* RT_Proxy = GetProxyAs<FNiagaraDataInterfaceProxyGrid3DCollectionProxy>();
 
-
-		// #todo(dmp): refactor
-		int32 MaxDim = 2048;
-		int32 MaxTilesX = floor(MaxDim / InstanceData->NumCells.X);
-		int32 MaxTilesY = floor(MaxDim / InstanceData->NumCells.Y);
-		int32 MaxTilesZ = floor(MaxDim / InstanceData->NumCells.Z);
-		int32 MaxAttributes = MaxTilesX * MaxTilesY * MaxTilesZ;
-		
-		if (InstanceData->NumCells.X > MaxDim || InstanceData->NumCells.Y > MaxDim || InstanceData->NumCells.Z > MaxDim)
+		// rgba textures only have a single tile since attributes are packed into r,g,b, and a channels
+		if (InstanceData->UseRGBATexture)
 		{
-			UE_LOG(LogNiagara, Error, TEXT("Resolution is too high on %s... max is 2048, defined as %d, %d %d"), *Proxy->SourceDIName.ToString(), InstanceData->NumCells.X, InstanceData->NumCells.Y, InstanceData->NumCells.Z);
-			return false;
+
+			InstanceData->NumTiles.X = 1;
+			InstanceData->NumTiles.Y = 1;
+			InstanceData->NumTiles.Z = 1;
 		}
-		
-		if ((InstanceData->TotalNumAttributes > MaxAttributes && MaxAttributes > 0))
+		else
 		{
-			UE_LOG(LogNiagara, Error, TEXT("Invalid number of attributes defined on %s... max is %d, num defined is %d"), *Proxy->SourceDIName.ToString(), MaxAttributes, InstanceData->TotalNumAttributes);
-			return false;
-		}
+			// #todo(dmp): refactor
+			int32 MaxDim = 2048;
+			int32 MaxTilesX = floor(MaxDim / InstanceData->NumCells.X);
+			int32 MaxTilesY = floor(MaxDim / InstanceData->NumCells.Y);
+			int32 MaxTilesZ = floor(MaxDim / InstanceData->NumCells.Z);
+			int32 MaxAttributes = MaxTilesX * MaxTilesY * MaxTilesZ;
+
+			if (InstanceData->NumCells.X > MaxDim || InstanceData->NumCells.Y > MaxDim || InstanceData->NumCells.Z > MaxDim)
+			{
+				UE_LOG(LogNiagara, Error, TEXT("Resolution is too high on %s... max is 2048, defined as %d, %d %d"), *Proxy->SourceDIName.ToString(), InstanceData->NumCells.X, InstanceData->NumCells.Y, InstanceData->NumCells.Z);
+				return false;
+			}
 
-		if (InstanceData->TotalNumAttributes <= 0)
-		{
-			UE_LOG(LogNiagara, Warning, TEXT("Invalid number of attributes defined on %s... max is %d, num defined is %d"), *Proxy->SourceDIName.ToString(), MaxAttributes, InstanceData->TotalNumAttributes);
-			return false;
-		}
+			if ((InstanceData->TotalNumAttributes > MaxAttributes && MaxAttributes > 0))
+			{
+				UE_LOG(LogNiagara, Error, TEXT("Invalid number of attributes defined on %s... max is %d, num defined is %d"), *Proxy->SourceDIName.ToString(), MaxAttributes, InstanceData->TotalNumAttributes);
+				return false;
+			}
 
-		// need to determine number of tiles in x and y based on number of attributes and max dimension size
-		int32 NumTilesX = FMath::Min<int32>(MaxTilesX, InstanceData->TotalNumAttributes);
-		int32 NumTilesY = FMath::Min<int32>(MaxTilesY, ceil(1.0 * InstanceData->TotalNumAttributes / NumTilesX));
-		int32 NumTilesZ = FMath::Min<int32>(MaxTilesZ, ceil(1.0 * InstanceData->TotalNumAttributes / (NumTilesX * NumTilesY)));
+			if (InstanceData->TotalNumAttributes <= 0)
+			{
+				UE_LOG(LogNiagara, Warning, TEXT("Invalid number of attributes defined on %s... max is %d, num defined is %d"), *Proxy->SourceDIName.ToString(), MaxAttributes, InstanceData->TotalNumAttributes);
+				return false;
+			}
 
-		InstanceData->NumTiles.X = NumTilesX;
-		InstanceData->NumTiles.Y = NumTilesY;
-		InstanceData->NumTiles.Z = NumTilesZ;
+			// need to determine number of tiles in x and y based on number of attributes and max dimension size
+			int32 NumTilesX = FMath::Min<int32>(MaxTilesX, InstanceData->TotalNumAttributes);
+			int32 NumTilesY = FMath::Min<int32>(MaxTilesY, ceil(1.0 * InstanceData->TotalNumAttributes / NumTilesX));
+			int32 NumTilesZ = FMath::Min<int32>(MaxTilesZ, ceil(1.0 * InstanceData->TotalNumAttributes / (NumTilesX * NumTilesY)));
 
-		check(InstanceData->NumTiles.X > 0);
-		check(InstanceData->NumTiles.Y > 0);
-		check(InstanceData->NumTiles.Z > 0);
+			InstanceData->NumTiles.X = NumTilesX;
+			InstanceData->NumTiles.Y = NumTilesY;
+			InstanceData->NumTiles.Z = NumTilesZ;
 
+			check(InstanceData->NumTiles.X > 0);
+			check(InstanceData->NumTiles.Y > 0);
+			check(InstanceData->NumTiles.Z > 0);
+		}
 		// #todo(dmp): we should align this method with the implementation in Grid3DCollection.  For now, we are relying on the next call to tick  to reset the User Texture
 
 		// Push Updates to Proxy.
 		ENQUEUE_RENDER_COMMAND(FUpdateData)(
 			[RT_Proxy, InstanceID = SystemInstance->GetId(), RT_InstanceData = *InstanceData](FRHICommandListImmediate& RHICmdList)
-		{
-			check(RT_Proxy->SystemInstancesToProxyData_RT.Contains(InstanceID));
-			FGrid3DCollectionRWInstanceData_RenderThread* TargetData = RT_Proxy->SystemInstancesToProxyData_RT.Find(InstanceID);
+			{
+				check(RT_Proxy->SystemInstancesToProxyData_RT.Contains(InstanceID));
+				FGrid3DCollectionRWInstanceData_RenderThread* TargetData = RT_Proxy->SystemInstancesToProxyData_RT.Find(InstanceID);
 
-			TargetData->NumCells = RT_InstanceData.NumCells;
-			TargetData->NumTiles = RT_InstanceData.NumTiles;
+				TargetData->NumCells = RT_InstanceData.NumCells;
+				TargetData->NumTiles = RT_InstanceData.NumTiles;
 
-			TargetData->TotalNumAttributes = RT_InstanceData.TotalNumAttributes;
-			TargetData->TotalNumNamedAttributes = RT_InstanceData.TotalNumNamedAttributes;
-			TargetData->CellSize = RT_InstanceData.CellSize;
+				TargetData->TotalNumAttributes = RT_InstanceData.TotalNumAttributes;
+				TargetData->TotalNumNamedAttributes = RT_InstanceData.TotalNumNamedAttributes;
+				TargetData->CellSize = RT_InstanceData.CellSize;
 
-			TargetData->Buffers.Empty();
-			TargetData->CurrentData = nullptr;
-			TargetData->DestinationData = nullptr;
-			TargetData->PerAttributeData.Release();
-		});
+				TargetData->Buffers.Empty();
+				TargetData->CurrentData = nullptr;
+				TargetData->DestinationData = nullptr;
+				TargetData->PerAttributeData.Release();
+			});
 
 	}
 
@@ -2960,17 +4145,17 @@ void FGrid3DCollectionRWInstanceData_RenderThread::BeginSimulate(FRDGBuilder& Gr
 
 		const FIntVector TextureSize(NumCells.X * NumTiles.X, NumCells.Y * NumTiles.Y, NumCells.Z * NumTiles.Z);
 		const FRDGTextureDesc TextureDesc = FRDGTextureDesc::Create3D(TextureSize, PixelFormat, FClearValueBinding::Black, ETextureCreateFlags::ShaderResource | ETextureCreateFlags::UAV);
-		
+
 		TStringBuilder<128> StringBuilder;
-		SourceDIName.AppendString(StringBuilder);		
-		StringBuilder.Append("_");		
+		SourceDIName.AppendString(StringBuilder);
+		StringBuilder.Append("_");
 		for (FName CurrName : Vars)
 		{
-			CurrName.AppendString(StringBuilder);			
-			StringBuilder.Append("_");			
+			CurrName.AppendString(StringBuilder);
+			StringBuilder.Append("_");
 		}
-		StringBuilder.Append("Grid3DCollection");		
-		
+		StringBuilder.Append("Grid3DCollection");
+
 		DestinationData->Initialize(GraphBuilder, *StringBuilder, TextureDesc);
 	}
 }
@@ -3090,12 +4275,12 @@ void FNiagaraDataInterfaceProxyGrid3DCollectionProxy::PostSimulate(const FNDIGpu
 
 FIntVector FNiagaraDataInterfaceProxyGrid3DCollectionProxy::GetElementCount(FNiagaraSystemInstanceID SystemInstanceID) const
 {
-	if ( const FGrid3DCollectionRWInstanceData_RenderThread* ProxyData = SystemInstancesToProxyData_RT.Find(SystemInstanceID) )
+	if (const FGrid3DCollectionRWInstanceData_RenderThread* ProxyData = SystemInstancesToProxyData_RT.Find(SystemInstanceID))
 	{
 		// support a grid reader acting as an iteration source
 		if (ProxyData->OtherProxy != nullptr)
 		{
-			FNiagaraDataInterfaceProxyGrid3DCollectionProxy *OtherGrid3DProxy = static_cast<FNiagaraDataInterfaceProxyGrid3DCollectionProxy*>(ProxyData->OtherProxy);
+			FNiagaraDataInterfaceProxyGrid3DCollectionProxy* OtherGrid3DProxy = static_cast<FNiagaraDataInterfaceProxyGrid3DCollectionProxy*>(ProxyData->OtherProxy);
 			const FGrid3DCollectionRWInstanceData_RenderThread* OtherProxyData = OtherGrid3DProxy->SystemInstancesToProxyData_RT.Find(SystemInstanceID);
 			return OtherProxyData->NumCells;
 		}
diff --git a/Engine/Plugins/FX/Niagara/Source/Niagara/Private/NiagaraDataInterfaceGrid3DCollectionUtils.cpp b/Engine/Plugins/FX/Niagara/Source/Niagara/Private/NiagaraDataInterfaceGrid3DCollectionUtils.cpp
new file mode 100644
--- /dev/null
+++ b/Engine/Plugins/FX/Niagara/Source/Niagara/Private/NiagaraDataInterfaceGrid3DCollectionUtils.cpp
@@ -0,0 +1,467 @@
+// Copyright Epic Games, Inc. All Rights Reserved.
+
+#include "NiagaraDataInterfaceGrid3DCollectionUtils.h"
+#include "NiagaraDataInterfaceGrid3DCollection.h"
+
+#define LOCTEXT_NAMESPACE "NiagaraDataInterfaceGrid3DCollection"
+
+static int32 GNiagaraGrid3DUseRGBAGrid = 0;
+static FAutoConsoleVariableRef CVarNiagaraGrid3DUseRGBAGrid(
+	TEXT("fx.Niagara.Grid3D.UseRGBAGrid"),
+	GNiagaraGrid3DUseRGBAGrid,
+	TEXT("Use RGBA textures when possible\n"),
+	ECVF_Default
+);
+
+TArray<FString> FGrid3DCollectionAttributeHelper::Channels = { "r", "g", "b", "a" };
+
+// static member function
+bool FGrid3DCollectionAttributeHelper::ShouldUseRGBAGrid(const int TotalChannels, const int TotalNumAttributes)
+{
+	return TotalNumAttributes == 1 && TotalChannels <= 4 && GNiagaraGrid3DUseRGBAGrid != 0;
+}
+
+FGrid3DCollectionAttributeHelper::FGrid3DCollectionAttributeHelper(const FNiagaraDataInterfaceGPUParamInfo& InParamInfo, TArray<FText>* OutWarnings /*= nullptr*/) : ParamInfo(InParamInfo)
+{
+	AttributeInfos.Reserve(ParamInfo.GeneratedFunctions.Num());
+	TotalChannels = 0;
+
+	for (const FNiagaraDataInterfaceGeneratedFunction& Function : ParamInfo.GeneratedFunctions)
+	{
+		const FName* AttributeName = Function.FindSpecifierValue(UNiagaraDataInterfaceGrid3DCollection::NAME_Attribute);
+		if (AttributeName == nullptr)
+		{
+			continue;
+		}
+
+		if (const FAttributeInfo* ExistingAttribute = FindAttributeInfo(*AttributeName))
+		{
+			if (OutWarnings != nullptr)
+			{
+				FNiagaraTypeDefinition AttributeTypeDef = UNiagaraDataInterfaceGrid3DCollection::GetValueTypeFromFuncName(Function.DefinitionName);
+				if (ExistingAttribute->TypeDef != AttributeTypeDef)
+				{
+					OutWarnings->Emplace(FText::Format(LOCTEXT("BadType", "Same name, different types! {0} vs {1}, Attribute {2}"), AttributeTypeDef.GetNameText(), ExistingAttribute->TypeDef.GetNameText(), FText::FromName(ExistingAttribute->Name)));
+				}
+			}
+			continue;
+		}
+
+		FAttributeInfo& AttributeInfo = AttributeInfos.AddDefaulted_GetRef();
+		AttributeInfo.Name = *AttributeName;
+		AttributeInfo.TypeDef = UNiagaraDataInterfaceGrid3DCollection::GetValueTypeFromFuncName(Function.DefinitionName);
+		AttributeInfo.NumChannels = UNiagaraDataInterfaceGrid3DCollection::GetComponentCountFromFuncName(Function.DefinitionName);
+		AttributeInfo.ChannelOffset = TotalChannels;
+		AttributeInfo.AttributeIndex = AttributeInfos.Num() - 1;
+		TotalChannels += AttributeInfo.NumChannels;
+	}
+}
+
+bool FGrid3DCollectionAttributeHelper::UseRGBAGrid()
+{
+	return ShouldUseRGBAGrid(TotalChannels, AttributeInfos.Num());
+}
+
+const FGrid3DCollectionAttributeHelper::FGrid3DCollectionAttributeHelper::FAttributeInfo* FGrid3DCollectionAttributeHelper::FindAttributeInfo(FName AttributeName) const
+{
+	return AttributeInfos.FindByPredicate([AttributeName](const FAttributeInfo& Info) { return Info.Name == AttributeName; });
+}
+
+#if WITH_EDITORONLY_DATA
+FString FGrid3DCollectionAttributeHelper::GetPerAttributePixelOffset(const TCHAR* DataInterfaceHLSLSymbol)
+{
+	return FString::Printf(TEXT("int3(%s_%s[(AttributeIndex * 2) + 0].xyz)"), DataInterfaceHLSLSymbol, *UNiagaraDataInterfaceGrid3DCollection::PerAttributeDataName);
+}
+
+FString FGrid3DCollectionAttributeHelper::GetPerAttributePixelOffset() const
+{
+	return GetPerAttributePixelOffset(*ParamInfo.DataInterfaceHLSLSymbol);
+}
+
+FString FGrid3DCollectionAttributeHelper::GetPerAttributeUVWOffset(const TCHAR* DataInterfaceHLSLSymbol)
+{
+	return FString::Printf(TEXT("%s%s[(AttributeIndex * 2) + 1].xyz"), *UNiagaraDataInterfaceGrid3DCollection::PerAttributeDataName, DataInterfaceHLSLSymbol);
+}
+
+FString FGrid3DCollectionAttributeHelper::GetPerAttributeUVWOffset() const
+{
+	return GetPerAttributeUVWOffset(*ParamInfo.DataInterfaceHLSLSymbol);
+}
+
+FString FGrid3DCollectionAttributeHelper::GetAttributeIndex(bool bUseAttributeIndirection, const FAttributeInfo* AttributeInfo, int Channel /*= 0*/) const
+{
+	check(AttributeInfo);
+	if (bUseAttributeIndirection)
+	{
+		return FString::Printf(TEXT("int(%s%s[%d].w + %d)"), *UNiagaraDataInterfaceGrid3DCollection::PerAttributeDataName, *ParamInfo.DataInterfaceHLSLSymbol, AttributeInfo->AttributeIndex, Channel);
+	}
+	else
+	{
+		return FString::Printf(TEXT("%d"), AttributeInfo->ChannelOffset + Channel);
+	}
+}
+
+FString FGrid3DCollectionAttributeHelper::GetGridChannelString()
+{
+	FString NumChannelsString = "";
+	if (TotalChannels > 1)
+	{
+		NumChannelsString = FString::FromInt(TotalChannels);
+	}
+	return NumChannelsString;
+}
+
+void FGrid3DCollectionAttributeHelper::GetChannelStrings(int AttributeIndex, int AttributeNumChannels, FString& NumChannelsString, FString& AttrGridChannels)
+{
+	NumChannelsString = "";
+	if (AttributeNumChannels > 1)
+	{
+		NumChannelsString = FString::FromInt(AttributeNumChannels);
+	}
+
+	AttrGridChannels = FGrid3DCollectionAttributeHelper::Channels[AttributeIndex];
+	for (int i = 1; i < AttributeNumChannels; ++i)
+	{
+		AttrGridChannels += FGrid3DCollectionAttributeHelper::Channels[AttributeIndex + i];
+	}
+}
+
+void FGrid3DCollectionAttributeHelper::GetChannelStrings(const FAttributeInfo* AttributeInfo, FString& NumChannelsString, FString& AttrGridChannels)
+{
+	NumChannelsString = "";
+	if (AttributeInfo->NumChannels > 1)
+	{
+		NumChannelsString = FString::FromInt(AttributeInfo->NumChannels);
+	}
+
+	AttrGridChannels = FGrid3DCollectionAttributeHelper::Channels[AttributeInfo->AttributeIndex];
+	for (int i = 1; i < AttributeInfo->NumChannels; ++i)
+	{
+		AttrGridChannels += FGrid3DCollectionAttributeHelper::Channels[AttributeInfo->AttributeIndex + i];
+	}
+}
+
+bool FGrid3DCollectionAttributeHelper::WriteGetHLSL(AttributeRetrievalMode AttributeStorage, const FNiagaraDataInterfaceGeneratedFunction& FunctionInfo, FString& OutHLSL)
+{
+	const FName* AttributeName = FunctionInfo.FindSpecifierValue(UNiagaraDataInterfaceGrid3DCollection::NAME_Attribute);
+	if (AttributeName == nullptr)
+	{
+		return false;
+	}
+
+	const FAttributeInfo* AttributeInfo = FindAttributeInfo(*AttributeName);
+	if (AttributeInfo == nullptr)
+	{
+		return false;
+	}
+
+	const FString GridNameHLSL = ParamInfo.DataInterfaceHLSLSymbol + UNiagaraDataInterfaceGrid3DCollection::GridName;
+
+	if (AttributeStorage == AttributeRetrievalMode::RGBAGrid)
+	{
+
+		FString NumChannelsString;
+		FString AttrGridChannels;
+
+		GetChannelStrings(AttributeInfo, NumChannelsString, AttrGridChannels);
+
+
+		OutHLSL.Appendf(TEXT("void %s(int IndexX, int IndexY, int IndexZ, out float%s Value)\n"), *FunctionInfo.InstanceName, *NumChannelsString);
+		OutHLSL.Appendf(TEXT("{\n"));
+		OutHLSL.Appendf(TEXT("	Value = %s.Load(int4(IndexX, IndexY, IndexZ, 0)).%s;\n"), *GridNameHLSL, *AttrGridChannels);
+		OutHLSL.Appendf(TEXT("}\n"));
+	}
+	else
+	{
+		if (AttributeInfo->NumChannels == 1)
+		{
+			OutHLSL.Appendf(TEXT("void %s(int IndexX, int IndexY, int IndexZ, out float Value)\n"), *FunctionInfo.InstanceName);
+			OutHLSL.Appendf(TEXT("{\n"));
+			OutHLSL.Appendf(TEXT("	int z %s;\n"), *GetAttributeIndex(AttributeStorage == AttributeRetrievalMode::Indirection, AttributeInfo));
+			OutHLSL.Appendf(TEXT("	int3 PixelOffset = int3(IndexX, IndexY, IndexZ) + %s;\n"), *GetPerAttributePixelOffset());
+			OutHLSL.Appendf(TEXT("	Value = %s.Load(int4(PixelOffset, 0));\n"), *GridNameHLSL);
+			OutHLSL.Appendf(TEXT("}\n"));
+		}
+		else
+		{
+			OutHLSL.Appendf(TEXT("void %s(int IndexX, int IndexY, int IndexZ, out float%d Value)\n"), *FunctionInfo.InstanceName, AttributeInfo->NumChannels);
+			OutHLSL.Appendf(TEXT("{\n"));
+			for (int32 i = 0; i < AttributeInfo->NumChannels; ++i)
+			{
+				OutHLSL.Appendf(TEXT("	{\n"));
+				OutHLSL.Appendf(TEXT("		int AttributeIndex = %s;\n"), *GetAttributeIndex(AttributeStorage == AttributeRetrievalMode::Indirection, AttributeInfo, i));
+				OutHLSL.Appendf(TEXT("		int3 PixelOffset = int3(IndexX, IndexY, IndexZ) + %s;\n"), *GetPerAttributePixelOffset());
+				OutHLSL.Appendf(TEXT("		Value[%d] = %s.Load(int4(PixelOffset, 0));\n"), i, *GridNameHLSL);
+				OutHLSL.Appendf(TEXT("	}\n"));
+			}
+			OutHLSL.Appendf(TEXT("}\n"));
+		}
+	}
+
+	return true;
+}
+
+bool FGrid3DCollectionAttributeHelper::WriteGetAtIndexHLSL(AttributeRetrievalMode AttributeStorage, const FNiagaraDataInterfaceGeneratedFunction& FunctionInfo, int NumChannels, FString& OutHLSL)
+{
+	FString NumChannelsString;
+	FString AttrGridChannels;
+
+	const FString GridNameHLSL = ParamInfo.DataInterfaceHLSLSymbol + UNiagaraDataInterfaceGrid3DCollection::GridName;
+
+	GetChannelStrings(0, NumChannels, NumChannelsString, AttrGridChannels);
+
+	if (AttributeInfos.Num() > 1)
+	{
+		//#todo(dmp): fill this in
+		return false;
+	}
+	else
+	{
+		OutHLSL.Appendf(TEXT("void %s(int IndexX, int IndexY, int IndexZ, int AttributeIndex, out float%s Value)\n"), *FunctionInfo.InstanceName, *NumChannelsString);
+		OutHLSL.Appendf(TEXT("{\n"));
+		OutHLSL.Appendf(TEXT("	Value = %s.Load(int4(IndexX, IndexY, IndexZ, 0)).%s;\n"), *GridNameHLSL, *AttrGridChannels);
+		OutHLSL.Appendf(TEXT("}\n"));
+	}
+
+	return true;
+}
+
+bool FGrid3DCollectionAttributeHelper::WriteSetAtIndexHLSL(AttributeRetrievalMode AttributeStorage, const FNiagaraDataInterfaceGeneratedFunction& FunctionInfo, int NumChannels, FString& OutHLSL)
+{
+	FString NumChannelsString;
+	FString AttrGridChannels;
+
+	const FString OutputGridNameHLSL = ParamInfo.DataInterfaceHLSLSymbol + UNiagaraDataInterfaceGrid3DCollection::OutputGridName;
+
+	GetChannelStrings(0, NumChannels, NumChannelsString, AttrGridChannels);
+
+	OutHLSL.Appendf(TEXT("void %s(int IndexX, int IndexY, int IndexZ, int AttributeIndex, float%s Value)\n"), *FunctionInfo.InstanceName, *NumChannelsString);
+	OutHLSL.Appendf(TEXT("{\n"));
+
+	// more than 1 attribute, we must read first
+	if (AttributeInfos.Num() > 1)
+	{
+		//#todo(dmp): fill this in
+		return false;
+	}
+	else
+	{
+		OutHLSL.Appendf(TEXT("	%s[float3(IndexX, IndexY, IndexZ)].%s = Value;\n"), *OutputGridNameHLSL, *AttrGridChannels);
+	}
+	OutHLSL.Appendf(TEXT("}\n"));
+
+	return true;
+}
+
+bool FGrid3DCollectionAttributeHelper::WriteSampleAtIndexHLSL(AttributeRetrievalMode AttributeStorage, const FNiagaraDataInterfaceGeneratedFunction& FunctionInfo, int NumChannels, bool IsCubic, FString& OutHLSL)
+{
+	FString NumChannelsString;
+	FString AttrGridChannels;
+
+	const FString GridNameHLSL = ParamInfo.DataInterfaceHLSLSymbol + UNiagaraDataInterfaceGrid3DCollection::GridName;
+	const FString SamplerNameHLSL = ParamInfo.DataInterfaceHLSLSymbol + UNiagaraDataInterfaceGrid3DCollection::SamplerName;
+
+	GetChannelStrings(0, NumChannels, NumChannelsString, AttrGridChannels);
+
+	if (AttributeInfos.Num() > 1)
+	{
+		//#todo(dmp): fill this in
+		return false;
+	}
+	else
+	{
+		OutHLSL.Appendf(TEXT("void %s(float3 Unit, int AttributeIndex, out float%s Value)\n"), *FunctionInfo.InstanceName, *NumChannelsString);
+		OutHLSL.Appendf(TEXT("{\n"));
+		if (IsCubic)
+		{
+			OutHLSL.Appendf(TEXT("	Value = SampleTriCubicLagrange_%s(%s, Unit, 0).%s;\n"), *ParamInfo.DataInterfaceHLSLSymbol, *SamplerNameHLSL, *AttrGridChannels);
+		}
+		else
+		{
+			OutHLSL.Appendf(TEXT("	Value = %s.SampleLevel(%s, Unit, 0).%s;\n"), *GridNameHLSL, *SamplerNameHLSL, *AttrGridChannels);
+		}
+
+		OutHLSL.Appendf(TEXT("}\n"));
+	}
+
+	return true;
+}
+
+bool FGrid3DCollectionAttributeHelper::WriteSetHLSL(AttributeRetrievalMode AttributeStorage, const FNiagaraDataInterfaceGeneratedFunction& FunctionInfo, FString& OutHLSL)
+{
+	const FName* AttributeName = FunctionInfo.FindSpecifierValue(UNiagaraDataInterfaceGrid3DCollection::NAME_Attribute);
+	if (AttributeName == nullptr)
+	{
+		return false;
+	}
+
+	const FAttributeInfo* AttributeInfo = FindAttributeInfo(*AttributeName);
+	if (AttributeInfo == nullptr)
+	{
+		return false;
+	}
+
+	const FString OutputGridNameHLSL = ParamInfo.DataInterfaceHLSLSymbol + UNiagaraDataInterfaceGrid3DCollection::OutputGridName;
+
+	if (AttributeStorage == AttributeRetrievalMode::RGBAGrid)
+	{
+		FString NumChannelsString;
+		FString AttrGridChannels;
+
+		GetChannelStrings(AttributeInfo, NumChannelsString, AttrGridChannels);
+
+		OutHLSL.Appendf(TEXT("void %s(int IndexX, int IndexY, int IndexZ, float%s Value)\n"), *FunctionInfo.InstanceName, *NumChannelsString);
+		OutHLSL.Appendf(TEXT("{\n"));
+
+		// more than 1 attrbiute, we must read first
+		if (AttributeInfos.Num() > 1)
+		{
+			FString GridNumChannels = FString::FromInt(TotalChannels);
+
+			OutHLSL.Appendf(TEXT("	float%s TmpValue = %s.Load(int4(IndexX, IndexY, IndexZ, 0));\n"), *GridNumChannels, *OutputGridNameHLSL);
+			OutHLSL.Appendf(TEXT("	TmpValue.%s = Value;\n"), *AttrGridChannels);
+			OutHLSL.Appendf(TEXT("	%s[float3(IndexX, IndexY, IndexZ)] = TmpValue;\n"), *OutputGridNameHLSL);
+		}
+		else
+		{
+			OutHLSL.Appendf(TEXT("	%s[float3(IndexX, IndexY, IndexZ)].%s = Value;\n"), *OutputGridNameHLSL, *AttrGridChannels);
+		}
+		OutHLSL.Appendf(TEXT("}\n"));
+	}
+	else
+	{
+		if (AttributeInfo->NumChannels == 1)
+		{
+			OutHLSL.Appendf(TEXT("void %s(int IndexX, int IndexY, int IndexZ, float Value)\n"), *FunctionInfo.InstanceName);
+			OutHLSL.Appendf(TEXT("{\n"));
+			OutHLSL.Appendf(TEXT("	int AttributeIndex = %s;\n"), *GetAttributeIndex(AttributeStorage == AttributeRetrievalMode::Indirection, AttributeInfo));
+			OutHLSL.Appendf(TEXT("	int3 PixelOffset = int3(IndexX, IndexY, IndexZ) + %s;\n"), *GetPerAttributePixelOffset());
+			OutHLSL.Appendf(TEXT("	%s[PixelOffset] = Value;\n"), *OutputGridNameHLSL);
+			OutHLSL.Appendf(TEXT("}\n"));
+		}
+		else
+		{
+			OutHLSL.Appendf(TEXT("void %s(int IndexX, int IndexY, int IndexZ, float%d Value)\n"), *FunctionInfo.InstanceName, AttributeInfo->NumChannels);
+			OutHLSL.Appendf(TEXT("{\n"));
+			for (int32 i = 0; i < AttributeInfo->NumChannels; ++i)
+			{
+				OutHLSL.Appendf(TEXT("	{\n"));
+				OutHLSL.Appendf(TEXT("		int AttributeIndex = %s;\n"), *GetAttributeIndex(AttributeStorage == AttributeRetrievalMode::Indirection, AttributeInfo, i));
+				OutHLSL.Appendf(TEXT("		int3 PixelOffset = int3(IndexX, IndexY, IndexZ) + %s;\n"), *GetPerAttributePixelOffset());
+				OutHLSL.Appendf(TEXT("		%s[PixelOffset] = Value[%d];\n"), *OutputGridNameHLSL, i);
+				OutHLSL.Appendf(TEXT("	}\n"));
+			}
+			OutHLSL.Appendf(TEXT("}\n"));
+		}
+	}
+	return true;
+}
+
+bool FGrid3DCollectionAttributeHelper::WriteSampleHLSL(AttributeRetrievalMode AttributeStorage, const FNiagaraDataInterfaceGeneratedFunction& FunctionInfo, bool IsCubic, FString& OutHLSL)
+{
+	const FName* AttributeName = FunctionInfo.FindSpecifierValue(UNiagaraDataInterfaceGrid3DCollection::NAME_Attribute);
+	if (AttributeName == nullptr)
+	{
+		return false;
+	}
+
+	const FAttributeInfo* AttributeInfo = FindAttributeInfo(*AttributeName);
+	if (AttributeInfo == nullptr)
+	{
+		return false;
+	}
+
+	const FString GridNameHLSL = ParamInfo.DataInterfaceHLSLSymbol + UNiagaraDataInterfaceGrid3DCollection::GridName;
+	const FString SamplerNameHLSL = ParamInfo.DataInterfaceHLSLSymbol + UNiagaraDataInterfaceGrid3DCollection::SamplerName;
+
+	if (AttributeStorage == AttributeRetrievalMode::RGBAGrid)
+	{
+		FString NumChannelsString;
+		FString AttrGridChannels;
+		GetChannelStrings(AttributeInfo, NumChannelsString, AttrGridChannels);
+
+		OutHLSL.Appendf(TEXT("void %s(float3 Unit, out float%s Value)\n"), *FunctionInfo.InstanceName, *NumChannelsString);
+		OutHLSL.Appendf(TEXT("{\n"));
+		if (IsCubic)
+		{
+			OutHLSL.Appendf(TEXT("	Value = SampleTriCubicLagrange_%s(%s, Unit, 0).%s;\n"), *ParamInfo.DataInterfaceHLSLSymbol, *SamplerNameHLSL, *AttrGridChannels);
+		}
+		else
+		{
+			OutHLSL.Appendf(TEXT("	Value = %s.SampleLevel(%s, Unit, 0).%s;\n"), *GridNameHLSL, *SamplerNameHLSL, *AttrGridChannels);
+		}
+
+		OutHLSL.Appendf(TEXT("}\n"));
+	}
+	else
+	{
+		const FString NumAttributesNameHLSL = UNiagaraDataInterfaceRWBase::NumAttributesName + ParamInfo.DataInterfaceHLSLSymbol;
+		if (AttributeInfo->NumChannels == 1)
+		{
+			OutHLSL.Appendf(TEXT("void %s(float3 Unit, out float Value)\n"), *FunctionInfo.InstanceName);
+			OutHLSL.Appendf(TEXT("{\n"));
+			OutHLSL.Appendf(TEXT("	float3 TileUVW = clamp(Unit, %s%s, %s%s) * %s%s;\n"), *UNiagaraDataInterfaceGrid3DCollection::UnitClampMinName, *ParamInfo.DataInterfaceHLSLSymbol, *UNiagaraDataInterfaceGrid3DCollection::UnitClampMaxName, *ParamInfo.DataInterfaceHLSLSymbol, *UNiagaraDataInterfaceGrid3DCollection::OneOverNumTilesName, *ParamInfo.DataInterfaceHLSLSymbol);
+			OutHLSL.Appendf(TEXT("	int AttributeIndex = %s;\n"), *GetAttributeIndex(AttributeStorage == AttributeRetrievalMode::Indirection, AttributeInfo));
+			OutHLSL.Appendf(TEXT("	float3 UVW = TileUVW + %s;\n"), *GetPerAttributeUVWOffset());
+			if (IsCubic)
+			{
+				OutHLSL.Appendf(TEXT("	Value = SampleTriCubicLagrange(%s, Unit, 0).%s;\n"), *ParamInfo.DataInterfaceHLSLSymbol, *SamplerNameHLSL);
+			}
+			else
+			{
+				OutHLSL.Appendf(TEXT("	Value = %s.SampleLevel(%s, UVW, 0);\n"), *GridNameHLSL, *SamplerNameHLSL);
+			}
+			OutHLSL.Appendf(TEXT("}\n"));
+		}
+		else
+		{
+			OutHLSL.Appendf(TEXT("void %s(float3 Unit, out float%d Value)\n"), *FunctionInfo.InstanceName, AttributeInfo->NumChannels);
+			OutHLSL.Appendf(TEXT("{\n"));
+			OutHLSL.Appendf(TEXT("	float3 TileUVW = clamp(Unit, %s%s, %s%s) * %s%s;\n"), *UNiagaraDataInterfaceGrid3DCollection::UnitClampMinName, *ParamInfo.DataInterfaceHLSLSymbol, *UNiagaraDataInterfaceGrid3DCollection::UnitClampMaxName, *ParamInfo.DataInterfaceHLSLSymbol, *UNiagaraDataInterfaceGrid3DCollection::OneOverNumTilesName, *ParamInfo.DataInterfaceHLSLSymbol);
+			for (int32 i = 0; i < AttributeInfo->NumChannels; ++i)
+			{
+				OutHLSL.Appendf(TEXT("	{\n"));
+				OutHLSL.Appendf(TEXT("		int AttributeIndex = %s;\n"), *GetAttributeIndex(AttributeStorage == AttributeRetrievalMode::Indirection, AttributeInfo, i));
+				OutHLSL.Appendf(TEXT("		float3 UVW = TileUVW + %s;\n"), *GetPerAttributeUVWOffset());
+				if (IsCubic)
+				{
+					OutHLSL.Appendf(TEXT("	Value = SampleTriCubicLagrange(%s, Unit, 0).%s;\n"), *ParamInfo.DataInterfaceHLSLSymbol, *SamplerNameHLSL);
+				}
+				else
+				{
+					OutHLSL.Appendf(TEXT("Value[%d] = %s.SampleLevel(%s, UVW, 0);\n"), i, *GridNameHLSL, *SamplerNameHLSL);
+				}
+
+				OutHLSL.Appendf(TEXT("	}\n"));
+			}
+			OutHLSL.Appendf(TEXT("}\n"));
+		}
+	}
+	return true;
+}
+
+bool FGrid3DCollectionAttributeHelper::WriteAttributeGetIndexHLSL(bool bUseAttributeIndirection, const FNiagaraDataInterfaceGeneratedFunction& FunctionInfo, FString& OutHLSL)
+{
+	const FName* AttributeName = FunctionInfo.FindSpecifierValue(UNiagaraDataInterfaceGrid3DCollection::NAME_Attribute);
+	if (AttributeName == nullptr)
+	{
+		return false;
+	}
+
+	const FAttributeInfo* AttributeInfo = FindAttributeInfo(*AttributeName);
+	if (AttributeInfo == nullptr)
+	{
+		return false;
+	}
+
+
+	OutHLSL.Appendf(TEXT("void %s(out int Value)\n"), *FunctionInfo.InstanceName);
+	OutHLSL.Appendf(TEXT("{\n"));
+	OutHLSL.Appendf(TEXT("	Value = %s;\n"), *GetAttributeIndex(bUseAttributeIndirection, AttributeInfo));
+	OutHLSL.Appendf(TEXT("}\n"));
+
+	return true;
+}
+#endif //WITH_EDITORONLY_DATA
+
+#undef LOCTEXT_NAMESPACE
\ No newline at end of file
diff --git a/Engine/Plugins/FX/Niagara/Source/Niagara/Private/NiagaraDataInterfaceGrid3DCollectionUtils.h b/Engine/Plugins/FX/Niagara/Source/Niagara/Private/NiagaraDataInterfaceGrid3DCollectionUtils.h
new file mode 100644
--- /dev/null
+++ b/Engine/Plugins/FX/Niagara/Source/Niagara/Private/NiagaraDataInterfaceGrid3DCollectionUtils.h
@@ -0,0 +1,64 @@
+// Copyright Epic Games, Inc. All Rights Reserved.
+
+#pragma once
+
+#include "CoreMinimal.h"
+#include "NiagaraDataInterface.h"
+
+struct FGrid3DCollectionAttributeHelper
+{
+
+	// only support rgba textures when we have a single attribute that contains up to 4 channels
+	static bool ShouldUseRGBAGrid(const int TotalChannels, const int TotalNumAttributes);
+
+	static TArray<FString> Channels;
+
+	struct FAttributeInfo
+	{
+		FName					Name;
+		FNiagaraTypeDefinition	TypeDef;
+		int32					NumChannels = 0;
+		int32					ChannelOffset = 0;
+		int32					AttributeIndex = 0;
+	};
+
+	enum AttributeRetrievalMode
+	{
+		RGBAGrid = 0,
+		Indirection,
+		NoIndirection
+	};
+
+	explicit FGrid3DCollectionAttributeHelper(const FNiagaraDataInterfaceGPUParamInfo& InParamInfo, TArray<FText>* OutWarnings = nullptr);
+
+	bool UseRGBAGrid();
+
+	const FAttributeInfo* FindAttributeInfo(FName AttributeName) const;
+
+#if WITH_EDITORONLY_DATA
+	// Returns pixel offset for the channel
+	static FString GetPerAttributePixelOffset(const TCHAR* DataInterfaceHLSLSymbol);
+	FString GetPerAttributePixelOffset() const;
+	// Returns UVW offset for the channel
+	static FString GetPerAttributeUVWOffset(const TCHAR* DataInterfaceHLSLSymbol);
+	FString GetPerAttributeUVWOffset() const;
+
+	// Translates named attribute into actual attribute index
+	FString GetAttributeIndex(bool bUseAttributeIndirection, const FAttributeInfo* AttributeInfo, int Channel = 0) const;
+	FString GetGridChannelString();
+	void GetChannelStrings(int AttributeIndex, int AttributeNumChannels, FString& NumChannelsString, FString& AttrGridChannels);
+	void GetChannelStrings(const FAttributeInfo* AttributeInfo, FString& NumChannelsString, FString& AttrGridChannels);
+
+	bool WriteGetHLSL(AttributeRetrievalMode AttributeStorage, const FNiagaraDataInterfaceGeneratedFunction& FunctionInfo, FString& OutHLSL);
+	bool WriteGetAtIndexHLSL(AttributeRetrievalMode AttributeStorage, const FNiagaraDataInterfaceGeneratedFunction& FunctionInfo, int NumChannels, FString& OutHLSL);
+	bool WriteSetAtIndexHLSL(AttributeRetrievalMode AttributeStorage, const FNiagaraDataInterfaceGeneratedFunction& FunctionInfo, int NumChannels, FString& OutHLSL);
+	bool WriteSampleAtIndexHLSL(AttributeRetrievalMode AttributeStorage, const FNiagaraDataInterfaceGeneratedFunction& FunctionInfo, int NumChannels, bool IsCubic, FString& OutHLSL);
+	bool WriteSetHLSL(AttributeRetrievalMode AttributeStorage, const FNiagaraDataInterfaceGeneratedFunction& FunctionInfo, FString& OutHLSL);
+	bool WriteSampleHLSL(AttributeRetrievalMode AttributeStorage, const FNiagaraDataInterfaceGeneratedFunction& FunctionInfo, bool IsCubic, FString& OutHLSL);
+	bool WriteAttributeGetIndexHLSL(bool bUseAttributeIndirection, const FNiagaraDataInterfaceGeneratedFunction& FunctionInfo, FString& OutHLSL);
+#endif //WITH_EDITORONLY_DATA
+
+	const FNiagaraDataInterfaceGPUParamInfo& ParamInfo;
+	TArray<FAttributeInfo>						AttributeInfos;
+	int32										TotalChannels = 0;
+};
diff --git a/Engine/Plugins/FX/Niagara/Source/NiagaraEditor/Private/NiagaraHlslTranslator.cpp b/Engine/Plugins/FX/Niagara/Source/NiagaraEditor/Private/NiagaraHlslTranslator.cpp
--- a/Engine/Plugins/FX/Niagara/Source/NiagaraEditor/Private/NiagaraHlslTranslator.cpp
+++ b/Engine/Plugins/FX/Niagara/Source/NiagaraEditor/Private/NiagaraHlslTranslator.cpp
@@ -2150,7 +2150,7 @@ void FHlslNiagaraTranslator::HandleSimStageSetupAndTeardown(int32 InWhichStage,
 	// Now take a look at any of the variables that were actually written to / read from in this stage.
 	TArray<FNiagaraVariable> ReadVars;
 	TArray<FNiagaraVariable> WriteVars;
-
+	TArray<FNiagaraVariable> AllVars;
 	for (int32 ParamHistoryIdx = 0; ParamHistoryIdx < ParamMapHistories.Num(); ParamHistoryIdx++)
 	{
 		if (InWhichStage != ParamHistoryIdx && TranslationStage.ShouldDoSpawnOnlyLogic() == false)
@@ -2170,6 +2170,20 @@ void FHlslNiagaraTranslator::HandleSimStageSetupAndTeardown(int32 InWhichStage,
 		}
 	}
 
+	for (int32 ParamHistoryIdx = 0; ParamHistoryIdx < ParamMapHistories.Num(); ParamHistoryIdx++)
+	{		
+		for (int32 i = 0; i < ParamMapHistories[ParamHistoryIdx].Variables.Num(); i++)
+		{
+			const FNiagaraVariable& Var = ParamMapHistories[ParamHistoryIdx].Variables[i];
+
+			if (Var.IsInNameSpace(TranslationStage.IterationSource))
+			{
+				if (!AllVars.Contains(Var))
+					AllVars.Emplace(Var);
+			}
+		}
+	}
+
 	// Find the data interface in the table. Note that this may not be found because we don't actually call any functions on the data interface yet.
 	int32 DataInterfaceOwnerIndex = INDEX_NONE;
 	for (int32 i = 0; i < CompilationOutput.ScriptData.DataInterfaceInfo.Num(); i++)
@@ -2337,7 +2351,7 @@ void FHlslNiagaraTranslator::HandleSimStageSetupAndTeardown(int32 InWhichStage,
 				AttributeHLSLNames.Emplace(TEXT("Map.") + GetSanitizedSymbolName(Var.GetName().ToString()));
 			}
 
-			if (CDO->GenerateIterationSourceNamespaceWriteAttributesHLSL(DIInstanceInfo, IterationSourceVar, MakeArrayView(Sig.Inputs), MakeArrayView(WriteVars), MakeArrayView(AttributeHLSLNames), bSpawnOnly, bPartialWrites, GeneratedErrors, AttributeWriteGeneratedHLSL) && AttributeWriteGeneratedHLSL.Len() > 0)
+			if (CDO->GenerateIterationSourceNamespaceWriteAttributesHLSL(DIInstanceInfo, IterationSourceVar, MakeArrayView(Sig.Inputs), MakeArrayView(WriteVars), MakeArrayView(AttributeHLSLNames), MakeArrayView(AllVars), bSpawnOnly, bPartialWrites, GeneratedErrors, AttributeWriteGeneratedHLSL) && AttributeWriteGeneratedHLSL.Len() > 0)
 			{
 				Sig.Name = *FString::Printf(TEXT("TeardownFromIterationSource_%s_GeneratedWriteAttributes"), *GetSanitizedFunctionNameSuffix(TranslationStage.PassNamespace));
 				FNiagaraFunctionSignature SignatureOut = Sig;
