Improved naming of utility function

diff --git a/Engine/Plugins/FX/Niagara/Source/Niagara/Classes/NiagaraScript.h b/Engine/Plugins/FX/Niagara/Source/Niagara/Classes/NiagaraScript.h
--- a/Engine/Plugins/FX/Niagara/Source/Niagara/Classes/NiagaraScript.h
+++ b/Engine/Plugins/FX/Niagara/Source/Niagara/Classes/NiagaraScript.h
@@ -1132,10 +1132,10 @@ public:
 	/** External call used to identify the values for a successful VM script compilation. OnVMScriptCompiled will be issued in this case.*/
 	void SetVMCompilationResults(const FNiagaraVMExecutableDataId& InCompileId, FNiagaraVMExecutableData& InScriptVM, FString EmitterUniqueName, const TMap<FName, UNiagaraDataInterface*>& ObjectNameMap);
 
-	/** In the event where we "merge" we duplicate the changes of the master copy onto the newly cloned copy. This function will synchronize the compiled script 
+	/** In the event where we "merge" we duplicate the changes of the source script onto the newly cloned copy. This function will synchronize the compiled script 
 		results assuming that the scripts themselves are bound to the same key. This saves looking things up in the DDC. It returns true if successfully synchronized and 
 		false if not.*/
-	NIAGARA_API bool SynchronizeExecutablesWithMaster(const UNiagaraScript* Script, const TMap<FString, FString>& RenameMap);
+	NIAGARA_API bool SynchronizeExecutablesWithCompilation(const UNiagaraScript* Script, const TMap<FString, FString>& RenameMap);
 
 	NIAGARA_API FString GetFriendlyName() const;
 
diff --git a/Engine/Plugins/FX/Niagara/Source/Niagara/Private/NiagaraScript.cpp b/Engine/Plugins/FX/Niagara/Source/Niagara/Private/NiagaraScript.cpp
--- a/Engine/Plugins/FX/Niagara/Source/Niagara/Private/NiagaraScript.cpp
+++ b/Engine/Plugins/FX/Niagara/Source/Niagara/Private/NiagaraScript.cpp
@@ -3425,7 +3425,7 @@ void UNiagaraScript::SyncAliases(const FNiagaraAliasContext& ResolveAliasesConte
 	}
 }
 
-bool UNiagaraScript::SynchronizeExecutablesWithMaster(const UNiagaraScript* Script, const TMap<FString, FString>& RenameMap)
+bool UNiagaraScript::SynchronizeExecutablesWithCompilation(const UNiagaraScript* Script, const TMap<FString, FString>& RenameMap)
 {
 	check(IsInGameThread());
 
diff --git a/Engine/Plugins/FX/Niagara/Source/NiagaraEditor/Private/NiagaraScriptMergeManager.cpp b/Engine/Plugins/FX/Niagara/Source/NiagaraEditor/Private/NiagaraScriptMergeManager.cpp
--- a/Engine/Plugins/FX/Niagara/Source/NiagaraEditor/Private/NiagaraScriptMergeManager.cpp
+++ b/Engine/Plugins/FX/Niagara/Source/NiagaraEditor/Private/NiagaraScriptMergeManager.cpp
@@ -1355,7 +1355,7 @@ FNiagaraScriptMergeManager::FApplyDiffResults FNiagaraScriptMergeManager::ForceI
 				{
 					if (Script->Usage == OriginalScript->Usage && Script->GetUsageId() == OriginalScript->GetUsageId())
 					{
-						bAnyUpdated |= Script->SynchronizeExecutablesWithMaster(OriginalScript, RenameMap);
+						bAnyUpdated |= Script->SynchronizeExecutablesWithCompilation(OriginalScript, RenameMap);
 					}
 				}
 			}
