- Add support for random socket to static mesh DI

diff --git a/Engine/Plugins/FX/Niagara/Shaders/Private/NiagaraDataInterfaceStaticMeshTemplate.ush b/Engine/Plugins/FX/Niagara/Shaders/Private/NiagaraDataInterfaceStaticMeshTemplate.ush
--- a/Engine/Plugins/FX/Niagara/Shaders/Private/NiagaraDataInterfaceStaticMeshTemplate.ush
+++ b/Engine/Plugins/FX/Niagara/Shaders/Private/NiagaraDataInterfaceStaticMeshTemplate.ush
@@ -330,11 +330,41 @@ void GetFilteredSocketCount_{ParameterName}(out int Count)
 {
 	Count = {ParameterName}_SocketCounts.y;
 }
+
 void GetUnfilteredSocketCount_{ParameterName}(out int Count)
 {
 	Count = {ParameterName}_SocketCounts.z;
 }
 
+void RandomSocket_{ParameterName}(NiagaraRandInfo RandInfo, out int Index)
+{
+	Index = -1;
+	if ( {ParameterName}_SocketCounts.x > 0 )
+	{
+		Index = NiagaraRandomInt(RandInfo, {ParameterName}_SocketCounts.x);
+	}
+}
+
+void RandomFilteredSocket_{ParameterName}(NiagaraRandInfo RandInfo, out int Index)
+{
+	Index = -1;
+	if ( {ParameterName}_SocketCounts.y > 0 )
+	{
+		Index = NiagaraRandomInt(RandInfo, {ParameterName}_SocketCounts.y);
+		Index = {ParameterName}_FilteredAndUnfilteredSockets[Index];
+	}
+}
+
+void RandomUnfilteredSocket_{ParameterName}(NiagaraRandInfo RandInfo, out int Index)
+{
+	Index = -1;
+	if ( {ParameterName}_SocketCounts.z > 0 )
+	{
+		Index = {ParameterName}_SocketCounts.y + NiagaraRandomInt(RandInfo, {ParameterName}_SocketCounts.z);
+		Index = {ParameterName}_FilteredAndUnfilteredSockets[Index];
+	}
+}
+
 void GetSocketTransform_{ParameterName}(int Socket, out float3 Position, out float4 Rotation, out float3 Scale)
 {
 	[branch]
diff --git a/Engine/Plugins/FX/Niagara/Source/Niagara/Private/DataInterface/NiagaraDataInterfaceStaticMesh.cpp b/Engine/Plugins/FX/Niagara/Source/Niagara/Private/DataInterface/NiagaraDataInterfaceStaticMesh.cpp
--- a/Engine/Plugins/FX/Niagara/Source/Niagara/Private/DataInterface/NiagaraDataInterfaceStaticMesh.cpp
+++ b/Engine/Plugins/FX/Niagara/Source/Niagara/Private/DataInterface/NiagaraDataInterfaceStaticMesh.cpp
@@ -154,6 +154,9 @@ namespace NDIStaticMeshLocal
 	static const FName	GetSocketCountName("GetSocketCount");
 	static const FName	GetFilteredSocketCountName("GetFilteredSocketCount");
 	static const FName	GetUnfilteredSocketCountName("GetUnfilteredSocketCount");
+	static const FName	RandomSocketName("RandomSocket");
+	static const FName	RandomFilteredSocketName("RandomFilteredSocket");
+	static const FName	RandomUnfilteredSocketName("RandomUnfilteredSocket");
 	static const FName	GetSocketTransformName("GetSocketTransform");
 	static const FName	GetSocketTransformWSName("GetSocketTransformWS");
 	static const FName	GetFilteredSocketTransformName("GetFilteredSocketTransform");
@@ -1881,6 +1884,15 @@ void UNiagaraDataInterfaceStaticMesh::GetFunctions(TArray<FNiagaraFunctionSignat
 		OutFunctions.Add_GetRef(Sig).Name = GetFilteredSocketCountName;
 		OutFunctions.Add_GetRef(Sig).Name = GetUnfilteredSocketCountName;
 	}
+	{
+		FNiagaraFunctionSignature Sig = BaseSignature;
+		Sig.Inputs.Emplace(FNiagaraTypeDefinition::GetRandInfoDef(), TEXT("RandomInfo"));
+		Sig.Outputs.Add(FNiagaraVariable(FNiagaraTypeDefinition::GetIntDef(), TEXT("Index")));
+
+		OutFunctions.Add_GetRef(Sig).Name = RandomSocketName;
+		OutFunctions.Add_GetRef(Sig).Name = RandomFilteredSocketName;
+		OutFunctions.Add_GetRef(Sig).Name = RandomUnfilteredSocketName;
+	}
 	{
 		FNiagaraFunctionSignature Sig = BaseSignature;
 		Sig.Inputs.Add(FNiagaraVariable(FNiagaraTypeDefinition::GetIntDef(), TEXT("Index")));
@@ -2356,6 +2368,18 @@ void UNiagaraDataInterfaceStaticMesh::GetVMExternalFunction(const FVMExternalFun
 	{
 		OutFunc = FVMExternalFunction::CreateLambda([this](FVectorVMExternalFunctionContext& Context) { VMGetUnfilteredSocketCount(Context); });
 	}
+	else if (BindingInfo.Name == RandomSocketName)
+	{
+		OutFunc = FVMExternalFunction::CreateUObject(this, &UNiagaraDataInterfaceStaticMesh::VMRandomSocket);
+	}
+	else if (BindingInfo.Name == RandomFilteredSocketName)
+	{
+		OutFunc = FVMExternalFunction::CreateUObject(this, &UNiagaraDataInterfaceStaticMesh::VMRandomFilteredSocket);
+	}
+	else if (BindingInfo.Name == RandomUnfilteredSocketName)
+	{
+		OutFunc = FVMExternalFunction::CreateUObject(this, &UNiagaraDataInterfaceStaticMesh::VMRandomUnfilteredSocket);
+	}
 	else if (BindingInfo.Name == GetSocketTransformName)
 	{
 		OutFunc = FVMExternalFunction::CreateLambda([this](FVectorVMExternalFunctionContext& Context) { VMGetSocketTransform<FNDITransformHandlerNoop>(Context); });
@@ -3684,6 +3708,70 @@ void UNiagaraDataInterfaceStaticMesh::VMGetUnfilteredSocketCount(FVectorVMExtern
 	}
 }
 
+void UNiagaraDataInterfaceStaticMesh::VMRandomSocket(FVectorVMExternalFunctionContext& Context)
+{
+	VectorVM::FUserPtrHandler<NDIStaticMeshLocal::FInstanceData_GameThread> InstanceData(Context);
+	FNDIRandomHelper RandHelper(Context);
+	FNDIOutputParam<int32> OutSocket(Context);
+
+	const int32 SocketRandMax = InstanceData->CachedSockets.Num() - 1;
+	for (int32 i = 0; i < Context.GetNumInstances(); ++i)
+	{
+		RandHelper.GetAndAdvance();
+		OutSocket.SetAndAdvance(SocketRandMax >= 0 ? RandHelper.RandRange(i, 0, SocketRandMax) : -1);
+	}
+}
+
+void UNiagaraDataInterfaceStaticMesh::VMRandomFilteredSocket(FVectorVMExternalFunctionContext& Context)
+{
+	VectorVM::FUserPtrHandler<NDIStaticMeshLocal::FInstanceData_GameThread> InstanceData(Context);
+	FNDIRandomHelper RandHelper(Context);
+	FNDIOutputParam<int32> OutSocket(Context);
+
+	const int32 SocketRandMax = InstanceData->NumFilteredSockets - 1;
+	if (SocketRandMax >= 0)
+	{
+		for (int32 i = 0; i < Context.GetNumInstances(); ++i)
+		{
+			RandHelper.GetAndAdvance();
+			const int32 SocketIndex = RandHelper.RandRange(i, 0, SocketRandMax);
+			OutSocket.SetAndAdvance(InstanceData->FilteredAndUnfilteredSockets[SocketIndex]);
+		}
+	}
+	else
+	{
+		for (int32 i = 0; i < Context.GetNumInstances(); ++i)
+		{
+			OutSocket.SetAndAdvance(-1);
+		}
+	}
+}
+
+void UNiagaraDataInterfaceStaticMesh::VMRandomUnfilteredSocket(FVectorVMExternalFunctionContext& Context)
+{
+	VectorVM::FUserPtrHandler<NDIStaticMeshLocal::FInstanceData_GameThread> InstanceData(Context);
+	FNDIRandomHelper RandHelper(Context);
+	FNDIOutputParam<int32> OutSocket(Context);
+
+	const int32 SocketRandMax = InstanceData->FilteredAndUnfilteredSockets.Num() - InstanceData->NumFilteredSockets - 1;
+	if (SocketRandMax >= 0)
+	{
+		for (int32 i = 0; i < Context.GetNumInstances(); ++i)
+		{
+			RandHelper.GetAndAdvance();
+			const int32 SocketIndex = InstanceData->NumFilteredSockets + RandHelper.RandRange(i, 0, SocketRandMax);
+			OutSocket.SetAndAdvance(InstanceData->FilteredAndUnfilteredSockets[SocketIndex]);
+		}
+	}
+	else
+	{
+		for (int32 i = 0; i < Context.GetNumInstances(); ++i)
+		{
+			OutSocket.SetAndAdvance(-1);
+		}
+	}
+}
+
 template<typename TTransformHandler>
 void UNiagaraDataInterfaceStaticMesh::VMGetSocketTransform(FVectorVMExternalFunctionContext& Context)
 {
diff --git a/Engine/Plugins/FX/Niagara/Source/Niagara/Private/DataInterface/NiagaraDataInterfaceStaticMesh.h b/Engine/Plugins/FX/Niagara/Source/Niagara/Private/DataInterface/NiagaraDataInterfaceStaticMesh.h
--- a/Engine/Plugins/FX/Niagara/Source/Niagara/Private/DataInterface/NiagaraDataInterfaceStaticMesh.h
+++ b/Engine/Plugins/FX/Niagara/Source/Niagara/Private/DataInterface/NiagaraDataInterfaceStaticMesh.h
@@ -200,6 +200,9 @@ protected:
 	void VMGetSocketCount(FVectorVMExternalFunctionContext& Context);
 	void VMGetFilteredSocketCount(FVectorVMExternalFunctionContext& Context);
 	void VMGetUnfilteredSocketCount(FVectorVMExternalFunctionContext& Context);
+	void VMRandomSocket(FVectorVMExternalFunctionContext& Context);
+	void VMRandomFilteredSocket(FVectorVMExternalFunctionContext& Context);
+	void VMRandomUnfilteredSocket(FVectorVMExternalFunctionContext& Context);
 	template<typename TTransformHandler>
 	void VMGetSocketTransform(FVectorVMExternalFunctionContext& Context);
 	template<typename TTransformHandler>
