Fixed incorrect order of condition for cull proxy registration.

diff --git a/Engine/Plugins/FX/Niagara/Source/Niagara/Private/NiagaraCullProxyComponent.cpp b/Engine/Plugins/FX/Niagara/Source/Niagara/Private/NiagaraCullProxyComponent.cpp
--- a/Engine/Plugins/FX/Niagara/Source/Niagara/Private/NiagaraCullProxyComponent.cpp
+++ b/Engine/Plugins/FX/Niagara/Source/Niagara/Private/NiagaraCullProxyComponent.cpp
@@ -66,7 +66,7 @@ bool UNiagaraCullProxyComponent::RegisterCulledComponent(UNiagaraComponent* Comp
 	}
 #endif
 
-	if (bForce || GetAsset()->GetScalabilitySettings().MaxSystemProxies < Instances.Num())
+	if (bForce || Instances.Num() < GetAsset()->GetScalabilitySettings().MaxSystemProxies)
 	{
 		FNiagaraCulledComponentInfo& NewInfo = Instances.AddDefaulted_GetRef();
 		NewInfo.WeakComponent = Component;
