Unshelved from pending changelist '21928365':

Fixing crash when deleting several parameters at once.

diff --git a/Engine/Plugins/FX/Niagara/Source/NiagaraEditor/Private/ViewModels/Stack/NiagaraStackModuleItem.cpp b/Engine/Plugins/FX/Niagara/Source/NiagaraEditor/Private/ViewModels/Stack/NiagaraStackModuleItem.cpp
--- a/Engine/Plugins/FX/Niagara/Source/NiagaraEditor/Private/ViewModels/Stack/NiagaraStackModuleItem.cpp
+++ b/Engine/Plugins/FX/Niagara/Source/NiagaraEditor/Private/ViewModels/Stack/NiagaraStackModuleItem.cpp
@@ -223,7 +223,7 @@ const UNiagaraStackModuleItem::FCollectedUsageData& UNiagaraStackModuleItem::Get
 	if (CachedCollectedUsageData.IsSet() == false)
 	{
 		CachedCollectedUsageData = FCollectedUsageData();
-		if (FunctionCallNode && FunctionCallNode->IsA<UNiagaraNodeAssignment>())
+		if (FunctionCallNode && FunctionCallNode->IsA<UNiagaraNodeAssignment>() && IsSystemViewModelValid())
 		{
 			UNiagaraNodeAssignment* AssignmentNode = CastChecked<UNiagaraNodeAssignment>(FunctionCallNode);
 			TSharedRef<FNiagaraSystemViewModel> SystemVM = GetSystemViewModel();
diff --git a/Engine/Plugins/FX/Niagara/Source/NiagaraEditor/Public/ViewModels/Stack/NiagaraStackEntry.h b/Engine/Plugins/FX/Niagara/Source/NiagaraEditor/Public/ViewModels/Stack/NiagaraStackEntry.h
--- a/Engine/Plugins/FX/Niagara/Source/NiagaraEditor/Public/ViewModels/Stack/NiagaraStackEntry.h
+++ b/Engine/Plugins/FX/Niagara/Source/NiagaraEditor/Public/ViewModels/Stack/NiagaraStackEntry.h
@@ -568,6 +568,8 @@ protected:
 
 	mutable TOptional<FCollectedUsageData> CachedCollectedUsageData;
 
+	bool IsSystemViewModelValid() const { return SystemViewModel.IsValid(); }
+
 
 private:
 	void ChildStructureChanged(ENiagaraStructureChangedFlags Info);
