- Cleanup some header includes

#rnx

diff --git a/Engine/Plugins/FX/Niagara/Source/Niagara/Private/NiagaraDataInterfaceRigidMeshCollisionQuery.cpp b/Engine/Plugins/FX/Niagara/Source/Niagara/Private/NiagaraDataInterfaceRigidMeshCollisionQuery.cpp
--- a/Engine/Plugins/FX/Niagara/Source/Niagara/Private/NiagaraDataInterfaceRigidMeshCollisionQuery.cpp
+++ b/Engine/Plugins/FX/Niagara/Source/Niagara/Private/NiagaraDataInterfaceRigidMeshCollisionQuery.cpp
@@ -4,7 +4,6 @@
 #include "Algo/ForEach.h"
 #include "Components/SkeletalMeshComponent.h"
 #include "Engine/Canvas.h"
-#include "Renderer/Private/DistanceFieldLightingShared.h"
 #include "Animation/SkeletalMeshActor.h"
 #include "AnimationRuntime.h"
 #include "SkeletalRenderPublic.h"
diff --git a/Engine/Plugins/FX/Niagara/Source/NiagaraShader/Private/NiagaraAsyncGpuTraceProviderGsdf.cpp b/Engine/Plugins/FX/Niagara/Source/NiagaraShader/Private/NiagaraAsyncGpuTraceProviderGsdf.cpp
--- a/Engine/Plugins/FX/Niagara/Source/NiagaraShader/Private/NiagaraAsyncGpuTraceProviderGsdf.cpp
+++ b/Engine/Plugins/FX/Niagara/Source/NiagaraShader/Private/NiagaraAsyncGpuTraceProviderGsdf.cpp
@@ -6,8 +6,6 @@
 #include "NiagaraGpuComputeDispatchInterface.h"
 #include "NiagaraDistanceFieldHelper.h"
 #include "NiagaraSettings.h"
-#include "Renderer/Private/ScenePrivate.h"
-#include "Renderer/Private/SceneRendering.h"
 #include "SceneManagement.h"
 
 static int GNiagaraAsyncGpuTraceGsdfEnabled = 1;
diff --git a/Engine/Plugins/FX/Niagara/Source/NiagaraShader/Private/NiagaraClearCounts.cpp b/Engine/Plugins/FX/Niagara/Source/NiagaraShader/Private/NiagaraClearCounts.cpp
--- a/Engine/Plugins/FX/Niagara/Source/NiagaraShader/Private/NiagaraClearCounts.cpp
+++ b/Engine/Plugins/FX/Niagara/Source/NiagaraShader/Private/NiagaraClearCounts.cpp
@@ -8,7 +8,7 @@
 #include "Modules/ModuleManager.h"
 
 #include "RenderGraphBuilder.h"
-#include "Runtime/Renderer/Private/ScreenPass.h"
+#include "RenderGraphUtils.h"
 
 //////////////////////////////////////////////////////////////////////////
 
diff --git a/Engine/Plugins/FX/Niagara/Source/NiagaraShader/Private/NiagaraDebugShaders.cpp b/Engine/Plugins/FX/Niagara/Source/NiagaraShader/Private/NiagaraDebugShaders.cpp
--- a/Engine/Plugins/FX/Niagara/Source/NiagaraShader/Private/NiagaraDebugShaders.cpp
+++ b/Engine/Plugins/FX/Niagara/Source/NiagaraShader/Private/NiagaraDebugShaders.cpp
@@ -8,9 +8,11 @@
 #include "Modules/ModuleManager.h"
 
 #include "RenderGraphBuilder.h"
-#include "Runtime/Renderer/Private/ScreenPass.h"
+#include "RenderGraphUtils.h"
 #include "PipelineStateCache.h"
 
+#include "Runtime/Renderer/Private/ScreenPass.h"
+
 int GNiagaraGpuComputeDebug_ShowNaNInf = 1;
 static FAutoConsoleVariableRef CVarNiagaraGpuComputeDebug_ShowNaNInf(
 	TEXT("fx.Niagara.GpuComputeDebug.ShowNaNInf"),
@@ -196,20 +198,20 @@ void NiagaraDebugShaders::ClearUAV(FRDGBuilder& GraphBuilder, FRDGBufferUAVRef U
 }
 
 void NiagaraDebugShaders::DrawDebugLines(
-	class FRDGBuilder& GraphBuilder, const class FViewInfo& View, FRDGTextureRef SceneColor, FRDGTextureRef SceneDepth,
+	class FRDGBuilder& GraphBuilder, const class FSceneView& View, FRDGTextureRef SceneColor, FRDGTextureRef SceneDepth,
 	const uint32 LineInstanceCount, FRDGBufferRef LineBuffer
 )
 {
-	TShaderMapRef<FNiagaraDebugDrawLineVS> VertexShader(View.ShaderMap);
-	TShaderMapRef<FNiagaraDebugDrawLinePS> PixelShader(View.ShaderMap);
+	TShaderMapRef<FNiagaraDebugDrawLineVS> VertexShader(GetGlobalShaderMap(View.GetFeatureLevel()));
+	TShaderMapRef<FNiagaraDebugDrawLinePS> PixelShader(GetGlobalShaderMap(View.GetFeatureLevel()));
 
 	FNiagaraDebugDrawLineParameters* PassParameters = GraphBuilder.AllocParameters<FNiagaraDebugDrawLineParameters>();
 	PassParameters->VSParameters.View					= View.ViewUniformBuffer;
 	PassParameters->VSParameters.GpuLineBuffer			= GraphBuilder.CreateSRV(LineBuffer, PF_R32_UINT);
 
 	PassParameters->PSParameters.OutputInvResolution	= FVector2f(1.0f / View.UnscaledViewRect.Width(), 1.0f / View.UnconstrainedViewRect.Height());
-	PassParameters->PSParameters.OriginalViewRectMin	= FVector2f(View.ViewRect.Min);
-	PassParameters->PSParameters.OriginalViewSize		= FVector2f(View.ViewRect.Width(), View.ViewRect.Height());
+	PassParameters->PSParameters.OriginalViewRectMin	= FVector2f(View.UnscaledViewRect.Min);
+	PassParameters->PSParameters.OriginalViewSize		= FVector2f(View.UnscaledViewRect.Width(), View.UnscaledViewRect.Height());
 	PassParameters->PSParameters.OriginalBufferInvSize	= FVector2f(1.f / SceneDepth->Desc.Extent.X, 1.f / SceneDepth->Desc.Extent.Y);
 	PassParameters->PSParameters.OccludedColorScale		= GNiagaraGpuComputeDebug_OccludedLineColorScale;
 	PassParameters->PSParameters.DepthTexture			= SceneDepth;
@@ -241,10 +243,10 @@ void NiagaraDebugShaders::DrawDebugLines(
 	);
 }
 
-void NiagaraDebugShaders::DrawDebugLines(FRDGBuilder& GraphBuilder, const FViewInfo& View, FRDGTextureRef SceneColor, FRDGTextureRef SceneDepth, FRDGBufferRef ArgsBuffer, FRDGBufferRef LineBuffer)
+void NiagaraDebugShaders::DrawDebugLines(FRDGBuilder& GraphBuilder, const FSceneView& View, FRDGTextureRef SceneColor, FRDGTextureRef SceneDepth, FRDGBufferRef ArgsBuffer, FRDGBufferRef LineBuffer)
 {
-	TShaderMapRef<FNiagaraDebugDrawLineVS> VertexShader(View.ShaderMap);
-	TShaderMapRef<FNiagaraDebugDrawLinePS> PixelShader(View.ShaderMap);
+	TShaderMapRef<FNiagaraDebugDrawLineVS> VertexShader(GetGlobalShaderMap(View.GetFeatureLevel()));
+	TShaderMapRef<FNiagaraDebugDrawLinePS> PixelShader(GetGlobalShaderMap(View.GetFeatureLevel()));
 
 	FNiagaraDebugDrawLineIndirectParameters* PassParameters = GraphBuilder.AllocParameters<FNiagaraDebugDrawLineIndirectParameters>();
 	PassParameters->AccessIndirectDrawArgsBuffer		= ArgsBuffer;
@@ -253,8 +255,8 @@ void NiagaraDebugShaders::DrawDebugLines(FRDGBuilder& GraphBuilder, const FViewI
 	PassParameters->VSParameters.GpuLineBuffer			= GraphBuilder.CreateSRV(LineBuffer, PF_R32_UINT);
 
 	PassParameters->PSParameters.OutputInvResolution	= FVector2f(1.0f / View.UnscaledViewRect.Width(), 1.0f / View.UnconstrainedViewRect.Height());
-	PassParameters->PSParameters.OriginalViewRectMin	= FVector2f(View.ViewRect.Min);
-	PassParameters->PSParameters.OriginalViewSize		= FVector2f(View.ViewRect.Width(), View.ViewRect.Height());
+	PassParameters->PSParameters.OriginalViewRectMin	= FVector2f(View.UnscaledViewRect.Min);
+	PassParameters->PSParameters.OriginalViewSize		= FVector2f(View.UnscaledViewRect.Width(), View.UnscaledViewRect.Height());
 	PassParameters->PSParameters.OriginalBufferInvSize	= FVector2f(1.f / SceneDepth->Desc.Extent.X, 1.f / SceneDepth->Desc.Extent.Y);
 	PassParameters->PSParameters.OccludedColorScale		= GNiagaraGpuComputeDebug_OccludedLineColorScale;
 	PassParameters->PSParameters.DepthTexture			= SceneDepth;
@@ -287,7 +289,7 @@ void NiagaraDebugShaders::DrawDebugLines(FRDGBuilder& GraphBuilder, const FViewI
 }
 
 void NiagaraDebugShaders::VisualizeTexture(
-	class FRDGBuilder& GraphBuilder, const FViewInfo& View, const FScreenPassRenderTarget& Output,
+	class FRDGBuilder& GraphBuilder, const FSceneView& View, const FScreenPassRenderTarget& Output,
 	const FIntPoint& Location, const int32& DisplayHeight,
 	const FIntVector4& InAttributesToVisualize, FRDGTextureRef Texture, const FIntVector4& NumTextureAttributes, uint32 TickCounter,
 	const FVector2D& PreviewDisplayRange
diff --git a/Engine/Plugins/FX/Niagara/Source/NiagaraShader/Public/NiagaraDebugShaders.h b/Engine/Plugins/FX/Niagara/Source/NiagaraShader/Public/NiagaraDebugShaders.h
--- a/Engine/Plugins/FX/Niagara/Source/NiagaraShader/Public/NiagaraDebugShaders.h
+++ b/Engine/Plugins/FX/Niagara/Source/NiagaraShader/Public/NiagaraDebugShaders.h
@@ -12,17 +12,17 @@ namespace NiagaraDebugShaders
 	NIAGARASHADER_API void ClearUAV(FRDGBuilder& GraphBuilder, FRDGBufferUAVRef UAV, FUintVector4 ClearValues, uint32 UIntsToSet);
 
 	NIAGARASHADER_API void DrawDebugLines(
-		class FRDGBuilder& GraphBuilder, const class FViewInfo& View, FRDGTextureRef SceneColor, FRDGTextureRef SceneDepth,
+		class FRDGBuilder& GraphBuilder, const class FSceneView& View, FRDGTextureRef SceneColor, FRDGTextureRef SceneDepth,
 		const uint32 LineInstanceCount, FRDGBufferRef LineBuffer
 	);
 
 	NIAGARASHADER_API void DrawDebugLines(
-		class FRDGBuilder& GraphBuilder, const class FViewInfo& View, FRDGTextureRef SceneColor, FRDGTextureRef SceneDepth,
+		class FRDGBuilder& GraphBuilder, const class FSceneView& View, FRDGTextureRef SceneColor, FRDGTextureRef SceneDepth,
 		FRDGBufferRef ArgsBuffer, FRDGBufferRef LineBuffer
 	);
 
 	NIAGARASHADER_API void VisualizeTexture(
-		class FRDGBuilder& GraphBuilder, const FViewInfo& View, const FScreenPassRenderTarget& Output,
+		class FRDGBuilder& GraphBuilder, const FSceneView& View, const FScreenPassRenderTarget& Output,
 		const FIntPoint& Location, const int32& DisplayHeight,
 		const FIntVector4& AttributesToVisualize, FRDGTextureRef Texture, const FIntVector4& NumTextureAttributes, uint32 TickCounter,
 		const FVector2D& PreviewDisplayRange = FVector2D::ZeroVector
