- Don't read the SimCache if time is the same
- Optimize volume updates, roughly 2x faster (i.e. 10ms -> 5ms)

#rnx

diff --git a/Engine/Plugins/FX/Niagara/Source/Niagara/Private/NiagaraDataInterfaceRenderTargetVolume.cpp b/Engine/Plugins/FX/Niagara/Source/Niagara/Private/NiagaraDataInterfaceRenderTargetVolume.cpp
--- a/Engine/Plugins/FX/Niagara/Source/Niagara/Private/NiagaraDataInterfaceRenderTargetVolume.cpp
+++ b/Engine/Plugins/FX/Niagara/Source/Niagara/Private/NiagaraDataInterfaceRenderTargetVolume.cpp
@@ -657,33 +657,43 @@ bool UNiagaraDataInterfaceRenderTargetVolume::SimCacheReadFrame(UObject* Storage
 				{
 					FUpdateTextureRegion3D UpdateRegion(FIntVector::ZeroValue, FIntVector::ZeroValue, InstanceData_RT->Size);
 
-					if (RT_CacheFrame->CompressedSize > 0)
+					FUpdateTexture3DData UpdateTexture = RHICmdList.BeginUpdateTexture3D(InstanceData_RT->RenderTarget->GetRHI(), 0, UpdateRegion);
+
+					const int32 SrcRowPitch = InstanceData_RT->Size.X * sizeof(FFloat16Color);
+					const int32 SrcDepthPitch = InstanceData_RT->Size.Y * SrcRowPitch;
+					if (UpdateTexture.RowPitch == SrcRowPitch && UpdateTexture.DepthPitch == SrcDepthPitch)
 					{
-						TArray<uint8> Decompressed;
-						Decompressed.AddUninitialized(sizeof(FFloat16Color) * InstanceData_RT->Size.X * InstanceData_RT->Size.Y * InstanceData_RT->Size.Z);
-
-						FCompression::UncompressMemory(RT_CompressionType, Decompressed.GetData(), Decompressed.Num(), RT_CacheFrame->GetPixelData(), RT_CacheFrame->CompressedSize);
-
-						RHICmdList.UpdateTexture3D(
-							InstanceData_RT->RenderTarget->GetRHI(),
-							0,
-							UpdateRegion,
-							InstanceData_RT->Size.X * sizeof(FFloat16Color),
-							InstanceData_RT->Size.X * InstanceData_RT->Size.Y * sizeof(FFloat16Color),
-							Decompressed.GetData()
-						);
+						if (RT_CacheFrame->CompressedSize > 0)
+						{
+							FCompression::UncompressMemory(RT_CompressionType, UpdateTexture.Data, UpdateTexture.DataSizeBytes, RT_CacheFrame->GetPixelData(), RT_CacheFrame->CompressedSize);
+						}
+						else
+						{
+							FMemory::Memcpy(UpdateTexture.Data, RT_CacheFrame->GetPixelData(), InstanceData_RT->Size.X * InstanceData_RT->Size.Y * sizeof(FFloat16Color));
+						}
 					}
 					else
 					{
-						RHICmdList.UpdateTexture3D(
-							InstanceData_RT->RenderTarget->GetRHI(),
-							0,
-							UpdateRegion,
-							InstanceData_RT->Size.X * sizeof(FFloat16Color),
-							InstanceData_RT->Size.X * InstanceData_RT->Size.Y * sizeof(FFloat16Color),
-							RT_CacheFrame->GetPixelData()
-						);
+						TArray<uint8> Decompressed;
+						if (RT_CacheFrame->CompressedSize > 0)
+						{
+							Decompressed.AddUninitialized(sizeof(FFloat16Color) * InstanceData_RT->Size.X * InstanceData_RT->Size.Y * InstanceData_RT->Size.Z);
+							FCompression::UncompressMemory(RT_CompressionType, Decompressed.GetData(), Decompressed.Num(), RT_CacheFrame->GetPixelData(), RT_CacheFrame->CompressedSize);
+						}
+
+						const uint8* SrcData = Decompressed.Num() > 0 ? Decompressed.GetData() : RT_CacheFrame->GetPixelData();
+						for (int32 z = 0; z < InstanceData_RT->Size.Z; ++z)
+						{
+							uint8* DstData = UpdateTexture.Data + (z * UpdateTexture.DepthPitch);
+							for (int32 y = 0; y < InstanceData_RT->Size.Y; ++y)
+							{
+								FMemory::Memcpy(DstData, SrcData, SrcRowPitch);
+								SrcData += SrcRowPitch;
+								DstData += UpdateTexture.RowPitch;
+							}
+						}
 					}
+					RHICmdList.EndUpdateTexture3D(UpdateTexture);
 				}
 			}
 		);
diff --git a/Engine/Plugins/FX/Niagara/Source/Niagara/Private/NiagaraSystemInstance.cpp b/Engine/Plugins/FX/Niagara/Source/Niagara/Private/NiagaraSystemInstance.cpp
--- a/Engine/Plugins/FX/Niagara/Source/Niagara/Private/NiagaraSystemInstance.cpp
+++ b/Engine/Plugins/FX/Niagara/Source/Niagara/Private/NiagaraSystemInstance.cpp
@@ -2180,6 +2180,12 @@ void FNiagaraSystemInstance::SimCacheTick_GameThread(UNiagaraSimCache* SimCache,
 		return;
 	}
 
+	// Early out if our age and desired age are virtually equal as there's nothing to do
+	if (FMath::IsNearlyEqual(Age, DesiredAge) )
+	{
+		return;
+	}
+
 	CachedDeltaSeconds = DeltaSeconds;
 	//FixedBounds_CNC = FixedBounds_GT;
 
