Niagara - In the system view model, don't try to update the emitter selection from sequencer selection changes if we're resetting the sequencer tracks.  Prevents a crash when adding emitters when there are multiple emitters selected.

diff --git a/Engine/Plugins/FX/Niagara/Source/NiagaraEditor/Private/ViewModels/NiagaraSystemViewModel.cpp b/Engine/Plugins/FX/Niagara/Source/NiagaraEditor/Private/ViewModels/NiagaraSystemViewModel.cpp
--- a/Engine/Plugins/FX/Niagara/Source/NiagaraEditor/Private/ViewModels/NiagaraSystemViewModel.cpp
+++ b/Engine/Plugins/FX/Niagara/Source/NiagaraEditor/Private/ViewModels/NiagaraSystemViewModel.cpp
@@ -88,6 +88,7 @@ FNiagaraSystemViewModel::FNiagaraSystemViewModel()
 	, bUpdatingSequencerFromEmitterDataChange(false)
 	, bUpdatingSystemSelectionFromSequencer(false)
 	, bUpdatingSequencerSelectionFromSystem(false)
+	, bResetingSequencerTracks(false)
 	, EditorSettings(GetMutableDefault<UNiagaraEditorSettings>())
 	, bResetRequestPending(false)
 	, bCompilePendingCompletion(false)
@@ -1686,6 +1687,7 @@ void FNiagaraSystemViewModel::ResetEmitterHandleViewModelsAndTracks()
 
 	if (NiagaraSequence)
 	{
+		TGuardValue<bool> UpdateGuard(bResetingSequencerTracks, true);
 		TArray<UMovieSceneTrack*> MasterTracks = NiagaraSequence->GetMovieScene()->GetMasterTracks();
 		for (UMovieSceneTrack* MasterTrack : MasterTracks)
 		{
@@ -2462,7 +2464,7 @@ void FNiagaraSystemViewModel::SystemSelectionChanged()
 
 void FNiagaraSystemViewModel::SequencerTrackSelectionChanged(TArray<UMovieSceneTrack*> SelectedTracks)
 {
-	if (bUpdatingSequencerSelectionFromSystem == false)
+	if (bUpdatingSequencerSelectionFromSystem == false && bResetingSequencerTracks == false)
 	{
 		UpdateEmitterHandleSelectionFromSequencer();
 	}
@@ -2470,7 +2472,7 @@ void FNiagaraSystemViewModel::SequencerTrackSelectionChanged(TArray<UMovieSceneT
 
 void FNiagaraSystemViewModel::SequencerSectionSelectionChanged(TArray<UMovieSceneSection*> SelectedSections)
 {
-	if (bUpdatingSequencerSelectionFromSystem == false)
+	if (bUpdatingSequencerSelectionFromSystem == false && bResetingSequencerTracks == false)
 	{
 		UpdateEmitterHandleSelectionFromSequencer();
 	}
diff --git a/Engine/Plugins/FX/Niagara/Source/NiagaraEditor/Public/ViewModels/NiagaraSystemViewModel.h b/Engine/Plugins/FX/Niagara/Source/NiagaraEditor/Public/ViewModels/NiagaraSystemViewModel.h
--- a/Engine/Plugins/FX/Niagara/Source/NiagaraEditor/Public/ViewModels/NiagaraSystemViewModel.h
+++ b/Engine/Plugins/FX/Niagara/Source/NiagaraEditor/Public/ViewModels/NiagaraSystemViewModel.h
@@ -661,6 +661,9 @@ private:
 	/** A flag for preventing reentrancy when synchronizing sequencer selection with system selection */
 	bool bUpdatingSequencerSelectionFromSystem;
 
+	/** A flag for preventing selection changes when resetting the sequencer tracks due to the emitter handles array changing. */
+	bool bResetingSequencerTracks;
+
 	TNiagaraViewModelManager<UNiagaraSystem, FNiagaraSystemViewModel>::Handle RegisteredHandle;
 
 	/** Array of FNiagaraEmitterHandleViewModel representing emitters that are pinned in the stack UI*/
