- Improve PlatformSet stack warning messaging
- Add a fix button to remove the offending cvar

#rnx

diff --git a/Engine/Plugins/FX/Niagara/Source/NiagaraEditor/Private/ViewModels/Stack/NiagaraStackObjectIssueGenerator.cpp b/Engine/Plugins/FX/Niagara/Source/NiagaraEditor/Private/ViewModels/Stack/NiagaraStackObjectIssueGenerator.cpp
--- a/Engine/Plugins/FX/Niagara/Source/NiagaraEditor/Private/ViewModels/Stack/NiagaraStackObjectIssueGenerator.cpp
+++ b/Engine/Plugins/FX/Niagara/Source/NiagaraEditor/Private/ViewModels/Stack/NiagaraStackObjectIssueGenerator.cpp
@@ -1,8 +1,9 @@
 // Copyright Epic Games, Inc. All Rights Reserved.
 
 #include "ViewModels/Stack/NiagaraStackObjectIssueGenerator.h"
+#include "NiagaraEmitter.h"
 #include "NiagaraPlatformSet.h"
-
+#include "NiagaraRendererProperties.h"
 
 void FNiagaraPlatformSetIssueGenerator::GenerateIssues(void* Object, UNiagaraStackObject* StackObject, TArray<UNiagaraStackEntry::FStackIssue>& NewIssues)
 {
@@ -18,13 +19,57 @@ void FNiagaraPlatformSetIssueGenerator::GenerateIssues(void* Object, UNiagaraSta
 		}
 	}
 
-	for (FName& CVarName : InvalidCVars)
+	if (InvalidCVars.Num() > 0)
+	{
+		FText ShortMessageFormat;
+		FText LongMessageFormat;
+		if (UNiagaraEmitter* OwnerEmitter = Cast<UNiagaraEmitter>(StackObject->GetObject()))
+		{
+			ShortMessageFormat = NSLOCTEXT("StackObject", "EmitterInvalidCVarConditionShort", "Emitter Scalability has an invalid cvar '{0}'!");
+			LongMessageFormat = NSLOCTEXT("StackObject", "EmitterInvalidCVarConditionLong", "Emitter Scalability has an invalid cvar '{0}'!");
+		}
+		else if (UNiagaraRendererProperties* OwnerRenderer = Cast<UNiagaraRendererProperties>(StackObject->GetObject()))
+		{
+			ShortMessageFormat = NSLOCTEXT("StackObject", "RendererInvalidCVarConditionShort", "Renderer Scalability has an invalid cvar '{0}'!");
+			LongMessageFormat = NSLOCTEXT("StackObject", "RendererInvalidCVarConditionLong", "Renderer Scalability has an invalid cvar '{0}'!");
+		}
+		else
+		{
+			ShortMessageFormat = NSLOCTEXT("StackObject", "InvalidCVarConditionShort", "PlatformSet has an invalid cvar '{0}'!");
+			LongMessageFormat = NSLOCTEXT("StackObject", "InvalidCVarConditionLong", "PlatformSet has an invalid cvar '{0}'!");
+		}
+
+		for (const FName& CVarName : InvalidCVars)
+		{
+			const FText CVarNameText = FText::FromName(CVarName);
+			NewIssues.Emplace(
+				EStackIssueSeverity::Error,
+				FText::Format(ShortMessageFormat, CVarNameText),
+				FText::Format(LongMessageFormat, CVarNameText),
+				StackObject->GetStackEditorDataKey(),
+				false,
+				UNiagaraStackEntry::FStackIssueFix(
+					FText::Format(NSLOCTEXT("StackObject", "RemoveInvalidCVar", "Remove invalid cvar '{0}'"), CVarNameText),
+					UNiagaraStackEntry::FStackIssueFixDelegate::CreateStatic(FNiagaraPlatformSetIssueGenerator::FixIssue, PlatformSet, TWeakObjectPtr<UObject>(StackObject->GetObject()), CVarName)
+				)
+			);
+		}
+	}
+}
+
+void FNiagaraPlatformSetIssueGenerator::FixIssue(FNiagaraPlatformSet* PlatformSet, TWeakObjectPtr<UObject> WeakObject, FName CVarName)
+{
+	UObject* OwnerObject = WeakObject.Get();
+	if (OwnerObject == nullptr)
 	{
-		NewIssues.Add(UNiagaraStackEntry::FStackIssue(
-			EStackIssueSeverity::Error,
-			NSLOCTEXT("StackObject", "InvalidCVarConditionShort", "Platform Set has an invalid CVar condition!"),
-			FText::Format(FTextFormat::FromString("{0} is not a valid CVar name."), FText::FromName(CVarName)),
-			StackObject->GetStackEditorDataKey(),
-			false));
+		return;
 	}
+
+	OwnerObject->Modify();
+	PlatformSet->CVarConditions.RemoveAll(
+		[&](const FNiagaraPlatformSetCVarCondition& CVarCondition)
+		{
+			return CVarCondition.CVarName == CVarName;
+		}
+	);
 }
diff --git a/Engine/Plugins/FX/Niagara/Source/NiagaraEditor/Public/ViewModels/Stack/NiagaraStackObjectIssueGenerator.h b/Engine/Plugins/FX/Niagara/Source/NiagaraEditor/Public/ViewModels/Stack/NiagaraStackObjectIssueGenerator.h
--- a/Engine/Plugins/FX/Niagara/Source/NiagaraEditor/Public/ViewModels/Stack/NiagaraStackObjectIssueGenerator.h
+++ b/Engine/Plugins/FX/Niagara/Source/NiagaraEditor/Public/ViewModels/Stack/NiagaraStackObjectIssueGenerator.h
@@ -4,6 +4,8 @@
 
 #include "ViewModels/Stack/NiagaraStackObject.h"
 
+struct FNiagaraPlatformSet;
+
 class INiagaraStackObjectIssueGenerator
 {
 public:
@@ -15,4 +17,6 @@ class FNiagaraPlatformSetIssueGenerator : public INiagaraStackObjectIssueGenerat
 {
 public:
 	virtual void GenerateIssues(void* Object, UNiagaraStackObject* StackObject, TArray<UNiagaraStackEntry::FStackIssue>& NewIssues)override;
+private:
+	static void FixIssue(FNiagaraPlatformSet* PlatformSet, TWeakObjectPtr<UObject> WeakObject, FName CVarName);
 };
\ No newline at end of file
