Niagara - When synchronizing referencing map pins with function calls, make sure that the matching bound variable is an input when synchronizing inputs and is in the same namespace when synchronizing attributes and outputs.  This fixes a fairly rare issue where you have inputs and outputs (or attributes) with the same name and they were both recorded in the bound pin names array without their namespaces.  This can result in the wrong variable being matched which can break existing data.

diff --git a/Engine/Plugins/FX/Niagara/Source/NiagaraEditor/Private/ViewModels/Stack/NiagaraStackGraphUtilities.cpp b/Engine/Plugins/FX/Niagara/Source/NiagaraEditor/Private/ViewModels/Stack/NiagaraStackGraphUtilities.cpp
--- a/Engine/Plugins/FX/Niagara/Source/NiagaraEditor/Private/ViewModels/Stack/NiagaraStackGraphUtilities.cpp
+++ b/Engine/Plugins/FX/Niagara/Source/NiagaraEditor/Private/ViewModels/Stack/NiagaraStackGraphUtilities.cpp
@@ -3508,26 +3508,29 @@ void FNiagaraStackGraphUtilities::SynchronizeReferencingMapPinsWithFunctionCall(
 				if (BoundVariable != nullptr)
 				{
 					FNiagaraParameterHandle BoundVariableHandle(BoundVariable->GetName());
-					if (OverrideHandle.GetName() != BoundVariableHandle.GetName())
+					if (BoundVariableHandle.IsModuleHandle())
 					{
-						FNiagaraParameterHandle UpdatedOverrideHandle = FNiagaraParameterHandle::CreateAliasedModuleParameterHandle(BoundVariable->GetName(), &InFunctionCallNode);
-						OverridePin->PinName = UpdatedOverrideHandle.GetParameterHandleString();
-						InFunctionCallNode.UpdateInputNameBinding(BoundGuid, BoundVariableHandle.GetParameterHandleString());
-					}
-					if (OverrideVariable.GetType() != BoundVariable->GetType())
-					{
-						FEdGraphPinType NewPinType = NiagaraSchema->TypeDefinitionToPinType(BoundVariable->GetType());
-						OverridePin->PinType = NewPinType;
-						if (OverridePin->LinkedTo.Num() == 1 &&
-							OverridePin->LinkedTo[0]->GetOwningNode() != nullptr &&
-							OverridePin->LinkedTo[0]->GetOwningNode()->IsA<UNiagaraNodeCustomHlsl>())
+						if (OverrideHandle.GetName() != BoundVariableHandle.GetName())
+						{
+							FNiagaraParameterHandle UpdatedOverrideHandle = FNiagaraParameterHandle::CreateAliasedModuleParameterHandle(BoundVariable->GetName(), &InFunctionCallNode);
+							OverridePin->PinName = UpdatedOverrideHandle.GetParameterHandleString();
+							InFunctionCallNode.UpdateInputNameBinding(BoundGuid, BoundVariableHandle.GetParameterHandleString());
+						}
+						if (OverrideVariable.GetType() != BoundVariable->GetType())
 						{
-							// Custom hlsl nodes in the stack are expression dynamic inputs and their output pin type must match their connected input
-							// so update that here.
-							OverridePin->LinkedTo[0]->PinType = NewPinType;
+							FEdGraphPinType NewPinType = NiagaraSchema->TypeDefinitionToPinType(BoundVariable->GetType());
+							OverridePin->PinType = NewPinType;
+							if (OverridePin->LinkedTo.Num() == 1 &&
+								OverridePin->LinkedTo[0]->GetOwningNode() != nullptr &&
+								OverridePin->LinkedTo[0]->GetOwningNode()->IsA<UNiagaraNodeCustomHlsl>())
+							{
+								// Custom hlsl nodes in the stack are expression dynamic inputs and their output pin type must match their connected input
+								// so update that here.
+								OverridePin->LinkedTo[0]->PinType = NewPinType;
+							}
 						}
+						break;
 					}
-					break;
 				}
 			}
 		}
@@ -3555,16 +3558,19 @@ void FNiagaraStackGraphUtilities::SynchronizeReferencingMapPinsWithFunctionCall(
 					{
 						FNiagaraVariable BoundVariableResolvedName = FNiagaraUtilities::ResolveAliases(*BoundVariable,FNiagaraAliasContext().ChangeModuleToModuleName(InFunctionCallNode.GetFunctionName()));
 						FNiagaraParameterHandle BoundVariableHandle(BoundVariableResolvedName.GetName());
-						if (LinkedModuleParameterHandle.GetName() != BoundVariableHandle.GetName())
+						if (LinkedModuleParameterHandle.GetNamespace() == BoundVariableHandle.GetNamespace())
 						{
-							StackLinkedModuleParameterPin->PinName = BoundVariableResolvedName.GetName();
-							TargetFunctionCallNode->UpdateInputNameBinding(BoundGuid, BoundVariableHandle.GetParameterHandleString());
-						}
-						if (LinkedModuleParameterVariable.GetType() != BoundVariable->GetType())
-						{
-							StackLinkedModuleParameterPin->PinType = NiagaraSchema->TypeDefinitionToPinType(BoundVariable->GetType());
+							if (LinkedModuleParameterHandle.GetName() != BoundVariableHandle.GetName())
+							{
+								StackLinkedModuleParameterPin->PinName = BoundVariableResolvedName.GetName();
+								TargetFunctionCallNode->UpdateInputNameBinding(BoundGuid, BoundVariableHandle.GetParameterHandleString());
+							}
+							if (LinkedModuleParameterVariable.GetType() != BoundVariable->GetType())
+							{
+								StackLinkedModuleParameterPin->PinType = NiagaraSchema->TypeDefinitionToPinType(BoundVariable->GetType());
+							}
+							break;
 						}
-						break;
 					}
 				}
 			}
