Fix for translation results being ommited from GPU compile results.
Adding a stack error for Events in GPU emitters.

diff --git a/Engine/Plugins/FX/Niagara/Source/NiagaraEditor/Private/NiagaraCompiler.cpp b/Engine/Plugins/FX/Niagara/Source/NiagaraEditor/Private/NiagaraCompiler.cpp
--- a/Engine/Plugins/FX/Niagara/Source/NiagaraEditor/Private/NiagaraCompiler.cpp
+++ b/Engine/Plugins/FX/Niagara/Source/NiagaraEditor/Private/NiagaraCompiler.cpp
@@ -1676,14 +1676,11 @@ int32 FHlslNiagaraCompiler::CompileScript(const FNiagaraCompileRequestData* InCo
 		CompileResults.bComputeSucceeded = false;
 		if (CompileResults.bVMSucceeded)
 		{
+			//Clear out current contents of compile results.
 			*(CompileResults.Data) = CompilationJob->TranslatorOutput.ScriptData;
 			CompileResults.Data->ByteCode.Reset();
 			CompileResults.bComputeSucceeded = true;
 		}
-		CompileResults.Data->LastHlslTranslationGPU = TranslatedHLSL;
-		CompileResults.Data->CompileTags = InTranslateResults.CompileTags;
-		DumpDebugInfo(CompileResults, Input, true);
-		CompilationJob->CompileResults = CompileResults;
 	}
 
 	CompileResults.AppendCompileEvents(MakeArrayView(InTranslateResults.CompileEvents));
@@ -1692,8 +1689,12 @@ int32 FHlslNiagaraCompiler::CompileScript(const FNiagaraCompileRequestData* InCo
 	CompileResults.Data->CompileTags = InTranslateResults.CompileTags;
 
 	// Early out if compiling a GPU particle script as we do not need to submit a CPU compile request.
+	// This must be done after we add in the translator errors etc so tha they are passed to the compile job correctly.
 	if (bCompilingGPUParticleScript)
 	{
+		CompileResults.Data->LastHlslTranslationGPU = TranslatedHLSL;
+		DumpDebugInfo(CompileResults, Input, true);
+		CompilationJob->CompileResults = CompileResults;
 		return JobID;
 	}
 
diff --git a/Engine/Plugins/FX/Niagara/Source/NiagaraEditor/Private/ViewModels/Stack/NiagaraStackEventScriptItemGroup.cpp b/Engine/Plugins/FX/Niagara/Source/NiagaraEditor/Private/ViewModels/Stack/NiagaraStackEventScriptItemGroup.cpp
--- a/Engine/Plugins/FX/Niagara/Source/NiagaraEditor/Private/ViewModels/Stack/NiagaraStackEventScriptItemGroup.cpp
+++ b/Engine/Plugins/FX/Niagara/Source/NiagaraEditor/Private/ViewModels/Stack/NiagaraStackEventScriptItemGroup.cpp
@@ -122,6 +122,36 @@ void UNiagaraStackEventHandlerPropertiesItem::RefreshChildrenInternal(const TArr
 		EmitterObject->SetOnSelectRootNodes(UNiagaraStackObject::FOnSelectRootNodes::CreateUObject(this, &UNiagaraStackEventHandlerPropertiesItem::SelectEmitterStackObjectRootTreeNodes));
 	}
 
+	FVersionedNiagaraEmitterData* EmitterData = GetEmitterViewModel()->GetEmitter().GetEmitterData();
+	if (EmitterData && (EmitterData->SimTarget == ENiagaraSimTarget::GPUComputeSim))
+	{
+		TArray<FStackIssueFix> IssueFixes;
+		IssueFixes.Emplace(
+			LOCTEXT("SetCpuSimulationFix", "Set CPU simulation"),
+			FStackIssueFixDelegate::CreateLambda(
+				[WeakEmitter = GetEmitterViewModel()->GetEmitter().ToWeakPtr()]()
+		{
+			if (WeakEmitter.IsValid())
+			{
+				FScopedTransaction Transaction(LOCTEXT("SetCpuSimulation", "Set Cpu Simulation"));
+				WeakEmitter.Emitter.Get()->Modify();
+				WeakEmitter.GetEmitterData()->SimTarget = ENiagaraSimTarget::CPUSim;
+				UNiagaraSystem::RequestCompileForEmitter(WeakEmitter.ResolveWeakPtr());
+			}
+		}
+		)
+		);
+
+		NewIssues.Emplace(
+			EStackIssueSeverity::Error,
+			LOCTEXT("EventHandlersNotSupportedOnCPU", "Event Handlers are not supported on GPU"),
+			LOCTEXT("EventHandlersNotSupportedOnCPULong", "Event Handlers are currently not supported on GPU, please disable or remove."),
+			GetStackEditorDataKey(),
+			false,
+			IssueFixes
+		);
+	}
+
 	NewChildren.Add(EmitterObject);
 
 	bCanResetToBaseCache.Reset();
