Fix for distance field atlas validation error.

diff --git a/Engine/Plugins/FX/Niagara/Source/NiagaraShader/Private/NiagaraDistanceFieldHelper.cpp b/Engine/Plugins/FX/Niagara/Source/NiagaraShader/Private/NiagaraDistanceFieldHelper.cpp
--- a/Engine/Plugins/FX/Niagara/Source/NiagaraShader/Private/NiagaraDistanceFieldHelper.cpp
+++ b/Engine/Plugins/FX/Niagara/Source/NiagaraShader/Private/NiagaraDistanceFieldHelper.cpp
@@ -83,11 +83,13 @@ void FNiagaraDistanceFieldHelper::SetMeshDistanceFieldParameters(FRDGBuilder& Gr
 		ObjectShaderParameters.SceneHeightfieldObjectData = DefaultVector4;
 		ObjectShaderParameters.NumSceneHeightfieldObjects = 0;
 
+		const FRDGSystemTextures& SystemTextures = FRDGSystemTextures::Get(GraphBuilder);
+
 		AtlasShaderParameters.SceneDistanceFieldAssetData = DefaultVector4;
 		AtlasShaderParameters.DistanceFieldIndirectionTable = DefaultUInt32;
 		AtlasShaderParameters.DistanceFieldIndirection2Table = DefaultVector4;
-		AtlasShaderParameters.DistanceFieldIndirectionAtlas = FRDGSystemTextures::Get(GraphBuilder).VolumetricBlack;
-		AtlasShaderParameters.DistanceFieldBrickTexture = GBlackVolumeTexture->GetTextureRHI();
+		AtlasShaderParameters.DistanceFieldIndirectionAtlas = SystemTextures.VolumetricBlack;
+		AtlasShaderParameters.DistanceFieldBrickTexture = SystemTextures.Black;
 		AtlasShaderParameters.DistanceFieldSampler = TStaticSamplerState<SF_Bilinear, AM_Wrap, AM_Wrap, AM_Wrap>::GetRHI();
 		AtlasShaderParameters.DistanceFieldBrickSize = FVector3f::ZeroVector;
 		AtlasShaderParameters.DistanceFieldUniqueDataBrickSize = FVector3f::ZeroVector;
diff --git a/Engine/Source/Runtime/Renderer/Private/DistanceFieldLightingShared.h b/Engine/Source/Runtime/Renderer/Private/DistanceFieldLightingShared.h
--- a/Engine/Source/Runtime/Renderer/Private/DistanceFieldLightingShared.h
+++ b/Engine/Source/Runtime/Renderer/Private/DistanceFieldLightingShared.h
@@ -96,7 +96,7 @@ BEGIN_SHADER_PARAMETER_STRUCT(FDistanceFieldAtlasParameters, )
 	SHADER_PARAMETER_RDG_BUFFER_SRV(StructuredBuffer<uint>, DistanceFieldIndirectionTable)
 	SHADER_PARAMETER_RDG_BUFFER_SRV(Buffer<float4>, DistanceFieldIndirection2Table)
 	SHADER_PARAMETER_RDG_TEXTURE(Texture3D<float4>, DistanceFieldIndirectionAtlas)
-	SHADER_PARAMETER_TEXTURE(Texture3D, DistanceFieldBrickTexture)
+	SHADER_PARAMETER_RDG_TEXTURE(Texture3D, DistanceFieldBrickTexture)
 	SHADER_PARAMETER_SAMPLER(SamplerState, DistanceFieldSampler)
 	SHADER_PARAMETER(FVector3f, DistanceFieldBrickSize)
 	SHADER_PARAMETER(FVector3f, DistanceFieldUniqueDataBrickSize)
diff --git a/Engine/Source/Runtime/Renderer/Private/DistanceFieldObjectManagement.cpp b/Engine/Source/Runtime/Renderer/Private/DistanceFieldObjectManagement.cpp
--- a/Engine/Source/Runtime/Renderer/Private/DistanceFieldObjectManagement.cpp
+++ b/Engine/Source/Runtime/Renderer/Private/DistanceFieldObjectManagement.cpp
@@ -146,7 +146,7 @@ FDistanceFieldAtlasParameters DistanceField::SetupAtlasParameters(FRDGBuilder& G
 			SceneParameters.DistanceFieldIndirectionAtlas = GraphBuilder.RegisterExternalTexture(DistanceFieldSceneData.IndirectionAtlas);
 		}
 
-		SceneParameters.DistanceFieldBrickTexture = DistanceFieldSceneData.DistanceFieldBrickVolumeTexture->GetRHI();
+		SceneParameters.DistanceFieldBrickTexture = GraphBuilder.RegisterExternalTexture(DistanceFieldSceneData.DistanceFieldBrickVolumeTexture);
 		SceneParameters.DistanceFieldSampler = TStaticSamplerState<SF_Bilinear,AM_Wrap, AM_Wrap, AM_Wrap>::GetRHI();
 
 		SceneParameters.DistanceFieldBrickSize = FVector3f(DistanceField::BrickSize);
