include rgba grid cvar in compile hash

diff --git a/Engine/Plugins/FX/Niagara/Source/Niagara/Private/NiagaraDataInterfaceGrid3DCollection.cpp b/Engine/Plugins/FX/Niagara/Source/Niagara/Private/NiagaraDataInterfaceGrid3DCollection.cpp
--- a/Engine/Plugins/FX/Niagara/Source/Niagara/Private/NiagaraDataInterfaceGrid3DCollection.cpp
+++ b/Engine/Plugins/FX/Niagara/Source/Niagara/Private/NiagaraDataInterfaceGrid3DCollection.cpp
@@ -338,7 +338,7 @@ bool UNiagaraDataInterfaceGrid3DCollection::AppendCompileHash(FNiagaraCompileHas
 
 
 	InVisitor->UpdatePOD(TEXT("UNiagaraDataInterfaceGrid3DCollectionVersion"), (int32)FNiagaraGridCollection3DDIFunctionVersion::LatestVersion);
-
+	InVisitor->UpdatePOD(TEXT("UNiagaraDataInterfaceGrid3DCollectionSupportsRGBAGrid"), FGrid3DCollectionAttributeHelper::SupportsRGBAGrid() ? 1 : 0);
 	//InVisitor->UpdateString(TEXT("UNiagaraDataInterfaceVolumeTextureHLSLSource"), GetShaderFileHash(TemplateShaderFilePath, EShaderPlatform::SP_PCD3D_SM5).ToString());
 	InVisitor->UpdateShaderParameters<FNDIGrid3DShaderParameters>();
 
diff --git a/Engine/Plugins/FX/Niagara/Source/Niagara/Private/NiagaraDataInterfaceGrid3DCollectionUtils.cpp b/Engine/Plugins/FX/Niagara/Source/Niagara/Private/NiagaraDataInterfaceGrid3DCollectionUtils.cpp
--- a/Engine/Plugins/FX/Niagara/Source/Niagara/Private/NiagaraDataInterfaceGrid3DCollectionUtils.cpp
+++ b/Engine/Plugins/FX/Niagara/Source/Niagara/Private/NiagaraDataInterfaceGrid3DCollectionUtils.cpp
@@ -21,6 +21,11 @@ bool FGrid3DCollectionAttributeHelper::ShouldUseRGBAGrid(const int TotalChannels
 	return TotalNumAttributes == 1 && TotalChannels <= 4 && GNiagaraGrid3DUseRGBAGrid != 0;
 }
 
+bool FGrid3DCollectionAttributeHelper::SupportsRGBAGrid()
+{
+	return GNiagaraGrid3DUseRGBAGrid == 1;
+}
+
 FGrid3DCollectionAttributeHelper::FGrid3DCollectionAttributeHelper(const FNiagaraDataInterfaceGPUParamInfo& InParamInfo, TArray<FText>* OutWarnings /*= nullptr*/) : ParamInfo(InParamInfo)
 {
 	AttributeInfos.Reserve(ParamInfo.GeneratedFunctions.Num());
diff --git a/Engine/Plugins/FX/Niagara/Source/Niagara/Private/NiagaraDataInterfaceGrid3DCollectionUtils.h b/Engine/Plugins/FX/Niagara/Source/Niagara/Private/NiagaraDataInterfaceGrid3DCollectionUtils.h
--- a/Engine/Plugins/FX/Niagara/Source/Niagara/Private/NiagaraDataInterfaceGrid3DCollectionUtils.h
+++ b/Engine/Plugins/FX/Niagara/Source/Niagara/Private/NiagaraDataInterfaceGrid3DCollectionUtils.h
@@ -10,6 +10,7 @@ struct FGrid3DCollectionAttributeHelper
 
 	// only support rgba textures when we have a single attribute that contains up to 4 channels
 	static bool ShouldUseRGBAGrid(const int TotalChannels, const int TotalNumAttributes);
+	static bool SupportsRGBAGrid();
 
 	static TArray<FString> Channels;
 
