update 3d dispatch thread group size to 4x4x4

diff --git a/Engine/Plugins/FX/Niagara/Source/NiagaraShader/Public/NiagaraShader.h b/Engine/Plugins/FX/Niagara/Source/NiagaraShader/Public/NiagaraShader.h
--- a/Engine/Plugins/FX/Niagara/Source/NiagaraShader/Public/NiagaraShader.h
+++ b/Engine/Plugins/FX/Niagara/Source/NiagaraShader/Public/NiagaraShader.h
@@ -200,7 +200,7 @@ public:
 		{
 			case ENiagaraGpuDispatchType::OneD:		return FIntVector(64, 1, 1);
 			case ENiagaraGpuDispatchType::TwoD:		return FIntVector(8, 8, 1);
-			case ENiagaraGpuDispatchType::ThreeD:	return FIntVector(4, 4, 2);
+			case ENiagaraGpuDispatchType::ThreeD:	return FIntVector(4, 4, 4);
 			default:								return FIntVector(64, 1, 1);
 		}
 	}
