Fix with User Param Tab not forwarding value changes to DIs and Object vars

diff --git a/Engine/Plugins/FX/Niagara/Source/NiagaraEditor/Private/Customizations/NiagaraComponentDetails.cpp b/Engine/Plugins/FX/Niagara/Source/NiagaraEditor/Private/Customizations/NiagaraComponentDetails.cpp
--- a/Engine/Plugins/FX/Niagara/Source/NiagaraEditor/Private/Customizations/NiagaraComponentDetails.cpp
+++ b/Engine/Plugins/FX/Niagara/Source/NiagaraEditor/Private/Customizations/NiagaraComponentDetails.cpp
@@ -212,11 +212,11 @@ public:
 				}
 				else if(ParameterValue.GetMode() == ENiagaraVariantMode::DataInterface)
 				{
-					System->GetExposedParameters().OnInterfaceChange();
+					System->GetExposedParameters().SetDataInterface(ParameterValue.GetDataInterface(), ParameterKey);
 				}
 				else if(ParameterValue.GetMode() == ENiagaraVariantMode::Object)
 				{
-					System->GetExposedParameters().OnUObjectChange();
+					System->GetExposedParameters().SetUObject(ParameterValue.GetUObject(), ParameterKey);
 				}
 
 				for (TSharedPtr<IPropertyHandle> PropertyHandle : PropertyHandles)
