Unshelved from pending changelist '21759277':

PR #9537: [NIAGARA] Fixes to "Reset to Default" functionality (Contributed by [REDACTED_PERSON])

diff --git a/Engine/Plugins/FX/Niagara/Source/NiagaraEditor/Private/Customizations/NiagaraTypeCustomizations.cpp b/Engine/Plugins/FX/Niagara/Source/NiagaraEditor/Private/Customizations/NiagaraTypeCustomizations.cpp
--- a/Engine/Plugins/FX/Niagara/Source/NiagaraEditor/Private/Customizations/NiagaraTypeCustomizations.cpp
+++ b/Engine/Plugins/FX/Niagara/Source/NiagaraEditor/Private/Customizations/NiagaraTypeCustomizations.cpp
@@ -1456,6 +1456,9 @@ void FNiagaraDataInterfaceBindingCustomization::CustomizeHeader(TSharedRef<IProp
 	PropertyHandle = InPropertyHandle;
 	TArray<UObject*> Objects;
 	PropertyHandle->GetOuterObjects(Objects);
+
+	PropertyHandle->MarkResetToDefaultCustomized(true);
+
 	bool bAddDefault = true;
 	if (Objects.Num() == 1)
 	{
@@ -1472,17 +1475,42 @@ void FNiagaraDataInterfaceBindingCustomization::CustomizeHeader(TSharedRef<IProp
 			.ValueContent()
 			.MaxDesiredWidth(250.f)
 			[
-				SNew(SComboButton)
-				.OnGetMenuContent(this, &FNiagaraDataInterfaceBindingCustomization::OnGetMenuContent)
-				.ContentPadding(1)
-				.ToolTipText(this, &FNiagaraDataInterfaceBindingCustomization::GetTooltipText)
-				.ButtonStyle(FAppStyle::Get(), "PropertyEditor.AssetComboStyle")
-				.ForegroundColor(FAppStyle::GetColor("PropertyEditor.AssetName.ColorAndOpacity"))
-				.ButtonContent()
+				SNew(SHorizontalBox)
+				+ SHorizontalBox::Slot()
+				.AutoWidth()
+				.Padding(0.0f, 0.0f, 4.0f, 0.0f)
 				[
-					SNew(SNiagaraParameterName)
-					.ParameterName(this, &FNiagaraDataInterfaceBindingCustomization::GetVariableName)
-					.IsReadOnly(true)
+					SNew(SComboButton)
+					.OnGetMenuContent(this, &FNiagaraDataInterfaceBindingCustomization::OnGetMenuContent)
+					.ContentPadding(1)
+					.ToolTipText(this, &FNiagaraDataInterfaceBindingCustomization::GetTooltipText)
+					.ButtonStyle(FAppStyle::Get(), "PropertyEditor.AssetComboStyle")
+					.ForegroundColor(FAppStyle::GetColor("PropertyEditor.AssetName.ColorAndOpacity"))
+					.ButtonContent()
+					[
+						SNew(SNiagaraParameterName)
+						.ParameterName(this, &FNiagaraDataInterfaceBindingCustomization::GetVariableName)
+						.IsReadOnly(true)
+					]
+				]
+				+ SHorizontalBox::Slot()
+				.AutoWidth()
+				.VAlign(VAlign_Center)
+				.Padding(2, 1)
+				[
+					SNew(SButton)
+					.IsFocusable(false)
+					.ToolTipText(LOCTEXT("ResetToDefaultToolTip", "Reset to Default"))
+					.ButtonStyle(FAppStyle::Get(), "NoBorder")
+					.ContentPadding(0)
+					.Visibility(this, &FNiagaraDataInterfaceBindingCustomization::IsResetToDefaultsVisible)
+					.OnClicked(this, &FNiagaraDataInterfaceBindingCustomization::OnResetToDefaultsClicked)
+					.Content()
+					[
+						SNew(SImage)
+						.Image(FAppStyle::GetBrush("PropertyWindow.DiffersFromDefault"))
+						.ColorAndOpacity(FSlateColor::UseForeground())
+					]
 				]
 			];
 			bAddDefault = false;
@@ -1507,6 +1535,18 @@ void FNiagaraDataInterfaceBindingCustomization::CustomizeHeader(TSharedRef<IProp
 	}
 }
 
+EVisibility FNiagaraDataInterfaceBindingCustomization::IsResetToDefaultsVisible() const
+{
+	return TargetDataInterfaceBinding->BoundVariable.IsValid() ? EVisibility::Visible : EVisibility::Hidden;
+}
+
+FReply FNiagaraDataInterfaceBindingCustomization::OnResetToDefaultsClicked()
+{
+	ChangeSource(NAME_None);
+
+	return FReply::Handled();
+}
+
 ////////////////////////
 FName FNiagaraScriptVariableBindingCustomization::GetVariableName() const
 {
diff --git a/Engine/Plugins/FX/Niagara/Source/NiagaraEditor/Private/Customizations/NiagaraTypeCustomizations.h b/Engine/Plugins/FX/Niagara/Source/NiagaraEditor/Private/Customizations/NiagaraTypeCustomizations.h
--- a/Engine/Plugins/FX/Niagara/Source/NiagaraEditor/Private/Customizations/NiagaraTypeCustomizations.h
+++ b/Engine/Plugins/FX/Niagara/Source/NiagaraEditor/Private/Customizations/NiagaraTypeCustomizations.h
@@ -298,6 +298,9 @@ private:
 	TSharedRef<SWidget> OnCreateWidgetForAction(struct FCreateWidgetForActionData* const InCreateData);
 	void OnActionSelected(const TArray< TSharedPtr<FEdGraphSchemaAction> >& SelectedActions, ESelectInfo::Type InSelectionType);
 
+	EVisibility IsResetToDefaultsVisible() const;
+	FReply OnResetToDefaultsClicked();
+
 	TSharedPtr<IPropertyHandle> PropertyHandle;
 	class UNiagaraSimulationStageBase* BaseStage;
 	struct FNiagaraVariableDataInterfaceBinding* TargetDataInterfaceBinding;
diff --git a/Engine/Plugins/FX/Niagara/Source/NiagaraEditorWidgets/Private/SNiagaraStack.cpp b/Engine/Plugins/FX/Niagara/Source/NiagaraEditorWidgets/Private/SNiagaraStack.cpp
--- a/Engine/Plugins/FX/Niagara/Source/NiagaraEditorWidgets/Private/SNiagaraStack.cpp
+++ b/Engine/Plugins/FX/Niagara/Source/NiagaraEditorWidgets/Private/SNiagaraStack.cpp
@@ -979,12 +979,19 @@ SNiagaraStack::FRowWidgets SNiagaraStack::ConstructNameAndValueWidgetsForItem(UN
 			[
 				PropertyRowWidgets.WholeRowWidget.ToSharedRef()
 			];
-			RowBox->AddSlot()
-			.AutoWidth()
-			.Padding(3, 0, 0, 0)
-			[
-				SNew(SResetToDefaultPropertyEditor, PropertyRow->GetDetailTreeNode()->CreatePropertyHandle())
-			];
+
+			// do not add reset widget if it is customized
+			TSharedPtr<IPropertyHandle> PropertyHandle = PropertyRow->GetDetailTreeNode()->CreatePropertyHandle();
+			if (PropertyHandle.IsValid() && !(PropertyHandle->IsResetToDefaultCustomized() || PropertyHandle->HasMetaData(TEXT("NoResetToDefault"))))
+			{
+				RowBox->AddSlot()
+				.AutoWidth()
+				.Padding(3, 0, 0, 0)
+				[
+					SNew(SResetToDefaultPropertyEditor, PropertyHandle)
+				];
+			}
+
 			return FRowWidgets(RowBox); 
 		}
 		else
@@ -1015,19 +1022,25 @@ SNiagaraStack::FRowWidgets SNiagaraStack::ConstructNameAndValueWidgetsForItem(UN
 				NameWidget = PropertyRowWidgets.NameWidget;
 			}
 
-			TSharedRef<SWidget> ValueWidget = SNew(SHorizontalBox)
+			TSharedRef<SHorizontalBox> ValueWidget = SNew(SHorizontalBox)
 			+ SHorizontalBox::Slot()
 			[
 				PropertyRowWidgets.ValueWidget.ToSharedRef()
-			]
-			+ SHorizontalBox::Slot()
-			.AutoWidth()
-			.VAlign(VAlign_Center)
-			.Padding(3, 0, 0, 0)
-			[
-				SNew(SResetToDefaultPropertyEditor, PropertyRow->GetDetailTreeNode()->CreatePropertyHandle())
 			];
 
+			// do not add reset widget if it is customized
+			TSharedPtr<IPropertyHandle> PropertyHandle = PropertyRow->GetDetailTreeNode()->CreatePropertyHandle();
+			if (PropertyHandle.IsValid() && !(PropertyHandle->IsResetToDefaultCustomized() || PropertyHandle->HasMetaData(TEXT("NoResetToDefault"))))
+			{
+				ValueWidget->AddSlot()
+				.AutoWidth()
+				.VAlign(VAlign_Center)
+				.Padding(3, 0, 0, 0)
+				[
+					SNew(SResetToDefaultPropertyEditor, PropertyHandle)
+				];
+			}
+
 			return FRowWidgets(NameWidget.ToSharedRef(), ValueWidget);
 		}
 	}
