Unshelved from pending changelist '19521705':

PR #9017: [Niagara] Error messages are not logged correctly for GPU Scripts (Contributed by [REDACTED_PERSON])

diff --git a/Engine/Plugins/FX/Niagara/Source/NiagaraEditor/Private/ViewModels/NiagaraSystemViewModel.cpp b/Engine/Plugins/FX/Niagara/Source/NiagaraEditor/Private/ViewModels/NiagaraSystemViewModel.cpp
--- a/Engine/Plugins/FX/Niagara/Source/NiagaraEditor/Private/ViewModels/NiagaraSystemViewModel.cpp
+++ b/Engine/Plugins/FX/Niagara/Source/NiagaraEditor/Private/ViewModels/NiagaraSystemViewModel.cpp
@@ -1416,7 +1416,7 @@ void FNiagaraSystemViewModel::TickCompileStatus()
 {
 	// Checking the compile status is expensive which is why it's updated on tick, and multiple scripts
 	// are time sliced across multiple frames.
-	if (System->HasOutstandingCompilationRequests())
+	if (WaitingOnCompilation())
 	{
 		// When compiling the compile status is always unknown.
 		InvalidateCachedCompileStatus();
