Fixing descripiton for ExecCount

diff --git a/Engine/Plugins/FX/Niagara/Source/Niagara/Private/NiagaraConstants.cpp b/Engine/Plugins/FX/Niagara/Source/Niagara/Private/NiagaraConstants.cpp
--- a/Engine/Plugins/FX/Niagara/Source/Niagara/Private/NiagaraConstants.cpp
+++ b/Engine/Plugins/FX/Niagara/Source/Niagara/Private/NiagaraConstants.cpp
@@ -239,7 +239,7 @@ void FNiagaraConstants::Init()
 		SystemStrMap.Add(SYS_PARAM_ENGINE_LOD_DISTANCE_FRACTION, LOCTEXT("LODDistanceFraction", "The distance fraction between this system and it's max culling distance defined in it's EffectType scalabiltiy settings."));
 		SystemStrMap.Add(SYS_PARAM_ENGINE_EXECUTION_STATE, LOCTEXT("ExecutionState", "The execution state of the systems owner. Takes precedence over the systems internal execution state."));
 
-		SystemStrMap.Add(SYS_PARAM_ENGINE_EXEC_COUNT, LOCTEXT("ExecCountDesc", "The index of this particle in the read buffer."));
+		SystemStrMap.Add(SYS_PARAM_ENGINE_EXEC_COUNT, LOCTEXT("ExecCountDesc", "Total number of items we are iterating over in this script invocation."));
 		SystemStrMap.Add(SYS_PARAM_ENGINE_EMITTER_NUM_PARTICLES, LOCTEXT("EmitterNumParticles", "The number of particles for this emitter at the beginning of simulation. Should only be used in Emitter scripts."));
 		SystemStrMap.Add(SYS_PARAM_ENGINE_EMITTER_SIMULATION_POSITION, LOCTEXT("EngineEmitterSimulationPosition", "The frame of reference space which the Emitter simulation is relative to. Is 0,0,0 for an Emitter in local space, and is Engine.Owner.Position for an Emitter in world space."));
 		SystemStrMap.Add(SYS_PARAM_ENGINE_EMITTER_TOTAL_SPAWNED_PARTICLES, LOCTEXT("EmitterTotalSpawnedParticles", "The total number of particles spawned for this emitter at the beginning of this simulation. Should only be used by the particle spawn script the assign unique IDs."));
