Fixing crash when deleting all simstages from grid3d complicated sim

diff --git a/Engine/Plugins/FX/Niagara/Source/NiagaraEditor/Private/ViewModels/Stack/NiagaraStackModuleItem.cpp b/Engine/Plugins/FX/Niagara/Source/NiagaraEditor/Private/ViewModels/Stack/NiagaraStackModuleItem.cpp
--- a/Engine/Plugins/FX/Niagara/Source/NiagaraEditor/Private/ViewModels/Stack/NiagaraStackModuleItem.cpp
+++ b/Engine/Plugins/FX/Niagara/Source/NiagaraEditor/Private/ViewModels/Stack/NiagaraStackModuleItem.cpp
@@ -1607,7 +1607,7 @@ void UNiagaraStackModuleItem::Delete()
 	FGuid EmitterHandleId = EmitterHandle != nullptr ? EmitterHandle->GetId() : FGuid();
 
 	TArray<TWeakObjectPtr<UNiagaraNodeInput>> RemovedNodes;
-	if (FNiagaraStackGraphUtilities::RemoveModuleFromStack(GetSystemViewModel()->GetSystem(), EmitterHandleId, *FunctionCallNode, RemovedNodes))
+	if (IsSystemViewModelValid() && FNiagaraStackGraphUtilities::RemoveModuleFromStack(GetSystemViewModel()->GetSystem(), EmitterHandleId, *FunctionCallNode, RemovedNodes))
 	{
 		UNiagaraGraph* Graph = FunctionCallNode->GetNiagaraGraph();
 		Graph->NotifyGraphNeedsRecompile();
