Fixed DIs being duplicated sometimes even if the owner was a system, where we want to operate on the DIs in the user param store directly.

diff --git a/Engine/Plugins/FX/Niagara/Source/NiagaraEditor/Private/Customizations/NiagaraComponentDetails.cpp b/Engine/Plugins/FX/Niagara/Source/NiagaraEditor/Private/Customizations/NiagaraComponentDetails.cpp
--- a/Engine/Plugins/FX/Niagara/Source/NiagaraEditor/Private/Customizations/NiagaraComponentDetails.cpp
+++ b/Engine/Plugins/FX/Niagara/Source/NiagaraEditor/Private/Customizations/NiagaraComponentDetails.cpp
@@ -212,7 +212,8 @@ public:
 				}
 				else if(ParameterValue.GetMode() == ENiagaraVariantMode::DataInterface)
 				{
-					System->GetExposedParameters().SetDataInterface(ParameterValue.GetDataInterface(), ParameterKey);
+					// we are operating on the DI objects in the parameter store directly, so no need to call SetDataInterface
+					System->GetExposedParameters().OnInterfaceChange();
 				}
 				else if(ParameterValue.GetMode() == ENiagaraVariantMode::Object)
 				{
@@ -772,7 +773,14 @@ void GenerateUserParameterRows(IDetailChildrenBuilder& ChildrenBuilder, UNiagara
 		
 			if (LeftoverUserParameter.IsDataInterface())
 			{
-				ParameterValue = FNiagaraVariant(DuplicateObject(ParameterValue.GetDataInterface(), Owner));
+				if(Owner->IsA<UNiagaraComponent>())
+				{
+					ParameterValue = FNiagaraVariant(DuplicateObject(ParameterValue.GetDataInterface(), Owner));
+				}
+				else if(Owner->IsA<UNiagaraSystem>())
+				{
+					ParameterValue = FNiagaraVariant(ParameterValue.GetDataInterface());
+				}
 			}
 			
 			TSharedPtr<FNiagaraParameterProxy> ParameterProxy = ParameterProxies.Add_GetRef(MakeShareable(new FNiagaraParameterProxy(Owner, LeftoverUserParameter, ParameterValue, OnRebuildChildren, OverridePropertyHandles)));
