default rgba grid to true

diff --git a/Engine/Plugins/FX/Niagara/Source/Niagara/Private/NiagaraDataInterfaceGrid3DCollectionUtils.cpp b/Engine/Plugins/FX/Niagara/Source/Niagara/Private/NiagaraDataInterfaceGrid3DCollectionUtils.cpp
--- a/Engine/Plugins/FX/Niagara/Source/Niagara/Private/NiagaraDataInterfaceGrid3DCollectionUtils.cpp
+++ b/Engine/Plugins/FX/Niagara/Source/Niagara/Private/NiagaraDataInterfaceGrid3DCollectionUtils.cpp
@@ -5,7 +5,7 @@
 
 #define LOCTEXT_NAMESPACE "NiagaraDataInterfaceGrid3DCollection"
 
-static int32 GNiagaraGrid3DUseRGBAGrid = 0;
+static int32 GNiagaraGrid3DUseRGBAGrid = 1;
 static FAutoConsoleVariableRef CVarNiagaraGrid3DUseRGBAGrid(
 	TEXT("fx.Niagara.Grid3D.UseRGBAGrid"),
 	GNiagaraGrid3DUseRGBAGrid,
