Add stack warning for emitters which exclusively use emitter sourced objects and are using dynamic bounds

diff --git a/Engine/Plugins/FX/Niagara/Source/Niagara/Private/NiagaraMeshRendererProperties.cpp b/Engine/Plugins/FX/Niagara/Source/Niagara/Private/NiagaraMeshRendererProperties.cpp
--- a/Engine/Plugins/FX/Niagara/Source/Niagara/Private/NiagaraMeshRendererProperties.cpp
+++ b/Engine/Plugins/FX/Niagara/Source/Niagara/Private/NiagaraMeshRendererProperties.cpp
@@ -155,6 +155,11 @@ FNiagaraRenderer* UNiagaraMeshRendererProperties::CreateEmitterRenderer(ERHIFeat
 
 FNiagaraBoundsCalculator* UNiagaraMeshRendererProperties::CreateBoundsCalculator()
 {
+	if (GetCurrentSourceMode() == ENiagaraRendererSourceDataMode::Emitter)
+	{
+		return nullptr;
+	}
+
 	FBox LocalBounds;
 	LocalBounds.Init();
 
diff --git a/Engine/Plugins/FX/Niagara/Source/Niagara/Private/NiagaraSpriteRendererProperties.cpp b/Engine/Plugins/FX/Niagara/Source/Niagara/Private/NiagaraSpriteRendererProperties.cpp
--- a/Engine/Plugins/FX/Niagara/Source/Niagara/Private/NiagaraSpriteRendererProperties.cpp
+++ b/Engine/Plugins/FX/Niagara/Source/Niagara/Private/NiagaraSpriteRendererProperties.cpp
@@ -105,6 +105,11 @@ FNiagaraRenderer* UNiagaraSpriteRendererProperties::CreateEmitterRenderer(ERHIFe
 
 FNiagaraBoundsCalculator* UNiagaraSpriteRendererProperties::CreateBoundsCalculator()
 {
+	if (GetCurrentSourceMode() == ENiagaraRendererSourceDataMode::Emitter)
+	{
+		return nullptr;
+	}
+
 	return new FNiagaraBoundsCalculatorHelper<true, false, false>();
 }
 
diff --git a/Engine/Plugins/FX/Niagara/Source/Niagara/Private/NiagaraSystemInstance.cpp b/Engine/Plugins/FX/Niagara/Source/Niagara/Private/NiagaraSystemInstance.cpp
--- a/Engine/Plugins/FX/Niagara/Source/Niagara/Private/NiagaraSystemInstance.cpp
+++ b/Engine/Plugins/FX/Niagara/Source/Niagara/Private/NiagaraSystemInstance.cpp
@@ -2533,7 +2533,10 @@ void FNiagaraSystemInstance::Tick_Concurrent(bool bEnqueueGPUTickIfNeeded)
 		// In the case that we have no bounds initialize to a sensible default value to avoid NaNs later on
 		if (LocalBounds.IsValid == false)
 		{
-			LocalBounds = FBox(FVector::ZeroVector, FVector::ZeroVector);
+			// we provide a small amount of volume to our default because small (read subpixel) volumes in screen
+			// space can still create problems with culling.  This system will still have flickering problems as
+			// this box becomes small in screen space, but it's an improvement.
+			LocalBounds = FBox(-FVector::OneVector, FVector::OneVector);
 		}
 	}
 
diff --git a/Engine/Plugins/FX/Niagara/Source/NiagaraEditor/Private/ViewModels/Stack/NiagaraStackEmitterSettingsGroup.cpp b/Engine/Plugins/FX/Niagara/Source/NiagaraEditor/Private/ViewModels/Stack/NiagaraStackEmitterSettingsGroup.cpp
--- a/Engine/Plugins/FX/Niagara/Source/NiagaraEditor/Private/ViewModels/Stack/NiagaraStackEmitterSettingsGroup.cpp
+++ b/Engine/Plugins/FX/Niagara/Source/NiagaraEditor/Private/ViewModels/Stack/NiagaraStackEmitterSettingsGroup.cpp
@@ -160,13 +160,6 @@ void UNiagaraStackEmitterPropertiesItem::RefreshIssues(TArray<FStackIssue>& NewI
 
 				NewIssues.Add(MissingRequiredFixedBoundsModuleError);
 			}
-
-			FStackIssue MissingRequiredFixedBoundsModuleError(
-				EStackIssueSeverity::Warning,
-				LOCTEXT("RequiredFixedBoundsWarningFormat", "The emitter is GPU and is using dynamic bounds mode.\r\nPlease update the Emitter or System properties otherwise bounds may be incorrect."),
-				LOCTEXT("MissingFixedBounds", "Missing fixed bounds."),
-				GetStackEditorDataKey(),
-				false);
 		}
 
 		if ( ActualEmitterData->SimTarget == ENiagaraSimTarget::GPUComputeSim && ActualEmitterData->bGpuAlwaysRunParticleUpdateScript )
@@ -254,6 +247,36 @@ void UNiagaraStackEmitterPropertiesItem::RefreshIssues(TArray<FStackIssue>& NewI
 				NewIssues.Add(WarmupDeltaTimeExceedsEmitterDeltaTimeWarning);
 			}
 		}
+
+		// check for any emitters which are exclusively using Emitter sourced renderers and dynamic bounds.  Currently our bounds calculators don't support
+		// emitter sources and so no valid bounds will be generated
+		if (ActualEmitterData->SimTarget == ENiagaraSimTarget::CPUSim && ActualEmitterData->CalculateBoundsMode == ENiagaraEmitterCalculateBoundMode::Dynamic)
+		{
+			UNiagaraSystem& Sys = GetSystemViewModel()->GetSystem();
+			if (!Sys.bFixedBounds)
+			{
+				bool bHasParticleSourcedRenderer = false;
+				ActualEmitterData->ForEachEnabledRenderer([&](UNiagaraRendererProperties* RendererProperties)
+				{
+					if (RendererProperties && RendererProperties->GetCurrentSourceMode() == ENiagaraRendererSourceDataMode::Particles)
+					{
+						bHasParticleSourcedRenderer = true;
+					}
+				});
+
+				if (!bHasParticleSourcedRenderer)
+				{
+					FStackIssue EmitterSourcedEmitterRequiredFixedBoundsError(
+						EStackIssueSeverity::Warning,
+						LOCTEXT("EmitterSourcedWarningFormat", "The emitter is using dynamic bounds mode but only using Emitter sourced renderers.\r\nPlease update the Emitter or System properties otherwise bounds may be incorrect."),
+						LOCTEXT("EmitterSourcedWarning", "Missing fixed bounds."),
+						GetStackEditorDataKey(),
+						false);
+
+					NewIssues.Add(EmitterSourcedEmitterRequiredFixedBoundsError);
+				}
+			}
+		}
 	}
 
 }
