suppress warning

diff --git a/Engine/Plugins/FX/Niagara/Source/Niagara/Private/NiagaraDataInterfaceRigidMeshCollisionQuery.cpp b/Engine/Plugins/FX/Niagara/Source/Niagara/Private/NiagaraDataInterfaceRigidMeshCollisionQuery.cpp
--- a/Engine/Plugins/FX/Niagara/Source/Niagara/Private/NiagaraDataInterfaceRigidMeshCollisionQuery.cpp
+++ b/Engine/Plugins/FX/Niagara/Source/Niagara/Private/NiagaraDataInterfaceRigidMeshCollisionQuery.cpp
@@ -298,11 +298,13 @@ void CountCollisionPrimitives(TConstArrayView<TComponentType*> Components, uint3
 			}
 		});
 
+		// suppress warning since we get it a lot
+		/*
 		if (HasConvexElements)
-		{
-			// suppress warning since we get it a lot
-			// UE_LOG(LogRigidMeshCollision, Warning, TEXT("Convex collision objects encountered and will be interpreted as a bounding box on %s"), *Component->GetOwner()->GetName());
+		{			
+			UE_LOG(LogRigidMeshCollision, Warning, TEXT("Convex collision objects encountered and will be interpreted as a bounding box on %s"), *Component->GetOwner()->GetName());
 		}
+		*/
 	}
 }
 
