Niagara - Fix an issue in the stack function input binder where code in a checkf statement was being stripped in a shipping editor build.

diff --git a/Engine/Plugins/FX/Niagara/Source/NiagaraEditor/Private/NiagaraStackFunctionInputBinder.cpp b/Engine/Plugins/FX/Niagara/Source/NiagaraEditor/Private/NiagaraStackFunctionInputBinder.cpp
--- a/Engine/Plugins/FX/Niagara/Source/NiagaraEditor/Private/NiagaraStackFunctionInputBinder.cpp
+++ b/Engine/Plugins/FX/Niagara/Source/NiagaraEditor/Private/NiagaraStackFunctionInputBinder.cpp
@@ -282,7 +282,8 @@ void FNiagaraStackFunctionInputBinder::SetData(const uint8* InValue, int32 InSiz
 
 			FString PinDefaultValue;
 			const UEdGraphSchema_Niagara* Schema = GetDefault<UEdGraphSchema_Niagara>();
-			checkf(Schema->TryGetPinDefaultValueFromNiagaraVariable(TempVariable, PinDefaultValue), TEXT("Default value not supported for type %s"), *InputType.GetName());
+			bool bSuccess = Schema->TryGetPinDefaultValueFromNiagaraVariable(TempVariable, PinDefaultValue);
+			checkf(bSuccess, TEXT("Default value not supported for type %s"), *InputType.GetName());
 
 			if (OverridePin == nullptr)
 			{
