- Add mesh bounds scale to allow adjusting bounds to compensate for WPO

diff --git a/Engine/Plugins/FX/Niagara/Source/Niagara/Private/NiagaraRendererMeshes.cpp b/Engine/Plugins/FX/Niagara/Source/Niagara/Private/NiagaraRendererMeshes.cpp
--- a/Engine/Plugins/FX/Niagara/Source/Niagara/Private/NiagaraRendererMeshes.cpp
+++ b/Engine/Plugins/FX/Niagara/Source/Niagara/Private/NiagaraRendererMeshes.cpp
@@ -187,6 +187,8 @@ void FNiagaraRendererMeshes::Initialize(const UNiagaraRendererProperties* InProp
 			MeshData.Scale = (FVector3f)MeshProperties.Scale;
 			MeshData.Rotation = FQuat4f(MeshProperties.Rotation.Quaternion());
 			MeshData.LocalBounds = Mesh->GetExtendedBounds().GetBox();
+			MeshData.LocalBounds.Min *= Properties->MeshBoundsScale;
+			MeshData.LocalBounds.Max *= Properties->MeshBoundsScale;
 
 			// Create an index remap from mesh material index to it's index in the base material list
 			TArray<UMaterialInterface*> MeshMaterials;
@@ -566,8 +568,8 @@ void FNiagaraRendererMeshes::PreparePerMeshData(FParticleMeshRenderData& Particl
 {
 	// Calculate pivot offset / culling sphere
 	FBox MeshLocalBounds = MeshData.LocalBounds;
-	MeshLocalBounds.Min *= (FVector)MeshData.Scale;
-	MeshLocalBounds.Max *= (FVector)MeshData.Scale;
+	MeshLocalBounds.Min *= FVector(MeshData.Scale);
+	MeshLocalBounds.Max *= FVector(MeshData.Scale);
 	ParticleMeshRenderData.CullingSphere.Center = MeshLocalBounds.GetCenter();
 	ParticleMeshRenderData.CullingSphere.W = MeshLocalBounds.GetExtent().Length();
 
diff --git a/Engine/Plugins/FX/Niagara/Source/Niagara/Public/NiagaraMeshRendererProperties.h b/Engine/Plugins/FX/Niagara/Source/Niagara/Public/NiagaraMeshRendererProperties.h
--- a/Engine/Plugins/FX/Niagara/Source/Niagara/Public/NiagaraMeshRendererProperties.h
+++ b/Engine/Plugins/FX/Niagara/Source/Niagara/Public/NiagaraMeshRendererProperties.h
@@ -281,6 +281,15 @@ public:
 	UPROPERTY(EditAnywhere, Category = "Mesh Rendering", meta = (EditCondition = "bLockedAxisEnable"))
 	ENiagaraMeshLockedAxisSpace LockedAxisSpace;
 
+	/**
+	Scale factor applied to all of the meshes bounds.
+	This impacts distance based and per instance frustum culling.  Per instance frustum culling is enabled by default
+	when GPU scene is enabled.  When using WPO with a material that may expand the mesh beyond the original bounds instances
+	can be frustum culled incorrectly, this allows you to grow the bounds to avoid this issue.
+	*/
+	UPROPERTY(EditAnywhere, Category = "Mesh Rendering", AdvancedDisplay)
+	FVector MeshBoundsScale = FVector::OneVector;
+
 	UPROPERTY(EditAnywhere, Category = "Visibility", meta = (EditCondition = "bEnableCameraDistanceCulling", ClampMin = 0.0f))
 	float MinCameraDistance;
 
