- Add support for sampling a MipLevel to Texture DI

diff --git a/Engine/Plugins/FX/Niagara/Shaders/Private/NiagaraDataInterfaceTextureTemplate.ush b/Engine/Plugins/FX/Niagara/Shaders/Private/NiagaraDataInterfaceTextureTemplate.ush
--- a/Engine/Plugins/FX/Niagara/Shaders/Private/NiagaraDataInterfaceTextureTemplate.ush
+++ b/Engine/Plugins/FX/Niagara/Shaders/Private/NiagaraDataInterfaceTextureTemplate.ush
@@ -1,17 +1,24 @@
 // Copyright Epic Games, Inc. All Rights Reserved.
 
-float2			{ParameterName}_TextureSize;
+int2			{ParameterName}_TextureSize;
+int				{ParameterName}_MipLevels;
 Texture2D		{ParameterName}_Texture;
 SamplerState	{ParameterName}_TextureSampler;
 
-void SampleTexture2D_{ParameterName}(in float2 UV, out float4 OutValue)
+void SampleTexture2D_{ParameterName}(in float2 UV, in float MipLevel, out float4 OutValue)
 {
-	OutValue = {ParameterName}_Texture.SampleLevel({ParameterName}_TextureSampler, UV, 0);
+	OutValue = {ParameterName}_Texture.SampleLevel({ParameterName}_TextureSampler, UV, MipLevel);
 }
 
-void TextureDimensions2D_{ParameterName}(out float2 OutValue)
+void TextureDimensions2D_{ParameterName}(int MipLevel, out float2 OutValue)
 {
-	OutValue = {ParameterName}_TextureSize;
+	OutValue.x = float(max({ParameterName}_TextureSize.x >> MipLevel, 1));
+	OutValue.y = float(max({ParameterName}_TextureSize.y >> MipLevel, 1));
+}
+
+void GetNumMipLevels_{ParameterName}(out int OutMipLevels)
+{
+	OutMipLevels = {ParameterName}_MipLevels;
 }
 
 void SamplePseudoVolumeTexture_{ParameterName}(in float3 UVW, in float2 XYNumFrames, in float TotalNumFrames, in int MipMode, in float MipLevel, in float2 DDX, in float2 DDY, out float4 OutValue)
diff --git a/Engine/Plugins/FX/Niagara/Source/Niagara/Classes/NiagaraDataInterfaceTexture.h b/Engine/Plugins/FX/Niagara/Source/Niagara/Classes/NiagaraDataInterfaceTexture.h
--- a/Engine/Plugins/FX/Niagara/Source/Niagara/Classes/NiagaraDataInterfaceTexture.h
+++ b/Engine/Plugins/FX/Niagara/Source/Niagara/Classes/NiagaraDataInterfaceTexture.h
@@ -14,7 +14,8 @@ class NIAGARA_API UNiagaraDataInterfaceTexture : public UNiagaraDataInterface
 	GENERATED_UCLASS_BODY()
 
 	BEGIN_SHADER_PARAMETER_STRUCT(FShaderParameters, )
-		SHADER_PARAMETER(FVector2f,				TextureSize)
+		SHADER_PARAMETER(FIntPoint,				TextureSize)
+		SHADER_PARAMETER(int32,					MipLevels)
 		SHADER_PARAMETER_RDG_TEXTURE(Texture2D,	Texture)
 		SHADER_PARAMETER_SAMPLER(SamplerState,	TextureSampler)
 	END_SHADER_PARAMETER_STRUCT()
@@ -45,9 +46,10 @@ public:
 	virtual bool PerInstanceTick(void* PerInstanceData, FNiagaraSystemInstance* SystemInstance, float DeltaSeconds) override;
 	//UNiagaraDataInterface Interface End
 
-	void SampleTexture(FVectorVMExternalFunctionContext& Context);
-	void GetTextureDimensions(FVectorVMExternalFunctionContext& Context);
-	void SamplePseudoVolumeTexture(FVectorVMExternalFunctionContext& Context);
+	void VMSampleTexture(FVectorVMExternalFunctionContext& Context);
+	void VMSamplePseudoVolumeTexture(FVectorVMExternalFunctionContext& Context);
+	void VMGetTextureDimensions(FVectorVMExternalFunctionContext& Context);
+	void VMGetNumMipLevels(FVectorVMExternalFunctionContext& Context);
 
 	virtual bool Equals(const UNiagaraDataInterface* Other) const override;
 
@@ -56,6 +58,7 @@ public:
 	virtual bool AppendCompileHash(FNiagaraCompileHashVisitor* InVisitor) const override;
 	virtual void GetParameterDefinitionHLSL(const FNiagaraDataInterfaceGPUParamInfo& ParamInfo, FString& OutHLSL) override;
 	virtual bool GetFunctionHLSL(const FNiagaraDataInterfaceGPUParamInfo& ParamInfo, const FNiagaraDataInterfaceGeneratedFunction& FunctionInfo, int FunctionInstanceIndex, FString& OutHLSL) override;
+	virtual bool UpgradeFunctionCall(FNiagaraFunctionSignature& FunctionSignature) override;
 #endif
 	virtual bool UseLegacyShaderBindings() const override { return false; }
 	virtual void BuildShaderParameters(FNiagaraShaderParametersBuilder& ShaderParametersBuilder) const override;
@@ -70,5 +73,6 @@ protected:
 	static const TCHAR* TemplateShaderFilePath;
 	static const FName SampleTexture2DName;
 	static const FName SamplePseudoVolumeTextureName;
-	static const FName TextureDimsName;
+	static const FName GetTextureDimensionsName;
+	static const FName GetNumMipLevelsName;
 };
diff --git a/Engine/Plugins/FX/Niagara/Source/Niagara/Private/NiagaraDataInterfaceTexture.cpp b/Engine/Plugins/FX/Niagara/Source/Niagara/Private/NiagaraDataInterfaceTexture.cpp
--- a/Engine/Plugins/FX/Niagara/Source/Niagara/Private/NiagaraDataInterfaceTexture.cpp
+++ b/Engine/Plugins/FX/Niagara/Source/Niagara/Private/NiagaraDataInterfaceTexture.cpp
@@ -13,12 +13,26 @@
 const TCHAR* UNiagaraDataInterfaceTexture::TemplateShaderFilePath = TEXT("/Plugin/FX/Niagara/Private/NiagaraDataInterfaceTextureTemplate.ush");
 const FName UNiagaraDataInterfaceTexture::SampleTexture2DName(TEXT("SampleTexture2D"));
 const FName UNiagaraDataInterfaceTexture::SamplePseudoVolumeTextureName(TEXT("SamplePseudoVolumeTexture"));
-const FName UNiagaraDataInterfaceTexture::TextureDimsName(TEXT("TextureDimensions2D"));
+const FName UNiagaraDataInterfaceTexture::GetTextureDimensionsName(TEXT("GetTextureDimensions"));
+const FName UNiagaraDataInterfaceTexture::GetNumMipLevelsName(TEXT("GetNumMipLevels"));
+
+struct FNDITextureFunctionVersion
+{
+	enum Type
+	{
+		InitialVersion = 0,
+		AddMipLevelSupport = 1,
+
+		VersionPlusOne,
+		LatestVersion = VersionPlusOne - 1
+	};
+};
 
 struct FNDITextureInstanceData_GameThread
 {
 	TWeakObjectPtr<UTexture> CurrentTexture = nullptr;
 	FIntPoint CurrentTextureSize = FIntPoint::ZeroValue;
+	int32 CurrentTextureMipLevels = 0;
 	FNiagaraParameterDirectBinding<UObject*> UserParamBinding;
 };
 
@@ -26,7 +40,8 @@ struct FNDITextureInstanceData_RenderThread
 {
 	FSamplerStateRHIRef		SamplerStateRHI;
 	FTextureReferenceRHIRef	TextureReferenceRHI;
-	FVector2f				TextureSize;
+	FIntPoint				TextureSize = FIntPoint(0, 0);
+	int32					MipLevels = 0;
 
 	FRDGTextureRef			TransientRDGTexture = nullptr;
 };
@@ -125,76 +140,82 @@ bool UNiagaraDataInterfaceTexture::Equals(const UNiagaraDataInterface* Other) co
 void UNiagaraDataInterfaceTexture::GetFunctions(TArray<FNiagaraFunctionSignature>& OutFunctions)
 {
 	{
-		FNiagaraFunctionSignature Sig;
+		FNiagaraFunctionSignature& Sig = OutFunctions.AddDefaulted_GetRef();
 		Sig.Name = SampleTexture2DName;
 		Sig.bMemberFunction = true;
 		Sig.bRequiresContext = false;		
 		Sig.bSupportsCPU = false;
 		Sig.bSupportsGPU = true;
-		Sig.Inputs.Add(FNiagaraVariable(FNiagaraTypeDefinition(GetClass()), TEXT("Texture")));
-		Sig.Inputs.Add(FNiagaraVariable(FNiagaraTypeDefinition::GetVec2Def(), TEXT("UV")));
-		Sig.SetDescription(LOCTEXT("TextureSampleTexture2DDesc", "Sample mip level 0 of the input 2d texture at the specified UV coordinates. The UV origin (0,0) is in the upper left hand corner of the image."));
-		Sig.Outputs.Add(FNiagaraVariable(FNiagaraTypeDefinition::GetVec4Def(), TEXT("Value")));
-		//Sig.Owner = *GetName();
-
-		OutFunctions.Add(Sig);
+		Sig.Inputs.Emplace(FNiagaraTypeDefinition(GetClass()), TEXT("Texture"));
+		Sig.Inputs.Emplace(FNiagaraTypeDefinition::GetVec2Def(), TEXT("UV"));
+		Sig.Inputs.Emplace(FNiagaraTypeDefinition::GetFloatDef(), TEXT("MipLevel"));
+		Sig.Outputs.Emplace(FNiagaraTypeDefinition::GetVec4Def(), TEXT("Value"));
+		Sig.SetDescription(LOCTEXT("TextureSampleTexture2DDesc", "Sample supplied mip level from input 2d texture at the specified UV coordinates. The UV origin (0,0) is in the upper left hand corner of the image."));
+		Sig.SetFunctionVersion(FNDITextureFunctionVersion::LatestVersion);
 	}
-
 	{
-		FNiagaraFunctionSignature Sig;
+		FNiagaraFunctionSignature& Sig = OutFunctions.AddDefaulted_GetRef();
 		Sig.Name = SamplePseudoVolumeTextureName;
 		Sig.bMemberFunction = true;
 		Sig.bRequiresContext = false;
 		Sig.bSupportsCPU = false;
 		Sig.bSupportsGPU = true;
-		Sig.Inputs.Add(FNiagaraVariable(FNiagaraTypeDefinition(GetClass()), TEXT("Texture")));
-		Sig.Inputs.Add(FNiagaraVariable(FNiagaraTypeDefinition::GetVec3Def(), TEXT("UVW")));
-		Sig.Inputs.Add(FNiagaraVariable(FNiagaraTypeDefinition::GetVec2Def(), TEXT("XYNumFrames")));
-		Sig.Inputs.Add(FNiagaraVariable(FNiagaraTypeDefinition::GetFloatDef(), TEXT("TotalNumFrames")));
-		Sig.Inputs.Add(FNiagaraVariable(FNiagaraTypeDefinition::GetIntDef(), TEXT("MipMode")));
-		Sig.Inputs.Add(FNiagaraVariable(FNiagaraTypeDefinition::GetFloatDef(), TEXT("MipLevel")));
-		Sig.Inputs.Add(FNiagaraVariable(FNiagaraTypeDefinition::GetVec2Def(), TEXT("DDX")));
-		Sig.Inputs.Add(FNiagaraVariable(FNiagaraTypeDefinition::GetVec2Def(), TEXT("DDY")));
-		
+		Sig.Inputs.Emplace(FNiagaraTypeDefinition(GetClass()), TEXT("Texture"));
+		Sig.Inputs.Emplace(FNiagaraTypeDefinition::GetVec3Def(), TEXT("UVW"));
+		Sig.Inputs.Emplace(FNiagaraTypeDefinition::GetVec2Def(), TEXT("XYNumFrames"));
+		Sig.Inputs.Emplace(FNiagaraTypeDefinition::GetFloatDef(), TEXT("TotalNumFrames"));
+		Sig.Inputs.Emplace(FNiagaraTypeDefinition::GetIntDef(), TEXT("MipMode"));
+		Sig.Inputs.Emplace(FNiagaraTypeDefinition::GetFloatDef(), TEXT("MipLevel"));
+		Sig.Inputs.Emplace(FNiagaraTypeDefinition::GetVec2Def(), TEXT("DDX"));
+		Sig.Inputs.Emplace(FNiagaraTypeDefinition::GetVec2Def(), TEXT("DDY"));
+		Sig.Outputs.Emplace(FNiagaraTypeDefinition::GetVec4Def(), TEXT("Value"));
 		Sig.SetDescription(LOCTEXT("TextureSamplePseudoVolumeTextureDesc", "Return a pseudovolume texture sample.\nUseful for simulating 3D texturing with a 2D texture or as a texture flipbook with lerped transitions.\nTreats 2d layout of frames as a 3d texture and performs bilinear filtering by blending with an offset Z frame.\nTexture = Input Texture Object storing Volume Data\nUVW = Input float3 for Position, 0 - 1\nXYNumFrames = Input float for num frames in x, y directions\nTotalNumFrames = Input float for num total frames\nMipMode = Sampling mode : 0 = use miplevel, 1 = use UV computed gradients, 2 = Use gradients(default = 0)\nMipLevel = MIP level to use in mipmode = 0 (default 0)\nDDX, DDY = Texture gradients in mipmode = 2\n"));
-		Sig.Outputs.Add(FNiagaraVariable(FNiagaraTypeDefinition::GetVec4Def(), TEXT("Value")));
-		//Sig.Owner = *GetName();
-
-		OutFunctions.Add(Sig);
+		Sig.SetFunctionVersion(FNDITextureFunctionVersion::LatestVersion);
 	}
-
 	{
-		FNiagaraFunctionSignature Sig;
-		Sig.Name = TextureDimsName;
+		FNiagaraFunctionSignature& Sig = OutFunctions.AddDefaulted_GetRef();
+		Sig.Name = GetTextureDimensionsName;
 		Sig.bMemberFunction = true;
 		Sig.bRequiresContext = false;
-		Sig.Inputs.Add(FNiagaraVariable(FNiagaraTypeDefinition(GetClass()), TEXT("Texture")));
-		Sig.SetDescription(LOCTEXT("TextureDimsDesc", "Get the dimensions of mip 0 of the texture."));
-		Sig.Outputs.Add(FNiagaraVariable(FNiagaraTypeDefinition::GetVec2Def(), TEXT("Dimensions2D")));
-		//Sig.Owner = *GetName();
-
-		OutFunctions.Add(Sig);
+		Sig.Inputs.Emplace(FNiagaraTypeDefinition(GetClass()), TEXT("Texture"));
+		Sig.Inputs.Emplace(FNiagaraTypeDefinition::GetIntDef(), TEXT("MipLevel"));
+		Sig.Outputs.Emplace(FNiagaraTypeDefinition::GetVec2Def(), TEXT("Dimensions2D"));
+		Sig.SetDescription(LOCTEXT("TextureDimsDesc", "Get the dimensions of the provided Mip level."));
+		Sig.SetFunctionVersion(FNDITextureFunctionVersion::LatestVersion);
+	}
+	{
+		FNiagaraFunctionSignature& Sig = OutFunctions.AddDefaulted_GetRef();
+		Sig.Name = GetNumMipLevelsName;
+		Sig.bMemberFunction = true;
+		Sig.bRequiresContext = false;
+		Sig.Inputs.Emplace(FNiagaraTypeDefinition(GetClass()), TEXT("Texture"));
+		Sig.Outputs.Emplace(FNiagaraTypeDefinition::GetIntDef(), TEXT("NumMipLevels"));
+		Sig.SetDescription(LOCTEXT("TextureDimsDesc", "Get the number of Mip Levels."));
+		Sig.SetFunctionVersion(FNDITextureFunctionVersion::LatestVersion);
 	}
 }
 
-DEFINE_NDI_DIRECT_FUNC_BINDER(UNiagaraDataInterfaceTexture, SampleTexture);
-DEFINE_NDI_DIRECT_FUNC_BINDER(UNiagaraDataInterfaceTexture, SamplePseudoVolumeTexture)
 void UNiagaraDataInterfaceTexture::GetVMExternalFunction(const FVMExternalFunctionBindingInfo& BindingInfo, void* InstanceData, FVMExternalFunction &OutFunc)
 {
 	if (BindingInfo.Name == SampleTexture2DName)
 	{
-		check(BindingInfo.GetNumInputs() == 3 && BindingInfo.GetNumOutputs() == 4);
-		NDI_FUNC_BINDER(UNiagaraDataInterfaceTexture, SampleTexture)::Bind(this, OutFunc);
+		check(BindingInfo.GetNumInputs() == 4 && BindingInfo.GetNumOutputs() == 4);
+		OutFunc = FVMExternalFunction::CreateUObject(this, &UNiagaraDataInterfaceTexture::VMSampleTexture);
 	}
 	else if (BindingInfo.Name == SamplePseudoVolumeTextureName)
 	{
 		check(BindingInfo.GetNumInputs() == 13 && BindingInfo.GetNumOutputs() == 4);
-		NDI_FUNC_BINDER(UNiagaraDataInterfaceTexture, SamplePseudoVolumeTexture)::Bind(this, OutFunc);
+		OutFunc = FVMExternalFunction::CreateUObject(this, &UNiagaraDataInterfaceTexture::VMSamplePseudoVolumeTexture);
+	}
+	else if (BindingInfo.Name == GetTextureDimensionsName)
+	{
+		check(BindingInfo.GetNumInputs() == 2 && BindingInfo.GetNumOutputs() == 2);
+		OutFunc = FVMExternalFunction::CreateUObject(this, &UNiagaraDataInterfaceTexture::VMGetTextureDimensions);
 	}
-	else if (BindingInfo.Name == TextureDimsName)
+	else if (BindingInfo.Name == GetNumMipLevelsName)
 	{
 		check(BindingInfo.GetNumInputs() == 1 && BindingInfo.GetNumOutputs() == 2);
-		OutFunc = FVMExternalFunction::CreateUObject(this, &UNiagaraDataInterfaceTexture::GetTextureDimensions);
+		OutFunc = FVMExternalFunction::CreateUObject(this, &UNiagaraDataInterfaceTexture::VMGetNumMipLevels);
 	}
 }
 
@@ -229,28 +250,43 @@ bool UNiagaraDataInterfaceTexture::PerInstanceTick(void* PerInstanceData, FNiaga
 	FNDITextureInstanceData_GameThread* InstanceData = static_cast<FNDITextureInstanceData_GameThread*>(PerInstanceData);
 
 	UTexture* CurrentTexture = InstanceData->UserParamBinding.GetValueOrDefault<UTexture>(Texture);
-	const FIntPoint CurrentTextureSize = CurrentTexture != nullptr ? FIntPoint(CurrentTexture->GetSurfaceWidth(), CurrentTexture->GetSurfaceHeight()) : FIntPoint::ZeroValue;
-	if ( (InstanceData->CurrentTexture != CurrentTexture) || (InstanceData->CurrentTextureSize != CurrentTextureSize) )
+	FIntPoint CurrentTextureSize = FIntPoint::ZeroValue;
+	int32 CurrentTextureMipLevels = 0;
+	if (UTexture2D* CurrentTexture2D = Cast<UTexture2D>(CurrentTexture) )
+	{
+		CurrentTextureSize = FIntPoint(CurrentTexture->GetSurfaceWidth(), CurrentTexture->GetSurfaceHeight());
+		CurrentTextureMipLevels = CurrentTexture2D->GetNumMips();
+	}
+	else if (UTextureRenderTarget2D* CurrentTexture2DRT = Cast<UTextureRenderTarget2D>(CurrentTexture))
+	{
+		CurrentTextureSize = FIntPoint(CurrentTexture->GetSurfaceWidth(), CurrentTexture->GetSurfaceHeight());
+		CurrentTextureMipLevels = CurrentTexture2DRT->GetNumMips();
+	}
+
+	if ( (InstanceData->CurrentTexture != CurrentTexture) || (InstanceData->CurrentTextureSize != CurrentTextureSize) || (InstanceData->CurrentTextureMipLevels != CurrentTextureMipLevels) )
 	{
 		InstanceData->CurrentTexture = CurrentTexture;
 		InstanceData->CurrentTextureSize = CurrentTextureSize;
+		InstanceData->CurrentTextureMipLevels = CurrentTextureMipLevels;
 
 		ENQUEUE_RENDER_COMMAND(NDITexture_UpdateInstance)
 		(
-			[RT_Proxy=GetProxyAs<FNiagaraDataInterfaceProxyTexture>(), RT_InstanceID=SystemInstance->GetId(), RT_Texture=CurrentTexture, RT_TextureSize=CurrentTextureSize](FRHICommandListImmediate&)
+			[RT_Proxy=GetProxyAs<FNiagaraDataInterfaceProxyTexture>(), RT_InstanceID=SystemInstance->GetId(), RT_Texture=CurrentTexture, RT_TextureSize=CurrentTextureSize, RT_MipLevels=CurrentTextureMipLevels](FRHICommandListImmediate&)
 			{
 				FNDITextureInstanceData_RenderThread& InstanceData = RT_Proxy->InstanceData_RT.FindOrAdd(RT_InstanceID);
 				if (RT_Texture)
 				{
 					InstanceData.TextureReferenceRHI = RT_Texture->TextureReference.TextureReferenceRHI;
 					InstanceData.SamplerStateRHI = RT_Texture->GetResource() ? RT_Texture->GetResource()->SamplerStateRHI.GetReference() : TStaticSamplerState<SF_Point>::GetRHI();
-					InstanceData.TextureSize = FVector2f(RT_TextureSize.X, RT_TextureSize.Y);
+					InstanceData.TextureSize = RT_TextureSize;
+					InstanceData.MipLevels = RT_MipLevels;
 				}
 				else
 				{
 					InstanceData.TextureReferenceRHI = nullptr;
 					InstanceData.SamplerStateRHI = nullptr;
-					InstanceData.TextureSize = FVector2f::ZeroVector;
+					InstanceData.TextureSize = FIntPoint(0, 0);
+					InstanceData.MipLevels = 0;
 				}
 			}
 		);
@@ -259,94 +295,64 @@ bool UNiagaraDataInterfaceTexture::PerInstanceTick(void* PerInstanceData, FNiaga
 	return false;
 }
 
-void UNiagaraDataInterfaceTexture::GetTextureDimensions(FVectorVMExternalFunctionContext& Context)
+void UNiagaraDataInterfaceTexture::VMSampleTexture(FVectorVMExternalFunctionContext& Context)
 {
 	VectorVM::FUserPtrHandler<FNDITextureInstanceData_GameThread> InstData(Context);
-	FNDIOutputParam<float> OutWidth(Context);
-	FNDIOutputParam<float> OutHeight(Context);
+	FNDIInputParam<FVector2f>	InUV(Context);
+	FNDIInputParam<float>		InMipLevel(Context);
+	FNDIOutputParam<FVector4f>	OutValue(Context);
 
+	const FVector4f DefaultValue(1.0f, 0.0f, 1.0f, 1.0f);
 	for (int32 i = 0; i < Context.GetNumInstances(); ++i)
 	{
-		OutWidth.SetAndAdvance(InstData->CurrentTextureSize.X);
-		OutHeight.SetAndAdvance(InstData->CurrentTextureSize.Y);
+		OutValue.SetAndAdvance(DefaultValue);
 	}
 }
 
-void UNiagaraDataInterfaceTexture::SampleTexture(FVectorVMExternalFunctionContext& Context)
+void UNiagaraDataInterfaceTexture::VMSamplePseudoVolumeTexture(FVectorVMExternalFunctionContext& Context)
 {
 	VectorVM::FUserPtrHandler<FNDITextureInstanceData_GameThread> InstData(Context);
-	VectorVM::FExternalFuncInputHandler<float> XParam(Context);
-	VectorVM::FExternalFuncInputHandler<float> YParam(Context);
-	VectorVM::FExternalFuncRegisterHandler<float> OutSampleR(Context);
-	VectorVM::FExternalFuncRegisterHandler<float> OutSampleG(Context);
-	VectorVM::FExternalFuncRegisterHandler<float> OutSampleB(Context);
-	VectorVM::FExternalFuncRegisterHandler<float> OutSampleA(Context);
-
+	FNDIInputParam<FVector3f>	InUVW(Context);
+	FNDIInputParam<FVector2f>	InXYNumFrames(Context);
+	FNDIInputParam<float>		InTotalNumFrames(Context);
+	FNDIInputParam<int32>		InMipMode(Context);
+	FNDIInputParam<float>		InMipLevel(Context);
+	FNDIInputParam<FVector2f>	InDDX(Context);
+	FNDIInputParam<FVector2f>	InDDY(Context);
+	FNDIOutputParam<FVector4f>	OutValue(Context);
+
+	const FVector4f DefaultValue(1.0f, 0.0f, 1.0f, 1.0f);
 	for (int32 i = 0; i < Context.GetNumInstances(); ++i)
 	{
-		float X = XParam.GetAndAdvance();
-		float Y = YParam.GetAndAdvance();
-		*OutSampleR.GetDestAndAdvance() = 1.0;
-		*OutSampleG.GetDestAndAdvance() = 0.0;
-		*OutSampleB.GetDestAndAdvance() = 1.0;
-		*OutSampleA.GetDestAndAdvance() = 1.0;
+		OutValue.SetAndAdvance(DefaultValue);
 	}
-
 }
 
-void UNiagaraDataInterfaceTexture::SamplePseudoVolumeTexture(FVectorVMExternalFunctionContext& Context)
+void UNiagaraDataInterfaceTexture::VMGetTextureDimensions(FVectorVMExternalFunctionContext& Context)
 {
-	// Noop handler which just returns magenta since this doesn't run on CPU.
 	VectorVM::FUserPtrHandler<FNDITextureInstanceData_GameThread> InstData(Context);
-	VectorVM::FExternalFuncInputHandler<float> UVW_UParam(Context);
-	VectorVM::FExternalFuncInputHandler<float> UVW_VParam(Context);
-	VectorVM::FExternalFuncInputHandler<float> UVW_WParam(Context);
-
-	VectorVM::FExternalFuncInputHandler<float> XYNumFrames_XParam(Context);
-	VectorVM::FExternalFuncInputHandler<float> XYNumFrames_YParam(Context);
-	
-	VectorVM::FExternalFuncInputHandler<float> TotalNumFramesParam(Context);
-
-	VectorVM::FExternalFuncInputHandler<int32> MipModeParam(Context);
-
-	VectorVM::FExternalFuncInputHandler<float> MipLevelParam(Context);
-
-	VectorVM::FExternalFuncInputHandler<float> DDX_XParam(Context);
-	VectorVM::FExternalFuncInputHandler<float> DDX_YParam(Context);
-
-	VectorVM::FExternalFuncInputHandler<float> DDY_XParam(Context);
-	VectorVM::FExternalFuncInputHandler<float> DDY_YParam(Context);
-
-	VectorVM::FExternalFuncRegisterHandler<float> OutSampleR(Context);
-	VectorVM::FExternalFuncRegisterHandler<float> OutSampleG(Context);
-	VectorVM::FExternalFuncRegisterHandler<float> OutSampleB(Context);
-	VectorVM::FExternalFuncRegisterHandler<float> OutSampleA(Context);
+	FNDIInputParam<int32>		InMipLevel(Context);
+	FNDIOutputParam<FVector2f>	OutSize(Context);
 
 	for (int32 i = 0; i < Context.GetNumInstances(); ++i)
 	{
-		UVW_UParam.Advance();
-		UVW_VParam.Advance();
-		UVW_WParam.Advance();
-
-		XYNumFrames_XParam.Advance();
-		XYNumFrames_YParam.Advance();
-
-		TotalNumFramesParam.Advance();
-
-		MipModeParam.Advance();
-
-		MipLevelParam.Advance();
-
-		DDX_XParam.Advance();
-		DDX_YParam.Advance();
+		const int32 MipLevel = InMipLevel.GetAndAdvance();
+		const FVector2f TextureSize(
+			FMath::Max(InstData->CurrentTextureSize.X >> MipLevel, 1),
+			FMath::Max(InstData->CurrentTextureSize.Y >> MipLevel, 1)
+		);
+		OutSize.SetAndAdvance(TextureSize);
+	}
+}
 
-		DDY_XParam.Advance();
-		DDY_YParam.Advance();
+void UNiagaraDataInterfaceTexture::VMGetNumMipLevels(FVectorVMExternalFunctionContext& Context)
+{
+	VectorVM::FUserPtrHandler<FNDITextureInstanceData_GameThread> InstData(Context);
+	FNDIOutputParam<int32> OutNumMipLevels(Context);
 
-		*OutSampleR.GetDestAndAdvance() = 1.0;
-		*OutSampleG.GetDestAndAdvance() = 0.0;
-		*OutSampleB.GetDestAndAdvance() = 1.0;
-		*OutSampleA.GetDestAndAdvance() = 1.0;
+	for (int32 i = 0; i < Context.GetNumInstances(); ++i)
+	{
+		OutNumMipLevels.SetAndAdvance(InstData->CurrentTextureMipLevels);
 	}
 }
 
@@ -373,13 +379,36 @@ void UNiagaraDataInterfaceTexture::GetParameterDefinitionHLSL(const FNiagaraData
 
 bool UNiagaraDataInterfaceTexture::GetFunctionHLSL(const FNiagaraDataInterfaceGPUParamInfo& ParamInfo, const FNiagaraDataInterfaceGeneratedFunction& FunctionInfo, int FunctionInstanceIndex, FString& OutHLSL)
 {
-	if ((FunctionInfo.DefinitionName == SampleTexture2DName) ||
-		(FunctionInfo.DefinitionName == SamplePseudoVolumeTextureName) ||
-		(FunctionInfo.DefinitionName == TextureDimsName))
+	static const TSet<FName> ValidGpuFunctions =
+	{
+		SampleTexture2DName,
+		SamplePseudoVolumeTextureName,
+		GetTextureDimensionsName,
+		GetNumMipLevelsName,
+	};
+
+	return ValidGpuFunctions.Contains(FunctionInfo.DefinitionName);
+}
+
+bool UNiagaraDataInterfaceTexture::UpgradeFunctionCall(FNiagaraFunctionSignature& FunctionSignature)
+{
+	if (FunctionSignature.FunctionVersion < FNDITextureFunctionVersion::AddMipLevelSupport)
 	{
-		return true;
+		static FName LegacyDimsName("TextureDimensions2D");
+		if ( FunctionSignature.Name == SampleTexture2DName )
+		{
+			FunctionSignature.Inputs.Emplace(FNiagaraTypeDefinition::GetFloatDef(), TEXT("MipLevel"));
+		}
+		else if (FunctionSignature.Name == LegacyDimsName )
+		{
+			FunctionSignature.Name = GetTextureDimensionsName;
+			FunctionSignature.Inputs.Emplace(FNiagaraTypeDefinition::GetIntDef(), TEXT("MipLevel"));
+		}
 	}
-	return false;
+
+	const bool bUpdated = FunctionSignature.FunctionVersion < FNDITextureFunctionVersion::LatestVersion;
+	FunctionSignature.FunctionVersion = FNDITextureFunctionVersion::LatestVersion;
+	return bUpdated;
 }
 #endif
 
@@ -397,6 +426,7 @@ void UNiagaraDataInterfaceTexture::SetShaderParameters(const FNiagaraDataInterfa
 	if (InstanceData && InstanceData->TextureReferenceRHI.IsValid())
 	{
 		Parameters->TextureSize = InstanceData->TextureSize;
+		Parameters->MipLevels = InstanceData->MipLevels;
 		Parameters->TextureSampler = InstanceData->SamplerStateRHI;
 		if (Context.IsResourceBound(&Parameters->Texture))
 		{
@@ -420,7 +450,8 @@ void UNiagaraDataInterfaceTexture::SetShaderParameters(const FNiagaraDataInterfa
 	}
 	else
 	{
-		Parameters->TextureSize = FVector2f::ZeroVector;
+		Parameters->TextureSize = FIntPoint(0, 0);
+		Parameters->MipLevels = 0;
 		Parameters->TextureSampler = TStaticSamplerState<SF_Point>::GetRHI();
 		Parameters->Texture = Context.GetComputeDispatchInterface().GetBlackTexture(Context.GetGraphBuilder(), ETextureDimension::Texture2D);
 	}
