- Fix for incorrect GetTextureDimensions name

#rnx

diff --git a/Engine/Plugins/FX/Niagara/Shaders/Private/NiagaraDataInterfaceTextureTemplate.ush b/Engine/Plugins/FX/Niagara/Shaders/Private/NiagaraDataInterfaceTextureTemplate.ush
--- a/Engine/Plugins/FX/Niagara/Shaders/Private/NiagaraDataInterfaceTextureTemplate.ush
+++ b/Engine/Plugins/FX/Niagara/Shaders/Private/NiagaraDataInterfaceTextureTemplate.ush
@@ -10,7 +10,7 @@ void SampleTexture2D_{ParameterName}(in float2 UV, in float MipLevel, out float4
 	OutValue = {ParameterName}_Texture.SampleLevel({ParameterName}_TextureSampler, UV, MipLevel);
 }
 
-void TextureDimensions2D_{ParameterName}(int MipLevel, out float2 OutValue)
+void GetTextureDimensions_{ParameterName}(int MipLevel, out float2 OutValue)
 {
 	OutValue.x = float(max({ParameterName}_TextureSize.x >> MipLevel, 1));
 	OutValue.y = float(max({ParameterName}_TextureSize.y >> MipLevel, 1));
