- Fix duplicate LOC text key

#rnx

diff --git a/Engine/Plugins/FX/Niagara/Source/Niagara/Private/NiagaraDataInterfaceTexture.cpp b/Engine/Plugins/FX/Niagara/Source/Niagara/Private/NiagaraDataInterfaceTexture.cpp
--- a/Engine/Plugins/FX/Niagara/Source/Niagara/Private/NiagaraDataInterfaceTexture.cpp
+++ b/Engine/Plugins/FX/Niagara/Source/Niagara/Private/NiagaraDataInterfaceTexture.cpp
@@ -190,7 +190,7 @@ void UNiagaraDataInterfaceTexture::GetFunctions(TArray<FNiagaraFunctionSignature
 		Sig.bRequiresContext = false;
 		Sig.Inputs.Emplace(FNiagaraTypeDefinition(GetClass()), TEXT("Texture"));
 		Sig.Outputs.Emplace(FNiagaraTypeDefinition::GetIntDef(), TEXT("NumMipLevels"));
-		Sig.SetDescription(LOCTEXT("TextureDimsDesc", "Get the number of Mip Levels."));
+		Sig.SetDescription(LOCTEXT("GetNumMipLevelsDesc", "Get the number of Mip Levels."));
 		Sig.SetFunctionVersion(FNDITextureFunctionVersion::LatestVersion);
 	}
 }
