Expose warmup time settings on the Niagara component

diff --git a/Engine/Plugins/FX/Niagara/Source/Niagara/Private/NiagaraComponent.cpp b/Engine/Plugins/FX/Niagara/Source/Niagara/Private/NiagaraComponent.cpp
--- a/Engine/Plugins/FX/Niagara/Source/Niagara/Private/NiagaraComponent.cpp
+++ b/Engine/Plugins/FX/Niagara/Source/Niagara/Private/NiagaraComponent.cpp
@@ -447,6 +447,8 @@ void FNiagaraSceneProxy::GatherSimpleLights(const FSceneViewFamily& ViewFamily,
 UNiagaraComponent::UNiagaraComponent(const FObjectInitializer& ObjectInitializer)
 	: Super(ObjectInitializer)
 	, bForceSolo(false)
+	, bOverrideWarmupSettings(false)
+	, WarmupTickDelta(1.0f / 15.0f)
 	, AgeUpdateMode(ENiagaraAgeUpdateMode::TickDeltaTime)
 	, DesiredAge(0.0f)
 	, LastHandledDesiredAge(0.0f)
@@ -941,8 +943,8 @@ bool UNiagaraComponent::InitializeSystem()
 		const bool bPooled = PoolingMethod != ENCPoolMethod::None;
 		OverrideParameters.MarkParametersDirty(); // new system instance means new lwc tile, so any position user params need to be re-evaluated
 
-		SystemInstanceController = MakeShared<FNiagaraSystemInstanceController, ESPMode::ThreadSafe>();
-		SystemInstanceController->Initialize(*World, *Asset, &OverrideParameters, this, TickBehavior, bPooled, RandomSeedOffset, RequiresSoloMode());
+		SystemInstanceController = MakeShared<FNiagaraSystemInstanceController, ESPMode::ThreadSafe>();		
+		SystemInstanceController->Initialize(*World, *Asset, &OverrideParameters, this, TickBehavior, bPooled, RandomSeedOffset, RequiresSoloMode(), bOverrideWarmupSettings ? WarmupTickCount : -1, WarmupTickDelta);
 		SystemInstanceController->SetOnPostTick(FNiagaraSystemInstance::FOnPostTick::CreateUObject(this, &UNiagaraComponent::PostSystemTick_GameThread));
 		SystemInstanceController->SetOnComplete(FNiagaraSystemInstance::FOnComplete::CreateUObject(this, &UNiagaraComponent::OnSystemComplete));
 
@@ -953,7 +955,7 @@ bool UNiagaraComponent::InitializeSystem()
 		{
 			SystemInstanceController->SetGpuComputeDebug(bEnableGpuComputeDebug != 0);
 		}
-
+		
 		if (bEnablePreviewLODDistance)
 		{
 			SystemInstanceController->SetLODDistance(PreviewLODDistance, PreviewMaxDistance, true);
diff --git a/Engine/Plugins/FX/Niagara/Source/Niagara/Private/NiagaraSystemInstance.cpp b/Engine/Plugins/FX/Niagara/Source/Niagara/Private/NiagaraSystemInstance.cpp
--- a/Engine/Plugins/FX/Niagara/Source/Niagara/Private/NiagaraSystemInstance.cpp
+++ b/Engine/Plugins/FX/Niagara/Source/Niagara/Private/NiagaraSystemInstance.cpp
@@ -598,6 +598,12 @@ void FNiagaraSystemInstance::SetGpuComputeDebug(bool bEnableDebug)
 #endif
 }
 
+void FNiagaraSystemInstance::SetWarmupSettings(int32 InWarmupTickCount, float InWarmupTickDelta)
+{
+	WarmupTickCount = InWarmupTickCount;
+	WarmupTickDelta = InWarmupTickDelta;
+}
+
 UActorComponent* FNiagaraSystemInstance::GetPrereqComponent() const
 {
 	UActorComponent* PrereqComponent = AttachComponent.Get();
@@ -916,12 +922,17 @@ void FNiagaraSystemInstance::Reset(EResetMode Mode)
 			// Add instance to simulation
 			SystemSimulation->AddInstance(this);
 
-			UNiagaraSystem* System = GetSystem();
-			if (System->NeedsWarmup())
+			int32 WarmupTicks = WarmupTickCount;
+			float WarmupDt = WarmupTickDelta;
+			if (WarmupTickCount == -1)
+			{
+				UNiagaraSystem* System = GetSystem();
+				WarmupTicks = System->GetWarmupTickCount();
+				WarmupDt = System->GetWarmupTickDelta();
+			}
+			
+			if (WarmupTicks > 0 && WarmupDt > SMALL_NUMBER)
 			{
-				int32 WarmupTicks = System->GetWarmupTickCount();
-				float WarmupDt = System->GetWarmupTickDelta();
-
 				AdvanceSimulation(WarmupTicks, WarmupDt);
 
 				//Reset age to zero.
diff --git a/Engine/Plugins/FX/Niagara/Source/Niagara/Private/NiagaraSystemInstanceController.cpp b/Engine/Plugins/FX/Niagara/Source/Niagara/Private/NiagaraSystemInstanceController.cpp
--- a/Engine/Plugins/FX/Niagara/Source/Niagara/Private/NiagaraSystemInstanceController.cpp
+++ b/Engine/Plugins/FX/Niagara/Source/Niagara/Private/NiagaraSystemInstanceController.cpp
@@ -12,13 +12,14 @@ FNiagaraSystemInstanceController::FNiagaraSystemInstanceController()
 }
 
 void FNiagaraSystemInstanceController::Initialize(UWorld& World, UNiagaraSystem& System, FNiagaraUserRedirectionParameterStore* InOverrideParameters,
-	USceneComponent* AttachComponent, ENiagaraTickBehavior TickBehavior, bool bPooled, int32 RandomSeedOffset, bool bForceSolo)
+	USceneComponent* AttachComponent, ENiagaraTickBehavior TickBehavior, bool bPooled, int32 RandomSeedOffset, bool bForceSolo, int32 InWarmupTickCount, float InWarmupTickDelta)
 {
 	OverrideParameters = InOverrideParameters;
 
 	FNiagaraSystemInstance::AllocateSystemInstance(SystemInstance, World, System, OverrideParameters, AttachComponent, TickBehavior, bPooled);
 	//UE_LOG(LogNiagara, Log, TEXT("Create System: %p | %s\n"), SystemInstance.Get(), *GetAsset()->GetFullName());
 	SystemInstance->SetRandomSeedOffset(RandomSeedOffset);
+	SystemInstance->SetWarmupSettings(InWarmupTickCount, InWarmupTickDelta);
 	SystemInstance->Init(bForceSolo);
 
 	WorldManager = FNiagaraWorldManager::Get(&World);
diff --git a/Engine/Plugins/FX/Niagara/Source/Niagara/Public/NiagaraComponent.h b/Engine/Plugins/FX/Niagara/Source/Niagara/Public/NiagaraComponent.h
--- a/Engine/Plugins/FX/Niagara/Source/Niagara/Public/NiagaraComponent.h
+++ b/Engine/Plugins/FX/Niagara/Source/Niagara/Public/NiagaraComponent.h
@@ -107,6 +107,18 @@ private:
 	UPROPERTY(EditAnywhere, Category = Parameters)
 	uint32 bEnableGpuComputeDebug : 1;
 
+	/** When true then this instance will override the system's warmup settings. */
+	UPROPERTY(EditAnywhere, Category = Warmup)
+	uint32 bOverrideWarmupSettings : 1;
+	
+	/** Number of ticks to process for warmup of the system. Total warmup time is WarmupTickCount * WarmupTickDelta. */
+	UPROPERTY(EditAnywhere, Category = Warmup, meta=(EditCondition="bOverrideWarmupSettings", ClampMin = "0"))
+	int32 WarmupTickCount = 0;
+
+	/** Delta time used when ticking the system in warmup mode. */
+	UPROPERTY(EditAnywhere, Category = Warmup, meta = (EditCondition="bOverrideWarmupSettings", ForceUnits=s, UIMin = "0.01", UIMax = "1"))
+	float WarmupTickDelta;
+
 	FNiagaraSystemInstanceControllerPtr SystemInstanceController;
 
 	/** Defines the mode use when updating the System age. */
diff --git a/Engine/Plugins/FX/Niagara/Source/Niagara/Public/NiagaraSystemInstance.h b/Engine/Plugins/FX/Niagara/Source/Niagara/Public/NiagaraSystemInstance.h
--- a/Engine/Plugins/FX/Niagara/Source/Niagara/Public/NiagaraSystemInstance.h
+++ b/Engine/Plugins/FX/Niagara/Source/Niagara/Public/NiagaraSystemInstance.h
@@ -127,6 +127,7 @@ public:
 	void SetSolo(bool bInSolo);
 
 	void SetGpuComputeDebug(bool bEnableDebug);
+	void SetWarmupSettings(int32 WarmupTickCount, float WarmupTickDelta);
 
 	UActorComponent* GetPrereqComponent() const;
 
@@ -473,6 +474,9 @@ private:
 	float LODDistance;
 	float MaxLODDistance;
 
+	int32 WarmupTickCount = -1;
+	float WarmupTickDelta = 0;
+	
 	TArray< TSharedRef<FNiagaraEmitterInstance, ESPMode::ThreadSafe> > Emitters;
 
 	FOnPostTick OnPostTickDelegate;
diff --git a/Engine/Plugins/FX/Niagara/Source/Niagara/Public/NiagaraSystemInstanceController.h b/Engine/Plugins/FX/Niagara/Source/Niagara/Public/NiagaraSystemInstanceController.h
--- a/Engine/Plugins/FX/Niagara/Source/Niagara/Public/NiagaraSystemInstanceController.h
+++ b/Engine/Plugins/FX/Niagara/Source/Niagara/Public/NiagaraSystemInstanceController.h
@@ -54,7 +54,7 @@ public:
 
 	/** Initializes the controller with an instance of the provided system */
 	void Initialize(UWorld& World, UNiagaraSystem& System, FNiagaraUserRedirectionParameterStore* OverrideParameters, USceneComponent* AttachComponent,
-		ENiagaraTickBehavior TickBehavior, bool bPooled, int32 RandomSeedOffset, bool bForceSolo);
+		ENiagaraTickBehavior TickBehavior, bool bPooled, int32 RandomSeedOffset, bool bForceSolo, int32 WarmupTickCount, float WarmupTickDelta);
 	/** Deactivates the underlying system and queues it up for destruction. After calling, this interface is unusable unless you call Initialize again */
 	void Release();
 
diff --git a/Engine/Plugins/FX/Niagara/Source/NiagaraEditor/Private/Customizations/NiagaraComponentDetails.cpp b/Engine/Plugins/FX/Niagara/Source/NiagaraEditor/Private/Customizations/NiagaraComponentDetails.cpp
--- a/Engine/Plugins/FX/Niagara/Source/NiagaraEditor/Private/Customizations/NiagaraComponentDetails.cpp
+++ b/Engine/Plugins/FX/Niagara/Source/NiagaraEditor/Private/Customizations/NiagaraComponentDetails.cpp
@@ -1135,6 +1135,7 @@ void FNiagaraComponentDetails::CustomizeDetails(IDetailLayoutBuilder& DetailBuil
 		Section->AddCategory(TEXT("Activation"));
 		Section->AddCategory(ScriptCategoryName);
 		Section->AddCategory(TEXT("Randomness"));
+		Section->AddCategory(TEXT("Warmup"));
 		bFirstTime = false;
 	}
 
@@ -1163,6 +1164,7 @@ void FNiagaraComponentDetails::CustomizeDetails(IDetailLayoutBuilder& DetailBuil
 	DetailBuilder.EditCategory("Attachment", FText::GetEmpty(), ECategoryPriority::TypeSpecific);
 	DetailBuilder.EditCategory("Randomness", FText::GetEmpty(), ECategoryPriority::TypeSpecific);
 	DetailBuilder.EditCategory("Parameters", FText::GetEmpty(), ECategoryPriority::TypeSpecific);
+	DetailBuilder.EditCategory("Warmup", FText::GetEmpty(), ECategoryPriority::TypeSpecific);
 	DetailBuilder.EditCategory("Materials", FText::GetEmpty(), ECategoryPriority::TypeSpecific);
 	
 	if (ObjectsCustomized.Num() == 1 && ObjectsCustomized[0]->IsA<UNiagaraComponent>())
