Add bounds calculator for niagara light renderers

diff --git a/Engine/Plugins/FX/Niagara/Source/Niagara/Private/NiagaraLightRendererProperties.cpp b/Engine/Plugins/FX/Niagara/Source/Niagara/Private/NiagaraLightRendererProperties.cpp
--- a/Engine/Plugins/FX/Niagara/Source/Niagara/Private/NiagaraLightRendererProperties.cpp
+++ b/Engine/Plugins/FX/Niagara/Source/Niagara/Private/NiagaraLightRendererProperties.cpp
@@ -109,6 +109,55 @@ FNiagaraRenderer* UNiagaraLightRendererProperties::CreateEmitterRenderer(ERHIFea
 	return NewRenderer;
 }
 
+class FNiagaraBoundsCalculatorLights : public FNiagaraBoundsCalculator
+{
+public:
+	FNiagaraBoundsCalculatorLights(const FNiagaraDataSetAccessor<FNiagaraPosition>& InPositionAccessor, const FNiagaraDataSetAccessor<float>& InRadiusAccessor)
+		: PositionAccessor(InPositionAccessor)
+		, RadiusAccessor(InRadiusAccessor)
+	{}
+
+	FNiagaraBoundsCalculatorLights() = delete;
+	virtual ~FNiagaraBoundsCalculatorLights() = default;
+
+	virtual void InitAccessors(const FNiagaraDataSetCompiledData* CompiledData) override {}
+	virtual FBox CalculateBounds(const FTransform& SystemTransform, const FNiagaraDataSet& DataSet, const int32 NumInstances) const override
+	{
+		if (!NumInstances || !PositionAccessor.IsValid())
+		{
+			return FBox(ForceInit);
+		}
+
+		FNiagaraPosition BoundsMin(ForceInitToZero);
+		FNiagaraPosition BoundsMax(ForceInitToZero);
+		PositionAccessor.GetReader(DataSet).GetMinMax(BoundsMin, BoundsMax);
+		FBox Bounds(BoundsMin, BoundsMax);
+
+		if (RadiusAccessor.IsValid())
+		{
+			return Bounds.ExpandBy(RadiusAccessor.GetReader(DataSet).GetMax());
+		}
+
+		const FNiagaraVariable DefaultRadius = FNiagaraConstants::GetAttributeWithDefaultValue(SYS_PARAM_PARTICLES_LIGHT_RADIUS);
+
+		return Bounds.ExpandBy(DefaultRadius.GetValue<float>());
+	}
+
+protected:
+	const FNiagaraDataSetAccessor<FNiagaraPosition>& PositionAccessor;
+	const FNiagaraDataSetAccessor<float>& RadiusAccessor;
+};
+
+FNiagaraBoundsCalculator* UNiagaraLightRendererProperties::CreateBoundsCalculator()
+{
+	if (GetCurrentSourceMode() == ENiagaraRendererSourceDataMode::Emitter)
+	{
+		return nullptr;
+	}
+
+	return new FNiagaraBoundsCalculatorLights(PositionDataSetAccessor, RadiusDataSetAccessor);
+}
+
 void UNiagaraLightRendererProperties::GetUsedMaterials(const FNiagaraEmitterInstance* InEmitter, TArray<UMaterialInterface*>& OutMaterials) const
 {
 	//OutMaterials.Add(Material);
diff --git a/Engine/Plugins/FX/Niagara/Source/Niagara/Public/NiagaraLightRendererProperties.h b/Engine/Plugins/FX/Niagara/Source/Niagara/Public/NiagaraLightRendererProperties.h
--- a/Engine/Plugins/FX/Niagara/Source/Niagara/Public/NiagaraLightRendererProperties.h
+++ b/Engine/Plugins/FX/Niagara/Source/Niagara/Public/NiagaraLightRendererProperties.h
@@ -29,7 +29,7 @@ public:
 
 	//~ UNiagaraRendererProperties interface
 	virtual FNiagaraRenderer* CreateEmitterRenderer(ERHIFeatureLevel::Type FeatureLevel, const FNiagaraEmitterInstance* Emitter, const FNiagaraSystemInstanceController& InController) override;
-	virtual class FNiagaraBoundsCalculator* CreateBoundsCalculator() override { return nullptr; }
+	virtual class FNiagaraBoundsCalculator* CreateBoundsCalculator() override;
 	virtual void GetUsedMaterials(const FNiagaraEmitterInstance* InEmitter, TArray<UMaterialInterface*>& OutMaterials) const override;
 	virtual bool IsSimTargetSupported(ENiagaraSimTarget InSimTarget) const override { return (InSimTarget == ENiagaraSimTarget::CPUSim); };
 #if WITH_EDITORONLY_DATA
