Add some additional debug drawing to the RigidMeshDI to highlight the actors generating the rigid bodies

diff --git a/Engine/Plugins/FX/Niagara/Source/Niagara/Private/NiagaraDataInterfaceRigidMeshCollisionQuery.cpp b/Engine/Plugins/FX/Niagara/Source/Niagara/Private/NiagaraDataInterfaceRigidMeshCollisionQuery.cpp
--- a/Engine/Plugins/FX/Niagara/Source/Niagara/Private/NiagaraDataInterfaceRigidMeshCollisionQuery.cpp
+++ b/Engine/Plugins/FX/Niagara/Source/Niagara/Private/NiagaraDataInterfaceRigidMeshCollisionQuery.cpp
@@ -819,6 +819,40 @@ void UNiagaraDataInterfaceRigidMeshCollisionQuery::DrawDebugHud(UCanvas* Canvas,
 				DrawDebugCanvasCapsule(Canvas, CurrentTransform, HalfTotalLength, RadiusLength.X, FColor::Blue);
 			}
 		}
+
+		if (!InstanceData_GT->ExplicitActors.IsEmpty() || !InstanceData_GT->FoundActors.IsEmpty())
+		{
+			const UFont* Font = GEngine->GetMediumFont();
+
+			auto DrawDebugActor = [&](TWeakObjectPtr<AActor>& InWeakActor, const TCHAR* ActorSourceString)
+			{
+				if (AActor* Actor = InWeakActor.Get())
+				{
+					FVector ActorOrigin;
+					FVector ActorBoundsExtent;
+					Actor->GetActorBounds(true, ActorOrigin, ActorBoundsExtent);
+
+					const FMatrix CurrentTransform = FTranslationMatrix(ActorOrigin);
+					if (!ShouldClip(Canvas, CurrentTransform, FSphere(FVector::ZeroVector, ActorBoundsExtent.Size())))
+					{
+						DrawDebugCanvasWireBox(Canvas, CurrentTransform, FBox(-ActorBoundsExtent, ActorBoundsExtent), FColor::Yellow);
+
+						const FVector ScreenLoc = Canvas->Project(ActorOrigin);
+						Canvas->DrawText(Font, FString::Printf(TEXT("RigidMeshDI[%s Actor] - %s"), ActorSourceString, *Actor->GetName()), ScreenLoc.X, ScreenLoc.Y);
+					}
+				}
+			};
+
+			for (TWeakObjectPtr<AActor>& ExplicitActor : InstanceData_GT->ExplicitActors)
+			{
+				DrawDebugActor(ExplicitActor, TEXT("Explicit"));
+			}
+
+			for (TWeakObjectPtr<AActor>& FoundActor : InstanceData_GT->FoundActors)
+			{
+				DrawDebugActor(FoundActor, TEXT("Found"));
+			}
+		}
 	}
 }
 #endif
