- GPU ribbon fixes
- Fix some OOB buffer access from various shaders
- Fix invalid indirect arg generation if counts are 0
- Don't generate dynamic data if the FX system is missing, various paths rely on the FX system (GPU sorting & culling, GPU ribbon generation)
- Remove twist permutation from FNiagaraRibbonUVParamCalculationCS as it's not required
- Fix missing twist permutations select from vertex reduction init / propagation leading to GPU hangs / incorrect tessellation

diff --git a/Engine/Plugins/FX/Niagara/Shaders/Private/NiagaraRibbonVertexFactory.ush b/Engine/Plugins/FX/Niagara/Shaders/Private/NiagaraRibbonVertexFactory.ush
--- a/Engine/Plugins/FX/Niagara/Shaders/Private/NiagaraRibbonVertexFactory.ush
+++ b/Engine/Plugins/FX/Niagara/Shaders/Private/NiagaraRibbonVertexFactory.ush
@@ -472,12 +472,15 @@ FMaterialVertexParameters GetMaterialVertexParameters(FVertexFactoryInput Input,
 	return Result;
 }
 
-#if 0
-float3 SafeNormalize(float3 V)
+float3 SafeNormalize(float3 V, float3 Default)
 {
-	return rsqrt(max(0.001, dot(V,V))) * V;
+	const float len2 = dot(V, V);
+	if ( len2 > 0.0f )
+	{
+		return V * rsqrt(len2);
+	}
+	return Default;
 }
-#endif
 
 float4 GetDirectionAndMagnitude(float3 V)
 {
@@ -536,7 +539,7 @@ FRibbonPositionIntermediates GetRibbonPositionIntermediates(
 			GetTangentAndDistanceComponent(InterpVtxData.RawParticleId1, 3));
 		
 		Intermediates.Position = TransformPosition(LocalToWorld, CubicInterp(Pos0, TangentAndDist0.xyz, Pos1, TangentAndDist1.xyz, InterpVtxData.Alpha, TangentAndDist1.w - TangentAndDist0.w));
-		Intermediates.Direction = normalize(TransformVector(LocalToWorld, InvScale, lerp(TangentAndDist0.xyz, TangentAndDist1.xyz, InterpVtxData.Alpha)));
+		Intermediates.Direction = SafeNormalize(TransformVector(LocalToWorld, InvScale, lerp(TangentAndDist0.xyz, TangentAndDist1.xyz, InterpVtxData.Alpha)), float3(0, 0, 0));
 		Intermediates.DistanceOnSegment = lerp(TangentAndDist0.w, TangentAndDist1.w, InterpVtxData.Alpha);
 	}
 
@@ -547,14 +550,16 @@ FRibbonPositionIntermediates GetRibbonPositionIntermediates(
 	{
 		const float3 Facing0 = GetVec3(FacingDataOffset, InterpVtxData.ParticleId0);
 		const float3 Facing1 = GetVec3(FacingDataOffset, InterpVtxData.ParticleId1);
-		Facing = normalize(TransformVector(LocalToWorld, InvScale, lerp(Facing0, Facing1, InterpVtxData.Alpha)));
+		Facing = TransformVector(LocalToWorld, InvScale, lerp(Facing0, Facing1, InterpVtxData.Alpha));
 	}
 	else
 	{
-		Facing = normalize(LWCToFloat(LWCSubtract(ResolvedView.WorldCameraOrigin, Intermediates.Position)));
+		Facing = LWCToFloat(LWCSubtract(ResolvedView.WorldCameraOrigin, Intermediates.Position));
 	}
+	Facing = SafeNormalize(Facing, float3(0, 0, 0));
 
-	float3 Right = normalize(cross(Facing, Intermediates.Direction));
+	float3 Right = cross(Facing, Intermediates.Direction);
+	Right = SafeNormalize(Right, float3(0, 0, 0));
 	
 	if (NiagaraRibbonVFLooseParameters.FacingMode == 2 /* CustomSide */)
 	{
@@ -588,11 +593,11 @@ FRibbonPositionIntermediates GetRibbonPositionIntermediates(
 	
 	Intermediates.TangentZ = NiagaraRibbonVFLooseParameters.Shape >= 1 ? Intermediates.GeoUp : Facing; // Intermediates.GeoUp;
 	
-	Intermediates.TangentZ = normalize((Intermediates.TangentZ * SliceVertexNormals.y + Intermediates.TangentX * SliceVertexNormals.x));
+	Intermediates.TangentZ = SafeNormalize(Intermediates.TangentZ * SliceVertexNormals.y + Intermediates.TangentX * SliceVertexNormals.x, float3(0, 0, 1));
 	
 	// Flip the normal only for multiplane shape if the normal is looking away from the camera.
 	// This lets us still render as two sided, but at least gets the normals corrected so the plane faces the camera.
-	float3 ViewNormal = normalize(LWCToFloat(LWCSubtract(ResolvedView.WorldCameraOrigin, Intermediates.Position)));
+	float3 ViewNormal = SafeNormalize(LWCToFloat(LWCSubtract(ResolvedView.WorldCameraOrigin, Intermediates.Position)), float3(0,0,0));
 	Intermediates.TangentZ *= (NiagaraRibbonVF.ShouldFlipNormalToView > 0 && dot(ViewNormal, Intermediates.TangentZ) < 0) ? -1 : 1;
 	
 	Intermediates.TangentX = cross(Intermediates.TangentZ, Intermediates.TangentY);
diff --git a/Engine/Plugins/FX/Niagara/Shaders/Private/Ribbons/NiagaraRibbonGenerateIndices.usf b/Engine/Plugins/FX/Niagara/Shaders/Private/Ribbons/NiagaraRibbonGenerateIndices.usf
--- a/Engine/Plugins/FX/Niagara/Shaders/Private/Ribbons/NiagaraRibbonGenerateIndices.usf
+++ b/Engine/Plugins/FX/Niagara/Shaders/Private/Ribbons/NiagaraRibbonGenerateIndices.usf
@@ -37,9 +37,6 @@ uint SegmentBitMask;
 uint SubSegmentBitShift;
 uint SubSegmentBitMask;
 
-
-
-
 uint PackIndex(uint SortedIndex, uint SubSegmentID, uint SliceVertex)
 {
 	return	((SortedIndex	& SegmentBitMask)		<< SegmentBitShift) |
@@ -47,12 +44,11 @@ uint PackIndex(uint SortedIndex, uint SubSegmentID, uint SliceVertex)
 			((SliceVertex	& BitMaskForShape));
 }
 
-
 [numthreads(THREADGROUP_SIZE, 1, 1)]
 void GenerateIndices(uint3 DispatchThreadId : SV_DispatchThreadID)
 {
 	const uint TotalNumParticles = GetTotalNumParticles();
-	const uint MaxPossibleSegments = TotalNumParticles - 1;
+	const uint MaxPossibleSegments = TotalNumParticles > 0 ? TotalNumParticles - 1 : 0;
 	
 	const uint IndirectDrawOffset = IndirectDrawInfoIndex * INDEX_GEN_INDIRECT_ARGS_STRIDE;
 	
@@ -106,9 +102,9 @@ void GenerateIndices(uint3 DispatchThreadId : SV_DispatchThreadID)
 				GeneratedIndicesBuffer[IndicesOffset + 0] = CurrSegment | FirstIndex;
 				GeneratedIndicesBuffer[IndicesOffset + 1] = CurrSegment | SecondIndex;				
 				GeneratedIndicesBuffer[IndicesOffset + 2] = NextSegment | FirstIndex;				
-				GeneratedIndicesBuffer[IndicesOffset + 3] = GeneratedIndicesBuffer[IndicesOffset + 1];				
+				GeneratedIndicesBuffer[IndicesOffset + 3] = CurrSegment | SecondIndex;	//GeneratedIndicesBuffer[IndicesOffset + 1];				
 				GeneratedIndicesBuffer[IndicesOffset + 4] = NextSegment | SecondIndex;				
-				GeneratedIndicesBuffer[IndicesOffset + 5] = GeneratedIndicesBuffer[IndicesOffset + 2];
+				GeneratedIndicesBuffer[IndicesOffset + 5] = NextSegment | FirstIndex;	//GeneratedIndicesBuffer[IndicesOffset + 2];
 
 				IndicesOffset += 6;
 			}
@@ -126,8 +122,8 @@ void GenerateIndices(uint3 DispatchThreadId : SV_DispatchThreadID)
 				GeneratedIndicesBuffer[IndicesOffset + 0] = CurrSegment | FirstIndex;				
 				GeneratedIndicesBuffer[IndicesOffset + 1] = CurrSegment | SecondIndex;				
 				GeneratedIndicesBuffer[IndicesOffset + 2] = NextSegment | SecondIndex;				
-				GeneratedIndicesBuffer[IndicesOffset + 3] = GeneratedIndicesBuffer[IndicesOffset + 0];				
-				GeneratedIndicesBuffer[IndicesOffset + 4] = GeneratedIndicesBuffer[IndicesOffset + 2];				
+				GeneratedIndicesBuffer[IndicesOffset + 3] = CurrSegment | FirstIndex;	//GeneratedIndicesBuffer[IndicesOffset + 0];				
+				GeneratedIndicesBuffer[IndicesOffset + 4] = NextSegment | SecondIndex;	//GeneratedIndicesBuffer[IndicesOffset + 2];				
 				GeneratedIndicesBuffer[IndicesOffset + 5] = NextSegment | FirstIndex;
 				
 				IndicesOffset += 6;
diff --git a/Engine/Plugins/FX/Niagara/Shaders/Private/Ribbons/NiagaraRibbonInitializeIndices.usf b/Engine/Plugins/FX/Niagara/Shaders/Private/Ribbons/NiagaraRibbonInitializeIndices.usf
--- a/Engine/Plugins/FX/Niagara/Shaders/Private/Ribbons/NiagaraRibbonInitializeIndices.usf
+++ b/Engine/Plugins/FX/Niagara/Shaders/Private/Ribbons/NiagaraRibbonInitializeIndices.usf
@@ -31,11 +31,32 @@ float GNiagaraRibbonTessellationMinDisplacementError;
 #define TESSELLATIONMODE_CUSTOM 1
 #define TESSELLATIONMODE_DISABLED 2
 
-
 [numthreads(1, 1, 1)]
 void InitializeIndices(uint3 DispatchThreadId : SV_DispatchThreadID)
 {
 	const uint TotalNumParticles = GetTotalNumParticles();
+	const uint IndirectDrawOffset = IndirectDrawOutputIndex * INDEX_GEN_INDIRECT_ARGS_STRIDE;
+
+	// Early out for no particles
+	if (TotalNumParticles == 0)
+	{
+		IndirectDrawOutput[IndirectDrawOffset + INDEX_GEN_EXECUTE_INDIRECT_OFFSET + 0] = 0;
+		IndirectDrawOutput[IndirectDrawOffset + INDEX_GEN_EXECUTE_INDIRECT_OFFSET + 1] = 0;
+		IndirectDrawOutput[IndirectDrawOffset + INDEX_GEN_EXECUTE_INDIRECT_OFFSET + 2] = 0;
+		IndirectDrawOutput[IndirectDrawOffset + INDEX_GEN_TESSELLATION_FACTOR_OFFSET] = 0;
+			
+		IndirectDrawOutput[IndirectDrawOffset + INDEX_GEN_INDIRECT_DRAW_OFFSET + 0] = 0;
+		IndirectDrawOutput[IndirectDrawOffset + INDEX_GEN_INDIRECT_DRAW_OFFSET + 1] = 0;
+		IndirectDrawOutput[IndirectDrawOffset + INDEX_GEN_INDIRECT_DRAW_OFFSET + 2] = 0;
+		IndirectDrawOutput[IndirectDrawOffset + INDEX_GEN_INDIRECT_DRAW_OFFSET + 3] = 0;
+		IndirectDrawOutput[IndirectDrawOffset + INDEX_GEN_INDIRECT_DRAW_OFFSET + 4] = 0;
+	
+		IndirectDrawOutput[IndirectDrawOffset + INDEX_GEN_NUM_SEGMENTS_OFFSET] = 0;
+		IndirectDrawOutput[IndirectDrawOffset + INDEX_GEN_NUM_SUB_SEGMENTS_OFFSET] = 0;
+		IndirectDrawOutput[IndirectDrawOffset + INDEX_GEN_ONE_OVER_NUM_SUB_SEGMENTS_OFFSET] = asuint(0.0f);
+		return;
+	}
+
 #if RIBBONS_WANTS_AUTOMATIC_TESSELLATION
 	
 	bool bUseConstantFactor = false;
@@ -68,7 +89,6 @@ void InitializeIndices(uint3 DispatchThreadId : SV_DispatchThreadID)
 	float TessellationTwistCurvature = asfloat(VertexGenerationResults[VertexGenerationResultsIndex * VERTEX_GEN_OUTPUT_DATA_STRIDE + VERTEX_GEN_OUTPUT_DATA_TESSELLATION_TWIST_CURVATURE]);
 	float TessellationTotalSegmentLength = asfloat(VertexGenerationResults[VertexGenerationResultsIndex * VERTEX_GEN_OUTPUT_DATA_STRIDE + VERTEX_GEN_OUTPUT_DATA_TESSELLATION_TOTAL_SEGMENT_LENGTH]);
 
-
 	if (bUseConstantFactor)
 	{
 		SegmentTessellation = TessellationFactor;
@@ -91,6 +111,9 @@ void InitializeIndices(uint3 DispatchThreadId : SV_DispatchThreadID)
 			Tess = max(TwistTess, Tess);
 		}
 		SegmentTessellation = clamp(uint(round(Tess)), uint(round(MinTesselation)), TessellationFactor);
+
+		// Ensure we don't blow the max tessellation allowed, this shouldn't happen but let's be safe
+		SegmentTessellation = min(SegmentTessellation, GNiagaraRibbonMaxTessellation);
 	}
 #elif RIBBONS_WANTS_CONSTANT_TESSELLATION
 	// Has a constant factor so just use it
@@ -106,8 +129,6 @@ void InitializeIndices(uint3 DispatchThreadId : SV_DispatchThreadID)
 
 	const uint SegmentCount = VertexGenerationResults[VertexGenerationResultsIndex * VERTEX_GEN_OUTPUT_DATA_STRIDE + VERTEX_GEN_OUTPUT_DATA_TOTAL_NUM_SEGMENTS_OFFSET];
 
-	const uint IndirectDrawOffset = IndirectDrawOutputIndex * INDEX_GEN_INDIRECT_ARGS_STRIDE;
-
 	const float OneOverSubSegmentTessellation = 1.0f / float(SegmentTessellation);
 	
 	// Output the dispatch indirect args followed by the segmentation amount
@@ -134,8 +155,4 @@ void InitializeIndices(uint3 DispatchThreadId : SV_DispatchThreadID)
 	IndirectDrawOutput[IndirectDrawOffset + INDEX_GEN_NUM_SUB_SEGMENTS_OFFSET] = SegmentTessellation;
 	// 1/SubSegmentCount
 	IndirectDrawOutput[IndirectDrawOffset + INDEX_GEN_ONE_OVER_NUM_SUB_SEGMENTS_OFFSET] = asuint(OneOverSubSegmentTessellation);
-
 }
-
-
-
diff --git a/Engine/Plugins/FX/Niagara/Shaders/Private/Ribbons/NiagaraRibbonRibbonUVParamCalculation.usf b/Engine/Plugins/FX/Niagara/Shaders/Private/Ribbons/NiagaraRibbonRibbonUVParamCalculation.usf
--- a/Engine/Plugins/FX/Niagara/Shaders/Private/Ribbons/NiagaraRibbonRibbonUVParamCalculation.usf
+++ b/Engine/Plugins/FX/Niagara/Shaders/Private/Ribbons/NiagaraRibbonRibbonUVParamCalculation.usf
@@ -15,8 +15,6 @@ RWBuffer<uint> PackedPerRibbonData;
 Buffer<uint> CommandBuffer;
 int CommandBufferOffset;
 
-int TotalNumRibbons;
-
 // Settings for UV Channel 0
 float2 UV0Settings_Offset;
 float2 UV0Settings_Scale;
@@ -91,7 +89,6 @@ struct FRibbonUVScaleAndOffsets
 #define UV_DISTRIBUTION_MODE_TILED_OVER_RIBBON_LENGTH 2
 #define UV_DISTRIBUTION_MODE_TILED_FROM_START_OVER_RIBBON_LENGTH 3
 
-
 float3 LoadTangent(int Index)
 {
 	return float3(
@@ -203,7 +200,6 @@ FRibbonUVScaleAndOffsets CalculateUVScaleAndOffsets(const FNiagaraRibbonUVSettin
 	return Output;
 }
 
-
 // We only need a single thread if we're not accounting for ribbon ids
 #if HAS_RIBBON_ID
 [numthreads(THREADGROUP_SIZE, 1, 1)]
@@ -212,7 +208,6 @@ FRibbonUVScaleAndOffsets CalculateUVScaleAndOffsets(const FNiagaraRibbonUVSettin
 #endif
 void GenerateRibbonUVParams(uint3 DispatchThreadId : SV_DispatchThreadID)
 {
-	const uint TotalNumParticles = GetTotalNumParticles();
 #if HAS_RIBBON_ID
 	const int RibbonIndex = DispatchThreadId.x;
 	const int TotalNumRibbons = CommandBuffer[CommandBufferOffset * VERTEX_GEN_OUTPUT_DATA_STRIDE + VERTEX_GEN_OUTPUT_DATA_TOTAL_NUM_RIBBONS_OFFSET];
@@ -229,49 +224,47 @@ void GenerateRibbonUVParams(uint3 DispatchThreadId : SV_DispatchThreadID)
 
 		// Grab the total length from the last particle
 		const float TotalLength = GetSegmentDistance(LastParticleID);
-		
+		const uint NumSegments = LastParticleID - FirstParticleID;
+	
 		const FNiagaraRibbonUVSettings UV0Settings = GetUV0Params();
 		const FNiagaraRibbonUVSettings UV1Settings = GetUV1Params();
 
 		const int PackedPerRibbonDataOffsetUVs = PackedPerRibbonDataOffset + PACKED_PER_RIBBON_UVDATAOFFSET;
-		
-		// UV0
-		BRANCH
-		if (UV0Settings.bEnablePerParticleUOverride && U0OverrideDataOffset != INDEX_NONE)
+		if ( NumSegments > 1 )
 		{
-			PackedPerRibbonData[PackedPerRibbonDataOffsetUVs + 0] = asuint(1.0f);
-			PackedPerRibbonData[PackedPerRibbonDataOffsetUVs + 1] = asuint(0.0f);
-			PackedPerRibbonData[PackedPerRibbonDataOffsetUVs + 2] = asuint(1.0f);
-		}
-		else
-		{
-			const FRibbonUVScaleAndOffsets UVData = CalculateUVScaleAndOffsets(UV0Settings, FirstParticleID, LastParticleID, TotalLength);
-			PackedPerRibbonData[PackedPerRibbonDataOffsetUVs + 0] = asuint(UVData.UScale);
-			PackedPerRibbonData[PackedPerRibbonDataOffsetUVs + 1] = asuint(UVData.UOffset);
-			PackedPerRibbonData[PackedPerRibbonDataOffsetUVs + 2] = asuint(UVData.UDistributionScalar);
-		}
+			// UV0
+			BRANCH
+			if (UV0Settings.bEnablePerParticleUOverride && U0OverrideDataOffset != INDEX_NONE)
+			{
+				PackedPerRibbonData[PackedPerRibbonDataOffsetUVs + 0] = asuint(1.0f);
+				PackedPerRibbonData[PackedPerRibbonDataOffsetUVs + 1] = asuint(0.0f);
+				PackedPerRibbonData[PackedPerRibbonDataOffsetUVs + 2] = asuint(1.0f);
+			}
+			else
+			{
+				const FRibbonUVScaleAndOffsets UVData = CalculateUVScaleAndOffsets(UV0Settings, FirstParticleID, LastParticleID, TotalLength);
+				PackedPerRibbonData[PackedPerRibbonDataOffsetUVs + 0] = asuint(UVData.UScale);
+				PackedPerRibbonData[PackedPerRibbonDataOffsetUVs + 1] = asuint(UVData.UOffset);
+				PackedPerRibbonData[PackedPerRibbonDataOffsetUVs + 2] = asuint(UVData.UDistributionScalar);
+			}
 		
-		// UV1
-		BRANCH
-		if (UV1Settings.bEnablePerParticleUOverride && U1OverrideDataOffset != INDEX_NONE)
-		{
-			PackedPerRibbonData[PackedPerRibbonDataOffsetUVs + 3] = asuint(1.0f);
-			PackedPerRibbonData[PackedPerRibbonDataOffsetUVs + 4] = asuint(0.0f);
-			PackedPerRibbonData[PackedPerRibbonDataOffsetUVs + 5] = asuint(1.0f);
-		}
-		else
-		{
-			const FRibbonUVScaleAndOffsets UVData = CalculateUVScaleAndOffsets(UV1Settings, FirstParticleID, LastParticleID, TotalLength);
-			PackedPerRibbonData[PackedPerRibbonDataOffsetUVs + 3] = asuint(UVData.UScale);
-			PackedPerRibbonData[PackedPerRibbonDataOffsetUVs + 4] = asuint(UVData.UOffset);
-			PackedPerRibbonData[PackedPerRibbonDataOffsetUVs + 5] = asuint(UVData.UDistributionScalar);
-		}
-
-		const uint NumSegments = LastParticleID - FirstParticleID;
-		
-		// Update the tangents for the first and last vertex, apply a reflect vector logic so that the initial and final curvature is continuous.
-		if (NumSegments > 1)
-		{
+			// UV1
+			BRANCH
+			if (UV1Settings.bEnablePerParticleUOverride && U1OverrideDataOffset != INDEX_NONE)
+			{
+				PackedPerRibbonData[PackedPerRibbonDataOffsetUVs + 3] = asuint(1.0f);
+				PackedPerRibbonData[PackedPerRibbonDataOffsetUVs + 4] = asuint(0.0f);
+				PackedPerRibbonData[PackedPerRibbonDataOffsetUVs + 5] = asuint(1.0f);
+			}
+			else
+			{
+				const FRibbonUVScaleAndOffsets UVData = CalculateUVScaleAndOffsets(UV1Settings, FirstParticleID, LastParticleID, TotalLength);
+				PackedPerRibbonData[PackedPerRibbonDataOffsetUVs + 3] = asuint(UVData.UScale);
+				PackedPerRibbonData[PackedPerRibbonDataOffsetUVs + 4] = asuint(UVData.UOffset);
+				PackedPerRibbonData[PackedPerRibbonDataOffsetUVs + 5] = asuint(UVData.UDistributionScalar);
+			}
+
+			// Update the tangents for the first and last vertex, apply a reflect vector logic so that the initial and final curvature is continuous.
 			const float3 FirstTangent = LoadTangent(FirstParticleID);
 			const float3 NextToFirstTangent = LoadTangent(FirstParticleID + 1);
 			StoreTangent(FirstParticleID, (2.0f * dot(FirstTangent, NextToFirstTangent)) * FirstTangent - NextToFirstTangent);
@@ -280,8 +273,15 @@ void GenerateRibbonUVParams(uint3 DispatchThreadId : SV_DispatchThreadID)
 			const float3 PrevToLastTangent = LoadTangent(LastParticleID - 1);
 			StoreTangent(LastParticleID, (2.0f * dot(LastTangent, PrevToLastTangent)) * LastTangent - PrevToLastTangent);
 		}
-		
-		
+		else
+		{
+			PackedPerRibbonData[PackedPerRibbonDataOffsetUVs + 0] = asuint(1.0f);
+			PackedPerRibbonData[PackedPerRibbonDataOffsetUVs + 1] = asuint(0.0f);
+			PackedPerRibbonData[PackedPerRibbonDataOffsetUVs + 2] = asuint(1.0f);
+			PackedPerRibbonData[PackedPerRibbonDataOffsetUVs + 3] = asuint(1.0f);
+			PackedPerRibbonData[PackedPerRibbonDataOffsetUVs + 4] = asuint(0.0f);
+			PackedPerRibbonData[PackedPerRibbonDataOffsetUVs + 5] = asuint(1.0f);
+		}
 #if HAS_RIBBON_ID
 	}
 #endif
diff --git a/Engine/Plugins/FX/Niagara/Shaders/Private/Ribbons/NiagaraRibbonVertexReductionFinalization.usf b/Engine/Plugins/FX/Niagara/Shaders/Private/Ribbons/NiagaraRibbonVertexReductionFinalization.usf
--- a/Engine/Plugins/FX/Niagara/Shaders/Private/Ribbons/NiagaraRibbonVertexReductionFinalization.usf
+++ b/Engine/Plugins/FX/Niagara/Shaders/Private/Ribbons/NiagaraRibbonVertexReductionFinalization.usf
@@ -5,8 +5,6 @@
 
 
 Buffer<uint> SortedIndices;
-// Tangent of the particle, and the distance along the ribbon if we're computing it
-RWBuffer<float> TangentsAndDistances;
 // Index of the ribbon for this particle
 // created through prefix sum across all the starting points of the ribbons
 #if HAS_RIBBON_ID
@@ -45,6 +43,33 @@ int FinalizationThreadBlockSize;
 void OutputRibbonStats(uint3 DispatchThreadId : SV_DispatchThreadID)
 {
 	const uint TotalNumParticles = GetTotalNumParticles();
+	const uint OutputCommandBufferOffset = OutputCommandBufferIndex * VERTEX_GEN_OUTPUT_DATA_STRIDE;
+
+	// Early out if we have no particles
+	if ( TotalNumParticles == 0 )
+	{
+		OutputCommandBuffer[OutputCommandBufferOffset + VERTEX_GEN_OUTPUT_DATA_FINALIZATION_INDIRECT_ARGS_OFFSET + 0] = 0;
+		OutputCommandBuffer[OutputCommandBufferOffset + VERTEX_GEN_OUTPUT_DATA_FINALIZATION_INDIRECT_ARGS_OFFSET + 1] = 0;
+		OutputCommandBuffer[OutputCommandBufferOffset + VERTEX_GEN_OUTPUT_DATA_FINALIZATION_INDIRECT_ARGS_OFFSET + 2] = 0;
+		OutputCommandBuffer[OutputCommandBufferOffset + VERTEX_GEN_OUTPUT_DATA_TOTAL_NUM_SEGMENTS_OFFSET] = 0;
+		OutputCommandBuffer[OutputCommandBufferOffset + VERTEX_GEN_OUTPUT_DATA_TOTAL_NUM_RIBBONS_OFFSET] = 0;
+		OutputCommandBuffer[OutputCommandBufferOffset + VERTEX_GEN_OUTPUT_DATA_TOTAL_NUM_SEGMENTS_OFFSET] = 0;
+		OutputCommandBuffer[OutputCommandBufferOffset + VERTEX_GEN_OUTPUT_DATA_TOTAL_NUM_RIBBONS_OFFSET] = 0;
+	#if RIBBONS_WANTS_AUTOMATIC_TESSELLATION
+		OutputCommandBuffer[OutputCommandBufferOffset + VERTEX_GEN_OUTPUT_DATA_TESSELLATION_ANGLE] = asuint(0.0f);
+		OutputCommandBuffer[OutputCommandBufferOffset + VERTEX_GEN_OUTPUT_DATA_TESSELLATION_CURVATURE] = asuint(0.0f);
+		OutputCommandBuffer[OutputCommandBufferOffset + VERTEX_GEN_OUTPUT_DATA_TESSELLATION_TWIST_ANGLE] = asuint(0.0f);
+		OutputCommandBuffer[OutputCommandBufferOffset + VERTEX_GEN_OUTPUT_DATA_TESSELLATION_TWIST_CURVATURE] = asuint(0.0f);
+		OutputCommandBuffer[OutputCommandBufferOffset + VERTEX_GEN_OUTPUT_DATA_TESSELLATION_TOTAL_SEGMENT_LENGTH] = asuint(0.0f);
+		OutputCommandBuffer[OutputCommandBufferOffset + VERTEX_GEN_OUTPUT_DATA_TESSELLATION_CURRENT_FRAME_TOTAL_SEGMENT_LENGTH] = asuint(0.0f);
+		OutputCommandBuffer[OutputCommandBufferOffset + VERTEX_GEN_OUTPUT_DATA_TESSELLATION_CURRENT_FRAME_AVERAGE_SEGMENT_LENGTH] = asuint(0.0f);
+		OutputCommandBuffer[OutputCommandBufferOffset + VERTEX_GEN_OUTPUT_DATA_TESSELLATION_CURRENT_FRAME_AVERAGE_SEGMENT_ANGLE] = asuint(0.0f);
+		OutputCommandBuffer[OutputCommandBufferOffset + VERTEX_GEN_OUTPUT_DATA_TESSELLATION_CURRENT_FRAME_AVERAGE_TWIST_ANGLE] = asuint(0.0f);
+		OutputCommandBuffer[OutputCommandBufferOffset + VERTEX_GEN_OUTPUT_DATA_TESSELLATION_CURRENT_FRAME_AVERAGE_WIDTH] = asuint(0.0f);
+	#endif
+		return;
+	}
+
 	const uint RawParticleID = DispatchThreadId.x;
 	const uint LastParticleIndex = TotalNumParticles - 1;
 	
@@ -55,8 +80,8 @@ void OutputRibbonStats(uint3 DispatchThreadId : SV_DispatchThreadID)
 	{
 		// First we Gen the ribbon lookup table if we have multiribbon enabled
 		const uint CurrRibbonIndex = MultiRibbonIndices[RawParticleID];
-		const uint PrevRibbonIndex = RawParticleID > 0? MultiRibbonIndices[RawParticleID - 1] : INDEX_NONE;
-		const uint NextRibbonIndex = RawParticleID < (TotalNumParticles - 1)? MultiRibbonIndices[RawParticleID + 1] : INDEX_NONE;
+		const uint PrevRibbonIndex = RawParticleID > 0 ? MultiRibbonIndices[RawParticleID - 1] : INDEX_NONE;
+		const uint NextRibbonIndex = RawParticleID < LastParticleIndex ? MultiRibbonIndices[RawParticleID + 1] : INDEX_NONE;
 		
 		// Track the starting index of each ribbon segment
 		BRANCH
@@ -78,7 +103,6 @@ void OutputRibbonStats(uint3 DispatchThreadId : SV_DispatchThreadID)
 	PackedPerRibbonData[PACKED_PER_RIBBON_ENDPARTICLEOFFSET] = LastParticleIndex;
 #endif
 
-
 	// Very first compute instance, we grab all the stats into the output buffer
 	// This is so if we don't need multiribbon, then we'll only run one compute instance
 #if HAS_RIBBON_ID	
@@ -92,14 +116,11 @@ void OutputRibbonStats(uint3 DispatchThreadId : SV_DispatchThreadID)
 		const int TotalNumRibbons = MultiRibbonIndices[LastParticleIndex] + 1;
 #endif		
 
-		const uint OutputCommandBufferOffset = OutputCommandBufferIndex * VERTEX_GEN_OUTPUT_DATA_STRIDE;
-
 		// Calculate the indirect args for the finalization shader
 		OutputCommandBuffer[OutputCommandBufferOffset + VERTEX_GEN_OUTPUT_DATA_FINALIZATION_INDIRECT_ARGS_OFFSET + 0] = DivideAndRoundUp(TotalNumRibbons, FinalizationThreadBlockSize);
 		OutputCommandBuffer[OutputCommandBufferOffset + VERTEX_GEN_OUTPUT_DATA_FINALIZATION_INDIRECT_ARGS_OFFSET + 1] = 1;
 		OutputCommandBuffer[OutputCommandBufferOffset + VERTEX_GEN_OUTPUT_DATA_FINALIZATION_INDIRECT_ARGS_OFFSET + 2] = 1;
 		
-
 		// Total num segments, we can get it from taking the last segment index and adding one to get the count (segments is one less in length than sorted particles, so we need the 2nd to last particle
 		const int TotalSegments = Segments[LastParticleIndex - 1] + 1;
 		OutputCommandBuffer[OutputCommandBufferOffset + VERTEX_GEN_OUTPUT_DATA_TOTAL_NUM_SEGMENTS_OFFSET] = TotalSegments;
@@ -107,7 +128,6 @@ void OutputRibbonStats(uint3 DispatchThreadId : SV_DispatchThreadID)
 		// Total num ribbons
 		OutputCommandBuffer[OutputCommandBufferOffset + VERTEX_GEN_OUTPUT_DATA_TOTAL_NUM_RIBBONS_OFFSET] = TotalNumRibbons;
 
-
 #if RIBBONS_WANTS_AUTOMATIC_TESSELLATION
 #if RIBBON_HAS_TWIST
 		const uint TessellationBaseOffset = LastParticleIndex * 5;
@@ -125,14 +145,22 @@ void OutputRibbonStats(uint3 DispatchThreadId : SV_DispatchThreadID)
 		float AverageWidth = 0;
 #endif
 		
-		float TessellationAngle = asfloat(OutputCommandBuffer[OutputCommandBufferOffset + VERTEX_GEN_OUTPUT_DATA_TESSELLATION_ANGLE]);
-		float TessellationCurvature = asfloat(OutputCommandBuffer[OutputCommandBufferOffset + VERTEX_GEN_OUTPUT_DATA_TESSELLATION_CURVATURE]);
-		float TessellationTwistAngle = asfloat(OutputCommandBuffer[OutputCommandBufferOffset + VERTEX_GEN_OUTPUT_DATA_TESSELLATION_TWIST_ANGLE]);
-		float TessellationTwistCurvature = asfloat(OutputCommandBuffer[OutputCommandBufferOffset + VERTEX_GEN_OUTPUT_DATA_TESSELLATION_TWIST_CURVATURE]);
-		float TessellationTotalSegmentLength = asfloat(OutputCommandBuffer[OutputCommandBufferOffset + VERTEX_GEN_OUTPUT_DATA_TESSELLATION_TOTAL_SEGMENT_LENGTH]);		
+		// Generate smoothed tessellation information
+		float TessellationAngle = 0.0f;
+		float TessellationCurvature = 0.0f;
+		float TessellationTwistAngle = 0.0f;
+		float TessellationTwistCurvature = 0.0f;
+		float TessellationTotalSegmentLength = 0.0f;
 		
 		if (TotalSegmentLength > 0.0)
 		{				
+			// Read Tessellation Smmothing Data (accumulated between frames)
+			TessellationAngle = asfloat(OutputCommandBuffer[OutputCommandBufferOffset + VERTEX_GEN_OUTPUT_DATA_TESSELLATION_ANGLE]);
+			TessellationCurvature = asfloat(OutputCommandBuffer[OutputCommandBufferOffset + VERTEX_GEN_OUTPUT_DATA_TESSELLATION_CURVATURE]);
+			TessellationTwistAngle = asfloat(OutputCommandBuffer[OutputCommandBufferOffset + VERTEX_GEN_OUTPUT_DATA_TESSELLATION_TWIST_ANGLE]);
+			TessellationTwistCurvature = asfloat(OutputCommandBuffer[OutputCommandBufferOffset + VERTEX_GEN_OUTPUT_DATA_TESSELLATION_TWIST_CURVATURE]);
+			TessellationTotalSegmentLength = asfloat(OutputCommandBuffer[OutputCommandBufferOffset + VERTEX_GEN_OUTPUT_DATA_TESSELLATION_TOTAL_SEGMENT_LENGTH]);		
+
 			// Blend the result between the last frame tessellation factors and the current frame base on the total length of all segments.
 			// This is only used to increase the tessellation value of the current frame data to prevent glitches where tessellation is significantly changin between frames.
 			const float OneOverTotalSegmentLength = 1.f / max(1.0f, TotalSegmentLength);
@@ -146,7 +174,6 @@ void OutputRibbonStats(uint3 DispatchThreadId : SV_DispatchThreadID)
 			TessellationAngle = lerp(AverageSegmentAngle, max(TessellationAngle, AverageSegmentAngle), AveragingFactor);
 			TessellationCurvature = lerp(AverageSegmentCurvature, max(TessellationCurvature, AverageSegmentCurvature), AveragingFactor);
 			
-
 #if RIBBON_HAS_TWIST				
 			AverageTwistAngle *= OneOverTotalSegmentLength;
 			AverageWidth *= OneOverTotalSegmentLength;
@@ -155,34 +182,20 @@ void OutputRibbonStats(uint3 DispatchThreadId : SV_DispatchThreadID)
 			TessellationTwistCurvature = lerp(AverageWidth, max(TessellationTwistCurvature, AverageWidth), AveragingFactor);		
 #endif
 		}
-		else
-		{
-			TessellationAngle = 0;
-			TessellationCurvature = 0;
-			TessellationTwistAngle = 0;
-			TessellationTwistCurvature = 0;
-			TessellationTotalSegmentLength = 0;
-		}
 
-		// TessellationAngle
+		// Output Tessellation Smmothing Data (accumulated between frames)
 		OutputCommandBuffer[OutputCommandBufferOffset + VERTEX_GEN_OUTPUT_DATA_TESSELLATION_ANGLE] = asuint(TessellationAngle);
-		// TessellationCurvature
 		OutputCommandBuffer[OutputCommandBufferOffset + VERTEX_GEN_OUTPUT_DATA_TESSELLATION_CURVATURE] = asuint(TessellationCurvature);
-		// TessellationTwistAngle
 		OutputCommandBuffer[OutputCommandBufferOffset + VERTEX_GEN_OUTPUT_DATA_TESSELLATION_TWIST_ANGLE] = asuint(TessellationTwistAngle);
-		// TessellationTwistCurvature
 		OutputCommandBuffer[OutputCommandBufferOffset + VERTEX_GEN_OUTPUT_DATA_TESSELLATION_TWIST_CURVATURE] = asuint(TessellationTwistCurvature);
-		// TessellationTotalSegmentLength
 		OutputCommandBuffer[OutputCommandBufferOffset + VERTEX_GEN_OUTPUT_DATA_TESSELLATION_TOTAL_SEGMENT_LENGTH] = asuint(TessellationTotalSegmentLength);
 
-
 		// Write out current frame tessellation stats
 		OutputCommandBuffer[OutputCommandBufferOffset + VERTEX_GEN_OUTPUT_DATA_TESSELLATION_CURRENT_FRAME_TOTAL_SEGMENT_LENGTH] = asuint(TotalSegmentLength);
 		OutputCommandBuffer[OutputCommandBufferOffset + VERTEX_GEN_OUTPUT_DATA_TESSELLATION_CURRENT_FRAME_AVERAGE_SEGMENT_LENGTH] = asuint(AverageSegmentLength);
 		OutputCommandBuffer[OutputCommandBufferOffset + VERTEX_GEN_OUTPUT_DATA_TESSELLATION_CURRENT_FRAME_AVERAGE_SEGMENT_ANGLE] = asuint(AverageSegmentAngle);
 		OutputCommandBuffer[OutputCommandBufferOffset + VERTEX_GEN_OUTPUT_DATA_TESSELLATION_CURRENT_FRAME_AVERAGE_TWIST_ANGLE] = asuint(AverageTwistAngle);
 		OutputCommandBuffer[OutputCommandBufferOffset + VERTEX_GEN_OUTPUT_DATA_TESSELLATION_CURRENT_FRAME_AVERAGE_WIDTH] = asuint(AverageWidth);
-
-#endif
+#endif //RIBBONS_WANTS_AUTOMATIC_TESSELLATION
 	}
 }
diff --git a/Engine/Plugins/FX/Niagara/Source/Niagara/Private/NiagaraComponent.cpp b/Engine/Plugins/FX/Niagara/Source/Niagara/Private/NiagaraComponent.cpp
--- a/Engine/Plugins/FX/Niagara/Source/Niagara/Private/NiagaraComponent.cpp
+++ b/Engine/Plugins/FX/Niagara/Source/Niagara/Private/NiagaraComponent.cpp
@@ -405,7 +405,12 @@ void FNiagaraSceneProxy::GetDynamicMeshElements(const TArray<const FSceneView*>&
 
 	if (RenderData)
 	{
-		RenderData->GetDynamicMeshElements(Views, ViewFamily, VisibilityMap, Collector, *this);
+		// Renderers often rely upon executing GPU code, such as sorting / culling, if the FX system is invalid then the count manager, etc,
+		// will not run but the mesh batch building will assume they do.  We will early out here for consistency across all Niagara renderers.
+		if (ViewFamily.Scene->GetFXSystem() != nullptr)
+		{
+			RenderData->GetDynamicMeshElements(Views, ViewFamily, VisibilityMap, Collector, *this);
+		}
 	}
 
 	if (ViewFamily.EngineShowFlags.Particles && ViewFamily.EngineShowFlags.Niagara)
diff --git a/Engine/Plugins/FX/Niagara/Source/Niagara/Private/NiagaraRendererRibbons.cpp b/Engine/Plugins/FX/Niagara/Source/Niagara/Private/NiagaraRendererRibbons.cpp
--- a/Engine/Plugins/FX/Niagara/Source/Niagara/Private/NiagaraRendererRibbons.cpp
+++ b/Engine/Plugins/FX/Niagara/Source/Niagara/Private/NiagaraRendererRibbons.cpp
@@ -357,7 +357,7 @@ public:
 	}
 };
 
-bool FNiagaraRibbonGpuBuffer::Allocate(uint32 NumElements, uint32 MaxElements, ERHIAccess InResourceState, bool bGpuReadOnly)
+bool FNiagaraRibbonGpuBuffer::Allocate(uint32 NumElements, uint32 MaxElements, ERHIAccess InResourceState, bool bGpuReadOnly, EBufferUsageFlags AdditionalBufferUsage)
 {
 	if (NumElements == 0)
 	{
@@ -379,7 +379,7 @@ bool FNiagaraRibbonGpuBuffer::Allocate(uint32 NumElements, uint32 MaxElements, E
 			NumElements = FMath::Min(MaxElements, uint32(FMath::RoundToInt32(NumElements * UpsizeMultipler)));
 
 			FRHIResourceCreateInfo CreateInfo(DebugName);
-			const EBufferUsageFlags UsageFlags = BUF_ShaderResource | (bGpuReadOnly ? BUF_Volatile : BUF_Static | BUF_UnorderedAccess);
+			const EBufferUsageFlags UsageFlags = AdditionalBufferUsage | BUF_ShaderResource | (bGpuReadOnly ? BUF_Volatile : BUF_Static | BUF_UnorderedAccess);
 
 			NumBytes = ElementBytes * NumElements;
 			Buffer = RHICreateVertexBuffer(NumBytes, UsageFlags, InResourceState, CreateInfo);
@@ -411,14 +411,29 @@ FNiagaraRibbonVertexBuffers::FNiagaraRibbonVertexBuffers()
 
 void FNiagaraRibbonVertexBuffers::InitializeOrUpdateBuffers(const FNiagaraRibbonGenerationConfig& GenerationConfig, const TSharedPtr<FNiagaraRibbonCPUGeneratedVertexData>& GeneratedGeometryData, const FNiagaraDataBuffer* SourceParticleData, int32 MaxAllocatedCount, bool bIsUsingGPUInit)
 {	
-	const uint32 MaxAllocatedRibbons = GenerationConfig.HasRibbonIDs()? (GenerationConfig.GetMaxNumRibbons() > 0? GenerationConfig.GetMaxNumRibbons() : MaxAllocatedCount / 2) : 1;
+	const uint32 MaxAllocatedRibbons = GenerationConfig.HasRibbonIDs() ? (GenerationConfig.GetMaxNumRibbons() > 0 ? GenerationConfig.GetMaxNumRibbons() : MaxAllocatedCount) : 1;
 	
 	constexpr ERHIAccess InitialBufferAccessFlags = ERHIAccess::SRVMask | ERHIAccess::VertexOrIndexBuffer;
 
-	if (!bIsUsingGPUInit)
+	if (bIsUsingGPUInit)
 	{
+		const uint32 TotalParticles = SourceParticleData->GetNumInstances();
+
+		//-OPT:  We should be able to assume 2 particles per ribbon, however the compute pass does not cull our single particle ribbons therefore we need to allocate
+		//       enough space to assume each particle will be the start of a unique ribbon to avoid running OOB on the buffers.
+		const int32 TotalRibbons = FMath::Clamp<int32>(TotalParticles, 1, MaxAllocatedRibbons);
+
+		SortedIndicesBuffer.Allocate(TotalParticles, MaxAllocatedCount, InitialBufferAccessFlags, false);
+		TangentsAndDistancesBuffer.Allocate(TotalParticles * 4, MaxAllocatedCount * 4, InitialBufferAccessFlags, false);
+		MultiRibbonIndicesBuffer.Allocate(GenerationConfig.HasRibbonIDs() ? TotalParticles : 0, MaxAllocatedCount, InitialBufferAccessFlags, false);
+		RibbonLookupTableBuffer.Allocate(TotalRibbons * FRibbonMultiRibbonInfoBufferEntry::NumElements, MaxAllocatedRibbons * FRibbonMultiRibbonInfoBufferEntry::NumElements, InitialBufferAccessFlags, false);
+		SegmentsBuffer.Allocate(TotalParticles, MaxAllocatedCount, InitialBufferAccessFlags, false);
+		bJustCreatedCommandBuffer |= GPUComputeCommandBuffer.Allocate(FNiagaraRibbonCommandBufferLayout::NumElements, FNiagaraRibbonCommandBufferLayout::NumElements, InitialBufferAccessFlags | ERHIAccess::IndirectArgs, false, EBufferUsageFlags::DrawIndirect);
+	}
+	else
+	{		
 		check(GeneratedGeometryData.IsValid());
-		
+
 		SortedIndicesBuffer.Allocate(GeneratedGeometryData->SortedIndices.Num(), MaxAllocatedCount, InitialBufferAccessFlags, true);
 		TangentsAndDistancesBuffer.Allocate(GeneratedGeometryData->TangentAndDistances.Num() * 4, MaxAllocatedCount * 4, InitialBufferAccessFlags, true);
 		MultiRibbonIndicesBuffer.Allocate(GenerationConfig.HasRibbonIDs() ? GeneratedGeometryData->MultiRibbonIndices.Num() : 0, MaxAllocatedCount, InitialBufferAccessFlags, true);
@@ -427,33 +442,16 @@ void FNiagaraRibbonVertexBuffers::InitializeOrUpdateBuffers(const FNiagaraRibbon
 		GPUComputeCommandBuffer.Release();
 		bJustCreatedCommandBuffer = false;
 	}
-	else
-	{		
-		const uint32 TotalParticles = SourceParticleData->GetNumInstancesAllocated();
-
-		// We assume to have at least 2 particles in each ribbon, therefor we can't have more than instances/2 ribbons
-		const int32 TotalRibbons = FMath::Clamp<int32>(TotalParticles / 2, 1, MaxAllocatedRibbons);
-		
-		SortedIndicesBuffer.Allocate(TotalParticles, MaxAllocatedCount, InitialBufferAccessFlags, false);
-		TangentsAndDistancesBuffer.Allocate(TotalParticles * 4, MaxAllocatedCount * 4, InitialBufferAccessFlags, false);
-		MultiRibbonIndicesBuffer.Allocate(GenerationConfig.HasRibbonIDs() ? TotalParticles : 0, MaxAllocatedCount, InitialBufferAccessFlags, false);
-		RibbonLookupTableBuffer.Allocate(TotalRibbons * FRibbonMultiRibbonInfoBufferEntry::NumElements, MaxAllocatedRibbons * FRibbonMultiRibbonInfoBufferEntry::NumElements, InitialBufferAccessFlags, false);
-		SegmentsBuffer.Allocate(TotalParticles, MaxAllocatedCount, InitialBufferAccessFlags, false);
-		bJustCreatedCommandBuffer |= GPUComputeCommandBuffer.Allocate(FNiagaraRibbonCommandBufferLayout::NumElements, FNiagaraRibbonCommandBufferLayout::NumElements, InitialBufferAccessFlags | ERHIAccess::IndirectArgs, false);
-	}	
 }
 
 struct FNiagaraRibbonGPUInitParameters
 {
 	const FNiagaraRendererRibbons* Renderer;
-	FNiagaraGpuComputeDispatchInterface* ComputeDispatchInterface;
 	const FNiagaraDataBuffer* SourceParticleData;
 	TWeakPtr<FNiagaraRibbonRenderingFrameResources> RenderingResources;
 	
-	FNiagaraRibbonGPUInitParameters(const FNiagaraRendererRibbons* InRenderer, FNiagaraGpuComputeDispatchInterface* InComputeDispatchInterface,
-		const FNiagaraDataBuffer* InSourceParticleData, const TSharedPtr<FNiagaraRibbonRenderingFrameResources>& InRenderingResources)
+	FNiagaraRibbonGPUInitParameters(const FNiagaraRendererRibbons* InRenderer, const FNiagaraDataBuffer* InSourceParticleData, const TSharedPtr<FNiagaraRibbonRenderingFrameResources>& InRenderingResources)
 		: Renderer(InRenderer)
-		, ComputeDispatchInterface(InComputeDispatchInterface)
 		, SourceParticleData(InSourceParticleData)
 		, RenderingResources(InRenderingResources)
 	{
@@ -506,57 +504,43 @@ struct FNiagaraRibbonGPUInitComputeBuffers
 	}	
 };
 
-
-
-class FNiagaraRibbonComputeDispatchManager	
+class FNiagaraGpuRibbonsDataManager final : public FNiagaraGpuComputeDataManager
 {
-	friend class FLazySingleton;
-private:
-	TArray<FNiagaraRibbonGPUInitParameters> RenderersToGenerate;
-
-	FGPUSortManager* SortManager;	
-	FDelegateHandle SortManagerEventRef;
-
-	FNiagaraRibbonGPUInitComputeBuffers ComputeBuffers;
-	
-	void RegisterForEvent(FGPUSortManager* InSortManager)
+public:
+	FNiagaraGpuRibbonsDataManager(FNiagaraGpuComputeDispatchInterface* InOwnerInterface)
+		: FNiagaraGpuComputeDataManager(InOwnerInterface)
 	{
-		if (!SortManagerEventRef.IsValid())
-		{
-			SortManager = InSortManager;
-			SortManagerEventRef = SortManager->PostPreRenderEvent.AddRaw(this, &FNiagaraRibbonComputeDispatchManager::OnPostPreRender);
-		}
+		FGPUSortManager* SortManager = InOwnerInterface->GetGPUSortManager();
+		SortManager->PostPreRenderEvent.AddRaw(this, &FNiagaraGpuRibbonsDataManager::OnPostPreRender);
 	}
 
-	void OnPostPreRender(FRHICommandListImmediate& RHICmdList)
+	virtual ~FNiagaraGpuRibbonsDataManager()
 	{
-		GenerateAllGPUData(RHICmdList);
-		
-		SortManager->PostPreRenderEvent.Remove(SortManagerEventRef);
-		SortManagerEventRef.Reset();
-		SortManager = nullptr;
 	}
 
-	void GenerateAllGPUData(FRHICommandListImmediate& RHICmdList);
-
-public:
-	static FNiagaraRibbonComputeDispatchManager& Get()
+	static FName GetManagerName()
 	{
-		return TLazySingleton<FNiagaraRibbonComputeDispatchManager>::Get();
+		static FName ManagerName("FNiagaraGpuRibbonsDataManager");
+		return ManagerName;
 	}
 
-	void RegisterRenderer(FGPUSortManager* InSortManager, FNiagaraRibbonGPUInitParameters NewRegistration)
+	void RegisterRenderer(FNiagaraRibbonGPUInitParameters NewRegistration)
 	{
 		RenderersToGenerate.Add(NewRegistration);
-		RegisterForEvent(InSortManager);
 	}
+
+private:
+	TArray<FNiagaraRibbonGPUInitParameters> RenderersToGenerate;
+
+	FNiagaraRibbonGPUInitComputeBuffers ComputeBuffers;
 	
-	void UnRegister(FNiagaraRendererRibbons* Renderer)
-	{		
-		RenderersToGenerate.RemoveAll([Renderer](const FNiagaraRibbonGPUInitParameters& Params) { return Params.Renderer == Renderer; });
+	void OnPostPreRender(FRHICommandListImmediate& RHICmdList)
+	{
+		GenerateAllGPUData(RHICmdList);
 	}
-};
 
+	void GenerateAllGPUData(FRHICommandListImmediate& RHICmdList);
+};
 
 
 FNiagaraRendererRibbons::FNiagaraRendererRibbons(ERHIFeatureLevel::Type FeatureLevel, const UNiagaraRendererProperties *InProps, const FNiagaraEmitterInstance* Emitter)
@@ -584,7 +568,6 @@ FNiagaraRendererRibbons::FNiagaraRendererRibbons(ERHIFeatureLevel::Type FeatureL
 
 FNiagaraRendererRibbons::~FNiagaraRendererRibbons()
 {
-	FNiagaraRibbonComputeDispatchManager::Get().UnRegister(this);
 }
 
 // FPrimitiveSceneProxy interface.
@@ -667,7 +650,7 @@ void FNiagaraRendererRibbons::GetDynamicMeshElements(const TArray<const FSceneVi
 	{
 		// Bail if we don't have enough particle data to have a valid ribbon
 		// or if somehow the sim targets don't match
-		if (SimTarget != ENiagaraSimTarget::GPUComputeSim || SourceParticleData->GetNumInstancesAllocated() < 2)
+		if (SimTarget != ENiagaraSimTarget::GPUComputeSim || SourceParticleData->GetNumInstances() < 2)
 		{
 			return;
 		}
@@ -694,8 +677,6 @@ void FNiagaraRendererRibbons::GetDynamicMeshElements(const TArray<const FSceneVi
 	FScopeCycleCounter EmitterStatsCounter(EmitterStatID);
 #endif
 	
-	auto* ComputeDispatchInterface = SceneProxy->GetComputeDispatchInterface();
-
 	const FNiagaraRibbonMeshCollectorResources& RenderingResources = Collector.AllocateOneFrameResource<FNiagaraRibbonMeshCollectorResources>();
 		
 	InitializeVertexBuffersResources(DynamicData, SourceParticleData, Collector.GetDynamicReadBuffer(), RenderingResources.RibbonResources, DynamicData->bUseGPUInit);
@@ -734,8 +715,9 @@ void FNiagaraRendererRibbons::GetDynamicMeshElements(const TArray<const FSceneVi
 	// Register this renderer for generation this frame if we're a gpu system or using gpu init
 	if (DynamicData->bUseGPUInit || DynamicData->bIsGPUSystem)
 	{
-		FGPUSortManager* SortManager = ComputeDispatchInterface->GetGPUSortManager();
-		FNiagaraRibbonComputeDispatchManager::Get().RegisterRenderer(SortManager, FNiagaraRibbonGPUInitParameters(this, ComputeDispatchInterface, SourceParticleData, RenderingResources.RibbonResources));
+		FNiagaraGpuComputeDispatchInterface* ComputeDispatchInterface = SceneProxy->GetComputeDispatchInterface();
+		FNiagaraGpuRibbonsDataManager& RibbonDataManager = ComputeDispatchInterface->GetOrCreateDataManager<FNiagaraGpuRibbonsDataManager>();
+		RibbonDataManager.RegisterRenderer(FNiagaraRibbonGPUInitParameters(this, SourceParticleData, RenderingResources.RibbonResources));
 	}
 }
 
@@ -872,7 +854,7 @@ void FNiagaraRendererRibbons::GetDynamicRayTracingInstances(FRayTracingMaterialG
 
 		// Bail if we don't have enough particle data to have a valid ribbon
 		// or if somehow the sim targets don't match
-		//if (SimTarget != ENiagaraSimTarget::GPUComputeSim || SourceParticleData->GetNumInstancesAllocated() < 2)
+		//if (SimTarget != ENiagaraSimTarget::GPUComputeSim || SourceParticleData->GetNumInstances() < 2)
 		//{
 		//	return;
 		//}
@@ -1675,9 +1657,8 @@ FNiagaraIndexGenerationInput FNiagaraRendererRibbons::CalculateIndexBufferConfig
 
 	if (bShouldUseGPUInitIndices)
 	{
-		// If we are a gpu sim, we rely on num instances allocated since between now and render we could have a different number of particles living
-		// If we're not a gpu sim, we're a cpu sim with gpu init, we can rely on the num particles - 1 being the max cap for segments
-		IndexGenInput.MaxSegmentCount = bIsGPUSim? SourceParticleData->GetNumInstancesAllocated() : SourceParticleData->GetNumInstances();
+		// NumInstances is precise for GPU init from CPU but may be > number of alive particles for GPU simulations
+		IndexGenInput.MaxSegmentCount = SourceParticleData->GetNumInstances();
 	}
 	else
 	{
@@ -1999,13 +1980,13 @@ void FNiagaraRendererRibbons::InitializeViewIndexBuffersGPU(FRHICommandListImmed
 {
 	SCOPE_CYCLE_COUNTER(STAT_NiagaraRenderRibbonsGenIndiciesGPU);
 	
-	const uint32 NumInstances = SourceParticleData->GetNumInstances();
-		
 	if (!RenderingViewResources->IndirectDrawBuffer.Buffer.IsValid())
 	{
 		return;
 	}
 
+	const uint32 NumInstances = SourceParticleData->GetNumInstances();
+
 	SCOPED_DRAW_EVENT(RHICmdList, NiagaraRenderRibbonsGenIndiciesGPU);
 	{
 		FNiagaraRibbonCreateIndexBufferParamsCS::FPermutationDomain PermutationVector;
@@ -2021,7 +2002,7 @@ void FNiagaraRendererRibbons::InitializeViewIndexBuffersGPU(FRHICommandListImmed
 		Params.VertexGenerationResults = VertexBuffers.GPUComputeCommandBuffer.SRV;
 
 		// Total particle Count
-		Params.TotalNumParticlesDirect = SourceParticleData->GetNumInstances();
+		Params.TotalNumParticlesDirect = NumInstances;
 
 		// Indirect particle Count
 		Params.EmitterParticleCountsBuffer = GetSrvOrDefaultUInt(ComputeDispatchInterface->GetGPUInstanceCounterManager().GetInstanceCountBuffer());
@@ -2270,7 +2251,7 @@ void FNiagaraRendererRibbons::InitializeVertexBuffersGPU(FRHICommandListImmediat
 
 	FRibbonComputeUniformParameters CommonParams = SetupComputeVertexGenParams(ComputeDispatchInterface, RenderingResources, SourceParticleData);
 
-	const int32 NumExecutableInstances = CommonParams.EmitterParticleCountsBufferOffset != INDEX_NONE? SourceParticleData->GetNumInstancesAllocated() : SourceParticleData->GetNumInstances();
+	const int32 NumExecutableInstances = SourceParticleData->GetNumInstances();
 
 	const bool bCanRun = NumExecutableInstances >= 2;
 	
@@ -2313,8 +2294,8 @@ void FNiagaraRendererRibbons::InitializeVertexBuffersGPU(FRHICommandListImmediat
 			CurrentBufferOrientation ^= 0x1;	
 			const bool bComputeOnOutputBuffer = CurrentBufferOrientation == 0;
 			
-			SortParams.DestinationSortedIndices = bComputeOnOutputBuffer? VertexBuffers.SortedIndicesBuffer.UAV : TempBuffers.SortBuffer.UAV;
-			SortParams.SortedIndices = bComputeOnOutputBuffer? TempBuffers.SortBuffer.SRV : VertexBuffers.SortedIndicesBuffer.SRV;			
+			SortParams.DestinationSortedIndices = bComputeOnOutputBuffer ? VertexBuffers.SortedIndicesBuffer.UAV : TempBuffers.SortBuffer.UAV;
+			SortParams.SortedIndices = bComputeOnOutputBuffer ? TempBuffers.SortBuffer.SRV : VertexBuffers.SortedIndicesBuffer.SRV;			
 		};
 	
 		const uint32 NumInitialThreadGroups = FMath::DivideAndRoundUp<uint32>(NumExecutableInstances, FNiagaraRibbonSortPhase1CS::BubbleSortGroupWidth);	
@@ -2333,11 +2314,7 @@ void FNiagaraRendererRibbons::InitializeVertexBuffersGPU(FRHICommandListImmediat
 			
 			// Initial sort, sets up the buffer, and runs a parallel bubble sort to create groups of BubbleSortGroupWidth size
 			RHICmdList.Transition(FRHITransitionInfo(SortParams.DestinationSortedIndices, ERHIAccess::SRVMask | ERHIAccess::VertexOrIndexBuffer, ERHIAccess::UAVCompute));
-			SetComputePipelineState(RHICmdList, BubbleSortShader.GetComputeShader());
-			ValidateShaderParameters(BubbleSortShader, SortParams);
-			SetShaderParameters(RHICmdList, BubbleSortShader, BubbleSortShader.GetComputeShader(), SortParams);
-			DispatchComputeShader(RHICmdList, BubbleSortShader, NumInitialThreadGroups, 1, 1);
-			UnsetShaderUAVs(RHICmdList, BubbleSortShader, BubbleSortShader.GetComputeShader());
+			FComputeShaderUtils::Dispatch(RHICmdList, BubbleSortShader, SortParams, FIntVector(NumInitialThreadGroups, 1, 1));
 			RHICmdList.Transition(FRHITransitionInfo(SortParams.DestinationSortedIndices, ERHIAccess::UAVCompute, ERHIAccess::SRVMask | ERHIAccess::VertexOrIndexBuffer));
 		}
 		
@@ -2355,10 +2332,7 @@ void FNiagaraRendererRibbons::InitializeVertexBuffersGPU(FRHICommandListImmediat
 				SwapBuffers();
 			
 				RHICmdList.Transition(FRHITransitionInfo(SortParams.DestinationSortedIndices, ERHIAccess::SRVMask | ERHIAccess::VertexOrIndexBuffer, ERHIAccess::UAVCompute));
-				SetComputePipelineState(RHICmdList, MergeSortShader.GetComputeShader());
-				SetShaderParameters(RHICmdList, MergeSortShader, MergeSortShader.GetComputeShader(), SortParams);
-				DispatchComputeShader(RHICmdList, MergeSortShader, NumMergeSortThreadGroups, 1, 1);
-				UnsetShaderUAVs(RHICmdList, MergeSortShader, MergeSortShader.GetComputeShader());
+				FComputeShaderUtils::Dispatch(RHICmdList, MergeSortShader, SortParams, FIntVector(NumInitialThreadGroups, 1, 1));
 				RHICmdList.Transition(FRHITransitionInfo(SortParams.DestinationSortedIndices, ERHIAccess::UAVCompute, ERHIAccess::SRVMask | ERHIAccess::VertexOrIndexBuffer));
 				
 				SortGroupSize *= 2;
@@ -2370,14 +2344,21 @@ void FNiagaraRendererRibbons::InitializeVertexBuffersGPU(FRHICommandListImmediat
 		SCOPED_DRAW_EVENT(RHICmdList, NiagaraRenderRibbonsGenVerticesReductionPhase1GPU);
 		SCOPE_CYCLE_COUNTER(STAT_NiagaraRenderRibbonsGenVerticesReductionPhase1GPU);
 		
-		FNiagaraRibbonVertexReductionInitializationCS::FPermutationDomain PermutationVector;
-		PermutationVector.Set<FRibbonHasFullRibbonID>(GenerationConfig.HasFullRibbonIDs());
-		PermutationVector.Set<FRibbonHasRibbonID>(GenerationConfig.HasSimpleRibbonIDs());
-		PermutationVector.Set<FRibbonWantsAutomaticTessellation>(GenerationConfig.WantsAutomaticTessellation());
-		PermutationVector.Set<FRibbonWantsConstantTessellation>(GenerationConfig.WantsConstantTessellation());
-		
-		TShaderMapRef<FNiagaraRibbonVertexReductionInitializationCS> ReductionInitializationShader(GetGlobalShaderMap(GMaxRHIFeatureLevel), PermutationVector);
-		TShaderMapRef<FNiagaraRibbonVertexReductionPropagateCS> ReductionPropgateShader(GetGlobalShaderMap(GMaxRHIFeatureLevel), PermutationVector);	
+		FNiagaraRibbonVertexReductionInitializationCS::FPermutationDomain InitPermutationVector;
+		InitPermutationVector.Set<FRibbonHasFullRibbonID>(GenerationConfig.HasFullRibbonIDs());
+		InitPermutationVector.Set<FRibbonHasRibbonID>(GenerationConfig.HasSimpleRibbonIDs());
+		InitPermutationVector.Set<FRibbonWantsAutomaticTessellation>(GenerationConfig.WantsAutomaticTessellation());
+		InitPermutationVector.Set<FRibbonWantsConstantTessellation>(GenerationConfig.WantsConstantTessellation());
+		InitPermutationVector.Set<FRibbonHasTwist>(GenerationConfig.HasTwist());
+		TShaderMapRef<FNiagaraRibbonVertexReductionInitializationCS> ReductionInitializationShader(GetGlobalShaderMap(GMaxRHIFeatureLevel), InitPermutationVector);
+
+		FNiagaraRibbonVertexReductionPropagateCS::FPermutationDomain  PropagationPermutationVector;
+		PropagationPermutationVector.Set<FRibbonHasFullRibbonID>(GenerationConfig.HasFullRibbonIDs());
+		PropagationPermutationVector.Set<FRibbonHasRibbonID>(GenerationConfig.HasSimpleRibbonIDs());
+		PropagationPermutationVector.Set<FRibbonWantsAutomaticTessellation>(GenerationConfig.WantsAutomaticTessellation());
+		PropagationPermutationVector.Set<FRibbonWantsConstantTessellation>(GenerationConfig.WantsConstantTessellation());
+		PropagationPermutationVector.Set<FRibbonHasTwist>(GenerationConfig.HasTwist());
+		TShaderMapRef<FNiagaraRibbonVertexReductionPropagateCS> ReductionPropgateShader(GetGlobalShaderMap(GMaxRHIFeatureLevel), PropagationPermutationVector);
 		
 		const int32 NumPrefixScanPasses = FMath::CeilLogTwo(NumExecutableInstances);
 	
@@ -2441,15 +2422,11 @@ void FNiagaraRendererRibbons::InitializeVertexBuffersGPU(FRHICommandListImmediat
 			RHICmdList.Transition(MakeArrayView(DataBufferTransitions, UE_ARRAY_COUNT(DataBufferTransitions)));
 		};
 		
-	
 		{			
 			
 			const uint32 NumThreadGroupsInitialization = FMath::DivideAndRoundUp<uint32>(NumExecutableInstances, FNiagaraRibbonComputeCommon::VertexGenReductionInitializationThreadSize);	
 			TransitionOutputBuffers(ERHIAccess::SRVMask | ERHIAccess::VertexOrIndexBuffer, ERHIAccess::UAVCompute);
-			SetComputePipelineState(RHICmdList, ReductionInitializationShader.GetComputeShader());
-			SetShaderParameters(RHICmdList, ReductionInitializationShader, ReductionInitializationShader.GetComputeShader(), Params);
-			DispatchComputeShader(RHICmdList, ReductionInitializationShader, NumThreadGroupsInitialization, 1, 1);
-			UnsetShaderUAVs(RHICmdList, ReductionInitializationShader, ReductionInitializationShader.GetComputeShader());
+			FComputeShaderUtils::Dispatch(RHICmdList, ReductionInitializationShader, Params, FIntVector(NumThreadGroupsInitialization, 1, 1));
 			TransitionOutputBuffers(ERHIAccess::UAVCompute, ERHIAccess::SRVMask | ERHIAccess::VertexOrIndexBuffer);
 		}
 		
@@ -2463,12 +2440,7 @@ void FNiagaraRendererRibbons::InitializeVertexBuffersGPU(FRHICommandListImmediat
 				SwapBuffers();
 				
 				TransitionOutputBuffers(ERHIAccess::SRVMask | ERHIAccess::VertexOrIndexBuffer, ERHIAccess::UAVCompute);
-			
-				SetComputePipelineState(RHICmdList, ReductionPropgateShader.GetComputeShader());
-				SetShaderParameters(RHICmdList, ReductionPropgateShader, ReductionPropgateShader.GetComputeShader(), Params);
-				DispatchComputeShader(RHICmdList, ReductionPropgateShader, NumThreadGroups, 1, 1);
-				UnsetShaderUAVs(RHICmdList, ReductionPropgateShader, ReductionPropgateShader.GetComputeShader());
-				
+				FComputeShaderUtils::Dispatch(RHICmdList, ReductionPropgateShader, Params, FIntVector(NumThreadGroups, 1, 1));
 				TransitionOutputBuffers(ERHIAccess::UAVCompute, ERHIAccess::SRVMask | ERHIAccess::VertexOrIndexBuffer);
 			}
 		}
@@ -2514,11 +2486,8 @@ void FNiagaraRendererRibbons::InitializeVertexBuffersGPU(FRHICommandListImmediat
 			RHICmdList.Transition({
 				FRHITransitionInfo(VertexBuffers.RibbonLookupTableBuffer.Buffer, ERHIAccess::SRVMask | ERHIAccess::VertexOrIndexBuffer, ERHIAccess::UAVCompute),
 				FRHITransitionInfo(VertexBuffers.GPUComputeCommandBuffer.Buffer, ERHIAccess::SRVMask | ERHIAccess::VertexOrIndexBuffer | ERHIAccess::IndirectArgs, ERHIAccess::UAVCompute)});
-			
-			SetComputePipelineState(RHICmdList, ComputeShader.GetComputeShader());
-			SetShaderParameters(RHICmdList, ComputeShader, ComputeShader.GetComputeShader(), FinalizationParams);
-			DispatchComputeShader(RHICmdList, ComputeShader, NumThreadGroups, 1, 1);
-			UnsetShaderUAVs(RHICmdList, ComputeShader, ComputeShader.GetComputeShader());
+
+			FComputeShaderUtils::Dispatch(RHICmdList, ComputeShader, FinalizationParams, FIntVector(NumThreadGroups, 1, 1));
 			
 			// We don't need to transition RibbonLookupTableBuffer as it's still needed for the next shader
 			RHICmdList.Transition({
@@ -2563,7 +2532,6 @@ void FNiagaraRendererRibbons::InitializeVertexBuffersGPU(FRHICommandListImmediat
 			PermutationVector.Set<FRibbonHasRibbonID>(GenerationConfig.HasSimpleRibbonIDs());
 			PermutationVector.Set<FRibbonWantsAutomaticTessellation>(GenerationConfig.WantsAutomaticTessellation());
 			PermutationVector.Set<FRibbonWantsConstantTessellation>(GenerationConfig.WantsConstantTessellation());
-			PermutationVector.Set<FRibbonHasTwist>(GenerationConfig.HasTwist());
 
 			TShaderMapRef<FNiagaraRibbonUVParamCalculationCS> ComputeShader(GetGlobalShaderMap(GMaxRHIFeatureLevel), PermutationVector);
 
@@ -2571,12 +2539,9 @@ void FNiagaraRendererRibbons::InitializeVertexBuffersGPU(FRHICommandListImmediat
 			RHICmdList.Transition({
 				FRHITransitionInfo(VertexBuffers.RibbonLookupTableBuffer.Buffer, ERHIAccess::UAVCompute, ERHIAccess::UAVCompute),
 				FRHITransitionInfo(VertexBuffers.TangentsAndDistancesBuffer.Buffer, ERHIAccess::SRVMask | ERHIAccess::VertexOrIndexBuffer, ERHIAccess::UAVCompute)});
-			
-			SetComputePipelineState(RHICmdList, ComputeShader.GetComputeShader());
-			SetShaderParameters(RHICmdList, ComputeShader, ComputeShader.GetComputeShader(), FinalizeParams);
-			DispatchIndirectComputeShader(RHICmdList, ComputeShader.GetShader(), VertexBuffers.GPUComputeCommandBuffer.Buffer, CommandBufferOffset * FNiagaraRibbonCommandBufferLayout::NumElements);
-			UnsetShaderUAVs(RHICmdList, ComputeShader, ComputeShader.GetComputeShader());
-			
+
+			FComputeShaderUtils::DispatchIndirect(RHICmdList, ComputeShader, FinalizeParams, VertexBuffers.GPUComputeCommandBuffer.Buffer, CommandBufferOffset * FNiagaraRibbonCommandBufferLayout::NumElements);
+
 			RHICmdList.Transition({
 				FRHITransitionInfo(VertexBuffers.TangentsAndDistancesBuffer.Buffer, ERHIAccess::UAVCompute, ERHIAccess::SRVMask | ERHIAccess::VertexOrIndexBuffer),
 				FRHITransitionInfo(VertexBuffers.RibbonLookupTableBuffer.Buffer, ERHIAccess::UAVCompute, ERHIAccess::SRVMask | ERHIAccess::VertexOrIndexBuffer)});
@@ -2584,8 +2549,15 @@ void FNiagaraRendererRibbons::InitializeVertexBuffersGPU(FRHICommandListImmediat
 	}
 }
 
-void FNiagaraRibbonComputeDispatchManager::GenerateAllGPUData(FRHICommandListImmediate& RHICmdList)
+void FNiagaraGpuRibbonsDataManager::GenerateAllGPUData(FRHICommandListImmediate& RHICmdList)
 {
+	if (RenderersToGenerate.Num() == 0)
+	{
+		return;
+	}
+
+	FNiagaraGpuComputeDispatchInterface* ComputeDispatchInterface = GetOwnerInterface();
+
 	// Handle all vertex gens first
 	for (int32 Index = 0; Index < RenderersToGenerate.Num(); Index++)
 	{
@@ -2596,16 +2568,17 @@ void FNiagaraRibbonComputeDispatchManager::GenerateAllGPUData(FRHICommandListImm
 		{
 			const bool bIsGPUSim = Params.SourceParticleData->GetGPUInstanceCountBufferOffset() != INDEX_NONE;
 
-			ComputeBuffers.InitOrUpdateBuffers(bIsGPUSim? Params.SourceParticleData->GetNumInstancesAllocated() : Params.SourceParticleData->GetNumInstances(),
+			ComputeBuffers.InitOrUpdateBuffers(
+				Params.SourceParticleData->GetNumInstances(),
 				Params.Renderer->GenerationConfig.HasRibbonIDs(),
 				Params.Renderer->GenerationConfig.WantsAutomaticTessellation(),
-				Params.Renderer->GenerationConfig.HasTwist());
-		
-		
-			Params.Renderer->InitializeVertexBuffersGPU(RHICmdList, Params.ComputeDispatchInterface, Params.SourceParticleData, ComputeBuffers, RenderingResources);
+				Params.Renderer->GenerationConfig.HasTwist()
+			);
+
+			Params.Renderer->InitializeVertexBuffersGPU(RHICmdList, ComputeDispatchInterface, Params.SourceParticleData, ComputeBuffers, RenderingResources);
 		}
 	}
-		
+
 	// Now handle all index gens
 	for (const auto& RendererToGen : RenderersToGenerate)
 	{
@@ -2615,7 +2588,7 @@ void FNiagaraRibbonComputeDispatchManager::GenerateAllGPUData(FRHICommandListImm
 			for (int32 Index = 0; Index < RenderingResources->ViewResources.Num(); Index++)
 			{
 				const auto& RenderingResourcesView = RenderingResources->ViewResources[Index];
-				RendererToGen.Renderer->InitializeViewIndexBuffersGPU(RHICmdList, RendererToGen.ComputeDispatchInterface, RendererToGen.SourceParticleData, RenderingResourcesView);
+				RendererToGen.Renderer->InitializeViewIndexBuffersGPU(RHICmdList, ComputeDispatchInterface, RendererToGen.SourceParticleData, RenderingResourcesView);
 			}
 		}
 	}
diff --git a/Engine/Plugins/FX/Niagara/Source/Niagara/Public/NiagaraGpuComputeDataManager.h b/Engine/Plugins/FX/Niagara/Source/Niagara/Public/NiagaraGpuComputeDataManager.h
new file mode 100644
--- /dev/null
+++ b/Engine/Plugins/FX/Niagara/Source/Niagara/Public/NiagaraGpuComputeDataManager.h
@@ -0,0 +1,33 @@
+// Copyright Epic Games, Inc. All Rights Reserved.
+
+#pragma once
+
+#include "NiagaraCommon.h"
+
+class FNiagaraGpuComputeDispatchInterface;
+
+// Abstract class for managing GPU data for Niagara
+// Once the manager is created it will last for the lifetime of the owner dispatch interface
+// The manager will be destroyed when the dispatch interface is also destroyed
+class NIAGARA_API FNiagaraGpuComputeDataManager
+{
+	UE_NONCOPYABLE(FNiagaraGpuComputeDataManager);
+
+public:
+	FNiagaraGpuComputeDataManager(FNiagaraGpuComputeDispatchInterface* InOwnerInterface)
+		: InternalOwnerInterface(InOwnerInterface)
+	{
+	}
+	virtual ~FNiagaraGpuComputeDataManager() {}
+
+	FNiagaraGpuComputeDispatchInterface* GetOwnerInterface() const { return InternalOwnerInterface; }
+
+	//static FName GetManagerName()
+	//{
+	//	static FName ManagerName("FNiagaraGpuComputeDataManager");
+	//	return ManagerName;
+	//}
+
+private:
+	FNiagaraGpuComputeDispatchInterface* InternalOwnerInterface = nullptr;
+};
diff --git a/Engine/Plugins/FX/Niagara/Source/Niagara/Public/NiagaraGpuComputeDispatchInterface.h b/Engine/Plugins/FX/Niagara/Source/Niagara/Public/NiagaraGpuComputeDispatchInterface.h
--- a/Engine/Plugins/FX/Niagara/Source/Niagara/Public/NiagaraGpuComputeDispatchInterface.h
+++ b/Engine/Plugins/FX/Niagara/Source/Niagara/Public/NiagaraGpuComputeDispatchInterface.h
@@ -4,6 +4,7 @@
 
 #include "NiagaraCommon.h"
 #include "NiagaraEmptyUAVPool.h"
+#include "NiagaraGpuComputeDataManager.h"
 #include "NiagaraGPUInstanceCountManager.h"
 #include "FXSystem.h"
 
@@ -52,6 +53,25 @@ public:
 	 */
 	virtual bool AddSortedGPUSimulation(struct FNiagaraGPUSortInfo& SortInfo) = 0;
 
+	/** Get or create the a data manager, must be done on the rendering thread only. */
+	template<typename TManager>
+	TManager& GetOrCreateDataManager()
+	{	
+		check(IsInRenderingThread());
+		const FName ManagerName = TManager::GetManagerName();
+		for (auto& DataManager : GpuDataManagers)
+		{
+			if (DataManager.Key == ManagerName)
+			{
+				return *static_cast<TManager*>(DataManager.Value.Get());
+			}
+		}
+
+		TManager* Manager = new TManager(this);
+		GpuDataManagers.Emplace(ManagerName, Manager);
+		return *Manager;
+	}
+
 	/** Get access to the instance count manager. */
 	FORCEINLINE FNiagaraGPUInstanceCountManager& GetGPUInstanceCounterManager() { check(IsInRenderingThread()); return GPUInstanceCounterManager; }
 	FORCEINLINE const FNiagaraGPUInstanceCountManager& GetGPUInstanceCounterManager() const { check(IsInRenderingThread()); return GPUInstanceCounterManager; }
@@ -149,4 +169,6 @@ protected:
 
 	// GPU emitter instance count buffer. Contains the actual particle / instance count generate in the GPU tick.
 	FNiagaraGPUInstanceCountManager			GPUInstanceCounterManager;
+
+	TArray<TPair<FName, TUniquePtr<FNiagaraGpuComputeDataManager>>> GpuDataManagers;
 };
diff --git a/Engine/Plugins/FX/Niagara/Source/Niagara/Public/NiagaraRendererRibbons.h b/Engine/Plugins/FX/Niagara/Source/Niagara/Public/NiagaraRendererRibbons.h
--- a/Engine/Plugins/FX/Niagara/Source/Niagara/Public/NiagaraRendererRibbons.h
+++ b/Engine/Plugins/FX/Niagara/Source/Niagara/Public/NiagaraRendererRibbons.h
@@ -256,7 +256,7 @@ struct FNiagaraRibbonGpuBuffer
 		Release();
 	}
 
-	bool Allocate(uint32 NumElements, uint32 MaxElements, ERHIAccess InResourceState, bool bGpuReadOnly);
+	bool Allocate(uint32 NumElements, uint32 MaxElements, ERHIAccess InResourceState, bool bGpuReadOnly, EBufferUsageFlags AdditionalBufferUsage = EBufferUsageFlags::None);
 	void Release();
 
 	const TCHAR*				DebugName = nullptr;
@@ -430,5 +430,5 @@ protected:
 	mutable FNiagaraRibbonTessellationSmoothingData TessellationSmoothingData;
 	int32 RibbonIDParamDataSetOffset;
 
-	friend class FNiagaraRibbonComputeDispatchManager;
+	friend class FNiagaraGpuRibbonsDataManager;
 };
diff --git a/Engine/Plugins/FX/Niagara/Source/NiagaraShader/Public/NiagaraRibbonCompute.h b/Engine/Plugins/FX/Niagara/Source/NiagaraShader/Public/NiagaraRibbonCompute.h
--- a/Engine/Plugins/FX/Niagara/Source/NiagaraShader/Public/NiagaraRibbonCompute.h
+++ b/Engine/Plugins/FX/Niagara/Source/NiagaraShader/Public/NiagaraRibbonCompute.h
@@ -68,8 +68,6 @@ class FRibbonWantsAutomaticTessellation : SHADER_PERMUTATION_BOOL("RIBBONS_WANTS
 class FRibbonHasTwist : SHADER_PERMUTATION_BOOL("RIBBON_HAS_TWIST");
 
 class FRibbonHasHighSliceComplexity : SHADER_PERMUTATION_BOOL("RIBBON_HAS_HIGH_SLICE_COMPLEXITY");
-using FRibbonComputePermutationDomain = TShaderPermutationDomain<FRibbonHasFullRibbonID, FRibbonHasRibbonID, FRibbonHasCustomLinkOrder, FRibbonWantsConstantTessellation, FRibbonWantsAutomaticTessellation, FRibbonHasTwist, FRibbonHasHighSliceComplexity>;
-
 
 struct FNiagaraRibbonComputeCommon
 {	
@@ -161,8 +159,7 @@ class NIAGARASHADER_API FNiagaraRibbonVertexReductionPropagateCS : public FGloba
 	SHADER_USE_PARAMETER_STRUCT(FNiagaraRibbonVertexReductionPropagateCS, FGlobalShader);
 	
 	using FPermutationDomain = TShaderPermutationDomain<FRibbonHasFullRibbonID, FRibbonHasRibbonID, FRibbonWantsConstantTessellation, FRibbonWantsAutomaticTessellation, FRibbonHasTwist>;
-	using FParameters = FNiagaraRibbonVertexReductionParameters
-	;
+	using FParameters = FNiagaraRibbonVertexReductionParameters;
 	
 	static void ModifyCompilationEnvironment(const FGlobalShaderPermutationParameters& Parameters, FShaderCompilerEnvironment& OutEnvironment);
 };
@@ -219,14 +216,12 @@ class NIAGARASHADER_API FNiagaraRibbonUVParamCalculationCS : public FGlobalShade
 	DECLARE_GLOBAL_SHADER(FNiagaraRibbonUVParamCalculationCS);
 	SHADER_USE_PARAMETER_STRUCT(FNiagaraRibbonUVParamCalculationCS, FGlobalShader);
 	
-	using FPermutationDomain = TShaderPermutationDomain<FRibbonHasFullRibbonID, FRibbonHasRibbonID, FRibbonWantsConstantTessellation, FRibbonWantsAutomaticTessellation, FRibbonHasTwist>;
+	using FPermutationDomain = TShaderPermutationDomain<FRibbonHasFullRibbonID, FRibbonHasRibbonID, FRibbonWantsConstantTessellation, FRibbonWantsAutomaticTessellation>;
 	using FParameters = FNiagaraRibbonVertexFinalizationParameters;
 	
 	static void ModifyCompilationEnvironment(const FGlobalShaderPermutationParameters& Parameters, FShaderCompilerEnvironment& OutEnvironment);
 };
 
-
-
 BEGIN_SHADER_PARAMETER_STRUCT(FNiagaraRibbonInitializeIndices, NIAGARASHADER_API)
 	SHADER_PARAMETER_UAV(RWBuffer<uint32>, IndirectDrawOutput)
 	SHADER_PARAMETER_SRV(Buffer<uint32>, VertexGenerationResults)
