Side step issue users can run into when trying to enable area weighted sampling of skeletal meshes with Niagara.  Mkae sure we only enable bAreaWeighting if the option is selected and the underlying data is present, this makes it so that if the user has modified the mesh, but hasn't actually applied the change, that it behaves sensibly.

diff --git a/Engine/Plugins/FX/Niagara/Source/Niagara/Classes/NDISkeletalMeshCommon.h b/Engine/Plugins/FX/Niagara/Source/Niagara/Classes/NDISkeletalMeshCommon.h
--- a/Engine/Plugins/FX/Niagara/Source/Niagara/Classes/NDISkeletalMeshCommon.h
+++ b/Engine/Plugins/FX/Niagara/Source/Niagara/Classes/NDISkeletalMeshCommon.h
@@ -546,18 +546,44 @@ struct TAreaWeightingModeBinder
 		{
 			if (InstData->SamplingRegionIndices.Num() > 1)
 			{
-				bAreaWeighting = InstData->SamplingRegionAreaWeightedSampler.IsValid();
+				if (InstData->SamplingRegionAreaWeightedSampler.IsValid())
+				{
+					bAreaWeighting = true;
+
+					// check that all the regions have sampling data
+					const FSkeletalMeshSamplingInfo& SamplingInfo = InstData->SkeletalMesh->GetSamplingInfo();
+
+					for (int32 RegionIndex : InstData->SamplingRegionIndices)
+					{
+						const FSkeletalMeshSamplingRegion& Region = SamplingInfo.GetRegion(RegionIndex);
+						const FSkeletalMeshSamplingRegionBuiltData& RegionBuiltData = SamplingInfo.GetRegionBuiltData(RegionIndex);
+
+						if (RegionBuiltData.AreaWeightedSampler.GetNumEntries() == 0)
+						{
+							bAreaWeighting = false;
+						}
+					}
+				}
 			}
 			else if (InstData->SamplingRegionIndices.Num() == 1)
 			{
+				const int32 RegionIndex = InstData->SamplingRegionIndices[0];
+
 				const FSkeletalMeshSamplingInfo& SamplingInfo = InstData->SkeletalMesh->GetSamplingInfo();
-				const FSkeletalMeshSamplingRegion& Region = SamplingInfo.GetRegion(InstData->SamplingRegionIndices[0]);
-				bAreaWeighting = Region.bSupportUniformlyDistributedSampling;
+				const FSkeletalMeshSamplingRegion& Region = SamplingInfo.GetRegion(RegionIndex);
+
+				bAreaWeighting = Region.bSupportUniformlyDistributedSampling
+					&& (SamplingInfo.GetRegionBuiltData(RegionIndex).AreaWeightedSampler.GetNumEntries() > 0);
 			}
 			else
 			{
-				int32 LODIndex = InstData->GetLODIndex();
-				bAreaWeighting = InstData->SkeletalMesh->GetLODInfo(LODIndex)->bSupportUniformlyDistributedSampling;
+				const int32 LODIndex = InstData->GetLODIndex();
+
+				const FSkeletalMeshSamplingInfo& SamplingInfo = InstData->SkeletalMesh->GetSamplingInfo();
+				const FSkeletalMeshSamplingLODBuiltData& WholeMeshBuiltData = SamplingInfo.GetWholeMeshLODBuiltData(LODIndex);
+
+				bAreaWeighting = InstData->SkeletalMesh->GetLODInfo(LODIndex)->bSupportUniformlyDistributedSampling
+					&& (WholeMeshBuiltData.AreaWeightedTriangleSampler.GetNumEntries() > 0);
 			}
 		}
 
