- Add methods to get socket index from filtered / unfiltered sockets
- Add Velocity to all functions to make it consistent
- Add Interpolated methods for WS functions

diff --git a/Engine/Plugins/FX/Niagara/Shaders/Private/NiagaraDataInterfaceStaticMeshTemplate.ush b/Engine/Plugins/FX/Niagara/Shaders/Private/NiagaraDataInterfaceStaticMeshTemplate.ush
--- a/Engine/Plugins/FX/Niagara/Shaders/Private/NiagaraDataInterfaceStaticMeshTemplate.ush
+++ b/Engine/Plugins/FX/Niagara/Shaders/Private/NiagaraDataInterfaceStaticMeshTemplate.ush
@@ -1,5 +1,7 @@
 // Copyright Epic Games, Inc. All Rights Reserved.
 
+#include "/Plugin/FX/Niagara/Private/NiagaraQuaternionUtils.ush"
+
 int3			{ParameterName}_NumTriangles;		// x = Num Sections, y = Num Filtered, z = Num Unfiltered
 int				{ParameterName}_NumVertices;
 int				{ParameterName}_NumUVs;
@@ -100,6 +102,7 @@ void GetVertex_{ParameterName}(int Vertex, out float3 Position, out float3 Veloc
 		float4 TangentZ = TangentBias({ParameterName}_TangentBuffer[Vertex * 2 + 1]);
 		float3 TangentY = cross(TangentZ.xyz, TangentX.xyz) * TangentZ.w;
 
+		Velocity	= float3(0, 0, 0);
 		Normal		= TangentZ.xyz;
 		Bitangent	= TangentY;
 		Tangent		= TangentX.xyz;
@@ -127,6 +130,28 @@ void GetVertexWS_{ParameterName}(int Vertex, out float3 Position, out float3 Vel
 	Tangent					= normalize(mul(float4(Tangent, 0), {ParameterName}_InstanceTransformInverseTransposed).xyz);
 }
 
+void GetVertexWSInterpolated_{ParameterName}(int Vertex, float Interp, out float3 Position, out float3 Velocity, out float3 Normal, out float3 Bitangent, out float3 Tangent)
+{
+	float3 LocalPosition, LocalNormal, LocalBitangent, LocalTangent;
+	GetVertex_{ParameterName}(Vertex, LocalPosition, Velocity, Normal, Bitangent, Tangent);
+
+	float3 PrevPosition		= mul(float4(LocalPosition, 1), {ParameterName}_InstancePreviousTransform).xyz;
+	float3 PrevNormal		= mul(float4(LocalNormal, 0), {ParameterName}_InstancePreviousTransformInverseTransposed).xyz;
+	float3 PrevBitangent	= mul(float4(LocalBitangent, 0), {ParameterName}_InstancePreviousTransformInverseTransposed).xyz;
+	float3 PrevTangent		= mul(float4(LocalTangent, 0), {ParameterName}_InstancePreviousTransformInverseTransposed).xyz;
+
+	float3 CurrPosition		= mul(float4(LocalPosition, 1), {ParameterName}_InstanceTransform).xyz;
+	float3 CurrNormal		= mul(float4(LocalNormal, 0), {ParameterName}_InstanceTransformInverseTransposed).xyz;
+	float3 CurrBitangent	= mul(float4(LocalBitangent, 0), {ParameterName}_InstanceTransformInverseTransposed).xyz;
+	float3 CurrTangent		= mul(float4(LocalTangent, 0), {ParameterName}_InstanceTransformInverseTransposed).xyz;
+
+	Position				= lerp(PrevPosition, CurrPosition, Interp);
+	Velocity				= (CurrPosition - PrevPosition) * {ParameterName}_InvDeltaSeconds;
+	Normal					= normalize(lerp(PrevNormal, CurrNormal, Interp));
+	Bitangent				= normalize(lerp(PrevBitangent, CurrBitangent, Interp));
+	Tangent					= normalize(lerp(PrevTangent, CurrTangent, Interp));
+}
+
 void GetVertexColor_{ParameterName}(int Vertex, out float4 Color)
 {
 	Vertex = clamp(Vertex, 0, {ParameterName}_NumVertices - 1);
@@ -293,6 +318,28 @@ void GetTriangleWS_{ParameterName}(int Tri, float3 BaryCoord, out float3 Positio
 	Tangent					= normalize(mul(float4(Tangent, 0), {ParameterName}_InstanceTransformInverseTransposed).xyz);
 }
 
+void GetTriangleWSInterpolated_{ParameterName}(int Tri, float3 BaryCoord, float Interp, out float3 Position, out float3 Velocity, out float3 Normal, out float3 Bitangent, out float3 Tangent)
+{
+	float3 LocalPosition, LocalNormal, LocalBitangent, LocalTangent;
+	GetTriangle_{ParameterName}(Tri, BaryCoord, LocalPosition, Velocity, LocalNormal, LocalBitangent, LocalTangent);
+
+	float3 PrevPosition		= mul(float4(LocalPosition, 1), {ParameterName}_InstancePreviousTransform).xyz;
+	float3 PrevNormal		= mul(float4(LocalNormal, 0), {ParameterName}_InstancePreviousTransformInverseTransposed).xyz;
+	float3 PrevBitangent	= mul(float4(LocalBitangent, 0), {ParameterName}_InstancePreviousTransformInverseTransposed).xyz;
+	float3 PrevTangent		= mul(float4(LocalTangent, 0), {ParameterName}_InstancePreviousTransformInverseTransposed).xyz;
+
+	float3 CurrPosition		= mul(float4(LocalPosition, 1), {ParameterName}_InstanceTransform).xyz;
+	float3 CurrNormal		= mul(float4(LocalNormal, 0), {ParameterName}_InstanceTransformInverseTransposed).xyz;
+	float3 CurrBitangent	= mul(float4(LocalBitangent, 0), {ParameterName}_InstanceTransformInverseTransposed).xyz;
+	float3 CurrTangent		= mul(float4(LocalTangent, 0), {ParameterName}_InstanceTransformInverseTransposed).xyz;
+
+	Position				= lerp(PrevPosition, CurrPosition, Interp);
+	Velocity				= (CurrPosition - PrevPosition) * {ParameterName}_InvDeltaSeconds;
+	Normal					= normalize(lerp(PrevNormal, CurrNormal, Interp));
+	Bitangent				= normalize(lerp(PrevBitangent, CurrBitangent, Interp));
+	Tangent					= normalize(lerp(PrevTangent, CurrTangent, Interp));
+}
+
 void GetTriangleColor_{ParameterName}(int Tri, float3 BaryCoord, out float4 Color)
 {
 	int Index[3];
@@ -336,6 +383,26 @@ void GetUnfilteredSocketCount_{ParameterName}(out int Count)
 	Count = {ParameterName}_SocketCounts.z;
 }
 
+void GetFilteredSocket_{ParameterName}(int Index, out int SocketIndex)
+{
+	SocketIndex = -1;
+	if ( {ParameterName}_SocketCounts.y > 0 )
+	{
+		Index = clamp(Index, 0, {ParameterName}_SocketCounts.y - 1);
+		SocketIndex = {ParameterName}_FilteredAndUnfilteredSockets[Index];
+	}
+}
+
+void GetUnfilteredSocket_{ParameterName}(int Index, out int SocketIndex)
+{
+	SocketIndex = -1;
+	if ( {ParameterName}_SocketCounts.z > 0 )
+	{
+		Index = {ParameterName}_SocketCounts.y + clamp(Index, 0, {ParameterName}_SocketCounts.z - 1);
+		SocketIndex = {ParameterName}_FilteredAndUnfilteredSockets[Index];
+	}
+}
+
 void RandomSocket_{ParameterName}(NiagaraRandInfo RandInfo, out int Index)
 {
 	Index = -1;
@@ -365,7 +432,7 @@ void RandomUnfilteredSocket_{ParameterName}(NiagaraRandInfo RandInfo, out int In
 	}
 }
 
-void GetSocketTransform_{ParameterName}(int Socket, out float3 Position, out float4 Rotation, out float3 Scale)
+void GetSocketTransform_{ParameterName}(int Socket, out float3 Position, out float4 Rotation, out float3 Scale, out float3 Velocity)
 {
 	[branch]
 	if ( {ParameterName}_SocketCounts.x > 0 )
@@ -374,71 +441,122 @@ void GetSocketTransform_{ParameterName}(int Socket, out float3 Position, out flo
 		Position = {ParameterName}_SocketTransforms[Socket * 3 + 0].xyz;
 		Rotation = {ParameterName}_SocketTransforms[Socket * 3 + 1];
 		Scale    = {ParameterName}_SocketTransforms[Socket * 3 + 2].xyz;
+		Velocity = float3(0, 0, 0);
 	}
 	else
 	{
+		Velocity = float3(0, 0, 0);
 		Position = float3(0, 0, 0);
 		Rotation = float4(0, 0, 0, 1);
 		Scale = float3(1, 1, 1);
 	}
 }
 
-void GetSocketTransformWS_{ParameterName}(int Socket, out float3 Position, out float4 Rotation, out float3 Scale)
+void GetSocketTransformWS_{ParameterName}(int Socket, out float3 Position, out float4 Rotation, out float3 Scale, out float3 Velocity)
 {
-	GetSocketTransform_{ParameterName}(Socket, Position, Rotation, Scale);
-	Position = mul(float4(Position, 1.0), {ParameterName}_InstanceTransform).xyz;
-	Rotation = NiagaraGPU_QuatMul({ParameterName}_InstanceRotation, Rotation);
-	Scale = mul(float4(Scale, 0.0), {ParameterName}_InstanceTransform).xyz;
+	float3 LocalPosition;
+	GetSocketTransform_{ParameterName}(Socket, LocalPosition, Rotation, Scale, Velocity);
+
+	Position				= mul(float4(LocalPosition, 1.0), {ParameterName}_InstanceTransform).xyz;
+	float3 PreviousPosition	= mul(float4(LocalPosition, 1), {ParameterName}_InstancePreviousTransform).xyz;
+	Rotation				= NiagaraGPU_QuatMul({ParameterName}_InstanceRotation, Rotation);
+	Scale					= mul(float4(Scale, 0.0), {ParameterName}_InstanceTransform).xyz;
+	Velocity				= (Position - PreviousPosition) * {ParameterName}_InvDeltaSeconds;
 }
 
-void GetFilteredSocketTransform_{ParameterName}(int SocketIndex, out float3 Position, out float4 Rotation, out float3 Scale)
+void GetSocketTransformWSInterpolated_{ParameterName}(int Socket, float Interp, out float3 Position, out float4 Rotation, out float3 Scale, out float3 Velocity)
+{
+	float3 LocalPosition, LocalScale, LocalVelocity;
+	float4 LocalRotation;
+	GetSocketTransform_{ParameterName}(Socket, LocalPosition, LocalRotation, LocalScale, LocalVelocity);
+
+	float3 CurrPosition		= mul(float4(LocalPosition, 1.0), {ParameterName}_InstanceTransform).xyz;
+	float3 PrevPosition		= mul(float4(LocalPosition, 1.0), {ParameterName}_InstancePreviousTransform).xyz;
+	float4 CurrRotation		= NiagaraGPU_QuatMul({ParameterName}_InstanceRotation, LocalRotation);
+	float4 PrevRotation		= NiagaraGPU_QuatMul({ParameterName}_InstancePreviousRotation, LocalRotation);
+	float3 CurrScale		= mul(float4(LocalScale, 0.0), {ParameterName}_InstanceTransform).xyz;
+	float3 PrevScale		= mul(float4(LocalScale, 0.0), {ParameterName}_InstancePreviousTransform).xyz;
+
+	Position	= lerp(PrevPosition, CurrPosition, Interp);
+	Rotation	= SlerpQuat(PrevRotation, CurrRotation, Interp);
+	Scale		= lerp(PrevScale, CurrScale, Interp);
+	Velocity	= (CurrPosition - PrevPosition) * {ParameterName}_InvDeltaSeconds;
+}
+
+void GetFilteredSocketTransform_{ParameterName}(int SocketIndex, out float3 Position, out float4 Rotation, out float3 Scale, out float3 Velocity)
 {
 	[branch]
 	if ( {ParameterName}_SocketCounts.y > 0 )
 	{
 		SocketIndex = clamp(SocketIndex, 0, {ParameterName}_SocketCounts.y - 1);
 		int Socket = {ParameterName}_FilteredAndUnfilteredSockets[SocketIndex];
-		GetSocketTransform_{ParameterName}(Socket, Position, Rotation, Scale);
+		GetSocketTransform_{ParameterName}(Socket, Position, Rotation, Scale, Velocity);
+		Velocity = float3(0, 0, 0);
 	}
 	else
 	{
 		Position = float3(0, 0, 0);
 		Rotation = float4(0, 0, 0, 1);
 		Scale = float3(1, 1, 1);
+		Velocity = float3(0, 0, 0);
 	}
 }
 
-void GetFilteredSocketTransformWS_{ParameterName}(int SocketIndex, out float3 Position, out float4 Rotation, out float3 Scale)
+void GetFilteredSocketTransformWS_{ParameterName}(int SocketIndex, out float3 Position, out float4 Rotation, out float3 Scale, out float3 Velocity)
 {
-	GetFilteredSocketTransform_{ParameterName}(SocketIndex, Position, Rotation, Scale);
-	Position = mul(float4(Position, 1.0), {ParameterName}_InstanceTransform).xyz;
-	Rotation = NiagaraGPU_QuatMul({ParameterName}_InstanceRotation, Rotation);
-	Scale = mul(float4(Scale, 0.0), {ParameterName}_InstanceTransform).xyz;
+	float3 LocalPosition;
+	GetFilteredSocketTransform_{ParameterName}(SocketIndex, LocalPosition, Rotation, Scale, Velocity);
+
+	Position				= mul(float4(LocalPosition, 1.0), {ParameterName}_InstanceTransform).xyz;
+	float3 PreviousPosition	= mul(float4(LocalPosition, 1), {ParameterName}_InstancePreviousTransform).xyz;
+	Rotation				= NiagaraGPU_QuatMul({ParameterName}_InstanceRotation, Rotation);
+	Scale					= mul(float4(Scale, 0.0), {ParameterName}_InstanceTransform).xyz;
+	Velocity				= (Position - PreviousPosition) * {ParameterName}_InvDeltaSeconds;
+}
+
+void GetFilteredSocketTransformWSInterpolated_{ParameterName}(int Index, float Interp, out float3 Position, out float4 Rotation, out float3 Scale, out float3 Velocity)
+{
+	int SocketIndex;
+	GetFilteredSocket_{ParameterName}(Index, SocketIndex);
+	GetSocketTransformWSInterpolated_{ParameterName}(SocketIndex, Interp, Position, Rotation, Scale, Velocity);
 }
 
-void GetUnfilteredSocketTransform_{ParameterName}(int SocketIndex, out float3 Position, out float4 Rotation, out float3 Scale)
+void GetUnfilteredSocketTransform_{ParameterName}(int SocketIndex, out float3 Position, out float4 Rotation, out float3 Scale, out float3 Velocity)
 {
 	[branch]
 	if ( {ParameterName}_SocketCounts.z > 0 )
 	{
 		SocketIndex = {ParameterName}_SocketCounts.y + clamp(SocketIndex, 0, {ParameterName}_SocketCounts.z - 1);
 		int Socket = {ParameterName}_FilteredAndUnfilteredSockets[SocketIndex];
-		GetSocketTransform_{ParameterName}(Socket, Position, Rotation, Scale);
+		GetSocketTransform_{ParameterName}(Socket, Position, Rotation, Scale, Velocity);
+		Velocity = float3(0, 0, 0);
 	}
 	else
 	{
 		Position = float3(0, 0, 0);
 		Rotation = float4(0, 0, 0, 1);
 		Scale = float3(1, 1, 1);
+		Velocity = float3(0, 0, 0);
 	}
 }
 
-void GetUnfilteredSocketTransformWS_{ParameterName}(int SocketIndex, out float3 Position, out float4 Rotation, out float3 Scale)
+void GetUnfilteredSocketTransformWS_{ParameterName}(int SocketIndex, out float3 Position, out float4 Rotation, out float3 Scale, out float3 Velocity)
+{
+	float3 LocalPosition;
+	GetUnfilteredSocketTransform_{ParameterName}(SocketIndex, LocalPosition, Rotation, Scale, Velocity);
+
+	Position				= mul(float4(LocalPosition, 1.0), {ParameterName}_InstanceTransform).xyz;
+	float3 PreviousPosition	= mul(float4(LocalPosition, 1), {ParameterName}_InstancePreviousTransform).xyz;
+	Rotation				= NiagaraGPU_QuatMul({ParameterName}_InstanceRotation, Rotation);
+	Scale					= mul(float4(Scale, 0.0), {ParameterName}_InstanceTransform).xyz;
+	Velocity				= (Position - PreviousPosition) * {ParameterName}_InvDeltaSeconds;
+}
+
+void GetUnfilteredSocketTransformWSInterpolated_{ParameterName}(int Index, float Interp, out float3 Position, out float4 Rotation, out float3 Scale, out float3 Velocity)
 {
-	GetUnfilteredSocketTransform_{ParameterName}(SocketIndex, Position, Rotation, Scale);
-	Position = mul(float4(Position, 1.0), {ParameterName}_InstanceTransform).xyz;
-	Rotation = NiagaraGPU_QuatMul({ParameterName}_InstanceRotation, Rotation);
-	Scale = mul(float4(Scale, 0.0), {ParameterName}_InstanceTransform).xyz;
+	int SocketIndex;
+	GetUnfilteredSocket_{ParameterName}(Index, SocketIndex);
+	GetSocketTransformWSInterpolated_{ParameterName}(SocketIndex, Interp, Position, Rotation, Scale, Velocity);
 }
 
 //////////////////////////////////////////////////////////////////////////
diff --git a/Engine/Plugins/FX/Niagara/Source/Niagara/Private/DataInterface/NiagaraDataInterfaceStaticMesh.cpp b/Engine/Plugins/FX/Niagara/Source/Niagara/Private/DataInterface/NiagaraDataInterfaceStaticMesh.cpp
--- a/Engine/Plugins/FX/Niagara/Source/Niagara/Private/DataInterface/NiagaraDataInterfaceStaticMesh.cpp
+++ b/Engine/Plugins/FX/Niagara/Source/Niagara/Private/DataInterface/NiagaraDataInterfaceStaticMesh.cpp
@@ -109,6 +109,7 @@ namespace NDIStaticMeshLocal
 			RefactoredV1 = 1,
 			LargeWorldCoordinates = 2,
 			LargeWorldCoordinates2 = 3,
+			AddSocketVelocity = 4,
 
 			VersionPlusOne,
 			LatestVersion = VersionPlusOne - 1
@@ -123,6 +124,7 @@ namespace NDIStaticMeshLocal
 
 	static const FName	GetVertexName("GetVertex");
 	static const FName	GetVertexWSName("GetVertexWS");
+	static const FName	GetVertexWSInterpolatedName("GetVertexWSInterpolated");
 	static const FName	GetVertexColorName("GetVertexColor");
 	static const FName	GetVertexUVName("GetVertexUV");
 
@@ -144,6 +146,7 @@ namespace NDIStaticMeshLocal
 
 	static const FName	GetTriangleName("GetTriangle");
 	static const FName	GetTriangleWSName("GetTriangleWS");
+	static const FName	GetTriangleWSInterpolatedName("GetTriangleWSInterpolated");
 	static const FName	GetTriangleColorName("GetTriangleColor");
 	static const FName	GetTriangleUVName("GetTriangleUV");
 
@@ -154,15 +157,25 @@ namespace NDIStaticMeshLocal
 	static const FName	GetSocketCountName("GetSocketCount");
 	static const FName	GetFilteredSocketCountName("GetFilteredSocketCount");
 	static const FName	GetUnfilteredSocketCountName("GetUnfilteredSocketCount");
+	
 	static const FName	RandomSocketName("RandomSocket");
 	static const FName	RandomFilteredSocketName("RandomFilteredSocket");
 	static const FName	RandomUnfilteredSocketName("RandomUnfilteredSocket");
+	
 	static const FName	GetSocketTransformName("GetSocketTransform");
 	static const FName	GetSocketTransformWSName("GetSocketTransformWS");
+	static const FName	GetSocketTransformWSInterpolatedName("GetSocketTransformWSInterpolated");
+
 	static const FName	GetFilteredSocketTransformName("GetFilteredSocketTransform");
 	static const FName	GetFilteredSocketTransformWSName("GetFilteredSocketTransformWS");
+	static const FName	GetFilteredSocketTransformWSInterpolatedName("GetFilteredSocketTransformWSInterpolated");
+
 	static const FName	GetUnfilteredSocketTransformName("GetUnfilteredSocketTransform");
 	static const FName	GetUnfilteredSocketTransformWSName("GetUnfilteredSocketTransformWS");
+	static const FName	GetUnfilteredSocketTransformWSInterpolatedName("GetUnfilteredSocketTransformWSInterpolated");
+
+	static const FName	GetFilteredSocketName("GetFilteredSocket");
+	static const FName	GetUnfilteredSocketName("GetUnfilteredSocket");
 
 	//////////////////////////////////////////////////////////////////////////
 	// Section functions
@@ -1233,6 +1246,45 @@ namespace NDIStaticMeshLocal
 			return Tangent;
 		}
 
+		FORCEINLINE FVector3f GetTriangleTangentXInterpolated(const FVector3f& BaryCoord, float Interp, int32 Index0, int32 Index1, int32 Index2) const
+		{
+			FVector3f Tangent;
+			Tangent = LODResource->VertexBuffers.StaticMeshVertexBuffer.VertexTangentX(Index0) * BaryCoord.X;
+			Tangent += LODResource->VertexBuffers.StaticMeshVertexBuffer.VertexTangentX(Index1) * BaryCoord.Y;
+			Tangent += LODResource->VertexBuffers.StaticMeshVertexBuffer.VertexTangentX(Index2) * BaryCoord.Z;
+			FVector3f CurrTangent = Tangent;
+			FVector3f PrevTangent = Tangent;
+			TransformHandler.TransformVector(PrevTangent, FMatrix44f(InstanceData->PrevTransform));						// LWC_TODO: Precision loss?
+			TransformHandler.TransformVector(CurrTangent, FMatrix44f(InstanceData->Transform));						// LWC_TODO: Precision loss?
+			return FMath::Lerp(PrevTangent, CurrTangent, Interp);
+		}
+
+		FORCEINLINE FVector3f GetTriangleTangentYInterpolated(const FVector3f& BaryCoord, float Interp, int32 Index0, int32 Index1, int32 Index2) const
+		{
+			FVector3f Tangent;
+			Tangent = LODResource->VertexBuffers.StaticMeshVertexBuffer.VertexTangentY(Index0) * BaryCoord.X;
+			Tangent += LODResource->VertexBuffers.StaticMeshVertexBuffer.VertexTangentY(Index1) * BaryCoord.Y;
+			Tangent += LODResource->VertexBuffers.StaticMeshVertexBuffer.VertexTangentY(Index2) * BaryCoord.Z;
+			FVector3f CurrTangent = Tangent;
+			FVector3f PrevTangent = Tangent;
+			TransformHandler.TransformVector(PrevTangent, FMatrix44f(InstanceData->PrevTransform));						// LWC_TODO: Precision loss?
+			TransformHandler.TransformVector(CurrTangent, FMatrix44f(InstanceData->Transform));						// LWC_TODO: Precision loss?
+			return FMath::Lerp(PrevTangent, CurrTangent, Interp);
+		}
+
+		FORCEINLINE FVector3f GetTriangleTangentZInterpolated(const FVector3f& BaryCoord, float Interp, int32 Index0, int32 Index1, int32 Index2) const
+		{
+			FVector3f Tangent;
+			Tangent = LODResource->VertexBuffers.StaticMeshVertexBuffer.VertexTangentZ(Index0) * BaryCoord.X;
+			Tangent += LODResource->VertexBuffers.StaticMeshVertexBuffer.VertexTangentZ(Index1) * BaryCoord.Y;
+			Tangent += LODResource->VertexBuffers.StaticMeshVertexBuffer.VertexTangentZ(Index2) * BaryCoord.Z;
+			FVector3f CurrTangent = Tangent;
+			FVector3f PrevTangent = Tangent;
+			TransformHandler.TransformVector(PrevTangent, FMatrix44f(InstanceData->PrevTransform));						// LWC_TODO: Precision loss?
+			TransformHandler.TransformVector(CurrTangent, FMatrix44f(InstanceData->Transform));						// LWC_TODO: Precision loss?
+			return FMath::Lerp(PrevTangent, CurrTangent, Interp);
+		}
+
 		FORCEINLINE FLinearColor GetTriangleColor(const FVector3f& BaryCoord, int32 Index0, int32 Index1, int32 Index2) const
 		{
 			FLinearColor Color;
@@ -1294,6 +1346,33 @@ namespace NDIStaticMeshLocal
 			return TangentZ;
 		}
 
+		FORCEINLINE FVector3f GetTangentXInterpolated(int32 Vertex, float Interp) const
+		{
+			FVector3f TangentXPrev = LODResource->VertexBuffers.StaticMeshVertexBuffer.VertexTangentX(Vertex);
+			FVector3f TangentXCurr = TangentXPrev;
+			TransformHandler.TransformVector(TangentXPrev, FMatrix44f(InstanceData->PrevTransformInverseTransposed));		// LWC_TODO: Precision loss?
+			TransformHandler.TransformVector(TangentXCurr, FMatrix44f(InstanceData->TransformInverseTransposed));				// LWC_TODO: Precision loss?
+			return FMath::Lerp(TangentXPrev, TangentXCurr, Interp);
+		}
+
+		FORCEINLINE FVector3f GetTangentYInterpolated(int32 Vertex, float Interp) const
+		{
+			FVector3f TangentYPrev = LODResource->VertexBuffers.StaticMeshVertexBuffer.VertexTangentY(Vertex);
+			FVector3f TangentYCurr = TangentYPrev;
+			TransformHandler.TransformVector(TangentYPrev, FMatrix44f(InstanceData->PrevTransformInverseTransposed));		// LWC_TODO: Precision loss?
+			TransformHandler.TransformVector(TangentYCurr, FMatrix44f(InstanceData->TransformInverseTransposed));				// LWC_TODO: Precision loss?
+			return FMath::Lerp(TangentYPrev, TangentYCurr, Interp);
+		}
+
+		FORCEINLINE FVector3f GetTangentZInterpolated(int32 Vertex, float Interp) const
+		{
+			FVector3f TangentZPrev = LODResource->VertexBuffers.StaticMeshVertexBuffer.VertexTangentZ(Vertex);
+			FVector3f TangentZCurr = TangentZPrev;
+			TransformHandler.TransformVector(TangentZPrev, FMatrix44f(InstanceData->PrevTransformInverseTransposed));		// LWC_TODO: Precision loss?
+			TransformHandler.TransformVector(TangentZCurr, FMatrix44f(InstanceData->TransformInverseTransposed));				// LWC_TODO: Precision loss?
+			return FMath::Lerp(TangentZPrev, TangentZCurr, Interp);
+		}
+
 		FORCEINLINE FLinearColor GetColor(int32 Vertex) const
 		{
 			if (OverrideVertexColors)
@@ -1412,6 +1491,22 @@ namespace NDIStaticMeshLocal
 			return 0;
 		}
 
+		FORCEINLINE void InterpolateSocket(int32 SocketIndex, float Interp, FVector3f& OutPosition, FQuat4f& OutRotation, FVector3f& OutScale, FVector3f& OutVelocity)
+		{
+			const FTransform3f& SocketTransform = InstanceData->CachedSockets[SocketIndex];
+			const FVector3f PrevPosition = PreviousTransformPosition(SocketTransform.GetLocation());
+			const FVector3f CurrPosition = TransformPosition(SocketTransform.GetLocation());
+			const FQuat4f PrevRotation = PreviousTransformRotation(SocketTransform.GetRotation());
+			const FQuat4f CurrRotation = TransformRotation(SocketTransform.GetRotation());
+			const FVector3f PrevScale = PreviousTransformVector(SocketTransform.GetScale3D());
+			const FVector3f CurrScale = TransformVector(SocketTransform.GetScale3D());
+
+			OutPosition = FMath::Lerp(PrevPosition, CurrPosition, Interp);
+			OutRotation = FQuat4f::Slerp(PrevRotation, CurrRotation, Interp);
+			OutScale = FMath::Lerp(PrevScale, CurrScale, Interp);
+			OutVelocity = (CurrPosition - PrevPosition) * GetInvDeltaSeconds();
+		}
+
 		VectorVM::FUserPtrHandler<FInstanceData_GameThread> InstanceData;
 		TRefCountPtr<const FStaticMeshLODResources> LODResource;
 		FColorVertexBuffer* OverrideVertexColors;
@@ -1748,6 +1843,11 @@ void UNiagaraDataInterfaceStaticMesh::GetFunctions(TArray<FNiagaraFunctionSignat
 		FNiagaraFunctionSignature& WsSig = OutFunctions.Add_GetRef(Sig);
 		WsSig.Name = GetVertexWSName;
 		WsSig.Outputs[0].SetType(FNiagaraTypeDefinition::GetPositionDef());
+
+		FNiagaraFunctionSignature& WsInterpSig = OutFunctions.Add_GetRef(Sig);
+		WsInterpSig.Name = GetVertexWSInterpolatedName;
+		WsInterpSig.Inputs.Emplace(FNiagaraTypeDefinition::GetFloatDef(), TEXT("Interp"));
+		WsInterpSig.Outputs[0].SetType(FNiagaraTypeDefinition::GetPositionDef());
 	}
 	{
 		FNiagaraFunctionSignature& Sig = OutFunctions.Add_GetRef(BaseSignature);
@@ -1848,6 +1948,11 @@ void UNiagaraDataInterfaceStaticMesh::GetFunctions(TArray<FNiagaraFunctionSignat
 		FNiagaraFunctionSignature& WsSig = OutFunctions.Add_GetRef(Sig);
 		WsSig.Name = GetTriangleWSName;
 		WsSig.Outputs[0].SetType(FNiagaraTypeDefinition::GetPositionDef());
+
+		FNiagaraFunctionSignature& WsInterpSig = OutFunctions.Add_GetRef(Sig);
+		WsInterpSig.Name = GetTriangleWSInterpolatedName;
+		WsInterpSig.Inputs.Emplace(FNiagaraTypeDefinition::GetFloatDef(), TEXT("Interp"));
+		WsInterpSig.Outputs[0].SetType(FNiagaraTypeDefinition::GetPositionDef());
 	}
 	{
 		FNiagaraFunctionSignature& Sig = OutFunctions.Add_GetRef(BaseSignature);
@@ -1899,16 +2004,32 @@ void UNiagaraDataInterfaceStaticMesh::GetFunctions(TArray<FNiagaraFunctionSignat
 		Sig.Outputs.Add(FNiagaraVariable(FNiagaraTypeDefinition::GetPositionDef(), TEXT("Position")));
 		Sig.Outputs.Add(FNiagaraVariable(FNiagaraTypeDefinition::GetQuatDef(), TEXT("Rotation")));
 		Sig.Outputs.Add(FNiagaraVariable(FNiagaraTypeDefinition::GetVec3Def(), TEXT("Scale")));
+		Sig.Outputs.Add(FNiagaraVariable(FNiagaraTypeDefinition::GetVec3Def(), TEXT("Velocity")));
 
 		FNiagaraFunctionSignature WsSig = Sig;
 		WsSig.Outputs[0].SetType(FNiagaraTypeDefinition::GetPositionDef());
 
+		FNiagaraFunctionSignature WsInterpSig = WsSig;
+		WsInterpSig.Inputs.Emplace(FNiagaraTypeDefinition::GetFloatDef(), TEXT("Interp"));
+
 		OutFunctions.Add_GetRef(Sig).Name = GetSocketTransformName;
 		OutFunctions.Add_GetRef(WsSig).Name = GetSocketTransformWSName;
+		OutFunctions.Add_GetRef(WsInterpSig).Name = GetSocketTransformWSInterpolatedName;
 		OutFunctions.Add_GetRef(Sig).Name = GetFilteredSocketTransformName;
 		OutFunctions.Add_GetRef(WsSig).Name = GetFilteredSocketTransformWSName;
+		OutFunctions.Add_GetRef(WsInterpSig).Name = GetFilteredSocketTransformWSInterpolatedName;
 		OutFunctions.Add_GetRef(Sig).Name = GetUnfilteredSocketTransformName;
 		OutFunctions.Add_GetRef(WsSig).Name = GetUnfilteredSocketTransformWSName;
+		OutFunctions.Add_GetRef(WsInterpSig).Name = GetUnfilteredSocketTransformWSInterpolatedName;
+	}
+
+	{
+		FNiagaraFunctionSignature Sig = BaseSignature;
+		Sig.Inputs.Emplace(FNiagaraTypeDefinition::GetIntDef(), TEXT("Index"));
+		Sig.Outputs.Emplace(FNiagaraTypeDefinition::GetIntDef(), TEXT("SocketIndex"));
+
+		OutFunctions.Add_GetRef(Sig).Name = GetFilteredSocketName;
+		OutFunctions.Add_GetRef(Sig).Name = GetUnfilteredSocketName;
 	}
 
 	////////////////////////////////////////////////////////////////////////////
@@ -2219,6 +2340,19 @@ bool UNiagaraDataInterfaceStaticMesh::UpgradeFunctionCall(FNiagaraFunctionSignat
 		}
 	}
 
+	if (FunctionSignature.FunctionVersion < EDIFunctionVersion::AddSocketVelocity)
+	{
+		if ((FunctionSignature.Name == GetSocketTransformName) ||
+			(FunctionSignature.Name == GetSocketTransformWSName) ||
+			(FunctionSignature.Name == GetFilteredSocketTransformName) ||
+			(FunctionSignature.Name == GetFilteredSocketTransformWSName) ||
+			(FunctionSignature.Name == GetUnfilteredSocketTransformName) ||
+			(FunctionSignature.Name == GetUnfilteredSocketTransformWSName))
+		{
+			FunctionSignature.Outputs.Emplace(FNiagaraTypeDefinition::GetVec3Def(), TEXT("Velocity"));
+		}
+	}
+
 	FunctionSignature.FunctionVersion = EDIFunctionVersion::LatestVersion;
 
 	return true;
@@ -2278,6 +2412,10 @@ void UNiagaraDataInterfaceStaticMesh::GetVMExternalFunction(const FVMExternalFun
 	{
 		OutFunc = FVMExternalFunction::CreateLambda([this](FVectorVMExternalFunctionContext& Context) { VMGetVertex<FNDITransformHandler>(Context); });
 	}
+	else if (BindingInfo.Name == GetVertexWSInterpolatedName)
+	{
+		OutFunc = FVMExternalFunction::CreateUObject(this, &UNiagaraDataInterfaceStaticMesh::VMGetVertexInterpolated<FNDITransformHandler>);
+	}
 	else if (BindingInfo.Name == GetVertexColorName)
 	{
 		OutFunc = FVMExternalFunction::CreateLambda([this](FVectorVMExternalFunctionContext& Context) { VMGetVertexColor(Context); });
@@ -2341,6 +2479,10 @@ void UNiagaraDataInterfaceStaticMesh::GetVMExternalFunction(const FVMExternalFun
 	{
 		OutFunc = FVMExternalFunction::CreateLambda([this](FVectorVMExternalFunctionContext& Context) { VMGetTriangle<FNDITransformHandler>(Context); });
 	}
+	else if (BindingInfo.Name == GetTriangleWSInterpolatedName)
+	{
+		OutFunc = FVMExternalFunction::CreateUObject(this, &UNiagaraDataInterfaceStaticMesh::VMGetTriangleInterpolated<FNDITransformHandler>);
+	}
 	else if (BindingInfo.Name == GetTriangleColorName)
 	{
 		OutFunc = FVMExternalFunction::CreateLambda([this](FVectorVMExternalFunctionContext& Context) { VMGetTriangleColor(Context); });
@@ -2388,6 +2530,10 @@ void UNiagaraDataInterfaceStaticMesh::GetVMExternalFunction(const FVMExternalFun
 	{
 		OutFunc = FVMExternalFunction::CreateLambda([this](FVectorVMExternalFunctionContext& Context) { VMGetSocketTransform<FNDITransformHandler>(Context); });
 	}
+	else if (BindingInfo.Name == GetSocketTransformWSInterpolatedName)
+	{
+		OutFunc = FVMExternalFunction::CreateLambda([this](FVectorVMExternalFunctionContext& Context) { VMGetSocketTransformInterpolated<FNDITransformHandler>(Context); });
+	}
 	else if (BindingInfo.Name == GetFilteredSocketTransformName)
 	{
 		OutFunc = FVMExternalFunction::CreateLambda([this](FVectorVMExternalFunctionContext& Context) { VMGetFilteredSocketTransform<FNDITransformHandlerNoop>(Context); });
@@ -2396,6 +2542,10 @@ void UNiagaraDataInterfaceStaticMesh::GetVMExternalFunction(const FVMExternalFun
 	{
 		OutFunc = FVMExternalFunction::CreateLambda([this](FVectorVMExternalFunctionContext& Context) { VMGetFilteredSocketTransform<FNDITransformHandler>(Context); });
 	}
+	else if (BindingInfo.Name == GetFilteredSocketTransformWSInterpolatedName)
+	{
+		OutFunc = FVMExternalFunction::CreateLambda([this](FVectorVMExternalFunctionContext& Context) { VMGetFilteredSocketTransformInterpolated<FNDITransformHandler>(Context); });
+	}
 	else if (BindingInfo.Name == GetUnfilteredSocketTransformName)
 	{
 		OutFunc = FVMExternalFunction::CreateLambda([this](FVectorVMExternalFunctionContext& Context) { VMGetUnfilteredSocketTransform<FNDITransformHandlerNoop>(Context); });
@@ -2404,6 +2554,18 @@ void UNiagaraDataInterfaceStaticMesh::GetVMExternalFunction(const FVMExternalFun
 	{
 		OutFunc = FVMExternalFunction::CreateLambda([this](FVectorVMExternalFunctionContext& Context) { VMGetUnfilteredSocketTransform<FNDITransformHandler>(Context); });
 	}
+	else if (BindingInfo.Name == GetUnfilteredSocketTransformWSInterpolatedName)
+	{
+		OutFunc = FVMExternalFunction::CreateLambda([this](FVectorVMExternalFunctionContext& Context) { VMGetUnfilteredSocketTransformInterpolated<FNDITransformHandler>(Context); });
+	}
+	else if (BindingInfo.Name == GetFilteredSocketName)
+	{
+		OutFunc = FVMExternalFunction::CreateUObject(this, &UNiagaraDataInterfaceStaticMesh::VMGetFilteredSocket);
+	}
+	else if (BindingInfo.Name == GetUnfilteredSocketName)
+	{
+		OutFunc = FVMExternalFunction::CreateUObject(this, &UNiagaraDataInterfaceStaticMesh::VMGetUnfilteredSocket);
+	}
 
 	//////////////////////////////////////////////////////////////////////////
 	// Section functions
@@ -3233,6 +3395,56 @@ void UNiagaraDataInterfaceStaticMesh::VMGetVertex(FVectorVMExternalFunctionConte
 	}
 }
 
+template<typename TTransformHandler>
+void UNiagaraDataInterfaceStaticMesh::VMGetVertexInterpolated(FVectorVMExternalFunctionContext& Context)
+{
+	NDIStaticMeshLocal::FStaticMeshCpuHelper<TTransformHandler> StaticMeshHelper(Context);
+	FNDIInputParam<int32> VertexParam(Context);
+	FNDIInputParam<float> InterpParam(Context);
+	FNDIOutputParam<FVector3f> OutPosition(Context);
+	FNDIOutputParam<FVector3f> OutVelocity(Context);
+	FNDIOutputParam<FVector3f> OutNormal(Context);
+	FNDIOutputParam<FVector3f> OutBitangent(Context);
+	FNDIOutputParam<FVector3f> OutTangent(Context);
+
+	const int32 PositionsMax = StaticMeshHelper.GetNumPositionVertices() - 1;
+	const int32 TangentsMax = StaticMeshHelper.GetNumTangentVertices() - 1;
+	if ((PositionsMax >= 0) && (PositionsMax == TangentsMax))
+	{
+		const float InvDt = StaticMeshHelper.GetInvDeltaSeconds();
+		for (int32 i = 0; i < Context.GetNumInstances(); ++i)
+		{
+			const float Interp = InterpParam.GetAndAdvance();
+			const int32 Vertex = FMath::Clamp(VertexParam.GetAndAdvance(), 0, PositionsMax);
+			const FVector3f LocalPosition = StaticMeshHelper.GetLocalPosition(Vertex);
+			const FVector3f Position = StaticMeshHelper.TransformPosition(LocalPosition);
+			const FVector3f PreviousPosition = StaticMeshHelper.PreviousTransformPosition(LocalPosition);
+			const FVector3f Velocity = (Position - PreviousPosition) * InvDt;
+			OutPosition.SetAndAdvance(FMath::Lerp(PreviousPosition, Position, Interp));
+			OutVelocity.SetAndAdvance(Velocity);
+			OutNormal.SetAndAdvance(StaticMeshHelper.GetTangentZInterpolated(Vertex, Interp));
+			OutBitangent.SetAndAdvance(StaticMeshHelper.GetTangentYInterpolated(Vertex, Interp));
+			OutTangent.SetAndAdvance(StaticMeshHelper.GetTangentXInterpolated(Vertex, Interp));
+		}
+	}
+	else
+	{
+		const FVector3f Position = StaticMeshHelper.TransformPosition(FVector3f::ZeroVector);
+		const FVector3f Velocity = FVector3f::ZeroVector;
+		const FVector3f TangentX = StaticMeshHelper.TransformVector(FVector3f(1.0f, 0.0f, 0.0f));
+		const FVector3f TangentY = StaticMeshHelper.TransformVector(FVector3f(0.0f, 1.0f, 0.0f));
+		const FVector3f TangentZ = StaticMeshHelper.TransformVector(FVector3f(0.0f, 0.0f, 1.0f));
+		for (int32 i = 0; i < Context.GetNumInstances(); ++i)
+		{
+			OutPosition.SetAndAdvance(Position);
+			OutVelocity.SetAndAdvance(Velocity);
+			OutNormal.SetAndAdvance(TangentZ);
+			OutBitangent.SetAndAdvance(TangentY);
+			OutTangent.SetAndAdvance(TangentX);
+		}
+	}
+}
+
 void UNiagaraDataInterfaceStaticMesh::VMGetVertexColor(FVectorVMExternalFunctionContext& Context)
 {
 	NDIStaticMeshLocal::FStaticMeshCpuHelper StaticMeshHelper(Context);
@@ -3582,6 +3794,61 @@ void UNiagaraDataInterfaceStaticMesh::VMGetTriangle(FVectorVMExternalFunctionCon
 	}
 }
 
+template<typename TTransformHandler>
+void UNiagaraDataInterfaceStaticMesh::VMGetTriangleInterpolated(FVectorVMExternalFunctionContext& Context)
+{
+	NDIStaticMeshLocal::FStaticMeshCpuHelper<TTransformHandler> StaticMeshHelper(Context);
+	FNDIInputParam<int32> TriangleParam(Context);
+	FNDIInputParam<FVector3f> BaryCoordParam(Context);
+	FNDIInputParam<float> InterpParam(Context);
+	FNDIOutputParam<FVector3f> OutPositionParam(Context);
+	FNDIOutputParam<FVector3f> OutVelocityParam(Context);
+	FNDIOutputParam<FVector3f> OutNormalParam(Context);
+	FNDIOutputParam<FVector3f> OutBitangentParam(Context);
+	FNDIOutputParam<FVector3f> OutTangentParam(Context);
+
+	const FIndexArrayView IndexArray = StaticMeshHelper.GetIndexArrayView();
+	const int32 TriangleMax = StaticMeshHelper.GetNumTriangles() - 1;
+	if (TriangleMax >= 0 && StaticMeshHelper.GetNumPositionVertices() > 0 && StaticMeshHelper.GetNumTangentVertices() > 0)
+	{
+		for (int32 i = 0; i < Context.GetNumInstances(); ++i)
+		{
+			const int32 Triangle = FMath::Clamp(TriangleParam.GetAndAdvance(), 0, TriangleMax);
+			const FVector3f BaryCoord = BaryCoordParam.GetAndAdvance();
+			const float Interp = InterpParam.GetAndAdvance();
+			const FVector3f LocalPosition = StaticMeshHelper.GetLocalTrianglePosition(BaryCoord, IndexArray[Triangle * 3 + 0], IndexArray[Triangle * 3 + 1], IndexArray[Triangle * 3 + 2]);
+			const FVector3f Position = StaticMeshHelper.TransformPosition(LocalPosition);
+			const FVector3f PreviousPosition = StaticMeshHelper.PreviousTransformPosition(LocalPosition);
+			const FVector3f Velocity = (Position - PreviousPosition) * StaticMeshHelper.GetInvDeltaSeconds();
+			const FVector3f TangentX = StaticMeshHelper.GetTriangleTangentXInterpolated(BaryCoord, Interp, IndexArray[Triangle * 3 + 0], IndexArray[Triangle * 3 + 1], IndexArray[Triangle * 3 + 2]);
+			const FVector3f TangentY = StaticMeshHelper.GetTriangleTangentYInterpolated(BaryCoord, Interp, IndexArray[Triangle * 3 + 0], IndexArray[Triangle * 3 + 1], IndexArray[Triangle * 3 + 2]);
+			const FVector3f TangentZ = StaticMeshHelper.GetTriangleTangentZInterpolated(BaryCoord, Interp, IndexArray[Triangle * 3 + 0], IndexArray[Triangle * 3 + 1], IndexArray[Triangle * 3 + 2]);
+			OutPositionParam.SetAndAdvance(FMath::Lerp(PreviousPosition, Position, Interp));
+			OutVelocityParam.SetAndAdvance(Velocity);
+			OutNormalParam.SetAndAdvance(TangentZ);
+			OutBitangentParam.SetAndAdvance(TangentY);
+			OutTangentParam.SetAndAdvance(TangentX);
+		}
+	}
+	else
+	{
+		const FVector3f Position = StaticMeshHelper.TransformPosition(FVector3f::ZeroVector);
+		const FVector3f PreviousPosition = StaticMeshHelper.PreviousTransformPosition(FVector3f::ZeroVector);
+		const FVector3f Velocity = FVector3f::ZeroVector;
+		const FVector3f TangentX = StaticMeshHelper.TransformVector(FVector3f(1.0f, 0.0f, 0.0f));
+		const FVector3f TangentY = StaticMeshHelper.TransformVector(FVector3f(0.0f, 1.0f, 0.0f));
+		const FVector3f TangentZ = StaticMeshHelper.TransformVector(FVector3f(0.0f, 0.0f, 1.0f));
+		for (int32 i = 0; i < Context.GetNumInstances(); ++i)
+		{
+			OutPositionParam.SetAndAdvance(Position);
+			OutVelocityParam.SetAndAdvance(Velocity);
+			OutNormalParam.SetAndAdvance(TangentZ);
+			OutBitangentParam.SetAndAdvance(TangentY);
+			OutTangentParam.SetAndAdvance(TangentX);
+		}
+	}
+}
+
 void UNiagaraDataInterfaceStaticMesh::VMGetTriangleColor(FVectorVMExternalFunctionContext& Context)
 {
 	NDIStaticMeshLocal::FStaticMeshCpuHelper<FNDITransformHandlerNoop> StaticMeshHelper(Context);
@@ -3780,6 +4047,7 @@ void UNiagaraDataInterfaceStaticMesh::VMGetSocketTransform(FVectorVMExternalFunc
 	FNDIOutputParam<FVector3f> OutTranslateParam(Context);
 	FNDIOutputParam<FQuat4f> OutRotateParam(Context);
 	FNDIOutputParam<FVector3f> OutScaleParam(Context);
+	FNDIOutputParam<FVector3f> OutVelocityParam(Context);
 
 	const int32 SocketMax = StaticMeshHelper.InstanceData->CachedSockets.Num() - 1;
 	if (SocketMax >= 0)
@@ -3788,9 +4056,13 @@ void UNiagaraDataInterfaceStaticMesh::VMGetSocketTransform(FVectorVMExternalFunc
 		{
 			const int32 SocketIndex = FMath::Clamp(SocketIndexParam.GetAndAdvance(), 0, SocketMax);
 			const FTransform3f& SocketTransform = StaticMeshHelper.InstanceData->CachedSockets[SocketIndex];
+			const FVector3f Position = StaticMeshHelper.TransformPosition(SocketTransform.GetLocation());
+			const FVector3f PreviousPosition = StaticMeshHelper.PreviousTransformPosition(SocketTransform.GetLocation());
+			const FVector3f Velocity = (Position - PreviousPosition) * StaticMeshHelper.GetInvDeltaSeconds();
 			OutTranslateParam.SetAndAdvance(StaticMeshHelper.TransformPosition(SocketTransform.GetTranslation()));
 			OutRotateParam.SetAndAdvance(StaticMeshHelper.TransformRotation(SocketTransform.GetRotation()));
 			OutScaleParam.SetAndAdvance(StaticMeshHelper.TransformVector(SocketTransform.GetScale3D()));
+			OutVelocityParam.SetAndAdvance(Velocity);
 		}
 	}
 	else
@@ -3798,11 +4070,55 @@ void UNiagaraDataInterfaceStaticMesh::VMGetSocketTransform(FVectorVMExternalFunc
 		const FVector3f DefaultTranslate = StaticMeshHelper.TransformPosition(FVector3f::ZeroVector);
 		const FQuat4f DefaultRotation = StaticMeshHelper.TransformRotation(FQuat4f::Identity);
 		const FVector3f DefaultScale = StaticMeshHelper.TransformVector(FVector3f::OneVector);
+		const FVector3f DefaultVelocity = FVector3f::ZeroVector;
 		for (int32 i = 0; i < Context.GetNumInstances(); ++i)
 		{
 			OutTranslateParam.SetAndAdvance(DefaultTranslate);
 			OutRotateParam.SetAndAdvance(DefaultRotation);
 			OutScaleParam.SetAndAdvance(DefaultScale);
+			OutVelocityParam.SetAndAdvance(DefaultVelocity);
+		}
+	}
+}
+
+template<typename TTransformHandler>
+void UNiagaraDataInterfaceStaticMesh::VMGetSocketTransformInterpolated(FVectorVMExternalFunctionContext& Context)
+{
+	NDIStaticMeshLocal::FStaticMeshCpuHelper<TTransformHandler> StaticMeshHelper(Context);
+	FNDIInputParam<int32> SocketIndexParam(Context);
+	FNDIInputParam<float> InterpParam(Context);
+	FNDIOutputParam<FVector3f> OutTranslateParam(Context);
+	FNDIOutputParam<FQuat4f> OutRotateParam(Context);
+	FNDIOutputParam<FVector3f> OutScaleParam(Context);
+	FNDIOutputParam<FVector3f> OutVelocityParam(Context);
+
+	const int32 SocketMax = StaticMeshHelper.InstanceData->CachedSockets.Num() - 1;
+	if (SocketMax >= 0)
+	{
+		for (int32 i = 0; i < Context.GetNumInstances(); ++i)
+		{
+			const int32 SocketIndex = FMath::Clamp(SocketIndexParam.GetAndAdvance(), 0, SocketMax);
+			FVector3f Position, Scale, Velocity;
+			FQuat4f Rotation;
+			StaticMeshHelper.InterpolateSocket(SocketIndex, InterpParam.GetAndAdvance(), Position, Rotation, Scale, Velocity);
+			OutTranslateParam.SetAndAdvance(Position);
+			OutRotateParam.SetAndAdvance(Rotation);
+			OutScaleParam.SetAndAdvance(Scale);
+			OutVelocityParam.SetAndAdvance(Velocity);
+		}
+	}
+	else
+	{
+		const FVector3f DefaultTranslate = StaticMeshHelper.TransformPosition(FVector3f::ZeroVector);
+		const FQuat4f DefaultRotation = StaticMeshHelper.TransformRotation(FQuat4f::Identity);
+		const FVector3f DefaultScale = StaticMeshHelper.TransformVector(FVector3f::OneVector);
+		const FVector3f DefaultVelocity = FVector3f::ZeroVector;
+		for (int32 i = 0; i < Context.GetNumInstances(); ++i)
+		{
+			OutTranslateParam.SetAndAdvance(DefaultTranslate);
+			OutRotateParam.SetAndAdvance(DefaultRotation);
+			OutScaleParam.SetAndAdvance(DefaultScale);
+			OutVelocityParam.SetAndAdvance(DefaultVelocity);
 		}
 	}
 }
@@ -3815,6 +4131,7 @@ void UNiagaraDataInterfaceStaticMesh::VMGetFilteredSocketTransform(FVectorVMExte
 	FNDIOutputParam<FVector3f> OutTranslateParam(Context);
 	FNDIOutputParam<FQuat4f> OutRotateParam(Context);
 	FNDIOutputParam<FVector3f> OutScaleParam(Context);
+	FNDIOutputParam<FVector3f> OutVelocityParam(Context);
 
 	const int32 SocketMax = StaticMeshHelper.InstanceData->NumFilteredSockets - 1;
 	if (SocketMax >= 0)
@@ -3824,9 +4141,13 @@ void UNiagaraDataInterfaceStaticMesh::VMGetFilteredSocketTransform(FVectorVMExte
 			const int32 FilteredSocketIndex = FMath::Clamp(SocketIndexParam.GetAndAdvance(), 0, SocketMax);
 			const int32 SocketIndex = StaticMeshHelper.InstanceData->FilteredAndUnfilteredSockets[FilteredSocketIndex];
 			const FTransform3f& SocketTransform = StaticMeshHelper.InstanceData->CachedSockets[SocketIndex];
+			const FVector3f Position = StaticMeshHelper.TransformPosition(SocketTransform.GetLocation());
+			const FVector3f PreviousPosition = StaticMeshHelper.PreviousTransformPosition(SocketTransform.GetLocation());
+			const FVector3f Velocity = (Position - PreviousPosition) * StaticMeshHelper.GetInvDeltaSeconds();
 			OutTranslateParam.SetAndAdvance(StaticMeshHelper.TransformPosition(SocketTransform.GetTranslation()));
 			OutRotateParam.SetAndAdvance(StaticMeshHelper.TransformRotation(SocketTransform.GetRotation()));
 			OutScaleParam.SetAndAdvance(StaticMeshHelper.TransformVector(SocketTransform.GetScale3D()));
+			OutVelocityParam.SetAndAdvance(Velocity);
 		}
 	}
 	else
@@ -3834,11 +4155,56 @@ void UNiagaraDataInterfaceStaticMesh::VMGetFilteredSocketTransform(FVectorVMExte
 		const FVector3f DefaultTranslate = StaticMeshHelper.TransformPosition(FVector3f::ZeroVector);
 		const FQuat4f DefaultRotation = StaticMeshHelper.TransformRotation(FQuat4f::Identity);
 		const FVector3f DefaultScale = StaticMeshHelper.TransformVector(FVector3f::OneVector);
+		const FVector3f DefaultVelocity = FVector3f::ZeroVector;
 		for (int32 i = 0; i < Context.GetNumInstances(); ++i)
 		{
 			OutTranslateParam.SetAndAdvance(DefaultTranslate);
 			OutRotateParam.SetAndAdvance(DefaultRotation);
 			OutScaleParam.SetAndAdvance(DefaultScale);
+			OutVelocityParam.SetAndAdvance(DefaultVelocity);
+		}
+	}
+}
+
+template<typename TTransformHandler>
+void UNiagaraDataInterfaceStaticMesh::VMGetFilteredSocketTransformInterpolated(FVectorVMExternalFunctionContext& Context)
+{
+	NDIStaticMeshLocal::FStaticMeshCpuHelper<TTransformHandler> StaticMeshHelper(Context);
+	FNDIInputParam<int32> SocketIndexParam(Context);
+	FNDIInputParam<float> InterpParam(Context);
+	FNDIOutputParam<FVector3f> OutTranslateParam(Context);
+	FNDIOutputParam<FQuat4f> OutRotateParam(Context);
+	FNDIOutputParam<FVector3f> OutScaleParam(Context);
+	FNDIOutputParam<FVector3f> OutVelocityParam(Context);
+
+	const int32 SocketMax = StaticMeshHelper.InstanceData->NumFilteredSockets - 1;
+	if (SocketMax >= 0)
+	{
+		for (int32 i = 0; i < Context.GetNumInstances(); ++i)
+		{
+			const int32 FilteredSocketIndex = FMath::Clamp(SocketIndexParam.GetAndAdvance(), 0, SocketMax);
+			const int32 SocketIndex = StaticMeshHelper.InstanceData->FilteredAndUnfilteredSockets[FilteredSocketIndex];
+			FVector3f Position, Scale, Velocity;
+			FQuat4f Rotation;
+			StaticMeshHelper.InterpolateSocket(SocketIndex, InterpParam.GetAndAdvance(), Position, Rotation, Scale, Velocity);
+			OutTranslateParam.SetAndAdvance(Position);
+			OutRotateParam.SetAndAdvance(Rotation);
+			OutScaleParam.SetAndAdvance(Scale);
+			OutVelocityParam.SetAndAdvance(Velocity);
+		}
+	}
+	else
+	{
+		const FVector3f DefaultTranslate = StaticMeshHelper.TransformPosition(FVector3f::ZeroVector);
+		const FQuat4f DefaultRotation = StaticMeshHelper.TransformRotation(FQuat4f::Identity);
+		const FVector3f DefaultScale = StaticMeshHelper.TransformVector(FVector3f::OneVector);
+		const FVector3f DefaultVelocity = FVector3f::ZeroVector;
+		for (int32 i = 0; i < Context.GetNumInstances(); ++i)
+		{
+			OutTranslateParam.SetAndAdvance(DefaultTranslate);
+			OutRotateParam.SetAndAdvance(DefaultRotation);
+			OutScaleParam.SetAndAdvance(DefaultScale);
+			OutVelocityParam.SetAndAdvance(DefaultVelocity);
 		}
 	}
 }
@@ -3851,6 +4217,7 @@ void UNiagaraDataInterfaceStaticMesh::VMGetUnfilteredSocketTransform(FVectorVMEx
 	FNDIOutputParam<FVector3f> OutTranslateParam(Context);
 	FNDIOutputParam<FQuat4f> OutRotateParam(Context);
 	FNDIOutputParam<FVector3f> OutScaleParam(Context);
+	FNDIOutputParam<FVector3f> OutVelocityParam(Context);
 
 	const int32 UnfilteredOffset = StaticMeshHelper.InstanceData->NumFilteredSockets;
 	const int32 SocketMax = StaticMeshHelper.InstanceData->FilteredAndUnfilteredSockets.Num() - UnfilteredOffset - 1;
@@ -3861,9 +4228,57 @@ void UNiagaraDataInterfaceStaticMesh::VMGetUnfilteredSocketTransform(FVectorVMEx
 			const int32 UnfilteredSocketIndex = FMath::Clamp(SocketIndexParam.GetAndAdvance(), 0, SocketMax) + UnfilteredOffset;
 			const int32 SocketIndex = StaticMeshHelper.InstanceData->FilteredAndUnfilteredSockets[UnfilteredSocketIndex];
 			const FTransform3f& SocketTransform = StaticMeshHelper.InstanceData->CachedSockets[SocketIndex];
+			const FVector3f Position = StaticMeshHelper.TransformPosition(SocketTransform.GetLocation());
+			const FVector3f PreviousPosition = StaticMeshHelper.PreviousTransformPosition(SocketTransform.GetLocation());
+			const FVector3f Velocity = (Position - PreviousPosition) * StaticMeshHelper.GetInvDeltaSeconds();
 			OutTranslateParam.SetAndAdvance(StaticMeshHelper.TransformPosition(SocketTransform.GetTranslation()));
 			OutRotateParam.SetAndAdvance(StaticMeshHelper.TransformRotation(SocketTransform.GetRotation()));
 			OutScaleParam.SetAndAdvance(StaticMeshHelper.TransformVector(SocketTransform.GetScale3D()));
+			OutVelocityParam.SetAndAdvance(Velocity);
+		}
+	}
+	else
+	{
+		const FVector3f DefaultTranslate = StaticMeshHelper.TransformPosition(FVector3f::ZeroVector);
+		const FQuat4f DefaultRotation = StaticMeshHelper.TransformRotation(FQuat4f::Identity);
+		const FVector3f DefaultScale = StaticMeshHelper.TransformVector(FVector3f::OneVector);
+		const FVector3f DefaultVelocity = FVector3f::ZeroVector;
+		for (int32 i = 0; i < Context.GetNumInstances(); ++i)
+		{
+			OutTranslateParam.SetAndAdvance(DefaultTranslate);
+			OutRotateParam.SetAndAdvance(DefaultRotation);
+			OutScaleParam.SetAndAdvance(DefaultScale);
+			OutVelocityParam.SetAndAdvance(DefaultVelocity);
+		}
+	}
+}
+
+template<typename TTransformHandler>
+void UNiagaraDataInterfaceStaticMesh::VMGetUnfilteredSocketTransformInterpolated(FVectorVMExternalFunctionContext& Context)
+{
+	NDIStaticMeshLocal::FStaticMeshCpuHelper<TTransformHandler> StaticMeshHelper(Context);
+	FNDIInputParam<int32> SocketIndexParam(Context);
+	FNDIInputParam<float> InterpParam(Context);
+	FNDIOutputParam<FVector3f> OutTranslateParam(Context);
+	FNDIOutputParam<FQuat4f> OutRotateParam(Context);
+	FNDIOutputParam<FVector3f> OutScaleParam(Context);
+	FNDIOutputParam<FVector3f> OutVelocityParam(Context);
+
+	const int32 UnfilteredOffset = StaticMeshHelper.InstanceData->NumFilteredSockets;
+	const int32 SocketMax = StaticMeshHelper.InstanceData->FilteredAndUnfilteredSockets.Num() - UnfilteredOffset - 1;
+	if (SocketMax >= 0)
+	{
+		for (int32 i = 0; i < Context.GetNumInstances(); ++i)
+		{
+			const int32 UnfilteredSocketIndex = FMath::Clamp(SocketIndexParam.GetAndAdvance(), 0, SocketMax) + UnfilteredOffset;
+			const int32 SocketIndex = StaticMeshHelper.InstanceData->FilteredAndUnfilteredSockets[UnfilteredSocketIndex];
+			FVector3f Position, Scale, Velocity;
+			FQuat4f Rotation;
+			StaticMeshHelper.InterpolateSocket(SocketIndex, InterpParam.GetAndAdvance(), Position, Rotation, Scale, Velocity);
+			OutTranslateParam.SetAndAdvance(Position);
+			OutRotateParam.SetAndAdvance(Rotation);
+			OutScaleParam.SetAndAdvance(Scale);
+			OutVelocityParam.SetAndAdvance(Velocity);
 		}
 	}
 	else
@@ -3871,15 +4286,50 @@ void UNiagaraDataInterfaceStaticMesh::VMGetUnfilteredSocketTransform(FVectorVMEx
 		const FVector3f DefaultTranslate = StaticMeshHelper.TransformPosition(FVector3f::ZeroVector);
 		const FQuat4f DefaultRotation = StaticMeshHelper.TransformRotation(FQuat4f::Identity);
 		const FVector3f DefaultScale = StaticMeshHelper.TransformVector(FVector3f::OneVector);
+		const FVector3f DefaultVelocity = FVector3f::ZeroVector;
 		for (int32 i = 0; i < Context.GetNumInstances(); ++i)
 		{
 			OutTranslateParam.SetAndAdvance(DefaultTranslate);
 			OutRotateParam.SetAndAdvance(DefaultRotation);
 			OutScaleParam.SetAndAdvance(DefaultScale);
+			OutVelocityParam.SetAndAdvance(DefaultVelocity);
 		}
 	}
 }
 
+void UNiagaraDataInterfaceStaticMesh::VMGetFilteredSocket(FVectorVMExternalFunctionContext& Context)
+{
+	VectorVM::FUserPtrHandler<NDIStaticMeshLocal::FInstanceData_GameThread> InstanceData(Context);
+	FNDIInputParam<int32> InIndex(Context);
+	FNDIOutputParam<int32> OutSocketIndex(Context);
+
+	const int32 SocketMax = InstanceData->NumFilteredSockets - 1;
+	for (int32 i = 0; i < Context.GetNumInstances(); ++i)
+	{
+		const int32 Index = InIndex.GetAndAdvance();
+		const int32 FilteredSocketIndex = FMath::Clamp(Index, 0, SocketMax);
+		const int32 SocketIndex = SocketMax >= 0 ? InstanceData->FilteredAndUnfilteredSockets[FilteredSocketIndex] : INDEX_NONE;
+		OutSocketIndex.SetAndAdvance(SocketIndex);
+	}
+}
+
+void UNiagaraDataInterfaceStaticMesh::VMGetUnfilteredSocket(FVectorVMExternalFunctionContext& Context)
+{
+	VectorVM::FUserPtrHandler<NDIStaticMeshLocal::FInstanceData_GameThread> InstanceData(Context);
+	FNDIInputParam<int32> InIndex(Context);
+	FNDIOutputParam<int32> OutSocketIndex(Context);
+
+	const int32 UnfilteredOffset = InstanceData->NumFilteredSockets;
+	const int32 SocketMax = InstanceData->FilteredAndUnfilteredSockets.Num() - UnfilteredOffset - 1;
+	for (int32 i = 0; i < Context.GetNumInstances(); ++i)
+	{
+		const int32 Index = InIndex.GetAndAdvance();
+		const int32 UnfilteredSocketIndex = FMath::Clamp(Index, 0, SocketMax) + UnfilteredOffset;
+		const int32 SocketIndex = SocketMax >= 0 ? InstanceData->FilteredAndUnfilteredSockets[UnfilteredSocketIndex] : INDEX_NONE;
+		OutSocketIndex.SetAndAdvance(SocketIndex);
+	}
+}
+
 //////////////////////////////////////////////////////////////////////////
 // Section functions
 void UNiagaraDataInterfaceStaticMesh::VMIsValidSection(FVectorVMExternalFunctionContext& Context)
diff --git a/Engine/Plugins/FX/Niagara/Source/Niagara/Private/DataInterface/NiagaraDataInterfaceStaticMesh.h b/Engine/Plugins/FX/Niagara/Source/Niagara/Private/DataInterface/NiagaraDataInterfaceStaticMesh.h
--- a/Engine/Plugins/FX/Niagara/Source/Niagara/Private/DataInterface/NiagaraDataInterfaceStaticMesh.h
+++ b/Engine/Plugins/FX/Niagara/Source/Niagara/Private/DataInterface/NiagaraDataInterfaceStaticMesh.h
@@ -170,6 +170,8 @@ protected:
 
 	template<typename TTransformHandler>
 	void VMGetVertex(FVectorVMExternalFunctionContext& Context);
+	template<typename TTransformHandler>
+	void VMGetVertexInterpolated(FVectorVMExternalFunctionContext& Context);
 	void VMGetVertexColor(FVectorVMExternalFunctionContext& Context);
 	void VMGetVertexUV(FVectorVMExternalFunctionContext& Context);
 
@@ -191,6 +193,8 @@ protected:
 
 	template<typename TTransformHandler>
 	void VMGetTriangle(FVectorVMExternalFunctionContext& Context);
+	template<typename TTransformHandler>
+	void VMGetTriangleInterpolated(FVectorVMExternalFunctionContext& Context);
 	void VMGetTriangleColor(FVectorVMExternalFunctionContext& Context);
 	void VMGetTriangleUV(FVectorVMExternalFunctionContext& Context);
 
@@ -206,9 +210,17 @@ protected:
 	template<typename TTransformHandler>
 	void VMGetSocketTransform(FVectorVMExternalFunctionContext& Context);
 	template<typename TTransformHandler>
+	void VMGetSocketTransformInterpolated(FVectorVMExternalFunctionContext& Context);
+	template<typename TTransformHandler>
 	void VMGetFilteredSocketTransform(FVectorVMExternalFunctionContext& Context);
 	template<typename TTransformHandler>
+	void VMGetFilteredSocketTransformInterpolated(FVectorVMExternalFunctionContext& Context);
+	template<typename TTransformHandler>
 	void VMGetUnfilteredSocketTransform(FVectorVMExternalFunctionContext& Context);
+	template<typename TTransformHandler>
+	void VMGetUnfilteredSocketTransformInterpolated(FVectorVMExternalFunctionContext& Context);
+	void VMGetFilteredSocket(FVectorVMExternalFunctionContext& Context);
+	void VMGetUnfilteredSocket(FVectorVMExternalFunctionContext& Context);
 
 	// Section functions
 	void VMIsValidSection(FVectorVMExternalFunctionContext& Context);
