Making it so that all loaded enums don't show up in the make/break list unless library only flag is off.

diff --git a/Engine/Plugins/FX/Niagara/Source/NiagaraEditor/Private/EdGraphSchema_Niagara.cpp b/Engine/Plugins/FX/Niagara/Source/NiagaraEditor/Private/EdGraphSchema_Niagara.cpp
--- a/Engine/Plugins/FX/Niagara/Source/NiagaraEditor/Private/EdGraphSchema_Niagara.cpp
+++ b/Engine/Plugins/FX/Niagara/Source/NiagaraEditor/Private/EdGraphSchema_Niagara.cpp
@@ -645,13 +645,13 @@ PRAGMA_ENABLE_DEPRECATION_WARNINGS
 			FText BreakCat = LOCTEXT("NiagaraBreak", "Break");
 
 			FText DescFmt = LOCTEXT("NiagaraMakeBreakFmt", "{0}");
-			auto MakeBreakType = [&](FNiagaraTypeDefinition Type, bool bMake)
+			auto MakeBreakType = [&](FNiagaraTypeDefinition Type, bool bMake, bool bInLibrary)
 			{
 				FText DisplayName = Type.GetNameText();
 				FText Desc = FText::Format(DescFmt, DisplayName);
 				
 				UNiagaraNodeConvert* ConvertNode = NewObject<UNiagaraNodeConvert>(OwnerOfTemporaries);
-				AddNewNodeMenuAction(NewActions, ConvertNode, DisplayName, ENiagaraMenuSections::General, {bMake ? MakeCat.ToString() : BreakCat.ToString()}, FText::GetEmpty(), FText::GetEmpty());
+				TSharedPtr<FNiagaraAction_NewNode> Action = AddNewNodeMenuAction(NewActions, ConvertNode, DisplayName, ENiagaraMenuSections::General, {bMake ? MakeCat.ToString() : BreakCat.ToString()}, FText::GetEmpty(), FText::GetEmpty());
 				if (bMake)
 				{
 					ConvertNode->InitAsMake(Type);
@@ -660,6 +660,7 @@ PRAGMA_ENABLE_DEPRECATION_WARNINGS
 				{
 					ConvertNode->InitAsBreak(Type);
 				}
+				Action->bIsInLibrary = bInLibrary;
 			};
 
 			// we don't allow make on the following types
@@ -671,6 +672,7 @@ PRAGMA_ENABLE_DEPRECATION_WARNINGS
 			for (const FNiagaraTypeDefinition& CandidateType : FNiagaraTypeRegistry::GetRegisteredTypes())
 			{
 				bool bAddMake = false;
+				bool bIsInLibrary = true;
 				
 				if(!CandidateType.IsValid() || CandidateType.IsUObject() || NoMakeAllowed.Contains(CandidateType))
 				{
@@ -715,11 +717,17 @@ PRAGMA_ENABLE_DEPRECATION_WARNINGS
 				else
 				{
 					bAddMake = true;
+					if (!FNiagaraTypeRegistry::GetRegisteredParameterTypes().Contains(CandidateType) &&
+						!FNiagaraTypeRegistry::GetRegisteredPayloadTypes().Contains(CandidateType) && 
+						!FNiagaraTypeRegistry::GetUserDefinedTypes().Contains(CandidateType))
+					{
+						bIsInLibrary = false;
+					}
 				}
 
 				if(bAddMake)
 				{
-					MakeBreakType(CandidateType, true);
+					MakeBreakType(CandidateType, true, bIsInLibrary);
 				}
 			}
 			
@@ -731,7 +739,7 @@ PRAGMA_ENABLE_DEPRECATION_WARNINGS
 				{
 					if(FromPin->Direction == EGPD_Output)
 					{
-						MakeBreakType(PinType, false);
+						MakeBreakType(PinType, false, true);
 					}					
 				}
 			}				
