Detecting if a signature actually changes when updgrading a data interface function call. This was causing an unneccessary reallaction of pins when doing a static duplicate.
Also added recovery logic for when some data pollution could have happened and led to crash on asset open.
Also added some detection logic for when this happens to trigger warnings that should trigger user reaction so that we can resolve additional issues.

diff --git a/Engine/Plugins/FX/Niagara/Source/Niagara/Public/NiagaraCommon.h b/Engine/Plugins/FX/Niagara/Source/Niagara/Public/NiagaraCommon.h
--- a/Engine/Plugins/FX/Niagara/Source/Niagara/Public/NiagaraCommon.h
+++ b/Engine/Plugins/FX/Niagara/Source/Niagara/Public/NiagaraCommon.h
@@ -474,6 +474,11 @@ struct NIAGARA_API FNiagaraFunctionSignature
 		}
 	}
 
+	bool operator!=(const FNiagaraFunctionSignature& Other) const
+	{
+		return !(*this == Other);
+	}
+
 	bool operator==(const FNiagaraFunctionSignature& Other) const
 	{
 		bool bFunctionSpecifiersEqual = [&]()
diff --git a/Engine/Plugins/FX/Niagara/Source/NiagaraEditor/Private/NiagaraGraph.cpp b/Engine/Plugins/FX/Niagara/Source/NiagaraEditor/Private/NiagaraGraph.cpp
--- a/Engine/Plugins/FX/Niagara/Source/NiagaraEditor/Private/NiagaraGraph.cpp
+++ b/Engine/Plugins/FX/Niagara/Source/NiagaraEditor/Private/NiagaraGraph.cpp
@@ -92,7 +92,18 @@ void FNiagaraGraphScriptUsageInfo::PostLoad(UObject* Owner)
 			CompileHash = FNiagaraCompileHash(DataHash_DEPRECATED);
 			CompileHashFromGraph = FNiagaraCompileHash(DataHash_DEPRECATED);
 		}
+	}	
+}
+
+bool FNiagaraGraphScriptUsageInfo::IsValid() const
+{
+	for (int32 i = 0; i < Traversal.Num(); i++)
+	{
+		if (Traversal[i] == nullptr)
+			return false;
 	}
+
+	return true;
 }
 
 UNiagaraGraph::UNiagaraGraph(const FObjectInitializer& ObjectInitializer)
@@ -182,9 +193,20 @@ void UNiagaraGraph::PostLoad()
 		Node->ConditionalPostLoad();
 	}
 
+	// There was a bug where the graph traversal could have nullptr entries in it, which causes crashes later. Detect these and just clear out the cache.
+	bool bAnyInvalid = false;
 	for (FNiagaraGraphScriptUsageInfo& CachedUsageInfoItem : CachedUsageInfo)
 	{
 		CachedUsageInfoItem.PostLoad(this);
+		if (!CachedUsageInfoItem.IsValid())
+		{
+			bAnyInvalid = true;
+		}
+	}
+	if (bAnyInvalid)
+	{
+		CachedUsageInfo.Empty();
+		LastBuiltTraversalDataChangeId = FGuid();
 	}
 
 	// In the past, we didn't bother setting the CallSortPriority and just used lexicographic ordering.
@@ -1458,6 +1480,7 @@ void BuildCompleteTraversal(UEdGraphNode* CurrentNode, TArray<UEdGraphNode*>& Al
 	}
 }
 
+
 UNiagaraGraph* UNiagaraGraph::CreateCompilationCopy(const TArray<ENiagaraScriptUsage>& CompileUsages)
 {
 	check(!bIsForCompilationOnly);
@@ -1522,11 +1545,37 @@ UNiagaraGraph* UNiagaraGraph::CreateCompilationCopy(const TArray<ENiagaraScriptU
 		Node->SetFlags(RF_Transactional);
 		DuplicationMapping.Add(Node, DupNode);
 	}
+
+	// Detect any invalidations that might have snuck in on the source graph!
+	for (UEdGraphNode* Node : Result->Nodes)
+	{
+		for (UEdGraphPin* Pin : Node->Pins)
+		{
+			for (int j = 0; j < Pin->LinkedTo.Num(); j++)
+			{
+				if (Pin->LinkedTo[j])
+				{
+					UEdGraphNode* NodeB = Pin->LinkedTo[j]->GetOwningNode();
+					if (Node && NodeB && Node->GetGraph() != NodeB->GetGraph())
+					{
+						UE_LOG(LogNiagaraEditor, Error, TEXT(""), *NodeB->GetPathName());
+						const FText WarningTemplate = LOCTEXT("DuplicationErrorToast", "Duplication Error! {0} was touched when cloning for compilation. Please report this issue to the Niagara team! ");
+						FText MismatchWarning = FText::Format(
+							WarningTemplate,
+							FText::FromString(NodeB->GetPathName()));
+						FNiagaraEditorUtilities::WarnWithToastAndLog(MismatchWarning);
+					}
+				}
+			}
+		}
+	}
+
 	for (UEdGraphNode* Node : NewNodes)
 	{
 		// fix up linked pins
 		for (UEdGraphPin* Pin : Node->Pins)
 		{
+			ensure(Pin->GetOwningNode()->GetGraph() == Result);
 			for (int i = 0; i < Pin->LinkedTo.Num(); i++)
 			{
 				UEdGraphPin* LinkedPin = Pin->LinkedTo[i];
diff --git a/Engine/Plugins/FX/Niagara/Source/NiagaraEditor/Private/NiagaraNodeFunctionCall.cpp b/Engine/Plugins/FX/Niagara/Source/NiagaraEditor/Private/NiagaraNodeFunctionCall.cpp
--- a/Engine/Plugins/FX/Niagara/Source/NiagaraEditor/Private/NiagaraNodeFunctionCall.cpp
+++ b/Engine/Plugins/FX/Niagara/Source/NiagaraEditor/Private/NiagaraNodeFunctionCall.cpp
@@ -106,10 +106,14 @@ void UNiagaraNodeFunctionCall::PostLoad()
 		if ((Signature.Inputs.Num() > 0) && Signature.Inputs[0].GetType().IsDataInterface())
 		{
 			UNiagaraDataInterface* CDO = CastChecked<UNiagaraDataInterface>(Signature.Inputs[0].GetType().GetClass()->GetDefaultObject());
+			FNiagaraFunctionSignature CopyForComparison = Signature;
 			if (CDO->UpgradeFunctionCall(Signature))
 			{
 				FunctionDisplayName.Empty();
-				ReallocatePins();
+				if (Signature != CopyForComparison)
+				{
+					ReallocatePins();
+				}
 			}
 		}
 	}
diff --git a/Engine/Plugins/FX/Niagara/Source/NiagaraEditor/Public/NiagaraGraph.h b/Engine/Plugins/FX/Niagara/Source/NiagaraEditor/Public/NiagaraGraph.h
--- a/Engine/Plugins/FX/Niagara/Source/NiagaraEditor/Public/NiagaraGraph.h
+++ b/Engine/Plugins/FX/Niagara/Source/NiagaraEditor/Public/NiagaraGraph.h
@@ -117,6 +117,8 @@ public:
 
 	void PostLoad(UObject* Owner);
 
+	bool IsValid() const;
+
 private:
 	UPROPERTY()
 	TArray<uint8> DataHash_DEPRECATED;
