- Fix Mesh Renderer dynamic parameter constant bindings not working
- Contribution by [REDACTED_PERSON]

diff --git a/Engine/Plugins/FX/Niagara/Source/Niagara/Private/NiagaraRendererMeshes.cpp b/Engine/Plugins/FX/Niagara/Source/Niagara/Private/NiagaraRendererMeshes.cpp
--- a/Engine/Plugins/FX/Niagara/Source/Niagara/Private/NiagaraRendererMeshes.cpp
+++ b/Engine/Plugins/FX/Niagara/Source/Niagara/Private/NiagaraRendererMeshes.cpp
@@ -891,16 +891,20 @@ FNiagaraMeshUniformBufferRef FNiagaraRendererMeshes::CreateVFUniformBuffer(const
 					break;
 				case ENiagaraMeshVFLayout::Type::DynamicParam0:
 					FMemory::Memcpy(&Params.DefaultDynamicMaterialParameter0, ParameterBoundData + VFBoundOffsetsInParamStore[i], sizeof(FVector4f));
+					Params.MaterialParamValidMask |= 0x1;
 					break;
 				case ENiagaraMeshVFLayout::Type::DynamicParam1:
 					FMemory::Memcpy(&Params.DefaultDynamicMaterialParameter1, ParameterBoundData + VFBoundOffsetsInParamStore[i], sizeof(FVector4f));
+					Params.MaterialParamValidMask |= 0x2;
 					break;
 				case ENiagaraMeshVFLayout::Type::DynamicParam2:
 					FMemory::Memcpy(&Params.DefaultDynamicMaterialParameter2, ParameterBoundData + VFBoundOffsetsInParamStore[i], sizeof(FVector4f));
+					Params.MaterialParamValidMask |= 0x4;
 					break;
 				case ENiagaraMeshVFLayout::Type::DynamicParam3:
 					FMemory::Memcpy(&Params.DefaultDynamicMaterialParameter3, ParameterBoundData + VFBoundOffsetsInParamStore[i], sizeof(FVector4f));
-					break;					
+					Params.MaterialParamValidMask |= 0x8;
+					break;
 				case ENiagaraMeshVFLayout::Type::CustomSorting:
 					// unsupported for now...
 					break;
