Fixed GPU compilation not being used for latest compilation status feedback properly
- Removed unused duplicate callback accessor function
- Added TOptional<LastCompilationStatus> member to NiagaraScript that is used to get the latest compilation status in case it's a GPU script

diff --git a/Engine/Plugins/FX/Niagara/Source/Niagara/Classes/NiagaraEmitter.h b/Engine/Plugins/FX/Niagara/Source/Niagara/Classes/NiagaraEmitter.h
--- a/Engine/Plugins/FX/Niagara/Source/Niagara/Classes/NiagaraEmitter.h
+++ b/Engine/Plugins/FX/Niagara/Source/Niagara/Classes/NiagaraEmitter.h
@@ -826,14 +826,8 @@ public:
 	/** Callback issued whenever a VM compilation successfully happened (even if the results are a script that cannot be executed due to errors)*/
 	NIAGARA_API FOnEmitterCompiled& OnEmitterVMCompiled();
 
-	/** Callback issued whenever a VM compilation successfully happened (even if the results are a script that cannot be executed due to errors)*/
-	NIAGARA_API FOnEmitterCompiled& OnEmitterGPUCompiled();
-
 	/** Callback issued whenever a GPU compilation successfully happened (even if the results are a script that cannot be executed due to errors)*/
-	NIAGARA_API FOnEmitterCompiled& OnGPUCompilationComplete()
-	{
-		return OnGPUScriptCompiledDelegate;
-	}
+	NIAGARA_API FOnEmitterCompiled& OnEmitterGPUCompiled();
 	
 	NIAGARA_API static bool GetForceCompileOnLoad();
 
diff --git a/Engine/Plugins/FX/Niagara/Source/Niagara/Classes/NiagaraScript.h b/Engine/Plugins/FX/Niagara/Source/Niagara/Classes/NiagaraScript.h
--- a/Engine/Plugins/FX/Niagara/Source/Niagara/Classes/NiagaraScript.h
+++ b/Engine/Plugins/FX/Niagara/Source/Niagara/Classes/NiagaraScript.h
@@ -1249,6 +1249,7 @@ private:
 #if WITH_EDITORONLY_DATA
 	TArray<FNiagaraShaderScript> LoadedScriptResources;
 	FNiagaraShaderScript* ScriptResourcesByFeatureLevel[ERHIFeatureLevel::Num];
+	TOptional<ENiagaraScriptCompileStatus> LastGPUCompileStatus;
 #endif
 
 	/** Compute shader compiled for this script */
diff --git a/Engine/Plugins/FX/Niagara/Source/Niagara/Private/NiagaraScript.cpp b/Engine/Plugins/FX/Niagara/Source/Niagara/Private/NiagaraScript.cpp
--- a/Engine/Plugins/FX/Niagara/Source/Niagara/Private/NiagaraScript.cpp
+++ b/Engine/Plugins/FX/Niagara/Source/Niagara/Private/NiagaraScript.cpp
@@ -2664,6 +2664,7 @@ void UNiagaraScript::SetVMCompilationResults(const FNiagaraVMExecutableDataId& I
 	CachedScriptVM.OptimizationTask.State = nullptr;
 	CachedParameterCollectionReferences.Empty();
 	// Proactively clear out the script resource, because it might be stale now.
+	LastGPUCompileStatus.Reset();
 	ScriptResource->Invalidate();
 
 	if (CachedScriptVM.LastCompileStatus == ENiagaraScriptCompileStatus::NCS_Error)
@@ -2870,6 +2871,33 @@ bool UNiagaraScript::RequestExternallyManagedAsyncCompile(const TSharedPtr<FNiag
 void UNiagaraScript::RaiseOnGPUCompilationComplete()
 {
 #if WITH_EDITORONLY_DATA
+	uint32 ErrorCount = 0;
+	uint32 WarningsCount = 0;
+	for(const FString& String : ScriptResource->GetCompileErrors())
+	{
+		if(String.Contains("err0r"))
+		{
+			ErrorCount++;
+		}
+		else
+		{
+			WarningsCount++;
+		}
+	}
+
+	if(ErrorCount > 0)
+	{
+		LastGPUCompileStatus = ENiagaraScriptCompileStatus::NCS_Error;
+	}
+	else if(WarningsCount > 0)
+	{
+		LastGPUCompileStatus = ENiagaraScriptCompileStatus::NCS_UpToDateWithWarnings;
+	}
+	else
+	{
+		LastGPUCompileStatus = ENiagaraScriptCompileStatus::NCS_UpToDate;
+	}
+	
 	OnGPUScriptCompiled().Broadcast(this, FGuid());
 	FNiagaraSystemUpdateContext(this, true);
 
@@ -3479,6 +3507,7 @@ void UNiagaraScript::InvalidateCompileResults(const FString& Reason)
 {
 	UE_LOG(LogNiagara, Verbose, TEXT("InvalidateCompileResults Script:%s Reason:%s"), *GetPathName(), *Reason);
 	CachedScriptVM.Reset();
+	LastGPUCompileStatus.Reset();
 	ScriptResource->Invalidate();
 	CachedScriptVMId.Invalidate();
 	if (VersionData.Num() > 0)
@@ -3850,10 +3879,18 @@ FGuid UNiagaraScript::GetBaseChangeID(const FGuid& VersionGuid) const
 
 ENiagaraScriptCompileStatus UNiagaraScript::GetLastCompileStatus() const
 {
-	if (CachedScriptVM.IsValid())
+	if(!IsGPUScript(Usage))
 	{
-		return CachedScriptVM.LastCompileStatus;
+		if (CachedScriptVM.IsValid())
+		{
+			return CachedScriptVM.LastCompileStatus;
+		}
 	}
+	else if(LastGPUCompileStatus.IsSet())
+	{
+		return LastGPUCompileStatus.GetValue(); 
+	}
+	
 	return ENiagaraScriptCompileStatus::NCS_Unknown;
 }
 #endif
