Return if selection is invalid when selecting the buffer in sim caches. No longer do a full refresh when updating the emitter index, which recreates the buffer selection list, invalidating the current selection.

diff --git a/Engine/Plugins/FX/Niagara/Source/NiagaraEditor/Private/ViewModels/NiagaraSimCacheViewModel.cpp b/Engine/Plugins/FX/Niagara/Source/NiagaraEditor/Private/ViewModels/NiagaraSimCacheViewModel.cpp
--- a/Engine/Plugins/FX/Niagara/Source/NiagaraEditor/Private/ViewModels/NiagaraSimCacheViewModel.cpp
+++ b/Engine/Plugins/FX/Niagara/Source/NiagaraEditor/Private/ViewModels/NiagaraSimCacheViewModel.cpp
@@ -107,7 +107,7 @@ void FNiagaraSimCacheViewModel::SetEmitterIndex(const int32 InEmitterIndex)
 {
 	EmitterIndex = InEmitterIndex;
 	UpdateCachedFrame();
-	OnViewDataChangedDelegate.Broadcast(true);
+	OnViewDataChangedDelegate.Broadcast(false);
 }
 
 bool FNiagaraSimCacheViewModel::IsCacheValid()
diff --git a/Engine/Plugins/FX/Niagara/Source/NiagaraEditor/Private/Widgets/SNiagaraSimCacheView.cpp b/Engine/Plugins/FX/Niagara/Source/NiagaraEditor/Private/Widgets/SNiagaraSimCacheView.cpp
--- a/Engine/Plugins/FX/Niagara/Source/NiagaraEditor/Private/Widgets/SNiagaraSimCacheView.cpp
+++ b/Engine/Plugins/FX/Niagara/Source/NiagaraEditor/Private/Widgets/SNiagaraSimCacheView.cpp
@@ -285,7 +285,12 @@ TSharedRef<SWidget> SNiagaraSimCacheView::BufferSelectionGenerateWidget(TSharedP
 
 void SNiagaraSimCacheView::BufferSelectionChanged(TSharedPtr<FBufferSelectionInfo> NewSelection, ESelectInfo::Type SelectInfo)
 {
-	SimCacheViewModel->SetEmitterIndex(NewSelection.IsValid() ? NewSelection->Key : INDEX_NONE);
+	if(!NewSelection.IsValid() || NewSelection->Key == SimCacheViewModel->GetEmitterIndex())
+	{
+		return;
+	}
+	
+	SimCacheViewModel->SetEmitterIndex(NewSelection->Key);
 	UpdateColumns(false);
 	UpdateRows(true);
 }
