- Make static mesh DI CPU access warning not include things like socket sampling

diff --git a/Engine/Plugins/FX/Niagara/Source/Niagara/Private/DataInterface/NiagaraDataInterfaceStaticMesh.cpp b/Engine/Plugins/FX/Niagara/Source/Niagara/Private/DataInterface/NiagaraDataInterfaceStaticMesh.cpp
--- a/Engine/Plugins/FX/Niagara/Source/Niagara/Private/DataInterface/NiagaraDataInterfaceStaticMesh.cpp
+++ b/Engine/Plugins/FX/Niagara/Source/Niagara/Private/DataInterface/NiagaraDataInterfaceStaticMesh.cpp
@@ -1787,8 +1787,19 @@ void UNiagaraDataInterfaceStaticMesh::GetFunctions(TArray<FNiagaraFunctionSignat
 	BaseSignature.FunctionVersion = EDIFunctionVersion::LatestVersion;
 #endif
 
-	//////////////////////////////////////////////////////////////////////////
-	// VM Vertex Sampling
+	GetVertexSamplingFunctions(OutFunctions, BaseSignature);
+	GetTriangleSamplingFunctions(OutFunctions, BaseSignature);
+	GetSocketSamplingFunctions(OutFunctions, BaseSignature);
+	GetSectionFunctions(OutFunctions, BaseSignature);
+	GetMiscFunctions(OutFunctions, BaseSignature);
+	GetUVMappingFunctions(OutFunctions, BaseSignature);
+	GetDistanceFieldFunctions(OutFunctions, BaseSignature);
+	GetDeprecatedFunctions(OutFunctions, BaseSignature);
+}
+
+void UNiagaraDataInterfaceStaticMesh::GetVertexSamplingFunctions(TArray<FNiagaraFunctionSignature>& OutFunctions, const FNiagaraFunctionSignature& BaseSignature) const
+{
+	using namespace NDIStaticMeshLocal;
 	{
 		FNiagaraFunctionSignature& Sig = OutFunctions.Add_GetRef(BaseSignature);
 		Sig.Name = IsValidVertexName;
@@ -1862,9 +1873,11 @@ void UNiagaraDataInterfaceStaticMesh::GetFunctions(TArray<FNiagaraFunctionSignat
 		Sig.Inputs.Emplace(FNiagaraTypeDefinition::GetIntDef(), TEXT("UV Set"));
 		Sig.Outputs.Emplace(FNiagaraTypeDefinition::GetVec2Def(), TEXT("UV"));
 	}
+}
 
-	//////////////////////////////////////////////////////////////////////////
-	// VM Triangle Sampling
+void UNiagaraDataInterfaceStaticMesh::GetTriangleSamplingFunctions(TArray<FNiagaraFunctionSignature>&OutFunctions, const FNiagaraFunctionSignature & BaseSignature) const
+{
+	using namespace NDIStaticMeshLocal;
 	{
 		FNiagaraFunctionSignature& Sig = OutFunctions.Add_GetRef(BaseSignature);
 		Sig.Name = IsValidTriangleName;
@@ -1978,9 +1991,11 @@ void UNiagaraDataInterfaceStaticMesh::GetFunctions(TArray<FNiagaraFunctionSignat
 		Sig.Outputs.Emplace(FNiagaraTypeDefinition::GetIntDef(), TEXT("Index1"));
 		Sig.Outputs.Emplace(FNiagaraTypeDefinition::GetIntDef(), TEXT("Index2"));
 	}
+}
 
-	//////////////////////////////////////////////////////////////////////////
-	// Socket Functions
+void UNiagaraDataInterfaceStaticMesh::GetSocketSamplingFunctions(TArray<FNiagaraFunctionSignature>& OutFunctions, const FNiagaraFunctionSignature& BaseSignature) const
+{
+	using namespace NDIStaticMeshLocal;
 	{
 		FNiagaraFunctionSignature Sig = BaseSignature;
 		Sig.Outputs.Add(FNiagaraVariable(FNiagaraTypeDefinition::GetIntDef(), TEXT("Num")));
@@ -2031,9 +2046,11 @@ void UNiagaraDataInterfaceStaticMesh::GetFunctions(TArray<FNiagaraFunctionSignat
 		OutFunctions.Add_GetRef(Sig).Name = GetFilteredSocketName;
 		OutFunctions.Add_GetRef(Sig).Name = GetUnfilteredSocketName;
 	}
+}
 
-	////////////////////////////////////////////////////////////////////////////
-	// Section functions
+void UNiagaraDataInterfaceStaticMesh::GetSectionFunctions(TArray<FNiagaraFunctionSignature>& OutFunctions, const FNiagaraFunctionSignature& BaseSignature) const
+{
+	using namespace NDIStaticMeshLocal;
 	{
 		FNiagaraFunctionSignature& Sig = OutFunctions.Add_GetRef(BaseSignature);
 		Sig.Name = IsValidSectionName;
@@ -2086,9 +2103,11 @@ void UNiagaraDataInterfaceStaticMesh::GetFunctions(TArray<FNiagaraFunctionSignat
 		OutFunctions.Add_GetRef(Sig).Name = RandomFilteredSectionName;
 		OutFunctions.Add_GetRef(Sig).Name = RandomUnfilteredSectionName;
 	}
+}
 
-	//////////////////////////////////////////////////////////////////////////
-	// VM Misc Functions
+void UNiagaraDataInterfaceStaticMesh::GetMiscFunctions(TArray<FNiagaraFunctionSignature>& OutFunctions, const FNiagaraFunctionSignature& BaseSignature) const
+{
+	using namespace NDIStaticMeshLocal;
 	{
 		FNiagaraFunctionSignature& Sig = OutFunctions.Add_GetRef(BaseSignature);
 		Sig.Name = IsValidName;
@@ -2119,9 +2138,11 @@ void UNiagaraDataInterfaceStaticMesh::GetFunctions(TArray<FNiagaraFunctionSignat
 		Sig.Name = GetWorldVelocityName;
 		Sig.Outputs.Emplace(FNiagaraTypeDefinition::GetVec3Def(), TEXT("Velocity"));
 	}
+}
 
-	//////////////////////////////////////////////////////////////////////////
-	// VM UV mapping functions
+void UNiagaraDataInterfaceStaticMesh::GetUVMappingFunctions(TArray<FNiagaraFunctionSignature>& OutFunctions, const FNiagaraFunctionSignature& BaseSignature) const
+{
+	using namespace NDIStaticMeshLocal;
 	{
 		FNiagaraFunctionSignature& Sig = OutFunctions.Add_GetRef(BaseSignature);
 		Sig.Name = GetTriangleCoordAtUVName;
@@ -2161,9 +2182,11 @@ void UNiagaraDataInterfaceStaticMesh::GetFunctions(TArray<FNiagaraFunctionSignat
 
 		Sig.bExperimental = true;
 	}
+}
 
-	//////////////////////////////////////////////////////////////////////////
-	// Distance Field Functions
+void UNiagaraDataInterfaceStaticMesh::GetDistanceFieldFunctions(TArray<FNiagaraFunctionSignature>&OutFunctions, const FNiagaraFunctionSignature & BaseSignature) const
+{
+	using namespace NDIStaticMeshLocal;
 	{
 		FNiagaraFunctionSignature& Sig = OutFunctions.Add_GetRef(BaseSignature);
 		Sig.Name = QueryDistanceFieldName;
@@ -2193,9 +2216,11 @@ void UNiagaraDataInterfaceStaticMesh::GetFunctions(TArray<FNiagaraFunctionSignat
 		);
 #endif
 	}
+}
 
-	//////////////////////////////////////////////////////////////////////////
-	// Deprecated Functions
+void UNiagaraDataInterfaceStaticMesh::GetDeprecatedFunctions(TArray<FNiagaraFunctionSignature>& OutFunctions, const FNiagaraFunctionSignature& BaseSignature) const
+{
+	using namespace NDIStaticMeshLocal;
 	{
 		FNiagaraFunctionSignature& Sig = OutFunctions.Add_GetRef(BaseSignature);
 		Sig.bSoftDeprecatedFunction = true;
@@ -2276,6 +2301,27 @@ void UNiagaraDataInterfaceStaticMesh::GetFunctions(TArray<FNiagaraFunctionSignat
 	}
 }
 
+void UNiagaraDataInterfaceStaticMesh::GetCpuAccessFunctions(TArray<FNiagaraFunctionSignature>& OutFunctions)
+{
+	using namespace NDIStaticMeshLocal;
+
+	// Setup base signature
+	FNiagaraFunctionSignature BaseSignature;
+	BaseSignature.Inputs.Emplace(FNiagaraTypeDefinition(GetClass()), TEXT("StaticMesh"));
+	BaseSignature.bMemberFunction = true;
+	BaseSignature.bRequiresContext = false;
+#if WITH_EDITORONLY_DATA
+	BaseSignature.FunctionVersion = EDIFunctionVersion::LatestVersion;
+#endif
+
+	GetVertexSamplingFunctions(OutFunctions, BaseSignature);
+	GetTriangleSamplingFunctions(OutFunctions, BaseSignature);
+	GetSectionFunctions(OutFunctions, BaseSignature);
+	GetUVMappingFunctions(OutFunctions, BaseSignature);
+	GetDistanceFieldFunctions(OutFunctions, BaseSignature);
+	GetDeprecatedFunctions(OutFunctions, BaseSignature);
+}
+
 #if WITH_EDITORONLY_DATA
 bool UNiagaraDataInterfaceStaticMesh::UpgradeFunctionCall(FNiagaraFunctionSignature& FunctionSignature)
 {
@@ -3033,44 +3079,62 @@ void UNiagaraDataInterfaceStaticMesh::DrawDebugHud(UCanvas* Canvas, FNiagaraSyst
 void UNiagaraDataInterfaceStaticMesh::GetFeedback(UNiagaraSystem* Asset, UNiagaraComponent* Component, TArray<FNiagaraDataInterfaceError>& OutErrors, TArray<FNiagaraDataInterfaceFeedback>& OutWarnings, TArray<FNiagaraDataInterfaceFeedback>& OutInfo)
 {
 	AActor* SourceActor = SoftSourceActor.Get();
-	if (SourceActor == nullptr && DefaultMesh != nullptr && !DefaultMesh->bAllowCPUAccess)
-	{
-		FNiagaraDataInterfaceError CPUAccessNotAllowedError(FText::Format(LOCTEXT("CPUAccessNotAllowedError", "This mesh needs CPU access in order to be used properly.({0})"), FText::FromString(DefaultMesh->GetName())),
-			LOCTEXT("CPUAccessNotAllowedErrorSummary", "CPU access error"),
-			FNiagaraDataInterfaceFix::CreateLambda([=]()
-		{
-			DefaultMesh->Modify();
-			DefaultMesh->bAllowCPUAccess = true;
-			return true;
-		}));
-
-		OutErrors.Add(CPUAccessNotAllowedError);
-	}
+	UStaticMesh* CurrentMesh = DefaultMesh;
 
 	bool bHasNoMeshAssignedWarning = (SourceActor == nullptr && DefaultMesh == nullptr);
-#if WITH_EDITORONLY_DATA
 	if (bHasNoMeshAssignedWarning)
 	{
-		if (UStaticMesh* LocalPreviewMesh = PreviewMesh.LoadSynchronous())
+		CurrentMesh = PreviewMesh.LoadSynchronous();
+		bHasNoMeshAssignedWarning = CurrentMesh == nullptr;
+	}
+
+	if (CurrentMesh != nullptr)
+	{
+		if (CurrentMesh->bAllowCPUAccess == false)
 		{
-			bHasNoMeshAssignedWarning = false;
+			// Get list of functions that require CPU access, this is slightly conservative but reduces the level of false warnings
+			TArray<FNiagaraFunctionSignature> CpuAccessFunctions;
+			GetCpuAccessFunctions(CpuAccessFunctions);
+
+			bool bCpuAccessWarning = false;
+
+			FNiagaraDataInterfaceUtilities::ForEachVMFunctionEquals(
+				this,
+				Asset,
+				[&](const FVMExternalFunctionBindingInfo& VMFunction) -> bool
+				{
+					bCpuAccessWarning |= CpuAccessFunctions.ContainsByPredicate([&](const FNiagaraFunctionSignature& CheckFunction) { return CheckFunction.Name == VMFunction.Name; });
+					return bCpuAccessWarning == false;
+				}
+			);
 
-			if (!LocalPreviewMesh->bAllowCPUAccess)
+			FNiagaraDataInterfaceUtilities::ForEachGpuFunctionEquals(
+				this,
+				Asset,
+				[&](const FNiagaraDataInterfaceGeneratedFunction& GpuFunction) -> bool
+				{
+					bCpuAccessWarning |= CpuAccessFunctions.ContainsByPredicate([&](const FNiagaraFunctionSignature& CheckFunction) { return CheckFunction.Name == GpuFunction.DefinitionName; });
+					return bCpuAccessWarning == false;
+				}
+			);
+
+			if (bCpuAccessWarning)
 			{
-				FNiagaraDataInterfaceError CPUAccessNotAllowedError(FText::Format(LOCTEXT("CPUAccessNotAllowedError", "This mesh needs CPU access in order to be used properly.({0})"), FText::FromString(LocalPreviewMesh->GetName())),
+				OutErrors.Emplace(
+					FText::Format(LOCTEXT("CPUAccessNotAllowedError", "This mesh needs CPU access in order to be used properly.({0})"), FText::FromString(CurrentMesh->GetName())),
 					LOCTEXT("CPUAccessNotAllowedErrorSummary", "CPU access error"),
-					FNiagaraDataInterfaceFix::CreateLambda([=]()
-					{
-						LocalPreviewMesh->Modify();
-						LocalPreviewMesh->bAllowCPUAccess = true;
-						return true;
-					}));
-
-				OutErrors.Add(CPUAccessNotAllowedError);
+					FNiagaraDataInterfaceFix::CreateLambda(
+						[=]()
+						{
+							CurrentMesh->Modify();
+							CurrentMesh->bAllowCPUAccess = true;
+							return true;
+						}
+					)
+				);
 			}
 		}
 	}
-#endif
 
 	if (bHasNoMeshAssignedWarning)
 	{
diff --git a/Engine/Plugins/FX/Niagara/Source/Niagara/Private/DataInterface/NiagaraDataInterfaceStaticMesh.h b/Engine/Plugins/FX/Niagara/Source/Niagara/Private/DataInterface/NiagaraDataInterfaceStaticMesh.h
--- a/Engine/Plugins/FX/Niagara/Source/Niagara/Private/DataInterface/NiagaraDataInterfaceStaticMesh.h
+++ b/Engine/Plugins/FX/Niagara/Source/Niagara/Private/DataInterface/NiagaraDataInterfaceStaticMesh.h
@@ -119,6 +119,16 @@ public:
 	virtual void ProvidePerInstanceDataForRenderThread(void* DataForRenderThread, void* PerInstanceData, const FNiagaraSystemInstanceID& SystemInstance) override;
 
 	virtual void GetFunctions(TArray<FNiagaraFunctionSignature>& OutFunctions) override;
+	void GetVertexSamplingFunctions(TArray<FNiagaraFunctionSignature>& OutFunctions, const FNiagaraFunctionSignature& BaseSignature) const;
+	void GetTriangleSamplingFunctions(TArray<FNiagaraFunctionSignature>& OutFunctions, const FNiagaraFunctionSignature& BaseSignature) const;
+	void GetSocketSamplingFunctions(TArray<FNiagaraFunctionSignature>& OutFunctions, const FNiagaraFunctionSignature& BaseSignature) const;
+	void GetSectionFunctions(TArray<FNiagaraFunctionSignature>& OutFunctions, const FNiagaraFunctionSignature& BaseSignature) const;
+	void GetMiscFunctions(TArray<FNiagaraFunctionSignature>& OutFunctions, const FNiagaraFunctionSignature& BaseSignature) const;
+	void GetUVMappingFunctions(TArray<FNiagaraFunctionSignature>& OutFunctions, const FNiagaraFunctionSignature& BaseSignature) const;
+	void GetDistanceFieldFunctions(TArray<FNiagaraFunctionSignature>& OutFunctions, const FNiagaraFunctionSignature& BaseSignature) const;
+	void GetDeprecatedFunctions(TArray<FNiagaraFunctionSignature>& OutFunctions, const FNiagaraFunctionSignature& BaseSignature) const;
+	void GetCpuAccessFunctions(TArray<FNiagaraFunctionSignature>& OutFunctions);
+
 #if WITH_EDITORONLY_DATA
 	virtual void GetCommonHLSL(FString& OutHLSL) override;
 	virtual bool UpgradeFunctionCall(FNiagaraFunctionSignature& FunctionSignature) override;
