Make sure that niagara shaders respect the r.DumpShaderDebugWorkerCommandLine

diff --git a/Engine/Plugins/FX/Niagara/Source/NiagaraEditor/Private/NiagaraCompiler.cpp b/Engine/Plugins/FX/Niagara/Source/NiagaraEditor/Private/NiagaraCompiler.cpp
--- a/Engine/Plugins/FX/Niagara/Source/NiagaraEditor/Private/NiagaraCompiler.cpp
+++ b/Engine/Plugins/FX/Niagara/Source/NiagaraEditor/Private/NiagaraCompiler.cpp
@@ -1627,6 +1627,8 @@ int32 FHlslNiagaraCompiler::CompileScript(const FNiagaraCompileRequestData* InCo
 {
 	SCOPE_CYCLE_COUNTER(STAT_NiagaraEditor_HlslCompiler_CompileScript);
 
+	static const auto CVarDumpCommandLine = IConsoleManager::Get().FindConsoleVariable(TEXT("r.DumpShaderDebugWorkerCommandLine"));
+
 	CompileResults.Data = MakeShared<FNiagaraVMExecutableData>();
 
 	//TODO: This should probably be done via the same route that other shaders take through the shader compiler etc.
@@ -1649,7 +1651,7 @@ int32 FHlslNiagaraCompiler::CompileScript(const FNiagaraCompileRequestData* InCo
 	Input.Environment.SetDefine(TEXT("MESHSHADER"), 0);
 	Input.Environment.SetDefine(TEXT("AMPLIFICATIONSHADER"), 0);
 	Input.Environment.IncludeVirtualPathToContentsMap.Add(TEXT("/Engine/Generated/NiagaraEmitterInstance.ush"), TranslatedHLSL);
-	Input.bGenerateDirectCompileFile = false;
+	Input.bGenerateDirectCompileFile = CVarDumpCommandLine ? CVarDumpCommandLine->GetBool() : false;
 	Input.DumpDebugInfoRootPath = GShaderCompilingManager->GetAbsoluteShaderDebugInfoDirectory() / TEXT("VM");
 	FString UsageIdStr = !InOptions.TargetUsageId.IsValid() ? TEXT("") : (TEXT("_") + InOptions.TargetUsageId.ToString());
 	Input.DebugGroupName = InCompileRequest->SourceName / InCompileRequest->EmitterUniqueName / InCompileRequest->ENiagaraScriptUsageEnum->GetNameStringByValue((int64)InOptions.TargetUsage) + UsageIdStr;
