Couple of fixes for compressed particle systems.  They remain experimental, but this addresses:
-only particle systems will have their attributes compressed
-Previous versions of our deny list will also be denied compression (we don't want to be comparing half Particles.Prevoius.Position with float Particles.Position)
-fixed stride being supplied to VF shaders

diff --git a/Engine/Plugins/FX/Niagara/Source/Niagara/Private/NiagaraRendererMeshes.cpp b/Engine/Plugins/FX/Niagara/Source/Niagara/Private/NiagaraRendererMeshes.cpp
--- a/Engine/Plugins/FX/Niagara/Source/Niagara/Private/NiagaraRendererMeshes.cpp
+++ b/Engine/Plugins/FX/Niagara/Source/Niagara/Private/NiagaraRendererMeshes.cpp
@@ -658,7 +658,7 @@ FNiagaraMeshCommonParameters FNiagaraRendererMeshes::CreateCommonShaderParams(co
 	Params.NiagaraParticleDataFloat	= ParticleMeshRenderData.ParticleFloatSRV;
 	Params.NiagaraParticleDataHalf 	= ParticleMeshRenderData.ParticleHalfSRV;
 	Params.NiagaraParticleDataInt 	= ParticleMeshRenderData.ParticleIntSRV;
-	Params.NiagaraFloatDataStride 	= ParticleMeshRenderData.ParticleFloatDataStride;
+	Params.NiagaraFloatDataStride 	= FMath::Max(ParticleMeshRenderData.ParticleFloatDataStride, ParticleMeshRenderData.ParticleHalfDataStride);
 	Params.NiagaraIntDataStride 	= ParticleMeshRenderData.ParticleIntDataStride;
 	
 	Params.SortedIndices 			= ParticleMeshRenderData.ParticleSortedIndicesSRV;
diff --git a/Engine/Plugins/FX/Niagara/Source/Niagara/Private/NiagaraRendererRibbons.cpp b/Engine/Plugins/FX/Niagara/Source/Niagara/Private/NiagaraRendererRibbons.cpp
--- a/Engine/Plugins/FX/Niagara/Source/Niagara/Private/NiagaraRendererRibbons.cpp
+++ b/Engine/Plugins/FX/Niagara/Source/Niagara/Private/NiagaraRendererRibbons.cpp
@@ -1907,7 +1907,7 @@ inline void FNiagaraRendererRibbons::SetupMeshBatchAndCollectorResourceForView(c
 	VFLooseParams.SliceVertexData = ShapeState.SliceVertexDataBuffer.SRV;
 	VFLooseParams.NiagaraParticleDataFloat = RenderingResources->ParticleFloatSRV;
 	VFLooseParams.NiagaraParticleDataHalf = RenderingResources->ParticleHalfSRV;
-	VFLooseParams.NiagaraFloatDataStride = RenderingResources->ParticleFloatDataStride;
+	VFLooseParams.NiagaraFloatDataStride = FMath::Max(RenderingResources->ParticleFloatDataStride, RenderingResources->ParticleHalfDataStride);
 	VFLooseParams.FacingMode = static_cast<uint32>(FacingMode);
 	VFLooseParams.Shape = static_cast<uint32>(ShapeState.Shape);
 	VFLooseParams.NeedsPreciseMotionVectors = GenerationConfig.NeedsPreciseMotionVectors();
diff --git a/Engine/Plugins/FX/Niagara/Source/Niagara/Private/NiagaraRendererSprites.cpp b/Engine/Plugins/FX/Niagara/Source/Niagara/Private/NiagaraRendererSprites.cpp
--- a/Engine/Plugins/FX/Niagara/Source/Niagara/Private/NiagaraRendererSprites.cpp
+++ b/Engine/Plugins/FX/Niagara/Source/Niagara/Private/NiagaraRendererSprites.cpp
@@ -825,7 +825,7 @@ void FNiagaraRendererSprites::CreateMeshBatchForView(
 	FNiagaraSpriteVFLooseParameters VFLooseParams;
 	VFLooseParams.NiagaraParticleDataFloat = ParticleSpriteRenderData.ParticleFloatSRV;
 	VFLooseParams.NiagaraParticleDataHalf = ParticleSpriteRenderData.ParticleHalfSRV;
-	VFLooseParams.NiagaraFloatDataStride = ParticleSpriteRenderData.ParticleFloatDataStride;
+	VFLooseParams.NiagaraFloatDataStride = FMath::Max(ParticleSpriteRenderData.ParticleFloatDataStride, ParticleSpriteRenderData.ParticleHalfDataStride);
 
 	FMaterialRenderProxy* MaterialRenderProxy = ParticleSpriteRenderData.DynamicDataSprites->Material;
 	check(MaterialRenderProxy);
diff --git a/Engine/Plugins/FX/Niagara/Source/Niagara/Public/NiagaraTypes.h b/Engine/Plugins/FX/Niagara/Source/Niagara/Public/NiagaraTypes.h
--- a/Engine/Plugins/FX/Niagara/Source/Niagara/Public/NiagaraTypes.h
+++ b/Engine/Plugins/FX/Niagara/Source/Niagara/Public/NiagaraTypes.h
@@ -1978,22 +1978,36 @@ struct FNiagaraVariable : public FNiagaraVariableBase
 
 	static int32 SearchArrayForPartialNameMatch(const TArray<FNiagaraVariable>& Variables, const FName& VariableName)
 	{
-		FString VarNameStr = VariableName.ToString();
-		FString BestMatchSoFar;
+		FNameBuilder VarNameBuilder(VariableName);
+		FStringView VarNameView(VarNameBuilder);
+
+		int32 BestMatchLength = INDEX_NONE;
 		int32 BestMatchIdx = INDEX_NONE;
 
+		const int32 VarNameLength = VarNameView.Len();
+
 		for (int32 i = 0; i < Variables.Num(); i++)
 		{
 			const FNiagaraVariable& TestVar = Variables[i];
-			FString TestVarNameStr = TestVar.GetName().ToString();
-			if (TestVarNameStr == VarNameStr)
+			FNameBuilder TestVarNameBuilder(TestVar.GetName());
+			FStringView TestVarNameView(TestVarNameBuilder);
+
+			if (VarNameView.Equals(TestVarNameView))
 			{
 				return i;
 			}
-			else if (VarNameStr.StartsWith(TestVarNameStr + TEXT(".")) && (BestMatchSoFar.Len() == 0 || TestVarNameStr.Len() > BestMatchSoFar.Len()))
+			else
 			{
-				BestMatchIdx = i;
-				BestMatchSoFar = TestVarNameStr;
+				const int32 TestVarNameLength = TestVarNameView.Len();
+
+				if (VarNameLength > TestVarNameLength && (BestMatchLength == INDEX_NONE || TestVarNameLength > BestMatchLength))
+				{
+					if (VarNameView.StartsWith(TestVarNameView) && VarNameView[TestVarNameLength] == TCHAR('.'))
+					{
+						BestMatchIdx = i;
+						BestMatchLength = TestVarNameLength;
+					}
+				}
 			}
 		}
 
diff --git a/Engine/Plugins/FX/Niagara/Source/NiagaraEditor/Private/NiagaraHlslTranslator.cpp b/Engine/Plugins/FX/Niagara/Source/NiagaraEditor/Private/NiagaraHlslTranslator.cpp
--- a/Engine/Plugins/FX/Niagara/Source/NiagaraEditor/Private/NiagaraHlslTranslator.cpp
+++ b/Engine/Plugins/FX/Niagara/Source/NiagaraEditor/Private/NiagaraHlslTranslator.cpp
@@ -1028,7 +1028,14 @@ void FHlslNiagaraTranslator::RecordParamMapDefinedAttributeToNamespaceVar(const
 
 static void ConvertFloatToHalf(const FNiagaraCompileOptions& InCompileOptions, TArray<FNiagaraVariable>& Attributes)
 {
-	if (InCompileOptions.AdditionalDefines.Contains(TEXT("CompressAttributes")))// && UNiagaraScript::IsParticleScript(InCompileOptions.TargetUsage))
+	// for now we're going to only process particle scripts as we don't currently support attributes
+	// being read transparently from the parameter store (which would be done for system/emitter scripts)
+	if (!UNiagaraScript::IsParticleScript(InCompileOptions.TargetUsage))
+	{
+		return;
+	}
+
+	if (InCompileOptions.AdditionalDefines.Contains(TEXT("CompressAttributes")))
 	{
 		static FNiagaraTypeDefinition ConvertMapping[][2] =
 		{
@@ -1072,6 +1079,18 @@ static void ConvertFloatToHalf(const FNiagaraCompileOptions& InCompileOptions, T
 				continue;
 			}
 
+			// also we'll check if the current attribute is a previous version of an exception because we wouldn't want
+			// those to be mismatched
+			if (FNiagaraParameterMapHistory::IsPreviousValue(Attribute))
+			{
+				FNiagaraVariable SrcAttribute = FNiagaraParameterMapHistory::GetSourceForPreviousValue(Attribute);
+
+				if (FNiagaraVariable::SearchArrayForPartialNameMatch(ConvertExceptions, SrcAttribute.GetName()) != INDEX_NONE)
+				{
+					continue;
+				}
+			}
+
 			for (int32 ConvertIt = 0; ConvertIt < UE_ARRAY_COUNT(ConvertMapping); ++ConvertIt)
 			{
 				if (Attribute.GetType() == ConvertMapping[ConvertIt][0])
diff --git a/Engine/Source/Runtime/Core/Private/UObject/UE5ReleaseStreamObjectVersion.cpp b/Engine/Source/Runtime/Core/Private/UObject/UE5ReleaseStreamObjectVersion.cpp
--- a/Engine/Source/Runtime/Core/Private/UObject/UE5ReleaseStreamObjectVersion.cpp
+++ b/Engine/Source/Runtime/Core/Private/UObject/UE5ReleaseStreamObjectVersion.cpp
@@ -12,7 +12,7 @@ TMap<FGuid, FGuid> FUE5ReleaseStreamObjectVersion::GetSystemGuids()
 	SystemGuids.Add(DevGuids.MATERIALSHADERMAP_DERIVEDDATA_VER, FGuid("79090A66F9D94B5BB285D49E5D39468E"));
 	SystemGuids.Add(DevGuids.NANITE_DERIVEDDATA_VER, FGuid("AC88CFBDDA614A9C8488F64D2DAF699D"));
 	SystemGuids.Add(DevGuids.NIAGARASHADERMAP_DERIVEDDATA_VER, FGuid("D204CCC166ED482685DAEE8AB72895AB"));
-	SystemGuids.Add(DevGuids.Niagara_LatestScriptCompileVersion, FGuid("7608C205F55D4EF8892807E4D5AF49C1"));
+	SystemGuids.Add(DevGuids.Niagara_LatestScriptCompileVersion, FGuid("8AF4589E892441FABA3E330276A74B92"));
 	SystemGuids.Add(DevGuids.SkeletalMeshDerivedDataVersion, FGuid("25C49E579B3142DDA2A8C14037267679"));
 
 	return SystemGuids;
