- Add on screen warning when we hit the instance count cap
- Change all emitter warnings to include a pointer to avoid filling the screen

#rnx

diff --git a/Engine/Plugins/FX/Niagara/Source/Niagara/Private/NiagaraEmitterInstance.cpp b/Engine/Plugins/FX/Niagara/Source/Niagara/Private/NiagaraEmitterInstance.cpp
--- a/Engine/Plugins/FX/Niagara/Source/Niagara/Private/NiagaraEmitterInstance.cpp
+++ b/Engine/Plugins/FX/Niagara/Source/Niagara/Private/NiagaraEmitterInstance.cpp
@@ -1406,6 +1406,12 @@ void FNiagaraEmitterInstance::Tick(float DeltaSeconds)
 	}
 
 	int32 AllocationSize = FMath::Max<int32>(AllocationEstimate, RequiredSize);
+#if !(UE_BUILD_SHIPPING || UE_BUILD_TEST)
+	if (uint32(AllocationSize) > MaxInstanceCount)
+	{
+		GEngine->AddOnScreenDebugMessage(uint64(uintptr_t(this)), 5.f, FColor::Yellow, FString::Printf(TEXT("%s has exceeded the maximum instance count (%u), requested (%d)"), *CachedEmitter.Emitter->GetFullName(), MaxInstanceCount, AllocationSize));
+	}
+#endif
 	AllocationSize = (int32)FMath::Min((uint32)AllocationSize, MaxInstanceCount);
 
 	if (AllocationSize > MaxAllocationCount)
@@ -1420,6 +1426,7 @@ void FNiagaraEmitterInstance::Tick(float DeltaSeconds)
 			UE_LOG(LogNiagara, Warning, TEXT("The emitter %s required many memory reallocation due to changing particle counts. Consider setting the emitter's AllocationMode property to 'manual' to improve runtime performance."), *FullName);
 		}
 	}
+	
 	int32 Overallocation = AllocationSize - RequiredSize;
 	if (Overallocation >= 0 && (MinOverallocation < 0 || Overallocation < MinOverallocation))
 	{
@@ -1498,7 +1505,7 @@ void FNiagaraEmitterInstance::Tick(float DeltaSeconds)
 						{
 							FString DebugMsg = FString::Printf(TEXT("%s has attempted to execeed max GPU per frame spawn! | Max: %d | Requested: %d | SpawnInfoEntry: %d"), *CachedEmitter.Emitter->GetFullName(), MaxParticlesSpawnedPerFrame, NumSpawnedOnGPUThisFrame, SpawnInfoIdx);
 							UE_LOG(LogNiagara, Warning, TEXT("%s"), *DebugMsg);
-							GEngine->AddOnScreenDebugMessage(-1, 5.f, FColor::Yellow, *DebugMsg);
+							GEngine->AddOnScreenDebugMessage(uint64(uintptr_t(this)), 5.f, FColor::Yellow, DebugMsg);
 							break;
 						}
 
@@ -1516,7 +1523,7 @@ void FNiagaraEmitterInstance::Tick(float DeltaSeconds)
 					{
 						FString DebugMsg = FString::Printf(TEXT("%s Exceeded Gpu spawn info count, see NIAGARA_MAX_GPU_SPAWN_INFOS for more information!"), *CachedEmitter.Emitter->GetUniqueEmitterName());
 						UE_LOG(LogNiagara, Warning, TEXT("%s"), *DebugMsg);
-						GEngine->AddOnScreenDebugMessage(-1, 5.f, FColor::Yellow, *DebugMsg);
+						GEngine->AddOnScreenDebugMessage(uint64(uintptr_t(this)), 5.f, FColor::Yellow, DebugMsg);
 						break;
 					}
 
@@ -1589,7 +1596,7 @@ void FNiagaraEmitterInstance::Tick(float DeltaSeconds)
 	{
 		FString DebugMsg = FString::Printf(TEXT("%s has attempted to exceed the max CPU particle count! | Max: %d | Requested: %u"), *CachedEmitter.Emitter->GetFullName(), GMaxNiagaraCPUParticlesPerEmitter, AllocationSize);
 		UE_LOG(LogNiagara, Warning, TEXT("%s"), *DebugMsg);
-		GEngine->AddOnScreenDebugMessage(-1, 5.f, FColor::Yellow, *DebugMsg);
+		GEngine->AddOnScreenDebugMessage(uint64(uintptr_t(this)), 5.f, FColor::Yellow, DebugMsg);
 
 		//We clear the emitters estimate otherwise we get stuck in this state forever.
 		EmitterData->ClearRuntimeAllocationEstimate();
