Disable backface culling for niagara ribbons during PSO precaching

#rnx

diff --git a/Engine/Plugins/FX/Niagara/Source/Niagara/Private/NiagaraRibbonRendererProperties.cpp b/Engine/Plugins/FX/Niagara/Source/Niagara/Private/NiagaraRibbonRendererProperties.cpp
--- a/Engine/Plugins/FX/Niagara/Source/Niagara/Private/NiagaraRibbonRendererProperties.cpp
+++ b/Engine/Plugins/FX/Niagara/Source/Niagara/Private/NiagaraRibbonRendererProperties.cpp
@@ -179,6 +179,18 @@ const FVertexFactoryType* UNiagaraRibbonRendererProperties::GetVertexFactoryType
 	return &FNiagaraRibbonVertexFactory::StaticType;
 }
 
+bool UNiagaraRibbonRendererProperties::IsBackfaceCullingDisabled() const
+{
+	if (Shape == ENiagaraRibbonShapeMode::MultiPlane)
+	{
+		return !bEnableAccurateGeometry;
+	}
+	else
+	{
+		return true;
+	}
+}
+
 bool UNiagaraRibbonRendererProperties::PopulateRequiredBindings(FNiagaraParameterStore& InParameterStore)
 {
 	bool bAnyAdded = Super::PopulateRequiredBindings(InParameterStore);
diff --git a/Engine/Plugins/FX/Niagara/Source/Niagara/Private/NiagaraSystem.cpp b/Engine/Plugins/FX/Niagara/Source/Niagara/Private/NiagaraSystem.cpp
--- a/Engine/Plugins/FX/Niagara/Source/Niagara/Private/NiagaraSystem.cpp
+++ b/Engine/Plugins/FX/Niagara/Source/Niagara/Private/NiagaraSystem.cpp
@@ -954,6 +954,7 @@ void UNiagaraSystem::PrecachePSOs()
 	struct VFsPerMaterialData
 	{
 		UMaterialInterface* MaterialInterface;
+		bool bDisableBackfaceCulling;
 		TArray<const FVertexFactoryType*, TInlineAllocator<2>> VertexFactoryTypes;
 	};
 	TArray<VFsPerMaterialData, TInlineAllocator<2>> VFsPerMaterials;
@@ -971,6 +972,7 @@ void UNiagaraSystem::PrecachePSOs()
 					{
 						return;
 					}
+					bool bDisableBackfaceCulling = Properties->IsBackfaceCullingDisabled();
 
 					// Don't have an instance yet to retrieve the possible material override data from
 					const FNiagaraEmitterInstance* EmitterInstance = nullptr;
@@ -979,11 +981,16 @@ void UNiagaraSystem::PrecachePSOs()
 
 					for (UMaterialInterface* MaterialInterface : Mats)
 					{
-						VFsPerMaterialData* VFsPerMaterial = VFsPerMaterials.FindByPredicate([MaterialInterface](const VFsPerMaterialData& Other) { return Other.MaterialInterface == MaterialInterface; });
+						VFsPerMaterialData* VFsPerMaterial = VFsPerMaterials.FindByPredicate([MaterialInterface, bDisableBackfaceCulling](const VFsPerMaterialData& Other) 
+						{ 
+							return (Other.MaterialInterface == MaterialInterface &&
+									Other.bDisableBackfaceCulling == bDisableBackfaceCulling);
+						});
 						if (VFsPerMaterial == nullptr)
 						{
 							VFsPerMaterial = &VFsPerMaterials.AddDefaulted_GetRef();
 							VFsPerMaterial->MaterialInterface = MaterialInterface;
+							VFsPerMaterial->bDisableBackfaceCulling = bDisableBackfaceCulling;
 						}
 						VFsPerMaterial->VertexFactoryTypes.AddUnique(VFType);
 					}
@@ -999,6 +1006,7 @@ void UNiagaraSystem::PrecachePSOs()
 	{
 		if (VFsPerMaterial.MaterialInterface)
 		{
+			PreCachePSOParams.bDisableBackFaceCulling = VFsPerMaterial.bDisableBackfaceCulling;
 			VFsPerMaterial.MaterialInterface->PrecachePSOs(VFsPerMaterial.VertexFactoryTypes, PreCachePSOParams);
 		}
 	}
diff --git a/Engine/Plugins/FX/Niagara/Source/Niagara/Public/NiagaraRendererProperties.h b/Engine/Plugins/FX/Niagara/Source/Niagara/Public/NiagaraRendererProperties.h
--- a/Engine/Plugins/FX/Niagara/Source/Niagara/Public/NiagaraRendererProperties.h
+++ b/Engine/Plugins/FX/Niagara/Source/Niagara/Public/NiagaraRendererProperties.h
@@ -264,6 +264,7 @@ public:
 	virtual class FNiagaraBoundsCalculator* CreateBoundsCalculator() PURE_VIRTUAL(UNiagaraRendererProperties::CreateBoundsCalculator, return nullptr;);
 	virtual void GetUsedMaterials(const FNiagaraEmitterInstance* InEmitter, TArray<UMaterialInterface*>& OutMaterials) const PURE_VIRTUAL(UNiagaraRendererProperties::GetUsedMaterials,);
 	virtual const FVertexFactoryType* GetVertexFactoryType() const { return nullptr; }
+	virtual bool IsBackfaceCullingDisabled() const { return false; }
 
 	virtual void GetStreamingMeshInfo(const FBoxSphereBounds& OwnerBounds, const FNiagaraEmitterInstance* InEmitter, TArray<FStreamingRenderAssetPrimitiveInfo>& OutStreamingRenderAssets) const {}
 
diff --git a/Engine/Plugins/FX/Niagara/Source/Niagara/Public/NiagaraRibbonRendererProperties.h b/Engine/Plugins/FX/Niagara/Source/Niagara/Public/NiagaraRibbonRendererProperties.h
--- a/Engine/Plugins/FX/Niagara/Source/Niagara/Public/NiagaraRibbonRendererProperties.h
+++ b/Engine/Plugins/FX/Niagara/Source/Niagara/Public/NiagaraRibbonRendererProperties.h
@@ -214,6 +214,7 @@ public:
 	virtual class FNiagaraBoundsCalculator* CreateBoundsCalculator() override;
 	virtual void GetUsedMaterials(const FNiagaraEmitterInstance* InEmitter, TArray<UMaterialInterface*>& OutMaterials) const override;
 	virtual const FVertexFactoryType* GetVertexFactoryType() const override;
+	virtual bool IsBackfaceCullingDisabled() const override;
 	virtual bool IsSimTargetSupported(ENiagaraSimTarget InSimTarget) const override { return true; };
 	virtual bool PopulateRequiredBindings(FNiagaraParameterStore& InParameterStore) override;
 
