Fixed default materials/mesh not being set on renderers

diff --git a/Engine/Plugins/FX/Niagara/Source/Niagara/Private/NiagaraRibbonRendererProperties.cpp b/Engine/Plugins/FX/Niagara/Source/Niagara/Private/NiagaraRibbonRendererProperties.cpp
--- a/Engine/Plugins/FX/Niagara/Source/Niagara/Private/NiagaraRibbonRendererProperties.cpp
+++ b/Engine/Plugins/FX/Niagara/Source/Niagara/Private/NiagaraRibbonRendererProperties.cpp
@@ -91,6 +91,8 @@ UNiagaraRibbonRendererProperties::UNiagaraRibbonRendererProperties()
 	AttributeBindings.Add(&PrevRibbonWidthBinding);
 	AttributeBindings.Add(&PrevRibbonFacingBinding);
 	AttributeBindings.Add(&PrevRibbonTwistBinding);
+	
+	Material = LoadObject<UMaterialInterface>(nullptr, TEXT("/Niagara/DefaultAssets/DefaultRibbonMaterial.DefaultRibbonMaterial"));
 }
 
 FNiagaraRenderer* UNiagaraRibbonRendererProperties::CreateEmitterRenderer(ERHIFeatureLevel::Type FeatureLevel, const FNiagaraEmitterInstance* Emitter, const FNiagaraSystemInstanceController& InController)
diff --git a/Engine/Plugins/FX/Niagara/Source/Niagara/Private/NiagaraSpriteRendererProperties.cpp b/Engine/Plugins/FX/Niagara/Source/Niagara/Private/NiagaraSpriteRendererProperties.cpp
--- a/Engine/Plugins/FX/Niagara/Source/Niagara/Private/NiagaraSpriteRendererProperties.cpp
+++ b/Engine/Plugins/FX/Niagara/Source/Niagara/Private/NiagaraSpriteRendererProperties.cpp
@@ -94,6 +94,8 @@ UNiagaraSpriteRendererProperties::UNiagaraSpriteRendererProperties()
 
 	// The remaining bindings are not associated with attributes in the VF layout
 	AttributeBindings.Add(&RendererVisibilityTagBinding);
+
+	Material = LoadObject<UMaterialInterface>(nullptr, TEXT("/Niagara/DefaultAssets/DefaultSpriteMaterial.DefaultSpriteMaterial"));
 }
 
 FNiagaraRenderer* UNiagaraSpriteRendererProperties::CreateEmitterRenderer(ERHIFeatureLevel::Type FeatureLevel, const FNiagaraEmitterInstance* Emitter, const FNiagaraSystemInstanceController& InController)
@@ -181,6 +183,7 @@ void UNiagaraSpriteRendererProperties::PostInitProperties()
 			return;
 		}
 		InitBindings();
+
 	}
 }
 
