- Fix static mesh DI extents being half extents
- Add tooltip for local space velocity

#rnx

diff --git a/Engine/Plugins/FX/Niagara/Shaders/Private/NiagaraDataInterfaceStaticMeshTemplate.ush b/Engine/Plugins/FX/Niagara/Shaders/Private/NiagaraDataInterfaceStaticMeshTemplate.ush
--- a/Engine/Plugins/FX/Niagara/Shaders/Private/NiagaraDataInterfaceStaticMeshTemplate.ush
+++ b/Engine/Plugins/FX/Niagara/Shaders/Private/NiagaraDataInterfaceStaticMeshTemplate.ush
@@ -64,8 +64,8 @@ int UniformTriangle_{ParameterName}(NiagaraRandInfo RandInfo, int Tri, int Secti
 void GetPreSkinnedLocalBounds_{ParameterName}(out float3 OutCenter, out float3 OutExtentsMin, out float3 OutExtentsMax, out float3 OutExtents, out float3 OutHalfExtents)
 {
 	OutCenter		= {ParameterName}_PreSkinnedLocalBoundsCenter;
-	OutExtents		= {ParameterName}_PreSkinnedLocalBoundsExtents;
-	OutHalfExtents	= OutExtents * 0.5f;
+	OutHalfExtents	= {ParameterName}_PreSkinnedLocalBoundsExtents;
+	OutExtents		= OutExtents * 2.0f;
 	OutExtentsMin	= OutCenter - OutHalfExtents;
 	OutExtentsMax	= OutCenter + OutHalfExtents;
 }
diff --git a/Engine/Plugins/FX/Niagara/Source/Niagara/Private/DataInterface/NiagaraDataInterfaceStaticMesh.cpp b/Engine/Plugins/FX/Niagara/Source/Niagara/Private/DataInterface/NiagaraDataInterfaceStaticMesh.cpp
--- a/Engine/Plugins/FX/Niagara/Source/Niagara/Private/DataInterface/NiagaraDataInterfaceStaticMesh.cpp
+++ b/Engine/Plugins/FX/Niagara/Source/Niagara/Private/DataInterface/NiagaraDataInterfaceStaticMesh.cpp
@@ -1850,6 +1850,7 @@ void UNiagaraDataInterfaceStaticMesh::GetVertexSamplingFunctions(TArray<FNiagara
 		Sig.Outputs.Emplace(FNiagaraTypeDefinition::GetVec3Def(), TEXT("Bitangent"));
 		Sig.Outputs.Emplace(FNiagaraTypeDefinition::GetVec3Def(), TEXT("Tangent"));
 		OutFunctions.Add_GetRef(Sig).Name = GetVertexName;
+		OutFunctions.Last().SetOutputDescription(FNiagaraVariableBase(FNiagaraTypeDefinition::GetVec3Def(), TEXT("Velocity")), LOCTEXT("LocalSpaceVelocityDesc", "Local space velocity is always zero as the mesh does not have any local deformation."));
 		
 		FNiagaraFunctionSignature& WsSig = OutFunctions.Add_GetRef(Sig);
 		WsSig.Name = GetVertexWSName;
@@ -1957,6 +1958,7 @@ void UNiagaraDataInterfaceStaticMesh::GetTriangleSamplingFunctions(TArray<FNiaga
 		Sig.Outputs.Emplace(FNiagaraTypeDefinition::GetVec3Def(), TEXT("Bitangent"));
 		Sig.Outputs.Emplace(FNiagaraTypeDefinition::GetVec3Def(), TEXT("Tangent"));
 		OutFunctions.Add_GetRef(Sig).Name = GetTriangleName;
+		OutFunctions.Last().SetOutputDescription(FNiagaraVariableBase(FNiagaraTypeDefinition::GetVec3Def(), TEXT("Velocity")), LOCTEXT("LocalSpaceVelocityDesc", "Local space velocity is always zero as the mesh does not have any local deformation."));
 
 		FNiagaraFunctionSignature& WsSig = OutFunctions.Add_GetRef(Sig);
 		WsSig.Name = GetTriangleWSName;
@@ -2027,6 +2029,8 @@ void UNiagaraDataInterfaceStaticMesh::GetSocketSamplingFunctions(TArray<FNiagara
 		FNiagaraFunctionSignature WsInterpSig = WsSig;
 		WsInterpSig.Inputs.Emplace(FNiagaraTypeDefinition::GetFloatDef(), TEXT("Interp"));
 
+		Sig.SetOutputDescription(FNiagaraVariableBase(FNiagaraTypeDefinition::GetVec3Def(), TEXT("Velocity")), LOCTEXT("LocalSpaceVelocityDesc", "Local space velocity is always zero as the mesh does not have any local deformation."));
+
 		OutFunctions.Add_GetRef(Sig).Name = GetSocketTransformName;
 		OutFunctions.Add_GetRef(WsSig).Name = GetSocketTransformWSName;
 		OutFunctions.Add_GetRef(WsInterpSig).Name = GetSocketTransformWSInterpolatedName;
@@ -2188,14 +2192,17 @@ void UNiagaraDataInterfaceStaticMesh::GetDistanceFieldFunctions(TArray<FNiagaraF
 {
 	using namespace NDIStaticMeshLocal;
 	{
+		const FNiagaraVariable UseMaxDistanceVar(FNiagaraTypeDefinition::GetBoolDef(), TEXT("Use Max Distance"));
+		const FNiagaraVariable MaxDistanceVar(FNiagaraTypeDefinition::GetFloatDef(), TEXT("Max Distance"));
+
 		FNiagaraFunctionSignature& Sig = OutFunctions.Add_GetRef(BaseSignature);
 		Sig.Name = QueryDistanceFieldName;
 		Sig.bSupportsCPU = false;
 		Sig.bExperimental = true;
 		Sig.Inputs.Emplace_GetRef(FNiagaraTypeDefinition::GetBoolDef(), TEXT("Execute")).SetValue(true);
 		Sig.Inputs.Emplace(FNiagaraTypeDefinition::GetPositionDef(), TEXT("World Position"));
-		const FNiagaraVariable& UseMaxDistanceVar = Sig.Inputs.Emplace_GetRef(FNiagaraTypeDefinition::GetBoolDef(), TEXT("Use Max Distance"));
-		const FNiagaraVariable& MaxDistanceVar = Sig.Inputs.Emplace_GetRef(FNiagaraTypeDefinition::GetFloatDef(), TEXT("Max Distance"));
+		Sig.Inputs.Emplace_GetRef(UseMaxDistanceVar);
+		Sig.Inputs.Emplace_GetRef(MaxDistanceVar);
 		Sig.Outputs.Emplace(FNiagaraTypeDefinition::GetBoolDef(), TEXT("IsValid"));
 		Sig.Outputs.Emplace(FNiagaraTypeDefinition::GetFloatDef(), TEXT("Distance"));
 		
@@ -3337,8 +3344,8 @@ void UNiagaraDataInterfaceStaticMesh::VMGetPreSkinnedLocalBounds(FVectorVMExtern
 	FNDIOutputParam<FVector3f>	OutHalfExtents(Context);
 
 	const FVector3f Center = StaticMeshHelper.InstanceData->PreSkinnedLocalBoundsCenter;
-	const FVector3f Extents = StaticMeshHelper.InstanceData->PreSkinnedLocalBoundsExtents;
-	const FVector3f HalfExtents = Extents * 0.5f;
+	const FVector3f HalfExtents = StaticMeshHelper.InstanceData->PreSkinnedLocalBoundsExtents;
+	const FVector3f Extents = HalfExtents * 2.0f;
 	const FVector3f ExtentsMin = Center - HalfExtents;
 	const FVector3f ExtentsMax = Center + HalfExtents;
 
diff --git a/Engine/Plugins/FX/Niagara/Source/Niagara/Public/NiagaraCommon.h b/Engine/Plugins/FX/Niagara/Source/Niagara/Public/NiagaraCommon.h
--- a/Engine/Plugins/FX/Niagara/Source/Niagara/Public/NiagaraCommon.h
+++ b/Engine/Plugins/FX/Niagara/Source/Niagara/Public/NiagaraCommon.h
@@ -544,6 +544,20 @@ struct NIAGARA_API FNiagaraFunctionSignature
 	#endif
 	}
 
+	void SetInputDescription(const FNiagaraVariableBase& Variable, const FText& Desc)
+	{
+#if WITH_EDITORONLY_DATA
+		InputDescriptions.FindOrAdd(Variable) = Desc;
+#endif
+	}
+
+	void SetOutputDescription(const FNiagaraVariableBase& Variable, const FText& Desc)
+	{
+#if WITH_EDITORONLY_DATA
+		OutputDescriptions.FindOrAdd(Variable) = Desc;
+#endif
+	}
+
 	void SetDescription(const FText& Desc)
 	{
 	#if WITH_EDITORONLY_DATA
