Fixing issue with cosmetics not binding to skeletal meshes correctly.

Recent change makes some parts of the flow to find the right skel mesh also depend on the component having a valid mesh to use.
In some cases we don't want to continue looking for a new component if we've found a valid one but it just dosn't have a valid mesh yet.

diff --git a/Engine/Plugins/FX/Niagara/Source/Niagara/Private/NiagaraDataInterfaceSkeletalMesh.cpp b/Engine/Plugins/FX/Niagara/Source/Niagara/Private/NiagaraDataInterfaceSkeletalMesh.cpp
--- a/Engine/Plugins/FX/Niagara/Source/Niagara/Private/NiagaraDataInterfaceSkeletalMesh.cpp
+++ b/Engine/Plugins/FX/Niagara/Source/Niagara/Private/NiagaraDataInterfaceSkeletalMesh.cpp
@@ -1416,9 +1416,9 @@ void UNiagaraDataInterfaceSkeletalMesh::ProvidePerInstanceDataForRenderThread(vo
 
 USkeletalMesh* UNiagaraDataInterfaceSkeletalMesh::GetSkeletalMesh(FNiagaraSystemInstance* SystemInstance, USceneComponent* AttachComponent, TWeakObjectPtr<USceneComponent>& SceneComponent, USkeletalMeshComponent*& FoundSkelComp, FNDISkeletalMesh_InstanceData* InstData)
 {
-	auto IsValidComponent = [&](const USkeletalMeshComponent* Component)
+	auto IsValidComponent = [&](const USkeletalMeshComponent* Component, bool bAllowNullMesh=false)
 	{
-		return IsValid(Component) && IsValid(Component->GetSkeletalMeshAsset()) &&
+		return IsValid(Component) && (bAllowNullMesh || IsValid(Component->GetSkeletalMeshAsset())) &&
 			(ComponentTags.IsEmpty()
 			|| ComponentTags.ContainsByPredicate([&](const FName& Tag) { return Tag == NAME_None || Component->ComponentHasTag(Tag); }));
 	};
@@ -1429,7 +1429,7 @@ USkeletalMesh* UNiagaraDataInterfaceSkeletalMesh::GetSkeletalMesh(FNiagaraSystem
 		if (ASkeletalMeshActor* SkelMeshActor = Cast<ASkeletalMeshActor>(Actor))
 		{
 			USkeletalMeshComponent* Comp = SkelMeshActor->GetSkeletalMeshComponent();
-			if (IsValidComponent(Comp))
+			if (IsValidComponent(Comp, true))
 			{
 				return Comp;
 			}
@@ -1519,7 +1519,7 @@ USkeletalMesh* UNiagaraDataInterfaceSkeletalMesh::GetSkeletalMesh(FNiagaraSystem
 		for (USceneComponent* Curr = AttachComponent; Curr; Curr = Curr->GetAttachParent())
 		{
 			USkeletalMeshComponent* ParentComp = Cast<USkeletalMeshComponent>(Curr);
-			if (IsValidComponent(ParentComp))
+			if (IsValidComponent(ParentComp, true))
 			{
 				FoundSkelComp = ParentComp;
 				break;
@@ -1530,7 +1530,7 @@ USkeletalMesh* UNiagaraDataInterfaceSkeletalMesh::GetSkeletalMesh(FNiagaraSystem
 		{
 			// Next, try to find one in our outer chain
 			USkeletalMeshComponent* OuterComp = AttachComponent->GetTypedOuter<USkeletalMeshComponent>();
-			if (IsValidComponent(OuterComp))
+			if (IsValidComponent(OuterComp, true))
 			{
 				FoundSkelComp = OuterComp;
 			}
