collide against complex collision shapes that have no simple shape

diff --git a/Engine/Plugins/FX/Niagara/Source/Niagara/Private/NiagaraDataInterfaceRigidMeshCollisionQuery.cpp b/Engine/Plugins/FX/Niagara/Source/Niagara/Private/NiagaraDataInterfaceRigidMeshCollisionQuery.cpp
--- a/Engine/Plugins/FX/Niagara/Source/Niagara/Private/NiagaraDataInterfaceRigidMeshCollisionQuery.cpp
+++ b/Engine/Plugins/FX/Niagara/Source/Niagara/Private/NiagaraDataInterfaceRigidMeshCollisionQuery.cpp
@@ -298,6 +298,15 @@ void CountCollisionPrimitives(TConstArrayView<TComponentType*> Components, uint3
 					++CapsuleCount;
 				}
 			}
+
+			// if we have no supported bodies associated with a mesh, but we want to collide with it, add a single box
+			if (BodySetup->AggGeom.ConvexElems.Num() == 0 &&
+				BodySetup->AggGeom.BoxElems.Num() == 0 &&
+				BodySetup->AggGeom.SphereElems.Num() == 0 &&
+				BodySetup->AggGeom.SphylElems.Num() == 0)
+			{
+				++BoxCount;
+			}
 		});
 
 		// suppress warning since we get it a lot
@@ -418,6 +427,34 @@ void UpdateAssetArrays(TConstArrayView<TComponentType*> Components, const FVecto
 				++CapsuleIndex;
 			}
 		}
+
+		// if we have no supported bodies associated with a mesh, but we want to collide with it, add a single box
+		if (BodySetup->AggGeom.ConvexElems.Num() == 0 &&
+			BodySetup->AggGeom.BoxElems.Num() == 0 &&
+			BodySetup->AggGeom.SphereElems.Num() == 0 &&
+			BodySetup->AggGeom.SphylElems.Num() == 0)
+		{
+			FVector Extent;
+			FVector Center;
+		
+			FBoxSphereBounds Bounds = static_cast< USceneComponent* >(Component)->GetLocalBounds();
+			Extent = Bounds.BoxExtent;
+			Center = Bounds.Origin;
+			
+			if (InitializeStatics)
+			{			
+				// local bounds extent is half the world extents of the bounding box in local space
+				Extent *= 2;
+
+				OutAssetArrays->ElementExtent[BoxIndex] = FVector4f(Extent.X, Extent.Y, Extent.Z, 0);
+				OutAssetArrays->PhysicsType[BoxIndex] = true;
+				OutAssetArrays->ComponentIdIndex[BoxIndex] = ComponentIdIndex;
+			}
+			
+			const FTransform ElementTransform = FTransform(Center) * MeshTransform;
+			FillCurrentTransforms(ElementTransform, BoxIndex, OutAssetArrays->CurrentTransform, OutAssetArrays->CurrentInverse);
+			++BoxIndex;
+		}
 	};
 
 	for (TComponentType* Component : Components)
@@ -977,6 +1014,7 @@ bool UNiagaraDataInterfaceRigidMeshCollisionQuery::CopyToInternal(UNiagaraDataIn
 	OtherTyped->ComponentTags = ComponentTags;
 	OtherTyped->SourceActors = SourceActors;
 	OtherTyped->OnlyUseMoveable = OnlyUseMoveable;
+	OtherTyped->UseComplexCollisions = UseComplexCollisions;
 	OtherTyped->GlobalSearchAllowed = GlobalSearchAllowed;
 	OtherTyped->GlobalSearchForced = GlobalSearchForced;
 	OtherTyped->GlobalSearchFallback_Unscripted = GlobalSearchFallback_Unscripted;
@@ -997,6 +1035,7 @@ bool UNiagaraDataInterfaceRigidMeshCollisionQuery::Equals(const UNiagaraDataInte
 		&& (OtherTyped->ComponentTags == ComponentTags)
 		&& (OtherTyped->SourceActors == SourceActors)
 		&& (OtherTyped->OnlyUseMoveable == OnlyUseMoveable)
+		&& (OtherTyped->UseComplexCollisions == UseComplexCollisions)
 		&& (OtherTyped->GlobalSearchAllowed == GlobalSearchAllowed)
 		&& (OtherTyped->GlobalSearchForced == GlobalSearchForced)
 		&& (OtherTyped->GlobalSearchFallback_Unscripted == GlobalSearchFallback_Unscripted)
@@ -1571,7 +1610,7 @@ bool UNiagaraDataInterfaceRigidMeshCollisionQuery::FindActors(UWorld* World, FND
 		ObjectParams.AddObjectTypesToQuery(Channel);
 
 		TArray<FOverlapResult> Overlaps;
-		FCollisionQueryParams Params(SCENE_QUERY_STAT(NiagaraRigidMeshCollisionQuery), false);
+		FCollisionQueryParams Params(SCENE_QUERY_STAT(NiagaraRigidMeshCollisionQuery), UseComplexCollisions);
 
 		World->OverlapMultiByChannel(Overlaps, OverlapLocation, OverlapRotation, Channel, FCollisionShape::MakeBox(0.5f * OverlapExtent), Params);
 
diff --git a/Engine/Plugins/FX/Niagara/Source/Niagara/Public/NiagaraDataInterfaceRigidMeshCollisionQuery.h b/Engine/Plugins/FX/Niagara/Source/Niagara/Public/NiagaraDataInterfaceRigidMeshCollisionQuery.h
--- a/Engine/Plugins/FX/Niagara/Source/Niagara/Public/NiagaraDataInterfaceRigidMeshCollisionQuery.h
+++ b/Engine/Plugins/FX/Niagara/Source/Niagara/Public/NiagaraDataInterfaceRigidMeshCollisionQuery.h
@@ -185,6 +185,10 @@ public:
 	UPROPERTY(EditAnywhere, Category = "Source")
 	bool OnlyUseMoveable = true;
 
+	/** If enabled, complex collisions will be searched for. */
+	UPROPERTY(EditAnywhere, Category = "Source")
+	bool UseComplexCollisions = false;
+
 	/** If enabled the global search can be executed dependeing on GlobalSearchForced and GlobalSearchFallback_Unscripted */
 	UPROPERTY(EditAnywhere, Category = "Source", meta = (DisplayName = "Global Search Allowed"))
 	bool GlobalSearchAllowed = true;
