Removed the inline macro in NiagaraSettings because it has issues when compiling the whole module as one unity file. It is unlikely to ever happen in practice because it is many times slower to compile but this fixes a buildfarm job. I will come back to this after other work is done.

diff --git a/Engine/Plugins/FX/Niagara/Source/Niagara/Private/NiagaraSettings.cpp b/Engine/Plugins/FX/Niagara/Source/Niagara/Private/NiagaraSettings.cpp
--- a/Engine/Plugins/FX/Niagara/Source/Niagara/Private/NiagaraSettings.cpp
+++ b/Engine/Plugins/FX/Niagara/Source/Niagara/Private/NiagaraSettings.cpp
@@ -2,8 +2,6 @@
 
 #include "NiagaraSettings.h"
 
-#include UE_INLINE_GENERATED_CPP_BY_NAME(NiagaraSettings)
-
 UNiagaraSettings::UNiagaraSettings(const FObjectInitializer& ObjectInitlaizer)
 	: Super(ObjectInitlaizer)
 	, NDISkelMesh_GpuMaxInfluences(ENDISkelMesh_GpuMaxInfluences::Unlimited)
