When saving from the NiagaraSystemToolkit we make sure that we're saving the version not in edit mode.  This prevents recompilation when someone references the system (in a level, or a cook, ...)

diff --git a/Engine/Plugins/FX/Niagara/Source/NiagaraEditor/Private/Toolkits/NiagaraSystemToolkit.cpp b/Engine/Plugins/FX/Niagara/Source/NiagaraEditor/Private/Toolkits/NiagaraSystemToolkit.cpp
--- a/Engine/Plugins/FX/Niagara/Source/NiagaraEditor/Private/Toolkits/NiagaraSystemToolkit.cpp
+++ b/Engine/Plugins/FX/Niagara/Source/NiagaraEditor/Private/Toolkits/NiagaraSystemToolkit.cpp
@@ -1210,6 +1210,7 @@ void FNiagaraSystemToolkit::SaveAsset_Execute()
 	}
 	SystemViewModel->NotifyPreSave();
 	FAssetEditorToolkit::SaveAsset_Execute();
+	SystemViewModel->NotifyPostSave();
 }
 
 void FNiagaraSystemToolkit::SaveAssetAs_Execute()
@@ -1221,6 +1222,7 @@ void FNiagaraSystemToolkit::SaveAssetAs_Execute()
 	}
 	SystemViewModel->NotifyPreSave();
 	FAssetEditorToolkit::SaveAssetAs_Execute();
+	SystemViewModel->NotifyPostSave();
 }
 
 bool FNiagaraSystemToolkit::OnRequestClose()
diff --git a/Engine/Plugins/FX/Niagara/Source/NiagaraEditor/Private/ViewModels/NiagaraSystemViewModel.cpp b/Engine/Plugins/FX/Niagara/Source/NiagaraEditor/Private/ViewModels/NiagaraSystemViewModel.cpp
--- a/Engine/Plugins/FX/Niagara/Source/NiagaraEditor/Private/ViewModels/NiagaraSystemViewModel.cpp
+++ b/Engine/Plugins/FX/Niagara/Source/NiagaraEditor/Private/ViewModels/NiagaraSystemViewModel.cpp
@@ -84,6 +84,7 @@ FNiagaraSystemViewModel::FNiagaraSystemViewModel()
 	, NiagaraSequence(nullptr)
 	, bSettingSequencerTimeDirectly(false)
 	, bForceAutoCompileOnce(false)
+	, bSupportCompileForEdit(true)
 	, bUpdatingEmittersFromSequencerDataChange(false)
 	, bUpdatingSequencerFromEmitterDataChange(false)
 	, bUpdatingSystemSelectionFromSequencer(false)
@@ -103,7 +104,7 @@ FNiagaraSystemViewModel::FNiagaraSystemViewModel()
 void FNiagaraSystemViewModel::Initialize(UNiagaraSystem& InSystem, FNiagaraSystemViewModelOptions InOptions)
 {
 	System = &InSystem;
-	System->bCompileForEdit = true;
+	System->bCompileForEdit = bSupportCompileForEdit;
 
 	RegisteredHandle = RegisterViewModelWithMap(System, this);
 
@@ -830,11 +831,29 @@ void FNiagaraSystemViewModel::Tick(float DeltaTime)
 
 void FNiagaraSystemViewModel::NotifyPreSave()
 {
-	if (GetSystem().HasOutstandingCompilationRequests(true))
+	if (System->HasOutstandingCompilationRequests(true))
 	{
-		UE_LOG(LogNiagaraEditor, Log, TEXT("System %s has pending compile jobs. Waiting for that code to complete before Saving.."), *GetSystem().GetName());
+		UE_LOG(LogNiagaraEditor, Log, TEXT("System %s has pending compile jobs. Waiting for that code to complete before Saving.."), *System->GetName());
+		System->WaitForCompilationComplete(true);
+	}
 
-		GetSystem().WaitForCompilationComplete(true);
+	if (bSupportCompileForEdit)
+	{
+		check(System->bCompileForEdit);
+		System->bCompileForEdit = false;
+		System->RequestCompile(false);
+		System->WaitForCompilationComplete(true);
+	}
+}
+
+void FNiagaraSystemViewModel::NotifyPostSave()
+{
+	if (bSupportCompileForEdit)
+	{
+		check(!System->bCompileForEdit);
+		System->bCompileForEdit = true;
+		System->RequestCompile(false);
+		System->WaitForCompilationComplete(true);
 	}
 }
 
diff --git a/Engine/Plugins/FX/Niagara/Source/NiagaraEditor/Public/ViewModels/NiagaraSystemViewModel.h b/Engine/Plugins/FX/Niagara/Source/NiagaraEditor/Public/ViewModels/NiagaraSystemViewModel.h
--- a/Engine/Plugins/FX/Niagara/Source/NiagaraEditor/Public/ViewModels/NiagaraSystemViewModel.h
+++ b/Engine/Plugins/FX/Niagara/Source/NiagaraEditor/Public/ViewModels/NiagaraSystemViewModel.h
@@ -343,6 +343,7 @@ public:
 	NIAGARAEDITOR_API bool GetIsEmitterPinned(TSharedRef<FNiagaraEmitterHandleViewModel> EmitterHandleModel);
 
 	NIAGARAEDITOR_API void NotifyPreSave();
+	NIAGARAEDITOR_API void NotifyPostSave();
 
 	NIAGARAEDITOR_API void NotifyPreClose();
 
@@ -620,6 +621,9 @@ private:
 	/** Whether or not the system represented by this view model can be simulated. */
 	bool bCanSimulate;
 
+	/** Whether we should use the CompileForEdit mode for the assigned system. */
+	bool bSupportCompileForEdit;
+
 	/** The current editing mode for this view model. */
 	ENiagaraSystemViewModelEditMode EditMode;
 
