- Fix emitter mode vis tag and mesh index not working on mesh renderers

diff --git a/Engine/Plugins/FX/Niagara/Source/Niagara/Private/NiagaraRendererMeshes.cpp b/Engine/Plugins/FX/Niagara/Source/Niagara/Private/NiagaraRendererMeshes.cpp
--- a/Engine/Plugins/FX/Niagara/Source/Niagara/Private/NiagaraRendererMeshes.cpp
+++ b/Engine/Plugins/FX/Niagara/Source/Niagara/Private/NiagaraRendererMeshes.cpp
@@ -106,20 +106,29 @@ FNiagaraRendererMeshes::FNiagaraRendererMeshes(ERHIFeatureLevel::Type FeatureLev
 
 	const FNiagaraDataSet& Data = Emitter->GetData();
 
-	RendererVisTagOffset = INDEX_NONE;
 	int32 FloatOffset;
 	int32 HalfOffset;
-	if (Data.GetVariableComponentOffsets(Properties->RendererVisibilityTagBinding.GetDataSetBindableVariable(), FloatOffset, RendererVisTagOffset, HalfOffset))
+
+	ParticleRendererVisTagOffset = INDEX_NONE;
+	EmitterRendererVisTagOffset = INDEX_NONE;
+	if (Properties->RendererVisibilityTagBinding.CanBindToHostParameterMap())
 	{
-		// If the renderer visibility tag is bound, we have to do it in the culling pass
-		bEnableCulling = true;
+		EmitterRendererVisTagOffset = Emitter->GetRendererBoundVariables().IndexOf(Properties->RendererVisibilityTagBinding.GetParamMapBindableVariable());
+	}
+	else if (Data.GetVariableComponentOffsets(Properties->RendererVisibilityTagBinding.GetDataSetBindableVariable(), FloatOffset, ParticleRendererVisTagOffset, HalfOffset))
+	{
+		bEnableCulling |= ParticleRendererVisTagOffset != INDEX_NONE;
 	}
 
-	MeshIndexOffset = INDEX_NONE;
-	if (Data.GetVariableComponentOffsets(Properties->MeshIndexBinding.GetDataSetBindableVariable(), FloatOffset, MeshIndexOffset, HalfOffset))
+	ParticleMeshIndexOffset = INDEX_NONE;
+	EmitterMeshIndexOffset = INDEX_NONE;
+	if (Properties->MeshIndexBinding.CanBindToHostParameterMap())
 	{
-		// If the mesh index is bound, we have to do it in the culling pass
-		bEnableCulling = true;
+		EmitterMeshIndexOffset = Emitter->GetRendererBoundVariables().IndexOf(Properties->MeshIndexBinding.GetParamMapBindableVariable());
+	}
+	else if (Data.GetVariableComponentOffsets(Properties->MeshIndexBinding.GetDataSetBindableVariable(), FloatOffset, EmitterMeshIndexOffset, HalfOffset))
+	{
+		bEnableCulling = ParticleMeshIndexOffset != INDEX_NONE;
 	}
 
 	MaterialParamValidMask = Properties->MaterialParamValidMask;
@@ -308,6 +317,18 @@ void FNiagaraRendererMeshes::PrepareParticleMeshRenderData(FParticleMeshRenderDa
 		return;
 	}
 
+	// If the visibility tag comes from a parameter map, so we can evaluate it here and just early out if it doesn't match up
+	if (EmitterRendererVisTagOffset != INDEX_NONE && ParticleMeshRenderData.DynamicDataMesh->ParameterDataBound.IsValidIndex(EmitterRendererVisTagOffset))
+	{
+		int32 VisTag = 0;
+		FMemory::Memcpy(&VisTag, ParticleMeshRenderData.DynamicDataMesh->ParameterDataBound.GetData() + EmitterRendererVisTagOffset, sizeof(int32));
+		if (RendererVisibility != VisTag)
+		{
+			ParticleMeshRenderData.SourceParticleData = nullptr;
+			return;
+		}
+	}
+
 	// Particle source mode
 	if (SourceMode == ENiagaraRendererSourceDataMode::Particles)
 	{
@@ -387,19 +408,19 @@ void FNiagaraRendererMeshes::PrepareParticleRenderBuffers(FParticleMeshRenderDat
 			{
 				if (ParticleMeshRenderData.bSortCullOnGpu || ParticleMeshRenderData.bUseGPUScene)
 				{
-					if (RendererVisTagOffset != INDEX_NONE)
+					if (ParticleRendererVisTagOffset != INDEX_NONE)
 					{
-						ParticleMeshRenderData.RendererVisTagOffset = IntParamsToCopy.Add(RendererVisTagOffset);
+						ParticleMeshRenderData.RendererVisTagOffset = IntParamsToCopy.Add(ParticleRendererVisTagOffset);
 					}
-					if (MeshIndexOffset != INDEX_NONE)
+					if (ParticleMeshIndexOffset != INDEX_NONE)
 					{
-						ParticleMeshRenderData.MeshIndexOffset = IntParamsToCopy.Add(MeshIndexOffset);
+						ParticleMeshRenderData.MeshIndexOffset = IntParamsToCopy.Add(ParticleMeshIndexOffset);
 					}
 				}
 				else
 				{
-					ParticleMeshRenderData.RendererVisTagOffset = RendererVisTagOffset;
-					ParticleMeshRenderData.MeshIndexOffset = MeshIndexOffset;
+					ParticleMeshRenderData.RendererVisTagOffset = ParticleRendererVisTagOffset;
+					ParticleMeshRenderData.MeshIndexOffset = ParticleMeshIndexOffset;
 				}
 			}
 
@@ -422,8 +443,8 @@ void FNiagaraRendererMeshes::PrepareParticleRenderBuffers(FParticleMeshRenderDat
 			ParticleMeshRenderData.ParticleHalfDataStride = ParticleMeshRenderData.SourceParticleData->GetHalfStride() / sizeof(FFloat16);
 			ParticleMeshRenderData.ParticleIntDataStride = ParticleMeshRenderData.SourceParticleData->GetInt32Stride() / sizeof(int32);
 
-			ParticleMeshRenderData.RendererVisTagOffset = RendererVisTagOffset;
-			ParticleMeshRenderData.MeshIndexOffset = MeshIndexOffset;
+			ParticleMeshRenderData.RendererVisTagOffset = ParticleRendererVisTagOffset;
+			ParticleMeshRenderData.MeshIndexOffset = ParticleMeshIndexOffset;
 		}
 	}
 	else
@@ -1171,6 +1192,13 @@ void FNiagaraRendererMeshes::GetDynamicMeshElements(const TArray<const FSceneVie
 
 	PrepareParticleRenderBuffers(ParticleMeshRenderData, Collector.GetDynamicReadBuffer());
 
+	// If mesh index comes from the parameter store grab the information now
+	int32 EmitterModeMeshIndex = INDEX_NONE;
+	if (EmitterMeshIndexOffset != INDEX_NONE && ParticleMeshRenderData.DynamicDataMesh->ParameterDataBound.IsValidIndex(EmitterMeshIndexOffset))
+	{
+		FMemory::Memcpy(&EmitterModeMeshIndex, ParticleMeshRenderData.DynamicDataMesh->ParameterDataBound.GetData() + EmitterMeshIndexOffset, sizeof(int32));
+	}
+
 	// Generate mesh batches per view
 	const int32 NumViews = Views.Num();
 	for (int32 ViewIndex = 0; ViewIndex < NumViews; ViewIndex++)
@@ -1216,6 +1244,12 @@ void FNiagaraRendererMeshes::GetDynamicMeshElements(const TArray<const FSceneVie
 					break;
 				}
 
+				// Emitter mode mesh index binding
+				if (EmitterModeMeshIndex != INDEX_NONE && MeshIndex != EmitterModeMeshIndex)
+				{
+					continue;
+				}
+
 				const FMeshData& MeshData = Meshes[MeshIndex];
 
 				// @TODO : support multiple LOD
@@ -1368,7 +1402,7 @@ void FNiagaraRendererMeshes::GetDynamicRayTracingInstances(FRayTracingMaterialGa
 	for (int32 MeshIndex = 0; MeshIndex < Meshes.Num(); ++MeshIndex)
 	{
 		// No binding for mesh index we only render the first mesh not all of them
-		if ( (MeshIndex > 0) && (MeshIndexOffset == INDEX_NONE) )
+		if ( (MeshIndex > 0) && (ParticleMeshIndexOffset == INDEX_NONE) )
 		{
 			break;
 		}
@@ -1586,8 +1620,8 @@ void FNiagaraRendererMeshes::GetDynamicRayTracingInstances(FRayTracingMaterialGa
 				const float* RESTRICT QuatArrayZ = reinterpret_cast<const float*>(ParticleMeshRenderData.SourceParticleData->GetComponentPtrFloat(TransformBaseCompOffset + 2));
 				const float* RESTRICT QuatArrayW = reinterpret_cast<const float*>(ParticleMeshRenderData.SourceParticleData->GetComponentPtrFloat(TransformBaseCompOffset + 3));
 
-				const int32* RESTRICT RenderVisibilityData = RendererVisTagOffset == INDEX_NONE ? nullptr : reinterpret_cast<const int32*>(ParticleMeshRenderData.SourceParticleData->GetComponentPtrInt32(RendererVisTagOffset));
-				const int32* RESTRICT MeshIndexData = MeshIndexOffset == INDEX_NONE ? nullptr : reinterpret_cast<const int32*>(ParticleMeshRenderData.SourceParticleData->GetComponentPtrInt32(MeshIndexOffset));
+				const int32* RESTRICT RenderVisibilityData = ParticleRendererVisTagOffset == INDEX_NONE ? nullptr : reinterpret_cast<const int32*>(ParticleMeshRenderData.SourceParticleData->GetComponentPtrInt32(ParticleRendererVisTagOffset));
+				const int32* RESTRICT MeshIndexData = ParticleMeshIndexOffset == INDEX_NONE ? nullptr : reinterpret_cast<const int32*>(ParticleMeshRenderData.SourceParticleData->GetComponentPtrInt32(ParticleMeshIndexOffset));
 
 				auto GetInstancePosition = [&PositionX, &PositionY, &PositionZ](int32 Idx)
 				{
@@ -1674,8 +1708,8 @@ void FNiagaraRendererMeshes::GetDynamicRayTracingInstances(FRayTracingMaterialGa
 					PassParameters->RotationDataOffset			= VFVariables[ENiagaraMeshVFLayout::Rotation].GetGPUOffset();
 					PassParameters->ScaleDataOffset				= VFVariables[ENiagaraMeshVFLayout::Scale].GetGPUOffset();
 					PassParameters->bLocalSpace					= bUseLocalSpace ? 1 : 0;
-					PassParameters->RenderVisibilityOffset		= RendererVisTagOffset;
-					PassParameters->MeshIndexOffset				= MeshIndexOffset;
+					PassParameters->RenderVisibilityOffset		= ParticleRendererVisTagOffset;
+					PassParameters->MeshIndexOffset				= ParticleMeshIndexOffset;
 					PassParameters->RenderVisibilityValue		= RendererVisibility;
 					PassParameters->MeshIndexValue				= MeshData.SourceMeshIndex;
 					PassParameters->LocalTransform				= FMatrix44f(LocalTransform);		// LWC_TODO: Precision loss
diff --git a/Engine/Plugins/FX/Niagara/Source/Niagara/Public/NiagaraRendererMeshes.h b/Engine/Plugins/FX/Niagara/Source/Niagara/Public/NiagaraRendererMeshes.h
--- a/Engine/Plugins/FX/Niagara/Source/Niagara/Public/NiagaraRendererMeshes.h
+++ b/Engine/Plugins/FX/Niagara/Source/Niagara/Public/NiagaraRendererMeshes.h
@@ -195,9 +195,11 @@ private:
 
 	FVector2f DistanceCullRange;
 	FVector2f DistanceCullRangeSquared;
-	int32 RendererVisTagOffset;
-	int32 RendererVisibility;
-	int32 MeshIndexOffset;
+	int32 ParticleRendererVisTagOffset = INDEX_NONE;
+	int32 ParticleMeshIndexOffset = INDEX_NONE;
+	int32 RendererVisibility = 0;
+	int32 EmitterRendererVisTagOffset = INDEX_NONE;
+	int32 EmitterMeshIndexOffset = INDEX_NONE;
 	uint32 MaterialParamValidMask;
 	uint32 MaxSectionCount;
 
