Fix for debug output for niagara compoennt bounds with NaNs

diff --git a/Engine/Plugins/FX/Niagara/Source/Niagara/Private/NiagaraComponent.cpp b/Engine/Plugins/FX/Niagara/Source/Niagara/Private/NiagaraComponent.cpp
--- a/Engine/Plugins/FX/Niagara/Source/Niagara/Private/NiagaraComponent.cpp
+++ b/Engine/Plugins/FX/Niagara/Source/Niagara/Private/NiagaraComponent.cpp
@@ -1976,7 +1976,7 @@ FBoxSphereBounds UNiagaraComponent::CalcBounds(const FTransform& LocalToWorld) c
 	FBoxSphereBounds SystemBounds;
 	if (CurrLocalBounds.IsValid)
 	{
-		ensureMsgf(!CurrLocalBounds.Min.ContainsNaN() && !CurrLocalBounds.Max.ContainsNaN(), TEXT("CurrLocalBounds contains NaN for System(%d)"), *GetNameSafe(Asset));
+		ensureMsgf(!CurrLocalBounds.Min.ContainsNaN() && !CurrLocalBounds.Max.ContainsNaN(), TEXT("CurrLocalBounds contains NaN for System(%s)"), *GetNameSafe(Asset));
 
 		SystemBounds = CurrLocalBounds;
 		SystemBounds.BoxExtent *= BoundsScale;
@@ -1986,13 +1986,13 @@ FBoxSphereBounds UNiagaraComponent::CalcBounds(const FTransform& LocalToWorld) c
 	{
 		if (SystemFixedBounds.IsValid)
 		{
-			ensureMsgf(!SystemFixedBounds.Min.ContainsNaN() && !SystemFixedBounds.Max.ContainsNaN(), TEXT("SystemFixedBounds contains NaN for System(%d)"), *GetNameSafe(Asset));
+			ensureMsgf(!SystemFixedBounds.Min.ContainsNaN() && !SystemFixedBounds.Max.ContainsNaN(), TEXT("SystemFixedBounds contains NaN for System(%s)"), *GetNameSafe(Asset));
 			SystemBounds = FBoxSphereBounds(SystemFixedBounds);
 		}
 		else if ( Asset && Asset->bFixedBounds )
 		{
 			const FBox AssetFixedBounds = Asset->GetFixedBounds();
-			ensureMsgf(!AssetFixedBounds.Min.ContainsNaN() && !AssetFixedBounds.Max.ContainsNaN(), TEXT("AssetFixedBounds contains NaN for System(%d)"), *GetNameSafe(Asset));
+			ensureMsgf(!AssetFixedBounds.Min.ContainsNaN() && !AssetFixedBounds.Max.ContainsNaN(), TEXT("AssetFixedBounds contains NaN for System(%s)"), *GetNameSafe(Asset));
 			SystemBounds = FBoxSphereBounds(AssetFixedBounds);
 		}
 		else
