Don't attempt to preload some niagara assets found in the asset registry if the plugin they live in hasn't been mounted

diff --git a/Engine/Plugins/FX/Niagara/Source/NiagaraEditor/Public/NiagaraEditorModule.h b/Engine/Plugins/FX/Niagara/Source/NiagaraEditor/Public/NiagaraEditorModule.h
--- a/Engine/Plugins/FX/Niagara/Source/NiagaraEditor/Public/NiagaraEditorModule.h
+++ b/Engine/Plugins/FX/Niagara/Source/NiagaraEditor/Public/NiagaraEditorModule.h
@@ -253,7 +253,7 @@ private:
 			CachedAssets.Reset(AssetData.Num());
 			for (const FAssetData& AssetDatum : AssetData)
 			{
-				if (AssetDatum.IsAssetLoaded() || bAllowLoading)
+				if (AssetDatum.IsAssetLoaded() || (bAllowLoading && FPackageName::GetPackageMountPoint(AssetDatum.PackageName.ToString()) != NAME_None))
 				{
 					if (AssetType* Asset = Cast<AssetType>(AssetDatum.GetAsset()))
 					{
