- Fix up UsesEmitter to handle unversioned emitters

#rnx

diff --git a/Engine/Plugins/FX/Niagara/Source/Niagara/Private/NiagaraEmitterHandle.cpp b/Engine/Plugins/FX/Niagara/Source/Niagara/Private/NiagaraEmitterHandle.cpp
--- a/Engine/Plugins/FX/Niagara/Source/Niagara/Private/NiagaraEmitterHandle.cpp
+++ b/Engine/Plugins/FX/Niagara/Source/Niagara/Private/NiagaraEmitterHandle.cpp
@@ -235,7 +235,18 @@ void FNiagaraEmitterHandle::ConditionalPostLoad(int32 NiagaraCustomVersion)
 
 bool FNiagaraEmitterHandle::UsesEmitter(const FVersionedNiagaraEmitter& InEmitter) const
 {
-	return VersionedInstance == InEmitter || (VersionedInstance.GetEmitterData() && InEmitter.Emitter && VersionedInstance.GetEmitterData()->UsesEmitter(*InEmitter.Emitter));
+	if (ensure(InEmitter.Emitter))
+	{
+		if (InEmitter.Emitter->IsVersioningEnabled())
+		{
+			return VersionedInstance == InEmitter || VersionedInstance.GetEmitterData() && VersionedInstance.GetEmitterData()->UsesEmitter(*InEmitter.Emitter);;
+		}
+		else
+		{
+			return FNiagaraEmitterHandle::UsesEmitter(*InEmitter.Emitter);
+		}
+	}
+	return false;
 }
 
 bool FNiagaraEmitterHandle::UsesEmitter(const UNiagaraEmitter& InEmitter) const
