- Fix CIS issue

#rnx

diff --git a/Engine/Plugins/FX/Niagara/Source/Niagara/Private/NiagaraEmitterHandle.cpp b/Engine/Plugins/FX/Niagara/Source/Niagara/Private/NiagaraEmitterHandle.cpp
--- a/Engine/Plugins/FX/Niagara/Source/Niagara/Private/NiagaraEmitterHandle.cpp
+++ b/Engine/Plugins/FX/Niagara/Source/Niagara/Private/NiagaraEmitterHandle.cpp
@@ -239,7 +239,7 @@ bool FNiagaraEmitterHandle::UsesEmitter(const FVersionedNiagaraEmitter& InEmitte
 	{
 		if (InEmitter.Emitter->IsVersioningEnabled())
 		{
-			return VersionedInstance == InEmitter || VersionedInstance.GetEmitterData() && VersionedInstance.GetEmitterData()->UsesEmitter(*InEmitter.Emitter);;
+			return VersionedInstance == InEmitter || (VersionedInstance.GetEmitterData() && VersionedInstance.GetEmitterData()->UsesEmitter(*InEmitter.Emitter));
 		}
 		else
 		{
