Change SimpleLight's color from double to single precision.

#rnx

diff --git a/Engine/Plugins/FX/Niagara/Source/Niagara/Private/NiagaraRendererLights.cpp b/Engine/Plugins/FX/Niagara/Source/Niagara/Private/NiagaraRendererLights.cpp
--- a/Engine/Plugins/FX/Niagara/Source/Niagara/Private/NiagaraRendererLights.cpp
+++ b/Engine/Plugins/FX/Niagara/Source/Niagara/Private/NiagaraRendererLights.cpp
@@ -126,7 +126,7 @@ FNiagaraDynamicDataBase* FNiagaraRendererLights::GenerateDynamicData(const FNiag
 			const float Brightness = Properties->bAlphaScalesBrightness ? Color.A : 1.0f;
 
 			LightData.LightEntry.Radius = LightRadius;
-			LightData.LightEntry.Color = FVector(Color) * Brightness + Properties->ColorAdd;
+			LightData.LightEntry.Color = FVector3f(Color) * Brightness + Properties->ColorAdd;
 			LightData.LightEntry.Exponent = Properties->bUseInverseSquaredFalloff ? 0 : ExponentReader.GetSafe(ParticleIndex, DefaultExponent);
 			LightData.LightEntry.bAffectTranslucency = Properties->bAffectsTranslucency;
 			LightData.LightEntry.VolumetricScatteringIntensity = ScatteringReader.GetSafe(ParticleIndex, DefaultScattering);
diff --git a/Engine/Plugins/FX/Niagara/Source/Niagara/Public/NiagaraLightRendererProperties.h b/Engine/Plugins/FX/Niagara/Source/Niagara/Public/NiagaraLightRendererProperties.h
--- a/Engine/Plugins/FX/Niagara/Source/Niagara/Public/NiagaraLightRendererProperties.h
+++ b/Engine/Plugins/FX/Niagara/Source/Niagara/Public/NiagaraLightRendererProperties.h
@@ -66,7 +66,7 @@ public:
 
 	/** A static color shift applied to each rendered light */
 	UPROPERTY(EditAnywhere, Category = "Light Rendering")
-	FVector ColorAdd;
+	FVector3f ColorAdd;
 
 	/** If a render visibility tag is present, particles whose tag matches this value will be visible in this renderer. */
 	UPROPERTY(EditAnywhere, Category = "Light Rendering")
diff --git a/Engine/Source/Runtime/Engine/Private/Particles/ParticleSystemRender.cpp b/Engine/Source/Runtime/Engine/Private/Particles/ParticleSystemRender.cpp
--- a/Engine/Source/Runtime/Engine/Private/Particles/ParticleSystemRender.cpp
+++ b/Engine/Source/Runtime/Engine/Private/Particles/ParticleSystemRender.cpp
@@ -711,7 +711,7 @@ void GatherParticleLightData(const FDynamicSpriteEmitterReplayDataBase& Source,
 				
 				FSimpleLightEntry ParticleLight;
 				ParticleLight.Radius =  LightPayload->RadiusScale * (Size.X + Size.Y) / 2.0f;
-				ParticleLight.Color = FVector(Particle.Color * Particle.Color.A * LightPayload->ColorScale);
+				ParticleLight.Color = FVector3f(Particle.Color) * Particle.Color.A * LightPayload->ColorScale;
 				ParticleLight.Exponent = LightPayload->LightExponent;
 				ParticleLight.VolumetricScatteringIntensity = Source.LightVolumetricScatteringIntensity;
 				ParticleLight.bAffectTranslucency = LightPayload->bAffectsTranslucency;
diff --git a/Engine/Source/Runtime/Engine/Public/PrimitiveSceneProxy.h b/Engine/Source/Runtime/Engine/Public/PrimitiveSceneProxy.h
--- a/Engine/Source/Runtime/Engine/Public/PrimitiveSceneProxy.h
+++ b/Engine/Source/Runtime/Engine/Public/PrimitiveSceneProxy.h
@@ -40,7 +40,7 @@ namespace Nanite
 class FSimpleLightEntry
 {
 public:
-	FVector Color;
+	FVector3f Color;
 	float Radius;
 	float Exponent;
 	float VolumetricScatteringIntensity;
