- Fix typo Emitter vs Particle

#rnx

diff --git a/Engine/Plugins/FX/Niagara/Source/Niagara/Private/NiagaraRendererMeshes.cpp b/Engine/Plugins/FX/Niagara/Source/Niagara/Private/NiagaraRendererMeshes.cpp
--- a/Engine/Plugins/FX/Niagara/Source/Niagara/Private/NiagaraRendererMeshes.cpp
+++ b/Engine/Plugins/FX/Niagara/Source/Niagara/Private/NiagaraRendererMeshes.cpp
@@ -126,7 +126,7 @@ FNiagaraRendererMeshes::FNiagaraRendererMeshes(ERHIFeatureLevel::Type FeatureLev
 	{
 		EmitterMeshIndexOffset = Emitter->GetRendererBoundVariables().IndexOf(Properties->MeshIndexBinding.GetParamMapBindableVariable());
 	}
-	else if (Data.GetVariableComponentOffsets(Properties->MeshIndexBinding.GetDataSetBindableVariable(), FloatOffset, EmitterMeshIndexOffset, HalfOffset))
+	else if (Data.GetVariableComponentOffsets(Properties->MeshIndexBinding.GetDataSetBindableVariable(), FloatOffset, ParticleMeshIndexOffset, HalfOffset))
 	{
 		bEnableCulling = ParticleMeshIndexOffset != INDEX_NONE;
 	}
