Improve performance when RigidMeshCollision DI has a large maximum number of collision elements

diff --git a/Engine/Plugins/FX/Niagara/Source/Niagara/Private/NiagaraDataInterfaceRigidMeshCollisionQuery.cpp b/Engine/Plugins/FX/Niagara/Source/Niagara/Private/NiagaraDataInterfaceRigidMeshCollisionQuery.cpp
--- a/Engine/Plugins/FX/Niagara/Source/Niagara/Private/NiagaraDataInterfaceRigidMeshCollisionQuery.cpp
+++ b/Engine/Plugins/FX/Niagara/Source/Niagara/Private/NiagaraDataInterfaceRigidMeshCollisionQuery.cpp
@@ -1527,28 +1527,32 @@ void UNiagaraDataInterfaceRigidMeshCollisionQuery::ProvidePerInstanceDataForRend
 	{
 		if (const FNDIRigidMeshCollisionArrays* SourceArrayData = GameThreadData->AssetArrays.Get())
 		{
+			const int32 ElementCount = GameThreadData->AssetArrays->ElementOffsets.NumElements;
+
 			RenderThreadData->ElementOffsets = GameThreadData->AssetArrays->ElementOffsets;
 
 			// compact the world/inverse transforms
-			const int32 TransformVectorCount = GameThreadData->AssetArrays->MaxPrimitives * 3;
+			const int32 TransformVectorCount = ElementCount * 3;
 
 			auto CompactTransforms = [&](const TArray<FVector4f>& Current, const TArray<FVector4f>& Previous, TArray<FVector4f>& Compact)
 			{
-				Compact.Reset(2 * TransformVectorCount);
-				Compact.Append(Current);
-				Compact.SetNumUninitialized(TransformVectorCount, false);
-				Compact.Append(Previous);
-				Compact.SetNumUninitialized(2 * TransformVectorCount, false);
+				check(Current.Num() >= TransformVectorCount);
+				check(Previous.Num() >= TransformVectorCount);
+
+				Compact.Reset(2 * TransformVectorCount); // space for current and previous transforms
+				Compact.Append(Current.GetData(), TransformVectorCount);
+				Compact.Append(Previous.GetData(), TransformVectorCount);
 			};
 
 			CompactTransforms(GameThreadData->AssetArrays->CurrentTransform, GameThreadData->AssetArrays->PreviousTransform, RenderThreadData->WorldTransform);
 			CompactTransforms(GameThreadData->AssetArrays->CurrentInverse, GameThreadData->AssetArrays->PreviousInverse, RenderThreadData->InverseTransform);
 
-			RenderThreadData->ElementExtent = GameThreadData->AssetArrays->ElementExtent;
-			RenderThreadData->PhysicsType = GameThreadData->AssetArrays->PhysicsType;
-			RenderThreadData->ComponentIdIndex = GameThreadData->AssetArrays->ComponentIdIndex;
+			RenderThreadData->ElementExtent.Append(GameThreadData->AssetArrays->ElementExtent.GetData(), ElementCount);
+			RenderThreadData->PhysicsType.Append(GameThreadData->AssetArrays->PhysicsType.GetData(), ElementCount);
+			RenderThreadData->ComponentIdIndex.Append(GameThreadData->AssetArrays->ComponentIdIndex.GetData(), ElementCount);
+
 			RenderThreadData->UniqueComponentIds = GameThreadData->AssetArrays->UniqueCompnentId;
-			RenderThreadData->MaxPrimitiveCount = GameThreadData->AssetArrays->MaxPrimitives;
+			RenderThreadData->MaxPrimitiveCount = ElementCount;
 			RenderThreadData->AssetBuffer = GameThreadData->AssetBuffer;
 		}
 	}
