Sim Cache Viewport

Adding tabs for viewport & preview scene settings to sim cache toolkit. The viewport will create a Niagara system using the sim cache and will similate to the selected frame in the timeline.

diff --git a/Engine/Plugins/FX/Niagara/Source/Niagara/Classes/NiagaraSimCache.h b/Engine/Plugins/FX/Niagara/Source/Niagara/Classes/NiagaraSimCache.h
--- a/Engine/Plugins/FX/Niagara/Source/Niagara/Classes/NiagaraSimCache.h
+++ b/Engine/Plugins/FX/Niagara/Source/Niagara/Classes/NiagaraSimCache.h
@@ -305,6 +305,9 @@ public:
 	/** Get the emitter index from a name */
 	int GetEmitterIndex(FName EmitterName) const;
 
+	/** Get the system this cache is based on */
+	UNiagaraSystem* GetSystem (bool bLoadSynchronous = false);
+	
 	/** Get the emitter name at the provided index. */
 	UFUNCTION(BlueprintCallable, Category = NiagaraSimCache)
 	FName GetEmitterName(int32 EmitterIndex) const { return CacheLayout.EmitterLayouts.IsValidIndex(EmitterIndex) ? CacheLayout.EmitterLayouts[EmitterIndex].LayoutName : NAME_None; }
diff --git a/Engine/Plugins/FX/Niagara/Source/Niagara/Private/NiagaraSimCache.cpp b/Engine/Plugins/FX/Niagara/Source/Niagara/Private/NiagaraSimCache.cpp
--- a/Engine/Plugins/FX/Niagara/Source/Niagara/Private/NiagaraSimCache.cpp
+++ b/Engine/Plugins/FX/Niagara/Source/Niagara/Private/NiagaraSimCache.cpp
@@ -486,6 +486,16 @@ int UNiagaraSimCache::GetEmitterIndex(FName EmitterName) const
 	return INDEX_NONE;
 }
 
+UNiagaraSystem* UNiagaraSimCache::GetSystem(bool bLoadSynchronous)
+{
+	UNiagaraSystem* NiagaraSystem = SoftNiagaraSystem.Get();
+	if (NiagaraSystem == nullptr && bLoadSynchronous)
+	{
+		NiagaraSystem = SoftNiagaraSystem.LoadSynchronous();
+	}
+	return NiagaraSystem;
+}
+
 TArray<FName> UNiagaraSimCache::GetEmitterNames() const
 {
 	TArray<FName> EmitterNames;
diff --git a/Engine/Plugins/FX/Niagara/Source/NiagaraEditor/Private/Toolkits/NiagaraSimCacheToolkit.cpp b/Engine/Plugins/FX/Niagara/Source/NiagaraEditor/Private/Toolkits/NiagaraSimCacheToolkit.cpp
--- a/Engine/Plugins/FX/Niagara/Source/NiagaraEditor/Private/Toolkits/NiagaraSimCacheToolkit.cpp
+++ b/Engine/Plugins/FX/Niagara/Source/NiagaraEditor/Private/Toolkits/NiagaraSimCacheToolkit.cpp
@@ -4,16 +4,18 @@
 #include "NiagaraEditorModule.h"
 
 // Widgets & ViewModels
+#include "AdvancedPreviewSceneModule.h"
 #include "Widgets/Docking/SDockTab.h"
 #include "NiagaraSimCacheViewModel.h"
+#include "SNiagaraSimCacheViewport.h"
 #include "SNiagaraSimCacheViewTimeline.h"
 #include "SNiagaraSimCacheViewTransportControls.h"
 
 #define LOCTEXT_NAMESPACE "NiagaraSimCacheToolkit"
 
 const FName FNiagaraSimCacheToolkit::NiagaraSimCacheSpreadsheetTabId(TEXT("NiagaraSimCacheEditor_Spreadsheet"));
-
-//TODO Preview
+const FName FNiagaraSimCacheToolkit::NiagaraSimCacheViewportTabId(TEXT("NiagaraSimCacheEditor_Viewport"));
+const FName FNiagaraSimCacheToolkit::NiagaraSimCachePreviewSettingsTabId(TEXT("NiagaraSimCacheEditor_PreviewSettings"));
 
 FNiagaraSimCacheToolkit::FNiagaraSimCacheToolkit()
 {
@@ -34,6 +36,14 @@ void FNiagaraSimCacheToolkit::RegisterTabSpawners(const TSharedRef<class FTabMan
 	InTabManager->RegisterTabSpawner(NiagaraSimCacheSpreadsheetTabId, FOnSpawnTab::CreateSP(this, &FNiagaraSimCacheToolkit::SpawnTab_SimCacheSpreadsheet))
 		.SetDisplayName(LOCTEXT("SpreadsheetTab", "Cache Spreadsheet"))
 		.SetGroup(WorkspaceMenuCategory.ToSharedRef());
+
+	InTabManager->RegisterTabSpawner(NiagaraSimCacheViewportTabId, FOnSpawnTab::CreateSP(this, &FNiagaraSimCacheToolkit::SpawnTab_SimCacheViewport))
+		.SetDisplayName(LOCTEXT("ViewpoirtTab", "Cache Viewport"))
+		.SetGroup(WorkspaceMenuCategory.ToSharedRef());
+
+	InTabManager->RegisterTabSpawner(NiagaraSimCachePreviewSettingsTabId, FOnSpawnTab::CreateSP(this, &FNiagaraSimCacheToolkit::SpawnTab_PreviewSettings))
+		.SetDisplayName(LOCTEXT("PreviewSceneSettingsTab","Preview Scene Settings"))
+		.SetGroup(WorkspaceMenuCategory.ToSharedRef());
 }
 
 void FNiagaraSimCacheToolkit::UnregisterTabSpawners(const TSharedRef<class FTabManager>& InTabManager)
@@ -62,6 +72,8 @@ void FNiagaraSimCacheToolkit::Initialize(const EToolkitMode::Type Mode, const TS
 					FTabManager::NewStack()
 					->SetSizeCoefficient(0.1f)
 					->AddTab(NiagaraSimCacheSpreadsheetTabId, ETabState::OpenedTab)
+					->AddTab(NiagaraSimCacheViewportTabId, ETabState::OpenedTab)
+					->SetForegroundTab(NiagaraSimCacheSpreadsheetTabId)
 				)
 			);
 
@@ -142,4 +154,45 @@ TSharedRef<SDockTab> FNiagaraSimCacheToolkit::SpawnTab_SimCacheSpreadsheet(const
 	return SpawnedTab;
 }
 
+TSharedRef<SDockTab>  FNiagaraSimCacheToolkit::SpawnTab_SimCacheViewport(const FSpawnTabArgs& Args)
+{
+	check(Args.GetTabId().TabType == NiagaraSimCacheViewportTabId);
+
+	if (!Viewport.IsValid())
+	{
+		Viewport = SNew(SNiagaraSimCacheViewport);
+	}
+	
+
+	TSharedRef<SDockTab> SpawnedTab =
+		SNew(SDockTab)
+		[
+			Viewport.ToSharedRef()
+		];
+
+	Viewport->SetPreviewComponent(SimCacheViewModel->GetPreviewComponent());
+	
+	return SpawnedTab;
+}
+
+TSharedRef<SDockTab> FNiagaraSimCacheToolkit::SpawnTab_PreviewSettings(const FSpawnTabArgs& Args)
+{
+	check(Args.GetTabId() == NiagaraSimCachePreviewSettingsTabId);
+
+	TSharedRef<SWidget> InWidget = SNullWidget::NullWidget;
+	if (Viewport.IsValid())
+	{
+		FAdvancedPreviewSceneModule& AdvancedPreviewSceneModule = FModuleManager::LoadModuleChecked<FAdvancedPreviewSceneModule>("AdvancedPreviewScene");
+		InWidget = AdvancedPreviewSceneModule.CreateAdvancedPreviewSceneSettingsWidget(Viewport->GetPreviewScene());
+	}
+
+	TSharedRef<SDockTab> SpawnedTab = SNew(SDockTab)
+		.Label(LOCTEXT("PreviewSceneSettingsTab", "Preview Scene Settings"))
+		[
+			InWidget
+		];
+
+	return SpawnedTab;
+}
+
 #undef LOCTEXT_NAMESPACE /** NiagaraSimCacheToolkit **/
\ No newline at end of file
diff --git a/Engine/Plugins/FX/Niagara/Source/NiagaraEditor/Private/Toolkits/NiagaraSimCacheToolkit.h b/Engine/Plugins/FX/Niagara/Source/NiagaraEditor/Private/Toolkits/NiagaraSimCacheToolkit.h
--- a/Engine/Plugins/FX/Niagara/Source/NiagaraEditor/Private/Toolkits/NiagaraSimCacheToolkit.h
+++ b/Engine/Plugins/FX/Niagara/Source/NiagaraEditor/Private/Toolkits/NiagaraSimCacheToolkit.h
@@ -4,6 +4,7 @@
 
 #include "CoreMinimal.h"
 #include "SNiagaraSimCacheView.h"
+#include "SNiagaraSimCacheViewport.h"
 #include "Toolkits/IToolkitHost.h"
 #include "Toolkits/AssetEditorToolkit.h"
 
@@ -16,7 +17,7 @@ class FNiagaraSimCacheToolkit : public FAssetEditorToolkit
 public:
 	FNiagaraSimCacheToolkit();
 	~FNiagaraSimCacheToolkit();
-
+	
 	virtual void RegisterTabSpawners(const TSharedRef<class FTabManager>& InTabManager) override;
 	virtual void UnregisterTabSpawners(const TSharedRef<class FTabManager>& InTabManager) override;
 
@@ -37,6 +38,8 @@ protected:
 
 private:
 	TSharedRef<SDockTab> SpawnTab_SimCacheSpreadsheet(const FSpawnTabArgs& Args);
+	TSharedRef<SDockTab>  SpawnTab_SimCacheViewport(const FSpawnTabArgs& Args);
+	TSharedRef<SDockTab> SpawnTab_PreviewSettings(const FSpawnTabArgs& Args);
 
 	// The sim cache being viewed.
 	TWeakObjectPtr <UNiagaraSimCache> SimCache;
@@ -47,5 +50,10 @@ private:
 	// Spreadsheet widget
 	TSharedPtr<SNiagaraSimCacheView> SimCacheSpreadsheetView;
 
+	// Viewport Widget
+	TSharedPtr<SNiagaraSimCacheViewport> Viewport;
+
 	static const FName NiagaraSimCacheSpreadsheetTabId;
+	static const FName NiagaraSimCacheViewportTabId;
+	static const FName NiagaraSimCachePreviewSettingsTabId;
 };
\ No newline at end of file
diff --git a/Engine/Plugins/FX/Niagara/Source/NiagaraEditor/Private/ViewModels/NiagaraSimCacheViewModel.cpp b/Engine/Plugins/FX/Niagara/Source/NiagaraEditor/Private/ViewModels/NiagaraSimCacheViewModel.cpp
--- a/Engine/Plugins/FX/Niagara/Source/NiagaraEditor/Private/ViewModels/NiagaraSimCacheViewModel.cpp
+++ b/Engine/Plugins/FX/Niagara/Source/NiagaraEditor/Private/ViewModels/NiagaraSimCacheViewModel.cpp
@@ -1,10 +1,10 @@
 // Copyright Epic Games, Inc. All Rights Reserved.
 
 #include "ViewModels/NiagaraSimCacheViewModel.h"
+#include "NiagaraComponent.h"
 #include "NiagaraSimCache.h"
 #include "NiagaraEditorCommon.h"
 #include "NiagaraDebuggerCommon.h"
-#include "Misc/ScopeExit.h"
 #include "Serialization/ObjectAndNameAsStringProxyArchive.h"
 
 #define LOCTEXT_NAMESPACE "NiagaraSimCacheViewModel"
@@ -29,6 +29,7 @@ void FNiagaraSimCacheViewModel::Initialize(TWeakObjectPtr<UNiagaraSimCache> SimC
 
 	WeakSimCache = SimCache;
 	UpdateCachedFrame();
+	SetupPreviewComponentAndInstance();
 	OnViewDataChangedDelegate.Broadcast(true);
 }
 
@@ -51,6 +52,30 @@ void FNiagaraSimCacheViewModel::UpdateSimCache(const FNiagaraSystemSimCacheCaptu
 	Initialize(SimCache);
 }
 
+void FNiagaraSimCacheViewModel::SetupPreviewComponentAndInstance()
+{
+	UNiagaraSimCache* SimCache = WeakSimCache.Get();
+	UNiagaraSystem* System = SimCache ? SimCache->GetSystem(true) : nullptr;
+
+	if(SimCache && System)
+	{
+		PreviewComponent = NewObject<UNiagaraComponent>(GetTransientPackage(), NAME_None, RF_Transient);
+		PreviewComponent->CastShadow = 1;
+		PreviewComponent->bCastDynamicShadow = 1;
+		PreviewComponent->SetAllowScalability(false);
+		PreviewComponent->SetAsset(System);
+		PreviewComponent->SetForceSolo(true);
+		PreviewComponent->SetAgeUpdateMode(ENiagaraAgeUpdateMode::DesiredAge);
+		PreviewComponent->SetCanRenderWhileSeeking(false);
+		PreviewComponent->Activate(true);
+		PreviewComponent->SetSimCache(SimCache);
+		PreviewComponent->SetRelativeLocation(FVector::ZeroVector);
+		PreviewComponent->SetDesiredAge(SimCache->GetStartSeconds());
+	}
+
+	
+}
+
 FText FNiagaraSimCacheViewModel::GetComponentText(const FName ComponentName, const int32 InstanceIndex) const
 {
 	const FComponentInfo* ComponentInfo = ComponentInfos.FindByPredicate([ComponentName](const FComponentInfo& FoundInfo) { return FoundInfo.Name == ComponentName; });
@@ -100,6 +125,16 @@ void FNiagaraSimCacheViewModel::SetFrameIndex(const int32 InFrameIndex)
 {
 	FrameIndex = InFrameIndex;
 	UpdateCachedFrame();
+	const UNiagaraSimCache* SimCache = WeakSimCache.Get();
+	if(PreviewComponent && SimCache)
+	{
+		const float Duration = SimCache->GetDurationSeconds();
+		const int NumFrames = SimCache->GetNumFrames();
+
+		const float NormalizedFrame = NumFrames == 0 ? 0.0f : float(InFrameIndex) / float(NumFrames - 1);
+		
+		PreviewComponent->SetDesiredAge(SimCache->GetStartSeconds() + (Duration * NormalizedFrame));
+	}
 	OnViewDataChangedDelegate.Broadcast(false);
 }
 
diff --git a/Engine/Plugins/FX/Niagara/Source/NiagaraEditor/Private/ViewModels/NiagaraSimCacheViewModel.h b/Engine/Plugins/FX/Niagara/Source/NiagaraEditor/Private/ViewModels/NiagaraSimCacheViewModel.h
--- a/Engine/Plugins/FX/Niagara/Source/NiagaraEditor/Private/ViewModels/NiagaraSimCacheViewModel.h
+++ b/Engine/Plugins/FX/Niagara/Source/NiagaraEditor/Private/ViewModels/NiagaraSimCacheViewModel.h
@@ -7,6 +7,7 @@
 
 struct FNiagaraSystemSimCacheCaptureReply;
 struct FNiagaraSimCacheDataBuffersLayout;
+class UNiagaraComponent;
 
 class NIAGARAEDITOR_API FNiagaraSimCacheViewModel : public TSharedFromThis<FNiagaraSimCacheViewModel>
 {
@@ -31,6 +32,8 @@ public:
 
 	void UpdateSimCache(const FNiagaraSystemSimCacheCaptureReply& Reply);
 
+	void SetupPreviewComponentAndInstance();
+
 	TConstArrayView<FComponentInfo> GetComponentInfos() const { return ComponentInfos; }
 
 	int32 GetNumInstances() const { return NumInstances; }
@@ -59,14 +62,27 @@ public:
 
 	void UpdateCachedFrame();
 
-private:
-	void BuildComponentInfos(const FName Name, const UScriptStruct* Struct);
+	UNiagaraComponent* GetPreviewComponent() { return PreviewComponent; }
 
 private:
+	void BuildComponentInfos(const FName Name, const UScriptStruct* Struct);
+	
+	// The sim cache being viewed
 	TWeakObjectPtr<UNiagaraSimCache>	WeakSimCache;
-	int32								FrameIndex = 0;
-	int32								EmitterIndex = INDEX_NONE;
-	int32								NumInstances = 0;
+
+	// Component for preview scene
+	UNiagaraComponent* PreviewComponent = nullptr;
+
+	// Which frame of the cached sim is being viewed
+	int32 FrameIndex = 0;
+	
+	// Which Emitter of the cached sim is being viewed
+	int32 EmitterIndex = INDEX_NONE;
+
+	// Number of particles in the given frame
+	int32 NumInstances = 0;
+
+	// Info about the variables in the cache
 	TArray<FComponentInfo>				ComponentInfos;
 	TArray<float>						FloatComponents;
 	TArray<FFloat16>					HalfComponents;
diff --git a/Engine/Plugins/FX/Niagara/Source/NiagaraEditor/Private/Widgets/SNiagaraSimCacheViewport.cpp b/Engine/Plugins/FX/Niagara/Source/NiagaraEditor/Private/Widgets/SNiagaraSimCacheViewport.cpp
new file mode 100644
--- /dev/null
+++ b/Engine/Plugins/FX/Niagara/Source/NiagaraEditor/Private/Widgets/SNiagaraSimCacheViewport.cpp
@@ -0,0 +1,198 @@
+﻿// Copyright Epic Games, Inc. All Rights Reserved.
+
+#include "SNiagaraSimCacheViewport.h"
+
+#include "AdvancedPreviewScene.h"
+#include "NiagaraSimCacheViewModel.h"
+#include "SEditorViewportToolBarMenu.h"
+#include "SNiagaraSimCacheViewportToolbar.h"
+#include "NiagaraComponent.h"
+
+//////////////////////////////////////////////////////////////////////////
+// Viewport Client
+//////////////////////////////////////////////////////////////////////////
+
+
+class FNiagaraSimCacheViewportClient : public FEditorViewportClient
+{
+public:
+
+	FNiagaraSimCacheViewportClient(FAdvancedPreviewScene& InPreviewScene, const TSharedRef<SNiagaraSimCacheViewport>& InNiagaraEditorViewport);
+
+	// FEditorViewportClient interface
+	virtual FLinearColor GetBackgroundColor() const override;
+	virtual void Tick(float DeltaSeconds) override;
+	virtual void Draw(FViewport* Viewport,FCanvas* Canvas) override;
+	virtual bool ShouldOrbitCamera() const override;
+	virtual FSceneView* CalcSceneView(FSceneViewFamily* ViewFamily, const int32 ViewIndex = INDEX_NONE) override;
+	virtual bool CanSetWidgetMode(UE::Widget::EWidgetMode NewMode) const override { return false; }
+	virtual bool CanCycleWidgetMode() const override { return false; }
+
+	FAdvancedPreviewScene* AdvancedPreviewScene = nullptr;
+
+	TWeakPtr<SNiagaraSimCacheViewport> SimCacheViewportPrt;
+};
+
+FNiagaraSimCacheViewportClient::FNiagaraSimCacheViewportClient(FAdvancedPreviewScene& InPreviewScene, const TSharedRef<SNiagaraSimCacheViewport>& InNiagaraEditorViewport):
+	FEditorViewportClient(nullptr, &InPreviewScene, StaticCastSharedRef<SEditorViewport>(InNiagaraEditorViewport))
+{
+	SimCacheViewportPrt = InNiagaraEditorViewport;
+
+	// Setup defaults for the common draw helper.
+	DrawHelper.bDrawPivot = false;
+	DrawHelper.bDrawWorldBox = false;
+	DrawHelper.bDrawKillZ = false;
+	DrawHelper.bDrawGrid = true;
+	DrawHelper.GridColorAxis = FColor(80,80,80);
+	DrawHelper.GridColorMajor = FColor(72,72,72);
+	DrawHelper.GridColorMinor = FColor(64,64,64);
+	DrawHelper.PerspectiveGridSize = UE_OLD_HALF_WORLD_MAX1;
+	ShowWidget(false);
+
+	FEditorViewportClient::SetViewMode(VMI_Lit);
+
+	EngineShowFlags.SetSnap(0);
+
+	OverrideNearClipPlane(1.0f);
+
+	FNiagaraSimCacheViewportClient::SetIsSimulateInEditorViewport(true);
+}
+
+FLinearColor FNiagaraSimCacheViewportClient::GetBackgroundColor() const
+{
+	return FEditorViewportClient::GetBackgroundColor();
+}
+
+void FNiagaraSimCacheViewportClient::Tick(float DeltaSeconds)
+{
+	FEditorViewportClient::Tick(DeltaSeconds);
+
+	// Tick the preview scene world.
+	if (!GIntraFrameDebuggingGameThread)
+	{
+		PreviewScene->GetWorld()->Tick(LEVELTICK_All, DeltaSeconds);
+	}
+}
+
+void FNiagaraSimCacheViewportClient::Draw(FViewport* InViewport, FCanvas* Canvas)
+{
+
+	FEditorViewportClient::Draw(InViewport, Canvas);
+}
+
+bool FNiagaraSimCacheViewportClient::ShouldOrbitCamera() const
+{
+	return FEditorViewportClient::ShouldOrbitCamera();
+}
+
+FSceneView* FNiagaraSimCacheViewportClient::CalcSceneView(FSceneViewFamily* ViewFamily, const int32 StereoViewIndex)
+{
+	return FEditorViewportClient::CalcSceneView(ViewFamily, StereoViewIndex);
+}
+
+//////////////////////////////////////////////////////////////////////////
+// Viewport Widget
+//////////////////////////////////////////////////////////////////////////
+
+void SNiagaraSimCacheViewport::Construct(FArguments InArgs)
+{
+	AdvancedPreviewScene = MakeShareable(new FAdvancedPreviewScene(FPreviewScene::ConstructionValues()));
+	AdvancedPreviewScene->SetFloorVisibility(false);
+	AdvancedPreviewScene->SetEnvironmentVisibility(false, false);
+	
+	SEditorViewport::Construct(SEditorViewport::FArguments());
+
+	Client->EngineShowFlags.SetGrid(false);
+}
+
+SNiagaraSimCacheViewport::~SNiagaraSimCacheViewport()
+{
+	if(SimCacheViewportClient.IsValid())
+	{
+		SimCacheViewportClient->Viewport = nullptr;
+	}
+}
+
+void SNiagaraSimCacheViewport::AddReferencedObjects(FReferenceCollector& Collector)
+{
+	if (PreviewComponent)
+	{
+		Collector.AddReferencedObject(PreviewComponent);
+	}
+}
+
+void SNiagaraSimCacheViewport::Tick(const FGeometry& AllottedGeometry, const double InCurrentTime,
+                                    const float InDeltaTime)
+{
+	SEditorViewport::Tick(AllottedGeometry, InCurrentTime, InDeltaTime);
+}
+
+void SNiagaraSimCacheViewport::SetPreviewComponent(UNiagaraComponent* NewPreviewComponent)
+{
+	if (PreviewComponent != nullptr)
+	{
+		AdvancedPreviewScene->RemoveComponent(PreviewComponent);
+	}
+	
+	PreviewComponent = NewPreviewComponent;
+
+	if (PreviewComponent != nullptr)
+	{
+		AdvancedPreviewScene->AddComponent(PreviewComponent, PreviewComponent->GetRelativeTransform());
+	}
+}
+
+TSharedRef<SEditorViewport> SNiagaraSimCacheViewport::GetViewportWidget()
+{
+	return SharedThis(this);
+}
+
+TSharedPtr<FExtender> SNiagaraSimCacheViewport::GetExtenders() const
+{
+	TSharedPtr<FExtender> Result(MakeShareable(new FExtender));
+	return Result;
+}
+
+void SNiagaraSimCacheViewport::OnFloatingButtonClicked()
+{
+}
+
+TSharedRef<FEditorViewportClient> SNiagaraSimCacheViewport::MakeEditorViewportClient()
+{
+	SimCacheViewportClient = MakeShareable( new FNiagaraSimCacheViewportClient(*AdvancedPreviewScene.Get(), SharedThis(this)));
+
+	SimCacheViewportClient->SetViewLocation(FVector::ZeroVector);
+	SimCacheViewportClient->SetViewRotation(FRotator::ZeroRotator);
+	SimCacheViewportClient->SetViewLocationForOrbiting(FVector::ZeroVector);
+	SimCacheViewportClient->bSetListenerPosition = false;
+
+	SimCacheViewportClient->SetRealtime(true);
+	SimCacheViewportClient->VisibilityDelegate.BindSP(this, &SNiagaraSimCacheViewport::IsVisible);
+	
+	return SimCacheViewportClient.ToSharedRef();
+}
+
+TSharedPtr<SWidget> SNiagaraSimCacheViewport::MakeViewportToolbar()
+{
+	return SNew(SNiagaraSimCacheViewportToolbar, SharedThis(this));
+}
+
+EVisibility SNiagaraSimCacheViewport::OnGetViewportContentVisibility() const
+{
+	return SEditorViewport::OnGetViewportContentVisibility();
+}
+
+void SNiagaraSimCacheViewport::BindCommands()
+{
+	SEditorViewport::BindCommands();
+}
+
+void SNiagaraSimCacheViewport::OnFocusViewportToSelection()
+{
+	SEditorViewport::OnFocusViewportToSelection();
+}
+
+bool SNiagaraSimCacheViewport::IsVisible() const
+{
+	return SEditorViewport::IsVisible();
+}
diff --git a/Engine/Plugins/FX/Niagara/Source/NiagaraEditor/Private/Widgets/SNiagaraSimCacheViewport.h b/Engine/Plugins/FX/Niagara/Source/NiagaraEditor/Private/Widgets/SNiagaraSimCacheViewport.h
new file mode 100644
--- /dev/null
+++ b/Engine/Plugins/FX/Niagara/Source/NiagaraEditor/Private/Widgets/SNiagaraSimCacheViewport.h
@@ -0,0 +1,57 @@
+﻿// Copyright Epic Games, Inc. All Rights Reserved.
+
+#pragma once
+#include "AdvancedPreviewScene.h"
+#include "SCommonEditorViewportToolbarBase.h"
+#include "UObject/GCObject.h"
+#include "SEditorViewport.h"
+
+class UNiagaraComponent;
+
+class SNiagaraSimCacheViewport : public SEditorViewport, public FGCObject, public ICommonEditorViewportToolbarInfoProvider
+{
+	SLATE_BEGIN_ARGS(SNiagaraSimCacheViewport) {}
+	SLATE_END_ARGS()
+	
+public:
+	void Construct(FArguments InArgs);
+	virtual ~SNiagaraSimCacheViewport() override;
+
+	/** FGCObject interface **/
+	virtual void AddReferencedObjects( FReferenceCollector& Collector ) override;
+
+	virtual FString GetReferencerName() const override
+	{
+		return TEXT("SNiagaraSimCacheViewport");
+	}
+
+	virtual void Tick(const FGeometry& AllottedGeometry, const double InCurrentTime, const float InDeltaTime) override;
+
+	void SetPreviewComponent (UNiagaraComponent* NewPreviewComponent);
+	UNiagaraComponent* GetPreviewComponent () const { return PreviewComponent; }
+
+	// ICommonEditorViewportToolbarInfoProvider interface
+	virtual TSharedRef<class SEditorViewport> GetViewportWidget() override;
+	virtual TSharedPtr<FExtender> GetExtenders() const override;
+	virtual void OnFloatingButtonClicked() override;
+	// End of ICommonEditorViewportToolbarInfoProvider interface
+
+	TSharedRef<FAdvancedPreviewScene> GetPreviewScene() const {return AdvancedPreviewScene.ToSharedRef();}
+
+protected:
+	/** SEditorViewport interface */
+	virtual TSharedRef<FEditorViewportClient> MakeEditorViewportClient() override;
+	virtual TSharedPtr<SWidget> MakeViewportToolbar() override;
+	virtual EVisibility OnGetViewportContentVisibility() const override;
+	virtual void BindCommands() override;
+	virtual void OnFocusViewportToSelection() override;
+
+private:
+	virtual bool IsVisible() const override;
+	
+	TSharedPtr<FAdvancedPreviewScene> AdvancedPreviewScene;
+
+	TSharedPtr<class FNiagaraSimCacheViewportClient> SimCacheViewportClient;
+
+	UNiagaraComponent* PreviewComponent = nullptr;
+};
diff --git a/Engine/Plugins/FX/Niagara/Source/NiagaraEditor/Private/Widgets/SNiagaraSimCacheViewportToolbar.cpp b/Engine/Plugins/FX/Niagara/Source/NiagaraEditor/Private/Widgets/SNiagaraSimCacheViewportToolbar.cpp
new file mode 100644
--- /dev/null
+++ b/Engine/Plugins/FX/Niagara/Source/NiagaraEditor/Private/Widgets/SNiagaraSimCacheViewportToolbar.cpp
@@ -0,0 +1,19 @@
+﻿// Copyright Epic Games, Inc. All Rights Reserved.
+
+#include "SNiagaraSimCacheViewportToolbar.h"
+
+void SNiagaraSimCacheViewportToolbar::Construct(const FArguments& InArgs, TSharedPtr<SNiagaraSimCacheViewport> InViewport)
+{
+	SCommonEditorViewportToolbarBase::Construct(SCommonEditorViewportToolbarBase::FArguments(), InViewport);
+}
+
+void SNiagaraSimCacheViewportToolbar::ExtendOptionsMenu(FMenuBuilder& OptionsMenuBuilder) const
+{
+	SCommonEditorViewportToolbarBase::ExtendOptionsMenu(OptionsMenuBuilder);
+}
+
+void SNiagaraSimCacheViewportToolbar::ExtendLeftAlignedToolbarSlots(TSharedPtr<SHorizontalBox> MainBoxPtr,
+	TSharedPtr<SViewportToolBar> ParentToolBarPtr) const
+{
+	SCommonEditorViewportToolbarBase::ExtendLeftAlignedToolbarSlots(MainBoxPtr, ParentToolBarPtr);
+}
diff --git a/Engine/Plugins/FX/Niagara/Source/NiagaraEditor/Private/Widgets/SNiagaraSimCacheViewportToolbar.h b/Engine/Plugins/FX/Niagara/Source/NiagaraEditor/Private/Widgets/SNiagaraSimCacheViewportToolbar.h
new file mode 100644
--- /dev/null
+++ b/Engine/Plugins/FX/Niagara/Source/NiagaraEditor/Private/Widgets/SNiagaraSimCacheViewportToolbar.h
@@ -0,0 +1,25 @@
+﻿// Copyright Epic Games, Inc. All Rights Reserved.
+
+#pragma once
+#include "SCommonEditorViewportToolbarBase.h"
+#include "SNiagaraSimCacheViewport.h"
+
+class SNiagaraSimCacheViewportToolbar final : public SCommonEditorViewportToolbarBase
+{
+public:
+	SLATE_BEGIN_ARGS(SNiagaraSimCacheViewportToolbar) {}
+	SLATE_END_ARGS()
+
+	void Construct(const FArguments& InArgs, TSharedPtr<SNiagaraSimCacheViewport> InViewport);
+
+	// SCommonEditorViewportToolbarBase interface
+	virtual void ExtendOptionsMenu(FMenuBuilder& OptionsMenuBuilder) const override;
+	/** Used for a custom realtime button as the default realtime button is for viewport realtime, not simulation realtime and overrides only allow for direct setting */
+	virtual void ExtendLeftAlignedToolbarSlots(TSharedPtr<SHorizontalBox> MainBoxPtr, TSharedPtr<SViewportToolBar> ParentToolBarPtr) const override;
+	virtual bool GetShowScalabilityMenu() const override
+	{
+		return true;
+	}
+	// End of SCommonEditorViewportToolbarBase
+	
+};
