Refactored the renderer add list to make use of a factory delegate pattern instead of adding one per class. This way we can customize the renderers before they get added without relying on the CDO default propagation. It also allows us to add variants later on.

diff --git a/Engine/Plugins/FX/Niagara/Source/NiagaraEditor/Private/NiagaraEditorModule.cpp b/Engine/Plugins/FX/Niagara/Source/NiagaraEditor/Private/NiagaraEditorModule.cpp
--- a/Engine/Plugins/FX/Niagara/Source/NiagaraEditor/Private/NiagaraEditorModule.cpp
+++ b/Engine/Plugins/FX/Niagara/Source/NiagaraEditor/Private/NiagaraEditorModule.cpp
@@ -150,6 +150,9 @@
 #include "NiagaraDebugVis.h"
 #include "NiagaraPerfBaseline.h"
 #include "NiagaraGraphDataCache.h"
+#include "NiagaraLightRendererProperties.h"
+#include "NiagaraRibbonRendererProperties.h"
+#include "NiagaraSpriteRendererProperties.h"
 #include "Misc/ScopedSlowTask.h"
 
 #include UE_INLINE_GENERATED_CPP_BY_NAME(NiagaraEditorModule)
@@ -1091,6 +1094,8 @@ void FNiagaraEditorModule::StartupModule()
 
 	RegisterSettings();
 
+	RegisterDefaultRendererFactories();
+	
 	// Register sequencer track editors
 	ISequencerModule &SequencerModule = FModuleManager::LoadModuleChecked<ISequencerModule>("Sequencer");
 	CreateEmitterTrackEditorHandle = SequencerModule.RegisterTrackEditor(FOnCreateTrackEditor::CreateStatic(&FNiagaraEmitterTrackEditor::CreateTrackEditor));
@@ -1546,6 +1551,11 @@ void FNiagaraEditorModule::ClearObjectPool()
 	ObjectPool.Empty();
 }
 
+void FNiagaraEditorModule::RegisterRendererCreationInfo(FNiagaraRendererCreationInfo InRendererCreationInfo)
+{
+	RendererCreationInfo.Add(InRendererCreationInfo);
+}
+
 void FNiagaraEditorModule::RegisterParameterTrackCreatorForType(const UScriptStruct& StructType, FOnCreateMovieSceneTrackForParameter CreateTrack)
 {
 	checkf(TypeToParameterTrackCreatorMap.Contains(&StructType) == false, TEXT("Type already registered"));
@@ -1714,6 +1724,64 @@ void FNiagaraEditorModule::GetTargetSystemAndEmitterForDataInterface(UNiagaraDat
 	}
 }
 
+void FNiagaraEditorModule::RegisterDefaultRendererFactories()
+{
+	RegisterRendererCreationInfo(FNiagaraRendererCreationInfo(
+		UNiagaraMeshRendererProperties::StaticClass()->GetDisplayNameText(),
+		FText::FromString(UNiagaraMeshRendererProperties::StaticClass()->GetDescription()),
+		FNiagaraRendererCreationInfo::FRendererFactory::CreateLambda([](UObject* OuterEmitter)
+		{
+			UNiagaraMeshRendererProperties* NewRenderer = NewObject<UNiagaraMeshRendererProperties>(OuterEmitter, NAME_None, RF_Transactional);
+			// we have an empty entry in the constructor. Due to CDO default value propagation being unwanted, we have to keep it in there
+			if(ensure(NewRenderer->Meshes.Num() == 1))
+			{
+				FSoftObjectPath DefaultMesh(TEXT("StaticMesh'/Niagara/DefaultAssets/S_Gnomon.S_Gnomon'"));
+				NewRenderer->Meshes[0].Mesh = Cast<UStaticMesh>(DefaultMesh.TryLoad());
+			}
+			return NewRenderer;
+		})));
+	
+	RegisterRendererCreationInfo(FNiagaraRendererCreationInfo(
+		UNiagaraSpriteRendererProperties::StaticClass()->GetDisplayNameText(),
+		FText::FromString(UNiagaraSpriteRendererProperties::StaticClass()->GetDescription()),
+		FNiagaraRendererCreationInfo::FRendererFactory::CreateLambda([](UObject* OuterEmitter)
+		{
+			UNiagaraSpriteRendererProperties* NewRenderer = NewObject<UNiagaraSpriteRendererProperties>(OuterEmitter, NAME_None, RF_Transactional);
+			FSoftObjectPath DefaultMaterial(TEXT("Material'/Niagara/DefaultAssets/DefaultSpriteMaterial.DefaultSpriteMaterial'"));
+			NewRenderer->Material = Cast<UMaterialInterface>(DefaultMaterial.TryLoad());
+			return NewRenderer;
+		})));
+	
+	RegisterRendererCreationInfo(FNiagaraRendererCreationInfo(
+		UNiagaraRibbonRendererProperties::StaticClass()->GetDisplayNameText(),
+		FText::FromString(UNiagaraRibbonRendererProperties::StaticClass()->GetDescription()),
+		FNiagaraRendererCreationInfo::FRendererFactory::CreateLambda([](UObject* OuterEmitter)
+		{
+			UNiagaraRibbonRendererProperties* NewRenderer = NewObject<UNiagaraRibbonRendererProperties>(OuterEmitter, NAME_None, RF_Transactional);
+			FSoftObjectPath DefaultMaterial(TEXT("Material'/Niagara/DefaultAssets/DefaultRibbonMaterial.DefaultRibbonMaterial'"));
+			NewRenderer->Material = Cast<UMaterialInterface>(DefaultMaterial.TryLoad());
+			return NewRenderer;
+		})));
+
+	RegisterRendererCreationInfo(FNiagaraRendererCreationInfo(
+		UNiagaraComponentRendererProperties::StaticClass()->GetDisplayNameText(),
+		FText::FromString(UNiagaraComponentRendererProperties::StaticClass()->GetDescription()),
+		FNiagaraRendererCreationInfo::FRendererFactory::CreateLambda([](UObject* OuterEmitter)
+		{
+			UNiagaraComponentRendererProperties* NewRenderer = NewObject<UNiagaraComponentRendererProperties>(OuterEmitter, NAME_None, RF_Transactional);
+			return NewRenderer;
+		})));
+
+	RegisterRendererCreationInfo(FNiagaraRendererCreationInfo(
+		UNiagaraLightRendererProperties::StaticClass()->GetDisplayNameText(),
+		FText::FromString(UNiagaraLightRendererProperties::StaticClass()->GetDescription()),
+		FNiagaraRendererCreationInfo::FRendererFactory::CreateLambda([](UObject* OuterEmitter)
+		{
+			UNiagaraLightRendererProperties* NewRenderer = NewObject<UNiagaraLightRendererProperties>(OuterEmitter, NAME_None, RF_Transactional);
+			return NewRenderer;
+		})));
+}
+
 void FNiagaraEditorModule::RegisterAssetTypeAction(IAssetTools& AssetTools, TSharedRef<IAssetTypeActions> Action)
 {
 	AssetTools.RegisterAssetTypeActions(Action);
diff --git a/Engine/Plugins/FX/Niagara/Source/NiagaraEditor/Private/ViewModels/Stack/NiagaraStackRenderItemGroup.cpp b/Engine/Plugins/FX/Niagara/Source/NiagaraEditor/Private/ViewModels/Stack/NiagaraStackRenderItemGroup.cpp
--- a/Engine/Plugins/FX/Niagara/Source/NiagaraEditor/Private/ViewModels/Stack/NiagaraStackRenderItemGroup.cpp
+++ b/Engine/Plugins/FX/Niagara/Source/NiagaraEditor/Private/ViewModels/Stack/NiagaraStackRenderItemGroup.cpp
@@ -8,6 +8,7 @@
 #include "ViewModels/NiagaraSystemViewModel.h"
 #include "NiagaraRendererProperties.h"
 #include "NiagaraClipboard.h"
+#include "NiagaraEditorModule.h"
 
 #include "ScopedTransaction.h"
 #include "Widgets/Notifications/SNotificationList.h"
@@ -23,8 +24,8 @@
 class FRenderItemGroupAddAction : public INiagaraStackItemGroupAddAction
 {
 public:
-	FRenderItemGroupAddAction(UClass* InRendererClass)
-		: RendererClass(InRendererClass)
+	FRenderItemGroupAddAction(FNiagaraRendererCreationInfo InRendererCreationInfo)
+		: RendererCreationInfo(InRendererCreationInfo)
 	{
 	}
 
@@ -35,12 +36,12 @@ public:
 
 	virtual FText GetDisplayName() const override
 	{
-		return RendererClass->GetDisplayNameText();
+		return RendererCreationInfo.DisplayName;
 	}
 
 	virtual FText GetDescription() const override
 	{
-		return FText::FromString(RendererClass->GetDescription());
+		return RendererCreationInfo.Description;
 	}
 
 	virtual FText GetKeywords() const override
@@ -48,13 +49,10 @@ public:
 		return FText();
 	}
 
-	UClass* GetRendererClass() const
-	{
-		return RendererClass;
-	}
+	const FNiagaraRendererCreationInfo& GetRendererCreationInfo() { return RendererCreationInfo; }
 
 private:
-	UClass* RendererClass;
+	FNiagaraRendererCreationInfo RendererCreationInfo;
 };
 
 class FRenderItemGroupAddUtilities : public TNiagaraStackItemGroupAddUtilities<UNiagaraRendererProperties*>
@@ -70,11 +68,11 @@ public:
 
 	virtual void GenerateAddActions(TArray<TSharedRef<INiagaraStackItemGroupAddAction>>& OutAddActions, const FNiagaraStackItemGroupAddOptions& AddProperties) const override
 	{
-		TArray<UClass*> RendererClasses;
-		GetDerivedClasses(UNiagaraRendererProperties::StaticClass(), RendererClasses);
-		for (UClass* RendererClass : RendererClasses)
+		const TArray<FNiagaraRendererCreationInfo>& RendererCreationInfos = FNiagaraEditorModule::Get().GetRendererCreationInfos();
+
+		for (const FNiagaraRendererCreationInfo& RendererCreationInfo : RendererCreationInfos)
 		{
-			OutAddActions.Add(MakeShared<FRenderItemGroupAddAction>(RendererClass));
+			OutAddActions.Add(MakeShared<FRenderItemGroupAddAction>(RendererCreationInfo));
 		}
 	}
 
@@ -91,7 +89,9 @@ public:
 		FScopedTransaction ScopedTransaction(LOCTEXT("AddNewRendererTransaction", "Add new renderer"));
 
 		FVersionedNiagaraEmitter VersionedEmitter = EmitterViewModelPinned->GetEmitter();
-		UNiagaraRendererProperties* RendererProperties = NewObject<UNiagaraRendererProperties>(VersionedEmitter.Emitter, RenderAddAction->GetRendererClass(), NAME_None, RF_Transactional);
+		UNiagaraRendererProperties* RendererProperties = RenderAddAction->GetRendererCreationInfo().RendererFactory.Execute(VersionedEmitter.Emitter);
+		check(RendererProperties != nullptr);
+		
 		VersionedEmitter.Emitter->AddRenderer(RendererProperties, VersionedEmitter.Version);
 
 		bool bVarsAdded = false;
diff --git a/Engine/Plugins/FX/Niagara/Source/NiagaraEditor/Public/NiagaraEditorModule.h b/Engine/Plugins/FX/Niagara/Source/NiagaraEditor/Public/NiagaraEditorModule.h
--- a/Engine/Plugins/FX/Niagara/Source/NiagaraEditor/Public/NiagaraEditorModule.h
+++ b/Engine/Plugins/FX/Niagara/Source/NiagaraEditor/Public/NiagaraEditorModule.h
@@ -11,6 +11,7 @@
 #include "AssetTypeCategories.h"
 #include "NiagaraPerfBaseline.h"
 #include "NiagaraDebuggerCommon.h"
+#include "NiagaraRendererProperties.h"
 #include "NiagaraEditorModule.generated.h"
 
 class IAssetTools;
@@ -93,6 +94,29 @@ private:
 	TObjectPtr<UNiagaraParameterDefinitions> ReservingDefinitionsAsset;
 };
 
+USTRUCT()
+struct NIAGARAEDITOR_API FNiagaraRendererCreationInfo
+{
+	DECLARE_DELEGATE_RetVal_OneParam(UNiagaraRendererProperties*, FRendererFactory, UObject* OuterEmitter);
+
+	GENERATED_BODY()
+
+	FNiagaraRendererCreationInfo() = default;
+	FNiagaraRendererCreationInfo(FText InDisplayName, FRendererFactory InFactory) : DisplayName(InDisplayName), RendererFactory(InFactory)
+	{}
+
+	FNiagaraRendererCreationInfo(FText InDisplayName, FText InDescription, FRendererFactory InFactory) : DisplayName(InDisplayName), Description(InDescription),RendererFactory(InFactory)
+	{}
+
+	UPROPERTY()
+	FText DisplayName;
+
+	UPROPERTY()
+	FText Description;
+	
+	FRendererFactory RendererFactory;
+};
+
 FORCEINLINE uint32 GetTypeHash(const FReservedParameter& ReservedParameter) { return GetTypeHash(ReservedParameter.GetParameter().GetName()); };
 
 /** Niagara Editor module */
@@ -155,6 +179,10 @@ public:
 	NIAGARAEDITOR_API void ReleaseObjectToPool(UObject* Obj);
 	NIAGARAEDITOR_API void ClearObjectPool();
 
+	/** Registers a new renderer creation delegate with the display name it's going to use for the UI. */
+	NIAGARAEDITOR_API void RegisterRendererCreationInfo(FNiagaraRendererCreationInfo RendererCreationInfo);
+	NIAGARAEDITOR_API const TArray<FNiagaraRendererCreationInfo>& GetRendererCreationInfos() const { return RendererCreationInfo; }
+	
 	void RegisterParameterTrackCreatorForType(const UScriptStruct& StructType, FOnCreateMovieSceneTrackForParameter CreateTrack);
 	void UnregisterParameterTrackCreatorForType(const UScriptStruct& StructType);
 	bool CanCreateParameterTrackForType(const UScriptStruct& StructType);
@@ -269,6 +297,7 @@ private:
 		TArray<TWeakObjectPtr<AssetType>> CachedAssets;
 	};
 
+	void RegisterDefaultRendererFactories();
 	void RegisterAssetTypeAction(IAssetTools& AssetTools, TSharedRef<IAssetTypeActions> Action);
 	void OnNiagaraSettingsChangedEvent(const FName& PropertyName, const UNiagaraSettings* Settings);
 	void OnPreGarbageCollection();
@@ -369,6 +398,8 @@ private:
 
 	TMap<UClass*, TArray<UObject*>> ObjectPool;
 
+	TArray<FNiagaraRendererCreationInfo> RendererCreationInfo;
+
 #if NIAGARA_PERF_BASELINES
 	void GeneratePerfBaselines(TArray<UNiagaraEffectType*>& BaselinesToGenerate);
 	void OnPerfBaselineWindowClosed(const TSharedRef<SWindow>& ClosedWindow);
diff --git a/Engine/Plugins/Runtime/GeometryCache/Source/GeometryCacheEd/GeometryCacheEd.Build.cs b/Engine/Plugins/Runtime/GeometryCache/Source/GeometryCacheEd/GeometryCacheEd.Build.cs
--- a/Engine/Plugins/Runtime/GeometryCache/Source/GeometryCacheEd/GeometryCacheEd.Build.cs
+++ b/Engine/Plugins/Runtime/GeometryCache/Source/GeometryCacheEd/GeometryCacheEd.Build.cs
@@ -21,6 +21,7 @@ public class GeometryCacheEd : ModuleRules
 				"AssetTools",
                 "GeometryCache",
 				"ToolMenus",
+				"NiagaraEditor"
 			}
 		);
 	}
diff --git a/Engine/Plugins/Runtime/GeometryCache/Source/GeometryCacheEd/Private/GeometryCacheEdModule.cpp b/Engine/Plugins/Runtime/GeometryCache/Source/GeometryCacheEd/Private/GeometryCacheEdModule.cpp
--- a/Engine/Plugins/Runtime/GeometryCache/Source/GeometryCacheEd/Private/GeometryCacheEdModule.cpp
+++ b/Engine/Plugins/Runtime/GeometryCache/Source/GeometryCacheEd/Private/GeometryCacheEdModule.cpp
@@ -10,6 +10,8 @@
 #include "GeometryCacheComponent.h"
 #include "GeometryCacheModule.h"
 #include "GeometryCacheThumbnailRenderer.h"
+#include "NiagaraGeometryCacheRendererProperties.h"
+#include "NiagaraEditor/Public/NiagaraEditorModule.h"
 
 IMPLEMENT_MODULE(FGeometryCacheEdModule, GeometryCacheEd)
 
@@ -27,6 +29,21 @@ void FGeometryCacheEdModule::StartupModule()
 	FComponentAssetBrokerage::RegisterBroker(MakeShareable(AssetBroker), UGeometryCacheComponent::StaticClass(), true, true);
 
 	UThumbnailManager::Get().RegisterCustomRenderer(UGeometryCache::StaticClass(), UGeometryCacheThumbnailRenderer::StaticClass());
+
+	FNiagaraEditorModule& NiagaraEditorModule = FNiagaraEditorModule::Get();
+	NiagaraEditorModule.RegisterRendererCreationInfo(FNiagaraRendererCreationInfo(
+		UNiagaraGeometryCacheRendererProperties::StaticClass()->GetDisplayNameText(),
+		FText::FromString(UNiagaraGeometryCacheRendererProperties::StaticClass()->GetDescription()),
+		FNiagaraRendererCreationInfo::FRendererFactory::CreateLambda([](UObject* OuterEmitter)
+		{
+			UNiagaraGeometryCacheRendererProperties* NewRenderer = NewObject<UNiagaraGeometryCacheRendererProperties>(OuterEmitter, NAME_None, RF_Transactional);
+			if(ensure(NewRenderer->GeometryCaches.Num() == 1))
+			{
+				FSoftObjectPath DefaultGeometryCache(TEXT("GeometryCache'/Niagara/DefaultAssets/DefaultGeometryCacheAsset.DefaultGeometryCacheAsset'"));
+				NewRenderer->GeometryCaches[0].GeometryCache = Cast<UGeometryCache>(DefaultGeometryCache.TryLoad());
+			}
+			return NewRenderer;
+		})));
 }
 
 void FGeometryCacheEdModule::ShutdownModule()
