- Fix validation issue when using ray tracing and Niagara meshes

diff --git a/Engine/Plugins/FX/Niagara/Source/Niagara/Classes/NiagaraGPUInstanceCountManager.h b/Engine/Plugins/FX/Niagara/Source/Niagara/Classes/NiagaraGPUInstanceCountManager.h
--- a/Engine/Plugins/FX/Niagara/Source/Niagara/Classes/NiagaraGPUInstanceCountManager.h
+++ b/Engine/Plugins/FX/Niagara/Source/Niagara/Classes/NiagaraGPUInstanceCountManager.h
@@ -109,6 +109,7 @@ public:
 	void UpdateDrawIndirectBuffers(FNiagaraGpuComputeDispatchInterface* ComputeDispatchInterface, FRHICommandList& RHICmdList, ENiagaraGPUCountUpdatePhase::Type CountPhase);
 
 	static const ERHIAccess kCountBufferDefaultState;
+	static const ERHIAccess kIndirectArgsDefaultState;
 
 	bool HasEntriesPendingFree() const { return InstanceCountClearTasks.Num() > 0; }
 
diff --git a/Engine/Plugins/FX/Niagara/Source/Niagara/Private/NiagaraGPUInstanceCountManager.cpp b/Engine/Plugins/FX/Niagara/Source/Niagara/Private/NiagaraGPUInstanceCountManager.cpp
--- a/Engine/Plugins/FX/Niagara/Source/Niagara/Private/NiagaraGPUInstanceCountManager.cpp
+++ b/Engine/Plugins/FX/Niagara/Source/Niagara/Private/NiagaraGPUInstanceCountManager.cpp
@@ -101,6 +101,7 @@ DECLARE_CYCLE_STAT(TEXT("GPU Readback Lock"), STAT_NiagaraGPUReadbackLock, STATG
 //*****************************************************************************
 
 const ERHIAccess FNiagaraGPUInstanceCountManager::kCountBufferDefaultState = ERHIAccess::SRVMask | ERHIAccess::CopySrc;
+const ERHIAccess FNiagaraGPUInstanceCountManager::kIndirectArgsDefaultState = ERHIAccess::IndirectArgs | ERHIAccess::SRVMask;
 
 FNiagaraGPUInstanceCountManager::FNiagaraGPUInstanceCountManager(ERHIFeatureLevel::Type InFeatureLevel)
 	: FeatureLevel(InFeatureLevel)
@@ -326,7 +327,7 @@ void FNiagaraGPUInstanceCountManager::FlushIndirectArgsPool()
 
 		TResourceArray<uint32> InitData;
 		InitData.AddZeroed(PoolEntry->AllocatedEntries * NIAGARA_DRAW_INDIRECT_ARGS_SIZE);
-		PoolEntry->Buffer.Initialize(TEXT("NiagaraGPUDrawIndirectArgs"), sizeof(uint32), PoolEntry->AllocatedEntries * NIAGARA_DRAW_INDIRECT_ARGS_SIZE, EPixelFormat::PF_R32_UINT, ERHIAccess::IndirectArgs, BUF_Static | BUF_DrawIndirect, &InitData);
+		PoolEntry->Buffer.Initialize(TEXT("NiagaraGPUDrawIndirectArgs"), sizeof(uint32), PoolEntry->AllocatedEntries * NIAGARA_DRAW_INDIRECT_ARGS_SIZE, EPixelFormat::PF_R32_UINT, kIndirectArgsDefaultState, BUF_Static | BUF_DrawIndirect, &InitData);
 
 		// Reset the timer
 		DrawIndirectLowWaterFrames = 0;
@@ -357,7 +358,7 @@ FNiagaraGPUInstanceCountManager::FIndirectArgSlot FNiagaraGPUInstanceCountManage
 
 			TResourceArray<uint32> InitData;
 			InitData.AddZeroed(NewEntry->AllocatedEntries * NIAGARA_DRAW_INDIRECT_ARGS_SIZE);
-			NewEntry->Buffer.Initialize(TEXT("NiagaraGPUDrawIndirectArgs"), sizeof(uint32), NewEntry->AllocatedEntries * NIAGARA_DRAW_INDIRECT_ARGS_SIZE, EPixelFormat::PF_R32_UINT, ERHIAccess::IndirectArgs, BUF_Static | BUF_DrawIndirect, &InitData);
+			NewEntry->Buffer.Initialize(TEXT("NiagaraGPUDrawIndirectArgs"), sizeof(uint32), NewEntry->AllocatedEntries * NIAGARA_DRAW_INDIRECT_ARGS_SIZE, EPixelFormat::PF_R32_UINT, kIndirectArgsDefaultState, BUF_Static | BUF_DrawIndirect, &InitData);
 
 			PoolEntry = NewEntry.Get();
 			DrawIndirectPool.Emplace(MoveTemp(NewEntry));
@@ -443,7 +444,7 @@ void FNiagaraGPUInstanceCountManager::UpdateDrawIndirectBuffers(FNiagaraGpuCompu
 		// Execute transitions
 		for (auto& PoolEntry : DrawIndirectPool)
 		{
-			Transitions.Emplace(PoolEntry->Buffer.UAV, ERHIAccess::IndirectArgs, ERHIAccess::UAVCompute);
+			Transitions.Emplace(PoolEntry->Buffer.UAV, kIndirectArgsDefaultState, ERHIAccess::UAVCompute);
 		}
 		RHICmdList.Transition(Transitions);
 
@@ -527,7 +528,7 @@ void FNiagaraGPUInstanceCountManager::UpdateDrawIndirectBuffers(FNiagaraGpuCompu
 		Transitions.Reset();
 		for (auto& PoolEntry : DrawIndirectPool)
 		{
-			Transitions.Emplace(PoolEntry->Buffer.UAV, ERHIAccess::UAVCompute, ERHIAccess::IndirectArgs);
+			Transitions.Emplace(PoolEntry->Buffer.UAV, ERHIAccess::UAVCompute, kIndirectArgsDefaultState);
 		}
 		Transitions.Emplace(CountBuffer.UAV, ERHIAccess::UAVCompute, kCountBufferDefaultState);
 		RHICmdList.Transition(Transitions);
