- Fix for Raytracing using PreViewTranslation for GPU instances
- This also fixes missing LWC support for Niagara & RT

diff --git a/Engine/Plugins/FX/Niagara/Shaders/Private/NiagaraGPURayTracingTransforms.usf b/Engine/Plugins/FX/Niagara/Shaders/Private/NiagaraGPURayTracingTransforms.usf
--- a/Engine/Plugins/FX/Niagara/Shaders/Private/NiagaraGPURayTracingTransforms.usf
+++ b/Engine/Plugins/FX/Niagara/Shaders/Private/NiagaraGPURayTracingTransforms.usf
@@ -35,12 +35,16 @@ float3	MeshScale;
 float4	MeshRotation;
 
 float4x4	LocalTransform;
+float3		LocalTransformTile;
 
 Buffer<float>	ParticleDataFloatBuffer;
 //Buffer<half>	ParticleDataHalfBuffer;
 Buffer<int>		ParticleDataIntBuffer;
 Buffer<uint>	GPUInstanceCountBuffer;
 
+float3 ViewTilePosition;
+float3 RelativePreViewTranslation;
+
 // Transforms are float3x4 row-major matrices stored as 3 float4
 // because of FXC matrix packing issues when using StructuredBuffer<float3x4>
 RWStructuredBuffer<float4> TLASTransforms;
@@ -93,18 +97,22 @@ void NiagaraGPURayTracingTransformsCS(uint3 DispatchThreadId : SV_DispatchThread
 	float4x4 InstanceTransformTransposed = 0;
 	if (InstanceIndex < GPUNumInstances)
 	{
-		float3 InstanceSimPosition = SafeGetVec3(PositionDataOffset, InstanceIndex, DefaultPosition);
-		// LWC_TODO: Precision loss
-		float3 InstancePosition = bLocalSpace ? InstanceSimPosition : LWCHackToFloat(MakeLWCVector3(SystemLWCTile, InstanceSimPosition));
-		float4 InstanceRotation = NiagaraQuatMul(SafeGetVec4(RotationDataOffset, InstanceIndex, DefaultRotation), MeshRotation);
-		float3 InstanceScale    = SafeGetVec3(ScaleDataOffset, InstanceIndex, DefaultScale) * MeshScale;
-
-		float4x4 InstanceTransform = NiagaraComposeTransformMatrix(InstanceScale, NiagaraQuatTo3x3(InstanceRotation), InstancePosition);
+		float3 InstancePosition	= SafeGetVec3(PositionDataOffset, InstanceIndex, DefaultPosition);
+		float4 InstanceRotation	= NiagaraQuatMul(SafeGetVec4(RotationDataOffset, InstanceIndex, DefaultRotation), MeshRotation);
+		float3 InstanceScale	= SafeGetVec3(ScaleDataOffset, InstanceIndex, DefaultScale) * MeshScale;
 
-		if (bLocalSpace == 1)
+		FLWCMatrix InstanceTransform;
+		if (bLocalSpace)
 		{
-			InstanceTransform = mul(InstanceTransform, LocalTransform);
+			float4x4 TempTrans = NiagaraComposeTransformMatrix(InstanceScale, NiagaraQuatTo3x3(InstanceRotation), InstancePosition);
+			TempTrans = mul(TempTrans, LocalTransform);
+			InstanceTransform = MakeLWCMatrix(LocalTransformTile, TempTrans);
 		}
+		else
+		{
+			InstanceTransform = NiagaraComposeTransformMatrix(InstanceScale, NiagaraQuatTo3x3(InstanceRotation), MakeLWCVector3(SystemLWCTile, InstancePosition));
+		}
+
 
 		uint bInstanceValid = 1;
 		if ( RenderVisibilityOffset != -1 )
@@ -123,10 +131,19 @@ void NiagaraGPURayTracingTransformsCS(uint3 DispatchThreadId : SV_DispatchThread
 
 		if ( bInstanceValid )
 		{
-			InstanceTransformTransposed = transpose(InstanceTransform);
+			FLWCVector3 PreViewTranslation = MakeLWCVector3(ViewTilePosition, RelativePreViewTranslation);
+			float3 MatrixTrans = LWCToFloat(LWCAdd(LWCGetOrigin(InstanceTransform), PreViewTranslation));
+			float3x3 MatrixRot = transpose(LWCToFloat3x3(InstanceTransform));
+
+			InstanceTransformTransposed[0].xyz = MatrixRot[0];
+			InstanceTransformTransposed[1].xyz = MatrixRot[1];
+			InstanceTransformTransposed[2].xyz = MatrixRot[2];
+			InstanceTransformTransposed[0].w = MatrixTrans.x;
+			InstanceTransformTransposed[1].w = MatrixTrans.y;
+			InstanceTransformTransposed[2].w = MatrixTrans.z;
 		}
 	}
 	TLASTransforms[InstanceIndex * 3 + 0] = InstanceTransformTransposed[0];
 	TLASTransforms[InstanceIndex * 3 + 1] = InstanceTransformTransposed[1];
 	TLASTransforms[InstanceIndex * 3 + 2] = InstanceTransformTransposed[2];
-}
\ No newline at end of file
+}
diff --git a/Engine/Plugins/FX/Niagara/Source/Niagara/Private/NiagaraRendererMeshes.cpp b/Engine/Plugins/FX/Niagara/Source/Niagara/Private/NiagaraRendererMeshes.cpp
--- a/Engine/Plugins/FX/Niagara/Source/Niagara/Private/NiagaraRendererMeshes.cpp
+++ b/Engine/Plugins/FX/Niagara/Source/Niagara/Private/NiagaraRendererMeshes.cpp
@@ -1695,8 +1695,14 @@ void FNiagaraRendererMeshes::GetDynamicRayTracingInstances(FRayTracingMaterialGa
 				const TRefCountPtr<FRDGPooledBuffer>& ExternalBuffer = Context.GraphBuilder.ConvertToExternalBuffer(InstanceGPUTransformsBufferRef);
 				RayTracingInstance.InstanceGPUTransformsSRV = ExternalBuffer->GetOrCreateSRV(FRHIBufferSRVCreateInfo());
 
-				FNiagaraGPURayTracingTransformsCS::FParameters* PassParameters = Context.GraphBuilder.AllocParameters< FNiagaraGPURayTracingTransformsCS::FParameters>();
+				const FLargeWorldRenderPosition AbsoluteViewOrigin(View->ViewMatrices.GetViewOrigin());
+				const FVector ViewTileOffset = AbsoluteViewOrigin.GetTileOffset();
+				const FVector RelativePreViewTranslation = View->ViewMatrices.GetPreViewTranslation() + ViewTileOffset;
+				const FVector3f ViewTilePosition = AbsoluteViewOrigin.GetTile();
+
+				const FLargeWorldRenderPosition LocalTransformOrigin(LocalTransform.GetOrigin());
 
+				FNiagaraGPURayTracingTransformsCS::FParameters* PassParameters = Context.GraphBuilder.AllocParameters< FNiagaraGPURayTracingTransformsCS::FParameters>();
 				{
 					PassParameters->ParticleDataFloatStride		= ParticleMeshRenderData.ParticleFloatDataStride;
 					//PassParameters.ParticleDataHalfStride		= ParticleMeshRenderData.ParticleHalfDataStride;
@@ -1712,7 +1718,8 @@ void FNiagaraRendererMeshes::GetDynamicRayTracingInstances(FRayTracingMaterialGa
 					PassParameters->MeshIndexOffset				= ParticleMeshIndexOffset;
 					PassParameters->RenderVisibilityValue		= RendererVisibility;
 					PassParameters->MeshIndexValue				= MeshData.SourceMeshIndex;
-					PassParameters->LocalTransform				= FMatrix44f(LocalTransform);		// LWC_TODO: Precision loss
+					PassParameters->LocalTransform				= FLargeWorldRenderScalar::MakeToRelativeWorldMatrix(LocalTransformOrigin.GetTileOffset(), LocalTransform);
+					PassParameters->LocalTransformTile			= LocalTransformOrigin.GetTile();
 					PassParameters->DefaultPosition				= CommonParameters.DefaultPosition;
 					PassParameters->DefaultRotation				= FVector4f(0.0f, 0.0f, 0.0f, 1.0f);
 					PassParameters->DefaultScale				= FVector3f(1.0f, 1.0f, 1.0f);
@@ -1723,6 +1730,9 @@ void FNiagaraRendererMeshes::GetDynamicRayTracingInstances(FRayTracingMaterialGa
 					PassParameters->ParticleDataIntBuffer		= ParticleMeshRenderData.ParticleIntSRV;
 					PassParameters->GPUInstanceCountBuffer		= ComputeDispatchInterface->GetGPUInstanceCounterManager().GetInstanceCountBuffer().SRV;
 					PassParameters->TLASTransforms				= ExternalBuffer->GetOrCreateUAV(FRHIBufferUAVCreateInfo());
+
+					PassParameters->ViewTilePosition			= ViewTilePosition;
+					PassParameters->RelativePreViewTranslation	= FVector3f(RelativePreViewTranslation);
 				}
 
 				Context.GraphBuilder.AddPass(
diff --git a/Engine/Plugins/FX/Niagara/Source/NiagaraVertexFactories/Public/NiagaraGPURayTracingTransformsShader.h b/Engine/Plugins/FX/Niagara/Source/NiagaraVertexFactories/Public/NiagaraGPURayTracingTransformsShader.h
--- a/Engine/Plugins/FX/Niagara/Source/NiagaraVertexFactories/Public/NiagaraGPURayTracingTransformsShader.h
+++ b/Engine/Plugins/FX/Niagara/Source/NiagaraVertexFactories/Public/NiagaraGPURayTracingTransformsShader.h
@@ -52,12 +52,16 @@ class NIAGARAVERTEXFACTORIES_API FNiagaraGPURayTracingTransformsCS : public FGlo
 		SHADER_PARAMETER(FVector4f, MeshRotation)
 
 		SHADER_PARAMETER(FMatrix44f,	LocalTransform)
+		SHADER_PARAMETER(FVector3f,		LocalTransformTile)
 
 		SHADER_PARAMETER_SRV(Buffer<float>,	ParticleDataFloatBuffer)
 		//SHADER_PARAMETER_SRV(Buffer<half>,	ParticleDataHalfBuffer)
 		SHADER_PARAMETER_SRV(Buffer<int>,	ParticleDataIntBuffer)
 		SHADER_PARAMETER_SRV(Buffer<uint>, GPUInstanceCountBuffer)
 
+		SHADER_PARAMETER(FVector3f,	ViewTilePosition)
+		SHADER_PARAMETER(FVector3f, RelativePreViewTranslation)
+
 		SHADER_PARAMETER_UAV(RWStructuredBuffer<float3x4>, TLASTransforms)
 	END_SHADER_PARAMETER_STRUCT()
 
