- Rewording a couple of tooltips for parameters

#rnx

diff --git a/Engine/Plugins/FX/Niagara/Source/Niagara/Private/NiagaraConstants.cpp b/Engine/Plugins/FX/Niagara/Source/Niagara/Private/NiagaraConstants.cpp
--- a/Engine/Plugins/FX/Niagara/Source/Niagara/Private/NiagaraConstants.cpp
+++ b/Engine/Plugins/FX/Niagara/Source/Niagara/Private/NiagaraConstants.cpp
@@ -242,7 +242,7 @@ void FNiagaraConstants::Init()
 		SystemStrMap.Add(SYS_PARAM_ENGINE_EXEC_COUNT, LOCTEXT("ExecCountDesc", "Total number of items we are iterating over in this script invocation."));
 		SystemStrMap.Add(SYS_PARAM_ENGINE_EMITTER_NUM_PARTICLES, LOCTEXT("EmitterNumParticles", "The number of particles for this emitter at the beginning of simulation. Should only be used in Emitter scripts."));
 		SystemStrMap.Add(SYS_PARAM_ENGINE_EMITTER_SIMULATION_POSITION, LOCTEXT("EngineEmitterSimulationPosition", "The frame of reference space which the Emitter simulation is relative to. Is 0,0,0 for an Emitter in local space, and is Engine.Owner.Position for an Emitter in world space."));
-		SystemStrMap.Add(SYS_PARAM_ENGINE_EMITTER_TOTAL_SPAWNED_PARTICLES, LOCTEXT("EmitterTotalSpawnedParticles", "The total number of particles spawned for this emitter at the beginning of this simulation. Should only be used by the particle spawn script the assign unique IDs."));
+		SystemStrMap.Add(SYS_PARAM_ENGINE_EMITTER_TOTAL_SPAWNED_PARTICLES, LOCTEXT("EmitterTotalSpawnedParticles", "The total number of particles spawned for this emitter instance since creation.  The value is cleared when the simulation is reset or reinitialized."));
 		SystemStrMap.Add(SYS_PARAM_ENGINE_EMITTER_SPAWN_COUNT_SCALE, LOCTEXT("EmitterSpawnCountScale", "The global spawn count scale, based on fx.NiagaraGlobalSpawnCountScale."));
 		SystemStrMap.Add(SYS_PARAM_ENGINE_EMITTER_INSTANCE_SEED, LOCTEXT("EmitterInstanceSeed", "A random seed that changes with every execution of the emitter instance."));
 		SystemStrMap.Add(SYS_PARAM_ENGINE_SYSTEM_NUM_EMITTERS_ALIVE, LOCTEXT("SystemNumEmittersAlive", "The number of emitters still alive attached to this system. Should only be used in System scripts."));
@@ -552,7 +552,7 @@ void FNiagaraConstants::Init()
 		AttrDescStrMap.Add(SYS_PARAM_PARTICLES_PIVOT_OFFSET, LOCTEXT("PivotOffsetParamDesc", "The central pivot point of the sprite in UV space"));
 		AttrDescStrMap.Add(SYS_PARAM_PARTICLES_MATERIAL_RANDOM, LOCTEXT("MaterialRandomParamDesc", "Used to drive the Particle Random node in the Material Editor. Without this set, any Particle Randoms will get 0.0."));
 		AttrDescStrMap.Add(SYS_PARAM_PARTICLES_LIGHT_RADIUS, LOCTEXT("LightRadiusParamDesc", "Used to drive the radius of the light when using a Light renderer."));
-		AttrDescStrMap.Add(SYS_PARAM_PARTICLES_LIGHT_ENABLED, LOCTEXT("LightEnabledParamDesc", "This can be used with the Light renderer control the spawn-rate of lights on a per-particle basis."));
+		AttrDescStrMap.Add(SYS_PARAM_PARTICLES_LIGHT_ENABLED, LOCTEXT("LightEnabledParamDesc", "This can be used with the Light renderer to dynamically control if a light exists on a per-particle basis."));
 		AttrDescStrMap.Add(SYS_PARAM_PARTICLES_LIGHT_EXPONENT, LOCTEXT("LightExponentParamDesc", "Used to drive the attenuation of the light when using a Light renderer without inverse squared falloff enabled."));
 		AttrDescStrMap.Add(SYS_PARAM_PARTICLES_LIGHT_VOLUMETRIC_SCATTERING, LOCTEXT("LightVolumetricScatteringParamDesc", "Used to drive the volumetric scattering intensity of the light when using a Light renderer."));
 		AttrDescStrMap.Add(SYS_PARAM_INSTANCE_ALIVE, LOCTEXT("AliveParamDesc", "Used to determine whether or not this particle instance is still valid or if it can be deleted."));
@@ -565,9 +565,9 @@ void FNiagaraConstants::Init()
 		AttrDescStrMap.Add(SYS_PARAM_PARTICLES_RIBBONLINKORDER, LOCTEXT("RibbonLinkOrderDesc", "Explicit order for linking particles within a ribbon. Particles of the same ribbon id will be connected into a ribbon in incrementing order of this attribute value."));
 		AttrDescStrMap.Add(RIBBONUVDISTANCE, LOCTEXT("RibbonUVDistanceDesc", "Distance from start of ribbon that the particle lies."));
 		AttrDescStrMap.Add(SYS_PARAM_PARTICLES_RIBBONU0OVERRIDE, LOCTEXT("RibbonU0OverrideDesc", "Overrides the U component of the UV0 texture coordinate of a ribbon particle."));
-		AttrDescStrMap.Add(SYS_PARAM_PARTICLES_RIBBONV0RANGEOVERRIDE, LOCTEXT("RibbonV0RangeOverrideDesc", "Overrives the V range across the width of a ribbon for the UV0 texture coordinate of a particle."));
+		AttrDescStrMap.Add(SYS_PARAM_PARTICLES_RIBBONV0RANGEOVERRIDE, LOCTEXT("RibbonV0RangeOverrideDesc", "Overrides the V range across the width of a ribbon for the UV0 texture coordinate of a particle."));
 		AttrDescStrMap.Add(SYS_PARAM_PARTICLES_RIBBONU1OVERRIDE, LOCTEXT("RibbonU1OverrideDesc", "Overrides the U component of the UV1 texture coordinate of a ribbon particle."));
-		AttrDescStrMap.Add(SYS_PARAM_PARTICLES_RIBBONV1RANGEOVERRIDE, LOCTEXT("RibbonV1RangeOverrideDesc", "Overrives the V range across the width of a ribbon for the UV1 texture coordinate of a particle."));
+		AttrDescStrMap.Add(SYS_PARAM_PARTICLES_RIBBONV1RANGEOVERRIDE, LOCTEXT("RibbonV1RangeOverrideDesc", "Overrides the V range across the width of a ribbon for the UV1 texture coordinate of a particle."));
 		AttrDescStrMap.Add(SYS_PARAM_PARTICLES_ID, LOCTEXT("IDDesc", "Engine managed particle attribute that is a persistent ID for each particle."));
 		AttrDescStrMap.Add(SYS_PARAM_PARTICLES_UNIQUE_ID, LOCTEXT("UniqueIDDesc", "Engine managed particle attribute that is a unique ID for each particle. The ID is incremented for each new particle spawned.")); 
 	}
