Fixed potential crash due to nullptr dereferencing

diff --git a/Engine/Plugins/FX/Niagara/Source/NiagaraEditor/Private/NiagaraSystemEditorData.cpp b/Engine/Plugins/FX/Niagara/Source/NiagaraEditor/Private/NiagaraSystemEditorData.cpp
--- a/Engine/Plugins/FX/Niagara/Source/NiagaraEditor/Private/NiagaraSystemEditorData.cpp
+++ b/Engine/Plugins/FX/Niagara/Source/NiagaraEditor/Private/NiagaraSystemEditorData.cpp
@@ -389,7 +389,12 @@ UNiagaraScriptVariable* UNiagaraSystemEditorData::FindUserScriptVariable(FGuid U
 			return ScriptVariable->Metadata.GetVariableGuid() == UserParameterGuid;
 		});
 
-	return *ExistingScriptVariable;
+	if(ExistingScriptVariable)
+	{
+		return *ExistingScriptVariable;
+	}
+	
+	return nullptr;
 }
 
 bool UNiagaraSystemEditorData::RenameUserScriptVariable(FNiagaraVariable OldVariable, FName NewName)
