Making all the renderer material bindings show up in GetBoundAttributes, which they weren't before. This could lead to subtle differences in the data being preserved in cooked/uncooked.

diff --git a/Engine/Plugins/FX/Niagara/Source/Niagara/Private/NiagaraMeshRendererProperties.cpp b/Engine/Plugins/FX/Niagara/Source/Niagara/Private/NiagaraMeshRendererProperties.cpp
--- a/Engine/Plugins/FX/Niagara/Source/Niagara/Private/NiagaraMeshRendererProperties.cpp
+++ b/Engine/Plugins/FX/Niagara/Source/Niagara/Private/NiagaraMeshRendererProperties.cpp
@@ -501,6 +501,21 @@ void UNiagaraMeshRendererProperties::GetStreamingMeshInfo(const FBoxSphereBounds
 	}
 }
 
+
+#if WITH_EDITORONLY_DATA
+TArray<FNiagaraVariable> UNiagaraMeshRendererProperties::GetBoundAttributes() const 
+{
+	TArray<FNiagaraVariable> BoundAttributes = Super::GetBoundAttributes();
+	BoundAttributes.Reserve(BoundAttributes.Num() + MaterialParameters.AttributeBindings.Num());
+
+	for (const FNiagaraMaterialAttributeBinding& MaterialParamBinding : MaterialParameters.AttributeBindings)
+	{
+		BoundAttributes.AddUnique(MaterialParamBinding.GetParamMapBindableVariable());
+	}
+	return BoundAttributes;
+}
+#endif
+
 bool UNiagaraMeshRendererProperties::PopulateRequiredBindings(FNiagaraParameterStore& InParameterStore)
 {
 	bool bAnyAdded = Super::PopulateRequiredBindings(InParameterStore);
diff --git a/Engine/Plugins/FX/Niagara/Source/Niagara/Private/NiagaraRibbonRendererProperties.cpp b/Engine/Plugins/FX/Niagara/Source/Niagara/Private/NiagaraRibbonRendererProperties.cpp
--- a/Engine/Plugins/FX/Niagara/Source/Niagara/Private/NiagaraRibbonRendererProperties.cpp
+++ b/Engine/Plugins/FX/Niagara/Source/Niagara/Private/NiagaraRibbonRendererProperties.cpp
@@ -194,6 +194,20 @@ bool UNiagaraRibbonRendererProperties::IsBackfaceCullingDisabled() const
 	}
 }
 
+#if WITH_EDITORONLY_DATA
+TArray<FNiagaraVariable> UNiagaraRibbonRendererProperties::GetBoundAttributes() const 
+{
+	TArray<FNiagaraVariable> BoundAttributes = Super::GetBoundAttributes();
+	BoundAttributes.Reserve(BoundAttributes.Num() + MaterialParameters.AttributeBindings.Num());
+
+	for (const FNiagaraMaterialAttributeBinding& MaterialParamBinding : MaterialParameters.AttributeBindings)
+	{
+		BoundAttributes.AddUnique(MaterialParamBinding.GetParamMapBindableVariable());
+	}
+	return BoundAttributes;
+}
+#endif
+
 bool UNiagaraRibbonRendererProperties::PopulateRequiredBindings(FNiagaraParameterStore& InParameterStore)
 {
 	bool bAnyAdded = Super::PopulateRequiredBindings(InParameterStore);
diff --git a/Engine/Plugins/FX/Niagara/Source/Niagara/Private/NiagaraSpriteRendererProperties.cpp b/Engine/Plugins/FX/Niagara/Source/Niagara/Private/NiagaraSpriteRendererProperties.cpp
--- a/Engine/Plugins/FX/Niagara/Source/Niagara/Private/NiagaraSpriteRendererProperties.cpp
+++ b/Engine/Plugins/FX/Niagara/Source/Niagara/Private/NiagaraSpriteRendererProperties.cpp
@@ -338,6 +338,19 @@ void UNiagaraSpriteRendererProperties::CacheFromCompiledData(const FNiagaraDataS
 
 }
 
+#if WITH_EDITORONLY_DATA
+TArray<FNiagaraVariable> UNiagaraSpriteRendererProperties::GetBoundAttributes() const
+{
+	TArray<FNiagaraVariable> BoundAttributes = Super::GetBoundAttributes();
+	BoundAttributes.Reserve(BoundAttributes.Num() + MaterialParameters.AttributeBindings.Num());
+
+	for (const FNiagaraMaterialAttributeBinding& MaterialParamBinding : MaterialParameters.AttributeBindings)
+	{
+		BoundAttributes.AddUnique(MaterialParamBinding.GetParamMapBindableVariable());
+	}
+	return BoundAttributes;
+}
+#endif
 
 bool UNiagaraSpriteRendererProperties::PopulateRequiredBindings(FNiagaraParameterStore& InParameterStore)
 {
diff --git a/Engine/Plugins/FX/Niagara/Source/Niagara/Public/NiagaraMeshRendererProperties.h b/Engine/Plugins/FX/Niagara/Source/Niagara/Public/NiagaraMeshRendererProperties.h
--- a/Engine/Plugins/FX/Niagara/Source/Niagara/Public/NiagaraMeshRendererProperties.h
+++ b/Engine/Plugins/FX/Niagara/Source/Niagara/Public/NiagaraMeshRendererProperties.h
@@ -190,6 +190,7 @@ public:
 	void OnMeshPostBuild(UStaticMesh*);
 	void OnAssetReimported(UObject*);
 	void CheckMaterialUsage();
+	virtual TArray<FNiagaraVariable> GetBoundAttributes() const override;
 #endif // WITH_EDITORONLY_DATA
 	virtual void CacheFromCompiledData(const FNiagaraDataSetCompiledData* CompiledData) override;
 
diff --git a/Engine/Plugins/FX/Niagara/Source/Niagara/Public/NiagaraRibbonRendererProperties.h b/Engine/Plugins/FX/Niagara/Source/Niagara/Public/NiagaraRibbonRendererProperties.h
--- a/Engine/Plugins/FX/Niagara/Source/Niagara/Public/NiagaraRibbonRendererProperties.h
+++ b/Engine/Plugins/FX/Niagara/Source/Niagara/Public/NiagaraRibbonRendererProperties.h
@@ -225,6 +225,8 @@ public:
 	virtual void GetRendererWidgets(const FNiagaraEmitterInstance* InEmitter, TArray<TSharedPtr<SWidget>>& OutWidgets, TSharedPtr<FAssetThumbnailPool> InThumbnailPool) const override;
 	virtual void GetRendererTooltipWidgets(const FNiagaraEmitterInstance* InEmitter, TArray<TSharedPtr<SWidget>>& OutWidgets, TSharedPtr<FAssetThumbnailPool> InThumbnailPool) const override;
 	virtual void GetRendererFeedback(const FVersionedNiagaraEmitter& InEmitter, TArray<FNiagaraRendererFeedback>& OutErrors, TArray<FNiagaraRendererFeedback>& OutWarnings, TArray<FNiagaraRendererFeedback>& OutInfo) const override;
+
+	virtual TArray<FNiagaraVariable> GetBoundAttributes() const override; 
 #endif
 	virtual void CacheFromCompiledData(const FNiagaraDataSetCompiledData* CompiledData) override;
 	
diff --git a/Engine/Plugins/FX/Niagara/Source/Niagara/Public/NiagaraSpriteRendererProperties.h b/Engine/Plugins/FX/Niagara/Source/Niagara/Public/NiagaraSpriteRendererProperties.h
--- a/Engine/Plugins/FX/Niagara/Source/Niagara/Public/NiagaraSpriteRendererProperties.h
+++ b/Engine/Plugins/FX/Niagara/Source/Niagara/Public/NiagaraSpriteRendererProperties.h
@@ -113,6 +113,7 @@ public:
 	virtual void PostEditChangeProperty(struct FPropertyChangedEvent& PropertyChangedEvent) override;
 	virtual void RenameVariable(const FNiagaraVariableBase& OldVariable, const FNiagaraVariableBase& NewVariable, const FVersionedNiagaraEmitter& InEmitter) override;
 	virtual void RemoveVariable(const FNiagaraVariableBase& OldVariable, const FVersionedNiagaraEmitter& InEmitter) override;
+	virtual TArray<FNiagaraVariable> GetBoundAttributes() const override;
 
 #endif // WITH_EDITORONLY_DATA
 	//UObject Interface END
