Fixing several issues with parameter panel usage highlights.
1) Fixing issue where base level SelectByName wasn't handling new variables that weren't used yet well. Also didn't clear out the cache.
2) Fixing issue where renderer items didn't show icon for used variables.
3) Fixing issue where emitter namespaced variables in bindings weren't mapping to selected items properly

diff --git a/Engine/Plugins/FX/Niagara/Source/NiagaraEditor/Private/ViewModels/NiagaraParameterPanelViewModel.cpp b/Engine/Plugins/FX/Niagara/Source/NiagaraEditor/Private/ViewModels/NiagaraParameterPanelViewModel.cpp
--- a/Engine/Plugins/FX/Niagara/Source/NiagaraEditor/Private/ViewModels/NiagaraParameterPanelViewModel.cpp
+++ b/Engine/Plugins/FX/Niagara/Source/NiagaraEditor/Private/ViewModels/NiagaraParameterPanelViewModel.cpp
@@ -733,6 +733,7 @@ const TArray<FNiagaraParameterPanelItem>& INiagaraParameterPanelViewModel::GetCa
 	return CachedViewedItems;
 }
 
+
 void INiagaraParameterPanelViewModel::SelectParameterItemByName(const FName ParameterName, const bool bRequestRename)
 {
 	if (bRequestRename)
@@ -755,12 +756,29 @@ void INiagaraParameterPanelViewModel::SelectParameterItemByName(const FName Para
 		}
 	}
 
+	if (IncludeViewItemsInSelectParameterItem() && !FoundScriptVariable)
+	{
+		const TArray < FNiagaraParameterPanelItem>& CachedItems = GetCachedViewedParameterItems();
+		for (const FNiagaraParameterPanelItem& Item : CachedItems)
+		{
+			if (Item.GetVariable().GetName() == ParameterName)
+			{
+				FoundScriptVariable = (UNiagaraScriptVariable*)Item.ScriptVariable;
+				break;
+			}
+		}
+	}
+
+
 	const TSharedPtr<FNiagaraObjectSelection>& Selection = GetVariableObjectSelection();
 	if (Selection.IsValid() && FoundScriptVariable)
 	{
 		Selection->SetSelectedObject(FoundScriptVariable);
 	}
 
+
+	OnINiagaraParameterPanelViewModelSelectionChanged(FoundScriptVariable);
+
 	TSharedPtr<INiagaraParameterPanelViewModel> MasterVM = MasterParameterPanelViewModel.Pin();
 	if (MasterVM.IsValid())
 	{
@@ -1015,8 +1033,25 @@ void FNiagaraSystemToolkitParameterPanelViewModel::Init(const FSystemToolkitUICo
 }
 
 
+void FNiagaraSystemToolkitParameterPanelViewModel::OnINiagaraParameterPanelViewModelSelectionChanged(UNiagaraScriptVariable* InVar)
+{
+
+	if (InVar)
+	{
+		SelectedVariable = InVar->Variable;
+	}
+	else
+	{
+		VariableObjectSelection->ClearSelectedObjects();
+		SelectedVariable = FNiagaraVariable();
+	}
+	
+	InvalidateCachedDependencies();
+}
+
 void FNiagaraSystemToolkitParameterPanelViewModel::OnParameterItemSelected(const FNiagaraParameterPanelItem& SelectedItem, ESelectInfo::Type SelectInfo) const
 {
+	SelectedVariable = SelectedItem.GetVariable();
 	bool bFound = false;
 	for (const UNiagaraGraph* Graph : GetEditableGraphsConst())
 	{
@@ -1029,20 +1064,13 @@ void FNiagaraSystemToolkitParameterPanelViewModel::OnParameterItemSelected(const
 		}
 	}
 
-	if (bFound && SystemViewModel.IsValid())
+	if (!bFound)
 	{
-		if (SystemViewModel->GetSystemStackViewModel())
-			SystemViewModel->GetSystemStackViewModel()->InvalidateCachedParameterUsage();
-
-		for (TSharedRef<FNiagaraEmitterHandleViewModel> EmitterHandleVM : SystemViewModel->GetEmitterHandleViewModels())
-		{
-			UNiagaraStackViewModel* EmitterStackViewModel = EmitterHandleVM->GetEmitterStackViewModel();
-			if (EmitterStackViewModel)
-			{
-				EmitterStackViewModel->InvalidateCachedParameterUsage();
-			}
-		}
+		VariableObjectSelection->ClearSelectedObjects();
 	}
+
+	InvalidateCachedDependencies();
+
 }
 
 bool FNiagaraSystemToolkitParameterPanelViewModel::IsVariableSelected(FNiagaraVariableBase& InVar) const
@@ -1060,6 +1088,10 @@ bool FNiagaraSystemToolkitParameterPanelViewModel::IsVariableSelected(FNiagaraVa
 		}
 
 	}
+	if (IncludeViewItemsInSelectParameterItem() && InVar == SelectedVariable)
+	{
+		return true;
+	}
 
 	return false;
 }
@@ -2176,11 +2208,16 @@ TArray<FNiagaraParameterPanelItem> FNiagaraSystemToolkitParameterPanelViewModel:
 				FVersionedNiagaraEmitterData* ED = EmitterVMS.Get().GetEmitterHandle()->GetEmitterData();
 				if (ED)
 				{
+					FNiagaraAliasContext ResolveAliasesContext(FNiagaraAliasContext::ERapidIterationParameterMode::EmitterOrParticleScript);
+					ResolveAliasesContext.ChangeEmitterNameToEmitter(EmitterVMS.Get().GetEmitterHandle()->GetUniqueInstanceName());
+					
 					ED->ForEachEnabledRenderer(
 						[&](UNiagaraRendererProperties* RenderProperties)
 						{
 							for (FNiagaraVariableBase BoundAttribute : RenderProperties->GetBoundAttributes())
 							{
+								BoundAttribute = FNiagaraUtilities::ResolveAliases(BoundAttribute, ResolveAliasesContext);
+
 								if (FNiagaraParameterPanelItem* ItemPtr = VisitedParameterToItemMap.Find(BoundAttribute))
 								{
 									// This variable has already been registered, increment the reference count. Otherwise, it is 
diff --git a/Engine/Plugins/FX/Niagara/Source/NiagaraEditor/Private/ViewModels/NiagaraSystemViewModel.cpp b/Engine/Plugins/FX/Niagara/Source/NiagaraEditor/Private/ViewModels/NiagaraSystemViewModel.cpp
--- a/Engine/Plugins/FX/Niagara/Source/NiagaraEditor/Private/ViewModels/NiagaraSystemViewModel.cpp
+++ b/Engine/Plugins/FX/Niagara/Source/NiagaraEditor/Private/ViewModels/NiagaraSystemViewModel.cpp
@@ -62,6 +62,7 @@
 #include "ViewModels/Stack/NiagaraStackViewModel.h"
 #include "ViewModels/TNiagaraViewModelManager.h"
 #include "ViewModels/HierarchyEditor/NiagaraUserParametersHierarchyViewModel.h"
+#include "ViewModels/NiagaraParameterPanelViewModel.h"
 
 
 DECLARE_CYCLE_STAT(TEXT("Niagara - SystemViewModel - CompileSystem"), STAT_NiagaraEditor_SystemViewModel_CompileSystem, STATGROUP_NiagaraEditor);
@@ -176,6 +177,34 @@ void FNiagaraSystemViewModel::DumpToText(FString& ExportText)
 	FEdGraphUtilities::ExportNodesToText(ExportObjs, ExportText);
 }
 
+void FNiagaraSystemViewModel::SetParameterPanelViewModel(TSharedPtr<INiagaraParameterPanelViewModel> InVM) 
+{
+	if (ParameterPanelViewModel.IsValid())
+	{
+		ParameterPanelViewModel.Pin()->GetOnInvalidateCachedDependencies().RemoveAll(this);
+	}
+
+	ParameterPanelViewModel = InVM;
+
+	if (InVM.IsValid())
+	{
+		InVM->GetOnInvalidateCachedDependencies().AddRaw(this, &FNiagaraSystemViewModel::InvalidateCachedParams);
+	}
+}
+
+void FNiagaraSystemViewModel::InvalidateCachedParams()
+{
+	GetSystemStackViewModel()->InvalidateCachedParameterUsage();
+	for (TSharedRef<FNiagaraEmitterHandleViewModel> EmitterHandleVM : GetEmitterHandleViewModels())
+	{
+		UNiagaraStackViewModel* EmitterStackViewModel = EmitterHandleVM->GetEmitterStackViewModel();
+		if (EmitterStackViewModel)
+		{
+			EmitterStackViewModel->InvalidateCachedParameterUsage();
+		}
+	}
+}
+
 void FNiagaraSystemViewModel::Cleanup()
 {
 	if (SystemInstance)
@@ -286,6 +315,11 @@ void FNiagaraSystemViewModel::Cleanup()
 		SystemGraphSelectionViewModel = nullptr;
 	}
 
+	if (ParameterPanelViewModel.IsValid())
+	{
+		ParameterPanelViewModel.Pin()->GetOnInvalidateCachedDependencies().RemoveAll(this);
+	}
+
 	System = nullptr;
 }
 
diff --git a/Engine/Plugins/FX/Niagara/Source/NiagaraEditor/Private/ViewModels/Stack/NiagaraStackRendererItem.cpp b/Engine/Plugins/FX/Niagara/Source/NiagaraEditor/Private/ViewModels/Stack/NiagaraStackRendererItem.cpp
--- a/Engine/Plugins/FX/Niagara/Source/NiagaraEditor/Private/ViewModels/Stack/NiagaraStackRendererItem.cpp
+++ b/Engine/Plugins/FX/Niagara/Source/NiagaraEditor/Private/ViewModels/Stack/NiagaraStackRendererItem.cpp
@@ -55,6 +55,51 @@ void UNiagaraStackRendererItem::FinalizeInternal()
 	Super::FinalizeInternal();
 }
 
+const UNiagaraStackRendererItem::FCollectedUsageData& UNiagaraStackRendererItem::GetCollectedUsageData() const
+{
+
+	if (CachedCollectedUsageData.IsSet() == false)
+	{
+		CachedCollectedUsageData = FCollectedUsageData();
+
+		if (RendererProperties.IsValid())
+		{
+			TArray<FNiagaraVariable> BoundAttribs = RendererProperties->GetBoundAttributes();
+			TSharedRef<FNiagaraSystemViewModel> SystemVM = GetSystemViewModel();
+			INiagaraParameterPanelViewModel* ParamVM = SystemVM->GetParameterPanelViewModel();
+
+			FNiagaraAliasContext ResolveAliasesContext(FNiagaraAliasContext::ERapidIterationParameterMode::EmitterOrParticleScript);
+			
+			if (GetEmitterViewModel().IsValid())
+			{
+				TSharedPtr<FNiagaraEmitterHandleViewModel> EmitterHandleVM = SystemVM->GetEmitterHandleViewModelForEmitter(GetEmitterViewModel().Get()->GetEmitter());
+				if (EmitterHandleVM.IsValid() && EmitterHandleVM->GetEmitterHandle())
+				{
+					ResolveAliasesContext.ChangeEmitterNameToEmitter(EmitterHandleVM->GetEmitterHandle()->GetUniqueInstanceName());
+				}
+			}
+
+			if (ParamVM)
+			{
+				bool bFoundMatch = false;
+				for (FNiagaraVariable Var : BoundAttribs)
+				{			
+					Var = FNiagaraUtilities::ResolveAliases(Var, ResolveAliasesContext);
+					bFoundMatch = ParamVM->IsVariableSelected(Var);
+					if (bFoundMatch)
+					{
+						break;
+					}
+				}
+
+				CachedCollectedUsageData.GetValue().bHasReferencedParameterRead = bFoundMatch;
+			}
+		}
+	}
+
+	return CachedCollectedUsageData.GetValue();
+}
+
 TArray<FNiagaraVariable> UNiagaraStackRendererItem::GetMissingVariables(UNiagaraRendererProperties* RendererProperties, const FVersionedNiagaraEmitterData* EmitterData)
 {
 	TArray<FNiagaraVariable> MissingAttributes;
diff --git a/Engine/Plugins/FX/Niagara/Source/NiagaraEditor/Public/ViewModels/NiagaraParameterPanelViewModel.h b/Engine/Plugins/FX/Niagara/Source/NiagaraEditor/Public/ViewModels/NiagaraParameterPanelViewModel.h
--- a/Engine/Plugins/FX/Niagara/Source/NiagaraEditor/Public/ViewModels/NiagaraParameterPanelViewModel.h
+++ b/Engine/Plugins/FX/Niagara/Source/NiagaraEditor/Public/ViewModels/NiagaraParameterPanelViewModel.h
@@ -165,6 +165,9 @@ public:
 
 	/** Delegate to get the names of all selected parameter items. */
 	DECLARE_DELEGATE_RetVal(TArray<FName>, FOnGetSelectedParameterNames);
+	
+	DECLARE_MULTICAST_DELEGATE(FOnInvalidateCachedDependencies);
+	
 	INiagaraParameterPanelViewModel();
 	virtual ~INiagaraParameterPanelViewModel() override;
 
@@ -265,6 +268,7 @@ public:
 	FOnNotifyParameterPendingRename& GetOnNotifyParameterPendingRenameDelegate() { return OnNotifyParameterPendingRenameDelegate; };
 	FOnNotifyParameterPendingNamespaceModifierRename& GetOnNotifyParameterPendingNamespaceModifierRenameDelegate() { return OnNotifyParameterPendingNamespaceModifierRenameDelegate; };
 	FOnGetParametersWithNamespaceModifierRenamePending& GetParametersWithNamespaceModifierRenamePendingDelegate() { return OnGetParametersWithNamespaceModifierRenamePendingDelegate; };
+	FOnInvalidateCachedDependencies& GetOnInvalidateCachedDependencies(){ return OnInvalidateCachedDependenciesDelegate;}
 
 	FText GetActiveSection() const
 	{
@@ -292,7 +296,12 @@ protected:
 	FOnNotifyParameterPendingRename OnNotifyParameterPendingRenameDelegate;
 	FOnNotifyParameterPendingNamespaceModifierRename OnNotifyParameterPendingNamespaceModifierRenameDelegate;
 	FOnGetParametersWithNamespaceModifierRenamePending OnGetParametersWithNamespaceModifierRenamePendingDelegate;
+    FOnInvalidateCachedDependencies OnInvalidateCachedDependenciesDelegate;
 
+	virtual bool IncludeViewItemsInSelectParameterItem() const { return false; }
+	
+	void InvalidateCachedDependencies() const { OnInvalidateCachedDependenciesDelegate.Broadcast(); };
+	
 	/** SharedPtr to menu and searchbox widget retained to prevent the shared ref returned by GetParameterMenu from being invalidated. */
 	mutable TSharedPtr<SWidget> ParameterMenuWidget;
 	mutable TSharedPtr<SEditableTextBox> ParameterMenuSearchBoxWidget;
@@ -308,6 +317,8 @@ protected:
 
 	TArray<FSectionDesc> Sections;
 	int32 ActiveSectionIndex;
+
+	virtual void OnINiagaraParameterPanelViewModelSelectionChanged(UNiagaraScriptVariable* InVar) {};
 };
 
 class FNiagaraSystemToolkitParameterPanelViewModel : public INiagaraParameterPanelViewModel, public TNiagaraViewModelManager<UNiagaraSystem, FNiagaraSystemToolkitParameterPanelViewModel>
@@ -390,6 +401,10 @@ public:
 	virtual bool UsesCategoryFilteringForInitialExpansion() const override { return true; }
 
 	virtual void PreSectionChange(const TArray<FNiagaraParameterPanelCategory>& ExpandedItems) override;
+protected:
+	virtual void OnINiagaraParameterPanelViewModelSelectionChanged(UNiagaraScriptVariable* InVar) override;
+	virtual bool IncludeViewItemsInSelectParameterItem() const { return true; }
+
 private:
 
 	bool  ShouldRouteThroughScratchParameterMap(const FNiagaraParameterPanelCategory* Category = nullptr, const FNiagaraVariableBase* Variable = nullptr);
@@ -419,6 +434,8 @@ private:
 	// Graphs viewed to gather UNiagaraScriptVariables that are displayed by the Parameter Panel.
 	TWeakObjectPtr<UNiagaraGraph> SystemScriptGraph;
 
+	mutable FNiagaraVariable SelectedVariable;
+
 	TSharedPtr<FNiagaraSystemViewModel> SystemViewModel;
 	TWeakPtr<FNiagaraSystemGraphSelectionViewModel> SystemGraphSelectionViewModelWeak;
 
diff --git a/Engine/Plugins/FX/Niagara/Source/NiagaraEditor/Public/ViewModels/NiagaraSystemViewModel.h b/Engine/Plugins/FX/Niagara/Source/NiagaraEditor/Public/ViewModels/NiagaraSystemViewModel.h
--- a/Engine/Plugins/FX/Niagara/Source/NiagaraEditor/Public/ViewModels/NiagaraSystemViewModel.h
+++ b/Engine/Plugins/FX/Niagara/Source/NiagaraEditor/Public/ViewModels/NiagaraSystemViewModel.h
@@ -431,7 +431,7 @@ public:
 	INiagaraParameterPanelViewModel* GetParameterPanelViewModel() const {
 		return ParameterPanelViewModel.Pin().Get();
 	}
-	void SetParameterPanelViewModel(TSharedPtr<INiagaraParameterPanelViewModel> InVM) { ParameterPanelViewModel = InVM; }
+	void SetParameterPanelViewModel(TSharedPtr<INiagaraParameterPanelViewModel> InVM);
 
 	UNiagaraSystemEditorDocumentsViewModel* GetDocumentViewModel() { return EditorDocumentsViewModel;};
 
@@ -493,6 +493,9 @@ private:
 	/** Called whenever a property on the emitter changes. */
 	void EmitterPropertyChanged();
 
+	/** Called whenever we need to invalidate the parameter cache from the parameter panel updates*/
+	void InvalidateCachedParams();
+
 	/** 
 	 * Called whenever a parameter store owned by the system changes.
 	 * @param ChangedParameterStore The parameter store that changed.
diff --git a/Engine/Plugins/FX/Niagara/Source/NiagaraEditor/Public/ViewModels/Stack/NiagaraStackRendererItem.h b/Engine/Plugins/FX/Niagara/Source/NiagaraEditor/Public/ViewModels/Stack/NiagaraStackRendererItem.h
--- a/Engine/Plugins/FX/Niagara/Source/NiagaraEditor/Public/ViewModels/Stack/NiagaraStackRendererItem.h
+++ b/Engine/Plugins/FX/Niagara/Source/NiagaraEditor/Public/ViewModels/Stack/NiagaraStackRendererItem.h
@@ -70,6 +70,7 @@ public:
 	bool IsExcludedFromScalability() const;
 	bool IsOwningEmitterExcludedFromScalability() const;
 
+	virtual const FCollectedUsageData& GetCollectedUsageData() const override;
 protected:
 	virtual void FinalizeInternal() override;
 
