- Allow data interfaces to add renderer only attributes for capturing
- For example if a DI ran a process post cache read to update state required to render the system we may not cache it as the sim cache will only store renderer bound attributes

diff --git a/Engine/Plugins/FX/Niagara/Source/Niagara/Classes/NiagaraDataInterface.h b/Engine/Plugins/FX/Niagara/Source/Niagara/Classes/NiagaraDataInterface.h
--- a/Engine/Plugins/FX/Niagara/Source/Niagara/Classes/NiagaraDataInterface.h
+++ b/Engine/Plugins/FX/Niagara/Source/Niagara/Classes/NiagaraDataInterface.h
@@ -512,6 +512,13 @@ public:
 	Only DataInterfaces with PerInstanceData are currently supported.
 	*/
 	virtual void SimCachePostReadFrame(void* OptionalPerInstanceData, FNiagaraSystemInstance* SystemInstance) {}
+	/**
+	This function allows you to preserve a list of attributes when building a renderer only cache.
+	The UsageContext will be either a UNiagaraSystem or a UNiagaraEmitter and can be used to scope your variables accordingly.
+	For example, if you were to require 'Particles.MyAttribute' in order to process the cache results you would need to convert
+	this into 'MyEmitter.Particles.MyAttribute' by checking the UsageContext is a UNiagaraEmitter and then creating the variable from the unique name.
+	*/
+	virtual TArray<FNiagaraVariableBase> GetSimCacheRendererAttributes(UObject* UsageContext) const { return TArray<FNiagaraVariableBase>(); }
 
 #if WITH_EDITORONLY_DATA
 	/** Allows the generic class defaults version of this class to specify any dependencies/version/etc that might invalidate the compile. It should never depend on the value of specific properties.*/
diff --git a/Engine/Plugins/FX/Niagara/Source/Niagara/Private/DataInterface/NiagaraDataInterfaceUtilities.cpp b/Engine/Plugins/FX/Niagara/Source/Niagara/Private/DataInterface/NiagaraDataInterfaceUtilities.cpp
--- a/Engine/Plugins/FX/Niagara/Source/Niagara/Private/DataInterface/NiagaraDataInterfaceUtilities.cpp
+++ b/Engine/Plugins/FX/Niagara/Source/Niagara/Private/DataInterface/NiagaraDataInterfaceUtilities.cpp
@@ -168,6 +168,30 @@ void ForEachDataInterface(const FNiagaraParameterStore& ParameterStore, TFunctio
 	}
 }
 
+void ForEachDataInterface(FDataInterfaceUsageContext& UsageContext, const FNiagaraParameterStore& ParameterStore, TFunction<bool(const FDataInterfaceUsageContext&)> Action)
+{
+	const TArray<UNiagaraDataInterface*>& DataInterfaces = ParameterStore.GetDataInterfaces();
+	for (const FNiagaraVariableWithOffset& Variable : ParameterStore.ReadParameterVariables())
+	{
+		if (Variable.IsDataInterface() == false)
+		{
+			continue;
+		}
+
+		if (Variable.IsInNameSpace(FNiagaraConstants::InternalNamespaceString))
+		{
+			continue;
+		}
+
+		if (UNiagaraDataInterface* DataInterface = DataInterfaces[Variable.Offset])
+		{
+			UsageContext.Variable = Variable;
+			UsageContext.DataInterface = DataInterface;
+			Action(UsageContext);
+		}
+	}
+}
+
 } //namespace FNiagaraDataInterfaceUtilities
 
 void FNiagaraDataInterfaceUtilities::ForEachVMFunctionEquals(UNiagaraDataInterface* DataInterface, UNiagaraSystem* NiagaraSystem, TFunction<bool(const FVMExternalFunctionBindingInfo&)> Action)
@@ -368,3 +392,46 @@ void FNiagaraDataInterfaceUtilities::ForEachDataInterface(FNiagaraSystemInstance
 		}
 	}
 }
+
+void FNiagaraDataInterfaceUtilities::ForEachDataInterface(FNiagaraSystemInstance* SystemInstance, TFunction<bool(const FDataInterfaceUsageContext&)> Action)
+{
+	if (SystemInstance == nullptr)
+	{
+		return;
+	}
+
+	FNiagaraSystemSimulationPtr SystemSimulation = SystemInstance->GetSystemSimulation();
+	if (SystemSimulation == nullptr || !SystemSimulation->IsValid())
+	{
+		return;
+	}
+
+	FDataInterfaceUsageContext UsageContext;
+	UsageContext.OwnerObject = SystemSimulation->GetSystem();
+	ForEachDataInterface(UsageContext, *SystemInstance->GetOverrideParameters(), Action);
+
+	ForEachDataInterface(UsageContext, SystemSimulation->GetSpawnExecutionContext()->Parameters, Action);
+	ForEachDataInterface(UsageContext, SystemSimulation->GetUpdateExecutionContext()->Parameters, Action);
+
+	for (const TSharedRef<FNiagaraEmitterInstance, ESPMode::ThreadSafe>& Emitter : SystemInstance->GetEmitters())
+	{
+		if (Emitter->IsDisabled())
+		{
+			continue;
+		}
+		UsageContext.OwnerObject = Emitter->GetCachedEmitter().Emitter;
+
+		ForEachDataInterface(UsageContext, Emitter->GetSpawnExecutionContext().Parameters, Action);
+		ForEachDataInterface(UsageContext, Emitter->GetUpdateExecutionContext().Parameters, Action);
+		for (int32 i = 0; i < Emitter->GetEventExecutionContexts().Num(); i++)
+		{
+			ForEachDataInterface(UsageContext, Emitter->GetEventExecutionContexts()[i].Parameters, Action);
+		}
+
+		FVersionedNiagaraEmitterData* EmitterData = Emitter->GetCachedEmitterData();
+		if (EmitterData && EmitterData->SimTarget == ENiagaraSimTarget::GPUComputeSim && Emitter->GetGPUContext())
+		{
+			ForEachDataInterface(UsageContext, Emitter->GetGPUContext()->CombinedParamStore, Action);
+		}
+	}
+}
diff --git a/Engine/Plugins/FX/Niagara/Source/Niagara/Private/NiagaraSimCache.cpp b/Engine/Plugins/FX/Niagara/Source/Niagara/Private/NiagaraSimCache.cpp
--- a/Engine/Plugins/FX/Niagara/Source/Niagara/Private/NiagaraSimCache.cpp
+++ b/Engine/Plugins/FX/Niagara/Source/Niagara/Private/NiagaraSimCache.cpp
@@ -93,7 +93,7 @@ bool UNiagaraSimCache::BeginWrite(FNiagaraSimCacheCreateParameters InCreateParam
 						NameBuilder.AppendChar(TEXT('.'));
 						BoundAttribute.GetName().AppendString(NameBuilder);
 
-						CreateParameters.ExplicitCaptureAttributes.Emplace(NameBuilder.ToString());
+						CreateParameters.ExplicitCaptureAttributes.AddUnique(NameBuilder.ToString());
 					}
 				}
 			);
@@ -104,9 +104,21 @@ bool UNiagaraSimCache::BeginWrite(FNiagaraSimCacheCreateParameters InCreateParam
 			const FNiagaraParameterStore& RendererBindings = EmitterInstance->GetRendererBoundVariables();
 			for ( const FNiagaraVariableWithOffset& Variable : RendererBindings.ReadParameterVariables() )
 			{
-				CreateParameters.ExplicitCaptureAttributes.Add(Variable.GetName());
+				CreateParameters.ExplicitCaptureAttributes.AddUnique(Variable.GetName());
 			}
 		}
+
+		FNiagaraDataInterfaceUtilities::ForEachDataInterface(
+			Helper.SystemInstance,
+			[&](const FNiagaraDataInterfaceUtilities::FDataInterfaceUsageContext& UsageContext)
+			{
+				for ( const FNiagaraVariableBase& PreserveAttribute : UsageContext.DataInterface->GetSimCacheRendererAttributes(UsageContext.OwnerObject) )
+				{
+					CreateParameters.ExplicitCaptureAttributes.AddUnique(PreserveAttribute.GetName());
+				}
+				return true;
+			}
+		);
 	}
 
 	// Build new layout for system / emitters
diff --git a/Engine/Plugins/FX/Niagara/Source/Niagara/Public/NiagaraDataInterfaceUtilities.h b/Engine/Plugins/FX/Niagara/Source/Niagara/Public/NiagaraDataInterfaceUtilities.h
--- a/Engine/Plugins/FX/Niagara/Source/Niagara/Public/NiagaraDataInterfaceUtilities.h
+++ b/Engine/Plugins/FX/Niagara/Source/Niagara/Public/NiagaraDataInterfaceUtilities.h
@@ -11,6 +11,13 @@ class FNiagaraSystemInstance;
 
 namespace FNiagaraDataInterfaceUtilities
 {
+	struct FDataInterfaceUsageContext
+	{
+		UObject*				OwnerObject = nullptr;
+		FNiagaraVariableBase	Variable;
+		UNiagaraDataInterface*	DataInterface = nullptr;
+	};
+
 	// Finds all VM function calls made using the data interface that equals this one (i.e. A->Equals(B))
 	// The same function call made be made multiple times across different scripts so you may see the same function multiple times
 	// The action should return True to continue iteration or False to stop
@@ -37,4 +44,7 @@ namespace FNiagaraDataInterfaceUtilities
 	// Loops over all data interfaces inside the SystemInstance
 	// The action should return True to continue iteration or False to stop
 	void ForEachDataInterface(FNiagaraSystemInstance* SystemInstance, TFunction<bool(const FNiagaraVariableBase Variable, UNiagaraDataInterface* DataInterface)> Action);
+	// Loops over all data interfaces inside the SystemInstance
+	// The action should return True to continue iteration or False to stop
+	void ForEachDataInterface(FNiagaraSystemInstance* SystemInstance, TFunction<bool(const FDataInterfaceUsageContext&)> Action);
 }
