- Remove some unused member variables from Mesh Vertex Factory

#rnx

diff --git a/Engine/Plugins/FX/Niagara/Source/NiagaraVertexFactories/Public/NiagaraMeshVertexFactory.h b/Engine/Plugins/FX/Niagara/Source/NiagaraVertexFactories/Public/NiagaraMeshVertexFactory.h
--- a/Engine/Plugins/FX/Niagara/Source/NiagaraVertexFactories/Public/NiagaraMeshVertexFactory.h
+++ b/Engine/Plugins/FX/Niagara/Source/NiagaraVertexFactories/Public/NiagaraMeshVertexFactory.h
@@ -131,8 +131,6 @@ public:
 		, MeshIndex(-1)
 		, LODIndex(-1)
 		, bAddPrimitiveIDElement(true)
-		, InstanceVerticesCPU(nullptr)
-		, FloatDataStride(0)
 	{}
 
 	FNiagaraMeshVertexFactory()
@@ -140,8 +138,6 @@ public:
 		, MeshIndex(-1)
 		, LODIndex(-1)
 		, bAddPrimitiveIDElement(true)
-		, InstanceVerticesCPU(nullptr)
-		, FloatDataStride(0)
 	{}
 
 	/**
@@ -227,15 +223,6 @@ protected:
 
 	/** Uniform buffer with mesh particle parameters. */
 	FRHIUniformBuffer* MeshParticleUniformBuffer;
-	
-	/** Used to remember this in the case that we reuse the same vertex factory for multiple renders . */
-	FNiagaraMeshInstanceVertices* InstanceVerticesCPU;
-
-	FShaderResourceViewRHIRef ParticleDataFloatSRV;
-	uint32 FloatDataStride;
-
-	FShaderResourceViewRHIRef ParticleDataHalfSRV;
-	uint32 HalfDataStride;
 };
 
 /**
