- GPU ribbon fixes for Android
- Add missing GroupSync
- Adjust groupshared usage to workaround compiler issues

diff --git a/Engine/Plugins/FX/Niagara/Shaders/Private/Ribbons/NiagaraRibbonSortParticles.usf b/Engine/Plugins/FX/Niagara/Shaders/Private/Ribbons/NiagaraRibbonSortParticles.usf
--- a/Engine/Plugins/FX/Niagara/Shaders/Private/Ribbons/NiagaraRibbonSortParticles.usf
+++ b/Engine/Plugins/FX/Niagara/Shaders/Private/Ribbons/NiagaraRibbonSortParticles.usf
@@ -11,17 +11,11 @@ uint CountOffset;
 int MergeSortSourceBlockSize;
 int MergeSortDestinationBlockSize;
 
-struct FBubbleShortEntry
-{
-	uint ParticleID;
-	float LinkID;
-	
+groupshared uint		SortParticleIDs[THREADGROUP_SIZE + 1];
+groupshared float		SortLinkIDs[THREADGROUP_SIZE + 1];
 #if HAS_RIBBON_ID
-	FRibbonID RibbonID;
+groupshared FRibbonID	SortRibbonIDs[THREADGROUP_SIZE + 1];
 #endif
-};
-
-groupshared FBubbleShortEntry BubbleSortTemp[THREADGROUP_SIZE];
 
 [numthreads(THREADGROUP_SIZE, 1, 1)]
 void InitialSortList(uint3 DispatchThreadId : SV_DispatchThreadID, uint InGroupIndex : SV_GroupIndex, uint3 InGroupID : SV_GroupID)
@@ -35,48 +29,53 @@ void InitialSortList(uint3 DispatchThreadId : SV_DispatchThreadID, uint InGroupI
 	const int GlobalIndexOffset = DispatchThreadId.x;
 	const int LocalIndexOffset = InGroupIndex;
 
-	FBubbleShortEntry NewEntry;
-	NewEntry.ParticleID = min(GlobalIndexOffset, int(TotalNumParticles - 1));
-	NewEntry.LinkID = GetRibbonLinkID(NewEntry.ParticleID);	
-#if HAS_RIBBON_ID
-	NewEntry.RibbonID = GetRibbonID(NewEntry.ParticleID);
-#endif
-	
-	BubbleSortTemp[LocalIndexOffset] = NewEntry;	
-	
-	int SortGroup = 1;
+	// Make local sort information
+	{
+		const uint LocalParticleID			= min(GlobalIndexOffset, int(TotalNumParticles - 1));
+		SortParticleIDs[LocalIndexOffset]	= LocalParticleID;
+		SortLinkIDs[LocalIndexOffset]		= GetRibbonLinkID(LocalParticleID);
+	#if HAS_RIBBON_ID
+		SortRibbonIDs[LocalIndexOffset]		= GetRibbonID(LocalParticleID);
+	#endif
+	}
+	GroupMemoryBarrierWithGroupSync();
 	
-	for (int Index = 0; Index < THREADGROUP_SIZE; Index++)
+	for (int SortGroup=0; SortGroup < THREADGROUP_SIZE; SortGroup++)
 	{
-		// Switch the offset and sync all threads
-		SortGroup ^= 0x1;		
-		GroupMemoryBarrierWithGroupSync();
-
-		if (LocalIndexOffset % 2 == SortGroup && LocalIndexOffset + 1 < TotalNumParticlesInSet)
+		bool bShouldSwap = false;
+		if ( ((LocalIndexOffset & 1) == (SortGroup & 1)) && (LocalIndexOffset + 1 < TotalNumParticlesInSet) )
 		{
-			const FBubbleShortEntry ThisEntry = BubbleSortTemp[LocalIndexOffset + 0];
-			const FBubbleShortEntry NextEntry = BubbleSortTemp[LocalIndexOffset + 1];		
-
 #if HAS_RIBBON_ID			
-			const int RibbonIDComparisonResult = CompareRibbonIDs(ThisEntry.RibbonID, NextEntry.RibbonID);
-			const bool bShouldSwap = RibbonIDComparisonResult > 0 || (RibbonIDComparisonResult == 0 && ThisEntry.LinkID > NextEntry.LinkID);
+			const int RibbonIDComparisonResult = CompareRibbonIDs(SortRibbonIDs[LocalIndexOffset + 0], SortRibbonIDs[LocalIndexOffset + 1]);
+			bShouldSwap = RibbonIDComparisonResult > 0 || (RibbonIDComparisonResult == 0 && SortLinkIDs[LocalIndexOffset + 0] > SortLinkIDs[LocalIndexOffset + 1]);
 #else
-			const bool bShouldSwap = ThisEntry.LinkID > NextEntry.LinkID;
+			bShouldSwap = SortLinkIDs[LocalIndexOffset + 0] > SortLinkIDs[LocalIndexOffset + 1];
 #endif
-			
-			if (bShouldSwap)
-			{
-				BubbleSortTemp[LocalIndexOffset + 0] = NextEntry;
-				BubbleSortTemp[LocalIndexOffset + 1] = ThisEntry;
-			}
 		}
-	}
 
-	GroupMemoryBarrierWithGroupSync();
+		GroupMemoryBarrierWithGroupSync();
+		if (bShouldSwap)
+		{
+			uint TempParticleID = SortParticleIDs[LocalIndexOffset + 0];
+			SortParticleIDs[LocalIndexOffset + 0] = SortParticleIDs[LocalIndexOffset + 1];
+			SortParticleIDs[LocalIndexOffset + 1] = TempParticleID;
+
+			float TempLinkID = SortLinkIDs[LocalIndexOffset + 0];
+			SortLinkIDs[LocalIndexOffset + 0] = SortLinkIDs[LocalIndexOffset + 1];
+			SortLinkIDs[LocalIndexOffset + 1] = TempLinkID;
+
+		#if HAS_RIBBON_ID
+			FRibbonID TempRibbonID = SortRibbonIDs[LocalIndexOffset + 0];
+			SortRibbonIDs[LocalIndexOffset + 0] = SortRibbonIDs[LocalIndexOffset + 1];
+			SortRibbonIDs[LocalIndexOffset + 1] = TempRibbonID;
+		#endif
+		}
+		GroupMemoryBarrierWithGroupSync();
+	}
 
 	if (LocalIndexOffset < TotalNumParticlesInSet)
 	{
-		DestinationSortedIndices[GlobalIndexOffset] = BubbleSortTemp[LocalIndexOffset].ParticleID;
+		DestinationSortedIndices[GlobalIndexOffset] = SortParticleIDs[LocalIndexOffset];
 	}
 }
 
diff --git a/Engine/Plugins/FX/Niagara/Source/Niagara/Private/NiagaraRendererRibbons.cpp b/Engine/Plugins/FX/Niagara/Source/Niagara/Private/NiagaraRendererRibbons.cpp
--- a/Engine/Plugins/FX/Niagara/Source/Niagara/Private/NiagaraRendererRibbons.cpp
+++ b/Engine/Plugins/FX/Niagara/Source/Niagara/Private/NiagaraRendererRibbons.cpp
@@ -244,6 +244,7 @@ struct FRWIndexBuffer final : FIndexBuffer
 		CreateInfo.ResourceArray = InResourceArray;
 		IndexBufferRHI = RHICreateIndexBuffer(BytesPerElement, NumBytes, BUF_UnorderedAccess | AdditionalUsage, InResourceState, CreateInfo);
 		UAV = RHICreateUnorderedAccessView(IndexBufferRHI, UE_PIXELFORMAT_TO_UINT8(Format));
+		UE_LOG(LogTemp, Warning, TEXT("Allocated UAV(%p)"), UAV.GetReference());
 	}
 
 	virtual void ReleaseRHI() override
@@ -2027,10 +2028,7 @@ void FNiagaraRendererRibbons::InitializeViewIndexBuffersGPU(FRHICommandListImmed
 		Params.GNiagaraRibbonTessellationMinDisplacementError = GNiagaraRibbonTessellationMinDisplacementError;
 
 		RHICmdList.Transition(FRHITransitionInfo(RenderingViewResources->IndirectDrawBuffer.UAV, ERHIAccess::SRVMask | ERHIAccess::IndirectArgs, ERHIAccess::UAVCompute));
-		SetComputePipelineState(RHICmdList, ComputeShader.GetComputeShader());
-		SetShaderParameters(RHICmdList, ComputeShader, ComputeShader.GetComputeShader(), Params);
-		DispatchComputeShader(RHICmdList, ComputeShader, 1, 1, 1);
-		UnsetShaderUAVs(RHICmdList, ComputeShader, ComputeShader.GetComputeShader());
+		FComputeShaderUtils::Dispatch(RHICmdList, ComputeShader, Params, FIntVector(1, 1, 1));
 		RHICmdList.Transition(FRHITransitionInfo(RenderingViewResources->IndirectDrawBuffer.UAV, ERHIAccess::UAVCompute, ERHIAccess::SRVMask | ERHIAccess::IndirectArgs));
 	}
 	
@@ -2088,10 +2086,7 @@ void FNiagaraRendererRibbons::InitializeViewIndexBuffersGPU(FRHICommandListImmed
 		Params.SubSegmentBitMask = RenderingViewResources->IndexGenerationSettings.SubSegmentBitMask;
 		
 		RHICmdList.Transition(FRHITransitionInfo(RenderingViewResources->IndexBuffer.UAV, ERHIAccess::VertexOrIndexBuffer, ERHIAccess::UAVCompute));
-		SetComputePipelineState(RHICmdList, ComputeShader.GetComputeShader());
-		SetShaderParameters(RHICmdList, ComputeShader, ComputeShader.GetComputeShader(), Params);
-		DispatchIndirectComputeShader(RHICmdList, ComputeShader.GetShader(), RenderingViewResources->IndirectDrawBuffer.Buffer, IndirectDispatchArgsOffset);
-		UnsetShaderUAVs(RHICmdList, ComputeShader, ComputeShader.GetComputeShader());
+		FComputeShaderUtils::DispatchIndirect(RHICmdList, ComputeShader, Params, RenderingViewResources->IndirectDrawBuffer.Buffer, IndirectDispatchArgsOffset);
 		RHICmdList.Transition(FRHITransitionInfo(RenderingViewResources->IndexBuffer.UAV, ERHIAccess::UAVCompute, ERHIAccess::VertexOrIndexBuffer));
 	}
 }
