- Remove debug logging

#rnx

diff --git a/Engine/Plugins/FX/Niagara/Source/Niagara/Private/NiagaraRendererRibbons.cpp b/Engine/Plugins/FX/Niagara/Source/Niagara/Private/NiagaraRendererRibbons.cpp
--- a/Engine/Plugins/FX/Niagara/Source/Niagara/Private/NiagaraRendererRibbons.cpp
+++ b/Engine/Plugins/FX/Niagara/Source/Niagara/Private/NiagaraRendererRibbons.cpp
@@ -244,7 +244,6 @@ struct FRWIndexBuffer final : FIndexBuffer
 		CreateInfo.ResourceArray = InResourceArray;
 		IndexBufferRHI = RHICreateIndexBuffer(BytesPerElement, NumBytes, BUF_UnorderedAccess | AdditionalUsage, InResourceState, CreateInfo);
 		UAV = RHICreateUnorderedAccessView(IndexBufferRHI, UE_PIXELFORMAT_TO_UINT8(Format));
-		UE_LOG(LogTemp, Warning, TEXT("Allocated UAV(%p)"), UAV.GetReference());
 	}
 
 	virtual void ReleaseRHI() override
