Support position types for material parameter bindings

diff --git a/Engine/Plugins/FX/Niagara/Source/Niagara/Private/NiagaraRenderer.cpp b/Engine/Plugins/FX/Niagara/Source/Niagara/Private/NiagaraRenderer.cpp
--- a/Engine/Plugins/FX/Niagara/Source/Niagara/Private/NiagaraRenderer.cpp
+++ b/Engine/Plugins/FX/Niagara/Source/Niagara/Private/NiagaraRenderer.cpp
@@ -536,6 +536,17 @@ void FNiagaraRenderer::ProcessMaterialParameterBindings(const FNiagaraRendererMa
 							InEmitter->GetBoundRendererValue_GT(Binding.GetParamMapBindableVariable(), Binding.NiagaraChildVariable, &Var);
 							MatDyn->SetVectorParameterValue(Binding.MaterialParameterName, Var);
 						}
+						else if (Binding.GetParamMapBindableVariable().GetType() == FNiagaraTypeDefinition::GetPositionDef() ||
+							(Binding.GetParamMapBindableVariable().GetType().IsDataInterface() && Binding.NiagaraChildVariable.GetType() == FNiagaraTypeDefinition::GetPositionDef()))
+						{
+							FNiagaraPosition Var(ForceInitToZero);
+							InEmitter->GetBoundRendererValue_GT(Binding.GetParamMapBindableVariable(), Binding.NiagaraChildVariable, &Var);
+							FNiagaraLWCConverter LwcConverter = SystemInstance->GetLWCConverter(InEmitter->GetCachedEmitterData() ? InEmitter->GetCachedEmitterData()->bLocalSpace : false);
+							FVector WorldPos = LwcConverter.ConvertSimulationPositionToWorld(Var);
+
+							//TODO: should we use SetDoubleVectorParamValue() instead to prevent accuracy loss? 
+							MatDyn->SetVectorParameterValue(Binding.MaterialParameterName, WorldPos);
+						}
 						else if (Binding.GetParamMapBindableVariable().GetType() == FNiagaraTypeDefinition::GetVec2Def() ||
 							(Binding.GetParamMapBindableVariable().GetType().IsDataInterface() && Binding.NiagaraChildVariable.GetType() == FNiagaraTypeDefinition::GetVec2Def()))
 						{
diff --git a/Engine/Plugins/FX/Niagara/Source/NiagaraEditor/Private/Customizations/NiagaraTypeCustomizations.cpp b/Engine/Plugins/FX/Niagara/Source/NiagaraEditor/Private/Customizations/NiagaraTypeCustomizations.cpp
--- a/Engine/Plugins/FX/Niagara/Source/NiagaraEditor/Private/Customizations/NiagaraTypeCustomizations.cpp
+++ b/Engine/Plugins/FX/Niagara/Source/NiagaraEditor/Private/Customizations/NiagaraTypeCustomizations.cpp
@@ -819,6 +819,7 @@ bool FNiagaraMaterialAttributeBindingCustomization::IsCompatibleNiagaraVariable(
 			InVar.GetType() == FNiagaraTypeDefinition::GetColorDef() ||
 			InVar.GetType() == FNiagaraTypeDefinition::GetVec2Def() ||
 			InVar.GetType() == FNiagaraTypeDefinition::GetVec3Def() ||
+			InVar.GetType() == FNiagaraTypeDefinition::GetPositionDef() ||
 			InVar.GetType() == FNiagaraTypeDefinition::GetUObjectDef() ||
 			InVar.GetType() == FNiagaraTypeDefinition::GetUTextureDef() ||
 			InVar.GetType() == FNiagaraTypeDefinition::GetUTextureRenderTargetDef() )
