NiagaraWorldManager - handle the newly added possibility that UWorld::InitWorld followed by CleanupWorld can be called multiple times and setup/teardown the world each time.
#rnx

diff --git a/Engine/Plugins/FX/Niagara/Source/Niagara/Private/NiagaraWorldManager.cpp b/Engine/Plugins/FX/Niagara/Source/Niagara/Private/NiagaraWorldManager.cpp
--- a/Engine/Plugins/FX/Niagara/Source/Niagara/Private/NiagaraWorldManager.cpp
+++ b/Engine/Plugins/FX/Niagara/Source/Niagara/Private/NiagaraWorldManager.cpp
@@ -299,13 +299,23 @@ void FNiagaraWorldManager::SetScalabilityCullingMode(ENiagaraScalabilityCullingM
 FNiagaraWorldManager::FNiagaraWorldManager()
 	: World(nullptr)
 	, ActiveNiagaraTickGroup(-1)
-	, CachedEffectsQuality(INDEX_NONE)
 	, bAppHasFocus(true)
 {
 }
 
 void FNiagaraWorldManager::Init(UWorld* InWorld)
 {
+	// Reset variables that are not reset during OnWorldCleanup; Init/OnWorldCleanup can be called multiple times
+	ComponentPool = nullptr; // Discard the existing ComponentPool
+	bPoolIsPrimed = false;
+	ActiveNiagaraTickGroup = -1;
+	bAppHasFocus = true;
+	bIsTearingDown = false;
+	WorldLoopTime = 0.0f;
+	RequestedDebugPlaybackMode = ENiagaraDebugPlaybackMode::Play;
+	DebugPlaybackMode = ENiagaraDebugPlaybackMode::Play;
+	DebugPlaybackRate = 1.0f;
+
 	World = InWorld;
 	for (int32 TickGroup = 0; TickGroup < NiagaraNumTickGroups; ++TickGroup)
 	{
@@ -594,6 +604,7 @@ void FNiagaraWorldManager::OnComputeDispatchInterfaceDestroyed_Internal(FNiagara
 
 void FNiagaraWorldManager::OnWorldCleanup(bool bSessionEnded, bool bCleanupResources)
 {
+	DeferredMethods.ExecuteAndClear();
 	ComponentPool->Cleanup(World);
 
 	for (int TG = 0; TG < NiagaraNumTickGroups; ++TG)
@@ -617,6 +628,7 @@ void FNiagaraWorldManager::OnWorldCleanup(bool bSessionEnded, bool bCleanupResou
 	ScalabilityManagers.Empty();
 
 	CullProxyMap.Empty();
+	DIGeneratedData.Empty();
 }
 
 void FNiagaraWorldManager::OnPostWorldCleanup(bool bSessionEnded, bool bCleanupResources)
@@ -691,9 +703,11 @@ void FNiagaraWorldManager::RefreshOwnerAllowsScalability()
 
 void FNiagaraWorldManager::OnWorldInit(UWorld* World, const UWorld::InitializationValues IVS)
 {
-	check(WorldManagers.Find(World) == nullptr);
-	FNiagaraWorldManager*& NewManager = WorldManagers.Add(World);
-	NewManager = new FNiagaraWorldManager();
+	FNiagaraWorldManager*& NewManager = WorldManagers.FindOrAdd(World);
+	if (!NewManager)
+	{
+		NewManager = new FNiagaraWorldManager();
+	}
 	NewManager->Init(World);
 }
 
diff --git a/Engine/Plugins/FX/Niagara/Source/Niagara/Public/NiagaraWorldManager.h b/Engine/Plugins/FX/Niagara/Source/Niagara/Public/NiagaraWorldManager.h
--- a/Engine/Plugins/FX/Niagara/Source/Niagara/Public/NiagaraWorldManager.h
+++ b/Engine/Plugins/FX/Niagara/Source/Niagara/Public/NiagaraWorldManager.h
@@ -279,7 +279,7 @@ private:
 
 	static TMap<class UWorld*, class FNiagaraWorldManager*> WorldManagers;
 
-	UWorld* World;
+	UWorld* World = nullptr;
 
 	int ActiveNiagaraTickGroup;
 
@@ -293,8 +293,6 @@ private:
 
 	TArray<FNiagaraSystemSimulationRef> SimulationsWithEndOfFrameWait;
 
-	int32 CachedEffectsQuality;
-
 	TArray<FNiagaraCachedViewInfo, TInlineAllocator<8> > CachedViewInfo;
 
 	UNiagaraComponentPool* ComponentPool;
