Moved deprecation flag into versioning settings, slightly changing the messaging and stack warnings to take this flag into consideration

diff --git a/Engine/Plugins/FX/Niagara/Source/Niagara/Classes/NiagaraEmitter.h b/Engine/Plugins/FX/Niagara/Source/Niagara/Classes/NiagaraEmitter.h
--- a/Engine/Plugins/FX/Niagara/Source/Niagara/Classes/NiagaraEmitter.h
+++ b/Engine/Plugins/FX/Niagara/Source/Niagara/Classes/NiagaraEmitter.h
@@ -259,6 +259,14 @@ struct NIAGARA_API FVersionedNiagaraEmitterData
 	UPROPERTY()
 	FFilePath ScriptAsset;
 
+	/* If this emitter is no longer meant to be used, this option should be set.*/
+	UPROPERTY()
+	bool bDeprecated = false;
+
+	/* Message to display when the script is deprecated. */
+	UPROPERTY()
+	FText DeprecationMessage;
+
 	/** Toggles whether or not the particles within this emitter are relative to the emitter origin or in global space. */ 
 	UPROPERTY(EditAnywhere, BlueprintReadWrite, Category = "Emitter")
 	bool bLocalSpace = false;
diff --git a/Engine/Plugins/FX/Niagara/Source/Niagara/Classes/NiagaraScript.h b/Engine/Plugins/FX/Niagara/Source/Niagara/Classes/NiagaraScript.h
--- a/Engine/Plugins/FX/Niagara/Source/Niagara/Classes/NiagaraScript.h
+++ b/Engine/Plugins/FX/Niagara/Source/Niagara/Classes/NiagaraScript.h
@@ -626,15 +626,15 @@ public:
 	TArray<FNiagaraModuleDependency> RequiredDependencies;
 	
 	/* If this script is no longer meant to be used, this option should be set.*/
-	UPROPERTY(AssetRegistrySearchable, EditAnywhere, Category = Script)
-	uint32 bDeprecated : 1;
+	UPROPERTY()
+	bool bDeprecated;
 	
 	/* Message to display when the script is deprecated. */
-	UPROPERTY(EditAnywhere, Category = Script, meta = (EditCondition = "bDeprecated", MultiLine = true))
+	UPROPERTY()
 	FText DeprecationMessage;
 
 	/* Which script to use if this is deprecated.*/
-	UPROPERTY(EditAnywhere, Category = Script, meta = (EditCondition = "bDeprecated"))
+	UPROPERTY(EditAnywhere, Category = Script)
 	TObjectPtr<UNiagaraScript> DeprecationRecommendation;
 
 	/* Custom logic to convert the contents of an existing script assignment to this script.*/
diff --git a/Engine/Plugins/FX/Niagara/Source/Niagara/Classes/NiagaraVersionedObject.h b/Engine/Plugins/FX/Niagara/Source/Niagara/Classes/NiagaraVersionedObject.h
--- a/Engine/Plugins/FX/Niagara/Source/Niagara/Classes/NiagaraVersionedObject.h
+++ b/Engine/Plugins/FX/Niagara/Source/Niagara/Classes/NiagaraVersionedObject.h
@@ -11,6 +11,8 @@ struct NIAGARA_API FNiagaraVersionDataAccessor
 
 	virtual FNiagaraAssetVersion& GetObjectVersion();
 	virtual FText& GetVersionChangeDescription();
+	virtual bool& IsDeprecated();
+	virtual FText& GetDeprecationMessage();
 	virtual ENiagaraPythonUpdateScriptReference& GetUpdateScriptExecutionType();
 	virtual FString& GetPythonUpdateScript();
 	virtual FFilePath& GetScriptAsset();
diff --git a/Engine/Plugins/FX/Niagara/Source/Niagara/Private/NiagaraEmitter.cpp b/Engine/Plugins/FX/Niagara/Source/Niagara/Private/NiagaraEmitter.cpp
--- a/Engine/Plugins/FX/Niagara/Source/Niagara/Private/NiagaraEmitter.cpp
+++ b/Engine/Plugins/FX/Niagara/Source/Niagara/Private/NiagaraEmitter.cpp
@@ -2269,6 +2269,8 @@ struct FNiagaraEmitterVersionDataAccessor final : FNiagaraVersionDataAccessor
 
 	virtual FNiagaraAssetVersion& GetObjectVersion() override { return EmitterData->Version; }
 	virtual FText& GetVersionChangeDescription() override { return EmitterData->VersionChangeDescription; }
+	virtual FText& GetDeprecationMessage() override { return EmitterData->DeprecationMessage; }
+	virtual bool& IsDeprecated() override { return EmitterData->bDeprecated; }
 	virtual ENiagaraPythonUpdateScriptReference& GetUpdateScriptExecutionType() override { return EmitterData->UpdateScriptExecution; }
 	virtual FString& GetPythonUpdateScript() override { return EmitterData->PythonUpdateScript; }
 	virtual FFilePath& GetScriptAsset() override { return EmitterData->ScriptAsset; }
diff --git a/Engine/Plugins/FX/Niagara/Source/Niagara/Private/NiagaraScript.cpp b/Engine/Plugins/FX/Niagara/Source/Niagara/Private/NiagaraScript.cpp
--- a/Engine/Plugins/FX/Niagara/Source/Niagara/Private/NiagaraScript.cpp
+++ b/Engine/Plugins/FX/Niagara/Source/Niagara/Private/NiagaraScript.cpp
@@ -568,6 +568,8 @@ struct FNiagaraScriptVersionDataAccessor final : FNiagaraVersionDataAccessor
 
 	virtual FNiagaraAssetVersion& GetObjectVersion() override { return ScriptData->Version; }
 	virtual FText& GetVersionChangeDescription() override { return ScriptData->VersionChangeDescription; }
+	virtual bool& IsDeprecated() override { return ScriptData->bDeprecated; }
+	virtual FText& GetDeprecationMessage() override { return ScriptData->DeprecationMessage; }
 	virtual ENiagaraPythonUpdateScriptReference& GetUpdateScriptExecutionType() override { return ScriptData->UpdateScriptExecution; }
 	virtual FString& GetPythonUpdateScript() override { return ScriptData->PythonUpdateScript; }
 	virtual FFilePath& GetScriptAsset() override { return ScriptData->ScriptAsset; }
diff --git a/Engine/Plugins/FX/Niagara/Source/Niagara/Private/NiagaraVersionedObject.cpp b/Engine/Plugins/FX/Niagara/Source/Niagara/Private/NiagaraVersionedObject.cpp
--- a/Engine/Plugins/FX/Niagara/Source/Niagara/Private/NiagaraVersionedObject.cpp
+++ b/Engine/Plugins/FX/Niagara/Source/Niagara/Private/NiagaraVersionedObject.cpp
@@ -12,15 +12,17 @@ namespace NiagaraVersionedObject
 	{
 		virtual FNiagaraAssetVersion& GetObjectVersion() override { return Version; }
 		virtual FText& GetVersionChangeDescription() override { return Desc; }
-		virtual ENiagaraPythonUpdateScriptReference& GetUpdateScriptExecutionType() override {return Type;}
-		virtual FString& GetPythonUpdateScript() override {return Script;}
-		virtual FFilePath& GetScriptAsset() override {return Path;};
+		virtual ENiagaraPythonUpdateScriptReference& GetUpdateScriptExecutionType() override {return Type; }
+		virtual FString& GetPythonUpdateScript() override {return Script; }
+		virtual FFilePath& GetScriptAsset() override {return Path; }
+		virtual bool& IsDeprecated() override { return bDeprecated; }
 
 		FNiagaraAssetVersion Version;
 		FText Desc;
 		ENiagaraPythonUpdateScriptReference Type = ENiagaraPythonUpdateScriptReference::None;
 		FString Script;
 		FFilePath Path;
+		bool bDeprecated;
 	};
 
 	static DummyObject Dummy;
@@ -38,6 +40,18 @@ FText& FNiagaraVersionDataAccessor::GetVersionChangeDescription()
 	return NiagaraVersionedObject::Dummy.GetVersionChangeDescription();
 }
 
+bool& FNiagaraVersionDataAccessor::IsDeprecated()
+{
+	ensure(false);
+	return NiagaraVersionedObject::Dummy.IsDeprecated();
+}
+
+FText& FNiagaraVersionDataAccessor::GetDeprecationMessage()
+{
+	ensure(false);
+	return NiagaraVersionedObject::Dummy.GetVersionChangeDescription();
+}
+
 ENiagaraPythonUpdateScriptReference& FNiagaraVersionDataAccessor::GetUpdateScriptExecutionType()
 {
 	ensure(false);
diff --git a/Engine/Plugins/FX/Niagara/Source/NiagaraEditor/Private/Commandlets/NiagaraDumpModuleInfoCommandlet.cpp b/Engine/Plugins/FX/Niagara/Source/NiagaraEditor/Private/Commandlets/NiagaraDumpModuleInfoCommandlet.cpp
--- a/Engine/Plugins/FX/Niagara/Source/NiagaraEditor/Private/Commandlets/NiagaraDumpModuleInfoCommandlet.cpp
+++ b/Engine/Plugins/FX/Niagara/Source/NiagaraEditor/Private/Commandlets/NiagaraDumpModuleInfoCommandlet.cpp
@@ -250,8 +250,12 @@ void UNiagaraDumpModuleInfoCommandlet::ProcessNiagaraScripts()
 			{
 				UNiagaraGraph* FunctionGraph = ScriptSrc->NodeGraph;
 				FText VersionDesc;
+				bool bDeprecated = false;
 				if (VersionDataAccessor)
+				{
 					VersionDesc = VersionDataAccessor->GetVersionChangeDescription();
+					bDeprecated = VersionDataAccessor->IsDeprecated();
+				}
 
 				if (FunctionGraph)
 				{
@@ -281,6 +285,7 @@ void UNiagaraDumpModuleInfoCommandlet::ProcessNiagaraScripts()
 					OutputStream->Log(TEXT("<Info>"));
 					OutputStream->Logf(TEXT("\t<VersionNumber>%d . %d</VersionNumber>"), Version.MajorVersion, Version.MinorVersion);
 					OutputStream->Logf(TEXT("\t<VersionDesc>%s</VersionDesc>"), *VersionDesc.ToString());
+					OutputStream->Logf(TEXT("\t<Deprecated>%s</Deprecated>"), bDeprecated ? TEXT("true") : TEXT("false"));
 					OutputStream->Logf(TEXT("\t<Path>%s</Path>"), *GetPathNameSafe(FunctionGraph));
 
 					for (UNiagaraScriptVariable*& ScriptVar : ScriptVars)
diff --git a/Engine/Plugins/FX/Niagara/Source/NiagaraEditor/Private/NiagaraVersionMetaData.h b/Engine/Plugins/FX/Niagara/Source/NiagaraEditor/Private/NiagaraVersionMetaData.h
--- a/Engine/Plugins/FX/Niagara/Source/NiagaraEditor/Private/NiagaraVersionMetaData.h
+++ b/Engine/Plugins/FX/Niagara/Source/NiagaraEditor/Private/NiagaraVersionMetaData.h
@@ -29,6 +29,14 @@ public:
 	UPROPERTY(EditAnywhere, Category="Version Details", meta=(EditCondition="!bIsExposedVersion"))
 	bool bIsVisibleInVersionSelector = true;
 
+	/* True if this asset version is deprecated and should no longer be used.*/
+	UPROPERTY(EditAnywhere, Category="Version Details")
+	bool bDeprecated = false;
+
+	/* Message to display when the asset is used in an emitter. */
+	UPROPERTY(EditAnywhere, Category="Version Details", meta = (EditCondition = "bDeprecated", MultiLine = true))
+	FText DeprecationMessage;
+
 	/** Internal version guid, mainly useful for debugging version conflicts. */
 	UPROPERTY(VisibleAnywhere, AdvancedDisplay, Category="Version Details")
 	FGuid VersionGuid;
diff --git a/Engine/Plugins/FX/Niagara/Source/NiagaraEditor/Private/ViewModels/Stack/NiagaraStackGraphUtilities.cpp b/Engine/Plugins/FX/Niagara/Source/NiagaraEditor/Private/ViewModels/Stack/NiagaraStackGraphUtilities.cpp
--- a/Engine/Plugins/FX/Niagara/Source/NiagaraEditor/Private/ViewModels/Stack/NiagaraStackGraphUtilities.cpp
+++ b/Engine/Plugins/FX/Niagara/Source/NiagaraEditor/Private/ViewModels/Stack/NiagaraStackGraphUtilities.cpp
@@ -518,19 +518,22 @@ void FNiagaraStackGraphUtilities::CheckForDeprecatedScriptVersion(UNiagaraNodeFu
 	if (InputFunctionCallNode->FunctionScript && InputFunctionCallNode->FunctionScript->IsVersioningEnabled() && VersionUpgradeFix.IsBound())
 	{
 		FNiagaraAssetVersion ExposedVersion = InputFunctionCallNode->FunctionScript->GetExposedVersion();
-		FNiagaraAssetVersion ReferencedVersion = InputFunctionCallNode->FunctionScript->GetScriptData(InputFunctionCallNode->SelectedScriptVersion)->Version;
+		FVersionedNiagaraScriptData* ScriptData = InputFunctionCallNode->FunctionScript->GetScriptData(InputFunctionCallNode->SelectedScriptVersion);
+		FNiagaraAssetVersion ReferencedVersion = ScriptData->Version;
 		if (ReferencedVersion.MajorVersion < ExposedVersion.MajorVersion)
 		{
 			TArray<UNiagaraStackEntry::FStackIssueFix> Fixes;
 			Fixes.Add(UNiagaraStackEntry::FStackIssueFix(LOCTEXT("UpgradeVersionFix", "Upgrade to newest version."), VersionUpgradeFix));
 			//TODO MV: add "fix all" and "copy and fix" actions
-						
+			
 			FText ChangelistDescriptions = GetChangelistText(InputFunctionCallNode->FunctionScript, ReferencedVersion, ExposedVersion);
-			FText LongDescription = FText::Format(LOCTEXT("DeprecatedVersionFormat", "This script has a newer version available.\nYou can upgrade now, but major version upgrades can sometimes come with breaking changes! So check that everything is still working as expected afterwards.{0}"),
+			FText DeprecationDescription = ScriptData->bDeprecated && !ScriptData->DeprecationMessage.IsEmpty() ? FText::Format(LOCTEXT("DeprecatedScriptVersionMessage", "\n\nDeprecation message:\n{0}"), ScriptData->DeprecationMessage) : FText();
+			FText LongDescription = FText::Format(LOCTEXT("DeprecatedVersionFormat", "This script has a newer version available.\nYou can upgrade now, but major version upgrades can sometimes come with breaking changes! So check that everything is still working as expected afterwards.{0}{1}"),
+				DeprecationDescription,
 				ChangelistDescriptions.IsEmpty() ? FText() : FText::Format(LOCTEXT("DeprecatedVersionFormatChanges", "\n\nVersion change description:\n{0}"), ChangelistDescriptions));
 			UNiagaraStackEntry::FStackIssue UpgradeVersion(
-                EStackIssueSeverity::Warning,
-                FText::Format(LOCTEXT("DeprecatedScriptVersionSummaryFormat", "Deprecated script version: {0}.{1} -> {2}.{3}"),
+                ScriptData->bDeprecated ? EStackIssueSeverity::Warning : EStackIssueSeverity::Info,
+                FText::Format(LOCTEXT("DeprecatedScriptVersionSummaryFormat", "Upgrade script version: {0}.{1} -> {2}.{3}"),
                 	FText::AsNumber(ReferencedVersion.MajorVersion), FText::AsNumber(ReferencedVersion.MinorVersion), FText::AsNumber(ExposedVersion.MajorVersion), FText::AsNumber(ExposedVersion.MinorVersion)),
                 LongDescription,
                 StackEditorDataKey,
@@ -621,11 +624,13 @@ void FNiagaraStackGraphUtilities::CheckForDeprecatedEmitterVersion(TSharedPtr<FN
 			Fixes.Add(UNiagaraStackEntry::FStackIssueFix(LOCTEXT("UpgradeNewestParentFix", "Upgrade to newest parent version."), VersionUpgradeFix));
 			
 			FText ChangelistDescriptions = GetChangelistText(Emitter, ReferencedVersion, ExposedVersion);
-			FText LongDescription = FText::Format(LOCTEXT("EmitterDeprecatedVersionFormat", "This emitter has a newer parent version available.\nYou can upgrade now, but major version upgrades can sometimes come with breaking changes! So check that everything is still working as expected afterwards.{0}"),
+			FText DeprecationDescription = EmitterData->bDeprecated && !EmitterData->DeprecationMessage.IsEmpty() ? FText::Format(LOCTEXT("DeprecatedEmitterVersionMessage", "\n\nDeprecation message:\n{0}"), EmitterData->DeprecationMessage) : FText();
+			FText LongDescription = FText::Format(LOCTEXT("EmitterDeprecatedVersionFormat", "This emitter has a newer parent version available.\nYou can upgrade now, but major version upgrades can sometimes come with breaking changes! So check that everything is still working as expected afterwards.{0}{1}"),
+				DeprecationDescription,
 				ChangelistDescriptions.IsEmpty() ? FText() : FText::Format(LOCTEXT("DeprecatedVersionFormatChanges", "\n\nVersion change description:\n{0}"), ChangelistDescriptions));
 			UNiagaraStackEntry::FStackIssue UpgradeVersion(
-                EStackIssueSeverity::Warning,
-                FText::Format(LOCTEXT("DeprecatedVersionSummaryFormat", "Deprecated emitter version: {0}.{1} -> {2}.{3}"),
+                EmitterData->bDeprecated ? EStackIssueSeverity::Warning : EStackIssueSeverity::Info,
+                FText::Format(LOCTEXT("DeprecatedVersionSummaryFormat", "Upgrade emitter version: {0}.{1} -> {2}.{3}"),
                 	FText::AsNumber(ReferencedVersion.MajorVersion), FText::AsNumber(ReferencedVersion.MinorVersion), FText::AsNumber(ExposedVersion.MajorVersion), FText::AsNumber(ExposedVersion.MinorVersion)),
                 LongDescription,
                 StackEditorDataKey,
diff --git a/Engine/Plugins/FX/Niagara/Source/NiagaraEditor/Private/Widgets/SNiagaraVersionWidget.cpp b/Engine/Plugins/FX/Niagara/Source/NiagaraEditor/Private/Widgets/SNiagaraVersionWidget.cpp
--- a/Engine/Plugins/FX/Niagara/Source/NiagaraEditor/Private/Widgets/SNiagaraVersionWidget.cpp
+++ b/Engine/Plugins/FX/Niagara/Source/NiagaraEditor/Private/Widgets/SNiagaraVersionWidget.cpp
@@ -217,6 +217,8 @@ void SNiagaraVersionWidget::NotifyPostChange(const FPropertyChangedEvent& Proper
 	check(ObjectData);
 	
 	ObjectData->GetVersionChangeDescription() = VersionMetadata->ChangeDescription;
+	ObjectData->IsDeprecated() = VersionMetadata->bDeprecated;
+	ObjectData->GetDeprecationMessage() = VersionMetadata->DeprecationMessage;
 	ObjectData->GetUpdateScriptExecutionType() = VersionMetadata->UpdateScriptExecution;
 	ObjectData->GetScriptAsset() = VersionMetadata->ScriptAsset;
 	ObjectData->GetPythonUpdateScript() = VersionMetadata->PythonUpdateScript;
@@ -437,6 +439,8 @@ void SNiagaraVersionWidget::VersionInListSelected(FNiagaraAssetVersion InSelecte
 	
 	VersionMetadata->VersionGuid = InSelectedVersion.VersionGuid;
 	VersionMetadata->ChangeDescription = ObjectData->GetVersionChangeDescription();
+	VersionMetadata->bDeprecated = ObjectData->IsDeprecated();
+	VersionMetadata->DeprecationMessage = ObjectData->GetDeprecationMessage();
 	VersionMetadata->bIsExposedVersion = InSelectedVersion == VersionedObject->GetExposedVersion();
 	VersionMetadata->bIsVisibleInVersionSelector = InSelectedVersion.bIsVisibleInVersionSelector;
 
diff --git a/Engine/Plugins/FX/Niagara/Source/NiagaraEditorWidgets/Private/SNiagaraOverviewGraphTitleBar.cpp b/Engine/Plugins/FX/Niagara/Source/NiagaraEditorWidgets/Private/SNiagaraOverviewGraphTitleBar.cpp
--- a/Engine/Plugins/FX/Niagara/Source/NiagaraEditorWidgets/Private/SNiagaraOverviewGraphTitleBar.cpp
+++ b/Engine/Plugins/FX/Niagara/Source/NiagaraEditorWidgets/Private/SNiagaraOverviewGraphTitleBar.cpp
@@ -65,6 +65,7 @@ void SNiagaraOverviewGraphTitleBar::RebuildWidget()
 		]
 	];
 
+	// warning note for affected assets
 	if (SystemViewModel.IsValid() && SystemViewModel.Pin()->GetEditMode() == ENiagaraSystemViewModelEditMode::EmitterAsset)
 	{
 		ContainerWidget->AddSlot()
@@ -84,6 +85,27 @@ void SNiagaraOverviewGraphTitleBar::RebuildWidget()
 			]
 		];
 	}
+
+	// usage of deprecated assets note
+	if (SystemViewModel.IsValid() && SystemViewModel.Pin()->GetEditMode() == ENiagaraSystemViewModelEditMode::SystemAsset)
+	{
+		ContainerWidget->AddSlot()
+		.HAlign(HAlign_Fill)
+		.AutoHeight()
+		[
+			SNew(SBorder)
+			.BorderImage(FAppStyle::GetBrush(TEXT("Graph.TitleBackground")))
+			.ColorAndOpacity(this, &SNiagaraOverviewGraphTitleBar::GetSystemSubheaderColor)
+			.Visibility(this, &SNiagaraOverviewGraphTitleBar::GetSystemSubheaderVisibility)
+			.HAlign(HAlign_Fill)
+			[
+				SNew(STextBlock)
+				.Text(this, &SNiagaraOverviewGraphTitleBar::GetSystemSubheaderText)
+				.TextStyle(FAppStyle::Get(), TEXT("GraphBreadcrumbButtonText"))
+				.Justification(ETextJustify::Center)
+			]
+		];
+	}
 	
 	if(bScalabilityModeActive)
 	{
@@ -126,6 +148,38 @@ FLinearColor SNiagaraOverviewGraphTitleBar::GetEmitterSubheaderColor() const
 	return GetEmitterAffectedAssets() >= 5 ? FNiagaraEditorWidgetsStyle::Get().GetColor("NiagaraEditor.SystemOverview.AffectedAssetsWarningColor") : FLinearColor::White;
 }
 
+FText SNiagaraOverviewGraphTitleBar::GetSystemSubheaderText() const
+{
+	TArray<FText> DeprecatedEmitterNames;
+	if (SystemViewModel.IsValid())
+	{
+		for (TSharedRef<FNiagaraEmitterHandleViewModel> EmitterModel : SystemViewModel.Pin()->GetEmitterHandleViewModels())
+		{
+			if (FVersionedNiagaraEmitterData* EmitterData = EmitterModel->GetEmitterHandle()->GetEmitterData())
+			{
+				if (FVersionedNiagaraEmitterData* ParentData = EmitterData->GetParent().GetEmitterData())
+				{
+					if (ParentData->bDeprecated)
+					{
+						DeprecatedEmitterNames.Add(FText::FromName(EmitterModel->GetName()));
+					}
+				}
+			}
+		}
+	}
+	return FText::Format(LOCTEXT("SystemSubheaderText", "The following emitters are marked as deprecated: {0}"), FText::Join(FText::FromString(", "), DeprecatedEmitterNames));
+}
+
+EVisibility SNiagaraOverviewGraphTitleBar::GetSystemSubheaderVisibility() const
+{
+	return IsUsingDeprecatedEmitter() ? EVisibility::Visible : EVisibility::Collapsed;
+}
+
+FLinearColor SNiagaraOverviewGraphTitleBar::GetSystemSubheaderColor() const
+{
+	return IsUsingDeprecatedEmitter() ? FNiagaraEditorWidgetsStyle::Get().GetColor("NiagaraEditor.SystemOverview.AffectedAssetsWarningColor") : FLinearColor::White;
+}
+
 void SNiagaraOverviewGraphTitleBar::ResetAssetCount(const FAssetData&)
 {
 	EmitterAffectedAssets.Reset();
@@ -192,4 +246,25 @@ int32 SNiagaraOverviewGraphTitleBar::GetEmitterAffectedAssets() const
 	return EmitterAffectedAssets.Get(0);
 }
 
+bool SNiagaraOverviewGraphTitleBar::IsUsingDeprecatedEmitter() const
+{
+	if (SystemViewModel.IsValid())
+	{
+		for (TSharedRef<FNiagaraEmitterHandleViewModel> EmitterModel : SystemViewModel.Pin()->GetEmitterHandleViewModels())
+		{
+			if (FVersionedNiagaraEmitterData* EmitterData = EmitterModel->GetEmitterHandle()->GetEmitterData())
+			{
+				if (FVersionedNiagaraEmitterData* ParentData = EmitterData->GetParent().GetEmitterData())
+				{
+					if (ParentData->bDeprecated)
+					{
+						return true;
+					}
+				}
+			}
+		}
+	}
+	return false;
+}
+
 #undef LOCTEXT_NAMESPACE
diff --git a/Engine/Plugins/FX/Niagara/Source/NiagaraEditorWidgets/Private/SNiagaraOverviewGraphTitleBar.h b/Engine/Plugins/FX/Niagara/Source/NiagaraEditorWidgets/Private/SNiagaraOverviewGraphTitleBar.h
--- a/Engine/Plugins/FX/Niagara/Source/NiagaraEditorWidgets/Private/SNiagaraOverviewGraphTitleBar.h
+++ b/Engine/Plugins/FX/Niagara/Source/NiagaraEditorWidgets/Private/SNiagaraOverviewGraphTitleBar.h
@@ -35,8 +35,14 @@ private:
 	FText GetEmitterSubheaderText() const;
 	EVisibility GetEmitterSubheaderVisibility() const;
 	FLinearColor GetEmitterSubheaderColor() const;
+
+	FText GetSystemSubheaderText() const;
+	EVisibility GetSystemSubheaderVisibility() const;
+	FLinearColor GetSystemSubheaderColor() const;
+	
 	
 	int32 GetEmitterAffectedAssets() const;
+	bool IsUsingDeprecatedEmitter() const;
 	void ResetAssetCount(const FAssetData& InAssetData);
 	void AddAssetListeners();
 	void ClearListeners();
