- Don't create a UAV for CPU ribbons as it's not required

diff --git a/Engine/Plugins/FX/Niagara/Source/Niagara/Private/NiagaraRendererRibbons.cpp b/Engine/Plugins/FX/Niagara/Source/Niagara/Private/NiagaraRendererRibbons.cpp
--- a/Engine/Plugins/FX/Niagara/Source/Niagara/Private/NiagaraRendererRibbons.cpp
+++ b/Engine/Plugins/FX/Niagara/Source/Niagara/Private/NiagaraRendererRibbons.cpp
@@ -243,7 +243,10 @@ struct FRWIndexBuffer final : FIndexBuffer
 		FRHIResourceCreateInfo CreateInfo(InDebugName);
 		CreateInfo.ResourceArray = InResourceArray;
 		IndexBufferRHI = RHICreateIndexBuffer(BytesPerElement, NumBytes, BUF_UnorderedAccess | AdditionalUsage, InResourceState, CreateInfo);
-		UAV = RHICreateUnorderedAccessView(IndexBufferRHI, UE_PIXELFORMAT_TO_UINT8(Format));
+		if (EnumHasAnyFlags(AdditionalUsage, BUF_UnorderedAccess))
+		{
+			UAV = RHICreateUnorderedAccessView(IndexBufferRHI, UE_PIXELFORMAT_TO_UINT8(Format));
+		}
 	}
 
 	virtual void ReleaseRHI() override
@@ -1711,9 +1714,10 @@ void FNiagaraRendererRibbons::GenerateIndexBufferForView(FNiagaraIndexGeneration
 		}
 		
 		// Can we pack into a uint16 vs a uint32?
+		const EBufferUsageFlags IndexBufferUsage = BUF_Static | (bShouldUseGPUInitIndices ? BUF_UnorderedAccess : BUF_None);
 		if (GeneratedData.TotalBitCount <= 16 && !bShouldUseGPUInitIndices)
 		{
-			RenderingViewResources->IndexBuffer.Initialize(TEXT("NiagaraRibbonIndexBuffer"), sizeof(uint16), GeneratedData.TotalNumIndices, PF_R16_UINT, ERHIAccess::VertexOrIndexBuffer, BUF_Static);
+			RenderingViewResources->IndexBuffer.Initialize(TEXT("NiagaraRibbonIndexBuffer"), sizeof(uint16), GeneratedData.TotalNumIndices, PF_R16_UINT, ERHIAccess::VertexOrIndexBuffer, IndexBufferUsage);
 			if (!bShouldUseGPUInitIndices)
 			{
 				GenerateIndexBufferCPU<uint16>(GeneratedData, DynamicDataRibbon, ShapeState, RenderingViewResources, View, ViewOriginForDistanceCulling, FeatureLevel, DrawDirection);
@@ -1721,7 +1725,7 @@ void FNiagaraRendererRibbons::GenerateIndexBufferForView(FNiagaraIndexGeneration
 		}
 		else
 		{		
-			RenderingViewResources->IndexBuffer.Initialize(TEXT("NiagaraRibbonIndexBuffer"), sizeof(uint32), GeneratedData.TotalNumIndices, PF_R32_UINT, ERHIAccess::VertexOrIndexBuffer, BUF_Static);
+			RenderingViewResources->IndexBuffer.Initialize(TEXT("NiagaraRibbonIndexBuffer"), sizeof(uint32), GeneratedData.TotalNumIndices, PF_R32_UINT, ERHIAccess::VertexOrIndexBuffer, IndexBufferUsage);
 			if (!bShouldUseGPUInitIndices)
 			{
 				GenerateIndexBufferCPU<uint32>(GeneratedData, DynamicDataRibbon, ShapeState, RenderingViewResources, View, ViewOriginForDistanceCulling, FeatureLevel, DrawDirection);				
