Niagara - Add a delegate to the niagara editor settings which can be set to restrict the use of features based on allowed uobject classes and update various parts of the UI to use this delegate.

diff --git a/Engine/Plugins/FX/Niagara/Source/NiagaraEditor/Private/Customizations/NiagaraStaticSwitchNodeDetails.cpp b/Engine/Plugins/FX/Niagara/Source/NiagaraEditor/Private/Customizations/NiagaraStaticSwitchNodeDetails.cpp
--- a/Engine/Plugins/FX/Niagara/Source/NiagaraEditor/Private/Customizations/NiagaraStaticSwitchNodeDetails.cpp
+++ b/Engine/Plugins/FX/Niagara/Source/NiagaraEditor/Private/Customizations/NiagaraStaticSwitchNodeDetails.cpp
@@ -16,6 +16,7 @@
 #include "Widgets/Layout/SWidgetSwitcher.h"
 #include "NiagaraConstants.h"
 #include "NiagaraScriptVariable.h"
+#include "NiagaraEditorUtilities.h"
 
 #define LOCTEXT_NAMESPACE "NiagaraStaticSwitchNodeDetails"
 
@@ -32,7 +33,9 @@ FNiagaraStaticSwitchNodeDetails::FNiagaraStaticSwitchNodeDetails()
 	DropdownOptions.Add(MakeShareable(new SwitchDropdownOption("Integer Constant")));
 	DropdownOptions.Add(MakeShareable(new SwitchDropdownOption("Enum Constant")));
 
-	for (FNiagaraTypeDefinition Type : FNiagaraTypeRegistry::GetRegisteredParameterTypes())
+	TArray<FNiagaraTypeDefinition> ParameterTypes;
+	FNiagaraEditorUtilities::GetAllowedParameterTypes(ParameterTypes);
+	for (FNiagaraTypeDefinition Type : ParameterTypes)
 	{
 		if (Type.IsEnum())
 		{
diff --git a/Engine/Plugins/FX/Niagara/Source/NiagaraEditor/Private/EdGraphSchema_Niagara.cpp b/Engine/Plugins/FX/Niagara/Source/NiagaraEditor/Private/EdGraphSchema_Niagara.cpp
--- a/Engine/Plugins/FX/Niagara/Source/NiagaraEditor/Private/EdGraphSchema_Niagara.cpp
+++ b/Engine/Plugins/FX/Niagara/Source/NiagaraEditor/Private/EdGraphSchema_Niagara.cpp
@@ -214,19 +214,24 @@ UEdGraphSchema_Niagara::UEdGraphSchema_Niagara(const FObjectInitializer& ObjectI
 {
 }
 
-TSharedPtr<FNiagaraSchemaAction_NewNode> AddNewNodeAction(TArray<TSharedPtr<FNiagaraSchemaAction_NewNode> >& NewActions, const FText& Category, const FText& DisplayName, const FName& InternalName, const FText& Tooltip, FText Keywords = FText(), int32 SectionID = 0)
-{
-	TSharedPtr<FNiagaraSchemaAction_NewNode> NewAction = MakeShared<FNiagaraSchemaAction_NewNode>(Category, DisplayName, InternalName, Tooltip, 0, Keywords, SectionID);
-	NewActions.Add(NewAction);
-
-	return NewAction;
-}
-
-TSharedPtr<FNiagaraAction_NewNode> AddNewNodeMenuAction(TArray<TSharedPtr<FNiagaraAction_NewNode>>& NewActions, UEdGraphNode* InNodeTemplate, const FText& DisplayName, ENiagaraMenuSections Section, TArray<FString> NestedCategories, const FText& Tooltip, FText Keywords = FText(), FNiagaraActionSourceData
-                                            SourceData = FNiagaraActionSourceData(EScriptSource::Niagara, FText::FromString(TEXT("Niagara")), true))
+TSharedPtr<FNiagaraAction_NewNode> AddNewNodeMenuAction(
+	TArray<TSharedPtr<FNiagaraAction_NewNode>>& NewActions,
+	UEdGraphNode* InNodeTemplate,
+	const FText& DisplayName,
+	ENiagaraMenuSections Section,
+	TArray<FString> NestedCategories,
+	const FText& Tooltip,
+	FText Keywords = FText(),
+	FNiagaraActionSourceData SourceData = FNiagaraActionSourceData(
+		EScriptSource::Niagara, FText::FromString(TEXT("Niagara")), true))
 {
+	const UNiagaraEditorSettings* NiagaraEditorSettings = GetDefault<UNiagaraEditorSettings>();
 	TSharedPtr<FNiagaraAction_NewNode> NewAction = MakeShared<FNiagaraAction_NewNode>(DisplayName, Section, NestedCategories, Tooltip, Keywords);
-	NewAction->NodeTemplate = InNodeTemplate;
+	if (ensureMsgf(InNodeTemplate == nullptr || NiagaraEditorSettings->IsAllowedClass(InNodeTemplate->GetClass()),
+		TEXT("Can not create a menu action for a node of class %s."), *InNodeTemplate->GetClass()->GetName()))
+	{
+		NewAction->NodeTemplate = InNodeTemplate;
+	}
 	NewAction->SourceData = SourceData;
 	NewActions.Add(NewAction);
 
@@ -398,6 +403,8 @@ TArray<TSharedPtr<FNiagaraAction_NewNode>> UEdGraphSchema_Niagara::GetGraphActio
 	bool bAllowSelectNodes = true;
 	bool bAllowStaticSwitchNodes = true;
 
+	const UNiagaraEditorSettings* NiagaraEditorSettings = GetDefault<UNiagaraEditorSettings>();
+
 	TOptional<FNiagaraTypeDefinition> FromType;
 	TOptional<EEdGraphPinDirection> FromDirection;
 
@@ -427,7 +434,9 @@ TArray<TSharedPtr<FNiagaraAction_NewNode>> UEdGraphSchema_Niagara::GetGraphActio
 	}
 	
 	// Add operations (add / mul / etc)
-	if (bAllowOpNodes && (GbAllowAllNiagaraNodesInEmitterGraphs || bModuleGraph || bFunctionGraph || bSystemGraph))
+	if (bAllowOpNodes &&
+		(GbAllowAllNiagaraNodesInEmitterGraphs || bModuleGraph || bFunctionGraph || bSystemGraph) &&
+		NiagaraEditorSettings->IsAllowedClass(UNiagaraNodeOp::StaticClass()))
 	{
 		const TArray<FNiagaraOpInfo>& OpInfos = FNiagaraOpInfo::GetOpInfoArray();
 
@@ -478,7 +487,7 @@ TArray<TSharedPtr<FNiagaraAction_NewNode>> UEdGraphSchema_Niagara::GetGraphActio
 	}
 
 	// Add custom code
-	if (bAllowCustomNode)
+	if (bAllowCustomNode && NiagaraEditorSettings->IsAllowedClass(UNiagaraNodeCustomHlsl::StaticClass()))
 	{
 		const FText DisplayName = LOCTEXT("CustomHLSLNode","Custom Hlsl");
 		const FText TooltipDesc = LOCTEXT("CustomHlslPopupTooltip", "Add a node with custom hlsl content");
@@ -518,7 +527,9 @@ PRAGMA_ENABLE_DEPRECATION_WARNINGS
 	};
 
 	//Add functions
-	if (bAllowFunctionNodes && (GbAllowAllNiagaraNodesInEmitterGraphs || bModuleGraph || bFunctionGraph || bDynamicInputGraph))
+	if (bAllowFunctionNodes &&
+		(GbAllowAllNiagaraNodesInEmitterGraphs || bModuleGraph || bFunctionGraph || bDynamicInputGraph) &&
+		NiagaraEditorSettings->IsAllowedClass(UNiagaraNodeFunctionCall::StaticClass()))
 	{
 		TArray<FAssetData> FunctionScriptAssets;
 		FNiagaraEditorUtilities::FGetFilteredScriptAssetsOptions FunctionScriptFilterOptions;
@@ -533,7 +544,8 @@ PRAGMA_ENABLE_DEPRECATION_WARNINGS
 	}
 
 	//Add modules
-	if (bAllowModuleNodes && !bFunctionGraph && !bModuleGraph && !bDynamicInputGraph)
+	if (bAllowModuleNodes && !bFunctionGraph && !bModuleGraph && !bDynamicInputGraph &&
+		NiagaraEditorSettings->IsAllowedClass(UNiagaraNodeFunctionCall::StaticClass()))
 	{
 		TArray<FAssetData> ModuleScriptAssets;
 		FNiagaraEditorUtilities::FGetFilteredScriptAssetsOptions ModuleScriptFilterOptions;
@@ -549,13 +561,16 @@ PRAGMA_ENABLE_DEPRECATION_WARNINGS
 	}
 
 	//Add event read and writes nodes
-	if (bAllowEventNodes && bModuleGraph)
+	if (bAllowEventNodes && bModuleGraph &&
+		NiagaraEditorSettings->IsAllowedClass(UNiagaraNodeReadDataSet::StaticClass()) &&
+		NiagaraEditorSettings->IsAllowedClass(UNiagaraNodeWriteDataSet::StaticClass()))
 	{
 		const FText MenuCat = LOCTEXT("NiagaraEventMenuCat", "Events");
 		TArray<FString> ReadCategories = {MenuCat.ToString(), LOCTEXT("NiagaraEventCategory_Read", "Read").ToString()};
 		TArray<FString> WriteCategories = {MenuCat.ToString(), LOCTEXT("NiagaraEventCategory_Write", "Write").ToString()};
 
-		const TArray<FNiagaraTypeDefinition>& PayloadTypes = FNiagaraTypeRegistry::GetRegisteredPayloadTypes();
+		TArray<FNiagaraTypeDefinition> PayloadTypes;
+		FNiagaraEditorUtilities::GetAllowedPayloadTypes(PayloadTypes);
 		for (const FNiagaraTypeDefinition& PayloadType : PayloadTypes)
 		{
 			if (PayloadType.IsInternalType())
@@ -605,7 +620,7 @@ PRAGMA_ENABLE_DEPRECATION_WARNINGS
 		UsageTypesToAdd.Add(ENiagaraScriptUsage::SystemUpdateScript);
 	}
 
-	if (UsageTypesToAdd.Num() != 0)
+	if (UsageTypesToAdd.Num() != 0 && NiagaraEditorSettings->IsAllowedClass(UNiagaraNodeOutput::StaticClass()))
 	{
 		for (ENiagaraScriptUsage Usage : UsageTypesToAdd)
 		{
@@ -632,6 +647,7 @@ PRAGMA_ENABLE_DEPRECATION_WARNINGS
 
 
 	// Add Convert Nodes
+	if (NiagaraEditorSettings->IsAllowedClass(UNiagaraNodeConvert::StaticClass()))
 	{
 		FNiagaraTypeDefinition PinType = FNiagaraTypeDefinition::GetGenericNumericDef();
 
@@ -671,6 +687,8 @@ PRAGMA_ENABLE_DEPRECATION_WARNINGS
 			NoMakeAllowed.Add(FNiagaraTypeDefinition::GetParameterMapDef());
 			
 			// Dynamic Makes from type registry. Should include the pin type as well.
+			TArray<FNiagaraTypeDefinition> CandidateTypes;
+			FNiagaraEditorUtilities::GetAllowedTypes(CandidateTypes);
 			for (const FNiagaraTypeDefinition& CandidateType : FNiagaraTypeRegistry::GetRegisteredTypes())
 			{
 				bool bAddMake = false;
@@ -818,26 +836,31 @@ PRAGMA_ENABLE_DEPRECATION_WARNINGS
 	if (bAllowParameterMapGetSetNodes)
 	{
 		FText MenuCat = FText::FromString("Parameter Map");
+		if (NiagaraEditorSettings->IsAllowedClass(UNiagaraNodeParameterMapGet::StaticClass()))
 		{
 			FString Name = TEXT("Parameter Map Get");
 			UNiagaraNodeParameterMapGet* BaseNode = NewObject<UNiagaraNodeParameterMapGet>(OwnerOfTemporaries);
 			AddNewNodeMenuAction(NewActions, BaseNode, FText::FromString(Name), ENiagaraMenuSections::Suggested, {MenuCat.ToString()}, FText::GetEmpty(), FText::GetEmpty());
 		}
+		if (NiagaraEditorSettings->IsAllowedClass(UNiagaraNodeParameterMapSet::StaticClass()))
 		{
 			FString Name = TEXT("Parameter Map Set");
 			UNiagaraNodeParameterMapSet* BaseNode = NewObject<UNiagaraNodeParameterMapSet>(OwnerOfTemporaries);
 			AddNewNodeMenuAction(NewActions, BaseNode, FText::FromString(Name), ENiagaraMenuSections::Suggested, {MenuCat.ToString()}, FText::GetEmpty(), FText::GetEmpty());
 		}
+		if (NiagaraEditorSettings->IsAllowedClass(UNiagaraNodeParameterMapFor::StaticClass()))
 		{
 			FString Name = TEXT("Parameter Map For");
 			UNiagaraNodeParameterMapFor* BaseNode = NewObject<UNiagaraNodeParameterMapFor>(OwnerOfTemporaries);
 			AddNewNodeMenuAction(NewActions, BaseNode, FText::FromString(Name), ENiagaraMenuSections::General, {MenuCat.ToString()}, FText::GetEmpty(), FText::GetEmpty());
 		}
+		if (NiagaraEditorSettings->IsAllowedClass(UNiagaraNodeParameterMapForWithContinue::StaticClass()))
 		{
 			FString Name = TEXT("Parameter Map For With Continue");
 			UNiagaraNodeParameterMapForWithContinue* BaseNode = NewObject<UNiagaraNodeParameterMapForWithContinue>(OwnerOfTemporaries);
 			AddNewNodeMenuAction(NewActions, BaseNode, FText::FromString(Name), ENiagaraMenuSections::General, { MenuCat.ToString() }, FText::GetEmpty(), FText::GetEmpty());
 		}
+		if (NiagaraEditorSettings->IsAllowedClass(UNiagaraNodeParameterMapForIndex::StaticClass()))
 		{
 			FString Name = TEXT("Parameter Map For Current Index");
 			UNiagaraNodeParameterMapForIndex* BaseNode = NewObject<UNiagaraNodeParameterMapForIndex>(OwnerOfTemporaries);
@@ -846,6 +869,7 @@ PRAGMA_ENABLE_DEPRECATION_WARNINGS
 	}
 
 	// Handle comment nodes
+	if (NiagaraEditorSettings->IsAllowedClass(UEdGraphNode_Comment::StaticClass()))
 	{
 		FText MenuCat = FText::FromString("Comments");
 
@@ -857,7 +881,7 @@ PRAGMA_ENABLE_DEPRECATION_WARNINGS
 	}
 
 	// Handle output tag nodes
-	if ( bAllowOutputTagNodes )
+	if (bAllowOutputTagNodes && NiagaraEditorSettings->IsAllowedClass(UNiagaraNodeOutputTag::StaticClass()))
 	{
 		FText MenuCat = FText::FromString("Compiler Tagging");
 
@@ -869,7 +893,7 @@ PRAGMA_ENABLE_DEPRECATION_WARNINGS
 	}
 	
 	//Add all input node options for input pins or no pin.
-	if (FromPin == nullptr || FromPin->Direction == EGPD_Input)
+	if ((FromPin == nullptr || FromPin->Direction == EGPD_Input) && NiagaraEditorSettings->IsAllowedClass(UNiagaraNodeInput::StaticClass()))
 	{
 		TArray<UNiagaraNodeInput*> InputNodes;
 		NiagaraGraph->GetNodesOfClass(InputNodes);
@@ -926,8 +950,9 @@ PRAGMA_ENABLE_DEPRECATION_WARNINGS
 			{
 				const FText MenuDescFmt = LOCTEXT("Add ParameterFmt", "Add {0} Parameter");
 				const FText AddParameterCategory = LOCTEXT("AddParameterCat", "Add Parameter");
-				const TArray<FNiagaraTypeDefinition>& RegisteredTypes = FNiagaraTypeRegistry::GetRegisteredParameterTypes();
-				for (FNiagaraTypeDefinition Type : RegisteredTypes)
+				TArray<FNiagaraTypeDefinition> RegisteredTypes;
+				FNiagaraEditorUtilities::GetAllowedParameterTypes(RegisteredTypes);
+				for (const FNiagaraTypeDefinition& Type : RegisteredTypes)
 				{
 					if (Type.IsUObject() && Type.IsDataInterface() == false)
 					{
@@ -978,6 +1003,7 @@ PRAGMA_ENABLE_DEPRECATION_WARNINGS
 	const FText UtilMenuCat = LOCTEXT("NiagaraUsageSelectorMenuCat", "Utility");
 
 	// Add reroute node
+	if (NiagaraEditorSettings->IsAllowedClass(UNiagaraNodeReroute::StaticClass()))
 	{
 		const FText RerouteMenuDesc = LOCTEXT("NiagaraRerouteMenuDesc", "Reroute");
 		
@@ -987,7 +1013,7 @@ PRAGMA_ENABLE_DEPRECATION_WARNINGS
 	
 	// Add select node
 	// Note: Data Interfaces are not supported for this type
-	if ( bAllowSelectNodes )
+	if (bAllowSelectNodes && NiagaraEditorSettings->IsAllowedClass(UNiagaraNodeSelect::StaticClass()))
 	{
 		const FText SelectMenuDesc = LOCTEXT("NiagaraSelectMenuDesc", "Select / If");
 		
@@ -997,7 +1023,7 @@ PRAGMA_ENABLE_DEPRECATION_WARNINGS
 
 	// Add static switch node
 	// Note: Data Interfaces are not supported for this type
-	if ( bAllowStaticSwitchNodes )
+	if (bAllowStaticSwitchNodes && NiagaraEditorSettings->IsAllowedClass(UNiagaraNodeStaticSwitch::StaticClass()))
 	{
 		const FText UsageSelectorMenuDesc = LOCTEXT("NiagaraStaticSwitchMenuDesc", "Static Switch");
 		
diff --git a/Engine/Plugins/FX/Niagara/Source/NiagaraEditor/Private/NiagaraEditorSettings.cpp b/Engine/Plugins/FX/Niagara/Source/NiagaraEditor/Private/NiagaraEditorSettings.cpp
--- a/Engine/Plugins/FX/Niagara/Source/NiagaraEditor/Private/NiagaraEditorSettings.cpp
+++ b/Engine/Plugins/FX/Niagara/Source/NiagaraEditor/Private/NiagaraEditorSettings.cpp
@@ -646,5 +646,20 @@ UNiagaraEditorSettings::FOnNiagaraEditorSettingsChanged& UNiagaraEditorSettings:
 	return GetMutableDefault<UNiagaraEditorSettings>()->SettingsChangedDelegate;
 }
 
+void UNiagaraEditorSettings::SetOnIsClassAllowed(const FOnIsClassAllowed& InOnIsClassAllowed)
+{
+	OnIsClassAllowedDelegate = InOnIsClassAllowed;
+}
+
+bool UNiagaraEditorSettings::IsAllowedClass(const UClass* InClass) const
+{
+	return OnIsClassAllowedDelegate.IsBound() == false || OnIsClassAllowedDelegate.Execute(InClass);
+}
+
+bool UNiagaraEditorSettings::IsAllowedTypeDefinition(const FNiagaraTypeDefinition& InTypeDefinition) const
+{
+	return InTypeDefinition.GetClass() == nullptr || IsAllowedClass(InTypeDefinition.GetClass());
+}
+
 #undef LOCTEXT_NAMESPACE
 
diff --git a/Engine/Plugins/FX/Niagara/Source/NiagaraEditor/Private/NiagaraEditorUtilities.cpp b/Engine/Plugins/FX/Niagara/Source/NiagaraEditor/Private/NiagaraEditorUtilities.cpp
--- a/Engine/Plugins/FX/Niagara/Source/NiagaraEditor/Private/NiagaraEditorUtilities.cpp
+++ b/Engine/Plugins/FX/Niagara/Source/NiagaraEditor/Private/NiagaraEditorUtilities.cpp
@@ -3924,6 +3924,90 @@ int FNiagaraEditorUtilities::GetReferencedAssetCount(const FAssetData& SourceAss
 	return FMath::Max(Count, 0);
 }
 
+void FNiagaraEditorUtilities::GetAllowedTypes(TArray<FNiagaraTypeDefinition>& OutAllowedTypes)
+{
+	const UNiagaraEditorSettings* NiagaraEditorSettings = GetDefault<UNiagaraEditorSettings>();
+	for (const FNiagaraTypeDefinition& RegisteredType : FNiagaraTypeRegistry::GetRegisteredTypes())
+	{
+		if (NiagaraEditorSettings->IsAllowedTypeDefinition(RegisteredType))
+		{
+			OutAllowedTypes.Add(RegisteredType);
+		}
+	}
+}
+
+void FNiagaraEditorUtilities::GetAllowedUserVariableTypes(TArray<FNiagaraTypeDefinition>& OutAllowedTypes)
+{
+	const UNiagaraEditorSettings* NiagaraEditorSettings = GetDefault<UNiagaraEditorSettings>();
+	for (const FNiagaraTypeDefinition& RegisteredUserVariableType : FNiagaraTypeRegistry::GetRegisteredUserVariableTypes())
+	{
+		if (NiagaraEditorSettings->IsAllowedTypeDefinition(RegisteredUserVariableType))
+		{
+			OutAllowedTypes.Add(RegisteredUserVariableType);
+		}
+	}
+}
+
+void FNiagaraEditorUtilities::GetAllowedSystemVariableTypes(TArray<FNiagaraTypeDefinition>& OutAllowedTypes)
+{
+	const UNiagaraEditorSettings* NiagaraEditorSettings = GetDefault<UNiagaraEditorSettings>();
+	for (const FNiagaraTypeDefinition& RegisteredSystemVariableType : FNiagaraTypeRegistry::GetRegisteredSystemVariableTypes())
+	{
+		if (NiagaraEditorSettings->IsAllowedTypeDefinition(RegisteredSystemVariableType))
+		{
+			OutAllowedTypes.Add(RegisteredSystemVariableType);
+		}
+	}
+}
+
+void FNiagaraEditorUtilities::GetAllowedEmitterVariableTypes(TArray<FNiagaraTypeDefinition>& OutAllowedTypes)
+{
+	const UNiagaraEditorSettings* NiagaraEditorSettings = GetDefault<UNiagaraEditorSettings>();
+	for (const FNiagaraTypeDefinition& RegisteredEmitterVariableType : FNiagaraTypeRegistry::GetRegisteredEmitterVariableTypes())
+	{
+		if (NiagaraEditorSettings->IsAllowedTypeDefinition(RegisteredEmitterVariableType))
+		{
+			OutAllowedTypes.Add(RegisteredEmitterVariableType);
+		}
+	}
+}
+
+void FNiagaraEditorUtilities::GetAllowedParticleVariableTypes(TArray<FNiagaraTypeDefinition>& OutAllowedTypes)
+{
+	const UNiagaraEditorSettings* NiagaraEditorSettings = GetDefault<UNiagaraEditorSettings>();
+	for (const FNiagaraTypeDefinition& RegisteredParticleVariableType : FNiagaraTypeRegistry::GetRegisteredParticleVariableTypes())
+	{
+		if (NiagaraEditorSettings->IsAllowedTypeDefinition(RegisteredParticleVariableType))
+		{
+			OutAllowedTypes.Add(RegisteredParticleVariableType);
+		}
+	}
+}
+
+void FNiagaraEditorUtilities::GetAllowedParameterTypes(TArray<FNiagaraTypeDefinition>& OutAllowedTypes)
+{
+	const UNiagaraEditorSettings* NiagaraEditorSettings = GetDefault<UNiagaraEditorSettings>();
+	for (const FNiagaraTypeDefinition& RegisteredParameterType : FNiagaraTypeRegistry::GetRegisteredParameterTypes())
+	{
+		if (NiagaraEditorSettings->IsAllowedTypeDefinition(RegisteredParameterType))
+		{
+			OutAllowedTypes.Add(RegisteredParameterType);
+		}
+	}
+}
+
+void FNiagaraEditorUtilities::GetAllowedPayloadTypes(TArray<FNiagaraTypeDefinition>& OutAllowedTypes)
+{
+	const UNiagaraEditorSettings* NiagaraEditorSettings = GetDefault<UNiagaraEditorSettings>();
+	for (const FNiagaraTypeDefinition& RegisteredPayloadType : FNiagaraTypeRegistry::GetRegisteredPayloadTypes())
+	{
+		if (NiagaraEditorSettings->IsAllowedTypeDefinition(RegisteredPayloadType))
+		{
+			OutAllowedTypes.Add(RegisteredPayloadType);
+		}
+	}
+}
+
 void FNiagaraParameterUtilities::FilterToRelevantStaticVariables(const TArray<FNiagaraVariable>& InVars, TArray<FNiagaraVariable>& OutVars, FName InOldEmitterAlias, FName InNewEmitterAlias, bool bFilterByEmitterAliasAndConvertToUnaliased)
 {
 	FNiagaraAliasContext RenameContext(ENiagaraScriptUsage::ParticleSpawnScript);
diff --git a/Engine/Plugins/FX/Niagara/Source/NiagaraEditor/Private/ViewModels/NiagaraParameterCollectionViewModel.cpp b/Engine/Plugins/FX/Niagara/Source/NiagaraEditor/Private/ViewModels/NiagaraParameterCollectionViewModel.cpp
--- a/Engine/Plugins/FX/Niagara/Source/NiagaraEditor/Private/ViewModels/NiagaraParameterCollectionViewModel.cpp
+++ b/Engine/Plugins/FX/Niagara/Source/NiagaraEditor/Private/ViewModels/NiagaraParameterCollectionViewModel.cpp
@@ -124,7 +124,9 @@ void FNiagaraParameterCollectionViewModel::RefreshAvailableTypes()
 	}
 
 	AvailableTypes->Empty();
-	for (const FNiagaraTypeDefinition& RegisteredType : FNiagaraTypeRegistry::GetRegisteredParameterTypes())
+	TArray<FNiagaraTypeDefinition> AllowedParameterTypes;
+	FNiagaraEditorUtilities::GetAllowedParameterTypes(AllowedParameterTypes);
+	for (const FNiagaraTypeDefinition& RegisteredType : AllowedParameterTypes)
 	{
 		if (SupportsType(RegisteredType))
 		{
diff --git a/Engine/Plugins/FX/Niagara/Source/NiagaraEditor/Private/ViewModels/NiagaraParameterPanelViewModel.cpp b/Engine/Plugins/FX/Niagara/Source/NiagaraEditor/Private/ViewModels/NiagaraParameterPanelViewModel.cpp
--- a/Engine/Plugins/FX/Niagara/Source/NiagaraEditor/Private/ViewModels/NiagaraParameterPanelViewModel.cpp
+++ b/Engine/Plugins/FX/Niagara/Source/NiagaraEditor/Private/ViewModels/NiagaraParameterPanelViewModel.cpp
@@ -576,7 +576,9 @@ bool INiagaraParameterPanelViewModel::GetCanChangeParameterType(const TArray<FNi
 void INiagaraParameterPanelViewModel::GetChangeTypeSubMenu(FMenuBuilder& MenuBuilder, TArray<FNiagaraParameterPanelItem> Items) 
 {
 	TArray<FNiagaraTypeDefinition> FilteredTypes;
-	for (const FNiagaraTypeDefinition& RegisteredType : FNiagaraTypeRegistry::GetRegisteredTypes())
+	TArray<FNiagaraTypeDefinition> AllowedTypes;
+	FNiagaraEditorUtilities::GetAllowedTypes(AllowedTypes);
+	for (const FNiagaraTypeDefinition& RegisteredType : AllowedTypes)
 	{
 		// only allow basic types for now
 		if (RegisteredType.IsDataInterface() || RegisteredType.IsEnum() || RegisteredType.IsUObject() || FNiagaraTypeRegistry::GetRegisteredPayloadTypes().Contains(RegisteredType) || RegisteredType.IsInternalType())
diff --git a/Engine/Plugins/FX/Niagara/Source/NiagaraEditor/Private/ViewModels/Stack/NiagaraStackEmitterPropertiesGroup.cpp b/Engine/Plugins/FX/Niagara/Source/NiagaraEditor/Private/ViewModels/Stack/NiagaraStackEmitterPropertiesGroup.cpp
--- a/Engine/Plugins/FX/Niagara/Source/NiagaraEditor/Private/ViewModels/Stack/NiagaraStackEmitterPropertiesGroup.cpp
+++ b/Engine/Plugins/FX/Niagara/Source/NiagaraEditor/Private/ViewModels/Stack/NiagaraStackEmitterPropertiesGroup.cpp
@@ -92,17 +92,21 @@ public:
 			LOCTEXT("AddEventActionDescription", "Add an event handler to this emitter."),
 			LOCTEXT("AddEventActionKeywords", "Event")));
 
+		const UNiagaraEditorSettings* NiagaraEditorSettings = GetDefault<UNiagaraEditorSettings>();
 		TArray<UClass*> SimulationStageClasses;
 		GetDerivedClasses(UNiagaraSimulationStageBase::StaticClass(), SimulationStageClasses);
 		for (UClass* SimulationStageClass : SimulationStageClasses)
 		{
-			OutAddActions.Add(MakeShared<FAddEmitterStageAction>(
-				EEmitterAddMode::SimulationStage,
-				SimulationStageClass,
-				SimulationStageCategories,
-				SimulationStageClass->GetDisplayNameText(),
-				FText::FromString(SimulationStageClass->GetDescription()),
-				FText::FromString(SimulationStageClass->GetName())));
+			if (NiagaraEditorSettings->IsAllowedClass(SimulationStageClass))
+			{
+				OutAddActions.Add(MakeShared<FAddEmitterStageAction>(
+					EEmitterAddMode::SimulationStage,
+					SimulationStageClass,
+					SimulationStageCategories,
+					SimulationStageClass->GetDisplayNameText(),
+					FText::FromString(SimulationStageClass->GetDescription()),
+					FText::FromString(SimulationStageClass->GetName())));
+			}
 		}
 	}
 
diff --git a/Engine/Plugins/FX/Niagara/Source/NiagaraEditor/Private/ViewModels/Stack/NiagaraStackFunctionInput.cpp b/Engine/Plugins/FX/Niagara/Source/NiagaraEditor/Private/ViewModels/Stack/NiagaraStackFunctionInput.cpp
--- a/Engine/Plugins/FX/Niagara/Source/NiagaraEditor/Private/ViewModels/Stack/NiagaraStackFunctionInput.cpp
+++ b/Engine/Plugins/FX/Niagara/Source/NiagaraEditor/Private/ViewModels/Stack/NiagaraStackFunctionInput.cpp
@@ -638,6 +638,16 @@ void UNiagaraStackFunctionInput::RefreshChildrenInternal(const TArray<UNiagaraSt
         ));
 	}
 
+	if (InputValues.Mode == EValueMode::Expression && SupportsCustomExpressions() == false)
+	{
+		NewIssues.Add(FStackIssue(
+			EStackIssueSeverity::Error,
+			LOCTEXT("UnsupportedExpressionShort", "Expression Input Unsupported"),
+			LOCTEXT("UnsupportedExpressionLong", "Use of expressions for function inputs is not currently supported in the current editor context."),
+			GetStackEditorDataKey(),
+			false));
+	}
+
 	if (InputValues.Mode == EValueMode::Dynamic && InputValues.DynamicNode.IsValid())
 	{
 		FVersionedNiagaraScriptData* ScriptData = InputValues.DynamicNode->GetScriptData();
@@ -3153,5 +3163,11 @@ bool UNiagaraStackFunctionInput::OpenSourceAsset() const
 	return false;
 }
 
+bool UNiagaraStackFunctionInput::SupportsCustomExpressions() const
+{
+	const UNiagaraEditorSettings* NiagaraEditorSettings = GetDefault<UNiagaraEditorSettings>();
+	return NiagaraEditorSettings->IsAllowedClass(UNiagaraNodeCustomHlsl::StaticClass());
+}
+
 #undef LOCTEXT_NAMESPACE
 
diff --git a/Engine/Plugins/FX/Niagara/Source/NiagaraEditor/Private/ViewModels/Stack/NiagaraStackGraphUtilities.cpp b/Engine/Plugins/FX/Niagara/Source/NiagaraEditor/Private/ViewModels/Stack/NiagaraStackGraphUtilities.cpp
--- a/Engine/Plugins/FX/Niagara/Source/NiagaraEditor/Private/ViewModels/Stack/NiagaraStackGraphUtilities.cpp
+++ b/Engine/Plugins/FX/Niagara/Source/NiagaraEditor/Private/ViewModels/Stack/NiagaraStackGraphUtilities.cpp
@@ -2195,26 +2195,26 @@ void FNiagaraStackGraphUtilities::CleanUpStaleRapidIterationParameters(FVersione
 
 void FNiagaraStackGraphUtilities::GetNewParameterAvailableTypes(TArray<FNiagaraTypeDefinition>& OutAvailableTypes, FName Namespace)
 {
-	TConstArrayView<FNiagaraTypeDefinition> ParamTypes;
+	TArray<FNiagaraTypeDefinition> ParamTypes;
 	if (Namespace == FNiagaraConstants::UserNamespace)
 	{
-		ParamTypes = MakeArrayView(FNiagaraTypeRegistry::GetRegisteredUserVariableTypes());
+		FNiagaraEditorUtilities::GetAllowedUserVariableTypes(ParamTypes);
 	}
 	else if (Namespace == FNiagaraConstants::SystemNamespace)
 	{
-		ParamTypes = MakeArrayView(FNiagaraTypeRegistry::GetRegisteredSystemVariableTypes());
+		FNiagaraEditorUtilities::GetAllowedSystemVariableTypes(ParamTypes);
 	}
 	else if (Namespace == FNiagaraConstants::EmitterNamespace)
 	{
-		ParamTypes = MakeArrayView(FNiagaraTypeRegistry::GetRegisteredEmitterVariableTypes());
+		FNiagaraEditorUtilities::GetAllowedEmitterVariableTypes(ParamTypes);
 	}
 	else if (Namespace == FNiagaraConstants::ParticleAttributeNamespace)
 	{
-		ParamTypes = MakeArrayView(FNiagaraTypeRegistry::GetRegisteredParticleVariableTypes());
+		FNiagaraEditorUtilities::GetAllowedParticleVariableTypes(ParamTypes);
 	}
 	else
 	{
-		ParamTypes = MakeArrayView(FNiagaraTypeRegistry::GetRegisteredParameterTypes());
+		FNiagaraEditorUtilities::GetAllowedParameterTypes(ParamTypes);
 	}
 
 	for (const FNiagaraTypeDefinition& RegisteredParameterType : ParamTypes)
diff --git a/Engine/Plugins/FX/Niagara/Source/NiagaraEditor/Private/ViewModels/Stack/NiagaraStackObject.cpp b/Engine/Plugins/FX/Niagara/Source/NiagaraEditor/Private/ViewModels/Stack/NiagaraStackObject.cpp
--- a/Engine/Plugins/FX/Niagara/Source/NiagaraEditor/Private/ViewModels/Stack/NiagaraStackObject.cpp
+++ b/Engine/Plugins/FX/Niagara/Source/NiagaraEditor/Private/ViewModels/Stack/NiagaraStackObject.cpp
@@ -7,6 +7,7 @@
 #include "NiagaraEditorModule.h"
 #include "NiagaraMessageManager.h"
 #include "NiagaraMessageUtilities.h"
+#include "NiagaraEditorSettings.h"
 
 #include "Modules/ModuleManager.h"
 #include "IPropertyRowGenerator.h"
@@ -257,6 +258,17 @@ void UNiagaraStackObject::RefreshChildrenInternal(const TArray<UNiagaraStackEntr
 		}
 	}
 
+	const UNiagaraEditorSettings* NiagaraEditorSettings = GetDefault<UNiagaraEditorSettings>();
+	if (Object != nullptr && NiagaraEditorSettings->IsAllowedClass(Object->GetClass()) == false)
+	{
+		NewIssues.Add(FStackIssue(
+			EStackIssueSeverity::Error,
+			NSLOCTEXT("StackObject", "InvalidClassShort", "Unsupported Object Type"),
+			FText::Format(NSLOCTEXT("StackObject", "InvalidClassLongFormat", "Use of an object of type {0} is unsupported in the current editor context."), Object->GetClass()->GetDisplayNameText()),
+			GetStackEditorDataKey(),
+			false));
+	}
+
 	if(PropertyRowGenerator.IsValid())
 	{
 		TArray<TSharedRef<IDetailTreeNode>> DefaultRootTreeNodes = PropertyRowGenerator->GetRootTreeNodes();
diff --git a/Engine/Plugins/FX/Niagara/Source/NiagaraEditor/Private/ViewModels/Stack/NiagaraStackRenderItemGroup.cpp b/Engine/Plugins/FX/Niagara/Source/NiagaraEditor/Private/ViewModels/Stack/NiagaraStackRenderItemGroup.cpp
--- a/Engine/Plugins/FX/Niagara/Source/NiagaraEditor/Private/ViewModels/Stack/NiagaraStackRenderItemGroup.cpp
+++ b/Engine/Plugins/FX/Niagara/Source/NiagaraEditor/Private/ViewModels/Stack/NiagaraStackRenderItemGroup.cpp
@@ -9,6 +9,7 @@
 #include "NiagaraRendererProperties.h"
 #include "NiagaraClipboard.h"
 #include "NiagaraEditorModule.h"
+#include "NiagaraEditorSettings.h"
 
 #include "ScopedTransaction.h"
 #include "Widgets/Notifications/SNotificationList.h"
@@ -72,7 +73,11 @@ public:
 
 		for (const FNiagaraRendererCreationInfo& RendererCreationInfo : RendererCreationInfos)
 		{
-			OutAddActions.Add(MakeShared<FRenderItemGroupAddAction>(RendererCreationInfo));
+			// TODO: UE-165740
+			//if (NiagaraEditorSettings->IsAllowedClass(RendererClass))
+			{
+				OutAddActions.Add(MakeShared<FRenderItemGroupAddAction>(RendererCreationInfo));
+			}
 		}
 	}
 
diff --git a/Engine/Plugins/FX/Niagara/Source/NiagaraEditor/Private/Widgets/SNiagaraGraphNode.cpp b/Engine/Plugins/FX/Niagara/Source/NiagaraEditor/Private/Widgets/SNiagaraGraphNode.cpp
--- a/Engine/Plugins/FX/Niagara/Source/NiagaraEditor/Private/Widgets/SNiagaraGraphNode.cpp
+++ b/Engine/Plugins/FX/Niagara/Source/NiagaraEditor/Private/Widgets/SNiagaraGraphNode.cpp
@@ -4,6 +4,7 @@
 #include "GraphEditorSettings.h"
 #include "Rendering/DrawElements.h"
 #include "SGraphPin.h"
+#include "NiagaraEditorSettings.h"
 
 #define LOCTEXT_NAMESPACE "SNiagaraGraphNode"
 
@@ -25,6 +26,7 @@ void SNiagaraGraphNode::Construct(const FArguments& InArgs, UEdGraphNode* InGrap
 	GraphNode = InGraphNode;
 	RegisterNiagaraGraphNode(InGraphNode);
 	UpdateGraphNode();
+	
 }
 
 
@@ -57,4 +59,18 @@ void SNiagaraGraphNode::CreateOutputSideAddButton(TSharedPtr<SVerticalBox> Outpu
 	NiagaraNode->AddWidgetsToOutputBox(OutputBox);
 }
 
+void SNiagaraGraphNode::UpdateErrorInfo()
+{
+	const UNiagaraEditorSettings* NiagaraEditorSettings = GetDefault<UNiagaraEditorSettings>();
+	if (GraphNode != nullptr && GraphNode->IsA<UNiagaraNode>() && NiagaraEditorSettings->IsAllowedClass(GraphNode->GetClass()) == false)
+	{
+		ErrorMsg = FString(TEXT("UNSUPPORTED!"));
+		ErrorColor = FAppStyle::GetColor("ErrorReporting.BackgroundColor");
+	}
+	else
+	{
+		SGraphNode::UpdateErrorInfo();
+	}
+}
+
 #undef LOCTEXT_NAMESPACE
\ No newline at end of file
diff --git a/Engine/Plugins/FX/Niagara/Source/NiagaraEditor/Private/Widgets/SNiagaraGraphNode.h b/Engine/Plugins/FX/Niagara/Source/NiagaraEditor/Private/Widgets/SNiagaraGraphNode.h
--- a/Engine/Plugins/FX/Niagara/Source/NiagaraEditor/Private/Widgets/SNiagaraGraphNode.h
+++ b/Engine/Plugins/FX/Niagara/Source/NiagaraEditor/Private/Widgets/SNiagaraGraphNode.h
@@ -21,6 +21,7 @@ public:
 	virtual void UpdateGraphNode() override;
 	virtual void CreateInputSideAddButton(TSharedPtr<SVerticalBox> InputBox) override;
 	virtual void CreateOutputSideAddButton(TSharedPtr<SVerticalBox> OutputBox) override;
+	virtual void UpdateErrorInfo() override;
 
 protected:
 	void RegisterNiagaraGraphNode(UEdGraphNode* InNode);
diff --git a/Engine/Plugins/FX/Niagara/Source/NiagaraEditor/Public/NiagaraEditorSettings.h b/Engine/Plugins/FX/Niagara/Source/NiagaraEditor/Public/NiagaraEditorSettings.h
--- a/Engine/Plugins/FX/Niagara/Source/NiagaraEditor/Public/NiagaraEditorSettings.h
+++ b/Engine/Plugins/FX/Niagara/Source/NiagaraEditor/Public/NiagaraEditorSettings.h
@@ -360,6 +360,17 @@ public:
 
 	FNiagaraParameterPanelSectionStorage& FindOrAddParameterPanelSectionStorage(FGuid PanelSectionId, bool& bOutAdded);
 
+	DECLARE_DELEGATE_RetVal_OneParam(bool, FOnIsClassAllowed, const UClass* /*InClass*/);
+	
+	/** Sets a delegate that allows external code to restrict which features can be used within the niagara editor by filtering which classes are allowed. */
+	void SetOnIsClassAllowed(const FOnIsClassAllowed& InOnIsClassAllowed);
+
+	/** Returns whether or not the supplied class can be used in the current editor context. */
+	bool IsAllowedClass(const UClass* InClass) const;
+
+	/** Returns whether or not the supplied niagara type definition can be used in the current editor context. */
+	bool IsAllowedTypeDefinition(const FNiagaraTypeDefinition& InTypeDefinition) const;
+
 private:
 	void SetupNamespaceMetadata();
 	void BuildCachedPlaybackSpeeds() const;
@@ -448,6 +459,8 @@ private:
 	UPROPERTY(config)
 	TArray<FNiagaraParameterPanelSectionStorage> SystemParameterPanelSectionData;
 
+	FOnIsClassAllowed OnIsClassAllowedDelegate;
+
 public:
 	bool IsShowInstructionsCount() const;
 	void SetShowInstructionsCount(bool bShowInstructionsCount);
diff --git a/Engine/Plugins/FX/Niagara/Source/NiagaraEditor/Public/NiagaraEditorUtilities.h b/Engine/Plugins/FX/Niagara/Source/NiagaraEditor/Public/NiagaraEditorUtilities.h
--- a/Engine/Plugins/FX/Niagara/Source/NiagaraEditor/Public/NiagaraEditorUtilities.h
+++ b/Engine/Plugins/FX/Niagara/Source/NiagaraEditor/Public/NiagaraEditorUtilities.h
@@ -367,6 +367,33 @@ namespace FNiagaraEditorUtilities
 
 	NIAGARAEDITOR_API int GetReferencedAssetCount(const FAssetData& SourceAsset, TFunction<ETrackAssetResult(const FAssetData&)> Predicate);
 	
+	/** Gets a list of the registered types which are allowed in the current editor context.  This API should be 
+		called when providing a list types to the user instead of getting the type list directly from the type registry. */
+	void GetAllowedTypes(TArray<FNiagaraTypeDefinition>& OutAllowedTypes);
+
+	/** Gets a list of the registered user variable types which are allowed in the current editor context.  This API should be
+		called when providing a list types to the user instead of getting the type list directly from the type registry. */
+	void GetAllowedUserVariableTypes(TArray<FNiagaraTypeDefinition>& OutAllowedTypes);
+
+	/** Gets a list of the registered system variable types which are allowed in the current editor context.  This API should be
+		called when providing a list types to the user instead of getting the type list directly from the type registry. */
+	void GetAllowedSystemVariableTypes(TArray<FNiagaraTypeDefinition>& OutAllowedTypes);
+
+	/** Gets a list of the registered emitter variable types which are allowed in the current editor context.  This API should be
+		called when providing a list types to the user instead of getting the type list directly from the type registry. */
+	void GetAllowedEmitterVariableTypes(TArray<FNiagaraTypeDefinition>& OutAllowedTypes);
+
+	/** Gets a list of the registered particle variable types which are allowed in the current editor context.  This API should be
+		called when providing a list types to the user instead of getting the type list directly from the type registry. */
+	void GetAllowedParticleVariableTypes(TArray<FNiagaraTypeDefinition>& OutAllowedTypes);
+
+	/** Gets a list of the registered parameter types which are allowed in the current editor context.  This API should be
+		called when providing a list types to the user instead of getting the type list directly from the type registry. */
+	void GetAllowedParameterTypes(TArray<FNiagaraTypeDefinition>& OutAllowedTypes);
+
+	/** Gets a list of the registered payload types which are allowed in the current editor context.  This API should be
+		called when providing a list types to the user instead of getting the type list directly from the type registry. */
+	void GetAllowedPayloadTypes(TArray<FNiagaraTypeDefinition>& OutAllowedTypes);
 };
 
 namespace FNiagaraParameterUtilities
diff --git a/Engine/Plugins/FX/Niagara/Source/NiagaraEditor/Public/ViewModels/Stack/NiagaraStackFunctionInput.h b/Engine/Plugins/FX/Niagara/Source/NiagaraEditor/Public/ViewModels/Stack/NiagaraStackFunctionInput.h
--- a/Engine/Plugins/FX/Niagara/Source/NiagaraEditor/Public/ViewModels/Stack/NiagaraStackFunctionInput.h
+++ b/Engine/Plugins/FX/Niagara/Source/NiagaraEditor/Public/ViewModels/Stack/NiagaraStackFunctionInput.h
@@ -284,6 +284,9 @@ public:
 	virtual const FCollectedUsageData& GetCollectedUsageData() const override;
 
 	bool OpenSourceAsset() const;
+
+	bool SupportsCustomExpressions() const;
+
 protected:
 	//~ UNiagaraStackEntry interface
 	virtual void FinalizeInternal() override;
diff --git a/Engine/Plugins/FX/Niagara/Source/NiagaraEditor/Public/Widgets/SNiagaraParameterMenu.cpp b/Engine/Plugins/FX/Niagara/Source/NiagaraEditor/Public/Widgets/SNiagaraParameterMenu.cpp
--- a/Engine/Plugins/FX/Niagara/Source/NiagaraEditor/Public/Widgets/SNiagaraParameterMenu.cpp
+++ b/Engine/Plugins/FX/Niagara/Source/NiagaraEditor/Public/Widgets/SNiagaraParameterMenu.cpp
@@ -155,26 +155,26 @@ void SNiagaraAddParameterFromPanelMenu::CollectMakeNew(FNiagaraMenuActionCollect
 		return;
 	}
 	
-	TConstArrayView<FNiagaraTypeDefinition> SectionTypes;
+	TArray<FNiagaraTypeDefinition> SectionTypes;
 	if(InNamespaceId == FNiagaraEditorGuids::UserNamespaceMetaDataGuid) 
 	{
-		SectionTypes = MakeArrayView(FNiagaraTypeRegistry::GetRegisteredUserVariableTypes());
+		FNiagaraEditorUtilities::GetAllowedUserVariableTypes(SectionTypes);
 	}
 	else if(InNamespaceId == FNiagaraEditorGuids::SystemNamespaceMetaDataGuid)
 	{
-		SectionTypes = MakeArrayView(FNiagaraTypeRegistry::GetRegisteredSystemVariableTypes());
+		FNiagaraEditorUtilities::GetAllowedSystemVariableTypes(SectionTypes);
 	}
 	else if (InNamespaceId == FNiagaraEditorGuids::EmitterNamespaceMetaDataGuid)
 	{
-		SectionTypes = MakeArrayView(FNiagaraTypeRegistry::GetRegisteredEmitterVariableTypes());
+		FNiagaraEditorUtilities::GetAllowedEmitterVariableTypes(SectionTypes);
 	}
 	else if (InNamespaceId == FNiagaraEditorGuids::ParticleAttributeNamespaceMetaDataGuid)
 	{
-		SectionTypes = MakeArrayView(FNiagaraTypeRegistry::GetRegisteredParticleVariableTypes());
+		FNiagaraEditorUtilities::GetAllowedParticleVariableTypes(SectionTypes);
 	}
 	else
 	{
-		SectionTypes = MakeArrayView(FNiagaraTypeRegistry::GetRegisteredTypes());
+		FNiagaraEditorUtilities::GetAllowedTypes(SectionTypes);
 	}
 
 	TArray<FNiagaraTypeDefinition> TypeDefinitions;
@@ -827,7 +827,8 @@ void SNiagaraAddParameterFromPinMenu::Construct(const FArguments& InArgs)
 void SNiagaraAddParameterFromPinMenu::CollectAllActions(FGraphActionListBuilderBase& OutAllActions)
 {
 	FNiagaraMenuActionCollector Collector;
-	TArray<FNiagaraTypeDefinition> Types(FNiagaraTypeRegistry::GetRegisteredTypes());
+	TArray<FNiagaraTypeDefinition> Types;
+	FNiagaraEditorUtilities::GetAllowedTypes(Types);
 	Types.Sort([](const FNiagaraTypeDefinition& A, const FNiagaraTypeDefinition& B) { return (A.GetNameText().ToLower().ToString() < B.GetNameText().ToLower().ToString()); });
 
 	for (const FNiagaraTypeDefinition& RegisteredType : Types)
@@ -875,7 +876,8 @@ void SNiagaraChangePinTypeMenu::CollectAllActions(FGraphActionListBuilderBase& O
 	UNiagaraNode* Node = Cast<UNiagaraNode>(PinToModify->GetOwningNode());
 	checkf(Node, TEXT("Niagara node pin did not have a valid outer node!"));
 
-	TArray<FNiagaraTypeDefinition> Types(FNiagaraTypeRegistry::GetRegisteredTypes());
+	TArray<FNiagaraTypeDefinition> Types;
+	FNiagaraEditorUtilities::GetAllowedTypes(Types);
 	Types.Sort([](const FNiagaraTypeDefinition& A, const FNiagaraTypeDefinition& B) { return (A.GetNameText().ToLower().ToString() < B.GetNameText().ToLower().ToString()); });
 
 	for (const FNiagaraTypeDefinition& RegisteredType : Types)
diff --git a/Engine/Plugins/FX/Niagara/Source/NiagaraEditorWidgets/Private/Stack/SNiagaraStackFunctionInputValue.cpp b/Engine/Plugins/FX/Niagara/Source/NiagaraEditorWidgets/Private/Stack/SNiagaraStackFunctionInputValue.cpp
--- a/Engine/Plugins/FX/Niagara/Source/NiagaraEditorWidgets/Private/Stack/SNiagaraStackFunctionInputValue.cpp
+++ b/Engine/Plugins/FX/Niagara/Source/NiagaraEditorWidgets/Private/Stack/SNiagaraStackFunctionInputValue.cpp
@@ -1217,7 +1217,7 @@ TArray<TSharedPtr<FNiagaraMenuAction_Generic>> SNiagaraStackFunctionInputValue::
 		}
 	}
 
-	if (bIsDataInterfaceOrObject == false)
+	if (bIsDataInterfaceOrObject == false && FunctionInput->SupportsCustomExpressions())
 	{
 		// Leaving the internal usage of bIsDataInterfaceObject that the tooltip and disabling will work properly when they're moved out of a graph action menu.
 		const FText DisplayName = LOCTEXT("ExpressionLabel", "New Expression");
