- Fix for arrays and LWC types as we didn't copy the user facing values over

#rnx

diff --git a/Engine/Plugins/FX/Niagara/Source/Niagara/Classes/NiagaraDataInterfaceArrayImpl.h b/Engine/Plugins/FX/Niagara/Source/Niagara/Classes/NiagaraDataInterfaceArrayImpl.h
--- a/Engine/Plugins/FX/Niagara/Source/Niagara/Classes/NiagaraDataInterfaceArrayImpl.h
+++ b/Engine/Plugins/FX/Niagara/Source/Niagara/Classes/NiagaraDataInterfaceArrayImpl.h
@@ -45,6 +45,27 @@
 			Super::PostEditChangeProperty(PropertyChangedEvent); \
 			GetProxyAs<FProxyType>()->template SetArrayData<decltype(MEMBERNAME)::ElementType>(MakeArrayView(MEMBERNAME)); \
 		} \
+		virtual bool CopyToInternal(UNiagaraDataInterface* Destination) const override \
+		{ \
+			if ( Super::CopyToInternal(Destination) == false ) \
+			{ \
+				return false; \
+			} \
+			CLASSNAME* TypedDestination = Cast<CLASSNAME>(Destination); \
+			if ( TypedDestination != nullptr ) \
+			{ \
+				TypedDestination->MEMBERNAME = MEMBERNAME; \
+			} \
+			return TypedDestination != nullptr;  \
+		} \
+		virtual bool Equals(const UNiagaraDataInterface* Other) const override \
+		{ \
+			const CLASSNAME* TypedOther = Cast<const CLASSNAME>(Other); \
+			return \
+				Super::Equals(Other) && \
+				TypedOther != nullptr && \
+				TypedOther->MEMBERNAME == MEMBERNAME; \
+		} \
 		TArray<TYPENAME>& GetArrayReference() { return Internal##MEMBERNAME; }
 #else
 	#define NDIARRAY_GENERATE_BODY_LWC(CLASSNAME, TYPENAME, MEMBERNAME) NDIARRAY_GENERATE_BODY(CLASSNAME, TYPENAME, Internal##MEMBERNAME)
