- Add checkf for null data interface in GPU parameter store
- Make warning about DI mismatch more clear

diff --git a/Engine/Plugins/FX/Niagara/Source/Niagara/Classes/NiagaraComputeExecutionContext.h b/Engine/Plugins/FX/Niagara/Source/Niagara/Classes/NiagaraComputeExecutionContext.h
--- a/Engine/Plugins/FX/Niagara/Source/Niagara/Classes/NiagaraComputeExecutionContext.h
+++ b/Engine/Plugins/FX/Niagara/Source/Niagara/Classes/NiagaraComputeExecutionContext.h
@@ -124,7 +124,7 @@ public:
 	//TArray<uint8, TAlignedHeapAllocator<16>> ParamData_RT;		// RT side copy of the parameter data
 	FNiagaraScriptInstanceParameterStore CombinedParamStore;
 #if DO_CHECK
-	TArray< FString >  DIClassNames;
+	TArray<FName>  DIClassNames;
 #endif
 
 	TArray<FNiagaraDataInterfaceProxy*> DataInterfaceProxies;
diff --git a/Engine/Plugins/FX/Niagara/Source/Niagara/Private/NiagaraComputeExecutionContext.cpp b/Engine/Plugins/FX/Niagara/Source/Niagara/Private/NiagaraComputeExecutionContext.cpp
--- a/Engine/Plugins/FX/Niagara/Source/Niagara/Private/NiagaraComputeExecutionContext.cpp
+++ b/Engine/Plugins/FX/Niagara/Source/Niagara/Private/NiagaraComputeExecutionContext.cpp
@@ -66,7 +66,7 @@ void FNiagaraComputeExecutionContext::InitParams(UNiagaraScript* InGPUComputeScr
 		DIClassNames.Empty(Shader->GetDIParameters().Num());
 		for (const FNiagaraDataInterfaceParamRef& DIParams : Shader->GetDIParameters())
 		{
-			DIClassNames.Add(DIParams.DIType.Get(Shader.GetPointerTable().DITypes)->GetClass()->GetName());
+			DIClassNames.Add(DIParams.DIType.Get(Shader.GetPointerTable().DITypes)->GetClass()->GetFName());
 		}
 	}
 	else
@@ -75,7 +75,7 @@ void FNiagaraComputeExecutionContext::InitParams(UNiagaraScript* InGPUComputeScr
 		DIClassNames.Empty(ScriptParametersMetadata->DataInterfaceParamInfo.Num());
 		for (const FNiagaraDataInterfaceGPUParamInfo& DIParams : ScriptParametersMetadata->DataInterfaceParamInfo)
 		{
-			DIClassNames.Add(DIParams.DIClassName);
+			DIClassNames.Emplace(DIParams.DIClassName);
 		}
 	}
 #endif
@@ -147,13 +147,15 @@ bool FNiagaraComputeExecutionContext::Tick(FNiagaraSystemInstance* ParentSystemI
 			return false;
 		}
 
-		for (int32 i = 0; i < DIClassNames.Num(); ++i)
+		for (int32 i=0; i < DIClassNames.Num(); ++i)
 		{
-			FString UsedClassName = DataInterfaces[i]->GetClass()->GetName();
+			UNiagaraDataInterface* UsedDataInterface = DataInterfaces[i];
+			checkf(UsedDataInterface, TEXT("DataInterface for GPU execution context is nullptr. System(%s) Index(%d) ExpectedType(%s)"), *GetNameSafe(ParentSystemInstance->GetSystem()), i, *DIClassNames[i].ToString());
+
+			const FName UsedClassName = UsedDataInterface->GetClass()->GetFName();
 			if (DIClassNames[i] != UsedClassName)
 			{
-				UE_LOG(LogNiagara, Warning, TEXT("Mismatched class between Niagara GPU Execution Context data interfaces and those in its script!\nIndex:%d\nShader:%s\nScript:%s")
-					, i, *DIClassNames[i], *UsedClassName);
+				UE_LOG(LogNiagara, Warning, TEXT("Mismatched class between Niagara GPU Execution Context data interfaces and those in its script! System(%s) Index(%d) ExpectedType(%s) CombinedParamStoreType(%s)"), *GetNameSafe(ParentSystemInstance->GetSystem()), i, *DIClassNames[i].ToString(), *UsedClassName.ToString());
 			}
 		}
 #endif
