- Fix PVS issue ParentSystemInstance is never null

#okforgithuib public
#rnx

diff --git a/Engine/Plugins/FX/Niagara/Source/Niagara/Private/NiagaraComputeExecutionContext.cpp b/Engine/Plugins/FX/Niagara/Source/Niagara/Private/NiagaraComputeExecutionContext.cpp
--- a/Engine/Plugins/FX/Niagara/Source/Niagara/Private/NiagaraComputeExecutionContext.cpp
+++ b/Engine/Plugins/FX/Niagara/Source/Niagara/Private/NiagaraComputeExecutionContext.cpp
@@ -136,6 +136,7 @@ void FNiagaraComputeExecutionContext::DirtyDataInterfaces()
 
 bool FNiagaraComputeExecutionContext::Tick(FNiagaraSystemInstance* ParentSystemInstance)
 {
+	check(ParentSystemInstance);
 	if (CombinedParamStore.GetInterfacesDirty())
 	{
 #if DO_CHECK
@@ -159,7 +160,7 @@ bool FNiagaraComputeExecutionContext::Tick(FNiagaraSystemInstance* ParentSystemI
 			}
 		}
 #endif
-		if (ParentSystemInstance && CombinedParamStore.GetPositionDataDirty())
+		if (CombinedParamStore.GetPositionDataDirty())
 		{
 			CombinedParamStore.ResolvePositions(ParentSystemInstance->GetLWCConverter());
 		}
