Fix for StackFuncitonInput cache not populating the source of the data so that we can reuse it on subsequent calls

diff --git a/Engine/Plugins/FX/Niagara/Source/NiagaraEditor/Private/NiagaraGraphDataCache.cpp b/Engine/Plugins/FX/Niagara/Source/NiagaraEditor/Private/NiagaraGraphDataCache.cpp
--- a/Engine/Plugins/FX/Niagara/Source/NiagaraEditor/Private/NiagaraGraphDataCache.cpp
+++ b/Engine/Plugins/FX/Niagara/Source/NiagaraEditor/Private/NiagaraGraphDataCache.cpp
@@ -4,6 +4,15 @@
 
 #include "NiagaraNodeFunctionCall.h"
 
+void FNiagaraGraphDataCache::FCachedStackFunctionInputPins::SetSource(UNiagaraNodeFunctionCall& FunctionCallNode)
+{
+	SourceGraphWeak = FunctionCallNode.GetNiagaraGraph();
+	LastSourceGraphChangeId = FunctionCallNode.GetNiagaraGraph()->GetChangeID();
+	FunctionCallNodeWeak = &FunctionCallNode;
+	CalledGraphWeak = FunctionCallNode.GetCalledGraph();
+	LastCalledGraphChangeId = FunctionCallNode.GetCalledGraph() != nullptr ? FunctionCallNode.GetCalledGraph()->GetChangeID() : FGuid();
+}
+
 void FNiagaraGraphDataCache::GetStackFunctionInputPins(UNiagaraNodeFunctionCall& FunctionCallNode, TConstArrayView<FNiagaraVariable> StaticVars, TArray<const UEdGraphPin*>& OutInputPins, FCompileConstantResolver ConstantResolver, FNiagaraStackGraphUtilities::ENiagaraGetStackFunctionInputPinsOptions Options, bool bIgnoreDisabled)
 {
 	FGetStackFunctionInputPinsKey Key(FunctionCallNode, ConstantResolver, Options, bIgnoreDisabled, false);
@@ -18,6 +27,8 @@ void FNiagaraGraphDataCache::GetStackFunctionInputPins(UNiagaraNodeFunctionCall&
 			Options,
 			bIgnoreDisabled,
 			false /*bFilterForCompilation*/);
+
+		CachedResult.SetSource(FunctionCallNode);
 	}
 
 	OutInputPins = CachedResult.InputPins;
@@ -51,6 +62,8 @@ void FNiagaraGraphDataCache::GetStackFunctionInputPins(UNiagaraNodeFunctionCall&
 			Options,
 			bIgnoreDisabled,
 			false /*bFilterForCompilation*/);
+
+		CachedResult.SetSource(FunctionCallNode);
 	}
 
 	if (!CompilationCachedResult.IsValid())
@@ -63,6 +76,8 @@ void FNiagaraGraphDataCache::GetStackFunctionInputPins(UNiagaraNodeFunctionCall&
 			Options,
 			bIgnoreDisabled,
 			true /*bFilterForCompilation*/);
+
+		CompilationCachedResult.SetSource(FunctionCallNode);
 	}
 
 	OutInputPins = CachedResult.InputPins;
diff --git a/Engine/Plugins/FX/Niagara/Source/NiagaraEditor/Public/NiagaraGraphDataCache.h b/Engine/Plugins/FX/Niagara/Source/NiagaraEditor/Public/NiagaraGraphDataCache.h
--- a/Engine/Plugins/FX/Niagara/Source/NiagaraEditor/Public/NiagaraGraphDataCache.h
+++ b/Engine/Plugins/FX/Niagara/Source/NiagaraEditor/Public/NiagaraGraphDataCache.h
@@ -61,13 +61,18 @@ private:
 
 	struct FCachedStackFunctionInputPins
 	{
+		void SetSource(UNiagaraNodeFunctionCall& InFunctionCallNode);
+		bool IsValid() const;
+
+		TArray<const UEdGraphPin*> InputPins;
+
+	private:
 		TWeakObjectPtr<UNiagaraGraph> SourceGraphWeak;
-		FGuid LastSourceGraphChangeId;
 		TWeakObjectPtr<UNiagaraNodeFunctionCall> FunctionCallNodeWeak;
 		TWeakObjectPtr<UNiagaraGraph> CalledGraphWeak;
+
+		FGuid LastSourceGraphChangeId;
 		FGuid LastCalledGraphChangeId;
-		TArray<const UEdGraphPin*> InputPins;
-		bool IsValid() const;
 	};
 
 	TMap<FGetStackFunctionInputPinsKey, FCachedStackFunctionInputPins> GetStackFunctionInputPinsCache;
