- Fix ribbons not respecting RDG deferred execution correctly for GPU buffers
- Fix ribbons not working with low latency GPU data

diff --git a/Engine/Plugins/FX/Niagara/Source/Niagara/Private/NiagaraRendererRibbons.cpp b/Engine/Plugins/FX/Niagara/Source/Niagara/Private/NiagaraRendererRibbons.cpp
--- a/Engine/Plugins/FX/Niagara/Source/Niagara/Private/NiagaraRendererRibbons.cpp
+++ b/Engine/Plugins/FX/Niagara/Source/Niagara/Private/NiagaraRendererRibbons.cpp
@@ -348,6 +348,22 @@ struct FNiagaraRibbonRenderingFrameResources
 	}	
 };
 
+struct FNiagaraRibbonGPUInitParameters
+{
+	FNiagaraRibbonGPUInitParameters(const FNiagaraRendererRibbons* InRenderer, const FNiagaraDataBuffer* InSourceParticleData, const TSharedPtr<FNiagaraRibbonRenderingFrameResources>& InRenderingResources)
+		: Renderer(InRenderer)
+		, NumInstances(InSourceParticleData->GetNumInstances())
+		, GPUInstanceCountBufferOffset(InSourceParticleData->GetGPUInstanceCountBufferOffset())
+		, RenderingResources(InRenderingResources)
+	{
+	}
+
+	const FNiagaraRendererRibbons*	Renderer = nullptr;
+	const uint32					NumInstances = 0;
+	const uint32					GPUInstanceCountBufferOffset = 0;
+	TWeakPtr<FNiagaraRibbonRenderingFrameResources> RenderingResources;
+};
+
 class FNiagaraRibbonMeshCollectorResources : public FOneFrameResource
 {
 public:
@@ -447,21 +463,6 @@ void FNiagaraRibbonVertexBuffers::InitializeOrUpdateBuffers(const FNiagaraRibbon
 	}
 }
 
-struct FNiagaraRibbonGPUInitParameters
-{
-	const FNiagaraRendererRibbons* Renderer;
-	const FNiagaraDataBuffer* SourceParticleData;
-	TWeakPtr<FNiagaraRibbonRenderingFrameResources> RenderingResources;
-	
-	FNiagaraRibbonGPUInitParameters(const FNiagaraRendererRibbons* InRenderer, const FNiagaraDataBuffer* InSourceParticleData, const TSharedPtr<FNiagaraRibbonRenderingFrameResources>& InRenderingResources)
-		: Renderer(InRenderer)
-		, SourceParticleData(InSourceParticleData)
-		, RenderingResources(InRenderingResources)
-	{
-		
-	}	 
-};
-
 struct FNiagaraRibbonGPUInitComputeBuffers
 {
 	FRWBuffer SortBuffer;
@@ -515,6 +516,7 @@ public:
 	{
 		FGPUSortManager* SortManager = InOwnerInterface->GetGPUSortManager();
 		SortManager->PostPreRenderEvent.AddRaw(this, &FNiagaraGpuRibbonsDataManager::OnPostPreRender);
+		SortManager->PostPostRenderEvent.AddRaw(this, &FNiagaraGpuRibbonsDataManager::OnPostPostRender);
 	}
 
 	virtual ~FNiagaraGpuRibbonsDataManager()
@@ -527,22 +529,28 @@ public:
 		return ManagerName;
 	}
 
-	void RegisterRenderer(FNiagaraRibbonGPUInitParameters NewRegistration)
+	void RegisterRenderer(const FNiagaraRendererRibbons* InRenderer, const FNiagaraDataBuffer* InSourceParticleData, const TSharedPtr<FNiagaraRibbonRenderingFrameResources>& InRenderingResources)
 	{
-		RenderersToGenerate.Add(NewRegistration);
+		const int32 GenerateIndex = InSourceParticleData->GetGPUDataReadyStage() == ENiagaraGpuComputeTickStage::PostOpaqueRender ? 1 : 0;
+		RenderersToGeneratePerStage[GenerateIndex].Emplace(InRenderer, InSourceParticleData, InRenderingResources);
 	}
 
 private:
-	TArray<FNiagaraRibbonGPUInitParameters> RenderersToGenerate;
+	TArray<FNiagaraRibbonGPUInitParameters> RenderersToGeneratePerStage[2];
 
 	FNiagaraRibbonGPUInitComputeBuffers ComputeBuffers;
 	
 	void OnPostPreRender(FRHICommandListImmediate& RHICmdList)
 	{
-		GenerateAllGPUData(RHICmdList);
+		GenerateAllGPUData(RHICmdList, RenderersToGeneratePerStage[0]);
 	}
 
-	void GenerateAllGPUData(FRHICommandListImmediate& RHICmdList);
+	void OnPostPostRender(FRHICommandListImmediate& RHICmdList)
+	{
+		GenerateAllGPUData(RHICmdList, RenderersToGeneratePerStage[1]);
+	}
+
+	void GenerateAllGPUData(FRHICommandListImmediate& RHICmdList, TArray<FNiagaraRibbonGPUInitParameters>& RenderersToGenerate);
 };
 
 
@@ -720,7 +728,7 @@ void FNiagaraRendererRibbons::GetDynamicMeshElements(const TArray<const FSceneVi
 	{
 		FNiagaraGpuComputeDispatchInterface* ComputeDispatchInterface = SceneProxy->GetComputeDispatchInterface();
 		FNiagaraGpuRibbonsDataManager& RibbonDataManager = ComputeDispatchInterface->GetOrCreateDataManager<FNiagaraGpuRibbonsDataManager>();
-		RibbonDataManager.RegisterRenderer(FNiagaraRibbonGPUInitParameters(this, SourceParticleData, RenderingResources.RibbonResources));
+		RibbonDataManager.RegisterRenderer(this, SourceParticleData, RenderingResources.RibbonResources);
 	}
 }
 
@@ -1968,19 +1976,7 @@ inline void FNiagaraRendererRibbons::SetupMeshBatchAndCollectorResourceForView(c
 	MeshElement.PrimitiveUniformBuffer = SceneProxy->GetCustomUniformBuffer(false);	// Note: Ribbons don't generate accurate velocities so disabling	
 }
 
-template<typename Type>
-static inline TArray<Type> ReadBackBuffer(FRHIBuffer* Buffer, int32 NumValues)
-{
-	TArray<Type> Test;
-	Test.SetNum(NumValues);	
-	const void* DebugDataIntPtr = RHILockBuffer(Buffer, 0, NumValues * sizeof(Type), RLM_ReadOnly);
-	FMemory::Memcpy(Test.GetData(), DebugDataIntPtr, NumValues * sizeof(Type));
-	RHIUnlockBuffer(Buffer);
-	return Test;
-}
-
-void FNiagaraRendererRibbons::InitializeViewIndexBuffersGPU(FRHICommandListImmediate& RHICmdList, FNiagaraGpuComputeDispatchInterface* ComputeDispatchInterface,
-	const FNiagaraDataBuffer* SourceParticleData, const TSharedPtr<FNiagaraRibbonRenderingFrameViewResources>& RenderingViewResources) const
+void FNiagaraRendererRibbons::InitializeViewIndexBuffersGPU(FRHICommandListImmediate& RHICmdList, FNiagaraGpuComputeDispatchInterface* ComputeDispatchInterface, const FNiagaraRibbonGPUInitParameters& GpuInitParameters, const TSharedPtr<FNiagaraRibbonRenderingFrameViewResources>& RenderingViewResources) const
 {
 	SCOPE_CYCLE_COUNTER(STAT_NiagaraRenderRibbonsGenIndiciesGPU);
 	
@@ -1989,7 +1985,7 @@ void FNiagaraRendererRibbons::InitializeViewIndexBuffersGPU(FRHICommandListImmed
 		return;
 	}
 
-	const uint32 NumInstances = SourceParticleData->GetNumInstances();
+	const uint32 NumInstances = GpuInitParameters.NumInstances;
 
 	SCOPED_DRAW_EVENT(RHICmdList, NiagaraRenderRibbonsGenIndiciesGPU);
 	{
@@ -2010,7 +2006,7 @@ void FNiagaraRendererRibbons::InitializeViewIndexBuffersGPU(FRHICommandListImmed
 
 		// Indirect particle Count
 		Params.EmitterParticleCountsBuffer = GetSrvOrDefaultUInt(ComputeDispatchInterface->GetGPUInstanceCounterManager().GetInstanceCountBuffer());
-		Params.EmitterParticleCountsBufferOffset = SourceParticleData->GetGPUInstanceCountBufferOffset();
+		Params.EmitterParticleCountsBufferOffset = GpuInitParameters.GPUInstanceCountBufferOffset;
 
 		Params.IndirectDrawOutputIndex = 0;
 		Params.VertexGenerationResultsIndex = 0; /*Offset into command buffer*/
@@ -2069,11 +2065,11 @@ void FNiagaraRendererRibbons::InitializeViewIndexBuffersGPU(FRHICommandListImmed
 		Params.TriangleToVertexIds = ShapeState.SliceTriangleToVertexIdsBuffer.SRV;
 
 		// Total particle Count
-		Params.TotalNumParticlesDirect = SourceParticleData->GetNumInstances();
+		Params.TotalNumParticlesDirect = GpuInitParameters.NumInstances;
 
 		// Indirect particle Count
 		Params.EmitterParticleCountsBuffer = GetSrvOrDefaultUInt(ComputeDispatchInterface->GetGPUInstanceCounterManager().GetInstanceCountBuffer());
-		Params.EmitterParticleCountsBufferOffset = SourceParticleData->GetGPUInstanceCountBufferOffset();
+		Params.EmitterParticleCountsBufferOffset = GpuInitParameters.GPUInstanceCountBufferOffset;
 
 		Params.IndexBufferOffset = 0;
 		Params.IndirectDrawInfoIndex = 0;
@@ -2178,18 +2174,17 @@ void FNiagaraRendererRibbons::InitializeVertexBuffersResources(const FNiagaraDyn
 	}
 }
 
-FRibbonComputeUniformParameters FNiagaraRendererRibbons::SetupComputeVertexGenParams(FNiagaraGpuComputeDispatchInterface* ComputeDispatchInterface,
-	const TSharedPtr<FNiagaraRibbonRenderingFrameResources>& RenderingResources, const FNiagaraDataBuffer* SourceParticleData) const
+FRibbonComputeUniformParameters FNiagaraRendererRibbons::SetupComputeVertexGenParams(FNiagaraGpuComputeDispatchInterface* ComputeDispatchInterface, const TSharedPtr<FNiagaraRibbonRenderingFrameResources>& RenderingResources, const FNiagaraRibbonGPUInitParameters& GpuInitParameters) const
 {
 	FRibbonComputeUniformParameters CommonParams;
 	FMemory::Memzero(CommonParams);
 
 	// Total particle Count
-	CommonParams.TotalNumParticlesDirect = SourceParticleData->GetNumInstances();
+	CommonParams.TotalNumParticlesDirect = GpuInitParameters.NumInstances;
 
 	// Indirect particle Count
 	CommonParams.EmitterParticleCountsBuffer = GetSrvOrDefaultUInt(ComputeDispatchInterface->GetGPUInstanceCounterManager().GetInstanceCountBuffer());
-	CommonParams.EmitterParticleCountsBufferOffset = SourceParticleData->GetGPUInstanceCountBufferOffset();
+	CommonParams.EmitterParticleCountsBufferOffset = GpuInitParameters.GPUInstanceCountBufferOffset;
 
 	// Niagara sim data
 	CommonParams.NiagaraParticleDataFloat = GetSrvOrDefaultFloat(RenderingResources->ParticleFloatSRV);
@@ -2242,14 +2237,14 @@ FRibbonComputeUniformParameters FNiagaraRendererRibbons::SetupComputeVertexGenPa
 	return CommonParams;
 }
 
-void FNiagaraRendererRibbons::InitializeVertexBuffersGPU(FRHICommandListImmediate& RHICmdList, FNiagaraGpuComputeDispatchInterface* ComputeDispatchInterface, const FNiagaraDataBuffer* SourceParticleData,
+void FNiagaraRendererRibbons::InitializeVertexBuffersGPU(FRHICommandListImmediate& RHICmdList, FNiagaraGpuComputeDispatchInterface* ComputeDispatchInterface, const FNiagaraRibbonGPUInitParameters& GpuInitParameters,
 	FNiagaraRibbonGPUInitComputeBuffers& TempBuffers, const TSharedPtr<FNiagaraRibbonRenderingFrameResources>& RenderingResources) const
 {	
 	SCOPE_CYCLE_COUNTER(STAT_NiagaraRenderRibbonsGenVerticesGPU);
 
-	FRibbonComputeUniformParameters CommonParams = SetupComputeVertexGenParams(ComputeDispatchInterface, RenderingResources, SourceParticleData);
+	FRibbonComputeUniformParameters CommonParams = SetupComputeVertexGenParams(ComputeDispatchInterface, RenderingResources, GpuInitParameters);
 
-	const int32 NumExecutableInstances = SourceParticleData->GetNumInstances();
+	const int32 NumExecutableInstances = GpuInitParameters.NumInstances;
 
 	const bool bCanRun = NumExecutableInstances >= 2;
 	
@@ -2547,7 +2542,7 @@ void FNiagaraRendererRibbons::InitializeVertexBuffersGPU(FRHICommandListImmediat
 	}
 }
 
-void FNiagaraGpuRibbonsDataManager::GenerateAllGPUData(FRHICommandListImmediate& RHICmdList)
+void FNiagaraGpuRibbonsDataManager::GenerateAllGPUData(FRHICommandListImmediate& RHICmdList, TArray<FNiagaraRibbonGPUInitParameters>& RenderersToGenerate)
 {
 	if (RenderersToGenerate.Num() == 0)
 	{
@@ -2557,36 +2552,31 @@ void FNiagaraGpuRibbonsDataManager::GenerateAllGPUData(FRHICommandListImmediate&
 	FNiagaraGpuComputeDispatchInterface* ComputeDispatchInterface = GetOwnerInterface();
 
 	// Handle all vertex gens first
-	for (int32 Index = 0; Index < RenderersToGenerate.Num(); Index++)
+	for (const FNiagaraRibbonGPUInitParameters& RendererToGen : RenderersToGenerate)
 	{
-		const auto& Params = RenderersToGenerate[Index];
-
-		const auto RenderingResources = Params.RenderingResources.Pin();
+		const TSharedPtr<FNiagaraRibbonRenderingFrameResources> RenderingResources = RendererToGen.RenderingResources.Pin();
 		if (RenderingResources.IsValid())
 		{
-			const bool bIsGPUSim = Params.SourceParticleData->GetGPUInstanceCountBufferOffset() != INDEX_NONE;
-
 			ComputeBuffers.InitOrUpdateBuffers(
-				Params.SourceParticleData->GetNumInstances(),
-				Params.Renderer->GenerationConfig.HasRibbonIDs(),
-				Params.Renderer->GenerationConfig.WantsAutomaticTessellation(),
-				Params.Renderer->GenerationConfig.HasTwist()
+				RendererToGen.NumInstances,
+				RendererToGen.Renderer->GenerationConfig.HasRibbonIDs(),
+				RendererToGen.Renderer->GenerationConfig.WantsAutomaticTessellation(),
+				RendererToGen.Renderer->GenerationConfig.HasTwist()
 			);
 
-			Params.Renderer->InitializeVertexBuffersGPU(RHICmdList, ComputeDispatchInterface, Params.SourceParticleData, ComputeBuffers, RenderingResources);
+			RendererToGen.Renderer->InitializeVertexBuffersGPU(RHICmdList, ComputeDispatchInterface, RendererToGen, ComputeBuffers, RenderingResources);
 		}
 	}
 
 	// Now handle all index gens
-	for (const auto& RendererToGen : RenderersToGenerate)
+	for (const FNiagaraRibbonGPUInitParameters& RendererToGen : RenderersToGenerate)
 	{
-		const auto RenderingResources = RendererToGen.RenderingResources.Pin();
+		const TSharedPtr<FNiagaraRibbonRenderingFrameResources> RenderingResources = RendererToGen.RenderingResources.Pin();
 		if (RenderingResources.IsValid())
 		{
-			for (int32 Index = 0; Index < RenderingResources->ViewResources.Num(); Index++)
+			for (const TSharedPtr<FNiagaraRibbonRenderingFrameViewResources>& RenderingResourcesView : RenderingResources->ViewResources)
 			{
-				const auto& RenderingResourcesView = RenderingResources->ViewResources[Index];
-				RendererToGen.Renderer->InitializeViewIndexBuffersGPU(RHICmdList, ComputeDispatchInterface, RendererToGen.SourceParticleData, RenderingResourcesView);
+				RendererToGen.Renderer->InitializeViewIndexBuffersGPU(RHICmdList, ComputeDispatchInterface, RendererToGen, RenderingResourcesView);
 			}
 		}
 	}
diff --git a/Engine/Plugins/FX/Niagara/Source/Niagara/Public/NiagaraRendererRibbons.h b/Engine/Plugins/FX/Niagara/Source/Niagara/Public/NiagaraRendererRibbons.h
--- a/Engine/Plugins/FX/Niagara/Source/Niagara/Public/NiagaraRendererRibbons.h
+++ b/Engine/Plugins/FX/Niagara/Source/Niagara/Public/NiagaraRendererRibbons.h
@@ -22,7 +22,8 @@ class FNiagaraDataSet;
 struct FNiagaraDynamicDataRibbon;
 struct FNiagaraRibbonRenderingFrameResources;
 struct FNiagaraRibbonRenderingFrameViewResources;
-
+struct FNiagaraRibbonGPUInitParameters;
+struct FNiagaraRibbonGPUInitComputeBuffers;
 
 struct FNiagaraGenerationInputDataCPUAccessors
 {
@@ -402,16 +403,16 @@ protected:
 	                                               const TSharedPtr<FNiagaraRibbonRenderingFrameResources>& RenderingResources, const TSharedPtr<FNiagaraRibbonRenderingFrameViewResources>& RenderingViewResources, FMeshBatch& OutMeshBatch, bool bShouldUseGPUInitIndices) const;
 
 
-	void InitializeViewIndexBuffersGPU(FRHICommandListImmediate& RHICmdList, FNiagaraGpuComputeDispatchInterface* ComputeDispatchInterface, const FNiagaraDataBuffer* SourceParticleData,
+	void InitializeViewIndexBuffersGPU(FRHICommandListImmediate& RHICmdList, FNiagaraGpuComputeDispatchInterface* ComputeDispatchInterface, const FNiagaraRibbonGPUInitParameters& GpuInitParameters,
 		const TSharedPtr<FNiagaraRibbonRenderingFrameViewResources>& RenderingViewResources) const;
 
 	void InitializeVertexBuffersResources(const FNiagaraDynamicDataRibbon* DynamicDataRibbon, FNiagaraDataBuffer* SourceParticleData,
 	                                      FGlobalDynamicReadBuffer& DynamicReadBuffer, const TSharedPtr<FNiagaraRibbonRenderingFrameResources>& RenderingResources, bool bShouldUseGPUInit) const;
 	
-	void InitializeVertexBuffersGPU(FRHICommandListImmediate& RHICmdList, FNiagaraGpuComputeDispatchInterface* ComputeDispatchInterface, const FNiagaraDataBuffer* SourceParticleData,
-		struct FNiagaraRibbonGPUInitComputeBuffers& TempBuffers, const TSharedPtr<FNiagaraRibbonRenderingFrameResources>& RenderingResources) const;
+	void InitializeVertexBuffersGPU(FRHICommandListImmediate& RHICmdList, FNiagaraGpuComputeDispatchInterface* ComputeDispatchInterface, const FNiagaraRibbonGPUInitParameters& GpuInitParameters,
+		FNiagaraRibbonGPUInitComputeBuffers& TempBuffers, const TSharedPtr<FNiagaraRibbonRenderingFrameResources>& RenderingResources) const;
 
-	FRibbonComputeUniformParameters SetupComputeVertexGenParams(FNiagaraGpuComputeDispatchInterface* ComputeDispatchInterface, const TSharedPtr<FNiagaraRibbonRenderingFrameResources>& RenderingResources, const FNiagaraDataBuffer* SourceParticleData) const;
+	FRibbonComputeUniformParameters SetupComputeVertexGenParams(FNiagaraGpuComputeDispatchInterface* ComputeDispatchInterface, const TSharedPtr<FNiagaraRibbonRenderingFrameResources>& RenderingResources, const FNiagaraRibbonGPUInitParameters& GpuInitParameters) const;
 
 	FNiagaraRibbonGenerationConfig GenerationConfig;
 	
