- Add function to get spline length

diff --git a/Engine/Plugins/FX/Niagara/Shaders/Private/NiagaraDataInterfaceSplineTemplate.ush b/Engine/Plugins/FX/Niagara/Shaders/Private/NiagaraDataInterfaceSplineTemplate.ush
--- a/Engine/Plugins/FX/Niagara/Shaders/Private/NiagaraDataInterfaceSplineTemplate.ush
+++ b/Engine/Plugins/FX/Niagara/Shaders/Private/NiagaraDataInterfaceSplineTemplate.ush
@@ -216,3 +216,8 @@ void FindClosestUnitDistanceFromPositionWS_{ParameterName}(float3 PositionWS, ou
 {
 	U = EvaluateFindNearestPosition_{ParameterName}(U * {ParameterName}_SplineLength);
 }
+
+void GetSplineLength_{ParameterName}(out float Length)
+{
+	Length = {ParameterName}_SplineLength;
+}
diff --git a/Engine/Plugins/FX/Niagara/Source/Niagara/Classes/NiagaraDataInterfaceSpline.h b/Engine/Plugins/FX/Niagara/Source/Niagara/Classes/NiagaraDataInterfaceSpline.h
--- a/Engine/Plugins/FX/Niagara/Source/Niagara/Classes/NiagaraDataInterfaceSpline.h
+++ b/Engine/Plugins/FX/Niagara/Source/Niagara/Classes/NiagaraDataInterfaceSpline.h
@@ -237,6 +237,8 @@ public:
 	void GetLocalToWorld(FVectorVMExternalFunctionContext& Context);
 	void GetLocalToWorldInverseTransposed(FVectorVMExternalFunctionContext& Context);
 
+	void VMGetSplineLength(FVectorVMExternalFunctionContext& Context);
+
 	virtual int32 PerInstanceDataPassedToRenderThreadSize() const override
 	{
 		return 0;
diff --git a/Engine/Plugins/FX/Niagara/Source/Niagara/Private/NiagaraDataInterfaceSpline.cpp b/Engine/Plugins/FX/Niagara/Source/Niagara/Private/NiagaraDataInterfaceSpline.cpp
--- a/Engine/Plugins/FX/Niagara/Source/Niagara/Private/NiagaraDataInterfaceSpline.cpp
+++ b/Engine/Plugins/FX/Niagara/Source/Niagara/Private/NiagaraDataInterfaceSpline.cpp
@@ -63,9 +63,10 @@ namespace NDISplineLocal
 	static const FName SampleSplineTangentByUnitDistanceName("SampleSplineTangentByUnitDistance");
 	static const FName SampleSplineTangentByUnitDistanceWSName("SampleSplineTangentByUnitDistanceWS");
 
-
 	static const FName FindClosestUnitDistanceFromPositionWSName("FindClosestUnitDistanceFromPositionWS");
 
+	static const FName GetSplineLengthName("GetSplineLength");
+
 	/** Temporary solution for exposing the transform of a mesh. Ideally this would be done by allowing interfaces to add to the uniform set for a simulation. */
 	static const FName GetSplineLocalToWorldName("GetSplineLocalToWorld");
 	static const FName GetSplineLocalToWorldInverseTransposedName("GetSplineLocalToWorldInverseTransposed");
@@ -107,188 +108,134 @@ void UNiagaraDataInterfaceSpline::PostInitProperties()
 
 void UNiagaraDataInterfaceSpline::GetFunctions(TArray<FNiagaraFunctionSignature>& OutFunctions)
 {
+	FNiagaraFunctionSignature DefaultSig;
+	DefaultSig.Inputs.Add(FNiagaraVariable(FNiagaraTypeDefinition(GetClass()), TEXT("Spline")));
+	DefaultSig.bMemberFunction = true;
+	DefaultSig.bRequiresContext = false;
+	DefaultSig.SetFunctionVersion(FNiagaraSplineDIFunctionVersion::LatestVersion);
+
 	{
-		FNiagaraFunctionSignature Sig;
+		FNiagaraFunctionSignature& Sig = OutFunctions.Add_GetRef(DefaultSig);
 		Sig.Name = NDISplineLocal::SampleSplinePositionByUnitDistanceName;
-		Sig.Inputs.Add(FNiagaraVariable(FNiagaraTypeDefinition(GetClass()), TEXT("Spline")));
 		Sig.Inputs.Add(FNiagaraVariable(FNiagaraTypeDefinition::GetFloatDef(), TEXT("U")));
 		Sig.Outputs.Add(FNiagaraVariable(FNiagaraTypeDefinition::GetVec3Def(), TEXT("Position")));
-		Sig.bMemberFunction = true;
-		Sig.bRequiresContext = false;
 		Sig.SetDescription(LOCTEXT("DataInterfaceSpline_SampleSplinePositionByUnitDistance", "Sample the spline Position where U is a 0 to 1 value representing the start and normalized length of the spline.\nThis is in the local space of the referenced USplineComponent."));
-		OutFunctions.Add(Sig);
 	}
 
 	{
-		FNiagaraFunctionSignature Sig;
+		FNiagaraFunctionSignature& Sig = OutFunctions.Add_GetRef(DefaultSig);
 		Sig.Name = NDISplineLocal::SampleSplinePositionByUnitDistanceWSName;
-		Sig.Inputs.Add(FNiagaraVariable(FNiagaraTypeDefinition(GetClass()), TEXT("Spline")));
 		Sig.Inputs.Add(FNiagaraVariable(FNiagaraTypeDefinition::GetFloatDef(), TEXT("U")));
 		Sig.Outputs.Add(FNiagaraVariable(FNiagaraTypeDefinition::GetPositionDef(), TEXT("Position")));
-		Sig.bMemberFunction = true;
-		Sig.bRequiresContext = false;
 		Sig.SetDescription(LOCTEXT("DataInterfaceSpline_SampleSplinePositionByUnitDistanceWS", "Sample the spline Position where U is a 0 to 1 value representing the start and normalized length of the spline.\nThis is in the world space of the level."));
-		OutFunctions.Add(Sig);
 	}
 
 	{
-		FNiagaraFunctionSignature Sig;
+		FNiagaraFunctionSignature& Sig = OutFunctions.Add_GetRef(DefaultSig);
 		Sig.Name = NDISplineLocal::SampleSplineRotationByUnitDistanceName;
-		Sig.Inputs.Add(FNiagaraVariable(FNiagaraTypeDefinition(GetClass()), TEXT("Spline")));
 		Sig.Inputs.Add(FNiagaraVariable(FNiagaraTypeDefinition::GetFloatDef(), TEXT("U")));
 		Sig.Outputs.Add(FNiagaraVariable(FNiagaraTypeDefinition::GetQuatDef(), TEXT("Rotation")));
-		Sig.bMemberFunction = true;
-		Sig.bRequiresContext = false;
 		Sig.SetDescription(LOCTEXT("DataInterfaceSpline_SampleSplineRotationByUnitDistance", "Sample the spline Rotation where U is a 0 to 1 value representing the start and normalized length of the spline.\nThis is in the local space of the referenced USplineComponent."));
-		OutFunctions.Add(Sig);
 	}
 
 	{
-		FNiagaraFunctionSignature Sig;
+		FNiagaraFunctionSignature& Sig = OutFunctions.Add_GetRef(DefaultSig);
 		Sig.Name = NDISplineLocal::SampleSplineRotationByUnitDistanceWSName;
-		Sig.Inputs.Add(FNiagaraVariable(FNiagaraTypeDefinition(GetClass()), TEXT("Spline")));
 		Sig.Inputs.Add(FNiagaraVariable(FNiagaraTypeDefinition::GetFloatDef(), TEXT("U")));
 		Sig.Outputs.Add(FNiagaraVariable(FNiagaraTypeDefinition::GetQuatDef(), TEXT("Rotation")));
-		Sig.bMemberFunction = true;
-		Sig.bRequiresContext = false;
 		Sig.SetDescription(LOCTEXT("DataInterfaceSpline_SampleSplineRotationByUnitDistanceWS", "Sample the spline Rotation where U is a 0 to 1 value representing the start and normalized length of the spline.\nThis is in the world space of the level."));
-		OutFunctions.Add(Sig);
 	}
 
 	{
-		FNiagaraFunctionSignature Sig;
+		FNiagaraFunctionSignature& Sig = OutFunctions.Add_GetRef(DefaultSig);
 		Sig.Name = NDISplineLocal::SampleSplineDirectionByUnitDistanceName;
-		Sig.Inputs.Add(FNiagaraVariable(FNiagaraTypeDefinition(GetClass()), TEXT("Spline")));
 		Sig.Inputs.Add(FNiagaraVariable(FNiagaraTypeDefinition::GetFloatDef(), TEXT("U")));
 		Sig.Outputs.Add(FNiagaraVariable(FNiagaraTypeDefinition::GetVec3Def(), TEXT("Direction")));
-		Sig.bMemberFunction = true;
-		Sig.bRequiresContext = false;
 		Sig.SetDescription(LOCTEXT("DataInterfaceSpline_SampleSplineDirectionByUnitDistance", "Sample the spline direction vector where U is a 0 to 1 value representing the start and normalized length of the spline.\nThis is in the local space of the referenced USplineComponent."));
-		OutFunctions.Add(Sig);
 	}
 
 	{
-		FNiagaraFunctionSignature Sig;
+		FNiagaraFunctionSignature& Sig = OutFunctions.Add_GetRef(DefaultSig);
 		Sig.Name = NDISplineLocal::SampleSplineDirectionByUnitDistanceWSName;
-		Sig.Inputs.Add(FNiagaraVariable(FNiagaraTypeDefinition(GetClass()), TEXT("Spline")));
 		Sig.Inputs.Add(FNiagaraVariable(FNiagaraTypeDefinition::GetFloatDef(), TEXT("U")));
 		Sig.Outputs.Add(FNiagaraVariable(FNiagaraTypeDefinition::GetVec3Def(), TEXT("Direction")));
-		Sig.bMemberFunction = true;
-		Sig.bRequiresContext = false;
 		Sig.SetDescription(LOCTEXT("DataInterfaceSpline_SampleSplineDirectionByUnitDistanceWS", "Sample the spline direction vector where U is a 0 to 1 value representing the start and normalized length of the spline.\nThis is in the world space of the level."));
-		OutFunctions.Add(Sig);
 	}
 
 	{
-		FNiagaraFunctionSignature Sig;
+		FNiagaraFunctionSignature& Sig = OutFunctions.Add_GetRef(DefaultSig);
 		Sig.Name = NDISplineLocal::SampleSplineUpVectorByUnitDistanceName;
-		Sig.Inputs.Add(FNiagaraVariable(FNiagaraTypeDefinition(GetClass()), TEXT("Spline")));
 		Sig.Inputs.Add(FNiagaraVariable(FNiagaraTypeDefinition::GetFloatDef(), TEXT("U")));
 		Sig.Outputs.Add(FNiagaraVariable(FNiagaraTypeDefinition::GetVec3Def(), TEXT("UpVector")));
-		Sig.bMemberFunction = true;
-		Sig.bRequiresContext = false;
 		Sig.SetDescription(LOCTEXT("DataInterfaceSpline_SampleSplineUpVectorByUnitDistance", "Sample the spline up vector where U is a 0 to 1 value representing the start and normalized length of the spline.\nThis is in the local space of the referenced USplineComponent."));
-		OutFunctions.Add(Sig);
 	}
 
 	{
-		FNiagaraFunctionSignature Sig;
+		FNiagaraFunctionSignature& Sig = OutFunctions.Add_GetRef(DefaultSig);
 		Sig.Name = NDISplineLocal::SampleSplineUpVectorByUnitDistanceWSName;
-		Sig.Inputs.Add(FNiagaraVariable(FNiagaraTypeDefinition(GetClass()), TEXT("Spline")));
 		Sig.Inputs.Add(FNiagaraVariable(FNiagaraTypeDefinition::GetFloatDef(), TEXT("U")));
-		Sig.SetDescription(LOCTEXT("DataInterfaceSpline_SampleSplineUpVectorByUnitDistanceWS", "Sample the spline up vector where U is a 0 to 1 value representing the start and normalized length of the spline.\nThis is in the world space of the level."));
 		Sig.Outputs.Add(FNiagaraVariable(FNiagaraTypeDefinition::GetVec3Def(), TEXT("UpVector")));
-		Sig.bMemberFunction = true;
-		Sig.bRequiresContext = false;
-		OutFunctions.Add(Sig);
+		Sig.SetDescription(LOCTEXT("DataInterfaceSpline_SampleSplineUpVectorByUnitDistanceWS", "Sample the spline up vector where U is a 0 to 1 value representing the start and normalized length of the spline.\nThis is in the world space of the level."));
 	}
 
 	{
-		FNiagaraFunctionSignature Sig;
+		FNiagaraFunctionSignature& Sig = OutFunctions.Add_GetRef(DefaultSig);
 		Sig.Name = NDISplineLocal::SampleSplineRightVectorByUnitDistanceName;
-		Sig.Inputs.Add(FNiagaraVariable(FNiagaraTypeDefinition(GetClass()), TEXT("Spline")));
 		Sig.Inputs.Add(FNiagaraVariable(FNiagaraTypeDefinition::GetFloatDef(), TEXT("U")));
-		Sig.SetDescription(LOCTEXT("DataInterfaceSpline_SampleSplineRightVectorByUnitDistance", "Sample the spline right vector where U is a 0 to 1 value representing the start and normalized length of the spline.\nThis is in the local space of the referenced USplineComponent."));
 		Sig.Outputs.Add(FNiagaraVariable(FNiagaraTypeDefinition::GetVec3Def(), TEXT("RightVector")));
-		Sig.bMemberFunction = true;
-		Sig.bRequiresContext = false;
-		OutFunctions.Add(Sig);
+		Sig.SetDescription(LOCTEXT("DataInterfaceSpline_SampleSplineRightVectorByUnitDistance", "Sample the spline right vector where U is a 0 to 1 value representing the start and normalized length of the spline.\nThis is in the local space of the referenced USplineComponent."));
 	}
 
 	{
-		FNiagaraFunctionSignature Sig;
+		FNiagaraFunctionSignature& Sig = OutFunctions.Add_GetRef(DefaultSig);
 		Sig.Name = NDISplineLocal::SampleSplineRightVectorByUnitDistanceWSName;
-		Sig.Inputs.Add(FNiagaraVariable(FNiagaraTypeDefinition(GetClass()), TEXT("Spline")));
 		Sig.Inputs.Add(FNiagaraVariable(FNiagaraTypeDefinition::GetFloatDef(), TEXT("U")));
-		Sig.SetDescription(LOCTEXT("DataInterfaceSpline_SampleSplineRightVectorByUnitDistanceWS", "Sample the spline right vector where U is a 0 to 1 value representing the start and normalized length of the spline.\nThis is in the world space of the level."));
 		Sig.Outputs.Add(FNiagaraVariable(FNiagaraTypeDefinition::GetVec3Def(), TEXT("RightVector")));
-		Sig.bMemberFunction = true;
-		Sig.bRequiresContext = false;
-		OutFunctions.Add(Sig);
+		Sig.SetDescription(LOCTEXT("DataInterfaceSpline_SampleSplineRightVectorByUnitDistanceWS", "Sample the spline right vector where U is a 0 to 1 value representing the start and normalized length of the spline.\nThis is in the world space of the level."));
 	}
 
 	{
-		FNiagaraFunctionSignature Sig;
+		FNiagaraFunctionSignature& Sig = OutFunctions.Add_GetRef(DefaultSig);
 		Sig.Name = NDISplineLocal::SampleSplineTangentByUnitDistanceName;
-		Sig.Inputs.Add(FNiagaraVariable(FNiagaraTypeDefinition(GetClass()), TEXT("Spline")));
 		Sig.Inputs.Add(FNiagaraVariable(FNiagaraTypeDefinition::GetFloatDef(), TEXT("U")));
 		Sig.Outputs.Add(FNiagaraVariable(FNiagaraTypeDefinition::GetVec3Def(), TEXT("Tangent")));
-		Sig.bMemberFunction = true;
-		Sig.bRequiresContext = false;
 		Sig.SetDescription(LOCTEXT("DataInterfaceSpline_SampleSplineTangentVectorByUnitDistance", "Sample the spline tangent vector where U is a 0 to 1 value representing the start and normalized length of the spline.\nThis is in the local space of the referenced USplineComponent."));
-		OutFunctions.Add(Sig);
 	}
 
 	{
-		FNiagaraFunctionSignature Sig;
+		FNiagaraFunctionSignature& Sig = OutFunctions.Add_GetRef(DefaultSig);
 		Sig.Name = NDISplineLocal::SampleSplineTangentByUnitDistanceWSName;
-		Sig.Inputs.Add(FNiagaraVariable(FNiagaraTypeDefinition(GetClass()), TEXT("Spline")));
 		Sig.Inputs.Add(FNiagaraVariable(FNiagaraTypeDefinition::GetFloatDef(), TEXT("U")));
 		Sig.Outputs.Add(FNiagaraVariable(FNiagaraTypeDefinition::GetVec3Def(), TEXT("Tangent")));
-		Sig.bMemberFunction = true;
-		Sig.bRequiresContext = false;
 		Sig.SetDescription(LOCTEXT("DataInterfaceSpline_SampleSplineTangentVectorByUnitDistanceWS", "Sample the spline tangent vector where U is a 0 to 1 value representing the start and normalized length of the spline.\nThis is in the world space of the level."));
-		OutFunctions.Add(Sig);
 	}
 
 	{
-		FNiagaraFunctionSignature Sig;
+		FNiagaraFunctionSignature& Sig = OutFunctions.Add_GetRef(DefaultSig);
 		Sig.Name = NDISplineLocal::GetSplineLocalToWorldName;
-		Sig.Inputs.Add(FNiagaraVariable(FNiagaraTypeDefinition(GetClass()), TEXT("Spline")));
 		Sig.Outputs.Add(FNiagaraVariable(FNiagaraTypeDefinition::GetMatrix4Def(), TEXT("Transform")));
-		Sig.bMemberFunction = true;
-		Sig.bRequiresContext = false;
 		Sig.SetDescription(LOCTEXT("DataInterfaceSpline_GetSplineLocalToWorld", "Get the transform from the USplineComponent's local space to world space."));
-		OutFunctions.Add(Sig);
 	}
 
 	{
-		FNiagaraFunctionSignature Sig;
+		FNiagaraFunctionSignature& Sig = OutFunctions.Add_GetRef(DefaultSig);
 		Sig.Name = NDISplineLocal::GetSplineLocalToWorldInverseTransposedName;
-		Sig.Inputs.Add(FNiagaraVariable(FNiagaraTypeDefinition(GetClass()), TEXT("Spline")));
 		Sig.Outputs.Add(FNiagaraVariable(FNiagaraTypeDefinition::GetMatrix4Def(), TEXT("Transform")));
-		Sig.bMemberFunction = true;
-		Sig.bRequiresContext = false;
 		Sig.SetDescription(LOCTEXT("DataInterfaceSpline_GetSplineLocalToWorldInverseTransposed", "Get the transform from the world space to the USplineComponent's local space."));
-		OutFunctions.Add(Sig);
 	}
 
 	{
-		FNiagaraFunctionSignature Sig;
+		FNiagaraFunctionSignature& Sig = OutFunctions.Add_GetRef(DefaultSig);
 		Sig.Name = NDISplineLocal::FindClosestUnitDistanceFromPositionWSName;
-		Sig.Inputs.Add(FNiagaraVariable(FNiagaraTypeDefinition(GetClass()), TEXT("Spline")));
 		Sig.Inputs.Add(FNiagaraVariable(FNiagaraTypeDefinition::GetPositionDef(), TEXT("PositionWS")));
-
 		Sig.Outputs.Add(FNiagaraVariable(FNiagaraTypeDefinition::GetFloatDef(), TEXT("U")));
-		Sig.bMemberFunction = true;
-		Sig.bRequiresContext = false;
 		Sig.SetDescription(LOCTEXT("DataInterfaceSpline_FindClosestUnitDistanceFromPositionWS", "Given a world space position, find the closest value 'U' on the USplineComponent to that point."));
-		OutFunctions.Add(Sig);
 	}
 
-	for (FNiagaraFunctionSignature& FunctionSignature : OutFunctions)
 	{
-		FunctionSignature.SetFunctionVersion(FNiagaraSplineDIFunctionVersion::LatestVersion);
+		FNiagaraFunctionSignature& Sig = OutFunctions.Add_GetRef(DefaultSig);
+		Sig.Name = NDISplineLocal::GetSplineLengthName;
+		Sig.Outputs.Add(FNiagaraVariable(FNiagaraTypeDefinition::GetFloatDef(), TEXT("Length")));
 	}
 }
 
@@ -396,8 +343,11 @@ void UNiagaraDataInterfaceSpline::GetVMExternalFunction(const FVMExternalFunctio
 		check(BindingInfo.GetNumInputs() == 1 && BindingInfo.GetNumOutputs() == 16);
 		OutFunc = FVMExternalFunction::CreateUObject(this, &UNiagaraDataInterfaceSpline::GetLocalToWorldInverseTransposed);
 	}
-
-	//check(0);
+	else if (BindingInfo.Name == NDISplineLocal::GetSplineLengthName)
+	{
+		check(BindingInfo.GetNumInputs() == 1 && BindingInfo.GetNumOutputs() == 1);
+		OutFunc = FVMExternalFunction::CreateUObject(this, &UNiagaraDataInterfaceSpline::VMGetSplineLength);
+	}
 }
 
 bool UNiagaraDataInterfaceSpline::CopyToInternal(UNiagaraDataInterface* Destination) const
@@ -457,22 +407,27 @@ void UNiagaraDataInterfaceSpline::GetParameterDefinitionHLSL(const FNiagaraDataI
 bool UNiagaraDataInterfaceSpline::GetFunctionHLSL(const FNiagaraDataInterfaceGPUParamInfo& ParamInfo, const FNiagaraDataInterfaceGeneratedFunction& FunctionInfo, int FunctionInstanceIndex, FString& OutHLSL)
 {
 	using namespace NDISplineLocal;	
-	const FName& DefinitionName = FunctionInfo.DefinitionName;
-	return DefinitionName == SampleSplinePositionByUnitDistanceName || 
-		DefinitionName == SampleSplinePositionByUnitDistanceWSName ||
-		DefinitionName == SampleSplineRotationByUnitDistanceName ||
-		DefinitionName == SampleSplineRotationByUnitDistanceWSName ||
-		DefinitionName == SampleSplineDirectionByUnitDistanceName ||
-		DefinitionName == SampleSplineDirectionByUnitDistanceWSName ||
-		DefinitionName == SampleSplineUpVectorByUnitDistanceName ||
-		DefinitionName == SampleSplineUpVectorByUnitDistanceWSName ||
-		DefinitionName == SampleSplineRightVectorByUnitDistanceName ||
-		DefinitionName == SampleSplineRightVectorByUnitDistanceWSName ||
-		DefinitionName == SampleSplineTangentByUnitDistanceName ||
-		DefinitionName == SampleSplineTangentByUnitDistanceWSName ||
-		DefinitionName == GetSplineLocalToWorldName ||
-		DefinitionName == GetSplineLocalToWorldInverseTransposedName ||
-		DefinitionName == FindClosestUnitDistanceFromPositionWSName;
+	static const TSet<FName> ValidGpuFunctions =
+	{
+		SampleSplinePositionByUnitDistanceName,
+		SampleSplinePositionByUnitDistanceWSName,
+		SampleSplineRotationByUnitDistanceName,
+		SampleSplineRotationByUnitDistanceWSName,
+		SampleSplineDirectionByUnitDistanceName,
+		SampleSplineDirectionByUnitDistanceWSName,
+		SampleSplineUpVectorByUnitDistanceName,
+		SampleSplineUpVectorByUnitDistanceWSName,
+		SampleSplineRightVectorByUnitDistanceName,
+		SampleSplineRightVectorByUnitDistanceWSName,
+		SampleSplineTangentByUnitDistanceName,
+		SampleSplineTangentByUnitDistanceWSName,
+		GetSplineLocalToWorldName,
+		GetSplineLocalToWorldInverseTransposedName,
+		FindClosestUnitDistanceFromPositionWSName,
+		GetSplineLengthName
+	};
+
+	return ValidGpuFunctions.Contains(FunctionInfo.DefinitionName);
 }
 
 bool UNiagaraDataInterfaceSpline::AppendCompileHash(FNiagaraCompileHashVisitor* InVisitor) const
@@ -1463,5 +1418,22 @@ void UNiagaraDataInterfaceSpline::GetLocalToWorldInverseTransposed(FVectorVMExte
 	WriteTransform(InstData->TransformInverseTransposed, Context);
 }
 
+void UNiagaraDataInterfaceSpline::VMGetSplineLength(FVectorVMExternalFunctionContext& Context)
+{
+	VectorVM::FUserPtrHandler<FNDISpline_InstanceData> InstData(Context);
+	FNDIOutputParam<float> OutLength(Context);
+
+	float SplineLength = 0.0f;
+	if (InstData->IsValid())
+	{
+		SplineLength = bUseLUT ? InstData->GetSplineLength<TIntegralConstant<bool, true>>() : InstData->GetSplineLength<TIntegralConstant<bool, false>>();
+	}
+
+	for (int32 i = 0; i < Context.GetNumInstances(); ++i)
+	{
+		OutLength.SetAndAdvance(SplineLength);
+	}
+}
+
 #undef LOCTEXT_NAMESPACE
 
