- Add get mesh bounds functions to static mesh DI

diff --git a/Engine/Plugins/FX/Niagara/Shaders/Private/NiagaraDataInterfaceStaticMeshTemplate.ush b/Engine/Plugins/FX/Niagara/Shaders/Private/NiagaraDataInterfaceStaticMeshTemplate.ush
--- a/Engine/Plugins/FX/Niagara/Shaders/Private/NiagaraDataInterfaceStaticMeshTemplate.ush
+++ b/Engine/Plugins/FX/Niagara/Shaders/Private/NiagaraDataInterfaceStaticMeshTemplate.ush
@@ -34,6 +34,9 @@ float3			{ParameterName}_InstanceWorldVelocity;
 
 float3			{ParameterName}_PreSkinnedLocalBoundsCenter;
 float3			{ParameterName}_PreSkinnedLocalBoundsExtents;
+float3			{ParameterName}_MeshBoundsWSCenter;
+float3			{ParameterName}_MeshBoundsWSExtents;
+float3			{ParameterName}_OwnerToMeshVector;
 float3			{ParameterName}_SystemLWCTile;
 
 #if DISTATICMESH_ALLOWDISTANCEFIELD
@@ -728,6 +731,26 @@ void GetWorldVelocity_{ParameterName}(out float3 Velocity)
 	Velocity = {ParameterName}_InstanceWorldVelocity;
 }
 
+void GetMeshBounds_{ParameterName}(out float3 Center, out float3 ExtentsMin, out float3 ExtentsMax, out float3 Extents, out float3 HalfExtents, out float3 OwnerToMeshVector)
+{
+	Center				= {ParameterName}_PreSkinnedLocalBoundsCenter;
+	HalfExtents			= {ParameterName}_MeshBoundsWSExtents;
+	Extents				= HalfExtents * 2.0f;
+	ExtentsMin			= Center - HalfExtents;
+	ExtentsMax			= Center + HalfExtents;
+	OwnerToMeshVector	= {ParameterName}_OwnerToMeshVector;
+}
+
+void GetMeshBoundsWS_{ParameterName}(out float3 Center, out float3 ExtentsMin, out float3 ExtentsMax, out float3 Extents, out float3 HalfExtents, out float3 OwnerToMeshVector)
+{
+	Center				= {ParameterName}_MeshBoundsWSCenter;
+	HalfExtents			= {ParameterName}_MeshBoundsWSExtents;
+	Extents				= HalfExtents * 2.0f;
+	ExtentsMin			= Center - HalfExtents;
+	ExtentsMax			= Center + HalfExtents;
+	OwnerToMeshVector	= {ParameterName}_OwnerToMeshVector;
+}
+
 //////////////////////////////////////////////////////////////////////////
 // Experimental Distance Field Functions
 void QueryDistanceField_{ParameterName}(bool bExecute, float3 WorldPosition, bool bUseMaxDistance, float MaxDistance, out bool OutIsValid, out float OutDistance)
diff --git a/Engine/Plugins/FX/Niagara/Source/Niagara/Private/DataInterface/NiagaraDataInterfaceStaticMesh.cpp b/Engine/Plugins/FX/Niagara/Source/Niagara/Private/DataInterface/NiagaraDataInterfaceStaticMesh.cpp
--- a/Engine/Plugins/FX/Niagara/Source/Niagara/Private/DataInterface/NiagaraDataInterfaceStaticMesh.cpp
+++ b/Engine/Plugins/FX/Niagara/Source/Niagara/Private/DataInterface/NiagaraDataInterfaceStaticMesh.cpp
@@ -77,6 +77,9 @@ namespace NDIStaticMeshLocal
 
 		SHADER_PARAMETER(FVector3f,				PreSkinnedLocalBoundsCenter)
 		SHADER_PARAMETER(FVector3f,				PreSkinnedLocalBoundsExtents)
+		SHADER_PARAMETER(FVector3f,				MeshBoundsWSCenter)
+		SHADER_PARAMETER(FVector3f,				MeshBoundsWSExtents)
+		SHADER_PARAMETER(FVector3f,				OwnerToMeshVector)
 		SHADER_PARAMETER(FVector3f,				SystemLWCTile)
 
 		SHADER_PARAMETER_SRV(Buffer<int>,		UvMappingBuffer)
@@ -202,6 +205,8 @@ namespace NDIStaticMeshLocal
 	static const FName	IsValidName("IsValid");
 
 	static const FName	GetPreSkinnedLocalBoundsName("GetPreSkinnedLocalBounds");
+	static const FName	GetMeshBoundsName("GetMeshBounds");
+	static const FName	GetMeshBoundsWSName("GetMeshBoundsWS");
 
 	static const FName	GetLocalToWorldName("GetLocalToWorld");
 	static const FName	GetLocalToWorldInverseTransposedName("GetLocalToWorldInverseTransposed");
@@ -237,40 +242,6 @@ namespace NDIStaticMeshLocal
 	static const FName Deprecated_GetTriPositionAndVelocityWSName("DeprecatedGetTriPositionAndVelocityWS");
 	//static const FName Deprecated_GetMeshLocalToWorldInverseTransposedName("DeprecatedGetMeshLocalToWorldInverseTransposed");
 
-	//////////////////////////////////////////////////////////////////////////
-	// Shader Parameters
-	static FString	NumTriangles_String("NumTriangles_");
-	static FString	NumVertices_String("NumVertices_");
-	static FString	NumUVs_String("NumUVs_");
-	static FString	HasColors_String("HasColors_");
-	static FString	IndexBuffer_String("IndexBuffer_");
-	static FString	PositionBuffer_String("PositionBuffer_");
-	static FString	TangentBuffer_String("TangentBuffer_");
-	static FString	UVBuffer_String("UVBuffer_");
-	static FString	ColorBuffer_String("ColorBuffer_");
-
-	static FString	HasUniformSampling_String("HasUniformSampling_");
-	static FString	UniformSamplingTriangles_String("UniformSamplingTriangles_");
-
-	static FString	SectionCounts_String("SectionCounts_");
-	static FString	SectionInfos_String("SectionInfos_");
-	static FString	FilteredAndUnfilteredSections_String("FilteredAndUnfilteredSections_");
-
-	static FString	SocketCounts_String("SocketCounts_");
-	static FString	SocketTransforms_String("SocketTransforms_");
-	static FString	FilteredAndUnfilteredSockets_String("FilteredAndUnfilteredSockets_");
-
-	static FString	InvDeltaSeconds_String("InvDeltaSeconds_");
-	static FString	InstanceTransform_String("InstanceTransform_");
-	static FString	InstanceTransformInverseTransposed_String("InstanceTransformInverseTransposed_");
-	static FString	InstanceRotation_String("InstanceRotation_");
-	static FString	InstancePreviousTransform_String("InstancePreviousTransform_");
-	static FString	InstancePreviousTransformInverseTransposed_String("InstancePreviousTransformInverseTransposed_");
-	static FString	InstancePreviousRotation_String("InstancePreviousRotation_");
-	static FString	InstanceWorldVelocity_String("InstanceWorldVelocity_");
-
-	static FString	InstanceDistanceFieldIndex_String("InstanceDistanceFieldIndex_");
-
 	//////////////////////////////////////////////////////////////////////////
 	struct FNDISectionInfo
 	{
@@ -386,6 +357,7 @@ namespace NDIStaticMeshLocal
 		float		DeltaSeconds = 0.0f;
 		FVector3f	PreSkinnedLocalBoundsCenter = FVector3f::ZeroVector;
 		FVector3f	PreSkinnedLocalBoundsExtents = FVector3f::ZeroVector;
+		FVector3f	OwnerToMeshVector = FVector3f::ZeroVector;
 		FPrimitiveComponentId	DistanceFieldPrimitiveId;
 
 		FIntVector NumTriangles = FIntVector::ZeroValue;
@@ -544,6 +516,7 @@ namespace NDIStaticMeshLocal
 		float					DeltaSeconds = 0.0f;
 		FVector3f				PreSkinnedLocalBoundsCenter = FVector3f::ZeroVector;
 		FVector3f				PreSkinnedLocalBoundsExtents = FVector3f::ZeroVector;
+		FVector3f				OwnerToMeshVector = FVector3f::ZeroVector;
 		FPrimitiveComponentId	DistanceFieldPrimitiveId;
 
 		const FMeshUvMappingBufferProxy* UvMappingBuffer = nullptr;
@@ -565,6 +538,9 @@ namespace NDIStaticMeshLocal
 		/** Handle to our uv mapping data. */
 		FStaticMeshUvMappingHandle UvMapping;
 
+		/** Vector from system owners instance location to the sampled mesh location. */
+		FVector3f OwnerToMeshVector = FVector3f::ZeroVector;
+
 		/** Cached ComponentToWorld. (Falls back to WorldTransform of the system instance) */
 		FMatrix Transform;
 		/** InverseTranspose of above for transforming normals/tangents. */
@@ -681,6 +657,7 @@ namespace NDIStaticMeshLocal
 			DeltaSeconds = 0.0f;
 			PreSkinnedLocalBoundsCenter = FVector3f::ZeroVector;
 			PreSkinnedLocalBoundsExtents = FVector3f::ZeroVector;
+			OwnerToMeshVector = FVector3f::ZeroVector;
 			PhysicsVelocity = FVector::ZeroVector;
 			bUsePhysicsVelocity = Interface->bUsePhysicsBodyVelocity;
 			bComponentValid = false;
@@ -709,6 +686,7 @@ namespace NDIStaticMeshLocal
 			// Gather attached information
 			bComponentValid = SceneComponent != nullptr;
 			FTransform ComponentTransform = bComponentValid ? SceneComponent->GetComponentToWorld() : SystemInstance->GetWorldTransform();
+			OwnerToMeshVector = FVector3f(ComponentTransform.GetLocation() -  SystemInstance->GetWorldTransform().GetLocation());
 			ComponentTransform.AddToTranslation(FVector(SystemInstance->GetLWCTile()) * -FLargeWorldRenderScalar::GetTileSize());
 
 			Transform = ComponentTransform.ToMatrixWithScale();
@@ -940,6 +918,7 @@ namespace NDIStaticMeshLocal
 
 			USceneComponent* SceneComponent = SceneComponentWeakPtr.Get();
 			FTransform ComponentTransform = SceneComponent != nullptr ? SceneComponent->GetComponentToWorld() : SystemInstance->GetWorldTransform();
+			OwnerToMeshVector = FVector3f(ComponentTransform.GetLocation() - SystemInstance->GetWorldTransform().GetLocation());
 			ComponentTransform.AddToTranslation(FVector(SystemInstance->GetLWCTile()) * -FLargeWorldRenderScalar::GetTileSize());
 
 			PrevTransform = Transform;
@@ -1088,6 +1067,7 @@ namespace NDIStaticMeshLocal
 			InstanceData->DeltaSeconds				= FromGameThread->DeltaSeconds;
 			InstanceData->PreSkinnedLocalBoundsCenter = FromGameThread->PreSkinnedLocalBoundsCenter;
 			InstanceData->PreSkinnedLocalBoundsExtents = FromGameThread->PreSkinnedLocalBoundsExtents;
+			InstanceData->OwnerToMeshVector			= FromGameThread->OwnerToMeshVector;
 			InstanceData->DistanceFieldPrimitiveId	= FromGameThread->DistanceFieldPrimitiveId;
 
 			InstanceData->UvMappingBuffer			= FromGameThread->UvMappingBuffer;
@@ -1760,6 +1740,7 @@ void UNiagaraDataInterfaceStaticMesh::ProvidePerInstanceDataForRenderThread(void
 	DataFromGT->DeltaSeconds				= InstanceData->DeltaSeconds;
 	DataFromGT->PreSkinnedLocalBoundsCenter = InstanceData->PreSkinnedLocalBoundsCenter;
 	DataFromGT->PreSkinnedLocalBoundsExtents = InstanceData->PreSkinnedLocalBoundsExtents;
+	DataFromGT->OwnerToMeshVector			= InstanceData->OwnerToMeshVector;
 	DataFromGT->DistanceFieldPrimitiveId	= FPrimitiveComponentId();
 	DataFromGT->UvMappingBuffer				= InstanceData->UvMapping.GetQuadTreeProxy();
 	DataFromGT->UvMappingSet				= InstanceData->UvMapping.GetUvSetIndex();
@@ -2127,7 +2108,30 @@ void UNiagaraDataInterfaceStaticMesh::GetMiscFunctions(TArray<FNiagaraFunctionSi
 		Sig.Outputs.Emplace(FNiagaraTypeDefinition::GetVec3Def(), TEXT("ExtentsMax"));
 		Sig.Outputs.Emplace(FNiagaraTypeDefinition::GetVec3Def(), TEXT("Extents"));
 		Sig.Outputs.Emplace(FNiagaraTypeDefinition::GetVec3Def(), TEXT("Half Extents"));
-		Sig.SetDescription(LOCTEXT("GetPreSkinnedLocalBoundsDesc", "Returns the local bounds of the static mesh"));
+		Sig.SetDescription(LOCTEXT("GetPreSkinnedLocalBoundsDesc", "Returns the bounding information of the mesh vertices."));
+		Sig.SetOutputDescription(Sig.Outputs[0], LOCTEXT("PreSkinnedLocalBounds_CenterDesc", "Center of the mesh vertices bounding box, i.e. ((max - min) * 0.5) + min"));
+		Sig.SetOutputDescription(Sig.Outputs[1], LOCTEXT("PreSkinnedLocalBounds_ExtentsMinDesc", "Min of the mesh vertices."));
+		Sig.SetOutputDescription(Sig.Outputs[2], LOCTEXT("PreSkinnedLocalBounds_ExtentsMaxDesc", "Max of the mesh vertices."));
+		Sig.SetOutputDescription(Sig.Outputs[3], LOCTEXT("PreSkinnedLocalBounds_ExtentsDesc", "Extents of the mesh vertices bounding box, i.e. max - min"));
+		Sig.SetOutputDescription(Sig.Outputs[4], LOCTEXT("PreSkinnedLocalBounds_HalfExtentsDesc", "Extents * 0.5f of the mesh vertices bounding box, i.e. (max - min) * 0.5"));
+	}
+	{
+		FNiagaraFunctionSignature BoundsSig = BaseSignature;
+		BoundsSig.Outputs.Emplace(FNiagaraTypeDefinition::GetPositionDef(), TEXT("Center"));
+		BoundsSig.Outputs.Emplace(FNiagaraTypeDefinition::GetVec3Def(), TEXT("ExtentsMin"));
+		BoundsSig.Outputs.Emplace(FNiagaraTypeDefinition::GetVec3Def(), TEXT("ExtentsMax"));
+		BoundsSig.Outputs.Emplace(FNiagaraTypeDefinition::GetVec3Def(), TEXT("Extents"));
+		BoundsSig.Outputs.Emplace(FNiagaraTypeDefinition::GetVec3Def(), TEXT("Half Extents"));
+		BoundsSig.Outputs.Emplace(FNiagaraTypeDefinition::GetVec3Def(), TEXT("OwnerToMeshVector"));
+		BoundsSig.SetOutputDescription(BoundsSig.Outputs.Last(), LOCTEXT("GetMeshBounds_MeshComponentOffsetDesc", "Offset vector to the Source Mesh Component from the Niagara instances location."));
+
+		FNiagaraFunctionSignature& LSSig = OutFunctions.Add_GetRef(BoundsSig);
+		LSSig.Name = GetMeshBoundsName;
+		LSSig.SetDescription(LOCTEXT("GetMeshBoundsDesc", "Returns the bounding information of the mesh in local space."));
+
+		FNiagaraFunctionSignature& WSSig = OutFunctions.Add_GetRef(BoundsSig);
+		WSSig.Name = GetMeshBoundsWSName;
+		WSSig.SetDescription(LOCTEXT("GetMeshBoundsWSDesc", "Returns the bounding information of the mesh in world space."));
 	}
 	{
 		FNiagaraFunctionSignature& Sig = OutFunctions.Add_GetRef(BaseSignature);
@@ -2681,7 +2685,15 @@ void UNiagaraDataInterfaceStaticMesh::GetVMExternalFunction(const FVMExternalFun
 	}
 	else if (BindingInfo.Name == GetPreSkinnedLocalBoundsName)
 	{
-	OutFunc = FVMExternalFunction::CreateUObject(this, &UNiagaraDataInterfaceStaticMesh::VMGetPreSkinnedLocalBounds);
+		OutFunc = FVMExternalFunction::CreateUObject(this, &UNiagaraDataInterfaceStaticMesh::VMGetPreSkinnedLocalBounds);
+	}
+	else if (BindingInfo.Name == GetMeshBoundsName)
+	{
+		OutFunc = FVMExternalFunction::CreateUObject(this, &UNiagaraDataInterfaceStaticMesh::VMGetMeshBounds<true>);
+	}
+	else if (BindingInfo.Name == GetMeshBoundsWSName)
+	{
+		OutFunc = FVMExternalFunction::CreateUObject(this, &UNiagaraDataInterfaceStaticMesh::VMGetMeshBounds<false>);
 	}
 	else if (BindingInfo.Name == GetLocalToWorldName)
 	{
@@ -3007,8 +3019,14 @@ void UNiagaraDataInterfaceStaticMesh::SetShaderParameters(const FNiagaraDataInte
 		FNiagaraDistanceFieldHelper::SetMeshDistanceFieldParameters(Context.GetGraphBuilder(), DistanceFieldSceneData, *ShaderDistanceFieldObjectParameters, *ShaderDistanceFieldAtlasParameters, FNiagaraRenderer::GetDummyFloat4Buffer());
 	}
 
+	const FBox3f MeshBounds(InstanceData.PreSkinnedLocalBoundsCenter - InstanceData.PreSkinnedLocalBoundsExtents, InstanceData.PreSkinnedLocalBoundsCenter + InstanceData.PreSkinnedLocalBoundsExtents);
+	const FBox3f WorldBounds = MeshBounds.TransformBy(InstanceData.Transform);
+
 	ShaderParameters->PreSkinnedLocalBoundsCenter = InstanceData.PreSkinnedLocalBoundsCenter;
 	ShaderParameters->PreSkinnedLocalBoundsExtents = InstanceData.PreSkinnedLocalBoundsExtents;
+	ShaderParameters->MeshBoundsWSCenter = WorldBounds.GetCenter();
+	ShaderParameters->MeshBoundsWSExtents = WorldBounds.GetExtent();
+	ShaderParameters->OwnerToMeshVector = InstanceData.OwnerToMeshVector;
 
 	ShaderParameters->SystemLWCTile = Context.GetSystemLWCTile();
 
@@ -3361,6 +3379,38 @@ void UNiagaraDataInterfaceStaticMesh::VMGetPreSkinnedLocalBounds(FVectorVMExtern
 	}
 }
 
+template<bool bLocalSpace>
+void UNiagaraDataInterfaceStaticMesh::VMGetMeshBounds(FVectorVMExternalFunctionContext& Context)
+{
+	NDIStaticMeshLocal::FStaticMeshCpuHelper<FNDITransformHandlerNoop> StaticMeshHelper(Context);
+	FNDIOutputParam<FVector3f>	OutCenter(Context);
+	FNDIOutputParam<FVector3f>	OutExtentsMin(Context);
+	FNDIOutputParam<FVector3f>	OutExtentsMax(Context);
+	FNDIOutputParam<FVector3f>	OutExtents(Context);
+	FNDIOutputParam<FVector3f>	OutHalfExtents(Context);
+	FNDIOutputParam<FVector3f>	OutOwnerToMeshVector(Context);
+
+	const FBox MeshBounds(StaticMeshHelper.InstanceData->PreSkinnedLocalBoundsCenter - StaticMeshHelper.InstanceData->PreSkinnedLocalBoundsExtents, StaticMeshHelper.InstanceData->PreSkinnedLocalBoundsCenter + StaticMeshHelper.InstanceData->PreSkinnedLocalBoundsExtents);
+	const FBox WorldBounds = MeshBounds.TransformBy(StaticMeshHelper.InstanceData->Transform);
+
+	const FVector3f Center = bLocalSpace ? FVector3f(MeshBounds.GetCenter())  : FVector3f(WorldBounds.GetCenter());
+	const FVector3f HalfExtents = FVector3f(WorldBounds.GetExtent());
+	const FVector3f Extents = HalfExtents * 2.0f;
+	const FVector3f ExtentsMin = Center - HalfExtents;
+	const FVector3f ExtentsMax = Center + HalfExtents;
+	const FVector3f OwnerToMeshVector = StaticMeshHelper.InstanceData->OwnerToMeshVector;
+
+	for (int32 i = 0; i < Context.GetNumInstances(); ++i)
+	{
+		OutCenter.SetAndAdvance(Center);
+		OutExtentsMin.SetAndAdvance(ExtentsMin);
+		OutExtentsMax.SetAndAdvance(ExtentsMax);
+		OutExtents.SetAndAdvance(Extents);
+		OutHalfExtents.SetAndAdvance(HalfExtents);
+		OutOwnerToMeshVector.SetAndAdvance(OwnerToMeshVector);
+	}
+}
+
 //////////////////////////////////////////////////////////////////////////
 // VM Vertex Sampling
 void UNiagaraDataInterfaceStaticMesh::VMIsValidVertex(FVectorVMExternalFunctionContext& Context)
@@ -5096,4 +5146,3 @@ FStaticMeshUvMappingHandle FNDI_StaticMesh_GeneratedData::GetCachedUvMapping(TWe
 }
 
 #undef LOCTEXT_NAMESPACE
-
diff --git a/Engine/Plugins/FX/Niagara/Source/Niagara/Private/DataInterface/NiagaraDataInterfaceStaticMesh.h b/Engine/Plugins/FX/Niagara/Source/Niagara/Private/DataInterface/NiagaraDataInterfaceStaticMesh.h
--- a/Engine/Plugins/FX/Niagara/Source/Niagara/Private/DataInterface/NiagaraDataInterfaceStaticMesh.h
+++ b/Engine/Plugins/FX/Niagara/Source/Niagara/Private/DataInterface/NiagaraDataInterfaceStaticMesh.h
@@ -256,6 +256,8 @@ protected:
 	void VMIsValid(FVectorVMExternalFunctionContext& Context);
 
 	void VMGetPreSkinnedLocalBounds(FVectorVMExternalFunctionContext& Context);
+	template<bool bLocalSpace>
+	void VMGetMeshBounds(FVectorVMExternalFunctionContext& Context);
 
 	void VMGetLocalToWorld(FVectorVMExternalFunctionContext& Context);
 	void VMGetLocalToWorldInverseTransposed(FVectorVMExternalFunctionContext& Context);
