Removed "Common" section in Niagara system editor's parameter panel

diff --git a/Engine/Plugins/FX/Niagara/Source/NiagaraEditor/Private/ViewModels/NiagaraParameterPanelViewModel.cpp b/Engine/Plugins/FX/Niagara/Source/NiagaraEditor/Private/ViewModels/NiagaraParameterPanelViewModel.cpp
--- a/Engine/Plugins/FX/Niagara/Source/NiagaraEditor/Private/ViewModels/NiagaraParameterPanelViewModel.cpp
+++ b/Engine/Plugins/FX/Niagara/Source/NiagaraEditor/Private/ViewModels/NiagaraParameterPanelViewModel.cpp
@@ -883,7 +883,6 @@ void FNiagaraSystemToolkitParameterPanelViewModel::Init(const FSystemToolkitUICo
 
 	// Define the sections
 	UserOnlyIdx = Sections.Emplace(FGuid(0x9E29BC43, 0x31524AB5, 0xA9D6675D, 0x741137CB), LOCTEXT("SectionUser", "User"), LOCTEXT("SectionUserDesc", "Displays only User exposed parameters"));
-	CommonOnlyIdx = Sections.Emplace(FGuid(0x3054B620, 0xBE714E54, 0xB7DB81E9, 0x00C0DE3D ),LOCTEXT("SectionCommon", "Common"), LOCTEXT("SectionCommonDesc", "Displays parameters defined in Parameter Libraries"));
 	ActiveSystemIdx = Sections.Emplace(FGuid(0x7B4AFB34, 0xD0DF4618, 0x9A05349E, 0x361CE735), LOCTEXT("SectionSystemActive", "Active Overview"), LOCTEXT("SectionActiveSystemDesc", "Displays parameters that are context-sensitive to the selected items in the System Overview."));
 	ActiveScriptIdx = Sections.Emplace(FGuid(0x6E11F3EA, 0x5E314E82, 0xA64FEA41, 0x4F1DB891), LOCTEXT("SectionScriptActive", "Active Module"), LOCTEXT("SectionActiveScriptDesc", "Displays parameters that are context-sensitive to the active Scratch module document. Use this for renaming/editing local parameters to the module."));
 
@@ -900,13 +899,6 @@ void FNiagaraSystemToolkitParameterPanelViewModel::Init(const FSystemToolkitUICo
 			bSaveConfig = true;
 		}
 
-		FNiagaraParameterPanelSectionStorage& CommonStorage = Settings->FindOrAddParameterPanelSectionStorage(Sections[CommonOnlyIdx].SectionId, bAdded);
-		if (bAdded)
-		{
-			CommonStorage.ExpandedCategories.Emplace(FNiagaraEditorGuids::ParticleAttributeNamespaceMetaDataGuid);
-			bSaveConfig = true;
-		}
-
 		FNiagaraParameterPanelSectionStorage& ScriptStorage = Settings->FindOrAddParameterPanelSectionStorage(Sections[ActiveScriptIdx].SectionId, bAdded);
 		if (bAdded)
 		{
@@ -1254,12 +1246,6 @@ void FNiagaraSystemToolkitParameterPanelViewModel::FindOrAddParameter(FNiagaraVa
 	SelectParameterItemByName(Variable.GetName(), false);
 }
 
-bool FNiagaraSystemToolkitParameterPanelViewModel::GetShowReferenceCounter() const
-{
-	return (ActiveSectionIndex != CommonOnlyIdx);
-}
-
-
 void FNiagaraSystemToolkitParameterPanelViewModel::SetActiveSection(int32 InSection)
 {
 	INiagaraParameterPanelViewModel::SetActiveSection(InSection);
@@ -1282,11 +1268,6 @@ void FNiagaraSystemToolkitParameterPanelViewModel::SetActiveSection(int32 InSect
 
 bool FNiagaraSystemToolkitParameterPanelViewModel::GetCanAddParametersToCategory(FNiagaraParameterPanelCategory Category) const
 {
-	if (ActiveSectionIndex == CommonOnlyIdx)
-	{
-		return false;
-	}
-
 	// If in emitter edit mode, don't support anything other than Emitter level and below
 	if (SystemViewModel && SystemViewModel->GetEditMode() == ENiagaraSystemViewModelEditMode::EmitterAsset)
 	{
@@ -1537,58 +1518,17 @@ const TArray<FNiagaraParameterPanelCategory>& FNiagaraSystemToolkitParameterPane
 	bool bForceScript = false;
 	bool bForceSystem = false;
 	bool bUserOnly = false;
-	bool bCommonOnly = false;
 	if (ActiveSectionIndex == UserOnlyIdx && UserOnlyIdx != -1)
 		bUserOnly = true;
 	else if (ActiveSectionIndex == ActiveScriptIdx && ActiveScriptIdx != -1)
 		bForceScript= true;
 	else if (ActiveSectionIndex == ActiveSystemIdx && ActiveSystemIdx != -1)
 		bForceSystem = true;
-	else if (ActiveSectionIndex == CommonOnlyIdx && CommonOnlyIdx != -1)
-		bCommonOnly = true;
-
 
 	if (bUserOnly)
 	{
 		CachedCurrentCategories = FNiagaraSystemToolkitParameterPanelViewModel::UserCategories;
 	}
-	else if (bCommonOnly)
-	{
-		// For common, dynamically create the cached array as it will depend on the parameter libraries that are live.
-		CachedCurrentCategories.Reset();
-
-		TArray<FNiagaraNamespaceMetadata> NamespaceMetadata = GetDefault<UNiagaraEditorSettings>()->GetAllNamespaceMetadata();
-		TArray<UNiagaraParameterDefinitions*> Defs = GetAvailableParameterDefinitions(false);
-		for (const FNiagaraNamespaceMetadata& NamespaceMetadatum : NamespaceMetadata)
-		{
-			for (int32 i = 0; i < NamespaceMetadatum.Namespaces.Num(); i++)
-			{
-				bool bFound = false;
-				for (UNiagaraParameterDefinitions* Def : Defs)
-				{
-					if (Def)
-					{
-						for (const UNiagaraScriptVariable* ScriptVar : Def->GetParametersConst())
-						{
-							if (ScriptVar && ScriptVar->Variable.IsInNameSpace(NamespaceMetadatum.Namespaces[i]))
-							{
-								bFound = true;
-								break;
-							}
-						}
-					}
-
-					if (bFound)
-						break;
-				}
-
-				if (bFound)
-				{
-					CachedCurrentCategories.Add(NamespaceMetadatum);
-				}
-			}
-		}
-	}
 	else
 	{	
 		if (bForceScript)
@@ -1600,9 +1540,7 @@ const TArray<FNiagaraParameterPanelCategory>& FNiagaraSystemToolkitParameterPane
 			CachedCurrentCategories = FNiagaraSystemToolkitParameterPanelViewModel::DefaultAdvancedCategories;
 		}
 	}
-
 	
-
 	return CachedCurrentCategories;
 }
 
@@ -1928,15 +1866,12 @@ TArray<FNiagaraParameterPanelItem> FNiagaraSystemToolkitParameterPanelViewModel:
 	bool bForceScript = false;
 	bool bForceSystem = false;
 	bool bUserOnly = false;
-	bool bCommonOnly = false;
 	if (ActiveSectionIndex == UserOnlyIdx && UserOnlyIdx != -1)
 		bUserOnly = true;
 	else if (ActiveSectionIndex == ActiveScriptIdx && ActiveScriptIdx != -1)
 		bForceScript = true;
 	else if (ActiveSectionIndex == ActiveSystemIdx && ActiveSystemIdx != -1)
 		bForceSystem = true;
-	else if (ActiveSectionIndex == CommonOnlyIdx && CommonOnlyIdx != -1)
-		bCommonOnly = true;
 
 	// On the first time opening the parameter panel view model we are not guaranteed to call GetDefaultCategories() before GetViewedParameterItems(). 
 	// We require CachedCurrentCategories being set as this is used to filter out parameter items that are being viewed. If CachedCurrentCategories 
@@ -1987,32 +1922,6 @@ TArray<FNiagaraParameterPanelItem> FNiagaraSystemToolkitParameterPanelViewModel:
 		CollectParamStore(&SystemViewModel->GetSystem().GetExposedParameters());
 	}
 
-	if (bCommonOnly)
-	{
-		TArray<UNiagaraParameterDefinitions*> Defs = GetAvailableParameterDefinitions(false);
-		for (UNiagaraParameterDefinitions* Def : Defs)
-		{
-			if (Def)
-			{
-				for (const UNiagaraScriptVariable* ScriptVar : Def->GetParametersConst())
-				{
-					FNiagaraVariableBase Var = ScriptVar->Variable;
-					FNiagaraParameterPanelItem Item = FNiagaraParameterPanelItem();
-					Item.ScriptVariable = ScriptVar;
-					Item.NamespaceMetaData = FNiagaraEditorUtilities::GetNamespaceMetaDataForVariableName(Var.GetName());
-					Item.bExternallyReferenced = false;
-					Item.bSourcedFromCustomStackContext = false;
-					Item.ReferenceCount = 0;
-
-					// Determine whether the item is name aliasing a parameter definition's parameter.
-					Item.DefinitionMatchState = FNiagaraParameterDefinitionsUtilities::GetDefinitionMatchStateForParameter(ScriptVar->Variable);
-
-					VisitedParameterToItemMap.Add(Var, Item);
-				}
-			}
-		}
-	}
-
 	if (bForceSystem)
 	{
 		// Collect parameters added to the system asset.
diff --git a/Engine/Plugins/FX/Niagara/Source/NiagaraEditor/Private/Widgets/SNiagaraParameterPanel.cpp b/Engine/Plugins/FX/Niagara/Source/NiagaraEditor/Private/Widgets/SNiagaraParameterPanel.cpp
--- a/Engine/Plugins/FX/Niagara/Source/NiagaraEditor/Private/Widgets/SNiagaraParameterPanel.cpp
+++ b/Engine/Plugins/FX/Niagara/Source/NiagaraEditor/Private/Widgets/SNiagaraParameterPanel.cpp
@@ -464,7 +464,6 @@ TSharedRef<SWidget> SNiagaraParameterPanel::OnGenerateWidgetForItem(const FNiaga
 			.ButtonStyle(FAppStyle::Get(), "RoundButton")
 			.ForegroundColor(FSlateColor::UseForeground())
 			.ContentPadding(FMargin(2.0f))
-			.Visibility(this, &SNiagaraParameterPanel::GetVisibilityForReferenceCounter)
 			.HAlign(HAlign_Right)
 			.VAlign(VAlign_Center)
 			.ButtonContent()
@@ -483,18 +482,6 @@ TSharedRef<SWidget> SNiagaraParameterPanel::OnGenerateWidgetForItem(const FNiaga
 	];
 }
 
-EVisibility SNiagaraParameterPanel::GetVisibilityForReferenceCounter() const
-{
-	if (ParameterPanelViewModel->GetShowReferenceCounter())
-	{
-		return EVisibility::Visible;
-	}
-	else
-	{
-		return EVisibility::Collapsed;
-	}
-}
-
 TSharedPtr<SWidget> SNiagaraParameterPanel::OnContextMenuOpening()
 {
 	const TArray<FNiagaraParameterPanelItem>& SelectedItems = ItemSelector->GetSelectedItems();
diff --git a/Engine/Plugins/FX/Niagara/Source/NiagaraEditor/Public/ViewModels/NiagaraParameterPanelViewModel.h b/Engine/Plugins/FX/Niagara/Source/NiagaraEditor/Public/ViewModels/NiagaraParameterPanelViewModel.h
--- a/Engine/Plugins/FX/Niagara/Source/NiagaraEditor/Public/ViewModels/NiagaraParameterPanelViewModel.h
+++ b/Engine/Plugins/FX/Niagara/Source/NiagaraEditor/Public/ViewModels/NiagaraParameterPanelViewModel.h
@@ -228,7 +228,6 @@ public:
 	virtual bool GetCanPasteParameterMetaDataAndToolTip(FText& OutCanPasteToolTip);
 
 	virtual bool GetAllowMultiSelect() const { return false; }
-	virtual bool GetShowReferenceCounter() const { return true; }
 
 	virtual void DuplicateParameters(const TArray<FNiagaraParameterPanelItem> ItemToDuplicate) ;
 	virtual bool GetCanDuplicateParameterAndToolTip(const TArray<FNiagaraParameterPanelItem>& ItemsToDuplicate, FText& OutCanDuplicateParameterToolTip) const;
@@ -393,7 +392,6 @@ public:
 	virtual bool GetNamespaceActive(const FName& InNamespace) const override { return true; }
 
 	virtual TSharedPtr<FNiagaraObjectSelection> GetVariableObjectSelection() override { return VariableObjectSelection; };
-	virtual bool GetShowReferenceCounter() const override;
 
 	virtual void SetActiveSection(int32 InSection) override;
 
@@ -454,7 +452,6 @@ private:
 
 	TSharedPtr<FNiagaraObjectSelection> VariableObjectSelection;
 	int32 UserOnlyIdx = -1;
-	int32 CommonOnlyIdx = -1;
 	int32 ActiveScriptIdx = -1;
 	int32 ActiveSystemIdx = -1;
 	int32 LastActiveSystemSectionIdx = -1;
diff --git a/Engine/Plugins/FX/Niagara/Source/NiagaraEditor/Public/Widgets/SNiagaraParameterPanel.h b/Engine/Plugins/FX/Niagara/Source/NiagaraEditor/Public/Widgets/SNiagaraParameterPanel.h
--- a/Engine/Plugins/FX/Niagara/Source/NiagaraEditor/Public/Widgets/SNiagaraParameterPanel.h
+++ b/Engine/Plugins/FX/Niagara/Source/NiagaraEditor/Public/Widgets/SNiagaraParameterPanel.h
@@ -134,8 +134,6 @@ private:
 
 	FText GetTooltipForSection(FText Section);
 
-	EVisibility GetVisibilityForReferenceCounter() const;
-
 private:
 	mutable bool bPendingRefresh;
 	mutable bool bPendingSelectionRestore;
