- Fix not gathering the MaxPendingTicks which meant we didn't always flush when we should
- Reduce pending frame counter 10 frames can be a lot of dispatches and potentially lead to TDR especially with fluids

#rnx

diff --git a/Engine/Plugins/FX/Niagara/Source/Niagara/Private/NiagaraGpuComputeDispatch.cpp b/Engine/Plugins/FX/Niagara/Source/Niagara/Private/NiagaraGpuComputeDispatch.cpp
--- a/Engine/Plugins/FX/Niagara/Source/Niagara/Private/NiagaraGpuComputeDispatch.cpp
+++ b/Engine/Plugins/FX/Niagara/Source/Niagara/Private/NiagaraGpuComputeDispatch.cpp
@@ -85,7 +85,7 @@ namespace FNiagaraGpuComputeDispatchLocal
 	const FName TemporalEffectTexturesName("FNiagaraGpuComputeDispatch_Textures");
 #endif // WITH_MGPU
 
-	int32 GTickFlushMaxQueuedFrames = 10;
+	int32 GTickFlushMaxQueuedFrames = 3;
 	static FAutoConsoleVariableRef CVarNiagaraTickFlushMaxQueuedFrames(
 		TEXT("fx.Niagara.Batcher.TickFlush.MaxQueuedFrames"),
 		GTickFlushMaxQueuedFrames,
@@ -333,28 +333,31 @@ void FNiagaraGpuComputeDispatch::ProcessPendingTicksFlush(FRHICommandListImmedia
 		return;
 	}
 
-	// We have pending ticks increment our counter, once we cross the threshold we will perform the appropriate operation
+	// Do we need to force a flush because we crossed the frame count threshold?
 	++FramesBeforeTickFlush;
-	int32 MaxPendingTicks = 0;
-	const int32 TickFlushMaxPendingTicks = FMath::Max(FNiagaraGpuComputeDispatchLocal::GTickFlushMaxPendingTicks, 1);
+	bForceFlush |= FramesBeforeTickFlush >= uint32(FMath::Max(0, FNiagaraGpuComputeDispatchLocal::GTickFlushMaxQueuedFrames));
 
-	if (bForceFlush || (FramesBeforeTickFlush < uint32(FNiagaraGpuComputeDispatchLocal::GTickFlushMaxQueuedFrames)))
+	// Do we need to force a flush because we crossed the max pending ticks on a single instance threashold?
+	//-OPT: Ideally we don't traverse this every frame but it's required if we need to flush as we may need to do so in batches to avoid an RDG limit
+	const int32 TickFlushMaxPendingTicks = FMath::Max(FNiagaraGpuComputeDispatchLocal::GTickFlushMaxPendingTicks, 1);
+	int32 MaxPendingTicks = 0;
+	for (int iTickStage = 0; iTickStage < ENiagaraGpuComputeTickStage::Max; ++iTickStage)
 	{
-		// If we have a lot of queued up ticks we will need to flush in batches as RDG can only handle a limited number of passes in a single Graph.
-		for (int iTickStage = 0; iTickStage < ENiagaraGpuComputeTickStage::Max; ++iTickStage)
-		{
-			for (FNiagaraSystemGpuComputeProxy* Proxy : ProxiesPerStage[iTickStage])
-			{
-				MaxPendingTicks = FMath::Max<int32>(MaxPendingTicks, Proxy->PendingTicks.Num());
-			}
-		}
-		if (!bForceFlush && (MaxPendingTicks < TickFlushMaxPendingTicks))
+		for (FNiagaraSystemGpuComputeProxy* Proxy : ProxiesPerStage[iTickStage])
 		{
-			return;
+			MaxPendingTicks = FMath::Max<int32>(MaxPendingTicks, Proxy->PendingTicks.Num());
 		}
 	}
-	FramesBeforeTickFlush = 0;
+	bForceFlush |= MaxPendingTicks >= TickFlushMaxPendingTicks;
+
+	// Do we need to execute the flush?
+	if (bForceFlush == false)
+	{
+		return;
+	}
 
+	// A tick flush was request reset the counter and execute according to the mode
+	FramesBeforeTickFlush = 0;
 	switch (FNiagaraGpuComputeDispatchLocal::GTickFlushMode)
 	{
 		// Do nothing
