- Fix some strings and reduce code

#rnx

diff --git a/Engine/Plugins/FX/Niagara/Source/Niagara/Private/NDISkeletalMesh_BoneSampling.cpp b/Engine/Plugins/FX/Niagara/Source/Niagara/Private/NDISkeletalMesh_BoneSampling.cpp
--- a/Engine/Plugins/FX/Niagara/Source/Niagara/Private/NDISkeletalMesh_BoneSampling.cpp
+++ b/Engine/Plugins/FX/Niagara/Source/Niagara/Private/NDISkeletalMesh_BoneSampling.cpp
@@ -47,304 +47,189 @@ const FName FSkeletalMeshInterfaceHelper::GetFilteredSocketOrBoneAtName("GetFilt
 
 void UNiagaraDataInterfaceSkeletalMesh::GetSkeletonSamplingFunctions(TArray<FNiagaraFunctionSignature>& OutFunctions)
 {
+	FNiagaraFunctionSignature DefaultSig;
+	DefaultSig.Inputs.Add(FNiagaraVariable(FNiagaraTypeDefinition(GetClass()), TEXT("SkeletalMesh")));
+	DefaultSig.bMemberFunction = true;
+	DefaultSig.bRequiresContext = false;
+
 	//////////////////////////////////////////////////////////////////////////
 	// Bone functions.
 	{
-		FNiagaraFunctionSignature Sig;
+		FNiagaraFunctionSignature& Sig = OutFunctions.Add_GetRef(DefaultSig);
 		Sig.Name = FSkeletalMeshInterfaceHelper::GetSkinnedBoneDataName;
-		Sig.Inputs.Add(FNiagaraVariable(FNiagaraTypeDefinition(GetClass()), TEXT("SkeletalMesh")));
 		Sig.Inputs.Add(FNiagaraVariable(FNiagaraTypeDefinition::GetIntDef(), TEXT("Bone")));
 		Sig.Outputs.Add(FNiagaraVariable(FNiagaraTypeDefinition::GetPositionDef(), TEXT("Position")));
 		Sig.Outputs.Add(FNiagaraVariable(FNiagaraTypeDefinition::GetQuatDef(), TEXT("Rotation")));
 		Sig.Outputs.Add(FNiagaraVariable(FNiagaraTypeDefinition::GetVec3Def(), TEXT("Scale")));
 		Sig.Outputs.Add(FNiagaraVariable(FNiagaraTypeDefinition::GetVec3Def(), TEXT("Velocity")));
-		Sig.bMemberFunction = true;
-		Sig.bRequiresContext = false;
-#if WITH_EDITORONLY_DATA
-		Sig.Description = LOCTEXT("GetOptionalSkinnedBoneDataDesc", "Returns skinning dependant data for the pased bone in local space. All outputs are optional and you will incur zero to minimal cost if they are not connected.");
-#endif
-		OutFunctions.Add(Sig);
+		Sig.SetDescription(LOCTEXT("GetOptionalSkinnedBoneDataDesc", "Returns skinning dependant data for the bone in local space. All outputs are optional and you will incur zero to minimal cost if they are not connected."));
 	}
 
 	{
-		FNiagaraFunctionSignature Sig;
+		FNiagaraFunctionSignature& Sig = OutFunctions.Add_GetRef(DefaultSig);
 		Sig.Name = FSkeletalMeshInterfaceHelper::GetSkinnedBoneDataWSName;
-		Sig.Inputs.Add(FNiagaraVariable(FNiagaraTypeDefinition(GetClass()), TEXT("SkeletalMesh")));
 		Sig.Inputs.Add(FNiagaraVariable(FNiagaraTypeDefinition::GetIntDef(), TEXT("Bone")));
 		Sig.Outputs.Add(FNiagaraVariable(FNiagaraTypeDefinition::GetPositionDef(), TEXT("Position")));
 		Sig.Outputs.Add(FNiagaraVariable(FNiagaraTypeDefinition::GetQuatDef(), TEXT("Rotation")));
 		Sig.Outputs.Add(FNiagaraVariable(FNiagaraTypeDefinition::GetVec3Def(), TEXT("Scale")));
 		Sig.Outputs.Add(FNiagaraVariable(FNiagaraTypeDefinition::GetVec3Def(), TEXT("Velocity")));
-		Sig.bMemberFunction = true;
-		Sig.bRequiresContext = false;
-#if WITH_EDITORONLY_DATA
-		Sig.Description = LOCTEXT("GetOptionalSkinnedBoneDataWSDesc", "Returns skinning dependant data for the pased bone in world space. All outputs are optional and you will incur zero to minimal cost if they are not connected.");
-#endif
-		OutFunctions.Add(Sig);
+		Sig.SetDescription(LOCTEXT("GetOptionalSkinnedBoneDataWSDesc", "Returns skinning dependant data for the bone in world space. All outputs are optional and you will incur zero to minimal cost if they are not connected."));
 	}
 
 
 	{
-		FNiagaraFunctionSignature Sig;
+		FNiagaraFunctionSignature& Sig = OutFunctions.Add_GetRef(DefaultSig);
 		Sig.Name = FSkeletalMeshInterfaceHelper::GetSkinnedBoneDataInterpolatedName;
-		Sig.Inputs.Add(FNiagaraVariable(FNiagaraTypeDefinition(GetClass()), TEXT("SkeletalMesh")));
 		Sig.Inputs.Add(FNiagaraVariable(FNiagaraTypeDefinition::GetIntDef(), TEXT("Bone")));
 		Sig.Inputs.Add(FNiagaraVariable(FNiagaraTypeDefinition::GetFloatDef(), TEXT("Interpolation")));
 		Sig.Outputs.Add(FNiagaraVariable(FNiagaraTypeDefinition::GetPositionDef(), TEXT("Position")));
 		Sig.Outputs.Add(FNiagaraVariable(FNiagaraTypeDefinition::GetQuatDef(), TEXT("Rotation")));
 		Sig.Outputs.Add(FNiagaraVariable(FNiagaraTypeDefinition::GetVec3Def(), TEXT("Scale")));
 		Sig.Outputs.Add(FNiagaraVariable(FNiagaraTypeDefinition::GetVec3Def(), TEXT("Velocity")));
-		Sig.bMemberFunction = true;
-		Sig.bRequiresContext = false;
-#if WITH_EDITORONLY_DATA
-		Sig.Description = LOCTEXT("GetSkinnedBoneDataDesc", "Returns skinning dependant data for the pased bone in local space. Interpolated between this frame and the previous based on passed interpolation factor. All outputs are optional and you will incur zero to minimal cost if they are not connected.");
-#endif
-		OutFunctions.Add(Sig);
+		Sig.SetDescription(LOCTEXT("GetSkinnedBoneDataDesc", "Returns skinning dependant data for the bone in local space. Interpolated between this frame and the previous based on passed interpolation factor. All outputs are optional and you will incur zero to minimal cost if they are not connected."));
 	}
 
 	{
-		FNiagaraFunctionSignature Sig;
+		FNiagaraFunctionSignature& Sig = OutFunctions.Add_GetRef(DefaultSig);
 		Sig.Name = FSkeletalMeshInterfaceHelper::GetSkinnedBoneDataWSInterpolatedName;
-		Sig.Inputs.Add(FNiagaraVariable(FNiagaraTypeDefinition(GetClass()), TEXT("SkeletalMesh")));
 		Sig.Inputs.Add(FNiagaraVariable(FNiagaraTypeDefinition::GetIntDef(), TEXT("Bone")));
 		Sig.Inputs.Add(FNiagaraVariable(FNiagaraTypeDefinition::GetFloatDef(), TEXT("Interpolation")));
 		Sig.Outputs.Add(FNiagaraVariable(FNiagaraTypeDefinition::GetPositionDef(), TEXT("Position")));
 		Sig.Outputs.Add(FNiagaraVariable(FNiagaraTypeDefinition::GetQuatDef(), TEXT("Rotation")));
 		Sig.Outputs.Add(FNiagaraVariable(FNiagaraTypeDefinition::GetVec3Def(), TEXT("Scale")));
 		Sig.Outputs.Add(FNiagaraVariable(FNiagaraTypeDefinition::GetVec3Def(), TEXT("Velocity")));
-		Sig.bMemberFunction = true;
-		Sig.bRequiresContext = false;
-#if WITH_EDITORONLY_DATA
-		Sig.Description = LOCTEXT("GetSkinnedBoneDataWSDesc", "Returns skinning dependant data for the pased bone in world space. Interpolated between this frame and the previous based on passed interpolation factor. All outputs are optional and you will incur zero to minimal cost if they are not connected.");
-#endif
-		OutFunctions.Add(Sig);
+		Sig.SetDescription(LOCTEXT("GetSkinnedBoneDataWSDesc", "Returns skinning dependant data for the bone in world space. Interpolated between this frame and the previous based on passed interpolation factor. All outputs are optional and you will incur zero to minimal cost if they are not connected."));
 	}
 
 	{
-		FNiagaraFunctionSignature Sig;
+		FNiagaraFunctionSignature& Sig = OutFunctions.Add_GetRef(DefaultSig);
 		Sig.Name = FSkeletalMeshInterfaceHelper::IsValidBoneName;
-		Sig.Inputs.Add(FNiagaraVariable(FNiagaraTypeDefinition(GetClass()), TEXT("SkeletalMesh")));
 		Sig.Inputs.Add(FNiagaraVariable(FNiagaraTypeDefinition::GetIntDef(), TEXT("Bone")));
 		Sig.Outputs.Add(FNiagaraVariable(FNiagaraTypeDefinition::GetBoolDef(), TEXT("IsValid")));
-		Sig.bMemberFunction = true;
-		Sig.bRequiresContext = false;
-#if WITH_EDITORONLY_DATA
-		Sig.Description = LOCTEXT("IsValidBoneDesc", "Determine if this bone index is valid for this mesh's skeleton.");
-#endif
-		OutFunctions.Add(Sig);
+		Sig.SetDescription(LOCTEXT("IsValidBoneDesc", "Determine if this bone index is valid for this mesh's skeleton."));
 	}
 
 	{
-		FNiagaraFunctionSignature Sig;
+		FNiagaraFunctionSignature& Sig = OutFunctions.Add_GetRef(DefaultSig);
 		Sig.Name = FSkeletalMeshInterfaceHelper::RandomBoneName;
-		Sig.Inputs.Add(FNiagaraVariable(FNiagaraTypeDefinition(GetClass()), TEXT("SkeletalMesh")));
 		Sig.Inputs.Add(FNiagaraVariable(FNiagaraTypeDefinition(FNiagaraRandInfo::StaticStruct()), TEXT("RandomInfo")));
 		Sig.Outputs.Add(FNiagaraVariable(FNiagaraTypeDefinition::GetIntDef(), TEXT("Bone")));
-		Sig.bMemberFunction = true;
-		Sig.bRequiresContext = false;
-		OutFunctions.Add(Sig);
 	}
 	{
-		FNiagaraFunctionSignature Sig;
+		FNiagaraFunctionSignature& Sig = OutFunctions.Add_GetRef(DefaultSig);
 		Sig.Name = FSkeletalMeshInterfaceHelper::GetBoneCountName;
-		Sig.Inputs.Add(FNiagaraVariable(FNiagaraTypeDefinition(GetClass()), TEXT("SkeletalMesh")));
 		Sig.Outputs.Add(FNiagaraVariable(FNiagaraTypeDefinition::GetIntDef(), TEXT("Count")));
-		Sig.bMemberFunction = true;
-		Sig.bRequiresContext = false;
-#if WITH_EDITORONLY_DATA
-		Sig.Description = LOCTEXT("GetBoneCountDesc", "Returns the number of bones in the skeletal mesh.");
-#endif
-		OutFunctions.Add(Sig);
+		Sig.SetDescription(LOCTEXT("GetBoneCountDesc", "Returns the number of bones in the skeletal mesh."));
 	}
 	{
-		FNiagaraFunctionSignature Sig;
+		FNiagaraFunctionSignature& Sig = OutFunctions.Add_GetRef(DefaultSig);
 		Sig.Name = FSkeletalMeshInterfaceHelper::GetParentBoneName;
-		Sig.Inputs.Add(FNiagaraVariable(FNiagaraTypeDefinition(GetClass()), TEXT("SkeletalMesh")));
 		Sig.Inputs.Add(FNiagaraVariable(FNiagaraTypeDefinition::GetIntDef(), TEXT("BoneIndex")));
 		Sig.Outputs.Add(FNiagaraVariable(FNiagaraTypeDefinition::GetIntDef(), TEXT("ParentIndex")));
-		Sig.bMemberFunction = true;
-		Sig.bRequiresContext = false;
-#if WITH_EDITORONLY_DATA
-		Sig.Description = LOCTEXT("GetParentBoneDesc", "Returns the parent bone index for the given bone, -1 if no parent exists.");
-#endif
-		OutFunctions.Add(Sig);
+		Sig.SetDescription(LOCTEXT("GetParentBoneDesc", "Returns the parent bone index for the given bone, -1 if no parent exists."));
 	}
 
 	{
-		FNiagaraFunctionSignature Sig;
+		FNiagaraFunctionSignature& Sig = OutFunctions.Add_GetRef(DefaultSig);
 		Sig.Name = FSkeletalMeshInterfaceHelper::RandomFilteredBoneName;
-		Sig.Inputs.Add(FNiagaraVariable(FNiagaraTypeDefinition(GetClass()), TEXT("SkeletalMesh")));
 		Sig.Inputs.Add(FNiagaraVariable(FNiagaraTypeDefinition(FNiagaraRandInfo::StaticStruct()), TEXT("RandomInfo")));
 		Sig.Outputs.Add(FNiagaraVariable(FNiagaraTypeDefinition::GetIntDef(), TEXT("Bone")));
-		Sig.bMemberFunction = true;
-		Sig.bRequiresContext = false;
-		OutFunctions.Add(Sig);
 	}
 
 	{
-		FNiagaraFunctionSignature Sig;
+		FNiagaraFunctionSignature& Sig = OutFunctions.Add_GetRef(DefaultSig);
 		Sig.Name = FSkeletalMeshInterfaceHelper::GetFilteredBoneCountName;
-		Sig.Inputs.Add(FNiagaraVariable(FNiagaraTypeDefinition(GetClass()), TEXT("SkeletalMesh")));
 		Sig.Outputs.Add(FNiagaraVariable(FNiagaraTypeDefinition::GetIntDef(), TEXT("Count")));
-		Sig.bMemberFunction = true;
-		Sig.bRequiresContext = false;
-#if WITH_EDITORONLY_DATA
-		Sig.Description = LOCTEXT("GetFilteredBoneCountDesc", "Returns the number of filtered bones in the DI list.");
-#endif
-		OutFunctions.Add(Sig);
+		Sig.SetDescription(LOCTEXT("GetFilteredBoneCountDesc", "Returns the number of filtered bones in the DI list."));
 	}
 
 	{
-		FNiagaraFunctionSignature Sig;
+		FNiagaraFunctionSignature& Sig = OutFunctions.Add_GetRef(DefaultSig);
 		Sig.Name = FSkeletalMeshInterfaceHelper::GetFilteredBoneAtName;
-		Sig.Inputs.Add(FNiagaraVariable(FNiagaraTypeDefinition(GetClass()), TEXT("SkeletalMesh")));
 		Sig.Inputs.Add(FNiagaraVariable(FNiagaraTypeDefinition::GetIntDef(), TEXT("Bone Index")));
 		Sig.Outputs.Add(FNiagaraVariable(FNiagaraTypeDefinition::GetIntDef(), TEXT("Bone")));
-		Sig.bMemberFunction = true;
-		Sig.bRequiresContext = false;
-#if WITH_EDITORONLY_DATA
-		Sig.Description = LOCTEXT("GetFilteredBoneAtDesc", "Gets the bone at the passed index in the DI's filter bones list.");
-#endif
-		OutFunctions.Add(Sig);
+		Sig.SetDescription(LOCTEXT("GetFilteredBoneAtDesc", "Gets the bone at the passed index in the DI's filter bones list."));
 	}
 
 	{
-		FNiagaraFunctionSignature Sig;
+		FNiagaraFunctionSignature& Sig = OutFunctions.Add_GetRef(DefaultSig);
 		Sig.Name = FSkeletalMeshInterfaceHelper::RandomUnfilteredBoneName;
-		Sig.Inputs.Add(FNiagaraVariable(FNiagaraTypeDefinition(GetClass()), TEXT("SkeletalMesh")));
 		Sig.Inputs.Add(FNiagaraVariable(FNiagaraTypeDefinition(FNiagaraRandInfo::StaticStruct()), TEXT("RandomInfo")));
 		Sig.Outputs.Add(FNiagaraVariable(FNiagaraTypeDefinition::GetIntDef(), TEXT("Bone")));
-		Sig.bMemberFunction = true;
-		Sig.bRequiresContext = false;
-		OutFunctions.Add(Sig);
 	}
 
 	{
-		FNiagaraFunctionSignature Sig;
+		FNiagaraFunctionSignature& Sig = OutFunctions.Add_GetRef(DefaultSig);
 		Sig.Name = FSkeletalMeshInterfaceHelper::GetUnfilteredBoneCountName;
-		Sig.Inputs.Add(FNiagaraVariable(FNiagaraTypeDefinition(GetClass()), TEXT("SkeletalMesh")));
 		Sig.Outputs.Add(FNiagaraVariable(FNiagaraTypeDefinition::GetIntDef(), TEXT("Count")));
-		Sig.bMemberFunction = true;
-		Sig.bRequiresContext = false;
-#if WITH_EDITORONLY_DATA
-		Sig.Description = LOCTEXT("GetUnfilteredBoneCountDesc", "Returns the number of unfiltered bones (i.e. the exclusion of filtered bones) in the DI list.");
-#endif
-		OutFunctions.Add(Sig);
+		Sig.SetDescription(LOCTEXT("GetUnfilteredBoneCountDesc", "Returns the number of unfiltered bones (i.e. the exclusion of filtered bones) in the DI list."));
 	}
 
 	{
-		FNiagaraFunctionSignature Sig;
+		FNiagaraFunctionSignature& Sig = OutFunctions.Add_GetRef(DefaultSig);
 		Sig.Name = FSkeletalMeshInterfaceHelper::GetUnfilteredBoneAtName;
-		Sig.Inputs.Add(FNiagaraVariable(FNiagaraTypeDefinition(GetClass()), TEXT("SkeletalMesh")));
 		Sig.Inputs.Add(FNiagaraVariable(FNiagaraTypeDefinition::GetIntDef(), TEXT("Bone Index")));
 		Sig.Outputs.Add(FNiagaraVariable(FNiagaraTypeDefinition::GetIntDef(), TEXT("Bone")));
-		Sig.bMemberFunction = true;
-		Sig.bRequiresContext = false;
-#if WITH_EDITORONLY_DATA
-		Sig.Description = LOCTEXT("GetUnfilteredBoneAtDesc", "Gets the bone at the passed index from the exlusion of the DI's filter bones list.");
-#endif
-		OutFunctions.Add(Sig);
+		Sig.SetDescription(LOCTEXT("GetUnfilteredBoneAtDesc", "Gets the bone at the passed index from the exlusion of the DI's filter bones list."));
 	}
 
 	//////////////////////////////////////////////////////////////////////////
 	//Socket functions
 	{
-		FNiagaraFunctionSignature Sig;
+		FNiagaraFunctionSignature& Sig = OutFunctions.Add_GetRef(DefaultSig);
 		Sig.Name = FSkeletalMeshInterfaceHelper::RandomFilteredSocketName;
-		Sig.Inputs.Add(FNiagaraVariable(FNiagaraTypeDefinition(GetClass()), TEXT("SkeletalMesh")));
 		Sig.Inputs.Add(FNiagaraVariable(FNiagaraTypeDefinition(FNiagaraRandInfo::StaticStruct()), TEXT("RandomInfo")));
 		Sig.Outputs.Add(FNiagaraVariable(FNiagaraTypeDefinition::GetIntDef(), TEXT("Socket Bone")));
-		Sig.bMemberFunction = true;
-		Sig.bRequiresContext = false;
-#if WITH_EDITORONLY_DATA
-		Sig.Description = LOCTEXT("RandomFilteredSocketDesc", "Gets the bone for a random socket in the DI's filtered socket list.");
-#endif
-		OutFunctions.Add(Sig);
+		Sig.SetDescription(LOCTEXT("RandomFilteredSocketDesc", "Gets the bone for a random socket in the DI's filtered socket list."));
 	}
 	
 	{
-		FNiagaraFunctionSignature Sig;
+		FNiagaraFunctionSignature& Sig = OutFunctions.Add_GetRef(DefaultSig);
 		Sig.Name = FSkeletalMeshInterfaceHelper::GetFilteredSocketCountName;
-		Sig.Inputs.Add(FNiagaraVariable(FNiagaraTypeDefinition(GetClass()), TEXT("SkeletalMesh")));
 		Sig.Outputs.Add(FNiagaraVariable(FNiagaraTypeDefinition::GetIntDef(), TEXT("Count")));
-		Sig.bMemberFunction = true;
-		Sig.bRequiresContext = false;
-#if WITH_EDITORONLY_DATA
-		Sig.Description = LOCTEXT("GetFilteredSocketCountDesc", "Returns the number of filtered Sockets in the DI list.");
-#endif
-		OutFunctions.Add(Sig);
+		Sig.SetDescription(LOCTEXT("GetFilteredSocketCountDesc", "Returns the number of filtered Sockets in the DI list."));
 	}
 
 	{
-		FNiagaraFunctionSignature Sig;
+		FNiagaraFunctionSignature& Sig = OutFunctions.Add_GetRef(DefaultSig);
 		Sig.Name = FSkeletalMeshInterfaceHelper::GetFilteredSocketBoneAtName;
-		Sig.Inputs.Add(FNiagaraVariable(FNiagaraTypeDefinition(GetClass()), TEXT("SkeletalMesh")));
 		Sig.Inputs.Add(FNiagaraVariable(FNiagaraTypeDefinition::GetIntDef(), TEXT("Socket Index")));
 		Sig.Outputs.Add(FNiagaraVariable(FNiagaraTypeDefinition::GetIntDef(), TEXT("Socket Bone")));
-		Sig.bMemberFunction = true;
-		Sig.bRequiresContext = false;
-#if WITH_EDITORONLY_DATA
-		Sig.Description = LOCTEXT("GetFilteredSocketBoneAtDesc", "Gets the bone for the socket at the passed index in the DI's filtered socket list.");
-#endif
-		OutFunctions.Add(Sig);
+		Sig.SetDescription(LOCTEXT("GetFilteredSocketBoneAtDesc", "Gets the bone for the socket at the passed index in the DI's filtered socket list."));
 	}
 
 	{
-		FNiagaraFunctionSignature Sig;
+		FNiagaraFunctionSignature& Sig = OutFunctions.Add_GetRef(DefaultSig);
 		Sig.Name = FSkeletalMeshInterfaceHelper::GetFilteredSocketTransformName;
-		Sig.Inputs.Add(FNiagaraVariable(FNiagaraTypeDefinition(GetClass()), TEXT("SkeletalMesh")));
 		Sig.Inputs.Add(FNiagaraVariable(FNiagaraTypeDefinition::GetIntDef(), TEXT("Socket Index")));
 		Sig.Inputs.Add(FNiagaraVariable(FNiagaraTypeDefinition::GetBoolDef(), TEXT("Apply World Transform")));
 		Sig.Outputs.Add(FNiagaraVariable(FNiagaraTypeDefinition::GetVec3Def(), TEXT("Socket Translation")));
 		Sig.Outputs.Add(FNiagaraVariable(FNiagaraTypeDefinition::GetQuatDef(), TEXT("Socket Rotation")));
 		Sig.Outputs.Add(FNiagaraVariable(FNiagaraTypeDefinition::GetVec3Def(), TEXT("Socket Scale")));
-		Sig.bMemberFunction = true;
-		Sig.bRequiresContext = false;
-#if WITH_EDITORONLY_DATA
-		Sig.Description = LOCTEXT("GetFilteredSocketTransformDesc", "Gets the transform for the socket at the passed index in the DI's filtered socket list. If the Source component is set it will respect the Relative Transform Space as well..");
-#endif
-		OutFunctions.Add(Sig);
+		Sig.SetDescription(LOCTEXT("GetFilteredSocketTransformDesc", "Gets the transform for the socket at the passed index in the DI's filtered socket list. If the Source component is set it will respect the Relative Transform Space as well.."));
 	}
 
 	//////////////////////////////////////////////////////////////////////////
 	// Misc Functions
 	{
-		FNiagaraFunctionSignature Sig;
+		FNiagaraFunctionSignature& Sig = OutFunctions.Add_GetRef(DefaultSig);
 		Sig.Name = FSkeletalMeshInterfaceHelper::RandomFilteredSocketOrBoneName;
-		Sig.Inputs.Add(FNiagaraVariable(FNiagaraTypeDefinition(GetClass()), TEXT("SkeletalMesh")));
 		Sig.Inputs.Add(FNiagaraVariable(FNiagaraTypeDefinition(FNiagaraRandInfo::StaticStruct()), TEXT("RandomInfo")));
 		Sig.Outputs.Add(FNiagaraVariable(FNiagaraTypeDefinition::GetIntDef(), TEXT("Bone")));
-		Sig.bMemberFunction = true;
-		Sig.bRequiresContext = false;
-#if WITH_EDITORONLY_DATA
-		Sig.Description = LOCTEXT("RandomFilteredSocketOrBoneDesc", "Gets the bone for a random filtered socket or bone from the DI's list.");
-#endif
-		OutFunctions.Add(Sig);
+		Sig.SetDescription(LOCTEXT("RandomFilteredSocketOrBoneDesc", "Gets the bone for a random filtered socket or bone from the DI's list."));
 	}
 	{
-		FNiagaraFunctionSignature Sig;
+		FNiagaraFunctionSignature& Sig = OutFunctions.Add_GetRef(DefaultSig);
 		Sig.Name = FSkeletalMeshInterfaceHelper::GetFilteredSocketOrBoneCountName;
-		Sig.Inputs.Add(FNiagaraVariable(FNiagaraTypeDefinition(GetClass()), TEXT("SkeletalMesh")));
 		Sig.Outputs.Add(FNiagaraVariable(FNiagaraTypeDefinition::GetIntDef(), TEXT("Count")));
-		Sig.bMemberFunction = true;
-		Sig.bRequiresContext = false;
-#if WITH_EDITORONLY_DATA
-		Sig.Description = LOCTEXT("GetFilteredSocketOrBoneCountDesc", "Gets the total filtered socket and bone count from the DI's list.");
-#endif
-		OutFunctions.Add(Sig);
+		Sig.SetDescription(LOCTEXT("GetFilteredSocketOrBoneCountDesc", "Gets the total filtered socket and bone count from the DI's list."));
 	}
 	{
-		FNiagaraFunctionSignature Sig;
+		FNiagaraFunctionSignature& Sig = OutFunctions.Add_GetRef(DefaultSig);
 		Sig.Name = FSkeletalMeshInterfaceHelper::GetFilteredSocketOrBoneAtName;
-		Sig.Inputs.Add(FNiagaraVariable(FNiagaraTypeDefinition(GetClass()), TEXT("SkeletalMesh")));
 		Sig.Inputs.Add(FNiagaraVariable(FNiagaraTypeDefinition::GetIntDef(), TEXT("Socket Or Bone Index")));
 		Sig.Outputs.Add(FNiagaraVariable(FNiagaraTypeDefinition::GetIntDef(), TEXT("Bone")));
-		Sig.bMemberFunction = true;
-		Sig.bRequiresContext = false;
-#if WITH_EDITORONLY_DATA
-		Sig.Description = LOCTEXT("GetFilteredSocketOrBoneAtDesc", "Gets a filtered socket or bone count from the DI's list.");
-#endif
-		OutFunctions.Add(Sig);
+		Sig.SetDescription(LOCTEXT("GetFilteredSocketOrBoneAtDesc", "Gets a filtered socket or bone count from the DI's list."));
 	}
 }
 
