Fix for non-determinism with version guids of NiagaraScript and NiagaraEmitter assets

diff --git a/Engine/Plugins/FX/Niagara/Source/Niagara/Private/NiagaraEmitter.cpp b/Engine/Plugins/FX/Niagara/Source/Niagara/Private/NiagaraEmitter.cpp
--- a/Engine/Plugins/FX/Niagara/Source/Niagara/Private/NiagaraEmitter.cpp
+++ b/Engine/Plugins/FX/Niagara/Source/Niagara/Private/NiagaraEmitter.cpp
@@ -2592,6 +2592,9 @@ void UNiagaraEmitter::CheckVersionDataAvailable()
 		Data.VersionedParentAtLastMerge.Emitter = ParentAtLastMerge_DEPRECATED;
 	}
 
+	// create a stable initial version guid for our versioned data based on the UNiagaraEmitter
+	Data.Version.VersionGuid = FNiagaraAssetVersion::CreateStableVersionGuid(this);
+
 	ExposedVersion = Data.Version.VersionGuid;
 
 	for (UNiagaraRendererProperties* Renderer : Data.RendererProperties)
diff --git a/Engine/Plugins/FX/Niagara/Source/Niagara/Private/NiagaraScript.cpp b/Engine/Plugins/FX/Niagara/Source/Niagara/Private/NiagaraScript.cpp
--- a/Engine/Plugins/FX/Niagara/Source/Niagara/Private/NiagaraScript.cpp
+++ b/Engine/Plugins/FX/Niagara/Source/Niagara/Private/NiagaraScript.cpp
@@ -613,6 +613,9 @@ void UNiagaraScript::CheckVersionDataAvailable()
 	Data.DeprecationRecommendation = DeprecationRecommendation_DEPRECATED;
 	Data.NumericOutputTypeSelectionMode = NumericOutputTypeSelectionMode_DEPRECATED;
 
+	// create a stable initial version guid for our versioned data based on the UNiagaraScript
+	Data.Version.VersionGuid = FNiagaraAssetVersion::CreateStableVersionGuid(this);
+
 	ExposedVersion = Data.Version.VersionGuid;
 }
 
diff --git a/Engine/Plugins/FX/Niagara/Source/Niagara/Private/NiagaraTypes.cpp b/Engine/Plugins/FX/Niagara/Source/Niagara/Private/NiagaraTypes.cpp
--- a/Engine/Plugins/FX/Niagara/Source/Niagara/Private/NiagaraTypes.cpp
+++ b/Engine/Plugins/FX/Niagara/Source/Niagara/Private/NiagaraTypes.cpp
@@ -101,6 +101,27 @@ bool FVersionedNiagaraEmitterWeakPtr::IsValid() const
 	return GetEmitterData() != nullptr;
 }
 
+FGuid FNiagaraAssetVersion::CreateStableVersionGuid(UObject* Object)
+{
+	if (Object)
+	{
+		uint32 HashBuffer[]{ 0, 0, 0, 0 };
+
+		{
+			TStringBuilder<256> ObjectPathString;
+			Object->GetPathName(nullptr, ObjectPathString);
+
+			FMD5 ObjectPathHash;
+			ObjectPathHash.Update(reinterpret_cast<uint8*>(GetData(ObjectPathString)), ObjectPathString.Len() * sizeof(TCHAR));
+			ObjectPathHash.Final(reinterpret_cast<uint8*>(&HashBuffer));
+		}
+
+		return FGuid(HashBuffer[0], HashBuffer[1], HashBuffer[2], HashBuffer[3]);
+	}
+
+	return FGuid();
+}
+
 FNiagaraLWCConverter::FNiagaraLWCConverter(FVector InSystemWorldPos)
 {
 	SystemWorldPos = InSystemWorldPos;
diff --git a/Engine/Plugins/FX/Niagara/Source/Niagara/Public/NiagaraTypes.h b/Engine/Plugins/FX/Niagara/Source/Niagara/Public/NiagaraTypes.h
--- a/Engine/Plugins/FX/Niagara/Source/Niagara/Public/NiagaraTypes.h
+++ b/Engine/Plugins/FX/Niagara/Source/Niagara/Public/NiagaraTypes.h
@@ -201,6 +201,8 @@ struct NIAGARA_API FNiagaraAssetVersion
 	bool operator!=(const FNiagaraAssetVersion& Other) const { return !(*this == Other); }
 	bool operator<(const FNiagaraAssetVersion& Other) const { return MajorVersion < Other.MajorVersion || (MajorVersion == Other.MajorVersion && MinorVersion < Other.MinorVersion); }
 	bool operator<=(const FNiagaraAssetVersion& Other) const { return *this < Other || (MajorVersion == Other.MajorVersion && MinorVersion == Other.MinorVersion); }
+
+	static FGuid CreateStableVersionGuid(UObject* Object);
 };
 
 struct NIAGARA_API FNiagaraLWCConverter
