Add project setting to control default behavior when dealing with experimental and legacy VM

diff --git a/Engine/Plugins/FX/Niagara/Source/Niagara/Classes/NiagaraScriptSourceBase.h b/Engine/Plugins/FX/Niagara/Source/Niagara/Classes/NiagaraScriptSourceBase.h
--- a/Engine/Plugins/FX/Niagara/Source/Niagara/Classes/NiagaraScriptSourceBase.h
+++ b/Engine/Plugins/FX/Niagara/Source/Niagara/Classes/NiagaraScriptSourceBase.h
@@ -61,6 +61,7 @@ public:
 	NIAGARA_API static const FString GpuScriptDefine;
 	NIAGARA_API static const FString EventSpawnDefine;
 	NIAGARA_API static const FString EventSpawnInitialAttribWritesDefine;
+	NIAGARA_API static const FString ExperimentalVMDisabled;
 
 	FNiagaraCompileOptions() : TargetUsage(ENiagaraScriptUsage::Function), TargetUsageBitmask(0)
 	{
diff --git a/Engine/Plugins/FX/Niagara/Source/Niagara/Private/NiagaraAsyncCompile.cpp b/Engine/Plugins/FX/Niagara/Source/Niagara/Private/NiagaraAsyncCompile.cpp
--- a/Engine/Plugins/FX/Niagara/Source/Niagara/Private/NiagaraAsyncCompile.cpp
+++ b/Engine/Plugins/FX/Niagara/Source/Niagara/Private/NiagaraAsyncCompile.cpp
@@ -114,6 +114,8 @@ FNiagaraAsyncCompileTask::FNiagaraAsyncCompileTask(UNiagaraSystem* InOwningSyste
 	ScriptPair = InScriptPair;
 
 	CurrentState = ENiagaraCompilationState::CheckDDC;
+
+	bExperimentalVMDisabled = InScriptPair.CompileId.AdditionalDefines.Contains(FNiagaraCompileOptions::ExperimentalVMDisabled);
 }
 
 void FNiagaraAsyncCompileTask::ProcessCurrentState()
@@ -327,6 +329,13 @@ void FNiagaraAsyncCompileTask::OptimizeByteCode()
 {
 	check(ExeData);
 
+#if VECTORVM_SUPPORTS_LEGACY
+	if (bExperimentalVMDisabled)
+	{
+		return;
+	}
+#endif
+
 #if VECTORVM_SUPPORTS_EXPERIMENTAL
 	if (const uint8* ByteCode = ExeData->ByteCode.GetDataPtr())
 	{
diff --git a/Engine/Plugins/FX/Niagara/Source/Niagara/Private/NiagaraAsyncCompile.h b/Engine/Plugins/FX/Niagara/Source/Niagara/Private/NiagaraAsyncCompile.h
--- a/Engine/Plugins/FX/Niagara/Source/Niagara/Private/NiagaraAsyncCompile.h
+++ b/Engine/Plugins/FX/Niagara/Source/Niagara/Private/NiagaraAsyncCompile.h
@@ -69,6 +69,7 @@ public:
 	bool bWaitForCompileJob = false;
 	bool bUsedShaderCompilerWorker = false;
 	bool bFetchedGCObjects = false;
+	bool bExperimentalVMDisabled = true;
 	UNiagaraSystem* OwningSystem;
 	FEmitterCompiledScriptPair ScriptPair;
 	TArray<FNiagaraVariable> EncounteredExposedVars;
diff --git a/Engine/Plugins/FX/Niagara/Source/Niagara/Private/NiagaraCommon.cpp b/Engine/Plugins/FX/Niagara/Source/Niagara/Private/NiagaraCommon.cpp
--- a/Engine/Plugins/FX/Niagara/Source/Niagara/Private/NiagaraCommon.cpp
+++ b/Engine/Plugins/FX/Niagara/Source/Niagara/Private/NiagaraCommon.cpp
@@ -1523,6 +1523,7 @@ const FString FNiagaraCompileOptions::CpuScriptDefine = TEXT("CPUSim");
 const FString FNiagaraCompileOptions::GpuScriptDefine = TEXT("GPUComputeSim");
 const FString FNiagaraCompileOptions::EventSpawnDefine = TEXT("EventSpawn");
 const FString FNiagaraCompileOptions::EventSpawnInitialAttribWritesDefine = TEXT("EventSpawnInitialAttribWrites");
+const FString FNiagaraCompileOptions::ExperimentalVMDisabled = TEXT("ExperimentalVMDisabled");
 
 FSynchronizeWithParameterDefinitionsArgs::FSynchronizeWithParameterDefinitionsArgs()
 	: SpecificDefinitionsUniqueIds(TArray<FGuid>())
diff --git a/Engine/Plugins/FX/Niagara/Source/Niagara/Private/NiagaraScript.cpp b/Engine/Plugins/FX/Niagara/Source/Niagara/Private/NiagaraScript.cpp
--- a/Engine/Plugins/FX/Niagara/Source/Niagara/Private/NiagaraScript.cpp
+++ b/Engine/Plugins/FX/Niagara/Source/Niagara/Private/NiagaraScript.cpp
@@ -966,6 +966,11 @@ void UNiagaraScript::ComputeVMCompilationId(FNiagaraVMExecutableDataId& Id, FGui
 				Id.AdditionalDefines.Add(TEXT("IgnoreParticleReadsForAttributeTrim"));
 			}
 
+			if (EmitterOwner->ShouldDisableExperimentalVM() || !GetDefault<UNiagaraSettings>()->bExperimentalVMEnabled)
+			{
+				Id.AdditionalDefines.Add(FNiagaraCompileOptions::ExperimentalVMDisabled);
+			}
+
 			bool TrimAttributes = EmitterOwner->ShouldTrimAttributes();
 			if (TrimAttributes)
 			{
diff --git a/Engine/Plugins/FX/Niagara/Source/Niagara/Public/NiagaraSettings.h b/Engine/Plugins/FX/Niagara/Source/Niagara/Public/NiagaraSettings.h
--- a/Engine/Plugins/FX/Niagara/Source/Niagara/Public/NiagaraSettings.h
+++ b/Engine/Plugins/FX/Niagara/Source/Niagara/Public/NiagaraSettings.h
@@ -118,6 +118,14 @@ class NIAGARA_API UNiagaraSettings : public UDeveloperSettings
 	UPROPERTY(config, EditAnywhere, Category = Niagara, meta = ( DisplayName = "Enforce strict type checks in the graph" ))
 	bool bEnforceStrictStackTypes = true;
 
+	/**
+	 True indicates that we will generate byte code for the new experimental VM.  Control over whether the new VM will
+	 be used when executing NiagaraScripts will also take into account the overrides on the system (bDisableExperimentalVM) and
+	 the cvars fx.NiagaraScript.StripByteCodeOnLoad and fx.ForceExecVMPath.
+	*/
+	UPROPERTY(config, EditAnywhere, Category = Niagara, meta = (DisplayName = "Enable building data for Experimental VM"))
+	bool bExperimentalVMEnabled = false;
+
 	/** Default effect type to use for effects that don't define their own. Can be null. */
 	UPROPERTY(config, EditAnywhere, Category = Niagara, meta = (AllowedClasses = "/Script/Niagara.NiagaraEffectType"))
 	FSoftObjectPath DefaultEffectType;
