- Add get velocity to actor DI

diff --git a/Engine/Plugins/FX/Niagara/Shaders/Private/NiagaraDataInterfaceActorComponentTemplate.ush b/Engine/Plugins/FX/Niagara/Shaders/Private/NiagaraDataInterfaceActorComponentTemplate.ush
--- a/Engine/Plugins/FX/Niagara/Shaders/Private/NiagaraDataInterfaceActorComponentTemplate.ush
+++ b/Engine/Plugins/FX/Niagara/Shaders/Private/NiagaraDataInterfaceActorComponentTemplate.ush
@@ -4,6 +4,7 @@ uint			{ParameterName}_Valid;
 float4x4		{ParameterName}_Matrix;
 float4			{ParameterName}_Rotation;
 float3			{ParameterName}_Scale;
+float3			{ParameterName}_Velocity;
 
 void GetMatrix_{ParameterName}(out bool bOutValid, out float4x4 OutMatrix)
 {
@@ -18,3 +19,10 @@ void GetTransform_{ParameterName}(out bool bOutValid, out float3 OutPosition, ou
 	OutRotation = {ParameterName}_Rotation;
 	OutScale = {ParameterName}_Scale;
 }
+
+void GetVelocity_{ParameterName}(out bool bOutValid, out float3 OutVelocity)
+{
+	bOutValid = {ParameterName}_Valid != 0;
+	OutVelocity = {ParameterName}_Velocity;
+}
+
diff --git a/Engine/Plugins/FX/Niagara/Source/Niagara/Private/DataInterface/NiagaraDataInterfaceActorComponent.cpp b/Engine/Plugins/FX/Niagara/Source/Niagara/Private/DataInterface/NiagaraDataInterfaceActorComponent.cpp
--- a/Engine/Plugins/FX/Niagara/Source/Niagara/Private/DataInterface/NiagaraDataInterfaceActorComponent.cpp
+++ b/Engine/Plugins/FX/Niagara/Source/Niagara/Private/DataInterface/NiagaraDataInterfaceActorComponent.cpp
@@ -34,6 +34,7 @@ namespace NDIActorComponentLocal
 
 	static const FName	GetMatrixName(TEXT("GetMatrix"));
 	static const FName	GetTransformName(TEXT("GetTransform"));
+	static const FName	GetVelocityName(TEXT("GetVelocity"));
 
 	static const FString ValidString(TEXT("Valid"));
 	static const FString MatrixString(TEXT("Matrix"));
@@ -45,6 +46,7 @@ namespace NDIActorComponentLocal
 		FNiagaraParameterDirectBinding<UObject*>	UserParamBinding;
 		bool										bCachedValid = false;
 		FTransform									CachedTransform = FTransform::Identity;
+		FVector3f									CachedVelocity = FVector3f::ZeroVector;
 
 		// our use of UserParamBinding can occur within CalculateTickGroup which occurs before we tick our parameter stores.  This
 		// can lead to a stale UObject reference being accessed (if the actor we're pointing at is deleted).  For now we cache
@@ -56,12 +58,14 @@ namespace NDIActorComponentLocal
 	{
 		bool		bCachedValid = false;
 		FTransform	CachedTransform = FTransform::Identity;
+		FVector3f	CachedVelocity = FVector3f::ZeroVector;
 	};
 
 	struct FInstanceData_RenderThread
 	{
 		bool		bCachedValid = false;
 		FTransform	CachedTransform = FTransform::Identity;
+		FVector3f	CachedVelocity = FVector3f::ZeroVector;
 	};
 
 	struct FNDIProxy : public FNiagaraDataInterfaceProxy
@@ -74,6 +78,7 @@ namespace NDIActorComponentLocal
 			FGameToRenderInstanceData* DataForRenderThread = reinterpret_cast<FGameToRenderInstanceData*>(InDataForRenderThread);
 			DataForRenderThread->bCachedValid = InstanceData->bCachedValid;
 			DataForRenderThread->CachedTransform = InstanceData->CachedTransform;
+			DataForRenderThread->CachedVelocity = InstanceData->CachedVelocity;
 		}
 
 		virtual void ConsumePerInstanceDataFromGameThread(void* PerInstanceData, const FNiagaraSystemInstanceID& InstanceID) override
@@ -83,6 +88,7 @@ namespace NDIActorComponentLocal
 			FInstanceData_RenderThread& InstanceData = SystemInstancesToInstanceData_RT.FindOrAdd(InstanceID);
 			InstanceData.bCachedValid = InstanceDataFromGT->bCachedValid;
 			InstanceData.CachedTransform = InstanceDataFromGT->CachedTransform;
+			InstanceData.CachedVelocity = InstanceDataFromGT->CachedVelocity;
 		}
 
 		TMap<FNiagaraSystemInstanceID, FInstanceData_RenderThread> SystemInstancesToInstanceData_RT;
@@ -134,31 +140,34 @@ UActorComponent* UNiagaraDataInterfaceActorComponent::ResolveComponent(const voi
 void UNiagaraDataInterfaceActorComponent::GetFunctions(TArray<FNiagaraFunctionSignature>& OutFunctions)
 {
 	using namespace NDIActorComponentLocal;
+	FNiagaraFunctionSignature DefaultSig;
+	DefaultSig.Inputs.Emplace(FNiagaraTypeDefinition(GetClass()), TEXT("ActorComponent"));
+	DefaultSig.bMemberFunction = true;
+	DefaultSig.bRequiresContext = false;
+	DefaultSig.bSupportsGPU = true;
+	DefaultSig.SetFunctionVersion(FNiagaraActorDIFunctionVersion::LatestVersion);
 	{
-		FNiagaraFunctionSignature& FunctionSignature = OutFunctions.AddDefaulted_GetRef();
+		FNiagaraFunctionSignature& FunctionSignature = OutFunctions.Add_GetRef(DefaultSig);
 		FunctionSignature.Name = GetMatrixName;
-		FunctionSignature.SetDescription(LOCTEXT("GetMatrix", "Returns the current matrix for the component if valid."));
-		FunctionSignature.SetFunctionVersion(FNiagaraActorDIFunctionVersion::LatestVersion);
-		FunctionSignature.bMemberFunction = true;
-		FunctionSignature.bRequiresContext = false;
-		FunctionSignature.bSupportsGPU = true;
-		FunctionSignature.Inputs.Emplace(FNiagaraTypeDefinition(GetClass()), TEXT("ActorComponent"));
 		FunctionSignature.Outputs.Emplace(FNiagaraTypeDefinition::GetBoolDef(), TEXT("IsValid"));
 		FunctionSignature.Outputs.Emplace(FNiagaraTypeDefinition::GetMatrix4Def(), TEXT("Matrix"));
+		FunctionSignature.SetDescription(LOCTEXT("GetMatrix", "Returns the current matrix for the component if valid."));
 	}
 	{
-		FNiagaraFunctionSignature& FunctionSignature = OutFunctions.AddDefaulted_GetRef();
+		FNiagaraFunctionSignature& FunctionSignature = OutFunctions.Add_GetRef(DefaultSig);
 		FunctionSignature.Name = GetTransformName;
-		FunctionSignature.SetDescription(LOCTEXT("GetTransform", "Returns the current transform for the component if valid."));
-		FunctionSignature.SetFunctionVersion(FNiagaraActorDIFunctionVersion::LatestVersion);
-		FunctionSignature.bMemberFunction = true;
-		FunctionSignature.bRequiresContext = false;
-		FunctionSignature.bSupportsGPU = true;
-		FunctionSignature.Inputs.Emplace(FNiagaraTypeDefinition(GetClass()), TEXT("ActorComponent"));
 		FunctionSignature.Outputs.Emplace(FNiagaraTypeDefinition::GetBoolDef(), TEXT("IsValid"));
 		FunctionSignature.Outputs.Emplace(FNiagaraTypeDefinition::GetPositionDef(), TEXT("Position"));
 		FunctionSignature.Outputs.Emplace(FNiagaraTypeDefinition::GetQuatDef(), TEXT("Rotation"));
 		FunctionSignature.Outputs.Emplace(FNiagaraTypeDefinition::GetVec3Def(), TEXT("Scale"));
+		FunctionSignature.SetDescription(LOCTEXT("GetTransform", "Returns the current transform for the component if valid."));
+	}
+	{
+		FNiagaraFunctionSignature& FunctionSignature = OutFunctions.Add_GetRef(DefaultSig);
+		FunctionSignature.Name = GetVelocityName;
+		FunctionSignature.Outputs.Emplace(FNiagaraTypeDefinition::GetBoolDef(), TEXT("IsValid"));
+		FunctionSignature.Outputs.Emplace(FNiagaraTypeDefinition::GetVec3Def(), TEXT("Velocity"));
+		FunctionSignature.SetDescription(LOCTEXT("GetVelocityDesc", "Returns the velocity from the component or actor if valid."));
 	}
 }
 
@@ -173,6 +182,10 @@ void UNiagaraDataInterfaceActorComponent::GetVMExternalFunction(const FVMExterna
 	{
 		OutFunc = FVMExternalFunction::CreateLambda([this](FVectorVMExternalFunctionContext& Context) { this->VMGetTransform(Context); });
 	}
+	else if (BindingInfo.Name == GetVelocityName)
+	{
+		OutFunc = FVMExternalFunction::CreateLambda([this](FVectorVMExternalFunctionContext& Context) { this->VMGetPhysicsVelocity(Context); });
+	}
 }
 
 #if WITH_EDITORONLY_DATA
@@ -200,7 +213,7 @@ void UNiagaraDataInterfaceActorComponent::GetParameterDefinitionHLSL(const FNiag
 bool UNiagaraDataInterfaceActorComponent::GetFunctionHLSL(const FNiagaraDataInterfaceGPUParamInfo& ParamInfo, const FNiagaraDataInterfaceGeneratedFunction& FunctionInfo, int FunctionInstanceIndex, FString& OutHLSL)
 {
 	using namespace NDIActorComponentLocal;
-	return (FunctionInfo.DefinitionName == GetMatrixName) || (FunctionInfo.DefinitionName == GetTransformName);
+	return (FunctionInfo.DefinitionName == GetMatrixName) || (FunctionInfo.DefinitionName == GetTransformName) || (FunctionInfo.DefinitionName == GetVelocityName);
 }
 
 bool UNiagaraDataInterfaceActorComponent::UpgradeFunctionCall(FNiagaraFunctionSignature& FunctionSignature)
@@ -287,8 +300,9 @@ bool UNiagaraDataInterfaceActorComponent::PerInstanceTick(void* PerInstanceData,
 
 	InstanceData->bCachedValid = false;
 	InstanceData->CachedTransform = FTransform::Identity;
-	UActorComponent* ActorComponent = ResolveComponent(PerInstanceData);
+	InstanceData->CachedVelocity = FVector3f::ZeroVector;
 
+	UActorComponent* ActorComponent = ResolveComponent(PerInstanceData);
 	if (ActorComponent)
 	{
 		if (USceneComponent* SceneComponent = Cast<USceneComponent>(ActorComponent) )
@@ -296,12 +310,21 @@ bool UNiagaraDataInterfaceActorComponent::PerInstanceTick(void* PerInstanceData,
 			InstanceData->bCachedValid = true;
 			InstanceData->CachedTransform = SceneComponent->GetComponentToWorld();
 			InstanceData->CachedTransform.AddToTranslation(FVector(SystemInstance->GetLWCTile()) * -FLargeWorldRenderScalar::GetTileSize());
+			if ( AActor* OwnerActor = ActorComponent->GetOwner() )
+			{
+				InstanceData->CachedVelocity = FVector3f(OwnerActor->GetVelocity());
+			}
+			else if (UPrimitiveComponent* PrimitiveComponent = Cast<UPrimitiveComponent>(ActorComponent))
+			{
+				InstanceData->CachedVelocity = FVector3f(PrimitiveComponent->GetPhysicsLinearVelocity());
+			}
 		}
 		else if (AActor* OwnerActor = ActorComponent->GetOwner())
 		{
 			InstanceData->bCachedValid = true;
 			InstanceData->CachedTransform = OwnerActor->GetTransform();
 			InstanceData->CachedTransform.AddToTranslation(FVector(SystemInstance->GetLWCTile()) * -FLargeWorldRenderScalar::GetTileSize());
+			InstanceData->CachedVelocity = FVector3f(OwnerActor->GetVelocity());
 		}
 	}
 
@@ -404,5 +427,20 @@ void UNiagaraDataInterfaceActorComponent::VMGetTransform(FVectorVMExternalFuncti
 	}
 }
 
+void UNiagaraDataInterfaceActorComponent::VMGetPhysicsVelocity(FVectorVMExternalFunctionContext& Context)
+{
+	using namespace NDIActorComponentLocal;
+
+	VectorVM::FUserPtrHandler<FInstanceData_GameThread> InstanceData(Context);
+	FNDIOutputParam<bool>		OutValid(Context);
+	FNDIOutputParam<FVector3f>	OutVelocity(Context);
+
+	for (int32 i = 0; i < Context.GetNumInstances(); ++i)
+	{
+		OutValid.SetAndAdvance(InstanceData->bCachedValid);
+		OutVelocity.SetAndAdvance(InstanceData->CachedVelocity);
+	}
+}
+
 #undef LOCTEXT_NAMESPACE
 
diff --git a/Engine/Plugins/FX/Niagara/Source/Niagara/Private/DataInterface/NiagaraDataInterfaceActorComponent.h b/Engine/Plugins/FX/Niagara/Source/Niagara/Private/DataInterface/NiagaraDataInterfaceActorComponent.h
--- a/Engine/Plugins/FX/Niagara/Source/Niagara/Private/DataInterface/NiagaraDataInterfaceActorComponent.h
+++ b/Engine/Plugins/FX/Niagara/Source/Niagara/Private/DataInterface/NiagaraDataInterfaceActorComponent.h
@@ -76,4 +76,5 @@ public:
 private:
 	void VMGetMatrix(FVectorVMExternalFunctionContext& Context);
 	void VMGetTransform(FVectorVMExternalFunctionContext& Context);
+	void VMGetPhysicsVelocity(FVectorVMExternalFunctionContext& Context);
 };
