- Add Source Mode this allows you to bind the attached component or a local player

diff --git a/Engine/Plugins/FX/Niagara/Source/Niagara/Private/DataInterface/NiagaraDataInterfaceActorComponent.cpp b/Engine/Plugins/FX/Niagara/Source/Niagara/Private/DataInterface/NiagaraDataInterfaceActorComponent.cpp
--- a/Engine/Plugins/FX/Niagara/Source/Niagara/Private/DataInterface/NiagaraDataInterfaceActorComponent.cpp
+++ b/Engine/Plugins/FX/Niagara/Source/Niagara/Private/DataInterface/NiagaraDataInterfaceActorComponent.cpp
@@ -5,9 +5,12 @@
 #include "NiagaraCustomVersion.h"
 #include "NiagaraShaderParametersBuilder.h"
 #include "NiagaraSystemInstance.h"
+#include "NiagaraWorldManager.h"
 
 #include "Components/SceneComponent.h"
+#include "Engine/LocalPlayer.h"
 #include "GameFramework/Actor.h"
+#include "GameFramework/PlayerController.h"
 #include "ShaderCompilerCore.h"
 #include "ShaderParameterUtils.h"
 #include "Misc/LargeWorldRenderPosition.h"
@@ -93,6 +96,36 @@ namespace NDIActorComponentLocal
 
 		TMap<FNiagaraSystemInstanceID, FInstanceData_RenderThread> SystemInstancesToInstanceData_RT;
 	};
+
+	USceneComponent* FindLocalPlayer(UWorld* World, int LocalPlayerIndex)
+	{
+		if ( World == nullptr )
+		{
+			return nullptr;
+		}
+
+		for (FConstPlayerControllerIterator Iterator=World->GetPlayerControllerIterator(); Iterator; ++Iterator)
+		{
+			APlayerController* PlayerController = Iterator->Get();
+			if ( PlayerController == nullptr || PlayerController->Player == nullptr )
+			{
+				continue;
+			}
+			
+			ULocalPlayer* LocalPlayer = Cast<ULocalPlayer>(PlayerController->Player);
+			if ( LocalPlayer == nullptr )
+			{
+				continue;
+			}
+
+			if (LocalPlayerIndex == LocalPlayer->GetLocalPlayerIndex())
+			{
+				AActor* PlayerPawn = PlayerController->GetPawn();
+				return PlayerPawn ? PlayerPawn->GetRootComponent() : nullptr;
+			}
+		}
+		return nullptr;
+	}
 }
 
 //////////////////////////////////////////////////////////////////////////
@@ -117,10 +150,25 @@ void UNiagaraDataInterfaceActorComponent::PostInitProperties()
 	}
 }
 
-UActorComponent* UNiagaraDataInterfaceActorComponent::ResolveComponent(const void* PerInstanceData) const
+UActorComponent* UNiagaraDataInterfaceActorComponent::ResolveComponent(FNiagaraSystemInstance* SystemInstance, const void* PerInstanceData) const
 {
 	using namespace NDIActorComponentLocal;
 
+	if (SourceMode == ENDIActorComponentSourceMode::AttachParent)
+	{
+		if (USceneComponent* AttachComponent = SystemInstance->GetAttachComponent())
+		{
+			return AttachComponent;
+		}
+	}
+	else if (SourceMode == ENDIActorComponentSourceMode::LocalPlayer)
+	{
+		if (USceneComponent* RootComponent = FindLocalPlayer(SystemInstance->GetWorldManager()->GetWorld(), LocalPlayerIndex))
+		{
+			return RootComponent;
+		}
+	}
+
 	FInstanceData_GameThread* InstanceData = (FInstanceData_GameThread*)PerInstanceData;
 	if (UObject* ObjectBinding = InstanceData->UserParamBinding.GetValue())
 	{
@@ -262,7 +310,7 @@ bool UNiagaraDataInterfaceActorComponent::InitPerInstanceData(void* PerInstanceD
 
 	FInstanceData_GameThread* InstanceData = new (PerInstanceData) FInstanceData_GameThread;
 	InstanceData->UserParamBinding.Init(SystemInstance->GetInstanceParameters(), ActorOrComponentParameter.Parameter);
-	InstanceData->CachedActorForCalcTickGroup = ResolveComponent(InstanceData);
+	InstanceData->CachedActorForCalcTickGroup = ResolveComponent(SystemInstance, InstanceData);
 
 	return true;
 }
@@ -303,7 +351,7 @@ bool UNiagaraDataInterfaceActorComponent::PerInstanceTick(void* PerInstanceData,
 	InstanceData->CachedTransform = FTransform::Identity;
 	InstanceData->CachedVelocity = FVector3f::ZeroVector;
 
-	UActorComponent* ActorComponent = ResolveComponent(PerInstanceData);
+	UActorComponent* ActorComponent = ResolveComponent(SystemInstance, PerInstanceData);
 	if (ActorComponent)
 	{
 		if (USceneComponent* SceneComponent = Cast<USceneComponent>(ActorComponent) )
@@ -376,6 +424,8 @@ bool UNiagaraDataInterfaceActorComponent::Equals(const UNiagaraDataInterface* Ot
 	const UNiagaraDataInterfaceActorComponent* OtherTyped = CastChecked<const UNiagaraDataInterfaceActorComponent>(Other);
 	return OtherTyped->SourceActor == SourceActor
 		&& OtherTyped->ActorOrComponentParameter == ActorOrComponentParameter
+		&& OtherTyped->SourceMode == SourceMode
+		&& OtherTyped->LocalPlayerIndex == LocalPlayerIndex
 		&& OtherTyped->bRequireCurrentFrameData == bRequireCurrentFrameData;
 }
 
@@ -387,6 +437,8 @@ bool UNiagaraDataInterfaceActorComponent::CopyToInternal(UNiagaraDataInterface*
 	}
 
 	UNiagaraDataInterfaceActorComponent* OtherTyped = CastChecked<UNiagaraDataInterfaceActorComponent>(Destination);
+	OtherTyped->SourceMode = SourceMode;
+	OtherTyped->LocalPlayerIndex = LocalPlayerIndex;
 	OtherTyped->SourceActor = SourceActor;
 	OtherTyped->ActorOrComponentParameter = ActorOrComponentParameter;
 	OtherTyped->bRequireCurrentFrameData = bRequireCurrentFrameData;
diff --git a/Engine/Plugins/FX/Niagara/Source/Niagara/Private/DataInterface/NiagaraDataInterfaceActorComponent.h b/Engine/Plugins/FX/Niagara/Source/Niagara/Private/DataInterface/NiagaraDataInterfaceActorComponent.h
--- a/Engine/Plugins/FX/Niagara/Source/Niagara/Private/DataInterface/NiagaraDataInterfaceActorComponent.h
+++ b/Engine/Plugins/FX/Niagara/Source/Niagara/Private/DataInterface/NiagaraDataInterfaceActorComponent.h
@@ -7,6 +7,27 @@
 #include "NiagaraParameterStore.h"
 #include "NiagaraDataInterfaceActorComponent.generated.h"
 
+UENUM()
+enum class ENDIActorComponentSourceMode : uint8
+{
+	/**
+	The default binding mode first we look for a valid binding on the ActorOrComponentParameter.
+	If this it no valid we then look at the SourceActor.
+	If these both fail we are bound to nothing.
+	*/
+	Default,
+	/**
+	We will first look at the attach parent.
+	If this is not valid we fallback to the Default binding mode.
+	*/
+	AttachParent,
+	/**
+	We will look for a ULocalPlayer with the provided index.
+	If this is not valid we fallback to the Default binding mode.
+	*/
+	LocalPlayer,
+};
+
 /**
 Data interface that gives you access to actor & component information.
 */
@@ -29,6 +50,13 @@ public:
 	UPROPERTY(EditAnywhere, Category = "Performance")
 	bool bRequireCurrentFrameData = true;
 
+	/** Controls how we find the actor / component we want to bind to. */
+	UPROPERTY(EditAnywhere, Category = "ActorComponent")
+	ENDIActorComponentSourceMode SourceMode = ENDIActorComponentSourceMode::Default;
+
+	UPROPERTY(EditAnywhere, Category = "ActorComponent", meta=(EditCondition="SourceMode==ENDIActorComponentSourceMode::LocalPlayer"))//, EditConditionHides))
+	int32 LocalPlayerIndex = 0;
+
 	/** Optional source actor to use, if the user parameter binding is valid this will be ignored. */
 	UPROPERTY(EditAnywhere, Category = "ActorComponent")
 	TLazyObjectPtr<AActor> SourceActor;
@@ -41,7 +69,7 @@ public:
 	virtual void PostInitProperties() override;
 	//UObject Interface End
 
-	class UActorComponent* ResolveComponent(const void* PerInstanceData) const;
+	class UActorComponent* ResolveComponent(FNiagaraSystemInstance* SystemInstance, const void* PerInstanceData) const;
 
 	//UNiagaraDataInterface Interface
 	virtual bool CanExecuteOnTarget(ENiagaraSimTarget Target) const override { return true; }
