Removed the default "New" prefix for new parameters in Niagara

diff --git a/Engine/Plugins/FX/Niagara/Source/NiagaraEditor/Private/ViewModels/NiagaraParameterCollectionAssetViewModel.cpp b/Engine/Plugins/FX/Niagara/Source/NiagaraEditor/Private/ViewModels/NiagaraParameterCollectionAssetViewModel.cpp
--- a/Engine/Plugins/FX/Niagara/Source/NiagaraEditor/Private/ViewModels/NiagaraParameterCollectionAssetViewModel.cpp
+++ b/Engine/Plugins/FX/Niagara/Source/NiagaraEditor/Private/ViewModels/NiagaraParameterCollectionAssetViewModel.cpp
@@ -102,7 +102,7 @@ FName FNiagaraParameterCollectionAssetViewModel::GenerateNewName(FNiagaraTypeDef
 {
 	check(Collection && Instance);
 
-	FName ProposedName = *(TEXT("New") + Type.GetName());
+	FName ProposedName = FName(Type.GetName());
 	TSet<FName> Existing;
 	Existing.Reserve(ParameterViewModels.Num());
 	for (TSharedRef<INiagaraParameterViewModel> ParameterViewModel : ParameterViewModels)
diff --git a/Engine/Plugins/FX/Niagara/Source/NiagaraEditor/Public/Widgets/SNiagaraParameterMenu.cpp b/Engine/Plugins/FX/Niagara/Source/NiagaraEditor/Public/Widgets/SNiagaraParameterMenu.cpp
--- a/Engine/Plugins/FX/Niagara/Source/NiagaraEditor/Public/Widgets/SNiagaraParameterMenu.cpp
+++ b/Engine/Plugins/FX/Niagara/Source/NiagaraEditor/Public/Widgets/SNiagaraParameterMenu.cpp
@@ -741,7 +741,7 @@ void SNiagaraAddParameterFromPanelMenu::NewParameterSelected(FNiagaraTypeDefinit
 
 	const FString NewVariableDefaultName = TypeDisplayName.IsEmpty()
 		? TEXT("New Variable")
-		: TEXT("New ") + TypeDisplayName;
+		: TypeDisplayName;
 
 	TArray<FString> NameParts;
 
