new function: AddParticle
cleaned up fill neighborgrid module

diff --git a/Engine/Plugins/FX/Niagara/Source/Niagara/Private/NiagaraDataInterfaceNeighborGrid3D.cpp b/Engine/Plugins/FX/Niagara/Source/Niagara/Private/NiagaraDataInterfaceNeighborGrid3D.cpp
--- a/Engine/Plugins/FX/Niagara/Source/Niagara/Private/NiagaraDataInterfaceNeighborGrid3D.cpp
+++ b/Engine/Plugins/FX/Niagara/Source/Niagara/Private/NiagaraDataInterfaceNeighborGrid3D.cpp
@@ -28,6 +28,7 @@ static const FName GetParticleNeighborFunctionName("GetParticleNeighbor");
 static const FName SetParticleNeighborFunctionName("SetParticleNeighbor");
 static const FName GetParticleNeighborCountFunctionName("GetParticleNeighborCount");
 static const FName SetParticleNeighborCountFunctionName("SetParticleNeighborCount");
+static const FName AddParticleFunctionName("AddParticle");
 
 static int32 GMaxNiagaraNeighborGridCells = (100*100*100);
 static FAutoConsoleVariableRef CVarMaxNiagaraNeighborGridCells(
@@ -175,6 +176,26 @@ void UNiagaraDataInterfaceNeighborGrid3D::GetFunctions(TArray<FNiagaraFunctionSi
 		Sig.bRequiresContext = false;
 		OutFunctions.Add(Sig);
 	}
+
+	{
+		FNiagaraFunctionSignature Sig;
+		Sig.Name = AddParticleFunctionName;
+		Sig.Inputs.Add(FNiagaraVariable(FNiagaraTypeDefinition(GetClass()), TEXT("Grid")));
+		Sig.Inputs.Add(FNiagaraVariable(FNiagaraTypeDefinition::GetIntDef(), TEXT("IndexX")));
+		Sig.Inputs.Add(FNiagaraVariable(FNiagaraTypeDefinition::GetIntDef(), TEXT("IndexY")));
+		Sig.Inputs.Add(FNiagaraVariable(FNiagaraTypeDefinition::GetIntDef(), TEXT("IndexZ")));		
+		Sig.Inputs.Add(FNiagaraVariable(FNiagaraTypeDefinition::GetIntDef(), TEXT("ParticleIndex")));
+		Sig.Outputs.Add(FNiagaraVariable(FNiagaraTypeDefinition::GetBoolDef(), TEXT("Success")));
+
+		Sig.bExperimental = true;
+		Sig.bWriteFunction = true;
+		Sig.bMemberFunction = true;
+		Sig.bRequiresContext = false;
+		Sig.bRequiresExecPin = true;
+		OutFunctions.Add(Sig);
+	}
+
+
 }
 
 DEFINE_NDI_DIRECT_FUNC_BINDER(UNiagaraDataInterfaceNeighborGrid3D, SetNumCells);
@@ -307,6 +328,7 @@ bool UNiagaraDataInterfaceNeighborGrid3D::GetFunctionHLSL(const FNiagaraDataInte
 	TMap<FString, FStringFormatArg> FormatArgs =
 	{
 		{TEXT("FunctionName"),					FunctionInfo.InstanceName},
+		{TEXT("DIName"),						ParamInfo.DataInterfaceHLSLSymbol},
 		{TEXT("NumCellsName"),					ParamInfo.DataInterfaceHLSLSymbol + UNiagaraDataInterfaceRWBase::NumCellsName},
 		{TEXT("UnitToUVName"),					ParamInfo.DataInterfaceHLSLSymbol + UNiagaraDataInterfaceRWBase::UnitToUVName},
 		{TEXT("MaxNeighborsPerCellName"),		ParamInfo.DataInterfaceHLSLSymbol + MaxNeighborsPerCellName},
@@ -395,6 +417,39 @@ bool UNiagaraDataInterfaceNeighborGrid3D::GetFunctionHLSL(const FNiagaraDataInte
 		OutHLSL += FString::Format(FormatBounds, FormatArgs);
 		return true;
 	}
+	else if (FunctionInfo.DefinitionName == AddParticleFunctionName)
+	{
+		static const TCHAR* FormatBounds = TEXT(R"(
+			void {FunctionName}(int IndexX, int IndexY, int IndexZ, int ParticleIndex, out bool Success)
+			{		
+				Success = false;
+		
+				if (IndexX >= 0 && IndexX < {NumCellsName}.x && 
+					IndexY >= 0 && IndexY < {NumCellsName}.y && 
+					IndexZ >= 0 && IndexZ < {NumCellsName}.z)
+				{
+					int LinearIndex = IndexX + IndexY * {NumCellsName}.x + IndexZ * {NumCellsName}.x * {NumCellsName}.y;
+
+					int PreviousNeighborCount;
+					InterlockedAdd({OutputParticleNeighborCount}[LinearIndex], 1, PreviousNeighborCount);					
+
+					if (PreviousNeighborCount < {MaxNeighborsPerCellName})
+					{
+						Success = true;
+
+						int NeighborGridLinear = 
+							PreviousNeighborCount + IndexX * {MaxNeighborsPerCellName} + 
+							IndexY * {MaxNeighborsPerCellName}*{NumCellsName}.x + 
+							IndexZ * {MaxNeighborsPerCellName}*{NumCellsName}.x*{NumCellsName}.y;
+						
+						{OutputParticleNeighbors}[NeighborGridLinear] = ParticleIndex;					
+					}		
+				}
+			}
+		)");
+		OutHLSL += FString::Format(FormatBounds, FormatArgs);
+		return true;
+	}
 
 	return false;
 }
