- Remove Ex factories from Niagara sprite and mesh to reduce memory & shader permutations

diff --git a/Engine/Plugins/FX/Niagara/Shaders/Private/NiagaraMeshVertexFactory.ush b/Engine/Plugins/FX/Niagara/Shaders/Private/NiagaraMeshVertexFactory.ush
--- a/Engine/Plugins/FX/Niagara/Shaders/Private/NiagaraMeshVertexFactory.ush
+++ b/Engine/Plugins/FX/Niagara/Shaders/Private/NiagaraMeshVertexFactory.ush
@@ -22,11 +22,6 @@
 #define USE_PARTICLE_LOCAL_TO_WORLD (NEEDS_PARTICLE_LOCAL_TO_WORLD || NEEDS_INSTANCE_LOCAL_TO_WORLD_PS)
 #define USE_PARTICLE_WORLD_TO_LOCAL (NEEDS_PARTICLE_WORLD_TO_LOCAL || NEEDS_INSTANCE_WORLD_TO_LOCAL_PS)
 
-
-#ifndef ENABLE_PRECISE_MOTION_VECTORS
-#define ENABLE_PRECISE_MOTION_VECTORS 0
-#endif
-
 #ifndef MANUAL_VERTEX_FETCH
 #define MANUAL_VERTEX_FETCH 0
 #endif
@@ -518,11 +513,7 @@ FVertexFactoryIntermediates GetVertexFactoryIntermediates(FVertexFactoryInput In
 	Params.ParticleIndex				= Intermediates.ParticleIndex;
 	Params.SystemLWCTile				= NiagaraMeshVF.SystemLWCTile;
 	Params.bLocalSpace					= NiagaraMeshVF.bLocalSpace;
-	#if ENABLE_PRECISE_MOTION_VECTORS
-		Params.bPreciseMotionVectors 	= true;
-	#else
-		Params.bPreciseMotionVectors 	= false;		
-	#endif
+	Params.bPreciseMotionVectors 		= NiagaraMeshVF.AccurateMotionVectors != 0;
 	Params.FacingMode					= NiagaraMeshVF.FacingMode;
 	Params.DeltaSeconds					= NiagaraMeshVF.DeltaSeconds;
 	Params.MeshScale 					= NiagaraMeshVF.MeshScale;
diff --git a/Engine/Plugins/FX/Niagara/Shaders/Private/NiagaraSpriteVertexFactory.ush b/Engine/Plugins/FX/Niagara/Shaders/Private/NiagaraSpriteVertexFactory.ush
--- a/Engine/Plugins/FX/Niagara/Shaders/Private/NiagaraSpriteVertexFactory.ush
+++ b/Engine/Plugins/FX/Niagara/Shaders/Private/NiagaraSpriteVertexFactory.ush
@@ -35,10 +35,6 @@
 /** Blends between FaceCamera and FaceCameraPosition.*/
 #define FACING_MODE_FACECAMERADISTANCEBLEND 4
 
-#ifndef ENABLE_PRECISE_MOTION_VECTORS
-#define ENABLE_PRECISE_MOTION_VECTORS 0
-#endif
-
 struct FVertexFactoryInput
 {
 	float2	TexCoord				: ATTRIBUTE0;
@@ -883,8 +879,8 @@ FVertexFactoryIntermediates GetVertexFactoryIntermediates(FVertexFactoryInput In
 
 	Intermediates.VertexWorldPosition = LWCAdd(ParticlePosition, VertexOffset);
 		
-	// Prev vertex position (approximated from current velocity unless ENABLE_PRECISE_MOTION_VECTORS)
-#if ENABLE_PRECISE_MOTION_VECTORS
+	// Prev vertex position (approximated from current velocity unless enabled)
+	if ( NiagaraSpriteVF.AccurateMotionVectors != 0 )
 	{
 		const FLWCVector3 PrevParticlePosition = GetNiagaraPreviousParticlePosition(ParticleID);
 		const float PrevParticleRotation = GetNiagaraPreviousParticleRotation(ParticleID);
@@ -910,9 +906,10 @@ FVertexFactoryIntermediates GetVertexFactoryIntermediates(FVertexFactoryInput In
 		
 		Intermediates.PreviousVertexWorldPosition = LWCAdd(PrevParticlePosition, PrevVertexOffset);
 	}
-#else
-	Intermediates.PreviousVertexWorldPosition = LWCSubtract(Intermediates.VertexWorldPosition, LWCPromote(ParticleVelocity * NiagaraSpriteVF.DeltaSeconds));
-#endif
+	else
+	{
+		Intermediates.PreviousVertexWorldPosition = LWCSubtract(Intermediates.VertexWorldPosition, LWCPromote(ParticleVelocity * NiagaraSpriteVF.DeltaSeconds));
+	}
 
 #if USE_PARTICLE_POSITION
 	// Do we really want the camera offset here ? What about when GENERATE_SPHERICAL_PARTICLE_NORMALS
diff --git a/Engine/Plugins/FX/Niagara/Source/Niagara/Private/NiagaraMeshRendererProperties.cpp b/Engine/Plugins/FX/Niagara/Source/Niagara/Private/NiagaraMeshRendererProperties.cpp
--- a/Engine/Plugins/FX/Niagara/Source/Niagara/Private/NiagaraMeshRendererProperties.cpp
+++ b/Engine/Plugins/FX/Niagara/Source/Niagara/Private/NiagaraMeshRendererProperties.cpp
@@ -458,7 +458,7 @@ void UNiagaraMeshRendererProperties::GetUsedMeshMaterials(int32 MeshIndex, const
 
 const FVertexFactoryType* UNiagaraMeshRendererProperties::GetVertexFactoryType() const
 {
-	return NeedsPreciseMotionVectors() ? &FNiagaraMeshVertexFactoryEx::StaticType : &FNiagaraMeshVertexFactory::StaticType;
+	return &FNiagaraMeshVertexFactory::StaticType;
 }
 
 void UNiagaraMeshRendererProperties::GetUsedMaterials(const FNiagaraEmitterInstance* InEmitter, TArray<UMaterialInterface*>& OutMaterials) const
diff --git a/Engine/Plugins/FX/Niagara/Source/Niagara/Private/NiagaraRendererMeshes.cpp b/Engine/Plugins/FX/Niagara/Source/Niagara/Private/NiagaraRendererMeshes.cpp
--- a/Engine/Plugins/FX/Niagara/Source/Niagara/Private/NiagaraRendererMeshes.cpp
+++ b/Engine/Plugins/FX/Niagara/Source/Niagara/Private/NiagaraRendererMeshes.cpp
@@ -687,6 +687,7 @@ FNiagaraMeshCommonParameters FNiagaraRendererMeshes::CreateCommonShaderParams(co
 
 	Params.SystemLWCTile			= SceneProxy.GetLWCRenderTile();
 	Params.bLocalSpace 				= bUseLocalSpace;
+	Params.AccurateMotionVectors	= bAccurateMotionVectors;
 	Params.DeltaSeconds 			= View.Family->Time.GetDeltaWorldTimeSeconds();
 	Params.FacingMode 				= (uint32)FacingMode;	
 
@@ -1262,19 +1263,11 @@ void FNiagaraRendererMeshes::GetDynamicMeshElements(const TArray<const FSceneVie
 				const FStaticMeshLODResources& LODModel = MeshData.RenderData->LODResources[LODIndex];
 				const int32 SectionCount = LODModel.Sections.Num();
 
-				FMeshCollectorResourcesBase* CollectorResources;
-				if (bAccurateMotionVectors)
-				{
-					CollectorResources = &Collector.AllocateOneFrameResource<FMeshCollectorResourcesEx>();
-				}
-				else
-				{
-					CollectorResources = &Collector.AllocateOneFrameResource<FMeshCollectorResources>();
-				}				
+				FMeshCollectorResources* CollectorResources = &Collector.AllocateOneFrameResource<FMeshCollectorResources>();
 
 				// Get the next vertex factory to use
 				// TODO: Find a way to safely pool these such that they won't be concurrently accessed by multiple views
-				FNiagaraMeshVertexFactory& VertexFactory = CollectorResources->GetVertexFactory();
+				FNiagaraMeshVertexFactory& VertexFactory = CollectorResources->VertexFactory;
 				VertexFactory.SetParticleFactoryType(NVFT_Mesh);
 				VertexFactory.SetMeshIndex(MeshIndex);
 				VertexFactory.SetLODIndex(LODIndex);
@@ -1420,19 +1413,11 @@ void FNiagaraRendererMeshes::GetDynamicRayTracingInstances(FRayTracingMaterialGa
 		FRayTracingInstance RayTracingInstance;
 		RayTracingInstance.Geometry = &Geometry;
 
-		FMeshCollectorResourcesBase* CollectorResources;
-		if (bAccurateMotionVectors)
-		{
-			CollectorResources = &Context.RayTracingMeshResourceCollector.AllocateOneFrameResource<FMeshCollectorResourcesEx>();
-		}
-		else
-		{
-			CollectorResources = &Context.RayTracingMeshResourceCollector.AllocateOneFrameResource<FMeshCollectorResources>();
-		}
+		FMeshCollectorResources* CollectorResources = &Context.RayTracingMeshResourceCollector.AllocateOneFrameResource<FMeshCollectorResources>();
 
 		// Get the next vertex factory to use
 		// TODO: Find a way to safely pool these such that they won't be concurrently accessed by multiple views
-		FNiagaraMeshVertexFactory& VertexFactory = CollectorResources->GetVertexFactory();
+		FNiagaraMeshVertexFactory& VertexFactory = CollectorResources->VertexFactory;
 		VertexFactory.SetParticleFactoryType(NVFT_Mesh);
 		VertexFactory.SetMeshIndex(MeshIndex);
 		VertexFactory.SetLODIndex(LODIndex);
diff --git a/Engine/Plugins/FX/Niagara/Source/Niagara/Private/NiagaraRendererSprites.cpp b/Engine/Plugins/FX/Niagara/Source/Niagara/Private/NiagaraRendererSprites.cpp
--- a/Engine/Plugins/FX/Niagara/Source/Niagara/Private/NiagaraRendererSprites.cpp
+++ b/Engine/Plugins/FX/Niagara/Source/Niagara/Private/NiagaraRendererSprites.cpp
@@ -598,6 +598,7 @@ FNiagaraSpriteUniformBufferRef FNiagaraRendererSprites::CreateViewUniformBuffer(
 		}
 	}
 
+	PerViewUniformParameters.AccurateMotionVectors = false;
 	if (SourceMode == ENiagaraRendererSourceDataMode::Particles)
 	{
 		TConstArrayView<FNiagaraRendererVariableInfo> VFVariables = ParticleSpriteRenderData.RendererLayout->GetVFVariables_RenderThread();
@@ -620,6 +621,7 @@ FNiagaraSpriteUniformBufferRef FNiagaraRendererSprites::CreateViewUniformBuffer(
 		PerViewUniformParameters.MaterialRandomDataOffset = VFVariables[ENiagaraSpriteVFLayout::MaterialRandom].GetGPUOffset();
 		if (bAccurateMotionVectors)
 		{
+			PerViewUniformParameters.AccurateMotionVectors = true;
 			PerViewUniformParameters.PrevPositionDataOffset = VFVariables[ENiagaraSpriteVFLayout::PrevPosition].GetGPUOffset();
 			PerViewUniformParameters.PrevVelocityDataOffset = VFVariables[ENiagaraSpriteVFLayout::PrevVelocity].GetGPUOffset();
 			PerViewUniformParameters.PrevRotationDataOffset = VFVariables[ENiagaraSpriteVFLayout::PrevRotation].GetGPUOffset();
@@ -974,19 +976,11 @@ void FNiagaraRendererSprites::GetDynamicMeshElements(const TArray<const FSceneVi
 				InitializeSortInfo(ParticleSpriteRenderData, *SceneProxy, *View, ViewIndex, SortInfo);
 			}
 
-			FMeshCollectorResourcesBase* CollectorResources;
-			if ( bAccurateMotionVectors )
-			{
-				CollectorResources = &Collector.AllocateOneFrameResource<FMeshCollectorResourcesEx>();
-			}
-			else
-			{
-				CollectorResources = &Collector.AllocateOneFrameResource<FMeshCollectorResources>();
-			}
+			FMeshCollectorResources* CollectorResources = &Collector.AllocateOneFrameResource<FMeshCollectorResources>();
 
 			// Get the next vertex factory to use
 			// TODO: Find a way to safely pool these such that they won't be concurrently accessed by multiple views
-			FNiagaraSpriteVertexFactory& VertexFactory = CollectorResources->GetVertexFactory();
+			FNiagaraSpriteVertexFactory& VertexFactory = CollectorResources->VertexFactory;
 
 			// Sort/Cull particles if needed.
 			uint32 NumInstances = SourceMode == ENiagaraRendererSourceDataMode::Particles ? ParticleSpriteRenderData.SourceParticleData->GetNumInstances() : 1;
@@ -1061,22 +1055,13 @@ void FNiagaraRendererSprites::GetDynamicRayTracingInstances(FRayTracingMaterialG
 		InitializeSortInfo(ParticleSpriteRenderData, *SceneProxy, *Context.ReferenceView, 0, SortInfo);
 	}
 
-	FMeshCollectorResourcesBase* CollectorResources;
-	if (bAccurateMotionVectors)
-	{
-		CollectorResources = &Context.RayTracingMeshResourceCollector.AllocateOneFrameResource<FMeshCollectorResourcesEx>();
-	}
-	else
-	{
-		CollectorResources = &Context.RayTracingMeshResourceCollector.AllocateOneFrameResource<FMeshCollectorResources>();
-	}
-
-	if (!CollectorResources->GetVertexFactory().GetType()->SupportsRayTracingDynamicGeometry())
+	if (FNiagaraSpriteVertexFactory::StaticType.SupportsRayTracingDynamicGeometry())
 	{
 		return;
 	}
 
-	FNiagaraSpriteVertexFactory& VertexFactory = CollectorResources->GetVertexFactory();
+	FMeshCollectorResources* CollectorResources = &Context.RayTracingMeshResourceCollector.AllocateOneFrameResource<FMeshCollectorResources>();
+	FNiagaraSpriteVertexFactory& VertexFactory = CollectorResources->VertexFactory;
 
 	// Sort/Cull particles if needed.
 	uint32 NumInstances = SourceMode == ENiagaraRendererSourceDataMode::Particles ? ParticleSpriteRenderData.SourceParticleData->GetNumInstances() : 1;
diff --git a/Engine/Plugins/FX/Niagara/Source/Niagara/Private/NiagaraSpriteRendererProperties.cpp b/Engine/Plugins/FX/Niagara/Source/Niagara/Private/NiagaraSpriteRendererProperties.cpp
--- a/Engine/Plugins/FX/Niagara/Source/Niagara/Private/NiagaraSpriteRendererProperties.cpp
+++ b/Engine/Plugins/FX/Niagara/Source/Niagara/Private/NiagaraSpriteRendererProperties.cpp
@@ -128,7 +128,7 @@ void UNiagaraSpriteRendererProperties::GetUsedMaterials(const FNiagaraEmitterIns
 
 const FVertexFactoryType* UNiagaraSpriteRendererProperties::GetVertexFactoryType() const
 {
-	return NeedsPreciseMotionVectors() ? &FNiagaraSpriteVertexFactoryEx::StaticType : &FNiagaraSpriteVertexFactory::StaticType;
+	return &FNiagaraSpriteVertexFactory::StaticType;
 }
 
 void UNiagaraSpriteRendererProperties::PostLoad()
diff --git a/Engine/Plugins/FX/Niagara/Source/Niagara/Public/NiagaraRendererMeshes.h b/Engine/Plugins/FX/Niagara/Source/Niagara/Public/NiagaraRendererMeshes.h
--- a/Engine/Plugins/FX/Niagara/Source/Niagara/Public/NiagaraRendererMeshes.h
+++ b/Engine/Plugins/FX/Niagara/Source/Niagara/Public/NiagaraRendererMeshes.h
@@ -90,28 +90,15 @@ protected:
 		TArray<uint32, TInlineAllocator<4>> MaterialRemapTable;
 	};
 
-	class FMeshCollectorResourcesBase : public FOneFrameResource
+	class FMeshCollectorResources : public FOneFrameResource
 	{
 	public:
-		FNiagaraMeshUniformBufferRef UniformBuffer;
-
-		virtual ~FMeshCollectorResourcesBase() {}
-		virtual FNiagaraMeshVertexFactory& GetVertexFactory() = 0;
-	};
-
-	template <typename TVertexFactory>
-	class TMeshCollectorResources : public FMeshCollectorResourcesBase
-	{
-	public:
-		TVertexFactory VertexFactory;
+		~FMeshCollectorResources() override { VertexFactory.ReleaseResource(); }
 
-		virtual ~TMeshCollectorResources() { VertexFactory.ReleaseResource(); }
-		virtual FNiagaraMeshVertexFactory& GetVertexFactory() override { return VertexFactory; }
+		FNiagaraMeshVertexFactory VertexFactory;
+		FNiagaraMeshUniformBufferRef UniformBuffer;
 	};
 
-	using FMeshCollectorResources = TMeshCollectorResources<FNiagaraMeshVertexFactory>;
-	using FMeshCollectorResourcesEx = TMeshCollectorResources<FNiagaraMeshVertexFactoryEx>;
-
 	struct FEmitterSourceInstanceData
 	{
 		FPrimitiveInstance InstanceSceneData;
diff --git a/Engine/Plugins/FX/Niagara/Source/Niagara/Public/NiagaraRendererSprites.h b/Engine/Plugins/FX/Niagara/Source/Niagara/Public/NiagaraRendererSprites.h
--- a/Engine/Plugins/FX/Niagara/Source/Niagara/Public/NiagaraRendererSprites.h
+++ b/Engine/Plugins/FX/Niagara/Source/Niagara/Public/NiagaraRendererSprites.h
@@ -62,28 +62,15 @@ private:
 	};
 
 	/* Mesh collector classes */
-	class FMeshCollectorResourcesBase : public FOneFrameResource
+	class FMeshCollectorResources : public FOneFrameResource
 	{
 	public:
-		FNiagaraSpriteUniformBufferRef UniformBuffer;
+		~FMeshCollectorResources() override { VertexFactory.ReleaseResource(); }
 
-		virtual ~FMeshCollectorResourcesBase() {}
-		virtual FNiagaraSpriteVertexFactory& GetVertexFactory() = 0;
+		FNiagaraSpriteVertexFactory		VertexFactory;
+		FNiagaraSpriteUniformBufferRef	UniformBuffer;
 	};
 
-	template <typename TVertexFactory>
-	class TMeshCollectorResources : public FMeshCollectorResourcesBase
-	{
-	public:
-		TVertexFactory VertexFactory;
-
-		virtual ~TMeshCollectorResources() { VertexFactory.ReleaseResource(); }
-		virtual FNiagaraSpriteVertexFactory& GetVertexFactory() override { return VertexFactory; }
-	};
-
-	using FMeshCollectorResources = TMeshCollectorResources<FNiagaraSpriteVertexFactory>;
-	using FMeshCollectorResourcesEx = TMeshCollectorResources<FNiagaraSpriteVertexFactoryEx>;
-
 	void PrepareParticleSpriteRenderData(FParticleSpriteRenderData& ParticleSpriteRenderData, FNiagaraDynamicDataBase* InDynamicData, const FNiagaraSceneProxy* SceneProxy) const;
 	void PrepareParticleRenderBuffers(FParticleSpriteRenderData& ParticleSpriteRenderData, FGlobalDynamicReadBuffer& DynamicReadBuffer) const;
 	void InitializeSortInfo(FParticleSpriteRenderData& ParticleSpriteRenderData, const FNiagaraSceneProxy& SceneProxy, const FSceneView& View, int32 ViewIndex, FNiagaraGPUSortInfo& OutSortInfo) const;
diff --git a/Engine/Plugins/FX/Niagara/Source/NiagaraVertexFactories/Private/NiagaraMeshVertexFactory.cpp b/Engine/Plugins/FX/Niagara/Source/NiagaraVertexFactories/Private/NiagaraMeshVertexFactory.cpp
--- a/Engine/Plugins/FX/Niagara/Source/NiagaraVertexFactories/Private/NiagaraMeshVertexFactory.cpp
+++ b/Engine/Plugins/FX/Niagara/Source/NiagaraVertexFactories/Private/NiagaraMeshVertexFactory.cpp
@@ -228,14 +228,3 @@ IMPLEMENT_VERTEX_FACTORY_PARAMETER_TYPE(FNiagaraMeshVertexFactory, SF_RayHitGrou
 IMPLEMENT_VERTEX_FACTORY_PARAMETER_TYPE(FNiagaraMeshVertexFactory, SF_Pixel, FNiagaraMeshVertexFactoryShaderParametersPS);
 
 IMPLEMENT_VERTEX_FACTORY_TYPE(FNiagaraMeshVertexFactory, "/Plugin/FX/Niagara/Private/NiagaraMeshVertexFactory.ush",	NIAGARA_MESH_VF_FLAGS);
-
-/////////////////////////////////////////////////////////////////////////////////////////////////////
-
-IMPLEMENT_VERTEX_FACTORY_PARAMETER_TYPE(FNiagaraMeshVertexFactoryEx, SF_Vertex, FNiagaraMeshVertexFactoryShaderParametersVS);
-#if RHI_RAYTRACING
-IMPLEMENT_VERTEX_FACTORY_PARAMETER_TYPE(FNiagaraMeshVertexFactoryEx, SF_Compute, FNiagaraMeshVertexFactoryShaderParametersVS);
-IMPLEMENT_VERTEX_FACTORY_PARAMETER_TYPE(FNiagaraMeshVertexFactoryEx, SF_RayHitGroup, FNiagaraMeshVertexFactoryShaderParametersVS);
-#endif
-IMPLEMENT_VERTEX_FACTORY_PARAMETER_TYPE(FNiagaraMeshVertexFactoryEx, SF_Pixel, FNiagaraMeshVertexFactoryShaderParametersPS);
-
-IMPLEMENT_VERTEX_FACTORY_TYPE(FNiagaraMeshVertexFactoryEx, "/Plugin/FX/Niagara/Private/NiagaraMeshVertexFactory.ush", NIAGARA_MESH_VF_FLAGS_EX);
diff --git a/Engine/Plugins/FX/Niagara/Source/NiagaraVertexFactories/Private/NiagaraSpriteVertexFactory.cpp b/Engine/Plugins/FX/Niagara/Source/NiagaraVertexFactories/Private/NiagaraSpriteVertexFactory.cpp
--- a/Engine/Plugins/FX/Niagara/Source/NiagaraVertexFactories/Private/NiagaraSpriteVertexFactory.cpp
+++ b/Engine/Plugins/FX/Niagara/Source/NiagaraVertexFactories/Private/NiagaraSpriteVertexFactory.cpp
@@ -224,20 +224,3 @@ IMPLEMENT_VERTEX_FACTORY_TYPE(FNiagaraSpriteVertexFactory,"/Plugin/FX/Niagara/Pr
 	| EVertexFactoryFlags::SupportsRayTracingDynamicGeometry
 	| EVertexFactoryFlags::SupportsPSOPrecaching
 );
-
-/////////////////////////////////////////////////////////////////////////////
-
-IMPLEMENT_VERTEX_FACTORY_PARAMETER_TYPE(FNiagaraSpriteVertexFactoryEx, SF_Vertex, FNiagaraSpriteVertexFactoryShaderParametersVS);
-IMPLEMENT_VERTEX_FACTORY_PARAMETER_TYPE(FNiagaraSpriteVertexFactoryEx, SF_Pixel, FNiagaraSpriteVertexFactoryShaderParametersPS);
-#if RHI_RAYTRACING
-IMPLEMENT_VERTEX_FACTORY_PARAMETER_TYPE(FNiagaraSpriteVertexFactoryEx, SF_Compute, FNiagaraSpriteVertexFactoryShaderParametersVS);
-IMPLEMENT_VERTEX_FACTORY_PARAMETER_TYPE(FNiagaraSpriteVertexFactoryEx, SF_RayHitGroup, FNiagaraSpriteVertexFactoryShaderParametersVS);
-#endif // RHI_RAYTRACING
-IMPLEMENT_VERTEX_FACTORY_TYPE(FNiagaraSpriteVertexFactoryEx,"/Plugin/FX/Niagara/Private/NiagaraSpriteVertexFactory.ush",
-	  EVertexFactoryFlags::UsedWithMaterials
-	| EVertexFactoryFlags::SupportsDynamicLighting
-	| EVertexFactoryFlags::SupportsRayTracing
-	| EVertexFactoryFlags::SupportsRayTracingDynamicGeometry
-	| EVertexFactoryFlags::SupportsPrecisePrevWorldPos
-	| EVertexFactoryFlags::SupportsPSOPrecaching
-);
diff --git a/Engine/Plugins/FX/Niagara/Source/NiagaraVertexFactories/Public/NiagaraMeshVertexFactory.h b/Engine/Plugins/FX/Niagara/Source/NiagaraVertexFactories/Public/NiagaraMeshVertexFactory.h
--- a/Engine/Plugins/FX/Niagara/Source/NiagaraVertexFactories/Public/NiagaraMeshVertexFactory.h
+++ b/Engine/Plugins/FX/Niagara/Source/NiagaraVertexFactories/Public/NiagaraMeshVertexFactory.h
@@ -40,6 +40,7 @@ BEGIN_SHADER_PARAMETER_STRUCT(FNiagaraMeshCommonParameters, NIAGARAVERTEXFACTORI
 	
 	SHADER_PARAMETER(FVector3f, SystemLWCTile)
 	SHADER_PARAMETER(int, bLocalSpace)
+	SHADER_PARAMETER(int, AccurateMotionVectors)
 	SHADER_PARAMETER(float, DeltaSeconds)
 	SHADER_PARAMETER(uint32, FacingMode)
 
@@ -224,24 +225,3 @@ protected:
 	/** Uniform buffer with mesh particle parameters. */
 	FRHIUniformBuffer* MeshParticleUniformBuffer;
 };
-
-/**
-* Advanced mesh vertex factory. Used for enabling accurate motion vector output
-*/
-class NIAGARAVERTEXFACTORIES_API FNiagaraMeshVertexFactoryEx : public FNiagaraMeshVertexFactory
-{
-	DECLARE_VERTEX_FACTORY_TYPE(FNiagaraMeshVertexFactoryEx);
-public:
-	FNiagaraMeshVertexFactoryEx(ENiagaraVertexFactoryType InType, ERHIFeatureLevel::Type InFeatureLevel)
-		: FNiagaraMeshVertexFactory(InType, InFeatureLevel)
-	{}
-
-	FNiagaraMeshVertexFactoryEx() {}
-
-	static void ModifyCompilationEnvironment(const FVertexFactoryShaderPermutationParameters& Parameters, FShaderCompilerEnvironment& OutEnvironment)
-	{
-		FNiagaraMeshVertexFactory::ModifyCompilationEnvironment(Parameters, OutEnvironment);
-
-		OutEnvironment.SetDefine(TEXT("ENABLE_PRECISE_MOTION_VECTORS"), TEXT("1"));
-	}
-};
\ No newline at end of file
diff --git a/Engine/Plugins/FX/Niagara/Source/NiagaraVertexFactories/Public/NiagaraSpriteVertexFactory.h b/Engine/Plugins/FX/Niagara/Source/NiagaraVertexFactories/Public/NiagaraSpriteVertexFactory.h
--- a/Engine/Plugins/FX/Niagara/Source/NiagaraVertexFactories/Public/NiagaraSpriteVertexFactory.h
+++ b/Engine/Plugins/FX/Niagara/Source/NiagaraVertexFactories/Public/NiagaraSpriteVertexFactory.h
@@ -86,7 +86,8 @@ BEGIN_GLOBAL_SHADER_PARAMETER_STRUCT( FNiagaraSpriteUniformParameters, NIAGARAVE
 	SHADER_PARAMETER(FVector4f, DefaultDynamicMaterialParameter3)
 	SHADER_PARAMETER(int, PixelCoverageEnabled)
 	SHADER_PARAMETER(FVector4f, PixelCoverageColorBlend)
-END_GLOBAL_SHADER_PARAMETER_STRUCT()
+	SHADER_PARAMETER(int, AccurateMotionVectors)
+	END_GLOBAL_SHADER_PARAMETER_STRUCT()
 
 typedef TUniformBufferRef<FNiagaraSpriteUniformParameters> FNiagaraSpriteUniformBufferRef;
 
@@ -246,25 +247,3 @@ private:
 	FShaderResourceViewRHIRef SortedIndicesSRV;
 	uint32 SortedIndicesOffset;
 };
-
-class NIAGARAVERTEXFACTORIES_API FNiagaraSpriteVertexFactoryEx : public FNiagaraSpriteVertexFactory
-{
-	DECLARE_VERTEX_FACTORY_TYPE(FNiagaraSpriteVertexFactoryEx);
-
-public:
-
-	FNiagaraSpriteVertexFactoryEx(ENiagaraVertexFactoryType InType, ERHIFeatureLevel::Type InFeatureLevel )
-		: FNiagaraSpriteVertexFactory(InType, InFeatureLevel)
-	{}
-
-	FNiagaraSpriteVertexFactoryEx()
-		: FNiagaraSpriteVertexFactory(NVFT_MAX, ERHIFeatureLevel::Num)
-	{}
-
-	static void ModifyCompilationEnvironment(const FVertexFactoryShaderPermutationParameters& Parameters, FShaderCompilerEnvironment& OutEnvironment)
-	{
-		FNiagaraSpriteVertexFactory::ModifyCompilationEnvironment(Parameters, OutEnvironment);
-
-		OutEnvironment.SetDefine(TEXT("ENABLE_PRECISE_MOTION_VECTORS"), TEXT("1"));
-	}
-};
