Niagara - Optimize creation of stack merge adapters.
+ When getting the function input defaults for the merge adapter use multi find to get all of the default pins at once since this is faster than getting them one at a time.
+ Use a cache when determining if an enum entry at a specific index is visible based on the metadata.

diff --git a/Engine/Plugins/FX/Niagara/Source/NiagaraEditor/Private/NiagaraEditorUtilities.cpp b/Engine/Plugins/FX/Niagara/Source/NiagaraEditor/Private/NiagaraEditorUtilities.cpp
--- a/Engine/Plugins/FX/Niagara/Source/NiagaraEditor/Private/NiagaraEditorUtilities.cpp
+++ b/Engine/Plugins/FX/Niagara/Source/NiagaraEditor/Private/NiagaraEditorUtilities.cpp
@@ -4008,6 +4008,11 @@ void FNiagaraEditorUtilities::GetAllowedPayloadTypes(TArray<FNiagaraTypeDefiniti
 	}
 }
 
+bool FNiagaraEditorUtilities::IsEnumIndexVisible(const UEnum* Enum, int32 Index)
+{
+	return FNiagaraEnumIndexVisibilityCache::GetVisibility(Enum, Index);
+}
+
 void FNiagaraParameterUtilities::FilterToRelevantStaticVariables(const TArray<FNiagaraVariable>& InVars, TArray<FNiagaraVariable>& OutVars, FName InOldEmitterAlias, FName InNewEmitterAlias, bool bFilterByEmitterAliasAndConvertToUnaliased)
 {
 	FNiagaraAliasContext RenameContext(ENiagaraScriptUsage::ParticleSpawnScript);
@@ -4126,4 +4131,20 @@ void FNiagaraParameterDefinitionsUtilities::TrySubscribeScriptVarToDefinitionByN
 	}
 }
 
+TMap<FNiagaraEnumIndexVisibilityCache::FEnumIndexPair, bool> FNiagaraEnumIndexVisibilityCache::Cache;
+FCriticalSection FNiagaraEnumIndexVisibilityCache::CacheLock;
+
+bool FNiagaraEnumIndexVisibilityCache::GetVisibility(const UEnum* InEnum, int32 InIndex)
+{
+	FEnumIndexPair Key(InEnum, InIndex);
+	FScopeLock Lock(&CacheLock);
+	bool* bCachedIsVisible = Cache.Find(Key);
+	if (bCachedIsVisible == nullptr)
+	{
+		bCachedIsVisible = &Cache.Add(Key);
+		*bCachedIsVisible = InEnum->HasMetaData(TEXT("Hidden"), InIndex) == false && InEnum->HasMetaData(TEXT("Spacer"), InIndex) == false;
+	}
+	return *bCachedIsVisible;
+}
+
 #undef LOCTEXT_NAMESPACE
diff --git a/Engine/Plugins/FX/Niagara/Source/NiagaraEditor/Private/NiagaraNodeSelect.cpp b/Engine/Plugins/FX/Niagara/Source/NiagaraEditor/Private/NiagaraNodeSelect.cpp
--- a/Engine/Plugins/FX/Niagara/Source/NiagaraEditor/Private/NiagaraNodeSelect.cpp
+++ b/Engine/Plugins/FX/Niagara/Source/NiagaraEditor/Private/NiagaraNodeSelect.cpp
@@ -705,8 +705,7 @@ TArray<int32> UNiagaraNodeSelect::GetOptionValues() const
 			int32 ValidEnumEntryCount = 0;
 			for (int32 EnumIndex = 0; EnumIndex < Enum->NumEnums()-1; EnumIndex++)
 			{
-				const bool bShouldHideEnumEntry = ShouldHideEnumEntry(Enum, EnumIndex);
-				if(!bShouldHideEnumEntry)
+				if(FNiagaraEditorUtilities::IsEnumIndexVisible(Enum, EnumIndex))
 				{
 					ValidEnumEntryCount++;
 					SelectorValues.Add(Enum->GetValueByIndex(EnumIndex));
diff --git a/Engine/Plugins/FX/Niagara/Source/NiagaraEditor/Private/NiagaraNodeStaticSwitch.cpp b/Engine/Plugins/FX/Niagara/Source/NiagaraEditor/Private/NiagaraNodeStaticSwitch.cpp
--- a/Engine/Plugins/FX/Niagara/Source/NiagaraEditor/Private/NiagaraNodeStaticSwitch.cpp
+++ b/Engine/Plugins/FX/Niagara/Source/NiagaraEditor/Private/NiagaraNodeStaticSwitch.cpp
@@ -114,9 +114,7 @@ TArray<int32> UNiagaraNodeStaticSwitch::GetOptionValues() const
 
 		for (int32 EnumIndex = 0; EnumIndex < Enum->NumEnums() - 1; ++EnumIndex)
 		{
-			bool const bShouldBeHidden = ShouldHideEnumEntry(Enum, EnumIndex);
-
-			if (!bShouldBeHidden)
+			if (FNiagaraEditorUtilities::IsEnumIndexVisible(Enum, EnumIndex))
 			{
 				OptionValues.Add(Enum->GetValueByIndex(EnumIndex));
 			}
diff --git a/Engine/Plugins/FX/Niagara/Source/NiagaraEditor/Private/NiagaraNodeUsageSelector.cpp b/Engine/Plugins/FX/Niagara/Source/NiagaraEditor/Private/NiagaraNodeUsageSelector.cpp
--- a/Engine/Plugins/FX/Niagara/Source/NiagaraEditor/Private/NiagaraNodeUsageSelector.cpp
+++ b/Engine/Plugins/FX/Niagara/Source/NiagaraEditor/Private/NiagaraNodeUsageSelector.cpp
@@ -65,11 +65,6 @@ void UNiagaraNodeUsageSelector::AddOptionPin(const FNiagaraVariable& OutputVaria
 	NewPin->PinFriendlyName = GetOptionPinFriendlyName(OutputVariable);
 }
 
-bool UNiagaraNodeUsageSelector::ShouldHideEnumEntry(UEnum* Enum, int32 Index) const
-{
-	return Enum->HasMetaData(TEXT("Hidden"), Index) || Enum->HasMetaData(TEXT("Spacer"), Index);
-}
-
 TArray<int32> UNiagaraNodeUsageSelector::GetOptionValues() const
 {
 	TArray<int32> OptionValues;
diff --git a/Engine/Plugins/FX/Niagara/Source/NiagaraEditor/Private/NiagaraNodeUsageSelector.h b/Engine/Plugins/FX/Niagara/Source/NiagaraEditor/Private/NiagaraNodeUsageSelector.h
--- a/Engine/Plugins/FX/Niagara/Source/NiagaraEditor/Private/NiagaraNodeUsageSelector.h
+++ b/Engine/Plugins/FX/Niagara/Source/NiagaraEditor/Private/NiagaraNodeUsageSelector.h
@@ -49,7 +49,6 @@ public:
 
 	/** Helper function to create a variable to add to the OutputVars and FGuid to add to OutputVarGuids. */
 	FGuid AddOutput(FNiagaraTypeDefinition Type, const FName& Name);
-	bool ShouldHideEnumEntry(UEnum* Enum, int32 Index) const;
 
 	bool AreInputPinsOutdated() const;
 protected:
diff --git a/Engine/Plugins/FX/Niagara/Source/NiagaraEditor/Private/NiagaraScriptMergeManager.cpp b/Engine/Plugins/FX/Niagara/Source/NiagaraEditor/Private/NiagaraScriptMergeManager.cpp
--- a/Engine/Plugins/FX/Niagara/Source/NiagaraEditor/Private/NiagaraScriptMergeManager.cpp
+++ b/Engine/Plugins/FX/Niagara/Source/NiagaraEditor/Private/NiagaraScriptMergeManager.cpp
@@ -259,37 +259,50 @@ FNiagaraStackFunctionMergeAdapter::FNiagaraStackFunctionMergeAdapter(const FVers
 		TArray<const UEdGraphPin*> FunctionInputPins;
 		GetStackFunctionInputPins(*FunctionCallNode, FunctionInputPins, Resolver, FNiagaraStackGraphUtilities::ENiagaraGetStackFunctionInputPinsOptions::ModuleInputsOnly, false);
 		
+		TArray<FNiagaraVariable> FunctionInputVariables;
+		TArray<FName> FunctionInputVariableNames;
 		for (const UEdGraphPin* FunctionInputPin : FunctionInputPins)
 		{
 			FNiagaraVariable FunctionInputVariable = NiagaraSchema->PinToNiagaraVariable(FunctionInputPin);
 			if (FunctionInputVariable.IsValid() && FNiagaraStackGraphUtilities::IsRapidIterationType(FunctionInputVariable.GetType()))
 			{
-				UEdGraphPin* FunctionInputDefaultPin = FunctionCallNode->FindParameterMapDefaultValuePin(FunctionInputPin->PinName, OwningScript->GetUsage(), ConstantResolver);
-				if (FunctionInputDefaultPin != nullptr)
-				{
-					// Try to get the default value from the default pin.
-					FNiagaraVariable FunctionInputDefaultVariable = NiagaraSchema->PinToNiagaraVariable(FunctionInputDefaultPin);
-					if (FunctionInputDefaultVariable.GetData() != nullptr)
-					{
-						FunctionInputVariable.SetData(FunctionInputDefaultVariable.GetData());
-					}
-				}
+				FunctionInputVariables.Add(FunctionInputVariable);
+				FunctionInputVariableNames.Add(FunctionInputVariable.GetName());
+			}
+		}
 
-				if (FunctionInputVariable.GetData() == nullptr)
-				{
-					// If the pin didn't have a default value then use the type default.
-					FNiagaraEditorUtilities::ResetVariableToDefaultValue(FunctionInputVariable);
-				}
+		TArray<UEdGraphPin*> FunctionInputDefaultValuePins;
+		FunctionInputDefaultValuePins.AddZeroed(FunctionInputVariables.Num());
+		FunctionCallNode->FindParameterMapDefaultValuePins(FunctionInputVariableNames, InOwningScript.GetUsage(), ConstantResolver, FunctionInputDefaultValuePins);
 
-				if (FunctionInputVariable.GetData() != nullptr)
+		for (int32 i = 0; i < FunctionInputVariables.Num(); i++)
+		{
+			FNiagaraVariable FunctionInputVariable = FunctionInputVariables[i];
+			UEdGraphPin* FunctionInputDefaultPin = FunctionInputDefaultValuePins[i];
+			if (FunctionInputDefaultPin != nullptr)
+			{
+				// Try to get the default value from the default pin.
+				FNiagaraVariable FunctionInputDefaultVariable = NiagaraSchema->PinToNiagaraVariable(FunctionInputDefaultPin);
+				if (FunctionInputDefaultVariable.GetData() != nullptr)
 				{
-					FNiagaraParameterHandle AliasedFunctionInputHandle = FNiagaraParameterHandle::CreateAliasedModuleParameterHandle(FNiagaraParameterHandle(FunctionInputVariable.GetName()), &InFunctionCallNode);
-					FNiagaraVariable FunctionInputRapidIterationParameter =
-						FNiagaraStackGraphUtilities::CreateRapidIterationParameter(UniqueEmitterName, OwningScript->GetUsage(), AliasedFunctionInputHandle.GetParameterHandleString(), FunctionInputVariable.GetType());
-					FunctionInputRapidIterationParameter.SetData(FunctionInputVariable.GetData());
-					RapidIterationInputDefaultValues.Add(FunctionInputRapidIterationParameter);
+					FunctionInputVariable.SetData(FunctionInputDefaultVariable.GetData());
 				}
 			}
+
+			if (FunctionInputVariable.GetData() == nullptr)
+			{
+				// If the pin didn't have a default value then use the type default.
+				FNiagaraEditorUtilities::ResetVariableToDefaultValue(FunctionInputVariable);
+			}
+
+			if (FunctionInputVariable.GetData() != nullptr)
+			{
+				FNiagaraParameterHandle AliasedFunctionInputHandle = FNiagaraParameterHandle::CreateAliasedModuleParameterHandle(FNiagaraParameterHandle(FunctionInputVariable.GetName()), &InFunctionCallNode);
+				FNiagaraVariable FunctionInputRapidIterationParameter =
+					FNiagaraStackGraphUtilities::CreateRapidIterationParameter(UniqueEmitterName, OwningScript->GetUsage(), AliasedFunctionInputHandle.GetParameterHandleString(), FunctionInputVariable.GetType());
+				FunctionInputRapidIterationParameter.SetData(FunctionInputVariable.GetData());
+				RapidIterationInputDefaultValues.Add(FunctionInputRapidIterationParameter);
+			}
 		}
 	}
 
diff --git a/Engine/Plugins/FX/Niagara/Source/NiagaraEditor/Public/NiagaraEditorUtilities.h b/Engine/Plugins/FX/Niagara/Source/NiagaraEditor/Public/NiagaraEditorUtilities.h
--- a/Engine/Plugins/FX/Niagara/Source/NiagaraEditor/Public/NiagaraEditorUtilities.h
+++ b/Engine/Plugins/FX/Niagara/Source/NiagaraEditor/Public/NiagaraEditorUtilities.h
@@ -10,7 +10,6 @@
 #include "NiagaraActions.h"
 #include "NiagaraGraph.h"
 #include "NiagaraEditorSettings.h"
-#include "UpgradeNiagaraScriptResults.h"
 #include "ViewModels/NiagaraSystemScalabilityViewModel.h"
 #include "ViewModels/NiagaraSystemViewModel.h"
 
@@ -42,6 +41,8 @@ enum class EScriptSource : uint8;
 struct FNiagaraNamespaceMetadata;
 class FNiagaraParameterHandle;
 class INiagaraParameterDefinitionsSubscriberViewModel;
+struct FNiagaraScriptVersionUpgradeContext;
+class UUpgradeNiagaraEmitterContext;
 
 enum class ENiagaraFunctionDebugState : uint8;
 
@@ -394,6 +395,8 @@ namespace FNiagaraEditorUtilities
 	/** Gets a list of the registered payload types which are allowed in the current editor context.  This API should be
 		called when providing a list types to the user instead of getting the type list directly from the type registry. */
 	void GetAllowedPayloadTypes(TArray<FNiagaraTypeDefinition>& OutAllowedTypes);
+
+	bool IsEnumIndexVisible(const UEnum* Enum, int32 Index);
 };
 
 namespace FNiagaraParameterUtilities
@@ -470,3 +473,30 @@ namespace FNiagaraParameterDefinitionsUtilities
 	EParameterDefinitionMatchState GetDefinitionMatchStateForParameter(const FNiagaraVariableBase& Parameter);
 	void TrySubscribeScriptVarToDefinitionByName(UNiagaraScriptVariable* ScriptVar, INiagaraParameterDefinitionsSubscriberViewModel* OwningDefinitionSubscriberViewModel);
 };
+
+class FNiagaraEnumIndexVisibilityCache
+{
+public:
+	static bool GetVisibility(const UEnum* InEnum, int32 InIndex);
+
+private:
+	struct FEnumIndexPair
+	{
+		FEnumIndexPair(const UEnum* InEnum, int32 InIndex) : Enum(InEnum), Index(InIndex) { }
+		bool operator==(const FEnumIndexPair& Other) const
+		{
+			return Enum == Other.Enum && Index == Other.Index;
+		}
+		const UEnum* Enum;
+		int32 Index;
+	};
+
+	friend FORCEINLINE uint32 GetTypeHash(const FEnumIndexPair& EnumIndexPair)
+	{
+		return HashCombineFast(GetTypeHash(EnumIndexPair.Enum), GetTypeHash(EnumIndexPair.Index));
+	}
+
+	static TMap<FEnumIndexPair, bool> Cache;
+	static FCriticalSection CacheLock;
+};
+
