- Remove Experimental from various DataInterfaces

#rnx

diff --git a/Engine/Plugins/FX/Niagara/Source/Niagara/Classes/NiagaraDataInterfaceGrid2DCollection.h b/Engine/Plugins/FX/Niagara/Source/Niagara/Classes/NiagaraDataInterfaceGrid2DCollection.h
--- a/Engine/Plugins/FX/Niagara/Source/Niagara/Classes/NiagaraDataInterfaceGrid2DCollection.h
+++ b/Engine/Plugins/FX/Niagara/Source/Niagara/Classes/NiagaraDataInterfaceGrid2DCollection.h
@@ -105,7 +105,7 @@ struct FNiagaraDataInterfaceParametersCS_Grid2DCollection : public FNiagaraDataI
 	LAYOUT_FIELD(TMemoryImageArray<uint32>, AttributeChannelCount);
 };
 
-UCLASS(EditInlineNew, Category = "Grid", meta = (DisplayName = "Grid2D Collection", Experimental), Blueprintable, BlueprintType)
+UCLASS(EditInlineNew, Category = "Grid", meta = (DisplayName = "Grid2D Collection"), Blueprintable, BlueprintType)
 class NIAGARA_API UNiagaraDataInterfaceGrid2DCollection : public UNiagaraDataInterfaceGrid2D
 {
 	GENERATED_UCLASS_BODY()
diff --git a/Engine/Plugins/FX/Niagara/Source/Niagara/Classes/NiagaraDataInterfaceGrid3DCollection.h b/Engine/Plugins/FX/Niagara/Source/Niagara/Classes/NiagaraDataInterfaceGrid3DCollection.h
--- a/Engine/Plugins/FX/Niagara/Source/Niagara/Classes/NiagaraDataInterfaceGrid3DCollection.h
+++ b/Engine/Plugins/FX/Niagara/Source/Niagara/Classes/NiagaraDataInterfaceGrid3DCollection.h
@@ -116,7 +116,7 @@ struct FNiagaraDataInterfaceParametersCS_Grid3DCollection : public FNiagaraDataI
 	LAYOUT_FIELD(TMemoryImageArray<uint32>, AttributeChannelCount);
 };
 
-UCLASS(EditInlineNew, Category = "Grid", meta = (DisplayName = "Grid3D Collection", Experimental), Blueprintable, BlueprintType)
+UCLASS(EditInlineNew, Category = "Grid", meta = (DisplayName = "Grid3D Collection"), Blueprintable, BlueprintType)
 class NIAGARA_API UNiagaraDataInterfaceGrid3DCollection : public UNiagaraDataInterfaceGrid3D
 {
 	GENERATED_UCLASS_BODY()
diff --git a/Engine/Plugins/FX/Niagara/Source/Niagara/Classes/NiagaraDataInterfaceGrid3DCollectionReader.h b/Engine/Plugins/FX/Niagara/Source/Niagara/Classes/NiagaraDataInterfaceGrid3DCollectionReader.h
--- a/Engine/Plugins/FX/Niagara/Source/Niagara/Classes/NiagaraDataInterfaceGrid3DCollectionReader.h
+++ b/Engine/Plugins/FX/Niagara/Source/Niagara/Classes/NiagaraDataInterfaceGrid3DCollectionReader.h
@@ -14,7 +14,7 @@ class FNiagaraSystemInstance;
 class UTextureRenderTarget;
 class UTextureRenderTargetVolume;
 
-UCLASS(EditInlineNew, Category = "Grid", hidecategories=(Grid3DCollection,Grid, Deprecated), meta = (DisplayName = "Grid3D Collection Reader", Experimental), Blueprintable, BlueprintType)
+UCLASS(EditInlineNew, Category = "Grid", hidecategories=(Grid3DCollection,Grid, Deprecated), meta = (DisplayName = "Grid3D Collection Reader"), Blueprintable, BlueprintType)
 class NIAGARA_API UNiagaraDataInterfaceGrid3DCollectionReader : public UNiagaraDataInterfaceGrid3DCollection
 {
 	GENERATED_UCLASS_BODY()
diff --git a/Engine/Plugins/FX/Niagara/Source/Niagara/Classes/NiagaraDataInterfaceRenderTarget2D.h b/Engine/Plugins/FX/Niagara/Source/Niagara/Classes/NiagaraDataInterfaceRenderTarget2D.h
--- a/Engine/Plugins/FX/Niagara/Source/Niagara/Classes/NiagaraDataInterfaceRenderTarget2D.h
+++ b/Engine/Plugins/FX/Niagara/Source/Niagara/Classes/NiagaraDataInterfaceRenderTarget2D.h
@@ -80,7 +80,7 @@ struct FNiagaraDataInterfaceProxyRenderTarget2DProxy : public FNiagaraDataInterf
 	TMap<FNiagaraSystemInstanceID, FRenderTarget2DRWInstanceData_RenderThread> SystemInstancesToProxyData_RT;
 };
 
-UCLASS(EditInlineNew, Category = "Grid", meta = (DisplayName = "Render Target 2D", Experimental), Blueprintable, BlueprintType)
+UCLASS(EditInlineNew, Category = "Rendering", meta = (DisplayName = "Render Target 2D"), Blueprintable, BlueprintType)
 class NIAGARA_API UNiagaraDataInterfaceRenderTarget2D : public UNiagaraDataInterfaceRWBase
 {
 	GENERATED_UCLASS_BODY()
diff --git a/Engine/Plugins/FX/Niagara/Source/Niagara/Classes/NiagaraDataInterfaceRenderTarget2DArray.h b/Engine/Plugins/FX/Niagara/Source/Niagara/Classes/NiagaraDataInterfaceRenderTarget2DArray.h
--- a/Engine/Plugins/FX/Niagara/Source/Niagara/Classes/NiagaraDataInterfaceRenderTarget2DArray.h
+++ b/Engine/Plugins/FX/Niagara/Source/Niagara/Classes/NiagaraDataInterfaceRenderTarget2DArray.h
@@ -77,7 +77,7 @@ struct FNiagaraDataInterfaceProxyRenderTarget2DArrayProxy : public FNiagaraDataI
 	TMap<FNiagaraSystemInstanceID, FRenderTarget2DArrayRWInstanceData_RenderThread> SystemInstancesToProxyData_RT;
 };
 
-UCLASS(EditInlineNew, Category = "Grid", meta = (DisplayName = "Render Target 2D Array", Experimental), Blueprintable, BlueprintType)
+UCLASS(EditInlineNew, Category = "Rendering", meta = (DisplayName = "Render Target 2D Array"), Blueprintable, BlueprintType)
 class NIAGARA_API UNiagaraDataInterfaceRenderTarget2DArray : public UNiagaraDataInterfaceRWBase
 {
 	GENERATED_UCLASS_BODY()
diff --git a/Engine/Plugins/FX/Niagara/Source/Niagara/Classes/NiagaraDataInterfaceRenderTargetCube.h b/Engine/Plugins/FX/Niagara/Source/Niagara/Classes/NiagaraDataInterfaceRenderTargetCube.h
--- a/Engine/Plugins/FX/Niagara/Source/Niagara/Classes/NiagaraDataInterfaceRenderTargetCube.h
+++ b/Engine/Plugins/FX/Niagara/Source/Niagara/Classes/NiagaraDataInterfaceRenderTargetCube.h
@@ -70,7 +70,7 @@ struct FNiagaraDataInterfaceProxyRenderTargetCubeProxy : public FNiagaraDataInte
 	TMap<FNiagaraSystemInstanceID, FRenderTargetCubeRWInstanceData_RenderThread> SystemInstancesToProxyData_RT;
 };
 
-UCLASS(EditInlineNew, Category = "Grid", meta = (DisplayName = "Render Target Cube", Experimental), Blueprintable, BlueprintType)
+UCLASS(EditInlineNew, Category = "Rendering", meta = (DisplayName = "Render Target Cube"), Blueprintable, BlueprintType)
 class NIAGARA_API UNiagaraDataInterfaceRenderTargetCube : public UNiagaraDataInterfaceRWBase
 {
 	GENERATED_UCLASS_BODY()
diff --git a/Engine/Plugins/FX/Niagara/Source/Niagara/Classes/NiagaraDataInterfaceRenderTargetVolume.h b/Engine/Plugins/FX/Niagara/Source/Niagara/Classes/NiagaraDataInterfaceRenderTargetVolume.h
--- a/Engine/Plugins/FX/Niagara/Source/Niagara/Classes/NiagaraDataInterfaceRenderTargetVolume.h
+++ b/Engine/Plugins/FX/Niagara/Source/Niagara/Classes/NiagaraDataInterfaceRenderTargetVolume.h
@@ -69,7 +69,7 @@ struct FNiagaraDataInterfaceProxyRenderTargetVolumeProxy : public FNiagaraDataIn
 	TMap<FNiagaraSystemInstanceID, FRenderTargetVolumeRWInstanceData_RenderThread> SystemInstancesToProxyData_RT;
 };
 
-UCLASS(EditInlineNew, Category = "Grid", meta = (DisplayName = "Render Target Volume", Experimental), Blueprintable, BlueprintType)
+UCLASS(EditInlineNew, Category = "Rendering", meta = (DisplayName = "Render Target Volume"), Blueprintable, BlueprintType)
 class NIAGARA_API UNiagaraDataInterfaceRenderTargetVolume : public UNiagaraDataInterfaceRWBase
 {
 	GENERATED_UCLASS_BODY()
diff --git a/Engine/Plugins/FX/Niagara/Source/Niagara/Private/DataInterface/NiagaraDataInterfaceEmitterProperties.h b/Engine/Plugins/FX/Niagara/Source/Niagara/Private/DataInterface/NiagaraDataInterfaceEmitterProperties.h
--- a/Engine/Plugins/FX/Niagara/Source/Niagara/Private/DataInterface/NiagaraDataInterfaceEmitterProperties.h
+++ b/Engine/Plugins/FX/Niagara/Source/Niagara/Private/DataInterface/NiagaraDataInterfaceEmitterProperties.h
@@ -11,7 +11,7 @@
 /**
 Allows access to various emitter properties that are not part of the simulation data.
 */
-UCLASS(EditInlineNew, Category = "DataInterface", meta=(DisplayName="Emitter Properties", Experimental))
+UCLASS(EditInlineNew, Category = "DataInterface", meta=(DisplayName="Emitter Properties"))
 class UNiagaraDataInterfaceEmitterProperties : public UNiagaraDataInterface
 {
 	GENERATED_UCLASS_BODY()
diff --git a/Engine/Plugins/FX/Niagara/Source/Niagara/Private/DataInterface/NiagaraDataInterfaceGBuffer.h b/Engine/Plugins/FX/Niagara/Source/Niagara/Private/DataInterface/NiagaraDataInterfaceGBuffer.h
--- a/Engine/Plugins/FX/Niagara/Source/Niagara/Private/DataInterface/NiagaraDataInterfaceGBuffer.h
+++ b/Engine/Plugins/FX/Niagara/Source/Niagara/Private/DataInterface/NiagaraDataInterfaceGBuffer.h
@@ -6,7 +6,7 @@
 #include "NiagaraDataInterface.h"
 #include "NiagaraDataInterfaceGBuffer.generated.h"
 
-UCLASS(Experimental, EditInlineNew, Category = "Camera", meta = (DisplayName = "GBuffer"))
+UCLASS(EditInlineNew, Category = "Rendering", meta = (DisplayName = "GBuffer"))
 class NIAGARA_API UNiagaraDataInterfaceGBuffer : public UNiagaraDataInterface
 {
 	GENERATED_UCLASS_BODY()
diff --git a/Engine/Plugins/FX/Niagara/Source/Niagara/Private/DataInterface/NiagaraDataInterfaceUObjectPropertyReader.h b/Engine/Plugins/FX/Niagara/Source/Niagara/Private/DataInterface/NiagaraDataInterfaceUObjectPropertyReader.h
--- a/Engine/Plugins/FX/Niagara/Source/Niagara/Private/DataInterface/NiagaraDataInterfaceUObjectPropertyReader.h
+++ b/Engine/Plugins/FX/Niagara/Source/Niagara/Private/DataInterface/NiagaraDataInterfaceUObjectPropertyReader.h
@@ -28,7 +28,7 @@ struct FNiagaraUObjectPropertyReaderRemap
 Data interface to read properties from UObjects.
 Rather than having BP tick functions that push data into Niagara this data interface will instead pull them.
 */
-UCLASS(EditInlineNew, Category = "DataInterface", meta=(DisplayName="Object Reader", Experimental))
+UCLASS(EditInlineNew, Category = "DataInterface", meta=(DisplayName="Object Reader"))
 class UNiagaraDataInterfaceUObjectPropertyReader : public UNiagaraDataInterface
 {
 	GENERATED_UCLASS_BODY()
