- Add reverse play to Niagara baker

diff --git a/Engine/Plugins/FX/Niagara/Source/NiagaraEditor/Private/Widgets/SNiagaraBakerWidget.cpp b/Engine/Plugins/FX/Niagara/Source/NiagaraEditor/Private/Widgets/SNiagaraBakerWidget.cpp
--- a/Engine/Plugins/FX/Niagara/Source/NiagaraEditor/Private/Widgets/SNiagaraBakerWidget.cpp
+++ b/Engine/Plugins/FX/Niagara/Source/NiagaraEditor/Private/Widgets/SNiagaraBakerWidget.cpp
@@ -25,7 +25,6 @@
 #include "EditorWidgetsModule.h"
 #include "IDetailCustomization.h"
 #include "IDocumentation.h"
-#include "ITransportControl.h"
 #include "PropertyEditorModule.h"
 
 #define LOCTEXT_NAMESPACE "NiagaraBakerWidget"
@@ -323,9 +322,10 @@ void SNiagaraBakerWidget::Construct(const FArguments& InArgs)
 	// Transport control args
 	{
 		FTransportControlArgs TransportControlArgs;
-		TransportControlArgs.OnGetPlaybackMode.BindLambda([&]() -> EPlaybackMode::Type { return bIsPlaying ? EPlaybackMode::PlayingForward : EPlaybackMode::Stopped; } );
+		TransportControlArgs.OnGetPlaybackMode.BindLambda([&]() -> EPlaybackMode::Type { return PlaybackMode; });
 		TransportControlArgs.OnBackwardEnd.BindSP(this, &SNiagaraBakerWidget::OnTransportBackwardEnd);
 		TransportControlArgs.OnBackwardStep.BindSP(this, &SNiagaraBakerWidget::OnTransportBackwardStep);
+		TransportControlArgs.OnBackwardPlay.BindSP(this, &SNiagaraBakerWidget::OnTransportBackwardPlay);
 		TransportControlArgs.OnForwardPlay.BindSP(this, &SNiagaraBakerWidget::OnTransportForwardPlay);
 		TransportControlArgs.OnForwardStep.BindSP(this, &SNiagaraBakerWidget::OnTransportForwardStep);
 		TransportControlArgs.OnForwardEnd.BindSP(this, &SNiagaraBakerWidget::OnTransportForwardEnd);
@@ -334,6 +334,7 @@ void SNiagaraBakerWidget::Construct(const FArguments& InArgs)
 
 		TransportControlArgs.WidgetsToCreate.Add(FTransportControlWidget(ETransportControlWidgetType::BackwardEnd));
 		TransportControlArgs.WidgetsToCreate.Add(FTransportControlWidget(ETransportControlWidgetType::BackwardStep));
+		TransportControlArgs.WidgetsToCreate.Add(FTransportControlWidget(ETransportControlWidgetType::BackwardPlay));
 		TransportControlArgs.WidgetsToCreate.Add(FTransportControlWidget(ETransportControlWidgetType::ForwardPlay));
 		TransportControlArgs.WidgetsToCreate.Add(FTransportControlWidget(ETransportControlWidgetType::ForwardStep));
 		TransportControlArgs.WidgetsToCreate.Add(FTransportControlWidget(ETransportControlWidgetType::ForwardEnd));
@@ -945,23 +946,24 @@ void SNiagaraBakerWidget::Tick(const FGeometry& AllottedGeometry, const double C
 	const float DurationSeconds = BakerSettings->DurationSeconds;
 	if ( DurationSeconds > 0.0f )
 	{
-		if (bIsPlaying)
+		if (PlaybackMode != EPlaybackMode::Stopped)
 		{
-			PreviewRelativeTime += DeltaTime;
+			PreviewRelativeTime += PlaybackMode == EPlaybackMode::PlayingForward ? DeltaTime : -DeltaTime;
 			if ( BakerSettings->bPreviewLooping )
 			{
 				PreviewRelativeTime = FMath::Fmod(PreviewRelativeTime, DurationSeconds);
+				PreviewRelativeTime += PreviewRelativeTime < 0.0f ? DurationSeconds : 0.0f;
 			}
-			else if ( PreviewRelativeTime >= DurationSeconds )
+			else if (PlaybackMode == EPlaybackMode::PlayingReverse && PreviewRelativeTime <= 0.0f)
 			{
-				PreviewRelativeTime = DurationSeconds;
-				bIsPlaying = false;
+				PlaybackMode = EPlaybackMode::Stopped;
+			}
+			else if (PlaybackMode == EPlaybackMode::PlayingForward && PreviewRelativeTime >= DurationSeconds)
+			{
+				PlaybackMode = EPlaybackMode::Stopped;
 			}
 		}
-		else
-		{
-			PreviewRelativeTime = FMath::Min(PreviewRelativeTime, DurationSeconds);
-		}
+		PreviewRelativeTime = FMath::Clamp(PreviewRelativeTime, 0.0f, DurationSeconds);
 
 		ViewportWidget->RefreshView(PreviewRelativeTime, DeltaTime);
 		TimelineWidget->SetRelativeTime(PreviewRelativeTime);
@@ -972,7 +974,7 @@ void SNiagaraBakerWidget::Tick(const FGeometry& AllottedGeometry, const double C
 
 void SNiagaraBakerWidget::SetPreviewRelativeTime(float RelativeTime)
 {
-	bIsPlaying = false;
+	PlaybackMode = EPlaybackMode::Stopped;
 	PreviewRelativeTime = RelativeTime;
 }
 
@@ -1005,7 +1007,7 @@ const UNiagaraBakerSettings* SNiagaraBakerWidget::GetBakerGeneratedSettings() co
 
 FReply SNiagaraBakerWidget::OnTransportBackwardEnd()
 {
-	bIsPlaying = false;
+	PlaybackMode = EPlaybackMode::Stopped;
 	PreviewRelativeTime = 0.0f;
 
 	return FReply::Handled();
@@ -1013,7 +1015,7 @@ FReply SNiagaraBakerWidget::OnTransportBackwardEnd()
 
 FReply SNiagaraBakerWidget::OnTransportBackwardStep()
 {
-	bIsPlaying = false;
+	PlaybackMode = EPlaybackMode::Stopped;
 	if (FNiagaraBakerViewModel* ViewModel = WeakViewModel.Pin().Get())
 	{
 		const FNiagaraBakerOutputFrameIndices OutputFrameIndices = ViewModel->GetCurrentOutputFrameIndices(PreviewRelativeTime);
@@ -1024,15 +1026,21 @@ FReply SNiagaraBakerWidget::OnTransportBackwardStep()
 	return FReply::Handled();
 }
 
+FReply SNiagaraBakerWidget::OnTransportBackwardPlay()
+{
+	PlaybackMode = PlaybackMode == EPlaybackMode::PlayingReverse ? EPlaybackMode::Stopped : EPlaybackMode::PlayingReverse;
+	return FReply::Handled();
+}
+
 FReply SNiagaraBakerWidget::OnTransportForwardPlay()
 {
-	bIsPlaying = !bIsPlaying;
+	PlaybackMode = PlaybackMode == EPlaybackMode::PlayingForward ? EPlaybackMode::Stopped : EPlaybackMode::PlayingForward;
 	return FReply::Handled();
 }
 
 FReply SNiagaraBakerWidget::OnTransportForwardStep()
 {
-	bIsPlaying = false;
+	PlaybackMode = EPlaybackMode::Stopped;
 	if (FNiagaraBakerViewModel* ViewModel = WeakViewModel.Pin().Get())
 	{
 		const FNiagaraBakerOutputFrameIndices OutputFrameIndices = ViewModel->GetCurrentOutputFrameIndices(PreviewRelativeTime);
@@ -1045,7 +1053,7 @@ FReply SNiagaraBakerWidget::OnTransportForwardStep()
 
 FReply SNiagaraBakerWidget::OnTransportForwardEnd()
 {
-	bIsPlaying = false;
+	PlaybackMode = EPlaybackMode::Stopped;
 	if (UNiagaraBakerSettings* BakerSettings = GetBakerSettings())
 	{
 		PreviewRelativeTime = BakerSettings->DurationSeconds - KINDA_SMALL_NUMBER;
diff --git a/Engine/Plugins/FX/Niagara/Source/NiagaraEditor/Private/Widgets/SNiagaraBakerWidget.h b/Engine/Plugins/FX/Niagara/Source/NiagaraEditor/Private/Widgets/SNiagaraBakerWidget.h
--- a/Engine/Plugins/FX/Niagara/Source/NiagaraEditor/Private/Widgets/SNiagaraBakerWidget.h
+++ b/Engine/Plugins/FX/Niagara/Source/NiagaraEditor/Private/Widgets/SNiagaraBakerWidget.h
@@ -4,6 +4,7 @@
 
 #include "Widgets/DeclarativeSyntaxSupport.h"
 #include "Widgets/SCompoundWidget.h"
+#include "ITransportControl.h"
 
 class FNiagaraBakerViewModel;
 class SNiagaraBakerViewport;
@@ -49,6 +50,7 @@ public:
 
 	FReply OnTransportBackwardEnd();
 	FReply OnTransportBackwardStep();
+	FReply OnTransportBackwardPlay();
 	FReply OnTransportForwardPlay();
 	FReply OnTransportForwardStep();
 	FReply OnTransportForwardEnd();
@@ -64,7 +66,7 @@ private:
 	TSharedPtr<IDetailsView>					BakerSettingsDetails;
 	TSharedPtr<SWidget>							TransportControls;
 
-	bool										bIsPlaying = true;
+	EPlaybackMode::Type							PlaybackMode = EPlaybackMode::PlayingForward;
 	float										PreviewRelativeTime = 0.0f;
 
 	FDelegateHandle								OnCurrentOutputIndexChangedHandle;
