- Fix for not including rotation as part of UB key hash

diff --git a/Engine/Plugins/FX/Niagara/Source/Niagara/Private/NiagaraComponent.cpp b/Engine/Plugins/FX/Niagara/Source/Niagara/Private/NiagaraComponent.cpp
--- a/Engine/Plugins/FX/Niagara/Source/Niagara/Private/NiagaraComponent.cpp
+++ b/Engine/Plugins/FX/Niagara/Source/Niagara/Private/NiagaraComponent.cpp
@@ -302,19 +302,15 @@ TUniformBuffer<FPrimitiveUniformShaderParameters>* FNiagaraSceneProxy::GetCustom
 
 	// we need to incorporate the transform of the PreviousLocalToWorld because it stopping moving is relevant
 	// and we need to be able to generate a new UB for that case
-	const auto PreviousOrigin = PreviousLocalToWorld.GetOrigin();
-	KeyHash = HashCombine(KeyHash, PreviousOrigin.X);
-	KeyHash = HashCombine(KeyHash, PreviousOrigin.Y);
-	KeyHash = HashCombine(KeyHash, PreviousOrigin.Z);
+	KeyHash = HashCombine(KeyHash, UE::Math::GetTypeHash(PreviousLocalToWorld.GetScaledAxis(EAxis::X)));
+	KeyHash = HashCombine(KeyHash, UE::Math::GetTypeHash(PreviousLocalToWorld.GetScaledAxis(EAxis::Y)));
+	KeyHash = HashCombine(KeyHash, UE::Math::GetTypeHash(PreviousLocalToWorld.GetScaledAxis(EAxis::Z)));
+	KeyHash = HashCombine(KeyHash, UE::Math::GetTypeHash(PreviousLocalToWorld.GetOrigin()));
 
 	if ( InstanceBounds.IsValid )
 	{
-		KeyHash = HashCombine(KeyHash, InstanceBounds.Min.X);
-		KeyHash = HashCombine(KeyHash, InstanceBounds.Min.Y);
-		KeyHash = HashCombine(KeyHash, InstanceBounds.Min.Z);
-		KeyHash = HashCombine(KeyHash, InstanceBounds.Max.X);
-		KeyHash = HashCombine(KeyHash, InstanceBounds.Max.Y);
-		KeyHash = HashCombine(KeyHash, InstanceBounds.Max.Z);
+		KeyHash = HashCombine(KeyHash, UE::Math::GetTypeHash(InstanceBounds.Min));
+		KeyHash = HashCombine(KeyHash, UE::Math::GetTypeHash(InstanceBounds.Max));
 	}
 
 	TUniformBuffer<FPrimitiveUniformShaderParameters>*& CustomUBRef = CustomUniformBuffers.FindOrAdd(KeyHash);
