Niagara - Change the update context so that when adding a component it checks that the component has been forced to solo, and when adding a component we don't do any work at the system level.  This fixes an issue where a non-infinitely looping emitter in the system editor would cause all systems in the world to be reset and never restarted again.

diff --git a/Engine/Plugins/FX/Niagara/Source/Niagara/Private/NiagaraCommon.cpp b/Engine/Plugins/FX/Niagara/Source/Niagara/Private/NiagaraCommon.cpp
--- a/Engine/Plugins/FX/Niagara/Source/Niagara/Private/NiagaraCommon.cpp
+++ b/Engine/Plugins/FX/Niagara/Source/Niagara/Private/NiagaraCommon.cpp
@@ -149,10 +149,13 @@ void FNiagaraSystemUpdateContext::AddAll(bool bReInit)
 	}
 }
 
-void FNiagaraSystemUpdateContext::Add(UNiagaraComponent* Component, bool bReInit)
+void FNiagaraSystemUpdateContext::AddSoloComponent(UNiagaraComponent* Component, bool bReInit)
 {
 	check(Component);
-	AddInternal(Component, bReInit);
+	if (ensureMsgf(Component->GetForceSolo(), TEXT("A component must have a solo system simulation when used with an update context.")))
+	{
+		AddInternal(Component, bReInit);
+	}
 }
 
 void FNiagaraSystemUpdateContext::Add(const UNiagaraSystem* System, bool bReInit)
@@ -228,15 +231,17 @@ void FNiagaraSystemUpdateContext::AddInternal(UNiagaraComponent* Comp, bool bReI
 		}
 	}
 
-	if (bReInit || bDestroySystemSim)
+	if (Comp->GetForceSolo() == false)
 	{
-		SystemSimsToRecache.AddUnique(Comp->GetAsset());
-	}
+		if (bReInit || bDestroySystemSim)
+		{
+			SystemSimsToRecache.AddUnique(Comp->GetAsset());
+		}
 
-	if (bReInit && bDestroySystemSim)
-	{
-		//Always destroy the system sims on a reinit, even if we're not reactivating the component.
-		SystemSimsToDestroy.AddUnique(Comp->GetAsset());
+		if (bReInit && bDestroySystemSim)
+		{
+			SystemSimsToDestroy.AddUnique(Comp->GetAsset());
+		}
 	}
 
 	bool bIsActive = (Comp->IsActive() && Comp->GetRequestedExecutionState() == ENiagaraExecutionState::Active) || Comp->IsRegisteredWithScalabilityManager();
diff --git a/Engine/Plugins/FX/Niagara/Source/Niagara/Public/NiagaraCommon.h b/Engine/Plugins/FX/Niagara/Source/Niagara/Public/NiagaraCommon.h
--- a/Engine/Plugins/FX/Niagara/Source/Niagara/Public/NiagaraCommon.h
+++ b/Engine/Plugins/FX/Niagara/Source/Niagara/Public/NiagaraCommon.h
@@ -750,7 +750,6 @@ struct NIAGARA_API FNiagaraSystemUpdateContext
 {
 	GENERATED_BODY()
 
-	FNiagaraSystemUpdateContext(UNiagaraComponent* Component, bool bReInit, bool bInDestroyOnAdd = false, bool bInOnlyActive = false, bool bInDestroySystemSim = true) :bDestroyOnAdd(bInDestroyOnAdd), bOnlyActive(bInOnlyActive), bDestroySystemSim(bInDestroySystemSim) { Add(Component, bReInit); }
 	FNiagaraSystemUpdateContext(const UNiagaraSystem* System, bool bReInit, bool bInDestroyOnAdd = false, bool bInOnlyActive = false, bool bInDestroySystemSim = true) :bDestroyOnAdd(bInDestroyOnAdd), bOnlyActive(bInOnlyActive), bDestroySystemSim(bInDestroySystemSim) { Add(System, bReInit); }
 #if WITH_EDITORONLY_DATA
 	FNiagaraSystemUpdateContext(const FVersionedNiagaraEmitter& Emitter, bool bReInit, bool bInDestroyOnAdd = false, bool bInOnlyActive = false, bool bInDestroySystemSim = true) :bDestroyOnAdd(bInDestroyOnAdd), bOnlyActive(bInOnlyActive), bDestroySystemSim(bInDestroySystemSim) { Add(Emitter, bReInit); }
@@ -766,7 +765,7 @@ struct NIAGARA_API FNiagaraSystemUpdateContext
 	void SetOnlyActive(bool bInOnlyActive) { bOnlyActive = bInOnlyActive; }
 	void SetDestroySystemSim(bool bInDestroySystemSim) { bDestroySystemSim = bInDestroySystemSim; }
 
-	void Add(UNiagaraComponent* Component, bool bReInit);
+	void AddSoloComponent(UNiagaraComponent* Component, bool bReInit);
 	void Add(const UNiagaraSystem* System, bool bReInit);
 #if WITH_EDITORONLY_DATA
 	void Add(const FVersionedNiagaraEmitter& Emitter, bool bReInit);
diff --git a/Engine/Plugins/FX/Niagara/Source/NiagaraEditor/Private/ViewModels/NiagaraSystemViewModel.cpp b/Engine/Plugins/FX/Niagara/Source/NiagaraEditor/Private/ViewModels/NiagaraSystemViewModel.cpp
--- a/Engine/Plugins/FX/Niagara/Source/NiagaraEditor/Private/ViewModels/NiagaraSystemViewModel.cpp
+++ b/Engine/Plugins/FX/Niagara/Source/NiagaraEditor/Private/ViewModels/NiagaraSystemViewModel.cpp
@@ -2040,9 +2040,11 @@ void FNiagaraSystemViewModel::ResetSystem(ETimeResetMode TimeResetMode, EMultiRe
 
 	if (MultiResetMode == EMultiResetMode::ResetThisInstance)
 	{
+		// Resetting the single preview instance is only valid when the change causing the reset hasn't invalidated any cached simulation data.
+		// Examples of these changes are bounds changes or preview changes such as resetting on loop or changing the preview quality level.
 		if (PreviewComponent != nullptr)
 		{
-			UpdateContext.Add(PreviewComponent, ReinitMode == EReinitMode::ReinitializeSystem);
+			UpdateContext.AddSoloComponent(PreviewComponent, ReinitMode == EReinitMode::ReinitializeSystem);
 		}
 	}
 	else
