Naigara - Reimplement renderer class filtering to work with the new renderer creation info paradigm.

diff --git a/Engine/Plugins/FX/Niagara/Source/NiagaraEditor/Private/NiagaraEditorModule.cpp b/Engine/Plugins/FX/Niagara/Source/NiagaraEditor/Private/NiagaraEditorModule.cpp
--- a/Engine/Plugins/FX/Niagara/Source/NiagaraEditor/Private/NiagaraEditorModule.cpp
+++ b/Engine/Plugins/FX/Niagara/Source/NiagaraEditor/Private/NiagaraEditorModule.cpp
@@ -1729,6 +1729,7 @@ void FNiagaraEditorModule::RegisterDefaultRendererFactories()
 	RegisterRendererCreationInfo(FNiagaraRendererCreationInfo(
 		UNiagaraMeshRendererProperties::StaticClass()->GetDisplayNameText(),
 		FText::FromString(UNiagaraMeshRendererProperties::StaticClass()->GetDescription()),
+		UNiagaraMeshRendererProperties::StaticClass()->GetClassPathName(),
 		FNiagaraRendererCreationInfo::FRendererFactory::CreateLambda([](UObject* OuterEmitter)
 		{
 			UNiagaraMeshRendererProperties* NewRenderer = NewObject<UNiagaraMeshRendererProperties>(OuterEmitter, NAME_None, RF_Transactional);
@@ -1744,6 +1745,7 @@ void FNiagaraEditorModule::RegisterDefaultRendererFactories()
 	RegisterRendererCreationInfo(FNiagaraRendererCreationInfo(
 		UNiagaraSpriteRendererProperties::StaticClass()->GetDisplayNameText(),
 		FText::FromString(UNiagaraSpriteRendererProperties::StaticClass()->GetDescription()),
+		UNiagaraSpriteRendererProperties::StaticClass()->GetClassPathName(),
 		FNiagaraRendererCreationInfo::FRendererFactory::CreateLambda([](UObject* OuterEmitter)
 		{
 			UNiagaraSpriteRendererProperties* NewRenderer = NewObject<UNiagaraSpriteRendererProperties>(OuterEmitter, NAME_None, RF_Transactional);
@@ -1755,6 +1757,7 @@ void FNiagaraEditorModule::RegisterDefaultRendererFactories()
 	RegisterRendererCreationInfo(FNiagaraRendererCreationInfo(
 		UNiagaraRibbonRendererProperties::StaticClass()->GetDisplayNameText(),
 		FText::FromString(UNiagaraRibbonRendererProperties::StaticClass()->GetDescription()),
+		UNiagaraRibbonRendererProperties::StaticClass()->GetClassPathName(),
 		FNiagaraRendererCreationInfo::FRendererFactory::CreateLambda([](UObject* OuterEmitter)
 		{
 			UNiagaraRibbonRendererProperties* NewRenderer = NewObject<UNiagaraRibbonRendererProperties>(OuterEmitter, NAME_None, RF_Transactional);
@@ -1766,6 +1769,7 @@ void FNiagaraEditorModule::RegisterDefaultRendererFactories()
 	RegisterRendererCreationInfo(FNiagaraRendererCreationInfo(
 		UNiagaraComponentRendererProperties::StaticClass()->GetDisplayNameText(),
 		FText::FromString(UNiagaraComponentRendererProperties::StaticClass()->GetDescription()),
+		UNiagaraComponentRendererProperties::StaticClass()->GetClassPathName(),
 		FNiagaraRendererCreationInfo::FRendererFactory::CreateLambda([](UObject* OuterEmitter)
 		{
 			UNiagaraComponentRendererProperties* NewRenderer = NewObject<UNiagaraComponentRendererProperties>(OuterEmitter, NAME_None, RF_Transactional);
@@ -1775,6 +1779,7 @@ void FNiagaraEditorModule::RegisterDefaultRendererFactories()
 	RegisterRendererCreationInfo(FNiagaraRendererCreationInfo(
 		UNiagaraLightRendererProperties::StaticClass()->GetDisplayNameText(),
 		FText::FromString(UNiagaraLightRendererProperties::StaticClass()->GetDescription()),
+		UNiagaraLightRendererProperties::StaticClass()->GetClassPathName(),
 		FNiagaraRendererCreationInfo::FRendererFactory::CreateLambda([](UObject* OuterEmitter)
 		{
 			UNiagaraLightRendererProperties* NewRenderer = NewObject<UNiagaraLightRendererProperties>(OuterEmitter, NAME_None, RF_Transactional);
diff --git a/Engine/Plugins/FX/Niagara/Source/NiagaraEditor/Private/NiagaraEditorSettings.cpp b/Engine/Plugins/FX/Niagara/Source/NiagaraEditor/Private/NiagaraEditorSettings.cpp
--- a/Engine/Plugins/FX/Niagara/Source/NiagaraEditor/Private/NiagaraEditorSettings.cpp
+++ b/Engine/Plugins/FX/Niagara/Source/NiagaraEditor/Private/NiagaraEditorSettings.cpp
@@ -651,11 +651,21 @@ void UNiagaraEditorSettings::SetOnIsClassAllowed(const FOnIsClassAllowed& InOnIs
 	OnIsClassAllowedDelegate = InOnIsClassAllowed;
 }
 
+void UNiagaraEditorSettings::SetOnIsClassPathAllowed(const FOnIsClassPathAllowed& InOnIsClassPathAllowed)
+{
+	OnIsClassPathAllowedDelegate = InOnIsClassPathAllowed;
+}
+
 bool UNiagaraEditorSettings::IsAllowedClass(const UClass* InClass) const
 {
 	return OnIsClassAllowedDelegate.IsBound() == false || OnIsClassAllowedDelegate.Execute(InClass);
 }
 
+bool UNiagaraEditorSettings::IsAllowedClassPath(const FTopLevelAssetPath& InClassPath) const
+{
+	return OnIsClassPathAllowedDelegate.IsBound() == false || OnIsClassPathAllowedDelegate.Execute(InClassPath);
+}
+
 bool UNiagaraEditorSettings::IsAllowedTypeDefinition(const FNiagaraTypeDefinition& InTypeDefinition) const
 {
 	return InTypeDefinition.GetClass() == nullptr || IsAllowedClass(InTypeDefinition.GetClass());
diff --git a/Engine/Plugins/FX/Niagara/Source/NiagaraEditor/Private/ViewModels/Stack/NiagaraStackRenderItemGroup.cpp b/Engine/Plugins/FX/Niagara/Source/NiagaraEditor/Private/ViewModels/Stack/NiagaraStackRenderItemGroup.cpp
--- a/Engine/Plugins/FX/Niagara/Source/NiagaraEditor/Private/ViewModels/Stack/NiagaraStackRenderItemGroup.cpp
+++ b/Engine/Plugins/FX/Niagara/Source/NiagaraEditor/Private/ViewModels/Stack/NiagaraStackRenderItemGroup.cpp
@@ -71,10 +71,10 @@ public:
 	{
 		const TArray<FNiagaraRendererCreationInfo>& RendererCreationInfos = FNiagaraEditorModule::Get().GetRendererCreationInfos();
 
+		const UNiagaraEditorSettings* NiagaraEditorSettings = GetDefault<UNiagaraEditorSettings>();
 		for (const FNiagaraRendererCreationInfo& RendererCreationInfo : RendererCreationInfos)
 		{
-			// TODO: UE-165740
-			//if (NiagaraEditorSettings->IsAllowedClass(RendererClass))
+			if (NiagaraEditorSettings->IsAllowedClassPath(RendererCreationInfo.RendererClassPath))
 			{
 				OutAddActions.Add(MakeShared<FRenderItemGroupAddAction>(RendererCreationInfo));
 			}
diff --git a/Engine/Plugins/FX/Niagara/Source/NiagaraEditor/Public/NiagaraEditorModule.h b/Engine/Plugins/FX/Niagara/Source/NiagaraEditor/Public/NiagaraEditorModule.h
--- a/Engine/Plugins/FX/Niagara/Source/NiagaraEditor/Public/NiagaraEditorModule.h
+++ b/Engine/Plugins/FX/Niagara/Source/NiagaraEditor/Public/NiagaraEditorModule.h
@@ -102,10 +102,10 @@ struct NIAGARAEDITOR_API FNiagaraRendererCreationInfo
 	GENERATED_BODY()
 
 	FNiagaraRendererCreationInfo() = default;
-	FNiagaraRendererCreationInfo(FText InDisplayName, FRendererFactory InFactory) : DisplayName(InDisplayName), RendererFactory(InFactory)
+	FNiagaraRendererCreationInfo(FText InDisplayName, const FTopLevelAssetPath& InRendererClassPath, FRendererFactory InFactory) : DisplayName(InDisplayName), RendererClassPath(InRendererClassPath), RendererFactory(InFactory)
 	{}
 
-	FNiagaraRendererCreationInfo(FText InDisplayName, FText InDescription, FRendererFactory InFactory) : DisplayName(InDisplayName), Description(InDescription),RendererFactory(InFactory)
+	FNiagaraRendererCreationInfo(FText InDisplayName, FText InDescription, const FTopLevelAssetPath& InRendererClassPath, FRendererFactory InFactory) : DisplayName(InDisplayName), Description(InDescription), RendererClassPath(InRendererClassPath), RendererFactory(InFactory)
 	{}
 
 	UPROPERTY()
@@ -113,6 +113,9 @@ struct NIAGARAEDITOR_API FNiagaraRendererCreationInfo
 
 	UPROPERTY()
 	FText Description;
+
+	UPROPERTY()
+	FTopLevelAssetPath RendererClassPath;
 	
 	FRendererFactory RendererFactory;
 };
diff --git a/Engine/Plugins/FX/Niagara/Source/NiagaraEditor/Public/NiagaraEditorSettings.h b/Engine/Plugins/FX/Niagara/Source/NiagaraEditor/Public/NiagaraEditorSettings.h
--- a/Engine/Plugins/FX/Niagara/Source/NiagaraEditor/Public/NiagaraEditorSettings.h
+++ b/Engine/Plugins/FX/Niagara/Source/NiagaraEditor/Public/NiagaraEditorSettings.h
@@ -361,13 +361,20 @@ public:
 	FNiagaraParameterPanelSectionStorage& FindOrAddParameterPanelSectionStorage(FGuid PanelSectionId, bool& bOutAdded);
 
 	DECLARE_DELEGATE_RetVal_OneParam(bool, FOnIsClassAllowed, const UClass* /*InClass*/);
+	DECLARE_DELEGATE_RetVal_OneParam(bool, FOnIsClassPathAllowed, const FTopLevelAssetPath& /*InClassPath*/);
 	
 	/** Sets a delegate that allows external code to restrict which features can be used within the niagara editor by filtering which classes are allowed. */
 	void SetOnIsClassAllowed(const FOnIsClassAllowed& InOnIsClassAllowed);
 
+	/** Sets a delegate that allows external code to restrict which features can be used within the niagara editor by filtering which classes are allowed by class path. */
+	void SetOnIsClassPathAllowed(const FOnIsClassPathAllowed& InOnIsClassPathAllowed);
+
 	/** Returns whether or not the supplied class can be used in the current editor context. */
 	bool IsAllowedClass(const UClass* InClass) const;
 
+	/** Returns whether or not the class referenced by the supplied class path can be used in the current editor context. */
+	bool IsAllowedClassPath(const FTopLevelAssetPath& InClassPath) const;
+
 	/** Returns whether or not the supplied niagara type definition can be used in the current editor context. */
 	bool IsAllowedTypeDefinition(const FNiagaraTypeDefinition& InTypeDefinition) const;
 
@@ -460,6 +467,7 @@ private:
 	TArray<FNiagaraParameterPanelSectionStorage> SystemParameterPanelSectionData;
 
 	FOnIsClassAllowed OnIsClassAllowedDelegate;
+	FOnIsClassPathAllowed OnIsClassPathAllowedDelegate;
 
 public:
 	bool IsShowInstructionsCount() const;
diff --git a/Engine/Plugins/Runtime/GeometryCache/Source/GeometryCacheEd/Private/GeometryCacheEdModule.cpp b/Engine/Plugins/Runtime/GeometryCache/Source/GeometryCacheEd/Private/GeometryCacheEdModule.cpp
--- a/Engine/Plugins/Runtime/GeometryCache/Source/GeometryCacheEd/Private/GeometryCacheEdModule.cpp
+++ b/Engine/Plugins/Runtime/GeometryCache/Source/GeometryCacheEd/Private/GeometryCacheEdModule.cpp
@@ -34,6 +34,7 @@ void FGeometryCacheEdModule::StartupModule()
 	NiagaraEditorModule.RegisterRendererCreationInfo(FNiagaraRendererCreationInfo(
 		UNiagaraGeometryCacheRendererProperties::StaticClass()->GetDisplayNameText(),
 		FText::FromString(UNiagaraGeometryCacheRendererProperties::StaticClass()->GetDescription()),
+		UNiagaraGeometryCacheRendererProperties::StaticClass()->GetClassPathName(),
 		FNiagaraRendererCreationInfo::FRendererFactory::CreateLambda([](UObject* OuterEmitter)
 		{
 			UNiagaraGeometryCacheRendererProperties* NewRenderer = NewObject<UNiagaraGeometryCacheRendererProperties>(OuterEmitter, NAME_None, RF_Transactional);
