Fixed DI categories in niagara not behaving as desired by removing categories again

diff --git a/Engine/Plugins/FX/Niagara/Source/NiagaraEditor/Private/Customizations/NiagaraComponentDetails.cpp b/Engine/Plugins/FX/Niagara/Source/NiagaraEditor/Private/Customizations/NiagaraComponentDetails.cpp
--- a/Engine/Plugins/FX/Niagara/Source/NiagaraEditor/Private/Customizations/NiagaraComponentDetails.cpp
+++ b/Engine/Plugins/FX/Niagara/Source/NiagaraEditor/Private/Customizations/NiagaraComponentDetails.cpp
@@ -344,7 +344,7 @@ void AddRowForUserParameter(TSharedPtr<FNiagaraParameterProxy> ParameterProxy, I
 		FAddPropertyParams Params = FAddPropertyParams()
 			.UniqueId(Parameter.GetName())
 			.AllowChildren(true)
-			.CreateCategoryNodes(true);
+			.CreateCategoryNodes(false);
 	
 		Row = ChildrenBuilder.AddExternalObjectProperty(Objects, NAME_None, Params);
 
