Fixing crash that occurs when stack entry has been finalized but we still invoke child data

diff --git a/Engine/Plugins/FX/Niagara/Source/NiagaraEditorWidgets/Private/SNiagaraOverviewStack.cpp b/Engine/Plugins/FX/Niagara/Source/NiagaraEditorWidgets/Private/SNiagaraOverviewStack.cpp
--- a/Engine/Plugins/FX/Niagara/Source/NiagaraEditorWidgets/Private/SNiagaraOverviewStack.cpp
+++ b/Engine/Plugins/FX/Niagara/Source/NiagaraEditorWidgets/Private/SNiagaraOverviewStack.cpp
@@ -709,41 +709,49 @@ void SNiagaraOverviewStack::Tick(const FGeometry& AllottedGeometry, const double
 
 EVisibility SNiagaraOverviewStack::GetIssueIconVisibility(UNiagaraStackEntry* StackEntry)
 {
-	bool bDisplayIcon;
-	if (StackEntry->IsA<UNiagaraStackItemGroup>() && StackEntry->GetCanExpandInOverview() && StackEntry->GetIsExpandedInOverview())
+	if (StackEntry != nullptr && StackEntry->IsFinalized() == false)
 	{
-		// If the entry is a group and it can expand and it is expanded, we only want to show the stack issue icon if the group itself has issues.
-		bDisplayIcon = StackEntry->GetIssues().Num() > 0;
-	}
-	else
-	{
-		bDisplayIcon = StackEntry->HasIssuesOrAnyChildHasIssues();
+		bool bDisplayIcon;
+		if (StackEntry->IsA<UNiagaraStackItemGroup>() && StackEntry->GetCanExpandInOverview() && StackEntry->GetIsExpandedInOverview())
+		{
+			// If the entry is a group and it can expand and it is expanded, we only want to show the stack issue icon if the group itself has issues.
+			bDisplayIcon = StackEntry->GetIssues().Num() > 0;
+		}
+		else
+		{
+			bDisplayIcon = StackEntry->HasIssuesOrAnyChildHasIssues();
+		}
+		return bDisplayIcon ? EVisibility::Visible : EVisibility::Collapsed;
 	}
-	return bDisplayIcon ? EVisibility::Visible : EVisibility::Collapsed;
+	return EVisibility::Collapsed;
 }
 
 const FSlateBrush* SNiagaraOverviewStack::GetUsageIcon(UNiagaraStackEntry* StackEntry)
 {
-	bool bRead = false;
-	bool bWrite = false;
-
-	StackEntry->GetRecursiveUsages(bRead, bWrite);
-	if (bRead && bWrite)
+	if (StackEntry != nullptr && StackEntry->IsFinalized() == false)
 	{
-		return FNiagaraEditorWidgetsStyle::Get().GetBrush("NiagaraEditor.Stack.ReadWriteIcon");
-	}
-	else if (bRead)
-	{
-		return FNiagaraEditorWidgetsStyle::Get().GetBrush("NiagaraEditor.Stack.ReadIcon");
-	}
-	else if (bWrite)
-	{
-		return FNiagaraEditorWidgetsStyle::Get().GetBrush("NiagaraEditor.Stack.WriteIcon");
-	}
-	else
-	{
-		return FAppStyle::Get().GetBrush("WhiteBrush");
+		bool bRead = false;
+		bool bWrite = false;
+
+		StackEntry->GetRecursiveUsages(bRead, bWrite);
+		if (bRead && bWrite)
+		{
+			return FNiagaraEditorWidgetsStyle::Get().GetBrush("NiagaraEditor.Stack.ReadWriteIcon");
+		}
+		else if (bRead)
+		{
+			return FNiagaraEditorWidgetsStyle::Get().GetBrush("NiagaraEditor.Stack.ReadIcon");
+		}
+		else if (bWrite)
+		{
+			return FNiagaraEditorWidgetsStyle::Get().GetBrush("NiagaraEditor.Stack.WriteIcon");
+		}
+		else
+		{
+			return FAppStyle::Get().GetBrush("WhiteBrush");
+		}
 	}
+	return FAppStyle::Get().GetBrush("WhiteBrush");
 }
 
 
@@ -765,34 +773,42 @@ int32 SNiagaraOverviewStack::GetUsageIconHeight(UNiagaraStackEntry* StackEntry)
 
 EVisibility SNiagaraOverviewStack::GetUsageIconVisibility(UNiagaraStackEntry* StackEntry)
 {
-	bool bDisplayIcon = StackEntry->HasUsagesOrAnyChildHasUsages();
+	if (StackEntry != nullptr && StackEntry->IsFinalized() == false)
+	{
+		bool bDisplayIcon = StackEntry->HasUsagesOrAnyChildHasUsages();
 	
-	return bDisplayIcon ? EVisibility::Visible : EVisibility::Collapsed;
+		return bDisplayIcon ? EVisibility::Visible : EVisibility::Collapsed;
+	}
+	return EVisibility::Collapsed;
 }
 
 
 FText SNiagaraOverviewStack::GetUsageTooltip(UNiagaraStackEntry* StackEntry)
 {
-	bool bRead = false;
-	bool bWrite = false;
-
-	StackEntry->GetRecursiveUsages(bRead, bWrite);
-	if (bRead && bWrite)
+	if (StackEntry != nullptr && StackEntry->IsFinalized() == false)
 	{
-		return LOCTEXT("UsageTooltipRW", "Reads and writes selected parameter in Parameters panel.");
-	}
-	else if (bRead)
-	{
-		return LOCTEXT("UsageTooltipR", "Reads selected parameter in Parameters panel.");
-	}
-	else if (bWrite)
-	{
-		return LOCTEXT("UsageTooltipW", "Writes selected parameter in Parameters panel.");
-	}
-	else
-	{
-		return LOCTEXT("UsageTooltipNone", "Does not touch selected parameter in Parameters panel.");
+		bool bRead = false;
+		bool bWrite = false;
+
+		StackEntry->GetRecursiveUsages(bRead, bWrite);
+		if (bRead && bWrite)
+		{
+			return LOCTEXT("UsageTooltipRW", "Reads and writes selected parameter in Parameters panel.");
+		}
+		else if (bRead)
+		{
+			return LOCTEXT("UsageTooltipR", "Reads selected parameter in Parameters panel.");
+		}
+		else if (bWrite)
+		{
+			return LOCTEXT("UsageTooltipW", "Writes selected parameter in Parameters panel.");
+		}
+		else
+		{
+			return LOCTEXT("UsageTooltipNone", "Does not touch selected parameter in Parameters panel.");
+		}
 	}
+	return FText::GetEmpty();
 }
 
 void SNiagaraOverviewStack::AddEntriesRecursive(UNiagaraStackEntry& EntryToAdd, TArray<UNiagaraStackEntry*>& EntryList, const TArray<UClass*>& AcceptableClasses, TArray<UNiagaraStackEntry*> ParentChain)
