- Inclusive word pass

#rnx

diff --git a/Engine/Plugins/FX/Niagara/Source/NiagaraEditor/Private/Customizations/NiagaraTypeCustomizations.cpp b/Engine/Plugins/FX/Niagara/Source/NiagaraEditor/Private/Customizations/NiagaraTypeCustomizations.cpp
--- a/Engine/Plugins/FX/Niagara/Source/NiagaraEditor/Private/Customizations/NiagaraTypeCustomizations.cpp
+++ b/Engine/Plugins/FX/Niagara/Source/NiagaraEditor/Private/Customizations/NiagaraTypeCustomizations.cpp
@@ -1980,14 +1980,14 @@ void FNiagaraSystemScalabilityOverrideCustomization::CustomizeHeader(TSharedRef<
 
 void FNiagaraSystemScalabilityOverrideCustomization::CustomizeChildren(TSharedRef<IPropertyHandle> PropertyHandle, IDetailChildrenBuilder& ChildBuilder, IPropertyTypeCustomizationUtils& CustomizationUtils)
 {
-	auto AddOverrideProperties = [&](FName MasterPropertyName, TArray<FName> OverrideProperties)
+	auto AddOverrideProperties = [&](FName MainPropertyName, TArray<FName> OverrideProperties)
 	{
-		if (TSharedPtr<IPropertyHandle> MasterProperty = PropertyHandle->GetChildHandle(MasterPropertyName))
+		if (TSharedPtr<IPropertyHandle> MainProperty = PropertyHandle->GetChildHandle(MainPropertyName))
 		{
-			TAttribute<EVisibility> VisibilityAttribute = TAttribute<EVisibility>::CreateLambda([MasterProperty]
+			TAttribute<EVisibility> VisibilityAttribute = TAttribute<EVisibility>::CreateLambda([MainProperty]
 			{
-				bool bMasterProp = false;
-				if(MasterProperty->GetValue(bMasterProp) == FPropertyAccess::Success && bMasterProp)
+				bool bMainProp = false;
+				if(MainProperty->GetValue(bMainProp) == FPropertyAccess::Success && bMainProp)
 				{
 					return EVisibility::Visible;
 				}
@@ -1995,11 +1995,11 @@ void FNiagaraSystemScalabilityOverrideCustomization::CustomizeChildren(TSharedRe
 				return EVisibility::Collapsed;
 			});
 			
-			bool bMasterProp;
-			if (MasterProperty->GetValue(bMasterProp) == FPropertyAccess::Success)
+			bool bMainProp;
+			if (MainProperty->GetValue(bMainProp) == FPropertyAccess::Success)
 			{
-				IDetailGroup& PropGroup = ChildBuilder.AddGroup(MasterPropertyName, MasterProperty->GetPropertyDisplayName());
-				PropGroup.HeaderProperty(MasterProperty.ToSharedRef());
+				IDetailGroup& PropGroup = ChildBuilder.AddGroup(MainPropertyName, MainProperty->GetPropertyDisplayName());
+				PropGroup.HeaderProperty(MainProperty.ToSharedRef());
 				for (FName OverridePropName : OverrideProperties)
 				{
 					if (TSharedPtr<IPropertyHandle> OverrideProp = PropertyHandle->GetChildHandle(OverridePropName))
@@ -2011,7 +2011,7 @@ void FNiagaraSystemScalabilityOverrideCustomization::CustomizeChildren(TSharedRe
 			}
 			else
 			{
-				ChildBuilder.AddProperty(MasterProperty.ToSharedRef());
+				ChildBuilder.AddProperty(MainProperty.ToSharedRef());
 			}
 		}
 	}; 
diff --git a/Engine/Plugins/FX/Niagara/Source/NiagaraEditor/Private/NiagaraHlslTranslator.cpp b/Engine/Plugins/FX/Niagara/Source/NiagaraEditor/Private/NiagaraHlslTranslator.cpp
--- a/Engine/Plugins/FX/Niagara/Source/NiagaraEditor/Private/NiagaraHlslTranslator.cpp
+++ b/Engine/Plugins/FX/Niagara/Source/NiagaraEditor/Private/NiagaraHlslTranslator.cpp
@@ -7236,7 +7236,7 @@ void FHlslNiagaraTranslator::CullMapSetInputPin(UEdGraphPin* InputPin)
 	}
 }
 
-// From a valid list of namespaces, resolve any aliased tokens and promote namespaced variables without a master namespace to the input parameter map instance namespace
+// From a valid list of namespaces, resolve any aliased tokens and promote namespaced variables without a main namespace to the input parameter map instance namespace
 void FHlslNiagaraTranslator::FinalResolveNamespacedTokens(const FString& ParameterMapInstanceNamespace, TArray<FString>& Tokens, TArray<FString>& ValidChildNamespaces, FNiagaraParameterMapHistoryBuilder& Builder, TArray<FNiagaraVariable>& UniqueParameterMapEntriesAliasesIntact, TArray<FNiagaraVariable>& UniqueParameterMapEntries, int32 ParamMapHistoryIdx)
 {
 	for (int32 i = 0; i < Tokens.Num(); i++)
@@ -7248,7 +7248,7 @@ void FHlslNiagaraTranslator::FinalResolveNamespacedTokens(const FString& Paramet
 				FNiagaraVariable Var;
 
 				// There are two possible paths here, one where we're using the namespace as-is from the valid list and one where we've already
-				// prepended with the master parameter map instance namespace but may not have resolved any internal aliases yet.
+				// prepended with the main parameter map instance namespace but may not have resolved any internal aliases yet.
 				if (Tokens[i].StartsWith(ValidNamespace, ESearchCase::CaseSensitive))
 				{
 					FNiagaraVariable TempVar(FNiagaraTypeDefinition::GetFloatDef(), *Tokens[i]); // Declare a dummy name here that we will convert aliases for and use later
diff --git a/Engine/Plugins/FX/Niagara/Source/NiagaraEditor/Private/NiagaraNodeAssignment.cpp b/Engine/Plugins/FX/Niagara/Source/NiagaraEditor/Private/NiagaraNodeAssignment.cpp
--- a/Engine/Plugins/FX/Niagara/Source/NiagaraEditor/Private/NiagaraNodeAssignment.cpp
+++ b/Engine/Plugins/FX/Niagara/Source/NiagaraEditor/Private/NiagaraNodeAssignment.cpp
@@ -252,7 +252,7 @@ void UNiagaraNodeAssignment::BuildParameterMapHistory(FNiagaraParameterMapHistor
 	Super::BuildParameterMapHistory(OutHistory, bRecursive, bFilterForCompilation);
 }
 
-void UNiagaraNodeAssignment::GatherExternalDependencyData(ENiagaraScriptUsage InMasterUsage, const FGuid& InMasterUsageId, TArray<FNiagaraCompileHash>& InReferencedCompileHashes, TArray<FString>& InReferencedObjs) const
+void UNiagaraNodeAssignment::GatherExternalDependencyData(ENiagaraScriptUsage InUsage, const FGuid& InUsageId, TArray<FNiagaraCompileHash>& InReferencedCompileHashes, TArray<FString>& InReferencedObjs) const
 {
 	// Assignment nodes own their function graphs and therefore have no external dependencies so we override the default function behavior here to avoid 
 	// adding additional non-deterministic guids to the compile id generation which can invalid the DDC for compiled scripts, especially during emitter merging.
diff --git a/Engine/Plugins/FX/Niagara/Source/NiagaraEditor/Private/NiagaraNodeCustomHlsl.cpp b/Engine/Plugins/FX/Niagara/Source/NiagaraEditor/Private/NiagaraNodeCustomHlsl.cpp
--- a/Engine/Plugins/FX/Niagara/Source/NiagaraEditor/Private/NiagaraNodeCustomHlsl.cpp
+++ b/Engine/Plugins/FX/Niagara/Source/NiagaraEditor/Private/NiagaraNodeCustomHlsl.cpp
@@ -476,7 +476,7 @@ FNiagaraCompileHash GetFileContentHash(const FString& FileContents)
 	return FNiagaraCompileHash(DataHash);
 }
 
-void UNiagaraNodeCustomHlsl::GatherExternalDependencyData(ENiagaraScriptUsage InMasterUsage, const FGuid& InMasterUsageId, TArray<FNiagaraCompileHash>& InReferencedCompileHashes, TArray<FString>& InReferencedObjs) const
+void UNiagaraNodeCustomHlsl::GatherExternalDependencyData(ENiagaraScriptUsage InUsage, const FGuid& InUsageId, TArray<FNiagaraCompileHash>& InReferencedCompileHashes, TArray<FString>& InReferencedObjs) const
 {
 	for (const auto& [IncludePath] : AbsoluteIncludeFilePaths)
 	{
diff --git a/Engine/Plugins/FX/Niagara/Source/NiagaraEditor/Private/NiagaraNodeEmitter.cpp b/Engine/Plugins/FX/Niagara/Source/NiagaraEditor/Private/NiagaraNodeEmitter.cpp
--- a/Engine/Plugins/FX/Niagara/Source/NiagaraEditor/Private/NiagaraNodeEmitter.cpp
+++ b/Engine/Plugins/FX/Niagara/Source/NiagaraEditor/Private/NiagaraNodeEmitter.cpp
@@ -482,14 +482,14 @@ void UNiagaraNodeEmitter::Compile(FHlslNiagaraTranslator *Translator, TArray<int
 	}
 }
 
-void UNiagaraNodeEmitter::GatherExternalDependencyData(ENiagaraScriptUsage InMasterUsage, const FGuid& InMasterUsageId, TArray<FNiagaraCompileHash>& InReferencedCompileHashes, TArray<FString>& InReferencedObjs) const
+void UNiagaraNodeEmitter::GatherExternalDependencyData(ENiagaraScriptUsage InUsage, const FGuid& InUsageId, TArray<FNiagaraCompileHash>& InReferencedCompileHashes, TArray<FString>& InReferencedObjs) const
 {
 	UNiagaraGraph* CalledGraph = GetCalledGraph();
 
 	if (CalledGraph && IsNodeEnabled()) // Skip if disabled
 	{
 		CalledGraph->RebuildCachedCompileIds();
-		ENiagaraScriptUsage TargetUsage = InMasterUsage == ENiagaraScriptUsage::SystemSpawnScript ? ENiagaraScriptUsage::EmitterSpawnScript : ENiagaraScriptUsage::EmitterUpdateScript;
+		ENiagaraScriptUsage TargetUsage = InUsage == ENiagaraScriptUsage::SystemSpawnScript ? ENiagaraScriptUsage::EmitterSpawnScript : ENiagaraScriptUsage::EmitterUpdateScript;
 		FNiagaraCompileHash Hash = CalledGraph->GetCompileDataHash(TargetUsage, FGuid(0,0,0,0));
 		InReferencedCompileHashes.AddUnique(Hash);
 		InReferencedObjs.Add(CalledGraph->GetPathName());
diff --git a/Engine/Plugins/FX/Niagara/Source/NiagaraEditor/Private/NiagaraNodeEmitter.h b/Engine/Plugins/FX/Niagara/Source/NiagaraEditor/Private/NiagaraNodeEmitter.h
--- a/Engine/Plugins/FX/Niagara/Source/NiagaraEditor/Private/NiagaraNodeEmitter.h
+++ b/Engine/Plugins/FX/Niagara/Source/NiagaraEditor/Private/NiagaraNodeEmitter.h
@@ -60,7 +60,7 @@ public:
 	UNiagaraGraph* GetCalledGraph() const;
 	
 	virtual void Compile(FHlslNiagaraTranslator *Translator, TArray<int32>& Outputs) override;
-	virtual void GatherExternalDependencyData(ENiagaraScriptUsage InMasterUsage, const FGuid& InMasterUsageId, TArray<FNiagaraCompileHash>& InReferencedCompileHashes, TArray<FString>& InReferencedObjs) const override;
+	virtual void GatherExternalDependencyData(ENiagaraScriptUsage InUsage, const FGuid& InUsageId, TArray<FNiagaraCompileHash>& InReferencedCompileHashes, TArray<FString>& InReferencedObjs) const override;
 
 	void SetCachedVariablesForCompilation(const FName& InUniqueName, UNiagaraGraph* InGraph, UNiagaraScriptSourceBase* InSource);
 
diff --git a/Engine/Plugins/FX/Niagara/Source/NiagaraEditor/Private/NiagaraNodeFunctionCall.cpp b/Engine/Plugins/FX/Niagara/Source/NiagaraEditor/Private/NiagaraNodeFunctionCall.cpp
--- a/Engine/Plugins/FX/Niagara/Source/NiagaraEditor/Private/NiagaraNodeFunctionCall.cpp
+++ b/Engine/Plugins/FX/Niagara/Source/NiagaraEditor/Private/NiagaraNodeFunctionCall.cpp
@@ -1021,7 +1021,7 @@ bool UNiagaraNodeFunctionCall::RefreshFromExternalChanges()
 	}
 }
 
-void UNiagaraNodeFunctionCall::GatherExternalDependencyData(ENiagaraScriptUsage InMasterUsage, const FGuid& InMasterUsageId, TArray<FNiagaraCompileHash>& InReferencedCompileHashes, TArray<FString>& InReferencedObjs) const
+void UNiagaraNodeFunctionCall::GatherExternalDependencyData(ENiagaraScriptUsage InUsage, const FGuid& InUsageId, TArray<FNiagaraCompileHash>& InReferencedCompileHashes, TArray<FString>& InReferencedObjs) const
 {
 	if (HasValidScriptAndGraph())
 	{
diff --git a/Engine/Plugins/FX/Niagara/Source/NiagaraEditor/Private/NiagaraNodeParameterMapGet.cpp b/Engine/Plugins/FX/Niagara/Source/NiagaraEditor/Private/NiagaraNodeParameterMapGet.cpp
--- a/Engine/Plugins/FX/Niagara/Source/NiagaraEditor/Private/NiagaraNodeParameterMapGet.cpp
+++ b/Engine/Plugins/FX/Niagara/Source/NiagaraEditor/Private/NiagaraNodeParameterMapGet.cpp
@@ -578,7 +578,7 @@ bool UNiagaraNodeParameterMapGet::CommitEditablePinName(const FText& InName, UEd
 	return false;
 }
 
-void UNiagaraNodeParameterMapGet::GatherExternalDependencyData(ENiagaraScriptUsage InMasterUsage, const FGuid& InMasterUsageId, TArray<FNiagaraCompileHash>& InReferencedCompileHashes, TArray<FString>& InReferencedObjs) const
+void UNiagaraNodeParameterMapGet::GatherExternalDependencyData(ENiagaraScriptUsage InUsage, const FGuid& InUsageId, TArray<FNiagaraCompileHash>& InReferencedCompileHashes, TArray<FString>& InReferencedObjs) const
 {
 	// If we are referencing any parameter collections, we need to register them here... might want to speeed this up in the future 
 	// by caching any parameter collections locally.
diff --git a/Engine/Plugins/FX/Niagara/Source/NiagaraEditor/Private/NiagaraNodeParameterMapGet.h b/Engine/Plugins/FX/Niagara/Source/NiagaraEditor/Private/NiagaraNodeParameterMapGet.h
--- a/Engine/Plugins/FX/Niagara/Source/NiagaraEditor/Private/NiagaraNodeParameterMapGet.h
+++ b/Engine/Plugins/FX/Niagara/Source/NiagaraEditor/Private/NiagaraNodeParameterMapGet.h
@@ -44,7 +44,7 @@ public:
 	
 	virtual void PostLoad() override;
 
-	void GatherExternalDependencyData(ENiagaraScriptUsage InMasterUsage, const FGuid& InMasterUsageId, TArray<FNiagaraCompileHash>& InReferencedCompileHashes, TArray<FString>& InReferencedObjs) const override;
+	void GatherExternalDependencyData(ENiagaraScriptUsage InUsage, const FGuid& InUsageId, TArray<FNiagaraCompileHash>& InReferencedCompileHashes, TArray<FString>& InReferencedObjs) const override;
 	virtual void GetPinHoverText(const UEdGraphPin& Pin, FString& HoverTextOut) const override;
 
 	virtual FName GetNewPinDefaultNamespace() const { return PARAM_MAP_MODULE_STR; }
diff --git a/Engine/Plugins/FX/Niagara/Source/NiagaraEditor/Private/NiagaraNodeReroute.h b/Engine/Plugins/FX/Niagara/Source/NiagaraEditor/Private/NiagaraNodeReroute.h
--- a/Engine/Plugins/FX/Niagara/Source/NiagaraEditor/Private/NiagaraNodeReroute.h
+++ b/Engine/Plugins/FX/Niagara/Source/NiagaraEditor/Private/NiagaraNodeReroute.h
@@ -39,7 +39,7 @@ public:
 	
 	/** Trace to an output pin that is not a reroute node output pin. If the reroute nodes resulted in a dead end and no output pin was found return nullptr.*/
 	virtual UEdGraphPin* GetTracedOutputPin(UEdGraphPin* LocallyOwnedOutputPin, bool bFilterForCompilation, TArray<const UNiagaraNode*>* OutNodesVisitedDuringTrace = nullptr) const override;
-	virtual UEdGraphPin* GetPassThroughPin(const UEdGraphPin* LocallyOwnedOutputPin, ENiagaraScriptUsage MasterUsage) const override { return GetPassThroughPin(LocallyOwnedOutputPin); }
+	virtual UEdGraphPin* GetPassThroughPin(const UEdGraphPin* LocallyOwnedOutputPin, ENiagaraScriptUsage InUsage) const override { return GetPassThroughPin(LocallyOwnedOutputPin); }
 	//~ End UNiagaraNode Interface
 
 	void PropagatePinType();
diff --git a/Engine/Plugins/FX/Niagara/Source/NiagaraEditor/Private/NiagaraNodeSimTargetSelector.cpp b/Engine/Plugins/FX/Niagara/Source/NiagaraEditor/Private/NiagaraNodeSimTargetSelector.cpp
--- a/Engine/Plugins/FX/Niagara/Source/NiagaraEditor/Private/NiagaraNodeSimTargetSelector.cpp
+++ b/Engine/Plugins/FX/Niagara/Source/NiagaraEditor/Private/NiagaraNodeSimTargetSelector.cpp
@@ -89,7 +89,7 @@ void UNiagaraNodeSimTargetSelector::Compile(class FHlslNiagaraTranslator* Transl
 	Outputs[OutputPins.Num() - 1] = INDEX_NONE;
 }
 
-UEdGraphPin* UNiagaraNodeSimTargetSelector::GetPassThroughPin(const UEdGraphPin* LocallyOwnedOutputPin, ENiagaraScriptUsage MasterUsage) const 
+UEdGraphPin* UNiagaraNodeSimTargetSelector::GetPassThroughPin(const UEdGraphPin* LocallyOwnedOutputPin, ENiagaraScriptUsage InUsage) const
 {
 	return nullptr;
 }
diff --git a/Engine/Plugins/FX/Niagara/Source/NiagaraEditor/Private/NiagaraNodeSimTargetSelector.h b/Engine/Plugins/FX/Niagara/Source/NiagaraEditor/Private/NiagaraNodeSimTargetSelector.h
--- a/Engine/Plugins/FX/Niagara/Source/NiagaraEditor/Private/NiagaraNodeSimTargetSelector.h
+++ b/Engine/Plugins/FX/Niagara/Source/NiagaraEditor/Private/NiagaraNodeSimTargetSelector.h
@@ -20,7 +20,7 @@ public:
 	//~ Begin UNiagaraNode Interface
 	virtual void Compile(class FHlslNiagaraTranslator* Translator, TArray<int32>& Outputs) override;
 	virtual void BuildParameterMapHistory(FNiagaraParameterMapHistoryBuilder& OutHistory, bool bRecursive = true, bool bFilterForCompilation = true) const override;
-	virtual UEdGraphPin* GetPassThroughPin(const UEdGraphPin* LocallyOwnedOutputPin, ENiagaraScriptUsage MasterUsage) const override;
+	virtual UEdGraphPin* GetPassThroughPin(const UEdGraphPin* LocallyOwnedOutputPin, ENiagaraScriptUsage InUsage) const override;
 	//~ End UNiagaraNode Interface
 
 protected:
diff --git a/Engine/Plugins/FX/Niagara/Source/NiagaraEditor/Private/NiagaraNodeStaticSwitch.cpp b/Engine/Plugins/FX/Niagara/Source/NiagaraEditor/Private/NiagaraNodeStaticSwitch.cpp
--- a/Engine/Plugins/FX/Niagara/Source/NiagaraEditor/Private/NiagaraNodeStaticSwitch.cpp
+++ b/Engine/Plugins/FX/Niagara/Source/NiagaraEditor/Private/NiagaraNodeStaticSwitch.cpp
@@ -594,16 +594,16 @@ ENiagaraNumericOutputTypeSelectionMode UNiagaraNodeStaticSwitch::GetNumericOutpu
 
 void UNiagaraNodeStaticSwitch::ResolveNumerics(const UEdGraphSchema_Niagara* Schema, bool bSetInline, TMap<TPair<FGuid, UEdGraphNode*>, FNiagaraTypeDefinition>* PinCache)
 {	
-	FPinCollectorArray MasterInputPins;
-	GetInputPins(MasterInputPins);
-	FPinCollectorArray MasterOutputPins;
-	GetOutputPins(MasterOutputPins);
+	FPinCollectorArray MainInputPins;
+	GetInputPins(MainInputPins);
+	FPinCollectorArray MainOutputPins;
+	GetOutputPins(MainOutputPins);
 
 	TArray<int32> OptionValues = GetOptionValues();
 	int32 VarIdx = 0;
-	for (int i = 0; i < MasterOutputPins.Num(); i++)
+	for (int i = 0; i < MainOutputPins.Num(); i++)
 	{
-		UEdGraphPin* OutPin = MasterOutputPins[i];
+		UEdGraphPin* OutPin = MainOutputPins[i];
 		if (IsAddPin(OutPin))
 		{
 			continue;
@@ -616,14 +616,14 @@ void UNiagaraNodeStaticSwitch::ResolveNumerics(const UEdGraphSchema_Niagara* Sch
 		for (int32 j = 0; j < OptionValues.Num(); j++)
 		{
 			const int TargetIdx = OutputVars.Num() * j + VarIdx;
-			if (MasterInputPins.IsValidIndex(TargetIdx))
+			if (MainInputPins.IsValidIndex(TargetIdx))
 			{
-				UEdGraphPin* InputPin = MasterInputPins[TargetIdx];
+				UEdGraphPin* InputPin = MainInputPins[TargetIdx];
 				InputPins.Add(InputPin);
 			}
 			else
 			{
-				UE_LOG(LogNiagaraEditor, Warning, TEXT("Invalid index on UNiagaraNodeStaticSwitch::ResolveNumerics %s, OptionIdx: %d VarIdx: %d MaxIndex: %d"), *GetPathName(), j, VarIdx, MasterInputPins.Num())
+				UE_LOG(LogNiagaraEditor, Warning, TEXT("Invalid index on UNiagaraNodeStaticSwitch::ResolveNumerics %s, OptionIdx: %d VarIdx: %d MaxIndex: %d"), *GetPathName(), j, VarIdx, MainInputPins.Num())
 			}
 		}
 
@@ -803,13 +803,13 @@ UEdGraphPin* UNiagaraNodeStaticSwitch::GetTracedOutputPin(UEdGraphPin* LocallyOw
 	return LocallyOwnedOutputPin;
 }
 
-UEdGraphPin* UNiagaraNodeStaticSwitch::GetPassThroughPin(const UEdGraphPin* LocallyOwnedOutputPin,	ENiagaraScriptUsage MasterUsage) const
+UEdGraphPin* UNiagaraNodeStaticSwitch::GetPassThroughPin(const UEdGraphPin* LocallyOwnedOutputPin,	ENiagaraScriptUsage InUsage) const
 {
 	if (IsValueSet)
 	{
 		return GetTracedOutputPin(const_cast<UEdGraphPin*>(LocallyOwnedOutputPin), true);
 	}
-	return Super::GetPassThroughPin(LocallyOwnedOutputPin, MasterUsage);
+	return Super::GetPassThroughPin(LocallyOwnedOutputPin, InUsage);
 }
 
 void UNiagaraNodeStaticSwitch::BuildParameterMapHistory(FNiagaraParameterMapHistoryBuilder& OutHistory, bool bRecursive /*= true*/, bool bFilterForCompilation /*= true*/) const
diff --git a/Engine/Plugins/FX/Niagara/Source/NiagaraEditor/Private/NiagaraNodeStaticSwitch.h b/Engine/Plugins/FX/Niagara/Source/NiagaraEditor/Private/NiagaraNodeStaticSwitch.h
--- a/Engine/Plugins/FX/Niagara/Source/NiagaraEditor/Private/NiagaraNodeStaticSwitch.h
+++ b/Engine/Plugins/FX/Niagara/Source/NiagaraEditor/Private/NiagaraNodeStaticSwitch.h
@@ -92,7 +92,7 @@ public:
 	virtual bool SubstituteCompiledPin(FHlslNiagaraTranslator* Translator, UEdGraphPin** LocallyOwnedPin) override;
 	virtual UEdGraphPin* GetTracedOutputPin(UEdGraphPin* LocallyOwnedOutputPin, bool bFilterForCompilation, TArray<const UNiagaraNode*>* OutNodesVisitedDuringTrace = nullptr) const override;
 	virtual UEdGraphPin* GetTracedOutputPin(UEdGraphPin* LocallyOwnedOutputPin, bool bRecursive, bool bFilterForCompilation, TArray<const UNiagaraNode*>* OutNodesVisitedDuringTrace = nullptr) const;
-	virtual UEdGraphPin* GetPassThroughPin(const UEdGraphPin* LocallyOwnedOutputPin, ENiagaraScriptUsage MasterUsage) const override;
+	virtual UEdGraphPin* GetPassThroughPin(const UEdGraphPin* LocallyOwnedOutputPin, ENiagaraScriptUsage InUsage) const override;
 	virtual bool AllowNiagaraTypeForAddPin(const FNiagaraTypeDefinition& InType) const override;
 	virtual void BuildParameterMapHistory(FNiagaraParameterMapHistoryBuilder& OutHistory, bool bRecursive = true, bool bFilterForCompilation = true) const override;
 	virtual void AddWidgetsToOutputBox(TSharedPtr<SVerticalBox> OutputBox) override;
diff --git a/Engine/Plugins/FX/Niagara/Source/NiagaraEditor/Private/NiagaraNodeUsageSelector.cpp b/Engine/Plugins/FX/Niagara/Source/NiagaraEditor/Private/NiagaraNodeUsageSelector.cpp
--- a/Engine/Plugins/FX/Niagara/Source/NiagaraEditor/Private/NiagaraNodeUsageSelector.cpp
+++ b/Engine/Plugins/FX/Niagara/Source/NiagaraEditor/Private/NiagaraNodeUsageSelector.cpp
@@ -331,12 +331,12 @@ void UNiagaraNodeUsageSelector::Compile(class FHlslNiagaraTranslator* Translator
 	}
 }
 
-UEdGraphPin* UNiagaraNodeUsageSelector::GetPassThroughPin(const UEdGraphPin* LocallyOwnedOutputPin, ENiagaraScriptUsage MasterUsage) const 
+UEdGraphPin* UNiagaraNodeUsageSelector::GetPassThroughPin(const UEdGraphPin* LocallyOwnedOutputPin, ENiagaraScriptUsage InUsage) const
 {
 	check(Pins.Contains(LocallyOwnedOutputPin) && LocallyOwnedOutputPin->Direction == EGPD_Output);
 
 	ENiagaraScriptGroup UsageGroup = ENiagaraScriptGroup::Max;
-	if (UNiagaraScript::ConvertUsageToGroup(MasterUsage, UsageGroup))
+	if (UNiagaraScript::ConvertUsageToGroup(InUsage, UsageGroup))
 	{
 		int32 VarIdx = 0;
 		for (int64 i = 0; i < (int64)ENiagaraScriptGroup::Max; i++)
diff --git a/Engine/Plugins/FX/Niagara/Source/NiagaraEditor/Private/NiagaraNodeUsageSelector.h b/Engine/Plugins/FX/Niagara/Source/NiagaraEditor/Private/NiagaraNodeUsageSelector.h
--- a/Engine/Plugins/FX/Niagara/Source/NiagaraEditor/Private/NiagaraNodeUsageSelector.h
+++ b/Engine/Plugins/FX/Niagara/Source/NiagaraEditor/Private/NiagaraNodeUsageSelector.h
@@ -43,7 +43,7 @@ public:
 	virtual bool AllowExternalPinTypeChanges(const UEdGraphPin* InGraphPin) const override;
 	virtual bool AllowNiagaraTypeForPinTypeChange(const FNiagaraTypeDefinition& InType, UEdGraphPin* Pin) const override;
 	virtual bool OnNewPinTypeRequested(UEdGraphPin* PinToChange, FNiagaraTypeDefinition NewType) override;
-	virtual UEdGraphPin* GetPassThroughPin(const UEdGraphPin* LocallyOwnedOutputPin, ENiagaraScriptUsage MasterUsage) const override; 
+	virtual UEdGraphPin* GetPassThroughPin(const UEdGraphPin* LocallyOwnedOutputPin, ENiagaraScriptUsage InUsage) const override;
 	virtual void AppendFunctionAliasForContext(const FNiagaraGraphFunctionAliasContext& InFunctionAliasContext, FString& InOutFunctionAlias, bool& OutOnlyOncePerNodeType) override;
 	//~ End UNiagaraNode Interface
 
diff --git a/Engine/Plugins/FX/Niagara/Source/NiagaraEditor/Private/ViewModels/NiagaraParameterPanelViewModel.cpp b/Engine/Plugins/FX/Niagara/Source/NiagaraEditor/Private/ViewModels/NiagaraParameterPanelViewModel.cpp
--- a/Engine/Plugins/FX/Niagara/Source/NiagaraEditor/Private/ViewModels/NiagaraParameterPanelViewModel.cpp
+++ b/Engine/Plugins/FX/Niagara/Source/NiagaraEditor/Private/ViewModels/NiagaraParameterPanelViewModel.cpp
@@ -259,10 +259,10 @@ void INiagaraParameterPanelViewModel::RefreshFull(bool bDoCategoryExpansion) con
 {
 	INiagaraImmutableParameterPanelViewModel::RefreshFull(bDoCategoryExpansion);
 
-	TSharedPtr<INiagaraParameterPanelViewModel> MasterVM = MasterParameterPanelViewModel.Pin();
-	if (MasterVM.IsValid())
+	TSharedPtr<INiagaraParameterPanelViewModel> MainVM = MainParameterPanelViewModel.Pin();
+	if (MainVM.IsValid())
 	{
-		MasterVM->RefreshFull(bDoCategoryExpansion);
+		MainVM->RefreshFull(bDoCategoryExpansion);
 	}
 }
 
@@ -781,10 +781,10 @@ void INiagaraParameterPanelViewModel::SelectParameterItemByName(const FName Para
 
 	OnINiagaraParameterPanelViewModelSelectionChanged(FoundScriptVariable);
 
-	TSharedPtr<INiagaraParameterPanelViewModel> MasterVM = MasterParameterPanelViewModel.Pin();
-	if (MasterVM.IsValid())
+	TSharedPtr<INiagaraParameterPanelViewModel> MainVM = MainParameterPanelViewModel.Pin();
+	if (MainVM.IsValid())
 	{
-		MasterVM->SelectParameterItemByName(ParameterName, bRequestRename);
+		MainVM->SelectParameterItemByName(ParameterName, bRequestRename);
 	}
 }
 
@@ -1261,7 +1261,7 @@ void FNiagaraSystemToolkitParameterPanelViewModel::SetActiveSection(int32 InSect
 		if (ScratchDocumentVM.IsValid())
 		{
 			TWeakPtr<INiagaraParameterPanelViewModel> WeakVM = StaticCastSharedRef< INiagaraParameterPanelViewModel>(AsShared());
-			ScratchDocumentVM->GetParameterPanelViewModel()->SetMasterParameterPanelViewModel(WeakVM);
+			ScratchDocumentVM->GetParameterPanelViewModel()->SetMainParameterPanelViewModel(WeakVM);
 		}
 	}
 }
diff --git a/Engine/Plugins/FX/Niagara/Source/NiagaraEditor/Private/ViewModels/NiagaraScratchPadScriptViewModel.cpp b/Engine/Plugins/FX/Niagara/Source/NiagaraEditor/Private/ViewModels/NiagaraScratchPadScriptViewModel.cpp
--- a/Engine/Plugins/FX/Niagara/Source/NiagaraEditor/Private/ViewModels/NiagaraScratchPadScriptViewModel.cpp
+++ b/Engine/Plugins/FX/Niagara/Source/NiagaraEditor/Private/ViewModels/NiagaraScratchPadScriptViewModel.cpp
@@ -48,10 +48,10 @@ FNiagaraScratchPadScriptViewModel::~FNiagaraScratchPadScriptViewModel()
 	}
 }
 
-void FNiagaraScratchPadScriptViewModel::Initialize(UNiagaraScript* Script, UNiagaraScript* InEditScript, TWeakPtr<FNiagaraSystemViewModel> InMasterViewModel)
+void FNiagaraScratchPadScriptViewModel::Initialize(UNiagaraScript* Script, UNiagaraScript* InEditScript, TWeakPtr<FNiagaraSystemViewModel> InSystemViewModel)
 {
 	OriginalScript = Script;
-	SystemViewModel = InMasterViewModel;
+	SystemViewModel = InSystemViewModel;
 
 	// Copy over the old edited graph, as this initialize might be incoming from inheritance changes and we want to maintain
 	// the edits to the old graph we were working on.
diff --git a/Engine/Plugins/FX/Niagara/Source/NiagaraEditor/Private/ViewModels/NiagaraSystemViewModel.cpp b/Engine/Plugins/FX/Niagara/Source/NiagaraEditor/Private/ViewModels/NiagaraSystemViewModel.cpp
--- a/Engine/Plugins/FX/Niagara/Source/NiagaraEditor/Private/ViewModels/NiagaraSystemViewModel.cpp
+++ b/Engine/Plugins/FX/Niagara/Source/NiagaraEditor/Private/ViewModels/NiagaraSystemViewModel.cpp
@@ -1741,12 +1741,12 @@ void FNiagaraSystemViewModel::ResetEmitterHandleViewModelsAndTracks()
 	if (NiagaraSequence)
 	{
 		TGuardValue<bool> UpdateGuard(bResetingSequencerTracks, true);
-		TArray<UMovieSceneTrack*> MasterTracks = NiagaraSequence->GetMovieScene()->GetMasterTracks();
-		for (UMovieSceneTrack* MasterTrack : MasterTracks)
+		TArray<UMovieSceneTrack*> MainTracks = NiagaraSequence->GetMovieScene()->GetMasterTracks();
+		for (UMovieSceneTrack* MainTrack : MainTracks)
 		{
-			if (MasterTrack != nullptr)
+			if (MainTrack != nullptr)
 			{
-				NiagaraSequence->GetMovieScene()->RemoveMasterTrack(*MasterTrack);
+				NiagaraSequence->GetMovieScene()->RemoveMasterTrack(*MainTrack);
 			}
 		}
 	}
@@ -1856,12 +1856,12 @@ void FNiagaraSystemViewModel::RefreshSequencerTracks()
 
 	if (Sequencer.IsValid())
 	{
-		TArray<UMovieSceneTrack*> MasterTracks = NiagaraSequence->GetMovieScene()->GetMasterTracks();
-		for (UMovieSceneTrack* MasterTrack : MasterTracks)
+		TArray<UMovieSceneTrack*> MainTracks = NiagaraSequence->GetMovieScene()->GetMasterTracks();
+		for (UMovieSceneTrack* MainTrack : MainTracks)
 		{
-			if (MasterTrack != nullptr)
+			if (MainTrack != nullptr)
 			{
-				NiagaraSequence->GetMovieScene()->RemoveMasterTrack(*MasterTrack);
+				NiagaraSequence->GetMovieScene()->RemoveMasterTrack(*MainTrack);
 			}
 		}
 
@@ -2605,9 +2605,9 @@ void FNiagaraSystemViewModel::UpdateSequencerFromEmitterHandleSelection()
 	Sequencer->EmptySelection();
 	for (FGuid SelectedEmitterHandleId : SelectionViewModel->GetSelectedEmitterHandleIds())
 	{
-		for (UMovieSceneTrack* MasterTrack : NiagaraSequence->GetMovieScene()->GetMasterTracks())
+		for (UMovieSceneTrack* MainTrack : NiagaraSequence->GetMovieScene()->GetMasterTracks())
 		{
-			UMovieSceneNiagaraEmitterTrack* EmitterTrack = Cast<UMovieSceneNiagaraEmitterTrack>(MasterTrack);
+			UMovieSceneNiagaraEmitterTrack* EmitterTrack = Cast<UMovieSceneNiagaraEmitterTrack>(MainTrack);
 			if (EmitterTrack != nullptr && EmitterTrack->GetEmitterHandleViewModel()->GetId() == SelectedEmitterHandleId)
 			{
 				Sequencer->SelectTrack(EmitterTrack);
diff --git a/Engine/Plugins/FX/Niagara/Source/NiagaraEditor/Private/ViewModels/Stack/NiagaraStackScriptItemGroup.cpp b/Engine/Plugins/FX/Niagara/Source/NiagaraEditor/Private/ViewModels/Stack/NiagaraStackScriptItemGroup.cpp
--- a/Engine/Plugins/FX/Niagara/Source/NiagaraEditor/Private/ViewModels/Stack/NiagaraStackScriptItemGroup.cpp
+++ b/Engine/Plugins/FX/Niagara/Source/NiagaraEditor/Private/ViewModels/Stack/NiagaraStackScriptItemGroup.cpp
@@ -1239,7 +1239,7 @@ void UNiagaraStackScriptItemGroup::ChildRequestDeprecatedRecommendation(UNiagara
 	}
 	TargetChild->Copy(ClipboardContent);
 
-	// Step 2: Disable the old master so that end users can use it as reference and inheritance isn't wiped out.
+	// Step 2: Disable the old owner so that end users can use it as reference and inheritance isn't wiped out.
 	//TargetChild->Rename();
 	//TargetChild->GetModuleNode().SuggestName(TEXT("Deprecated Original " + TargetChild->GetModuleNode().GetFunctionName()), true);
 	TargetChild->OnRenamed(FText::Format(LOCTEXT("NewDeprecationName", "Deprecated Original {0}"), TargetChild->GetDisplayName()));
diff --git a/Engine/Plugins/FX/Niagara/Source/NiagaraEditor/Public/NiagaraNode.h b/Engine/Plugins/FX/Niagara/Source/NiagaraEditor/Public/NiagaraNode.h
--- a/Engine/Plugins/FX/Niagara/Source/NiagaraEditor/Public/NiagaraNode.h
+++ b/Engine/Plugins/FX/Niagara/Source/NiagaraEditor/Public/NiagaraNode.h
@@ -144,7 +144,7 @@ public:
 	virtual void BuildParameterMapHistory(FNiagaraParameterMapHistoryBuilder& OutHistory, bool bRecursive = true, bool bFilterForCompilation = true) const;
 	
 	/** Go through all the external dependencies of this node in isolation and add them to the reference id list.*/
-	virtual void GatherExternalDependencyData(ENiagaraScriptUsage InMasterUsage, const FGuid& InMasterUsageId, TArray<FNiagaraCompileHash>& InReferencedCompileHashes, TArray<FString>& InReferencedObjs) const {};
+	virtual void GatherExternalDependencyData(ENiagaraScriptUsage InUsage, const FGuid& InUsageId, TArray<FNiagaraCompileHash>& InReferencedCompileHashes, TArray<FString>& InReferencedObjs) const {};
 
 	/** Traces one of this node's output pins to its source output pin.
 	 *
@@ -164,7 +164,7 @@ public:
 	virtual bool SubstituteCompiledPin(FHlslNiagaraTranslator* Translator, UEdGraphPin** LocallyOwnedPin);
 
 	virtual UEdGraphPin* GetPassThroughPin(const UEdGraphPin* LocallyOwnedOutputPin) const override { return nullptr; }
-	virtual UEdGraphPin* GetPassThroughPin(const UEdGraphPin* LocallyOwnedOutputPin, ENiagaraScriptUsage MasterUsage) const { return nullptr; }
+	virtual UEdGraphPin* GetPassThroughPin(const UEdGraphPin* LocallyOwnedOutputPin, ENiagaraScriptUsage InUsage) const { return nullptr; }
 
 	/** Identify that this node has undergone changes that will require synchronization with a compiled script.*/
 	void MarkNodeRequiresSynchronization(FString Reason, bool bRaiseGraphNeedsRecompile);
diff --git a/Engine/Plugins/FX/Niagara/Source/NiagaraEditor/Public/NiagaraNodeAssignment.h b/Engine/Plugins/FX/Niagara/Source/NiagaraEditor/Public/NiagaraNodeAssignment.h
--- a/Engine/Plugins/FX/Niagara/Source/NiagaraEditor/Public/NiagaraNodeAssignment.h
+++ b/Engine/Plugins/FX/Niagara/Source/NiagaraEditor/Public/NiagaraNodeAssignment.h
@@ -42,7 +42,7 @@ public:
 	//~ UNiagaraNodeFunctionCall interface
 	virtual bool RefreshFromExternalChanges() override;
 	virtual void BuildParameterMapHistory(FNiagaraParameterMapHistoryBuilder& OutHistory, bool bRecursive = true, bool bFilterForCompilation = true) const override;
-	virtual void GatherExternalDependencyData(ENiagaraScriptUsage InMasterUsage, const FGuid& InMasterUsageId, TArray<FNiagaraCompileHash>& InReferencedCompileHashes, TArray<FString>& InReferencedObjs) const override;
+	virtual void GatherExternalDependencyData(ENiagaraScriptUsage InUsage, const FGuid& InUsageId, TArray<FNiagaraCompileHash>& InReferencedCompileHashes, TArray<FString>& InReferencedObjs) const override;
 
 	void AddParameter(FNiagaraVariable InVar, FString InDefaultValue);
 	void RemoveParameter(const FNiagaraVariable& InVar);
diff --git a/Engine/Plugins/FX/Niagara/Source/NiagaraEditor/Public/NiagaraNodeCustomHlsl.h b/Engine/Plugins/FX/Niagara/Source/NiagaraEditor/Public/NiagaraNodeCustomHlsl.h
--- a/Engine/Plugins/FX/Niagara/Source/NiagaraEditor/Public/NiagaraNodeCustomHlsl.h
+++ b/Engine/Plugins/FX/Niagara/Source/NiagaraEditor/Public/NiagaraNodeCustomHlsl.h
@@ -34,7 +34,7 @@ public:
 	bool GetTokens(TArray<FString>& OutTokens, bool IncludeComments = true, bool IncludeWhitespace = true) const;
 
 	virtual void BuildParameterMapHistory(FNiagaraParameterMapHistoryBuilder& OutHistory, bool bRecursive = true, bool bFilterForCompilation = true) const override;
-	virtual void GatherExternalDependencyData(ENiagaraScriptUsage InMasterUsage, const FGuid& InMasterUsageId, TArray<FNiagaraCompileHash>& InReferencedCompileHashes, TArray<FString>& InReferencedObjs) const override;
+	virtual void GatherExternalDependencyData(ENiagaraScriptUsage InUsage, const FGuid& InUsageId, TArray<FNiagaraCompileHash>& InReferencedCompileHashes, TArray<FString>& InReferencedObjs) const override;
 
 	// Replace items in the tokens array if they start with the src string or optionally src string and a namespace delimiter
 	static uint32 ReplaceExactMatchTokens(TArray<FString>& Tokens, FStringView SrcString, FStringView ReplaceString, bool bAllowNamespaceSeparation);
diff --git a/Engine/Plugins/FX/Niagara/Source/NiagaraEditor/Public/NiagaraNodeFunctionCall.h b/Engine/Plugins/FX/Niagara/Source/NiagaraEditor/Public/NiagaraNodeFunctionCall.h
--- a/Engine/Plugins/FX/Niagara/Source/NiagaraEditor/Public/NiagaraNodeFunctionCall.h
+++ b/Engine/Plugins/FX/Niagara/Source/NiagaraEditor/Public/NiagaraNodeFunctionCall.h
@@ -111,7 +111,7 @@ public:
 	virtual bool RefreshFromExternalChanges() override;
 	virtual ENiagaraNumericOutputTypeSelectionMode GetNumericOutputTypeSelectionMode() const override;
 	virtual bool CanAddToGraph(UNiagaraGraph* TargetGraph, FString& OutErrorMsg) const override;
-	virtual void GatherExternalDependencyData(ENiagaraScriptUsage InMasterUsage, const FGuid& InMasterUsageId, TArray<FNiagaraCompileHash>& InReferencedCompileHashes, TArray<FString>& InReferencedObjs) const override;
+	virtual void GatherExternalDependencyData(ENiagaraScriptUsage InUsage, const FGuid& InUsageId, TArray<FNiagaraCompileHash>& InReferencedCompileHashes, TArray<FString>& InReferencedObjs) const override;
 	virtual void UpdateCompileHashForNode(FSHA1& HashState) const override;
 	//End UNiagaraNode interface
 
diff --git a/Engine/Plugins/FX/Niagara/Source/NiagaraEditor/Public/ViewModels/NiagaraParameterPanelViewModel.h b/Engine/Plugins/FX/Niagara/Source/NiagaraEditor/Public/ViewModels/NiagaraParameterPanelViewModel.h
--- a/Engine/Plugins/FX/Niagara/Source/NiagaraEditor/Public/ViewModels/NiagaraParameterPanelViewModel.h
+++ b/Engine/Plugins/FX/Niagara/Source/NiagaraEditor/Public/ViewModels/NiagaraParameterPanelViewModel.h
@@ -279,7 +279,7 @@ public:
 	virtual FText GetTooltipForSection(FText& InSection) const;
 	void SetActiveSection(FText& InSection);
 	virtual void SetActiveSection(int32 InSection);
-	virtual void SetMasterParameterPanelViewModel(TWeakPtr<INiagaraParameterPanelViewModel> InMasterParameterPanelViewModel) { MasterParameterPanelViewModel = InMasterParameterPanelViewModel; }
+	virtual void SetMainParameterPanelViewModel(TWeakPtr<INiagaraParameterPanelViewModel> InMainParameterPanelViewModel) { MainParameterPanelViewModel = InMainParameterPanelViewModel; }
 
 
 	virtual bool IsVariableSelected(FNiagaraVariableBase& InVar) const { return false; };
@@ -289,7 +289,7 @@ public:
 	static bool CanMakeNewParameterOfType(const FNiagaraTypeDefinition& InType);
 	
 protected:
-	TWeakPtr<INiagaraParameterPanelViewModel> MasterParameterPanelViewModel;
+	TWeakPtr<INiagaraParameterPanelViewModel> MainParameterPanelViewModel;
 	FOnParameterPanelViewModelExternalSelectionChanged OnParameterPanelViewModelExternalSelectionChangedDelegate;
 	FOnSelectParameterItemByName OnSelectParameterItemByNameDelegate;
 	FOnNotifyParameterPendingRename OnNotifyParameterPendingRenameDelegate;
