Moved curve options from the curve category header row into its own custom row.
This fixes the disappearance due to lack of categories in DIs

diff --git a/Engine/Plugins/FX/Niagara/Source/NiagaraEditorWidgets/Private/DetailCustomizations/NiagaraDataInterfaceCurveDetails.cpp b/Engine/Plugins/FX/Niagara/Source/NiagaraEditorWidgets/Private/DetailCustomizations/NiagaraDataInterfaceCurveDetails.cpp
--- a/Engine/Plugins/FX/Niagara/Source/NiagaraEditorWidgets/Private/DetailCustomizations/NiagaraDataInterfaceCurveDetails.cpp
+++ b/Engine/Plugins/FX/Niagara/Source/NiagaraEditorWidgets/Private/DetailCustomizations/NiagaraDataInterfaceCurveDetails.cpp
@@ -1442,7 +1442,7 @@ void FNiagaraDataInterfaceCurveDetailsBase::CustomizeDetails(IDetailLayoutBuilde
 	}
 
 	IDetailCategoryBuilder& CurveCategory = DetailBuilder.EditCategory("Curve");
-	CurveCategory.HeaderContent(ToolBarBuilder.MakeWidget());
+	CurveCategory.AddCustomRow(NSLOCTEXT("NiagaraDataInterfaceCurveDetails", "CurveOptions", "Options")).WholeRowContent() [ ToolBarBuilder.MakeWidget() ];
 	CurveCategory.AddCustomRow(NSLOCTEXT("NiagaraDataInterfaceCurveDetails", "CurveFilterText", "Curve")).WholeRowContent() [ CurveDataInterfaceEditor.ToSharedRef() ];
 }
 
