Niagara - Fix a bug in my previous change to the update context which prevented solo system simulations from being reset when they should be.

diff --git a/Engine/Plugins/FX/Niagara/Source/Niagara/Private/NiagaraCommon.cpp b/Engine/Plugins/FX/Niagara/Source/Niagara/Private/NiagaraCommon.cpp
--- a/Engine/Plugins/FX/Niagara/Source/Niagara/Private/NiagaraCommon.cpp
+++ b/Engine/Plugins/FX/Niagara/Source/Niagara/Private/NiagaraCommon.cpp
@@ -144,8 +144,8 @@ void FNiagaraSystemUpdateContext::AddAll(bool bReInit)
 	{
 		UNiagaraComponent* Comp = *It;
 		check(Comp);
-
-		AddInternal(Comp, bReInit);
+		bool bAllowDestroySystemSim = true;
+		AddInternal(Comp, bReInit, bAllowDestroySystemSim);
 	}
 }
 
@@ -154,7 +154,8 @@ void FNiagaraSystemUpdateContext::AddSoloComponent(UNiagaraComponent* Component,
 	check(Component);
 	if (ensureMsgf(Component->GetForceSolo(), TEXT("A component must have a solo system simulation when used with an update context.")))
 	{
-		AddInternal(Component, bReInit);
+		bool bAllowDestroySystemSim = false;
+		AddInternal(Component, bReInit, bAllowDestroySystemSim);
 	}
 }
 
@@ -166,7 +167,8 @@ void FNiagaraSystemUpdateContext::Add(const UNiagaraSystem* System, bool bReInit
 		check(Comp);
 		if (Comp->GetAsset() == System)
 		{
-			AddInternal(Comp, bReInit);
+			bool bAllowDestroySystemSim = true;
+			AddInternal(Comp, bReInit, bAllowDestroySystemSim);
 		}
 	}
 }
@@ -181,7 +183,8 @@ void FNiagaraSystemUpdateContext::Add(const FVersionedNiagaraEmitter& Emitter, b
 		UNiagaraSystem* System = Comp->GetAsset();
 		if (System && System->UsesEmitter(Emitter))
 		{
-			AddInternal(Comp, bReInit);
+			bool bAllowDestroySystemSim = true;
+			AddInternal(Comp, bReInit, bAllowDestroySystemSim);
 		}
 	}
 }
@@ -195,7 +198,8 @@ void FNiagaraSystemUpdateContext::Add(const UNiagaraScript* Script, bool bReInit
 		UNiagaraSystem* System = Comp->GetAsset();
 		if (System && System->UsesScript(Script))
 		{
-			AddInternal(Comp, bReInit);
+			bool bAllowDestroySystemSim = true;
+			AddInternal(Comp, bReInit, bAllowDestroySystemSim);
 		}
 	}
 }
@@ -209,13 +213,14 @@ void FNiagaraSystemUpdateContext::Add(const UNiagaraParameterCollection* Collect
 		UNiagaraSystem* System = Comp->GetAsset();
 		if (System && System->UsesCollection(Collection))
 		{
-			AddInternal(Comp, bReInit);
+			bool bAllowDestroySystemSim = true;
+			AddInternal(Comp, bReInit, bAllowDestroySystemSim);
 		}
 	}
 }
 #endif
 
-void FNiagaraSystemUpdateContext::AddInternal(UNiagaraComponent* Comp, bool bReInit)
+void FNiagaraSystemUpdateContext::AddInternal(UNiagaraComponent* Comp, bool bReInit, bool bAllowDestroySystemSim)
 {
 	PreWork.ExecuteIfBound(Comp);
 
@@ -231,7 +236,7 @@ void FNiagaraSystemUpdateContext::AddInternal(UNiagaraComponent* Comp, bool bReI
 		}
 	}
 
-	if (Comp->GetForceSolo() == false)
+	if (bAllowDestroySystemSim)
 	{
 		if (bReInit || bDestroySystemSim)
 		{
diff --git a/Engine/Plugins/FX/Niagara/Source/Niagara/Public/NiagaraCommon.h b/Engine/Plugins/FX/Niagara/Source/Niagara/Public/NiagaraCommon.h
--- a/Engine/Plugins/FX/Niagara/Source/Niagara/Public/NiagaraCommon.h
+++ b/Engine/Plugins/FX/Niagara/Source/Niagara/Public/NiagaraCommon.h
@@ -784,7 +784,7 @@ struct NIAGARA_API FNiagaraSystemUpdateContext
 	FCustomWorkDelegate& GetPostWork() { return PostWork; }
 
 private:
-	void AddInternal(class UNiagaraComponent* Comp, bool bReInit);
+	void AddInternal(class UNiagaraComponent* Comp, bool bReInit, bool bAllowDestroySystemSim);
 
 	UPROPERTY(transient)
 	TArray<TObjectPtr<UNiagaraComponent>> ComponentsToReset;
