Sim cache column refactor

Fix issue with filtering shrinking columns by specifying a starting width
Moved column generation to a helper function, and now only generate new columns if a reset is needed
Filtering now hides columns rather than destroying them

diff --git a/Engine/Plugins/FX/Niagara/Source/NiagaraEditor/Private/Widgets/SNiagaraSimCacheView.cpp b/Engine/Plugins/FX/Niagara/Source/NiagaraEditor/Private/Widgets/SNiagaraSimCacheView.cpp
--- a/Engine/Plugins/FX/Niagara/Source/NiagaraEditor/Private/Widgets/SNiagaraSimCacheView.cpp
+++ b/Engine/Plugins/FX/Niagara/Source/NiagaraEditor/Private/Widgets/SNiagaraSimCacheView.cpp
@@ -71,7 +71,7 @@ void SNiagaraSimCacheView::Construct(const FArguments& InArgs)
 
 	HeaderRowWidget = SNew(SHeaderRow);
 
-	UpdateColumns(false);
+	UpdateColumns(true);
 	UpdateRows(false);
 	UpdateBufferSelectionList();
 
@@ -143,34 +143,30 @@ void SNiagaraSimCacheView::Construct(const FArguments& InArgs)
 		+ SVerticalBox::Slot()
 		.FillHeight(1.0f)
 		[
-			SNew(SVerticalBox)
-			+ SVerticalBox::Slot()
-			.FillHeight(1.0f)
+			SNew(SHorizontalBox)
+			+ SHorizontalBox::Slot()
 			[
-				SNew(SHorizontalBox)
-				+ SHorizontalBox::Slot()
-				.FillWidth(1.0f)
-				[
-					SNew(SScrollBox)
-					.Orientation(Orient_Horizontal)
-					.ExternalScrollbar(HorizontalScrollBar)
-					+ SScrollBox::Slot()
-					[
-						SAssignNew(ListViewWidget, SListView<TSharedPtr<int32>>)
-						.ListItemsSource(&RowItems)
-						.OnGenerateRow(this, &SNiagaraSimCacheView::MakeRowWidget)
-						.Visibility(EVisibility::Visible)
-						.SelectionMode(ESelectionMode::Single)
-						.ExternalScrollbar(VerticalScrollBar)
-						.HeaderRow(HeaderRowWidget)
-					]
-				]
-				+ SHorizontalBox::Slot()
-				.AutoWidth()
+				 SNew(SScrollBox)
+				 .Orientation(Orient_Horizontal)
+				 .ExternalScrollbar(HorizontalScrollBar)
+				 + SScrollBox::Slot()
 				[
-					VerticalScrollBar
+					SAssignNew(ListViewWidget, SListView<TSharedPtr<int32>>)
+					.ListItemsSource(&RowItems)
+					.OnGenerateRow(this, &SNiagaraSimCacheView::MakeRowWidget)
+					.Visibility(EVisibility::Visible)
+					.SelectionMode(ESelectionMode::Single)
+					.ExternalScrollbar(VerticalScrollBar)
+					.ConsumeMouseWheel(EConsumeMouseWheel::Always)
+					.AllowOverscroll(EAllowOverscroll::No)
+					.HeaderRow(HeaderRowWidget)
 				]
 			]
+			+ SHorizontalBox::Slot()
+			.AutoWidth()
+			[
+				VerticalScrollBar
+			]
 		]
 		+ SVerticalBox::Slot()
 		.AutoHeight()
@@ -188,11 +184,13 @@ TSharedRef<ITableRow> SNiagaraSimCacheView::MakeRowWidget(const TSharedPtr<int32
 		.SimCacheViewModel(SimCacheViewModel);
 }
 
-void SNiagaraSimCacheView::UpdateColumns(const bool bRefresh)
+void SNiagaraSimCacheView::GenerateColumns()
 {
-	// TODO: Hide columns rather than regenerating everything?
+	//  Give columns a width to prevent them from being shrunk when filtering. 
+	constexpr float ManualWidth = 125.0f;
 	HeaderRowWidget->ClearColumns();
 
+	// Generate instance count column
 	HeaderRowWidget->AddColumn(
 		SHeaderRow::Column(NAME_Instance)
 		.DefaultLabel(FText::FromName(NAME_Instance))
@@ -200,29 +198,13 @@ void SNiagaraSimCacheView::UpdateColumns(const bool bRefresh)
 		.VAlignHeader(EVerticalAlignment::VAlign_Fill)
 		.HAlignCell(EHorizontalAlignment::HAlign_Center)
 		.VAlignCell(EVerticalAlignment::VAlign_Fill)
+		.ManualWidth(ManualWidth)
 		.SortMode(EColumnSortMode::None)
 	);
-
-
-
-	const bool bFilterActive = ComponentFilterArray.Num() > 0;
+		
+	// Generate a column for each component
 	for (const FNiagaraSimCacheViewModel::FComponentInfo& ComponentInfo : SimCacheViewModel->GetComponentInfos())
 	{
-		if (bFilterActive)
-		{
-			const FString ComponentInfoString = ComponentInfo.Name.ToString();
-			const bool bPassedFilter = ComponentFilterArray.ContainsByPredicate(
-				[ComponentInfoString = ComponentInfo.Name.ToString()](const FString& ComponentFilter)
-			{
-				return ComponentInfoString.Contains(ComponentFilter);
-			}
-			);
-			if (bPassedFilter == false)
-			{
-				continue;
-			}
-		}
-
 		HeaderRowWidget->AddColumn(
 			SHeaderRow::Column(ComponentInfo.Name)
 			.DefaultLabel(FText::FromName(ComponentInfo.Name))
@@ -230,15 +212,38 @@ void SNiagaraSimCacheView::UpdateColumns(const bool bRefresh)
 			.VAlignHeader(EVerticalAlignment::VAlign_Fill)
 			.HAlignCell(EHorizontalAlignment::HAlign_Center)
 			.VAlignCell(EVerticalAlignment::VAlign_Fill)
+			.FillWidth(1.0f)
+			.ManualWidth(ManualWidth)
 			.SortMode(EColumnSortMode::None)
 		);
 	}
+}
 
-	HeaderRowWidget->ResetColumnWidths();
-	HeaderRowWidget->RefreshColumns();
-	if (bRefresh && ListViewWidget)
+void SNiagaraSimCacheView::UpdateColumns(const bool bReset)
+{
+	if(bReset)
 	{
-		ListViewWidget->RequestListRefresh();
+		GenerateColumns();
+	}
+
+	const bool bFilterEmpty = ComponentFilterArray.Num() == 0;
+
+	// If the columns are newly generated, and there are no filters to apply they are already up to date.
+	const bool bColumnsUpToDate = bReset && bFilterEmpty;
+
+	if(!bColumnsUpToDate)
+	{
+		for (SHeaderRow::FColumn Column : HeaderRowWidget->GetColumns())
+		{
+			const bool bPassedFilter = bFilterEmpty || ComponentFilterArray.ContainsByPredicate([ColumnName = Column.DefaultText.Get().ToString()](const FString& ComponentFilter)
+			{
+				return ColumnName.Contains(ComponentFilter) || ColumnName.Equals(NAME_Instance.ToString());
+			});
+
+			HeaderRowWidget->SetShowGeneratedColumn(Column.ColumnId, bPassedFilter);
+		}
+
+		HeaderRowWidget->RefreshColumns();
 	}
 }
 
@@ -291,7 +296,7 @@ void SNiagaraSimCacheView::BufferSelectionChanged(TSharedPtr<FBufferSelectionInf
 	}
 	
 	SimCacheViewModel->SetEmitterIndex(NewSelection->Key);
-	UpdateColumns(false);
+	UpdateColumns(true);
 	UpdateRows(true);
 }
 
@@ -312,7 +317,7 @@ void SNiagaraSimCacheView::OnComponentFilterChange(const FText& InFilter)
 {
 	ComponentFilterArray.Empty();
 	InFilter.ToString().ParseIntoArray(ComponentFilterArray, TEXT(","));
-	UpdateColumns(true);
+	UpdateColumns(false);
 }
 
 void SNiagaraSimCacheView::OnSimCacheChanged(const FAssetData& InAsset)
@@ -333,7 +338,7 @@ void SNiagaraSimCacheView::OnViewDataChanged(const bool bFullRefresh)
 	UpdateRows(true);
 	if (bFullRefresh)
 	{
-		UpdateColumns(true);
+		UpdateColumns(false);
 		UpdateBufferSelectionList();
 	}
 }
diff --git a/Engine/Plugins/FX/Niagara/Source/NiagaraEditor/Private/Widgets/SNiagaraSimCacheView.h b/Engine/Plugins/FX/Niagara/Source/NiagaraEditor/Private/Widgets/SNiagaraSimCacheView.h
--- a/Engine/Plugins/FX/Niagara/Source/NiagaraEditor/Private/Widgets/SNiagaraSimCacheView.h
+++ b/Engine/Plugins/FX/Niagara/Source/NiagaraEditor/Private/Widgets/SNiagaraSimCacheView.h
@@ -23,7 +23,9 @@ public:
 
 	TSharedRef<ITableRow> MakeRowWidget(const TSharedPtr<int32> RowIndexPtr, const TSharedRef<STableViewBase>& OwnerTable) const;
 
-	void UpdateColumns(const bool bRefresh);
+	// Update existing or generate new columns and apply any filters.
+	// @param bReset Weather the columns need a full reset. If true this will destroy existing columns and rebuild new ones. 
+	void UpdateColumns(const bool bReset);
 
 	void UpdateRows(const bool bRefresh);
 
@@ -41,13 +43,14 @@ public:
 
 	void OnViewDataChanged(const bool bFullRefresh);
 
+private:
+	void GenerateColumns();
+
 	TArray<TSharedPtr<int32>>					RowItems;
 	TSharedPtr<FNiagaraSimCacheViewModel>		SimCacheViewModel;
 
 	TArray<TSharedPtr<FBufferSelectionInfo>>	BufferSelectionList;
-
-	int32										FrameIndex = 0;
-
+	
 	TSharedPtr<SHeaderRow>						HeaderRowWidget;
 	TSharedPtr<SListView<TSharedPtr<int32>>>	ListViewWidget;
 
