- Ensure Thumbnails are cached before using otherwise they show up as empty / missing

diff --git a/Engine/Plugins/FX/Niagara/Source/Niagara/Classes/NiagaraEmitter.h b/Engine/Plugins/FX/Niagara/Source/Niagara/Classes/NiagaraEmitter.h
--- a/Engine/Plugins/FX/Niagara/Source/Niagara/Classes/NiagaraEmitter.h
+++ b/Engine/Plugins/FX/Niagara/Source/Niagara/Classes/NiagaraEmitter.h
@@ -671,13 +671,9 @@ public:
 	UPROPERTY(AssetRegistrySearchable, EditAnywhere, Category = Script)
 	FText Category;
 
-	/** Internal: The thumbnail image.*/
+	///** Internal: The thumbnail image.*/
 	UPROPERTY()
 	TObjectPtr<UTexture2D> ThumbnailImage;
-
-	/** Internal: Indicates the thumbnail image is out of date.*/
-	UPROPERTY()
-	uint32 ThumbnailImageOutOfDate : 1;
 	
 	/** If this emitter is exposed to the library, or should be explicitly hidden. */
     UPROPERTY(EditAnywhere, AdvancedDisplay, Category = "Asset Options", AssetRegistrySearchable)
diff --git a/Engine/Plugins/FX/Niagara/Source/Niagara/Classes/NiagaraSystem.h b/Engine/Plugins/FX/Niagara/Source/Niagara/Classes/NiagaraSystem.h
--- a/Engine/Plugins/FX/Niagara/Source/Niagara/Classes/NiagaraSystem.h
+++ b/Engine/Plugins/FX/Niagara/Source/Niagara/Classes/NiagaraSystem.h
@@ -411,10 +411,6 @@ public:
 	UPROPERTY()
 	TObjectPtr<class UTexture2D> ThumbnailImage;
 
-	/** Internal: Indicates the thumbnail image is out of date.*/
-	UPROPERTY()
-	uint32 ThumbnailImageOutOfDate : 1;
-
 	/** Deprecated library exposure bool. Use the LibraryVisibility enum instead. FNiagaraEditorUtilities has accessor functions that takes deprecation into account */
 	UPROPERTY()
 	bool bExposeToLibrary_DEPRECATED;
diff --git a/Engine/Plugins/FX/Niagara/Source/Niagara/Private/NiagaraEmitter.cpp b/Engine/Plugins/FX/Niagara/Source/Niagara/Private/NiagaraEmitter.cpp
--- a/Engine/Plugins/FX/Niagara/Source/Niagara/Private/NiagaraEmitter.cpp
+++ b/Engine/Plugins/FX/Niagara/Source/Niagara/Private/NiagaraEmitter.cpp
@@ -165,7 +165,6 @@ UNiagaraEmitter::UNiagaraEmitter(const FObjectInitializer& Initializer)
 #endif
 {
 #if WITH_EDITORONLY_DATA
-	ThumbnailImageOutOfDate = true;
 	IsCooked = false;
 #endif
 }
@@ -1102,7 +1101,6 @@ void UNiagaraEmitter::PostEditChangeProperty(FPropertyChangedEvent& PropertyChan
 		Data.RuntimeEstimation.Init();
 		Data.UpdateDebugName(*this, nullptr);
 	}
-	ThumbnailImageOutOfDate = true;
 	UpdateChangeId(TEXT("PostEditChangeProperty"));
 	
 	OnPropertiesChangedDelegate.Broadcast();
@@ -1210,7 +1208,6 @@ void UNiagaraEmitter::PostEditChangeVersionedProperty(FPropertyChangedEvent& Pro
 
 	ResolveScalabilitySettings();
 
-	ThumbnailImageOutOfDate = true;
 	UpdateChangeId(TEXT("PostEditChangeProperty"));
 	OnPropertiesChangedDelegate.Broadcast();
 
diff --git a/Engine/Plugins/FX/Niagara/Source/Niagara/Private/NiagaraSystem.cpp b/Engine/Plugins/FX/Niagara/Source/Niagara/Private/NiagaraSystem.cpp
--- a/Engine/Plugins/FX/Niagara/Source/Niagara/Private/NiagaraSystem.cpp
+++ b/Engine/Plugins/FX/Niagara/Source/Niagara/Private/NiagaraSystem.cpp
@@ -324,9 +324,6 @@ bool UNiagaraSystem::ChangeEmitterVersion(const FVersionedNiagaraEmitter& Versio
 void UNiagaraSystem::PostInitProperties()
 {
 	Super::PostInitProperties();
-#if WITH_EDITORONLY_DATA
-	ThumbnailImageOutOfDate = true;
-#endif
 	if (HasAnyFlags(RF_ClassDefaultObject | RF_NeedLoad) == false)
 	{
 		SystemSpawnScript = NewObject<UNiagaraScript>(this, "SystemSpawnScript", RF_Transactional);
@@ -769,8 +766,6 @@ void UNiagaraSystem::PostEditChangeProperty(struct FPropertyChangedEvent& Proper
 {
 	Super::PostEditChangeProperty(PropertyChangedEvent);
 
-	ThumbnailImageOutOfDate = true;
-
 	if (PropertyChangedEvent.Property != nullptr)
 	{
 		if (PropertyChangedEvent.Property->GetFName() == GET_MEMBER_NAME_CHECKED(UNiagaraSystem, WarmupTickCount))
diff --git a/Engine/Plugins/FX/Niagara/Source/NiagaraEditor/Private/NiagaraThumbnailRenderer.cpp b/Engine/Plugins/FX/Niagara/Source/NiagaraEditor/Private/NiagaraThumbnailRenderer.cpp
--- a/Engine/Plugins/FX/Niagara/Source/NiagaraEditor/Private/NiagaraThumbnailRenderer.cpp
+++ b/Engine/Plugins/FX/Niagara/Source/NiagaraEditor/Private/NiagaraThumbnailRenderer.cpp
@@ -14,33 +14,21 @@ bool UNiagaraThumbnailRendererBase::CanVisualizeAsset(UObject* Object)
 
 void UNiagaraThumbnailRendererBase::GetThumbnailSize(UObject* Object, float Zoom, uint32& OutWidth, uint32& OutHeight) const
 {
-	UTexture2D* ObjectTexture = GetThumbnailTextureFromObject(Object);
-	if (ObjectTexture != nullptr)
-	{
-		OutWidth = Zoom * ObjectTexture->GetSizeX();
-		OutHeight = Zoom * ObjectTexture->GetSizeY();
-	}
-	else
-	{
-		OutWidth = 0;
-		OutHeight = 0;
-	}
+	Super::GetThumbnailSize(GetThumbnailTextureFromObject(Object), Zoom, OutWidth, OutHeight);
 }
 
 void UNiagaraThumbnailRendererBase::Draw(UObject* Object, int32 X, int32 Y, uint32 Width, uint32 Height, FRenderTarget* RenderTarget, FCanvas* Canvas, bool bAdditionalViewFamily)
 {
-	UTexture2D* ObjectTexture = GetThumbnailTextureFromObject(Object);
-	if (ObjectTexture != nullptr)
-	{
-		Super::Draw(ObjectTexture, X, Y, Width, Height, RenderTarget, Canvas, bAdditionalViewFamily);
-	}
+	Super::Draw(GetThumbnailTextureFromObject(Object), X, Y, Width, Height, RenderTarget, Canvas, bAdditionalViewFamily);
 }
 
 UTexture2D* UNiagaraEmitterThumbnailRenderer::GetThumbnailTextureFromObject(UObject* Object) const
 {
 	UNiagaraEmitter* Emitter = Cast<UNiagaraEmitter>(Object);
-	if (Emitter != nullptr)
+	if (Emitter && Emitter->ThumbnailImage)
 	{
+		Emitter->ThumbnailImage->FinishCachePlatformData();
+		Emitter->ThumbnailImage->UpdateResource();
 		return Emitter->ThumbnailImage;
 	}
 	return nullptr;
@@ -49,8 +37,10 @@ UTexture2D* UNiagaraEmitterThumbnailRenderer::GetThumbnailTextureFromObject(UObj
 UTexture2D* UNiagaraSystemThumbnailRenderer::GetThumbnailTextureFromObject(UObject* Object) const
 {
 	UNiagaraSystem* System = Cast<UNiagaraSystem>(Object);
-	if (System != nullptr)
+	if (System && System->ThumbnailImage)
 	{
+		System->ThumbnailImage->FinishCachePlatformData();
+		System->ThumbnailImage->UpdateResource();
 		return System->ThumbnailImage;
 	}
 	return nullptr;
