Fixed a bug where the Niagara stack did not reflect the compile status correctly

diff --git a/Engine/Plugins/FX/Niagara/Source/NiagaraEditor/Private/ViewModels/NiagaraSystemViewModel.cpp b/Engine/Plugins/FX/Niagara/Source/NiagaraEditor/Private/ViewModels/NiagaraSystemViewModel.cpp
--- a/Engine/Plugins/FX/Niagara/Source/NiagaraEditor/Private/ViewModels/NiagaraSystemViewModel.cpp
+++ b/Engine/Plugins/FX/Niagara/Source/NiagaraEditor/Private/ViewModels/NiagaraSystemViewModel.cpp
@@ -1560,6 +1560,7 @@ void FNiagaraSystemViewModel::TickCompileStatus()
 			if (bCompileStatusRefreshed)
 			{
 				SendLastCompileMessageJobs();
+				RefreshStackViewModels();
 			}
 		}
 	}
@@ -2845,6 +2846,18 @@ void FNiagaraSystemViewModel::ScratchPadScriptsChanged()
 	}
 }
 
+void FNiagaraSystemViewModel::RefreshStackViewModels()
+{
+	if (SystemStackViewModel)
+	{
+		SystemStackViewModel->Refresh();
+	}
+	for (TSharedRef<FNiagaraEmitterHandleViewModel> EmitterHandleViewModel : EmitterHandleViewModels)
+	{
+		EmitterHandleViewModel->GetEmitterStackViewModel()->Refresh();
+	}
+}
+
 void FNiagaraSystemViewModel::RefreshAssetMessages()
 {
 	FNiagaraMessageManager* MessageManager = FNiagaraMessageManager::Get();
diff --git a/Engine/Plugins/FX/Niagara/Source/NiagaraEditor/Public/ViewModels/NiagaraSystemViewModel.h b/Engine/Plugins/FX/Niagara/Source/NiagaraEditor/Public/ViewModels/NiagaraSystemViewModel.h
--- a/Engine/Plugins/FX/Niagara/Source/NiagaraEditor/Public/ViewModels/NiagaraSystemViewModel.h
+++ b/Engine/Plugins/FX/Niagara/Source/NiagaraEditor/Public/ViewModels/NiagaraSystemViewModel.h
@@ -573,6 +573,9 @@ private:
 	/** Called whenever one of the scripts in the scratch pad changes. */
 	void ScratchPadScriptsChanged();
 
+	/** Deferred refreshes the system and emitter view models the next time they are ticked. */
+	void RefreshStackViewModels();
+
 	/** Called whenever the map of messages associated with the managed Emitter/System changes. */
 	void RefreshAssetMessages();
 
