- Atomic ops for Array Data Interface
- Only supported on int array and not optimized fully
- Also fixes a bug with array readback that came in with RDG conversion.

diff --git a/Engine/Plugins/FX/Niagara/Shaders/Private/NiagaraDataInterfaceArrayRWTemplate.ush b/Engine/Plugins/FX/Niagara/Shaders/Private/NiagaraDataInterfaceArrayRWTemplate.ush
--- a/Engine/Plugins/FX/Niagara/Shaders/Private/NiagaraDataInterfaceArrayRWTemplate.ush
+++ b/Engine/Plugins/FX/Niagara/Shaders/Private/NiagaraDataInterfaceArrayRWTemplate.ush
@@ -90,3 +90,55 @@ void RemoveLastElem_{ParameterName}_UEImpureCall(bool bSkip, out {VariableType}
 	{RWBufferRead}
 	#undef BUFFER_NAME
 }
+
+#if {bSupportsAtomicOps}
+
+void AtomicAdd_{ParameterName}_UEImpureCall(bool bSkip, int Index, int Value, out {VariableType} PrevValue, out {VariableType} CurrValue)
+{
+	bool bIsValidIndex;
+	IsValidIndex_{ParameterName}(Index, bIsValidIndex);
+	if ( !bSkip && bIsValidIndex )
+	{
+		InterlockedAdd({ParameterName}_ArrayRWBuffer[Index], Value, PrevValue);
+		CurrValue = PrevValue + Value;
+	}
+	else
+	{
+		PrevValue = 0;
+		CurrValue = 0;
+	}
+}
+
+void AtomicMin_{ParameterName}_UEImpureCall(bool bSkip, int Index, int Value, out {VariableType} PrevValue, out {VariableType} CurrValue)
+{
+	bool bIsValidIndex;
+	IsValidIndex_{ParameterName}(Index, bIsValidIndex);
+	if ( !bSkip && bIsValidIndex )
+	{
+		InterlockedMin({ParameterName}_ArrayRWBuffer[Index], Value, PrevValue);
+		CurrValue = PrevValue + Value;
+	}
+	else
+	{
+		PrevValue = 0;
+		CurrValue = 0;
+	}
+}
+
+void AtomicMax_{ParameterName}_UEImpureCall(bool bSkip, int Index, int Value, out {VariableType} PrevValue, out {VariableType} CurrValue)
+{
+	bool bIsValidIndex;
+	IsValidIndex_{ParameterName}(Index, bIsValidIndex);
+	if ( !bSkip && bIsValidIndex )
+	{
+		InterlockedMax({ParameterName}_ArrayRWBuffer[Index], Value, PrevValue);
+		CurrValue = PrevValue + Value;
+	}
+	else
+	{
+		PrevValue = 0;
+		CurrValue = 0;
+	}
+}
+
+#endif
diff --git a/Engine/Plugins/FX/Niagara/Source/Niagara/Classes/NiagaraDataInterfaceArrayImpl.h b/Engine/Plugins/FX/Niagara/Source/Niagara/Classes/NiagaraDataInterfaceArrayImpl.h
--- a/Engine/Plugins/FX/Niagara/Source/Niagara/Classes/NiagaraDataInterfaceArrayImpl.h
+++ b/Engine/Plugins/FX/Niagara/Source/Niagara/Classes/NiagaraDataInterfaceArrayImpl.h
@@ -76,6 +76,7 @@ struct FNDIArrayImplHelperBase
 {
 	static constexpr bool bSupportsCPU = true;
 	static constexpr bool bSupportsGPU = true;
+	static constexpr bool bSupportsAtomicOps = false;
 
 	//-OPT: We can reduce the differences between read and RW if we have typed UAV loads
 	//static constexpr TCHAR const* HLSLVariableType		= TEXT("float");
@@ -104,6 +105,10 @@ struct FNDIArrayImplHelperBase
 	{
 		FMemory::Memcpy(Dest, Src, NumElements * sizeof(TArrayType));
 	}
+
+	//static TArrayType AtomicAdd(TArrayType* Dest, TArrayType Value) { check(false); }
+	//static TArrayType AtomicMin(TArrayType* Dest, TArrayType Value) { check(false); }
+	//static TArrayType AtomicMax(TArrayType* Dest, TArrayType Value) { check(false); }
 };
 
 template<typename TArrayType>
@@ -139,6 +144,10 @@ struct FNiagaraDataInterfaceArrayImplHelper
 	static const FName Function_AddName;
 	static const FName Function_RemoveLastElemName;
 
+	static const FName Function_AtomicAddName;
+	static const FName Function_AtomicMinName;
+	static const FName Function_AtomicMaxName;
+
 #if WITH_EDITORONLY_DATA
 	static bool UpgradeFunctionCall(FNiagaraFunctionSignature& FunctionSignature);
 #endif
@@ -146,6 +155,7 @@ struct FNiagaraDataInterfaceArrayImplHelper
 	{
 		return bIsRWArray ? FNiagaraDataInterfaceArrayImplHelper::HLSLReadWriteTemplateFile : FNiagaraDataInterfaceArrayImplHelper::HLSLReadTemplateFile;
 	}
+	static bool IsRWFunction(const FName FunctionName);
 };
 
 //////////////////////////////////////////////////////////////////////////
@@ -442,21 +452,21 @@ struct FNDIArrayProxyImpl : public INDIArrayProxyBase
 		return INDEX_NONE;
 	}
 
-	virtual void PostSimulate(FRHICommandList& RHICmdList, const FNiagaraDataInterfaceArgs& Context) override
+	virtual void PostSimulate(const FNDIGpuComputePostSimulateContext& Context) override
 	{
 		if (bShouldSyncToCpu == false)
 		{
 			return;
 		}
 		
-		const FNDIArrayInstanceData_RenderThread<TArrayType>* InstanceData_RT = PerInstanceData_RenderThread.Find(Context.SystemInstanceID);
+		const FNDIArrayInstanceData_RenderThread<TArrayType>* InstanceData_RT = PerInstanceData_RenderThread.Find(Context.GetSystemInstanceID());
 		if ( !InstanceData_RT || InstanceData_RT->IsReadOnly() || (InstanceData_RT->ArrayNumBytes == 0) )
 		{
 			return;
 		}
 
-		const FNiagaraGPUInstanceCountManager& CountManager = Context.ComputeDispatchInterface->GetGPUInstanceCounterManager();
-		FNiagaraGpuReadbackManager* ReadbackManager = Context.ComputeDispatchInterface->GetGpuReadbackManager();
+		const FNiagaraGPUInstanceCountManager& CountManager = Context.GetComputeDispatchInterface().GetGPUInstanceCounterManager();
+		FNiagaraGpuReadbackManager* ReadbackManager = Context.GetComputeDispatchInterface().GetGpuReadbackManager();
 
 		const FNiagaraGpuReadbackManager::FBufferRequest BufferRequests[] =
 		{
@@ -475,36 +485,43 @@ struct FNDIArrayProxyImpl : public INDIArrayProxyBase
 			FRHITransitionInfo(InstanceData_RT->ArrayBuffer, ERHIAccess::CopySrc, ERHIAccess::UAVCompute),
 		};
 
-		RHICmdList.Transition(TransitionsBefore);
-		ReadbackManager->EnqueueReadbacks(
-			RHICmdList,
-			BufferRequests,
-			[SystemInstanceID=Context.SystemInstanceID, WeakOwner=TWeakObjectPtr(Owner), Proxy=this](TConstArrayView<TPair<void*, uint32>> ReadbackData)
+		AddPass(
+			Context.GetGraphBuilder(),
+			RDG_EVENT_NAME("NDIArrayReadback"),
+			[BufferRequests, TransitionsBefore, TransitionsAfter, SystemInstanceID=Context.GetSystemInstanceID(), WeakOwner=TWeakObjectPtr(Owner), Proxy=this, ReadbackManager](FRHICommandListImmediate& RHICmdList)
 			{
-				const int32 NumElements = *reinterpret_cast<const uint32*>(ReadbackData[0].Key);
-				TArray<TArrayType> ArrayData;
-				if ( NumElements > 0 )
-				{
-					ArrayData.AddUninitialized(NumElements);
-					FNDIArrayImplHelper<TArrayType>::CopyGpuToCpuMemory(ArrayData.GetData(), reinterpret_cast<const TVMArrayType*>(ReadbackData[1].Key), NumElements);
-				}
-
-				AsyncTask(
-					ENamedThreads::GameThread,
-					[SystemInstanceID, ArrayData=MoveTemp(ArrayData), WeakOwner, Proxy]()
+				RHICmdList.Transition(TransitionsBefore);
+				ReadbackManager->EnqueueReadbacks(
+					RHICmdList,
+					BufferRequests,
+					[SystemInstanceID, WeakOwner, Proxy](TConstArrayView<TPair<void*, uint32>> ReadbackData)
 					{
-						// If this is nullptr the proxy is no longer valid so discard
-						if ( WeakOwner.Get() == nullptr )
+						const int32 NumElements = *reinterpret_cast<const uint32*>(ReadbackData[0].Key);
+						TArray<TArrayType> ArrayData;
+						if ( NumElements > 0 )
 						{
-							return;
+							ArrayData.AddUninitialized(NumElements);
+							FNDIArrayImplHelper<TArrayType>::CopyGpuToCpuMemory(ArrayData.GetData(), reinterpret_cast<const TVMArrayType*>(ReadbackData[1].Key), NumElements);
 						}
 
-						Proxy->SetInstanceArrayData(SystemInstanceID, ArrayData);
+						AsyncTask(
+							ENamedThreads::GameThread,
+							[SystemInstanceID, ArrayData=MoveTemp(ArrayData), WeakOwner, Proxy]()
+							{
+								// If this is nullptr the proxy is no longer valid so discard
+								if ( WeakOwner.Get() == nullptr )
+								{
+									return;
+								}
+
+								Proxy->SetInstanceArrayData(SystemInstanceID, ArrayData);
+							}
+						);
 					}
 				);
+				RHICmdList.Transition(TransitionsAfter);
 			}
 		);
-		RHICmdList.Transition(TransitionsAfter);
 	}
 
 	//////////////////////////////////////////////////////////////////////////
@@ -749,10 +766,53 @@ struct FNDIArrayProxyImpl : public INDIArrayProxyBase
 			Sig.OutputDescriptions.Add(Sig.Outputs[1], NSLOCTEXT("Niagara", "Array_RemoveLastElemDesc_IsValid", "True if we removed a value from the array, False if no entries or we skipped the remove."));
 		#endif
 		}
+
+		if (FNDIArrayImplHelper<TArrayType>::bSupportsAtomicOps)
+		{
+			{
+				FNiagaraFunctionSignature& Sig = OutFunctions.Add_GetRef(DefaultMutableSig);
+				Sig.Name = FNiagaraDataInterfaceArrayImplHelper::Function_AtomicAddName;
+				Sig.Inputs.Emplace(FNiagaraTypeDefinition::GetBoolDef(), TEXT("SkipAdd"));
+				Sig.Inputs.Emplace(FNiagaraTypeDefinition::GetIntDef(), TEXT("Index"));
+				Sig.Inputs.Emplace(FNDIArrayImplHelper<TArrayType>::GetTypeDefinition(), TEXT("Value"));
+				Sig.Outputs.Emplace(FNDIArrayImplHelper<TArrayType>::GetTypeDefinition(), TEXT("PreviousValue"));
+				Sig.Outputs.Emplace(FNDIArrayImplHelper<TArrayType>::GetTypeDefinition(), TEXT("CurrentValue"));
+				Sig.SetDescription(NSLOCTEXT("Niagara", "Array_AtomicAddDesc", "Optionally perform an atomic add on the array element."));
+				Sig.SetInputDescription(Sig.Inputs[1], NSLOCTEXT("Niagara", "Array_AtomicAdd_SkipAdd", "When enabled will not perform the add operation, the return values will therefore be invalid."));
+				Sig.SetOutputDescription(Sig.Inputs[0], NSLOCTEXT("Niagara", "Array_AtomicAdd_PrevValue", "The value before the operation was performed."));
+				Sig.SetOutputDescription(Sig.Inputs[1], NSLOCTEXT("Niagara", "Array_AtomicAdd_CurrValue", "The value after the operation was performed."));
+			}
+			{
+				FNiagaraFunctionSignature& Sig = OutFunctions.Add_GetRef(DefaultMutableSig);
+				Sig.Name = FNiagaraDataInterfaceArrayImplHelper::Function_AtomicMinName;
+				Sig.Inputs.Emplace(FNiagaraTypeDefinition::GetBoolDef(), TEXT("SkipMin"));
+				Sig.Inputs.Emplace(FNiagaraTypeDefinition::GetIntDef(), TEXT("Index"));
+				Sig.Inputs.Emplace(FNDIArrayImplHelper<TArrayType>::GetTypeDefinition(), TEXT("Value"));
+				Sig.Outputs.Emplace(FNDIArrayImplHelper<TArrayType>::GetTypeDefinition(), TEXT("PreviousValue"));
+				Sig.Outputs.Emplace(FNDIArrayImplHelper<TArrayType>::GetTypeDefinition(), TEXT("CurrentValue"));
+				Sig.SetDescription(NSLOCTEXT("Niagara", "Array_AtomicMinDesc", "Optionally perform an atomic min on the array element."));
+				Sig.SetInputDescription(Sig.Inputs[1], NSLOCTEXT("Niagara", "Array_AtomicMin_SkipMin", "When enabled will not perform the min operation, the return values will therefore be invalid."));
+				Sig.SetOutputDescription(Sig.Inputs[0], NSLOCTEXT("Niagara", "Array_AtomicMin_PrevValue", "The value before the operation was performed."));
+				Sig.SetOutputDescription(Sig.Inputs[1], NSLOCTEXT("Niagara", "Array_AtomicMin_CurrValue", "The value after the operation was performed."));
+			}
+			{
+				FNiagaraFunctionSignature& Sig = OutFunctions.Add_GetRef(DefaultMutableSig);
+				Sig.Name = FNiagaraDataInterfaceArrayImplHelper::Function_AtomicMaxName;
+				Sig.Inputs.Emplace(FNiagaraTypeDefinition::GetBoolDef(), TEXT("SkipMax"));
+				Sig.Inputs.Emplace(FNiagaraTypeDefinition::GetIntDef(), TEXT("Index"));
+				Sig.Inputs.Emplace(FNDIArrayImplHelper<TArrayType>::GetTypeDefinition(), TEXT("Value"));
+				Sig.Outputs.Emplace(FNDIArrayImplHelper<TArrayType>::GetTypeDefinition(), TEXT("PreviousValue"));
+				Sig.Outputs.Emplace(FNDIArrayImplHelper<TArrayType>::GetTypeDefinition(), TEXT("CurrentValue"));
+				Sig.SetDescription(NSLOCTEXT("Niagara", "Array_AtomicMaxDesc", "Optionally perform an atomic max on the array element."));
+				Sig.SetInputDescription(Sig.Inputs[1], NSLOCTEXT("Niagara", "Array_AtomicMax_SkipMax", "When enabled will not perform the max operation, the return values will therefore be invalid."));
+				Sig.SetOutputDescription(Sig.Inputs[0], NSLOCTEXT("Niagara", "Array_AtomicMax_PrevValue", "The value before the operation was performed."));
+				Sig.SetOutputDescription(Sig.Inputs[1], NSLOCTEXT("Niagara", "Array_AtomicMax_CurrValue", "The value after the operation was performed."));
+			}
+		}
 	}
 
 	template<typename T = FNDIArrayImplHelper<TArrayType>>
-	typename TEnableIf<T::bSupportsCPU>::Type GetVMExternalFunction_Internal(const FVMExternalFunctionBindingInfo& BindingInfo, void* InstanceData, FVMExternalFunction &OutFunc)
+	typename TEnableIf<T::bSupportsCPU>::Type GetVMExternalFunction_Internal(const FVMExternalFunctionBindingInfo& BindingInfo, void* InstanceData, FVMExternalFunction& OutFunc)
 	{
 		// Immutable functions
 		if (BindingInfo.Name == FNiagaraDataInterfaceArrayImplHelper::Function_LengthName)
@@ -804,23 +864,39 @@ struct FNDIArrayProxyImpl : public INDIArrayProxyBase
 	}
 
 	template<typename T = FNDIArrayImplHelper<TArrayType>>
-	typename TEnableIf<!T::bSupportsCPU>::Type GetVMExternalFunction_Internal(const FVMExternalFunctionBindingInfo& BindingInfo, void* InstanceData, FVMExternalFunction &OutFunc)
+	typename TEnableIf<!T::bSupportsCPU>::Type GetVMExternalFunction_Internal(const FVMExternalFunctionBindingInfo& BindingInfo, void* InstanceData, FVMExternalFunction& OutFunc)
 	{
 	}
 
-	virtual void GetVMExternalFunction(const FVMExternalFunctionBindingInfo& BindingInfo, void* InstanceData, FVMExternalFunction &OutFunc) override
+	template<typename T = FNDIArrayImplHelper<TArrayType>>
+	typename TEnableIf<T::bSupportsAtomicOps>::Type GetVMExternalFunction_AtomicInternal(const FVMExternalFunctionBindingInfo& BindingInfo, void* InstanceData, FVMExternalFunction& OutFunc)
+	{
+		if (BindingInfo.Name == FNiagaraDataInterfaceArrayImplHelper::Function_AtomicAddName)
+		{
+			OutFunc = FVMExternalFunction::CreateLambda([this](FVectorVMExternalFunctionContext& Context) { this->VMAtomicAdd(Context); });
+		}
+		else if (BindingInfo.Name == FNiagaraDataInterfaceArrayImplHelper::Function_AtomicMinName)
+		{
+			OutFunc = FVMExternalFunction::CreateLambda([this](FVectorVMExternalFunctionContext& Context) { this->VMAtomicMin(Context); });
+		}
+		else if (BindingInfo.Name == FNiagaraDataInterfaceArrayImplHelper::Function_AtomicMaxName)
+		{
+			OutFunc = FVMExternalFunction::CreateLambda([this](FVectorVMExternalFunctionContext& Context) { this->VMAtomicMax(Context); });
+		}
+	}
+
+	template<typename T = FNDIArrayImplHelper<TArrayType>>
+	typename TEnableIf<!T::bSupportsAtomicOps>::Type GetVMExternalFunction_AtomicInternal(const FVMExternalFunctionBindingInfo& BindingInfo, void* InstanceData, FVMExternalFunction& OutFunc)
 	{
-		GetVMExternalFunction_Internal(BindingInfo, InstanceData, OutFunc);
 	}
 
-	bool IsRWFunction(const FName FunctionName) const
+	virtual void GetVMExternalFunction(const FVMExternalFunctionBindingInfo& BindingInfo, void* InstanceData, FVMExternalFunction& OutFunc) override
 	{
-		return
-			FunctionName == FNiagaraDataInterfaceArrayImplHelper::Function_ClearName ||
-			FunctionName == FNiagaraDataInterfaceArrayImplHelper::Function_ResizeName ||
-			FunctionName == FNiagaraDataInterfaceArrayImplHelper::Function_SetArrayElemName ||
-			FunctionName == FNiagaraDataInterfaceArrayImplHelper::Function_AddName ||
-			FunctionName == FNiagaraDataInterfaceArrayImplHelper::Function_RemoveLastElemName;
+		GetVMExternalFunction_Internal(BindingInfo, InstanceData, OutFunc);
+		if (OutFunc.IsBound() == false)
+		{
+			GetVMExternalFunction_AtomicInternal(BindingInfo, InstanceData, OutFunc);
+		}
 	}
 
 #if WITH_EDITORONLY_DATA
@@ -829,7 +905,7 @@ struct FNDIArrayProxyImpl : public INDIArrayProxyBase
 		return ParamInfo.GeneratedFunctions.ContainsByPredicate(
 			[&](const FNiagaraDataInterfaceGeneratedFunction& Function)
 			{
-				return IsRWFunction(Function.DefinitionName);
+				return FNiagaraDataInterfaceArrayImplHelper::IsRWFunction(Function.DefinitionName);
 			}
 		);
 	}
@@ -846,6 +922,7 @@ struct FNDIArrayProxyImpl : public INDIArrayProxyBase
 			{TEXT("RWBufferType"),		FNDIArrayImplHelper<TArrayType>::RWHLSLBufferType},
 			{TEXT("RWBufferRead"),		FNDIArrayImplHelper<TArrayType>::RWHLSLBufferRead},
 			{TEXT("RWBufferWrite"),		FNDIArrayImplHelper<TArrayType>::RWHLSLBufferWrite},
+			{TEXT("bSupportsAtomicOps"),	FNDIArrayImplHelper<TArrayType>::bSupportsAtomicOps ? 1 : 0},
 		};
 
 		FString TemplateFile;
@@ -883,6 +960,16 @@ struct FNDIArrayProxyImpl : public INDIArrayProxyBase
 			{
 				return true;
 			}
+
+			if ( FNDIArrayImplHelper<TArrayType>::bSupportsAtomicOps )
+			{
+				if ((FunctionInfo.DefinitionName == FNiagaraDataInterfaceArrayImplHelper::Function_AtomicAddName) ||
+					(FunctionInfo.DefinitionName == FNiagaraDataInterfaceArrayImplHelper::Function_AtomicMinName) ||
+					(FunctionInfo.DefinitionName == FNiagaraDataInterfaceArrayImplHelper::Function_AtomicMaxName) )
+				{
+					return true;
+				}
+			}
 		}
 		return false;
 	}
@@ -952,7 +1039,7 @@ struct FNDIArrayProxyImpl : public INDIArrayProxyBase
 				DataInterface, SystemInstance,
 				[&](const FNiagaraDataInterfaceGeneratedFunction& Function) -> bool
 				{
-					bRWGpuArray = IsRWFunction(Function.DefinitionName);
+					bRWGpuArray = FNiagaraDataInterfaceArrayImplHelper::IsRWFunction(Function.DefinitionName);
 					return !bRWGpuArray;
 				}
 			);
@@ -1180,6 +1267,118 @@ struct FNDIArrayProxyImpl : public INDIArrayProxyBase
 		InstanceData->bIsRenderDirty = bShouldSyncToGpu;
 	}
 
+	template<typename T = FNDIArrayImplHelper<TArrayType>>
+	typename TEnableIf<!T::bSupportsAtomicOps>::Type VMAtomicAdd(FVectorVMExternalFunctionContext& Context) { check(false); }
+	template<typename T = FNDIArrayImplHelper<TArrayType>>
+	typename TEnableIf<!T::bSupportsAtomicOps>::Type VMAtomicMin(FVectorVMExternalFunctionContext& Context) { check(false); }
+	template<typename T = FNDIArrayImplHelper<TArrayType>>
+	typename TEnableIf<!T::bSupportsAtomicOps>::Type VMAtomicMax(FVectorVMExternalFunctionContext& Context) { check(false); }
+
+	template<typename T = FNDIArrayImplHelper<TArrayType>>
+	typename TEnableIf<T::bSupportsAtomicOps>::Type VMAtomicAdd(FVectorVMExternalFunctionContext& Context)
+	{
+		VectorVM::FUserPtrHandler<FNDIArrayInstanceData_GameThread<TArrayType>> InstanceData(Context);
+		FNDIInputParam<FNiagaraBool>	InSkipOp(Context);
+		FNDIInputParam<int32>			InIndex(Context);
+		FNDIInputParam<TVMArrayType>	InValue(Context);
+		FNDIOutputParam<TVMArrayType>	OutPrevValue(Context);
+		FNDIOutputParam<TVMArrayType>	OutCurrValue(Context);
+		
+		const TVMArrayType DefaultValue = FNDIArrayImplHelper<TArrayType>::GetDefaultValue();
+
+		//-OPT: Can do a in batches or single atomic op rather than one per instance
+		FWriteArrayRef ArrayData(Owner, InstanceData);
+		for (int32 i = 0; i < Context.GetNumInstances(); ++i)
+		{
+			const bool bSkipExecute = InSkipOp.GetAndAdvance();
+			const int32 Index = InIndex.GetAndAdvance();
+			const TVMArrayType Value = InValue.GetAndAdvance();
+			if (!bSkipExecute && ArrayData.GetArray().IsValidIndex(Index))
+			{
+				TVMArrayType PreviousValue = FNDIArrayImplHelper<TArrayType>::AtomicAdd(&ArrayData.GetArray().GetData()[Index], Value);
+				OutPrevValue.SetAndAdvance(PreviousValue);
+				OutCurrValue.SetAndAdvance(PreviousValue + Value);
+			}
+			else
+			{
+				OutPrevValue.SetAndAdvance(DefaultValue);
+				OutCurrValue.SetAndAdvance(DefaultValue);
+			}
+		}
+
+		InstanceData->bIsRenderDirty = bShouldSyncToGpu;
+	}
+
+	template<typename T = FNDIArrayImplHelper<TArrayType>>
+	typename TEnableIf<T::bSupportsAtomicOps>::Type VMAtomicMin(FVectorVMExternalFunctionContext& Context)
+	{
+		VectorVM::FUserPtrHandler<FNDIArrayInstanceData_GameThread<TArrayType>> InstanceData(Context);
+		FNDIInputParam<FNiagaraBool>	InSkipOp(Context);
+		FNDIInputParam<int32>			InIndex(Context);
+		FNDIInputParam<TVMArrayType>	InValue(Context);
+		FNDIOutputParam<TVMArrayType>	OutPrevValue(Context);
+		FNDIOutputParam<TVMArrayType>	OutCurrValue(Context);
+
+		const TVMArrayType DefaultValue = FNDIArrayImplHelper<TArrayType>::GetDefaultValue();
+
+		//-OPT: Can do a in batches or single atomic op rather than one per instance
+		FWriteArrayRef ArrayData(Owner, InstanceData);
+		for (int32 i = 0; i < Context.GetNumInstances(); ++i)
+		{
+			const bool bSkipExecute = InSkipOp.GetAndAdvance();
+			const int32 Index = InIndex.GetAndAdvance();
+			const TVMArrayType Value = InValue.GetAndAdvance();
+			if (!bSkipExecute && ArrayData.GetArray().IsValidIndex(Index))
+			{
+				TVMArrayType PreviousValue = FNDIArrayImplHelper<TArrayType>::AtomicMin(&ArrayData.GetArray().GetData()[Index], Value);
+				OutPrevValue.SetAndAdvance(PreviousValue);
+				OutCurrValue.SetAndAdvance(PreviousValue + Value);
+			}
+			else
+			{
+				OutPrevValue.SetAndAdvance(DefaultValue);
+				OutCurrValue.SetAndAdvance(DefaultValue);
+			}
+		}
+
+		InstanceData->bIsRenderDirty = bShouldSyncToGpu;
+	}
+
+	template<typename T = FNDIArrayImplHelper<TArrayType>>
+	typename TEnableIf<T::bSupportsAtomicOps>::Type VMAtomicMax(FVectorVMExternalFunctionContext& Context)
+	{
+		VectorVM::FUserPtrHandler<FNDIArrayInstanceData_GameThread<TArrayType>> InstanceData(Context);
+		FNDIInputParam<FNiagaraBool>	InSkipOp(Context);
+		FNDIInputParam<int32>			InIndex(Context);
+		FNDIInputParam<TVMArrayType>	InValue(Context);
+		FNDIOutputParam<TVMArrayType>	OutPrevValue(Context);
+		FNDIOutputParam<TVMArrayType>	OutCurrValue(Context);
+
+		const TVMArrayType DefaultValue = FNDIArrayImplHelper<TArrayType>::GetDefaultValue();
+
+		//-OPT: Can do a in batches or single atomic op rather than one per instance
+		FWriteArrayRef ArrayData(Owner, InstanceData);
+		for (int32 i = 0; i < Context.GetNumInstances(); ++i)
+		{
+			const bool bSkipExecute = InSkipOp.GetAndAdvance();
+			const int32 Index = InIndex.GetAndAdvance();
+			const TVMArrayType Value = InValue.GetAndAdvance();
+			if (!bSkipExecute && ArrayData.GetArray().IsValidIndex(Index))
+			{
+				TVMArrayType PreviousValue = FNDIArrayImplHelper<TArrayType>::AtomicMax(&ArrayData.GetArray().GetData()[Index], Value);
+				OutPrevValue.SetAndAdvance(PreviousValue);
+				OutCurrValue.SetAndAdvance(PreviousValue + Value);
+			}
+			else
+			{
+				OutPrevValue.SetAndAdvance(DefaultValue);
+				OutCurrValue.SetAndAdvance(DefaultValue);
+			}
+		}
+
+		InstanceData->bIsRenderDirty = bShouldSyncToGpu;
+	}
+
 private:
 	TOwnerType*	Owner = nullptr;
 	bool bShouldSyncToGpu = false;
diff --git a/Engine/Plugins/FX/Niagara/Source/Niagara/Classes/NiagaraDataInterfaceArrayInt.h b/Engine/Plugins/FX/Niagara/Source/Niagara/Classes/NiagaraDataInterfaceArrayInt.h
--- a/Engine/Plugins/FX/Niagara/Source/Niagara/Classes/NiagaraDataInterfaceArrayInt.h
+++ b/Engine/Plugins/FX/Niagara/Source/Niagara/Classes/NiagaraDataInterfaceArrayInt.h
@@ -11,6 +11,8 @@ struct FNDIArrayImplHelper<int32> : public FNDIArrayImplHelperBase<int32>
 {
 	typedef int32 TVMArrayType;
 
+	static constexpr bool bSupportsAtomicOps = true;
+
 	static constexpr TCHAR const* HLSLVariableType = TEXT("int");
 	static constexpr EPixelFormat ReadPixelFormat = PF_R32_SINT;
 	static constexpr TCHAR const* ReadHLSLBufferType = TEXT("int");
@@ -22,6 +24,29 @@ struct FNDIArrayImplHelper<int32> : public FNDIArrayImplHelperBase<int32>
 
 	static const FNiagaraTypeDefinition& GetTypeDefinition() { return FNiagaraTypeDefinition::GetIntDef(); }
 	static const int32 GetDefaultValue() { return 0; }
+
+	static int32 AtomicAdd(int32* Dest, int32 Value)
+	{
+		return FPlatformAtomics::InterlockedAdd(Dest, Value);
+	}
+	static int32 AtomicMin(int32* Dest, int32 Value)
+	{
+		int32 PrevValue = *Dest;
+		while (PrevValue > Value)
+		{
+			PrevValue = FPlatformAtomics::InterlockedCompareExchange(Dest, Value, PrevValue);
+		}
+		return PrevValue;
+	}
+	static int32 AtomicMax(int32* Dest, int32 Value)
+	{
+		int32 PrevValue = *Dest;
+		while (PrevValue < Value)
+		{
+			PrevValue = FPlatformAtomics::InterlockedCompareExchange(Dest, Value, PrevValue);
+		}
+		return PrevValue;
+	}
 };
 
 template<>
diff --git a/Engine/Plugins/FX/Niagara/Source/Niagara/Private/NiagaraDataInterfaceArrayImpl.cpp b/Engine/Plugins/FX/Niagara/Source/Niagara/Private/NiagaraDataInterfaceArrayImpl.cpp
--- a/Engine/Plugins/FX/Niagara/Source/Niagara/Private/NiagaraDataInterfaceArrayImpl.cpp
+++ b/Engine/Plugins/FX/Niagara/Source/Niagara/Private/NiagaraDataInterfaceArrayImpl.cpp
@@ -16,6 +16,10 @@ const FName FNiagaraDataInterfaceArrayImplHelper::Function_SetArrayElemName(TEXT
 const FName FNiagaraDataInterfaceArrayImplHelper::Function_AddName(TEXT("Add"));
 const FName FNiagaraDataInterfaceArrayImplHelper::Function_RemoveLastElemName(TEXT("RemoveLastElem"));
 
+const FName FNiagaraDataInterfaceArrayImplHelper::Function_AtomicAddName("AtomicAdd");
+const FName FNiagaraDataInterfaceArrayImplHelper::Function_AtomicMinName("AtomicMin");
+const FName FNiagaraDataInterfaceArrayImplHelper::Function_AtomicMaxName("AtomicMax");
+
 #if WITH_EDITORONLY_DATA
 bool FNiagaraDataInterfaceArrayImplHelper::UpgradeFunctionCall(FNiagaraFunctionSignature& FunctionSignature)
 {
@@ -60,3 +64,20 @@ bool FNiagaraDataInterfaceArrayImplHelper::UpgradeFunctionCall(FNiagaraFunctionS
 	return true;
 }
 #endif
+
+bool FNiagaraDataInterfaceArrayImplHelper::IsRWFunction(const FName FunctionName)
+{
+	static const TSet<FName> RWFunctions =
+	{
+		Function_ClearName,
+		Function_ResizeName,
+		Function_SetArrayElemName,
+		Function_AddName,
+		Function_RemoveLastElemName,
+
+		Function_AtomicAddName,
+		Function_AtomicMinName,
+		Function_AtomicMaxName,
+	};
+	return RWFunctions.Contains(FunctionName);
+}
