Fixed a bug where spawned components were never destroyed when a Niagara system was reset

diff --git a/Engine/Plugins/FX/Niagara/Source/Niagara/Private/NiagaraRendererComponents.cpp b/Engine/Plugins/FX/Niagara/Source/Niagara/Private/NiagaraRendererComponents.cpp
--- a/Engine/Plugins/FX/Niagara/Source/Niagara/Private/NiagaraRendererComponents.cpp
+++ b/Engine/Plugins/FX/Niagara/Source/Niagara/Private/NiagaraRendererComponents.cpp
@@ -317,9 +317,7 @@ void FNiagaraRendererComponents::DestroyRenderState_Concurrent()
 					}
 				}
 			}
-
 		}
-		ComponentPool.Empty();
 	}
 
 	if (bRequiresAyncTask)
@@ -360,8 +358,8 @@ void FNiagaraRendererComponents::DestroyRenderState_Concurrent()
 				Component->DestroyComponent();
 			}
 		}
-		ComponentPool.Empty();
 	}
+	ComponentPool.Empty();
 	SpawnedOwner.Reset();
 }
 
