- Don't allow emitter to run if we will hit the RHI constant buffer limit

diff --git a/Engine/Plugins/FX/Niagara/Source/Niagara/Private/NiagaraEmitterInstance.cpp b/Engine/Plugins/FX/Niagara/Source/Niagara/Private/NiagaraEmitterInstance.cpp
--- a/Engine/Plugins/FX/Niagara/Source/Niagara/Private/NiagaraEmitterInstance.cpp
+++ b/Engine/Plugins/FX/Niagara/Source/Niagara/Private/NiagaraEmitterInstance.cpp
@@ -233,6 +233,22 @@ bool FNiagaraEmitterInstance::IsAllowedToExecute() const
 		{
 			return false;
 		}
+
+		if (const UNiagaraScript* GPUComputeScript = EmitterData->GetGPUComputeScript())
+		{
+			if (const FNiagaraShaderScript* ShaderScript = GPUComputeScript->GetRenderThreadScript())
+			{
+				const uint64 ScriptCBufferSize = ShaderScript->GetScriptParametersMetadata()->ShaderParametersMetadata->GetLayout().ConstantBufferSize;
+				const uint64 RHIMaxCBufferSize = GetMaxConstantBufferByteSize();
+				if (ScriptCBufferSize > RHIMaxCBufferSize)
+				{
+				#if !(UE_BUILD_SHIPPING || UE_BUILD_TEST)
+					GEngine->AddOnScreenDebugMessage(uint64(EmitterData), 1.f, FColor::Red, *FString::Printf(TEXT("GPU Simulation(%s) is disabled due to using too much constant buffer space (%d/%d)."), EmitterData->GetDebugSimName(), ScriptCBufferSize, RHIMaxCBufferSize));
+				#endif
+					return false;
+				}
+			}
+		}
 	}
 
 	if (UNiagaraComponentSettings::ShouldSuppressEmitterActivation(this))
