Improved User Parameters Tab.
Renaming via Context Menu
Deleting via Context Menu

Refactoring of the user param details code so that it is easier to maintain & improve

Removed "User" section in Palette tab

Added "OnParameterRenamed" editor only delegate to the param store so that dependent code can rely on it. The dichtomy between the "ParameterPANELViewModel" handling both parameter internal logic & handling UI interferes with the logic the User Param Tab is supposed to handle. i.e. "ParameterPanelViewModel->RenameParameter" will also attempt to select the parameter we just renamed in the parameter palette tab, so we can't really use that to handle our renaming.

Clicking "user parameters" in the overview graph now summons the user parameter tab

diff --git a/Engine/Plugins/FX/Niagara/Source/Niagara/Private/NiagaraParameterStore.cpp b/Engine/Plugins/FX/Niagara/Source/Niagara/Private/NiagaraParameterStore.cpp
--- a/Engine/Plugins/FX/Niagara/Source/Niagara/Private/NiagaraParameterStore.cpp
+++ b/Engine/Plugins/FX/Niagara/Source/Niagara/Private/NiagaraParameterStore.cpp
@@ -813,6 +813,9 @@ void FNiagaraParameterStore::RenameParameter(const FNiagaraVariableBase& Param,
 			RemoveParameter(Param);
 
 			OnLayoutChange();
+#if WITH_EDITORONLY_DATA
+			OnParameterRenamedDelegate.Broadcast(Param, NewName);
+#endif
 		}
 		else
 		{
diff --git a/Engine/Plugins/FX/Niagara/Source/Niagara/Public/NiagaraParameterStore.h b/Engine/Plugins/FX/Niagara/Source/Niagara/Public/NiagaraParameterStore.h
--- a/Engine/Plugins/FX/Niagara/Source/Niagara/Public/NiagaraParameterStore.h
+++ b/Engine/Plugins/FX/Niagara/Source/Niagara/Public/NiagaraParameterStore.h
@@ -176,6 +176,7 @@ struct NIAGARA_API FNiagaraParameterStore
 	};
 	
 	DECLARE_MULTICAST_DELEGATE(FOnStructureChanged)
+	DECLARE_MULTICAST_DELEGATE_TwoParams(FOnParameterRenamed, FNiagaraVariable OldVariable, FName NewName)
 #endif
 
 	GENERATED_USTRUCT_BODY()
@@ -240,6 +241,7 @@ private:
 	FOnChanged OnChangedDelegate;
 	bool bSuppressOnChanged = false;
 	FOnStructureChanged OnStructureChangedDelegate;
+	FOnParameterRenamed OnParameterRenamedDelegate;
 #endif
 
 	void SetPositionData(const FName& Name, const FVector& Position);
@@ -669,6 +671,7 @@ public:
 	void RemoveAllOnChangedHandlers(const void* InUserObject);
 
 	FOnStructureChanged& OnStructureChanged() { return OnStructureChangedDelegate; }
+	FOnParameterRenamed& OnParameterRenamed() { return OnParameterRenamedDelegate; }
 #endif
 
 	void TriggerOnLayoutChanged() { OnLayoutChange(); }
diff --git a/Engine/Plugins/FX/Niagara/Source/NiagaraEditor/Private/Customizations/NiagaraComponentDetails.cpp b/Engine/Plugins/FX/Niagara/Source/NiagaraEditor/Private/Customizations/NiagaraComponentDetails.cpp
--- a/Engine/Plugins/FX/Niagara/Source/NiagaraEditor/Private/Customizations/NiagaraComponentDetails.cpp
+++ b/Engine/Plugins/FX/Niagara/Source/NiagaraEditor/Private/Customizations/NiagaraComponentDetails.cpp
@@ -37,6 +37,7 @@
 #include "Widgets/SNiagaraParameterName.h"
 #include "ViewModels/HierarchyEditor/NiagaraUserParametersHierarchyViewModel.h"
 #include "Widgets/Input/SCheckBox.h"
+#include "SNiagaraSystemUserParameters.h"
 
 #define LOCTEXT_NAMESPACE "NiagaraComponentDetails"
 
@@ -147,7 +148,7 @@ public:
 		}
 	}
 
-	EVisibility GetResetToDefaultVisibility() const 
+	EVisibility GetResetToDefaultVisibility() const
 	{
 		TWeakObjectPtr<UNiagaraComponent> RawComponent = GetComponent();
 		if (RawComponent.IsValid())
@@ -180,7 +181,7 @@ public:
 			}
 		}
 	}
-	
+
 	void OnParameterChanged(const FPropertyChangedEvent& PropertyChangedEvent)
 	{
 		if(Owner->IsA<UNiagaraComponent>())
@@ -208,7 +209,16 @@ public:
 			{
 				if(ParameterValue.GetMode() == ENiagaraVariantMode::Bytes)
 				{
-					System->GetExposedParameters().SetParameterData(ParameterValue.GetBytes(), ParameterKey);
+					// we have to special case the position definition because, while SetParameterData does handle Position variables, it expects 12 Bytes of data
+					if(ParameterKey.GetType() == FNiagaraTypeDefinition::GetPositionDef())
+					{
+						FVector Value = *reinterpret_cast<const FVector*>(ParameterValue.GetBytes());
+						System->GetExposedParameters().SetPositionParameterValue(Value, ParameterKey.GetName());
+					}
+					else
+					{
+						System->GetExposedParameters().SetParameterData(ParameterValue.GetBytes(), ParameterKey);
+					}
 				}
 				else if(ParameterValue.GetMode() == ENiagaraVariantMode::DataInterface)
 				{
@@ -312,8 +322,7 @@ public:
 
 		return CurrentObject != nullptr ? CurrentObject->GetPathName() : FString();
 	}
-
-
+	
 	const FNiagaraVariableBase& Key() const { return ParameterKey; }
 	FNiagaraVariant& Value() { return ParameterValue; }
 
@@ -328,21 +337,51 @@ private:
 	FSimpleDelegate OnRebuild;
 	bool bResettingToDefault;
 };
+#undef LOCTEXT_NAMESPACE
 
-void AddRowForUserParameter(TSharedPtr<FNiagaraParameterProxy> ParameterProxy, IDetailChildrenBuilder& ChildrenBuilder, TMap<FName, TWeakPtr<FStructOnScope>>& ParameterNameToDisplayStruct)
+#define LOCTEXT_NAMESPACE "NiagaraUserParameterDetails"
+
+void FNiagaraUserParameterNodeBuilder::GenerateRowForUserParameter(IDetailChildrenBuilder& ChildrenBuilder, FNiagaraVariable UserParameter, UObject* Owner)
 {
-	FNiagaraVariable Parameter = ParameterProxy->Key();
+	FNiagaraVariable ChoppedUserParameter(UserParameter);
+
+	if(UserParameter.IsInNameSpace(FNiagaraConstants::UserNamespaceString))
+	{
+		ChoppedUserParameter.SetName(FName(UserParameter.GetName().ToString().RightChop(5)));
+	}
+	
+	FNiagaraVariant ParameterValue = GetCurrentParameterValue(ChoppedUserParameter, Owner);
+	if (!ParameterValue.IsValid())
+	{
+		return;
+	}
+		
+	if (UserParameter.IsDataInterface())
+	{
+		if(Owner->IsA<UNiagaraComponent>())
+		{
+			ParameterValue = FNiagaraVariant(DuplicateObject(ParameterValue.GetDataInterface(), Owner));
+		}
+		else if(Owner->IsA<UNiagaraSystem>())
+		{
+			ParameterValue = FNiagaraVariant(ParameterValue.GetDataInterface());
+		}
+	}
+
+	TSharedPtr<FNiagaraParameterProxy> ParameterProxy = ParameterProxies.Add_GetRef(MakeShared<FNiagaraParameterProxy>(Owner, UserParameter, ParameterValue, OnRebuildChildren, OverridePropertyHandles));
+
 	IDetailPropertyRow* Row = nullptr;
 
+	TSharedPtr<SWidget> CustomNameWidget = GenerateCustomNameWidget(UserParameter);
  	TSharedPtr<SWidget> CustomValueWidget;
-			
-	if (Parameter.IsDataInterface())
+	
+	if (UserParameter.IsDataInterface())
 	{
 		// if no changes are made, then it'll just be the same as the asset
 		TArray<UObject*> Objects { ParameterProxy->Value().GetDataInterface() };
 	
 		FAddPropertyParams Params = FAddPropertyParams()
-			.UniqueId(Parameter.GetName())
+			.UniqueId(ChoppedUserParameter.GetName())
 			.AllowChildren(true)
 			.CreateCategoryNodes(false);
 	
@@ -351,22 +390,22 @@ void AddRowForUserParameter(TSharedPtr<FNiagaraParameterProxy> ParameterProxy, I
 		CustomValueWidget =
 			SNew(STextBlock)
 			.TextStyle(FNiagaraEditorStyle::Get(), "NiagaraEditor.ParameterText")
-			.Text(FText::FromString(FName::NameToDisplayString(Parameter.GetType().GetClass()->GetName(), false)));
+			.Text(FText::FromString(FName::NameToDisplayString(UserParameter.GetType().GetClass()->GetName(), false)));
 	}
-	else if(Parameter.IsUObject())
+	else if(UserParameter.IsUObject())
 	{
 		TArray<UObject*> Objects { ParameterProxy->Value().GetUObject() };
 
 		FAddPropertyParams Params = FAddPropertyParams()
-			.UniqueId(Parameter.GetName())
+			.UniqueId(ChoppedUserParameter.GetName())
 			.AllowChildren(false) // Don't show the material's properties
 			.CreateCategoryNodes(false);
 
 		Row = ChildrenBuilder.AddExternalObjectProperty(Objects, NAME_None, Params);
-
+		
 		CustomValueWidget = SNew(SObjectPropertyEntryBox)
 			.ObjectPath(ParameterProxy.ToSharedRef(), &FNiagaraParameterProxy::GetCurrentAssetPath)
-			.AllowedClass(Parameter.GetType().GetClass())
+			.AllowedClass(UserParameter.GetType().GetClass())
 			.OnObjectChanged(ParameterProxy.ToSharedRef(), &FNiagaraParameterProxy::OnAssetSelectedFromPicker)
 			.AllowClear(false)
 			.DisplayUseSelected(true)
@@ -376,7 +415,7 @@ void AddRowForUserParameter(TSharedPtr<FNiagaraParameterProxy> ParameterProxy, I
 	}
 	else
 	{
-		FNiagaraTypeDefinition Type = Parameter.GetType();
+		FNiagaraTypeDefinition Type = UserParameter.GetType();
 		UStruct* Struct = Type.GetStruct();
 			
 		if (Type == FNiagaraTypeDefinition::GetPositionDef())
@@ -389,27 +428,21 @@ void AddRowForUserParameter(TSharedPtr<FNiagaraParameterProxy> ParameterProxy, I
 		TSharedPtr<FStructOnScope> StructOnScope = MakeShareable(new FStructOnScope(Struct, ParameterProxy->Value().GetBytes()));
 
 		FAddPropertyParams Params = FAddPropertyParams()
-			.UniqueId(Parameter.GetName());
+			.UniqueId(ChoppedUserParameter.GetName());
 
 		Row = ChildrenBuilder.AddExternalStructureProperty(StructOnScope.ToSharedRef(), NAME_None, Params);
-		FNiagaraVariable UserParameter = Parameter;
-		FNiagaraUserRedirectionParameterStore::MakeUserVariable(UserParameter);
 		ParameterNameToDisplayStruct.Add(UserParameter.GetName(), TWeakPtr<FStructOnScope>(StructOnScope));
 	}
 
 	TSharedPtr<SWidget> DefaultNameWidget;
 	TSharedPtr<SWidget> DefaultValueWidget;
 	
-	if(Parameter.IsInNameSpace(FNiagaraConstants::UserNamespaceString))
-	{
-		Parameter.SetName(FName(Parameter.GetName().ToString().RightChop(5)));
-	}
+	Row->DisplayName(FText::FromName(ChoppedUserParameter.GetName()));
 	
-	Row->DisplayName(FText::FromName(Parameter.GetName()));
-
 	FDetailWidgetRow& CustomWidget = Row->CustomWidget(true);
-
-	FNiagaraVariable UserParameter = Parameter;
+	
+	AddCustomMenuActionsForParameter(CustomWidget, UserParameter);
+	
 	UNiagaraSystem* System = nullptr;
 	if(ParameterProxy->GetSystem().IsValid())
 	{
@@ -422,39 +455,42 @@ void AddRowForUserParameter(TSharedPtr<FNiagaraParameterProxy> ParameterProxy, I
 
 	if(System)
 	{
-		System->GetExposedParameters().RedirectUserVariable(UserParameter);
 		UNiagaraScriptVariable* ScriptVariable = FNiagaraEditorUtilities::GetScriptVariableForUserParameter(UserParameter, *System);
 			
-		TAttribute<FText> Tooltip = FText::GetEmpty();
-		if(ScriptVariable)
-		{
-			Tooltip = TAttribute<FText>::CreateLambda([ScriptVariable]
-			{
-				return ScriptVariable->Metadata.Description;
-			});
-		}
-			
 		Row->GetDefaultWidgets(DefaultNameWidget, DefaultValueWidget, CustomWidget);
 
+		Row->GetPropertyHandle()->MarkResetToDefaultCustomized(true);
+		
 		Row->GetPropertyHandle()->SetOnPropertyValuePreChange(FSimpleDelegate::CreateSP(ParameterProxy.ToSharedRef(), &FNiagaraParameterProxy::OnParameterPreChange));
 		Row->GetPropertyHandle()->SetOnChildPropertyValuePreChange(FSimpleDelegate::CreateSP(ParameterProxy.ToSharedRef(), &FNiagaraParameterProxy::OnParameterPreChange));
 		Row->GetPropertyHandle()->SetOnPropertyValueChangedWithData(TDelegate<void(const FPropertyChangedEvent&)>::CreateSP(ParameterProxy.ToSharedRef(), &FNiagaraParameterProxy::OnParameterChanged));
 		Row->GetPropertyHandle()->SetOnChildPropertyValueChangedWithData(TDelegate<void(const FPropertyChangedEvent&)>::CreateSP(ParameterProxy.ToSharedRef(), &FNiagaraParameterProxy::OnParameterChanged));
 		Row->GetPropertyHandle()->SetOnPropertyResetToDefault(FSimpleDelegate::CreateSP(ParameterProxy.ToSharedRef(), &FNiagaraParameterProxy::OnResetToDefault));
-		
-		CustomWidget
-			.NameContent()
+
+		TSharedPtr<SWidget> NameWidget = SNew(SBox)
+			.Padding(FMargin(0.0f, 2.0f))
 			[
-				SNew(SBox)
-				.Padding(FMargin(0.0f, 2.0f))
-				[
-					SNew(STextBlock)
-					.TextStyle(FNiagaraEditorStyle::Get(), "NiagaraEditor.ParameterText")
-					.Text(FText::FromName(Parameter.GetName()))
-					.ToolTipText(Tooltip)
-				]
+				SNew(STextBlock)
+				.TextStyle(FNiagaraEditorStyle::Get(), "NiagaraEditor.ParameterText")
+				.Text(FText::FromName(ChoppedUserParameter.GetName()))
 			];
+		
+		if (CustomNameWidget.IsValid())
+		{
+			NameWidget = CustomNameWidget;
+		}
 
+		// regardless of which widget we use for the name, we want to use the script variable's description as tooltip
+		TAttribute<FText> Tooltip = FText::GetEmpty();
+		if(ScriptVariable)
+		{
+			Tooltip = TAttribute<FText>::CreateLambda([ScriptVariable]
+			{
+				return ScriptVariable->Metadata.Description;
+			});
+		}
+		NameWidget->SetToolTipText(Tooltip);
+		
 		TSharedPtr<SWidget> ValueWidget = DefaultValueWidget;
 		if (CustomValueWidget.IsValid())
 		{
@@ -462,6 +498,12 @@ void AddRowForUserParameter(TSharedPtr<FNiagaraParameterProxy> ParameterProxy, I
 		}
 		
 		CustomWidget
+		.FilterString(FText::FromName(ChoppedUserParameter.GetName()))
+		.NameContent()
+		.HAlign(HAlign_Left)
+		[
+			NameWidget.ToSharedRef()
+		]
 		.ValueContent()
 		[
 			SNew(SHorizontalBox)
@@ -493,15 +535,41 @@ void AddRowForUserParameter(TSharedPtr<FNiagaraParameterProxy> ParameterProxy, I
 	}
 }
 
+TSharedRef<IDetailCustomization> FNiagaraSystemUserParameterDetails::MakeInstance()
+{
+	return MakeShared<FNiagaraSystemUserParameterDetails>();
+}
+
+void FNiagaraSystemUserParameterDetails::CustomizeDetails(IDetailLayoutBuilder& DetailBuilder)
+{
+	TArray<TWeakObjectPtr<UObject>> CustomizedObjects;
+	DetailBuilder.GetObjectsBeingCustomized(CustomizedObjects);
+	
+	TArray<FName> Categories;
+	DetailBuilder.GetCategoryNames(Categories);
+
+	for(FName& CategoryName : Categories)
+	{
+		DetailBuilder.HideCategory(CategoryName);
+	}
+	
+	static const FName ParamCategoryName = TEXT("NiagaraSystem_UserParameters");
+
+	if(CustomizedObjects.Num() == 1 && CustomizedObjects[0]->IsA<UNiagaraSystem>())
+	{
+		System = Cast<UNiagaraSystem>(CustomizedObjects[0]);
+		IDetailCategoryBuilder& InputParamCategory = DetailBuilder.EditCategory(ParamCategoryName, LOCTEXT("ParamCategoryName", "User Parameters"), ECategoryPriority::Important);
+		InputParamCategory.AddCustomBuilder(MakeShared<FNiagaraSystemUserParameterBuilder>(System.Get(), ParamCategoryName));
+	}
+}
+
+/** The category builder will display all sub-categories & parameters contained within a given category in the details panel. */
 class FNiagaraUserParameterCategoryBuilder : public IDetailCustomNodeBuilder, public TSharedFromThis<FNiagaraUserParameterCategoryBuilder>
 {
 public:
-	FNiagaraUserParameterCategoryBuilder(UObject* InOwner, const UNiagaraHierarchyCategory* InNiagaraHierarchyCategory, FName InCustomBuilderRowName,
-		TArray<TSharedPtr<FNiagaraParameterProxy>>& InParameterProxies, TMap<FName, TWeakPtr<FStructOnScope>>& InParameterNameToDisplayStruct,
-		TArray<TSharedPtr<IPropertyHandle>> InOverridePropertyHandles) : ParameterProxies(InParameterProxies), ParameterNameToDisplayStruct(InParameterNameToDisplayStruct)
+	FNiagaraUserParameterCategoryBuilder(UObject* InOwner, const UNiagaraHierarchyCategory* InNiagaraHierarchyCategory, FName InCustomBuilderRowName, FNiagaraUserParameterNodeBuilder& InNodeBuilder) : NodeBuilder(InNodeBuilder) 
 	{
 		HierarchyCategory = InNiagaraHierarchyCategory;
-		OverridePropertyHandles = InOverridePropertyHandles;
 		CustomBuilderRowName = InCustomBuilderRowName;
 
 		Owner = InOwner;
@@ -513,12 +581,13 @@ public:
 		[
 			SNew(STextBlock)
 			.TextStyle(FNiagaraEditorStyle::Get(), "NiagaraEditor.ParameterText")
-			.Text(FText::FromName(CustomBuilderRowName))
+			.Text_UObject(HierarchyCategory.Get(), &UNiagaraHierarchyCategory::GetCategoryAsText)
 			.ToolTipText_UObject(HierarchyCategory.Get(), &UNiagaraHierarchyCategory::GetTooltip)
 		];
 
 		NodeRow.FilterString(FText::FromName(CustomBuilderRowName));
 	}
+	
 	virtual void Tick(float DeltaTime) override {}
 	virtual bool RequiresTick() const override { return false; }
 	virtual bool InitiallyCollapsed() const override { return false; }
@@ -551,39 +620,16 @@ public:
 
 		for(const UNiagaraHierarchyCategory* Category : ContainedCategories)
 		{
-			if(Category->DoesOneChildExist<UNiagaraHierarchyUserParameter>())
+			if(Category->DoesOneChildExist<UNiagaraHierarchyUserParameter>(true))
 			{
 				ChildrenBuilder.AddCustomBuilder(MakeShared<FNiagaraUserParameterCategoryBuilder>(
-					Owner.Get(),Category, Category->GetCategoryName(),
-					ParameterProxies, ParameterNameToDisplayStruct,
-					OverridePropertyHandles));
+					Owner.Get(),Category, Category->GetCategoryName(), NodeBuilder));
 			}
 		}
 		
 		for(const UNiagaraHierarchyUserParameter* HierarchyUserParameter : ContainedParameters)
 		{
-			FNiagaraVariable Parameter = HierarchyUserParameter->GetUserParameter();
-			FNiagaraVariant ParameterValue = GetCurrentParameterValue(Parameter, Owner.Get());
-			if (!ParameterValue.IsValid())
-			{
-				continue;
-			}
-
-			// the parameter value should operate on a duplicate for components and on the contained DI for a system
-			if (Parameter.IsDataInterface())
-			{
-				if(Owner->IsA<UNiagaraComponent>())
-				{
-					ParameterValue = FNiagaraVariant(DuplicateObject(ParameterValue.GetDataInterface(), GetComponent().Get()));
-				}
-				else if(Owner->IsA<UNiagaraSystem>())
-				{
-					ParameterValue = FNiagaraVariant(ParameterValue.GetDataInterface());
-				}
-			}
-			
-			TSharedPtr<FNiagaraParameterProxy> ParameterProxy = ParameterProxies.Add_GetRef(MakeShareable(new FNiagaraParameterProxy(Owner, Parameter, ParameterValue, OnRebuildChildren, OverridePropertyHandles)));
-			AddRowForUserParameter(ParameterProxy, ChildrenBuilder, ParameterNameToDisplayStruct);
+			NodeBuilder.GenerateRowForUserParameter(ChildrenBuilder, HierarchyUserParameter->GetUserParameter(), Owner.Get());
 		}
 	}
 
@@ -593,15 +639,11 @@ public:
 private:
 	TWeakObjectPtr<UObject> Owner;
 	TWeakObjectPtr<const UNiagaraHierarchyCategory> HierarchyCategory;
-	TArray<TSharedPtr<IPropertyHandle>> OverridePropertyHandles;
-	TArray<TSharedPtr<FNiagaraParameterProxy>>& ParameterProxies;
-	TMap<FName, TWeakPtr<FStructOnScope>>& ParameterNameToDisplayStruct;
-	FSimpleDelegate OnRebuildChildren;
 	FName CustomBuilderRowName;
+	FNiagaraUserParameterNodeBuilder& NodeBuilder;
 };
 
-void GenerateUserParameterRows(IDetailChildrenBuilder& ChildrenBuilder, UNiagaraHierarchyRoot& UserParameterHierarchyRoot, UObject* Owner, const UNiagaraHierarchySection*& ActiveSection,
-	TArray<TSharedPtr<FNiagaraParameterProxy>>& ParameterProxies, TMap<FName, TWeakPtr<FStructOnScope>>& ParameterNameToDisplayStruct, TArray<TSharedPtr<IPropertyHandle>> OverridePropertyHandles, FSimpleDelegate OnRebuildChildren)
+void FNiagaraUserParameterNodeBuilder::GenerateUserParameterRows(IDetailChildrenBuilder& ChildrenBuilder, UNiagaraHierarchyRoot& UserParameterHierarchyRoot, UObject* Owner)
 {
 	UNiagaraSystem* SystemAsset = nullptr;
 
@@ -618,7 +660,8 @@ void GenerateUserParameterRows(IDetailChildrenBuilder& ChildrenBuilder, UNiagara
 	{
 		return;
 	}
-	
+
+	TArray<const UNiagaraHierarchySection*> OrderedSections;
 	TMap<const UNiagaraHierarchySection*, TArray<const UNiagaraHierarchyCategory*>> SectionCategoryMap;
 
 	UNiagaraHierarchyRoot* Root = &UserParameterHierarchyRoot;
@@ -637,14 +680,21 @@ void GenerateUserParameterRows(IDetailChildrenBuilder& ChildrenBuilder, UNiagara
 	if(SectionCategoryMap.Num() > 0)
 	{
 		TSharedRef<SHorizontalBox> SectionsBox = SNew(SHorizontalBox);
-		
-		for(auto& SectionCategories : SectionCategoryMap)
+
+		// maps don't guarantee order, so we iterate over the original section data instead
+		for(auto& Section : Root->GetSectionData())
 		{
-			const UNiagaraHierarchySection* Section = SectionCategories.Key;
-			TArray<const UNiagaraHierarchyCategory*>& HierarchyCategories = SectionCategories.Value;
+			// if the map doesn't contain at the entry, it means there are 0 categories for that section, so we skip it
+			if(!SectionCategoryMap.Contains(Section))
+			{
+				continue;
+			}
+			
+			TArray<const UNiagaraHierarchyCategory*>& HierarchyCategories = SectionCategoryMap[Section];
 			
 			if(HierarchyCategories.Num() > 0)
 			{
+				// we only want to add a section if any of its categories contain at least one user parameter
 				bool bDoesUserParamExist = false;
 				for(const UNiagaraHierarchyCategory* Category : HierarchyCategories)
 				{
@@ -667,7 +717,7 @@ void GenerateUserParameterRows(IDetailChildrenBuilder& ChildrenBuilder, UNiagara
 						[
 							SNew(SCheckBox)
 							.Style(FAppStyle::Get(), "DetailsView.SectionButton")
-							.OnCheckStateChanged_Lambda([OnRebuildChildren, &ActiveSection, Section](ECheckBoxState NewState)
+							.OnCheckStateChanged_Lambda([this, Section](ECheckBoxState NewState)
 							{
 								ActiveSection = Section;
 								OnRebuildChildren.ExecuteIfBound();
@@ -700,7 +750,7 @@ void GenerateUserParameterRows(IDetailChildrenBuilder& ChildrenBuilder, UNiagara
 				[
 					SNew(SCheckBox)
 					.Style(FAppStyle::Get(), "DetailsView.SectionButton")
-					.OnCheckStateChanged_Lambda([&ActiveSection, OnRebuildChildren](ECheckBoxState NewState)
+					.OnCheckStateChanged_Lambda([this](ECheckBoxState NewState)
 					{
 						ActiveSection = nullptr;
 						OnRebuildChildren.ExecuteIfBound();
@@ -724,18 +774,25 @@ void GenerateUserParameterRows(IDetailChildrenBuilder& ChildrenBuilder, UNiagara
 			];
 		}
 	}
-
+	
 	// we add the categories now if the active section is set to all (nullptr) or if the section it belongs to is active
 	for(const UNiagaraHierarchyItemBase* Child : Root->GetChildren())
-	{
+	{		
 		if(const UNiagaraHierarchyCategory* HierarchyCategory = Cast<UNiagaraHierarchyCategory>(Child))
 		{
-			if(HierarchyCategory->DoesOneChildExist<UNiagaraHierarchyUserParameter>(true) && (ActiveSection == nullptr || HierarchyCategory->GetSection() == ActiveSection))
+			if((ActiveSection == nullptr || HierarchyCategory->GetSection() == ActiveSection) && HierarchyCategory->DoesOneChildExist<UNiagaraHierarchyUserParameter>(true))
 			{
 				ChildrenBuilder.AddCustomBuilder(MakeShared<FNiagaraUserParameterCategoryBuilder>(
-					Owner, HierarchyCategory, HierarchyCategory->GetCategoryName(),
-					ParameterProxies, ParameterNameToDisplayStruct,
-					OverridePropertyHandles));
+					Owner, HierarchyCategory, HierarchyCategory->GetCategoryName(), *this));
+			
+			}
+		}
+		else if(const UNiagaraHierarchyUserParameter* UserParameter = Cast<UNiagaraHierarchyUserParameter>(Child))
+		{
+			// we only want to display parameters within the root in the all section
+			if(ActiveSection == nullptr)
+			{
+				GenerateRowForUserParameter(ChildrenBuilder, UserParameter->GetUserParameter(), Owner);
 			}
 		}
 	}
@@ -764,309 +821,253 @@ void GenerateUserParameterRows(IDetailChildrenBuilder& ChildrenBuilder, UNiagara
 		}
 
 		for(FNiagaraVariable& LeftoverUserParameter : LeftoverUserParameters)
-		{
-			FNiagaraVariant ParameterValue = GetCurrentParameterValue(LeftoverUserParameter, Owner);
-			if (!ParameterValue.IsValid())
-			{
-				continue;
-			}
-		
-			if (LeftoverUserParameter.IsDataInterface())
-			{
-				if(Owner->IsA<UNiagaraComponent>())
-				{
-					ParameterValue = FNiagaraVariant(DuplicateObject(ParameterValue.GetDataInterface(), Owner));
-				}
-				else if(Owner->IsA<UNiagaraSystem>())
-				{
-					ParameterValue = FNiagaraVariant(ParameterValue.GetDataInterface());
-				}
-			}
-			
-			TSharedPtr<FNiagaraParameterProxy> ParameterProxy = ParameterProxies.Add_GetRef(MakeShareable(new FNiagaraParameterProxy(Owner, LeftoverUserParameter, ParameterValue, OnRebuildChildren, OverridePropertyHandles)));
-			AddRowForUserParameter(ParameterProxy, ChildrenBuilder, ParameterNameToDisplayStruct);
+		{			
+			GenerateRowForUserParameter(ChildrenBuilder, LeftoverUserParameter, Owner);
 		}
 	}
 }
 
-class FNiagaraComponentNodeBuilder : public IDetailCustomNodeBuilder, public TSharedFromThis<FNiagaraComponentNodeBuilder>
+void FNiagaraComponentUserParametersNodeBuilder::GenerateChildContent(IDetailChildrenBuilder& ChildrenBuilder)
 {
-public:
-	FNiagaraComponentNodeBuilder(UNiagaraComponent* InComponent, TArray<TSharedPtr<IPropertyHandle>> InOverridePropertyHandles, FName InCustomBuilderRowName) 
+	if (bDelegatesInitialized == false)
 	{
-		OverridePropertyHandles = InOverridePropertyHandles;
-		CustomBuilderRowName = InCustomBuilderRowName;
+		Component->OnSynchronizedWithAssetParameters().AddSP(this, &FNiagaraComponentUserParametersNodeBuilder::Rebuild);
+		Component->GetOverrideParameters().AddOnChangedHandler(
+			FNiagaraParameterStore::FOnChanged::FDelegate::CreateSP(this, &FNiagaraComponentUserParametersNodeBuilder::ParameterValueChanged));
 
-		Component = InComponent;
-		bDelegatesInitialized = false;
+		RegisterRebuildOnHierarchyChanged();
+		bDelegatesInitialized = true;
 	}
 
-	virtual ~FNiagaraComponentNodeBuilder() override
-	{
-		if (Component.IsValid() && bDelegatesInitialized)
-		{
-			Component->OnSynchronizedWithAssetParameters().RemoveAll(this);
-			Component->GetOverrideParameters().RemoveAllOnChangedHandlers(this);
+	check(Component.IsValid());
 
-			if(UNiagaraSystem* System = Component->GetAsset())
-			{
-				TSharedPtr<FNiagaraSystemViewModel> SystemViewModel = TNiagaraViewModelManager<UNiagaraSystem, FNiagaraSystemViewModel>::GetExistingViewModelForObject(System);
-				if(SystemViewModel.IsValid())
-				{
-					SystemViewModel->GetUserParametersHierarchyViewModel()->OnHierarchyChanged().RemoveAll(this);
-				}
-			}
-		}
-	}
-
-	virtual void SetOnRebuildChildren(FSimpleDelegate InOnRegenerateChildren) override
-	{
-		OnRebuildChildren = InOnRegenerateChildren;
-	}
+	ParameterProxies.Reset();
 
-	virtual void GenerateHeaderRowContent(FDetailWidgetRow& NodeRow) override {}
-	virtual void Tick(float DeltaTime) override {}
-	virtual bool RequiresTick() const override { return false; }
-	virtual bool InitiallyCollapsed() const override { return false; }
-	virtual FName GetName() const  override
+	UNiagaraSystem* SystemAsset = Component->GetAsset();
+	if (SystemAsset == nullptr)
 	{
-		return CustomBuilderRowName;
+		return;
 	}
-
-	virtual void GenerateChildContent(IDetailChildrenBuilder& ChildrenBuilder) override
-	{
-		if (bDelegatesInitialized == false)
-		{
-			Component->OnSynchronizedWithAssetParameters().AddSP(this, &FNiagaraComponentNodeBuilder::Rebuild);
-			Component->GetOverrideParameters().AddOnChangedHandler(
-				FNiagaraParameterStore::FOnChanged::FDelegate::CreateSP(this, &FNiagaraComponentNodeBuilder::ParameterValueChanged));
-
-			RegisterRebuildOnHierarchyChanged();
-			bDelegatesInitialized = true;
-		}
-
-		check(Component.IsValid());
-
-		ParameterProxies.Reset();
-
-		UNiagaraSystem* SystemAsset = Component->GetAsset();
-		if (SystemAsset == nullptr)
-		{
-			return;
-		}
 		
-		TArray<FNiagaraVariable> UserParameters;
-		SystemAsset->GetExposedParameters().GetUserParameters(UserParameters);
+	TArray<FNiagaraVariable> UserParameters;
+	SystemAsset->GetExposedParameters().GetUserParameters(UserParameters);
 		
-		ParameterProxies.Reserve(UserParameters.Num());
+	ParameterProxies.Reserve(UserParameters.Num());
 
-		ParameterNameToDisplayStruct.Empty();
-
-		GenerateUserParameterRows(ChildrenBuilder, *Cast<UNiagaraSystemEditorData>(SystemAsset->GetEditorData())->UserParameterHierarchy.Get(),
-			Component.Get(), ActiveSection, ParameterProxies, ParameterNameToDisplayStruct, OverridePropertyHandles, OnRebuildChildren);
-	}
+	ParameterNameToDisplayStruct.Empty();
 
-private:
-
-	void Rebuild() const
-	{
-		OnRebuildChildren.ExecuteIfBound();
-	}
+	GenerateUserParameterRows(ChildrenBuilder, *Cast<UNiagaraSystemEditorData>(SystemAsset->GetEditorData())->UserParameterHierarchy.Get(),Component.Get());
+}
 
-	void RegisterRebuildOnHierarchyChanged()
+void FNiagaraComponentUserParametersNodeBuilder::RegisterRebuildOnHierarchyChanged()
+{
+	if(UNiagaraSystem* Asset = Component->GetAsset())
 	{
-		if(UNiagaraSystem* Asset = Component->GetAsset())
+		TSharedPtr<FNiagaraSystemViewModel> SystemViewModel = TNiagaraViewModelManager<UNiagaraSystem, FNiagaraSystemViewModel>::GetExistingViewModelForObject(Asset);
+		if(SystemViewModel.IsValid())
 		{
-			TSharedPtr<FNiagaraSystemViewModel> SystemViewModel = TNiagaraViewModelManager<UNiagaraSystem, FNiagaraSystemViewModel>::GetExistingViewModelForObject(Asset);
-			if(SystemViewModel.IsValid())
-			{
-				SystemViewModel->GetUserParametersHierarchyViewModel()->OnHierarchyChanged().RemoveAll(this);
-				SystemViewModel->GetUserParametersHierarchyViewModel()->OnHierarchyChanged().Add(UNiagaraHierarchyViewModelBase::FOnHierarchyChanged::FDelegate::CreateSP(this, &FNiagaraComponentNodeBuilder::Rebuild));
-			}
+			SystemViewModel->GetUserParametersHierarchyViewModel()->OnHierarchyChanged().RemoveAll(this);
+			SystemViewModel->GetUserParametersHierarchyViewModel()->OnHierarchyChanged().Add(UNiagaraHierarchyViewModelBase::FOnHierarchyChanged::FDelegate::CreateSP(this, &FNiagaraComponentUserParametersNodeBuilder::Rebuild));
 		}
 	}
-	
-	void ParameterValueChanged()
+}
+
+void FNiagaraComponentUserParametersNodeBuilder::ParameterValueChanged()
+{
+	if (Component.IsValid())
 	{
-		if (Component.IsValid())
+		const FNiagaraParameterStore& OverrideParameters = Component->GetOverrideParameters();
+		TArray<FNiagaraVariable> UserParameters;
+		OverrideParameters.GetParameters(UserParameters);
+		for (const FNiagaraVariable& UserParameter : UserParameters)
 		{
-			const FNiagaraParameterStore& OverrideParameters = Component->GetOverrideParameters();
-			TArray<FNiagaraVariable> UserParameters;
-			OverrideParameters.GetParameters(UserParameters);
-			for (const FNiagaraVariable& UserParameter : UserParameters)
+			if (UserParameter.IsUObject() == false)
 			{
-				if (UserParameter.IsUObject() == false)
+				TWeakPtr<FStructOnScope>* DisplayStructPtr = ParameterNameToDisplayStruct.Find(UserParameter.GetName());
+				if (DisplayStructPtr != nullptr && DisplayStructPtr->IsValid())
 				{
-					TWeakPtr<FStructOnScope>* DisplayStructPtr = ParameterNameToDisplayStruct.Find(UserParameter.GetName());
-					if (DisplayStructPtr != nullptr && DisplayStructPtr->IsValid())
+					TSharedPtr<FStructOnScope> DisplayStruct = DisplayStructPtr->Pin();
+					if (UserParameter.GetType() == FNiagaraTypeDefinition::GetPositionDef())
 					{
-						TSharedPtr<FStructOnScope> DisplayStruct = DisplayStructPtr->Pin();
-						if (UserParameter.GetType() == FNiagaraTypeDefinition::GetPositionDef())
-						{
-							FMemory::Memcpy(DisplayStruct->GetStructMemory(), OverrideParameters.GetPositionParameterValue(UserParameter.GetName()), UserParameter.GetSizeInBytes());
-						}
-						else
-						{
-							OverrideParameters.CopyParameterData(UserParameter, DisplayStruct->GetStructMemory());
-						}
+						FMemory::Memcpy(DisplayStruct->GetStructMemory(), OverrideParameters.GetPositionParameterValue(UserParameter.GetName()), UserParameter.GetSizeInBytes());
+					}
+					else
+					{
+						OverrideParameters.CopyParameterData(UserParameter, DisplayStruct->GetStructMemory());
 					}
 				}
 			}
 		}
 	}
+}
 
-	void OnActiveSectionChanged(ECheckBoxState State, const UNiagaraHierarchySection* NewSelection)
-	{
-		ActiveSection = NewSelection;
-		OnRebuildChildren.Execute();
-	}
-
-	ECheckBoxState IsSectionActive(const UNiagaraHierarchySection* Section) const
-	{
-		return ActiveSection == Section ? ECheckBoxState::Checked : ECheckBoxState::Unchecked;
-	}
-
-private:
-	bool bDelegatesInitialized;
-	TWeakObjectPtr<UNiagaraComponent> Component;
-	TArray<TSharedPtr<IPropertyHandle>> OverridePropertyHandles;
-	FSimpleDelegate OnRebuildChildren;
-	TArray<TSharedPtr<FNiagaraParameterProxy>> ParameterProxies;
-	TMap<FName, TWeakPtr<FStructOnScope>> ParameterNameToDisplayStruct;
-	FName CustomBuilderRowName;
-	const UNiagaraHierarchySection* ActiveSection = nullptr;
-	TArray<UNiagaraHierarchySection*> AvailableSections;
-};
-
-class FNiagaraSystemUserParameterBuilder : public IDetailCustomNodeBuilder, public TSharedFromThis<FNiagaraSystemUserParameterBuilder>
+FNiagaraSystemUserParameterBuilder::FNiagaraSystemUserParameterBuilder(UNiagaraSystem* InSystem, FName InCustomBuilderRowName)
 {
-public:
-	FNiagaraSystemUserParameterBuilder(UNiagaraSystem* InSystem, FName InCustomBuilderRowName) 
-	{
-		CustomBuilderRowName = InCustomBuilderRowName;
-
-		System = InSystem;
-		bDelegatesInitialized = false;
-	}
+	CustomBuilderRowName = InCustomBuilderRowName;
+	System = InSystem;
+	bDelegatesInitialized = false;
+}
 
-	virtual ~FNiagaraSystemUserParameterBuilder() override
-	{
-		if (System.IsValid() && bDelegatesInitialized)
+TSharedPtr<SWidget> FNiagaraSystemUserParameterBuilder::GenerateCustomNameWidget(FNiagaraVariable UserParameter)
+{
+	TSharedPtr<FNiagaraSystemViewModel> SystemViewModel = TNiagaraViewModelManager<UNiagaraSystem, FNiagaraSystemViewModel>::GetExistingViewModelForObject(System.Get());
+	
+	TSharedRef<SNiagaraParameterNameTextBlock> ParameterName = SNew(SNiagaraParameterNameTextBlock)
+		// if this is specified, we avoid the default behavior, which is "when clicked when the widget had keyboard focus already, enter editing mode"
+		// we don't want to be able to rename by clicking, only by context menu
+		.IsSelected_Lambda([]()
 		{
-			System->GetExposedParameters().OnStructureChanged().RemoveAll(this);
-
-			TSharedPtr<FNiagaraSystemViewModel> SystemViewModel = TNiagaraViewModelManager<UNiagaraSystem, FNiagaraSystemViewModel>::GetExistingViewModelForObject(System.Get());
-			if(SystemViewModel.IsValid())
-			{
-				SystemViewModel->GetUserParametersHierarchyViewModel()->OnHierarchyChanged().RemoveAll(this);
-			}
-		}
-	}
-
-	void Rebuild() const
-	{
-		OnRebuildChildren.ExecuteIfBound();
-	}
+			return false;
+		})
+		.ParameterText(FText::FromName(UserParameter.GetName()))
+		.OnDragDetected(this, &FNiagaraSystemUserParameterBuilder::GenerateParameterDragDropOp, UserParameter)
+		.OnTextCommitted_Lambda([this, UserParameter](const FText& InText, ETextCommit::Type Type)
+		{
+			RenameParameter(UserParameter, FName(InText.ToString()));
+		});
 	
-	virtual void SetOnRebuildChildren(FSimpleDelegate InOnRegenerateChildren) override
-	{
-		OnRebuildChildren = InOnRegenerateChildren;
-	}
+	UserParamToWidgetMap.Add(UserParameter, ParameterName);
+	
+	return ParameterName;
+}
 
-	virtual void GenerateHeaderRowContent(FDetailWidgetRow& NodeRow) override {}
-	virtual void Tick(float DeltaTime) override {}
-	virtual bool RequiresTick() const override { return false; }
-	virtual bool InitiallyCollapsed() const override { return false; }
-	virtual FName GetName() const  override
-	{
-		return CustomBuilderRowName;
-	}
+void FNiagaraSystemUserParameterBuilder::GenerateChildContent(IDetailChildrenBuilder& ChildrenBuilder)
+{
+	check(System.IsValid());
 
-	virtual void GenerateChildContent(IDetailChildrenBuilder& ChildrenBuilder) override
+	if (System.IsValid())
 	{
-		check(System.IsValid());
-
-		if (System.IsValid())
+		if(bDelegatesInitialized == false)
 		{
-			if(bDelegatesInitialized == false)
-			{
-				System->GetExposedParameters().OnStructureChanged().Add(FNiagaraParameterStore::FOnStructureChanged::FDelegate::CreateSP(this, &FNiagaraSystemUserParameterBuilder::Rebuild));
-				TSharedPtr<FNiagaraSystemViewModel> SystemViewModel = TNiagaraViewModelManager<UNiagaraSystem, FNiagaraSystemViewModel>::GetExistingViewModelForObject(System.Get());
-				SystemViewModel->GetUserParametersHierarchyViewModel()->OnHierarchyChanged().Add(UNiagaraHierarchyViewModelBase::FOnHierarchyChanged::FDelegate::CreateSP(this, &FNiagaraSystemUserParameterBuilder::Rebuild));
+			System->GetExposedParameters().OnStructureChanged().Add(FNiagaraParameterStore::FOnStructureChanged::FDelegate::CreateSP(this, &FNiagaraSystemUserParameterBuilder::Rebuild));
+			TSharedPtr<FNiagaraSystemViewModel> SystemViewModel = TNiagaraViewModelManager<UNiagaraSystem, FNiagaraSystemViewModel>::GetExistingViewModelForObject(System.Get());
+			SystemViewModel->GetUserParametersHierarchyViewModel()->OnHierarchyChanged().Add(UNiagaraHierarchyViewModelBase::FOnHierarchyChanged::FDelegate::CreateSP(this, &FNiagaraSystemUserParameterBuilder::Rebuild));
 
-				bDelegatesInitialized = true;
+			TSharedPtr<FNiagaraUserParameterPanelViewModel> UserParameterPanelViewModel = SystemViewModel->GetUserParameterPanelViewModel();
+			if(ensure(UserParameterPanelViewModel.IsValid()))
+			{
+				UserParameterPanelViewModel->OnParameterAdded().BindSP(this, &FNiagaraSystemUserParameterBuilder::RequestRename);
 			}
 			
-			ParameterProxies.Reset();
+			bDelegatesInitialized = true;
+		}
 			
-			TArray<FNiagaraVariable> UserParameters;
-			System->GetExposedParameters().GetUserParameters(UserParameters);
-			ParameterProxies.Reserve(UserParameters.Num());
+		ParameterProxies.Reset();
+			
+		TArray<FNiagaraVariable> UserParameters;
+		System->GetExposedParameters().GetUserParameters(UserParameters);
+		ParameterProxies.Reserve(UserParameters.Num());
 
-			ParameterNameToDisplayStruct.Empty();
+		ParameterNameToDisplayStruct.Empty();
 
-			UNiagaraHierarchyRoot* Root = Cast<UNiagaraSystemEditorData>(System->GetEditorData())->UserParameterHierarchy;
-			GenerateUserParameterRows(ChildrenBuilder, *Root, System.Get(), ActiveSection, ParameterProxies, ParameterNameToDisplayStruct, {}, OnRebuildChildren);
-		}
+		UNiagaraHierarchyRoot* Root = Cast<UNiagaraSystemEditorData>(System->GetEditorData())->UserParameterHierarchy;
+		GenerateUserParameterRows(ChildrenBuilder, *Root, System.Get());
 	}
+}
 
-private:
+void FNiagaraSystemUserParameterBuilder::AddCustomMenuActionsForParameter(FDetailWidgetRow& WidgetRow, FNiagaraVariable UserParameter)
+{
+	FNiagaraUserParameterNodeBuilder::AddCustomMenuActionsForParameter(WidgetRow, UserParameter);
+	
+	FUIAction RenameAction(FExecuteAction::CreateSP(this, &FNiagaraSystemUserParameterBuilder::RequestRename, UserParameter));
+	WidgetRow.AddCustomContextMenuAction(RenameAction, LOCTEXT("RenameParameterAction", "Rename"), LOCTEXT("RenameParameterActionTooltip", "Rename this user parameter"));
 
-	void ComponentSynchronizedWithAssetParameters()
-	{
-		OnRebuildChildren.ExecuteIfBound();
-	}
+	FUIAction DeleteAction(FExecuteAction::CreateSP(this, &FNiagaraSystemUserParameterBuilder::DeleteParameter, UserParameter));
+	WidgetRow.AddCustomContextMenuAction(DeleteAction, LOCTEXT("DeleteParameterAction", "Delete"), LOCTEXT("DeleteParameterActionTooltip", "Delete this user parameter"));
+}
 
-	void ParameterValueChanged()
+void FNiagaraSystemUserParameterBuilder::ParameterValueChanged()
+{
+	if (System.IsValid())
 	{
-		if (System.IsValid())
+		const FNiagaraParameterStore& ExposedParameters = System->GetExposedParameters();
+		TArray<FNiagaraVariable> UserParameters;
+		ExposedParameters.GetParameters(UserParameters);
+		for (const FNiagaraVariable& UserParameter : UserParameters)
 		{
-			const FNiagaraParameterStore& ExposedParameters = System->GetExposedParameters();
-			TArray<FNiagaraVariable> UserParameters;
-			ExposedParameters.GetParameters(UserParameters);
-			for (const FNiagaraVariable& UserParameter : UserParameters)
+			if (UserParameter.IsUObject() == false)
 			{
-				if (UserParameter.IsUObject() == false)
+				TWeakPtr<FStructOnScope>* DisplayStructPtr = ParameterNameToDisplayStruct.Find(UserParameter.GetName());
+				if (DisplayStructPtr != nullptr && DisplayStructPtr->IsValid())
 				{
-					TWeakPtr<FStructOnScope>* DisplayStructPtr = ParameterNameToDisplayStruct.Find(UserParameter.GetName());
-					if (DisplayStructPtr != nullptr && DisplayStructPtr->IsValid())
+					TSharedPtr<FStructOnScope> DisplayStruct = DisplayStructPtr->Pin();
+					if (UserParameter.GetType() == FNiagaraTypeDefinition::GetPositionDef())
 					{
-						TSharedPtr<FStructOnScope> DisplayStruct = DisplayStructPtr->Pin();
-						if (UserParameter.GetType() == FNiagaraTypeDefinition::GetPositionDef())
-						{
-							FMemory::Memcpy(DisplayStruct->GetStructMemory(), ExposedParameters.GetPositionParameterValue(UserParameter.GetName()), UserParameter.GetSizeInBytes());
-						}
-						else
-						{
-							ExposedParameters.CopyParameterData(UserParameter, DisplayStruct->GetStructMemory());
-						}
+						FMemory::Memcpy(DisplayStruct->GetStructMemory(), ExposedParameters.GetPositionParameterValue(UserParameter.GetName()), UserParameter.GetSizeInBytes());
+					}
+					else
+					{
+						ExposedParameters.CopyParameterData(UserParameter, DisplayStruct->GetStructMemory());
 					}
 				}
 			}
 		}
 	}
+}
 
-	void OnActiveSectionChanged(ECheckBoxState State, const UNiagaraHierarchySection* NewSelection)
+void FNiagaraSystemUserParameterBuilder::DeleteParameter(FNiagaraVariable UserParameter) const
+{
+	FScopedTransaction Transaction(LOCTEXT("DeleteParameterTransaction", "Deleted user parameter"));
+	System->Modify();
+	
+	if(System->GetExposedParameters().IndexOf(UserParameter) != INDEX_NONE)
 	{
-		ActiveSection = NewSelection;
-		OnRebuildChildren.Execute();
+		System->GetExposedParameters().RemoveParameter(UserParameter);		
 	}
 
-	ECheckBoxState IsSectionActive(const UNiagaraHierarchySection* Section) const
+	Cast<UNiagaraSystemEditorData>(System->GetEditorData())->UserParameterHierarchy->Modify();
+	System->HandleVariableRemoved(UserParameter, true);
+}
+
+void FNiagaraSystemUserParameterBuilder::RequestRename(FNiagaraVariable UserParameter)
+{
+	SelectedParameter = UserParameter;
+
+	// we add a timer for next frame, as requesting a rename on a parameter that has just been created will not preselect the entire text. Waiting a frame fixes that.
+	UserParamToWidgetMap[UserParameter]->RegisterActiveTimer(0.f, FWidgetActiveTimerDelegate::CreateLambda([this, UserParameter](double CurrentTime, float DeltaTime)
+	{
+		UserParamToWidgetMap[UserParameter]->EnterEditingMode();
+		return EActiveTimerReturnType::Stop;
+	}));
+}
+
+void FNiagaraSystemUserParameterBuilder::RenameParameter(FNiagaraVariable UserParameter, FName NewName)
+{
+	if(UserParameter.GetName() == NewName)
 	{
-		return ActiveSection == Section ? ECheckBoxState::Checked : ECheckBoxState::Unchecked;
+		return;
 	}
+	
+	if(System->GetExposedParameters().IndexOf(UserParameter) != INDEX_NONE)
+	{		
+		TArray<FNiagaraVariable> Parameters;
+		System->GetExposedParameters().GetParameters(Parameters);
 
-private:
-	bool bDelegatesInitialized;
-	TWeakObjectPtr<UNiagaraSystem> System;
-	FSimpleDelegate OnRebuildChildren;
-	TArray<TSharedPtr<FNiagaraParameterProxy>> ParameterProxies;
-	TMap<FName, TWeakPtr<FStructOnScope>> ParameterNameToDisplayStruct;
-	FName CustomBuilderRowName;
-	const UNiagaraHierarchySection* ActiveSection = nullptr;
-	TArray<UNiagaraHierarchySection*> AvailableSections;
-};
+		TSet<FName> Names;
+		for(FNiagaraVariable& Parameter : Parameters)
+		{
+			Names.Add(Parameter.GetName());
+		}
+		
+		FName UniqueName = FNiagaraUtilities::GetUniqueName(NewName, Names);
+
+		FNiagaraVariable VariableWithNewName(UserParameter);
+		VariableWithNewName.SetName(UniqueName);
+		FNiagaraUserRedirectionParameterStore::MakeUserVariable(VariableWithNewName);
+
+		TSharedPtr<FNiagaraSystemViewModel> SystemViewModel = TNiagaraViewModelManager<UNiagaraSystem, FNiagaraSystemViewModel>::GetExistingViewModelForObject(System.Get());
+		SystemViewModel->RenameParameter(UserParameter, NewName, ENiagaraGetGraphParameterReferencesMode::AllGraphs);
+	}
+}
+
+FReply FNiagaraSystemUserParameterBuilder::GenerateParameterDragDropOp(const FGeometry& Geometry, const FPointerEvent& MouseEvent, FNiagaraVariable UserParameter) const
+{
+	TSharedPtr<FEdGraphSchemaAction> ParameterDragAction = MakeShared<FNiagaraParameterAction>(UserParameter, FText::GetEmpty(), FText::GetEmpty(), FText::GetEmpty(), 0, FText(), nullptr, 0/*SectionID*/);
+	TSharedRef<FNiagaraParameterDragOperation> ParameterDragOperation = MakeShared<FNiagaraParameterDragOperation>(ParameterDragAction);
+	ParameterDragOperation->SetupDefaults();
+	ParameterDragOperation->Construct();
+	return FReply::Handled().BeginDragDrop(ParameterDragOperation);
+}
 
 TSharedRef<IDetailCustomization> FNiagaraComponentDetails::MakeInstance()
 {
@@ -1178,7 +1179,7 @@ void FNiagaraComponentDetails::CustomizeDetails(IDetailLayoutBuilder& DetailBuil
 		FGameDelegates::Get().GetEndPlayMapDelegate().AddRaw(this, &FNiagaraComponentDetails::OnPiEEnd);
 			
 		IDetailCategoryBuilder& InputParamCategory = DetailBuilder.EditCategory(ParamCategoryName, LOCTEXT("ParamCategoryName", "User Parameters"), ECategoryPriority::Important);
-		InputParamCategory.AddCustomBuilder(MakeShared<FNiagaraComponentNodeBuilder>(Component.Get(), PropertyHandles, ParamCategoryName));
+		InputParamCategory.AddCustomBuilder(MakeShared<FNiagaraComponentUserParametersNodeBuilder>(Component.Get(), PropertyHandles, ParamCategoryName));
 	}
 	else if (ObjectsCustomized.Num() > 1)
 	{
@@ -1307,32 +1308,4 @@ FReply FNiagaraComponentDetails::OnDebugSelectedSystem()
 	return FReply::Handled();
 }
 
-TSharedRef<IDetailCustomization> FNiagaraSystemUserParameterDetails::MakeInstance()
-{
-	return MakeShared<FNiagaraSystemUserParameterDetails>();
-}
-
-void FNiagaraSystemUserParameterDetails::CustomizeDetails(IDetailLayoutBuilder& DetailBuilder)
-{
-	TArray<TWeakObjectPtr<UObject>> CustomizedObjects;
-	DetailBuilder.GetObjectsBeingCustomized(CustomizedObjects);
-	
-	TArray<FName> Categories;
-	DetailBuilder.GetCategoryNames(Categories);
-
-	for(FName& CategoryName : Categories)
-	{
-		DetailBuilder.HideCategory(CategoryName);
-	}
-	
-	static const FName ParamCategoryName = TEXT("NiagaraSystem_UserParameters");
-
-	if(CustomizedObjects.Num() == 1 && CustomizedObjects[0]->IsA<UNiagaraSystem>())
-	{
-		System = Cast<UNiagaraSystem>(CustomizedObjects[0]);
-		IDetailCategoryBuilder& InputParamCategory = DetailBuilder.EditCategory(ParamCategoryName, LOCTEXT("ParamCategoryName", "User Parameters"), ECategoryPriority::Important);
-		InputParamCategory.AddCustomBuilder(MakeShared<FNiagaraSystemUserParameterBuilder>(System.Get(), ParamCategoryName));
-	}
-}
-
 #undef LOCTEXT_NAMESPACE
diff --git a/Engine/Plugins/FX/Niagara/Source/NiagaraEditor/Private/Customizations/NiagaraComponentDetails.h b/Engine/Plugins/FX/Niagara/Source/NiagaraEditor/Private/Customizations/NiagaraComponentDetails.h
--- a/Engine/Plugins/FX/Niagara/Source/NiagaraEditor/Private/Customizations/NiagaraComponentDetails.h
+++ b/Engine/Plugins/FX/Niagara/Source/NiagaraEditor/Private/Customizations/NiagaraComponentDetails.h
@@ -4,8 +4,12 @@
 
 #include "CoreMinimal.h"
 #include "IDetailCustomization.h"
+#include "IDetailCustomNodeBuilder.h"
 #include "Input/Reply.h"
 #include "NiagaraComponent.h"
+#include "ViewModels/TNiagaraViewModelManager.h"
+#include "ViewModels/HierarchyEditor/NiagaraHierarchyViewModelBase.h"
+#include "ViewModels/HierarchyEditor/NiagaraUserParametersHierarchyViewModel.h"
 
 class IDetailLayoutBuilder;
 class IPropertyHandle;
@@ -45,5 +49,146 @@ public:
 	virtual void CustomizeDetails( IDetailLayoutBuilder& DetailBuilder ) override;
 
 protected:
-	TWeakObjectPtr<UNiagaraSystem> System;	
+	TWeakObjectPtr<UNiagaraSystem> System;
+};
+
+/** A base class for generating rows based on a given hierarchy root that contains sections, categories and user parameters. */
+class FNiagaraUserParameterNodeBuilder : public IDetailCustomNodeBuilder, public TSharedFromThis<FNiagaraUserParameterNodeBuilder>, FSelfRegisteringEditorUndoClient
+{
+public:
+	virtual void GenerateHeaderRowContent(FDetailWidgetRow& NodeRow) override {}
+	virtual void Tick(float DeltaTime) override {}
+	virtual bool RequiresTick() const override { return false; }
+	virtual bool InitiallyCollapsed() const override { return false; }
+	virtual FName GetName() const  override
+	{
+		return CustomBuilderRowName;
+	}
+	
+	void Rebuild() const
+	{
+		OnRebuildChildren.ExecuteIfBound();
+	}
+	
+	virtual void SetOnRebuildChildren(FSimpleDelegate InOnRegenerateChildren) override
+	{
+		OnRebuildChildren = InOnRegenerateChildren;
+	}
+
+	virtual void PostUndo(bool bSuccess) override { Rebuild(); }
+	virtual void PostRedo(bool bSuccess) override { Rebuild(); }
+
+	/** Allows the addition of right-click context actions for a specific user parameter */
+	virtual void AddCustomMenuActionsForParameter(FDetailWidgetRow& WidgetRow, FNiagaraVariable UserParameter) {}
+	/** Allows the creation of a custom name widget for a user parameter. */
+	virtual TSharedPtr<SWidget> GenerateCustomNameWidget(FNiagaraVariable UserParameter) { return nullptr; }
+
+	/** Generates all rows for the node builder. Includes setting up sections, categories and user parameters */
+	void GenerateUserParameterRows(IDetailChildrenBuilder& ChildrenBuilder, UNiagaraHierarchyRoot& UserParameterHierarchyRoot, UObject* Owner);
+	/** Generates a single row for a specific user parameter. */
+	void GenerateRowForUserParameter(IDetailChildrenBuilder& ChildrenBuilder, FNiagaraVariable UserParameter, UObject* Owner);
+	
+	void OnActiveSectionChanged(ECheckBoxState State, const UNiagaraHierarchySection* NewSelection)
+	{
+		ActiveSection = NewSelection;
+		OnRebuildChildren.Execute();
+	}
+
+	ECheckBoxState IsSectionActive(const UNiagaraHierarchySection* Section) const
+	{
+		return ActiveSection == Section ? ECheckBoxState::Checked : ECheckBoxState::Unchecked;
+	}
+	
+protected:
+	bool bDelegatesInitialized = false;
+	TArray<TSharedPtr<IPropertyHandle>> OverridePropertyHandles;
+	FSimpleDelegate OnRebuildChildren;
+	TArray<TSharedPtr<class FNiagaraParameterProxy>> ParameterProxies;
+	TMap<FName, TWeakPtr<FStructOnScope>> ParameterNameToDisplayStruct;
+	FName CustomBuilderRowName;
+	const UNiagaraHierarchySection* ActiveSection = nullptr;
+	TArray<UNiagaraHierarchySection*> AvailableSections;
+};
+
+/** A specialized class to generate user parameters for a component. */
+class FNiagaraComponentUserParametersNodeBuilder : public FNiagaraUserParameterNodeBuilder
+{
+public:
+	FNiagaraComponentUserParametersNodeBuilder(UNiagaraComponent* InComponent, TArray<TSharedPtr<IPropertyHandle>> InOverridePropertyHandles, FName InCustomBuilderRowName) 
+	{
+		OverridePropertyHandles = InOverridePropertyHandles;
+		CustomBuilderRowName = InCustomBuilderRowName;
+
+		Component = InComponent;
+		bDelegatesInitialized = false;
+	}
+
+	virtual ~FNiagaraComponentUserParametersNodeBuilder() override
+	{
+		if (Component.IsValid() && bDelegatesInitialized)
+		{
+			Component->OnSynchronizedWithAssetParameters().RemoveAll(this);
+			Component->GetOverrideParameters().RemoveAllOnChangedHandlers(this);
+
+			if(UNiagaraSystem* System = Component->GetAsset())
+			{
+				TSharedPtr<FNiagaraSystemViewModel> SystemViewModel = TNiagaraViewModelManager<UNiagaraSystem, FNiagaraSystemViewModel>::GetExistingViewModelForObject(System);
+				if(SystemViewModel.IsValid())
+				{
+					SystemViewModel->GetUserParametersHierarchyViewModel()->OnHierarchyChanged().RemoveAll(this);
+				}
+			}
+		}
+	}
+	
+	virtual void GenerateChildContent(IDetailChildrenBuilder& ChildrenBuilder) override;
+
+private:
+	void RegisterRebuildOnHierarchyChanged();
+
+	void ParameterValueChanged();
+
+private:
+	TWeakObjectPtr<UNiagaraComponent> Component;
+};
+
+/** A specialized class to generate user parameters for a system asset */
+class FNiagaraSystemUserParameterBuilder : public FNiagaraUserParameterNodeBuilder
+{
+public:
+	FNiagaraSystemUserParameterBuilder(UNiagaraSystem* InSystem, FName InCustomBuilderRowName);
+
+	virtual ~FNiagaraSystemUserParameterBuilder() override
+	{
+		if (System.IsValid() && bDelegatesInitialized)
+		{
+			System->GetExposedParameters().OnStructureChanged().RemoveAll(this);
+
+			TSharedPtr<FNiagaraSystemViewModel> SystemViewModel = TNiagaraViewModelManager<UNiagaraSystem, FNiagaraSystemViewModel>::GetExistingViewModelForObject(System.Get());
+			if(SystemViewModel.IsValid())
+			{
+				SystemViewModel->GetUserParametersHierarchyViewModel()->OnHierarchyChanged().RemoveAll(this);
+			}
+		}
+	}
+
+	virtual void GenerateChildContent(IDetailChildrenBuilder& ChildrenBuilder) override;
+
+	/** Instead of pure text, we use the parameter widgets for display in a system asset. */
+	virtual TSharedPtr<SWidget> GenerateCustomNameWidget(FNiagaraVariable UserParameter) override;
+	
+	/** We want to add rename & delete actions within a system asset */
+	virtual void AddCustomMenuActionsForParameter(FDetailWidgetRow& WidgetRow, FNiagaraVariable UserParameter) override;
+private:
+	void ParameterValueChanged();
+
+	void DeleteParameter(FNiagaraVariable UserParameter) const;
+	void RequestRename(FNiagaraVariable UserParameter);
+	void RenameParameter(FNiagaraVariable UserParameter, FName NewName);
+
+	FReply GenerateParameterDragDropOp(const FGeometry& Geometry, const FPointerEvent& MouseEvent, FNiagaraVariable UserParameter) const;
+private:
+	TWeakObjectPtr<UNiagaraSystem> System;
+	TOptional<FNiagaraVariable> SelectedParameter;
+	TMap<FNiagaraVariable, TSharedPtr<class SNiagaraParameterNameTextBlock>> UserParamToWidgetMap;
 };
diff --git a/Engine/Plugins/FX/Niagara/Source/NiagaraEditor/Private/NiagaraEditorStyle.cpp b/Engine/Plugins/FX/Niagara/Source/NiagaraEditor/Private/NiagaraEditorStyle.cpp
--- a/Engine/Plugins/FX/Niagara/Source/NiagaraEditor/Private/NiagaraEditorStyle.cpp
+++ b/Engine/Plugins/FX/Niagara/Source/NiagaraEditor/Private/NiagaraEditorStyle.cpp
@@ -698,7 +698,8 @@ void FNiagaraEditorStyle::InitHierarchyEditor()
 	FInlineEditableTextBlockStyle HierarchyCategoryTextStyle = FInlineEditableTextBlockStyle()
 		.SetTextStyle(CategoryText)
 		.SetEditableTextBoxStyle(CategoryEditableText);
-	
+
+	Set("NiagaraEditor.HierarchyEditor.CategoryTextBlock", CategoryText);
 	Set("NiagaraEditor.HierarchyEditor.Category", HierarchyCategoryTextStyle);
 
 	Set("NiagaraEditor.Stack.DropTarget.BorderVertical", new IMAGE_BRUSH("Icons/StackDropTargetBorder_Vertical", FVector2D(2, 8), FLinearColor::White, ESlateBrushTileType::Vertical));
diff --git a/Engine/Plugins/FX/Niagara/Source/NiagaraEditor/Private/SNiagaraParameterName.cpp b/Engine/Plugins/FX/Niagara/Source/NiagaraEditor/Private/SNiagaraParameterName.cpp
--- a/Engine/Plugins/FX/Niagara/Source/NiagaraEditor/Private/SNiagaraParameterName.cpp
+++ b/Engine/Plugins/FX/Niagara/Source/NiagaraEditor/Private/SNiagaraParameterName.cpp
@@ -342,6 +342,7 @@ void SNiagaraParameterNameTextBlock::Construct(const FArguments& InArgs)
 	ParameterText = InArgs._ParameterText;
 	OnVerifyNameTextChangedDelegate = InArgs._OnVerifyTextChanged;
 	OnNameTextCommittedDelegate = InArgs._OnTextCommitted;
+	OnDragDetectedHandlerDelegate = InArgs._OnDragDetected;
 
 	ChildSlot
 	[
@@ -391,6 +392,27 @@ void SNiagaraParameterNameTextBlock::NameChanged(FName InNewName)
 	OnNameTextCommittedDelegate.ExecuteIfBound(FText::FromName(InNewName), ETextCommit::OnEnter);
 }
 
+FReply SNiagaraParameterNameTextBlock::OnMouseButtonDown(const FGeometry& MyGeometry, const FPointerEvent& MouseEvent)
+{
+	// for OnDragDetected to fire, the widget has to handle mouse button down
+	if(OnDragDetectedHandlerDelegate.IsBound())
+	{
+		return FReply::Handled().DetectDrag(SharedThis(this), MouseEvent.GetEffectingButton());
+	}
+
+	return FReply::Unhandled();
+}
+
+FReply SNiagaraParameterNameTextBlock::OnDragDetected(const FGeometry& MyGeometry, const FPointerEvent& MouseEvent)
+{
+	if(OnDragDetectedHandlerDelegate.IsBound())
+	{
+		return OnDragDetectedHandlerDelegate.Execute(MyGeometry, MouseEvent);
+	}
+
+	return FReply::Unhandled();
+}
+
 void SNiagaraParameterNameTextBlock::EnterEditingMode()
 {
 	ParameterName->EnterEditingMode();
diff --git a/Engine/Plugins/FX/Niagara/Source/NiagaraEditor/Private/Toolkits/NiagaraSystemToolkit.cpp b/Engine/Plugins/FX/Niagara/Source/NiagaraEditor/Private/Toolkits/NiagaraSystemToolkit.cpp
--- a/Engine/Plugins/FX/Niagara/Source/NiagaraEditor/Private/Toolkits/NiagaraSystemToolkit.cpp
+++ b/Engine/Plugins/FX/Niagara/Source/NiagaraEditor/Private/Toolkits/NiagaraSystemToolkit.cpp
@@ -1403,9 +1403,14 @@ void FNiagaraSystemToolkit::OnSystemSelectionChanged()
 	RefreshParameters();
 }
 
-void FNiagaraSystemToolkit::OnViewModelRequestFocusTab(FName TabName)
+void FNiagaraSystemToolkit::OnViewModelRequestFocusTab(FName TabName, bool bDrawAttention)
 {
-	GetTabManager()->TryInvokeTab(TabName);
+	TSharedPtr<SDockTab> DockTab = GetTabManager()->TryInvokeTab(TabName);
+
+	if(DockTab.IsValid() && bDrawAttention)
+	{
+		DockTab->FlashTab();
+	}
 }
 
 
diff --git a/Engine/Plugins/FX/Niagara/Source/NiagaraEditor/Private/Toolkits/NiagaraSystemToolkit.h b/Engine/Plugins/FX/Niagara/Source/NiagaraEditor/Private/Toolkits/NiagaraSystemToolkit.h
--- a/Engine/Plugins/FX/Niagara/Source/NiagaraEditor/Private/Toolkits/NiagaraSystemToolkit.h
+++ b/Engine/Plugins/FX/Niagara/Source/NiagaraEditor/Private/Toolkits/NiagaraSystemToolkit.h
@@ -192,7 +192,7 @@ private:
 	void OnPinnedCurvesChanged();
 	void RefreshParameters();
 	void OnSystemSelectionChanged();
-	void OnViewModelRequestFocusTab(FName TabName);
+	void OnViewModelRequestFocusTab(FName TabName, bool bDrawAttention = false);
 	
 	const FName GetNiagaraSystemMessageLogName(UNiagaraSystem* InSystem) const;
 	void OnSaveThumbnailImage();
diff --git a/Engine/Plugins/FX/Niagara/Source/NiagaraEditor/Private/Toolkits/SystemToolkitModes/NiagaraSystemToolkitModeBase.cpp b/Engine/Plugins/FX/Niagara/Source/NiagaraEditor/Private/Toolkits/SystemToolkitModes/NiagaraSystemToolkitModeBase.cpp
--- a/Engine/Plugins/FX/Niagara/Source/NiagaraEditor/Private/Toolkits/SystemToolkitModes/NiagaraSystemToolkitModeBase.cpp
+++ b/Engine/Plugins/FX/Niagara/Source/NiagaraEditor/Private/Toolkits/SystemToolkitModes/NiagaraSystemToolkitModeBase.cpp
@@ -790,6 +790,56 @@ TSharedRef<SDockTab> FNiagaraSystemToolkitModeBase::SpawnTab_UserParameters(cons
 	return SpawnedTab;
 }
 
+TSharedRef<SWidget> GenerateRowContentForUserParameterHierarchyEditor(TSharedRef<FNiagaraHierarchyItemViewModelBase> HierarchyItem, TSharedPtr<FNiagaraSystemViewModel> SystemViewModel)
+{
+	if(HierarchyItem->GetDataMutable()->IsA<UNiagaraHierarchyCategory>())
+	{
+		TSharedRef<FNiagaraHierarchyCategoryViewModel> TreeViewCategory = StaticCastSharedRef<FNiagaraHierarchyCategoryViewModel>(HierarchyItem);
+		return SNew(SNiagaraHierarchyCategory, TreeViewCategory);
+	}
+	else if(const UNiagaraHierarchyUserParameter* UserParameter = Cast<UNiagaraHierarchyUserParameter>(HierarchyItem->GetData()))
+	{
+		TSharedRef<SWidget> ParameterWidget = FNiagaraParameterUtilities::GetParameterWidget(UserParameter->GetUserParameter(), true, false);
+		UNiagaraScriptVariable* ScriptVariable = FNiagaraEditorUtilities::GetScriptVariableForUserParameter(UserParameter->GetUserParameter(), SystemViewModel->GetSystem());
+
+		ParameterWidget->SetToolTipText(TAttribute<FText>::CreateLambda([ScriptVariable]()
+		{
+			return ScriptVariable->Metadata.Description;
+		}));
+		
+		return ParameterWidget;
+	}
+
+	return SNullWidget::NullWidget;
+}
+
+TSharedRef<SWidget> GenerateCustomDetailsPanelNameWidget(TSharedPtr<FNiagaraHierarchyItemViewModelBase> HierarchyItem)
+{
+	if(!HierarchyItem.IsValid())
+	{
+		return SNew(STextBlock).Text(FText::FromString("None selected"));
+	}
+	
+	if(HierarchyItem->GetData()->IsA<UNiagaraHierarchyCategory>() || HierarchyItem->GetData()->IsA<UNiagaraHierarchySection>())
+	{
+		return SNew(SBox).Padding(2.f)
+		[
+			SNew(STextBlock)
+			.Text_Lambda([HierarchyItem]
+			{
+				return FText::FromString(HierarchyItem->ToString());
+			})
+			.TextStyle(&FNiagaraEditorStyle::Get().GetWidgetStyle<FTextBlockStyle>("NiagaraEditor.HierarchyEditor.CategoryTextBlock"))
+		];
+	}
+	else if(const UNiagaraHierarchyUserParameter* UserParameter = Cast<UNiagaraHierarchyUserParameter>(HierarchyItem->GetData()))
+	{
+		return FNiagaraParameterUtilities::GetParameterWidget(UserParameter->GetUserParameter(), false, false);
+	}
+
+	return SNullWidget::NullWidget;
+}
+
 TSharedRef<SDockTab> FNiagaraSystemToolkitModeBase::SpawnTab_UserParametersHierarchyEditor(const FSpawnTabArgs& Args)
 {
 	check(Args.GetTabId().TabType == UserParametersHierarchyTabID);
@@ -801,6 +851,8 @@ TSharedRef<SDockTab> FNiagaraSystemToolkitModeBase::SpawnTab_UserParametersHiera
 			.AddMetaData<FTagMetaData>(FTagMetaData(TEXT("UserParameters")))
 			[
 				SNew(SNiagaraHierarchy, SystemToolkit.Pin()->GetSystemViewModel()->GetUserParametersHierarchyViewModel())
+				.OnGenerateRowContentWidget_Static(&GenerateRowContentForUserParameterHierarchyEditor, SystemToolkit.Pin()->GetSystemViewModel())
+				.OnGenerateCustomDetailsPanelNameWidget_Static(&GenerateCustomDetailsPanelNameWidget)
 			]
 		];
 	
diff --git a/Engine/Plugins/FX/Niagara/Source/NiagaraEditor/Private/ViewModels/HierarchyEditor/NiagaraHierarchyViewModelBase.cpp b/Engine/Plugins/FX/Niagara/Source/NiagaraEditor/Private/ViewModels/HierarchyEditor/NiagaraHierarchyViewModelBase.cpp
--- a/Engine/Plugins/FX/Niagara/Source/NiagaraEditor/Private/ViewModels/HierarchyEditor/NiagaraHierarchyViewModelBase.cpp
+++ b/Engine/Plugins/FX/Niagara/Source/NiagaraEditor/Private/ViewModels/HierarchyEditor/NiagaraHierarchyViewModelBase.cpp
@@ -39,12 +39,20 @@ void UNiagaraHierarchyItemBase::Refresh()
 	Children.Empty();
 	for(UNiagaraHierarchyItemBase* Item : OldChildren)
 	{
+		// if an item is already finalized we can skip it
 		if(Item->IsFinalized())
 		{
 			continue;
 		}
-
+		
 		Item->Refresh();
+
+		// the refresh can have caused the item to finalize. In that case, we don't add it back
+		if(Item->IsFinalized())
+		{
+			continue;
+		}
+		
 		Children.Add(Item);
 	}
 }
@@ -154,7 +162,7 @@ bool UNiagaraHierarchyRoot::Modify(bool bAlwaysMarkDirty)
 
 bool FNiagaraHierarchyCategoryViewModel::IsTopCategoryActive() const
 {
-	if(UNiagaraHierarchyCategory* Category = Cast<UNiagaraHierarchyCategory>(GetDataMutable()))
+	if(UNiagaraHierarchyCategory* Category = GetDataMutable<UNiagaraHierarchyCategory>())
 	{
 		const UNiagaraHierarchyCategory* Result = Category;
 		const UNiagaraHierarchyCategory* TopLevelCategory = Result;
@@ -190,12 +198,10 @@ UNiagaraHierarchyViewModelBase::~UNiagaraHierarchyViewModelBase()
 	RefreshSourceViewDelegate.Unbind();
 	RefreshHierarchyWidgetDelegate.Unbind();
 	RefreshSectionsWidgetDelegate.Unbind();
-	IsItemSelectedDelegate.Unbind();
-	SelectObjectInDetailsPanelDelegate.Unbind();
 }
 
 void UNiagaraHierarchyViewModelBase::ForceFullRefresh()
-{
+{	
 	PrepareSourceItems();
 	HierarchyViewModelRoot->SyncToData();
 	RefreshSourceView(true);
@@ -203,21 +209,6 @@ void UNiagaraHierarchyViewModelBase::ForceFullRefresh()
 	RefreshSectionsWidget();
 }
 
-TSharedRef<SWidget> UNiagaraHierarchyViewModelBase::GenerateRowContentWidget(TSharedRef<FNiagaraHierarchyItemViewModelBase> ItemBaseViewModel) const
-{
-	if(UNiagaraHierarchyCategory* Category = Cast<UNiagaraHierarchyCategory>(ItemBaseViewModel->GetDataMutable()))
-	{
-		TSharedRef<FNiagaraHierarchyCategoryViewModel> TreeViewCategory = StaticCastSharedRef<FNiagaraHierarchyCategoryViewModel>(ItemBaseViewModel);
-		return SNew(SNiagaraHierarchyCategory, TreeViewCategory)
-		.IsSelected(FIsSelected::CreateLambda([this, ItemBaseViewModel]
-		{
-			return this->IsItemSelected(ItemBaseViewModel);
-		}));
-	}
-
-	return SNullWidget::NullWidget;
-}
-
 void UNiagaraHierarchyViewModelBase::RefreshSourceView(bool bFullRefresh) const
 {
 	if(RefreshSourceViewDelegate.IsBound())
@@ -242,24 +233,6 @@ void UNiagaraHierarchyViewModelBase::RefreshSectionsWidget() const
 	}
 }
 
-void UNiagaraHierarchyViewModelBase::SelectItemInDetailsPanel(UObject* Object) const
-{
-	if(SelectObjectInDetailsPanelDelegate.IsBound())
-	{
-		SelectObjectInDetailsPanelDelegate.Execute(Object);
-	}
-}
-
-bool UNiagaraHierarchyViewModelBase::IsItemSelected(TSharedPtr<FNiagaraHierarchyItemViewModelBase> Item) const
-{
-	if(IsItemSelectedDelegate.IsBound())
-	{
-		return IsItemSelectedDelegate.Execute(Item);
-	}
-
-	return false;
-}
-
 void UNiagaraHierarchyViewModelBase::PostUndo(bool bSuccess)
 {
 	ForceFullRefresh();
@@ -320,6 +293,20 @@ void UNiagaraHierarchyViewModelBase::Finalize()
 	FinalizeInternal();
 }
 
+TSharedPtr<FNiagaraHierarchyItemViewModelBase> UNiagaraHierarchyViewModelBase::CreateViewModelForData(UNiagaraHierarchyItemBase* ItemBase, TSharedPtr<FNiagaraHierarchyItemViewModelBase> Parent)
+{
+	if(UNiagaraHierarchyItem* Item = Cast<UNiagaraHierarchyItem>(ItemBase))
+	{
+		return MakeShared<FNiagaraHierarchyItemViewModel>(Item, Parent, this);
+	}
+	else if(UNiagaraHierarchyCategory* Category = Cast<UNiagaraHierarchyCategory>(ItemBase))
+	{
+		return MakeShared<FNiagaraHierarchyCategoryViewModel>(Category, Parent, this);
+	}
+
+	return nullptr;
+}
+
 const TArray<TSharedPtr<FNiagaraHierarchyItemViewModelBase>>& UNiagaraHierarchyViewModelBase::GetSourceItems() const
 {
 	return SourceViewModelRoot->GetChildren();
@@ -337,16 +324,7 @@ void UNiagaraHierarchyViewModelBase::OnGetChildren(TSharedPtr<FNiagaraHierarchyI
 
 void UNiagaraHierarchyViewModelBase::SetActiveSection(TSharedPtr<FNiagaraHierarchySectionViewModel> Section)
 {
-	ActiveSection = Section;
-	if(Section.IsValid())
-	{
-		SelectItemInDetailsPanel(Section->GetDataMutable());
-	}
-	else
-	{
-		SelectItemInDetailsPanel(nullptr);
-	}
-	
+	ActiveSection = Section;	
 	RefreshHierarchyView(true);
 }
 
@@ -357,7 +335,7 @@ TSharedPtr<FNiagaraHierarchySectionViewModel> UNiagaraHierarchyViewModelBase::Ge
 
 UNiagaraHierarchySection* UNiagaraHierarchyViewModelBase::GetActiveSectionData() const
 {
-	return ActiveSection.IsValid() ? Cast<UNiagaraHierarchySection>(ActiveSection.Pin()->GetDataMutable()) : nullptr;
+	return ActiveSection.IsValid() ? ActiveSection.Pin()->GetDataMutable<UNiagaraHierarchySection>() : nullptr;
 }
 
 bool UNiagaraHierarchyViewModelBase::IsSectionActive(const UNiagaraHierarchySection* Section) const
@@ -383,18 +361,13 @@ void FNiagaraHierarchyItemViewModelBase::SyncToData()
 		// if we couldn't find a view model for a data child, we create it here
 		if(ViewModelIndex == INDEX_NONE)
 		{
-			if(UNiagaraHierarchyItem* Item = Cast<UNiagaraHierarchyItem>(Child))
+			TSharedPtr<FNiagaraHierarchyItemViewModelBase> NewItem = HierarchyViewModel->CreateViewModelForData(Child, AsShared());
+
+			if(ensure(NewItem.IsValid()))
 			{
-				TSharedPtr<FNiagaraHierarchyItemViewModel> NewItem = MakeShared<FNiagaraHierarchyItemViewModel>(Item, AsShared(), HierarchyViewModel);
 				NewItem->SyncToData();
 				NewChildren.Add(NewItem);
 			}
-			else if(UNiagaraHierarchyCategory* Category = Cast<UNiagaraHierarchyCategory>(Child))
-			{
-				TSharedPtr<FNiagaraHierarchyCategoryViewModel> NewCategory = MakeShared<FNiagaraHierarchyCategoryViewModel>(Category, AsShared(), HierarchyViewModel);
-				NewCategory->SyncToData();
-				NewChildren.Add(NewCategory);
-			}
 		}
 		// if we could find it, we refresh its contained view models and readd it
 		else
@@ -434,9 +407,18 @@ const TArray<TSharedPtr<FNiagaraHierarchyItemViewModelBase>>& FNiagaraHierarchyI
 
 	for(TSharedPtr<FNiagaraHierarchyItemViewModelBase> Child : Children)
 	{
-		if(const UNiagaraHierarchyCategory* Category = Cast<UNiagaraHierarchyCategory>(Child->GetData()))
+		// we add categories only if the top most category is considered active
+		if(const UNiagaraHierarchyCategory* Category = Child->GetData<UNiagaraHierarchyCategory>())
 		{
-			if(StaticCastSharedPtr<FNiagaraHierarchyCategoryViewModel>(Child)->IsTopCategoryActive()) //HierarchyViewModel->IsSectionActive(Category->GetSection()))
+			if(StaticCastSharedPtr<FNiagaraHierarchyCategoryViewModel>(Child)->IsTopCategoryActive())
+			{
+				FilteredChildren.Add(Child);
+			}
+		}
+		// if we have an item directly under the root without being categorized, we only allow if if the active section is set to "All"
+		else if(Child->GetParent().IsValid() && Child->GetParent() == HierarchyViewModel->GetHierarchyViewModelRoot() && Child->GetData()->IsA<UNiagaraHierarchyItem>())
+		{
+			if(HierarchyViewModel->GetActiveSectionData() == nullptr)
 			{
 				FilteredChildren.Add(Child);
 			}
@@ -546,6 +528,19 @@ int32 FNiagaraHierarchyItemViewModelBase::FindIndexOfDataChild(UNiagaraHierarchy
 	return GetData()->GetChildren().Find(Child);
 }
 
+void FNiagaraHierarchyItemViewModelBase::Delete()
+{
+	FScopedTransaction Transaction(LOCTEXT("DeletedItemTransaction", "Deleted hierarchy item"));
+	ItemBase->Modify();
+
+	if(Parent.IsValid())
+	{
+		Parent.Pin()->ItemBase->Modify();
+	}
+	
+	Finalize();
+}
+
 void FNiagaraHierarchyItemViewModelBase::Finalize()
 {	
 	for(TSharedPtr<FNiagaraHierarchyItemViewModelBase> ChildViewModel : Children)
@@ -603,6 +598,34 @@ FReply FNiagaraHierarchyItemViewModelBase::OnDragDetected(const FGeometry& Geome
 	return FReply::Unhandled();
 }
 
+FReply FNiagaraHierarchyRootViewModel::OnDroppedOn(const FDragDropEvent& DragDropEvent, EItemDropZone ItemDropZone, TSharedPtr<FNiagaraHierarchyItemViewModelBase> Item)
+{
+	if(DragDropEvent.GetOperation()->IsOfType<FNiagaraHierarchyDragDropOp>())
+	{
+		FScopedTransaction Transaction(LOCTEXT("Transaction_OnDropOnRoot", "Dropped item on root"));
+		HierarchyViewModel->GetHierarchyDataRoot()->Modify();
+		
+		TSharedPtr<FNiagaraHierarchyDragDropOp> HierarchyDragDropOp = DragDropEvent.GetOperationAs<FNiagaraHierarchyDragDropOp>();
+		TWeakPtr<FNiagaraHierarchyItemViewModelBase> DraggedItem = HierarchyDragDropOp->GetDraggedItem();
+
+		if(DraggedItem.Pin()->GetDataMutable()->IsA<UNiagaraHierarchyItem>() || DraggedItem.Pin()->GetDataMutable()->IsA<UNiagaraHierarchyCategory>())
+		{
+			if(HierarchyDragDropOp->GetIsFromSourceList())
+			{
+				DuplicateToThis(DraggedItem.Pin());
+			}
+			else
+			{
+				ReparentToThis(DraggedItem.Pin());
+			}
+
+			HierarchyViewModel->RefreshHierarchyView();
+		}
+	}
+
+	return FReply::Unhandled();
+}
+
 TSharedPtr<FNiagaraHierarchySectionViewModel> FNiagaraHierarchyRootViewModel::AddNewSection()
 {
 	FScopedTransaction ScopedTransaction(LOCTEXT("NewSectionAdded","Added Section"));
@@ -615,14 +638,13 @@ TSharedPtr<FNiagaraHierarchySectionViewModel> FNiagaraHierarchyRootViewModel::Ad
 	SyncToData();
 	HierarchyViewModel->SetActiveSection(NewSectionViewModel);
 	NewSectionViewModel->RequestRename();
-	
 
 	return NewSectionViewModel;
 }
 
 void FNiagaraHierarchyRootViewModel::SyncChildrenViewModelsInternal()
 {
-	const UNiagaraHierarchyRoot* RootData = Cast<UNiagaraHierarchyRoot>(GetData());
+	const UNiagaraHierarchyRoot* RootData = GetData<UNiagaraHierarchyRoot>();
 
 	TArray<TSharedPtr<FNiagaraHierarchySectionViewModel>> SectionsToRemove;
 	for(TSharedPtr<FNiagaraHierarchySectionViewModel> SectionViewModel : SectionViewModels)
@@ -659,9 +681,10 @@ void FNiagaraHierarchyRootViewModel::SyncChildrenViewModelsInternal()
 
 	SectionViewModels.Sort([this](const TSharedPtr<FNiagaraHierarchySectionViewModel>& ItemA, const TSharedPtr<FNiagaraHierarchySectionViewModel>& ItemB)
 		{
-			return Cast<UNiagaraHierarchyRoot>(GetDataMutable())->GetSectionData().Find(Cast<UNiagaraHierarchySection>(ItemA->GetDataMutable()))
-						<
-					Cast<UNiagaraHierarchyRoot>(GetDataMutable())->GetSectionData().Find(Cast<UNiagaraHierarchySection>(ItemB->GetDataMutable())); 
+			return
+			GetDataMutable<UNiagaraHierarchyRoot>()->GetSectionData().Find(Cast<UNiagaraHierarchySection>(ItemA->GetDataMutable()))
+				<
+			GetDataMutable<UNiagaraHierarchyRoot>()->GetSectionData().Find(Cast<UNiagaraHierarchySection>(ItemB->GetDataMutable())); 
 		});
 }
 
@@ -671,10 +694,10 @@ FReply FNiagaraHierarchySectionViewModel::OnDroppedOn(const FDragDropEvent& Drag
 	{
 		TSharedPtr<FNiagaraSectionDragDropOp> SectionDragDropOp = DragDropEvent.GetOperationAs<FNiagaraSectionDragDropOp>();
 		TWeakPtr<FNiagaraHierarchySectionViewModel> DraggedItem = SectionDragDropOp->GetDraggedSection();
-		UNiagaraHierarchySection* DraggedSectionData = Cast<UNiagaraHierarchySection>(DraggedItem.Pin()->GetDataMutable());
+		UNiagaraHierarchySection* DraggedSectionData = DraggedItem.Pin()->GetDataMutable<UNiagaraHierarchySection>();
 
-		int32 IndexOfThis = HierarchyViewModel->GetHierarchyDataRoot()->GetSectionData().Find(Cast<UNiagaraHierarchySection>(GetDataMutable()));
-		int32 DraggedSectionIndex = HierarchyViewModel->GetHierarchyDataRoot()->GetSectionData().Find(Cast<UNiagaraHierarchySection>(DraggedItem.Pin()->GetDataMutable()));
+		int32 IndexOfThis = HierarchyViewModel->GetHierarchyDataRoot()->GetSectionData().Find(GetDataMutable<UNiagaraHierarchySection>());
+		int32 DraggedSectionIndex = HierarchyViewModel->GetHierarchyDataRoot()->GetSectionData().Find(DraggedItem.Pin()->GetDataMutable<UNiagaraHierarchySection>());
 
 		TArray<UNiagaraHierarchySection*>& SectionData = HierarchyViewModel->GetHierarchyDataRoot()->GetSectionDataMutable();
 		int32 Count = SectionData.Num();
@@ -719,12 +742,12 @@ FReply FNiagaraHierarchySectionViewModel::OnDroppedOn(const FDragDropEvent& Drag
 		TSharedPtr<FNiagaraHierarchyDragDropOp> HierarchyDragDropOp = DragDropEvent.GetOperationAs<FNiagaraHierarchyDragDropOp>();
 		TWeakPtr<FNiagaraHierarchyItemViewModelBase> DraggedItem = HierarchyDragDropOp->GetDraggedItem();
 
-		if(UNiagaraHierarchyCategory* HierarchyCategory = Cast<UNiagaraHierarchyCategory>(DraggedItem.Pin()->GetDataMutable()))
+		if(UNiagaraHierarchyCategory* HierarchyCategory = DraggedItem.Pin()->GetDataMutable<UNiagaraHierarchyCategory>())
 		{
 			FScopedTransaction Transaction(LOCTEXT("Transaction_OnSectionDrop", "Moved category to section"));
 			HierarchyViewModel->GetHierarchyDataRoot()->Modify();
 			
-			HierarchyCategory->SetSection(Cast<UNiagaraHierarchySection>(GetDataMutable()));
+			HierarchyCategory->SetSection(GetDataMutable<UNiagaraHierarchySection>());
 
 			// we null out any sections for all contained categories
 			TArray<UNiagaraHierarchyCategory*> AllChildCategories;
@@ -822,11 +845,27 @@ FReply FNiagaraHierarchyCategoryViewModel::OnDroppedOn(const FDragDropEvent& Dra
 		
 		TSharedPtr<FNiagaraHierarchyDragDropOp> HierarchyDragDropOp = DragDropEvent.GetOperationAs<FNiagaraHierarchyDragDropOp>();
 		TSharedPtr<FNiagaraHierarchyItemViewModelBase> DraggedItem = HierarchyDragDropOp->GetDraggedItem().Pin();
-		
-		if(ItemDropZone == EItemDropZone::OntoItem)
+
+		if(UNiagaraHierarchyCategory* Category = DraggedItem->GetDataMutable<UNiagaraHierarchyCategory>())
 		{
-			// if we are dragging a category onto another category, we null out its section
-			if(UNiagaraHierarchyCategory* Category = Cast<UNiagaraHierarchyCategory>(DraggedItem->GetDataMutable()))
+			// if we are dragging a category above/below another category and the new parent is going to be the root, we update its section to the active section
+			if(ItemDropZone != EItemDropZone::OntoItem)
+			{
+				if(Parent.IsValid() && Parent == HierarchyViewModel->GetHierarchyViewModelRoot())
+				{
+					Category->SetSection(HierarchyViewModel->GetActiveSectionData());
+
+					// we null out any sections for all contained categories
+					TArray<UNiagaraHierarchyCategory*> AllChildCategories;
+					Category->GetChildrenOfType<UNiagaraHierarchyCategory>(AllChildCategories, true);
+					for(UNiagaraHierarchyCategory* ChildCategory : AllChildCategories)
+					{
+						ChildCategory->SetSection(nullptr);
+					}
+				}				
+			}
+			// if we are dragging a category onto another category, we null out its section instead
+			else
 			{
 				Category->SetSection(nullptr);
 
@@ -837,8 +876,12 @@ FReply FNiagaraHierarchyCategoryViewModel::OnDroppedOn(const FDragDropEvent& Dra
 				{
 					ChildCategory->SetSection(nullptr);
 				}
-			}
+			}			
+		}
 
+		// the actual moving of the item happens here
+		if(ItemDropZone == EItemDropZone::OntoItem)
+		{
 			if(HierarchyDragDropOp->GetIsFromSourceList())
 			{
 				DuplicateToThis(HierarchyDragDropOp->GetDraggedItem().Pin());
@@ -856,7 +899,7 @@ FReply FNiagaraHierarchyCategoryViewModel::OnDroppedOn(const FDragDropEvent& Dra
 				Parent.Pin()->DuplicateToThis(HierarchyDragDropOp->GetDraggedItem().Pin(), FMath::Max(IndexOfThis, 0));
 			}
 			else
-			{
+			{				
 				Parent.Pin()->ReparentToThis(HierarchyDragDropOp->GetDraggedItem().Pin(), FMath::Max(IndexOfThis, 0));
 			}
 		}
diff --git a/Engine/Plugins/FX/Niagara/Source/NiagaraEditor/Private/ViewModels/HierarchyEditor/NiagaraUserParametersHierarchyViewModel.cpp b/Engine/Plugins/FX/Niagara/Source/NiagaraEditor/Private/ViewModels/HierarchyEditor/NiagaraUserParametersHierarchyViewModel.cpp
--- a/Engine/Plugins/FX/Niagara/Source/NiagaraEditor/Private/ViewModels/HierarchyEditor/NiagaraUserParametersHierarchyViewModel.cpp
+++ b/Engine/Plugins/FX/Niagara/Source/NiagaraEditor/Private/ViewModels/HierarchyEditor/NiagaraUserParametersHierarchyViewModel.cpp
@@ -33,6 +33,16 @@ void UNiagaraHierarchyUserParameter::RefreshDataInternal()
 	}
 }
 
+UObject* FNiagaraHierarchyUserParameterViewModel::GetDataForEditing()
+{
+	UNiagaraHierarchyUserParameter* HierarchyUserParameter = Cast<UNiagaraHierarchyUserParameter>(GetDataMutable<UNiagaraHierarchyUserParameter>());
+	FNiagaraVariable ContainedVariable = HierarchyUserParameter->GetUserParameter();
+	UNiagaraUserParametersHierarchyViewModel* UserParametersHierarchyViewModel = Cast<UNiagaraUserParametersHierarchyViewModel>(GetHierarchyViewModel());
+	UserParametersHierarchyViewModel->GetSystemViewModel()->GetSystem().GetExposedParameters().RedirectUserVariable(ContainedVariable);
+	TObjectPtr<UNiagaraScriptVariable> ScriptVariable = FNiagaraEditorUtilities::GetScriptVariableForUserParameter(ContainedVariable, UserParametersHierarchyViewModel->GetSystemViewModel());
+	return ScriptVariable.Get();
+}
+
 TSharedRef<FNiagaraSystemViewModel> UNiagaraUserParametersHierarchyViewModel::GetSystemViewModel() const
 {
 	TSharedPtr<FNiagaraSystemViewModel> SystemViewModelPinned = SystemViewModelWeak.Pin();
@@ -73,6 +83,21 @@ UNiagaraHierarchyRoot* UNiagaraUserParametersHierarchyViewModel::GetHierarchyDat
 	return RootItem;
 }
 
+TSharedPtr<FNiagaraHierarchyItemViewModelBase> UNiagaraUserParametersHierarchyViewModel::CreateViewModelForData(UNiagaraHierarchyItemBase* ItemBase, TSharedPtr<FNiagaraHierarchyItemViewModelBase> Parent)
+{
+	if(UNiagaraHierarchyUserParameter* UserParameter = Cast<UNiagaraHierarchyUserParameter>(ItemBase))
+	{
+		return MakeShared<FNiagaraHierarchyUserParameterViewModel>(UserParameter, Parent, this);
+	}
+	else if(UNiagaraHierarchyCategory* Category = Cast<UNiagaraHierarchyCategory>(ItemBase))
+	{
+		return MakeShared<FNiagaraHierarchyCategoryViewModel>(Category, Parent, this);
+	}
+
+	check(false);
+	return nullptr;
+}
+
 void UNiagaraUserParametersHierarchyViewModel::PrepareSourceItems()
 {	
 	TArray<FNiagaraVariable> UserParameters;
@@ -120,12 +145,20 @@ TOptional<EItemDropZone> UNiagaraUserParametersHierarchyViewModel::CanDropOn(TSh
 			!TargetDropItem->HasParent(SourceDropItem, true)
 		&&
 		(
+			// user parameters can be dropped onto categories
 			(SourceDropItem->GetData()->IsA<UNiagaraHierarchyUserParameter>() && TargetDropItem->GetData()->IsA<UNiagaraHierarchyCategory>() && DropZone == EItemDropZone::OntoItem)
 				||
+			// user parameters can be dropped above/below other user parameters
 			(SourceDropItem->GetData()->IsA<UNiagaraHierarchyUserParameter>() && TargetDropItem->GetData()->IsA<UNiagaraHierarchyUserParameter>() && DropZone != EItemDropZone::OntoItem)
 				||
-			(SourceDropItem->GetData()->IsA<UNiagaraHierarchyCategory>() && TargetDropItem->GetData()->IsA<UNiagaraHierarchyCategory>()
-				)
+			// categories can be dropped on categories
+			(SourceDropItem->GetData()->IsA<UNiagaraHierarchyCategory>() && TargetDropItem->GetData()->IsA<UNiagaraHierarchyCategory>())
+				||
+			// user parameters can be dropped onto the root if the section is set to "All"
+			(SourceDropItem->GetData()->IsA<UNiagaraHierarchyUserParameter>() && TargetDropItem->GetData()->IsA<UNiagaraHierarchyRoot>() && GetActiveSectionData() == nullptr)
+				||
+			// categories can be dropped onto the root always
+			(SourceDropItem->GetData()->IsA<UNiagaraHierarchyCategory>() && TargetDropItem->GetData()-IsA<UNiagaraHierarchyRoot>())
 		);
 	
 	return bAllowDropOn ? DropZone : TOptional<EItemDropZone>();
@@ -151,54 +184,9 @@ TSharedRef<FNiagaraHierarchyDragDropOp> UNiagaraUserParametersHierarchyViewModel
 	return MakeShared<FNiagaraHierarchyDragDropOp>(nullptr);
 }
 
-void UNiagaraUserParametersHierarchyViewModel::OnSelectionChanged(TSharedPtr<FNiagaraHierarchyItemViewModelBase> HierarchyItem)
-{
-	if(!HierarchyItem.IsValid())
-	{
-		SelectObjectInDetailsPanelDelegate.Execute(nullptr);
-		return;
-	}
-
-	UObject* ObjectToSelect = nullptr;
-	if(HierarchyItem->GetData()->IsA<UNiagaraHierarchyUserParameter>())
-	{
-		UNiagaraHierarchyUserParameter* HierarchyUserParameter = Cast<UNiagaraHierarchyUserParameter>(HierarchyItem->GetDataMutable());
-		FNiagaraVariable ContainedVariable = HierarchyUserParameter->GetUserParameter();
-		GetSystemViewModel()->GetSystem().GetExposedParameters().RedirectUserVariable(ContainedVariable);
-		TObjectPtr<UNiagaraScriptVariable> ScriptVariable = FNiagaraEditorUtilities::GetScriptVariableForUserParameter(ContainedVariable, GetSystemViewModel());
-
-		if(ScriptVariable)
-		{
-			ObjectToSelect = ScriptVariable.Get();
-		}
-	}
-	else
-	{
-		ObjectToSelect = HierarchyItem->GetDataMutable();
-	}
-
-	SelectObjectInDetailsPanelDelegate.Execute(ObjectToSelect);
-}
-
 TArray<TTuple<UClass*, FOnGetDetailCustomizationInstance>> UNiagaraUserParametersHierarchyViewModel::GetInstanceCustomizations()
 {
 	return {{UNiagaraScriptVariable::StaticClass(), FOnGetDetailCustomizationInstance::CreateStatic(&FNiagaraScriptVariableHierarchyDetails::MakeInstance)}};
 }
 
-TSharedRef<SWidget> UNiagaraUserParametersHierarchyViewModel::GenerateRowContentWidget(TSharedRef<FNiagaraHierarchyItemViewModelBase> ItemViewModelBase) const
-{
-	TSharedRef<SWidget> RowWidget = Super::GenerateRowContentWidget(ItemViewModelBase);
-
-	if(RowWidget == SNullWidget::NullWidget)
-	{
-		if(const UNiagaraHierarchyUserParameter* UserParameter = Cast<UNiagaraHierarchyUserParameter>(ItemViewModelBase->GetData()))
-		{
-			RowWidget = FNiagaraParameterUtilities::GetParameterWidget(UserParameter->GetUserParameter(), true, false);
-		}
-	}	
-
-	check(RowWidget != SNullWidget::NullWidget);
-	return RowWidget;
-}
-
 #undef LOCTEXT_NAMESPACE
diff --git a/Engine/Plugins/FX/Niagara/Source/NiagaraEditor/Private/ViewModels/NiagaraParameterPanelViewModel.cpp b/Engine/Plugins/FX/Niagara/Source/NiagaraEditor/Private/ViewModels/NiagaraParameterPanelViewModel.cpp
--- a/Engine/Plugins/FX/Niagara/Source/NiagaraEditor/Private/ViewModels/NiagaraParameterPanelViewModel.cpp
+++ b/Engine/Plugins/FX/Niagara/Source/NiagaraEditor/Private/ViewModels/NiagaraParameterPanelViewModel.cpp
@@ -882,7 +882,6 @@ void FNiagaraSystemToolkitParameterPanelViewModel::Init(const FSystemToolkitUICo
 	}
 
 	// Define the sections
-	UserOnlyIdx = Sections.Emplace(FGuid(0x9E29BC43, 0x31524AB5, 0xA9D6675D, 0x741137CB), LOCTEXT("SectionUser", "User"), LOCTEXT("SectionUserDesc", "Displays only User exposed parameters"));
 	ActiveSystemIdx = Sections.Emplace(FGuid(0x7B4AFB34, 0xD0DF4618, 0x9A05349E, 0x361CE735), LOCTEXT("SectionSystemActive", "Active Overview"), LOCTEXT("SectionActiveSystemDesc", "Displays parameters that are context-sensitive to the selected items in the System Overview."));
 	ActiveScriptIdx = Sections.Emplace(FGuid(0x6E11F3EA, 0x5E314E82, 0xA64FEA41, 0x4F1DB891), LOCTEXT("SectionScriptActive", "Active Module"), LOCTEXT("SectionActiveScriptDesc", "Displays parameters that are context-sensitive to the active Scratch module document. Use this for renaming/editing local parameters to the module."));
 
@@ -892,12 +891,6 @@ void FNiagaraSystemToolkitParameterPanelViewModel::Init(const FSystemToolkitUICo
 	{
 		bool bSaveConfig = false;
 		bool bAdded = false;
-		FNiagaraParameterPanelSectionStorage& UserStorage = Settings->FindOrAddParameterPanelSectionStorage(Sections[UserOnlyIdx].SectionId, bAdded);
-		if (bAdded)
-		{
-			UserStorage.ExpandedCategories.Emplace(FNiagaraEditorGuids::UserNamespaceMetaDataGuid);
-			bSaveConfig = true;
-		}
 
 		FNiagaraParameterPanelSectionStorage& ScriptStorage = Settings->FindOrAddParameterPanelSectionStorage(Sections[ActiveScriptIdx].SectionId, bAdded);
 		if (bAdded)
@@ -1517,28 +1510,19 @@ const TArray<FNiagaraParameterPanelCategory>& FNiagaraSystemToolkitParameterPane
 	// Show advanced in Active/Common case. Only show User when in user mode.
 	bool bForceScript = false;
 	bool bForceSystem = false;
-	bool bUserOnly = false;
-	if (ActiveSectionIndex == UserOnlyIdx && UserOnlyIdx != -1)
-		bUserOnly = true;
-	else if (ActiveSectionIndex == ActiveScriptIdx && ActiveScriptIdx != -1)
+	
+	if (ActiveSectionIndex == ActiveScriptIdx && ActiveScriptIdx != -1)
 		bForceScript= true;
 	else if (ActiveSectionIndex == ActiveSystemIdx && ActiveSystemIdx != -1)
 		bForceSystem = true;
-
-	if (bUserOnly)
+	
+	if (bForceScript)
 	{
-		CachedCurrentCategories = FNiagaraSystemToolkitParameterPanelViewModel::UserCategories;
+		CachedCurrentCategories = FNiagaraSystemToolkitParameterPanelViewModel::DefaultAdvancedScriptCategories;
 	}
 	else
-	{	
-		if (bForceScript)
-		{
-			CachedCurrentCategories = FNiagaraSystemToolkitParameterPanelViewModel::DefaultAdvancedScriptCategories;
-		}
-		else
-		{
-			CachedCurrentCategories = FNiagaraSystemToolkitParameterPanelViewModel::DefaultAdvancedCategories;
-		}
+	{
+		CachedCurrentCategories = FNiagaraSystemToolkitParameterPanelViewModel::DefaultAdvancedCategories;
 	}
 	
 	return CachedCurrentCategories;
@@ -1865,10 +1849,7 @@ TArray<FNiagaraParameterPanelItem> FNiagaraSystemToolkitParameterPanelViewModel:
 {
 	bool bForceScript = false;
 	bool bForceSystem = false;
-	bool bUserOnly = false;
-	if (ActiveSectionIndex == UserOnlyIdx && UserOnlyIdx != -1)
-		bUserOnly = true;
-	else if (ActiveSectionIndex == ActiveScriptIdx && ActiveScriptIdx != -1)
+	if (ActiveSectionIndex == ActiveScriptIdx && ActiveScriptIdx != -1)
 		bForceScript = true;
 	else if (ActiveSectionIndex == ActiveSystemIdx && ActiveSystemIdx != -1)
 		bForceSystem = true;
@@ -1917,7 +1898,7 @@ TArray<FNiagaraParameterPanelItem> FNiagaraSystemToolkitParameterPanelViewModel:
 	};
 
 	// Collect user parameters from system.
-	if (bUserOnly || bForceSystem)
+	if (bForceSystem)
 	{
 		CollectParamStore(&SystemViewModel->GetSystem().GetExposedParameters());
 	}
@@ -1944,7 +1925,7 @@ TArray<FNiagaraParameterPanelItem> FNiagaraSystemToolkitParameterPanelViewModel:
 		}
 	}
 
-	if (bForceSystem || bUserOnly)
+	if (bForceSystem)
 	{
 		// Collect parameters for all emitters.
 		TArray<FNiagaraVariable> VisitedInvalidParameters;
@@ -2041,7 +2022,7 @@ TArray<FNiagaraParameterPanelItem> FNiagaraSystemToolkitParameterPanelViewModel:
 						// This variable has already been registered, increment the reference count.
 						ItemPtr->ReferenceCount += Builder.Histories[0].PerVariableReadHistory[VariableIndex].Num() + Builder.Histories[0].PerVariableWriteHistory[VariableIndex].Num();
 					}
-					else if (!bUserOnly) // Add newly found variables, unless in user-only mode in which we use the parameter store defined ones only
+					else  // Add newly found variables
 					{
 						// This variable has not been registered, prepare the FNiagaraParameterPanelItem.
 						// -First make sure the variable namespace is in a valid category. If not, skip it.
diff --git a/Engine/Plugins/FX/Niagara/Source/NiagaraEditor/Private/ViewModels/NiagaraSystemViewModel.cpp b/Engine/Plugins/FX/Niagara/Source/NiagaraEditor/Private/ViewModels/NiagaraSystemViewModel.cpp
--- a/Engine/Plugins/FX/Niagara/Source/NiagaraEditor/Private/ViewModels/NiagaraSystemViewModel.cpp
+++ b/Engine/Plugins/FX/Niagara/Source/NiagaraEditor/Private/ViewModels/NiagaraSystemViewModel.cpp
@@ -151,6 +151,8 @@ void FNiagaraSystemViewModel::Initialize(UNiagaraSystem& InSystem, FNiagaraSyste
 	ScalabilityViewModel = NewObject<UNiagaraSystemScalabilityViewModel>(GetTransientPackage());
 	ScalabilityViewModel->Initialize(this->AsShared());
 
+	UserParameterPanelViewModel = MakeShared<FNiagaraUserParameterPanelViewModel>();
+	
 	UserParametersHierarchyViewModel = NewObject<UNiagaraUserParametersHierarchyViewModel>(GetTransientPackage());
 	UserParametersHierarchyViewModel->Initialize(this->AsShared());
 	
@@ -305,6 +307,11 @@ void FNiagaraSystemViewModel::Cleanup()
 		UserParametersHierarchyViewModel = nullptr;
 	}
 
+	if(UserParameterPanelViewModel.IsValid())
+	{
+		UserParameterPanelViewModel.Reset();
+	}
+
 	if (PlaceholderDataInterfaceManager.IsValid())
 	{
 		PlaceholderDataInterfaceManager.Reset();
@@ -911,9 +918,9 @@ FNiagaraSystemViewModel::FOnRequestFocusTab& FNiagaraSystemViewModel::OnRequestF
 	return OnRequestFocusTabDelegate;
 }
 
-void FNiagaraSystemViewModel::FocusTab(FName TabName)
+void FNiagaraSystemViewModel::FocusTab(FName TabName, bool bDrawAttention)
 {
-	OnRequestFocusTabDelegate.Broadcast(TabName);
+	OnRequestFocusTabDelegate.Broadcast(TabName, bDrawAttention);
 }
 
 TSharedPtr<FUICommandList> FNiagaraSystemViewModel::GetToolkitCommands()
diff --git a/Engine/Plugins/FX/Niagara/Source/NiagaraEditor/Private/ViewModels/NiagaraUserParameterPanelViewModel.cpp b/Engine/Plugins/FX/Niagara/Source/NiagaraEditor/Private/ViewModels/NiagaraUserParameterPanelViewModel.cpp
new file mode 100644
--- /dev/null
+++ b/Engine/Plugins/FX/Niagara/Source/NiagaraEditor/Private/ViewModels/NiagaraUserParameterPanelViewModel.cpp
@@ -0,0 +1,20 @@
+// Copyright Epic Games, Inc. All Rights Reserved.
+
+#pragma once
+
+#include "ViewModels/NiagaraUserParameterPanelViewModel.h"
+#include "CoreMinimal.h"
+#include "ViewModels/NiagaraSystemViewModel.h"
+
+template<> TMap<UNiagaraSystem*, TArray<FNiagaraUserParameterPanelViewModel*>> TNiagaraViewModelManager<UNiagaraSystem, FNiagaraUserParameterPanelViewModel>::ObjectsToViewModels{};
+
+FNiagaraUserParameterPanelViewModel::~FNiagaraUserParameterPanelViewModel()
+{
+	UnregisterViewModelWithMap(RegisteredHandle);
+}
+
+void FNiagaraUserParameterPanelViewModel::Initialize(TSharedRef<FNiagaraSystemViewModel> InSystemViewModel)
+{
+	SystemViewModelWeak = InSystemViewModel;
+	RegisteredHandle = RegisterViewModelWithMap(&(SystemViewModelWeak.Pin()->GetSystem()), this);
+}
diff --git a/Engine/Plugins/FX/Niagara/Source/NiagaraEditor/Private/Widgets/SNiagaraHierarchy.cpp b/Engine/Plugins/FX/Niagara/Source/NiagaraEditor/Private/Widgets/SNiagaraHierarchy.cpp
--- a/Engine/Plugins/FX/Niagara/Source/NiagaraEditor/Private/Widgets/SNiagaraHierarchy.cpp
+++ b/Engine/Plugins/FX/Niagara/Source/NiagaraEditor/Private/Widgets/SNiagaraHierarchy.cpp
@@ -23,6 +23,12 @@ void SNiagaraHierarchyCategory::Construct(const FArguments& InArgs,	TSharedPtr<F
 	Category = InCategory;
 
 	InCategory->GetOnRequestRename().BindSP(this, &SNiagaraHierarchyCategory::EnterEditingMode);
+
+	const UNiagaraHierarchyCategory* CategoryData = InCategory->GetData<UNiagaraHierarchyCategory>();
+	SetToolTipText(TAttribute<FText>::CreateLambda([CategoryData]
+	{
+		return CategoryData->GetTooltip();
+	}));
 	
 	ChildSlot
 	[
@@ -81,6 +87,8 @@ void SNiagaraHierarchySection::Construct(const FArguments& InArgs, TSharedPtr<FN
 	SectionViewModel = InSection;
 	HierarchyViewModel = InHierarchyViewModel;
 
+	OnSectionActivatedDelegate = InArgs._OnSectionActivated;
+	
 	if(SectionViewModel != nullptr)
 	{
 		SectionViewModel->GetOnRequestRename().BindSP(this, &SNiagaraHierarchySection::EnterEditingMode);
@@ -104,7 +112,9 @@ void SNiagaraHierarchySection::Construct(const FArguments& InArgs, TSharedPtr<FN
 			.OnDropped(this, &SNiagaraHierarchySection::OnDroppedOn, EItemDropZone::BelowItem)
 			.VerticalImage(FAppStyle::GetNoBrush())
 			.HorizontalImage(FNiagaraEditorStyle::Get().GetBrush("NiagaraEditor.Stack.DropTarget.BorderHorizontal"));
-		
+
+		SetToolTipText(TAttribute<FText>::CreateSP(this, &SNiagaraHierarchySection::GetTooltipText));
+
 		ChildSlot
 		[
 			SNew(SHorizontalBox)
@@ -144,7 +154,7 @@ void SNiagaraHierarchySection::Construct(const FArguments& InArgs, TSharedPtr<FN
 				SNew(SNiagaraSectionDragDropTarget)
 				.DropTargetArgs(RightDropTargetArgs)				
 			]
-		];
+		];		
 	}
 	// if this section doesn't represent data, it's the "All" widget
 	else
@@ -263,16 +273,30 @@ FReply SNiagaraHierarchySection::OnMouseButtonDown(const FGeometry& MyGeometry,
 {
 	if(SectionViewModel.IsValid())
 	{
+		// we handle the event here so we can react on mouse button up
 		if(MouseEvent.IsMouseButtonDown(EKeys::RightMouseButton))
 		{
-			MenuAnchor->SetIsOpen(true);
-			OnSectionCheckChanged(ECheckBoxState::Checked);
 			return FReply::Handled();
 		}
 		else if(MouseEvent.IsMouseButtonDown(EKeys::LeftMouseButton))
 		{
-			OnSectionCheckChanged(ECheckBoxState::Checked);
-			return FReply::Handled().DetectDrag(AsShared(), EKeys::LeftMouseButton);	
+			OnSectionActivatedDelegate.ExecuteIfBound(SectionViewModel);
+			return FReply::Handled().DetectDrag(AsShared(), EKeys::LeftMouseButton).SetUserFocus(AsShared());	
+		}
+	}
+
+	return FReply::Unhandled();
+}
+
+FReply SNiagaraHierarchySection::OnMouseButtonUp(const FGeometry& MyGeometry, const FPointerEvent& MouseEvent)
+{
+	if(SectionViewModel.IsValid())
+	{
+		if(MouseEvent.GetEffectingButton() == EKeys::RightMouseButton)
+		{
+			MenuAnchor->SetIsOpen(true);
+			OnSectionActivatedDelegate.ExecuteIfBound(SectionViewModel);
+			return FReply::Handled();
 		}
 	}
 
@@ -319,6 +343,11 @@ FText SNiagaraHierarchySection::GetText() const
 	return SectionViewModel->GetSectionNameAsText();
 }
 
+FText SNiagaraHierarchySection::GetTooltipText() const
+{
+	return SectionViewModel->GetSectionTooltip();
+}
+
 void SNiagaraHierarchySection::OnRenameSection(const FText& Text, ETextCommit::Type CommitType) const
 {
 	FScopedTransaction Transaction(LOCTEXT("Transaction_Rename_Section", "Renamed hierarchy section"));
@@ -362,15 +391,7 @@ ECheckBoxState SNiagaraHierarchySection::GetSectionCheckState() const
 
 void SNiagaraHierarchySection::OnSectionCheckChanged(ECheckBoxState NewState)
 {
-	if(NewState == ECheckBoxState::Checked)
-	{
-		HierarchyViewModel->SetActiveSection(SectionViewModel);
-	}
-
-	if(SectionViewModel.IsValid())
-	{
-		HierarchyViewModel->SelectItemInDetailsPanel(SectionViewModel->GetDataMutable());
-	}
+	OnSectionActivatedDelegate.ExecuteIfBound(SectionViewModel);
 }
 
 void SNiagaraHierarchy::Construct(const FArguments& InArgs, TObjectPtr<UNiagaraHierarchyViewModelBase> InHierarchyViewModel)
@@ -380,7 +401,9 @@ void SNiagaraHierarchy::Construct(const FArguments& InArgs, TObjectPtr<UNiagaraH
 	HierarchyViewModel->OnRefreshSourceView().BindSP(this, &SNiagaraHierarchy::RefreshSourceView);
 	HierarchyViewModel->OnRefreshHierarchyView().BindSP(this, &SNiagaraHierarchy::RefreshHierarchyView);
 	HierarchyViewModel->OnRefreshSections().BindSP(this, &SNiagaraHierarchy::RefreshSectionsWidget);
-	HierarchyViewModel->OnSelectObjectInDetailsPanel().BindSP(this, &SNiagaraHierarchy::SelectObjectInDetailsPanel);
+
+	OnGenerateRowContentWidget = InArgs._OnGenerateRowContentWidget;
+	OnGenerateCustomDetailsPanelNameWidget = InArgs._OnGenerateCustomDetailsPanelNameWidget;
 	
 	SHorizontalBox::FSlot* DetailsPanelSlot = nullptr;
 	
@@ -389,12 +412,16 @@ void SNiagaraHierarchy::Construct(const FArguments& InArgs, TObjectPtr<UNiagaraH
 		SNew(SHorizontalBox)
 		+ SHorizontalBox::Slot()
 		[
-			SAssignNew(SourceTreeView, STreeView<TSharedPtr<FNiagaraHierarchyItemViewModelBase>>)
-			.TreeItemsSource(&InHierarchyViewModel->GetSourceItems())
-			.OnSelectionChanged(this, &SNiagaraHierarchy::OnSelectionChanged)
-			.OnGenerateRow(this, &SNiagaraHierarchy::GenerateSourceItemRow)
-			.OnGetChildren_UObject(InHierarchyViewModel.Get(), &UNiagaraHierarchyViewModelBase::OnGetChildren)
-			.OnItemToString_Debug_UObject(InHierarchyViewModel.Get(), &UNiagaraHierarchyViewModelBase::OnItemToStringDebug)
+			SNew(SVerticalBox)
+			+ SVerticalBox::Slot()
+			[
+				SAssignNew(SourceTreeView, STreeView<TSharedPtr<FNiagaraHierarchyItemViewModelBase>>)
+				.TreeItemsSource(&InHierarchyViewModel->GetSourceItems())
+				.OnSelectionChanged(this, &SNiagaraHierarchy::OnSelectionChanged, false)
+				.OnGenerateRow(this, &SNiagaraHierarchy::GenerateSourceItemRow)
+				.OnGetChildren_UObject(InHierarchyViewModel.Get(), &UNiagaraHierarchyViewModelBase::OnGetChildren)
+				.OnItemToString_Debug_UObject(InHierarchyViewModel.Get(), &UNiagaraHierarchyViewModelBase::OnItemToStringDebug)
+			]
 		]
 		+ SHorizontalBox::Slot()
 		.AutoWidth()
@@ -421,12 +448,24 @@ void SNiagaraHierarchy::Construct(const FArguments& InArgs, TObjectPtr<UNiagaraH
 			]
 			+ SVerticalBox::Slot()
 			[
-				SAssignNew(HierarchyTreeView, STreeView<TSharedPtr<FNiagaraHierarchyItemViewModelBase>>)
-				.TreeItemsSource(&InHierarchyViewModel->GetHierarchyItems())
-				.OnSelectionChanged(this, &SNiagaraHierarchy::OnSelectionChanged)
-				.OnGenerateRow(this, &SNiagaraHierarchy::GenerateHierarchyItemRow)
-				.OnGetChildren_UObject(InHierarchyViewModel.Get(), &UNiagaraHierarchyViewModelBase::OnGetChildren)
-				.OnItemToString_Debug_UObject(InHierarchyViewModel.Get(), &UNiagaraHierarchyViewModelBase::OnItemToStringDebug)
+				SNew(SVerticalBox)
+				+ SVerticalBox::Slot()
+				.AutoHeight()
+				[
+					SAssignNew(HierarchyTreeView, STreeView<TSharedPtr<FNiagaraHierarchyItemViewModelBase>>)
+					.TreeItemsSource(&InHierarchyViewModel->GetHierarchyItems())
+					.OnSelectionChanged(this, &SNiagaraHierarchy::OnSelectionChanged, true)
+					.OnGenerateRow(this, &SNiagaraHierarchy::GenerateHierarchyItemRow)
+					.OnGetChildren_UObject(InHierarchyViewModel.Get(), &UNiagaraHierarchyViewModelBase::OnGetChildren)
+					.OnItemToString_Debug_UObject(InHierarchyViewModel.Get(), &UNiagaraHierarchyViewModelBase::OnItemToStringDebug)
+				]
+				+ SVerticalBox::Slot()
+				.VAlign(VAlign_Fill)
+				[
+					SNew(SDropTarget)
+					.OnDropped(this, &SNiagaraHierarchy::HandleHierarchyRootDrop)
+					.OnAllowDrop(this, &SNiagaraHierarchy::CanDropOnRoot)
+				]
 			]
 		]
 		+ SHorizontalBox::Slot().Expose(DetailsPanelSlot)		
@@ -438,9 +477,16 @@ void SNiagaraHierarchy::Construct(const FArguments& InArgs, TObjectPtr<UNiagaraH
 
 		FDetailsViewArgs DetailsViewArgs;
 		DetailsViewArgs.NameAreaSettings = FDetailsViewArgs::ENameAreaSettings::ObjectsUseNameArea;
-		DetailsViewArgs.bShowObjectLabel = true;
-		DetailsPanel = PropertyEditorModule.CreateDetailView(DetailsViewArgs);
+		DetailsViewArgs.bShowObjectLabel = false;
 
+		DetailsPanel = PropertyEditorModule.CreateDetailView(DetailsViewArgs);
+		
+		if(OnGenerateCustomDetailsPanelNameWidget.IsBound())
+		{
+			TSharedRef<SWidget> CustomDetailsPanelNameWidget = OnGenerateCustomDetailsPanelNameWidget.Execute(nullptr);
+			DetailsPanel->SetNameAreaCustomContent(CustomDetailsPanelNameWidget);
+		}		
+		
 		for(auto& Customizations : HierarchyViewModel->GetInstanceCustomizations())
 		{
 			DetailsPanel->RegisterInstancedCustomPropertyLayout(Customizations.Key, Customizations.Value);
@@ -459,6 +505,8 @@ void SNiagaraHierarchy::Construct(const FArguments& InArgs, TObjectPtr<UNiagaraH
 	HierarchyViewModel->GetCommands()->MapAction(FGenericCommands::Get().Delete,
 		FExecuteAction::CreateSP(this, &SNiagaraHierarchy::DeleteSelectedItems),
 		FCanExecuteAction::CreateSP(this, &SNiagaraHierarchy::CanDeleteSelectedItems));
+
+	HierarchyViewModel->ForceFullRefresh();
 }
 
 SNiagaraHierarchy::~SNiagaraHierarchy()
@@ -468,7 +516,6 @@ SNiagaraHierarchy::~SNiagaraHierarchy()
 	HierarchyViewModel->OnRefreshSourceView().Unbind();
 	HierarchyViewModel->OnRefreshHierarchyView().Unbind();
 	HierarchyViewModel->OnRefreshSections().Unbind();
-	HierarchyViewModel->OnSelectObjectInDetailsPanel().Unbind();
 
 	HierarchyViewModel->GetCommands()->UnmapAction(FGenericCommands::Get().Delete);
 	HierarchyViewModel->GetCommands()->UnmapAction(FGenericCommands::Get().Rename);
@@ -479,6 +526,17 @@ FReply SNiagaraHierarchy::OnKeyDown(const FGeometry& MyGeometry, const FKeyEvent
 	return HierarchyViewModel->GetCommands()->ProcessCommandBindings(InKeyEvent) ? FReply::Handled() : FReply::Unhandled();
 }
 
+FReply SNiagaraHierarchy::OnMouseButtonDown(const FGeometry& MyGeometry, const FPointerEvent& MouseEvent)
+{
+	// we catch any mouse button down event so that we can continue using our commands
+	return FReply::Handled().SetUserFocus(AsShared(), EFocusCause::Mouse, true);
+}
+
+FReply SNiagaraHierarchy::OnMouseButtonUp(const FGeometry& MyGeometry, const FPointerEvent& MouseEvent)
+{
+	return FReply::Handled().SetUserFocus(AsShared(), EFocusCause::Mouse, true);
+}
+
 void SNiagaraHierarchy::RefreshSourceView(bool bFullRefresh) const
 {
 	if(bFullRefresh)
@@ -557,7 +615,12 @@ TSharedRef<SWidget> SNiagaraHierarchy::CreateSectionWidgets()
 	
 	for (TSharedPtr<FNiagaraHierarchySectionViewModel>& Section : HierarchyViewModel->GetHierarchyViewModelRoot()->GetSectionViewModels())
 	{
-		TSharedPtr<SNiagaraHierarchySection> SectionWidget = SNew(SNiagaraHierarchySection, Section, HierarchyViewModel);
+		TSharedPtr<SNiagaraHierarchySection> SectionWidget = SNew(SNiagaraHierarchySection, Section, HierarchyViewModel)
+		.OnSectionActivated_Lambda([this](TSharedPtr<FNiagaraHierarchySectionViewModel> SectionViewModel)
+		{
+			HierarchyViewModel->SetActiveSection(SectionViewModel);
+			OnSelectionChanged(SectionViewModel, ESelectInfo::Type::OnMouseClick, true);
+		});
 		SectionsWidgetMap.Add(Section, SectionWidget);
 		
 		SectionsBox->AddSlot()
@@ -566,7 +629,13 @@ TSharedRef<SWidget> SNiagaraHierarchy::CreateSectionWidgets()
 		];
 	}
 	
-	TSharedPtr<SNiagaraHierarchySection> DefaultSectionWidget = SNew(SNiagaraHierarchySection, nullptr, HierarchyViewModel);
+	TSharedPtr<SNiagaraHierarchySection> DefaultSectionWidget = SNew(SNiagaraHierarchySection, nullptr, HierarchyViewModel)
+	.OnSectionActivated_Lambda([this](TSharedPtr<FNiagaraHierarchySectionViewModel> SectionViewModel)
+	{
+		HierarchyViewModel->SetActiveSection(SectionViewModel);
+		OnSelectionChanged(SectionViewModel, ESelectInfo::Type::OnMouseClick, true);
+	});
+	
 	SectionsWidgetMap.Add(nullptr, DefaultSectionWidget);
 
 	SectionsBox->AddSlot()
@@ -585,7 +654,7 @@ TSharedRef<ITableRow> SNiagaraHierarchy::GenerateSourceItemRow(TSharedPtr<FNiaga
 	return SNew(STableRow<TSharedPtr<FNiagaraHierarchyItemViewModelBase>>, TableViewBase)
 	.OnDragDetected(HierarchyItem.ToSharedRef(), &FNiagaraHierarchyItemViewModelBase::OnDragDetected, true)
 	[
-		HierarchyViewModel->GenerateRowContentWidget(HierarchyItem.ToSharedRef())
+		OnGenerateRowContentWidget.Execute(HierarchyItem.ToSharedRef())
 	];
 }
 
@@ -596,7 +665,7 @@ TSharedRef<ITableRow> SNiagaraHierarchy::GenerateHierarchyItemRow(TSharedPtr<FNi
 	.OnCanAcceptDrop(HierarchyItem.ToSharedRef(), &FNiagaraHierarchyItemViewModelBase::OnCanAcceptDrop)
 	.OnDragDetected(HierarchyItem.ToSharedRef(), &FNiagaraHierarchyItemViewModelBase::OnDragDetected, false)
 	[
-		HierarchyViewModel->GenerateRowContentWidget(HierarchyItem.ToSharedRef())
+		OnGenerateRowContentWidget.Execute(HierarchyItem.ToSharedRef())
 	];
 }
 
@@ -673,8 +742,10 @@ void SNiagaraHierarchy::DeleteSelectedItems()
 bool SNiagaraHierarchy::CanDeleteSelectedItems() const
 {
 	TArray<TSharedPtr<FNiagaraHierarchyItemViewModelBase>> SelectedItems = HierarchyTreeView->GetSelectedItems();
+	TArray<TSharedPtr<FNiagaraHierarchyItemViewModelBase>> SelectedItemsInSource = SourceTreeView->GetSelectedItems();
 
-	if(SelectedItems.Num() == 0)
+	// we only want to allow inferring the intent to delete a section if _nothing_ else is selected
+	if(SelectedItems.Num() == 0 && SelectedItemsInSource.Num() == 0)
 	{
 		TSharedPtr<FNiagaraHierarchySectionViewModel> ActiveSection = HierarchyViewModel->GetActiveSection();
 		if(ActiveSection)
@@ -697,11 +768,65 @@ bool SNiagaraHierarchy::CanDeleteSelectedItems() const
 	return false;
 }
 
-void SNiagaraHierarchy::OnSelectionChanged(TSharedPtr<FNiagaraHierarchyItemViewModelBase> HierarchyItem, ESelectInfo::Type Type) const
+void SNiagaraHierarchy::OnSelectionChanged(TSharedPtr<FNiagaraHierarchyItemViewModelBase> HierarchyItem, ESelectInfo::Type Type, bool bFromHierarchy) const
 {
+	if(DetailsPanel.IsValid())
+	{
+		if(HierarchyItem.IsValid())
+		{
+			// when se select a section, and the previous item selection is no longer available due to it, we would get a selection refresh next tick
+			// to wipe out the current selection. We want to avoid that, so we manually clear the selected items in that case.
+			if(HierarchyItem->GetData()->IsA<UNiagaraHierarchySection>())
+			{
+				HierarchyTreeView->ClearSelection();
+			}
+
+			// we clear the selection of the other tree view to avoid 
+			if(bFromHierarchy)
+			{
+				SourceTreeView->ClearSelection();
+			}
+			else
+			{
+				HierarchyTreeView->ClearSelection();
+			}
+			
+			DetailsPanel->SetObject(HierarchyItem->GetDataForEditing());	
+		}
+		else
+		{
+			DetailsPanel->SetObject(nullptr);
+		}
+	}
+	
+	if(DetailsPanel.IsValid() && OnGenerateCustomDetailsPanelNameWidget.IsBound() && HierarchyItem.IsValid())
+	{
+		TSharedRef<SWidget> NameWidget = OnGenerateCustomDetailsPanelNameWidget.Execute(HierarchyItem);
+		DetailsPanel->SetNameAreaCustomContent(NameWidget);
+	}
+
 	HierarchyViewModel->OnSelectionChanged(HierarchyItem);
 }
 
+FReply SNiagaraHierarchy::HandleHierarchyRootDrop(const FGeometry& Geometry, const FDragDropEvent& DragDropEvent) const
+{
+	// we pass in nullptr because the dragged object is contained within the drag drop event
+	return HierarchyViewModel->GetHierarchyViewModelRoot()->OnDroppedOn(DragDropEvent, EItemDropZone::OntoItem, nullptr);
+}
+
+bool SNiagaraHierarchy::CanDropOnRoot(TSharedPtr<FDragDropOperation> DragDropOp) const
+{
+	if(DragDropOp->IsOfType<FNiagaraHierarchyDragDropOp>())
+	{
+		TSharedPtr<FNiagaraHierarchyDragDropOp> HierarchyDragDropOp = StaticCastSharedPtr<FNiagaraHierarchyDragDropOp>(DragDropOp);
+		TSharedPtr<FNiagaraHierarchyItemViewModelBase> DraggedItem = HierarchyDragDropOp->GetDraggedItem().Pin();
+		TOptional<EItemDropZone> Result = HierarchyViewModel->CanDropOn(DraggedItem, HierarchyViewModel->GetHierarchyViewModelRoot(), EItemDropZone::OntoItem);
+		return Result.IsSet() && Result.GetValue() == EItemDropZone::OntoItem ? true : false;
+	}
+
+	return false;
+}
+
 void SNiagaraSectionDragDropTarget::Construct(const FArguments& InArgs)
 {
 	SDropTarget::Construct(InArgs._DropTargetArgs);
diff --git a/Engine/Plugins/FX/Niagara/Source/NiagaraEditor/Private/Widgets/SNiagaraSystemUserParameters.cpp b/Engine/Plugins/FX/Niagara/Source/NiagaraEditor/Private/Widgets/SNiagaraSystemUserParameters.cpp
--- a/Engine/Plugins/FX/Niagara/Source/NiagaraEditor/Private/Widgets/SNiagaraSystemUserParameters.cpp
+++ b/Engine/Plugins/FX/Niagara/Source/NiagaraEditor/Private/Widgets/SNiagaraSystemUserParameters.cpp
@@ -22,24 +22,26 @@ void SNiagaraSystemUserParameters::Construct(const FArguments& InArgs, TSharedPt
 	SystemViewModel = InSystemViewModel;
 	
 	FPropertyEditorModule& PropertyEditorModule = FModuleManager::Get().LoadModuleChecked<FPropertyEditorModule>("PropertyEditor");
-
+	
 	FDetailsViewArgs DetailsViewArgs;
 	DetailsViewArgs.NameAreaSettings = FDetailsViewArgs::ENameAreaSettings::HideNameArea;
+	DetailsViewArgs.bShowOptions = false;
+	DetailsViewArgs.bShowPropertyMatrixButton = false;
 	TSharedRef<IDetailsView> ObjectDetailsView = PropertyEditorModule.CreateDetailView(DetailsViewArgs);
 	ObjectDetailsView->RegisterInstancedCustomPropertyLayout(UNiagaraSystem::StaticClass(), FOnGetDetailCustomizationInstance::CreateStatic(&FNiagaraSystemUserParameterDetails::MakeInstance));
-	ObjectDetailsView->SetObject(&SystemViewModel->GetSystem());
+	ObjectDetailsView->SetObject(&InSystemViewModel->GetSystem());
 	
 	ChildSlot
 	[
 		SNew(SVerticalBox)
 		+ SVerticalBox::Slot()
 		.AutoHeight()
-		.HAlign(HAlign_Right)
+		.HAlign(HAlign_Fill)
 		.Padding(2.f)
 		[
 			SNew(SHorizontalBox)
 			+ SHorizontalBox::Slot()
-			.AutoWidth()
+			.HAlign(HAlign_Left)
 			[
 				SAssignNew(AddParameterButton, SComboButton)
 				.OnGetMenuContent(this, &SNiagaraSystemUserParameters::GetParameterMenu)
@@ -54,9 +56,10 @@ void SNiagaraSystemUserParameters::Construct(const FArguments& InArgs, TSharedPt
 				]
 			]
 			+ SHorizontalBox::Slot()
-			.AutoWidth()
+			.HAlign(HAlign_Right)
 			[
 				SNew(SButton)
+				.HAlign(HAlign_Right)
 				.OnClicked(this, &SNiagaraSystemUserParameters::SummonHierarchyEditor)
 				.Text(LOCTEXT("SummonUserParametersHierarchyButtonLabel", "Edit Hierarchy"))
 			]
@@ -76,7 +79,7 @@ TSharedRef<SWidget> SNiagaraSystemUserParameters::GetParameterMenu()
 {
 	FNiagaraParameterPanelCategory UserCategory = GetDefault<UNiagaraEditorSettings>()->GetMetaDataForNamespaces({FNiagaraConstants::UserNamespace});
 	ParameterPanel = SNew(SNiagaraAddParameterFromPanelMenu)
-		.Graphs(SystemViewModel->GetParameterPanelViewModel()->GetEditableGraphsConst())
+		.Graphs(SystemViewModel.Pin()->GetParameterPanelViewModel()->GetEditableGraphsConst())
 		.OnNewParameterRequested(this, &SNiagaraSystemUserParameters::AddParameter)
 		.OnAllowMakeType(this, &SNiagaraSystemUserParameters::CanMakeNewParameterOfType)
 		.NamespaceId(UserCategory.NamespaceMetaData.GetGuid())
@@ -90,18 +93,18 @@ TSharedRef<SWidget> SNiagaraSystemUserParameters::GetParameterMenu()
 void SNiagaraSystemUserParameters::AddParameter(FNiagaraVariable NewParameter) const
 {
 	FNiagaraParameterPanelCategory UserCategory = GetDefault<UNiagaraEditorSettings>()->GetMetaDataForNamespaces({FNiagaraConstants::UserNamespace});
-	// TODO (ME) Change request rename to true when it's supported in the UI
-	SystemViewModel->GetParameterPanelViewModel()->AddParameter(NewParameter, UserCategory, false, true);
+	FNiagaraEditorUtilities::AddParameter(NewParameter, SystemViewModel.Pin()->GetSystem().GetExposedParameters(), SystemViewModel.Pin()->GetSystem(), nullptr);
+	SystemViewModel.Pin()->GetUserParameterPanelViewModel()->OnParameterAdded().Execute(NewParameter);
 }
 
 bool SNiagaraSystemUserParameters::CanMakeNewParameterOfType(const FNiagaraTypeDefinition& InType) const
 {
-	return SystemViewModel->GetParameterPanelViewModel()->CanMakeNewParameterOfType(InType);
+	return SystemViewModel.Pin()->GetParameterPanelViewModel()->CanMakeNewParameterOfType(InType);
 }
 
 FReply SNiagaraSystemUserParameters::SummonHierarchyEditor()
 {
-	SystemViewModel->FocusTab(FNiagaraSystemToolkitModeBase::UserParametersHierarchyTabID);
+	SystemViewModel.Pin()->FocusTab(FNiagaraSystemToolkitModeBase::UserParametersHierarchyTabID);
 	return FReply::Handled();
 }
 
diff --git a/Engine/Plugins/FX/Niagara/Source/NiagaraEditor/Private/Widgets/SNiagaraSystemUserParameters.h b/Engine/Plugins/FX/Niagara/Source/NiagaraEditor/Private/Widgets/SNiagaraSystemUserParameters.h
--- a/Engine/Plugins/FX/Niagara/Source/NiagaraEditor/Private/Widgets/SNiagaraSystemUserParameters.h
+++ b/Engine/Plugins/FX/Niagara/Source/NiagaraEditor/Private/Widgets/SNiagaraSystemUserParameters.h
@@ -8,18 +8,15 @@
 #include "Widgets/SCompoundWidget.h"
 #include "Widgets/DeclarativeSyntaxSupport.h"
 #include "Widgets/SNiagaraParameterMenu.h"
+#include "ViewModels/TNiagaraViewModelManager.h"
 
 class SNiagaraSystemUserParameters : public SCompoundWidget
 {
-public:
-	DECLARE_DELEGATE_OneParam(FOnThumbnailCaptured, UTexture2D*);
-
 public:
 	SLATE_BEGIN_ARGS( SNiagaraSystemUserParameters ){}
 	SLATE_END_ARGS()
 	
 	void Construct(const FArguments& InArgs, TSharedPtr<FNiagaraSystemViewModel> InSystemViewModel);
-	virtual ~SNiagaraSystemUserParameters() override {}
 	
 	virtual void Tick( const FGeometry& AllottedGeometry, const double InCurrentTime, const float InDeltaTime ) override {}
 
@@ -31,7 +28,7 @@ private:
 
 	FReply SummonHierarchyEditor();
 private:
-	TSharedPtr<FNiagaraSystemViewModel> SystemViewModel;
+	TWeakPtr<FNiagaraSystemViewModel> SystemViewModel;
 	
 	TSharedPtr<SComboButton> AddParameterButton;
 	TSharedPtr<SNiagaraAddParameterFromPanelMenu> ParameterPanel;
diff --git a/Engine/Plugins/FX/Niagara/Source/NiagaraEditor/Public/ViewModels/HierarchyEditor/NiagaraHierarchyViewModelBase.h b/Engine/Plugins/FX/Niagara/Source/NiagaraEditor/Public/ViewModels/HierarchyEditor/NiagaraHierarchyViewModelBase.h
--- a/Engine/Plugins/FX/Niagara/Source/NiagaraEditor/Public/ViewModels/HierarchyEditor/NiagaraHierarchyViewModelBase.h
+++ b/Engine/Plugins/FX/Niagara/Source/NiagaraEditor/Public/ViewModels/HierarchyEditor/NiagaraHierarchyViewModelBase.h
@@ -155,6 +155,7 @@ public:
 	void SetCategoryName(FName NewCategory) { Category = NewCategory; }
 	FName GetCategoryName() const { return Category; }
 
+	FText GetCategoryAsText() const { return FText::FromName(Category); }
 	FText GetTooltip() const { return Tooltip; }
 
 	void SetSection(UNiagaraHierarchySection* InSection) { Section = InSection; }
@@ -165,7 +166,7 @@ public:
 	virtual FString ToString() const override { return Category.ToString(); }
 
 private:
-	UPROPERTY()
+	UPROPERTY(EditAnywhere, Category = "Category")
 	FName Category;
 
 	/** The tooltip used when the user is hovering this category */
@@ -192,8 +193,9 @@ public:
 
 	FText GetTooltip() const { return Tooltip; }
 
+	virtual FString ToString() const override { return Section.ToString(); }
 private:
-	UPROPERTY()
+	UPROPERTY(EditAnywhere, Category = "Section")
 	FName Section;
 
 	/** The tooltip used when the user is hovering this section */
@@ -242,7 +244,7 @@ public:
 	DECLARE_DELEGATE_OneParam(FRefreshTreeWidget, bool bFullRefresh)
 	DECLARE_DELEGATE_RetVal_OneParam(bool, FIsItemSelected, TSharedPtr<FNiagaraHierarchyItemViewModelBase> Item)
 	DECLARE_DELEGATE_OneParam(FSelectObjectInDetailsPanel, UObject* Object)
-
+	
 	GENERATED_BODY()
 	
 	UNiagaraHierarchyViewModelBase();
@@ -260,6 +262,9 @@ public:
 	/** The hierarchy root the widget is editing. This should point to persistent data stored somewhere else as the serialized root of the hierarchy. */
 	virtual UNiagaraHierarchyRoot* GetHierarchyDataRoot() const PURE_VIRTUAL(UNiagaraHierarchyViewModelBase::GetHierarchyDataRoot, return nullptr;);
 	TSharedPtr<struct FNiagaraHierarchyRootViewModel> GetHierarchyViewModelRoot() const { return HierarchyViewModelRoot; }
+
+	/** This function will create the view model for a given item. Override to customize view model behavior by providing custom classes. */
+	virtual TSharedPtr<FNiagaraHierarchyItemViewModelBase> CreateViewModelForData(UNiagaraHierarchyItemBase* ItemBase, TSharedPtr<FNiagaraHierarchyItemViewModelBase> Parent);
 	
 	/** Source items reflect the base, unedited status of items to edit into a hierarchy */
 	const TArray<TSharedPtr<FNiagaraHierarchyItemViewModelBase>>& GetSourceItems() const;
@@ -284,21 +289,16 @@ public:
 	virtual bool SupportsDetailsPanel() { return false; }
 	// Overriding this will give the details panel instance customizations for specific UClasses
 	virtual TArray<TTuple<UClass*, FOnGetDetailCustomizationInstance>> GetInstanceCustomizations() { return {}; }
-	// Child view models need to override this to handle their own row generation. By default, we provide category widgets
-	virtual TSharedRef<SWidget> GenerateRowContentWidget(TSharedRef<FNiagaraHierarchyItemViewModelBase>) const;
 	
 	void RefreshSourceView(bool bFullRefresh = false) const;
 	void RefreshHierarchyView(bool bFullRefresh = false) const;
 	void RefreshSectionsWidget() const;
 	void SelectItemInDetailsPanel(UObject* Object) const;
-	bool IsItemSelected(TSharedPtr<FNiagaraHierarchyItemViewModelBase> Item) const;
 	
 	// Delegate that call functions from SNiagaraHierarchy
 	FRefreshTreeWidget& OnRefreshSourceView() { return RefreshSourceViewDelegate; }
 	FRefreshTreeWidget& OnRefreshHierarchyView() { return RefreshHierarchyWidgetDelegate; }
 	FSimpleDelegate& OnRefreshSections() { return RefreshSectionsWidgetDelegate; }
-	FSelectObjectInDetailsPanel& OnSelectObjectInDetailsPanel() { return SelectObjectInDetailsPanelDelegate; }
-	FIsItemSelected& IsItemSelected() { return IsItemSelectedDelegate; }
 
 	// Delegates for external systems
 	FOnHierarchyChanged& OnHierarchyChanged() { return OnHierarchyChangedDelegate; } 
@@ -310,11 +310,10 @@ public:
 	bool IsSectionActive(const UNiagaraHierarchySection* Section) const;
 	
 	FString OnItemToStringDebug(TSharedPtr<FNiagaraHierarchyItemViewModelBase> ItemBaseViewModel) const; 
-
 protected:
 	virtual void InitializeInternal() {}
 	virtual void FinalizeInternal() {}
-
+	
 	virtual void PostUndo(bool bSuccess) override;
 	virtual void PostRedo(bool bSuccess) override;
 protected:
@@ -335,13 +334,11 @@ protected:
 	FRefreshTreeWidget RefreshSourceViewDelegate;
 	FRefreshTreeWidget RefreshHierarchyWidgetDelegate;
 	FSimpleDelegate RefreshSectionsWidgetDelegate;
-	FIsItemSelected IsItemSelectedDelegate;
-	FSelectObjectInDetailsPanel SelectObjectInDetailsPanelDelegate;
 
 	FOnHierarchyChanged OnHierarchyChangedDelegate;
 };
 
-struct FNiagaraHierarchyItemViewModelBase : TSharedFromThis<FNiagaraHierarchyItemViewModelBase>
+struct NIAGARAEDITOR_API FNiagaraHierarchyItemViewModelBase : TSharedFromThis<FNiagaraHierarchyItemViewModelBase>
 {
 	DECLARE_MULTICAST_DELEGATE(FOnSynced)
 	
@@ -357,6 +354,12 @@ struct FNiagaraHierarchyItemViewModelBase : TSharedFromThis<FNiagaraHierarchyIte
 
 	UNiagaraHierarchyItemBase* GetDataMutable() const { return ItemBase; }
 	const UNiagaraHierarchyItemBase* GetData() const { return ItemBase; }
+	
+	template<class T>
+	T* GetDataMutable() const { return Cast<T>(ItemBase); }
+	
+	template<class T>
+	const T* GetData() const { return Cast<T>(ItemBase); }
 
 	FString ToString() const { return ItemBase->ToString(); }
 	
@@ -378,6 +381,8 @@ struct FNiagaraHierarchyItemViewModelBase : TSharedFromThis<FNiagaraHierarchyIte
 	int32 FindIndexOfDataChild(TSharedPtr<FNiagaraHierarchyItemViewModelBase> Child) const;
 	int32 FindIndexOfDataChild(UNiagaraHierarchyItemBase* Child) const;
 
+	/** Deleting will finalize as well, but is a user triggered action whereas Finalize is an internal call. */
+	void Delete();
 	/** Finalizing the view model will finalize the data & the children as well as handle cleanup. */
 	void Finalize();
 
@@ -405,6 +410,9 @@ struct FNiagaraHierarchyItemViewModelBase : TSharedFromThis<FNiagaraHierarchyIte
 			OnRequestRenameDelegate.Execute();
 		}
 	}
+
+	/** The UObject we display in the details panel when this item is selected. By default it's the item the view model represents. */
+	virtual UObject* GetDataForEditing() { return ItemBase; }
 	
 	/** Used to create customized drag drop ops. */
 	TSharedRef<class FNiagaraHierarchyDragDropOp> CreateDragDropOp();
@@ -457,6 +465,7 @@ struct FNiagaraHierarchyRootViewModel : FNiagaraHierarchyItemViewModelBase
 	virtual ~FNiagaraHierarchyRootViewModel() override {}
 
 	virtual bool CanDrag() override { return false; }
+	virtual FReply OnDroppedOn(const FDragDropEvent& DragDropEvent, EItemDropZone ItemDropZone, TSharedPtr<FNiagaraHierarchyItemViewModelBase> Item) override;
 
 	TSharedPtr<struct FNiagaraHierarchySectionViewModel> AddNewSection();
 	TArray<TSharedPtr<struct FNiagaraHierarchySectionViewModel>>& GetSectionViewModels() { return SectionViewModels; }
@@ -477,7 +486,7 @@ struct FNiagaraHierarchySectionViewModel : FNiagaraHierarchyItemViewModelBase
 	
 	void SetSectionNameAsText(const FText& Text) const { Cast<UNiagaraHierarchySection>(ItemBase)->SetSectionNameAsText(Text); }
 	FText GetSectionNameAsText() const { return Cast<UNiagaraHierarchySection>(ItemBase)->GetSectionNameAsText(); }
-	
+	FText GetSectionTooltip() const { return Cast<UNiagaraHierarchySection>(ItemBase)->GetTooltip(); }
 	virtual bool CanDrag() override { return false; }
 	virtual bool CanRename() override { return true; }
 	virtual bool CanDelete() override { return true; }
@@ -487,7 +496,7 @@ struct FNiagaraHierarchySectionViewModel : FNiagaraHierarchyItemViewModelBase
 	virtual void FinalizeInternal() override;
 };
 
-struct FNiagaraHierarchyItemViewModel : FNiagaraHierarchyItemViewModelBase
+struct NIAGARAEDITOR_API FNiagaraHierarchyItemViewModel : FNiagaraHierarchyItemViewModelBase
 {
 	FNiagaraHierarchyItemViewModel(UNiagaraHierarchyItem* InItem, TSharedPtr<FNiagaraHierarchyItemViewModelBase> InParent, TWeakObjectPtr<UNiagaraHierarchyViewModelBase> InHierarchyViewModel) : FNiagaraHierarchyItemViewModelBase(InItem, InParent, InHierarchyViewModel) {}
 	
diff --git a/Engine/Plugins/FX/Niagara/Source/NiagaraEditor/Public/ViewModels/HierarchyEditor/NiagaraUserParametersHierarchyViewModel.h b/Engine/Plugins/FX/Niagara/Source/NiagaraEditor/Public/ViewModels/HierarchyEditor/NiagaraUserParametersHierarchyViewModel.h
--- a/Engine/Plugins/FX/Niagara/Source/NiagaraEditor/Public/ViewModels/HierarchyEditor/NiagaraUserParametersHierarchyViewModel.h
+++ b/Engine/Plugins/FX/Niagara/Source/NiagaraEditor/Public/ViewModels/HierarchyEditor/NiagaraUserParametersHierarchyViewModel.h
@@ -30,6 +30,15 @@ private:
 	TObjectPtr<UNiagaraSystem> System;
 };
 
+struct FNiagaraHierarchyUserParameterViewModel : public FNiagaraHierarchyItemViewModel
+{
+	FNiagaraHierarchyUserParameterViewModel(UNiagaraHierarchyUserParameter* InItem, TSharedPtr<FNiagaraHierarchyItemViewModelBase> InParent, TWeakObjectPtr<UNiagaraHierarchyViewModelBase> InHierarchyViewModel)
+	: FNiagaraHierarchyItemViewModel(InItem, InParent, InHierarchyViewModel) {}
+
+	/** For editing in the details panel we want to handle the script variable that is represented by the hierarchy item, not the hierarchy item itself. */
+	virtual UObject* GetDataForEditing() override;
+};
+
 UCLASS()
 class UNiagaraUserParametersHierarchyViewModel : public UNiagaraHierarchyViewModelBase
 {
@@ -45,18 +54,16 @@ public:
 	void Initialize(TSharedRef<FNiagaraSystemViewModel> InSystemViewModel);
 	
 	virtual UNiagaraHierarchyRoot* GetHierarchyDataRoot() const override;
+	virtual TSharedPtr<FNiagaraHierarchyItemViewModelBase> CreateViewModelForData(UNiagaraHierarchyItemBase* ItemBase, TSharedPtr<FNiagaraHierarchyItemViewModelBase> Parent) override;
 	
 	virtual void PrepareSourceItems() override;
 	virtual void SetupCommands() override;
 
 	virtual TOptional<EItemDropZone> CanDropOn(TSharedPtr<FNiagaraHierarchyItemViewModelBase> SourceDropItem, TSharedPtr<FNiagaraHierarchyItemViewModelBase> TargetDropItem, EItemDropZone DropZone = EItemDropZone::OntoItem) override;
 	virtual TSharedRef<FNiagaraHierarchyDragDropOp> CreateDragDropOp(TSharedRef<FNiagaraHierarchyItemViewModelBase> Item) override;
-
-	virtual void OnSelectionChanged(TSharedPtr<FNiagaraHierarchyItemViewModelBase> HierarchyItem) override;
 	
 	virtual bool SupportsDetailsPanel() override { return true; }
 	virtual TArray<TTuple<UClass*, FOnGetDetailCustomizationInstance>> GetInstanceCustomizations() override;
-	virtual TSharedRef<SWidget> GenerateRowContentWidget(TSharedRef<FNiagaraHierarchyItemViewModelBase>) const override;
 protected:
 	virtual void FinalizeInternal() override;
 private:
diff --git a/Engine/Plugins/FX/Niagara/Source/NiagaraEditor/Public/ViewModels/NiagaraParameterPanelViewModel.h b/Engine/Plugins/FX/Niagara/Source/NiagaraEditor/Public/ViewModels/NiagaraParameterPanelViewModel.h
--- a/Engine/Plugins/FX/Niagara/Source/NiagaraEditor/Public/ViewModels/NiagaraParameterPanelViewModel.h
+++ b/Engine/Plugins/FX/Niagara/Source/NiagaraEditor/Public/ViewModels/NiagaraParameterPanelViewModel.h
@@ -386,8 +386,6 @@ public:
 
 	virtual bool GetSectionEnabled(FText Section) const override;
 
-	bool IsUserSectionActive() const { return ActiveSectionIndex == UserOnlyIdx && UserOnlyIdx != -1; }
-
 	virtual bool GetShowSections() const override { return true; }
 	virtual bool GetNamespaceActive(const FName& InNamespace) const override { return true; }
 
@@ -451,7 +449,6 @@ private:
 	static TArray<FNiagaraParameterPanelCategory> DefaultAdvancedScriptCategories;
 
 	TSharedPtr<FNiagaraObjectSelection> VariableObjectSelection;
-	int32 UserOnlyIdx = -1;
 	int32 ActiveScriptIdx = -1;
 	int32 ActiveSystemIdx = -1;
 	int32 LastActiveSystemSectionIdx = -1;
diff --git a/Engine/Plugins/FX/Niagara/Source/NiagaraEditor/Public/ViewModels/NiagaraSystemViewModel.h b/Engine/Plugins/FX/Niagara/Source/NiagaraEditor/Public/ViewModels/NiagaraSystemViewModel.h
--- a/Engine/Plugins/FX/Niagara/Source/NiagaraEditor/Public/ViewModels/NiagaraSystemViewModel.h
+++ b/Engine/Plugins/FX/Niagara/Source/NiagaraEditor/Public/ViewModels/NiagaraSystemViewModel.h
@@ -10,6 +10,7 @@
 #include "NiagaraEditorCommon.h"
 #include "NiagaraOverviewNode.h"
 #include "NiagaraSystemScalabilityViewModel.h"
+#include "NiagaraUserParameterPanelViewModel.h"
 #include "TickableEditorObject.h"
 #include "UObject/GCObject.h"
 #include "UObject/ObjectKey.h"
@@ -127,7 +128,7 @@ public:
 
 	DECLARE_MULTICAST_DELEGATE(FOnPreClose);
 
-	DECLARE_MULTICAST_DELEGATE_OneParam(FOnRequestFocusTab, FName /* TabName */);
+	DECLARE_MULTICAST_DELEGATE_TwoParams(FOnRequestFocusTab, FName /* TabName */, bool /* DrawAttention */);
 
 	DECLARE_MULTICAST_DELEGATE_ThreeParams(FOnExternalRenameParameter, const FNiagaraVariableBase&, const FNiagaraVariableBase&, UNiagaraEmitter*);
 	DECLARE_MULTICAST_DELEGATE_TwoParams(FOnExternalRemoveParameter, const FNiagaraVariableBase&,  UNiagaraEmitter*);
@@ -351,7 +352,7 @@ public:
 
 	FOnRequestFocusTab& OnRequestFocusTab();
 
-	NIAGARAEDITOR_API void FocusTab(FName TabName);
+	NIAGARAEDITOR_API void FocusTab(FName TabName, bool bDrawAttention = false);
 	
 	/** Gets the system toolkit command list. */
 	NIAGARAEDITOR_API TSharedPtr<FUICommandList> GetToolkitCommands();
@@ -440,6 +441,8 @@ public:
 
 	const TSharedPtr<FNiagaraSystemGraphSelectionViewModel>& GetSystemGraphSelectionViewModel() const { return SystemGraphSelectionViewModel; };
 
+	TSharedPtr<FNiagaraUserParameterPanelViewModel> GetUserParameterPanelViewModel() const { return UserParameterPanelViewModel; }
+	
 	/** Utility method to gather all graphs. */
 	TArray<UNiagaraGraph*> GetAllGraphs();
 
@@ -754,4 +757,6 @@ private:
 
 	/** ViewModel for caching the selected graph state and notifying when it changes. */
 	TSharedPtr<FNiagaraSystemGraphSelectionViewModel> SystemGraphSelectionViewModel;
+
+	TSharedPtr<FNiagaraUserParameterPanelViewModel> UserParameterPanelViewModel;
 };
\ No newline at end of file
diff --git a/Engine/Plugins/FX/Niagara/Source/NiagaraEditor/Public/ViewModels/NiagaraUserParameterPanelViewModel.h b/Engine/Plugins/FX/Niagara/Source/NiagaraEditor/Public/ViewModels/NiagaraUserParameterPanelViewModel.h
new file mode 100644
--- /dev/null
+++ b/Engine/Plugins/FX/Niagara/Source/NiagaraEditor/Public/ViewModels/NiagaraUserParameterPanelViewModel.h
@@ -0,0 +1,27 @@
+// Copyright Epic Games, Inc. All Rights Reserved.
+
+#pragma once
+
+#include "CoreMinimal.h"
+#include "NiagaraSystem.h"
+#include "ViewModels/TNiagaraViewModelManager.h"
+
+class FNiagaraSystemViewModel;
+
+class FNiagaraUserParameterPanelViewModel 
+	: public TSharedFromThis<FNiagaraUserParameterPanelViewModel>
+    , public TNiagaraViewModelManager<UNiagaraSystem, FNiagaraUserParameterPanelViewModel>
+{
+public:
+    DECLARE_DELEGATE_OneParam(FOnParameterAdded, FNiagaraVariable);
+
+	virtual ~FNiagaraUserParameterPanelViewModel();
+	
+	void Initialize(TSharedRef<FNiagaraSystemViewModel> InSystemViewModel);
+	
+    FOnParameterAdded& OnParameterAdded() { return OnParameterAddedDelegate; }
+private:
+    FOnParameterAdded OnParameterAddedDelegate;
+	TWeakPtr<FNiagaraSystemViewModel> SystemViewModelWeak;
+    TNiagaraViewModelManager<UNiagaraSystem, FNiagaraUserParameterPanelViewModel>::Handle RegisteredHandle;
+};
\ No newline at end of file
diff --git a/Engine/Plugins/FX/Niagara/Source/NiagaraEditor/Public/Widgets/SNiagaraHierarchy.h b/Engine/Plugins/FX/Niagara/Source/NiagaraEditor/Public/Widgets/SNiagaraHierarchy.h
--- a/Engine/Plugins/FX/Niagara/Source/NiagaraEditor/Public/Widgets/SNiagaraHierarchy.h
+++ b/Engine/Plugins/FX/Niagara/Source/NiagaraEditor/Public/Widgets/SNiagaraHierarchy.h
@@ -18,18 +18,21 @@ struct FNiagaraHierarchySectionViewModel;
 
 class SNiagaraHierarchy : public SCompoundWidget
 {
+	DECLARE_DELEGATE_RetVal_OneParam(TSharedRef<SWidget>, FOnGenerateRowContentWidget, TSharedRef<FNiagaraHierarchyItemViewModelBase> HierarchyItem);
+	DECLARE_DELEGATE_RetVal_OneParam(TSharedRef<SWidget>, FOnGenerateCustomDetailsPanelNameWidget, TSharedPtr<FNiagaraHierarchyItemViewModelBase> HierarchyItem);
+
 	SLATE_BEGIN_ARGS(SNiagaraHierarchy)
 		: _bReadOnly(true)
 	{}
 		SLATE_ARGUMENT(bool, bReadOnly)
+		SLATE_EVENT(FOnGenerateRowContentWidget, OnGenerateRowContentWidget)
+		SLATE_EVENT(FOnGenerateCustomDetailsPanelNameWidget, OnGenerateCustomDetailsPanelNameWidget)
 	SLATE_END_ARGS()
 
 public:
 	void Construct(const FArguments& InArgs, TObjectPtr<UNiagaraHierarchyViewModelBase> InHierarchyViewModel);
 	virtual ~SNiagaraHierarchy() override;
 	
-	virtual FReply OnKeyDown(const FGeometry& MyGeometry, const FKeyEvent& InKeyEvent) override;
-
 	void RefreshSourceView(bool bFullRefresh = false) const;
 	void RefreshHierarchyView(bool bFullRefresh = false) const;
 	void RefreshSectionsWidget();
@@ -48,7 +51,15 @@ private:
 	void DeleteSelectedItems();
 	bool CanDeleteSelectedItems() const;
 
-	void OnSelectionChanged(TSharedPtr<FNiagaraHierarchyItemViewModelBase> HierarchyItem, ESelectInfo::Type Type) const;
+	void OnSelectionChanged(TSharedPtr<FNiagaraHierarchyItemViewModelBase> HierarchyItem, ESelectInfo::Type Type, bool bFromHierarchy) const;
+
+	virtual FReply OnKeyDown(const FGeometry& MyGeometry, const FKeyEvent& InKeyEvent) override;
+	virtual FReply OnMouseButtonDown(const FGeometry& MyGeometry, const FPointerEvent& MouseEvent) override;
+	virtual FReply OnMouseButtonUp(const FGeometry& MyGeometry, const FPointerEvent& MouseEvent) override;
+
+	/** This handles the drag & drop at the root level of the hierarchy, i.e. when dragging and dropping a parameter into the "empty" hierarchy */
+	FReply HandleHierarchyRootDrop(const FGeometry&, const FDragDropEvent& DragDropEvent) const;
+	bool CanDropOnRoot(TSharedPtr<FDragDropOperation> DragDropOp) const;
 private:
 	TWeakObjectPtr<UNiagaraHierarchyViewModelBase> HierarchyViewModel;
 	TSharedPtr<STreeView<TSharedPtr<FNiagaraHierarchyItemViewModelBase>>> SourceTreeView;
@@ -56,26 +67,37 @@ private:
 	TMap<TSharedPtr<FNiagaraHierarchySectionViewModel>, TSharedPtr<class SNiagaraHierarchySection>> SectionsWidgetMap;
 	SVerticalBox::FSlot* SectionsSlot = nullptr;
 	TSharedPtr<IDetailsView> DetailsPanel;
+
+	// Slate widget arguments
+private:
+	FOnGenerateRowContentWidget OnGenerateRowContentWidget;
+	FOnGenerateCustomDetailsPanelNameWidget OnGenerateCustomDetailsPanelNameWidget;
 };
 
 class SNiagaraHierarchySection : public SCompoundWidget
 {
+	DECLARE_DELEGATE_OneParam(FOnSectionActivated, TSharedPtr<FNiagaraHierarchySectionViewModel> SectionViewModel)
+	
 	SLATE_BEGIN_ARGS(SNiagaraHierarchySection)
 	{}
 		SLATE_EVENT(FIsSelected, IsSelected)
+		SLATE_EVENT(FOnSectionActivated, OnSectionActivated)
 	SLATE_END_ARGS()
 
 public:
 	void Construct(const FArguments& InArgs, TSharedPtr<struct FNiagaraHierarchySectionViewModel> InSection, TWeakObjectPtr<UNiagaraHierarchyViewModelBase> HierarchyViewModel);
 	virtual ~SNiagaraHierarchySection() override;
-	
 
 	bool OnCanAcceptDrop(TSharedPtr<FDragDropOperation> DragDropOperation) const;
 	bool OnCanAcceptSectionDrop(TSharedPtr<FDragDropOperation> DragDropOperation, EItemDropZone ItemDropZone) const;
 	FReply OnDroppedOn(const FGeometry&, const FDragDropEvent& DragDropEvent, EItemDropZone DropZone) const;
 
 	void EnterEditingMode() const;
+
+private:
 	virtual FReply OnMouseButtonDown(const FGeometry& MyGeometry, const FPointerEvent& MouseEvent) override;
+	virtual FReply OnMouseButtonUp(const FGeometry& MyGeometry, const FPointerEvent& MouseEvent) override;
+	
 	virtual FReply OnDragDetected(const FGeometry& MyGeometry, const FPointerEvent& MouseEvent) override;
 	virtual void OnDragLeave(const FDragDropEvent& DragDropEvent) override;
 	
@@ -84,6 +106,7 @@ public:
 	UNiagaraHierarchySection* TryGetSectionData() const;
 
 	FText GetText() const;
+	FText GetTooltipText() const;
 	void OnRenameSection(const FText& Text, ETextCommit::Type CommitType) const;
 	bool OnVerifySectionRename(const FText& NewName, FText& OutTooltip) const;
 
@@ -94,6 +117,8 @@ private:
 	TSharedPtr<SInlineEditableTextBlock> InlineEditableTextBlock;
 	TWeakObjectPtr<UNiagaraHierarchyViewModelBase> HierarchyViewModel;
 	TSharedPtr<FNiagaraHierarchySectionViewModel> SectionViewModel;
+private:
+	FOnSectionActivated OnSectionActivatedDelegate;
 };
 
 class SNiagaraHierarchyCategory : public SCompoundWidget
diff --git a/Engine/Plugins/FX/Niagara/Source/NiagaraEditor/Public/Widgets/SNiagaraParameterName.h b/Engine/Plugins/FX/Niagara/Source/NiagaraEditor/Public/Widgets/SNiagaraParameterName.h
--- a/Engine/Plugins/FX/Niagara/Source/NiagaraEditor/Public/Widgets/SNiagaraParameterName.h
+++ b/Engine/Plugins/FX/Niagara/Source/NiagaraEditor/Public/Widgets/SNiagaraParameterName.h
@@ -106,6 +106,8 @@ private:
 class NIAGARAEDITOR_API SNiagaraParameterNameTextBlock : public SCompoundWidget
 {
 public:
+	DECLARE_DELEGATE_RetVal_TwoParams(FReply, FOnDragDetectedHandler, const FGeometry&, const FPointerEvent&);
+	
 	SLATE_BEGIN_ARGS(SNiagaraParameterNameTextBlock)
 		: _EditableTextStyle(&FAppStyle::Get().GetWidgetStyle<FInlineEditableTextBlockStyle>("InlineEditableTextBlockStyle"))
 		, _ReadOnlyTextStyle(&FAppStyle::Get().GetWidgetStyle<FTextBlockStyle>("NormalText"))
@@ -122,11 +124,11 @@ public:
 		SLATE_ATTRIBUTE(FText, HighlightText)
 		SLATE_EVENT(FOnVerifyTextChanged, OnVerifyTextChanged)
 		SLATE_EVENT(FOnTextCommitted, OnTextCommitted)
+		SLATE_EVENT(FOnDragDetectedHandler, OnDragDetected)
 		SLATE_EVENT(FIsSelected, IsSelected)
 		SLATE_ARGUMENT(EHorizontalAlignment, DecoratorHAlign)
 		SLATE_ARGUMENT(FMargin, DecoratorPadding)
 		SLATE_NAMED_SLOT(FArguments, Decorator)
-		
 	SLATE_END_ARGS()
 
 	void Construct(const FArguments& InArgs);
@@ -142,11 +144,15 @@ private:
 
 	void NameChanged(FName InNewName);
 
+	virtual FReply OnMouseButtonDown(const FGeometry& MyGeometry, const FPointerEvent& MouseEvent) override;
+	virtual FReply OnDragDetected(const FGeometry& MyGeometry, const FPointerEvent& MouseEvent) override;
+
 private:
 	TAttribute<FText> ParameterText;
 	FOnVerifyTextChanged OnVerifyNameTextChangedDelegate;
 	FOnTextCommitted OnNameTextCommittedDelegate;
-
+	FOnDragDetectedHandler OnDragDetectedHandlerDelegate;
+	
 	mutable FText DisplayedParameterTextCache;
 	mutable FName ParameterNameCache;
 	TSharedPtr<SNiagaraParameterName> ParameterName;
diff --git a/Engine/Plugins/FX/Niagara/Source/NiagaraEditorWidgets/Private/SNiagaraOverviewStack.cpp b/Engine/Plugins/FX/Niagara/Source/NiagaraEditorWidgets/Private/SNiagaraOverviewStack.cpp
--- a/Engine/Plugins/FX/Niagara/Source/NiagaraEditorWidgets/Private/SNiagaraOverviewStack.cpp
+++ b/Engine/Plugins/FX/Niagara/Source/NiagaraEditorWidgets/Private/SNiagaraOverviewStack.cpp
@@ -54,6 +54,7 @@
 #include "ScopedTransaction.h"
 #include "Styling/StyleColors.h"
 #include "ViewModels/Stack/NiagaraStackRendererItem.h"
+#include "ViewModels/Stack/NiagaraStackSystemSettingsGroup.h"
 #include "Widgets/Layout/SWidgetSwitcher.h"
 #include "Widgets/Layout/SWrapBox.h"
 
@@ -1387,7 +1388,16 @@ void SNiagaraOverviewStack::OnSelectionChanged(UNiagaraStackEntry* InNewSelectio
 		}
 
 		bool bClearCurrentSelection = FSlateApplication::Get().GetModifierKeys().IsControlDown() == false;
-		OverviewSelectionViewModel->UpdateSelectedEntries(SelectedEntries, DeselectedEntries, bClearCurrentSelection);
+
+		// todo (ME) remove this special case when we remove the user parameters group from the stack in favor of the new user param tab
+		if(InNewSelection->IsA<UNiagaraStackSystemUserParametersGroup>())
+		{
+			InNewSelection->GetSystemViewModel()->FocusTab(TEXT("NiagaraSystemEditor_UserParameters"), true);
+		}
+		else
+		{
+			OverviewSelectionViewModel->UpdateSelectedEntries(SelectedEntries, DeselectedEntries, bClearCurrentSelection);
+		}
 
 		PreviousSelection.Empty();
 		for (UNiagaraStackEntry* SelectedEntry : SelectedEntries)
