- Fix static variables not working with SimStage bindings

diff --git a/Engine/Plugins/FX/Niagara/Source/Niagara/Classes/NiagaraEmitter.h b/Engine/Plugins/FX/Niagara/Source/Niagara/Classes/NiagaraEmitter.h
--- a/Engine/Plugins/FX/Niagara/Source/Niagara/Classes/NiagaraEmitter.h
+++ b/Engine/Plugins/FX/Niagara/Source/Niagara/Classes/NiagaraEmitter.h
@@ -380,6 +380,7 @@ struct NIAGARA_API FVersionedNiagaraEmitterData
 	FORCEINLINE const FNiagaraEmitterScalabilitySettings& GetScalabilitySettings()const { return CurrentScalabilitySettings; }
 	const FNiagaraEmitterScalabilityOverride& GetCurrentOverrideSettings() const;
 	UNiagaraSimulationStageBase* GetSimulationStageById(FGuid ScriptUsageId) const;
+	bool BuildParameterStoreRendererBindings(FNiagaraParameterStore& ParameterStore) const;
 	void RebuildRendererBindings(const UNiagaraEmitter& Emitter);
 	FORCEINLINE static FBox GetDefaultFixedBounds() { return FBox(FVector(-100), FVector(100)); }
 	
diff --git a/Engine/Plugins/FX/Niagara/Source/Niagara/Private/NiagaraEmitter.cpp b/Engine/Plugins/FX/Niagara/Source/Niagara/Private/NiagaraEmitter.cpp
--- a/Engine/Plugins/FX/Niagara/Source/Niagara/Private/NiagaraEmitter.cpp
+++ b/Engine/Plugins/FX/Niagara/Source/Niagara/Private/NiagaraEmitter.cpp
@@ -1643,32 +1643,67 @@ void FVersionedNiagaraEmitterData::UpdateDebugName(const UNiagaraEmitter& Emitte
 	}
 }
 
-void FVersionedNiagaraEmitterData::RebuildRendererBindings(const UNiagaraEmitter& Emitter)
+bool FVersionedNiagaraEmitterData::BuildParameterStoreRendererBindings(FNiagaraParameterStore& ParameterStore) const
 {
-#if WITH_EDITORONLY_DATA
-
-	// For right now we only capture the values of the static variables here. So we create a temp
-	// parameter store for ALL possible values, then fill in the real one with the static variables only.
-	bool bAnyRendererBindingsAdded = false;
-	FNiagaraParameterStore TempStore;
+	bool bAnyBindingsAdded = false;
 	for (UNiagaraRendererProperties* Props : GetRenderers())
 	{
 		if (Props && Props->bIsEnabled)
 		{
-			bAnyRendererBindingsAdded |= Props->PopulateRequiredBindings(TempStore);
+			bAnyBindingsAdded |= Props->PopulateRequiredBindings(ParameterStore);
 		}
 	}
 
-	TArray<UNiagaraScript*> TargetScripts;
-	if (UNiagaraSystem* SystemOwner = Emitter.GetTypedOuter<UNiagaraSystem>())
+	if (GPUComputeScript)
 	{
-		TargetScripts.Add(SystemOwner->GetSystemSpawnScript());
-		TargetScripts.Add(SystemOwner->GetSystemUpdateScript());
+		for (const FSimulationStageMetaData& SimStageMetaData : GPUComputeScript->GetSimulationStageMetaData())
+		{
+			if (!SimStageMetaData.EnabledBinding.IsNone())
+			{
+				bAnyBindingsAdded |= ParameterStore.AddParameter(FNiagaraVariable(FNiagaraTypeDefinition::GetBoolDef(), SimStageMetaData.EnabledBinding), false);
+			}
+			if (SimStageMetaData.bOverrideElementCount)
+			{
+				if (!SimStageMetaData.ElementCountXBinding.IsNone())
+				{
+					bAnyBindingsAdded |= ParameterStore.AddParameter(FNiagaraVariable(FNiagaraTypeDefinition::GetIntDef(), SimStageMetaData.ElementCountXBinding), false);
+				}
+				if (!SimStageMetaData.ElementCountYBinding.IsNone())
+				{
+					bAnyBindingsAdded |= ParameterStore.AddParameter(FNiagaraVariable(FNiagaraTypeDefinition::GetIntDef(), SimStageMetaData.ElementCountYBinding), false);
+				}
+				if (!SimStageMetaData.ElementCountZBinding.IsNone())
+				{
+					bAnyBindingsAdded |= ParameterStore.AddParameter(FNiagaraVariable(FNiagaraTypeDefinition::GetIntDef(), SimStageMetaData.ElementCountZBinding), false);
+				}
+			}
+			if (!SimStageMetaData.NumIterationsBinding.IsNone())
+			{
+				bAnyBindingsAdded |= ParameterStore.AddParameter(FNiagaraVariable(FNiagaraTypeDefinition::GetIntDef(), SimStageMetaData.NumIterationsBinding), false);
+			}
+		}
 	}
-	GetScripts(TargetScripts, false, true);
+	return bAnyBindingsAdded;
+}
+
+void FVersionedNiagaraEmitterData::RebuildRendererBindings(const UNiagaraEmitter& Emitter)
+{
+#if WITH_EDITORONLY_DATA
+	// For right now we only capture the values of the static variables here. So we create a temp
+	// parameter store for ALL possible values, then fill in the real one with the static variables only.
+	FNiagaraParameterStore TempStore;
+	const bool bAnyRendererBindingsAdded = BuildParameterStoreRendererBindings(TempStore);
 
 	if (bAnyRendererBindingsAdded)
 	{
+		TArray<UNiagaraScript*> TargetScripts;
+		if (UNiagaraSystem* SystemOwner = Emitter.GetTypedOuter<UNiagaraSystem>())
+		{
+			TargetScripts.Add(SystemOwner->GetSystemSpawnScript());
+			TargetScripts.Add(SystemOwner->GetSystemUpdateScript());
+		}
+		GetScripts(TargetScripts, false, true);
+
 		RendererBindings.Empty();
 		TArrayView<const FNiagaraVariableWithOffset> Vars = TempStore.ReadParameterVariables();
 		for (const FNiagaraVariableWithOffset& Var : Vars)
diff --git a/Engine/Plugins/FX/Niagara/Source/Niagara/Private/NiagaraEmitterInstance.cpp b/Engine/Plugins/FX/Niagara/Source/Niagara/Private/NiagaraEmitterInstance.cpp
--- a/Engine/Plugins/FX/Niagara/Source/Niagara/Private/NiagaraEmitterInstance.cpp
+++ b/Engine/Plugins/FX/Niagara/Source/Niagara/Private/NiagaraEmitterInstance.cpp
@@ -490,44 +490,7 @@ void FNiagaraEmitterInstance::Init(int32 InEmitterIdx, FNiagaraSystemInstanceID
 		// We may need to populate bindings that will be used in rendering
 		RendererBindings.UnbindAll();
 
-		bool bAnyRendererBindingsAdded = false;
-		for (UNiagaraRendererProperties* Props : EmitterData->GetRenderers())
-		{
-			if (Props && Props->bIsEnabled)
-			{
-				bAnyRendererBindingsAdded |= Props->PopulateRequiredBindings(RendererBindings);
-			}
-		}
-
-		if ( GPUExecContext != nullptr )
-		{
-			for (const FSimulationStageMetaData& SimStageMetaData : GPUExecContext->SimStageInfo )
-			{
-				if (!SimStageMetaData.EnabledBinding.IsNone())
-				{
-					bAnyRendererBindingsAdded |= RendererBindings.AddParameter(FNiagaraVariable(FNiagaraTypeDefinition::GetBoolDef(), SimStageMetaData.EnabledBinding), false);
-				}
-				if (SimStageMetaData.bOverrideElementCount)
-				{
-					if (!SimStageMetaData.ElementCountXBinding.IsNone())
-					{
-						bAnyRendererBindingsAdded |= RendererBindings.AddParameter(FNiagaraVariable(FNiagaraTypeDefinition::GetIntDef(), SimStageMetaData.ElementCountXBinding), false);
-					}
-					if (!SimStageMetaData.ElementCountYBinding.IsNone())
-					{
-						bAnyRendererBindingsAdded |= RendererBindings.AddParameter(FNiagaraVariable(FNiagaraTypeDefinition::GetIntDef(), SimStageMetaData.ElementCountYBinding), false);
-					}
-					if (!SimStageMetaData.ElementCountZBinding.IsNone())
-					{
-						bAnyRendererBindingsAdded |= RendererBindings.AddParameter(FNiagaraVariable(FNiagaraTypeDefinition::GetIntDef(), SimStageMetaData.ElementCountZBinding), false);
-					}
-				}
-				if ( !SimStageMetaData.NumIterationsBinding.IsNone() )
-				{
-					bAnyRendererBindingsAdded |= RendererBindings.AddParameter(FNiagaraVariable(FNiagaraTypeDefinition::GetIntDef(), SimStageMetaData.NumIterationsBinding), false);
-				}				
-			}
-		}
+		const bool bAnyRendererBindingsAdded = EmitterData->BuildParameterStoreRendererBindings(RendererBindings);
 
 		if (bAnyRendererBindingsAdded)
 		{
