- Fix validation errors in Niagara Landscape / Distance Field default bindings

diff --git a/Engine/Plugins/FX/Niagara/Source/Niagara/Private/NiagaraDataInterfaceLandscape.cpp b/Engine/Plugins/FX/Niagara/Source/Niagara/Private/NiagaraDataInterfaceLandscape.cpp
--- a/Engine/Plugins/FX/Niagara/Source/Niagara/Private/NiagaraDataInterfaceLandscape.cpp
+++ b/Engine/Plugins/FX/Niagara/Source/Niagara/Private/NiagaraDataInterfaceLandscape.cpp
@@ -323,7 +323,7 @@ struct FNDILandscapeData_RenderThread
 		const FUintVector4 DummyUint4(ForceInitToZero);
 
 		ShaderParameters->BaseColorVirtualTexture = GBlackTextureWithSRV->ShaderResourceViewRHI;
-		ShaderParameters->BaseColorVirtualTexturePageTable = GBlackTexture->TextureRHI;
+		ShaderParameters->BaseColorVirtualTexturePageTable = GBlackUintTexture->TextureRHI;
 		ShaderParameters->BaseColorVirtualTextureWorldToUvTransform = FMatrix44f::Identity;
 		ShaderParameters->BaseColorVirtualTextureSRGB = 0;
 		ShaderParameters->BaseColorVirtualTextureEnabled = 0;
@@ -398,7 +398,7 @@ struct FNDILandscapeData_RenderThread
 		const FUintVector4 DummyUint4(ForceInitToZero);
 
 		ShaderParameters->HeightVirtualTexture = GBlackTextureWithSRV->ShaderResourceViewRHI;
-		ShaderParameters->HeightVirtualTexturePageTable = GBlackTexture->TextureRHI;
+		ShaderParameters->HeightVirtualTexturePageTable = GBlackUintTexture->TextureRHI;
 		ShaderParameters->HeightVirtualTextureWorldToUvTransform = FMatrix44f::Identity;
 		ShaderParameters->HeightVirtualTextureEnabled = 0;
 		ShaderParameters->HeightVirtualTexturePackedUniform0 = DummyUint4;
@@ -502,7 +502,7 @@ struct FNDILandscapeData_RenderThread
 
 		ShaderParameters->NormalVirtualTexture0 =  GBlackTextureWithSRV->ShaderResourceViewRHI;
 		ShaderParameters->NormalVirtualTexture1 = GBlackTextureWithSRV->ShaderResourceViewRHI;
-		ShaderParameters->NormalVirtualTexturePageTable = GBlackTexture->TextureRHI;
+		ShaderParameters->NormalVirtualTexturePageTable = GBlackUintTexture->TextureRHI;
 		ShaderParameters->NormalVirtualTextureWorldToUvTransform = FMatrix44f::Identity;
 		ShaderParameters->NormalVirtualTextureEnabled = 0;
 		ShaderParameters->NormalVirtualTextureUnpackMode = (int32)ENormalUnpackType::None;
diff --git a/Engine/Plugins/FX/Niagara/Source/NiagaraShader/Private/NiagaraDistanceFieldHelper.cpp b/Engine/Plugins/FX/Niagara/Source/NiagaraShader/Private/NiagaraDistanceFieldHelper.cpp
--- a/Engine/Plugins/FX/Niagara/Source/NiagaraShader/Private/NiagaraDistanceFieldHelper.cpp
+++ b/Engine/Plugins/FX/Niagara/Source/NiagaraShader/Private/NiagaraDistanceFieldHelper.cpp
@@ -73,8 +73,37 @@ void FNiagaraDistanceFieldHelper::SetMeshDistanceFieldParameters(FRDGBuilder& Gr
 	}
 	else
 	{
+		class FDFDummyByteAddress : public FRenderResource
+		{
+		public:
+			TRefCountPtr<FRDGPooledBuffer> PooledBuffer;
+
+			virtual void InitRHI() override
+			{
+				FRDGBufferDesc BufferDesc = FRDGBufferDesc::CreateStructuredDesc(sizeof(uint32), 1);
+				BufferDesc.Usage = BUF_Static | BUF_ShaderResource | BUF_ByteAddressBuffer;
+
+				FRHIResourceCreateInfo CreateInfo(TEXT("FDFDummyByteAddress"));
+				FBufferRHIRef RHIBuffer = RHICreateStructuredBuffer(BufferDesc.BytesPerElement, BufferDesc.GetSize(), BufferDesc.Usage, CreateInfo);
+				{
+					void* Data = RHILockBuffer(RHIBuffer, 0, BufferDesc.GetSize(), RLM_WriteOnly);
+					FMemory::Memset(Data, 0, BufferDesc.GetSize());
+					RHIUnlockBuffer(RHIBuffer);
+				}
+
+				PooledBuffer = new FRDGPooledBuffer(RHIBuffer, BufferDesc, BufferDesc.NumElements, CreateInfo.DebugName);
+			}
+
+			virtual void ReleaseRHI() override
+			{
+				PooledBuffer.SafeRelease();
+			}
+		};
+
+		static TGlobalResource<FDFDummyByteAddress> GDFDummyByteAddress;
+
 		FRDGBufferSRVRef DefaultVector4 = GraphBuilder.CreateSRV(FRDGBufferSRVDesc(GSystemTextures.GetDefaultStructuredBuffer(GraphBuilder, sizeof(FVector4f))));
-		FRDGBufferSRVRef DefaultUInt32 = GraphBuilder.CreateSRV(FRDGBufferSRVDesc(GSystemTextures.GetDefaultStructuredBuffer(GraphBuilder, sizeof(uint32))));
+		FRDGBufferSRVRef DefaultUInt32 = GraphBuilder.CreateSRV(FRDGBufferSRVDesc(GraphBuilder.RegisterExternalBuffer(GDFDummyByteAddress.PooledBuffer, ERDGBufferFlags::SkipTracking)));
 
 		ObjectShaderParameters.SceneObjectBounds = DefaultVector4;
 		ObjectShaderParameters.SceneObjectData = DefaultVector4;
@@ -89,7 +118,7 @@ void FNiagaraDistanceFieldHelper::SetMeshDistanceFieldParameters(FRDGBuilder& Gr
 		AtlasShaderParameters.DistanceFieldIndirectionTable = DefaultUInt32;
 		AtlasShaderParameters.DistanceFieldIndirection2Table = DefaultVector4;
 		AtlasShaderParameters.DistanceFieldIndirectionAtlas = SystemTextures.VolumetricBlack;
-		AtlasShaderParameters.DistanceFieldBrickTexture = SystemTextures.Black;
+		AtlasShaderParameters.DistanceFieldBrickTexture = SystemTextures.VolumetricBlack;
 		AtlasShaderParameters.DistanceFieldSampler = TStaticSamplerState<SF_Bilinear, AM_Wrap, AM_Wrap, AM_Wrap>::GetRHI();
 		AtlasShaderParameters.DistanceFieldBrickSize = FVector3f::ZeroVector;
 		AtlasShaderParameters.DistanceFieldUniqueDataBrickSize = FVector3f::ZeroVector;
