- Add GPU stat for Niagara GPU ribbon generation

#rnx

diff --git a/Engine/Plugins/FX/Niagara/Source/Niagara/Private/NiagaraRendererRibbons.cpp b/Engine/Plugins/FX/Niagara/Source/Niagara/Private/NiagaraRendererRibbons.cpp
--- a/Engine/Plugins/FX/Niagara/Source/Niagara/Private/NiagaraRendererRibbons.cpp
+++ b/Engine/Plugins/FX/Niagara/Source/Niagara/Private/NiagaraRendererRibbons.cpp
@@ -43,6 +43,7 @@ DECLARE_CYCLE_STAT(TEXT("Generate Vertices GPU - Reduction Finalize [RT]"), STAT
 DECLARE_CYCLE_STAT(TEXT("Generate Vertices GPU - MultiRibbon Init [RT]"), STAT_NiagaraRenderRibbonsGenVerticesMultiRibbonInitGPU, STATGROUP_NiagaraRibbons);
 DECLARE_CYCLE_STAT(TEXT("Generate Vertices GPU - MultiRibbon Init Compute [RT]"), STAT_NiagaraRenderRibbonsGenVerticesMultiRibbonInitComputeGPU, STATGROUP_NiagaraRibbons);
 
+DECLARE_GPU_STAT_NAMED(NiagaraGPURibbons, TEXT("Niagara GPU Ribbons"));
 
 int32 GNiagaraRibbonTessellationEnabled = 1;
 static FAutoConsoleVariableRef CVarNiagaraRibbonTessellationEnabled(
@@ -2549,6 +2550,8 @@ void FNiagaraGpuRibbonsDataManager::GenerateAllGPUData(FRHICommandListImmediate&
 		return;
 	}
 
+	SCOPED_GPU_STAT(RHICmdList, NiagaraGPURibbons);
+
 	FNiagaraGpuComputeDispatchInterface* ComputeDispatchInterface = GetOwnerInterface();
 
 	// Handle all vertex gens first
