Remove #pragma once in cpp

diff --git a/Engine/Plugins/FX/Niagara/Source/NiagaraEditor/Private/ViewModels/NiagaraUserParameterPanelViewModel.cpp b/Engine/Plugins/FX/Niagara/Source/NiagaraEditor/Private/ViewModels/NiagaraUserParameterPanelViewModel.cpp
--- a/Engine/Plugins/FX/Niagara/Source/NiagaraEditor/Private/ViewModels/NiagaraUserParameterPanelViewModel.cpp
+++ b/Engine/Plugins/FX/Niagara/Source/NiagaraEditor/Private/ViewModels/NiagaraUserParameterPanelViewModel.cpp
@@ -1,7 +1,5 @@
 // Copyright Epic Games, Inc. All Rights Reserved.
 
-#pragma once
-
 #include "ViewModels/NiagaraUserParameterPanelViewModel.h"
 #include "CoreMinimal.h"
 #include "ViewModels/NiagaraSystemViewModel.h"
