Couple of fixes/improvements for the rigid body DI
-make the actor name text more legible
-use the actor label if one exists
-fix crash when a acquired actor goes away (deleted in the editor)

diff --git a/Engine/Plugins/FX/Niagara/Source/Niagara/Private/NiagaraDataInterfaceRigidMeshCollisionQuery.cpp b/Engine/Plugins/FX/Niagara/Source/Niagara/Private/NiagaraDataInterfaceRigidMeshCollisionQuery.cpp
--- a/Engine/Plugins/FX/Niagara/Source/Niagara/Private/NiagaraDataInterfaceRigidMeshCollisionQuery.cpp
+++ b/Engine/Plugins/FX/Niagara/Source/Niagara/Private/NiagaraDataInterfaceRigidMeshCollisionQuery.cpp
@@ -2,6 +2,7 @@
 
 #include "NiagaraDataInterfaceRigidMeshCollisionQuery.h"
 #include "Algo/ForEach.h"
+#include "Algo/RemoveIf.h"
 #include "Components/SkeletalMeshComponent.h"
 #include "Engine/Canvas.h"
 #include "Animation/SkeletalMeshActor.h"
@@ -598,6 +599,22 @@ void FNDIRigidMeshCollisionData::MergeActors(FMergedActorArray& MergedActors) co
 	Algo::ForEach(FoundActors, AppendActors);
 }
 
+bool FNDIRigidMeshCollisionData::TrimMissingActors()
+{
+	auto EvaluateWeakActor = [&](const TWeakObjectPtr<AActor>& ActorPtr)
+	{
+		return !ActorPtr.IsValid();
+	};
+
+	const int32 ExplicitActorCount = ExplicitActors.Num();
+	ExplicitActors.SetNum(Algo::StableRemoveIf(ExplicitActors, EvaluateWeakActor));
+
+	const int32 FoundActorCount = FoundActors.Num();
+	FoundActors.SetNum(Algo::StableRemoveIf(FoundActors, EvaluateWeakActor));
+
+	return ExplicitActorCount != ExplicitActors.Num() || FoundActorCount != FoundActors.Num();
+}
+
 void FNDIRigidMeshCollisionData::Init(int32 MaxNumPrimitives)
 {
 	const bool bHasActors = HasActors();
@@ -655,6 +672,8 @@ void FNDIRigidMeshCollisionData::Update(UNiagaraDataInterfaceRigidMeshCollisionQ
 
 	if (bHasActors)
 	{
+		bRequiresFullUpdate |= TrimMissingActors();
+
 		FMergedActorArray MergedActors;
 		MergeActors(MergedActors);
 
@@ -862,6 +881,7 @@ void UNiagaraDataInterfaceRigidMeshCollisionQuery::DrawDebugHud(UCanvas* Canvas,
 		if (!InstanceData_GT->ExplicitActors.IsEmpty() || !InstanceData_GT->FoundActors.IsEmpty())
 		{
 			const UFont* Font = GEngine->GetMediumFont();
+			Canvas->SetDrawColor(FColor::White);
 
 			auto DrawDebugActor = [&](TWeakObjectPtr<AActor>& InWeakActor, const TCHAR* ActorSourceString)
 			{
@@ -876,8 +896,14 @@ void UNiagaraDataInterfaceRigidMeshCollisionQuery::DrawDebugHud(UCanvas* Canvas,
 					{
 						DrawDebugCanvasWireBox(Canvas, CurrentTransform, FBox(-ActorBoundsExtent, ActorBoundsExtent), FColor::Yellow);
 
+						FString ActorLabel = Actor->GetActorLabel();
+						if (ActorLabel.Len() == 0)
+						{
+							ActorLabel = Actor->GetName();
+						}
+
 						const FVector ScreenLoc = Canvas->Project(ActorOrigin);
-						Canvas->DrawText(Font, FString::Printf(TEXT("RigidMeshDI[%s Actor] - %s"), ActorSourceString, *Actor->GetName()), ScreenLoc.X, ScreenLoc.Y);
+						Canvas->DrawText(Font, FString::Printf(TEXT("RigidMeshDI[%s Actor] - %s"), ActorSourceString, *ActorLabel), ScreenLoc.X, ScreenLoc.Y);
 					}
 				}
 			};
diff --git a/Engine/Plugins/FX/Niagara/Source/Niagara/Public/NiagaraDataInterfaceRigidMeshCollisionQuery.h b/Engine/Plugins/FX/Niagara/Source/Niagara/Public/NiagaraDataInterfaceRigidMeshCollisionQuery.h
--- a/Engine/Plugins/FX/Niagara/Source/Niagara/Public/NiagaraDataInterfaceRigidMeshCollisionQuery.h
+++ b/Engine/Plugins/FX/Niagara/Source/Niagara/Public/NiagaraDataInterfaceRigidMeshCollisionQuery.h
@@ -127,6 +127,8 @@ struct FNDIRigidMeshCollisionData
 	using FMergedActorArray = TArray<AActor*, TInlineAllocator<16>>;
 	void MergeActors(FMergedActorArray& MergedActors) const;
 
+	bool TrimMissingActors();
+
 	/** Physics asset Gpu buffer */
 	FNDIRigidMeshCollisionBuffer* AssetBuffer = nullptr;
 
