Fix for accessing ActorLabel with RigidBody DI in non editor builds

diff --git a/Engine/Plugins/FX/Niagara/Source/Niagara/Private/NiagaraDataInterfaceRigidMeshCollisionQuery.cpp b/Engine/Plugins/FX/Niagara/Source/Niagara/Private/NiagaraDataInterfaceRigidMeshCollisionQuery.cpp
--- a/Engine/Plugins/FX/Niagara/Source/Niagara/Private/NiagaraDataInterfaceRigidMeshCollisionQuery.cpp
+++ b/Engine/Plugins/FX/Niagara/Source/Niagara/Private/NiagaraDataInterfaceRigidMeshCollisionQuery.cpp
@@ -896,7 +896,10 @@ void UNiagaraDataInterfaceRigidMeshCollisionQuery::DrawDebugHud(UCanvas* Canvas,
 					{
 						DrawDebugCanvasWireBox(Canvas, CurrentTransform, FBox(-ActorBoundsExtent, ActorBoundsExtent), FColor::Yellow);
 
-						FString ActorLabel = Actor->GetActorLabel();
+						FString ActorLabel;
+#if WITH_EDITOR
+						ActorLabel = Actor->GetActorLabel();
+#endif
 						if (ActorLabel.Len() == 0)
 						{
 							ActorLabel = Actor->GetName();
