Incorporate project setting (bExperimentalVMEnabled) for running niagara's EVVM

diff --git a/Engine/Plugins/FX/Niagara/Source/Niagara/Private/NiagaraScriptExecutionContext.cpp b/Engine/Plugins/FX/Niagara/Source/Niagara/Private/NiagaraScriptExecutionContext.cpp
--- a/Engine/Plugins/FX/Niagara/Source/Niagara/Private/NiagaraScriptExecutionContext.cpp
+++ b/Engine/Plugins/FX/Niagara/Source/Niagara/Private/NiagaraScriptExecutionContext.cpp
@@ -1,6 +1,7 @@
 // Copyright Epic Games, Inc. All Rights Reserved.
 
 #include "NiagaraScriptExecutionContext.h"
+#include "NiagaraSettings.h"
 #include "NiagaraStats.h"
 #include "NiagaraDataInterface.h"
 #include "NiagaraSystemInstance.h"
@@ -63,7 +64,8 @@ bool FNiagaraScriptExecutionContextBase::Init(UNiagaraScript* InScript, ENiagara
 	HasInterpolationParameters = Script && Script->GetComputedVMCompilationId().HasInterpolatedParameters();
 
 #if VECTORVM_SUPPORTS_EXPERIMENTAL && VECTORVM_SUPPORTS_LEGACY
-	const bool ScriptSupportsExperimentalVM = Script ? Script->GetVMExecutableData().SupportsExperimentalVM() : false;
+	const bool ProjectSupportsExperimentalVM = GetDefault<UNiagaraSettings>()->bExperimentalVMEnabled;
+	const bool ScriptSupportsExperimentalVM = (Script ? Script->GetVMExecutableData().SupportsExperimentalVM() : false);
 	bool SystemDisabledExperimentalVM = false;
 
 	if (InScript)
@@ -77,7 +79,7 @@ bool FNiagaraScriptExecutionContextBase::Init(UNiagaraScript* InScript, ENiagara
 		}
 	}
 
-	bUsingExperimentalVM = ScriptSupportsExperimentalVM && ((GbForceExecVMPath > 0) || (!GbForceExecVMPath && !SystemDisabledExperimentalVM));
+	bUsingExperimentalVM = ProjectSupportsExperimentalVM && ScriptSupportsExperimentalVM && ((GbForceExecVMPath > 0) || (!GbForceExecVMPath && !SystemDisabledExperimentalVM));
 #endif
 
 #	if VECTORVM_SUPPORTS_EXPERIMENTAL
