Fixed undo/redo not working when editing user parameter default values

#rnx

diff --git a/Engine/Plugins/FX/Niagara/Source/NiagaraEditor/Private/Customizations/NiagaraComponentDetails.cpp b/Engine/Plugins/FX/Niagara/Source/NiagaraEditor/Private/Customizations/NiagaraComponentDetails.cpp
--- a/Engine/Plugins/FX/Niagara/Source/NiagaraEditor/Private/Customizations/NiagaraComponentDetails.cpp
+++ b/Engine/Plugins/FX/Niagara/Source/NiagaraEditor/Private/Customizations/NiagaraComponentDetails.cpp
@@ -170,14 +170,25 @@ public:
 
 	void OnParameterPreChange()
 	{
-		TWeakObjectPtr<UNiagaraComponent> RawComponent = GetComponent();
-		if (RawComponent.IsValid())
+		if (Owner->IsA<UNiagaraComponent>())
 		{
-			RawComponent->Modify();
+			TWeakObjectPtr<UNiagaraComponent> RawComponent = GetComponent();
+			if(RawComponent.IsValid())
+			{
+				RawComponent->Modify();
 
-			for (TSharedPtr<IPropertyHandle> PropertyHandle : PropertyHandles)
+				for (TSharedPtr<IPropertyHandle> PropertyHandle : PropertyHandles)
+				{
+					PropertyHandle->NotifyPreChange();
+				}
+			}
+		}
+		else if(Owner->IsA<UNiagaraSystem>())
+		{
+			TWeakObjectPtr<UNiagaraSystem> System = GetSystem();
+			if(System.IsValid())
 			{
-				PropertyHandle->NotifyPreChange();
+				System->Modify();
 			}
 		}
 	}
