Removal of "+" button for system node "User Parameters"
Also changed the tooltip to clarify User Parameters tab and the button going away in future releases

diff --git a/Engine/Plugins/FX/Niagara/Source/NiagaraEditor/Private/ViewModels/Stack/NiagaraStackSystemSettingsGroup.cpp b/Engine/Plugins/FX/Niagara/Source/NiagaraEditor/Private/ViewModels/Stack/NiagaraStackSystemSettingsGroup.cpp
--- a/Engine/Plugins/FX/Niagara/Source/NiagaraEditor/Private/ViewModels/Stack/NiagaraStackSystemSettingsGroup.cpp
+++ b/Engine/Plugins/FX/Niagara/Source/NiagaraEditor/Private/ViewModels/Stack/NiagaraStackSystemSettingsGroup.cpp
@@ -130,14 +130,11 @@ void UNiagaraStackSystemUserParametersGroup::Initialize(
 	UObject* InOwner,
 	FNiagaraParameterStore* InParameterStore)
 {
-	AddUtilities = MakeShared<FParameterStoreGroupAddUtiliites>(*InOwner, *InParameterStore, *InRequiredEntryData.StackEditorData,
-		FParameterStoreGroupAddUtiliites::FOnItemAdded::CreateUObject(this, &UNiagaraStackSystemUserParametersGroup::ParameterAdded));
 	FText DisplayName = LOCTEXT("SystemUserParametersGroupName", "User Parameters");
-	FText Tooltip = LOCTEXT("SystemUserParametersTooltip", "Parameters for the system which are exposed externally.");
-	Super::Initialize(InRequiredEntryData, DisplayName, Tooltip, AddUtilities.Get());
+	FText Tooltip = LOCTEXT("SystemUserParametersTooltip", "Brings the User Parameters tab to the front. This button will go away with the next release.\nThe User Parameters tab allows you to set default values and hierarchies for user parameters.\nUser Parameters are exposed to systems outside Niagara (like Blueprints) and can be dynamically changed at runtime.\n");
+	Super::Initialize(InRequiredEntryData, DisplayName, Tooltip, nullptr);
 
 	Owner = InOwner;
-	UserParameterStore = InParameterStore;
 }
 
 FText UNiagaraStackSystemUserParametersGroup::GetIconText() const
@@ -145,40 +142,6 @@ FText UNiagaraStackSystemUserParametersGroup::GetIconText() const
 	return FText::FromString(FString(TEXT("\xf007")/* fa-user */));
 }
 
-void UNiagaraStackSystemUserParametersGroup::RefreshChildrenInternal(const TArray<UNiagaraStackEntry*>& CurrentChildren, TArray<UNiagaraStackEntry*>& NewChildren, TArray<FStackIssue>& NewIssues)
-{
-	if (UserParameterStore != nullptr)
-	{
-		UNiagaraStackParameterStoreItem* UserParameterStoreItem = FindCurrentChildOfTypeByPredicate<UNiagaraStackParameterStoreItem>(CurrentChildren,
-			[=](UNiagaraStackParameterStoreItem* CurrentItem) { return true; });
-
-		if (UserParameterStoreItem == nullptr)
-		{
-			UserParameterStoreItem = NewObject<UNiagaraStackParameterStoreItem>(this);
-			UserParameterStoreItem->Initialize(CreateDefaultChildRequiredData(), Owner.Get(), UserParameterStore, AddUtilities.Get());
-		}
-
-		NewChildren.Add(UserParameterStoreItem);
-	}
-
-	Super::RefreshChildrenInternal(CurrentChildren, NewChildren, NewIssues);
-}
-
-void UNiagaraStackSystemUserParametersGroup::ParameterAdded(FNiagaraVariable AddedParameter)
-{
-	if (AddedParameter.GetType().IsDataInterface())
-	{
-		UNiagaraDataInterface* DataInterfaceParameter = UserParameterStore->GetDataInterface(AddedParameter);
-		if (DataInterfaceParameter != nullptr)
-		{
-			TArray<UObject*> ChangedObjects;
-			ChangedObjects.Add(DataInterfaceParameter);
-			OnDataObjectModified().Broadcast(ChangedObjects, ENiagaraDataObjectChange::Added);
-		}
-	}
-	RefreshChildren();
-}
-
 void UNiagaraStackParameterStoreItem::Initialize(
 	FRequiredEntryData InRequiredEntryData,
 	UObject* InOwner,
diff --git a/Engine/Plugins/FX/Niagara/Source/NiagaraEditor/Public/ViewModels/Stack/NiagaraStackSystemSettingsGroup.h b/Engine/Plugins/FX/Niagara/Source/NiagaraEditor/Public/ViewModels/Stack/NiagaraStackSystemSettingsGroup.h
--- a/Engine/Plugins/FX/Niagara/Source/NiagaraEditor/Public/ViewModels/Stack/NiagaraStackSystemSettingsGroup.h
+++ b/Engine/Plugins/FX/Niagara/Source/NiagaraEditor/Public/ViewModels/Stack/NiagaraStackSystemSettingsGroup.h
@@ -40,18 +40,9 @@ public:
 	virtual FText GetIconText() const override;
 	virtual bool GetCanExpandInOverview() const override { return false; }
 	virtual bool GetShouldShowInStack() const override { return false; }
-
-protected:
-	virtual void RefreshChildrenInternal(const TArray<UNiagaraStackEntry*>& CurrentChildren, TArray<UNiagaraStackEntry*>& NewChildren, TArray<FStackIssue>& NewIssues) override;
-
-private:
-	void ParameterAdded(FNiagaraVariable AddedParameter);
-
 private:
 	TWeakObjectPtr<UObject> Owner;
-	FNiagaraParameterStore* UserParameterStore;
 	FDelegateHandle ParameterStoreChangedHandle;
-	TSharedPtr<INiagaraStackItemGroupAddUtilities> AddUtilities;
 };
 
 UCLASS()
