Invalidate script data for a simulation stage when it is disabled in the UI (prevents previous error messages from being reported)

diff --git a/Engine/Plugins/FX/Niagara/Source/Niagara/Private/NiagaraSimulationStageBase.cpp b/Engine/Plugins/FX/Niagara/Source/Niagara/Private/NiagaraSimulationStageBase.cpp
--- a/Engine/Plugins/FX/Niagara/Source/Niagara/Private/NiagaraSimulationStageBase.cpp
+++ b/Engine/Plugins/FX/Niagara/Source/Niagara/Private/NiagaraSimulationStageBase.cpp
@@ -40,6 +40,11 @@ void UNiagaraSimulationStageBase::SetEnabled(bool bInEnabled)
 
 void UNiagaraSimulationStageBase::RequestRecompile()
 {
+	if (!bEnabled)
+	{
+		Script->InvalidateCompileResults(TEXT("SimulatoinStage changed."));
+	}
+
 	FVersionedNiagaraEmitter OuterEmitter = GetOuterEmitter();
 	if (FVersionedNiagaraEmitterData* EmitterData = OuterEmitter.GetEmitterData())
 	{
