- Remove check because AddInstance can be called while we are moving system from PendingSpawn -> Spawning

#rnx

diff --git a/Engine/Plugins/FX/Niagara/Source/Niagara/Private/NiagaraSystemSimulation.cpp b/Engine/Plugins/FX/Niagara/Source/Niagara/Private/NiagaraSystemSimulation.cpp
--- a/Engine/Plugins/FX/Niagara/Source/Niagara/Private/NiagaraSystemSimulation.cpp
+++ b/Engine/Plugins/FX/Niagara/Source/Niagara/Private/NiagaraSystemSimulation.cpp
@@ -2176,7 +2176,8 @@ void FNiagaraSystemSimulation::AddInstance(FNiagaraSystemInstance* Instance)
 	}
 
 	check(GetSystemInstances(ENiagaraSystemInstanceState::Running).Num() == MainDataSet.GetCurrentDataChecked().GetNumInstances());
-	check(GetSystemInstances(ENiagaraSystemInstanceState::Spawning).Num() == SpawningDataSet.GetCurrentDataChecked().GetNumInstances());
+	//Note: We can't check this here because AddInstance can be called while we are moving system from PendingSpawn -> Spawning
+	//check(GetSystemInstances(ENiagaraSystemInstanceState::Spawning).Num() == SpawningDataSet.GetCurrentDataChecked().GetNumInstances());
 	check(GetSystemInstances(ENiagaraSystemInstanceState::Paused).Num() == PausedDataSet.GetCurrentDataChecked().GetNumInstances());
 }
 
