Optimized generating compilation hash for NiagaraScriptVariable saved about 5% of the cook time for a heavy niagara level

diff --git a/Engine/Plugins/FX/Niagara/Source/Niagara/Private/NiagaraModule.cpp b/Engine/Plugins/FX/Niagara/Source/Niagara/Private/NiagaraModule.cpp
--- a/Engine/Plugins/FX/Niagara/Source/Niagara/Private/NiagaraModule.cpp
+++ b/Engine/Plugins/FX/Niagara/Source/Niagara/Private/NiagaraModule.cpp
@@ -899,23 +899,18 @@ bool FNiagaraTypeDefinition::AppendCompileHash(FNiagaraCompileHashVisitor* InVis
 		return false;
 	}
 
-	UStruct* TDStruct = GetStruct();
-	UClass* TDClass = GetClass();
-	UEnum* TDEnum = GetEnum();
-
 	TStringBuilder<128> PathName;
-	if (TDEnum)
+	if (UEnum* TDEnum = GetEnum())
 	{
 		// Do we need to enumerate all the enum values and rebuild if that changes or are we ok with just knowing that there are the same  count of enum entries?
 		// For now, am just going to be ok with the number of entries. The actual string values don't matter so much.
-		FString CppType = TDEnum->CppType;
 		PathName.Reset();
 		TDEnum->GetPathName(nullptr, PathName);
 		InVisitor->UpdateString(TEXT("\tEnumPath"), PathName);
-		InVisitor->UpdateString(TEXT("\tEnumCppType"), CppType);
+		InVisitor->UpdateString(TEXT("\tEnumCppType"), TDEnum->CppType);
 		InVisitor->UpdatePOD(TEXT("\t\tNumEnums"),TDEnum->NumEnums());
 	}
-	else if (TDClass)
+	else if (UClass* TDClass = GetClass())
 	{
 		// For data interfaces, get the default object and the compile version so that we can properly update when code changes.
 		check(IsInGameThread());
@@ -937,7 +932,7 @@ bool FNiagaraTypeDefinition::AppendCompileHash(FNiagaraCompileHashVisitor* InVis
 			}
 		}
 	}
-	else if (TDStruct)
+	else if (UStruct* TDStruct = GetStruct())
 	{
 		PathName.Reset();
 		TDStruct->GetPathName(nullptr, PathName);
@@ -1553,6 +1548,12 @@ void FNiagaraVariableBase::PostSerialize(const FArchive& Ar)
 }
 #endif
 
+bool FNiagaraVariableBase::AppendCompileHash(FNiagaraCompileHashVisitor* InVisitor) const
+{
+	return InVisitor->UpdateName(TEXT("NiagaraVariable.Name"), Name)
+		&& TypeDefHandle->AppendCompileHash(InVisitor);
+}
+
 bool FNiagaraVariable::Serialize(FArchive& Ar)
 {
 	FNiagaraVariableBase::Serialize(Ar);
diff --git a/Engine/Plugins/FX/Niagara/Source/Niagara/Private/NiagaraSimulationStageBase.cpp b/Engine/Plugins/FX/Niagara/Source/Niagara/Private/NiagaraSimulationStageBase.cpp
--- a/Engine/Plugins/FX/Niagara/Source/Niagara/Private/NiagaraSimulationStageBase.cpp
+++ b/Engine/Plugins/FX/Niagara/Source/Niagara/Private/NiagaraSimulationStageBase.cpp
@@ -21,8 +21,10 @@ namespace NiagaraSimulationStageLocal
 bool UNiagaraSimulationStageBase::AppendCompileHash(FNiagaraCompileHashVisitor* InVisitor) const
 {
 #if WITH_EDITORONLY_DATA
-	const int32 Index = InVisitor->Values.AddDefaulted();
-	InVisitor->Values[Index].Object = FString::Printf(TEXT("Class: \"%s\"  Name: \"%s\""), *GetClass()->GetName(), *GetName());
+	if (FNiagaraCompileHashVisitorDebugInfo* DebugInfo = InVisitor->AddDebugInfo())
+	{
+		DebugInfo->Object = FString::Printf(TEXT("Class: \"%s\"  Name: \"%s\""), *GetClass()->GetName(), *GetName());
+	}
 #endif
 	InVisitor->UpdatePOD(TEXT("Enabled"), bEnabled ? 1 : 0);
 	return true;
diff --git a/Engine/Plugins/FX/Niagara/Source/Niagara/Private/NiagaraVariant.cpp b/Engine/Plugins/FX/Niagara/Source/Niagara/Private/NiagaraVariant.cpp
--- a/Engine/Plugins/FX/Niagara/Source/Niagara/Private/NiagaraVariant.cpp
+++ b/Engine/Plugins/FX/Niagara/Source/Niagara/Private/NiagaraVariant.cpp
@@ -134,3 +134,10 @@ bool FNiagaraVariant::operator!=(const FNiagaraVariant& Other) const
 	return !(operator==(Other));
 }
 
+bool FNiagaraVariant::AppendCompileHash(FNiagaraCompileHashVisitor* InVisitor) const
+{
+	return InVisitor->UpdateName(TEXT("NiagaraVariant.Object"), Object ? Object->GetFName() : NAME_None)
+		&& InVisitor->UpdateName(TEXT("NiagaraVariant.DataInterface"), DataInterface ? DataInterface->GetFName() : NAME_None)
+		&& InVisitor->UpdateArray(TEXT("NiagaraVariant.Bytes"), Bytes.GetData(), Bytes.Num())
+		&& InVisitor->UpdatePOD<uint8>(TEXT("NiagaraVariant.CurrentMode"), static_cast<uint8>(CurrentMode));
+}
\ No newline at end of file
diff --git a/Engine/Plugins/FX/Niagara/Source/Niagara/Public/NiagaraTypes.h b/Engine/Plugins/FX/Niagara/Source/Niagara/Public/NiagaraTypes.h
--- a/Engine/Plugins/FX/Niagara/Source/Niagara/Public/NiagaraTypes.h
+++ b/Engine/Plugins/FX/Niagara/Source/Niagara/Public/NiagaraTypes.h
@@ -590,6 +590,14 @@ public:
 	static int LogCompileIdGeneration;
 
 #if WITH_EDITORONLY_DATA
+	FNiagaraCompileHashVisitorDebugInfo* AddDebugInfo()
+	{
+		if (LogCompileIdGeneration != 0)
+		{
+			return &Values.AddDefaulted_GetRef();
+		}
+		return nullptr;
+	}
 
 	// Debug data about the compilation hash, including key value pairs to detect differences.
 	TArray<FNiagaraCompileHashVisitorDebugInfo> Values;
@@ -655,6 +663,15 @@ public:
 		return true;
 	}
 
+	/**
+	 * Adds the string representation of an FName to the hash
+	 */
+	bool UpdateName(const TCHAR* InDebugName, FName InName)
+	{
+		FNameBuilder NameString(InName);
+		return UpdateString(InDebugName, NameString.ToView());
+	}
+
 	/**
 	Adds a single value of typed POD (plain old data) to the hash.
 	*/
@@ -1819,6 +1836,8 @@ struct FNiagaraVariableBase
 	void PostSerialize(const FArchive& Ar);
 #endif
 
+	bool NIAGARA_API AppendCompileHash(FNiagaraCompileHashVisitor* InVisitor) const;
+
 protected:
 	UPROPERTY(EditAnywhere, Category = "Variable")
 	FName Name;
diff --git a/Engine/Plugins/FX/Niagara/Source/Niagara/Public/NiagaraVariant.h b/Engine/Plugins/FX/Niagara/Source/Niagara/Public/NiagaraVariant.h
--- a/Engine/Plugins/FX/Niagara/Source/Niagara/Public/NiagaraVariant.h
+++ b/Engine/Plugins/FX/Niagara/Source/Niagara/Public/NiagaraVariant.h
@@ -6,6 +6,7 @@
 
 #include "NiagaraVariant.generated.h"
 
+struct FNiagaraCompileHashVisitor;
 class UNiagaraDataInterface;
 
 UENUM()
@@ -46,6 +47,8 @@ struct NIAGARA_API FNiagaraVariant
 	bool operator==(const FNiagaraVariant& Other) const;
 	bool operator!=(const FNiagaraVariant& Other) const;
 
+	bool AppendCompileHash(FNiagaraCompileHashVisitor* InVisitor) const;
+
 private:
 	UPROPERTY(EditAnywhere, Category=Variant, Instanced)
 	TObjectPtr<UObject> Object;
diff --git a/Engine/Plugins/FX/Niagara/Source/NiagaraEditor/Private/NiagaraEditorUtilities.cpp b/Engine/Plugins/FX/Niagara/Source/NiagaraEditor/Private/NiagaraEditorUtilities.cpp
--- a/Engine/Plugins/FX/Niagara/Source/NiagaraEditor/Private/NiagaraEditorUtilities.cpp
+++ b/Engine/Plugins/FX/Niagara/Source/NiagaraEditor/Private/NiagaraEditorUtilities.cpp
@@ -366,75 +366,75 @@
 }
 
 
-bool FNiagaraEditorUtilities::PODPropertyAppendCompileHash(const void* Container, FProperty* Property, const FString& PropertyName, FNiagaraCompileHashVisitor* InVisitor) 
+bool FNiagaraEditorUtilities::PODPropertyAppendCompileHash(const void* Container, FProperty* Property, FStringView PropertyName, FNiagaraCompileHashVisitor* InVisitor) 
 {
 	if (Property->IsA(FFloatProperty::StaticClass()))
 	{
 		FFloatProperty* CastProp = CastFieldChecked<FFloatProperty>(Property);
 		float Value = CastProp->GetPropertyValue_InContainer(Container, 0);
-		InVisitor->UpdatePOD(*PropertyName, Value);
+		InVisitor->UpdatePOD(PropertyName.GetData(), Value);
 		return true;
 	}
 	else if (Property->IsA(FIntProperty::StaticClass()))
 	{
 		FIntProperty* CastProp = CastFieldChecked<FIntProperty>(Property);
 		int32 Value = CastProp->GetPropertyValue_InContainer(Container, 0);
-		InVisitor->UpdatePOD(*PropertyName, Value);
+		InVisitor->UpdatePOD(PropertyName.GetData(), Value);
 		return true;
 	}
 	else if (Property->IsA(FInt16Property::StaticClass()))
 	{
 		FInt16Property* CastProp = CastFieldChecked<FInt16Property>(Property);
 		int16 Value = CastProp->GetPropertyValue_InContainer(Container, 0);
-		InVisitor->UpdatePOD(*PropertyName, Value);
+		InVisitor->UpdatePOD(PropertyName.GetData(), Value);
 		return true;
 	}
 	else if (Property->IsA(FUInt32Property::StaticClass()))
 	{
 		FUInt32Property* CastProp = CastFieldChecked<FUInt32Property>(Property);
 		uint32 Value = CastProp->GetPropertyValue_InContainer(Container, 0);
-		InVisitor->UpdatePOD(*PropertyName, Value);
+		InVisitor->UpdatePOD(PropertyName.GetData(), Value);
 		return true;
 	}
 	else if (Property->IsA(FUInt16Property::StaticClass()))
 	{
 		FUInt16Property* CastProp = CastFieldChecked<FUInt16Property>(Property);
 		uint16 Value = CastProp->GetPropertyValue_InContainer(Container, 0);
-		InVisitor->UpdatePOD(*PropertyName, Value);
+		InVisitor->UpdatePOD(PropertyName.GetData(), Value);
 		return true;
 	}
 	else if (Property->IsA(FByteProperty::StaticClass()))
 	{
 		FByteProperty* CastProp = CastFieldChecked<FByteProperty>(Property);
 		uint8 Value = CastProp->GetPropertyValue_InContainer(Container, 0);
-		InVisitor->UpdatePOD(*PropertyName, Value);
+		InVisitor->UpdatePOD(PropertyName.GetData(), Value);
 		return true;
 	}
 	else if (Property->IsA(FBoolProperty::StaticClass()))
 	{
 		FBoolProperty* CastProp = CastFieldChecked<FBoolProperty>(Property);
 		bool Value = CastProp->GetPropertyValue_InContainer(Container, 0);
-		InVisitor->UpdatePOD(*PropertyName, Value);
+		InVisitor->UpdatePOD(PropertyName.GetData(), Value);
 		return true;
 	}
 	else if (Property->IsA(FNameProperty::StaticClass()))
 	{
 		FNameProperty* CastProp = CastFieldChecked<FNameProperty>(Property);
 		FName Value = CastProp->GetPropertyValue_InContainer(Container);
-		InVisitor->UpdateString(*PropertyName, Value.ToString());
+		InVisitor->UpdateName(PropertyName.GetData(), Value);
 		return true;
 	}
 	else if (Property->IsA(FStrProperty::StaticClass()))
 	{
 		FStrProperty* CastProp = CastFieldChecked<FStrProperty>(Property);
 		FString Value = CastProp->GetPropertyValue_InContainer(Container);
-		InVisitor->UpdateString(*PropertyName, Value);
+		InVisitor->UpdateString(PropertyName.GetData(), Value);
 		return true;
 	}
 	return false;
 }
 
-bool FNiagaraEditorUtilities::NestedPropertiesAppendCompileHash(const void* Container, const UStruct* Struct, EFieldIteratorFlags::SuperClassFlags IteratorFlags, const FString& BaseName, FNiagaraCompileHashVisitor* InVisitor) 
+bool FNiagaraEditorUtilities::NestedPropertiesAppendCompileHash(const void* Container, const UStruct* Struct, EFieldIteratorFlags::SuperClassFlags IteratorFlags, FStringView BaseName, FNiagaraCompileHashVisitor* InVisitor) 
 {
 	// We special case FNiagaraTypeDefinitions here because they need to write out a lot more than just their standalone uproperties.
 	if (Struct == FNiagaraTypeDefinition::StaticStruct())
@@ -454,7 +454,7 @@
 		}
 	}
 
-	TFieldIterator<FProperty> PropertyCountIt(Struct, IteratorFlags);
+	TFieldIterator<FProperty> PropertyCountIt(Struct, IteratorFlags, EFieldIteratorFlags::ExcludeDeprecated);
 	int32 NumProperties = 0;
 	for (; PropertyCountIt; ++PropertyCountIt)
 	{
@@ -462,7 +462,9 @@
 	}
 
 	TStringBuilder<128> PathName;
-	for (TFieldIterator<FProperty> PropertyIt(Struct, IteratorFlags); PropertyIt; ++PropertyIt)
+	TStringBuilder<128> PropertyName;
+
+	for (TFieldIterator<FProperty> PropertyIt(Struct, IteratorFlags, EFieldIteratorFlags::ExcludeDeprecated); PropertyIt; ++PropertyIt)
 	{
 		FProperty* Property = *PropertyIt;
 
@@ -472,17 +474,57 @@
 			continue;
 		}
 
-		FString PropertyName = (NumProperties == 1) ? (BaseName) : (BaseName + "." + Property->GetName());
+		PropertyName = BaseName;
+		if (NumProperties > 1)
+		{
+			PropertyName << TCHAR('.');
+			Property->GetFName().AppendString(PropertyName);
+		}
 
 		if (PODPropertyAppendCompileHash(Container, Property, PropertyName, InVisitor))
 		{
 			continue;
 		}
+		else if (Property->IsA(FStructProperty::StaticClass()))
+		{
+			FStructProperty* StructProp = CastFieldChecked<FStructProperty>(Property);
+			const void* StructContainer = Property->ContainerPtrToValuePtr<uint8>(Container);
+			NestedPropertiesAppendCompileHash(StructContainer, StructProp->Struct, EFieldIteratorFlags::IncludeSuper, PropertyName, InVisitor);
+			continue;
+		}
+		else if (Property->IsA(FEnumProperty::StaticClass()))
+		{
+			FEnumProperty* CastProp = CastFieldChecked<FEnumProperty>(Property);
+			const void* EnumContainer = Property->ContainerPtrToValuePtr<uint8>(Container);
+			if (PODPropertyAppendCompileHash(EnumContainer, CastProp->GetUnderlyingProperty(), PropertyName, InVisitor))
+			{
+				continue;
+			}
+			check(false);
+			return false;
+		}
+		else if (Property->IsA(FObjectProperty::StaticClass()))
+		{
+			FObjectProperty* CastProp = CastFieldChecked<FObjectProperty>(Property);
+			UObject* Obj = CastProp->GetObjectPropertyValue_InContainer(Container);
+			if (Obj != nullptr)
+			{
+				// We just do name here as sometimes things will be in a transient package or something tricky.
+				// Because we do nested id's for each called graph, it should work out in the end to have a different
+				// value in the compile array if the scripts are the same name but different locations.
+				InVisitor->UpdateString(PropertyName.GetData(), Obj->GetName());
+			}
+			else
+			{
+				InVisitor->UpdateString(PropertyName.GetData(), TEXT("nullptr"));
+			}
+			continue;
+		}
 		else if (Property->IsA(FMapProperty::StaticClass()))
 		{
 			FMapProperty* CastProp = CastFieldChecked<FMapProperty>(Property);
 			FScriptMapHelper MapHelper(CastProp, CastProp->ContainerPtrToValuePtr<void>(Container));
-			InVisitor->UpdatePOD(*PropertyName, MapHelper.Num());
+			InVisitor->UpdatePOD(PropertyName.GetData(), MapHelper.Num());
 			if (MapHelper.GetKeyProperty())
 			{
 				PathName.Reset();
@@ -567,7 +609,7 @@
 			FArrayProperty* CastProp = CastFieldChecked<FArrayProperty>(Property);
 
 			FScriptArrayHelper ArrayHelper(CastProp, CastProp->ContainerPtrToValuePtr<void>(Container));
-			InVisitor->UpdatePOD(*PropertyName, ArrayHelper.Num());
+			InVisitor->UpdatePOD(PropertyName.GetData(), ArrayHelper.Num());
 			PathName.Reset();
 			CastProp->Inner->GetPathName(nullptr, PathName);
 			InVisitor->UpdateString(TEXT("InnerPathname"), PathName);
@@ -622,41 +664,6 @@
 			UE_LOG(LogNiagaraEditor, Warning, TEXT("Skipping %s because it is a UText property, please add \"meta = (SkipForCompileHash=\"true\")\" to avoid this warning in the future or handle it yourself in NestedPropertiesAppendCompileHash!"), *Property->GetName());
 			return true;
 		}
-		else if (Property->IsA(FEnumProperty::StaticClass()))
-		{
-			FEnumProperty* CastProp = CastFieldChecked<FEnumProperty>(Property);
-			const void* EnumContainer = Property->ContainerPtrToValuePtr<uint8>(Container);
-			if (PODPropertyAppendCompileHash(EnumContainer, CastProp->GetUnderlyingProperty(), PropertyName, InVisitor))
-			{
-				continue;
-			}
-			check(false);
-			return false;
-		}
-		else if (Property->IsA(FObjectProperty::StaticClass()))
-		{
-			FObjectProperty* CastProp = CastFieldChecked<FObjectProperty>(Property);
-			UObject* Obj = CastProp->GetObjectPropertyValue_InContainer(Container);
-			if (Obj != nullptr)
-			{
-				// We just do name here as sometimes things will be in a transient package or something tricky.
-				// Because we do nested id's for each called graph, it should work out in the end to have a different
-				// value in the compile array if the scripts are the same name but different locations.
-				InVisitor->UpdateString(*PropertyName, Obj->GetName());
-			}
-			else
-			{
-				InVisitor->UpdateString(*PropertyName, TEXT("nullptr"));
-			}
-			continue;
-		}
-		else if (Property->IsA(FStructProperty::StaticClass()))
-		{
-			FStructProperty* StructProp = CastFieldChecked<FStructProperty>(Property);
-			const void* StructContainer = Property->ContainerPtrToValuePtr<uint8>(Container);
-			NestedPropertiesAppendCompileHash(StructContainer, StructProp->Struct, EFieldIteratorFlags::IncludeSuper, PropertyName, InVisitor);
-			continue;
-		}
 		else
 		{
 			check(false);
diff --git a/Engine/Plugins/FX/Niagara/Source/NiagaraEditor/Private/NiagaraGraph.cpp b/Engine/Plugins/FX/Niagara/Source/NiagaraEditor/Private/NiagaraGraph.cpp
--- a/Engine/Plugins/FX/Niagara/Source/NiagaraEditor/Private/NiagaraGraph.cpp
+++ b/Engine/Plugins/FX/Niagara/Source/NiagaraEditor/Private/NiagaraGraph.cpp
@@ -2881,8 +2881,10 @@ bool UNiagaraGraph::AppendCompileHash(FNiagaraCompileHashVisitor* InVisitor, con
 	check(!bIsForCompilationOnly);
 
 #if WITH_EDITORONLY_DATA
-	int32 Index = InVisitor->Values.AddDefaulted();
-	InVisitor->Values[Index].Object = FString::Printf(TEXT("Class: \"%s\"  Name: \"%s\""), *this->GetClass()->GetName(),  *this->GetName());
+	if (FNiagaraCompileHashVisitorDebugInfo* DebugInfo = InVisitor->AddDebugInfo())
+	{
+		DebugInfo->Object = FString::Printf(TEXT("Class: \"%s\"  Name: \"%s\""), *this->GetClass()->GetName(), *this->GetName());
+	}
 #endif
 	InVisitor->UpdateString(TEXT("ForceRebuildId"), ForceRebuildId.ToString());
 
@@ -2957,8 +2959,10 @@ bool UNiagaraGraph::AppendCompileHash(FNiagaraCompileHashVisitor* InVisitor, con
 			}
 
 		#if WITH_EDITORONLY_DATA
-			Index = InVisitor->Values.AddDefaulted();
-			InVisitor->Values[Index].Object = FString::Printf(TEXT("Class: \"%s\"  Name: \"%s\""), *Var->GetClass()->GetName(), *Var->Variable.GetName().ToString());
+			if (FNiagaraCompileHashVisitorDebugInfo* DebugInfo = InVisitor->AddDebugInfo())
+			{
+				DebugInfo->Object = FString::Printf(TEXT("Class: \"%s\"  Name: \"%s\""), *Var->GetClass()->GetName(), *Var->Variable.GetName().ToString());
+			}
 		#endif
 			verify(Var->AppendCompileHash(InVisitor));
 			
@@ -2975,8 +2979,10 @@ bool UNiagaraGraph::AppendCompileHash(FNiagaraCompileHashVisitor* InVisitor, con
 	for (const UNiagaraNode* Node : InTraversal)
 	{
 #if WITH_EDITORONLY_DATA
-		Index = InVisitor->Values.AddDefaulted();
-		InVisitor->Values[Index].Object = FString::Printf(TEXT("Class: \"%s\" Title: \"%s\" Name: \"%s\" Guid: %s"), *Node->GetClass()->GetName(), *Node->GetNodeTitle(ENodeTitleType::EditableTitle).ToString(), *Node->GetName(), *LexToString(Node->NodeGuid));
+		if (FNiagaraCompileHashVisitorDebugInfo* DebugInfo = InVisitor->AddDebugInfo())
+		{
+			DebugInfo->Object = FString::Printf(TEXT("Class: \"%s\" Title: \"%s\" Name: \"%s\" Guid: %s"), *Node->GetClass()->GetName(), *Node->GetNodeTitle(ENodeTitleType::EditableTitle).ToString(), *Node->GetName(), *LexToString(Node->NodeGuid));
+		}
 #endif
 		verify(Node->AppendCompileHash(InVisitor));
 	}
diff --git a/Engine/Plugins/FX/Niagara/Source/NiagaraEditor/Private/NiagaraScriptVariable.cpp b/Engine/Plugins/FX/Niagara/Source/NiagaraEditor/Private/NiagaraScriptVariable.cpp
--- a/Engine/Plugins/FX/Niagara/Source/NiagaraEditor/Private/NiagaraScriptVariable.cpp
+++ b/Engine/Plugins/FX/Niagara/Source/NiagaraEditor/Private/NiagaraScriptVariable.cpp
@@ -110,11 +110,19 @@ void UNiagaraScriptVariable::PostLoad()
 
 bool UNiagaraScriptVariable::AppendCompileHash(FNiagaraCompileHashVisitor* InVisitor) const
 {
-	if (!FNiagaraEditorUtilities::NestedPropertiesAppendCompileHash(static_cast<const void*>(this), StaticClass(), EFieldIteratorFlags::ExcludeSuper, StaticClass()->GetName(), InVisitor))
+	bool bSuccess = InVisitor->UpdatePOD<uint8>(TEXT("NiagaraScriptVariable.DefaultMode"), static_cast<uint8>(DefaultMode))
+		&& InVisitor->UpdateName(TEXT("NiagaraScriptVariable.DefaultBinding.Name"), DefaultBinding.Name)
+		&& InVisitor->UpdatePOD(TEXT("NiagaraScriptVariable.StaticSwitchDefaultValue"), StaticSwitchDefaultValue)
+		&& InVisitor->UpdatePOD(TEXT("NiagaraScriptVariable.bIsStaticSwitch"), bIsStaticSwitch)
+		&& Variable.AppendCompileHash(InVisitor)
+		&& DefaultValueVariant.AppendCompileHash(InVisitor);
+
+	for (const FName& AlternateAlias : Metadata.AlternateAliases)
 	{
-		return false;
+		bSuccess = bSuccess && InVisitor->UpdateName(TEXT("NiagaraScriptVariable.Metadata.AlternateAliases"), AlternateAlias);
 	}
-	return true;
+
+	return bSuccess;
 }
 
 void UNiagaraScriptVariable::SetIsSubscribedToParameterDefinitions(bool bInSubscribedToParameterDefinitions)
diff --git a/Engine/Plugins/FX/Niagara/Source/NiagaraEditor/Public/NiagaraEditorUtilities.h b/Engine/Plugins/FX/Niagara/Source/NiagaraEditor/Public/NiagaraEditorUtilities.h
--- a/Engine/Plugins/FX/Niagara/Source/NiagaraEditor/Public/NiagaraEditorUtilities.h
+++ b/Engine/Plugins/FX/Niagara/Source/NiagaraEditor/Public/NiagaraEditorUtilities.h
@@ -180,8 +180,8 @@ namespace FNiagaraEditorUtilities
 
 	void PreprocessFunctionGraph(const UEdGraphSchema_Niagara* Schema, UNiagaraGraph* Graph, TArrayView<UEdGraphPin* const> CallInputs, TArrayView<UEdGraphPin* const> CallOutputs, ENiagaraScriptUsage ScriptUsage, const FCompileConstantResolver& ConstantResolver);
 
-	bool PODPropertyAppendCompileHash(const void* Container, FProperty* Property, const FString& PropertyName, struct FNiagaraCompileHashVisitor* InVisitor);
-	bool NestedPropertiesAppendCompileHash(const void* Container, const UStruct* Struct, EFieldIteratorFlags::SuperClassFlags IteratorFlags, const FString& BaseName, struct FNiagaraCompileHashVisitor* InVisitor);
+	bool PODPropertyAppendCompileHash(const void* Container, FProperty* Property, FStringView PropertyName, struct FNiagaraCompileHashVisitor* InVisitor);
+	bool NestedPropertiesAppendCompileHash(const void* Container, const UStruct* Struct, EFieldIteratorFlags::SuperClassFlags IteratorFlags, FStringView BaseName, struct FNiagaraCompileHashVisitor* InVisitor);
 
 	/** Options for the GetScriptsByFilter function. 
 	** @Param ScriptUsageToInclude Only return Scripts that have this usage
