- Add support for InverseExposureBlend to simple lights
- Niagara default is controlled inside NiagaraSettings and can be overridden per light renderer
- Cascade default setting is controlled by console variable fx.ParticleDefaultLightInverseExposureBlend and can be overridden per light module

diff --git a/Engine/Plugins/FX/Niagara/Source/Niagara/Private/NiagaraRendererLights.cpp b/Engine/Plugins/FX/Niagara/Source/Niagara/Private/NiagaraRendererLights.cpp
--- a/Engine/Plugins/FX/Niagara/Source/Niagara/Private/NiagaraRendererLights.cpp
+++ b/Engine/Plugins/FX/Niagara/Source/Niagara/Private/NiagaraRendererLights.cpp
@@ -4,6 +4,7 @@
 #include "ParticleResources.h"
 #include "NiagaraDataSet.h"
 #include "NiagaraDataSetAccessor.h"
+#include "NiagaraSettings.h"
 #include "NiagaraStats.h"
 #include "NiagaraVertexFactory.h"
 #include "NiagaraComponent.h"
@@ -110,6 +111,8 @@ FNiagaraDynamicDataBase* FNiagaraRendererLights::GenerateDynamicData(const FNiag
 	const int32 DefaultVisibilityTag(0);
 	const float DefaultExponent = Properties->DefaultExponent;
 
+	const float InverseExposureBlend = Properties->bOverrideInverseExposureBlend ? Properties->InverseExposureBlend : GetDefault<UNiagaraSettings>()->DefaultLightInverseExposureBlend;
+
 	for (uint32 ParticleIndex = 0; ParticleIndex < DataToRender->GetNumInstances(); ParticleIndex++)
 	{
 		bool bShouldRenderParticleLight = EnabledReader.GetSafe(ParticleIndex, DefaultEnabled).GetValue();
@@ -128,6 +131,7 @@ FNiagaraDynamicDataBase* FNiagaraRendererLights::GenerateDynamicData(const FNiag
 			LightData.LightEntry.Radius = LightRadius;
 			LightData.LightEntry.Color = FVector3f(Color) * Brightness + Properties->ColorAdd;
 			LightData.LightEntry.Exponent = Properties->bUseInverseSquaredFalloff ? 0 : ExponentReader.GetSafe(ParticleIndex, DefaultExponent);
+			LightData.LightEntry.InverseExposureBlend = InverseExposureBlend;
 			LightData.LightEntry.bAffectTranslucency = Properties->bAffectsTranslucency;
 			LightData.LightEntry.VolumetricScatteringIntensity = ScatteringReader.GetSafe(ParticleIndex, DefaultScattering);
 			LightData.PerViewEntry.Position = LwcConverter.ConvertSimulationPositionToWorld(PositionReader.GetSafe(ParticleIndex, DefaultPos));
@@ -160,7 +164,3 @@ void FNiagaraRendererLights::GatherSimpleLights(FSimpleLightArray& OutParticleLi
 		}
 	}
 }
-
-
-
-
diff --git a/Engine/Plugins/FX/Niagara/Source/Niagara/Public/NiagaraLightRendererProperties.h b/Engine/Plugins/FX/Niagara/Source/Niagara/Public/NiagaraLightRendererProperties.h
--- a/Engine/Plugins/FX/Niagara/Source/Niagara/Public/NiagaraLightRendererProperties.h
+++ b/Engine/Plugins/FX/Niagara/Source/Niagara/Public/NiagaraLightRendererProperties.h
@@ -56,6 +56,10 @@ public:
 	UPROPERTY(EditAnywhere, Category = "Light Rendering")
 	uint32 bAlphaScalesBrightness : 1;
 
+	/** When enabled we will override the project default setting with our local setting. */
+	UPROPERTY(EditAnywhere, Category = "Light Rendering", meta = (PinHiddenByDefault, InlineEditConditionToggle))
+	uint32 bOverrideInverseExposureBlend : 1;
+
 	/** A factor used to scale each particle light radius */
 	UPROPERTY(EditAnywhere, Category = "Light Rendering", meta = (ClampMin = "0"))
 	float RadiusScale;
@@ -68,6 +72,14 @@ public:
 	UPROPERTY(EditAnywhere, Category = "Light Rendering")
 	FVector3f ColorAdd;
 
+	/**
+	* Blend Factor used to blend between Intensity and Intensity/Exposure.
+	* This is useful for gameplay lights that should have constant brighness on screen independent of current exposure.
+	* This feature can cause issues with exposure particularly when used on the primary light on a scene, as such it's usage should be limited.
+	*/
+	UPROPERTY(EditAnywhere, BlueprintReadOnly, Category = "Light Rendering", AdvancedDisplay, meta = (UIMin = "0.0", UIMax = "1.0", EditCondition="bOverrideInverseExposureBlend"))
+	float InverseExposureBlend = 0.0f;
+
 	/** If a render visibility tag is present, particles whose tag matches this value will be visible in this renderer. */
 	UPROPERTY(EditAnywhere, Category = "Light Rendering")
 	int32 RendererVisibility;
diff --git a/Engine/Plugins/FX/Niagara/Source/Niagara/Public/NiagaraSettings.h b/Engine/Plugins/FX/Niagara/Source/Niagara/Public/NiagaraSettings.h
--- a/Engine/Plugins/FX/Niagara/Source/Niagara/Public/NiagaraSettings.h
+++ b/Engine/Plugins/FX/Niagara/Source/Niagara/Public/NiagaraSettings.h
@@ -162,6 +162,10 @@ class NIAGARA_API UNiagaraSettings : public UDeveloperSettings
 	UPROPERTY(config, EditAnywhere, Category = Renderer)
 	ENiagaraDefaultSortPrecision DefaultSortPrecision = ENiagaraDefaultSortPrecision::Low;
 
+	/** The default InverseExposureBlend used for the light renderer. */
+	UPROPERTY(config, EditAnywhere, Category = LightRenderer)
+	float DefaultLightInverseExposureBlend = 0.0f;
+
 	UPROPERTY(config, EditAnywhere, Category=SkeletalMeshDI, meta = ( DisplayName = "Gpu Max Bone Influences", ToolTip = "Controls the maximum number of influences we allow the Skeletal Mesh Data Interface to use on the GPU.  Changing this setting requires restarting the editor.", ConfigRestartRequired = true))
 	TEnumAsByte<ENDISkelMesh_GpuMaxInfluences::Type> NDISkelMesh_GpuMaxInfluences;
 
diff --git a/Engine/Source/Runtime/Engine/Classes/Components/PointLightComponent.h b/Engine/Source/Runtime/Engine/Classes/Components/PointLightComponent.h
--- a/Engine/Source/Runtime/Engine/Classes/Components/PointLightComponent.h
+++ b/Engine/Source/Runtime/Engine/Classes/Components/PointLightComponent.h
@@ -63,6 +63,9 @@ class ENGINE_API UPointLightComponent : public ULocalLightComponent
 	UFUNCTION(BlueprintCallable, Category="Rendering|Lighting")
 	void SetLightFalloffExponent(float NewLightFalloffExponent);
 
+	UFUNCTION(BlueprintCallable, Category = "Rendering|Lighting")
+	void SetInverseExposureBlend(float NewInverseExposureBlend);
+
 	UFUNCTION(BlueprintCallable, Category="Rendering|Lighting")
 	void SetSourceRadius(float bNewValue);
 
diff --git a/Engine/Source/Runtime/Engine/Classes/Particles/Light/ParticleModuleLight.h b/Engine/Source/Runtime/Engine/Classes/Particles/Light/ParticleModuleLight.h
--- a/Engine/Source/Runtime/Engine/Classes/Particles/Light/ParticleModuleLight.h
+++ b/Engine/Source/Runtime/Engine/Classes/Particles/Light/ParticleModuleLight.h
@@ -31,6 +31,10 @@ class UParticleModuleLight : public UParticleModuleLightBase
 	UPROPERTY(EditAnywhere, Category=Light)
 	bool bAffectsTranslucency;
 
+	/** When enabled we will override the project default setting with our local setting. */
+	UPROPERTY(EditAnywhere, Category = Light, meta = (PinHiddenByDefault, InlineEditConditionToggle))
+	uint32 bOverrideInverseExposureBlend : 1;
+
 	/** 
 	 * Will draw wireframe spheres to preview the light radius if enabled.
 	 * Note: this is intended for previewing and the value will not be saved, it will always revert to disabled.
@@ -58,6 +62,14 @@ class UParticleModuleLight : public UParticleModuleLightBase
 	UPROPERTY(EditAnywhere, Category=Light)
 	struct FRawDistributionFloat LightExponent;
 
+	/**
+	* Blend Factor used to blend between Intensity and Intensity/Exposure.
+	* This is useful for gameplay lights that should have constant brighness on screen independent of current exposure.
+	* This feature can cause issues with exposure particularly when used on the primary light on a scene, as such it's usage should be limited.
+	*/
+	UPROPERTY(EditAnywhere, BlueprintReadOnly, Category = Light, AdvancedDisplay, meta = (UIMin = "0.0", UIMax = "1.0", EditCondition = "bOverrideInverseExposureBlend"))
+	float InverseExposureBlend = 0.0f;
+
 	/**
 	* Channels that this light should affect.
 	* Only affect high quality lights
diff --git a/Engine/Source/Runtime/Engine/Private/Components/LightComponent.cpp b/Engine/Source/Runtime/Engine/Private/Components/LightComponent.cpp
--- a/Engine/Source/Runtime/Engine/Private/Components/LightComponent.cpp
+++ b/Engine/Source/Runtime/Engine/Private/Components/LightComponent.cpp
@@ -334,7 +334,7 @@ void FLightRenderParameters::MakeShaderParameters(const FViewMatrices& ViewMatri
 }
 
 // match logic in EyeAdaptationInverseLookup(...)
-float FLightRenderParameters::GetLightExposureScale(float Exposure) const
+float FLightRenderParameters::GetLightExposureScale(float Exposure, float InverseExposureBlend)
 {
 	if (Exposure <= 0.0f)
 	{
@@ -357,6 +357,11 @@ float FLightRenderParameters::GetLightExposureScale(float Exposure) const
 	return Scale;
 }
 
+float FLightRenderParameters::GetLightExposureScale(float Exposure) const
+{
+	return GetLightExposureScale(Exposure, InverseExposureBlend);
+}
+
 FLightSceneProxy::FLightSceneProxy(const ULightComponent* InLightComponent)
 	: LightComponent(InLightComponent)
 	, SceneInterface(InLightComponent->GetScene())
diff --git a/Engine/Source/Runtime/Engine/Private/Components/PointLightComponent.cpp b/Engine/Source/Runtime/Engine/Private/Components/PointLightComponent.cpp
--- a/Engine/Source/Runtime/Engine/Private/Components/PointLightComponent.cpp
+++ b/Engine/Source/Runtime/Engine/Private/Components/PointLightComponent.cpp
@@ -147,6 +147,16 @@ void UPointLightComponent::SetLightFalloffExponent(float NewLightFalloffExponent
 	}
 }
 
+void UPointLightComponent::SetInverseExposureBlend(float NewInverseExposureBlend)
+{
+	if (AreDynamicDataChangesAllowed()
+		&& NewInverseExposureBlend != InverseExposureBlend)
+	{
+		InverseExposureBlend = NewInverseExposureBlend;
+		MarkRenderStateDirty();
+	}
+}
+
 void UPointLightComponent::SetSourceRadius(float NewValue)
 {
 	if (AreDynamicDataChangesAllowed()
diff --git a/Engine/Source/Runtime/Engine/Private/Particles/ParticleModules.cpp b/Engine/Source/Runtime/Engine/Private/Particles/ParticleModules.cpp
--- a/Engine/Source/Runtime/Engine/Private/Particles/ParticleModules.cpp
+++ b/Engine/Source/Runtime/Engine/Private/Particles/ParticleModules.cpp
@@ -3038,6 +3038,13 @@ static TAutoConsoleVariable<int32> CVarParticleLightQuality(
 	ECVF_Scalability
 	);
 
+static TAutoConsoleVariable<int32> CVarParticleDefaultLightInverseExposureBlend(
+	TEXT("fx.ParticleDefaultLightInverseExposureBlend"),
+	0.0f,
+	TEXT("Blend Factor used to blend between Intensity and Intensity / Exposure."),
+	ECVF_Default
+);
+
 void UParticleModuleLight::SpawnEx(FParticleEmitterInstance* Owner, int32 Offset, float SpawnTime, struct FRandomStream* InRandomStream, FBaseParticle* ParticleBase)
 {
 	int32 ParticleLightQuality = CVarParticleLightQuality.GetValueOnAnyThread();
@@ -3050,6 +3057,7 @@ void UParticleModuleLight::SpawnEx(FParticleEmitterInstance* Owner, int32 Offset
 		LightData.RadiusScale = RadiusScale.GetValue(Owner->EmitterTime, Owner->Component, InRandomStream);
 		// Exponent of 0 is interpreted by renderer as inverse squared falloff
 		LightData.LightExponent = bUseInverseSquaredFalloff ? 0 : LightExponent.GetValue(Owner->EmitterTime, Owner->Component, InRandomStream);
+		LightData.InverseExposureBlend = bOverrideInverseExposureBlend ? InverseExposureBlend : CVarParticleDefaultLightInverseExposureBlend.GetValueOnAnyThread();
 		const float RandomNumber = InRandomStream->GetFraction();
 		LightData.bValid = RandomNumber < SpawnFraction;
 		LightData.bAffectsTranslucency = bAffectsTranslucency;
@@ -3163,6 +3171,7 @@ void UParticleModuleLight::UpdateHQLight(UPointLightComponent* PointLightCompone
 	PointLightComponent->SetLightColor(NormalizedColor);
 	PointLightComponent->SetAttenuationRadius(Radius);
 	PointLightComponent->SetLightFalloffExponent(Payload.LightExponent);
+	PointLightComponent->SetInverseExposureBlend(Payload.InverseExposureBlend);
 
 	if (OwnerScene && bDoRTUpdate)
 	{
diff --git a/Engine/Source/Runtime/Engine/Private/Particles/ParticleSystemRender.cpp b/Engine/Source/Runtime/Engine/Private/Particles/ParticleSystemRender.cpp
--- a/Engine/Source/Runtime/Engine/Private/Particles/ParticleSystemRender.cpp
+++ b/Engine/Source/Runtime/Engine/Private/Particles/ParticleSystemRender.cpp
@@ -713,6 +713,7 @@ void GatherParticleLightData(const FDynamicSpriteEmitterReplayDataBase& Source,
 				ParticleLight.Radius =  LightPayload->RadiusScale * (Size.X + Size.Y) / 2.0f;
 				ParticleLight.Color = FVector3f(Particle.Color) * Particle.Color.A * LightPayload->ColorScale;
 				ParticleLight.Exponent = LightPayload->LightExponent;
+				ParticleLight.InverseExposureBlend = LightPayload->InverseExposureBlend;
 				ParticleLight.VolumetricScatteringIntensity = Source.LightVolumetricScatteringIntensity;
 				ParticleLight.bAffectTranslucency = LightPayload->bAffectsTranslucency;
 
diff --git a/Engine/Source/Runtime/Engine/Public/ParticleHelper.h b/Engine/Source/Runtime/Engine/Public/ParticleHelper.h
--- a/Engine/Source/Runtime/Engine/Public/ParticleHelper.h
+++ b/Engine/Source/Runtime/Engine/Public/ParticleHelper.h
@@ -602,6 +602,7 @@ struct FLightParticlePayload
 	uint64		LightId;
 	float		RadiusScale;
 	float		LightExponent;
+	float		InverseExposureBlend;
 	bool		bValid;
 	bool		bAffectsTranslucency;
 	bool		bHighQuality;
diff --git a/Engine/Source/Runtime/Engine/Public/PrimitiveSceneProxy.h b/Engine/Source/Runtime/Engine/Public/PrimitiveSceneProxy.h
--- a/Engine/Source/Runtime/Engine/Public/PrimitiveSceneProxy.h
+++ b/Engine/Source/Runtime/Engine/Public/PrimitiveSceneProxy.h
@@ -43,6 +43,7 @@ public:
 	FVector3f Color;
 	float Radius;
 	float Exponent;
+	float InverseExposureBlend = 0.0f;
 	float VolumetricScatteringIntensity;
 	bool bAffectTranslucency;
 };
diff --git a/Engine/Source/Runtime/Engine/Public/SceneManagement.h b/Engine/Source/Runtime/Engine/Public/SceneManagement.h
--- a/Engine/Source/Runtime/Engine/Public/SceneManagement.h
+++ b/Engine/Source/Runtime/Engine/Public/SceneManagement.h
@@ -1340,6 +1340,7 @@ struct FLightRenderParameters
 {
 	ENGINE_API void MakeShaderParameters(const FViewMatrices& ViewMatrices, float Exposure, FLightShaderParameters& OutShaderParameters) const;
 	ENGINE_API float GetLightExposureScale(float Exposure) const;
+	static ENGINE_API float GetLightExposureScale(float Exposure, float InverseExposureBlend);
 
 	// Position of the light in world space.
 	FVector WorldPosition;
diff --git a/Engine/Source/Runtime/Renderer/Private/LightGridInjection.cpp b/Engine/Source/Runtime/Renderer/Private/LightGridInjection.cpp
--- a/Engine/Source/Runtime/Renderer/Private/LightGridInjection.cpp
+++ b/Engine/Source/Runtime/Renderer/Private/LightGridInjection.cpp
@@ -315,7 +315,7 @@ static void PackLocalLightData(
 	const uint32 RectPackedW = FFloat16(-2.f).Encoded;
 
 	Out.LightPositionAndInvRadius				= FVector4f(LightTranslatedWorldPosition, 1.0f / FMath::Max(SimpleLight.Radius, KINDA_SMALL_NUMBER));
-	Out.LightColorAndFalloffExponent			= FVector4f((FVector3f)SimpleLight.Color, SimpleLight.Exponent);
+	Out.LightColorAndFalloffExponent			= FVector4f((FVector3f)SimpleLight.Color * FLightRenderParameters::GetLightExposureScale(View.GetLastEyeAdaptationExposure(), SimpleLight.InverseExposureBlend), SimpleLight.Exponent);
 	Out.LightDirectionAndShadowMapChannelMask	= FVector4f(FVector3f(1, 0, 0), FMath::AsFloat(ShadowMapChannelMask));
 	Out.SpotAnglesAndSourceRadiusPacked			= FVector4f(-2, 1, FMath::AsFloat(PackedZ), FMath::AsFloat(PackedW));
 	Out.LightTangentAndSpecularScale			= FVector4f(1.0f, 0.0f, 0.0f, 1.0f);
@@ -402,6 +402,8 @@ void FSceneRenderer::ComputeLightGrid(FRDGBuilder& GraphBuilder, bool bCullLight
 		int32 ClusteredSupportedEnd = 0;
 		int32 LumenSupportedStart = 0;
 
+		const float Exposure = View.GetLastEyeAdaptationExposure();
+
 		if (bCullLightsToGrid)
 		{
 			// Simple lights are copied without view dependent checks, so same in and out
@@ -453,8 +455,6 @@ void FSceneRenderer::ComputeLightGrid(FRDGBuilder& GraphBuilder, bool bCullLight
 				}
 			}
 
-			const float Exposure = View.GetLastEyeAdaptationExposure();
-
 			float SelectedForwardDirectionalLightIntensitySq = 0.0f;
 			int32 SelectedForwardDirectionalLightPriority = -1;
 			const TArray<FSortedLightSceneInfo, SceneRenderingAllocator>& SortedLights = SortedLightSet.SortedLights;
diff --git a/Engine/Source/Runtime/Renderer/Private/LightRendering.cpp b/Engine/Source/Runtime/Renderer/Private/LightRendering.cpp
--- a/Engine/Source/Runtime/Renderer/Private/LightRendering.cpp
+++ b/Engine/Source/Runtime/Renderer/Private/LightRendering.cpp
@@ -378,7 +378,7 @@ FDeferredLightUniformStruct GetSimpleDeferredLightParameters(
 
 	Out.LightParameters.TranslatedWorldPosition = FVector3f(LightWorldPosition + View.ViewMatrices.GetPreViewTranslation());
 	Out.LightParameters.InvRadius = 1.0f / FMath::Max(SimpleLight.Radius, KINDA_SMALL_NUMBER);
-	Out.LightParameters.Color = (FVector3f)SimpleLight.Color;
+	Out.LightParameters.Color = (FVector3f)SimpleLight.Color * FLightRenderParameters::GetLightExposureScale(View.GetLastEyeAdaptationExposure(), SimpleLight.InverseExposureBlend);
 	Out.LightParameters.FalloffExponent = SimpleLight.Exponent;
 	Out.LightParameters.Direction = FVector3f(1, 0, 0);
 	Out.LightParameters.Tangent = FVector3f(1, 0, 0);
diff --git a/Engine/Source/Runtime/Renderer/Private/TranslucentLighting.cpp b/Engine/Source/Runtime/Renderer/Private/TranslucentLighting.cpp
--- a/Engine/Source/Runtime/Renderer/Private/TranslucentLighting.cpp
+++ b/Engine/Source/Runtime/Renderer/Private/TranslucentLighting.cpp
@@ -1275,6 +1275,8 @@ void InjectSimpleTranslucencyLightingVolumeArray(
 
 		const int32 TranslucencyLightingVolumeDim = GetTranslucencyLightingVolumeDim();
 
+		const float Exposure = View.GetLastEyeAdaptationExposure();
+
 		// Inject into each volume cascade
 		// Operate on one cascade at a time to reduce render target switches
 		for (int32 VolumeCascadeIndex = 0; VolumeCascadeIndex < TVC_MAX; VolumeCascadeIndex++)
@@ -1303,7 +1305,7 @@ void InjectSimpleTranslucencyLightingVolumeArray(
 						PassParameters->View = View.ViewUniformBuffer;
 						PassParameters->VolumeCascadeIndex = VolumeCascadeIndex;
 						PassParameters->SimpleLightPositionAndRadius = FVector4f(TranslatedLightPosition, SimpleLight.Radius);
-						PassParameters->SimpleLightColorAndExponent = FVector4f((FVector3f)SimpleLight.Color, SimpleLight.Exponent);
+						PassParameters->SimpleLightColorAndExponent = FVector4f((FVector3f)SimpleLight.Color * FLightRenderParameters::GetLightExposureScale(Exposure, SimpleLight.InverseExposureBlend), SimpleLight.Exponent);
 						PassParameters->RenderTargets[0] = FRenderTargetBinding(VolumeAmbientTexture, ERenderTargetLoadAction::ELoad);
 						PassParameters->RenderTargets[1] = FRenderTargetBinding(VolumeDirectionalTexture, ERenderTargetLoadAction::ELoad);
 
