- Add warning if data from VM produces cells that would result in allocation failure
- Clean up render vs game thread data

#rnx

diff --git a/Engine/Plugins/FX/Niagara/Source/Niagara/Classes/NiagaraDataInterfaceNeighborGrid3D.h b/Engine/Plugins/FX/Niagara/Source/Niagara/Classes/NiagaraDataInterfaceNeighborGrid3D.h
--- a/Engine/Plugins/FX/Niagara/Source/Niagara/Classes/NiagaraDataInterfaceNeighborGrid3D.h
+++ b/Engine/Plugins/FX/Niagara/Source/Niagara/Classes/NiagaraDataInterfaceNeighborGrid3D.h
@@ -16,19 +16,24 @@ class FNiagaraSystemInstance;
 // move all data management to use per instance data
 // remove references to push data to render thread
 
-class NeighborGrid3DRWInstanceData
+struct FNDINeighborGrid3DInstanceData_GT
+{
+	FIntVector	NumCells = FIntVector::ZeroValue;
+	float		CellSize = 0.0f;
+	uint32		MaxNeighborsPerCell = 0;
+	FVector		WorldBBoxSize = FVector::ZeroVector;
+	bool		bNeedsRealloc = false;
+};
+
+struct FNDINeighborGrid3DInstanceData_RT
 {
-public:
 	void ResizeBuffers(FRDGBuilder& GraphBuilder);
 
 	FIntVector	NumCells = FIntVector::ZeroValue;
 	float		CellSize = 0.0f;
-	bool		SetGridFromCellSize = false;
 	uint32		MaxNeighborsPerCell = 0;
 	FVector		WorldBBoxSize = FVector::ZeroVector;
-
-	bool		NeedsRealloc_GT = false;
-	bool		NeedsRealloc_RT = false;
+	bool		bNeedsRealloc = false;
 
 	FNiagaraPooledRWBuffer NeighborhoodBuffer;
 	FNiagaraPooledRWBuffer NeighborhoodCountBuffer;
@@ -40,13 +45,13 @@ struct FNiagaraDataInterfaceProxyNeighborGrid3D : public FNiagaraDataInterfacePr
 	virtual void PostSimulate(const FNDIGpuComputePostSimulateContext& Context) override;
 
 	virtual void ConsumePerInstanceDataFromGameThread(void* PerInstanceData, const FNiagaraSystemInstanceID& Instance) override {}
-	virtual int32 PerInstanceDataPassedToRenderThreadSize() const override { return sizeof(NeighborGrid3DRWInstanceData); }	
+	virtual int32 PerInstanceDataPassedToRenderThreadSize() const override { return 0; }	
 
 	virtual FIntVector GetElementCount(FNiagaraSystemInstanceID SystemInstanceID) const override;
 
 	/* List of proxy data for each system instances*/
 	// #todo(dmp): this should all be refactored to avoid duplicate code
-	TMap<FNiagaraSystemInstanceID, NeighborGrid3DRWInstanceData> SystemInstancesToProxyData;
+	TMap<FNiagaraSystemInstanceID, FNDINeighborGrid3DInstanceData_RT> SystemInstancesToProxyData_RT;
 };
 
 UCLASS(EditInlineNew, Category = "Grid", meta = (DisplayName = "Neighbor Grid3D"))
@@ -107,7 +112,7 @@ public:
 	virtual bool InitPerInstanceData(void* PerInstanceData, FNiagaraSystemInstance* SystemInstance) override;
 	virtual void DestroyPerInstanceData(void* PerInstanceData, FNiagaraSystemInstance* SystemInstance) override;
 	virtual bool PerInstanceTick(void* PerInstanceData, FNiagaraSystemInstance* SystemInstance, float DeltaSeconds) override { return false;  }
-	virtual int32 PerInstanceDataSize()const override { return sizeof(NeighborGrid3DRWInstanceData); }
+	virtual int32 PerInstanceDataSize()const override { return sizeof(FNDINeighborGrid3DInstanceData_GT); }
 	virtual bool PerInstanceTickPostSimulate(void* PerInstanceData, FNiagaraSystemInstance* SystemInstance, float DeltaSeconds) override;
 	virtual bool HasPostSimulateTick() const override { return true; }
 	virtual bool HasPreSimulateTick() const override { return true; }
diff --git a/Engine/Plugins/FX/Niagara/Source/Niagara/Private/NiagaraDataInterfaceNeighborGrid3D.cpp b/Engine/Plugins/FX/Niagara/Source/Niagara/Private/NiagaraDataInterfaceNeighborGrid3D.cpp
--- a/Engine/Plugins/FX/Niagara/Source/Niagara/Private/NiagaraDataInterfaceNeighborGrid3D.cpp
+++ b/Engine/Plugins/FX/Niagara/Source/Niagara/Private/NiagaraDataInterfaceNeighborGrid3D.cpp
@@ -40,12 +40,12 @@ static FAutoConsoleVariableRef CVarMaxNiagaraNeighborGridCells(
 
 //////////////////////////////////////////////////////////////////////////
 
-void NeighborGrid3DRWInstanceData::ResizeBuffers(FRDGBuilder& GraphBuilder)
+void FNDINeighborGrid3DInstanceData_RT::ResizeBuffers(FRDGBuilder& GraphBuilder)
 {
 	const uint32 NumTotalCells = NumCells.X * NumCells.Y * NumCells.Z;
 	const uint32 NumIntsInGridBuffer = NumTotalCells * MaxNeighborsPerCell;
 
-	NeedsRealloc_RT = false;
+	bNeedsRealloc = false;
 	if (NumTotalCells > (uint32)GMaxNiagaraNeighborGridCells)
 	{
 		return;
@@ -54,7 +54,7 @@ void NeighborGrid3DRWInstanceData::ResizeBuffers(FRDGBuilder& GraphBuilder)
 	NeighborhoodCountBuffer.Release();
 	NeighborhoodBuffer.Release();
 
-	NeighborhoodCountBuffer.Initialize(GraphBuilder, TEXT("NiagaraNeighborGrid3D::NeighborCount"), EPixelFormat::PF_R32_SINT, sizeof(int32), FMath::Max(NumTotalCells, 1u), BUF_Static);
+	NeighborhoodCountBuffer.Initialize(GraphBuilder, TEXT("NiagaraNeighborGrid3D::NeighborCount"), EPixelFormat::PF_R32_SINT, sizeof(int32), FMath::Max<uint32>(NumTotalCells, 1u), BUF_Static);
 	NeighborhoodBuffer.Initialize(GraphBuilder, TEXT("NiagaraNeighborGrid3D::NeighborsGrid"), EPixelFormat::PF_R32_SINT, sizeof(int32), FMath::Max(NumIntsInGridBuffer, 1u), BUF_Static);
 }
 
@@ -233,7 +233,7 @@ void UNiagaraDataInterfaceNeighborGrid3D::GetVMExternalFunction(const FVMExterna
 
 void UNiagaraDataInterfaceNeighborGrid3D::GetWorldBBoxSize(FVectorVMExternalFunctionContext& Context)
 {
-	VectorVM::FUserPtrHandler<NeighborGrid3DRWInstanceData> InstData(Context);
+	VectorVM::FUserPtrHandler<FNDINeighborGrid3DInstanceData_GT> InstData(Context);
 
 	FNDIOutputParam<FVector3f> OutWorldBounds(Context);
 
@@ -245,7 +245,7 @@ void UNiagaraDataInterfaceNeighborGrid3D::GetWorldBBoxSize(FVectorVMExternalFunc
 
 void UNiagaraDataInterfaceNeighborGrid3D::GetNumCells(FVectorVMExternalFunctionContext& Context)
 {
-	VectorVM::FUserPtrHandler<NeighborGrid3DRWInstanceData> InstData(Context);
+	VectorVM::FUserPtrHandler<FNDINeighborGrid3DInstanceData_GT> InstData(Context);
 
 	FNDIOutputParam<int32> NumCellsX(Context);
 	FNDIOutputParam<int32> NumCellsY(Context);
@@ -265,7 +265,7 @@ void UNiagaraDataInterfaceNeighborGrid3D::GetNumCells(FVectorVMExternalFunctionC
 
 void UNiagaraDataInterfaceNeighborGrid3D::GetMaxNeighborsPerCell(FVectorVMExternalFunctionContext& Context)
 {
-	VectorVM::FUserPtrHandler<NeighborGrid3DRWInstanceData> InstData(Context);
+	VectorVM::FUserPtrHandler<FNDINeighborGrid3DInstanceData_GT> InstData(Context);
 
 	FNDIOutputParam<int32> OutMaxNeighborsPerCell(Context);
 
@@ -463,7 +463,7 @@ void UNiagaraDataInterfaceNeighborGrid3D::BuildShaderParameters(FNiagaraShaderPa
 void UNiagaraDataInterfaceNeighborGrid3D::SetShaderParameters(const FNiagaraDataInterfaceSetShaderParametersContext& Context) const
 {
 	FNiagaraDataInterfaceProxyNeighborGrid3D& DIProxy = Context.GetProxy<FNiagaraDataInterfaceProxyNeighborGrid3D>();
-	NeighborGrid3DRWInstanceData* ProxyData = DIProxy.SystemInstancesToProxyData.Find(Context.GetSystemInstanceID());
+	FNDINeighborGrid3DInstanceData_RT* ProxyData = DIProxy.SystemInstancesToProxyData_RT.Find(Context.GetSystemInstanceID());
 	FRDGBuilder& GraphBuilder = Context.GetGraphBuilder();
 
 	FShaderParameters* ShaderParameters = Context.GetParameterNestedStruct<FShaderParameters>();
@@ -509,8 +509,7 @@ void UNiagaraDataInterfaceNeighborGrid3D::SetShaderParameters(const FNiagaraData
 
 bool UNiagaraDataInterfaceNeighborGrid3D::InitPerInstanceData(void* PerInstanceData, FNiagaraSystemInstance* SystemInstance)
 {
-
-	NeighborGrid3DRWInstanceData* InstanceData = new (PerInstanceData) NeighborGrid3DRWInstanceData();
+	FNDINeighborGrid3DInstanceData_GT* InstanceData = new (PerInstanceData) FNDINeighborGrid3DInstanceData_GT();
 
 	FNiagaraDataInterfaceProxyNeighborGrid3D* RT_Proxy = GetProxyAs<FNiagaraDataInterfaceProxyNeighborGrid3D>();
 
@@ -585,25 +584,27 @@ bool UNiagaraDataInterfaceNeighborGrid3D::InitPerInstanceData(void* PerInstanceD
 	InstanceData->MaxNeighborsPerCell = RT_MaxNeighborsPerCell;	
 	InstanceData->NumCells = RT_NumCells;
 
-	if ((RT_NumCells.X * RT_NumCells.Y * RT_NumCells.Z) > GMaxNiagaraNeighborGridCells)
+	const uint32 NumTotalCells = RT_NumCells.X * RT_NumCells.Y * RT_NumCells.Z;
+	if (NumTotalCells == 0 || NumTotalCells > uint64(GMaxNiagaraNeighborGridCells))
 	{
-		UE_LOG(LogNiagara, Warning, TEXT("Dimensions are too big! Please adjust! %d x %d x %d > %d for ==> %s"), RT_NumCells.X, RT_NumCells.Y, RT_NumCells.Z, GMaxNiagaraNeighborGridCells , *GetFullNameSafe(this))
-			return false;
+		UE_LOG(LogNiagara, Warning, TEXT("NiagaraDataInterfaceNeighborGrid3D - Invalid NumCells(%dx%dx%d) total (%u / %d) for '%s'"), RT_NumCells.X, RT_NumCells.Y, RT_NumCells.Z, NumTotalCells, GMaxNiagaraNeighborGridCells, *GetFullNameSafe(this));
+		return false;
 	}
 
 	// @todo-threadsafety. This would be a race but I'm taking a ref here. Not ideal in the long term.
 	// Push Updates to Proxy.
 	ENQUEUE_RENDER_COMMAND(FUpdateData)(
-		[RT_Proxy, RT_NumCells, RT_MaxNeighborsPerCell, RT_WorldBBoxSize, InstanceID = SystemInstance->GetId()](FRHICommandListImmediate& RHICmdList)
-	{
-		check(!RT_Proxy->SystemInstancesToProxyData.Contains(InstanceID));
-		NeighborGrid3DRWInstanceData* TargetData = &RT_Proxy->SystemInstancesToProxyData.Add(InstanceID);
+		[RT_Proxy, RT_NumCells, RT_MaxNeighborsPerCell, RT_WorldBBoxSize, InstanceID=SystemInstance->GetId()](FRHICommandListImmediate& RHICmdList)
+		{
+			check(!RT_Proxy->SystemInstancesToProxyData_RT.Contains(InstanceID));
+			FNDINeighborGrid3DInstanceData_RT* TargetData = &RT_Proxy->SystemInstancesToProxyData_RT.Add(InstanceID);
 
-		TargetData->NumCells = RT_NumCells;
-		TargetData->MaxNeighborsPerCell = RT_MaxNeighborsPerCell;		
-		TargetData->WorldBBoxSize = RT_WorldBBoxSize;
-		TargetData->NeedsRealloc_RT = true;
-	});
+			TargetData->NumCells = RT_NumCells;
+			TargetData->MaxNeighborsPerCell = RT_MaxNeighborsPerCell;		
+			TargetData->WorldBBoxSize = RT_WorldBBoxSize;
+			TargetData->bNeedsRealloc = true;
+		}
+	);
 
 	return true;
 }
@@ -611,7 +612,7 @@ bool UNiagaraDataInterfaceNeighborGrid3D::InitPerInstanceData(void* PerInstanceD
 void UNiagaraDataInterfaceNeighborGrid3D::SetNumCells(FVectorVMExternalFunctionContext& Context)
 {
 	// This should only be called from a system or emitter script due to a need for only setting up initially.
-	VectorVM::FUserPtrHandler<NeighborGrid3DRWInstanceData> InstData(Context);
+	VectorVM::FUserPtrHandler<FNDINeighborGrid3DInstanceData_GT> InstData(Context);
 	VectorVM::FExternalFuncInputHandler<int> InNumCellsX(Context);
 	VectorVM::FExternalFuncInputHandler<int> InNumCellsY(Context);
 	VectorVM::FExternalFuncInputHandler<int> InNumCellsZ(Context);
@@ -620,11 +621,19 @@ void UNiagaraDataInterfaceNeighborGrid3D::SetNumCells(FVectorVMExternalFunctionC
 
 	for (int32 InstanceIdx = 0; InstanceIdx < Context.GetNumInstances(); ++InstanceIdx)
 	{
-		int NewNumCellsX = InNumCellsX.GetAndAdvance();
-		int NewNumCellsY = InNumCellsY.GetAndAdvance();
-		int NewNumCellsZ = InNumCellsZ.GetAndAdvance();
-		int NewMaxNeighborsPerCell = InMaxNeighborsPerCell.GetAndAdvance();
+		const int NewNumCellsX = InNumCellsX.GetAndAdvance();
+		const int NewNumCellsY = InNumCellsY.GetAndAdvance();
+		const int NewNumCellsZ = InNumCellsZ.GetAndAdvance();
+		const int NewMaxNeighborsPerCell = InMaxNeighborsPerCell.GetAndAdvance();
 		bool bSuccess = (InstData.Get() != nullptr && Context.GetNumInstances() == 1 && NumCells.X >= 0 && NumCells.Y >= 0 && NumCells.Z >= 0 && MaxNeighborsPerCell >= 0);
+
+		const uint32 NumTotalCells = NewNumCellsX * NewNumCellsY * NewNumCellsZ;
+		if (NumTotalCells == 0 || NumTotalCells > uint64(GMaxNiagaraNeighborGridCells))
+		{
+			UE_LOG(LogNiagara, Warning, TEXT("NiagaraDataInterfaceNeighborGrid3D - Invalid NumCells(%dx%dx%d) total (%u / %d) for '%s'"), NewNumCellsX, NewNumCellsY, NewNumCellsZ, NumTotalCells, GMaxNiagaraNeighborGridCells, *GetFullNameSafe(this));
+			bSuccess = false;
+		}
+
 		*OutSuccess.GetDestAndAdvance() = bSuccess;
 		if (bSuccess)
 		{
@@ -637,19 +646,19 @@ void UNiagaraDataInterfaceNeighborGrid3D::SetNumCells(FVectorVMExternalFunctionC
 
 			InstData->MaxNeighborsPerCell = NewMaxNeighborsPerCell;
 		
-			InstData->NeedsRealloc_GT = OldNumCells != InstData->NumCells || OldMaxNeighborsPerCell != InstData->MaxNeighborsPerCell;
+			InstData->bNeedsRealloc = OldNumCells != InstData->NumCells || OldMaxNeighborsPerCell != InstData->MaxNeighborsPerCell;
 		}
 	}
 }
 
 bool UNiagaraDataInterfaceNeighborGrid3D::PerInstanceTickPostSimulate(void* PerInstanceData, FNiagaraSystemInstance* SystemInstance, float DeltaSeconds)
 {
-	NeighborGrid3DRWInstanceData* InstanceData = static_cast<NeighborGrid3DRWInstanceData*>(PerInstanceData);
+	FNDINeighborGrid3DInstanceData_GT* InstanceData = static_cast<FNDINeighborGrid3DInstanceData_GT*>(PerInstanceData);
 	bool bNeedsReset = false;
 
-	if (InstanceData->NeedsRealloc_GT && InstanceData->NumCells.X > 0 && InstanceData->NumCells.Y > 0 && InstanceData->NumCells.Z > 0 && InstanceData->MaxNeighborsPerCell > 0)
+	if (InstanceData->bNeedsRealloc && InstanceData->NumCells.X > 0 && InstanceData->NumCells.Y > 0 && InstanceData->NumCells.Z > 0 && InstanceData->MaxNeighborsPerCell > 0)
 	{
-		InstanceData->NeedsRealloc_GT = false;
+		InstanceData->bNeedsRealloc = false;
 
 		InstanceData->CellSize = (InstanceData->WorldBBoxSize / FVector(InstanceData->NumCells.X, InstanceData->NumCells.Y, InstanceData->NumCells.Z))[0];
 
@@ -657,13 +666,13 @@ bool UNiagaraDataInterfaceNeighborGrid3D::PerInstanceTickPostSimulate(void* PerI
 		ENQUEUE_RENDER_COMMAND(FUpdateData)(
 			[RT_Proxy, RT_NumCells = InstanceData->NumCells, RT_MaxNeighborsPerCell = InstanceData->MaxNeighborsPerCell, RT_CellSize = InstanceData->CellSize, InstanceID = SystemInstance->GetId()](FRHICommandListImmediate& RHICmdList)
 		{
-			check(RT_Proxy->SystemInstancesToProxyData.Contains(InstanceID));
-			NeighborGrid3DRWInstanceData* TargetData = &RT_Proxy->SystemInstancesToProxyData.Add(InstanceID);
+			check(RT_Proxy->SystemInstancesToProxyData_RT.Contains(InstanceID));
+			FNDINeighborGrid3DInstanceData_RT* TargetData = &RT_Proxy->SystemInstancesToProxyData_RT.Add(InstanceID);
 
 			TargetData->NumCells = RT_NumCells;
 			TargetData->MaxNeighborsPerCell = RT_MaxNeighborsPerCell;			
 			TargetData->CellSize = RT_CellSize;
-			TargetData->NeedsRealloc_RT = true;
+			TargetData->bNeedsRealloc = true;
 		});
 	}
 
@@ -673,8 +682,8 @@ bool UNiagaraDataInterfaceNeighborGrid3D::PerInstanceTickPostSimulate(void* PerI
 
 void UNiagaraDataInterfaceNeighborGrid3D::DestroyPerInstanceData(void* PerInstanceData, FNiagaraSystemInstance* SystemInstance)
 {		
-	NeighborGrid3DRWInstanceData* InstanceData = static_cast<NeighborGrid3DRWInstanceData*>(PerInstanceData);
-	InstanceData->~NeighborGrid3DRWInstanceData();
+	FNDINeighborGrid3DInstanceData_GT* InstanceData = static_cast<FNDINeighborGrid3DInstanceData_GT*>(PerInstanceData);
+	InstanceData->~FNDINeighborGrid3DInstanceData_GT();
 
 	FNiagaraDataInterfaceProxyNeighborGrid3D* ThisProxy = GetProxyAs<FNiagaraDataInterfaceProxyNeighborGrid3D>();
 	if (!ThisProxy)
@@ -683,8 +692,8 @@ void UNiagaraDataInterfaceNeighborGrid3D::DestroyPerInstanceData(void* PerInstan
 	ENQUEUE_RENDER_COMMAND(FNiagaraDIDestroyInstanceData) (
 		[ThisProxy, InstanceID = SystemInstance->GetId()](FRHICommandListImmediate& CmdList)
 	{
-		//check(ThisProxy->SystemInstancesToProxyData.Contains(InstanceID));
-		ThisProxy->SystemInstancesToProxyData.Remove(InstanceID);
+		//check(ThisProxy->SystemInstancesToProxyData_RT.Contains(InstanceID));
+		ThisProxy->SystemInstancesToProxyData_RT.Remove(InstanceID);
 	}
 	);
 }
@@ -692,8 +701,8 @@ void UNiagaraDataInterfaceNeighborGrid3D::DestroyPerInstanceData(void* PerInstan
 void FNiagaraDataInterfaceProxyNeighborGrid3D::PreStage(const FNDIGpuComputePreStageContext& Context)
 {
 	FRDGBuilder& GraphBuilder = Context.GetGraphBuilder();
-	NeighborGrid3DRWInstanceData& ProxyData = SystemInstancesToProxyData.FindChecked(Context.GetSystemInstanceID());
-	if (ProxyData.NeedsRealloc_RT)
+	FNDINeighborGrid3DInstanceData_RT& ProxyData = SystemInstancesToProxyData_RT.FindChecked(Context.GetSystemInstanceID());
+	if (ProxyData.bNeedsRealloc)
 	{
 		ProxyData.ResizeBuffers(GraphBuilder);
 	}
@@ -710,7 +719,7 @@ void FNiagaraDataInterfaceProxyNeighborGrid3D::PostSimulate(const FNDIGpuCompute
 {
 	if (Context.IsFinalPostSimulate())
 	{
-		NeighborGrid3DRWInstanceData& ProxyData = SystemInstancesToProxyData.FindChecked(Context.GetSystemInstanceID());
+		FNDINeighborGrid3DInstanceData_RT& ProxyData = SystemInstancesToProxyData_RT.FindChecked(Context.GetSystemInstanceID());
 		ProxyData.NeighborhoodBuffer.EndGraphUsage();
 		ProxyData.NeighborhoodCountBuffer.EndGraphUsage();
 	}
@@ -718,7 +727,7 @@ void FNiagaraDataInterfaceProxyNeighborGrid3D::PostSimulate(const FNDIGpuCompute
 
 FIntVector FNiagaraDataInterfaceProxyNeighborGrid3D::GetElementCount(FNiagaraSystemInstanceID SystemInstanceID) const
 {
-	if ( const NeighborGrid3DRWInstanceData* TargetData = SystemInstancesToProxyData.Find(SystemInstanceID) )
+	if ( const FNDINeighborGrid3DInstanceData_RT* TargetData = SystemInstancesToProxyData_RT.Find(SystemInstanceID) )
 	{
 		return TargetData->NumCells;
 	}
