- Fix missing null check in SNiagaraOverviewStack::OnSelectionChanged

#rnx

diff --git a/Engine/Plugins/FX/Niagara/Source/NiagaraEditorWidgets/Private/SNiagaraOverviewStack.cpp b/Engine/Plugins/FX/Niagara/Source/NiagaraEditorWidgets/Private/SNiagaraOverviewStack.cpp
--- a/Engine/Plugins/FX/Niagara/Source/NiagaraEditorWidgets/Private/SNiagaraOverviewStack.cpp
+++ b/Engine/Plugins/FX/Niagara/Source/NiagaraEditorWidgets/Private/SNiagaraOverviewStack.cpp
@@ -1390,7 +1390,7 @@ void SNiagaraOverviewStack::OnSelectionChanged(UNiagaraStackEntry* InNewSelectio
 		bool bClearCurrentSelection = FSlateApplication::Get().GetModifierKeys().IsControlDown() == false;
 
 		// todo (ME) remove this special case when we remove the user parameters group from the stack in favor of the new user param tab
-		if(InNewSelection->IsA<UNiagaraStackSystemUserParametersGroup>())
+		if(InNewSelection && InNewSelection->IsA<UNiagaraStackSystemUserParametersGroup>())
 		{
 			InNewSelection->GetSystemViewModel()->FocusTab(TEXT("NiagaraSystemEditor_UserParameters"), true);
 		}
