Fixed new systems not registering their callback to update the sync variable's mapped to the user parameter store.
Also made sure to sync manually on post load

#rnx

diff --git a/Engine/Plugins/FX/Niagara/Source/NiagaraEditor/Private/NiagaraEditorModule.cpp b/Engine/Plugins/FX/Niagara/Source/NiagaraEditor/Private/NiagaraEditorModule.cpp
--- a/Engine/Plugins/FX/Niagara/Source/NiagaraEditor/Private/NiagaraEditorModule.cpp
+++ b/Engine/Plugins/FX/Niagara/Source/NiagaraEditor/Private/NiagaraEditorModule.cpp
@@ -211,6 +211,7 @@ class FNiagaraEditorOnlyDataUtilities : public INiagaraEditorOnlyDataUtilities
 		{
 			UNiagaraSystemEditorData* SystemEditorData = NewObject<UNiagaraSystemEditorData>(InOuter);
 			SystemEditorData->SynchronizeOverviewGraphWithSystem(*System);
+			SystemEditorData->InitOnSyncScriptVariables(*System);
 			return SystemEditorData;
 		}
 		return nullptr;
diff --git a/Engine/Plugins/FX/Niagara/Source/NiagaraEditor/Private/NiagaraSystemEditorData.cpp b/Engine/Plugins/FX/Niagara/Source/NiagaraEditor/Private/NiagaraSystemEditorData.cpp
--- a/Engine/Plugins/FX/Niagara/Source/NiagaraEditor/Private/NiagaraSystemEditorData.cpp
+++ b/Engine/Plugins/FX/Niagara/Source/NiagaraEditor/Private/NiagaraSystemEditorData.cpp
@@ -108,9 +108,7 @@ void UNiagaraSystemEditorData::PostLoadFromOwner(UObject* InOwner)
 		OwningSystem->GetExposedParameters().RedirectUserVariable(RedirectedUserParameter);
 		FindOrAddUserScriptVariable(RedirectedUserParameter, *OwningSystem);
 	}
-
-	OwningSystem->GetExposedParameters().AddOnChangedHandler(FNiagaraParameterStore::FOnChanged::FDelegate::CreateUObject(this, &UNiagaraSystemEditorData::SyncUserScriptVariables, OwningSystem));
-
+	
 	const int32 NiagaraVer = GetLinkerCustomVersion(FNiagaraCustomVersion::GUID);
 
 	if (NiagaraVer < FNiagaraCustomVersion::PlaybackRangeStoredOnSystem)
@@ -124,6 +122,10 @@ void UNiagaraSystemEditorData::PostLoadFromOwner(UObject* InOwner)
 		SystemOverviewGraph->Schema = UEdGraphSchema_NiagaraSystemOverview::StaticClass();
 	}
 
+	// we make sure to refresh the script variables mapping to the user parameter store
+	SyncUserScriptVariables(OwningSystem);
+	InitOnSyncScriptVariables(*OwningSystem);
+
 	if(SystemOverviewGraph->Nodes.Num() == 0)
 	{
 		SynchronizeOverviewGraphWithSystem(*OwningSystem);
@@ -344,6 +346,11 @@ void UNiagaraSystemEditorData::SynchronizeOverviewGraphWithSystem(UNiagaraSystem
 	SystemOverviewGraph->NotifyGraphChanged();
 }
 
+void UNiagaraSystemEditorData::InitOnSyncScriptVariables(UNiagaraSystem& System)
+{
+	System.GetExposedParameters().AddOnChangedHandler(FNiagaraParameterStore::FOnChanged::FDelegate::CreateUObject(this, &UNiagaraSystemEditorData::SyncUserScriptVariables, &System));
+}
+
 void UNiagaraSystemEditorData::SyncUserScriptVariables(UNiagaraSystem* System)
 {
 	TArray<FNiagaraVariable> UserParameters;
diff --git a/Engine/Plugins/FX/Niagara/Source/NiagaraEditor/Public/NiagaraSystemEditorData.h b/Engine/Plugins/FX/Niagara/Source/NiagaraEditor/Public/NiagaraSystemEditorData.h
--- a/Engine/Plugins/FX/Niagara/Source/NiagaraEditor/Public/NiagaraSystemEditorData.h
+++ b/Engine/Plugins/FX/Niagara/Source/NiagaraEditor/Public/NiagaraSystemEditorData.h
@@ -92,6 +92,8 @@ public:
 
 	void SynchronizeOverviewGraphWithSystem(UNiagaraSystem& OwnerSystem);
 
+	void InitOnSyncScriptVariables(UNiagaraSystem& System);
+	
 	void SyncUserScriptVariables(UNiagaraSystem* System);
 	FOnUserParameterScriptVariablesSynced& OnUserParameterScriptVariablesSynced() { return OnUserParameterScriptVariablesSyncedDelegate; }
 	
