Fixed confusing open tabs when a Niagara system's cached layout is applied to an emitter (i.e. you open a system, close it, open an emitter) since user param tab & user param hierarchy tab are not available for emitters.
Solution:
We close the hierarchy editor tab when default mode is activated, and we swap out the contents of the user param tab with a small explanation.

The reason why we don't close the user param tab is that system & emitter layout state will influence each other, so if the user closes the weird empty tab in an emitter, he might get confused where the tab is when he opens a system.

#rnx

diff --git a/Engine/Plugins/FX/Niagara/Source/NiagaraEditor/Private/Toolkits/SystemToolkitModes/NiagaraSystemToolkitModeBase.cpp b/Engine/Plugins/FX/Niagara/Source/NiagaraEditor/Private/Toolkits/SystemToolkitModes/NiagaraSystemToolkitModeBase.cpp
--- a/Engine/Plugins/FX/Niagara/Source/NiagaraEditor/Private/Toolkits/SystemToolkitModes/NiagaraSystemToolkitModeBase.cpp
+++ b/Engine/Plugins/FX/Niagara/Source/NiagaraEditor/Private/Toolkits/SystemToolkitModes/NiagaraSystemToolkitModeBase.cpp
@@ -161,6 +161,17 @@ void FNiagaraSystemToolkitModeBase::OnSystemSelectionChanged()
 	// Do nothing for now
 }
 
+void FNiagaraSystemToolkitModeBase::PostActivateMode()
+{
+	// by default we want to close the user parameters hierarchy tab if it has been open before. T
+	// his fixes the issue of an empty tab being summoned from a system's cached layout when applied to an emitter.
+	// for consistency, we close it not only for emitters but for systems too
+	if(TSharedPtr<SDockTab> UserParametersHierarchyDockTab = SystemToolkit.Pin()->GetTabManager()->FindExistingLiveTab(UserParametersHierarchyTabID))
+	{
+		UserParametersHierarchyDockTab->RequestCloseTab();
+	}
+}
+
 void FNiagaraSystemToolkitModeBase::UpdateSelectionForActiveDocument()
 {
 	TSharedPtr<FNiagaraSystemToolkit> Toolkit = StaticCastSharedPtr<FNiagaraSystemToolkit>(SystemToolkit.Pin());
@@ -294,17 +305,14 @@ void FNiagaraSystemToolkitModeBase::RegisterTabFactories(TSharedPtr<FTabManager>
 		.SetDisplayName(LOCTEXT("VersioningTab", "Versioning"))
 		.SetGroup(WorkspaceMenuCategory.ToSharedRef())
 		.SetIcon(FSlateIcon(FAppStyle::GetAppStyleSetName(), "Versions"));
-
-	if(SystemToolkit.Pin()->GetSystemViewModel()->GetEditMode() == ENiagaraSystemViewModelEditMode::SystemAsset)
-	{
-		InTabManager->RegisterTabSpawner(UserParametersTabID, FOnSpawnTab::CreateSP(this, &FNiagaraSystemToolkitModeBase::SpawnTab_UserParameters))
-			.SetDisplayName(LOCTEXT("UserParametersTab", "User Parameters"))
-			.SetGroup(WorkspaceMenuCategory.ToSharedRef());
-		
-		InTabManager->RegisterTabSpawner(UserParametersHierarchyTabID, FOnSpawnTab::CreateSP(this, &FNiagaraSystemToolkitModeBase::SpawnTab_UserParametersHierarchyEditor))
-			.SetDisplayName(LOCTEXT("UserParametersHierarchyTab", "User Parameters Hierarchy"))
-			.SetGroup(WorkspaceMenuCategory.ToSharedRef());
-	}
+	
+	InTabManager->RegisterTabSpawner(UserParametersTabID, FOnSpawnTab::CreateSP(this, &FNiagaraSystemToolkitModeBase::SpawnTab_UserParameters))
+		.SetDisplayName(LOCTEXT("UserParametersTab", "User Parameters"))
+		.SetGroup(WorkspaceMenuCategory.ToSharedRef());
+	
+	InTabManager->RegisterTabSpawner(UserParametersHierarchyTabID, FOnSpawnTab::CreateSP(this, &FNiagaraSystemToolkitModeBase::SpawnTab_UserParametersHierarchyEditor))
+		.SetDisplayName(LOCTEXT("UserParametersHierarchyTab", "User Parameters Hierarchy"))
+		.SetGroup(WorkspaceMenuCategory.ToSharedRef());
 }
 
 int FNiagaraSystemToolkitModeBase::GetActiveSelectionDetailsIndex() const
@@ -780,13 +788,24 @@ TSharedRef<SDockTab> FNiagaraSystemToolkitModeBase::SpawnTab_Versioning(const FS
 TSharedRef<SDockTab> FNiagaraSystemToolkitModeBase::SpawnTab_UserParameters(const FSpawnTabArgs& Args)
 {
 	check(Args.GetTabId().TabType == UserParametersTabID);
-	
+
 	TSharedRef<SDockTab> SpawnedTab = SNew(SDockTab)
-		.Label(LOCTEXT("UserParametersTabTitle", "User Parameters"))
-		[
-			SNew(SNiagaraSystemUserParameters, SystemToolkit.Pin()->GetSystemViewModel())
-		];
+	.Label(LOCTEXT("UserParametersTabTitle", "User Parameters"));
 
+	if(SystemToolkit.Pin()->GetSystemViewModel()->GetEditMode() == ENiagaraSystemViewModelEditMode::SystemAsset)
+	{
+		SpawnedTab->SetContent(SNew(SNiagaraSystemUserParameters, SystemToolkit.Pin()->GetSystemViewModel()));
+		return SpawnedTab;
+	}
+
+	TSharedRef<SWidget> EmptyTabContent = SNew(SBox)
+	.HAlign(HAlign_Center)
+	.VAlign(VAlign_Center)
+	[
+		SNew(STextBlock).Text(LOCTEXT("EmptyUserParametersTabText", "User Parameters are only supported in System assets.")).AutoWrapText(true)
+	];
+	
+	SpawnedTab->SetContent(EmptyTabContent);
 	return SpawnedTab;
 }
 
diff --git a/Engine/Plugins/FX/Niagara/Source/NiagaraEditor/Private/Toolkits/SystemToolkitModes/NiagaraSystemToolkitModeBase.h b/Engine/Plugins/FX/Niagara/Source/NiagaraEditor/Private/Toolkits/SystemToolkitModes/NiagaraSystemToolkitModeBase.h
--- a/Engine/Plugins/FX/Niagara/Source/NiagaraEditor/Private/Toolkits/SystemToolkitModes/NiagaraSystemToolkitModeBase.h
+++ b/Engine/Plugins/FX/Niagara/Source/NiagaraEditor/Private/Toolkits/SystemToolkitModes/NiagaraSystemToolkitModeBase.h
@@ -23,6 +23,8 @@ public:
 
 	virtual void OnActiveDocumentChanged(TSharedPtr<class SDockTab> NewActiveTab);
 	virtual void OnSystemSelectionChanged();
+
+	virtual void PostActivateMode() override;
 private:
 	TSharedRef<SDockTab> SpawnTab_Viewport(const FSpawnTabArgs& Args);
 	TSharedRef<SDockTab> SpawnTab_CurveEd(const FSpawnTabArgs& Args);
diff --git a/Engine/Plugins/FX/Niagara/Source/NiagaraEditor/Private/Toolkits/SystemToolkitModes/NiagaraSystemToolkitMode_Default.cpp b/Engine/Plugins/FX/Niagara/Source/NiagaraEditor/Private/Toolkits/SystemToolkitModes/NiagaraSystemToolkitMode_Default.cpp
--- a/Engine/Plugins/FX/Niagara/Source/NiagaraEditor/Private/Toolkits/SystemToolkitModes/NiagaraSystemToolkitMode_Default.cpp
+++ b/Engine/Plugins/FX/Niagara/Source/NiagaraEditor/Private/Toolkits/SystemToolkitModes/NiagaraSystemToolkitMode_Default.cpp
@@ -304,6 +304,8 @@ void FNiagaraSystemToolkitMode_Default::ExtendToolbar()
 
 void FNiagaraSystemToolkitMode_Default::PostActivateMode()
 {
+	FNiagaraSystemToolkitModeBase::PostActivateMode();
+	
 	// We need to register the primary tab ID with the document system so that we can easily update and 
 	// reference it elsewhere and only this class knows it's identity since it is dynamically allocated.
 	TSharedPtr<FNiagaraSystemToolkit> Toolkit = StaticCastSharedPtr<FNiagaraSystemToolkit>(SystemToolkit.Pin());
