- Translucent draw latency settings are now part of the project settings and can be overridden on each renderer

diff --git a/Engine/Plugins/FX/Niagara/Source/Niagara/Private/NiagaraGpuComputeDispatch.cpp b/Engine/Plugins/FX/Niagara/Source/Niagara/Private/NiagaraGpuComputeDispatch.cpp
--- a/Engine/Plugins/FX/Niagara/Source/Niagara/Private/NiagaraGpuComputeDispatch.cpp
+++ b/Engine/Plugins/FX/Niagara/Source/Niagara/Private/NiagaraGpuComputeDispatch.cpp
@@ -58,7 +58,7 @@ static FAutoConsoleVariableRef CVarNiagaraGpuSubmitCommandHint(
 	ECVF_Default
 );
 
-int32 GNiagaraGpuLowLatencyTranslucencyEnabled = 0;
+int32 GNiagaraGpuLowLatencyTranslucencyEnabled = 1;
 static FAutoConsoleVariableRef CVarNiagaraGpuLowLatencyTranslucencyEnabled(
 	TEXT("fx.NiagaraGpuLowLatencyTranslucencyEnabled"),
 	GNiagaraGpuLowLatencyTranslucencyEnabled,
diff --git a/Engine/Plugins/FX/Niagara/Source/Niagara/Private/NiagaraMeshRendererProperties.cpp b/Engine/Plugins/FX/Niagara/Source/Niagara/Private/NiagaraMeshRendererProperties.cpp
--- a/Engine/Plugins/FX/Niagara/Source/Niagara/Private/NiagaraMeshRendererProperties.cpp
+++ b/Engine/Plugins/FX/Niagara/Source/Niagara/Private/NiagaraMeshRendererProperties.cpp
@@ -87,14 +87,11 @@ bool FNiagaraMeshRendererMeshProperties::HasValidMeshProperties() const
 	return Mesh || UserParamBinding.Parameter.IsValid();
 }
 
-
-
 UNiagaraMeshRendererProperties::UNiagaraMeshRendererProperties()
 	: SourceMode(ENiagaraRendererSourceDataMode::Particles)
 	, SortMode(ENiagaraSortMode::None)
 	, bOverrideMaterials(false)
 	, bSortOnlyWhenTranslucent(true)
-	, bGpuLowLatencyTranslucency(true)
 	, bSubImageBlend(false)
 	, SubImageSize(1.0f, 1.0f)
 	, FacingMode(ENiagaraMeshFacingMode::Default)
diff --git a/Engine/Plugins/FX/Niagara/Source/Niagara/Private/NiagaraRendererMeshes.cpp b/Engine/Plugins/FX/Niagara/Source/Niagara/Private/NiagaraRendererMeshes.cpp
--- a/Engine/Plugins/FX/Niagara/Source/Niagara/Private/NiagaraRendererMeshes.cpp
+++ b/Engine/Plugins/FX/Niagara/Source/Niagara/Private/NiagaraRendererMeshes.cpp
@@ -79,7 +79,7 @@ FNiagaraRendererMeshes::FNiagaraRendererMeshes(ERHIFeatureLevel::Type FeatureLev
 	SortMode = Properties->SortMode;
 	bSortHighPrecision = UNiagaraRendererProperties::IsSortHighPrecision(Properties->SortPrecision);
 	bSortOnlyWhenTranslucent = Properties->bSortOnlyWhenTranslucent;
-	bGpuLowLatencyTranslucency = Properties->bGpuLowLatencyTranslucency;
+	bGpuLowLatencyTranslucency = UNiagaraRendererProperties::IsGpuTranslucentThisFrame(Properties->GpuTranslucentLatency);
 	bOverrideMaterials = Properties->bOverrideMaterials;
 	SubImageSize = FVector2f(Properties->SubImageSize);	// LWC_TODO: Precision loss
 	bSubImageBlend = Properties->bSubImageBlend;
diff --git a/Engine/Plugins/FX/Niagara/Source/Niagara/Private/NiagaraRendererProperties.cpp b/Engine/Plugins/FX/Niagara/Source/Niagara/Private/NiagaraRendererProperties.cpp
--- a/Engine/Plugins/FX/Niagara/Source/Niagara/Private/NiagaraRendererProperties.cpp
+++ b/Engine/Plugins/FX/Niagara/Source/Niagara/Private/NiagaraRendererProperties.cpp
@@ -609,5 +609,14 @@ bool UNiagaraRendererProperties::IsSortHighPrecision(ENiagaraRendererSortPrecisi
 	return SortPrecision == ENiagaraRendererSortPrecision::High;
 }
 
+bool UNiagaraRendererProperties::IsGpuTranslucentThisFrame(ENiagaraRendererGpuTranslucentLatency Latency)
+{
+	if (Latency == ENiagaraRendererGpuTranslucentLatency::ProjectDefault)
+	{
+		return GetDefault<UNiagaraSettings>()->DefaultGpuTranslucentLatency == ENiagaraDefaultGpuTranslucentLatency::Immediate;
+	}
+	return Latency == ENiagaraRendererGpuTranslucentLatency::Immediate;
+}
+
 #undef LOCTEXT_NAMESPACE
 
diff --git a/Engine/Plugins/FX/Niagara/Source/Niagara/Private/NiagaraRendererSprites.cpp b/Engine/Plugins/FX/Niagara/Source/Niagara/Private/NiagaraRendererSprites.cpp
--- a/Engine/Plugins/FX/Niagara/Source/Niagara/Private/NiagaraRendererSprites.cpp
+++ b/Engine/Plugins/FX/Niagara/Source/Niagara/Private/NiagaraRendererSprites.cpp
@@ -93,7 +93,7 @@ FNiagaraRendererSprites::FNiagaraRendererSprites(ERHIFeatureLevel::Type FeatureL
 	bRemoveHMDRollInVR = Properties->bRemoveHMDRollInVR;
 	bSortHighPrecision = UNiagaraRendererProperties::IsSortHighPrecision(Properties->SortPrecision);
 	bSortOnlyWhenTranslucent = Properties->bSortOnlyWhenTranslucent;
-	bGpuLowLatencyTranslucency = Properties->bGpuLowLatencyTranslucency;
+	bGpuLowLatencyTranslucency = UNiagaraRendererProperties::IsGpuTranslucentThisFrame(Properties->GpuTranslucentLatency);
 	MinFacingCameraBlendDistance = Properties->MinFacingCameraBlendDistance;
 	MaxFacingCameraBlendDistance = Properties->MaxFacingCameraBlendDistance;
 	RendererVisibility = Properties->RendererVisibility;
diff --git a/Engine/Plugins/FX/Niagara/Source/Niagara/Private/NiagaraSpriteRendererProperties.cpp b/Engine/Plugins/FX/Niagara/Source/Niagara/Private/NiagaraSpriteRendererProperties.cpp
--- a/Engine/Plugins/FX/Niagara/Source/Niagara/Private/NiagaraSpriteRendererProperties.cpp
+++ b/Engine/Plugins/FX/Niagara/Source/Niagara/Private/NiagaraSpriteRendererProperties.cpp
@@ -54,7 +54,6 @@ UNiagaraSpriteRendererProperties::UNiagaraSpriteRendererProperties()
 	, bSubImageBlend(false)
 	, bRemoveHMDRollInVR(false)
 	, bSortOnlyWhenTranslucent(true)
-	, bGpuLowLatencyTranslucency(true)
 	, MinFacingCameraBlendDistance(0.0f)
 	, MaxFacingCameraBlendDistance(0.0f)
 #if WITH_EDITORONLY_DATA
diff --git a/Engine/Plugins/FX/Niagara/Source/Niagara/Public/NiagaraMeshRendererProperties.h b/Engine/Plugins/FX/Niagara/Source/Niagara/Public/NiagaraMeshRendererProperties.h
--- a/Engine/Plugins/FX/Niagara/Source/Niagara/Public/NiagaraMeshRendererProperties.h
+++ b/Engine/Plugins/FX/Niagara/Source/Niagara/Public/NiagaraMeshRendererProperties.h
@@ -234,12 +234,12 @@ public:
 	ENiagaraRendererSortPrecision SortPrecision = ENiagaraRendererSortPrecision::Default;
 
 	/**
-	If true and a GPU emitter, we will use the current frames data to render with regardless of where the batcher may execute the dispatches.
-	If you have other emitters that are not translucent and using data that forces it to be a frame latent (i.e. view uniform buffer) you may need to disable
-	on renderers with translucent materials if you need the frame they are reading to match exactly.
+	Gpu simulations run at different points in the frame depending on what features are used, i.e. depth buffer, distance fields, etc.
+	Opaque materials will run latent when these features are used.
+	Translucent materials can choose if they want to use this frames or the previous frames data to match opaque draws.
 	*/
-	UPROPERTY(EditAnywhere, AdvancedDisplay, Category = "Rendering")
-	uint32 bGpuLowLatencyTranslucency : 1;
+	UPROPERTY(EditAnywhere, AdvancedDisplay, Category = "Sprite Rendering")
+	ENiagaraRendererGpuTranslucentLatency GpuTranslucentLatency = ENiagaraRendererGpuTranslucentLatency::ProjectDefault;
 
 	/** If true, blends the sub-image UV lookup with its next adjacent member using the fractional part of the SubImageIndex float value as the linear interpolation factor.*/
 	UPROPERTY(EditAnywhere, Category = "SubUV", meta = (DisplayName = "Sub UV Blending Enabled"))
diff --git a/Engine/Plugins/FX/Niagara/Source/Niagara/Public/NiagaraRendererProperties.h b/Engine/Plugins/FX/Niagara/Source/Niagara/Public/NiagaraRendererProperties.h
--- a/Engine/Plugins/FX/Niagara/Source/Niagara/Public/NiagaraRendererProperties.h
+++ b/Engine/Plugins/FX/Niagara/Source/Niagara/Public/NiagaraRendererProperties.h
@@ -161,6 +161,17 @@ enum class ENiagaraRendererSortPrecision : uint8
 	High,
 };
 
+UENUM()
+enum class ENiagaraRendererGpuTranslucentLatency : uint8
+{
+	/** Uses the project default value. */
+	ProjectDefault,
+	/** Gpu simulations will always read this frames data for translucent materials. */
+	Immediate,
+	/** Gpu simulations will read the previous frames data if the simulation has to run in PostRenderOpaque. */
+	Latent,
+};
+
 USTRUCT()
 struct FNiagaraRendererMaterialScalarParameter
 {
@@ -339,6 +350,8 @@ public:
 
 	static bool IsSortHighPrecision(ENiagaraRendererSortPrecision SortPrecision);
 
+	static bool IsGpuTranslucentThisFrame(ENiagaraRendererGpuTranslucentLatency Latency);
+
 	template<typename TAction>
 	void ForEachPlatformSet(TAction Func);
 
diff --git a/Engine/Plugins/FX/Niagara/Source/Niagara/Public/NiagaraSettings.h b/Engine/Plugins/FX/Niagara/Source/Niagara/Public/NiagaraSettings.h
--- a/Engine/Plugins/FX/Niagara/Source/Niagara/Public/NiagaraSettings.h
+++ b/Engine/Plugins/FX/Niagara/Source/Niagara/Public/NiagaraSettings.h
@@ -71,6 +71,15 @@ enum class ENiagaraDefaultSortPrecision : uint8
 	High,
 };
 
+UENUM()
+enum class ENiagaraDefaultGpuTranslucentLatency : uint8
+{
+	/** Gpu simulations will always read this frames data for translucent materials. */
+	Immediate,
+	/** Gpu simulations will read the previous frames data if the simulation has to run in PostRenderOpaque. */
+	Latent,
+};
+
 UENUM()
 namespace ENDICollisionQuery_AsyncGpuTraceProvider
 {
@@ -162,6 +171,10 @@ class NIAGARA_API UNiagaraSettings : public UDeveloperSettings
 	UPROPERTY(config, EditAnywhere, Category = Renderer)
 	ENiagaraDefaultSortPrecision DefaultSortPrecision = ENiagaraDefaultSortPrecision::Low;
 
+	/** The default setting for Gpu simulation translucent draw latency. */
+	UPROPERTY(config, EditAnywhere, Category = Renderer)
+	ENiagaraDefaultGpuTranslucentLatency DefaultGpuTranslucentLatency = ENiagaraDefaultGpuTranslucentLatency::Immediate;
+
 	/** The default InverseExposureBlend used for the light renderer. */
 	UPROPERTY(config, EditAnywhere, Category = LightRenderer)
 	float DefaultLightInverseExposureBlend = 0.0f;
diff --git a/Engine/Plugins/FX/Niagara/Source/Niagara/Public/NiagaraSpriteRendererProperties.h b/Engine/Plugins/FX/Niagara/Source/Niagara/Public/NiagaraSpriteRendererProperties.h
--- a/Engine/Plugins/FX/Niagara/Source/Niagara/Public/NiagaraSpriteRendererProperties.h
+++ b/Engine/Plugins/FX/Niagara/Source/Niagara/Public/NiagaraSpriteRendererProperties.h
@@ -199,12 +199,12 @@ public:
 	ENiagaraRendererSortPrecision SortPrecision = ENiagaraRendererSortPrecision::Default;
 
 	/**
-	If true and a GPU emitter, we will use the current frames data to render with regardless of where the batcher may execute the dispatches.
-	If you have other emitters that are not translucent and using data that forces it to be a frame latent (i.e. view uniform buffer) you may need to disable
-	on renderers with translucent materials if you need the frame they are reading to match exactly.
+	Gpu simulations run at different points in the frame depending on what features are used, i.e. depth buffer, distance fields, etc.
+	Opaque materials will run latent when these features are used.
+	Translucent materials can choose if they want to use this frames or the previous frames data to match opaque draws.
 	*/
-	UPROPERTY(EditAnywhere, AdvancedDisplay, Category = "Rendering")
-	uint32 bGpuLowLatencyTranslucency : 1;
+	UPROPERTY(EditAnywhere, AdvancedDisplay, Category = "Sprite Rendering")
+	ENiagaraRendererGpuTranslucentLatency GpuTranslucentLatency = ENiagaraRendererGpuTranslucentLatency::ProjectDefault;
 
 	/**
 	This setting controls what happens when a sprite becomes less than a pixel in size.
