Fixed include paths in ../Niagara/...

diff --git a/Engine/Plugins/FX/Niagara/Source/Niagara/Classes/NiagaraDataInterface.h b/Engine/Plugins/FX/Niagara/Source/Niagara/Classes/NiagaraDataInterface.h
--- a/Engine/Plugins/FX/Niagara/Source/Niagara/Classes/NiagaraDataInterface.h
+++ b/Engine/Plugins/FX/Niagara/Source/Niagara/Classes/NiagaraDataInterface.h
@@ -1,11 +1,11 @@
 // Copyright Epic Games, Inc. All Rights Reserved.
 #pragma once
 
-#include "Niagara/Public/NiagaraCommon.h"
-#include "NiagaraShared.h"
-#include "VectorVM.h"
+#include "NiagaraCommon.h"
 #include "NiagaraDataInterfaceBase.h"
 #include "NiagaraScriptBase.h"
+#include "NiagaraShared.h"
+#include "VectorVM.h"
 #include "NiagaraDataInterface.generated.h"
 
 class RENDERCORE_API FRDGBuilder;
@@ -1356,4 +1356,4 @@ public:
 
 private:
 	TSharedPtr<ResourceType> Resource;
-};
\ No newline at end of file
+};
diff --git a/Engine/Plugins/FX/Niagara/Source/Niagara/Niagara.Build.cs b/Engine/Plugins/FX/Niagara/Source/Niagara/Niagara.Build.cs
--- a/Engine/Plugins/FX/Niagara/Source/Niagara/Niagara.Build.cs
+++ b/Engine/Plugins/FX/Niagara/Source/Niagara/Niagara.Build.cs
@@ -61,6 +61,7 @@ public class Niagara : ModuleRules
         PrivateIncludePaths.AddRange(
             new string[] {
                 "Niagara/Private",
+				Path.Combine(Path.GetFullPath(Target.RelativeEnginePath), "Source/Runtime/Engine/Private"),
 				Path.Combine(Path.GetFullPath(Target.RelativeEnginePath), "Source/Runtime/Renderer/Private")
 			});
 
diff --git a/Engine/Plugins/FX/Niagara/Source/Niagara/Private/NiagaraDataInterfaceRigidMeshCollisionQuery.cpp b/Engine/Plugins/FX/Niagara/Source/Niagara/Private/NiagaraDataInterfaceRigidMeshCollisionQuery.cpp
--- a/Engine/Plugins/FX/Niagara/Source/Niagara/Private/NiagaraDataInterfaceRigidMeshCollisionQuery.cpp
+++ b/Engine/Plugins/FX/Niagara/Source/Niagara/Private/NiagaraDataInterfaceRigidMeshCollisionQuery.cpp
@@ -3,12 +3,13 @@
 #include "NiagaraDataInterfaceRigidMeshCollisionQuery.h"
 #include "Algo/ForEach.h"
 #include "Algo/RemoveIf.h"
-#include "Components/SkeletalMeshComponent.h"
-#include "Engine/Canvas.h"
 #include "Animation/SkeletalMeshActor.h"
 #include "AnimationRuntime.h"
-#include "SkeletalRenderPublic.h"
+#include "Components/SkeletalMeshComponent.h"
+#include "DistanceFieldLightingShared.h"
+#include "Engine/Canvas.h"
 #include "EngineUtils.h"
+#include "SkeletalRenderPublic.h"
 
 #include "NiagaraComponent.h"
 #include "NiagaraDataInterfaceUtilities.h"
diff --git a/Engine/Plugins/FX/Niagara/Source/Niagara/Private/NiagaraGpuComputeDebug.cpp b/Engine/Plugins/FX/Niagara/Source/Niagara/Private/NiagaraGpuComputeDebug.cpp
--- a/Engine/Plugins/FX/Niagara/Source/Niagara/Private/NiagaraGpuComputeDebug.cpp
+++ b/Engine/Plugins/FX/Niagara/Source/Niagara/Private/NiagaraGpuComputeDebug.cpp
@@ -5,12 +5,12 @@
 
 #include "CanvasTypes.h"
 #include "CommonRenderResources.h"
-#include "RenderGraphBuilder.h"
-#include "RHI.h"
-#include "ScreenRendering.h"
 #include "Engine/Font.h"
 #include "Modules/ModuleManager.h"
-#include "Runtime/Renderer/Private/ScreenPass.h"
+#include "RHI.h"
+#include "RenderGraphBuilder.h"
+#include "ScreenPass.h"
+#include "ScreenRendering.h"
 
 int32 GNiagaraGpuComputeDebug_MinTextureHeight = 128;
 static FAutoConsoleVariableRef CVarNiagaraGpuComputeDebug_MinTextureHeight(
diff --git a/Engine/Plugins/FX/Niagara/Source/Niagara/Private/NiagaraGpuComputeDispatch.cpp b/Engine/Plugins/FX/Niagara/Source/Niagara/Private/NiagaraGpuComputeDispatch.cpp
--- a/Engine/Plugins/FX/Niagara/Source/Niagara/Private/NiagaraGpuComputeDispatch.cpp
+++ b/Engine/Plugins/FX/Niagara/Source/Niagara/Private/NiagaraGpuComputeDispatch.cpp
@@ -17,19 +17,19 @@
 #endif
 #include "NiagaraGPUProfilerInterface.h"
 #include "NiagaraGpuReadbackManager.h"
+#include "NiagaraRenderViewDataManager.h"
 #include "NiagaraRenderer.h"
-#include "NiagaraSystemGpuComputeProxy.h"
 #include "NiagaraShader.h"
 #include "NiagaraShaderParticleID.h"
 #include "NiagaraSortingGPU.h"
 #include "NiagaraStats.h"
-#include "NiagaraRenderViewDataManager.h"
+#include "NiagaraSystemGpuComputeProxy.h"
 #include "NiagaraWorldManager.h"
-#include "RHI.h"
-#include "Runtime/Renderer/Private/SceneRendering.h"
 #include "PipelineStateCache.h"
+#include "RHI.h"
+#include "ScenePrivate.h"
+#include "SceneRendering.h"
 #include "ShaderParameterUtils.h"
-#include "Renderer/Private/ScenePrivate.h"
 
 DECLARE_CYCLE_STAT(TEXT("GPU Dispatch Setup [RT]"), STAT_NiagaraGPUDispatchSetup_RT, STATGROUP_Niagara);
 DECLARE_CYCLE_STAT(TEXT("GPU Emitter Dispatch [RT]"), STAT_NiagaraGPUSimTick_RT, STATGROUP_Niagara);
diff --git a/Engine/Plugins/FX/Niagara/Source/Niagara/Private/NiagaraGpuComputeDispatch.h b/Engine/Plugins/FX/Niagara/Source/Niagara/Private/NiagaraGpuComputeDispatch.h
--- a/Engine/Plugins/FX/Niagara/Source/Niagara/Private/NiagaraGpuComputeDispatch.h
+++ b/Engine/Plugins/FX/Niagara/Source/Niagara/Private/NiagaraGpuComputeDispatch.h
@@ -10,20 +10,20 @@ the same VectorVM byte code / compute shader code
 #include "CoreMinimal.h"
 #include "FXSystem.h"
 #include "Modules/ModuleManager.h"
-#include "ParticleResources.h"
 #include "NiagaraCommon.h"
 #include "NiagaraComputeExecutionContext.h"
 #include "NiagaraEmitter.h"
-#include "NiagaraGpuComputeDispatchInterface.h"
 #include "NiagaraGPUProfiler.h"
 #include "NiagaraGPUSortInfo.h"
 #include "NiagaraGPUSystemTick.h"
+#include "NiagaraGpuComputeDispatchInterface.h"
 #include "NiagaraParameters.h"
 #include "NiagaraRendererProperties.h"
 #include "NiagaraSystemGpuComputeProxy.h"
-#include "RendererInterface.h"
+#include "ParticleResources.h"
+#include "Particles/ParticleSortingGPU.h"
 #include "RHIResources.h"
-#include "Runtime/Engine/Private/Particles/ParticleSortingGPU.h"
+#include "RendererInterface.h"
 #include "Tickable.h"
 
 class FGPUSortManager;
diff --git a/Engine/Plugins/FX/Niagara/Source/Niagara/Private/NiagaraGpuComputeDispatchInterface.cpp b/Engine/Plugins/FX/Niagara/Source/Niagara/Private/NiagaraGpuComputeDispatchInterface.cpp
--- a/Engine/Plugins/FX/Niagara/Source/Niagara/Private/NiagaraGpuComputeDispatchInterface.cpp
+++ b/Engine/Plugins/FX/Niagara/Source/Niagara/Private/NiagaraGpuComputeDispatchInterface.cpp
@@ -7,10 +7,10 @@
 #include "NiagaraGpuReadbackManager.h"
 
 #include "Engine/World.h"
-#include "Runtime/Renderer/Private/SystemTextures.h"
-#include "SceneInterface.h"
 #include "FXSystem.h"
 #include "RenderGraphBuilder.h"
+#include "SceneInterface.h"
+#include "SystemTextures.h"
 
 FNiagaraGpuComputeDispatchInterface::FNiagaraGpuComputeDispatchInterface(EShaderPlatform InShaderPlatform, ERHIFeatureLevel::Type InFeatureLevel)
 	: ShaderPlatform(InShaderPlatform)
diff --git a/Engine/Plugins/FX/Niagara/Source/Niagara/Private/NiagaraRendererMeshes.cpp b/Engine/Plugins/FX/Niagara/Source/Niagara/Private/NiagaraRendererMeshes.cpp
--- a/Engine/Plugins/FX/Niagara/Source/Niagara/Private/NiagaraRendererMeshes.cpp
+++ b/Engine/Plugins/FX/Niagara/Source/Niagara/Private/NiagaraRendererMeshes.cpp
@@ -1,24 +1,24 @@
 // Copyright Epic Games, Inc. All Rights Reserved.
 
 #include "NiagaraRendererMeshes.h"
-#include "ParticleResources.h"
-#include "NiagaraDataSet.h"
-#include "NiagaraEmitterInstance.h"
-#include "NiagaraStats.h"
-#include "NiagaraComponent.h"
 #include "Async/ParallelFor.h"
 #include "Engine/StaticMesh.h"
-#include "NiagaraSortingGPU.h"
+#include "FXRenderingUtils.h"
+#include "IXRTrackingSystem.h"
+#include "NiagaraComponent.h"
+#include "NiagaraCullProxyComponent.h"
+#include "NiagaraDataSet.h"
+#include "NiagaraEmitterInstance.h"
 #include "NiagaraGPURayTracingTransformsShader.h"
 #include "NiagaraGpuComputeDispatchInterface.h"
+#include "NiagaraSortingGPU.h"
+#include "NiagaraStats.h"
+#include "ParticleResources.h"
+#include "PipelineStateCache.h"
 #include "RayTracingDefinitions.h"
 #include "RayTracingDynamicGeometryCollection.h"
 #include "RayTracingInstance.h"
-#include "Renderer/Private/ScenePrivate.h"
-#include "PipelineStateCache.h"
-#include "NiagaraCullProxyComponent.h"
-#include "IXRTrackingSystem.h"
-#include "FXRenderingUtils.h"
+#include "ScenePrivate.h"
 
 DECLARE_CYCLE_STAT(TEXT("Generate Mesh Vertex Data [GT]"), STAT_NiagaraGenMeshVertexData, STATGROUP_Niagara);
 
diff --git a/Engine/Plugins/FX/Niagara/Source/Niagara/Private/NiagaraRendererSprites.cpp b/Engine/Plugins/FX/Niagara/Source/Niagara/Private/NiagaraRendererSprites.cpp
--- a/Engine/Plugins/FX/Niagara/Source/Niagara/Private/NiagaraRendererSprites.cpp
+++ b/Engine/Plugins/FX/Niagara/Source/Niagara/Private/NiagaraRendererSprites.cpp
@@ -1,20 +1,20 @@
 // Copyright Epic Games, Inc. All Rights Reserved.
 
 #include "NiagaraRendererSprites.h"
-#include "ParticleResources.h"
-#include "NiagaraDataSet.h"
-#include "NiagaraStats.h"
-#include "NiagaraSortingGPU.h"
+#include "Materials/MaterialInstanceDynamic.h"
 #include "NiagaraComponent.h"
+#include "NiagaraCullProxyComponent.h"
 #include "NiagaraCutoutVertexBuffer.h"
+#include "NiagaraDataSet.h"
 #include "NiagaraGpuComputeDispatchInterface.h"
 #include "NiagaraSettings.h"
+#include "NiagaraSortingGPU.h"
+#include "NiagaraStats.h"
+#include "ParticleResources.h"
 #include "RayTracingDefinitions.h"
 #include "RayTracingDynamicGeometryCollection.h"
 #include "RayTracingInstance.h"
-#include "Materials/MaterialInstanceDynamic.h"
-#include "Renderer/Private/ScenePrivate.h"
-#include "NiagaraCullProxyComponent.h"
+#include "ScenePrivate.h"
 
 DECLARE_DWORD_COUNTER_STAT(TEXT("NumSprites"), STAT_NiagaraNumSprites, STATGROUP_Niagara);
 
diff --git a/Engine/Plugins/FX/Niagara/Source/NiagaraEditor/NiagaraEditor.Build.cs b/Engine/Plugins/FX/Niagara/Source/NiagaraEditor/NiagaraEditor.Build.cs
--- a/Engine/Plugins/FX/Niagara/Source/NiagaraEditor/NiagaraEditor.Build.cs
+++ b/Engine/Plugins/FX/Niagara/Source/NiagaraEditor/NiagaraEditor.Build.cs
@@ -14,7 +14,10 @@ public class NiagaraEditor : ModuleRules
 			"NiagaraEditor/Private/Sequencer/NiagaraSequence",
 			"NiagaraEditor/Private/ViewModels",
 			"NiagaraEditor/Private/TypeEditorUtilities",
-			Path.Combine(Path.GetFullPath(Target.RelativeEnginePath), "Source/Runtime/Renderer/Private")
+			Path.Combine(GetModuleDirectory("GraphEditor"), "Private"),
+			Path.Combine(GetModuleDirectory("Niagara"), "Private"),
+			Path.Combine(GetModuleDirectory("PropertyEditor"), "Private"),
+			Path.Combine(GetModuleDirectory("Renderer"), "Private"),
 		});
 
 		PrivateDependencyModuleNames.AddRange(
@@ -74,12 +77,14 @@ public class NiagaraEditor : ModuleRules
 		PrivateIncludePathModuleNames.AddRange(
 			new string[] {
 				"Engine",
+				"MessageLog",
 				"Messaging",
 				"AssetTools",
 				"ContentBrowser",
                 "DerivedDataCache",
                 "LevelEditor",
-            }
+				"WorkspaceMenuStructure"
+			}
         );
 
 		PublicDependencyModuleNames.AddRange(
diff --git a/Engine/Plugins/FX/Niagara/Source/NiagaraEditor/Private/Customizations/NiagaraPlatformSetCustomization.cpp b/Engine/Plugins/FX/Niagara/Source/NiagaraEditor/Private/Customizations/NiagaraPlatformSetCustomization.cpp
--- a/Engine/Plugins/FX/Niagara/Source/NiagaraEditor/Private/Customizations/NiagaraPlatformSetCustomization.cpp
+++ b/Engine/Plugins/FX/Niagara/Source/NiagaraEditor/Private/Customizations/NiagaraPlatformSetCustomization.cpp
@@ -2,29 +2,29 @@
 
 #include "NiagaraPlatformSetCustomization.h"
 #include "DetailWidgetRow.h"
+#include "DeviceProfiles/DeviceProfile.h"
+#include "DeviceProfiles/DeviceProfileManager.h"
 #include "IDetailChildrenBuilder.h"
+#include "IDetailPropertyRow.h"
+#include "IPropertyUtilities.h"
+#include "Layout/Visibility.h"
 #include "NiagaraEditorStyle.h"
 #include "NiagaraEmitter.h"
 #include "NiagaraPlatformSet.h"
+#include "NiagaraSettings.h"
 #include "NiagaraSystem.h"
 #include "NiagaraTypes.h"
 #include "PlatformInfo.h"
 #include "PropertyHandle.h"
-#include "DeviceProfiles/DeviceProfile.h"
-#include "DeviceProfiles/DeviceProfileManager.h"
-#include "Layout/Visibility.h"
+#include "PropertyNode.h"
+#include "StructurePropertyNode.h"
 #include "Widgets/Images/SImage.h"
 #include "Widgets/Input/SButton.h"
 #include "Widgets/Input/SCheckBox.h"
 #include "Widgets/Input/SMenuAnchor.h"
 #include "Widgets/Layout/SWrapBox.h"
-#include "Widgets/Text/STextBlock.h"
 #include "Widgets/SToolTip.h"
-#include "NiagaraSettings.h"
-#include "IDetailPropertyRow.h"
-#include "IPropertyUtilities.h"
-#include "PropertyEditor/Private/PropertyNode.h"
-#include "PropertyEditor/Private/StructurePropertyNode.h"
+#include "Widgets/Text/STextBlock.h"
 
 
 #define LOCTEXT_NAMESPACE "FNiagaraPlatformSetCustomization"
diff --git a/Engine/Plugins/FX/Niagara/Source/NiagaraEditor/Private/GlobalDistanceFieldCapture.cpp b/Engine/Plugins/FX/Niagara/Source/NiagaraEditor/Private/GlobalDistanceFieldCapture.cpp
--- a/Engine/Plugins/FX/Niagara/Source/NiagaraEditor/Private/GlobalDistanceFieldCapture.cpp
+++ b/Engine/Plugins/FX/Niagara/Source/NiagaraEditor/Private/GlobalDistanceFieldCapture.cpp
@@ -1,15 +1,15 @@
 // Copyright Epic Games, Inc. All Rights Reserved.
 
 #include "GlobalDistanceFieldCapture.h"
+#include "ContentBrowserModule.h"
 #include "Editor.h"
-#include "Stats/Stats2.h"
-#include "Renderer/Private/GlobalDistanceField.h"
-#include "HAL/IConsoleManager.h"
 #include "Engine/VolumeTexture.h"
-#include "ContentBrowser/Public/ContentBrowserModule.h"
-#include "ContentBrowser/Public/IContentBrowserSingleton.h"
+#include "GlobalDistanceField.h"
+#include "HAL/IConsoleManager.h"
+#include "IContentBrowserSingleton.h"
 #include "LevelEditorViewport.h"
 #include "NiagaraEditorModule.h"
+#include "Stats/Stats2.h"
 
 
 
@@ -134,4 +134,4 @@ void FGlobalDistanceFieldCapture::OnReadbackComplete()
 
 	// Work is done, we can kill the singleton instance now
 	delete this;
-}
\ No newline at end of file
+}
diff --git a/Engine/Plugins/FX/Niagara/Source/NiagaraEditor/Private/NiagaraBakerRenderer.cpp b/Engine/Plugins/FX/Niagara/Source/NiagaraEditor/Private/NiagaraBakerRenderer.cpp
--- a/Engine/Plugins/FX/Niagara/Source/NiagaraEditor/Private/NiagaraBakerRenderer.cpp
+++ b/Engine/Plugins/FX/Niagara/Source/NiagaraEditor/Private/NiagaraBakerRenderer.cpp
@@ -9,8 +9,8 @@
 #include "NiagaraGpuComputeDispatchInterface.h"
 #include "NiagaraBatchedElements.h"
 
-#include "Niagara/Classes/NiagaraDataInterfaceRenderTarget2D.h"
-#include "Niagara/Classes/NiagaraDataInterfaceGrid2DCollection.h"
+#include "NiagaraDataInterfaceGrid2DCollection.h"
+#include "NiagaraDataInterfaceRenderTarget2D.h"
 
 #include "Components/SceneCaptureComponent2D.h"
 #include "Engine/Canvas.h"
@@ -748,4 +748,4 @@ bool FNiagaraBakerRenderer::ExportVolume(FStringView FilePath, FIntVector ImageS
 	{
 		return ExportImage(FilePath, FIntPoint(ImageSize.X, ImageSize.Y * ImageSize.Z), ImageData);
 	}
-}
\ No newline at end of file
+}
diff --git a/Engine/Plugins/FX/Niagara/Source/NiagaraEditor/Private/NiagaraCompiler.cpp b/Engine/Plugins/FX/Niagara/Source/NiagaraEditor/Private/NiagaraCompiler.cpp
--- a/Engine/Plugins/FX/Niagara/Source/NiagaraEditor/Private/NiagaraCompiler.cpp
+++ b/Engine/Plugins/FX/Niagara/Source/NiagaraEditor/Private/NiagaraCompiler.cpp
@@ -1,34 +1,34 @@
 // Copyright Epic Games, Inc. All Rights Reserved.
 
 #include "NiagaraCompiler.h"
-#include "NiagaraHlslTranslator.h"
-#include "NiagaraScriptSource.h"
-#include "NiagaraGraph.h"
-#include "NiagaraEditorModule.h"
+#include "EdGraphSchema_Niagara.h"
+#include "EdGraphUtilities.h"
+#include "INiagaraEditorTypeUtilities.h"
+#include "Misc/FileHelper.h"
+#include "Modules/ModuleManager.h"
 #include "NiagaraComponent.h"
-#include "ShaderFormatVectorVM.h"
-#include "NiagaraSystem.h"
-#include "NiagaraNodeEmitter.h"
-#include "NiagaraNodeInput.h"
-#include "NiagaraFunctionLibrary.h"
 #include "NiagaraDataInterface.h"
-#include "NiagaraNodeFunctionCall.h"
-#include "INiagaraEditorTypeUtilities.h"
+#include "NiagaraEditorModule.h"
+#include "NiagaraEditorModule.h"
 #include "NiagaraEditorUtilities.h"
+#include "NiagaraFunctionLibrary.h"
+#include "NiagaraGraph.h"
+#include "NiagaraHlslTranslator.h"
+#include "NiagaraNodeEmitter.h"
+#include "NiagaraNodeFunctionCall.h"
+#include "NiagaraNodeInput.h"
 #include "NiagaraNodeOutput.h"
-#include "ShaderCore.h"
-#include "EdGraphSchema_Niagara.h"
-#include "EdGraphUtilities.h"
-#include "Misc/FileHelper.h"
-#include "ShaderCompiler.h"
-#include "NiagaraShader.h"
+#include "NiagaraPrecompileContainer.h"
 #include "NiagaraScript.h"
+#include "NiagaraScriptSource.h"
+#include "NiagaraShader.h"
 #include "NiagaraSimulationStageBase.h"
+#include "NiagaraSystem.h"
 #include "NiagaraTrace.h"
 #include "Serialization/MemoryReader.h"
-#include "../../Niagara/Private/NiagaraPrecompileContainer.h"
-#include "Modules/ModuleManager.h"
-#include "NiagaraEditorModule.h"
+#include "ShaderCompiler.h"
+#include "ShaderCore.h"
+#include "ShaderFormatVectorVM.h"
 
 #define LOCTEXT_NAMESPACE "NiagaraCompiler"
 
diff --git a/Engine/Plugins/FX/Niagara/Source/NiagaraEditor/Private/NiagaraDebugger.h b/Engine/Plugins/FX/Niagara/Source/NiagaraEditor/Private/NiagaraDebugger.h
--- a/Engine/Plugins/FX/Niagara/Source/NiagaraEditor/Private/NiagaraDebugger.h
+++ b/Engine/Plugins/FX/Niagara/Source/NiagaraEditor/Private/NiagaraDebugger.h
@@ -2,13 +2,13 @@
 
 #pragma once
 
+#include "IMessageContext.h"
+#include "NiagaraDebugHud.h"
+#include "NiagaraDebugHud.h"
 #include "NiagaraDebuggerCommon.h"
-#include "Niagara/Private/NiagaraDebugHud.h"
-#include "NiagaraWorldManager.h"
 #include "NiagaraDebuggerCommon.h"
-#include "Niagara/Private/NiagaraDebugHud.h"
-#include "IMessageContext.h"
 #include "NiagaraOutliner.h"
+#include "NiagaraWorldManager.h"
 
 class ISessionManager;
 class FMessageEndpoint;
@@ -131,4 +131,4 @@ void FNiagaraDebugger::ForAllConnectedClients(TAction Func)
 	}
 }
 
-#endif//WITH_NIAGARA_DEBUGGER
\ No newline at end of file
+#endif//WITH_NIAGARA_DEBUGGER
diff --git a/Engine/Plugins/FX/Niagara/Source/NiagaraEditor/Private/NiagaraEditorModule.cpp b/Engine/Plugins/FX/Niagara/Source/NiagaraEditor/Private/NiagaraEditorModule.cpp
--- a/Engine/Plugins/FX/Niagara/Source/NiagaraEditor/Private/NiagaraEditorModule.cpp
+++ b/Engine/Plugins/FX/Niagara/Source/NiagaraEditor/Private/NiagaraEditorModule.cpp
@@ -23,21 +23,21 @@
 #include "AssetTypeActions/AssetTypeActions_NiagaraParameterDefinitions.h"
 #include "AssetTypeActions/AssetTypeActions_NiagaraSimCache.h"
 
+#include "EdGraphSchema_Niagara.h"
 #include "EdGraphUtilities.h"
-#include "SGraphPin.h"
-#include "KismetPins/SGraphPinVector4.h"
-#include "KismetPins/SGraphPinNum.h"
+#include "KismetPins/SGraphPinBool.h"
+#include "KismetPins/SGraphPinColor.h"
+#include "KismetPins/SGraphPinEnum.h"
 #include "KismetPins/SGraphPinExec.h"
 #include "KismetPins/SGraphPinInteger.h"
+#include "KismetPins/SGraphPinNum.h"
 #include "KismetPins/SGraphPinVector.h"
 #include "KismetPins/SGraphPinVector2D.h"
-#include "KismetPins/SGraphPinColor.h"
-#include "KismetPins/SGraphPinBool.h"
-#include "Editor/GraphEditor/Private/KismetPins/SGraphPinEnum.h"
-#include "SNiagaraGraphPinNumeric.h"
-#include "SNiagaraGraphPinAdd.h"
+#include "KismetPins/SGraphPinVector4.h"
 #include "NiagaraNodeAssignment.h"
-#include "EdGraphSchema_Niagara.h"
+#include "SGraphPin.h"
+#include "SNiagaraGraphPinAdd.h"
+#include "SNiagaraGraphPinNumeric.h"
 
 #include "TypeEditorUtilities/NiagaraFloatTypeEditorUtilities.h"
 #include "TypeEditorUtilities/NiagaraIntegerTypeEditorUtilities.h"
diff --git a/Engine/Plugins/FX/Niagara/Source/NiagaraEditor/Private/SNewEmitterDialog.cpp b/Engine/Plugins/FX/Niagara/Source/NiagaraEditor/Private/SNewEmitterDialog.cpp
--- a/Engine/Plugins/FX/Niagara/Source/NiagaraEditor/Private/SNewEmitterDialog.cpp
+++ b/Engine/Plugins/FX/Niagara/Source/NiagaraEditor/Private/SNewEmitterDialog.cpp
@@ -1,9 +1,9 @@
 // Copyright Epic Games, Inc. All Rights Reserved.
 
 #include "SNewEmitterDialog.h"
-#include "NiagaraEmitter.h"
 #include "NiagaraEditorStyle.h"
-#include "NiagaraEditor/Private/SNiagaraAssetPickerList.h"
+#include "NiagaraEmitter.h"
+#include "SNiagaraAssetPickerList.h"
 #include "SNiagaraNewAssetDialog.h"
 
 #include "AssetRegistry/AssetData.h"
diff --git a/Engine/Plugins/FX/Niagara/Source/NiagaraEditor/Private/SNewSystemDialog.cpp b/Engine/Plugins/FX/Niagara/Source/NiagaraEditor/Private/SNewSystemDialog.cpp
--- a/Engine/Plugins/FX/Niagara/Source/NiagaraEditor/Private/SNewSystemDialog.cpp
+++ b/Engine/Plugins/FX/Niagara/Source/NiagaraEditor/Private/SNewSystemDialog.cpp
@@ -1,10 +1,10 @@
 // Copyright Epic Games, Inc. All Rights Reserved.
 
 #include "SNewSystemDialog.h"
-#include "NiagaraSystem.h"
 #include "NiagaraEditorModule.h"
 #include "NiagaraEditorStyle.h"
-#include "NiagaraEditor/Private/SNiagaraAssetPickerList.h"
+#include "NiagaraSystem.h"
+#include "SNiagaraAssetPickerList.h"
 
 #include "AssetRegistry/AssetData.h"
 
diff --git a/Engine/Plugins/FX/Niagara/Source/NiagaraEditor/Private/SNiagaraNewAssetDialog.cpp b/Engine/Plugins/FX/Niagara/Source/NiagaraEditor/Private/SNiagaraNewAssetDialog.cpp
--- a/Engine/Plugins/FX/Niagara/Source/NiagaraEditor/Private/SNiagaraNewAssetDialog.cpp
+++ b/Engine/Plugins/FX/Niagara/Source/NiagaraEditor/Private/SNiagaraNewAssetDialog.cpp
@@ -1,10 +1,10 @@
 // Copyright Epic Games, Inc. All Rights Reserved.
 
 #include "SNiagaraNewAssetDialog.h"
-#include "NiagaraEmitter.h"
-#include "NiagaraEditorStyle.h"
-#include "NiagaraEditor/Private/SNiagaraAssetPickerList.h"
 #include "NiagaraEditorSettings.h"
+#include "NiagaraEditorStyle.h"
+#include "NiagaraEmitter.h"
+#include "SNiagaraAssetPickerList.h"
 
 #include "AssetRegistry/AssetData.h"
 
@@ -289,4 +289,4 @@ void SNiagaraNewAssetDialog::SaveConfig()
 	GetMutableDefault<UNiagaraEditorSettings>()->SetNewAssetDialogConfig(SaveConfigKey, Config);
 }
 
-#undef LOCTEXT_NAMESPACE
\ No newline at end of file
+#undef LOCTEXT_NAMESPACE
diff --git a/Engine/Plugins/FX/Niagara/Source/NiagaraEditor/Private/Toolkits/NiagaraScriptToolkit.cpp b/Engine/Plugins/FX/Niagara/Source/NiagaraEditor/Private/Toolkits/NiagaraScriptToolkit.cpp
--- a/Engine/Plugins/FX/Niagara/Source/NiagaraEditor/Private/Toolkits/NiagaraScriptToolkit.cpp
+++ b/Engine/Plugins/FX/Niagara/Source/NiagaraEditor/Private/Toolkits/NiagaraScriptToolkit.cpp
@@ -2,9 +2,19 @@
 
 #include "NiagaraScriptToolkit.h"
 
+#include "AssetTypeActions/AssetTypeActions_NiagaraScript.h"
 #include "BusyCursor.h"
 #include "DetailLayoutBuilder.h"
 #include "Editor.h"
+#include "Engine/Selection.h"
+#include "Framework/MultiBox/MultiBoxBuilder.h"
+#include "IMessageLogListing.h"
+#include "MessageLogInitializationOptions.h"
+#include "MessageLogModule.h"
+#include "Misc/FeedbackContext.h"
+#include "Misc/MessageDialog.h"
+#include "Misc/TransactionObjectEvent.h"
+#include "Modules/ModuleManager.h"
 #include "NiagaraEditorCommands.h"
 #include "NiagaraEditorModule.h"
 #include "NiagaraEditorStyle.h"
@@ -18,31 +28,21 @@
 #include "NiagaraScript.h"
 #include "NiagaraScriptGraphViewModel.h"
 #include "NiagaraScriptInputCollectionViewModel.h"
-#include "Widgets/SNiagaraParameterDefinitionsPanel.h"
 #include "NiagaraScriptSource.h"
 #include "NiagaraStandaloneScriptViewModel.h"
 #include "NiagaraVersionMetaData.h"
 #include "PropertyEditorModule.h"
 #include "SGraphActionMenu.h"
 #include "SNiagaraScriptVersionWidget.h"
-#include "AssetTypeActions/AssetTypeActions_NiagaraScript.h"
-#include "Developer/MessageLog/Public/IMessageLogListing.h"
-#include "Developer/MessageLog/Public/MessageLogInitializationOptions.h"
-#include "Developer/MessageLog/Public/MessageLogModule.h"
-#include "Engine/Selection.h"
-#include "Framework/MultiBox/MultiBoxBuilder.h"
-#include "Misc/FeedbackContext.h"
-#include "Misc/MessageDialog.h"
-#include "Misc/TransactionObjectEvent.h"
-#include "Modules/ModuleManager.h"
 #include "UObject/Package.h"
+#include "ViewModels/NiagaraParameterDefinitionsPanelViewModel.h"
 #include "ViewModels/NiagaraParameterPanelViewModel.h"
 #include "ViewModels/NiagaraScriptViewModel.h"
+#include "Widgets/Docking/SDockTab.h"
+#include "Widgets/SNiagaraParameterDefinitionsPanel.h"
 #include "Widgets/SNiagaraParameterPanel.h"
 #include "Widgets/SNiagaraScriptGraph.h"
 #include "Widgets/SNiagaraSelectedObjectsDetails.h"
-#include "Widgets/Docking/SDockTab.h"
-#include "ViewModels/NiagaraParameterDefinitionsPanelViewModel.h"
 
 #define LOCTEXT_NAMESPACE "NiagaraScriptToolkit"
 
diff --git a/Engine/Plugins/FX/Niagara/Source/NiagaraEditor/Private/Toolkits/NiagaraSystemToolkit.cpp b/Engine/Plugins/FX/Niagara/Source/NiagaraEditor/Private/Toolkits/NiagaraSystemToolkit.cpp
--- a/Engine/Plugins/FX/Niagara/Source/NiagaraEditor/Private/Toolkits/NiagaraSystemToolkit.cpp
+++ b/Engine/Plugins/FX/Niagara/Source/NiagaraEditor/Private/Toolkits/NiagaraSystemToolkit.cpp
@@ -37,30 +37,30 @@
 #include "Toolkits/AssetEditorToolkit.h"
 #include "ScopedTransaction.h"
 
-#include "Framework/Application/SlateApplication.h"
-#include "Widgets/SBoxPanel.h"
-#include "Widgets/Layout/SBox.h"
-#include "Widgets/Docking/SDockTab.h"
-#include "Widgets/SNiagaraEmitterVersionWidget.h"
 #include "BusyCursor.h"
-#include "Misc/FeedbackContext.h"
-#include "Misc/TransactionObjectEvent.h"
+#include "Dialogs/Dialogs.h"
 #include "Editor.h"
+#include "Engine/Selection.h"
+#include "Framework/Application/SlateApplication.h"
 #include "LevelEditor.h"
 #include "LevelEditorActions.h"
-#include "Engine/Selection.h"
+#include "Misc/FeedbackContext.h"
+#include "Misc/FileHelper.h"
 #include "Misc/MessageDialog.h"
+#include "Misc/TransactionObjectEvent.h"
 #include "Modules/ModuleManager.h"
-#include "Misc/FileHelper.h"
 #include "NiagaraMessageLogViewModel.h"
 #include "NiagaraVersionMetaData.h"
-#include "Dialogs/Dialogs.h"
-#include "ViewModels/NiagaraParameterPanelViewModel.h"
-#include "NiagaraEditor/Private/SNiagaraAssetPickerList.h"
+#include "SNiagaraAssetPickerList.h"
 #include "SystemToolkitModes/NiagaraSystemToolkitModeBase.h"
 #include "SystemToolkitModes/NiagaraSystemToolkitMode_Default.h"
 #include "SystemToolkitModes/NiagaraSystemToolkitMode_Scalability.h"
 #include "ViewModels/NiagaraParameterDefinitionsPanelViewModel.h"
+#include "ViewModels/NiagaraParameterPanelViewModel.h"
+#include "Widgets/Docking/SDockTab.h"
+#include "Widgets/Layout/SBox.h"
+#include "Widgets/SBoxPanel.h"
+#include "Widgets/SNiagaraEmitterVersionWidget.h"
 
 #define LOCTEXT_NAMESPACE "NiagaraSystemEditor"
 
diff --git a/Engine/Plugins/FX/Niagara/Source/NiagaraEditor/Private/TypeEditorUtilities/NiagaraEnumTypeEditorUtilities.cpp b/Engine/Plugins/FX/Niagara/Source/NiagaraEditor/Private/TypeEditorUtilities/NiagaraEnumTypeEditorUtilities.cpp
--- a/Engine/Plugins/FX/Niagara/Source/NiagaraEditor/Private/TypeEditorUtilities/NiagaraEnumTypeEditorUtilities.cpp
+++ b/Engine/Plugins/FX/Niagara/Source/NiagaraEditor/Private/TypeEditorUtilities/NiagaraEnumTypeEditorUtilities.cpp
@@ -6,9 +6,9 @@
 #include "NiagaraEditorStyle.h"
 #include "NiagaraEditorCommon.h"
 
-#include "Widgets/Input/SSpinBox.h"
-#include "EditorWidgets/Public/SEnumCombo.h"
+#include "SEnumCombo.h"
 #include "Styling/AppStyle.h"
+#include "Widgets/Input/SSpinBox.h"
 
 class SNiagaraEnumParameterEditor : public SNiagaraParameterEditor
 {
diff --git a/Engine/Plugins/FX/Niagara/Source/NiagaraEditor/Private/ViewModels/NiagaraBakerViewModel.cpp b/Engine/Plugins/FX/Niagara/Source/NiagaraEditor/Private/ViewModels/NiagaraBakerViewModel.cpp
--- a/Engine/Plugins/FX/Niagara/Source/NiagaraEditor/Private/ViewModels/NiagaraBakerViewModel.cpp
+++ b/Engine/Plugins/FX/Niagara/Source/NiagaraEditor/Private/ViewModels/NiagaraBakerViewModel.cpp
@@ -10,7 +10,7 @@
 #include "NiagaraSettings.h"
 #include "NiagaraSystem.h"
 
-#include "Niagara/Classes/NiagaraDataInterfaceRenderTarget2D.h"
+#include "NiagaraDataInterfaceRenderTarget2D.h"
 
 #include "Misc/ScopedSlowTask.h"
 #include "Modules/ModuleManager.h"
diff --git a/Engine/Plugins/FX/Niagara/Source/NiagaraEditor/Private/ViewModels/NiagaraMessageLogViewModel.cpp b/Engine/Plugins/FX/Niagara/Source/NiagaraEditor/Private/ViewModels/NiagaraMessageLogViewModel.cpp
--- a/Engine/Plugins/FX/Niagara/Source/NiagaraEditor/Private/ViewModels/NiagaraMessageLogViewModel.cpp
+++ b/Engine/Plugins/FX/Niagara/Source/NiagaraEditor/Private/ViewModels/NiagaraMessageLogViewModel.cpp
@@ -1,10 +1,10 @@
 // Copyright Epic Games, Inc. All Rights Reserved.
 
 #include "NiagaraMessageLogViewModel.h"
-#include "NiagaraScriptSourceBase.h"
+#include "MessageLogModule.h"
 #include "Modules/ModuleManager.h"
-#include "MessageLog/Public/MessageLogModule.h"
 #include "NiagaraMessages.h"
+#include "NiagaraScriptSourceBase.h"
 
 FNiagaraMessageLogViewModel::FNiagaraMessageLogViewModel(const FName& InMessageLogName, const FGuid& InMessageLogGuidKey, TSharedPtr<class SWidget>& OutMessageLogWidget)
 	: MessageLogGuidKey(InMessageLogGuidKey)
diff --git a/Engine/Plugins/FX/Niagara/Source/NiagaraEditor/Private/ViewModels/NiagaraMessageLogViewModel.h b/Engine/Plugins/FX/Niagara/Source/NiagaraEditor/Private/ViewModels/NiagaraMessageLogViewModel.h
--- a/Engine/Plugins/FX/Niagara/Source/NiagaraEditor/Private/ViewModels/NiagaraMessageLogViewModel.h
+++ b/Engine/Plugins/FX/Niagara/Source/NiagaraEditor/Private/ViewModels/NiagaraMessageLogViewModel.h
@@ -3,10 +3,10 @@
 #pragma once
 
 #include "CoreMinimal.h"
+#include "IMessageLogListing.h"
 #include "Logging/TokenizedMessage.h"
-#include "Templates/SharedPointer.h"
-#include "Developer/MessageLog/Public/IMessageLogListing.h"
 #include "NiagaraMessageManager.h"
+#include "Templates/SharedPointer.h"
 
 class FNiagaraMessageLogViewModel : public TSharedFromThis<FNiagaraMessageLogViewModel>
 {
@@ -23,4 +23,4 @@ private:
 	TSharedPtr<IMessageLogListing> MessageLogListing;
 	FGuid MessageLogGuidKey;
 	FGuid MessageManagerRegistrationKey;
-};
\ No newline at end of file
+};
diff --git a/Engine/Plugins/FX/Niagara/Source/NiagaraEditor/Private/Widgets/SNiagaraDebugger.cpp b/Engine/Plugins/FX/Niagara/Source/NiagaraEditor/Private/Widgets/SNiagaraDebugger.cpp
--- a/Engine/Plugins/FX/Niagara/Source/NiagaraEditor/Private/Widgets/SNiagaraDebugger.cpp
+++ b/Engine/Plugins/FX/Niagara/Source/NiagaraEditor/Private/Widgets/SNiagaraDebugger.cpp
@@ -1,35 +1,35 @@
 // Copyright Epic Games, Inc. All Rights Reserved.
 
 #include "SNiagaraDebugger.h"
-#include "SNiagaraDebuggerSpawn.h"
+#include "Customizations/NiagaraDebugHUDCustomization.h"
+#include "Customizations/NiagaraOutlinerCustomization.h"
+#include "EditorWidgetsModule.h"
+#include "Framework/Docking/LayoutService.h"
+#include "Framework/Docking/TabManager.h"
+#include "Framework/MultiBox/MultiBoxBuilder.h"
+#include "IStructureDetailsView.h"
+#include "Modules/ModuleManager.h"
+#include "Modules/ModuleManager.h"
+#include "NiagaraComponent.h"
+#include "NiagaraConstants.h"
+#include "NiagaraEditorModule.h"
 #include "NiagaraEditorStyle.h"
+#include "PropertyEditorModule.h"
+#include "SNiagaraDebuggerSpawn.h"
 #include "Styling/AppStyle.h"
-#include "Editor/WorkspaceMenuStructure/Public/WorkspaceMenuStructure.h"
-#include "Editor/WorkspaceMenuStructure/Public/WorkspaceMenuStructureModule.h"
-#include "Modules/ModuleManager.h"
-#include "Widgets/Images/SImage.h"
-#include "Widgets/SBoxPanel.h"
-#include "Widgets/SOverlay.h"
+#include "UObject/UObjectIterator.h"
 #include "Widgets/Docking/SDockTab.h"
+#include "Widgets/Images/SImage.h"
 #include "Widgets/Input/SButton.h"
 #include "Widgets/Input/SNumericDropDown.h"
 #include "Widgets/Input/SNumericEntryBox.h"
 #include "Widgets/Layout/SScrollBox.h"
 #include "Widgets/Layout/SUniformGridPanel.h"
+#include "Widgets/SBoxPanel.h"
+#include "Widgets/SOverlay.h"
 #include "Widgets/Text/STextBlock.h"
-#include "Framework/Docking/LayoutService.h"
-#include "Framework/Docking/TabManager.h"
-#include "Framework/MultiBox/MultiBoxBuilder.h"
-#include "UObject/UObjectIterator.h"
-#include "EditorWidgetsModule.h"
-#include "PropertyEditorModule.h"
-#include "Modules/ModuleManager.h"
-#include "NiagaraEditorModule.h"
-#include "Customizations/NiagaraDebugHUDCustomization.h"
-#include "Customizations/NiagaraOutlinerCustomization.h"
-#include "IStructureDetailsView.h"
-#include "NiagaraComponent.h"
-#include "NiagaraConstants.h"
+#include "WorkspaceMenuStructure.h"
+#include "WorkspaceMenuStructureModule.h"
 
 // the SessionFrontend is a "target" developer tool for talking to other devices, etc, and may get disabled
 #define WITH_SESSION_FRONTEND 	WITH_UNREAL_TARGET_DEVELOPER_TOOLS
@@ -1359,4 +1359,4 @@ TSharedRef<SWidget> SNiagaraDebugger::MakePlaybackOptionsMenu()
 
 #undef LOCTEXT_NAMESPACE
 
-#endif // WITH_NIAGARA_DEBUGGER
\ No newline at end of file
+#endif // WITH_NIAGARA_DEBUGGER
diff --git a/Engine/Plugins/FX/Niagara/Source/NiagaraEditor/Private/Widgets/SNiagaraEmitterVersionWidget.cpp b/Engine/Plugins/FX/Niagara/Source/NiagaraEditor/Private/Widgets/SNiagaraEmitterVersionWidget.cpp
--- a/Engine/Plugins/FX/Niagara/Source/NiagaraEditor/Private/Widgets/SNiagaraEmitterVersionWidget.cpp
+++ b/Engine/Plugins/FX/Niagara/Source/NiagaraEditor/Private/Widgets/SNiagaraEmitterVersionWidget.cpp
@@ -3,10 +3,10 @@
 #include "SNiagaraEmitterVersionWidget.h"
 
 #include "AssetToolsModule.h"
-#include "NiagaraEmitter.h"
 #include "Editor.h"
+#include "NiagaraEmitter.h"
+#include "NiagaraVersionMetaData.h"
 #include "Subsystems/AssetEditorSubsystem.h"
-#include "NiagaraEditor/Private/NiagaraVersionMetaData.h"
 
 #define LOCTEXT_NAMESPACE "NiagaraEmitterVersionWidget"
 
diff --git a/Engine/Plugins/FX/Niagara/Source/NiagaraEditor/Private/Widgets/SNiagaraParameterCollection.cpp b/Engine/Plugins/FX/Niagara/Source/NiagaraEditor/Private/Widgets/SNiagaraParameterCollection.cpp
--- a/Engine/Plugins/FX/Niagara/Source/NiagaraEditor/Private/Widgets/SNiagaraParameterCollection.cpp
+++ b/Engine/Plugins/FX/Niagara/Source/NiagaraEditor/Private/Widgets/SNiagaraParameterCollection.cpp
@@ -16,15 +16,15 @@
 #include "Widgets/Layout/SSplitter.h"
 #include "Framework/MultiBox/MultiBoxBuilder.h"
 
-#include "Modules/ModuleManager.h"
-#include "IStructureDetailsView.h"
 #include "Framework/Commands/GenericCommands.h"
-#include "Editor/PropertyEditor/Public/PropertyEditorModule.h"
 #include "IDetailsView.h"
-#include "Widgets/SNullWidget.h"
-#include "Widgets/Input/SSpinBox.h"
+#include "IStructureDetailsView.h"
+#include "Modules/ModuleManager.h"
+#include "PropertyEditorModule.h"
 #include "ScopedTransaction.h"
 #include "Widgets/Input/SCheckBox.h"
+#include "Widgets/Input/SSpinBox.h"
+#include "Widgets/SNullWidget.h"
 
 
 #define LOCTEXT_NAMESPACE "NiagaraParameterCollectionEditor"
diff --git a/Engine/Plugins/FX/Niagara/Source/NiagaraEditor/Private/Widgets/SNiagaraScriptVersionWidget.cpp b/Engine/Plugins/FX/Niagara/Source/NiagaraEditor/Private/Widgets/SNiagaraScriptVersionWidget.cpp
--- a/Engine/Plugins/FX/Niagara/Source/NiagaraEditor/Private/Widgets/SNiagaraScriptVersionWidget.cpp
+++ b/Engine/Plugins/FX/Niagara/Source/NiagaraEditor/Private/Widgets/SNiagaraScriptVersionWidget.cpp
@@ -6,8 +6,8 @@
 #include "Editor.h"
 #include "Editor/EditorEngine.h"
 #include "Modules/ModuleManager.h"
-#include "NiagaraEditor/Private/NiagaraVersionMetaData.h"
-#include "NiagaraEditor/Public/NiagaraActions.h"
+#include "NiagaraActions.h"
+#include "NiagaraVersionMetaData.h"
 #include "Subsystems/AssetEditorSubsystem.h"
 
 #define LOCTEXT_NAMESPACE "NiagaraScriptVersionWidget"
diff --git a/Engine/Plugins/FX/Niagara/Source/NiagaraEditor/Private/Widgets/SNiagaraSelectedObjectsDetails.cpp b/Engine/Plugins/FX/Niagara/Source/NiagaraEditor/Private/Widgets/SNiagaraSelectedObjectsDetails.cpp
--- a/Engine/Plugins/FX/Niagara/Source/NiagaraEditor/Private/Widgets/SNiagaraSelectedObjectsDetails.cpp
+++ b/Engine/Plugins/FX/Niagara/Source/NiagaraEditor/Private/Widgets/SNiagaraSelectedObjectsDetails.cpp
@@ -3,14 +3,14 @@
 #include "Widgets/SNiagaraSelectedObjectsDetails.h"
 #include "NiagaraObjectSelection.h"
 
-#include "Modules/ModuleManager.h"
-#include "Editor/PropertyEditor/Public/PropertyEditorModule.h"
+#include "Customizations/NiagaraTypeCustomizations.h"
 #include "IDetailsView.h"
+#include "Modules/ModuleManager.h"
 #include "NiagaraEditorCommon.h"
 #include "NiagaraParameterDefinitions.h"
 #include "NiagaraScriptVariable.h"
 #include "PropertyEditorDelegates.h"
-#include "Customizations/NiagaraTypeCustomizations.h"
+#include "PropertyEditorModule.h"
 
 #define LOCTEXT_NAMESPACE "NiagaraSelectedObjectsDetails"
 
diff --git a/Engine/Plugins/FX/Niagara/Source/NiagaraEditor/Private/Widgets/SNiagaraVersionWidget.cpp b/Engine/Plugins/FX/Niagara/Source/NiagaraEditor/Private/Widgets/SNiagaraVersionWidget.cpp
--- a/Engine/Plugins/FX/Niagara/Source/NiagaraEditor/Private/Widgets/SNiagaraVersionWidget.cpp
+++ b/Engine/Plugins/FX/Niagara/Source/NiagaraEditor/Private/Widgets/SNiagaraVersionWidget.cpp
@@ -2,19 +2,19 @@
 
 #include "SNiagaraVersionWidget.h"
 
-#include "Editor.h"
 #include "AssetToolsModule.h"
 #include "DetailLayoutBuilder.h"
+#include "Editor.h"
+#include "Framework/MultiBox/MultiBoxBuilder.h"
 #include "IDetailsView.h"
+#include "Modules/ModuleManager.h"
+#include "NiagaraActions.h"
+#include "NiagaraVersionMetaData.h"
 #include "PropertyEditorModule.h"
 #include "SGraphActionMenu.h"
-#include "SlateOptMacros.h"
 #include "ScopedTransaction.h"
-#include "Framework/MultiBox/MultiBoxBuilder.h"
-#include "Modules/ModuleManager.h"
+#include "SlateOptMacros.h"
 #include "Styling/AppStyle.h"
-#include "NiagaraEditor/Private/NiagaraVersionMetaData.h"
-#include "NiagaraEditor/Public/NiagaraActions.h"
 #include "Widgets/Images/SImage.h"
 #include "Widgets/Input/SButton.h"
 #include "Widgets/Layout/SBox.h"
diff --git a/Engine/Plugins/FX/Niagara/Source/NiagaraEditorWidgets/NiagaraEditorWidgets.Build.cs b/Engine/Plugins/FX/Niagara/Source/NiagaraEditorWidgets/NiagaraEditorWidgets.Build.cs
--- a/Engine/Plugins/FX/Niagara/Source/NiagaraEditorWidgets/NiagaraEditorWidgets.Build.cs
+++ b/Engine/Plugins/FX/Niagara/Source/NiagaraEditorWidgets/NiagaraEditorWidgets.Build.cs
@@ -1,5 +1,6 @@
 // Copyright Epic Games, Inc. All Rights Reserved.
 
+using System.IO;
 using UnrealBuildTool;
 
 public class NiagaraEditorWidgets : ModuleRules
@@ -7,7 +8,9 @@ public class NiagaraEditorWidgets : ModuleRules
 	public NiagaraEditorWidgets(ReadOnlyTargetRules Target) : base(Target)
 	{
 		PrivateIncludePaths.AddRange(new string[] {
-			"NiagaraEditorWidgets/Private",
+			System.IO.Path.Combine(GetModuleDirectory("Niagara"), "Private"),
+			System.IO.Path.Combine(GetModuleDirectory("NiagaraEditor"), "Private"),
+			System.IO.Path.Combine(GetModuleDirectory("GraphEditor"), "Private"),
 		});
 
 		PrivateDependencyModuleNames.AddRange(new string[] {
diff --git a/Engine/Plugins/FX/Niagara/Source/NiagaraEditorWidgets/Private/DetailCustomizations/NiagaraDataInterfaceParticleReadDetails.cpp b/Engine/Plugins/FX/Niagara/Source/NiagaraEditorWidgets/Private/DetailCustomizations/NiagaraDataInterfaceParticleReadDetails.cpp
--- a/Engine/Plugins/FX/Niagara/Source/NiagaraEditorWidgets/Private/DetailCustomizations/NiagaraDataInterfaceParticleReadDetails.cpp
+++ b/Engine/Plugins/FX/Niagara/Source/NiagaraEditorWidgets/Private/DetailCustomizations/NiagaraDataInterfaceParticleReadDetails.cpp
@@ -1,8 +1,8 @@
 // Copyright Epic Games, Inc. All Rights Reserved.
 
 #include "NiagaraDataInterfaceParticleReadDetails.h"
-#include "Niagara/Classes/NiagaraDataInterfaceParticleRead.h" 
-#include "Niagara/Classes/NiagaraSystem.h" 
+#include "NiagaraDataInterfaceParticleRead.h" 
+#include "NiagaraSystem.h" 
 
 #include "Widgets/Input/SSuggestionTextBox.h"
 #include "DetailCategoryBuilder.h"
diff --git a/Engine/Plugins/FX/Niagara/Source/NiagaraEditorWidgets/Private/DetailCustomizations/NiagaraDataInterfaceStaticMeshDetails.cpp b/Engine/Plugins/FX/Niagara/Source/NiagaraEditorWidgets/Private/DetailCustomizations/NiagaraDataInterfaceStaticMeshDetails.cpp
--- a/Engine/Plugins/FX/Niagara/Source/NiagaraEditorWidgets/Private/DetailCustomizations/NiagaraDataInterfaceStaticMeshDetails.cpp
+++ b/Engine/Plugins/FX/Niagara/Source/NiagaraEditorWidgets/Private/DetailCustomizations/NiagaraDataInterfaceStaticMeshDetails.cpp
@@ -1,14 +1,14 @@
 // Copyright Epic Games, Inc. All Rights Reserved.
 
 #include "NiagaraDataInterfaceStaticMeshDetails.h"
-#include "NiagaraDetailSourcedArrayBuilder.h"
-#include "NiagaraDataInterfaceDetails.h"
-#include "Niagara/Private/DataInterface/NiagaraDataInterfaceStaticMesh.h" 
-#include "DetailLayoutBuilder.h"
+#include "DataInterface/NiagaraDataInterfaceStaticMesh.h" 
 #include "DetailCategoryBuilder.h"
+#include "DetailLayoutBuilder.h"
+#include "Engine/StaticMeshSocket.h"
 #include "IDetailGroup.h"
 #include "NiagaraComponent.h"
-#include "Engine/StaticMeshSocket.h"
+#include "NiagaraDataInterfaceDetails.h"
+#include "NiagaraDetailSourcedArrayBuilder.h"
 #include "SNiagaraNamePropertySelector.h"
 
 #define LOCTEXT_NAMESPACE "FNiagaraDataInterfaceStaticMeshDetails"
@@ -127,4 +127,4 @@ void FNiagaraDataInterfaceStaticMeshDetails::GenerateSocketsArray(TArray<TShared
 	}
 }
 
-#undef LOCTEXT_NAMESPACE
\ No newline at end of file
+#undef LOCTEXT_NAMESPACE
diff --git a/Engine/Plugins/FX/Niagara/Source/NiagaraEditorWidgets/Private/SNiagaraOverviewGraph.cpp b/Engine/Plugins/FX/Niagara/Source/NiagaraEditorWidgets/Private/SNiagaraOverviewGraph.cpp
--- a/Engine/Plugins/FX/Niagara/Source/NiagaraEditorWidgets/Private/SNiagaraOverviewGraph.cpp
+++ b/Engine/Plugins/FX/Niagara/Source/NiagaraEditorWidgets/Private/SNiagaraOverviewGraph.cpp
@@ -1,32 +1,32 @@
 // Copyright Epic Games, Inc. All Rights Reserved.
 
 #include "SNiagaraOverviewGraph.h"
-#include "NiagaraOverviewGraphNodeFactory.h"
-#include "SNiagaraStack.h"
-#include "ViewModels/NiagaraSystemViewModel.h"
-#include "ViewModels/NiagaraEmitterHandleViewModel.h"
-#include "ViewModels/NiagaraOverviewGraphViewModel.h"
-#include "NiagaraSystem.h"
-#include "NiagaraOverviewNode.h"
-#include "NiagaraSystemEditorData.h"
-#include "NiagaraObjectSelection.h"
+#include "EdGraph/EdGraph.h"
+#include "EdGraph/EdGraphNode.h"
 #include "EdGraphSchema_Niagara.h"
+#include "Framework/Application/SlateApplication.h"
+#include "Framework/Commands/GenericCommands.h"
+#include "Framework/MultiBox/MultiBoxBuilder.h"
 #include "GraphEditAction.h"
+#include "GraphEditorActions.h"
 #include "NiagaraEditorCommands.h"
 #include "NiagaraEditorModule.h"
-#include "GraphEditorActions.h"
-#include "NiagaraEditorWidgetsStyle.h"
 #include "NiagaraEditorSettings.h"
-#include "Framework/MultiBox/MultiBoxBuilder.h"
+#include "NiagaraEditorWidgetsStyle.h"
+#include "NiagaraObjectSelection.h"
+#include "NiagaraOverviewGraphNodeFactory.h"
+#include "NiagaraOverviewNode.h"
+#include "NiagaraSystem.h"
+#include "NiagaraSystemEditorData.h"
+#include "SNiagaraAssetPickerList.h"
+#include "SNiagaraOverviewGraphTitleBar.h"
+#include "SNiagaraStack.h"
 #include "ScopedTransaction.h"
-#include "Framework/Application/SlateApplication.h"
 #include "Templates/SharedPointer.h"
-#include "EdGraph/EdGraph.h"
-#include "EdGraph/EdGraphNode.h"
-#include "Framework/Commands/GenericCommands.h"
-#include "NiagaraEditor/Private/SNiagaraAssetPickerList.h"
+#include "ViewModels/NiagaraEmitterHandleViewModel.h"
+#include "ViewModels/NiagaraOverviewGraphViewModel.h"
+#include "ViewModels/NiagaraSystemViewModel.h"
 #include "Widgets/SItemSelector.h"
-#include "SNiagaraOverviewGraphTitleBar.h"
 
 #define LOCTEXT_NAMESPACE "NiagaraOverviewGraph"
 
diff --git a/Engine/Plugins/FX/Niagara/Source/NiagaraEditorWidgets/Private/SNiagaraOverviewStackNode.cpp b/Engine/Plugins/FX/Niagara/Source/NiagaraEditorWidgets/Private/SNiagaraOverviewStackNode.cpp
--- a/Engine/Plugins/FX/Niagara/Source/NiagaraEditorWidgets/Private/SNiagaraOverviewStackNode.cpp
+++ b/Engine/Plugins/FX/Niagara/Source/NiagaraEditorWidgets/Private/SNiagaraOverviewStackNode.cpp
@@ -3,40 +3,40 @@
 #include "SNiagaraOverviewStackNode.h"
 
 #include "EditorFontGlyphs.h"
+#include "Framework/MultiBox/MultiBoxBuilder.h"
+#include "Materials/Material.h"
+#include "Materials/MaterialInterface.h"
+#include "Modules/ModuleManager.h"
+#include "NiagaraEditorModule.h"
+#include "NiagaraEditorSettings.h"
+#include "NiagaraEditorStyle.h"
+#include "NiagaraEditorUtilities.h"
+#include "NiagaraEditorWidgetsStyle.h"
+#include "NiagaraEmitter.h"
+#include "NiagaraEmitterInstance.h"
 #include "NiagaraOverviewNode.h"
-#include "ViewModels/NiagaraSystemViewModel.h"
-#include "ViewModels/NiagaraEmitterViewModel.h"
+#include "NiagaraRendererProperties.h"
+#include "NiagaraSettings.h"
+#include "SGraphPanel.h"
+#include "SLevelOfDetailBranchNode.h"
+#include "SNiagaraOverviewStack.h"
+#include "Stack/SNiagaraStackIssueIcon.h"
+#include "Styling/AppStyle.h"
+#include "Styling/StyleColors.h"
 #include "ViewModels/NiagaraEmitterHandleViewModel.h"
+#include "ViewModels/NiagaraEmitterViewModel.h"
 #include "ViewModels/NiagaraSystemSelectionViewModel.h"
+#include "ViewModels/NiagaraSystemViewModel.h"
 #include "ViewModels/Stack/NiagaraStackItemGroup.h"
 #include "ViewModels/Stack/NiagaraStackModuleItem.h"
+#include "ViewModels/Stack/NiagaraStackRendererItem.h"
 #include "ViewModels/Stack/NiagaraStackViewModel.h"
-#include "SNiagaraOverviewStack.h"
-#include "NiagaraEditorModule.h"
-#include "NiagaraEditorStyle.h"
-#include "NiagaraEditorWidgetsStyle.h"
-#include "NiagaraEmitter.h"
-#include "Stack/SNiagaraStackIssueIcon.h"
-#include "Modules/ModuleManager.h"
-#include "Widgets/Layout/SBox.h"
 #include "Widgets/Input/SButton.h"
 #include "Widgets/Input/SCheckBox.h"
-#include "NiagaraEmitterInstance.h"
-#include "NiagaraEditorUtilities.h"
-#include "Widgets/SToolTip.h"
+#include "Widgets/Layout/SBox.h"
 #include "Widgets/SBoxPanel.h"
-#include "NiagaraRendererProperties.h"
-#include "NiagaraSettings.h"
-#include "NiagaraEditorSettings.h"
-#include "SGraphPanel.h"
+#include "Widgets/SToolTip.h"
 #include "Widgets/SWidget.h"
-#include "ViewModels/Stack/NiagaraStackRendererItem.h"
-#include "Materials/MaterialInterface.h"
-#include "Materials/Material.h"
-#include "Framework/MultiBox/MultiBoxBuilder.h"
-#include "Styling/StyleColors.h"
-#include "Styling/AppStyle.h"
-#include "KismetWidgets/Public/SLevelOfDetailBranchNode.h"
 
 #define LOCTEXT_NAMESPACE "NiagaraOverviewStackNode"
 
diff --git a/Engine/Plugins/FX/Niagara/Source/NiagaraEditorWidgets/Private/SNiagaraScalabilityPreviewSettings.cpp b/Engine/Plugins/FX/Niagara/Source/NiagaraEditorWidgets/Private/SNiagaraScalabilityPreviewSettings.cpp
--- a/Engine/Plugins/FX/Niagara/Source/NiagaraEditorWidgets/Private/SNiagaraScalabilityPreviewSettings.cpp
+++ b/Engine/Plugins/FX/Niagara/Source/NiagaraEditorWidgets/Private/SNiagaraScalabilityPreviewSettings.cpp
@@ -1,21 +1,21 @@
 // Copyright Epic Games, Inc. All Rights Reserved.
 
 #include "SNiagaraScalabilityPreviewSettings.h"
+#include "Customizations/NiagaraPlatformSetCustomization.h"
+#include "DeviceProfiles/DeviceProfile.h"
+#include "DeviceProfiles/DeviceProfileManager.h"
 #include "IDetailTreeNode.h"
 #include "IPropertyRowGenerator.h"
 #include "ISinglePropertyView.h"
 #include "NiagaraEditorStyle.h"
 #include "NiagaraSettings.h"
 #include "PlatformInfo.h"
-#include "DeviceProfiles/DeviceProfile.h"
-#include "DeviceProfiles/DeviceProfileManager.h"
-#include "Widgets/Layout/SWrapBox.h"
-#include "Widgets/Layout/SSeparator.h"
-#include "NiagaraEditor/Private/Customizations/NiagaraPlatformSetCustomization.h"
 #include "Widgets/Images/SImage.h"
 #include "Widgets/Input/SButton.h"
 #include "Widgets/Input/SCheckBox.h"
+#include "Widgets/Layout/SSeparator.h"
 #include "Widgets/Layout/SSpacer.h"
+#include "Widgets/Layout/SWrapBox.h"
 
 #define LOCTEXT_NAMESPACE "NiagaraScalabilityPreviewSettings"
 
diff --git a/Engine/Plugins/FX/Niagara/Source/NiagaraEditorWidgets/Private/Stack/SNiagaraStackFunctionInputValue.cpp b/Engine/Plugins/FX/Niagara/Source/NiagaraEditorWidgets/Private/Stack/SNiagaraStackFunctionInputValue.cpp
--- a/Engine/Plugins/FX/Niagara/Source/NiagaraEditorWidgets/Private/Stack/SNiagaraStackFunctionInputValue.cpp
+++ b/Engine/Plugins/FX/Niagara/Source/NiagaraEditorWidgets/Private/Stack/SNiagaraStackFunctionInputValue.cpp
@@ -3,29 +3,36 @@
 #include "Stack/SNiagaraStackFunctionInputValue.h"
 #include "AssetRegistry/AssetRegistryModule.h"
 #include "EditorFontGlyphs.h"
+#include "Framework/Application/SlateApplication.h"
+#include "Framework/MultiBox/MultiBoxBuilder.h"
 #include "INiagaraEditorTypeUtilities.h"
 #include "IStructureDetailsView.h"
-#include "NiagaraConstants.h"
+#include "Modules/ModuleManager.h"
 #include "NiagaraActions.h"
+#include "NiagaraCommon.h"
+#include "NiagaraConstants.h"
 #include "NiagaraEditorModule.h"
 #include "NiagaraEditorSettings.h"
 #include "NiagaraEditorStyle.h"
 #include "NiagaraEditorUtilities.h"
 #include "NiagaraEditorWidgetsStyle.h"
 #include "NiagaraEditorWidgetsUtilities.h"
+#include "NiagaraNodeAssignment.h"
 #include "NiagaraNodeCustomHlsl.h"
 #include "NiagaraNodeFunctionCall.h"
-#include "NiagaraScriptVariable.h"
 #include "NiagaraParameterCollection.h"
+#include "NiagaraScriptVariable.h"
+#include "NiagaraScriptVariable.h"
+#include "NiagaraSettings.h"
 #include "NiagaraSystem.h"
-#include "ScopedTransaction.h"
+#include "PropertyEditorModule.h"
 #include "SDropTarget.h"
 #include "SNiagaraGraphActionWidget.h"
+#include "SNiagaraParameterDropTarget.h"
 #include "SNiagaraParameterEditor.h"
-#include "Editor/PropertyEditor/Public/PropertyEditorModule.h"
-#include "Framework/Application/SlateApplication.h"
-#include "Framework/MultiBox/MultiBoxBuilder.h"
-#include "Modules/ModuleManager.h"
+#include "ScopedTransaction.h"
+#include "Stack/SNiagaraStackIndent.h"
+#include "Styling/AppStyle.h"
 #include "Subsystems/AssetEditorSubsystem.h"
 #include "ViewModels/NiagaraEmitterViewModel.h"
 #include "ViewModels/NiagaraScratchPadScriptViewModel.h"
@@ -36,26 +43,19 @@
 #include "ViewModels/Stack/NiagaraStackGraphUtilities.h"
 #include "ViewModels/Stack/NiagaraStackInputCategory.h"
 #include "ViewModels/Stack/NiagaraStackSummaryViewInputCollection.h"
-#include "Widgets/SBoxPanel.h"
-#include "Widgets/Text/SRichTextBlock.h"
-#include "Widgets/SNiagaraParameterName.h"
 #include "Widgets/Images/SImage.h"
 #include "Widgets/Input/SButton.h"
-#include "Widgets/Input/SComboButton.h"
 #include "Widgets/Input/SCheckBox.h"
+#include "Widgets/Input/SComboButton.h"
 #include "Widgets/Input/SEditableTextBox.h"
+#include "Widgets/Input/SNumericEntryBox.h"
+#include "Widgets/Input/SSpinBox.h"
 #include "Widgets/Layout/SBox.h"
 #include "Widgets/NiagaraHLSLSyntaxHighlighter.h"
+#include "Widgets/SBoxPanel.h"
+#include "Widgets/SNiagaraParameterName.h"
 #include "Widgets/Text/SMultiLineEditableText.h"
-#include "Widgets/Input/SSpinBox.h"
-#include "Widgets/Input/SNumericEntryBox.h"
-#include "NiagaraNodeAssignment.h"
-#include "NiagaraCommon.h"
-#include "NiagaraScriptVariable.h"
-#include "NiagaraSettings.h"
-#include "SNiagaraParameterDropTarget.h"
-#include "Stack/SNiagaraStackIndent.h"
-#include "Styling/AppStyle.h"
+#include "Widgets/Text/SRichTextBlock.h"
 
 #define LOCTEXT_NAMESPACE "NiagaraStackFunctionInputValue"
 
diff --git a/Engine/Plugins/FX/Niagara/Source/NiagaraShader/Private/NiagaraAsyncGpuTraceProvider.cpp b/Engine/Plugins/FX/Niagara/Source/NiagaraShader/Private/NiagaraAsyncGpuTraceProvider.cpp
--- a/Engine/Plugins/FX/Niagara/Source/NiagaraShader/Private/NiagaraAsyncGpuTraceProvider.cpp
+++ b/Engine/Plugins/FX/Niagara/Source/NiagaraShader/Private/NiagaraAsyncGpuTraceProvider.cpp
@@ -7,7 +7,7 @@
 #include "NiagaraAsyncGpuTraceProviderHwrt.h"
 #include "NiagaraSettings.h"
 #include "NiagaraShaderParticleID.h"
-#include "Renderer/Private/ScenePrivate.h"
+#include "ScenePrivate.h"
 
 ////////////////////////////////////////////////////////////////////////////////
 ////////////////////////////////////////////////////////////////////////////////
diff --git a/Engine/Plugins/FX/Niagara/Source/NiagaraShader/Private/NiagaraAsyncGpuTraceProviderGsdf.cpp b/Engine/Plugins/FX/Niagara/Source/NiagaraShader/Private/NiagaraAsyncGpuTraceProviderGsdf.cpp
--- a/Engine/Plugins/FX/Niagara/Source/NiagaraShader/Private/NiagaraAsyncGpuTraceProviderGsdf.cpp
+++ b/Engine/Plugins/FX/Niagara/Source/NiagaraShader/Private/NiagaraAsyncGpuTraceProviderGsdf.cpp
@@ -3,10 +3,12 @@
 #include "NiagaraAsyncGpuTraceProviderGsdf.h"
 
 #include "GlobalShader.h"
-#include "NiagaraGpuComputeDispatchInterface.h"
 #include "NiagaraDistanceFieldHelper.h"
+#include "NiagaraGpuComputeDispatchInterface.h"
 #include "NiagaraSettings.h"
 #include "SceneManagement.h"
+#include "ScenePrivate.h"
+#include "SceneRendering.h"
 
 static int GNiagaraAsyncGpuTraceGsdfEnabled = 1;
 static FAutoConsoleVariableRef CVarNiagaraAsyncGpuTraceGsdfEnabled(
@@ -146,4 +148,4 @@ void FNiagaraAsyncGpuTraceProviderGsdf::Reset()
 {
 	m_DistanceFieldData = FGlobalDistanceFieldParameterData();
 	m_ViewUniformBuffer = TUniformBufferRef<FViewUniformShaderParameters>();
-}
\ No newline at end of file
+}
diff --git a/Engine/Plugins/FX/Niagara/Source/NiagaraShader/Private/NiagaraAsyncGpuTraceProviderHwrt.cpp b/Engine/Plugins/FX/Niagara/Source/NiagaraShader/Private/NiagaraAsyncGpuTraceProviderHwrt.cpp
--- a/Engine/Plugins/FX/Niagara/Source/NiagaraShader/Private/NiagaraAsyncGpuTraceProviderHwrt.cpp
+++ b/Engine/Plugins/FX/Niagara/Source/NiagaraShader/Private/NiagaraAsyncGpuTraceProviderHwrt.cpp
@@ -8,7 +8,7 @@
 #include "GlobalShader.h"
 #include "NiagaraRenderer.h"
 #include "NiagaraSettings.h"
-#include "Renderer/Private/ScenePrivate.h"
+#include "ScenePrivate.h"
 
 static int GNiagaraAsyncGpuTraceHwrtEnabled = 1;
 static FAutoConsoleVariableRef CVarNiagaraAsyncGpuTraceHwrtEnabled(
@@ -406,4 +406,4 @@ void FNiagaraAsyncGpuTraceProviderHwrt::Reset()
 	ViewUniformBuffer = TUniformBufferRef<FViewUniformShaderParameters>();
 }
 
-#endif // RHI_RAYTRACING
\ No newline at end of file
+#endif // RHI_RAYTRACING
diff --git a/Engine/Plugins/FX/Niagara/Source/NiagaraShader/Private/NiagaraClearCounts.cpp b/Engine/Plugins/FX/Niagara/Source/NiagaraShader/Private/NiagaraClearCounts.cpp
--- a/Engine/Plugins/FX/Niagara/Source/NiagaraShader/Private/NiagaraClearCounts.cpp
+++ b/Engine/Plugins/FX/Niagara/Source/NiagaraShader/Private/NiagaraClearCounts.cpp
@@ -9,6 +9,7 @@
 
 #include "RenderGraphBuilder.h"
 #include "RenderGraphUtils.h"
+#include "ScreenPass.h"
 
 //////////////////////////////////////////////////////////////////////////
 
diff --git a/Engine/Plugins/FX/Niagara/Source/NiagaraShader/Private/NiagaraDebugShaders.cpp b/Engine/Plugins/FX/Niagara/Source/NiagaraShader/Private/NiagaraDebugShaders.cpp
--- a/Engine/Plugins/FX/Niagara/Source/NiagaraShader/Private/NiagaraDebugShaders.cpp
+++ b/Engine/Plugins/FX/Niagara/Source/NiagaraShader/Private/NiagaraDebugShaders.cpp
@@ -7,11 +7,10 @@
 #include "CommonRenderResources.h"
 #include "Modules/ModuleManager.h"
 
+#include "PipelineStateCache.h"
 #include "RenderGraphBuilder.h"
 #include "RenderGraphUtils.h"
-#include "PipelineStateCache.h"
-
-#include "Runtime/Renderer/Private/ScreenPass.h"
+#include "ScreenPass.h"
 
 int GNiagaraGpuComputeDebug_ShowNaNInf = 1;
 static FAutoConsoleVariableRef CVarNiagaraGpuComputeDebug_ShowNaNInf(
diff --git a/Engine/Plugins/FX/Niagara/Source/NiagaraShader/Private/NiagaraDistanceFieldHelper.cpp b/Engine/Plugins/FX/Niagara/Source/NiagaraShader/Private/NiagaraDistanceFieldHelper.cpp
--- a/Engine/Plugins/FX/Niagara/Source/NiagaraShader/Private/NiagaraDistanceFieldHelper.cpp
+++ b/Engine/Plugins/FX/Niagara/Source/NiagaraShader/Private/NiagaraDistanceFieldHelper.cpp
@@ -1,7 +1,7 @@
 // Copyright Epic Games, Inc. All Rights Reserved.
 
 #include "NiagaraDistanceFieldHelper.h"
-#include "Runtime/Renderer/Private/SystemTextures.h"
+#include "SystemTextures.h"
 
 // todo - currently duplicated from SetupGlobalDistanceFieldParameters (GlobalDistanceField.cpp) because of problems getting it properly exported from the dll
 void FNiagaraDistanceFieldHelper::SetGlobalDistanceFieldParameters(const FGlobalDistanceFieldParameterData* OptionalParameterData, FGlobalDistanceFieldParameters2& ShaderParameters)
diff --git a/Engine/Plugins/FX/Niagara/Source/NiagaraVertexFactories/Public/NiagaraMeshVertexFactory.h b/Engine/Plugins/FX/Niagara/Source/NiagaraVertexFactories/Public/NiagaraMeshVertexFactory.h
--- a/Engine/Plugins/FX/Niagara/Source/NiagaraVertexFactories/Public/NiagaraMeshVertexFactory.h
+++ b/Engine/Plugins/FX/Niagara/Source/NiagaraVertexFactories/Public/NiagaraMeshVertexFactory.h
@@ -6,14 +6,14 @@ ParticleVertexFactory.h: Particle vertex factory definitions.
 
 #pragma once
 
+#include "Components.h"
 #include "CoreMinimal.h"
-#include "RenderResource.h"
-#include "UniformBuffer.h"
+#include "NiagaraDataSet.h"
 #include "NiagaraVertexFactory.h"
-#include "../../Niagara/Classes/NiagaraDataSet.h"
-#include "SceneView.h"
-#include "Components.h"
+#include "RenderResource.h"
 #include "SceneManagement.h"
+#include "SceneView.h"
+#include "UniformBuffer.h"
 #include "VertexFactory.h"
 
 // Disable this define to test disabling the use of GPU Scene with Niagara mesh renderer
