- Don't cull systems if we have no cached views
- This can occur if we reregister systems outside of the main game loop as the WorldManager will clear the cached view infos

#rnx

diff --git a/Engine/Plugins/FX/Niagara/Source/Niagara/Private/NiagaraWorldManager.cpp b/Engine/Plugins/FX/Niagara/Source/Niagara/Private/NiagaraWorldManager.cpp
--- a/Engine/Plugins/FX/Niagara/Source/Niagara/Private/NiagaraWorldManager.cpp
+++ b/Engine/Plugins/FX/Niagara/Source/Niagara/Private/NiagaraWorldManager.cpp
@@ -1524,6 +1524,9 @@ void FNiagaraWorldManager::ViewBasedCulling(UNiagaraEffectType* EffectType, cons
 				bInsideAnyView = true;
 			}
 		}
+
+		// If we have no CachedViewInfo this can mean one of two things, we have no views or we are outside of the main loop (i.e. reregister context) so start the sim anyway
+		bInsideAnyView |= CachedViewInfo.Num() == 0;
 	}
 
 	float TimeSinceInsideView = 0.0f;
