- Move CPU ribbons back to using the global index buffer allocation
- GPU ribbons now have an optional index buffer pool to improve performance

#rnx

diff --git a/Engine/Plugins/FX/Niagara/Source/Niagara/Private/NiagaraRendererRibbons.cpp b/Engine/Plugins/FX/Niagara/Source/Niagara/Private/NiagaraRendererRibbons.cpp
--- a/Engine/Plugins/FX/Niagara/Source/Niagara/Private/NiagaraRendererRibbons.cpp
+++ b/Engine/Plugins/FX/Niagara/Source/Niagara/Private/NiagaraRendererRibbons.cpp
@@ -27,7 +27,6 @@ DECLARE_CYCLE_STAT(TEXT("Generate Indices GPU [RT]"), STAT_NiagaraRenderRibbonsG
 DECLARE_CYCLE_STAT(TEXT("Generate Vertices CPU [GT]"), STAT_NiagaraRenderRibbonsGenVerticesCPU, STATGROUP_Niagara);
 DECLARE_CYCLE_STAT(TEXT("Generate Vertices GPU [RT]"), STAT_NiagaraRenderRibbonsGenVerticesGPU, STATGROUP_Niagara);
 
-
 DECLARE_STATS_GROUP(TEXT("NiagaraRibbons"), STATGROUP_NiagaraRibbons, STATCAT_Advanced);
 
 DECLARE_CYCLE_STAT(TEXT("Generate Vertices GPU - Sort [RT]"), STAT_NiagaraRenderRibbonsGenVerticesSortGPU, STATGROUP_NiagaraRibbons);
@@ -124,6 +123,33 @@ static FAutoConsoleVariableRef CVarNiagaraRibbonGpuInitMode(
 	ECVF_Scalability
 );
 
+static int32 GNiagaraRibbonGpuBufferCachePurgeCounter = 0;
+static FAutoConsoleVariableRef CVarNiagaraRibbonGpuBufferCachePurgeCounter(
+	TEXT("Niagara.Ribbon.GpuBufferCachePurgeCounter"),
+	GNiagaraRibbonGpuBufferCachePurgeCounter,
+	TEXT("The number of frames we hold onto ribbon buffer for.")
+	TEXT("Where 0 (Default) we purge them if not used next frame.")
+	TEXT("Negative values will purge the buffers the same frame, essentially zero reusing."),
+	ECVF_Default
+);
+
+static int32 GNiagaraRibbonGpuAllocateMaxCount = 1;
+static FAutoConsoleVariableRef CVarNiagaraRibbonGpuAllocateMaxCount(
+	TEXT("Niagara.Ribbon.GpuAllocateMaxCount"),
+	GNiagaraRibbonGpuAllocateMaxCount,
+	TEXT("When enabled (default) we allocate the maximum number of required elements.")
+	TEXT("This can result in memory bloat if the count is highly variable but will be more stable performance wise"),
+	ECVF_Default
+);
+
+static int32 GNiagaraRibbonGpuBufferAlign = 512;
+static FAutoConsoleVariableRef CVarNiagaraRibbonGpuBufferAlign(
+	TEXT("Niagara.Ribbon.GpuBufferAlign"),
+	GNiagaraRibbonGpuBufferAlign,
+	TEXT("When not allocating the maximum number of required elements we align up the request elements to this size to improve buffer reuse."),
+	ECVF_Default
+);
+
 static TAutoConsoleVariable<int32> CVarRayTracingNiagaraRibbons(
 	TEXT("r.RayTracing.Geometry.NiagaraRibbons"),
 	1,
@@ -221,41 +247,42 @@ struct FNiagaraRibbonIndirectDrawBufferLayout
 	float OneOverSubSegmentCount;	
 };
 
-
-struct FRWIndexBuffer final : FIndexBuffer
+struct FNiagaraRibbonIndexBuffer final : FIndexBuffer
 {
-	FUnorderedAccessViewRHIRef UAV;
-
-	FRWIndexBuffer()
-	{}
+	FNiagaraRibbonIndexBuffer() {}
 
-	virtual ~FRWIndexBuffer() override
+	virtual ~FNiagaraRibbonIndexBuffer() override
 	{
 		FRenderResource::ReleaseResource();
 	}
 
-	void Initialize(const TCHAR* InDebugName, const uint32 BytesPerElement, const uint32 NumElements, const EPixelFormat Format, const ERHIAccess InResourceState, const EBufferUsageFlags AdditionalUsage = BUF_None, FResourceArrayInterface *InResourceArray = nullptr)
+	// CPU allocation path
+	void Initialize(FGlobalDynamicIndexBuffer::FAllocationEx& IndexAllocation)
 	{
 		InitResource();
-		
-		// Provide a debug name if using Fast VRAM so the allocators diagnostics will work
-		ensure(!(EnumHasAnyFlags(AdditionalUsage, BUF_FastVRAM) && !InDebugName));
-		const int32 NumBytes = BytesPerElement * NumElements;
-		FRHIResourceCreateInfo CreateInfo(InDebugName);
-		CreateInfo.ResourceArray = InResourceArray;
-		IndexBufferRHI = RHICreateIndexBuffer(BytesPerElement, NumBytes, BUF_UnorderedAccess | AdditionalUsage, InResourceState, CreateInfo);
-		if (EnumHasAnyFlags(AdditionalUsage, BUF_UnorderedAccess))
-		{
-			UAV = RHICreateUnorderedAccessView(IndexBufferRHI, UE_PIXELFORMAT_TO_UINT8(Format));
-		}
+
+		IndexBufferRHI = IndexAllocation.IndexBuffer->IndexBufferRHI;
+		FirstIndex = IndexAllocation.FirstIndex;
+	}
+
+	// GPU allocation path assumes 32 bit indicies
+	void Initialize(const uint32 NumElements)
+	{
+		InitResource();
+
+		FRHIResourceCreateInfo CreateInfo(TEXT("NiagaraRibbonIndexBuffer"));
+		IndexBufferRHI = RHICreateIndexBuffer(sizeof(uint32), sizeof(uint32) * NumElements, BUF_Static | BUF_UnorderedAccess, ERHIAccess::VertexOrIndexBuffer, CreateInfo);
+		UAV = RHICreateUnorderedAccessView(IndexBufferRHI, PF_R32_UINT);
 	}
 
 	virtual void ReleaseRHI() override
 	{
 		UAV.SafeRelease();
-
 		FIndexBuffer::ReleaseRHI();
 	}
+
+	uint32						FirstIndex = 0;
+	FUnorderedAccessViewRHIRef	UAV;
 };
 
 
@@ -297,13 +324,11 @@ struct FNiagaraDynamicDataRibbon : public FNiagaraDynamicDataBase
 
 struct FNiagaraRibbonRenderingFrameViewResources
 {
-	FNiagaraRibbonVertexFactory VertexFactory;
-	FNiagaraRibbonUniformBufferRef UniformBuffer;
-
-	FRWIndexBuffer IndexBuffer;
-	FRWBuffer IndirectDrawBuffer;
-
-	FNiagaraIndexGenerationInput IndexGenerationSettings;
+	FNiagaraRibbonVertexFactory		VertexFactory;
+	FNiagaraRibbonUniformBufferRef	UniformBuffer;
+	FNiagaraRibbonIndexBuffer		IndexBuffer;
+	FRWBuffer						IndirectDrawBuffer;
+	FNiagaraIndexGenerationInput	IndexGenerationSettings;
 
 	int32 IndirectDrawBufferStartOffset = FNiagaraRibbonIndirectDrawBufferLayout::IndirectDrawCommandIndex;
 	int32 IndirectDrawBufferStartByteOffset = FNiagaraRibbonIndirectDrawBufferLayout::IndirectDrawCommandByteOffset;
@@ -511,6 +536,19 @@ struct FNiagaraRibbonGPUInitComputeBuffers
 
 class FNiagaraGpuRibbonsDataManager final : public FNiagaraGpuComputeDataManager
 {
+	struct FIndirectDrawBufferEntry
+	{
+		uint64		FrameUsed = 0;
+		FRWBuffer	Buffer;
+	};
+
+	struct FIndexBufferEntry
+	{
+		uint64						FrameUsed = 0;
+		int32						NumIndices = 0;
+		FNiagaraRibbonIndexBuffer	Buffer;
+	};
+
 public:
 	FNiagaraGpuRibbonsDataManager(FNiagaraGpuComputeDispatchInterface* InOwnerInterface)
 		: FNiagaraGpuComputeDataManager(InOwnerInterface)
@@ -536,13 +574,72 @@ public:
 		RenderersToGeneratePerStage[GenerateIndex].Emplace(InRenderer, InSourceParticleData, InRenderingResources);
 	}
 
+	//-OPT: These caches should be more central and are as a simple solution to reduce memory thrashing / poor performance for ribbons
+	FRWBuffer GetOrAllocateIndirectDrawBuffer()
+	{
+		FIndirectDrawBufferEntry* BufferEntry = IndirectDrawBufferCache.FindByPredicate(
+			[&](const FIndirectDrawBufferEntry& ExistingBuffer)
+			{
+				return ExistingBuffer.FrameUsed != FrameCounter;
+			}
+		);
+		if (BufferEntry == nullptr)
+		{
+			BufferEntry = &IndirectDrawBufferCache.AddDefaulted_GetRef();
+			BufferEntry->Buffer.Initialize(TEXT("RibbonIndirectDrawBuffer"), sizeof(uint32), FNiagaraRibbonIndirectDrawBufferLayout::NumElements, EPixelFormat::PF_R32_UINT, ERHIAccess::IndirectArgs | ERHIAccess::SRVMask, BUF_Static | BUF_DrawIndirect);
+		}
+		BufferEntry->FrameUsed = FrameCounter;
+		return BufferEntry->Buffer;
+	}
+
+	FNiagaraRibbonIndexBuffer GetOrAllocateIndexBuffer(int32 NumIndices, int32 MaxIndices)
+	{
+		if (GNiagaraRibbonGpuBufferCachePurgeCounter >= 0)
+		{
+			NumIndices = GNiagaraRibbonGpuAllocateMaxCount == 0 ? Align(NumIndices, GNiagaraRibbonGpuBufferAlign) : MaxIndices;
+		}
+
+		FIndexBufferEntry* BufferEntry = Index32BufferCache.FindByPredicate(
+			[&](const FIndexBufferEntry& ExistingBuffer)
+			{
+				return ExistingBuffer.FrameUsed != FrameCounter && ExistingBuffer.NumIndices == NumIndices;
+			}
+		);
+		if (BufferEntry == nullptr)
+		{
+			BufferEntry = &Index32BufferCache.AddDefaulted_GetRef();
+			BufferEntry->NumIndices = NumIndices;
+			BufferEntry->Buffer.Initialize(NumIndices);
+		}
+
+		BufferEntry->FrameUsed = FrameCounter;
+		return BufferEntry->Buffer;
+	}
+
 private:
-	TArray<FNiagaraRibbonGPUInitParameters> RenderersToGeneratePerStage[2];
+	TArray<FNiagaraRibbonGPUInitParameters>	RenderersToGeneratePerStage[2];
+	FNiagaraRibbonGPUInitComputeBuffers		ComputeBuffers;
+
+	TArray<FIndirectDrawBufferEntry>		IndirectDrawBufferCache;
+	TArray<FIndexBufferEntry>				Index32BufferCache;
+
+	uint64									FrameCounter = 0;
 
-	FNiagaraRibbonGPUInitComputeBuffers ComputeBuffers;
-	
 	void OnPostPreRender(FRHICommandListImmediate& RHICmdList)
 	{
+		if (GNiagaraRibbonGpuBufferCachePurgeCounter < 0)
+		{
+			IndirectDrawBufferCache.Empty();
+			Index32BufferCache.Empty();
+		}
+		else
+		{
+			const uint64 PurgeCounter = GNiagaraRibbonGpuBufferCachePurgeCounter;
+			IndirectDrawBufferCache.RemoveAll([&](const FIndirectDrawBufferEntry& InBuffer) { return FrameCounter - InBuffer.FrameUsed > PurgeCounter; });
+			Index32BufferCache.RemoveAll([&](const FIndexBufferEntry& InBuffer) { return FrameCounter - InBuffer.FrameUsed > PurgeCounter; });
+			++FrameCounter;
+		}
+
 		GenerateAllGPUData(RHICmdList, RenderersToGeneratePerStage[0]);
 	}
 
@@ -692,7 +789,10 @@ void FNiagaraRendererRibbons::GetDynamicMeshElements(const TArray<const FSceneVi
 	const FNiagaraRibbonMeshCollectorResources& RenderingResources = Collector.AllocateOneFrameResource<FNiagaraRibbonMeshCollectorResources>();
 		
 	InitializeVertexBuffersResources(DynamicData, SourceParticleData, Collector.GetDynamicReadBuffer(), RenderingResources.RibbonResources, DynamicData->bUseGPUInit);
-	
+
+	FNiagaraGpuComputeDispatchInterface* ComputeDispatchInterface = SceneProxy->GetComputeDispatchInterface();
+	FNiagaraGpuRibbonsDataManager& GpuRibbonDataManager = ComputeDispatchInterface->GetOrCreateDataManager<FNiagaraGpuRibbonsDataManager>();
+
 	// Compute the per-view uniform buffers.
 	for (int32 ViewIndex = 0; ViewIndex < Views.Num(); ViewIndex++)
 	{
@@ -714,7 +814,7 @@ void FNiagaraRendererRibbons::GetDynamicMeshElements(const TArray<const FSceneVi
 			auto& RenderingViewResources = RenderingResources.RibbonResources->ViewResources.Add_GetRef(MakeShared<FNiagaraRibbonRenderingFrameViewResources>());
 			RenderingViewResources->IndexGenerationSettings = CalculateIndexBufferConfiguration(DynamicData->GenerationOutput, SourceParticleData, SceneProxy, View, ViewOriginForDistanceCulling, DynamicData->bUseGPUInit, DynamicData->bIsGPUSystem);
 			
-			GenerateIndexBufferForView(RenderingViewResources->IndexGenerationSettings, DynamicData, RenderingViewResources, View, ViewOriginForDistanceCulling, DynamicData->bUseGPUInit);
+			GenerateIndexBufferForView(GpuRibbonDataManager, Collector, RenderingViewResources->IndexGenerationSettings, DynamicData, RenderingViewResources, View, ViewOriginForDistanceCulling);
 			
 			SetupPerViewUniformBuffer(RenderingViewResources->IndexGenerationSettings, View, ViewFamily, SceneProxy, RenderingViewResources->UniformBuffer);
 			
@@ -727,9 +827,7 @@ void FNiagaraRendererRibbons::GetDynamicMeshElements(const TArray<const FSceneVi
 	// Register this renderer for generation this frame if we're a gpu system or using gpu init
 	if (DynamicData->bUseGPUInit || DynamicData->bIsGPUSystem)
 	{
-		FNiagaraGpuComputeDispatchInterface* ComputeDispatchInterface = SceneProxy->GetComputeDispatchInterface();
-		FNiagaraGpuRibbonsDataManager& RibbonDataManager = ComputeDispatchInterface->GetOrCreateDataManager<FNiagaraGpuRibbonsDataManager>();
-		RibbonDataManager.RegisterRenderer(this, SourceParticleData, RenderingResources.RibbonResources);
+		GpuRibbonDataManager.RegisterRenderer(this, SourceParticleData, RenderingResources.RibbonResources);
 	}
 }
 
@@ -909,10 +1007,12 @@ void FNiagaraRendererRibbons::GetDynamicRayTracingInstances(FRayTracingMaterialG
 	{
 		return;
 	}
-	
+
+	FNiagaraGpuRibbonsDataManager& GpuRibbonDataManager = ComputeDispatchInterface->GetOrCreateDataManager<FNiagaraGpuRibbonsDataManager>();
+
 	InitializeVertexBuffersResources(DynamicDataRibbon, SourceParticleData, Context.RayTracingMeshResourceCollector.GetDynamicReadBuffer(), RenderingResources.RibbonResources, DynamicDataRibbon->bUseGPUInit);
 	
-	GenerateIndexBufferForView(RenderingViewResources->IndexGenerationSettings, DynamicDataRibbon, RenderingViewResources, View, ViewOriginForDistanceCulling, DynamicDataRibbon->bUseGPUInit);
+	GenerateIndexBufferForView(GpuRibbonDataManager, Context.RayTracingMeshResourceCollector, RenderingViewResources->IndexGenerationSettings, DynamicDataRibbon, RenderingViewResources, View, ViewOriginForDistanceCulling);
 			
 	SetupPerViewUniformBuffer(RenderingViewResources->IndexGenerationSettings, View, ViewFamily, SceneProxy, RenderingViewResources->UniformBuffer);
 	
@@ -926,7 +1026,7 @@ void FNiagaraRendererRibbons::GetDynamicRayTracingInstances(FRayTracingMaterialG
 	RayTracingInstance.InstanceTransforms.Add(FMatrix::Identity);
 	
 	RayTracingGeometry.Initializer.IndexBuffer = RenderingViewResources->IndexBuffer.IndexBufferRHI;// PerViewGeneratedData.IndexAllocation.IndexBuffer->IndexBufferRHI;
-	RayTracingGeometry.Initializer.IndexBufferOffset = 0;//PerViewGeneratedData.IndexAllocation.FirstIndex * PerViewGeneratedData.IndexAllocation.IndexStride;
+	RayTracingGeometry.Initializer.IndexBufferOffset = RenderingViewResources->IndexBuffer.FirstIndex;//PerViewGeneratedData.IndexAllocation.FirstIndex * PerViewGeneratedData.IndexAllocation.IndexStride;
 	
 	FMeshBatch MeshBatch;
 	
@@ -1711,33 +1811,34 @@ FNiagaraIndexGenerationInput FNiagaraRendererRibbons::CalculateIndexBufferConfig
 	return IndexGenInput;
 }
 
-void FNiagaraRendererRibbons::GenerateIndexBufferForView(FNiagaraIndexGenerationInput& GeneratedData, FNiagaraDynamicDataRibbon* DynamicDataRibbon,
-                                                         const TSharedPtr<FNiagaraRibbonRenderingFrameViewResources>& RenderingViewResources, const FSceneView* View,
-                                                         const FVector& ViewOriginForDistanceCulling, bool bShouldUseGPUInitIndices) const
+void FNiagaraRendererRibbons::GenerateIndexBufferForView(
+	FNiagaraGpuRibbonsDataManager& GpuRibbonsDataManager,
+	FMeshElementCollector& Collector,
+	FNiagaraIndexGenerationInput& GeneratedData, FNiagaraDynamicDataRibbon* DynamicDataRibbon,
+    const TSharedPtr<FNiagaraRibbonRenderingFrameViewResources>& RenderingViewResources, const FSceneView* View,
+    const FVector& ViewOriginForDistanceCulling
+) const
 {
 	if (GeneratedData.MaxSegmentCount > 0)
 	{
-		if (bShouldUseGPUInitIndices)
+		if (DynamicDataRibbon->bUseGPUInit)
 		{
-			RenderingViewResources->IndirectDrawBuffer.Initialize(TEXT("RibbonIndirectDrawBuffer"), sizeof(uint32), FNiagaraRibbonIndirectDrawBufferLayout::NumElements, EPixelFormat::PF_R32_UINT, ERHIAccess::IndirectArgs | ERHIAccess::SRVMask, BUF_Static | BUF_DrawIndirect);			
+			RenderingViewResources->IndirectDrawBuffer = GpuRibbonsDataManager.GetOrAllocateIndirectDrawBuffer();
+			RenderingViewResources->IndexBuffer = GpuRibbonsDataManager.GetOrAllocateIndexBuffer(GeneratedData.TotalNumIndices, DynamicDataRibbon->MaxAllocatedParticleCount);
 		}
-		
-		// Can we pack into a uint16 vs a uint32?
-		const EBufferUsageFlags IndexBufferUsage = BUF_Static | (bShouldUseGPUInitIndices ? BUF_UnorderedAccess : BUF_None);
-		if (GeneratedData.TotalBitCount <= 16 && !bShouldUseGPUInitIndices)
+		else
 		{
-			RenderingViewResources->IndexBuffer.Initialize(TEXT("NiagaraRibbonIndexBuffer"), sizeof(uint16), GeneratedData.TotalNumIndices, PF_R16_UINT, ERHIAccess::VertexOrIndexBuffer, IndexBufferUsage);
-			if (!bShouldUseGPUInitIndices)
+			if (GeneratedData.TotalBitCount <= 16)
 			{
-				GenerateIndexBufferCPU<uint16>(GeneratedData, DynamicDataRibbon, ShapeState, RenderingViewResources, View, ViewOriginForDistanceCulling, FeatureLevel, DrawDirection);
+				FGlobalDynamicIndexBuffer::FAllocationEx IndexAllocation = Collector.GetDynamicIndexBuffer().Allocate<uint16>(GeneratedData.TotalNumIndices);
+				RenderingViewResources->IndexBuffer.Initialize(IndexAllocation);
+				GenerateIndexBufferCPU<uint16>(GeneratedData, DynamicDataRibbon, ShapeState, reinterpret_cast<uint16*>(IndexAllocation.Buffer), View, ViewOriginForDistanceCulling, FeatureLevel, DrawDirection);
 			}
-		}
-		else
-		{		
-			RenderingViewResources->IndexBuffer.Initialize(TEXT("NiagaraRibbonIndexBuffer"), sizeof(uint32), GeneratedData.TotalNumIndices, PF_R32_UINT, ERHIAccess::VertexOrIndexBuffer, IndexBufferUsage);
-			if (!bShouldUseGPUInitIndices)
+			else
 			{
-				GenerateIndexBufferCPU<uint32>(GeneratedData, DynamicDataRibbon, ShapeState, RenderingViewResources, View, ViewOriginForDistanceCulling, FeatureLevel, DrawDirection);				
+				FGlobalDynamicIndexBuffer::FAllocationEx IndexAllocation = Collector.GetDynamicIndexBuffer().Allocate<uint32>(GeneratedData.TotalNumIndices);
+				RenderingViewResources->IndexBuffer.Initialize(IndexAllocation);
+				GenerateIndexBufferCPU<uint32>(GeneratedData, DynamicDataRibbon, ShapeState, reinterpret_cast<uint32*>(IndexAllocation.Buffer), View, ViewOriginForDistanceCulling, FeatureLevel, DrawDirection);
 			}
 		}
 	}
@@ -1745,16 +1846,16 @@ void FNiagaraRendererRibbons::GenerateIndexBufferForView(FNiagaraIndexGeneration
 
 template <typename TValue>
 void FNiagaraRendererRibbons::GenerateIndexBufferCPU(FNiagaraIndexGenerationInput& GeneratedData, FNiagaraDynamicDataRibbon* DynamicDataRibbon, const FNiagaraRibbonShapeGeometryData& ShapeState,
-    const TSharedPtr<FNiagaraRibbonRenderingFrameViewResources>& RenderingViewResources, const FSceneView* View, const FVector& ViewOriginForDistanceCulling, ERHIFeatureLevel::Type FeatureLevel, ENiagaraRibbonDrawDirection DrawDirection)
+    TValue* StartIndexBuffer, const FSceneView* View, const FVector& ViewOriginForDistanceCulling, ERHIFeatureLevel::Type FeatureLevel, ENiagaraRibbonDrawDirection DrawDirection)
 {
 	SCOPE_CYCLE_COUNTER(STAT_NiagaraRenderRibbonsGenIndiciesCPU);
+
 	FMaterialRenderProxy* MaterialRenderProxy = DynamicDataRibbon->Material;
 	check(MaterialRenderProxy);
 	const EBlendMode BlendMode = MaterialRenderProxy->GetIncompleteMaterialWithFallback(FeatureLevel).GetBlendMode();
 	
 	const TSharedPtr<FNiagaraRibbonCPUGeneratedVertexData>& GeneratedGeometryData = DynamicDataRibbon->GenerationOutput;
 
-	TValue* StartIndexBuffer = reinterpret_cast<TValue*>(RHILockBuffer(RenderingViewResources->IndexBuffer.IndexBufferRHI, 0, GeneratedData.TotalNumIndices * sizeof(TValue), RLM_WriteOnly));
 	TValue* CurrentIndexBuffer = StartIndexBuffer;
 	if (IsTranslucentBlendMode(BlendMode) && GeneratedGeometryData->RibbonInfoLookup.Num())
 	{
@@ -1771,7 +1872,6 @@ void FNiagaraRendererRibbons::GenerateIndexBufferCPU(FNiagaraIndexGenerationInpu
 	}
 	GeneratedData.CPUTriangleCount = (CurrentIndexBuffer - StartIndexBuffer) / 3;
 	check(CurrentIndexBuffer <= StartIndexBuffer + GeneratedData.TotalNumIndices);
-	RHIUnlockBuffer(RenderingViewResources->IndexBuffer.IndexBufferRHI);
 }
 
 template <typename TValue>
@@ -1952,7 +2052,7 @@ inline void FNiagaraRendererRibbons::SetupMeshBatchAndCollectorResourceForView(c
 	
 	FMeshBatchElement& MeshElement = OutMeshBatch.Elements[0];
 	MeshElement.IndexBuffer = &RenderingViewResources->IndexBuffer;
-	MeshElement.FirstIndex = 0;
+	MeshElement.FirstIndex = RenderingViewResources->IndexBuffer.FirstIndex;
 	MeshElement.NumInstances = 1;
 	MeshElement.MinVertexIndex = 0;
 	MeshElement.MaxVertexIndex = 0;
diff --git a/Engine/Plugins/FX/Niagara/Source/Niagara/Public/NiagaraRendererRibbons.h b/Engine/Plugins/FX/Niagara/Source/Niagara/Public/NiagaraRendererRibbons.h
--- a/Engine/Plugins/FX/Niagara/Source/Niagara/Public/NiagaraRendererRibbons.h
+++ b/Engine/Plugins/FX/Niagara/Source/Niagara/Public/NiagaraRendererRibbons.h
@@ -20,6 +20,7 @@ NiagaraRenderer.h: Base class for Niagara render modules
 class FNiagaraDataSet;
 
 struct FNiagaraDynamicDataRibbon;
+class FNiagaraGpuRibbonsDataManager;
 struct FNiagaraRibbonRenderingFrameResources;
 struct FNiagaraRibbonRenderingFrameViewResources;
 struct FNiagaraRibbonGPUInitParameters;
@@ -385,12 +386,17 @@ protected:
 	FNiagaraIndexGenerationInput CalculateIndexBufferConfiguration(const TSharedPtr<FNiagaraRibbonCPUGeneratedVertexData>& GeneratedVertices, const FNiagaraDataBuffer* SourceParticleData,
 		const FNiagaraSceneProxy* SceneProxy, const FSceneView* View, const FVector& ViewOriginForDistanceCulling, bool bShouldUseGPUInitIndices, bool bIsGPUSim) const;
 	
-	void GenerateIndexBufferForView(FNiagaraIndexGenerationInput& GeneratedData, FNiagaraDynamicDataRibbon* DynamicDataRibbon,
-	                                const TSharedPtr<FNiagaraRibbonRenderingFrameViewResources>& RenderingViewResources, const FSceneView* View, const FVector& ViewOriginForDistanceCulling, bool bShouldUseGPUInitIndices) const;
+	void GenerateIndexBufferForView(
+		FNiagaraGpuRibbonsDataManager& GpuRibbonsDataManager, FMeshElementCollector& Collector,
+		FNiagaraIndexGenerationInput& GeneratedData, FNiagaraDynamicDataRibbon* DynamicDataRibbon,
+		const TSharedPtr<FNiagaraRibbonRenderingFrameViewResources>& RenderingViewResources, const FSceneView* View, const FVector& ViewOriginForDistanceCulling
+	) const;
 	
 	template <typename TValue>
-	static void GenerateIndexBufferCPU(FNiagaraIndexGenerationInput& GeneratedData, FNiagaraDynamicDataRibbon* DynamicDataRibbon, const FNiagaraRibbonShapeGeometryData& ShapeState,
-	                                   const TSharedPtr<FNiagaraRibbonRenderingFrameViewResources>& RenderingViewResources, const FSceneView* View, const FVector& ViewOriginForDistanceCulling, ERHIFeatureLevel::Type FeatureLevel, ENiagaraRibbonDrawDirection DrawDirection);
+	static void GenerateIndexBufferCPU(
+		FNiagaraIndexGenerationInput& GeneratedData, FNiagaraDynamicDataRibbon* DynamicDataRibbon, const FNiagaraRibbonShapeGeometryData& ShapeState,
+		TValue* StartIndexBuffer, const FSceneView* View, const FVector& ViewOriginForDistanceCulling, ERHIFeatureLevel::Type FeatureLevel, ENiagaraRibbonDrawDirection DrawDirection
+	);
 
 	template <typename TValue>
 	static TValue* AppendToIndexBufferCPU(TValue* OutIndices, const FNiagaraIndexGenerationInput& GeneratedData, const FNiagaraRibbonShapeGeometryData& ShapeState, const TArrayView<uint32>& SegmentData, bool bInvertOrder);
@@ -431,5 +437,5 @@ protected:
 	mutable FNiagaraRibbonTessellationSmoothingData TessellationSmoothingData;
 	int32 RibbonIDParamDataSetOffset;
 
-	friend class FNiagaraGpuRibbonsDataManager;
+	friend FNiagaraGpuRibbonsDataManager;
 };
