- Fix raytracing for Niagara sprites

#rnx

diff --git a/Engine/Plugins/FX/Niagara/Source/Niagara/Private/NiagaraRendererSprites.cpp b/Engine/Plugins/FX/Niagara/Source/Niagara/Private/NiagaraRendererSprites.cpp
--- a/Engine/Plugins/FX/Niagara/Source/Niagara/Private/NiagaraRendererSprites.cpp
+++ b/Engine/Plugins/FX/Niagara/Source/Niagara/Private/NiagaraRendererSprites.cpp
@@ -1055,7 +1055,7 @@ void FNiagaraRendererSprites::GetDynamicRayTracingInstances(FRayTracingMaterialG
 		InitializeSortInfo(ParticleSpriteRenderData, *SceneProxy, *Context.ReferenceView, 0, SortInfo);
 	}
 
-	if (FNiagaraSpriteVertexFactory::StaticType.SupportsRayTracingDynamicGeometry())
+	if (!FNiagaraSpriteVertexFactory::StaticType.SupportsRayTracingDynamicGeometry())
 	{
 		return;
 	}
