Fix for landscape DI using shared samplers for the Normal VT

diff --git a/Engine/Plugins/FX/Niagara/Shaders/Private/NiagaraDataInterfaceLandscape.ush b/Engine/Plugins/FX/Niagara/Shaders/Private/NiagaraDataInterfaceLandscape.ush
--- a/Engine/Plugins/FX/Niagara/Shaders/Private/NiagaraDataInterfaceLandscape.ush
+++ b/Engine/Plugins/FX/Niagara/Shaders/Private/NiagaraDataInterfaceLandscape.ush
@@ -31,7 +31,8 @@ uint4				{ParameterName}_HeightVirtualTextureUniforms;
 Texture2D			{ParameterName}_NormalVirtualTexture0;
 Texture2D			{ParameterName}_NormalVirtualTexture1;
 Texture2D<uint4>	{ParameterName}_NormalVirtualTexturePageTable;
-SamplerState		{ParameterName}_NormalVirtualTextureSampler;
+SamplerState		{ParameterName}_NormalVirtualTexture0Sampler;
+SamplerState		{ParameterName}_NormalVirtualTexture1Sampler;
 float4x4			{ParameterName}_NormalVirtualTextureWorldToUvTransform;
 uint4				{ParameterName}_NormalVirtualTexturePackedUniform0;
 uint4				{ParameterName}_NormalVirtualTexturePackedUniform1;
@@ -151,14 +152,14 @@ void GetWorldNormal_{ParameterName}(float3 InWorldPos, out float3 OutNormal, out
 
 		float4 PackedValue0 = TextureVirtualSample(
 			{ParameterName}_NormalVirtualTexture0,
-			{ParameterName}_NormalVirtualTextureSampler,
+			{ParameterName}_NormalVirtualTexture0Sampler,
 			PageTable,
 			0,
 			VTUniform_Unpack({ParameterName}_NormalVirtualTextureUniforms0));
 
 		float4 PackedValue1 = TextureVirtualSample(
 			{ParameterName}_NormalVirtualTexture1,
-			{ParameterName}_NormalVirtualTextureSampler,
+			{ParameterName}_NormalVirtualTexture1Sampler,
 			PageTable,
 			0,
 			VTUniform_Unpack({ParameterName}_NormalVirtualTextureUniforms1));
diff --git a/Engine/Plugins/FX/Niagara/Source/Niagara/Private/NiagaraDataInterfaceLandscape.cpp b/Engine/Plugins/FX/Niagara/Source/Niagara/Private/NiagaraDataInterfaceLandscape.cpp
--- a/Engine/Plugins/FX/Niagara/Source/Niagara/Private/NiagaraDataInterfaceLandscape.cpp
+++ b/Engine/Plugins/FX/Niagara/Source/Niagara/Private/NiagaraDataInterfaceLandscape.cpp
@@ -67,7 +67,8 @@ namespace NiagaraDataInterfaceLandscape
 		SHADER_PARAMETER_SRV(Texture2D,				NormalVirtualTexture0)
 		SHADER_PARAMETER_SRV(Texture2D,				NormalVirtualTexture1)
 		SHADER_PARAMETER_TEXTURE(Texture2D<uint4>,	NormalVirtualTexturePageTable)
-		SHADER_PARAMETER_SAMPLER(SamplerState,		NormalVirtualTextureSampler)
+		SHADER_PARAMETER_SAMPLER(SamplerState,		NormalVirtualTexture0Sampler)
+		SHADER_PARAMETER_SAMPLER(SamplerState,		NormalVirtualTexture1Sampler)
 		SHADER_PARAMETER(FMatrix44f,				NormalVirtualTextureWorldToUvTransform)
 		SHADER_PARAMETER(FUintVector4,				NormalVirtualTexturePackedUniform0)
 		SHADER_PARAMETER(FUintVector4,				NormalVirtualTexturePackedUniform1)
@@ -1315,7 +1316,8 @@ void UNiagaraDataInterfaceLandscape::SetShaderParameters(const FNiagaraDataInter
 	FRHISamplerState* BilinearSamplerState				= TStaticSamplerState<SF_Bilinear, AM_Clamp, AM_Clamp, AM_Clamp>::GetRHI();
 	ShaderParameters->BaseColorVirtualTextureSampler	= BilinearSamplerState;
 	ShaderParameters->HeightVirtualTextureSampler		= BilinearSamplerState;
-	ShaderParameters->NormalVirtualTextureSampler		= BilinearSamplerState;
+	ShaderParameters->NormalVirtualTexture0Sampler		= BilinearSamplerState;
+	ShaderParameters->NormalVirtualTexture1Sampler		= BilinearSamplerState;
 	ShaderParameters->CachedHeightTextureSampler		= BilinearSamplerState;
 
 	// Set Textures
diff --git a/Engine/Source/Runtime/Core/Private/UObject/FortniteMainBranchObjectVersion.cpp b/Engine/Source/Runtime/Core/Private/UObject/FortniteMainBranchObjectVersion.cpp
--- a/Engine/Source/Runtime/Core/Private/UObject/FortniteMainBranchObjectVersion.cpp
+++ b/Engine/Source/Runtime/Core/Private/UObject/FortniteMainBranchObjectVersion.cpp
@@ -11,7 +11,7 @@ TMap<FGuid, FGuid> FFortniteMainBranchObjectVersion::GetSystemGuids()
 	SystemGuids.Add(DevGuids.MATERIALSHADERMAP_DERIVEDDATA_VER, FGuid("1DA85DA5733C4A489A1BE117CDC5ACCC"));
 	SystemGuids.Add(DevGuids.NANITE_DERIVEDDATA_VER, FGuid("B62CDCAF66AA45289DDB82E276E7C3E7"));
 	SystemGuids.Add(DevGuids.NIAGARASHADERMAP_DERIVEDDATA_VER, FGuid("2289D5116CF94BC0AFEAEC541468E645"));
-	SystemGuids.Add(DevGuids.Niagara_LatestScriptCompileVersion, FGuid("A9222A671841874BA54368C57E3BCE82"));
+	SystemGuids.Add(DevGuids.Niagara_LatestScriptCompileVersion, FGuid("BC25219B21894B4684FF9289CEACCB5A"));
 	SystemGuids.Add(DevGuids.SkeletalMeshDerivedDataVersion, FGuid("ACF593EAAE354FCBB34CF44F5AA1BFD2"));
 	SystemGuids.Add(DevGuids.STATICMESH_DERIVEDDATA_VER, FGuid("18E028FCF1BC469385806AA5F8AE07CC"));
 	return SystemGuids;
