Fixing cases where multiple static variable sets to the same varaible but different values should trigger an error and it does not.
Fixing cases where we weren't properly migrating varaibles across different script stacks, especially emitter variables.
Fixing cases where static variables as function call inputs weren't being properly propagated when they were using the pin's value and not linked to another pin.

diff --git a/Engine/Plugins/FX/Niagara/Source/Niagara/Public/NiagaraTypes.h b/Engine/Plugins/FX/Niagara/Source/Niagara/Public/NiagaraTypes.h
--- a/Engine/Plugins/FX/Niagara/Source/Niagara/Public/NiagaraTypes.h
+++ b/Engine/Plugins/FX/Niagara/Source/Niagara/Public/NiagaraTypes.h
@@ -1814,18 +1814,25 @@ struct FNiagaraVariableBase
 
 	FORCEINLINE bool IsInNameSpace(const FStringView& Namespace) const
 	{
-		TStringBuilder<128> NameString;
+		FNameBuilder NameString;
 		Name.ToString(NameString);
 		
 		FStringView NameStringView = NameString.ToView();
-		return (NameStringView.Len() > Namespace.Len() + 1) && (NameStringView[Namespace.Len()] == '.') && NameStringView.StartsWith(Namespace);
+		if (Namespace.Len() > 1 && Namespace[Namespace.Len() - 1] == '.') // Includes period in namespace
+		{
+			return (NameStringView.Len() > Namespace.Len()) && NameStringView.StartsWith(Namespace);
+		}
+		else // Skips period, makes sure it's in the name string to delineate
+		{
+			return (NameStringView.Len() > Namespace.Len() + 1) && (NameStringView[Namespace.Len()] == '.') && NameStringView.StartsWith(Namespace);
+		}
 	}
 
 #if WITH_EDITORONLY_DATA
 	// This method should not be used at runtime as we have pre-defined strings in FNiagaraConstants for runtime cases
 	FORCEINLINE bool IsInNameSpace(const FName& Namespace) const
 	{
-		TStringBuilder<128> NamespaceString;
+		FNameBuilder NamespaceString;
 		Namespace.ToString(NamespaceString);
 		return IsInNameSpace(NamespaceString.ToView());
 	}
diff --git a/Engine/Plugins/FX/Niagara/Source/NiagaraEditor/Private/EdGraphSchema_Niagara.cpp b/Engine/Plugins/FX/Niagara/Source/NiagaraEditor/Private/EdGraphSchema_Niagara.cpp
--- a/Engine/Plugins/FX/Niagara/Source/NiagaraEditor/Private/EdGraphSchema_Niagara.cpp
+++ b/Engine/Plugins/FX/Niagara/Source/NiagaraEditor/Private/EdGraphSchema_Niagara.cpp
@@ -1786,6 +1786,10 @@ FEdGraphPinType UEdGraphSchema_Niagara::TypeDefinitionToPinType(FNiagaraTypeDefi
 
 	}
 }
+bool UEdGraphSchema_Niagara::IsPinStatic(const UEdGraphPin* Pin)
+{
+	return Pin->PinType.PinCategory == PinCategoryStaticEnum || Pin->PinType.PinCategory == PinCategoryStaticClass || Pin->PinType.PinCategory == PinCategoryStaticType;
+}
 
 bool UEdGraphSchema_Niagara::IsPinWildcard(const UEdGraphPin* Pin)
 {
diff --git a/Engine/Plugins/FX/Niagara/Source/NiagaraEditor/Private/NiagaraCompiler.cpp b/Engine/Plugins/FX/Niagara/Source/NiagaraEditor/Private/NiagaraCompiler.cpp
--- a/Engine/Plugins/FX/Niagara/Source/NiagaraEditor/Private/NiagaraCompiler.cpp
+++ b/Engine/Plugins/FX/Niagara/Source/NiagaraEditor/Private/NiagaraCompiler.cpp
@@ -262,6 +262,62 @@ UNiagaraDataInterface* FNiagaraCompileRequestDuplicateData::GetDuplicatedDataInt
 	return nullptr;
 }
 
+void FNiagaraCompileRequestData::SortOutputNodesByDependencies(TArray<class UNiagaraNodeOutput*>& NodesToSort, const TArray<class UNiagaraSimulationStageBase*>* SimStages)
+{
+	if (!SimStages)
+		return;
+
+	TArray<class UNiagaraNodeOutput*> NewArray;
+	NewArray.Reserve(NodesToSort.Num());
+
+	// First gather up the non-simstage items
+	bool bFoundAnySimStages = false;
+	for (class UNiagaraNodeOutput* OutputNode : NodesToSort)
+	{
+		// Add any non sim stage entries back to the array in the order of encounter
+		if (OutputNode->GetUsage() != ENiagaraScriptUsage::ParticleSimulationStageScript)
+		{
+			NewArray.Emplace(OutputNode);
+		}
+		else
+		{
+			bFoundAnySimStages = true;
+		}
+	}
+
+	// No Sim stages, no problem! Just return
+	if (!bFoundAnySimStages)
+	{
+		return;
+	}
+
+	ensure(SimStages->Num() == (NodesToSort.Num() - NewArray.Num()));
+
+	// Add any sim stage entries back to the array in the order of encounter in the SimStage entry list from the Emitter (Handles reordering)
+	for (const UNiagaraSimulationStageBase* Stage : *SimStages)
+	{
+		if (Stage && Stage->Script)
+		{
+			const FGuid & StageId = Stage->Script->GetUsageId();
+
+			for (class UNiagaraNodeOutput* OutputNode : NodesToSort)
+			{
+				if (OutputNode->GetUsage() == ENiagaraScriptUsage::ParticleSimulationStageScript && OutputNode->GetUsageId() == StageId)
+				{
+					NewArray.Emplace(OutputNode);
+					break;
+				}
+			}
+		}
+	}
+
+	ensure(NodesToSort.Num() == NewArray.Num());
+
+	// Copy out final results
+	NodesToSort = NewArray;
+}
+
+
 FName FNiagaraCompileRequestData::ResolveEmitterAlias(FName VariableName) const
 {
 	return FNiagaraParameterMapHistory::ResolveEmitterAlias(VariableName, EmitterUniqueName);
@@ -509,6 +565,7 @@ void FNiagaraCompileRequestData::CompareAgainst(FNiagaraGraphCachedBuiltHistory*
 	}
 }
 
+
 void FNiagaraCompileRequestData::FinishPrecompile(const TArray<FNiagaraVariable>& EncounterableVariables, const TArray<FNiagaraVariable>& InStaticVariables, FCompileConstantResolver ConstantResolver, const TArray<ENiagaraScriptUsage>& UsagesToProcess, const TArray<class UNiagaraSimulationStageBase*>* SimStages, const TArray<FString> EmitterNames)
 {
 	FNiagaraEditorModule& NiagaraEditorModule = FModuleManager::GetModuleChecked<FNiagaraEditorModule>("NiagaraEditor");
@@ -522,6 +579,8 @@ void FNiagaraCompileRequestData::FinishPrecompile(const TArray<FNiagaraVariable>
 			Source->NodeGraph->FindOutputNodes(OutputNodes);
 		}
 
+		SortOutputNodesByDependencies(OutputNodes, SimStages);
+
 		bool bFilterByEmitterAlias = true;
 		for (UNiagaraNodeOutput* FoundOutputNode : OutputNodes)
 		{
@@ -584,6 +643,8 @@ void FNiagaraCompileRequestData::FinishPrecompile(const TArray<FNiagaraVariable>
 				for (FNiagaraParameterMapHistory& History : Builder.Histories)
 				{
 					History.OriginatingScriptUsage = FoundOutputNode->GetUsage();
+					History.UsageGuid = FoundOutputNode->ScriptTypeId;
+					History.UsageName = SimStageName;
 					for (int32 i = 0; i < History.Variables.Num(); i++)
 					{
 						FNiagaraVariable& Var = History.Variables[i];
@@ -595,6 +656,36 @@ void FNiagaraCompileRequestData::FinishPrecompile(const TArray<FNiagaraVariable>
 
 						if (Var.GetType().IsStatic())
 						{
+							int32 NumValues = 0;
+							int32 LastIndex = INDEX_NONE;
+
+							// The logic for the static variables array adds the full payload static variable to the builder list. This will result
+							// in duplicates with the same name and type, but *different* value payloads. We detect this and error out.
+							for (int32 StaticIdx = 0; StaticIdx < Builder.StaticVariables.Num(); StaticIdx++)
+							{
+								const FNiagaraVariable& BuilderVar = Builder.StaticVariables[StaticIdx];
+								if (Var == BuilderVar) // operator == ignores the value field, which is what we want here 
+								{
+									if (NumValues == 0)
+									{
+										LastIndex = StaticIdx;
+										NumValues++;
+									}
+									else if (LastIndex != INDEX_NONE && !BuilderVar.HoldsSameData(Builder.StaticVariables[LastIndex]))
+									{
+										if (UNiagaraScript::LogCompileStaticVars > 0)
+										{
+											UE_LOG(LogNiagaraEditor, Log, TEXT("Mismatch in static vars %s: \"%s\" vs \"%s\""), *BuilderVar.GetName().ToString(),
+												*BuilderVar.ToString(), *Builder.StaticVariables[LastIndex].ToString());
+										}
+
+										StaticVariablesWithMultipleWrites.AddUnique(Var);
+										NumValues++;
+										break;
+									}
+								} 
+							}
+
 							if (History.PerVariableConstantValue[i].Num() > 1)
 							{
 								StaticVariablesWithMultipleWrites.AddUnique(Var);
@@ -623,14 +714,23 @@ void FNiagaraCompileRequestData::FinishPrecompile(const TArray<FNiagaraVariable>
 					UE_LOG(LogNiagaraEditor, Log, TEXT("Builder.StaticVariables After Param Map Traversal............................"));
 				}
 
-				for (const FNiagaraVariable& Var : Builder.StaticVariables)
+				for (int32 StaticIdx = 0; StaticIdx < Builder.StaticVariables.Num(); StaticIdx++)
 				{
-					StaticVariables.AddUnique(Var);
+					const FNiagaraVariable& Var = Builder.StaticVariables[StaticIdx];
+					bool bProcess = Builder.StaticVariableExportable[StaticIdx];			
 
 					if (UNiagaraScript::LogCompileStaticVars > 0)
 					{
-						UE_LOG(LogNiagaraEditor, Log, TEXT("%s"), *Var.ToString());
+						UE_LOG(LogNiagaraEditor, Log, TEXT("%s > %s"), *Var.ToString(), bProcess ? TEXT("EXPORT") : TEXT("SkipExport"));
 					}
+
+					if (!bProcess)
+					{
+						continue;
+					}
+
+					StaticVariables.AddUnique(Var);
+
 				}
 
 				if (UNiagaraScript::LogCompileStaticVars > 0)
@@ -734,6 +834,8 @@ void FNiagaraCompileRequestDuplicateData::FinishPrecompileDuplicate(const TArray
 		NodeGraphDeepCopy->FindOutputNodes(OutputNodes);
 	}
 
+	FNiagaraCompileRequestData::SortOutputNodesByDependencies(OutputNodes, SimStages);
+
 	for (UNiagaraNodeOutput* FoundOutputNode : OutputNodes)
 	{
 		FName SimStageName;
@@ -777,6 +879,8 @@ void FNiagaraCompileRequestDuplicateData::FinishPrecompileDuplicate(const TArray
 			for (FNiagaraParameterMapHistory& History : Builder.Histories)
 			{
 				History.OriginatingScriptUsage = FoundOutputNode->GetUsage();
+				History.UsageGuid = FoundOutputNode->ScriptTypeId;
+				History.UsageName = SimStageName;
 				for (int32 VarIdx = 0; VarIdx < History.Variables.Num(); VarIdx++)
 				{
 					const FNiagaraVariable& Var = History.Variables[VarIdx];
diff --git a/Engine/Plugins/FX/Niagara/Source/NiagaraEditor/Private/NiagaraHlslTranslator.cpp b/Engine/Plugins/FX/Niagara/Source/NiagaraEditor/Private/NiagaraHlslTranslator.cpp
--- a/Engine/Plugins/FX/Niagara/Source/NiagaraEditor/Private/NiagaraHlslTranslator.cpp
+++ b/Engine/Plugins/FX/Niagara/Source/NiagaraEditor/Private/NiagaraHlslTranslator.cpp
@@ -5282,13 +5282,8 @@ void FHlslNiagaraTranslator::ParameterMapSet(UNiagaraNodeParameterMapSet* SetNod
 					else 
 					{
 						if (Var.GetType().IsStatic())
-						{
-							if (Var.GetName().ToString().Contains("MathFun"))
-							{
-								UE_LOG(LogNiagaraEditor, Log, TEXT("A"));
-							}
-
-							bool bAllSame = true;//CompileDataStaticVariablesWithMultipleWrites;
+						{	
+							bool bAllSame = true;
 							int32 FoundOverrideIdx = CompileData->StaticVariablesWithMultipleWrites.IndexOfByPredicate([&](const FNiagaraVariable& InObj) -> bool
 								{
 									return (InObj.GetName() == Var.GetName());
@@ -5745,10 +5740,20 @@ void FHlslNiagaraTranslator::SetConstantByStaticVariable(int32& OutValue, const
 		const FString* ConstantPtr = CompileData->PinToConstantValues.Find(PinHandle);
 		if (ConstantPtr != nullptr)
 		{
-			TSharedPtr<INiagaraEditorTypeUtilities, ESPMode::ThreadSafe> TypeEditorUtilities = NiagaraEditorModule.GetTypeUtilities(Var.GetType());
-			if (TypeEditorUtilities.IsValid() && TypeEditorUtilities->CanHandlePinDefaults())
+			FNiagaraVariable SearchVar(Var.GetType(), *(*ConstantPtr));
+			int32 StaticVarSearchIdx = CompileData->StaticVariables.Find(SearchVar);
+
+			if (StaticVarSearchIdx == -1)
 			{
-				TypeEditorUtilities->SetValueFromPinDefaultString(*ConstantPtr, VarWithValue);
+				TSharedPtr<INiagaraEditorTypeUtilities, ESPMode::ThreadSafe> TypeEditorUtilities = NiagaraEditorModule.GetTypeUtilities(Var.GetType());
+				if (TypeEditorUtilities.IsValid() && TypeEditorUtilities->CanHandlePinDefaults())
+				{
+					TypeEditorUtilities->SetValueFromPinDefaultString(*ConstantPtr, VarWithValue);
+				}
+			}
+			else
+			{
+				VarWithValue = CompileData->StaticVariables[StaticVarSearchIdx];
 			}
 		}
 
@@ -6129,9 +6134,21 @@ void FHlslNiagaraTranslator::ParameterMapGet(UNiagaraNodeParameterMapGet* GetNod
 				{
 					if (DefaultMode.IsSet() && *DefaultMode == ENiagaraDefaultMode::FailIfPreviouslyNotSet && !Var.IsInNameSpace(FNiagaraConstants::UserNamespaceString))
 					{
+						// Register an external dependency...
 						RegisterCompileDependency(Var, FText::Format(LOCTEXT("UsedBeforeSet", "Variable {0} was read before being set. It's default mode is \"Fail If Previously Not Set\", so this isn't allowed."), FText::FromName(Var.GetName())), GetNode, OutputPins[i], true, ParamMapHistoryIdx);
 					}
 				}
+				else if (DefaultMode.IsSet() && *DefaultMode == ENiagaraDefaultMode::FailIfPreviouslyNotSet && !Var.IsInNameSpace(FNiagaraConstants::UserNamespaceString) && !Var.IsInNameSpace(FNiagaraConstants::ModuleNamespaceString))
+				{
+					// Check for an internal dependency
+					bool bFailIfNotSet = false;
+					FNiagaraVariable TestVar = ActiveHistoryForFunctionCalls.ResolveAliases(Var);
+					ValidateFailIfPreviouslyNotSet(TestVar, bFailIfNotSet);
+					if (bFailIfNotSet)
+					{
+						RegisterCompileDependency(Var, FText::Format(LOCTEXT("UsedBeforeSet", "Variable {0} was read before being set. It's default mode is \"Fail If Previously Not Set\", so this isn't allowed."), FText::FromName(Var.GetName())), GetNode, nullptr, false, ParamMapHistoryIdx);
+					}
+				}
 
 				Outputs[i] = MakeStaticVariableDirect(OutputPins[i]);
 			}
@@ -6153,6 +6170,45 @@ int32 FHlslNiagaraTranslator::MakeStaticVariableDirect(const UEdGraphPin* InDefa
 	return GetConstantDirect(Constant);
 }
 
+void FHlslNiagaraTranslator::ValidateFailIfPreviouslyNotSet(const FNiagaraVariable& InVar, bool& bFailIfNotSet)
+{
+	bFailIfNotSet = false;
+	FNiagaraVariable SearchVar = InVar;
+	if (FNiagaraParameterMapHistory::IsInitialValue(InVar))
+	{
+		SearchVar = FNiagaraParameterMapHistory::GetSourceForInitialValue(InVar);
+	}
+	else if (FNiagaraParameterMapHistory::IsPreviousValue(InVar))
+	{
+		SearchVar = FNiagaraParameterMapHistory::GetSourceForPreviousValue(InVar);
+	}
+
+	FVarAndDefaultSource* ParamMapDefinedVarAndDefaultSource = ParamMapDefinedAttributesToNamespaceVars.Find(SearchVar.GetName());
+	bool bSetPreviously = ParamMapDefinedVarAndDefaultSource != nullptr && ParamMapDefinedVarAndDefaultSource->bDefaultExplicit;
+	for (int32 OtherParamIdx = 0; OtherParamIdx < OtherOutputParamMapHistories.Num() && !bSetPreviously; OtherParamIdx++)
+	{
+		// Stop if this is already in our evaluation chain. Assume only indices above us are valid sourcers for this.
+		if (ParamMapHistoriesSourceInOtherHistories.Contains(OtherParamIdx))
+			break;
+
+		int32 FoundInParamIdx = OtherOutputParamMapHistories[OtherParamIdx].FindVariableByName(SearchVar.GetName());
+		if (INDEX_NONE != FoundInParamIdx)
+		{
+			for (const FModuleScopedPin& ScopedPin : OtherOutputParamMapHistories[OtherParamIdx].PerVariableWriteHistory[FoundInParamIdx])
+			{
+				if (ScopedPin.Pin->Direction == EEdGraphPinDirection::EGPD_Input && ScopedPin.Pin->bHidden == false)
+				{
+					bSetPreviously = true;
+					break;
+				}
+			}
+		}
+	}
+	if (!bSetPreviously && !UNiagaraScript::IsStandaloneScript(CompileOptions.TargetUsage))
+		bFailIfNotSet = true;
+}
+
+
 void FHlslNiagaraTranslator::HandleParameterRead(int32 ParamMapHistoryIdx, const FNiagaraVariable& InVar, const UEdGraphPin* DefaultPin, UNiagaraNode* ErrorNode, int32& OutputChunkId, TOptional<ENiagaraDefaultMode> DefaultMode, TOptional<FNiagaraScriptVariableBinding> DefaultBinding, bool bTreatAsUnknownParameterMap, bool bIgnoreDefaultSetFirst)
 {
 	FString ParameterMapInstanceName = GetParameterMapInstanceName(ParamMapHistoryIdx);
@@ -6263,39 +6319,7 @@ void FHlslNiagaraTranslator::HandleParameterRead(int32 ParamMapHistoryIdx, const
 
 	if (bValidateFailIfPreviouslyNotSet)
 	{
-		FNiagaraVariable SearchVar = Var;
-		if (FNiagaraParameterMapHistory::IsInitialValue(Var))
-		{
-			SearchVar = FNiagaraParameterMapHistory::GetSourceForInitialValue(Var);
-		}
-		else if (FNiagaraParameterMapHistory::IsPreviousValue(Var))
-		{
-			SearchVar = FNiagaraParameterMapHistory::GetSourceForPreviousValue(Var);
-		}
-
-		FVarAndDefaultSource* ParamMapDefinedVarAndDefaultSource = ParamMapDefinedAttributesToNamespaceVars.Find(SearchVar.GetName());
-		bool bSetPreviously = ParamMapDefinedVarAndDefaultSource != nullptr && ParamMapDefinedVarAndDefaultSource->bDefaultExplicit;
-		for (int32 OtherParamIdx = 0; OtherParamIdx < OtherOutputParamMapHistories.Num() && !bSetPreviously; OtherParamIdx++)
-		{
-			// Stop if this is already in our evaluation chain. Assume only indices above us are valid sourcers for this.
-			if (ParamMapHistoriesSourceInOtherHistories.Contains(OtherParamIdx))
-				break;
-
-			int32 FoundInParamIdx = OtherOutputParamMapHistories[OtherParamIdx].FindVariableByName(SearchVar.GetName());
-			if (INDEX_NONE != FoundInParamIdx)
-			{
-				for (const FModuleScopedPin& ScopedPin : OtherOutputParamMapHistories[OtherParamIdx].PerVariableWriteHistory[FoundInParamIdx])
-				{
-					if (ScopedPin.Pin->Direction == EEdGraphPinDirection::EGPD_Input && ScopedPin.Pin->bHidden == false)
-					{
-						bSetPreviously = true;
-						break;
-					}
-				}
-			}
-		}
-		if (!bSetPreviously && !UNiagaraScript::IsStandaloneScript(CompileOptions.TargetUsage))
-			bFailIfPreviouslyNotSetSentinel = true;
+		ValidateFailIfPreviouslyNotSet(Var, bFailIfPreviouslyNotSetSentinel);
 	}
 
 	
@@ -7237,7 +7261,7 @@ void FHlslNiagaraTranslator::CullMapSetInputPin(UEdGraphPin* InputPin)
 }
 
 // From a valid list of namespaces, resolve any aliased tokens and promote namespaced variables without a main namespace to the input parameter map instance namespace
-void FHlslNiagaraTranslator::FinalResolveNamespacedTokens(const FString& ParameterMapInstanceNamespace, TArray<FString>& Tokens, TArray<FString>& ValidChildNamespaces, FNiagaraParameterMapHistoryBuilder& Builder, TArray<FNiagaraVariable>& UniqueParameterMapEntriesAliasesIntact, TArray<FNiagaraVariable>& UniqueParameterMapEntries, int32 ParamMapHistoryIdx)
+void FHlslNiagaraTranslator::FinalResolveNamespacedTokens(const FString& ParameterMapInstanceNamespace, TArray<FString>& Tokens, TArray<FString>& ValidChildNamespaces, FNiagaraParameterMapHistoryBuilder& Builder, TArray<FNiagaraVariable>& UniqueParameterMapEntriesAliasesIntact, TArray<FNiagaraVariable>& UniqueParameterMapEntries, int32 ParamMapHistoryIdx, UNiagaraNode* InNodeForErrorReporting)
 {
 	for (int32 i = 0; i < Tokens.Num(); i++)
 	{
@@ -7289,7 +7313,7 @@ void FHlslNiagaraTranslator::FinalResolveNamespacedTokens(const FString& Paramet
 						}
 						if (!bAdded && !UNiagaraScript::IsStandaloneScript(CompileOptions.TargetUsage)) // Don't warn in modules, they don't have enough context.
 						{
-							Error(FText::Format(LOCTEXT("GetCustomFail1", "Cannot use variable in custom expression, it hasn't been encountered yet: {0}"), FText::FromName(Var.GetName())), nullptr, nullptr);
+							Error(FText::Format(LOCTEXT("GetCustomFail1", "Cannot use variable in custom expression, it hasn't been encountered yet: {0}"), FText::FromName(Var.GetName())), InNodeForErrorReporting, nullptr);
 						}
 
 					}
@@ -7714,7 +7738,7 @@ void FHlslNiagaraTranslator::ProcessCustomHlsl(const FString& InCustomHlsl, ENia
 
 		TArray<FNiagaraVariable> UniqueParamMapEntries;
 		TArray<FNiagaraVariable> UniqueParamMapEntriesAliasesIntact;
-		FinalResolveNamespacedTokens(GetParameterMapInstanceName(0) + TEXT("."), Tokens, PossibleNamespaces, ActiveHistoryForFunctionCalls, UniqueParamMapEntriesAliasesIntact, UniqueParamMapEntries, ParamMapHistoryIdx);
+		FinalResolveNamespacedTokens(GetParameterMapInstanceName(0) + TEXT("."), Tokens, PossibleNamespaces, ActiveHistoryForFunctionCalls, UniqueParamMapEntriesAliasesIntact, UniqueParamMapEntries, ParamMapHistoryIdx, InNodeForErrorReporting);
 
 		// We must register any external constant variables that we encountered.
 		for (int32 VarIdx = 0; VarIdx < UniqueParamMapEntriesAliasesIntact.Num(); VarIdx++)
diff --git a/Engine/Plugins/FX/Niagara/Source/NiagaraEditor/Private/NiagaraHlslTranslator.h b/Engine/Plugins/FX/Niagara/Source/NiagaraEditor/Private/NiagaraHlslTranslator.h
--- a/Engine/Plugins/FX/Niagara/Source/NiagaraEditor/Private/NiagaraHlslTranslator.h
+++ b/Engine/Plugins/FX/Niagara/Source/NiagaraEditor/Private/NiagaraHlslTranslator.h
@@ -128,6 +128,8 @@ public:
 	TArray<FNiagaraVariable> StaticVariablesWithMultipleWrites;
 
 	void CompareAgainst(FNiagaraGraphCachedBuiltHistory* InCachedDataBase);
+
+	static void SortOutputNodesByDependencies(TArray<class UNiagaraNodeOutput*>& NodesToSort, const TArray<class UNiagaraSimulationStageBase*>* SimStages);
 };
 
 class FNiagaraCompileRequestDuplicateData : public FNiagaraCompileRequestDuplicateDataBase
@@ -687,7 +689,7 @@ private:
 	bool ShouldConsiderTargetParameterMap(ENiagaraScriptUsage InUsage) const;
 	FString BuildParameterMapHlslDefinitions(TArray<FNiagaraVariable>& PrimaryDataSetOutputEntries);
 	void BuildMissingDefaults();
-	void FinalResolveNamespacedTokens(const FString& ParameterMapInstanceNamespace, TArray<FString>& Tokens, TArray<FString>& ValidChildNamespaces, FNiagaraParameterMapHistoryBuilder& Builder, TArray<FNiagaraVariable>& UniqueParameterMapEntriesAliasesIntact, TArray<FNiagaraVariable>& UniqueParameterMapEntries, int32 ParamMapHistoryIdx);
+	void FinalResolveNamespacedTokens(const FString& ParameterMapInstanceNamespace, TArray<FString>& Tokens, TArray<FString>& ValidChildNamespaces, FNiagaraParameterMapHistoryBuilder& Builder, TArray<FNiagaraVariable>& UniqueParameterMapEntriesAliasesIntact, TArray<FNiagaraVariable>& UniqueParameterMapEntries, int32 ParamMapHistoryIdx, UNiagaraNode* InNodeForErrorReporting);
 
 	void HandleNamespacedExternalVariablesToDataSetRead(TArray<FNiagaraVariable>& InDataSetVars, FString InNamespaceStr);
 
@@ -743,6 +745,8 @@ private:
 	void GenerateFunctionSignature(ENiagaraScriptUsage ScriptUsage, FString InName, const FString& InFullName, const FString& InFunctionNameSuffix, UNiagaraGraph* FuncGraph, TArray<int32>& Inputs, 
 		bool bHadNumericInputs, bool bHasParameterMapParameters, TArray<UEdGraphPin*> StaticSwitchValues, FNiagaraFunctionSignature& OutSig) const;
 
+	void ValidateFailIfPreviouslyNotSet(const FNiagaraVariable& InVar, bool& bFailIfNotSet);
+
 	bool ShouldInterpolateParameter(const FNiagaraVariable& Parameter);
 
 	void UpdateStaticSwitchConstants(UEdGraphNode* Node);
diff --git a/Engine/Plugins/FX/Niagara/Source/NiagaraEditor/Private/NiagaraNodeEmitter.cpp b/Engine/Plugins/FX/Niagara/Source/NiagaraEditor/Private/NiagaraNodeEmitter.cpp
--- a/Engine/Plugins/FX/Niagara/Source/NiagaraEditor/Private/NiagaraNodeEmitter.cpp
+++ b/Engine/Plugins/FX/Niagara/Source/NiagaraEditor/Private/NiagaraNodeEmitter.cpp
@@ -395,6 +395,35 @@ void UNiagaraNodeEmitter::BuildParameterMapHistory(FNiagaraParameterMapHistoryBu
 				OutHistory.Histories[ParamMapIdx].ParameterCollectionVariables.Append(History.ParameterCollectionVariables);
 				OutHistory.Histories[ParamMapIdx].PinToConstantValues.Append(History.PinToConstantValues);
 			}
+
+			// We only want to push out appropriately scoped static variables that should be in the system builder, not in-betweens like "Module.MyInputVar"
+			// or per-particle or others vars. Really only want Emitter or System parameters here.
+			for (int32 StaticVarIdx = 0; StaticVarIdx < ChildBuilder.StaticVariables.Num(); StaticVarIdx++)
+			{
+				const FNiagaraVariable& ChildStaticVar = ChildBuilder.StaticVariables[StaticVarIdx];
+				if (ChildBuilder.StaticVariableExportable[StaticVarIdx])
+				{
+					// Should match logic in FNiagaraParameterMapHistoryBuilder::RegisterConstantVariableWrite
+					FNiagaraVariable ResolvedStaticVar = FNiagaraUtilities::ResolveAliases(ChildStaticVar, ResolveAliasesContext);
+
+					// Index of uses == operator, which only checks name and type. This will allow us to detect instances of the duplicate
+					// data down the line.
+					int32 FoundStaticVarIdx = OutHistory.StaticVariables.Find(ResolvedStaticVar);
+					
+					if (FoundStaticVarIdx == INDEX_NONE) // Didn't find it, so add it.
+					{
+						OutHistory.StaticVariables.Add(ResolvedStaticVar);
+						OutHistory.StaticVariableExportable.Emplace(true);
+					}
+					else if (false == OutHistory.StaticVariables[FoundStaticVarIdx].HoldsSameData(ResolvedStaticVar))
+					{
+						OutHistory.StaticVariables.Add(ResolvedStaticVar);// Add as a duplicate here. We will filter out later
+						OutHistory.StaticVariableExportable.Emplace(true);
+					}
+
+					ensure(OutHistory.StaticVariables.Num() == OutHistory.StaticVariableExportable.Num());
+				}
+			}
 		}
 
 		OutHistory.EndNodeVisitation(ParamMapIdx, NodeIdx);
diff --git a/Engine/Plugins/FX/Niagara/Source/NiagaraEditor/Private/NiagaraNodeFunctionCall.cpp b/Engine/Plugins/FX/Niagara/Source/NiagaraEditor/Private/NiagaraNodeFunctionCall.cpp
--- a/Engine/Plugins/FX/Niagara/Source/NiagaraEditor/Private/NiagaraNodeFunctionCall.cpp
+++ b/Engine/Plugins/FX/Niagara/Source/NiagaraEditor/Private/NiagaraNodeFunctionCall.cpp
@@ -1095,6 +1095,19 @@ void UNiagaraNodeFunctionCall::BuildParameterMapHistory(FNiagaraParameterMapHist
 
 		FPinCollectorArray InputPins;
 		GetInputPins(InputPins);
+		
+		// Because the pin traversals above the the Super::BuildParameterMapHistory don't traverse into the input pins if they aren't linked.
+		// This leaves static variables that aren't linked to anything basically ignored. The code below fills in that gap.
+		// We don't do it in the base code because many input pins are traversed only if the owning node wants them to be.
+		// OutputNode is another example of manually registering the input pins.
+		for (UEdGraphPin* Pin : InputPins)
+		{
+			// Make sure to register pin constants
+			if (Pin->LinkedTo.Num() == 0 && Schema->IsPinStatic(Pin) && Pin->DefaultValue.Len() != 0)
+			{
+				OutHistory.RegisterConstantFromInputPin(Pin);
+			}
+		}
 
 		FCompileConstantResolver FunctionResolver = OutHistory.ConstantResolver.WithDebugState(DebugState);
 		if (OutHistory.HasCurrentUsageContext())
diff --git a/Engine/Plugins/FX/Niagara/Source/NiagaraEditor/Private/NiagaraParameterMapHistory.cpp b/Engine/Plugins/FX/Niagara/Source/NiagaraEditor/Private/NiagaraParameterMapHistory.cpp
--- a/Engine/Plugins/FX/Niagara/Source/NiagaraEditor/Private/NiagaraParameterMapHistory.cpp
+++ b/Engine/Plugins/FX/Niagara/Source/NiagaraEditor/Private/NiagaraParameterMapHistory.cpp
@@ -44,6 +44,38 @@ static FAutoConsoleVariableRef CVarNiagaraForcePrecompilerCullDataset(
 	ECVF_Default
 );
 
+//#define NIAGARA_SHOULD_DO_PER_VARIABLE_SUBSTRING_MATCH_LOGGING 1
+#if defined(NIAGARA_SHOULD_DO_PER_VARIABLE_SUBSTRING_MATCH_LOGGING)
+static TArray<FString> CheckTheseNameSubstrings;// = { TEXT("XXXX") };
+
+bool NiagaraDebugShouldLogEntriesByNameSubstring(const FString& InNameStr)
+{
+	for (const FString& Str : CheckTheseNameSubstrings)
+	{
+		if (InNameStr.EndsWith(Str))
+		{
+			return true;
+		}
+	}
+
+	return false;
+}
+bool NiagaraDebugShouldLogEntriesByNameSubstring(const FName& InName)
+{
+	if (CheckTheseNameSubstrings.Num() == 0)
+	{
+		return false;
+	}
+	
+	FString NameStr = InName.ToString();
+	return NiagaraDebugShouldLogEntriesByNameSubstring(NameStr);
+}
+#else
+
+#define NiagaraDebugShouldLogEntriesByNameSubstring(expr) (false) // Because the code may do specific allocations, we want to bypass the expression
+
+#endif
+
 void FGraphTraversalHandle::Push(const UEdGraphNode* Node)
 {
 	ensure(Node->NodeGuid.IsValid());
@@ -193,7 +225,7 @@ void FNiagaraParameterMapHistory::RegisterConstantPin(const FGraphTraversalHandl
 		}
 		PinToConstantValues.Add(Handle, InValue);
 
-		if (UNiagaraScript::LogCompileStaticVars > 0)
+		if (UNiagaraScript::LogCompileStaticVars > 0 || NiagaraDebugShouldLogEntriesByNameSubstring(Handle.ToString(true)))
 		{
 			UE_LOG(LogNiagaraEditor, Log, TEXT("RegisterConstantPin \"%s\" Pin: \"%s\""), *InValue, *Handle.ToString(true));
 		}
@@ -201,7 +233,7 @@ void FNiagaraParameterMapHistory::RegisterConstantPin(const FGraphTraversalHandl
 }
 
 
-void FNiagaraParameterMapHistory::RegisterConstantVariableWrite(const FString& InValue, int32 VarIdx, bool bIsSettingDefault)
+void FNiagaraParameterMapHistory::RegisterConstantVariableWrite(const FString& InValue, int32 VarIdx, bool bIsSettingDefault, bool bLinkNotValue)
 {
 	
 	FNiagaraEditorModule& NiagaraEditorModule = FModuleManager::GetModuleChecked<FNiagaraEditorModule>("NiagaraEditor");
@@ -210,7 +242,7 @@ void FNiagaraParameterMapHistory::RegisterConstantVariableWrite(const FString& I
 	{
 		PerVariableConstantValue[VarIdx].AddUnique(InValue);
 
-		if (UNiagaraScript::LogCompileStaticVars > 0)
+		if (UNiagaraScript::LogCompileStaticVars > 0 || NiagaraDebugShouldLogEntriesByNameSubstring(Variables[VarIdx].GetName().ToString()))
 		{
 			UE_LOG(LogNiagaraEditor, Log, TEXT("RegisterConstantVariableWrite Value \"%s\" Var: \"%s\""), *InValue, *Variables[VarIdx].GetName().ToString());
 		}
@@ -219,7 +251,7 @@ void FNiagaraParameterMapHistory::RegisterConstantVariableWrite(const FString& I
 	{
 		PerVariableConstantDefaultValue[VarIdx].AddUnique(InValue);
 
-		if (UNiagaraScript::LogCompileStaticVars > 0)
+		if (UNiagaraScript::LogCompileStaticVars > 0 || NiagaraDebugShouldLogEntriesByNameSubstring(Variables[VarIdx].GetName().ToString()))
 		{
 			UE_LOG(LogNiagaraEditor, Log, TEXT("RegisterConstantVariableWrite Default \"%s\" Var: \"%s\""), *InValue, *Variables[VarIdx].GetName().ToString());
 		}
@@ -423,11 +455,19 @@ bool FNiagaraParameterMapHistory::IsInNamespace(const FNiagaraVariableBase& InVa
 {
 	if (Namespace.EndsWith(TEXT(".")))
 	{
-		return InVar.GetName().ToString().StartsWith(Namespace);
+		bool FastTest = InVar.IsInNameSpace(Namespace);
+		// Leaving old code commented out for now just as a reference to what was before.
+		//bool SlowTest =  InVar.GetName().ToString().StartsWith(Namespace);
+		//check(FastTest == SlowTest);
+		return FastTest;
 	}
 	else
 	{
-		return InVar.GetName().ToString().StartsWith(Namespace + TEXT("."));
+		bool FastTest = InVar.IsInNameSpace(Namespace); // This should check to see if there is a period in the name of the variable after the namespace string.
+		// Leaving old code commented out for now just as a reference to what was before.
+		//bool SlowTest = InVar.GetName().ToString().StartsWith(Namespace + TEXT("."));
+		//check(FastTest == SlowTest);
+		return FastTest;
 	}
 }
 
@@ -973,7 +1013,24 @@ int32 FNiagaraParameterMapHistoryBuilder::RegisterConstantPin(int32 WhichConstan
 	}
 }
 
-int32 FNiagaraParameterMapHistoryBuilder::RegisterConstantVariableWrite(int32 WhichParamMapIdx, int32 WhichConstant,  int32 WhichVarIdx, bool bIsSettingDefault)
+bool FNiagaraParameterMapHistoryBuilder::IsStaticVariableExportableToOuterScopeBasedOnCurrentContext(const FNiagaraVariable& Var) const
+{
+	if ((IsInEncounteredEmitterNamespace(Var) || FNiagaraParameterMapHistory::IsAliasedEmitterParameter(Var)) && (ContextContains(ENiagaraScriptUsage::EmitterSpawnScript) || ContextContains(ENiagaraScriptUsage::EmitterUpdateScript)))
+	{
+		return true;
+	}
+	else if (FNiagaraParameterMapHistory::IsAttribute(Var) && (UNiagaraScript::IsParticleScript(GetBaseUsageContext())))
+	{
+		return true;
+	}
+	else if (FNiagaraParameterMapHistory::IsSystemParameter(Var) && (UNiagaraScript::IsSystemScript(GetBaseUsageContext()) || UNiagaraScript::IsEmitterScript(GetBaseUsageContext())))
+	{
+		return true;
+	}
+	return false;
+}
+
+int32 FNiagaraParameterMapHistoryBuilder::RegisterConstantVariableWrite(int32 WhichParamMapIdx, int32 WhichConstant,  int32 WhichVarIdx, bool bIsSettingDefault, bool bIsLinkNotValue)
 {
 
 	FNiagaraEditorModule& NiagaraEditorModule = FModuleManager::GetModuleChecked<FNiagaraEditorModule>("NiagaraEditor");
@@ -986,10 +1043,10 @@ int32 FNiagaraParameterMapHistoryBuilder::RegisterConstantVariableWrite(int32 Wh
 
 			if (WhichParamMapIdx != INDEX_NONE)
 			{
-				Histories[WhichParamMapIdx].RegisterConstantVariableWrite(Constants[WhichConstant], WhichVarIdx, bIsSettingDefault);
+				Histories[WhichParamMapIdx].RegisterConstantVariableWrite(Constants[WhichConstant], WhichVarIdx, bIsSettingDefault, bIsLinkNotValue);
 			}
 
-			if (Constants[WhichConstant].Len() != 0 && Var.GetType().IsStatic())
+			if (bIsLinkNotValue == false && Constants[WhichConstant].Len() != 0 && Var.GetType().IsStatic())
 			{
 				TSharedPtr<INiagaraEditorTypeUtilities, ESPMode::ThreadSafe> TypeEditorUtilities = NiagaraEditorModule.GetTypeUtilities(Var.GetType());
 				if (TypeEditorUtilities.IsValid() && TypeEditorUtilities->CanHandlePinDefaults())
@@ -998,7 +1055,25 @@ int32 FNiagaraParameterMapHistoryBuilder::RegisterConstantVariableWrite(int32 Wh
 					VarWithValue.AllocateData();
 					if (TypeEditorUtilities->SetValueFromPinDefaultString(Constants[WhichConstant], VarWithValue))
 					{
-						StaticVariables.AddUnique(VarWithValue);
+						// Index of uses == operator, which only checks name and type. This will allow us to detect instances of the duplicate
+						// data down the line.
+						int32 LastStaticVarIdx = StaticVariables.Find(Var);
+						int32 AddedIdx = INDEX_NONE;
+						if (LastStaticVarIdx == INDEX_NONE)
+						{
+							AddedIdx = StaticVariables.Add(VarWithValue); // Didn't find it, so add it.
+						}
+						else if (false == VarWithValue.HoldsSameData(StaticVariables[LastStaticVarIdx]))
+						{
+							AddedIdx = StaticVariables.Add(VarWithValue); // Add as a duplicate here. We will filter out later
+						}
+
+						if (AddedIdx != INDEX_NONE)
+						{
+							bool bExportable = IsStaticVariableExportableToOuterScopeBasedOnCurrentContext(Var);
+							int32 AddedBoolIdx = StaticVariableExportable.Emplace(bExportable);
+							ensure(AddedBoolIdx == AddedIdx);
+						}
 					}
 				}
 			}
@@ -1037,7 +1112,7 @@ int32 FNiagaraParameterMapHistoryBuilder::RegisterConstantFromInputPin(const UEd
 			ConstantIdx = GetConstantFromOutputPin(ConnectedOutputPin);
 			RegisterConstantPin(ConstantIdx, InputPin);
 		}
-		else if (InputPin && InputPin->Direction == EEdGraphPinDirection::EGPD_Input && InputPin->LinkedTo.Num() == 0 && InputPin->DefaultValue.Len() != 0)
+		else if (InputPin && InputPin->Direction == EEdGraphPinDirection::EGPD_Input && InputPin->LinkedTo.Num() == 0 && InputPin->DefaultValue.Len() != 0 && !InputPin->PinName.IsNone())
 		{
 			Value = InputPin->DefaultValue;
 			ConstantIdx = AddOrGetConstantFromValue(Value);
@@ -1045,7 +1120,7 @@ int32 FNiagaraParameterMapHistoryBuilder::RegisterConstantFromInputPin(const UEd
 		}
 	}
 
-	if (UNiagaraScript::LogCompileStaticVars > 0)
+	if (UNiagaraScript::LogCompileStaticVars > 0 || NiagaraDebugShouldLogEntriesByNameSubstring(InputPin->PinName.ToString()))
 	{
 		FGraphTraversalHandle TestPath = ActivePath;
 		TestPath.Push(InputPin);
@@ -1064,7 +1139,7 @@ int32 FNiagaraParameterMapHistoryBuilder::GetConstantFromInputPin(const UEdGraph
 			const int32* IdxPtr = PinToConstantIndices.Last().Find(InputPin);
 			if (IdxPtr != nullptr)
 			{
-				if (UNiagaraScript::LogCompileStaticVars > 0)
+				if (UNiagaraScript::LogCompileStaticVars > 0 || NiagaraDebugShouldLogEntriesByNameSubstring(InputPin->PinName.ToString()))
 				{
 					FGraphTraversalHandle TestPath = ActivePath;
 					TestPath.Push(InputPin);
@@ -1086,7 +1161,7 @@ int32 FNiagaraParameterMapHistoryBuilder::GetConstantFromOutputPin(const UEdGrap
 			const int32* IdxPtr = PinToConstantIndices.Last().Find(OutputPin);
 			if (IdxPtr != nullptr)
 			{
-				if (UNiagaraScript::LogCompileStaticVars > 0)
+				if (UNiagaraScript::LogCompileStaticVars > 0 || NiagaraDebugShouldLogEntriesByNameSubstring(OutputPin->GetName()))
 				{
 					FGraphTraversalHandle TestPath = ActivePath;
 					TestPath.Push(OutputPin);
@@ -1109,7 +1184,7 @@ int32 FNiagaraParameterMapHistoryBuilder::GetConstantFromVariableRead(const FNia
 			const int32* IdxPtr = VariableToConstantIndices[i].Find(InVar);
 			if (IdxPtr != nullptr)
 			{
-				if (UNiagaraScript::LogCompileStaticVars > 0)
+				if (UNiagaraScript::LogCompileStaticVars > 0 || NiagaraDebugShouldLogEntriesByNameSubstring(ActivePath.ToString()))
 				{
 					UE_LOG(LogNiagaraEditor, Log, TEXT("GetConstantFromVariableRead Value \"%s\" Var: \"%s\" Path: \"%s\""), *Constants[*IdxPtr], *InVar.GetName().ToString(),  *ActivePath.ToString());
 				}
@@ -1227,6 +1302,7 @@ int32 FNiagaraParameterMapHistoryBuilder::FindMatchingStaticFromContextInputs(co
 		FNiagaraVariable CallInputVar = Schema->PinToNiagaraVariable(Inputs[i]);
 		if (CallInputVar.IsEquivalent(InVar))
 		{
+			// If linked to something, then the output we are linked to will carry the constant.
 			if (Inputs[i]->LinkedTo.Num() != 0 && PinToConstantIndices.Num() >= 2)
 			{
 				const UEdGraphPin* OutputPin = Inputs[i]->LinkedTo[0];
@@ -1235,6 +1311,13 @@ int32 FNiagaraParameterMapHistoryBuilder::FindMatchingStaticFromContextInputs(co
 				if (IdxPtr != nullptr)
 					return *IdxPtr;
 			}
+			// If just using the default, the input pin will carry the constant.
+			else if (Inputs[i]->LinkedTo.Num() == 0 && PinToConstantIndices.Num() >= 2)
+			{				
+				const int32* IdxPtr = PinToConstantIndices[PinToConstantIndices.Num() - 2].Find(Inputs[i]);
+				if (IdxPtr != nullptr)
+					return *IdxPtr;
+			}
 		}
 	}
 	return INDEX_NONE;
@@ -1422,7 +1505,7 @@ void FNiagaraParameterMapHistoryBuilder::ExitEmitter(const FString& InEmitterNam
 }
 
 
-bool FNiagaraParameterMapHistoryBuilder::IsInEncounteredFunctionNamespace(FNiagaraVariable& InVar) const
+bool FNiagaraParameterMapHistoryBuilder::IsInEncounteredFunctionNamespace(const FNiagaraVariable& InVar) const
 {
 	if (EncounteredFunctionNames.Num() != 0)
 	{
@@ -1437,7 +1520,7 @@ bool FNiagaraParameterMapHistoryBuilder::IsInEncounteredFunctionNamespace(FNiaga
 	return false;
 }
 
-bool FNiagaraParameterMapHistoryBuilder::IsInEncounteredEmitterNamespace(FNiagaraVariable& InVar) const
+bool FNiagaraParameterMapHistoryBuilder::IsInEncounteredEmitterNamespace(const FNiagaraVariable& InVar) const
 {
 	for (FString EmitterEncounteredNamespace : EncounteredEmitterNames)
 	{
@@ -1574,7 +1657,7 @@ int32 FNiagaraParameterMapHistoryBuilder::HandleVariableWrite(int32 ParamMapIdx,
 
 		if (FoundIdx != INDEX_NONE && ParamMapIdx != INDEX_NONE && ConstantIdx != INDEX_NONE)
 		{
-			RegisterConstantVariableWrite(ParamMapIdx, ConstantIdx, FoundIdx, false);
+			RegisterConstantVariableWrite(ParamMapIdx, ConstantIdx, FoundIdx, false, false);
 		}
 	}
 	else if (InPin && Var.GetType().IsStatic() && InPin->LinkedTo.Num() == 0)
@@ -1584,7 +1667,7 @@ int32 FNiagaraParameterMapHistoryBuilder::HandleVariableWrite(int32 ParamMapIdx,
 
 		if (FoundIdx != INDEX_NONE && ParamMapIdx != INDEX_NONE && ConstantIdx != INDEX_NONE)
 		{
-			RegisterConstantVariableWrite(ParamMapIdx, ConstantIdx, FoundIdx, false);
+			RegisterConstantVariableWrite(ParamMapIdx, ConstantIdx, FoundIdx, false, false);
 		}
 	}
 
@@ -1719,6 +1802,8 @@ int32 FNiagaraParameterMapHistoryBuilder::HandleVariableRead(int32 ParamMapIdx,
 				EmitterNamespace = TEXT("Emitter");
 		}
 
+		bool bWroteVarNameToDefaultValue = false;
+
 		if (bUseRapidIterationParam && /*InTopLevelFunctionCall(Usage) && */FNiagaraParameterMapHistory::IsAliasedModuleParameter(Histories[ParamMapIdx].VariablesWithOriginalAliasesIntact[FoundIdx]))
 		{					
 			FNiagaraVariable VarRapid = FNiagaraParameterMapHistory::ConvertVariableToRapidIterationConstantName(Histories[ParamMapIdx].Variables[FoundIdx], EmitterNamespace.Len() == 0 ? nullptr : *EmitterNamespace, Usage);
@@ -1780,6 +1865,7 @@ int32 FNiagaraParameterMapHistoryBuilder::HandleVariableRead(int32 ParamMapIdx,
 			else
 			{
 				DefaultValue = Var.GetName().ToString();
+				bWroteVarNameToDefaultValue = true;
 			}
 			
 			if (UNiagaraScript::LogCompileStaticVars > 0)
@@ -1795,7 +1881,7 @@ int32 FNiagaraParameterMapHistoryBuilder::HandleVariableRead(int32 ParamMapIdx,
 		RegisterConstantPin(ConstantIdx, InPin);
 		if (FoundIdx != INDEX_NONE && ParamMapIdx != INDEX_NONE && ConstantIdx != INDEX_NONE)
 		{
-			RegisterConstantVariableWrite(ParamMapIdx, ConstantIdx, FoundIdx, true);
+			RegisterConstantVariableWrite(ParamMapIdx, ConstantIdx, FoundIdx, true, bWroteVarNameToDefaultValue);
 		}
 	}
 	else if (FoundIdx != INDEX_NONE && ParamMapIdx != INDEX_NONE && UEdGraphSchema_Niagara::IsStaticPin(InPin))
@@ -1814,8 +1900,15 @@ void FNiagaraParameterMapHistoryBuilder::RegisterEncounterableVariables(TConstAr
 
 void FNiagaraParameterMapHistoryBuilder::RegisterExternalStaticVariables(TConstArrayView<FNiagaraVariable> Variables)
 {
+	int32 LastIdx = StaticVariables.Num();
 	StaticVariables.Append(Variables.GetData(), Variables.Num());
 
+	// Register all new entries as exportable true
+	for (int32 i = LastIdx; i < StaticVariables.Num(); i++)
+	{
+		StaticVariableExportable.Emplace(true);
+	}
+
 	if (UNiagaraScript::LogCompileStaticVars > 0)
 	{
 		UE_LOG(LogNiagaraEditor, Log, TEXT("FNiagaraParameterMapHistoryBuilder::RegisterExternalStaticVariables ==================================="));
diff --git a/Engine/Plugins/FX/Niagara/Source/NiagaraEditor/Public/EdGraphSchema_Niagara.h b/Engine/Plugins/FX/Niagara/Source/NiagaraEditor/Public/EdGraphSchema_Niagara.h
--- a/Engine/Plugins/FX/Niagara/Source/NiagaraEditor/Public/EdGraphSchema_Niagara.h
+++ b/Engine/Plugins/FX/Niagara/Source/NiagaraEditor/Public/EdGraphSchema_Niagara.h
@@ -152,6 +152,9 @@ class NIAGARAEDITOR_API UEdGraphSchema_Niagara : public UEdGraphSchema
 	static FEdGraphPinType TypeDefinitionToPinType(FNiagaraTypeDefinition TypeDef);
 	
 	static bool IsPinWildcard(const UEdGraphPin* Pin);
+	
+	// Fast path helper to skip doing full type conversion
+	static bool IsPinStatic(const UEdGraphPin* Pin);
 	static FPinConnectionResponse GetWildcardConnectionResponse(const UEdGraphPin* PinA, const UEdGraphPin* PinB);
 	
 	bool IsSystemConstant(const FNiagaraVariable& Variable)const;
diff --git a/Engine/Plugins/FX/Niagara/Source/NiagaraEditor/Public/NiagaraParameterMapHistory.h b/Engine/Plugins/FX/Niagara/Source/NiagaraEditor/Public/NiagaraParameterMapHistory.h
--- a/Engine/Plugins/FX/Niagara/Source/NiagaraEditor/Public/NiagaraParameterMapHistory.h
+++ b/Engine/Plugins/FX/Niagara/Source/NiagaraEditor/Public/NiagaraParameterMapHistory.h
@@ -118,6 +118,8 @@ public:
 	FNiagaraParameterMapHistory();
 
 	ENiagaraScriptUsage OriginatingScriptUsage;
+	FGuid UsageGuid;
+	FName UsageName;
 
 	/** The variables that have been identified during the traversal. */
 	TArray<FNiagaraVariable> Variables;
@@ -182,7 +184,7 @@ public:
 	*/
 	int32 RegisterParameterMapPin(const UEdGraphPin* Pin);
 
-	void RegisterConstantVariableWrite(const FString& InValue, int32 InVarIdx, bool bIsSettingDefault);
+	void RegisterConstantVariableWrite(const FString& InValue, int32 InVarIdx, bool bIsSettingDefault, bool bIsLinkNotValue);
 
 	uint32 BeginNodeVisitation(const UNiagaraNode* Node);
 	void EndNodeVisitation(uint32 IndexFromBeginNode);
@@ -364,7 +366,7 @@ public:
 	int32 RegisterParameterMapPin(int32 WhichParameterMap, const UEdGraphPin* Pin);
 
 	int32 RegisterConstantPin(int32 WhichConstant, const UEdGraphPin* Pin);
-	int32 RegisterConstantVariableWrite(int32 WhichParamMapIdx, int32 WhichConstant, int32 WhichVarIdx, bool bIsSettingDefault);
+	int32 RegisterConstantVariableWrite(int32 WhichParamMapIdx, int32 WhichConstant, int32 WhichVarIdx, bool bIsSettingDefault, bool bLinkNotValue);
 
 	int32  RegisterConstantFromInputPin(const UEdGraphPin* InputPin);
 
@@ -498,8 +500,8 @@ public:
 	bool GetIgnoreDisabled() const { return bIgnoreDisabled; }
 	void SetIgnoreDisabled(bool bInIgnore) { bIgnoreDisabled = bInIgnore; }
 
-	bool IsInEncounteredFunctionNamespace(FNiagaraVariable& InVar) const;
-	bool IsInEncounteredEmitterNamespace(FNiagaraVariable& InVar) const;
+	bool IsInEncounteredFunctionNamespace(const FNiagaraVariable& InVar) const;
+	bool IsInEncounteredEmitterNamespace(const FNiagaraVariable& InVar) const;
 
 	/** Register any user or other external variables that could possibly be encountered but may not be declared explicitly. */
 	void RegisterEncounterableVariables(TConstArrayView<FNiagaraVariable> Variables);
@@ -517,6 +519,7 @@ public:
 	int32 FindStaticVariable(const FNiagaraVariable& Var) const;
 
 	TArray<FNiagaraVariable> StaticVariables;
+	TArray<bool> StaticVariableExportable; // Should we export out these static variables to calling context?
 
 	void GetContextuallyVisitedNodes(TArray<const class UNiagaraNode*>& OutVistedNodes)
 	{
@@ -524,6 +527,10 @@ public:
 			OutVistedNodes.Append(ContextuallyVisitedNodes.Last());
 	}
 protected:
+	/** Internal helper function to decide whether or not we should use this static variable when merging with other
+		graph traversal static variables, like Emitter graphs merging with System graphs. */
+	bool IsStaticVariableExportableToOuterScopeBasedOnCurrentContext(const FNiagaraVariable& Var) const;
+
 	/**
 	* Generate the internal alias map from the current traversal state.
 	*/
diff --git a/Engine/Source/Runtime/Core/Private/UObject/UE5ReleaseStreamObjectVersion.cpp b/Engine/Source/Runtime/Core/Private/UObject/UE5ReleaseStreamObjectVersion.cpp
--- a/Engine/Source/Runtime/Core/Private/UObject/UE5ReleaseStreamObjectVersion.cpp
+++ b/Engine/Source/Runtime/Core/Private/UObject/UE5ReleaseStreamObjectVersion.cpp
@@ -12,7 +12,7 @@ TMap<FGuid, FGuid> FUE5ReleaseStreamObjectVersion::GetSystemGuids()
 	SystemGuids.Add(DevGuids.MATERIALSHADERMAP_DERIVEDDATA_VER, FGuid("79090A66F9D94B5BB285D49E5D39468E"));
 	SystemGuids.Add(DevGuids.NANITE_DERIVEDDATA_VER, FGuid("AC88CFBDDA614A9C8488F64D2DAF699D"));
 	SystemGuids.Add(DevGuids.NIAGARASHADERMAP_DERIVEDDATA_VER, FGuid("D204CCC166ED482685DAEE8AB72895AB"));
-	SystemGuids.Add(DevGuids.Niagara_LatestScriptCompileVersion, FGuid("8AF4589E892441FABA3E330276A74B92"));
+	SystemGuids.Add(DevGuids.Niagara_LatestScriptCompileVersion, FGuid("29F4804724014091BC9B1AF11BA53330"));
 	SystemGuids.Add(DevGuids.SkeletalMeshDerivedDataVersion, FGuid("25C49E579B3142DDA2A8C14037267679"));
 
 	return SystemGuids;
