Optimizations to preparation of niagara rapid iteration parameters - a substantial cost of loading niagara systems when serialized data doesn't match runtime (including when pulling from ddc).  Shaves off about 70% of the cost of preparing rapid iteration parameters (about 10% of total cook time for a very niagara heavy map).
-relax cache key for module stack pins, allows much more reuse of past results
-make sure that gathering static variables doesn't re-traverse the same scripts
-clean up some dynamic string allocations
-when looking for module inputs make sure that we don't recurse into function graphs beyond what is necessary

diff --git a/Engine/Plugins/FX/Niagara/Source/NiagaraEditor/Private/NiagaraCompiler.cpp b/Engine/Plugins/FX/Niagara/Source/NiagaraEditor/Private/NiagaraCompiler.cpp
--- a/Engine/Plugins/FX/Niagara/Source/NiagaraEditor/Private/NiagaraCompiler.cpp
+++ b/Engine/Plugins/FX/Niagara/Source/NiagaraEditor/Private/NiagaraCompiler.cpp
@@ -1460,29 +1460,31 @@ TSharedPtr<FNiagaraGraphCachedDataBase, ESPMode::ThreadSafe> FNiagaraEditorModul
 	const UNiagaraSystem* System = Cast<UNiagaraSystem>(Obj);
 	const UNiagaraEmitter* Emitter = Cast<UNiagaraEmitter>(Obj);
 
-	/*
-	* const TArray<FNiagaraVariable>& EncounterableVariables, const TArray<FNiagaraVariable>& InStaticVariables, FCompileConstantResolver ConstantResolver, const TArray<ENiagaraScriptUsage>& UsagesToProcess, const 
-	*/
 	const TArray<class UNiagaraSimulationStageBase*>* SimStages = nullptr;
 	const UNiagaraScriptSource* ScriptSource = nullptr;
 	FCompileConstantResolver ConstantResolver;
-	TArray<FNiagaraVariable> EncounterableVariables;
-	TArray<FNiagaraVariable> StaticVariablesFromSystem;
-	TArray<FNiagaraVariable> StaticVariablesFromSystemAndEmitters;
 	TArray<FNiagaraVariable> SrcStaticVariables;
 	TSharedPtr<FNiagaraGraphCachedDataBase, ESPMode::ThreadSafe> ParentCachedData;
 	FString SrcUniqueEmitterName;
 
+	int32 ValidUsageMask = 0;
+	bool bIncludeSubHistories = true;
+
 	if (System)
 	{
+		TArray<FNiagaraVariable> StaticVariablesFromEmitters;
+
 		ScriptSource = CastChecked<UNiagaraScriptSource>(System->GetSystemSpawnScript()->GetLatestSource());
 		ConstantResolver = FCompileConstantResolver(System, ENiagaraScriptUsage::SystemSpawnScript);
-		System->GatherStaticVariables(StaticVariablesFromSystem, StaticVariablesFromSystemAndEmitters);
-
-		SrcStaticVariables = StaticVariablesFromSystem;
+		System->GatherStaticVariables(SrcStaticVariables, StaticVariablesFromEmitters);
 
-		for (const FNiagaraVariable& Var : StaticVariablesFromSystemAndEmitters)
+		for (const FNiagaraVariable& Var : StaticVariablesFromEmitters)
+		{
 			SrcStaticVariables.AddUnique(Var);
+		}
+
+		ValidUsageMask = ENiagaraScriptUsageMask::System;
+		bIncludeSubHistories = false;
 	}
 	else if (Emitter)
 	{
@@ -1495,13 +1497,6 @@ TSharedPtr<FNiagaraGraphCachedDataBase, ESPMode::ThreadSafe> FNiagaraEditorModul
 			{
 				SrcStaticVariables = ((FNiagaraGraphCachedBuiltHistory*)ParentCachedData.Get())->StaticVariables;
 			}
-			/*for (const FNiagaraEmitterHandle& Handle : SysParent->GetEmitterHandles())
-			{
-				if (Handle.GetInstance() == Emitter)
-				{
-					SrcUniqueEmitterName = Handle.GetUniqueEmitterName().ToString();
-				}
-			}*/
 		}
 		SrcUniqueEmitterName = Emitter->GetUniqueEmitterName();
 		const FVersionedNiagaraEmitterData* EmitterData = Emitter->GetEmitterData(Version);
@@ -1510,8 +1505,9 @@ TSharedPtr<FNiagaraGraphCachedDataBase, ESPMode::ThreadSafe> FNiagaraEditorModul
 		ScriptSource = CastChecked<UNiagaraScriptSource>(EmitterData->GraphSource);
 		SimStages = &EmitterData->GetSimulationStages();
 		ConstantResolver = FCompileConstantResolver(FVersionedNiagaraEmitter(const_cast<UNiagaraEmitter*>(Emitter), Version), ENiagaraScriptUsage::EmitterSpawnScript);
-	}
 
+		ValidUsageMask = ENiagaraScriptUsageMask::Particle | ENiagaraScriptUsageMask::Emitter;
+	}
 
 	TArray<UNiagaraNodeOutput*> OutputNodes;
 	if (ScriptSource != nullptr && ScriptSource->NodeGraph != nullptr)
@@ -1535,6 +1531,14 @@ TSharedPtr<FNiagaraGraphCachedDataBase, ESPMode::ThreadSafe> FNiagaraEditorModul
 	int32 NumSimStageNodes = 0;
 	for (UNiagaraNodeOutput* FoundOutputNode : OutputNodes)
 	{
+		// only process those nodes that we care about
+		const bool bShouldSkipOutputNode = !(ValidUsageMask & (1 << (int32)FoundOutputNode->GetUsage()));
+
+		if (bShouldSkipOutputNode)
+		{
+			continue;
+		}
+
 		FName SimStageName;
 		bool bStageEnabled = true;
 		if (FoundOutputNode->GetUsage() == ENiagaraScriptUsage::ParticleSimulationStageScript && SimStages)
@@ -1542,9 +1546,9 @@ TSharedPtr<FNiagaraGraphCachedDataBase, ESPMode::ThreadSafe> FNiagaraEditorModul
 			// Find the simulation stage for this output node.
 			const FGuid& UsageId = FoundOutputNode->GetUsageId();
 			UNiagaraSimulationStageBase* const* MatchingStagePtr = SimStages->FindByPredicate([UsageId](UNiagaraSimulationStageBase* SimStage)
-				{
+			{
 					return SimStage != nullptr && SimStage->Script != nullptr && SimStage->Script->GetUsageId() == UsageId;
-				});
+			});
 
 			// Set whether or not the stage is enabled, and get the iteration source name if available.
 			bStageEnabled = MatchingStagePtr != nullptr && (*MatchingStagePtr)->bEnabled;
@@ -1565,18 +1569,20 @@ TSharedPtr<FNiagaraGraphCachedDataBase, ESPMode::ThreadSafe> FNiagaraEditorModul
 			{
 				Builder.AddGraphToCallingGraphContextStack(ScriptSource->NodeGraph);
 			}
-			Builder.RegisterEncounterableVariables(EncounterableVariables);
 			Builder.RegisterExternalStaticVariables(StaticVariables);
 
 			FString TranslationName = TEXT("Emitter");
 			Builder.BeginTranslation(TranslationName);
 			Builder.BeginUsage(FoundOutputNode->GetUsage(), SimStageName);
 			Builder.EnableScriptAllowList(true, FoundOutputNode->GetUsage());
+			Builder.bShouldBuildSubHistories = bIncludeSubHistories;
 			Builder.BuildParameterMaps(FoundOutputNode, true);
 			Builder.EndUsage();
 
 			for (const FNiagaraVariable& Var : Builder.StaticVariables)
+			{
 				StaticVariables.AddUnique(Var);
+			}
 		}
 	}
 
diff --git a/Engine/Plugins/FX/Niagara/Source/NiagaraEditor/Private/NiagaraGraphDataCache.cpp b/Engine/Plugins/FX/Niagara/Source/NiagaraEditor/Private/NiagaraGraphDataCache.cpp
--- a/Engine/Plugins/FX/Niagara/Source/NiagaraEditor/Private/NiagaraGraphDataCache.cpp
+++ b/Engine/Plugins/FX/Niagara/Source/NiagaraEditor/Private/NiagaraGraphDataCache.cpp
@@ -4,135 +4,181 @@
 
 #include "NiagaraNodeFunctionCall.h"
 
-void FNiagaraGraphDataCache::FCachedStackFunctionInputPins::SetSource(UNiagaraNodeFunctionCall& FunctionCallNode)
+static int32 GNiagaraGraphDataCacheSize = 16384;
+static FAutoConsoleVariableRef CVarNiagaraGraphDataCacheSize(
+	TEXT("fx.Niagara.GraphDataCacheSize"),
+	GNiagaraGraphDataCacheSize,
+	TEXT("Maximum number of elements to store within the GraphDataCache."),
+	ECVF_ReadOnly
+);
+
+static bool GNiagaraGraphDataCacheValidation = false;
+static FAutoConsoleVariableRef CVarNiagaraGraphDataCacheValidation(
+	TEXT("fx.Niagara.GraphDataCacheValidation"),
+	GNiagaraGraphDataCacheValidation,
+	TEXT("If true will perform validation on retrieving data from the data FNiagaraGraphDataCache."),
+	ECVF_Default
+);
+
+FNiagaraGraphDataCache::FNiagaraGraphDataCache()
+	: StackFunctionInputPinCache(GNiagaraGraphDataCacheSize)
 {
-	SourceGraphWeak = FunctionCallNode.GetNiagaraGraph();
-	LastSourceGraphChangeId = FunctionCallNode.GetNiagaraGraph()->GetChangeID();
-	FunctionCallNodeWeak = &FunctionCallNode;
-	CalledGraphWeak = FunctionCallNode.GetCalledGraph();
-	LastCalledGraphChangeId = FunctionCallNode.GetCalledGraph() != nullptr ? FunctionCallNode.GetCalledGraph()->GetChangeID() : FGuid();
 }
 
-void FNiagaraGraphDataCache::GetStackFunctionInputPins(UNiagaraNodeFunctionCall& FunctionCallNode, TConstArrayView<FNiagaraVariable> StaticVars, TArray<const UEdGraphPin*>& OutInputPins, FCompileConstantResolver ConstantResolver, FNiagaraStackGraphUtilities::ENiagaraGetStackFunctionInputPinsOptions Options, bool bIgnoreDisabled)
+void FNiagaraGraphDataCache::GetStackFunctionInputPinsInternal(
+	UNiagaraNodeFunctionCall& FunctionCallNode,
+	const UNiagaraGraph* CalledGraph,
+	TConstArrayView<FNiagaraVariable> StaticVars,
+	TArray<const UEdGraphPin*>& OutInputPins,
+	FCompileConstantResolver ConstantResolver,
+	FNiagaraStackGraphUtilities::ENiagaraGetStackFunctionInputPinsOptions Options,
+	bool bIgnoreDisabled,
+	bool bFilterForCompilation)
 {
-	FGetStackFunctionInputPinsKey Key(FunctionCallNode, ConstantResolver, Options, bIgnoreDisabled, false);
-	FCachedStackFunctionInputPins& CachedResult = GetStackFunctionInputPinsCache.FindOrAdd(Key);
-	if (!CachedResult.IsValid())
+	// we don't currently support hashing the internal details of the translator within FStackFunctionInputPinKey so in the case
+	// we've been supplied a translator we'll just fall back to not using the cache
+	if (const FHlslNiagaraTranslator* Translator = ConstantResolver.GetTranslator())
 	{
 		FNiagaraStackGraphUtilities::GetStackFunctionInputPinsWithoutCache(
 			FunctionCallNode,
 			StaticVars,
-			CachedResult.InputPins,
+			OutInputPins,
 			ConstantResolver,
 			Options,
 			bIgnoreDisabled,
-			false /*bFilterForCompilation*/);
+			bFilterForCompilation);
 
-		CachedResult.SetSource(FunctionCallNode);
+		return;
 	}
 
-	OutInputPins = CachedResult.InputPins;
-}
+	FStackFunctionInputPinKey CacheKey(FunctionCallNode, CalledGraph, StaticVars, ConstantResolver, Options, bIgnoreDisabled, bFilterForCompilation);
 
-void FNiagaraGraphDataCache::GetStackFunctionInputPins(UNiagaraNodeFunctionCall& FunctionCallNode, TConstArrayView<FNiagaraVariable> StaticVars, TArray<const UEdGraphPin*>& OutInputPins, TSet<const UEdGraphPin*>& OutHiddenPins, FCompileConstantResolver ConstantResolver, FNiagaraStackGraphUtilities::ENiagaraGetStackFunctionInputPinsOptions Options, bool bIgnoreDisabled)
-{
-	FGetStackFunctionInputPinsKey Key(FunctionCallNode, ConstantResolver, Options, bIgnoreDisabled, false);
-	FGetStackFunctionInputPinsKey CompilationKey(FunctionCallNode, ConstantResolver, Options, bIgnoreDisabled, true);
+	bool FoundInCache = false;
 
-	if (!GetStackFunctionInputPinsCache.Contains(Key))
+	if (const FStackFunctionInputPinValue* CachedValue = StackFunctionInputPinCache.FindAndTouch(CacheKey))
 	{
-		GetStackFunctionInputPinsCache.Add(Key);
+		FoundInCache = true;
+		OutInputPins = CachedValue->InputPins;
 	}
-
-	if (!GetStackFunctionInputPinsCache.Contains(CompilationKey))
-	{
-		GetStackFunctionInputPinsCache.Add(CompilationKey);
-	}
-
-	FCachedStackFunctionInputPins& CachedResult = GetStackFunctionInputPinsCache.FindChecked(Key);
-	FCachedStackFunctionInputPins& CompilationCachedResult = GetStackFunctionInputPinsCache.FindChecked(CompilationKey);
-
-	if (!CachedResult.IsValid())
+	else
 	{
+		FStackFunctionInputPinValue NewValue;
 		FNiagaraStackGraphUtilities::GetStackFunctionInputPinsWithoutCache(
 			FunctionCallNode,
 			StaticVars,
-			CachedResult.InputPins,
+			NewValue.InputPins,
 			ConstantResolver,
 			Options,
 			bIgnoreDisabled,
-			false /*bFilterForCompilation*/);
+			bFilterForCompilation);
 
-		CachedResult.SetSource(FunctionCallNode);
+		OutInputPins = NewValue.InputPins;
+
+		StackFunctionInputPinCache.Add(CacheKey, NewValue);
 	}
 
-	if (!CompilationCachedResult.IsValid())
+	if (GNiagaraGraphDataCacheValidation)
 	{
+		TArray<const UEdGraphPin*> ValidationPins;
 		FNiagaraStackGraphUtilities::GetStackFunctionInputPinsWithoutCache(
 			FunctionCallNode,
 			StaticVars,
-			CompilationCachedResult.InputPins,
+			ValidationPins,
 			ConstantResolver,
 			Options,
 			bIgnoreDisabled,
-			true /*bFilterForCompilation*/);
+			bFilterForCompilation);
 
-		CompilationCachedResult.SetSource(FunctionCallNode);
+		if (!ensureMsgf(ValidationPins == OutInputPins, TEXT("FNiagaraGraphDataCache failed to accurately collect StackFunctionInputPins for Function {0}"), *FunctionCallNode.GetPathName()))
+		{
+			OutInputPins = ValidationPins;
+		}
 	}
+}
 
-	OutInputPins = CachedResult.InputPins;
-
-	// generate hidden pins
-	auto PinsMatch = [](const UEdGraphPin* Lhs, const UEdGraphPin* Rhs) -> bool
+void FNiagaraGraphDataCache::GetStackFunctionInputPins(UNiagaraNodeFunctionCall& FunctionCallNode, TConstArrayView<FNiagaraVariable> StaticVars, TArray<const UEdGraphPin*>& OutInputPins, FCompileConstantResolver ConstantResolver, FNiagaraStackGraphUtilities::ENiagaraGetStackFunctionInputPinsOptions Options, bool bIgnoreDisabled)
+{
+	if (const UNiagaraGraph* CalledGraph = FunctionCallNode.GetCalledGraph())
 	{
-		return Lhs->GetFName() == Rhs->GetFName()
-			&& Lhs->PinType.PinCategory == Rhs->PinType.PinCategory
-			&& Lhs->PinType.PinSubCategoryObject == Rhs->PinType.PinSubCategoryObject;
-	};
+		GetStackFunctionInputPinsInternal(FunctionCallNode, CalledGraph, StaticVars, OutInputPins, ConstantResolver, Options, bIgnoreDisabled, false /*bFilterForCompilation*/);
+	}
+}
 
-	for (const UEdGraphPin* InputPin : CachedResult.InputPins)
+void FNiagaraGraphDataCache::GetStackFunctionInputPins(UNiagaraNodeFunctionCall& FunctionCallNode, TConstArrayView<FNiagaraVariable> StaticVars, TArray<const UEdGraphPin*>& OutInputPins, TSet<const UEdGraphPin*>& OutHiddenPins, FCompileConstantResolver ConstantResolver, FNiagaraStackGraphUtilities::ENiagaraGetStackFunctionInputPinsOptions Options, bool bIgnoreDisabled)
+{
+	if (const UNiagaraGraph* CalledGraph = FunctionCallNode.GetCalledGraph())
 	{
-		if (!CompilationCachedResult.InputPins.ContainsByPredicate([&](const UEdGraphPin* CompilationPin) { return PinsMatch(InputPin, CompilationPin); }))
+		TArray<const UEdGraphPin*> FilteredPins;
+
+		GetStackFunctionInputPinsInternal(FunctionCallNode, CalledGraph, StaticVars, OutInputPins, ConstantResolver, Options, bIgnoreDisabled, false /*bFilterForCompilation*/);
+		GetStackFunctionInputPinsInternal(FunctionCallNode, CalledGraph, StaticVars, FilteredPins, ConstantResolver, Options, bIgnoreDisabled, true /*bFilterForCompilation*/);
+
+		// generate hidden pins
+		auto PinsMatch = [](const UEdGraphPin* Lhs, const UEdGraphPin* Rhs) -> bool
 		{
-			OutHiddenPins.Add(InputPin);
+			return Lhs->GetFName() == Rhs->GetFName()
+				&& Lhs->PinType.PinCategory == Rhs->PinType.PinCategory
+				&& Lhs->PinType.PinSubCategoryObject == Rhs->PinType.PinSubCategoryObject;
+		};
+
+		for (const UEdGraphPin* InputPin : OutInputPins)
+		{
+			if (!FilteredPins.ContainsByPredicate([&](const UEdGraphPin* CompilationPin) { return PinsMatch(InputPin, CompilationPin); }))
+			{
+				OutHiddenPins.Add(InputPin);
+			}
 		}
 	}
 }
 
-FNiagaraGraphDataCache::FGetStackFunctionInputPinsKey::FGetStackFunctionInputPinsKey(UNiagaraNodeFunctionCall& InFunctionCallNode, const FCompileConstantResolver& InConstantResolver, FNiagaraStackGraphUtilities::ENiagaraGetStackFunctionInputPinsOptions InOptions, bool bInIgnoreDisabled, bool bInFilterForCompilation)
+FNiagaraGraphDataCache::FStackFunctionInputPinKey::FStackFunctionInputPinKey(const UNiagaraNodeFunctionCall& FunctionCallNode, const UNiagaraGraph* CalledGraph, TConstArrayView<FNiagaraVariable> StaticVariables, const FCompileConstantResolver& InConstantResolver, FNiagaraStackGraphUtilities::ENiagaraGetStackFunctionInputPinsOptions InOptions, bool bInIgnoreDisabled, bool bInFilterForCompilation)
 {
-	FunctionCallNodeKey = FObjectKey(&InFunctionCallNode);
-
-	ConstantResolverSystemKey = FObjectKey(InConstantResolver.GetSystem());
-	ConstantResolverEmitterKey = FObjectKey(InConstantResolver.GetEmitter());
-	ConstantResolverTranslator = InConstantResolver.GetTranslator();
-	ConstantResolverUsage = InConstantResolver.GetUsage();
-	ConstantResolverDebugState = InConstantResolver.GetDebugState();
+	if (ensure(CalledGraph))
+	{
+		FunctionCallGraphKey = FObjectKey(CalledGraph);
+		FunctionCallGraphChangeId = CalledGraph->GetChangeID();
+		ContextStaticVariables = StaticVariables;
+	}
 
-	Options = InOptions;
+	InputPinOption = (uint8)InputPinOption;
 	bIgnoreDisabled = bInIgnoreDisabled;
 	bFilterForCompilation = bInFilterForCompilation;
+	ConstantResolverHash = InConstantResolver.BuildTypeHash();
 
-	Hash = GetTypeHash(FunctionCallNodeKey);
-	Hash = HashCombine(Hash, GetTypeHash(ConstantResolverSystemKey));
-	Hash = HashCombine(Hash, GetTypeHash(ConstantResolverEmitterKey));
-	Hash = HashCombine(Hash, GetTypeHash(ConstantResolverTranslator));
-	Hash = HashCombine(Hash, GetTypeHash(ConstantResolverUsage));
-	Hash = HashCombine(Hash, GetTypeHash(ConstantResolverDebugState));
-
-	Hash = HashCombine(Hash, GetTypeHash(Options));
-	Hash = HashCombine(Hash, GetTypeHash(bIgnoreDisabled));
-	Hash = HashCombine(Hash, GetTypeHash(bFilterForCompilation));
+	Hash = GetTypeHash(FunctionCallGraphKey);
+	Hash = HashCombine(Hash, GetTypeHash(FunctionCallGraphChangeId));
+	for (const FNiagaraVariable& StaticVariable : ContextStaticVariables)
+	{
+		Hash = HashCombine(Hash, GetTypeHash(StaticVariable));
+		const uint8* StaticVariableData = StaticVariable.GetData();
+		const int32 ByteCount = StaticVariable.GetAllocatedSizeInBytes();
+		for (int32 ByteIt = 0; ByteIt < ByteCount; ++ByteIt)
+		{
+			Hash = HashCombine(Hash, GetTypeHash(StaticVariableData[ByteIt]));
+		}
+	}
+	Hash = HashCombine(Hash, ConstantResolverHash);
+	Hash = HashCombine(Hash, InputPinOption);
+	Hash = HashCombine(Hash, bIgnoreDisabled);
+	Hash = HashCombine(Hash, bFilterForCompilation);
+
+	// the default values of the input pins can be relevant when it comes to static switch values defined within the graph
+	FPinCollectorArray InputPins;
+	FunctionCallNode.GetInputPins(InputPins);
+	for (const UEdGraphPin* InputPin : InputPins)
+	{
+		Hash = HashCombine(Hash, GetTypeHash(InputPin->PinName));
+		Hash = HashCombine(Hash, GetTypeHash(InputPin->DefaultValue));
+	}
 }
 
-bool FNiagaraGraphDataCache::FCachedStackFunctionInputPins::IsValid() const
+bool FNiagaraGraphDataCache::FStackFunctionInputPinKey::operator==(const FStackFunctionInputPinKey& Other) const
 {
-	return
-		SourceGraphWeak.IsValid() &&
-		FunctionCallNodeWeak.IsValid() &&
-		CalledGraphWeak.IsValid() &&
-
-		SourceGraphWeak->GetChangeID() == LastSourceGraphChangeId &&
-		FunctionCallNodeWeak->GetCalledGraph() == CalledGraphWeak.Get() &&
-		(CalledGraphWeak.IsValid() == false || CalledGraphWeak->GetChangeID() == LastCalledGraphChangeId);
+	return FunctionCallGraphKey == Other.FunctionCallGraphKey
+		&& FunctionCallGraphChangeId == Other.FunctionCallGraphChangeId
+		&& ContextStaticVariables == Other.ContextStaticVariables
+		&& ConstantResolverHash == Other.ConstantResolverHash
+		&& InputPinOption == Other.InputPinOption
+		&& bIgnoreDisabled == Other.bIgnoreDisabled
+		&& bFilterForCompilation == Other.bFilterForCompilation;
 }
diff --git a/Engine/Plugins/FX/Niagara/Source/NiagaraEditor/Private/NiagaraNodeFunctionCall.cpp b/Engine/Plugins/FX/Niagara/Source/NiagaraEditor/Private/NiagaraNodeFunctionCall.cpp
--- a/Engine/Plugins/FX/Niagara/Source/NiagaraEditor/Private/NiagaraNodeFunctionCall.cpp
+++ b/Engine/Plugins/FX/Niagara/Source/NiagaraEditor/Private/NiagaraNodeFunctionCall.cpp
@@ -1056,8 +1056,26 @@ bool UNiagaraNodeFunctionCall::HasValidScriptAndGraph() const
 void UNiagaraNodeFunctionCall::BuildParameterMapHistory(FNiagaraParameterMapHistoryBuilder& OutHistory, bool bRecursive /*= true*/, bool bFilterForCompilation /*= true*/) const
 {
 	const UEdGraphSchema_Niagara* Schema = CastChecked<UEdGraphSchema_Niagara>(GetSchema());
-	const UEdGraphPin* ParamMapPin = GetInputPin(0);
-	bool bHasParamMapPin = ParamMapPin && Schema->PinToTypeDefinition(ParamMapPin) == FNiagaraTypeDefinition::GetParameterMapDef();
+
+	auto FindParameterMapPinIndex = [&]()
+	{
+		FPinCollectorArray InputPins;
+		GetInputPins(InputPins);
+		for (int32 PinIndex = 0; PinIndex < InputPins.Num(); ++PinIndex)
+		{
+			if (Schema->PinToTypeDefinition(InputPins[PinIndex]) == FNiagaraTypeDefinition::GetParameterMapDef())
+			{
+				return PinIndex;
+			}
+		}
+		return int32(INDEX_NONE);
+	};
+
+	// when dealing with function scripts it's not sufficient to assume that InputPin[0] is the ParamMap pin, and so we
+	// need to search through the inputs to find the right one
+	const int32 ParamMapPinIndex = FindParameterMapPinIndex();
+	const bool bHasParamMapPin = ParamMapPinIndex != INDEX_NONE;
+	const UEdGraphPin* ParamMapPin = bHasParamMapPin ? GetInputPin(ParamMapPinIndex) : nullptr;
 	if (bHasParamMapPin && ParamMapPin->LinkedTo.Num() == 0)
 	{
 		// Looks like this function call is not yet hooked up. Skip it to prevent cascading errors in the compilation
@@ -1122,7 +1140,7 @@ void UNiagaraNodeFunctionCall::BuildParameterMapHistory(FNiagaraParameterMapHist
 
 		if (bHasParamMapPin && bRecursive)
 		{
-			ParamMapIdx = OutHistory.TraceParameterMapOutputPin(UNiagaraNode::TraceOutputPin(GetInputPin(0)->LinkedTo[0]));
+			ParamMapIdx = OutHistory.TraceParameterMapOutputPin(UNiagaraNode::TraceOutputPin(ParamMapPin->LinkedTo[0]));
 		}
 
 		OutHistory.EnterFunction(GetFunctionName(), FunctionScript, FunctionGraph, this);
@@ -1130,7 +1148,18 @@ void UNiagaraNodeFunctionCall::BuildParameterMapHistory(FNiagaraParameterMapHist
 		{
 			NodeIdx = OutHistory.BeginNodeVisitation(ParamMapIdx, this);
 		}
-		OutputNode->BuildParameterMapHistory(OutHistory, true, bFilterForCompilation);
+
+		// check if we should be recursing deeper into the graph
+		const bool DoDepthTraversal = (OutHistory.MaxGraphDepthTraversal == INDEX_NONE || OutHistory.CurrentGraphDepth < OutHistory.MaxGraphDepthTraversal);
+
+		if (DoDepthTraversal)
+		{
+			++OutHistory.CurrentGraphDepth;
+
+			OutputNode->BuildParameterMapHistory(OutHistory, true, bFilterForCompilation);
+
+			--OutHistory.CurrentGraphDepth;
+		}
 
 		// Since we're about to lose the pin calling context, we finish up the function call parameter map pin wiring
 		// here when we have the calling context and the child context still available to us...
@@ -1197,14 +1226,28 @@ void UNiagaraNodeFunctionCall::BuildParameterMapHistory(FNiagaraParameterMapHist
 
 		OutHistory.ExitFunction(GetFunctionName(), FunctionScript, this);
 
-		for (int32 i = 0; i < MatchedPairs.Num(); i++)
+		if (DoDepthTraversal)
 		{
-			OutHistory.RegisterParameterMapPin(MatchedPairs[i].Value, MatchedPairs[i].Key);
-		}
+			for (int32 i = 0; i < MatchedPairs.Num(); i++)
+			{
+				OutHistory.RegisterParameterMapPin(MatchedPairs[i].Value, MatchedPairs[i].Key);
+			}
 
-		for (int32 i = 0; i < MatchedConstants.Num(); i++)
+			for (int32 i = 0; i < MatchedConstants.Num(); i++)
+			{
+				OutHistory.RegisterConstantPin(MatchedConstants[i].Value, MatchedConstants[i].Key);
+			}
+		}
+		else
 		{
-			OutHistory.RegisterConstantPin(MatchedConstants[i].Value, MatchedConstants[i].Key);
+			for (int32 OutputPinIt = 0; OutputPinIt < OutputPinCount; ++OutputPinIt)
+			{
+				if (OutputVariables[OutputPinIt].GetType() == FNiagaraTypeDefinition::GetParameterMapDef())
+				{
+					OutHistory.RegisterParameterMapPin(ParamMapIdx, OutputPins[OutputPinIt]);
+					break;
+				}
+			}
 		}
 	}
 	else if (Signature.bRequiresExecPin)
diff --git a/Engine/Plugins/FX/Niagara/Source/NiagaraEditor/Private/NiagaraParameterMapHistory.cpp b/Engine/Plugins/FX/Niagara/Source/NiagaraEditor/Private/NiagaraParameterMapHistory.cpp
--- a/Engine/Plugins/FX/Niagara/Source/NiagaraEditor/Private/NiagaraParameterMapHistory.cpp
+++ b/Engine/Plugins/FX/Niagara/Source/NiagaraEditor/Private/NiagaraParameterMapHistory.cpp
@@ -2364,6 +2364,49 @@ FCompileConstantResolver FCompileConstantResolver::WithUsage(ENiagaraScriptUsage
 	return Copy;
 }
 
+uint32 FCompileConstantResolver::BuildTypeHash() const
+{
+	union
+	{
+		struct  
+		{
+			uint8 Usage;
+			uint8 DebugState;
+			uint8 EmitterSimTarget;
+			uint8 bHasSystem : 1;
+			uint8 bHasEmitter : 1;
+			uint8 bEmitterLocalSpace : 1;
+			uint8 bEmitterDeterminism : 1;
+			uint8 bEmitterInterpolatedSpawn : 1;
+		};
+		uint32 Hash;
+	} HashUnion;
+
+	HashUnion.Hash = 0;
+
+	if (Emitter.Emitter)
+	{
+		HashUnion.bHasEmitter = true;
+		if (FVersionedNiagaraEmitterData* EmitterData = Emitter.GetEmitterData())
+		{
+			HashUnion.EmitterSimTarget = (uint8)EmitterData->SimTarget;
+			HashUnion.bEmitterLocalSpace = EmitterData->bLocalSpace;
+			HashUnion.bEmitterDeterminism = EmitterData->bDeterminism;
+			HashUnion.bEmitterInterpolatedSpawn = EmitterData->bInterpolatedSpawning;
+		}
+	}
+
+	if (System)
+	{
+		HashUnion.bHasSystem = true;
+	}
+
+	HashUnion.Usage = (uint8)Usage;
+	HashUnion.DebugState = (uint8)DebugState;
+
+	return HashUnion.Hash;
+}
+
 int32 FNiagaraParameterMapHistoryWithMetaDataBuilder::AddVariableToHistory(FNiagaraParameterMapHistory& History, const FNiagaraVariable& InVar, const FNiagaraVariable& InAliasedVar, const UEdGraphPin* InPin)
 {
 	const FString* ModuleAlias = GetModuleAlias();
diff --git a/Engine/Plugins/FX/Niagara/Source/NiagaraEditor/Private/NiagaraScriptSource.cpp b/Engine/Plugins/FX/Niagara/Source/NiagaraEditor/Private/NiagaraScriptSource.cpp
--- a/Engine/Plugins/FX/Niagara/Source/NiagaraEditor/Private/NiagaraScriptSource.cpp
+++ b/Engine/Plugins/FX/Niagara/Source/NiagaraEditor/Private/NiagaraScriptSource.cpp
@@ -299,32 +299,32 @@ bool UNiagaraScriptSource::AddModuleIfMissing(FString ModulePath, ENiagaraScript
 	return false;
 }
 
-void UNiagaraScriptSource::FixupRenamedParameters(UNiagaraNodeFunctionCall* FunctionCallNode, TConstArrayView<const UEdGraphPin*> ModuleInputPins, FNiagaraParameterStore& RapidIterationParameters, const TArray<FNiagaraVariable>& OldRapidIterationVariables, const FVersionedNiagaraEmitter& VersionedEmitter, ENiagaraScriptUsage ScriptUsage) const
+void UNiagaraScriptSource::FixupRenamedParameters(UNiagaraNodeFunctionCall* FunctionCallNode, TConstArrayView<const UEdGraphPin*> ModuleInputPins, TConstArrayView<FNiagaraVariable> ModulePinVariables, FNiagaraParameterStore& RapidIterationParameters, const TArray<FNiagaraVariable>& OldRapidIterationVariables, const FVersionedNiagaraEmitter& VersionedEmitter, ENiagaraScriptUsage ScriptUsage) const
 {
 	// the rapid iteration parameters and the function input pins use different variable naming schemes, so most of this is just used to convert one name to the other 
 	UNiagaraGraph* Graph = FunctionCallNode->GetCalledGraph();
-	const UEdGraphSchema_Niagara* NiagaraSchema = Graph->GetNiagaraSchema();
 	const FString UniqueEmitterName = VersionedEmitter.Emitter ? VersionedEmitter.Emitter->GetUniqueEmitterName() : FString();
 
-	TArray<FNiagaraVariable, TInlineAllocator<16>> ModulePinVariables;
 	TArray<TOptional<FNiagaraVariableMetaData>, TInlineAllocator<16>> ModulePinMetaData;
 	TArray<FName, TInlineAllocator<16>> ModulePinAliasedName;
 	TArray<FString, TInlineAllocator<16>> ModulePinConstantName;
 	TArray<int32, TInlineAllocator<16>> ModulePinNameSpaceEnd;
 
 	const int32 ModuleInputPinCount = ModuleInputPins.Num();
-	ModulePinVariables.Reserve(ModuleInputPinCount);
 	ModulePinMetaData.Reserve(ModuleInputPinCount);
 	ModulePinAliasedName.Reserve(ModuleInputPinCount);
 	ModulePinConstantName.Reserve(ModuleInputPinCount);
 	ModulePinNameSpaceEnd.Reserve(ModuleInputPinCount);
 
-	for (const UEdGraphPin* ModuleInputPin : ModuleInputPins)
+	const FName FunctionName = *FunctionCallNode->GetFunctionName();
+
+	for (int32 ModuleInputPinIt = 0; ModuleInputPinIt < ModuleInputPinCount; ++ModuleInputPinIt)
 	{
-		const FNiagaraVariable& InputVar = ModulePinVariables.Add_GetRef(NiagaraSchema->PinToNiagaraVariable(ModuleInputPin));
-		ModulePinMetaData.Add(Graph->GetMetaData(InputVar));
+		const UEdGraphPin* ModuleInputPin = ModuleInputPins[ModuleInputPinIt];
+		const FNiagaraVariable& InputVar = ModulePinVariables[ModuleInputPinIt];
 
-		FNiagaraParameterHandle AliasedFunctionInputHandle = FNiagaraParameterHandle::CreateAliasedModuleParameterHandle(FNiagaraParameterHandle(ModuleInputPin->PinName), FunctionCallNode);
+		ModulePinMetaData.Add(Graph->GetMetaData(InputVar));
+		FNiagaraParameterHandle AliasedFunctionInputHandle = FNiagaraParameterHandle::CreateAliasedModuleParameterHandle(ModuleInputPin->PinName, FunctionName);
 		const FName& AliasedName = ModulePinAliasedName.Add_GetRef(AliasedFunctionInputHandle.GetParameterHandleString());
 		const FString& ConstantName = ModulePinConstantName.Add_GetRef(FNiagaraUtilities::CreateRapidIterationConstantName(AliasedName, *UniqueEmitterName, ScriptUsage));
 
@@ -393,7 +393,7 @@ void UNiagaraScriptSource::FixupRenamedParameters(UNiagaraNodeFunctionCall* Func
 	}
 }
 
-void UNiagaraScriptSource::InitializeNewParameters(UNiagaraNodeFunctionCall* FunctionCallNode, TConstArrayView<const UEdGraphPin*> FunctionInputPins, FNiagaraParameterStore& RapidIterationParameters, const FVersionedNiagaraEmitter& VersionedEmitter, ENiagaraScriptUsage ScriptUsage, TSet<FName>& ValidRapidIterationParameterNames) const
+void UNiagaraScriptSource::InitializeNewParameters(UNiagaraNodeFunctionCall* FunctionCallNode, TConstArrayView<const UEdGraphPin*> FunctionInputPins, TConstArrayView<FNiagaraVariable> PinVariables, FNiagaraParameterStore& RapidIterationParameters, const FVersionedNiagaraEmitter& VersionedEmitter, ENiagaraScriptUsage ScriptUsage, TSet<FName>& ValidRapidIterationParameterNames) const
 {
 	const FString UniqueEmitterName = VersionedEmitter.Emitter ? VersionedEmitter.Emitter->GetUniqueEmitterName() : FString();
 	UNiagaraGraph* Graph = FunctionCallNode->GetCalledGraph();
@@ -410,10 +410,12 @@ void UNiagaraScriptSource::InitializeNewParameters(UNiagaraNodeFunctionCall* Fun
 	AliasHandles.Reserve(FunctionInputPinCount);
 	Parameters.Reserve(FunctionInputPinCount);
 
+	const FName FunctionNodeName = *FunctionCallNode->GetFunctionName();
 
-	for (const UEdGraphPin* FunctionInputPin : FunctionInputPins)
+	for (int32 FunctionInputPinIt = 0; FunctionInputPinIt < FunctionInputPinCount; ++FunctionInputPinIt)
 	{
-		FNiagaraTypeDefinition InputType = NiagaraSchema->PinToTypeDefinition(FunctionInputPin);
+		const UEdGraphPin* FunctionInputPin = FunctionInputPins[FunctionInputPinIt];
+		const FNiagaraTypeDefinition& InputType = PinVariables[FunctionInputPinIt].GetType();
 		if (InputType.IsValid() == false)
 		{
 			UE_LOG(LogNiagaraEditor, Error, TEXT("Invalid input type found while attempting initialize new rapid iteration parameters. Function Node: %s %s Input Name: %s"),
@@ -423,7 +425,7 @@ void UNiagaraScriptSource::InitializeNewParameters(UNiagaraNodeFunctionCall* Fun
 
 		if (FNiagaraStackGraphUtilities::IsRapidIterationType(InputType))
 		{
-			FNiagaraParameterHandle AliasedFunctionInputHandle = FNiagaraParameterHandle::CreateAliasedModuleParameterHandle(FNiagaraParameterHandle(FunctionInputPin->PinName), FunctionCallNode);
+			FNiagaraParameterHandle AliasedFunctionInputHandle = FNiagaraParameterHandle::CreateAliasedModuleParameterHandle(FunctionInputPin->PinName, FunctionNodeName);
 			FNiagaraVariable RapidIterationParameter = FNiagaraStackGraphUtilities::CreateRapidIterationParameter(UniqueEmitterName, ScriptUsage, AliasedFunctionInputHandle.GetParameterHandleString(), InputType);
 			ValidRapidIterationParameterNames.Add(RapidIterationParameter.GetName());
 			int32 ParameterIndex = RapidIterationParameters.IndexOf(RapidIterationParameter);
@@ -508,12 +510,29 @@ void UNiagaraScriptSource::CleanUpOldAndInitializeNewRapidIterationParameters(co
 				}
 
 				// find out which inputs the module offers
+				UNiagaraGraph* Graph = FunctionCallNode->GetCalledGraph();
+				const UEdGraphSchema_Niagara* NiagaraSchema = Graph->GetNiagaraSchema();
+
 				TArray<const UEdGraphPin*> ModuleInputPins;
 				FCompileConstantResolver ConstantResolver(Emitter, ScriptUsage);
+
 				GetStackFunctionInputPins(*FunctionCallNode, ModuleInputPins, ConstantResolver, FNiagaraStackGraphUtilities::ENiagaraGetStackFunctionInputPinsOptions::ModuleInputsOnly);
 
-				FixupRenamedParameters(FunctionCallNode, ModuleInputPins, RapidIterationParameters, OldRapidIterationVariables, Emitter, ScriptUsage);
-				InitializeNewParameters(FunctionCallNode, ModuleInputPins, RapidIterationParameters, Emitter, OutputNode->GetUsage(), ValidRapidIterationParameterNames);
+				// if we found any module input pins, then we can continue with fixing up and initializing them
+				if (!ModuleInputPins.IsEmpty())
+				{
+					TArray<FNiagaraVariable> ModuleInputPinVariables;
+					ModuleInputPinVariables.Reserve(ModuleInputPins.Num());
+					for (const UEdGraphPin* InputPin : ModuleInputPins)
+					{
+						// PinToNiagaraVariable can be costly (when in sufficient numbers) so do this work once for
+						// both FixupRenamedParameters & InitializeNewParameters
+						ModuleInputPinVariables.Add(NiagaraSchema->PinToNiagaraVariable(InputPin));
+					}
+
+					FixupRenamedParameters(FunctionCallNode, ModuleInputPins, ModuleInputPinVariables, RapidIterationParameters, OldRapidIterationVariables, Emitter, ScriptUsage);
+					InitializeNewParameters(FunctionCallNode, ModuleInputPins, ModuleInputPinVariables, RapidIterationParameters, Emitter, OutputNode->GetUsage(), ValidRapidIterationParameterNames);
+				}
 			}
 		}
 	}
diff --git a/Engine/Plugins/FX/Niagara/Source/NiagaraEditor/Private/ViewModels/NiagaraCurveSelectionViewModel.cpp b/Engine/Plugins/FX/Niagara/Source/NiagaraEditor/Private/ViewModels/NiagaraCurveSelectionViewModel.cpp
--- a/Engine/Plugins/FX/Niagara/Source/NiagaraEditor/Private/ViewModels/NiagaraCurveSelectionViewModel.cpp
+++ b/Engine/Plugins/FX/Niagara/Source/NiagaraEditor/Private/ViewModels/NiagaraCurveSelectionViewModel.cpp
@@ -512,7 +512,7 @@ TSharedPtr<FNiagaraCurveSelectionTreeNode> UNiagaraCurveSelectionViewModel::Crea
 	const UEdGraphSchema_Niagara* NiagaraSchema = GetDefault<UEdGraphSchema_Niagara>();
 	TArray<const UEdGraphPin*> InputPins;
 	TSet<const UEdGraphPin*> HiddenPins;
-	FNiagaraStackGraphUtilities::GetStackFunctionInputPins(FunctionCallNode, InputPins, HiddenPins, ConstantResolver);
+	FNiagaraStackGraphUtilities::GetStackFunctionInputPins(FunctionCallNode, InputPins, HiddenPins, ConstantResolver, FNiagaraStackGraphUtilities::ENiagaraGetStackFunctionInputPinsOptions::ModuleInputsOnly);
 	for (const UEdGraphPin* InputPin : InputPins)
 	{
 		if (HiddenPins.Contains(InputPin))
diff --git a/Engine/Plugins/FX/Niagara/Source/NiagaraEditor/Private/ViewModels/Stack/NiagaraParameterHandle.cpp b/Engine/Plugins/FX/Niagara/Source/NiagaraEditor/Private/ViewModels/Stack/NiagaraParameterHandle.cpp
--- a/Engine/Plugins/FX/Niagara/Source/NiagaraEditor/Private/ViewModels/Stack/NiagaraParameterHandle.cpp
+++ b/Engine/Plugins/FX/Niagara/Source/NiagaraEditor/Private/ViewModels/Stack/NiagaraParameterHandle.cpp
@@ -15,11 +15,13 @@ FNiagaraParameterHandle::FNiagaraParameterHandle(FName InParameterHandleName)
 	: ParameterHandleName(InParameterHandleName)
 {
 	int32 DotIndex;
-	const FString ParameterHandleString = ParameterHandleName.ToString();
+	FNameBuilder ParameterHandleBuilder(ParameterHandleName);
+	FStringView ParameterHandleString(ParameterHandleBuilder);
+
 	if (ParameterHandleString.FindChar(TEXT('.'), DotIndex))
 	{
-		Name = *ParameterHandleString.RightChop(DotIndex + 1);
-		Namespace = *ParameterHandleString.Left(DotIndex);
+		Name = FName(ParameterHandleString.RightChop(DotIndex + 1));
+		Namespace = FName(ParameterHandleString.Left(DotIndex));
 	}
 	else
 	{
@@ -31,7 +33,11 @@ FNiagaraParameterHandle::FNiagaraParameterHandle(FName InNamespace, FName InName
 	: Name(InName)
 	, Namespace(InNamespace)
 {
-	ParameterHandleName = *FString::Printf(TEXT("%s.%s"), *Namespace.ToString(), *Name.ToString());
+	FNameBuilder ParameterHandleBuilder(Namespace);
+	ParameterHandleBuilder.AppendChar(TCHAR('.'));
+	InName.AppendString(ParameterHandleBuilder);
+
+	ParameterHandleName = FName(ParameterHandleBuilder);
 }
 
 bool FNiagaraParameterHandle::operator==(const FNiagaraParameterHandle& Other) const
@@ -46,6 +52,23 @@ FNiagaraParameterHandle FNiagaraParameterHandle::CreateAliasedModuleParameterHan
 		: ModuleParameterHandle;
 }
 
+FNiagaraParameterHandle FNiagaraParameterHandle::CreateAliasedModuleParameterHandle(const FName FullInputName, const FName FunctionName)
+{
+	FNameBuilder ParameterHandleBuilder(FullInputName);
+	FStringView ParamHandleString(ParameterHandleBuilder);
+
+	const int32 ModuleNamespaceLen = FNiagaraConstants::ModuleNamespaceString.Len();
+
+	if (ParamHandleString.Len() > ModuleNamespaceLen
+		&& ParamHandleString[ModuleNamespaceLen] == TCHAR('.')
+		&& ParamHandleString.StartsWith(FNiagaraConstants::ModuleNamespaceString))
+	{
+		return FNiagaraParameterHandle(FunctionName, FName(ParamHandleString.RightChop(ModuleNamespaceLen + 1)));
+	}
+
+	return FNiagaraParameterHandle(FullInputName);
+}
+
 FNiagaraParameterHandle FNiagaraParameterHandle::CreateEngineParameterHandle(const FNiagaraVariable& SystemVariable)
 {
 	return FNiagaraParameterHandle(SystemVariable.GetName());
@@ -53,7 +76,6 @@ FNiagaraParameterHandle FNiagaraParameterHandle::CreateEngineParameterHandle(con
 
 FNiagaraParameterHandle FNiagaraParameterHandle::CreateEmitterParameterHandle(const FNiagaraVariable& EmitterVariable)
 {
-
 	return FNiagaraParameterHandle(FNiagaraConstants::EmitterNamespace, EmitterVariable.GetName());
 }
 
@@ -69,7 +91,12 @@ FNiagaraParameterHandle FNiagaraParameterHandle::CreateModuleParameterHandle(con
 
 FNiagaraParameterHandle FNiagaraParameterHandle::CreateInitialParameterHandle(const FNiagaraParameterHandle& Handle)
 {
-	return FNiagaraParameterHandle(Handle.GetNamespace(), *FString::Printf(TEXT("%s.%s"), *FNiagaraConstants::InitialPrefix, *Handle.GetName().ToString()));
+	FNameBuilder NameBuilder;
+	NameBuilder.Append(FNiagaraConstants::InitialPrefix);
+	NameBuilder.AppendChar(TCHAR('.'));
+	Handle.GetName().AppendString(NameBuilder);
+
+	return FNiagaraParameterHandle(Handle.GetNamespace(), FName(FStringView(NameBuilder)));
 }
 
 bool FNiagaraParameterHandle::IsValid() const 
diff --git a/Engine/Plugins/FX/Niagara/Source/NiagaraEditor/Private/ViewModels/Stack/NiagaraStackGraphUtilities.cpp b/Engine/Plugins/FX/Niagara/Source/NiagaraEditor/Private/ViewModels/Stack/NiagaraStackGraphUtilities.cpp
--- a/Engine/Plugins/FX/Niagara/Source/NiagaraEditor/Private/ViewModels/Stack/NiagaraStackGraphUtilities.cpp
+++ b/Engine/Plugins/FX/Niagara/Source/NiagaraEditor/Private/ViewModels/Stack/NiagaraStackGraphUtilities.cpp
@@ -942,6 +942,12 @@ void FNiagaraStackGraphUtilities::GetStackFunctionInputPinsWithoutCache(
 	Builder.ConstantResolver = ConstantResolver;
 	Builder.RegisterExternalStaticVariables(StaticVars);
 
+	// if we are only dealing with the module input pins then we don't need to delve deep into the graph
+	if (Options == ENiagaraGetStackFunctionInputPinsOptions::ModuleInputsOnly)
+	{
+		Builder.MaxGraphDepthTraversal = 1;
+	}
+
 	FunctionCallNode.BuildParameterMapHistory(Builder, false, bFilterForCompilation);
 
 	OutInputPins.Empty();
diff --git a/Engine/Plugins/FX/Niagara/Source/NiagaraEditor/Private/ViewModels/Stack/NiagaraStackScriptItemGroup.cpp b/Engine/Plugins/FX/Niagara/Source/NiagaraEditor/Private/ViewModels/Stack/NiagaraStackScriptItemGroup.cpp
--- a/Engine/Plugins/FX/Niagara/Source/NiagaraEditor/Private/ViewModels/Stack/NiagaraStackScriptItemGroup.cpp
+++ b/Engine/Plugins/FX/Niagara/Source/NiagaraEditor/Private/ViewModels/Stack/NiagaraStackScriptItemGroup.cpp
@@ -353,7 +353,7 @@ private:
 		UNiagaraNodeAssignment* NewAssignmentModule = FNiagaraStackGraphUtilities::AddParameterModuleToStack(Vars, *OutputNode, TargetIndex,DefaultVals );
 		
 		TArray<const UEdGraphPin*> InputPins;
-		FNiagaraStackGraphUtilities::GetStackFunctionInputPins(*NewAssignmentModule, InputPins);
+		FNiagaraStackGraphUtilities::GetStackFunctionInputPins(*NewAssignmentModule, InputPins, FNiagaraStackGraphUtilities::ENiagaraGetStackFunctionInputPinsOptions::AllInputs);
 		if (InputPins.Num() == 1)
 		{
 			FString FunctionInputEditorDataKey = FNiagaraStackGraphUtilities::GenerateStackFunctionInputEditorDataKey(*NewAssignmentModule, InputPins[0]->PinName);
diff --git a/Engine/Plugins/FX/Niagara/Source/NiagaraEditor/Public/NiagaraGraph.h b/Engine/Plugins/FX/Niagara/Source/NiagaraEditor/Public/NiagaraGraph.h
--- a/Engine/Plugins/FX/Niagara/Source/NiagaraEditor/Public/NiagaraGraph.h
+++ b/Engine/Plugins/FX/Niagara/Source/NiagaraEditor/Public/NiagaraGraph.h
@@ -257,7 +257,7 @@ class UNiagaraGraph : public UEdGraph
 	/** Each graph is given a Change Id that occurs anytime the graph's content is manipulated. This key changing induces several important activities, including being a 
 	value that third parties can poll to see if their cached handling of the graph needs to potentially adjust to changes. Furthermore, for script compilation we cache 
 	the changes that were produced during the traversal of each output node, which are referred to as the CompileID.*/
-	FGuid GetChangeID() { return ChangeId; }
+	FGuid GetChangeID() const { return ChangeId; }
 
 	/** Gets the current compile data hash associated with the output node traversal specified by InUsage and InUsageId. If the usage is not found, an invalid hash is returned.*/
 	FNiagaraCompileHash GetCompileDataHash(ENiagaraScriptUsage InUsage, const FGuid& InUsageId) const;
diff --git a/Engine/Plugins/FX/Niagara/Source/NiagaraEditor/Public/NiagaraGraphDataCache.h b/Engine/Plugins/FX/Niagara/Source/NiagaraEditor/Public/NiagaraGraphDataCache.h
--- a/Engine/Plugins/FX/Niagara/Source/NiagaraEditor/Public/NiagaraGraphDataCache.h
+++ b/Engine/Plugins/FX/Niagara/Source/NiagaraEditor/Public/NiagaraGraphDataCache.h
@@ -2,6 +2,7 @@
 
 #pragma once
 
+#include "Containers/LruCache.h"
 #include "NiagaraCommon.h"
 #include "NiagaraEditorUtilities.h"
 #include "UObject/ObjectKey.h"
@@ -16,64 +17,55 @@ class UNiagaraNodeFunctionCall;
 class FNiagaraGraphDataCache
 {
 public:
+	FNiagaraGraphDataCache();
+
 	void GetStackFunctionInputPins(UNiagaraNodeFunctionCall& FunctionCallNode, TConstArrayView<FNiagaraVariable> StaticVars, TArray<const UEdGraphPin*>& OutInputPins, FCompileConstantResolver ConstantResolver, FNiagaraStackGraphUtilities::ENiagaraGetStackFunctionInputPinsOptions Options, bool bIgnoreDisabled);
 	void GetStackFunctionInputPins(UNiagaraNodeFunctionCall& FunctionCallNode, TConstArrayView<FNiagaraVariable> StaticVars, TArray<const UEdGraphPin*>& OutInputPins, TSet<const UEdGraphPin*>& OutHiddenPins, FCompileConstantResolver ConstantResolver, FNiagaraStackGraphUtilities::ENiagaraGetStackFunctionInputPinsOptions Options, bool bIgnoreDisabled);
 
 private:
-	struct FGetStackFunctionInputPinsKey
+	void GetStackFunctionInputPinsInternal(
+		UNiagaraNodeFunctionCall& FunctionCallNode,
+		const UNiagaraGraph* CalledGraph,
+		TConstArrayView<FNiagaraVariable> StaticVars,
+		TArray<const UEdGraphPin*>& OutInputPins,
+		FCompileConstantResolver ConstantResolver,
+		FNiagaraStackGraphUtilities::ENiagaraGetStackFunctionInputPinsOptions Options,
+		bool bIgnoreDisabled,
+		bool bFilterForCompilation);
+
+	struct FStackFunctionInputPinKey
 	{
-		FGetStackFunctionInputPinsKey(UNiagaraNodeFunctionCall& InFunctionCallNode, const FCompileConstantResolver& InConstantResolver, FNiagaraStackGraphUtilities::ENiagaraGetStackFunctionInputPinsOptions InOptions, bool bInIgnoreDisabled, bool bFilterForCompilation);
+		FStackFunctionInputPinKey(const UNiagaraNodeFunctionCall& FunctionCallNode, const UNiagaraGraph* CalledGraph, TConstArrayView<FNiagaraVariable> StaticVazriables, const FCompileConstantResolver& InConstantResolver, FNiagaraStackGraphUtilities::ENiagaraGetStackFunctionInputPinsOptions InOptions, bool bInIgnoreDisabled, bool bInFilterForCompilation);
+		bool operator==(const FStackFunctionInputPinKey& Other) const;
 
-		FORCEINLINE bool operator == (const FGetStackFunctionInputPinsKey& Other) const
+		friend uint32 GetTypeHash(const FStackFunctionInputPinKey& Key)
 		{
-			return
-				FunctionCallNodeKey == Other.FunctionCallNodeKey &&
-				ConstantResolverSystemKey == Other.ConstantResolverSystemKey &&
-				ConstantResolverEmitterKey == Other.ConstantResolverEmitterKey &&
-				ConstantResolverTranslator == Other.ConstantResolverTranslator &&
-				ConstantResolverUsage == Other.ConstantResolverUsage &&
-				ConstantResolverDebugState == Other.ConstantResolverDebugState &&
-				Options == Other.Options &&
-				bIgnoreDisabled == Other.bIgnoreDisabled &&
-				bFilterForCompilation == Other.bFilterForCompilation;
+			return Key.Hash;
 		}
 
-		friend uint32 GetTypeHash(const FGetStackFunctionInputPinsKey& Key)
+		const FGuid& GetGraphChangeId() const
 		{
-			return Key.Hash;
+			return FunctionCallGraphChangeId;
 		}
-		
-	private:
-		FObjectKey FunctionCallNodeKey;
-
-		FObjectKey ConstantResolverSystemKey;
-		FObjectKey ConstantResolverEmitterKey;
-		const FHlslNiagaraTranslator* ConstantResolverTranslator;
-		ENiagaraScriptUsage ConstantResolverUsage;
-		ENiagaraFunctionDebugState ConstantResolverDebugState;
 
-		FNiagaraStackGraphUtilities::ENiagaraGetStackFunctionInputPinsOptions Options;
-		bool bIgnoreDisabled;
-		bool bFilterForCompilation;
+	private:
+		FObjectKey FunctionCallGraphKey;
+		FGuid FunctionCallGraphChangeId;
+		TArray<FNiagaraVariable> ContextStaticVariables;
+		uint32 ConstantResolverHash = 0;
+		uint8 InputPinOption = 0;
+		bool bIgnoreDisabled = false;
+		bool bFilterForCompilation = false;
 
-		int32 Hash;
+		int32 Hash = 0;
 	};
 
-	struct FCachedStackFunctionInputPins
+	struct FStackFunctionInputPinValue
 	{
-		void SetSource(UNiagaraNodeFunctionCall& InFunctionCallNode);
-		bool IsValid() const;
-
 		TArray<const UEdGraphPin*> InputPins;
-
-	private:
-		TWeakObjectPtr<UNiagaraGraph> SourceGraphWeak;
-		TWeakObjectPtr<UNiagaraNodeFunctionCall> FunctionCallNodeWeak;
-		TWeakObjectPtr<UNiagaraGraph> CalledGraphWeak;
-
-		FGuid LastSourceGraphChangeId;
-		FGuid LastCalledGraphChangeId;
 	};
 
-	TMap<FGetStackFunctionInputPinsKey, FCachedStackFunctionInputPins> GetStackFunctionInputPinsCache;
+	// cache holding a finite number (specific value pulled from a cvar) of results to previous calls
+	// to retrieve the input pins
+	TLruCache<FStackFunctionInputPinKey, FStackFunctionInputPinValue> StackFunctionInputPinCache;
 };
\ No newline at end of file
diff --git a/Engine/Plugins/FX/Niagara/Source/NiagaraEditor/Public/NiagaraParameterMapHistory.h b/Engine/Plugins/FX/Niagara/Source/NiagaraEditor/Public/NiagaraParameterMapHistory.h
--- a/Engine/Plugins/FX/Niagara/Source/NiagaraEditor/Public/NiagaraParameterMapHistory.h
+++ b/Engine/Plugins/FX/Niagara/Source/NiagaraEditor/Public/NiagaraParameterMapHistory.h
@@ -34,6 +34,12 @@ public:
 	ENiagaraFunctionDebugState CalculateDebugState() const;
 	FCompileConstantResolver WithDebugState(ENiagaraFunctionDebugState InDebugState) const;
 	FCompileConstantResolver WithUsage(ENiagaraScriptUsage ScriptUsage) const;
+
+	// returns a hash of the data that is used when resolving constants using this resolver.  The specific
+	// Emitter/System/Translator shouldn't be included but the values that it might use from those objects
+	// should be
+	uint32 BuildTypeHash() const;
+
 private:
 	FVersionedNiagaraEmitter Emitter;
 	const UNiagaraSystem* System;
@@ -513,6 +519,8 @@ public:
 
 	bool bShouldBuildSubHistories = true;
 
+	int32 MaxGraphDepthTraversal = INDEX_NONE;
+	int32 CurrentGraphDepth = 0;
 
 	void SetConstantByStaticVariable(int32& OutValue, const UEdGraphPin* InDefaultPin);
 	void SetConstantByStaticVariable(int32& OutValue, const FNiagaraVariable& Var);
diff --git a/Engine/Plugins/FX/Niagara/Source/NiagaraEditor/Public/NiagaraScriptSource.h b/Engine/Plugins/FX/Niagara/Source/NiagaraEditor/Public/NiagaraScriptSource.h
--- a/Engine/Plugins/FX/Niagara/Source/NiagaraEditor/Public/NiagaraScriptSource.h
+++ b/Engine/Plugins/FX/Niagara/Source/NiagaraEditor/Public/NiagaraScriptSource.h
@@ -50,8 +50,8 @@ class UNiagaraScriptSource : public UNiagaraScriptSourceBase
 
 	NIAGARAEDITOR_API virtual bool AddModuleIfMissing(FString ModulePath, ENiagaraScriptUsage Usage, bool& bOutFoundModule)override;
 
-	virtual void FixupRenamedParameters(UNiagaraNodeFunctionCall* FunctionCallNode, TConstArrayView<const UEdGraphPin*> FunctionPins, FNiagaraParameterStore& RapidIterationParameters, const TArray<FNiagaraVariable>& OldRapidIterationVariables, const FVersionedNiagaraEmitter& Emitter, ENiagaraScriptUsage ScriptUsage) const;
-	virtual void InitializeNewParameters(UNiagaraNodeFunctionCall* FunctionCallNode, TConstArrayView<const UEdGraphPin*> FunctionPins, FNiagaraParameterStore& RapidIterationParameters, const FVersionedNiagaraEmitter& VersionedEmitter, ENiagaraScriptUsage ScriptUsage, TSet<FName>& ValidRapidIterationParameterNames) const;
+	virtual void FixupRenamedParameters(UNiagaraNodeFunctionCall* FunctionCallNode, TConstArrayView<const UEdGraphPin*> FunctionPins, TConstArrayView<FNiagaraVariable> PinVariables, FNiagaraParameterStore& RapidIterationParameters, const TArray<FNiagaraVariable>& OldRapidIterationVariables, const FVersionedNiagaraEmitter& Emitter, ENiagaraScriptUsage ScriptUsage) const;
+	virtual void InitializeNewParameters(UNiagaraNodeFunctionCall* FunctionCallNode, TConstArrayView<const UEdGraphPin*> FunctionPins, TConstArrayView<FNiagaraVariable> PinVariables, FNiagaraParameterStore& RapidIterationParameters, const FVersionedNiagaraEmitter& VersionedEmitter, ENiagaraScriptUsage ScriptUsage, TSet<FName>& ValidRapidIterationParameterNames) const;
 
 	virtual void CleanUpOldAndInitializeNewRapidIterationParameters(const FVersionedNiagaraEmitter& Emitter, ENiagaraScriptUsage ScriptUsage, FGuid ScriptUsageId, FNiagaraParameterStore& RapidIterationParameters) const override;
 	virtual void ForceGraphToRecompileOnNextCheck() override;
diff --git a/Engine/Plugins/FX/Niagara/Source/NiagaraEditor/Public/ViewModels/Stack/NiagaraParameterHandle.h b/Engine/Plugins/FX/Niagara/Source/NiagaraEditor/Public/ViewModels/Stack/NiagaraParameterHandle.h
--- a/Engine/Plugins/FX/Niagara/Source/NiagaraEditor/Public/ViewModels/Stack/NiagaraParameterHandle.h
+++ b/Engine/Plugins/FX/Niagara/Source/NiagaraEditor/Public/ViewModels/Stack/NiagaraParameterHandle.h
@@ -22,6 +22,7 @@ public:
 	bool operator!=(const FNiagaraParameterHandle& Other) const { return (*this == Other) == false; }
 
 	static FNiagaraParameterHandle CreateAliasedModuleParameterHandle(const FNiagaraParameterHandle& ModuleParameterHandle, const UNiagaraNodeFunctionCall* ModuleNode);
+	static FNiagaraParameterHandle CreateAliasedModuleParameterHandle(const FName FullInputName, const FName FunctionName);
 
 	static FNiagaraParameterHandle CreateEngineParameterHandle(const FNiagaraVariable& SystemVariable);
 
diff --git a/Engine/Plugins/FX/Niagara/Source/NiagaraEditor/Public/ViewModels/Stack/NiagaraStackGraphUtilities.h b/Engine/Plugins/FX/Niagara/Source/NiagaraEditor/Public/ViewModels/Stack/NiagaraStackGraphUtilities.h
--- a/Engine/Plugins/FX/Niagara/Source/NiagaraEditor/Public/ViewModels/Stack/NiagaraStackGraphUtilities.h
+++ b/Engine/Plugins/FX/Niagara/Source/NiagaraEditor/Public/ViewModels/Stack/NiagaraStackGraphUtilities.h
@@ -93,11 +93,11 @@ namespace FNiagaraStackGraphUtilities
 	};
 
 	/* Returns the input pins for the given function call node. Try not to use this method if possible and use the version that accepts a FCompileConstantResolver parameter. This method should only be used if the current context is completely outside of any system or emitter (e.g. inside the graph editor itself) and constants cannot be resolved at all. */
-	void GetStackFunctionInputPins(UNiagaraNodeFunctionCall& FunctionCallNode, TArray<const UEdGraphPin*>& OutInputPins, ENiagaraGetStackFunctionInputPinsOptions Options = ENiagaraGetStackFunctionInputPinsOptions::AllInputs, bool bIgnoreDisabled = false);
+	void GetStackFunctionInputPins(UNiagaraNodeFunctionCall& FunctionCallNode, TArray<const UEdGraphPin*>& OutInputPins, ENiagaraGetStackFunctionInputPinsOptions Options, bool bIgnoreDisabled = false);
 
 	/* Returns the input pins for the given function call node. */
-	void GetStackFunctionInputPins(UNiagaraNodeFunctionCall& FunctionCallNode, TArray<const UEdGraphPin*>& OutInputPins, TSet<const UEdGraphPin*>& OutHiddenPins, FCompileConstantResolver ConstantResolver, ENiagaraGetStackFunctionInputPinsOptions Options = ENiagaraGetStackFunctionInputPinsOptions::AllInputs, bool bIgnoreDisabled = false);
-	void GetStackFunctionInputPins(UNiagaraNodeFunctionCall& FunctionCallNode, TArray<const UEdGraphPin*>& OutInputPins, FCompileConstantResolver ConstantResolver, ENiagaraGetStackFunctionInputPinsOptions Options = ENiagaraGetStackFunctionInputPinsOptions::AllInputs, bool bIgnoreDisabled = false);
+	void GetStackFunctionInputPins(UNiagaraNodeFunctionCall& FunctionCallNode, TArray<const UEdGraphPin*>& OutInputPins, TSet<const UEdGraphPin*>& OutHiddenPins, FCompileConstantResolver ConstantResolver, ENiagaraGetStackFunctionInputPinsOptions Options, bool bIgnoreDisabled = false);
+	void GetStackFunctionInputPins(UNiagaraNodeFunctionCall& FunctionCallNode, TArray<const UEdGraphPin*>& OutInputPins, FCompileConstantResolver ConstantResolver, ENiagaraGetStackFunctionInputPinsOptions Options, bool bIgnoreDisabled = false);
 
 	/* Returns the input pins for the given function call node.  Bypasses the module level caching of this data and generates it each call. */
 	void GetStackFunctionInputPinsWithoutCache(
