Fixed a bug where changing a Blueprint struct that is used in Niagara could lead to a crash

diff --git a/Engine/Plugins/FX/Niagara/Source/Niagara/Private/NiagaraTypes.cpp b/Engine/Plugins/FX/Niagara/Source/Niagara/Private/NiagaraTypes.cpp
--- a/Engine/Plugins/FX/Niagara/Source/Niagara/Private/NiagaraTypes.cpp
+++ b/Engine/Plugins/FX/Niagara/Source/Niagara/Private/NiagaraTypes.cpp
@@ -522,13 +522,32 @@ UScriptStruct* FNiagaraTypeHelper::GetSWCStruct(UScriptStruct* LWCStruct)
 			// If this is a LWC structure we need to built a new structure now
 			if (IsLWCStructure(LWCStruct))
 			{
-				SWCStruct = NewObject<UScriptStruct>(GetTransientPackage(), FName(LWCStruct->GetName() + TEXT("_SWC")), RF_NoFlags);
+				FName SWCName = FName(LWCStruct->GetName() + TEXT("_SWC"));
+
+				//check if the remap table contains a previous entry for that struct. This might happen when the lwc struct was changed at runtime.
+				// In that case we want to reuse the existing swc struct object, since it might already be referenced by compile results.
+				for (auto Iter = RemapTable.CreateIterator(); Iter; ++Iter)
+				{
+					TWeakObjectPtr<UScriptStruct> ExistingStruct = Iter.Value().Struct;
+					if (ExistingStruct.IsValid() && ExistingStruct->GetFName() == SWCName)
+					{
+						SWCStruct = ExistingStruct.Get();
+						Iter.RemoveCurrent();
+						RemapTableDirty = true;
+						break;
+					}
+				}
+
+				if (SWCStruct == nullptr)
+				{
+					SWCStruct = NewObject<UScriptStruct>(GetTransientPackage(), SWCName, RF_NoFlags);
+					SWCStruct->AddToRoot();
+				}
 				FNiagaraLwcStructConverter StructConverter = BuildSWCStructure(SWCStruct, LWCStruct);
 				FNiagaraTypeRegistry::RegisterStructConverter(FNiagaraTypeDefinition(LWCStruct, FNiagaraTypeDefinition::EAllowUnfriendlyStruct::Allow), StructConverter);
 				SWCStruct->Bind();
 				SWCStruct->PrepareCppStructOps();
 				SWCStruct->StaticLink(true);
-				SWCStruct->AddToRoot();
 			}
 			else
 			{
@@ -556,7 +575,7 @@ UScriptStruct* FNiagaraTypeHelper::GetLWCStruct(UScriptStruct* SWCStruct)
 	for (auto It = RemapTable.CreateConstIterator(); It; ++It)
 	{
 		const FRemapEntry& RemapEntry = It.Value();
-		if (RemapEntry.Get(SWCStruct) == SWCStruct)
+		if (RemapEntry.Struct.Get() == SWCStruct)
 		{
 			return It.Key().Get();
 		}
