Fixing issue where we failed to persist the Niagara Editor state if closed as a side effect of the editor closing.

diff --git a/Engine/Plugins/FX/Niagara/Source/NiagaraEditor/Private/ViewModels/NiagaraSystemEditorDocumentsViewModel.cpp b/Engine/Plugins/FX/Niagara/Source/NiagaraEditor/Private/ViewModels/NiagaraSystemEditorDocumentsViewModel.cpp
--- a/Engine/Plugins/FX/Niagara/Source/NiagaraEditor/Private/ViewModels/NiagaraSystemEditorDocumentsViewModel.cpp
+++ b/Engine/Plugins/FX/Niagara/Source/NiagaraEditor/Private/ViewModels/NiagaraSystemEditorDocumentsViewModel.cpp
@@ -49,6 +49,8 @@ void UNiagaraSystemEditorDocumentsViewModel::Initialize(TSharedRef<FNiagaraSyste
 
 void UNiagaraSystemEditorDocumentsViewModel::Finalize()
 {
+	DocumentManager.Reset();
+	TabManager.Reset();
 	ActiveDocumentTabScriptViewModel.Reset();
 }
 
@@ -183,7 +185,7 @@ void UNiagaraSystemEditorDocumentsViewModel::OpenChildScript(UEdGraph* InGraph)
 	}
 	else if (Results.Num() > 0)
 	{
-		if (Results[0].IsValid())
+		if (Results[0].IsValid() && TabManager.IsValid())
 		{
 			TabManager->DrawAttention(Results[0].ToSharedRef());
 			Results[0]->ActivateInParent(ETabActivationCause::SetDirectly);
@@ -257,7 +259,11 @@ void UNiagaraSystemEditorDocumentsViewModel::InitializePreTabManager(TSharedPtr<
 void UNiagaraSystemEditorDocumentsViewModel::InitializePostTabManager(TSharedPtr<FNiagaraSystemToolkit> InToolkit)
 {
 	TabManager = InToolkit->GetTabManager();
-	DocumentManager->SetTabManager(TabManager.ToSharedRef());
+
+	if (DocumentManager.IsValid())
+	{
+		DocumentManager->SetTabManager(InToolkit->GetTabManager().ToSharedRef());
+	}
 	SystemGraphSelectionVMWeak = InToolkit->GetSystemGraphSelectionViewModel();
 }
 
@@ -267,7 +273,7 @@ void UNiagaraSystemEditorDocumentsViewModel::DrawAttentionToPrimaryDocument()
 	TSharedPtr<SDockTab>  ActiveTab = GetActiveDocumentTab().Pin();
 	if (ActiveTab.IsValid())
 	{
-		if (ActiveTab->GetLayoutIdentifier().TabType != PrimaryDocumentTabId)
+		if (ActiveTab->GetLayoutIdentifier().TabType != PrimaryDocumentTabId && TabManager.IsValid())
 		{
 			TSharedPtr<SDockTab> FoundPrimaryDocumentTab = TabManager->FindExistingLiveTab(PrimaryDocumentTabId);
 			if (FoundPrimaryDocumentTab.IsValid())
@@ -302,7 +308,14 @@ TSharedRef<SNiagaraScratchPadScriptEditor> UNiagaraSystemEditorDocumentsViewMode
 TSharedPtr<SDockTab> UNiagaraSystemEditorDocumentsViewModel::OpenDocument(const UObject* DocumentID, FDocumentTracker::EOpenDocumentCause Cause)
 {
 	TSharedRef<FTabPayload_UObject> Payload = FTabPayload_UObject::Make(DocumentID);
-	return DocumentManager->OpenDocument(Payload, Cause);
+	if (DocumentManager.IsValid())
+	{
+		return DocumentManager->OpenDocument(Payload, Cause);
+	}
+	else
+	{
+		return nullptr;
+	}
 }
 
 void UNiagaraSystemEditorDocumentsViewModel::NavigateTab(FDocumentTracker::EOpenDocumentCause InCause)
@@ -313,8 +326,11 @@ void UNiagaraSystemEditorDocumentsViewModel::NavigateTab(FDocumentTracker::EOpen
 void UNiagaraSystemEditorDocumentsViewModel::CloseDocumentTab(const UObject* DocumentID)
 {
 	TSharedRef<FTabPayload_UObject> Payload = FTabPayload_UObject::Make(DocumentID);
-	DocumentManager->CloseTab(Payload);
-	DocumentManager->CleanInvalidTabs();
+	if (DocumentManager.IsValid())
+	{
+		DocumentManager->CloseTab(Payload);
+		DocumentManager->CleanInvalidTabs();
+	}
 }
 
 // Finds any open tabs containing the specified document and adds them to the specified array; returns true if at least one is found
@@ -323,8 +339,10 @@ bool UNiagaraSystemEditorDocumentsViewModel::FindOpenTabsContainingDocument(cons
 	int32 StartingCount = Results.Num();
 
 	TSharedRef<FTabPayload_UObject> Payload = FTabPayload_UObject::Make(DocumentID);
-
-	DocumentManager->FindMatchingTabs(Payload, /*inout*/ Results);
+	if (DocumentManager.IsValid())
+	{
+		DocumentManager->FindMatchingTabs(Payload, /*inout*/ Results);
+	}
 
 	// Did we add anything new?
 	return (StartingCount != Results.Num());
