Fixing auto-editing when created and also making sure that it is undoable.

diff --git a/Engine/Plugins/FX/Niagara/Source/NiagaraEditor/Private/Toolkits/SystemToolkitModes/NiagaraSystemToolkitModeBase.cpp b/Engine/Plugins/FX/Niagara/Source/NiagaraEditor/Private/Toolkits/SystemToolkitModes/NiagaraSystemToolkitModeBase.cpp
--- a/Engine/Plugins/FX/Niagara/Source/NiagaraEditor/Private/Toolkits/SystemToolkitModes/NiagaraSystemToolkitModeBase.cpp
+++ b/Engine/Plugins/FX/Niagara/Source/NiagaraEditor/Private/Toolkits/SystemToolkitModes/NiagaraSystemToolkitModeBase.cpp
@@ -37,6 +37,7 @@
 #include "Styling/StyleColors.h"
 #include "ViewModels/NiagaraSystemEditorDocumentsViewModel.h"
 #include "ViewModels/HierarchyEditor/NiagaraUserParametersHierarchyViewModel.h"
+#include "ViewModels/NiagaraScratchPadViewModel.h"
 
 #define LOCTEXT_NAMESPACE "NiagaraSystemToolkitModeBase"
 
@@ -77,6 +78,9 @@ FNiagaraSystemToolkitModeBase::~FNiagaraSystemToolkitModeBase()
 			TSharedPtr<FNiagaraScratchPadScriptViewModel> OldScratchScriptVM = LastActiveDocumentModel.Pin();
 			OldScratchScriptVM->GetGraphViewModel()->GetNodeSelection()->OnSelectedObjectsChanged().Remove(LastSelectionUpdateDelegate);
 		}
+
+		SystemToolkit.Pin()->GetSystemViewModel()->GetScriptScratchPadViewModel()->OnFocusNeeded().Remove(ScratchPadNeedsFocusHandle);
+
 	}
 }
 
@@ -220,6 +224,7 @@ void FNiagaraSystemToolkitModeBase::RegisterTabFactories(TSharedPtr<FTabManager>
 	WorkspaceMenuCategory = InTabManager->AddLocalWorkspaceMenuCategory(LOCTEXT("WorkspaceMenu_NiagaraSystemEditor", "Niagara System"));
 
 	SystemToolkit.Pin()->RegisterToolbarTab(InTabManager.ToSharedRef());
+	TabManager = InTabManager;
 	
 	InTabManager->RegisterTabSpawner(ViewportTabID, FOnSpawnTab::CreateSP(this, &FNiagaraSystemToolkitModeBase::SpawnTab_Viewport))
 		.SetDisplayName(LOCTEXT("Preview", "Preview"))
@@ -711,6 +716,14 @@ TSharedRef<SDockTab> FNiagaraSystemToolkitModeBase::SpawnTab_SystemOverview(cons
 	return SpawnedTab;
 }
 
+void FNiagaraSystemToolkitModeBase::OnScratchPadNeedsFocus()
+{
+
+	if (TabManager.IsValid() && ScratchPadTab.Pin().IsValid())
+	{
+		TabManager->DrawAttention(ScratchPadTab.Pin().ToSharedRef());
+	}
+}
 
 
 TSharedRef<SDockTab> FNiagaraSystemToolkitModeBase::SpawnTab_ScratchPadScripts(const FSpawnTabArgs& Args)
@@ -720,11 +733,14 @@ TSharedRef<SDockTab> FNiagaraSystemToolkitModeBase::SpawnTab_ScratchPadScripts(c
 		SystemToolkit.Pin()->SetScriptScratchpadManager(FNiagaraEditorModule::Get().GetWidgetProvider()->CreateScriptScratchPadManager(*SystemToolkit.Pin()->GetSystemViewModel()->GetScriptScratchPadViewModel()));
 	}
 
+	ScratchPadNeedsFocusHandle = SystemToolkit.Pin()->GetSystemViewModel()->GetScriptScratchPadViewModel()->OnFocusNeeded().AddRaw(this, &FNiagaraSystemToolkitModeBase::OnScratchPadNeedsFocus);
+
 	TSharedRef<SDockTab> SpawnedTab = SNew(SDockTab)
 		.Label(LOCTEXT("ScratchPadLocalModulesTabLabel", "Local Modules"))
 		[
 			SystemToolkit.Pin()->GetScriptScratchpadManager().ToSharedRef()
 		];
+	ScratchPadTab = SpawnedTab;
 
 	SpawnedTab->SetOnTabClosed(SDockTab::FOnTabClosedCallback::CreateLambda([this](TSharedRef<SDockTab>)
 		{
diff --git a/Engine/Plugins/FX/Niagara/Source/NiagaraEditor/Private/Toolkits/SystemToolkitModes/NiagaraSystemToolkitModeBase.h b/Engine/Plugins/FX/Niagara/Source/NiagaraEditor/Private/Toolkits/SystemToolkitModes/NiagaraSystemToolkitModeBase.h
--- a/Engine/Plugins/FX/Niagara/Source/NiagaraEditor/Private/Toolkits/SystemToolkitModes/NiagaraSystemToolkitModeBase.h
+++ b/Engine/Plugins/FX/Niagara/Source/NiagaraEditor/Private/Toolkits/SystemToolkitModes/NiagaraSystemToolkitModeBase.h
@@ -46,13 +46,19 @@ private:
 	TSharedRef<SDockTab> SpawnTab_UserParameters(const FSpawnTabArgs& Args);
 	TSharedRef<SDockTab> SpawnTab_UserParametersHierarchyEditor(const FSpawnTabArgs& Args);
 
+	void OnScratchPadNeedsFocus();
+
 protected:
 	TWeakPtr<FNiagaraSystemToolkit> SystemToolkit;
+	TWeakPtr< SDockTab> ScratchPadTab;
+
 	FDelegateHandle DocChangedHandle;
 	FDelegateHandle LastSelectionUpdateDelegate;
 	FDelegateHandle LastParamPanelSelectionUpdateDelegate;
 	FDelegateHandle LastSystemSelectionUpdateDelegate;
 	FDelegateHandle LastGraphEditDelegate;
+	FDelegateHandle ScratchPadNeedsFocusHandle;
+
 	int32 SwitcherIdx;
 	TSharedPtr<FNiagaraObjectSelection> ObjectSelection;
 	FText ObjectSelectionSubHeaderText;
@@ -66,6 +72,8 @@ protected:
 
 	TWeakPtr< class FNiagaraScratchPadScriptViewModel > LastActiveDocumentModel;
 
+	TSharedPtr<FTabManager> TabManager;
+
 public:
 	static const FName ViewportTabID;
 	static const FName CurveEditorTabID;
diff --git a/Engine/Plugins/FX/Niagara/Source/NiagaraEditor/Private/ViewModels/NiagaraScratchPadScriptViewModel.cpp b/Engine/Plugins/FX/Niagara/Source/NiagaraEditor/Private/ViewModels/NiagaraScratchPadScriptViewModel.cpp
--- a/Engine/Plugins/FX/Niagara/Source/NiagaraEditor/Private/ViewModels/NiagaraScratchPadScriptViewModel.cpp
+++ b/Engine/Plugins/FX/Niagara/Source/NiagaraEditor/Private/ViewModels/NiagaraScratchPadScriptViewModel.cpp
@@ -53,6 +53,21 @@ void FNiagaraScratchPadScriptViewModel::Initialize(UNiagaraScript* Script, UNiag
 	OriginalScript = Script;
 	SystemViewModel = InSystemViewModel;
 
+	if (OriginalScript)
+	{
+		CachedName = OriginalScript->GetName();
+		OriginalScript->SetFlags(RF_Transactional);
+	}
+	else
+	{
+		CachedName.Reset();
+	}
+
+	if (InEditScript)
+	{
+		InEditScript->SetFlags(RF_Transactional);
+	}
+
 	// Copy over the old edited graph, as this initialize might be incoming from inheritance changes and we want to maintain
 	// the edits to the old graph we were working on.
 	if (InEditScript)
@@ -145,6 +160,22 @@ bool FNiagaraScratchPadScriptViewModel::GetIsPendingRename() const
 void FNiagaraScratchPadScriptViewModel::SetIsPendingRename(bool bInIsPendingRename)
 {
 	bIsPendingRename = bInIsPendingRename;
+	OnPendingRenameDelegate.Broadcast();
+}
+
+bool FNiagaraScratchPadScriptViewModel::ValidateAndUpdateCachedNameIfDifferent()
+{
+	if (CachedName.IsSet() && OriginalScript)
+	{
+		const FString& CachedNameStr = CachedName.GetValue();
+		if (CachedNameStr != OriginalScript->GetName())
+		{
+			CachedName = OriginalScript->GetName();
+			return true;
+		}
+	}
+
+	return false;
 }
 
 void FNiagaraScratchPadScriptViewModel::SetScriptName(FText InScriptName)
@@ -153,6 +184,10 @@ void FNiagaraScratchPadScriptViewModel::SetScriptName(FText InScriptName)
 	if (OriginalScript->GetName() != NewName)
 	{
 		FScopedTransaction RenameTransaction(LOCTEXT("RenameScriptTransaction", "Rename scratch pad script."));
+		if (OriginalScript->GetOuter())
+		{
+			OriginalScript->GetOuter()->Modify();
+		}
 
 		FName NewUniqueName = FNiagaraEditorUtilities::GetUniqueObjectName<UNiagaraScript>(OriginalScript->GetOuter(), *NewName);
 		OriginalScript->Modify();
@@ -165,7 +200,7 @@ void FNiagaraScratchPadScriptViewModel::SetScriptName(FText InScriptName)
 			FNiagaraScratchPadUtilities::FixFunctionInputsFromFunctionScriptRename(*ReferencingFunctionCallNode, NewUniqueName);
 			ReferencingFunctionCallNode->MarkNodeRequiresSynchronization(TEXT("ScratchPad script renamed"), true);
 		}
-
+		CachedName = NewUniqueName.ToString();
 		OnRenamedDelegate.Broadcast();
 	}
 }
@@ -242,6 +277,11 @@ FNiagaraScratchPadScriptViewModel::FOnRenamed& FNiagaraScratchPadScriptViewModel
 	return OnRenamedDelegate;
 }
 
+FNiagaraScratchPadScriptViewModel::FOnPendingRename& FNiagaraScratchPadScriptViewModel::OnPendingRename()
+{
+	return OnPendingRenameDelegate;
+}
+
 FNiagaraScratchPadScriptViewModel::FOnPinnedChanged& FNiagaraScratchPadScriptViewModel::OnPinnedChanged()
 {
 	return OnPinnedChangedDelegate;
diff --git a/Engine/Plugins/FX/Niagara/Source/NiagaraEditor/Private/ViewModels/NiagaraScratchPadViewModel.cpp b/Engine/Plugins/FX/Niagara/Source/NiagaraEditor/Private/ViewModels/NiagaraScratchPadViewModel.cpp
--- a/Engine/Plugins/FX/Niagara/Source/NiagaraEditor/Private/ViewModels/NiagaraScratchPadViewModel.cpp
+++ b/Engine/Plugins/FX/Niagara/Source/NiagaraEditor/Private/ViewModels/NiagaraScratchPadViewModel.cpp
@@ -216,10 +216,17 @@ void UNiagaraScratchPadViewModel::RefreshScriptViewModels()
 	{
 		for (int32 i = 0; i < ScriptViewModels.Num(); i++)
 		{
+			// First double-check that we're not creating multiple view models per script....
 			TArray<TSharedPtr<FNiagaraScriptViewModel>> ExistingViewModels;
 			TNiagaraViewModelManager<UNiagaraScript, FNiagaraScriptViewModel>::GetAllViewModelsForObject(ScriptViewModels[i]->GetEditScript().Script, ExistingViewModels);
 
 			(ensureMsgf(ExistingViewModels.Num() == 1, TEXT("Missing or duplicate script view models detected.  Can not create script details customization.")));
+
+			// Now double-check that we haven't reused a view model where the underlying script name has changed and need to act as if a rename happened.
+			if (ScriptViewModels[i]->ValidateAndUpdateCachedNameIfDifferent())
+			{
+				ScriptViewModelScriptRenamed(ScriptViewModels[i]);
+			}
 		}
 	}
 	
@@ -473,10 +480,17 @@ FName GetUniqueScriptName(UObject* Outer, const FString& CandidateName)
 
 TSharedPtr<FNiagaraScratchPadScriptViewModel> UNiagaraScratchPadViewModel::CreateNewScript(ENiagaraScriptUsage InScriptUsage, ENiagaraScriptUsage InTargetSupportedUsage, FNiagaraTypeDefinition InOutputType)
 {
+	FScopedTransaction Transaction(LOCTEXT("CreateNewScriptTransaction", "Create a new local script."));
+
 	UObject* ScriptOuter;
 	TArray<TObjectPtr<UNiagaraScript>>* TargetScripts;
 	GetOuterAndTargetScripts(GetSystemViewModel(), ScriptOuter, TargetScripts);
 
+	if (ScriptOuter)
+	{
+		ScriptOuter->Modify();
+	}
+
 	UNiagaraScript* NewScript = nullptr;
 	switch (InScriptUsage)
 	{
@@ -530,7 +544,7 @@ TSharedPtr<FNiagaraScratchPadScriptViewModel> UNiagaraScratchPadViewModel::Creat
 	}
 
 	NewScript->ClearFlags(RF_Public | RF_Standalone);
-	ScriptOuter->Modify();
+	NewScript->SetFlags(RF_Transactional);
 	TargetScripts->Add(NewScript);
 	NewScript->GetLatestScriptData()->ModuleUsageBitmask |= (1 << (int32)InTargetSupportedUsage);
 	RefreshScriptViewModels();
@@ -541,13 +555,21 @@ TSharedPtr<FNiagaraScratchPadScriptViewModel> UNiagaraScratchPadViewModel::Creat
 
 TSharedPtr<FNiagaraScratchPadScriptViewModel> UNiagaraScratchPadViewModel::CreateNewScriptAsDuplicate(const UNiagaraScript* ScriptToDuplicate)
 {
+
+	FScopedTransaction Transaction(LOCTEXT("CreateNewScriptDupeTransaction", "Create a new local script as duplicate."));
+
 	UObject* ScriptOuter;
 	TArray<TObjectPtr<UNiagaraScript>>* TargetScripts;
 	GetOuterAndTargetScripts(GetSystemViewModel(), ScriptOuter, TargetScripts);
 
+	if (ScriptOuter)
+	{
+		ScriptOuter->Modify();
+	}
+	
 	UNiagaraScript* NewScript = CastChecked<UNiagaraScript>(StaticDuplicateObject(ScriptToDuplicate, ScriptOuter, GetUniqueScriptName(ScriptOuter, *ScriptToDuplicate->GetFName().ToString())));
 	NewScript->ClearFlags(RF_Public | RF_Standalone);
-	ScriptOuter->Modify();
+	NewScript->SetFlags(RF_Transactional);
 	TargetScripts->Add(NewScript);
 	RefreshScriptViewModels();
 	UpdateChangeId(GetSystemViewModel());
@@ -720,6 +742,8 @@ TSharedRef<FNiagaraScratchPadScriptViewModel> UNiagaraScratchPadViewModel::Creat
 	ScriptViewModel->Initialize(ScratchPadScript, ScratchPadEditScript, SystemViewModelWeak);
 	ScriptViewModel->GetGraphViewModel()->GetNodeSelection()->OnSelectedObjectsChanged().AddUObject(this, &UNiagaraScratchPadViewModel::ScriptGraphNodeSelectionChanged, TWeakPtr<FNiagaraScratchPadScriptViewModel>(ScriptViewModel));
 	ScriptViewModel->OnRenamed().AddUObject(this, &UNiagaraScratchPadViewModel::ScriptViewModelScriptRenamed, TWeakPtr<FNiagaraScratchPadScriptViewModel>(ScriptViewModel));
+	ScriptViewModel->OnPendingRename().AddUObject(this, &UNiagaraScratchPadViewModel::ScriptViewModelScriptPendingRename, TWeakPtr<FNiagaraScratchPadScriptViewModel>(ScriptViewModel));
+
 	ScriptViewModel->OnGraphSelectionChanged().AddUObject(this, &UNiagaraScratchPadViewModel::ScriptViewModelGraphSelectionChanged, TWeakPtr<FNiagaraScratchPadScriptViewModel>(ScriptViewModel));
 	ScriptViewModel->OnPinnedChanged().AddUObject(this, &UNiagaraScratchPadViewModel::ScriptViewModelPinnedChanged, TWeakPtr<FNiagaraScratchPadScriptViewModel>(ScriptViewModel));
 	ScriptViewModel->OnHasUnappliedChangesChanged().AddUObject(this, &UNiagaraScratchPadViewModel::ScriptViewModelHasUnappliedChangesChanged);
@@ -740,6 +764,13 @@ void UNiagaraScratchPadViewModel::TearDownScriptViewModel(TSharedRef<FNiagaraScr
 	InScriptViewModel->OnRequestDiscardChanges().Unbind();
 	InScriptViewModel->GetVariableSelection()->OnSelectedObjectsChanged().RemoveAll(this);
 	InScriptViewModel->Finalize();
+
+	UNiagaraSystemEditorDocumentsViewModel* DocVM = GetSystemViewModel()->GetDocumentViewModel();
+	if (DocVM)
+	{
+		DocVM->CloseChildScript(InScriptViewModel->GetGraphViewModel()->GetGraph());
+	}
+
 }
 
 void UNiagaraScratchPadViewModel::RefreshEditScriptViewModels()
@@ -775,6 +806,16 @@ void UNiagaraScratchPadViewModel::ScriptGraphNodeSelectionChanged(TWeakPtr<FNiag
 	}
 }
 
+void UNiagaraScratchPadViewModel::PostUndo()
+{
+	RefreshScriptViewModels();
+	if (ObjectSelection)
+	{
+		ObjectSelection->ClearSelectedObjects();
+	}
+}
+
+
 void UNiagaraScratchPadViewModel::ScriptViewModelScriptRenamed(TWeakPtr<FNiagaraScratchPadScriptViewModel> ScriptViewModelWeak)
 {
 	TSharedPtr<FNiagaraScratchPadScriptViewModel> ScriptViewModel = ScriptViewModelWeak.Pin();
@@ -792,6 +833,11 @@ void UNiagaraScratchPadViewModel::ScriptViewModelScriptRenamed(TWeakPtr<FNiagara
 	}
 }
 
+void UNiagaraScratchPadViewModel::ScriptViewModelScriptPendingRename(TWeakPtr<FNiagaraScratchPadScriptViewModel> ScriptViewModelWeak)
+{
+	OnFocusNeededDelegate.Broadcast();
+}
+
 void UNiagaraScratchPadViewModel::ScriptViewModelGraphSelectionChanged(const UObject* Obj, TWeakPtr<FNiagaraScratchPadScriptViewModel> ScriptViewModelWeak)
 {
 	const UNiagaraScriptVariable* Var = Cast<UNiagaraScriptVariable>(Obj);
diff --git a/Engine/Plugins/FX/Niagara/Source/NiagaraEditor/Private/ViewModels/NiagaraSystemViewModel.cpp b/Engine/Plugins/FX/Niagara/Source/NiagaraEditor/Private/ViewModels/NiagaraSystemViewModel.cpp
--- a/Engine/Plugins/FX/Niagara/Source/NiagaraEditor/Private/ViewModels/NiagaraSystemViewModel.cpp
+++ b/Engine/Plugins/FX/Niagara/Source/NiagaraEditor/Private/ViewModels/NiagaraSystemViewModel.cpp
@@ -792,6 +792,12 @@ void FNiagaraSystemViewModel::AddReferencedObjects(FReferenceCollector& Collecto
 
 void FNiagaraSystemViewModel::PostUndo(bool bSuccess)
 {
+	// We may have undone scratch script renames or adds, so we need to make sure that they
+	// are accurate and potentially close anything that is now gone.
+	if (ScriptScratchPadViewModel)
+	{
+		ScriptScratchPadViewModel->PostUndo();
+	}
 	// Reset system stack and emitter handle view models to prevent accessing invalid data if they were in the undo operation.
 	ResetStack();
 	System->InvalidateActiveCompiles();
diff --git a/Engine/Plugins/FX/Niagara/Source/NiagaraEditor/Public/ViewModels/NiagaraScratchPadScriptViewModel.h b/Engine/Plugins/FX/Niagara/Source/NiagaraEditor/Public/ViewModels/NiagaraScratchPadScriptViewModel.h
--- a/Engine/Plugins/FX/Niagara/Source/NiagaraEditor/Public/ViewModels/NiagaraScratchPadScriptViewModel.h
+++ b/Engine/Plugins/FX/Niagara/Source/NiagaraEditor/Public/ViewModels/NiagaraScratchPadScriptViewModel.h
@@ -14,6 +14,7 @@ class NIAGARAEDITOR_API FNiagaraScratchPadScriptViewModel : public FNiagaraScrip
 {
 public:
 	DECLARE_MULTICAST_DELEGATE(FOnRenamed);
+	DECLARE_MULTICAST_DELEGATE(FOnPendingRename);
 	DECLARE_MULTICAST_DELEGATE(FOnPinnedChanged);
 	DECLARE_MULTICAST_DELEGATE(FOnHasUnappliedChangesChanged);
 	DECLARE_MULTICAST_DELEGATE(FOnChangesApplied);
@@ -84,6 +85,8 @@ public:
 
 	FOnRenamed& OnRenamed();
 
+	FOnPendingRename& OnPendingRename();
+
 	FOnPinnedChanged& OnPinnedChanged();
 
 	FOnHasUnappliedChangesChanged& OnHasUnappliedChangesChanged();
@@ -100,6 +103,9 @@ public:
 	void RaisePinFocusRequested(FNiagaraScriptIDAndGraphFocusInfo*  InFocusInfo);
 	void RaiseNodeFocusRequested(FNiagaraScriptIDAndGraphFocusInfo* InFocusInfo);
 
+	// Updates the cached name post undo to see if we need to broadcast a changed event. Returns true if they were out of sync and tabs will potentially need to be reopned.
+	bool ValidateAndUpdateCachedNameIfDifferent();
+
 private:
 	FText GetDisplayNameInternal() const;
 
@@ -124,6 +130,7 @@ private:
 	FDelegateHandle OnGraphNeedsRecompileHandle;
 
 	FOnRenamed OnRenamedDelegate;
+	FOnPendingRename OnPendingRenameDelegate;
 	FOnPinnedChanged OnPinnedChangedDelegate;
 	FOnHasUnappliedChangesChanged OnHasUnappliedChangesChangedDelegate;
 	FOnChangesApplied OnChangesAppliedDelegate;
@@ -132,4 +139,6 @@ private:
 
 	FOnNodeIDFocusRequested OnNodeIDFocusRequestedDelegate;
 	FOnPinIDFocusRequested OnPinIDFocusRequestedDelegate;
+
+	TOptional<FString> CachedName;
 };
\ No newline at end of file
diff --git a/Engine/Plugins/FX/Niagara/Source/NiagaraEditor/Public/ViewModels/NiagaraScratchPadViewModel.h b/Engine/Plugins/FX/Niagara/Source/NiagaraEditor/Public/ViewModels/NiagaraScratchPadViewModel.h
--- a/Engine/Plugins/FX/Niagara/Source/NiagaraEditor/Public/ViewModels/NiagaraScratchPadViewModel.h
+++ b/Engine/Plugins/FX/Niagara/Source/NiagaraEditor/Public/ViewModels/NiagaraScratchPadViewModel.h
@@ -21,6 +21,7 @@ public:
 	DECLARE_MULTICAST_DELEGATE(FOnActiveScriptChanged);
 	DECLARE_MULTICAST_DELEGATE(FOnScriptRenamed);
 	DECLARE_MULTICAST_DELEGATE(FOnScriptDeleted);
+	DECLARE_MULTICAST_DELEGATE(FOnFocusNeeded);
 
 public:
 	void Initialize(TSharedRef<FNiagaraSystemViewModel> InSystemViewModel);
@@ -86,6 +87,12 @@ public:
 
 	void OpenEditorForActive();
 
+	FOnFocusNeeded& OnFocusNeeded() {
+		return OnFocusNeededDelegate;
+	};
+
+	void PostUndo();
+
 private:
 	TSharedRef<FNiagaraSystemViewModel> GetSystemViewModel();
 
@@ -101,6 +108,8 @@ private:
 
 	void ScriptViewModelScriptRenamed(TWeakPtr<FNiagaraScratchPadScriptViewModel> ScriptViewModelWeak);
 
+	void ScriptViewModelScriptPendingRename(TWeakPtr<FNiagaraScratchPadScriptViewModel> ScriptViewModelWeak);	
+
 	void ScriptViewModelPinnedChanged(TWeakPtr<FNiagaraScratchPadScriptViewModel> ScriptViewModelWeak);
 
 	void ScriptViewModelHasUnappliedChangesChanged();
@@ -141,5 +150,8 @@ private:
 
 	FOnScriptDeleted OnScriptDeletedDelegate;
 
+	FOnFocusNeeded OnFocusNeededDelegate;
+
+
 	bool bIsBulkApplying = false;
 };
\ No newline at end of file
