[Backout] - CL22647753
Original CL Desc
-----------------------------------------------------------------
Fixing auto-editing when created and also making sure that it is undoable.

diff --git a/Engine/Plugins/FX/Niagara/Source/NiagaraEditor/Private/Toolkits/SystemToolkitModes/NiagaraSystemToolkitModeBase.cpp b/Engine/Plugins/FX/Niagara/Source/NiagaraEditor/Private/Toolkits/SystemToolkitModes/NiagaraSystemToolkitModeBase.cpp
--- a/Engine/Plugins/FX/Niagara/Source/NiagaraEditor/Private/Toolkits/SystemToolkitModes/NiagaraSystemToolkitModeBase.cpp
+++ b/Engine/Plugins/FX/Niagara/Source/NiagaraEditor/Private/Toolkits/SystemToolkitModes/NiagaraSystemToolkitModeBase.cpp
@@ -37,7 +37,6 @@
 #include "Styling/StyleColors.h"
 #include "ViewModels/NiagaraSystemEditorDocumentsViewModel.h"
 #include "ViewModels/HierarchyEditor/NiagaraUserParametersHierarchyViewModel.h"
-#include "ViewModels/NiagaraScratchPadViewModel.h"
 
 #define LOCTEXT_NAMESPACE "NiagaraSystemToolkitModeBase"
 
@@ -78,9 +77,6 @@ FNiagaraSystemToolkitModeBase::~FNiagaraSystemToolkitModeBase()
 			TSharedPtr<FNiagaraScratchPadScriptViewModel> OldScratchScriptVM = LastActiveDocumentModel.Pin();
 			OldScratchScriptVM->GetGraphViewModel()->GetNodeSelection()->OnSelectedObjectsChanged().Remove(LastSelectionUpdateDelegate);
 		}
-
-		SystemToolkit.Pin()->GetSystemViewModel()->GetScriptScratchPadViewModel()->OnFocusNeeded().Remove(ScratchPadNeedsFocusHandle);
-
 	}
 }
 
@@ -224,7 +220,6 @@ void FNiagaraSystemToolkitModeBase::RegisterTabFactories(TSharedPtr<FTabManager>
 	WorkspaceMenuCategory = InTabManager->AddLocalWorkspaceMenuCategory(LOCTEXT("WorkspaceMenu_NiagaraSystemEditor", "Niagara System"));
 
 	SystemToolkit.Pin()->RegisterToolbarTab(InTabManager.ToSharedRef());
-	TabManager = InTabManager;
 	
 	InTabManager->RegisterTabSpawner(ViewportTabID, FOnSpawnTab::CreateSP(this, &FNiagaraSystemToolkitModeBase::SpawnTab_Viewport))
 		.SetDisplayName(LOCTEXT("Preview", "Preview"))
@@ -716,14 +711,6 @@ TSharedRef<SDockTab> FNiagaraSystemToolkitModeBase::SpawnTab_SystemOverview(cons
 	return SpawnedTab;
 }
 
-void FNiagaraSystemToolkitModeBase::OnScratchPadNeedsFocus()
-{
-
-	if (TabManager.IsValid() && ScratchPadTab.Pin().IsValid())
-	{
-		TabManager->DrawAttention(ScratchPadTab.Pin().ToSharedRef());
-	}
-}
 
 
 TSharedRef<SDockTab> FNiagaraSystemToolkitModeBase::SpawnTab_ScratchPadScripts(const FSpawnTabArgs& Args)
@@ -733,14 +720,11 @@ TSharedRef<SDockTab> FNiagaraSystemToolkitModeBase::SpawnTab_ScratchPadScripts(c
 		SystemToolkit.Pin()->SetScriptScratchpadManager(FNiagaraEditorModule::Get().GetWidgetProvider()->CreateScriptScratchPadManager(*SystemToolkit.Pin()->GetSystemViewModel()->GetScriptScratchPadViewModel()));
 	}
 
-	ScratchPadNeedsFocusHandle = SystemToolkit.Pin()->GetSystemViewModel()->GetScriptScratchPadViewModel()->OnFocusNeeded().AddRaw(this, &FNiagaraSystemToolkitModeBase::OnScratchPadNeedsFocus);
-
 	TSharedRef<SDockTab> SpawnedTab = SNew(SDockTab)
 		.Label(LOCTEXT("ScratchPadLocalModulesTabLabel", "Local Modules"))
 		[
 			SystemToolkit.Pin()->GetScriptScratchpadManager().ToSharedRef()
 		];
-	ScratchPadTab = SpawnedTab;
 
 	SpawnedTab->SetOnTabClosed(SDockTab::FOnTabClosedCallback::CreateLambda([this](TSharedRef<SDockTab>)
 		{
diff --git a/Engine/Plugins/FX/Niagara/Source/NiagaraEditor/Private/Toolkits/SystemToolkitModes/NiagaraSystemToolkitModeBase.h b/Engine/Plugins/FX/Niagara/Source/NiagaraEditor/Private/Toolkits/SystemToolkitModes/NiagaraSystemToolkitModeBase.h
--- a/Engine/Plugins/FX/Niagara/Source/NiagaraEditor/Private/Toolkits/SystemToolkitModes/NiagaraSystemToolkitModeBase.h
+++ b/Engine/Plugins/FX/Niagara/Source/NiagaraEditor/Private/Toolkits/SystemToolkitModes/NiagaraSystemToolkitModeBase.h
@@ -46,19 +46,13 @@ private:
 	TSharedRef<SDockTab> SpawnTab_UserParameters(const FSpawnTabArgs& Args);
 	TSharedRef<SDockTab> SpawnTab_UserParametersHierarchyEditor(const FSpawnTabArgs& Args);
 
-	void OnScratchPadNeedsFocus();
-
 protected:
 	TWeakPtr<FNiagaraSystemToolkit> SystemToolkit;
-	TWeakPtr< SDockTab> ScratchPadTab;
-
 	FDelegateHandle DocChangedHandle;
 	FDelegateHandle LastSelectionUpdateDelegate;
 	FDelegateHandle LastParamPanelSelectionUpdateDelegate;
 	FDelegateHandle LastSystemSelectionUpdateDelegate;
 	FDelegateHandle LastGraphEditDelegate;
-	FDelegateHandle ScratchPadNeedsFocusHandle;
-
 	int32 SwitcherIdx;
 	TSharedPtr<FNiagaraObjectSelection> ObjectSelection;
 	FText ObjectSelectionSubHeaderText;
@@ -72,8 +66,6 @@ protected:
 
 	TWeakPtr< class FNiagaraScratchPadScriptViewModel > LastActiveDocumentModel;
 
-	TSharedPtr<FTabManager> TabManager;
-
 public:
 	static const FName ViewportTabID;
 	static const FName CurveEditorTabID;
diff --git a/Engine/Plugins/FX/Niagara/Source/NiagaraEditor/Private/ViewModels/NiagaraScratchPadScriptViewModel.cpp b/Engine/Plugins/FX/Niagara/Source/NiagaraEditor/Private/ViewModels/NiagaraScratchPadScriptViewModel.cpp
--- a/Engine/Plugins/FX/Niagara/Source/NiagaraEditor/Private/ViewModels/NiagaraScratchPadScriptViewModel.cpp
+++ b/Engine/Plugins/FX/Niagara/Source/NiagaraEditor/Private/ViewModels/NiagaraScratchPadScriptViewModel.cpp
@@ -53,21 +53,6 @@ void FNiagaraScratchPadScriptViewModel::Initialize(UNiagaraScript* Script, UNiag
 	OriginalScript = Script;
 	SystemViewModel = InSystemViewModel;
 
-	if (OriginalScript)
-	{
-		CachedName = OriginalScript->GetName();
-		OriginalScript->SetFlags(RF_Transactional);
-	}
-	else
-	{
-		CachedName.Reset();
-	}
-
-	if (InEditScript)
-	{
-		InEditScript->SetFlags(RF_Transactional);
-	}
-
 	// Copy over the old edited graph, as this initialize might be incoming from inheritance changes and we want to maintain
 	// the edits to the old graph we were working on.
 	if (InEditScript)
@@ -160,22 +145,6 @@ bool FNiagaraScratchPadScriptViewModel::GetIsPendingRename() const
 void FNiagaraScratchPadScriptViewModel::SetIsPendingRename(bool bInIsPendingRename)
 {
 	bIsPendingRename = bInIsPendingRename;
-	OnPendingRenameDelegate.Broadcast();
-}
-
-bool FNiagaraScratchPadScriptViewModel::ValidateAndUpdateCachedNameIfDifferent()
-{
-	if (CachedName.IsSet() && OriginalScript)
-	{
-		const FString& CachedNameStr = CachedName.GetValue();
-		if (CachedNameStr != OriginalScript->GetName())
-		{
-			CachedName = OriginalScript->GetName();
-			return true;
-		}
-	}
-
-	return false;
 }
 
 void FNiagaraScratchPadScriptViewModel::SetScriptName(FText InScriptName)
@@ -184,10 +153,6 @@ void FNiagaraScratchPadScriptViewModel::SetScriptName(FText InScriptName)
 	if (OriginalScript->GetName() != NewName)
 	{
 		FScopedTransaction RenameTransaction(LOCTEXT("RenameScriptTransaction", "Rename scratch pad script."));
-		if (OriginalScript->GetOuter())
-		{
-			OriginalScript->GetOuter()->Modify();
-		}
 
 		FName NewUniqueName = FNiagaraEditorUtilities::GetUniqueObjectName<UNiagaraScript>(OriginalScript->GetOuter(), *NewName);
 		OriginalScript->Modify();
@@ -200,7 +165,7 @@ void FNiagaraScratchPadScriptViewModel::SetScriptName(FText InScriptName)
 			FNiagaraScratchPadUtilities::FixFunctionInputsFromFunctionScriptRename(*ReferencingFunctionCallNode, NewUniqueName);
 			ReferencingFunctionCallNode->MarkNodeRequiresSynchronization(TEXT("ScratchPad script renamed"), true);
 		}
-		CachedName = NewUniqueName.ToString();
+
 		OnRenamedDelegate.Broadcast();
 	}
 }
@@ -277,11 +242,6 @@ FNiagaraScratchPadScriptViewModel::FOnRenamed& FNiagaraScratchPadScriptViewModel
 	return OnRenamedDelegate;
 }
 
-FNiagaraScratchPadScriptViewModel::FOnPendingRename& FNiagaraScratchPadScriptViewModel::OnPendingRename()
-{
-	return OnPendingRenameDelegate;
-}
-
 FNiagaraScratchPadScriptViewModel::FOnPinnedChanged& FNiagaraScratchPadScriptViewModel::OnPinnedChanged()
 {
 	return OnPinnedChangedDelegate;
diff --git a/Engine/Plugins/FX/Niagara/Source/NiagaraEditor/Private/ViewModels/NiagaraScratchPadViewModel.cpp b/Engine/Plugins/FX/Niagara/Source/NiagaraEditor/Private/ViewModels/NiagaraScratchPadViewModel.cpp
--- a/Engine/Plugins/FX/Niagara/Source/NiagaraEditor/Private/ViewModels/NiagaraScratchPadViewModel.cpp
+++ b/Engine/Plugins/FX/Niagara/Source/NiagaraEditor/Private/ViewModels/NiagaraScratchPadViewModel.cpp
@@ -216,17 +216,10 @@ void UNiagaraScratchPadViewModel::RefreshScriptViewModels()
 	{
 		for (int32 i = 0; i < ScriptViewModels.Num(); i++)
 		{
-			// First double-check that we're not creating multiple view models per script....
 			TArray<TSharedPtr<FNiagaraScriptViewModel>> ExistingViewModels;
 			TNiagaraViewModelManager<UNiagaraScript, FNiagaraScriptViewModel>::GetAllViewModelsForObject(ScriptViewModels[i]->GetEditScript().Script, ExistingViewModels);
 
 			(ensureMsgf(ExistingViewModels.Num() == 1, TEXT("Missing or duplicate script view models detected.  Can not create script details customization.")));
-
-			// Now double-check that we haven't reused a view model where the underlying script name has changed and need to act as if a rename happened.
-			if (ScriptViewModels[i]->ValidateAndUpdateCachedNameIfDifferent())
-			{
-				ScriptViewModelScriptRenamed(ScriptViewModels[i]);
-			}
 		}
 	}
 	
@@ -480,17 +473,10 @@ FName GetUniqueScriptName(UObject* Outer, const FString& CandidateName)
 
 TSharedPtr<FNiagaraScratchPadScriptViewModel> UNiagaraScratchPadViewModel::CreateNewScript(ENiagaraScriptUsage InScriptUsage, ENiagaraScriptUsage InTargetSupportedUsage, FNiagaraTypeDefinition InOutputType)
 {
-	FScopedTransaction Transaction(LOCTEXT("CreateNewScriptTransaction", "Create a new local script."));
-
 	UObject* ScriptOuter;
 	TArray<TObjectPtr<UNiagaraScript>>* TargetScripts;
 	GetOuterAndTargetScripts(GetSystemViewModel(), ScriptOuter, TargetScripts);
 
-	if (ScriptOuter)
-	{
-		ScriptOuter->Modify();
-	}
-
 	UNiagaraScript* NewScript = nullptr;
 	switch (InScriptUsage)
 	{
@@ -544,7 +530,7 @@ TSharedPtr<FNiagaraScratchPadScriptViewModel> UNiagaraScratchPadViewModel::Creat
 	}
 
 	NewScript->ClearFlags(RF_Public | RF_Standalone);
-	NewScript->SetFlags(RF_Transactional);
+	ScriptOuter->Modify();
 	TargetScripts->Add(NewScript);
 	NewScript->GetLatestScriptData()->ModuleUsageBitmask |= (1 << (int32)InTargetSupportedUsage);
 	RefreshScriptViewModels();
@@ -555,21 +541,13 @@ TSharedPtr<FNiagaraScratchPadScriptViewModel> UNiagaraScratchPadViewModel::Creat
 
 TSharedPtr<FNiagaraScratchPadScriptViewModel> UNiagaraScratchPadViewModel::CreateNewScriptAsDuplicate(const UNiagaraScript* ScriptToDuplicate)
 {
-
-	FScopedTransaction Transaction(LOCTEXT("CreateNewScriptDupeTransaction", "Create a new local script as duplicate."));
-
 	UObject* ScriptOuter;
 	TArray<TObjectPtr<UNiagaraScript>>* TargetScripts;
 	GetOuterAndTargetScripts(GetSystemViewModel(), ScriptOuter, TargetScripts);
 
-	if (ScriptOuter)
-	{
-		ScriptOuter->Modify();
-	}
-	
 	UNiagaraScript* NewScript = CastChecked<UNiagaraScript>(StaticDuplicateObject(ScriptToDuplicate, ScriptOuter, GetUniqueScriptName(ScriptOuter, *ScriptToDuplicate->GetFName().ToString())));
 	NewScript->ClearFlags(RF_Public | RF_Standalone);
-	NewScript->SetFlags(RF_Transactional);
+	ScriptOuter->Modify();
 	TargetScripts->Add(NewScript);
 	RefreshScriptViewModels();
 	UpdateChangeId(GetSystemViewModel());
@@ -742,8 +720,6 @@ TSharedRef<FNiagaraScratchPadScriptViewModel> UNiagaraScratchPadViewModel::Creat
 	ScriptViewModel->Initialize(ScratchPadScript, ScratchPadEditScript, SystemViewModelWeak);
 	ScriptViewModel->GetGraphViewModel()->GetNodeSelection()->OnSelectedObjectsChanged().AddUObject(this, &UNiagaraScratchPadViewModel::ScriptGraphNodeSelectionChanged, TWeakPtr<FNiagaraScratchPadScriptViewModel>(ScriptViewModel));
 	ScriptViewModel->OnRenamed().AddUObject(this, &UNiagaraScratchPadViewModel::ScriptViewModelScriptRenamed, TWeakPtr<FNiagaraScratchPadScriptViewModel>(ScriptViewModel));
-	ScriptViewModel->OnPendingRename().AddUObject(this, &UNiagaraScratchPadViewModel::ScriptViewModelScriptPendingRename, TWeakPtr<FNiagaraScratchPadScriptViewModel>(ScriptViewModel));
-
 	ScriptViewModel->OnGraphSelectionChanged().AddUObject(this, &UNiagaraScratchPadViewModel::ScriptViewModelGraphSelectionChanged, TWeakPtr<FNiagaraScratchPadScriptViewModel>(ScriptViewModel));
 	ScriptViewModel->OnPinnedChanged().AddUObject(this, &UNiagaraScratchPadViewModel::ScriptViewModelPinnedChanged, TWeakPtr<FNiagaraScratchPadScriptViewModel>(ScriptViewModel));
 	ScriptViewModel->OnHasUnappliedChangesChanged().AddUObject(this, &UNiagaraScratchPadViewModel::ScriptViewModelHasUnappliedChangesChanged);
@@ -764,13 +740,6 @@ void UNiagaraScratchPadViewModel::TearDownScriptViewModel(TSharedRef<FNiagaraScr
 	InScriptViewModel->OnRequestDiscardChanges().Unbind();
 	InScriptViewModel->GetVariableSelection()->OnSelectedObjectsChanged().RemoveAll(this);
 	InScriptViewModel->Finalize();
-
-	UNiagaraSystemEditorDocumentsViewModel* DocVM = GetSystemViewModel()->GetDocumentViewModel();
-	if (DocVM)
-	{
-		DocVM->CloseChildScript(InScriptViewModel->GetGraphViewModel()->GetGraph());
-	}
-
 }
 
 void UNiagaraScratchPadViewModel::RefreshEditScriptViewModels()
@@ -806,16 +775,6 @@ void UNiagaraScratchPadViewModel::ScriptGraphNodeSelectionChanged(TWeakPtr<FNiag
 	}
 }
 
-void UNiagaraScratchPadViewModel::PostUndo()
-{
-	RefreshScriptViewModels();
-	if (ObjectSelection)
-	{
-		ObjectSelection->ClearSelectedObjects();
-	}
-}
-
-
 void UNiagaraScratchPadViewModel::ScriptViewModelScriptRenamed(TWeakPtr<FNiagaraScratchPadScriptViewModel> ScriptViewModelWeak)
 {
 	TSharedPtr<FNiagaraScratchPadScriptViewModel> ScriptViewModel = ScriptViewModelWeak.Pin();
@@ -833,11 +792,6 @@ void UNiagaraScratchPadViewModel::ScriptViewModelScriptRenamed(TWeakPtr<FNiagara
 	}
 }
 
-void UNiagaraScratchPadViewModel::ScriptViewModelScriptPendingRename(TWeakPtr<FNiagaraScratchPadScriptViewModel> ScriptViewModelWeak)
-{
-	OnFocusNeededDelegate.Broadcast();
-}
-
 void UNiagaraScratchPadViewModel::ScriptViewModelGraphSelectionChanged(const UObject* Obj, TWeakPtr<FNiagaraScratchPadScriptViewModel> ScriptViewModelWeak)
 {
 	const UNiagaraScriptVariable* Var = Cast<UNiagaraScriptVariable>(Obj);
diff --git a/Engine/Plugins/FX/Niagara/Source/NiagaraEditor/Private/ViewModels/NiagaraSystemViewModel.cpp b/Engine/Plugins/FX/Niagara/Source/NiagaraEditor/Private/ViewModels/NiagaraSystemViewModel.cpp
--- a/Engine/Plugins/FX/Niagara/Source/NiagaraEditor/Private/ViewModels/NiagaraSystemViewModel.cpp
+++ b/Engine/Plugins/FX/Niagara/Source/NiagaraEditor/Private/ViewModels/NiagaraSystemViewModel.cpp
@@ -792,12 +792,6 @@ void FNiagaraSystemViewModel::AddReferencedObjects(FReferenceCollector& Collecto
 
 void FNiagaraSystemViewModel::PostUndo(bool bSuccess)
 {
-	// We may have undone scratch script renames or adds, so we need to make sure that they
-	// are accurate and potentially close anything that is now gone.
-	if (ScriptScratchPadViewModel)
-	{
-		ScriptScratchPadViewModel->PostUndo();
-	}
 	// Reset system stack and emitter handle view models to prevent accessing invalid data if they were in the undo operation.
 	ResetStack();
 	System->InvalidateActiveCompiles();
diff --git a/Engine/Plugins/FX/Niagara/Source/NiagaraEditor/Public/ViewModels/NiagaraScratchPadScriptViewModel.h b/Engine/Plugins/FX/Niagara/Source/NiagaraEditor/Public/ViewModels/NiagaraScratchPadScriptViewModel.h
--- a/Engine/Plugins/FX/Niagara/Source/NiagaraEditor/Public/ViewModels/NiagaraScratchPadScriptViewModel.h
+++ b/Engine/Plugins/FX/Niagara/Source/NiagaraEditor/Public/ViewModels/NiagaraScratchPadScriptViewModel.h
@@ -14,7 +14,6 @@ class NIAGARAEDITOR_API FNiagaraScratchPadScriptViewModel : public FNiagaraScrip
 {
 public:
 	DECLARE_MULTICAST_DELEGATE(FOnRenamed);
-	DECLARE_MULTICAST_DELEGATE(FOnPendingRename);
 	DECLARE_MULTICAST_DELEGATE(FOnPinnedChanged);
 	DECLARE_MULTICAST_DELEGATE(FOnHasUnappliedChangesChanged);
 	DECLARE_MULTICAST_DELEGATE(FOnChangesApplied);
@@ -85,8 +84,6 @@ public:
 
 	FOnRenamed& OnRenamed();
 
-	FOnPendingRename& OnPendingRename();
-
 	FOnPinnedChanged& OnPinnedChanged();
 
 	FOnHasUnappliedChangesChanged& OnHasUnappliedChangesChanged();
@@ -103,9 +100,6 @@ public:
 	void RaisePinFocusRequested(FNiagaraScriptIDAndGraphFocusInfo*  InFocusInfo);
 	void RaiseNodeFocusRequested(FNiagaraScriptIDAndGraphFocusInfo* InFocusInfo);
 
-	// Updates the cached name post undo to see if we need to broadcast a changed event. Returns true if they were out of sync and tabs will potentially need to be reopned.
-	bool ValidateAndUpdateCachedNameIfDifferent();
-
 private:
 	FText GetDisplayNameInternal() const;
 
@@ -130,7 +124,6 @@ private:
 	FDelegateHandle OnGraphNeedsRecompileHandle;
 
 	FOnRenamed OnRenamedDelegate;
-	FOnPendingRename OnPendingRenameDelegate;
 	FOnPinnedChanged OnPinnedChangedDelegate;
 	FOnHasUnappliedChangesChanged OnHasUnappliedChangesChangedDelegate;
 	FOnChangesApplied OnChangesAppliedDelegate;
@@ -139,6 +132,4 @@ private:
 
 	FOnNodeIDFocusRequested OnNodeIDFocusRequestedDelegate;
 	FOnPinIDFocusRequested OnPinIDFocusRequestedDelegate;
-
-	TOptional<FString> CachedName;
 };
\ No newline at end of file
diff --git a/Engine/Plugins/FX/Niagara/Source/NiagaraEditor/Public/ViewModels/NiagaraScratchPadViewModel.h b/Engine/Plugins/FX/Niagara/Source/NiagaraEditor/Public/ViewModels/NiagaraScratchPadViewModel.h
--- a/Engine/Plugins/FX/Niagara/Source/NiagaraEditor/Public/ViewModels/NiagaraScratchPadViewModel.h
+++ b/Engine/Plugins/FX/Niagara/Source/NiagaraEditor/Public/ViewModels/NiagaraScratchPadViewModel.h
@@ -21,7 +21,6 @@ public:
 	DECLARE_MULTICAST_DELEGATE(FOnActiveScriptChanged);
 	DECLARE_MULTICAST_DELEGATE(FOnScriptRenamed);
 	DECLARE_MULTICAST_DELEGATE(FOnScriptDeleted);
-	DECLARE_MULTICAST_DELEGATE(FOnFocusNeeded);
 
 public:
 	void Initialize(TSharedRef<FNiagaraSystemViewModel> InSystemViewModel);
@@ -87,12 +86,6 @@ public:
 
 	void OpenEditorForActive();
 
-	FOnFocusNeeded& OnFocusNeeded() {
-		return OnFocusNeededDelegate;
-	};
-
-	void PostUndo();
-
 private:
 	TSharedRef<FNiagaraSystemViewModel> GetSystemViewModel();
 
@@ -108,8 +101,6 @@ private:
 
 	void ScriptViewModelScriptRenamed(TWeakPtr<FNiagaraScratchPadScriptViewModel> ScriptViewModelWeak);
 
-	void ScriptViewModelScriptPendingRename(TWeakPtr<FNiagaraScratchPadScriptViewModel> ScriptViewModelWeak);	
-
 	void ScriptViewModelPinnedChanged(TWeakPtr<FNiagaraScratchPadScriptViewModel> ScriptViewModelWeak);
 
 	void ScriptViewModelHasUnappliedChangesChanged();
@@ -150,8 +141,5 @@ private:
 
 	FOnScriptDeleted OnScriptDeletedDelegate;
 
-	FOnFocusNeeded OnFocusNeededDelegate;
-
-
 	bool bIsBulkApplying = false;
 };
\ No newline at end of file
