Fixed architectural issue that centralized the OnCanAcceptDrop functionality in one function instead of letting individual view models handle it individually and streamlined the section drag & drop handling

#rnx

diff --git a/Engine/Plugins/FX/Niagara/Source/NiagaraEditor/Private/ViewModels/HierarchyEditor/NiagaraHierarchyViewModelBase.cpp b/Engine/Plugins/FX/Niagara/Source/NiagaraEditor/Private/ViewModels/HierarchyEditor/NiagaraHierarchyViewModelBase.cpp
--- a/Engine/Plugins/FX/Niagara/Source/NiagaraEditor/Private/ViewModels/HierarchyEditor/NiagaraHierarchyViewModelBase.cpp
+++ b/Engine/Plugins/FX/Niagara/Source/NiagaraEditor/Private/ViewModels/HierarchyEditor/NiagaraHierarchyViewModelBase.cpp
@@ -181,6 +181,47 @@ bool FNiagaraHierarchyCategoryViewModel::IsTopCategoryActive() const
 	return false;	
 }
 
+TOptional<EItemDropZone> FNiagaraHierarchyCategoryViewModel::OnCanAcceptDropInternal(TSharedPtr<FDragDropOperation> DragDropOp, EItemDropZone ItemDropZone)
+{
+	bool bAllowDrop = false;
+
+	if(TSharedPtr<FNiagaraHierarchyDragDropOp> HierarchyDragDropOp = StaticCastSharedPtr<FNiagaraHierarchyDragDropOp>(DragDropOp))
+	{
+		TSharedPtr<FNiagaraHierarchyItemViewModelBase> SourceDropItem = HierarchyDragDropOp->GetDraggedItem().Pin();
+		TSharedPtr<FNiagaraHierarchyItemViewModelBase> TargetDropItem = AsShared();
+		
+		TArray<TSharedPtr<FNiagaraHierarchyCategoryViewModel>> TargetChildrenCategories;
+		TargetDropItem->GetChildrenViewModelsForType<UNiagaraHierarchyCategory, FNiagaraHierarchyCategoryViewModel>(TargetChildrenCategories);
+
+		// we only allow drops if some general conditions are fulfilled
+		if(SourceDropItem->GetData() != TargetDropItem->GetData() &&
+			(!SourceDropItem->HasParent(TargetDropItem, false) || ItemDropZone != EItemDropZone::OntoItem)  &&
+			!TargetDropItem->HasParent(SourceDropItem, true))
+		{
+			// categories can be dropped on categories, but only if the resulting sibling categories have different names
+			if(SourceDropItem->GetData()->IsA<UNiagaraHierarchyCategory>())
+			{
+				bAllowDrop =
+				(
+					(ItemDropZone == EItemDropZone::OntoItem && !TargetChildrenCategories.ContainsByPredicate([SourceDropItem](TSharedPtr<FNiagaraHierarchyCategoryViewModel> HierarchyCategoryViewModel)
+						{
+							return SourceDropItem->ToString() == HierarchyCategoryViewModel->ToString();
+						}))
+						||
+					(ItemDropZone != EItemDropZone::OntoItem && SourceDropItem->ToString() != TargetDropItem->ToString())
+				);
+			}
+			else if(SourceDropItem->GetData()->IsA<UNiagaraHierarchyItem>())
+			{
+				// items can generally be dropped onto categories
+				bAllowDrop = EItemDropZone::OntoItem == ItemDropZone;
+			}
+		}
+	}
+
+	return bAllowDrop ? ItemDropZone : TOptional<EItemDropZone>();
+}
+
 void UNiagaraHierarchySection::SetSectionNameAsText(const FText& Text)
 {
 	Section = FName(Text.ToString());
@@ -555,11 +596,11 @@ void FNiagaraHierarchyItemViewModelBase::Finalize()
 
 TOptional<EItemDropZone> FNiagaraHierarchyItemViewModelBase::OnCanAcceptDrop(const FDragDropEvent& DragDropEvent, EItemDropZone ItemDropZone, TSharedPtr<FNiagaraHierarchyItemViewModelBase> Item)
 {
-	if(TSharedPtr<FNiagaraHierarchyDragDropOp> HierarchyDragDropOp = DragDropEvent.GetOperationAs<FNiagaraHierarchyDragDropOp>())
-	{
-		return HierarchyViewModel->CanDropOn(HierarchyDragDropOp->GetDraggedItem().Pin().ToSharedRef(), AsShared(), ItemDropZone);
-	}
+	return OnCanAcceptDropInternal(DragDropEvent.GetOperation(), ItemDropZone);
+}
 
+TOptional<EItemDropZone> FNiagaraHierarchyItemViewModelBase::OnCanAcceptDropInternal(TSharedPtr<FDragDropOperation> DragDropOp, EItemDropZone ItemDropZone)
+{
 	return TOptional<EItemDropZone>();	
 }
 
@@ -598,6 +639,33 @@ FReply FNiagaraHierarchyItemViewModelBase::OnDragDetected(const FGeometry& Geome
 	return FReply::Unhandled();
 }
 
+TOptional<EItemDropZone> FNiagaraHierarchyRootViewModel::OnCanAcceptDropInternal(TSharedPtr<FDragDropOperation> DragDropOp, EItemDropZone ItemDropZone)
+{
+	bool bAllowDrop = false;
+
+	if(TSharedPtr<FNiagaraHierarchyDragDropOp> HierarchyDragDropOp = StaticCastSharedPtr<FNiagaraHierarchyDragDropOp>(DragDropOp))
+	{
+		TSharedPtr<FNiagaraHierarchyItemViewModelBase> SourceDropItem = HierarchyDragDropOp->GetDraggedItem().Pin();
+		TSharedPtr<FNiagaraHierarchyItemViewModelBase> TargetDropItem = AsShared();
+
+		// we only allow drops if some general conditions are fulfilled
+		if(SourceDropItem->GetData() != TargetDropItem->GetData() &&
+			(!SourceDropItem->HasParent(TargetDropItem, false) || ItemDropZone != EItemDropZone::OntoItem)  &&
+			!TargetDropItem->HasParent(SourceDropItem, true))
+		{
+			bAllowDrop = 
+				// items can be dropped onto the root directly if the section is set to "All"
+				(SourceDropItem->GetData()->IsA<UNiagaraHierarchyItem>() && HierarchyViewModel->GetActiveSectionData() == nullptr)
+					||
+				// categories can be dropped onto the root always
+				(SourceDropItem->GetData()->IsA<UNiagaraHierarchyCategory>());
+		}
+		
+	}
+
+	return bAllowDrop ? ItemDropZone : TOptional<EItemDropZone>();
+}
+
 FReply FNiagaraHierarchyRootViewModel::OnDroppedOn(const FDragDropEvent& DragDropEvent, EItemDropZone ItemDropZone, TSharedPtr<FNiagaraHierarchyItemViewModelBase> Item)
 {
 	if(DragDropEvent.GetOperation()->IsOfType<FNiagaraHierarchyDragDropOp>())
@@ -688,6 +756,42 @@ void FNiagaraHierarchyRootViewModel::SyncChildrenViewModelsInternal()
 		});
 }
 
+TOptional<EItemDropZone> FNiagaraHierarchySectionViewModel::OnCanAcceptDropInternal(TSharedPtr<FDragDropOperation> DragDropOp, EItemDropZone ItemDropZone)
+{
+	bool bAllowDrop = false;
+
+	if(DragDropOp->IsOfType<FNiagaraHierarchyDragDropOp>())
+	{
+		TSharedPtr<FNiagaraHierarchyDragDropOp> HierarchyDragDropOp = StaticCastSharedPtr<FNiagaraHierarchyDragDropOp>(DragDropOp);
+		TWeakPtr<FNiagaraHierarchyItemViewModelBase> DraggedItem = HierarchyDragDropOp->GetDraggedItem();
+
+		if(const UNiagaraHierarchyCategory* Category = Cast<UNiagaraHierarchyCategory>(DraggedItem.Pin()->GetData()))
+		{
+			if(ItemDropZone == EItemDropZone::OntoItem)
+			{
+				bAllowDrop = GetData() != Category->GetSection();
+			}
+		}
+		else if(UNiagaraHierarchySection* Section = Cast<UNiagaraHierarchySection>(DraggedItem.Pin()->GetDataMutable()))
+		{
+			if(ItemDropZone != EItemDropZone::OntoItem)
+			{
+				int32 DraggedItemIndex = GetHierarchyViewModel()->GetHierarchyDataRoot()->GetSectionIndex(Section);
+				bool bSameSection = GetData() == Section;
+
+				int32 InsertionIndex = GetHierarchyViewModel()->GetHierarchyDataRoot()->GetSectionIndex(GetDataMutable<UNiagaraHierarchySection>());
+				// we add 1 to the insertion index if it's below an item because we either want to insert at the current index to place the item above, or at current+1 for below
+				InsertionIndex += ItemDropZone == EItemDropZone::AboveItem ? -1 : 1;
+
+				bAllowDrop = !bSameSection && DraggedItemIndex != InsertionIndex;			
+			}
+		}
+	}
+	 
+
+	return bAllowDrop ? ItemDropZone : TOptional<EItemDropZone>();
+}
+
 FReply FNiagaraHierarchySectionViewModel::OnDroppedOn(const FDragDropEvent& DragDropEvent, EItemDropZone ItemDropZone, TSharedPtr<FNiagaraHierarchyItemViewModelBase> Item)
 {
 	if(DragDropEvent.GetOperation()->IsOfType<FNiagaraSectionDragDropOp>())
@@ -793,6 +897,28 @@ void FNiagaraHierarchySectionViewModel::FinalizeInternal()
 	}
 }
 
+TOptional<EItemDropZone> FNiagaraHierarchyItemViewModel::OnCanAcceptDropInternal(TSharedPtr<FDragDropOperation> DragDropOp, EItemDropZone ItemDropZone)
+{
+	bool bAllowDrop = false;
+
+	if(TSharedPtr<FNiagaraHierarchyDragDropOp> HierarchyDragDropOp = StaticCastSharedPtr<FNiagaraHierarchyDragDropOp>(DragDropOp))
+	{
+		TSharedPtr<FNiagaraHierarchyItemViewModelBase> SourceDropItem = HierarchyDragDropOp->GetDraggedItem().Pin();
+		TSharedPtr<FNiagaraHierarchyItemViewModelBase> TargetDropItem = AsShared();
+
+		// we only allow drops if some general conditions are fulfilled
+		if(SourceDropItem->GetData() != TargetDropItem->GetData() &&
+			(!SourceDropItem->HasParent(TargetDropItem, false) || ItemDropZone != EItemDropZone::OntoItem)  &&
+			!TargetDropItem->HasParent(SourceDropItem, true))
+		{
+			// items can be generally be dropped above/below other items
+			bAllowDrop = (SourceDropItem->GetData()->IsA<UNiagaraHierarchyItem>() && ItemDropZone != EItemDropZone::OntoItem);
+		}
+	}
+
+	return bAllowDrop ? ItemDropZone : TOptional<EItemDropZone>();
+}
+
 FReply FNiagaraHierarchyItemViewModel::OnDroppedOn(const FDragDropEvent& DragDropEvent, EItemDropZone ItemDropZone, TSharedPtr<FNiagaraHierarchyItemViewModelBase> Item)
 {
 	FScopedTransaction Transaction(LOCTEXT("Transaction_MovedItem", "Moved an item in the hierarchy"));
diff --git a/Engine/Plugins/FX/Niagara/Source/NiagaraEditor/Private/ViewModels/HierarchyEditor/NiagaraUserParametersHierarchyViewModel.cpp b/Engine/Plugins/FX/Niagara/Source/NiagaraEditor/Private/ViewModels/HierarchyEditor/NiagaraUserParametersHierarchyViewModel.cpp
--- a/Engine/Plugins/FX/Niagara/Source/NiagaraEditor/Private/ViewModels/HierarchyEditor/NiagaraUserParametersHierarchyViewModel.cpp
+++ b/Engine/Plugins/FX/Niagara/Source/NiagaraEditor/Private/ViewModels/HierarchyEditor/NiagaraUserParametersHierarchyViewModel.cpp
@@ -137,33 +137,6 @@ void UNiagaraUserParametersHierarchyViewModel::SetupCommands()
 	// no custom commands yet
 }
 
-TOptional<EItemDropZone> UNiagaraUserParametersHierarchyViewModel::CanDropOn(TSharedPtr<FNiagaraHierarchyItemViewModelBase> SourceDropItem, TSharedPtr<FNiagaraHierarchyItemViewModelBase> TargetDropItem, EItemDropZone DropZone)
-{
-	bool bAllowDropOn =
-		SourceDropItem->GetData() != TargetDropItem->GetData() &&
-			(!SourceDropItem->HasParent(TargetDropItem, false) || DropZone != EItemDropZone::OntoItem)  &&
-			!TargetDropItem->HasParent(SourceDropItem, true)
-		&&
-		(
-			// user parameters can be dropped onto categories
-			(SourceDropItem->GetData()->IsA<UNiagaraHierarchyUserParameter>() && TargetDropItem->GetData()->IsA<UNiagaraHierarchyCategory>() && DropZone == EItemDropZone::OntoItem)
-				||
-			// user parameters can be dropped above/below other user parameters
-			(SourceDropItem->GetData()->IsA<UNiagaraHierarchyUserParameter>() && TargetDropItem->GetData()->IsA<UNiagaraHierarchyUserParameter>() && DropZone != EItemDropZone::OntoItem)
-				||
-			// categories can be dropped on categories
-			(SourceDropItem->GetData()->IsA<UNiagaraHierarchyCategory>() && TargetDropItem->GetData()->IsA<UNiagaraHierarchyCategory>())
-				||
-			// user parameters can be dropped onto the root if the section is set to "All"
-			(SourceDropItem->GetData()->IsA<UNiagaraHierarchyUserParameter>() && TargetDropItem->GetData()->IsA<UNiagaraHierarchyRoot>() && GetActiveSectionData() == nullptr)
-				||
-			// categories can be dropped onto the root always
-			(SourceDropItem->GetData()->IsA<UNiagaraHierarchyCategory>() && TargetDropItem->GetData()-IsA<UNiagaraHierarchyRoot>())
-		);
-	
-	return bAllowDropOn ? DropZone : TOptional<EItemDropZone>();
-}
-
 TSharedRef<FNiagaraHierarchyDragDropOp> UNiagaraUserParametersHierarchyViewModel::CreateDragDropOp(TSharedRef<FNiagaraHierarchyItemViewModelBase> Item)
 {
 	if(const UNiagaraHierarchyUserParameter* UserParameter = Cast<UNiagaraHierarchyUserParameter>(Item->GetData()))
diff --git a/Engine/Plugins/FX/Niagara/Source/NiagaraEditor/Private/Widgets/SNiagaraHierarchy.cpp b/Engine/Plugins/FX/Niagara/Source/NiagaraEditor/Private/Widgets/SNiagaraHierarchy.cpp
--- a/Engine/Plugins/FX/Niagara/Source/NiagaraEditor/Private/Widgets/SNiagaraHierarchy.cpp
+++ b/Engine/Plugins/FX/Niagara/Source/NiagaraEditor/Private/Widgets/SNiagaraHierarchy.cpp
@@ -98,17 +98,17 @@ void SNiagaraHierarchySection::Construct(const FArguments& InArgs, TSharedPtr<FN
 		SDropTarget::FArguments RightDropTargetArgs;
 
 		LeftDropTargetArgs
-			.OnAllowDrop(this, &SNiagaraHierarchySection::OnCanAcceptSectionDrop, EItemDropZone::AboveItem)
+			.OnAllowDrop(this, &SNiagaraHierarchySection::OnCanAcceptDrop, EItemDropZone::AboveItem)
 			.OnDropped(this, &SNiagaraHierarchySection::OnDroppedOn, EItemDropZone::AboveItem)
 			.VerticalImage(FAppStyle::GetNoBrush())
 			.HorizontalImage(FNiagaraEditorStyle::Get().GetBrush("NiagaraEditor.Stack.DropTarget.BorderHorizontal"));
 
 		OntoDropTargetArgs
-			.OnAllowDrop(this, &SNiagaraHierarchySection::OnCanAcceptDrop)
+			.OnAllowDrop(this, &SNiagaraHierarchySection::OnCanAcceptDrop, EItemDropZone::OntoItem)
 			.OnDropped(this, &SNiagaraHierarchySection::OnDroppedOn, EItemDropZone::OntoItem);
 		
 		RightDropTargetArgs
-			.OnAllowDrop(this, &SNiagaraHierarchySection::OnCanAcceptSectionDrop, EItemDropZone::BelowItem)
+			.OnAllowDrop(this, &SNiagaraHierarchySection::OnCanAcceptDrop, EItemDropZone::BelowItem)
 			.OnDropped(this, &SNiagaraHierarchySection::OnDroppedOn, EItemDropZone::BelowItem)
 			.VerticalImage(FAppStyle::GetNoBrush())
 			.HorizontalImage(FNiagaraEditorStyle::Get().GetBrush("NiagaraEditor.Stack.DropTarget.BorderHorizontal"));
@@ -162,7 +162,7 @@ void SNiagaraHierarchySection::Construct(const FArguments& InArgs, TSharedPtr<FN
 		ChildSlot
 		[
 			SNew(SDropTarget)
-			.OnAllowDrop(this, &SNiagaraHierarchySection::OnCanAcceptDrop)
+			.OnAllowDrop(this, &SNiagaraHierarchySection::OnCanAcceptDrop, EItemDropZone::OntoItem)
 			.OnDropped(this, &SNiagaraHierarchySection::OnDroppedOn, EItemDropZone::OntoItem)
 			[
 				SNew(SCheckBox)
@@ -188,44 +188,24 @@ void SNiagaraHierarchySection::EnterEditingMode() const
 	InlineEditableTextBlock->EnterEditingMode();
 }
 
-bool SNiagaraHierarchySection::OnCanAcceptDrop(TSharedPtr<FDragDropOperation> DragDropOperation) const
+bool SNiagaraHierarchySection::OnCanAcceptDrop(TSharedPtr<FDragDropOperation> DragDropOperation, EItemDropZone ItemDropZone) const
 {
-	if(DragDropOperation->IsOfType<FNiagaraHierarchyDragDropOp>())
+	if(SectionViewModel.IsValid())
 	{
-		TSharedPtr<FNiagaraHierarchyDragDropOp> HierarchyDragDropOp = StaticCastSharedPtr<FNiagaraHierarchyDragDropOp>(DragDropOperation);
-		TWeakPtr<FNiagaraHierarchyItemViewModelBase> DraggedItem = HierarchyDragDropOp->GetDraggedItem();
-
-		if(const UNiagaraHierarchyCategory* Category = Cast<UNiagaraHierarchyCategory>(DraggedItem.Pin()->GetData()))
-		{
-			return TryGetSectionData() != Category->GetSection();
-		}
+		return SectionViewModel->OnCanAcceptDropInternal(DragDropOperation, ItemDropZone).IsSet();
 	}
-
-	return false;
-}
-
-bool SNiagaraHierarchySection::OnCanAcceptSectionDrop(TSharedPtr<FDragDropOperation> DragDropOperation, EItemDropZone ItemDropZone) const
-{
-	if(DragDropOperation->IsOfType<FNiagaraSectionDragDropOp>() && ItemDropZone != EItemDropZone::OntoItem)
+	// for the All section which has no valid view model, we simply do a check if the sections of the dragged categories are different
+	else
 	{
-		TSharedPtr<FNiagaraSectionDragDropOp> SectionDragDropOp = StaticCastSharedPtr<FNiagaraSectionDragDropOp>(DragDropOperation);
-		TWeakPtr<FNiagaraHierarchyItemViewModelBase> DraggedItem = SectionDragDropOp->GetDraggedItem();
-		if(UNiagaraHierarchySection* DraggedSection = Cast<UNiagaraHierarchySection>(DraggedItem.Pin()->GetDataMutable()))
+		if(DragDropOperation->IsOfType<FNiagaraHierarchyDragDropOp>())
 		{
-			int32 DraggedItemIndex = SectionViewModel->GetHierarchyViewModel()->GetHierarchyDataRoot()->GetSectionIndex(DraggedSection);
-			bool bSameSection = TryGetSectionData() == DraggedSection;
+			TSharedPtr<FNiagaraHierarchyDragDropOp> HierarchyDragDropOp = StaticCastSharedPtr<FNiagaraHierarchyDragDropOp>(DragDropOperation);
+			TWeakPtr<FNiagaraHierarchyItemViewModelBase> DraggedItem = HierarchyDragDropOp->GetDraggedItem();
 
-			int32 InsertionIndex = INDEX_NONE;
-			if(SectionViewModel.IsValid())
+			if(const UNiagaraHierarchyCategory* Category = Cast<UNiagaraHierarchyCategory>(DraggedItem.Pin()->GetData()))
 			{
-				InsertionIndex = SectionViewModel->GetHierarchyViewModel()->GetHierarchyDataRoot()->GetSectionIndex(TryGetSectionData());
-
-				// we add 1 to the insertion index if it's below an item because we either want to insert at the current index to place the item above, or at current+1 for below
-				InsertionIndex += ItemDropZone == EItemDropZone::AboveItem ? -1 : 1;
+				return TryGetSectionData() != Category->GetSection();
 			}
-
-			bool bCanDrop = !bSameSection && DraggedItemIndex != InsertionIndex;			
-			return bCanDrop;
 		}
 	}
 
@@ -820,7 +800,7 @@ bool SNiagaraHierarchy::CanDropOnRoot(TSharedPtr<FDragDropOperation> DragDropOp)
 	{
 		TSharedPtr<FNiagaraHierarchyDragDropOp> HierarchyDragDropOp = StaticCastSharedPtr<FNiagaraHierarchyDragDropOp>(DragDropOp);
 		TSharedPtr<FNiagaraHierarchyItemViewModelBase> DraggedItem = HierarchyDragDropOp->GetDraggedItem().Pin();
-		TOptional<EItemDropZone> Result = HierarchyViewModel->CanDropOn(DraggedItem, HierarchyViewModel->GetHierarchyViewModelRoot(), EItemDropZone::OntoItem);
+		TOptional<EItemDropZone> Result = HierarchyViewModel->GetHierarchyViewModelRoot()->OnCanAcceptDropInternal(DragDropOp, EItemDropZone::OntoItem);
 		return Result.IsSet() && Result.GetValue() == EItemDropZone::OntoItem ? true : false;
 	}
 
diff --git a/Engine/Plugins/FX/Niagara/Source/NiagaraEditor/Public/ViewModels/HierarchyEditor/NiagaraHierarchyViewModelBase.h b/Engine/Plugins/FX/Niagara/Source/NiagaraEditor/Public/ViewModels/HierarchyEditor/NiagaraHierarchyViewModelBase.h
--- a/Engine/Plugins/FX/Niagara/Source/NiagaraEditor/Public/ViewModels/HierarchyEditor/NiagaraHierarchyViewModelBase.h
+++ b/Engine/Plugins/FX/Niagara/Source/NiagaraEditor/Public/ViewModels/HierarchyEditor/NiagaraHierarchyViewModelBase.h
@@ -277,9 +277,8 @@ public:
 	/** Additional commands can be specified overriding the Commands function. */
 	virtual void SetupCommands() {}
 	TSharedRef<FUICommandList> GetCommands() const { return Commands.ToSharedRef(); }
-	
-	/** Central function to determine whether an item can be dropped on, above or below another item. */
-	virtual TOptional<EItemDropZone> CanDropOn(TSharedPtr<FNiagaraHierarchyItemViewModelBase> SourceDropItem, TSharedPtr<FNiagaraHierarchyItemViewModelBase> TargetDropItem, EItemDropZone DropZone = EItemDropZone::OntoItem) PURE_VIRTUAL(UNiagaraHierarchyViewModelBase::CanDropOn, return TOptional<EItemDropZone>(););
+
+	/** Function to implement for custom drag drop ops. */
 	virtual TSharedRef<FNiagaraHierarchyDragDropOp> CreateDragDropOp(TSharedRef<FNiagaraHierarchyItemViewModelBase> Item) PURE_VIRTUAL(UNiagaraHierarchyViewModelBase::CreateDragDropOp, return MakeShared<FNiagaraHierarchyDragDropOp>(nullptr););
 	
 	virtual void OnSelectionChanged(TSharedPtr<FNiagaraHierarchyItemViewModelBase> HierarchyItem) {}
@@ -392,8 +391,9 @@ struct NIAGARAEDITOR_API FNiagaraHierarchyItemViewModelBase : TSharedFromThis<FN
 	{
 		return FReply::Unhandled();
 	}
-	
+
 	TOptional<EItemDropZone> OnCanAcceptDrop(const FDragDropEvent& DragDropEvent, EItemDropZone ItemDropZone, TSharedPtr<FNiagaraHierarchyItemViewModelBase> Item);
+	virtual TOptional<EItemDropZone> OnCanAcceptDropInternal(TSharedPtr<FDragDropOperation> DragDropOp, EItemDropZone ItemDropZone);
 
 	/** Should return true if draggable */
 	virtual bool CanDrag() { return false; }
@@ -465,6 +465,8 @@ struct FNiagaraHierarchyRootViewModel : FNiagaraHierarchyItemViewModelBase
 	virtual ~FNiagaraHierarchyRootViewModel() override {}
 
 	virtual bool CanDrag() override { return false; }
+
+	virtual TOptional<EItemDropZone> OnCanAcceptDropInternal(TSharedPtr<FDragDropOperation> DragDropOp, EItemDropZone ItemDropZone) override;
 	virtual FReply OnDroppedOn(const FDragDropEvent& DragDropEvent, EItemDropZone ItemDropZone, TSharedPtr<FNiagaraHierarchyItemViewModelBase> Item) override;
 
 	TSharedPtr<struct FNiagaraHierarchySectionViewModel> AddNewSection();
@@ -491,6 +493,7 @@ struct FNiagaraHierarchySectionViewModel : FNiagaraHierarchyItemViewModelBase
 	virtual bool CanRename() override { return true; }
 	virtual bool CanDelete() override { return true; }
 
+	virtual TOptional<EItemDropZone> OnCanAcceptDropInternal(TSharedPtr<FDragDropOperation> DragDropOp, EItemDropZone ItemDropZone) override;
 	virtual FReply OnDroppedOn(const FDragDropEvent& DragDropEvent, EItemDropZone ItemDropZone, TSharedPtr<FNiagaraHierarchyItemViewModelBase> Item) override;
 
 	virtual void FinalizeInternal() override;
@@ -503,7 +506,8 @@ struct NIAGARAEDITOR_API FNiagaraHierarchyItemViewModel : FNiagaraHierarchyItemV
 	virtual ~FNiagaraHierarchyItemViewModel() override {}
 
 	virtual bool CanDrag() override { return true; }
-	
+
+	virtual TOptional<EItemDropZone> OnCanAcceptDropInternal(TSharedPtr<FDragDropOperation> DragDropOp, EItemDropZone ItemDropZone) override;
 	virtual FReply OnDroppedOn(const FDragDropEvent& DragDropEvent, EItemDropZone ItemDropZone, TSharedPtr<FNiagaraHierarchyItemViewModelBase> Item) override;
 };
 
@@ -519,6 +523,7 @@ struct FNiagaraHierarchyCategoryViewModel : FNiagaraHierarchyItemViewModelBase
 
 	bool IsTopCategoryActive() const;
 
+	virtual TOptional<EItemDropZone> OnCanAcceptDropInternal(TSharedPtr<FDragDropOperation> DragDropOp, EItemDropZone ItemDropZone) override;
 	virtual FReply OnDroppedOn(const FDragDropEvent& DragDropEvent, EItemDropZone ItemDropZone, TSharedPtr<FNiagaraHierarchyItemViewModelBase> Item) override;
 };
 
diff --git a/Engine/Plugins/FX/Niagara/Source/NiagaraEditor/Public/ViewModels/HierarchyEditor/NiagaraUserParametersHierarchyViewModel.h b/Engine/Plugins/FX/Niagara/Source/NiagaraEditor/Public/ViewModels/HierarchyEditor/NiagaraUserParametersHierarchyViewModel.h
--- a/Engine/Plugins/FX/Niagara/Source/NiagaraEditor/Public/ViewModels/HierarchyEditor/NiagaraUserParametersHierarchyViewModel.h
+++ b/Engine/Plugins/FX/Niagara/Source/NiagaraEditor/Public/ViewModels/HierarchyEditor/NiagaraUserParametersHierarchyViewModel.h
@@ -34,7 +34,7 @@ struct FNiagaraHierarchyUserParameterViewModel : public FNiagaraHierarchyItemVie
 {
 	FNiagaraHierarchyUserParameterViewModel(UNiagaraHierarchyUserParameter* InItem, TSharedPtr<FNiagaraHierarchyItemViewModelBase> InParent, TWeakObjectPtr<UNiagaraHierarchyViewModelBase> InHierarchyViewModel)
 	: FNiagaraHierarchyItemViewModel(InItem, InParent, InHierarchyViewModel) {}
-
+	
 	/** For editing in the details panel we want to handle the script variable that is represented by the hierarchy item, not the hierarchy item itself. */
 	virtual UObject* GetDataForEditing() override;
 };
@@ -58,8 +58,7 @@ public:
 	
 	virtual void PrepareSourceItems() override;
 	virtual void SetupCommands() override;
-
-	virtual TOptional<EItemDropZone> CanDropOn(TSharedPtr<FNiagaraHierarchyItemViewModelBase> SourceDropItem, TSharedPtr<FNiagaraHierarchyItemViewModelBase> TargetDropItem, EItemDropZone DropZone = EItemDropZone::OntoItem) override;
+	
 	virtual TSharedRef<FNiagaraHierarchyDragDropOp> CreateDragDropOp(TSharedRef<FNiagaraHierarchyItemViewModelBase> Item) override;
 	
 	virtual bool SupportsDetailsPanel() override { return true; }
diff --git a/Engine/Plugins/FX/Niagara/Source/NiagaraEditor/Public/Widgets/SNiagaraHierarchy.h b/Engine/Plugins/FX/Niagara/Source/NiagaraEditor/Public/Widgets/SNiagaraHierarchy.h
--- a/Engine/Plugins/FX/Niagara/Source/NiagaraEditor/Public/Widgets/SNiagaraHierarchy.h
+++ b/Engine/Plugins/FX/Niagara/Source/NiagaraEditor/Public/Widgets/SNiagaraHierarchy.h
@@ -68,7 +68,6 @@ private:
 	SVerticalBox::FSlot* SectionsSlot = nullptr;
 	TSharedPtr<IDetailsView> DetailsPanel;
 
-	// Slate widget arguments
 private:
 	FOnGenerateRowContentWidget OnGenerateRowContentWidget;
 	FOnGenerateCustomDetailsPanelNameWidget OnGenerateCustomDetailsPanelNameWidget;
@@ -88,8 +87,7 @@ public:
 	void Construct(const FArguments& InArgs, TSharedPtr<struct FNiagaraHierarchySectionViewModel> InSection, TWeakObjectPtr<UNiagaraHierarchyViewModelBase> HierarchyViewModel);
 	virtual ~SNiagaraHierarchySection() override;
 
-	bool OnCanAcceptDrop(TSharedPtr<FDragDropOperation> DragDropOperation) const;
-	bool OnCanAcceptSectionDrop(TSharedPtr<FDragDropOperation> DragDropOperation, EItemDropZone ItemDropZone) const;
+	bool OnCanAcceptDrop(TSharedPtr<FDragDropOperation> DragDropOperation, EItemDropZone ItemDropZone) const;
 	FReply OnDroppedOn(const FGeometry&, const FDragDropEvent& DragDropEvent, EItemDropZone DropZone) const;
 
 	void EnterEditingMode() const;
