Moved the "Add Parameter" button to the detail's panel's customization & streamlined UI to be consistent with the parameter's palette.
Also removed the "alias" property in the details for a user parameter's metadata

#rnx

diff --git a/Engine/Plugins/FX/Niagara/Source/NiagaraEditor/Private/Customizations/NiagaraComponentDetails.cpp b/Engine/Plugins/FX/Niagara/Source/NiagaraEditor/Private/Customizations/NiagaraComponentDetails.cpp
--- a/Engine/Plugins/FX/Niagara/Source/NiagaraEditor/Private/Customizations/NiagaraComponentDetails.cpp
+++ b/Engine/Plugins/FX/Niagara/Source/NiagaraEditor/Private/Customizations/NiagaraComponentDetails.cpp
@@ -37,7 +37,9 @@
 #include "Widgets/SNiagaraParameterName.h"
 #include "ViewModels/HierarchyEditor/NiagaraUserParametersHierarchyViewModel.h"
 #include "Widgets/Input/SCheckBox.h"
+#include "Widgets/Input/SComboButton.h"
 #include "SNiagaraSystemUserParameters.h"
+#include "Widgets/Input/SEditableTextBox.h"
 
 #define LOCTEXT_NAMESPACE "NiagaraComponentDetails"
 
@@ -570,7 +572,12 @@ void FNiagaraSystemUserParameterDetails::CustomizeDetails(IDetailLayoutBuilder&
 	{
 		System = Cast<UNiagaraSystem>(CustomizedObjects[0]);
 		IDetailCategoryBuilder& InputParamCategory = DetailBuilder.EditCategory(ParamCategoryName, LOCTEXT("ParamCategoryName", "User Parameters"), ECategoryPriority::Important);
-		InputParamCategory.AddCustomBuilder(MakeShared<FNiagaraSystemUserParameterBuilder>(System.Get(), ParamCategoryName));
+		TSharedPtr<FNiagaraSystemViewModel> SystemViewModel = TNiagaraViewModelManager<UNiagaraSystem, FNiagaraSystemViewModel>::GetExistingViewModelForObject(System.Get());
+		ensure(SystemViewModel.IsValid());
+		TSharedRef<FNiagaraSystemUserParameterBuilder> SystemUserParameterBuilder = MakeShared<FNiagaraSystemUserParameterBuilder>(SystemViewModel, ParamCategoryName);
+		InputParamCategory.AddCustomBuilder(SystemUserParameterBuilder);
+		TSharedRef<SWidget> AddParameterButton = SystemUserParameterBuilder->GetAddParameterButton();
+		InputParamCategory.HeaderContent(AddParameterButton);
 	}
 }
 
@@ -912,17 +919,16 @@ void FNiagaraComponentUserParametersNodeBuilder::ParameterValueChanged()
 	}
 }
 
-FNiagaraSystemUserParameterBuilder::FNiagaraSystemUserParameterBuilder(UNiagaraSystem* InSystem, FName InCustomBuilderRowName)
+FNiagaraSystemUserParameterBuilder::FNiagaraSystemUserParameterBuilder(TSharedPtr<FNiagaraSystemViewModel> InSystemViewModel, FName InCustomBuilderRowName)
 {
 	CustomBuilderRowName = InCustomBuilderRowName;
-	System = InSystem;
+	System = &InSystemViewModel->GetSystem();
+	SystemViewModel = InSystemViewModel;
 	bDelegatesInitialized = false;
 }
 
 TSharedPtr<SWidget> FNiagaraSystemUserParameterBuilder::GenerateCustomNameWidget(FNiagaraVariable UserParameter)
-{
-	TSharedPtr<FNiagaraSystemViewModel> SystemViewModel = TNiagaraViewModelManager<UNiagaraSystem, FNiagaraSystemViewModel>::GetExistingViewModelForObject(System.Get());
-	
+{	
 	TSharedRef<SNiagaraParameterNameTextBlock> ParameterName = SNew(SNiagaraParameterNameTextBlock)
 		// if this is specified, we avoid the default behavior, which is "when clicked when the widget had keyboard focus already, enter editing mode"
 		// we don't want to be able to rename by clicking, only by context menu
@@ -951,10 +957,11 @@ void FNiagaraSystemUserParameterBuilder::GenerateChildContent(IDetailChildrenBui
 		if(bDelegatesInitialized == false)
 		{
 			System->GetExposedParameters().OnStructureChanged().Add(FNiagaraParameterStore::FOnStructureChanged::FDelegate::CreateSP(this, &FNiagaraSystemUserParameterBuilder::Rebuild));
-			TSharedPtr<FNiagaraSystemViewModel> SystemViewModel = TNiagaraViewModelManager<UNiagaraSystem, FNiagaraSystemViewModel>::GetExistingViewModelForObject(System.Get());
-			SystemViewModel->GetUserParametersHierarchyViewModel()->OnHierarchyChanged().Add(UNiagaraHierarchyViewModelBase::FOnHierarchyChanged::FDelegate::CreateSP(this, &FNiagaraSystemUserParameterBuilder::Rebuild));
 
-			TSharedPtr<FNiagaraUserParameterPanelViewModel> UserParameterPanelViewModel = SystemViewModel->GetUserParameterPanelViewModel();
+			SystemViewModel.Pin()->GetUserParametersHierarchyViewModel()->OnHierarchyChanged().RemoveAll(this);
+			SystemViewModel.Pin()->GetUserParametersHierarchyViewModel()->OnHierarchyChanged().Add(UNiagaraHierarchyViewModelBase::FOnHierarchyChanged::FDelegate::CreateSP(this, &FNiagaraSystemUserParameterBuilder::Rebuild));
+
+			TSharedPtr<FNiagaraUserParameterPanelViewModel> UserParameterPanelViewModel = SystemViewModel.Pin()->GetUserParameterPanelViewModel();
 			if(ensure(UserParameterPanelViewModel.IsValid()))
 			{
 				UserParameterPanelViewModel->OnParameterAdded().BindSP(this, &FNiagaraSystemUserParameterBuilder::RequestRename);
@@ -987,6 +994,64 @@ void FNiagaraSystemUserParameterBuilder::AddCustomMenuActionsForParameter(FDetai
 	WidgetRow.AddCustomContextMenuAction(DeleteAction, LOCTEXT("DeleteParameterAction", "Delete"), LOCTEXT("DeleteParameterActionTooltip", "Delete this user parameter"));
 }
 
+TSharedRef<SWidget> FNiagaraSystemUserParameterBuilder::GetAddParameterButton()
+{
+	if(!AddParameterButtonContainer.IsValid())
+	{
+		AddParameterButtonContainer = SNew(SBox)
+		.HAlign(HAlign_Right)
+		.VAlign(VAlign_Center)
+		[
+			SAssignNew(AddParameterButton, SComboButton)
+			.HAlign(HAlign_Center)
+			.VAlign(VAlign_Center)
+			.ButtonStyle(FAppStyle::Get(), "RoundButton")
+			.ForegroundColor(FAppStyle::GetSlateColor("DefaultForeground"))
+			.ContentPadding(FMargin(2, 0))
+			.OnGetMenuContent(this, &FNiagaraSystemUserParameterBuilder::GetAddParameterMenu)
+			.OnComboBoxOpened_Lambda([this]
+			{
+				AddParameterButton->SetMenuContentWidgetToFocus(AddParameterMenu->GetSearchBox());
+			})
+			.HasDownArrow(false)
+			.ButtonContent()
+			[
+				SNew(SImage)
+				.Image(FAppStyle::GetBrush("Plus"))
+			]
+		];
+	}
+	
+	return AddParameterButtonContainer.ToSharedRef(); 
+}
+
+TSharedRef<SWidget> FNiagaraSystemUserParameterBuilder::GetAddParameterMenu()
+{
+	FNiagaraParameterPanelCategory UserCategory = GetDefault<UNiagaraEditorSettings>()->GetMetaDataForNamespaces({FNiagaraConstants::UserNamespace});
+	AddParameterMenu = SNew(SNiagaraAddParameterFromPanelMenu)
+		.Graphs(SystemViewModel.Pin()->GetParameterPanelViewModel()->GetEditableGraphsConst())
+		.OnNewParameterRequested(this, &FNiagaraSystemUserParameterBuilder::AddParameter)
+		.OnAllowMakeType(this, &FNiagaraSystemUserParameterBuilder::CanMakeNewParameterOfType)
+		.NamespaceId(UserCategory.NamespaceMetaData.GetGuid())
+		.ShowNamespaceCategory(false)
+		.ShowGraphParameters(false)
+		.AutoExpandMenu(false);
+
+	return AddParameterMenu.ToSharedRef();
+}
+
+void FNiagaraSystemUserParameterBuilder::AddParameter(FNiagaraVariable NewParameter) const
+{
+	FNiagaraParameterPanelCategory UserCategory = GetDefault<UNiagaraEditorSettings>()->GetMetaDataForNamespaces({FNiagaraConstants::UserNamespace});
+	FNiagaraEditorUtilities::AddParameter(NewParameter, SystemViewModel.Pin()->GetSystem().GetExposedParameters(), SystemViewModel.Pin()->GetSystem(), nullptr);
+	SystemViewModel.Pin()->GetUserParameterPanelViewModel()->OnParameterAdded().Execute(NewParameter);
+}
+
+bool FNiagaraSystemUserParameterBuilder::CanMakeNewParameterOfType(const FNiagaraTypeDefinition& InType) const
+{
+	return SystemViewModel.Pin()->GetParameterPanelViewModel()->CanMakeNewParameterOfType(InType);
+}
+
 void FNiagaraSystemUserParameterBuilder::ParameterValueChanged()
 {
 	if (System.IsValid())
@@ -1065,9 +1130,8 @@ void FNiagaraSystemUserParameterBuilder::RenameParameter(FNiagaraVariable UserPa
 		FNiagaraVariable VariableWithNewName(UserParameter);
 		VariableWithNewName.SetName(UniqueName);
 		FNiagaraUserRedirectionParameterStore::MakeUserVariable(VariableWithNewName);
-
-		TSharedPtr<FNiagaraSystemViewModel> SystemViewModel = TNiagaraViewModelManager<UNiagaraSystem, FNiagaraSystemViewModel>::GetExistingViewModelForObject(System.Get());
-		SystemViewModel->RenameParameter(UserParameter, NewName, ENiagaraGetGraphParameterReferencesMode::AllGraphs);
+		
+		SystemViewModel.Pin()->RenameParameter(UserParameter, NewName, ENiagaraGetGraphParameterReferencesMode::AllGraphs);
 	}
 }
 
diff --git a/Engine/Plugins/FX/Niagara/Source/NiagaraEditor/Private/Customizations/NiagaraComponentDetails.h b/Engine/Plugins/FX/Niagara/Source/NiagaraEditor/Private/Customizations/NiagaraComponentDetails.h
--- a/Engine/Plugins/FX/Niagara/Source/NiagaraEditor/Private/Customizations/NiagaraComponentDetails.h
+++ b/Engine/Plugins/FX/Niagara/Source/NiagaraEditor/Private/Customizations/NiagaraComponentDetails.h
@@ -152,22 +152,21 @@ private:
 	TWeakObjectPtr<UNiagaraComponent> Component;
 };
 
-/** A specialized class to generate user parameters for a system asset */
+/** A specialized class to generate user parameters for a system asset's editor. Due to reliance on the system view model this should only be used within the asset editor. */
 class FNiagaraSystemUserParameterBuilder : public FNiagaraUserParameterNodeBuilder
 {
 public:
-	FNiagaraSystemUserParameterBuilder(UNiagaraSystem* InSystem, FName InCustomBuilderRowName);
+	FNiagaraSystemUserParameterBuilder(TSharedPtr<FNiagaraSystemViewModel> InSystemViewModel, FName InCustomBuilderRowName);
 
 	virtual ~FNiagaraSystemUserParameterBuilder() override
 	{
 		if (System.IsValid() && bDelegatesInitialized)
 		{
 			System->GetExposedParameters().OnStructureChanged().RemoveAll(this);
-
-			TSharedPtr<FNiagaraSystemViewModel> SystemViewModel = TNiagaraViewModelManager<UNiagaraSystem, FNiagaraSystemViewModel>::GetExistingViewModelForObject(System.Get());
+			
 			if(SystemViewModel.IsValid())
 			{
-				SystemViewModel->GetUserParametersHierarchyViewModel()->OnHierarchyChanged().RemoveAll(this);
+				SystemViewModel.Pin()->GetUserParametersHierarchyViewModel()->OnHierarchyChanged().RemoveAll(this);
 			}
 		}
 	}
@@ -179,7 +178,13 @@ public:
 	
 	/** We want to add rename & delete actions within a system asset */
 	virtual void AddCustomMenuActionsForParameter(FDetailWidgetRow& WidgetRow, FNiagaraVariable UserParameter) override;
+
+	TSharedRef<SWidget> GetAddParameterButton();
 private:
+	TSharedRef<SWidget> GetAddParameterMenu();
+	void AddParameter(FNiagaraVariable NewParameter) const;
+	bool CanMakeNewParameterOfType(const FNiagaraTypeDefinition& InType) const;
+
 	void ParameterValueChanged();
 
 	void DeleteParameter(FNiagaraVariable UserParameter) const;
@@ -189,6 +194,10 @@ private:
 	FReply GenerateParameterDragDropOp(const FGeometry& Geometry, const FPointerEvent& MouseEvent, FNiagaraVariable UserParameter) const;
 private:
 	TWeakObjectPtr<UNiagaraSystem> System;
+	TWeakPtr<FNiagaraSystemViewModel> SystemViewModel;
 	TOptional<FNiagaraVariable> SelectedParameter;
 	TMap<FNiagaraVariable, TSharedPtr<class SNiagaraParameterNameTextBlock>> UserParamToWidgetMap;
+	TSharedPtr<SWidget> AddParameterButtonContainer;
+	TSharedPtr<class SComboButton> AddParameterButton;
+	TSharedPtr<class SNiagaraAddParameterFromPanelMenu> AddParameterMenu;
 };
diff --git a/Engine/Plugins/FX/Niagara/Source/NiagaraEditor/Private/Customizations/NiagaraScriptVariableCustomization.cpp b/Engine/Plugins/FX/Niagara/Source/NiagaraEditor/Private/Customizations/NiagaraScriptVariableCustomization.cpp
--- a/Engine/Plugins/FX/Niagara/Source/NiagaraEditor/Private/Customizations/NiagaraScriptVariableCustomization.cpp
+++ b/Engine/Plugins/FX/Niagara/Source/NiagaraEditor/Private/Customizations/NiagaraScriptVariableCustomization.cpp
@@ -743,7 +743,7 @@ void FNiagaraScriptVariableHierarchyDetails::CustomizeDetails(IDetailLayoutBuild
 	DetailBuilder.HideProperty(GET_MEMBER_NAME_CHECKED(UNiagaraScriptVariable, Metadata.InlineParameterEnumOverrides));
 	DetailBuilder.HideProperty(GET_MEMBER_NAME_CHECKED(UNiagaraScriptVariable, Metadata.InlineParameterColorOverride));
 	DetailBuilder.HideProperty(GET_MEMBER_NAME_CHECKED(UNiagaraScriptVariable, Metadata.InlineParameterBoolOverride));
-
+	DetailBuilder.HideProperty(GET_MEMBER_NAME_CHECKED(UNiagaraScriptVariable, Metadata.AlternateAliases));
 }
 
 #undef LOCTEXT_NAMESPACE
diff --git a/Engine/Plugins/FX/Niagara/Source/NiagaraEditor/Private/Widgets/SNiagaraSystemUserParameters.cpp b/Engine/Plugins/FX/Niagara/Source/NiagaraEditor/Private/Widgets/SNiagaraSystemUserParameters.cpp
--- a/Engine/Plugins/FX/Niagara/Source/NiagaraEditor/Private/Widgets/SNiagaraSystemUserParameters.cpp
+++ b/Engine/Plugins/FX/Niagara/Source/NiagaraEditor/Private/Widgets/SNiagaraSystemUserParameters.cpp
@@ -41,21 +41,6 @@ void SNiagaraSystemUserParameters::Construct(const FArguments& InArgs, TSharedPt
 		[
 			SNew(SHorizontalBox)
 			+ SHorizontalBox::Slot()
-			.HAlign(HAlign_Left)
-			[
-				SAssignNew(AddParameterButton, SComboButton)
-				.OnGetMenuContent(this, &SNiagaraSystemUserParameters::GetParameterMenu)
-				.OnComboBoxOpened_Lambda([this]
-				{
-					AddParameterButton->SetMenuContentWidgetToFocus(ParameterPanel->GetSearchBox());
-				})
-				.ButtonContent()
-				[
-					SNew(STextBlock)
-					.Text(LOCTEXT("AddUserParameterLabel", "Add Parameter"))
-				]
-			]
-			+ SHorizontalBox::Slot()
 			.HAlign(HAlign_Right)
 			[
 				SNew(SButton)
@@ -75,33 +60,6 @@ void SNiagaraSystemUserParameters::Construct(const FArguments& InArgs, TSharedPt
 	];
 }
 
-TSharedRef<SWidget> SNiagaraSystemUserParameters::GetParameterMenu() 
-{
-	FNiagaraParameterPanelCategory UserCategory = GetDefault<UNiagaraEditorSettings>()->GetMetaDataForNamespaces({FNiagaraConstants::UserNamespace});
-	ParameterPanel = SNew(SNiagaraAddParameterFromPanelMenu)
-		.Graphs(SystemViewModel.Pin()->GetParameterPanelViewModel()->GetEditableGraphsConst())
-		.OnNewParameterRequested(this, &SNiagaraSystemUserParameters::AddParameter)
-		.OnAllowMakeType(this, &SNiagaraSystemUserParameters::CanMakeNewParameterOfType)
-		.NamespaceId(UserCategory.NamespaceMetaData.GetGuid())
-		.ShowNamespaceCategory(false)
-		.ShowGraphParameters(false)
-		.AutoExpandMenu(false);
-	
-	return ParameterPanel.ToSharedRef();
-}
-
-void SNiagaraSystemUserParameters::AddParameter(FNiagaraVariable NewParameter) const
-{
-	FNiagaraParameterPanelCategory UserCategory = GetDefault<UNiagaraEditorSettings>()->GetMetaDataForNamespaces({FNiagaraConstants::UserNamespace});
-	FNiagaraEditorUtilities::AddParameter(NewParameter, SystemViewModel.Pin()->GetSystem().GetExposedParameters(), SystemViewModel.Pin()->GetSystem(), nullptr);
-	SystemViewModel.Pin()->GetUserParameterPanelViewModel()->OnParameterAdded().Execute(NewParameter);
-}
-
-bool SNiagaraSystemUserParameters::CanMakeNewParameterOfType(const FNiagaraTypeDefinition& InType) const
-{
-	return SystemViewModel.Pin()->GetParameterPanelViewModel()->CanMakeNewParameterOfType(InType);
-}
-
 FReply SNiagaraSystemUserParameters::SummonHierarchyEditor()
 {
 	SystemViewModel.Pin()->FocusTab(FNiagaraSystemToolkitModeBase::UserParametersHierarchyTabID);
diff --git a/Engine/Plugins/FX/Niagara/Source/NiagaraEditor/Private/Widgets/SNiagaraSystemUserParameters.h b/Engine/Plugins/FX/Niagara/Source/NiagaraEditor/Private/Widgets/SNiagaraSystemUserParameters.h
--- a/Engine/Plugins/FX/Niagara/Source/NiagaraEditor/Private/Widgets/SNiagaraSystemUserParameters.h
+++ b/Engine/Plugins/FX/Niagara/Source/NiagaraEditor/Private/Widgets/SNiagaraSystemUserParameters.h
@@ -21,15 +21,7 @@ public:
 	virtual void Tick( const FGeometry& AllottedGeometry, const double InCurrentTime, const float InDeltaTime ) override {}
 
 private:
-	TSharedRef<SWidget> GetParameterMenu();
-	
-	void AddParameter(FNiagaraVariable NewParameter) const;
-	bool CanMakeNewParameterOfType(const FNiagaraTypeDefinition& InType) const;
-
 	FReply SummonHierarchyEditor();
 private:
 	TWeakPtr<FNiagaraSystemViewModel> SystemViewModel;
-	
-	TSharedPtr<SComboButton> AddParameterButton;
-	TSharedPtr<SNiagaraAddParameterFromPanelMenu> ParameterPanel;
 };
