Fix static analysis warning with NiagaraGraphDataCache
#rnx

diff --git a/Engine/Plugins/FX/Niagara/Source/NiagaraEditor/Private/NiagaraGraphDataCache.cpp b/Engine/Plugins/FX/Niagara/Source/NiagaraEditor/Private/NiagaraGraphDataCache.cpp
--- a/Engine/Plugins/FX/Niagara/Source/NiagaraEditor/Private/NiagaraGraphDataCache.cpp
+++ b/Engine/Plugins/FX/Niagara/Source/NiagaraEditor/Private/NiagaraGraphDataCache.cpp
@@ -140,7 +140,7 @@ FNiagaraGraphDataCache::FStackFunctionInputPinKey::FStackFunctionInputPinKey(con
 		ContextStaticVariables = StaticVariables;
 	}
 
-	InputPinOption = (uint8)InputPinOption;
+	InputPinOption = (uint8)InOptions;
 	bIgnoreDisabled = bInIgnoreDisabled;
 	bFilterForCompilation = bInFilterForCompilation;
 	ConstantResolverHash = InConstantResolver.BuildTypeHash();
