- Fix missing world / system stat tracking for Niagara
- Fix incorrect avg vs max on HUD display

#rnx

diff --git a/Engine/Plugins/FX/Niagara/Source/Niagara/Private/NiagaraDebugHud.cpp b/Engine/Plugins/FX/Niagara/Source/Niagara/Private/NiagaraDebugHud.cpp
--- a/Engine/Plugins/FX/Niagara/Source/Niagara/Private/NiagaraDebugHud.cpp
+++ b/Engine/Plugins/FX/Niagara/Source/Niagara/Private/NiagaraDebugHud.cpp
@@ -3011,7 +3011,7 @@ bool FNiagaraDebugHUDStatsListener::Tick()
 			if (Settings.PerfSampleMode == ENiagaraDebugHUDPerfSampleMode::FrameTotal)
 			{
 				SysAvg = Stats->GetGameThreadStats().GetPerFrameAvg();
-				SysMax = Stats->GetGameThreadStats().GetPerFrameAvg();
+				SysMax = Stats->GetGameThreadStats().GetPerFrameMax();
 			}
 			else
 			{
diff --git a/Engine/Plugins/FX/Niagara/Source/Niagara/Private/NiagaraSystemSimulation.cpp b/Engine/Plugins/FX/Niagara/Source/Niagara/Private/NiagaraSystemSimulation.cpp
--- a/Engine/Plugins/FX/Niagara/Source/Niagara/Private/NiagaraSystemSimulation.cpp
+++ b/Engine/Plugins/FX/Niagara/Source/Niagara/Private/NiagaraSystemSimulation.cpp
@@ -429,9 +429,10 @@ struct FNiagaraSystemSimulationSpawnConcurrentTask
 // Task to run FNiagaraSystemInstance::Tick_Concurrent
 struct FNiagaraSystemInstanceTickConcurrentTask
 {
-	FNiagaraSystemInstanceTickConcurrentTask(FNiagaraSystemSimulation* InSystemSimulation, FNiagaraSystemTickBatch& InBatch)
+	FNiagaraSystemInstanceTickConcurrentTask(FNiagaraSystemSimulation* InSystemSimulation, FNiagaraSystemTickBatch& InBatch, UWorld* InWorldContext)
 		: SystemSimulation(InSystemSimulation)
 		, Batch(InBatch)
+		, WorldContext(InWorldContext)
 	{
 	}
 
@@ -441,6 +442,7 @@ struct FNiagaraSystemInstanceTickConcurrentTask
 
 	void DoTask(ENamedThreads::Type CurrentThread, const FGraphEventRef& MyCompletionGraphEvent)
 	{
+		PARTICLE_PERF_STAT_CYCLES_GT(FParticlePerfStatsContext(WorldContext, SystemSimulation->GetSystem()), TickConcurrent);
 #if NIAGARA_SYSTEMSIMULATION_DEBUGGING
 		NiagaraSystemSimulationLocal::DebugDelayInstancesTask();
 #endif
@@ -486,6 +488,7 @@ struct FNiagaraSystemInstanceTickConcurrentTask
 
 	FNiagaraSystemSimulation* SystemSimulation = nullptr;
 	FNiagaraSystemTickBatch Batch;
+	UWorld* WorldContext = nullptr;
 };
 
 
@@ -1101,7 +1104,7 @@ void FNiagaraSystemSimulation::FlushTickBatch(FNiagaraSystemSimulationTickContex
 		if ( Context.IsRunningAsync() )
 		{
 			// Queue instance concurrent task and track information in the instance
-			FGraphEventRef InstanceAsyncGraphEvent = TGraphTask<FNiagaraSystemInstanceTickConcurrentTask>::CreateTask(&Context.BeforeInstancesTickGraphEvents).ConstructAndDispatchWhenReady(this, TickBatch);
+			FGraphEventRef InstanceAsyncGraphEvent = TGraphTask<FNiagaraSystemInstanceTickConcurrentTask>::CreateTask(&Context.BeforeInstancesTickGraphEvents).ConstructAndDispatchWhenReady(this, TickBatch, Context.World);
 
 			for (FNiagaraSystemInstance* Inst : TickBatch)
 			{
