Fixed new sections and categories not being correctly selected in the details panel in the Hierarchy Editor automatically by introducing new delegates
Also auto rename for new categories & sections is now working

#rnx

diff --git a/Engine/Plugins/FX/Niagara/Source/NiagaraEditor/Private/ViewModels/HierarchyEditor/NiagaraHierarchyViewModelBase.cpp b/Engine/Plugins/FX/Niagara/Source/NiagaraEditor/Private/ViewModels/HierarchyEditor/NiagaraHierarchyViewModelBase.cpp
--- a/Engine/Plugins/FX/Niagara/Source/NiagaraEditor/Private/ViewModels/HierarchyEditor/NiagaraHierarchyViewModelBase.cpp
+++ b/Engine/Plugins/FX/Niagara/Source/NiagaraEditor/Private/ViewModels/HierarchyEditor/NiagaraHierarchyViewModelBase.cpp
@@ -367,6 +367,7 @@ void UNiagaraHierarchyViewModelBase::SetActiveSection(TSharedPtr<FNiagaraHierarc
 {
 	ActiveSection = Section;	
 	RefreshHierarchyView(true);
+	OnSectionActivatedDelegate.Broadcast(Section);
 }
 
 TSharedPtr<FNiagaraHierarchySectionViewModel> UNiagaraHierarchyViewModelBase::GetActiveSection() const
@@ -389,6 +390,19 @@ FString UNiagaraHierarchyViewModelBase::OnItemToStringDebug(TSharedPtr<FNiagaraH
 	return ItemBaseViewModel->ToString();	
 }
 
+void FNiagaraHierarchyItemViewModelBase::Tick(float DeltaTime)
+{
+	if(bRenamePending)
+	{
+		RequestRename();
+	}
+}
+
+TStatId FNiagaraHierarchyItemViewModelBase::GetStatId() const
+{
+	RETURN_QUICK_DECLARE_CYCLE_STAT(FNiagaraHierarchyItemViewModelBase, STATGROUP_Tickables);
+}
+
 void FNiagaraHierarchyItemViewModelBase::SyncToData()
 {
 	// this will refresh all underlying data recursively
@@ -500,7 +514,7 @@ void FNiagaraHierarchyItemViewModelBase::AddChild(TSharedPtr<FNiagaraHierarchyIt
 	Children.Add(Item);
 }
 
-UNiagaraHierarchyItemBase* FNiagaraHierarchyItemViewModelBase::AddNewItem(TSubclassOf<UNiagaraHierarchyItemBase> NewItemClass)
+TSharedPtr<FNiagaraHierarchyItemViewModelBase> FNiagaraHierarchyItemViewModelBase::AddNewItem(TSubclassOf<UNiagaraHierarchyItemBase> NewItemClass)
 {
 	FText TransactionText = FText::FormatOrdered(LOCTEXT("Transaction_AddedItem", "Added new {0} to hierarchy"), FText::FromString(NewItemClass->GetName()));
 	FScopedTransaction Transaction(TransactionText);
@@ -511,7 +525,9 @@ UNiagaraHierarchyItemBase* FNiagaraHierarchyItemViewModelBase::AddNewItem(TSubcl
 	GetDataMutable()->GetChildrenMutable().Add(NewItem);
 	SyncToData();
 
-	return NewItem;
+	TSharedPtr<FNiagaraHierarchyItemViewModelBase> ViewModel = FindViewModelForChild(NewItem);
+	ensure(ViewModel.IsValid());
+	return ViewModel;
 }
 
 void FNiagaraHierarchyItemViewModelBase::DuplicateToThis(TSharedPtr<FNiagaraHierarchyItemViewModelBase> ItemToDuplicate, int32 InsertIndex)
@@ -546,6 +562,17 @@ void FNiagaraHierarchyItemViewModelBase::ReparentToThis(TSharedPtr<FNiagaraHiera
 	HierarchyViewModel->OnHierarchyChanged().Broadcast();
 }
 
+TSharedPtr<FNiagaraHierarchyItemViewModelBase> FNiagaraHierarchyItemViewModelBase::FindViewModelForChild(UNiagaraHierarchyItemBase* Child) const
+{
+	int32 Index = FindIndexOfChild(Child);
+	if(Index != INDEX_NONE)
+	{
+		return Children[Index];
+	}
+
+	return nullptr;
+}
+
 int32 FNiagaraHierarchyItemViewModelBase::FindIndexOfChild(TSharedPtr<FNiagaraHierarchyItemViewModelBase> Child) const
 {
 	return Children.Find(Child);
@@ -705,7 +732,6 @@ TSharedPtr<FNiagaraHierarchySectionViewModel> FNiagaraHierarchyRootViewModel::Ad
 	SectionViewModels.Add(NewSectionViewModel);
 	SyncToData();
 	HierarchyViewModel->SetActiveSection(NewSectionViewModel);
-	NewSectionViewModel->RequestRename();
 
 	return NewSectionViewModel;
 }
@@ -1054,29 +1080,35 @@ FReply FNiagaraHierarchyCategoryViewModel::OnDroppedOn(const FDragDropEvent& Dra
 
 void UNiagaraHierarchyViewModelBase::AddCategory() const
 {
-	UNiagaraHierarchyCategory* Category = Cast<UNiagaraHierarchyCategory>(HierarchyViewModelRoot->AddNewItem(UNiagaraHierarchyCategory::StaticClass()));
-	TArray<UNiagaraHierarchyCategory*> SiblingCategories;
-	Category->GetTypedOuter<UNiagaraHierarchyItemBase>()->GetChildrenOfType<UNiagaraHierarchyCategory>(SiblingCategories);
-	
-	TSet<FName> CategoryNames;
-	for(const auto& SiblingCategory : SiblingCategories)
+	TSharedPtr<FNiagaraHierarchyItemViewModelBase> ViewModel = HierarchyViewModelRoot->AddNewItem(UNiagaraHierarchyCategory::StaticClass());
+
+	if(UNiagaraHierarchyCategory* Category = ViewModel->GetDataMutable<UNiagaraHierarchyCategory>())
 	{
-		CategoryNames.Add(SiblingCategory->GetCategoryName());
-	}
+		TArray<UNiagaraHierarchyCategory*> SiblingCategories;
+		Category->GetTypedOuter<UNiagaraHierarchyItemBase>()->GetChildrenOfType<UNiagaraHierarchyCategory>(SiblingCategories);
+		
+		TSet<FName> CategoryNames;
+		for(const auto& SiblingCategory : SiblingCategories)
+		{
+			CategoryNames.Add(SiblingCategory->GetCategoryName());
+		}
 
-	Category->SetCategoryName(FNiagaraUtilities::GetUniqueName(FName("New Category"), CategoryNames));
-	
-	// we only set the section property if the current section isn't set to "All"
-	Category->SetSection(GetActiveSectionData());
-	
-	RefreshHierarchyView();
+		Category->SetCategoryName(FNiagaraUtilities::GetUniqueName(FName("New Category"), CategoryNames));
+		
+		// we only set the section property if the current section isn't set to "All"
+		Category->SetSection(GetActiveSectionData());
+		
+		RefreshHierarchyView();
 
-	OnHierarchyChangedDelegate.Broadcast();
+		OnItemAddedDelegate.Broadcast(ViewModel);
+		OnHierarchyChangedDelegate.Broadcast();
+	}
 }
 
 void UNiagaraHierarchyViewModelBase::AddSection() const
 {
-	HierarchyViewModelRoot->AddNewSection();
+	TSharedPtr<FNiagaraHierarchySectionViewModel> HierarchySectionViewModel = HierarchyViewModelRoot->AddNewSection();
+	OnItemAddedDelegate.Broadcast(HierarchySectionViewModel);
 	OnHierarchyChangedDelegate.Broadcast();
 }
 
diff --git a/Engine/Plugins/FX/Niagara/Source/NiagaraEditor/Private/Widgets/SNiagaraHierarchy.cpp b/Engine/Plugins/FX/Niagara/Source/NiagaraEditor/Private/Widgets/SNiagaraHierarchy.cpp
--- a/Engine/Plugins/FX/Niagara/Source/NiagaraEditor/Private/Widgets/SNiagaraHierarchy.cpp
+++ b/Engine/Plugins/FX/Niagara/Source/NiagaraEditor/Private/Widgets/SNiagaraHierarchy.cpp
@@ -382,6 +382,9 @@ void SNiagaraHierarchy::Construct(const FArguments& InArgs, TObjectPtr<UNiagaraH
 	HierarchyViewModel->OnRefreshHierarchyView().BindSP(this, &SNiagaraHierarchy::RefreshHierarchyView);
 	HierarchyViewModel->OnRefreshSections().BindSP(this, &SNiagaraHierarchy::RefreshSectionsWidget);
 
+	HierarchyViewModel->OnSectionActivated().AddSP(this, &SNiagaraHierarchy::OnSectionActivated);
+	HierarchyViewModel->OnItemAdded().AddSP(this, &SNiagaraHierarchy::OnItemAdded);
+	
 	OnGenerateRowContentWidget = InArgs._OnGenerateRowContentWidget;
 	OnGenerateCustomDetailsPanelNameWidget = InArgs._OnGenerateCustomDetailsPanelNameWidget;
 	
@@ -423,15 +426,19 @@ void SNiagaraHierarchy::Construct(const FArguments& InArgs, TObjectPtr<UNiagaraH
 						CreateHierarchyButtonWidgets()
 					]
 				]
-				+ SVerticalBox::Slot().Expose(SectionsSlot)
+				+ SVerticalBox::Slot()
 				.AutoHeight()
+				[
+					SAssignNew(SectionBox, SWrapBox)
+					.UseAllottedWidth(true)
+				]
 			]
 			+ SVerticalBox::Slot()
 			[
-				SNew(SVerticalBox)
-				+ SVerticalBox::Slot()
-				.AutoHeight()
-				[
+				 SNew(SVerticalBox)
+				 + SVerticalBox::Slot()
+				 .AutoHeight()
+				 [
 					SAssignNew(HierarchyTreeView, STreeView<TSharedPtr<FNiagaraHierarchyItemViewModelBase>>)
 					.TreeItemsSource(&InHierarchyViewModel->GetHierarchyItems())
 					.OnSelectionChanged(this, &SNiagaraHierarchy::OnSelectionChanged, true)
@@ -497,6 +504,9 @@ SNiagaraHierarchy::~SNiagaraHierarchy()
 	HierarchyViewModel->OnRefreshHierarchyView().Unbind();
 	HierarchyViewModel->OnRefreshSections().Unbind();
 
+	HierarchyViewModel->OnSectionActivated().RemoveAll(this);
+	HierarchyViewModel->OnItemAdded().RemoveAll(this);
+	
 	HierarchyViewModel->GetCommands()->UnmapAction(FGenericCommands::Get().Delete);
 	HierarchyViewModel->GetCommands()->UnmapAction(FGenericCommands::Get().Rename);
 }
@@ -545,7 +555,36 @@ void SNiagaraHierarchy::RefreshHierarchyView(bool bFullRefresh) const
 
 void SNiagaraHierarchy::RefreshSectionsWidget()
 {
-	SectionsSlot->AttachWidget(CreateSectionWidgets());
+	SectionBox->ClearChildren();
+	SectionsWidgetMap.Empty();
+	
+	for (TSharedPtr<FNiagaraHierarchySectionViewModel>& Section : HierarchyViewModel->GetHierarchyViewModelRoot()->GetSectionViewModels())
+	{
+		TSharedPtr<SNiagaraHierarchySection> SectionWidget = SNew(SNiagaraHierarchySection, Section, HierarchyViewModel)
+		.OnSectionActivated_Lambda([this](TSharedPtr<FNiagaraHierarchySectionViewModel> SectionViewModel)
+		{
+			HierarchyViewModel->SetActiveSection(SectionViewModel);
+		});
+		SectionsWidgetMap.Add(Section, SectionWidget);
+		
+		SectionBox->AddSlot()
+		[
+			SectionWidget.ToSharedRef()
+		];
+	}
+	
+	TSharedPtr<SNiagaraHierarchySection> DefaultSectionWidget = SNew(SNiagaraHierarchySection, nullptr, HierarchyViewModel)
+	.OnSectionActivated_Lambda([this](TSharedPtr<FNiagaraHierarchySectionViewModel> SectionViewModel)
+	{
+		HierarchyViewModel->SetActiveSection(SectionViewModel);
+	});
+	
+	SectionsWidgetMap.Add(nullptr, DefaultSectionWidget);
+
+	SectionBox->AddSlot()
+	[
+		DefaultSectionWidget.ToSharedRef()
+	];
 }
 
 bool SNiagaraHierarchy::IsItemSelected(TSharedPtr<FNiagaraHierarchyItemViewModelBase> Item) const
@@ -564,16 +603,7 @@ void SNiagaraHierarchy::SelectObjectInDetailsPanel(UObject* Object) const
 TSharedRef<SWidget> SNiagaraHierarchy::CreateHierarchyButtonWidgets() const
 {
 	FSlimHorizontalToolBarBuilder ToolBarBuilder(nullptr, FMultiBoxCustomization::None);
-
-	// Add Category
-	{
-		FUIAction AddCategoryAction;
-		AddCategoryAction.ExecuteAction = FExecuteAction::CreateUObject(HierarchyViewModel.Get(), &UNiagaraHierarchyViewModelBase::AddCategory);
-		ToolBarBuilder.AddToolBarButton(AddCategoryAction, NAME_None,
-			FText::FromString("Add Category")
-			);
-	}
-
+	
 	// Add Section
 	{
 		FUIAction AddSectionAction;
@@ -582,51 +612,17 @@ TSharedRef<SWidget> SNiagaraHierarchy::CreateHierarchyButtonWidgets() const
 			FText::FromString("Add Section")
 			);
 	}
-
-	return ToolBarBuilder.MakeWidget();
-}
-
-TSharedRef<SWidget> SNiagaraHierarchy::CreateSectionWidgets()
-{
-	SectionsWidgetMap.Empty();
 	
-	TSharedRef<SWrapBox> SectionsBox = SNew(SWrapBox)
-		.UseAllottedWidth(true);
-	
-	for (TSharedPtr<FNiagaraHierarchySectionViewModel>& Section : HierarchyViewModel->GetHierarchyViewModelRoot()->GetSectionViewModels())
+	// Add Category
 	{
-		TSharedPtr<SNiagaraHierarchySection> SectionWidget = SNew(SNiagaraHierarchySection, Section, HierarchyViewModel)
-		.OnSectionActivated_Lambda([this](TSharedPtr<FNiagaraHierarchySectionViewModel> SectionViewModel)
-		{
-			HierarchyViewModel->SetActiveSection(SectionViewModel);
-			OnSelectionChanged(SectionViewModel, ESelectInfo::Type::OnMouseClick, true);
-		});
-		SectionsWidgetMap.Add(Section, SectionWidget);
-		
-		SectionsBox->AddSlot()
-		[
-			SectionWidget.ToSharedRef()
-		];
+		FUIAction AddCategoryAction;
+		AddCategoryAction.ExecuteAction = FExecuteAction::CreateUObject(HierarchyViewModel.Get(), &UNiagaraHierarchyViewModelBase::AddCategory);
+		ToolBarBuilder.AddToolBarButton(AddCategoryAction, NAME_None,
+			FText::FromString("Add Category")
+			);
 	}
-	
-	TSharedPtr<SNiagaraHierarchySection> DefaultSectionWidget = SNew(SNiagaraHierarchySection, nullptr, HierarchyViewModel)
-	.OnSectionActivated_Lambda([this](TSharedPtr<FNiagaraHierarchySectionViewModel> SectionViewModel)
-	{
-		HierarchyViewModel->SetActiveSection(SectionViewModel);
-		OnSelectionChanged(SectionViewModel, ESelectInfo::Type::OnMouseClick, true);
-	});
-	
-	SectionsWidgetMap.Add(nullptr, DefaultSectionWidget);
 
-	SectionsBox->AddSlot()
-	[
-		DefaultSectionWidget.ToSharedRef()
-	];
-	
-	return SNew(SBorder).BorderBackgroundColor(FStyleColors::Recessed)
-	[
-		SectionsBox
-	];
+	return ToolBarBuilder.MakeWidget();
 }
 
 TSharedRef<ITableRow> SNiagaraHierarchy::GenerateSourceItemRow(TSharedPtr<FNiagaraHierarchyItemViewModelBase> HierarchyItem, const TSharedRef<STableViewBase>& TableViewBase)
@@ -748,6 +744,28 @@ bool SNiagaraHierarchy::CanDeleteSelectedItems() const
 	return false;
 }
 
+void SNiagaraHierarchy::OnItemAdded(TSharedPtr<FNiagaraHierarchyItemViewModelBase> AddedItem)
+{
+	// when a new item is created (opposed to dragged & dropped from source view, i.e. only categories so far)
+	// we make sure to request a tree refresh, select the row, and request a pending rename since the widget will created a frame later
+	if(AddedItem->GetData()->IsA<UNiagaraHierarchyItem>() || AddedItem->GetData()->IsA<UNiagaraHierarchyCategory>())
+	{
+		HierarchyTreeView->RequestTreeRefresh();
+		HierarchyTreeView->SetSelection(AddedItem);
+	}
+	else if(AddedItem->GetData()->IsA<UNiagaraHierarchySection>())
+	{
+		RefreshSectionsWidget();
+	}
+
+	AddedItem->RequestRenamePending();
+}
+
+void SNiagaraHierarchy::OnSectionActivated(TSharedPtr<FNiagaraHierarchySectionViewModel> Section) const
+{
+	OnSelectionChanged(Section, ESelectInfo::Direct, true);
+}
+
 void SNiagaraHierarchy::OnSelectionChanged(TSharedPtr<FNiagaraHierarchyItemViewModelBase> HierarchyItem, ESelectInfo::Type Type, bool bFromHierarchy) const
 {
 	if(DetailsPanel.IsValid())
diff --git a/Engine/Plugins/FX/Niagara/Source/NiagaraEditor/Public/ViewModels/HierarchyEditor/NiagaraHierarchyViewModelBase.h b/Engine/Plugins/FX/Niagara/Source/NiagaraEditor/Public/ViewModels/HierarchyEditor/NiagaraHierarchyViewModelBase.h
--- a/Engine/Plugins/FX/Niagara/Source/NiagaraEditor/Public/ViewModels/HierarchyEditor/NiagaraHierarchyViewModelBase.h
+++ b/Engine/Plugins/FX/Niagara/Source/NiagaraEditor/Public/ViewModels/HierarchyEditor/NiagaraHierarchyViewModelBase.h
@@ -204,6 +204,7 @@ private:
 };
 
 struct FNiagaraHierarchyItemViewModelBase;
+struct FNiagaraHierarchySectionViewModel;
 
 class FNiagaraHierarchyDragDropOp : public FDragDropOperation
 {
@@ -244,7 +245,9 @@ public:
 	DECLARE_DELEGATE_OneParam(FRefreshTreeWidget, bool bFullRefresh)
 	DECLARE_DELEGATE_RetVal_OneParam(bool, FIsItemSelected, TSharedPtr<FNiagaraHierarchyItemViewModelBase> Item)
 	DECLARE_DELEGATE_OneParam(FSelectObjectInDetailsPanel, UObject* Object)
-	
+	DECLARE_MULTICAST_DELEGATE_OneParam(FOnSectionActivated, TSharedPtr<FNiagaraHierarchySectionViewModel> Section)
+	DECLARE_MULTICAST_DELEGATE_OneParam(FOnItemAdded, TSharedPtr<FNiagaraHierarchyItemViewModelBase> AddedItem)
+
 	GENERATED_BODY()
 	
 	UNiagaraHierarchyViewModelBase();
@@ -292,7 +295,6 @@ public:
 	void RefreshSourceView(bool bFullRefresh = false) const;
 	void RefreshHierarchyView(bool bFullRefresh = false) const;
 	void RefreshSectionsWidget() const;
-	void SelectItemInDetailsPanel(UObject* Object) const;
 	
 	// Delegate that call functions from SNiagaraHierarchy
 	FRefreshTreeWidget& OnRefreshSourceView() { return RefreshSourceViewDelegate; }
@@ -301,12 +303,15 @@ public:
 
 	// Delegates for external systems
 	FOnHierarchyChanged& OnHierarchyChanged() { return OnHierarchyChangedDelegate; } 
+
+	FOnItemAdded& OnItemAdded() { return OnItemAddedDelegate; }
 	
 	// Sections
 	void SetActiveSection(TSharedPtr<struct FNiagaraHierarchySectionViewModel>);
 	TSharedPtr<FNiagaraHierarchySectionViewModel> GetActiveSection() const;
 	UNiagaraHierarchySection* GetActiveSectionData() const;
 	bool IsSectionActive(const UNiagaraHierarchySection* Section) const;
+	FOnSectionActivated& OnSectionActivated() { return OnSectionActivatedDelegate; }
 	
 	FString OnItemToStringDebug(TSharedPtr<FNiagaraHierarchyItemViewModelBase> ItemBaseViewModel) const; 
 protected:
@@ -334,13 +339,15 @@ protected:
 	FRefreshTreeWidget RefreshHierarchyWidgetDelegate;
 	FSimpleDelegate RefreshSectionsWidgetDelegate;
 
+	FOnItemAdded OnItemAddedDelegate;
+	FOnSectionActivated OnSectionActivatedDelegate;
 	FOnHierarchyChanged OnHierarchyChangedDelegate;
 };
 
-struct NIAGARAEDITOR_API FNiagaraHierarchyItemViewModelBase : TSharedFromThis<FNiagaraHierarchyItemViewModelBase>
+struct NIAGARAEDITOR_API FNiagaraHierarchyItemViewModelBase : TSharedFromThis<FNiagaraHierarchyItemViewModelBase>, public FTickableEditorObject
 {
 	DECLARE_MULTICAST_DELEGATE(FOnSynced)
-	
+
 	FNiagaraHierarchyItemViewModelBase(UNiagaraHierarchyItemBase* InItemBase, TSharedPtr<FNiagaraHierarchyItemViewModelBase> InParent, TWeakObjectPtr<UNiagaraHierarchyViewModelBase> InHierarchyViewModel)
 		: ItemBase(InItemBase)
 		, Parent(InParent)
@@ -360,6 +367,9 @@ struct NIAGARAEDITOR_API FNiagaraHierarchyItemViewModelBase : TSharedFromThis<FN
 	template<class T>
 	const T* GetData() const { return Cast<T>(ItemBase); }
 
+	virtual void Tick(float DeltaTime) override;
+	virtual TStatId GetStatId() const override;
+	
 	FString ToString() const { return ItemBase->ToString(); }
 	
 	void SyncToData();
@@ -372,9 +382,11 @@ struct NIAGARAEDITOR_API FNiagaraHierarchyItemViewModelBase : TSharedFromThis<FN
 	bool HasParent(TSharedPtr<FNiagaraHierarchyItemViewModelBase> ParentCandidate, bool bRecursive = false);
 
 	void AddChild(TSharedPtr<FNiagaraHierarchyItemViewModelBase> Item);
-	UNiagaraHierarchyItemBase* AddNewItem(TSubclassOf<UNiagaraHierarchyItemBase> NewItemClass);
+	TSharedPtr<FNiagaraHierarchyItemViewModelBase> AddNewItem(TSubclassOf<UNiagaraHierarchyItemBase> NewItemClass);
 	void DuplicateToThis(TSharedPtr<FNiagaraHierarchyItemViewModelBase> ItemToDuplicate, int32 InsertIndex = INDEX_NONE);
 	void ReparentToThis(TSharedPtr<FNiagaraHierarchyItemViewModelBase> ItemToMove, int32 InsertIndex = INDEX_NONE);
+
+	TSharedPtr<FNiagaraHierarchyItemViewModelBase> FindViewModelForChild(UNiagaraHierarchyItemBase* Child) const;
 	int32 FindIndexOfChild(TSharedPtr<FNiagaraHierarchyItemViewModelBase> Child) const;
 	int32 FindIndexOfChild(UNiagaraHierarchyItemBase* Child) const;
 	int32 FindIndexOfDataChild(TSharedPtr<FNiagaraHierarchyItemViewModelBase> Child) const;
@@ -405,12 +417,21 @@ struct NIAGARAEDITOR_API FNiagaraHierarchyItemViewModelBase : TSharedFromThis<FN
 
 	void RequestRename()
 	{
-		if(CanRename())
+		if(CanRename() && OnRequestRenameDelegate.IsBound())
 		{
+			bRenamePending = false;
 			OnRequestRenameDelegate.Execute();
 		}
 	}
 
+	void RequestRenamePending()
+	{
+		if(CanRename())
+		{
+			bRenamePending = true;
+		}
+	}
+
 	/** The UObject we display in the details panel when this item is selected. By default it's the item the view model represents. */
 	virtual UObject* GetDataForEditing() { return ItemBase; }
 	
@@ -436,6 +457,7 @@ protected:
 	TWeakObjectPtr<UNiagaraHierarchyViewModelBase> HierarchyViewModel;
 	FSimpleDelegate OnRequestRenameDelegate;
 	FOnSynced OnSyncedDelegate;
+	bool bRenamePending = false;
 };
 
 template <class DataClass, class ViewModelClass>
diff --git a/Engine/Plugins/FX/Niagara/Source/NiagaraEditor/Public/Widgets/SNiagaraHierarchy.h b/Engine/Plugins/FX/Niagara/Source/NiagaraEditor/Public/Widgets/SNiagaraHierarchy.h
--- a/Engine/Plugins/FX/Niagara/Source/NiagaraEditor/Public/Widgets/SNiagaraHierarchy.h
+++ b/Engine/Plugins/FX/Niagara/Source/NiagaraEditor/Public/Widgets/SNiagaraHierarchy.h
@@ -41,7 +41,6 @@ public:
 	
 private:
 	TSharedRef<SWidget> CreateHierarchyButtonWidgets() const;
-	TSharedRef<SWidget> CreateSectionWidgets();
 	TSharedRef<ITableRow> GenerateSourceItemRow(TSharedPtr<FNiagaraHierarchyItemViewModelBase> HierarchyItem, const TSharedRef<STableViewBase>& TableViewBase);
 	TSharedRef<ITableRow> GenerateHierarchyItemRow(TSharedPtr<FNiagaraHierarchyItemViewModelBase> HierarchyItem, const TSharedRef<STableViewBase>& TableViewBase);
 private:	
@@ -51,6 +50,8 @@ private:
 	void DeleteSelectedItems();
 	bool CanDeleteSelectedItems() const;
 
+	void OnItemAdded(TSharedPtr<FNiagaraHierarchyItemViewModelBase> AddedItem);
+	void OnSectionActivated(TSharedPtr<FNiagaraHierarchySectionViewModel> Section) const;
 	void OnSelectionChanged(TSharedPtr<FNiagaraHierarchyItemViewModelBase> HierarchyItem, ESelectInfo::Type Type, bool bFromHierarchy) const;
 
 	virtual FReply OnKeyDown(const FGeometry& MyGeometry, const FKeyEvent& InKeyEvent) override;
@@ -65,7 +66,7 @@ private:
 	TSharedPtr<STreeView<TSharedPtr<FNiagaraHierarchyItemViewModelBase>>> SourceTreeView;
 	TSharedPtr<STreeView<TSharedPtr<FNiagaraHierarchyItemViewModelBase>>> HierarchyTreeView;
 	TMap<TSharedPtr<FNiagaraHierarchySectionViewModel>, TSharedPtr<class SNiagaraHierarchySection>> SectionsWidgetMap;
-	SVerticalBox::FSlot* SectionsSlot = nullptr;
+	TSharedPtr<class SWrapBox> SectionBox;
 	TSharedPtr<IDetailsView> DetailsPanel;
 
 private:
