Niagara - When creating instance components from blueprint templates, reapply template and instance overrides to the parameter store to keep the override parameters consistent durring cooking.  This fixes a very long standing issue where blueprint instances with niagara components would lose their instance user parameter values in a cooked build.

diff --git a/Engine/Plugins/FX/Niagara/Source/Niagara/Private/NiagaraComponent.cpp b/Engine/Plugins/FX/Niagara/Source/Niagara/Private/NiagaraComponent.cpp
--- a/Engine/Plugins/FX/Niagara/Source/Niagara/Private/NiagaraComponent.cpp
+++ b/Engine/Plugins/FX/Niagara/Source/Niagara/Private/NiagaraComponent.cpp
@@ -1730,15 +1730,17 @@ void UNiagaraComponent::OnComponentCreated()
 #endif
 }
 
-// for handling cases where niagara components are stored in blueprints we can get into situations where
-// the instancing code will get confused with objects stored in the different parameter override variables
-// (TemplateParameterOverrides, InstanceParameterOverrides, OverrideParameters), so we make sure that
-// the outers are properly fixed up
+// When creating niagara components from blueprint templates the constructions scripts are rerun which returns the instance properties to the template
+// defaults.  The instance properties are then reapplied to that template, but in the case of Niagara components, the OverrideParameters are not saved
+// in the instance properties so the template and instance overrides need to be reapplied here.  In addition to this, when instancing from blueprint 
+// we can get into situations where the instancing code will get confused with objects stored in the different parameter override variables
+// (TemplateParameterOverrides, InstanceParameterOverrides, OverrideParameters), so we make sure that the outers are properly fixed up here as well.
 void UNiagaraComponent::PostApplyToComponent()
 {
 	Super::PostApplyToComponent();
 
 #if WITH_EDITOR
+	ApplyOverridesToParameterStore();
 	FixDataInterfaceOuters();
 #endif
 }
