Fix for static analysis warning with NiagaraParameterMapHistory.cpp
#rnx

diff --git a/Engine/Plugins/FX/Niagara/Source/NiagaraEditor/Private/NiagaraParameterMapHistory.cpp b/Engine/Plugins/FX/Niagara/Source/NiagaraEditor/Private/NiagaraParameterMapHistory.cpp
--- a/Engine/Plugins/FX/Niagara/Source/NiagaraEditor/Private/NiagaraParameterMapHistory.cpp
+++ b/Engine/Plugins/FX/Niagara/Source/NiagaraEditor/Private/NiagaraParameterMapHistory.cpp
@@ -453,22 +453,23 @@ FNiagaraVariable FNiagaraParameterMapHistory::VariableToNamespacedVariable(const
 
 bool FNiagaraParameterMapHistory::IsInNamespace(const FNiagaraVariableBase& InVar, const FString& Namespace)
 {
-	if (Namespace.EndsWith(TEXT(".")))
-	{
-		bool FastTest = InVar.IsInNameSpace(Namespace);
-		// Leaving old code commented out for now just as a reference to what was before.
-		//bool SlowTest =  InVar.GetName().ToString().StartsWith(Namespace);
-		//check(FastTest == SlowTest);
-		return FastTest;
-	}
-	else
-	{
-		bool FastTest = InVar.IsInNameSpace(Namespace); // This should check to see if there is a period in the name of the variable after the namespace string.
-		// Leaving old code commented out for now just as a reference to what was before.
-		//bool SlowTest = InVar.GetName().ToString().StartsWith(Namespace + TEXT("."));
-		//check(FastTest == SlowTest);
-		return FastTest;
-	}
+	const bool FastTest = InVar.IsInNameSpace(Namespace);
+
+	// Leaving old code commented out for now just as a reference to what was before.
+	//if (Namespace.EndsWith(TEXT(".")))
+	//{
+	//	bool SlowTest =  InVar.GetName().ToString().StartsWith(Namespace);
+	//	check(FastTest == SlowTest);
+	//	return FastTest;
+	//}
+	//else
+	//{
+	//	bool SlowTest = InVar.GetName().ToString().StartsWith(Namespace + TEXT("."));
+	//	check(FastTest == SlowTest);
+	//	return FastTest;
+	//}
+
+	return FastTest;
 }
 
 bool FNiagaraParameterMapHistory::IsAliasedModuleParameter(const FNiagaraVariable& InVar)
