Make sure RendererEnabled binding gts updated with changes to the emitter (like renaming)

diff --git a/Engine/Plugins/FX/Niagara/Source/Niagara/Private/NiagaraRendererProperties.cpp b/Engine/Plugins/FX/Niagara/Source/Niagara/Private/NiagaraRendererProperties.cpp
--- a/Engine/Plugins/FX/Niagara/Source/Niagara/Private/NiagaraRendererProperties.cpp
+++ b/Engine/Plugins/FX/Niagara/Source/Niagara/Private/NiagaraRendererProperties.cpp
@@ -556,6 +556,8 @@ void UNiagaraRendererProperties::UpdateSourceModeDerivates(ENiagaraRendererSourc
 			((FNiagaraVariableAttributeBinding*)Binding)->CacheValues(SrcEmitter, InSourceMode);
 		}
 
+		RendererEnabledBinding.CacheValues(SrcEmitter, InSourceMode);
+
 #if WITH_EDITORONLY_DATA
 		// If we added or removed any valid bindings to a non-particle source during editing, we need to reset to prevent hazards and
 		// to ensure new ones get bound by the simulation
