Fix for nullptr access when closing Niagara editor
#rnx

diff --git a/Engine/Plugins/FX/Niagara/Source/NiagaraEditor/Private/Customizations/NiagaraComponentDetails.h b/Engine/Plugins/FX/Niagara/Source/NiagaraEditor/Private/Customizations/NiagaraComponentDetails.h
--- a/Engine/Plugins/FX/Niagara/Source/NiagaraEditor/Private/Customizations/NiagaraComponentDetails.h
+++ b/Engine/Plugins/FX/Niagara/Source/NiagaraEditor/Private/Customizations/NiagaraComponentDetails.h
@@ -166,7 +166,10 @@ public:
 			
 			if(SystemViewModel.IsValid())
 			{
-				SystemViewModel.Pin()->GetUserParametersHierarchyViewModel()->OnHierarchyChanged().RemoveAll(this);
+				if (UNiagaraUserParametersHierarchyViewModel* HierarchyViewModel = SystemViewModel.Pin()->GetUserParametersHierarchyViewModel())
+				{
+					HierarchyViewModel->OnHierarchyChanged().RemoveAll(this);
+				}
 			}
 		}
 	}
