PR #9662: Correct spelling of "receive" in log messages. (Contributed by [REDACTED_PERSON])

diff --git a/Engine/Plugins/FX/Niagara/Source/Niagara/Private/NiagaraDebuggerClient.cpp b/Engine/Plugins/FX/Niagara/Source/Niagara/Private/NiagaraDebuggerClient.cpp
--- a/Engine/Plugins/FX/Niagara/Source/Niagara/Private/NiagaraDebuggerClient.cpp
+++ b/Engine/Plugins/FX/Niagara/Source/Niagara/Private/NiagaraDebuggerClient.cpp
@@ -136,12 +136,12 @@ void FNiagaraDebuggerClient::HandleConnectionRequestMessage(const FNiagaraDebugg
 	{
 		if (Connection.IsValid())
 		{
-			UE_LOG(LogNiagaraDebuggerClient, Warning, TEXT("Connection request recieved but we already have a connected debugger. Current connection being dropped and new connection accepted. | Session: %s | Instance: %s (%s)."), *SessionId.ToString(), *InstanceId.ToString(), *InstanceName);
+			UE_LOG(LogNiagaraDebuggerClient, Warning, TEXT("Connection request received but we already have a connected debugger. Current connection being dropped and new connection accepted. | Session: %s | Instance: %s (%s)."), *SessionId.ToString(), *InstanceId.ToString(), *InstanceName);
 			CloseConnection();
 		}
 		else
 		{
-			UE_LOG(LogNiagaraDebuggerClient, Log, TEXT("Connection request recieved and accepted. | Session: %s | Instance: %s (%s)."), *SessionId.ToString(), *InstanceId.ToString(), *InstanceName);
+			UE_LOG(LogNiagaraDebuggerClient, Log, TEXT("Connection request received and accepted. | Session: %s | Instance: %s (%s)."), *SessionId.ToString(), *InstanceId.ToString(), *InstanceName);
 		}
 		
 		//Accept the connection and inform the debugger we have done so with an accepted message.
@@ -164,7 +164,7 @@ void FNiagaraDebuggerClient::HandleConnectionClosedMessage(const FNiagaraDebugge
 		}
 		else
 		{
-			UE_LOG(LogNiagaraDebuggerClient, Warning, TEXT("Recieved connection closed message for unconnected debugger. | Session: %s | Instance: %s (%s)."), *SessionId.ToString(), *InstanceId.ToString(), *InstanceName);
+			UE_LOG(LogNiagaraDebuggerClient, Warning, TEXT("Received connection closed message for unconnected debugger. | Session: %s | Instance: %s (%s)."), *SessionId.ToString(), *InstanceId.ToString(), *InstanceName);
 		}
 	}
 }
@@ -236,24 +236,24 @@ void FNiagaraDebuggerClient::HandleOutlinerSettingsMessage(const FNiagaraOutline
 #endif
 				if (Message.CaptureDelayFrames <= 0)
 				{
-					UE_LOG(LogNiagaraDebuggerClient, Log, TEXT("Recieved request to capture outliner data. Capturing now. | Session: %s | Instance: %s (%s)."), *SessionId.ToString(), *InstanceId.ToString(), *InstanceName);
+					UE_LOG(LogNiagaraDebuggerClient, Log, TEXT("Received request to capture outliner data. Capturing now. | Session: %s | Instance: %s (%s)."), *SessionId.ToString(), *InstanceId.ToString(), *InstanceName);
 					UpdateOutliner(0.001f);
 				}
 				else
 				{
 					OutlinerCountdown = Message.CaptureDelayFrames;
-					UE_LOG(LogNiagaraDebuggerClient, Log, TEXT("Recieved request to capture outliner data. Capturing in %u frames. | Session: %s | Instance: %s (%s)."), Message.CaptureDelayFrames, *SessionId.ToString(), *InstanceId.ToString(), *InstanceName);
+					UE_LOG(LogNiagaraDebuggerClient, Log, TEXT("Received request to capture outliner data. Capturing in %u frames. | Session: %s | Instance: %s (%s)."), Message.CaptureDelayFrames, *SessionId.ToString(), *InstanceId.ToString(), *InstanceName);
 					FTSTicker::GetCoreTicker().AddTicker(FTickerDelegate::CreateRaw(this, &FNiagaraDebuggerClient::UpdateOutliner));
 				}
 			}
 			else
 			{
-				UE_LOG(LogNiagaraDebuggerClient, Log, TEXT("Recieved request to capture outliner data. Ignoring as we already have a pending outliner capture. | Session: %s | Instance: %s (%s)."), *SessionId.ToString(), *InstanceId.ToString(), *InstanceName);
+				UE_LOG(LogNiagaraDebuggerClient, Log, TEXT("Received request to capture outliner data. Ignoring as we already have a pending outliner capture. | Session: %s | Instance: %s (%s)."), *SessionId.ToString(), *InstanceId.ToString(), *InstanceName);
 			}
 		}
 		else
 		{
-			UE_LOG(LogNiagaraDebuggerClient, Log, TEXT("Recieved request to capture outliner data but the capture bool is false. | Session: %s | Instance: %s."), *SessionId.ToString(), *InstanceId.ToString(), *InstanceName);
+			UE_LOG(LogNiagaraDebuggerClient, Log, TEXT("Received request to capture outliner data but the capture bool is false. | Session: %s | Instance: %s."), *SessionId.ToString(), *InstanceId.ToString(), *InstanceName);
 		}
 	}
 }
diff --git a/Engine/Plugins/FX/Niagara/Source/NiagaraEditor/Private/NiagaraDebugger.cpp b/Engine/Plugins/FX/Niagara/Source/NiagaraEditor/Private/NiagaraDebugger.cpp
--- a/Engine/Plugins/FX/Niagara/Source/NiagaraEditor/Private/NiagaraDebugger.cpp
+++ b/Engine/Plugins/FX/Niagara/Source/NiagaraEditor/Private/NiagaraDebugger.cpp
@@ -279,7 +279,7 @@ void FNiagaraDebugger::HandleConnectionAcceptedMessage(const FNiagaraDebuggerAcc
 	int32 FoundActive = FindActiveConnection(Message.SessionId, Message.InstanceId);
 	if (FoundActive != INDEX_NONE)
 	{
-		UE_LOG(LogNiagaraDebugger, Log, TEXT("Recieved connection accepted message from an already connected client. Ignored. | Session: %s | Instance: %s |"), *Message.SessionId.ToString(), *Message.InstanceId.ToString());
+		UE_LOG(LogNiagaraDebugger, Log, TEXT("Received connection accepted message from an already connected client. Ignored. | Session: %s | Instance: %s |"), *Message.SessionId.ToString(), *Message.InstanceId.ToString());
 		return;
 	}
 
@@ -301,7 +301,7 @@ void FNiagaraDebugger::HandleConnectionAcceptedMessage(const FNiagaraDebuggerAcc
 	}
 	else
 	{
-		UE_LOG(LogNiagaraDebugger, Log, TEXT("Recieved connection accepted message from a client that is not in our pending list. Ignored. | Session: %s | Instance: %s |"), *Message.SessionId.ToString(), *Message.InstanceId.ToString());
+		UE_LOG(LogNiagaraDebugger, Log, TEXT("Received connection accepted message from a client that is not in our pending list. Ignored. | Session: %s | Instance: %s |"), *Message.SessionId.ToString(), *Message.InstanceId.ToString());
 	}
 }
 
@@ -328,7 +328,7 @@ void FNiagaraDebugger::HandleConnectionClosedMessage(const FNiagaraDebuggerConne
 
 void FNiagaraDebugger::HandleOutlinerUpdateMessage(const FNiagaraDebuggerOutlinerUpdate& Message, const TSharedRef<IMessageContext, ESPMode::ThreadSafe>& Context)
 {
-	UE_LOG(LogNiagaraDebugger, Log, TEXT("Recieved outliner update. | Sender: %s "), *Context->GetSender().ToString());
+	UE_LOG(LogNiagaraDebugger, Log, TEXT("Received outliner update. | Sender: %s "), *Context->GetSender().ToString());
 	if(UNiagaraOutliner* Outliner = GetOutliner())
 	{
 		Outliner->UpdateData(Message.OutlinerData);
@@ -337,14 +337,14 @@ void FNiagaraDebugger::HandleOutlinerUpdateMessage(const FNiagaraDebuggerOutline
 
 void FNiagaraDebugger::UpdateSimpleClientInfo(const FNiagaraSimpleClientInfo& Message, const TSharedRef<IMessageContext, ESPMode::ThreadSafe>& Context)
 {
-	UE_LOG(LogNiagaraDebugger, Log, TEXT("Recieved simple client info update. | Sender: %s "), *Context->GetSender().ToString());
+	UE_LOG(LogNiagaraDebugger, Log, TEXT("Received simple client info update. | Sender: %s "), *Context->GetSender().ToString());
 	SimpleClientInfo = Message;
 	OnSimpleClientInfoChangedDelegate.Broadcast(SimpleClientInfo);
 }
 
 void FNiagaraDebugger::HandleSimCacheCaptureReply(const FNiagaraSystemSimCacheCaptureReply& Message, const TSharedRef<IMessageContext, ESPMode::ThreadSafe>& Context)
 {
-	UE_LOG(LogNiagaraDebugger, Log, TEXT("Recieved Niagara Sim Cache Capture Reply. | Sender: %s "), *Context->GetSender().ToString());
+	UE_LOG(LogNiagaraDebugger, Log, TEXT("Received Niagara Sim Cache Capture Reply. | Sender: %s "), *Context->GetSender().ToString());
 	
 	if (UNiagaraOutliner* Outliner = GetOutliner())
 	{
diff --git a/Engine/Plugins/FX/Niagara/Source/NiagaraEditor/Private/NiagaraOutliner.cpp b/Engine/Plugins/FX/Niagara/Source/NiagaraEditor/Private/NiagaraOutliner.cpp
--- a/Engine/Plugins/FX/Niagara/Source/NiagaraEditor/Private/NiagaraOutliner.cpp
+++ b/Engine/Plugins/FX/Niagara/Source/NiagaraEditor/Private/NiagaraOutliner.cpp
@@ -59,7 +59,7 @@ void UNiagaraOutliner::UpdateSystemSimCache(const FNiagaraSystemSimCacheCaptureR
 	}
 	else
 	{
-		UE_LOG(LogNiagaraEditor, Warning, TEXT("Niagara Outliner recieved empty sim cache data."));
+		UE_LOG(LogNiagaraEditor, Warning, TEXT("Niagara Outliner received empty sim cache data."));
 	}
 	OnChanged();
 }
diff --git a/Engine/Plugins/FX/Niagara/Source/NiagaraEditor/Private/ViewModels/NiagaraSimCacheViewModel.cpp b/Engine/Plugins/FX/Niagara/Source/NiagaraEditor/Private/ViewModels/NiagaraSimCacheViewModel.cpp
--- a/Engine/Plugins/FX/Niagara/Source/NiagaraEditor/Private/ViewModels/NiagaraSimCacheViewModel.cpp
+++ b/Engine/Plugins/FX/Niagara/Source/NiagaraEditor/Private/ViewModels/NiagaraSimCacheViewModel.cpp
@@ -47,7 +47,7 @@ void FNiagaraSimCacheViewModel::UpdateSimCache(const FNiagaraSystemSimCacheCaptu
 	}
 	else
 	{
-		UE_LOG(LogNiagaraEditor, Warning, TEXT("Debug Spreadsheet recieved empty sim cache data."));
+		UE_LOG(LogNiagaraEditor, Warning, TEXT("Debug Spreadsheet received empty sim cache data."));
 	}
 	Initialize(SimCache);
 }
diff --git a/Engine/Plugins/Media/HardwareEncoders/Source/EncoderAMF/Private/Amf_EncoderH264.cpp b/Engine/Plugins/Media/HardwareEncoders/Source/EncoderAMF/Private/Amf_EncoderH264.cpp
--- a/Engine/Plugins/Media/HardwareEncoders/Source/EncoderAMF/Private/Amf_EncoderH264.cpp
+++ b/Engine/Plugins/Media/HardwareEncoders/Source/EncoderAMF/Private/Amf_EncoderH264.cpp
@@ -732,7 +732,7 @@ namespace AVEncoder
 		}
 		else
 		{
-			UE_LOG(LogEncoderAMF, Error, TEXT("Amf recieved nullptr to texture."));
+			UE_LOG(LogEncoderAMF, Error, TEXT("Amf received nullptr to texture."));
 			return false;
 		}
 
diff --git a/Engine/Plugins/MovieScene/MovieRenderPipeline/Source/MovieRenderPipelineCore/Private/MoviePipelineOutputBuilder.cpp b/Engine/Plugins/MovieScene/MovieRenderPipeline/Source/MovieRenderPipelineCore/Private/MoviePipelineOutputBuilder.cpp
--- a/Engine/Plugins/MovieScene/MovieRenderPipeline/Source/MovieRenderPipelineCore/Private/MoviePipelineOutputBuilder.cpp
+++ b/Engine/Plugins/MovieScene/MovieRenderPipeline/Source/MovieRenderPipelineCore/Private/MoviePipelineOutputBuilder.cpp
@@ -70,13 +70,13 @@ void FMoviePipelineOutputMerger::OnCompleteRenderPassDataAvailable_AnyThread(TUn
 	 		break;
 	 	}
 	}
-	if (!ensureAlwaysMsgf(OutputFrame, TEXT("Recieved data for unknown frame. Frame was either already processed or not queued yet!")))
+	if (!ensureAlwaysMsgf(OutputFrame, TEXT("Received data for unknown frame. Frame was either already processed or not queued yet!")))
 	{
 		return;
 	}
 
 	// Ensure this pass is expected as well...
-	if (!ensureAlwaysMsgf(OutputFrame->ExpectedRenderPasses.Contains(Payload->PassIdentifier), TEXT("Recieved data for unexpected render pass: %s"), *Payload->PassIdentifier.Name))
+	if (!ensureAlwaysMsgf(OutputFrame->ExpectedRenderPasses.Contains(Payload->PassIdentifier), TEXT("Received data for unexpected render pass: %s"), *Payload->PassIdentifier.Name))
 	{
 		return;
 	}
diff --git a/Engine/Plugins/Online/OnlineSubsystemUtils/Source/OnlineSubsystemUtils/Private/Tests/TestPresenceInterface.cpp b/Engine/Plugins/Online/OnlineSubsystemUtils/Source/OnlineSubsystemUtils/Private/Tests/TestPresenceInterface.cpp
--- a/Engine/Plugins/Online/OnlineSubsystemUtils/Source/OnlineSubsystemUtils/Private/Tests/TestPresenceInterface.cpp
+++ b/Engine/Plugins/Online/OnlineSubsystemUtils/Source/OnlineSubsystemUtils/Private/Tests/TestPresenceInterface.cpp
@@ -441,7 +441,7 @@ void FTestPresenceInterface::OnRandomUserFetchComplete(const FUniqueNetId& UserI
 
 void FTestPresenceInterface::OnPresenceRecieved(const FUniqueNetId& UserId, const TSharedRef<FOnlineUserPresence>& PresenceData)
 {
-	UE_LOG_ONLINE_PRESENCE(Verbose, TEXT("Recieved a presence update event about user %s"), *UserId.ToString());
+	UE_LOG_ONLINE_PRESENCE(Verbose, TEXT("Received a presence update event about user %s"), *UserId.ToString());
 	// We expect to get a status update about someone on your friends list at least.
 	CompletedTasks |= EPresenceTestStatus::WorkingDelegate;
 	TasksAttempted |= EPresenceTestStatus::WorkingDelegate;
diff --git a/Engine/Plugins/Runtime/Steam/SteamSockets/Source/SteamSockets/Private/SteamSocketsNetDriver.cpp b/Engine/Plugins/Runtime/Steam/SteamSockets/Source/SteamSockets/Private/SteamSocketsNetDriver.cpp
--- a/Engine/Plugins/Runtime/Steam/SteamSockets/Source/SteamSockets/Private/SteamSocketsNetDriver.cpp
+++ b/Engine/Plugins/Runtime/Steam/SteamSockets/Source/SteamSockets/Private/SteamSocketsNetDriver.cpp
@@ -307,7 +307,7 @@ void USteamSocketsNetDriver::TickDispatch(float DeltaTime)
 				// Process the message for this connection.
 				if (ConnectionToHandleMessage != nullptr)
 				{
-					UE_LOG(LogNet, VeryVerbose, TEXT("SteamSockets: Recieved packet from %s with size %d"), *MessageSender.ToString(true), Message->GetSize());
+					UE_LOG(LogNet, VeryVerbose, TEXT("SteamSockets: Received packet from %s with size %d"), *MessageSender.ToString(true), Message->GetSize());
 					ConnectionToHandleMessage->HandleRecvMessage(Message->m_pData, Message->GetSize(), &MessageSender);
 					Message->Release();
 					continue;
diff --git a/Engine/Plugins/VirtualProduction/DMX/DMXProtocol/Source/DMXProtocolArtNet/Private/DMXProtocolArtNetReceiver.cpp b/Engine/Plugins/VirtualProduction/DMX/DMXProtocol/Source/DMXProtocolArtNet/Private/DMXProtocolArtNetReceiver.cpp
--- a/Engine/Plugins/VirtualProduction/DMX/DMXProtocol/Source/DMXProtocolArtNet/Private/DMXProtocolArtNetReceiver.cpp
+++ b/Engine/Plugins/VirtualProduction/DMX/DMXProtocol/Source/DMXProtocolArtNet/Private/DMXProtocolArtNetReceiver.cpp
@@ -16,7 +16,7 @@
 #include "HAL/RunnableThread.h"
 #include "Serialization/ArrayReader.h"
 
-DECLARE_DWORD_ACCUMULATOR_STAT(TEXT("Art-Net Packages Recieved Total"), STAT_ArtNetPackagesReceived, STATGROUP_DMX);
+DECLARE_DWORD_ACCUMULATOR_STAT(TEXT("Art-Net Packages Received Total"), STAT_ArtNetPackagesReceived, STATGROUP_DMX);
 
 
 FDMXProtocolArtNetReceiver::FDMXProtocolArtNetReceiver(const TSharedPtr<FDMXProtocolArtNet, ESPMode::ThreadSafe>& InArtNetProtocol, FSocket& InSocket, TSharedRef<FInternetAddr> InEndpointInternetAddr)
diff --git a/Engine/Plugins/VirtualProduction/DMX/DMXProtocol/Source/DMXProtocolSACN/Private/DMXProtocolSACNReceiver.cpp b/Engine/Plugins/VirtualProduction/DMX/DMXProtocol/Source/DMXProtocolSACN/Private/DMXProtocolSACNReceiver.cpp
--- a/Engine/Plugins/VirtualProduction/DMX/DMXProtocol/Source/DMXProtocolSACN/Private/DMXProtocolSACNReceiver.cpp
+++ b/Engine/Plugins/VirtualProduction/DMX/DMXProtocol/Source/DMXProtocolSACN/Private/DMXProtocolSACNReceiver.cpp
@@ -17,7 +17,7 @@
 #include "HAL/RunnableThread.h"
 #include "Serialization/ArrayReader.h"
 
-DECLARE_DWORD_ACCUMULATOR_STAT(TEXT("sACN Packages Recieved Total"), STAT_SACNPackagesReceived, STATGROUP_DMX);
+DECLARE_DWORD_ACCUMULATOR_STAT(TEXT("sACN Packages Received Total"), STAT_SACNPackagesReceived, STATGROUP_DMX);
 
 
 FDMXProtocolSACNReceiver::FDMXProtocolSACNReceiver(const TSharedPtr<FDMXProtocolSACN, ESPMode::ThreadSafe>& InSACNProtocol, FSocket& InSocket, TSharedRef<FInternetAddr> InEndpointInternetAddr)
diff --git a/Engine/Source/Editor/MovieSceneCaptureDialog/Private/MovieSceneCaptureDialogModule.cpp b/Engine/Source/Editor/MovieSceneCaptureDialog/Private/MovieSceneCaptureDialogModule.cpp
--- a/Engine/Source/Editor/MovieSceneCaptureDialog/Private/MovieSceneCaptureDialogModule.cpp
+++ b/Engine/Source/Editor/MovieSceneCaptureDialog/Private/MovieSceneCaptureDialogModule.cpp
@@ -557,7 +557,7 @@ void FInEditorCapture::OnPIEViewportStarted()
 		}
 	}
 
-	UE_LOG(LogTemp, Warning, TEXT("Recieved PIE Creation callback but failed to find PIE World or missing FSlatePlayInEditorInfo for world."));
+	UE_LOG(LogTemp, Warning, TEXT("Received PIE Creation callback but failed to find PIE World or missing FSlatePlayInEditorInfo for world."));
 }
 
 void FInEditorCapture::Shutdown()
diff --git a/Engine/Source/Runtime/AudioMixer/Private/Quartz/QuartzSubsystem.cpp b/Engine/Source/Runtime/AudioMixer/Private/Quartz/QuartzSubsystem.cpp
--- a/Engine/Source/Runtime/AudioMixer/Private/Quartz/QuartzSubsystem.cpp
+++ b/Engine/Source/Runtime/AudioMixer/Private/Quartz/QuartzSubsystem.cpp
@@ -487,7 +487,7 @@ void UQuartzSubsystem::AddProxyForExternalClock(const Audio::FQuartzClockProxy&
 	// make sure we aren't adding a duplicate name
 	if(FindProxyByName(InProxy.GetClockName()))
 	{
-		UE_LOG(LogAudioQuartz, Warning, TEXT("Recieved request to add external Clock: (%s) when a clock of that name already exists (Ignoring Request)"), *InProxy.GetClockName().ToString());
+		UE_LOG(LogAudioQuartz, Warning, TEXT("Received request to add external Clock: (%s) when a clock of that name already exists (Ignoring Request)"), *InProxy.GetClockName().ToString());
 		return;
 	}
 }
diff --git a/Engine/Source/Runtime/Core/Private/Async/TaskGraph.cpp b/Engine/Source/Runtime/Core/Private/Async/TaskGraph.cpp
--- a/Engine/Source/Runtime/Core/Private/Async/TaskGraph.cpp
+++ b/Engine/Source/Runtime/Core/Private/Async/TaskGraph.cpp
@@ -2634,7 +2634,7 @@ public:
 			const double ThisTime = FPlatformTime::Seconds() - StartTime;
 			if (ThisTime > 0.02)
 			{
-				UE_CLOG(GPrintBroadcastWarnings, LogTaskGraph, Warning, TEXT("Task graph took %6.2fms for %s to recieve broadcast."), ThisTime * 1000.0, Name);
+				UE_CLOG(GPrintBroadcastWarnings, LogTaskGraph, Warning, TEXT("Task graph took %6.2fms for %s to receive broadcast."), ThisTime * 1000.0, Name);
 			}
 		}
 
@@ -2646,7 +2646,7 @@ public:
 			const double ThisTime = FPlatformTime::Seconds() - StartTime;
 			if (ThisTime > 0.02)
 			{
-				UE_CLOG(GPrintBroadcastWarnings, LogTaskGraph, Warning, TEXT("Task graph took %6.2fms for %s to recieve broadcast and do processing."), ThisTime * 1000.0, Name);
+				UE_CLOG(GPrintBroadcastWarnings, LogTaskGraph, Warning, TEXT("Task graph took %6.2fms for %s to receive broadcast and do processing."), ThisTime * 1000.0, Name);
 			}
 		}
 		if (StallForTaskThread)
@@ -2661,7 +2661,7 @@ public:
 						const double ThisTime = FPlatformTime::Seconds() - StartTime;
 						if (ThisTime > 0.02)
 						{
-							UE_CLOG(GPrintBroadcastWarnings, LogTaskGraph, Warning, TEXT("Task graph took %6.2fms for %s to recieve broadcast do processing and wait for other task threads."), ThisTime * 1000.0f, Name);
+							UE_CLOG(GPrintBroadcastWarnings, LogTaskGraph, Warning, TEXT("Task graph took %6.2fms for %s to receive broadcast do processing and wait for other task threads."), ThisTime * 1000.0f, Name);
 						}
 					}
 				}
@@ -2673,7 +2673,7 @@ public:
 					const double ThisTime = FPlatformTime::Seconds() - StartTime;
 					if (ThisTime > 0.02)
 					{
-						UE_CLOG(GPrintBroadcastWarnings, LogTaskGraph, Warning, TEXT("Task graph took %6.2fms for %s to recieve broadcast do processing and trigger other task threads."), ThisTime * 1000.0, Name);
+						UE_CLOG(GPrintBroadcastWarnings, LogTaskGraph, Warning, TEXT("Task graph took %6.2fms for %s to receive broadcast do processing and trigger other task threads."), ThisTime * 1000.0, Name);
 					}
 				}
 			}
diff --git a/Engine/Source/Runtime/CoreUObject/Private/Serialization/EditorBulkData.cpp b/Engine/Source/Runtime/CoreUObject/Private/Serialization/EditorBulkData.cpp
--- a/Engine/Source/Runtime/CoreUObject/Private/Serialization/EditorBulkData.cpp
+++ b/Engine/Source/Runtime/CoreUObject/Private/Serialization/EditorBulkData.cpp
@@ -517,7 +517,7 @@ static FGuid CreateUniqueGuid(const FGuid& NonUniqueGuid, const UObject* Owner,
 	else
 	{
 		UE_LOG(LogSerialization, Warning,
-			TEXT("CreateFromBulkData recieved an invalid FGuid. A temporary one will be generated until the package is next re-saved! Package: '%s'"),
+			TEXT("CreateFromBulkData received an invalid FGuid. A temporary one will be generated until the package is next re-saved! Package: '%s'"),
 			DebugName);
 		return FGuid::NewGuid();
 	}
diff --git a/Engine/Source/Runtime/Experimental/Chaos/Private/Chaos/AABBTree.cpp b/Engine/Source/Runtime/Experimental/Chaos/Private/Chaos/AABBTree.cpp
--- a/Engine/Source/Runtime/Experimental/Chaos/Private/Chaos/AABBTree.cpp
+++ b/Engine/Source/Runtime/Experimental/Chaos/Private/Chaos/AABBTree.cpp
@@ -3,7 +3,7 @@
 #include "Chaos/AABBTree.h"
 
 int32 FAABBTreeCVars::UpdateDirtyElementPayloadData = 1;
-FAutoConsoleVariableRef FAABBTreeCVars::CVarUpdateDirtyElementPayloadData(TEXT("p.aabbtree.updatedirtyelementpayloads"), FAABBTreeCVars::UpdateDirtyElementPayloadData, TEXT("Allow AABB tree elements to update internal payload data when they recieve a payload update"));
+FAutoConsoleVariableRef FAABBTreeCVars::CVarUpdateDirtyElementPayloadData(TEXT("p.aabbtree.updatedirtyelementpayloads"), FAABBTreeCVars::UpdateDirtyElementPayloadData, TEXT("Allow AABB tree elements to update internal payload data when they receive a payload update"));
 
 int32 FAABBTreeDirtyGridCVars::DirtyElementGridCellSize = 1000; // 0 means disabled
 FAutoConsoleVariableRef FAABBTreeDirtyGridCVars::CVarDirtyElementGridCellSize(TEXT("p.aabbtree.DirtyElementGridCellSize"), FAABBTreeDirtyGridCVars::DirtyElementGridCellSize, TEXT("DirtyElement Grid acceleration structure cell size in cm. 0 or less will disable the feature"));
diff --git a/Engine/Source/Runtime/Online/BackgroundHTTP/Private/GenericPlatform/GenericPlatformBackgroundHttpRequest.cpp b/Engine/Source/Runtime/Online/BackgroundHTTP/Private/GenericPlatform/GenericPlatformBackgroundHttpRequest.cpp
--- a/Engine/Source/Runtime/Online/BackgroundHTTP/Private/GenericPlatform/GenericPlatformBackgroundHttpRequest.cpp
+++ b/Engine/Source/Runtime/Online/BackgroundHTTP/Private/GenericPlatform/GenericPlatformBackgroundHttpRequest.cpp
@@ -139,7 +139,7 @@ const FString FGenericPlatformBackgroundHttpRequest::FGenericPlatformBackgroundH
 
 void FGenericPlatformBackgroundHttpRequest::FGenericPlatformBackgroundHttpWrapper::HttpRequestComplete(FHttpRequestPtr HttpRequestIn, FHttpResponsePtr HttpResponse, bool bSuccess)
 {
-	if (ensureAlwaysMsgf(OriginalRequest.IsValid(), TEXT("Recieved HttpRequestComplete callback with invalid OriginalRequest pointer! Can not complete request!")))
+	if (ensureAlwaysMsgf(OriginalRequest.IsValid(), TEXT("Received HttpRequestComplete callback with invalid OriginalRequest pointer! Can not complete request!")))
 	{
 		//Create a new FGenericPlatformBackgroundHttpRequest from the HttpRequest/Response
 		FBackgroundHttpResponsePtr ConstructedResponse = MakeShareable(new FGenericPlatformBackgroundHttpResponse(HttpRequestIn, HttpResponse, bSuccess));
@@ -170,7 +170,7 @@ bool FGenericPlatformBackgroundHttpRequest::FGenericPlatformBackgroundHttpWrappe
 
 void FGenericPlatformBackgroundHttpRequest::FGenericPlatformBackgroundHttpWrapper::UpdateHttpProgress(FHttpRequestPtr UnderlyingHttpRequest, int32 BytesSent, int32 BytesReceived)
 {
-	if (ensureAlwaysMsgf(OriginalRequest.IsValid(), TEXT("Recieved UpdateHttpProgress callback with invalid OriginalRequest pointer! Can not update request progress!")))
+	if (ensureAlwaysMsgf(OriginalRequest.IsValid(), TEXT("Received UpdateHttpProgress callback with invalid OriginalRequest pointer! Can not update request progress!")))
 	{
 		const int32 ByteDifference = (LastProgressUpdateBytes > 0) ? BytesReceived - LastProgressUpdateBytes : BytesReceived;
 		ensureAlwaysMsgf((ByteDifference >= 0), TEXT("Invalid Byte Difference in UpdateHttpProgress -- ByteDifference:%d | LastProgressUpdateBytes:%d | BytesSent:%d | BytesReceived:%d"), ByteDifference, LastProgressUpdateBytes, BytesSent, BytesReceived);
diff --git a/Engine/Source/Runtime/Online/BackgroundHTTP/Private/IOS/ApplePlatformBackgroundHttpManager.cpp b/Engine/Source/Runtime/Online/BackgroundHTTP/Private/IOS/ApplePlatformBackgroundHttpManager.cpp
--- a/Engine/Source/Runtime/Online/BackgroundHTTP/Private/IOS/ApplePlatformBackgroundHttpManager.cpp
+++ b/Engine/Source/Runtime/Online/BackgroundHTTP/Private/IOS/ApplePlatformBackgroundHttpManager.cpp
@@ -517,7 +517,7 @@ void FApplePlatformBackgroundHttpManager::OnTask_DidFinishDownloadingToURL(NSURL
     IPlatformFile& PlatformFile = FPlatformFileManager::Get().GetPlatformFile();
 	const bool bFileExists = PlatformFile.FileExists(*TempFilePath);
 	
-    UE_LOG(LogBackgroundHttpManager, Display, TEXT("Recieved Background Session Callback for URL:%s | TaskIdentifier:%d | bFileExists:%d | ErrorCode:%d | ErrorDescription:%s | Location:%s"), *TaskURL, TaskIdentifier, (int)(bFileExists), ErrorCode, *ErrorDescription, *TempFilePath);
+    UE_LOG(LogBackgroundHttpManager, Display, TEXT("Received Background Session Callback for URL:%s | TaskIdentifier:%d | bFileExists:%d | ErrorCode:%d | ErrorDescription:%s | Location:%s"), *TaskURL, TaskIdentifier, (int)(bFileExists), ErrorCode, *ErrorDescription, *TempFilePath);
 
     if (bFileExists)
     {
