Fix for ensuring static variable traversals are unique when a pin has (non-unique) PersistentGuids
#rnx

diff --git a/Engine/Plugins/FX/Niagara/Source/NiagaraEditor/Private/NiagaraParameterMapHistory.cpp b/Engine/Plugins/FX/Niagara/Source/NiagaraEditor/Private/NiagaraParameterMapHistory.cpp
--- a/Engine/Plugins/FX/Niagara/Source/NiagaraEditor/Private/NiagaraParameterMapHistory.cpp
+++ b/Engine/Plugins/FX/Niagara/Source/NiagaraEditor/Private/NiagaraParameterMapHistory.cpp
@@ -85,20 +85,23 @@ void FGraphTraversalHandle::Push(const UEdGraphNode* Node)
 
 void FGraphTraversalHandle::Push(const UEdGraphPin* Pin)
 {
+	const FGuid& NodeGuid = Pin->GetOwningNode()->NodeGuid;
+	ensure(NodeGuid.IsValid());
+
+	Path.Push(NodeGuid);
+	FriendlyPath.Push(Pin->GetOwningNode()->GetNodeTitle(ENodeTitleType::FullTitle).ToString());
+
 	if (Pin->PersistentGuid.IsValid())
 	{
 		Path.Push(Pin->PersistentGuid);
 	}
 	else if (Pin->GetOwningNode())
 	{
-		FGuid NodeGuid = Pin->GetOwningNode()->NodeGuid;
 		int32 Index = Pin->GetOwningNode()->GetPinIndex((UEdGraphPin*)Pin);
-		ensure(NodeGuid.IsValid());
-		NodeGuid.A += Index; // This is a bit hacky, but the chance of collisions within the same graph is really low and we want a unique value for this graph.
-		Path.Push(NodeGuid);
+		Path.Push(FGuid(Index, 0, 0, 0));
 	}
 	
-	FriendlyPath.Push(Pin->GetOwningNode()->GetNodeTitle(ENodeTitleType::FullTitle).ToString() + TEXT("->") + Pin->PinName.ToString());
+	FriendlyPath.Push(TEXT("->") + Pin->PinName.ToString());
 }
 
 void FGraphTraversalHandle::Pop()
