Niagara - When moving a data interface input to a new user parameter broadcast the interface changed delegate so that existing components in the world are synchronized correctly.

diff --git a/Engine/Plugins/FX/Niagara/Source/NiagaraEditor/Private/ViewModels/Stack/NiagaraStackFunctionInput.cpp b/Engine/Plugins/FX/Niagara/Source/NiagaraEditor/Private/ViewModels/Stack/NiagaraStackFunctionInput.cpp
--- a/Engine/Plugins/FX/Niagara/Source/NiagaraEditor/Private/ViewModels/Stack/NiagaraStackFunctionInput.cpp
+++ b/Engine/Plugins/FX/Niagara/Source/NiagaraEditor/Private/ViewModels/Stack/NiagaraStackFunctionInput.cpp
@@ -1298,6 +1298,7 @@ void UNiagaraStackFunctionInput::SetLinkedValueHandle(const FNiagaraParameterHan
 					if (InputValues.Mode == EValueMode::Data && InputValues.DataObject.IsValid())
 					{
 						InputValues.DataObject->CopyTo(UserParameters.GetDataInterface(DataInterfaceOffset));
+						UserParameters.OnInterfaceChange();
 					}
 				}
 				else
