Rollback performance improvements to static variable collection to keep things in sync with the compilation code path till we can unify them
#rnx

diff --git a/Engine/Plugins/FX/Niagara/Source/NiagaraEditor/Private/NiagaraCompiler.cpp b/Engine/Plugins/FX/Niagara/Source/NiagaraEditor/Private/NiagaraCompiler.cpp
--- a/Engine/Plugins/FX/Niagara/Source/NiagaraEditor/Private/NiagaraCompiler.cpp
+++ b/Engine/Plugins/FX/Niagara/Source/NiagaraEditor/Private/NiagaraCompiler.cpp
@@ -1460,31 +1460,29 @@ TSharedPtr<FNiagaraGraphCachedDataBase, ESPMode::ThreadSafe> FNiagaraEditorModul
 	const UNiagaraSystem* System = Cast<UNiagaraSystem>(Obj);
 	const UNiagaraEmitter* Emitter = Cast<UNiagaraEmitter>(Obj);
 
+	/*
+	* const TArray<FNiagaraVariable>& EncounterableVariables, const TArray<FNiagaraVariable>& InStaticVariables, FCompileConstantResolver ConstantResolver, const TArray<ENiagaraScriptUsage>& UsagesToProcess, const 
+	*/
 	const TArray<class UNiagaraSimulationStageBase*>* SimStages = nullptr;
 	const UNiagaraScriptSource* ScriptSource = nullptr;
 	FCompileConstantResolver ConstantResolver;
+	TArray<FNiagaraVariable> EncounterableVariables;
+	TArray<FNiagaraVariable> StaticVariablesFromSystem;
+	TArray<FNiagaraVariable> StaticVariablesFromSystemAndEmitters;
 	TArray<FNiagaraVariable> SrcStaticVariables;
 	TSharedPtr<FNiagaraGraphCachedDataBase, ESPMode::ThreadSafe> ParentCachedData;
 	FString SrcUniqueEmitterName;
 
-	int32 ValidUsageMask = 0;
-	bool bIncludeSubHistories = true;
-
 	if (System)
 	{
-		TArray<FNiagaraVariable> StaticVariablesFromEmitters;
-
 		ScriptSource = CastChecked<UNiagaraScriptSource>(System->GetSystemSpawnScript()->GetLatestSource());
 		ConstantResolver = FCompileConstantResolver(System, ENiagaraScriptUsage::SystemSpawnScript);
-		System->GatherStaticVariables(SrcStaticVariables, StaticVariablesFromEmitters);
+		System->GatherStaticVariables(StaticVariablesFromSystem, StaticVariablesFromSystemAndEmitters);
 
-		for (const FNiagaraVariable& Var : StaticVariablesFromEmitters)
-		{
-			SrcStaticVariables.AddUnique(Var);
-		}
+		SrcStaticVariables = StaticVariablesFromSystem;
 
-		ValidUsageMask = ENiagaraScriptUsageMask::System;
-		bIncludeSubHistories = false;
+		for (const FNiagaraVariable& Var : StaticVariablesFromSystemAndEmitters)
+			SrcStaticVariables.AddUnique(Var);
 	}
 	else if (Emitter)
 	{
@@ -1497,6 +1495,13 @@ TSharedPtr<FNiagaraGraphCachedDataBase, ESPMode::ThreadSafe> FNiagaraEditorModul
 			{
 				SrcStaticVariables = ((FNiagaraGraphCachedBuiltHistory*)ParentCachedData.Get())->StaticVariables;
 			}
+			/*for (const FNiagaraEmitterHandle& Handle : SysParent->GetEmitterHandles())
+			{
+				if (Handle.GetInstance() == Emitter)
+				{
+					SrcUniqueEmitterName = Handle.GetUniqueEmitterName().ToString();
+				}
+			}*/
 		}
 		SrcUniqueEmitterName = Emitter->GetUniqueEmitterName();
 		const FVersionedNiagaraEmitterData* EmitterData = Emitter->GetEmitterData(Version);
@@ -1505,10 +1510,9 @@ TSharedPtr<FNiagaraGraphCachedDataBase, ESPMode::ThreadSafe> FNiagaraEditorModul
 		ScriptSource = CastChecked<UNiagaraScriptSource>(EmitterData->GraphSource);
 		SimStages = &EmitterData->GetSimulationStages();
 		ConstantResolver = FCompileConstantResolver(FVersionedNiagaraEmitter(const_cast<UNiagaraEmitter*>(Emitter), Version), ENiagaraScriptUsage::EmitterSpawnScript);
-
-		ValidUsageMask = ENiagaraScriptUsageMask::Particle | ENiagaraScriptUsageMask::Emitter;
 	}
 
+
 	TArray<UNiagaraNodeOutput*> OutputNodes;
 	if (ScriptSource != nullptr && ScriptSource->NodeGraph != nullptr)
 	{
@@ -1531,14 +1535,6 @@ TSharedPtr<FNiagaraGraphCachedDataBase, ESPMode::ThreadSafe> FNiagaraEditorModul
 	int32 NumSimStageNodes = 0;
 	for (UNiagaraNodeOutput* FoundOutputNode : OutputNodes)
 	{
-		// only process those nodes that we care about
-		const bool bShouldSkipOutputNode = !(ValidUsageMask & (1 << (int32)FoundOutputNode->GetUsage()));
-
-		if (bShouldSkipOutputNode)
-		{
-			continue;
-		}
-
 		FName SimStageName;
 		bool bStageEnabled = true;
 		if (FoundOutputNode->GetUsage() == ENiagaraScriptUsage::ParticleSimulationStageScript && SimStages)
@@ -1546,9 +1542,9 @@ TSharedPtr<FNiagaraGraphCachedDataBase, ESPMode::ThreadSafe> FNiagaraEditorModul
 			// Find the simulation stage for this output node.
 			const FGuid& UsageId = FoundOutputNode->GetUsageId();
 			UNiagaraSimulationStageBase* const* MatchingStagePtr = SimStages->FindByPredicate([UsageId](UNiagaraSimulationStageBase* SimStage)
-			{
+				{
 					return SimStage != nullptr && SimStage->Script != nullptr && SimStage->Script->GetUsageId() == UsageId;
-			});
+				});
 
 			// Set whether or not the stage is enabled, and get the iteration source name if available.
 			bStageEnabled = MatchingStagePtr != nullptr && (*MatchingStagePtr)->bEnabled;
@@ -1569,20 +1565,18 @@ TSharedPtr<FNiagaraGraphCachedDataBase, ESPMode::ThreadSafe> FNiagaraEditorModul
 			{
 				Builder.AddGraphToCallingGraphContextStack(ScriptSource->NodeGraph);
 			}
+			Builder.RegisterEncounterableVariables(EncounterableVariables);
 			Builder.RegisterExternalStaticVariables(StaticVariables);
 
 			FString TranslationName = TEXT("Emitter");
 			Builder.BeginTranslation(TranslationName);
 			Builder.BeginUsage(FoundOutputNode->GetUsage(), SimStageName);
 			Builder.EnableScriptAllowList(true, FoundOutputNode->GetUsage());
-			Builder.bShouldBuildSubHistories = bIncludeSubHistories;
 			Builder.BuildParameterMaps(FoundOutputNode, true);
 			Builder.EndUsage();
 
 			for (const FNiagaraVariable& Var : Builder.StaticVariables)
-			{
 				StaticVariables.AddUnique(Var);
-			}
 		}
 	}
 
diff --git a/Engine/Source/Runtime/Core/Private/UObject/FortniteMainBranchObjectVersion.cpp b/Engine/Source/Runtime/Core/Private/UObject/FortniteMainBranchObjectVersion.cpp
--- a/Engine/Source/Runtime/Core/Private/UObject/FortniteMainBranchObjectVersion.cpp
+++ b/Engine/Source/Runtime/Core/Private/UObject/FortniteMainBranchObjectVersion.cpp
@@ -11,7 +11,7 @@ TMap<FGuid, FGuid> FFortniteMainBranchObjectVersion::GetSystemGuids()
 	SystemGuids.Add(DevGuids.MATERIALSHADERMAP_DERIVEDDATA_VER, FGuid("1DA85DA5733C4A489A1BE117CDC5ACCC"));
 	SystemGuids.Add(DevGuids.NANITE_DERIVEDDATA_VER, FGuid("B62CDCAF66AA45289DDB82E276E7C3E7"));
 	SystemGuids.Add(DevGuids.NIAGARASHADERMAP_DERIVEDDATA_VER, FGuid("2289D5116CF94BC0AFEAEC541468E645"));
-	SystemGuids.Add(DevGuids.Niagara_LatestScriptCompileVersion, FGuid("BC25219B21894B4684FF9289CEACCB5A"));
+	SystemGuids.Add(DevGuids.Niagara_LatestScriptCompileVersion, FGuid("ECAC41C5265C441F8FF141A1E4FEF577"));
 	SystemGuids.Add(DevGuids.SkeletalMeshDerivedDataVersion, FGuid("ACF593EAAE354FCBB34CF44F5AA1BFD2"));
 	SystemGuids.Add(DevGuids.STATICMESH_DERIVEDDATA_VER, FGuid("18E028FCF1BC469385806AA5F8AE07CC"));
 	return SystemGuids;
