- Modify Niagara dynamic parameter discovery to use the expression mask which handles RGB / RGBA pins correctly and fixes a crash from the output being out of range

diff --git a/Engine/Plugins/FX/Niagara/Source/NiagaraEditor/Private/ViewModels/Stack/NiagaraStackFunctionInput.cpp b/Engine/Plugins/FX/Niagara/Source/NiagaraEditor/Private/ViewModels/Stack/NiagaraStackFunctionInput.cpp
--- a/Engine/Plugins/FX/Niagara/Source/NiagaraEditor/Private/ViewModels/Stack/NiagaraStackFunctionInput.cpp
+++ b/Engine/Plugins/FX/Niagara/Source/NiagaraEditor/Private/ViewModels/Stack/NiagaraStackFunctionInput.cpp
@@ -868,46 +868,41 @@ FString UNiagaraStackFunctionInput::ResolveDisplayNameArgument(const FString& In
 		return ResultMaterials;
 	};
 
-	// Duplicated from UMaterialGraph::GetValidOutputIndex()
-	auto GetValidOutputIndex = [](FExpressionInput* Input)->int32 /*OutputIndex*/ {
-		int32 OutputIndex = 0;
-
-		if (Input->Expression)
+	auto GetChannelUsedBitMask =
+		[](FExpressionInput* Input, TStaticBitArray<4>& ChannelUsedMask)
 		{
-			TArray<FExpressionOutput>& Outputs = Input->Expression->GetOutputs();
-
-			if (Outputs.Num() > 0)
+			if (Input->Expression)
 			{
-				const bool bOutputIndexIsValid = Outputs.IsValidIndex(Input->OutputIndex)
-					// Attempt to handle legacy connections before OutputIndex was used that had a mask
-					&& (Input->OutputIndex != 0 || Input->Mask == 0);
+				TArray<FExpressionOutput>& Outputs = Input->Expression->GetOutputs();
 
-				for (; OutputIndex < Outputs.Num(); ++OutputIndex)
+				if (Outputs.Num() > 0)
 				{
-					const FExpressionOutput& Output = Outputs[OutputIndex];
-
-					if ((bOutputIndexIsValid && OutputIndex == Input->OutputIndex)
-						|| (!bOutputIndexIsValid
-							&& Output.Mask == Input->Mask
-							&& Output.MaskR == Input->MaskR
-							&& Output.MaskG == Input->MaskG
-							&& Output.MaskB == Input->MaskB
-							&& Output.MaskA == Input->MaskA))
+					const bool bOutputIndexIsValid = Outputs.IsValidIndex(Input->OutputIndex)
+						// Attempt to handle legacy connections before OutputIndex was used that had a mask
+						&& (Input->OutputIndex != 0 || Input->Mask == 0);
+
+					for (int32 OutputIndex=0; OutputIndex < Outputs.Num(); ++OutputIndex)
 					{
-						break;
+						const FExpressionOutput& Output = Outputs[OutputIndex];
+
+						if ((bOutputIndexIsValid && OutputIndex == Input->OutputIndex)
+							|| (!bOutputIndexIsValid
+								&& Output.Mask == Input->Mask
+								&& Output.MaskR == Input->MaskR
+								&& Output.MaskG == Input->MaskG
+								&& Output.MaskB == Input->MaskB
+								&& Output.MaskA == Input->MaskA))
+						{
+							ChannelUsedMask[0] = ChannelUsedMask[0] || (Input->MaskR != 0);
+							ChannelUsedMask[1] = ChannelUsedMask[1] || (Input->MaskG != 0);
+							ChannelUsedMask[2] = ChannelUsedMask[2] || (Input->MaskB != 0);
+							ChannelUsedMask[3] = ChannelUsedMask[3] || (Input->MaskA != 0);
+							return;
+						}
 					}
 				}
-
-				if (OutputIndex >= Outputs.Num())
-				{
-					// Work around for non-reproducible crash where OutputIndex would be out of bounds
-					OutputIndex = Outputs.Num() - 1;
-				}
 			}
-		}
-
-		return OutputIndex;
-	};
+		};
 	//~ End helper functions
 
 	// If the DisplayNameArgument to resolve is not a MaterialDynamicParam, early out.
@@ -985,8 +980,7 @@ FString UNiagaraStackFunctionInput::ResolveDisplayNameArgument(const FString& In
 		}
 
 		bAnyDynamicParametersFound = true;
-		TStaticBitArray<4>& DynamicParameterExpressionOutputMask = DynamicParameterExpressionToOutputMaskMap.FindOrAdd(DynamicParameterExpression);
-		DynamicParameterExpressionOutputMask[GetValidOutputIndex(ExpressionInput)] = true;
+		GetChannelUsedBitMask(ExpressionInput, DynamicParameterExpressionToOutputMaskMap.FindOrAdd(DynamicParameterExpression));
 	}
 
 	// Construct the final dynamic param UI name. Visit each UMaterialExpressionDynamicParameter and for those which have an output which is used, consider them for the param name.
