Fixed a bug where closing a newly created Niagara system could lead to a crash

diff --git a/Engine/Plugins/FX/Niagara/Source/NiagaraEditor/Private/Customizations/NiagaraComponentDetails.h b/Engine/Plugins/FX/Niagara/Source/NiagaraEditor/Private/Customizations/NiagaraComponentDetails.h
--- a/Engine/Plugins/FX/Niagara/Source/NiagaraEditor/Private/Customizations/NiagaraComponentDetails.h
+++ b/Engine/Plugins/FX/Niagara/Source/NiagaraEditor/Private/Customizations/NiagaraComponentDetails.h
@@ -133,7 +133,7 @@ public:
 			if(UNiagaraSystem* System = Component->GetAsset())
 			{
 				TSharedPtr<FNiagaraSystemViewModel> SystemViewModel = TNiagaraViewModelManager<UNiagaraSystem, FNiagaraSystemViewModel>::GetExistingViewModelForObject(System);
-				if(SystemViewModel.IsValid())
+				if(SystemViewModel.IsValid() && SystemViewModel->GetUserParametersHierarchyViewModel())
 				{
 					SystemViewModel->GetUserParametersHierarchyViewModel()->OnHierarchyChanged().RemoveAll(this);
 				}
@@ -164,7 +164,7 @@ public:
 		{
 			System->GetExposedParameters().OnStructureChanged().RemoveAll(this);
 			
-			if(SystemViewModel.IsValid())
+			if(SystemViewModel.IsValid() && SystemViewModel.Pin()->GetUserParametersHierarchyViewModel())
 			{
 				if (UNiagaraUserParametersHierarchyViewModel* HierarchyViewModel = SystemViewModel.Pin()->GetUserParametersHierarchyViewModel())
 				{
