Added helper functions to the camera DI to use the translated view matrices with lwc positions

diff --git a/Engine/Plugins/FX/Niagara/Shaders/Private/NiagaraDataInterfaceCamera.ush b/Engine/Plugins/FX/Niagara/Shaders/Private/NiagaraDataInterfaceCamera.ush
--- a/Engine/Plugins/FX/Niagara/Shaders/Private/NiagaraDataInterfaceCamera.ush
+++ b/Engine/Plugins/FX/Niagara/Shaders/Private/NiagaraDataInterfaceCamera.ush
@@ -64,3 +64,18 @@ void GetViewSpaceTransformsGPU_{ParameterName}(out float4x4 Out_TranslatedWorldT
 	Out_ViewToClipTransform = View.ViewToClip;
 	Out_ViewToClipNoAATransform = View.ViewToClipNoAA;
 }
+
+void ApplyPreViewTranslationToPositionGPU_{ParameterName}(in float3 In_WorldPosition, out float4 Out_TranslatedPosition)
+{
+	FLWCVector3 LwcPos = MakeLWCVector3({ParameterName}_SystemLWCTile, In_WorldPosition);
+	float3 ResolvedLwcPos = LWCToFloat(LWCAdd(LwcPos, PrimaryView.PreViewTranslation));
+	Out_TranslatedPosition = float4(ResolvedLwcPos, 1);
+}
+
+void RemovePreViewTranslationFromPositionGPU_{ParameterName}(in float4 In_TranslatedPosition, out float3 Out_WorldPosition)
+{
+	FLWCVector3 TranslatedPos = LWCPromote(In_TranslatedPosition.xyz / In_TranslatedPosition.w);
+	FLWCVector3 WorldPos = LWCSubtract(TranslatedPos, PrimaryView.PreViewTranslation);
+	WorldPos.Tile -= {ParameterName}_SystemLWCTile;
+	Out_WorldPosition = LWCToFloat(WorldPos);
+}
diff --git a/Engine/Plugins/FX/Niagara/Source/Niagara/Classes/NiagaraDataInterfaceCamera.h b/Engine/Plugins/FX/Niagara/Source/Niagara/Classes/NiagaraDataInterfaceCamera.h
--- a/Engine/Plugins/FX/Niagara/Source/Niagara/Classes/NiagaraDataInterfaceCamera.h
+++ b/Engine/Plugins/FX/Niagara/Source/Niagara/Classes/NiagaraDataInterfaceCamera.h
@@ -91,6 +91,8 @@ private:
 	static const FName GetCameraPropertiesName;
 	static const FName GetFieldOfViewName;
 	static const FName GetTAAJitterName;
+	static const FName ApplyPreViewTranslationToPosition;
+	static const FName RemovePreViewTranslationFromPosition;
 };
 
 struct FNiagaraDataIntefaceProxyCameraQuery : public FNiagaraDataInterfaceProxy
diff --git a/Engine/Plugins/FX/Niagara/Source/Niagara/Private/NiagaraDataInterfaceCamera.cpp b/Engine/Plugins/FX/Niagara/Source/Niagara/Private/NiagaraDataInterfaceCamera.cpp
--- a/Engine/Plugins/FX/Niagara/Source/Niagara/Private/NiagaraDataInterfaceCamera.cpp
+++ b/Engine/Plugins/FX/Niagara/Source/Niagara/Private/NiagaraDataInterfaceCamera.cpp
@@ -45,6 +45,8 @@ const FName UNiagaraDataInterfaceCamera::GetFieldOfViewName(TEXT("GetFieldOfView
 const FName UNiagaraDataInterfaceCamera::CalculateDistancesName(TEXT("CalculateParticleDistancesCPU"));
 const FName UNiagaraDataInterfaceCamera::QueryClosestName(TEXT("QueryClosestParticlesCPU"));
 const FName UNiagaraDataInterfaceCamera::GetTAAJitterName(TEXT("GetTAAJitter"));
+const FName UNiagaraDataInterfaceCamera::ApplyPreViewTranslationToPosition(TEXT("ApplyPreViewTranslationToPositionGPU"));
+const FName UNiagaraDataInterfaceCamera::RemovePreViewTranslationFromPosition(TEXT("RemovePreViewTranslationFromPositionGPU"));
 
 UNiagaraDataInterfaceCamera::UNiagaraDataInterfaceCamera(FObjectInitializer const& ObjectInitializer)
 	: Super(ObjectInitializer)
@@ -206,6 +208,35 @@ void UNiagaraDataInterfaceCamera::GetFunctions(TArray<FNiagaraFunctionSignature>
 	Sig.AddOutput(FNiagaraVariable(FNiagaraTypeDefinition::GetMatrix4Def(), TEXT("View To ClipNoAA Transform")), LOCTEXT("ViewToClipNoAATransformDescription", "Transforms a view space position to clip space without the temporal AA jittering"));
 	OutFunctions.Add(Sig);
 
+	
+	Sig = FNiagaraFunctionSignature();
+	Sig.Name = ApplyPreViewTranslationToPosition;
+#if WITH_EDITORONLY_DATA
+	Sig.Description = LOCTEXT("ApplyPreViewTranslationToPositionDescription", "This function applies the pre view translation to a lwc positions such as Particles.Position. The resulting vector can be used with any of the 'Translated X to Y' transform matrices.");
+	Sig.FunctionVersion = FNiagaraCameraDIFunctionVersion::LatestVersion;
+#endif
+	Sig.bMemberFunction = true;
+	Sig.bRequiresContext = false;
+	Sig.bSupportsCPU = false;
+	Sig.AddInput(FNiagaraVariable(FNiagaraTypeDefinition(GetClass()), TEXT("Camera interface")));
+	Sig.AddInput(FNiagaraVariable(FNiagaraTypeDefinition::GetPositionDef(), TEXT("World Position")), LOCTEXT("ApplyPreViewTranslationToPositionInputDescription", "Position input to be transformed."));
+	Sig.AddOutput(FNiagaraVariable(FNiagaraTypeDefinition::GetVec4Def(), TEXT("Translated Position")), LOCTEXT("ApplyPreViewTranslationToPositionOutputDescription", "Translated vector that can be used with the 'Translated X to Y' transforms matrices."));
+	OutFunctions.Add(Sig);
+
+	Sig = FNiagaraFunctionSignature();
+	Sig.Name = RemovePreViewTranslationFromPosition;
+#if WITH_EDITORONLY_DATA
+	Sig.Description = LOCTEXT("RemovePreViewTranslationFromPositionDescription", "This function removes the pre view translation from a vector to return a lwc position. Use the result of a 'X to Translated Y' transform matrix multiplication as input vector.");
+	Sig.FunctionVersion = FNiagaraCameraDIFunctionVersion::LatestVersion;
+#endif
+	Sig.bMemberFunction = true;
+	Sig.bRequiresContext = false;
+	Sig.bSupportsCPU = false;
+	Sig.AddInput(FNiagaraVariable(FNiagaraTypeDefinition(GetClass()), TEXT("Camera interface")));
+	Sig.AddInput(FNiagaraVariable(FNiagaraTypeDefinition::GetVec4Def(), TEXT("Translated Position")), LOCTEXT("RemovePreViewTranslationFromPositionInputDescription", "Translated vector that is a result of a 'X to Translated Y' transforms matrix."));
+	Sig.AddOutput(FNiagaraVariable(FNiagaraTypeDefinition::GetPositionDef(), TEXT("World Position")), LOCTEXT("RemovePreViewTranslationFromPositionOutputDescription", "The transformed position."));
+	OutFunctions.Add(Sig);
+
 
 	Sig = FNiagaraFunctionSignature();
 	Sig.Name = GetFieldOfViewName;
@@ -287,6 +318,8 @@ bool UNiagaraDataInterfaceCamera::GetFunctionHLSL(const FNiagaraDataInterfaceGPU
 		(FunctionInfo.DefinitionName == GetClipSpaceTransformsName) ||
 		(FunctionInfo.DefinitionName == GetViewSpaceTransformsName) ||
 		(FunctionInfo.DefinitionName == GetCameraPropertiesName) ||
+		(FunctionInfo.DefinitionName == ApplyPreViewTranslationToPosition) ||
+		(FunctionInfo.DefinitionName == RemovePreViewTranslationFromPosition) ||
 		(FunctionInfo.DefinitionName == GetTAAJitterName))
 	{
 		return true;
