Fixed Niagara perf widget layout and improved display of large numbers

diff --git a/Engine/Plugins/FX/Niagara/Source/NiagaraEditorWidgets/Private/SNiagaraOverviewStack.cpp b/Engine/Plugins/FX/Niagara/Source/NiagaraEditorWidgets/Private/SNiagaraOverviewStack.cpp
--- a/Engine/Plugins/FX/Niagara/Source/NiagaraEditorWidgets/Private/SNiagaraOverviewStack.cpp
+++ b/Engine/Plugins/FX/Niagara/Source/NiagaraEditorWidgets/Private/SNiagaraOverviewStack.cpp
@@ -975,14 +975,6 @@ TSharedRef<ITableRow> SNiagaraOverviewStack::OnGenerateRowForEntry(UNiagaraStack
 		];
 
 
-		// Perf widget
-		ContentBox->AddSlot()
-        .AutoWidth()
-        .Padding(3, 0, 0, 0)
-        [
-            SNew(SNiagaraStackRowPerfWidget, Item)
-        ];
-
 		TSharedRef<SHorizontalBox> OptionsBox = SNew(SHorizontalBox);
 		
 		UNiagaraStackModuleItem* StackModuleItem = Cast<UNiagaraStackModuleItem>(StackItem);
@@ -1087,6 +1079,13 @@ TSharedRef<ITableRow> SNiagaraOverviewStack::OnGenerateRowForEntry(UNiagaraStack
 					]
 				]
 			]
+			// Perf widget
+			+ SHorizontalBox::Slot()
+			.AutoWidth()
+			.Padding(3, 0, 0, 0)
+			[
+				SNew(SNiagaraStackRowPerfWidget, Item)
+			]
 			+ SHorizontalBox::Slot()
 			.VAlign(VAlign_Center)
 			.AutoWidth()
@@ -1105,7 +1104,14 @@ TSharedRef<ITableRow> SNiagaraOverviewStack::OnGenerateRowForEntry(UNiagaraStack
 				.MaxDesiredWidth(250.f)
 				[
 					ContentBox
-				]			
+				]
+			]
+			// Perf widget
+			+ SHorizontalBox::Slot()
+			.AutoWidth()
+			.Padding(3, 0, 0, 0)
+			[
+				SNew(SNiagaraStackRowPerfWidget, Item)
 			]
 			+ SHorizontalBox::Slot()
 			.AutoWidth()
diff --git a/Engine/Plugins/FX/Niagara/Source/NiagaraEditorWidgets/Private/Stack/SNiagaraStackRowPerfWidget.cpp b/Engine/Plugins/FX/Niagara/Source/NiagaraEditorWidgets/Private/Stack/SNiagaraStackRowPerfWidget.cpp
--- a/Engine/Plugins/FX/Niagara/Source/NiagaraEditorWidgets/Private/Stack/SNiagaraStackRowPerfWidget.cpp
+++ b/Engine/Plugins/FX/Niagara/Source/NiagaraEditorWidgets/Private/Stack/SNiagaraStackRowPerfWidget.cpp
@@ -268,6 +268,7 @@ FText SNiagaraStackRowPerfWidget::GetPerformanceDisplayText() const
 	TimeFormatOptions.MinimumIntegralDigits = 1;
 	TimeFormatOptions.MinimumFractionalDigits = 2;
 	TimeFormatOptions.MaximumFractionalDigits = 2;
+	TimeFormatOptions.UseGrouping = false;
 	
 	if (IsGroupHeaderEntry())
 	{
@@ -275,10 +276,18 @@ FText SNiagaraStackRowPerfWidget::GetPerformanceDisplayText() const
 		{
 			return FText::FromString("N/A");
 		}
+		if (GroupOverallTime >= 1000)
+		{
+			return FText::Format(FText::FromString("{0}s"), FText::AsNumber(GroupOverallTime / 1000, &TimeFormatOptions));
+		}
 		return FText::Format(FText::FromString("{0}ms"), FText::AsNumber(GroupOverallTime, &TimeFormatOptions));
 	}
 	if (GetDisplayMode() == ENiagaraStatDisplayMode::Absolute)
 	{
+		if (StackEntryTime >= 1000)
+		{
+			return FText::Format(FText::FromString("{0}s"), FText::AsNumber(StackEntryTime / 1000, &TimeFormatOptions));
+		}
 		return FText::Format(FText::FromString("{0}ms"), FText::AsNumber(StackEntryTime, &TimeFormatOptions));
 	}
 	FNumberFormattingOptions PercentFormatOptions;
