Fix for renaming static switches variable names resulting in overlapping VariableGuids in the graph's NiagaraScriptVariables

diff --git a/Engine/Plugins/FX/Niagara/Source/NiagaraEditor/Private/NiagaraGraph.cpp b/Engine/Plugins/FX/Niagara/Source/NiagaraEditor/Private/NiagaraGraph.cpp
--- a/Engine/Plugins/FX/Niagara/Source/NiagaraEditor/Private/NiagaraGraph.cpp
+++ b/Engine/Plugins/FX/Niagara/Source/NiagaraEditor/Private/NiagaraGraph.cpp
@@ -2314,19 +2314,52 @@ bool UNiagaraGraph::RenameParameterFromPin(const FNiagaraVariable& Parameter, FN
 	return false;
 }
 
-bool UNiagaraGraph::RenameParameter(const FNiagaraVariable& Parameter, FName NewName, bool bRenameRequestedFromStaticSwitch, bool* bMerged, bool bSuppressEvents)
+void UNiagaraGraph::FixupReferenceCollectionsPostRename(const FNiagaraVariable& OrigVariable, const FNiagaraVariable& DestVariable, bool bParameterMerged, bool bSuppressEvents)
+{
+	// Fixup reference collection and pin names
+	FNiagaraGraphParameterReferenceCollection NewReferences;
+	if (ParameterToReferencesMap.RemoveAndCopyValue(OrigVariable, NewReferences))
+	{
+		if (!NewReferences.ParameterReferences.IsEmpty())
+		{
+			const FText NewNameText = FText::FromName(DestVariable.GetName());
+			for (FNiagaraGraphParameterReference& Reference : NewReferences.ParameterReferences)
+			{
+				UNiagaraNode* Node = Cast<UNiagaraNode>(Reference.Value.Get());
+				if (Node && Node->GetGraph() == this)
+				{
+					Node->Modify();
+					if (UEdGraphPin* Pin = Node->GetPinByPersistentGuid(Reference.Key))
+					{
+						Pin->Modify();
+						Node->CommitEditablePinName(NewNameText, Pin, bSuppressEvents);
+					}
+				}
+			}
+		}
+
+		if (!bParameterMerged)
+		{
+			ParameterToReferencesMap.Add(DestVariable, NewReferences);
+		}
+		else
+		{
+			RefreshParameterReferences();
+		}
+	}
+}
+
+bool UNiagaraGraph::RenameParameter(const FNiagaraVariable& Parameter, FName NewName, bool bRenameRequestedFromStaticSwitch, bool* bOutMerged, bool bSuppressEvents)
 {
 	check(!bIsForCompilationOnly);
 
 	// Initialize the merger state if requested
-	if (bMerged)
-		*bMerged = false;
+	if (bOutMerged)
+		*bOutMerged = false;
 
 	if (Parameter.GetName() == NewName)
 		return true;
 
-	
-
 	// Block rename when already renaming. This prevents recursion when CommitEditablePinName is called on referenced nodes. 
 	if (bIsRenamingParameter)
 	{
@@ -2334,6 +2367,8 @@ bool UNiagaraGraph::RenameParameter(const FNiagaraVariable& Parameter, FName New
 	}
 	bIsRenamingParameter = true;
 
+	bool bParameterMerged = false;
+
 	// Create the new parameter
 	FNiagaraVariable NewParameter = Parameter;
 	NewParameter.SetName(NewName);
@@ -2395,8 +2430,7 @@ bool UNiagaraGraph::RenameParameter(const FNiagaraVariable& Parameter, FName New
 	}
 	else
 	{
-		if (bMerged)
-			*bMerged = true;
+		bParameterMerged = true;
 		FName NewParamName = NewName;
 		FNiagaraEditorUtilities::DecomposeVariableNamespace(NewName, NewParamName);
 		FName OldParamName = Parameter.GetName();
@@ -2407,40 +2441,87 @@ bool UNiagaraGraph::RenameParameter(const FNiagaraVariable& Parameter, FName New
 			FNiagaraParameterUtilities::FormatParameterNameForTextDisplay(NewName)));
 	}
 
-	// Fixup reference collection and pin names
-	if (FNiagaraGraphParameterReferenceCollection* ReferenceCollection = ParameterToReferencesMap.Find(Parameter))
+	FixupReferenceCollectionsPostRename(Parameter, NewParameter, bParameterMerged, bSuppressEvents);
+
+	bIsRenamingParameter = false;
+
+	NotifyGraphChanged();
+
+	if (bOutMerged)
 	{
-		const FText NewNameText = FText::FromName(NewName);
-		FNiagaraGraphParameterReferenceCollection NewReferences = *ReferenceCollection;
-		for (FNiagaraGraphParameterReference& Reference : NewReferences.ParameterReferences)
-		{
-			UNiagaraNode* Node = Cast<UNiagaraNode>(Reference.Value.Get());
-			if (Node && Node->GetGraph() == this)
-			{
-				Node->Modify();
-				UEdGraphPin* Pin = Node->GetPinByPersistentGuid(Reference.Key);
-				if (Pin)
-				{
-					Pin->Modify();
-					Node->CommitEditablePinName(NewNameText, Pin, bSuppressEvents);
-				}
-			}
-		}
+		*bOutMerged = bParameterMerged;
+	}
 
-		ParameterToReferencesMap.Remove(Parameter);
-		if (!bMerged)
+	return true;
+}
+
+bool UNiagaraGraph::RenameStaticSwitch(UNiagaraNodeStaticSwitch* SwitchNode, FName NewParameterName)
+{
+	if (!SwitchNode || SwitchNode->InputParameterName == NewParameterName)
+	{
+		return false;
+	}
+
+	bIsRenamingParameter = true;
+	ON_SCOPE_EXIT{ bIsRenamingParameter = false; };
+
+	const FName OrigParameterName = SwitchNode->InputParameterName;
+
+	const FNiagaraVariableBase OrigVariable(SwitchNode->GetInputType(), OrigParameterName);
+	const FNiagaraVariableBase DestVariable(SwitchNode->GetInputType(), NewParameterName);
+
+	// see if the old/new names map to existing UNiagaraScriptVariables
+	UNiagaraScriptVariable* OrigScriptVariable = VariableToScriptVariable.FindRef(OrigVariable);
+	UNiagaraScriptVariable* DestScriptVariable = VariableToScriptVariable.FindRef(DestVariable);
+
+	// assign the new name to the node
+	SwitchNode->InputParameterName = NewParameterName;
+
+	// if we're just merging to an existing variable then we don't need to create a new one
+	if (DestScriptVariable)
+	{
+		FNiagaraEditorUtilities::InfoWithToastAndLog(FText::Format(
+			LOCTEXT("MergedVar", "\"{0}\" has been merged with parameter \"{1}\".\nAll of \"{1}\"'s meta-data will be used, overwriting \"{0}\"'s meta-data."),
+			FNiagaraParameterUtilities::FormatParameterNameForTextDisplay(OrigParameterName),
+			FNiagaraParameterUtilities::FormatParameterNameForTextDisplay(NewParameterName)));
+	}
+	else
+	{
+		TObjectPtr<UNiagaraScriptVariable>& CreatedParameter = VariableToScriptVariable.Add(DestVariable);
+		
+		if (OrigScriptVariable)
 		{
-			ParameterToReferencesMap.Add(NewParameter, NewReferences);
+			CreatedParameter = DuplicateObject<UNiagaraScriptVariable>(OrigScriptVariable, this);
 		}
 		else
 		{
-			RefreshParameterReferences();
+			CreatedParameter = NewObject<UNiagaraScriptVariable>(this);
 		}
+		
+		CreatedParameter->SetFlags(RF_Transactional);
+		CreatedParameter->Metadata.CreateNewGuid();
+		CreatedParameter->Variable = DestVariable;
 	}
 
-	bIsRenamingParameter = false;
+	// now see if we need to remove any script variables that are now obsolete because they have been renamed
+	if (OrigScriptVariable)
+	{
+		const TArray<FNiagaraVariable> StaticSwitchInputs = FindStaticSwitchInputs(false /*bReachableOnly*/);
+
+		if (!StaticSwitchInputs.Contains(OrigVariable))
+		{
+			RemoveParameter(OrigVariable);
+		}
+	}
+
+	// Fixup reference collection and pin names
+	FixupReferenceCollectionsPostRename(OrigVariable, DestVariable, DestScriptVariable != nullptr /* bIsMerging */, false /* bSuppressEvents */);
 
 	NotifyGraphChanged();
+
+	// force the graph to refresh the metadata
+	GetParameterReferenceMap();
+
 	return true;
 }
 
diff --git a/Engine/Plugins/FX/Niagara/Source/NiagaraEditor/Private/NiagaraNodeStaticSwitch.cpp b/Engine/Plugins/FX/Niagara/Source/NiagaraEditor/Private/NiagaraNodeStaticSwitch.cpp
--- a/Engine/Plugins/FX/Niagara/Source/NiagaraEditor/Private/NiagaraNodeStaticSwitch.cpp
+++ b/Engine/Plugins/FX/Niagara/Source/NiagaraEditor/Private/NiagaraNodeStaticSwitch.cpp
@@ -168,11 +168,8 @@ FName UNiagaraNodeStaticSwitch::GetOptionPinName(const FNiagaraVariable& Variabl
 
 void UNiagaraNodeStaticSwitch::ChangeSwitchParameterName(const FName& NewName)
 {
-	FNiagaraVariable OldValue(GetInputType(), InputParameterName);
-	InputParameterName = NewName;
-	GetNiagaraGraph()->RenameParameter(OldValue, NewName, true);
+	GetNiagaraGraph()->RenameStaticSwitch(this, NewName);
 	VisualsChangedDelegate.Broadcast(this);
-	RemoveUnusedGraphParameter(OldValue);	
 }
 
 void UNiagaraNodeStaticSwitch::OnSwitchParameterTypeChanged(const FNiagaraTypeDefinition& OldType)
diff --git a/Engine/Plugins/FX/Niagara/Source/NiagaraEditor/Public/NiagaraGraph.h b/Engine/Plugins/FX/Niagara/Source/NiagaraEditor/Public/NiagaraGraph.h
--- a/Engine/Plugins/FX/Niagara/Source/NiagaraEditor/Public/NiagaraGraph.h
+++ b/Engine/Plugins/FX/Niagara/Source/NiagaraEditor/Public/NiagaraGraph.h
@@ -8,12 +8,12 @@
 #include "NiagaraScript.h"
 #include "NiagaraGraph.generated.h"
 
+class UNiagaraNodeStaticSwitch;
 class UNiagaraParameterDefinitions;
 class UNiagaraScriptVariable;
 struct FSynchronizeWithParameterDefinitionsArgs;
 struct FNiagaraScriptVariableData;
 
-
 /** This is the type of action that occurred on a given Niagara graph. Note that this should follow from EEdGraphActionType, leaving some slop for growth. */
 enum ENiagaraGraphActionType
 {
@@ -319,6 +319,8 @@ class UNiagaraGraph : public UEdGraph
 	/** Rename a pin inline in a graph. If this is the only instance used in the graph, then rename them all, otherwise make a duplicate. */
 	bool RenameParameterFromPin(const FNiagaraVariable& Parameter, FName NewName, UEdGraphPin* InPin);
 
+	bool RenameStaticSwitch(UNiagaraNodeStaticSwitch* SwitchNode, FName NewName);
+
 	/** Changes the type of existing graph parameters.
 	 *  Optionally creates orphaned pins for any connection that won't be kept, but tries to keep connections as long as types are matching.
 	 *  Changing multiple parameters at once helps with maintaining connections.
@@ -445,6 +447,8 @@ private:
 	static bool VariableLess(const FNiagaraVariable& Lhs, const FNiagaraVariable& Rhs);
 	TOptional<FNiagaraScriptVariableData> GetScriptVariableData(const FNiagaraVariable& Variable) const;
 
+	void FixupReferenceCollectionsPostRename(const FNiagaraVariable& OrigVariable, const FNiagaraVariable& DestVariable, bool bParameterMerged, bool bSuppressEvents);
+
 private:
 	/** The current change identifier for this graph overall. Used to sync status with UNiagaraScripts.*/
 	UPROPERTY()
