Fix for ensuring static variable traversals are unique when a pin has (non-unique) PersistentGuids

diff --git a/Engine/Plugins/FX/Niagara/Source/NiagaraEditor/Private/NiagaraHlslTranslator.cpp b/Engine/Plugins/FX/Niagara/Source/NiagaraEditor/Private/NiagaraHlslTranslator.cpp
--- a/Engine/Plugins/FX/Niagara/Source/NiagaraEditor/Private/NiagaraHlslTranslator.cpp
+++ b/Engine/Plugins/FX/Niagara/Source/NiagaraEditor/Private/NiagaraHlslTranslator.cpp
@@ -5735,7 +5735,7 @@ void FHlslNiagaraTranslator::SetConstantByStaticVariable(int32& OutValue, const
 			PinToTest = InDefaultPin->LinkedTo[0];
 		}
 		FGraphTraversalHandle PinHandle(ActiveHistoryForFunctionCalls.ActivePath);
-		PinHandle.Push(PinToTest);
+		PinHandle.PushPin(PinToTest);
 
 		const FString* ConstantPtr = CompileData->PinToConstantValues.Find(PinHandle);
 		if (ConstantPtr != nullptr)
diff --git a/Engine/Plugins/FX/Niagara/Source/NiagaraEditor/Private/NiagaraParameterMapHistory.cpp b/Engine/Plugins/FX/Niagara/Source/NiagaraEditor/Private/NiagaraParameterMapHistory.cpp
--- a/Engine/Plugins/FX/Niagara/Source/NiagaraEditor/Private/NiagaraParameterMapHistory.cpp
+++ b/Engine/Plugins/FX/Niagara/Source/NiagaraEditor/Private/NiagaraParameterMapHistory.cpp
@@ -76,33 +76,42 @@ bool NiagaraDebugShouldLogEntriesByNameSubstring(const FName& InName)
 
 #endif
 
-void FGraphTraversalHandle::Push(const UEdGraphNode* Node)
+void FGraphTraversalHandle::PushNode(const UEdGraphNode* Node)
 {
 	ensure(Node->NodeGuid.IsValid());
 	Path.Push(Node->NodeGuid);
 	FriendlyPath.Push(Node->GetNodeTitle(ENodeTitleType::FullTitle).ToString());
 }
 
-void FGraphTraversalHandle::Push(const UEdGraphPin* Pin)
+void FGraphTraversalHandle::PushPin(const UEdGraphPin* Pin)
 {
+	PushNode(Pin->GetOwningNode());
+
 	if (Pin->PersistentGuid.IsValid())
 	{
 		Path.Push(Pin->PersistentGuid);
 	}
 	else if (Pin->GetOwningNode())
 	{
-		FGuid NodeGuid = Pin->GetOwningNode()->NodeGuid;
 		int32 Index = Pin->GetOwningNode()->GetPinIndex((UEdGraphPin*)Pin);
-		ensure(NodeGuid.IsValid());
-		NodeGuid.A += Index; // This is a bit hacky, but the chance of collisions within the same graph is really low and we want a unique value for this graph.
-		Path.Push(NodeGuid);
+		FGuid PinIndexGuid(Index, 0, 0, 0);
+		Path.Push(PinIndexGuid);
 	}
 	
 	FriendlyPath.Push(Pin->GetOwningNode()->GetNodeTitle(ENodeTitleType::FullTitle).ToString() + TEXT("->") + Pin->PinName.ToString());
 }
 
-void FGraphTraversalHandle::Pop()
+void FGraphTraversalHandle::PopNode()
+{
+	Path.Pop();
+	FriendlyPath.Pop();
+}
+
+
+void FGraphTraversalHandle::PopPin()
 {
+	PopNode();
+
 	Path.Pop();
 	FriendlyPath.Pop();
 }
@@ -216,7 +225,7 @@ void FNiagaraParameterMapHistory::RegisterConstantPin(const FGraphTraversalHandl
 	if (UEdGraphSchema_Niagara::IsStaticPin(InPin))
 	{
 		FGraphTraversalHandle Handle = InTraversalPath;
-		Handle.Push(InPin);
+		Handle.PushPin(InPin);
 		
 		FString* FoundValue = PinToConstantValues.Find(Handle);
 		if (FoundValue != nullptr)
@@ -1124,7 +1133,7 @@ int32 FNiagaraParameterMapHistoryBuilder::RegisterConstantFromInputPin(const UEd
 	if (UNiagaraScript::LogCompileStaticVars > 0 || NiagaraDebugShouldLogEntriesByNameSubstring(InputPin->PinName.ToString()))
 	{
 		FGraphTraversalHandle TestPath = ActivePath;
-		TestPath.Push(InputPin);
+		TestPath.PushPin(InputPin);
 		UE_LOG(LogNiagaraEditor, Log, TEXT("RegisterConstantFromInputPin Value \"%s\" Var: \"%s\""), ConstantIdx != INDEX_NONE ? *Constants[ConstantIdx] : TEXT("INDEX_NONE"), *TestPath.ToString());
 	}
 
@@ -1143,7 +1152,7 @@ int32 FNiagaraParameterMapHistoryBuilder::GetConstantFromInputPin(const UEdGraph
 				if (UNiagaraScript::LogCompileStaticVars > 0 || NiagaraDebugShouldLogEntriesByNameSubstring(InputPin->PinName.ToString()))
 				{
 					FGraphTraversalHandle TestPath = ActivePath;
-					TestPath.Push(InputPin);
+					TestPath.PushPin(InputPin);
 					UE_LOG(LogNiagaraEditor, Log, TEXT("GetConstantFromInputPin Value \"%s\" Var: \"%s\""), *Constants[*IdxPtr], *TestPath.ToString());
 				}
 				return *IdxPtr;
@@ -1165,7 +1174,7 @@ int32 FNiagaraParameterMapHistoryBuilder::GetConstantFromOutputPin(const UEdGrap
 				if (UNiagaraScript::LogCompileStaticVars > 0 || NiagaraDebugShouldLogEntriesByNameSubstring(OutputPin->GetName()))
 				{
 					FGraphTraversalHandle TestPath = ActivePath;
-					TestPath.Push(OutputPin);
+					TestPath.PushPin(OutputPin);
 					UE_LOG(LogNiagaraEditor, Log, TEXT("GetConstantFromOutputPin Value \"%s\" Var: \"%s\""), *Constants[*IdxPtr], *TestPath.ToString());
 				}
 				return *IdxPtr;
@@ -1433,7 +1442,7 @@ void FNiagaraParameterMapHistoryBuilder::EnterFunction(const FString& InNodeName
 	{
 		RegisterNodeVisitation(Node);
 		CallingContext.Push(Node);
-		ActivePath.Push(Node);
+		ActivePath.PushNode(Node);
 		CallingGraphContext.Push(InGraph);
 		PinToParameterMapIndices.Emplace();
 		PinToConstantIndices.Emplace();
@@ -1453,7 +1462,7 @@ void FNiagaraParameterMapHistoryBuilder::ExitFunction(const FString& InNodeName,
 	{
 		CallingContext.Pop();
 		CallingGraphContext.Pop();
-		ActivePath.Pop();
+		ActivePath.PopNode();
 		PinToParameterMapIndices.Pop();
 		FunctionNameContextStack.Pop();
 		PinToConstantIndices.Pop();
@@ -1466,7 +1475,7 @@ void FNiagaraParameterMapHistoryBuilder::EnterEmitter(const FString& InEmitterNa
 {
 	RegisterNodeVisitation(Node);
 	CallingContext.Push(Node);
-	ActivePath.Push(Node);
+	ActivePath.PushNode(Node);
 	CallingGraphContext.Push(InGraph);
 	EmitterNameContextStack.Emplace(*InEmitterName);
 	BuildCurrentAliases();
@@ -1499,7 +1508,7 @@ void FNiagaraParameterMapHistoryBuilder::ExitEmitter(const FString& InEmitterNam
 	CallingContext.Pop();
 	CallingGraphContext.Pop();
 	EmitterNameContextStack.Pop();
-	ActivePath.Pop();
+	ActivePath.PopNode();
 	BuildCurrentAliases();
 	ContextuallyVisitedNodes.Pop();
 	EncounteredFunctionNames.Pop();
@@ -1841,7 +1850,7 @@ int32 FNiagaraParameterMapHistoryBuilder::HandleVariableRead(int32 ParamMapIdx,
 			if (UNiagaraScript::LogCompileStaticVars > 0)
 			{
 				FGraphTraversalHandle TestPath = ActivePath;
-				TestPath.Push(InPin);
+				TestPath.PushPin(InPin);
 				UE_LOG(LogNiagaraEditor, Log, TEXT("HandleVariableRead RapidIt Value \"%s\" Var: \"%s\" Path: \"%s\""), *DefaultValue, *VarRapid.GetName().ToString(),  *TestPath.ToString());
 			}
 			
@@ -1872,7 +1881,7 @@ int32 FNiagaraParameterMapHistoryBuilder::HandleVariableRead(int32 ParamMapIdx,
 			if (UNiagaraScript::LogCompileStaticVars > 0)
 			{
 				FGraphTraversalHandle TestPath = ActivePath;
-				TestPath.Push(InPin);
+				TestPath.PushPin(InPin);
 				UE_LOG(LogNiagaraEditor, Log, TEXT("HandleVariableRead Ext Value \"%s\" Var: \"%s\" Path: \"%s\""), *DefaultValue, *VarRapid.GetName().ToString(),  *TestPath.ToString());
 			}
 		}
diff --git a/Engine/Plugins/FX/Niagara/Source/NiagaraEditor/Public/NiagaraParameterMapHistory.h b/Engine/Plugins/FX/Niagara/Source/NiagaraEditor/Public/NiagaraParameterMapHistory.h
--- a/Engine/Plugins/FX/Niagara/Source/NiagaraEditor/Public/NiagaraParameterMapHistory.h
+++ b/Engine/Plugins/FX/Niagara/Source/NiagaraEditor/Public/NiagaraParameterMapHistory.h
@@ -77,9 +77,11 @@ public:
 		Path.Add(Guid);
 	}
 
-	void Push(const UEdGraphNode* Node);
-	void Push(const UEdGraphPin* Pin);
-	void Pop();
+	void PushNode(const UEdGraphNode* Node);
+	void PopNode();
+
+	void PushPin(const UEdGraphPin* Pin);
+	void PopPin();
 
 	bool operator==(const FGraphTraversalHandle& Rhs) const
 	{
