Fix static analysis warnings with ParamMapHistory.cpp
#rnx

diff --git a/Engine/Plugins/FX/Niagara/Source/NiagaraEditor/Private/NiagaraParameterMapHistory.cpp b/Engine/Plugins/FX/Niagara/Source/NiagaraEditor/Private/NiagaraParameterMapHistory.cpp
--- a/Engine/Plugins/FX/Niagara/Source/NiagaraEditor/Private/NiagaraParameterMapHistory.cpp
+++ b/Engine/Plugins/FX/Niagara/Source/NiagaraEditor/Private/NiagaraParameterMapHistory.cpp
@@ -78,27 +78,46 @@ bool NiagaraDebugShouldLogEntriesByNameSubstring(const FName& InName)
 
 void FGraphTraversalHandle::PushNode(const UEdGraphNode* Node)
 {
-	ensure(Node->NodeGuid.IsValid());
-	Path.Push(Node->NodeGuid);
-	FriendlyPath.Push(Node->GetNodeTitle(ENodeTitleType::FullTitle).ToString());
+	if (ensure(Node))
+	{
+		ensure(Node->NodeGuid.IsValid());
+		Path.Push(Node->NodeGuid);
+		FriendlyPath.Push(Node->GetNodeTitle(ENodeTitleType::FullTitle).ToString());
+	}
+	else
+	{
+		Path.Push(FGuid());
+		FriendlyPath.Push(FString());
+	}
 }
 
 void FGraphTraversalHandle::PushPin(const UEdGraphPin* Pin)
 {
-	PushNode(Pin->GetOwningNode());
-
-	if (Pin->PersistentGuid.IsValid())
+	if (ensure(Pin && Pin->GetOwningNode()))
 	{
-		Path.Push(Pin->PersistentGuid);
+		PushNode(Pin->GetOwningNode());
+
+		if (Pin->PersistentGuid.IsValid())
+		{
+			Path.Push(Pin->PersistentGuid);
+		}
+		else if (Pin->GetOwningNode())
+		{
+			int32 Index = Pin->GetOwningNode()->GetPinIndex((UEdGraphPin*)Pin);
+			FGuid PinIndexGuid(Index, 0, 0, 0);
+			Path.Push(PinIndexGuid);
+		}
+	
+		FriendlyPath.Push(Pin->GetOwningNode()->GetNodeTitle(ENodeTitleType::FullTitle).ToString() + TEXT("->") + Pin->PinName.ToString());
 	}
-	else if (Pin->GetOwningNode())
+	else
 	{
-		int32 Index = Pin->GetOwningNode()->GetPinIndex((UEdGraphPin*)Pin);
-		FGuid PinIndexGuid(Index, 0, 0, 0);
-		Path.Push(PinIndexGuid);
+		Path.Push(FGuid());
+		FriendlyPath.Push(FString());
+
+		Path.Push(FGuid());
+		FriendlyPath.Push(FString());
 	}
-	
-	FriendlyPath.Push(Pin->GetOwningNode()->GetNodeTitle(ENodeTitleType::FullTitle).ToString() + TEXT("->") + Pin->PinName.ToString());
 }
 
 void FGraphTraversalHandle::PopNode()
