2d dispatch for grid2d collection
avoid pre stage clear when resolution of dispatch matches grid DI resolution

diff --git a/Engine/Plugins/FX/Niagara/Source/Niagara/Classes/NiagaraDataInterfaceGrid2DCollection.h b/Engine/Plugins/FX/Niagara/Source/Niagara/Classes/NiagaraDataInterfaceGrid2DCollection.h
--- a/Engine/Plugins/FX/Niagara/Source/Niagara/Classes/NiagaraDataInterfaceGrid2DCollection.h
+++ b/Engine/Plugins/FX/Niagara/Source/Niagara/Classes/NiagaraDataInterfaceGrid2DCollection.h
@@ -160,6 +160,11 @@ public:
 	virtual FNiagaraDataInterfaceParametersCS* CreateShaderStorage(const FNiagaraDataInterfaceGPUParamInfo& ParameterInfo, const FShaderParameterMap& ParameterMap) const override;
 	virtual const FTypeLayoutDesc* GetShaderStorageType() const override;
 
+#if WITH_EDITOR
+	virtual ENiagaraGpuDispatchType GetGpuDispatchType() const override { return ENiagaraGpuDispatchType::TwoD; }
+	virtual FIntVector GetGpuDispatchNumThreads() const { return FIntVector(8, 8, 1); }
+#endif
+
 	virtual void ProvidePerInstanceDataForRenderThread(void* DataForRenderThread, void* PerInstanceData, const FNiagaraSystemInstanceID& SystemInstance) override {}
 	virtual bool InitPerInstanceData(void* PerInstanceData, FNiagaraSystemInstance* SystemInstance) override;
 	virtual void DestroyPerInstanceData(void* PerInstanceData, FNiagaraSystemInstance* SystemInstance) override;
diff --git a/Engine/Plugins/FX/Niagara/Source/Niagara/Private/NiagaraDataInterfaceGrid2DCollection.cpp b/Engine/Plugins/FX/Niagara/Source/Niagara/Private/NiagaraDataInterfaceGrid2DCollection.cpp
--- a/Engine/Plugins/FX/Niagara/Source/Niagara/Private/NiagaraDataInterfaceGrid2DCollection.cpp
+++ b/Engine/Plugins/FX/Niagara/Source/Niagara/Private/NiagaraDataInterfaceGrid2DCollection.cpp
@@ -9,6 +9,7 @@
 #include "NiagaraShader.h"
 #include "NiagaraShaderParametersBuilder.h"
 #include "NiagaraSystemInstance.h"
+#include "NiagaraSimStageData.h"
 
 #include "Engine/Canvas.h"
 #include "Engine/TextureRenderTarget2D.h"
@@ -2589,8 +2590,22 @@ void FNiagaraDataInterfaceProxyGrid2DCollectionProxy::PreStage(const FNDIGpuComp
 
 		// #todo(dmp): we might want to expose an option where we have buffers that are write only and need a clear (ie: no buffering like the neighbor grid).  They would be considered transient perhaps?  It'd be more
 		// memory efficient since it would theoretically not require any double buffering.
+
+		// #todo(dmp): for now, if there is an output DI that is NOT an iteration DI AND if both the iteration DI and this DI have the same total number of instances, do not do the UAV clear prior to the stage.
+		// this isn't optimal, but should work to some degree to reduce overhead in cases where we have multiple grids of the same resolution being processed/written to in 1 stage
 		if (!Context.IsIterationStage())
 		{
+			FNiagaraDataInterfaceProxyRW* IterationInterface = Context.GetSimStageData().AlternateIterationSource;
+			if (IterationInterface != nullptr)
+			{
+				const FIntVector ElementCount = IterationInterface->GetElementCount(Context.GetSystemInstanceID());				
+
+				if (ElementCount.X == ProxyData->NumCells.X && ElementCount.Y == ProxyData->NumCells.Y)
+				{
+					return;
+				}
+			}
+
 			AddClearUAVPass(GraphBuilder, ProxyData->DestinationData->GetOrCreateUAV(GraphBuilder), FVector4f(ForceInitToZero));
 		}
 	}
