- [Niagara] Reset ActiveWriteFrame at EOF (Contributed by [REDACTED_PERSON])

#rnx
#9697

diff --git a/Engine/Plugins/FX/Niagara/Source/Niagara/Private/NiagaraGPUProfiler.cpp b/Engine/Plugins/FX/Niagara/Source/Niagara/Private/NiagaraGPUProfiler.cpp
--- a/Engine/Plugins/FX/Niagara/Source/Niagara/Private/NiagaraGPUProfiler.cpp
+++ b/Engine/Plugins/FX/Niagara/Source/Niagara/Private/NiagaraGPUProfiler.cpp
@@ -73,6 +73,7 @@ void FNiagaraGPUProfiler::EndFrame(FRHICommandList& RHICmdList)
 	ActiveWriteFrame->EndQuery = QueryPool->AllocateQuery();
 	RHICmdList.EndRenderQuery(ActiveWriteFrame->EndQuery.GetQuery());
 
+	ActiveWriteFrame = nullptr;
 	CurrentWriteFrame = (CurrentWriteFrame + 1) % NumBufferFrames;
 }
 
