- Expose FViewInfos to Niagara Compute Interface
- These are only available while generating the graph
- This cleans up some existing hacks for reaching into data for distance fields

diff --git a/Engine/Plugins/FX/Niagara/Source/Niagara/Private/DataInterface/NiagaraDataInterfaceStaticMesh.cpp b/Engine/Plugins/FX/Niagara/Source/Niagara/Private/DataInterface/NiagaraDataInterfaceStaticMesh.cpp
--- a/Engine/Plugins/FX/Niagara/Source/Niagara/Private/DataInterface/NiagaraDataInterfaceStaticMesh.cpp
+++ b/Engine/Plugins/FX/Niagara/Source/Niagara/Private/DataInterface/NiagaraDataInterfaceStaticMesh.cpp
@@ -2992,7 +2992,13 @@ void UNiagaraDataInterfaceStaticMesh::SetShaderParameters(const FNiagaraDataInte
 	ShaderParameters->InstancePreviousRotation						= InstanceData.PrevRotation;
 	ShaderParameters->InstanceWorldVelocity							= DeltaPosition;
 
-	const FDistanceFieldSceneData* DistanceFieldSceneData = static_cast<const FNiagaraGpuComputeDispatch&>(Context.GetComputeDispatchInterface()).GetMeshDistanceFieldParameters();	//-BATCHERTODO:
+	const FDistanceFieldSceneData* DistanceFieldSceneData = nullptr;
+	TConstArrayView<FViewInfo> SimulationViewInfos = Context.GetComputeDispatchInterface().GetSimulationViewInfos();
+	if (SimulationViewInfos.Num() > 0 && SimulationViewInfos[0].Family && SimulationViewInfos[0].Family->Scene && SimulationViewInfos[0].Family->Scene->GetRenderScene())
+	{
+		DistanceFieldSceneData = &SimulationViewInfos[0].Family->Scene->GetRenderScene()->DistanceFieldSceneData;
+	}
+
 	if (Context.IsParameterBound(&ShaderParameters->InstanceDistanceFieldIndex))
 	{
 		int32 DistanceFieldIndex = -1;
diff --git a/Engine/Plugins/FX/Niagara/Source/Niagara/Private/NiagaraDataInterfaceCollisionQuery.cpp b/Engine/Plugins/FX/Niagara/Source/Niagara/Private/NiagaraDataInterfaceCollisionQuery.cpp
--- a/Engine/Plugins/FX/Niagara/Source/Niagara/Private/NiagaraDataInterfaceCollisionQuery.cpp
+++ b/Engine/Plugins/FX/Niagara/Source/Niagara/Private/NiagaraDataInterfaceCollisionQuery.cpp
@@ -548,8 +548,11 @@ void UNiagaraDataInterfaceCollisionQuery::SetShaderParameters(const FNiagaraData
 	FGlobalDistanceFieldParameters2* ShaderGDFParameters = Context.GetParameterIncludedStruct<FGlobalDistanceFieldParameters2>();
 	if (Context.IsStructBound(ShaderGDFParameters))
 	{
-		const FGlobalDistanceFieldParameterData* GDFParameterData = static_cast<const FNiagaraGpuComputeDispatch&>(Context.GetComputeDispatchInterface()).GetGlobalDistanceFieldParameters();//-BATCHERTODO:
-		FNiagaraDistanceFieldHelper::SetGlobalDistanceFieldParameters(GDFParameterData, *ShaderGDFParameters);
+		TConstArrayView<FViewInfo> SimulationViewInfos = Context.GetComputeDispatchInterface().GetSimulationViewInfos();
+		FNiagaraDistanceFieldHelper::SetGlobalDistanceFieldParameters(
+			SimulationViewInfos.Num() > 0 ? &SimulationViewInfos[0].GlobalDistanceFieldInfo.ParameterData : nullptr,
+			*ShaderGDFParameters
+		);
 	}
 }
 
diff --git a/Engine/Plugins/FX/Niagara/Source/Niagara/Private/NiagaraDataInterfaceRigidMeshCollisionQuery.cpp b/Engine/Plugins/FX/Niagara/Source/Niagara/Private/NiagaraDataInterfaceRigidMeshCollisionQuery.cpp
--- a/Engine/Plugins/FX/Niagara/Source/Niagara/Private/NiagaraDataInterfaceRigidMeshCollisionQuery.cpp
+++ b/Engine/Plugins/FX/Niagara/Source/Niagara/Private/NiagaraDataInterfaceRigidMeshCollisionQuery.cpp
@@ -1509,7 +1509,11 @@ void UNiagaraDataInterfaceRigidMeshCollisionQuery::SetShaderParameters(const FNi
 
 		if (bDistanceFieldDataBound)
 		{
-			DistanceFieldSceneData = static_cast<const FNiagaraGpuComputeDispatch&>(Context.GetComputeDispatchInterface()).GetMeshDistanceFieldParameters();	//-BATCHERTODO:
+			TConstArrayView<FViewInfo> SimulationViewInfos = Context.GetComputeDispatchInterface().GetSimulationViewInfos();
+			if (SimulationViewInfos.Num() > 0 && SimulationViewInfos[0].Family && SimulationViewInfos[0].Family->Scene && SimulationViewInfos[0].Family->Scene->GetRenderScene())
+			{
+				DistanceFieldSceneData = &SimulationViewInfos[0].Family->Scene->GetRenderScene()->DistanceFieldSceneData;
+			}
 		}
 	}
 	else
diff --git a/Engine/Plugins/FX/Niagara/Source/Niagara/Private/NiagaraGpuComputeDispatch.cpp b/Engine/Plugins/FX/Niagara/Source/Niagara/Private/NiagaraGpuComputeDispatch.cpp
--- a/Engine/Plugins/FX/Niagara/Source/Niagara/Private/NiagaraGpuComputeDispatch.cpp
+++ b/Engine/Plugins/FX/Niagara/Source/Niagara/Private/NiagaraGpuComputeDispatch.cpp
@@ -1037,7 +1037,7 @@ void FNiagaraGpuComputeDispatch::ExecuteTicks(FRDGBuilder& GraphBuilder, TConstA
 	FNiagaraSceneTextureParameters SceneTextures;
 	GNiagaraViewDataManager.GetSceneTextureParameters(GraphBuilder, Views.Num() > 0 ? GetViewFamilyInfo(Views).GetSceneTexturesChecked() : nullptr, SceneTextures);
 	NiagaraSceneTextures = &SceneTextures;
-	CurrentPassViews = Views;
+	SimulationViewInfos = Views;
 
 	// Loop over dispatches
 	for ( const FNiagaraGpuDispatchGroup& DispatchGroup : DispatchList.DispatchGroups )
@@ -1351,8 +1351,8 @@ void FNiagaraGpuComputeDispatch::ExecuteTicks(FRDGBuilder& GraphBuilder, TConstA
 	);
 
 	// Tear down for tick pass
-	NiagaraSceneTextures = &SceneTextures;
-	CurrentPassViews = Views;
+	NiagaraSceneTextures = nullptr;
+	SimulationViewInfos = TConstArrayView<FViewInfo>();
 
 	CurrentPassExternalAccessQueue.Submit(GraphBuilder);
 }
@@ -1543,7 +1543,7 @@ void FNiagaraGpuComputeDispatch::DispatchStage(FRDGBuilder& GraphBuilder, const
 	// Set ViewUniformBuffer if it's required
 	if (ComputeShader->bNeedsViewUniformBuffer)
 	{
-		DispatchParameters->View = CurrentPassViews[0].ViewUniformBuffer;
+		DispatchParameters->View = SimulationViewInfos[0].ViewUniformBuffer;
 	}
 	DispatchParameters->SceneTextures.SceneTextures			= NiagaraSceneTextures->SceneTextures;
 	DispatchParameters->SceneTextures.MobileSceneTextures	= NiagaraSceneTextures->MobileSceneTextures;
@@ -1758,7 +1758,6 @@ void FNiagaraGpuComputeDispatch::PostRenderOpaque(FRDGBuilder& GraphBuilder, TCo
 			[this, Views](FRHICommandListImmediate& RHICmdList)
 			{
 				AsyncGpuTraceHelper->PostRenderOpaque(RHICmdList, this, Views);
-				CurrentPassViews = Views;
 			}
 		);
 
@@ -1776,7 +1775,6 @@ void FNiagaraGpuComputeDispatch::PostRenderOpaque(FRDGBuilder& GraphBuilder, TCo
 
 				AsyncGpuTraceHelper->EndFrame(RHICmdList, this);
 			}
-			CurrentPassViews = TConstArrayView<FViewInfo>();
 
 			ProcessDebugReadbacks(RHICmdList, false);
 
@@ -1909,22 +1907,6 @@ bool FNiagaraGpuComputeDispatch::AddSortedGPUSimulation(FNiagaraGPUSortInfo& Sor
 	}
 }
 
-const FGlobalDistanceFieldParameterData* FNiagaraGpuComputeDispatch::GetGlobalDistanceFieldParameters() const
-{ 
-	check(CurrentPassViews.Num() > 0); 
-	return &CurrentPassViews[0].GlobalDistanceFieldInfo.ParameterData; 
-}
-
-const FDistanceFieldSceneData* FNiagaraGpuComputeDispatch::GetMeshDistanceFieldParameters() const
-{
-	if (CurrentPassViews.Num() == 0 || CurrentPassViews[0].Family == nullptr || CurrentPassViews[0].Family->Scene == nullptr || CurrentPassViews[0].Family->Scene->GetRenderScene() == nullptr)
-	{
-		return nullptr;
-	}
-
-	return &CurrentPassViews[0].Family->Scene->GetRenderScene()->DistanceFieldSceneData;
-}
-
 void FNiagaraGpuComputeDispatch::GenerateSortKeys(FRHICommandListImmediate& RHICmdList, int32 BatchId, int32 NumElementsInBatch, EGPUSortFlags Flags, FRHIUnorderedAccessView* KeysUAV, FRHIUnorderedAccessView* ValuesUAV)
 {
 	const bool bHighPrecision = EnumHasAnyFlags(Flags, EGPUSortFlags::HighPrecisionKeys);
diff --git a/Engine/Plugins/FX/Niagara/Source/Niagara/Private/NiagaraGpuComputeDispatch.h b/Engine/Plugins/FX/Niagara/Source/Niagara/Private/NiagaraGpuComputeDispatch.h
--- a/Engine/Plugins/FX/Niagara/Source/Niagara/Private/NiagaraGpuComputeDispatch.h
+++ b/Engine/Plugins/FX/Niagara/Source/Niagara/Private/NiagaraGpuComputeDispatch.h
@@ -147,9 +147,6 @@ public:
 
 	virtual bool AddSortedGPUSimulation(FNiagaraGPUSortInfo& SortInfo) override;
 
-	const FGlobalDistanceFieldParameterData* GetGlobalDistanceFieldParameters() const;
-	const FDistanceFieldSceneData* GetMeshDistanceFieldParameters() const;
-
 	void LegacySetDataInterfaceParameters(FRHICommandList& RHICmdList, const FNiagaraGPUSystemTick& Tick, const FNiagaraComputeInstanceData& InstanceData, const FNiagaraShaderRef& ComputeShader, const FNiagaraSimStageData& SimStageData) const;
 	void LegacyUnsetDataInterfaceParameters(FRHICommandList& RHICmdList, const FNiagaraGPUSystemTick& Tick, const FNiagaraComputeInstanceData& InstanceData, const FNiagaraShaderRef& ComputeShader, const FNiagaraSimStageData& SimStageData) const;
 	void LegacyResetDataInterfaces(FRHICommandList& RHICmdList, const FNiagaraGPUSystemTick& Tick, const FNiagaraComputeInstanceData& InstanceData) const;
@@ -236,7 +233,6 @@ private:
 
 	uint32 FramesBeforeTickFlush = 0;
 
-	TConstArrayView<FViewInfo> CurrentPassViews;
 	mutable FRDGExternalAccessQueue CurrentPassExternalAccessQueue;
 
 	/** A buffer of list sizes used by UpdateFreeIDBuffers to allow overlapping several dispatches. */
diff --git a/Engine/Plugins/FX/Niagara/Source/Niagara/Public/NiagaraGpuComputeDispatchInterface.h b/Engine/Plugins/FX/Niagara/Source/Niagara/Public/NiagaraGpuComputeDispatchInterface.h
--- a/Engine/Plugins/FX/Niagara/Source/Niagara/Public/NiagaraGpuComputeDispatchInterface.h
+++ b/Engine/Plugins/FX/Niagara/Source/Niagara/Public/NiagaraGpuComputeDispatchInterface.h
@@ -10,6 +10,7 @@
 
 class RENDERCORE_API FRDGBuilder;
 class FRDGExternalAccessQueue;
+class FViewInfo;
 
 class FNiagaraAsyncGpuTraceHelper;
 struct FNiagaraComputeExecutionContext;
@@ -72,6 +73,12 @@ public:
 		return *Manager;
 	}
 
+	/**
+	Get access to the Views the simulation is being rendered with.
+	List is only valid during graph building (i.e. during ExecuteTicks) and for simulations in PostInitViews / PostRenderOpaque.
+	*/
+	TConstArrayView<FViewInfo> GetSimulationViewInfos() const { return SimulationViewInfos; }
+
 	/** Get access to the instance count manager. */
 	FORCEINLINE FNiagaraGPUInstanceCountManager& GetGPUInstanceCounterManager() { check(IsInRenderingThread()); return GPUInstanceCounterManager; }
 	FORCEINLINE const FNiagaraGPUInstanceCountManager& GetGPUInstanceCounterManager() const { check(IsInRenderingThread()); return GPUInstanceCounterManager; }
@@ -171,4 +178,6 @@ protected:
 	FNiagaraGPUInstanceCountManager			GPUInstanceCounterManager;
 
 	TArray<TPair<FName, TUniquePtr<FNiagaraGpuComputeDataManager>>> GpuDataManagers;
+
+	TConstArrayView<FViewInfo>				SimulationViewInfos;
 };
