Removed lots of includes in UE based on compile profiling

diff --git a/Engine/Plugins/Experimental/CommonUI/Source/CommonInput/Public/CommonInputSubsystem.h b/Engine/Plugins/Experimental/CommonUI/Source/CommonInput/Public/CommonInputSubsystem.h
--- a/Engine/Plugins/Experimental/CommonUI/Source/CommonInput/Public/CommonInputSubsystem.h
+++ b/Engine/Plugins/Experimental/CommonUI/Source/CommonInput/Public/CommonInputSubsystem.h
@@ -4,14 +4,20 @@
 
 #include "Subsystems/LocalPlayerSubsystem.h"
 #include "CommonInputBaseTypes.h"
+#include "Containers/Ticker.h"
+
+#if UE_ENABLE_INCLUDE_ORDER_DEPRECATED_IN_5_2
 #include "Framework/Application/SlateApplication.h"
 #include "Styling/SlateTypes.h"
+#endif
+
 #include "CommonInputSubsystem.generated.h"
 
 class UWidget;
 class ULocalPlayer;
 class APlayerController;
 class FCommonInputPreprocessor;
+class FSlateUser;
 class UCommonInputActionDomainTable;
 
 DECLARE_DYNAMIC_MULTICAST_DELEGATE_OneParam(FInputMethodChangedDelegate, ECommonInputType, bNewInputType);
diff --git a/Engine/Plugins/Experimental/CommonUI/Source/CommonUI/Public/Input/CommonAnalogCursor.h b/Engine/Plugins/Experimental/CommonUI/Source/CommonUI/Public/Input/CommonAnalogCursor.h
--- a/Engine/Plugins/Experimental/CommonUI/Source/CommonUI/Public/Input/CommonAnalogCursor.h
+++ b/Engine/Plugins/Experimental/CommonUI/Source/CommonUI/Public/Input/CommonAnalogCursor.h
@@ -14,7 +14,7 @@ class UGameViewportClient;
 
 struct FInputEvent;
 enum class ECommonInputType : uint8;
-enum EOrientation;
+enum EOrientation : int;
 
 /**
  * Analog cursor preprocessor that tastefully hijacks things a bit to support controller navigation by moving a hidden cursor around based on focus.
diff --git a/Engine/Plugins/Experimental/Water/Source/Runtime/Private/WaterBodyComponent.cpp b/Engine/Plugins/Experimental/Water/Source/Runtime/Private/WaterBodyComponent.cpp
--- a/Engine/Plugins/Experimental/Water/Source/Runtime/Private/WaterBodyComponent.cpp
+++ b/Engine/Plugins/Experimental/Water/Source/Runtime/Private/WaterBodyComponent.cpp
@@ -28,6 +28,7 @@
 #include "WaterRuntimeSettings.h"
 #include "WaterUtils.h"
 #include "Engine/TextureRenderTarget2D.h"
+#include "Engine/Texture2D.h"
 #include "GerstnerWaterWaves.h"
 #include "WaterMeshComponent.h"
 #include "WaterVersion.h"
@@ -71,6 +72,22 @@ TAutoConsoleVariable<float> CVarWaterOceanFallbackDepth(
 	TEXT("Depth to report for the ocean when no terrain is found under the query location. Not used when <= 0."),
 	ECVF_Default);
 
+
+FWaterBodyMeshSection::FWaterBodyMeshSection()
+	: Vertices()
+	, Indices()
+	, SectionBounds()
+{
+}
+
+FWaterBodyMeshSection::~FWaterBodyMeshSection() = default;
+
+uint32 FWaterBodyMeshSection::GetAllocatedSize() const
+{
+	return Vertices.GetAllocatedSize() + Indices.GetAllocatedSize();
+}
+
+
 const FName UWaterBodyComponent::WaterBodyIndexParamName(TEXT("WaterBodyIndex"));
 const FName UWaterBodyComponent::WaterBodyZOffsetParamName(TEXT("WaterBodyZOffset"));
 const FName UWaterBodyComponent::WaterVelocityAndHeightName(TEXT("WaterVelocityAndHeight"));
diff --git a/Engine/Plugins/Experimental/Water/Source/Runtime/Public/WaterBodyComponent.h b/Engine/Plugins/Experimental/Water/Source/Runtime/Public/WaterBodyComponent.h
--- a/Engine/Plugins/Experimental/Water/Source/Runtime/Public/WaterBodyComponent.h
+++ b/Engine/Plugins/Experimental/Water/Source/Runtime/Public/WaterBodyComponent.h
@@ -16,7 +16,10 @@
 #include "Interfaces/Interface_PostProcessVolume.h"
 #include "WaterWaves.h"
 #include "WaterBodyTypes.h"
+
+#if UE_ENABLE_INCLUDE_ORDER_DEPRECATED_IN_5_2
 #include "DynamicMeshBuilder.h"
+#endif
 
 #if WITH_EDITOR
 #include "MeshDescription.h"
@@ -33,6 +36,7 @@ class ALandscapeProxy;
 class UMaterialInstanceDynamic;
 class FTokenizedMessage;
 namespace UE::Geometry { class FDynamicMesh3; }
+struct FDynamicMeshVertex;
 
 // ----------------------------------------------------------------------------------
 
@@ -88,16 +92,12 @@ enum class EWaterBodyStatus : uint8
 struct FWaterBodyMeshSection
 {
 public:
-	FWaterBodyMeshSection()
-		: Vertices()
-		, Indices()
-		, SectionBounds()
-	{
-	}
+	FWaterBodyMeshSection();
+	~FWaterBodyMeshSection();
 
 	bool IsValid() const { return (bool)SectionBounds.bIsValid; }
 
-	uint32 GetAllocatedSize() const { return Vertices.GetAllocatedSize() + Indices.GetAllocatedSize(); }
+	uint32 GetAllocatedSize() const;
 
 	TArray<FDynamicMeshVertex> Vertices;
 	TArray<uint32> Indices;
diff --git a/Engine/Plugins/Experimental/Water/Source/Runtime/Public/WaterBrushActorInterface.h b/Engine/Plugins/Experimental/Water/Source/Runtime/Public/WaterBrushActorInterface.h
--- a/Engine/Plugins/Experimental/Water/Source/Runtime/Public/WaterBrushActorInterface.h
+++ b/Engine/Plugins/Experimental/Water/Source/Runtime/Public/WaterBrushActorInterface.h
@@ -5,13 +5,16 @@
 #include "CoreMinimal.h"
 #include "CoreGlobals.h"
 #include "UObject/Interface.h"
+#if UE_ENABLE_INCLUDE_ORDER_DEPRECATED_IN_5_2
 #include "Engine/TextureRenderTarget2D.h"
+#endif
 #include "WaterBrushActorInterface.generated.h"
 
 struct FWaterCurveSettings;
 struct FWaterBodyHeightmapSettings;
 struct FWaterBodyWeightmapSettings;
 class UPrimitiveComponent;
+enum ETextureRenderTargetFormat : int;
 
 /** Dummy class needed to support Cast<IWaterBrushActorInterface>(Object). */
 UINTERFACE()
diff --git a/Engine/Plugins/FX/Niagara/Source/Niagara/Classes/NiagaraBakerSettings.h b/Engine/Plugins/FX/Niagara/Source/Niagara/Classes/NiagaraBakerSettings.h
--- a/Engine/Plugins/FX/Niagara/Source/Niagara/Classes/NiagaraBakerSettings.h
+++ b/Engine/Plugins/FX/Niagara/Source/Niagara/Classes/NiagaraBakerSettings.h
@@ -4,7 +4,9 @@
 
 #include "CoreMinimal.h"
 #include "Camera/CameraTypes.h"
+#if UE_ENABLE_INCLUDE_ORDER_DEPRECATED_IN_5_2
 #include "Engine/TextureRenderTarget2D.h"
+#endif
 #include "Misc/PathViews.h"
 #include "UObject/ObjectMacros.h"
 #include "UObject/Object.h"
diff --git a/Engine/Plugins/FX/Niagara/Source/Niagara/Private/NiagaraComponent.cpp b/Engine/Plugins/FX/Niagara/Source/Niagara/Private/NiagaraComponent.cpp
--- a/Engine/Plugins/FX/Niagara/Source/Niagara/Private/NiagaraComponent.cpp
+++ b/Engine/Plugins/FX/Niagara/Source/Niagara/Private/NiagaraComponent.cpp
@@ -2,6 +2,7 @@
 
 #include "NiagaraComponent.h"
 #include "Engine/CollisionProfile.h"
+#include "Engine/TextureRenderTarget.h"
 #include "EngineUtils.h"
 #include "Materials/MaterialInstanceDynamic.h"
 #include "NiagaraCommon.h"
@@ -12,8 +13,10 @@
 #include "NiagaraFunctionLibrary.h"
 #include "NiagaraGpuComputeDispatchInterface.h"
 #include "NiagaraRenderer.h"
+#include "NiagaraSceneProxy.h"
 #include "NiagaraStats.h"
 #include "NiagaraSystem.h"
+#include "NiagaraSystemInstanceController.h"
 #include "NiagaraWorldManager.h"
 #include "PrimitiveSceneInfo.h"
 #include "ProfilingDebugging/CsvProfiler.h"
@@ -287,6 +290,11 @@ void FNiagaraSceneProxy::ReleaseUniformBuffers(bool bEmpty)
 	CustomUniformBuffers.Empty(bEmpty ? 0 : FMath::Min(ExpectedRendererCount, CurrentUBCount));
 }
 
+const FVector3f& FNiagaraSceneProxy::GetLWCRenderTile() const
+{
+	return RenderData ? RenderData->LWCRenderTile : FVector3f::ZeroVector;
+}
+
 TUniformBuffer<FPrimitiveUniformShaderParameters>* FNiagaraSceneProxy::GetCustomUniformBufferResource(bool bHasVelocity, const FBox& InstanceBounds) const
 {
 	// Use a hash to determine if we can re-use any uniform buffer
@@ -546,6 +554,14 @@ void UNiagaraComponent::InitForPerformanceBaseline()
 	SetPreviewLODDistance(true, 1.0f, 10000.0f);
 }
 
+void UNiagaraComponent::SetLODDistance(float InLODDistance, float InMaxLODDistance)
+{
+	if (!bEnablePreviewLODDistance && SystemInstanceController)
+	{
+		SystemInstanceController->SetLODDistance(InLODDistance, InMaxLODDistance, false);
+	}
+}
+
 void UNiagaraComponent::SetEmitterEnable(FName EmitterName, bool bNewEnableState)
 {
 	if (SystemInstanceController.IsValid() && !SystemInstanceController->IsComplete())
@@ -911,6 +927,11 @@ UNiagaraDataInterface* UNiagaraComponent::GetDataInterface(const FString& Name)
 	return UNiagaraFunctionLibrary::GetDataInterface(UNiagaraDataInterface::StaticClass(), this, *Name);
 }
 
+void UNiagaraComponent::SetSystemSignificanceIndex(int32 InIndex)
+{
+	if(SystemInstanceController) SystemInstanceController->SetSystemSignificanceIndex(InIndex);
+}
+
 bool UNiagaraComponent::IsWorldReadyToRun() const
 {
 	// The Niagara system instance assumes that a ComputeDispatchInterface exists when it is created. We need to wait until this has happened before successfully activating this system.
@@ -1296,6 +1317,21 @@ void UNiagaraComponent::DeactivateImmediate()
 	DeactivateImmediateInternal(false);
 }
 
+ENiagaraExecutionState UNiagaraComponent::GetRequestedExecutionState() const
+{
+	return SystemInstanceController ? SystemInstanceController->GetRequestedExecutionState() : ENiagaraExecutionState::Complete;
+}
+
+ENiagaraExecutionState UNiagaraComponent::GetExecutionState() const
+{
+	return SystemInstanceController ? SystemInstanceController->GetActualExecutionState() : ENiagaraExecutionState::Complete;
+}
+
+bool UNiagaraComponent::IsComplete() const
+{
+	return SystemInstanceController ? SystemInstanceController->IsComplete() : true;
+}
+
 void UNiagaraComponent::DeactivateImmediateInternal(bool bIsScalabilityCull)
 {
 	SCOPE_CYCLE_COUNTER(STAT_NiagaraComponentDeactivate);
@@ -1852,6 +1888,11 @@ void UNiagaraComponent::DestroyCullProxy()
 	}
 }
 
+FParticlePerfStatsContext UNiagaraComponent::GetPerfStatsContext()
+{
+	return FParticlePerfStatsContext(GetWorld(), Asset, this);
+}
+
 TSharedPtr<FNiagaraSystemSimulation, ESPMode::ThreadSafe> UNiagaraComponent::GetSystemSimulation()
 {
 	if (SystemInstanceController)
@@ -3528,6 +3569,11 @@ void UNiagaraComponent::SetAutoDestroy(bool bInAutoDestroy)
 	}
 }
 
+FNiagaraSystemInstance* UNiagaraComponent::GetSystemInstance() const
+{
+	return SystemInstanceController ? SystemInstanceController->GetSystemInstance_Unsafe() : nullptr;
+}
+
 void UNiagaraComponent::SetPreviewLODDistance(bool bInEnablePreviewLODDistance, float InPreviewLODDistance, float InPreviewMaxDistance)
 {
 	bEnablePreviewLODDistance = bInEnablePreviewLODDistance;
diff --git a/Engine/Plugins/FX/Niagara/Source/Niagara/Private/NiagaraParameterStore.cpp b/Engine/Plugins/FX/Niagara/Source/Niagara/Private/NiagaraParameterStore.cpp
--- a/Engine/Plugins/FX/Niagara/Source/Niagara/Private/NiagaraParameterStore.cpp
+++ b/Engine/Plugins/FX/Niagara/Source/Niagara/Private/NiagaraParameterStore.cpp
@@ -1,9 +1,11 @@
 // Copyright Epic Games, Inc. All Rights Reserved.
 
 #include "NiagaraParameterStore.h"
+#include "Algo/Find.h"
 #include "NiagaraDataSet.h"
 #include "NiagaraComponent.h"
 #include "NiagaraCustomVersion.h"
+#include "NiagaraDataInterface.h"
 #include "NiagaraSystemInstance.h"
 #include "NiagaraStats.h"
 
@@ -319,6 +321,37 @@ void FNiagaraParameterStoreBinding::MatchParameters(FNiagaraParameterStore* Dest
 	}
 }
 
+void FNiagaraParameterStoreBinding::Tick(FNiagaraParameterStore* DestStore, FNiagaraParameterStore* SrcStore, bool bForce)
+{
+	if (SrcStore->GetParametersDirty() || bForce)
+	{
+		for (FParameterBinding& Binding : ParameterBindings)
+		{
+			DestStore->SetParameterData(SrcStore->GetParameterData(Binding.SrcOffset), Binding.DestOffset, Binding.Size);
+		}
+	}
+
+	if (SrcStore->GetInterfacesDirty() || bForce)
+	{
+		for (FInterfaceBinding& Binding : InterfaceBindings)
+		{
+			DestStore->SetDataInterface(SrcStore->GetDataInterface(Binding.SrcOffset), Binding.DestOffset);
+		}
+	}
+
+	if (SrcStore->GetUObjectsDirty() || bForce)
+	{
+		for (FUObjectBinding& Binding : UObjectBindings)
+		{
+			DestStore->SetUObject(SrcStore->GetUObject(Binding.SrcOffset), Binding.DestOffset);
+		}
+	}
+
+#if NIAGARA_NAN_CHECKING
+	DestStore->CheckForNaNs();
+#endif
+}
+
 void FNiagaraParameterStoreBinding::GetBindingData(FNiagaraParameterStore* DestStore, FNiagaraParameterStore* SrcStore, FNiagaraBoundParameterArray& OutBoundParameters)
 {
 	OutBoundParameters.Empty();
@@ -533,6 +566,22 @@ void FNiagaraParameterStore::TickBindings()
 	Dump();
 }
 
+void FNiagaraParameterStore::Tick()
+{
+#if NIAGARA_NAN_CHECKING
+	CheckForNaNs();
+#endif
+	if (Bindings.Num() > 0 && (bParametersDirty || bInterfacesDirty || bUObjectsDirty))
+	{
+		TickBindings();
+	}
+
+	//We have to have ticked all our source stores before now.
+	bParametersDirty = false;
+	bInterfacesDirty = false;
+	bUObjectsDirty = false;
+}
+
 void FNiagaraParameterStore::UnbindFromSourceStores()
 {
 	//UE_LOG(LogNiagara, Log, TEXT("FNiagaraParameterStore::UnUnbindFromSourceStoresbind() - Src: 0x%p - SrcName: %s"), this, *DebugName);
@@ -1256,6 +1305,48 @@ const FNiagaraVariableBase* FNiagaraParameterStore::FindVariable(const UNiagaraD
 	return nullptr;
 }
 
+void FNiagaraParameterStore::CopyParameterData(FNiagaraParameterStore& DestStore, const FNiagaraVariable& Parameter) const
+{
+	int32 DestIndex = DestStore.IndexOf(Parameter);
+	int32 SrcIndex = IndexOf(Parameter);
+	if (DestIndex != INDEX_NONE && SrcIndex != INDEX_NONE)
+	{
+		if (Parameter.IsDataInterface())
+		{
+			ensure(DestStore.DataInterfaces.IsValidIndex(DestIndex));
+			UNiagaraDataInterface* DestDataInterface = DestStore.DataInterfaces[DestIndex];
+			if (DestDataInterface == nullptr)
+			{
+				if (ensureMsgf(DestStore.Owner != nullptr, TEXT("Destination data interface pointer was null and a new one couldn't be created because the destination store's owner pointer was also null.")))
+				{
+					UE_LOG(LogNiagara, Warning, TEXT("While trying to copy parameter data the destination data interface was null, creating a new one.  Parameter: %s Destination Store Owner: %s"), *Parameter.GetName().ToString(), *GetPathNameSafe(DestStore.Owner.Get()));
+					DestDataInterface = NewObject<UNiagaraDataInterface>(DestStore.Owner.Get(), Parameter.GetType().GetClass(), NAME_None, RF_Transactional | RF_Public);
+					DestStore.DataInterfaces[DestIndex] = DestDataInterface;
+				}
+				else
+				{
+					return;
+				}
+			}
+			DataInterfaces[SrcIndex]->CopyTo(DestDataInterface);
+			DestStore.OnInterfaceChange();
+		}
+		else if (Parameter.IsUObject())
+		{
+			DestStore.SetUObject(GetUObject(SrcIndex), DestIndex);
+		}
+		else if (Parameter.GetType() == FNiagaraTypeDefinition::GetPositionDef())
+		{
+			const FVector* SourceVector = GetPositionParameterValue(Parameter.GetName());
+			DestStore.SetPositionParameterValue(SourceVector ? *SourceVector : FVector::ZeroVector, Parameter.GetName());
+		}
+		else
+		{
+			DestStore.SetParameterData(GetParameterData(SrcIndex), DestIndex, Parameter.GetSizeInBytes());
+		}
+	}
+}
+
 bool FNiagaraParameterStore::CopyParameterData(const FNiagaraVariable& Parameter, uint8* DestinationData) const
 {
 	if (const FNiagaraVariableWithOffset* NiagaraVariableWithOffset = FindParameterVariable(Parameter))
@@ -1274,6 +1365,14 @@ bool FNiagaraParameterStore::CopyParameterData(const FNiagaraVariable& Parameter
 	return false;
 }
 
+UNiagaraDataInterface* FNiagaraParameterStore::GetDataInterface(const FNiagaraVariable& Parameter)const
+{
+	int32 Offset = IndexOf(Parameter);
+	UNiagaraDataInterface* Interface = GetDataInterface(Offset);
+	checkSlow(!Interface || Parameter.GetType().GetClass() == Interface->GetClass());
+	return Interface;
+}
+
 FNiagaraLwcStructConverter FNiagaraParameterStore::GetStructConverter(const FNiagaraVariable& Parameter) const
 {
 	auto ParameterVariables = ReadParameterVariables();
diff --git a/Engine/Plugins/FX/Niagara/Source/Niagara/Public/NiagaraCommon.h b/Engine/Plugins/FX/Niagara/Source/Niagara/Public/NiagaraCommon.h
--- a/Engine/Plugins/FX/Niagara/Source/Niagara/Public/NiagaraCommon.h
+++ b/Engine/Plugins/FX/Niagara/Source/Niagara/Public/NiagaraCommon.h
@@ -4,13 +4,21 @@
 
 #include "CoreMinimal.h"
 #include "Engine/EngineBaseTypes.h"
+#if UE_ENABLE_INCLUDE_ORDER_DEPRECATED_IN_5_2
 #include "Engine/TextureRenderTarget2D.h"
+#endif
 #include "Engine/BlueprintGeneratedClass.h"
 #include "Particles/ParticlePerfStats.h"
+#include "RHIDefinitions.h"
 #include "UObject/ObjectKey.h"
 #include "UObject/ObjectMacros.h"
 #include "UObject/WeakFieldPtr.h"
 
+#include "NiagaraCompileHashVisitor.h"
+#include "NiagaraComponentPoolMethodEnum.h"
+#include "NiagaraDefines.h"
+#include "NiagaraScalabilityState.h"
+#include "NiagaraTickBehaviorEnum.h"
 #include "NiagaraTypes.h"
 #include "NiagaraCore.h"
 #include "NiagaraCommon.generated.h"
@@ -25,9 +33,8 @@ class FNiagaraSystemInstance;
 class UNiagaraParameterCollection;
 class UNiagaraParameterDefinitionsBase;
 struct FNiagaraParameterStore;
-
-//#define NIAGARA_NAN_CHECKING 1
-#define NIAGARA_NAN_CHECKING 0
+enum class EPSCPoolMethod : uint8;
+enum ETextureRenderTargetFormat : int;
 
 #define NIAGARA_MEMORY_TRACKING	!UE_BUILD_SHIPPING
 
@@ -35,10 +42,6 @@ struct FNiagaraParameterStore;
 	#define NIAGARA_COMPUTEDEBUG_ENABLED WITH_EDITOR
 #endif
 
-#ifndef WITH_NIAGARA_DEBUGGER
-	#define WITH_NIAGARA_DEBUGGER !UE_BUILD_SHIPPING
-#endif
-
 #define WITH_NIAGARA_GPU_PROFILER_EDITOR (WITH_EDITOR && STATS)
 #define WITH_NIAGARA_GPU_PROFILER ((WITH_PARTICLE_PERF_STATS && !UE_BUILD_SHIPPING) || WITH_NIAGARA_GPU_PROFILER_EDITOR)
 
@@ -58,20 +61,6 @@ constexpr ETickingGroup NiagaraFirstTickGroup = TG_PrePhysics;
 constexpr ETickingGroup NiagaraLastTickGroup = TG_LastDemotable;
 constexpr int NiagaraNumTickGroups = NiagaraLastTickGroup - NiagaraFirstTickGroup + 1;
 
-/** Niagara ticking behaviour */
-UENUM()
-enum class ENiagaraTickBehavior : uint8
-{
-	/** Niagara will tick after all prereqs have ticked for attachements / data interfaces, this is the safest option. */
-	UsePrereqs,
-	/** Niagara will ignore prereqs (attachments / data interface dependencies) and use the tick group set on the component. */
-	UseComponentTickGroup,
-	/** Niagara will tick in the first tick group (default is TG_PrePhysics). */
-	ForceTickFirst,
-	/** Niagara will tick in the last tick group (default is TG_LastDemotable). */
-	ForceTickLast,
-};
-
 enum ENiagaraBaseTypes
 {
 	NBT_Half,
@@ -1463,101 +1452,6 @@ public:
 };
 
 
-USTRUCT()
-struct FNiagaraScalabilityState
-{
-	GENERATED_BODY()
-
-	FNiagaraScalabilityState()
-		: Significance(1.0f)
-		, LastVisibleTime(0.0f)
-		, SystemDataIndex(INDEX_NONE)
-		, bCulled(0)
-		, bPreviousCulled(0)
-		, bCulledByDistance(0)
-		, bCulledByInstanceCount(0)
-		, bCulledByVisibility(0)
-		, bCulledByGlobalBudget(0)
-	{
-	}
-
-	FNiagaraScalabilityState(float InSignificance, bool InCulled, bool InPreviousCulled)
-		: Significance(InSignificance)
-		, LastVisibleTime(0.0f)
-		, SystemDataIndex(INDEX_NONE)
-		, bCulled(InCulled)
-		, bPreviousCulled(InPreviousCulled)
-		, bCulledByDistance(0)
-		, bCulledByInstanceCount(0)
-		, bCulledByVisibility(0)
-		, bCulledByGlobalBudget(0)
-	{
-	}
-
-	bool IsDirty() const { return bCulled != bPreviousCulled; }
-	void Apply() { bPreviousCulled = bCulled; }
-
-	UPROPERTY(VisibleAnywhere, Category="Scalability")
-	float Significance;
-
-	UPROPERTY(VisibleAnywhere, Category = "Scalability")
-	float LastVisibleTime;
-
-	int16 SystemDataIndex;
-
-	UPROPERTY(VisibleAnywhere, Category = "Scalability")
-	uint8 bCulled : 1;
-
-	UPROPERTY(VisibleAnywhere, Category="Scalability")
-	uint8 bPreviousCulled : 1;
-
-	UPROPERTY(VisibleAnywhere, Category="Scalability")
-	uint8 bCulledByDistance : 1;
-
-	UPROPERTY(VisibleAnywhere, Category = "Scalability")
-	uint8 bCulledByInstanceCount : 1;
-
-	UPROPERTY(VisibleAnywhere, Category = "Scalability")
-	uint8 bCulledByVisibility : 1;
-	
-	UPROPERTY(VisibleAnywhere, Category = "Scalability")
-	uint8 bCulledByGlobalBudget : 1;
-};
-
-
-UENUM()
-enum class ENCPoolMethod : uint8
-{
-	/**
-	* The component will be created fresh and not allocated from the pool.
-	*/
-	None,
-
-	/**
-	* The component is allocated from the pool and will be automatically released back to it.
-	* User need not handle this any more that other NCs but interaction with the NC after the tick it's spawned in are unsafe.
-	* This method is useful for one-shot fx that you don't need to keep a reference to and can fire and forget.
-	*/
-	AutoRelease,
-
-	/**
-	* NC is allocated from the pool but will NOT be automatically released back to it. User has ownership of the NC and must call ReleaseToPool when finished with it otherwise the NC will leak.
-	* Interaction with the NC after it has been released are unsafe.
-	* This method is useful for persistent FX that you need to modify parameters upon etc over it's lifetime.
-	*/
-	ManualRelease,
-
-	/**
-	Special entry allowing manual release NCs to be manually released but wait until completion to be returned to the pool.
-	*/
-	ManualRelease_OnComplete UMETA(Hidden),
-
-	/**
-	Special entry that marks a NC as having been returned to the pool. All NCs currently in the pool are marked this way.
-	*/
-	FreeInPool UMETA(Hidden),
-};
-
 extern NIAGARA_API ENCPoolMethod ToNiagaraPooling(EPSCPoolMethod PoolingMethod);
 extern NIAGARA_API EPSCPoolMethod ToPSCPoolMethod(ENCPoolMethod PoolingMethod);
 
diff --git a/Engine/Plugins/FX/Niagara/Source/Niagara/Public/NiagaraCompileHashVisitor.h b/Engine/Plugins/FX/Niagara/Source/Niagara/Public/NiagaraCompileHashVisitor.h
--- a/Engine/Plugins/FX/Niagara/Source/Niagara/Public/NiagaraCompileHashVisitor.h
+++ b/Engine/Plugins/FX/Niagara/Source/Niagara/Public/NiagaraCompileHashVisitor.h
@@ -2,4 +2,161 @@
 
 #pragma once
 
-// Placeholder for upcoming change
\ No newline at end of file
+#include "ShaderParameterMacros.h"
+#include "NiagaraCompileHashVisitor.generated.h"
+
+USTRUCT()
+struct NIAGARA_API FNiagaraCompileHashVisitorDebugInfo
+{
+	GENERATED_USTRUCT_BODY()
+public:
+	UPROPERTY()
+	FString Object;
+
+	UPROPERTY()
+	TArray<FString> PropertyKeys;
+
+	UPROPERTY()
+	TArray<FString> PropertyValues;
+};
+
+/**
+Used to store the state of a graph when deciding if it has been dirtied for recompile.
+*/
+struct NIAGARA_API FNiagaraCompileHashVisitor
+{
+public:
+	FNiagaraCompileHashVisitor(FSHA1& InHashState) : HashState(InHashState) {}
+
+	FSHA1& HashState;
+	TArray<const void*> ObjectList;
+
+	static int LogCompileIdGeneration;
+
+#if WITH_EDITORONLY_DATA
+	FNiagaraCompileHashVisitorDebugInfo* AddDebugInfo()
+	{
+		if (LogCompileIdGeneration != 0)
+		{
+			return &Values.AddDefaulted_GetRef();
+		}
+		return nullptr;
+	}
+
+	// Debug data about the compilation hash, including key value pairs to detect differences.
+	TArray<FNiagaraCompileHashVisitorDebugInfo> Values;
+
+	template<typename T>
+	void ToDebugString(const T* InData, uint64 InDataCount, FString& OutStr)
+	{
+		for (int32 i = 0; i < InDataCount; i++)
+		{
+			FString DataStr = LexToString(InData[i]);
+			OutStr.Appendf(TEXT("%s "), *DataStr);
+		}
+	}
+#endif
+	/**
+	Registers a pointer for later reference in the compile id in a deterministic manner.
+	*/
+	int32 RegisterReference(const void* InObject)
+	{
+		if (InObject == nullptr)
+		{
+			return -1;
+		}
+
+		int32 Index = ObjectList.Find(InObject);
+		if (Index < 0)
+		{
+			Index = ObjectList.Add(InObject);
+		}
+		return Index;
+	}
+
+	/**
+	We don't usually want to save GUID's or pointer values because they have nondeterministic values. Consider a PostLoad upgrade operation that creates a new node.
+	Each pin and node gets a unique ID. If you close the editor and reopen, you'll get a different set of values. One of the characteristics we want for compilation
+	behavior is that the same graph structure produces the same compile results, so we only want to embed information that is deterministic. This method is for use
+	when registering a pointer to an object that is serialized within the compile hash.
+	*/
+	bool UpdateReference(const TCHAR* InDebugName, const void* InObject)
+	{
+		int32 Index = RegisterReference(InObject);
+		return UpdatePOD(InDebugName, Index);
+	}
+
+	/**
+	Adds an array of POD (plain old data) values to the hash.
+	*/
+	template<typename T>
+	bool UpdateArray(const TCHAR* InDebugName, const T* InData, uint64 InDataCount = 1)
+	{
+		static_assert(TIsPODType<T>::Value, "UpdateArray does not support a constructor / destructor on the element class.");
+		HashState.Update((const uint8 *)InData, sizeof(T)*InDataCount);
+#if WITH_EDITORONLY_DATA
+		if (LogCompileIdGeneration != 0)
+		{
+			FString ValuesStr = InDebugName;
+			ValuesStr.Append(TEXT(" = "));
+			ToDebugString(InData, InDataCount, ValuesStr);
+			Values.Top().PropertyKeys.Push(InDebugName);
+			Values.Top().PropertyValues.Push(ValuesStr);
+		}
+#endif
+		return true;
+	}
+
+	/**
+	 * Adds the string representation of an FName to the hash
+	 */
+	bool UpdateName(const TCHAR* InDebugName, FName InName)
+	{
+		FNameBuilder NameString(InName);
+		return UpdateString(InDebugName, NameString.ToView());
+	}
+
+	/**
+	Adds a single value of typed POD (plain old data) to the hash.
+	*/
+	template<typename T>
+	bool UpdatePOD(const TCHAR* InDebugName, const T& InData)
+	{
+		static_assert(TIsPODType<T>::Value, "Update does not support a constructor / destructor on the element class.");
+		static_assert(!TIsPointer<T>::Value, "Update does not support pointer types.");
+		HashState.Update((const uint8 *)&InData, sizeof(T));
+#if WITH_EDITORONLY_DATA
+		if (LogCompileIdGeneration != 0)
+		{
+			FString ValuesStr;
+			ToDebugString(&InData, 1, ValuesStr);
+			Values.Top().PropertyKeys.Push(InDebugName);
+			Values.Top().PropertyValues.Push(ValuesStr);
+		}
+#endif
+		return true;
+	}
+
+	template<typename T>
+	bool UpdateShaderParameters()
+	{
+		const FShaderParametersMetadata* ShaderParametersMetadata = TShaderParameterStructTypeInfo<T>::GetStructMetadata();
+		return UpdatePOD(ShaderParametersMetadata->GetStructTypeName(), ShaderParametersMetadata->GetLayoutHash());
+	}
+
+	/**
+	Adds an string value to the hash.
+	*/
+	bool UpdateString(const TCHAR* InDebugName, FStringView InData)
+	{
+		HashState.Update((const uint8 *)InData.GetData(), sizeof(TCHAR)*InData.Len());
+#if WITH_EDITORONLY_DATA
+		if (LogCompileIdGeneration != 0)
+		{
+			Values.Top().PropertyKeys.Push(InDebugName);
+			Values.Top().PropertyValues.Push(FString(InData));
+		}
+#endif
+		return true;
+	}
+};
diff --git a/Engine/Plugins/FX/Niagara/Source/Niagara/Public/NiagaraComponent.h b/Engine/Plugins/FX/Niagara/Source/Niagara/Public/NiagaraComponent.h
--- a/Engine/Plugins/FX/Niagara/Source/Niagara/Public/NiagaraComponent.h
+++ b/Engine/Plugins/FX/Niagara/Source/Niagara/Public/NiagaraComponent.h
@@ -4,12 +4,17 @@
 
 #include "CoreMinimal.h"
 #include "UObject/ObjectMacros.h"
+#include "NiagaraDefines.h"
+#include "NiagaraScalabilityState.h"
+#include "NiagaraTickBehaviorEnum.h"
+#if UE_ENABLE_INCLUDE_ORDER_DEPRECATED_IN_5_2
 #include "NiagaraCommon.h"
 #include "NiagaraComponentPool.h"
 #include "NiagaraSystemInstanceController.h"
+#include "PrimitiveSceneProxy.h"
+#endif
 #include "NiagaraUserRedirectionParameterStore.h"
 #include "NiagaraVariant.h"
-#include "PrimitiveSceneProxy.h"
 #include "PrimitiveViewRelevance.h"
 #include "Particles/ParticlePerfStats.h"
 #include "Particles/ParticleSystemComponent.h"
@@ -18,11 +23,18 @@
 
 class FMeshElementCollector;
 class FNiagaraRenderer;
+class UNiagaraEffectType;
 class UNiagaraSystem;
 class UNiagaraParameterCollection;
 class UNiagaraParameterCollectionInstance;
+class FNiagaraSystemRenderData;
+class FNiagaraSystemInstance;
+class FNiagaraSystemInstanceController;
 class FNiagaraSystemSimulation;
 class FNiagaraGpuComputeDispatchInterface;
+enum class ENCPoolMethod : uint8;
+using FNiagaraSystemInstanceControllerPtr = TSharedPtr<FNiagaraSystemInstanceController, ESPMode::ThreadSafe>;
+using FNiagaraSystemInstanceControllerConstPtr = TSharedPtr<const FNiagaraSystemInstanceController, ESPMode::ThreadSafe>;
 
 // Called when the particle system is done
 DECLARE_DYNAMIC_MULTICAST_DELEGATE_OneParam(FOnNiagaraSystemFinished, class UNiagaraComponent*, PSystem);
@@ -194,10 +206,10 @@ public:
 	virtual void Deactivate() override;
 	virtual void DeactivateImmediate() override;
 
-	FORCEINLINE ENiagaraExecutionState GetRequestedExecutionState() const { return SystemInstanceController ? SystemInstanceController->GetRequestedExecutionState() : ENiagaraExecutionState::Complete; }
-	FORCEINLINE ENiagaraExecutionState GetExecutionState() const { return SystemInstanceController ? SystemInstanceController->GetActualExecutionState() : ENiagaraExecutionState::Complete; }
+	ENiagaraExecutionState GetRequestedExecutionState() const;
+	ENiagaraExecutionState GetExecutionState() const;
 
-	FORCEINLINE bool IsComplete() const { return SystemInstanceController ? SystemInstanceController->IsComplete() : true; }
+	bool IsComplete() const;
 
 private:
 
@@ -348,7 +360,7 @@ public:
 	void SetAutoDestroy(bool bInAutoDestroy);
 
 	UE_DEPRECATED(5.0, "This interface is no longer safe to access directly. Use the interface provided by GetSystemInstanceController instead.")
-	FNiagaraSystemInstance* GetSystemInstance() const { return SystemInstanceController ? SystemInstanceController->GetSystemInstance_Unsafe() : nullptr; }
+	FNiagaraSystemInstance* GetSystemInstance() const;
 
 	FNiagaraSystemInstanceControllerPtr GetSystemInstanceController() { return SystemInstanceController; }
 	FNiagaraSystemInstanceControllerConstPtr GetSystemInstanceController() const { return SystemInstanceController; }
@@ -559,7 +571,7 @@ public:
 		The significant index for this component. i.e. this is the Nth most significant instance of it's system in the scene.
 		Passed to the script to allow us to scale down internally for less significant systems instances.
 	*/
-	FORCEINLINE void SetSystemSignificanceIndex(int32 InIndex) 	{ if(SystemInstanceController) SystemInstanceController->SetSystemSignificanceIndex(InIndex); }
+	void SetSystemSignificanceIndex(int32 InIndex);
 
 	//~ Begin UObject Interface.
 	virtual void Serialize(FStructuredArchive::FRecord Record) override;
@@ -611,7 +623,7 @@ public:
 	UFUNCTION(BlueprintCallable, Category = Performance, meta = (Keywords = "Niagara Performance"))
 	void InitForPerformanceBaseline();
 
-	FORCEINLINE void SetLODDistance(float InLODDistance, float InMaxLODDistance);
+	void SetLODDistance(float InLODDistance, float InMaxLODDistance);
 
 #if WITH_EDITOR
 	void PostLoadNormalizeOverrideNames();
@@ -835,113 +847,8 @@ private:
 	void DestroyCullProxy();
 
 public:
-	FORCEINLINE FParticlePerfStatsContext GetPerfStatsContext(){ return FParticlePerfStatsContext(GetWorld(), Asset, this); }
+	FParticlePerfStatsContext GetPerfStatsContext();
 };
 
 FORCEINLINE bool UNiagaraComponent::GetPreviewLODDistanceEnabled()const { return bEnablePreviewLODDistance; }
 FORCEINLINE float UNiagaraComponent::GetPreviewLODDistance()const { return bEnablePreviewLODDistance ? PreviewLODDistance : 0.0f; }
-
-/**
-* Scene proxy for drawing niagara particle simulations.
-*/
-class NIAGARA_API FNiagaraSceneProxy : public FPrimitiveSceneProxy
-{
-public:
-	virtual SIZE_T GetTypeHash() const override;
-
-	FNiagaraSceneProxy(UNiagaraComponent* InComponent);
-	~FNiagaraSceneProxy();
-
-	/** Retrieves the render data for a single system */
-	FNiagaraSystemRenderData* GetSystemRenderData() { return RenderData; }
-
-	/** Called to allow renderers to free render state */
-	void DestroyRenderState_Concurrent();
-
-	/** Gets whether or not this scene proxy should be rendered. */
-	bool GetRenderingEnabled() const;
-
-	/** Sets whether or not this scene proxy should be rendered. */
-	void SetRenderingEnabled(bool bInRenderingEnabled);
-
-	FNiagaraGpuComputeDispatchInterface* GetComputeDispatchInterface() const { return ComputeDispatchInterface; }
-
-#if RHI_RAYTRACING
-	virtual void GetDynamicRayTracingInstances(FRayTracingMaterialGatheringContext& Context, TArray<FRayTracingInstance>& OutRayTracingInstances) override;
-	virtual bool IsRayTracingRelevant() const override { return true; }
-	virtual bool HasRayTracingRepresentation() const override { return true; }
-#endif
-
-	FORCEINLINE const FMatrix& GetLocalToWorldInverse() const { return LocalToWorldInverse; }
-
-	const FVector3f& GetLWCRenderTile() const { return RenderData ? RenderData->LWCRenderTile : FVector3f::ZeroVector; }
-
-	TUniformBuffer<FPrimitiveUniformShaderParameters>* GetCustomUniformBufferResource(bool bHasVelocity, const FBox& InstanceBounds = FBox(ForceInitToZero)) const;
-	FRHIUniformBuffer* GetCustomUniformBuffer(bool bHasVelocity, const FBox& InstanceBounds = FBox(ForceInitToZero)) const;
-
-	virtual FPrimitiveViewRelevance GetViewRelevance(const FSceneView* View) const override;
-
-	/** Some proxy wide dynamic settings passed down with the emitter dynamic data. */
-	struct FDynamicData
-	{
-		bool bUseCullProxy = false;
-		float LODDistanceOverride = -1.0f;
-#if WITH_PARTICLE_PERF_STATS
-		FParticlePerfStatsContext PerfStatsContext;
-#endif
-	};
-	const FDynamicData& GetProxyDynamicData()const { return DynamicData; }
-	void SetProxyDynamicData(const FDynamicData& NewData) { DynamicData = NewData; }
-
-private:
-	void ReleaseRenderThreadResources();
-
-	void ReleaseUniformBuffers(bool bEmpty);
-
-	//~ Begin FPrimitiveSceneProxy Interface.
-	virtual void CreateRenderThreadResources() override;
-
-	//virtual void OnActorPositionChanged() override;
-	virtual void OnTransformChanged() override;
-
-	virtual void GetDynamicMeshElements(const TArray<const FSceneView*>& Views, const FSceneViewFamily& ViewFamily, uint32 VisibilityMap, FMeshElementCollector& Collector) const override;
-
-
-	virtual bool CanBeOccluded() const override
-	{
-		// TODO account for MaterialRelevance.bDisableDepthTest and MaterialRelevance.bPostMotionBlurTranslucency as well
-		return !ShouldRenderCustomDepth();
-	}
-
-
-	/** Callback from the renderer to gather simple lights that this proxy wants renderered. */
-	virtual void GatherSimpleLights(const FSceneViewFamily& ViewFamily, FSimpleLightArray& OutParticleLights) const override;
-
-	virtual uint32 GetMemoryFootprint() const override;
-
-	uint32 GetAllocatedSize() const;
-
-private:
-	/** Custom Uniform Buffers, allows us to have renderer specific data packed inside such as pre-skinned bounds. */
-	mutable TMap<uint64, TUniformBuffer<FPrimitiveUniformShaderParameters>*> CustomUniformBuffers;
-
-	/** The data required to render a single instance of a NiagaraSystem */
-	FNiagaraSystemRenderData* RenderData = nullptr;
-
-	FNiagaraGpuComputeDispatchInterface* ComputeDispatchInterface = nullptr;
-
-	FMatrix LocalToWorldInverse;
-
-	TStatId SystemStatID;
-
-	FDynamicData DynamicData;
-};
-
-FORCEINLINE void UNiagaraComponent::SetLODDistance(float InLODDistance, float InMaxLODDistance)
-{
-	if (!bEnablePreviewLODDistance && SystemInstanceController)
-	{
-		SystemInstanceController->SetLODDistance(InLODDistance, InMaxLODDistance, false);
-	}
-}
-
diff --git a/Engine/Plugins/FX/Niagara/Source/Niagara/Public/NiagaraParameterStore.h b/Engine/Plugins/FX/Niagara/Source/Niagara/Public/NiagaraParameterStore.h
--- a/Engine/Plugins/FX/Niagara/Source/Niagara/Public/NiagaraParameterStore.h
+++ b/Engine/Plugins/FX/Niagara/Source/Niagara/Public/NiagaraParameterStore.h
@@ -2,13 +2,19 @@
 
 #pragma once
 
+#include "NiagaraTypes.h"
+
+#if UE_ENABLE_INCLUDE_ORDER_DEPRECATED_IN_5_2
 #include "NiagaraCommon.h"
 #include "NiagaraDataInterface.h"
+#endif
+
 #include "NiagaraParameterStore.generated.h"
 
 // When not cooked, sort by actual name to ensure deterministic cooked data
 #define NIAGARA_VARIABLE_LEXICAL_SORTING WITH_EDITORONLY_DATA
 
+class UNiagaraDataInterface;
 struct FNiagaraParameterStore;
 
 USTRUCT()
@@ -112,7 +118,7 @@ struct FNiagaraParameterStoreBinding
 	FORCEINLINE_DEBUGGABLE void Empty(FNiagaraParameterStore* DestStore, FNiagaraParameterStore* SrcStore);
 	FORCEINLINE_DEBUGGABLE bool Initialize(FNiagaraParameterStore* DestStore, FNiagaraParameterStore* SrcStore, const FNiagaraBoundParameterArray* BoundParameters = nullptr);
 	FORCEINLINE_DEBUGGABLE bool VerifyBinding(const FNiagaraParameterStore* DestStore, const FNiagaraParameterStore* SrcStore)const;
-	FORCEINLINE_DEBUGGABLE void Tick(FNiagaraParameterStore* DestStore, FNiagaraParameterStore* SrcStore, bool bForce = false);
+	void Tick(FNiagaraParameterStore* DestStore, FNiagaraParameterStore* SrcStore, bool bForce = false);
 	FORCEINLINE_DEBUGGABLE void Dump(const FNiagaraParameterStore* DestStore, const FNiagaraParameterStore* SrcStore)const;
 	/** TODO: Merge contiguous ranges into a single binding? */
 	//FORCEINLINE_DEBUGGABLE void Optimize();
@@ -294,7 +300,7 @@ public:
 	/** Recreates any bindings to reflect a layout change etc. */
 	void TransferBindings(FNiagaraParameterStore& OtherStore);
 	/** Handles any update such as pushing parameters to bound stores etc. */
-	FORCEINLINE_DEBUGGABLE void Tick();
+	void Tick();
 	/** Unbinds this store from all stores it's being driven by. */
 	void UnbindFromSourceStores();
 	
@@ -429,13 +435,7 @@ public:
 	}
 
 	/** Returns the data interface for the passed parameter if it exists in this store. */
-	FORCEINLINE_DEBUGGABLE UNiagaraDataInterface* GetDataInterface(const FNiagaraVariable& Parameter)const
-	{
-		int32 Offset = IndexOf(Parameter);
-		UNiagaraDataInterface* Interface = GetDataInterface(Offset);
-		checkSlow(!Interface || Parameter.GetType().GetClass() == Interface->GetClass());
-		return Interface;
-	}
+	UNiagaraDataInterface* GetDataInterface(const FNiagaraVariable& Parameter) const;
 
 	/** Returns a struct converter for the given variable, if the store contains the variable and it's a LWC type. */
 	FNiagaraLwcStructConverter GetStructConverter(const FNiagaraVariable& Parameter) const;
@@ -468,47 +468,7 @@ public:
 	}
 
 	/** Copies the passed parameter from this parameter store into another. */
-	FORCEINLINE_DEBUGGABLE void CopyParameterData(FNiagaraParameterStore& DestStore, const FNiagaraVariable& Parameter) const
-	{
-		int32 DestIndex = DestStore.IndexOf(Parameter);
-		int32 SrcIndex = IndexOf(Parameter);
-		if (DestIndex != INDEX_NONE && SrcIndex != INDEX_NONE)
-		{
-			if (Parameter.IsDataInterface())
-			{
-				ensure(DestStore.DataInterfaces.IsValidIndex(DestIndex));
-				UNiagaraDataInterface* DestDataInterface = DestStore.DataInterfaces[DestIndex];
-				if (DestDataInterface == nullptr)
-				{
-					if (ensureMsgf(DestStore.Owner != nullptr, TEXT("Destination data interface pointer was null and a new one couldn't be created because the destination store's owner pointer was also null.")))
-					{
-						UE_LOG(LogNiagara, Warning, TEXT("While trying to copy parameter data the destination data interface was null, creating a new one.  Parameter: %s Destination Store Owner: %s"), *Parameter.GetName().ToString(), *GetPathNameSafe(DestStore.Owner.Get()));
-						DestDataInterface = NewObject<UNiagaraDataInterface>(DestStore.Owner.Get(), Parameter.GetType().GetClass(), NAME_None, RF_Transactional | RF_Public);
-						DestStore.DataInterfaces[DestIndex] = DestDataInterface;
-					}
-					else
-					{
-						return;
-					}
-				}
-				DataInterfaces[SrcIndex]->CopyTo(DestDataInterface);
-				DestStore.OnInterfaceChange();
-			}
-			else if (Parameter.IsUObject())
-			{
-				DestStore.SetUObject(GetUObject(SrcIndex), DestIndex);
-			}
-			else if (Parameter.GetType() == FNiagaraTypeDefinition::GetPositionDef())
-			{
-				const FVector* SourceVector = GetPositionParameterValue(Parameter.GetName());
-				DestStore.SetPositionParameterValue(SourceVector ? *SourceVector : FVector::ZeroVector, Parameter.GetName());
-			}
-			else
-			{
-				DestStore.SetParameterData(GetParameterData(SrcIndex), DestIndex, Parameter.GetSizeInBytes());
-			}
-		}
-	}
+	void CopyParameterData(FNiagaraParameterStore& DestStore, const FNiagaraVariable& Parameter) const;
 
 	enum class EDataInterfaceCopyMethod
 	{
@@ -723,22 +683,6 @@ FORCEINLINE bool FNiagaraParameterStore::SetParameterValue(const FQuat4d& InValu
 	return SetParameterValue((FQuat4f)InValue, Param, bAdd);
 }
 
-FORCEINLINE_DEBUGGABLE void FNiagaraParameterStore::Tick()
-{
-#if NIAGARA_NAN_CHECKING
-	CheckForNaNs();
-#endif
-	if (Bindings.Num() > 0 && (bParametersDirty || bInterfacesDirty || bUObjectsDirty))
-	{
-		TickBindings();
-	}
-
-	//We have to have ticked all our source stores before now.
-	bParametersDirty = false;
-	bInterfacesDirty = false;
-	bUObjectsDirty = false;
-}
-
 FORCEINLINE_DEBUGGABLE void FNiagaraParameterStoreBinding::Empty(FNiagaraParameterStore* DestStore, FNiagaraParameterStore* SrcStore)
 {
 	if (DestStore)
@@ -829,37 +773,6 @@ FORCEINLINE_DEBUGGABLE bool FNiagaraParameterStoreBinding::VerifyBinding(const F
 // 	//TODO
 // }
 
-FORCEINLINE_DEBUGGABLE void FNiagaraParameterStoreBinding::Tick(FNiagaraParameterStore* DestStore, FNiagaraParameterStore* SrcStore, bool bForce)
-{
-	if (SrcStore->GetParametersDirty() || bForce)
-	{
-		for (FParameterBinding& Binding : ParameterBindings)
-		{
-			DestStore->SetParameterData(SrcStore->GetParameterData(Binding.SrcOffset), Binding.DestOffset, Binding.Size);
-		}
-	}
-
-	if (SrcStore->GetInterfacesDirty() || bForce)
-	{
-		for (FInterfaceBinding& Binding : InterfaceBindings)
-		{
-			DestStore->SetDataInterface(SrcStore->GetDataInterface(Binding.SrcOffset), Binding.DestOffset);
-		}
-	}
-
-	if (SrcStore->GetUObjectsDirty() || bForce)
-	{
-		for (FUObjectBinding& Binding : UObjectBindings)
-		{
-			DestStore->SetUObject(SrcStore->GetUObject(Binding.SrcOffset), Binding.DestOffset);
-		}
-	}
-
-#if NIAGARA_NAN_CHECKING
-	DestStore->CheckForNaNs();
-#endif
-}
-
 FORCEINLINE_DEBUGGABLE void FNiagaraParameterStoreBinding::Dump(const FNiagaraParameterStore* DestStore, const FNiagaraParameterStore* SrcStore)const
 {
 #if WITH_EDITORONLY_DATA
diff --git a/Engine/Plugins/FX/Niagara/Source/Niagara/Public/NiagaraSceneProxy.h b/Engine/Plugins/FX/Niagara/Source/Niagara/Public/NiagaraSceneProxy.h
--- a/Engine/Plugins/FX/Niagara/Source/Niagara/Public/NiagaraSceneProxy.h
+++ b/Engine/Plugins/FX/Niagara/Source/Niagara/Public/NiagaraSceneProxy.h
@@ -2,4 +2,103 @@
 
 #pragma once
 
-// Placeholder for upcoming change
\ No newline at end of file
+#include "PrimitiveSceneProxy.h"
+
+class FNiagaraSystemRenderData;
+
+/**
+* Scene proxy for drawing niagara particle simulations.
+*/
+class NIAGARA_API FNiagaraSceneProxy : public FPrimitiveSceneProxy
+{
+public:
+	virtual SIZE_T GetTypeHash() const override;
+
+	FNiagaraSceneProxy(UNiagaraComponent* InComponent);
+	~FNiagaraSceneProxy();
+
+	/** Retrieves the render data for a single system */
+	FNiagaraSystemRenderData* GetSystemRenderData() { return RenderData; }
+
+	/** Called to allow renderers to free render state */
+	void DestroyRenderState_Concurrent();
+
+	/** Gets whether or not this scene proxy should be rendered. */
+	bool GetRenderingEnabled() const;
+
+	/** Sets whether or not this scene proxy should be rendered. */
+	void SetRenderingEnabled(bool bInRenderingEnabled);
+
+	FNiagaraGpuComputeDispatchInterface* GetComputeDispatchInterface() const { return ComputeDispatchInterface; }
+
+#if RHI_RAYTRACING
+	virtual void GetDynamicRayTracingInstances(FRayTracingMaterialGatheringContext& Context, TArray<FRayTracingInstance>& OutRayTracingInstances) override;
+	virtual bool IsRayTracingRelevant() const override { return true; }
+	virtual bool HasRayTracingRepresentation() const override { return true; }
+#endif
+
+	FORCEINLINE const FMatrix& GetLocalToWorldInverse() const { return LocalToWorldInverse; }
+
+	const FVector3f& GetLWCRenderTile() const;
+
+	TUniformBuffer<FPrimitiveUniformShaderParameters>* GetCustomUniformBufferResource(bool bHasVelocity, const FBox& InstanceBounds = FBox(ForceInitToZero)) const;
+	FRHIUniformBuffer* GetCustomUniformBuffer(bool bHasVelocity, const FBox& InstanceBounds = FBox(ForceInitToZero)) const;
+
+	virtual FPrimitiveViewRelevance GetViewRelevance(const FSceneView* View) const override;
+
+	/** Some proxy wide dynamic settings passed down with the emitter dynamic data. */
+	struct FDynamicData
+	{
+		bool bUseCullProxy = false;
+		float LODDistanceOverride = -1.0f;
+#if WITH_PARTICLE_PERF_STATS
+		FParticlePerfStatsContext PerfStatsContext;
+#endif
+	};
+	const FDynamicData& GetProxyDynamicData()const { return DynamicData; }
+	void SetProxyDynamicData(const FDynamicData& NewData) { DynamicData = NewData; }
+
+private:
+	void ReleaseRenderThreadResources();
+
+	void ReleaseUniformBuffers(bool bEmpty);
+
+	//~ Begin FPrimitiveSceneProxy Interface.
+	virtual void CreateRenderThreadResources() override;
+
+	//virtual void OnActorPositionChanged() override;
+	virtual void OnTransformChanged() override;
+
+	virtual void GetDynamicMeshElements(const TArray<const FSceneView*>& Views, const FSceneViewFamily& ViewFamily, uint32 VisibilityMap, FMeshElementCollector& Collector) const override;
+
+
+	virtual bool CanBeOccluded() const override
+	{
+		// TODO account for MaterialRelevance.bDisableDepthTest and MaterialRelevance.bPostMotionBlurTranslucency as well
+		return !ShouldRenderCustomDepth();
+	}
+
+
+	/** Callback from the renderer to gather simple lights that this proxy wants renderered. */
+	virtual void GatherSimpleLights(const FSceneViewFamily& ViewFamily, FSimpleLightArray& OutParticleLights) const override;
+
+	virtual uint32 GetMemoryFootprint() const override;
+
+	uint32 GetAllocatedSize() const;
+
+private:
+	/** Custom Uniform Buffers, allows us to have renderer specific data packed inside such as pre-skinned bounds. */
+	mutable TMap<uint64, TUniformBuffer<FPrimitiveUniformShaderParameters>*> CustomUniformBuffers;
+
+	/** The data required to render a single instance of a NiagaraSystem */
+	FNiagaraSystemRenderData* RenderData = nullptr;
+
+	FNiagaraGpuComputeDispatchInterface* ComputeDispatchInterface = nullptr;
+
+	FMatrix LocalToWorldInverse;
+
+	TStatId SystemStatID;
+
+	FDynamicData DynamicData;
+};
+
diff --git a/Engine/Plugins/FX/Niagara/Source/Niagara/Public/NiagaraTypes.h b/Engine/Plugins/FX/Niagara/Source/Niagara/Public/NiagaraTypes.h
--- a/Engine/Plugins/FX/Niagara/Source/Niagara/Public/NiagaraTypes.h
+++ b/Engine/Plugins/FX/Niagara/Source/Niagara/Public/NiagaraTypes.h
@@ -7,12 +7,16 @@
 #include "Misc/SecureHash.h"
 #include "UObject/GCObject.h"
 #include "UObject/UnrealType.h"
+#if UE_ENABLE_INCLUDE_ORDER_DEPRECATED_IN_5_2
 #include "Engine/Texture2D.h"
+#endif
 
 #include "NiagaraTypes.generated.h"
 
 struct FNiagaraTypeDefinition;
+struct FNiagaraCompileHashVisitor;
 class UNiagaraDataInterfaceBase;
+class UTexture2D;
 
 DECLARE_LOG_CATEGORY_EXTERN(LogNiagara, Log, Verbose);
 
@@ -561,162 +565,6 @@ enum class ENiagaraOrientationAxis : uint32
 };
 
 
-USTRUCT()
-struct NIAGARA_API FNiagaraCompileHashVisitorDebugInfo
-{
-	GENERATED_USTRUCT_BODY()
-public:
-	UPROPERTY()
-	FString Object;
-
-	UPROPERTY()
-	TArray<FString> PropertyKeys;
-
-	UPROPERTY()
-	TArray<FString> PropertyValues;
-};
-
-/**
-Used to store the state of a graph when deciding if it has been dirtied for recompile.
-*/
-struct NIAGARA_API FNiagaraCompileHashVisitor
-{
-public:
-	FNiagaraCompileHashVisitor(FSHA1& InHashState) : HashState(InHashState) {}
-
-	FSHA1& HashState;
-	TArray<const void*> ObjectList;
-
-	static int LogCompileIdGeneration;
-
-#if WITH_EDITORONLY_DATA
-	FNiagaraCompileHashVisitorDebugInfo* AddDebugInfo()
-	{
-		if (LogCompileIdGeneration != 0)
-		{
-			return &Values.AddDefaulted_GetRef();
-		}
-		return nullptr;
-	}
-
-	// Debug data about the compilation hash, including key value pairs to detect differences.
-	TArray<FNiagaraCompileHashVisitorDebugInfo> Values;
-
-	template<typename T>
-	void ToDebugString(const T* InData, uint64 InDataCount, FString& OutStr)
-	{
-		for (int32 i = 0; i < InDataCount; i++)
-		{
-			FString DataStr = LexToString(InData[i]);
-			OutStr.Appendf(TEXT("%s "), *DataStr);
-		}
-	}
-#endif
-	/**
-	Registers a pointer for later reference in the compile id in a deterministic manner.
-	*/
-	int32 RegisterReference(const void* InObject)
-	{
-		if (InObject == nullptr)
-		{
-			return -1;
-		}
-
-		int32 Index = ObjectList.Find(InObject);
-		if (Index < 0)
-		{
-			Index = ObjectList.Add(InObject);
-		}
-		return Index;
-	}
-
-	/**
-	We don't usually want to save GUID's or pointer values because they have nondeterministic values. Consider a PostLoad upgrade operation that creates a new node.
-	Each pin and node gets a unique ID. If you close the editor and reopen, you'll get a different set of values. One of the characteristics we want for compilation
-	behavior is that the same graph structure produces the same compile results, so we only want to embed information that is deterministic. This method is for use
-	when registering a pointer to an object that is serialized within the compile hash.
-	*/
-	bool UpdateReference(const TCHAR* InDebugName, const void* InObject)
-	{
-		int32 Index = RegisterReference(InObject);
-		return UpdatePOD(InDebugName, Index);
-	}
-
-	/**
-	Adds an array of POD (plain old data) values to the hash.
-	*/
-	template<typename T>
-	bool UpdateArray(const TCHAR* InDebugName, const T* InData, uint64 InDataCount = 1)
-	{
-		static_assert(TIsPODType<T>::Value, "UpdateArray does not support a constructor / destructor on the element class.");
-		HashState.Update((const uint8 *)InData, sizeof(T)*InDataCount);
-#if WITH_EDITORONLY_DATA
-		if (LogCompileIdGeneration != 0)
-		{
-			FString ValuesStr = InDebugName;
-			ValuesStr.Append(TEXT(" = "));
-			ToDebugString(InData, InDataCount, ValuesStr);
-			Values.Top().PropertyKeys.Push(InDebugName);
-			Values.Top().PropertyValues.Push(ValuesStr);
-		}
-#endif
-		return true;
-	}
-
-	/**
-	 * Adds the string representation of an FName to the hash
-	 */
-	bool UpdateName(const TCHAR* InDebugName, FName InName)
-	{
-		FNameBuilder NameString(InName);
-		return UpdateString(InDebugName, NameString.ToView());
-	}
-
-	/**
-	Adds a single value of typed POD (plain old data) to the hash.
-	*/
-	template<typename T>
-	bool UpdatePOD(const TCHAR* InDebugName, const T& InData)
-	{
-		static_assert(TIsPODType<T>::Value, "Update does not support a constructor / destructor on the element class.");
-		static_assert(!TIsPointer<T>::Value, "Update does not support pointer types.");
-		HashState.Update((const uint8 *)&InData, sizeof(T));
-#if WITH_EDITORONLY_DATA
-		if (LogCompileIdGeneration != 0)
-		{
-			FString ValuesStr;
-			ToDebugString(&InData, 1, ValuesStr);
-			Values.Top().PropertyKeys.Push(InDebugName);
-			Values.Top().PropertyValues.Push(ValuesStr);
-		}
-#endif
-		return true;
-	}
-
-	template<typename T>
-	bool UpdateShaderParameters()
-	{
-		const FShaderParametersMetadata* ShaderParametersMetadata = TShaderParameterStructTypeInfo<T>::GetStructMetadata();
-		return UpdatePOD(ShaderParametersMetadata->GetStructTypeName(), ShaderParametersMetadata->GetLayoutHash());
-	}
-
-	/**
-	Adds an string value to the hash.
-	*/
-	bool UpdateString(const TCHAR* InDebugName, FStringView InData)
-	{
-		HashState.Update((const uint8 *)InData.GetData(), sizeof(TCHAR)*InData.Len());
-#if WITH_EDITORONLY_DATA
-		if (LogCompileIdGeneration != 0)
-		{
-			Values.Top().PropertyKeys.Push(InDebugName);
-			Values.Top().PropertyValues.Push(FString(InData));
-		}
-#endif
-		return true;
-	}
-};
-
 /** Defines options for conditionally editing and showing script inputs in the UI. */
 USTRUCT()
 struct NIAGARA_API FNiagaraInputConditionMetadata
diff --git a/Engine/Plugins/FX/Niagara/Source/Niagara/Public/NiagaraUserRedirectionParameterStore.h b/Engine/Plugins/FX/Niagara/Source/Niagara/Public/NiagaraUserRedirectionParameterStore.h
--- a/Engine/Plugins/FX/Niagara/Source/Niagara/Public/NiagaraUserRedirectionParameterStore.h
+++ b/Engine/Plugins/FX/Niagara/Source/Niagara/Public/NiagaraUserRedirectionParameterStore.h
@@ -2,8 +2,11 @@
 
 #pragma once
 
+#if UE_ENABLE_INCLUDE_ORDER_DEPRECATED_IN_5_2
 #include "NiagaraCommon.h"
 #include "NiagaraDataInterface.h"
+#endif
+
 #include "NiagaraParameterStore.h"
 #include "NiagaraUserRedirectionParameterStore.generated.h"
 
diff --git a/Engine/Plugins/Online/OnlineSubsystemUtils/Source/OnlineSubsystemUtils/Private/OnlineSubsystemUtils.cpp b/Engine/Plugins/Online/OnlineSubsystemUtils/Source/OnlineSubsystemUtils/Private/OnlineSubsystemUtils.cpp
--- a/Engine/Plugins/Online/OnlineSubsystemUtils/Source/OnlineSubsystemUtils/Private/OnlineSubsystemUtils.cpp
+++ b/Engine/Plugins/Online/OnlineSubsystemUtils/Source/OnlineSubsystemUtils/Private/OnlineSubsystemUtils.cpp
@@ -16,6 +16,7 @@
 #include "OnlineSubsystemBPCallHelper.h"
 
 #include "VoiceModule.h"
+#include "VoipListenerSynthComponent.h"
 #include "AudioDevice.h"
 #include "Sound/AudioSettings.h"
 #include "Sound/SoundWaveProcedural.h"
diff --git a/Engine/Plugins/Online/OnlineSubsystemUtils/Source/OnlineSubsystemUtils/Public/OnlineSubsystemUtils.h b/Engine/Plugins/Online/OnlineSubsystemUtils/Source/OnlineSubsystemUtils/Public/OnlineSubsystemUtils.h
--- a/Engine/Plugins/Online/OnlineSubsystemUtils/Source/OnlineSubsystemUtils/Public/OnlineSubsystemUtils.h
+++ b/Engine/Plugins/Online/OnlineSubsystemUtils/Source/OnlineSubsystemUtils/Public/OnlineSubsystemUtils.h
@@ -10,11 +10,17 @@
 #include "Interfaces/OnlineIdentityInterface.h"
 #include "Online.h"
 #include "EngineGlobals.h"
-#include "VoipListenerSynthComponent.h"
 #include "Features/IModularFeatures.h"
 
+#if UE_ENABLE_INCLUDE_ORDER_DEPRECATED_IN_5_2
+#include "VoipListenerSynthComponent.h"
+#endif
+
+
 class UWorld;
 class UAudioComponent;
+class UVoipListenerSynthComponent;
+struct FVoiceSettings;
 
 #ifdef ONLINESUBSYSTEMUTILS_API
 
diff --git a/Engine/Plugins/Runtime/SignificanceManager/Source/SignificanceManager/Private/SignificanceManager.cpp b/Engine/Plugins/Runtime/SignificanceManager/Source/SignificanceManager/Private/SignificanceManager.cpp
--- a/Engine/Plugins/Runtime/SignificanceManager/Source/SignificanceManager/Private/SignificanceManager.cpp
+++ b/Engine/Plugins/Runtime/SignificanceManager/Source/SignificanceManager/Private/SignificanceManager.cpp
@@ -5,6 +5,7 @@
 #include "Engine/Canvas.h"
 #include "GameFramework/HUD.h"
 #include "Engine/Engine.h"
+#include "Modules/ModuleManager.h"
 #include "Async/ParallelFor.h"
 
 #include UE_INLINE_GENERATED_CPP_BY_NAME(SignificanceManager)
@@ -14,6 +15,52 @@
 #include "ConsoleSettings.h"
 #endif // ALLOW_CONSOLE
 
+/* Module definition for significance manager. Owns the references to created significance managers*/
+class SIGNIFICANCEMANAGER_API FSignificanceManagerModule : public FDefaultModuleImpl, public FGCObject
+{
+public:
+	// Begin IModuleInterface overrides
+	virtual void StartupModule() override;
+	// End IModuleInterface overrides
+
+	// Begin FGCObject overrides
+	virtual void AddReferencedObjects( FReferenceCollector& Collector ) override;
+	virtual FString GetReferencerName() const override
+	{
+		return TEXT("FSignificanceManagerModule");
+	}
+	// End FGCObject overrides
+
+	// Returns the significance manager for the specified World
+	FORCEINLINE static USignificanceManager* Get(const UWorld* World)
+	{
+		return WorldSignificanceManagers.FindRef(World);
+	}
+
+private:
+
+	// Callback function registered with global world delegates to instantiate significance manager when a game world is created
+	static void OnWorldInit(UWorld* World, const UWorld::InitializationValues IVS);
+
+	// Callback function registered with global world delegates to cleanup significance manager when a game world is destroyed
+	static void OnWorldCleanup(UWorld* World, bool bSessionEnded, bool bCleanupResources);
+
+	// Callback function registered with HUD to supply debug info when ShowDebug SignificanceManager has been entered on the console
+	static void OnShowDebugInfo(AHUD* HUD, UCanvas* Canvas, const FDebugDisplayInfo& DisplayInfo, float& YL, float& YPos);
+
+#if ALLOW_CONSOLE
+	// Callback function registered with Console to inject show debug auto complete command
+	static void PopulateAutoCompleteEntries(TArray<FAutoCompleteCommand>& AutoCompleteList);
+#endif // ALLOW_CONSOLE
+
+	// Map of worlds to their significance manager
+	static TMap<const UWorld*, USignificanceManager*> WorldSignificanceManagers;
+
+	// Cached class for instantiating significance manager
+	static TSubclassOf<USignificanceManager>  SignificanceManagerClass;
+};
+
+
 IMPLEMENT_MODULE( FSignificanceManagerModule, SignificanceManager );
 
 DECLARE_CYCLE_STAT(TEXT("Update Total"), STAT_SignificanceManager_Update, STATGROUP_SignificanceManager);
@@ -380,6 +427,11 @@ bool USignificanceManager::QuerySignificance(const UObject* Object, float& OutSi
 	}
 }
 
+USignificanceManager* USignificanceManager::Get(const UWorld* World)
+{
+	return FSignificanceManagerModule::Get(World);
+}
+
 void USignificanceManager::FManagedObjectInfo::UpdateSignificance(const TArray<FTransform>& InViewpoints, const bool bSortAscending)
 {
 	float OldSignificance = Significance;
diff --git a/Engine/Plugins/Runtime/SignificanceManager/Source/SignificanceManager/Public/SignificanceManager.h b/Engine/Plugins/Runtime/SignificanceManager/Source/SignificanceManager/Public/SignificanceManager.h
--- a/Engine/Plugins/Runtime/SignificanceManager/Source/SignificanceManager/Public/SignificanceManager.h
+++ b/Engine/Plugins/Runtime/SignificanceManager/Source/SignificanceManager/Public/SignificanceManager.h
@@ -7,7 +7,9 @@
 #include "UObject/ObjectMacros.h"
 #include "UObject/Object.h"
 #include "Templates/SubclassOf.h"
+#if UE_ENABLE_INCLUDE_ORDER_DEPRECATED_IN_5_2
 #include "Modules/ModuleManager.h"
+#endif
 #include "UObject/GCObject.h"
 #include "UObject/SoftObjectPath.h"
 #include "Engine/World.h"
@@ -22,51 +24,6 @@ struct FAutoCompleteCommand;
 DECLARE_STATS_GROUP(TEXT("Significance Manager"), STATGROUP_SignificanceManager, STATCAT_Advanced);
 DECLARE_LOG_CATEGORY_EXTERN(LogSignificance, Log, All);
 
-/* Module definition for significance manager. Owns the references to created significance managers*/
-class SIGNIFICANCEMANAGER_API FSignificanceManagerModule : public FDefaultModuleImpl, public FGCObject
-{
-public:
-	// Begin IModuleInterface overrides
-	virtual void StartupModule() override;
-	// End IModuleInterface overrides
-
-	// Begin FGCObject overrides
-	virtual void AddReferencedObjects( FReferenceCollector& Collector ) override;
-	virtual FString GetReferencerName() const override
-	{
-		return TEXT("FSignificanceManagerModule");
-	}
-	// End FGCObject overrides
-
-	// Returns the significance manager for the specified World
-	FORCEINLINE static USignificanceManager* Get(const UWorld* World)
-	{
-		return WorldSignificanceManagers.FindRef(World);
-	}
-
-private:
-
-	// Callback function registered with global world delegates to instantiate significance manager when a game world is created
-	static void OnWorldInit(UWorld* World, const UWorld::InitializationValues IVS);
-
-	// Callback function registered with global world delegates to cleanup significance manager when a game world is destroyed
-	static void OnWorldCleanup(UWorld* World, bool bSessionEnded, bool bCleanupResources);
-
-	// Callback function registered with HUD to supply debug info when ShowDebug SignificanceManager has been entered on the console
-	static void OnShowDebugInfo(AHUD* HUD, UCanvas* Canvas, const FDebugDisplayInfo& DisplayInfo, float& YL, float& YPos);
-
-#if ALLOW_CONSOLE
-	// Callback function registered with Console to inject show debug auto complete command
-	static void PopulateAutoCompleteEntries(TArray<FAutoCompleteCommand>& AutoCompleteList);
-#endif // ALLOW_CONSOLE
-
-	// Map of worlds to their significance manager
-	static TMap<const UWorld*, USignificanceManager*> WorldSignificanceManagers;
-
-	// Cached class for instantiating significance manager
-	static TSubclassOf<USignificanceManager>  SignificanceManagerClass;
-};
-
 /* The significance manager provides a framework for registering objects by tag to each have a significance
  * value calculated from which a game specific subclass and game logic can make decisions about what level
  * of detail objects should be at, tick frequency, whether to spawn effects, and other such functionality
@@ -192,10 +149,7 @@ public:
 	bool QuerySignificance(const UObject* Object, float& OutSignificance) const;
 
 	// Returns the significance manager for the specified World
-	FORCEINLINE static USignificanceManager* Get(const UWorld* World)
-	{
-		return FSignificanceManagerModule::Get(World);
-	}
+	static USignificanceManager* Get(const UWorld* World);
 
 	// Templated convenience function to return a significance manager cast to a known type
 	template<class T>
@@ -256,5 +210,5 @@ private:
 	void OnShowDebugInfo(AHUD* HUD, UCanvas* Canvas, const FDebugDisplayInfo& DisplayInfo, float& YL, float& YPos);
 
 	// Friend FSignificanceManagerModule so it can call OnShowDebugInfo and check 
-	friend FSignificanceManagerModule;
+	friend class FSignificanceManagerModule;
 };
diff --git a/Engine/Source/Developer/AnimationWidgets/Public/SAdvancedRotationInputBox.h b/Engine/Source/Developer/AnimationWidgets/Public/SAdvancedRotationInputBox.h
--- a/Engine/Source/Developer/AnimationWidgets/Public/SAdvancedRotationInputBox.h
+++ b/Engine/Source/Developer/AnimationWidgets/Public/SAdvancedRotationInputBox.h
@@ -18,7 +18,7 @@
 
 namespace ESlateRotationRepresentation
 {
-	enum Type
+	enum Type : int
 	{
 		EulerXYZ,
 		EulerXZY,
@@ -35,7 +35,7 @@ namespace ESlateRotationRepresentation
 
 namespace ESlateTransformSubComponent
 {
-	enum Type
+	enum Type : int
 	{
 		X,
 		Y,
diff --git a/Engine/Source/Developer/AnimationWidgets/Public/SAdvancedTransformInputBox.h b/Engine/Source/Developer/AnimationWidgets/Public/SAdvancedTransformInputBox.h
--- a/Engine/Source/Developer/AnimationWidgets/Public/SAdvancedTransformInputBox.h
+++ b/Engine/Source/Developer/AnimationWidgets/Public/SAdvancedTransformInputBox.h
@@ -24,7 +24,7 @@
 
 namespace ESlateTransformComponent
 {
-	enum Type
+	enum Type : int
 	{
 		Location,
 		Rotation,
diff --git a/Engine/Source/Editor/LevelEditor/Public/SLevelViewport.h b/Engine/Source/Editor/LevelEditor/Public/SLevelViewport.h
--- a/Engine/Source/Editor/LevelEditor/Public/SLevelViewport.h
+++ b/Engine/Source/Editor/LevelEditor/Public/SLevelViewport.h
@@ -31,6 +31,7 @@ class SCaptureRegionWidget;
 class SGameLayerManager;
 class UFoliageType;
 enum class EMapChangeType : uint8;
+enum ELabelAnchorMode : int;
 
 /**
  * Encapsulates an SViewport and an SLevelViewportToolBar
diff --git a/Engine/Source/Editor/SkeletonEditor/Private/BoneProxy.cpp b/Engine/Source/Editor/SkeletonEditor/Private/BoneProxy.cpp
--- a/Engine/Source/Editor/SkeletonEditor/Private/BoneProxy.cpp
+++ b/Engine/Source/Editor/SkeletonEditor/Private/BoneProxy.cpp
@@ -4,6 +4,7 @@
 #include "IPersonaPreviewScene.h"
 #include "BoneControllers/AnimNode_ModifyBone.h"
 #include "AnimPreviewInstance.h"
+#include "SAdvancedTransformInputBox.h"
 #include "ScopedTransaction.h"
 #include "Animation/DebugSkelMeshComponent.h"
 
diff --git a/Engine/Source/Editor/SkeletonEditor/Private/SkeletonTreeBoneItem.h b/Engine/Source/Editor/SkeletonEditor/Private/SkeletonTreeBoneItem.h
--- a/Engine/Source/Editor/SkeletonEditor/Private/SkeletonTreeBoneItem.h
+++ b/Engine/Source/Editor/SkeletonEditor/Private/SkeletonTreeBoneItem.h
@@ -3,21 +3,31 @@
 #pragma once
 
 #include "CoreMinimal.h"
-#include "Styling/SlateColor.h"
 #include "Input/Reply.h"
+#include "ISkeletonTreeItem.h"
+#include "SkeletonTreeItem.h"
+#include "Animation/Skeleton.h"
+#include "BoneProxy.h"
+#include "UObject/GCObject.h"
+
+#if UE_ENABLE_INCLUDE_ORDER_DEPRECATED_IN_5_2
+#include "Styling/SlateColor.h"
 #include "Layout/Visibility.h"
 #include "Widgets/SWidget.h"
 #include "Fonts/SlateFontInfo.h"
 #include "Widgets/Views/STableViewBase.h"
 #include "Widgets/Views/STableRow.h"
-#include "ISkeletonTreeItem.h"
-#include "SkeletonTreeItem.h"
-#include "Animation/Skeleton.h"
 #include "Widgets/Input/SComboButton.h"
-#include "BoneProxy.h"
-#include "UObject/GCObject.h"
+#endif
+
 
+class SComboButton;
+class SWidget;
 class UDebugSkelMeshComponent;
+struct EVisibility;
+struct FSlateBrush;
+struct FSlateColor;
+struct FSlateFontInfo;
 
 class FSkeletonTreeBoneItem : public FSkeletonTreeItem, public FGCObject
 {
diff --git a/Engine/Source/Editor/SkeletonEditor/Public/BoneProxy.h b/Engine/Source/Editor/SkeletonEditor/Public/BoneProxy.h
--- a/Engine/Source/Editor/SkeletonEditor/Public/BoneProxy.h
+++ b/Engine/Source/Editor/SkeletonEditor/Public/BoneProxy.h
@@ -7,11 +7,17 @@
 #include "TickableEditorObject.h"
 #include "UObject/Object.h"
 #include "UObject/WeakObjectPtr.h"
+#if UE_ENABLE_INCLUDE_ORDER_DEPRECATED_IN_5_2
 #include "SAdvancedTransformInputBox.h"
+#endif
 #include "BoneProxy.generated.h"
 
-
 class UDebugSkelMeshComponent;
+enum class ETransformType : uint8;
+namespace ESlateTransformComponent { enum Type : int; }
+namespace ESlateRotationRepresentation { enum Type : int; }
+namespace ESlateTransformSubComponent { enum Type : int; }
+namespace ETextCommit { enum Type : int; }
 
 /** Proxy object used to display/edit bone transforms */
 UCLASS()
diff --git a/Engine/Source/Editor/UnrealEd/Classes/Editor/EditorEngine.h b/Engine/Source/Editor/UnrealEd/Classes/Editor/EditorEngine.h
--- a/Engine/Source/Editor/UnrealEd/Classes/Editor/EditorEngine.h
+++ b/Engine/Source/Editor/UnrealEd/Classes/Editor/EditorEngine.h
@@ -19,8 +19,10 @@
 #include "Engine/Brush.h"
 #include "Model.h"
 #include "Engine/Engine.h"
+#if UE_ENABLE_INCLUDE_ORDER_DEPRECATED_IN_5_2
 #include "Settings/LevelEditorPlaySettings.h"
 #include "Settings/LevelEditorViewportSettings.h"
+#endif
 #include "Misc/CompilationResult.h"
 #include "Interfaces/ITargetPlatform.h"
 #include "Interfaces/ITargetPlatformManagerModule.h"
@@ -39,6 +41,7 @@ class APlayerStart;
 class Error;
 class FEditorViewportClient;
 class FEditorWorldManager;
+class FLevelEditorViewportClient;
 class FMessageLog;
 class FOutputLogErrorsToMessageLogProxy;
 class FPoly;
@@ -58,6 +61,7 @@ class UBrushBuilder;
 class UFoliageType;
 class UFbxImportUI;
 class UGameViewportClient;
+class ULevelEditorPlaySettings;
 class ULocalPlayer;
 class UNetDriver;
 class UPrimitiveComponent;
@@ -67,6 +71,7 @@ class USoundNode;
 class UTextureRenderTarget2D;
 class UTransactor;
 class FTransactionObjectEvent;
+enum ERotationGridMode : int;
 struct FEdge;
 struct FTransactionContext;
 struct FEditorTransactionDeltaContext;
diff --git a/Engine/Source/Editor/UnrealEd/Classes/Settings/LevelEditorViewportSettings.h b/Engine/Source/Editor/UnrealEd/Classes/Settings/LevelEditorViewportSettings.h
--- a/Engine/Source/Editor/UnrealEd/Classes/Settings/LevelEditorViewportSettings.h
+++ b/Engine/Source/Editor/UnrealEd/Classes/Settings/LevelEditorViewportSettings.h
@@ -16,7 +16,7 @@
  * Enumerates modes for the viewport's rotation grid.
  */
 UENUM()
-enum ERotationGridMode
+enum ERotationGridMode : int
 {
 	/** Using Divisions of 360 degrees (e.g 360/2. 360/3, 360/4, ... ). */
 	GridMode_DivisionsOf360,
diff --git a/Engine/Source/Runtime/Engine/Classes/Animation/AnimInstance.h b/Engine/Source/Runtime/Engine/Classes/Animation/AnimInstance.h
--- a/Engine/Source/Runtime/Engine/Classes/Animation/AnimInstance.h
+++ b/Engine/Source/Runtime/Engine/Classes/Animation/AnimInstance.h
@@ -6,12 +6,15 @@
 #include "UObject/ObjectMacros.h"
 #include "UObject/Object.h"
 #include "Animation/AnimTypes.h"
+#include "Animation/AttributesRuntime.h"
 #include "Animation/Skeleton.h"
 #include "Animation/AnimationAsset.h"
 #include "Animation/AnimCurveTypes.h"
 #include "Animation/AnimMontage.h"
 #include "BonePose.h"
+#if UE_ENABLE_INCLUDE_ORDER_DEPRECATED_IN_5_2
 #include "Components/SkeletalMeshComponent.h"
+#endif
 #include "Animation/AnimNotifyQueue.h"
 #include "Animation/AnimSubsystemInstance.h"
 #include "Animation/AnimNotifies/AnimNotify.h"
@@ -25,7 +28,9 @@ class FDebugDisplayInfo;
 class IAnimClassInterface;
 class UAnimInstance;
 class UCanvas;
+class USkeletalMeshComponent;
 struct FAnimInstanceProxy;
+struct FAnimationEvaluationContext;
 struct FAnimNode_AssetPlayerBase;
 struct FAnimNode_AssetPlayerRelevancyBase;
 struct FAnimNode_StateMachine;
@@ -1454,7 +1459,7 @@ public:
 	void RecalcRequiredCurves(const FCurveEvaluationOption& CurveEvalOption);
 
 	// @todo document
-	inline USkeletalMeshComponent* GetSkelMeshComponent() const { return CastChecked<USkeletalMeshComponent>(GetOuter()); }
+	USkeletalMeshComponent* GetSkelMeshComponent() const;
 
 	virtual UWorld* GetWorld() const override;
 
@@ -1568,11 +1573,9 @@ protected:
 		{
 			check(IsInGameThread());
 			UObject* OuterObj = InAnimInstance->GetOuter();
-			if (OuterObj && OuterObj->IsA<USkeletalMeshComponent>())
+			if (IsSkeletalMeshComponent(OuterObj))
 			{
-				bool bBlockOnTask = true;
-				bool bPerformPostAnimEvaluation = true;
-				InAnimInstance->GetSkelMeshComponent()->HandleExistingParallelEvaluationTask(bBlockOnTask, bPerformPostAnimEvaluation);
+				HandleExistingParallelEvaluationTask(InAnimInstance->GetSkelMeshComponent());
 			}
 			if (InAnimInstance->AnimInstanceProxy == nullptr)
 			{
@@ -1595,11 +1598,9 @@ protected:
 	FORCEINLINE const T& GetProxyOnGameThread() const
 	{
 		check(IsInGameThread());
-		if(GetOuter() && GetOuter()->IsA<USkeletalMeshComponent>())
+		if(IsSkeletalMeshComponent(GetOuter()))
 		{
-			bool bBlockOnTask = true;
-			bool bPerformPostAnimEvaluation = true;
-			GetSkelMeshComponent()->HandleExistingParallelEvaluationTask(bBlockOnTask, bPerformPostAnimEvaluation);
+			HandleExistingParallelEvaluationTask(GetSkelMeshComponent());
 		}
 		if(AnimInstanceProxy == nullptr)
 		{
@@ -1612,13 +1613,11 @@ protected:
 	template <typename T/* = FAnimInstanceProxy*/>	// @TODO: Cant default parameters to this function on Xbox One until we move off the VS2012 compiler
 	FORCEINLINE T& GetProxyOnAnyThread()
 	{
-		if(GetOuter() && GetOuter()->IsA<USkeletalMeshComponent>())
+		if(IsSkeletalMeshComponent(GetOuter()))
 		{
 			if(IsInGameThread())
 			{
-				bool bBlockOnTask = true;
-				bool bPerformPostAnimEvaluation = true;
-				GetSkelMeshComponent()->HandleExistingParallelEvaluationTask(bBlockOnTask, bPerformPostAnimEvaluation);
+				HandleExistingParallelEvaluationTask(GetSkelMeshComponent());
 			}
 		}
 		if(AnimInstanceProxy == nullptr)
@@ -1632,13 +1631,11 @@ protected:
 	template <typename T/* = FAnimInstanceProxy*/>	// @TODO: Cant default parameters to this function on Xbox One until we move off the VS2012 compiler
 	FORCEINLINE const T& GetProxyOnAnyThread() const
 	{
-		if(GetOuter() && GetOuter()->IsA<USkeletalMeshComponent>())
+		if(IsSkeletalMeshComponent(GetOuter()))
 		{
 			if(IsInGameThread())
 			{
-				bool bBlockOnTask = true;
-				bool bPerformPostAnimEvaluation = true;
-				GetSkelMeshComponent()->HandleExistingParallelEvaluationTask(bBlockOnTask, bPerformPostAnimEvaluation);
+				HandleExistingParallelEvaluationTask(GetSkelMeshComponent());
 			}
 		}
 		if(AnimInstanceProxy == nullptr)
@@ -1657,6 +1654,9 @@ public:
 	virtual bool ShouldTriggerAnimNotifyState(const UAnimNotifyState* AnimNotifyState) const;
 
 protected:
+	static bool IsSkeletalMeshComponent(const UObject* Object);
+	static void HandleExistingParallelEvaluationTask(USkeletalMeshComponent* Component);
+
 	/** Proxy object, nothing should access this from an externally-callable API as it is used as a scratch area on worker threads */
 	mutable FAnimInstanceProxy* AnimInstanceProxy;
 
diff --git a/Engine/Source/Runtime/Engine/Classes/Animation/AnimNodeAlphaOptions.h b/Engine/Source/Runtime/Engine/Classes/Animation/AnimNodeAlphaOptions.h
--- a/Engine/Source/Runtime/Engine/Classes/Animation/AnimNodeAlphaOptions.h
+++ b/Engine/Source/Runtime/Engine/Classes/Animation/AnimNodeAlphaOptions.h
@@ -2,4 +2,37 @@
 
 #pragma once
 
-// Placeholder for upcoming change
\ No newline at end of file
+#include "Animation/AnimInstanceProxy.h"
+
+// Alpha blending options helper functions for anim nodes
+// Assumes that the specified node contains the members:
+// AlphaInputType, ActualAlpha, AlphaScaleBias, Alpha, bAlphaBoolEnabled, AlphaCurveName
+struct FAnimNodeAlphaOptions
+{
+	// Per-tick update
+	template<typename AnimNodeType>
+	static bool Update(AnimNodeType& InAnimNode, const FAnimationUpdateContext& InContext)
+	{
+		// Determine Actual Alpha.
+		switch (InAnimNode.AlphaInputType)
+		{
+		case EAnimAlphaInputType::Float:
+			InAnimNode.ActualAlpha = InAnimNode.AlphaScaleBias.ApplyTo(InAnimNode.AlphaScaleBiasClamp.ApplyTo(InAnimNode.Alpha, InContext.GetDeltaTime()));
+			break;
+		case EAnimAlphaInputType::Bool:
+			InAnimNode.ActualAlpha = InAnimNode.AlphaBoolBlend.ApplyTo(InAnimNode.bAlphaBoolEnabled, InContext.GetDeltaTime());
+			break;
+		case EAnimAlphaInputType::Curve:
+			if (UAnimInstance* AnimInstance = Cast<UAnimInstance>(InContext.AnimInstanceProxy->GetAnimInstanceObject()))
+			{
+				InAnimNode.ActualAlpha = InAnimNode.AlphaScaleBiasClamp.ApplyTo(AnimInstance->GetCurveValue(InAnimNode.AlphaCurveName), InContext.GetDeltaTime());
+			}
+			break;
+		};
+
+		// Make sure Alpha is clamped between 0 and 1.
+		InAnimNode.ActualAlpha = FMath::Clamp<float>(InAnimNode.ActualAlpha, 0.f, 1.f);
+
+		return FAnimWeight::IsRelevant(InAnimNode.ActualAlpha);
+	}
+};
\ No newline at end of file
diff --git a/Engine/Source/Runtime/Engine/Classes/Animation/InputScaleBias.h b/Engine/Source/Runtime/Engine/Classes/Animation/InputScaleBias.h
--- a/Engine/Source/Runtime/Engine/Classes/Animation/InputScaleBias.h
+++ b/Engine/Source/Runtime/Engine/Classes/Animation/InputScaleBias.h
@@ -6,7 +6,9 @@
 #include "UObject/ObjectMacros.h"
 #include "AlphaBlend.h"
 #include "Animation/AnimNodeBase.h"
+#if UE_ENABLE_INCLUDE_ORDER_DEPRECATED_IN_5_2
 #include "Animation/AnimInstanceProxy.h"
+#endif
 #include "InputScaleBias.generated.h"
 
 struct FAnimationUpdateContext;
@@ -318,36 +320,3 @@ struct ENGINE_API FInputAlphaBoolBlend
 
 	void Reinitialize() { bInitialized = false; }
 };
-
-// Alpha blending options helper functions for anim nodes
-// Assumes that the specified node contains the members:
-// AlphaInputType, ActualAlpha, AlphaScaleBias, Alpha, bAlphaBoolEnabled, AlphaCurveName
-struct FAnimNodeAlphaOptions
-{
-	// Per-tick update
-	template<typename AnimNodeType>
-	static bool Update(AnimNodeType& InAnimNode, const FAnimationUpdateContext& InContext)
-	{
-		// Determine Actual Alpha.
-		switch (InAnimNode.AlphaInputType)
-		{
-		case EAnimAlphaInputType::Float:
-			InAnimNode.ActualAlpha = InAnimNode.AlphaScaleBias.ApplyTo(InAnimNode.AlphaScaleBiasClamp.ApplyTo(InAnimNode.Alpha, InContext.GetDeltaTime()));
-			break;
-		case EAnimAlphaInputType::Bool:
-			InAnimNode.ActualAlpha = InAnimNode.AlphaBoolBlend.ApplyTo(InAnimNode.bAlphaBoolEnabled, InContext.GetDeltaTime());
-			break;
-		case EAnimAlphaInputType::Curve:
-			if (UAnimInstance* AnimInstance = Cast<UAnimInstance>(InContext.AnimInstanceProxy->GetAnimInstanceObject()))
-			{
-				InAnimNode.ActualAlpha = InAnimNode.AlphaScaleBiasClamp.ApplyTo(AnimInstance->GetCurveValue(InAnimNode.AlphaCurveName), InContext.GetDeltaTime());
-			}
-			break;
-		};
-
-		// Make sure Alpha is clamped between 0 and 1.
-		InAnimNode.ActualAlpha = FMath::Clamp<float>(InAnimNode.ActualAlpha, 0.f, 1.f);
-
-		return FAnimWeight::IsRelevant(InAnimNode.ActualAlpha);
-	}
-};
\ No newline at end of file
diff --git a/Engine/Source/Runtime/Engine/Classes/Animation/Skeleton.h b/Engine/Source/Runtime/Engine/Classes/Animation/Skeleton.h
--- a/Engine/Source/Runtime/Engine/Classes/Animation/Skeleton.h
+++ b/Engine/Source/Runtime/Engine/Classes/Animation/Skeleton.h
@@ -16,7 +16,9 @@
 #include "Animation/PreviewAssetAttachComponent.h"
 #include "Animation/SmartName.h"
 #include "Engine/AssetUserData.h"
+#if UE_ENABLE_INCLUDE_ORDER_DEPRECATED_IN_5_2
 #include "Engine/SkeletalMesh.h"
+#endif
 #include "HAL/CriticalSection.h"
 #include "Interfaces/Interface_AssetUserData.h"
 #include "Interfaces/Interface_PreviewMeshProvider.h"
@@ -27,6 +29,7 @@ class UAnimSequence;
 class UBlendProfile;
 class URig;
 class USkeletalMeshSocket;
+class USkinnedAsset;
 class FPackageReloadedEvent;
 struct FAssetData;
 enum class EPackageReloadPhase : uint8;
diff --git a/Engine/Source/Runtime/Engine/Classes/Components/ActorComponent.h b/Engine/Source/Runtime/Engine/Classes/Components/ActorComponent.h
--- a/Engine/Source/Runtime/Engine/Classes/Components/ActorComponent.h
+++ b/Engine/Source/Runtime/Engine/Classes/Components/ActorComponent.h
@@ -11,7 +11,9 @@
 #include "Engine/EngineTypes.h"
 #include "Engine/EngineBaseTypes.h"
 #include "UObject/ScriptMacros.h"
+#if UE_ENABLE_INCLUDE_ORDER_DEPRECATED_IN_5_2
 #include "EdGraph/EdGraphPin.h"
+#endif
 #include "Interfaces/Interface_AssetUserData.h"
 #include "UObject/StructOnScope.h"
 #include "PropertyPairsMap.h"
@@ -26,7 +28,7 @@ class UAssetUserData;
 class ULevel;
 class UWorld;
 class UPrimitiveComponent;
-
+struct FSimpleMemberReference;
 
 ENGINE_API extern int32 GEnableDeferredPhysicsCreation;
 
diff --git a/Engine/Source/Runtime/Engine/Classes/Components/AudioComponent.h b/Engine/Source/Runtime/Engine/Classes/Components/AudioComponent.h
--- a/Engine/Source/Runtime/Engine/Classes/Components/AudioComponent.h
+++ b/Engine/Source/Runtime/Engine/Classes/Components/AudioComponent.h
@@ -12,7 +12,10 @@
 #include "Sound/QuartzQuantizationUtilities.h"
 #include "Sound/QuartzSubscription.h"
 #include "Sound/SoundAttenuation.h"
+#include "Sound/SoundModulationDestination.h"
+#if UE_ENABLE_INCLUDE_ORDER_DEPRECATED_IN_5_2
 #include "Sound/SoundWave.h"
+#endif
 #include "UObject/ObjectMacros.h"
 
 #include "AudioComponent.generated.h"
@@ -20,13 +23,16 @@
 
 // Forward Declarations
 class FAudioDevice;
+class ISourceBufferListener;
 class UAudioComponent;
 class USoundBase;
 class USoundClass;
 class USoundConcurrency;
+class USoundEffectSourcePresetChain;
 class USoundWave;
-
+enum class EBusSendType : uint8;
 struct FAudioComponentParam;
+using FSharedISourceBufferListenerPtr = TSharedPtr<ISourceBufferListener, ESPMode::ThreadSafe>;
 
 
 // Enum describing the audio component play state
diff --git a/Engine/Source/Runtime/Engine/Classes/Components/ExternalMorphSet.h b/Engine/Source/Runtime/Engine/Classes/Components/ExternalMorphSet.h
--- a/Engine/Source/Runtime/Engine/Classes/Components/ExternalMorphSet.h
+++ b/Engine/Source/Runtime/Engine/Classes/Components/ExternalMorphSet.h
@@ -1,5 +1,57 @@
 // Copyright Epic Games, Inc. All Rights Reserved.
 
+
 #pragma once
 
-// Placeholder for upcoming change
\ No newline at end of file
+#include "Rendering/MorphTargetVertexInfoBuffers.h"
+
+/** An external morph target set. External morph targets are managed by systems outside of the skinned meshes. */
+struct ENGINE_API FExternalMorphSet
+{
+	/** A name for this set, useful for debugging. */
+	FName Name = FName(TEXT("Unknown"));
+
+	/** The GPU compressed morph buffers. */
+	FMorphTargetVertexInfoBuffers MorphBuffers;
+};
+
+/** The weight data for a specific external morph set. */
+struct ENGINE_API FExternalMorphSetWeights
+{
+	/** Update the number of active morph targets. */
+	void UpdateNumActiveMorphTargets();
+
+	/** Set all weights to 0. Optionally set the NumActiveMorphTargets to zero as well. */
+	void ZeroWeights(bool bZeroNumActiveMorphTargets=true);
+
+	/** The debug name. */
+	FName Name = FName(TEXT("Unknown ExternalMorphSetWeights"));
+
+	/** The weights, which can also be negative and go beyond 1.0 or -1.0. */
+	TArray<float> Weights;
+
+	/** The number of active morph targets. */
+	int32 NumActiveMorphTargets = 0;
+
+	/** The treshold used to determine if a morph target is active or not. Any weight equal to or above this value is seen as active morph target. */
+	float ActiveWeightThreshold = 0.001f;
+};
+
+/** The morph target weight data for all external morph target sets. */
+struct ENGINE_API FExternalMorphWeightData
+{
+	/** Update the number of active morph targets for all sets. */
+	void UpdateNumActiveMorphTargets();
+
+	/** Reset the morph target sets. */
+	void Reset() { MorphSets.Reset(); NumActiveMorphTargets = 0; }
+
+	/** Check if we have active morph targets or not. */
+	bool HasActiveMorphs() const { return (NumActiveMorphTargets > 0); }
+
+	/** The map with a collection of morph sets. Each set can contains multiple morph targets. */
+	TMap<int32, FExternalMorphSetWeights> MorphSets;
+
+	/** The number of active morph targets. */
+	int32 NumActiveMorphTargets = 0;
+};
diff --git a/Engine/Source/Runtime/Engine/Classes/Components/LineBatchComponent.h b/Engine/Source/Runtime/Engine/Classes/Components/LineBatchComponent.h
--- a/Engine/Source/Runtime/Engine/Classes/Components/LineBatchComponent.h
+++ b/Engine/Source/Runtime/Engine/Classes/Components/LineBatchComponent.h
@@ -6,8 +6,11 @@
 #include "CoreMinimal.h"
 #include "UObject/ObjectMacros.h"
 #include "Components/PrimitiveComponent.h"
+#if UE_ENABLE_INCLUDE_ORDER_DEPRECATED_IN_5_2
 #include "PrimitiveSceneProxy.h"
 #include "DynamicMeshBuilder.h"
+#endif
+
 #include "LineBatchComponent.generated.h"
 
 class FPrimitiveSceneProxy;
@@ -179,31 +182,3 @@ class ULineBatchComponent : public UPrimitiveComponent
 	/** Clear all batched lines, points and meshes */
 	ENGINE_API void Flush();
 };
-
-
-
-
-/** Represents a LineBatchComponent to the scene manager. */
-class ENGINE_API FLineBatcherSceneProxy : public FPrimitiveSceneProxy
-{
-public:
-	SIZE_T GetTypeHash() const override;
-
-	FLineBatcherSceneProxy(const ULineBatchComponent* InComponent);
-
-	virtual void GetDynamicMeshElements(const TArray<const FSceneView*>& Views, const FSceneViewFamily& ViewFamily, uint32 VisibilityMap, FMeshElementCollector& Collector) const override;
-
-	/**
-	*  Returns a struct that describes to the renderer when to draw this proxy.
-	*	@param		Scene view to use to determine our relevence.
-	*  @return		View relevance struct
-	*/
-	virtual FPrimitiveViewRelevance GetViewRelevance(const FSceneView* View) const override;
-	virtual uint32 GetMemoryFootprint( void ) const override;
-	uint32 GetAllocatedSize( void ) const;
-
-private:
-	TArray<FBatchedLine> Lines;
-	TArray<FBatchedPoint> Points;
-	TArray<FBatchedMesh> Meshes;
-};
diff --git a/Engine/Source/Runtime/Engine/Classes/Components/MeshComponent.h b/Engine/Source/Runtime/Engine/Classes/Components/MeshComponent.h
--- a/Engine/Source/Runtime/Engine/Classes/Components/MeshComponent.h
+++ b/Engine/Source/Runtime/Engine/Classes/Components/MeshComponent.h
@@ -7,10 +7,15 @@
 #include "UObject/ObjectMacros.h"
 #include "Engine/TextureStreamingTypes.h"
 #include "Components/PrimitiveComponent.h"
+#if UE_ENABLE_INCLUDE_ORDER_DEPRECATED_IN_5_2
 #include "Materials/MaterialInterface.h"
+#endif
 #include "Containers/SortedMap.h"
 #include "MeshComponent.generated.h"
 
+class UMaterialInterface;
+struct FMaterialRelevance;
+
 /**
  * MeshComponent is an abstract base for any component that is an instance of a renderable collection of triangles.
  *
diff --git a/Engine/Source/Runtime/Engine/Classes/Components/SceneComponent.h b/Engine/Source/Runtime/Engine/Classes/Components/SceneComponent.h
--- a/Engine/Source/Runtime/Engine/Classes/Components/SceneComponent.h
+++ b/Engine/Source/Runtime/Engine/Classes/Components/SceneComponent.h
@@ -11,7 +11,9 @@
 #include "Engine/OverlapInfo.h"
 #include "Engine/ScopedMovementUpdate.h"
 #endif
+#if UE_ENABLE_INCLUDE_ORDER_DEPRECATED_IN_5_2
 #include "Engine/HitResult.h"
+#endif
 #include "ComponentInstanceDataCache.h"
 #include "Components/ActorComponent.h"
 #include "RHIDefinitions.h"
@@ -21,6 +23,7 @@ class AActor;
 class APhysicsVolume;
 class USceneComponent;
 class FScopedMovementUpdate;
+struct FHitResult;
 struct FLevelCollection;
 
 struct FOverlapInfo;
diff --git a/Engine/Source/Runtime/Engine/Classes/Components/SkeletalMeshComponent.h b/Engine/Source/Runtime/Engine/Classes/Components/SkeletalMeshComponent.h
--- a/Engine/Source/Runtime/Engine/Classes/Components/SkeletalMeshComponent.h
+++ b/Engine/Source/Runtime/Engine/Classes/Components/SkeletalMeshComponent.h
@@ -9,22 +9,22 @@
 #include "Components/SceneComponent.h"
 #include "EngineDefines.h"
 #include "CollisionQueryParams.h"
-#include "SkeletalMeshTypes.h"
 #include "Interfaces/Interface_CollisionDataProvider.h"
+#if UE_ENABLE_INCLUDE_ORDER_DEPRECATED_IN_5_2
+#include "SkeletalMeshTypes.h"
 #include "Engine/SkeletalMesh.h"
+#include "ClothCollisionPrim.h"
+#include "PhysicsEngine/PhysicsAsset.h"
+#endif
 #include "Animation/AnimationAsset.h"
 #include "Animation/AnimCurveTypes.h"
 #include "Components/SkinnedMeshComponent.h"
 #include "ClothSimData.h"
 #include "SingleAnimationPlayData.h"
 #include "Animation/PoseSnapshot.h"
-#include "SkeletalMeshTypes.h"
-
 #include "ClothingSystemRuntimeTypes.h"
 #include "ClothingSimulationInterface.h"
 #include "ClothingSimulationFactory.h"
-#include "ClothCollisionPrim.h"
-#include "PhysicsEngine/PhysicsAsset.h"
 #include "Animation/AttributesRuntime.h"
 #if WITH_ENGINE
 #include "Engine/PoseWatchRenderData.h"
@@ -39,8 +39,11 @@ class FCanvas;
 class FSceneView;
 class UAnimInstance;
 class UPhysicalMaterial;
+class USkeletalMesh;
 class USkeletalMeshComponent;
+struct FClothCollisionSource;
 struct FConstraintInstance;
+struct FConstraintProfileProperties;
 struct FNavigableGeometryExport;
 struct FCompactPose;
 
@@ -62,18 +65,6 @@ DECLARE_DYNAMIC_MULTICAST_DELEGATE(FOnBoneTransformsFinalized);  // Deprecated,
 
 DECLARE_MULTICAST_DELEGATE(FOnBoneTransformsFinalizedMultiCast);
 
-UENUM()
-enum class EAnimCurveType : uint8 
-{
-	AttributeCurve,
-	MaterialCurve, 
-	MorphTargetCurve, 
-	// make sure to update MaxCurve 
-	MaxAnimCurveType UMETA(Hidden)
-};
-
-ENUM_RANGE_BY_COUNT(EAnimCurveType, EAnimCurveType::MaxAnimCurveType);
-
 /** Method used when retrieving a attribute value*/
 UENUM()
 enum class ECustomBoneAttributeLookup : uint8
@@ -302,35 +293,6 @@ struct ENGINE_API FClosestPointOnPhysicsAsset
 	}
 };
 
-/** Helper struct used to store info about a cloth collision source */
-struct FClothCollisionSource
-{
-	FClothCollisionSource(USkeletalMeshComponent* InSourceComponent, UPhysicsAsset* InSourcePhysicsAsset, const FOnBoneTransformsFinalizedMultiCast::FDelegate& InOnBoneTransformsFinalizedDelegate);
-	ENGINE_API ~FClothCollisionSource();
-
-	/** Component that collision data will be copied from */
-	TWeakObjectPtr<USkeletalMeshComponent> SourceComponent;
-
-	/** Physics asset to use to generate collision against the source component */
-	TWeakObjectPtr<UPhysicsAsset> SourcePhysicsAsset;
-
-	/** Callback used to remove the cloth transform updates delegate */
-	FDelegateHandle OnBoneTransformsFinalizedHandle;
-
-	/** Cached skeletal mesh used to invalidate the cache if the skeletal mesh has changed */
-	TWeakObjectPtr<USkeletalMesh> CachedSkeletalMesh;
-
-	/** Cached spheres from physics asset */
-	TArray<FClothCollisionPrim_Sphere> CachedSpheres;
-
-	/** Cached sphere connections from physics asset */
-	TArray<FClothCollisionPrim_SphereConnection> CachedSphereConnections;
-
-	/** Flag whether the cache is valid */
-	bool bCached;
-};
-
-
 /**
  * SkeletalMeshComponent is used to create an instance of an animated SkeletalMesh asset.
  *
diff --git a/Engine/Source/Runtime/Engine/Classes/Components/SkinnedMeshComponent.h b/Engine/Source/Runtime/Engine/Classes/Components/SkinnedMeshComponent.h
--- a/Engine/Source/Runtime/Engine/Classes/Components/SkinnedMeshComponent.h
+++ b/Engine/Source/Runtime/Engine/Classes/Components/SkinnedMeshComponent.h
@@ -8,7 +8,10 @@
 #include "UObject/ObjectMacros.h"
 #include "UObject/Object.h"
 #include "Engine/EngineTypes.h"
+#if UE_ENABLE_INCLUDE_ORDER_DEPRECATED_IN_5_2
 #include "Engine/SkeletalMesh.h"
+#include "Rendering/MorphTargetVertexInfoBuffers.h"
+#endif
 #include "Components/SceneComponent.h"
 #include "Interfaces/Interface_AsyncCompilation.h"
 #include "Engine/TextureStreamingTypes.h"
@@ -17,7 +20,6 @@
 #include "LODSyncInterface.h"
 #include "MeshDeformerInterface.h"
 #include "BoneContainer.h"
-#include "Rendering/MorphTargetVertexInfoBuffers.h"
 #include "ClothingSystemRuntimeTypes.h"
 #include "SkinnedMeshComponent.generated.h"
 
@@ -33,7 +35,10 @@ class FPositionVertexBuffer;
 class UMeshDeformer;
 class UMeshDeformerInstance;
 class UMeshDeformerInstanceSettings;
+class UMorphTarget;
 class USkinnedAsset;
+struct FExternalMorphSet;
+struct FExternalMorphWeightData;
 
 DECLARE_DELEGATE_OneParam(FOnAnimUpdateRateParamsCreated, FAnimUpdateRateParameters*)
 DECLARE_MULTICAST_DELEGATE_ThreeParams(FOnTickPose, USkinnedMeshComponent* /*SkinnedMeshComponent*/, float /*DeltaTime*/, bool /*bNeedsValidRootMotion*/)
@@ -117,60 +122,6 @@ namespace EBoneSpaces
 /** WeightIndex is an into the MorphTargetWeights array */
 using FMorphTargetWeightMap = TMap<const UMorphTarget* /* MorphTarget */, int32 /* WeightIndex */>;
 
-/** An external morph target set. External morph targets are managed by systems outside of the skinned meshes. */
-struct ENGINE_API FExternalMorphSet
-{
-	/** A name for this set, useful for debugging. */
-	FName Name = FName(TEXT("Unknown"));
-
-	/** The GPU compressed morph buffers. */
-	FMorphTargetVertexInfoBuffers MorphBuffers;
-};
-
-/** The map of external morph sets registered on the skinned mesh component. */
-using FExternalMorphSets = TMap<int32, TSharedPtr<FExternalMorphSet>>;
-
-/** The weight data for a specific external morph set. */
-struct ENGINE_API FExternalMorphSetWeights
-{
-	/** Update the number of active morph targets. */
-	void UpdateNumActiveMorphTargets();
-
-	/** Set all weights to 0. Optionally set the NumActiveMorphTargets to zero as well. */
-	void ZeroWeights(bool bZeroNumActiveMorphTargets=true);
-
-	/** The debug name. */
-	FName Name = FName(TEXT("Unknown ExternalMorphSetWeights"));
-
-	/** The weights, which can also be negative and go beyond 1.0 or -1.0. */
-	TArray<float> Weights;
-
-	/** The number of active morph targets. */
-	int32 NumActiveMorphTargets = 0;
-
-	/** The treshold used to determine if a morph target is active or not. Any weight equal to or above this value is seen as active morph target. */
-	float ActiveWeightThreshold = 0.001f;
-};
-
-/** The morph target weight data for all external morph target sets. */
-struct ENGINE_API FExternalMorphWeightData
-{
-	/** Update the number of active morph targets for all sets. */
-	void UpdateNumActiveMorphTargets();
-
-	/** Reset the morph target sets. */
-	void Reset() { MorphSets.Reset(); NumActiveMorphTargets = 0; }
-
-	/** Check if we have active morph targets or not. */
-	bool HasActiveMorphs() const { return (NumActiveMorphTargets > 0); }
-
-	/** The map with a collection of morph sets. Each set can contains multiple morph targets. */
-	TMap<int32, FExternalMorphSetWeights> MorphSets;
-
-	/** The number of active morph targets. */
-	int32 NumActiveMorphTargets = 0;
-};
-
 /** Vertex skin weight info supplied for a component override. */
 USTRUCT(BlueprintType, meta = (HasNativeMake = "/Script/Engine.KismetRenderingLibrary.MakeSkinWeightInfo", HasNativeBreak = "/Script/Engine.KismetRenderingLibrary.BreakSkinWeightInfo"))
 struct FSkelMeshSkinWeightInfo
@@ -231,6 +182,9 @@ struct FVertexOffsetUsage
 	int32 Usage = 0;
 };
 
+/** The map of external morph sets registered on the skinned mesh component. */
+using FExternalMorphSets = TMap<int32, TSharedPtr<FExternalMorphSet>>;
+
 /**
  *
  * Skinned mesh component that supports bone skinned mesh rendering.
@@ -258,6 +212,9 @@ private:
 	TObjectPtr<class USkinnedAsset> SkinnedAsset;
 
 public:
+	USkinnedMeshComponent(FVTableHelper& Helper);
+	~USkinnedMeshComponent();
+
 	//
 	// LeaderPoseComponent.
 	//
@@ -428,10 +385,10 @@ public:
 	const TArray<int32>& GetMasterBoneMap() const { return GetLeaderBoneMap(); }
 
 	/** Get the weights in read-only mode for a given external morph target set at a specific LOD. */
-	const FExternalMorphWeightData& GetExternalMorphWeights(int32 LOD) const { return ExternalMorphWeightData[LOD]; }
+	const FExternalMorphWeightData& GetExternalMorphWeights(int32 LOD) const;
 
 	/** Get the weights for a given external morph target set at a specific LOD. */
-	FExternalMorphWeightData& GetExternalMorphWeights(int32 LOD) { return ExternalMorphWeightData[LOD]; }
+	FExternalMorphWeightData& GetExternalMorphWeights(int32 LOD);
 
 	/**
 	 * Register an external set of GPU compressed morph targets.
@@ -933,7 +890,7 @@ public:
 	class USkeletalMesh* GetSkeletalMesh_DEPRECATED() const;
 
 	UE_DEPRECATED(5.1, "Use USkeletalMeshComponent::SetSkinnedAssetAndUpdate() instead.")
-	void SetSkeletalMesh_DEPRECATED(USkeletalMesh* NewMesh) { SetSkinnedAssetAndUpdate(Cast<USkinnedAsset>(NewMesh)); }
+	void SetSkeletalMesh_DEPRECATED(USkeletalMesh* NewMesh);
 
 	/**
 	 * Change the SkinnedAsset that is rendered for this Component. Will re-initialize the animation tree etc.
diff --git a/Engine/Source/Runtime/Engine/Classes/Components/StaticMeshComponent.h b/Engine/Source/Runtime/Engine/Classes/Components/StaticMeshComponent.h
--- a/Engine/Source/Runtime/Engine/Classes/Components/StaticMeshComponent.h
+++ b/Engine/Source/Runtime/Engine/Classes/Components/StaticMeshComponent.h
@@ -11,7 +11,9 @@
 #include "Engine/TextureStreamingTypes.h"
 #include "Components/MeshComponent.h"
 #include "PackedNormal.h"
+#if UE_ENABLE_INCLUDE_ORDER_DEPRECATED_IN_5_2
 #include "RawIndexBuffer.h"
+#endif
 #include "Templates/UniquePtr.h"
 #include "Runtime/Launch/Resources/Version.h"
 #include "UObject/RenderingObjectVersion.h"
@@ -31,6 +33,7 @@ struct FConvexVolume;
 struct FEngineShowFlags;
 struct FNavigableGeometryExport;
 struct FNavigationRelevantData;
+struct FStaticMeshComponentLODInfo;
 struct FStaticLightingPrimitiveInfo;
 
 /** Whether FStaticMeshSceneProxy should to store data and enable codepaths needed for debug rendering */
@@ -83,95 +86,6 @@ struct FPaintedVertex
 	
 };
 
-struct FPreCulledStaticMeshSection
-{
-	/** Range of vertices and indices used when rendering this section. */
-	uint32 FirstIndex;
-	uint32 NumTriangles;
-};
-
-USTRUCT()
-struct FStaticMeshComponentLODInfo
-{
-	GENERATED_USTRUCT_BODY()
-
-	/** Uniquely identifies this LOD's built map data. */
-	FGuid MapBuildDataId;
-
-	/** Used during deserialization to temporarily store legacy lightmap data. */
-	FMeshMapBuildData* LegacyMapBuildData;
-
-	/** 
-	 * Transient override lightmap data, used by landscape grass.
-	 * Be sure to add your component to UMapBuildDataRegistry::CleanupTransientOverrideMapBuildData() for proper cleanup
-	 * so that you don't get stale rendering resource references if the underlying MapBuildData is gone (lighting scenario changes, new static lighting build, etc.)
-	 */
-	TUniquePtr<FMeshMapBuildData> OverrideMapBuildData;
-
-	/** Vertex data cached at the time this LOD was painted, if any */
-	TArray<struct FPaintedVertex> PaintedVertices;
-
-	/** Vertex colors to use for this mesh LOD */
-	FColorVertexBuffer* OverrideVertexColors;
-
-	/** Information for each section about what range of PreCulledIndexBuffer to use.  If no preculled index data is available, PreCulledSections will be empty. */
-	TArray<FPreCulledStaticMeshSection> PreCulledSections;
-
-	FRawStaticIndexBuffer PreCulledIndexBuffer;
-
-	/** 
-	 * Owner of this FStaticMeshComponentLODInfo 
-	 * Warning, can be NULL for a component created via SpawnActor off of a blueprint default (LODData will be created without a call to SetLODDataCount).
-	 */
-	class UStaticMeshComponent* OwningComponent;
-
-	/** Default constructor */
-	FStaticMeshComponentLODInfo();
-	FStaticMeshComponentLODInfo(UStaticMeshComponent* InOwningComponent);
-	/** Destructor */
-	~FStaticMeshComponentLODInfo();
-
-	/** Delete existing resources */
-	void CleanUp();
-
-	/** 
-	 * Ensure this LODInfo has a valid MapBuildDataId GUID.
-	 * @param LodIndex Index of the LOD this LODInfo represents.
-	 * @return true if a new GUID was created, false otherwise.
-	 */
-	ENGINE_API bool CreateMapBuildDataId(int32 LodIndex);
-
-	/**
-	* Enqueues a rendering command to release the vertex colors.
-	* The game thread must block until the rendering thread has processed the command before deleting OverrideVertexColors.
-	*/
-	ENGINE_API void BeginReleaseOverrideVertexColors();
-
-	ENGINE_API void ReleaseOverrideVertexColorsAndBlock();
-
-	void ReleaseResources();
-
-	/** Methods for importing and exporting the painted vertex array to text */
-	void ExportText(FString& ValueStr);
-	void ImportText(const TCHAR** SourceText);
-
-	/** Serializer. */
-	friend FArchive& operator<<(FArchive& Ar,FStaticMeshComponentLODInfo& I);
-
-private:
-	/** Purposely hidden */
-	FStaticMeshComponentLODInfo &operator=( const FStaticMeshComponentLODInfo &rhs ) { check(0); return *this; }
-};
-
-template<>
-struct TStructOpsTypeTraits<FStaticMeshComponentLODInfo> : public TStructOpsTypeTraitsBase2<FStaticMeshComponentLODInfo>
-{
-	enum
-	{
-		WithCopy = false
-	};
-};
-
 /**
  * StaticMeshComponent is used to create an instance of a UStaticMesh.
  * A static mesh is a piece of geometry that consists of a static set of polygons.
diff --git a/Engine/Source/Runtime/Engine/Classes/Engine/Blueprint.h b/Engine/Source/Runtime/Engine/Classes/Engine/Blueprint.h
--- a/Engine/Source/Runtime/Engine/Classes/Engine/Blueprint.h
+++ b/Engine/Source/Runtime/Engine/Classes/Engine/Blueprint.h
@@ -11,7 +11,10 @@
 #include "Engine/EngineTypes.h"
 #include "EdGraph/EdGraphPin.h"
 #include "Engine/BlueprintCore.h"
+#include "Blueprint/BlueprintPropertyGuidProvider.h"
+#if UE_ENABLE_INCLUDE_ORDER_DEPRECATED_IN_5_2
 #include "Engine/BlueprintGeneratedClass.h"
+#endif
 #include "UObject/SoftObjectPath.h"
 #include "Blueprint/BlueprintSupport.h"
 #include "Blueprint.generated.h"
@@ -23,7 +26,9 @@ class UEdGraph;
 class FKismetCompilerContext;
 class UInheritableComponentHandler;
 class UBlueprintExtension;
+class UBlueprintGeneratedClass;
 class FBlueprintActionDatabaseRegistrar;
+struct FBPComponentClassOverride;
 struct FDiffResults;
 
 /**
diff --git a/Engine/Source/Runtime/Engine/Classes/Engine/BlueprintGeneratedClass.h b/Engine/Source/Runtime/Engine/Classes/Engine/BlueprintGeneratedClass.h
--- a/Engine/Source/Runtime/Engine/Classes/Engine/BlueprintGeneratedClass.h
+++ b/Engine/Source/Runtime/Engine/Classes/Engine/BlueprintGeneratedClass.h
@@ -2,6 +2,7 @@
 
 #pragma once
 
+#include "Blueprint/BlueprintPropertyGuidProvider.h"
 #include "CoreMinimal.h"
 #include "Containers/IndirectArray.h"
 #include "Stats/Stats.h"
@@ -625,28 +626,6 @@ struct FBPComponentClassOverride
 	}
 };
 
-/**
- * Interface to query the property name<->GUID relationship using either a UBlueprint or a UBlueprintGeneratedClass.
- * This allows cooked and uncooked Blueprints to be queried via the same API.
- */
-class IBlueprintPropertyGuidProvider
-{
-public:
-	virtual ~IBlueprintPropertyGuidProvider() = default;
-
-	/**
-	 * Returns the property name for the given GUID, if any.
-	 * @note Does not recurse into parents.
-	 */
-	virtual FName FindBlueprintPropertyNameFromGuid(const FGuid& PropertyGuid) const = 0;
-
-	/**
-	 * Returns the property GUID for the given name, if any.
-	 * @note Does not recurse into parents.
-	 */
-	virtual FGuid FindBlueprintPropertyGuidFromName(const FName PropertyName) const = 0;
-};
-
 UCLASS(NeedsDeferredDependencyLoading)
 class ENGINE_API UBlueprintGeneratedClass : public UClass, public IBlueprintPropertyGuidProvider
 {
diff --git a/Engine/Source/Runtime/Engine/Classes/Engine/Engine.h b/Engine/Source/Runtime/Engine/Classes/Engine/Engine.h
--- a/Engine/Source/Runtime/Engine/Classes/Engine/Engine.h
+++ b/Engine/Source/Runtime/Engine/Classes/Engine/Engine.h
@@ -16,8 +16,11 @@
 #include "Misc/FrameRate.h"
 #include "Subsystems/SubsystemCollection.h"
 #include "Subsystems/EngineSubsystem.h"
+#include "RHIDefinitions.h"
+#if UE_ENABLE_INCLUDE_ORDER_DEPRECATED_IN_5_2
 #include "RHI.h"
 #include "AudioDeviceManager.h"
+#endif
 #include "Templates/UniqueObj.h"
 #include "Containers/Ticker.h"
 #include "DynamicRenderScaling.h"
@@ -27,6 +30,7 @@
 
 class APlayerController;
 class Error;
+class FAudioDeviceManager;
 class FCanvas;
 class FCommonViewportClient;
 class FFineGrainedPerformanceTracker;
