- Add Cloud Occlusion test to Occlusion Data Interface

diff --git a/Engine/Plugins/FX/Niagara/Shaders/Private/NiagaraDataInterfaceOcclusion.ush b/Engine/Plugins/FX/Niagara/Shaders/Private/NiagaraDataInterfaceOcclusion.ush
--- a/Engine/Plugins/FX/Niagara/Shaders/Private/NiagaraDataInterfaceOcclusion.ush
+++ b/Engine/Plugins/FX/Niagara/Shaders/Private/NiagaraDataInterfaceOcclusion.ush
@@ -1,6 +1,8 @@
 // Copyright Epic Games, Inc. All Rights Reserved.
 
-float3 {ParameterName}_SystemLWCTile;
+float3				{ParameterName}_SystemLWCTile;
+Texture2D<float4>	{ParameterName}_CloudVolumetricTexture;
+SamplerState		{ParameterName}_CloudVolumetricTextureSampler;
 
 void QueryOcclusionFactorWithRectangleGPU_{ParameterName}(in float3 In_SampleCenterWorldPos, in float In_SampleWindowWidthWorld, in float In_SampleWindowHeightWorld, in float In_SampleSteps, out float Out_VisibilityFraction, out float Out_SampleFraction)
 {
@@ -121,3 +123,61 @@ void QueryOcclusionFactorWithCircleGPU_{ParameterName}(in float3 In_SampleCenter
 	Out_VisibilityFraction = TotalSamples > 0 ? 1 - OccludedSamples / TotalSamples : 0;
 	Out_SampleFraction = Steps == 0 ? 0 : (TotalSamples / (Rays * Steps + 1));
 }
+
+void QueryCloudOcclusionWithCircle_{ParameterName}(float3 WorldPosition, float WorldDiameter, bool bIncludeCenterSample, int NumberOfRings, int SamplesPerRing, out float VisibilityFraction, out float SampleFraction)
+{
+	const float PI = 3.14159265;
+	const float SPIRAL_TURN = 2.0f * PI * 0.61803399; // use golden ratio to rotate sample pattern each ring so we get a spiral
+
+	const FLWCVector3 LWCWorldPosition = MakeLWCVector3({ParameterName}_SystemLWCTile, WorldPosition);
+
+	const float3 PreViewPosition = LWCToFloat(LWCAdd(LWCWorldPosition, PrimaryView.PreViewTranslation));
+	const float4 ClipPosition = mul(float4(PreViewPosition, 1), PrimaryView.TranslatedWorldToClip);
+	const float2 ScreenPosition = ClipPosition.xy / ClipPosition.w;
+	const float2 ScreenUV = (ScreenPosition + 1.0f) * 0.5f;
+	
+	const float WorldDepth = dot(LWCToFloat(LWCSubtract(LWCWorldPosition, PrimaryView.WorldCameraOrigin)), PrimaryView.ViewForward);
+	const float2 ScreenRadiusUV = (WorldDiameter * 0.5f) / (GetTanHalfFieldOfView() * 2.0f * WorldDepth);
+
+	float TotalSamples = 0.0f;
+	float TotalVisibility = 0.0f;
+
+	// Take center tap
+	float ExpectedSamples = 0.0f;
+	if ( bIncludeCenterSample && (all(ScreenUV >= 0.0f) && all(ScreenUV < 1.0f)) )
+	{
+		const float4 SampleColor = {ParameterName}_CloudVolumetricTexture.SampleLevel({ParameterName}_CloudVolumetricTextureSampler, ScreenUV, 0);
+
+		TotalVisibility += SampleColor.a;
+		TotalSamples += 1.0f;
+		ExpectedSamples += 1.0f;
+	}
+
+	// Take ring taps
+	if ( NumberOfRings > 0 && SamplesPerRing > 0 )
+	{
+		const float RadStep = 2.0f * PI / float(SamplesPerRing + 1);
+		float Rad = 0.0f;
+		ExpectedSamples += float(NumberOfRings * SamplesPerRing);
+		for ( int iRing=0; iRing < NumberOfRings; ++iRing )datainterfaceocclusion.cpp
+		{
+			const float2 RingRadiusUV = ScreenRadiusUV * (float(iRing + 1) / float(NumberOfRings));
+			Rad += SPIRAL_TURN;
+			for ( int iSample=0; iSample < SamplesPerRing; ++iSample )
+			{
+				const float2 SampleUV = ScreenUV + (RingRadiusUV * float2(sin(Rad), cos(Rad)));
+				if ( all(SampleUV >= 0.0f) && all(SampleUV < 1.0f) )
+				{
+					const float4 SampleColor = {ParameterName}_CloudVolumetricTexture.SampleLevel({ParameterName}_CloudVolumetricTextureSampler, SampleUV, 0);
+
+					TotalVisibility += SampleColor.a;
+					TotalSamples += 1.0f;
+				}
+				Rad += RadStep;
+			}
+		}
+	}
+
+	VisibilityFraction = TotalSamples > 0.0f ? TotalVisibility / TotalSamples : 0.0f;
+	SampleFraction = TotalSamples > 0.0f ? ExpectedSamples / TotalSamples : 0.0f;
+}
diff --git a/Engine/Plugins/FX/Niagara/Source/Niagara/Classes/NiagaraDataInterfaceOcclusion.h b/Engine/Plugins/FX/Niagara/Source/Niagara/Classes/NiagaraDataInterfaceOcclusion.h
--- a/Engine/Plugins/FX/Niagara/Source/Niagara/Classes/NiagaraDataInterfaceOcclusion.h
+++ b/Engine/Plugins/FX/Niagara/Source/Niagara/Classes/NiagaraDataInterfaceOcclusion.h
@@ -12,7 +12,9 @@ class NIAGARA_API UNiagaraDataInterfaceOcclusion : public UNiagaraDataInterface
 	GENERATED_UCLASS_BODY()
 
 	BEGIN_SHADER_PARAMETER_STRUCT(FShaderParameters, )
-		SHADER_PARAMETER(FVector3f, SystemLWCTile)
+		SHADER_PARAMETER(FVector3f,				SystemLWCTile)
+		SHADER_PARAMETER_RDG_TEXTURE(Texture2D, CloudVolumetricTexture)
+		SHADER_PARAMETER_SAMPLER(SamplerState,	CloudVolumetricTextureSampler)
 	END_SHADER_PARAMETER_STRUCT();
 
 public:
@@ -41,6 +43,7 @@ private:
 	static const TCHAR* TemplateShaderFilePath;
 	static const FName GetCameraOcclusionRectangleName;
 	static const FName GetCameraOcclusionCircleName;
+	static const FName QueryCloudOcclusionWithCircleName;
 };
 
 struct FNiagaraDataIntefaceProxyOcclusionQuery : public FNiagaraDataInterfaceProxy
diff --git a/Engine/Plugins/FX/Niagara/Source/Niagara/Private/NiagaraDataInterfaceOcclusion.cpp b/Engine/Plugins/FX/Niagara/Source/Niagara/Private/NiagaraDataInterfaceOcclusion.cpp
--- a/Engine/Plugins/FX/Niagara/Source/Niagara/Private/NiagaraDataInterfaceOcclusion.cpp
+++ b/Engine/Plugins/FX/Niagara/Source/Niagara/Private/NiagaraDataInterfaceOcclusion.cpp
@@ -1,15 +1,19 @@
 // Copyright Epic Games, Inc. All Rights Reserved.
 
 #include "NiagaraDataInterfaceOcclusion.h"
+#include "NiagaraGpuComputeDispatchInterface.h"
 #include "NiagaraTypes.h"
 #include "NiagaraWorldManager.h"
-#include "ShaderParameterUtils.h"
-#include "Internationalization/Internationalization.h"
 #include "NiagaraGPUSystemTick.h"
 #include "NiagaraShaderParametersBuilder.h"
 #include "NiagaraSystemGpuComputeProxy.h"
 #include "NiagaraSystemInstance.h"
 
+#include "SceneRendering.h"
+#include "SceneView.h"
+#include "ShaderParameterUtils.h"
+#include "Internationalization/Internationalization.h"
+
 #include UE_INLINE_GENERATED_CPP_BY_NAME(NiagaraDataInterfaceOcclusion)
 
 #define LOCTEXT_NAMESPACE "NiagaraDataInterfaceOcclusion"
@@ -17,6 +21,7 @@
 const TCHAR* UNiagaraDataInterfaceOcclusion::TemplateShaderFilePath = TEXT("/Plugin/FX/Niagara/Private/NiagaraDataInterfaceOcclusion.ush");
 const FName UNiagaraDataInterfaceOcclusion::GetCameraOcclusionRectangleName(TEXT("QueryOcclusionFactorWithRectangleGPU"));
 const FName UNiagaraDataInterfaceOcclusion::GetCameraOcclusionCircleName(TEXT("QueryOcclusionFactorWithCircleGPU"));
+const FName UNiagaraDataInterfaceOcclusion::QueryCloudOcclusionWithCircleName(TEXT("QueryCloudOcclusionWithCircle"));
 
 struct FNiagaraOcclusionDIFunctionVersion
 {
@@ -49,43 +54,65 @@ void UNiagaraDataInterfaceOcclusion::PostInitProperties()
 
 void UNiagaraDataInterfaceOcclusion::GetFunctions(TArray<FNiagaraFunctionSignature>& OutFunctions)
 {
-	FNiagaraFunctionSignature Sig;
-	Sig.Name = GetCameraOcclusionRectangleName;
-#if WITH_EDITORONLY_DATA
-	Sig.Description = LOCTEXT("GetCameraOcclusionRectFunctionDescription", "This function returns the occlusion factor of a sprite. It samples the depth buffer in a rectangular grid around the given world position and compares each sample with the camera distance.");
-	Sig.FunctionVersion = FNiagaraOcclusionDIFunctionVersion::LatestVersion;
-#endif
-	Sig.bMemberFunction = true;
-	Sig.bRequiresContext = false;
-	Sig.bSupportsCPU = false;
-	FText VisibilityFractionDescription = LOCTEXT("VisibilityFractionDescription", "Returns a value 0..1 depending on how many of the samples on the screen were occluded.\nFor example, a value of 0.3 means that 70% of visible samples were occluded.\nIf the sample fraction is 0 then this also returns 0.");
-	FText SampleFractionDescription = LOCTEXT("SampleFractionDescription", "Returns a value 0..1 depending on how many samples were inside the viewport or outside of it.\nFor example, a value of 0.3 means that 70% of samples were outside the current viewport and therefore not visible.");
-	Sig.AddInput(FNiagaraVariable(FNiagaraTypeDefinition(GetClass()), TEXT("Occlusion interface")));
-	Sig.AddInput(FNiagaraVariable(FNiagaraTypeDefinition::GetPositionDef(), TEXT("Sample Center World Position")), LOCTEXT("RectCenterPosDescription", "This world space position where the center of the sample rectangle should be."));
-	Sig.AddInput(FNiagaraVariable(FNiagaraTypeDefinition::GetFloatDef(), TEXT("Sample Window Width World")), LOCTEXT("SampleWindowWidthWorldDescription", "The total width of the sample rectangle in world space.\nIf the particle is a camera-aligned sprite then this is the sprite width."));
-	Sig.AddInput(FNiagaraVariable(FNiagaraTypeDefinition::GetFloatDef(), TEXT("Sample Window Height World")), LOCTEXT("SampleWindowHeightWorldDescription", "The total height of the sample rectangle in world space.\nIf the particle is a camera-aligned sprite then this is the sprite height."));
-	Sig.AddInput(FNiagaraVariable(FNiagaraTypeDefinition::GetFloatDef(), TEXT("Sample Steps Per Line")), LOCTEXT("StepsPerLineDescription", "The number of samples to take horizontally. The total number of samples is this value squared."));
-	Sig.AddOutput(FNiagaraVariable(FNiagaraTypeDefinition::GetFloatDef(), TEXT("Visibility Fraction")), VisibilityFractionDescription);
-	Sig.AddOutput(FNiagaraVariable(FNiagaraTypeDefinition::GetFloatDef(), TEXT("Sample Fraction")), SampleFractionDescription);
-	OutFunctions.Add(Sig);
-
-	Sig = FNiagaraFunctionSignature();
-	Sig.Name = GetCameraOcclusionCircleName;
-#if WITH_EDITORONLY_DATA
-	Sig.Description = LOCTEXT("GetCameraOcclusionCircleFunctionDescription", "This function returns the occlusion factor of a sprite. It samples the depth buffer in concentric rings around the given world position and compares each sample with the camera distance.");
-	Sig.FunctionVersion = FNiagaraOcclusionDIFunctionVersion::LatestVersion;
-#endif
-	Sig.bMemberFunction = true;
-	Sig.bRequiresContext = false;
-	Sig.bSupportsCPU = false;
-	Sig.AddInput(FNiagaraVariable(FNiagaraTypeDefinition(GetClass()), TEXT("Occlusion interface")));
-	Sig.AddInput(FNiagaraVariable(FNiagaraTypeDefinition::GetPositionDef(), TEXT("Sample Center World Position")), LOCTEXT("CircleCenterPosDescription", "This world space position where the center of the sample circle should be."));
-	Sig.AddInput(FNiagaraVariable(FNiagaraTypeDefinition::GetFloatDef(), TEXT("Sample Window Diameter World")), LOCTEXT("SampleWindowDiameterDescription", "The world space diameter of the circle to sample.\nIf the particle is a spherical sprite then this is the sprite size."));
-	Sig.AddInput(FNiagaraVariable(FNiagaraTypeDefinition::GetFloatDef(), TEXT("Samples per ring")),LOCTEXT("SamplesPerRingDescription", "The number of samples for each ring inside the circle.\nThe total number of samples is NumRings * SamplesPerRing."));
-	Sig.AddInput(FNiagaraVariable(FNiagaraTypeDefinition::GetFloatDef(), TEXT("Number of sample rings")), LOCTEXT("NumberOfSampleRingsDescription", "This number of concentric rings to sample inside the circle.\nThe total number of samples is NumRings * SamplesPerRing."));
-	Sig.AddOutput(FNiagaraVariable(FNiagaraTypeDefinition::GetFloatDef(), TEXT("Visibility Fraction")), VisibilityFractionDescription);
-	Sig.AddOutput(FNiagaraVariable(FNiagaraTypeDefinition::GetFloatDef(), TEXT("Sample Fraction")), SampleFractionDescription);
-	OutFunctions.Add(Sig);
+	FNiagaraFunctionSignature DefaultSig;
+	DefaultSig.bMemberFunction = true;
+	DefaultSig.bRequiresContext = false;
+	DefaultSig.bSupportsCPU = false;
+	DefaultSig.AddInput(FNiagaraVariable(FNiagaraTypeDefinition(GetClass()), TEXT("Occlusion interface")));
+	DefaultSig.SetFunctionVersion(FNiagaraOcclusionDIFunctionVersion::LatestVersion);
+
+	static const FText VisibilityFractionDescription(LOCTEXT("VisibilityFractionDescription", "Returns a value 0..1 depending on how many of the samples on the screen were occluded.\nFor example, a value of 0.3 means that 70% of visible samples were occluded.\nIf the sample fraction is 0 then this also returns 0."));
+	static const FText SampleFractionDescription(LOCTEXT("SampleFractionDescription", "Returns a value 0..1 depending on how many samples were inside the viewport or outside of it.\nFor example, a value of 0.3 means that 70% of samples were outside the current viewport and therefore not visible."));
+	static const FText CircleCenterPosDescription(LOCTEXT("CircleCenterPosDescription", "This world space position where the center of the sample circle should be."));
+	static const FText SampleWindowDiameterDescription(LOCTEXT("SampleWindowDiameterDescription", "The world space diameter of the circle to sample.\nIf the particle is a spherical sprite then this is the sprite size."));
+	static const FText SamplesPerRingDescription(LOCTEXT("SamplesPerRingDescription", "The number of samples for each ring inside the circle.\nThe total number of samples is NumRings * SamplesPerRing."));
+	static const FText NumberOfSampleRingsDescription(LOCTEXT("NumberOfSampleRingsDescription", "This number of concentric rings to sample inside the circle.\nThe total number of samples is NumRings * SamplesPerRing."));
+	static const FText IncludeCenterSampleDescription(LOCTEXT("IncludeCenterSampleDescription", "When enabled we sample the center of the circle in addition to the rings."));
+	{
+		FNiagaraFunctionSignature& Sig = OutFunctions.Add_GetRef(DefaultSig);
+		Sig.Name = GetCameraOcclusionRectangleName;
+		Sig.SetDescription(LOCTEXT("GetCameraOcclusionRectFunctionDescription", "This function returns the occlusion factor of a sprite. It samples the depth buffer in a rectangular grid around the given world position and compares each sample with the camera distance."));
+		Sig.AddInput(FNiagaraVariable(FNiagaraTypeDefinition(GetClass()), TEXT("Occlusion interface")));
+		Sig.AddInput(FNiagaraVariable(FNiagaraTypeDefinition::GetPositionDef(), TEXT("Sample Center World Position")), LOCTEXT("RectCenterPosDescription", "This world space position where the center of the sample rectangle should be."));
+		Sig.AddInput(FNiagaraVariable(FNiagaraTypeDefinition::GetFloatDef(), TEXT("Sample Window Width World")), LOCTEXT("SampleWindowWidthWorldDescription", "The total width of the sample rectangle in world space.\nIf the particle is a camera-aligned sprite then this is the sprite width."));
+		Sig.AddInput(FNiagaraVariable(FNiagaraTypeDefinition::GetFloatDef(), TEXT("Sample Window Height World")), LOCTEXT("SampleWindowHeightWorldDescription", "The total height of the sample rectangle in world space.\nIf the particle is a camera-aligned sprite then this is the sprite height."));
+		Sig.AddInput(FNiagaraVariable(FNiagaraTypeDefinition::GetFloatDef(), TEXT("Sample Steps Per Line")), LOCTEXT("StepsPerLineDescription", "The number of samples to take horizontally. The total number of samples is this value squared."));
+		Sig.AddOutput(FNiagaraVariable(FNiagaraTypeDefinition::GetFloatDef(), TEXT("Visibility Fraction")), VisibilityFractionDescription);
+		Sig.AddOutput(FNiagaraVariable(FNiagaraTypeDefinition::GetFloatDef(), TEXT("Sample Fraction")), SampleFractionDescription);
+	}
+
+	{
+		FNiagaraFunctionSignature& Sig = OutFunctions.Add_GetRef(DefaultSig);
+		Sig.Name = GetCameraOcclusionCircleName;
+		Sig.SetDescription(LOCTEXT("GetCameraOcclusionCircleFunctionDescription", "This function returns the occlusion factor of a sprite. It samples the depth buffer in concentric rings around the given world position and compares each sample with the camera distance."));
+		Sig.AddInput(FNiagaraVariable(FNiagaraTypeDefinition::GetPositionDef(), TEXT("Sample Center World Position")), CircleCenterPosDescription);
+		Sig.AddInput(FNiagaraVariable(FNiagaraTypeDefinition::GetFloatDef(), TEXT("Sample Window Diameter World")), SampleWindowDiameterDescription);
+		Sig.AddInput(FNiagaraVariable(FNiagaraTypeDefinition::GetFloatDef(), TEXT("Samples per ring")), SamplesPerRingDescription);
+		Sig.AddInput(FNiagaraVariable(FNiagaraTypeDefinition::GetFloatDef(), TEXT("Number of sample rings")), NumberOfSampleRingsDescription);
+		Sig.AddOutput(FNiagaraVariable(FNiagaraTypeDefinition::GetFloatDef(), TEXT("Visibility Fraction")), VisibilityFractionDescription);
+		Sig.AddOutput(FNiagaraVariable(FNiagaraTypeDefinition::GetFloatDef(), TEXT("Sample Fraction")), SampleFractionDescription);
+	}
+
+	{
+		FNiagaraFunctionSignature& Sig = OutFunctions.Add_GetRef(DefaultSig);
+		Sig.Name = QueryCloudOcclusionWithCircleName;
+		Sig.SetDescription(LOCTEXT("QueryCloudOcclusionWithCircleDescription", "Returns the cloud occlusion factor for the world position. "));
+		Sig.Inputs.Emplace(FNiagaraTypeDefinition::GetPositionDef(), TEXT("WorldPosition"));
+		Sig.Inputs.Emplace(FNiagaraTypeDefinition::GetFloatDef(), TEXT("WorldDiameter"));
+		Sig.Inputs.Emplace_GetRef(FNiagaraTypeDefinition::GetBoolDef(), TEXT("IncludeCenterSample")).SetValue(true);
+		Sig.Inputs.Emplace_GetRef(FNiagaraTypeDefinition::GetIntDef(), TEXT("NumberOfRings")).SetValue(1);
+		Sig.Inputs.Emplace_GetRef(FNiagaraTypeDefinition::GetIntDef(), TEXT("SamplesPerRing")).SetValue(1);
+		Sig.Outputs.Emplace(FNiagaraTypeDefinition::GetFloatDef(), TEXT("VisibilityFraction"));
+		Sig.Outputs.Emplace(FNiagaraTypeDefinition::GetFloatDef(), TEXT("SampleFraction"));
+
+		Sig.SetInputDescription(Sig.Inputs[1], CircleCenterPosDescription);
+		Sig.SetInputDescription(Sig.Inputs[2], SampleWindowDiameterDescription);
+		Sig.SetInputDescription(Sig.Inputs[3], IncludeCenterSampleDescription);
+		Sig.SetInputDescription(Sig.Inputs[4], NumberOfSampleRingsDescription);
+		Sig.SetInputDescription(Sig.Inputs[5], SamplesPerRingDescription);
+		Sig.SetOutputDescription(Sig.Outputs[0], VisibilityFractionDescription);
+		Sig.SetOutputDescription(Sig.Outputs[1], SampleFractionDescription);
+	}
 }
 
 #if WITH_EDITORONLY_DATA
@@ -106,12 +133,13 @@ void UNiagaraDataInterfaceOcclusion::GetCommonHLSL(FString& OutHLSL)
 
 bool UNiagaraDataInterfaceOcclusion::GetFunctionHLSL(const FNiagaraDataInterfaceGPUParamInfo& ParamInfo, const FNiagaraDataInterfaceGeneratedFunction& FunctionInfo, int FunctionInstanceIndex, FString& OutHLSL)
 {
-	if ((FunctionInfo.DefinitionName == GetCameraOcclusionRectangleName) ||
-		(FunctionInfo.DefinitionName == GetCameraOcclusionCircleName))
+	static const TSet<FName> ValidGpuFunctions =
 	{
-		return true;
-	}
-	return false;
+		GetCameraOcclusionRectangleName,
+		GetCameraOcclusionCircleName,
+		QueryCloudOcclusionWithCircleName,
+	};
+	return ValidGpuFunctions.Contains(FunctionInfo.DefinitionName);
 }
 
 bool UNiagaraDataInterfaceOcclusion::UpgradeFunctionCall(FNiagaraFunctionSignature& FunctionSignature)
@@ -157,6 +185,18 @@ void UNiagaraDataInterfaceOcclusion::SetShaderParameters(const FNiagaraDataInter
 {
 	FShaderParameters* ShaderParameters = Context.GetParameterNestedStruct<FShaderParameters>();
 	ShaderParameters->SystemLWCTile = Context.GetSystemLWCTile();
+
+	if (Context.IsResourceBound(&ShaderParameters->CloudVolumetricTexture))
+	{
+		TConstArrayView<FViewInfo> SimulationViewInfos = Context.GetComputeDispatchInterface().GetSimulationViewInfos();
+		FRDGTextureRef CloudTexture = SimulationViewInfos.Num() > 0 ? SimulationViewInfos[0].GetVolumetricCloudTexture(Context.GetGraphBuilder()) : nullptr;
+		if (CloudTexture == nullptr)
+		{
+			CloudTexture = Context.GetComputeDispatchInterface().GetBlackTexture(Context.GetGraphBuilder(), ETextureDimension::Texture2D);
+		}
+		ShaderParameters->CloudVolumetricTexture = CloudTexture;
+	}
+	ShaderParameters->CloudVolumetricTextureSampler = TStaticSamplerState<SF_Bilinear>::GetRHI();
 }
 
 #undef LOCTEXT_NAMESPACE
diff --git a/Engine/Source/Runtime/Engine/Public/SceneManagement.h b/Engine/Source/Runtime/Engine/Public/SceneManagement.h
--- a/Engine/Source/Runtime/Engine/Public/SceneManagement.h
+++ b/Engine/Source/Runtime/Engine/Public/SceneManagement.h
@@ -219,6 +219,9 @@ public:
 	 */
 	virtual uint32 GetViewKey() const = 0;
 
+	/* Return the active volumetric cloud texture, can be null. */
+	virtual FRDGTextureRef GetVolumetricCloudTexture(FRDGBuilder& GraphBuilder) = 0;
+
 	//
 	virtual uint32 GetCurrentTemporalAASampleIndex() const = 0;
 
diff --git a/Engine/Source/Runtime/Renderer/Private/ScenePrivate.h b/Engine/Source/Runtime/Renderer/Private/ScenePrivate.h
--- a/Engine/Source/Runtime/Renderer/Private/ScenePrivate.h
+++ b/Engine/Source/Runtime/Renderer/Private/ScenePrivate.h
@@ -1037,6 +1037,11 @@ public:
 	FVector VolumetricCloudShadowmapPreviousAnchorPoint[NUM_ATMOSPHERE_LIGHTS];
 	FVector VolumetricCloudShadowmapPreviousAtmosphericLightDir[NUM_ATMOSPHERE_LIGHTS];
 
+	virtual FRDGTextureRef GetVolumetricCloudTexture(FRDGBuilder& GraphBuilder) override
+	{
+		return VolumetricCloudRenderTarget.GetDstVolumetricReconstructRT(GraphBuilder);
+	}
+
 	FHairStrandsViewStateData HairStrandsViewStateData;
 
 	FShaderPrintStateData ShaderPrintStateData;
diff --git a/Engine/Source/Runtime/Renderer/Private/SceneRendering.cpp b/Engine/Source/Runtime/Renderer/Private/SceneRendering.cpp
--- a/Engine/Source/Runtime/Renderer/Private/SceneRendering.cpp
+++ b/Engine/Source/Runtime/Renderer/Private/SceneRendering.cpp
@@ -2150,6 +2150,15 @@ FInt32Range FViewInfo::GetDynamicMeshElementRange(uint32 PrimitiveIndex) const
 	return FInt32Range(Start, AfterEnd);
 }
 
+FRDGTextureRef FViewInfo::GetVolumetricCloudTexture(FRDGBuilder& GraphBuilder) const
+{
+	if (State)
+	{
+		return State->GetVolumetricCloudTexture(GraphBuilder);
+	}
+	return nullptr;
+}
+
 FSceneViewState* FViewInfo::GetEyeAdaptationViewState() const
 {
 	return static_cast<FSceneViewState*>(EyeAdaptationViewState);
diff --git a/Engine/Source/Runtime/Renderer/Private/SceneRendering.h b/Engine/Source/Runtime/Renderer/Private/SceneRendering.h
--- a/Engine/Source/Runtime/Renderer/Private/SceneRendering.h
+++ b/Engine/Source/Runtime/Renderer/Private/SceneRendering.h
@@ -1737,6 +1737,8 @@ public:
 	const FSceneTextures& GetSceneTextures() const;
 	const FSceneTextures* GetSceneTexturesChecked() const;
 
+	RENDERER_API FRDGTextureRef GetVolumetricCloudTexture(FRDGBuilder& GraphBuilder) const;
+
 	/**
 	 * Collector for view-dependent data.
 	 */
diff --git a/Engine/Source/Runtime/Renderer/Private/VolumetricRenderTarget.cpp b/Engine/Source/Runtime/Renderer/Private/VolumetricRenderTarget.cpp
--- a/Engine/Source/Runtime/Renderer/Private/VolumetricRenderTarget.cpp
+++ b/Engine/Source/Runtime/Renderer/Private/VolumetricRenderTarget.cpp
@@ -283,6 +283,10 @@ FRDGTextureRef FVolumetricRenderTargetViewStateData::GetOrCreateVolumetricTracin
 	return GraphBuilder.RegisterExternalTexture(VolumetricTracingRTDepth);
 }
 
+FRDGTextureRef FVolumetricRenderTargetViewStateData::GetDstVolumetricReconstructRT(FRDGBuilder& GraphBuilder)
+{
+	return VolumetricReconstructRT[CurrentRT].IsValid() ? GraphBuilder.RegisterExternalTexture(VolumetricReconstructRT[CurrentRT]) : nullptr;
+}
 
 FRDGTextureRef FVolumetricRenderTargetViewStateData::GetOrCreateDstVolumetricReconstructRT(FRDGBuilder& GraphBuilder)
 {
diff --git a/Engine/Source/Runtime/Renderer/Private/VolumetricRenderTargetViewStateData.h b/Engine/Source/Runtime/Renderer/Private/VolumetricRenderTargetViewStateData.h
--- a/Engine/Source/Runtime/Renderer/Private/VolumetricRenderTargetViewStateData.h
+++ b/Engine/Source/Runtime/Renderer/Private/VolumetricRenderTargetViewStateData.h
@@ -30,6 +30,7 @@ public:
 	FRDGTextureRef GetOrCreateVolumetricTracingRT(FRDGBuilder& GraphBuilder);
 	FRDGTextureRef GetOrCreateVolumetricTracingRTDepth(FRDGBuilder& GraphBuilder);
 
+	FRDGTextureRef GetDstVolumetricReconstructRT(FRDGBuilder& GraphBuilder);
 	FRDGTextureRef GetOrCreateDstVolumetricReconstructRT(FRDGBuilder& GraphBuilder);
 	FRDGTextureRef GetOrCreateDstVolumetricReconstructRTDepth(FRDGBuilder& GraphBuilder);
 
