- Fix for count buffer not being transferred for mGPU
- The GPU mask is set to all GPUs where we previously transferred the buffer which meant it was never transferred

#rnx

diff --git a/Engine/Plugins/FX/Niagara/Source/Niagara/Private/NiagaraGpuComputeDispatch.cpp b/Engine/Plugins/FX/Niagara/Source/Niagara/Private/NiagaraGpuComputeDispatch.cpp
--- a/Engine/Plugins/FX/Niagara/Source/Niagara/Private/NiagaraGpuComputeDispatch.cpp
+++ b/Engine/Plugins/FX/Niagara/Source/Niagara/Private/NiagaraGpuComputeDispatch.cpp
@@ -1708,6 +1708,28 @@ void FNiagaraGpuComputeDispatch::PreInitViews(FRDGBuilder& GraphBuilder, bool bA
 		GPUInstanceCounterManager.ResizeBuffers(GraphBuilder.RHICmdList, 0);
 		FinishDispatches();
 	}
+
+#if WITH_MGPU
+	// If we have any simulations ticking we need to add the counter buffer to the list of buffers to transfer
+	// This must be done during graph building rather than when we transfer as the GPU mask will be all nodes
+	auto HasAnyWork =
+		[&]() -> bool
+		{
+			for (const FNiagaraGpuDispatchList& DispatchList : DispatchListPerStage)
+			{
+				if (DispatchList.HasWork())
+				{
+					return true;
+				}
+			}
+			return false;
+		};
+
+	if (GPUInstanceCounterManager.HasEntriesPendingFree() || HasAnyWork())
+	{
+		MultiGPUResourceModified(GraphBuilder, GPUInstanceCounterManager.GetInstanceCountBuffer().Buffer, true, true);
+	}
+#endif
 }
 
 void FNiagaraGpuComputeDispatch::PostInitViews(FRDGBuilder& GraphBuilder, TArrayView<const class FViewInfo> Views, bool bAllowGPUParticleUpdate)
@@ -2434,7 +2456,6 @@ void FNiagaraGpuComputeDispatch::TransferMultiGPUBufers(FRHICommandList& RHICmdL
 	// Transfer buffers for AFR rendering
 	if (AFRBuffers.Num())
 	{
-		AddAFRBuffer(GPUInstanceCounterManager.GetInstanceCountBuffer().Buffer);
 		RHICmdList.BroadcastTemporalEffect(FNiagaraGpuComputeDispatchLocal::TemporalEffectBuffersName, AFRBuffers);
 		AFRBuffers.Reset();
 	}
@@ -2447,7 +2468,6 @@ void FNiagaraGpuComputeDispatch::TransferMultiGPUBufers(FRHICommandList& RHICmdL
 	// Transfer buffers for cross GPU rendering
 	if (CrossGPUTransferBuffers.Num())
 	{
-		AddCrossGPUTransfer(RHICmdList, GPUInstanceCounterManager.GetInstanceCountBuffer().Buffer);
 		RHICmdList.TransferResources(CrossGPUTransferBuffers);
 		CrossGPUTransferBuffers.Reset();
 	}
