- Remove invalid text that accidentally got into a previous submit

#rnx

diff --git a/Engine/Plugins/FX/Niagara/Shaders/Private/NiagaraDataInterfaceOcclusion.ush b/Engine/Plugins/FX/Niagara/Shaders/Private/NiagaraDataInterfaceOcclusion.ush
--- a/Engine/Plugins/FX/Niagara/Shaders/Private/NiagaraDataInterfaceOcclusion.ush
+++ b/Engine/Plugins/FX/Niagara/Shaders/Private/NiagaraDataInterfaceOcclusion.ush
@@ -159,7 +159,7 @@ void QueryCloudOcclusionWithCircle_{ParameterName}(float3 WorldPosition, float W
 		const float RadStep = 2.0f * PI / float(SamplesPerRing + 1);
 		float Rad = 0.0f;
 		ExpectedSamples += float(NumberOfRings * SamplesPerRing);
-		for ( int iRing=0; iRing < NumberOfRings; ++iRing )datainterfaceocclusion.cpp
+		for ( int iRing=0; iRing < NumberOfRings; ++iRing )
 		{
 			const float2 RingRadiusUV = ScreenRadiusUV * (float(iRing + 1) / float(NumberOfRings));
 			Rad += SPIRAL_TURN;
