- Invert Y to fix ScreenUV

#rnx

diff --git a/Engine/Plugins/FX/Niagara/Shaders/Private/NiagaraDataInterfaceOcclusion.ush b/Engine/Plugins/FX/Niagara/Shaders/Private/NiagaraDataInterfaceOcclusion.ush
--- a/Engine/Plugins/FX/Niagara/Shaders/Private/NiagaraDataInterfaceOcclusion.ush
+++ b/Engine/Plugins/FX/Niagara/Shaders/Private/NiagaraDataInterfaceOcclusion.ush
@@ -134,10 +134,7 @@ void QueryCloudOcclusionWithCircle_{ParameterName}(float3 WorldPosition, float W
 	const float3 PreViewPosition = LWCToFloat(LWCAdd(LWCWorldPosition, PrimaryView.PreViewTranslation));
 	const float4 ClipPosition = mul(float4(PreViewPosition, 1), PrimaryView.TranslatedWorldToClip);
 	const float2 ScreenPosition = ClipPosition.xy / ClipPosition.w;
-	const float2 ScreenUV = (ScreenPosition + 1.0f) * 0.5f;
-	
-	const float WorldDepth = dot(LWCToFloat(LWCSubtract(LWCWorldPosition, PrimaryView.WorldCameraOrigin)), PrimaryView.ViewForward);
-	const float2 ScreenRadiusUV = (WorldDiameter * 0.5f) / (GetTanHalfFieldOfView() * 2.0f * WorldDepth);
+	const float2 ScreenUV = (ScreenPosition * float2(0.5f, -0.5f)) + 0.5f;
 
 	float TotalSamples = 0.0f;
 	float TotalVisibility = 0.0f;
@@ -156,6 +153,9 @@ void QueryCloudOcclusionWithCircle_{ParameterName}(float3 WorldPosition, float W
 	// Take ring taps
 	if ( NumberOfRings > 0 && SamplesPerRing > 0 )
 	{
+		const float WorldDepth = dot(LWCToFloat(LWCSubtract(LWCWorldPosition, PrimaryView.WorldCameraOrigin)), PrimaryView.ViewForward);
+		const float2 ScreenRadiusUV = (WorldDiameter * 0.5f) / (GetTanHalfFieldOfView() * 2.0f * WorldDepth);
+
 		const float RadStep = 2.0f * PI / float(SamplesPerRing + 1);
 		float Rad = 0.0f;
 		ExpectedSamples += float(NumberOfRings * SamplesPerRing);
@@ -179,5 +179,5 @@ void QueryCloudOcclusionWithCircle_{ParameterName}(float3 WorldPosition, float W
 	}
 
 	VisibilityFraction = TotalSamples > 0.0f ? TotalVisibility / TotalSamples : 0.0f;
-	SampleFraction = TotalSamples > 0.0f ? ExpectedSamples / TotalSamples : 0.0f;
+	SampleFraction = TotalSamples > 0.0f ? TotalSamples / ExpectedSamples : 0.0f;
 }
