- Add validation rule for banning data interfaces based on sim target and platform

diff --git a/Engine/Plugins/FX/Niagara/Source/Niagara/Private/DataInterface/NiagaraDataInterfaceUtilities.cpp b/Engine/Plugins/FX/Niagara/Source/Niagara/Private/DataInterface/NiagaraDataInterfaceUtilities.cpp
--- a/Engine/Plugins/FX/Niagara/Source/Niagara/Private/DataInterface/NiagaraDataInterfaceUtilities.cpp
+++ b/Engine/Plugins/FX/Niagara/Source/Niagara/Private/DataInterface/NiagaraDataInterfaceUtilities.cpp
@@ -437,3 +437,51 @@ void FNiagaraDataInterfaceUtilities::ForEachDataInterface(FNiagaraSystemInstance
 		}
 	}
 }
+
+void FNiagaraDataInterfaceUtilities::ForEachDataInterface(UNiagaraSystem* NiagaraSystem, TFunction<bool(const FDataInterfaceUsageContext&)> Action)
+{
+	if (NiagaraSystem == nullptr)
+	{
+		return;
+	}
+
+	FDataInterfaceUsageContext UsageContext;
+	UsageContext.OwnerObject = NiagaraSystem;
+
+	for (UNiagaraScript* SystemScript : { NiagaraSystem->GetSystemSpawnScript(), NiagaraSystem->GetSystemUpdateScript() } )
+	{
+		if (SystemScript == nullptr)
+		{
+			continue;
+		}
+
+		if (const FNiagaraScriptExecutionParameterStore* ParameterStore = SystemScript->GetExecutionReadyParameterStore(ENiagaraSimTarget::CPUSim))
+		{
+			ForEachDataInterface(UsageContext, *ParameterStore, Action);
+		}
+	}
+
+	for (const FNiagaraEmitterHandle& EmitterHandle : NiagaraSystem->GetEmitterHandles())
+	{
+		FVersionedNiagaraEmitterData* EmitterData = EmitterHandle.GetEmitterData();
+		if (EmitterHandle.GetIsEnabled() == false || EmitterData == nullptr)
+		{
+			continue;
+		}
+		const ENiagaraSimTarget SimTarget = EmitterData->SimTarget;
+		UsageContext.OwnerObject = EmitterHandle.GetInstance().Emitter;
+
+		EmitterData->ForEachScript(
+			[&](UNiagaraScript* NiagaraScript)
+			{
+				const FNiagaraParameterStore* ParameterStore = NiagaraScript->GetExecutionReadyParameterStore(SimTarget);
+				if (ParameterStore == nullptr)
+				{
+					return;
+				}
+
+				ForEachDataInterface(UsageContext, *ParameterStore, Action);
+			}
+		);
+	}
+}
diff --git a/Engine/Plugins/FX/Niagara/Source/Niagara/Public/NiagaraDataInterfaceUtilities.h b/Engine/Plugins/FX/Niagara/Source/Niagara/Public/NiagaraDataInterfaceUtilities.h
--- a/Engine/Plugins/FX/Niagara/Source/Niagara/Public/NiagaraDataInterfaceUtilities.h
+++ b/Engine/Plugins/FX/Niagara/Source/Niagara/Public/NiagaraDataInterfaceUtilities.h
@@ -43,8 +43,11 @@ namespace FNiagaraDataInterfaceUtilities
 
 	// Loops over all data interfaces inside the SystemInstance
 	// The action should return True to continue iteration or False to stop
-	void ForEachDataInterface(FNiagaraSystemInstance* SystemInstance, TFunction<bool(const FNiagaraVariableBase Variable, UNiagaraDataInterface* DataInterface)> Action);
+	NIAGARA_API void ForEachDataInterface(FNiagaraSystemInstance* SystemInstance, TFunction<bool(const FNiagaraVariableBase Variable, UNiagaraDataInterface* DataInterface)> Action);
 	// Loops over all data interfaces inside the SystemInstance
 	// The action should return True to continue iteration or False to stop
-	void ForEachDataInterface(FNiagaraSystemInstance* SystemInstance, TFunction<bool(const FDataInterfaceUsageContext&)> Action);
+	NIAGARA_API void ForEachDataInterface(FNiagaraSystemInstance* SystemInstance, TFunction<bool(const FDataInterfaceUsageContext&)> Action);
+	// Loops over all data interfaces inside the NiagaraSystem
+	// The action should return True to continue iteration or False to stop
+	NIAGARA_API void ForEachDataInterface(UNiagaraSystem* NiagaraSystem, TFunction<bool(const FDataInterfaceUsageContext&)> Action);
 }
diff --git a/Engine/Plugins/FX/Niagara/Source/NiagaraEditor/Private/NiagaraValidationRules.cpp b/Engine/Plugins/FX/Niagara/Source/NiagaraEditor/Private/NiagaraValidationRules.cpp
--- a/Engine/Plugins/FX/Niagara/Source/NiagaraEditor/Private/NiagaraValidationRules.cpp
+++ b/Engine/Plugins/FX/Niagara/Source/NiagaraEditor/Private/NiagaraValidationRules.cpp
@@ -4,6 +4,7 @@
 
 #include "NiagaraClipboard.h"
 #include "NiagaraComponentRendererProperties.h"
+#include "NiagaraDataInterfaceUtilities.h"
 #include "NiagaraNodeFunctionCall.h"
 #include "NiagaraScriptSource.h"
 #include "NiagaraSettings.h"
@@ -109,6 +110,47 @@ namespace NiagaraValidation
 				}
 			});
 	}
+
+	FString GetPlatformConflictsString(const FNiagaraPlatformSet& PlatformSetA, const FNiagaraPlatformSet& PlatformSetB, int MaxPlatformsToShow = 4)
+	{
+		TArray<const FNiagaraPlatformSet*> CheckSets;
+		CheckSets.Add(&PlatformSetA);
+		CheckSets.Add(&PlatformSetB);
+
+		TArray<FNiagaraPlatformSetConflictInfo> ConflictInfos;
+		FNiagaraPlatformSet::GatherConflicts(CheckSets, ConflictInfos);
+
+		if (ConflictInfos.Num() > 0)
+		{
+			TSet<FName> BannedPlatformNames;
+			for (const FNiagaraPlatformSetConflictInfo& ConflictInfo : ConflictInfos)
+			{
+				for (const FNiagaraPlatformSetConflictEntry& ConflictEntry : ConflictInfo.Conflicts)
+				{
+					BannedPlatformNames.Add(ConflictEntry.ProfileName);
+				}
+			}
+
+			TStringBuilder<256> BannedPlatformsString;
+			int NumFounds = 0;
+			for (FName PlatformName : BannedPlatformNames)
+			{
+				if (NumFounds >= MaxPlatformsToShow)
+				{
+					BannedPlatformsString.Append(TEXT(", ..."));
+					break;
+				}
+				if (NumFounds != 0)
+				{
+					BannedPlatformsString.Append(TEXT(", "));
+				}
+				++NumFounds;
+				PlatformName.AppendString(BannedPlatformsString);
+			}
+			return BannedPlatformsString.ToString();
+		}
+		return FString();
+	}
 }
 
 // --------------------------------------------------------------------------------------------------------------------------------------------
@@ -328,6 +370,83 @@ void UNiagaraValidationRule_BannedModules::CheckValidity(const FNiagaraValidatio
 	}
 }
 
+void UNiagaraValidationRule_BannedDataInterfaces::CheckValidity(const FNiagaraValidationContext& Context, TArray<FNiagaraValidationResult>& Results)  const
+{
+	UNiagaraSystem* NiagaraSystem = &Context.ViewModel->GetSystem();
+
+	FNiagaraDataInterfaceUtilities::ForEachDataInterface(
+		NiagaraSystem,
+		[&](const FNiagaraDataInterfaceUtilities::FDataInterfaceUsageContext& UsageContext) -> bool
+		{
+			UClass* DIClass = UsageContext.DataInterface->GetClass();
+			if (BannedDataInterfaces.Contains(DIClass) == false)
+			{
+				return true;
+			}
+
+			static const FText WarningFormat(LOCTEXT("BannedDataInteraceFormatWarn", "DataInterface '{0}' is banned on currently enabled platforms"));
+			static const FText SystemDescFormat(LOCTEXT("BannedDataInteraceFormatDesc", "DataInterface '{0} - {1}' is banned on currently enabled platforms"));
+			static const FText EmitterDescFormat(LOCTEXT("BannedDataInteraceFormatDesc", "DataInterface '{0} - {1}' is banned on currently enabled platforms '{2}'"));
+
+			UObject* WarningObject = nullptr;
+			if (UNiagaraEmitter* NiagaraEmitter = Cast<UNiagaraEmitter>(UsageContext.OwnerObject))
+			{
+				const TSharedRef<FNiagaraEmitterHandleViewModel>* EmitterViewModel =
+					Context.ViewModel->GetEmitterHandleViewModels().FindByPredicate(
+						[NiagaraEmitter](const TSharedRef<FNiagaraEmitterHandleViewModel>& EmitterViewModelRef)
+						{
+							FNiagaraEmitterHandle* EmitterHandle = EmitterViewModelRef->GetEmitterHandle();
+							return EmitterHandle && EmitterHandle->GetInstance().Emitter == NiagaraEmitter;
+						}
+					);
+
+				if (EmitterViewModel == nullptr)
+				{
+					return true;
+				}
+
+				const FVersionedNiagaraEmitterData* EmitterData = (*EmitterViewModel)->GetEmitterHandle()->GetEmitterData();
+				if (EmitterData == nullptr)
+				{
+					return true;
+				}
+
+				const bool bIsBanEnabled =
+					((EmitterData->SimTarget == ENiagaraSimTarget::CPUSim) && bBanOnCpu) ||
+					((EmitterData->SimTarget == ENiagaraSimTarget::GPUComputeSim) && bBanOnGpu);
+
+				if (bIsBanEnabled)
+				{
+					FString PlatformConflictsString = NiagaraValidation::GetPlatformConflictsString(Platforms, EmitterData->Platforms);
+					if (PlatformConflictsString.IsEmpty() == false)
+					{
+						Results.Emplace(
+							ENiagaraValidationSeverity::Warning,
+							FText::Format(WarningFormat, FText::FromName(UsageContext.Variable.GetName())),
+							FText::Format(EmitterDescFormat, FText::FromName(UsageContext.Variable.GetName()), FText::FromName(DIClass->GetFName()), FText::FromString(PlatformConflictsString)),
+							NiagaraValidation::GetStackEntry<UNiagaraStackEmitterPropertiesItem>((*EmitterViewModel)->GetEmitterStackViewModel())
+						);
+					}
+				}
+			}
+			else if (UNiagaraSystem* NiagaraSystem = Cast<UNiagaraSystem>(UsageContext.OwnerObject))
+			{
+				if (bBanOnCpu == true)
+				{
+					Results.Emplace(
+						ENiagaraValidationSeverity::Warning,
+						FText::Format(WarningFormat, FText::FromName(UsageContext.Variable.GetName())),
+						FText::Format(SystemDescFormat, FText::FromName(UsageContext.Variable.GetName()), FText::FromName(DIClass->GetFName())),
+						NiagaraValidation::GetStackEntry<UNiagaraStackSystemPropertiesItem>(Context.ViewModel->GetSystemStackViewModel())
+					);
+				}
+			}
+
+			return true;
+		}
+	);
+}
+
 void UNiagaraValidationRule_InvalidEffectType::CheckValidity(const FNiagaraValidationContext& Context, TArray<FNiagaraValidationResult>& Results)  const
 {
 	UNiagaraStackSystemPropertiesItem* SystemProperties = NiagaraValidation::GetStackEntry<UNiagaraStackSystemPropertiesItem>(Context.ViewModel->GetSystemStackViewModel());
diff --git a/Engine/Plugins/FX/Niagara/Source/NiagaraEditor/Public/NiagaraValidationRules.h b/Engine/Plugins/FX/Niagara/Source/NiagaraEditor/Public/NiagaraValidationRules.h
--- a/Engine/Plugins/FX/Niagara/Source/NiagaraEditor/Public/NiagaraValidationRules.h
+++ b/Engine/Plugins/FX/Niagara/Source/NiagaraEditor/Public/NiagaraValidationRules.h
@@ -69,6 +69,28 @@ public:
 	virtual void CheckValidity(const FNiagaraValidationContext& Context, TArray<FNiagaraValidationResult>& OutResults) const override;
 };
 
+/** This validation rule can ban the use of certain datainterfaces on all or a subset of platforms. */
+UCLASS(Category = "Validation", DisplayName = "Banned DataInterfaces")
+class UNiagaraValidationRule_BannedDataInterfaces : public UNiagaraValidationRule
+{
+	GENERATED_BODY()
+
+public:
+	UPROPERTY(EditAnywhere, Category = Validation)
+	bool bBanOnGpu = true;
+
+	UPROPERTY(EditAnywhere, Category = Validation)
+	bool bBanOnCpu = true;
+
+	UPROPERTY(EditAnywhere, Category = Validation)
+	FNiagaraPlatformSet Platforms;
+
+	UPROPERTY(EditAnywhere, Category = Validation)
+	TArray<TSubclassOf<UNiagaraDataInterface>> BannedDataInterfaces;
+
+	virtual void CheckValidity(const FNiagaraValidationContext& Context, TArray<FNiagaraValidationResult>& OutResults) const override;
+};
+
 /** This validation rule can marks this effect type as invalid and so must be changed. Forces a choice of correct Effect Type for an System rather than. Leaving as the default. */
 UCLASS(Category = "Validation", DisplayName = "Invalid Effect Type")
 class UNiagaraValidationRule_InvalidEffectType : public UNiagaraValidationRule
