- Clean up FX asset menu and move infrequently used items into sub menus

diff --git a/Engine/Plugins/FX/Niagara/Source/NiagaraEditor/Private/AssetTypeActions/AssetTypeActions_NiagaraEffectType.cpp b/Engine/Plugins/FX/Niagara/Source/NiagaraEditor/Private/AssetTypeActions/AssetTypeActions_NiagaraEffectType.cpp
--- a/Engine/Plugins/FX/Niagara/Source/NiagaraEditor/Private/AssetTypeActions/AssetTypeActions_NiagaraEffectType.cpp
+++ b/Engine/Plugins/FX/Niagara/Source/NiagaraEditor/Private/AssetTypeActions/AssetTypeActions_NiagaraEffectType.cpp
@@ -37,4 +37,13 @@ UClass* FAssetTypeActions_NiagaraEffectType::GetSupportedClass() const
 	return UNiagaraEffectType::StaticClass();
 }
 
+const TArray<FText>& FAssetTypeActions_NiagaraEffectType::GetSubMenus() const
+{
+	static const TArray<FText> SubMenus
+	{
+		NSLOCTEXT("Niagara", "NiagaraAssetSubMenu_Advanced", "Advanced")
+	};
+	return SubMenus;
+}
+
 #undef LOCTEXT_NAMESPACE
diff --git a/Engine/Plugins/FX/Niagara/Source/NiagaraEditor/Private/AssetTypeActions/AssetTypeActions_NiagaraEffectType.h b/Engine/Plugins/FX/Niagara/Source/NiagaraEditor/Private/AssetTypeActions/AssetTypeActions_NiagaraEffectType.h
--- a/Engine/Plugins/FX/Niagara/Source/NiagaraEditor/Private/AssetTypeActions/AssetTypeActions_NiagaraEffectType.h
+++ b/Engine/Plugins/FX/Niagara/Source/NiagaraEditor/Private/AssetTypeActions/AssetTypeActions_NiagaraEffectType.h
@@ -16,4 +16,5 @@ public:
 	virtual UClass* GetSupportedClass() const override;
 	//virtual void OpenAssetEditor(const TArray<UObject*>& InObjects, TSharedPtr<class IToolkitHost> EditWithinLevelEditor = TSharedPtr<IToolkitHost>()) override;
 	virtual uint32 GetCategories() override { return FNiagaraEditorModule::GetAssetCategory(); }
-};
\ No newline at end of file
+	virtual const TArray<FText>& GetSubMenus() const override;
+};
diff --git a/Engine/Plugins/FX/Niagara/Source/NiagaraEditor/Private/AssetTypeActions/AssetTypeActions_NiagaraParameterCollection.cpp b/Engine/Plugins/FX/Niagara/Source/NiagaraEditor/Private/AssetTypeActions/AssetTypeActions_NiagaraParameterCollection.cpp
--- a/Engine/Plugins/FX/Niagara/Source/NiagaraEditor/Private/AssetTypeActions/AssetTypeActions_NiagaraParameterCollection.cpp
+++ b/Engine/Plugins/FX/Niagara/Source/NiagaraEditor/Private/AssetTypeActions/AssetTypeActions_NiagaraParameterCollection.cpp
@@ -14,6 +14,15 @@
 
 #define LOCTEXT_NAMESPACE "AssetTypeActions_NiagaraParameterCollection"
 
+const TArray<FText>& FAssetTypeActions_NiagaraParameterCollectionBase::GetSubMenus() const
+{
+	static const TArray<FText> SubMenus
+	{
+		NSLOCTEXT("Niagara", "NiagaraAssetSubMenu_Advanced", "Advanced")
+	};
+	return SubMenus;
+}
+
 FColor FAssetTypeActions_NiagaraParameterCollection::GetTypeColor() const
 {
 	return FNiagaraEditorStyle::Get().GetColor("NiagaraEditor.AssetColors.ParameterCollection").ToFColor(true);
diff --git a/Engine/Plugins/FX/Niagara/Source/NiagaraEditor/Private/AssetTypeActions/AssetTypeActions_NiagaraParameterCollection.h b/Engine/Plugins/FX/Niagara/Source/NiagaraEditor/Private/AssetTypeActions/AssetTypeActions_NiagaraParameterCollection.h
--- a/Engine/Plugins/FX/Niagara/Source/NiagaraEditor/Private/AssetTypeActions/AssetTypeActions_NiagaraParameterCollection.h
+++ b/Engine/Plugins/FX/Niagara/Source/NiagaraEditor/Private/AssetTypeActions/AssetTypeActions_NiagaraParameterCollection.h
@@ -7,7 +7,13 @@
 #include "AssetTypeActions_Base.h"
 #include "NiagaraEditorModule.h"
 
-class FAssetTypeActions_NiagaraParameterCollection : public FAssetTypeActions_Base
+class FAssetTypeActions_NiagaraParameterCollectionBase : public FAssetTypeActions_Base
+{
+public:
+	virtual const TArray<FText>& GetSubMenus() const override;
+};
+
+class FAssetTypeActions_NiagaraParameterCollection : public FAssetTypeActions_NiagaraParameterCollectionBase
 {
 public:
 	// IAssetTypeActions Implementation
@@ -21,7 +27,7 @@ public:
 	void ExecuteNewNPC(TArray<TWeakObjectPtr<class UNiagaraParameterCollection>> Objects);
 };
 
-class FAssetTypeActions_NiagaraParameterCollectionInstance : public FAssetTypeActions_Base
+class FAssetTypeActions_NiagaraParameterCollectionInstance : public FAssetTypeActions_NiagaraParameterCollectionBase
 {
 public:
 	// IAssetTypeActions Implementation
diff --git a/Engine/Plugins/FX/Niagara/Source/NiagaraEditor/Private/AssetTypeActions/AssetTypeActions_NiagaraParameterDefinitions.cpp b/Engine/Plugins/FX/Niagara/Source/NiagaraEditor/Private/AssetTypeActions/AssetTypeActions_NiagaraParameterDefinitions.cpp
--- a/Engine/Plugins/FX/Niagara/Source/NiagaraEditor/Private/AssetTypeActions/AssetTypeActions_NiagaraParameterDefinitions.cpp
+++ b/Engine/Plugins/FX/Niagara/Source/NiagaraEditor/Private/AssetTypeActions/AssetTypeActions_NiagaraParameterDefinitions.cpp
@@ -30,3 +30,12 @@ void FAssetTypeActions_NiagaraParameterDefinitions::OpenAssetEditor(const TArray
 	}
 }
 
+const TArray<FText>& FAssetTypeActions_NiagaraParameterDefinitions::GetSubMenus() const
+{
+	static const TArray<FText> SubMenus
+	{
+		NSLOCTEXT("Niagara", "NiagaraAssetSubMenu_Advanced", "Advanced")
+	};
+	return SubMenus;
+}
+
diff --git a/Engine/Plugins/FX/Niagara/Source/NiagaraEditor/Private/AssetTypeActions/AssetTypeActions_NiagaraParameterDefinitions.h b/Engine/Plugins/FX/Niagara/Source/NiagaraEditor/Private/AssetTypeActions/AssetTypeActions_NiagaraParameterDefinitions.h
--- a/Engine/Plugins/FX/Niagara/Source/NiagaraEditor/Private/AssetTypeActions/AssetTypeActions_NiagaraParameterDefinitions.h
+++ b/Engine/Plugins/FX/Niagara/Source/NiagaraEditor/Private/AssetTypeActions/AssetTypeActions_NiagaraParameterDefinitions.h
@@ -16,5 +16,6 @@ public:
 	virtual UClass* GetSupportedClass() const override;
 	virtual void OpenAssetEditor(const TArray<UObject*>& InObjects, TSharedPtr<class IToolkitHost> EditWithinLevelEditor = TSharedPtr<IToolkitHost>()) override;
 	virtual uint32 GetCategories() override { return FNiagaraEditorModule::GetAssetCategory(); }
+	virtual const TArray<FText>& GetSubMenus() const override;
 };
 
diff --git a/Engine/Plugins/FX/Niagara/Source/NiagaraEditor/Private/AssetTypeActions/AssetTypeActions_NiagaraScript.cpp b/Engine/Plugins/FX/Niagara/Source/NiagaraEditor/Private/AssetTypeActions/AssetTypeActions_NiagaraScript.cpp
--- a/Engine/Plugins/FX/Niagara/Source/NiagaraEditor/Private/AssetTypeActions/AssetTypeActions_NiagaraScript.cpp
+++ b/Engine/Plugins/FX/Niagara/Source/NiagaraEditor/Private/AssetTypeActions/AssetTypeActions_NiagaraScript.cpp
@@ -72,6 +72,15 @@ void FAssetTypeActions_NiagaraScript::GetActions(const TArray<UObject*>& InObjec
 		FUIAction(FExecuteAction::CreateStatic(&FNiagaraEditorUtilities::MarkDependentCompilableAssetsDirty, InObjects)));
 }
 
+const TArray<FText>& FAssetTypeActions_NiagaraScript:: GetSubMenus() const
+{
+	static const TArray<FText> SubMenus
+	{
+		NSLOCTEXT("Niagara", "NiagaraAssetSubMenu_Script", "Script")
+	};
+	return SubMenus;
+}
+
 UClass* FAssetTypeActions_NiagaraScript::GetSupportedClass() const
 { 
 	return UNiagaraScript::StaticClass(); 
diff --git a/Engine/Plugins/FX/Niagara/Source/NiagaraEditor/Private/AssetTypeActions/AssetTypeActions_NiagaraScript.h b/Engine/Plugins/FX/Niagara/Source/NiagaraEditor/Private/AssetTypeActions/AssetTypeActions_NiagaraScript.h
--- a/Engine/Plugins/FX/Niagara/Source/NiagaraEditor/Private/AssetTypeActions/AssetTypeActions_NiagaraScript.h
+++ b/Engine/Plugins/FX/Niagara/Source/NiagaraEditor/Private/AssetTypeActions/AssetTypeActions_NiagaraScript.h
@@ -18,6 +18,7 @@ public:
 	virtual uint32 GetCategories() override { return FNiagaraEditorModule::GetAssetCategory(); }
 	virtual FText GetDisplayNameFromAssetData(const FAssetData& AssetData) const override;
 	virtual void GetActions(const TArray<UObject*>& InObjects, struct FToolMenuSection& Section) override;
+	virtual const TArray<FText>& GetSubMenus() const override;
 	// End IAssetTypeActions Interface
 
 	/** Get the action type name. */
diff --git a/Engine/Plugins/FX/Niagara/Source/NiagaraEditor/Private/AssetTypeActions/AssetTypeActions_NiagaraSimCache.cpp b/Engine/Plugins/FX/Niagara/Source/NiagaraEditor/Private/AssetTypeActions/AssetTypeActions_NiagaraSimCache.cpp
--- a/Engine/Plugins/FX/Niagara/Source/NiagaraEditor/Private/AssetTypeActions/AssetTypeActions_NiagaraSimCache.cpp
+++ b/Engine/Plugins/FX/Niagara/Source/NiagaraEditor/Private/AssetTypeActions/AssetTypeActions_NiagaraSimCache.cpp
@@ -29,3 +29,12 @@ void FAssetTypeActions_NiagaraSimCache::OpenAssetEditor(const TArray<UObject*>&
 		}
 	}
 }
+
+const TArray<FText>& FAssetTypeActions_NiagaraSimCache::GetSubMenus() const
+{
+	static const TArray<FText> SubMenus
+	{
+		NSLOCTEXT("Niagara", "NiagaraAssetSubMenu_Advanced", "Advanced")
+	};
+	return SubMenus;
+}
diff --git a/Engine/Plugins/FX/Niagara/Source/NiagaraEditor/Private/AssetTypeActions/AssetTypeActions_NiagaraSimCache.h b/Engine/Plugins/FX/Niagara/Source/NiagaraEditor/Private/AssetTypeActions/AssetTypeActions_NiagaraSimCache.h
--- a/Engine/Plugins/FX/Niagara/Source/NiagaraEditor/Private/AssetTypeActions/AssetTypeActions_NiagaraSimCache.h
+++ b/Engine/Plugins/FX/Niagara/Source/NiagaraEditor/Private/AssetTypeActions/AssetTypeActions_NiagaraSimCache.h
@@ -17,4 +17,5 @@ public:
 	virtual void OpenAssetEditor(const TArray<UObject*>& InObjects, TSharedPtr<class IToolkitHost> EditWithinLevelEditor = TSharedPtr<IToolkitHost>()) override;
 	virtual bool CanFilter() override { return true; }
 	virtual uint32 GetCategories() override { return FNiagaraEditorModule::GetAssetCategory(); }
+	virtual const TArray<FText>& GetSubMenus() const override;
 };
