- Added Niagara Validation Rule Set
- Can be used to share rules between different FX types
- Can be used to apply default validation rules to all Niagara assets

diff --git a/Engine/Plugins/FX/Niagara/Source/Niagara/Classes/NiagaraEffectType.h b/Engine/Plugins/FX/Niagara/Source/Niagara/Classes/NiagaraEffectType.h
--- a/Engine/Plugins/FX/Niagara/Source/Niagara/Classes/NiagaraEffectType.h
+++ b/Engine/Plugins/FX/Niagara/Source/Niagara/Classes/NiagaraEffectType.h
@@ -7,6 +7,7 @@
 #include "NiagaraPlatformSet.h"
 #include "NiagaraPerfBaseline.h"
 #include "NiagaraValidationRule.h"
+#include "NiagaraValidationRuleSet.h"
 #include "NiagaraEffectType.generated.h"
 
 #define DEBUG_SCALABILITY_STATE (!UE_BUILD_SHIPPING)
@@ -441,10 +442,16 @@ class NIAGARA_API UNiagaraEffectType : public UObject
 	UPROPERTY(EditAnywhere, Category = "Scalability")
 	FNiagaraEmitterScalabilitySettingsArray EmitterScalabilitySettings;
 
+#if WITH_EDITORONLY_DATA
 	/** A set of rules to apply when checking content. To create your own rules, write a custom class that extends UNiagaraValidationRule. */
 	UPROPERTY(EditAnywhere, Category = "Validation", Instanced)
 	TArray<TObjectPtr<UNiagaraValidationRule>> ValidationRules;
 
+	/** Array of reusable rule sets to apply when checking content. To create your own rules, write a custom class that extends UNiagaraValidationRule. */
+	UPROPERTY(EditAnywhere, Category = "Validation")
+	TArray<TObjectPtr<UNiagaraValidationRuleSet>> ValidationRuleSets;
+#endif
+
 	FORCEINLINE const FNiagaraSystemScalabilitySettingsArray& GetSystemScalabilitySettings()const { return SystemScalabilitySettings; }
 	FORCEINLINE const FNiagaraEmitterScalabilitySettingsArray& GetEmitterScalabilitySettings()const { return EmitterScalabilitySettings; }
 
diff --git a/Engine/Plugins/FX/Niagara/Source/Niagara/Classes/NiagaraValidationRuleSet.h b/Engine/Plugins/FX/Niagara/Source/Niagara/Classes/NiagaraValidationRuleSet.h
new file mode 100644
--- /dev/null
+++ b/Engine/Plugins/FX/Niagara/Source/Niagara/Classes/NiagaraValidationRuleSet.h
@@ -0,0 +1,17 @@
+// Copyright Epic Games, Inc. All Rights Reserved.
+
+#pragma once
+
+#include "NiagaraValidationRule.h"
+#include "NiagaraValidationRuleSet.generated.h"
+
+/** A set of reusable validation rules. */
+UCLASS(BlueprintType)
+class NIAGARA_API UNiagaraValidationRuleSet : public UObject
+{
+	GENERATED_BODY()
+public:
+
+	UPROPERTY(EditAnywhere, Category = "Validation", Instanced)
+	TArray<TObjectPtr<UNiagaraValidationRule>> ValidationRules;
+};
diff --git a/Engine/Plugins/FX/Niagara/Source/NiagaraEditor/Private/AssetTypeActions/AssetTypeActions_NiagaraValidationRuleSet.cpp b/Engine/Plugins/FX/Niagara/Source/NiagaraEditor/Private/AssetTypeActions/AssetTypeActions_NiagaraValidationRuleSet.cpp
new file mode 100644
--- /dev/null
+++ b/Engine/Plugins/FX/Niagara/Source/NiagaraEditor/Private/AssetTypeActions/AssetTypeActions_NiagaraValidationRuleSet.cpp
@@ -0,0 +1,29 @@
+// Copyright Epic Games, Inc. All Rights Reserved.
+
+#include "AssetTypeActions_NiagaraValidationRuleSet.h"
+#include "NiagaraEditorStyle.h"
+#include "NiagaraValidationRuleSet.h"
+
+FText FAssetTypeActions_NiagaraValidationRuleSet::GetName() const
+{
+	return NSLOCTEXT("AssetTypeActions", "AssetTypeActions_NiagaraValidationRuleSet", "Niagara Validation Rule Set");
+}
+
+FColor FAssetTypeActions_NiagaraValidationRuleSet::GetTypeColor() const
+{
+	return FNiagaraEditorStyle::Get().GetColor("NiagaraEditor.AssetColors.ValidationRuleSet").ToFColor(true);
+}
+
+UClass* FAssetTypeActions_NiagaraValidationRuleSet::GetSupportedClass() const
+{
+	return UNiagaraValidationRuleSet::StaticClass();
+}
+
+const TArray<FText>& FAssetTypeActions_NiagaraValidationRuleSet::GetSubMenus() const
+{
+	static const TArray<FText> SubMenus
+	{
+		NSLOCTEXT("Niagara", "NiagaraAssetSubMenu_Advanced", "Advanced")
+	};
+	return SubMenus;
+}
diff --git a/Engine/Plugins/FX/Niagara/Source/NiagaraEditor/Private/AssetTypeActions/AssetTypeActions_NiagaraValidationRuleSet.h b/Engine/Plugins/FX/Niagara/Source/NiagaraEditor/Private/AssetTypeActions/AssetTypeActions_NiagaraValidationRuleSet.h
new file mode 100644
--- /dev/null
+++ b/Engine/Plugins/FX/Niagara/Source/NiagaraEditor/Private/AssetTypeActions/AssetTypeActions_NiagaraValidationRuleSet.h
@@ -0,0 +1,20 @@
+// Copyright Epic Games, Inc. All Rights Reserved.
+
+#pragma once
+
+#include "CoreMinimal.h"
+#include "Toolkits/IToolkitHost.h"
+#include "AssetTypeActions_Base.h"
+#include "NiagaraEditorModule.h"
+
+class FAssetTypeActions_NiagaraValidationRuleSet : public FAssetTypeActions_Base
+{
+public:
+	// IAssetTypeActions Implementation
+	virtual FText GetName() const override;
+	virtual FColor GetTypeColor() const override;
+	virtual UClass* GetSupportedClass() const override;
+	virtual bool CanFilter() override { return true; }
+	virtual uint32 GetCategories() override { return FNiagaraEditorModule::GetAssetCategory(); }
+	virtual const TArray<FText>& GetSubMenus() const override;
+};
diff --git a/Engine/Plugins/FX/Niagara/Source/NiagaraEditor/Private/NiagaraEditorModule.cpp b/Engine/Plugins/FX/Niagara/Source/NiagaraEditor/Private/NiagaraEditorModule.cpp
--- a/Engine/Plugins/FX/Niagara/Source/NiagaraEditor/Private/NiagaraEditorModule.cpp
+++ b/Engine/Plugins/FX/Niagara/Source/NiagaraEditor/Private/NiagaraEditorModule.cpp
@@ -22,6 +22,7 @@
 #include "AssetTypeActions/AssetTypeActions_NiagaraEffectType.h"
 #include "AssetTypeActions/AssetTypeActions_NiagaraParameterDefinitions.h"
 #include "AssetTypeActions/AssetTypeActions_NiagaraSimCache.h"
+#include "AssetTypeActions/AssetTypeActions_NiagaraValidationRuleSet.h"
 
 #include "EdGraphSchema_Niagara.h"
 #include "EdGraphUtilities.h"
@@ -866,6 +867,7 @@ void FNiagaraEditorModule::StartupModule()
 	RegisterAssetTypeAction(AssetTools, MakeShareable(new FAssetTypeActions_NiagaraParameterDefinitions()));
 	RegisterAssetTypeAction(AssetTools, MakeShareable(new FAssetTypeActions_NiagaraEffectType()));
 	RegisterAssetTypeAction(AssetTools, MakeShareable(new FAssetTypeActions_NiagaraSimCache()));
+	RegisterAssetTypeAction(AssetTools, MakeShareable(new FAssetTypeActions_NiagaraValidationRuleSet()));
 
 	// Preload all parameter definition & collection assets so that they will be postloaded before postload calls to scripts/emitters/systems that rely on them.
 	{
diff --git a/Engine/Plugins/FX/Niagara/Source/NiagaraEditor/Private/NiagaraEditorStyle.cpp b/Engine/Plugins/FX/Niagara/Source/NiagaraEditor/Private/NiagaraEditorStyle.cpp
--- a/Engine/Plugins/FX/Niagara/Source/NiagaraEditor/Private/NiagaraEditorStyle.cpp
+++ b/Engine/Plugins/FX/Niagara/Source/NiagaraEditor/Private/NiagaraEditorStyle.cpp
@@ -420,6 +420,7 @@ void FNiagaraEditorStyle::InitAssetColors()
 	Set("NiagaraEditor.AssetColors.ParameterDefinitions", FLinearColor(0.57f, 0.82f, 0.06f));
 	Set("NiagaraEditor.AssetColors.EffectType", FLinearColor(1.f, 1.f, 1.f));
 	Set("NiagaraEditor.AssetColors.SimCache", FLinearColor(0.9f, 0.3f, 0.25f));
+	Set("NiagaraEditor.AssetColors.ValidationRuleSet", FLinearColor(0.25f, 0.25f, 0.25f));
 }
 
 void FNiagaraEditorStyle::InitThumbnails()
diff --git a/Engine/Plugins/FX/Niagara/Source/NiagaraEditor/Private/NiagaraValidationRuleSetFactoryNew.cpp b/Engine/Plugins/FX/Niagara/Source/NiagaraEditor/Private/NiagaraValidationRuleSetFactoryNew.cpp
new file mode 100644
--- /dev/null
+++ b/Engine/Plugins/FX/Niagara/Source/NiagaraEditor/Private/NiagaraValidationRuleSetFactoryNew.cpp
@@ -0,0 +1,24 @@
+// Copyright Epic Games, Inc. All Rights Reserved.
+
+#include "NiagaraValidationRuleSetFactoryNew.h"
+#include "NiagaraValidationRuleSet.h"
+
+#include UE_INLINE_GENERATED_CPP_BY_NAME(NiagaraValidationRuleSetFactoryNew)
+
+UNiagaraValidationRuleSetFactoryNew::UNiagaraValidationRuleSetFactoryNew(const FObjectInitializer& ObjectInitializer)
+{
+	SupportedClass = UNiagaraValidationRuleSet::StaticClass();
+	bCreateNew = true;
+	bEditAfterNew = true;
+}
+
+bool UNiagaraValidationRuleSetFactoryNew::ShouldShowInNewMenu() const
+{
+	return true;
+}
+
+UObject* UNiagaraValidationRuleSetFactoryNew::FactoryCreateNew(UClass* InClass, UObject* InParent, FName InName, EObjectFlags Flags, UObject* Context, FFeedbackContext* Warn)
+{
+	return NewObject<UNiagaraValidationRuleSet>(InParent, InClass, InName, Flags);
+}
+
diff --git a/Engine/Plugins/FX/Niagara/Source/NiagaraEditor/Private/NiagaraValidationRuleSetFactoryNew.h b/Engine/Plugins/FX/Niagara/Source/NiagaraEditor/Private/NiagaraValidationRuleSetFactoryNew.h
new file mode 100644
--- /dev/null
+++ b/Engine/Plugins/FX/Niagara/Source/NiagaraEditor/Private/NiagaraValidationRuleSetFactoryNew.h
@@ -0,0 +1,17 @@
+// Copyright Epic Games, Inc. All Rights Reserved.
+
+#pragma once
+
+#include "Factories/Factory.h"
+#include "NiagaraValidationRuleSetFactoryNew.generated.h"
+
+class UNiagaraValidationRuleSet;
+
+UCLASS(hidecategories=Object, MinimalAPI)
+class UNiagaraValidationRuleSetFactoryNew : public UFactory
+{
+	GENERATED_UCLASS_BODY()
+
+	virtual bool ShouldShowInNewMenu() const override;
+	virtual UObject* FactoryCreateNew(UClass* InClass, UObject* InParent, FName InName, EObjectFlags Flags, UObject* Context, FFeedbackContext* Warn) override;
+};
diff --git a/Engine/Plugins/FX/Niagara/Source/NiagaraEditor/Private/NiagaraValidationRules.cpp b/Engine/Plugins/FX/Niagara/Source/NiagaraEditor/Private/NiagaraValidationRules.cpp
--- a/Engine/Plugins/FX/Niagara/Source/NiagaraEditor/Private/NiagaraValidationRules.cpp
+++ b/Engine/Plugins/FX/Niagara/Source/NiagaraEditor/Private/NiagaraValidationRules.cpp
@@ -5,6 +5,7 @@
 #include "NiagaraClipboard.h"
 #include "NiagaraComponentRendererProperties.h"
 #include "NiagaraDataInterfaceUtilities.h"
+#include "NiagaraEditorSettings.h"
 #include "NiagaraNodeFunctionCall.h"
 #include "NiagaraScriptSource.h"
 #include "NiagaraSettings.h"
@@ -167,19 +168,46 @@ void NiagaraValidation::ValidateAllRulesInSystem(TSharedPtr<FNiagaraSystemViewMo
 	TArray<FNiagaraValidationResult> NiagaraValidationResults;
 	
 	UNiagaraSystem& NiagaraSystem = SysViewModel->GetSystem();
-	if (NiagaraSystem.GetEffectType())
+
+	// Helper function
+	const auto& ExecuteValidateRules =
+		[&](TConstArrayView<TObjectPtr<UNiagaraValidationRule>> ValidationRules)
+		{
+			for (const UNiagaraValidationRule* ValidationRule : ValidationRules)
+			{
+				if (ValidationRule)
+				{
+					ValidationRule->CheckValidity(Context, NiagaraValidationResults);
+				}
+		}
+	};
+
+	// Validate Global Rules
+	if (const UNiagaraEditorSettings* EditorSettings = GetDefault<UNiagaraEditorSettings>())
+	{
+		for (const TSoftObjectPtr<UNiagaraValidationRuleSet>& ValidationRuleSetPtr : EditorSettings->DefaultValidationRuleSets)
+		{
+			if (const UNiagaraValidationRuleSet* ValidationRuleSet = ValidationRuleSetPtr.Get())
+			{
+				ExecuteValidateRules(ValidationRuleSet->ValidationRules);
+			}
+		}
+	}
+
+	// Validate EffectType Rules
+	if (UNiagaraEffectType* EffectType = NiagaraSystem.GetEffectType())
 	{
-		// go over the validation rules in the effect type
-		for (UNiagaraValidationRule* ValidationRule : NiagaraSystem.GetEffectType()->ValidationRules)
+		ExecuteValidateRules(EffectType->ValidationRules);
+		for (UNiagaraValidationRuleSet* ValidationRuleSet : EffectType->ValidationRuleSets)
 		{
-			if (ValidationRule)
+			if (ValidationRuleSet != nullptr)
 			{
-				ValidationRule->CheckValidity(Context, NiagaraValidationResults);
+				ExecuteValidateRules(ValidationRuleSet->ValidationRules);
 			}
 		}
 	}
 
-	// go over the module-specific rules
+	// Validate Module Specific Rules
 	TArray<UNiagaraStackModuleItem*> StackModuleItems =	NiagaraValidation::GetAllStackEntriesInSystem<UNiagaraStackModuleItem>(Context.ViewModel);
 	for (UNiagaraStackModuleItem* Module : StackModuleItems)
 	{
@@ -647,7 +675,7 @@ void UNiagaraValidationRule_NoFixedDeltaTime::CheckValidity(const FNiagaraValida
 		{
 			OutResults.Emplace_GetRef(
 				ENiagaraValidationSeverity::Error,
-				LOCTEXT("NoFixedDeltaTime", "Effect tyoe does not allow fixed tick delta time"),
+				LOCTEXT("NoFixedDeltaTime", "Effect type does not allow fixed tick delta time"),
 				LOCTEXT("NoFixedDeltaTimeDetailed", "This system uses a fixed tick delta time, which means it might tick multiple times per frame or might skip ticks depending on the global tick rate.\nThe selected effect type does not allow fixed tick delta times."),
 				NiagaraValidation::GetStackEntry<UNiagaraStackSystemPropertiesItem>(Context.ViewModel->GetSystemStackViewModel())
 			);
diff --git a/Engine/Plugins/FX/Niagara/Source/NiagaraEditor/Public/NiagaraEditorSettings.h b/Engine/Plugins/FX/Niagara/Source/NiagaraEditor/Public/NiagaraEditorSettings.h
--- a/Engine/Plugins/FX/Niagara/Source/NiagaraEditor/Public/NiagaraEditorSettings.h
+++ b/Engine/Plugins/FX/Niagara/Source/NiagaraEditor/Public/NiagaraEditorSettings.h
@@ -5,6 +5,7 @@
 #include "NiagaraSystem.h"
 #include "NiagaraEmitter.h"
 #include "NiagaraScript.h"
+#include "NiagaraValidationRuleSet.h"
 #include "Engine/DeveloperSettings.h"
 #include "NiagaraSpawnShortcut.h"
 #include "NiagaraEditorSettings.generated.h"
@@ -263,6 +264,10 @@ public:
 	UPROPERTY(config, EditAnywhere, Category = Niagara)
 	FSoftObjectPath RequiredSystemUpdateScript;
 
+	/** Validation rules applied to all Niagara systems. */
+	UPROPERTY(config, EditAnywhere, Category = Niagara)
+	TArray<TSoftObjectPtr<UNiagaraValidationRuleSet>> DefaultValidationRuleSets;
+
 	/** Shortcut key bindings that if held down while doing a mouse click, will spawn the specified type of Niagara node.*/
 	UPROPERTY(config, EditAnywhere, Category = Niagara)
 	TArray<FNiagaraSpawnShortcut> GraphCreationShortcuts;
