- Remove double DI Class on rect occlusion

#rnx

diff --git a/Engine/Plugins/FX/Niagara/Source/Niagara/Private/NiagaraDataInterfaceOcclusion.cpp b/Engine/Plugins/FX/Niagara/Source/Niagara/Private/NiagaraDataInterfaceOcclusion.cpp
--- a/Engine/Plugins/FX/Niagara/Source/Niagara/Private/NiagaraDataInterfaceOcclusion.cpp
+++ b/Engine/Plugins/FX/Niagara/Source/Niagara/Private/NiagaraDataInterfaceOcclusion.cpp
@@ -72,7 +72,6 @@ void UNiagaraDataInterfaceOcclusion::GetFunctions(TArray<FNiagaraFunctionSignatu
 		FNiagaraFunctionSignature& Sig = OutFunctions.Add_GetRef(DefaultSig);
 		Sig.Name = GetCameraOcclusionRectangleName;
 		Sig.SetDescription(LOCTEXT("GetCameraOcclusionRectFunctionDescription", "This function returns the occlusion factor of a sprite. It samples the depth buffer in a rectangular grid around the given world position and compares each sample with the camera distance."));
-		Sig.AddInput(FNiagaraVariable(FNiagaraTypeDefinition(GetClass()), TEXT("Occlusion interface")));
 		Sig.AddInput(FNiagaraVariable(FNiagaraTypeDefinition::GetPositionDef(), TEXT("Sample Center World Position")), LOCTEXT("RectCenterPosDescription", "This world space position where the center of the sample rectangle should be."));
 		Sig.AddInput(FNiagaraVariable(FNiagaraTypeDefinition::GetFloatDef(), TEXT("Sample Window Width World")), LOCTEXT("SampleWindowWidthWorldDescription", "The total width of the sample rectangle in world space.\nIf the particle is a camera-aligned sprite then this is the sprite width."));
 		Sig.AddInput(FNiagaraVariable(FNiagaraTypeDefinition::GetFloatDef(), TEXT("Sample Window Height World")), LOCTEXT("SampleWindowHeightWorldDescription", "The total height of the sample rectangle in world space.\nIf the particle is a camera-aligned sprite then this is the sprite height."));
diff --git a/Engine/Source/Runtime/Core/Private/UObject/FortniteMainBranchObjectVersion.cpp b/Engine/Source/Runtime/Core/Private/UObject/FortniteMainBranchObjectVersion.cpp
--- a/Engine/Source/Runtime/Core/Private/UObject/FortniteMainBranchObjectVersion.cpp
+++ b/Engine/Source/Runtime/Core/Private/UObject/FortniteMainBranchObjectVersion.cpp
@@ -11,7 +11,7 @@ TMap<FGuid, FGuid> FFortniteMainBranchObjectVersion::GetSystemGuids()
 	SystemGuids.Add(DevGuids.MATERIALSHADERMAP_DERIVEDDATA_VER, FGuid("1DA85DA5733C4A489A1BE117CDC5ACCC"));
 	SystemGuids.Add(DevGuids.NANITE_DERIVEDDATA_VER, FGuid("B62CDCAF66AA45289DDB82E276E7C3E7"));
 	SystemGuids.Add(DevGuids.NIAGARASHADERMAP_DERIVEDDATA_VER, FGuid("2289D5116CF94BC0AFEAEC541468E645"));
-	SystemGuids.Add(DevGuids.Niagara_LatestScriptCompileVersion, FGuid("1F88A64AFB06BD4EACA3246E75D77770"));
+	SystemGuids.Add(DevGuids.Niagara_LatestScriptCompileVersion, FGuid("9CFE679BEE95F345BBC06744597AC064"));
 	SystemGuids.Add(DevGuids.SkeletalMeshDerivedDataVersion, FGuid("ACF593EAAE354FCBB34CF44F5AA1BFD2"));
 	SystemGuids.Add(DevGuids.STATICMESH_DERIVEDDATA_VER, FGuid("18E028FCF1BC469385806AA5F8AE07CC"));
 	return SystemGuids;
