More Niagara shader compiler code should be behind WITH_EDITOR checks.
Also has some minor cleanups.

diff --git a/Engine/Plugins/FX/Niagara/Source/NiagaraShader/Private/NiagaraShader.cpp b/Engine/Plugins/FX/Niagara/Source/NiagaraShader/Private/NiagaraShader.cpp
--- a/Engine/Plugins/FX/Niagara/Source/NiagaraShader/Private/NiagaraShader.cpp
+++ b/Engine/Plugins/FX/Niagara/Source/NiagaraShader/Private/NiagaraShader.cpp
@@ -70,24 +70,18 @@ namespace NiagaraShaderCookStats
 //
 // Globals
 //
-NIAGARASHADER_API FCriticalSection GIdToNiagaraShaderMapCS;
+FCriticalSection GIdToNiagaraShaderMapCS;
 TMap<FNiagaraShaderMapId, FNiagaraShaderMap*> FNiagaraShaderMap::GIdToNiagaraShaderMap[SP_NumPlatforms];
 TArray<FNiagaraShaderMap*> FNiagaraShaderMap::AllNiagaraShaderMaps;
 
-
-/** 
- * Tracks FNiagaraShaderScripts and their shader maps that are being compiled.
- * Uses a TRefCountPtr as this will be the only reference to a shader map while it is being compiled.
- */
-TMap<TRefCountPtr<FNiagaraShaderMap>, TArray<FNiagaraShaderScript*> > FNiagaraShaderMap::NiagaraShaderMapsBeingCompiled;
-
-
+#if WITH_EDITOR
+TMap<TRefCountPtr<FNiagaraShaderMap>, TArray<FNiagaraShaderScript*>> FNiagaraShaderMap::NiagaraShaderMapsBeingCompiled;
 
 static inline bool ShouldCacheNiagaraShader(const FNiagaraShaderType* ShaderType, EShaderPlatform Platform, const FNiagaraShaderScript* Script)
 {
 	return ShaderType->ShouldCache(Platform, Script) && Script->ShouldCache(Platform, ShaderType);
 }
-
+#endif //WITH_EDITOR
 
 
 /** Called for every script shader to update the appropriate stats. */
@@ -248,6 +242,7 @@ bool FNiagaraShaderMapId::operator==(const FNiagaraShaderMapId& ReferenceSet) co
 	return true;
 }
 
+#if WITH_EDITOR
 void FNiagaraShaderMapId::AppendKeyString(FString& KeyString) const
 {
 	BaseCompileHash.AppendString(KeyString);
@@ -312,9 +307,6 @@ void FNiagaraShaderMapId::AppendKeyString(FString& KeyString) const
 	AppendKeyStringShaderDependencies(MakeArrayView(ShaderTypeDependencies), LayoutParams, KeyString);
 }
 
-
-#if WITH_EDITOR
-
 /**
  * Enqueues a compilation for a new shader of this type.
  * @param script - The script to link the shader with.
@@ -392,7 +384,6 @@ void FNiagaraShaderType::BeginCompileShader(
 
 	NewJobs.Add(FShaderCommonCompileJobPtr(NewJob));
 }
-#endif
 
 void FNiagaraShaderType::CacheUniformBufferIncludes(TMap<const TCHAR*, FCachedUniformBufferDeclaration>& Cache, EShaderPlatform Platform) const
 {
@@ -412,8 +403,6 @@ void FNiagaraShaderType::CacheUniformBufferIncludes(TMap<const TCHAR*, FCachedUn
 	}
 }
 
-
-
 void FNiagaraShaderType::AddReferencedUniformBufferIncludes(FShaderCompilerEnvironment& OutEnvironment, FString& OutSourceFilePrefix, EShaderPlatform Platform) const
 {
 	// Cache uniform buffer struct declarations referenced by this shader type's files
@@ -474,6 +463,7 @@ FShader* FNiagaraShaderType::FinishCompileShader(
 
 	return Shader;
 }
+#endif // WITH_EDITOR
 
 /**
 * Finds the shader map for a script.
@@ -488,19 +478,7 @@ FNiagaraShaderMap* FNiagaraShaderMap::FindId(const FNiagaraShaderMapId& ShaderMa
 	return Result;
 }
 
-/** Flushes the given shader types from any loaded FNiagaraShaderMap's. */
-void FNiagaraShaderMap::FlushShaderTypes(TArray<const FShaderType*>& ShaderTypesToFlush)
-{
-	for (int32 ShaderMapIndex = 0; ShaderMapIndex < AllNiagaraShaderMaps.Num(); ShaderMapIndex++)
-	{
-		FNiagaraShaderMap* CurrentShaderMap = AllNiagaraShaderMaps[ShaderMapIndex];
-
-		for (int32 ShaderTypeIndex = 0; ShaderTypeIndex < ShaderTypesToFlush.Num(); ShaderTypeIndex++)
-		{
-			CurrentShaderMap->FlushShadersByShaderType(ShaderTypesToFlush[ShaderTypeIndex]);
-		}
-	}
-}
+#if WITH_EDITOR
 
 void NiagaraShaderMapAppendKeyString(EShaderPlatform Platform, FString& KeyString)
 {
@@ -508,8 +486,6 @@ void NiagaraShaderMapAppendKeyString(EShaderPlatform Platform, FString& KeyStrin
 	// for example static switches will need to go here
 }
 
-#if WITH_EDITOR
-
 /** Creates a string key for the derived data cache given a shader map id. */
 static FString GetNiagaraShaderMapKeyString(const FNiagaraShaderMapId& ShaderMapId, EShaderPlatform Platform)
 {
@@ -901,7 +877,7 @@ void FNiagaraShaderMap::LoadMissingShadersFromMemory(const FNiagaraShaderScript*
 	}
 #endif
 }
-#endif
+#endif // WITH_EDITOR
 
 void FNiagaraShaderMap::GetShaderList(TMap<FShaderId, TShaderRef<FShader>>& OutShaders) const
 {
@@ -985,21 +961,6 @@ FNiagaraShaderMap::~FNiagaraShaderMap()
 	AllNiagaraShaderMaps.RemoveSwap(this);
 }
 
-/**
- * Removes all entries in the cache with exceptions based on a shader type
- * @param ShaderType - The shader type to flush
- */
-void FNiagaraShaderMap::FlushShadersByShaderType(const FShaderType* ShaderType)
-{
-	if (ShaderType->GetNiagaraShaderType())
-	{
-		for (int32 PermutationId = 0; PermutationId < ShaderType->GetPermutationCount(); ++PermutationId)
-		{
-			GetMutableContent()->RemoveShaderTypePermutaion(ShaderType->GetNiagaraShaderType(), PermutationId);
-		}
-	}
-}
-
 bool FNiagaraShaderMap::Serialize(FArchive& Ar, bool bInlineShaderResources, bool bLoadedByCookedMaterial)
 {
 	// Note: This is saved to the DDC, not into packages (except when cooked)
@@ -1008,6 +969,7 @@ bool FNiagaraShaderMap::Serialize(FArchive& Ar, bool bInlineShaderResources, boo
 	return Super::Serialize(Ar, bInlineShaderResources, bLoadedByCookedMaterial);
 }
 
+#if WITH_EDITOR
 bool FNiagaraShaderMap::RemovePendingScript(FNiagaraShaderScript* Script)
 {
 	bool bRemoved = false;
@@ -1082,6 +1044,7 @@ const FNiagaraShaderMap* FNiagaraShaderMap::GetShaderMapBeingCompiled(const FNia
 
 	return NULL;
 }
+#endif // WITH_EDITOR
 
 //////////////////////////////////////////////////////////////////////////
 
diff --git a/Engine/Plugins/FX/Niagara/Source/NiagaraShader/Private/NiagaraShared.cpp b/Engine/Plugins/FX/Niagara/Source/NiagaraShader/Private/NiagaraShared.cpp
--- a/Engine/Plugins/FX/Niagara/Source/NiagaraShader/Private/NiagaraShared.cpp
+++ b/Engine/Plugins/FX/Niagara/Source/NiagaraShader/Private/NiagaraShared.cpp
@@ -48,6 +48,19 @@ static FAutoConsoleVariableRef CVarNiagaraTranslatorSilenceFailIfNotSet(
 	FNiagaraCompilationQueue* FNiagaraCompilationQueue::Singleton = nullptr;
 #endif
 
+FNiagaraShaderScript::FNiagaraShaderScript()
+	: BaseVMScript(nullptr)
+	, GameThreadShaderMap(nullptr)
+	, RenderingThreadShaderMap(nullptr)
+	, ScriptParametersMetadata(MakeShared<FNiagaraShaderScriptParametersMetadata>())
+	, FeatureLevel(GMaxRHIFeatureLevel)
+	, ShaderPlatform(SP_NumPlatforms)
+	, bLoadedCookedShaderMapId(false)
+	, bLoadedFromCookedMaterial(false)
+	, bQueuedForRelease(false)
+{
+}
+
 FNiagaraShaderScript::~FNiagaraShaderScript()
 {
 #if WITH_EDITOR
@@ -56,6 +69,7 @@ FNiagaraShaderScript::~FNiagaraShaderScript()
 #endif
 }
 
+#if WITH_EDITOR
 /** Populates OutEnvironment with defines needed to compile shaders for this script. */
 void FNiagaraShaderScript::SetupShaderCompilationEnvironment(
 	EShaderPlatform Platform,
@@ -64,7 +78,7 @@ void FNiagaraShaderScript::SetupShaderCompilationEnvironment(
 {
 	OutEnvironment.SetDefine(TEXT("GPU_SIMULATION_SHADER"), TEXT("1"));
 }
-
+#endif // WITH_EDITOR
 
 bool FNiagaraShaderScript::ShouldCache(EShaderPlatform Platform, const FShaderType* ShaderType) const
 {
@@ -92,6 +106,7 @@ bool FNiagaraShaderScript::GetUsesCompressedAttributes() const
 	return AdditionalDefines.Contains(TEXT("CompressAttributes"));
 }
 
+#if WITH_EDITOR
 void FNiagaraShaderScript::NotifyCompilationFinished()
 {
 	UpdateCachedData_PostCompile();
@@ -101,7 +116,6 @@ void FNiagaraShaderScript::NotifyCompilationFinished()
 
 void FNiagaraShaderScript::CancelCompilation()
 {
-#if WITH_EDITOR
 	check(IsInGameThread());
 	bool bWasPending = FNiagaraShaderMap::RemovePendingScript(this);
 	FNiagaraCompilationQueue::Get()->RemovePending(this);
@@ -112,7 +126,6 @@ void FNiagaraShaderScript::CancelCompilation()
 		UE_LOG(LogShaders, Log, TEXT("CancelCompilation %p."), this);
 	}
 	OutstandingCompileShaderMapIds.Empty();
-#endif
 }
 
 void FNiagaraShaderScript::RemoveOutstandingCompileId(const int32 OldOutstandingCompileShaderMapId)
@@ -123,10 +136,13 @@ void FNiagaraShaderScript::RemoveOutstandingCompileId(const int32 OldOutstanding
 		//UE_LOG(LogShaders, Log, TEXT("RemoveOutstandingCompileId %p %d"), this, OldOutstandingCompileShaderMapId);
 	}
 }
+#endif // WITH_EDITOR
 
 void FNiagaraShaderScript::Invalidate()
 {
+#if WITH_EDITOR
 	CancelCompilation();
+#endif
 	ReleaseShaderMap();
 #if WITH_EDITOR
 	CompileErrors.Empty();
@@ -134,10 +150,6 @@ void FNiagaraShaderScript::Invalidate()
 #endif
 }
 
-void FNiagaraShaderScript::LegacySerialize(FArchive& Ar)
-{
-}
-
 bool FNiagaraShaderScript::IsSame(const FNiagaraShaderMapId& InId) const
 {
 	if (InId.ReferencedCompileHashes.Num() != ReferencedCompileHashes.Num() ||
@@ -184,10 +196,12 @@ bool FNiagaraShaderScript::IsShaderMapComplete() const
 		return false;
 	}
 
+#if WITH_EDITOR
 	if (FNiagaraShaderMap::GetShaderMapBeingCompiled(this) != nullptr)
 	{
 		return false;
 	}
+#endif
 
 	if (!GameThreadShaderMap->IsValid())
 	{
@@ -204,6 +218,7 @@ bool FNiagaraShaderScript::IsShaderMapComplete() const
 	return true;
 }
 
+#if WITH_EDITOR
 void FNiagaraShaderScript::GetDependentShaderTypes(EShaderPlatform Platform, TArray<FShaderType*>& OutShaderTypes) const
 {
 	for (TLinkedList<FShaderType*>::TIterator ShaderTypeIt(FShaderType::GetTypeList()); ShaderTypeIt; ShaderTypeIt.Next())
@@ -259,7 +274,6 @@ void FNiagaraShaderScript::GetShaderMapId(EShaderPlatform Platform, const ITarge
 			OutId.ShaderTypeDependencies.Emplace(ShaderType, Platform);
 		}
 
-#if WITH_EDITOR
 		if (TargetPlatform)
 		{
 			OutId.LayoutParams.InitializeForPlatform(TargetPlatform->IniPlatformName(), TargetPlatform->HasEditorOnlyData());
@@ -268,16 +282,9 @@ void FNiagaraShaderScript::GetShaderMapId(EShaderPlatform Platform, const ITarge
 		{
 			OutId.LayoutParams.InitializeForCurrent();
 		}
-#else
-		if (TargetPlatform != nullptr)
-		{
-			UE_LOG(LogShaders, Error, TEXT("FNiagaraShaderScript::GetShaderMapId: TargetPlatform is not null, but a cooked executable cannot target platforms other than its own."));
-		}
-		OutId.LayoutParams.InitializeForCurrent();
-#endif
 	}
 }
-
+#endif // WITH_EDITOR
 
 
 void FNiagaraShaderScript::AddReferencedObjects(FReferenceCollector& Collector)
@@ -396,7 +403,7 @@ void FNiagaraShaderScript::SaveShaderStableKeys(EShaderPlatform TargetShaderPlat
 		GameThreadShaderMap->SaveShaderStableKeys(TargetShaderPlatform, SaveKeyVal);
 	}
 }
-#endif
+#endif // WITH_EDITOR
 
 void FNiagaraShaderScript::SetScript(UNiagaraScriptBase* InScript, ERHIFeatureLevel::Type InFeatureLevel, EShaderPlatform InShaderPlatform, const FGuid& InCompilerVersionID,  const TArray<FString>& InAdditionalDefines, const TArray<FString>& InAdditionalVariables,
 		const FNiagaraCompileHash& InBaseCompileHash, const TArray<FNiagaraCompileHash>& InReferencedCompileHashes, 
@@ -428,7 +435,7 @@ bool FNiagaraShaderScript::MatchesScript(ERHIFeatureLevel::Type InFeatureLevel,
 		&& FeatureLevel == InFeatureLevel
 		&& ShaderPlatform == InShaderPlatform;
 }
-#endif
+#endif // WITH_EDITOR
 
 void FNiagaraShaderScript::SetRenderingThreadShaderMap(FNiagaraShaderMap* InShaderMap)
 {
@@ -438,13 +445,17 @@ void FNiagaraShaderScript::SetRenderingThreadShaderMap(FNiagaraShaderMap* InShad
 
 bool FNiagaraShaderScript::IsCompilationFinished() const
 {
+#if WITH_EDITOR
 	check(IsInGameThread());
-	bool bRet = GameThreadShaderMap && GameThreadShaderMap.IsValid() && GameThreadShaderMap->IsCompilationFinalized();
 	if (OutstandingCompileShaderMapIds.Num() == 0)
 	{
 		return true;
 	}
+	bool bRet = GameThreadShaderMap && GameThreadShaderMap.IsValid() && GameThreadShaderMap->IsCompilationFinalized();
 	return bRet;
+#else
+	return true;
+#endif
 }
 
 void FNiagaraShaderScript::SetRenderThreadCachedData(const FNiagaraShaderMapCachedData& CachedData)
@@ -682,7 +693,7 @@ void FNiagaraShaderScript::FinishCompilation()
 	}
 }
 
-#endif
+#endif // WITH_EDITOR
 
 void FNiagaraShaderScript::BuildScriptParametersMetadata(const FNiagaraShaderScriptParametersMetadata& InScriptParametersMetadata)
 {
@@ -734,7 +745,7 @@ FNiagaraShaderRef FNiagaraShaderScript::GetShaderGameThread(int32 PermutationId)
 	return FNiagaraShaderRef();
 };
 
-
+#if WITH_EDITOR
 void FNiagaraShaderScript::GetShaderMapIDsWithUnfinishedCompilation(TArray<int32>& ShaderMapIds)
 {
 	// Build an array of the shader map Id's are not finished compiling.
@@ -748,8 +759,6 @@ void FNiagaraShaderScript::GetShaderMapIDsWithUnfinishedCompilation(TArray<int32
 	}
 }
 
-#if WITH_EDITOR
-
 /**
 * Compiles this script for Platform, storing the result in OutShaderMap
 *
@@ -765,7 +774,6 @@ bool FNiagaraShaderScript::BeginCompileShaderMap(
 	bool bSynchronous)
 {
 	check(IsInGameThread());
-#if WITH_EDITORONLY_DATA
 	bool bSuccess = false;
 
 	STAT(double NiagaraCompileTime = 0);
@@ -794,10 +802,6 @@ bool FNiagaraShaderScript::BeginCompileShaderMap(
 	INC_FLOAT_STAT_BY(STAT_ShaderCompiling_NiagaraShaders, (float)NiagaraCompileTime);
 
 	return true;
-#else
-	UE_LOG(LogShaders, Fatal, TEXT("Compiling of shaders in a build without editordata is not supported."));
-	return false;
-#endif
 }
 
 FNiagaraCompileEventSeverity FNiagaraCVarUtilities::GetCompileEventSeverityForFailIfNotSet()
@@ -819,5 +823,5 @@ bool FNiagaraCVarUtilities::GetShouldEmitMessagesForFailIfNotSet()
 	return GNiagaraTranslatorFailIfNotSetSeverity != 0;
 }
 
-#endif
+#endif // WITH_EDITOR
 
diff --git a/Engine/Plugins/FX/Niagara/Source/NiagaraShader/Public/NiagaraShaderCompilationManager.h b/Engine/Plugins/FX/Niagara/Source/NiagaraShader/Public/NiagaraShaderCompilationManager.h
--- a/Engine/Plugins/FX/Niagara/Source/NiagaraShader/Public/NiagaraShaderCompilationManager.h
+++ b/Engine/Plugins/FX/Niagara/Source/NiagaraShader/Public/NiagaraShaderCompilationManager.h
@@ -2,6 +2,8 @@
 
 #pragma once
 
+#if WITH_EDITOR
+
 #include "NiagaraShared.h"
 #include "ShaderCompiler.h"
 
@@ -55,5 +57,6 @@ private:
 	TMap<int32, FNiagaraShaderMapFinalizeResults> PendingFinalizeNiagaraShaderMaps;
 };
 
-
 extern NIAGARASHADER_API FNiagaraShaderCompilationManager GNiagaraShaderCompilationManager;
+
+#endif //WITH_EDITOR
diff --git a/Engine/Plugins/FX/Niagara/Source/NiagaraShader/Public/NiagaraShaderType.h b/Engine/Plugins/FX/Niagara/Source/NiagaraShader/Public/NiagaraShaderType.h
--- a/Engine/Plugins/FX/Niagara/Source/NiagaraShader/Public/NiagaraShaderType.h
+++ b/Engine/Plugins/FX/Niagara/Source/NiagaraShader/Public/NiagaraShaderType.h
@@ -110,6 +110,7 @@ public:
 		check(InTotalPermutationCount == 1);
 	}
 
+#if WITH_EDITOR
 	/**
 	 * Enqueues a compilation for a new shader of this type.
 	 * @param Script - The script to link the shader with.
@@ -134,6 +135,7 @@ public:
 		const FShaderCompileJob& CurrentJob,
 		const FString& InDebugDescription
 		) const;
+#endif // WITH_EDITOR
 
 	/**
 	 * Checks if the shader type should be cached for a particular platform and script.
@@ -146,12 +148,12 @@ public:
 		return ShouldCompilePermutation(FNiagaraShaderPermutationParameters(Platform, Script));
 	}
 
+#if WITH_EDITOR
 	/** Adds include statements for uniform buffers that this shader type references, and builds a prefix for the shader file with the include statements. */
 	void AddReferencedUniformBufferIncludes(FShaderCompilerEnvironment& OutEnvironment, FString& OutSourceFilePrefix, EShaderPlatform Platform) const;
 
 	void CacheUniformBufferIncludes(TMap<const TCHAR*, FCachedUniformBufferDeclaration>& Cache, EShaderPlatform Platform) const;
 
-
 protected:
 
 	/**
@@ -163,4 +165,5 @@ protected:
 	{
 		ModifyCompilationEnvironment(FNiagaraShaderPermutationParameters(Platform, Script), Environment);
 	}
+#endif // WITH_EDITOR
 };
diff --git a/Engine/Plugins/FX/Niagara/Source/NiagaraShader/Public/NiagaraShared.h b/Engine/Plugins/FX/Niagara/Source/NiagaraShader/Public/NiagaraShared.h
--- a/Engine/Plugins/FX/Niagara/Source/NiagaraShader/Public/NiagaraShared.h
+++ b/Engine/Plugins/FX/Niagara/Source/NiagaraShader/Public/NiagaraShared.h
@@ -385,8 +385,10 @@ public:
 		return !(*this == ReferenceSet);
 	}
 
+#if WITH_EDITOR
 	/** Appends string representations of this Id to a key string. */
 	void AppendKeyString(FString& KeyString) const;
+#endif
 
 	/** Returns true if the requested shader type is a dependency of this shader map Id. */
 	//bool ContainsShaderType(const FShaderType* ShaderType) const;
@@ -494,26 +496,26 @@ public:
 	*/
 	static FNiagaraShaderMap* FindId(const FNiagaraShaderMapId& ShaderMapId, EShaderPlatform Platform);
 
-	/** Flushes the given shader types from any loaded FNiagaraShaderMaps. */
-	static void FlushShaderTypes(TArray<const FShaderType*>& ShaderTypesToFlush);
-
 	// ShaderMap interface
 	template<typename ShaderType> TNiagaraShaderRef<ShaderType> GetShader(int32 PermutationId) const { return TNiagaraShaderRef<ShaderType>(GetContent()->GetShader<ShaderType>(PermutationId), *this); }
 	TNiagaraShaderRef<FShader> GetShader(FShaderType* ShaderType, int32 PermutationId) const { return TNiagaraShaderRef<FShader>(GetContent()->GetShader(ShaderType, PermutationId), *this); }
 
 	//static void FixupShaderTypes(EShaderPlatform Platform, const TMap<FShaderType*, FString>& ShaderTypeNames);
 
+#if WITH_EDITOR
 	/**
 	* Attempts to load the shader map for the given script from the Derived Data Cache.
 	* If InOutShaderMap is valid, attempts to load the individual missing shaders instead.
 	*/
 	static void LoadFromDerivedDataCache(const FNiagaraShaderScript* Script, const FNiagaraShaderMapId& ShaderMapId, EShaderPlatform Platform, TRefCountPtr<FNiagaraShaderMap>& InOutShaderMap);
+#endif
 
 	FNiagaraShaderMap();
 
 	// Destructor.
 	~FNiagaraShaderMap();
 
+#if WITH_EDITOR
 	/**
 	* Compiles the shaders for a script and caches them in this shader map.
 	* @param script - The script to compile shaders for.
@@ -550,6 +552,7 @@ public:
 	* @return True if the shader map was being compiled and we added Script to the list to be applied.
 	*/
 	bool TryToAddToExistingCompilationTask(FNiagaraShaderScript* Script);
+#endif // WITH_EDITOR
 
 	/** Builds a list of the shaders in a shader map. */
 	NIAGARASHADER_API  void GetShaderList(TMap<FShaderId, TShaderRef<FShader>>& OutShaders) const;
@@ -564,12 +567,7 @@ public:
 	NIAGARASHADER_API  void AddRef();
 	NIAGARASHADER_API  void Release();
 
-	/**
-	* Removes all entries in the cache with exceptions based on a shader type
-	* @param ShaderType - The shader type to flush
-	*/
-	void FlushShadersByShaderType(const FShaderType* ShaderType);
-
+#if WITH_EDITOR
 	/** Removes a Script from NiagaraShaderMapsBeingCompiled. 
 	* @return true if something was actually removed.
 	*/
@@ -578,17 +576,21 @@ public:
 
 	/** Finds a shader map currently being compiled that was enqueued for the given script. */
 	static const FNiagaraShaderMap* GetShaderMapBeingCompiled(const FNiagaraShaderScript* Script);
+#endif // WITH_EDITOR
 
 	/** Serializes the shader map. */
 	bool Serialize(FArchive& Ar, bool bInlineShaderResources = true, bool bLoadedByCookedMaterial = false);
 
+#if WITH_EDITOR
 	/** Saves this shader map to the derived data cache. */
 	void SaveToDerivedDataCache();
 
 	/** Backs up any FShaders in this shader map to memory through serialization and clears FShader references. */
 	TArray<uint8>* BackupShadersToMemory();
+
 	/** Recreates FShaders from the passed in memory, handling shader key changes. */
 	void RestoreShadersFromMemory(const TArray<uint8>& ShaderData);
+#endif // WITH_EDITOR
 
 	// Accessors.
 	const FNiagaraShaderMapId& GetShaderMapId() const		{ return GetContent()->ShaderMapId; }
@@ -608,7 +610,9 @@ public:
 
 	int32 GetNumRefs() const { return NumRefs; }
 	uint32 GetCompilingId()  { return CompilingId; }
-	static TMap<TRefCountPtr<FNiagaraShaderMap>, TArray<FNiagaraShaderScript*> > &GetInFlightShaderMaps() 
+
+#if WITH_EDITOR
+	static TMap<TRefCountPtr<FNiagaraShaderMap>, TArray<FNiagaraShaderScript*>>& GetInFlightShaderMaps()
 	{
 		//All access to NiagaraShaderMapsBeingCompiled must be done on the game thread!
 		check(IsInGameThread());
@@ -616,6 +620,8 @@ public:
 	}
 
 	void SetCompiledSuccessfully(bool bSuccess) { bCompiledSuccessfully = bSuccess; }
+#endif // WITH_EDITOR
+
 private:
 
 	/**
@@ -631,8 +637,10 @@ private:
 	*/
 	static TArray<FNiagaraShaderMap*> AllNiagaraShaderMaps;
 
+#if WITH_EDITOR
 	/** Tracks resources and their shader maps that need to be compiled but whose compilation is being deferred. */
 	static TMap<TRefCountPtr<FNiagaraShaderMap>, TArray<FNiagaraShaderScript*> > NiagaraShaderMapsBeingCompiled;
+#endif
 
 	/** Uniquely identifies this shader map during compilation, needed for deferred compilation where shaders from multiple shader maps are compiled together. */
 	uint32 CompilingId;
@@ -656,9 +664,11 @@ private:
 	/** Indicates whether the shader map should be stored in the shader cache. */
 	uint32 bIsPersistent : 1;
 
+#if WITH_EDITOR
 	FShader* ProcessCompilationResultsForSingleJob(const TRefCountPtr<class FShaderCommonCompileJob>& SingleJob, const FSHAHash& ShaderMapHash);
 
 	bool IsNiagaraShaderComplete(const FNiagaraShaderScript* Script, const FNiagaraShaderType* ShaderType, bool bSilent);
+#endif
 
 	friend class FShaderCompilingManager;
 };
@@ -691,18 +701,7 @@ public:
 	/**
 	 * Minimal initialization constructor.
 	 */
-	FNiagaraShaderScript()
-		: BaseVMScript(nullptr)
-		, GameThreadShaderMap(nullptr)
-		, RenderingThreadShaderMap(nullptr)
-		, ScriptParametersMetadata(MakeShared<FNiagaraShaderScriptParametersMetadata>())
-		, FeatureLevel(GMaxRHIFeatureLevel)
-		, ShaderPlatform(SP_NumPlatforms)
-		, bLoadedCookedShaderMapId(false)
-		, bLoadedFromCookedMaterial(false)
-		, bQueuedForRelease(false)
-	{
-	}
+	NIAGARASHADER_API FNiagaraShaderScript();
 
 	/**
 	 * Destructor
@@ -736,14 +735,16 @@ public:
 	NIAGARASHADER_API  virtual void ModifyCompilationEnvironment(EShaderPlatform Platform, struct FShaderCompilerEnvironment& OutEnvironment) const;
 
 	/** Serializes the script. */
-	NIAGARASHADER_API  virtual void LegacySerialize(FArchive& Ar);
-
 	NIAGARASHADER_API void SerializeShaderMap(FArchive& Ar);
+
 	/** Releases this script's shader map.  Must only be called on scripts not exposed to the rendering thread! */
 	NIAGARASHADER_API void ReleaseShaderMap();
 
+#if WITH_EDITOR
 	void GetDependentShaderTypes(EShaderPlatform Platform, TArray<FShaderType*>& OutShaderTypes) const;
+
 	NIAGARASHADER_API  virtual void GetShaderMapId(EShaderPlatform Platform, const ITargetPlatform* TargetPlatform, FNiagaraShaderMapId& OutId) const;
+#endif
 
 	NIAGARASHADER_API void Invalidate();
 
@@ -752,6 +753,7 @@ public:
 	 */
 	virtual bool IsPersistent() const { return true; }
 
+#if WITH_EDITOR
 	/**
 	* Called when compilation finishes, after the GameThreadShaderMap is set and the render command to set the RenderThreadShaderMap is queued
 	*/
@@ -766,6 +768,7 @@ public:
 	 * Blocks until compilation has completed. Returns immediately if a compilation is not outstanding.
 	 */
 	NIAGARASHADER_API  void FinishCompilation();
+#endif // WITH_EDITOR
 
 	/**
 	 * Checks if the compilation for this shader is finished
@@ -804,11 +807,13 @@ public:
 
 	NIAGARASHADER_API bool QueueForRelease(FThreadSafeBool& Fence);
 
+#if WITH_EDITOR
 	void AddCompileId(uint32 Id) 
 	{
 		check(IsInGameThread());
 		OutstandingCompileShaderMapIds.Add(Id);
 	}
+#endif // WITH_EDITOR
 
 	void SetShaderMap(FNiagaraShaderMap* InShaderMap)
 	{
@@ -818,7 +823,9 @@ public:
 		CookedShaderMapId = InShaderMap->GetShaderMapId();
 	}
 
+#if WITH_EDITOR
 	NIAGARASHADER_API void RemoveOutstandingCompileId(const int32 OldOutstandingCompileShaderMapId);
+#endif
 
 	NIAGARASHADER_API  virtual void AddReferencedObjects(FReferenceCollector& Collector);
 
@@ -882,11 +889,12 @@ public:
 	FORCEINLINE bool IsViewUniformBufferUsed_RenderThread() const { return CachedData_RenderThread.bViewUniformBufferUsed != 0; }
 
 protected:
-
+#if WITH_EDITOR
 	/**
 	* Fills the passed array with IDs of shader maps unfinished compilation jobs.
 	*/
 	void GetShaderMapIDsWithUnfinishedCompilation(TArray<int32>& ShaderMapIds);
+#endif
 
 	void SetFeatureLevel(ERHIFeatureLevel::Type InFeatureLevel)
 	{
@@ -943,11 +951,13 @@ private:
 	/** The compile hashes for the top level scripts referenced by the script. */
 	TArray<FNiagaraCompileHash> ReferencedCompileHashes;
 
+#if WITH_EDITOR
 	/** 
 	 * Contains the compiling id of this shader map when it is being compiled asynchronously. 
 	 * This can be used to access the shader map during async compiling, since GameThreadShaderMap will not have been set yet.
 	 */
 	TArray<int32, TInlineAllocator<1> > OutstandingCompileShaderMapIds;
+#endif // WITH_EDITOR
 
 	/** Feature level and shader platform that this script is representing. */
 	ERHIFeatureLevel::Type FeatureLevel;
@@ -965,6 +975,7 @@ private:
 
 	FOnNiagaraScriptCompilationComplete OnCompilationCompleteDelegate;
 
+#if WITH_EDITOR
 	/**
 	* Compiles this script for Platform, storing the result in OutShaderMap if the compile was synchronous
 	*/
@@ -979,7 +990,7 @@ private:
 		EShaderPlatform Platform,
 		FShaderCompilerEnvironment& OutEnvironment
 		) const;
-
+#endif // WITH_EDITOR
 
 	FString FriendlyName;
 
