- Use syncrhonous load for rules set to avoid missing rules

#rnx

diff --git a/Engine/Plugins/FX/Niagara/Source/NiagaraEditor/Private/NiagaraValidationRules.cpp b/Engine/Plugins/FX/Niagara/Source/NiagaraEditor/Private/NiagaraValidationRules.cpp
--- a/Engine/Plugins/FX/Niagara/Source/NiagaraEditor/Private/NiagaraValidationRules.cpp
+++ b/Engine/Plugins/FX/Niagara/Source/NiagaraEditor/Private/NiagaraValidationRules.cpp
@@ -187,7 +187,7 @@ void NiagaraValidation::ValidateAllRulesInSystem(TSharedPtr<FNiagaraSystemViewMo
 	{
 		for (const TSoftObjectPtr<UNiagaraValidationRuleSet>& ValidationRuleSetPtr : EditorSettings->DefaultValidationRuleSets)
 		{
-			if (const UNiagaraValidationRuleSet* ValidationRuleSet = ValidationRuleSetPtr.Get())
+			if (const UNiagaraValidationRuleSet* ValidationRuleSet = ValidationRuleSetPtr.LoadSynchronous())
 			{
 				ExecuteValidateRules(ValidationRuleSet->ValidationRules);
 			}
