Revert check in FNiagaraDataInterfaceProxySubmix::~FNiagaraDataInterfaceProxySubmix() introduced with cl 22806983

#rnx

diff --git a/Engine/Plugins/FX/Niagara/Source/Niagara/Private/NiagaraDataInterfaceAudio.cpp b/Engine/Plugins/FX/Niagara/Source/Niagara/Private/NiagaraDataInterfaceAudio.cpp
--- a/Engine/Plugins/FX/Niagara/Source/Niagara/Private/NiagaraDataInterfaceAudio.cpp
+++ b/Engine/Plugins/FX/Niagara/Source/Niagara/Private/NiagaraDataInterfaceAudio.cpp
@@ -124,7 +124,6 @@ FNiagaraDataInterfaceProxySubmix::FNiagaraDataInterfaceProxySubmix(int32 InNumSa
 
 FNiagaraDataInterfaceProxySubmix::~FNiagaraDataInterfaceProxySubmix()
 {
-	check(IsInRenderingThread());
 }
 
 void FNiagaraDataInterfaceProxySubmix::RegisterToAllAudioDevices()
