- Handle mismatched data interface class on user parameter store
- Logs and replaces with the default instances when this occurs, how we got into this state requires further investigation

#rnx

diff --git a/Engine/Plugins/FX/Niagara/Source/Niagara/Private/NiagaraComponent.cpp b/Engine/Plugins/FX/Niagara/Source/Niagara/Private/NiagaraComponent.cpp
--- a/Engine/Plugins/FX/Niagara/Source/Niagara/Private/NiagaraComponent.cpp
+++ b/Engine/Plugins/FX/Niagara/Source/Niagara/Private/NiagaraComponent.cpp
@@ -3189,10 +3189,29 @@ void FixInvalidDataInterfaceOverrides(TMap<FNiagaraVariableBase, FNiagaraVariant
 		FNiagaraVariant& Value = VariableValuePair.Value;
 		if (Variable.IsDataInterface())
 		{
-			if (Value.GetDataInterface() == nullptr)
+			UNiagaraDataInterface* DataInterface = Value.GetDataInterface();
+			if (DataInterface == nullptr)
 			{
-				UE_LOG(LogNiagara, Warning, TEXT("Replaced invalid user parameter data interface with it's default.  Component: %s Override Source: %s Parameter Name: %s."), *OwningComponent->GetPathName(), *OverrideSource, *Variable.GetName().ToString());
-				Value.SetDataInterface(NewObject<UNiagaraDataInterface>(OwningComponent, Variable.GetType().GetClass(), NAME_None, RF_Transactional | RF_Public));
+				UE_LOG(LogNiagara, Warning, TEXT("Replaced null user parameter data interface with it's default.  Component: %s Override Source: %s Parameter Name: %s."), *OwningComponent->GetPathName(), *OverrideSource, *Variable.GetName().ToString());
+			}
+			else if (DataInterface->GetClass() != Variable.GetType().GetClass())
+			{
+				UE_LOG(LogNiagara, Warning, TEXT("Replaced invalid class '%s' user parameter data interface '%s' with it's default.  Component: %s Override Source: %s Parameter Name: %s."), *DataInterface->GetClass()->GetName(), *Variable.GetType().GetClass()->GetName(), *OwningComponent->GetPathName(), *OverrideSource, *Variable.GetName().ToString());
+				DataInterface = nullptr;
+			}
+			if (DataInterface == nullptr)
+			{
+				DataInterface = NewObject<UNiagaraDataInterface>(OwningComponent, Variable.GetType().GetClass(), NAME_None, RF_Transactional | RF_Public);
+				if (UNiagaraSystem* NiagaraSystem = OwningComponent->GetAsset())
+				{
+					const int32* DataInterfaceOffset = NiagaraSystem->GetExposedParameters().FindParameterOffset(Variable);
+					UNiagaraDataInterface* SourceDataInterface = DataInterfaceOffset ? NiagaraSystem->GetExposedParameters().GetDataInterface(*DataInterfaceOffset) : nullptr;
+					if (DataInterfaceOffset != nullptr)
+					{
+						SourceDataInterface->CopyTo(DataInterface);
+					}
+				}
+				Value.SetDataInterface(DataInterface);
 			}
 		}
 	}
@@ -3238,6 +3257,33 @@ void UNiagaraComponent::FixInvalidUserParameterOverrideData()
 				ExistingVar = PositionVar;
 			}
 		}
+		// Check for invalid or missing data interfaces
+		else if (ExistingVar.IsDataInterface())
+		{
+			UNiagaraDataInterface* DataInterface = OverrideParameters.GetDataInterface(ExistingVar);
+			if (DataInterface == nullptr)
+			{
+				UE_LOG(LogNiagara, Log, TEXT("ParameterStore has null data interface '%s' for variable '%s', making a default."), *ExistingVar.GetType().GetClass()->GetName(), *ExistingVar.GetName().ToString());
+			}
+			else if (DataInterface->GetClass() != ExistingVar.GetType().GetClass())
+			{
+				UE_LOG(LogNiagara, Log, TEXT("ParameterStore has invalid type '%s' data interface '%s' for variable '%s', making a default."), *DataInterface->GetClass()->GetName(), *ExistingVar.GetType().GetClass()->GetName(), *ExistingVar.GetName().ToString());
+				DataInterface = nullptr;
+			}
+
+			if (DataInterface == nullptr)
+			{
+				DataInterface = NewObject<UNiagaraDataInterface>(this, ExistingVar.GetType().GetClass(), NAME_None, RF_Transactional | RF_Public);
+				if ( SourceVars.Contains(ExistingVar) )
+				{
+					if ( UNiagaraDataInterface* SourceDataInterface = Asset->GetExposedParameters().GetDataInterface(ExistingVar) )
+					{
+						SourceDataInterface->CopyTo(DataInterface);
+					}
+				}
+				OverrideParameters.SetDataInterface(DataInterface, ExistingVar);
+			}
+		}
 	}
 }
 #endif
