- Make sure rule sets are applied to validator

#rnx

diff --git a/Engine/Plugins/FX/Niagara/Source/NiagaraEditor/Private/NiagaraValidationRules.cpp b/Engine/Plugins/FX/Niagara/Source/NiagaraEditor/Private/NiagaraValidationRules.cpp
--- a/Engine/Plugins/FX/Niagara/Source/NiagaraEditor/Private/NiagaraValidationRules.cpp
+++ b/Engine/Plugins/FX/Niagara/Source/NiagaraEditor/Private/NiagaraValidationRules.cpp
@@ -156,6 +156,43 @@ namespace NiagaraValidation
 
 // --------------------------------------------------------------------------------------------------------------------------------------------
 
+bool NiagaraValidation::HasValidationRules(UNiagaraSystem* NiagaraSystem)
+{
+	if ( NiagaraSystem != nullptr )
+	{
+		if (const UNiagaraEditorSettings* EditorSettings = GetDefault<UNiagaraEditorSettings>())
+		{
+			for (const TSoftObjectPtr<UNiagaraValidationRuleSet>& ValidationRuleSetPtr : EditorSettings->DefaultValidationRuleSets)
+			{
+				const UNiagaraValidationRuleSet* ValidationRuleSet = ValidationRuleSetPtr.LoadSynchronous();
+				if (ValidationRuleSet != nullptr && ValidationRuleSet->ValidationRules.Num() > 0)
+				{
+					return true;
+				}
+			}
+		}
+
+		if (UNiagaraEffectType* EffectType = NiagaraSystem->GetEffectType())
+		{
+			if (EffectType->ValidationRules.Num() > 0)
+			{
+				return true;
+			}
+
+			for (UNiagaraValidationRuleSet* ValidationRuleSet : EffectType->ValidationRuleSets)
+			{
+				if (ValidationRuleSet != nullptr && ValidationRuleSet->ValidationRules.Num() > 0)
+				{
+					return true;
+				}
+			}
+		}
+	}
+	return false;
+}
+
+// --------------------------------------------------------------------------------------------------------------------------------------------
+
 void NiagaraValidation::ValidateAllRulesInSystem(TSharedPtr<FNiagaraSystemViewModel> SysViewModel, TFunction<void(const FNiagaraValidationResult& Result)> ResultCallback)
 {
 	if (SysViewModel == nullptr)
@@ -413,8 +450,8 @@ void UNiagaraValidationRule_BannedDataInterfaces::CheckValidity(const FNiagaraVa
 			}
 
 			static const FText WarningFormat(LOCTEXT("BannedDataInteraceFormatWarn", "DataInterface '{0}' is banned on currently enabled platforms"));
-			static const FText SystemDescFormat(LOCTEXT("BannedDataInteraceFormatDesc", "DataInterface '{0} - {1}' is banned on currently enabled platforms"));
-			static const FText EmitterDescFormat(LOCTEXT("BannedDataInteraceFormatDesc", "DataInterface '{0} - {1}' is banned on currently enabled platforms '{2}'"));
+			static const FText SystemDescFormat(LOCTEXT("BannedDataInteraceSystemFormatDesc", "DataInterface '{0} - {1}' is banned on currently enabled platforms"));
+			static const FText EmitterDescFormat(LOCTEXT("BannedDataInteraceEmitterFormatDesc", "DataInterface '{0} - {1}' is banned on currently enabled platforms '{2}'"));
 
 			UObject* WarningObject = nullptr;
 			if (UNiagaraEmitter* NiagaraEmitter = Cast<UNiagaraEmitter>(UsageContext.OwnerObject))
diff --git a/Engine/Plugins/FX/Niagara/Source/NiagaraEditor/Public/NiagaraValidationRules.h b/Engine/Plugins/FX/Niagara/Source/NiagaraEditor/Public/NiagaraValidationRules.h
--- a/Engine/Plugins/FX/Niagara/Source/NiagaraEditor/Public/NiagaraValidationRules.h
+++ b/Engine/Plugins/FX/Niagara/Source/NiagaraEditor/Public/NiagaraValidationRules.h
@@ -12,6 +12,7 @@ class UNiagaraScript;
 
 namespace NiagaraValidation
 {
+	NIAGARAEDITOR_API bool HasValidationRules(UNiagaraSystem* NiagaraSystem);
 	NIAGARAEDITOR_API void ValidateAllRulesInSystem(TSharedPtr<FNiagaraSystemViewModel> ViewModel, TFunction<void(const FNiagaraValidationResult& Result)> ResultCallback);
 }
 
