- Niagara Audit Commandlet can also run validation

#rnx

diff --git a/Engine/Plugins/FX/Niagara/Source/NiagaraEditor/Private/Commandlets/NiagaraSystemAuditCommandlet.cpp b/Engine/Plugins/FX/Niagara/Source/NiagaraEditor/Private/Commandlets/NiagaraSystemAuditCommandlet.cpp
--- a/Engine/Plugins/FX/Niagara/Source/NiagaraEditor/Private/Commandlets/NiagaraSystemAuditCommandlet.cpp
+++ b/Engine/Plugins/FX/Niagara/Source/NiagaraEditor/Private/Commandlets/NiagaraSystemAuditCommandlet.cpp
@@ -15,17 +15,57 @@
 #include "CollectionManagerModule.h"
 #include "FileHelpers.h"
 
+#include "ViewModels/NiagaraSystemViewModel.h"
 #include "NiagaraSettings.h"
 #include "NiagaraSystem.h"
 #include "NiagaraSimulationStageBase.h"
 #include "NiagaraRendererProperties.h"
 #include "NiagaraLightRendererProperties.h"
 #include "NiagaraRibbonRendererProperties.h"
+#include "NiagaraValidationRules.h"
+#include "NiagaraValidationRule.h"
 
 #include UE_INLINE_GENERATED_CPP_BY_NAME(NiagaraSystemAuditCommandlet)
 
 DEFINE_LOG_CATEGORY_STATIC(LogNiagaraSystemAuditCommandlet, Log, All);
 
+namespace NiagaraSystemAuditCommandletLocal
+{
+	void RunValidation(FArchive* Output, UNiagaraSystem* NiagaraSystem)
+	{
+		NiagaraSystem->WaitForCompilationComplete();
+
+		TSharedPtr<FNiagaraSystemViewModel> SystemViewModel = MakeShared<FNiagaraSystemViewModel>();
+		{
+			FNiagaraSystemViewModelOptions SystemOptions;
+			SystemOptions.bCanModifyEmittersFromTimeline = false;
+			SystemOptions.bCanSimulate = false;
+			SystemOptions.bCanAutoCompile = false;
+			SystemOptions.bIsForDataProcessingOnly = true;
+			SystemOptions.MessageLogGuid = NiagaraSystem->GetAssetGuid();
+			SystemViewModel->Initialize(*NiagaraSystem, SystemOptions);
+		}
+
+		NiagaraValidation::ValidateAllRulesInSystem(
+			SystemViewModel,
+			[Output, NiagaraSystem](const FNiagaraValidationResult& Result)
+			{
+				if ( Result.Severity == ENiagaraValidationSeverity::Error || Result.Severity == ENiagaraValidationSeverity::Warning )
+				{
+					Output->Logf(
+						TEXT("%s,%s,%s,%s"),
+						*GetFullNameSafe(NiagaraSystem),
+						(Result.Severity == ENiagaraValidationSeverity::Error) ? TEXT("Error") : TEXT("Warning"),
+						*Result.SummaryText.ToString(),
+						*Result.Description.ToString()
+					);
+				}
+			}
+		);
+	}
+}
+
+
 UNiagaraSystemAuditCommandlet::UNiagaraSystemAuditCommandlet(const FObjectInitializer& ObjectInitializer)
 	: Super(ObjectInitializer)
 {
@@ -91,6 +131,14 @@ int32 UNiagaraSystemAuditCommandlet::Main(const FString& Params)
 
 	FParse::Bool(*Params, TEXT("CaptureDataInterfaceUsage="), bCaptureDataInterfaceUsage);
 
+	// Validation enabled?
+	bool bRunValidation = false;
+	FParse::Bool(*Params, TEXT("RunValidation="), bRunValidation);
+	if ( bRunValidation )
+	{
+		NiagaraValidationIssues.Reset(GetOutputFile(TEXT("NiagaraValidationIssues")));
+	}
+
 	// Package Paths
 	FString PackagePathsString;
 	if (FParse::Value(*Params, TEXT("PackagePaths="), PackagePathsString, false))
@@ -110,12 +158,15 @@ int32 UNiagaraSystemAuditCommandlet::Main(const FString& Params)
 
 	ProcessNiagaraSystems();
 	DumpResults();
+	NiagaraValidationIssues.Reset();
 
 	return 0;
 }
 
 bool UNiagaraSystemAuditCommandlet::ProcessNiagaraSystems()
 {
+	using namespace NiagaraSystemAuditCommandletLocal;
+
 	FAssetRegistryModule& AssetRegistryModule = FModuleManager::LoadModuleChecked<FAssetRegistryModule>(TEXT("AssetRegistry"));
 	IAssetRegistry& AssetRegistry = AssetRegistryModule.Get();
 	AssetRegistry.SearchAllAssets(true);
@@ -143,7 +194,10 @@ bool UNiagaraSystemAuditCommandlet::ProcessNiagaraSystems()
 	const FString DevelopersFolder = FPackageName::FilenameToLongPackageName(FPaths::GameDevelopersDir().LeftChop(1));
 	FString LastPackageName = TEXT("");
 	UPackage* CurrentPackage = nullptr;
-	TArray<UPackage*> PackagesToSave;
+
+	const bool bColledGarbage = NiagaraValidationIssues.IsValid();
+	const int32 CollectGarbageFrequency = 1024;
+	int32 CollectGarabageCounter = 0;
 
 	for (const FAssetData& AssetIt : AssetList)
 	{
@@ -381,6 +435,18 @@ bool UNiagaraSystemAuditCommandlet::ProcessNiagaraSystems()
 			}
 			NiagaraSystemsWithUserDataInterface.Add(FString::Printf(TEXT("%s,%s"), *NiagaraSystem->GetPathName(), *DataInterfaceNames));
 		}
+
+		// Run validation last as it will potentially modify the asset
+		if (NiagaraValidationIssues.IsValid())
+		{
+			RunValidation(NiagaraValidationIssues.Get(), NiagaraSystem);
+		}
+
+		if (bColledGarbage && (++CollectGarabageCounter >= CollectGarbageFrequency) )
+		{
+			CollectGarabageCounter = 0;
+			::CollectGarbage(RF_NoFlags);
+		}
 	}
 
 	// Anything to save do it
@@ -488,7 +554,7 @@ bool UNiagaraSystemAuditCommandlet::DumpSimpleSet(TSet<FString>& InSet, const TC
 FArchive* UNiagaraSystemAuditCommandlet::GetOutputFile(const TCHAR* InShortFilename)
 {
 	const FString Filename = FString::Printf(TEXT("%s/%s.csv"), *AuditOutputFolder, InShortFilename);
-	FArchive* OutputStream = IFileManager::Get().CreateDebugFileWriter(*Filename);
+	FArchive* OutputStream = IFileManager::Get().CreateDebugFileWriter(*Filename, FILEWRITE_AllowRead);
 	if (OutputStream == NULL)
 	{
 		UE_LOG(LogNiagaraSystemAuditCommandlet, Warning, TEXT("Failed to create output stream %s"), *Filename);
diff --git a/Engine/Plugins/FX/Niagara/Source/NiagaraEditor/Private/ViewModels/NiagaraSystemViewModel.cpp b/Engine/Plugins/FX/Niagara/Source/NiagaraEditor/Private/ViewModels/NiagaraSystemViewModel.cpp
--- a/Engine/Plugins/FX/Niagara/Source/NiagaraEditor/Private/ViewModels/NiagaraSystemViewModel.cpp
+++ b/Engine/Plugins/FX/Niagara/Source/NiagaraEditor/Private/ViewModels/NiagaraSystemViewModel.cpp
@@ -162,8 +162,11 @@ void FNiagaraSystemViewModel::Initialize(UNiagaraSystem& InSystem, FNiagaraSyste
 	SystemGraphSelectionViewModel = MakeShared<FNiagaraSystemGraphSelectionViewModel>();
 	SystemGraphSelectionViewModel->Initialize(this->AsShared());
 	
-	SetupPreviewComponentAndInstance();
-	SetupSequencer();
+	if (bIsForDataProcessingOnly == false)
+	{
+		SetupPreviewComponentAndInstance();
+		SetupSequencer();
+	}
 	RefreshAll();
 	AddSystemEventHandlers();
 }
diff --git a/Engine/Plugins/FX/Niagara/Source/NiagaraEditor/Public/Commandlets/NiagaraSystemAuditCommandlet.h b/Engine/Plugins/FX/Niagara/Source/NiagaraEditor/Public/Commandlets/NiagaraSystemAuditCommandlet.h
--- a/Engine/Plugins/FX/Niagara/Source/NiagaraEditor/Public/Commandlets/NiagaraSystemAuditCommandlet.h
+++ b/Engine/Plugins/FX/Niagara/Source/NiagaraEditor/Public/Commandlets/NiagaraSystemAuditCommandlet.h
@@ -29,6 +29,9 @@ class UNiagaraSystemAuditCommandlet : public UCommandlet
 	/** All Niagara Systems & Emitters with simulation stages enabled*/
 	TSet<FString> NiagaraSystemsWithSimulationStages;
 
+	/** All Validation Issues */
+	TUniquePtr<FArchive> NiagaraValidationIssues;
+
 	/** Optional list of data interface usage */
 	struct FDataInterfaceUsage
 	{
@@ -57,6 +60,9 @@ class UNiagaraSystemAuditCommandlet : public UCommandlet
 	/** DeviceProfiles we want to disable GPU emitters for. */
 	TArray<class UDeviceProfile*> DeviceProfilesToDisableGpu;
 
+	UPROPERTY()
+	TArray<TObjectPtr<UPackage>> PackagesToSave;
+
 	/** Entry point */
 	int32 Main(const FString& Params) override;
 
