- Add atmosphere transmittance to cloud occlusion

#rnx

diff --git a/Engine/Plugins/FX/Niagara/Shaders/Private/NiagaraDataInterfaceOcclusion.ush b/Engine/Plugins/FX/Niagara/Shaders/Private/NiagaraDataInterfaceOcclusion.ush
--- a/Engine/Plugins/FX/Niagara/Shaders/Private/NiagaraDataInterfaceOcclusion.ush
+++ b/Engine/Plugins/FX/Niagara/Shaders/Private/NiagaraDataInterfaceOcclusion.ush
@@ -1,5 +1,8 @@
 // Copyright Epic Games, Inc. All Rights Reserved.
 
+#define AP_SUPPORT_SAMPLE_ATMOSHPERE 0
+#include "/Engine/Private/SkyAtmosphereCommon.ush"
+
 float3				{ParameterName}_SystemLWCTile;
 Texture2D<float4>	{ParameterName}_CloudVolumetricTexture;
 SamplerState		{ParameterName}_CloudVolumetricTextureSampler;
@@ -124,7 +127,7 @@ void QueryOcclusionFactorWithCircleGPU_{ParameterName}(in float3 In_SampleCenter
 	Out_SampleFraction = Steps == 0 ? 0 : (TotalSamples / (Rays * Steps + 1));
 }
 
-void QueryCloudOcclusionWithCircle_{ParameterName}(float3 WorldPosition, float WorldDiameter, bool bIncludeCenterSample, int NumberOfRings, int SamplesPerRing, out float VisibilityFraction, out float SampleFraction)
+void QueryCloudOcclusionWithCircle_{ParameterName}(float3 WorldPosition, float WorldDiameter, bool bIncludeCenterSample, int NumberOfRings, int SamplesPerRing, out float VisibilityFraction, out float SampleFraction, out float3 AtmosphereTransmittance)
 {
 	const float PI = 3.14159265;
 	const float SPIRAL_TURN = 2.0f * PI * 0.61803399; // use golden ratio to rotate sample pattern each ring so we get a spiral
@@ -180,4 +183,17 @@ void QueryCloudOcclusionWithCircle_{ParameterName}(float3 WorldPosition, float W
 
 	VisibilityFraction = TotalSamples > 0.0f ? TotalVisibility / TotalSamples : 0.0f;
 	SampleFraction = TotalSamples > 0.0f ? TotalSamples / ExpectedSamples : 0.0f;
+
+	// Sample Atmosphere Transmittance
+	{
+		const float3 RayDirection = normalize(PreViewPosition);
+
+		const float3 PlanetCenterToTranslatedWorldPos = (LWCToFloat(LWCAdd(LWCWorldPosition, PrimaryView.PreViewTranslation)) - PrimaryView.SkyPlanetTranslatedWorldCenterAndViewHeight.xyz) * CM_TO_SKY_UNIT;
+		AtmosphereTransmittance = GetAtmosphereTransmittance(
+			PlanetCenterToTranslatedWorldPos,
+			RayDirection,
+			PrimaryView.SkyAtmosphereBottomRadiusKm, PrimaryView.SkyAtmosphereTopRadiusKm,
+			View.TransmittanceLutTexture, View.TransmittanceLutTextureSampler
+		);
+	}
 }
diff --git a/Engine/Plugins/FX/Niagara/Source/Niagara/Private/NiagaraDataInterfaceOcclusion.cpp b/Engine/Plugins/FX/Niagara/Source/Niagara/Private/NiagaraDataInterfaceOcclusion.cpp
--- a/Engine/Plugins/FX/Niagara/Source/Niagara/Private/NiagaraDataInterfaceOcclusion.cpp
+++ b/Engine/Plugins/FX/Niagara/Source/Niagara/Private/NiagaraDataInterfaceOcclusion.cpp
@@ -29,6 +29,7 @@ struct FNiagaraOcclusionDIFunctionVersion
 	{
 		InitialVersion = 0,
 		LWCConversion = 1,
+		AddAtmosphereTransmittance = 2,
 
 		VersionPlusOne,
 		LatestVersion = VersionPlusOne - 1
@@ -103,6 +104,7 @@ void UNiagaraDataInterfaceOcclusion::GetFunctions(TArray<FNiagaraFunctionSignatu
 		Sig.Inputs.Emplace_GetRef(FNiagaraTypeDefinition::GetIntDef(), TEXT("SamplesPerRing")).SetValue(1);
 		Sig.Outputs.Emplace(FNiagaraTypeDefinition::GetFloatDef(), TEXT("VisibilityFraction"));
 		Sig.Outputs.Emplace(FNiagaraTypeDefinition::GetFloatDef(), TEXT("SampleFraction"));
+		Sig.Outputs.Emplace(FNiagaraTypeDefinition::GetVec3Def(), TEXT("AtmosphereTransmittance"));
 
 		Sig.SetInputDescription(Sig.Inputs[1], CircleCenterPosDescription);
 		Sig.SetInputDescription(Sig.Inputs[2], SampleWindowDiameterDescription);
diff --git a/Engine/Shaders/Private/SkyAtmosphereCommon.ush b/Engine/Shaders/Private/SkyAtmosphereCommon.ush
--- a/Engine/Shaders/Private/SkyAtmosphereCommon.ush
+++ b/Engine/Shaders/Private/SkyAtmosphereCommon.ush
@@ -12,6 +12,11 @@
 // Disabled for all mobile platforms, see USkyLightComponent::IsRealTimeCaptureEnabled
 #define AP_SUPPORT_REALTIME_REFLECTION_CAPTURE (FEATURE_LEVEL > FEATURE_LEVEL_ES3_1)
 
+// Controlls Atmosphere. struct access as not all shaders have access to thi
+#ifndef AP_SUPPORT_SAMPLE_ATMOSHPERE
+	#define AP_SUPPORT_SAMPLE_ATMOSHPERE 1
+#endif
+
 // The constants below should match the one in SceneRendering.cpp
 // Kilometers as unit for computations related to the sky and its atmosphere
 #define CM_TO_SKY_UNIT 0.00001f
@@ -254,6 +259,7 @@ float3 GetLightDiskLuminance(
 	return 0.0f;
 }
 
+#if AP_SUPPORT_SAMPLE_ATMOSHPERE
 ////////////////////////////////////////////////////////////
 // Participating medium properties
 ////////////////////////////////////////////////////////////
@@ -318,3 +324,4 @@ MediumSampleRGB SampleAtmosphereMediumRGB(in float3 WorldPos)
 
 	return s;
 }
+#endif //AP_SUPPORT_SAMPLE_ATMOSHPERE
