- Fix race with AllNiagaraShaderMaps
- Move AllNiagaraShaderMaps behind ALLOW_SHADERMAP_DEBUG_DATA

#rnx

diff --git a/Engine/Plugins/FX/Niagara/Source/NiagaraShader/Private/NiagaraShader.cpp b/Engine/Plugins/FX/Niagara/Source/NiagaraShader/Private/NiagaraShader.cpp
--- a/Engine/Plugins/FX/Niagara/Source/NiagaraShader/Private/NiagaraShader.cpp
+++ b/Engine/Plugins/FX/Niagara/Source/NiagaraShader/Private/NiagaraShader.cpp
@@ -72,7 +72,10 @@ namespace NiagaraShaderCookStats
 //
 FCriticalSection GIdToNiagaraShaderMapCS;
 TMap<FNiagaraShaderMapId, FNiagaraShaderMap*> FNiagaraShaderMap::GIdToNiagaraShaderMap[SP_NumPlatforms];
+#if ALLOW_SHADERMAP_DEBUG_DATA
 TArray<FNiagaraShaderMap*> FNiagaraShaderMap::AllNiagaraShaderMaps;
+FCriticalSection FNiagaraShaderMap::AllNiagaraShaderMapsGuard;
+#endif
 
 #if WITH_EDITOR
 TMap<TRefCountPtr<FNiagaraShaderMap>, TArray<FNiagaraShaderScript*>> FNiagaraShaderMap::NiagaraShaderMapsBeingCompiled;
@@ -950,7 +953,10 @@ FNiagaraShaderMap::FNiagaraShaderMap() :
 	bIsPersistent(true) 
 {
 	checkSlow(IsInGameThread() || IsAsyncLoading());
+#if ALLOW_SHADERMAP_DEBUG_DATA
+	FScopeLock AllSMAccess(&AllNiagaraShaderMapsGuard);
 	AllNiagaraShaderMaps.Add(this);
+#endif
 }
 
 FNiagaraShaderMap::~FNiagaraShaderMap()
@@ -958,7 +964,10 @@ FNiagaraShaderMap::~FNiagaraShaderMap()
 	checkSlow(IsInGameThread() || IsAsyncLoading());
 	check(bDeletedThroughDeferredCleanup);
 	check(!bRegistered);
+#if ALLOW_SHADERMAP_DEBUG_DATA
+	FScopeLock AllSMAccess(&AllNiagaraShaderMapsGuard);
 	AllNiagaraShaderMaps.RemoveSwap(this);
+#endif
 }
 
 bool FNiagaraShaderMap::Serialize(FArchive& Ar, bool bInlineShaderResources, bool bLoadedByCookedMaterial)
diff --git a/Engine/Plugins/FX/Niagara/Source/NiagaraShader/Public/NiagaraShared.h b/Engine/Plugins/FX/Niagara/Source/NiagaraShader/Public/NiagaraShared.h
--- a/Engine/Plugins/FX/Niagara/Source/NiagaraShader/Public/NiagaraShared.h
+++ b/Engine/Plugins/FX/Niagara/Source/NiagaraShader/Public/NiagaraShared.h
@@ -12,6 +12,7 @@
 #include "Engine/EngineTypes.h"
 #include "Templates/RefCounting.h"
 #include "Misc/SecureHash.h"
+#include "MaterialShared.h"
 #include "RHI.h"
 #include "RenderResource.h"
 #include "RenderingThread.h"
@@ -631,12 +632,17 @@ private:
 	*/
 	static TMap<FNiagaraShaderMapId, FNiagaraShaderMap*> GIdToNiagaraShaderMap[SP_NumPlatforms];
 
+#if ALLOW_SHADERMAP_DEBUG_DATA
 	/**
 	* All script shader maps in memory.
 	* No ref counting needed as these are removed on destruction of the shader map.
 	*/
 	static TArray<FNiagaraShaderMap*> AllNiagaraShaderMaps;
 
+	/** Guards access to AllNiagaraShaderMaps, which can be written to from an async loading thread. */
+	static FCriticalSection AllNiagaraShaderMapsGuard;
+#endif
+
 #if WITH_EDITOR
 	/** Tracks resources and their shader maps that need to be compiled but whose compilation is being deferred. */
 	static TMap<TRefCountPtr<FNiagaraShaderMap>, TArray<FNiagaraShaderScript*> > NiagaraShaderMapsBeingCompiled;
