- Niagara ribbon optiomizations to reduce allocations
- Test asset wost case performance 369us -> 307us, average case ~30us saving

#rnx

diff --git a/Engine/Plugins/FX/Niagara/Source/Niagara/Private/NiagaraRendererRibbons.cpp b/Engine/Plugins/FX/Niagara/Source/Niagara/Private/NiagaraRendererRibbons.cpp
--- a/Engine/Plugins/FX/Niagara/Source/Niagara/Private/NiagaraRendererRibbons.cpp
+++ b/Engine/Plugins/FX/Niagara/Source/Niagara/Private/NiagaraRendererRibbons.cpp
@@ -1139,6 +1139,7 @@ void FNiagaraRendererRibbons::GenerateShapeStateTube(FNiagaraRibbonShapeGeometry
 	State.BitsNeededForShape = CalculateBitsForRange(State.NumVerticesInSlice);
 	State.BitMaskForShape = CalculateBitMask(State.BitsNeededForShape);
 	
+	State.SliceVertexData.Reserve(TubeSubdivisions + 1);
 	for (int32 VertexId = 0; VertexId <= TubeSubdivisions; VertexId++)
 	{
 		const float RotationAngle = (static_cast<float>(VertexId) / TubeSubdivisions) * -360.0f;
@@ -1149,7 +1150,7 @@ void FNiagaraRendererRibbons::GenerateShapeStateTube(FNiagaraRibbonShapeGeometry
 		State.SliceVertexData.Emplace(Position, Normal, TextureV);
 	}
 	
-	State.SliceTriangleToVertexIds.Reserve(TubeSubdivisions);
+	State.SliceTriangleToVertexIds.Reserve(TubeSubdivisions * 2);
 	for (int32 VertexIdx = 0; VertexIdx < TubeSubdivisions; VertexIdx++)
 	{
 		State.SliceTriangleToVertexIds.Add(VertexIdx);
@@ -1188,7 +1189,7 @@ void FNiagaraRendererRibbons::GenerateShapeStateCustom(FNiagaraRibbonShapeGeomet
 		State.SliceVertexData.Emplace(Position, Normal, TextureV);
 	}
 
-	State.SliceTriangleToVertexIds.Reserve(CustomVertices.Num());
+	State.SliceTriangleToVertexIds.Reserve(CustomVertices.Num() * 2);
 	for (int32 VertexIdx = 0; VertexIdx < CustomVertices.Num(); VertexIdx++)
 	{
 		State.SliceTriangleToVertexIds.Add(VertexIdx);
@@ -1206,6 +1207,7 @@ void FNiagaraRendererRibbons::GenerateShapeStatePlane(FNiagaraRibbonShapeGeometr
 	State.BitsNeededForShape = CalculateBitsForRange(State.NumVerticesInSlice);
 	State.BitMaskForShape = CalculateBitMask(State.BitsNeededForShape);
 	
+	State.SliceVertexData.Reserve(WidthSegmentationCount + 1);
 	for (int32 VertexId = 0; VertexId <= WidthSegmentationCount; VertexId++)
 	{
 		const FVector2f Position = FVector2f((static_cast<float>(VertexId) / WidthSegmentationCount) - 0.5f, 0);
@@ -1215,7 +1217,7 @@ void FNiagaraRendererRibbons::GenerateShapeStatePlane(FNiagaraRibbonShapeGeometr
 		State.SliceVertexData.Emplace(Position, Normal, TextureV);
 	}
 	
-	State.SliceTriangleToVertexIds.Reserve(WidthSegmentationCount);
+	State.SliceTriangleToVertexIds.Reserve(WidthSegmentationCount * 2);
 	for (int32 VertexIdx = 0; VertexIdx < WidthSegmentationCount; VertexIdx++)
 	{
 		State.SliceTriangleToVertexIds.Add(VertexIdx);
@@ -1351,10 +1353,9 @@ void FNiagaraRendererRibbons::CalculateUVScaleAndOffsets(const FNiagaraRibbonUVS
 
 
 template<bool bWantsTessellation, bool bHasTwist, bool bWantsMultiRibbon>
-void FNiagaraRendererRibbons::GenerateVertexBufferForRibbonPart(const FNiagaraGenerationInputDataCPUAccessors& CPUData, const TArray<uint32>& RibbonIndices, uint32 RibbonIndex, FNiagaraRibbonCPUGeneratedVertexData& OutputData) const
+void FNiagaraRendererRibbons::GenerateVertexBufferForRibbonPart(const FNiagaraGenerationInputDataCPUAccessors& CPUData, TConstArrayView<uint32> RibbonIndices, uint32 RibbonIndex, FNiagaraRibbonCPUGeneratedVertexData& OutputData) const
 {
 	TArray<uint32>& SegmentData = OutputData.SegmentData;
-	TArray<FRibbonMultiRibbonInfo>& MultiRibbonInfos = OutputData.RibbonInfoLookup;
 	
 	const FNiagaraDataSetReaderFloat<FNiagaraPosition>& PosData = CPUData.PosData;	
 	const FNiagaraDataSetReaderFloat<float>& AgeData = CPUData.AgeData;
@@ -1371,9 +1372,6 @@ void FNiagaraRendererRibbons::GenerateVertexBufferForRibbonPart(const FNiagaraGe
 	float LastWidth = 0;
 	double TotalDistance = 0.0f;
 
-	OutputData.SortedIndices.Reserve(OutputData.SortedIndices.Num() + RibbonIndices.Num());
-	OutputData.TangentAndDistances.Reserve(OutputData.TangentAndDistances.Num() + RibbonIndices.Num());
-
 	// Find the first position with enough distance.
 	int32 CurrentIndex = 1;
 	while (CurrentIndex < RibbonIndices.Num())
@@ -1493,7 +1491,7 @@ void FNiagaraRendererRibbons::GenerateVertexBufferForRibbonPart(const FNiagaraGe
 
 	if (NumSegments > 0)
 	{
-		FRibbonMultiRibbonInfo& MultiRibbonInfo = MultiRibbonInfos[RibbonIndex];
+		FRibbonMultiRibbonInfo& MultiRibbonInfo = OutputData.RibbonInfoLookup[RibbonIndex];
 		MultiRibbonInfo.StartPos = (FVector)PosData[RibbonIndices[0]];
 		MultiRibbonInfo.EndPos = (FVector)PosData[RibbonIndices.Last()];
 		MultiRibbonInfo.BaseSegmentDataIndex = SegmentData.Num();
@@ -1569,23 +1567,29 @@ void FNiagaraRendererRibbons::GenerateVertexBufferForRibbonPart(const FNiagaraGe
 template<typename IDType, typename ReaderType, bool bWantsTessellation, bool bHasTwist>
 void FNiagaraRendererRibbons::GenerateVertexBufferForMultiRibbonInternal(const FNiagaraGenerationInputDataCPUAccessors& CPUData, const ReaderType& IDReader, FNiagaraRibbonCPUGeneratedVertexData& OutputData) const
 {
-	//-OPT: Consider MemStack
-	TMap<IDType, TArray<uint32>> MultiRibbonSortedIndices;
+	using FIndexArray = TArray<uint32, TMemStackAllocator<>>;
+	TMap<IDType, FIndexArray> MultiRibbonSortedIndices;
 
-	for (uint32 i = 0; i < CPUData.TotalNumParticles; ++i)
+	for (uint32 i=0; i < CPUData.TotalNumParticles; ++i)
 	{
-		TArray<uint32>& Indices = MultiRibbonSortedIndices.FindOrAdd(IDReader[i]);
+		FIndexArray& Indices = MultiRibbonSortedIndices.FindOrAdd(IDReader[i]);
 		Indices.Add(i);
 	}
+
+	const int32 NumRibbons = MultiRibbonSortedIndices.Num();
 	OutputData.RibbonInfoLookup.AddZeroed(MultiRibbonSortedIndices.Num());
+	OutputData.SortedIndices.Reserve(OutputData.SortedIndices.Num() + CPUData.TotalNumParticles + NumRibbons);
+	OutputData.TangentAndDistances.Reserve(OutputData.TangentAndDistances.Num() + CPUData.TotalNumParticles + NumRibbons);
+	OutputData.MultiRibbonIndices.Reserve(OutputData.MultiRibbonIndices.Num() + CPUData.TotalNumParticles + NumRibbons);
+	OutputData.SegmentData.Reserve(OutputData.SegmentData.Num() + CPUData.TotalNumParticles + NumRibbons);
 
 	// Sort the ribbons by ID so that the draw order stays consistent.
 	MultiRibbonSortedIndices.KeySort(TLess<IDType>());
 
 	uint32 RibbonIndex = 0;
-	for (TPair<IDType, TArray<uint32>>& Pair : MultiRibbonSortedIndices)
+	for (TPair<IDType, FIndexArray>& Pair : MultiRibbonSortedIndices)
 	{
-		TArray<uint32>& SortedIndices = Pair.Value;
+		FIndexArray& SortedIndices = Pair.Value;
 		const auto& SortKeyReader = CPUData.SortKeyReader;
 		SortedIndices.Sort([&SortKeyReader](const uint32& A, const uint32& B) { return (SortKeyReader[A] < SortKeyReader[B]); });
 		GenerateVertexBufferForRibbonPart<bWantsTessellation, bHasTwist, true>(CPUData, SortedIndices, RibbonIndex, OutputData);
@@ -1620,6 +1624,7 @@ void FNiagaraRendererRibbons::GenerateVertexBufferCPU(const FNiagaraGenerationIn
 	check(CPUData.PosData.IsValid() && CPUData.SortKeyReader.IsValid());
 
 	// TODO: Move sorting to share code with sprite and mesh sorting and support the custom sorting key.
+	FMemMark Mark(FMemStack::Get());
 	if (GenerationConfig.HasRibbonIDs())
 	{
 		if (GenerationConfig.HasFullRibbonIDs())
@@ -1636,8 +1641,7 @@ void FNiagaraRendererRibbons::GenerateVertexBufferCPU(const FNiagaraGenerationIn
 	}
 	else
 	{
-		//-OPT: Consider MemStack
-		TArray<uint32> SortedIndices;
+		TArray<uint32, TMemStackAllocator<>> SortedIndices;
 		SortedIndices.Reserve(CPUData.TotalNumParticles + 1);
 		for (uint32 i = 0; i < CPUData.TotalNumParticles; ++i)
 		{
@@ -1647,7 +1651,11 @@ void FNiagaraRendererRibbons::GenerateVertexBufferCPU(const FNiagaraGenerationIn
 
 		const auto& SortKeyReader = CPUData.SortKeyReader;
 		SortedIndices.Sort([&SortKeyReader](const uint32& A, const uint32& B) {	return (SortKeyReader[A] < SortKeyReader[B]); });
-		
+
+		OutputData.SortedIndices.Reserve(OutputData.SortedIndices.Num() + CPUData.TotalNumParticles + 1);
+		OutputData.TangentAndDistances.Reserve(OutputData.TangentAndDistances.Num() + CPUData.TotalNumParticles + 1);
+		OutputData.SegmentData.Reserve(OutputData.SegmentData.Num() + CPUData.TotalNumParticles + 1);
+
 		if (GenerationConfig.WantsAutomaticTessellation())
 		{
 			if (GenerationConfig.HasTwist())
diff --git a/Engine/Plugins/FX/Niagara/Source/Niagara/Private/NiagaraScriptExecutionContext.cpp b/Engine/Plugins/FX/Niagara/Source/Niagara/Private/NiagaraScriptExecutionContext.cpp
--- a/Engine/Plugins/FX/Niagara/Source/Niagara/Private/NiagaraScriptExecutionContext.cpp
+++ b/Engine/Plugins/FX/Niagara/Source/Niagara/Private/NiagaraScriptExecutionContext.cpp
@@ -141,8 +141,9 @@ void FNiagaraScriptExecutionContextBase::BindData(int32 Index, FNiagaraDataBuffe
 #if STATS
 void FNiagaraScriptExecutionContextBase::CreateStatScopeData()
 {
-	StatScopeData.Empty();
-	for (const TStatId& StatId : Script->GetStatScopeIDs())
+	TArrayView<const TStatId> StatScopeIDs = Script->GetStatScopeIDs();
+	StatScopeData.Empty(StatScopeIDs.Num());
+	for (const TStatId& StatId : StatScopeIDs)
 	{
 		StatScopeData.Add(FStatScopeData(StatId));
 	}
diff --git a/Engine/Plugins/FX/Niagara/Source/Niagara/Public/NiagaraRendererRibbons.h b/Engine/Plugins/FX/Niagara/Source/Niagara/Public/NiagaraRendererRibbons.h
--- a/Engine/Plugins/FX/Niagara/Source/Niagara/Public/NiagaraRendererRibbons.h
+++ b/Engine/Plugins/FX/Niagara/Source/Niagara/Public/NiagaraRendererRibbons.h
@@ -192,7 +192,7 @@ struct FNiagaraRibbonCPUGeneratedVertexData
 	TArray<uint32> MultiRibbonIndices;
 	
 	/** Ribbon perperties required for sorting. */
-	TArray<FRibbonMultiRibbonInfo> RibbonInfoLookup;
+	TArray<FRibbonMultiRibbonInfo, TInlineAllocator<1>> RibbonInfoLookup;
 
 	double TotalSegmentLength;
 	double AverageSegmentLength;
@@ -371,7 +371,7 @@ protected:
 		int32 StartIndex, int32 EndIndex, int32 NumSegments, float TotalLength, float& OutUScale, float& OutUOffset, float& OutUDistributionScaler);
 
 	template<bool bWantsTessellation, bool bHasTwist, bool bWantsMultiRibbon>
-	void GenerateVertexBufferForRibbonPart(const FNiagaraGenerationInputDataCPUAccessors& CPUData, const TArray<uint32>& RibbonIndices, uint32 RibbonIndex, FNiagaraRibbonCPUGeneratedVertexData& OutputData) const;
+	void GenerateVertexBufferForRibbonPart(const FNiagaraGenerationInputDataCPUAccessors& CPUData, TConstArrayView<uint32> RibbonIndices, uint32 RibbonIndex, FNiagaraRibbonCPUGeneratedVertexData& OutputData) const;
 
 	template<typename IDType, typename ReaderType, bool bWantsTessellation, bool bHasTwist>
 	void GenerateVertexBufferForMultiRibbonInternal(const FNiagaraGenerationInputDataCPUAccessors& CPUData, const ReaderType& IDReader, FNiagaraRibbonCPUGeneratedVertexData& OutputData) const;
