- [Niagara] Exporting some types/functions (Contributed by [REDACTED_PERSON])

#9695

diff --git a/Engine/Plugins/FX/Niagara/Source/Niagara/Classes/NiagaraEmitterInstance.h b/Engine/Plugins/FX/Niagara/Source/Niagara/Classes/NiagaraEmitterInstance.h
--- a/Engine/Plugins/FX/Niagara/Source/Niagara/Classes/NiagaraEmitterInstance.h
+++ b/Engine/Plugins/FX/Niagara/Source/Niagara/Classes/NiagaraEmitterInstance.h
@@ -82,7 +82,7 @@ public:
 	NIAGARA_API void CalculateFixedBounds(const FTransform& ToWorldSpace);
 #endif
 	
-	bool GetBoundRendererValue_GT(const FNiagaraVariableBase& InBaseVar, const FNiagaraVariableBase& InSubVar, void* OutValueData) const;
+	NIAGARA_API bool GetBoundRendererValue_GT(const FNiagaraVariableBase& InBaseVar, const FNiagaraVariableBase& InSubVar, void* OutValueData) const;
 
 	FNiagaraDataSet& GetData()const { return *ParticleDataSet; }
 
@@ -155,8 +155,8 @@ public:
 	}
 	FBox GetFixedBounds() const;
 
-	UObject* FindBinding(const FNiagaraVariable& InVariable) const;
-	UNiagaraDataInterface* FindDataInterface(const FNiagaraVariable& InVariable) const;
+	NIAGARA_API UObject* FindBinding(const FNiagaraVariable& InVariable) const;
+	NIAGARA_API UNiagaraDataInterface* FindDataInterface(const FNiagaraVariable& InVariable) const;
 
 	bool HasTicked() const { return TickCount > 0;  }
 
diff --git a/Engine/Plugins/FX/Niagara/Source/Niagara/Public/NiagaraCommon.h b/Engine/Plugins/FX/Niagara/Source/Niagara/Public/NiagaraCommon.h
--- a/Engine/Plugins/FX/Niagara/Source/Niagara/Public/NiagaraCommon.h
+++ b/Engine/Plugins/FX/Niagara/Source/Niagara/Public/NiagaraCommon.h
@@ -1077,11 +1077,11 @@ struct FNiagaraMaterialAttributeBinding
 	UPROPERTY(EditAnywhere, Category = "Variable")
 	FNiagaraVariableBase NiagaraChildVariable;
 
-	void NIAGARA_API CacheValues(const UNiagaraEmitter* InEmitter); 
-	const FNiagaraVariableBase& GetParamMapBindableVariable() const;
+	NIAGARA_API void CacheValues(const UNiagaraEmitter* InEmitter);
+	NIAGARA_API const FNiagaraVariableBase& GetParamMapBindableVariable() const;
 
-	bool NIAGARA_API RenameVariableIfMatching(const FNiagaraVariableBase& OldVariable, const FNiagaraVariableBase& NewVariable, const UNiagaraEmitter* InEmitter, ENiagaraRendererSourceDataMode InSourceMode);
-	bool NIAGARA_API Matches(const FNiagaraVariableBase& OldVariable, const UNiagaraEmitter* InEmitter, ENiagaraRendererSourceDataMode InSourceMode);
+	NIAGARA_API bool RenameVariableIfMatching(const FNiagaraVariableBase& OldVariable, const FNiagaraVariableBase& NewVariable, const UNiagaraEmitter* InEmitter, ENiagaraRendererSourceDataMode InSourceMode);
+	NIAGARA_API bool Matches(const FNiagaraVariableBase& OldVariable, const UNiagaraEmitter* InEmitter, ENiagaraRendererSourceDataMode InSourceMode);
 };
 
 USTRUCT()
diff --git a/Engine/Plugins/FX/Niagara/Source/Niagara/Public/NiagaraRenderer.h b/Engine/Plugins/FX/Niagara/Source/Niagara/Public/NiagaraRenderer.h
--- a/Engine/Plugins/FX/Niagara/Source/Niagara/Public/NiagaraRenderer.h
+++ b/Engine/Plugins/FX/Niagara/Source/Niagara/Public/NiagaraRenderer.h
@@ -24,7 +24,7 @@ class FNiagaraGPURendererCount;
 class FNiagaraSystemInstanceController;
 
 /** Struct used to pass dynamic data from game thread to render thread */
-struct FNiagaraDynamicDataBase
+struct NIAGARA_API FNiagaraDynamicDataBase
 {
 	explicit FNiagaraDynamicDataBase(const FNiagaraEmitterInstance* InEmitter);
 	virtual ~FNiagaraDynamicDataBase();
