- [Niagara] SV_* params passthrough for custom nodes. (Contributed by [REDACTED_PERSON])

#9694

diff --git a/Engine/Plugins/FX/Niagara/Shaders/Private/NiagaraEmitterInstanceShader.usf b/Engine/Plugins/FX/Niagara/Shaders/Private/NiagaraEmitterInstanceShader.usf
--- a/Engine/Plugins/FX/Niagara/Shaders/Private/NiagaraEmitterInstanceShader.usf
+++ b/Engine/Plugins/FX/Niagara/Shaders/Private/NiagaraEmitterInstanceShader.usf
@@ -82,12 +82,14 @@ bool NiagaraAny(bool4 b) { return b.x || b.y || b.z || b.w; }
 	uint3 DispatchThreadIdBounds;		// Used to track thread bounds
 	uint3 DispatchThreadIdToLinear;		// Used to convert from DispatchThreadId to a linear thread id
 
+	static uint3 GDispatchThreadId;		// SV_DispatchThreadId
+	static uint3 GGroupId;				// SV_GroupdId
+	static uint3 GGroupThreadId;		// SV_GroupThreadId
+	static uint  GGroupIndex;			// SV_GroupIndex
 	static uint GLinearThreadId;
-	static uint3 GDispatchThreadId;
-	static uint3 GGroupThreadId;
 	static uint GEmitterTickCounter;
 	static uint GRandomSeedOffset = 0;
-	
+
 	// To be removed as well but don't know who is using it
 	#include "/Engine/Private/SceneTexturesCommon.ush"
 
diff --git a/Engine/Plugins/FX/Niagara/Source/NiagaraEditor/Private/NiagaraHlslTranslator.cpp b/Engine/Plugins/FX/Niagara/Source/NiagaraEditor/Private/NiagaraHlslTranslator.cpp
--- a/Engine/Plugins/FX/Niagara/Source/NiagaraEditor/Private/NiagaraHlslTranslator.cpp
+++ b/Engine/Plugins/FX/Niagara/Source/NiagaraEditor/Private/NiagaraHlslTranslator.cpp
@@ -3323,17 +3323,21 @@ void FHlslNiagaraTranslator::DefineMainGPUFunctions(
 			"\n"
 			"[numthreads(THREADGROUP_SIZE_X, THREADGROUP_SIZE_Y, THREADGROUP_SIZE_Z)]\n"
 			"void SimulateMainComputeCS(\n"
-			"	uint3 DispatchThreadId : SV_DispatchThreadID,\n"
-			"	uint3 GroupThreadId : SV_GroupThreadID)\n"
+			"	uint3 InDispatchThreadId	: SV_DispatchThreadID,\n"
+			"	uint3 InGroupId				: SV_GroupID,\n"
+			"	uint3 InGroupThreadId		: SV_GroupThreadID,\n"
+			"	uint  InGroupIndex			: SV_GroupIndex)\n"
 			"{\n"
-			"	GLinearThreadId = DispatchThreadId.x + (DispatchThreadId.y * DispatchThreadIdToLinear.y);\n"
+			"	GDispatchThreadId	= InDispatchThreadId;\n"
+			"	GGroupId			= InGroupId;\n"
+			"	GGroupThreadId		= InGroupThreadId;\n"
+			"	GGroupIndex			= InGroupIndex;\n"
+			"	GLinearThreadId		= GDispatchThreadId.x + (GDispatchThreadId.y * DispatchThreadIdToLinear.y);\n"
 			"	#if NIAGARA_DISPATCH_TYPE >= NIAGARA_DISPATCH_TYPE_THREE_D\n"
-			"		GLinearThreadId += DispatchThreadId.z * DispatchThreadIdToLinear.z;\n"
+			"		GLinearThreadId += GDispatchThreadId.z * DispatchThreadIdToLinear.z;\n"
 			"	#endif\n"
-			"	GDispatchThreadId = DispatchThreadId;\n"
-			"	GGroupThreadId = GroupThreadId;\n"
-			"	GEmitterTickCounter = EmitterTickCounter;\n"
-			"	GRandomSeedOffset = 0;\n"
+			"	GEmitterTickCounter	= EmitterTickCounter;\n"
+			"	GRandomSeedOffset	= 0;\n"
 			"\n"
 		);
 
@@ -3423,7 +3427,7 @@ void FHlslNiagaraTranslator::DefineMainGPUFunctions(
 			//		 can really be variable, everything else is driven from CPU code.
 			HlslOutput += TEXT(
 				"	const uint MaxInstances = SimulationStage_GetInstanceCount();\n"
-				"	GLinearThreadId = all(DispatchThreadId < DispatchThreadIdBounds) ? GLinearThreadId : MaxInstances;\n"
+				"	GLinearThreadId = all(GDispatchThreadId < DispatchThreadIdBounds) ? GLinearThreadId : MaxInstances;\n"
 				"	GSpawnStartInstance = MaxInstances;\n"
 				"	const bool bRunUpdateLogic = (GLinearThreadId < GSpawnStartInstance) && (SimStart != 1);\n"
 				"	const bool bRunSpawnLogic = (GLinearThreadId < GSpawnStartInstance) && (SimStart == 1);\n"
