Fixes for non-null-terminated strings getting into the reporting of cache compile ids

diff --git a/Engine/Plugins/FX/Niagara/Source/NiagaraEditor/Private/NiagaraEditorUtilities.cpp b/Engine/Plugins/FX/Niagara/Source/NiagaraEditor/Private/NiagaraEditorUtilities.cpp
--- a/Engine/Plugins/FX/Niagara/Source/NiagaraEditor/Private/NiagaraEditorUtilities.cpp
+++ b/Engine/Plugins/FX/Niagara/Source/NiagaraEditor/Private/NiagaraEditorUtilities.cpp
@@ -482,7 +482,7 @@ bool FNiagaraEditorUtilities::NestedPropertiesAppendCompileHash(const void* Cont
 			Property->GetFName().AppendString(PropertyName);
 		}
 
-		if (PODPropertyAppendCompileHash(Container, Property, PropertyName, InVisitor))
+		if (PODPropertyAppendCompileHash(Container, Property, *PropertyName, InVisitor))
 		{
 			continue;
 		}
@@ -490,14 +490,14 @@ bool FNiagaraEditorUtilities::NestedPropertiesAppendCompileHash(const void* Cont
 		{
 			FStructProperty* StructProp = CastFieldChecked<FStructProperty>(Property);
 			const void* StructContainer = Property->ContainerPtrToValuePtr<uint8>(Container);
-			NestedPropertiesAppendCompileHash(StructContainer, StructProp->Struct, EFieldIteratorFlags::IncludeSuper, PropertyName, InVisitor);
+			NestedPropertiesAppendCompileHash(StructContainer, StructProp->Struct, EFieldIteratorFlags::IncludeSuper, *PropertyName, InVisitor);
 			continue;
 		}
 		else if (Property->IsA(FEnumProperty::StaticClass()))
 		{
 			FEnumProperty* CastProp = CastFieldChecked<FEnumProperty>(Property);
 			const void* EnumContainer = Property->ContainerPtrToValuePtr<uint8>(Container);
-			if (PODPropertyAppendCompileHash(EnumContainer, CastProp->GetUnderlyingProperty(), PropertyName, InVisitor))
+			if (PODPropertyAppendCompileHash(EnumContainer, CastProp->GetUnderlyingProperty(), *PropertyName, InVisitor))
 			{
 				continue;
 			}
@@ -513,11 +513,11 @@ bool FNiagaraEditorUtilities::NestedPropertiesAppendCompileHash(const void* Cont
 				// We just do name here as sometimes things will be in a transient package or something tricky.
 				// Because we do nested id's for each called graph, it should work out in the end to have a different
 				// value in the compile array if the scripts are the same name but different locations.
-				InVisitor->UpdateString(PropertyName.GetData(), Obj->GetName());
+				InVisitor->UpdateString(*PropertyName, Obj->GetName());
 			}
 			else
 			{
-				InVisitor->UpdateString(PropertyName.GetData(), TEXT("nullptr"));
+				InVisitor->UpdateString(*PropertyName, TEXT("nullptr"));
 			}
 			continue;
 		}
@@ -525,7 +525,7 @@ bool FNiagaraEditorUtilities::NestedPropertiesAppendCompileHash(const void* Cont
 		{
 			FMapProperty* CastProp = CastFieldChecked<FMapProperty>(Property);
 			FScriptMapHelper MapHelper(CastProp, CastProp->ContainerPtrToValuePtr<void>(Container));
-			InVisitor->UpdatePOD(PropertyName.GetData(), MapHelper.Num());
+			InVisitor->UpdatePOD(*PropertyName, MapHelper.Num());
 			if (MapHelper.GetKeyProperty())
 			{
 				PathName.Reset();
@@ -610,7 +610,7 @@ bool FNiagaraEditorUtilities::NestedPropertiesAppendCompileHash(const void* Cont
 			FArrayProperty* CastProp = CastFieldChecked<FArrayProperty>(Property);
 
 			FScriptArrayHelper ArrayHelper(CastProp, CastProp->ContainerPtrToValuePtr<void>(Container));
-			InVisitor->UpdatePOD(PropertyName.GetData(), ArrayHelper.Num());
+			InVisitor->UpdatePOD(*PropertyName, ArrayHelper.Num());
 			PathName.Reset();
 			CastProp->Inner->GetPathName(nullptr, PathName);
 			InVisitor->UpdateString(TEXT("InnerPathname"), PathName);
@@ -623,7 +623,7 @@ bool FNiagaraEditorUtilities::NestedPropertiesAppendCompileHash(const void* Cont
 
 				for (int32 ArrayIdx = 0; ArrayIdx < ArrayHelper.Num(); ArrayIdx++)
 				{
-					if (!NestedPropertiesAppendCompileHash(ArrayHelper.GetRawPtr(ArrayIdx), StructProp->Struct, EFieldIteratorFlags::IncludeSuper, PropertyName, InVisitor))
+					if (!NestedPropertiesAppendCompileHash(ArrayHelper.GetRawPtr(ArrayIdx), StructProp->Struct, EFieldIteratorFlags::IncludeSuper, *PropertyName, InVisitor))
 					{
 						UE_LOG(LogNiagaraEditor, Warning, TEXT("Skipping %s because it is an array property of unsupported underlying type, please add \"meta = (SkipForCompileHash=\"true\")\" to avoid this warning in the future or handle it yourself in NestedPropertiesAppendCompileHash!"), *Property->GetName());
 						bPassed = false;
@@ -640,7 +640,7 @@ bool FNiagaraEditorUtilities::NestedPropertiesAppendCompileHash(const void* Cont
 				bool bPassed = true;
 				for (int32 ArrayIdx = 0; ArrayIdx < ArrayHelper.Num(); ArrayIdx++)
 				{
-					if (!PODPropertyAppendCompileHash(ArrayHelper.GetRawPtr(ArrayIdx), CastProp->Inner, PropertyName, InVisitor))
+					if (!PODPropertyAppendCompileHash(ArrayHelper.GetRawPtr(ArrayIdx), CastProp->Inner, *PropertyName, InVisitor))
 					{
 						if (bPassed)
 						{
