Remove some debug code from investigation into niagara static switches

diff --git a/Engine/Plugins/FX/Niagara/Source/NiagaraEditor/Private/NiagaraNodeStaticSwitch.cpp b/Engine/Plugins/FX/Niagara/Source/NiagaraEditor/Private/NiagaraNodeStaticSwitch.cpp
--- a/Engine/Plugins/FX/Niagara/Source/NiagaraEditor/Private/NiagaraNodeStaticSwitch.cpp
+++ b/Engine/Plugins/FX/Niagara/Source/NiagaraEditor/Private/NiagaraNodeStaticSwitch.cpp
@@ -725,12 +725,6 @@ void UNiagaraNodeStaticSwitch::PostLoad()
 	// Make sure that we are added to the static switch list.
 	if (GetInputType().IsValid() && InputParameterName.IsValid())
 	{
-		if (InputParameterName.ToString() == TEXT("Enable Flattened Endcaps"))
-		{
-			int asdf = 0;
-			++asdf;
-		}
-
 		if (UNiagaraGraph* NiagaraGraph = GetNiagaraGraph())
 		{
 			const FNiagaraVariable Variable(GetInputType(), InputParameterName);
