- Fix Niagara debug line drawing when TSR is enabled

#rnx

diff --git a/Engine/Plugins/FX/Niagara/Source/NiagaraShader/Private/NiagaraDebugShaders.cpp b/Engine/Plugins/FX/Niagara/Source/NiagaraShader/Private/NiagaraDebugShaders.cpp
--- a/Engine/Plugins/FX/Niagara/Source/NiagaraShader/Private/NiagaraDebugShaders.cpp
+++ b/Engine/Plugins/FX/Niagara/Source/NiagaraShader/Private/NiagaraDebugShaders.cpp
@@ -197,7 +197,7 @@ void NiagaraDebugShaders::ClearUAV(FRDGBuilder& GraphBuilder, FRDGBufferUAVRef U
 }
 
 void NiagaraDebugShaders::DrawDebugLines(
-	class FRDGBuilder& GraphBuilder, const class FSceneView& View, FRDGTextureRef SceneColor, FRDGTextureRef SceneDepth,
+	class FRDGBuilder& GraphBuilder, const class FViewInfo& View, FRDGTextureRef SceneColor, FRDGTextureRef SceneDepth,
 	const uint32 LineInstanceCount, FRDGBufferRef LineBuffer
 )
 {
@@ -209,8 +209,8 @@ void NiagaraDebugShaders::DrawDebugLines(
 	PassParameters->VSParameters.GpuLineBuffer			= GraphBuilder.CreateSRV(LineBuffer, PF_R32_UINT);
 
 	PassParameters->PSParameters.OutputInvResolution	= FVector2f(1.0f / View.UnscaledViewRect.Width(), 1.0f / View.UnconstrainedViewRect.Height());
-	PassParameters->PSParameters.OriginalViewRectMin	= FVector2f(View.UnscaledViewRect.Min);
-	PassParameters->PSParameters.OriginalViewSize		= FVector2f(View.UnscaledViewRect.Width(), View.UnscaledViewRect.Height());
+	PassParameters->PSParameters.OriginalViewRectMin	= FVector2f(View.ViewRect.Min);
+	PassParameters->PSParameters.OriginalViewSize		= FVector2f(View.ViewRect.Width(), View.ViewRect.Height());
 	PassParameters->PSParameters.OriginalBufferInvSize	= FVector2f(1.f / SceneDepth->Desc.Extent.X, 1.f / SceneDepth->Desc.Extent.Y);
 	PassParameters->PSParameters.OccludedColorScale		= GNiagaraGpuComputeDebug_OccludedLineColorScale;
 	PassParameters->PSParameters.DepthTexture			= SceneDepth;
@@ -242,7 +242,7 @@ void NiagaraDebugShaders::DrawDebugLines(
 	);
 }
 
-void NiagaraDebugShaders::DrawDebugLines(FRDGBuilder& GraphBuilder, const FSceneView& View, FRDGTextureRef SceneColor, FRDGTextureRef SceneDepth, FRDGBufferRef ArgsBuffer, FRDGBufferRef LineBuffer)
+void NiagaraDebugShaders::DrawDebugLines(FRDGBuilder& GraphBuilder, const FViewInfo& View, FRDGTextureRef SceneColor, FRDGTextureRef SceneDepth, FRDGBufferRef ArgsBuffer, FRDGBufferRef LineBuffer)
 {
 	TShaderMapRef<FNiagaraDebugDrawLineVS> VertexShader(GetGlobalShaderMap(View.GetFeatureLevel()));
 	TShaderMapRef<FNiagaraDebugDrawLinePS> PixelShader(GetGlobalShaderMap(View.GetFeatureLevel()));
@@ -254,8 +254,8 @@ void NiagaraDebugShaders::DrawDebugLines(FRDGBuilder& GraphBuilder, const FScene
 	PassParameters->VSParameters.GpuLineBuffer			= GraphBuilder.CreateSRV(LineBuffer, PF_R32_UINT);
 
 	PassParameters->PSParameters.OutputInvResolution	= FVector2f(1.0f / View.UnscaledViewRect.Width(), 1.0f / View.UnconstrainedViewRect.Height());
-	PassParameters->PSParameters.OriginalViewRectMin	= FVector2f(View.UnscaledViewRect.Min);
-	PassParameters->PSParameters.OriginalViewSize		= FVector2f(View.UnscaledViewRect.Width(), View.UnscaledViewRect.Height());
+	PassParameters->PSParameters.OriginalViewRectMin	= FVector2f(View.ViewRect.Min);
+	PassParameters->PSParameters.OriginalViewSize		= FVector2f(View.ViewRect.Width(), View.ViewRect.Height());
 	PassParameters->PSParameters.OriginalBufferInvSize	= FVector2f(1.f / SceneDepth->Desc.Extent.X, 1.f / SceneDepth->Desc.Extent.Y);
 	PassParameters->PSParameters.OccludedColorScale		= GNiagaraGpuComputeDebug_OccludedLineColorScale;
 	PassParameters->PSParameters.DepthTexture			= SceneDepth;
@@ -288,7 +288,7 @@ void NiagaraDebugShaders::DrawDebugLines(FRDGBuilder& GraphBuilder, const FScene
 }
 
 void NiagaraDebugShaders::VisualizeTexture(
-	class FRDGBuilder& GraphBuilder, const FSceneView& View, const FScreenPassRenderTarget& Output,
+	class FRDGBuilder& GraphBuilder, const FViewInfo& View, const FScreenPassRenderTarget& Output,
 	const FIntPoint& Location, const int32& DisplayHeight,
 	const FIntVector4& InAttributesToVisualize, FRDGTextureRef Texture, const FIntVector4& NumTextureAttributes, uint32 TickCounter,
 	const FVector2D& PreviewDisplayRange
diff --git a/Engine/Plugins/FX/Niagara/Source/NiagaraShader/Public/NiagaraDebugShaders.h b/Engine/Plugins/FX/Niagara/Source/NiagaraShader/Public/NiagaraDebugShaders.h
--- a/Engine/Plugins/FX/Niagara/Source/NiagaraShader/Public/NiagaraDebugShaders.h
+++ b/Engine/Plugins/FX/Niagara/Source/NiagaraShader/Public/NiagaraDebugShaders.h
@@ -12,17 +12,17 @@ namespace NiagaraDebugShaders
 	NIAGARASHADER_API void ClearUAV(FRDGBuilder& GraphBuilder, FRDGBufferUAVRef UAV, FUintVector4 ClearValues, uint32 UIntsToSet);
 
 	NIAGARASHADER_API void DrawDebugLines(
-		class FRDGBuilder& GraphBuilder, const class FSceneView& View, FRDGTextureRef SceneColor, FRDGTextureRef SceneDepth,
+		class FRDGBuilder& GraphBuilder, const class FViewInfo& View, FRDGTextureRef SceneColor, FRDGTextureRef SceneDepth,
 		const uint32 LineInstanceCount, FRDGBufferRef LineBuffer
 	);
 
 	NIAGARASHADER_API void DrawDebugLines(
-		class FRDGBuilder& GraphBuilder, const class FSceneView& View, FRDGTextureRef SceneColor, FRDGTextureRef SceneDepth,
+		class FRDGBuilder& GraphBuilder, const class FViewInfo& View, FRDGTextureRef SceneColor, FRDGTextureRef SceneDepth,
 		FRDGBufferRef ArgsBuffer, FRDGBufferRef LineBuffer
 	);
 
 	NIAGARASHADER_API void VisualizeTexture(
-		class FRDGBuilder& GraphBuilder, const FSceneView& View, const FScreenPassRenderTarget& Output,
+		class FRDGBuilder& GraphBuilder, const FViewInfo& View, const FScreenPassRenderTarget& Output,
 		const FIntPoint& Location, const int32& DisplayHeight,
 		const FIntVector4& AttributesToVisualize, FRDGTextureRef Texture, const FIntVector4& NumTextureAttributes, uint32 TickCounter,
 		const FVector2D& PreviewDisplayRange = FVector2D::ZeroVector
