- Fix race with ScriptParametersMetadata
- RT now has mirror for ScriptParametersMetadata as updating scalability, etc, can force a recache of the metadata while the system is running as we don't flush render commands.

diff --git a/Engine/Plugins/FX/Niagara/Source/Niagara/Private/NiagaraGpuComputeDispatch.cpp b/Engine/Plugins/FX/Niagara/Source/Niagara/Private/NiagaraGpuComputeDispatch.cpp
--- a/Engine/Plugins/FX/Niagara/Source/Niagara/Private/NiagaraGpuComputeDispatch.cpp
+++ b/Engine/Plugins/FX/Niagara/Source/Niagara/Private/NiagaraGpuComputeDispatch.cpp
@@ -1436,7 +1436,7 @@ void FNiagaraGpuComputeDispatch::DispatchStage(FRDGBuilder& GraphBuilder, const
 
 	checkf(DispatchNumThreads.X * DispatchNumThreads.Y * DispatchNumThreads.Z > 0, TEXT("DispatchNumThreads(%d, %d, %d) is invalid"), DispatchNumThreads.X, DispatchNumThreads.Y, DispatchNumThreads.Z);
 
-	const FNiagaraShaderScriptParametersMetadata& NiagaraShaderParametersMetadata = InstanceData.Context->GPUScript_RT->GetScriptParametersMetadata().Get();
+	const FNiagaraShaderScriptParametersMetadata& NiagaraShaderParametersMetadata = InstanceData.Context->GPUScript_RT->GetScriptParametersMetadata_RT();
 	const FShaderParametersMetadata* ShaderParametersMetadata = NiagaraShaderParametersMetadata.ShaderParametersMetadata.Get();
 	FNiagaraShader::FParameters* DispatchParameters = GraphBuilder.AllocParameters<FNiagaraShader::FParameters>(ShaderParametersMetadata);
 
diff --git a/Engine/Plugins/FX/Niagara/Source/NiagaraShader/Private/NiagaraShared.cpp b/Engine/Plugins/FX/Niagara/Source/NiagaraShader/Private/NiagaraShared.cpp
--- a/Engine/Plugins/FX/Niagara/Source/NiagaraShader/Private/NiagaraShared.cpp
+++ b/Engine/Plugins/FX/Niagara/Source/NiagaraShader/Private/NiagaraShared.cpp
@@ -53,6 +53,7 @@ FNiagaraShaderScript::FNiagaraShaderScript()
 	, GameThreadShaderMap(nullptr)
 	, RenderingThreadShaderMap(nullptr)
 	, ScriptParametersMetadata(MakeShared<FNiagaraShaderScriptParametersMetadata>())
+	, ScriptParametersMetadata_RT(ScriptParametersMetadata)
 	, FeatureLevel(GMaxRHIFeatureLevel)
 	, ShaderPlatform(SP_NumPlatforms)
 	, bLoadedCookedShaderMapId(false)
@@ -723,6 +724,13 @@ void FNiagaraShaderScript::BuildScriptParametersMetadata(const FNiagaraShaderScr
 	}
 
 	ScriptParametersMetadata->ShaderParametersMetadata = MakeShareable<FShaderParametersMetadata>(ShaderMetadataBuilder.Build(FShaderParametersMetadata::EUseCase::ShaderParameterStruct, TEXT("FNiagaraShaderScript")));
+
+	ENQUEUE_RENDER_COMMAND(SetScriptParametersMetadata)(
+		[this, NewMetaData=ScriptParametersMetadata](FRHICommandListImmediate& RHICmdList)
+		{
+			ScriptParametersMetadata_RT=NewMetaData;
+		}
+	);
 }
 
 FNiagaraShaderRef FNiagaraShaderScript::GetShader(int32 PermutationId) const
diff --git a/Engine/Plugins/FX/Niagara/Source/NiagaraShader/Public/NiagaraShared.h b/Engine/Plugins/FX/Niagara/Source/NiagaraShader/Public/NiagaraShared.h
--- a/Engine/Plugins/FX/Niagara/Source/NiagaraShader/Public/NiagaraShared.h
+++ b/Engine/Plugins/FX/Niagara/Source/NiagaraShader/Public/NiagaraShared.h
@@ -877,6 +877,7 @@ public:
 	
 	NIAGARASHADER_API void BuildScriptParametersMetadata(const FNiagaraShaderScriptParametersMetadata& ScriptParametersMetadata);
 	TSharedRef<FNiagaraShaderScriptParametersMetadata> GetScriptParametersMetadata() const { check(ScriptParametersMetadata.IsValid()); return ScriptParametersMetadata.ToSharedRef(); }
+	const FNiagaraShaderScriptParametersMetadata& GetScriptParametersMetadata_RT() const { check(ScriptParametersMetadata_RT.IsValid()); return *ScriptParametersMetadata_RT.Get(); }
 
 	NIAGARASHADER_API FOnNiagaraScriptCompilationComplete& OnCompilationComplete()
 	{
@@ -935,6 +936,7 @@ private:
 	FNiagaraShaderMap* RenderingThreadShaderMap;
 
 	TSharedPtr<FNiagaraShaderScriptParametersMetadata> ScriptParametersMetadata;
+	TSharedPtr<FNiagaraShaderScriptParametersMetadata> ScriptParametersMetadata_RT;
 
 	/** Guid id for base script*/
 	FGuid BaseScriptId_DEPRECATED;
