Reorganize NiagaraGraph::PostLoad to reduce the cost by limiting the number of times we invalidate results during PostLoad for the graph and associated scripts.  Using -trace=loadtime and cooking a heavy niagara based map, it reduces the loadtime for NiagaraSystems by 25%

diff --git a/Engine/Plugins/FX/Niagara/Source/Niagara/Classes/NiagaraScriptSourceBase.h b/Engine/Plugins/FX/Niagara/Source/Niagara/Classes/NiagaraScriptSourceBase.h
--- a/Engine/Plugins/FX/Niagara/Source/Niagara/Classes/NiagaraScriptSourceBase.h
+++ b/Engine/Plugins/FX/Niagara/Source/Niagara/Classes/NiagaraScriptSourceBase.h
@@ -118,7 +118,12 @@ class UNiagaraScriptSourceBase : public UObject
 
 	virtual FGuid GetChangeID() { return FGuid(); }
 
-	virtual void ComputeVMCompilationId(struct FNiagaraVMExecutableDataId& Id, ENiagaraScriptUsage InUsage, const FGuid& InUsageId, bool bForceRebuild = false) const {};
+	virtual void ComputeVMCompilationId(struct FNiagaraVMExecutableDataId& Id, ENiagaraScriptUsage InUsage, const FGuid& InUsageId) const {};
+
+#if WITH_EDITORONLY_DATA
+	// Will conditionally refresh the graph's CompileId
+	virtual void RefreshGraphCompileId() {};
+#endif
 
 	virtual TMap<FName, UNiagaraDataInterface*> ComputeObjectNameMap(UNiagaraSystem& System, ENiagaraScriptUsage Usage, FGuid UsageId, FString EmitterUniqueName) const { return TMap<FName, UNiagaraDataInterface*>(); }
 
diff --git a/Engine/Plugins/FX/Niagara/Source/Niagara/Private/NiagaraScript.cpp b/Engine/Plugins/FX/Niagara/Source/Niagara/Private/NiagaraScript.cpp
--- a/Engine/Plugins/FX/Niagara/Source/Niagara/Private/NiagaraScript.cpp
+++ b/Engine/Plugins/FX/Niagara/Source/Niagara/Private/NiagaraScript.cpp
@@ -2070,65 +2070,66 @@ void UNiagaraScript::PostLoad()
 	for (FVersionedNiagaraScriptData& Data : VersionData)
 	{
 		UNiagaraScriptSourceBase* Source = Data.Source;
-        if (Source != nullptr)
-        {
-        	Source->ConditionalPostLoad();
+		if (Source != nullptr)
+		{
+			Source->ConditionalPostLoad();
 
 			// Synchronize with Definitions after source scripts have been postloaded.
 			FVersionedNiagaraScript& VersionedScriptAdapter = VersionedScriptAdapters.Emplace_GetRef(this, Data.Version.VersionGuid);
 			VersionedScriptAdapter.PostLoadDefinitionsSubscriptions();
 
-        	bool bScriptVMNeedsRebuild = false;
-        	FString RebuildReason;
-        	if (NiagaraVer < FNiagaraCustomVersion::UseHashesToIdentifyCompileStateOfTopLevelScripts && CachedScriptVMId.CompilerVersionID.IsValid())
-        	{
-        		FGuid BaseId = Source->GetCompileBaseId(Usage, UsageId);
-        		if (BaseId.IsValid() == false)
-        		{
-        			UE_LOG(LogNiagara, Warning,
-                        TEXT("Invalidating compile ids for script %s because it doesn't have a valid base id.  The owning asset will continue to compile on load until it is resaved."),
-                        *GetPathName());
-        			InvalidateCompileResults(TEXT("Script didn't have a valid base id."));
-        			Source->ForceGraphToRecompileOnNextCheck();
-        		}
-        		else
-        		{
-        			FNiagaraCompileHash CompileHash = Source->GetCompileHash(Usage, UsageId);
-        			if (CompileHash.IsValid())
-        			{
-        				CachedScriptVMId.BaseScriptCompileHash = CompileHash;
-        			}
-        			else
-        			{
-        				// If the compile hash isn't valid, the vm id needs to be recalculated and the cached vm needs to be invalidated.
-        				bScriptVMNeedsRebuild = true;
-        				RebuildReason = TEXT("Script did not have a valid compile hash.");
-        			}
-        		}
-        	}
-
-        	if (CachedScriptVMId.CompilerVersionID.IsValid() && CachedScriptVMId.CompilerVersionID != FNiagaraCustomVersion::GetLatestScriptCompileVersion())
-        	{
-        		bScriptVMNeedsRebuild = true;
-        		RebuildReason = TEXT("Niagara compiler version changed since the last time the script was compiled.");
-        	}
-
-        	if (bScriptVMNeedsRebuild)
-        	{
-        		// Force a rebuild on the source vm ids, and then invalidate the current cache to force the script to be unsynchronized.
-        		// We modify here in post load so that it will cause the owning asset to resave when running the resave commandlet.
-        		bool bForceRebuild = true;
-        		Modify();
-        		Source->ComputeVMCompilationId(CachedScriptVMId, Usage, UsageId, bForceRebuild);
-        		InvalidateCompileResults(RebuildReason);
-        	}
-
-        	// Convert visibility of old assets
-        	if (NiagaraVer < FNiagaraCustomVersion::AddLibraryAssetProperty || (NiagaraVer < FNiagaraCustomVersion::AddLibraryVisibilityProperty && bExposeToLibrary_DEPRECATED))
-        	{
-        		ScriptData->LibraryVisibility = ENiagaraScriptLibraryVisibility::Library;
-        	}
-        }
+			bool bScriptVMNeedsRebuild = false;
+			FString RebuildReason;
+			if (NiagaraVer < FNiagaraCustomVersion::UseHashesToIdentifyCompileStateOfTopLevelScripts && CachedScriptVMId.CompilerVersionID.IsValid())
+			{
+				FGuid BaseId = Source->GetCompileBaseId(Usage, UsageId);
+				if (BaseId.IsValid() == false)
+				{
+					UE_LOG(LogNiagara, Warning,
+						TEXT("Invalidating compile ids for script %s because it doesn't have a valid base id.  The owning asset will continue to compile on load until it is resaved."),
+						*GetPathName());
+					InvalidateCompileResults(TEXT("Script didn't have a valid base id."));
+					Source->ForceGraphToRecompileOnNextCheck();
+				}
+				else
+				{
+					FNiagaraCompileHash CompileHash = Source->GetCompileHash(Usage, UsageId);
+					if (CompileHash.IsValid())
+					{
+						CachedScriptVMId.BaseScriptCompileHash = CompileHash;
+					}
+					else
+					{
+						// If the compile hash isn't valid, the vm id needs to be recalculated and the cached vm needs to be invalidated.
+						bScriptVMNeedsRebuild = true;
+						RebuildReason = TEXT("Script did not have a valid compile hash.");
+					}
+				}
+			}
+
+			if (CachedScriptVMId.CompilerVersionID.IsValid() && CachedScriptVMId.CompilerVersionID != FNiagaraCustomVersion::GetLatestScriptCompileVersion())
+			{
+				bScriptVMNeedsRebuild = true;
+				RebuildReason = TEXT("Niagara compiler version changed since the last time the script was compiled.");
+			}
+
+			if (bScriptVMNeedsRebuild)
+			{
+        			// We modify here in post load so that it will cause the owning asset to resave when running the resave commandlet.
+				Modify();
+				InvalidateCompileResults(RebuildReason);
+
+				// Older versions of graphs require the script to drive the CachedCompileIds being fully refreshed when the script version changes.
+				// In order to bridge the gap without requiring a new package version we trigger the refresh here
+				Source->RefreshGraphCompileId();
+			}
+
+			// Convert visibility of old assets
+			if (NiagaraVer < FNiagaraCustomVersion::AddLibraryAssetProperty || (NiagaraVer < FNiagaraCustomVersion::AddLibraryVisibilityProperty && bExposeToLibrary_DEPRECATED))
+			{
+				ScriptData->LibraryVisibility = ENiagaraScriptLibraryVisibility::Library;
+			}
+		}
 
 		if (Data.InputSections.Num() > 0)
 		{
diff --git a/Engine/Plugins/FX/Niagara/Source/NiagaraEditor/Private/NiagaraGraph.cpp b/Engine/Plugins/FX/Niagara/Source/NiagaraEditor/Private/NiagaraGraph.cpp
--- a/Engine/Plugins/FX/Niagara/Source/NiagaraEditor/Private/NiagaraGraph.cpp
+++ b/Engine/Plugins/FX/Niagara/Source/NiagaraEditor/Private/NiagaraGraph.cpp
@@ -151,130 +151,176 @@ void UNiagaraGraph::NotifyGraphChanged()
 	InvalidateNumericCache();
 }
 
-void UNiagaraGraph::PostLoad()
+void UNiagaraGraph::PostLoad_LWCFixup(int32 NiagaraVersion)
 {
-	check(!bIsForCompilationOnly);
+	// TODO - put behind a CustomVersion
+	TArray<UNiagaraNodeParameterMapBase*> NiagaraParameterNodes;
+	GetNodesOfClass<UNiagaraNodeParameterMapBase>(NiagaraParameterNodes);
 
-	Super::PostLoad();
+	if (NiagaraParameterNodes.IsEmpty())
+	{
+		return;
+	}
+	bool GraphChanged = false;
 
-	NIAGARA_SCOPE_CYCLE_COUNTER(STAT_NiagaraEditor_Graph_PostLoad);
+	const UEdGraphSchema_Niagara* NiagaraSchema = GetDefault<UEdGraphSchema_Niagara>();
+	TArray<FNiagaraVariable> OldTypes = FNiagaraConstants::GetOldPositionTypeVariables();
 
-	const int32 NiagaraVer = GetLinkerCustomVersion(FNiagaraCustomVersion::GUID);
+	const FEdGraphPinType PositionPinType = NiagaraSchema->TypeDefinitionToPinType(FNiagaraTypeDefinition::GetPositionDef());
 
-	if (!VariableToScriptVariable.IsEmpty())
+	for (UNiagaraNodeParameterMapBase* NiagaraNode : NiagaraParameterNodes)
 	{
-		const TArray<FNiagaraVariable> StaticSwitchInputs = FindStaticSwitchInputs();
-
-		for (auto It = VariableToScriptVariable.CreateIterator(); It; ++It)
+		for (UEdGraphPin* Pin : NiagaraNode->Pins)
 		{
-			FNiagaraVariable Var = It.Key();
-			TObjectPtr<UNiagaraScriptVariable>& ScriptVar = It.Value();
-
-			if (ScriptVar == nullptr)
+			if (NiagaraNode->IsAddPin(Pin))
 			{
-				ScriptVar = NewObject<UNiagaraScriptVariable>(const_cast<UNiagaraGraph*>(this));
-				ScriptVar->Init(Var, FNiagaraVariableMetaData());
-				UE_LOG(LogNiagaraEditor, Display, TEXT("Fixed null UNiagaraScriptVariable | variable %s | asset path %s"), *Var.GetName().ToString(), *GetPathName());
+				continue;
 			}
-			else
+			FNiagaraVariable Variable = FNiagaraVariable(NiagaraSchema->PinToTypeDefinition(Pin), Pin->PinName);
+			for (const FNiagaraVariable& OldVarType : OldTypes)
 			{
-				if (ScriptVar->GetName().Contains(TEXT(".")) || ScriptVar->GetName().Contains(TEXT(" ")))
+				if (Variable == OldVarType)
 				{
-					ScriptVar->Rename(*MakeUniqueObjectName(this, UNiagaraScriptVariable::StaticClass(), NAME_None).ToString(), nullptr, REN_NonTransactional | REN_ForceNoResetLoaders);
+					Pin->PinType = PositionPinType;
+					Pin->PinType.PinSubCategory = UNiagaraNodeParameterMapBase::ParameterPinSubCategory;
+					NiagaraNode->MarkNodeRequiresSynchronization(__FUNCTION__, false);
+					GraphChanged = true;
+					break;
 				}
+			}
 
-				// Conditional postload all ScriptVars to ensure static switch default values are allocated as these are required when postloading all graph nodes later.
-				ScriptVar->ConditionalPostLoad();
+			// fix default pin types if necessary
+			// NOTE - that this code is impacting more than just LWC related nodes, it is touching all MapGet nodes
+			//		- Further by assigning the Pin->PinType to the DefaultPin->PinType it is going to assign
+			//			UNiagaraNodeParameterMapBase::ParameterPinSubCategory to the SubCategory of the input (default) pin
+			//			which currently is reserved for output pins.  This seems to be an unnecessary change and does have
+			//			an impact when dealing with RefreshParameterReferences' AddParameterReference lambda where the subcategory
+			//			is evaluated.  This results in DefaultValues not necessarily getting into the compile hash (if the default pin
+			//			is the first reference added, then it will be encountered and rejected in AppendCompileHash and then the
+			//			default value won't be handled
+			if (Pin->Direction == EGPD_Output)
+			{
+				if (UNiagaraNodeParameterMapGet* GetNode = Cast<UNiagaraNodeParameterMapGet>(NiagaraNode))
+				{
+					UEdGraphPin* DefaultPin = GetNode->GetDefaultPin(Pin);
+					if (DefaultPin && DefaultPin->PinType != Pin->PinType)
+					{
+						DefaultPin->PinType = Pin->PinType;
+						NiagaraNode->MarkNodeRequiresSynchronization(__FUNCTION__, false);
+						GraphChanged = true;
+					}
+				}
 			}
-			ScriptVar->SetIsStaticSwitch(StaticSwitchInputs.Contains(Var));
 		}
 	}
 
-	for (UEdGraphNode* Node : Nodes)
-	{
-		Node->ConditionalPostLoad();
-	}
-
-	// There was a bug where the graph traversal could have nullptr entries in it, which causes crashes later. Detect these and just clear out the cache.
-	bool bAnyInvalid = false;
-	for (FNiagaraGraphScriptUsageInfo& CachedUsageInfoItem : CachedUsageInfo)
+	TArray<UNiagaraScriptVariable*> ChangedScriptVariables;
+	for (FNiagaraVariable OldVarType : OldTypes)
 	{
-		CachedUsageInfoItem.PostLoad(this);
-		if (!CachedUsageInfoItem.IsValid())
+		TObjectPtr<UNiagaraScriptVariable>* ScriptVariablePtr = VariableToScriptVariable.Find(OldVarType);
+		if (ScriptVariablePtr != nullptr && *ScriptVariablePtr)
 		{
-			bAnyInvalid = true;
+			const FNiagaraVariable NewVarType = FNiagaraVariable(FNiagaraTypeDefinition::GetPositionDef(), OldVarType.GetName());
+
+			UNiagaraScriptVariable* ScriptVar = ScriptVariablePtr->Get();
+
+			if (ScriptVar->DefaultMode == ENiagaraDefaultMode::Value)
+			{
+				ChangedScriptVariables.Add(ScriptVar);
+			}
+
+			ScriptVar->Variable.SetType(FNiagaraTypeDefinition::GetPositionDef());
+			VariableToScriptVariable.Remove(OldVarType);
+
+			VariableToScriptVariable.Add(NewVarType, TObjectPtr<UNiagaraScriptVariable>(ScriptVar));
+			GraphChanged = true;
 		}
 	}
-	if (bAnyInvalid)
-	{
-		CachedUsageInfo.Empty();
-		LastBuiltTraversalDataChangeId = FGuid();
-	}
 
-	// In the past, we didn't bother setting the CallSortPriority and just used lexicographic ordering.
-	// In the event that we have multiple non-matching nodes with a zero call sort priority, this will
-	// give every node a unique order value.
-	TArray<UNiagaraNodeInput*> InputNodes;
-	GetNodesOfClass(InputNodes);
-	bool bAllZeroes = true;
-	TArray<FName> UniqueNames;
-	for (UNiagaraNodeInput* InputNode : InputNodes)
+	if (!ChangedScriptVariables.IsEmpty())
 	{
-		if (InputNode->CallSortPriority != 0)
+		TSet<UNiagaraNodeParameterMapGet*> MapGetNodes;
+		for (UNiagaraScriptVariable* ChangedScriptVariable : ChangedScriptVariables)
 		{
-			bAllZeroes = false;
-		}
+			FNiagaraVariable& Variable = ChangedScriptVariable->Variable;
 
-		if (InputNode->Usage == ENiagaraInputNodeUsage::Parameter)
-		{
-			UniqueNames.AddUnique(InputNode->Input.GetName());
-		}
+			FNiagaraEditorModule& EditorModule = FNiagaraEditorModule::Get();
+			auto TypeUtilityValue = EditorModule.GetTypeUtilities(Variable.GetType());
 
-		if (InputNode->Usage == ENiagaraInputNodeUsage::SystemConstant)
-		{
-			InputNode->Input = FNiagaraConstants::UpdateEngineConstant(InputNode->Input);
+			if (TypeUtilityValue.IsValid() && TypeUtilityValue->CanHandlePinDefaults())
+			{
+				TArray<UEdGraphPin*> Pins = FindParameterMapDefaultValuePins(Variable.GetName());
+				for (UEdGraphPin* Pin : Pins)
+				{
+					check(ChangedScriptVariable->DefaultMode == ENiagaraDefaultMode::Value && !Variable.GetType().IsDataInterface());
+
+					Variable.AllocateData();
+					FString NewDefaultValue = TypeUtilityValue->GetPinDefaultStringFromValue(Variable);
+					NiagaraSchema->TrySetDefaultValue(*Pin, NewDefaultValue, false);
+
+					MapGetNodes.Add(CastChecked<UNiagaraNodeParameterMapGet>(Pin->GetOwningNode()));
+				}
+			}
 		}
-	}
 
-	if (bAllZeroes && UniqueNames.Num() > 1)
-	{
-		// Just do the lexicographic sort and assign the call order to their ordered index value.
-		UniqueNames.Sort(FNameLexicalLess());
-		for (UNiagaraNodeInput* InputNode : InputNodes)
+		if (!MapGetNodes.IsEmpty())
 		{
-			if (InputNode->Usage == ENiagaraInputNodeUsage::Parameter)
+			GraphChanged = true;
+			for (UNiagaraNodeParameterMapGet* MapGet : MapGetNodes)
 			{
-				int32 FoundIndex = UniqueNames.Find(InputNode->Input.GetName());
-				check(FoundIndex != -1);
-				InputNode->CallSortPriority = FoundIndex;
+				MapGet->SynchronizeDefaultPins();
 			}
 		}
 	}
 
+	if (GraphChanged)
+	{
+		Modify();
+		ChangeId = FGuid::NewGuid();
+	}
+}
+
+void UNiagaraGraph::PostLoad()
+{
+	check(!bIsForCompilationOnly);
+
+	Super::PostLoad();
+
+	//////////////////////////////////////////////////////////////////////////
+	/// GLOBAL CLEANUP/INVALIDATION
+	///		-has no dependence on other chunks
+	///		-can add/modify UNiagaraScriptVariable
+	//////////////////////////////////////////////////////////////////////////
+
+	if (GIsEditor)
+	{
+		SetFlags(RF_Transactional);
+	}
+
+	const int32 NiagaraVer = GetLinkerCustomVersion(FNiagaraCustomVersion::GUID);
+
 	// If this is from a prior version, enforce a valid Change Id!
 	if (ChangeId.IsValid() == false)
 	{
 		MarkGraphRequiresSynchronization(TEXT("Graph change id was invalid"));
 	}
 
-	// Assume that all externally referenced assets have changed, so update to match. They will return true if they have changed.
-	for (UEdGraphNode* Node : Nodes)
+	// There was a bug where the graph traversal could have nullptr entries in it, which causes crashes later. Detect these and just clear out the cache.
 	{
-		if (UNiagaraNode* NiagaraNode = Cast<UNiagaraNode>(Node))
+		bool bAnyInvalid = false;
+		for (FNiagaraGraphScriptUsageInfo& CachedUsageInfoItem : CachedUsageInfo)
 		{
-			if (UObject* ReferencedAsset = NiagaraNode->GetReferencedAsset())
+			CachedUsageInfoItem.PostLoad(this);
+			if (!CachedUsageInfoItem.IsValid())
 			{
-				ReferencedAsset->ConditionalPostLoad();
-				NiagaraNode->RefreshFromExternalChanges();
+				bAnyInvalid = true;
 			}
 		}
-	}
-
-	RebuildCachedCompileIds();
-
-	if (GIsEditor)
-	{
-		SetFlags(RF_Transactional);
+		if (bAnyInvalid)
+		{
+			CachedUsageInfo.Empty();
+			LastBuiltTraversalDataChangeId = FGuid();
+		}
 	}
 
 	// Migrate input condition metadata
@@ -293,24 +339,94 @@ void UNiagaraGraph::PostLoad()
 		VariableToMetaData_DEPRECATED.Empty();
 	}
 
-	// Fix inconsistencies in the default value declaration between graph and metadata
-	ValidateDefaultPins();
-	
+	// TODO - put behind a CustomVersion
+	if (!VariableToScriptVariable.IsEmpty())
 	{
-		// Create a UNiagaraScriptVariable instance for every entry in the parameter map for which there is no existing script variable. 
-		TArray<FNiagaraVariable> VarsToAdd;
-		for (auto& ParameterToReferences : GetParameterReferenceMap())
+		const TArray<FNiagaraVariable> StaticSwitchInputs = FindStaticSwitchInputs();
+
+		for (auto It = VariableToScriptVariable.CreateIterator(); It; ++It)
 		{
-			UNiagaraScriptVariable* Variable = GetScriptVariable(ParameterToReferences.Key.GetName());
-			if (Variable == nullptr)
+			FNiagaraVariable Var = It.Key();
+			TObjectPtr<UNiagaraScriptVariable>& ScriptVar = It.Value();
+
+			if (ScriptVar == nullptr)
 			{
-				VarsToAdd.Add(ParameterToReferences.Key);
+				ScriptVar = NewObject<UNiagaraScriptVariable>(const_cast<UNiagaraGraph*>(this));
+				ScriptVar->Init(Var, FNiagaraVariableMetaData());
+				UE_LOG(LogNiagaraEditor, Display, TEXT("Fixed null UNiagaraScriptVariable | variable %s | asset path %s"), *Var.GetName().ToString(), *GetPathName());
+			}
+			else
+			{
+				if (ScriptVar->GetName().Contains(TEXT(".")) || ScriptVar->GetName().Contains(TEXT(" ")))
+				{
+					ScriptVar->Rename(*MakeUniqueObjectName(this, UNiagaraScriptVariable::StaticClass(), NAME_None).ToString(), nullptr, REN_NonTransactional | REN_ForceNoResetLoaders);
+				}
+
+				// Conditional postload all ScriptVars to ensure static switch default values are allocated as these are required when postloading all graph nodes later.
+				ScriptVar->ConditionalPostLoad();
 			}
+			ScriptVar->SetIsStaticSwitch(StaticSwitchInputs.Contains(Var));
 		}
+	}
 
-		for (FNiagaraVariable& Var : VarsToAdd)
+	if (NiagaraVer < FNiagaraCustomVersion::StaticSwitchFunctionPinsUsePersistentGuids)
+	{
+		// Make sure that the variable persistent ids are unique, otherwise pin allocation for static switches will be inconsistent.
+		TArray<TObjectPtr<UNiagaraScriptVariable>> ScriptVariables;
+		VariableToScriptVariable.GenerateValueArray(ScriptVariables);
+		TMap<FGuid, UNiagaraScriptVariable*> VariableGuidToScriptVariable;
+		for (TObjectPtr<UNiagaraScriptVariable>& ScriptVariable : ScriptVariables)
 		{
-			AddParameter(Var);
+			if (ScriptVariable != nullptr)
+			{
+				UNiagaraScriptVariable** MatchingScriptVariablePtr = VariableGuidToScriptVariable.Find(ScriptVariable->Metadata.GetVariableGuid());
+				if (MatchingScriptVariablePtr == nullptr)
+				{
+					VariableGuidToScriptVariable.Add(ScriptVariable->Metadata.GetVariableGuid(), ScriptVariable);
+				}
+				else
+				{
+					UNiagaraScriptVariable* MatchingScriptVariable = *MatchingScriptVariablePtr;
+					if (ScriptVariable->GetIsSubscribedToParameterDefinitions() && MatchingScriptVariable->GetIsSubscribedToParameterDefinitions())
+					{
+						// Both of the script variables with duplicate ids are controlled by parameter definitions so issue a warning because neither will be updated.
+						UE_LOG(LogNiagaraEditor, Log, TEXT("Duplicate ids found for script variables which are both subscribed to parameter definitions.\nScript Variable 1 Name: %s Type: %s Path: %s\nScript Variable 2 Name: %s Type: %s Path: %s"),
+							*MatchingScriptVariable->Variable.GetName().ToString(), *MatchingScriptVariable->Variable.GetType().GetName(), *MatchingScriptVariable->GetPathName(),
+							*ScriptVariable->Variable.GetName().ToString(), *ScriptVariable->Variable.GetType().GetName(), *ScriptVariable->GetPathName());
+					}
+					else
+					{
+						// Remove the duplicated entry and regenerate the ids for entries not subscribed to parameter definitions.
+						VariableGuidToScriptVariable.Remove(ScriptVariable->Metadata.GetVariableGuid());
+						if (ScriptVariable->GetIsSubscribedToParameterDefinitions() == false)
+						{
+							ScriptVariable->Metadata.SetVariableGuid(UNiagaraScriptVariable::GenerateStableGuid(ScriptVariable));
+							if (VariableGuidToScriptVariable.Contains(ScriptVariable->Metadata.GetVariableGuid()))
+							{
+								UE_LOG(LogNiagaraEditor, Log, TEXT("Could not generate a stable unique variable guid for script variable. Name: %s Type: %s Path: %s"),
+									*ScriptVariable->Variable.GetName().ToString(), *ScriptVariable->Variable.GetType().GetName(), *ScriptVariable->GetPathName());
+							}
+							else
+							{
+								VariableGuidToScriptVariable.Add(ScriptVariable->Metadata.GetVariableGuid(), ScriptVariable);
+							}
+						}
+						if (MatchingScriptVariable->GetIsSubscribedToParameterDefinitions() == false)
+						{
+							MatchingScriptVariable->Metadata.SetVariableGuid(UNiagaraScriptVariable::GenerateStableGuid(MatchingScriptVariable));
+							if (VariableGuidToScriptVariable.Contains(MatchingScriptVariable->Metadata.GetVariableGuid()))
+							{
+								UE_LOG(LogNiagaraEditor, Log, TEXT("Could not generate a stable unique variable guid for script variable. Name: %s Type: %s Path: %s"),
+									*MatchingScriptVariable->Variable.GetName().ToString(), *MatchingScriptVariable->Variable.GetType().GetName(), *MatchingScriptVariable->GetPathName());
+							}
+							else
+							{
+								VariableGuidToScriptVariable.Add(MatchingScriptVariable->Metadata.GetVariableGuid(), MatchingScriptVariable);
+							}
+						}
+					}
+				}
+			}
 		}
 	}
 
@@ -366,134 +482,117 @@ void UNiagaraGraph::PostLoad()
 		StandardizeParameterNames();
 	}
 
-	// Fix LWC position type mapping; before LWC, all position were FVector3f, but now that they are a separate type, we need to upgrade the pins and metadata to match
-	const UEdGraphSchema_Niagara* NiagaraSchema = GetDefault<UEdGraphSchema_Niagara>();
-	TArray<UNiagaraNodeParameterMapBase*> NiagaraParameterNodes;
-	GetNodesOfClass<UNiagaraNodeParameterMapBase>(NiagaraParameterNodes);
-	TArray<FNiagaraVariable> OldTypes = FNiagaraConstants::GetOldPositionTypeVariables();
-	for (UNiagaraNodeParameterMapBase* NiagaraNode : NiagaraParameterNodes)
+	//////////////////////////////////////////////////////////////////////////
+	/// NODE VALIDATION
+	//////////////////////////////////////////////////////////////////////////
+
+
 	{
-		for (UEdGraphPin* Pin : NiagaraNode->Pins)
+		TArray<UNiagaraNodeInput*> InputNodes;
+
+		for (UEdGraphNode* Node : Nodes)
 		{
-			if (NiagaraNode->IsAddPin(Pin))
+			Node->ConditionalPostLoad();
+
+			if (UNiagaraNode* NiagaraNode = Cast<UNiagaraNode>(Node))
 			{
-				continue;
+				// Assume that all externally referenced assets have changed, so update to match. They will return true if they have changed.
+				if (UObject* ReferencedAsset = NiagaraNode->GetReferencedAsset())
+				{
+					ReferencedAsset->ConditionalPostLoad();
+					NiagaraNode->RefreshFromExternalChanges();
+				}
+
+				if (UNiagaraNodeInput* InputNode = Cast<UNiagaraNodeInput>(NiagaraNode))
+				{
+					InputNodes.Add(InputNode);
+				}
 			}
-			FNiagaraVariable Variable = FNiagaraVariable(NiagaraSchema->PinToTypeDefinition(Pin), Pin->PinName);
-			for (const FNiagaraVariable& OldVarType : OldTypes)
+		}
+
+		// In the past, we didn't bother setting the CallSortPriority and just used lexicographic ordering.
+		// In the event that we have multiple non-matching nodes with a zero call sort priority, this will
+		// give every node a unique order value.
+		if (!InputNodes.IsEmpty())
+		{
+			bool bPrioritiesAllZeroes = true;
+			TArray<FName> UniqueNames;
+
+			for (UNiagaraNodeInput* InputNode : InputNodes)
 			{
-				if (Variable == OldVarType)
+				if (InputNode->CallSortPriority != 0)
 				{
-					Pin->PinType = NiagaraSchema->TypeDefinitionToPinType(FNiagaraTypeDefinition::GetPositionDef());
-					Pin->PinType.PinSubCategory = UNiagaraNodeParameterMapBase::ParameterPinSubCategory;
-					NiagaraNode->MarkNodeRequiresSynchronization(__FUNCTION__, true);
-					break;
+					bPrioritiesAllZeroes = false;
+				}
+
+				if (InputNode->Usage == ENiagaraInputNodeUsage::Parameter)
+				{
+					UniqueNames.AddUnique(InputNode->Input.GetName());
+				}
+
+				if (InputNode->Usage == ENiagaraInputNodeUsage::SystemConstant)
+				{
+					InputNode->Input = FNiagaraConstants::UpdateEngineConstant(InputNode->Input);
 				}
 			}
 
-			// fix default pin types if necessary
-			if (UNiagaraNodeParameterMapGet* GetNode = Cast<UNiagaraNodeParameterMapGet>(NiagaraNode))
+			if (bPrioritiesAllZeroes && UniqueNames.Num() > 1)
 			{
-				UEdGraphPin* DefaultPin = GetNode->GetDefaultPin(Pin);
-				if (Pin->Direction == EGPD_Output && DefaultPin && DefaultPin->PinType != Pin->PinType)
+				// Just do the lexicographic sort and assign the call order to their ordered index value.
+				UniqueNames.Sort(FNameLexicalLess());
+				for (UNiagaraNodeInput* InputNode : InputNodes)
 				{
-					DefaultPin->PinType = Pin->PinType;
-					NiagaraNode->MarkNodeRequiresSynchronization(__FUNCTION__, true);
+					if (InputNode->Usage == ENiagaraInputNodeUsage::Parameter)
+					{
+						int32 FoundIndex = UniqueNames.Find(InputNode->Input.GetName());
+						check(FoundIndex != INDEX_NONE);
+						InputNode->CallSortPriority = FoundIndex;
+					}
 				}
 			}
 		}
 	}
-	for (FNiagaraVariable OldVarType : OldTypes)
-	{
-		TObjectPtr<UNiagaraScriptVariable>* ScriptVariablePtr = VariableToScriptVariable.Find(OldVarType);
-		if (ScriptVariablePtr != nullptr && *ScriptVariablePtr)
-		{
-			UNiagaraScriptVariable& ScriptVar = *ScriptVariablePtr->Get();
-			ScriptVar.Variable.SetType(FNiagaraTypeDefinition::GetPositionDef());
-			VariableToScriptVariable.Remove(OldVarType);
 
-			FNiagaraVariable NewVarType(FNiagaraTypeDefinition::GetPositionDef(), OldVarType.GetName());
-			VariableToScriptVariable.Add(NewVarType, TObjectPtr<UNiagaraScriptVariable>(&ScriptVar));
-			ScriptVariableChanged(NewVarType);
-		}
+	// Fix LWC position type mapping; before LWC, all position were FVector3f, but now that they are a separate type, we need to upgrade the pins and metadata to match
+	// NOTE - this function has the capacity to change the graph and so needs be done before we RefreshParameterReferences below (GetParameterReferenceMap)
+	// otherwise we'll be getting out of date results
+	PostLoad_LWCFixup(NiagaraVer);
 
-		// Also fix stale parameter reference map entries. This should usually happen in RefreshParameterReferences, but since we just changed VariableToScriptVariable we need to fix that as well. 
-		FNiagaraGraphParameterReferenceCollection* ReferenceCollection = ParameterToReferencesMap.Find(OldVarType);
-		if (ReferenceCollection != nullptr)
+	check(bParameterReferenceRefreshPending);
+
+	{
+		// Create a UNiagaraScriptVariable instance for every entry in the parameter map for which there is no existing script variable. 
+		// Note that the call to GetParameterReferenceMap() will:
+		//	-generate a parameter map traversal (slow)
+		//	-remove/modify UNiagaraScriptVariables from VariableToScriptVariable
+		//
+		// Once GetParameterReferenceMap is called we need to make sure that we don't invalidate the ParameterReferenceMap.  Any changes
+		// to the graph through the postload needs to also fixup the parameter map to avoid us having to regenerate the map from scratch
+		// costing us a parameter map traversal
+		TArray<FNiagaraVariable> VarsToAdd;
+		for (auto& ParameterToReferences : GetParameterReferenceMap())
 		{
-			if (ReferenceCollection->ParameterReferences.Num() > 0)
+			UNiagaraScriptVariable* Variable = GetScriptVariable(ParameterToReferences.Key.GetName());
+			if (Variable == nullptr)
 			{
-				FNiagaraVariable NewVarType(FNiagaraTypeDefinition::GetPositionDef(), OldVarType.GetName());
-				FNiagaraGraphParameterReferenceCollection NewReferenceCollection = *ReferenceCollection;
-				ParameterToReferencesMap.Add(NewVarType, NewReferenceCollection);
+				VarsToAdd.Add(ParameterToReferences.Key);
 			}
-			ParameterToReferencesMap.Remove(OldVarType);
 		}
-	}
 
-	if (NiagaraVer < FNiagaraCustomVersion::StaticSwitchFunctionPinsUsePersistentGuids)
-	{
-		// Make sure that the variable persistent ids are unique, otherwise pin allocation for static switches will be inconsistent.
-		TArray<TObjectPtr<UNiagaraScriptVariable>> ScriptVariables;
-		VariableToScriptVariable.GenerateValueArray(ScriptVariables);
-		TMap<FGuid, UNiagaraScriptVariable*> VariableGuidToScriptVariable;
-		for (TObjectPtr<UNiagaraScriptVariable>& ScriptVariable : ScriptVariables)
+		for (FNiagaraVariable& Var : VarsToAdd)
 		{
-			if (ScriptVariable != nullptr)
-			{
-				ScriptVariable->ConditionalPostLoad();
-				UNiagaraScriptVariable** MatchingScriptVariablePtr = VariableGuidToScriptVariable.Find(ScriptVariable->Metadata.GetVariableGuid());
-				if (MatchingScriptVariablePtr == nullptr)
-				{
-					VariableGuidToScriptVariable.Add(ScriptVariable->Metadata.GetVariableGuid(), ScriptVariable);
-				}
-				else
-				{
-					UNiagaraScriptVariable* MatchingScriptVariable = *MatchingScriptVariablePtr;
-					if (ScriptVariable->GetIsSubscribedToParameterDefinitions() && MatchingScriptVariable->GetIsSubscribedToParameterDefinitions())
-					{
-						// Both of the script variables with duplicate ids are controlled by parameter definitions so issue a warning because neither will be updated.
-						UE_LOG(LogNiagaraEditor, Log, TEXT("Duplicate ids found for script variables which are both subscribed to parameter definitions.\nScript Variable 1 Name: %s Type: %s Path: %s\nScript Variable 2 Name: %s Type: %s Path: %s"),
-							*MatchingScriptVariable->Variable.GetName().ToString(), *MatchingScriptVariable->Variable.GetType().GetName(), *MatchingScriptVariable->GetPathName(),
-							*ScriptVariable->Variable.GetName().ToString(), *ScriptVariable->Variable.GetType().GetName(), *ScriptVariable->GetPathName());
-					}
-					else
-					{
-						// Remove the duplicated entry and regenerate the ids for entries not subscribed to parameter definitions.
-						VariableGuidToScriptVariable.Remove(ScriptVariable->Metadata.GetVariableGuid());
-						if (ScriptVariable->GetIsSubscribedToParameterDefinitions() == false)
-						{
-							ScriptVariable->Metadata.SetVariableGuid(UNiagaraScriptVariable::GenerateStableGuid(ScriptVariable));
-							if (VariableGuidToScriptVariable.Contains(ScriptVariable->Metadata.GetVariableGuid()))
-							{
-								UE_LOG(LogNiagaraEditor, Log, TEXT("Could not generate a stable unique variable guid for script variable. Name: %s Type: %s Path: %s"),
-									*ScriptVariable->Variable.GetName().ToString(), *ScriptVariable->Variable.GetType().GetName(), *ScriptVariable->GetPathName());
-							}
-							else
-							{
-								VariableGuidToScriptVariable.Add(ScriptVariable->Metadata.GetVariableGuid(), ScriptVariable);
-							}
-						}
-						if (MatchingScriptVariable->GetIsSubscribedToParameterDefinitions() == false)
-						{
-							MatchingScriptVariable->Metadata.SetVariableGuid(UNiagaraScriptVariable::GenerateStableGuid(MatchingScriptVariable));
-							if (VariableGuidToScriptVariable.Contains(MatchingScriptVariable->Metadata.GetVariableGuid()))
-							{
-								UE_LOG(LogNiagaraEditor, Log, TEXT("Could not generate a stable unique variable guid for script variable. Name: %s Type: %s Path: %s"),
-									*MatchingScriptVariable->Variable.GetName().ToString(), *MatchingScriptVariable->Variable.GetType().GetName(), *MatchingScriptVariable->GetPathName());
-							}
-							else
-							{
-								VariableGuidToScriptVariable.Add(MatchingScriptVariable->Metadata.GetVariableGuid(), MatchingScriptVariable);
-							}
-						}
-					}
-				}
-			}
+			AddParameter(Var, FNiagaraVariableMetaData(), false, false);
 		}
 	}
 
-	InvalidateCachedParameterData();
+	// Fix inconsistencies in the default value declaration between graph and metadata
+	ValidateDefaultPins();
+
+	// this will ensure that the parameter maps have been refreshed and our NiagaraScriptVariable are up to date.
+	// There's an argument to be made that we should just call RefreshParameterReferences directly and avoid explicitly
+	// generating the compile ids but there's still code out there that is expecting to have fully generated compile ids
+	// available on load.
+	RebuildCachedCompileIds();
 }
 
 #if WITH_EDITORONLY_DATA
@@ -805,6 +904,8 @@ TArray<UEdGraphPin*> UNiagaraGraph::FindParameterMapDefaultValuePins(const FName
 	return DefaultPins;
 }
 
+// Impacts of this function include changing the default value of UNiagaraScriptVariable::Variable & DefaultMode (note that this won't impact the
+// key in the VariableToScriptVariable map.  Additionally properties on the pins can also be changed.
 void UNiagaraGraph::ValidateDefaultPins()
 {
 	check(!bIsForCompilationOnly);
@@ -3154,12 +3255,25 @@ void DiffProperties(const FNiagaraCompileHashVisitorDebugInfo& A, const FNiagara
 }
 
 
-void UNiagaraGraph::RebuildCachedCompileIds(bool bForce)
+void UNiagaraGraph::ConditionalRebuildCompileIdCache()
+{
+	if (!bHasValidLastBuiltScriptVersionId)
+	{
+		bHasValidLastBuiltScriptVersionId = true;
+		LastBuiltScriptVersionId = FGuid();
+	}
+}
+
+void UNiagaraGraph::RebuildCachedCompileIds()
 {
 	check(!bIsForCompilationOnly);
 
+	const FGuid CurrentScriptVersionId = FNiagaraCustomVersion::GetLatestScriptCompileVersion();
+
 	// If the graph hasn't changed since last rebuild, then do nothing.
-	if (!bForce && ChangeId == LastBuiltTraversalDataChangeId && LastBuiltTraversalDataChangeId.IsValid())
+	const bool bForceRebuild = bHasValidLastBuiltScriptVersionId && (LastBuiltScriptVersionId != CurrentScriptVersionId);
+	
+	if (!bForceRebuild && ChangeId == LastBuiltTraversalDataChangeId && LastBuiltTraversalDataChangeId.IsValid())
 	{
 		return;
 	}
@@ -3177,7 +3291,6 @@ void UNiagaraGraph::RebuildCachedCompileIds(bool bForce)
 	TArray<FNiagaraGraphScriptUsageInfo> NewUsageCache;
 	NewUsageCache.AddDefaulted(NiagaraOutputNodes.Num());
 
-	UEnum* FoundEnum = nullptr;
 	bool bNeedsAnyNewCompileIds = false;
 
 	FNiagaraGraphScriptUsageInfo* ParticleSpawnUsageInfo = nullptr;
@@ -3227,20 +3340,18 @@ void UNiagaraGraph::RebuildCachedCompileIds(bool bForce)
 
 		// We can't store in a FShaHash struct directly because you can't FProperty it. Using a standin of the same size.
 		{
-			TArray<uint8> DataHash;
-			DataHash.AddUninitialized(FSHA1::DigestSize);
+			TStaticArray<uint8, FSHA1::DigestSize> DataHash;
 			HashState.GetHash(DataHash.GetData());
 		
-			NewUsageCache[i].CompileHash = FNiagaraCompileHash(DataHash);
+			NewUsageCache[i].CompileHash = FNiagaraCompileHash(DataHash.GetData(), DataHash.Num());
 		}
 
 
 		{
 			// We can't store in a FShaHash struct directly because you can't UProperty it. Using a standin of the same size.
-			TArray<uint8> DataHash;
-			DataHash.AddUninitialized(FSHA1::DigestSize);
+			TStaticArray<uint8, FSHA1::DigestSize> DataHash;
 			GraphHashState.GetHash(DataHash.GetData());
-			NewUsageCache[i].CompileHashFromGraph = FNiagaraCompileHash(DataHash);
+			NewUsageCache[i].CompileHashFromGraph = FNiagaraCompileHash(DataHash.GetData(), DataHash.Num());
 			NewUsageCache[i].CompileLastObjects = Visitor.Values;
 
 #if WITH_EDITORONLY_DATA
@@ -3284,22 +3395,6 @@ void UNiagaraGraph::RebuildCachedCompileIds(bool bForce)
 #endif
 		}
 
-
-		// TODO sckime debug logic... should be disabled or put under a cvar in the future
-		{
-
-			if (FoundEnum == nullptr)
-			{
-				FoundEnum = StaticEnum<ENiagaraScriptUsage>();
-			}
-
-			FString ResultsEnum = TEXT("??");
-			if (FoundEnum)
-			{
-				ResultsEnum = FoundEnum->GetNameStringByValue((int64)NewUsageCache[i].UsageType);
-			}
-		}
-
 		if (UNiagaraScript::IsEquivalentUsage(NewUsageCache[i].UsageType, ENiagaraScriptUsage::ParticleSpawnScript) && NewUsageCache[i].UsageId == FGuid())
 		{
 			ParticleSpawnUsageInfo = &NewUsageCache[i];
@@ -3349,6 +3444,9 @@ void UNiagaraGraph::RebuildCachedCompileIds(bool bForce)
 	CachedUsageInfo = NewUsageCache;
 	LastBuiltTraversalDataChangeId = ChangeId;
 
+	LastBuiltScriptVersionId = CurrentScriptVersionId;
+	bHasValidLastBuiltScriptVersionId = true;
+
 	RebuildNumericCache();
 }
 
@@ -3585,7 +3683,7 @@ void UNiagaraGraph::RefreshParameterReferences() const
 	TSet<FNiagaraVariable> CandidateUnreferencedParametersToRemove;
 
 	// The set of pins which has already been handled by add parameters.
-	TSet<const UEdGraphPin*> HandledPins;
+	TSet<const UEdGraphPin*> HandledParameterMapPins;
 
 	// Purge existing parameter references and collect candidate unreferenced parameters.
 	for (auto& ParameterToReferences : ParameterToReferencesMap)
@@ -3600,22 +3698,17 @@ void UNiagaraGraph::RefreshParameterReferences() const
 
 	auto AddParameterReference = [&](const FNiagaraVariable& Parameter, const UEdGraphPin* Pin)
 	{
-		if (Pin->PinType.PinSubCategory == UNiagaraNodeParameterMapBase::ParameterPinSubCategory)
+		FNiagaraGraphParameterReferenceCollection* ReferenceCollection = ParameterToReferencesMap.Find(Parameter);
+		if (ReferenceCollection == nullptr)
 		{
-			FNiagaraGraphParameterReferenceCollection* ReferenceCollection = ParameterToReferencesMap.Find(Parameter);
-			if (ReferenceCollection == nullptr)
-			{
-				FNiagaraGraphParameterReferenceCollection& NewReferenceCollection = ParameterToReferencesMap.Add(Parameter);
-				NewReferenceCollection.Graph = this;
-				ReferenceCollection = &NewReferenceCollection;
-			}
-			ReferenceCollection->ParameterReferences.AddUnique(FNiagaraGraphParameterReference(Pin->PersistentGuid, Cast<UNiagaraNode>(Pin->GetOwningNode())));
-
-			// If we're adding a parameter reference then it needs to be removed from the list of candidate variables to remove since it's been referenced.
-			CandidateUnreferencedParametersToRemove.Remove(Parameter);
+			FNiagaraGraphParameterReferenceCollection& NewReferenceCollection = ParameterToReferencesMap.Add(Parameter);
+			NewReferenceCollection.Graph = this;
+			ReferenceCollection = &NewReferenceCollection;
 		}
+		ReferenceCollection->ParameterReferences.AddUnique(FNiagaraGraphParameterReference(Pin->PersistentGuid, Cast<UNiagaraNode>(Pin->GetOwningNode())));
 
-		HandledPins.Add(Pin);
+		// If we're adding a parameter reference then it needs to be removed from the list of candidate variables to remove since it's been referenced.
+		CandidateUnreferencedParametersToRemove.Remove(Parameter);
 	};
 
 	auto AddStaticParameterReference = [&](const FNiagaraVariable& Variable, UNiagaraNode* Node)
@@ -3654,12 +3747,20 @@ void UNiagaraGraph::RefreshParameterReferences() const
 			const FNiagaraVariable& Parameter = History.VariablesWithOriginalAliasesIntact[Index];
 			for (const FModuleScopedPin& WriteEvent : History.PerVariableWriteHistory[Index])
 			{
-				AddParameterReference(Parameter, WriteEvent.Pin);
+				if (WriteEvent.Pin->PinType.PinSubCategory == UNiagaraNodeParameterMapBase::ParameterPinSubCategory)
+				{
+					AddParameterReference(Parameter, WriteEvent.Pin);
+					HandledParameterMapPins.Add(WriteEvent.Pin);
+				}
 			}
 
 			for (const FNiagaraParameterMapHistory::FReadHistory& ReadHistory : History.PerVariableReadHistory[Index])
 			{
-				AddParameterReference(Parameter, ReadHistory.ReadPin.Pin);
+				if (ReadHistory.ReadPin.Pin->PinType.PinSubCategory == UNiagaraNodeParameterMapBase::ParameterPinSubCategory)
+				{
+					AddParameterReference(Parameter, ReadHistory.ReadPin.Pin);
+					HandledParameterMapPins.Add(ReadHistory.ReadPin.Pin);
+				}
 			}
 		}
 	}
@@ -3669,16 +3770,13 @@ void UNiagaraGraph::RefreshParameterReferences() const
 	const UEdGraphSchema_Niagara* NiagaraSchema = GetNiagaraSchema();
 	for (UEdGraphNode* Node : Nodes)
 	{
-		if (Node->IsA<UNiagaraNodeReroute>())
+		if (UNiagaraNodeStaticSwitch* SwitchNode = Cast<UNiagaraNodeStaticSwitch>(Node))
 		{
-			continue;
-		}
-
-		UNiagaraNodeStaticSwitch* SwitchNode = Cast<UNiagaraNodeStaticSwitch>(Node);
-		if (SwitchNode && !SwitchNode->IsSetByCompiler() && !SwitchNode->IsSetByPin())
-		{
-			FNiagaraVariable Variable(SwitchNode->GetInputType(), SwitchNode->InputParameterName);
-			AddStaticParameterReference(Variable, SwitchNode);
+			if (!SwitchNode->IsSetByCompiler() && !SwitchNode->IsSetByPin())
+			{
+				FNiagaraVariable Variable(SwitchNode->GetInputType(), SwitchNode->InputParameterName);
+				AddStaticParameterReference(Variable, SwitchNode);
+			}
 		}
 		else if (UNiagaraNodeFunctionCall* FunctionNode = Cast<UNiagaraNodeFunctionCall>(Node))
 		{
@@ -3687,13 +3785,22 @@ void UNiagaraGraph::RefreshParameterReferences() const
 				AddStaticParameterReference(Propagated.ToVariable(), FunctionNode);
 			}
 		}
-
-		for (UEdGraphPin* Pin : Node->Pins)
+		else if (UNiagaraNodeParameterMapBase* ParameterMapNode = Cast<UNiagaraNodeParameterMapBase>(Node))
 		{
-			if (HandledPins.Contains(Pin) == false)
+			for (UEdGraphPin* Pin : Node->Pins)
 			{
-				const FNiagaraVariable Parameter = NiagaraSchema->PinToNiagaraVariable(Pin, false);
-				AddParameterReference(Parameter, Pin);
+				// we only consider output parameter pins
+				if ((Pin->Direction == EGPD_Output) && (Pin->PinType.PinSubCategory == UNiagaraNodeParameterMapBase::ParameterPinSubCategory))
+				{
+					bool bAlreadyHandledPin = false;
+					HandledParameterMapPins.FindOrAdd(Pin, &bAlreadyHandledPin);
+
+					if (!bAlreadyHandledPin)
+					{
+						const FNiagaraVariable Parameter = NiagaraSchema->PinToNiagaraVariable(Pin, false);
+						AddParameterReference(Parameter, Pin);
+					}
+				}
 			}
 		}
 	}
diff --git a/Engine/Plugins/FX/Niagara/Source/NiagaraEditor/Private/NiagaraNode.cpp b/Engine/Plugins/FX/Niagara/Source/NiagaraEditor/Private/NiagaraNode.cpp
--- a/Engine/Plugins/FX/Niagara/Source/NiagaraEditor/Private/NiagaraNode.cpp
+++ b/Engine/Plugins/FX/Niagara/Source/NiagaraEditor/Private/NiagaraNode.cpp
@@ -590,10 +590,12 @@ void UNiagaraNode::MarkNodeRequiresSynchronization(FString Reason, bool bRaiseGr
 	ChangeId = FGuid::NewGuid();
 	//UE_LOG(LogNiagaraEditor, Verbose, TEXT("Node %s was marked requires synchronization.  Reason: %s"), *GetPathName(), *Reason);
 
-	UNiagaraGraph* Graph = Cast<UNiagaraGraph>(GetGraph());
-	if (Graph != nullptr && bRaiseGraphNeedsRecompile)
+	if (bRaiseGraphNeedsRecompile)
 	{
-		Graph->NotifyGraphNeedsRecompile();
+		if (UNiagaraGraph* Graph = Cast<UNiagaraGraph>(GetGraph()))
+		{
+			Graph->NotifyGraphNeedsRecompile();
+		}
 	}
 
 	VisualsChangedDelegate.Broadcast(this);
diff --git a/Engine/Plugins/FX/Niagara/Source/NiagaraEditor/Private/NiagaraScriptSource.cpp b/Engine/Plugins/FX/Niagara/Source/NiagaraEditor/Private/NiagaraScriptSource.cpp
--- a/Engine/Plugins/FX/Niagara/Source/NiagaraEditor/Private/NiagaraScriptSource.cpp
+++ b/Engine/Plugins/FX/Niagara/Source/NiagaraEditor/Private/NiagaraScriptSource.cpp
@@ -31,7 +31,7 @@ UNiagaraScriptSource::UNiagaraScriptSource(const FObjectInitializer& ObjectIniti
 {
 }
 
-void UNiagaraScriptSource::ComputeVMCompilationId(FNiagaraVMExecutableDataId& Id, ENiagaraScriptUsage InUsage, const FGuid& InUsageId, bool bForceRebuild) const
+void UNiagaraScriptSource::ComputeVMCompilationId(FNiagaraVMExecutableDataId& Id, ENiagaraScriptUsage InUsage, const FGuid& InUsageId) const
 {
 	if (!AllowShaderCompiling())
 	{
@@ -43,7 +43,7 @@ void UNiagaraScriptSource::ComputeVMCompilationId(FNiagaraVMExecutableDataId& Id
 	Id.CompilerVersionID = FNiagaraCustomVersion::GetLatestScriptCompileVersion();
 	if (NodeGraph)
 	{
-		NodeGraph->RebuildCachedCompileIds(bForceRebuild);
+		NodeGraph->RebuildCachedCompileIds();
 		Id.BaseScriptCompileHash = FNiagaraCompileHash(NodeGraph->GetCompileDataHash(InUsage, InUsageId));
 		NodeGraph->GatherExternalDependencyData(InUsage, InUsageId, Id.ReferencedCompileHashes, Id.DebugReferencedObjects);
 	}
@@ -60,6 +60,14 @@ void UNiagaraScriptSource::ComputeVMCompilationId(FNiagaraVMExecutableDataId& Id
 	Id.DebugReferencedObjects.Emplace(TEXT("/Engine/Public/ShaderVersion.ush"));
 }
 
+void UNiagaraScriptSource::RefreshGraphCompileId()
+{
+	if (NodeGraph)
+	{
+		NodeGraph->ConditionalRebuildCompileIdCache();
+	}
+}
+
 FGuid UNiagaraScriptSource::GetCompileBaseId(ENiagaraScriptUsage InUsage, const FGuid& InUsageId) const
 {
 	return NodeGraph->GetBaseId(InUsage, InUsageId);
diff --git a/Engine/Plugins/FX/Niagara/Source/NiagaraEditor/Public/NiagaraGraph.h b/Engine/Plugins/FX/Niagara/Source/NiagaraEditor/Public/NiagaraGraph.h
--- a/Engine/Plugins/FX/Niagara/Source/NiagaraEditor/Public/NiagaraGraph.h
+++ b/Engine/Plugins/FX/Niagara/Source/NiagaraEditor/Public/NiagaraGraph.h
@@ -356,7 +356,12 @@ class UNiagaraGraph : public UEdGraph
 	  */
 	FString GetFunctionAliasByContext(const FNiagaraGraphFunctionAliasContext& FunctionAliasContext);
 
-	void RebuildCachedCompileIds(bool bForce = false);
+	/** In order to support reducing the number of times we need to fully generate the CompileId this function is introduced
+	  * to work with LastBuiltScriptVersionId & bHasValidLastBuiltScriptVersionId.
+	  */
+	void ConditionalRebuildCompileIdCache();
+
+	void RebuildCachedCompileIds();
 
 	void CopyCachedReferencesMap(UNiagaraGraph* TargetGraph);
 
@@ -449,6 +454,8 @@ private:
 
 	void FixupReferenceCollectionsPostRename(const FNiagaraVariable& OrigVariable, const FNiagaraVariable& DestVariable, bool bParameterMerged, bool bSuppressEvents);
 
+	void PostLoad_LWCFixup(int32 NiagaraVersion);
+
 private:
 	/** The current change identifier for this graph overall. Used to sync status with UNiagaraScripts.*/
 	UPROPERTY()
@@ -461,6 +468,13 @@ private:
 	UPROPERTY()
 	FGuid LastBuiltTraversalDataChangeId;
 
+	/** The script version that was used when the cached CompileId was generated, a change
+	 *	in script version will invalidate the cached CompileId and a new one will be generated.
+	 *	Will be initialized to an invalid guid but it won't be used until a valid script has been
+	 *	assigned (as dictated by bHasValidLastBuiltScriptVersionId) */
+	UPROPERTY()
+	FGuid LastBuiltScriptVersionId;
+
 	UPROPERTY()
 	TArray<FNiagaraGraphScriptUsageInfo> CachedUsageInfo;
 
@@ -479,6 +493,11 @@ private:
 	UPROPERTY(Transient)
 	mutable TArray<FNiagaraScriptVariableData> CompilationScriptVariables;
 
+	/** Used in conjunction with LastBuiltScriptVersionId to note that we've got a valid script Id stored.
+	 *	Works around things without having to add a custom version. */
+	UPROPERTY()
+	bool bHasValidLastBuiltScriptVersionId = false;
+
 	FOnDataInterfaceChanged OnDataInterfaceChangedDelegate;
 	FOnSubObjectSelectionChanged OnSelectedSubObjectChanged;
 	FOnGetParameterDefinitionsForDetailsCustomization OnGetParameterDefinitionsForDetailsCustomizationDelegate;
diff --git a/Engine/Plugins/FX/Niagara/Source/NiagaraEditor/Public/NiagaraScriptSource.h b/Engine/Plugins/FX/Niagara/Source/NiagaraEditor/Public/NiagaraScriptSource.h
--- a/Engine/Plugins/FX/Niagara/Source/NiagaraEditor/Public/NiagaraScriptSource.h
+++ b/Engine/Plugins/FX/Niagara/Source/NiagaraEditor/Public/NiagaraScriptSource.h
@@ -38,7 +38,10 @@ class UNiagaraScriptSource : public UNiagaraScriptSourceBase
 	virtual FGuid GetChangeID() override;
 	FVersionedNiagaraEmitter GetOuterEmitter() const;
 
-	virtual void ComputeVMCompilationId(struct FNiagaraVMExecutableDataId& Id, ENiagaraScriptUsage InUsage, const FGuid& InUsageId, bool bForceRebuild = false) const override;
+	virtual void ComputeVMCompilationId(struct FNiagaraVMExecutableDataId& Id, ENiagaraScriptUsage InUsage, const FGuid& InUsageId) const override;
+
+	// Will conditionally refresh the graph's CompileId
+	virtual void RefreshGraphCompileId() override;
 
 	virtual FGuid GetCompileBaseId(ENiagaraScriptUsage InUsage, const FGuid& InUsageId) const override;
 
