Niagara - When removing a parameter, only remove position data by name if a position parameter type is actually being removed.  Fixes a rare case where a parameter with the same name as a position parameter but with a different type is removed from the parameter store breaking the position data.

diff --git a/Engine/Plugins/FX/Niagara/Source/Niagara/Private/NiagaraParameterStore.cpp b/Engine/Plugins/FX/Niagara/Source/Niagara/Private/NiagaraParameterStore.cpp
--- a/Engine/Plugins/FX/Niagara/Source/Niagara/Private/NiagaraParameterStore.cpp
+++ b/Engine/Plugins/FX/Niagara/Source/Niagara/Private/NiagaraParameterStore.cpp
@@ -758,7 +758,10 @@ bool FNiagaraParameterStore::RemoveParameter(const FNiagaraVariableBase& ToRemov
 	}
 	//check(!ParameterOffsets.Num()); // Migration to SortedParameterOffsets
 #endif
-	RemovePositionData(ToRemove.GetName());
+	if (ToRemove.GetType() == FNiagaraTypeDefinition::GetPositionDef())
+	{
+		RemovePositionData(ToRemove.GetName());
+	}
 	
 	if (IndexOf(ToRemove) != INDEX_NONE)
 	{
