better error reporting
trap case where grid3d collection reader is used with a grid that is not rgba, but the reader thinks it is rgba

diff --git a/Engine/Plugins/FX/Niagara/Source/Niagara/Classes/NiagaraDataInterfaceGrid3DCollection.h b/Engine/Plugins/FX/Niagara/Source/Niagara/Classes/NiagaraDataInterfaceGrid3DCollection.h
--- a/Engine/Plugins/FX/Niagara/Source/Niagara/Classes/NiagaraDataInterfaceGrid3DCollection.h
+++ b/Engine/Plugins/FX/Niagara/Source/Niagara/Classes/NiagaraDataInterfaceGrid3DCollection.h
@@ -193,13 +193,10 @@ public:
 	virtual bool GetExposedVariableValue(const FNiagaraVariableBase& InVariable, void* InPerInstanceData, FNiagaraSystemInstance* InSystemInstance, void* OutData) const override;
 	//~ UNiagaraDataInterface interface END
 
-private:
-	static void CollectAttributesForScript(UNiagaraScript* Script, FName VariableName, TArray<FNiagaraVariableBase>& OutVariables, TArray<uint32>& OutVariableOffsets, int32& TotalAttributes, TArray<FText>* OutWarnings = nullptr);
 public:
 	/** Finds all attributes by locating the variable name inside the parameter stores. */
 	void FindAttributesByName(FName DataInterfaceName, TArray<FNiagaraVariableBase>& OutVariables, TArray<uint32>& OutVariableOffsets, int32& OutNumAttribChannelsFound, TArray<FText>* OutWarnings = nullptr) const;
-	/** Finds all attributes by locating the data interface amongst the parameter stores. */
-	void FindAttributes(TArray<FNiagaraVariableBase>& OutVariables, TArray<uint32>& OutVariableOffsets, int32& OutNumAttribChannelsFound, TArray<FText>* OutWarnings = nullptr) const;
+
 	// Fills a texture render target 2d with the current data from the simulation
 	// #todo(dmp): this will eventually go away when we formalize how data makes it out of Niagara
 	// #todo(dmp): reimplement for 3d
@@ -344,5 +341,104 @@ protected:
 	TMap<FNiagaraSystemInstanceID, FGrid3DCollectionRWInstanceData_GameThread*> SystemInstancesToProxyData_GT;
 
 	static FNiagaraVariableBase ExposedRTVar;
+
+public:
+	/** Finds all attributes by locating the data interface amongst the parameter stores. */
+	template< typename T = UNiagaraDataInterfaceGrid3DCollection>
+	void FindAttributes(TArray<FNiagaraVariableBase>& OutVariables, TArray<uint32>& OutVariableOffsets, int32& OutNumAttribChannelsFound, TArray<FText>* OutWarnings = nullptr) const
+	{
+		{
+			OutNumAttribChannelsFound = 0;
+
+			UNiagaraSystem* OwnerSystem = GetTypedOuter<UNiagaraSystem>();
+			if (OwnerSystem == nullptr)
+			{
+				return;
+			}
+
+			int32 TotalAttributes = NumAttributes;
+			for (const FNiagaraEmitterHandle& EmitterHandle : OwnerSystem->GetEmitterHandles())
+			{
+				FVersionedNiagaraEmitterData* EmitterData = EmitterHandle.GetEmitterData();
+				if (EmitterData && EmitterHandle.GetIsEnabled() && EmitterData->IsValid() && (EmitterData->SimTarget == ENiagaraSimTarget::GPUComputeSim))
+				{
+					// Search scripts for this data interface so we get the variable name
+					auto FindDataInterfaceVariable =
+						[&OwnerSystem, &EmitterData](const UNiagaraDataInterface* DataInterface) -> FName
+					{
+						UNiagaraScript* Scripts[] =
+						{
+							OwnerSystem->GetSystemSpawnScript(),
+							OwnerSystem->GetSystemUpdateScript(),
+							EmitterData->GetGPUComputeScript(),
+						};
+
+						for (UNiagaraScript* Script : Scripts)
+						{
+							for (FNiagaraScriptDataInterfaceInfo& DataInterfaceInfo : Script->GetCachedDefaultDataInterfaces())
+							{
+								if (DataInterfaceInfo.DataInterface == DataInterface)
+								{
+									return DataInterfaceInfo.RegisteredParameterMapRead.IsNone() ? DataInterfaceInfo.RegisteredParameterMapWrite : DataInterfaceInfo.RegisteredParameterMapRead;
+								}
+							}
+						}
+						return NAME_None;
+					};
+
+					const FName VariableName = FindDataInterfaceVariable(this);
+					if (!VariableName.IsNone())
+					{
+						CollectAttributesForScript<T>(EmitterData->GetGPUComputeScript(), VariableName, OutVariables, OutVariableOffsets, TotalAttributes, OutWarnings);
+					}
+				}
+			}
+			OutNumAttribChannelsFound = TotalAttributes - NumAttributes;
+		}
+	}
+protected:
+	template< typename T>
+	static void CollectAttributesForScript(UNiagaraScript* Script, FName VariableName, TArray<FNiagaraVariableBase>& OutVariables, TArray<uint32>& OutVariableOffsets, int32& TotalAttributes, TArray<FText>* OutWarnings = nullptr)
+	{
+		if (const FNiagaraScriptExecutionParameterStore* ParameterStore = Script->GetExecutionReadyParameterStore(ENiagaraSimTarget::GPUComputeSim))
+		{
+			const FNiagaraVariableBase DataInterfaceVariable(FNiagaraTypeDefinition(T::StaticClass()), VariableName);
+
+			const int32* IndexOfDataInterface = ParameterStore->FindParameterOffset(DataInterfaceVariable);
+			if (IndexOfDataInterface != nullptr)
+			{
+				TConstArrayView<FNiagaraDataInterfaceGPUParamInfo> ParamInfoArray = Script->GetDataInterfaceGPUParamInfos();
+				for (const FNiagaraDataInterfaceGeneratedFunction& Func : ParamInfoArray[*IndexOfDataInterface].GeneratedFunctions)
+				{
+					if (const FName* AttributeName = Func.FindSpecifierValue(NAME_Attribute))
+					{
+						FNiagaraVariableBase NewVar(UNiagaraDataInterfaceGrid3DCollection::GetValueTypeFromFuncName(Func.DefinitionName), *AttributeName);
+						if (UNiagaraDataInterfaceGrid3DCollection::CanCreateVarFromFuncName(Func.DefinitionName))
+						{
+							if (!OutVariables.Contains(NewVar))
+							{
+								const int32 FoundNameMatch = OutVariables.IndexOfByPredicate([&](const FNiagaraVariableBase& Var) { return Var.GetName() == *AttributeName; });
+								if (FoundNameMatch == INDEX_NONE)
+								{
+									OutVariables.Add(NewVar);
+									const int32 NumComponents = NewVar.GetSizeInBytes() / sizeof(float);
+									OutVariableOffsets.Add(TotalAttributes);
+									TotalAttributes += NumComponents;
+								}
+								else
+								{
+									if (OutWarnings)
+									{
+										FText Warning = FText::Format(NSLOCTEXT("NiagaraDataInterfaceGrid3DCollection", "BadType", "Same name, different types! {0} vs {1}, Attribute {2}"), NewVar.GetType().GetNameText(), OutVariables[FoundNameMatch].GetType().GetNameText(), FText::FromName(NewVar.GetName()));
+										OutWarnings->Add(Warning);
+									}
+								}
+							}
+						}
+					}
+				}
+			}
+		}
+	}
 };
 
diff --git a/Engine/Plugins/FX/Niagara/Source/Niagara/Private/NiagaraDataInterfaceGrid3DCollection.cpp b/Engine/Plugins/FX/Niagara/Source/Niagara/Private/NiagaraDataInterfaceGrid3DCollection.cpp
--- a/Engine/Plugins/FX/Niagara/Source/Niagara/Private/NiagaraDataInterfaceGrid3DCollection.cpp
+++ b/Engine/Plugins/FX/Niagara/Source/Niagara/Private/NiagaraDataInterfaceGrid3DCollection.cpp
@@ -3568,48 +3568,7 @@ bool UNiagaraDataInterfaceGrid3DCollection::GetExposedVariableValue(const FNiaga
 	return false;
 }
 
-void UNiagaraDataInterfaceGrid3DCollection::CollectAttributesForScript(UNiagaraScript* Script, FName VariableName, TArray<FNiagaraVariableBase>& OutVariables, TArray<uint32>& OutVariableOffsets, int32& TotalAttributes, TArray<FText>* OutWarnings)
-{
-	if (const FNiagaraScriptExecutionParameterStore* ParameterStore = Script->GetExecutionReadyParameterStore(ENiagaraSimTarget::GPUComputeSim))
-	{
-		const FNiagaraVariableBase DataInterfaceVariable(FNiagaraTypeDefinition(UNiagaraDataInterfaceGrid3DCollection::StaticClass()), VariableName);
 
-		const int32* IndexOfDataInterface = ParameterStore->FindParameterOffset(DataInterfaceVariable);
-		if (IndexOfDataInterface != nullptr)
-		{
-			TConstArrayView<FNiagaraDataInterfaceGPUParamInfo> ParamInfoArray = Script->GetDataInterfaceGPUParamInfos();
-			for (const FNiagaraDataInterfaceGeneratedFunction& Func : ParamInfoArray[*IndexOfDataInterface].GeneratedFunctions)
-			{
-				if (const FName* AttributeName = Func.FindSpecifierValue(NAME_Attribute))
-				{
-					FNiagaraVariableBase NewVar(UNiagaraDataInterfaceGrid3DCollection::GetValueTypeFromFuncName(Func.DefinitionName), *AttributeName);
-					if (UNiagaraDataInterfaceGrid3DCollection::CanCreateVarFromFuncName(Func.DefinitionName))
-					{
-						if (!OutVariables.Contains(NewVar))
-						{
-							const int32 FoundNameMatch = OutVariables.IndexOfByPredicate([&](const FNiagaraVariableBase& Var) { return Var.GetName() == *AttributeName; });
-							if (FoundNameMatch == INDEX_NONE)
-							{
-								OutVariables.Add(NewVar);
-								const int32 NumComponents = NewVar.GetSizeInBytes() / sizeof(float);
-								OutVariableOffsets.Add(TotalAttributes);
-								TotalAttributes += NumComponents;
-							}
-							else
-							{
-								if (OutWarnings)
-								{
-									FText Warning = FText::Format(LOCTEXT("BadType", "Same name, different types! {0} vs {1}, Attribute {2}"), NewVar.GetType().GetNameText(), OutVariables[FoundNameMatch].GetType().GetNameText(), FText::FromName(NewVar.GetName()));
-									OutWarnings->Add(Warning);
-								}
-							}
-						}
-					}
-				}
-			}
-		}
-	}
-}
 
 void UNiagaraDataInterfaceGrid3DCollection::FindAttributesByName(FName VariableName, TArray<FNiagaraVariableBase>& OutVariables, TArray<uint32>& OutVariableOffsets, int32& OutNumAttribChannelsFound, TArray<FText>* OutWarnings) const
 {
@@ -3627,57 +3586,7 @@ void UNiagaraDataInterfaceGrid3DCollection::FindAttributesByName(FName VariableN
 		FVersionedNiagaraEmitterData* EmitterData = EmitterHandle.GetEmitterData();
 		if (EmitterData && EmitterHandle.GetIsEnabled() && EmitterData->IsValid() && (EmitterData->SimTarget == ENiagaraSimTarget::GPUComputeSim))
 		{
-			CollectAttributesForScript(EmitterData->GetGPUComputeScript(), VariableName, OutVariables, OutVariableOffsets, TotalAttributes, OutWarnings);
-		}
-	}
-	OutNumAttribChannelsFound = TotalAttributes - NumAttributes;
-}
-
-void UNiagaraDataInterfaceGrid3DCollection::FindAttributes(TArray<FNiagaraVariableBase>& OutVariables, TArray<uint32>& OutVariableOffsets, int32& OutNumAttribChannelsFound, TArray<FText>* OutWarnings) const
-{
-	OutNumAttribChannelsFound = 0;
-
-	UNiagaraSystem* OwnerSystem = GetTypedOuter<UNiagaraSystem>();
-	if (OwnerSystem == nullptr)
-	{
-		return;
-	}
-
-	int32 TotalAttributes = NumAttributes;
-	for (const FNiagaraEmitterHandle& EmitterHandle : OwnerSystem->GetEmitterHandles())
-	{
-		FVersionedNiagaraEmitterData* EmitterData = EmitterHandle.GetEmitterData();
-		if (EmitterData && EmitterHandle.GetIsEnabled() && EmitterData->IsValid() && (EmitterData->SimTarget == ENiagaraSimTarget::GPUComputeSim))
-		{
-			// Search scripts for this data interface so we get the variable name
-			auto FindDataInterfaceVariable =
-				[&OwnerSystem, &EmitterData](const UNiagaraDataInterface* DataInterface) -> FName
-			{
-				UNiagaraScript* Scripts[] =
-				{
-					OwnerSystem->GetSystemSpawnScript(),
-					OwnerSystem->GetSystemUpdateScript(),
-					EmitterData->GetGPUComputeScript(),
-				};
-
-				for (UNiagaraScript* Script : Scripts)
-				{
-					for (FNiagaraScriptDataInterfaceInfo& DataInterfaceInfo : Script->GetCachedDefaultDataInterfaces())
-					{
-						if (DataInterfaceInfo.DataInterface == DataInterface)
-						{
-							return DataInterfaceInfo.RegisteredParameterMapRead.IsNone() ? DataInterfaceInfo.RegisteredParameterMapWrite : DataInterfaceInfo.RegisteredParameterMapRead;
-						}
-					}
-				}
-				return NAME_None;
-			};
-
-			const FName VariableName = FindDataInterfaceVariable(this);
-			if (!VariableName.IsNone())
-			{
-				CollectAttributesForScript(EmitterData->GetGPUComputeScript(), VariableName, OutVariables, OutVariableOffsets, TotalAttributes, OutWarnings);
-			}
+			CollectAttributesForScript<UNiagaraDataInterfaceGrid3DCollection>(EmitterData->GetGPUComputeScript(), VariableName, OutVariables, OutVariableOffsets, TotalAttributes, OutWarnings);
 		}
 	}
 	OutNumAttribChannelsFound = TotalAttributes - NumAttributes;
diff --git a/Engine/Plugins/FX/Niagara/Source/Niagara/Private/NiagaraDataInterfaceGrid3DCollectionReader.cpp b/Engine/Plugins/FX/Niagara/Source/Niagara/Private/NiagaraDataInterfaceGrid3DCollectionReader.cpp
--- a/Engine/Plugins/FX/Niagara/Source/Niagara/Private/NiagaraDataInterfaceGrid3DCollectionReader.cpp
+++ b/Engine/Plugins/FX/Niagara/Source/Niagara/Private/NiagaraDataInterfaceGrid3DCollectionReader.cpp
@@ -8,6 +8,7 @@
 #include "NiagaraSettings.h"
 #include "NiagaraConstants.h"
 #include "NiagaraComputeExecutionContext.h"
+#include "NiagaraDataInterfaceGrid3DCollectionUtils.h"
 
 #include UE_INLINE_GENERATED_CPP_BY_NAME(NiagaraDataInterfaceGrid3DCollectionReader)
 
@@ -104,11 +105,28 @@ bool UNiagaraDataInterfaceGrid3DCollectionReader::InitPerInstanceData(void* PerI
 
 		if (GridInterface && ProxyToUse)
 		{
+			// Does the reader think it should be an rgba grid?
+			int32 NumAttribChannelsFound = 0;
+			int32 NumNamedAttribChannelsFound = 0;
+			TArray<FNiagaraVariableBase> Vars;
+			TArray<uint32> Offsets;
+			FindAttributes<UNiagaraDataInterfaceGrid3DCollectionReader>(Vars, Offsets, NumNamedAttribChannelsFound);
+
+			// #todo(dmp): slight hack here - these aren't always the same, but we generally can't use unnamed attrs very well
+			NumAttribChannelsFound = NumNamedAttribChannelsFound;
+			bool UseRGBATexture = FGrid3DCollectionAttributeHelper::ShouldUseRGBAGrid(NumAttribChannelsFound, Vars.Num());
 
 			FGrid3DCollectionRWInstanceData_GameThread* Grid3DInstanceData = GridInterface->GetSystemInstancesToProxyData_GT().FindRef(SystemInstance->GetId());
 			InstanceData->OtherDI = GridInterface;
 			InstanceData->OtherInstanceData = Grid3DInstanceData;
 
+			// error if the reader thinks it should be an rgba grid and the underlying grid collection disagrees.  This is due to the reader not using all attributes defined
+			if (NumAttribChannelsFound > 0 && UseRGBATexture != InstanceData->OtherInstanceData->UseRGBATexture)
+			{
+				UE_LOG(LogNiagara, Error, TEXT("RGBA grids and Grid3DCollectionReaders are not compatible with Emitter %s and DataInterface %s.  Turn off RGBA grids with fx.Niagara.Grid3D.UseRGBAGrid 0, or split up your grid collections so the one passed to this reader has only one attribute."), *this->EmitterName, *this->DIName);
+				return false;
+			}
+
 			// Push Updates to Proxy.
 			FNiagaraDataInterfaceProxyGrid3DCollectionProxy* RT_Proxy = GetProxyAs<FNiagaraDataInterfaceProxyGrid3DCollectionProxy>();
 			ENQUEUE_RENDER_COMMAND(FUpdateData)(
@@ -123,6 +141,14 @@ bool UNiagaraDataInterfaceGrid3DCollectionReader::InitPerInstanceData(void* PerI
 			});
 			return true;
 		}
+		else
+		{
+			UE_LOG(LogNiagara, Error, TEXT("Emitter %s doesn't have Grid3D Collection DI named %s for DI: %s"), *this->EmitterName, *this->DIName, *Proxy->SourceDIName.ToString());
+		}
+	}
+	else
+	{
+		UE_LOG(LogNiagara, Error, TEXT("Emitter %s doesn't exist for DI: %s"), *this->EmitterName, *Proxy->SourceDIName.ToString());
 	}
 
 	return false;
