cis fixes

diff --git a/Engine/Plugins/FX/Niagara/Source/Niagara/Classes/NiagaraDataInterfaceGrid3DCollection.h b/Engine/Plugins/FX/Niagara/Source/Niagara/Classes/NiagaraDataInterfaceGrid3DCollection.h
--- a/Engine/Plugins/FX/Niagara/Source/Niagara/Classes/NiagaraDataInterfaceGrid3DCollection.h
+++ b/Engine/Plugins/FX/Niagara/Source/Niagara/Classes/NiagaraDataInterfaceGrid3DCollection.h
@@ -141,16 +141,16 @@ public:
 	UPROPERTY(EditAnywhere, Category = "Grid", meta = (PinHiddenByDefault, InlineEditConditionToggle))
 	uint8 bPreviewGrid : 1;
 
-	UPROPERTY(EditAnywhere, Category = "Grid", meta = (EditCondition = "bPreviewGrid", ToolTip = "When enabled allows you to preview the grid in a debug display") )
+	UPROPERTY(EditAnywhere, Category = "Grid", meta = (EditCondition = "bPreviewGrid", ToolTip = "When enabled allows you to preview the grid in a debug display"))
 	FName PreviewAttribute = NAME_None;
 #endif
 
 	virtual void PostInitProperties() override;
-	
+
 	//~ UNiagaraDataInterface interface
 	// VM functionality
 	virtual void GetFunctions(TArray<FNiagaraFunctionSignature>& OutFunctions) override;
-	virtual void GetVMExternalFunction(const FVMExternalFunctionBindingInfo& BindingInfo, void* InstanceData, FVMExternalFunction &OutFunc) override;
+	virtual void GetVMExternalFunction(const FVMExternalFunctionBindingInfo& BindingInfo, void* InstanceData, FVMExternalFunction& OutFunc) override;
 
 #if WITH_EDITOR	
 	virtual void GetFeedback(UNiagaraSystem* Asset, UNiagaraComponent* Component, TArray<FNiagaraDataInterfaceError>& OutErrors,
@@ -202,16 +202,16 @@ public:
 	// #todo(dmp): reimplement for 3d
 
 	UFUNCTION(BlueprintCallable, Category = Niagara, meta = (DeprecatedFunction, DeprecationMessage = "This function has been replaced by object user variables on the emitter to specify render targets to fill with data."))
-	virtual bool FillVolumeTexture(const UNiagaraComponent *Component, UVolumeTexture *dest, int AttributeIndex);
-	
+	virtual bool FillVolumeTexture(const UNiagaraComponent* Component, UVolumeTexture* dest, int AttributeIndex);
+
 	UFUNCTION(BlueprintCallable, Category = Niagara, meta = (DeprecatedFunction, DeprecationMessage = "This function has been replaced by object user variables on the emitter to specify render targets to fill with data."))
-	virtual bool FillRawVolumeTexture(const UNiagaraComponent *Component, UVolumeTexture*Dest, int &TilesX, int &TilesY, int &TileZ);
-	
+	virtual bool FillRawVolumeTexture(const UNiagaraComponent* Component, UVolumeTexture* Dest, int& TilesX, int& TilesY, int& TileZ);
+
 	UFUNCTION(BlueprintCallable, Category = Niagara)
-	virtual void GetRawTextureSize(const UNiagaraComponent *Component, int &SizeX, int &SizeY, int &SizeZ);
+	virtual void GetRawTextureSize(const UNiagaraComponent* Component, int& SizeX, int& SizeY, int& SizeZ);
 
 	UFUNCTION(BlueprintCallable, Category = Niagara)
-	virtual void GetTextureSize(const UNiagaraComponent *Component, int &SizeX, int &SizeY, int &SizeZ);	
+	virtual void GetTextureSize(const UNiagaraComponent* Component, int& SizeX, int& SizeY, int& SizeZ);
 
 	void VMGetWorldBBoxSize(FVectorVMExternalFunctionContext& Context);
 	void VMGetCellSize(FVectorVMExternalFunctionContext& Context);
@@ -257,7 +257,7 @@ public:
 
 	static const FName SetVector4ValueAtIndexFunctionName;
 	static const FName SetVector3ValueAtIndexFunctionName;
-	static const FName SetVector2ValueAtIndexFunctionName;	
+	static const FName SetVector2ValueAtIndexFunctionName;
 	static const FName SetFloatValueAtIndexFunctionName;
 
 	static const FName GetPreviousVector4ValueAtIndexFunctionName;
@@ -344,101 +344,7 @@ protected:
 
 public:
 	/** Finds all attributes by locating the data interface amongst the parameter stores. */
-	template< typename T = UNiagaraDataInterfaceGrid3DCollection>
-	void FindAttributes(TArray<FNiagaraVariableBase>& OutVariables, TArray<uint32>& OutVariableOffsets, int32& OutNumAttribChannelsFound, TArray<FText>* OutWarnings = nullptr) const
-	{
-		{
-			OutNumAttribChannelsFound = 0;
-
-			UNiagaraSystem* OwnerSystem = GetTypedOuter<UNiagaraSystem>();
-			if (OwnerSystem == nullptr)
-			{
-				return;
-			}
-
-			int32 TotalAttributes = NumAttributes;
-			for (const FNiagaraEmitterHandle& EmitterHandle : OwnerSystem->GetEmitterHandles())
-			{
-				FVersionedNiagaraEmitterData* EmitterData = EmitterHandle.GetEmitterData();
-				if (EmitterData && EmitterHandle.GetIsEnabled() && EmitterData->IsValid() && (EmitterData->SimTarget == ENiagaraSimTarget::GPUComputeSim))
-				{
-					// Search scripts for this data interface so we get the variable name
-					auto FindDataInterfaceVariable =
-						[&OwnerSystem, &EmitterData](const UNiagaraDataInterface* DataInterface) -> FName
-					{
-						UNiagaraScript* Scripts[] =
-						{
-							OwnerSystem->GetSystemSpawnScript(),
-							OwnerSystem->GetSystemUpdateScript(),
-							EmitterData->GetGPUComputeScript(),
-						};
-
-						for (UNiagaraScript* Script : Scripts)
-						{
-							for (FNiagaraScriptDataInterfaceInfo& DataInterfaceInfo : Script->GetCachedDefaultDataInterfaces())
-							{
-								if (DataInterfaceInfo.DataInterface == DataInterface)
-								{
-									return DataInterfaceInfo.RegisteredParameterMapRead.IsNone() ? DataInterfaceInfo.RegisteredParameterMapWrite : DataInterfaceInfo.RegisteredParameterMapRead;
-								}
-							}
-						}
-						return NAME_None;
-					};
-
-					const FName VariableName = FindDataInterfaceVariable(this);
-					if (!VariableName.IsNone())
-					{
-						CollectAttributesForScript<T>(EmitterData->GetGPUComputeScript(), VariableName, OutVariables, OutVariableOffsets, TotalAttributes, OutWarnings);
-					}
-				}
-			}
-			OutNumAttribChannelsFound = TotalAttributes - NumAttributes;
-		}
-	}
+	void FindAttributes(TArray<FNiagaraVariableBase>& OutVariables, TArray<uint32>& OutVariableOffsets, int32& OutNumAttribChannelsFound, TArray<FText>* OutWarnings = nullptr, bool UseReader = false) const;
 protected:
-	template< typename T>
-	static void CollectAttributesForScript(UNiagaraScript* Script, FName VariableName, TArray<FNiagaraVariableBase>& OutVariables, TArray<uint32>& OutVariableOffsets, int32& TotalAttributes, TArray<FText>* OutWarnings = nullptr)
-	{
-		if (const FNiagaraScriptExecutionParameterStore* ParameterStore = Script->GetExecutionReadyParameterStore(ENiagaraSimTarget::GPUComputeSim))
-		{
-			const FNiagaraVariableBase DataInterfaceVariable(FNiagaraTypeDefinition(T::StaticClass()), VariableName);
-
-			const int32* IndexOfDataInterface = ParameterStore->FindParameterOffset(DataInterfaceVariable);
-			if (IndexOfDataInterface != nullptr)
-			{
-				TConstArrayView<FNiagaraDataInterfaceGPUParamInfo> ParamInfoArray = Script->GetDataInterfaceGPUParamInfos();
-				for (const FNiagaraDataInterfaceGeneratedFunction& Func : ParamInfoArray[*IndexOfDataInterface].GeneratedFunctions)
-				{
-					if (const FName* AttributeName = Func.FindSpecifierValue(NAME_Attribute))
-					{
-						FNiagaraVariableBase NewVar(UNiagaraDataInterfaceGrid3DCollection::GetValueTypeFromFuncName(Func.DefinitionName), *AttributeName);
-						if (UNiagaraDataInterfaceGrid3DCollection::CanCreateVarFromFuncName(Func.DefinitionName))
-						{
-							if (!OutVariables.Contains(NewVar))
-							{
-								const int32 FoundNameMatch = OutVariables.IndexOfByPredicate([&](const FNiagaraVariableBase& Var) { return Var.GetName() == *AttributeName; });
-								if (FoundNameMatch == INDEX_NONE)
-								{
-									OutVariables.Add(NewVar);
-									const int32 NumComponents = NewVar.GetSizeInBytes() / sizeof(float);
-									OutVariableOffsets.Add(TotalAttributes);
-									TotalAttributes += NumComponents;
-								}
-								else
-								{
-									if (OutWarnings)
-									{
-										FText Warning = FText::Format(NSLOCTEXT("NiagaraDataInterfaceGrid3DCollection", "BadType", "Same name, different types! {0} vs {1}, Attribute {2}"), NewVar.GetType().GetNameText(), OutVariables[FoundNameMatch].GetType().GetNameText(), FText::FromName(NewVar.GetName()));
-										OutWarnings->Add(Warning);
-									}
-								}
-							}
-						}
-					}
-				}
-			}
-		}
-	}
-};
-
+	static void CollectAttributesForScript(UNiagaraScript* Script, FName VariableName, TArray<FNiagaraVariableBase>& OutVariables, TArray<uint32>& OutVariableOffsets, int32& TotalAttributes, TArray<FText>* OutWarnings = nullptr, bool UseReader = false);
+};
\ No newline at end of file
diff --git a/Engine/Plugins/FX/Niagara/Source/Niagara/Private/NiagaraDataInterfaceGrid3DCollection.cpp b/Engine/Plugins/FX/Niagara/Source/Niagara/Private/NiagaraDataInterfaceGrid3DCollection.cpp
--- a/Engine/Plugins/FX/Niagara/Source/Niagara/Private/NiagaraDataInterfaceGrid3DCollection.cpp
+++ b/Engine/Plugins/FX/Niagara/Source/Niagara/Private/NiagaraDataInterfaceGrid3DCollection.cpp
@@ -18,6 +18,8 @@
 #include "ShaderParameterUtils.h"
 #include "TextureResource.h"
 
+#include "NiagaraDataInterfaceGrid3DCollectionReader.h"
+
 #include UE_INLINE_GENERATED_CPP_BY_NAME(NiagaraDataInterfaceGrid3DCollection)
 
 #define LOCTEXT_NAMESPACE "NiagaraDataInterfaceGrid3DCollection"
@@ -3586,7 +3588,7 @@ void UNiagaraDataInterfaceGrid3DCollection::FindAttributesByName(FName VariableN
 		FVersionedNiagaraEmitterData* EmitterData = EmitterHandle.GetEmitterData();
 		if (EmitterData && EmitterHandle.GetIsEnabled() && EmitterData->IsValid() && (EmitterData->SimTarget == ENiagaraSimTarget::GPUComputeSim))
 		{
-			CollectAttributesForScript<UNiagaraDataInterfaceGrid3DCollection>(EmitterData->GetGPUComputeScript(), VariableName, OutVariables, OutVariableOffsets, TotalAttributes, OutWarnings);
+			CollectAttributesForScript(EmitterData->GetGPUComputeScript(), VariableName, OutVariables, OutVariableOffsets, TotalAttributes, OutWarnings);
 		}
 	}
 	OutNumAttribChannelsFound = TotalAttributes - NumAttributes;
@@ -4204,5 +4206,109 @@ FIntVector FNiagaraDataInterfaceProxyGrid3DCollectionProxy::GetElementCount(FNia
 	return FIntVector::ZeroValue;
 }
 
+void UNiagaraDataInterfaceGrid3DCollection::FindAttributes(TArray<FNiagaraVariableBase>& OutVariables, TArray<uint32>& OutVariableOffsets, int32& OutNumAttribChannelsFound, TArray<FText>* OutWarnings, bool UseReader) const
+{
+	{
+		OutNumAttribChannelsFound = 0;
+
+		UNiagaraSystem* OwnerSystem = GetTypedOuter<UNiagaraSystem>();
+		if (OwnerSystem == nullptr)
+		{
+			return;
+		}
+
+		int32 TotalAttributes = NumAttributes;
+		for (const FNiagaraEmitterHandle& EmitterHandle : OwnerSystem->GetEmitterHandles())
+		{
+			FVersionedNiagaraEmitterData* EmitterData = EmitterHandle.GetEmitterData();
+			if (EmitterData && EmitterHandle.GetIsEnabled() && EmitterData->IsValid() && (EmitterData->SimTarget == ENiagaraSimTarget::GPUComputeSim))
+			{
+				// Search scripts for this data interface so we get the variable name
+				auto FindDataInterfaceVariable =
+					[&OwnerSystem, &EmitterData](const UNiagaraDataInterface* DataInterface) -> FName
+				{
+					UNiagaraScript* Scripts[] =
+					{
+						OwnerSystem->GetSystemSpawnScript(),
+						OwnerSystem->GetSystemUpdateScript(),
+						EmitterData->GetGPUComputeScript(),
+					};
+
+					for (UNiagaraScript* Script : Scripts)
+					{
+						for (FNiagaraScriptDataInterfaceInfo& DataInterfaceInfo : Script->GetCachedDefaultDataInterfaces())
+						{
+							if (DataInterfaceInfo.DataInterface == DataInterface)
+							{
+								return DataInterfaceInfo.RegisteredParameterMapRead.IsNone() ? DataInterfaceInfo.RegisteredParameterMapWrite : DataInterfaceInfo.RegisteredParameterMapRead;
+							}
+						}
+					}
+					return NAME_None;
+				};
+
+				const FName VariableName = FindDataInterfaceVariable(this);
+				if (!VariableName.IsNone())
+				{
+					CollectAttributesForScript(EmitterData->GetGPUComputeScript(), VariableName, OutVariables, OutVariableOffsets, TotalAttributes, OutWarnings, UseReader);
+				}
+			}
+		}
+		OutNumAttribChannelsFound = TotalAttributes - NumAttributes;
+	}
+}
+
+void UNiagaraDataInterfaceGrid3DCollection::CollectAttributesForScript(UNiagaraScript* Script, FName VariableName, TArray<FNiagaraVariableBase>& OutVariables, TArray<uint32>& OutVariableOffsets, int32& TotalAttributes, TArray<FText>* OutWarnings, bool UseReader)
+{
+	if (const FNiagaraScriptExecutionParameterStore* ParameterStore = Script->GetExecutionReadyParameterStore(ENiagaraSimTarget::GPUComputeSim))
+	{
+		FNiagaraVariableBase DataInterfaceVariable;
+
+		if (UseReader)
+		{
+			DataInterfaceVariable = FNiagaraVariableBase(FNiagaraTypeDefinition(UNiagaraDataInterfaceGrid3DCollectionReader::StaticClass()), VariableName);
+		}
+		else
+		{
+			DataInterfaceVariable = FNiagaraVariableBase(FNiagaraTypeDefinition(UNiagaraDataInterfaceGrid3DCollection::StaticClass()), VariableName);
+		}
+
+		const int32* IndexOfDataInterface = ParameterStore->FindParameterOffset(DataInterfaceVariable);
+		if (IndexOfDataInterface != nullptr)
+		{
+			TConstArrayView<FNiagaraDataInterfaceGPUParamInfo> ParamInfoArray = Script->GetDataInterfaceGPUParamInfos();
+			for (const FNiagaraDataInterfaceGeneratedFunction& Func : ParamInfoArray[*IndexOfDataInterface].GeneratedFunctions)
+			{
+				if (const FName* AttributeName = Func.FindSpecifierValue(NAME_Attribute))
+				{
+					FNiagaraVariableBase NewVar(UNiagaraDataInterfaceGrid3DCollection::GetValueTypeFromFuncName(Func.DefinitionName), *AttributeName);
+					if (UNiagaraDataInterfaceGrid3DCollection::CanCreateVarFromFuncName(Func.DefinitionName))
+					{
+						if (!OutVariables.Contains(NewVar))
+						{
+							const int32 FoundNameMatch = OutVariables.IndexOfByPredicate([&](const FNiagaraVariableBase& Var) { return Var.GetName() == *AttributeName; });
+							if (FoundNameMatch == INDEX_NONE)
+							{
+								OutVariables.Add(NewVar);
+								const int32 NumComponents = NewVar.GetSizeInBytes() / sizeof(float);
+								OutVariableOffsets.Add(TotalAttributes);
+								TotalAttributes += NumComponents;
+							}
+							else
+							{
+								if (OutWarnings)
+								{
+									FText Warning = FText::Format(NSLOCTEXT("NiagaraDataInterfaceGrid3DCollection", "BadType", "Same name, different types! {0} vs {1}, Attribute {2}"), NewVar.GetType().GetNameText(), OutVariables[FoundNameMatch].GetType().GetNameText(), FText::FromName(NewVar.GetName()));
+									OutWarnings->Add(Warning);
+								}
+							}
+						}
+					}
+				}
+			}
+		}
+	}
+}
+
 #undef LOCTEXT_NAMESPACE
 
diff --git a/Engine/Plugins/FX/Niagara/Source/Niagara/Private/NiagaraDataInterfaceGrid3DCollectionReader.cpp b/Engine/Plugins/FX/Niagara/Source/Niagara/Private/NiagaraDataInterfaceGrid3DCollectionReader.cpp
--- a/Engine/Plugins/FX/Niagara/Source/Niagara/Private/NiagaraDataInterfaceGrid3DCollectionReader.cpp
+++ b/Engine/Plugins/FX/Niagara/Source/Niagara/Private/NiagaraDataInterfaceGrid3DCollectionReader.cpp
@@ -110,7 +110,7 @@ bool UNiagaraDataInterfaceGrid3DCollectionReader::InitPerInstanceData(void* PerI
 			int32 NumNamedAttribChannelsFound = 0;
 			TArray<FNiagaraVariableBase> Vars;
 			TArray<uint32> Offsets;
-			FindAttributes<UNiagaraDataInterfaceGrid3DCollectionReader>(Vars, Offsets, NumNamedAttribChannelsFound);
+			FindAttributes(Vars, Offsets, NumNamedAttribChannelsFound, nullptr, true);
 
 			// #todo(dmp): slight hack here - these aren't always the same, but we generally can't use unnamed attrs very well
 			NumAttribChannelsFound = NumNamedAttribChannelsFound;
