- Validation rule to notify about tick dependencies

#rnx

diff --git a/Engine/Plugins/FX/Niagara/Source/NiagaraEditor/Private/NiagaraValidationRules.cpp b/Engine/Plugins/FX/Niagara/Source/NiagaraEditor/Private/NiagaraValidationRules.cpp
--- a/Engine/Plugins/FX/Niagara/Source/NiagaraEditor/Private/NiagaraValidationRules.cpp
+++ b/Engine/Plugins/FX/Niagara/Source/NiagaraEditor/Private/NiagaraValidationRules.cpp
@@ -4,6 +4,8 @@
 
 #include "NiagaraClipboard.h"
 #include "NiagaraComponentRendererProperties.h"
+#include "NiagaraDataInterfaceCamera.h"
+#include "NiagaraDataInterfaceSkeletalMesh.h"
 #include "NiagaraDataInterfaceUtilities.h"
 #include "NiagaraEditorSettings.h"
 #include "NiagaraNodeFunctionCall.h"
@@ -12,6 +14,7 @@
 #include "NiagaraSimulationStageBase.h"
 #include "NiagaraSystem.h"
 #include "NiagaraSystemEditorData.h"
+#include "DataInterface/NiagaraDataInterfaceActorComponent.h"
 #include "ViewModels/NiagaraEmitterHandleViewModel.h"
 #include "ViewModels/NiagaraSystemViewModel.h"
 #include "ViewModels/Stack/NiagaraStackEmitterSettingsGroup.h"
@@ -152,6 +155,29 @@ namespace NiagaraValidation
 		}
 		return FString();
 	}
+
+	TSharedPtr<FNiagaraEmitterHandleViewModel> GetEmitterViewModel(const FNiagaraValidationContext& Context, UNiagaraEmitter* NiagaraEmitter)
+	{
+		if (NiagaraEmitter == nullptr)
+		{
+			return nullptr;
+		}
+
+		const TSharedRef<FNiagaraEmitterHandleViewModel>* EmitterViewModel =
+			Context.ViewModel->GetEmitterHandleViewModels().FindByPredicate(
+				[NiagaraEmitter](const TSharedRef<FNiagaraEmitterHandleViewModel>& EmitterViewModelRef)
+				{
+					FNiagaraEmitterHandle* EmitterHandle = EmitterViewModelRef->GetEmitterHandle();
+					return EmitterHandle && EmitterHandle->GetInstance().Emitter == NiagaraEmitter;
+				}
+			);
+
+		if ( EmitterViewModel )
+		{
+			return *EmitterViewModel;
+		}
+		return nullptr;
+	}
 }
 
 // --------------------------------------------------------------------------------------------------------------------------------------------
@@ -456,21 +482,13 @@ void UNiagaraValidationRule_BannedDataInterfaces::CheckValidity(const FNiagaraVa
 			UObject* WarningObject = nullptr;
 			if (UNiagaraEmitter* NiagaraEmitter = Cast<UNiagaraEmitter>(UsageContext.OwnerObject))
 			{
-				const TSharedRef<FNiagaraEmitterHandleViewModel>* EmitterViewModel =
-					Context.ViewModel->GetEmitterHandleViewModels().FindByPredicate(
-						[NiagaraEmitter](const TSharedRef<FNiagaraEmitterHandleViewModel>& EmitterViewModelRef)
-						{
-							FNiagaraEmitterHandle* EmitterHandle = EmitterViewModelRef->GetEmitterHandle();
-							return EmitterHandle && EmitterHandle->GetInstance().Emitter == NiagaraEmitter;
-						}
-					);
-
+				const TSharedPtr<FNiagaraEmitterHandleViewModel> EmitterViewModel = NiagaraValidation::GetEmitterViewModel(Context, NiagaraEmitter);
 				if (EmitterViewModel == nullptr)
 				{
 					return true;
 				}
 
-				const FVersionedNiagaraEmitterData* EmitterData = (*EmitterViewModel)->GetEmitterHandle()->GetEmitterData();
+				const FVersionedNiagaraEmitterData* EmitterData = EmitterViewModel->GetEmitterHandle()->GetEmitterData();
 				if (EmitterData == nullptr)
 				{
 					return true;
@@ -489,7 +507,7 @@ void UNiagaraValidationRule_BannedDataInterfaces::CheckValidity(const FNiagaraVa
 							ENiagaraValidationSeverity::Warning,
 							FText::Format(WarningFormat, FText::FromName(UsageContext.Variable.GetName())),
 							FText::Format(EmitterDescFormat, FText::FromName(UsageContext.Variable.GetName()), FText::FromName(DIClass->GetFName()), FText::FromString(PlatformConflictsString)),
-							NiagaraValidation::GetStackEntry<UNiagaraStackEmitterPropertiesItem>((*EmitterViewModel)->GetEmitterStackViewModel())
+							NiagaraValidation::GetStackEntry<UNiagaraStackEmitterPropertiesItem>(EmitterViewModel->GetEmitterStackViewModel())
 						);
 					}
 				}
@@ -785,5 +803,68 @@ void UNiagaraValidationRule_SimulationStageBudget::CheckValidity(const FNiagaraV
 	}
 }
 
-#undef LOCTEXT_NAMESPACE
+void UNiagaraValidationRule_TickDependencyCheck::CheckValidity(const FNiagaraValidationContext& Context, TArray<FNiagaraValidationResult>& OutResults) const
+{
+	UNiagaraSystem* NiagaraSystem = &Context.ViewModel->GetSystem();
+
+	if (!bCheckActorComponentInterface && !bCheckCameraDataInterface && !bCheckSkeletalMeshInterface)
+	{
+		return;
+	}
+
+	NiagaraSystem->ForEachScript(
+		[&](UNiagaraScript* NiagaraScript)
+		{
+			const TArray<FNiagaraScriptDataInterfaceInfo>& CachedDefaultDIs = NiagaraScript->GetCachedDefaultDataInterfaces();
+			for ( const FNiagaraScriptDataInterfaceInfo& CachedDefaultDI : CachedDefaultDIs )
+			{
+				bool bWarnTickDependency = false;
+				if (UNiagaraDataInterfaceCamera* CameraDataInterface = Cast<UNiagaraDataInterfaceCamera>(CachedDefaultDI.DataInterface))
+				{
+					bWarnTickDependency = bCheckCameraDataInterface && CameraDataInterface->bRequireCurrentFrameData;
+				}
+				else if (UNiagaraDataInterfaceSkeletalMesh* SkeletalMeshDataInterface = Cast<UNiagaraDataInterfaceSkeletalMesh>(CachedDefaultDI.DataInterface))
+				{
+					bWarnTickDependency = bCheckSkeletalMeshInterface && SkeletalMeshDataInterface->bRequireCurrentFrameData;
+				}
+				else if (UNiagaraDataInterfaceActorComponent* ActorComponentDataInterface = Cast<UNiagaraDataInterfaceActorComponent>(CachedDefaultDI.DataInterface))
+				{
+					bWarnTickDependency = bCheckActorComponentInterface && ActorComponentDataInterface->bRequireCurrentFrameData;
+				}
+
+				if (bWarnTickDependency)
+				{
+					UNiagaraEmitter* NiagaraEmitter = NiagaraScript->GetTypedOuter<UNiagaraEmitter>();
 
+					UObject* StackObject = nullptr;
+					if ( NiagaraEmitter )
+					{
+						// If we are attatched to an emitter than is not enabled bail
+						const TSharedPtr<FNiagaraEmitterHandleViewModel> EmitterViewModel = NiagaraValidation::GetEmitterViewModel(Context, NiagaraEmitter);
+						if (EmitterViewModel == nullptr || EmitterViewModel->GetIsEnabled() == false)
+						{
+							return;
+						}
+						StackObject = NiagaraValidation::GetStackEntry<UNiagaraStackEmitterPropertiesItem>(EmitterViewModel->GetEmitterStackViewModel());
+					}
+					else
+					{
+						StackObject = NiagaraValidation::GetStackEntry<UNiagaraStackSystemPropertiesItem>(Context.ViewModel->GetSystemStackViewModel());
+					}
+
+					const FText OwnerStackText = FText::FromString(NiagaraEmitter ? NiagaraEmitter->GetName() : FString(TEXT("System")));
+					const FText DIClassText = FText::FromName(CachedDefaultDI.DataInterface->GetClass()->GetFName());
+					const FText VariableText = FText::FromName(CachedDefaultDI.Name);
+					OutResults.Emplace(
+						Severity,
+						FText::Format(LOCTEXT("TickDependencyCheckFormat", "'{0}' has a tick dependency that can potentially be removed"), OwnerStackText),
+						FText::Format(LOCTEXT("TickDependencyCheckDetailedFormat", "'{0}' has a tick dependency from data interace '{1}' variable '{2}' that can removed by unchecking 'RequireCurrentFrameData'"), OwnerStackText, DIClassText, VariableText),
+						StackObject
+					);
+				}
+			}
+		}
+	);
+}
+
+#undef LOCTEXT_NAMESPACE
diff --git a/Engine/Plugins/FX/Niagara/Source/NiagaraEditor/Public/NiagaraValidationRules.h b/Engine/Plugins/FX/Niagara/Source/NiagaraEditor/Public/NiagaraValidationRules.h
--- a/Engine/Plugins/FX/Niagara/Source/NiagaraEditor/Public/NiagaraValidationRules.h
+++ b/Engine/Plugins/FX/Niagara/Source/NiagaraEditor/Public/NiagaraValidationRules.h
@@ -164,3 +164,29 @@ public:
 	UPROPERTY(EditAnywhere, Category = Validation, meta = (EditCondition = "bMaxTotalIterationsEnabled"))
 	int32 MaxTotalIterations = 1;
 };
+
+/** Validation rule to check that we don't have a tick dependency we don't want.  */
+UCLASS(Category = "Validation", DisplayName = "Tick Dependency Check")
+class UNiagaraValidationRule_TickDependencyCheck : public UNiagaraValidationRule
+{
+	GENERATED_BODY()
+
+public:
+	virtual void CheckValidity(const FNiagaraValidationContext& Context, TArray<FNiagaraValidationResult>& OutResults) const override;
+
+	/** How do we want to repro the error in the stack */
+	UPROPERTY(EditAnywhere, Category = Validation)
+	ENiagaraValidationSeverity Severity = ENiagaraValidationSeverity::Info;
+
+	/** Check that the actor component interface isn't adding a tick dependency on the CPU. */
+	UPROPERTY(EditAnywhere, Category = Validation)
+	bool bCheckActorComponentInterface = true;
+
+	/** Check that the camera data interface isn't adding a tick dependency on the CPU. */
+	UPROPERTY(EditAnywhere, Category = Validation)
+	bool bCheckCameraDataInterface = true;
+
+	/** Check that the skeletal mesh interface isn't adding a tick dependency on the CPU. */
+	UPROPERTY(EditAnywhere, Category = Validation)
+	bool bCheckSkeletalMeshInterface = true;
+};
