Guard against re-entrant shader compiling by making sure that we don't try to fully load a system just because it shows up in a TObjectIterator resulting from an UpdateContext

diff --git a/Engine/Plugins/FX/Niagara/Source/Niagara/Private/NiagaraCommon.cpp b/Engine/Plugins/FX/Niagara/Source/Niagara/Private/NiagaraCommon.cpp
--- a/Engine/Plugins/FX/Niagara/Source/Niagara/Private/NiagaraCommon.cpp
+++ b/Engine/Plugins/FX/Niagara/Source/Niagara/Private/NiagaraCommon.cpp
@@ -182,7 +182,9 @@ void FNiagaraSystemUpdateContext::Add(const FVersionedNiagaraEmitter& Emitter, b
 		UNiagaraComponent* Comp = *It;
 		check(Comp);
 		UNiagaraSystem* System = Comp->GetAsset();
-		if (System && System->UsesEmitter(Emitter))
+
+		// only worry about systems that are fully loaded (and don't force it to be fully loaded just because it shows up in the iterator
+		if (System && System->bFullyLoaded && System->UsesEmitter(Emitter))
 		{
 			bool bAllowDestroySystemSim = true;
 			AddInternal(Comp, bReInit, bAllowDestroySystemSim);
@@ -196,8 +198,10 @@ void FNiagaraSystemUpdateContext::Add(const UNiagaraScript* Script, bool bReInit
 	{
 		UNiagaraComponent* Comp = *It;
 		check(Comp);
+
+		// only worry about systems that are fully loaded (and don't force it to be fully loaded just because it shows up in the iterator
 		UNiagaraSystem* System = Comp->GetAsset();
-		if (System && System->UsesScript(Script))
+		if (System && System->bFullyLoaded && System->UsesScript(Script))
 		{
 			bool bAllowDestroySystemSim = true;
 			AddInternal(Comp, bReInit, bAllowDestroySystemSim);
@@ -212,7 +216,9 @@ void FNiagaraSystemUpdateContext::Add(const UNiagaraParameterCollection* Collect
 		UNiagaraComponent* Comp = *It;
 		check(Comp);
 		UNiagaraSystem* System = Comp->GetAsset();
-		if (System && System->UsesCollection(Collection))
+
+		// only worry about systems that are fully loaded (and don't force it to be fully loaded just because it shows up in the iterator
+		if (System && System->bFullyLoaded && System->UsesCollection(Collection))
 		{
 			bool bAllowDestroySystemSim = true;
 			AddInternal(Comp, bReInit, bAllowDestroySystemSim);
diff --git a/Engine/Plugins/FX/Niagara/Source/NiagaraEditor/Private/NiagaraHlslTranslator.cpp b/Engine/Plugins/FX/Niagara/Source/NiagaraEditor/Private/NiagaraHlslTranslator.cpp
--- a/Engine/Plugins/FX/Niagara/Source/NiagaraEditor/Private/NiagaraHlslTranslator.cpp
+++ b/Engine/Plugins/FX/Niagara/Source/NiagaraEditor/Private/NiagaraHlslTranslator.cpp
@@ -81,6 +81,17 @@ DECLARE_CYCLE_STAT(TEXT("Niagara - HlslTranslator - GenerateFunctionSignature_Fi
 // because editor tickables aren't ticked during cooking
 void FNiagaraShaderQueueTickable::ProcessQueue()
 {
+	static bool bReentrantProcessingGuard = false;
+
+	// The current implementation of the compilation queue will drain the contents during one invocation of ProcessQueue.  If processing
+	// the results of one of the shader script indirectly results in another invocation of ProcessQueue we'll run into trouble because
+	// the contents of the queue will be double processed.  While we could move the queue to a local variable to process to prevent this
+	// double processing we'll still run into trouble because the second scope could be expecting a result, which won't be available till
+	// the earlier scope is complete.  Needs a better system...
+	checkf(bReentrantProcessingGuard == false, TEXT("FNiagaraShaderQueueTickable::ProcessQueue() is not re-entrant!  Only a single scope can process results!"));
+
+	TGuardValue<bool> GuardProcessQueue(bReentrantProcessingGuard, true);
+
 	check(IsInGameThread());
 
 	for (FNiagaraCompilationQueue::NiagaraCompilationQueueItem &Item : FNiagaraCompilationQueue::Get()->GetQueue())
