Add an ensure to provide more info about potentially missing GraphSource

diff --git a/Engine/Plugins/FX/Niagara/Source/Niagara/Private/NiagaraEmitter.cpp b/Engine/Plugins/FX/Niagara/Source/Niagara/Private/NiagaraEmitter.cpp
--- a/Engine/Plugins/FX/Niagara/Source/Niagara/Private/NiagaraEmitter.cpp
+++ b/Engine/Plugins/FX/Niagara/Source/Niagara/Private/NiagaraEmitter.cpp
@@ -754,8 +754,13 @@ void FVersionedNiagaraEmitterData::PostLoad(UNiagaraEmitter& Emitter, bool bIsCo
 
 	if (!Emitter.GetOutermost()->bIsCookedForEditor)
 	{
-		GraphSource->ConditionalPostLoad();
-		GraphSource->PostLoadFromEmitter(FVersionedNiagaraEmitter(&Emitter, Version.VersionGuid));
+		if (ensureMsgf(GraphSource != nullptr, TEXT("Niagara emitter %s [Flags: %x] - Version (%d.%d - %s) - Missing GraphSource"),
+			*Emitter.GetPathName(), Emitter.GetFlags(),
+			Version.MajorVersion, Version.MinorVersion, *Version.VersionGuid.ToString()))
+		{
+			GraphSource->ConditionalPostLoad();
+			GraphSource->PostLoadFromEmitter(FVersionedNiagaraEmitter(&Emitter, Version.VersionGuid));
+		}
 		
 		// Prepare for emitter inheritance.
 		if (GetParent().Emitter != nullptr)
