Strategical submit which adds includes in preparation for coming change which removes includes in headers

diff --git a/Engine/Plugins/Animation/AnimToTexture/Source/AnimToTextureEditor/Private/AnimToTextureBPLibrary.cpp b/Engine/Plugins/Animation/AnimToTexture/Source/AnimToTextureEditor/Private/AnimToTextureBPLibrary.cpp
--- a/Engine/Plugins/Animation/AnimToTexture/Source/AnimToTextureEditor/Private/AnimToTextureBPLibrary.cpp
+++ b/Engine/Plugins/Animation/AnimToTexture/Source/AnimToTextureEditor/Private/AnimToTextureBPLibrary.cpp
@@ -1,6 +1,7 @@
 // Copyright Epic Games, Inc. All Rights Reserved.
 
 #include "AnimToTextureBPLibrary.h"
+#include "Animation/Skeleton.h"
 #include "AnimToTexture.h"
 #include "AnimToTextureUtils.h"
 #include "AnimToTextureSkeletalMesh.h"
diff --git a/Engine/Plugins/Animation/IKRig/Source/IKRigEditor/Public/RetargetEditor/IKRetargetEditorController.h b/Engine/Plugins/Animation/IKRig/Source/IKRigEditor/Public/RetargetEditor/IKRetargetEditorController.h
--- a/Engine/Plugins/Animation/IKRig/Source/IKRigEditor/Public/RetargetEditor/IKRetargetEditorController.h
+++ b/Engine/Plugins/Animation/IKRig/Source/IKRigEditor/Public/RetargetEditor/IKRetargetEditorController.h
@@ -17,6 +17,7 @@ class UIKRetargetProcessor;
 class SIKRetargetChainMapList;
 class UIKRetargetAnimInstance;
 class FIKRetargetEditor;
+class FPrimitiveDrawInterface;
 class UDebugSkelMeshComponent;
 class UIKRigDefinition;
 class UIKRetargeterController;
diff --git a/Engine/Plugins/Animation/MLDeformer/MLDeformerFramework/Source/MLDeformerFrameworkEditor/Private/MLDeformerModelDetails.cpp b/Engine/Plugins/Animation/MLDeformer/MLDeformerFramework/Source/MLDeformerFrameworkEditor/Private/MLDeformerModelDetails.cpp
--- a/Engine/Plugins/Animation/MLDeformer/MLDeformerFramework/Source/MLDeformerFrameworkEditor/Private/MLDeformerModelDetails.cpp
+++ b/Engine/Plugins/Animation/MLDeformer/MLDeformerFramework/Source/MLDeformerFrameworkEditor/Private/MLDeformerModelDetails.cpp
@@ -1,6 +1,7 @@
 // Copyright Epic Games, Inc. All Rights Reserved.
 
 #include "MLDeformerModelDetails.h"
+#include "Animation/Skeleton.h"
 #include "MLDeformerModule.h"
 #include "MLDeformerAsset.h"
 #include "MLDeformerModel.h"
diff --git a/Engine/Plugins/Animation/MLDeformer/MLDeformerFramework/Source/MLDeformerFrameworkEditor/Private/MLDeformerVizSettingsDetails.cpp b/Engine/Plugins/Animation/MLDeformer/MLDeformerFramework/Source/MLDeformerFrameworkEditor/Private/MLDeformerVizSettingsDetails.cpp
--- a/Engine/Plugins/Animation/MLDeformer/MLDeformerFramework/Source/MLDeformerFrameworkEditor/Private/MLDeformerVizSettingsDetails.cpp
+++ b/Engine/Plugins/Animation/MLDeformer/MLDeformerFramework/Source/MLDeformerFrameworkEditor/Private/MLDeformerVizSettingsDetails.cpp
@@ -1,6 +1,7 @@
 // Copyright Epic Games, Inc. All Rights Reserved.
 
 #include "MLDeformerVizSettingsDetails.h"
+#include "Animation/Skeleton.h"
 #include "MLDeformerModule.h"
 #include "MLDeformerModel.h"
 #include "MLDeformerEditorModel.h"
diff --git a/Engine/Plugins/Animation/MotionWarping/Source/MotionWarping/Private/MotionWarpingComponent.cpp b/Engine/Plugins/Animation/MotionWarping/Source/MotionWarping/Private/MotionWarpingComponent.cpp
--- a/Engine/Plugins/Animation/MotionWarping/Source/MotionWarping/Private/MotionWarpingComponent.cpp
+++ b/Engine/Plugins/Animation/MotionWarping/Source/MotionWarping/Private/MotionWarpingComponent.cpp
@@ -6,6 +6,7 @@
 #include "Animation/AnimMontage.h"
 #include "Animation/AnimInstance.h"
 #include "Animation/AnimationPoseData.h"
+#include "BonePose.h"
 #include "Components/SkeletalMeshComponent.h"
 #include "DrawDebugHelpers.h"
 #include "Engine/World.h"
diff --git a/Engine/Plugins/Animation/MotionWarping/Source/MotionWarping/Public/MotionWarpingComponent.h b/Engine/Plugins/Animation/MotionWarping/Source/MotionWarping/Public/MotionWarpingComponent.h
--- a/Engine/Plugins/Animation/MotionWarping/Source/MotionWarping/Public/MotionWarpingComponent.h
+++ b/Engine/Plugins/Animation/MotionWarping/Source/MotionWarping/Public/MotionWarpingComponent.h
@@ -14,6 +14,8 @@ class UAnimSequenceBase;
 class UCharacterMovementComponent;
 class UMotionWarpingComponent;
 class UAnimNotifyState_MotionWarping;
+struct FCompactPose;
+template<class PoseType> struct FCSPose;
 
 DECLARE_LOG_CATEGORY_EXTERN(LogMotionWarping, Log, All);
 
diff --git a/Engine/Plugins/Animation/MotionWarping/Source/MotionWarping/Public/RootMotionModifier_AdjustmentBlendWarp.h b/Engine/Plugins/Animation/MotionWarping/Source/MotionWarping/Public/RootMotionModifier_AdjustmentBlendWarp.h
--- a/Engine/Plugins/Animation/MotionWarping/Source/MotionWarping/Public/RootMotionModifier_AdjustmentBlendWarp.h
+++ b/Engine/Plugins/Animation/MotionWarping/Source/MotionWarping/Public/RootMotionModifier_AdjustmentBlendWarp.h
@@ -2,6 +2,7 @@
 
 #pragma once
 
+#include "BonePose.h"
 #include "CoreMinimal.h"
 #include "RootMotionModifier.h"
 #include "RootMotionModifier_AdjustmentBlendWarp.generated.h"
diff --git a/Engine/Plugins/Developer/AnimationSharing/Source/AnimationSharing/Private/AnimationSharingManager.cpp b/Engine/Plugins/Developer/AnimationSharing/Source/AnimationSharing/Private/AnimationSharingManager.cpp
--- a/Engine/Plugins/Developer/AnimationSharing/Source/AnimationSharing/Private/AnimationSharingManager.cpp
+++ b/Engine/Plugins/Developer/AnimationSharing/Source/AnimationSharing/Private/AnimationSharingManager.cpp
@@ -1,6 +1,7 @@
 // Copyright Epic Games, Inc. All Rights Reserved.
 
 #include "AnimationSharingManager.h"
+#include "Animation/Skeleton.h"
 #include "AnimationSharingModule.h"
 #include "Stats/Stats.h"
 #include "Components/SkinnedMeshComponent.h"
diff --git a/Engine/Plugins/EnhancedInput/Source/EnhancedInput/Private/EnhancedPlayerInput.cpp b/Engine/Plugins/EnhancedInput/Source/EnhancedInput/Private/EnhancedPlayerInput.cpp
--- a/Engine/Plugins/EnhancedInput/Source/EnhancedInput/Private/EnhancedPlayerInput.cpp
+++ b/Engine/Plugins/EnhancedInput/Source/EnhancedInput/Private/EnhancedPlayerInput.cpp
@@ -3,6 +3,7 @@
 #include "EnhancedPlayerInput.h"
 #include "EnhancedInputComponent.h"
 #include "GameFramework/PlayerController.h"
+#include "HAL/IConsoleManager.h"
 #include "InputAction.h"
 #include "InputModifiers.h"
 #include "InputTriggers.h"
diff --git a/Engine/Plugins/Enterprise/DatasmithFBXImporter/Source/DatasmithFBXTranslator/Private/DatasmithFBXSceneProcessor.cpp b/Engine/Plugins/Enterprise/DatasmithFBXImporter/Source/DatasmithFBXTranslator/Private/DatasmithFBXSceneProcessor.cpp
--- a/Engine/Plugins/Enterprise/DatasmithFBXImporter/Source/DatasmithFBXTranslator/Private/DatasmithFBXSceneProcessor.cpp
+++ b/Engine/Plugins/Enterprise/DatasmithFBXImporter/Source/DatasmithFBXTranslator/Private/DatasmithFBXSceneProcessor.cpp
@@ -4,6 +4,7 @@
 
 #include "DatasmithFBXHashUtils.h"
 #include "DatasmithFBXImporterLog.h"
+#include "Misc/Paths.h"
 #include "Utility/DatasmithMeshHelper.h"
 
 #include "FbxImporter.h"
diff --git a/Engine/Plugins/Enterprise/GLTFExporter/Source/GLTFExporter/Private/Converters/GLTFBoneUtility.cpp b/Engine/Plugins/Enterprise/GLTFExporter/Source/GLTFExporter/Private/Converters/GLTFBoneUtility.cpp
--- a/Engine/Plugins/Enterprise/GLTFExporter/Source/GLTFExporter/Private/Converters/GLTFBoneUtility.cpp
+++ b/Engine/Plugins/Enterprise/GLTFExporter/Source/GLTFExporter/Private/Converters/GLTFBoneUtility.cpp
@@ -1,6 +1,7 @@
 // Copyright Epic Games, Inc. All Rights Reserved.
 
 #include "Converters/GLTFBoneUtility.h"
+#include "BonePose.h"
 #include "Sections/MovieScene3DTransformSection.h"
 #include "Animation/AnimationPoseData.h"
 #include "Animation/AnimSequence.h"
diff --git a/Engine/Plugins/Experimental/Animation/ContextualAnimation/Source/ContextualAnimation/Public/ContextualAnimUtilities.h b/Engine/Plugins/Experimental/Animation/ContextualAnimation/Source/ContextualAnimation/Public/ContextualAnimUtilities.h
--- a/Engine/Plugins/Experimental/Animation/ContextualAnimation/Source/ContextualAnimation/Public/ContextualAnimUtilities.h
+++ b/Engine/Plugins/Experimental/Animation/ContextualAnimation/Source/ContextualAnimation/Public/ContextualAnimUtilities.h
@@ -13,7 +13,9 @@ class UAnimInstance;
 class AActor;
 class FPrimitiveDrawInterface;
 struct FAnimMontageInstance;
+struct FCompactPose;
 struct FContextualAnimSet;
+template<class PoseType> struct FCSPose;
 
 UCLASS()
 class CONTEXTUALANIMATION_API UContextualAnimUtilities : public UBlueprintFunctionLibrary
diff --git a/Engine/Plugins/Experimental/Animation/PoseSearch/Source/Runtime/Public/PoseSearch/PoseSearch.h b/Engine/Plugins/Experimental/Animation/PoseSearch/Source/Runtime/Public/PoseSearch/PoseSearch.h
--- a/Engine/Plugins/Experimental/Animation/PoseSearch/Source/Runtime/Public/PoseSearch/PoseSearch.h
+++ b/Engine/Plugins/Experimental/Animation/PoseSearch/Source/Runtime/Public/PoseSearch/PoseSearch.h
@@ -3,6 +3,7 @@
 #pragma once
 
 #include "Async/ParallelFor.h"
+#include "BonePose.h"
 #include "CoreMinimal.h"
 #include "Containers/RingBuffer.h"
 #include "Engine/DataAsset.h"
diff --git a/Engine/Plugins/Experimental/GeometryScripting/Source/GeometryScriptingEditor/Private/GeometryActors/GeneratedDynamicMeshActor.cpp b/Engine/Plugins/Experimental/GeometryScripting/Source/GeometryScriptingEditor/Private/GeometryActors/GeneratedDynamicMeshActor.cpp
--- a/Engine/Plugins/Experimental/GeometryScripting/Source/GeometryScriptingEditor/Private/GeometryActors/GeneratedDynamicMeshActor.cpp
+++ b/Engine/Plugins/Experimental/GeometryScripting/Source/GeometryScriptingEditor/Private/GeometryActors/GeneratedDynamicMeshActor.cpp
@@ -2,6 +2,7 @@
 
 #include "GeometryActors/GeneratedDynamicMeshActor.h"
 #include "Components/StaticMeshComponent.h"
+#include "Engine/Level.h"
 #include "GeometryActors/EditorGeometryGenerationSubsystem.h"
 
 #include "Editor/EditorEngine.h" // for CopyPropertiesForUnrelatedObjects
diff --git a/Engine/Plugins/Experimental/Mutable/Source/CustomizableObjectEditor/Private/MuCOE/CustomizableObjectEditorViewportClient.cpp b/Engine/Plugins/Experimental/Mutable/Source/CustomizableObjectEditor/Private/MuCOE/CustomizableObjectEditorViewportClient.cpp
--- a/Engine/Plugins/Experimental/Mutable/Source/CustomizableObjectEditor/Private/MuCOE/CustomizableObjectEditorViewportClient.cpp
+++ b/Engine/Plugins/Experimental/Mutable/Source/CustomizableObjectEditor/Private/MuCOE/CustomizableObjectEditorViewportClient.cpp
@@ -7,6 +7,7 @@
 #include "Animation/DebugSkelMeshComponent.h"
 #include "Animation/PoseAsset.h"
 #include "Animation/AnimSingleNodeInstance.h"
+#include "Animation/Skeleton.h"
 #include "AssetViewerSettings.h"
 #include "BatchedElements.h"
 #include "CanvasTypes.h"
diff --git a/Engine/Plugins/Experimental/PCG/Source/PCG/Private/Helpers/PCGBlueprintHelpers.cpp b/Engine/Plugins/Experimental/PCG/Source/PCG/Private/Helpers/PCGBlueprintHelpers.cpp
--- a/Engine/Plugins/Experimental/PCG/Source/PCG/Private/Helpers/PCGBlueprintHelpers.cpp
+++ b/Engine/Plugins/Experimental/PCG/Source/PCG/Private/Helpers/PCGBlueprintHelpers.cpp
@@ -1,6 +1,7 @@
 // Copyright Epic Games, Inc. All Rights Reserved.
 
 #include "Helpers/PCGBlueprintHelpers.h"
+#include "LandscapeComponent.h"
 #include "PCGComponent.h"
 #include "PCGHelpers.h"
 #include "PCGSettings.h"
diff --git a/Engine/Plugins/FX/Niagara/Source/Niagara/Classes/NiagaraDataInterfaceArrayImpl.h b/Engine/Plugins/FX/Niagara/Source/Niagara/Classes/NiagaraDataInterfaceArrayImpl.h
--- a/Engine/Plugins/FX/Niagara/Source/Niagara/Classes/NiagaraDataInterfaceArrayImpl.h
+++ b/Engine/Plugins/FX/Niagara/Source/Niagara/Classes/NiagaraDataInterfaceArrayImpl.h
@@ -3,6 +3,7 @@
 #pragma once
 
 #include "NiagaraClearCounts.h"
+#include "NiagaraCompileHashVisitor.h"
 #include "NiagaraDataInterfaceArray.h"
 #include "NiagaraDataInterfaceUtilities.h"
 #include "NiagaraGpuComputeDispatchInterface.h"
diff --git a/Engine/Plugins/FX/Niagara/Source/Niagara/Classes/NiagaraGPUInstanceCountManager.h b/Engine/Plugins/FX/Niagara/Source/Niagara/Classes/NiagaraGPUInstanceCountManager.h
--- a/Engine/Plugins/FX/Niagara/Source/Niagara/Classes/NiagaraGPUInstanceCountManager.h
+++ b/Engine/Plugins/FX/Niagara/Source/Niagara/Classes/NiagaraGPUInstanceCountManager.h
@@ -7,7 +7,7 @@ NiagaraGPUInstanceCountManager.h: GPU particle count handling
 
 #include "CoreMinimal.h"
 #include "NiagaraCommon.h"
-#include "RHIUtilities.h"
+#include "RHI.h"
 #include "NiagaraDrawIndirect.h"
 #include "RHIGPUReadback.h"
 
diff --git a/Engine/Plugins/FX/Niagara/Source/Niagara/Private/MovieScene/Parameters/MovieSceneNiagaraParameterSectionTemplate.cpp b/Engine/Plugins/FX/Niagara/Source/Niagara/Private/MovieScene/Parameters/MovieSceneNiagaraParameterSectionTemplate.cpp
--- a/Engine/Plugins/FX/Niagara/Source/Niagara/Private/MovieScene/Parameters/MovieSceneNiagaraParameterSectionTemplate.cpp
+++ b/Engine/Plugins/FX/Niagara/Source/Niagara/Private/MovieScene/Parameters/MovieSceneNiagaraParameterSectionTemplate.cpp
@@ -2,6 +2,7 @@
 
 #include "MovieSceneNiagaraParameterSectionTemplate.h"
 #include "NiagaraComponent.h"
+#include "NiagaraEffectType.h"
 #include "NiagaraSystem.h"
 #include "NiagaraTypes.h"
 #include "IMovieScenePlayer.h"
diff --git a/Engine/Plugins/FX/Niagara/Source/Niagara/Private/NiagaraActor.cpp b/Engine/Plugins/FX/Niagara/Source/Niagara/Private/NiagaraActor.cpp
--- a/Engine/Plugins/FX/Niagara/Source/Niagara/Private/NiagaraActor.cpp
+++ b/Engine/Plugins/FX/Niagara/Source/Niagara/Private/NiagaraActor.cpp
@@ -8,6 +8,7 @@
 #include "Components/BillboardComponent.h"
 #include "NiagaraComponent.h"
 #include "NiagaraSystem.h"
+#include "NiagaraEffectType.h"
 
 #include UE_INLINE_GENERATED_CPP_BY_NAME(NiagaraActor)
 
diff --git a/Engine/Plugins/FX/Niagara/Source/Niagara/Private/NiagaraAsyncCompile.cpp b/Engine/Plugins/FX/Niagara/Source/Niagara/Private/NiagaraAsyncCompile.cpp
--- a/Engine/Plugins/FX/Niagara/Source/Niagara/Private/NiagaraAsyncCompile.cpp
+++ b/Engine/Plugins/FX/Niagara/Source/Niagara/Private/NiagaraAsyncCompile.cpp
@@ -4,6 +4,7 @@
 #include "NiagaraAsyncCompile.h"
 
 #include "Modules/ModuleManager.h"
+#include "NiagaraEmitter.h"
 #include "NiagaraEmitterHandle.h"
 #include "NiagaraModule.h"
 #include "NiagaraPrecompileContainer.h"
diff --git a/Engine/Plugins/FX/Niagara/Source/Niagara/Private/NiagaraAsyncGpuTraceHelper.cpp b/Engine/Plugins/FX/Niagara/Source/Niagara/Private/NiagaraAsyncGpuTraceHelper.cpp
--- a/Engine/Plugins/FX/Niagara/Source/Niagara/Private/NiagaraAsyncGpuTraceHelper.cpp
+++ b/Engine/Plugins/FX/Niagara/Source/Niagara/Private/NiagaraAsyncGpuTraceHelper.cpp
@@ -2,6 +2,7 @@
 
 #include "NiagaraAsyncGpuTraceHelper.h"
 
+#include "Components/PrimitiveComponent.h"
 #include "NiagaraGpuComputeDispatchInterface.h"
 #include "NiagaraStats.h"
 
diff --git a/Engine/Plugins/FX/Niagara/Source/Niagara/Private/NiagaraBakerOutput.cpp b/Engine/Plugins/FX/Niagara/Source/Niagara/Private/NiagaraBakerOutput.cpp
--- a/Engine/Plugins/FX/Niagara/Source/Niagara/Private/NiagaraBakerOutput.cpp
+++ b/Engine/Plugins/FX/Niagara/Source/Niagara/Private/NiagaraBakerOutput.cpp
@@ -1,6 +1,8 @@
 // Copyright Epic Games, Inc. All Rights Reserved.
 
 #include "NiagaraBakerOutput.h"
+#include "Misc/Paths.h"
+#include "NiagaraEffectType.h"
 #include "NiagaraSystem.h"
 
 #include UE_INLINE_GENERATED_CPP_BY_NAME(NiagaraBakerOutput)
diff --git a/Engine/Plugins/FX/Niagara/Source/Niagara/Private/NiagaraBakerSettings.cpp b/Engine/Plugins/FX/Niagara/Source/Niagara/Private/NiagaraBakerSettings.cpp
--- a/Engine/Plugins/FX/Niagara/Source/Niagara/Private/NiagaraBakerSettings.cpp
+++ b/Engine/Plugins/FX/Niagara/Source/Niagara/Private/NiagaraBakerSettings.cpp
@@ -4,6 +4,7 @@
 #include "Framework/Application/SlateApplication.h"
 #include "NiagaraBakerOutputTexture2D.h"
 #include "NiagaraCustomVersion.h"
+#include "NiagaraEffectType.h"
 #include "NiagaraSystem.h"
 #include "Engine/Texture2D.h"
 #include "Misc/PathViews.h"
diff --git a/Engine/Plugins/FX/Niagara/Source/Niagara/Private/NiagaraComponentRendererProperties.cpp b/Engine/Plugins/FX/Niagara/Source/Niagara/Private/NiagaraComponentRendererProperties.cpp
--- a/Engine/Plugins/FX/Niagara/Source/Niagara/Private/NiagaraComponentRendererProperties.cpp
+++ b/Engine/Plugins/FX/Niagara/Source/Niagara/Private/NiagaraComponentRendererProperties.cpp
@@ -1,6 +1,7 @@
 // Copyright Epic Games, Inc. All Rights Reserved.
 
 #include "NiagaraComponentRendererProperties.h"
+#include "NiagaraEmitterInstance.h"
 #include "NiagaraRenderer.h"
 #include "NiagaraConstants.h"
 #include "NiagaraRendererComponents.h"
diff --git a/Engine/Plugins/FX/Niagara/Source/Niagara/Private/NiagaraDataInterfaceArray.cpp b/Engine/Plugins/FX/Niagara/Source/Niagara/Private/NiagaraDataInterfaceArray.cpp
--- a/Engine/Plugins/FX/Niagara/Source/Niagara/Private/NiagaraDataInterfaceArray.cpp
+++ b/Engine/Plugins/FX/Niagara/Source/Niagara/Private/NiagaraDataInterfaceArray.cpp
@@ -1,6 +1,7 @@
 // Copyright Epic Games, Inc. All Rights Reserved.
 
 #include "NiagaraDataInterfaceArray.h"
+#include "NiagaraCompileHashVisitor.h"
 #include "NiagaraShaderParametersBuilder.h"
 
 #include UE_INLINE_GENERATED_CPP_BY_NAME(NiagaraDataInterfaceArray)
diff --git a/Engine/Plugins/FX/Niagara/Source/Niagara/Private/NiagaraDataInterfaceCurlNoise.cpp b/Engine/Plugins/FX/Niagara/Source/Niagara/Private/NiagaraDataInterfaceCurlNoise.cpp
--- a/Engine/Plugins/FX/Niagara/Source/Niagara/Private/NiagaraDataInterfaceCurlNoise.cpp
+++ b/Engine/Plugins/FX/Niagara/Source/Niagara/Private/NiagaraDataInterfaceCurlNoise.cpp
@@ -1,6 +1,7 @@
 // Copyright Epic Games, Inc. All Rights Reserved.
 
 #include "NiagaraDataInterfaceCurlNoise.h"
+#include "NiagaraCompileHashVisitor.h"
 #include "NiagaraShader.h"
 #include "NiagaraShaderParametersBuilder.h"
 #include "ShaderParameterUtils.h"
diff --git a/Engine/Plugins/FX/Niagara/Source/Niagara/Private/NiagaraDataInterfaceGrid3DCollection.cpp b/Engine/Plugins/FX/Niagara/Source/Niagara/Private/NiagaraDataInterfaceGrid3DCollection.cpp
--- a/Engine/Plugins/FX/Niagara/Source/Niagara/Private/NiagaraDataInterfaceGrid3DCollection.cpp
+++ b/Engine/Plugins/FX/Niagara/Source/Niagara/Private/NiagaraDataInterfaceGrid3DCollection.cpp
@@ -1,5 +1,6 @@
 // Copyright Epic Games, Inc. All Rights Reserved.
 #include "NiagaraDataInterfaceGrid3DCollection.h"
+#include "NiagaraCompileHashVisitor.h"
 #include "NiagaraConstants.h"
 #include "NiagaraGpuComputeDispatchInterface.h"
 #include "NiagaraSimStageData.h"
diff --git a/Engine/Plugins/FX/Niagara/Source/Niagara/Private/NiagaraDataInterfaceParticleRead.cpp b/Engine/Plugins/FX/Niagara/Source/Niagara/Private/NiagaraDataInterfaceParticleRead.cpp
--- a/Engine/Plugins/FX/Niagara/Source/Niagara/Private/NiagaraDataInterfaceParticleRead.cpp
+++ b/Engine/Plugins/FX/Niagara/Source/Niagara/Private/NiagaraDataInterfaceParticleRead.cpp
@@ -1,6 +1,7 @@
 // Copyright Epic Games, Inc. All Rights Reserved.
 
 #include "NiagaraDataInterfaceParticleRead.h"
+#include "NiagaraCompileHashVisitor.h"
 #include "NiagaraComputeExecutionContext.h"
 #include "NiagaraConstants.h"
 #include "NiagaraDataSetAccessor.h"
diff --git a/Engine/Plugins/FX/Niagara/Source/Niagara/Private/NiagaraDataInterfaceSkeletalMesh.cpp b/Engine/Plugins/FX/Niagara/Source/Niagara/Private/NiagaraDataInterfaceSkeletalMesh.cpp
--- a/Engine/Plugins/FX/Niagara/Source/Niagara/Private/NiagaraDataInterfaceSkeletalMesh.cpp
+++ b/Engine/Plugins/FX/Niagara/Source/Niagara/Private/NiagaraDataInterfaceSkeletalMesh.cpp
@@ -3,12 +3,14 @@
 #include "NiagaraDataInterfaceSkeletalMesh.h"
 
 #include "Animation/SkeletalMeshActor.h"
+#include "Animation/Skeleton.h"
 #include "Async/ParallelFor.h"
 #include "Components/SkeletalMeshComponent.h"
 #include "Engine/Canvas.h"
 #include "Engine/SkeletalMeshSocket.h"
 #include "Internationalization/Internationalization.h"
 #include "NDISkeletalMeshCommon.h"
+#include "NiagaraCompileHashVisitor.h"
 #include "NiagaraEmitterInstance.h"
 #include "NiagaraComponent.h"
 #include "NiagaraDataInterfaceSkeletalMeshConnectivity.h"
diff --git a/Engine/Plugins/FX/Niagara/Source/Niagara/Private/NiagaraDataInterfaceVectorField.cpp b/Engine/Plugins/FX/Niagara/Source/Niagara/Private/NiagaraDataInterfaceVectorField.cpp
--- a/Engine/Plugins/FX/Niagara/Source/Niagara/Private/NiagaraDataInterfaceVectorField.cpp
+++ b/Engine/Plugins/FX/Niagara/Source/Niagara/Private/NiagaraDataInterfaceVectorField.cpp
@@ -1,6 +1,7 @@
 // Copyright Epic Games, Inc. All Rights Reserved.
 
 #include "NiagaraDataInterfaceVectorField.h"
+#include "NiagaraEmitter.h"
 #include "VectorField/VectorFieldStatic.h"
 #include "VectorField/VectorFieldAnimated.h"
 #include "NiagaraShader.h"
diff --git a/Engine/Plugins/FX/Niagara/Source/Niagara/Private/NiagaraDataSetDebugAccessor.cpp b/Engine/Plugins/FX/Niagara/Source/Niagara/Private/NiagaraDataSetDebugAccessor.cpp
--- a/Engine/Plugins/FX/Niagara/Source/Niagara/Private/NiagaraDataSetDebugAccessor.cpp
+++ b/Engine/Plugins/FX/Niagara/Source/Niagara/Private/NiagaraDataSetDebugAccessor.cpp
@@ -1,6 +1,5 @@
 // Copyright Epic Games, Inc. All Rights Reserved.
 #include "NiagaraDataSetDebugAccessor.h"
-#include "NiagaraSystem.h"
 
 bool FNiagaraDataSetDebugAccessor::Init(const FNiagaraDataSetCompiledData& CompiledData, FName InVariableName)
 {
diff --git a/Engine/Plugins/FX/Niagara/Source/Niagara/Private/NiagaraDebugHud.h b/Engine/Plugins/FX/Niagara/Source/Niagara/Private/NiagaraDebugHud.h
--- a/Engine/Plugins/FX/Niagara/Source/Niagara/Private/NiagaraDebugHud.h
+++ b/Engine/Plugins/FX/Niagara/Source/Niagara/Private/NiagaraDebugHud.h
@@ -8,6 +8,7 @@ Class used help debugging Niagara simulations
 #include "CoreMinimal.h"
 #include "NiagaraCommon.h"
 #include "NiagaraTypes.h"
+#include "RHI.h"
 #include "RHICommandList.h"
 #include "NiagaraDebuggerCommon.h"
 #include "NiagaraGPUProfilerInterface.h"
diff --git a/Engine/Plugins/FX/Niagara/Source/Niagara/Private/NiagaraGpuReadbackManager.h b/Engine/Plugins/FX/Niagara/Source/Niagara/Private/NiagaraGpuReadbackManager.h
--- a/Engine/Plugins/FX/Niagara/Source/Niagara/Private/NiagaraGpuReadbackManager.h
+++ b/Engine/Plugins/FX/Niagara/Source/Niagara/Private/NiagaraGpuReadbackManager.h
@@ -6,6 +6,7 @@
 #include "Containers/Queue.h"
 #include "NiagaraCommon.h"
 #include "RenderGraphResources.h"
+#include "RHI.h"
 #include "RHICommandList.h"
 
 class FRDGBuilder;
diff --git a/Engine/Plugins/FX/Niagara/Source/Niagara/Private/NiagaraLensEffectBase.cpp b/Engine/Plugins/FX/Niagara/Source/Niagara/Private/NiagaraLensEffectBase.cpp
--- a/Engine/Plugins/FX/Niagara/Source/Niagara/Private/NiagaraLensEffectBase.cpp
+++ b/Engine/Plugins/FX/Niagara/Source/Niagara/Private/NiagaraLensEffectBase.cpp
@@ -3,6 +3,7 @@
 #include "NiagaraLensEffectBase.h"
 
 #include "NiagaraComponent.h"
+#include "NiagaraEffectType.h"
 #include "NiagaraSystem.h"
 #include "Camera/PlayerCameraManager.h"
 
diff --git a/Engine/Plugins/FX/Niagara/Source/Niagara/Private/NiagaraModule.cpp b/Engine/Plugins/FX/Niagara/Source/Niagara/Private/NiagaraModule.cpp
--- a/Engine/Plugins/FX/Niagara/Source/Niagara/Private/NiagaraModule.cpp
+++ b/Engine/Plugins/FX/Niagara/Source/Niagara/Private/NiagaraModule.cpp
@@ -3,6 +3,7 @@
 #include "NiagaraModule.h"
 #include "Misc/LazySingleton.h"
 #include "Modules/ModuleManager.h"
+#include "NiagaraCompileHashVisitor.h"
 #include "NiagaraTypes.h"
 #include "NiagaraDebugVis.h"
 #include "NiagaraEvents.h"
diff --git a/Engine/Plugins/FX/Niagara/Source/Niagara/Private/NiagaraPerfBaseline.cpp b/Engine/Plugins/FX/Niagara/Source/Niagara/Private/NiagaraPerfBaseline.cpp
--- a/Engine/Plugins/FX/Niagara/Source/Niagara/Private/NiagaraPerfBaseline.cpp
+++ b/Engine/Plugins/FX/Niagara/Source/Niagara/Private/NiagaraPerfBaseline.cpp
@@ -3,6 +3,7 @@
 #include "NiagaraPerfBaseline.h"
 #include "Engine/Engine.h"
 #include "NiagaraComponent.h"
+#include "NiagaraEffectType.h"
 #include "NiagaraFunctionLibrary.h"
 #include "Components/TextRenderComponent.h"
 #include "Misc/OutputDeviceArchiveWrapper.h"
diff --git a/Engine/Plugins/FX/Niagara/Source/Niagara/Private/NiagaraPrecompileContainer.cpp b/Engine/Plugins/FX/Niagara/Source/Niagara/Private/NiagaraPrecompileContainer.cpp
--- a/Engine/Plugins/FX/Niagara/Source/Niagara/Private/NiagaraPrecompileContainer.cpp
+++ b/Engine/Plugins/FX/Niagara/Source/Niagara/Private/NiagaraPrecompileContainer.cpp
@@ -1,6 +1,7 @@
 // Copyright Epic Games, Inc. All Rights Reserved.
 
 #include "NiagaraPrecompileContainer.h"
+#include "NiagaraEffectType.h"
 
 #include UE_INLINE_GENERATED_CPP_BY_NAME(NiagaraPrecompileContainer)
 
diff --git a/Engine/Plugins/FX/Niagara/Source/Niagara/Private/NiagaraPreviewGrid.cpp b/Engine/Plugins/FX/Niagara/Source/Niagara/Private/NiagaraPreviewGrid.cpp
--- a/Engine/Plugins/FX/Niagara/Source/Niagara/Private/NiagaraPreviewGrid.cpp
+++ b/Engine/Plugins/FX/Niagara/Source/Niagara/Private/NiagaraPreviewGrid.cpp
@@ -5,6 +5,7 @@
 #include "Engine/BlueprintGeneratedClass.h"
 #include "Engine/Texture2D.h"
 #include "NiagaraComponent.h"
+#include "NiagaraEffectType.h"
 #include "NiagaraSystem.h"
 
 #include "Components/TextRenderComponent.h"
diff --git a/Engine/Plugins/FX/Niagara/Source/Niagara/Private/NiagaraSystemGpuComputeProxy.cpp b/Engine/Plugins/FX/Niagara/Source/Niagara/Private/NiagaraSystemGpuComputeProxy.cpp
--- a/Engine/Plugins/FX/Niagara/Source/Niagara/Private/NiagaraSystemGpuComputeProxy.cpp
+++ b/Engine/Plugins/FX/Niagara/Source/Niagara/Private/NiagaraSystemGpuComputeProxy.cpp
@@ -7,6 +7,7 @@
 #include "NiagaraGpuComputeDispatchInterface.h"
 #include "NiagaraGPUSystemTick.h"
 #include "NiagaraRenderer.h"
+#include "NiagaraSystemStaticBuffers.h"
 
 FNiagaraSystemGpuComputeProxy::FNiagaraSystemGpuComputeProxy(FNiagaraSystemInstance* OwnerInstance)
 	: DebugOwnerInstance(OwnerInstance)
diff --git a/Engine/Plugins/FX/Niagara/Source/Niagara/Private/NiagaraVariant.cpp b/Engine/Plugins/FX/Niagara/Source/Niagara/Private/NiagaraVariant.cpp
--- a/Engine/Plugins/FX/Niagara/Source/Niagara/Private/NiagaraVariant.cpp
+++ b/Engine/Plugins/FX/Niagara/Source/Niagara/Private/NiagaraVariant.cpp
@@ -1,6 +1,7 @@
 // Copyright Epic Games, Inc. All Rights Reserved.
 
 #include "NiagaraVariant.h"
+#include "NiagaraCompileHashVisitor.h"
 #include "NiagaraDataInterface.h"
 
 #include UE_INLINE_GENERATED_CPP_BY_NAME(NiagaraVariant)
diff --git a/Engine/Plugins/FX/Niagara/Source/Niagara/Public/NiagaraEmptyUAVPool.h b/Engine/Plugins/FX/Niagara/Source/Niagara/Public/NiagaraEmptyUAVPool.h
--- a/Engine/Plugins/FX/Niagara/Source/Niagara/Public/NiagaraEmptyUAVPool.h
+++ b/Engine/Plugins/FX/Niagara/Source/Niagara/Public/NiagaraEmptyUAVPool.h
@@ -3,6 +3,7 @@
 #pragma once
 
 #include "NiagaraCommon.h"
+#include "RHI.h"
 
 // Scoped access for UAVs this is not required when running inside the main dispatch loop
 // but is required for external usage that is outside, i.e. if you are doing some custom dispatch setup
diff --git a/Engine/Plugins/FX/Niagara/Source/Niagara/Public/NiagaraSystemGpuComputeProxy.h b/Engine/Plugins/FX/Niagara/Source/Niagara/Public/NiagaraSystemGpuComputeProxy.h
--- a/Engine/Plugins/FX/Niagara/Source/Niagara/Public/NiagaraSystemGpuComputeProxy.h
+++ b/Engine/Plugins/FX/Niagara/Source/Niagara/Public/NiagaraSystemGpuComputeProxy.h
@@ -3,6 +3,7 @@
 #pragma once
 
 #include "NiagaraCommon.h"
+#include "RHI.h"
 
 class FNiagaraSystemInstance;
 class FNiagaraGpuComputeDispatchInterface;
diff --git a/Engine/Plugins/FX/Niagara/Source/NiagaraEditor/Private/Commandlets/NiagaraDumpBytecodeCommandlet.cpp b/Engine/Plugins/FX/Niagara/Source/NiagaraEditor/Private/Commandlets/NiagaraDumpBytecodeCommandlet.cpp
--- a/Engine/Plugins/FX/Niagara/Source/NiagaraEditor/Private/Commandlets/NiagaraDumpBytecodeCommandlet.cpp
+++ b/Engine/Plugins/FX/Niagara/Source/NiagaraEditor/Private/Commandlets/NiagaraDumpBytecodeCommandlet.cpp
@@ -13,6 +13,8 @@
 #include "Misc/PackageName.h"
 #include "Misc/Paths.h"
 #include "Modules/ModuleManager.h"
+#include "NiagaraEmitter.h"
+#include "NiagaraScript.h"
 #include "NiagaraSystem.h"
 
 #include UE_INLINE_GENERATED_CPP_BY_NAME(NiagaraDumpBytecodeCommandlet)
diff --git a/Engine/Plugins/FX/Niagara/Source/NiagaraEditor/Private/NiagaraBakerRenderer.h b/Engine/Plugins/FX/Niagara/Source/NiagaraEditor/Private/NiagaraBakerRenderer.h
--- a/Engine/Plugins/FX/Niagara/Source/NiagaraEditor/Private/NiagaraBakerRenderer.h
+++ b/Engine/Plugins/FX/Niagara/Source/NiagaraEditor/Private/NiagaraBakerRenderer.h
@@ -3,8 +3,10 @@
 #pragma once
 
 #include "CoreMinimal.h"
-#include "UObject/GCObject.h"
+#include "Engine/EngineTypes.h"
+#include "NiagaraBakerOutput.h"
 #include "NiagaraSystem.h"
+#include "UObject/GCObject.h"
 
 class FNiagaraBakerRenderer;
 class UNiagaraComponent;
diff --git a/Engine/Plugins/FX/Niagara/Source/NiagaraEditor/Private/NiagaraBakerRendererOutputSimCache.cpp b/Engine/Plugins/FX/Niagara/Source/NiagaraEditor/Private/NiagaraBakerRendererOutputSimCache.cpp
--- a/Engine/Plugins/FX/Niagara/Source/NiagaraEditor/Private/NiagaraBakerRendererOutputSimCache.cpp
+++ b/Engine/Plugins/FX/Niagara/Source/NiagaraEditor/Private/NiagaraBakerRendererOutputSimCache.cpp
@@ -3,6 +3,7 @@
 #include "NiagaraBakerRendererOutputSimCache.h"
 #include "NiagaraBakerOutputSimCache.h"
 
+#include "NiagaraBakerSettings.h"
 #include "NiagaraComponent.h"
 #include "NiagaraGpuComputeDispatchInterface.h"
 #include "NiagaraSimCacheFactoryNew.h"
diff --git a/Engine/Plugins/FX/Niagara/Source/NiagaraEditor/Private/NiagaraBakerRendererOutputTexture2D.cpp b/Engine/Plugins/FX/Niagara/Source/NiagaraEditor/Private/NiagaraBakerRendererOutputTexture2D.cpp
--- a/Engine/Plugins/FX/Niagara/Source/NiagaraEditor/Private/NiagaraBakerRendererOutputTexture2D.cpp
+++ b/Engine/Plugins/FX/Niagara/Source/NiagaraEditor/Private/NiagaraBakerRendererOutputTexture2D.cpp
@@ -8,6 +8,7 @@
 #include "Factories/Texture2dFactoryNew.h"
 #include "CanvasTypes.h"
 #include "CanvasItem.h"
+#include "NiagaraBakerSettings.h"
 #include "TextureResource.h"
 
 TArray<FNiagaraBakerOutputBinding> FNiagaraBakerRendererOutputTexture2D::GetRendererBindings(UNiagaraBakerOutput* InBakerOutput) const
diff --git a/Engine/Plugins/FX/Niagara/Source/NiagaraEditor/Private/NiagaraBakerRendererOutputVolumeTexture.cpp b/Engine/Plugins/FX/Niagara/Source/NiagaraEditor/Private/NiagaraBakerRendererOutputVolumeTexture.cpp
--- a/Engine/Plugins/FX/Niagara/Source/NiagaraEditor/Private/NiagaraBakerRendererOutputVolumeTexture.cpp
+++ b/Engine/Plugins/FX/Niagara/Source/NiagaraEditor/Private/NiagaraBakerRendererOutputVolumeTexture.cpp
@@ -3,6 +3,7 @@
 #include "NiagaraBakerRendererOutputVolumeTexture.h"
 #include "NiagaraBakerOutputVolumeTexture.h"
 
+#include "NiagaraBakerSettings.h"
 #include "NiagaraComponent.h"
 #include "NiagaraSystemInstance.h"
 #include "NiagaraSystemInstanceController.h"
diff --git a/Engine/Plugins/FX/Niagara/Source/NiagaraEditor/Private/SNewSystemDialog.cpp b/Engine/Plugins/FX/Niagara/Source/NiagaraEditor/Private/SNewSystemDialog.cpp
--- a/Engine/Plugins/FX/Niagara/Source/NiagaraEditor/Private/SNewSystemDialog.cpp
+++ b/Engine/Plugins/FX/Niagara/Source/NiagaraEditor/Private/SNewSystemDialog.cpp
@@ -3,6 +3,7 @@
 #include "SNewSystemDialog.h"
 #include "NiagaraEditorModule.h"
 #include "NiagaraEditorStyle.h"
+#include "NiagaraEmitter.h"
 #include "NiagaraSystem.h"
 #include "SNiagaraAssetPickerList.h"
 
diff --git a/Engine/Plugins/FX/Niagara/Source/NiagaraEditor/Private/ViewModels/NiagaraBakerViewModel.h b/Engine/Plugins/FX/Niagara/Source/NiagaraEditor/Private/ViewModels/NiagaraBakerViewModel.h
--- a/Engine/Plugins/FX/Niagara/Source/NiagaraEditor/Private/ViewModels/NiagaraBakerViewModel.h
+++ b/Engine/Plugins/FX/Niagara/Source/NiagaraEditor/Private/ViewModels/NiagaraBakerViewModel.h
@@ -11,6 +11,7 @@
 #include "NiagaraBakerSettings.h"
 
 class UNiagaraBakerSettings;
+struct FSlateIcon;
 
 enum class ENiagaraBakerColorChannel : uint8
 {
diff --git a/Engine/Plugins/FX/Niagara/Source/NiagaraEditor/Private/ViewModels/Stack/NiagaraStackSystemPropertiesItem.cpp b/Engine/Plugins/FX/Niagara/Source/NiagaraEditor/Private/ViewModels/Stack/NiagaraStackSystemPropertiesItem.cpp
--- a/Engine/Plugins/FX/Niagara/Source/NiagaraEditor/Private/ViewModels/Stack/NiagaraStackSystemPropertiesItem.cpp
+++ b/Engine/Plugins/FX/Niagara/Source/NiagaraEditor/Private/ViewModels/Stack/NiagaraStackSystemPropertiesItem.cpp
@@ -1,6 +1,7 @@
 // Copyright Epic Games, Inc. All Rights Reserved.
 
 #include "ViewModels/Stack/NiagaraStackSystemPropertiesItem.h"
+#include "NiagaraEmitter.h"
 #include "ViewModels/Stack/NiagaraStackObject.h"
 #include "NiagaraSystem.h"
 #include "NiagaraSystemDetailsCustomization.h"
diff --git a/Engine/Plugins/FX/Niagara/Source/NiagaraEditorWidgets/Private/SNiagaraSystemEffectTypeBar.cpp b/Engine/Plugins/FX/Niagara/Source/NiagaraEditorWidgets/Private/SNiagaraSystemEffectTypeBar.cpp
--- a/Engine/Plugins/FX/Niagara/Source/NiagaraEditorWidgets/Private/SNiagaraSystemEffectTypeBar.cpp
+++ b/Engine/Plugins/FX/Niagara/Source/NiagaraEditorWidgets/Private/SNiagaraSystemEffectTypeBar.cpp
@@ -6,6 +6,7 @@
 #include "NiagaraEditorWidgetsStyle.h"
 #include "Delegates/DelegateInstanceInterface.h"
 #include "Modules/ModuleManager.h"
+#include "NiagaraEffectType.h"
 #include "Styling/StyleColors.h"
 #include "Widgets/Images/SImage.h"
 #include "Widgets/Layout/SWidgetSwitcher.h"
diff --git a/Engine/Plugins/Importers/AlembicImporter/Source/AlembicImporter/Private/AlembicImportFactory.cpp b/Engine/Plugins/Importers/AlembicImporter/Source/AlembicImporter/Private/AlembicImportFactory.cpp
--- a/Engine/Plugins/Importers/AlembicImporter/Source/AlembicImporter/Private/AlembicImportFactory.cpp
+++ b/Engine/Plugins/Importers/AlembicImporter/Source/AlembicImporter/Private/AlembicImportFactory.cpp
@@ -3,6 +3,7 @@
 #include "AlembicImportFactory.h"
 #include "AssetImportTask.h"
 #include "AssetRegistry/AssetRegistryModule.h"
+#include "Engine/AssetUserData.h"
 #include "Engine/StaticMesh.h"
 #include "Engine/SkeletalMesh.h"
 #include "Editor.h"
@@ -20,6 +21,7 @@
 #include "AbcAssetImportData.h"
 
 #include "GeometryCache.h"
+#include "RenderingThread.h"
 #include "Subsystems/AssetEditorSubsystem.h"
 
 #include "AI/Navigation/NavCollisionBase.h"
diff --git a/Engine/Plugins/Importers/AlembicImporter/Source/AlembicLibrary/Private/AbcImporter.cpp b/Engine/Plugins/Importers/AlembicImporter/Source/AlembicLibrary/Private/AbcImporter.cpp
--- a/Engine/Plugins/Importers/AlembicImporter/Source/AlembicLibrary/Private/AbcImporter.cpp
+++ b/Engine/Plugins/Importers/AlembicImporter/Source/AlembicLibrary/Private/AbcImporter.cpp
@@ -1,6 +1,7 @@
 // Copyright Epic Games, Inc. All Rights Reserved.
 
 #include "AbcImporter.h"
+#include "Animation/Skeleton.h"
 
 #include "BoneWeights.h"
 
diff --git a/Engine/Plugins/Interchange/Runtime/Source/Pipelines/Private/InterchangeGenericMeshPipeline.cpp b/Engine/Plugins/Interchange/Runtime/Source/Pipelines/Private/InterchangeGenericMeshPipeline.cpp
--- a/Engine/Plugins/Interchange/Runtime/Source/Pipelines/Private/InterchangeGenericMeshPipeline.cpp
+++ b/Engine/Plugins/Interchange/Runtime/Source/Pipelines/Private/InterchangeGenericMeshPipeline.cpp
@@ -2,6 +2,7 @@
 #include "InterchangeGenericMeshPipeline.h"
 
 #include "Animation/AnimSequence.h"
+#include "Animation/Skeleton.h"
 #include "Engine/SkeletalMesh.h"
 #include "Engine/StaticMesh.h"
 #include "InterchangeGenericAssetsPipeline.h"
diff --git a/Engine/Plugins/Runtime/HairStrands/Source/HairStrandsEditor/Private/GroomCustomAssetEditorToolkit.cpp b/Engine/Plugins/Runtime/HairStrands/Source/HairStrandsEditor/Private/GroomCustomAssetEditorToolkit.cpp
--- a/Engine/Plugins/Runtime/HairStrands/Source/HairStrandsEditor/Private/GroomCustomAssetEditorToolkit.cpp
+++ b/Engine/Plugins/Runtime/HairStrands/Source/HairStrandsEditor/Private/GroomCustomAssetEditorToolkit.cpp
@@ -1,6 +1,7 @@
 // Copyright Epic Games, Inc. All Rights Reserved.
 
 #include "GroomCustomAssetEditorToolkit.h"
+#include "Animation/Skeleton.h"
 #include "Modules/ModuleManager.h"
 #include "Styling/AppStyle.h"
 #include "Widgets/Docking/SDockTab.h"
diff --git a/Engine/Plugins/Runtime/Oculus/OculusVR/Source/OculusInput/Private/OculusHandTracking.cpp b/Engine/Plugins/Runtime/Oculus/OculusVR/Source/OculusInput/Private/OculusHandTracking.cpp
--- a/Engine/Plugins/Runtime/Oculus/OculusVR/Source/OculusInput/Private/OculusHandTracking.cpp
+++ b/Engine/Plugins/Runtime/Oculus/OculusVR/Source/OculusInput/Private/OculusHandTracking.cpp
@@ -1,6 +1,7 @@
 // Copyright Epic Games, Inc. All Rights Reserved.
 
 #include "OculusHandTracking.h"
+#include "Animation/Skeleton.h"
 #include "OculusHMD.h"
 #include "Misc/CoreDelegates.h"
 #include "IOculusInputModule.h"
diff --git a/Engine/Plugins/VirtualProduction/Takes/Source/TakeTrackRecorders/Private/MovieScene3DTransformTrackRecorder.cpp b/Engine/Plugins/VirtualProduction/Takes/Source/TakeTrackRecorders/Private/MovieScene3DTransformTrackRecorder.cpp
--- a/Engine/Plugins/VirtualProduction/Takes/Source/TakeTrackRecorders/Private/MovieScene3DTransformTrackRecorder.cpp
+++ b/Engine/Plugins/VirtualProduction/Takes/Source/TakeTrackRecorders/Private/MovieScene3DTransformTrackRecorder.cpp
@@ -1,6 +1,7 @@
 // Copyright Epic Games, Inc. All Rights Reserved.
 
 #include "TrackRecorders/MovieScene3DTransformTrackRecorder.h"
+#include "Animation/Skeleton.h"
 #include "TrackRecorders/MovieSceneAnimationTrackRecorder.h"
 #include "Tracks/MovieScene3DTransformTrack.h"
 #include "Sections/MovieScene3DTransformSection.h"
diff --git a/Engine/Plugins/VirtualProduction/Takes/Source/TakeTrackRecorders/Private/MovieSceneAnimationTrackRecorder.cpp b/Engine/Plugins/VirtualProduction/Takes/Source/TakeTrackRecorders/Private/MovieSceneAnimationTrackRecorder.cpp
--- a/Engine/Plugins/VirtualProduction/Takes/Source/TakeTrackRecorders/Private/MovieSceneAnimationTrackRecorder.cpp
+++ b/Engine/Plugins/VirtualProduction/Takes/Source/TakeTrackRecorders/Private/MovieSceneAnimationTrackRecorder.cpp
@@ -1,6 +1,7 @@
 // Copyright Epic Games, Inc. All Rights Reserved.
 
 #include "TrackRecorders/MovieSceneAnimationTrackRecorder.h"
+#include "Animation/Skeleton.h"
 #include "Engine/SkeletalMesh.h"
 #include "TrackRecorders/MovieSceneAnimationTrackRecorderSettings.h"
 #include "TakesUtils.h"
diff --git a/Engine/Source/Developer/AnimationDataController/Private/AnimDataController.cpp b/Engine/Source/Developer/AnimationDataController/Private/AnimDataController.cpp
--- a/Engine/Source/Developer/AnimationDataController/Private/AnimDataController.cpp
+++ b/Engine/Source/Developer/AnimationDataController/Private/AnimDataController.cpp
@@ -1,6 +1,7 @@
 // Copyright Epic Games, Inc. All Rights Reserved.
 
 #include "AnimDataController.h"
+#include "Animation/Skeleton.h"
 #include "AnimDataControllerActions.h"
 #include "MovieSceneTimeHelpers.h"
 
diff --git a/Engine/Source/Editor/LevelInstanceEditor/Private/LevelInstancePivotDetails.cpp b/Engine/Source/Editor/LevelInstanceEditor/Private/LevelInstancePivotDetails.cpp
--- a/Engine/Source/Editor/LevelInstanceEditor/Private/LevelInstancePivotDetails.cpp
+++ b/Engine/Source/Editor/LevelInstanceEditor/Private/LevelInstancePivotDetails.cpp
@@ -4,6 +4,7 @@
 #include "DetailLayoutBuilder.h"
 #include "DetailCategoryBuilder.h"
 #include "DetailWidgetRow.h"
+#include "Engine/Level.h"
 #include "Widgets/Input/SButton.h"
 #include "Widgets/Input/SComboBox.h"
 #include "Widgets/Input/SMultiLineEditableTextBox.h"
diff --git a/Engine/Source/Editor/Persona/Private/AnimNotifyDetails.cpp b/Engine/Source/Editor/Persona/Private/AnimNotifyDetails.cpp
--- a/Engine/Source/Editor/Persona/Private/AnimNotifyDetails.cpp
+++ b/Engine/Source/Editor/Persona/Private/AnimNotifyDetails.cpp
@@ -1,6 +1,7 @@
 // Copyright Epic Games, Inc. All Rights Reserved.
 
 #include "AnimNotifyDetails.h"
+#include "Animation/Skeleton.h"
 #include "Widgets/Text/STextBlock.h"
 #include "Animation/AnimTypes.h"
 #include "Animation/AnimationAsset.h"
diff --git a/Engine/Source/Editor/Persona/Private/AnimTimeline/AnimTimelineTrack_MontagePanel.cpp b/Engine/Source/Editor/Persona/Private/AnimTimeline/AnimTimelineTrack_MontagePanel.cpp
--- a/Engine/Source/Editor/Persona/Private/AnimTimeline/AnimTimelineTrack_MontagePanel.cpp
+++ b/Engine/Source/Editor/Persona/Private/AnimTimeline/AnimTimelineTrack_MontagePanel.cpp
@@ -1,6 +1,7 @@
 // Copyright Epic Games, Inc. All Rights Reserved.
 
 #include "AnimTimeline/AnimTimelineTrack_MontagePanel.h"
+#include "Animation/Skeleton.h"
 #include "SAnimMontagePanel.h"
 #include "Animation/AnimComposite.h"
 #include "PersonaUtils.h"
diff --git a/Engine/Source/Editor/Persona/Private/SAnimSegmentsPanel.cpp b/Engine/Source/Editor/Persona/Private/SAnimSegmentsPanel.cpp
--- a/Engine/Source/Editor/Persona/Private/SAnimSegmentsPanel.cpp
+++ b/Engine/Source/Editor/Persona/Private/SAnimSegmentsPanel.cpp
@@ -2,6 +2,7 @@
 
 
 #include "SAnimSegmentsPanel.h"
+#include "Animation/Skeleton.h"
 #include "Misc/MessageDialog.h"
 #include "Modules/ModuleManager.h"
 #include "Framework/MultiBox/MultiBoxBuilder.h"
diff --git a/Engine/Source/Editor/SequenceRecorder/Private/Sections/MovieScene3DTransformSectionRecorder.cpp b/Engine/Source/Editor/SequenceRecorder/Private/Sections/MovieScene3DTransformSectionRecorder.cpp
--- a/Engine/Source/Editor/SequenceRecorder/Private/Sections/MovieScene3DTransformSectionRecorder.cpp
+++ b/Engine/Source/Editor/SequenceRecorder/Private/Sections/MovieScene3DTransformSectionRecorder.cpp
@@ -1,6 +1,7 @@
 // Copyright Epic Games, Inc. All Rights Reserved.
 
 #include "Sections/MovieScene3DTransformSectionRecorder.h"
+#include "Animation/Skeleton.h"
 #include "Misc/ScopedSlowTask.h"
 #include "GameFramework/Character.h"
 #include "KeyParams.h"
diff --git a/Engine/Source/Editor/UnrealEd/Private/Animation/EditorAnimCompositeSegment.cpp b/Engine/Source/Editor/UnrealEd/Private/Animation/EditorAnimCompositeSegment.cpp
--- a/Engine/Source/Editor/UnrealEd/Private/Animation/EditorAnimCompositeSegment.cpp
+++ b/Engine/Source/Editor/UnrealEd/Private/Animation/EditorAnimCompositeSegment.cpp
@@ -6,6 +6,7 @@
 
 #include "Animation/EditorAnimCompositeSegment.h"
 #include "Animation/AnimComposite.h"
+#include "Animation/Skeleton.h"
 
 
 
diff --git a/Engine/Source/Editor/UnrealEd/Private/Animation/EditorAnimSegment.cpp b/Engine/Source/Editor/UnrealEd/Private/Animation/EditorAnimSegment.cpp
--- a/Engine/Source/Editor/UnrealEd/Private/Animation/EditorAnimSegment.cpp
+++ b/Engine/Source/Editor/UnrealEd/Private/Animation/EditorAnimSegment.cpp
@@ -6,6 +6,7 @@
 
 #include "Animation/EditorAnimSegment.h"
 #include "Animation/AnimMontage.h"
+#include "Animation/Skeleton.h"
 
 
 
diff --git a/Engine/Source/Editor/UnrealEd/Private/Commandlets/PackageUtilities.cpp b/Engine/Source/Editor/UnrealEd/Private/Commandlets/PackageUtilities.cpp
--- a/Engine/Source/Editor/UnrealEd/Private/Commandlets/PackageUtilities.cpp
+++ b/Engine/Source/Editor/UnrealEd/Private/Commandlets/PackageUtilities.cpp
@@ -5,6 +5,7 @@
 =============================================================================*/
 
 #include "CoreMinimal.h"
+#include "Animation/Skeleton.h"
 #include "HAL/FileManager.h"
 #include "Misc/CommandLine.h"
 #include "Misc/FileHelper.h"
diff --git a/Engine/Source/Editor/UnrealEd/Private/EditorExporters.cpp b/Engine/Source/Editor/UnrealEd/Private/EditorExporters.cpp
--- a/Engine/Source/Editor/UnrealEd/Private/EditorExporters.cpp
+++ b/Engine/Source/Editor/UnrealEd/Private/EditorExporters.cpp
@@ -5,6 +5,7 @@
 =============================================================================*/
 
 #include "CoreMinimal.h"
+#include "Animation/Skeleton.h"
 #include "Audio.h"
 #include "Misc/MessageDialog.h"
 #include "HAL/FileManager.h"
diff --git a/Engine/Source/Editor/UnrealEd/Private/Factories/MirrorDataTableFactory.cpp b/Engine/Source/Editor/UnrealEd/Private/Factories/MirrorDataTableFactory.cpp
--- a/Engine/Source/Editor/UnrealEd/Private/Factories/MirrorDataTableFactory.cpp
+++ b/Engine/Source/Editor/UnrealEd/Private/Factories/MirrorDataTableFactory.cpp
@@ -1,6 +1,7 @@
 // Copyright Epic Games, Inc. All Rights Reserved.
 
 #include "Factories/MirrorDataTableFactory.h"
+#include "Animation/Skeleton.h"
 #include "DataTableEditorUtils.h"
 #include "Animation/MirrorDataTable.h"
 #include "Animation/AnimationSettings.h"
diff --git a/Engine/Source/Editor/UnrealEd/Private/Fbx/FbxFactory.cpp b/Engine/Source/Editor/UnrealEd/Private/Fbx/FbxFactory.cpp
--- a/Engine/Source/Editor/UnrealEd/Private/Fbx/FbxFactory.cpp
+++ b/Engine/Source/Editor/UnrealEd/Private/Fbx/FbxFactory.cpp
@@ -3,6 +3,7 @@
 #include "Factories/FbxFactory.h"
 
 #include "Animation/AnimSequence.h"
+#include "Animation/Skeleton.h"
 #include "AssetImportTask.h"
 #include "AssetRegistry/AssetRegistryModule.h"
 #include "ContentBrowserModule.h"
diff --git a/Engine/Source/Editor/UnrealEd/Private/Fbx/FbxSkeletalMeshExport.cpp b/Engine/Source/Editor/UnrealEd/Private/Fbx/FbxSkeletalMeshExport.cpp
--- a/Engine/Source/Editor/UnrealEd/Private/Fbx/FbxSkeletalMeshExport.cpp
+++ b/Engine/Source/Editor/UnrealEd/Private/Fbx/FbxSkeletalMeshExport.cpp
@@ -5,6 +5,7 @@
 =============================================================================*/
 
 #include "CoreMinimal.h"
+#include "Animation/Skeleton.h"
 #include "BoneWeights.h"
 #include "GPUSkinPublicDefs.h"
 #include "Components/SkeletalMeshComponent.h"
diff --git a/Engine/Source/Editor/UnrealEd/Private/Fbx/ReimportFbxSceneFactory.cpp b/Engine/Source/Editor/UnrealEd/Private/Fbx/ReimportFbxSceneFactory.cpp
--- a/Engine/Source/Editor/UnrealEd/Private/Fbx/ReimportFbxSceneFactory.cpp
+++ b/Engine/Source/Editor/UnrealEd/Private/Fbx/ReimportFbxSceneFactory.cpp
@@ -1,6 +1,7 @@
 // Copyright Epic Games, Inc. All Rights Reserved.
 
 #include "Factories/ReimportFbxSceneFactory.h"
+#include "Animation/Skeleton.h"
 #include "Misc/Paths.h"
 #include "Misc/FeedbackContext.h"
 #include "Modules/ModuleManager.h"
diff --git a/Engine/Source/Editor/UnrealEd/Private/LevelEditorViewport.cpp b/Engine/Source/Editor/UnrealEd/Private/LevelEditorViewport.cpp
--- a/Engine/Source/Editor/UnrealEd/Private/LevelEditorViewport.cpp
+++ b/Engine/Source/Editor/UnrealEd/Private/LevelEditorViewport.cpp
@@ -2,6 +2,7 @@
 
 
 #include "LevelEditorViewport.h"
+#include "Animation/Skeleton.h"
 #include "Materials/MaterialInterface.h"
 #include "Modules/ModuleManager.h"
 #include "Misc/PackageName.h"
diff --git a/Engine/Source/Editor/UnrealEd/Private/ThumbnailHelpers.cpp b/Engine/Source/Editor/UnrealEd/Private/ThumbnailHelpers.cpp
--- a/Engine/Source/Editor/UnrealEd/Private/ThumbnailHelpers.cpp
+++ b/Engine/Source/Editor/UnrealEd/Private/ThumbnailHelpers.cpp
@@ -1,6 +1,7 @@
 // Copyright Epic Games, Inc. All Rights Reserved.
 
 #include "ThumbnailHelpers.h"
+#include "Engine/Level.h"
 #include "FinalPostProcessSettings.h"
 #include "SceneView.h"
 #include "Components/PrimitiveComponent.h"
diff --git a/Engine/Source/Editor/UnrealEd/Public/SSkeletonWidget.h b/Engine/Source/Editor/UnrealEd/Public/SSkeletonWidget.h
--- a/Engine/Source/Editor/UnrealEd/Public/SSkeletonWidget.h
+++ b/Engine/Source/Editor/UnrealEd/Public/SSkeletonWidget.h
@@ -2,6 +2,7 @@
 
 #pragma once
 
+#include "Animation/Skeleton.h"
 #include "CoreMinimal.h"
 #include "IContentBrowserSingleton.h"
 #include "Input/Reply.h"
diff --git a/Engine/Source/Runtime/AnimGraphRuntime/Private/AnimCustomInstanceHelper.cpp b/Engine/Source/Runtime/AnimGraphRuntime/Private/AnimCustomInstanceHelper.cpp
--- a/Engine/Source/Runtime/AnimGraphRuntime/Private/AnimCustomInstanceHelper.cpp
+++ b/Engine/Source/Runtime/AnimGraphRuntime/Private/AnimCustomInstanceHelper.cpp
@@ -2,6 +2,7 @@
 
 
 #include "AnimCustomInstanceHelper.h"
+#include "Animation/Skeleton.h"
 #include "Engine/SkeletalMesh.h"
 
 /////////////////////////////////////////////////////
diff --git a/Engine/Source/Runtime/AudioMixer/Private/Effects/AudioMixerSubmixEffectReverb.cpp b/Engine/Source/Runtime/AudioMixer/Private/Effects/AudioMixerSubmixEffectReverb.cpp
--- a/Engine/Source/Runtime/AudioMixer/Private/Effects/AudioMixerSubmixEffectReverb.cpp
+++ b/Engine/Source/Runtime/AudioMixer/Private/Effects/AudioMixerSubmixEffectReverb.cpp
@@ -2,6 +2,7 @@
 
 #include "SubmixEffects/AudioMixerSubmixEffectReverb.h"
 #include "AudioMixerEffectsManager.h"
+#include "HAL/IConsoleManager.h"
 #include "Sound/ReverbEffect.h"
 #include "Audio.h"
 #include "AudioMixer.h"
diff --git a/Engine/Source/Runtime/AudioMixer/Private/Quartz/QuartzSubsystem.cpp b/Engine/Source/Runtime/AudioMixer/Private/Quartz/QuartzSubsystem.cpp
--- a/Engine/Source/Runtime/AudioMixer/Private/Quartz/QuartzSubsystem.cpp
+++ b/Engine/Source/Runtime/AudioMixer/Private/Quartz/QuartzSubsystem.cpp
@@ -2,6 +2,7 @@
 
 #include "Quartz/QuartzSubsystem.h"
 
+#include "Quartz/AudioMixerClockHandle.h"
 #include "Quartz/QuartzMetronome.h"
 #include "Quartz/AudioMixerClockManager.h"
 #include "Sound/QuartzQuantizationUtilities.h"
diff --git a/Engine/Source/Runtime/Engine/Classes/Animation/AnimCurveCompressionCodec_UniformIndexable.h b/Engine/Source/Runtime/Engine/Classes/Animation/AnimCurveCompressionCodec_UniformIndexable.h
--- a/Engine/Source/Runtime/Engine/Classes/Animation/AnimCurveCompressionCodec_UniformIndexable.h
+++ b/Engine/Source/Runtime/Engine/Classes/Animation/AnimCurveCompressionCodec_UniformIndexable.h
@@ -6,6 +6,7 @@
 * Stores the raw rich curves as FCompressedRichCurve internally with optional key reduction and key time quantization.
 */
 
+#include "Animation/Skeleton.h"
 #include "CoreMinimal.h"
 #include "Animation/AnimCurveCompressionCodec.h"
 #include "AnimCurveCompressionCodec_UniformIndexable.generated.h"
diff --git a/Engine/Source/Runtime/Engine/Private/AI/Navigation/NavigationTypes.cpp b/Engine/Source/Runtime/Engine/Private/AI/Navigation/NavigationTypes.cpp
--- a/Engine/Source/Runtime/Engine/Private/AI/Navigation/NavigationTypes.cpp
+++ b/Engine/Source/Runtime/Engine/Private/AI/Navigation/NavigationTypes.cpp
@@ -4,6 +4,7 @@
 #include "AI/Navigation/NavRelevantInterface.h"
 #include "AI/NavigationSystemBase.h"
 #include "AI/Navigation/NavQueryFilter.h"
+#include "Engine/Level.h"
 #include "EngineStats.h"
 #include "Components/ShapeComponent.h"
 #include "AI/Navigation/NavAreaBase.h"
diff --git a/Engine/Source/Runtime/Engine/Private/AI/NavigationSystemBase.cpp b/Engine/Source/Runtime/Engine/Private/AI/NavigationSystemBase.cpp
--- a/Engine/Source/Runtime/Engine/Private/AI/NavigationSystemBase.cpp
+++ b/Engine/Source/Runtime/Engine/Private/AI/NavigationSystemBase.cpp
@@ -2,6 +2,7 @@
 
 #include "AI/NavigationSystemBase.h"
 #include "Engine/Engine.h"
+#include "Engine/Level.h"
 #include "Engine/World.h"
 #include "GameFramework/WorldSettings.h"
 #include "AI/Navigation/NavigationTypes.h"
diff --git a/Engine/Source/Runtime/Engine/Private/ActorEditorUtils.cpp b/Engine/Source/Runtime/Engine/Private/ActorEditorUtils.cpp
--- a/Engine/Source/Runtime/Engine/Private/ActorEditorUtils.cpp
+++ b/Engine/Source/Runtime/Engine/Private/ActorEditorUtils.cpp
@@ -2,6 +2,8 @@
 
 #include "ActorEditorUtils.h"
 #include "Components/ActorComponent.h"
+#include "Engine/Level.h"
+#include "Engine/Level.h"
 #include "GameFramework/Actor.h"
 #include "Engine/World.h"
 #include "Engine/Brush.h"
diff --git a/Engine/Source/Runtime/Engine/Private/ActorPartition/PartitionActor.cpp b/Engine/Source/Runtime/Engine/Private/ActorPartition/PartitionActor.cpp
--- a/Engine/Source/Runtime/Engine/Private/ActorPartition/PartitionActor.cpp
+++ b/Engine/Source/Runtime/Engine/Private/ActorPartition/PartitionActor.cpp
@@ -6,6 +6,7 @@
 
 #if WITH_EDITOR
 #include "Components/BoxComponent.h"
+#include "Engine/Level.h"
 #include "WorldPartition/WorldPartitionActorDesc.h"
 #include "WorldPartition/ActorPartition/PartitionActorDesc.h"
 #include "GameFramework/WorldSettings.h"
diff --git a/Engine/Source/Runtime/Engine/Private/Animation/AnimPhysicsSolver.cpp b/Engine/Source/Runtime/Engine/Private/Animation/AnimPhysicsSolver.cpp
--- a/Engine/Source/Runtime/Engine/Private/Animation/AnimPhysicsSolver.cpp
+++ b/Engine/Source/Runtime/Engine/Private/Animation/AnimPhysicsSolver.cpp
@@ -29,6 +29,7 @@ SOFTWARE.
 */
 
 #include "Animation/AnimPhysicsSolver.h"
+#include "HAL/IConsoleManager.h"
 #include "PhysicsEngine/PhysicsSettings.h"
 
 DEFINE_STAT(STAT_AnimDynamicsUpdate);
diff --git a/Engine/Source/Runtime/Engine/Private/Animation/CurveIdentifier.cpp b/Engine/Source/Runtime/Engine/Private/Animation/CurveIdentifier.cpp
--- a/Engine/Source/Runtime/Engine/Private/Animation/CurveIdentifier.cpp
+++ b/Engine/Source/Runtime/Engine/Private/Animation/CurveIdentifier.cpp
@@ -1,6 +1,7 @@
 // Copyright Epic Games, Inc. All Rights Reserved.
 
 #include "Animation/AnimData/CurveIdentifier.h"
+#include "Animation/Skeleton.h"
 #include "EngineLogs.h"
 
 #include UE_INLINE_GENERATED_CPP_BY_NAME(CurveIdentifier)
diff --git a/Engine/Source/Runtime/Engine/Private/Animation/MirrorDataTable.cpp b/Engine/Source/Runtime/Engine/Private/Animation/MirrorDataTable.cpp
--- a/Engine/Source/Runtime/Engine/Private/Animation/MirrorDataTable.cpp
+++ b/Engine/Source/Runtime/Engine/Private/Animation/MirrorDataTable.cpp
@@ -2,6 +2,7 @@
 
 #include "Animation/MirrorDataTable.h"
 #include "Animation/AnimationSettings.h"
+#include "Animation/Skeleton.h"
 #include "Internationalization/Regex.h"
 
 #include UE_INLINE_GENERATED_CPP_BY_NAME(MirrorDataTable)
diff --git a/Engine/Source/Runtime/Engine/Private/Components/AudioComponent.cpp b/Engine/Source/Runtime/Engine/Private/Components/AudioComponent.cpp
--- a/Engine/Source/Runtime/Engine/Private/Components/AudioComponent.cpp
+++ b/Engine/Source/Runtime/Engine/Private/Components/AudioComponent.cpp
@@ -11,6 +11,7 @@
 #include "Engine/Texture2D.h"
 #include "Kismet/GameplayStatics.h"
 #include "Misc/App.h"
+#include "Quartz/AudioMixerClockHandle.h"
 #include "Sound/SoundCue.h"
 #include "Sound/SoundNodeAttenuation.h"
 #include "UObject/FrameworkObjectVersion.h"
diff --git a/Engine/Source/Runtime/Engine/Private/Components/ModelComponent.cpp b/Engine/Source/Runtime/Engine/Private/Components/ModelComponent.cpp
--- a/Engine/Source/Runtime/Engine/Private/Components/ModelComponent.cpp
+++ b/Engine/Source/Runtime/Engine/Private/Components/ModelComponent.cpp
@@ -5,6 +5,7 @@
 =============================================================================*/
 
 #include "Components/ModelComponent.h"
+#include "Engine/Level.h"
 #include "Engine/MapBuildDataRegistry.h"
 #include "Materials/Material.h"
 #include "Engine/CollisionProfile.h"
diff --git a/Engine/Source/Runtime/Engine/Private/Components/PrimitiveComponent.cpp b/Engine/Source/Runtime/Engine/Private/Components/PrimitiveComponent.cpp
--- a/Engine/Source/Runtime/Engine/Private/Components/PrimitiveComponent.cpp
+++ b/Engine/Source/Runtime/Engine/Private/Components/PrimitiveComponent.cpp
@@ -5,6 +5,7 @@
 =============================================================================*/
 
 #include "Components/PrimitiveComponent.h"
+#include "Engine/Level.h"
 #include "EngineStats.h"
 #include "GameFramework/DamageType.h"
 #include "GameFramework/Pawn.h"
diff --git a/Engine/Source/Runtime/Engine/Private/Components/SceneComponent.cpp b/Engine/Source/Runtime/Engine/Private/Components/SceneComponent.cpp
--- a/Engine/Source/Runtime/Engine/Private/Components/SceneComponent.cpp
+++ b/Engine/Source/Runtime/Engine/Private/Components/SceneComponent.cpp
@@ -6,6 +6,7 @@
 
 
 #include "Components/SceneComponent.h"
+#include "Engine/Level.h"
 #include "EngineStats.h"
 #include "Engine/Blueprint.h"
 #include "GameFramework/Actor.h"
diff --git a/Engine/Source/Runtime/Engine/Private/HLOD/HLODProxyDesc.cpp b/Engine/Source/Runtime/Engine/Private/HLOD/HLODProxyDesc.cpp
--- a/Engine/Source/Runtime/Engine/Private/HLOD/HLODProxyDesc.cpp
+++ b/Engine/Source/Runtime/Engine/Private/HLOD/HLODProxyDesc.cpp
@@ -8,6 +8,7 @@
 #if WITH_EDITOR
 #include "Engine/LODActor.h"
 #include "Algo/Transform.h"
+#include "Engine/Level.h"
 #include "GameFramework/WorldSettings.h"
 #include "LevelUtils.h"
 #include "Engine/LevelStreaming.h"
diff --git a/Engine/Source/Runtime/Engine/Private/HLODProxy.cpp b/Engine/Source/Runtime/Engine/Private/HLODProxy.cpp
--- a/Engine/Source/Runtime/Engine/Private/HLODProxy.cpp
+++ b/Engine/Source/Runtime/Engine/Private/HLODProxy.cpp
@@ -2,6 +2,7 @@
 
 #include "Engine/HLODProxy.h"
 
+#include "Engine/Level.h"
 #include "Engine/LODActor.h"
 #include "GameFramework/WorldSettings.h"
 #include "Materials/Material.h"
diff --git a/Engine/Source/Runtime/Engine/Private/InstancedStaticMesh.cpp b/Engine/Source/Runtime/Engine/Private/InstancedStaticMesh.cpp
--- a/Engine/Source/Runtime/Engine/Private/InstancedStaticMesh.cpp
+++ b/Engine/Source/Runtime/Engine/Private/InstancedStaticMesh.cpp
@@ -5,6 +5,7 @@
 =============================================================================*/
 
 #include "Engine/InstancedStaticMesh.h"
+#include "Engine/Level.h"
 #include "InstancedStaticMeshDelegates.h"
 #include "AI/NavigationSystemBase.h"
 #include "Engine/MapBuildDataRegistry.h"
diff --git a/Engine/Source/Runtime/Engine/Private/LevelUtils.cpp b/Engine/Source/Runtime/Engine/Private/LevelUtils.cpp
--- a/Engine/Source/Runtime/Engine/Private/LevelUtils.cpp
+++ b/Engine/Source/Runtime/Engine/Private/LevelUtils.cpp
@@ -3,6 +3,7 @@
 
 #include "LevelUtils.h"
 #include "Engine/Engine.h"
+#include "Engine/Level.h"
 #include "Engine/LevelStreaming.h"
 #include "HAL/FileManager.h"
 #include "UObject/Package.h"
diff --git a/Engine/Source/Runtime/Engine/Private/LightMap.cpp b/Engine/Source/Runtime/Engine/Private/LightMap.cpp
--- a/Engine/Source/Runtime/Engine/Private/LightMap.cpp
+++ b/Engine/Source/Runtime/Engine/Private/LightMap.cpp
@@ -5,6 +5,7 @@
 =============================================================================*/
 
 #include "LightMap.h"
+#include "Engine/Level.h"
 #include "ShadowMap.h"
 #include "Engine/ShadowMapTexture2D.h"
 #include "UnrealEngine.h"
diff --git a/Engine/Source/Runtime/Engine/Private/Rendering/ColorVertexBuffer.cpp b/Engine/Source/Runtime/Engine/Private/Rendering/ColorVertexBuffer.cpp
--- a/Engine/Source/Runtime/Engine/Private/Rendering/ColorVertexBuffer.cpp
+++ b/Engine/Source/Runtime/Engine/Private/Rendering/ColorVertexBuffer.cpp
@@ -6,6 +6,7 @@
 #include "Components.h"
 #include "EngineUtils.h"
 #include "ProfilingDebugging/LoadTimeTracker.h"
+#include "VertexFactory.h"
 
 /*-----------------------------------------------------------------------------
 FColorVertexBuffer
diff --git a/Engine/Source/Runtime/Engine/Private/ShadowMap.cpp b/Engine/Source/Runtime/Engine/Private/ShadowMap.cpp
--- a/Engine/Source/Runtime/Engine/Private/ShadowMap.cpp
+++ b/Engine/Source/Runtime/Engine/Private/ShadowMap.cpp
@@ -1,6 +1,7 @@
 // Copyright Epic Games, Inc. All Rights Reserved.
 
 #include "ShadowMap.h"
+#include "Engine/Level.h"
 #include "Engine/MapBuildDataRegistry.h"
 #include "Components/LightComponent.h"
 
diff --git a/Engine/Source/Runtime/Engine/Private/WorldPartition/WorldPartitionLevelStreamingDynamic.cpp b/Engine/Source/Runtime/Engine/Private/WorldPartition/WorldPartitionLevelStreamingDynamic.cpp
--- a/Engine/Source/Runtime/Engine/Private/WorldPartition/WorldPartitionLevelStreamingDynamic.cpp
+++ b/Engine/Source/Runtime/Engine/Private/WorldPartition/WorldPartitionLevelStreamingDynamic.cpp
@@ -11,6 +11,7 @@
 #include "UObject/UObjectHash.h"
 #include "Misc/PackageName.h"
 #include "Engine/Engine.h"
+#include "Engine/Level.h"
 #include "Engine/World.h"
 #include "Misc/PathViews.h"
 #include "Model.h"
diff --git a/Engine/Source/Runtime/Engine/Public/SurfaceIterators.h b/Engine/Source/Runtime/Engine/Public/SurfaceIterators.h
--- a/Engine/Source/Runtime/Engine/Public/SurfaceIterators.h
+++ b/Engine/Source/Runtime/Engine/Public/SurfaceIterators.h
@@ -7,6 +7,7 @@
 #pragma once
 
 #include "CoreMinimal.h"
+#include "Engine/Level.h"
 #include "Engine/World.h"
 #include "Model.h"
 
diff --git a/Engine/Source/Runtime/Experimental/Chaos/Public/Chaos/Deformable/Utilities.h b/Engine/Source/Runtime/Experimental/Chaos/Public/Chaos/Deformable/Utilities.h
--- a/Engine/Source/Runtime/Experimental/Chaos/Public/Chaos/Deformable/Utilities.h
+++ b/Engine/Source/Runtime/Experimental/Chaos/Public/Chaos/Deformable/Utilities.h
@@ -9,6 +9,7 @@
 #include "Chaos/Vector.h"
 #include "Chaos/Triangle.h"
 
+#include "ChaosFlesh/FleshCollectionUtility.h"
 #include "Math/NumericLimits.h"
 #include "Templates/IsIntegral.h"
 #include "Chaos/Plane.h"
diff --git a/Engine/Source/Runtime/Landscape/Classes/Landscape.h b/Engine/Source/Runtime/Landscape/Classes/Landscape.h
--- a/Engine/Source/Runtime/Landscape/Classes/Landscape.h
+++ b/Engine/Source/Runtime/Landscape/Classes/Landscape.h
@@ -3,6 +3,7 @@
 #pragma once
 
 #include "CoreMinimal.h"
+#include "LandscapeComponent.h"
 #include "UObject/ObjectMacros.h"
 #include "LandscapeProxy.h"
 #include "LandscapeBlueprintBrushBase.h"
diff --git a/Engine/Source/Runtime/Renderer/Private/RendererScene.cpp b/Engine/Source/Runtime/Renderer/Private/RendererScene.cpp
--- a/Engine/Source/Runtime/Renderer/Private/RendererScene.cpp
+++ b/Engine/Source/Runtime/Renderer/Private/RendererScene.cpp
@@ -5,6 +5,7 @@
 =============================================================================*/
 
 #include "CoreMinimal.h"
+#include "Engine/Level.h"
 #include "Engine/TextureLightProfile.h"
 #include "HAL/ThreadSafeCounter.h"
 #include "HAL/PlatformFileManager.h"
diff --git a/Engine/Source/Runtime/UMG/Private/WidgetBlueprintLibrary.cpp b/Engine/Source/Runtime/UMG/Private/WidgetBlueprintLibrary.cpp
--- a/Engine/Source/Runtime/UMG/Private/WidgetBlueprintLibrary.cpp
+++ b/Engine/Source/Runtime/UMG/Private/WidgetBlueprintLibrary.cpp
@@ -2,6 +2,7 @@
 
 #include "Blueprint/WidgetBlueprintLibrary.h"
 #include "Materials/MaterialInterface.h"
+#include "Slate/SGameLayerManager.h"
 #include "UObject/UObjectHash.h"
 #include "UObject/UObjectIterator.h"
 #include "UObject/Package.h"
