support dynamically creating attributes from emitter scripts

diff --git a/Engine/Plugins/FX/Niagara/Source/Niagara/Classes/NiagaraDataInterfaceRasterizationGrid3D.h b/Engine/Plugins/FX/Niagara/Source/Niagara/Classes/NiagaraDataInterfaceRasterizationGrid3D.h
--- a/Engine/Plugins/FX/Niagara/Source/Niagara/Classes/NiagaraDataInterfaceRasterizationGrid3D.h
+++ b/Engine/Plugins/FX/Niagara/Source/Niagara/Classes/NiagaraDataInterfaceRasterizationGrid3D.h
@@ -113,6 +113,7 @@ public:
 
 	void VMGetNumCells(FVectorVMExternalFunctionContext& Context);
 	void VMSetNumCells(FVectorVMExternalFunctionContext& Context);
+	void VMSetNumAttributes(FVectorVMExternalFunctionContext& Context);
 	void VMSetFloatResetValue(FVectorVMExternalFunctionContext& Context);
 
 protected:
diff --git a/Engine/Plugins/FX/Niagara/Source/Niagara/Private/NiagaraDataInterfaceRasterizationGrid3D.cpp b/Engine/Plugins/FX/Niagara/Source/Niagara/Private/NiagaraDataInterfaceRasterizationGrid3D.cpp
--- a/Engine/Plugins/FX/Niagara/Source/Niagara/Private/NiagaraDataInterfaceRasterizationGrid3D.cpp
+++ b/Engine/Plugins/FX/Niagara/Source/Niagara/Private/NiagaraDataInterfaceRasterizationGrid3D.cpp
@@ -29,6 +29,7 @@ namespace NDIRasterizationGrid3DLocal
 	static const FString PrecisionName(TEXT("_Precision"));
 
 	static const FName SetNumCellsFunctionName("SetNumCells");
+	static const FName SetNumAttributesFunctionName("SetNumAttributes");
 	static const FName SetFloatResetValueFunctionName("SetFloatResetValue");
 
 	static const FString PerAttributeDataName(TEXT("_PerAttributeData"));
@@ -97,6 +98,23 @@ void UNiagaraDataInterfaceRasterizationGrid3D::GetFunctions(TArray<FNiagaraFunct
 		OutFunctions.Add(Sig);
 	}
 
+	{
+		FNiagaraFunctionSignature Sig;
+		Sig.Name = SetNumAttributesFunctionName;
+		Sig.Inputs.Add(FNiagaraVariable(FNiagaraTypeDefinition(GetClass()), TEXT("Grid")));
+		Sig.Inputs.Add(FNiagaraVariable(FNiagaraTypeDefinition::GetIntDef(), TEXT("NumAttributes")));		
+		Sig.Outputs.Add(FNiagaraVariable(FNiagaraTypeDefinition::GetBoolDef(), TEXT("Success")));
+
+		Sig.ModuleUsageBitmask = ENiagaraScriptUsageMask::Emitter | ENiagaraScriptUsageMask::System;
+		Sig.bExperimental = true;
+		Sig.bMemberFunction = true;
+		Sig.bRequiresExecPin = true;
+		Sig.bRequiresContext = false;
+		Sig.bSupportsCPU = true;
+		Sig.bSupportsGPU = false;
+		OutFunctions.Add(Sig);
+	}
+	
 	{
 		FNiagaraFunctionSignature Sig;
 		Sig.Name = SetFloatResetValueFunctionName;
@@ -438,6 +456,11 @@ void UNiagaraDataInterfaceRasterizationGrid3D::GetVMExternalFunction(const FVMEx
 		check(BindingInfo.GetNumInputs() == 4 && BindingInfo.GetNumOutputs() == 1);
 		OutFunc.BindUObject(this, &UNiagaraDataInterfaceRasterizationGrid3D::VMSetNumCells);
 	}
+	else if (BindingInfo.Name == SetNumAttributesFunctionName)
+	{
+		check(BindingInfo.GetNumInputs() == 2 && BindingInfo.GetNumOutputs() == 1);
+		OutFunc.BindUObject(this, &UNiagaraDataInterfaceRasterizationGrid3D::VMSetNumAttributes);
+	}
 	else if (BindingInfo.Name == SetFloatResetValueFunctionName)
 	{
 		check(BindingInfo.GetNumInputs() == 2 && BindingInfo.GetNumOutputs() == 1);
@@ -908,9 +931,14 @@ bool UNiagaraDataInterfaceRasterizationGrid3D::InitPerInstanceData(void* PerInst
 	InstanceData->NumTiles.Y = NumTilesY;
 	InstanceData->NumTiles.Z = NumTilesZ;
 
-	check(InstanceData->NumTiles.X > 0);
-	check(InstanceData->NumTiles.Y > 0);
-	check(InstanceData->NumTiles.Z > 0);
+	if (NumAttribChannelsFound == 0)
+	{
+		UE_LOG(LogNiagara, Warning, TEXT("Zero attributes defined on %s"), *Proxy->SourceDIName.ToString());
+	}
+
+//	check(InstanceData->NumTiles.X > 0);
+//	check(InstanceData->NumTiles.Y > 0);
+//	check(InstanceData->NumTiles.Z > 0);
 
 	// @todo-threadsafety. This would be a race but I'm taking a ref here. Not ideal in the long term.
 	// Push Updates to Proxy.
@@ -959,6 +987,29 @@ void UNiagaraDataInterfaceRasterizationGrid3D::VMSetNumCells(FVectorVMExternalFu
 	}
 }
 
+void UNiagaraDataInterfaceRasterizationGrid3D::VMSetNumAttributes(FVectorVMExternalFunctionContext& Context)
+{
+	// This should only be called from a system or emitter script due to a need for only setting up initially.
+	VectorVM::FUserPtrHandler<RasterizationGrid3DRWInstanceData> InstData(Context);
+	VectorVM::FExternalFuncInputHandler<int> InNumAttributes(Context);	
+	VectorVM::FExternalFuncRegisterHandler<FNiagaraBool> OutSuccess(Context);
+
+	for (int32 InstanceIdx = 0; InstanceIdx < Context.GetNumInstances(); ++InstanceIdx)
+	{
+		int NewNumAttributes = InNumAttributes.GetAndAdvance();
+		bool bSuccess = (InstData.Get() != nullptr && Context.GetNumInstances() == 1 && NumCells.X >= 0 && NumCells.Y >= 0 && NumCells.Z >= 0);
+		*OutSuccess.GetDestAndAdvance() = bSuccess;
+		if (bSuccess)
+		{
+			int OldNumAttributes = InstData->TotalNumAttributes;
+
+			InstData->TotalNumAttributes = NewNumAttributes;
+
+			InstData->NeedsRealloc = OldNumAttributes != InstData->TotalNumAttributes;
+		}
+	}
+}
+
 void UNiagaraDataInterfaceRasterizationGrid3D::VMSetFloatResetValue(FVectorVMExternalFunctionContext& Context)
 {
 	VectorVM::FUserPtrHandler<RasterizationGrid3DRWInstanceData> InstData(Context);
@@ -983,7 +1034,7 @@ bool UNiagaraDataInterfaceRasterizationGrid3D::PerInstanceTickPostSimulate(void*
 	RasterizationGrid3DRWInstanceData* InstanceData = static_cast<RasterizationGrid3DRWInstanceData*>(PerInstanceData);
 	bool bNeedsReset = false;
 
-	if (InstanceData->NeedsRealloc && InstanceData->NumCells.X > 0 && InstanceData->NumCells.Y > 0 && InstanceData->NumCells.Z > 0)
+	if (InstanceData->NeedsRealloc && InstanceData->NumCells.X > 0 && InstanceData->NumCells.Y > 0 && InstanceData->NumCells.Z > 0 && InstanceData->TotalNumAttributes > 0)
 	{
 		InstanceData->NeedsRealloc = false;
 		
@@ -1004,6 +1055,12 @@ bool UNiagaraDataInterfaceRasterizationGrid3D::PerInstanceTickPostSimulate(void*
 		InstanceData->NumTiles.Y = NumTilesY;
 		InstanceData->NumTiles.Z = NumTilesZ;
 
+		if (InstanceData->TotalNumAttributes <= 0)
+		{
+			UE_LOG(LogNiagara, Warning, TEXT("Invalid number of attributes defined on %s... max is %d, num defined is %d"), *Proxy->SourceDIName.ToString(), MaxAttributes, InstanceData->TotalNumAttributes);
+			return false;
+		}
+
 		check(InstanceData->NumTiles.X > 0);
 		check(InstanceData->NumTiles.Y > 0);
 		check(InstanceData->NumTiles.Z > 0);
@@ -1053,14 +1110,18 @@ void FNiagaraDataInterfaceProxyRasterizationGrid3D::PreStage(const FNDIGpuComput
 
 	RasterizationGrid3DRWInstanceData& InstanceData = SystemInstancesToProxyData.FindChecked(Context.GetSystemInstanceID());
 
+	const uint32 NumTotalCells = InstanceData.NumCells.X * InstanceData.NumCells.Y * InstanceData.NumCells.Z * InstanceData.TotalNumAttributes;
+
 	// Resize was requested
 	if (InstanceData.NeedsRealloc)
 	{
 		InstanceData.NeedsRealloc = false;
-
-		const uint32 NumTotalCells = InstanceData.NumCells.X * InstanceData.NumCells.Y * InstanceData.NumCells.Z;
-		if (NumTotalCells <= (uint32)GMaxNiagaraRasterizationGridCells)
+		
+		
+		if (NumTotalCells <= (uint32)GMaxNiagaraRasterizationGridCells && NumTotalCells > 0)
 		{
+			InstanceData.PerAttributeData.Release();
+
 			const FIntVector TextureSize(InstanceData.NumCells.X * InstanceData.NumTiles.X, InstanceData.NumCells.Y * InstanceData.NumTiles.Y, InstanceData.NumCells.Z * InstanceData.NumTiles.Z);
 			const FRDGTextureDesc TextureDesc = FRDGTextureDesc::Create3D(TextureSize, PF_R32_SINT, FClearValueBinding::Black, ETextureCreateFlags::ShaderResource | ETextureCreateFlags::UAV);
 			InstanceData.RasterizationTexture.Initialize(Context.GetGraphBuilder(), TEXT("NiagaraRasterizationGrid3D::IntGrid"), TextureDesc);
@@ -1092,7 +1153,7 @@ void FNiagaraDataInterfaceProxyRasterizationGrid3D::PreStage(const FNDIGpuComput
 		InstanceData.PerAttributeData.Initialize(TEXT("Grid3D::PerAttributeData"), sizeof(FVector4f), PerAttributeData.Num(), EPixelFormat::PF_A32B32G32R32F, BUF_Static, &PerAttributeData);
 	}
 
-	if (Context.IsOutputStage())
+	if (Context.IsOutputStage() && NumTotalCells > 0)
 	{
 		FRDGBuilder& GraphBuilder = Context.GetGraphBuilder();
 		const FUintVector4 ResetValue(InstanceData.ResetValue, InstanceData.ResetValue, InstanceData.ResetValue, InstanceData.ResetValue);
