fixed error checking for case where reader and the grid are mismatched for rgba status
fixed error checking when grids have unnamed and named attributes

diff --git a/Engine/Plugins/FX/Niagara/Source/Niagara/Private/NiagaraDataInterfaceGrid3DCollection.cpp b/Engine/Plugins/FX/Niagara/Source/Niagara/Private/NiagaraDataInterfaceGrid3DCollection.cpp
--- a/Engine/Plugins/FX/Niagara/Source/Niagara/Private/NiagaraDataInterfaceGrid3DCollection.cpp
+++ b/Engine/Plugins/FX/Niagara/Source/Niagara/Private/NiagaraDataInterfaceGrid3DCollection.cpp
@@ -3215,7 +3215,7 @@ bool UNiagaraDataInterfaceGrid3DCollection::InitPerInstanceData(void* PerInstanc
 	InstanceData->TotalNumNamedAttributes = NumNamedAttribChannelsFound;
 
 	// give an error if there are any unnamed attributes
-	if (NumAttribChannelsFound - NumNamedAttribChannelsFound > 0)
+	if (NumAttributes > 0 && NumNamedAttribChannelsFound > 0)
 	{
 		UE_LOG(LogNiagara, Warning, TEXT("%s: Unnamed attributes should not be used with named.  This is a deprecated workflow"), *Proxy->SourceDIName.ToString());
 	}
diff --git a/Engine/Plugins/FX/Niagara/Source/Niagara/Private/NiagaraDataInterfaceGrid3DCollectionReader.cpp b/Engine/Plugins/FX/Niagara/Source/Niagara/Private/NiagaraDataInterfaceGrid3DCollectionReader.cpp
--- a/Engine/Plugins/FX/Niagara/Source/Niagara/Private/NiagaraDataInterfaceGrid3DCollectionReader.cpp
+++ b/Engine/Plugins/FX/Niagara/Source/Niagara/Private/NiagaraDataInterfaceGrid3DCollectionReader.cpp
@@ -121,7 +121,7 @@ bool UNiagaraDataInterfaceGrid3DCollectionReader::InitPerInstanceData(void* PerI
 			InstanceData->OtherInstanceData = Grid3DInstanceData;
 
 			// error if the reader thinks it should be an rgba grid and the underlying grid collection disagrees.  This is due to the reader not using all attributes defined
-			if (NumAttribChannelsFound > 0 && UseRGBATexture != InstanceData->OtherInstanceData->UseRGBATexture)
+			if (NumAttribChannelsFound > 0 && InstanceData->OtherInstanceData != nullptr && UseRGBATexture != InstanceData->OtherInstanceData->UseRGBATexture)
 			{
 				UE_LOG(LogNiagara, Error, TEXT("RGBA grids and Grid3DCollectionReaders are not compatible with Emitter %s and DataInterface %s.  Turn off RGBA grids with fx.Niagara.Grid3D.UseRGBAGrid 0, or split up your grid collections so the one passed to this reader has only one attribute."), *this->EmitterName, *this->DIName);
 				return false;
@@ -220,6 +220,31 @@ bool UNiagaraDataInterfaceGrid3DCollectionReader::PerInstanceTickPostSimulate(vo
 		FGrid3DCollectionRWInstanceData_GameThread* Grid3DInstanceData = OtherGridDI->GetSystemInstancesToProxyData_GT().FindRef(SystemInstance->GetId());
 
 		InstanceData->OtherInstanceData = Grid3DInstanceData;
+
+		// Does the reader think it should be an rgba grid?
+		int32 NumAttribChannelsFound = 0;
+		int32 NumNamedAttribChannelsFound = 0;
+		TArray<FNiagaraVariableBase> Vars;
+		TArray<uint32> Offsets;
+		FindAttributes(Vars, Offsets, NumNamedAttribChannelsFound, nullptr, true);
+
+		// #todo(dmp): slight hack here - these aren't always the same, but we generally can't use unnamed attrs very well
+		NumAttribChannelsFound = NumNamedAttribChannelsFound;
+		bool UseRGBATexture = FGrid3DCollectionAttributeHelper::ShouldUseRGBAGrid(NumAttribChannelsFound, Vars.Num());
+
+		if (InstanceData->OtherInstanceData == nullptr)
+		{
+			UE_LOG(LogNiagara, Error, TEXT("%s: The DI for Emitter %s and DataInterface %s is not initialized."), *Proxy->SourceDIName.ToString(), *this->EmitterName, *this->DIName);
+			return false;
+		}
+
+		// error if the reader thinks it should be an rgba grid and the underlying grid collection disagrees.  This is due to the reader not using all attributes defined
+		if (NumAttribChannelsFound > 0 && UseRGBATexture != InstanceData->OtherInstanceData->UseRGBATexture)
+		{
+			UE_LOG(LogNiagara, Error, TEXT("RGBA grids and Grid3DCollectionReaders are not compatible with Emitter %s and DataInterface %s.  Turn off RGBA grids with fx.Niagara.Grid3D.UseRGBAGrid 0, or split up your grid collections so the one passed to this reader has only one attribute."), *this->EmitterName, *this->DIName);
+			return false;
+		}
+
 	}
 
 	return false;
