- Add static variables to debug HUD
- Replace user variables with parameter store variables so we can look in all locations

diff --git a/Engine/Plugins/FX/Niagara/Source/Niagara/Private/NiagaraDebugHud.cpp b/Engine/Plugins/FX/Niagara/Source/Niagara/Private/NiagaraDebugHud.cpp
--- a/Engine/Plugins/FX/Niagara/Source/Niagara/Private/NiagaraDebugHud.cpp
+++ b/Engine/Plugins/FX/Niagara/Source/Niagara/Private/NiagaraDebugHud.cpp
@@ -39,6 +39,23 @@ namespace NiagaraDebugLocal
 		TEXT("Engine.LODFraction"),
 	};
 
+	enum class EParameterStoreLocation
+	{
+		 UserOverride,
+		 SystemSpawn,
+		 SystemUpdate,
+		 EmitterSpawn,
+		 EmitterUpdate,
+		 EmitterRenderBindings,
+	};
+
+	struct FParameterStoreVariable
+	{
+		EParameterStoreLocation	StoreLocation;
+		int32					EmitterIndex = INDEX_NONE;
+		FNiagaraVariableBase	Variable;
+	};
+
 	struct FCachedVariables
 	{
 		~FCachedVariables()
@@ -58,11 +75,14 @@ namespace NiagaraDebugLocal
 
 		bool bShowEngineVariable[(int)EEngineVariables::Num] = {};				// Engine variables that are not contained within the store
 
-		TArray<FNiagaraDataSetDebugAccessor> SystemVariables;					// System & Emitter variables since both are inside the same DataBuffer
-		TArray<FNiagaraVariableWithOffset> UserVariables;						// Exposed user parameters which will pull from the component
+		TArray<FNiagaraDataSetDebugAccessor>				SystemVariables;			// System & Emitter variables since both are inside the same DataBuffer
+		TArray<FParameterStoreVariable>						ParameterStoreVariables;	// Variables Stored inside Parameter Stores
+	#if WITH_EDITORONLY_DATA
+		TArray<FNiagaraVariable>							StaticVariables;			// Static variables do not exist in cooked builds
+	#endif
 
-		TArray<TArray<FNiagaraDataSetDebugAccessor>> ParticleVariables;			// Per Emitter Particle variables
-		TArray<FNiagaraDataSetAccessor<FNiagaraPosition>> ParticlePositionAccessors;		// Only valid if we have particle attributes
+		TArray<TArray<FNiagaraDataSetDebugAccessor>>		ParticleVariables;			// Per Emitter Particle variables
+		TArray<FNiagaraDataSetAccessor<FNiagaraPosition>>	ParticlePositionAccessors;	// Only valid if we have particle attributes
 	};
 
 	static TMap<TWeakObjectPtr<UNiagaraSystem>, FCachedVariables> GCachedSystemVariables;
@@ -330,6 +350,93 @@ namespace NiagaraDebugLocal
 		}
 	}
 
+	void FindParameterStoreVariablesByWildcard(const FNiagaraParameterStore* ParameterStore, const TArray<FNiagaraDebugHUDVariable>& DebugVariables, FCachedVariables* CachedVariables, EParameterStoreLocation StoreLocation, int32 EmitterIndex = INDEX_NONE, ENiagaraSimTarget SimTarget = ENiagaraSimTarget::CPUSim)
+	{
+		if (ParameterStore)
+		{
+			FindVariablesByWildcard(
+				ParameterStore->ReadParameterVariables(),
+				DebugVariables,
+				[&](const FNiagaraVariableWithOffset& Variable)
+				{
+					if (CachedVariables->ParameterStoreVariables.ContainsByPredicate([&](const FParameterStoreVariable& Existing) { return Existing.Variable == Variable; }) == false)
+					{
+						FParameterStoreVariable& NewVariable = CachedVariables->ParameterStoreVariables.AddDefaulted_GetRef();
+						NewVariable.StoreLocation = StoreLocation;
+						NewVariable.EmitterIndex = EmitterIndex;
+						NewVariable.Variable = Variable;
+					}
+				}
+			);
+		}
+	}
+
+	void FindScriptVariablesByWildcard(UNiagaraScript* NiagaraScript, const TArray<FNiagaraDebugHUDVariable>& DebugVariables, FCachedVariables* CachedVariables, EParameterStoreLocation StoreLocation, int32 EmitterIndex = INDEX_NONE, ENiagaraSimTarget SimTarget = ENiagaraSimTarget::CPUSim)
+	{
+		if (NiagaraScript)
+		{
+			FindParameterStoreVariablesByWildcard(NiagaraScript->GetExecutionReadyParameterStore(SimTarget), DebugVariables, CachedVariables, StoreLocation, EmitterIndex, SimTarget);
+
+		#if WITH_EDITORONLY_DATA
+			FindVariablesByWildcard(
+				NiagaraScript->GetVMExecutableData().StaticVariablesWritten,
+				DebugVariables,
+				[&](const FNiagaraVariable& Variable)
+				{
+					CachedVariables->StaticVariables.AddUnique(Variable);
+				}
+			);
+		#endif
+		}
+	}
+
+	void FindSystemVariablesByWildcard(UNiagaraSystem* NiagaraSystem, const TArray<FNiagaraDebugHUDVariable>& DebugVariables, FCachedVariables* CachedVariables)
+	{
+		if (DebugVariables.Num() == 0)
+		{
+			return;
+		}
+
+		const FNiagaraDataSetCompiledData& SystemCompiledData = NiagaraSystem->GetSystemCompiledData().DataSetCompiledData;
+		FindVariablesByWildcard(
+			SystemCompiledData.Variables,
+			DebugVariables,
+			[&](const FNiagaraVariable& Variable) { CachedVariables->SystemVariables.AddDefaulted_GetRef().Init(SystemCompiledData, Variable.GetName()); }
+		);
+
+		FindParameterStoreVariablesByWildcard(&NiagaraSystem->GetExposedParameters(), DebugVariables, CachedVariables, EParameterStoreLocation::UserOverride);
+		FindScriptVariablesByWildcard(NiagaraSystem->GetSystemSpawnScript(), DebugVariables, CachedVariables, EParameterStoreLocation::SystemSpawn);
+		FindScriptVariablesByWildcard(NiagaraSystem->GetSystemUpdateScript(), DebugVariables, CachedVariables, EParameterStoreLocation::SystemUpdate);
+
+		int32 EmitterIndex = 0;
+		for (const FNiagaraEmitterHandle& EmitterHandle : NiagaraSystem->GetEmitterHandles())
+		{
+			FVersionedNiagaraEmitterData* EmitterData = EmitterHandle.GetEmitterData();
+			if (EmitterData && EmitterHandle.GetIsEnabled())
+			{
+				FindScriptVariablesByWildcard(EmitterData->SpawnScriptProps.Script, DebugVariables, CachedVariables, EParameterStoreLocation::EmitterSpawn, EmitterIndex, EmitterData->SimTarget);
+				FindScriptVariablesByWildcard(EmitterData->UpdateScriptProps.Script, DebugVariables, CachedVariables, EParameterStoreLocation::EmitterUpdate, EmitterIndex, EmitterData->SimTarget);
+				FindParameterStoreVariablesByWildcard(&EmitterData->RendererBindings, DebugVariables, CachedVariables, EParameterStoreLocation::EmitterRenderBindings, EmitterIndex);
+			}
+			++EmitterIndex;
+		}
+
+	#if WITH_EDITORONLY_DATA
+		// Remove any static variables we found that are used inside renderer bindings to avoid doubling up
+		CachedVariables->StaticVariables.RemoveAll(
+			[&](const FNiagaraVariable& StaticVariable)
+			{
+				return CachedVariables->ParameterStoreVariables.ContainsByPredicate(
+						[&](const FParameterStoreVariable& Existing)
+						{
+							return Existing.Variable == StaticVariable;
+						}
+					);
+			}
+		);
+	#endif
+	}
+
 	const FCachedVariables& GetCachedVariables(UNiagaraSystem* NiagaraSystem)
 	{
 		FCachedVariables* CachedVariables = GCachedSystemVariables.Find(NiagaraSystem);
@@ -343,18 +450,7 @@ namespace NiagaraDebugLocal
 
 			if (Settings.bShowSystemVariables && Settings.SystemVariables.Num() > 0)
 			{
-				const FNiagaraDataSetCompiledData& SystemCompiledData = NiagaraSystem->GetSystemCompiledData().DataSetCompiledData;
-				FindVariablesByWildcard(
-					SystemCompiledData.Variables,
-					Settings.SystemVariables,
-					[&](const FNiagaraVariable& Variable) { CachedVariables->SystemVariables.AddDefaulted_GetRef().Init(SystemCompiledData, Variable.GetName()); }
-				);
-
-				FindVariablesByWildcard(
-					NiagaraSystem->GetExposedParameters().ReadParameterVariables(),
-					Settings.SystemVariables,
-					[&](const FNiagaraVariableWithOffset& Variable) { CachedVariables->UserVariables.Add(Variable); }
-				);
+				FindSystemVariablesByWildcard(NiagaraSystem, Settings.SystemVariables, CachedVariables);
 
 				for (int32 iVariable = 0; iVariable < (int32)EEngineVariables::Num; ++iVariable)
 				{
@@ -2647,49 +2743,70 @@ void FNiagaraDebugHud::DrawComponents(FNiagaraWorldManager* WorldManager, UCanva
 							}
 						}
 
-						// User variables
-						if (CachedVariables.UserVariables.Num() > 0)
+						// Parameter Store Variables (i.e. user variables)
+						if (CachedVariables.ParameterStoreVariables.Num() > 0)
 						{
-							if (FNiagaraUserRedirectionParameterStore* ParameterStore = SystemInstance->GetOverrideParameters())
+							for (const FParameterStoreVariable& ParameterStoreVariable : CachedVariables.ParameterStoreVariables)
 							{
-								for (const FNiagaraVariableWithOffset& UserVariable : CachedVariables.UserVariables)
+								FNiagaraParameterStore* ParameterStore = nullptr;
+								switch (ParameterStoreVariable.StoreLocation)
 								{
-									if (UserVariable.IsDataInterface())
+									case EParameterStoreLocation::UserOverride:			ParameterStore = SystemInstance->GetOverrideParameters(); break;
+									case EParameterStoreLocation::SystemSpawn:			ParameterStore = &SystemSimulation->GetSpawnExecutionContext()->Parameters; break;
+									case EParameterStoreLocation::SystemUpdate:			ParameterStore = &SystemSimulation->GetUpdateExecutionContext()->Parameters; break;
+									case EParameterStoreLocation::EmitterSpawn:			ParameterStore = &SystemInstance->Emitters[ParameterStoreVariable.EmitterIndex]->GetSpawnExecutionContext().Parameters; break;
+									case EParameterStoreLocation::EmitterUpdate:		ParameterStore = &SystemInstance->Emitters[ParameterStoreVariable.EmitterIndex]->GetUpdateExecutionContext().Parameters; break;
+									case EParameterStoreLocation::EmitterRenderBindings:ParameterStore = &SystemInstance->Emitters[ParameterStoreVariable.EmitterIndex]->GetRendererBoundVariables(); break;
+								}
+								if (ParameterStore == nullptr)
+								{
+									continue;
+								}
+
+								if (ParameterStoreVariable.Variable.IsDataInterface())
+								{
+									UNiagaraDataInterface* DataInterface = ParameterStore->GetDataInterface(ParameterStoreVariable.Variable);
+									FString DataInterfaceString;
+									if ( DataInterface != nullptr && Settings.DataInterfaceVerbosity != ENiagaraDebugHudVerbosity::None )
 									{
-										UNiagaraDataInterface* DataInterface = ParameterStore->GetDataInterface(UserVariable);
-										FString DataInterfaceString;
-										if ( DataInterface != nullptr && Settings.DataInterfaceVerbosity != ENiagaraDebugHudVerbosity::None )
-										{
-											DataInterface->DrawDebugHud(Canvas, SystemInstance, DataInterfaceString, Settings.DataInterfaceVerbosity == ENiagaraDebugHudVerbosity::Verbose);
-										}
-
-										if ( DataInterfaceString.IsEmpty() )
-										{
-											StringBuilder.Appendf(TEXT("%s(%s %s)\n"), *UserVariable.GetName().ToString(), *GetNameSafe(UserVariable.GetType().GetClass()) , *GetNameSafe(DataInterface));
-										}
-										else
-										{
-											StringBuilder.Appendf(TEXT("%s(%s)\n"), *UserVariable.GetName().ToString(), *DataInterfaceString);
-										}
+										DataInterface->DrawDebugHud(Canvas, SystemInstance, DataInterfaceString, Settings.DataInterfaceVerbosity == ENiagaraDebugHudVerbosity::Verbose);
 									}
-									else if (UserVariable.IsUObject())
+
+									if ( DataInterfaceString.IsEmpty() )
 									{
-										StringBuilder.Appendf(TEXT("%s(%s %s)\n"), *UserVariable.GetName().ToString(), *GetNameSafe(UserVariable.GetType().GetClass()) , *GetNameSafe(ParameterStore->GetUObjects()[UserVariable.Offset]));
+										StringBuilder.Appendf(TEXT("%s(%s %s)"), *ParameterStoreVariable.Variable.GetName().ToString(), *GetNameSafe(ParameterStoreVariable.Variable.GetType().GetClass()) , *GetNameSafe(DataInterface));
 									}
 									else
 									{
-										FNiagaraVariable UserVariableWithValue(UserVariable);
-										if ( const uint8* ParameterData = ParameterStore->GetParameterData(UserVariableWithValue) )
-										{
-											UserVariableWithValue.SetData(ParameterData);
-										}
-										StringBuilder.Append(*UserVariableWithValue.ToString());
-										StringBuilder.Append(TEXT("\n"));
+										StringBuilder.Appendf(TEXT("%s(%s)"), *ParameterStoreVariable.Variable.GetName().ToString(), *DataInterfaceString);
+									}
+								}
+								else if (ParameterStoreVariable.Variable.IsUObject())
+								{
+									StringBuilder.Appendf(TEXT("%s(%s %s)"), *ParameterStoreVariable.Variable.GetName().ToString(), *GetNameSafe(ParameterStoreVariable.Variable.GetType().GetClass()) , *GetNameSafe(ParameterStore->GetUObject(ParameterStoreVariable.Variable)));
+								}
+								else
+								{
+									FNiagaraVariable VariableWithValue(ParameterStoreVariable.Variable);
+									if ( const uint8* ParameterData = ParameterStore->GetParameterData(VariableWithValue) )
+									{
+										VariableWithValue.SetData(ParameterData);
 									}
+									StringBuilder.Append(*VariableWithValue.ToString());
 								}
+								StringBuilder.AppendChar('\n');
 							}
 						}
 
+						// Static variables
+					#if WITH_EDITORONLY_DATA
+						for (const FNiagaraVariable& StaticVar : CachedVariables.StaticVariables)
+						{
+							StringBuilder.Append(*StaticVar.ToString());
+							StringBuilder.Append(TEXT(" (Static)\n"));
+						}
+					#endif
+
 						// Append particle data if we don't show them in world
 						if (Settings.bShowParticlesVariablesWithSystem)
 						{
