- Add method to hash a shader template file
- Add method to append a tempate shader file, dealing with potential lack of return at end of file

diff --git a/Engine/Plugins/Experimental/ChaosNiagara/Source/ChaosNiagara/Private/NiagaraDataInterfaceChaosDestruction.cpp b/Engine/Plugins/Experimental/ChaosNiagara/Source/ChaosNiagara/Private/NiagaraDataInterfaceChaosDestruction.cpp
--- a/Engine/Plugins/Experimental/ChaosNiagara/Source/ChaosNiagara/Private/NiagaraDataInterfaceChaosDestruction.cpp
+++ b/Engine/Plugins/Experimental/ChaosNiagara/Source/ChaosNiagara/Private/NiagaraDataInterfaceChaosDestruction.cpp
@@ -3355,7 +3355,7 @@ void UNiagaraDataInterfaceChaosDestruction::GetTrailingData(FVectorVMExternalFun
 bool UNiagaraDataInterfaceChaosDestruction::AppendCompileHash(FNiagaraCompileHashVisitor* InVisitor) const
 {
 	bool bSuccess = Super::AppendCompileHash(InVisitor);
-	bSuccess &= InVisitor->UpdateString(TEXT("UNiagaraDataInterfaceChaosDestructionSource"), GetShaderFileHash(NDIChaosDestructionLocal::TemplateShaderFilePath, EShaderPlatform::SP_PCD3D_SM5).ToString());
+	bSuccess &= InVisitor->UpdateShaderFile(NDIChaosDestructionLocal::TemplateShaderFilePath);
 	bSuccess &= InVisitor->UpdateShaderParameters<NDIChaosDestructionLocal::FShaderParameters>();
 	return bSuccess;
 }
@@ -3363,10 +3363,7 @@ bool UNiagaraDataInterfaceChaosDestruction::AppendCompileHash(FNiagaraCompileHas
 void UNiagaraDataInterfaceChaosDestruction::GetParameterDefinitionHLSL(const FNiagaraDataInterfaceGPUParamInfo& ParamInfo, FString& OutHLSL)
 {
 	const TMap<FString, FStringFormatArg> TemplateArgs = { {TEXT("ParameterName"),	ParamInfo.DataInterfaceHLSLSymbol}, };
-
-	FString TemplateFile;
-	LoadShaderSourceFile(NDIChaosDestructionLocal::TemplateShaderFilePath, EShaderPlatform::SP_PCD3D_SM5, &TemplateFile, nullptr);
-	OutHLSL += FString::Format(*TemplateFile, TemplateArgs);
+	AppendTemplateHLSL(OutHLSL, NDIChaosDestructionLocal::TemplateShaderFilePath, TemplateArgs);
 }
 
 bool UNiagaraDataInterfaceChaosDestruction::GetFunctionHLSL(const FNiagaraDataInterfaceGPUParamInfo& ParamInfo, const FNiagaraDataInterfaceGeneratedFunction& FunctionInfo, int FunctionInstanceIndex, FString& OutHLSL)
diff --git a/Engine/Plugins/Experimental/ChaosNiagara/Source/ChaosNiagara/Private/NiagaraDataInterfaceGeometryCollection.cpp b/Engine/Plugins/Experimental/ChaosNiagara/Source/ChaosNiagara/Private/NiagaraDataInterfaceGeometryCollection.cpp
--- a/Engine/Plugins/Experimental/ChaosNiagara/Source/ChaosNiagara/Private/NiagaraDataInterfaceGeometryCollection.cpp
+++ b/Engine/Plugins/Experimental/ChaosNiagara/Source/ChaosNiagara/Private/NiagaraDataInterfaceGeometryCollection.cpp
@@ -447,7 +447,7 @@ bool UNiagaraDataInterfaceGeometryCollection::GetFunctionHLSL(const FNiagaraData
 bool UNiagaraDataInterfaceGeometryCollection::AppendCompileHash(FNiagaraCompileHashVisitor* InVisitor) const
 {
 	bool bSuccess = Super::AppendCompileHash(InVisitor);
-	bSuccess &= InVisitor->UpdateString(TEXT("UNiagaraDataInterfaceGeometryCollectionSource"), GetShaderFileHash(NDIGeometryCollectionLocal::TemplateShaderFilePath, EShaderPlatform::SP_PCD3D_SM5).ToString());
+	bSuccess &= InVisitor->UpdateShaderFile(NDIGeometryCollectionLocal::TemplateShaderFilePath);
 	bSuccess &= InVisitor->UpdateShaderParameters<FShaderParameters>();
 	return bSuccess;
 }
@@ -455,10 +455,7 @@ bool UNiagaraDataInterfaceGeometryCollection::AppendCompileHash(FNiagaraCompileH
 void UNiagaraDataInterfaceGeometryCollection::GetParameterDefinitionHLSL(const FNiagaraDataInterfaceGPUParamInfo& ParamInfo, FString& OutHLSL)
 {
 	const TMap<FString, FStringFormatArg> TemplateArgs = { {TEXT("ParameterName"),	ParamInfo.DataInterfaceHLSLSymbol}, };
-
-	FString TemplateFile;
-	LoadShaderSourceFile(NDIGeometryCollectionLocal::TemplateShaderFilePath, EShaderPlatform::SP_PCD3D_SM5, &TemplateFile, nullptr);
-	OutHLSL += FString::Format(*TemplateFile, TemplateArgs);
+	AppendTemplateHLSL(OutHLSL, NDIGeometryCollectionLocal::TemplateShaderFilePath, TemplateArgs);
 }
 #endif
 
diff --git a/Engine/Plugins/Experimental/ChaosNiagara/Source/ChaosNiagara/Private/NiagaraDataInterfacePhysicsField.cpp b/Engine/Plugins/Experimental/ChaosNiagara/Source/ChaosNiagara/Private/NiagaraDataInterfacePhysicsField.cpp
--- a/Engine/Plugins/Experimental/ChaosNiagara/Source/ChaosNiagara/Private/NiagaraDataInterfacePhysicsField.cpp
+++ b/Engine/Plugins/Experimental/ChaosNiagara/Source/ChaosNiagara/Private/NiagaraDataInterfacePhysicsField.cpp
@@ -651,7 +651,7 @@ bool UNiagaraDataInterfacePhysicsField::UpgradeFunctionCall(FNiagaraFunctionSign
 bool UNiagaraDataInterfacePhysicsField::AppendCompileHash(FNiagaraCompileHashVisitor* InVisitor) const
 {
 	bool bSuccess = Super::AppendCompileHash(InVisitor);
-	bSuccess &= InVisitor->UpdateString(TEXT("UNiagaraDataInterfacePhysicsFieldTemplateHLSLSource"), GetShaderFileHash(NDIPhysicsFieldLocal::TemplateShaderFilePath, EShaderPlatform::SP_PCD3D_SM5).ToString());
+	bSuccess &= InVisitor->UpdateShaderFile(NDIPhysicsFieldLocal::TemplateShaderFilePath);
 	bSuccess &= InVisitor->UpdateShaderParameters<NDIPhysicsFieldLocal::FShaderInstanceParameters>();
 	bSuccess &= InVisitor->UpdateShaderParameters<NDIPhysicsFieldLocal::FShaderGlobalParameters>();
 	return bSuccess;
@@ -659,14 +659,11 @@ bool UNiagaraDataInterfacePhysicsField::AppendCompileHash(FNiagaraCompileHashVis
 
 void UNiagaraDataInterfacePhysicsField::GetParameterDefinitionHLSL(const FNiagaraDataInterfaceGPUParamInfo& ParamInfo, FString& OutHLSL)
 {
-	TMap<FString, FStringFormatArg> TemplateArgs =
+	const TMap<FString, FStringFormatArg> TemplateArgs =
 	{
 		{TEXT("ParameterName"),	ParamInfo.DataInterfaceHLSLSymbol},
 	};
-
-	FString TemplateFile;
-	LoadShaderSourceFile(NDIPhysicsFieldLocal::TemplateShaderFilePath, EShaderPlatform::SP_PCD3D_SM5, &TemplateFile, nullptr);
-	OutHLSL += FString::Format(*TemplateFile, TemplateArgs);
+	AppendTemplateHLSL(OutHLSL, NDIPhysicsFieldLocal::TemplateShaderFilePath, TemplateArgs);
 }
 #endif
 
diff --git a/Engine/Plugins/FX/ExampleCustomDataInterface/Source/ExampleCustomDataInterface/Private/NiagaraDataInterfaceMousePosition.cpp b/Engine/Plugins/FX/ExampleCustomDataInterface/Source/ExampleCustomDataInterface/Private/NiagaraDataInterfaceMousePosition.cpp
--- a/Engine/Plugins/FX/ExampleCustomDataInterface/Source/ExampleCustomDataInterface/Private/NiagaraDataInterfaceMousePosition.cpp
+++ b/Engine/Plugins/FX/ExampleCustomDataInterface/Source/ExampleCustomDataInterface/Private/NiagaraDataInterfaceMousePosition.cpp
@@ -210,8 +210,7 @@ bool UNiagaraDataInterfaceMousePosition::AppendCompileHash(FNiagaraCompileHashVi
 		return false;
 	}
 
-	FSHAHash Hash = GetShaderFileHash(MouseDITemplateShaderFile, SP_PCD3D_SM5);
-	InVisitor->UpdateString(TEXT("NiagaraDataInterfaceMousePositionHLSLSource"), Hash.ToString());
+	InVisitor->UpdateShaderFile(MouseDITemplateShaderFile);
 	InVisitor->UpdateShaderParameters<FShaderParameters>();
 	return true;
 }
@@ -227,14 +226,11 @@ bool UNiagaraDataInterfaceMousePosition::GetFunctionHLSL(const FNiagaraDataInter
 // this loads our hlsl template script file and 
 void UNiagaraDataInterfaceMousePosition::GetParameterDefinitionHLSL(const FNiagaraDataInterfaceGPUParamInfo& ParamInfo,	FString& OutHLSL)
 {
-	TMap<FString, FStringFormatArg> TemplateArgs =
+	const TMap<FString, FStringFormatArg> TemplateArgs =
 	{
 		{TEXT("ParameterName"),	ParamInfo.DataInterfaceHLSLSymbol},
 	};
-
-	FString TemplateFile;
-	LoadShaderSourceFile(MouseDITemplateShaderFile, SP_PCD3D_SM5, &TemplateFile, nullptr);
-	OutHLSL += FString::Format(*TemplateFile, TemplateArgs);
+	AppendTemplateHLSL(OutHLSL, MouseDITemplateShaderFile, TemplateArgs);
 }
 
 #endif
diff --git a/Engine/Plugins/FX/Niagara/Source/Niagara/Classes/NiagaraDataInterface.h b/Engine/Plugins/FX/Niagara/Source/Niagara/Classes/NiagaraDataInterface.h
--- a/Engine/Plugins/FX/Niagara/Source/Niagara/Classes/NiagaraDataInterface.h
+++ b/Engine/Plugins/FX/Niagara/Source/Niagara/Classes/NiagaraDataInterface.h
@@ -633,6 +633,9 @@ public:
 	{
 		return false;
 	}
+
+	/** Formats and appends a template file onto the output HLSL */
+	void AppendTemplateHLSL(FString& OutHLSL, const TCHAR* TemplateShaderFile, const TMap<FString, FStringFormatArg>& TemplateArgs) const;
 #endif
 
 	virtual void PostExecute() {}
diff --git a/Engine/Plugins/FX/Niagara/Source/Niagara/Classes/NiagaraDataInterfaceArrayImpl.h b/Engine/Plugins/FX/Niagara/Source/Niagara/Classes/NiagaraDataInterfaceArrayImpl.h
--- a/Engine/Plugins/FX/Niagara/Source/Niagara/Classes/NiagaraDataInterfaceArrayImpl.h
+++ b/Engine/Plugins/FX/Niagara/Source/Niagara/Classes/NiagaraDataInterfaceArrayImpl.h
@@ -979,10 +979,8 @@ struct FNDIArrayProxyImpl : public INDIArrayProxyBase
 	{
 		if (FNDIArrayImplHelper<TArrayType>::bSupportsGPU)
 		{
-			FSHAHash Hash = GetShaderFileHash(FNiagaraDataInterfaceArrayImplHelper::GetHLSLTemplateFile(false), EShaderPlatform::SP_PCD3D_SM5);
-			InVisitor->UpdateString(TEXT("NiagaraDataInterfaceArrayTemplateHLSLSource"), Hash.ToString());
-			Hash = GetShaderFileHash(FNiagaraDataInterfaceArrayImplHelper::GetHLSLTemplateFile(true), EShaderPlatform::SP_PCD3D_SM5);
-			InVisitor->UpdateString(TEXT("NiagaraDataInterfaceArrayTemplateHLSLSource"), Hash.ToString());
+			InVisitor->UpdateShaderFile(FNiagaraDataInterfaceArrayImplHelper::GetHLSLTemplateFile(false));
+			InVisitor->UpdateShaderFile(FNiagaraDataInterfaceArrayImplHelper::GetHLSLTemplateFile(true));
 		}
 		return true;
 	}
diff --git a/Engine/Plugins/FX/Niagara/Source/Niagara/Private/DataInterface/NiagaraDataInterfaceActorComponent.cpp b/Engine/Plugins/FX/Niagara/Source/Niagara/Private/DataInterface/NiagaraDataInterfaceActorComponent.cpp
--- a/Engine/Plugins/FX/Niagara/Source/Niagara/Private/DataInterface/NiagaraDataInterfaceActorComponent.cpp
+++ b/Engine/Plugins/FX/Niagara/Source/Niagara/Private/DataInterface/NiagaraDataInterfaceActorComponent.cpp
@@ -241,9 +241,8 @@ void UNiagaraDataInterfaceActorComponent::GetVMExternalFunction(const FVMExterna
 bool UNiagaraDataInterfaceActorComponent::AppendCompileHash(FNiagaraCompileHashVisitor* InVisitor) const
 {
 	bool bSuccess = Super::AppendCompileHash(InVisitor);
-	FSHAHash Hash = GetShaderFileHash(NDIActorComponentLocal::TemplateShaderFile, EShaderPlatform::SP_PCD3D_SM5);
-	InVisitor->UpdateString(TEXT("NiagaraDataInterfaceActorComponentTemplateHLSLSource"), Hash.ToString());
-	InVisitor->UpdateShaderParameters<FShaderParameters>();
+	bSuccess &= InVisitor->UpdateShaderFile(NDIActorComponentLocal::TemplateShaderFile);
+	bSuccess &= InVisitor->UpdateShaderParameters<FShaderParameters>();
 	return bSuccess;
 }
 
@@ -253,10 +252,7 @@ void UNiagaraDataInterfaceActorComponent::GetParameterDefinitionHLSL(const FNiag
 	{
 		{TEXT("ParameterName"),	ParamInfo.DataInterfaceHLSLSymbol},
 	};
-
-	FString TemplateFile;
-	LoadShaderSourceFile(NDIActorComponentLocal::TemplateShaderFile, EShaderPlatform::SP_PCD3D_SM5, &TemplateFile, nullptr);
-	OutHLSL += FString::Format(*TemplateFile, TemplateArgs);
+	AppendTemplateHLSL(OutHLSL, NDIActorComponentLocal::TemplateShaderFile, TemplateArgs);
 }
 
 bool UNiagaraDataInterfaceActorComponent::GetFunctionHLSL(const FNiagaraDataInterfaceGPUParamInfo& ParamInfo, const FNiagaraDataInterfaceGeneratedFunction& FunctionInfo, int FunctionInstanceIndex, FString& OutHLSL)
diff --git a/Engine/Plugins/FX/Niagara/Source/Niagara/Private/DataInterface/NiagaraDataInterfaceAsyncGpuTrace.cpp b/Engine/Plugins/FX/Niagara/Source/Niagara/Private/DataInterface/NiagaraDataInterfaceAsyncGpuTrace.cpp
--- a/Engine/Plugins/FX/Niagara/Source/Niagara/Private/DataInterface/NiagaraDataInterfaceAsyncGpuTrace.cpp
+++ b/Engine/Plugins/FX/Niagara/Source/Niagara/Private/DataInterface/NiagaraDataInterfaceAsyncGpuTrace.cpp
@@ -378,14 +378,12 @@ bool UNiagaraDataInterfaceAsyncGpuTrace::UpgradeFunctionCall(FNiagaraFunctionSig
 
 void UNiagaraDataInterfaceAsyncGpuTrace::GetParameterDefinitionHLSL(const FNiagaraDataInterfaceGPUParamInfo& ParamInfo, FString& OutHLSL)
 {
-	TMap<FString, FStringFormatArg> TemplateArgs =
+	const TMap<FString, FStringFormatArg> TemplateArgs =
 	{
 		{TEXT("ParameterName"),	ParamInfo.DataInterfaceHLSLSymbol},
 	};
 
-	FString TemplateFile;
-	LoadShaderSourceFile(NDIAsyncGpuTraceLocal::TemplateShaderFile, EShaderPlatform::SP_PCD3D_SM5, &TemplateFile, nullptr);
-	OutHLSL += FString::Format(*TemplateFile, TemplateArgs);
+	AppendTemplateHLSL(OutHLSL, NDIAsyncGpuTraceLocal::TemplateShaderFile, TemplateArgs);
 }
 
 void UNiagaraDataInterfaceAsyncGpuTrace::GetCommonHLSL(FString& OutHlsl)
@@ -400,8 +398,8 @@ bool UNiagaraDataInterfaceAsyncGpuTrace::AppendCompileHash(FNiagaraCompileHashVi
 		return false;
 	}
 
-	InVisitor->UpdateString(TEXT("NDIAsyncGpuTraceCommonHLSLSource"), GetShaderFileHash(NDIAsyncGpuTraceLocal::CommonShaderFile, EShaderPlatform::SP_PCD3D_SM5).ToString());
-	InVisitor->UpdateString(TEXT("NDIAsyncGpuTraceTemplateHLSLSource"), GetShaderFileHash(NDIAsyncGpuTraceLocal::TemplateShaderFile, EShaderPlatform::SP_PCD3D_SM5).ToString());
+	InVisitor->UpdateShaderFile(NDIAsyncGpuTraceLocal::CommonShaderFile);
+	InVisitor->UpdateShaderFile(NDIAsyncGpuTraceLocal::TemplateShaderFile);
 	InVisitor->UpdateShaderParameters<NDIAsyncGpuTraceLocal::FShaderParameters>();
 
 	return true;
diff --git a/Engine/Plugins/FX/Niagara/Source/Niagara/Private/DataInterface/NiagaraDataInterfaceDebugDraw.cpp b/Engine/Plugins/FX/Niagara/Source/Niagara/Private/DataInterface/NiagaraDataInterfaceDebugDraw.cpp
--- a/Engine/Plugins/FX/Niagara/Source/Niagara/Private/DataInterface/NiagaraDataInterfaceDebugDraw.cpp
+++ b/Engine/Plugins/FX/Niagara/Source/Niagara/Private/DataInterface/NiagaraDataInterfaceDebugDraw.cpp
@@ -2918,8 +2918,8 @@ bool UNiagaraDataInterfaceDebugDraw::AppendCompileHash(FNiagaraCompileHashVisito
 	if (!Super::AppendCompileHash(InVisitor))
 		return false;
 
-	InVisitor->UpdateString(TEXT("NiagaraDataInterfaceDebugDrawHLSLSource"), GetShaderFileHash(NDIDebugDrawLocal::CommonShaderFile, EShaderPlatform::SP_PCD3D_SM5).ToString());
-	InVisitor->UpdateString(TEXT("NiagaraDataInterfaceDebugDrawTemplateHLSLSource"), GetShaderFileHash(NDIDebugDrawLocal::TemplateShaderFile, EShaderPlatform::SP_PCD3D_SM5).ToString());
+	InVisitor->UpdateShaderFile(NDIDebugDrawLocal::CommonShaderFile);
+	InVisitor->UpdateShaderFile(NDIDebugDrawLocal::TemplateShaderFile);
 
 	InVisitor->UpdatePOD<int32>(TEXT("NiagaraDataInterfaceDebugDrawVersion"), FNiagaraDebugDrawDIFunctionVersion::LatestVersion);	
 	InVisitor->UpdateShaderParameters<FNDIDebugDrawShaderParameters>();
@@ -2935,14 +2935,11 @@ void UNiagaraDataInterfaceDebugDraw::GetCommonHLSL(FString& OutHLSL)
 
 void UNiagaraDataInterfaceDebugDraw::GetParameterDefinitionHLSL(const FNiagaraDataInterfaceGPUParamInfo& ParamInfo, FString& OutHLSL)
 {
-	TMap<FString, FStringFormatArg> TemplateArgs =
+	const TMap<FString, FStringFormatArg> TemplateArgs =
 	{
 		{TEXT("ParameterName"),	ParamInfo.DataInterfaceHLSLSymbol},
 	};
-
-	FString TemplateFile;
-	LoadShaderSourceFile(NDIDebugDrawLocal::TemplateShaderFile, EShaderPlatform::SP_PCD3D_SM5, &TemplateFile, nullptr);
-	OutHLSL += FString::Format(*TemplateFile, TemplateArgs);
+	AppendTemplateHLSL(OutHLSL, NDIDebugDrawLocal::TemplateShaderFile, TemplateArgs);
 }
 
 bool UNiagaraDataInterfaceDebugDraw::GetFunctionHLSL(const FNiagaraDataInterfaceGPUParamInfo& ParamInfo, const FNiagaraDataInterfaceGeneratedFunction& FunctionInfo, int FunctionInstanceIndex, FString& OutHLSL)
diff --git a/Engine/Plugins/FX/Niagara/Source/Niagara/Private/DataInterface/NiagaraDataInterfaceEmitterProperties.cpp b/Engine/Plugins/FX/Niagara/Source/Niagara/Private/DataInterface/NiagaraDataInterfaceEmitterProperties.cpp
--- a/Engine/Plugins/FX/Niagara/Source/Niagara/Private/DataInterface/NiagaraDataInterfaceEmitterProperties.cpp
+++ b/Engine/Plugins/FX/Niagara/Source/Niagara/Private/DataInterface/NiagaraDataInterfaceEmitterProperties.cpp
@@ -274,22 +274,18 @@ void UNiagaraDataInterfaceEmitterProperties::GetVMExternalFunction(const FVMExte
 bool UNiagaraDataInterfaceEmitterProperties::AppendCompileHash(FNiagaraCompileHashVisitor* InVisitor) const
 {
 	bool bSuccess = Super::AppendCompileHash(InVisitor);
-	FSHAHash Hash = GetShaderFileHash(NDIEmitterPropertiesLocal::TemplateShaderFile, EShaderPlatform::SP_PCD3D_SM5);
-	InVisitor->UpdateString(TEXT("NiagaraDataInterfaceEmitterPropertiesTemplateHLSLSource"), Hash.ToString());
+	InVisitor->UpdateShaderFile(NDIEmitterPropertiesLocal::TemplateShaderFile);
 	InVisitor->UpdateShaderParameters<FShaderParameters>();
 	return bSuccess;
 }
 
 void UNiagaraDataInterfaceEmitterProperties::GetParameterDefinitionHLSL(const FNiagaraDataInterfaceGPUParamInfo& ParameterInfo, FString& OutHLSL)
 {
-	TMap<FString, FStringFormatArg> TemplateArgs =
+	const TMap<FString, FStringFormatArg> TemplateArgs =
 	{
 		{TEXT("ParameterName"),	ParameterInfo.DataInterfaceHLSLSymbol},
 	};
-
-	FString TemplateFile;
-	LoadShaderSourceFile(NDIEmitterPropertiesLocal::TemplateShaderFile, EShaderPlatform::SP_PCD3D_SM5, &TemplateFile, nullptr);
-	OutHLSL += FString::Format(*TemplateFile, TemplateArgs);
+	AppendTemplateHLSL(OutHLSL, NDIEmitterPropertiesLocal::TemplateShaderFile, TemplateArgs);
 }
 
 bool UNiagaraDataInterfaceEmitterProperties::GetFunctionHLSL(const FNiagaraDataInterfaceGPUParamInfo& ParamInfo, const FNiagaraDataInterfaceGeneratedFunction& FunctionInfo, int FunctionInstanceIndex, FString& OutHLSL)
diff --git a/Engine/Plugins/FX/Niagara/Source/Niagara/Private/DataInterface/NiagaraDataInterfaceGBuffer.cpp b/Engine/Plugins/FX/Niagara/Source/Niagara/Private/DataInterface/NiagaraDataInterfaceGBuffer.cpp
--- a/Engine/Plugins/FX/Niagara/Source/Niagara/Private/DataInterface/NiagaraDataInterfaceGBuffer.cpp
+++ b/Engine/Plugins/FX/Niagara/Source/Niagara/Private/DataInterface/NiagaraDataInterfaceGBuffer.cpp
@@ -198,22 +198,18 @@ void UNiagaraDataInterfaceGBuffer::SetShaderParameters(const FNiagaraDataInterfa
 bool UNiagaraDataInterfaceGBuffer::AppendCompileHash(FNiagaraCompileHashVisitor* InVisitor) const
 {
 	bool bSuccess = Super::AppendCompileHash(InVisitor);
-	FSHAHash Hash = GetShaderFileHash(NiagaraDataInterfaceGBufferLocal::TemplateShaderFile, EShaderPlatform::SP_PCD3D_SM5);
-	InVisitor->UpdateString(TEXT("UNiagaraDataInterfaceGBufferTemplateHLSLSource"), Hash.ToString());
+	InVisitor->UpdateShaderFile(NiagaraDataInterfaceGBufferLocal::TemplateShaderFile);
 	InVisitor->UpdateShaderParameters<NiagaraDataInterfaceGBufferLocal::FShaderParameters>();
 	return bSuccess;
 }
 
 void UNiagaraDataInterfaceGBuffer::GetParameterDefinitionHLSL(const FNiagaraDataInterfaceGPUParamInfo& ParamInfo, FString& OutHLSL)
 {
-	TMap<FString, FStringFormatArg> TemplateArgs =
+	const TMap<FString, FStringFormatArg> TemplateArgs =
 	{
 		{TEXT("ParameterName"),	ParamInfo.DataInterfaceHLSLSymbol},
 	};
-
-	FString TemplateFile;
-	LoadShaderSourceFile(NiagaraDataInterfaceGBufferLocal::TemplateShaderFile, EShaderPlatform::SP_PCD3D_SM5, &TemplateFile, nullptr);
-	OutHLSL += FString::Format(*TemplateFile, TemplateArgs);
+	AppendTemplateHLSL(OutHLSL, NiagaraDataInterfaceGBufferLocal::TemplateShaderFile, TemplateArgs);
 }
 
 bool UNiagaraDataInterfaceGBuffer::GetFunctionHLSL(const FNiagaraDataInterfaceGPUParamInfo& ParamInfo, const FNiagaraDataInterfaceGeneratedFunction& FunctionInfo, int FunctionInstanceIndex, FString& OutHLSL)
diff --git a/Engine/Plugins/FX/Niagara/Source/Niagara/Private/DataInterface/NiagaraDataInterfacePhysicsAsset.cpp b/Engine/Plugins/FX/Niagara/Source/Niagara/Private/DataInterface/NiagaraDataInterfacePhysicsAsset.cpp
--- a/Engine/Plugins/FX/Niagara/Source/Niagara/Private/DataInterface/NiagaraDataInterfacePhysicsAsset.cpp
+++ b/Engine/Plugins/FX/Niagara/Source/Niagara/Private/DataInterface/NiagaraDataInterfacePhysicsAsset.cpp
@@ -1129,7 +1129,7 @@ bool UNiagaraDataInterfacePhysicsAsset::AppendCompileHash(FNiagaraCompileHashVis
 		return false;
 	}
 
-	InVisitor->UpdateString(TEXT("NiagaraDataInterfaceRigidMeshCollisionQueryHLSLSource"), GetShaderFileHash(NDIPhysicsAssetLocal::TemplateShaderFile, EShaderPlatform::SP_PCD3D_SM5).ToString());
+	InVisitor->UpdateShaderFile(NDIPhysicsAssetLocal::TemplateShaderFile);
 	InVisitor->UpdateShaderParameters<NDIPhysicsAssetLocal::FShaderParameters>();
 
 	return true;
@@ -1144,14 +1144,11 @@ void UNiagaraDataInterfacePhysicsAsset::GetParameterDefinitionHLSL(const FNiagar
 {
 	Super::GetParameterDefinitionHLSL(ParamInfo, OutHLSL);
 
-	TMap<FString, FStringFormatArg> TemplateArgs =
+	const TMap<FString, FStringFormatArg> TemplateArgs =
 	{
 		{TEXT("ParameterName"),	ParamInfo.DataInterfaceHLSLSymbol},
 	};
-
-	FString TemplateFile;
-	LoadShaderSourceFile(NDIPhysicsAssetLocal::TemplateShaderFile, EShaderPlatform::SP_PCD3D_SM5, &TemplateFile, nullptr);
-	OutHLSL += FString::Format(*TemplateFile, TemplateArgs);
+	AppendTemplateHLSL(OutHLSL, NDIPhysicsAssetLocal::TemplateShaderFile, TemplateArgs);
 }
 #endif
 
diff --git a/Engine/Plugins/FX/Niagara/Source/Niagara/Private/DataInterface/NiagaraDataInterfaceSimCacheReader.cpp b/Engine/Plugins/FX/Niagara/Source/Niagara/Private/DataInterface/NiagaraDataInterfaceSimCacheReader.cpp
--- a/Engine/Plugins/FX/Niagara/Source/Niagara/Private/DataInterface/NiagaraDataInterfaceSimCacheReader.cpp
+++ b/Engine/Plugins/FX/Niagara/Source/Niagara/Private/DataInterface/NiagaraDataInterfaceSimCacheReader.cpp
@@ -442,24 +442,21 @@ void UNiagaraDataInterfaceSimCacheReader::GetVMExternalFunction(const FVMExterna
 bool UNiagaraDataInterfaceSimCacheReader::AppendCompileHash(FNiagaraCompileHashVisitor* InVisitor) const
 {
 	bool bSuccess = Super::AppendCompileHash(InVisitor);
-	FSHAHash Hash = GetShaderFileHash(NDISimCacheReaderLocal::TemplateShaderFile, EShaderPlatform::SP_PCD3D_SM5);
-	InVisitor->UpdateString(TEXT("NDISimCacheReaderTemplateShaderFile"), Hash.ToString());
+	InVisitor->UpdateShaderFile(NDISimCacheReaderLocal::TemplateShaderFile);
 	InVisitor->UpdateShaderParameters<FShaderParameters>();
 	return bSuccess;
 }
 
 void UNiagaraDataInterfaceSimCacheReader::GetParameterDefinitionHLSL(const FNiagaraDataInterfaceGPUParamInfo& ParameterInfo, FString& OutHLSL)
 {
-	TMap<FString, FStringFormatArg> TemplateArgs =
+	const TMap<FString, FStringFormatArg> TemplateArgs =
 	{
 		{TEXT("ParameterName"),	ParameterInfo.DataInterfaceHLSLSymbol},
 	};
 
 	OutHLSL.Appendf(TEXT("int4 %s_ComponentOffsets[%d];\n"), *ParameterInfo.DataInterfaceHLSLSymbol, FMath::DivideAndRoundUp(ParameterInfo.GeneratedFunctions.Num(), 4));
 
-	FString TemplateFile;
-	LoadShaderSourceFile(NDISimCacheReaderLocal::TemplateShaderFile, EShaderPlatform::SP_PCD3D_SM5, &TemplateFile, nullptr);
-	OutHLSL += FString::Format(*TemplateFile, TemplateArgs);
+	AppendTemplateHLSL(OutHLSL, NDISimCacheReaderLocal::TemplateShaderFile, TemplateArgs);
 }
 
 bool UNiagaraDataInterfaceSimCacheReader::GetFunctionHLSL(const FNiagaraDataInterfaceGPUParamInfo& ParamInfo, const FNiagaraDataInterfaceGeneratedFunction& FunctionInfo, int FunctionInstanceIndex, FString& OutHLSL)
diff --git a/Engine/Plugins/FX/Niagara/Source/Niagara/Private/DataInterface/NiagaraDataInterfaceSimpleCounter.cpp b/Engine/Plugins/FX/Niagara/Source/Niagara/Private/DataInterface/NiagaraDataInterfaceSimpleCounter.cpp
--- a/Engine/Plugins/FX/Niagara/Source/Niagara/Private/DataInterface/NiagaraDataInterfaceSimpleCounter.cpp
+++ b/Engine/Plugins/FX/Niagara/Source/Niagara/Private/DataInterface/NiagaraDataInterfaceSimpleCounter.cpp
@@ -360,14 +360,11 @@ void UNiagaraDataInterfaceSimpleCounter::GetVMExternalFunction(const FVMExternal
 #if WITH_EDITORONLY_DATA
 void UNiagaraDataInterfaceSimpleCounter::GetParameterDefinitionHLSL(const FNiagaraDataInterfaceGPUParamInfo& ParamInfo, FString& OutHLSL)
 {
-	TMap<FString, FStringFormatArg> TemplateArgs =
+	const TMap<FString, FStringFormatArg> TemplateArgs =
 	{
 		{TEXT("ParameterName"),	ParamInfo.DataInterfaceHLSLSymbol},
 	};
-
-	FString TemplateFile;
-	LoadShaderSourceFile(NDISimpleCounterLocal::TemplateShaderFile, EShaderPlatform::SP_PCD3D_SM5, &TemplateFile, nullptr);
-	OutHLSL += FString::Format(*TemplateFile, TemplateArgs);
+	AppendTemplateHLSL(OutHLSL, NDISimpleCounterLocal::TemplateShaderFile, TemplateArgs);
 }
 
 bool UNiagaraDataInterfaceSimpleCounter::GetFunctionHLSL(const FNiagaraDataInterfaceGPUParamInfo& ParamInfo, const FNiagaraDataInterfaceGeneratedFunction& FunctionInfo, int FunctionInstanceIndex, FString& OutHLSL)
@@ -390,8 +387,7 @@ bool UNiagaraDataInterfaceSimpleCounter::GetFunctionHLSL(const FNiagaraDataInter
 bool UNiagaraDataInterfaceSimpleCounter::AppendCompileHash(FNiagaraCompileHashVisitor* InVisitor) const
 {
 	bool bSuccess = Super::AppendCompileHash(InVisitor);
-	FSHAHash Hash = GetShaderFileHash(NDISimpleCounterLocal::TemplateShaderFile, EShaderPlatform::SP_PCD3D_SM5);
-	InVisitor->UpdateString(TEXT("NiagaraDataInterfaceSimpleCounterTemplateHLSLSource"), Hash.ToString());
+	InVisitor->UpdateShaderFile(NDISimpleCounterLocal::TemplateShaderFile);
 	InVisitor->UpdateShaderParameters<FShaderParameters>();
 	return bSuccess;
 }
diff --git a/Engine/Plugins/FX/Niagara/Source/Niagara/Private/DataInterface/NiagaraDataInterfaceStaticMesh.cpp b/Engine/Plugins/FX/Niagara/Source/Niagara/Private/DataInterface/NiagaraDataInterfaceStaticMesh.cpp
--- a/Engine/Plugins/FX/Niagara/Source/Niagara/Private/DataInterface/NiagaraDataInterfaceStaticMesh.cpp
+++ b/Engine/Plugins/FX/Niagara/Source/Niagara/Private/DataInterface/NiagaraDataInterfaceStaticMesh.cpp
@@ -2795,8 +2795,7 @@ bool UNiagaraDataInterfaceStaticMesh::RequiresDistanceFieldData() const
 bool UNiagaraDataInterfaceStaticMesh::AppendCompileHash(FNiagaraCompileHashVisitor* InVisitor) const
 {
 	bool bSuccess = Super::AppendCompileHash(InVisitor);
-	FSHAHash Hash = GetShaderFileHash(NDIStaticMeshLocal::TemplateShaderFile, EShaderPlatform::SP_PCD3D_SM5);
-	bSuccess &= InVisitor->UpdateString(TEXT("NiagaraDataInterfaceStaticMeshTemplateHLSLSource"), Hash.ToString());
+	bSuccess &= InVisitor->UpdateShaderFile(NDIStaticMeshLocal::TemplateShaderFile);
 	bSuccess &= InVisitor->UpdatePOD(TEXT("NDIStaticMesh_AllowDistanceField"), GetDefault<UNiagaraSettings>()->NDIStaticMesh_AllowDistanceFields ? 1 : 0);
 	bSuccess &= InVisitor->UpdateShaderParameters<NDIStaticMeshLocal::FShaderParameters>();
 	return bSuccess;
@@ -2811,14 +2810,11 @@ void UNiagaraDataInterfaceStaticMesh::ModifyCompilationEnvironment(EShaderPlatfo
 
 void UNiagaraDataInterfaceStaticMesh::GetParameterDefinitionHLSL(const FNiagaraDataInterfaceGPUParamInfo& ParamInfo, FString& OutHLSL)
 {
-	TMap<FString, FStringFormatArg> TemplateArgs =
+	const TMap<FString, FStringFormatArg> TemplateArgs =
 	{
 		{TEXT("ParameterName"),	ParamInfo.DataInterfaceHLSLSymbol},
 	};
-
-	FString TemplateFile;
-	LoadShaderSourceFile(NDIStaticMeshLocal::TemplateShaderFile, EShaderPlatform::SP_PCD3D_SM5, &TemplateFile, nullptr);
-	OutHLSL += FString::Format(*TemplateFile, TemplateArgs);
+	AppendTemplateHLSL(OutHLSL, NDIStaticMeshLocal::TemplateShaderFile, TemplateArgs);
 }
 
 bool UNiagaraDataInterfaceStaticMesh::GetFunctionHLSL(const FNiagaraDataInterfaceGPUParamInfo& ParamInfo, const FNiagaraDataInterfaceGeneratedFunction& FunctionInfo, int FunctionInstanceIndex, FString& OutHLSL)
diff --git a/Engine/Plugins/FX/Niagara/Source/Niagara/Private/DataInterface/NiagaraDataInterfaceVirtualTexture.cpp b/Engine/Plugins/FX/Niagara/Source/Niagara/Private/DataInterface/NiagaraDataInterfaceVirtualTexture.cpp
--- a/Engine/Plugins/FX/Niagara/Source/Niagara/Private/DataInterface/NiagaraDataInterfaceVirtualTexture.cpp
+++ b/Engine/Plugins/FX/Niagara/Source/Niagara/Private/DataInterface/NiagaraDataInterfaceVirtualTexture.cpp
@@ -291,8 +291,8 @@ bool UNiagaraDataInterfaceVirtualTexture::AppendCompileHash(FNiagaraCompileHashV
 	using namespace NDIVirtualTextureLocal;
 
 	bool bSuccess = Super::AppendCompileHash(InVisitor);
-	InVisitor->UpdateString(TEXT("UNiagaraDataInterfaceVirtualTextureHLSLSource"), GetShaderFileHash(CommonShaderFilePath, EShaderPlatform::SP_PCD3D_SM5).ToString());
-	InVisitor->UpdateString(TEXT("UNiagaraDataInterfaceVirtualTextureHLSLSource"), GetShaderFileHash(TemplateShaderFilePath, EShaderPlatform::SP_PCD3D_SM5).ToString());
+	InVisitor->UpdateShaderFile(CommonShaderFilePath);
+	InVisitor->UpdateShaderFile(TemplateShaderFilePath);
 	bSuccess &= InVisitor->UpdateShaderParameters<FShaderParameters>();
 	return bSuccess;
 }
@@ -306,15 +306,11 @@ void UNiagaraDataInterfaceVirtualTexture::GetCommonHLSL(FString& OutHlsl)
 
 void UNiagaraDataInterfaceVirtualTexture::GetParameterDefinitionHLSL(const FNiagaraDataInterfaceGPUParamInfo& ParamInfo, FString& OutHLSL)
 {
-	using namespace NDIVirtualTextureLocal;
-	TMap<FString, FStringFormatArg> TemplateArgs =
+	const TMap<FString, FStringFormatArg> TemplateArgs =
 	{
 		{TEXT("ParameterName"),	ParamInfo.DataInterfaceHLSLSymbol},
 	};
-
-	FString TemplateFile;
-	LoadShaderSourceFile(TemplateShaderFilePath, EShaderPlatform::SP_PCD3D_SM5, &TemplateFile, nullptr);
-	OutHLSL += FString::Format(*TemplateFile, TemplateArgs);
+	AppendTemplateHLSL(OutHLSL, NDIVirtualTextureLocal::TemplateShaderFilePath, TemplateArgs);
 }
 
 bool UNiagaraDataInterfaceVirtualTexture::GetFunctionHLSL(const FNiagaraDataInterfaceGPUParamInfo& ParamInfo, const FNiagaraDataInterfaceGeneratedFunction& FunctionInfo, int FunctionInstanceIndex, FString& OutHLSL)
diff --git a/Engine/Plugins/FX/Niagara/Source/Niagara/Private/NDIRenderTargetVolumeSimCacheData.h b/Engine/Plugins/FX/Niagara/Source/Niagara/Private/NDIRenderTargetVolumeSimCacheData.h
--- a/Engine/Plugins/FX/Niagara/Source/Niagara/Private/NDIRenderTargetVolumeSimCacheData.h
+++ b/Engine/Plugins/FX/Niagara/Source/Niagara/Private/NDIRenderTargetVolumeSimCacheData.h
@@ -18,7 +18,7 @@ struct FNDIRenderTargetVolumeSimCacheFrame
 
 	uint8* GetPixelData()
 	{
-		if (PixelData == nullptr && CompressedSize > 0)
+		if (PixelData == nullptr && UncompressedSize > 0)
 		{
 			BulkData.GetCopy(reinterpret_cast<void**>(&PixelData), true);
 		}
diff --git a/Engine/Plugins/FX/Niagara/Source/Niagara/Private/NiagaraCompileHashVisitor.cpp b/Engine/Plugins/FX/Niagara/Source/Niagara/Private/NiagaraCompileHashVisitor.cpp
new file mode 100644
--- /dev/null
+++ b/Engine/Plugins/FX/Niagara/Source/Niagara/Private/NiagaraCompileHashVisitor.cpp
@@ -0,0 +1,26 @@
+// Copyright Epic Games, Inc. All Rights Reserved.
+
+#include "NiagaraCompileHashVisitor.h"
+#include "ShaderCompilerCore.h"
+
+bool FNiagaraCompileHashVisitor::UpdateShaderFile(const TCHAR* ShaderFilePath)
+{
+	const FSHAHash Hash = GetShaderFileHash(ShaderFilePath, EShaderPlatform::SP_PCD3D_SM5);
+	FString SanitizedShaderName(ShaderFilePath);
+	SanitizedShaderName.ReplaceCharInline('/', '_');
+	SanitizedShaderName.ReplaceCharInline('.', '_');
+	return UpdateString(*SanitizedShaderName, Hash.ToString());
+}
+
+bool FNiagaraCompileHashVisitor::UpdateString(const TCHAR* InDebugName, FStringView InData)
+{
+	HashState.Update((const uint8*)InData.GetData(), sizeof(TCHAR) * InData.Len());
+#if WITH_EDITORONLY_DATA
+	if (LogCompileIdGeneration != 0)
+	{
+		Values.Top().PropertyKeys.Push(InDebugName);
+		Values.Top().PropertyValues.Push(FString(InData));
+	}
+#endif
+	return true;
+}
diff --git a/Engine/Plugins/FX/Niagara/Source/Niagara/Private/NiagaraDataInterface.cpp b/Engine/Plugins/FX/Niagara/Source/Niagara/Private/NiagaraDataInterface.cpp
--- a/Engine/Plugins/FX/Niagara/Source/Niagara/Private/NiagaraDataInterface.cpp
+++ b/Engine/Plugins/FX/Niagara/Source/Niagara/Private/NiagaraDataInterface.cpp
@@ -143,6 +143,16 @@ bool UNiagaraDataInterface::IsDataInterfaceType(const FNiagaraTypeDefinition& Ty
 	return false;
 }
 
+#if WITH_EDITORONLY_DATA
+void UNiagaraDataInterface::AppendTemplateHLSL(FString& OutHLSL, const TCHAR* TemplateShaderFile, const TMap<FString, FStringFormatArg>& TemplateArgs) const
+{
+	FString TemplateFile;
+	LoadShaderSourceFile(TemplateShaderFile, EShaderPlatform::SP_PCD3D_SM5, &TemplateFile, nullptr);
+	OutHLSL.Append(FString::Format(*TemplateFile, TemplateArgs));
+	OutHLSL.AppendChar('\n');
+}
+#endif
+
 bool UNiagaraDataInterface::CopyToInternal(UNiagaraDataInterface* Destination) const
 {
 	if (Destination == nullptr || Destination->GetClass() != GetClass())
diff --git a/Engine/Plugins/FX/Niagara/Source/Niagara/Private/NiagaraDataInterface2DArrayTexture.cpp b/Engine/Plugins/FX/Niagara/Source/Niagara/Private/NiagaraDataInterface2DArrayTexture.cpp
--- a/Engine/Plugins/FX/Niagara/Source/Niagara/Private/NiagaraDataInterface2DArrayTexture.cpp
+++ b/Engine/Plugins/FX/Niagara/Source/Niagara/Private/NiagaraDataInterface2DArrayTexture.cpp
@@ -224,21 +224,18 @@ void UNiagaraDataInterface2DArrayTexture::GetTextureDimensions(FVectorVMExternal
 bool UNiagaraDataInterface2DArrayTexture::AppendCompileHash(FNiagaraCompileHashVisitor* InVisitor) const
 {
 	bool bSuccess = Super::AppendCompileHash(InVisitor);
-	InVisitor->UpdateString(TEXT("UNiagaraDataInterface2DArrayTextureHLSLSource"), GetShaderFileHash(TemplateShaderFilePath, EShaderPlatform::SP_PCD3D_SM5).ToString());
+	bSuccess &= InVisitor->UpdateShaderFile(TemplateShaderFilePath);
 	bSuccess &= InVisitor->UpdateShaderParameters<FShaderParameters>();
 	return bSuccess;
 }
 
 void UNiagaraDataInterface2DArrayTexture::GetParameterDefinitionHLSL(const FNiagaraDataInterfaceGPUParamInfo& ParamInfo, FString& OutHLSL)
 {
-	TMap<FString, FStringFormatArg> TemplateArgs =
+	const TMap<FString, FStringFormatArg> TemplateArgs =
 	{
 		{TEXT("ParameterName"),	ParamInfo.DataInterfaceHLSLSymbol},
 	};
-
-	FString TemplateFile;
-	LoadShaderSourceFile(TemplateShaderFilePath, EShaderPlatform::SP_PCD3D_SM5, &TemplateFile, nullptr);
-	OutHLSL += FString::Format(*TemplateFile, TemplateArgs);
+	AppendTemplateHLSL(OutHLSL, TemplateShaderFilePath, TemplateArgs);
 }
 
 bool UNiagaraDataInterface2DArrayTexture::GetFunctionHLSL(const FNiagaraDataInterfaceGPUParamInfo& ParamInfo, const FNiagaraDataInterfaceGeneratedFunction& FunctionInfo, int FunctionInstanceIndex, FString& OutHLSL)
diff --git a/Engine/Plugins/FX/Niagara/Source/Niagara/Private/NiagaraDataInterfaceCamera.cpp b/Engine/Plugins/FX/Niagara/Source/Niagara/Private/NiagaraDataInterfaceCamera.cpp
--- a/Engine/Plugins/FX/Niagara/Source/Niagara/Private/NiagaraDataInterfaceCamera.cpp
+++ b/Engine/Plugins/FX/Niagara/Source/Niagara/Private/NiagaraDataInterfaceCamera.cpp
@@ -381,21 +381,18 @@ bool UNiagaraDataInterfaceCamera::AppendCompileHash(FNiagaraCompileHashVisitor*
 	{
 		return false;
 	}
-	InVisitor->UpdateString(TEXT("UNiagaraDataInterfaceCameraHLSLSource"), GetShaderFileHash(TemplateShaderFilePath, EShaderPlatform::SP_PCD3D_SM5).ToString());
+	InVisitor->UpdateShaderFile(TemplateShaderFilePath);
 	InVisitor->UpdateShaderParameters<FShaderParameters>();
 	return true;
 }
 
 void UNiagaraDataInterfaceCamera::GetParameterDefinitionHLSL(const FNiagaraDataInterfaceGPUParamInfo& ParamInfo, FString& OutHLSL)
 {
-	TMap<FString, FStringFormatArg> TemplateArgs =
+	const TMap<FString, FStringFormatArg> TemplateArgs =
 	{
 		{TEXT("ParameterName"),	ParamInfo.DataInterfaceHLSLSymbol},
 	};
-
-	FString TemplateFile;
-	LoadShaderSourceFile(TemplateShaderFilePath, EShaderPlatform::SP_PCD3D_SM5, &TemplateFile, nullptr);
-	OutHLSL += FString::Format(*TemplateFile, TemplateArgs);
+	AppendTemplateHLSL(OutHLSL, TemplateShaderFilePath, TemplateArgs);
 }
 #endif
 
diff --git a/Engine/Plugins/FX/Niagara/Source/Niagara/Private/NiagaraDataInterfaceCollisionQuery.cpp b/Engine/Plugins/FX/Niagara/Source/Niagara/Private/NiagaraDataInterfaceCollisionQuery.cpp
--- a/Engine/Plugins/FX/Niagara/Source/Niagara/Private/NiagaraDataInterfaceCollisionQuery.cpp
+++ b/Engine/Plugins/FX/Niagara/Source/Niagara/Private/NiagaraDataInterfaceCollisionQuery.cpp
@@ -479,14 +479,11 @@ void UNiagaraDataInterfaceCollisionQuery::ValidateFunction(const FNiagaraFunctio
 #if WITH_EDITORONLY_DATA
 void UNiagaraDataInterfaceCollisionQuery::GetParameterDefinitionHLSL(const FNiagaraDataInterfaceGPUParamInfo& ParamInfo, FString& OutHLSL)
 {
-	TMap<FString, FStringFormatArg> TemplateArgs =
+	const TMap<FString, FStringFormatArg> TemplateArgs =
 	{
 		{TEXT("ParameterName"),	ParamInfo.DataInterfaceHLSLSymbol},
 	};
-
-	FString TemplateFile;
-	LoadShaderSourceFile(NDICollisionQueryLocal::TemplateShaderFile, EShaderPlatform::SP_PCD3D_SM5, &TemplateFile, nullptr);
-	OutHLSL += FString::Format(*TemplateFile, TemplateArgs);
+	AppendTemplateHLSL(OutHLSL, NDICollisionQueryLocal::TemplateShaderFile, TemplateArgs);
 }
 #endif
 
@@ -524,8 +521,8 @@ bool UNiagaraDataInterfaceCollisionQuery::AppendCompileHash(FNiagaraCompileHashV
 	}
 
 	InVisitor->UpdatePOD(TEXT("NiagaraCollisionDI_DistanceField"), IsDistanceFieldEnabled());
-	InVisitor->UpdateString(TEXT("NDICollisionQueryCommonHLSLSource"), GetShaderFileHash(NDICollisionQueryLocal::CommonShaderFile, EShaderPlatform::SP_PCD3D_SM5).ToString());
-	InVisitor->UpdateString(TEXT("NDICollisionQueryTemplateHLSLSource"), GetShaderFileHash(NDICollisionQueryLocal::TemplateShaderFile, EShaderPlatform::SP_PCD3D_SM5).ToString());
+	InVisitor->UpdateShaderFile(NDICollisionQueryLocal::CommonShaderFile);
+	InVisitor->UpdateShaderFile(NDICollisionQueryLocal::TemplateShaderFile);
 	InVisitor->UpdateShaderParameters<FShaderParameters>();
 	InVisitor->UpdateShaderParameters<FGlobalDistanceFieldParameters2>();
 
diff --git a/Engine/Plugins/FX/Niagara/Source/Niagara/Private/NiagaraDataInterfaceCubeTexture.cpp b/Engine/Plugins/FX/Niagara/Source/Niagara/Private/NiagaraDataInterfaceCubeTexture.cpp
--- a/Engine/Plugins/FX/Niagara/Source/Niagara/Private/NiagaraDataInterfaceCubeTexture.cpp
+++ b/Engine/Plugins/FX/Niagara/Source/Niagara/Private/NiagaraDataInterfaceCubeTexture.cpp
@@ -244,21 +244,18 @@ void UNiagaraDataInterfaceCubeTexture::SampleCubeTexture(FVectorVMExternalFuncti
 bool UNiagaraDataInterfaceCubeTexture::AppendCompileHash(FNiagaraCompileHashVisitor* InVisitor) const
 {
 	bool bSuccess = Super::AppendCompileHash(InVisitor);
-	InVisitor->UpdateString(TEXT("UNiagaraDataInterfaceCubeTextureHLSLSource"), GetShaderFileHash(TemplateShaderFilePath, EShaderPlatform::SP_PCD3D_SM5).ToString());
+	bSuccess &= InVisitor->UpdateShaderFile(TemplateShaderFilePath);
 	bSuccess &= InVisitor->UpdateShaderParameters<FShaderParameters>();
 	return bSuccess;
 }
 
 void UNiagaraDataInterfaceCubeTexture::GetParameterDefinitionHLSL(const FNiagaraDataInterfaceGPUParamInfo& ParamInfo, FString& OutHLSL)
 {
-	TMap<FString, FStringFormatArg> TemplateArgs =
+	const TMap<FString, FStringFormatArg> TemplateArgs =
 	{
 		{TEXT("ParameterName"),	ParamInfo.DataInterfaceHLSLSymbol},
 	};
-
-	FString TemplateFile;
-	LoadShaderSourceFile(TemplateShaderFilePath, EShaderPlatform::SP_PCD3D_SM5, &TemplateFile, nullptr);
-	OutHLSL += FString::Format(*TemplateFile, TemplateArgs);
+	AppendTemplateHLSL(OutHLSL, TemplateShaderFilePath, TemplateArgs);
 }
 
 bool UNiagaraDataInterfaceCubeTexture::GetFunctionHLSL(const FNiagaraDataInterfaceGPUParamInfo& ParameterInfo, const FNiagaraDataInterfaceGeneratedFunction& FunctionInfo, int FunctionInstanceIndex, FString& OutHLSL)
diff --git a/Engine/Plugins/FX/Niagara/Source/Niagara/Private/NiagaraDataInterfaceCurveBase.cpp b/Engine/Plugins/FX/Niagara/Source/Niagara/Private/NiagaraDataInterfaceCurveBase.cpp
--- a/Engine/Plugins/FX/Niagara/Source/Niagara/Private/NiagaraDataInterfaceCurveBase.cpp
+++ b/Engine/Plugins/FX/Niagara/Source/Niagara/Private/NiagaraDataInterfaceCurveBase.cpp
@@ -345,8 +345,7 @@ void UNiagaraDataInterfaceCurveBase::SetShaderParameters(const FNiagaraDataInter
 bool UNiagaraDataInterfaceCurveBase::AppendCompileHash(FNiagaraCompileHashVisitor* InVisitor) const
 {
 	bool bSuccess = Super::AppendCompileHash(InVisitor);
-	FSHAHash Hash = GetShaderFileHash(NDICurve_TemplateShaderFile, EShaderPlatform::SP_PCD3D_SM5);
-	InVisitor->UpdateString(TEXT("UNiagaraDataInterfaceCurveTemplateHLSLSource"), Hash.ToString());
+	InVisitor->UpdateShaderFile(NDICurve_TemplateShaderFile);
 	InVisitor->UpdateShaderParameters<FNiagaraDataInterfaceCurveParameters>();
 	return bSuccess;
 }
@@ -356,17 +355,14 @@ void UNiagaraDataInterfaceCurveBase::GetParameterDefinitionHLSL(const FNiagaraDa
 	const int32 NumElements = GetCurveNumElems();
 	check(NumElements > 0 && NumElements <= 4);
 
-	TMap<FString, FStringFormatArg> TemplateArgs =
+	const TMap<FString, FStringFormatArg> TemplateArgs =
 	{
 		{TEXT("ParameterName"),				ParamInfo.DataInterfaceHLSLSymbol},
 		{TEXT("UseStaticBuffer"),			ParamInfo.IsUserParameter() ? TEXT("1") : TEXT("0")},
 		{TEXT("NumElements"),				FString::FromInt(NumElements)},
 		{TEXT("CurveSampleFunctionName"),	GetCurveSampleFunctionName().ToString()},
 	};
-
-	FString TemplateFile;
-	LoadShaderSourceFile(NDICurve_TemplateShaderFile, EShaderPlatform::SP_PCD3D_SM5, &TemplateFile, nullptr);
-	OutHLSL += FString::Format(*TemplateFile, TemplateArgs);
+	AppendTemplateHLSL(OutHLSL, NDICurve_TemplateShaderFile, TemplateArgs);
 }
 #endif
 
diff --git a/Engine/Plugins/FX/Niagara/Source/Niagara/Private/NiagaraDataInterfaceExport.cpp b/Engine/Plugins/FX/Niagara/Source/Niagara/Private/NiagaraDataInterfaceExport.cpp
--- a/Engine/Plugins/FX/Niagara/Source/Niagara/Private/NiagaraDataInterfaceExport.cpp
+++ b/Engine/Plugins/FX/Niagara/Source/Niagara/Private/NiagaraDataInterfaceExport.cpp
@@ -354,14 +354,11 @@ void UNiagaraDataInterfaceExport::GetFunctions(TArray<FNiagaraFunctionSignature>
 #if WITH_EDITORONLY_DATA
 void UNiagaraDataInterfaceExport::GetParameterDefinitionHLSL(const FNiagaraDataInterfaceGPUParamInfo& ParamInfo, FString& OutHLSL)
 {
-	TMap<FString, FStringFormatArg> TemplateArgs =
+	const TMap<FString, FStringFormatArg> TemplateArgs =
 	{
 		{TEXT("ParameterName"),	ParamInfo.DataInterfaceHLSLSymbol},
 	};
-
-	FString TemplateFile;
-	LoadShaderSourceFile(NDIExportLocal::TemplateShaderFile, EShaderPlatform::SP_PCD3D_SM5, &TemplateFile, nullptr);
-	OutHLSL += FString::Format(*TemplateFile, TemplateArgs);
+	AppendTemplateHLSL(OutHLSL, NDIExportLocal::TemplateShaderFile, TemplateArgs);
 }
 
 bool UNiagaraDataInterfaceExport::GetFunctionHLSL(const FNiagaraDataInterfaceGPUParamInfo& ParamInfo, const FNiagaraDataInterfaceGeneratedFunction& FunctionInfo, int FunctionInstanceIndex, FString& OutHLSL)
@@ -380,8 +377,7 @@ bool UNiagaraDataInterfaceExport::GetFunctionHLSL(const FNiagaraDataInterfaceGPU
 bool UNiagaraDataInterfaceExport::AppendCompileHash(FNiagaraCompileHashVisitor* InVisitor) const
 {
 	bool bSuccess = Super::AppendCompileHash(InVisitor);
-	FSHAHash Hash = GetShaderFileHash(NDIExportLocal::TemplateShaderFile, EShaderPlatform::SP_PCD3D_SM5);
-	InVisitor->UpdateString(TEXT("NiagaraDataInterfaceExportTemplateHLSLSource"), Hash.ToString());
+	bSuccess &= InVisitor->UpdateShaderFile(NDIExportLocal::TemplateShaderFile);
 	return bSuccess;
 }
 #endif
diff --git a/Engine/Plugins/FX/Niagara/Source/Niagara/Private/NiagaraDataInterfaceGrid3DCollection.cpp b/Engine/Plugins/FX/Niagara/Source/Niagara/Private/NiagaraDataInterfaceGrid3DCollection.cpp
--- a/Engine/Plugins/FX/Niagara/Source/Niagara/Private/NiagaraDataInterfaceGrid3DCollection.cpp
+++ b/Engine/Plugins/FX/Niagara/Source/Niagara/Private/NiagaraDataInterfaceGrid3DCollection.cpp
@@ -345,7 +345,7 @@ bool UNiagaraDataInterfaceGrid3DCollection::AppendCompileHash(FNiagaraCompileHas
 
 	InVisitor->UpdatePOD(TEXT("UNiagaraDataInterfaceGrid3DCollectionVersion"), (int32)FNiagaraGridCollection3DDIFunctionVersion::LatestVersion);
 	InVisitor->UpdatePOD(TEXT("UNiagaraDataInterfaceGrid3DCollectionSupportsRGBAGrid"), FGrid3DCollectionAttributeHelper::SupportsRGBAGrid() ? 1 : 0);
-	//InVisitor->UpdateString(TEXT("UNiagaraDataInterfaceVolumeTextureHLSLSource"), GetShaderFileHash(TemplateShaderFilePath, EShaderPlatform::SP_PCD3D_SM5).ToString());
+	//InVisitor->UpdateShaderFile(TemplateShaderFilePath);
 	InVisitor->UpdateShaderParameters<FNDIGrid3DShaderParameters>();
 
 	return true;
@@ -1673,9 +1673,7 @@ void UNiagaraDataInterfaceGrid3DCollection::GetParameterDefinitionHLSL(const FNi
 
 	OutHLSL += FString::Format(FormatDeclarations, ArgsDeclarations);
 
-	FString TemplateFile;
-	LoadShaderSourceFile(TemplateShaderFilePath, EShaderPlatform::SP_PCD3D_SM5, &TemplateFile, nullptr);
-	OutHLSL += FString::Format(*TemplateFile, ArgsDeclarations);
+	AppendTemplateHLSL(OutHLSL, TemplateShaderFilePath, ArgsDeclarations);
 }
 void UNiagaraDataInterfaceGrid3DCollection::WriteSetHLSL(const FNiagaraDataInterfaceGPUParamInfo& ParamInfo, const FNiagaraDataInterfaceGeneratedFunction& FunctionInfo, int FunctionInstanceIndex, int32 InNumChannels, FString& OutHLSL)
 {
diff --git a/Engine/Plugins/FX/Niagara/Source/Niagara/Private/NiagaraDataInterfaceIntRenderTarget2D.cpp b/Engine/Plugins/FX/Niagara/Source/Niagara/Private/NiagaraDataInterfaceIntRenderTarget2D.cpp
--- a/Engine/Plugins/FX/Niagara/Source/Niagara/Private/NiagaraDataInterfaceIntRenderTarget2D.cpp
+++ b/Engine/Plugins/FX/Niagara/Source/Niagara/Private/NiagaraDataInterfaceIntRenderTarget2D.cpp
@@ -516,22 +516,19 @@ bool UNiagaraDataInterfaceIntRenderTarget2D::CopyToInternal(UNiagaraDataInterfac
 bool UNiagaraDataInterfaceIntRenderTarget2D::AppendCompileHash(FNiagaraCompileHashVisitor* InVisitor) const
 {
 	bool bSuccess = Super::AppendCompileHash(InVisitor);
-	FSHAHash Hash = GetShaderFileHash(NDIIntRenderTarget2DLocal::TemplateShaderFile, EShaderPlatform::SP_PCD3D_SM5);
-	InVisitor->UpdateString(TEXT("NiagaraDataInterfaceExportTemplateHLSLSource"), Hash.ToString());
-	InVisitor->UpdateShaderParameters<FShaderParameters>();
+	bSuccess &= InVisitor->UpdateShaderFile(NDIIntRenderTarget2DLocal::TemplateShaderFile);
+	bSuccess &= InVisitor->UpdateShaderParameters<FShaderParameters>();
 	return bSuccess;
 }
 
 void UNiagaraDataInterfaceIntRenderTarget2D::GetParameterDefinitionHLSL(const FNiagaraDataInterfaceGPUParamInfo& ParamInfo, FString& OutHLSL)
 {
-	TMap<FString, FStringFormatArg> TemplateArgs =
+	const TMap<FString, FStringFormatArg> TemplateArgs =
 	{
 		{TEXT("ParameterName"),	ParamInfo.DataInterfaceHLSLSymbol},
 	};
 
-	FString TemplateFile;
-	LoadShaderSourceFile(NDIIntRenderTarget2DLocal::TemplateShaderFile, EShaderPlatform::SP_PCD3D_SM5, &TemplateFile, nullptr);
-	OutHLSL += FString::Format(*TemplateFile, TemplateArgs);
+	AppendTemplateHLSL(OutHLSL, NDIIntRenderTarget2DLocal::TemplateShaderFile, TemplateArgs);
 }
 
 bool UNiagaraDataInterfaceIntRenderTarget2D::GetFunctionHLSL(const FNiagaraDataInterfaceGPUParamInfo& ParamInfo, const FNiagaraDataInterfaceGeneratedFunction& FunctionInfo, int FunctionInstanceIndex, FString& OutHLSL)
diff --git a/Engine/Plugins/FX/Niagara/Source/Niagara/Private/NiagaraDataInterfaceLandscape.cpp b/Engine/Plugins/FX/Niagara/Source/Niagara/Private/NiagaraDataInterfaceLandscape.cpp
--- a/Engine/Plugins/FX/Niagara/Source/Niagara/Private/NiagaraDataInterfaceLandscape.cpp
+++ b/Engine/Plugins/FX/Niagara/Source/Niagara/Private/NiagaraDataInterfaceLandscape.cpp
@@ -1264,8 +1264,7 @@ bool UNiagaraDataInterfaceLandscape::AppendCompileHash(FNiagaraCompileHashVisito
 		return false;
 	}
 
-	FSHAHash Hash = GetShaderFileHash(NiagaraDataInterfaceLandscape::TemplateShaderFile, EShaderPlatform::SP_PCD3D_SM5);
-	InVisitor->UpdateString(TEXT("NiagaraDataInterfaceLandscapeHLSLSource"), Hash.ToString());
+	InVisitor->UpdateShaderFile(NiagaraDataInterfaceLandscape::TemplateShaderFile);
 	InVisitor->UpdateShaderParameters<NiagaraDataInterfaceLandscape::FShaderParameters>();
 
 	return true;
@@ -1276,14 +1275,11 @@ void UNiagaraDataInterfaceLandscape::GetParameterDefinitionHLSL(const FNiagaraDa
 {
 	Super::GetParameterDefinitionHLSL(ParamInfo, OutHLSL);
 
-	TMap<FString, FStringFormatArg> TemplateArgs =
+	const TMap<FString, FStringFormatArg> TemplateArgs =
 	{
 		{TEXT("ParameterName"),	ParamInfo.DataInterfaceHLSLSymbol},
 	};
-
-	FString TemplateFile;
-	LoadShaderSourceFile(NiagaraDataInterfaceLandscape::TemplateShaderFile, EShaderPlatform::SP_PCD3D_SM5, &TemplateFile, nullptr);
-	OutHLSL += FString::Format(*TemplateFile, TemplateArgs);
+	AppendTemplateHLSL(OutHLSL, NiagaraDataInterfaceLandscape::TemplateShaderFile, TemplateArgs);
 }
 
 bool UNiagaraDataInterfaceLandscape::GetFunctionHLSL(const FNiagaraDataInterfaceGPUParamInfo& ParamInfo, const FNiagaraDataInterfaceGeneratedFunction& FunctionInfo, int FunctionInstanceIndex, FString& OutHLSL)
diff --git a/Engine/Plugins/FX/Niagara/Source/Niagara/Private/NiagaraDataInterfaceMeshRendererInfo.cpp b/Engine/Plugins/FX/Niagara/Source/Niagara/Private/NiagaraDataInterfaceMeshRendererInfo.cpp
--- a/Engine/Plugins/FX/Niagara/Source/Niagara/Private/NiagaraDataInterfaceMeshRendererInfo.cpp
+++ b/Engine/Plugins/FX/Niagara/Source/Niagara/Private/NiagaraDataInterfaceMeshRendererInfo.cpp
@@ -204,21 +204,18 @@ bool UNiagaraDataInterfaceMeshRendererInfo::Equals(const UNiagaraDataInterface*
 bool UNiagaraDataInterfaceMeshRendererInfo::AppendCompileHash(FNiagaraCompileHashVisitor* InVisitor) const
 {
 	bool bSuccess = Super::AppendCompileHash(InVisitor);
-	bSuccess &= InVisitor->UpdateString(TEXT("NiagaraDataInterfaceMeshRendererInfoHLSLSource"), GetShaderFileHash(NDIMeshRendererInfoInternal::TemplateShaderFile, EShaderPlatform::SP_PCD3D_SM5).ToString());
+	bSuccess &= InVisitor->UpdateShaderFile(NDIMeshRendererInfoInternal::TemplateShaderFile);
 	bSuccess &= InVisitor->UpdateShaderParameters<NDIMeshRendererInfoInternal::FShaderParameters>();
 	return bSuccess;
 }
 
 void UNiagaraDataInterfaceMeshRendererInfo::GetParameterDefinitionHLSL(const FNiagaraDataInterfaceGPUParamInfo& ParamInfo, FString& OutHLSL)
 {
-	TMap<FString, FStringFormatArg> TemplateArgs =
+	const TMap<FString, FStringFormatArg> TemplateArgs =
 	{
 		{TEXT("ParameterName"),	ParamInfo.DataInterfaceHLSLSymbol},
 	};
-
-	FString TemplateFile;
-	LoadShaderSourceFile(NDIMeshRendererInfoInternal::TemplateShaderFile, EShaderPlatform::SP_PCD3D_SM5, &TemplateFile, nullptr);
-	OutHLSL += FString::Format(*TemplateFile, TemplateArgs);
+	AppendTemplateHLSL(OutHLSL, NDIMeshRendererInfoInternal::TemplateShaderFile, TemplateArgs);
 }
 
 bool UNiagaraDataInterfaceMeshRendererInfo::GetFunctionHLSL(const FNiagaraDataInterfaceGPUParamInfo& ParamInfo, const FNiagaraDataInterfaceGeneratedFunction& FunctionInfo, int FunctionInstanceIndex, FString& OutHLSL)
diff --git a/Engine/Plugins/FX/Niagara/Source/Niagara/Private/NiagaraDataInterfaceOcclusion.cpp b/Engine/Plugins/FX/Niagara/Source/Niagara/Private/NiagaraDataInterfaceOcclusion.cpp
--- a/Engine/Plugins/FX/Niagara/Source/Niagara/Private/NiagaraDataInterfaceOcclusion.cpp
+++ b/Engine/Plugins/FX/Niagara/Source/Niagara/Private/NiagaraDataInterfaceOcclusion.cpp
@@ -122,7 +122,7 @@ bool UNiagaraDataInterfaceOcclusion::AppendCompileHash(FNiagaraCompileHashVisito
 	if (!Super::AppendCompileHash(InVisitor))
 		return false;
 
-	InVisitor->UpdateString(TEXT("NiagaraDataInterfaceOcclusionHLSLSource"), GetShaderFileHash(TemplateShaderFilePath, EShaderPlatform::SP_PCD3D_SM5).ToString());
+	InVisitor->UpdateShaderFile(TemplateShaderFilePath);
 	InVisitor->UpdateShaderParameters<FShaderParameters>();
 	return true;
 }
@@ -166,14 +166,11 @@ void UNiagaraDataInterfaceOcclusion::GetParameterDefinitionHLSL(const FNiagaraDa
 {
 	Super::GetParameterDefinitionHLSL(ParamInfo, OutHLSL);
 
-	TMap<FString, FStringFormatArg> TemplateArgs =
+	const TMap<FString, FStringFormatArg> TemplateArgs =
 	{
 		{TEXT("ParameterName"),	ParamInfo.DataInterfaceHLSLSymbol},
 	};
-
-	FString TemplateFile;
-	LoadShaderSourceFile(TemplateShaderFilePath, EShaderPlatform::SP_PCD3D_SM5, &TemplateFile, nullptr);
-	OutHLSL += FString::Format(*TemplateFile, TemplateArgs);
+	AppendTemplateHLSL(OutHLSL, TemplateShaderFilePath, TemplateArgs);
 }
 #endif
 
diff --git a/Engine/Plugins/FX/Niagara/Source/Niagara/Private/NiagaraDataInterfaceRenderTarget2D.cpp b/Engine/Plugins/FX/Niagara/Source/Niagara/Private/NiagaraDataInterfaceRenderTarget2D.cpp
--- a/Engine/Plugins/FX/Niagara/Source/Niagara/Private/NiagaraDataInterfaceRenderTarget2D.cpp
+++ b/Engine/Plugins/FX/Niagara/Source/Niagara/Private/NiagaraDataInterfaceRenderTarget2D.cpp
@@ -341,22 +341,18 @@ bool UNiagaraDataInterfaceRenderTarget2D::CopyToInternal(UNiagaraDataInterface*
 bool UNiagaraDataInterfaceRenderTarget2D::AppendCompileHash(FNiagaraCompileHashVisitor* InVisitor) const
 {
 	bool bSuccess = Super::AppendCompileHash(InVisitor);
-	FSHAHash Hash = GetShaderFileHash(NDIRenderTarget2DLocal::TemplateShaderFile, EShaderPlatform::SP_PCD3D_SM5);
-	InVisitor->UpdateString(TEXT("UNiagaraDataInterfaceRenderTarget2DTemplateHLSLSource"), Hash.ToString());
+	InVisitor->UpdateShaderFile(NDIRenderTarget2DLocal::TemplateShaderFile);
 	InVisitor->UpdateShaderParameters<NDIRenderTarget2DLocal::FShaderParameters>();
 	return bSuccess;
 }
 
 void UNiagaraDataInterfaceRenderTarget2D::GetParameterDefinitionHLSL(const FNiagaraDataInterfaceGPUParamInfo& ParamInfo, FString& OutHLSL)
 {
-	TMap<FString, FStringFormatArg> TemplateArgs =
+	const TMap<FString, FStringFormatArg> TemplateArgs =
 	{
 		{TEXT("ParameterName"),				ParamInfo.DataInterfaceHLSLSymbol},
 	};
-
-	FString TemplateFile;
-	LoadShaderSourceFile(NDIRenderTarget2DLocal::TemplateShaderFile, EShaderPlatform::SP_PCD3D_SM5, &TemplateFile, nullptr);
-	OutHLSL += FString::Format(*TemplateFile, TemplateArgs);
+	AppendTemplateHLSL(OutHLSL, NDIRenderTarget2DLocal::TemplateShaderFile, TemplateArgs);
 }
 
 bool UNiagaraDataInterfaceRenderTarget2D::GetFunctionHLSL(const FNiagaraDataInterfaceGPUParamInfo& ParamInfo, const FNiagaraDataInterfaceGeneratedFunction& FunctionInfo, int FunctionInstanceIndex, FString& OutHLSL)
diff --git a/Engine/Plugins/FX/Niagara/Source/Niagara/Private/NiagaraDataInterfaceRenderTarget2DArray.cpp b/Engine/Plugins/FX/Niagara/Source/Niagara/Private/NiagaraDataInterfaceRenderTarget2DArray.cpp
--- a/Engine/Plugins/FX/Niagara/Source/Niagara/Private/NiagaraDataInterfaceRenderTarget2DArray.cpp
+++ b/Engine/Plugins/FX/Niagara/Source/Niagara/Private/NiagaraDataInterfaceRenderTarget2DArray.cpp
@@ -339,22 +339,18 @@ bool UNiagaraDataInterfaceRenderTarget2DArray::CopyToInternal(UNiagaraDataInterf
 bool UNiagaraDataInterfaceRenderTarget2DArray::AppendCompileHash(FNiagaraCompileHashVisitor* InVisitor) const
 {
 	bool bSuccess = Super::AppendCompileHash(InVisitor);
-	FSHAHash Hash = GetShaderFileHash(NDIRenderTarget2DArrayLocal::TemplateShaderFile, EShaderPlatform::SP_PCD3D_SM5);
-	InVisitor->UpdateString(TEXT("UNiagaraDataInterfaceRenderTarget2DArrayTemplateHLSLSource"), Hash.ToString());
+	InVisitor->UpdateShaderFile(NDIRenderTarget2DArrayLocal::TemplateShaderFile);
 	InVisitor->UpdateShaderParameters<NDIRenderTarget2DArrayLocal::FShaderParameters>();
 	return bSuccess;
 }
 
 void UNiagaraDataInterfaceRenderTarget2DArray::GetParameterDefinitionHLSL(const FNiagaraDataInterfaceGPUParamInfo& ParamInfo, FString& OutHLSL)
 {
-	TMap<FString, FStringFormatArg> TemplateArgs =
+	const TMap<FString, FStringFormatArg> TemplateArgs =
 	{
 		{TEXT("ParameterName"),				ParamInfo.DataInterfaceHLSLSymbol},
 	};
-
-	FString TemplateFile;
-	LoadShaderSourceFile(NDIRenderTarget2DArrayLocal::TemplateShaderFile, EShaderPlatform::SP_PCD3D_SM5, &TemplateFile, nullptr);
-	OutHLSL += FString::Format(*TemplateFile, TemplateArgs);
+	AppendTemplateHLSL(OutHLSL, NDIRenderTarget2DArrayLocal::TemplateShaderFile, TemplateArgs);
 }
 
 bool UNiagaraDataInterfaceRenderTarget2DArray::GetFunctionHLSL(const FNiagaraDataInterfaceGPUParamInfo& ParamInfo, const FNiagaraDataInterfaceGeneratedFunction& FunctionInfo, int FunctionInstanceIndex, FString& OutHLSL)
diff --git a/Engine/Plugins/FX/Niagara/Source/Niagara/Private/NiagaraDataInterfaceRenderTargetCube.cpp b/Engine/Plugins/FX/Niagara/Source/Niagara/Private/NiagaraDataInterfaceRenderTargetCube.cpp
--- a/Engine/Plugins/FX/Niagara/Source/Niagara/Private/NiagaraDataInterfaceRenderTargetCube.cpp
+++ b/Engine/Plugins/FX/Niagara/Source/Niagara/Private/NiagaraDataInterfaceRenderTargetCube.cpp
@@ -333,22 +333,18 @@ bool UNiagaraDataInterfaceRenderTargetCube::CopyToInternal(UNiagaraDataInterface
 bool UNiagaraDataInterfaceRenderTargetCube::AppendCompileHash(FNiagaraCompileHashVisitor* InVisitor) const
 {
 	bool bSuccess = Super::AppendCompileHash(InVisitor);
-	FSHAHash Hash = GetShaderFileHash(NDIRenderTargetCubeLocal::TemplateShaderFile, EShaderPlatform::SP_PCD3D_SM5);
-	InVisitor->UpdateString(TEXT("UNiagaraDataInterfaceRenderTargetCubeTemplateHLSLSource"), Hash.ToString());
+	InVisitor->UpdateShaderFile(NDIRenderTargetCubeLocal::TemplateShaderFile);
 	InVisitor->UpdateShaderParameters<NDIRenderTargetCubeLocal::FShaderParameters>();
 	return bSuccess;
 }
 
 void UNiagaraDataInterfaceRenderTargetCube::GetParameterDefinitionHLSL(const FNiagaraDataInterfaceGPUParamInfo& ParamInfo, FString& OutHLSL)
 {
-	TMap<FString, FStringFormatArg> TemplateArgs =
+	const TMap<FString, FStringFormatArg> TemplateArgs =
 	{
 		{TEXT("ParameterName"),				ParamInfo.DataInterfaceHLSLSymbol},
 	};
-
-	FString TemplateFile;
-	LoadShaderSourceFile(NDIRenderTargetCubeLocal::TemplateShaderFile, EShaderPlatform::SP_PCD3D_SM5, &TemplateFile, nullptr);
-	OutHLSL += FString::Format(*TemplateFile, TemplateArgs);
+	AppendTemplateHLSL(OutHLSL, NDIRenderTargetCubeLocal::TemplateShaderFile, TemplateArgs);
 }
 
 bool UNiagaraDataInterfaceRenderTargetCube::GetFunctionHLSL(const FNiagaraDataInterfaceGPUParamInfo& ParamInfo, const FNiagaraDataInterfaceGeneratedFunction& FunctionInfo, int FunctionInstanceIndex, FString& OutHLSL)
diff --git a/Engine/Plugins/FX/Niagara/Source/Niagara/Private/NiagaraDataInterfaceRigidMeshCollisionQuery.cpp b/Engine/Plugins/FX/Niagara/Source/Niagara/Private/NiagaraDataInterfaceRigidMeshCollisionQuery.cpp
--- a/Engine/Plugins/FX/Niagara/Source/Niagara/Private/NiagaraDataInterfaceRigidMeshCollisionQuery.cpp
+++ b/Engine/Plugins/FX/Niagara/Source/Niagara/Private/NiagaraDataInterfaceRigidMeshCollisionQuery.cpp
@@ -1443,14 +1443,11 @@ void UNiagaraDataInterfaceRigidMeshCollisionQuery::GetParameterDefinitionHLSL(co
 {
 	Super::GetParameterDefinitionHLSL(ParamInfo, OutHLSL);
 
-	TMap<FString, FStringFormatArg> TemplateArgs =
+	const TMap<FString, FStringFormatArg> TemplateArgs =
 	{
 		{TEXT("ParameterName"),	ParamInfo.DataInterfaceHLSLSymbol},
 	};
-
-	FString TemplateFile;
-	LoadShaderSourceFile(NDIRigidMeshCollisionLocal::TemplateShaderFile, EShaderPlatform::SP_PCD3D_SM5, &TemplateFile, nullptr);
-	OutHLSL += FString::Format(*TemplateFile, TemplateArgs);
+	AppendTemplateHLSL(OutHLSL, NDIRigidMeshCollisionLocal::TemplateShaderFile, TemplateArgs);
 }
 
 bool UNiagaraDataInterfaceRigidMeshCollisionQuery::AppendCompileHash(FNiagaraCompileHashVisitor* InVisitor) const
@@ -1460,8 +1457,7 @@ bool UNiagaraDataInterfaceRigidMeshCollisionQuery::AppendCompileHash(FNiagaraCom
 		return false;
 	}
 
-	FSHAHash Hash = GetShaderFileHash(NDIRigidMeshCollisionLocal::TemplateShaderFile, EShaderPlatform::SP_PCD3D_SM5);
-	InVisitor->UpdateString(TEXT("NiagaraDataInterfaceRigidMeshCollisionQueryHLSLSource"), Hash.ToString());
+	InVisitor->UpdateShaderFile(NDIRigidMeshCollisionLocal::TemplateShaderFile);
 	InVisitor->UpdateShaderParameters<NDIRigidMeshCollisionLocal::FShaderParameters>();
 
 	return true;
diff --git a/Engine/Plugins/FX/Niagara/Source/Niagara/Private/NiagaraDataInterfaceSkeletalMesh.cpp b/Engine/Plugins/FX/Niagara/Source/Niagara/Private/NiagaraDataInterfaceSkeletalMesh.cpp
--- a/Engine/Plugins/FX/Niagara/Source/Niagara/Private/NiagaraDataInterfaceSkeletalMesh.cpp
+++ b/Engine/Plugins/FX/Niagara/Source/Niagara/Private/NiagaraDataInterfaceSkeletalMesh.cpp
@@ -2774,8 +2774,8 @@ bool UNiagaraDataInterfaceSkeletalMesh::AppendCompileHash(FNiagaraCompileHashVis
 	if (!Super::AppendCompileHash(InVisitor))
 		return false;
 
-	InVisitor->UpdateString(TEXT("NiagaraDataInterfaceSkeletalMeshHLSLSource"), GetShaderFileHash(NDISkelMeshLocal::CommonShaderFile, EShaderPlatform::SP_PCD3D_SM5).ToString());
-	InVisitor->UpdateString(TEXT("NiagaraDataInterfaceSkeletalMeshTemplateHLSLSource"), GetShaderFileHash(NDISkelMeshLocal::TemplateShaderFile, EShaderPlatform::SP_PCD3D_SM5).ToString());
+	InVisitor->UpdateShaderFile(NDISkelMeshLocal::CommonShaderFile);
+	InVisitor->UpdateShaderFile(NDISkelMeshLocal::TemplateShaderFile);
 	InVisitor->UpdateShaderParameters<NDISkelMeshLocal::FShaderParameters>();
 
 	InVisitor->UpdatePOD(TEXT("NDISkelmesh_Influences"), int(GetDefault<UNiagaraSettings>()->NDISkelMesh_GpuMaxInfluences));
@@ -2805,14 +2805,11 @@ void UNiagaraDataInterfaceSkeletalMesh::GetCommonHLSL(FString& OutHLSL)
 
 void UNiagaraDataInterfaceSkeletalMesh::GetParameterDefinitionHLSL(const FNiagaraDataInterfaceGPUParamInfo& ParamInfo, FString& OutHLSL)
 {
-	TMap<FString, FStringFormatArg> TemplateArgs =
+	const TMap<FString, FStringFormatArg> TemplateArgs =
 	{
 		{TEXT("ParameterName"),	ParamInfo.DataInterfaceHLSLSymbol},
 	};
-
-	FString TemplateFile;
-	LoadShaderSourceFile(NDISkelMeshLocal::TemplateShaderFile, EShaderPlatform::SP_PCD3D_SM5, &TemplateFile, nullptr);
-	OutHLSL += FString::Format(*TemplateFile, TemplateArgs);
+	AppendTemplateHLSL(OutHLSL, NDISkelMeshLocal::TemplateShaderFile, TemplateArgs);
 }
 
 bool UNiagaraDataInterfaceSkeletalMesh::GetFunctionHLSL(const FNiagaraDataInterfaceGPUParamInfo& ParamInfo, const FNiagaraDataInterfaceGeneratedFunction& FunctionInfo, int FunctionInstanceIndex, FString& OutHLSL)
diff --git a/Engine/Plugins/FX/Niagara/Source/Niagara/Private/NiagaraDataInterfaceSpline.cpp b/Engine/Plugins/FX/Niagara/Source/Niagara/Private/NiagaraDataInterfaceSpline.cpp
--- a/Engine/Plugins/FX/Niagara/Source/Niagara/Private/NiagaraDataInterfaceSpline.cpp
+++ b/Engine/Plugins/FX/Niagara/Source/Niagara/Private/NiagaraDataInterfaceSpline.cpp
@@ -396,13 +396,11 @@ void UNiagaraDataInterfaceSpline::GetCommonHLSL(FString& OutHLSL)
 
 void UNiagaraDataInterfaceSpline::GetParameterDefinitionHLSL(const FNiagaraDataInterfaceGPUParamInfo& ParamInfo, FString& OutHLSL)
 {
-	const TMap<FString, FStringFormatArg> TemplateArgs = {
+	const TMap<FString, FStringFormatArg> TemplateArgs =
+	{
 		{TEXT("ParameterName"), ParamInfo.DataInterfaceHLSLSymbol},
 	};
-	
-	FString TemplateFile;
-	LoadShaderSourceFile(NDISplineLocal::TemplateShaderFile, EShaderPlatform::SP_PCD3D_SM5, &TemplateFile, nullptr);
-	OutHLSL += FString::Format(*TemplateFile, TemplateArgs);	
+	AppendTemplateHLSL(OutHLSL, NDISplineLocal::TemplateShaderFile, TemplateArgs);
 }
 
 bool UNiagaraDataInterfaceSpline::GetFunctionHLSL(const FNiagaraDataInterfaceGPUParamInfo& ParamInfo, const FNiagaraDataInterfaceGeneratedFunction& FunctionInfo, int FunctionInstanceIndex, FString& OutHLSL)
@@ -434,8 +432,7 @@ bool UNiagaraDataInterfaceSpline::GetFunctionHLSL(const FNiagaraDataInterfaceGPU
 bool UNiagaraDataInterfaceSpline::AppendCompileHash(FNiagaraCompileHashVisitor* InVisitor) const
 {
 	bool bSuccess = Super::AppendCompileHash(InVisitor);
-	FSHAHash Hash = GetShaderFileHash(NDISplineLocal::TemplateShaderFile, EShaderPlatform::SP_PCD3D_SM5);
-	InVisitor->UpdateString(TEXT("NiagaraDataInterfaceExportTemplateHLSLSource"), Hash.ToString());
+	InVisitor->UpdateShaderFile(NDISplineLocal::TemplateShaderFile);
 	InVisitor->UpdateShaderParameters<NDISplineLocal::FShaderParameters>();
 	return bSuccess;
 }
diff --git a/Engine/Plugins/FX/Niagara/Source/Niagara/Private/NiagaraDataInterfaceSpriteRendererInfo.cpp b/Engine/Plugins/FX/Niagara/Source/Niagara/Private/NiagaraDataInterfaceSpriteRendererInfo.cpp
--- a/Engine/Plugins/FX/Niagara/Source/Niagara/Private/NiagaraDataInterfaceSpriteRendererInfo.cpp
+++ b/Engine/Plugins/FX/Niagara/Source/Niagara/Private/NiagaraDataInterfaceSpriteRendererInfo.cpp
@@ -214,21 +214,18 @@ bool UNiagaraDataInterfaceSpriteRendererInfo::Equals(const UNiagaraDataInterface
 bool UNiagaraDataInterfaceSpriteRendererInfo::AppendCompileHash(FNiagaraCompileHashVisitor* InVisitor) const
 {
 	bool bSuccess = Super::AppendCompileHash(InVisitor);
-	bSuccess &= InVisitor->UpdateString(TEXT("NiagaraDataInterfaceSpriteRendererInfoHLSLSource"), GetShaderFileHash(NDISpriteRendererInfoLocal::TemplateShaderFile, EShaderPlatform::SP_PCD3D_SM5).ToString());
+	bSuccess &= InVisitor->UpdateShaderFile(NDISpriteRendererInfoLocal::TemplateShaderFile);
 	bSuccess &= InVisitor->UpdateShaderParameters<NDISpriteRendererInfoLocal::FShaderParameters>();
 	return bSuccess;
 }
 
 void UNiagaraDataInterfaceSpriteRendererInfo::GetParameterDefinitionHLSL(const FNiagaraDataInterfaceGPUParamInfo& ParamInfo, FString& OutHLSL)
 {
-	TMap<FString, FStringFormatArg> TemplateArgs =
+	const TMap<FString, FStringFormatArg> TemplateArgs =
 	{
 		{TEXT("ParameterName"),	ParamInfo.DataInterfaceHLSLSymbol},
 	};
-
-	FString TemplateFile;
-	LoadShaderSourceFile(NDISpriteRendererInfoLocal::TemplateShaderFile, EShaderPlatform::SP_PCD3D_SM5, &TemplateFile, nullptr);
-	OutHLSL += FString::Format(*TemplateFile, TemplateArgs);
+	AppendTemplateHLSL(OutHLSL, NDISpriteRendererInfoLocal::TemplateShaderFile, TemplateArgs);
 }
 
 bool UNiagaraDataInterfaceSpriteRendererInfo::GetFunctionHLSL(const FNiagaraDataInterfaceGPUParamInfo& ParamInfo, const FNiagaraDataInterfaceGeneratedFunction& FunctionInfo, int FunctionInstanceIndex, FString& OutHLSL)
diff --git a/Engine/Plugins/FX/Niagara/Source/Niagara/Private/NiagaraDataInterfaceTexture.cpp b/Engine/Plugins/FX/Niagara/Source/Niagara/Private/NiagaraDataInterfaceTexture.cpp
--- a/Engine/Plugins/FX/Niagara/Source/Niagara/Private/NiagaraDataInterfaceTexture.cpp
+++ b/Engine/Plugins/FX/Niagara/Source/Niagara/Private/NiagaraDataInterfaceTexture.cpp
@@ -368,21 +368,18 @@ void UNiagaraDataInterfaceTexture::VMGetNumMipLevels(FVectorVMExternalFunctionCo
 bool UNiagaraDataInterfaceTexture::AppendCompileHash(FNiagaraCompileHashVisitor* InVisitor) const
 {
 	bool bSuccess = Super::AppendCompileHash(InVisitor);
-	InVisitor->UpdateString(TEXT("UNiagaraDataInterfaceTextureHLSLSource"), GetShaderFileHash(TemplateShaderFilePath, EShaderPlatform::SP_PCD3D_SM5).ToString());
+	InVisitor->UpdateShaderFile(TemplateShaderFilePath);
 	bSuccess &= InVisitor->UpdateShaderParameters<FShaderParameters>();
 	return bSuccess;
 }
 
 void UNiagaraDataInterfaceTexture::GetParameterDefinitionHLSL(const FNiagaraDataInterfaceGPUParamInfo& ParamInfo, FString& OutHLSL)
 {
-	TMap<FString, FStringFormatArg> TemplateArgs =
+	const TMap<FString, FStringFormatArg> TemplateArgs =
 	{
 		{TEXT("ParameterName"),	ParamInfo.DataInterfaceHLSLSymbol},
 	};
-
-	FString TemplateFile;
-	LoadShaderSourceFile(TemplateShaderFilePath, EShaderPlatform::SP_PCD3D_SM5, &TemplateFile, nullptr);
-	OutHLSL += FString::Format(*TemplateFile, TemplateArgs);
+	AppendTemplateHLSL(OutHLSL, TemplateShaderFilePath, TemplateArgs);
 }
 
 bool UNiagaraDataInterfaceTexture::GetFunctionHLSL(const FNiagaraDataInterfaceGPUParamInfo& ParamInfo, const FNiagaraDataInterfaceGeneratedFunction& FunctionInfo, int FunctionInstanceIndex, FString& OutHLSL)
diff --git a/Engine/Plugins/FX/Niagara/Source/Niagara/Private/NiagaraDataInterfaceVolumeCache.cpp b/Engine/Plugins/FX/Niagara/Source/Niagara/Private/NiagaraDataInterfaceVolumeCache.cpp
--- a/Engine/Plugins/FX/Niagara/Source/Niagara/Private/NiagaraDataInterfaceVolumeCache.cpp
+++ b/Engine/Plugins/FX/Niagara/Source/Niagara/Private/NiagaraDataInterfaceVolumeCache.cpp
@@ -419,21 +419,19 @@ bool UNiagaraDataInterfaceVolumeCache::PerInstanceTick(void* PerInstanceData, FN
 bool UNiagaraDataInterfaceVolumeCache::AppendCompileHash(FNiagaraCompileHashVisitor* InVisitor) const
 {
 	bool bSuccess = Super::AppendCompileHash(InVisitor);
-	InVisitor->UpdateString(TEXT("UNiagaraDataInterfaceVolumeTextureHLSLSource"), GetShaderFileHash(TemplateShaderFilePath, EShaderPlatform::SP_PCD3D_SM5).ToString());
+	InVisitor->UpdateShaderFile(TemplateShaderFilePath);
 	bSuccess &= InVisitor->UpdateShaderParameters<FShaderParameters>();
 	return bSuccess;
 }
 
 void UNiagaraDataInterfaceVolumeCache::GetParameterDefinitionHLSL(const FNiagaraDataInterfaceGPUParamInfo& ParamInfo, FString& OutHLSL)
 {
-	TMap<FString, FStringFormatArg> TemplateArgs =
+	const TMap<FString, FStringFormatArg> TemplateArgs =
 	{
 		{TEXT("ParameterName"),	ParamInfo.DataInterfaceHLSLSymbol},
 	};
 
-	FString TemplateFile;
-	LoadShaderSourceFile(TemplateShaderFilePath, EShaderPlatform::SP_PCD3D_SM5, &TemplateFile, nullptr);
-	OutHLSL += FString::Format(*TemplateFile, TemplateArgs);
+	AppendTemplateHLSL(OutHLSL, TemplateShaderFilePath, TemplateArgs);
 }
 
 bool UNiagaraDataInterfaceVolumeCache::GetFunctionHLSL(const FNiagaraDataInterfaceGPUParamInfo& ParamInfo, const FNiagaraDataInterfaceGeneratedFunction& FunctionInfo, int FunctionInstanceIndex, FString& OutHLSL)
diff --git a/Engine/Plugins/FX/Niagara/Source/Niagara/Private/NiagaraDataInterfaceVolumeTexture.cpp b/Engine/Plugins/FX/Niagara/Source/Niagara/Private/NiagaraDataInterfaceVolumeTexture.cpp
--- a/Engine/Plugins/FX/Niagara/Source/Niagara/Private/NiagaraDataInterfaceVolumeTexture.cpp
+++ b/Engine/Plugins/FX/Niagara/Source/Niagara/Private/NiagaraDataInterfaceVolumeTexture.cpp
@@ -253,21 +253,19 @@ void UNiagaraDataInterfaceVolumeTexture::SampleVolumeTexture(FVectorVMExternalFu
 bool UNiagaraDataInterfaceVolumeTexture::AppendCompileHash(FNiagaraCompileHashVisitor* InVisitor) const
 {
 	bool bSuccess = Super::AppendCompileHash(InVisitor);
-	InVisitor->UpdateString(TEXT("UNiagaraDataInterfaceVolumeTextureHLSLSource"), GetShaderFileHash(TemplateShaderFilePath, EShaderPlatform::SP_PCD3D_SM5).ToString());
+	InVisitor->UpdateShaderFile(TemplateShaderFilePath);
 	bSuccess &= InVisitor->UpdateShaderParameters<FShaderParameters>();
 	return bSuccess;
 }
 
 void UNiagaraDataInterfaceVolumeTexture::GetParameterDefinitionHLSL(const FNiagaraDataInterfaceGPUParamInfo& ParamInfo, FString& OutHLSL)
 {
-	TMap<FString, FStringFormatArg> TemplateArgs =
+	const TMap<FString, FStringFormatArg> TemplateArgs =
 	{
 		{TEXT("ParameterName"),	ParamInfo.DataInterfaceHLSLSymbol},
 	};
 
-	FString TemplateFile;
-	LoadShaderSourceFile(TemplateShaderFilePath, EShaderPlatform::SP_PCD3D_SM5, &TemplateFile, nullptr);
-	OutHLSL += FString::Format(*TemplateFile, TemplateArgs);
+	AppendTemplateHLSL(OutHLSL, TemplateShaderFilePath, TemplateArgs);
 }
 
 bool UNiagaraDataInterfaceVolumeTexture::GetFunctionHLSL(const FNiagaraDataInterfaceGPUParamInfo& ParamInfo, const FNiagaraDataInterfaceGeneratedFunction& FunctionInfo, int FunctionInstanceIndex, FString& OutHLSL)
diff --git a/Engine/Plugins/FX/Niagara/Source/Niagara/Public/NiagaraCompileHashVisitor.h b/Engine/Plugins/FX/Niagara/Source/Niagara/Public/NiagaraCompileHashVisitor.h
--- a/Engine/Plugins/FX/Niagara/Source/Niagara/Public/NiagaraCompileHashVisitor.h
+++ b/Engine/Plugins/FX/Niagara/Source/Niagara/Public/NiagaraCompileHashVisitor.h
@@ -144,19 +144,10 @@ public:
 		return UpdatePOD(ShaderParametersMetadata->GetStructTypeName(), ShaderParametersMetadata->GetLayoutHash());
 	}
 
+	bool UpdateShaderFile(const TCHAR* ShaderFilePath);
+
 	/**
 	Adds an string value to the hash.
 	*/
-	bool UpdateString(const TCHAR* InDebugName, FStringView InData)
-	{
-		HashState.Update((const uint8 *)InData.GetData(), sizeof(TCHAR)*InData.Len());
-#if WITH_EDITORONLY_DATA
-		if (LogCompileIdGeneration != 0)
-		{
-			Values.Top().PropertyKeys.Push(InDebugName);
-			Values.Top().PropertyValues.Push(FString(InData));
-		}
-#endif
-		return true;
-	}
+	bool UpdateString(const TCHAR* InDebugName, FStringView InData);
 };
