- Add solid box option for drawing bounds via Niagara Debugger

#rnx

diff --git a/Engine/Plugins/FX/Niagara/Source/Niagara/Private/NiagaraDebugHud.cpp b/Engine/Plugins/FX/Niagara/Source/Niagara/Private/NiagaraDebugHud.cpp
--- a/Engine/Plugins/FX/Niagara/Source/Niagara/Private/NiagaraDebugHud.cpp
+++ b/Engine/Plugins/FX/Niagara/Source/Niagara/Private/NiagaraDebugHud.cpp
@@ -573,24 +573,34 @@ namespace NiagaraDebugLocal
 		return TPair<FVector2D, FVector2D>(StringSize, OutLocation);
 	}
 
-	void DrawBox(UWorld* World, const FVector& Location, const FVector& Extents, const FLinearColor& Color, float Thickness = 3.0f)
+	void DrawBox(UWorld* World, const FVector& Location, const FVector& Extents, const FLinearColor& Color, float SolidAlpha = 0.0f, float Thickness = 3.0f)
 	{
 		if (ULineBatchComponent* LineBatcher = World->LineBatcher)
 		{
-			LineBatcher->DrawLine(Location + FVector( Extents.X,  Extents.Y,  Extents.Z), Location + FVector( Extents.X, -Extents.Y,  Extents.Z), Color, 0, Thickness);
-			LineBatcher->DrawLine(Location + FVector( Extents.X, -Extents.Y,  Extents.Z), Location + FVector(-Extents.X, -Extents.Y,  Extents.Z), Color, 0, Thickness);
-			LineBatcher->DrawLine(Location + FVector(-Extents.X, -Extents.Y,  Extents.Z), Location + FVector(-Extents.X,  Extents.Y,  Extents.Z), Color, 0, Thickness);
-			LineBatcher->DrawLine(Location + FVector(-Extents.X,  Extents.Y,  Extents.Z), Location + FVector( Extents.X,  Extents.Y,  Extents.Z), Color, 0, Thickness);
-
-			LineBatcher->DrawLine(Location + FVector( Extents.X,  Extents.Y, -Extents.Z), Location + FVector( Extents.X, -Extents.Y, -Extents.Z), Color, 0, Thickness);
-			LineBatcher->DrawLine(Location + FVector( Extents.X, -Extents.Y, -Extents.Z), Location + FVector(-Extents.X, -Extents.Y, -Extents.Z), Color, 0, Thickness);
-			LineBatcher->DrawLine(Location + FVector(-Extents.X, -Extents.Y, -Extents.Z), Location + FVector(-Extents.X,  Extents.Y, -Extents.Z), Color, 0, Thickness);
-			LineBatcher->DrawLine(Location + FVector(-Extents.X,  Extents.Y, -Extents.Z), Location + FVector( Extents.X,  Extents.Y, -Extents.Z), Color, 0, Thickness);
-
-			LineBatcher->DrawLine(Location + FVector( Extents.X,  Extents.Y,  Extents.Z), Location + FVector( Extents.X,  Extents.Y, -Extents.Z), Color, 0, Thickness);
-			LineBatcher->DrawLine(Location + FVector( Extents.X, -Extents.Y,  Extents.Z), Location + FVector( Extents.X, -Extents.Y, -Extents.Z), Color, 0, Thickness);
-			LineBatcher->DrawLine(Location + FVector(-Extents.X, -Extents.Y,  Extents.Z), Location + FVector(-Extents.X, -Extents.Y, -Extents.Z), Color, 0, Thickness);
-			LineBatcher->DrawLine(Location + FVector(-Extents.X,  Extents.Y,  Extents.Z), Location + FVector(-Extents.X,  Extents.Y, -Extents.Z), Color, 0, Thickness);
+			if (SolidAlpha > 0.0f)
+			{
+				const FBox BoundsBox(-Extents, Extents);
+				FColor BoxColor = Color.ToFColor(false);
+				BoxColor.A = FMath::Clamp(int(SolidAlpha * 255.0f), 0, 255);
+				LineBatcher->DrawSolidBox(BoundsBox, FTransform(FQuat::Identity, Location), BoxColor, 0, 0.0f);
+			}
+			else
+			{
+				LineBatcher->DrawLine(Location + FVector(Extents.X, Extents.Y, Extents.Z), Location + FVector(Extents.X, -Extents.Y, Extents.Z), Color, 0, Thickness);
+				LineBatcher->DrawLine(Location + FVector(Extents.X, -Extents.Y, Extents.Z), Location + FVector(-Extents.X, -Extents.Y, Extents.Z), Color, 0, Thickness);
+				LineBatcher->DrawLine(Location + FVector(-Extents.X, -Extents.Y, Extents.Z), Location + FVector(-Extents.X, Extents.Y, Extents.Z), Color, 0, Thickness);
+				LineBatcher->DrawLine(Location + FVector(-Extents.X, Extents.Y, Extents.Z), Location + FVector(Extents.X, Extents.Y, Extents.Z), Color, 0, Thickness);
+
+				LineBatcher->DrawLine(Location + FVector(Extents.X, Extents.Y, -Extents.Z), Location + FVector(Extents.X, -Extents.Y, -Extents.Z), Color, 0, Thickness);
+				LineBatcher->DrawLine(Location + FVector(Extents.X, -Extents.Y, -Extents.Z), Location + FVector(-Extents.X, -Extents.Y, -Extents.Z), Color, 0, Thickness);
+				LineBatcher->DrawLine(Location + FVector(-Extents.X, -Extents.Y, -Extents.Z), Location + FVector(-Extents.X, Extents.Y, -Extents.Z), Color, 0, Thickness);
+				LineBatcher->DrawLine(Location + FVector(-Extents.X, Extents.Y, -Extents.Z), Location + FVector(Extents.X, Extents.Y, -Extents.Z), Color, 0, Thickness);
+
+				LineBatcher->DrawLine(Location + FVector(Extents.X, Extents.Y, Extents.Z), Location + FVector(Extents.X, Extents.Y, -Extents.Z), Color, 0, Thickness);
+				LineBatcher->DrawLine(Location + FVector(Extents.X, -Extents.Y, Extents.Z), Location + FVector(Extents.X, -Extents.Y, -Extents.Z), Color, 0, Thickness);
+				LineBatcher->DrawLine(Location + FVector(-Extents.X, -Extents.Y, Extents.Z), Location + FVector(-Extents.X, -Extents.Y, -Extents.Z), Color, 0, Thickness);
+				LineBatcher->DrawLine(Location + FVector(-Extents.X, Extents.Y, Extents.Z), Location + FVector(-Extents.X, Extents.Y, -Extents.Z), Color, 0, Thickness);
+			}
 		}
 	}
 
@@ -2482,7 +2492,7 @@ void FNiagaraDebugHud::DrawComponents(FNiagaraWorldManager* WorldManager, UCanva
 			const FBox Bounds = NiagaraComponent->CalcBounds(NiagaraComponent->GetComponentTransform()).GetBox();
 			if (Bounds.IsValid)
 			{
-				DrawBox(World, Bounds.GetCenter(), Bounds.GetExtent(), FColor::Red);
+				DrawBox(World, Bounds.GetCenter(), Bounds.GetExtent(), FColor::Red, Settings.SystemBoundsSolidBoxAlpha);
 			}
 		}
 
diff --git a/Engine/Plugins/FX/Niagara/Source/Niagara/Public/NiagaraDebuggerCommon.h b/Engine/Plugins/FX/Niagara/Source/Niagara/Public/NiagaraDebuggerCommon.h
--- a/Engine/Plugins/FX/Niagara/Source/Niagara/Public/NiagaraDebuggerCommon.h
+++ b/Engine/Plugins/FX/Niagara/Source/Niagara/Public/NiagaraDebuggerCommon.h
@@ -531,6 +531,10 @@ struct NIAGARA_API FNiagaraDebugHUDSettingsData
 	UPROPERTY(Config, EditAnywhere, Category = "Debug System")
 	bool bSystemShowBounds = false;
 
+	/** When bounds display is enabled allows you to draw a solid box if alpha is > 0. */
+	UPROPERTY(Config, EditAnywhere, Category = "Debug System", meta = (EditCondition = "bSystemShowBounds"))
+	float SystemBoundsSolidBoxAlpha = 0.0f;
+
 	/** When disabled in world rendering will show systems deactivated by scalability. */
 	UPROPERTY(Config, EditAnywhere, Category = "Debug System", meta = (EditCondition = "SystemDebugVerbosity != ENiagaraDebugHudVerbosity::None"))
 	bool bSystemShowActiveOnlyInWorld = true;
