- Fix for Mesh Renderer Wireframe mode crashing when using transparent material

#rnx

diff --git a/Engine/Plugins/FX/Niagara/Source/Niagara/Private/NiagaraRendererMeshes.cpp b/Engine/Plugins/FX/Niagara/Source/Niagara/Private/NiagaraRendererMeshes.cpp
--- a/Engine/Plugins/FX/Niagara/Source/Niagara/Private/NiagaraRendererMeshes.cpp
+++ b/Engine/Plugins/FX/Niagara/Source/Niagara/Private/NiagaraRendererMeshes.cpp
@@ -285,8 +285,10 @@ void FNiagaraRendererMeshes::PrepareParticleMeshRenderData(FParticleMeshRenderDa
 	}
 
 	// Check if any materials are translucent and if we can pickup the low latency data
+	const bool bIsWireframe = AllowDebugViewmodes() && ViewFamily.EngineShowFlags.Wireframe;
 	ParticleMeshRenderData.bIsGpuLowLatencyTranslucency =
 		bGpuLowLatencyTranslucency &&
+		!bIsWireframe &&
 		!SceneProxy->CastsVolumetricTranslucentShadow() &&
 		ParticleMeshRenderData.DynamicDataMesh->Materials.Num() > 0 &&
 		ParticleMeshRenderData.DynamicDataMesh->IsGpuLowLatencyTranslucencyEnabled();
