Fix for UNiagaragraph::PostLoad missing some script variables during creation.  Ensure that we include MapSets as well as MapGets.
#rnx

diff --git a/Engine/Plugins/FX/Niagara/Source/NiagaraEditor/Private/NiagaraGraph.cpp b/Engine/Plugins/FX/Niagara/Source/NiagaraEditor/Private/NiagaraGraph.cpp
--- a/Engine/Plugins/FX/Niagara/Source/NiagaraEditor/Private/NiagaraGraph.cpp
+++ b/Engine/Plugins/FX/Niagara/Source/NiagaraEditor/Private/NiagaraGraph.cpp
@@ -3789,9 +3789,17 @@ void UNiagaraGraph::RefreshParameterReferences() const
 		{
 			for (UEdGraphPin* Pin : Node->Pins)
 			{
-				// we only consider output parameter pins
-				if ((Pin->Direction == EGPD_Output) && (Pin->PinType.PinSubCategory == UNiagaraNodeParameterMapBase::ParameterPinSubCategory))
+				// we only consider parameter pins, inputs for MapSets and outputs for MapGets
+				if (Pin->PinType.PinSubCategory == UNiagaraNodeParameterMapBase::ParameterPinSubCategory)
 				{
+					const bool bRelevantPin = ((Pin->Direction == EEdGraphPinDirection::EGPD_Input) && Node->IsA<UNiagaraNodeParameterMapSet>())
+						|| ((Pin->Direction == EEdGraphPinDirection::EGPD_Output) && Node->IsA<UNiagaraNodeParameterMapGet>());
+					
+					if (!bRelevantPin)
+					{
+						continue;
+					}
+
 					bool bAlreadyHandledPin = false;
 					HandledParameterMapPins.FindOrAdd(Pin, &bAlreadyHandledPin);
 
