Works around object loading shenanigans with the UNiagaraEmitter/UNiagaraScript by ensuring that the UNiagaraEmitter is fully loaded during PostLoad() through a call to Preload.

diff --git a/Engine/Plugins/FX/Niagara/Source/Niagara/Private/NiagaraEmitter.cpp b/Engine/Plugins/FX/Niagara/Source/Niagara/Private/NiagaraEmitter.cpp
--- a/Engine/Plugins/FX/Niagara/Source/Niagara/Private/NiagaraEmitter.cpp
+++ b/Engine/Plugins/FX/Niagara/Source/Niagara/Private/NiagaraEmitter.cpp
@@ -547,6 +547,21 @@ void UNiagaraEmitter::PostLoad()
 {
 	Super::PostLoad();
 
+#if WITH_EDITOR
+	// It is unclear in what conditions we can have our PostLoad() invoked without the package having been fully serialized.  Failure to
+	// have the data loaded results in the emitter (and potentially all other objects that reference it) being left in an indeterminate
+	// state.  While investigation continues into hwo this happens, this call to Preload() seems to resolve the somewhat reproducible case
+	// that we have.  There's some evidence that UNiagaraScript having instanced properties is leading to an edge case in object serialization
+	// but more investigation is required.
+	if (HasAnyFlags(RF_NeedLoad))
+	{
+		if (FLinkerLoad* LinkerLoad = GetLinker())
+		{
+			LinkerLoad->Preload(this);
+		}
+	}
+#endif
+
 #if WITH_EDITORONLY_DATA
 	CheckVersionDataAvailable();
 #endif
@@ -761,7 +776,41 @@ void FVersionedNiagaraEmitterData::PostLoad(UNiagaraEmitter& Emitter, bool bIsCo
 			GraphSource->ConditionalPostLoad();
 			GraphSource->PostLoadFromEmitter(FVersionedNiagaraEmitter(&Emitter, Version.VersionGuid));
 		}
-		
+		// more debug code to track what's going wrong with GraphSource being null
+#if !(UE_BUILD_SHIPPING || UE_BUILD_TEST)
+		else
+		{
+			ensureMsgf(!Emitter.HasAnyFlags(RF_NeedLoad), TEXT("Emitter (%s) should have been serialized by now!"), *Emitter.GetPathName());
+
+			FUObjectThreadContext& ThreadContext = FUObjectThreadContext::Get();
+
+			FString ErrorMsg;
+
+			const int32 PostLoadContextCount = ThreadContext.DebugPostLoad.Num();
+
+			ErrorMsg.Appendf(TEXT("ObjFlags along outer: "));
+			UE_LOG(LogNiagara, Warning, TEXT("ObjFlags along outer"));
+			UObject* ThisObject = &Emitter;
+			while (ThisObject != nullptr)
+			{
+				UE_LOG(LogNiagara, Warning, TEXT("\t%s = %x"), *ThisObject->GetName(), ThisObject->GetFlags());
+				ErrorMsg.Appendf(TEXT("(%s = %x), "), *ThisObject->GetName(), ThisObject->GetFlags());
+				ThisObject = ThisObject->GetOuter();
+			}
+
+			ErrorMsg.Appendf(TEXT("Stack of objects (%d) in PostLoad chain: "), PostLoadContextCount);
+
+			for (int32 PostLoadIt = PostLoadContextCount - 1; PostLoadIt >= 0; --PostLoadIt)
+			{
+				ErrorMsg.Appendf(TEXT("{%s}, "), ThreadContext.DebugPostLoad[PostLoadIt]
+					? *ThreadContext.DebugPostLoad[PostLoadIt]->GetPathName()
+					: TEXT("<unknown>"));
+			}
+
+			checkf(false, *ErrorMsg);
+		}
+#endif
+
 		// Prepare for emitter inheritance.
 		if (GetParent().Emitter != nullptr)
 		{
