Added keywords to Niagara component functions and fixed comments

diff --git a/Engine/Plugins/FX/Niagara/Source/Niagara/Public/NiagaraComponent.h b/Engine/Plugins/FX/Niagara/Source/Niagara/Public/NiagaraComponent.h
--- a/Engine/Plugins/FX/Niagara/Source/Niagara/Public/NiagaraComponent.h
+++ b/Engine/Plugins/FX/Niagara/Source/Niagara/Public/NiagaraComponent.h
@@ -418,105 +418,104 @@ public:
 	int32 GetRandomSeedOffset() const { return RandomSeedOffset; }
 
 	/** Sets a Niagara FLinearColor parameter by name, overriding locally if necessary.*/
-	UFUNCTION(BlueprintCallable, Category = Niagara, meta = (DisplayName = "Set Niagara Variable By String (LinearColor)"))
+	UFUNCTION(BlueprintCallable, Category = Niagara, meta = (DisplayName = "Set Niagara Variable By String (LinearColor)", Keywords="user parameter variable color"))
 	void SetNiagaraVariableLinearColor(const FString& InVariableName, const FLinearColor& InValue);
 
 	/** Sets a Niagara FLinearColor parameter by name, overriding locally if necessary.*/
-	UFUNCTION(BlueprintCallable, Category = Niagara, meta = (DisplayName = "Set Niagara Variable (LinearColor)"))
+	UFUNCTION(BlueprintCallable, Category = Niagara, meta = (DisplayName = "Set Niagara Variable (LinearColor)", Keywords="user parameter variable color"))
 	void SetVariableLinearColor(FName InVariableName, const FLinearColor& InValue);
 
-
 	/** Sets a Niagara Vector4 parameter by name, overriding locally if necessary.*/
-	UFUNCTION(BlueprintCallable, Category = Niagara, meta = (DisplayName = "Set Niagara Variable By String (Vector4)"))
+	UFUNCTION(BlueprintCallable, Category = Niagara, meta = (DisplayName = "Set Niagara Variable By String (Vector4)", Keywords="user parameter variable vector"))
 	void SetNiagaraVariableVec4(const FString& InVariableName, const FVector4& InValue);
 
 	/** Sets a Niagara Vector4 parameter by name, overriding locally if necessary.*/
-	UFUNCTION(BlueprintCallable, Category = Niagara, meta = (DisplayName = "Set Niagara Variable (Vector4)"))
+	UFUNCTION(BlueprintCallable, Category = Niagara, meta = (DisplayName = "Set Niagara Variable (Vector4)", Keywords="user parameter variable vector"))
 	void SetVariableVec4(FName InVariableName, const FVector4& InValue);
 
-	/** Sets a Niagara Vector3 parameter by name, overriding locally if necessary.*/
-	UFUNCTION(BlueprintCallable, Category = Niagara, meta = (DisplayName = "Set Niagara Variable By String (Quaternion)"))
+	/** Sets a Niagara quaternion parameter by name, overriding locally if necessary.*/
+	UFUNCTION(BlueprintCallable, Category = Niagara, meta = (DisplayName = "Set Niagara Variable By String (Quaternion)", Keywords="user parameter variable quaternion rotation"))
 	void SetNiagaraVariableQuat(const FString& InVariableName, const FQuat& InValue);
 
-	/** Sets a Niagara Vector3 parameter by name, overriding locally if necessary.*/
-	UFUNCTION(BlueprintCallable, Category = Niagara, meta = (DisplayName = "Set Niagara Variable (Quaternion)"))
+	/** Sets a Niagara quaternion parameter by name, overriding locally if necessary.*/
+	UFUNCTION(BlueprintCallable, Category = Niagara, meta = (DisplayName = "Set Niagara Variable (Quaternion)", Keywords="user parameter variable quaternion rotation"))
 	void SetVariableQuat(FName InVariableName, const FQuat& InValue);
 
-	/** Sets a Niagara Vector3 parameter by name, overriding locally if necessary.*/
-	UFUNCTION(BlueprintCallable, Category = Niagara, meta = (DisplayName = "Set Niagara Variable By String (Matrix)"))
+	/** Sets a Niagara matrix parameter by name, overriding locally if necessary.*/
+	UFUNCTION(BlueprintCallable, Category = Niagara, meta = (DisplayName = "Set Niagara Variable By String (Matrix)", Keywords="user parameter variable matrix"))
 	void SetNiagaraVariableMatrix(const FString& InVariableName, const FMatrix& InValue);
 
-	/** Sets a Niagara Vector3 parameter by name, overriding locally if necessary.*/
-	UFUNCTION(BlueprintCallable, Category = Niagara, meta = (DisplayName = "Set Niagara Variable (Matrix)"))
+	/** Sets a Niagara matrix parameter by name, overriding locally if necessary.*/
+	UFUNCTION(BlueprintCallable, Category = Niagara, meta = (DisplayName = "Set Niagara Variable (Matrix)", Keywords="user parameter variable matrix"))
 	void SetVariableMatrix(FName InVariableName, const FMatrix& InValue);
 
 	/** Sets a Niagara Vector3 parameter by name, overriding locally if necessary.*/
-	UFUNCTION(BlueprintCallable, Category = Niagara, meta = (DisplayName = "Set Niagara Variable By String (Vector3)"))
+	UFUNCTION(BlueprintCallable, Category = Niagara, meta = (DisplayName = "Set Niagara Variable By String (Vector3)", Keywords="user parameter variable vector"))
 	void SetNiagaraVariableVec3(const FString& InVariableName, FVector InValue);
 
 	/** Sets a Niagara Vector3 parameter by name, overriding locally if necessary.*/
-	UFUNCTION(BlueprintCallable, Category = Niagara, meta = (DisplayName = "Set Niagara Variable (Vector3)"))
+	UFUNCTION(BlueprintCallable, Category = Niagara, meta = (DisplayName = "Set Niagara Variable (Vector3)", Keywords="user parameter variable vector"))
 	void SetVariableVec3(FName InVariableName, FVector InValue);
 
 	/** Sets a Niagara Position parameter by name, overriding locally if necessary.*/
-	UFUNCTION(BlueprintCallable, Category = Niagara, meta = (DisplayName = "Set Niagara Variable By String (Position)"))
+	UFUNCTION(BlueprintCallable, Category = Niagara, meta = (DisplayName = "Set Niagara Variable By String (Position)", Keywords="user parameter variable vector position lwc"))
 	void SetNiagaraVariablePosition(const FString& InVariableName, FVector InValue);
 
 	/** Sets a Niagara Position parameter by name, overriding locally if necessary.*/
-	UFUNCTION(BlueprintCallable, Category = Niagara, meta = (DisplayName = "Set Niagara Variable (Position)"))
+	UFUNCTION(BlueprintCallable, Category = Niagara, meta = (DisplayName = "Set Niagara Variable (Position)", Keywords="user parameter variable vector position lwc"))
 	void SetVariablePosition(FName InVariableName, FVector InValue);
 
-	/** Sets a Niagara Vector3 parameter by name, overriding locally if necessary.*/
-	UFUNCTION(BlueprintCallable, Category = Niagara, meta = (DisplayName = "Set Niagara Variable By String (Vector2)"))
+	/** Sets a Niagara Vector2 parameter by name, overriding locally if necessary.*/
+	UFUNCTION(BlueprintCallable, Category = Niagara, meta = (DisplayName = "Set Niagara Variable By String (Vector2)", Keywords="user parameter variable vector"))
 	void SetNiagaraVariableVec2(const FString& InVariableName, FVector2D InValue);
 
-	/** Sets a Niagara Vector3 parameter by name, overriding locally if necessary.*/
-	UFUNCTION(BlueprintCallable, Category = Niagara, meta = (DisplayName = "Set Niagara Variable (Vector2)"))
+	/** Sets a Niagara Vector2 parameter by name, overriding locally if necessary.*/
+	UFUNCTION(BlueprintCallable, Category = Niagara, meta = (DisplayName = "Set Niagara Variable (Vector2)", Keywords="user parameter variable vector"))
 	void SetVariableVec2(FName InVariableName, FVector2D InValue);
 
 	/** Sets a Niagara float parameter by name, overriding locally if necessary.*/
-	UFUNCTION(BlueprintCallable, Category = Niagara, meta = (DisplayName = "Set Niagara Variable By String (Float)"))
+	UFUNCTION(BlueprintCallable, Category = Niagara, meta = (DisplayName = "Set Niagara Variable By String (Float)", Keywords="user parameter variable float"))
 	void SetNiagaraVariableFloat(const FString& InVariableName, float InValue);
 
 	/** Sets a Niagara float parameter by name, overriding locally if necessary.*/
-	UFUNCTION(BlueprintCallable, Category = Niagara, meta = (DisplayName = "Set Niagara Variable (Float)"))
+	UFUNCTION(BlueprintCallable, Category = Niagara, meta = (DisplayName = "Set Niagara Variable (Float)", Keywords="user parameter variable float"))
 	void SetVariableFloat(FName InVariableName, float InValue);
 
 	/** Sets a Niagara int parameter by name, overriding locally if necessary.*/
-	UFUNCTION(BlueprintCallable, Category = Niagara, meta = (DisplayName = "Set Niagara Variable By String (Int32)"))
+	UFUNCTION(BlueprintCallable, Category = Niagara, meta = (DisplayName = "Set Niagara Variable By String (Int32)", Keywords="user parameter variable int"))
 	void SetNiagaraVariableInt(const FString& InVariableName, int32 InValue);
 
 	/** Sets a Niagara int parameter by name, overriding locally if necessary.*/
-	UFUNCTION(BlueprintCallable, Category = Niagara, meta = (DisplayName = "Set Niagara Variable (Int32)"))
+	UFUNCTION(BlueprintCallable, Category = Niagara, meta = (DisplayName = "Set Niagara Variable (Int32)", Keywords="user parameter variable int"))
 	void SetVariableInt(FName InVariableName, int32 InValue);
 
-	/** Sets a Niagara float parameter by name, overriding locally if necessary.*/
-	UFUNCTION(BlueprintCallable, Category = Niagara, meta = (DisplayName = "Set Niagara Variable By String (Bool)"))
+	/** Sets a Niagara bool parameter by name, overriding locally if necessary.*/
+	UFUNCTION(BlueprintCallable, Category = Niagara, meta = (DisplayName = "Set Niagara Variable By String (Bool)", Keywords="user parameter variable bool"))
 	void SetNiagaraVariableBool(const FString& InVariableName, bool InValue);
 
-	/** Sets a Niagara float parameter by name, overriding locally if necessary.*/
-	UFUNCTION(BlueprintCallable, Category = Niagara, meta = (DisplayName = "Set Niagara Variable (Bool)"))
+	/** Sets a Niagara bool parameter by name, overriding locally if necessary.*/
+	UFUNCTION(BlueprintCallable, Category = Niagara, meta = (DisplayName = "Set Niagara Variable (Bool)", Keywords="user parameter variable bool"))
 	void SetVariableBool(FName InVariableName, bool InValue);
 
-	UFUNCTION(BlueprintCallable, Category = Niagara, meta = (DisplayName = "Set Niagara Variable By String (Actor)"))
+	UFUNCTION(BlueprintCallable, Category = Niagara, meta = (DisplayName = "Set Niagara Variable By String (Actor)", Keywords="user parameter variable actor"))
 	void SetNiagaraVariableActor(const FString& InVariableName, AActor* Actor);
 
-	UFUNCTION(BlueprintCallable, Category = Niagara, meta = (DisplayName = "Set Niagara Variable (Actor)"))
+	UFUNCTION(BlueprintCallable, Category = Niagara, meta = (DisplayName = "Set Niagara Variable (Actor)", Keywords="user parameter variable actor"))
 	void SetVariableActor(FName InVariableName, AActor* Actor);
 
-	UFUNCTION(BlueprintCallable, Category = Niagara, meta = (DisplayName = "Set Niagara Variable By String (Object)"))
+	UFUNCTION(BlueprintCallable, Category = Niagara, meta = (DisplayName = "Set Niagara Variable By String (Object)", Keywords="user parameter variable object"))
 	void SetNiagaraVariableObject(const FString& InVariableName, UObject* Object);
 
-	UFUNCTION(BlueprintCallable, Category = Niagara, meta = (DisplayName = "Set Niagara Variable (Object)"))
+	UFUNCTION(BlueprintCallable, Category = Niagara, meta = (DisplayName = "Set Niagara Variable (Object)", Keywords="user parameter variable object"))
 	void SetVariableObject(FName InVariableName, UObject* Object);
 
-	UFUNCTION(BlueprintCallable, Category = Niagara, meta = (DisplayName = "Set Niagara Variable (Material)"))
+	UFUNCTION(BlueprintCallable, Category = Niagara, meta = (DisplayName = "Set Niagara Variable (Material)", Keywords="user parameter variable material"))
 	void SetVariableMaterial(FName InVariableName, UMaterialInterface* Object);
 
-	UFUNCTION(BlueprintCallable, Category = Niagara, meta = (DisplayName = "Set Niagara Variable (Static Mesh)"))
+	UFUNCTION(BlueprintCallable, Category = Niagara, meta = (DisplayName = "Set Niagara Variable (Static Mesh)", Keywords="user parameter variable mesh"))
 	void SetVariableStaticMesh(FName InVariableName, UStaticMesh* InValue);
 
-	UFUNCTION(BlueprintCallable, Category = Niagara, meta = (DisplayName = "Set Niagara Variable (Texture)"))
+	UFUNCTION(BlueprintCallable, Category = Niagara, meta = (DisplayName = "Set Niagara Variable (Texture)", Keywords="user parameter variable texture"))
 	void SetVariableTexture(FName InVariableName, class UTexture* Texture);
 
 	UFUNCTION(BlueprintCallable, Category = Niagara, meta = (DisplayName = "Set Niagara Variable (TextureRenderTarget)"))
