Fix for unity build (also removing some debug code) for NiagaraEmitter.cpp
#rnx

diff --git a/Engine/Plugins/FX/Niagara/Source/Niagara/Private/NiagaraEmitter.cpp b/Engine/Plugins/FX/Niagara/Source/Niagara/Private/NiagaraEmitter.cpp
--- a/Engine/Plugins/FX/Niagara/Source/Niagara/Private/NiagaraEmitter.cpp
+++ b/Engine/Plugins/FX/Niagara/Source/Niagara/Private/NiagaraEmitter.cpp
@@ -19,6 +19,7 @@
 #include "NiagaraTrace.h"
 #include "Interfaces/ITargetPlatform.h"
 #include "Modules/ModuleManager.h"
+#include "UObject/LinkerLoad.h"
 #include "UObject/ObjectSaveContext.h"
 #include "UObject/Package.h"
 #include "UObject/UObjectIterator.h"
@@ -776,40 +777,6 @@ void FVersionedNiagaraEmitterData::PostLoad(UNiagaraEmitter& Emitter, bool bIsCo
 			GraphSource->ConditionalPostLoad();
 			GraphSource->PostLoadFromEmitter(FVersionedNiagaraEmitter(&Emitter, Version.VersionGuid));
 		}
-		// more debug code to track what's going wrong with GraphSource being null
-#if !(UE_BUILD_SHIPPING || UE_BUILD_TEST)
-		else
-		{
-			ensureMsgf(!Emitter.HasAnyFlags(RF_NeedLoad), TEXT("Emitter (%s) should have been serialized by now!"), *Emitter.GetPathName());
-
-			FUObjectThreadContext& ThreadContext = FUObjectThreadContext::Get();
-
-			FString ErrorMsg;
-
-			const int32 PostLoadContextCount = ThreadContext.DebugPostLoad.Num();
-
-			ErrorMsg.Appendf(TEXT("ObjFlags along outer: "));
-			UE_LOG(LogNiagara, Warning, TEXT("ObjFlags along outer"));
-			UObject* ThisObject = &Emitter;
-			while (ThisObject != nullptr)
-			{
-				UE_LOG(LogNiagara, Warning, TEXT("\t%s = %x"), *ThisObject->GetName(), ThisObject->GetFlags());
-				ErrorMsg.Appendf(TEXT("(%s = %x), "), *ThisObject->GetName(), ThisObject->GetFlags());
-				ThisObject = ThisObject->GetOuter();
-			}
-
-			ErrorMsg.Appendf(TEXT("Stack of objects (%d) in PostLoad chain: "), PostLoadContextCount);
-
-			for (int32 PostLoadIt = PostLoadContextCount - 1; PostLoadIt >= 0; --PostLoadIt)
-			{
-				ErrorMsg.Appendf(TEXT("{%s}, "), ThreadContext.DebugPostLoad[PostLoadIt]
-					? *ThreadContext.DebugPostLoad[PostLoadIt]->GetPathName()
-					: TEXT("<unknown>"));
-			}
-
-			checkf(false, *ErrorMsg);
-		}
-#endif
 
 		// Prepare for emitter inheritance.
 		if (GetParent().Emitter != nullptr)
