Added user parameter setting option for attenuation and sound concurrency to the audio player DI

diff --git a/Engine/Plugins/FX/Niagara/Source/Niagara/Classes/NiagaraDataInterfaceAudioPlayer.h b/Engine/Plugins/FX/Niagara/Source/Niagara/Classes/NiagaraDataInterfaceAudioPlayer.h
--- a/Engine/Plugins/FX/Niagara/Source/Niagara/Classes/NiagaraDataInterfaceAudioPlayer.h
+++ b/Engine/Plugins/FX/Niagara/Source/Niagara/Classes/NiagaraDataInterfaceAudioPlayer.h
@@ -4,6 +4,7 @@
 #include "NiagaraCommon.h"
 #include "NiagaraShared.h"
 #include "NiagaraDataInterface.h"
+#include "NiagaraParameterStore.h"
 #include "Sound/SoundAttenuation.h"
 
 #include "NiagaraDataInterfaceAudioPlayer.generated.h"
@@ -30,6 +31,27 @@ struct FPersistentAudioParticleData
 	TFunction<void(struct FAudioPlayerInterface_InstanceData*,UAudioComponent*,FNiagaraSystemInstance*)> UpdateCallback;
 };
 
+UCLASS(EditInlineNew, Category = "Audio", meta = (DisplayName = "Niagara Audio Player Settings"), Blueprintable, BlueprintType)
+class NIAGARA_API UNiagaraDataInterfaceAudioPlayerSettings : public UObject
+{
+public:
+	GENERATED_BODY()
+
+	UNiagaraDataInterfaceAudioPlayerSettings();
+
+	UPROPERTY(EditAnywhere, BlueprintReadWrite, Category = "Audio", meta=(InlineEditConditionToggle))
+	bool bOverrideConcurrency = true;
+
+	UPROPERTY(EditAnywhere, BlueprintReadWrite, Category = "Audio", meta=(EditCondition="bOverrideConcurrency"))
+	TObjectPtr<USoundConcurrency> Concurrency;
+
+	UPROPERTY(EditAnywhere, BlueprintReadWrite, Category = "Audio")
+	bool bOverrideAttenuationSettings = true;
+
+	UPROPERTY(EditAnywhere, BlueprintReadWrite, Category = "Audio", meta=(EditCondition="bOverrideAttenuationSettings"))
+	FSoundAttenuationSettings AttenuationSettings;
+};
+
 struct FAudioPlayerInterface_InstanceData
 {
 	/** We use a lock-free queue here because multiple threads might try to push data to it at the same time. */
@@ -42,6 +64,7 @@ struct FAudioPlayerInterface_InstanceData
 	TWeakObjectPtr<USoundBase> SoundToPlay;
 	TWeakObjectPtr<USoundAttenuation> Attenuation;
 	TWeakObjectPtr<USoundConcurrency> Concurrency;
+	TWeakObjectPtr<UNiagaraDataInterfaceAudioPlayerSettings> CachedUserParam;
 	TArray<FName> ParameterNames;
 
 	FNiagaraLWCConverter LWCConverter;
@@ -51,6 +74,9 @@ struct FAudioPlayerInterface_InstanceData
 #if WITH_EDITORONLY_DATA
 	bool bOnlyActiveDuringGameplay = false;
 #endif
+
+	/** A binding to the user ptr we're reading from (if we are). */
+	FNiagaraParameterDirectBinding<UObject*> UserParamBinding;
 };
 
 /** This Data Interface can be used to play one-shot audio effects driven by particle data. */
@@ -76,6 +102,11 @@ public:
 	UPROPERTY(EditAnywhere, Category = "Parameters")
 	TArray<FName> ParameterNames;
 
+	/** If bound to a valid user parameter object of type UNiagaraDataInterfaceAudioPlayerSettings, then configured settings like sound attenuation are set via the user parameter. This allows the sound settings to be dynamically changed via blueprint or C++.
+	 *  Only used by persistent audio, one-shot audio ignores this option. */
+	UPROPERTY(EditAnywhere, Category = "Parameters")
+	FNiagaraUserParameterBinding ConfigurationUserParameter;
+
 	UPROPERTY(EditAnywhere, AdvancedDisplay, Category = "Audio", meta = (InlineEditConditionToggle))
     bool bLimitPlaysPerTick;
 
diff --git a/Engine/Plugins/FX/Niagara/Source/Niagara/Private/NiagaraDataInterfaceAudioPlayer.cpp b/Engine/Plugins/FX/Niagara/Source/Niagara/Private/NiagaraDataInterfaceAudioPlayer.cpp
--- a/Engine/Plugins/FX/Niagara/Source/Niagara/Private/NiagaraDataInterfaceAudioPlayer.cpp
+++ b/Engine/Plugins/FX/Niagara/Source/Niagara/Private/NiagaraDataInterfaceAudioPlayer.cpp
@@ -41,55 +41,6 @@ struct FNiagaraAudioPlayerDIFunctionVersion
 	};
 };
 
-/**
-Async task to play the audio on the game thread and isolate from the niagara tick
-*/
-class FNiagaraAudioPlayerAsyncTask
-{
-	TWeakObjectPtr<USoundBase> WeakSound;
-	TWeakObjectPtr<USoundAttenuation> WeakAttenuation;
-	TWeakObjectPtr<USoundConcurrency> WeakConcurrency;
-	TArray<FAudioParticleData> Data;
-	TWeakObjectPtr<UWorld> WeakWorld;
-
-public:
-	FNiagaraAudioPlayerAsyncTask(TWeakObjectPtr<USoundBase> InSound, TWeakObjectPtr<USoundAttenuation> InAttenuation, TWeakObjectPtr<USoundConcurrency> InConcurrency, TArray<FAudioParticleData>& Data, TWeakObjectPtr<UWorld> InWorld)
-		: WeakSound(InSound)
-	    , WeakAttenuation(InAttenuation)
-		, WeakConcurrency(InConcurrency)
-		, Data(Data)
-		, WeakWorld(InWorld)
-	{
-	}
-
-	FORCEINLINE TStatId GetStatId() const { RETURN_QUICK_DECLARE_CYCLE_STAT(FNiagaraAudioPlayerAsyncTask, STATGROUP_TaskGraphTasks); }
-	static FORCEINLINE ENamedThreads::Type GetDesiredThread() { return ENamedThreads::GameThread; }
-	static FORCEINLINE ESubsequentsMode::Type GetSubsequentsMode() { return ESubsequentsMode::FireAndForget; }
-
-	void DoTask(ENamedThreads::Type CurrentThread, const FGraphEventRef& MyCompletionGraphEvent)
-	{
-		UWorld* World = WeakWorld.Get();
-		if (World == nullptr)
-		{
-			UE_LOG(LogNiagara, Warning, TEXT("Invalid world reference in audio player DI, skipping play"));
-			return;
-		}
-		
-		USoundBase* Sound = WeakSound.Get();
-		if (Sound == nullptr)
-		{
-			UE_LOG(LogNiagara, Warning, TEXT("Invalid sound reference in audio player DI, skipping play"));
-			return;
-		}
-
-		for (const FAudioParticleData& ParticleData : Data)
-		{
-			UGameplayStatics::PlaySoundAtLocation(World, Sound, ParticleData.Position, ParticleData.Rotation, ParticleData.Volume,
-				ParticleData.Pitch, ParticleData.StartTime, WeakAttenuation.Get(), WeakConcurrency.Get());
-		}
-	}
-};
-
 UNiagaraDataInterfaceAudioPlayer::UNiagaraDataInterfaceAudioPlayer(FObjectInitializer const& ObjectInitializer)
 	: Super(ObjectInitializer)
 {
@@ -98,9 +49,17 @@ UNiagaraDataInterfaceAudioPlayer::UNiagaraDataInterfaceAudioPlayer(FObjectInitia
 	Concurrency = nullptr;
 	bLimitPlaysPerTick = true;
 	MaxPlaysPerTick = 10;
+
+	FNiagaraTypeDefinition Def(UObject::StaticClass());
+	ConfigurationUserParameter.Parameter.SetType(Def);
+}
+
+UNiagaraDataInterfaceAudioPlayerSettings::UNiagaraDataInterfaceAudioPlayerSettings()
+{
 }
 
 #if WITH_EDITORONLY_DATA
+
 bool UNiagaraDataInterfaceAudioPlayer::UpgradeFunctionCall(FNiagaraFunctionSignature& FunctionSignature)
 {
 	// always upgrade to the latest version
@@ -120,6 +79,7 @@ bool UNiagaraDataInterfaceAudioPlayer::UpgradeFunctionCall(FNiagaraFunctionSigna
 
 	return false;
 }
+
 #endif
 
 void UNiagaraDataInterfaceAudioPlayer::PostInitProperties()
@@ -163,11 +123,16 @@ void UNiagaraDataInterfaceAudioPlayer::DestroyPerInstanceData(void* PerInstanceD
 
 bool UNiagaraDataInterfaceAudioPlayer::PerInstanceTick(void* PerInstanceData, FNiagaraSystemInstance* SystemInstance, float DeltaSeconds)
 {
-	FAudioPlayerInterface_InstanceData* PIData = (FAudioPlayerInterface_InstanceData*)PerInstanceData;
+	FAudioPlayerInterface_InstanceData* PIData = static_cast<FAudioPlayerInterface_InstanceData*>(PerInstanceData);
 	if (!PIData)
 	{
 		return true;
 	}
+	if (ConfigurationUserParameter.Parameter.IsValid() && PIData->UserParamBinding.GetValue() != PIData->CachedUserParam)
+	{
+		// Reset if the user object ptr has been changed
+		return true;
+	}
 	
 	if (IsValid(SoundToPlay) && SystemInstance)
 	{
@@ -175,12 +140,20 @@ bool UNiagaraDataInterfaceAudioPlayer::PerInstanceTick(void* PerInstanceData, FN
 		PIData->Attenuation = Attenuation;
 		PIData->Concurrency = Concurrency;
 		PIData->bValidOneShotSound = SoundToPlay->IsLooping() ? bAllowLoopingOneShotSounds : true;
+
+		if (ConfigurationUserParameter.Parameter.IsValid())
+		{
+			// Initialize the binding and retrieve the object.
+			UObject* UserParamObject = PIData->UserParamBinding.Init(SystemInstance->GetInstanceParameters(), ConfigurationUserParameter.Parameter);
+			PIData->CachedUserParam = Cast<UNiagaraDataInterfaceAudioPlayerSettings>(UserParamObject);
+		}
 	}
 	else
 	{
 		PIData->SoundToPlay.Reset();
 		PIData->Attenuation.Reset();
 		PIData->Concurrency.Reset();
+		PIData->CachedUserParam.Reset();
 		PIData->bValidOneShotSound = false;
 	}
 
@@ -219,8 +192,16 @@ bool UNiagaraDataInterfaceAudioPlayer::PerInstanceTickPostSimulate(void* PerInst
 				break;
 			}
 		}
-		
-		TGraphTask<FNiagaraAudioPlayerAsyncTask>::CreateTask().ConstructAndDispatchWhenReady(PIData->SoundToPlay, PIData->Attenuation, PIData->Concurrency, Data, World);
+
+		// play the fire-and-forget sounds
+		if (World && PIData->SoundToPlay.IsValid())
+		{
+			for (const FAudioParticleData& ParticleData : Data)
+			{
+				UGameplayStatics::PlaySoundAtLocation(World, PIData->SoundToPlay.Get(), ParticleData.Position, ParticleData.Rotation, ParticleData.Volume,
+					ParticleData.Pitch, ParticleData.StartTime, PIData->Attenuation.Get(), PIData->Concurrency.Get());
+			}
+		}
 	}
 
 	// process the persistent audio updates
@@ -273,7 +254,7 @@ bool UNiagaraDataInterfaceAudioPlayer::Equals(const UNiagaraDataInterface* Other
 	}
 
 	const UNiagaraDataInterfaceAudioPlayer* OtherPlayer = CastChecked<UNiagaraDataInterfaceAudioPlayer>(Other);
-	return OtherPlayer->SoundToPlay == SoundToPlay && OtherPlayer->Attenuation == Attenuation && OtherPlayer->Concurrency == Concurrency && OtherPlayer->bLimitPlaysPerTick == bLimitPlaysPerTick && OtherPlayer->MaxPlaysPerTick == MaxPlaysPerTick;
+	return OtherPlayer->SoundToPlay == SoundToPlay && OtherPlayer->Attenuation == Attenuation && OtherPlayer->Concurrency == Concurrency && OtherPlayer->ConfigurationUserParameter == ConfigurationUserParameter && OtherPlayer->bLimitPlaysPerTick == bLimitPlaysPerTick && OtherPlayer->MaxPlaysPerTick == MaxPlaysPerTick;
 }
 
 void UNiagaraDataInterfaceAudioPlayer::GetFunctions(TArray<FNiagaraFunctionSignature>& OutFunctions)
@@ -842,12 +823,24 @@ void UNiagaraDataInterfaceAudioPlayer::PlayPersistentAudio(FVectorVMExternalFunc
 					if (NiagaraComponent && Sound.IsValid())
 					{
 						bool bStopWithEffect = InstanceData->bStopWhenComponentIsDestroyed || Sound->IsLooping(); // we don't allow looping effects to outlive us because then they keep playing forever
-						UAudioComponent* AudioComponent = UGameplayStatics::SpawnSoundAttached(Sound.Get(), NiagaraComponent, NAME_None, Position, Rotation, EAttachLocation::KeepWorldPosition, bStopWithEffect, Volume, Pitch, StartTime, InstanceData->Attenuation.Get(), InstanceData->Concurrency.Get(), true);
+						USoundConcurrency* SoundConcurrency = InstanceData->Concurrency.Get();
+						if (InstanceData->CachedUserParam.IsValid() && InstanceData->CachedUserParam->bOverrideConcurrency)
+						{
+							SoundConcurrency = InstanceData->CachedUserParam->Concurrency;
+						}
+
+						USoundAttenuation* AttenuationSettings = InstanceData->Attenuation.Get();
+						UAudioComponent* AudioComponent = UGameplayStatics::SpawnSoundAttached(Sound.Get(), NiagaraComponent, NAME_None, Position, Rotation, EAttachLocation::KeepWorldPosition, bStopWithEffect, Volume, Pitch, StartTime, AttenuationSettings, SoundConcurrency, true);
 						if (AudioComponent == nullptr)
 						{
 							// looks like audio is disabled, so we'll skip adding a mapping
 							return;
 						}
+						if (InstanceData->CachedUserParam.IsValid() && InstanceData->CachedUserParam->bOverrideAttenuationSettings)
+						{
+							AudioComponent->AdjustAttenuation(InstanceData->CachedUserParam->AttenuationSettings);
+							AudioComponent->Play(StartTime);
+						}
 						if (FadeIn > 0.0)
 						{
 							AudioComponent->FadeIn(FadeIn, Volume, StartTime);
@@ -903,6 +896,7 @@ bool UNiagaraDataInterfaceAudioPlayer::CopyToInternal(UNiagaraDataInterface* Des
 	OtherTyped->SoundToPlay = SoundToPlay;
 	OtherTyped->Attenuation = Attenuation;
 	OtherTyped->Concurrency = Concurrency;
+	OtherTyped->ConfigurationUserParameter = ConfigurationUserParameter;
 	OtherTyped->bLimitPlaysPerTick = bLimitPlaysPerTick;
 	OtherTyped->MaxPlaysPerTick = MaxPlaysPerTick;
 	OtherTyped->ParameterNames = ParameterNames;
