Fixed enums appearing as ints in the Niagara user parameters tab & component details tab

diff --git a/Engine/Plugins/FX/Niagara/Source/NiagaraEditor/Private/Customizations/NiagaraComponentDetails.cpp b/Engine/Plugins/FX/Niagara/Source/NiagaraEditor/Private/Customizations/NiagaraComponentDetails.cpp
--- a/Engine/Plugins/FX/Niagara/Source/NiagaraEditor/Private/Customizations/NiagaraComponentDetails.cpp
+++ b/Engine/Plugins/FX/Niagara/Source/NiagaraEditor/Private/Customizations/NiagaraComponentDetails.cpp
@@ -437,6 +437,14 @@ void FNiagaraUserParameterNodeBuilder::GenerateRowForUserParameter(IDetailChildr
 			static UScriptStruct* VectorStruct = FindObjectChecked<UScriptStruct>(CoreUObjectPkg, TEXT("Vector"));
 			Struct = VectorStruct;
 		}
+
+		// the details panel uses byte properties with enums assigned to display enums
+		if (Type.IsEnum())
+		{
+			static UPackage* NiagaraEditorPkg = FindObjectChecked<UPackage>(nullptr, TEXT("/Script/NiagaraEditor"));
+			static UScriptStruct* ByteStruct = FindObjectChecked<UScriptStruct>(NiagaraEditorPkg, TEXT("NiagaraEnumToByteHelper"));
+			Struct = ByteStruct;
+		}
 			
 		TSharedPtr<FStructOnScope> StructOnScope = MakeShareable(new FStructOnScope(Struct, ParameterProxy->Value().GetBytes()));
 
@@ -444,6 +452,22 @@ void FNiagaraUserParameterNodeBuilder::GenerateRowForUserParameter(IDetailChildr
 			.UniqueId(ChoppedUserParameter.GetName());
 
 		Row = ChildrenBuilder.AddExternalStructureProperty(StructOnScope.ToSharedRef(), NAME_None, Params);
+
+		// we set the enum of the contained byte property to the enum carried with the type
+		if(Type.IsEnum())
+		{
+			FProperty* Property = Row->GetPropertyHandle()->GetProperty();
+			TSharedPtr<IPropertyHandle> ValuePropertyHandle = Row->GetPropertyHandle()->GetChildHandle(0);
+			FProperty* ValueProperty = ValuePropertyHandle->GetProperty();
+			if(ValueProperty->IsA<FByteProperty>())
+			{
+				if(FByteProperty* ByteProperty = CastField<FByteProperty>(ValueProperty))
+				{
+					ByteProperty->Enum = Type.GetEnum();
+				}
+			}
+		}
+	
 		ParameterNameToDisplayStruct.Add(UserParameter.GetName(), TWeakPtr<FStructOnScope>(StructOnScope));
 	}
 
diff --git a/Engine/Plugins/FX/Niagara/Source/NiagaraEditor/Private/Customizations/NiagaraComponentDetails.h b/Engine/Plugins/FX/Niagara/Source/NiagaraEditor/Private/Customizations/NiagaraComponentDetails.h
--- a/Engine/Plugins/FX/Niagara/Source/NiagaraEditor/Private/Customizations/NiagaraComponentDetails.h
+++ b/Engine/Plugins/FX/Niagara/Source/NiagaraEditor/Private/Customizations/NiagaraComponentDetails.h
@@ -10,6 +10,7 @@
 #include "ViewModels/TNiagaraViewModelManager.h"
 #include "ViewModels/HierarchyEditor/NiagaraHierarchyViewModelBase.h"
 #include "ViewModels/HierarchyEditor/NiagaraUserParametersHierarchyViewModel.h"
+#include "NiagaraComponentDetails.generated.h"
 
 class IDetailLayoutBuilder;
 class IPropertyHandle;
@@ -19,6 +20,15 @@ class SNiagaraParameterEditor;
 class FNiagaraParameterViewModelCustomDetails;
 class UNiagaraSystem;
 
+USTRUCT()
+struct FNiagaraEnumToByteHelper
+{
+	GENERATED_USTRUCT_BODY()
+
+	UPROPERTY(EditAnywhere, Category = Parameters)
+	uint8 Value = 0;
+};
+
 class FNiagaraComponentDetails : public IDetailCustomization
 {
 public:
diff --git a/Engine/Plugins/FX/Niagara/Source/NiagaraEditor/Private/NiagaraEditorModule.cpp b/Engine/Plugins/FX/Niagara/Source/NiagaraEditor/Private/NiagaraEditorModule.cpp
--- a/Engine/Plugins/FX/Niagara/Source/NiagaraEditor/Private/NiagaraEditorModule.cpp
+++ b/Engine/Plugins/FX/Niagara/Source/NiagaraEditor/Private/NiagaraEditorModule.cpp
@@ -932,6 +932,10 @@ void FNiagaraEditorModule::StartupModule()
 		FNiagaraParameterMap::StaticStruct()->GetFName(),
 		FOnGetPropertyTypeCustomizationInstance::CreateStatic(&FNiagaraNumericCustomization::MakeInstance));
 
+	PropertyModule.RegisterCustomPropertyTypeLayout(
+		FNiagaraEnumToByteHelper::StaticStruct()->GetFName(),
+		FOnGetPropertyTypeCustomizationInstance::CreateStatic(&FNiagaraNumericCustomization::MakeInstance));
+
 	PropertyModule.RegisterCustomPropertyTypeLayout(
 		FNiagaraBool::StaticStruct()->GetFName(),
 		FOnGetPropertyTypeCustomizationInstance::CreateStatic(&FNiagaraBoolCustomization::MakeInstance));
