Fixed static analysis warning for possible nullptr access

diff --git a/Engine/Plugins/FX/Niagara/Source/Niagara/Private/NiagaraDataInterfaceAudioPlayer.cpp b/Engine/Plugins/FX/Niagara/Source/Niagara/Private/NiagaraDataInterfaceAudioPlayer.cpp
--- a/Engine/Plugins/FX/Niagara/Source/Niagara/Private/NiagaraDataInterfaceAudioPlayer.cpp
+++ b/Engine/Plugins/FX/Niagara/Source/Niagara/Private/NiagaraDataInterfaceAudioPlayer.cpp
@@ -169,7 +169,7 @@ bool UNiagaraDataInterfaceAudioPlayer::PerInstanceTickPostSimulate(void* PerInst
 	UWorld* World = SystemInstance->GetWorldManager()->GetWorld();
 
 #if WITH_EDITORONLY_DATA
-	if (World->HasBegunPlay() == false && PIData->bOnlyActiveDuringGameplay)
+	if (World == nullptr || (World->HasBegunPlay() == false && PIData->bOnlyActiveDuringGameplay))
 	{
 		PIData->PlayAudioQueue.Empty();
 		PIData->PersistentAudioMapping.Empty();
