Fixing issue with changing NiagaraQuality levels at runtime.

- Prevented premature update of cached ql so that we do not skip proper scalability updates under some calling conditions.
- Tidied up and clarified refreshing code.
- Decoupled the override ql setting from directly setting the CVar value.

diff --git a/Engine/Plugins/FX/Niagara/Source/Niagara/Classes/NiagaraPlatformSet.h b/Engine/Plugins/FX/Niagara/Source/Niagara/Classes/NiagaraPlatformSet.h
--- a/Engine/Plugins/FX/Niagara/Source/Niagara/Classes/NiagaraPlatformSet.h
+++ b/Engine/Plugins/FX/Niagara/Source/Niagara/Classes/NiagaraPlatformSet.h
@@ -262,7 +262,6 @@ private:
 	mutable uint32 LastCachedFrame = INDEX_NONE;
 };
 
-void NIAGARA_API SetGNiagaraQualityLevel(int32 QualityLevel);
 void NIAGARA_API SetGNiagaraDeviceProfile(UDeviceProfile* Profile);
 
 UENUM()
@@ -394,7 +393,6 @@ public:
 	/** Invalidates any cached data on this platform set when something has changed. */
 	void OnChanged();
 
-
 	/** Inspects the passed sets and generates an array of all conflicts between these sets. Used to keep arrays of platform sets orthogonal. */
 	static bool GatherConflicts(const TArray<const FNiagaraPlatformSet*>& PlatformSets, TArray<FNiagaraPlatformSetConflictInfo>& OutConflicts);
 	
@@ -438,6 +436,9 @@ public:
 
 	static uint32 GetLastDirtiedFrame(){ return LastDirtiedFrame; }
 
+	static void SetNiagaraQualityLevelOverride(int32 Override);
+	static void ClearNiagaraQualityLevelOverride();
+
 private:
 
 	mutable uint32 bEnabledForCurrentProfileAndEffectQuality : 1;
@@ -448,9 +449,13 @@ private:
 	//Set from outside when we need to force all cached values to be regenerated. For example on CVar changes.
 	static uint32 LastDirtiedFrame;
 
+	/** Cached value of the CVar fx.Niagara.QualityLevel. Allows us to detect when it has changed so we can properly apply scalability changes. */
 	static int32 CachedQualityLevel;
 	static int32 CachedAvailableQualityLevelMask;
 
+	/** An override for Niagara quality level, driven from the Niagara UI. */
+	static int32 QualityLevelOverride;
+
 #if WITH_EDITOR
 	//Cached data read from platform ini files.
 	static TMap<FName, FPlatformIniSettings> CachedPlatformIniSettings;
@@ -522,5 +527,3 @@ FORCEINLINE void FNiagaraDeviceProfileStateEntry::SetState(int32 QualityLevel, E
 		QualityLevelMask &= ~QLMask;
 	}
 }
-
-extern int32 GNiagaraQualityLevel;
\ No newline at end of file
diff --git a/Engine/Plugins/FX/Niagara/Source/Niagara/Private/NiagaraPlatformSet.cpp b/Engine/Plugins/FX/Niagara/Source/Niagara/Private/NiagaraPlatformSet.cpp
--- a/Engine/Plugins/FX/Niagara/Source/Niagara/Private/NiagaraPlatformSet.cpp
+++ b/Engine/Plugins/FX/Niagara/Source/Niagara/Private/NiagaraPlatformSet.cpp
@@ -55,14 +55,14 @@ int32 GNiagaraQualityLevel = DefaultQualityLevel;
 
 static FAutoConsoleCommand GCmdSetNiagaraQualityLevelOverride(
 	TEXT("fx.Niagara.SetOverrideQualityLevel"),
-	TEXT("Sets which quality level we should override with, no args means reset to default (Epic). Valid levels are 0-4 (Low-Cinematic)"),
+	TEXT("Sets which quality level we should override with, no args means clear the override and return to non overriden quality level). Valid levels are 0-4 (Low-Cinematic)"),
 	FConsoleCommandWithArgsDelegate::CreateLambda(
 		[](const TArray<FString>& Args)
 		{
 			if(Args.Num() == 0)
 			{
-				SetGNiagaraQualityLevel(DefaultQualityLevel);
-				UE_LOG(LogNiagara, Display, TEXT("Reset the global Niagara quality level to %s"), *FNiagaraPlatformSet::GetQualityLevelText(GNiagaraQualityLevel).ToString());
+				FNiagaraPlatformSet::ClearNiagaraQualityLevelOverride();
+				UE_LOG(LogNiagara, Display, TEXT("Clearing the Niagara Quality Level override and returning quality level to %s"), *FNiagaraPlatformSet::GetQualityLevelText(FNiagaraPlatformSet::GetQualityLevel()).ToString());
 			}
 			else if(Args.Num() == 1)
 			{
@@ -71,8 +71,8 @@ static FAutoConsoleCommand GCmdSetNiagaraQualityLevelOverride(
 					int32 NewQualityLevel = FCString::Atoi(*Args[0]);
 					if(ensure(NewQualityLevel >= 0 && NewQualityLevel <= 4))
 					{
-						SetGNiagaraQualityLevel(NewQualityLevel);
-						UE_LOG(LogNiagara, Display, TEXT("Set the global Niagara quality level to %s"), *FNiagaraPlatformSet::GetQualityLevelText(GNiagaraQualityLevel).ToString());
+						FNiagaraPlatformSet::SetNiagaraQualityLevelOverride(NewQualityLevel);
+						UE_LOG(LogNiagara, Display, TEXT("Set the global Niagara quality level to %s"), *FNiagaraPlatformSet::GetQualityLevelText(FNiagaraPlatformSet::GetQualityLevel()).ToString());
 					}
 				}
 			}
@@ -87,20 +87,24 @@ static FAutoConsoleVariableRef CVarNiagaraQualityLevel(
 	ECVF_Scalability
 );
 
+void FNiagaraPlatformSet::SetNiagaraQualityLevelOverride(int32 Override) 
+{
+	if(Override != QualityLevelOverride)
+	{
+		QualityLevelOverride = Override; 
+		InvalidateCachedData();
+		RefreshScalability();
+	}
+}
+
+void FNiagaraPlatformSet::ClearNiagaraQualityLevelOverride() 
+{
+	SetNiagaraQualityLevelOverride(INDEX_NONE);
+};
+
 // Sync function for checking quality level hasn't changed
 // This must be done at the end of all console variable changes otherwise we can encounter multiple changes to the Q level resulting in some systems being left deactivated
-FAutoConsoleVariableSink CVarNiagaraQualityLevelSync(
-	FConsoleCommandDelegate::CreateLambda(
-		[]()
-		{
-			const int32 CurrentLevel = FNiagaraPlatformSet::GetQualityLevel();
-			if (CurrentLevel != GNiagaraQualityLevel)
-			{
-				SetGNiagaraQualityLevel(GNiagaraQualityLevel);
-			}
-		}
-	)
-);
+FAutoConsoleVariableSink CVarNiagaraQualityLevelSync(FConsoleCommandDelegate::CreateStatic(&FNiagaraPlatformSet::RefreshScalability));
 
 /** Minimum quality level available for a platform. Defaults to -1, indicating there is no minimum. */
 static const int32 DefaultMinQualityLevel = INDEX_NONE;
@@ -213,13 +217,23 @@ static UDeviceProfile* NiagaraGetActiveDeviceProfile()
 
 int32 FNiagaraPlatformSet::CachedQualityLevel = INDEX_NONE;
 int32 FNiagaraPlatformSet::CachedAvailableQualityLevelMask = INDEX_NONE;
+int32 FNiagaraPlatformSet::QualityLevelOverride = INDEX_NONE;
 
 int32 FNiagaraPlatformSet::GetQualityLevel()
 {
+	if(QualityLevelOverride != INDEX_NONE)
+	{
+		return QualityLevelOverride;
+	}
+
 	if (CachedQualityLevel == INDEX_NONE)
 	{
-		CachedQualityLevel = GNiagaraQualityLevel;
+		//Return GNiagaraQualityLevel so that we have a valid level but do not update cached quality level. 
+		//This must be kept in order to trigger a proper refresh of scalability state.
+		//It is possible that InvalidateCachedData to be called and we reach here before we hit our own CVar Sinks and refresh handling.
+		return GNiagaraQualityLevel;
 	}
+
 	return CachedQualityLevel;
 }
 
@@ -976,6 +990,23 @@ IConsoleVariable* FNiagaraPlatformSet::GetCVar(FName CVarName)
 
 void FNiagaraPlatformSet::RefreshScalability()
 {
+	//If the main quality level has changed we do a full refresh, other wise we can do a more targeted refresh just based on the changed CVars.
+	if (GNiagaraQualityLevel != CachedQualityLevel)
+	{
+		FNiagaraPlatformSet::InvalidateCachedData();
+		CachedQualityLevel = GNiagaraQualityLevel;
+
+		for (TObjectIterator<UNiagaraSystem> It; It; ++It)
+		{
+			UNiagaraSystem* System = *It;
+			check(System);
+			System->UpdateScalability();
+		}
+
+		//We've done a full scalability so clear these out as we wont need to do a targeted refresh.
+		ChangedCVars.Reset();
+	}
+
 	bool bNeedsRefresh = ChangedCVars.Num() > 0;//Other reasons to refresh deferred like this?
 
 	if (bNeedsRefresh)
@@ -1031,7 +1062,6 @@ void FNiagaraPlatformSet::InvalidateCachedData()
 
 	CachedQualityLevel = INDEX_NONE;
 	LastDirtiedFrame = GFrameNumber;
-	CachedQualityLevel = INDEX_NONE;
 }
 
 bool FNiagaraPlatformSet::IsEnabledForPlatform(const FString& PlatformName)const
@@ -1078,19 +1108,6 @@ bool FNiagaraPlatformSet::CanChangeScalabilityAtRuntime(const UDeviceProfile* De
 	return true;//Assuming true if we fail to find the platform seems safest.
 }
 
-void SetGNiagaraQualityLevel(int32 QualityLevel)
-{
-	GNiagaraQualityLevel = QualityLevel;
-	FNiagaraPlatformSet::InvalidateCachedData();
-
-	for (TObjectIterator<UNiagaraSystem> It; It; ++It)
-	{
-		UNiagaraSystem* System = *It;
-		check(System);
-		System->UpdateScalability();
-	}
-}
-
 void SetGNiagaraDeviceProfile(UDeviceProfile* Profile)
 {
 	GNiagaraPlatformOverride = Profile;
diff --git a/Engine/Plugins/FX/Niagara/Source/NiagaraEditor/Private/ViewModels/NiagaraBakerViewModel.cpp b/Engine/Plugins/FX/Niagara/Source/NiagaraEditor/Private/ViewModels/NiagaraBakerViewModel.cpp
--- a/Engine/Plugins/FX/Niagara/Source/NiagaraEditor/Private/ViewModels/NiagaraBakerViewModel.cpp
+++ b/Engine/Plugins/FX/Niagara/Source/NiagaraEditor/Private/ViewModels/NiagaraBakerViewModel.cpp
@@ -631,7 +631,6 @@ void FNiagaraBakerViewModel::RenderBaker()
 		FScopedSlowTask SlowTask(TotalFrames);
 		SlowTask.MakeDialog();
 
-		const int32 ExistingQualityLevel = FNiagaraPlatformSet::GetQualityLevel();
 		if (BakerSettings->BakeQualityLevel.IsNone() == false)
 		{
 			const UNiagaraSettings* NiagaraSettings = GetDefault<UNiagaraSettings>();
@@ -639,7 +638,7 @@ void FNiagaraBakerViewModel::RenderBaker()
 			{
 				if (FName(NiagaraSettings->QualityLevels[i].ToString()) == BakerSettings->BakeQualityLevel)
 				{
-					SetGNiagaraQualityLevel(i);
+					FNiagaraPlatformSet::SetNiagaraQualityLevelOverride(i);
 					break;
 				}
 			}
@@ -665,7 +664,7 @@ void FNiagaraBakerViewModel::RenderBaker()
 
 		if (BakerSettings->BakeQualityLevel.IsNone() == false)
 		{
-			SetGNiagaraQualityLevel(ExistingQualityLevel);
+			FNiagaraPlatformSet::ClearNiagaraQualityLevelOverride();
 		}
 	}
 
