Niagara - Fix a bunch of issues with user parameter display and with position user parameter sequencer integration.
+ Added an ensure in the user parameter node builder so that missing position data doesn't cause a crash and will be reported.
+ Fixed an issue with the user parameter node builder where the wrong display struct could be matched up with a user parameter when the override parameters contain two parameters with the same name but with different types.
+ Fixed an issue with the user parameter node builder where position parameters would not be refreshed correctly when edited externally due to copying the wrong size data.
+ Fixed an issue with the previous value cache in the sequencer niagara parameter track where only the first animated parameter would actually get cached and restored because the type id was using the seciton class.
+ Added support for alternate types for parameter tracks so that vector3 tracks can contiue to animate user parameters which have automatically been upgraded to position values without breaking content.
+ Fixed an issue where you could add two parameter tracks with the same name but different types to a niagara component.

diff --git a/Engine/Plugins/FX/Niagara/Source/Niagara/Private/MovieScene/Parameters/MovieSceneNiagaraBoolParameterSectionTemplate.cpp b/Engine/Plugins/FX/Niagara/Source/Niagara/Private/MovieScene/Parameters/MovieSceneNiagaraBoolParameterSectionTemplate.cpp
--- a/Engine/Plugins/FX/Niagara/Source/Niagara/Private/MovieScene/Parameters/MovieSceneNiagaraBoolParameterSectionTemplate.cpp
+++ b/Engine/Plugins/FX/Niagara/Source/Niagara/Private/MovieScene/Parameters/MovieSceneNiagaraBoolParameterSectionTemplate.cpp
@@ -17,7 +17,7 @@ FMovieSceneNiagaraBoolParameterSectionTemplate::FMovieSceneNiagaraBoolParameterS
 {
 }
 
-void FMovieSceneNiagaraBoolParameterSectionTemplate::GetParameterValue(FFrameTime InTime, const TArray<uint8>& InCurrentValueData, TArray<uint8>& OutAnimatedValueData) const
+void FMovieSceneNiagaraBoolParameterSectionTemplate::GetAnimatedParameterValue(FFrameTime InTime, const FNiagaraVariableBase& InTargetParameter, const TArray<uint8>& InCurrentValueData, TArray<uint8>& OutAnimatedValueData) const
 {
 	FNiagaraBool const* CurrentValue = (FNiagaraBool const*)InCurrentValueData.GetData();
 	FNiagaraBool AnimatedNiagaraValue = *CurrentValue;
diff --git a/Engine/Plugins/FX/Niagara/Source/Niagara/Private/MovieScene/Parameters/MovieSceneNiagaraBoolParameterSectionTemplate.h b/Engine/Plugins/FX/Niagara/Source/Niagara/Private/MovieScene/Parameters/MovieSceneNiagaraBoolParameterSectionTemplate.h
--- a/Engine/Plugins/FX/Niagara/Source/Niagara/Private/MovieScene/Parameters/MovieSceneNiagaraBoolParameterSectionTemplate.h
+++ b/Engine/Plugins/FX/Niagara/Source/Niagara/Private/MovieScene/Parameters/MovieSceneNiagaraBoolParameterSectionTemplate.h
@@ -20,7 +20,7 @@ private:
 	virtual UScriptStruct& GetScriptStructImpl() const override { return *StaticStruct(); }
 
 protected:
-	virtual void GetParameterValue(FFrameTime InTime, const TArray<uint8>& InCurrentValueData, TArray<uint8>& OutAnimatedValueData) const override;
+	virtual void GetAnimatedParameterValue(FFrameTime InTime, const FNiagaraVariableBase& InTargetParameter, const TArray<uint8>& InCurrentValueData, TArray<uint8>& OutAnimatedValueData) const override;
 
 private:
 	UPROPERTY()
diff --git a/Engine/Plugins/FX/Niagara/Source/Niagara/Private/MovieScene/Parameters/MovieSceneNiagaraColorParameterSectionTemplate.cpp b/Engine/Plugins/FX/Niagara/Source/Niagara/Private/MovieScene/Parameters/MovieSceneNiagaraColorParameterSectionTemplate.cpp
--- a/Engine/Plugins/FX/Niagara/Source/Niagara/Private/MovieScene/Parameters/MovieSceneNiagaraColorParameterSectionTemplate.cpp
+++ b/Engine/Plugins/FX/Niagara/Source/Niagara/Private/MovieScene/Parameters/MovieSceneNiagaraColorParameterSectionTemplate.cpp
@@ -20,7 +20,7 @@ FMovieSceneNiagaraColorParameterSectionTemplate::FMovieSceneNiagaraColorParamete
 {
 }
 
-void FMovieSceneNiagaraColorParameterSectionTemplate::GetParameterValue(FFrameTime InTime, const TArray<uint8>& InCurrentValueData, TArray<uint8>& OutAnimatedValueData) const
+void FMovieSceneNiagaraColorParameterSectionTemplate::GetAnimatedParameterValue(FFrameTime InTime, const FNiagaraVariableBase& InTargetParameter, const TArray<uint8>& InCurrentValueData, TArray<uint8>& OutAnimatedValueData) const
 {
 	FLinearColor const* CurrentValue = (FLinearColor const*)InCurrentValueData.GetData();
 	FLinearColor AnimatedValue = *CurrentValue;
diff --git a/Engine/Plugins/FX/Niagara/Source/Niagara/Private/MovieScene/Parameters/MovieSceneNiagaraColorParameterSectionTemplate.h b/Engine/Plugins/FX/Niagara/Source/Niagara/Private/MovieScene/Parameters/MovieSceneNiagaraColorParameterSectionTemplate.h
--- a/Engine/Plugins/FX/Niagara/Source/Niagara/Private/MovieScene/Parameters/MovieSceneNiagaraColorParameterSectionTemplate.h
+++ b/Engine/Plugins/FX/Niagara/Source/Niagara/Private/MovieScene/Parameters/MovieSceneNiagaraColorParameterSectionTemplate.h
@@ -20,7 +20,7 @@ private:
 	virtual UScriptStruct& GetScriptStructImpl() const override { return *StaticStruct(); }
 
 protected:
-	virtual void GetParameterValue(FFrameTime InTime, const TArray<uint8>& InCurrentValueData, TArray<uint8>& OutAnimatedValueData) const override;
+	virtual void GetAnimatedParameterValue(FFrameTime InTime, const FNiagaraVariableBase& InTargetParameter, const TArray<uint8>& InCurrentValueData, TArray<uint8>& OutAnimatedValueData) const override;
 
 private:
 	UPROPERTY()
diff --git a/Engine/Plugins/FX/Niagara/Source/Niagara/Private/MovieScene/Parameters/MovieSceneNiagaraFloatParameterSectionTemplate.cpp b/Engine/Plugins/FX/Niagara/Source/Niagara/Private/MovieScene/Parameters/MovieSceneNiagaraFloatParameterSectionTemplate.cpp
--- a/Engine/Plugins/FX/Niagara/Source/Niagara/Private/MovieScene/Parameters/MovieSceneNiagaraFloatParameterSectionTemplate.cpp
+++ b/Engine/Plugins/FX/Niagara/Source/Niagara/Private/MovieScene/Parameters/MovieSceneNiagaraFloatParameterSectionTemplate.cpp
@@ -17,7 +17,7 @@ FMovieSceneNiagaraFloatParameterSectionTemplate::FMovieSceneNiagaraFloatParamete
 {
 }
 
-void FMovieSceneNiagaraFloatParameterSectionTemplate::GetParameterValue(FFrameTime InTime, const TArray<uint8>& InCurrentValueData, TArray<uint8>& OutAnimatedValueData) const
+void FMovieSceneNiagaraFloatParameterSectionTemplate::GetAnimatedParameterValue(FFrameTime InTime, const FNiagaraVariableBase& InTargetParameter, const TArray<uint8>& InCurrentValueData, TArray<uint8>& OutAnimatedValueData) const
 {
 	FNiagaraFloat const* CurrentValue = (FNiagaraFloat const*)InCurrentValueData.GetData();
 	FNiagaraFloat AnimatedValue = *CurrentValue;
diff --git a/Engine/Plugins/FX/Niagara/Source/Niagara/Private/MovieScene/Parameters/MovieSceneNiagaraFloatParameterSectionTemplate.h b/Engine/Plugins/FX/Niagara/Source/Niagara/Private/MovieScene/Parameters/MovieSceneNiagaraFloatParameterSectionTemplate.h
--- a/Engine/Plugins/FX/Niagara/Source/Niagara/Private/MovieScene/Parameters/MovieSceneNiagaraFloatParameterSectionTemplate.h
+++ b/Engine/Plugins/FX/Niagara/Source/Niagara/Private/MovieScene/Parameters/MovieSceneNiagaraFloatParameterSectionTemplate.h
@@ -20,7 +20,7 @@ private:
 	virtual UScriptStruct& GetScriptStructImpl() const override { return *StaticStruct(); }
 
 protected:
-	virtual void GetParameterValue(FFrameTime InTime, const TArray<uint8>& InCurrentValueData, TArray<uint8>& OutAnimatedValueData) const override;
+	virtual void GetAnimatedParameterValue(FFrameTime InTime, const FNiagaraVariableBase& InTargetParameter, const TArray<uint8>& InCurrentValueData, TArray<uint8>& OutAnimatedValueData) const override;
 
 private:
 	UPROPERTY()
diff --git a/Engine/Plugins/FX/Niagara/Source/Niagara/Private/MovieScene/Parameters/MovieSceneNiagaraIntegerParameterSectionTemplate.cpp b/Engine/Plugins/FX/Niagara/Source/Niagara/Private/MovieScene/Parameters/MovieSceneNiagaraIntegerParameterSectionTemplate.cpp
--- a/Engine/Plugins/FX/Niagara/Source/Niagara/Private/MovieScene/Parameters/MovieSceneNiagaraIntegerParameterSectionTemplate.cpp
+++ b/Engine/Plugins/FX/Niagara/Source/Niagara/Private/MovieScene/Parameters/MovieSceneNiagaraIntegerParameterSectionTemplate.cpp
@@ -17,7 +17,7 @@ FMovieSceneNiagaraIntegerParameterSectionTemplate::FMovieSceneNiagaraIntegerPara
 {
 }
 
-void FMovieSceneNiagaraIntegerParameterSectionTemplate::GetParameterValue(FFrameTime InTime, const TArray<uint8>& InCurrentValueData, TArray<uint8>& OutAnimatedValueData) const
+void FMovieSceneNiagaraIntegerParameterSectionTemplate::GetAnimatedParameterValue(FFrameTime InTime, const FNiagaraVariableBase& InTargetParameter, const TArray<uint8>& InCurrentValueData, TArray<uint8>& OutAnimatedValueData) const
 {
 	FNiagaraInt32 const* CurrentValue = (FNiagaraInt32 const*)InCurrentValueData.GetData();
 	FNiagaraInt32 AnimatedValue = *CurrentValue;
diff --git a/Engine/Plugins/FX/Niagara/Source/Niagara/Private/MovieScene/Parameters/MovieSceneNiagaraIntegerParameterSectionTemplate.h b/Engine/Plugins/FX/Niagara/Source/Niagara/Private/MovieScene/Parameters/MovieSceneNiagaraIntegerParameterSectionTemplate.h
--- a/Engine/Plugins/FX/Niagara/Source/Niagara/Private/MovieScene/Parameters/MovieSceneNiagaraIntegerParameterSectionTemplate.h
+++ b/Engine/Plugins/FX/Niagara/Source/Niagara/Private/MovieScene/Parameters/MovieSceneNiagaraIntegerParameterSectionTemplate.h
@@ -20,7 +20,7 @@ private:
 	virtual UScriptStruct& GetScriptStructImpl() const override { return *StaticStruct(); }
 
 protected:
-	virtual void GetParameterValue(FFrameTime InTime, const TArray<uint8>& InCurrentValueData, TArray<uint8>& OutAnimatedValueData) const override;
+	virtual void GetAnimatedParameterValue(FFrameTime InTime, const FNiagaraVariableBase& InTargetParameter, const TArray<uint8>& InCurrentValueData, TArray<uint8>& OutAnimatedValueData) const override;
 
 private:
 	UPROPERTY()
diff --git a/Engine/Plugins/FX/Niagara/Source/Niagara/Private/MovieScene/Parameters/MovieSceneNiagaraParameterSectionTemplate.cpp b/Engine/Plugins/FX/Niagara/Source/Niagara/Private/MovieScene/Parameters/MovieSceneNiagaraParameterSectionTemplate.cpp
--- a/Engine/Plugins/FX/Niagara/Source/Niagara/Private/MovieScene/Parameters/MovieSceneNiagaraParameterSectionTemplate.cpp
+++ b/Engine/Plugins/FX/Niagara/Source/Niagara/Private/MovieScene/Parameters/MovieSceneNiagaraParameterSectionTemplate.cpp
@@ -12,6 +12,7 @@
 struct FComponentData
 {
 	TWeakObjectPtr<UNiagaraComponent> Component;
+	FNiagaraVariableBase TargetParameter;
 	TOptional<TArray<uint8>> CurrentValue;
 };
 
@@ -81,18 +82,27 @@ struct FSetParameterValueExecutionToken : IMovieSceneExecutionToken
 
 	virtual void Execute(const FMovieSceneContext& Context, const FMovieSceneEvaluationOperand& Operand, FPersistentEvaluationData& PersistentData, IMovieScenePlayer& Player) override
 	{
-		static FMovieSceneAnimTypeID TypeID = TMovieSceneAnimTypeID<FSetParameterValueExecutionToken, 0>();
+		static TMovieSceneAnimTypeIDContainer<FNiagaraVariableBase> ParameterToTypeID;
+		FMovieSceneAnimTypeID TypeID = ParameterToTypeID.GetAnimTypeID(Parameter);
 		UNiagaraComponent* NiagaraComponent = ComponentPtr.Get();
 		if (NiagaraComponent != nullptr)
 		{
-			Player.SavePreAnimatedState(*NiagaraComponent, TypeID, FPreAnimatedParameterValueTokenProducer(Parameter));
+			Player.PreAnimatedState.SavePreAnimatedState(*NiagaraComponent, TypeID, FPreAnimatedParameterValueTokenProducer(Parameter));
 			NiagaraComponent->GetOverrideParameters().AddParameter(Parameter, false);
-			NiagaraComponent->GetOverrideParameters().SetParameterData(Data.GetData(), Parameter);
+			if (Parameter.GetType() == FNiagaraTypeDefinition::GetPositionDef())
+			{
+				FVector* PositionValue = reinterpret_cast<FVector*>(Data.GetData());
+				NiagaraComponent->GetOverrideParameters().SetPositionParameterValue(*PositionValue, Parameter.GetName(), false);
+			}
+			else
+			{
+				NiagaraComponent->GetOverrideParameters().SetParameterData(Data.GetData(), Parameter);
+			}
 		}
 	}
 
 	TWeakObjectPtr<UNiagaraComponent> ComponentPtr;
-	FNiagaraVariable Parameter;
+	FNiagaraVariableBase Parameter;
 	TArray<uint8> Data;
 };
 
@@ -112,32 +122,52 @@ void FMovieSceneNiagaraParameterSectionTemplate::Initialize(const FMovieSceneEva
 	for (TWeakObjectPtr<> ObjectPtr : Player.FindBoundObjects(Operand))
 	{
 		UNiagaraComponent* NiagaraComponent = Cast<UNiagaraComponent>(ObjectPtr.Get());
-		if (NiagaraComponent != nullptr)
+		if (NiagaraComponent != nullptr && NiagaraComponent->GetAsset() != nullptr)
 		{
+			FNiagaraVariableBase TargetParameter;
 			const uint8* ParameterData = NiagaraComponent->GetOverrideParameters().GetParameterData(Parameter);
 			if (ParameterData == nullptr)
 			{
-				if (UNiagaraSystem* NiagaraSystem = NiagaraComponent->GetAsset())
+				ParameterData = NiagaraComponent->GetAsset()->GetExposedParameters().GetParameterData(Parameter);
+			};
+
+			if (ParameterData != nullptr)
+			{
+				TargetParameter = Parameter;
+			}
+			else
+			{
+				// If a matching parameter was not found in the component overrides or the assets 
+				// exposed parameters check if there are alternate types which are supported for this section.
+				for (const FNiagaraTypeDefinition& AlternateParameterType : GetAlternateParameterTypes())
 				{
-					ParameterData = NiagaraSystem->GetExposedParameters().GetParameterData(Parameter);
+					FNiagaraVariable AlternateParameter(AlternateParameterType, Parameter.GetName());
+					ParameterData = NiagaraComponent->GetOverrideParameters().GetParameterData(AlternateParameter);
+					if (ParameterData == nullptr)
+					{
+						ParameterData = NiagaraComponent->GetAsset()->GetExposedParameters().GetParameterData(AlternateParameter);
+					}
+
+					if (ParameterData != nullptr)
+					{
+						TargetParameter = AlternateParameter;
+						break;
+					}
 				}
 			}
 
-			TArray<uint8> CurrentValueData;
-			if(ParameterData != nullptr)
+			if (TargetParameter.IsValid() && ParameterData != nullptr)
 			{
+				TArray<uint8> CurrentValueData;
 				CurrentValueData.AddUninitialized(Parameter.GetSizeInBytes());
 				FMemory::Memcpy(CurrentValueData.GetData(), ParameterData, Parameter.GetSizeInBytes());
-			}
-			else
-			{
-				CurrentValueData.AddDefaulted(Parameter.GetSizeInBytes());
-			}
 
-			FComponentData Data;
-			Data.Component = NiagaraComponent;
-			Data.CurrentValue.Emplace(CurrentValueData);
-			SectionData.CachedComponentData.Add(Data);
+				FComponentData Data;
+				Data.Component = NiagaraComponent;
+				Data.TargetParameter = TargetParameter;
+				Data.CurrentValue.Emplace(CurrentValueData);
+				SectionData.CachedComponentData.Add(Data);
+			}
 		}
 	}
 }
@@ -152,8 +182,8 @@ void FMovieSceneNiagaraParameterSectionTemplate::Evaluate(const FMovieSceneEvalu
 			if (ComponentData.CurrentValue.IsSet())
 			{
 				TArray<uint8> AnimatedValueData;
-				GetParameterValue(Context.GetTime(), ComponentData.CurrentValue.GetValue(), AnimatedValueData);
-				ExecutionTokens.Add(FSetParameterValueExecutionToken(ComponentData.Component, Parameter, MoveTemp(AnimatedValueData)));
+				GetAnimatedParameterValue(Context.GetTime(), ComponentData.TargetParameter, ComponentData.CurrentValue.GetValue(),  AnimatedValueData);
+				ExecutionTokens.Add(FSetParameterValueExecutionToken(ComponentData.Component, ComponentData.TargetParameter, MoveTemp(AnimatedValueData)));
 			}
 		}
 	}
diff --git a/Engine/Plugins/FX/Niagara/Source/Niagara/Private/MovieScene/Parameters/MovieSceneNiagaraParameterSectionTemplate.h b/Engine/Plugins/FX/Niagara/Source/Niagara/Private/MovieScene/Parameters/MovieSceneNiagaraParameterSectionTemplate.h
--- a/Engine/Plugins/FX/Niagara/Source/Niagara/Private/MovieScene/Parameters/MovieSceneNiagaraParameterSectionTemplate.h
+++ b/Engine/Plugins/FX/Niagara/Source/Niagara/Private/MovieScene/Parameters/MovieSceneNiagaraParameterSectionTemplate.h
@@ -27,7 +27,10 @@ private:
 	virtual void Evaluate(const FMovieSceneEvaluationOperand& Operand, const FMovieSceneContext& Context, const FPersistentEvaluationData& PersistentData, FMovieSceneExecutionTokens& ExecutionTokens) const override;
 
 protected:
-	virtual void GetParameterValue(FFrameTime InTime, const TArray<uint8>& InCurrentValueData, TArray<uint8>& OutAnimatedValueData) const { };
+	virtual void GetAnimatedParameterValue(FFrameTime InTime, const FNiagaraVariableBase& InTargetParameter, const TArray<uint8>& InCurrentValueData, TArray<uint8>& OutAnimatedValueData) const { };
+	
+	// Specifies a list of alternate types which are supported by this parameter section which can be used if the original parameter isn't found.
+	virtual TArrayView<FNiagaraTypeDefinition> GetAlternateParameterTypes() const { return TArrayView<FNiagaraTypeDefinition>(); }
 
 private:
 	UPROPERTY()
diff --git a/Engine/Plugins/FX/Niagara/Source/Niagara/Private/MovieScene/Parameters/MovieSceneNiagaraVectorParameterSectionTemplate.cpp b/Engine/Plugins/FX/Niagara/Source/Niagara/Private/MovieScene/Parameters/MovieSceneNiagaraVectorParameterSectionTemplate.cpp
--- a/Engine/Plugins/FX/Niagara/Source/Niagara/Private/MovieScene/Parameters/MovieSceneNiagaraVectorParameterSectionTemplate.cpp
+++ b/Engine/Plugins/FX/Niagara/Source/Niagara/Private/MovieScene/Parameters/MovieSceneNiagaraVectorParameterSectionTemplate.cpp
@@ -21,7 +21,7 @@ FMovieSceneNiagaraVectorParameterSectionTemplate::FMovieSceneNiagaraVectorParame
 	}
 }
 
-void FMovieSceneNiagaraVectorParameterSectionTemplate::GetParameterValue(FFrameTime InTime, const TArray<uint8>& InCurrentValueData, TArray<uint8>& OutAnimatedValueData) const
+void FMovieSceneNiagaraVectorParameterSectionTemplate::GetAnimatedParameterValue(FFrameTime InTime, const FNiagaraVariableBase& InTargetParameter, const TArray<uint8>& InCurrentValueData, TArray<uint8>& OutAnimatedValueData) const
 {
 	if (ChannelsUsed == 2)
 	{
@@ -43,8 +43,19 @@ void FMovieSceneNiagaraVectorParameterSectionTemplate::GetParameterValue(FFrameT
 		VectorChannels[1].Evaluate(InTime, AnimatedValue.Y);
 		VectorChannels[2].Evaluate(InTime, AnimatedValue.Z);
 
-		OutAnimatedValueData.AddUninitialized(sizeof(FVector3f));
-		FMemory::Memcpy(OutAnimatedValueData.GetData(), (uint8*)&AnimatedValue, sizeof(FVector3f));
+		if (InTargetParameter.GetType() != FNiagaraTypeDefinition::GetPositionDef())
+		{
+			OutAnimatedValueData.AddUninitialized(sizeof(FVector3f));
+			FMemory::Memcpy(OutAnimatedValueData.GetData(), (uint8*)&AnimatedValue, sizeof(FVector3f));
+		}
+		else
+		{
+			OutAnimatedValueData.AddUninitialized(sizeof(FVector));
+			FVector* PositionVector = (FVector*)OutAnimatedValueData.GetData();
+			PositionVector->X = AnimatedValue.X;
+			PositionVector->Y = AnimatedValue.Y;
+			PositionVector->Z = AnimatedValue.Z;
+		}
 	}
 	else if (ChannelsUsed == 4)
 	{
@@ -60,3 +71,16 @@ void FMovieSceneNiagaraVectorParameterSectionTemplate::GetParameterValue(FFrameT
 		FMemory::Memcpy(OutAnimatedValueData.GetData(), (uint8*)&AnimatedValue, sizeof(FVector4f));
 	}
 }
+
+TArrayView<FNiagaraTypeDefinition> FMovieSceneNiagaraVectorParameterSectionTemplate::GetAlternateParameterTypes() const
+{
+	if (ChannelsUsed == 3)
+	{
+		static TArray<FNiagaraTypeDefinition> AlternateParameterTypes = { FNiagaraTypeDefinition::GetPositionDef() };
+		return TArrayView<FNiagaraTypeDefinition>(AlternateParameterTypes);
+	}
+	else
+	{
+		return TArrayView<FNiagaraTypeDefinition>();
+	}
+}
diff --git a/Engine/Plugins/FX/Niagara/Source/Niagara/Private/MovieScene/Parameters/MovieSceneNiagaraVectorParameterSectionTemplate.h b/Engine/Plugins/FX/Niagara/Source/Niagara/Private/MovieScene/Parameters/MovieSceneNiagaraVectorParameterSectionTemplate.h
--- a/Engine/Plugins/FX/Niagara/Source/Niagara/Private/MovieScene/Parameters/MovieSceneNiagaraVectorParameterSectionTemplate.h
+++ b/Engine/Plugins/FX/Niagara/Source/Niagara/Private/MovieScene/Parameters/MovieSceneNiagaraVectorParameterSectionTemplate.h
@@ -20,7 +20,8 @@ private:
 	virtual UScriptStruct& GetScriptStructImpl() const override { return *StaticStruct(); }
 
 protected:
-	virtual void GetParameterValue(FFrameTime InTime, const TArray<uint8>& InCurrentValueData, TArray<uint8>& OutAnimatedValueData) const override;
+	virtual void GetAnimatedParameterValue(FFrameTime InTime, const FNiagaraVariableBase& InTargetParameter, const TArray<uint8>& InCurrentValueData, TArray<uint8>& OutAnimatedValueData) const override;
+	virtual TArrayView<FNiagaraTypeDefinition> GetAlternateParameterTypes() const override;
 
 private:
 	UPROPERTY()
diff --git a/Engine/Plugins/FX/Niagara/Source/NiagaraEditor/Private/Customizations/NiagaraComponentDetails.cpp b/Engine/Plugins/FX/Niagara/Source/NiagaraEditor/Private/Customizations/NiagaraComponentDetails.cpp
--- a/Engine/Plugins/FX/Niagara/Source/NiagaraEditor/Private/Customizations/NiagaraComponentDetails.cpp
+++ b/Engine/Plugins/FX/Niagara/Source/NiagaraEditor/Private/Customizations/NiagaraComponentDetails.cpp
@@ -468,7 +468,7 @@ void FNiagaraUserParameterNodeBuilder::GenerateRowForUserParameter(IDetailChildr
 			}
 		}
 	
-		ParameterNameToDisplayStruct.Add(UserParameter.GetName(), TWeakPtr<FStructOnScope>(StructOnScope));
+		ParameterToDisplayStruct.Add(UserParameter, TWeakPtr<FStructOnScope>(StructOnScope));
 	}
 
 	TSharedPtr<SWidget> DefaultNameWidget;
@@ -896,7 +896,7 @@ void FNiagaraComponentUserParametersNodeBuilder::GenerateChildContent(IDetailChi
 		
 	ParameterProxies.Reserve(UserParameters.Num());
 
-	ParameterNameToDisplayStruct.Empty();
+	ParameterToDisplayStruct.Empty();
 
 	GenerateUserParameterRows(ChildrenBuilder, *Cast<UNiagaraSystemEditorData>(SystemAsset->GetEditorData())->UserParameterHierarchy.Get(),Component.Get());
 }
@@ -925,13 +925,18 @@ void FNiagaraComponentUserParametersNodeBuilder::ParameterValueChanged()
 		{
 			if (UserParameter.IsUObject() == false)
 			{
-				TWeakPtr<FStructOnScope>* DisplayStructPtr = ParameterNameToDisplayStruct.Find(UserParameter.GetName());
+				TWeakPtr<FStructOnScope>* DisplayStructPtr = ParameterToDisplayStruct.Find(UserParameter);
 				if (DisplayStructPtr != nullptr && DisplayStructPtr->IsValid())
 				{
 					TSharedPtr<FStructOnScope> DisplayStruct = DisplayStructPtr->Pin();
 					if (UserParameter.GetType() == FNiagaraTypeDefinition::GetPositionDef())
 					{
-						FMemory::Memcpy(DisplayStruct->GetStructMemory(), OverrideParameters.GetPositionParameterValue(UserParameter.GetName()), UserParameter.GetSizeInBytes());
+						const FVector* PositionValue = OverrideParameters.GetPositionParameterValue(UserParameter.GetName());
+						if (ensureMsgf(PositionValue != nullptr, TEXT("Position user parameter %s was missing it's position data.  Path: %s"),
+							*UserParameter.GetName().ToString(), *Component->GetPathName()))
+						{
+							FMemory::Memcpy(DisplayStruct->GetStructMemory(), PositionValue, DisplayStruct->GetStruct()->GetStructureSize());
+						}
 					}
 					else
 					{
@@ -1000,7 +1005,7 @@ void FNiagaraSystemUserParameterBuilder::GenerateChildContent(IDetailChildrenBui
 		System->GetExposedParameters().GetUserParameters(UserParameters);
 		ParameterProxies.Reserve(UserParameters.Num());
 
-		ParameterNameToDisplayStruct.Empty();
+		ParameterToDisplayStruct.Empty();
 
 		UNiagaraHierarchyRoot* Root = Cast<UNiagaraSystemEditorData>(System->GetEditorData())->UserParameterHierarchy;
 		GenerateUserParameterRows(ChildrenBuilder, *Root, System.Get());
@@ -1087,7 +1092,7 @@ void FNiagaraSystemUserParameterBuilder::ParameterValueChanged()
 		{
 			if (UserParameter.IsUObject() == false)
 			{
-				TWeakPtr<FStructOnScope>* DisplayStructPtr = ParameterNameToDisplayStruct.Find(UserParameter.GetName());
+				TWeakPtr<FStructOnScope>* DisplayStructPtr = ParameterToDisplayStruct.Find(UserParameter);
 				if (DisplayStructPtr != nullptr && DisplayStructPtr->IsValid())
 				{
 					TSharedPtr<FStructOnScope> DisplayStruct = DisplayStructPtr->Pin();
diff --git a/Engine/Plugins/FX/Niagara/Source/NiagaraEditor/Private/Customizations/NiagaraComponentDetails.h b/Engine/Plugins/FX/Niagara/Source/NiagaraEditor/Private/Customizations/NiagaraComponentDetails.h
--- a/Engine/Plugins/FX/Niagara/Source/NiagaraEditor/Private/Customizations/NiagaraComponentDetails.h
+++ b/Engine/Plugins/FX/Niagara/Source/NiagaraEditor/Private/Customizations/NiagaraComponentDetails.h
@@ -114,7 +114,7 @@ protected:
 	TArray<TSharedPtr<IPropertyHandle>> OverridePropertyHandles;
 	FSimpleDelegate OnRebuildChildren;
 	TArray<TSharedPtr<class FNiagaraParameterProxy>> ParameterProxies;
-	TMap<FName, TWeakPtr<FStructOnScope>> ParameterNameToDisplayStruct;
+	TMap<FNiagaraVariableBase, TWeakPtr<FStructOnScope>> ParameterToDisplayStruct;
 	FName CustomBuilderRowName;
 	const UNiagaraHierarchySection* ActiveSection = nullptr;
 	TArray<UNiagaraHierarchySection*> AvailableSections;
diff --git a/Engine/Plugins/FX/Niagara/Source/NiagaraEditor/Private/Sequencer/LevelSequence/NiagaraSystemTrackEditor.cpp b/Engine/Plugins/FX/Niagara/Source/NiagaraEditor/Private/Sequencer/LevelSequence/NiagaraSystemTrackEditor.cpp
--- a/Engine/Plugins/FX/Niagara/Source/NiagaraEditor/Private/Sequencer/LevelSequence/NiagaraSystemTrackEditor.cpp
+++ b/Engine/Plugins/FX/Niagara/Source/NiagaraEditor/Private/Sequencer/LevelSequence/NiagaraSystemTrackEditor.cpp
@@ -55,7 +55,7 @@ bool HasLifeCycleTrack(UMovieScene& MovieScene, FGuid ObjectBinding)
 	return false;
 }
 
-void GetAnimatedParameters(UMovieScene& MovieScene, FGuid ObjectBinding, TSet<FNiagaraVariable>& AnimatedParameters)
+void GetAnimatedParameterNames(UMovieScene& MovieScene, FGuid ObjectBinding, TSet<FName>& AnimatedParameterNames)
 {
 	for (const FMovieSceneBinding& Binding : MovieScene.GetBindings())
 	{
@@ -66,7 +66,7 @@ void GetAnimatedParameters(UMovieScene& MovieScene, FGuid ObjectBinding, TSet<FN
 				UMovieSceneNiagaraParameterTrack* ParameterTrack = Cast<UMovieSceneNiagaraParameterTrack>(Track);
 				if (ParameterTrack != nullptr)
 				{
-					AnimatedParameters.Add(ParameterTrack->GetParameter());
+					AnimatedParameterNames.Add(ParameterTrack->GetParameter().GetName());
 				}
 			}
 			break;
@@ -103,8 +103,8 @@ void FNiagaraSystemTrackEditor::BuildObjectBindingTrackMenu(FMenuBuilder& MenuBu
 		TArray<FNiagaraVariable> ParameterVariables;
 		System->GetExposedParameters().GetUserParameters(ParameterVariables);
 
-		TSet<FNiagaraVariable> AnimatedParameters;
-		GetAnimatedParameters(*GetSequencer()->GetFocusedMovieSceneSequence()->GetMovieScene(), ObjectBindings[0], AnimatedParameters);
+		TSet<FName> AnimatedParameterNames;
+		GetAnimatedParameterNames(*GetSequencer()->GetFocusedMovieSceneSequence()->GetMovieScene(), ObjectBindings[0], AnimatedParameterNames);
 
 		FNiagaraEditorModule& NiagaraEditorModule = FModuleManager::GetModuleChecked<FNiagaraEditorModule>("NiagaraEditor");
 
@@ -113,7 +113,7 @@ void FNiagaraSystemTrackEditor::BuildObjectBindingTrackMenu(FMenuBuilder& MenuBu
 		{
 			if (ParameterVariable.GetType().IsDataInterface() == false &&
 				NiagaraEditorModule.CanCreateParameterTrackForType(*ParameterVariable.GetType().GetScriptStruct()) && 
-				AnimatedParameters.Contains(ParameterVariable) == false)
+				AnimatedParameterNames.Contains(ParameterVariable.GetName()) == false)
 			{
 				TArray<uint8> DefaultValueData;
 				DefaultValueData.AddZeroed(ParameterVariable.GetSizeInBytes());
