- Add tracking for Niagara GPU compute work to CSV files

diff --git a/Engine/Plugins/FX/Niagara/Source/Niagara/Private/NiagaraEmitter.cpp b/Engine/Plugins/FX/Niagara/Source/Niagara/Private/NiagaraEmitter.cpp
--- a/Engine/Plugins/FX/Niagara/Source/Niagara/Private/NiagaraEmitter.cpp
+++ b/Engine/Plugins/FX/Niagara/Source/Niagara/Private/NiagaraEmitter.cpp
@@ -1627,23 +1627,19 @@ void FVersionedNiagaraEmitterData::UpdateDebugName(const UNiagaraEmitter& Emitte
 {
 #if WITH_NIAGARA_DEBUG_EMITTER_NAME
 	// Ensure our debug simulation name is up to date
-	// Only required on uncooked as it can change due to compilation
-	if (!FPlatformProperties::RequiresCookedData())
+	DebugSimName.Empty();
+	if (const UNiagaraSystem* SystemOwner = Emitter.GetTypedOuter<class UNiagaraSystem>())
 	{
-		DebugSimName.Empty();
-		if (const UNiagaraSystem* SystemOwner = Emitter.GetTypedOuter<class UNiagaraSystem>())
-		{
-			DebugSimName = SystemOwner->GetName();
-			DebugSimName.AppendChar(':');
-		}
-		DebugSimName.Append(Emitter.GetName());
-		if (Emitter.IsVersioningEnabled())
-		{
-			DebugSimName.AppendChar(':');
-			DebugSimName.AppendInt(Version.MajorVersion);
-			DebugSimName.AppendChar('.');
-			DebugSimName.AppendInt(Version.MinorVersion);
-		}
+		DebugSimName = SystemOwner->GetName();
+		DebugSimName.AppendChar(':');
+	}
+	DebugSimName.Append(Emitter.GetName());
+	if (Emitter.IsVersioningEnabled())
+	{
+		DebugSimName.AppendChar(':');
+		DebugSimName.AppendInt(Version.MajorVersion);
+		DebugSimName.AppendChar('.');
+		DebugSimName.AppendInt(Version.MinorVersion);
 	}
 #endif
 
diff --git a/Engine/Plugins/FX/Niagara/Source/Niagara/Private/NiagaraGpuComputeDispatch.cpp b/Engine/Plugins/FX/Niagara/Source/Niagara/Private/NiagaraGpuComputeDispatch.cpp
--- a/Engine/Plugins/FX/Niagara/Source/Niagara/Private/NiagaraGpuComputeDispatch.cpp
+++ b/Engine/Plugins/FX/Niagara/Source/Niagara/Private/NiagaraGpuComputeDispatch.cpp
@@ -49,6 +49,8 @@ DECLARE_GPU_STAT_NAMED(NiagaraGPUComputeFreeIDs, TEXT("Niagara GPU Compute All F
 DECLARE_GPU_STAT_NAMED(NiagaraGPUComputeFreeIDsEmitter, TEXT("Niagara GPU Compute Emitter Free IDs"));
 DECLARE_GPU_STAT_NAMED(NiagaraGPUSorting, TEXT("Niagara GPU sorting"));
 
+CSV_DEFINE_CATEGORY_MODULE(NIAGARA_API, NiagaraGpuCompute, true);
+
 uint32 FNiagaraComputeExecutionContext::TickCounter = 0;
 
 int32 GNiagaraGpuSubmitCommandHint = 0;
@@ -146,6 +148,16 @@ namespace FNiagaraGpuComputeDispatchLocal
 			AfterTransitionArray.Emplace(IntUAV, AfterState, BeforeState);
 		}
 	}
+
+	static bool CsvStatsEnabled()
+	{
+	#if WITH_PARTICLE_PERF_CSV_STATS
+		static IConsoleVariable* DetailedCVar = IConsoleManager::Get().FindConsoleVariable(TEXT("fx.DetailedCSVStats"));
+		return DetailedCVar && DetailedCVar->GetBool();
+	#else
+		return false;
+	#endif
+	}
 }
 
 //////////////////////////////////////////////////////////////////////////
@@ -1624,6 +1636,15 @@ void FNiagaraGpuComputeDispatch::DispatchStage(FRDGBuilder& GraphBuilder, const
 		);
 
 		INC_DWORD_STAT(STAT_NiagaraGPUDispatches);
+	#if WITH_PER_SYSTEM_PARTICLE_PERF_STATS && WITH_NIAGARA_DEBUG_EMITTER_NAME
+		if (FNiagaraGpuComputeDispatchLocal::CsvStatsEnabled())
+		{
+			if (FCsvProfiler* CSVProfiler = FCsvProfiler::Get())
+			{
+				CSVProfiler->RecordCustomStat(InstanceData.Context->GetDebugSimFName(), CSV_CATEGORY_INDEX(NiagaraGpuCompute), 1, ECsvCustomStatOp::Accumulate);
+			}
+		}
+	#endif
 	}
 
 	// Optionally submit commands to the GPU
@@ -1814,6 +1835,11 @@ void FNiagaraGpuComputeDispatch::PostRenderOpaque(FRDGBuilder& GraphBuilder, TCo
 	bRequiresReadback = false;
 
 	GNiagaraViewDataManager.ClearSceneTextureParameters();
+
+	if (FNiagaraGpuComputeDispatchLocal::CsvStatsEnabled())
+	{
+		CSV_CUSTOM_STAT(NiagaraGpuCompute, TotalDispatchesThisFrame, TotalDispatchesThisFrame, ECsvCustomStatOp::Set);
+	}
 }
 
 void FNiagaraGpuComputeDispatch::ProcessDebugReadbacks(FRHICommandListImmediate& RHICmdList, bool bWaitCompletion)
