- IWYU fix

#rnx

diff --git a/Engine/Plugins/FX/Niagara/Source/Niagara/Private/NiagaraGpuComputeDispatch.cpp b/Engine/Plugins/FX/Niagara/Source/Niagara/Private/NiagaraGpuComputeDispatch.cpp
--- a/Engine/Plugins/FX/Niagara/Source/Niagara/Private/NiagaraGpuComputeDispatch.cpp
+++ b/Engine/Plugins/FX/Niagara/Source/Niagara/Private/NiagaraGpuComputeDispatch.cpp
@@ -9,6 +9,7 @@
 #include "Engine/GameViewportClient.h"
 #include "Engine/World.h"
 #include "Misc/ScopeExit.h"
+#include "ProfilingDebugging/CsvProfiler.h"
 
 #include "NiagaraAsyncGpuTraceHelper.h"
 #include "NiagaraDataInterfaceRW.h"
