- Guard using CSVProfiler with CSV_PROFILER

#rnx

diff --git a/Engine/Plugins/FX/Niagara/Source/Niagara/Private/NiagaraGpuComputeDispatch.cpp b/Engine/Plugins/FX/Niagara/Source/Niagara/Private/NiagaraGpuComputeDispatch.cpp
--- a/Engine/Plugins/FX/Niagara/Source/Niagara/Private/NiagaraGpuComputeDispatch.cpp
+++ b/Engine/Plugins/FX/Niagara/Source/Niagara/Private/NiagaraGpuComputeDispatch.cpp
@@ -1637,7 +1637,7 @@ void FNiagaraGpuComputeDispatch::DispatchStage(FRDGBuilder& GraphBuilder, const
 		);
 
 		INC_DWORD_STAT(STAT_NiagaraGPUDispatches);
-	#if WITH_PER_SYSTEM_PARTICLE_PERF_STATS && WITH_NIAGARA_DEBUG_EMITTER_NAME
+	#if CSV_PROFILER && WITH_PER_SYSTEM_PARTICLE_PERF_STATS && WITH_NIAGARA_DEBUG_EMITTER_NAME
 		if (FNiagaraGpuComputeDispatchLocal::CsvStatsEnabled())
 		{
 			if (FCsvProfiler* CSVProfiler = FCsvProfiler::Get())
