Fixed a bug where the component renderer could no longer bind to float properties created in BP

diff --git a/Engine/Plugins/FX/Niagara/Source/Niagara/Private/NiagaraComponentRendererProperties.cpp b/Engine/Plugins/FX/Niagara/Source/Niagara/Private/NiagaraComponentRendererProperties.cpp
--- a/Engine/Plugins/FX/Niagara/Source/Niagara/Private/NiagaraComponentRendererProperties.cpp
+++ b/Engine/Plugins/FX/Niagara/Source/Niagara/Private/NiagaraComponentRendererProperties.cpp
@@ -47,6 +47,7 @@ bool UNiagaraComponentRendererProperties::IsConvertible(const FNiagaraTypeDefini
 		(SourceType == FNiagaraTypeDefinition::GetQuatDef() && TargetType.GetStruct() == GetFQuatDef().GetStruct()) ||
 		(SourceType == FNiagaraTypeDefinition::GetPositionDef() && TargetType.GetStruct() == GetFVectorDef().GetStruct()) ||
 		(SourceType == FNiagaraTypeDefinition::GetPositionDef() && TargetType.GetStruct() == GetFVector3fDef().GetStruct()) ||
+		(SourceType == FNiagaraTypeDefinition::GetFloatDef() && TargetType.GetStruct() == GetDoubleDef().GetStruct()) ||
 		(SourceType == FNiagaraTypeDefinition::GetQuatDef() && TargetType.GetStruct() == GetFRotatorDef().GetStruct()))
 	{
 		return true;
@@ -69,6 +70,10 @@ NIAGARA_API FNiagaraTypeDefinition UNiagaraComponentRendererProperties::ToNiagar
 	{
 		return FNiagaraTypeDefinition::GetFloatDef();
 	}
+	if (FieldClass->IsChildOf(FDoubleProperty::StaticClass()))
+	{
+		return GetDoubleDef();
+	}
 	if (FieldClass->IsChildOf(FStructProperty::StaticClass()))
 	{
 		FStructProperty* StructProperty = (FStructProperty*)Property;
@@ -172,6 +177,13 @@ FNiagaraTypeDefinition UNiagaraComponentRendererProperties::GetFQuatDef()
 	return FNiagaraTypeDefinition(Struct, FNiagaraTypeDefinition::EAllowUnfriendlyStruct::Allow);
 }
 
+FNiagaraTypeDefinition UNiagaraComponentRendererProperties::GetDoubleDef()
+{
+	static UPackage* NiagaraPkg = FindObjectChecked<UPackage>(nullptr, TEXT("/Script/Niagara"));
+	static UScriptStruct* DoubleStruct = FindObjectChecked<UScriptStruct>(NiagaraPkg, TEXT("NiagaraDouble"));
+	return FNiagaraTypeDefinition(DoubleStruct, FNiagaraTypeDefinition::EAllowUnfriendlyStruct::Allow);
+}
+
 TArray<TWeakObjectPtr<UNiagaraComponentRendererProperties>> UNiagaraComponentRendererProperties::ComponentRendererPropertiesToDeferredInit;
 
 UNiagaraComponentRendererProperties::UNiagaraComponentRendererProperties()
diff --git a/Engine/Plugins/FX/Niagara/Source/Niagara/Private/NiagaraRendererComponents.cpp b/Engine/Plugins/FX/Niagara/Source/Niagara/Private/NiagaraRendererComponents.cpp
--- a/Engine/Plugins/FX/Niagara/Source/Niagara/Private/NiagaraRendererComponents.cpp
+++ b/Engine/Plugins/FX/Niagara/Source/Niagara/Private/NiagaraRendererComponents.cpp
@@ -131,6 +131,12 @@ void ConvertVariableToType(const FNiagaraVariable& SourceVariable, const FNiagar
 		TargetData.AddUninitialized(sizeof(FRotator));
 		*reinterpret_cast<FRotator*>(TargetData.GetData()) = FRotator(Data.Rotator());
 	}
+	else if (SourceType == FNiagaraTypeDefinition::GetFloatDef() && TargetType == UNiagaraComponentRendererProperties::GetDoubleDef())
+	{
+		const FNiagaraFloat Data = SourceVariable.GetValue<FNiagaraFloat>();
+		TargetData.AddUninitialized(sizeof(FNiagaraDouble));
+		reinterpret_cast<FNiagaraDouble*>(TargetData.GetData())->Value = Data.Value;
+	}
 }
 
 void InvokeSetterFunction(UObject* InRuntimeObject, UFunction* Setter, const uint8* InData, int32 DataSize, const TMap<FString, FString>& SetterDefaultValues)
@@ -707,6 +713,12 @@ void FNiagaraRendererComponents::TickPropertyBindings(
 						break;
 					}
 
+					if (Property->GetClass()->IsChildOf(FDoubleProperty::StaticClass()))
+					{
+						PropertyType = UNiagaraComponentRendererProperties::GetDoubleDef();
+						break;
+					}
+
 					// the property is missing the struct info and has a different size. In case of LWC types we try to guess the type if it fits the size.
 					bool bDoubleSize = Property->GetSize() == VarType.GetSize() * 2;
 					if (VarType == FNiagaraTypeDefinition::GetVec2Def() && bDoubleSize)
diff --git a/Engine/Plugins/FX/Niagara/Source/Niagara/Public/NiagaraComponentRendererProperties.h b/Engine/Plugins/FX/Niagara/Source/Niagara/Public/NiagaraComponentRendererProperties.h
--- a/Engine/Plugins/FX/Niagara/Source/Niagara/Public/NiagaraComponentRendererProperties.h
+++ b/Engine/Plugins/FX/Niagara/Source/Niagara/Public/NiagaraComponentRendererProperties.h
@@ -149,6 +149,7 @@ public:
 	static FNiagaraTypeDefinition GetFVector4Def();
 	static FNiagaraTypeDefinition GetFVector3fDef();
 	static FNiagaraTypeDefinition GetFQuatDef();
+	static FNiagaraTypeDefinition GetDoubleDef();
 
 	virtual void CacheFromCompiledData(const FNiagaraDataSetCompiledData* CompiledData) override;
 
diff --git a/Engine/Plugins/FX/Niagara/Source/Niagara/Public/NiagaraTypes.h b/Engine/Plugins/FX/Niagara/Source/Niagara/Public/NiagaraTypes.h
--- a/Engine/Plugins/FX/Niagara/Source/Niagara/Public/NiagaraTypes.h
+++ b/Engine/Plugins/FX/Niagara/Source/Niagara/Public/NiagaraTypes.h
@@ -161,6 +161,15 @@ struct FNiagaraParameterMap
 	GENERATED_USTRUCT_BODY()
 };
 
+// only used for LWC conversions, not supported by Niagara yet
+USTRUCT(meta = (DisplayName = "double"))
+struct FNiagaraDouble
+{
+	GENERATED_USTRUCT_BODY()
+
+	UPROPERTY(EditAnywhere, Category=Parameters)
+	double Value = 0;
+};
 
 USTRUCT(meta = (DisplayName = "Matrix"))
 struct FNiagaraMatrix
