- Remove NiagaraViewDataManager it is no longer required now that we are on RDG timeline

diff --git a/Engine/Plugins/FX/Niagara/Source/Niagara/Private/DataInterface/NiagaraDataInterfaceGBuffer.cpp b/Engine/Plugins/FX/Niagara/Source/Niagara/Private/DataInterface/NiagaraDataInterfaceGBuffer.cpp
--- a/Engine/Plugins/FX/Niagara/Source/Niagara/Private/DataInterface/NiagaraDataInterfaceGBuffer.cpp
+++ b/Engine/Plugins/FX/Niagara/Source/Niagara/Private/DataInterface/NiagaraDataInterfaceGBuffer.cpp
@@ -4,13 +4,14 @@
 #include "NiagaraGpuComputeDispatchInterface.h"
 #include "NiagaraGpuComputeDispatch.h"
 #include "NiagaraTypes.h"
-#include "NiagaraRenderViewDataManager.h"
 #include "NiagaraShaderParametersBuilder.h"
 #include "NiagaraSystemInstance.h"
 #include "NiagaraWorldManager.h"
 
 #include "Internationalization/Internationalization.h"
 #include "ShaderParameterUtils.h"
+#include "SceneRendering.h"
+#include "SceneTextures.h"
 
 #include UE_INLINE_GENERATED_CPP_BY_NAME(NiagaraDataInterfaceGBuffer)
 
@@ -179,11 +180,15 @@ void UNiagaraDataInterfaceGBuffer::SetShaderParameters(const FNiagaraDataInterfa
 	NiagaraDataInterfaceGBufferLocal::FShaderParameters* Parameters = Context.GetParameterNestedStruct<NiagaraDataInterfaceGBufferLocal::FShaderParameters>();
 	if (Context.IsResourceBound(&Parameters->VelocityTexture))
 	{
-		if (FNiagaraSceneTextureParameters* NiagaraSceneTextures = static_cast<const FNiagaraGpuComputeDispatch&>(Context.GetComputeDispatchInterface()).GetNiagaraSceneTextures())	//-BATCHERTODO:
+		TConstArrayView<FViewInfo> ViewInfos = Context.GetComputeDispatchInterface().GetSimulationViewInfos();
+		if (ViewInfos.Num() > 0)
 		{
-			VelocityTexture = NiagaraSceneTextures->Velocity.GetTexture();
+			if ( const FSceneTextures* SceneTextures = GetViewFamilyInfo(ViewInfos).GetSceneTexturesChecked() )
+			{
+				VelocityTexture = SceneTextures->Velocity;
+			}
 		}
-
+		
 		if (VelocityTexture == nullptr)
 		{
 			VelocityTexture = Context.GetComputeDispatchInterface().GetBlackTexture(Context.GetGraphBuilder(), ETextureDimension::Texture2D);
diff --git a/Engine/Plugins/FX/Niagara/Source/Niagara/Private/NiagaraGpuComputeDispatch.cpp b/Engine/Plugins/FX/Niagara/Source/Niagara/Private/NiagaraGpuComputeDispatch.cpp
--- a/Engine/Plugins/FX/Niagara/Source/Niagara/Private/NiagaraGpuComputeDispatch.cpp
+++ b/Engine/Plugins/FX/Niagara/Source/Niagara/Private/NiagaraGpuComputeDispatch.cpp
@@ -18,7 +18,6 @@
 #endif
 #include "NiagaraGPUProfilerInterface.h"
 #include "NiagaraGpuReadbackManager.h"
-#include "NiagaraRenderViewDataManager.h"
 #include "NiagaraRenderer.h"
 #include "NiagaraShader.h"
 #include "NiagaraShaderParticleID.h"
@@ -28,6 +27,7 @@
 #include "NiagaraWorldManager.h"
 #include "PipelineStateCache.h"
 #include "RHI.h"
+#include "SceneInterface.h"
 #include "ScenePrivate.h"
 #include "SceneRendering.h"
 #include "ShaderParameterUtils.h"
@@ -1047,11 +1047,28 @@ void FNiagaraGpuComputeDispatch::ExecuteTicks(FRDGBuilder& GraphBuilder, TConstA
 	RDG_RHI_EVENT_SCOPE(GraphBuilder, NiagaraGpuComputeDispatch);	//-TODO:RDG: Show TickStage
 
 	// Setup Parameters that can be read from data interfaces
-	FNiagaraSceneTextureParameters SceneTextures;
-	GNiagaraViewDataManager.GetSceneTextureParameters(GraphBuilder, Views.Num() > 0 ? GetViewFamilyInfo(Views).GetSceneTexturesChecked() : nullptr, SceneTextures);
-	NiagaraSceneTextures = &SceneTextures;
 	SimulationViewInfos = Views;
 
+	{
+		const FSceneTextures* SceneTextures = Views.Num() > 0 ? GetViewFamilyInfo(Views).GetSceneTexturesChecked() : nullptr;
+		if (FSceneInterface::GetShadingPath(FeatureLevel) == EShadingPath::Deferred)
+		{
+			SceneTexturesUniformParams = SceneTextures ? SceneTextures->UniformBuffer : nullptr;
+			if (SceneTexturesUniformParams == nullptr)
+			{
+				SceneTexturesUniformParams = CreateSceneTextureUniformBuffer(GraphBuilder, nullptr, FeatureLevel, ESceneTextureSetupMode::SceneVelocity);
+			}
+		}
+		else if (FSceneInterface::GetShadingPath(FeatureLevel) == EShadingPath::Mobile)
+		{
+			MobileSceneTexturesUniformParams = SceneTextures ? SceneTextures->MobileUniformBuffer : nullptr;
+			if (MobileSceneTexturesUniformParams == nullptr)
+			{
+				MobileSceneTexturesUniformParams = CreateMobileSceneTextureUniformBuffer(GraphBuilder, SceneTextures, EMobileSceneTextureSetupMode::None);
+			}
+		}
+	}
+
 	// Loop over dispatches
 	for ( const FNiagaraGpuDispatchGroup& DispatchGroup : DispatchList.DispatchGroups )
 	{
@@ -1364,8 +1381,9 @@ void FNiagaraGpuComputeDispatch::ExecuteTicks(FRDGBuilder& GraphBuilder, TConstA
 	);
 
 	// Tear down for tick pass
-	NiagaraSceneTextures = nullptr;
 	SimulationViewInfos = TConstArrayView<FViewInfo>();
+	SceneTexturesUniformParams = nullptr;
+	MobileSceneTexturesUniformParams = nullptr;
 
 	CurrentPassExternalAccessQueue.Submit(GraphBuilder);
 }
@@ -1558,8 +1576,8 @@ void FNiagaraGpuComputeDispatch::DispatchStage(FRDGBuilder& GraphBuilder, const
 	{
 		DispatchParameters->View = SimulationViewInfos[0].ViewUniformBuffer;
 	}
-	DispatchParameters->SceneTextures.SceneTextures			= NiagaraSceneTextures->SceneTextures;
-	DispatchParameters->SceneTextures.MobileSceneTextures	= NiagaraSceneTextures->MobileSceneTextures;
+	DispatchParameters->SceneTextures.SceneTextures			= SceneTexturesUniformParams;
+	DispatchParameters->SceneTextures.MobileSceneTextures	= MobileSceneTexturesUniformParams;
 
 	// Execute the dispatch
 	{
@@ -1667,7 +1685,6 @@ void FNiagaraGpuComputeDispatch::PreInitViews(FRDGBuilder& GraphBuilder, bool bA
 	RDG_CSV_STAT_EXCLUSIVE_SCOPE(GraphBuilder, Niagara);
 
 	bRequiresReadback = false;
-	GNiagaraViewDataManager.ClearSceneTextureParameters();
 #if WITH_EDITOR
 	bRaisedWarningThisFrame = false;
 #endif
@@ -1835,8 +1852,6 @@ void FNiagaraGpuComputeDispatch::PostRenderOpaque(FRDGBuilder& GraphBuilder, TCo
 	);
 	bRequiresReadback = false;
 
-	GNiagaraViewDataManager.ClearSceneTextureParameters();
-
 	if (FNiagaraGpuComputeDispatchLocal::CsvStatsEnabled())
 	{
 		CSV_CUSTOM_STAT(NiagaraGpuCompute, TotalDispatchesThisFrame, TotalDispatchesThisFrame, ECsvCustomStatOp::Set);
diff --git a/Engine/Plugins/FX/Niagara/Source/Niagara/Private/NiagaraGpuComputeDispatch.h b/Engine/Plugins/FX/Niagara/Source/Niagara/Private/NiagaraGpuComputeDispatch.h
--- a/Engine/Plugins/FX/Niagara/Source/Niagara/Private/NiagaraGpuComputeDispatch.h
+++ b/Engine/Plugins/FX/Niagara/Source/Niagara/Private/NiagaraGpuComputeDispatch.h
@@ -170,9 +170,6 @@ public:
 	virtual FNiagaraGPUProfilerInterface* GetGPUProfiler() const final { return GPUProfilerPtr.Get(); }
 #endif
 
-	/** When inside a pass returns the active scene textures for the pass, can be nullptr. */
-	class FNiagaraSceneTextureParameters* GetNiagaraSceneTextures() const { return NiagaraSceneTextures; }
-
 	/** Allows access to the current passes external access queue. */
 	FRDGExternalAccessQueue& GetCurrentPassExternalAccessQueue() const { return CurrentPassExternalAccessQueue; }
 
@@ -253,8 +250,6 @@ private:
 
 	FNiagaraGpuDispatchList DispatchListPerStage[ENiagaraGpuComputeTickStage::Max];
 
-	class FNiagaraSceneTextureParameters* NiagaraSceneTextures = nullptr;
-
 	struct FDebugReadbackInfo
 	{
 		FNiagaraSystemInstanceID InstanceID;
@@ -288,4 +283,7 @@ private:
 #if WITH_EDITOR
 	bool bRaisedWarningThisFrame = false;
 #endif
+
+	TRDGUniformBufferRef<FSceneTextureUniformParameters> SceneTexturesUniformParams = nullptr;
+	TRDGUniformBufferRef<FMobileSceneTextureUniformParameters> MobileSceneTexturesUniformParams = nullptr;
 };
diff --git a/Engine/Plugins/FX/Niagara/Source/Niagara/Private/NiagaraModule.cpp b/Engine/Plugins/FX/Niagara/Source/Niagara/Private/NiagaraModule.cpp
--- a/Engine/Plugins/FX/Niagara/Source/Niagara/Private/NiagaraModule.cpp
+++ b/Engine/Plugins/FX/Niagara/Source/Niagara/Private/NiagaraModule.cpp
@@ -13,7 +13,6 @@
 #include "UObject/ObjectRedirector.h"
 #include "UObject/Package.h"
 #include "NiagaraWorldManager.h"
-#include "NiagaraRenderViewDataManager.h"
 #include "VectorVM.h"
 #include "NiagaraConstants.h"
 #include "NiagaraLightRendererProperties.h"
@@ -217,8 +216,6 @@ void INiagaraModule::StartupModule()
 #endif
 	LLM_SCOPE(ELLMTag::Niagara);
 	FNiagaraTypeDefinition::Init();
-	FNiagaraRenderViewDataManager::Init();
-
 	
 #if PLATFORM_WINDOWS
 	// Global registration of  the vdb types.
@@ -486,8 +483,6 @@ void INiagaraModule::OnWorldBeginTearDown(UWorld* World)
 void INiagaraModule::ShutdownRenderingResources()
 {
 	FFXSystemInterface::UnregisterCustomFXSystem(FNiagaraGpuComputeDispatch::Name);
-
-	FNiagaraRenderViewDataManager::Shutdown();
 }
 
 void INiagaraModule::ShutdownModule()
diff --git a/Engine/Plugins/FX/Niagara/Source/Niagara/Private/NiagaraRenderViewDataManager.cpp b/Engine/Plugins/FX/Niagara/Source/Niagara/Private/NiagaraRenderViewDataManager.cpp
deleted file mode 100644
diff --git a/Engine/Plugins/FX/Niagara/Source/Niagara/Public/NiagaraRenderViewDataManager.h b/Engine/Plugins/FX/Niagara/Source/Niagara/Public/NiagaraRenderViewDataManager.h
deleted file mode 100644
