- Add method to destroy instance but avoid destroying the component
- This is used when calling SetAsset / Reinitializing the system
- What was occuring before this change is auto destroy / pooled components would destroy both the instance and the component, leading to missing FX or a crash when creating datainterface user parameters as the owner is nullptr

diff --git a/Engine/Plugins/FX/Niagara/Source/Niagara/Private/NiagaraComponent.cpp b/Engine/Plugins/FX/Niagara/Source/Niagara/Private/NiagaraComponent.cpp
--- a/Engine/Plugins/FX/Niagara/Source/Niagara/Private/NiagaraComponent.cpp
+++ b/Engine/Plugins/FX/Niagara/Source/Niagara/Private/NiagaraComponent.cpp
@@ -870,11 +870,7 @@ void UNiagaraComponent::ResetSystem()
 
 void UNiagaraComponent::ReinitializeSystem()
 {
-	const bool bCachedAutoDestroy = bAutoDestroy;
-	bAutoDestroy = false;
-	DestroyInstance();
-	bAutoDestroy = bCachedAutoDestroy;
-
+	DestroyInstanceNotComponent();
 	Activate(true);
 }
 
@@ -1589,6 +1585,7 @@ void UNiagaraComponent::DestroyInstance()
 {
 	//UE_LOG(LogNiagara, Log, TEXT("UNiagaraComponent::DestroyInstance: %p - %p  %s\n"), this, SystemInstance.Get(), *GetAsset()->GetFullName());
 	//UE_LOG(LogNiagara, Log, TEXT("DestroyInstance: %p - %s"), this, *Asset->GetName());
+
 	SetActiveFlag(false);
 	UnregisterWithScalabilityManager();
 
@@ -1609,6 +1606,19 @@ void UNiagaraComponent::DestroyInstance()
 	MarkRenderStateDirty();
 }
 
+void UNiagaraComponent::DestroyInstanceNotComponent()
+{
+	const bool bCachedAutoDestroy = bAutoDestroy;
+	const ENCPoolMethod CachedPoolMethod = PoolingMethod;
+	bAutoDestroy = false;
+	PoolingMethod = ENCPoolMethod::None;
+
+	DestroyInstance();
+
+	bAutoDestroy = bCachedAutoDestroy;
+	PoolingMethod = CachedPoolMethod;
+}
+
 void UNiagaraComponent::OnPooledReuse(UWorld* NewWorld)
 {
 	check(IsValid(this));
@@ -3154,7 +3164,7 @@ void UNiagaraComponent::SynchronizeWithSourceSystem()
 	// Synchronizing parameters will create new data interface objects and if the old data
 	// interface objects are currently being used by a simulation they may be destroyed due to garbage
 	// collection, so preemptively kill the instance here.
-	DestroyInstance();
+	DestroyInstanceNotComponent();
 
 	//TODO: Look through params in system in "Owner" namespace and add to our parameters.
 	if (Asset == nullptr)
@@ -3737,7 +3747,7 @@ void UNiagaraComponent::SetAsset(UNiagaraSystem* InAsset, bool bResetExistingOve
 
 	const bool bWasActive = SystemInstanceController && SystemInstanceController->GetRequestedExecutionState() == ENiagaraExecutionState::Active;
 
-	DestroyInstance();
+	DestroyInstanceNotComponent();
 
 	// Set new asset, update parameters and reactivate it it was already active
 	Asset = InAsset;
diff --git a/Engine/Plugins/FX/Niagara/Source/Niagara/Public/NiagaraComponent.h b/Engine/Plugins/FX/Niagara/Source/Niagara/Public/NiagaraComponent.h
--- a/Engine/Plugins/FX/Niagara/Source/Niagara/Public/NiagaraComponent.h
+++ b/Engine/Plugins/FX/Niagara/Source/Niagara/Public/NiagaraComponent.h
@@ -256,6 +256,7 @@ public:
 
 	bool InitializeSystem();
 	void DestroyInstance();
+	void DestroyInstanceNotComponent();
 
 	void OnPooledReuse(UWorld* NewWorld);
 
