- Move functions that shouldn't be in cooked behind editor define
- Fixes an ensure that can occur on cooked builds due to calling editor only function

#rnx

diff --git a/Engine/Plugins/FX/Niagara/Source/Niagara/Classes/NiagaraVersionedObject.h b/Engine/Plugins/FX/Niagara/Source/Niagara/Classes/NiagaraVersionedObject.h
--- a/Engine/Plugins/FX/Niagara/Source/Niagara/Classes/NiagaraVersionedObject.h
+++ b/Engine/Plugins/FX/Niagara/Source/Niagara/Classes/NiagaraVersionedObject.h
@@ -23,14 +23,15 @@ class NIAGARA_API FNiagaraVersionedObject
 public:
 	virtual ~FNiagaraVersionedObject() = default;
 
-	virtual TSharedPtr<FNiagaraVersionDataAccessor> GetVersionDataAccessor(const FGuid& Version); 
+	/** Returns all available versions for this object. */
+	virtual TArray<FNiagaraAssetVersion> GetAllAvailableVersions() const;
+
+#if WITH_EDITORONLY_DATA
+	virtual TSharedPtr<FNiagaraVersionDataAccessor> GetVersionDataAccessor(const FGuid& Version);
 
 	/** If true then this script asset uses active version control to track changes. */
 	virtual bool IsVersioningEnabled() const;
 	
-	/** Returns all available versions for this object. */
-	virtual TArray<FNiagaraAssetVersion> GetAllAvailableVersions() const;
-
 	/** Returns the version of the exposed version data (i.e. the version used when adding a module to the stack) */
 	virtual FNiagaraAssetVersion GetExposedVersion() const;
 
@@ -51,4 +52,5 @@ public:
 
 	/** Disables versioning and keeps only the data from the given version guid. Note that this breaks ALL references from existing assets and should only be used when creating a copy of an object, as the effect is very destructive.  */
 	virtual void DisableVersioning(const FGuid& VersionGuidToUse);
+#endif //WITH_EDITORONLY_DATA
 };
diff --git a/Engine/Plugins/FX/Niagara/Source/Niagara/Private/NiagaraEmitter.cpp b/Engine/Plugins/FX/Niagara/Source/Niagara/Private/NiagaraEmitter.cpp
--- a/Engine/Plugins/FX/Niagara/Source/Niagara/Private/NiagaraEmitter.cpp
+++ b/Engine/Plugins/FX/Niagara/Source/Niagara/Private/NiagaraEmitter.cpp
@@ -1633,6 +1633,7 @@ void FVersionedNiagaraEmitterData::UpdateDebugName(const UNiagaraEmitter& Emitte
 		DebugSimName = SystemOwner->GetName();
 		DebugSimName.AppendChar(':');
 	}
+#if WITH_EDITORONLY_DATA
 	DebugSimName.Append(Emitter.GetName());
 	if (Emitter.IsVersioningEnabled())
 	{
@@ -1641,6 +1642,7 @@ void FVersionedNiagaraEmitterData::UpdateDebugName(const UNiagaraEmitter& Emitte
 		DebugSimName.AppendChar('.');
 		DebugSimName.AppendInt(Version.MinorVersion);
 	}
+#endif
 #endif
 
 	RebuildRendererBindings(Emitter);
diff --git a/Engine/Plugins/FX/Niagara/Source/Niagara/Private/NiagaraVersionedObject.cpp b/Engine/Plugins/FX/Niagara/Source/Niagara/Private/NiagaraVersionedObject.cpp
--- a/Engine/Plugins/FX/Niagara/Source/Niagara/Private/NiagaraVersionedObject.cpp
+++ b/Engine/Plugins/FX/Niagara/Source/Niagara/Private/NiagaraVersionedObject.cpp
@@ -70,22 +70,23 @@ FFilePath& FNiagaraVersionDataAccessor::GetScriptAsset()
 	return NiagaraVersionedObject::Dummy.GetScriptAsset();
 }
 
-TSharedPtr<FNiagaraVersionDataAccessor> FNiagaraVersionedObject::GetVersionDataAccessor(const FGuid& Version)
+TArray<FNiagaraAssetVersion> FNiagaraVersionedObject::GetAllAvailableVersions() const
 {
 	ensure(false);
-	return TSharedPtr<FNiagaraVersionDataAccessor>();
+	return TArray<FNiagaraAssetVersion>();
 }
 
-bool FNiagaraVersionedObject::IsVersioningEnabled() const
+#if WITH_EDITORONLY_DATA
+TSharedPtr<FNiagaraVersionDataAccessor> FNiagaraVersionedObject::GetVersionDataAccessor(const FGuid& Version)
 {
 	ensure(false);
-	return false;
+	return TSharedPtr<FNiagaraVersionDataAccessor>();
 }
 
-TArray<FNiagaraAssetVersion> FNiagaraVersionedObject::GetAllAvailableVersions() const
+bool FNiagaraVersionedObject::IsVersioningEnabled() const
 {
 	ensure(false);
-	return TArray<FNiagaraAssetVersion>();
+	return false;
 }
 
 FNiagaraAssetVersion FNiagaraVersionedObject::GetExposedVersion() const
@@ -125,3 +126,4 @@ void FNiagaraVersionedObject::DisableVersioning(const FGuid& VersionGuidToUse)
 {
 	ensure(false);
 }
+#endif //WITH_EDITORONLY_DATA
