Fixed a bug that caused an ensure to fire when attempting to ctrl+drag connected pins in script graphs
Note: this might cause subsequent ui bugs that need to be fixed.

diff --git a/Engine/Plugins/FX/Niagara/Source/NiagaraEditor/Private/NiagaraGraph.cpp b/Engine/Plugins/FX/Niagara/Source/NiagaraEditor/Private/NiagaraGraph.cpp
--- a/Engine/Plugins/FX/Niagara/Source/NiagaraEditor/Private/NiagaraGraph.cpp
+++ b/Engine/Plugins/FX/Niagara/Source/NiagaraEditor/Private/NiagaraGraph.cpp
@@ -3616,7 +3616,6 @@ void UNiagaraGraph::MarkGraphRequiresSynchronization(FString Reason)
 {
 	Modify();
 	ChangeId = FGuid::NewGuid();
-	NotifyGraphChanged();
 	if (GEnableVerboseNiagaraChangeIdLogging)
 	{
 		UE_LOG(LogNiagaraEditor, Verbose, TEXT("Graph %s was marked requires synchronization.  Reason: %s"), *GetPathName(), *Reason);
