Unshelved from pending changelist '23096088':

Fixed the system's root stack entry not being found when the selection view model is being refreshed. This resulted in deselection of the system node when changing a system's value (including in scalability mode)

diff --git a/Engine/Plugins/FX/Niagara/Source/NiagaraEditor/Private/ViewModels/NiagaraSystemSelectionViewModel.cpp b/Engine/Plugins/FX/Niagara/Source/NiagaraEditor/Private/ViewModels/NiagaraSystemSelectionViewModel.cpp
--- a/Engine/Plugins/FX/Niagara/Source/NiagaraEditor/Private/ViewModels/NiagaraSystemSelectionViewModel.cpp
+++ b/Engine/Plugins/FX/Niagara/Source/NiagaraEditor/Private/ViewModels/NiagaraSystemSelectionViewModel.cpp
@@ -4,6 +4,7 @@
 #include "ViewModels/Stack/NiagaraStackEntry.h"
 #include "ViewModels/Stack/NiagaraStackItemGroup.h"
 #include "ViewModels/Stack/NiagaraStackItem.h"
+#include "ViewModels/Stack/NiagaraStackRoot.h"
 #include "ViewModels/Stack/NiagaraStackSelection.h"
 #include "ViewModels/Stack/NiagaraStackViewModel.h"
 #include "ViewModels/NiagaraSystemViewModel.h"
@@ -238,7 +239,7 @@ void CollectGroupAndItemEntries(TSharedRef<FNiagaraSystemViewModel> SystemViewMo
 			UNiagaraStackEntry* Entry = EntriesToCheck[EntriesToCheck.Num() - 1];
 			EntriesToCheck.RemoveAt(EntriesToCheck.Num() - 1);
 
-			if (Entry->IsA<UNiagaraStackItemGroup>() || Entry->IsA<UNiagaraStackItem>())
+			if (Entry->IsA<UNiagaraStackItemGroup>() || Entry->IsA<UNiagaraStackItem>() || Entry->IsA<UNiagaraStackRoot>())
 			{
 				EntryKeyToEntryMap.Add(Entry->GetStackEditorDataKey(), Entry);
 			}
