Curve Editor: Added missing curve color tweaks to some additional Niagara curves.

#rnx

diff --git a/Engine/Plugins/FX/Niagara/Source/Niagara/Private/NiagaraDataInterfaceVector2DCurve.cpp b/Engine/Plugins/FX/Niagara/Source/Niagara/Private/NiagaraDataInterfaceVector2DCurve.cpp
--- a/Engine/Plugins/FX/Niagara/Source/Niagara/Private/NiagaraDataInterfaceVector2DCurve.cpp
+++ b/Engine/Plugins/FX/Niagara/Source/Niagara/Private/NiagaraDataInterfaceVector2DCurve.cpp
@@ -138,8 +138,8 @@ bool UNiagaraDataInterfaceVector2DCurve::Equals(const UNiagaraDataInterface* Oth
 }
 void UNiagaraDataInterfaceVector2DCurve::GetCurveData(TArray<FCurveData>& OutCurveData)
 {
-	OutCurveData.Add(FCurveData(&XCurve, TEXT("X"), FLinearColor::Red));
-	OutCurveData.Add(FCurveData(&YCurve, TEXT("Y"), FLinearColor::Green));
+	OutCurveData.Add(FCurveData(&XCurve, TEXT("X"), FLinearColor(1.0f, 0.05f, 0.05f)));
+	OutCurveData.Add(FCurveData(&YCurve, TEXT("Y"), FLinearColor(0.05f, 1.0f, 0.05f)));
 }
 
 void UNiagaraDataInterfaceVector2DCurve::GetFunctions(TArray<FNiagaraFunctionSignature>& OutFunctions)
diff --git a/Engine/Plugins/FX/Niagara/Source/Niagara/Private/NiagaraDataInterfaceVector4Curve.cpp b/Engine/Plugins/FX/Niagara/Source/Niagara/Private/NiagaraDataInterfaceVector4Curve.cpp
--- a/Engine/Plugins/FX/Niagara/Source/Niagara/Private/NiagaraDataInterfaceVector4Curve.cpp
+++ b/Engine/Plugins/FX/Niagara/Source/Niagara/Private/NiagaraDataInterfaceVector4Curve.cpp
@@ -157,10 +157,10 @@ bool UNiagaraDataInterfaceVector4Curve::Equals(const UNiagaraDataInterface* Othe
 
 void UNiagaraDataInterfaceVector4Curve::GetCurveData(TArray<FCurveData>& OutCurveData)
 {
-	OutCurveData.Add(FCurveData(&XCurve, TEXT("X"), FLinearColor::Red));
-	OutCurveData.Add(FCurveData(&YCurve, TEXT("Y"), FLinearColor::Green));
-	OutCurveData.Add(FCurveData(&ZCurve, TEXT("Z"), FLinearColor::Blue));
-	OutCurveData.Add(FCurveData(&WCurve, TEXT("W"), FLinearColor::White));
+	OutCurveData.Add(FCurveData(&XCurve, TEXT("X"), FLinearColor(1.0f, 0.05f, 0.05f)));
+	OutCurveData.Add(FCurveData(&YCurve, TEXT("Y"), FLinearColor(0.05f, 1.0f, 0.05f)));
+	OutCurveData.Add(FCurveData(&ZCurve, TEXT("Z"), FLinearColor(0.1f, 0.2f, 1.0f)));
+	OutCurveData.Add(FCurveData(&WCurve, TEXT("W"), FLinearColor(0.2f, 0.2f, 0.2f)));
 }
 
 void UNiagaraDataInterfaceVector4Curve::GetFunctions(TArray<FNiagaraFunctionSignature>& OutFunctions)
diff --git a/Engine/Plugins/FX/Niagara/Source/Niagara/Private/NiagaraDataInterfaceVectorCurve.cpp b/Engine/Plugins/FX/Niagara/Source/Niagara/Private/NiagaraDataInterfaceVectorCurve.cpp
--- a/Engine/Plugins/FX/Niagara/Source/Niagara/Private/NiagaraDataInterfaceVectorCurve.cpp
+++ b/Engine/Plugins/FX/Niagara/Source/Niagara/Private/NiagaraDataInterfaceVectorCurve.cpp
@@ -149,9 +149,9 @@ bool UNiagaraDataInterfaceVectorCurve::Equals(const UNiagaraDataInterface* Other
 
 void UNiagaraDataInterfaceVectorCurve::GetCurveData(TArray<FCurveData>& OutCurveData)
 {
-	OutCurveData.Add(FCurveData(&XCurve, TEXT("X"), FLinearColor::Red));
-	OutCurveData.Add(FCurveData(&YCurve, TEXT("Y"), FLinearColor::Green));
-	OutCurveData.Add(FCurveData(&ZCurve, TEXT("Z"), FLinearColor::Blue));
+	OutCurveData.Add(FCurveData(&XCurve, TEXT("X"), FLinearColor(1.0f, 0.05f, 0.05f)));
+	OutCurveData.Add(FCurveData(&YCurve, TEXT("Y"), FLinearColor(0.05f, 1.0f, 0.05f)));
+	OutCurveData.Add(FCurveData(&ZCurve, TEXT("Z"), FLinearColor(0.1f, 0.2f, 1.0f)));
 }
 
 void UNiagaraDataInterfaceVectorCurve::GetFunctions(TArray<FNiagaraFunctionSignature>& OutFunctions)
