- Remove legacy data interface binding code

#rnx

diff --git a/Engine/Plugins/FX/Niagara/Source/Niagara/Private/NiagaraSystemInstance.cpp b/Engine/Plugins/FX/Niagara/Source/Niagara/Private/NiagaraSystemInstance.cpp
--- a/Engine/Plugins/FX/Niagara/Source/Niagara/Private/NiagaraSystemInstance.cpp
+++ b/Engine/Plugins/FX/Niagara/Source/Niagara/Private/NiagaraSystemInstance.cpp
@@ -99,32 +99,6 @@ static FAutoConsoleVariableRef CVarNiagaraAllowDeferredReset(
 	ECVF_Default
 );
 
-bool DoSystemDataInterfacesRequireSolo(const UNiagaraSystem& System, const FNiagaraUserRedirectionParameterStore* OverrideParameters)
-{
-	if (FNiagaraSystemSimulation::UseLegacySystemSimulationContexts())
-	{
-		if (System.HasSystemScriptDIsWithPerInstanceData())
-		{
-			return true;
-		}
-
-		const TArray<FName>& UserDINamesReadInSystemScripts = System.GetUserDINamesReadInSystemScripts();
-		if (OverrideParameters != nullptr && UserDINamesReadInSystemScripts.Num() > 0)
-		{
-			TArray<FNiagaraVariable> OverrideParameterVariables;
-			OverrideParameters->GetParameters(OverrideParameterVariables);
-			for (const FNiagaraVariable& OverrideParameterVariable : OverrideParameterVariables)
-			{
-				if (OverrideParameterVariable.IsDataInterface() && UserDINamesReadInSystemScripts.Contains(OverrideParameterVariable.GetName()))
-				{
-					return true;
-				}
-			}
-		}
-	}
-	return false;
-}
-
 FNiagaraSystemInstance::FNiagaraSystemInstance(UWorld& InWorld, UNiagaraSystem& InAsset, FNiagaraUserRedirectionParameterStore* InOverrideParameters,
                                                USceneComponent* InAttachComponent, ENiagaraTickBehavior InTickBehavior, bool bInPooled)
 	: SystemInstanceIndex(INDEX_NONE)
@@ -527,8 +501,6 @@ void FNiagaraSystemInstance::SetSolo(bool bInSolo)
 		}
 		else
 		{
-			ensureMsgf(DoSystemDataInterfacesRequireSolo(*GetSystem(), OverrideParameters) == false, TEXT("Disabling solo mode but data interfaces require it"));
-
 			const ETickingGroup TickGroup = CalculateTickGroup();
 			TSharedPtr<FNiagaraSystemSimulation, ESPMode::ThreadSafe> NewSim = GetWorldManager()->GetSystemSimulation(TickGroup, System);
 
@@ -1108,7 +1080,6 @@ void FNiagaraSystemInstance::ReInitInternal()
 		return;
 	}
 
-	bSolo |= DoSystemDataInterfacesRequireSolo(*System, OverrideParameters);
 	if (bSolo)
 	{
 		if (!SystemSimulation.IsValid())
@@ -1236,14 +1207,6 @@ void FNiagaraSystemInstance::BindParameters()
 			// NOTE: We don't rebind if it's already bound to improve reset times.
 			OverrideParameters->Bind(&InstanceParameters);
 		}
-
-		if (SystemSimulation->GetIsSolo() && FNiagaraSystemSimulation::UseLegacySystemSimulationContexts())
-		{
-			// If this simulation is solo than we can bind the instance parameters to the system simulation contexts so that
-			// the system and emitter scripts use the per-instance data interfaces.
-			OverrideParameters->Bind(&SystemSimulation->GetSpawnExecutionContext()->Parameters);
-			OverrideParameters->Bind(&SystemSimulation->GetUpdateExecutionContext()->Parameters);
-		}
 	}
 	else
 	{
@@ -1515,20 +1478,8 @@ void FNiagaraSystemInstance::InitDataInterfaces()
 	};
 
 	CalcInstDataSize(InstanceParameters, false, false, false);//This probably should be a proper exec context.
-
-	if (SystemSimulation->GetIsSolo() && FNiagaraSystemSimulation::UseLegacySystemSimulationContexts())
-	{
-		CalcInstDataSize(SystemSimulation->GetSpawnExecutionContext()->Parameters, true, false, false);
-		SystemSimulation->GetSpawnExecutionContext()->DirtyDataInterfaces();
-
-		CalcInstDataSize(SystemSimulation->GetUpdateExecutionContext()->Parameters, true, false, false);
-		SystemSimulation->GetUpdateExecutionContext()->DirtyDataInterfaces();
-	}
-	else
-	{
-		CalcInstDataSize(SystemSimulation->GetSpawnExecutionContext()->Parameters, true, false, true);
-		CalcInstDataSize(SystemSimulation->GetUpdateExecutionContext()->Parameters, true, false, true);
-	}
+	CalcInstDataSize(SystemSimulation->GetSpawnExecutionContext()->Parameters, true, false, true);
+	CalcInstDataSize(SystemSimulation->GetUpdateExecutionContext()->Parameters, true, false, true);
 
 	//Iterate over interfaces to get size for table and clear their interface bindings.
 	for (const TSharedRef<FNiagaraEmitterInstance, ESPMode::ThreadSafe>& Simulation : Emitters)
@@ -1620,7 +1571,6 @@ void FNiagaraSystemInstance::InitDataInterfaces()
 	//We have valid DI instance data so now generate the table of function calls.
 	//When using the new exec contexts, each system instance builds it's own tables of DI function bindings for DI calls that require it.
 	//i.e. User DIs or those with per instance data that are called from system scripts.
-	if (FNiagaraSystemSimulation::UseLegacySystemSimulationContexts() == false)
 	{
 		bool bSuccess = true;
 		bSuccess &= SystemSimulation->GetSpawnExecutionContext()->GeneratePerInstanceDIFunctionTable(this, PerInstanceDIFunctions[(int32)ENiagaraSystemSimulationScript::Spawn]);
@@ -1713,7 +1663,6 @@ void FNiagaraSystemInstance::TickDataInterfaces(float DeltaSeconds, bool bPostSi
 		}
 
 		// Rebind funcs for system scripts
-		if (FNiagaraSystemSimulation::UseLegacySystemSimulationContexts() == false)
 		{
 			PerInstanceDIFunctions[(int32)ENiagaraSystemSimulationScript::Spawn].Reset();
 			PerInstanceDIFunctions[(int32)ENiagaraSystemSimulationScript::Update].Reset();
@@ -1729,11 +1678,6 @@ void FNiagaraSystemInstance::TickDataInterfaces(float DeltaSeconds, bool bPostSi
 				Complete(true);
 			}
 		}
-		else
-		{
-			SystemSimulation->GetSpawnExecutionContext()->DirtyDataInterfaces();
-			SystemSimulation->GetUpdateExecutionContext()->DirtyDataInterfaces();
-		}
 	}
 }
 
@@ -2826,10 +2770,9 @@ void FNiagaraSystemInstance::SetForceSolo(bool bInForceSolo)
 		return;
 	}
 
-	const bool bNewSolo = bForceSolo || DoSystemDataInterfacesRequireSolo(*GetSystem(), OverrideParameters);
-	if (bNewSolo != bSolo)
+	if (bForceSolo != bSolo)
 	{
-		SetSolo(bNewSolo);
+		SetSolo(bForceSolo);
 	}
 }
 
diff --git a/Engine/Plugins/FX/Niagara/Source/Niagara/Private/NiagaraSystemSimulation.cpp b/Engine/Plugins/FX/Niagara/Source/Niagara/Private/NiagaraSystemSimulation.cpp
--- a/Engine/Plugins/FX/Niagara/Source/Niagara/Private/NiagaraSystemSimulation.cpp
+++ b/Engine/Plugins/FX/Niagara/Source/Niagara/Private/NiagaraSystemSimulation.cpp
@@ -681,21 +681,10 @@ bool FNiagaraSystemSimulation::Init(UNiagaraSystem* InSystem, UWorld* InWorld, b
 
 		{
 			//SCOPE_CYCLE_COUNTER(STAT_NiagaraSystemSim_Init_ExecContexts);
-
-			if (UseLegacySystemSimulationContexts())
-			{
-				SpawnExecContext = MakeUnique<FNiagaraScriptExecutionContext>();
-				UpdateExecContext = MakeUnique<FNiagaraScriptExecutionContext>();
-				bCanExecute &= SpawnExecContext->Init(SpawnScript, ENiagaraSimTarget::CPUSim);
-				bCanExecute &= UpdateExecContext->Init(UpdateScript, ENiagaraSimTarget::CPUSim);
-			}
-			else
-			{
-				SpawnExecContext = MakeUnique<FNiagaraSystemScriptExecutionContext>(ENiagaraSystemSimulationScript::Spawn);
-				UpdateExecContext = MakeUnique<FNiagaraSystemScriptExecutionContext>(ENiagaraSystemSimulationScript::Update);
-				bCanExecute &= SpawnExecContext->Init(SpawnScript, ENiagaraSimTarget::CPUSim);
-				bCanExecute &= UpdateExecContext->Init(UpdateScript, ENiagaraSimTarget::CPUSim);
-			}
+			SpawnExecContext = MakeUnique<FNiagaraSystemScriptExecutionContext>(ENiagaraSystemSimulationScript::Spawn);
+			UpdateExecContext = MakeUnique<FNiagaraSystemScriptExecutionContext>(ENiagaraSystemSimulationScript::Update);
+			bCanExecute &= SpawnExecContext->Init(SpawnScript, ENiagaraSimTarget::CPUSim);
+			bCanExecute &= UpdateExecContext->Init(UpdateScript, ENiagaraSimTarget::CPUSim);
 		}
 
 		{
@@ -741,15 +730,6 @@ bool FNiagaraSystemSimulation::Init(UNiagaraSystem* InSystem, UWorld* InWorld, b
 
 			SpawnScript->RapidIterationParameters.Bind(&SpawnExecContext->Parameters);
 			UpdateScript->RapidIterationParameters.Bind(&UpdateExecContext->Parameters);
-
-			// If this simulation is not solo than we have bind the source system parameters to the system simulation contexts so that
-			// the system and emitter scripts use the default shared data interfaces.
-			if (UseLegacySystemSimulationContexts() && !bIsSolo)
-			{
-				FNiagaraUserRedirectionParameterStore& ExposedParameters = System->GetExposedParameters();
-				ExposedParameters.Bind(&SpawnExecContext->Parameters);
-				ExposedParameters.Bind(&UpdateExecContext->Parameters);
-			}
 		}
 
 		{
@@ -2397,52 +2377,3 @@ ENiagaraGPUTickHandlingMode FNiagaraSystemSimulation::GetGPUTickHandlingMode()co
 
 	return ENiagaraGPUTickHandlingMode::None;
 }
-
-static int32 GbNiagaraUseLegacySystemSimContexts = 0;
-static FAutoConsoleVariableRef CVarNiagaraUseLevgacySystemSimContexts(
-	TEXT("fx.Niagara.UseLegacySystemSimContexts"),
-	GbNiagaraUseLegacySystemSimContexts,
-	TEXT("If > 0, Niagara will use legacy system simulation contexts which would force the whole simulation solo if there were per instance DI calls in the system scripts. \n"),
-	FConsoleVariableDelegate::CreateStatic(&FNiagaraSystemSimulation::OnChanged_UseLegacySystemSimulationContexts),
-	ECVF_Default
-);
-
-bool FNiagaraSystemSimulation::bUseLegacyExecContexts = GbNiagaraUseLegacySystemSimContexts != 0;
-bool FNiagaraSystemSimulation::UseLegacySystemSimulationContexts()
-{
-	return bUseLegacyExecContexts;
-}
-
-void FNiagaraSystemSimulation::OnChanged_UseLegacySystemSimulationContexts(IConsoleVariable* CVar)
-{
-	bool bNewValue = GbNiagaraUseLegacySystemSimContexts != 0;
-	if( bUseLegacyExecContexts != bNewValue)
-	{
-		//To change at runtime we have to reinit all systems so they have the correct per instance DI bindings.
-		FNiagaraSystemUpdateContext UpdateContext;
-		UpdateContext.SetDestroyOnAdd(true);
-		UpdateContext.SetOnlyActive(true);
-		UpdateContext.AddAll(true);
-
-		//Just to be sure there's no lingering state, clear out the pools.
-		//TODO: Moveinto the update context itself?
-		FNiagaraWorldManager::ForAllWorldManagers(
-			[](FNiagaraWorldManager& WorldMan)
-			{
-				WorldMan.GetComponentPool()->Cleanup(nullptr);
-			}
-		);
-
-		//Reactivate any FX that were active.
-		bUseLegacyExecContexts = bNewValue;
-		UpdateContext.CommitUpdate();
-
-		//Re-prime the pools.
-		FNiagaraWorldManager::ForAllWorldManagers(
-			[](FNiagaraWorldManager& WorldMan)
-			{
-				WorldMan.PrimePoolForAllSystems();
-			}
-		);
-	}
-}
\ No newline at end of file
diff --git a/Engine/Plugins/FX/Niagara/Source/Niagara/Public/NiagaraSystemSimulation.h b/Engine/Plugins/FX/Niagara/Source/Niagara/Public/NiagaraSystemSimulation.h
--- a/Engine/Plugins/FX/Niagara/Source/Niagara/Public/NiagaraSystemSimulation.h
+++ b/Engine/Plugins/FX/Niagara/Source/Niagara/Public/NiagaraSystemSimulation.h
@@ -315,10 +315,6 @@ public:
 
 	ENiagaraGPUTickHandlingMode GetGPUTickHandlingMode() const;
 
-	/** If true we use legacy simulation contexts that could not handle per instance DI calls in the system scripts and would force the whole simulation solo. */
-	static bool UseLegacySystemSimulationContexts();
-	static void OnChanged_UseLegacySystemSimulationContexts(class IConsoleVariable* CVar);
-
 	FORCEINLINE UWorld* GetWorld()const{return World;}
 
 protected:
