Niagara: Handle private includes from public headers so the path isn't rooted in the module parent. This should be cleaned up to not require including private headers in public files.

#rnx

diff --git a/Engine/Plugins/FX/Niagara/Source/Niagara/Classes/NiagaraDataInterfaceArrayImpl.h b/Engine/Plugins/FX/Niagara/Source/Niagara/Classes/NiagaraDataInterfaceArrayImpl.h
--- a/Engine/Plugins/FX/Niagara/Source/Niagara/Classes/NiagaraDataInterfaceArrayImpl.h
+++ b/Engine/Plugins/FX/Niagara/Source/Niagara/Classes/NiagaraDataInterfaceArrayImpl.h
@@ -7,8 +7,8 @@
 #include "NiagaraDataInterfaceArray.h"
 #include "NiagaraDataInterfaceUtilities.h"
 #include "NiagaraGpuComputeDispatchInterface.h"
-#include "Niagara/Private/NiagaraGpuReadbackManager.h"
-#include "Niagara/Private/NiagaraStats.h"
+#include "NiagaraGpuReadbackManager.h"
+#include "NiagaraStats.h"
 #include "NiagaraSystemInstance.h"
 
 #include "Async/Async.h"
diff --git a/Engine/Plugins/FX/Niagara/Source/Niagara/Classes/NiagaraDataInterfaceCurveBase.h b/Engine/Plugins/FX/Niagara/Source/Niagara/Classes/NiagaraDataInterfaceCurveBase.h
--- a/Engine/Plugins/FX/Niagara/Source/Niagara/Classes/NiagaraDataInterfaceCurveBase.h
+++ b/Engine/Plugins/FX/Niagara/Source/Niagara/Classes/NiagaraDataInterfaceCurveBase.h
@@ -6,8 +6,8 @@
 #include "VectorVM.h"
 #include "StaticMeshResources.h"
 #include "Curves/RichCurve.h"
-#include "Niagara/Private/NiagaraStats.h"
 #include "NiagaraDataInterface.h"
+#include "NiagaraStats.h"
 #include "NiagaraDataInterfaceCurveBase.generated.h"
 
 /** Base class for curve data proxy data. */
diff --git a/Engine/Plugins/FX/Niagara/Source/Niagara/Classes/NiagaraDataInterfaceGrid2DCollection.h b/Engine/Plugins/FX/Niagara/Source/Niagara/Classes/NiagaraDataInterfaceGrid2DCollection.h
--- a/Engine/Plugins/FX/Niagara/Source/Niagara/Classes/NiagaraDataInterfaceGrid2DCollection.h
+++ b/Engine/Plugins/FX/Niagara/Source/Niagara/Classes/NiagaraDataInterfaceGrid2DCollection.h
@@ -6,7 +6,7 @@
 #include "ClearQuad.h"
 #include "NiagaraComponent.h"
 #include "NiagaraRenderGraphUtils.h"
-#include "Niagara/Private/NiagaraStats.h"
+#include "NiagaraStats.h"
 
 #include "NiagaraDataInterfaceGrid2DCollection.generated.h"
 
diff --git a/Engine/Plugins/FX/Niagara/Source/Niagara/Classes/NiagaraDataInterfaceGrid2DCollectionReader.h b/Engine/Plugins/FX/Niagara/Source/Niagara/Classes/NiagaraDataInterfaceGrid2DCollectionReader.h
--- a/Engine/Plugins/FX/Niagara/Source/Niagara/Classes/NiagaraDataInterfaceGrid2DCollectionReader.h
+++ b/Engine/Plugins/FX/Niagara/Source/Niagara/Classes/NiagaraDataInterfaceGrid2DCollectionReader.h
@@ -5,8 +5,8 @@
 #include "NiagaraDataInterfaceRW.h"
 #include "ClearQuad.h"
 #include "NiagaraComponent.h"
-#include "Niagara/Private/NiagaraStats.h"
 #include "NiagaraDataInterfaceGrid2DCollection.h"
+#include "NiagaraStats.h"
 
 #include "NiagaraDataInterfaceGrid2DCollectionReader.generated.h"
 
diff --git a/Engine/Plugins/FX/Niagara/Source/Niagara/Classes/NiagaraDataInterfaceGrid3DCollection.h b/Engine/Plugins/FX/Niagara/Source/Niagara/Classes/NiagaraDataInterfaceGrid3DCollection.h
--- a/Engine/Plugins/FX/Niagara/Source/Niagara/Classes/NiagaraDataInterfaceGrid3DCollection.h
+++ b/Engine/Plugins/FX/Niagara/Source/Niagara/Classes/NiagaraDataInterfaceGrid3DCollection.h
@@ -6,7 +6,7 @@
 #include "ClearQuad.h"
 #include "NiagaraComponent.h"
 #include "NiagaraRenderGraphUtils.h"
-#include "Niagara/Private/NiagaraStats.h"
+#include "NiagaraStats.h"
 
 #include "NiagaraDataInterfaceGrid3DCollection.generated.h"
 
diff --git a/Engine/Plugins/FX/Niagara/Source/Niagara/Classes/NiagaraDataInterfaceGrid3DCollectionReader.h b/Engine/Plugins/FX/Niagara/Source/Niagara/Classes/NiagaraDataInterfaceGrid3DCollectionReader.h
--- a/Engine/Plugins/FX/Niagara/Source/Niagara/Classes/NiagaraDataInterfaceGrid3DCollectionReader.h
+++ b/Engine/Plugins/FX/Niagara/Source/Niagara/Classes/NiagaraDataInterfaceGrid3DCollectionReader.h
@@ -5,8 +5,8 @@
 #include "NiagaraDataInterfaceRW.h"
 #include "ClearQuad.h"
 #include "NiagaraComponent.h"
-#include "Niagara/Private/NiagaraStats.h"
 #include "NiagaraDataInterfaceGrid3DCollection.h"
+#include "NiagaraStats.h"
 
 #include "NiagaraDataInterfaceGrid3DCollectionReader.generated.h"
 
diff --git a/Engine/Plugins/FX/Niagara/Source/Niagara/Classes/NiagaraDataInterfaceNeighborGrid3D.h b/Engine/Plugins/FX/Niagara/Source/Niagara/Classes/NiagaraDataInterfaceNeighborGrid3D.h
--- a/Engine/Plugins/FX/Niagara/Source/Niagara/Classes/NiagaraDataInterfaceNeighborGrid3D.h
+++ b/Engine/Plugins/FX/Niagara/Source/Niagara/Classes/NiagaraDataInterfaceNeighborGrid3D.h
@@ -1,12 +1,12 @@
 // Copyright Epic Games, Inc. All Rights Reserved.
 #pragma once
 
+#include "ClearQuad.h"
 #include "NiagaraDataInterface.h"
 #include "NiagaraCommon.h"
 #include "NiagaraDataInterfaceRW.h"
 #include "NiagaraRenderGraphUtils.h"
-#include "ClearQuad.h"
-#include "Niagara/Private/NiagaraStats.h"
+#include "NiagaraStats.h"
 
 #include "NiagaraDataInterfaceNeighborGrid3D.generated.h"
 
diff --git a/Engine/Plugins/FX/Niagara/Source/Niagara/Classes/NiagaraDataInterfaceSpline.h b/Engine/Plugins/FX/Niagara/Source/Niagara/Classes/NiagaraDataInterfaceSpline.h
--- a/Engine/Plugins/FX/Niagara/Source/Niagara/Classes/NiagaraDataInterfaceSpline.h
+++ b/Engine/Plugins/FX/Niagara/Source/Niagara/Classes/NiagaraDataInterfaceSpline.h
@@ -1,12 +1,12 @@
 // Copyright Epic Games, Inc. All Rights Reserved.
 #pragma once
 
-#include "NiagaraDataInterface.h"
+#include "Components/SplineComponent.h"
 #include "NiagaraCommon.h"
-#include "VectorVM.h"
+#include "NiagaraDataInterface.h"
 #include "NiagaraParameterStore.h"
-#include "Components/SplineComponent.h"
-#include "Niagara/Private/NiagaraStats.h"
+#include "NiagaraStats.h"
+#include "VectorVM.h"
 #include "NiagaraDataInterfaceSpline.generated.h"
 
 class UNiagaraDataInterfaceSpline;
diff --git a/Engine/Plugins/FX/Niagara/Source/Niagara/Niagara.Build.cs b/Engine/Plugins/FX/Niagara/Source/Niagara/Niagara.Build.cs
--- a/Engine/Plugins/FX/Niagara/Source/Niagara/Niagara.Build.cs
+++ b/Engine/Plugins/FX/Niagara/Source/Niagara/Niagara.Build.cs
@@ -60,7 +60,6 @@ public class Niagara : ModuleRules
 
         PrivateIncludePaths.AddRange(
             new string[] {
-                "Niagara/Private",
 				System.IO.Path.Combine(GetModuleDirectory("Engine"), "Private"),
 				System.IO.Path.Combine(GetModuleDirectory("Renderer"), "Private"),
 			});
@@ -94,5 +93,8 @@ public class Niagara : ModuleRules
                 "VECTORVM_SUPPORTS_EXPERIMENTAL=1",
                 "VECTORVM_SUPPORTS_LEGACY=1"
             });
-    }
+
+		// TODO: Should not be including private headers in public code
+		PublicIncludePaths.Add(Path.Combine(ModuleDirectory, "Private")); // For NiagaraStats.h & NiagaraGpuReadbackManager.h
+	}
 }
diff --git a/Engine/Plugins/FX/Niagara/Source/Niagara/Private/DataInterface/NiagaraDataInterfaceStaticMesh.cpp b/Engine/Plugins/FX/Niagara/Source/Niagara/Private/DataInterface/NiagaraDataInterfaceStaticMesh.cpp
--- a/Engine/Plugins/FX/Niagara/Source/Niagara/Private/DataInterface/NiagaraDataInterfaceStaticMesh.cpp
+++ b/Engine/Plugins/FX/Niagara/Source/Niagara/Private/DataInterface/NiagaraDataInterfaceStaticMesh.cpp
@@ -21,6 +21,7 @@
 #include "Subsystems/ImportSubsystem.h"
 #endif
 
+#include "DistanceFieldLightingShared.h"
 #include "Engine/StaticMeshActor.h"
 #include "Engine/StaticMeshSocket.h"
 #include "NiagaraDataInterfaceStaticMeshUvMapping.h"
diff --git a/Engine/Plugins/FX/Niagara/Source/Niagara/Private/NiagaraDataInterfaceCollisionQuery.cpp b/Engine/Plugins/FX/Niagara/Source/Niagara/Private/NiagaraDataInterfaceCollisionQuery.cpp
--- a/Engine/Plugins/FX/Niagara/Source/Niagara/Private/NiagaraDataInterfaceCollisionQuery.cpp
+++ b/Engine/Plugins/FX/Niagara/Source/Niagara/Private/NiagaraDataInterfaceCollisionQuery.cpp
@@ -2,6 +2,7 @@
 
 #include "NiagaraDataInterfaceCollisionQuery.h"
 
+#include "DistanceFieldLightingShared.h"
 #include "GlobalDistanceFieldParameters.h"
 #include "NiagaraAsyncGpuTraceHelper.h"
 #include "NiagaraDistanceFieldHelper.h"
diff --git a/Engine/Plugins/FX/Niagara/Source/NiagaraShader/Private/NiagaraDistanceFieldHelper.cpp b/Engine/Plugins/FX/Niagara/Source/NiagaraShader/Private/NiagaraDistanceFieldHelper.cpp
--- a/Engine/Plugins/FX/Niagara/Source/NiagaraShader/Private/NiagaraDistanceFieldHelper.cpp
+++ b/Engine/Plugins/FX/Niagara/Source/NiagaraShader/Private/NiagaraDistanceFieldHelper.cpp
@@ -1,6 +1,7 @@
 // Copyright Epic Games, Inc. All Rights Reserved.
 
 #include "NiagaraDistanceFieldHelper.h"
+#include "DistanceFieldLightingShared.h"
 #include "SystemTextures.h"
 
 // todo - currently duplicated from SetupGlobalDistanceFieldParameters (GlobalDistanceField.cpp) because of problems getting it properly exported from the dll
diff --git a/Engine/Plugins/FX/Niagara/Source/NiagaraShader/Public/NiagaraDistanceFieldHelper.h b/Engine/Plugins/FX/Niagara/Source/NiagaraShader/Public/NiagaraDistanceFieldHelper.h
--- a/Engine/Plugins/FX/Niagara/Source/NiagaraShader/Public/NiagaraDistanceFieldHelper.h
+++ b/Engine/Plugins/FX/Niagara/Source/NiagaraShader/Public/NiagaraDistanceFieldHelper.h
@@ -3,9 +3,12 @@
 #pragma once
 
 #include "GlobalDistanceFieldParameters.h"
-#include "Renderer/Private/DistanceFieldLightingShared.h"
 #include "RenderGraphBuilder.h"
 
+class FDistanceFieldSceneData;
+class FDistanceFieldObjectBufferParameters;
+class FDistanceFieldAtlasParameters;
+
 namespace FNiagaraDistanceFieldHelper
 {
 	NIAGARASHADER_API void SetGlobalDistanceFieldParameters(const FGlobalDistanceFieldParameterData* OptionalParameterData, FGlobalDistanceFieldParameters2& ShaderParameters);
