- Static bool parameter support for Niagara Material bindings
- Generates a MIC to replace the material on the renderer
- On cooked builds the material is replaced with the MIC during serialization

diff --git a/Engine/Plugins/FX/Niagara/Source/Niagara/Private/NiagaraMeshRendererProperties.cpp b/Engine/Plugins/FX/Niagara/Source/Niagara/Private/NiagaraMeshRendererProperties.cpp
--- a/Engine/Plugins/FX/Niagara/Source/Niagara/Private/NiagaraMeshRendererProperties.cpp
+++ b/Engine/Plugins/FX/Niagara/Source/Niagara/Private/NiagaraMeshRendererProperties.cpp
@@ -7,8 +7,11 @@
 #include "NiagaraBoundsCalculatorHelper.h"
 #include "NiagaraCustomVersion.h"
 #include "NiagaraEmitterInstance.h"
-#include "Modules/ModuleManager.h"
 #include "NiagaraGPUSortInfo.h"
+#include "NiagaraSystem.h"
+
+#include "Materials/MaterialInstanceConstant.h"
+#include "Modules/ModuleManager.h"
 
 #include UE_INLINE_GENERATED_CPP_BY_NAME(NiagaraMeshRendererProperties)
 
@@ -244,7 +247,45 @@ void UNiagaraMeshRendererProperties::Serialize(FArchive& Ar)
 	{
 		SortMode = ENiagaraSortMode::ViewDistance;
 	}
-	Super::Serialize(Ar);
+
+	// MICs will replace the main material during serialize
+	// Be careful if adding code that looks at the material to make sure you get the correct one
+	{
+	#if WITH_EDITORONLY_DATA
+		// For cooked builds we put our MIC material directly into the OverrideMaterial slot to avoid wasting memory
+		TOptional<bool> OverrideMaterialsGuard;
+		TArray<FNiagaraMeshMaterialOverride> OverrideMaterialsArrayGuard;
+		if (Ar.IsSaving() && Ar.IsCooking() && MICMaterials.Num() > 0)
+		{
+			OverrideMaterialsGuard = bOverrideMaterials;
+			OverrideMaterialsArrayGuard = OverrideMaterials;
+
+			for (int i=0; i < OverrideMaterials.Num(); ++i)
+			{
+				FNiagaraMeshMaterialOverride& MaterialOverride = OverrideMaterials[i];
+				if (MICMaterials.IsValidIndex(i) && MICMaterials[i])
+				{
+					MaterialOverride.ExplicitMat = MICMaterials[i];
+				}
+				else if (bOverrideMaterials == false)
+				{
+					MaterialOverride.ExplicitMat = nullptr;
+				}
+			}
+			bOverrideMaterials = true;
+		}
+	#endif
+
+		Super::Serialize(Ar);
+
+#if WITH_EDITORONLY_DATA
+		if (OverrideMaterialsGuard.IsSet())
+		{
+			bOverrideMaterials = OverrideMaterialsGuard.GetValue();
+			OverrideMaterials = OverrideMaterialsArrayGuard;
+		}
+#endif
+	}
 }
 
 /** The bindings depend on variables that are created during the NiagaraModule startup. However, the CDO's are build prior to this being initialized, so we defer setting these values until later.*/
@@ -317,6 +358,7 @@ void UNiagaraMeshRendererProperties::UpdateSourceModeDerivates(ENiagaraRendererS
 void UNiagaraMeshRendererProperties::CacheFromCompiledData(const FNiagaraDataSetCompiledData* CompiledData)
 {
 	UpdateSourceModeDerivates(SourceMode);
+	UpdateMICs();
 
 	// Initialize layout
 	const int32 NumLayoutVars = NeedsPreciseMotionVectors() ? ENiagaraMeshVFLayout::Num_Max : ENiagaraMeshVFLayout::Num_Default;
@@ -370,6 +412,21 @@ void UNiagaraMeshRendererProperties::CacheFromCompiledData(const FNiagaraDataSet
 	RendererLayoutWithoutCustomSorting.Finalize();
 }
 
+void UNiagaraMeshRendererProperties::UpdateMICs()
+{
+#if WITH_EDITORONLY_DATA
+	// Note: We need to swap out MICMaterials as GetUsedMaterials uses them to override the material
+	TArray<TObjectPtr<UMaterialInstanceConstant>> TempMICMaterials;
+	Swap(TempMICMaterials, MICMaterials);
+
+	TArray<UMaterialInterface*> Materials;
+	GetUsedMaterials(nullptr, Materials);
+	UpdateMaterialParametersMIC(MaterialParameters, Materials, TempMICMaterials);
+
+	Swap(TempMICMaterials, MICMaterials);
+#endif
+}
+
 #if WITH_EDITORONLY_DATA
 bool UNiagaraMeshRendererProperties::IsSupportedVariableForBinding(const FNiagaraVariableBase& InSourceForBinding, const FName& InTargetBindingName) const
 {
@@ -442,6 +499,9 @@ void UNiagaraMeshRendererProperties::GetUsedMeshMaterials(int32 MeshIndex, const
 				if (!OverrideMat)
 				{
 					OverrideMat = OverrideMaterials[OverrideIndex].ExplicitMat;
+				#if WITH_EDITORONLY_DATA
+					OverrideMat = MICMaterials.IsValidIndex(OverrideIndex) && MICMaterials[OverrideIndex] ? MICMaterials[OverrideIndex] : OverrideMat;
+				#endif
 				}
 
 				if (OverrideMat)
@@ -451,6 +511,19 @@ void UNiagaraMeshRendererProperties::GetUsedMeshMaterials(int32 MeshIndex, const
 			}
 		}
 	}
+	else
+	{
+	#if WITH_EDITORONLY_DATA
+		if (MICMaterials.Num() > 0)
+		{
+			const int MaterialMax = FMath::Min(OutMaterials.Num(), MICMaterials.Num());
+			for (int i = 0; i < MaterialMax; ++i)
+			{
+				OutMaterials[i] = MICMaterials[i] ? MICMaterials[i] : OutMaterials[i];
+			}
+		}
+	#endif
+	}
 }
 
 const FVertexFactoryType* UNiagaraMeshRendererProperties::GetVertexFactoryType() const
@@ -754,6 +827,8 @@ void UNiagaraMeshRendererProperties::PostEditChangeProperty(FPropertyChangedEven
 
 	const bool bIsRedirect = PropertyChangedEvent.ChangeType == EPropertyChangeType::Redirected;
 	const bool bRebuildMeshList = ChangeRequiresMeshListRebuild(PropertyChangedEvent.Property);
+	const FName PropertyName = PropertyChangedEvent.GetPropertyName();
+	const FName MemberPropertyName = PropertyChangedEvent.GetMemberPropertyName();
 
 	if (bIsRedirect)
 	{
@@ -771,7 +846,7 @@ void UNiagaraMeshRendererProperties::PostEditChangeProperty(FPropertyChangedEven
 	if (bRebuildMeshList)
 	{
 		if (!IsRunningCommandlet() &&
-			PropertyChangedEvent.Property->GetFName() == GET_MEMBER_NAME_CHECKED(UNiagaraMeshRendererProperties, bEnableMeshFlipbook) &&
+			PropertyName == GET_MEMBER_NAME_CHECKED(UNiagaraMeshRendererProperties, bEnableMeshFlipbook) &&
 			bEnableMeshFlipbook &&
 			Meshes.Num() > 0)
 		{
@@ -815,7 +890,7 @@ void UNiagaraMeshRendererProperties::PostEditChangeProperty(FPropertyChangedEven
 	}
 
 	// If changing the source mode, we may need to update many of our values.
-	if (PropertyChangedEvent.GetPropertyName() == TEXT("SourceMode"))
+	if (PropertyName == TEXT("SourceMode"))
 	{
 		UpdateSourceModeDerivates(SourceMode, true);
 	}
@@ -838,33 +913,32 @@ void UNiagaraMeshRendererProperties::PostEditChangeProperty(FPropertyChangedEven
 		}
 	}
 
+	// Update our MICs if we change override material / material bindings / meshes
+	//-OPT: Could narrow down further to only static materials
+	if ((MemberPropertyName == GET_MEMBER_NAME_CHECKED(UNiagaraMeshRendererProperties, OverrideMaterials)) ||
+		(MemberPropertyName == GET_MEMBER_NAME_CHECKED(UNiagaraMeshRendererProperties, Meshes)) ||
+		(MemberPropertyName == GET_MEMBER_NAME_CHECKED(UNiagaraMeshRendererProperties, MaterialParameters)))
+	{
+		UpdateMICs();
+	}
+
 	Super::PostEditChangeProperty(PropertyChangedEvent);
 }
 
 void UNiagaraMeshRendererProperties::RenameVariable(const FNiagaraVariableBase& OldVariable, const FNiagaraVariableBase& NewVariable, const FVersionedNiagaraEmitter& InEmitter)
 {
 	Super::RenameVariable(OldVariable, NewVariable, InEmitter);
-
-	// Handle renaming material bindings
-	for (FNiagaraMaterialAttributeBinding& Binding : MaterialParameters.AttributeBindings)
-	{
-		Binding.RenameVariableIfMatching(OldVariable, NewVariable, InEmitter.Emitter, GetCurrentSourceMode());
-	}
+#if WITH_EDITORONLY_DATA
+	MaterialParameters.RenameVariable(OldVariable, NewVariable, InEmitter, GetCurrentSourceMode());
+#endif
 }
 
 void UNiagaraMeshRendererProperties::RemoveVariable(const FNiagaraVariableBase& OldVariable, const FVersionedNiagaraEmitter& InEmitter)
 {
 	Super::RemoveVariable(OldVariable, InEmitter);
-
-	// Handle resetting material bindings to defaults
-	for (FNiagaraMaterialAttributeBinding& Binding : MaterialParameters.AttributeBindings)
-	{
-		if (Binding.Matches(OldVariable, InEmitter.Emitter, GetCurrentSourceMode()))
-		{
-			Binding.NiagaraVariable = FNiagaraVariable();
-			Binding.CacheValues(InEmitter.Emitter);
-		}
-	}
+#if WITH_EDITORONLY_DATA
+	MaterialParameters.RemoveVariable(OldVariable, InEmitter, GetCurrentSourceMode());
+#endif
 }
 
 void UNiagaraMeshRendererProperties::OnMeshChanged()
@@ -878,6 +952,7 @@ void UNiagaraMeshRendererProperties::OnMeshChanged()
 	}
 
 	CheckMaterialUsage();
+	UpdateMICs();
 }
 
 void UNiagaraMeshRendererProperties::OnMeshPostBuild(UStaticMesh*)
diff --git a/Engine/Plugins/FX/Niagara/Source/Niagara/Private/NiagaraRendererMeshes.cpp b/Engine/Plugins/FX/Niagara/Source/Niagara/Private/NiagaraRendererMeshes.cpp
--- a/Engine/Plugins/FX/Niagara/Source/Niagara/Private/NiagaraRendererMeshes.cpp
+++ b/Engine/Plugins/FX/Niagara/Source/Niagara/Private/NiagaraRendererMeshes.cpp
@@ -203,20 +203,24 @@ void FNiagaraRendererMeshes::Initialize(const UNiagaraRendererProperties* InProp
 			// Create an index remap from mesh material index to it's index in the base material list
 			TArray<UMaterialInterface*> MeshMaterials;
 			Properties->GetUsedMeshMaterials(SourceMeshIndex, Emitter, MeshMaterials);
-			for (auto MeshMaterial : MeshMaterials)
+			for (UMaterialInterface* MeshMaterial : MeshMaterials)
 			{
-				MeshData.MaterialRemapTable.Add(BaseMaterials_GT.IndexOfByPredicate([&](UMaterialInterface* LookMat)
-				{
-					if (LookMat == MeshMaterial)
-					{
-						return true;
-					}
-					if (UMaterialInstanceDynamic* MID = Cast<UMaterialInstanceDynamic>(LookMat))
-					{
-						return MeshMaterial == MID->Parent;
-					}
-					return false;
-				}));
+				MeshData.MaterialRemapTable.Add(
+					BaseMaterials_GT.IndexOfByPredicate(
+						[&](UMaterialInterface* LookMat)
+						{
+							if (LookMat == MeshMaterial)
+							{
+								return true;
+							}
+							if (UMaterialInstanceDynamic* MID = Cast<UMaterialInstanceDynamic>(LookMat))
+							{
+								return MeshMaterial == MID->Parent;
+							}
+							return false;
+						}
+					)
+				);
 			}
 
 			// Determine the max section count for all LODs of this mesh and accumulate it on the max for all meshes
diff --git a/Engine/Plugins/FX/Niagara/Source/Niagara/Private/NiagaraRendererProperties.cpp b/Engine/Plugins/FX/Niagara/Source/Niagara/Private/NiagaraRendererProperties.cpp
--- a/Engine/Plugins/FX/Niagara/Source/Niagara/Private/NiagaraRendererProperties.cpp
+++ b/Engine/Plugins/FX/Niagara/Source/Niagara/Private/NiagaraRendererProperties.cpp
@@ -8,8 +8,10 @@
 #include "NiagaraScriptSourceBase.h"
 #include "NiagaraSettings.h"
 #include "NiagaraSystem.h"
+
 #include "Interfaces/ITargetPlatform.h"
 #include "Materials/MaterialInterface.h"
+#include "Materials/MaterialInstanceConstant.h"
 #include "Styling/SlateIconFinder.h"
 
 #include UE_INLINE_GENERATED_CPP_BY_NAME(NiagaraRendererProperties)
@@ -108,6 +110,48 @@ void FNiagaraRendererLayout::Finalize()
 //////////////////////////////////////////////////////////////////////////
 
 #if WITH_EDITORONLY_DATA
+void FNiagaraRendererMaterialParameters::RenameVariable(const FNiagaraVariableBase& OldVariable, const FNiagaraVariableBase& NewVariable, const FVersionedNiagaraEmitter& InEmitter, ENiagaraRendererSourceDataMode SourceMode)
+{
+	for (FNiagaraMaterialAttributeBinding& Binding : AttributeBindings)
+	{
+		Binding.RenameVariableIfMatching(OldVariable, NewVariable, InEmitter.Emitter, SourceMode);
+	}
+
+	if ( OldVariable.GetType() == FNiagaraTypeDefinition::GetBoolDef().ToStaticDef() )
+	{
+		for (FNiagaraRendererMaterialStaticBoolParameter& Binding : StaticBoolParameters)
+		{
+			if (Binding.StaticVariableName == OldVariable.GetName())
+			{
+				Binding.StaticVariableName = NewVariable.GetName();
+			}
+		}
+	}
+}
+
+void FNiagaraRendererMaterialParameters::RemoveVariable(const FNiagaraVariableBase& OldVariable, const FVersionedNiagaraEmitter& InEmitter, ENiagaraRendererSourceDataMode SourceMode)
+{
+	for (FNiagaraMaterialAttributeBinding& Binding : AttributeBindings)
+	{
+		if (Binding.Matches(OldVariable, InEmitter.Emitter, SourceMode))
+		{
+			Binding.NiagaraVariable = FNiagaraVariable();
+			Binding.CacheValues(InEmitter.Emitter);
+		}
+	}
+
+	if (OldVariable.GetType() == FNiagaraTypeDefinition::GetBoolDef().ToStaticDef())
+	{
+		for (FNiagaraRendererMaterialStaticBoolParameter& Binding : StaticBoolParameters)
+		{
+			if (Binding.StaticVariableName == OldVariable.GetName())
+			{
+				Binding.StaticVariableName = FName();
+			}
+		}
+	}
+}
+
 void FNiagaraRendererMaterialParameters::GetFeedback(TArrayView<UMaterialInterface*> Materials, TArray<FNiagaraRendererFeedback>& OutWarnings) const
 {
 	TArray<bool> AttributeBindingsValid;
@@ -266,6 +310,146 @@ void UNiagaraRendererProperties::ChangeToPositionBinding(FNiagaraVariableAttribu
 	}
 }
 
+bool UNiagaraRendererProperties::UpdateMaterialStaticParameters(const FNiagaraRendererMaterialParameters& MaterialParameters, UMaterialInstanceConstant* MIC)
+{
+	FStaticParameterSet StaticParameterSet;
+
+	UNiagaraEmitter* NiagaraEmitter = GetTypedOuter<UNiagaraEmitter>();
+	UNiagaraSystem* NiagaraSystem = GetTypedOuter<UNiagaraSystem>();
+	if (NiagaraEmitter == nullptr || NiagaraSystem == nullptr)
+	{
+		return false;
+	}
+
+	bool bModified = false;
+	for (const FNiagaraRendererMaterialStaticBoolParameter& ParameterBinding : MaterialParameters.StaticBoolParameters)
+	{
+		NiagaraSystem->ForEachScript(
+			[&](UNiagaraScript* NiagaraScript)
+			{
+				for ( const FNiagaraVariable& StaticVariable : NiagaraScript->GetVMExecutableData().StaticVariablesWritten )
+				{
+					if (StaticVariable.GetType() != FNiagaraTypeDefinition::GetBoolDef().ToStaticDef())
+					{
+						continue;
+					}
+
+					FNiagaraVariableBase ResolvedStaticVariable = StaticVariable;
+					ResolvedStaticVariable.ReplaceRootNamespace(NiagaraEmitter->GetUniqueEmitterName(), FNiagaraConstants::EmitterNamespaceString);
+					if (ResolvedStaticVariable.GetName() != ParameterBinding.StaticVariableName)
+					{
+						continue;
+					}
+
+					FStaticSwitchParameter* StaticParameter = StaticParameterSet.StaticSwitchParameters.FindByPredicate(
+						[ParameterBinding](const FStaticSwitchParameter& StaticParameter)
+						{
+							return StaticParameter.ParameterInfo.Name == ParameterBinding.MaterialParameterName;
+						}
+					);
+
+					const bool bNewValue = StaticVariable.GetValue<bool>();
+					if (StaticParameter == nullptr)
+					{
+						FGuid ParameterGuid;
+						bool bDefaultValue = false;
+						if (MIC->GetStaticSwitchParameterDefaultValue(ParameterBinding.MaterialParameterName, bDefaultValue, ParameterGuid))
+						{
+							if (bDefaultValue != bNewValue)
+							{
+								StaticParameterSet.SetParameterValue(ParameterBinding.MaterialParameterName, FMaterialParameterValue(bNewValue), EMaterialSetParameterValueFlags::None);
+								bModified = true;
+							}
+						}
+					}
+					else
+					{
+						if (StaticParameter && StaticParameter->Value != bNewValue)
+						{
+							StaticParameter->Value = bNewValue;
+							bModified = true;
+						}
+					}
+
+					break;
+				}
+			}
+		);
+	}
+
+	if (bModified)
+	{
+		MIC->UpdateStaticPermutation(StaticParameterSet);
+	}
+
+	return bModified;
+}
+
+void UNiagaraRendererProperties::UpdateMaterialParametersMIC(const FNiagaraRendererMaterialParameters& MaterialParameters, TObjectPtr<UMaterialInterface>& InOutMaterial, TObjectPtr<UMaterialInstanceConstant>& InOutMIC)
+{
+	if (InOutMaterial == nullptr || MaterialParameters.StaticBoolParameters.Num() == 0)
+	{
+		InOutMIC = nullptr;
+		return;
+	}
+
+	if (InOutMIC == nullptr)
+	{
+		InOutMIC = NewObject<UMaterialInstanceConstant>(this);
+	}
+	InOutMIC->SetParentEditorOnly(InOutMaterial);
+
+	if (UpdateMaterialStaticParameters(MaterialParameters, InOutMIC) == false)
+	{
+		InOutMIC = nullptr;
+	}
+}
+
+void UNiagaraRendererProperties::UpdateMaterialParametersMIC(const FNiagaraRendererMaterialParameters& MaterialParameters, TArrayView<UMaterialInterface*> Materials, TArray<TObjectPtr<UMaterialInstanceConstant>>& InOutMICs)
+{
+	if (Materials.Num() == 0 || MaterialParameters.StaticBoolParameters.Num() == 0)
+	{
+		InOutMICs.Empty();
+		return;
+	}
+
+	// Create a pool of unique MICs that we can potentially reuse
+	TArray<UMaterialInstanceConstant*> MICPool;
+	MICPool.Reserve(InOutMICs.Num());
+	for (UMaterialInstanceConstant* MIC : InOutMICs)
+	{
+		if (MIC != nullptr)
+		{
+			MICPool.AddUnique(MIC);
+		}
+	}
+	InOutMICs.Reset(0);
+
+	// Loop over each material to see if we need to generate a MIC for it
+	for ( int i=0; i < Materials.Num(); ++i )
+	{
+		UMaterialInterface* Material = Materials[i];
+		if ( Material == nullptr )
+		{
+			continue;
+		}
+
+		//-OPT: We should be able to reuse rather than create
+		UMaterialInstanceConstant* MIC = MICPool.Num() > 0 ? MICPool.Pop() : NewObject<UMaterialInstanceConstant>(this);
+		MIC->SetParentEditorOnly(Material);
+
+		if (UpdateMaterialStaticParameters(MaterialParameters, MIC) == false)
+		{
+			MICPool.Add(MIC);
+		}
+		else
+		{
+			InOutMICs.SetNum(i + 1);
+			InOutMICs[i] = MIC;
+		}
+	}
+}
+
 FNiagaraVariable UNiagaraRendererProperties::GetBoundAttribute(const FNiagaraVariableAttributeBinding* Binding) const
 {
 	if (Binding->GetParamMapBindableVariable().IsValid())
diff --git a/Engine/Plugins/FX/Niagara/Source/Niagara/Private/NiagaraRibbonRendererProperties.cpp b/Engine/Plugins/FX/Niagara/Source/Niagara/Private/NiagaraRibbonRendererProperties.cpp
--- a/Engine/Plugins/FX/Niagara/Source/Niagara/Private/NiagaraRibbonRendererProperties.cpp
+++ b/Engine/Plugins/FX/Niagara/Source/Niagara/Private/NiagaraRibbonRendererProperties.cpp
@@ -6,6 +6,9 @@
 #include "NiagaraBoundsCalculatorHelper.h"
 #include "NiagaraCustomVersion.h"
 #include "NiagaraEmitterInstance.h"
+#include "NiagaraSystem.h"
+
+#include "Materials/MaterialInstanceConstant.h"
 #include "Modules/ModuleManager.h"
 
 #include UE_INLINE_GENERATED_CPP_BY_NAME(NiagaraRibbonRendererProperties)
@@ -174,6 +177,10 @@ void UNiagaraRibbonRendererProperties::GetUsedMaterials(const FNiagaraEmitterIns
 		MaterialInterface = Cast<UMaterialInterface>(InEmitter->FindBinding(MaterialUserParamBinding.Parameter));
 	}
 
+#if WITH_EDITORONLY_DATA
+	MaterialInterface = MaterialInterface ? MaterialInterface : ToRawPtr(MICMaterial);
+#endif
+
 	OutMaterials.Add(MaterialInterface ? MaterialInterface : ToRawPtr(Material));
 }
 
@@ -260,6 +267,24 @@ void UNiagaraRibbonRendererProperties::PostInitProperties()
 	}
 }
 
+void UNiagaraRibbonRendererProperties::Serialize(FStructuredArchive::FRecord Record)
+{
+	FArchive& Ar = Record.GetUnderlyingArchive();
+
+	// MIC will replace the main material during serialize
+	// Be careful if adding code that looks at the material to make sure you get the correct one
+	{
+#if WITH_EDITORONLY_DATA
+		TOptional<TGuardValue<TObjectPtr<UMaterialInterface>>> MICGuard;
+		if (Ar.IsSaving() && Ar.IsCooking() && MICMaterial)
+		{
+			MICGuard.Emplace(Material, MICMaterial);
+		}
+#endif
+		Super::Serialize(Record);
+	}
+}
+
 /** The bindings depend on variables that are created during the NiagaraModule startup. However, the CDO's are build prior to this being initialized, so we defer setting these values until later.*/
 void UNiagaraRibbonRendererProperties::InitCDOPropertiesAfterModuleStartup()
 {
@@ -313,6 +338,8 @@ void UNiagaraRibbonRendererProperties::SetPreviousBindings(const FVersionedNiaga
 
 void UNiagaraRibbonRendererProperties::CacheFromCompiledData(const FNiagaraDataSetCompiledData* CompiledData)
 {
+	UpdateMICs();
+
 	// Initialize accessors
 	bSortKeyDataSetAccessorIsAge = false;
 	SortKeyDataSetAccessor.Init(CompiledData, RibbonLinkOrderBinding.GetDataSetBindableVariable().GetName());
@@ -394,6 +421,12 @@ void UNiagaraRibbonRendererProperties::CacheFromCompiledData(const FNiagaraDataS
 	RendererLayout.Finalize();
 }
 
+void UNiagaraRibbonRendererProperties::UpdateMICs()
+{
+#if WITH_EDITORONLY_DATA
+	UpdateMaterialParametersMIC(MaterialParameters, Material, MICMaterial);
+#endif
+}
 
 #if WITH_EDITORONLY_DATA
 
@@ -412,7 +445,10 @@ bool UNiagaraRibbonRendererProperties::IsSupportedVariableForBinding(const FNiag
 void UNiagaraRibbonRendererProperties::PostEditChangeProperty(struct FPropertyChangedEvent& PropertyChangedEvent)
 {
 	Super::PostEditChangeProperty(PropertyChangedEvent);
-	FName PropertyName = PropertyChangedEvent.GetPropertyName();
+
+	const FName PropertyName = PropertyChangedEvent.GetPropertyName();
+	const FName MemberPropertyName = PropertyChangedEvent.GetMemberPropertyName();
+
 	if (PropertyName == GET_MEMBER_NAME_CHECKED(UNiagaraRibbonRendererProperties, TessellationAngle))
 	{
 		if (TessellationAngle > 0.f && TessellationAngle < 1.f)
@@ -420,6 +456,14 @@ void UNiagaraRibbonRendererProperties::PostEditChangeProperty(struct FPropertyCh
 			TessellationAngle = 1.f;
 		}
 	}
+
+	// Update our MICs if we change material / material bindings
+	//-OPT: Could narrow down further to only static materials
+	if ((PropertyName == GET_MEMBER_NAME_CHECKED(UNiagaraRibbonRendererProperties, Material)) ||
+		(MemberPropertyName == GET_MEMBER_NAME_CHECKED(UNiagaraRibbonRendererProperties, MaterialParameters)))
+	{
+		UpdateMICs();
+	}
 }
 
 const TArray<FNiagaraVariable>& UNiagaraRibbonRendererProperties::GetOptionalAttributes()
@@ -589,27 +633,17 @@ bool UNiagaraRibbonRendererProperties::CanEditChange(const FProperty* InProperty
 void UNiagaraRibbonRendererProperties::RenameVariable(const FNiagaraVariableBase& OldVariable, const FNiagaraVariableBase& NewVariable, const FVersionedNiagaraEmitter& InEmitter)
 {
 	Super::RenameVariable(OldVariable, NewVariable, InEmitter);
-
-	// Handle renaming material bindings
-	for (FNiagaraMaterialAttributeBinding& Binding : MaterialParameters.AttributeBindings)
-	{
-		Binding.RenameVariableIfMatching(OldVariable, NewVariable, InEmitter.Emitter, GetCurrentSourceMode());
-	}
+#if WITH_EDITORONLY_DATA
+	MaterialParameters.RenameVariable(OldVariable, NewVariable, InEmitter, GetCurrentSourceMode());
+#endif
 }
 
 void UNiagaraRibbonRendererProperties::RemoveVariable(const FNiagaraVariableBase& OldVariable, const FVersionedNiagaraEmitter& InEmitter)
 {
 	Super::RemoveVariable(OldVariable, InEmitter);
-
-	// Handle resetting material bindings to defaults
-	for (FNiagaraMaterialAttributeBinding& Binding : MaterialParameters.AttributeBindings)
-	{
-		if (Binding.Matches(OldVariable, InEmitter.Emitter, GetCurrentSourceMode()))
-		{
-			Binding.NiagaraVariable = FNiagaraVariable();
-			Binding.CacheValues(InEmitter.Emitter);
-		}
-	}
+#if WITH_EDITORONLY_DATA
+	MaterialParameters.RemoveVariable(OldVariable, InEmitter, GetCurrentSourceMode());
+#endif
 }
 
 #endif // WITH_EDITORONLY_DATA
diff --git a/Engine/Plugins/FX/Niagara/Source/Niagara/Private/NiagaraSpriteRendererProperties.cpp b/Engine/Plugins/FX/Niagara/Source/Niagara/Private/NiagaraSpriteRendererProperties.cpp
--- a/Engine/Plugins/FX/Niagara/Source/Niagara/Private/NiagaraSpriteRendererProperties.cpp
+++ b/Engine/Plugins/FX/Niagara/Source/Niagara/Private/NiagaraSpriteRendererProperties.cpp
@@ -1,17 +1,20 @@
 // Copyright Epic Games, Inc. All Rights Reserved.
 
 #include "NiagaraSpriteRendererProperties.h"
-#include "Engine/Texture2D.h"
 #include "NiagaraRenderer.h"
-#include "Internationalization/Internationalization.h"
 #include "NiagaraConstants.h"
 #include "NiagaraRendererSprites.h"
 #include "NiagaraBoundsCalculatorHelper.h"
 #include "NiagaraCustomVersion.h"
 #include "NiagaraEmitterInstance.h"
-#include "Modules/ModuleManager.h"
 #include "NiagaraEmitter.h"
 #include "NiagaraScriptSourceBase.h"
+#include "NiagaraSystem.h"
+
+#include "Engine/Texture2D.h"
+#include "Internationalization/Internationalization.h"
+#include "Materials/MaterialInstanceConstant.h"
+#include "Modules/ModuleManager.h"
 
 #include UE_INLINE_GENERATED_CPP_BY_NAME(NiagaraSpriteRendererProperties)
 
@@ -123,6 +126,10 @@ void UNiagaraSpriteRendererProperties::GetUsedMaterials(const FNiagaraEmitterIns
 		MaterialInterface = Cast<UMaterialInterface>(InEmitter->FindBinding(MaterialUserParamBinding.Parameter));
 	}
 
+#if WITH_EDITORONLY_DATA
+	MaterialInterface = MaterialInterface ? MaterialInterface : ToRawPtr(MICMaterial);
+#endif
+
 	OutMaterials.Add(MaterialInterface ? MaterialInterface : ToRawPtr(Material));
 }
 
@@ -197,15 +204,26 @@ void UNiagaraSpriteRendererProperties::Serialize(FStructuredArchive::FRecord Rec
 		SortMode = ENiagaraSortMode::ViewDistance;
 	}
 
-	Super::Serialize(Record);
+	// MIC will replace the main material during serialize
+	// Be careful if adding code that looks at the material to make sure you get the correct one
+	{
+	#if WITH_EDITORONLY_DATA
+		TOptional<TGuardValue<TObjectPtr<UMaterialInterface>>> MICGuard;
+		if (Ar.IsSaving() && Ar.IsCooking() && MICMaterial)
+		{
+			MICGuard.Emplace(Material, MICMaterial);
+		}
+	#endif
+
+		Super::Serialize(Record);
+	}
 
 	bool bIsCookedForEditor = false;
 #if WITH_EDITORONLY_DATA
 	bIsCookedForEditor = ((Ar.GetPortFlags() & PPF_Duplicate) == 0) && GetPackage()->HasAnyPackageFlags(PKG_Cooked);
 #endif // WITH_EDITORONLY_DATA
 
-	FArchive& UnderlyingArchive = Record.GetUnderlyingArchive();
-	if (UnderlyingArchive.IsCooking() || (FPlatformProperties::RequiresCookedData() && UnderlyingArchive.IsLoading()) || bIsCookedForEditor)
+	if (Ar.IsCooking() || (FPlatformProperties::RequiresCookedData() && Ar.IsLoading()) || bIsCookedForEditor)
 	{
 		DerivedData.Serialize(Record.EnterField(TEXT("DerivedData")));
 	}
@@ -271,6 +289,7 @@ void UNiagaraSpriteRendererProperties::SetPreviousBindings(const FVersionedNiaga
 void UNiagaraSpriteRendererProperties::CacheFromCompiledData(const FNiagaraDataSetCompiledData* CompiledData)
 {
 	UpdateSourceModeDerivates(SourceMode);
+	UpdateMICs();
 
 	const int32 NumLayoutVars = NeedsPreciseMotionVectors() ? ENiagaraSpriteVFLayout::Num_Max : ENiagaraSpriteVFLayout::Num_Default;
 	RendererLayoutWithCustomSort.Initialize(NumLayoutVars);
@@ -335,7 +354,6 @@ void UNiagaraSpriteRendererProperties::CacheFromCompiledData(const FNiagaraDataS
 		RendererLayoutWithoutCustomSort.SetVariableFromBinding(CompiledData, PrevPivotOffsetBinding, ENiagaraSpriteVFLayout::PrevPivotOffset);
 	}
 	RendererLayoutWithoutCustomSort.Finalize();
-
 }
 
 #if WITH_EDITORONLY_DATA
@@ -390,6 +408,13 @@ void UNiagaraSpriteRendererProperties::UpdateSourceModeDerivates(ENiagaraRendere
 	}
 }
 
+void UNiagaraSpriteRendererProperties::UpdateMICs()
+{
+#if WITH_EDITORONLY_DATA
+	UpdateMaterialParametersMIC(MaterialParameters, Material, MICMaterial);
+#endif
+}
+
 #if WITH_EDITORONLY_DATA
 
 bool UNiagaraSpriteRendererProperties::IsSupportedVariableForBinding(const FNiagaraVariableBase& InSourceForBinding, const FName& InTargetBindingName) const
@@ -421,6 +446,9 @@ bool UNiagaraSpriteRendererProperties::IsSupportedVariableForBinding(const FNiag
 
 void UNiagaraSpriteRendererProperties::PostEditChangeProperty(struct FPropertyChangedEvent& PropertyChangedEvent)
 {
+	const FName PropertyName = PropertyChangedEvent.GetPropertyName();
+	const FName MemberPropertyName = PropertyChangedEvent.GetMemberPropertyName();
+
 	SubImageSize.X = FMath::Max<float>(SubImageSize.X, 1.f);
 	SubImageSize.Y = FMath::Max<float>(SubImageSize.Y, 1.f);
 
@@ -428,12 +456,12 @@ void UNiagaraSpriteRendererProperties::PostEditChangeProperty(struct FPropertyCh
 	if (bUseMaterialCutoutTexture || CutoutTexture || DerivedData.BoundingGeometry.Num())
 	{
 		const bool bUpdateCutoutDDC =
-			PropertyChangedEvent.GetPropertyName() == TEXT("bUseMaterialCutoutTexture") ||
-			PropertyChangedEvent.GetPropertyName() == TEXT("CutoutTexture") ||
-			PropertyChangedEvent.GetPropertyName() == TEXT("BoundingMode") ||
-			PropertyChangedEvent.GetPropertyName() == TEXT("OpacitySourceMode") ||
-			PropertyChangedEvent.GetPropertyName() == TEXT("AlphaThreshold") ||
-			(bUseMaterialCutoutTexture && PropertyChangedEvent.GetPropertyName() == TEXT("Material"));
+			PropertyName == TEXT("bUseMaterialCutoutTexture") ||
+			PropertyName == TEXT("CutoutTexture") ||
+			PropertyName == TEXT("BoundingMode") ||
+			PropertyName == TEXT("OpacitySourceMode") ||
+			PropertyName == TEXT("AlphaThreshold") ||
+			(bUseMaterialCutoutTexture && PropertyName == TEXT("Material"));
 
 		if (bUpdateCutoutDDC)
 		{
@@ -442,6 +470,14 @@ void UNiagaraSpriteRendererProperties::PostEditChangeProperty(struct FPropertyCh
 		}
 	}
 
+	// Update our MICs if we change material / material bindings
+	//-OPT: Could narrow down further to only static materials
+	if ((PropertyName == GET_MEMBER_NAME_CHECKED(UNiagaraSpriteRendererProperties, Material)) ||
+		(MemberPropertyName == GET_MEMBER_NAME_CHECKED(UNiagaraSpriteRendererProperties, MaterialParameters)) )
+	{
+		UpdateMICs();
+	}
+
 	// If changing the source mode, we may need to update many of our values.
 	if (PropertyChangedEvent.GetPropertyName() == TEXT("SourceMode"))
 	{
@@ -467,34 +503,22 @@ void UNiagaraSpriteRendererProperties::PostEditChangeProperty(struct FPropertyCh
 	}
 
 	Super::PostEditChangeProperty(PropertyChangedEvent);
-
 }
 
-
 void UNiagaraSpriteRendererProperties::RenameVariable(const FNiagaraVariableBase& OldVariable, const FNiagaraVariableBase& NewVariable, const FVersionedNiagaraEmitter& InEmitter)
 {
 	Super::RenameVariable(OldVariable, NewVariable, InEmitter);
-
-	// Handle renaming material bindings
-	for (FNiagaraMaterialAttributeBinding& Binding : MaterialParameters.AttributeBindings)
-	{
-		Binding.RenameVariableIfMatching(OldVariable, NewVariable, InEmitter.Emitter, GetCurrentSourceMode());
-	}
+#if WITH_EDITORONLY_DATA
+	MaterialParameters.RenameVariable(OldVariable, NewVariable, InEmitter, GetCurrentSourceMode());
+#endif
 }
 
 void UNiagaraSpriteRendererProperties::RemoveVariable(const FNiagaraVariableBase& OldVariable, const FVersionedNiagaraEmitter& InEmitter)
 {
 	Super::RemoveVariable(OldVariable, InEmitter);
-
-	// Handle resetting material bindings to defaults
-	for (FNiagaraMaterialAttributeBinding& Binding : MaterialParameters.AttributeBindings)
-	{
-		if (Binding.Matches(OldVariable, InEmitter.Emitter, GetCurrentSourceMode()))
-		{
-			Binding.NiagaraVariable = FNiagaraVariable();
-			Binding.CacheValues(InEmitter.Emitter);
-		}
-	}
+#if WITH_EDITORONLY_DATA
+	MaterialParameters.RemoveVariable(OldVariable, InEmitter, GetCurrentSourceMode());
+#endif
 }
 
 const TArray<FNiagaraVariable>& UNiagaraSpriteRendererProperties::GetOptionalAttributes()
diff --git a/Engine/Plugins/FX/Niagara/Source/Niagara/Private/NiagaraTypes.cpp b/Engine/Plugins/FX/Niagara/Source/Niagara/Private/NiagaraTypes.cpp
--- a/Engine/Plugins/FX/Niagara/Source/Niagara/Private/NiagaraTypes.cpp
+++ b/Engine/Plugins/FX/Niagara/Source/Niagara/Private/NiagaraTypes.cpp
@@ -49,7 +49,7 @@ bool FNiagaraVariableBase::ReplaceRootNamespace(const FStringView& ExpectedNames
 	{
 		FNameBuilder NewNameString;
 		NewNameString.Append(NewNamespace);
-		NewNameString.Append(NameStringView.Mid(ExpectedNamespace.Len() + 1));
+		NewNameString.Append(NameStringView.Mid(ExpectedNamespace.Len()));
 		Name = FName(FStringView(NewNameString));
 		return true;
 	}
diff --git a/Engine/Plugins/FX/Niagara/Source/Niagara/Public/NiagaraMeshRendererProperties.h b/Engine/Plugins/FX/Niagara/Source/Niagara/Public/NiagaraMeshRendererProperties.h
--- a/Engine/Plugins/FX/Niagara/Source/Niagara/Public/NiagaraMeshRendererProperties.h
+++ b/Engine/Plugins/FX/Niagara/Source/Niagara/Public/NiagaraMeshRendererProperties.h
@@ -8,6 +8,7 @@
 #include "NiagaraCommon.h"
 #include "NiagaraMeshRendererProperties.generated.h"
 
+class UMaterialInstanceConstant;
 class FNiagaraEmitterInstance;
 class FAssetThumbnailPool;
 class SWidget;
@@ -195,6 +196,8 @@ public:
 #endif // WITH_EDITORONLY_DATA
 	virtual void CacheFromCompiledData(const FNiagaraDataSetCompiledData* CompiledData) override;
 
+	void UpdateMICs();
+
 #if WITH_EDITORONLY_DATA
 	bool IsSupportedVariableForBinding(const FNiagaraVariableBase& InSourceForBinding, const FName& InTargetBindingName) const;
 #endif
@@ -263,6 +266,11 @@ public:
 	UPROPERTY(EditAnywhere, Category = "Mesh Rendering", meta = (EditCondition = "bOverrideMaterials"))
 	TArray<FNiagaraMeshMaterialOverride> OverrideMaterials;
 
+#if WITH_EDITORONLY_DATA
+	UPROPERTY(transient)
+	TArray<TObjectPtr<UMaterialInstanceConstant>> MICMaterials;
+#endif
+
 	/** When using SubImage lookups for particles, this variable contains the number of columns in X and the number of rows in Y.*/
 	UPROPERTY(EditAnywhere, Category = "SubUV")
 	FVector2D SubImageSize;
diff --git a/Engine/Plugins/FX/Niagara/Source/Niagara/Public/NiagaraRendererProperties.h b/Engine/Plugins/FX/Niagara/Source/Niagara/Public/NiagaraRendererProperties.h
--- a/Engine/Plugins/FX/Niagara/Source/Niagara/Public/NiagaraRendererProperties.h
+++ b/Engine/Plugins/FX/Niagara/Source/Niagara/Public/NiagaraRendererProperties.h
@@ -17,6 +17,7 @@ class FNiagaraSystemInstanceController;
 class FVertexFactoryType;
 class UMaterial;
 class UMaterialInterface;
+class UMaterialInstanceConstant;
 class UTexture;
 class FNiagaraEmitterInstance;
 class SWidget;
@@ -211,6 +212,19 @@ struct FNiagaraRendererMaterialTextureParameter
 	TObjectPtr<UTexture> Texture;
 };
 
+USTRUCT()
+struct FNiagaraRendererMaterialStaticBoolParameter
+{
+	GENERATED_BODY()
+
+	UPROPERTY(EditAnywhere, Category = "Material")
+	FName MaterialParameterName;
+
+	UPROPERTY(EditAnywhere, Category = "Material")
+	FName StaticVariableName;
+};
+
+
 /**
 * Parameters to apply to the material, these are both constant and dynamic bindings
 * Having any bindings set will cause a MID to be generated
@@ -232,7 +246,12 @@ struct FNiagaraRendererMaterialParameters
 	UPROPERTY(EditAnywhere, Category = "Material")
 	TArray<FNiagaraRendererMaterialTextureParameter> TextureParameters;
 
+	UPROPERTY(EditAnywhere, Category = "Material")
+	TArray<FNiagaraRendererMaterialStaticBoolParameter> StaticBoolParameters;
+
 #if WITH_EDITORONLY_DATA
+	void RenameVariable(const FNiagaraVariableBase& OldVariable, const FNiagaraVariableBase& NewVariable, const FVersionedNiagaraEmitter& InEmitter, ENiagaraRendererSourceDataMode SourceMode);
+	void RemoveVariable(const FNiagaraVariableBase& OldVariable, const FVersionedNiagaraEmitter& InEmitter, ENiagaraRendererSourceDataMode SourceMode);
 	void GetFeedback(TArrayView<UMaterialInterface*> Materials, TArray<FNiagaraRendererFeedback>& OutWarnings) const;
 #endif
 
@@ -402,6 +421,13 @@ protected:
 
 	/** utility function that can be used to fix up old vec3 bindings into position bindings. */
 	static void ChangeToPositionBinding(FNiagaraVariableAttributeBinding& Binding);
+
+	/** Update MIC Static Parameters. */
+	bool UpdateMaterialStaticParameters(const FNiagaraRendererMaterialParameters& MaterialParameters, UMaterialInstanceConstant* MIC);
+
+	/** Utility function to updates MICs. */
+	void UpdateMaterialParametersMIC(const FNiagaraRendererMaterialParameters& MaterialParameters, TObjectPtr<UMaterialInterface>& InOutMaterial, TObjectPtr<UMaterialInstanceConstant>& InOutMIC);
+	void UpdateMaterialParametersMIC(const FNiagaraRendererMaterialParameters& MaterialParameters, TArrayView<UMaterialInterface*> Materials, TArray<TObjectPtr<UMaterialInstanceConstant>>& InOutMICs);
 #endif
 
 #if WITH_EDITOR
diff --git a/Engine/Plugins/FX/Niagara/Source/Niagara/Public/NiagaraRibbonRendererProperties.h b/Engine/Plugins/FX/Niagara/Source/Niagara/Public/NiagaraRibbonRendererProperties.h
--- a/Engine/Plugins/FX/Niagara/Source/Niagara/Public/NiagaraRibbonRendererProperties.h
+++ b/Engine/Plugins/FX/Niagara/Source/Niagara/Public/NiagaraRibbonRendererProperties.h
@@ -9,6 +9,7 @@
 #include "NiagaraDataSetAccessor.h"
 #include "NiagaraRibbonRendererProperties.generated.h"
 
+class UMaterialInstanceConstant;
 class FNiagaraEmitterInstance;
 class FAssetThumbnailPool;
 class FVertexFactoryType;
@@ -200,6 +201,7 @@ public:
 	//UObject Interface
 	virtual void PostLoad() override;
 	virtual void PostInitProperties() override;
+	virtual void Serialize(FStructuredArchive::FRecord Record) override;
 #if WITH_EDITOR
 	virtual void PostEditChangeProperty(struct FPropertyChangedEvent& PropertyChangedEvent) override;
 	virtual bool CanEditChange(const FProperty* InProperty) const override;
@@ -239,6 +241,11 @@ public:
 	UPROPERTY(EditAnywhere, Category = "Ribbon Rendering")
 	TObjectPtr<UMaterialInterface> Material;
 
+#if WITH_EDITORONLY_DATA
+	UPROPERTY(transient)
+	TObjectPtr<UMaterialInstanceConstant> MICMaterial;
+#endif
+
 	/** Use the UMaterialInterface bound to this user variable if it is set to a valid value. If this is bound to a valid value and Material is also set, UserParamBinding wins.*/
 	UPROPERTY(EditAnywhere, Category = "Ribbon Rendering")
 	FNiagaraUserParameterBinding MaterialUserParamBinding;
@@ -491,6 +498,8 @@ protected:
 
 	void UpdateSourceModeDerivates(ENiagaraRendererSourceDataMode InSourceMode, bool bFromPropertyEdit);
 
+	void UpdateMICs();
+
 	virtual bool NeedsMIDsForMaterials() const { return MaterialParameters.HasAnyBindings(); }
 private: 
 	static TArray<TWeakObjectPtr<UNiagaraRibbonRendererProperties>> RibbonRendererPropertiesToDeferredInit;
diff --git a/Engine/Plugins/FX/Niagara/Source/Niagara/Public/NiagaraSpriteRendererProperties.h b/Engine/Plugins/FX/Niagara/Source/Niagara/Public/NiagaraSpriteRendererProperties.h
--- a/Engine/Plugins/FX/Niagara/Source/Niagara/Public/NiagaraSpriteRendererProperties.h
+++ b/Engine/Plugins/FX/Niagara/Source/Niagara/Public/NiagaraSpriteRendererProperties.h
@@ -9,6 +9,7 @@
 #include "Particles/SubUVAnimation.h"
 #include "NiagaraSpriteRendererProperties.generated.h"
 
+class UMaterialInstanceConstant;
 class FVertexFactoryType;
 
 /** This enum decides how a sprite particle will orient its "up" axis. Must keep these in sync with NiagaraSpriteVertexFactory.ush*/
@@ -149,6 +150,11 @@ public:
 	UPROPERTY(EditAnywhere, Category = "Sprite Rendering")
 	TObjectPtr<UMaterialInterface> Material;
 
+#if WITH_EDITORONLY_DATA
+	UPROPERTY(transient)
+	TObjectPtr<UMaterialInstanceConstant> MICMaterial;
+#endif
+
 	/** Whether or not to draw a single element for the Emitter or to draw the particles.*/
 	UPROPERTY(EditAnywhere, Category = "Sprite Rendering")
 	ENiagaraRendererSourceDataMode SourceMode;
@@ -399,6 +405,8 @@ protected:
 	void SetPreviousBindings(const FVersionedNiagaraEmitter& SrcEmitter, ENiagaraRendererSourceDataMode InSourceMode);
 	virtual void UpdateSourceModeDerivates(ENiagaraRendererSourceDataMode InSourceMode, bool bFromPropertyEdit = false) override;
 
+	void UpdateMICs();
+
 #if WITH_EDITORONLY_DATA
 	virtual FNiagaraVariable GetBoundAttribute(const FNiagaraVariableAttributeBinding* Binding) const override;
 #endif
diff --git a/Engine/Plugins/FX/Niagara/Source/NiagaraEditor/Private/Customizations/NiagaraTypeCustomizations.cpp b/Engine/Plugins/FX/Niagara/Source/NiagaraEditor/Private/Customizations/NiagaraTypeCustomizations.cpp
--- a/Engine/Plugins/FX/Niagara/Source/NiagaraEditor/Private/Customizations/NiagaraTypeCustomizations.cpp
+++ b/Engine/Plugins/FX/Niagara/Source/NiagaraEditor/Private/Customizations/NiagaraTypeCustomizations.cpp
@@ -2040,17 +2040,17 @@ void FNiagaraRendererMaterialParameterCustomization::CustomizeHeader(TSharedRef<
 	}
 
 	HeaderRow
-		.NameContent()
-		[
-			PropertyHandle->CreatePropertyNameWidget()
-		]
-		.ValueContent()
-		[
-			PropertyHandle->CreatePropertyValueWidget()
-		];
+	.NameContent()
+	[
+		PropertyHandle->CreatePropertyNameWidget()
+	]
+	.ValueContent()
+	[
+		PropertyHandle->CreatePropertyValueWidget()
+	];
 }
 
-void FNiagaraRendererMaterialParameterCustomization::CustomizeChildren(TSharedRef<IPropertyHandle> StructPropertyHandle, class IDetailChildrenBuilder& ChildBuilder, IPropertyTypeCustomizationUtils& StructCustomizationUtils)
+void FNiagaraRendererMaterialParameterCustomization::CustomizeChildren(TSharedRef<IPropertyHandle> StructPropertyHandle, IDetailChildrenBuilder& ChildBuilder, IPropertyTypeCustomizationUtils& StructCustomizationUtils)
 {
 	uint32 NumChildren;
 	StructPropertyHandle->GetNumChildren(NumChildren);
@@ -2061,8 +2061,6 @@ void FNiagaraRendererMaterialParameterCustomization::CustomizeChildren(TSharedRe
 		TSharedPtr<IPropertyHandle> ChildHandle = StructPropertyHandle->GetChildHandle(ChildIndex);
 		if ( ChildHandle->GetProperty()->GetFName() == NAME_MaterialParameterName )
 		{
-			MaterialBindingNameProperty = ChildHandle;
-
 			ChildBuilder.AddCustomRow(ChildHandle->GetPropertyDisplayName())
 			.NameWidget
 			[
@@ -2071,41 +2069,41 @@ void FNiagaraRendererMaterialParameterCustomization::CustomizeChildren(TSharedRe
 			.ValueWidget
 			[
 				SNew(SComboButton)
-				.OnGetMenuContent(this, &FNiagaraRendererMaterialParameterCustomization::OnGetMaterialBindingNameMenuContent)
+				.OnGetMenuContent(this, &FNiagaraRendererMaterialParameterCustomization::OnGetMaterialBindingNameMenuContent, ChildHandle)
 				.ContentPadding(1)
 				.ButtonStyle(FAppStyle::Get(), "PropertyEditor.AssetComboStyle")
 				.ForegroundColor(FAppStyle::GetColor("PropertyEditor.AssetName.ColorAndOpacity"))
 				.ButtonContent()
 				[
 					SNew(STextBlock)
-					.Text(this, &FNiagaraRendererMaterialParameterCustomization::GetMaterialBindingNameText)
+					.Text(this, &FNiagaraRendererMaterialParameterCustomization::GetBindingNameText, ChildHandle)
 					.Font(IDetailLayoutBuilder::GetDetailFont())
 				]
 			];
 		}
-		else
+		else if ( CustomizeChildProperty(ChildBuilder, ChildHandle) == false )
 		{
 			ChildBuilder.AddProperty(ChildHandle.ToSharedRef());
 		}
 	}
 }
 
-FText FNiagaraRendererMaterialParameterCustomization::GetMaterialBindingNameText() const
+FText FNiagaraRendererMaterialParameterCustomization::GetBindingNameText(TSharedPtr<IPropertyHandle> PropertyHandle) const
 {
 	FName BindingName;
-	if ( MaterialBindingNameProperty.IsValid() )
+	if (PropertyHandle.IsValid() )
 	{
-		MaterialBindingNameProperty->GetValue(BindingName);
+		PropertyHandle->GetValue(BindingName);
 	}
 	return FText::FromName(BindingName);
 }
 
-TSharedRef<SWidget> FNiagaraRendererMaterialParameterCustomization::OnGetMaterialBindingNameMenuContent() const
+TSharedRef<SWidget> FNiagaraRendererMaterialParameterCustomization::OnGetMaterialBindingNameMenuContent(TSharedPtr<IPropertyHandle> PropertyHandle) const
 {
 	FMenuBuilder MenuBuilder(true, NULL);
 
 	UNiagaraRendererProperties* RenderProperties = WeakRenderProperties.Get();
-	if (RenderProperties != nullptr && MaterialBindingNameProperty.IsValid())
+	if (RenderProperties != nullptr && PropertyHandle.IsValid())
 	{
 		TArray<UMaterialInterface*> UsedMaterials;
 		RenderProperties->GetUsedMaterials(nullptr, UsedMaterials);
@@ -2129,9 +2127,9 @@ TSharedRef<SWidget> FNiagaraRendererMaterialParameterCustomization::OnGetMateria
 				FSlateIcon(),
 				FUIAction(
 					FExecuteAction::CreateLambda(
-						[ValidBinding, PropertyToSet=MaterialBindingNameProperty]()
+						[ValidBinding, PropertyHandle]()
 						{
-							PropertyToSet->SetValue(ValidBinding);
+				PropertyHandle->SetValue(ValidBinding);
 						}
 					)
 				)
@@ -2179,5 +2177,109 @@ void FNiagaraRendererMaterialTextureParameterCustomization::GetMaterialBindingNa
 	}
 }
 
-#undef LOCTEXT_NAMESPACE
+bool FNiagaraRendererMaterialStaticBoolParameterCustomization::CustomizeChildProperty(class IDetailChildrenBuilder& ChildBuilder, TSharedPtr<IPropertyHandle> PropertyHandle)
+{
+	if (PropertyHandle->GetProperty()->GetFName() == GET_MEMBER_NAME_CHECKED(FNiagaraRendererMaterialStaticBoolParameter, StaticVariableName))
+	{
+		ChildBuilder.AddCustomRow(PropertyHandle->GetPropertyDisplayName())
+		.NameWidget
+		[
+			PropertyHandle->CreatePropertyNameWidget()
+		]
+		.ValueWidget
+		[
+			SNew(SComboButton)
+			.OnGetMenuContent(this, &FNiagaraRendererMaterialStaticBoolParameterCustomization::OnGetStaticVariablelBindingNameMenuContent, PropertyHandle)
+			.ContentPadding(1)
+			.ButtonStyle(FAppStyle::Get(), "PropertyEditor.AssetComboStyle")
+			.ForegroundColor(FAppStyle::GetColor("PropertyEditor.AssetName.ColorAndOpacity"))
+			.ButtonContent()
+			[
+				SNew(STextBlock)
+				.Text(this, &FNiagaraRendererMaterialParameterCustomization::GetBindingNameText, PropertyHandle)
+				.Font(IDetailLayoutBuilder::GetDetailFont())
+			]
+		];
 
+		return true;
+	}
+	return false;
+}
+
+void FNiagaraRendererMaterialStaticBoolParameterCustomization::GetMaterialBindingNames(UMaterialInterface* Material, TArray<FName>& OutBindings) const
+{
+	TArray<FMaterialParameterInfo> MaterialParameterInfos;
+	TArray<FGuid> Guids;
+	Material->GetAllStaticSwitchParameterInfo(MaterialParameterInfos, Guids);
+
+	for (const FMaterialParameterInfo& MaterialParameterInfo : MaterialParameterInfos)
+	{
+		OutBindings.AddUnique(MaterialParameterInfo.Name);
+	}
+}
+
+TSharedRef<SWidget> FNiagaraRendererMaterialStaticBoolParameterCustomization::OnGetStaticVariablelBindingNameMenuContent(TSharedPtr<IPropertyHandle> PropertyHandle) const
+{
+	FMenuBuilder MenuBuilder(true, NULL);
+
+	UNiagaraRendererProperties* RenderProperties = WeakRenderProperties.Get();
+	if (RenderProperties != nullptr && PropertyHandle.IsValid())
+	{
+		const FVersionedNiagaraEmitter NiagaraEmitter = RenderProperties->GetOuterEmitter();
+		const FNiagaraTypeDefinition StaticBoolDef = FNiagaraTypeDefinition::GetBoolDef().ToStaticDef();
+
+		if (FVersionedNiagaraEmitterData* EmitterData = NiagaraEmitter.GetEmitterData())
+		{
+			const FNiagaraAliasContext AliasContext = FNiagaraAliasContext().ChangeEmitterNameToEmitter(NiagaraEmitter.Emitter->GetUniqueEmitterName());
+
+			auto AddStaticVariables =
+				[&](UNiagaraScript* NiagaraScript)
+				{
+					if (NiagaraScript == nullptr)
+					{
+						return;
+					}
+					for ( const FNiagaraVariable& StaticVariable : NiagaraScript->GetVMExecutableData().StaticVariablesWritten )
+					{
+						if ( StaticVariable.GetType() != StaticBoolDef)
+						{
+							continue;
+						}
+						const FNiagaraVariableBase ResolvedVariable = FNiagaraUtilities::ResolveAliases(StaticVariable, AliasContext);
+						const bool bSystemAttribute = ResolvedVariable.IsInNameSpace(FNiagaraConstants::SystemNamespaceString);
+						const bool bEmitterAttribute = ResolvedVariable.IsInNameSpace(FNiagaraConstants::EmitterNamespaceString);
+						const bool bParticleAttribute = ResolvedVariable.IsInNameSpace(FNiagaraConstants::ParticleAttributeNamespaceString);
+						if (!bSystemAttribute && !bEmitterAttribute && !bParticleAttribute)
+						{
+							continue;
+						}
+
+						MenuBuilder.AddMenuEntry(
+							FText::FromName(ResolvedVariable.GetName()),
+							TAttribute<FText>(),
+							FSlateIcon(),
+							FUIAction(
+								FExecuteAction::CreateLambda(
+									[ValidBinding= ResolvedVariable.GetName(), PropertyHandle]()
+									{
+										PropertyHandle->SetValue(ValidBinding);
+									}
+								)
+							)
+						);
+					}
+				};
+
+			if (UNiagaraSystem* NiagaraSystem = RenderProperties->GetTypedOuter<UNiagaraSystem>())
+			{
+				AddStaticVariables(NiagaraSystem->GetSystemSpawnScript());
+				AddStaticVariables(NiagaraSystem->GetSystemUpdateScript());
+			}
+			EmitterData->ForEachScript(AddStaticVariables);
+		}
+	}
+
+	return MenuBuilder.MakeWidget();
+}
+
+#undef LOCTEXT_NAMESPACE
diff --git a/Engine/Plugins/FX/Niagara/Source/NiagaraEditor/Private/Customizations/NiagaraTypeCustomizations.h b/Engine/Plugins/FX/Niagara/Source/NiagaraEditor/Private/Customizations/NiagaraTypeCustomizations.h
--- a/Engine/Plugins/FX/Niagara/Source/NiagaraEditor/Private/Customizations/NiagaraTypeCustomizations.h
+++ b/Engine/Plugins/FX/Niagara/Source/NiagaraEditor/Private/Customizations/NiagaraTypeCustomizations.h
@@ -381,14 +381,14 @@ public:
 	virtual void CustomizeChildren(TSharedRef<IPropertyHandle> StructPropertyHandle, class IDetailChildrenBuilder& ChildBuilder, IPropertyTypeCustomizationUtils& StructCustomizationUtils) override;
 	// IPropertyTypeCustomization interface end
 
-	FText GetMaterialBindingNameText() const;
-	TSharedRef<SWidget> OnGetMaterialBindingNameMenuContent() const;
+	FText GetBindingNameText(TSharedPtr<IPropertyHandle> PropertyHandle) const;
+	TSharedRef<SWidget> OnGetMaterialBindingNameMenuContent(TSharedPtr<IPropertyHandle> PropertyHandle) const;
 
+	virtual bool CustomizeChildProperty(class IDetailChildrenBuilder& ChildBuilder, TSharedPtr<IPropertyHandle> PropertyHandle) { return false; }
 	virtual void GetMaterialBindingNames(UMaterialInterface* Material, TArray<FName>& OutBindings) const = 0;
 
-private:
+protected:
 	TWeakObjectPtr<UNiagaraRendererProperties>	WeakRenderProperties;
-	TSharedPtr<IPropertyHandle>					MaterialBindingNameProperty;
 };
 
 class FNiagaraRendererMaterialScalarParameterCustomization : public FNiagaraRendererMaterialParameterCustomization
@@ -423,3 +423,16 @@ public:
 
 	virtual void GetMaterialBindingNames(UMaterialInterface* Material, TArray<FName>& OutBindings) const override;
 };
+
+class FNiagaraRendererMaterialStaticBoolParameterCustomization : public FNiagaraRendererMaterialParameterCustomization
+{
+public:
+	static TSharedRef<IPropertyTypeCustomization> MakeInstance()
+	{
+		return MakeShared<FNiagaraRendererMaterialStaticBoolParameterCustomization>();
+	}
+
+	virtual bool CustomizeChildProperty(class IDetailChildrenBuilder& ChildBuilder, TSharedPtr<IPropertyHandle> PropertyHandle) override;
+	virtual void GetMaterialBindingNames(UMaterialInterface* Material, TArray<FName>& OutBindings) const override;
+	TSharedRef<SWidget> OnGetStaticVariablelBindingNameMenuContent(TSharedPtr<IPropertyHandle> PropertyHandle) const;
+};
diff --git a/Engine/Plugins/FX/Niagara/Source/NiagaraEditor/Private/NiagaraEditorModule.cpp b/Engine/Plugins/FX/Niagara/Source/NiagaraEditor/Private/NiagaraEditorModule.cpp
--- a/Engine/Plugins/FX/Niagara/Source/NiagaraEditor/Private/NiagaraEditorModule.cpp
+++ b/Engine/Plugins/FX/Niagara/Source/NiagaraEditor/Private/NiagaraEditorModule.cpp
@@ -1025,6 +1025,10 @@ void FNiagaraEditorModule::StartupModule()
 		FNiagaraRendererMaterialTextureParameter::StaticStruct()->GetFName(),
 		FOnGetPropertyTypeCustomizationInstance::CreateStatic(&FNiagaraRendererMaterialTextureParameterCustomization::MakeInstance));
 
+	PropertyModule.RegisterCustomPropertyTypeLayout(
+		FNiagaraRendererMaterialStaticBoolParameter::StaticStruct()->GetFName(),
+		FOnGetPropertyTypeCustomizationInstance::CreateStatic(&FNiagaraRendererMaterialStaticBoolParameterCustomization::MakeInstance));
+
 	// Outliner Customizations end.
 
 	//Register Stack Object Issue Generators.
diff --git a/Engine/Source/Runtime/Engine/Public/StaticParameterSet.h b/Engine/Source/Runtime/Engine/Public/StaticParameterSet.h
--- a/Engine/Source/Runtime/Engine/Public/StaticParameterSet.h
+++ b/Engine/Source/Runtime/Engine/Public/StaticParameterSet.h
@@ -522,7 +522,7 @@ struct FStaticParameterSet : public FStaticParameterSetRuntimeData
 	bool Equivalent(const FStaticParameterSet& ReferenceSet) const;
 
 #if WITH_EDITORONLY_DATA
-	void SetParameterValue(const FMaterialParameterInfo& ParameterInfo, const FMaterialParameterMetadata& Meta, EMaterialSetParameterValueFlags Flags = EMaterialSetParameterValueFlags::None);
+	ENGINE_API void SetParameterValue(const FMaterialParameterInfo& ParameterInfo, const FMaterialParameterMetadata& Meta, EMaterialSetParameterValueFlags Flags = EMaterialSetParameterValueFlags::None);
 	void AddParametersOfType(EMaterialParameterType Type, const TMap<FMaterialParameterInfo, FMaterialParameterMetadata>& Values);
 #endif // WITH_EDITORONLY_DATA
 
