- IWYU fix post header reorganization

#rnx

diff --git a/Engine/Plugins/FX/Niagara/Source/Niagara/Private/NiagaraRendererProperties.cpp b/Engine/Plugins/FX/Niagara/Source/Niagara/Private/NiagaraRendererProperties.cpp
--- a/Engine/Plugins/FX/Niagara/Source/Niagara/Private/NiagaraRendererProperties.cpp
+++ b/Engine/Plugins/FX/Niagara/Source/Niagara/Private/NiagaraRendererProperties.cpp
@@ -8,6 +8,7 @@
 #include "NiagaraScriptSourceBase.h"
 #include "NiagaraSettings.h"
 #include "NiagaraSystem.h"
+#include "NiagaraSystemImpl.h"
 
 #include "Interfaces/ITargetPlatform.h"
 #include "Materials/MaterialInterface.h"
