Converted Niagara asset type actions to asset definitions

diff --git a/Engine/Plugins/FX/Niagara/Source/NiagaraEditor/NiagaraEditor.Build.cs b/Engine/Plugins/FX/Niagara/Source/NiagaraEditor/NiagaraEditor.Build.cs
--- a/Engine/Plugins/FX/Niagara/Source/NiagaraEditor/NiagaraEditor.Build.cs
+++ b/Engine/Plugins/FX/Niagara/Source/NiagaraEditor/NiagaraEditor.Build.cs
@@ -61,6 +61,8 @@ public class NiagaraEditor : ModuleRules
 				"DeveloperSettings",
 				"PythonScriptPlugin",
 				"ImageWrapper",
+				"AssetDefinition",
+				"ContentBrowser"
 			}
 		);
 
diff --git a/Engine/Plugins/FX/Niagara/Source/NiagaraEditor/Private/AssetDefinitions/AssetDefinition_NiagaraEffectType.cpp b/Engine/Plugins/FX/Niagara/Source/NiagaraEditor/Private/AssetDefinitions/AssetDefinition_NiagaraEffectType.cpp
new file mode 100644
--- /dev/null
+++ b/Engine/Plugins/FX/Niagara/Source/NiagaraEditor/Private/AssetDefinitions/AssetDefinition_NiagaraEffectType.cpp
@@ -0,0 +1,26 @@
+// Copyright Epic Games, Inc. All Rights Reserved.
+
+#include "AssetDefinitions/AssetDefinition_NiagaraEffectType.h"
+
+#include "NiagaraEditorStyle.h"
+#include "NiagaraEffectType.h"
+
+#define LOCTEXT_NAMESPACE "UAssetDefinition_NiagaraEffectType"
+
+FLinearColor UAssetDefinition_NiagaraEffectType::GetAssetColor() const
+{
+	return FNiagaraEditorStyle::Get().GetColor("NiagaraEditor.AssetColors.EffectType");
+}
+
+TSoftClassPtr<> UAssetDefinition_NiagaraEffectType::GetAssetClass() const
+{
+	return UNiagaraEffectType::StaticClass();
+}
+
+TConstArrayView<FAssetCategoryPath> UAssetDefinition_NiagaraEffectType::GetAssetCategories() const
+{
+	static FAssetCategoryPath AssetPaths[] = { FAssetCategoryPath(LOCTEXT("NiagaraAssetsCategory", "FX"), LOCTEXT("NiagaraEffectType_SubCategory", "Advanced")) };
+	return AssetPaths;
+}
+
+#undef LOCTEXT_NAMESPACE
diff --git a/Engine/Plugins/FX/Niagara/Source/NiagaraEditor/Private/AssetDefinitions/AssetDefinition_NiagaraEffectType.h b/Engine/Plugins/FX/Niagara/Source/NiagaraEditor/Private/AssetDefinitions/AssetDefinition_NiagaraEffectType.h
new file mode 100644
--- /dev/null
+++ b/Engine/Plugins/FX/Niagara/Source/NiagaraEditor/Private/AssetDefinitions/AssetDefinition_NiagaraEffectType.h
@@ -0,0 +1,21 @@
+// Copyright Epic Games, Inc. All Rights Reserved.
+
+#pragma once
+
+#include "CoreMinimal.h"
+#include "AssetDefinitionDefault.h"
+
+#include "AssetDefinition_NiagaraEffectType.generated.h"
+
+UCLASS()
+class UAssetDefinition_NiagaraEffectType : public UAssetDefinitionDefault
+{
+	GENERATED_BODY()
+
+public:
+	// UAssetDefinition Implementation
+	virtual FText GetAssetDisplayName() const override { return NSLOCTEXT("AssetDefinition", "NiagaraEffectType", "Niagara Effect Type"); }
+	virtual FLinearColor GetAssetColor() const override;
+	virtual TSoftClassPtr<UObject> GetAssetClass() const override;
+	virtual TConstArrayView<FAssetCategoryPath> GetAssetCategories() const override;
+};
diff --git a/Engine/Plugins/FX/Niagara/Source/NiagaraEditor/Private/AssetDefinitions/AssetDefinition_NiagaraEmitter.cpp b/Engine/Plugins/FX/Niagara/Source/NiagaraEditor/Private/AssetDefinitions/AssetDefinition_NiagaraEmitter.cpp
new file mode 100644
--- /dev/null
+++ b/Engine/Plugins/FX/Niagara/Source/NiagaraEditor/Private/AssetDefinitions/AssetDefinition_NiagaraEmitter.cpp
@@ -0,0 +1,223 @@
+// Copyright Epic Games, Inc. All Rights Reserved.
+
+#include "AssetDefinitions/AssetDefinition_NiagaraEmitter.h"
+
+#include "AssetToolsModule.h"
+#include "ContentBrowserMenuContexts.h"
+#include "ContentBrowserModule.h"
+#include "IContentBrowserSingleton.h"
+#include "ToolMenus.h"
+#include "ToolMenuDelegates.h"
+#include "NiagaraEditorStyle.h"
+#include "NiagaraEditorUtilities.h"
+#include "NiagaraEmitter.h"
+#include "NiagaraScriptSourceBase.h"
+#include "NiagaraSystemFactoryNew.h"
+#include "NiagaraSystemToolkit.h"
+#include "ViewModels/NiagaraSystemViewModel.h"
+
+#define LOCTEXT_NAMESPACE "UAssetDefinition_NiagaraEmitter"
+
+FLinearColor UAssetDefinition_NiagaraEmitter::GetAssetColor() const
+{
+	return FNiagaraEditorStyle::Get().GetColor("NiagaraEditor.AssetColors.Emitter");
+}
+
+TSoftClassPtr<> UAssetDefinition_NiagaraEmitter::GetAssetClass() const
+{
+	return UNiagaraEmitter::StaticClass();
+}
+
+TConstArrayView<FAssetCategoryPath> UAssetDefinition_NiagaraEmitter::GetAssetCategories() const
+{
+	static FAssetCategoryPath AssetPaths[] = { LOCTEXT("NiagaraAssetsCategory", "FX") };
+	return AssetPaths;
+}
+
+EAssetCommandResult UAssetDefinition_NiagaraEmitter::OpenAssets(const FAssetOpenArgs& OpenArgs) const
+{
+	if (OpenArgs.OpenMethod == EAssetOpenMethod::Edit)
+	{
+		EToolkitMode::Type Mode = OpenArgs.ToolkitHost.IsValid() ? EToolkitMode::WorldCentric : EToolkitMode::Standalone;
+		TArray<UObject*> Objects = OpenArgs.LoadObjects<UObject>();
+		for (UObject* OpenObj : Objects)
+		{
+			if (UNiagaraEmitter* Emitter = Cast<UNiagaraEmitter>(OpenObj))
+			{
+				if (!Emitter->VersionToOpenInEditor.IsValid())
+				{
+					Emitter->VersionToOpenInEditor = Emitter->GetExposedVersion().VersionGuid;
+				}
+				TSharedRef<FNiagaraSystemToolkit> SystemToolkit(new FNiagaraSystemToolkit());
+				SystemToolkit->InitializeWithEmitter(Mode, OpenArgs.ToolkitHost, *Emitter);
+			}
+		}
+		return EAssetCommandResult::Handled;
+	}
+
+	return EAssetCommandResult::Unhandled;
+}
+
+namespace MenuExtension_Emitter
+{
+	void CreateUniqueAssetName(const FString& InBasePackageName, const FString& InSuffix, FString& OutPackageName, FString& OutAssetName)
+	{
+		FAssetToolsModule& AssetToolsModule = FModuleManager::Get().LoadModuleChecked<FAssetToolsModule>("AssetTools");
+		AssetToolsModule.Get().CreateUniqueAssetName(InBasePackageName, InSuffix, OutPackageName, OutAssetName);
+	}
+	
+	static void ExecuteNewNiagaraSystem(const FToolMenuContext& MenuContext)
+	{
+		if (const UContentBrowserAssetContextMenuContext* Context = UContentBrowserAssetContextMenuContext::FindContextWithAssets(MenuContext))
+		{
+			TArray<UNiagaraEmitter*> Emitters = Context->LoadSelectedObjects<UNiagaraEmitter>();
+			const FString DefaultSuffix = TEXT("_System");
+
+			FNiagaraSystemViewModelOptions SystemOptions;
+			SystemOptions.bCanModifyEmittersFromTimeline = true;
+			SystemOptions.EditMode = ENiagaraSystemViewModelEditMode::SystemAsset;
+
+			TArray<UObject*> ObjectsToSync;
+			for (UNiagaraEmitter* Emitter : Emitters)
+			{
+				if (Emitter)
+				{
+					// Determine an appropriate names
+					FString Name;
+					FString PackageName;
+					CreateUniqueAssetName(Emitter->GetOutermost()->GetName(), DefaultSuffix, PackageName, Name);
+
+					// Create the factory used to generate the asset
+					UNiagaraSystemFactoryNew* Factory = NewObject<UNiagaraSystemFactoryNew>();
+					Factory->EmittersToAddToNewSystem.Add(FVersionedNiagaraEmitter(Emitter, Emitter->GetExposedVersion().VersionGuid));
+					FAssetToolsModule& AssetToolsModule = FModuleManager::GetModuleChecked<FAssetToolsModule>("AssetTools");
+					UObject* NewAsset = AssetToolsModule.Get().CreateAsset(Name, FPackageName::GetLongPackagePath(PackageName), UNiagaraSystem::StaticClass(), Factory);
+			
+					UNiagaraSystem* System = Cast<UNiagaraSystem>(NewAsset);
+					if (System != nullptr)
+					{
+						ObjectsToSync.Add(NewAsset);
+					}
+				}
+			}
+
+			if (ObjectsToSync.Num() > 0)
+			{
+				FContentBrowserModule& ContentBrowserModule = FModuleManager::LoadModuleChecked<FContentBrowserModule>("ContentBrowser");
+				ContentBrowserModule.Get().SyncBrowserToAssets(ObjectsToSync);
+			}
+		}
+	}
+
+	static void ExecuteCreateDuplicateParent(const FToolMenuContext& MenuContext)
+	{
+		if (const UContentBrowserAssetContextMenuContext* Context = UContentBrowserAssetContextMenuContext::FindContextWithAssets(MenuContext))
+		{
+			TArray<UNiagaraEmitter*> Emitters = Context->LoadSelectedObjects<UNiagaraEmitter>();
+			const FString DefaultSuffix = TEXT("_Parent");
+
+			TArray<UObject*> ObjectsToSync;
+			FAssetToolsModule& AssetToolsModule = FModuleManager::GetModuleChecked<FAssetToolsModule>("AssetTools");
+
+			for (UNiagaraEmitter* Emitter : Emitters)
+			{
+				if (Emitter == nullptr)
+				{
+					continue;
+				}
+				FString Name;
+				FString PackageName;
+				CreateUniqueAssetName(Emitter->GetOutermost()->GetName(), DefaultSuffix, PackageName, Name);
+
+				// Create the factory used to generate the asset
+
+				if (UNiagaraEmitter* NewEmitter = Cast<UNiagaraEmitter>(AssetToolsModule.Get().DuplicateAssetWithDialog(Name, FPackageName::GetLongPackagePath(PackageName), Emitter)))
+				{
+					GEditor->GetEditorSubsystem<UAssetEditorSubsystem>()->CloseAllEditorsForAsset(Emitter);
+
+					NewEmitter->Modify();
+					NewEmitter->GetLatestEmitterData()->GraphSource->MarkNotSynchronized(TEXT("Emitter created"));
+
+					Emitter->Modify();
+					Emitter->SetParent(FVersionedNiagaraEmitter(NewEmitter, NewEmitter->GetExposedVersion().VersionGuid));
+					Emitter->GetLatestEmitterData()->GraphSource->MarkNotSynchronized(TEXT("Emitter parent changed"));
+
+					ObjectsToSync.Add(NewEmitter);
+				}
+			}
+
+			if (ObjectsToSync.Num() > 0)
+			{
+				AssetToolsModule.Get().OpenEditorForAssets(ObjectsToSync);
+
+				FContentBrowserModule& ContentBrowserModule = FModuleManager::LoadModuleChecked<FContentBrowserModule>("ContentBrowser");
+				ContentBrowserModule.Get().SyncBrowserToAssets(ObjectsToSync);
+			}
+		}
+	}
+
+	static bool CanCreateDuplicateParent(const FToolMenuContext& MenuContext)
+	{
+		if (const UContentBrowserAssetContextMenuContext* Context = UContentBrowserAssetContextMenuContext::FindContextWithAssets(MenuContext))
+		{
+			TArray<UNiagaraEmitter*> Emitters = Context->LoadSelectedObjects<UNiagaraEmitter>();
+			for (UNiagaraEmitter* EmitterPtr : Emitters)
+			{
+				if (EmitterPtr && EmitterPtr->IsVersioningEnabled())
+				{
+					return false;
+				}
+			}
+		}
+		return true;
+	}
+
+	static void MarkDependentCompilableAssetsDirty(const FToolMenuContext& MenuContext)
+	{
+		if (const UContentBrowserAssetContextMenuContext* Context = UContentBrowserAssetContextMenuContext::FindContextWithAssets(MenuContext))
+		{
+			FNiagaraEditorUtilities::MarkDependentCompilableAssetsDirty(Context->SelectedAssets);
+		}
+	}
+	
+	static FDelayedAutoRegisterHelper DelayedAutoRegister(EDelayedRegisterRunPhase::EndOfEngineInit, []
+	{
+		UToolMenus::RegisterStartupCallback(FSimpleMulticastDelegate::FDelegate::CreateLambda([]()
+		{
+			FToolMenuOwnerScoped OwnerScoped(UE_MODULE_NAME);
+			UToolMenu* Menu = UE::ContentBrowser::ExtendToolMenu_AssetContextMenu(UNiagaraEmitter::StaticClass());
+	        
+			FToolMenuSection& Section = Menu->FindOrAddSection("GetAssetActions");
+			Section.AddDynamicEntry("GetAssetActions_Emitter", FNewToolMenuSectionDelegate::CreateLambda([](FToolMenuSection& InSection)
+			{
+				{
+					const TAttribute Label = LOCTEXT("Emitter_NewNiagaraSystem", "Create Niagara System");
+					const TAttribute ToolTip = LOCTEXT("Emitter_NewNiagaraSystemTooltip", "Creates a Niagara system using this emitter as a base.");
+					const FSlateIcon Icon = FSlateIcon(FAppStyle::GetAppStyleSetName(), "ClassIcon.ParticleSystem");
+					const FToolMenuExecuteAction UIAction = FToolMenuExecuteAction::CreateStatic(&ExecuteNewNiagaraSystem);
+
+					InSection.AddMenuEntry("Emitter_NewNiagaraSystem", Label, ToolTip, Icon, UIAction);
+				}
+				{
+					const TAttribute Label = LOCTEXT("Emitter_CreateDuplicateParentLabel", "Create Duplicate Parent");
+					const TAttribute ToolTip = LOCTEXT("Emitter_CreateDuplicateParentTooltip", "Duplicate this emitter and set this emitter's parent to the new emitter.");
+					const FSlateIcon Icon = FSlateIcon();
+					FToolUIAction UIAction = FToolMenuExecuteAction::CreateStatic(&ExecuteCreateDuplicateParent);
+					UIAction.CanExecuteAction = FToolMenuCanExecuteAction::CreateStatic(&CanCreateDuplicateParent);
+
+					InSection.AddMenuEntry("Emitter_CreateDuplicateParent", Label, ToolTip, Icon, UIAction);
+				}
+				{
+					const TAttribute Label = LOCTEXT("MarkDependentCompilableAssetsDirtyLabel", "Mark Dependent Compilable Assets Dirty");
+					const TAttribute ToolTip = LOCTEXT("MarkDependentCompilableAssetsDirtyToolTip", "Finds all niagara assets which depend on this asset either directly or indirectly,\n and marks them dirty so they can be saved with the latest version.");
+					const FSlateIcon Icon = FSlateIcon();
+					const FToolMenuExecuteAction UIAction = FToolMenuExecuteAction::CreateStatic(&MarkDependentCompilableAssetsDirty);
+
+					InSection.AddMenuEntry("MarkDependentCompilableAssetsDirty", Label, ToolTip, Icon, UIAction);
+				}
+			}));
+		}));
+	});
+}
+
+#undef LOCTEXT_NAMESPACE
diff --git a/Engine/Plugins/FX/Niagara/Source/NiagaraEditor/Private/AssetDefinitions/AssetDefinition_NiagaraEmitter.h b/Engine/Plugins/FX/Niagara/Source/NiagaraEditor/Private/AssetDefinitions/AssetDefinition_NiagaraEmitter.h
new file mode 100644
--- /dev/null
+++ b/Engine/Plugins/FX/Niagara/Source/NiagaraEditor/Private/AssetDefinitions/AssetDefinition_NiagaraEmitter.h
@@ -0,0 +1,22 @@
+// Copyright Epic Games, Inc. All Rights Reserved.
+
+#pragma once
+
+#include "CoreMinimal.h"
+#include "AssetDefinitionDefault.h"
+
+#include "AssetDefinition_NiagaraEmitter.generated.h"
+
+UCLASS()
+class UAssetDefinition_NiagaraEmitter : public UAssetDefinitionDefault
+{
+	GENERATED_BODY()
+
+public:
+	// UAssetDefinition Implementation
+	virtual FText GetAssetDisplayName() const override { return NSLOCTEXT("AssetDefinition", "NiagaraEmitter", "Niagara Emitter"); }
+	virtual FLinearColor GetAssetColor() const override;
+	virtual TSoftClassPtr<UObject> GetAssetClass() const override;
+	virtual TConstArrayView<FAssetCategoryPath> GetAssetCategories() const override;
+	virtual EAssetCommandResult OpenAssets(const FAssetOpenArgs& OpenArgs) const override;
+};
diff --git a/Engine/Plugins/FX/Niagara/Source/NiagaraEditor/Private/AssetDefinitions/AssetDefinition_NiagaraParameterDefinitions.cpp b/Engine/Plugins/FX/Niagara/Source/NiagaraEditor/Private/AssetDefinitions/AssetDefinition_NiagaraParameterDefinitions.cpp
new file mode 100644
--- /dev/null
+++ b/Engine/Plugins/FX/Niagara/Source/NiagaraEditor/Private/AssetDefinitions/AssetDefinition_NiagaraParameterDefinitions.cpp
@@ -0,0 +1,47 @@
+// Copyright Epic Games, Inc. All Rights Reserved.
+
+#include "AssetDefinitions/AssetDefinition_NiagaraParameterDefinitions.h"
+
+#include "NiagaraEditorStyle.h"
+#include "NiagaraParameterDefinitions.h"
+#include "NiagaraParameterDefinitionsToolkit.h"
+
+#define LOCTEXT_NAMESPACE "UAssetDefinition_NiagaraParameterDefinitions"
+
+FLinearColor UAssetDefinition_NiagaraParameterDefinitions::GetAssetColor() const
+{
+	return FNiagaraEditorStyle::Get().GetColor("NiagaraEditor.AssetColors.ParameterDefinitions");
+}
+
+TSoftClassPtr<> UAssetDefinition_NiagaraParameterDefinitions::GetAssetClass() const
+{
+	return UNiagaraParameterDefinitions::StaticClass();
+}
+
+TConstArrayView<FAssetCategoryPath> UAssetDefinition_NiagaraParameterDefinitions::GetAssetCategories() const
+{
+	static FAssetCategoryPath AssetPaths[] = { FAssetCategoryPath(LOCTEXT("NiagaraAssetsCategory", "FX"), LOCTEXT("NiagaraParameterDefinitions_SubCategory", "Advanced")) };
+	return AssetPaths;
+}
+
+EAssetCommandResult UAssetDefinition_NiagaraParameterDefinitions::OpenAssets(const FAssetOpenArgs& OpenArgs) const
+{
+	if (OpenArgs.OpenMethod == EAssetOpenMethod::Edit)
+	{
+		EToolkitMode::Type Mode = OpenArgs.ToolkitHost.IsValid() ? EToolkitMode::WorldCentric : EToolkitMode::Standalone;
+		TArray<UObject*> Objects = OpenArgs.LoadObjects<UObject>();
+		for (UObject* OpenObj : Objects)
+		{
+			if (UNiagaraParameterDefinitions* ParameterDefinitions = Cast<UNiagaraParameterDefinitions>(OpenObj))
+			{
+				TSharedRef< FNiagaraParameterDefinitionsToolkit > NewNiagaraParameterDefinitionsToolkit(new FNiagaraParameterDefinitionsToolkit());
+				NewNiagaraParameterDefinitionsToolkit->Initialize(Mode, OpenArgs.ToolkitHost, ParameterDefinitions);
+			}
+		}
+		return EAssetCommandResult::Handled;
+	}
+
+	return EAssetCommandResult::Unhandled;
+}
+
+#undef LOCTEXT_NAMESPACE
diff --git a/Engine/Plugins/FX/Niagara/Source/NiagaraEditor/Private/AssetDefinitions/AssetDefinition_NiagaraParameterDefinitions.h b/Engine/Plugins/FX/Niagara/Source/NiagaraEditor/Private/AssetDefinitions/AssetDefinition_NiagaraParameterDefinitions.h
new file mode 100644
--- /dev/null
+++ b/Engine/Plugins/FX/Niagara/Source/NiagaraEditor/Private/AssetDefinitions/AssetDefinition_NiagaraParameterDefinitions.h
@@ -0,0 +1,22 @@
+// Copyright Epic Games, Inc. All Rights Reserved.
+
+#pragma once
+
+#include "CoreMinimal.h"
+#include "AssetDefinitionDefault.h"
+
+#include "AssetDefinition_NiagaraParameterDefinitions.generated.h"
+
+UCLASS()
+class UAssetDefinition_NiagaraParameterDefinitions : public UAssetDefinitionDefault
+{
+	GENERATED_BODY()
+
+public:
+	// UAssetDefinition Implementation
+	virtual FText GetAssetDisplayName() const override { return NSLOCTEXT("AssetDefinition", "NiagaraParameterDefinitions", "Niagara Parameter Definitions"); }
+	virtual FLinearColor GetAssetColor() const override;
+	virtual TSoftClassPtr<UObject> GetAssetClass() const override;
+	virtual TConstArrayView<FAssetCategoryPath> GetAssetCategories() const override;
+	virtual EAssetCommandResult OpenAssets(const FAssetOpenArgs& OpenArgs) const override;
+};
diff --git a/Engine/Plugins/FX/Niagara/Source/NiagaraEditor/Private/AssetDefinitions/AssetDefinition_NiagaraSimCache.cpp b/Engine/Plugins/FX/Niagara/Source/NiagaraEditor/Private/AssetDefinitions/AssetDefinition_NiagaraSimCache.cpp
new file mode 100644
--- /dev/null
+++ b/Engine/Plugins/FX/Niagara/Source/NiagaraEditor/Private/AssetDefinitions/AssetDefinition_NiagaraSimCache.cpp
@@ -0,0 +1,47 @@
+// Copyright Epic Games, Inc. All Rights Reserved.
+
+#include "AssetDefinitions/AssetDefinition_NiagaraSimCache.h"
+
+#include "NiagaraEditorStyle.h"
+#include "NiagaraSimCache.h"
+#include "NiagaraSimCacheToolkit.h"
+
+#define LOCTEXT_NAMESPACE "UAssetDefinition_NiagaraSimCache"
+
+FLinearColor UAssetDefinition_NiagaraSimCache::GetAssetColor() const
+{
+	return FNiagaraEditorStyle::Get().GetColor("NiagaraEditor.AssetColors.SimCache");
+}
+
+TSoftClassPtr<> UAssetDefinition_NiagaraSimCache::GetAssetClass() const
+{
+	return UNiagaraSimCache::StaticClass();
+}
+
+TConstArrayView<FAssetCategoryPath> UAssetDefinition_NiagaraSimCache::GetAssetCategories() const
+{
+	static FAssetCategoryPath AssetPaths[] = { FAssetCategoryPath(LOCTEXT("NiagaraAssetsCategory", "FX"), LOCTEXT("NiagaraSimCache_SubCategory", "Advanced")) };
+	return AssetPaths;
+}
+
+EAssetCommandResult UAssetDefinition_NiagaraSimCache::OpenAssets(const FAssetOpenArgs& OpenArgs) const
+{
+	if (OpenArgs.OpenMethod == EAssetOpenMethod::Edit)
+	{
+		EToolkitMode::Type Mode = OpenArgs.ToolkitHost.IsValid() ? EToolkitMode::WorldCentric : EToolkitMode::Standalone;
+		TArray<UObject*> Objects = OpenArgs.LoadObjects<UObject>();
+		for (UObject* OpenObj : Objects)
+		{
+			if (UNiagaraSimCache* Cache = Cast<UNiagaraSimCache>(OpenObj))
+			{
+				TSharedRef<FNiagaraSimCacheToolkit> NewNiagaraSystemToolkit(new FNiagaraSimCacheToolkit());
+				NewNiagaraSystemToolkit->Initialize(Mode, OpenArgs.ToolkitHost, Cache);
+			}
+		}
+		return EAssetCommandResult::Handled;
+	}
+
+	return EAssetCommandResult::Unhandled;
+}
+
+#undef LOCTEXT_NAMESPACE
diff --git a/Engine/Plugins/FX/Niagara/Source/NiagaraEditor/Private/AssetDefinitions/AssetDefinition_NiagaraSimCache.h b/Engine/Plugins/FX/Niagara/Source/NiagaraEditor/Private/AssetDefinitions/AssetDefinition_NiagaraSimCache.h
new file mode 100644
--- /dev/null
+++ b/Engine/Plugins/FX/Niagara/Source/NiagaraEditor/Private/AssetDefinitions/AssetDefinition_NiagaraSimCache.h
@@ -0,0 +1,22 @@
+// Copyright Epic Games, Inc. All Rights Reserved.
+
+#pragma once
+
+#include "CoreMinimal.h"
+#include "AssetDefinitionDefault.h"
+
+#include "AssetDefinition_NiagaraSimCache.generated.h"
+
+UCLASS()
+class UAssetDefinition_NiagaraSimCache : public UAssetDefinitionDefault
+{
+	GENERATED_BODY()
+
+public:
+	// UAssetDefinition Implementation
+	virtual FText GetAssetDisplayName() const override { return NSLOCTEXT("AssetDefinition", "NiagaraSimCache", "Niagara Simulation Cache"); }
+	virtual FLinearColor GetAssetColor() const override;
+	virtual TSoftClassPtr<UObject> GetAssetClass() const override;
+	virtual TConstArrayView<FAssetCategoryPath> GetAssetCategories() const override;
+	virtual EAssetCommandResult OpenAssets(const FAssetOpenArgs& OpenArgs) const override;
+};
diff --git a/Engine/Plugins/FX/Niagara/Source/NiagaraEditor/Private/AssetDefinitions/AssetDefinition_NiagaraSystem.cpp b/Engine/Plugins/FX/Niagara/Source/NiagaraEditor/Private/AssetDefinitions/AssetDefinition_NiagaraSystem.cpp
new file mode 100644
--- /dev/null
+++ b/Engine/Plugins/FX/Niagara/Source/NiagaraEditor/Private/AssetDefinitions/AssetDefinition_NiagaraSystem.cpp
@@ -0,0 +1,47 @@
+// Copyright Epic Games, Inc. All Rights Reserved.
+
+#include "AssetDefinitions/AssetDefinition_NiagaraSystem.h"
+
+#include "NiagaraEditorStyle.h"
+#include "NiagaraSystem.h"
+#include "NiagaraSystemToolkit.h"
+
+#define LOCTEXT_NAMESPACE "UAssetDefinition_NiagaraSystem"
+
+FLinearColor UAssetDefinition_NiagaraSystem::GetAssetColor() const
+{
+	return FNiagaraEditorStyle::Get().GetColor("NiagaraEditor.AssetColors.System");
+}
+
+TSoftClassPtr<> UAssetDefinition_NiagaraSystem::GetAssetClass() const
+{
+	return UNiagaraSystem::StaticClass();
+}
+
+TConstArrayView<FAssetCategoryPath> UAssetDefinition_NiagaraSystem::GetAssetCategories() const
+{
+	static FAssetCategoryPath AssetPaths[] = { LOCTEXT("NiagaraAssetsCategory", "FX"), EAssetCategoryPaths::Basic };
+	return AssetPaths;
+}
+
+EAssetCommandResult UAssetDefinition_NiagaraSystem::OpenAssets(const FAssetOpenArgs& OpenArgs) const
+{
+	if (OpenArgs.OpenMethod == EAssetOpenMethod::Edit)
+	{
+		EToolkitMode::Type Mode = OpenArgs.ToolkitHost.IsValid() ? EToolkitMode::WorldCentric : EToolkitMode::Standalone;
+		TArray<UObject*> Objects = OpenArgs.LoadObjects<UObject>();
+		for (UObject* OpenObj : Objects)
+		{
+			if (UNiagaraSystem* System = Cast<UNiagaraSystem>(OpenObj))
+			{
+				TSharedRef<FNiagaraSystemToolkit> NewNiagaraSystemToolkit(new FNiagaraSystemToolkit());
+				NewNiagaraSystemToolkit->InitializeWithSystem(Mode, OpenArgs.ToolkitHost, *System);
+			}
+		}
+		return EAssetCommandResult::Handled;
+	}
+
+	return EAssetCommandResult::Unhandled;
+}
+
+#undef LOCTEXT_NAMESPACE
diff --git a/Engine/Plugins/FX/Niagara/Source/NiagaraEditor/Private/AssetDefinitions/AssetDefinition_NiagaraSystem.h b/Engine/Plugins/FX/Niagara/Source/NiagaraEditor/Private/AssetDefinitions/AssetDefinition_NiagaraSystem.h
new file mode 100644
--- /dev/null
+++ b/Engine/Plugins/FX/Niagara/Source/NiagaraEditor/Private/AssetDefinitions/AssetDefinition_NiagaraSystem.h
@@ -0,0 +1,22 @@
+// Copyright Epic Games, Inc. All Rights Reserved.
+
+#pragma once
+
+#include "CoreMinimal.h"
+#include "AssetDefinitionDefault.h"
+
+#include "AssetDefinition_NiagaraSystem.generated.h"
+
+UCLASS()
+class UAssetDefinition_NiagaraSystem : public UAssetDefinitionDefault
+{
+	GENERATED_BODY()
+
+public:
+	// UAssetDefinition Implementation
+	virtual FText GetAssetDisplayName() const override { return NSLOCTEXT("AssetDefinition", "NiagaraSystem", "Niagara System"); }
+	virtual FLinearColor GetAssetColor() const override;
+	virtual TSoftClassPtr<UObject> GetAssetClass() const override;
+	virtual TConstArrayView<FAssetCategoryPath> GetAssetCategories() const override;
+	virtual EAssetCommandResult OpenAssets(const FAssetOpenArgs& OpenArgs) const override;
+};
diff --git a/Engine/Plugins/FX/Niagara/Source/NiagaraEditor/Private/AssetDefinitions/AssetDefinition_NiagaraValidationRuleSet.cpp b/Engine/Plugins/FX/Niagara/Source/NiagaraEditor/Private/AssetDefinitions/AssetDefinition_NiagaraValidationRuleSet.cpp
new file mode 100644
--- /dev/null
+++ b/Engine/Plugins/FX/Niagara/Source/NiagaraEditor/Private/AssetDefinitions/AssetDefinition_NiagaraValidationRuleSet.cpp
@@ -0,0 +1,26 @@
+// Copyright Epic Games, Inc. All Rights Reserved.
+
+#include "AssetDefinitions/AssetDefinition_NiagaraValidationRuleSet.h"
+
+#include "NiagaraEditorStyle.h"
+#include "NiagaraValidationRuleSet.h"
+
+#define LOCTEXT_NAMESPACE "UAssetDefinition_NiagaraValidationRuleSet"
+
+FLinearColor UAssetDefinition_NiagaraValidationRuleSet::GetAssetColor() const
+{
+	return FNiagaraEditorStyle::Get().GetColor("NiagaraEditor.AssetColors.ValidationRuleSet");
+}
+
+TSoftClassPtr<> UAssetDefinition_NiagaraValidationRuleSet::GetAssetClass() const
+{
+	return UNiagaraValidationRuleSet::StaticClass();
+}
+
+TConstArrayView<FAssetCategoryPath> UAssetDefinition_NiagaraValidationRuleSet::GetAssetCategories() const
+{
+	static FAssetCategoryPath AssetPaths[] = { FAssetCategoryPath(LOCTEXT("NiagaraAssetsCategory", "FX"), LOCTEXT("NiagaraValidationRuleSet_SubCategory", "Advanced")) };
+	return AssetPaths;
+}
+
+#undef LOCTEXT_NAMESPACE
diff --git a/Engine/Plugins/FX/Niagara/Source/NiagaraEditor/Private/AssetDefinitions/AssetDefinition_NiagaraValidationRuleSet.h b/Engine/Plugins/FX/Niagara/Source/NiagaraEditor/Private/AssetDefinitions/AssetDefinition_NiagaraValidationRuleSet.h
new file mode 100644
--- /dev/null
+++ b/Engine/Plugins/FX/Niagara/Source/NiagaraEditor/Private/AssetDefinitions/AssetDefinition_NiagaraValidationRuleSet.h
@@ -0,0 +1,21 @@
+// Copyright Epic Games, Inc. All Rights Reserved.
+
+#pragma once
+
+#include "CoreMinimal.h"
+#include "AssetDefinitionDefault.h"
+
+#include "AssetDefinition_NiagaraValidationRuleSet.generated.h"
+
+UCLASS()
+class UAssetDefinition_NiagaraValidationRuleSet : public UAssetDefinitionDefault
+{
+	GENERATED_BODY()
+
+public:
+	// UAssetDefinition Implementation
+	virtual FText GetAssetDisplayName() const override { return NSLOCTEXT("AssetDefinition", "NiagaraValidationRuleSet", "Niagara Validation Rule Set"); }
+	virtual FLinearColor GetAssetColor() const override;
+	virtual TSoftClassPtr<UObject> GetAssetClass() const override;
+	virtual TConstArrayView<FAssetCategoryPath> GetAssetCategories() const override;
+};
diff --git a/Engine/Plugins/FX/Niagara/Source/NiagaraEditor/Private/AssetTypeActions/AssetTypeActions_NiagaraEffectType.cpp b/Engine/Plugins/FX/Niagara/Source/NiagaraEditor/Private/AssetTypeActions/AssetTypeActions_NiagaraEffectType.cpp
deleted file mode 100644
diff --git a/Engine/Plugins/FX/Niagara/Source/NiagaraEditor/Private/AssetTypeActions/AssetTypeActions_NiagaraEffectType.h b/Engine/Plugins/FX/Niagara/Source/NiagaraEditor/Private/AssetTypeActions/AssetTypeActions_NiagaraEffectType.h
deleted file mode 100644
diff --git a/Engine/Plugins/FX/Niagara/Source/NiagaraEditor/Private/AssetTypeActions/AssetTypeActions_NiagaraEmitter.cpp b/Engine/Plugins/FX/Niagara/Source/NiagaraEditor/Private/AssetTypeActions/AssetTypeActions_NiagaraEmitter.cpp
deleted file mode 100644
diff --git a/Engine/Plugins/FX/Niagara/Source/NiagaraEditor/Private/AssetTypeActions/AssetTypeActions_NiagaraEmitter.h b/Engine/Plugins/FX/Niagara/Source/NiagaraEditor/Private/AssetTypeActions/AssetTypeActions_NiagaraEmitter.h
deleted file mode 100644
diff --git a/Engine/Plugins/FX/Niagara/Source/NiagaraEditor/Private/AssetTypeActions/AssetTypeActions_NiagaraParameterDefinitions.cpp b/Engine/Plugins/FX/Niagara/Source/NiagaraEditor/Private/AssetTypeActions/AssetTypeActions_NiagaraParameterDefinitions.cpp
deleted file mode 100644
diff --git a/Engine/Plugins/FX/Niagara/Source/NiagaraEditor/Private/AssetTypeActions/AssetTypeActions_NiagaraParameterDefinitions.h b/Engine/Plugins/FX/Niagara/Source/NiagaraEditor/Private/AssetTypeActions/AssetTypeActions_NiagaraParameterDefinitions.h
deleted file mode 100644
diff --git a/Engine/Plugins/FX/Niagara/Source/NiagaraEditor/Private/AssetTypeActions/AssetTypeActions_NiagaraSimCache.cpp b/Engine/Plugins/FX/Niagara/Source/NiagaraEditor/Private/AssetTypeActions/AssetTypeActions_NiagaraSimCache.cpp
deleted file mode 100644
diff --git a/Engine/Plugins/FX/Niagara/Source/NiagaraEditor/Private/AssetTypeActions/AssetTypeActions_NiagaraSimCache.h b/Engine/Plugins/FX/Niagara/Source/NiagaraEditor/Private/AssetTypeActions/AssetTypeActions_NiagaraSimCache.h
deleted file mode 100644
diff --git a/Engine/Plugins/FX/Niagara/Source/NiagaraEditor/Private/AssetTypeActions/AssetTypeActions_NiagaraSystem.cpp b/Engine/Plugins/FX/Niagara/Source/NiagaraEditor/Private/AssetTypeActions/AssetTypeActions_NiagaraSystem.cpp
deleted file mode 100644
diff --git a/Engine/Plugins/FX/Niagara/Source/NiagaraEditor/Private/AssetTypeActions/AssetTypeActions_NiagaraSystem.h b/Engine/Plugins/FX/Niagara/Source/NiagaraEditor/Private/AssetTypeActions/AssetTypeActions_NiagaraSystem.h
deleted file mode 100644
diff --git a/Engine/Plugins/FX/Niagara/Source/NiagaraEditor/Private/AssetTypeActions/AssetTypeActions_NiagaraValidationRuleSet.cpp b/Engine/Plugins/FX/Niagara/Source/NiagaraEditor/Private/AssetTypeActions/AssetTypeActions_NiagaraValidationRuleSet.cpp
deleted file mode 100644
diff --git a/Engine/Plugins/FX/Niagara/Source/NiagaraEditor/Private/AssetTypeActions/AssetTypeActions_NiagaraValidationRuleSet.h b/Engine/Plugins/FX/Niagara/Source/NiagaraEditor/Private/AssetTypeActions/AssetTypeActions_NiagaraValidationRuleSet.h
deleted file mode 100644
diff --git a/Engine/Plugins/FX/Niagara/Source/NiagaraEditor/Private/NiagaraEditorModule.cpp b/Engine/Plugins/FX/Niagara/Source/NiagaraEditor/Private/NiagaraEditorModule.cpp
--- a/Engine/Plugins/FX/Niagara/Source/NiagaraEditor/Private/NiagaraEditorModule.cpp
+++ b/Engine/Plugins/FX/Niagara/Source/NiagaraEditor/Private/NiagaraEditorModule.cpp
@@ -15,14 +15,8 @@
 #include "Subsystems/ImportSubsystem.h"
 #include "UObject/UObjectThreadContext.h"
 
-#include "AssetTypeActions/AssetTypeActions_NiagaraSystem.h"
-#include "AssetTypeActions/AssetTypeActions_NiagaraEmitter.h"
 #include "AssetTypeActions/AssetTypeActions_NiagaraScript.h"
 #include "AssetTypeActions/AssetTypeActions_NiagaraParameterCollection.h"
-#include "AssetTypeActions/AssetTypeActions_NiagaraEffectType.h"
-#include "AssetTypeActions/AssetTypeActions_NiagaraParameterDefinitions.h"
-#include "AssetTypeActions/AssetTypeActions_NiagaraSimCache.h"
-#include "AssetTypeActions/AssetTypeActions_NiagaraValidationRuleSet.h"
 
 #include "EdGraphSchema_Niagara.h"
 #include "EdGraphUtilities.h"
@@ -857,17 +851,11 @@ void FNiagaraEditorModule::StartupModule()
 
 	IAssetTools& AssetTools = FModuleManager::LoadModuleChecked<FAssetToolsModule>("AssetTools").Get();
 	NiagaraAssetCategory = AssetTools.RegisterAdvancedAssetCategory(FName(TEXT("FX")), LOCTEXT("NiagaraAssetsCategory", "FX"));
-	RegisterAssetTypeAction(AssetTools, MakeShareable(new FAssetTypeActions_NiagaraSystem()));
-	RegisterAssetTypeAction(AssetTools, MakeShareable(new FAssetTypeActions_NiagaraEmitter()));
 	RegisterAssetTypeAction(AssetTools, MakeShareable(new FAssetTypeActions_NiagaraScriptFunctions()));
 	RegisterAssetTypeAction(AssetTools, MakeShareable(new FAssetTypeActions_NiagaraScriptModules()));
 	RegisterAssetTypeAction(AssetTools, MakeShareable(new FAssetTypeActions_NiagaraScriptDynamicInputs()));
 	RegisterAssetTypeAction(AssetTools, MakeShareable(new FAssetTypeActions_NiagaraParameterCollection())); 
 	RegisterAssetTypeAction(AssetTools, MakeShareable(new FAssetTypeActions_NiagaraParameterCollectionInstance()));
-	RegisterAssetTypeAction(AssetTools, MakeShareable(new FAssetTypeActions_NiagaraParameterDefinitions()));
-	RegisterAssetTypeAction(AssetTools, MakeShareable(new FAssetTypeActions_NiagaraEffectType()));
-	RegisterAssetTypeAction(AssetTools, MakeShareable(new FAssetTypeActions_NiagaraSimCache()));
-	RegisterAssetTypeAction(AssetTools, MakeShareable(new FAssetTypeActions_NiagaraValidationRuleSet()));
 
 	// Preload all parameter definition & collection assets so that they will be postloaded before postload calls to scripts/emitters/systems that rely on them.
 	{
diff --git a/Engine/Plugins/FX/Niagara/Source/NiagaraEditor/Private/NiagaraEditorUtilities.cpp b/Engine/Plugins/FX/Niagara/Source/NiagaraEditor/Private/NiagaraEditorUtilities.cpp
--- a/Engine/Plugins/FX/Niagara/Source/NiagaraEditor/Private/NiagaraEditorUtilities.cpp
+++ b/Engine/Plugins/FX/Niagara/Source/NiagaraEditor/Private/NiagaraEditorUtilities.cpp
@@ -1013,6 +1013,16 @@ bool FNiagaraEditorUtilities::AreTypesAssignable(const FNiagaraTypeDefinition& T
 }
 
 void FNiagaraEditorUtilities::MarkDependentCompilableAssetsDirty(TArray<UObject*> InObjects)
+{
+	TArray<FAssetData> AssetsToCheck;
+	for (UObject* InObject : InObjects)
+	{	
+		AssetsToCheck.Add(FAssetData(InObject));
+	}
+	MarkDependentCompilableAssetsDirty(AssetsToCheck);
+}
+
+void FNiagaraEditorUtilities::MarkDependentCompilableAssetsDirty(const TArray<FAssetData>& InAssetsToCheck)
 {
 	const FText LoadAndMarkDirtyDisplayName = NSLOCTEXT("NiagaraEditor", "MarkDependentAssetsDirtySlowTask", "Loading and marking dependent assets dirty.");
 	GWarn->BeginSlowTask(LoadAndMarkDirtyDisplayName, true, true);
@@ -1021,12 +1031,7 @@ void FNiagaraEditorUtilities::MarkDependentCompilableAssetsDirty(TArray<UObject*
 	TArray<FAssetIdentifier> ReferenceNames;
 
 	TArray<FAssetData> AssetsToLoadAndMarkDirty;
-	TArray<FAssetData> AssetsToCheck;
-
-	for (UObject* InObject : InObjects)
-	{	
-		AssetsToCheck.Add(FAssetData(InObject));
-	}
+	TArray<FAssetData> AssetsToCheck = InAssetsToCheck;
 
 	while (AssetsToCheck.Num() > 0)
 	{
diff --git a/Engine/Plugins/FX/Niagara/Source/NiagaraEditor/Public/NiagaraEditorUtilities.h b/Engine/Plugins/FX/Niagara/Source/NiagaraEditor/Public/NiagaraEditorUtilities.h
--- a/Engine/Plugins/FX/Niagara/Source/NiagaraEditor/Public/NiagaraEditorUtilities.h
+++ b/Engine/Plugins/FX/Niagara/Source/NiagaraEditor/Public/NiagaraEditorUtilities.h
@@ -168,6 +168,7 @@ namespace FNiagaraEditorUtilities
 	NIAGARAEDITOR_API bool AreTypesAssignable(const FNiagaraTypeDefinition& TypeA, const FNiagaraTypeDefinition& TypeB);
 
 	void MarkDependentCompilableAssetsDirty(TArray<UObject*> InObjects);
+	void MarkDependentCompilableAssetsDirty(const TArray<FAssetData>& InAssets);
 
 	void ResolveNumerics(UNiagaraGraph* SourceGraph, bool bForceParametersToResolveNumerics, TArray<FNiagaraVariable>& ChangedNumericParams);
 
diff --git a/Engine/Source/Editor/AssetDefinition/Private/AssetDefinition.cpp b/Engine/Source/Editor/AssetDefinition/Private/AssetDefinition.cpp
--- a/Engine/Source/Editor/AssetDefinition/Private/AssetDefinition.cpp
+++ b/Engine/Source/Editor/AssetDefinition/Private/AssetDefinition.cpp
@@ -22,6 +22,11 @@ FAssetCategoryPath::FAssetCategoryPath(const FText InCategory)
 	CategoryPath = { TPair<FName, FText>(FName(*FTextInspector::GetSourceString(InCategory)), InCategory) };
 }
 
+FAssetCategoryPath::FAssetCategoryPath(FText InCategory, FText InSubCategory)
+{
+	CategoryPath = { TPair<FName, FText>(FName(*FTextInspector::GetSourceString(InCategory)), InCategory), TPair<FName, FText>(FName(*FTextInspector::GetSourceString(InSubCategory)), InSubCategory) };
+}
+
 FAssetCategoryPath::FAssetCategoryPath(TConstArrayView<FText> InCategoryPath)
 {
 	check(InCategoryPath.Num() > 0);
diff --git a/Engine/Source/Editor/AssetDefinition/Public/AssetDefinition.h b/Engine/Source/Editor/AssetDefinition/Public/AssetDefinition.h
--- a/Engine/Source/Editor/AssetDefinition/Public/AssetDefinition.h
+++ b/Engine/Source/Editor/AssetDefinition/Public/AssetDefinition.h
@@ -195,6 +195,7 @@ struct FAssetOpenSupportArgs
 struct ASSETDEFINITION_API FAssetCategoryPath
 {
 	FAssetCategoryPath(FText InCategory);
+	FAssetCategoryPath(FText InCategory, FText InSubCategory);
 	FAssetCategoryPath(TConstArrayView<FText> InCategoryPath);
 
 	FName GetCategory() const { return CategoryPath[0].Key; }
