- Tooltips for Niagara Renderer bindings now include bindings name & description if available

diff --git a/Engine/Plugins/FX/Niagara/Source/NiagaraEditor/Private/Customizations/NiagaraTypeCustomizations.cpp b/Engine/Plugins/FX/Niagara/Source/NiagaraEditor/Private/Customizations/NiagaraTypeCustomizations.cpp
--- a/Engine/Plugins/FX/Niagara/Source/NiagaraEditor/Private/Customizations/NiagaraTypeCustomizations.cpp
+++ b/Engine/Plugins/FX/Niagara/Source/NiagaraEditor/Private/Customizations/NiagaraTypeCustomizations.cpp
@@ -5,13 +5,30 @@
 #include "CoreMinimal.h"
 #include "DetailLayoutBuilder.h"
 #include "DetailWidgetRow.h"
+#include "MaterialTypes.h"
 #include "IDetailChildrenBuilder.h"
+#include "IDetailGroup.h"
 #include "IPropertyUtilities.h"
+#include "PropertyHandle.h"
+#include "ScopedTransaction.h"
+#include "SGraphActionMenu.h"
+#include "DeviceProfiles/DeviceProfileManager.h"
+#include "Framework/Application/SlateApplication.h"
+#include "Framework/MultiBox/MultiBoxBuilder.h"
 #include "Materials/MaterialInterface.h"
+#include "Modules/ModuleManager.h"
+#include "ViewModels/NiagaraSystemViewModel.h"
+#include "ViewModels/TNiagaraViewModelManager.h"
+#include "Widgets/SNiagaraParameterName.h"
+#include "Widgets/Input/SButton.h"
+#include "Widgets/Input/SCheckBox.h"
+#include "Widgets/Input/SComboButton.h"
+#include "Widgets/Input/SEditableTextBox.h"
+#include "Widgets/Text/STextBlock.h"
+
 #include "NiagaraConstants.h"
 #include "NiagaraDataInterfaceRW.h"
 #include "NiagaraEditorModule.h"
-#include "Modules/ModuleManager.h"
 #include "NiagaraEditorUtilities.h"
 #include "NiagaraEmitter.h"
 #include "NiagaraNodeOutput.h"
@@ -24,21 +41,6 @@
 #include "NiagaraSimulationStageBase.h"
 #include "NiagaraSystem.h"
 #include "NiagaraTypes.h"
-#include "PropertyHandle.h"
-#include "ScopedTransaction.h"
-#include "SGraphActionMenu.h"
-#include "DeviceProfiles/DeviceProfileManager.h"
-#include "Framework/Application/SlateApplication.h"
-#include "Framework/MultiBox/MultiBoxBuilder.h"
-#include "ViewModels/NiagaraSystemViewModel.h"
-#include "ViewModels/TNiagaraViewModelManager.h"
-#include "Widgets/SNiagaraParameterName.h"
-#include "Widgets/Input/SButton.h"
-#include "Widgets/Input/SCheckBox.h"
-#include "Widgets/Input/SComboButton.h"
-#include "Widgets/Input/SEditableTextBox.h"
-#include "Widgets/Text/STextBlock.h"
-#include "IDetailGroup.h"
 
 #include UE_INLINE_GENERATED_CPP_BY_NAME(NiagaraTypeCustomizations)
 
@@ -1061,9 +1063,19 @@ TSharedRef<SWidget> FNiagaraMaterialAttributeBindingCustomization::OnGetMaterial
 	return MenuContent;
 }
 
-TArray<FName> FNiagaraMaterialAttributeBindingCustomization::GetMaterialNames() const
+void FNiagaraMaterialAttributeBindingCustomization::GetMaterialParameters(TArray<TPair<FName, FString>>& OutBindingNameAndDescription) const
 {
-	TArray<FName> Names;
+	auto TransformMaterialParameterInfo =
+		[&OutBindingNameAndDescription](UMaterialInterface* Material, TArray<FMaterialParameterInfo>& ParameterInfos)
+		{
+			for (const FMaterialParameterInfo& ParameterInfo : ParameterInfos)
+			{
+				FString ParameterDesc;
+				Material->GetParameterDesc(ParameterInfo, ParameterDesc);
+				OutBindingNameAndDescription.Emplace(ParameterInfo.Name, ParameterDesc);
+			}
+			ParameterInfos.Reset();
+		};
 
 	if (BaseSystem && TargetParameterBinding && PropertyHandle)
 	{
@@ -1079,7 +1091,8 @@ TArray<FName> FNiagaraMaterialAttributeBindingCustomization::GetMaterialNames()
 				RendererProperties->GetUsedMaterials(nullptr, Materials);
 			}
 
-			TArray<FMaterialParameterInfo> ParameterInfo;
+			TArray<FMaterialParameterInfo> TempParameterInfo;
+			TArray<FGuid> ParameterIds;
 			for (UMaterialInterface* Material : Materials)
 			{
 				if (!Material)
@@ -1087,47 +1100,38 @@ TArray<FName> FNiagaraMaterialAttributeBindingCustomization::GetMaterialNames()
 					continue;
 				}
 
-				{
-					TArray<FMaterialParameterInfo> LocalParameterInfo;
-					TArray<FGuid> ParameterIds;
-					Material->GetAllTextureParameterInfo(LocalParameterInfo, ParameterIds);
-					ParameterInfo.Append(LocalParameterInfo);
-				}
-				{
-					TArray<FMaterialParameterInfo> LocalParameterInfo;
-					TArray<FGuid> ParameterIds;
-					Material->GetAllScalarParameterInfo(LocalParameterInfo, ParameterIds);
-					ParameterInfo.Append(LocalParameterInfo);
-				}
-				{
-					TArray<FMaterialParameterInfo> LocalParameterInfo;
-					TArray<FGuid> ParameterIds;
-					Material->GetAllVectorParameterInfo(LocalParameterInfo, ParameterIds);
-					ParameterInfo.Append(LocalParameterInfo);
-				}
-			}
+				Material->GetAllTextureParameterInfo(TempParameterInfo, ParameterIds);
+				TransformMaterialParameterInfo(Material, TempParameterInfo);
 
-			for (const FMaterialParameterInfo& Var : ParameterInfo)
-			{
-				Names.AddUnique(Var.Name);
+				Material->GetAllScalarParameterInfo(TempParameterInfo, ParameterIds);
+				TransformMaterialParameterInfo(Material, TempParameterInfo);
+
+				Material->GetAllVectorParameterInfo(TempParameterInfo, ParameterIds);
+				TransformMaterialParameterInfo(Material, TempParameterInfo);
+
+				ParameterIds.Reset();
 			}
 		}
 	}
-
-	return Names;
 }
 
 void FNiagaraMaterialAttributeBindingCustomization::CollectAllMaterialActions(FGraphActionListBuilderBase& OutAllActions)
 {
-	TArray<FName> ParamNames = GetMaterialNames();
-	for (FName ParamName : ParamNames)
+	TArray<TPair<FName, FString>> MaterialParameters;
+	GetMaterialParameters(MaterialParameters);
+
+	for (const TPair<FName, FString>& Parameter : MaterialParameters)
 	{
-		FText CategoryName = FText();
-		FString DisplayNameString = FName::NameToDisplayString(ParamName.ToString(), false);
-		const FText NameText = FText::FromString(DisplayNameString);
-		const FText TooltipDesc = FText::Format(LOCTEXT("BindToMaterialParameter", "Bind to the Material Variable \"{0}\" "), FText::FromString(DisplayNameString));
-		TSharedPtr<FNiagaraStackAssetAction_VarBind> NewNodeAction(new FNiagaraStackAssetAction_VarBind(ParamName, CategoryName, NameText,
-			TooltipDesc, 0, FText()));
+		TSharedPtr<FNiagaraStackAssetAction_VarBind> NewNodeAction(
+			new FNiagaraStackAssetAction_VarBind(
+				Parameter.Key,
+				FText(),
+				FText::FromName(Parameter.Key),
+				FNiagaraRendererMaterialParameterCustomization::GetMaterialBindingTooltip(Parameter.Key, Parameter.Value),
+				0,
+				FText()
+			)
+		);
 		OutAllActions.AddAction(NewNodeAction);
 	}
 }
@@ -2098,6 +2102,18 @@ FText FNiagaraRendererMaterialParameterCustomization::GetBindingNameText(TShared
 	return FText::FromName(BindingName);
 }
 
+FText FNiagaraRendererMaterialParameterCustomization::GetMaterialBindingTooltip(FName ParameterName, const FString& ParameterDesc)
+{
+	if (ParameterDesc.Len() > 0)
+	{
+		return FText::Format(LOCTEXT("RendererMaterialParameterTooltip", "Set Material Parameter: {0}\nDescription: {1}"), FText::FromName(ParameterName), FText::FromString(ParameterDesc));
+	}
+	else
+	{
+		return FText::Format(LOCTEXT("BasicRendererMaterialParameterTooltip", "Set Material Parameter: {0}"), FText::FromName(ParameterName));
+	}
+}
+
 TSharedRef<SWidget> FNiagaraRendererMaterialParameterCustomization::OnGetMaterialBindingNameMenuContent(TSharedPtr<IPropertyHandle> PropertyHandle) const
 {
 	FMenuBuilder MenuBuilder(true, NULL);
@@ -2108,7 +2124,6 @@ TSharedRef<SWidget> FNiagaraRendererMaterialParameterCustomization::OnGetMateria
 		TArray<UMaterialInterface*> UsedMaterials;
 		RenderProperties->GetUsedMaterials(nullptr, UsedMaterials);
 
-		TArray<FName> ValidBindings;
 		for (UMaterialInterface* Material : UsedMaterials)
 		{
 			if (Material == nullptr)
@@ -2116,65 +2131,50 @@ TSharedRef<SWidget> FNiagaraRendererMaterialParameterCustomization::OnGetMateria
 				continue;
 			}
 
-			GetMaterialBindingNames(Material, ValidBindings);
-		}
+			TArray<FMaterialParameterInfo> MaterialParameterInfos;
+			GetMaterialParameterInfos(Material, MaterialParameterInfos);
 
-		for (FName ValidBinding : ValidBindings)
-		{
-			MenuBuilder.AddMenuEntry(
-				FText::FromName(ValidBinding),
-				TAttribute<FText>(),
-				FSlateIcon(),
-				FUIAction(
-					FExecuteAction::CreateLambda(
-						[ValidBinding, PropertyHandle]()
-						{
-				PropertyHandle->SetValue(ValidBinding);
-						}
+			for (const FMaterialParameterInfo& MaterialParameterInfo : MaterialParameterInfos)
+			{
+				FString ParameterDesc;
+				Material->GetParameterDesc(MaterialParameterInfo, ParameterDesc);
+
+				MenuBuilder.AddMenuEntry(
+					FText::FromName(MaterialParameterInfo.Name),
+					GetMaterialBindingTooltip(MaterialParameterInfo.Name, ParameterDesc),
+					FSlateIcon(),
+					FUIAction(
+						FExecuteAction::CreateLambda(
+							[BindingName=MaterialParameterInfo.Name, PropertyHandle]()
+							{
+								PropertyHandle->SetValue(BindingName);
+							}
+						)
 					)
-				)
-			);
+				);
+			}
 		}
 	}
 
 	return MenuBuilder.MakeWidget();
 }
 
-void FNiagaraRendererMaterialScalarParameterCustomization::GetMaterialBindingNames(UMaterialInterface* Material, TArray<FName>& OutBindings) const
+void FNiagaraRendererMaterialScalarParameterCustomization::GetMaterialParameterInfos(UMaterialInterface* Material, TArray<FMaterialParameterInfo>& OutMaterialParameterInfos) const
 {
-	TArray<FName> ValidBindings;
-	TArray<FMaterialParameterInfo> MaterialParameterInfos;
 	TArray<FGuid> Guids;
-	Material->GetAllScalarParameterInfo(MaterialParameterInfos, Guids);
-
-	for (const FMaterialParameterInfo& MaterialParameterInfo : MaterialParameterInfos)
-	{
-		OutBindings.AddUnique(MaterialParameterInfo.Name);
-	}
+	Material->GetAllScalarParameterInfo(OutMaterialParameterInfos, Guids);
 }
 
-void FNiagaraRendererMaterialVectorParameterCustomization::GetMaterialBindingNames(UMaterialInterface* Material, TArray<FName>& OutBindings) const
+void FNiagaraRendererMaterialVectorParameterCustomization::GetMaterialParameterInfos(UMaterialInterface* Material, TArray<FMaterialParameterInfo>& OutMaterialParameterInfos) const
 {
-	TArray<FMaterialParameterInfo> MaterialParameterInfos;
 	TArray<FGuid> Guids;
-	Material->GetAllVectorParameterInfo(MaterialParameterInfos, Guids);
-
-	for (const FMaterialParameterInfo& MaterialParameterInfo : MaterialParameterInfos)
-	{
-		OutBindings.AddUnique(MaterialParameterInfo.Name);
-	}
+	Material->GetAllVectorParameterInfo(OutMaterialParameterInfos, Guids);
 }
 
-void FNiagaraRendererMaterialTextureParameterCustomization::GetMaterialBindingNames(UMaterialInterface* Material, TArray<FName>& OutBindings) const
+void FNiagaraRendererMaterialTextureParameterCustomization::GetMaterialParameterInfos(UMaterialInterface* Material, TArray<FMaterialParameterInfo>& OutMaterialParameterInfos) const
 {
-	TArray<FMaterialParameterInfo> MaterialParameterInfos;
 	TArray<FGuid> Guids;
-	Material->GetAllTextureParameterInfo(MaterialParameterInfos, Guids);
-
-	for (const FMaterialParameterInfo& MaterialParameterInfo : MaterialParameterInfos)
-	{
-		OutBindings.AddUnique(MaterialParameterInfo.Name);
-	}
+	Material->GetAllTextureParameterInfo(OutMaterialParameterInfos, Guids);
 }
 
 bool FNiagaraRendererMaterialStaticBoolParameterCustomization::CustomizeChildProperty(class IDetailChildrenBuilder& ChildBuilder, TSharedPtr<IPropertyHandle> PropertyHandle)
@@ -2206,16 +2206,10 @@ bool FNiagaraRendererMaterialStaticBoolParameterCustomization::CustomizeChildPro
 	return false;
 }
 
-void FNiagaraRendererMaterialStaticBoolParameterCustomization::GetMaterialBindingNames(UMaterialInterface* Material, TArray<FName>& OutBindings) const
+void FNiagaraRendererMaterialStaticBoolParameterCustomization::GetMaterialParameterInfos(UMaterialInterface* Material, TArray<FMaterialParameterInfo>& OutMaterialParameterInfos) const
 {
-	TArray<FMaterialParameterInfo> MaterialParameterInfos;
 	TArray<FGuid> Guids;
-	Material->GetAllStaticSwitchParameterInfo(MaterialParameterInfos, Guids);
-
-	for (const FMaterialParameterInfo& MaterialParameterInfo : MaterialParameterInfos)
-	{
-		OutBindings.AddUnique(MaterialParameterInfo.Name);
-	}
+	Material->GetAllStaticSwitchParameterInfo(OutMaterialParameterInfos, Guids);
 }
 
 TSharedRef<SWidget> FNiagaraRendererMaterialStaticBoolParameterCustomization::OnGetStaticVariablelBindingNameMenuContent(TSharedPtr<IPropertyHandle> PropertyHandle) const
diff --git a/Engine/Plugins/FX/Niagara/Source/NiagaraEditor/Private/Customizations/NiagaraTypeCustomizations.h b/Engine/Plugins/FX/Niagara/Source/NiagaraEditor/Private/Customizations/NiagaraTypeCustomizations.h
--- a/Engine/Plugins/FX/Niagara/Source/NiagaraEditor/Private/Customizations/NiagaraTypeCustomizations.h
+++ b/Engine/Plugins/FX/Niagara/Source/NiagaraEditor/Private/Customizations/NiagaraTypeCustomizations.h
@@ -15,6 +15,7 @@ class FNiagaraScriptViewModel;
 class IPropertyHandle;
 class IPropertyHandleArray;
 class UMaterialInterface;
+struct FMaterialParameterInfo;
 class UNiagaraGraph;
 class UNiagaraParameterDefinitions;
 class UNiagaraRendererProperties;
@@ -255,7 +256,7 @@ private:
 	FText GetMaterialCurrentText() const;
 	FText GetMaterialTooltipText() const;
 	TSharedRef<SWidget> OnGetMaterialMenuContent() const;
-	TArray<FName> GetMaterialNames() const;
+	void GetMaterialParameters(TArray<TPair<FName, FString>>& OutBindingNameAndDescription) const;
 	void ChangeMaterialSource(FName InVarName);
 	void CollectAllMaterialActions(FGraphActionListBuilderBase& OutAllActions);
 	TSharedRef<SWidget> OnCreateWidgetForMaterialAction(struct FCreateWidgetForActionData* const InCreateData);
@@ -382,10 +383,11 @@ public:
 	// IPropertyTypeCustomization interface end
 
 	FText GetBindingNameText(TSharedPtr<IPropertyHandle> PropertyHandle) const;
+	static FText GetMaterialBindingTooltip(FName ParameterName, const FString& ParameterDesc);
 	TSharedRef<SWidget> OnGetMaterialBindingNameMenuContent(TSharedPtr<IPropertyHandle> PropertyHandle) const;
 
 	virtual bool CustomizeChildProperty(class IDetailChildrenBuilder& ChildBuilder, TSharedPtr<IPropertyHandle> PropertyHandle) { return false; }
-	virtual void GetMaterialBindingNames(UMaterialInterface* Material, TArray<FName>& OutBindings) const = 0;
+	virtual void GetMaterialParameterInfos(UMaterialInterface* Material, TArray<FMaterialParameterInfo>& OutMaterialParameterInfos) const = 0;
 
 protected:
 	TWeakObjectPtr<UNiagaraRendererProperties>	WeakRenderProperties;
@@ -399,7 +401,7 @@ public:
 		return MakeShared<FNiagaraRendererMaterialScalarParameterCustomization>();
 	}
 
-	virtual void GetMaterialBindingNames(UMaterialInterface* Material, TArray<FName>& OutBindings) const override;
+	virtual void GetMaterialParameterInfos(UMaterialInterface* Material, TArray<FMaterialParameterInfo>& OutMaterialParameterInfos) const override;
 };
 
 class FNiagaraRendererMaterialVectorParameterCustomization : public FNiagaraRendererMaterialParameterCustomization
@@ -410,7 +412,7 @@ public:
 		return MakeShared<FNiagaraRendererMaterialVectorParameterCustomization>();
 	}
 
-	virtual void GetMaterialBindingNames(UMaterialInterface* Material, TArray<FName>& OutBindings) const override;
+	virtual void GetMaterialParameterInfos(UMaterialInterface* Material, TArray<FMaterialParameterInfo>& OutMaterialParameterInfos) const override;
 };
 
 class FNiagaraRendererMaterialTextureParameterCustomization : public FNiagaraRendererMaterialParameterCustomization
@@ -421,7 +423,7 @@ public:
 		return MakeShared<FNiagaraRendererMaterialTextureParameterCustomization>();
 	}
 
-	virtual void GetMaterialBindingNames(UMaterialInterface* Material, TArray<FName>& OutBindings) const override;
+	virtual void GetMaterialParameterInfos(UMaterialInterface* Material, TArray<FMaterialParameterInfo>& OutMaterialParameterInfos) const override;
 };
 
 class FNiagaraRendererMaterialStaticBoolParameterCustomization : public FNiagaraRendererMaterialParameterCustomization
@@ -433,6 +435,6 @@ public:
 	}
 
 	virtual bool CustomizeChildProperty(class IDetailChildrenBuilder& ChildBuilder, TSharedPtr<IPropertyHandle> PropertyHandle) override;
-	virtual void GetMaterialBindingNames(UMaterialInterface* Material, TArray<FName>& OutBindings) const override;
+	virtual void GetMaterialParameterInfos(UMaterialInterface* Material, TArray<FMaterialParameterInfo>& OutMaterialParameterInfos) const override;
 	TSharedRef<SWidget> OnGetStaticVariablelBindingNameMenuContent(TSharedPtr<IPropertyHandle> PropertyHandle) const;
 };
