Handle renaming cached default data interfaces when we rename an emitter.  Avoids having to recompile a shader to get all the data up to date.

diff --git a/Engine/Plugins/FX/Niagara/Source/Niagara/Private/NiagaraScript.cpp b/Engine/Plugins/FX/Niagara/Source/Niagara/Private/NiagaraScript.cpp
--- a/Engine/Plugins/FX/Niagara/Source/Niagara/Private/NiagaraScript.cpp
+++ b/Engine/Plugins/FX/Niagara/Source/Niagara/Private/NiagaraScript.cpp
@@ -3437,6 +3437,28 @@ void UNiagaraScript::SyncAliases(const FNiagaraAliasContext& ResolveAliasesConte
 		}
 		++Iterator;
 	}
+
+	// and cached default data interfaces
+	for (FNiagaraScriptDataInterfaceInfo& DataInterfaceInfo : CachedDefaultDataInterfaces)
+	{
+		const FNiagaraVariableBase TempNamedVariable(FNiagaraTypeDefinition::GetFloatDef(), DataInterfaceInfo.Name);
+		const FNiagaraVariableBase ResolvedVariable = FNiagaraUtilities::ResolveAliases(TempNamedVariable, ResolveAliasesContext);
+
+		// update the name
+		DataInterfaceInfo.Name = ResolvedVariable.GetName();
+
+		auto UpdateName = [&](FName& NameVariable)
+		{
+			if (NameVariable == TempNamedVariable.GetName())
+			{
+				NameVariable = ResolvedVariable.GetName();
+			}
+		};
+
+		// also update the MapRead/MapWrite member variables, they seem to typically match the Name parameter, but may be None
+		UpdateName(DataInterfaceInfo.RegisteredParameterMapRead);
+		UpdateName(DataInterfaceInfo.RegisteredParameterMapWrite);
+	}
 }
 
 bool UNiagaraScript::SynchronizeExecutablesWithCompilation(const UNiagaraScript* Script, const TMap<FString, FString>& RenameMap)
