[Niagara] Collapse option for CustomHLSL node (Contributed by [REDACTED_PERSON])

PR #9752

diff --git a/Engine/Plugins/FX/Niagara/Source/NiagaraEditor/Private/NiagaraEditorStyle.cpp b/Engine/Plugins/FX/Niagara/Source/NiagaraEditor/Private/NiagaraEditorStyle.cpp
--- a/Engine/Plugins/FX/Niagara/Source/NiagaraEditor/Private/NiagaraEditorStyle.cpp
+++ b/Engine/Plugins/FX/Niagara/Source/NiagaraEditor/Private/NiagaraEditorStyle.cpp
@@ -317,6 +317,7 @@ void FNiagaraEditorStyle::InitCodeView()
 	const FTextBlockStyle NormalText = FAppStyle::GetWidgetStyle<FTextBlockStyle>("NormalText");
 
 	// Code View
+	Set("NiagaraEditor.CodeView.ToggledHide", FLinearColor(FColor(24, 160, 251)));
 	Set("NiagaraEditor.CodeView.Checkbox.Text", FTextBlockStyle(NormalText)
 		.SetFont(DEFAULT_FONT("Bold", 12))
 		.SetColorAndOpacity(FLinearColor(1.0f, 1.0f, 1.0f, 0.9f))
diff --git a/Engine/Plugins/FX/Niagara/Source/NiagaraEditor/Private/NiagaraNodeCustomHlsl.cpp b/Engine/Plugins/FX/Niagara/Source/NiagaraEditor/Private/NiagaraNodeCustomHlsl.cpp
--- a/Engine/Plugins/FX/Niagara/Source/NiagaraEditor/Private/NiagaraNodeCustomHlsl.cpp
+++ b/Engine/Plugins/FX/Niagara/Source/NiagaraEditor/Private/NiagaraNodeCustomHlsl.cpp
@@ -20,6 +20,8 @@ UNiagaraNodeCustomHlsl::UNiagaraNodeCustomHlsl(const FObjectInitializer& ObjectI
 
 	Signature.Name = TEXT("Custom Hlsl");
 	FunctionDisplayName = Signature.Name.ToString();
+
+	bIsShaderCodeShown = true;
 }
 
 const FString& UNiagaraNodeCustomHlsl::GetCustomHlsl() const
@@ -40,6 +42,21 @@ void UNiagaraNodeCustomHlsl::SetCustomHlsl(const FString& InCustomHlsl)
 	}
 }
 
+bool UNiagaraNodeCustomHlsl::IsShaderCodeShown() const
+{
+	return bIsShaderCodeShown;
+}
+
+void UNiagaraNodeCustomHlsl::SetShaderCodeShown(bool bInShown)
+{
+	if (bIsShaderCodeShown != bInShown)
+	{
+		Modify();
+		bIsShaderCodeShown = bInShown;
+		PostEditChange();
+	}
+}
+
 void UNiagaraNodeCustomHlsl::GetIncludeFilePaths(TArray<FNiagaraCustomHlslInclude>& OutCustomHlslIncludeFilePaths) const
 {
 	for (const FString& FilePath : VirtualIncludeFilePaths)
@@ -89,6 +106,11 @@ FLinearColor UNiagaraNodeCustomHlsl::GetNodeTitleColor() const
 	return UEdGraphSchema_Niagara::NodeTitleColor_CustomHlsl;
 }
 
+FText UNiagaraNodeCustomHlsl::GetTooltipText() const
+{
+	return LOCTEXT("CustomHlslTooltip", "Inserts the entered hlsl code into the translated script.");
+}
+
 bool UNiagaraNodeCustomHlsl::GetTokensFromString(const FString& InHlsl, TArray<FString>& OutTokens, bool IncludeComments, bool IncludeWhitespace)
 {
 	if (InHlsl.Len() == 0)
diff --git a/Engine/Plugins/FX/Niagara/Source/NiagaraEditor/Private/Widgets/SNiagaraGraphNodeConvert.cpp b/Engine/Plugins/FX/Niagara/Source/NiagaraEditor/Private/Widgets/SNiagaraGraphNodeConvert.cpp
--- a/Engine/Plugins/FX/Niagara/Source/NiagaraEditor/Private/Widgets/SNiagaraGraphNodeConvert.cpp
+++ b/Engine/Plugins/FX/Niagara/Source/NiagaraEditor/Private/Widgets/SNiagaraGraphNodeConvert.cpp
@@ -6,9 +6,9 @@
 #include "NiagaraConvertPinViewModel.h"
 #include "NiagaraConvertPinSocketViewModel.h"
 #include "SNiagaraConvertPinSocket.h"
-#include "Widgets/Input/SButton.h"
 #include "GraphEditorSettings.h"
 #include "Rendering/DrawElements.h"
+#include "Widgets/Input/SCheckBox.h"
 #include "SGraphPin.h"
 
 #define LOCTEXT_NAMESPACE "SNiagaraGraphNodeConvert"
@@ -22,41 +22,70 @@ void SNiagaraGraphNodeConvert::Construct(const FArguments& InArgs, UEdGraphNode*
 	UpdateGraphNode();
 }
 
-TSharedRef<SWidget> SNiagaraGraphNodeConvert::CreateTitleWidget(TSharedPtr<SNodeTitle> NodeTitle)
+void SNiagaraGraphNodeConvert::SetDefaultTitleAreaWidget(TSharedRef<SOverlay> DefaultTitleAreaWidget)
 {
-	return SNew(SHorizontalBox)
-		+ SHorizontalBox::Slot()
-		.AutoWidth()
+	DefaultTitleAreaWidget->AddSlot()
+	.Padding(5)
+	.VAlign(VAlign_Center)
+	.HAlign(HAlign_Right)
+	[
+		SNew(SCheckBox)
+		.OnCheckStateChanged(this, &SNiagaraGraphNodeConvert::ToggleShowWiring)
+		.IsChecked(this, &SNiagaraGraphNodeConvert::GetToggleButtonChecked)
+		.Cursor(EMouseCursor::Default)
+		.ToolTipText(LOCTEXT("ToggleShaderCode_Tooltip", "Toggle visibility of shader code."))
+		.Style(FAppStyle::Get(), "Graph.Node.AdvancedView")
 		[
-			SGraphNode::CreateTitleWidget(NodeTitle)
-		]
-		+ SHorizontalBox::Slot()
-		.AutoWidth()
-		.Padding(10.0f, 0.0f)
-		.VAlign(VAlign_Center)
-		.HAlign(HAlign_Right)
-		[
-			SNew(SButton)
-			.ButtonStyle(FAppStyle::Get(), "FlatButton")
-			.ToolTipText(LOCTEXT("ShowWiring_Tooltip", "Toggle visibility of the internal wiring switchboard."))
-			.OnClicked(this, &SNiagaraGraphNodeConvert::ToggleShowWiring)
-			.Content()
+			SNew(SHorizontalBox)
+			+SHorizontalBox::Slot()
+			.VAlign(VAlign_Center)
+			.HAlign(HAlign_Center)
 			[
 				SNew(SImage)
-				.Image(FAppStyle::GetBrush("Icons.Edit"))
+				.Image(this, &SNiagaraGraphNodeConvert::GetToggleButtonArrow)
 			]
-		];
+		]
+	];
 }
 
-FReply SNiagaraGraphNodeConvert::ToggleShowWiring()
+FReply SNiagaraGraphNodeConvert::OnMouseButtonDoubleClick(const FGeometry& InMyGeometry, const FPointerEvent& InMouseEvent)
 {
-	UNiagaraNodeConvert* ConvertNode = Cast<UNiagaraNodeConvert>(GraphNode);
-	if (ConvertNode != nullptr)
+	if(InMouseEvent.IsMouseButtonDown(EKeys::LeftMouseButton))
 	{
-		ConvertNode->SetWiringShown(!ConvertNode->IsWiringShown());
+		if (UNiagaraNodeConvert* ConvertNode = Cast<UNiagaraNodeConvert>(GraphNode))
+		{
+			ConvertNode->SetWiringShown(true);
+		}
+		OnDoubleClick.ExecuteIfBound(GraphNode);
+		return FReply::Handled();
 	}
+	return FReply::Unhandled();
+}
 
-	return FReply::Handled();
+void SNiagaraGraphNodeConvert::ToggleShowWiring(const ECheckBoxState NewState)
+{
+	if (UNiagaraNodeConvert* ConvertNode = Cast<UNiagaraNodeConvert>(GraphNode))
+	{
+		ConvertNode->SetWiringShown(NewState == ECheckBoxState::Unchecked ? false : true);
+	}
+}
+
+ECheckBoxState SNiagaraGraphNodeConvert::GetToggleButtonChecked() const
+{
+	if (UNiagaraNodeConvert* ConvertNode = Cast<UNiagaraNodeConvert>(GraphNode))
+	{
+		return ConvertNode->IsWiringShown() ? ECheckBoxState::Checked : ECheckBoxState::Unchecked;
+	}
+	return ECheckBoxState::Checked;
+}
+
+const FSlateBrush* SNiagaraGraphNodeConvert::GetToggleButtonArrow() const
+{
+	if (UNiagaraNodeConvert* ConvertNode = Cast<UNiagaraNodeConvert>(GraphNode))
+	{
+		return FAppStyle::GetBrush(ConvertNode->IsWiringShown() ? TEXT("Icons.ChevronUp") : TEXT("Icons.ChevronDown"));
+	}
+	return FAppStyle::GetBrush(TEXT("Icons.ChevronUp"));
 }
 
 TSharedRef<SWidget> ConstructPinSocketsRecursive(const TArray<TSharedRef<FNiagaraConvertPinSocketViewModel>>& SocketViewModels)
diff --git a/Engine/Plugins/FX/Niagara/Source/NiagaraEditor/Private/Widgets/SNiagaraGraphNodeConvert.h b/Engine/Plugins/FX/Niagara/Source/NiagaraEditor/Private/Widgets/SNiagaraGraphNodeConvert.h
--- a/Engine/Plugins/FX/Niagara/Source/NiagaraEditor/Private/Widgets/SNiagaraGraphNodeConvert.h
+++ b/Engine/Plugins/FX/Niagara/Source/NiagaraEditor/Private/Widgets/SNiagaraGraphNodeConvert.h
@@ -21,12 +21,15 @@ public:
 
 	virtual int32 OnPaint(const FPaintArgs& Args, const FGeometry& AllottedGeometry, const FSlateRect& MyCullingRect, FSlateWindowElementList& OutDrawElements, int32 LayerId, const FWidgetStyle& InWidgetStyle, bool bParentEnabled) const override;
 
-	virtual TSharedRef<SWidget> CreateTitleWidget(TSharedPtr<SNodeTitle> NodeTitle) override;
+	virtual void SetDefaultTitleAreaWidget(TSharedRef<SOverlay> DefaultTitleAreaWidget) override;
+	virtual FReply OnMouseButtonDoubleClick(const FGeometry& InMyGeometry, const FPointerEvent& InMouseEvent) override;
 
 protected:
 	virtual void AddPin(const TSharedRef<SGraphPin>& PinToAdd) override;
 
-	FReply ToggleShowWiring();
+	void ToggleShowWiring(const ECheckBoxState NewState);
+	ECheckBoxState GetToggleButtonChecked() const;
+	const FSlateBrush* GetToggleButtonArrow() const;
 
 private:
 	TSharedPtr<FNiagaraConvertPinViewModel> GetViewModelForPinWidget(TSharedRef<SGraphPin> GraphPin);
diff --git a/Engine/Plugins/FX/Niagara/Source/NiagaraEditor/Private/Widgets/SNiagaraGraphNodeCustomHlsl.cpp b/Engine/Plugins/FX/Niagara/Source/NiagaraEditor/Private/Widgets/SNiagaraGraphNodeCustomHlsl.cpp
--- a/Engine/Plugins/FX/Niagara/Source/NiagaraEditor/Private/Widgets/SNiagaraGraphNodeCustomHlsl.cpp
+++ b/Engine/Plugins/FX/Niagara/Source/NiagaraEditor/Private/Widgets/SNiagaraGraphNodeCustomHlsl.cpp
@@ -5,6 +5,7 @@
 #include "NiagaraGraph.h"
 #include "Widgets/SBoxPanel.h"
 #include "GraphEditorSettings.h"
+#include "Widgets/Input/SCheckBox.h"
 #include "Widgets/Input/SMultiLineEditableTextBox.h"
 
 #define LOCTEXT_NAMESPACE "SNiagaraGraphNodeCustomHlsl"
@@ -55,14 +56,42 @@ void SNiagaraGraphNodeCustomHlsl::RequestRenameOnSpawn()
 	}
 }
 
-TSharedRef<SWidget> SNiagaraGraphNodeCustomHlsl::CreateTitleWidget(TSharedPtr<SNodeTitle> NodeTitle)
+void SNiagaraGraphNodeCustomHlsl::SetDefaultTitleAreaWidget(TSharedRef<SOverlay> DefaultTitleAreaWidget)
 {
-	return SNew(SHorizontalBox)
-		+ SHorizontalBox::Slot()
-		.AutoWidth()
+	DefaultTitleAreaWidget->AddSlot()
+	.Padding(5)
+	.VAlign(VAlign_Center)
+	.HAlign(HAlign_Right)
+	[
+		SNew(SCheckBox)
+		.OnCheckStateChanged(this, &SNiagaraGraphNodeCustomHlsl::ToggleShowShaderCode)
+		.IsChecked(this, &SNiagaraGraphNodeCustomHlsl::GetToggleButtonChecked)
+		.Cursor(EMouseCursor::Default)
+		.ToolTipText(LOCTEXT("ToggleShaderCode_Tooltip", "Toggle visibility of shader code."))
+		.Style(FAppStyle::Get(), "Graph.Node.AdvancedView")
 		[
-			SGraphNode::CreateTitleWidget(NodeTitle)
-		];
+			SNew(SHorizontalBox)
+			+SHorizontalBox::Slot()
+			.VAlign(VAlign_Center)
+			.HAlign(HAlign_Center)
+			[
+				SNew(SImage)
+				.Image(this, &SNiagaraGraphNodeCustomHlsl::GetToggleButtonArrow)
+			]
+		]
+	];
+}
+
+FReply SNiagaraGraphNodeCustomHlsl::OnMouseButtonDoubleClick(const FGeometry& InMyGeometry,	const FPointerEvent& InMouseEvent)
+{
+	if(InMouseEvent.IsMouseButtonDown(EKeys::LeftMouseButton))
+	{
+		UNiagaraNodeCustomHlsl* CustomNode = CastChecked<UNiagaraNodeCustomHlsl>(GraphNode);
+		CustomNode->SetShaderCodeShown(true);
+		OnDoubleClick.ExecuteIfBound(CustomNode);
+		return FReply::Handled();
+	}
+	return FReply::Unhandled();
 }
 
 void SNiagaraGraphNodeCustomHlsl::CreateBelowPinControls(TSharedPtr<SVerticalBox> MainBox)
@@ -81,7 +110,8 @@ void SNiagaraGraphNodeCustomHlsl::CreateBelowPinControls(TSharedPtr<SVerticalBox
 		.Text(GetText)
 		.Marshaller(SyntaxHighlighter)
 		.OnKeyCharHandler(this, &SNiagaraGraphNodeCustomHlsl::OnShaderTextKeyChar)
-		.OnTextCommitted(TextCommit);
+		.OnTextCommitted(TextCommit)
+		.Visibility(this, &SNiagaraGraphNodeCustomHlsl::GetShaderTextVisibility);
 	MainBox->AddSlot()[
 		SNew(SHorizontalBox)
 		+ SHorizontalBox::Slot()
@@ -107,4 +137,28 @@ FReply SNiagaraGraphNodeCustomHlsl::OnShaderTextKeyChar(const FGeometry&, const
 	return FReply::Unhandled();
 }
 
+EVisibility SNiagaraGraphNodeCustomHlsl::GetShaderTextVisibility() const
+{
+	UNiagaraNodeCustomHlsl* CustomNode = CastChecked<UNiagaraNodeCustomHlsl>(GraphNode);
+	return CustomNode->IsShaderCodeShown() ? EVisibility::Visible : EVisibility::Collapsed;
+}
+
+void SNiagaraGraphNodeCustomHlsl::ToggleShowShaderCode(const ECheckBoxState NewCheckedState)
+{
+	UNiagaraNodeCustomHlsl* CustomNode = CastChecked<UNiagaraNodeCustomHlsl>(GraphNode);
+	CustomNode->SetShaderCodeShown(NewCheckedState == ECheckBoxState::Unchecked ? false : true);
+}
+
+ECheckBoxState SNiagaraGraphNodeCustomHlsl::GetToggleButtonChecked() const
+{
+	UNiagaraNodeCustomHlsl* CustomNode = CastChecked<UNiagaraNodeCustomHlsl>(GraphNode);
+	return CustomNode->IsShaderCodeShown() ? ECheckBoxState::Checked : ECheckBoxState::Unchecked;
+}
+
+const FSlateBrush* SNiagaraGraphNodeCustomHlsl::GetToggleButtonArrow() const
+{
+	UNiagaraNodeCustomHlsl* CustomNode = CastChecked<UNiagaraNodeCustomHlsl>(GraphNode);
+	return FAppStyle::GetBrush(CustomNode->IsShaderCodeShown() ? TEXT("Icons.ChevronUp") : TEXT("Icons.ChevronDown"));
+}
+
 #undef LOCTEXT_NAMESPACE
diff --git a/Engine/Plugins/FX/Niagara/Source/NiagaraEditor/Private/Widgets/SNiagaraGraphNodeCustomHlsl.h b/Engine/Plugins/FX/Niagara/Source/NiagaraEditor/Private/Widgets/SNiagaraGraphNodeCustomHlsl.h
--- a/Engine/Plugins/FX/Niagara/Source/NiagaraEditor/Private/Widgets/SNiagaraGraphNodeCustomHlsl.h
+++ b/Engine/Plugins/FX/Niagara/Source/NiagaraEditor/Private/Widgets/SNiagaraGraphNodeCustomHlsl.h
@@ -19,12 +19,12 @@ public:
 	void Construct(const FArguments& InArgs, UEdGraphNode* InGraphNode);
 
 	//~ SGraphNode api
-	virtual TSharedRef<SWidget> CreateTitleWidget(TSharedPtr<SNodeTitle> NodeTitle) override;
+	virtual void SetDefaultTitleAreaWidget(TSharedRef<SOverlay> DefaultTitleAreaWidget) override;
 	virtual bool IsNameReadOnly() const override;
 	virtual void RequestRenameOnSpawn() override;
 	virtual bool ShouldAllowCulling() const override { return false; }
-
-	virtual void CreateBelowPinControls(TSharedPtr<SVerticalBox> MainBox);
+	virtual FReply OnMouseButtonDoubleClick( const FGeometry& InMyGeometry, const FPointerEvent& InMouseEvent ) override;
+	virtual void CreateBelowPinControls(TSharedPtr<SVerticalBox> MainBox) override;
 protected:
 
 	TSharedPtr<FNiagaraHLSLSyntaxHighlighter> SyntaxHighlighter;
@@ -32,4 +32,9 @@ protected:
 
 private:
 	FReply OnShaderTextKeyChar(const FGeometry& MyGeometry, const FCharacterEvent& InCharacterEvent);
+
+	EVisibility GetShaderTextVisibility() const;
+	void ToggleShowShaderCode(const ECheckBoxState NewCheckedState);
+	ECheckBoxState GetToggleButtonChecked() const;
+	const FSlateBrush* GetToggleButtonArrow() const;
 };
diff --git a/Engine/Plugins/FX/Niagara/Source/NiagaraEditor/Public/NiagaraNodeCustomHlsl.h b/Engine/Plugins/FX/Niagara/Source/NiagaraEditor/Public/NiagaraNodeCustomHlsl.h
--- a/Engine/Plugins/FX/Niagara/Source/NiagaraEditor/Public/NiagaraNodeCustomHlsl.h
+++ b/Engine/Plugins/FX/Niagara/Source/NiagaraEditor/Public/NiagaraNodeCustomHlsl.h
@@ -19,6 +19,9 @@ public:
 	const FString& GetCustomHlsl() const;
 	void SetCustomHlsl(const FString& InCustomHlsl);
 
+	bool IsShaderCodeShown() const;
+	void SetShaderCodeShown(bool bInShown);
+
 	void GetIncludeFilePaths(TArray<FNiagaraCustomHlslInclude>& OutCustomHlslIncludeFilePaths) const;
 
 	UPROPERTY()
@@ -27,6 +30,7 @@ public:
 	virtual TSharedPtr<SGraphNode> CreateVisualWidget() override;
 	virtual void OnRenameNode(const FString& NewName) override;
 	virtual FLinearColor GetNodeTitleColor() const override;
+	virtual FText GetTooltipText() const override;
 
 	FText GetHlslText() const;
 	void OnCustomHlslTextCommitted(const FText& InText, ETextCommit::Type InType);
@@ -93,4 +97,8 @@ private:
 	// For example, /Plugin/FX/Niagara maps to /Engine/Plugins/FX/Niagara/Shaders. Custom mappings can be added via AddShaderSourceDirectoryMapping().
 	UPROPERTY(EditAnywhere, Category = "HLSL")
 	TArray<FString> VirtualIncludeFilePaths;
+	
+	// Is the shader code UI shown?
+	UPROPERTY()
+	bool bIsShaderCodeShown;
 };
