Couple of fixes for NiagaraScript's DDC content to be non-deterministic
-Cleans up some of the data written into the payload (removing compile times & writing out zeroes for unneeded non-determinstic pointers)
-Incorporates the AssetPath of the script to the DDC key to prevent identical scripts from sharing the same DDC as we still have some data in the DDC value (node guis & emitter names, ..)

diff --git a/Engine/Plugins/FX/Niagara/Source/Niagara/Classes/NiagaraScript.h b/Engine/Plugins/FX/Niagara/Source/Niagara/Classes/NiagaraScript.h
--- a/Engine/Plugins/FX/Niagara/Source/Niagara/Classes/NiagaraScript.h
+++ b/Engine/Plugins/FX/Niagara/Source/Niagara/Classes/NiagaraScript.h
@@ -534,9 +534,6 @@ public:
 	UPROPERTY()
 	FString ErrorMsg;
 
-	UPROPERTY()
-	float CompileTime;
-
 	/** Array of all compile events generated last time the script was compiled.*/
 	UPROPERTY()
 	TArray<FNiagaraCompileEvent> LastCompileEvents;
@@ -1114,11 +1111,11 @@ public:
 		gather up the results with. The function returns whether or not any compiles were actually issued. */
 	NIAGARA_API bool RequestExternallyManagedAsyncCompile(const TSharedPtr<FNiagaraCompileRequestDataBase, ESPMode::ThreadSafe>& RequestData, const TSharedPtr<FNiagaraCompileRequestDuplicateDataBase, ESPMode::ThreadSafe>& RequestDuplicateData, FNiagaraVMExecutableDataId& OutCompileId, uint32& OutAsyncHandle);
 
-	/** Builds the DDC string for the derived data cache using the supplied CompiledId */
-	static FString BuildNiagaraDDCKeyString(const FNiagaraVMExecutableDataId& CompileId);
+	/** Builds the DDC string for the derived data cache using the supplied CompiledId and ScriptPath */
+	static FString BuildNiagaraDDCKeyString(const FNiagaraVMExecutableDataId& CompileId, const FString& ScriptPath);
 
 	/** Creates a string key for the derived data cache */
-	FString GetNiagaraDDCKeyString(const FGuid& ScriptVersion);
+	FString GetNiagaraDDCKeyString(const FGuid& ScriptVersion, const FString& ScriptPath);
 
 	/** Callback issued whenever a VM script compilation successfully happened (even if the results are a script that cannot be executed due to errors)*/
 	NIAGARA_API FOnScriptCompiled& OnVMScriptCompiled();
diff --git a/Engine/Plugins/FX/Niagara/Source/Niagara/Private/NiagaraAsyncCompile.cpp b/Engine/Plugins/FX/Niagara/Source/Niagara/Private/NiagaraAsyncCompile.cpp
--- a/Engine/Plugins/FX/Niagara/Source/Niagara/Private/NiagaraAsyncCompile.cpp
+++ b/Engine/Plugins/FX/Niagara/Source/Niagara/Private/NiagaraAsyncCompile.cpp
@@ -259,7 +259,7 @@ void FNiagaraAsyncCompileTask::PrecompileData()
 	// because otherwise we would save the compiled data under the wrong ddc key
 	FNiagaraVMExecutableDataId NewID;
 	ScriptPair.CompiledScript->ComputeVMCompilationId(NewID, FGuid());
-	FString CurrentDDCKey = UNiagaraScript::BuildNiagaraDDCKeyString(NewID);
+	FString CurrentDDCKey = UNiagaraScript::BuildNiagaraDDCKeyString(NewID, AssetPath);
 	if (DDCKey != CurrentDDCKey && LogNiagara.GetVerbosity() >= ELogVerbosity::Verbose)
 	{
 		UE_LOG(LogNiagara, Warning, TEXT("Compile ID for %s changed during the compilation task."), *AssetPath);
@@ -372,7 +372,7 @@ void FNiagaraAsyncCompileTask::ProcessResult()
 	// convert results to be saved in the ddc
 	ScriptPair.CompileResults = ExeData;
 	ScriptPair.bResultsReady = true;
-	ExeData->CompileTime = FPlatformTime::Seconds() - StartCompileTime;
+	ScriptPair.CompileTime = FPlatformTime::Seconds() - StartCompileTime;
 	ScriptPair.CompiledScript->ExecToBinaryData(ScriptPair.CompiledScript, DDCOutData, *ExeData);
 	UE_LOG(LogNiagara, Verbose, TEXT("Got %i bytes in ddc data for %s"), DDCOutData.Num(), *AssetPath);
 }
diff --git a/Engine/Plugins/FX/Niagara/Source/Niagara/Private/NiagaraAsyncCompile.h b/Engine/Plugins/FX/Niagara/Source/Niagara/Private/NiagaraAsyncCompile.h
--- a/Engine/Plugins/FX/Niagara/Source/Niagara/Private/NiagaraAsyncCompile.h
+++ b/Engine/Plugins/FX/Niagara/Source/Niagara/Private/NiagaraAsyncCompile.h
@@ -22,6 +22,7 @@ struct FEmitterCompiledScriptPair
 	uint32 PendingJobID = INDEX_NONE; // this is the ID for any active shader compiler worker job
 	FNiagaraVMExecutableDataId CompileId;
 	TSharedPtr<FNiagaraVMExecutableData> CompileResults;
+	double CompileTime = 0;
 };
 
 UENUM()
diff --git a/Engine/Plugins/FX/Niagara/Source/Niagara/Private/NiagaraScript.cpp b/Engine/Plugins/FX/Niagara/Source/Niagara/Private/NiagaraScript.cpp
--- a/Engine/Plugins/FX/Niagara/Source/Niagara/Private/NiagaraScript.cpp
+++ b/Engine/Plugins/FX/Niagara/Source/Niagara/Private/NiagaraScript.cpp
@@ -234,7 +234,6 @@ FNiagaraVMExecutableData::FNiagaraVMExecutableData()
 	, LastCompileStatus(ENiagaraScriptCompileStatus::NCS_Unknown)
 #if WITH_EDITORONLY_DATA
 	, bReadsAttributeData(false)
-	, CompileTime(0.0f)
 #endif
 	, bReadsSignificanceIndex(false)
 	, bNeedsGPUContextInit(false)
@@ -911,20 +910,33 @@ FNiagaraVMExecutableDataId& UNiagaraScript::GetLastGeneratedVMId(const FGuid& Ve
 	return VersionData[0].LastGeneratedVMId;
 }
 
-FString UNiagaraScript::BuildNiagaraDDCKeyString(const FNiagaraVMExecutableDataId& CompileId)
+// while it would be nice to not have to include the script path into the DDC key (so that two scripts on different
+// could share the same DDC data) we currently have a number of elements of the system/emitter scripts that get
+// incorporated into the data stored in the DDC that make it difficult to have scripts owned by different emitters
+// generated deterministic results.  Current limitations include:
+//	-stat scopes and emitter functions incorporating identifiers into the generated hlsl
+//	-graph specific information (node guids) used in serialized errors/warnings for backtracking to the source in the UI
+FString UNiagaraScript::BuildNiagaraDDCKeyString(const FNiagaraVMExecutableDataId& CompileId, const FString& ScriptPath)
 {
 	enum { UE_NIAGARA_COMPILATION_DERIVEDDATA_VER = 2 };
 
-	FString KeyString = FString::Printf(TEXT("%i_%i"),
+	FString KeyString;
+	KeyString.Reserve(1024);
+
+	KeyString.Appendf(TEXT("%i_%i"),
 		(int32)UE_NIAGARA_COMPILATION_DERIVEDDATA_VER, GNiagaraSkipVectorVMBackendOptimizations);
 
+	KeyString.AppendChar(TCHAR('_'));
+	KeyString.Append(ScriptPath);
+	KeyString.AppendChar(TCHAR('_'));
+
 	CompileId.AppendKeyString(KeyString);
 	return FDerivedDataCacheInterface::BuildCacheKey(TEXT("NiagaraScriptDerivedData"), NIAGARASCRIPT_DERIVEDDATA_VER, *KeyString);
 }
 
-FString UNiagaraScript::GetNiagaraDDCKeyString(const FGuid& ScriptVersion)
+FString UNiagaraScript::GetNiagaraDDCKeyString(const FGuid& ScriptVersion, const FString& ScriptPath)
 {
-	return BuildNiagaraDDCKeyString(GetLastGeneratedVMId(ScriptVersion));
+	return BuildNiagaraDDCKeyString(GetLastGeneratedVMId(ScriptVersion), ScriptPath);
 }
 
 void UNiagaraScript::ComputeVMCompilationId(FNiagaraVMExecutableDataId& Id, FGuid VersionGuid) const
@@ -2806,8 +2818,10 @@ void UNiagaraScript::RequestCompile(const FGuid& ScriptVersion, bool bForceCompi
 			return;
 		}
 
+		const FString ScriptPathName = GetPathName();
+
 		// check the ddc first
-		if (GetDerivedDataCacheRef().GetSynchronous(*GetNiagaraDDCKeyString(ScriptVersion), OutData, GetPathName()))
+		if (GetDerivedDataCacheRef().GetSynchronous(*GetNiagaraDDCKeyString(ScriptVersion, ScriptPathName), OutData, ScriptPathName))
 		{
 			FNiagaraVMExecutableData ExeData;
 			if (BinaryToExecData(this, OutData, ExeData))
@@ -2821,7 +2835,7 @@ void UNiagaraScript::RequestCompile(const FGuid& ScriptVersion, bool bForceCompi
 		ActiveCompileRoots.Empty();
 		RequestDuplicateData->GetDuplicatedObjects(ActiveCompileRoots);
 
-		FNiagaraCompileOptions Options(GetUsage(), GetUsageId(), ScriptData->ModuleUsageBitmask, GetPathName(), GetFullName(), GetName());
+		FNiagaraCompileOptions Options(GetUsage(), GetUsageId(), ScriptData->ModuleUsageBitmask, ScriptPathName, GetFullName(), GetName());
 		int32 JobHandle = NiagaraModule.StartScriptCompileJob(RequestData.Get(), RequestDuplicateData.Get(), Options);
 		TSharedPtr<FNiagaraVMExecutableData> ExeData = NiagaraModule.GetCompileJobResult(JobHandle, true);
 		if (ExeData)
@@ -2831,7 +2845,7 @@ void UNiagaraScript::RequestCompile(const FGuid& ScriptVersion, bool bForceCompi
 			if (ExecToBinaryData(this, OutData, *ExeData))
 			{
 				COOK_STAT(Timer.AddMiss(OutData.Num()));
-				GetDerivedDataCacheRef().Put(*GetNiagaraDDCKeyString(ScriptVersion), OutData, GetPathName());
+				GetDerivedDataCacheRef().Put(*GetNiagaraDDCKeyString(ScriptVersion, ScriptPathName), OutData, ScriptPathName);
 			}
 		}
 		ActiveCompileRoots.Empty();
diff --git a/Engine/Plugins/FX/Niagara/Source/Niagara/Private/NiagaraSystem.cpp b/Engine/Plugins/FX/Niagara/Source/Niagara/Private/NiagaraSystem.cpp
--- a/Engine/Plugins/FX/Niagara/Source/Niagara/Private/NiagaraSystem.cpp
+++ b/Engine/Plugins/FX/Niagara/Source/Niagara/Private/NiagaraSystem.cpp
@@ -2683,7 +2683,7 @@ bool UNiagaraSystem::QueryCompileComplete(bool bWait, bool bDoPost, bool bDoNotA
 			if (EmitterCompiledScriptPair.bResultsReady)
 			{
 				TSharedPtr<FNiagaraVMExecutableData> ExeData = EmitterCompiledScriptPair.CompileResults;
-				CompileRequest.CombinedCompileTime += ExeData->CompileTime;
+				CompileRequest.CombinedCompileTime += EmitterCompiledScriptPair.CompileTime;
 				UNiagaraScript* CompiledScript = EmitterCompiledScriptPair.CompiledScript;
 
 				// generate the ObjectNameMap from the source (from the duplicated data if available).
@@ -3182,7 +3182,7 @@ bool UNiagaraSystem::RequestCompile(bool bForce, FNiagaraSystemUpdateContext* Op
 			};
 
 			TSharedPtr<FNiagaraAsyncCompileTask, ESPMode::ThreadSafe> AsyncTask = *ActiveCompilation.DDCTasks.FindByPredicate(InPairs);
-			AsyncTask->DDCKey = UNiagaraScript::BuildNiagaraDDCKeyString(AsyncTask->GetScriptPair().CompileId);
+			AsyncTask->DDCKey = UNiagaraScript::BuildNiagaraDDCKeyString(AsyncTask->GetScriptPair().CompileId, AsyncTask->AssetPath);
 			AsyncTask->PrecompileReference = PrecompileReference;
 			AsyncTask->StartTaskTime = FPlatformTime::Seconds();
 
diff --git a/Engine/Plugins/FX/Niagara/Source/NiagaraEditor/Private/NiagaraCompiler.cpp b/Engine/Plugins/FX/Niagara/Source/NiagaraEditor/Private/NiagaraCompiler.cpp
--- a/Engine/Plugins/FX/Niagara/Source/NiagaraEditor/Private/NiagaraCompiler.cpp
+++ b/Engine/Plugins/FX/Niagara/Source/NiagaraEditor/Private/NiagaraCompiler.cpp
@@ -2034,8 +2034,6 @@ TOptional<FNiagaraCompileResults> FHlslNiagaraCompiler::GetCompileResult(int32 J
 			}
 		}
 		Results.CompileTime = (float)(FPlatformTime::Seconds() - CompilationJob->StartTime);
-		Results.Data->CompileTime = Results.CompileTime;
-
 
 		Results.Data->CalledVMExternalFunctions.Empty(CompilationOutput.CalledVMFunctionTable.Num());
 		for (FCalledVMFunction& FuncInfo : CompilationOutput.CalledVMFunctionTable)
diff --git a/Engine/Source/Runtime/VectorVM/Private/VectorVMExperimental_Optimizer.inl b/Engine/Source/Runtime/VectorVM/Private/VectorVMExperimental_Optimizer.inl
--- a/Engine/Source/Runtime/VectorVM/Private/VectorVMExperimental_Optimizer.inl
+++ b/Engine/Source/Runtime/VectorVM/Private/VectorVMExperimental_Optimizer.inl
@@ -163,11 +163,29 @@ void FreezeVectorVMOptimizeContext(const FVectorVMOptimizeContext& Context, TArr
 	FMemory::Memcpy(DestData + InputDataSetOffsetsOffset, Context.InputDataSetOffsets, InputDataSetOffsetsSize);
 	FMemory::Memcpy(DestData + ExtFnOffset, Context.ExtFnTable, ExtFnSize);
 
-	// we want to clear out the pointers to any callbacks
+	// we want to clear out the pointers to any callbacks or pointers to the data that is being pushed past
+	// the end of the struct
 	FVectorVMOptimizeContext& DestContext = *reinterpret_cast<FVectorVMOptimizeContext*>(DestData);
+
+	// relocated buffers
+	DestContext.OutputBytecode = nullptr;
+	DestContext.ConstRemap[0] = nullptr;
+	DestContext.ConstRemap[1] = nullptr;
+	DestContext.InputRemapTable = nullptr;
+	DestContext.InputDataSetOffsets = nullptr;
+	DestContext.ExtFnTable = nullptr;
+
+	// callbacks
 	DestContext.Init.ReallocFn = nullptr;
 	DestContext.Init.FreeFn = nullptr;
 	DestContext.Error.CallbackFn = nullptr;
+
+	// external function pointers
+	FVectorVMExtFunctionData* ExtFunctionTable = reinterpret_cast<FVectorVMExtFunctionData*>(DestData + ExtFnOffset);
+	for (uint32 ExtFunctionIt = 0; ExtFunctionIt < Context.NumExtFns; ++ExtFunctionIt)
+	{
+		ExtFunctionTable[ExtFunctionIt].Function = nullptr;
+	}
 }
 
 static void* VectorVMFrozenRealloc(void* Ptr, size_t NumBytes, const char* Filename, int LineNum)
