- PR #9805: [Niagara] GetPreSkinnedLocalBounds was missing GPU support (Contributed by [REDACTED_PERSON])

#rnx

diff --git a/Engine/Plugins/FX/Niagara/Source/Niagara/Private/NiagaraDataInterfaceSkeletalMesh.cpp b/Engine/Plugins/FX/Niagara/Source/Niagara/Private/NiagaraDataInterfaceSkeletalMesh.cpp
--- a/Engine/Plugins/FX/Niagara/Source/Niagara/Private/NiagaraDataInterfaceSkeletalMesh.cpp
+++ b/Engine/Plugins/FX/Niagara/Source/Niagara/Private/NiagaraDataInterfaceSkeletalMesh.cpp
@@ -2816,6 +2816,8 @@ bool UNiagaraDataInterfaceSkeletalMesh::GetFunctionHLSL(const FNiagaraDataInterf
 {
 	static const TSet<FName> ValidGpuFunctions =
 	{
+		NDISkelMeshLocal::NAME_GetPreSkinnedLocalBounds,
+		//////////////////////////////////////////////////////////////////////////////////////////////////////////////////////////////
 		FSkeletalMeshInterfaceHelper::GetTriCoordVerticesName,
 		FSkeletalMeshInterfaceHelper::GetTriangleCountName,
 		FSkeletalMeshInterfaceHelper::GetFilteredTriangleCountName,
