- Ensure we don't overflow buffer size for debug lines on GPU

diff --git a/Engine/Plugins/FX/Niagara/Source/Niagara/Private/NiagaraGpuComputeDebug.cpp b/Engine/Plugins/FX/Niagara/Source/Niagara/Private/NiagaraGpuComputeDebug.cpp
--- a/Engine/Plugins/FX/Niagara/Source/Niagara/Private/NiagaraGpuComputeDebug.cpp
+++ b/Engine/Plugins/FX/Niagara/Source/Niagara/Private/NiagaraGpuComputeDebug.cpp
@@ -204,7 +204,8 @@ FNiagaraSimulationDebugDrawData* FNiagaraGpuComputeDebug::GetSimulationDebugDraw
 
 	FNiagaraSimulationDebugDrawData* DebugDrawData = GetSimulationDebugDrawData(SystemInstanceID);
 
-	const int MaxLineInstancesToUse = FMath::Max3(DebugDrawData->GpuLineMaxInstances, (uint32)GNiagaraGpuComputeDebug_MaxLineInstances, OverrideMaxDebugLines);
+	const uint32 MaxPossibleLines = FMath::Min(FMath::DivideAndRoundDown(GetMaxBufferDimension(), uint64(NiagaraGpuComputeDebugLocal::NumUintsPerLine)), uint64(TNumericLimits<uint32>::Max() >> 1));
+	const uint32 MaxLineInstancesToUse = FMath::Min(FMath::Max3(DebugDrawData->GpuLineMaxInstances, uint32(GNiagaraGpuComputeDebug_MaxLineInstances), OverrideMaxDebugLines), MaxPossibleLines);
 	if (DebugDrawData->GpuLineMaxInstances != MaxLineInstancesToUse)
 	{
 		DebugDrawData->GpuLineBufferArgs.Release();
