- Add severity selection to banned DI / simulation stage budget rules

#rnx

diff --git a/Engine/Plugins/FX/Niagara/Source/NiagaraEditor/Public/NiagaraValidationRules.h b/Engine/Plugins/FX/Niagara/Source/NiagaraEditor/Public/NiagaraValidationRules.h
--- a/Engine/Plugins/FX/Niagara/Source/NiagaraEditor/Public/NiagaraValidationRules.h
+++ b/Engine/Plugins/FX/Niagara/Source/NiagaraEditor/Public/NiagaraValidationRules.h
@@ -77,6 +77,9 @@ class UNiagaraValidationRule_BannedDataInterfaces : public UNiagaraValidationRul
 	GENERATED_BODY()
 
 public:
+	UPROPERTY(EditAnywhere, Category = Validation)
+	ENiagaraValidationSeverity Severity = ENiagaraValidationSeverity::Warning;
+
 	UPROPERTY(EditAnywhere, Category = Validation)
 	bool bBanOnGpu = true;
 
@@ -146,6 +149,10 @@ public:
 	UPROPERTY(EditAnywhere, Category = Validation)
 	bool bMaxTotalIterationsEnabled = false;
 
+	/** How do we want to repro the error in the stack */
+	UPROPERTY(EditAnywhere, Category = Validation)
+	ENiagaraValidationSeverity Severity = ENiagaraValidationSeverity::Warning;
+
 	/** Maximum number of simulation stages allowed, where 0 means no simulation stages. */
 	UPROPERTY(EditAnywhere, Category = Validation, meta=(EditCondition="bMaxSimulationStagesEnabled"))
 	int32 MaxSimulationStages = 0;
