- Missing file which used the severity for the validation rule

#rnx
#preflght trivial

diff --git a/Engine/Plugins/FX/Niagara/Source/NiagaraEditor/Private/NiagaraValidationRules.cpp b/Engine/Plugins/FX/Niagara/Source/NiagaraEditor/Private/NiagaraValidationRules.cpp
--- a/Engine/Plugins/FX/Niagara/Source/NiagaraEditor/Private/NiagaraValidationRules.cpp
+++ b/Engine/Plugins/FX/Niagara/Source/NiagaraEditor/Private/NiagaraValidationRules.cpp
@@ -504,7 +504,7 @@ void UNiagaraValidationRule_BannedDataInterfaces::CheckValidity(const FNiagaraVa
 					if (PlatformConflictsString.IsEmpty() == false)
 					{
 						Results.Emplace(
-							ENiagaraValidationSeverity::Warning,
+							Severity,
 							FText::Format(WarningFormat, FText::FromName(UsageContext.Variable.GetName())),
 							FText::Format(EmitterDescFormat, FText::FromName(UsageContext.Variable.GetName()), FText::FromName(DIClass->GetFName()), FText::FromString(PlatformConflictsString)),
 							NiagaraValidation::GetStackEntry<UNiagaraStackEmitterPropertiesItem>(EmitterViewModel->GetEmitterStackViewModel())
@@ -771,7 +771,7 @@ void UNiagaraValidationRule_SimulationStageBudget::CheckValidity(const FNiagaraV
 			{
 				UNiagaraStackEmitterPropertiesItem* EmitterProperties = NiagaraValidation::GetStackEntry<UNiagaraStackEmitterPropertiesItem>(EmitterHandleModel.Get().GetEmitterStackViewModel());
 				OutResults.Emplace_GetRef(
-					ENiagaraValidationSeverity::Error,
+					Severity,
 					FText::Format(LOCTEXT("SimStageTooManyIterationsFormat", "Simulation Stage '{0}' has too many iterations"), FText::FromName(SimStage->SimulationStageName)),
 					FText::Format(LOCTEXT("SimStageTooManyIterationsDetailedFormat", "Simulation Stage '{0}' has {1} iterations and we only allow {2}"), FText::FromName(SimStage->SimulationStageName), FText::AsNumber(SimStage->Iterations), FText::AsNumber(MaxIterationsPerStage)),
 					EmitterProperties
@@ -783,7 +783,7 @@ void UNiagaraValidationRule_SimulationStageBudget::CheckValidity(const FNiagaraV
 		{
 			UNiagaraStackEmitterPropertiesItem* EmitterProperties = NiagaraValidation::GetStackEntry<UNiagaraStackEmitterPropertiesItem>(EmitterHandleModel.Get().GetEmitterStackViewModel());
 			OutResults.Emplace(
-				ENiagaraValidationSeverity::Error,
+				Severity,
 				LOCTEXT("SimStageTooManyTotalIterationsFormat", "Emitter has too many total simulation stage iterations"),
 				FText::Format(LOCTEXT("SimStageTooManyTotalIterationsDetailedFormat", "Emitter has {0} total simulation stage iterations and we only allow {1}"), FText::AsNumber(TotalIterations), FText::AsNumber(MaxTotalIterations)),
 				EmitterProperties
@@ -794,7 +794,7 @@ void UNiagaraValidationRule_SimulationStageBudget::CheckValidity(const FNiagaraV
 		{
 			UNiagaraStackEmitterPropertiesItem* EmitterProperties = NiagaraValidation::GetStackEntry<UNiagaraStackEmitterPropertiesItem>(EmitterHandleModel.Get().GetEmitterStackViewModel());
 			OutResults.Emplace(
-				ENiagaraValidationSeverity::Error,
+				Severity,
 				LOCTEXT("TooManySimStagesFormat", "Emitter has too many simulation stages"),
 				FText::Format(LOCTEXT("TooManySimStagesDetailedFormat", "Emitter has {0} simulation stages active and we only allow {1}"), FText::AsNumber(TotalEnabledStages), FText::AsNumber(MaxSimulationStages)),
 				EmitterProperties
