- Add dispatch costs to system / emitter stats

diff --git a/Engine/Plugins/FX/Niagara/Source/Niagara/Classes/NiagaraComputeExecutionContext.h b/Engine/Plugins/FX/Niagara/Source/Niagara/Classes/NiagaraComputeExecutionContext.h
--- a/Engine/Plugins/FX/Niagara/Source/Niagara/Classes/NiagaraComputeExecutionContext.h
+++ b/Engine/Plugins/FX/Niagara/Source/Niagara/Classes/NiagaraComputeExecutionContext.h
@@ -97,6 +97,11 @@ struct FNiagaraComputeExecutionContext
 	void SetDebugSimName(const TCHAR*) { }
 #endif
 
+#if STATS
+	TStatId SystemStatID;
+	TStatId EmitterStatID;
+#endif
+
 //-TOOD:private:
 	void ResetInternal(FNiagaraGpuComputeDispatchInterface* ComputeDispatchInterface);
 
diff --git a/Engine/Plugins/FX/Niagara/Source/Niagara/Private/NiagaraEmitterInstance.cpp b/Engine/Plugins/FX/Niagara/Source/Niagara/Private/NiagaraEmitterInstance.cpp
--- a/Engine/Plugins/FX/Niagara/Source/Niagara/Private/NiagaraEmitterInstance.cpp
+++ b/Engine/Plugins/FX/Niagara/Source/Niagara/Private/NiagaraEmitterInstance.cpp
@@ -416,6 +416,10 @@ void FNiagaraEmitterInstance::Init(int32 InEmitterIdx, FNiagaraSystemInstanceID
 			GPUExecContext->ProfilingEmitterPtr = CachedEmitter.ToWeakPtr();
 			GPUExecContext->MainDataSet = ParticleDataSet;
 			GPUExecContext->GPUScript_RT = EmitterData->GetGPUComputeScript()->GetRenderThreadScript();
+		#if STATS
+			GPUExecContext->SystemStatID = ParentSystemInstance->GetSystem()->GetStatID(false, false);
+			GPUExecContext->EmitterStatID = CachedEmitter.Emitter->GetStatID(false, false);
+		#endif
 		}
 	}
 
diff --git a/Engine/Plugins/FX/Niagara/Source/Niagara/Private/NiagaraGpuComputeDispatch.cpp b/Engine/Plugins/FX/Niagara/Source/Niagara/Private/NiagaraGpuComputeDispatch.cpp
--- a/Engine/Plugins/FX/Niagara/Source/Niagara/Private/NiagaraGpuComputeDispatch.cpp
+++ b/Engine/Plugins/FX/Niagara/Source/Niagara/Private/NiagaraGpuComputeDispatch.cpp
@@ -1390,6 +1390,11 @@ void FNiagaraGpuComputeDispatch::ExecuteTicks(FRDGBuilder& GraphBuilder, TConstA
 
 void FNiagaraGpuComputeDispatch::DispatchStage(FRDGBuilder& GraphBuilder, const FNiagaraGPUSystemTick& Tick, const FNiagaraComputeInstanceData& InstanceData, const FNiagaraSimStageData& SimStageData)
 {
+#if STATS
+	FScopeCycleCounter SystemCounterStat(InstanceData.Context->SystemStatID);
+	FScopeCycleCounter EmitterCounterStat(InstanceData.Context->EmitterStatID);
+#endif
+
 	// Setup source buffer
 	if ( SimStageData.Source != nullptr )
 	{
