- Fix ExecuteBehavior not showing for particle iteration stages
- Move dispatch parameters into separate category
- Move Simulation Stage Name to top of sim stage

diff --git a/Engine/Plugins/FX/Niagara/Source/Niagara/Public/NiagaraSimulationStageBase.h b/Engine/Plugins/FX/Niagara/Source/Niagara/Public/NiagaraSimulationStageBase.h
--- a/Engine/Plugins/FX/Niagara/Source/Niagara/Public/NiagaraSimulationStageBase.h
+++ b/Engine/Plugins/FX/Niagara/Source/Niagara/Public/NiagaraSimulationStageBase.h
@@ -117,18 +117,18 @@ public:
 	/**
 	Controls when the simulation stage should execute, only valid for data interface iteration stages
 	*/
-	UPROPERTY(EditAnywhere, Category = "Simulation Stage", meta = (EditCondition = "IterationSource == ENiagaraIterationSource::DataInterface"))
+	UPROPERTY(EditAnywhere, Category = "Simulation Stage")
 	ENiagaraSimStageExecuteBehavior ExecuteBehavior = ENiagaraSimStageExecuteBehavior::Always;
 
 	/**
 	Disables the ability to read / write from the same particle buffer, i.e. only update position and no other attributes.
 	By default this should not be changed, but can be used to debug issues.
 	*/
-	UPROPERTY(EditAnywhere, AdvancedDisplay, Category = "Simulation Stage", meta = (EditCondition = "IterationSource == ENiagaraIterationSource::Particles"))
+	UPROPERTY(EditAnywhere, AdvancedDisplay, Category = "Simulation Stage")
 	uint32 bDisablePartialParticleUpdate : 1;
 
 	/** Source data interface to use for the simulation stage. The data interface needs to be a subclass of UNiagaraDataInterfaceRWBase, for example the Grid2D and Grid3D data interfaces. */
-	UPROPERTY(EditAnywhere, Category = "Simulation Stage", meta = (editcondition = "IterationSource == ENiagaraIterationSource::DataInterface"))
+	UPROPERTY(EditAnywhere, Category = "Simulation Stage")
 	FNiagaraVariableDataInterfaceBinding DataInterface;
 
 	/**
@@ -136,29 +136,29 @@ public:
 	For example, if the binding was set to Particles.MyStateIndex and the range was 1-3 only the particles
 	with values of 1, 2 or 3 for Particles.MyStateIndex would run the script, all others would be ignored.
 	*/
-	UPROPERTY(EditAnywhere, Category = "Simulation Stage", meta = (editcondition = "IterationSource == ENiagaraIterationSource::Particles"))
+	UPROPERTY(EditAnywhere, Category = "Simulation Stage")
 	uint32 bParticleIterationStateEnabled : 1;
 
 	/** Particle state attribute binding, when enabled we will only allow particles who pass the state range check to be processed. */
-	UPROPERTY(EditAnywhere, Category = "Simulation Stage", meta = (editcondition = "IterationSource == ENiagaraIterationSource::Particles"))
+	UPROPERTY(EditAnywhere, Category = "Simulation Stage")
 	FNiagaraVariableAttributeBinding ParticleIterationStateBinding;
 
 	/** The inclusive range used to check particle state binding against when enabled. */
-	UPROPERTY(EditAnywhere, Category = "Simulation Stage", meta = (editcondition = "IterationSource == ENiagaraIterationSource::Particles"))
+	UPROPERTY(EditAnywhere, Category = "Simulation Stage")
 	FIntPoint ParticleIterationStateRange = FIntPoint(0, 0);
 
 	/**
 	When enabled we force the dispatch to be linear (i.e. one dimension is used).
 	NOTE: Cannot be used with OverrideGpuDispatchType.
 	*/
-	UPROPERTY(EditAnywhere, Category = "Simulation Stage", meta = (editcondition = "!bOverrideGpuDispatchType"))
+	UPROPERTY(EditAnywhere, Category = "Simulation Stage")
 	uint32 bGpuDispatchForceLinear : 1;
 
 	/**
 	When enabled we use a custom number of element count for the dispatch.
 	NOTE: Cannot be used with GpuDispatchForceLinear.
 	*/
-	UPROPERTY(EditAnywhere, Category = "Simulation Stage", meta = (editcondition = "!bGpuDispatchForceLinear"))
+	UPROPERTY(EditAnywhere, Category = "Simulation Stage")
 	uint32 bOverrideGpuDispatchType : 1;
 
 	/** Dimensions to use for dispatch. */
diff --git a/Engine/Plugins/FX/Niagara/Source/NiagaraEditor/Private/Customizations/NiagaraSimulationStageCustomization.cpp b/Engine/Plugins/FX/Niagara/Source/NiagaraEditor/Private/Customizations/NiagaraSimulationStageCustomization.cpp
--- a/Engine/Plugins/FX/Niagara/Source/NiagaraEditor/Private/Customizations/NiagaraSimulationStageCustomization.cpp
+++ b/Engine/Plugins/FX/Niagara/Source/NiagaraEditor/Private/Customizations/NiagaraSimulationStageCustomization.cpp
@@ -57,61 +57,45 @@ void FNiagaraSimulationStageGenericCustomization::CustomizeDetails(class IDetail
 
 	SimStage->OnChanged().AddRaw(this, &FNiagaraSimulationStageGenericCustomization::OnPropertyChanged);
 
-	static const FName CategoryName = TEXT("Simulation Stage");
-	static const FName ParticleStageName = TEXT("Particle Parameters");
-	static const FName DataInterfaceStageName = TEXT("DataInterface Parameters");
-	IDetailCategoryBuilder& SimStageCategory = DetailBuilder.EditCategory(CategoryName);
-	IDetailCategoryBuilder& ParticleStageCategory = DetailBuilder.EditCategory(ParticleStageName);
-	IDetailCategoryBuilder& DataInterfaceStageCategory = DetailBuilder.EditCategory(DataInterfaceStageName);
-
-	// Hide all
-	DetailBuilder.HideProperty(GET_MEMBER_NAME_CHECKED(UNiagaraSimulationStageGeneric, EnabledBinding));
-	DetailBuilder.HideProperty(GET_MEMBER_NAME_CHECKED(UNiagaraSimulationStageGeneric, ElementCountXBinding));
-	DetailBuilder.HideProperty(GET_MEMBER_NAME_CHECKED(UNiagaraSimulationStageGeneric, ElementCountYBinding));
-	DetailBuilder.HideProperty(GET_MEMBER_NAME_CHECKED(UNiagaraSimulationStageGeneric, ElementCountZBinding));
-	DetailBuilder.HideProperty(GET_MEMBER_NAME_CHECKED(UNiagaraSimulationStageGeneric, IterationSource));
-	DetailBuilder.HideProperty(GET_MEMBER_NAME_CHECKED(UNiagaraSimulationStageGeneric, Iterations));
-	DetailBuilder.HideProperty(GET_MEMBER_NAME_CHECKED(UNiagaraSimulationStageGeneric, NumIterationsBinding));
-	DetailBuilder.HideProperty(GET_MEMBER_NAME_CHECKED(UNiagaraSimulationStageGeneric, ExecuteBehavior));
-	DetailBuilder.HideProperty(GET_MEMBER_NAME_CHECKED(UNiagaraSimulationStageGeneric, bDisablePartialParticleUpdate));
-	DetailBuilder.HideProperty(GET_MEMBER_NAME_CHECKED(UNiagaraSimulationStageGeneric, DataInterface));
-	DetailBuilder.HideProperty(GET_MEMBER_NAME_CHECKED(UNiagaraSimulationStageGeneric, bParticleIterationStateEnabled));
-	DetailBuilder.HideProperty(GET_MEMBER_NAME_CHECKED(UNiagaraSimulationStageGeneric, ParticleIterationStateBinding));
-	DetailBuilder.HideProperty(GET_MEMBER_NAME_CHECKED(UNiagaraSimulationStageGeneric, ParticleIterationStateRange));
-	DetailBuilder.HideProperty(GET_MEMBER_NAME_CHECKED(UNiagaraSimulationStageGeneric, bGpuDispatchForceLinear));
-	DetailBuilder.HideProperty(GET_MEMBER_NAME_CHECKED(UNiagaraSimulationStageGeneric, bOverrideGpuDispatchType));
-	DetailBuilder.HideProperty(GET_MEMBER_NAME_CHECKED(UNiagaraSimulationStageGeneric, OverrideGpuDispatchType));
-	DetailBuilder.HideProperty(GET_MEMBER_NAME_CHECKED(UNiagaraSimulationStageGeneric, bOverrideGpuDispatchNumThreads));
-	DetailBuilder.HideProperty(GET_MEMBER_NAME_CHECKED(UNiagaraSimulationStageGeneric, OverrideGpuDispatchNumThreads));
-
-	// Show properties in the order we want them to appear
-	SimStageCategory.AddProperty(DetailBuilder.GetProperty(GET_MEMBER_NAME_CHECKED(UNiagaraSimulationStageGeneric, EnabledBinding)));
-	SimStageCategory.AddProperty(DetailBuilder.GetProperty(GET_MEMBER_NAME_CHECKED(UNiagaraSimulationStageGeneric, Iterations)));
-	SimStageCategory.AddProperty(DetailBuilder.GetProperty(GET_MEMBER_NAME_CHECKED(UNiagaraSimulationStageGeneric, NumIterationsBinding)));
-	SimStageCategory.AddProperty(DetailBuilder.GetProperty(GET_MEMBER_NAME_CHECKED(UNiagaraSimulationStageGeneric, IterationSource)));
-
-	DataInterfaceStageCategory.AddProperty(DetailBuilder.GetProperty(GET_MEMBER_NAME_CHECKED(UNiagaraSimulationStageGeneric, bOverrideGpuDispatchType)));
-	if (SimStage->bOverrideGpuDispatchType)
+	static const FName NAME_SimulationStage = TEXT("Simulation Stage");
+	static const FName NAME_ParticleParameters = TEXT("Particle Parameters");
+	static const FName NAME_DataInterfaceParameters = TEXT("DataInterface Parameters");
+	static const FName NAME_DispatchParameters = TEXT("Dispatch Parameters");
+	IDetailCategoryBuilder& SimStageCategory = DetailBuilder.EditCategory(NAME_SimulationStage);
+	IDetailCategoryBuilder& ParticleStageCategory = DetailBuilder.EditCategory(NAME_ParticleParameters);
+	IDetailCategoryBuilder& DataInterfaceStageCategory = DetailBuilder.EditCategory(NAME_DataInterfaceParameters);
+	IDetailCategoryBuilder& DispatchParametersCategory = DetailBuilder.EditCategory(NAME_DispatchParameters);
+
+	// Hide all properties by default
+	TArray<TSharedRef<IPropertyHandle>> CategoryProperties;
+	SimStageCategory.GetDefaultProperties(CategoryProperties);
+	for (TSharedRef<IPropertyHandle>& PropertyHandle : CategoryProperties)
 	{
-		DataInterfaceStageCategory.AddProperty(DetailBuilder.GetProperty(GET_MEMBER_NAME_CHECKED(UNiagaraSimulationStageGeneric, OverrideGpuDispatchType)));
-		// Always true as we always dispatch across at least 1 dimension
-		//if (SimStage->OverrideGpuDispatchType >= ENiagaraGpuDispatchType::OneD)
+		const FName PropertyName = PropertyHandle->GetProperty() ? PropertyHandle->GetProperty()->GetFName() : NAME_None;
+		if (PropertyName == GET_MEMBER_NAME_CHECKED(UNiagaraSimulationStageBase, SimulationStageName))
 		{
-			DataInterfaceStageCategory.AddProperty(DetailBuilder.GetProperty(GET_MEMBER_NAME_CHECKED(UNiagaraSimulationStageGeneric, ElementCountXBinding)));
+			SimStageCategory.AddProperty(PropertyHandle);
 		}
-		if (SimStage->OverrideGpuDispatchType >= ENiagaraGpuDispatchType::TwoD)
+		else
 		{
-			DataInterfaceStageCategory.AddProperty(DetailBuilder.GetProperty(GET_MEMBER_NAME_CHECKED(UNiagaraSimulationStageGeneric, ElementCountYBinding)));
-		}
-		if (SimStage->OverrideGpuDispatchType >= ENiagaraGpuDispatchType::ThreeD)
-		{
-			DataInterfaceStageCategory.AddProperty(DetailBuilder.GetProperty(GET_MEMBER_NAME_CHECKED(UNiagaraSimulationStageGeneric, ElementCountZBinding)));
+			DetailBuilder.HideProperty(PropertyHandle);
 		}
 	}
 
-	if ( SimStage->IterationSource == ENiagaraIterationSource::Particles )
+	// Main Category
+	SimStageCategory.AddProperty(DetailBuilder.GetProperty(GET_MEMBER_NAME_CHECKED(UNiagaraSimulationStageGeneric, EnabledBinding)));
+
+	SimStageCategory.AddProperty(DetailBuilder.GetProperty(GET_MEMBER_NAME_CHECKED(UNiagaraSimulationStageGeneric, IterationSource)));
+	SimStageCategory.AddProperty(DetailBuilder.GetProperty(GET_MEMBER_NAME_CHECKED(UNiagaraSimulationStageGeneric, Iterations)));
+	SimStageCategory.AddProperty(DetailBuilder.GetProperty(GET_MEMBER_NAME_CHECKED(UNiagaraSimulationStageGeneric, NumIterationsBinding)));
+
+	SimStageCategory.AddProperty(DetailBuilder.GetProperty(GET_MEMBER_NAME_CHECKED(UNiagaraSimulationStageGeneric, ExecuteBehavior)));
+
+	// Paterial Iteration Category
+	if (SimStage->IterationSource == ENiagaraIterationSource::Particles)
 	{
 		ParticleStageCategory.AddProperty(DetailBuilder.GetProperty(GET_MEMBER_NAME_CHECKED(UNiagaraSimulationStageGeneric, bDisablePartialParticleUpdate)));
+
 		ParticleStageCategory.AddProperty(DetailBuilder.GetProperty(GET_MEMBER_NAME_CHECKED(UNiagaraSimulationStageGeneric, bParticleIterationStateEnabled)));
 		if (SimStage->bParticleIterationStateEnabled)
 		{
@@ -119,19 +103,35 @@ void FNiagaraSimulationStageGenericCustomization::CustomizeDetails(class IDetail
 			ParticleStageCategory.AddProperty(DetailBuilder.GetProperty(GET_MEMBER_NAME_CHECKED(UNiagaraSimulationStageGeneric, ParticleIterationStateRange)));
 		}
 	}
+	// DataInterface Iteration Category
 	else
 	{
 		DataInterfaceStageCategory.AddProperty(DetailBuilder.GetProperty(GET_MEMBER_NAME_CHECKED(UNiagaraSimulationStageGeneric, DataInterface)));
-		DataInterfaceStageCategory.AddProperty(DetailBuilder.GetProperty(GET_MEMBER_NAME_CHECKED(UNiagaraSimulationStageGeneric, ExecuteBehavior)));
-
 		DataInterfaceStageCategory.AddProperty(DetailBuilder.GetProperty(GET_MEMBER_NAME_CHECKED(UNiagaraSimulationStageGeneric, bGpuDispatchForceLinear)));
+	}
+
+	// Dispatch Parameters
+	DispatchParametersCategory.AddProperty(DetailBuilder.GetProperty(GET_MEMBER_NAME_CHECKED(UNiagaraSimulationStageGeneric, bOverrideGpuDispatchType)));
+	if (SimStage->bOverrideGpuDispatchType)
+	{
+		DispatchParametersCategory.AddProperty(DetailBuilder.GetProperty(GET_MEMBER_NAME_CHECKED(UNiagaraSimulationStageGeneric, OverrideGpuDispatchType)));
 
-		DataInterfaceStageCategory.AddProperty(DetailBuilder.GetProperty(GET_MEMBER_NAME_CHECKED(UNiagaraSimulationStageGeneric, bOverrideGpuDispatchNumThreads)));
-		if (SimStage->bOverrideGpuDispatchNumThreads)
+		DispatchParametersCategory.AddProperty(DetailBuilder.GetProperty(GET_MEMBER_NAME_CHECKED(UNiagaraSimulationStageGeneric, ElementCountXBinding)));
+		if (SimStage->OverrideGpuDispatchType >= ENiagaraGpuDispatchType::TwoD)
+		{
+			DispatchParametersCategory.AddProperty(DetailBuilder.GetProperty(GET_MEMBER_NAME_CHECKED(UNiagaraSimulationStageGeneric, ElementCountYBinding)));
+		}
+		if (SimStage->OverrideGpuDispatchType >= ENiagaraGpuDispatchType::ThreeD)
 		{
-			DataInterfaceStageCategory.AddProperty(DetailBuilder.GetProperty(GET_MEMBER_NAME_CHECKED(UNiagaraSimulationStageGeneric, OverrideGpuDispatchNumThreads)));
+			DispatchParametersCategory.AddProperty(DetailBuilder.GetProperty(GET_MEMBER_NAME_CHECKED(UNiagaraSimulationStageGeneric, ElementCountZBinding)));
 		}
 	}
+
+	DispatchParametersCategory.AddProperty(DetailBuilder.GetProperty(GET_MEMBER_NAME_CHECKED(UNiagaraSimulationStageGeneric, bOverrideGpuDispatchNumThreads)));
+	if (SimStage->bOverrideGpuDispatchNumThreads)
+	{
+		DispatchParametersCategory.AddProperty(DetailBuilder.GetProperty(GET_MEMBER_NAME_CHECKED(UNiagaraSimulationStageGeneric, OverrideGpuDispatchNumThreads)));
+	}
 }
 
 void FNiagaraSimulationStageGenericCustomization::OnPropertyChanged()
