- Sanitize asin / acos inputs to avoid NaNs

diff --git a/Engine/Plugins/FX/Niagara/Source/NiagaraEditor/Private/NiagaraEditorCommon.cpp b/Engine/Plugins/FX/Niagara/Source/NiagaraEditor/Private/NiagaraEditorCommon.cpp
--- a/Engine/Plugins/FX/Niagara/Source/NiagaraEditor/Private/NiagaraEditorCommon.cpp
+++ b/Engine/Plugins/FX/Niagara/Source/NiagaraEditor/Private/NiagaraEditorCommon.cpp
@@ -645,7 +645,7 @@ void FNiagaraOpInfo::Init()
 		Op->Keywords = FText::FromString(TEXT("ArcSine"));
 		Op->Inputs.Add(FNiagaraOpInOutInfo(A, NumericType, AText, AText, DefaultStr_One));
 		Op->Inputs.Add(FNiagaraOpInOutInfo(TEXT("Period"), PeriodType, PeriodFriendlyNameText, PeriodDescText, DefaultStr_One));
-		Op->Outputs.Add(FNiagaraOpInOutInfo(Result, NumericType, ResultText, ResultText, DefaultStr_One, TEXT("asin({0})*({1}/TWO_PI)")));
+		Op->Outputs.Add(FNiagaraOpInOutInfo(Result, NumericType, ResultText, ResultText, DefaultStr_One, TEXT("asin(clamp({0}, -1.0f, 1.0f))*({1}/TWO_PI)")));
 		Op->BuildName(TEXT("ArcSine"), CategoryName);
 		Op->InputTypeValidationFunction.BindLambda(NoPositionTypeCheck);
 		OpInfoMap.Add(Op->Name) = Idx;
@@ -657,7 +657,7 @@ void FNiagaraOpInfo::Init()
 		Op->Description = NSLOCTEXT("NiagaraOpInfo", "ArcSineRad Desc", "Result = asin(A)");
 		Op->Keywords = FText::FromString(TEXT("ArcSine"));
 		Op->Inputs.Add(FNiagaraOpInOutInfo(A, NumericType, AText, AText, DefaultStr_One));
-		Op->Outputs.Add(FNiagaraOpInOutInfo(Result, NumericType, ResultText, ResultText, DefaultStr_One, TEXT("asin({0})")));
+		Op->Outputs.Add(FNiagaraOpInOutInfo(Result, NumericType, ResultText, ResultText, DefaultStr_One, TEXT("asin(clamp({0}, -1.0f, 1.0f))")));
 		Op->BuildName(TEXT("ArcSine(Radians)"), CategoryName);
 		Op->InputTypeValidationFunction.BindLambda(NoPositionTypeCheck);
 		OpInfoMap.Add(Op->Name) = Idx;
@@ -669,7 +669,7 @@ void FNiagaraOpInfo::Init()
 		Op->Description = NSLOCTEXT("NiagaraOpInfo", "ArcSineDeg Desc", "Result = asin(A)*RadiansToDegrees");
 		Op->Keywords = FText::FromString(TEXT("ArcSine"));
 		Op->Inputs.Add(FNiagaraOpInOutInfo(A, NumericType, AText, AText, DefaultStr_One));
-		Op->Outputs.Add(FNiagaraOpInOutInfo(Result, NumericType, ResultText, ResultText, DefaultStr_One, TEXT("(180.0f/PI)*asin({0})")));
+		Op->Outputs.Add(FNiagaraOpInOutInfo(Result, NumericType, ResultText, ResultText, DefaultStr_One, TEXT("(180.0f/PI)*asin(clamp({0}, -1.0f, 1.0f))")));
 		Op->BuildName(TEXT("ArcSine(Degrees)"), CategoryName);
 		Op->InputTypeValidationFunction.BindLambda(NoPositionTypeCheck);
 		OpInfoMap.Add(Op->Name) = Idx;
@@ -703,7 +703,7 @@ void FNiagaraOpInfo::Init()
 		Op->Keywords = FText::FromString(TEXT("ArcCosine"));
 		Op->Inputs.Add(FNiagaraOpInOutInfo(A, NumericType, AText, AText, DefaultStr_One));
 		Op->Inputs.Add(FNiagaraOpInOutInfo(TEXT("Period"), PeriodType, PeriodFriendlyNameText, PeriodDescText, DefaultStr_One));
-		Op->Outputs.Add(FNiagaraOpInOutInfo(Result, NumericType, ResultText, ResultText, DefaultStr_One, TEXT("acos({0})*({1}/TWO_PI)")));
+		Op->Outputs.Add(FNiagaraOpInOutInfo(Result, NumericType, ResultText, ResultText, DefaultStr_One, TEXT("acos(clamp({0}, -1.0f, 1.0f))*({1}/TWO_PI)")));
 		Op->BuildName(TEXT("ArcCosine"), CategoryName);
 		Op->InputTypeValidationFunction.BindLambda(NoPositionTypeCheck);
 		OpInfoMap.Add(Op->Name) = Idx;
@@ -715,7 +715,7 @@ void FNiagaraOpInfo::Init()
 		Op->Description = NSLOCTEXT("NiagaraOpInfo", "ArcCosineRad Desc", "Result = acos(A)");
 		Op->Keywords = FText::FromString(TEXT("ArcCosine"));
 		Op->Inputs.Add(FNiagaraOpInOutInfo(A, NumericType, AText, AText, DefaultStr_One));
-		Op->Outputs.Add(FNiagaraOpInOutInfo(Result, NumericType, ResultText, ResultText, DefaultStr_One, TEXT("acos({0})")));
+		Op->Outputs.Add(FNiagaraOpInOutInfo(Result, NumericType, ResultText, ResultText, DefaultStr_One, TEXT("acos(clamp({0}, -1.0f, 1.0f))")));
 		Op->BuildName(TEXT("ArcCosine(Radians)"), CategoryName);
 		Op->InputTypeValidationFunction.BindLambda(NoPositionTypeCheck);
 		OpInfoMap.Add(Op->Name) = Idx;
@@ -727,7 +727,7 @@ void FNiagaraOpInfo::Init()
 		Op->Description = NSLOCTEXT("NiagaraOpInfo", "ArcCosineDeg Desc", "Result = acos(A)*RadiansToDegrees");
 		Op->Keywords = FText::FromString(TEXT("ArcCosine"));
 		Op->Inputs.Add(FNiagaraOpInOutInfo(A, NumericType, AText, AText, DefaultStr_One));
-		Op->Outputs.Add(FNiagaraOpInOutInfo(Result, NumericType, ResultText, ResultText, DefaultStr_One, TEXT("(180.0f/PI)*acos({0})")));
+		Op->Outputs.Add(FNiagaraOpInOutInfo(Result, NumericType, ResultText, ResultText, DefaultStr_One, TEXT("(180.0f/PI)*acos(clamp({0}, -1.0f, 1.0f))")));
 		Op->BuildName(TEXT("ArcCosine(Degrees)"), CategoryName);
 		Op->InputTypeValidationFunction.BindLambda(NoPositionTypeCheck);
 		OpInfoMap.Add(Op->Name) = Idx;
diff --git a/Engine/Source/Runtime/Core/Private/UObject/UE5MainStreamObjectVersion.cpp b/Engine/Source/Runtime/Core/Private/UObject/UE5MainStreamObjectVersion.cpp
--- a/Engine/Source/Runtime/Core/Private/UObject/UE5MainStreamObjectVersion.cpp
+++ b/Engine/Source/Runtime/Core/Private/UObject/UE5MainStreamObjectVersion.cpp
@@ -14,7 +14,7 @@ TMap<FGuid, FGuid> FUE5MainStreamObjectVersion::GetSystemGuids()
 	SystemGuids.Add(DevGuids.MATERIALSHADERMAP_DERIVEDDATA_VER, FGuid("65C01D817C9A4EEFAE9E988D41A1F3DD"));
 	SystemGuids.Add(DevGuids.NANITE_DERIVEDDATA_VER, FGuid("7A3EA93546624597BEBE2ADFDE3A7628"));
 	SystemGuids.Add(DevGuids.NIAGARASHADERMAP_DERIVEDDATA_VER, FGuid("8A37C45D24F2423CBE5F8F371DE33575"));
-	SystemGuids.Add(DevGuids.Niagara_LatestScriptCompileVersion, FGuid("E5A56F8E17E7614C9040A7F8F168F0C0"));
+	SystemGuids.Add(DevGuids.Niagara_LatestScriptCompileVersion, FGuid("FFA132206A19A641AA7A972631B17232"));
 	SystemGuids.Add(DevGuids.SkeletalMeshDerivedDataVersion, FGuid("4084BE94509A8C3FA7CF410DDF8ABF04"));
 	SystemGuids.Add(DevGuids.STATICMESH_DERIVEDDATA_VER, FGuid("98F7D79A5811013825E75BBCC41ED3E9"));
 	SystemGuids.Add(DevGuids.POSESEARCHDB_DERIVEDDATA_VER, FGuid("4E595C2AC5E947D6BA9ABC874353E5BC"));
