- PR #9797: [Niagara] Filtering option for Niagara RT (Contributed by [REDACTED_PERSON])

diff --git a/Engine/Plugins/FX/Niagara/Source/Niagara/Classes/NiagaraDataInterfaceRenderTarget2D.h b/Engine/Plugins/FX/Niagara/Source/Niagara/Classes/NiagaraDataInterfaceRenderTarget2D.h
--- a/Engine/Plugins/FX/Niagara/Source/Niagara/Classes/NiagaraDataInterfaceRenderTarget2D.h
+++ b/Engine/Plugins/FX/Niagara/Source/Niagara/Classes/NiagaraDataInterfaceRenderTarget2D.h
@@ -1,10 +1,12 @@
 // Copyright Epic Games, Inc. All Rights Reserved.
 #pragma once
 
+#include "Engine/TextureDefines.h"
+#include "Engine/TextureRenderTarget2D.h"
+#include "RenderGraphDefinitions.h"
+
 #include "NiagaraDataInterface.h"
 #include "NiagaraDataInterfaceRW.h"
-#include "ClearQuad.h"
-#include "Engine/TextureRenderTarget2D.h"
 #include "NiagaraComponent.h"
 #include "NiagaraGenerateMips.h"
 #include "NiagaraDataInterfaceRenderTarget2D.generated.h"
@@ -23,6 +25,7 @@ struct FRenderTarget2DRWInstanceData_GameThread
 
 	FIntPoint Size = FIntPoint(EForceInit::ForceInitToZero);
 	ETextureRenderTargetFormat Format = RTF_RGBA16f;
+	TextureFilter Filter = TextureFilter::TF_Default;
 	ENiagaraMipMapGeneration MipMapGeneration = ENiagaraMipMapGeneration::Disabled;
 	ENiagaraMipMapGenerationType MipMapGenerationType = ENiagaraMipMapGenerationType::Linear;
 	
@@ -146,6 +149,10 @@ public:
 	UPROPERTY(EditAnywhere, Category = "Render Target", meta = (EditCondition = "!bInheritUserParameterSettings && bOverrideFormat"))
 	TEnumAsByte<ETextureRenderTargetFormat> OverrideRenderTargetFormat;
 
+	/** When enabled overrides the filter of the render target, otherwise uses the project default setting. */
+	UPROPERTY(EditAnywhere, Category = "Render Target", meta = (EditCondition = "!bInheritUserParameterSettings"))
+	TEnumAsByte<TextureFilter> OverrideRenderTargetFilter = TextureFilter::TF_Default;
+
 	/**
 	When enabled texture parameters (size / etc) are taken from the user provided render target.
 	If no valid user parameter is set the system will be invalid.
diff --git a/Engine/Plugins/FX/Niagara/Source/Niagara/Classes/NiagaraDataInterfaceRenderTarget2DArray.h b/Engine/Plugins/FX/Niagara/Source/Niagara/Classes/NiagaraDataInterfaceRenderTarget2DArray.h
--- a/Engine/Plugins/FX/Niagara/Source/Niagara/Classes/NiagaraDataInterfaceRenderTarget2DArray.h
+++ b/Engine/Plugins/FX/Niagara/Source/Niagara/Classes/NiagaraDataInterfaceRenderTarget2DArray.h
@@ -1,9 +1,11 @@
 // Copyright Epic Games, Inc. All Rights Reserved.
 #pragma once
 
+#include "Engine/TextureDefines.h"
+#include "RenderGraphDefinitions.h"
+
 #include "NiagaraDataInterface.h"
 #include "NiagaraDataInterfaceRW.h"
-#include "ClearQuad.h"
 #include "NiagaraComponent.h"
 #include "NiagaraDataInterfaceRenderTarget2DArray.generated.h"
 
@@ -22,6 +24,7 @@ struct FRenderTarget2DArrayRWInstanceData_GameThread
 
 	FIntVector Size = FIntVector(EForceInit::ForceInitToZero);
 	EPixelFormat Format = EPixelFormat::PF_A16B16G16R16;
+	TextureFilter Filter = TextureFilter::TF_Default;
 
 	UTextureRenderTarget2DArray* TargetTexture = nullptr;
 #if WITH_EDITORONLY_DATA
@@ -131,6 +134,10 @@ public:
 	UPROPERTY(EditAnywhere, Category = "Render Target", meta = (EditCondition = "!bInheritUserParameterSettings && bOverrideFormat"))
 	TEnumAsByte<ETextureRenderTargetFormat> OverrideRenderTargetFormat;
 
+	/** When enabled overrides the filter of the render target, otherwise uses the project default setting. */
+	UPROPERTY(EditAnywhere, Category = "Render Target", meta = (EditCondition = "!bInheritUserParameterSettings"))
+	TEnumAsByte<TextureFilter> OverrideRenderTargetFilter = TextureFilter::TF_Default;
+
 	/**
 	When enabled texture parameters (size / etc) are taken from the user provided render target.
 	If no valid user parameter is set the system will be invalid.
diff --git a/Engine/Plugins/FX/Niagara/Source/Niagara/Classes/NiagaraDataInterfaceRenderTargetCube.h b/Engine/Plugins/FX/Niagara/Source/Niagara/Classes/NiagaraDataInterfaceRenderTargetCube.h
--- a/Engine/Plugins/FX/Niagara/Source/Niagara/Classes/NiagaraDataInterfaceRenderTargetCube.h
+++ b/Engine/Plugins/FX/Niagara/Source/Niagara/Classes/NiagaraDataInterfaceRenderTargetCube.h
@@ -1,8 +1,10 @@
 // Copyright Epic Games, Inc. All Rights Reserved.
 #pragma once
 
+#include "Engine/TextureDefines.h"
+#include "RenderGraphDefinitions.h"
+
 #include "NiagaraDataInterfaceRW.h"
-#include "ClearQuad.h"
 #include "NiagaraComponent.h"
 #include "NiagaraDataInterfaceRenderTargetCube.generated.h"
 
@@ -21,6 +23,7 @@ struct FRenderTargetCubeRWInstanceData_GameThread
 
 	int Size = 0;
 	EPixelFormat Format = EPixelFormat::PF_A16B16G16R16;
+	TextureFilter Filter = TextureFilter::TF_Default;
 
 	UTextureRenderTargetCube* TargetTexture = nullptr;
 #if WITH_EDITORONLY_DATA
@@ -124,6 +127,10 @@ public:
 	UPROPERTY(EditAnywhere, Category = "Render Target", meta = (EditCondition = "!bInheritUserParameterSettings && bOverrideFormat"))
 	TEnumAsByte<ETextureRenderTargetFormat> OverrideRenderTargetFormat;
 
+	/** When enabled overrides the filter of the render target, otherwise uses the project default setting. */
+	UPROPERTY(EditAnywhere, Category = "Render Target", meta = (EditCondition = "!bInheritUserParameterSettings"))
+	TEnumAsByte<TextureFilter> OverrideRenderTargetFilter = TextureFilter::TF_Default;
+
 	/**
 	When enabled texture parameters (size / etc) are taken from the user provided render target.
 	If no valid user parameter is set the system will be invalid.
diff --git a/Engine/Plugins/FX/Niagara/Source/Niagara/Classes/NiagaraDataInterfaceRenderTargetVolume.h b/Engine/Plugins/FX/Niagara/Source/Niagara/Classes/NiagaraDataInterfaceRenderTargetVolume.h
--- a/Engine/Plugins/FX/Niagara/Source/Niagara/Classes/NiagaraDataInterfaceRenderTargetVolume.h
+++ b/Engine/Plugins/FX/Niagara/Source/Niagara/Classes/NiagaraDataInterfaceRenderTargetVolume.h
@@ -1,8 +1,10 @@
 // Copyright Epic Games, Inc. All Rights Reserved.
 #pragma once
 
+#include "Engine/TextureDefines.h"
+#include "RenderGraphDefinitions.h"
+
 #include "NiagaraDataInterfaceRW.h"
-#include "ClearQuad.h"
 #include "NiagaraComponent.h"
 #include "NiagaraDataInterfaceRenderTargetVolume.generated.h"
 
@@ -20,6 +22,7 @@ struct FRenderTargetVolumeRWInstanceData_GameThread
 
 	FIntVector Size = FIntVector(EForceInit::ForceInitToZero);
 	EPixelFormat Format = EPixelFormat::PF_A16B16G16R16;
+	TextureFilter Filter = TextureFilter::TF_Default;
 
 	UTextureRenderTargetVolume* TargetTexture = nullptr;
 #if WITH_EDITORONLY_DATA
@@ -132,6 +135,10 @@ public:
 	UPROPERTY(EditAnywhere, Category = "Render Target", meta = (EditCondition = "!bInheritUserParameterSettings && bOverrideFormat"))
 	TEnumAsByte<ETextureRenderTargetFormat> OverrideRenderTargetFormat;
 
+	/** When enabled overrides the filter of the render target, otherwise uses the project default setting. */
+	UPROPERTY(EditAnywhere, Category = "Render Target", meta = (EditCondition = "!bInheritUserParameterSettings"))
+	TEnumAsByte<TextureFilter> OverrideRenderTargetFilter = TextureFilter::TF_Default;
+
 	/**
 	When enabled texture parameters (size / etc) are taken from the user provided render target.
 	If no valid user parameter is set the system will be invalid.
diff --git a/Engine/Plugins/FX/Niagara/Source/Niagara/Private/NiagaraDataInterfaceRenderTarget2D.cpp b/Engine/Plugins/FX/Niagara/Source/Niagara/Private/NiagaraDataInterfaceRenderTarget2D.cpp
--- a/Engine/Plugins/FX/Niagara/Source/Niagara/Private/NiagaraDataInterfaceRenderTarget2D.cpp
+++ b/Engine/Plugins/FX/Niagara/Source/Niagara/Private/NiagaraDataInterfaceRenderTarget2D.cpp
@@ -307,6 +307,7 @@ bool UNiagaraDataInterfaceRenderTarget2D::Equals(const UNiagaraDataInterface* Ot
 		OtherTyped->MipMapGeneration == MipMapGeneration &&
 		OtherTyped->MipMapGenerationType == MipMapGenerationType &&
 		OtherTyped->OverrideRenderTargetFormat == OverrideRenderTargetFormat &&
+		OtherTyped->OverrideRenderTargetFilter == OverrideRenderTargetFilter &&
 		OtherTyped->bInheritUserParameterSettings == bInheritUserParameterSettings &&
 		OtherTyped->bOverrideFormat == bOverrideFormat;
 }
@@ -328,6 +329,7 @@ bool UNiagaraDataInterfaceRenderTarget2D::CopyToInternal(UNiagaraDataInterface*
 	DestinationTyped->MipMapGeneration = MipMapGeneration;
 	DestinationTyped->MipMapGenerationType = MipMapGenerationType;
 	DestinationTyped->OverrideRenderTargetFormat = OverrideRenderTargetFormat;
+	DestinationTyped->OverrideRenderTargetFilter = OverrideRenderTargetFilter;
 	DestinationTyped->bInheritUserParameterSettings = bInheritUserParameterSettings;
 	DestinationTyped->bOverrideFormat = bOverrideFormat;
 #if WITH_EDITORONLY_DATA
@@ -472,6 +474,7 @@ bool UNiagaraDataInterfaceRenderTarget2D::InitPerInstanceData(void* PerInstanceD
 	InstanceData->MipMapGeneration = MipMapGeneration;
 	InstanceData->MipMapGenerationType = MipMapGenerationType;
 	InstanceData->Format = RenderTargetFormat;
+	InstanceData->Filter = OverrideRenderTargetFilter;
 	InstanceData->RTUserParamBinding.Init(SystemInstance->GetInstanceParameters(), RenderTargetUserParameter.Parameter);
 #if WITH_EDITORONLY_DATA
 	InstanceData->bPreviewTexture = bPreviewRenderTarget;
@@ -638,6 +641,7 @@ bool UNiagaraDataInterfaceRenderTarget2D::PerInstanceTick(void* PerInstanceData,
 				InstanceData->MipMapGenerationType = ENiagaraMipMapGenerationType::Unfiltered;
 			}
 			InstanceData->Format = InstanceData->TargetTexture->RenderTargetFormat;
+			InstanceData->Filter = InstanceData->TargetTexture->Filter;
 		}
 		else
 		{
@@ -670,6 +674,7 @@ bool UNiagaraDataInterfaceRenderTarget2D::PerInstanceTickPostSimulate(void* PerI
 		InstanceData->TargetTexture->bAutoGenerateMips = InstanceData->MipMapGeneration != ENiagaraMipMapGeneration::Disabled;
 		InstanceData->TargetTexture->RenderTargetFormat = InstanceData->Format;
 		InstanceData->TargetTexture->ClearColor = FLinearColor(0.0, 0, 0, 0);
+		InstanceData->TargetTexture->Filter = InstanceData->Filter;
 		InstanceData->TargetTexture->InitAutoFormat(InstanceData->Size.X, InstanceData->Size.Y);
 		InstanceData->TargetTexture->UpdateResourceImmediate(true);
 
@@ -682,6 +687,7 @@ bool UNiagaraDataInterfaceRenderTarget2D::PerInstanceTickPostSimulate(void* PerI
 		const bool bAutoGenerateMips = InstanceData->MipMapGeneration != ENiagaraMipMapGeneration::Disabled;
 		if ((InstanceData->TargetTexture->SizeX != InstanceData->Size.X) || (InstanceData->TargetTexture->SizeY != InstanceData->Size.Y) ||
 			(InstanceData->TargetTexture->RenderTargetFormat != InstanceData->Format) ||
+			(InstanceData->TargetTexture->Filter != InstanceData->Filter) ||
 			!InstanceData->TargetTexture->bCanCreateUAV ||
 			(InstanceData->TargetTexture->bAutoGenerateMips != bAutoGenerateMips) ||
 			!InstanceData->TargetTexture->GetResource())
@@ -690,6 +696,7 @@ bool UNiagaraDataInterfaceRenderTarget2D::PerInstanceTickPostSimulate(void* PerI
 			InstanceData->TargetTexture->bCanCreateUAV = true;
 			InstanceData->TargetTexture->bAutoGenerateMips = bAutoGenerateMips;
 			InstanceData->TargetTexture->RenderTargetFormat = InstanceData->Format;
+			InstanceData->TargetTexture->Filter = InstanceData->Filter;
 			InstanceData->TargetTexture->InitAutoFormat(InstanceData->Size.X, InstanceData->Size.Y);
 			InstanceData->TargetTexture->UpdateResourceImmediate(true);
 		}
diff --git a/Engine/Plugins/FX/Niagara/Source/Niagara/Private/NiagaraDataInterfaceRenderTarget2DArray.cpp b/Engine/Plugins/FX/Niagara/Source/Niagara/Private/NiagaraDataInterfaceRenderTarget2DArray.cpp
--- a/Engine/Plugins/FX/Niagara/Source/Niagara/Private/NiagaraDataInterfaceRenderTarget2DArray.cpp
+++ b/Engine/Plugins/FX/Niagara/Source/Niagara/Private/NiagaraDataInterfaceRenderTarget2DArray.cpp
@@ -307,6 +307,7 @@ bool UNiagaraDataInterfaceRenderTarget2DArray::Equals(const UNiagaraDataInterfac
 		OtherTyped->RenderTargetUserParameter == RenderTargetUserParameter &&
 		OtherTyped->Size == Size &&
 		OtherTyped->OverrideRenderTargetFormat == OverrideRenderTargetFormat &&
+		OtherTyped->OverrideRenderTargetFilter == OverrideRenderTargetFilter &&
 		OtherTyped->bInheritUserParameterSettings == bInheritUserParameterSettings &&
 		OtherTyped->bOverrideFormat == bOverrideFormat;
 }
@@ -326,6 +327,7 @@ bool UNiagaraDataInterfaceRenderTarget2DArray::CopyToInternal(UNiagaraDataInterf
 
 	DestinationTyped->Size = Size;
 	DestinationTyped->OverrideRenderTargetFormat = OverrideRenderTargetFormat;
+	DestinationTyped->OverrideRenderTargetFilter = OverrideRenderTargetFilter;
 	DestinationTyped->bInheritUserParameterSettings = bInheritUserParameterSettings;
 	DestinationTyped->bOverrideFormat = bOverrideFormat;
 #if WITH_EDITORONLY_DATA
@@ -463,6 +465,7 @@ bool UNiagaraDataInterfaceRenderTarget2DArray::InitPerInstanceData(void* PerInst
 	InstanceData->Size.Y = FMath::Clamp<int>(int(float(Size.Y) * NiagaraDataInterfaceRenderTargetCommon::GResolutionMultiplier), 1, GMaxTextureDimensions);
 	InstanceData->Size.Z = FMath::Clamp<int>(Size.Z, 1, GMaxTextureDimensions);
 	InstanceData->Format = GetPixelFormatFromRenderTargetFormat(RenderTargetFormat);
+	InstanceData->Filter = OverrideRenderTargetFilter;
 	InstanceData->RTUserParamBinding.Init(SystemInstance->GetInstanceParameters(), RenderTargetUserParameter.Parameter);
 #if WITH_EDITORONLY_DATA
 	InstanceData->bPreviewTexture = bPreviewRenderTarget;
@@ -595,6 +598,7 @@ bool UNiagaraDataInterfaceRenderTarget2DArray::PerInstanceTick(void* PerInstance
 			//	InstanceData->MipMapGeneration = ENiagaraMipMapGeneration::Disabled;
 			//}
 			InstanceData->Format = InstanceData->TargetTexture->OverrideFormat;
+			InstanceData->Filter = InstanceData->TargetTexture->Filter;
 		}
 		else
 		{
@@ -625,6 +629,7 @@ bool UNiagaraDataInterfaceRenderTarget2DArray::PerInstanceTickPostSimulate(void*
 		InstanceData->TargetTexture->bCanCreateUAV = true;
 		//InstanceData->TargetTexture->bAutoGenerateMips = InstanceData->MipMapGeneration != ENiagaraMipMapGeneration::Disabled;
 		InstanceData->TargetTexture->ClearColor = FLinearColor(0.0, 0, 0, 0);
+		InstanceData->TargetTexture->Filter = InstanceData->Filter;
 		InstanceData->TargetTexture->Init(InstanceData->Size.X, InstanceData->Size.Y, InstanceData->Size.Z, InstanceData->Format);
 		InstanceData->TargetTexture->UpdateResourceImmediate(true);
 
@@ -637,6 +642,7 @@ bool UNiagaraDataInterfaceRenderTarget2DArray::PerInstanceTickPostSimulate(void*
 		//const bool bAutoGenerateMips = InstanceData->MipMapGeneration != ENiagaraMipMapGeneration::Disabled;
 		if ((InstanceData->TargetTexture->SizeX != InstanceData->Size.X) || (InstanceData->TargetTexture->SizeY != InstanceData->Size.Y) || (InstanceData->TargetTexture->Slices != InstanceData->Size.Z) ||
 			(InstanceData->TargetTexture->OverrideFormat != InstanceData->Format) ||
+			(InstanceData->TargetTexture->Filter != InstanceData->Filter) ||
 			!InstanceData->TargetTexture->bCanCreateUAV ||
 			//(InstanceData->TargetTexture->bAutoGenerateMips != bAutoGenerateMips) ||
 			!InstanceData->TargetTexture->GetResource())
@@ -644,6 +650,7 @@ bool UNiagaraDataInterfaceRenderTarget2DArray::PerInstanceTickPostSimulate(void*
 			// resize RT to match what we need for the output
 			InstanceData->TargetTexture->bCanCreateUAV = true;
 			//InstanceData->TargetTexture->bAutoGenerateMips = bAutoGenerateMips;
+			InstanceData->TargetTexture->Filter = InstanceData->Filter;
 			InstanceData->TargetTexture->Init(InstanceData->Size.X, InstanceData->Size.Y, InstanceData->Size.Z, InstanceData->Format);
 			InstanceData->TargetTexture->UpdateResourceImmediate(true);
 		}
diff --git a/Engine/Plugins/FX/Niagara/Source/Niagara/Private/NiagaraDataInterfaceRenderTargetCube.cpp b/Engine/Plugins/FX/Niagara/Source/Niagara/Private/NiagaraDataInterfaceRenderTargetCube.cpp
--- a/Engine/Plugins/FX/Niagara/Source/Niagara/Private/NiagaraDataInterfaceRenderTargetCube.cpp
+++ b/Engine/Plugins/FX/Niagara/Source/Niagara/Private/NiagaraDataInterfaceRenderTargetCube.cpp
@@ -301,6 +301,7 @@ bool UNiagaraDataInterfaceRenderTargetCube::Equals(const UNiagaraDataInterface*
 		OtherTyped->RenderTargetUserParameter == RenderTargetUserParameter &&
 		OtherTyped->Size == Size &&
 		OtherTyped->OverrideRenderTargetFormat == OverrideRenderTargetFormat &&
+		OtherTyped->OverrideRenderTargetFilter == OverrideRenderTargetFilter &&
 		OtherTyped->bInheritUserParameterSettings == bInheritUserParameterSettings &&
 		OtherTyped->bOverrideFormat == bOverrideFormat;
 }
@@ -320,6 +321,7 @@ bool UNiagaraDataInterfaceRenderTargetCube::CopyToInternal(UNiagaraDataInterface
 
 	DestinationTyped->Size = Size;
 	DestinationTyped->OverrideRenderTargetFormat = OverrideRenderTargetFormat;
+	DestinationTyped->OverrideRenderTargetFilter = OverrideRenderTargetFilter;
 	DestinationTyped->bInheritUserParameterSettings = bInheritUserParameterSettings;
 	DestinationTyped->bOverrideFormat = bOverrideFormat;
 #if WITH_EDITORONLY_DATA
@@ -449,6 +451,7 @@ bool UNiagaraDataInterfaceRenderTargetCube::InitPerInstanceData(void* PerInstanc
 
 	InstanceData->Size = FMath::Clamp<int>(int(float(Size) * NiagaraDataInterfaceRenderTargetCommon::GResolutionMultiplier), 1, GMaxCubeTextureDimensions);
 	InstanceData->Format = GetPixelFormatFromRenderTargetFormat(RenderTargetFormat);
+	InstanceData->Filter = OverrideRenderTargetFilter;
 	InstanceData->RTUserParamBinding.Init(SystemInstance->GetInstanceParameters(), RenderTargetUserParameter.Parameter);
 #if WITH_EDITORONLY_DATA
 	InstanceData->bPreviewTexture = bPreviewRenderTarget;
@@ -569,6 +572,7 @@ bool UNiagaraDataInterfaceRenderTargetCube::PerInstanceTick(void* PerInstanceDat
 			//	InstanceData->MipMapGeneration = ENiagaraMipMapGeneration::Disabled;
 			//}
 			InstanceData->Format = InstanceData->TargetTexture->OverrideFormat;
+			InstanceData->Filter = InstanceData->TargetTexture->Filter;
 		}
 		else
 		{
@@ -600,6 +604,7 @@ bool UNiagaraDataInterfaceRenderTargetCube::PerInstanceTickPostSimulate(void* Pe
 		InstanceData->TargetTexture->bCanCreateUAV = true;
 		//InstanceData->TargetTexture->bAutoGenerateMips = InstanceData->MipMapGeneration != ENiagaraMipMapGeneration::Disabled;
 		InstanceData->TargetTexture->ClearColor = FLinearColor(0.0, 0, 0, 0);
+		InstanceData->TargetTexture->Filter = InstanceData->Filter;
 		InstanceData->TargetTexture->Init(InstanceData->Size, InstanceData->Format);
 		InstanceData->TargetTexture->UpdateResourceImmediate(true);
 
@@ -612,6 +617,7 @@ bool UNiagaraDataInterfaceRenderTargetCube::PerInstanceTickPostSimulate(void* Pe
 		//const bool bAutoGenerateMips = InstanceData->MipMapGeneration != ENiagaraMipMapGeneration::Disabled;
 		if ((InstanceData->TargetTexture->SizeX != InstanceData->Size) || 
 			(InstanceData->TargetTexture->OverrideFormat != InstanceData->Format) ||
+			(InstanceData->TargetTexture->Filter != InstanceData->Filter) ||
 			!InstanceData->TargetTexture->bCanCreateUAV ||
 			//(InstanceData->TargetTexture->bAutoGenerateMips != bAutoGenerateMips) ||
 			!InstanceData->TargetTexture->GetResource())
@@ -619,6 +625,7 @@ bool UNiagaraDataInterfaceRenderTargetCube::PerInstanceTickPostSimulate(void* Pe
 			// resize RT to match what we need for the output
 			InstanceData->TargetTexture->bCanCreateUAV = true;
 			//InstanceData->TargetTexture->bAutoGenerateMips = bAutoGenerateMips;
+			InstanceData->TargetTexture->Filter = InstanceData->Filter;
 			InstanceData->TargetTexture->Init(InstanceData->Size, InstanceData->Format);
 			InstanceData->TargetTexture->UpdateResourceImmediate(true);
 		}
diff --git a/Engine/Plugins/FX/Niagara/Source/Niagara/Private/NiagaraDataInterfaceRenderTargetVolume.cpp b/Engine/Plugins/FX/Niagara/Source/Niagara/Private/NiagaraDataInterfaceRenderTargetVolume.cpp
--- a/Engine/Plugins/FX/Niagara/Source/Niagara/Private/NiagaraDataInterfaceRenderTargetVolume.cpp
+++ b/Engine/Plugins/FX/Niagara/Source/Niagara/Private/NiagaraDataInterfaceRenderTargetVolume.cpp
@@ -323,6 +323,7 @@ bool UNiagaraDataInterfaceRenderTargetVolume::Equals(const UNiagaraDataInterface
 		OtherTyped->RenderTargetUserParameter == RenderTargetUserParameter &&
 		OtherTyped->Size == Size &&
 		OtherTyped->OverrideRenderTargetFormat == OverrideRenderTargetFormat &&
+		OtherTyped->OverrideRenderTargetFilter == OverrideRenderTargetFilter &&
 		OtherTyped->bInheritUserParameterSettings == bInheritUserParameterSettings &&
 		OtherTyped->bOverrideFormat == bOverrideFormat;
 }
@@ -342,6 +343,7 @@ bool UNiagaraDataInterfaceRenderTargetVolume::CopyToInternal(UNiagaraDataInterfa
 
 	DestinationTyped->Size = Size;
 	DestinationTyped->OverrideRenderTargetFormat = OverrideRenderTargetFormat;
+	DestinationTyped->OverrideRenderTargetFilter = OverrideRenderTargetFilter;
 	DestinationTyped->bInheritUserParameterSettings = bInheritUserParameterSettings;
 	DestinationTyped->bOverrideFormat = bOverrideFormat;
 #if WITH_EDITORONLY_DATA
@@ -477,6 +479,7 @@ bool UNiagaraDataInterfaceRenderTargetVolume::InitPerInstanceData(void* PerInsta
 	InstanceData->Size.Y = FMath::Clamp<int>(int(float(Size.Y) * NiagaraDataInterfaceRenderTargetCommon::GResolutionMultiplier), 1, GMaxVolumeTextureDimensions);
 	InstanceData->Size.Z = FMath::Clamp<int>(int(float(Size.Z) * NiagaraDataInterfaceRenderTargetCommon::GResolutionMultiplier), 1, GMaxVolumeTextureDimensions);
 	InstanceData->Format = GetPixelFormatFromRenderTargetFormat(RenderTargetFormat);
+	InstanceData->Filter = OverrideRenderTargetFilter;
 	InstanceData->RTUserParamBinding.Init(SystemInstance->GetInstanceParameters(), RenderTargetUserParameter.Parameter);
 #if WITH_EDITORONLY_DATA
 	InstanceData->bPreviewTexture = bPreviewRenderTarget;
@@ -644,6 +647,7 @@ bool UNiagaraDataInterfaceRenderTargetVolume::SimCacheReadFrame(UObject* Storage
 	FRenderTargetVolumeRWInstanceData_GameThread* InstanceData_GT = reinterpret_cast<FRenderTargetVolumeRWInstanceData_GameThread*>(OptionalPerInstanceData);
 	InstanceData_GT->Size = CacheFrame->Size;
 	InstanceData_GT->Format = CacheFrame->Format;
+	InstanceData_GT->Filter = TextureFilter::TF_Default;
 
 	PerInstanceTick(InstanceData_GT, SystemInstance, 0.0f);
 	PerInstanceTickPostSimulate(InstanceData_GT, SystemInstance, 0.0f);
@@ -781,6 +785,7 @@ bool UNiagaraDataInterfaceRenderTargetVolume::PerInstanceTick(void* PerInstanceD
 			//	InstanceData->MipMapGeneration = ENiagaraMipMapGeneration::Disabled;
 			//}
 			InstanceData->Format = InstanceData->TargetTexture->OverrideFormat;
+			InstanceData->Filter = InstanceData->TargetTexture->Filter;
 		}
 		else
 		{
@@ -812,6 +817,7 @@ bool UNiagaraDataInterfaceRenderTargetVolume::PerInstanceTickPostSimulate(void*
 		InstanceData->TargetTexture->bCanCreateUAV = true;
 		//InstanceData->TargetTexture->bAutoGenerateMips = InstanceData->MipMapGeneration != ENiagaraMipMapGeneration::Disabled;
 		InstanceData->TargetTexture->ClearColor = FLinearColor(0.0, 0, 0, 0);
+		InstanceData->TargetTexture->Filter = InstanceData->Filter;
 		InstanceData->TargetTexture->Init(InstanceData->Size.X, InstanceData->Size.Y, InstanceData->Size.Z, InstanceData->Format);
 		InstanceData->TargetTexture->UpdateResourceImmediate(true);
 
@@ -824,6 +830,7 @@ bool UNiagaraDataInterfaceRenderTargetVolume::PerInstanceTickPostSimulate(void*
 		//const bool bAutoGenerateMips = InstanceData->MipMapGeneration != ENiagaraMipMapGeneration::Disabled;
 		if ((InstanceData->TargetTexture->SizeX != InstanceData->Size.X) || (InstanceData->TargetTexture->SizeY != InstanceData->Size.Y) || (InstanceData->TargetTexture->SizeZ != InstanceData->Size.Z) ||
 			(InstanceData->TargetTexture->OverrideFormat != InstanceData->Format) ||
+			(InstanceData->TargetTexture->Filter != InstanceData->Filter) ||
 			!InstanceData->TargetTexture->bCanCreateUAV ||
 			//(InstanceData->TargetTexture->bAutoGenerateMips != bAutoGenerateMips) ||
 			!InstanceData->TargetTexture->GetResource())
@@ -831,6 +838,7 @@ bool UNiagaraDataInterfaceRenderTargetVolume::PerInstanceTickPostSimulate(void*
 			// resize RT to match what we need for the output
 			InstanceData->TargetTexture->bCanCreateUAV = true;
 			//InstanceData->TargetTexture->bAutoGenerateMips = bAutoGenerateMips;
+			InstanceData->TargetTexture->Filter = InstanceData->Filter;
 			InstanceData->TargetTexture->Init(InstanceData->Size.X, InstanceData->Size.Y, InstanceData->Size.Z, InstanceData->Format);
 			InstanceData->TargetTexture->UpdateResourceImmediate(true);
 		}
