- Fix for C++20 compilation

#rnx

diff --git a/Engine/Plugins/FX/Niagara/Source/Niagara/Private/NiagaraDebugHud.cpp b/Engine/Plugins/FX/Niagara/Source/Niagara/Private/NiagaraDebugHud.cpp
--- a/Engine/Plugins/FX/Niagara/Source/Niagara/Private/NiagaraDebugHud.cpp
+++ b/Engine/Plugins/FX/Niagara/Source/Niagara/Private/NiagaraDebugHud.cpp
@@ -424,7 +424,7 @@ namespace NiagaraDebugLocal
 	#if WITH_EDITORONLY_DATA
 		// Remove any static variables we found that are used inside renderer bindings to avoid doubling up
 		CachedVariables->StaticVariables.RemoveAll(
-			[&](const FNiagaraVariable& StaticVariable)
+			[&](const FNiagaraVariableBase& StaticVariable)
 			{
 				return CachedVariables->ParameterStoreVariables.ContainsByPredicate(
 						[&](const FParameterStoreVariable& Existing)
