Display version number on inherited modules

diff --git a/Engine/Plugins/FX/Niagara/Source/NiagaraEditorWidgets/Private/Stack/SNiagaraStackModuleItem.cpp b/Engine/Plugins/FX/Niagara/Source/NiagaraEditorWidgets/Private/Stack/SNiagaraStackModuleItem.cpp
--- a/Engine/Plugins/FX/Niagara/Source/NiagaraEditorWidgets/Private/Stack/SNiagaraStackModuleItem.cpp
+++ b/Engine/Plugins/FX/Niagara/Source/NiagaraEditorWidgets/Private/Stack/SNiagaraStackModuleItem.cpp
@@ -394,11 +394,26 @@ bool SNiagaraStackModuleItem::GetVersionSelectionMenuEnabled() const
 
 FText SNiagaraStackModuleItem::GetVersionSelectionMenuTooltip() const
 {
+	FText BaseText;
 	if (ModuleItem->CanMoveAndDelete())
 	{
-		return LOCTEXT("VersionTooltip", "Change the version of this module script");
+		BaseText = LOCTEXT("VersionTooltip", "Change the version of this module script.");
 	}
-	return LOCTEXT("VersionTooltipDisabled", "The version of this module script can only be changed in the parent emitter.");
+	else
+	{
+		BaseText = LOCTEXT("VersionTooltipDisabled", "The version of this module script can only be changed in the parent emitter.");
+	}
+	
+	if (UNiagaraScript* Script = ModuleItem->GetModuleNode().FunctionScript; Script && Script->IsVersioningEnabled())
+	{
+		if (FVersionedNiagaraScriptData* ScriptData = Script->GetScriptData(ModuleItem->GetModuleNode().SelectedScriptVersion))
+		{
+			// add the version information if possible
+			FText VersionText = FText::Format(FText::FromString("{0}.{1}"), ScriptData->Version.MajorVersion, ScriptData->Version.MinorVersion);
+			return FText::Format(LOCTEXT("VersionTooltipFormat", "{0}\nCurrently used module version: {1}"), BaseText, VersionText);
+		}
+	}
+	return BaseText;
 }
 
 FReply SNiagaraStackModuleItem::ScratchButtonPressed() const
