- Change how we deal with Niagara deny lists so they can be more easily updated and applied via device previews

#rnx

diff --git a/Engine/Plugins/FX/Niagara/Source/Niagara/Private/NiagaraComponent.cpp b/Engine/Plugins/FX/Niagara/Source/Niagara/Private/NiagaraComponent.cpp
--- a/Engine/Plugins/FX/Niagara/Source/Niagara/Private/NiagaraComponent.cpp
+++ b/Engine/Plugins/FX/Niagara/Source/Niagara/Private/NiagaraComponent.cpp
@@ -1069,8 +1069,8 @@ void UNiagaraComponent::ActivateInternal(bool bReset /* = false */, bool bIsScal
 		return;
 	}
 
-	// Should we force activation to fail?
-	if (UNiagaraComponentSettings::ShouldSuppressActivation(Asset))
+	// Is the system allowed to run or not?
+	if (FNiagaraComponentSettings::IsSystemAllowedToRun(Asset) == false)
 	{
 		return;
 	}
diff --git a/Engine/Plugins/FX/Niagara/Source/Niagara/Private/NiagaraComponentSettings.cpp b/Engine/Plugins/FX/Niagara/Source/Niagara/Private/NiagaraComponentSettings.cpp
--- a/Engine/Plugins/FX/Niagara/Source/Niagara/Private/NiagaraComponentSettings.cpp
+++ b/Engine/Plugins/FX/Niagara/Source/Niagara/Private/NiagaraComponentSettings.cpp
@@ -1,51 +1,198 @@
 // Copyright Epic Games, Inc. All Rights Reserved.
 
 #include "NiagaraComponentSettings.h"
+#include "NiagaraEmitter.h"
+#include "NiagaraEmitterInstance.h"
+#include "NiagaraScript.h"
+#include "NiagaraSystem.h"
+#include "NiagaraSystemInstance.h"
 
-#include UE_INLINE_GENERATED_CPP_BY_NAME(NiagaraComponentSettings)
-
-int32 UNiagaraComponentSettings::bAllowSuppressActivation = 0;
-int32 UNiagaraComponentSettings::bAllowForceAutoPooling = 0;
-int32 UNiagaraComponentSettings::bUseSuppressEmitterList = 0;
-int32 UNiagaraComponentSettings::bUseGpuEmitterAllowList = 0;
-int32 UNiagaraComponentSettings::bUseGpuDataInterfaceDenyList = 0;
-
-static FAutoConsoleVariableRef CVarNiagaraUseSupressActivateList(
-	TEXT("fx.Niagara.UseSupressActivateList"),
-	UNiagaraComponentSettings::bAllowSuppressActivation,
-	TEXT("When a component is activated we will check the surpession list."),
-	ECVF_Default
-);
-
-static FAutoConsoleVariableRef CVarNiagaraForceAutoPooling(
-	TEXT("fx.Niagara.ForceAutoPooling"),
-	UNiagaraComponentSettings::bAllowForceAutoPooling,
-	TEXT("Forces auto pooling to be enabled on spawned components."),
-	ECVF_Default
-);
-
-static FAutoConsoleVariableRef CVarNiagaraUseEmitterSupressList(
-	TEXT("fx.Niagara.UseEmitterSuppressList"),
-	UNiagaraComponentSettings::bUseSuppressEmitterList,
-	TEXT("When an emitter is activated we will check the surpession list."),
-	ECVF_Default
-);
-
-static FAutoConsoleVariableRef CVarNiagaraUseGpuEmitterAllowList(
-	TEXT("fx.Niagara.UseGpuEmitterAllowList"),
-	UNiagaraComponentSettings::bUseGpuEmitterAllowList,
-	TEXT("When enabled only GPU emitters on the allow list are allowed to run."),
-	ECVF_Default
-);
-
-static FAutoConsoleVariableRef CVarNiagaraUseGpuDataInterfaceDenyList(
-	TEXT("fx.Niagara.UseGpuDataInterfaceDenyList"),
-	UNiagaraComponentSettings::bUseGpuDataInterfaceDenyList,
-	TEXT("When enabled GPU emitters will be disabled if they use a data interface on the deny list."),
-	ECVF_Default
-);
-
-UNiagaraComponentSettings::UNiagaraComponentSettings(FObjectInitializer const& ObjectInitializer)
+namespace FNiagaraComponentSettings
 {
-}
+	struct FNiagaraEmitterNameSettingsRef
+	{
+		FName SystemName;
+		FName EmitterName;
+
+		FNiagaraEmitterNameSettingsRef() = default;
+		explicit FNiagaraEmitterNameSettingsRef(const FStringView& Pair)
+		{
+			int32 SepIndex;
+			Pair.FindChar(':', SepIndex);
+			if (SepIndex != INDEX_NONE && Pair.Len() > SepIndex + 1)
+			{
+				SystemName = FName(Pair.Mid(0, SepIndex));
+				EmitterName = FName(Pair.Mid(SepIndex + 1));
+			}
+		}
+
+		bool operator==(const FNiagaraEmitterNameSettingsRef& Other) const { return SystemName == Other.SystemName && EmitterName == Other.EmitterName; }
+		bool operator!=(const FNiagaraEmitterNameSettingsRef& Other) const { return !(*this == Other); }
+		bool IsValid() const { return !SystemName.IsNone() && !EmitterName.IsNone(); }
+	};
+	
+	uint32 GetTypeHash(const FNiagaraEmitterNameSettingsRef& Var)
+	{
+		return HashCombine(GetTypeHash(Var.SystemName), GetTypeHash(Var.EmitterName));
+	}
+
+	bool									bUseSystemDenyList = false;
+	bool									bUseEmitterDenyList = false;
+
+	TSet<FName>								SystemDenyList;
+	TSet<FNiagaraEmitterNameSettingsRef>	EmitterDenyList;
+	TSet<FNiagaraEmitterNameSettingsRef>	GpuEmitterDenyList;
+	TSet<FName>								GpuDataInterfaceDenyList;
+
+	FString									SystemDenyListString;
+	FString									EmitterDenyListString;
+	FString									GpuEmitterDenyListString;
+	FString									GpuDataInterfaceDenyListString;
+
+	static bool ParseIntoSet(const FString& StringList, TSet<FName>& OutSet)
+	{
+		TArray<FString> Names;
+		StringList.ParseIntoArray(Names, TEXT(","));
+
+		TSet<FName> ExistingSet;
+		Swap(ExistingSet, OutSet);
+		for (const FString& Name : Names)
+		{
+			OutSet.Emplace(Name);
+		}
+		return ExistingSet.Difference(OutSet).Num() > 0;
+	}
+
+	static bool ParseIntoSet(const FString& StringList, TSet<FNiagaraEmitterNameSettingsRef>& OutSet)
+	{
+		TArray<FString> Pairs;
+		StringList.ParseIntoArray(Pairs, TEXT(","));
+
+		TSet<FNiagaraEmitterNameSettingsRef> ExistingSet;
+		Swap(ExistingSet, OutSet);
+		for (const FString& Pair : Pairs)
+		{
+			FNiagaraEmitterNameSettingsRef EmitterRef(Pair);
+			if (EmitterRef.IsValid())
+			{
+				OutSet.Emplace(EmitterRef);
+			}
+		}
+		return ExistingSet.Difference(OutSet).Num() > 0;
+	}
 
+	static void UpdateSystemDenyList(IConsoleVariable*)
+	{
+		const bool bWasChanged = ParseIntoSet(SystemDenyListString, SystemDenyList);
+		bUseSystemDenyList = SystemDenyList.Num() > 0;
+		if (bWasChanged)
+		{
+			for (TObjectIterator<UNiagaraSystem> It; It; ++It)
+			{
+				if (UNiagaraSystem* System = *It)
+				{
+					System->UpdateScalability();
+				}
+			}
+		}
+	}
+
+	static void UpdateEmitterDenyList(IConsoleVariable*)
+	{
+		bool bWasChanged = ParseIntoSet(EmitterDenyListString, EmitterDenyList);
+		bWasChanged |= ParseIntoSet(GpuEmitterDenyListString, GpuEmitterDenyList);
+		bWasChanged |= ParseIntoSet(GpuDataInterfaceDenyListString, GpuDataInterfaceDenyList);
+		bUseEmitterDenyList = EmitterDenyList.Num() > 0 || GpuEmitterDenyList.Num() > 0 || GpuDataInterfaceDenyList.Num() > 0;
+		if (bWasChanged)
+		{
+			for (TObjectIterator<UNiagaraSystem> It; It; ++It)
+			{
+				if (UNiagaraSystem* System = *It)
+				{
+					System->UpdateScalability();
+				}
+			}
+		}
+	}
+
+	static FAutoConsoleVariableRef CVarNiagaraSetSystemDenyList(
+		TEXT("fx.Niagara.SetSystemDenyList"),
+		SystemDenyListString,
+		TEXT("Set the system deny List to use. (i.e. NS_SystemA,NS_SystemB)"),
+		FConsoleVariableDelegate::CreateStatic(UpdateSystemDenyList), 
+		ECVF_Scalability | ECVF_Default
+	);
+	static FAutoConsoleVariableRef CVarNiagaraSetEmitterDenyList(
+		TEXT("fx.Niagara.SetEmitterDenyList"),
+		EmitterDenyListString,
+		TEXT("Set the emitter deny list to use. (i.e. NS_SystemA:EmitterA,NS_SystemB:EmitterA)"),
+		FConsoleVariableDelegate::CreateStatic(UpdateEmitterDenyList),
+		ECVF_Scalability | ECVF_Default
+	);
+	static FAutoConsoleVariableRef CVarNiagaraSetGpuEmitterDenyList(
+		TEXT("fx.Niagara.SetGpuEmitterDenyList"),
+		GpuEmitterDenyListString,
+		TEXT("Set the Gpu emitter deny list to use. (i.e. NS_SystemA:EmitterA,NS_SystemB:EmitterA)"),
+		FConsoleVariableDelegate::CreateStatic(UpdateEmitterDenyList),
+		ECVF_Scalability | ECVF_Default
+	);
+	static FAutoConsoleVariableRef CVarNiagaraSetGpuDataInterfaceDenyList(
+		TEXT("fx.Niagara.SetGpuDataInterfaceDenyList"),
+		GpuDataInterfaceDenyListString,
+		TEXT("Set the Gpu data interface deny list to use.  (i.e. UMyDataInteraceA,UMyDataInteraceB)"),
+		FConsoleVariableDelegate::CreateStatic(UpdateEmitterDenyList),
+		ECVF_Scalability | ECVF_Default
+	);
+
+	bool IsSystemAllowedToRun(const UNiagaraSystem* System)
+	{
+		check(System != nullptr);
+		return bUseSystemDenyList ? !SystemDenyList.Contains(System->GetFName()) : true;
+	}
+
+	bool IsEmitterAllowedToRun(const FNiagaraEmitterInstance* EmitterInstance)
+	{
+		if (bUseEmitterDenyList)
+		{
+			const UNiagaraSystem* NiagaraSystem = EmitterInstance->GetParentSystemInstance()->GetSystem();
+			const FVersionedNiagaraEmitter CachedEmitter = EmitterInstance->GetCachedEmitter();
+			FVersionedNiagaraEmitterData* EmitterData = CachedEmitter.GetEmitterData();
+			if (EmitterData == nullptr)
+			{
+				return false;
+			}
+
+			FNiagaraEmitterNameSettingsRef EmitterRef;
+			EmitterRef.SystemName = NiagaraSystem ? NiagaraSystem->GetFName() : NAME_None;
+			EmitterRef.EmitterName = FName(*CachedEmitter.Emitter->GetUniqueEmitterName());
+
+			if (EmitterDenyList.Contains(EmitterRef))
+			{
+				return false;
+			}
+
+			if (EmitterData->SimTarget == ENiagaraSimTarget::GPUComputeSim)
+			{
+				if (GpuEmitterDenyList.Contains(EmitterRef))
+				{
+					return false;
+				}
+
+				if (GpuDataInterfaceDenyList.Num() > 0)
+				{
+					if (const UNiagaraScript* GPUComputeScript = EmitterData->GetGPUComputeScript())
+					{
+						for (const FNiagaraScriptDataInterfaceInfo& DefaultDIInfo : GPUComputeScript->GetCachedDefaultDataInterfaces())
+						{
+							if (GpuDataInterfaceDenyList.Contains(DefaultDIInfo.Type.GetFName()))
+							{
+								return false;
+							}
+						}
+					}
+				}
+			}
+		}
+		return true;
+	}
+}
diff --git a/Engine/Plugins/FX/Niagara/Source/Niagara/Private/NiagaraEmitterInstance.cpp b/Engine/Plugins/FX/Niagara/Source/Niagara/Private/NiagaraEmitterInstance.cpp
--- a/Engine/Plugins/FX/Niagara/Source/Niagara/Private/NiagaraEmitterInstance.cpp
+++ b/Engine/Plugins/FX/Niagara/Source/Niagara/Private/NiagaraEmitterInstance.cpp
@@ -252,12 +252,8 @@ bool FNiagaraEmitterInstance::IsAllowedToExecute() const
 		}
 	}
 
-	if (UNiagaraComponentSettings::ShouldSuppressEmitterActivation(this))
-	{
-		return false;
-	}
-
-	return true;
+	// Do we allow the emitter to run or not?
+	return FNiagaraComponentSettings::IsEmitterAllowedToRun(this);
 }
 
 void FNiagaraEmitterInstance::Init(int32 InEmitterIdx, FNiagaraSystemInstanceID InSystemInstanceID)
diff --git a/Engine/Plugins/FX/Niagara/Source/Niagara/Private/NiagaraFunctionLibrary.cpp b/Engine/Plugins/FX/Niagara/Source/Niagara/Private/NiagaraFunctionLibrary.cpp
--- a/Engine/Plugins/FX/Niagara/Source/Niagara/Private/NiagaraFunctionLibrary.cpp
+++ b/Engine/Plugins/FX/Niagara/Source/Niagara/Private/NiagaraFunctionLibrary.cpp
@@ -68,23 +68,12 @@ UNiagaraComponent* CreateNiagaraSystem(UNiagaraSystem* SystemTemplate, UWorld* W
 	{
 		if (PoolingMethod == ENCPoolMethod::None)
 		{
-			if (UNiagaraComponentSettings::ShouldForceAutoPooling(SystemTemplate))
-			{
-				if (FNiagaraWorldManager* WorldManager = FNiagaraWorldManager::Get(World))
-				{
-					UNiagaraComponentPool* ComponentPool = WorldManager->GetComponentPool();
-					NiagaraComponent = ComponentPool->CreateWorldParticleSystem(SystemTemplate, World, ENCPoolMethod::AutoRelease);
-				}
-			}
-			else
-			{
-				NiagaraComponent = NewObject<UNiagaraComponent>((Actor ? Actor : (UObject*)World));
-				NiagaraComponent->SetAsset(SystemTemplate);
-				NiagaraComponent->bAutoActivate = false;
-				NiagaraComponent->SetAutoDestroy(bAutoDestroy);
-				NiagaraComponent->bAllowAnyoneToDestroyMe = true;
-				NiagaraComponent->SetVisibleInRayTracing(false);
-			}
+			NiagaraComponent = NewObject<UNiagaraComponent>((Actor ? Actor : (UObject*)World));
+			NiagaraComponent->SetAsset(SystemTemplate);
+			NiagaraComponent->bAutoActivate = false;
+			NiagaraComponent->SetAutoDestroy(bAutoDestroy);
+			NiagaraComponent->bAllowAnyoneToDestroyMe = true;
+			NiagaraComponent->SetVisibleInRayTracing(false);
 		}
 		else if (FNiagaraWorldManager* WorldManager = FNiagaraWorldManager::Get(World))
 		{
diff --git a/Engine/Plugins/FX/Niagara/Source/Niagara/Public/NiagaraComponentSettings.h b/Engine/Plugins/FX/Niagara/Source/Niagara/Public/NiagaraComponentSettings.h
--- a/Engine/Plugins/FX/Niagara/Source/Niagara/Public/NiagaraComponentSettings.h
+++ b/Engine/Plugins/FX/Niagara/Source/Niagara/Public/NiagaraComponentSettings.h
@@ -2,145 +2,11 @@
 
 #pragma once
 
-#include "NiagaraScript.h"
-#include "NiagaraSystem.h"
-#include "NiagaraSystemInstance.h"
-#include "NiagaraEmitter.h"
-#include "NiagaraEmitterInstance.h"
+class UNiagaraSystem;
+class FNiagaraEmitterInstance;
 
-#include "NiagaraComponentSettings.generated.h"
-
-USTRUCT(meta = (DisplayName = "Emitter Name Settings Reference"))
-struct FNiagaraEmitterNameSettingsRef
-{
-	GENERATED_USTRUCT_BODY()
-
-	UPROPERTY(EditAnywhere, Category = Parameters)
-	FName SystemName;
-
-	UPROPERTY(EditAnywhere, Category = Parameters)
-	FString EmitterName;
-
-	FORCEINLINE bool operator==(const FNiagaraEmitterNameSettingsRef& Other)const
-	{
-		return SystemName == Other.SystemName && EmitterName == Other.EmitterName;
-	}
-
-	FORCEINLINE bool operator!=(const FNiagaraEmitterNameSettingsRef& Other)const
-	{
-		return !(*this == Other);
-	}
-};
-
-FORCEINLINE uint32 GetTypeHash(const FNiagaraEmitterNameSettingsRef& Var)
-{
-	return HashCombine(GetTypeHash(Var.SystemName), GetTypeHash(Var.EmitterName));
-}
-
-UCLASS(config=Game, defaultconfig)
-class NIAGARA_API UNiagaraComponentSettings : public UObject
+namespace FNiagaraComponentSettings
 {
-	GENERATED_UCLASS_BODY()
-
-public:
-	FORCEINLINE static bool ShouldSuppressActivation(const UNiagaraSystem* System)
-	{
-		check(System != nullptr);
-		if ( bAllowSuppressActivation )
-		{
-			if (const UNiagaraComponentSettings* ComponentSettings = GetDefault<UNiagaraComponentSettings>())
-			{
-				if (ComponentSettings->SuppressActivationList.Contains(System->GetFName()))
-				{
-					return true;
-				}
-			}
-		}
-		return false;
-	}
-
-	FORCEINLINE static bool ShouldForceAutoPooling(const UNiagaraSystem* System)
-	{
-		check(System != nullptr);
-		if (bAllowForceAutoPooling)
-		{
-			if (const UNiagaraComponentSettings* ComponentSettings = GetDefault<UNiagaraComponentSettings>())
-			{
-				if (ComponentSettings->ForceAutoPooolingList.Contains(System->GetFName()))
-				{
-					return true;
-				}
-			}
-		}
-		return false;
-	}
-
-	FORCEINLINE static bool ShouldSuppressEmitterActivation(const FNiagaraEmitterInstance* EmitterInstance)
-	{
-		if ( bUseSuppressEmitterList || bUseGpuEmitterAllowList || bUseGpuDataInterfaceDenyList)
-		{
-			const UNiagaraComponentSettings* ComponentSettings = GetDefault<UNiagaraComponentSettings>();
-
-			FNiagaraEmitterNameSettingsRef EmitterRef;
-			if (const UNiagaraSystem* ParentSystem = EmitterInstance->GetParentSystemInstance()->GetSystem())
-			{
-				EmitterRef.SystemName = ParentSystem->GetFName();
-			}
-			FVersionedNiagaraEmitter CachedEmitter = EmitterInstance->GetCachedEmitter();
-			EmitterRef.EmitterName = CachedEmitter.Emitter->GetUniqueEmitterName();
-
-			if (bUseSuppressEmitterList && ComponentSettings->SuppressEmitterList.Contains(EmitterRef))
-			{
-				return true;
-			}
-
-			FVersionedNiagaraEmitterData* EmitterData = CachedEmitter.GetEmitterData();
-			if (bUseGpuEmitterAllowList && (EmitterData->SimTarget == ENiagaraSimTarget::GPUComputeSim) && !ComponentSettings->GPUEmitterAllowList.Contains(EmitterRef))
-			{
-				return true;
-			}
-
-			if ( bUseGpuDataInterfaceDenyList && (EmitterData->SimTarget == ENiagaraSimTarget::GPUComputeSim) )
-			{
-				if (const UNiagaraScript* GPUComputeScript = EmitterData->GetGPUComputeScript())
-				{
-					for (const FNiagaraScriptDataInterfaceInfo& DefaultDIInfo : GPUComputeScript->GetCachedDefaultDataInterfaces())
-					{
-						if (ComponentSettings->GpuDataInterfaceDenyList.Contains(DefaultDIInfo.Type.GetFName()))
-						{
-							return true;
-						}
-					}
-				}
-			}
-		}
-		return false;
-	}
-
-	UPROPERTY(config)
-	TSet<FName> SuppressActivationList;
-
-	UPROPERTY(config)
-	TSet<FName> ForceAutoPooolingList;
-
-	/** 
-		Config file to tweak individual emitters being disabled. Syntax is as follows for the config file:
-		[/Script/Niagara.NiagaraComponentSettings]
-		SuppressEmitterList=((SystemName="BasicSpriteSystem",EmitterName="BasicSprite001"))
-		+GpuDataInterfaceDenyList=("NiagaraDataInterfaceCollisionQuery")
-	*/
-	UPROPERTY(config)
-	TSet<FNiagaraEmitterNameSettingsRef> SuppressEmitterList;
-
-	UPROPERTY(config)
-	TSet<FNiagaraEmitterNameSettingsRef> GPUEmitterAllowList;
-
-	UPROPERTY(config)
-	TSet<FName> GpuDataInterfaceDenyList;
-
-	static int32 bAllowSuppressActivation;
-	static int32 bAllowForceAutoPooling;
-	static int32 bUseSuppressEmitterList;
-	static int32 bUseGpuEmitterAllowList;
-	static int32 bUseGpuDataInterfaceDenyList;
+	extern NIAGARA_API bool IsSystemAllowedToRun(const UNiagaraSystem* System);
+	extern NIAGARA_API bool IsEmitterAllowedToRun(const FNiagaraEmitterInstance* EmitterInstance);
 };
