- IWYU fixes

#rnx

diff --git a/Engine/Plugins/FX/Niagara/Source/Niagara/Private/NiagaraComponentSettings.cpp b/Engine/Plugins/FX/Niagara/Source/Niagara/Private/NiagaraComponentSettings.cpp
--- a/Engine/Plugins/FX/Niagara/Source/Niagara/Private/NiagaraComponentSettings.cpp
+++ b/Engine/Plugins/FX/Niagara/Source/Niagara/Private/NiagaraComponentSettings.cpp
@@ -6,6 +6,7 @@
 #include "NiagaraScript.h"
 #include "NiagaraSystem.h"
 #include "NiagaraSystemInstance.h"
+#include "UObject/UObjectIterator.h"
 
 namespace FNiagaraComponentSettings
 {
diff --git a/Engine/Plugins/FX/Niagara/Source/Niagara/Public/NiagaraComponentSettings.h b/Engine/Plugins/FX/Niagara/Source/Niagara/Public/NiagaraComponentSettings.h
--- a/Engine/Plugins/FX/Niagara/Source/Niagara/Public/NiagaraComponentSettings.h
+++ b/Engine/Plugins/FX/Niagara/Source/Niagara/Public/NiagaraComponentSettings.h
@@ -2,6 +2,8 @@
 
 #pragma once
 
+#include "NiagaraCommon.h"
+
 class UNiagaraSystem;
 class FNiagaraEmitterInstance;
 
