Fixing bad parameter binding in Niagara from CL 23176960

diff --git a/Engine/Plugins/FX/Niagara/Source/NiagaraVertexFactories/Private/NiagaraDrawIndirect.cpp b/Engine/Plugins/FX/Niagara/Source/NiagaraVertexFactories/Private/NiagaraDrawIndirect.cpp
--- a/Engine/Plugins/FX/Niagara/Source/NiagaraVertexFactories/Private/NiagaraDrawIndirect.cpp
+++ b/Engine/Plugins/FX/Niagara/Source/NiagaraVertexFactories/Private/NiagaraDrawIndirect.cpp
@@ -44,7 +44,7 @@ void FNiagaraDrawIndirectArgsGenCS::SetOutput(FRHICommandList& RHICmdList, FRHIU
 {
 	FRHIComputeShader* ComputeShaderRHI = RHICmdList.GetBoundComputeShader();
 	DrawIndirectArgsParam.SetUAV(RHICmdList, ComputeShaderRHI, DrawIndirectArgsUAV);
-	InstanceCountsParam.SetUAV(RHICmdList, ComputeShaderRHI, DrawIndirectArgsUAV);
+	InstanceCountsParam.SetUAV(RHICmdList, ComputeShaderRHI, InstanceCountsUAV);
 }
 
 void FNiagaraDrawIndirectArgsGenCS::SetParameters(FRHICommandList& RHICmdList, FRHIShaderResourceView* TaskInfosBuffer, FRHIShaderResourceView* CulledInstanceCountsBuffer, int32 ArgGenTaskOffset, int32 NumArgGenTasks, int32 NumInstanceCountClearTasks)
