Added data interface auto-completion to custom hlsl nodes

diff --git a/Engine/Plugins/FX/Niagara/Source/NiagaraEditor/Private/EdGraphSchema_Niagara.cpp b/Engine/Plugins/FX/Niagara/Source/NiagaraEditor/Private/EdGraphSchema_Niagara.cpp
--- a/Engine/Plugins/FX/Niagara/Source/NiagaraEditor/Private/EdGraphSchema_Niagara.cpp
+++ b/Engine/Plugins/FX/Niagara/Source/NiagaraEditor/Private/EdGraphSchema_Niagara.cpp
@@ -2159,43 +2159,55 @@ void UEdGraphSchema_Niagara::GetNumericConversionToSubMenuActionsAll(UToolMenu*
 	}
 }
 
-void UEdGraphSchema_Niagara::GenerateDataInterfacePinMenu(UToolMenu* ToolMenu, const FName SectionName, const UEdGraphPin* GraphPin, FNiagaraTypeDefinition TypeDef) const
+TArray<TPair<FString, FString>> UEdGraphSchema_Niagara::GetDataInterfaceFunctionPrototypes(const UEdGraphPin* InGraphPin)
 {
-	UNiagaraDataInterface* DataInterface = CastChecked<UNiagaraDataInterface>(TypeDef.GetClass()->GetDefaultObject());
-
-	// No functions, don't add an option
-	TArray<FNiagaraFunctionSignature> FunctionSignatures;
-	DataInterface->GetFunctions(FunctionSignatures);
-	if (FunctionSignatures.Num() == 0)
-	{
-		return;
-	}
-
-	// Make the data interface name
-	FString DataInterfaceName = GraphPin->GetName();
-	{
-		int32 LastPinNameDot = 0;
-		if (DataInterfaceName.FindLastChar('.', LastPinNameDot))
-		{
-			DataInterfaceName.RightChopInline(LastPinNameDot + 1);
-		}
-		DataInterfaceName = FHlslNiagaraTranslator::GetSanitizedSymbolName(DataInterfaceName);
-	}
-
-	// Generate all prototypes
 	TArray<TPair<FString, FString>> FunctionPrototypes;
-	FunctionPrototypes.Reserve(FunctionSignatures.Num());
-	for (const FNiagaraFunctionSignature& FunctionSignature : FunctionSignatures)
+	if (InGraphPin)
 	{
-		FString Prototype = FHlslNiagaraTranslator::GenerateFunctionHlslPrototype(DataInterfaceName, FunctionSignature);
-		if ( Prototype.Len() > 0 )
+		FNiagaraTypeDefinition GraphPinType = UEdGraphSchema_Niagara::PinToTypeDefinition(InGraphPin);
+		if (GraphPinType.IsDataInterface())
 		{
-			Prototype.Append(TEXT("\r\n"));
-			FunctionPrototypes.Emplace(FunctionSignature.GetNameString(), Prototype);
+			UNiagaraDataInterface* DataInterface = CastChecked<UNiagaraDataInterface>(GraphPinType.GetClass()->GetDefaultObject());
+			
+			// No functions, don't add an option
+			TArray<FNiagaraFunctionSignature> FunctionSignatures;
+			DataInterface->GetFunctions(FunctionSignatures);
+			if (FunctionSignatures.Num() == 0)
+			{
+				return FunctionPrototypes;
+			}
+
+			// Make the data interface name
+			FString DataInterfaceName = InGraphPin->GetName();
+			{
+				int32 LastPinNameDot = 0;
+				if (DataInterfaceName.FindLastChar('.', LastPinNameDot))
+				{
+					DataInterfaceName.RightChopInline(LastPinNameDot + 1);
+				}
+				DataInterfaceName = FHlslNiagaraTranslator::GetSanitizedSymbolName(DataInterfaceName);
+			}
+
+			// Generate all prototypes
+			FunctionPrototypes.Reserve(FunctionSignatures.Num());
+			for (const FNiagaraFunctionSignature& FunctionSignature : FunctionSignatures)
+			{
+				FString Prototype = FHlslNiagaraTranslator::GenerateFunctionHlslPrototype(DataInterfaceName, FunctionSignature);
+				if ( Prototype.Len() > 0 )
+				{
+					Prototype.Append(TEXT("\r\n"));
+					FunctionPrototypes.Emplace(FunctionSignature.GetNameString(), Prototype);
+				}
+			}
 		}
 	}
+	return FunctionPrototypes;
+}
 
-	if ( FunctionPrototypes.Num () == 0 )
+void UEdGraphSchema_Niagara::GenerateDataInterfacePinMenu(UToolMenu* ToolMenu, const FName SectionName, const UEdGraphPin* GraphPin, FNiagaraTypeDefinition TypeDef) const
+{
+	TArray<TPair<FString, FString>> FunctionPrototypes = GetDataInterfaceFunctionPrototypes(GraphPin);
+	if (FunctionPrototypes.Num () == 0)
 	{
 		return;
 	}
diff --git a/Engine/Plugins/FX/Niagara/Source/NiagaraEditor/Private/Widgets/SNiagaraGraphNodeCustomHlsl.cpp b/Engine/Plugins/FX/Niagara/Source/NiagaraEditor/Private/Widgets/SNiagaraGraphNodeCustomHlsl.cpp
--- a/Engine/Plugins/FX/Niagara/Source/NiagaraEditor/Private/Widgets/SNiagaraGraphNodeCustomHlsl.cpp
+++ b/Engine/Plugins/FX/Niagara/Source/NiagaraEditor/Private/Widgets/SNiagaraGraphNodeCustomHlsl.cpp
@@ -1,10 +1,11 @@
 // Copyright Epic Games, Inc. All Rights Reserved.
 
 #include "SNiagaraGraphNodeCustomHlsl.h"
+
+#include "EdGraphSchema_Niagara.h"
 #include "NiagaraNodeCustomHlsl.h"
 #include "NiagaraGraph.h"
 #include "Widgets/SBoxPanel.h"
-#include "GraphEditorSettings.h"
 #include "Widgets/Input/SCheckBox.h"
 #include "Widgets/Input/SMultiLineEditableTextBox.h"
 
@@ -112,7 +113,37 @@ void SNiagaraGraphNodeCustomHlsl::CreateBelowPinControls(TSharedPtr<SVerticalBox
 		.OnKeyCharHandler(this, &SNiagaraGraphNodeCustomHlsl::OnShaderTextKeyChar)
 		.OnTextCommitted(TextCommit)
 		.Visibility(this, &SNiagaraGraphNodeCustomHlsl::GetShaderTextVisibility);
-	MainBox->AddSlot()[
+
+	// the menu anchor is a bit of a hack to show the auto-complete window.
+	// Main advantage is that it has all the layout logic inside to place the popup window correctly and auto-closes it.
+	// Disadvantage is that it's always at the bottom of the node and not where the current line cursor is.
+	SAssignNew(MenuAnchor, SMenuAnchor)
+		.Placement(MenuPlacement_BelowAnchor)
+		.Method(EPopupMethod::CreateNewWindow)
+		.Visibility(EVisibility::HitTestInvisible)
+		.OnGetMenuContent_Lambda([this]()
+		{
+			return SNew(STextBlock)
+			.Margin(5)
+			.Text_Lambda([this]()
+			{
+				TStringBuilder<256> Options;
+				for (int i = 0; (i + AutoCompletePageIndex * 9) < FunctionPrototypes.Num() && i < 9; i++)
+				{
+					TPair<FString, FString> Prototype = FunctionPrototypes[i + AutoCompletePageIndex * 9];
+					Options.Appendf(TEXT("\n(%i) %s"), i + 1, *Prototype.Key);
+				}
+				FText HiddenText = FText::GetEmpty();
+				if (FunctionPrototypes.Num() > (AutoCompletePageIndex + 1) * 9)
+				{
+					HiddenText = LOCTEXT("Autocomplete_HiddenFormat", "\n(0) more...");
+				}
+				return FText::Format(LOCTEXT("Autocomplete_Format", "Select auto-complete option:{0}{1}"), FText::FromString(Options.ToString()), HiddenText);
+			});
+		});
+	
+	MainBox->AddSlot()
+	[
 		SNew(SHorizontalBox)
 		+ SHorizontalBox::Slot()
 		.AutoWidth()
@@ -121,17 +152,123 @@ void SNiagaraGraphNodeCustomHlsl::CreateBelowPinControls(TSharedPtr<SVerticalBox
 			ShaderTextBox.ToSharedRef()
 		]
 	];
+
+	
+	MainBox->AddSlot()
+	.MaxHeight(0)
+	.AutoHeight()
+	[
+		MenuAnchor.ToSharedRef()
+	];
+}
+
+bool IsValidIdentifierChar(const TCHAR& Char)
+{
+	return ((Char >= TCHAR('0') && Char <= TCHAR('9')) ||
+			(Char >= TCHAR('A') && Char <= TCHAR('Z')) ||
+			(Char >= TCHAR('a') && Char <= TCHAR('z')) ||
+			Char == TCHAR('_'));
+}
+
+FString SNiagaraGraphNodeCustomHlsl::GetIdentifierUnderCursor() const
+{
+	FString CurrentLine;
+	ShaderTextBox->GetCurrentTextLine(CurrentLine);
+	FTextLocation CursorLocation = ShaderTextBox->GetCursorLocation();
+	int32 Offset = CursorLocation.GetOffset();
+	int32 StartIndex = Offset - 1;
+	while (StartIndex >= 0)
+	{
+		if (IsValidIdentifierChar(CurrentLine[StartIndex]))
+		{
+			StartIndex--;
+		}
+		else
+		{
+			break;
+		}
+	}
+	StartIndex++;
+
+	if (StartIndex >= 0 && Offset >= 0 && StartIndex < Offset)
+	{
+		return CurrentLine.Mid(StartIndex, Offset - StartIndex);
+	}
+	return FString();
 }
 
 FReply SNiagaraGraphNodeCustomHlsl::OnShaderTextKeyChar(const FGeometry&, const FCharacterEvent& InCharacterEvent)
 {
 	const TCHAR Character = InCharacterEvent.GetCharacter();
+
+	if (bIsAutoCompleteActive)
+	{
+		if (Character == TCHAR('0') && FunctionPrototypes.Num() > (AutoCompletePageIndex + 1) * 9)
+		{
+			AutoCompletePageIndex++;
+			return FReply::Handled();
+		}
+		if ((Character >= TCHAR('1') && Character <= TCHAR('9')))
+		{
+			int32 SelectedOption = (Character - TCHAR('1')) + AutoCompletePageIndex * 9;
+			if (AutoCompleteOptions.IsValidIndex(SelectedOption))
+			{
+				ShaderTextBox->InsertTextAtCursor(AutoCompleteOptions[SelectedOption]);
+				bIsAutoCompleteActive = false;
+				MenuAnchor->SetIsOpen(false);
+				return FReply::Handled();
+			}
+		}
+	}
+	MenuAnchor->SetIsOpen(false);
+	bIsAutoCompleteActive = false;
+	
 	if (Character == TEXT('\t'))
 	{
 		// Convert tab to four spaces
 		ShaderTextBox->InsertTextAtCursor(TEXT("    "));
 		return FReply::Handled();
 	}
+	if (Character == TEXT('.'))
+	{
+		// Show auto-complete hint for data interfaces
+		FString Identifier = GetIdentifierUnderCursor();
+		if (Identifier.Len() > 0)
+		{
+			UNiagaraNodeCustomHlsl* CustomNode = CastChecked<UNiagaraNodeCustomHlsl>(GraphNode);
+			FPinCollectorArray NodeInputPins;
+			CustomNode->GetInputPins(NodeInputPins);
+			for (UEdGraphPin* Pin : NodeInputPins)
+			{
+				if (Pin->GetName().Equals(Identifier, ESearchCase::IgnoreCase))
+				{
+					FunctionPrototypes = UEdGraphSchema_Niagara::GetDataInterfaceFunctionPrototypes(Pin);
+					if (FunctionPrototypes.Num() > 0)
+					{
+						bIsAutoCompleteActive = true;
+						AutoCompletePageIndex = 0;
+						AutoCompleteOptions.Empty();
+
+						for (int i = 0; i < FunctionPrototypes.Num(); i++)
+						{
+							TPair<FString, FString> Prototype = FunctionPrototypes[i];
+							int32 NameDot = 0;
+							if (Prototype.Value.FindChar('.', NameDot))
+							{
+								Prototype.Value.RightChopInline(NameDot + 1);
+							}
+							Prototype.Value.TrimStartAndEndInline();
+							AutoCompleteOptions.Add(Prototype.Value);
+						}						
+						MenuAnchor->SetIsOpen(true, false);
+						break;
+					}
+				}
+			}
+		}
+		ShaderTextBox->InsertTextAtCursor(TEXT("."));
+		return FReply::Handled();
+	}
 	
 	// Let SMultiLineEditableTextBox::OnKeyChar handle it.
 	return FReply::Unhandled();
diff --git a/Engine/Plugins/FX/Niagara/Source/NiagaraEditor/Private/Widgets/SNiagaraGraphNodeCustomHlsl.h b/Engine/Plugins/FX/Niagara/Source/NiagaraEditor/Private/Widgets/SNiagaraGraphNodeCustomHlsl.h
--- a/Engine/Plugins/FX/Niagara/Source/NiagaraEditor/Private/Widgets/SNiagaraGraphNodeCustomHlsl.h
+++ b/Engine/Plugins/FX/Niagara/Source/NiagaraEditor/Private/Widgets/SNiagaraGraphNodeCustomHlsl.h
@@ -3,6 +3,7 @@
 #pragma once
 
 #include "Widgets/NiagaraHLSLSyntaxHighlighter.h"
+#include "Widgets/Input/SMenuAnchor.h"
 #include "SNiagaraGraphNode.h"
 
 class SMultiLineEditableTextBox;
@@ -29,12 +30,19 @@ protected:
 
 	TSharedPtr<FNiagaraHLSLSyntaxHighlighter> SyntaxHighlighter;
 	TSharedPtr<SMultiLineEditableTextBox> ShaderTextBox;
+	TSharedPtr<SMenuAnchor> MenuAnchor;
 
 private:
 	FReply OnShaderTextKeyChar(const FGeometry& MyGeometry, const FCharacterEvent& InCharacterEvent);
+	FString GetIdentifierUnderCursor() const;
 
 	EVisibility GetShaderTextVisibility() const;
 	void ToggleShowShaderCode(const ECheckBoxState NewCheckedState);
 	ECheckBoxState GetToggleButtonChecked() const;
 	const FSlateBrush* GetToggleButtonArrow() const;
+
+	bool bIsAutoCompleteActive = false;
+	TArray<FString> AutoCompleteOptions;
+	TArray<TPair<FString, FString>> FunctionPrototypes;
+	int32 AutoCompletePageIndex = 0;
 };
diff --git a/Engine/Plugins/FX/Niagara/Source/NiagaraEditor/Public/EdGraphSchema_Niagara.h b/Engine/Plugins/FX/Niagara/Source/NiagaraEditor/Public/EdGraphSchema_Niagara.h
--- a/Engine/Plugins/FX/Niagara/Source/NiagaraEditor/Public/EdGraphSchema_Niagara.h
+++ b/Engine/Plugins/FX/Niagara/Source/NiagaraEditor/Public/EdGraphSchema_Niagara.h
@@ -185,7 +185,8 @@ class NIAGARAEDITOR_API UEdGraphSchema_Niagara : public UEdGraphSchema
 	static const FText CircularConnectionFoundText;
 	
 	bool PinTypesValidForNumericConversion(FEdGraphPinType AType, FEdGraphPinType BType) const;	
-	
+
+	static TArray<TPair<FString, FString>> GetDataInterfaceFunctionPrototypes(const UEdGraphPin* GraphPin);
 private:
 	void GetNumericConversionToSubMenuActions(class UToolMenu* Menu, const FName SectionName, UEdGraphPin* InGraphPin);
 	void GetNumericConversionToSubMenuActionsAll(class UToolMenu* Menu, const FName SectionName, UNiagaraNode* InGraphPin);
