- Update NiagaraDrawIndirect to use shader parameters

#rnx

diff --git a/Engine/Plugins/FX/Niagara/Source/Niagara/Private/NiagaraGPUInstanceCountManager.cpp b/Engine/Plugins/FX/Niagara/Source/Niagara/Private/NiagaraGPUInstanceCountManager.cpp
--- a/Engine/Plugins/FX/Niagara/Source/Niagara/Private/NiagaraGPUInstanceCountManager.cpp
+++ b/Engine/Plugins/FX/Niagara/Source/Niagara/Private/NiagaraGPUInstanceCountManager.cpp
@@ -486,33 +486,39 @@ void FNiagaraGPUInstanceCountManager::UpdateDrawIndirectBuffers(FNiagaraGpuCompu
 				continue;
 			}
 
-			SetComputePipelineState(RHICmdList, DrawIndirectArgsGenCS.GetComputeShader());
-			DrawIndirectArgsGenCS->SetOutput(RHICmdList, ArgsUAV, CountsUAV);
-			DrawIndirectArgsGenCS->SetParameters(RHICmdList, TaskInfosBuffer.SRV, CulledCountsSRV, ArgGenTaskOffset, NumArgGenTasks, NumInstanceCountClearTasks);
+			FNiagaraDrawIndirectArgsGenCS::FParameters ArgsGenParameters;
+			ArgsGenParameters.TaskInfos				= TaskInfosBuffer.SRV;
+			ArgsGenParameters.CulledInstanceCounts	= CulledCountsSRV;
+			ArgsGenParameters.RWInstanceCounts		= CountsUAV;
+			ArgsGenParameters.RWDrawIndirectArgs	= ArgsUAV;
+			ArgsGenParameters.TaskCount.X			= ArgGenTaskOffset;
+			ArgsGenParameters.TaskCount.Y			= NumArgGenTasks;
+			ArgsGenParameters.TaskCount.Z			= NumInstanceCountClearTasks;
+			ArgsGenParameters.TaskCount.W			= NumArgGenTasks + NumInstanceCountClearTasks;
 
 			// If the device supports RW Texture buffers then we can use a single compute pass, otherwise we need to split into two passes
 			if (GRHISupportsRWTextureBuffers)
 			{
-				DispatchComputeShader(RHICmdList, DrawIndirectArgsGenCS.GetShader(), FMath::DivideAndRoundUp(NumArgGenTasks + NumInstanceCountClearTasks, NIAGARA_DRAW_INDIRECT_ARGS_GEN_THREAD_COUNT), 1, 1);
-				DrawIndirectArgsGenCS->UnbindBuffers(RHICmdList);
+				FComputeShaderUtils::Dispatch(RHICmdList, DrawIndirectArgsGenCS, ArgsGenParameters, FIntVector(FMath::DivideAndRoundUp(NumArgGenTasks + NumInstanceCountClearTasks, NIAGARA_DRAW_INDIRECT_ARGS_GEN_THREAD_COUNT), 1, 1));
 			}
 			else
 			{
 				if (NumArgGenTasks > 0)
 				{
-					DispatchComputeShader(RHICmdList, DrawIndirectArgsGenCS.GetShader(), FMath::DivideAndRoundUp(NumArgGenTasks, NIAGARA_DRAW_INDIRECT_ARGS_GEN_THREAD_COUNT), 1, 1);
-					DrawIndirectArgsGenCS->UnbindBuffers(RHICmdList);
+					FComputeShaderUtils::Dispatch(RHICmdList, DrawIndirectArgsGenCS, ArgsGenParameters, FIntVector(FMath::DivideAndRoundUp(NumArgGenTasks, NIAGARA_DRAW_INDIRECT_ARGS_GEN_THREAD_COUNT), 1, 1));
 				}
 
 				if (NumInstanceCountClearTasks > 0)
 				{
+					FNiagaraDrawIndirectResetCountsCS::FParameters ClearCountParameters;
+					ClearCountParameters.TaskInfos			= ArgsGenParameters.TaskInfos;
+					ClearCountParameters.RWInstanceCounts	= ArgsGenParameters.RWInstanceCounts;
+					ClearCountParameters.TaskCount			= ArgsGenParameters.TaskCount;
+					ClearCountParameters.TaskCount.X		= 0;
+
 					FNiagaraDrawIndirectResetCountsCS::FPermutationDomain PermutationVectorResetCounts;
 					TShaderMapRef<FNiagaraDrawIndirectResetCountsCS> DrawIndirectResetCountsArgsGenCS(GetGlobalShaderMap(FeatureLevel), PermutationVectorResetCounts);
-					SetComputePipelineState(RHICmdList, DrawIndirectResetCountsArgsGenCS.GetComputeShader());
-					DrawIndirectResetCountsArgsGenCS->SetOutput(RHICmdList, CountBuffer.UAV);
-					DrawIndirectResetCountsArgsGenCS->SetParameters(RHICmdList, TaskInfosBuffer.SRV, ArgTasks.Num(), NumInstanceCountClearTasks);
-					DispatchComputeShader(RHICmdList, DrawIndirectResetCountsArgsGenCS.GetShader(), FMath::DivideAndRoundUp(NumInstanceCountClearTasks, NIAGARA_DRAW_INDIRECT_ARGS_GEN_THREAD_COUNT), 1, 1);
-					DrawIndirectResetCountsArgsGenCS->UnbindBuffers(RHICmdList);
+					FComputeShaderUtils::Dispatch(RHICmdList, DrawIndirectResetCountsArgsGenCS, ClearCountParameters, FIntVector(FMath::DivideAndRoundUp(NumInstanceCountClearTasks, NIAGARA_DRAW_INDIRECT_ARGS_GEN_THREAD_COUNT), 1, 1));
 				}
 			}
 
diff --git a/Engine/Plugins/FX/Niagara/Source/NiagaraVertexFactories/Private/NiagaraDrawIndirect.cpp b/Engine/Plugins/FX/Niagara/Source/NiagaraVertexFactories/Private/NiagaraDrawIndirect.cpp
--- a/Engine/Plugins/FX/Niagara/Source/NiagaraVertexFactories/Private/NiagaraDrawIndirect.cpp
+++ b/Engine/Plugins/FX/Niagara/Source/NiagaraVertexFactories/Private/NiagaraDrawIndirect.cpp
@@ -19,53 +19,6 @@ void FNiagaraDrawIndirectArgsGenCS::ModifyCompilationEnvironment(const FGlobalSh
 	OutEnvironment.SetDefine(TEXT("NIAGARA_DRAW_INDIRECT_TASK_INFO_SIZE"), NIAGARA_DRAW_INDIRECT_TASK_INFO_SIZE);
 }
 
-FNiagaraDrawIndirectArgsGenCS::FNiagaraDrawIndirectArgsGenCS(const ShaderMetaType::CompiledShaderInitializerType& Initializer)
-	: FGlobalShader(Initializer)
-{
-	TaskInfosParam.Bind(Initializer.ParameterMap, TEXT("TaskInfos"));
-	CulledInstanceCountsParam.Bind(Initializer.ParameterMap, TEXT("CulledInstanceCounts"));
-	InstanceCountsParam.Bind(Initializer.ParameterMap, TEXT("InstanceCounts"));
-	DrawIndirectArgsParam.Bind(Initializer.ParameterMap, TEXT("DrawIndirectArgs"));
-	TaskCountParam.Bind(Initializer.ParameterMap, TEXT("TaskCount"));
-}
-
-/*bool FNiagaraDrawIndirectArgsGenCS::Serialize(FArchive& Ar)
-{
-	bool bShaderHasOutdatedParameters = FGlobalShader::Serialize(Ar);
-	Ar << TaskInfosParam;
-	Ar << CulledInstanceCountsParam;
-	Ar << InstanceCountsParam;
-	Ar << DrawIndirectArgsParam;
-	Ar << TaskCountParam;
-	return bShaderHasOutdatedParameters;
-}*/
-
-void FNiagaraDrawIndirectArgsGenCS::SetOutput(FRHICommandList& RHICmdList, FRHIUnorderedAccessView* DrawIndirectArgsUAV, FRHIUnorderedAccessView* InstanceCountsUAV)
-{
-	FRHIComputeShader* ComputeShaderRHI = RHICmdList.GetBoundComputeShader();
-	DrawIndirectArgsParam.SetUAV(RHICmdList, ComputeShaderRHI, DrawIndirectArgsUAV);
-	InstanceCountsParam.SetUAV(RHICmdList, ComputeShaderRHI, InstanceCountsUAV);
-}
-
-void FNiagaraDrawIndirectArgsGenCS::SetParameters(FRHICommandList& RHICmdList, FRHIShaderResourceView* TaskInfosBuffer, FRHIShaderResourceView* CulledInstanceCountsBuffer, int32 ArgGenTaskOffset, int32 NumArgGenTasks, int32 NumInstanceCountClearTasks)
-{
-	FRHIComputeShader* ComputeShaderRHI = RHICmdList.GetBoundComputeShader();
-
-	SetSRVParameter(RHICmdList, ComputeShaderRHI, TaskInfosParam, TaskInfosBuffer);
-	SetSRVParameter(RHICmdList, ComputeShaderRHI, CulledInstanceCountsParam, CulledInstanceCountsBuffer);
-
-	const FUintVector4 TaskCountValue(ArgGenTaskOffset, NumArgGenTasks, NumInstanceCountClearTasks, NumArgGenTasks + NumInstanceCountClearTasks);
-	SetShaderValue(RHICmdList, ComputeShaderRHI, TaskCountParam, TaskCountValue);
-}
-
-void FNiagaraDrawIndirectArgsGenCS::UnbindBuffers(FRHICommandList& RHICmdList)
-{
-	FRHIComputeShader* ComputeShaderRHI = RHICmdList.GetBoundComputeShader();
-
-	SetSRVParameter(RHICmdList, ComputeShaderRHI, TaskInfosParam, nullptr);
-	DrawIndirectArgsParam.UnsetUAV(RHICmdList, ComputeShaderRHI);
-	InstanceCountsParam.UnsetUAV(RHICmdList, ComputeShaderRHI);
-}
 void FNiagaraDrawIndirectResetCountsCS::ModifyCompilationEnvironment(const FGlobalShaderPermutationParameters& Parameters, FShaderCompilerEnvironment& OutEnvironment)
 {
 	FGlobalShader::ModifyCompilationEnvironment(Parameters, OutEnvironment);
@@ -73,35 +26,3 @@ void FNiagaraDrawIndirectResetCountsCS::ModifyCompilationEnvironment(const FGlob
 	OutEnvironment.SetDefine(TEXT("NIAGARA_DRAW_INDIRECT_ARGS_SIZE"), NIAGARA_DRAW_INDIRECT_ARGS_SIZE);
 	OutEnvironment.SetDefine(TEXT("NIAGARA_DRAW_INDIRECT_TASK_INFO_SIZE"), NIAGARA_DRAW_INDIRECT_TASK_INFO_SIZE);
 }
-
-FNiagaraDrawIndirectResetCountsCS::FNiagaraDrawIndirectResetCountsCS(const ShaderMetaType::CompiledShaderInitializerType& Initializer)
-	: FGlobalShader(Initializer)
-{
-	TaskInfosParam.Bind(Initializer.ParameterMap, TEXT("TaskInfos"));
-	InstanceCountsParam.Bind(Initializer.ParameterMap, TEXT("InstanceCounts"));
-	TaskCountParam.Bind(Initializer.ParameterMap, TEXT("TaskCount"));
-}
-
-void FNiagaraDrawIndirectResetCountsCS::SetOutput(FRHICommandList& RHICmdList, FRHIUnorderedAccessView* InstanceCountsUAV)
-{
-	FRHIComputeShader* ComputeShaderRHI = RHICmdList.GetBoundComputeShader();
-	InstanceCountsParam.SetUAV(RHICmdList, ComputeShaderRHI, InstanceCountsUAV);
-}
-
-void FNiagaraDrawIndirectResetCountsCS::SetParameters(FRHICommandList& RHICmdList, FRHIShaderResourceView* TaskInfosBuffer, int32 NumArgGenTasks, int32 NumInstanceCountClearTasks)
-{
-	FRHIComputeShader* ComputeShaderRHI = RHICmdList.GetBoundComputeShader();
-
-	SetSRVParameter(RHICmdList, ComputeShaderRHI, TaskInfosParam, TaskInfosBuffer);
-
-	const FUintVector4 TaskCountValue(0, NumArgGenTasks, NumInstanceCountClearTasks, NumArgGenTasks + NumInstanceCountClearTasks);
-	SetShaderValue(RHICmdList, ComputeShaderRHI, TaskCountParam, TaskCountValue);
-}
-
-void FNiagaraDrawIndirectResetCountsCS::UnbindBuffers(FRHICommandList& RHICmdList)
-{
-	FRHIComputeShader* ComputeShaderRHI = RHICmdList.GetBoundComputeShader();
-
-	SetSRVParameter(RHICmdList, ComputeShaderRHI, TaskInfosParam, nullptr);
-	InstanceCountsParam.UnsetUAV(RHICmdList, ComputeShaderRHI);
-}
\ No newline at end of file
diff --git a/Engine/Plugins/FX/Niagara/Source/NiagaraVertexFactories/Public/NiagaraDrawIndirect.h b/Engine/Plugins/FX/Niagara/Source/NiagaraVertexFactories/Public/NiagaraDrawIndirect.h
--- a/Engine/Plugins/FX/Niagara/Source/NiagaraVertexFactories/Public/NiagaraDrawIndirect.h
+++ b/Engine/Plugins/FX/Niagara/Source/NiagaraVertexFactories/Public/NiagaraDrawIndirect.h
@@ -9,6 +9,7 @@ NiagaraDrawIndirect.h: Niagara shader to generate the draw indirect args for Nia
 #include "CoreMinimal.h"
 #include "GlobalShader.h"
 #include "ShaderPermutation.h"
+#include "ShaderParameterStruct.h"
 #include "ShaderParameterUtils.h"
 
 #define NIAGARA_DRAW_INDIRECT_ARGS_GEN_THREAD_COUNT 64
@@ -65,28 +66,23 @@ struct FNiagaraDrawIndirectArgGenTaskInfo
 class NIAGARAVERTEXFACTORIES_API FNiagaraDrawIndirectArgsGenCS : public FGlobalShader
 {
 	DECLARE_GLOBAL_SHADER(FNiagaraDrawIndirectArgsGenCS);
+	SHADER_USE_PARAMETER_STRUCT(FNiagaraDrawIndirectArgsGenCS, FGlobalShader);
 
 public:
 	class FSupportsTextureRW : SHADER_PERMUTATION_INT("SUPPORTS_TEXTURE_RW", 2);
 	using FPermutationDomain = TShaderPermutationDomain<FSupportsTextureRW>;
 
-	static void ModifyCompilationEnvironment(const FGlobalShaderPermutationParameters& Parameters, FShaderCompilerEnvironment& OutEnvironment);
-
-	FNiagaraDrawIndirectArgsGenCS() {}
-	FNiagaraDrawIndirectArgsGenCS(const ShaderMetaType::CompiledShaderInitializerType& Initializer);
+	BEGIN_SHADER_PARAMETER_STRUCT(FParameters, NIAGARAVERTEXFACTORIES_API)
+		SHADER_PARAMETER_SRV(Buffer<uint>,		TaskInfos)
+		SHADER_PARAMETER_SRV(Buffer<uint>,		CulledInstanceCounts)
 
-	//bool Serialize(FArchive& Ar);
-	void SetOutput(FRHICommandList& RHICmdList, FRHIUnorderedAccessView* DrawIndirectArgsUAV, FRHIUnorderedAccessView* InstanceCountsUAV);
-	void SetParameters(FRHICommandList& RHICmdList, FRHIShaderResourceView* TaskInfosBuffer, FRHIShaderResourceView* CulledInstanceCountsBuffer, int32 NumArgGenTaskOffset, int32 NumArgGenTasks, int32 NumInstanceCountClearTasks);
-	void UnbindBuffers(FRHICommandList& RHICmdList);
+		SHADER_PARAMETER_UAV(RWBuffer<uint>,	RWInstanceCounts)
+		SHADER_PARAMETER_UAV(RWBuffer<uint>,	RWDrawIndirectArgs)
 
-protected:
+		SHADER_PARAMETER(FUintVector4,			TaskCount)
+	END_SHADER_PARAMETER_STRUCT()
 
-	LAYOUT_FIELD(FShaderResourceParameter, TaskInfosParam)
-	LAYOUT_FIELD(FShaderResourceParameter, CulledInstanceCountsParam);
-	LAYOUT_FIELD(FRWShaderParameter, InstanceCountsParam)
-	LAYOUT_FIELD(FRWShaderParameter, DrawIndirectArgsParam)
-	LAYOUT_FIELD(FShaderParameter, TaskCountParam)
+	static void ModifyCompilationEnvironment(const FGlobalShaderPermutationParameters& Parameters, FShaderCompilerEnvironment& OutEnvironment);
 };
 
 /**
@@ -95,23 +91,16 @@ protected:
 class NIAGARAVERTEXFACTORIES_API FNiagaraDrawIndirectResetCountsCS : public FGlobalShader
 {
 	DECLARE_GLOBAL_SHADER(FNiagaraDrawIndirectResetCountsCS);
+	SHADER_USE_PARAMETER_STRUCT(FNiagaraDrawIndirectResetCountsCS, FGlobalShader);
 
 public:
 	using FPermutationDomain = TShaderPermutationDomain<>;
 
-	static void ModifyCompilationEnvironment(const FGlobalShaderPermutationParameters& Parameters, FShaderCompilerEnvironment& OutEnvironment);
-
-	FNiagaraDrawIndirectResetCountsCS() {}
-	FNiagaraDrawIndirectResetCountsCS(const ShaderMetaType::CompiledShaderInitializerType& Initializer);
+	BEGIN_SHADER_PARAMETER_STRUCT(FParameters, NIAGARAVERTEXFACTORIES_API)
+		SHADER_PARAMETER_SRV(Buffer<uint>,		TaskInfos)
+		SHADER_PARAMETER_UAV(RWBuffer<uint>,	RWInstanceCounts)
+		SHADER_PARAMETER(FUintVector4,			TaskCount)
+	END_SHADER_PARAMETER_STRUCT()
 
-	//bool Serialize(FArchive& Ar);
-	void SetOutput(FRHICommandList& RHICmdList, FRHIUnorderedAccessView* InstanceCountsUAV);
-	void SetParameters(FRHICommandList& RHICmdList, FRHIShaderResourceView* TaskInfosBuffer, int32 NumArgGenTasks, int32 NumInstanceCountClearTasks);
-	void UnbindBuffers(FRHICommandList& RHICmdList);
-
-protected:
-
-	LAYOUT_FIELD(FShaderResourceParameter, TaskInfosParam)
-	LAYOUT_FIELD(FRWShaderParameter, InstanceCountsParam)
-	LAYOUT_FIELD(FShaderParameter, TaskCountParam)
-};
\ No newline at end of file
+	static void ModifyCompilationEnvironment(const FGlobalShaderPermutationParameters& Parameters, FShaderCompilerEnvironment& OutEnvironment);
+};
