PR #9821: [Niagara] Fixed "UnitToUV" binding for RasterGrid3D (Contributed by [REDACTED_PERSON])

diff --git a/Engine/Plugins/FX/Niagara/Source/Niagara/Private/NiagaraDataInterfaceRasterizationGrid3D.cpp b/Engine/Plugins/FX/Niagara/Source/Niagara/Private/NiagaraDataInterfaceRasterizationGrid3D.cpp
--- a/Engine/Plugins/FX/Niagara/Source/Niagara/Private/NiagaraDataInterfaceRasterizationGrid3D.cpp
+++ b/Engine/Plugins/FX/Niagara/Source/Niagara/Private/NiagaraDataInterfaceRasterizationGrid3D.cpp
@@ -16,7 +16,7 @@ namespace NDIRasterizationGrid3DLocal
 {
 	BEGIN_SHADER_PARAMETER_STRUCT(FShaderParameters,)
 		SHADER_PARAMETER(FIntVector,						NumCells)
-		SHADER_PARAMETER(FVector3f,							UnitToUVParam)
+		SHADER_PARAMETER(FVector3f,							UnitToUV)
 		SHADER_PARAMETER(float,								Precision)
 		SHADER_PARAMETER(int,								NumAttributes)
 		SHADER_PARAMETER_RDG_TEXTURE_SRV(Texture3D<int>,	IntGrid)
@@ -836,7 +836,7 @@ void UNiagaraDataInterfaceRasterizationGrid3D::SetShaderParameters(const FNiagar
 	if (InstanceData && InstanceData->RasterizationTexture.IsValid())
 	{
 		ShaderParameters->NumCells = InstanceData->NumCells;
-		ShaderParameters->UnitToUVParam = FVector3f(1.0f) / FVector3f(InstanceData->NumCells);
+		ShaderParameters->UnitToUV = FVector3f(1.0f) / FVector3f(InstanceData->NumCells);
 		ShaderParameters->Precision = InstanceData->Precision;
 		ShaderParameters->NumAttributes = InstanceData->TotalNumAttributes;
 		ShaderParameters->PerAttributeData = InstanceData->PerAttributeData.SRV;
@@ -861,7 +861,7 @@ void UNiagaraDataInterfaceRasterizationGrid3D::SetShaderParameters(const FNiagar
 	else
 	{
 		ShaderParameters->NumCells = FIntVector::ZeroValue;
-		ShaderParameters->UnitToUVParam = FVector3f::ZeroVector;
+		ShaderParameters->UnitToUV = FVector3f::ZeroVector;
 		ShaderParameters->Precision = 0;
 		ShaderParameters->NumAttributes = 0;
 		ShaderParameters->PerAttributeData = FNiagaraRenderer::GetDummyFloat4Buffer();
