- Change cutout texture detection to work with material instances that override the texture

#rnx
#preflght 63780b1f170bc34a93846ec9

diff --git a/Engine/Plugins/FX/Niagara/Source/Niagara/Private/NiagaraSpriteRendererProperties.cpp b/Engine/Plugins/FX/Niagara/Source/Niagara/Private/NiagaraSpriteRendererProperties.cpp
--- a/Engine/Plugins/FX/Niagara/Source/Niagara/Private/NiagaraSpriteRendererProperties.cpp
+++ b/Engine/Plugins/FX/Niagara/Source/Niagara/Private/NiagaraSpriteRendererProperties.cpp
@@ -640,29 +640,42 @@ void UNiagaraSpriteRendererProperties::GetRendererTooltipWidgets(const FNiagaraE
 
 void UNiagaraSpriteRendererProperties::UpdateCutoutTexture()
 {
-	if (bUseMaterialCutoutTexture)
+	if (bUseMaterialCutoutTexture == false)
 	{
-		CutoutTexture = nullptr;
-		if (Material)
+		return;
+	}
+
+	CutoutTexture = nullptr;
+	if (Material == nullptr || Material->GetMaterial() == nullptr)
+	{
+		return;
+	}
+
+	// The property should probably come from the blend mode, but keeping consistency with how this has always been
+	for (const EMaterialProperty MaterialProperty : {MP_OpacityMask, MP_Opacity})
+	{
+		TArray<UMaterialExpression*> MaterialExpressions;
+		Material->GetMaterial()->GetExpressionsInPropertyChain(MaterialProperty, MaterialExpressions, nullptr);
+		for (UMaterialExpression* MaterialExpression : MaterialExpressions)
 		{
-			// Try to find an opacity mask texture to default to, if not try to find an opacity texture
-			TArray<UTexture*> OpacityMaskTextures;
-			Material->GetTexturesInPropertyChain(MP_OpacityMask, OpacityMaskTextures, nullptr, nullptr);
-			if (OpacityMaskTextures.Num())
+			UMaterialExpressionTextureSample* TextureExpression = MaterialExpression ? Cast<UMaterialExpressionTextureSample>(MaterialExpression) : nullptr;
+			if (TextureExpression == nullptr)
 			{
-				CutoutTexture = (UTexture2D*)OpacityMaskTextures[0];
+				continue;
 			}
-			else
+
+			UTexture* ExpressionTexture = TextureExpression->Texture;
+			if (UMaterialExpressionTextureSampleParameter* TextureParameterExpression = Cast<UMaterialExpressionTextureSampleParameter>(TextureExpression))
 			{
-				TArray<UTexture*> OpacityTextures;
-				Material->GetTexturesInPropertyChain(MP_Opacity, OpacityTextures, nullptr, nullptr);
-				if (OpacityTextures.Num())
-				{
-					CutoutTexture = (UTexture2D*)OpacityTextures[0];
-				}
+				Material->GetTextureParameterValue(TextureExpression->GetParameterName(), ExpressionTexture);
 			}
+
+			// We bail on the first texture found, which isn't accurate but good enough
+			CutoutTexture = Cast<UTexture2D>(ExpressionTexture);
+			return;
 		}
 	}
+	// Note we do not get here unless we failed to find a cutout texture
 }
 
 void UNiagaraSpriteRendererProperties::CacheDerivedData()
