- Fix mobile preview and Niagara GPU shaders which use scene textures
- We track the device feature level in editor and force a recache of scripts that use GPU compute if it changes

diff --git a/Engine/Plugins/FX/Niagara/Source/Niagara/Classes/NiagaraScript.h b/Engine/Plugins/FX/Niagara/Source/Niagara/Classes/NiagaraScript.h
--- a/Engine/Plugins/FX/Niagara/Source/Niagara/Classes/NiagaraScript.h
+++ b/Engine/Plugins/FX/Niagara/Source/Niagara/Classes/NiagaraScript.h
@@ -982,6 +982,9 @@ public:
 	NIAGARA_API bool CanBeRunOnGpu() const;
 	NIAGARA_API bool IsReadyToRun(ENiagaraSimTarget SimTarget) const;
 	NIAGARA_API bool ShouldCacheShadersForCooking(const ITargetPlatform* TargetPlatform) const;
+#if WITH_EDITORONLY_DATA
+	NIAGARA_API static void SetPreviewFeatureLevel(ERHIFeatureLevel::Type PreviewFeatureLevel);
+#endif
 
 #if WITH_EDITORONLY_DATA
 	NIAGARA_API class UNiagaraScriptSourceBase* GetLatestSource();
diff --git a/Engine/Plugins/FX/Niagara/Source/Niagara/Private/NiagaraScript.cpp b/Engine/Plugins/FX/Niagara/Source/Niagara/Private/NiagaraScript.cpp
--- a/Engine/Plugins/FX/Niagara/Source/Niagara/Private/NiagaraScript.cpp
+++ b/Engine/Plugins/FX/Niagara/Source/Niagara/Private/NiagaraScript.cpp
@@ -118,10 +118,16 @@ static FAutoConsoleVariableRef CVarNiagaraEventSpawnsUpdateInitialAttributeValue
 	ECVF_Default
 );
 
-FNiagaraScriptDebuggerInfo::FNiagaraScriptDebuggerInfo() : bWaitForGPU(false), Usage(ENiagaraScriptUsage::Function), FrameLastWriteId(-1), bWritten(false)
+namespace NiagaraScriptLocal
 {
+#if WITH_EDITORONLY_DATA
+	TOptional<ERHIFeatureLevel::Type> PreviewFeatureLevel;
+#endif
 }
 
+FNiagaraScriptDebuggerInfo::FNiagaraScriptDebuggerInfo() : bWaitForGPU(false), Usage(ENiagaraScriptUsage::Function), FrameLastWriteId(-1), bWritten(false)
+{
+}
 
 FNiagaraScriptDebuggerInfo::FNiagaraScriptDebuggerInfo(FName InName, ENiagaraScriptUsage InUsage, const FGuid& InUsageId) : HandleName(InName), Usage(InUsage), UsageId(InUsageId), FrameLastWriteId(-1), bWritten(false)
 {
@@ -2229,6 +2235,46 @@ bool UNiagaraScript::ShouldCacheShadersForCooking(const ITargetPlatform* TargetP
 	return false;
 }
 
+#if WITH_EDITORONLY_DATA
+void UNiagaraScript::SetPreviewFeatureLevel(ERHIFeatureLevel::Type InPreviewFeatureLevel)
+{
+	// Anything to do?
+	if ( NiagaraScriptLocal::PreviewFeatureLevel.Get(GMaxRHIFeatureLevel) == InPreviewFeatureLevel )
+	{
+		return;
+	}
+
+	// Either Set of Clear the Preview Feature Level
+	if ( InPreviewFeatureLevel == GMaxRHIFeatureLevel )
+	{
+		NiagaraScriptLocal::PreviewFeatureLevel.Reset();
+	}
+	else
+	{
+		NiagaraScriptLocal::PreviewFeatureLevel = InPreviewFeatureLevel;
+	}
+
+	// Force scripts to recache that can run on the GPU and already have a script resource
+	for (TObjectIterator<UNiagaraScript> It; It; ++It)
+	{
+		UNiagaraScript* Script = *It;
+		if ( Script == nullptr || !Script->OwnerCanBeRunOnGpu() || Script->ScriptResource == nullptr )
+		{
+			continue;
+		}
+
+		for ( FNiagaraShaderScript* ExistingScript : Script->ScriptResourcesByFeatureLevel )
+		{
+			if (ExistingScript == Script->ScriptResource.Get())
+			{
+				Script->CacheResourceShadersForRendering(true, false);
+				break;
+			}
+		}
+	}
+}
+#endif
+
 void UNiagaraScript::GenerateStatIDs()
 {
 #if STATS
@@ -3290,7 +3336,11 @@ void UNiagaraScript::CacheResourceShadersForRendering(bool bRegenerateId, bool b
 		UNiagaraScriptSourceBase* Source = GetLatestScriptData()->Source;
 		if (Source && OwnerCanBeRunOnGpu())
 		{
-			ERHIFeatureLevel::Type CacheFeatureLevel = GMaxRHIFeatureLevel;
+		#if WITH_EDITORONLY_DATA
+			const ERHIFeatureLevel::Type CacheFeatureLevel = NiagaraScriptLocal::PreviewFeatureLevel.Get(GMaxRHIFeatureLevel);
+		#else
+			const ERHIFeatureLevel::Type CacheFeatureLevel = GMaxRHIFeatureLevel;
+		#endif
 			const EShaderPlatform ShaderPlatform = GShaderPlatformForFeatureLevel[CacheFeatureLevel];
 
 			ScriptResource->SetScript(this, CacheFeatureLevel, ShaderPlatform, CachedScriptVMId.CompilerVersionID, CachedScriptVMId.AdditionalDefines,
@@ -3300,9 +3350,9 @@ void UNiagaraScript::CacheResourceShadersForRendering(bool bRegenerateId, bool b
 
 			ScriptResource->BuildScriptParametersMetadata(CachedScriptVM.ShaderScriptParametersMetadata);
 
-				CacheShadersForResources(ScriptResource.Get(), true);
-				ScriptResourcesByFeatureLevel[CacheFeatureLevel] = ScriptResource.Get();
-			}
+			CacheShadersForResources(ScriptResource.Get(), true);
+			ScriptResourcesByFeatureLevel[CacheFeatureLevel] = ScriptResource.Get();
+		}
 		else
 		{
 			ScriptResource->Invalidate();
diff --git a/Engine/Plugins/FX/Niagara/Source/NiagaraEditor/Private/NiagaraEditorModule.cpp b/Engine/Plugins/FX/Niagara/Source/NiagaraEditor/Private/NiagaraEditorModule.cpp
--- a/Engine/Plugins/FX/Niagara/Source/NiagaraEditor/Private/NiagaraEditorModule.cpp
+++ b/Engine/Plugins/FX/Niagara/Source/NiagaraEditor/Private/NiagaraEditorModule.cpp
@@ -816,6 +816,7 @@ void FNiagaraEditorModule::OnPreExit()
 	if (GEditor)
 	{
 		CastChecked<UEditorEngine>(GEngine)->OnPreviewPlatformChanged().Remove(PreviewPlatformChangedHandle);
+		CastChecked<UEditorEngine>(GEngine)->OnPreviewFeatureLevelChanged().Remove(PreviewFeatureLevelChangedHandle);
 	}
 
 	// Ensure that we don't have any lingering compiles laying around that will explode after this module shuts down.
@@ -1383,8 +1384,10 @@ void FNiagaraEditorModule::OnPostEngineInit()
 	{
 		GEditor->OnExecParticleInvoked().AddRaw(this, &FNiagaraEditorModule::OnExecParticleInvoked);
 
-		PreviewPlatformChangedHandle = CastChecked<UEditorEngine>(GEngine)->OnPreviewPlatformChanged().AddRaw(this, &FNiagaraEditorModule::OnPreviewPlatformChanged);
+		UEditorEngine* EditorEngine = CastChecked<UEditorEngine>(GEngine);
+		PreviewPlatformChangedHandle = EditorEngine->OnPreviewPlatformChanged().AddRaw(this, &FNiagaraEditorModule::OnPreviewPlatformChanged);
 
+		PreviewFeatureLevelChangedHandle = EditorEngine->OnPreviewFeatureLevelChanged().AddStatic(UNiagaraScript::SetPreviewFeatureLevel);
 
 		// Handle a re-import for mesh renderers
 		if (UImportSubsystem* ImportSubsystem = GEditor->GetEditorSubsystem<UImportSubsystem>())
diff --git a/Engine/Plugins/FX/Niagara/Source/NiagaraEditor/Public/NiagaraEditorModule.h b/Engine/Plugins/FX/Niagara/Source/NiagaraEditor/Public/NiagaraEditorModule.h
--- a/Engine/Plugins/FX/Niagara/Source/NiagaraEditor/Public/NiagaraEditorModule.h
+++ b/Engine/Plugins/FX/Niagara/Source/NiagaraEditor/Public/NiagaraEditorModule.h
@@ -364,6 +364,7 @@ private:
 	FDelegateHandle DeviceProfileManagerUpdatedHandle;
 
 	FDelegateHandle PreviewPlatformChangedHandle;
+	FDelegateHandle PreviewFeatureLevelChangedHandle;
 
 	USequencerSettings* SequencerSettings;
 
