- TickInfo exposed on NiagaraSystemSimulation to contain information about sub frame ticking (typically from fixed delta time
- Expose Engine.WorldDeltaTime, System.NumTimeSteps, System.CurrentTimeStep, System.TimeFraction to expose substeps to emitters

diff --git a/Engine/Plugins/FX/Niagara/Source/Niagara/Classes/NiagaraConstants.h b/Engine/Plugins/FX/Niagara/Source/Niagara/Classes/NiagaraConstants.h
--- a/Engine/Plugins/FX/Niagara/Source/Niagara/Classes/NiagaraConstants.h
+++ b/Engine/Plugins/FX/Niagara/Source/Niagara/Classes/NiagaraConstants.h
@@ -32,6 +32,7 @@
 #define TRANSLATOR_SET_VARIABLES_STR TEXT("SetVariables")
 #define TRANSLATOR_SET_VARIABLES_UNDERSCORE_STR TEXT("SetVariables_")
 
+#define SYS_PARAM_ENGINE_WORLD_DELTA_TIME                INiagaraModule::GetVar_Engine_WorldDeltaTime()
 #define SYS_PARAM_ENGINE_DELTA_TIME                      INiagaraModule::GetVar_Engine_DeltaTime()
 #define SYS_PARAM_ENGINE_INV_DELTA_TIME                  INiagaraModule::GetVar_Engine_InvDeltaTime()
 #define SYS_PARAM_ENGINE_TIME                            INiagaraModule::GetVar_Engine_Time()
@@ -70,6 +71,12 @@
 #define SYS_PARAM_ENGINE_SYSTEM_NUM_EMITTERS_ALIVE       INiagaraModule::GetVar_Engine_System_NumEmittersAlive()
 #define SYS_PARAM_ENGINE_SYSTEM_SIGNIFICANCE_INDEX		 INiagaraModule::GetVar_Engine_System_SignificanceIndex()
 #define SYS_PARAM_ENGINE_SYSTEM_RANDOM_SEED				INiagaraModule::GetVar_Engine_System_RandomSeed()
+
+#define SYS_PARAM_ENGINE_SYSTEM_CURRENT_TIME_STEP		 INiagaraModule::GetVar_Engine_System_CurrentTimeStep()
+#define SYS_PARAM_ENGINE_SYSTEM_NUM_TIME_STEPS			 INiagaraModule::GetVar_Engine_System_NumTimeSteps()
+#define SYS_PARAM_ENGINE_SYSTEM_TIME_STEP_FRACTION			 INiagaraModule::GetVar_Engine_System_TimeStepFraction()
+
+
 #define SYS_PARAM_ENGINE_SYSTEM_NUM_EMITTERS             INiagaraModule::GetVar_Engine_System_NumEmitters()
 #define SYS_PARAM_ENGINE_NUM_SYSTEM_INSTANCES            INiagaraModule::GetVar_Engine_NumSystemInstances()
 
diff --git a/Engine/Plugins/FX/Niagara/Source/Niagara/Private/NiagaraConstants.cpp b/Engine/Plugins/FX/Niagara/Source/Niagara/Private/NiagaraConstants.cpp
--- a/Engine/Plugins/FX/Niagara/Source/Niagara/Private/NiagaraConstants.cpp
+++ b/Engine/Plugins/FX/Niagara/Source/Niagara/Private/NiagaraConstants.cpp
@@ -75,6 +75,7 @@ void FNiagaraConstants::Init()
 {
 	if (SystemParameters.Num() == 0)
 	{
+		SystemParameters.Add(SYS_PARAM_ENGINE_WORLD_DELTA_TIME);
 		SystemParameters.Add(SYS_PARAM_ENGINE_DELTA_TIME);
 		SystemParameters.Add(SYS_PARAM_ENGINE_INV_DELTA_TIME);
 		SystemParameters.Add(SYS_PARAM_ENGINE_TIME);
@@ -112,6 +113,9 @@ void FNiagaraConstants::Init()
 		SystemParameters.Add(SYS_PARAM_ENGINE_SYSTEM_NUM_EMITTERS_ALIVE);
 		SystemParameters.Add(SYS_PARAM_ENGINE_SYSTEM_SIGNIFICANCE_INDEX);
 		SystemParameters.Add(SYS_PARAM_ENGINE_SYSTEM_RANDOM_SEED);
+		SystemParameters.Add(SYS_PARAM_ENGINE_SYSTEM_CURRENT_TIME_STEP);
+		SystemParameters.Add(SYS_PARAM_ENGINE_SYSTEM_NUM_TIME_STEPS);
+		SystemParameters.Add(SYS_PARAM_ENGINE_SYSTEM_TIME_STEP_FRACTION);
 		SystemParameters.Add(SYS_PARAM_ENGINE_SYSTEM_NUM_EMITTERS);
 		SystemParameters.Add(SYS_PARAM_ENGINE_NUM_SYSTEM_INSTANCES);
 		SystemParameters.Add(SYS_PARAM_ENGINE_GLOBAL_SPAWN_COUNT_SCALE);
@@ -159,6 +163,7 @@ void FNiagaraConstants::Init()
 
 	if (UpdatedSystemParameters.Num() == 0)
 	{
+		UpdatedSystemParameters.Add(FName(TEXT("System World Delta Time")), SYS_PARAM_ENGINE_WORLD_DELTA_TIME);
 		UpdatedSystemParameters.Add(FName(TEXT("System Delta Time")), SYS_PARAM_ENGINE_DELTA_TIME);
 		UpdatedSystemParameters.Add(FName(TEXT("System Inv Delta Time")), SYS_PARAM_ENGINE_INV_DELTA_TIME);
 		UpdatedSystemParameters.Add(FName(TEXT("System Position")), SYS_PARAM_ENGINE_POSITION);
@@ -181,6 +186,7 @@ void FNiagaraConstants::Init()
 		UpdatedSystemParameters.Add(FName(TEXT("Emitter Interp Spawn Start Dt")), SYS_PARAM_EMITTER_INTERP_SPAWN_START_DT);
 		UpdatedSystemParameters.Add(FName(TEXT("Emitter Spawn Group")), SYS_PARAM_EMITTER_SPAWN_GROUP);
 
+		UpdatedSystemParameters.Add(FName(TEXT("World Delta Time")), SYS_PARAM_ENGINE_WORLD_DELTA_TIME);
 		UpdatedSystemParameters.Add(FName(TEXT("Delta Time")), SYS_PARAM_ENGINE_DELTA_TIME);
 		UpdatedSystemParameters.Add(FName(TEXT("Emitter Age")), SYS_PARAM_EMITTER_AGE);
 		UpdatedSystemParameters.Add(FName(TEXT("Emitter Local Space")), SYS_PARAM_EMITTER_LOCALSPACE);
@@ -210,7 +216,8 @@ void FNiagaraConstants::Init()
 
 	if (SystemStrMap.Num() == 0)
 	{
-		SystemStrMap.Add(SYS_PARAM_ENGINE_DELTA_TIME, LOCTEXT("EngineDeltaTimeDesc", "Time in seconds since the last tick."));
+		SystemStrMap.Add(SYS_PARAM_ENGINE_WORLD_DELTA_TIME, LOCTEXT("EngineWorldDeltaTimeDesc", "Time in seconds since the last tick of the engine."));
+		SystemStrMap.Add(SYS_PARAM_ENGINE_DELTA_TIME, LOCTEXT("EngineDeltaTimeDesc", "Time in seconds since the last simulation tick."));
 		SystemStrMap.Add(SYS_PARAM_ENGINE_INV_DELTA_TIME, LOCTEXT("EngineInvDeltaTimeDesc", "One over Engine.DeltaTime"));
 		SystemStrMap.Add(SYS_PARAM_ENGINE_TIME, LOCTEXT("EngineTimeDesc", "Time in seconds since level began play, but IS paused when the game is paused, and IS dilated/clamped."));
 		SystemStrMap.Add(SYS_PARAM_ENGINE_REAL_TIME, LOCTEXT("EngineRealTimeDesc", "Time in seconds since level began play, but IS NOT paused when the game is paused, and IS NOT dilated/clamped."));
@@ -249,6 +256,10 @@ void FNiagaraConstants::Init()
 		SystemStrMap.Add(SYS_PARAM_ENGINE_SYSTEM_SIGNIFICANCE_INDEX, LOCTEXT("SystemSignificanceIndex", "Index denoting how significant this system instance is in relation to others of the same system in this scene. e.g. 0 is the most significanct instance."));
 		SystemStrMap.Add(SYS_PARAM_ENGINE_SYSTEM_RANDOM_SEED, LOCTEXT("SystemRandomSeed", "A random seed controlled used for generating system random numbers."));
 		
+		SystemStrMap.Add(SYS_PARAM_ENGINE_SYSTEM_CURRENT_TIME_STEP, LOCTEXT("CurrentTimeStep", "The current time step."));
+		SystemStrMap.Add(SYS_PARAM_ENGINE_SYSTEM_NUM_TIME_STEPS, LOCTEXT("NumTimeSteps", "Number of steps used, can be > 1 when using fixed delta time."));
+		SystemStrMap.Add(SYS_PARAM_ENGINE_SYSTEM_TIME_STEP_FRACTION, LOCTEXT("TimeStepFraction", "The fraction of the world delta time integrated for the current sub step."));
+
 		SystemStrMap.Add(SYS_PARAM_ENGINE_SYSTEM_NUM_EMITTERS, LOCTEXT("SystemNumEmitters", "The number of emitters attached to this system. Should only be used in System scripts."));
 		SystemStrMap.Add(SYS_PARAM_ENGINE_NUM_SYSTEM_INSTANCES, LOCTEXT("SystemNumInstances", "The number of instances of this system currently ticking. Should only be used in System scripts."));
 		SystemStrMap.Add(SYS_PARAM_ENGINE_GLOBAL_SPAWN_COUNT_SCALE, LOCTEXT("GlobalSpawnCountScale", "Global Spawn Count Scale. Should only be used in System scripts."));
diff --git a/Engine/Plugins/FX/Niagara/Source/Niagara/Private/NiagaraGPUSystemTick.cpp b/Engine/Plugins/FX/Niagara/Source/Niagara/Private/NiagaraGPUSystemTick.cpp
--- a/Engine/Plugins/FX/Niagara/Source/Niagara/Private/NiagaraGPUSystemTick.cpp
+++ b/Engine/Plugins/FX/Niagara/Source/Niagara/Private/NiagaraGPUSystemTick.cpp
@@ -92,6 +92,9 @@ void FNiagaraGPUSystemTick::Init(FNiagaraSystemInstance* InSystemInstance)
 		FMemory::Memcpy(OwnerParamData + sizeof(FNiagaraOwnerParameters), &InSystemInstance->GetOwnerParameters(true), sizeof(FNiagaraOwnerParameters));
 	}
 
+	FNiagaraSystemParameters *Tmp = reinterpret_cast<FNiagaraSystemParameters*>(SystemParamData);
+	FNiagaraSystemParameters Tmp2 = *Tmp;
+
 	ParamDataBufferPtr = OwnerParamData + InterpFactor * sizeof(FNiagaraOwnerParameters);
 
 	// Now we will generate instance data for every GPU simulation we want to run on the render thread.
diff --git a/Engine/Plugins/FX/Niagara/Source/Niagara/Private/NiagaraModule.cpp b/Engine/Plugins/FX/Niagara/Source/Niagara/Private/NiagaraModule.cpp
--- a/Engine/Plugins/FX/Niagara/Source/Niagara/Private/NiagaraModule.cpp
+++ b/Engine/Plugins/FX/Niagara/Source/Niagara/Private/NiagaraModule.cpp
@@ -104,6 +104,7 @@ namespace NiagaraDebugVisHelper
 	const FNiagaraTypeRegistry*& GTypeRegistrySingletonPtr = GCoreTypeRegistrySingletonPtr;
 }
 
+FNiagaraVariable INiagaraModule::Engine_WorldDeltaTime;
 FNiagaraVariable INiagaraModule::Engine_DeltaTime;
 FNiagaraVariable INiagaraModule::Engine_InvDeltaTime;
 FNiagaraVariable INiagaraModule::Engine_Time;
@@ -141,6 +142,10 @@ FNiagaraVariable INiagaraModule::Engine_System_TickCount;
 FNiagaraVariable INiagaraModule::Engine_System_NumEmittersAlive;
 FNiagaraVariable INiagaraModule::Engine_System_SignificanceIndex;
 FNiagaraVariable INiagaraModule::Engine_System_RandomSeed;
+FNiagaraVariable INiagaraModule::Engine_System_CurrentTimeStep;
+FNiagaraVariable INiagaraModule::Engine_System_NumTimeSteps;
+FNiagaraVariable INiagaraModule::Engine_System_TimeStepFraction;
+
 FNiagaraVariable INiagaraModule::Engine_System_NumEmitters;
 FNiagaraVariable INiagaraModule::Engine_NumSystemInstances;
 
@@ -216,7 +221,7 @@ void INiagaraModule::StartupModule()
 #endif
 	LLM_SCOPE(ELLMTag::Niagara);
 	FNiagaraTypeDefinition::Init();
-	
+
 #if PLATFORM_WINDOWS
 	// Global registration of  the vdb types.
 	openvdb::initialize();
@@ -248,6 +253,7 @@ void INiagaraModule::StartupModule()
 
 	//Init commonly used FNiagaraVariables
 
+	Engine_WorldDeltaTime = FNiagaraVariable(FNiagaraTypeDefinition::GetFloatDef(), TEXT("Engine.WorldDeltaTime"));
 	Engine_DeltaTime = FNiagaraVariable(FNiagaraTypeDefinition::GetFloatDef(), TEXT("Engine.DeltaTime"));
 	Engine_InvDeltaTime = FNiagaraVariable(FNiagaraTypeDefinition::GetFloatDef(), TEXT("Engine.InverseDeltaTime"));
 	
@@ -287,6 +293,11 @@ void INiagaraModule::StartupModule()
 	Engine_System_NumEmittersAlive = FNiagaraVariable(FNiagaraTypeDefinition::GetIntDef(), TEXT("Engine.System.NumEmittersAlive"));
 	Engine_System_SignificanceIndex = FNiagaraVariable(FNiagaraTypeDefinition::GetIntDef(), TEXT("Engine.System.SignificanceIndex"));
 	Engine_System_RandomSeed = FNiagaraVariable(FNiagaraTypeDefinition::GetIntDef(), TEXT("Engine.System.RandomSeed"));
+
+	Engine_System_CurrentTimeStep = FNiagaraVariable(FNiagaraTypeDefinition::GetIntDef(), TEXT("Engine.System.CurrentTimeStep"));
+	Engine_System_NumTimeSteps = FNiagaraVariable(FNiagaraTypeDefinition::GetIntDef(), TEXT("Engine.System.NumTimeSteps"));
+	Engine_System_TimeStepFraction = FNiagaraVariable(FNiagaraTypeDefinition::GetFloatDef(), TEXT("Engine.System.TimeStepFraction"));
+
 	Engine_System_NumEmitters = FNiagaraVariable(FNiagaraTypeDefinition::GetIntDef(), TEXT("Engine.System.NumEmitters"));
 	Engine_NumSystemInstances = FNiagaraVariable(FNiagaraTypeDefinition::GetIntDef(), TEXT("Engine.NumSystemInstances"));
 
@@ -1580,10 +1591,10 @@ const TArray<FNiagaraVariable>& FNiagaraGlobalParameters::GetVariables()
 {
 	static const FName NAME_NiagaraStructPadding0 = "Engine.PaddingInt32_0";
 	static const FName NAME_NiagaraStructPadding1 = "Engine.PaddingInt32_1";
-	static const FName NAME_NiagaraStructPadding2 = "Engine.PaddingInt32_2";
 
 	static const TArray<FNiagaraVariable> Variables =
 	{
+		SYS_PARAM_ENGINE_WORLD_DELTA_TIME,
 		SYS_PARAM_ENGINE_DELTA_TIME,
 		SYS_PARAM_ENGINE_INV_DELTA_TIME,
 		SYS_PARAM_ENGINE_TIME,
@@ -1591,7 +1602,6 @@ const TArray<FNiagaraVariable>& FNiagaraGlobalParameters::GetVariables()
 		SYS_PARAM_ENGINE_QUALITY_LEVEL,
 		FNiagaraVariable(FNiagaraTypeDefinition::GetIntDef(), NAME_NiagaraStructPadding0),
 		FNiagaraVariable(FNiagaraTypeDefinition::GetIntDef(), NAME_NiagaraStructPadding1),
-		FNiagaraVariable(FNiagaraTypeDefinition::GetIntDef(), NAME_NiagaraStructPadding2),
 	};
 
 	return Variables;
@@ -1601,6 +1611,7 @@ const TArray<FNiagaraVariable>& FNiagaraSystemParameters::GetVariables()
 {
 	static const FName NAME_NiagaraStructPadding0 = "Engine.System.PaddingInt32_0";
 	static const FName NAME_NiagaraStructPadding1 = "Engine.System.PaddingInt32_1";
+	static const FName NAME_NiagaraStructPadding2 = "Engine.System.PaddingInt32_2";
 
 	static const TArray<FNiagaraVariable> Variables =
 	{
@@ -1614,8 +1625,12 @@ const TArray<FNiagaraVariable>& FNiagaraSystemParameters::GetVariables()
 		SYS_PARAM_ENGINE_SYSTEM_NUM_EMITTERS_ALIVE,
 		SYS_PARAM_ENGINE_SYSTEM_SIGNIFICANCE_INDEX,
 		SYS_PARAM_ENGINE_SYSTEM_RANDOM_SEED,
+		SYS_PARAM_ENGINE_SYSTEM_CURRENT_TIME_STEP,
+		SYS_PARAM_ENGINE_SYSTEM_NUM_TIME_STEPS,
+		SYS_PARAM_ENGINE_SYSTEM_TIME_STEP_FRACTION,			
 		FNiagaraVariable(FNiagaraTypeDefinition::GetIntDef(), NAME_NiagaraStructPadding0),
 		FNiagaraVariable(FNiagaraTypeDefinition::GetIntDef(), NAME_NiagaraStructPadding1),
+		FNiagaraVariable(FNiagaraTypeDefinition::GetIntDef(), NAME_NiagaraStructPadding2),
 	};
 
 	return Variables;
diff --git a/Engine/Plugins/FX/Niagara/Source/Niagara/Private/NiagaraSystemInstance.cpp b/Engine/Plugins/FX/Niagara/Source/Niagara/Private/NiagaraSystemInstance.cpp
--- a/Engine/Plugins/FX/Niagara/Source/Niagara/Private/NiagaraSystemInstance.cpp
+++ b/Engine/Plugins/FX/Niagara/Source/Niagara/Private/NiagaraSystemInstance.cpp
@@ -1479,8 +1479,8 @@ void FNiagaraSystemInstance::InitDataInterfaces()
 	};
 
 	CalcInstDataSize(InstanceParameters, false, false, false);//This probably should be a proper exec context.
-	CalcInstDataSize(SystemSimulation->GetSpawnExecutionContext()->Parameters, true, false, true);
-	CalcInstDataSize(SystemSimulation->GetUpdateExecutionContext()->Parameters, true, false, true);
+		CalcInstDataSize(SystemSimulation->GetSpawnExecutionContext()->Parameters, true, false, true);
+		CalcInstDataSize(SystemSimulation->GetUpdateExecutionContext()->Parameters, true, false, true);
 
 	//Iterate over interfaces to get size for table and clear their interface bindings.
 	for (const TSharedRef<FNiagaraEmitterInstance, ESPMode::ThreadSafe>& Simulation : Emitters)
@@ -1679,8 +1679,8 @@ void FNiagaraSystemInstance::TickDataInterfaces(float DeltaSeconds, bool bPostSi
 				Complete(true);
 			}
 		}
+		}
 	}
-}
 
 float FNiagaraSystemInstance::GetLODDistance()
 {
@@ -1865,6 +1865,9 @@ void FNiagaraSystemInstance::TickInstanceParameters_GameThread(float DeltaSecond
 
 	GatheredInstanceParameters.EmitterCount = Emitters.Num();
 	GatheredInstanceParameters.DeltaSeconds = DeltaSeconds;
+	
+	GatheredInstanceParameters.WorldDeltaSeconds = GetSystemSimulation()->GetTickInfo().EngineTick;	
+	
 	GatheredInstanceParameters.NumAlive = 0;
 
 	//Bias the LastRenderTime slightly to account for any delay as it's written by the RT.
@@ -1907,6 +1910,10 @@ void FNiagaraSystemInstance::TickInstanceParameters_GameThread(float DeltaSecond
 	CurrentSystemParameters.SignificanceIndex = SignificanceIndex;
 	CurrentSystemParameters.RandomSeed = RandomSeed + RandomSeedOffset;
 
+	CurrentSystemParameters.CurrentTimeStep = GetSystemSimulation()->GetTickInfo().TickNumber;
+	CurrentSystemParameters.NumTimeSteps = GetSystemSimulation()->GetTickInfo().TickCount;
+	CurrentSystemParameters.TimeStepFraction = GetSystemSimulation()->GetTickInfo().TimeStepFraction;	
+
 	if (OverrideParameters)
 	{
 		OverrideParameters->ResolvePositions(GetLWCConverter());
@@ -1951,7 +1958,12 @@ void FNiagaraSystemInstance::TickInstanceParameters_Concurrent()
 	CurrentSystemParameters.SignificanceIndex = SignificanceIndex;
 	CurrentSystemParameters.RandomSeed = RandomSeed + RandomSeedOffset;
 
+	CurrentSystemParameters.CurrentTimeStep = GetSystemSimulation()->GetTickInfo().TickNumber;
+	CurrentSystemParameters.NumTimeSteps = GetSystemSimulation()->GetTickInfo().TickCount;
+	CurrentSystemParameters.TimeStepFraction = GetSystemSimulation()->GetTickInfo().TimeStepFraction;
+
 	FNiagaraGlobalParameters& CurrentGlobalParameter = GlobalParameters[ParameterIndex];
+	CurrentGlobalParameter.WorldDeltaTime = GatheredInstanceParameters.WorldDeltaSeconds;
 	CurrentGlobalParameter.EngineDeltaTime = GatheredInstanceParameters.DeltaSeconds;
 	CurrentGlobalParameter.EngineInvDeltaTime = GatheredInstanceParameters.DeltaSeconds > 0.0f ? 1.0f / GatheredInstanceParameters.DeltaSeconds : 0.0f;
 	CurrentGlobalParameter.EngineRealTime = GatheredInstanceParameters.RealTimeSeconds;
diff --git a/Engine/Plugins/FX/Niagara/Source/Niagara/Private/NiagaraSystemSimulation.cpp b/Engine/Plugins/FX/Niagara/Source/Niagara/Private/NiagaraSystemSimulation.cpp
--- a/Engine/Plugins/FX/Niagara/Source/Niagara/Private/NiagaraSystemSimulation.cpp
+++ b/Engine/Plugins/FX/Niagara/Source/Niagara/Private/NiagaraSystemSimulation.cpp
@@ -681,11 +681,11 @@ bool FNiagaraSystemSimulation::Init(UNiagaraSystem* InSystem, UWorld* InWorld, b
 
 		{
 			//SCOPE_CYCLE_COUNTER(STAT_NiagaraSystemSim_Init_ExecContexts);
-			SpawnExecContext = MakeUnique<FNiagaraSystemScriptExecutionContext>(ENiagaraSystemSimulationScript::Spawn);
-			UpdateExecContext = MakeUnique<FNiagaraSystemScriptExecutionContext>(ENiagaraSystemSimulationScript::Update);
-			bCanExecute &= SpawnExecContext->Init(SpawnScript, ENiagaraSimTarget::CPUSim);
-			bCanExecute &= UpdateExecContext->Init(UpdateScript, ENiagaraSimTarget::CPUSim);
-		}
+				SpawnExecContext = MakeUnique<FNiagaraSystemScriptExecutionContext>(ENiagaraSystemSimulationScript::Spawn);
+				UpdateExecContext = MakeUnique<FNiagaraSystemScriptExecutionContext>(ENiagaraSystemSimulationScript::Update);
+				bCanExecute &= SpawnExecContext->Init(SpawnScript, ENiagaraSimTarget::CPUSim);
+				bCanExecute &= UpdateExecContext->Init(UpdateScript, ENiagaraSimTarget::CPUSim);
+			}
 
 		{
 			//SCOPE_CYCLE_COUNTER(STAT_NiagaraSystemSim_Init_BindParams);
@@ -1120,8 +1120,24 @@ void FNiagaraSystemSimulation::Tick_GameThread(float DeltaSeconds, const FGraphE
 		float FixedDelta = System->GetFixedTickDeltaTime();
 		float Budget = FixedDelta > 0 ? FMath::Fmod(FixedDeltaTickAge, FixedDelta) + DeltaSeconds : 0;
 		int32 Ticks = FixedDelta > 0 ? FMath::Min(Budget / FixedDelta, GNiagaraSystemSimulationMaxTickSubsteps) : 0;
+
+		TickInfo.UsesFixedTick = true;
+		TickInfo.EngineTick = DeltaSeconds;
+		TickInfo.SystemTick = FixedDelta;
+		TickInfo.TickCount = Ticks;
+		TickInfo.TickNumber = -1;
+		TickInfo.TimeStepFraction = 0;
+		
 		for (int i = 0; i < Ticks; i++)
 		{
+			TickInfo.TickNumber = i;
+
+			// Fraction of the total number of ticks we are evaluating
+			// @note: I'd prefer if this were the fraction of the requested timestep, but
+			// the implementation of budget doesn't readily allow it since time fraction could be
+			// greater than 1
+			TickInfo.TimeStepFraction = 1.0f * (i + 1) / Ticks;
+
 			//Cannot do multiple tick off the game thread here without additional work. So we pass in null for the completion event which will force GT execution.
 			Tick_GameThread_Internal(FixedDelta, nullptr);
 			Budget -= FixedDelta;
@@ -1130,6 +1146,13 @@ void FNiagaraSystemSimulation::Tick_GameThread(float DeltaSeconds, const FGraphE
 	}
 	else
 	{
+		TickInfo.UsesFixedTick = false;
+		TickInfo.EngineTick = DeltaSeconds;
+		TickInfo.SystemTick = DeltaSeconds;
+		TickInfo.TickCount = 1;
+		TickInfo.TickNumber = 0;
+		TickInfo.TimeStepFraction = 1;
+
 		Tick_GameThread_Internal(DeltaSeconds, MyCompletionGraphEvent);
 	}
 }
@@ -1183,6 +1206,9 @@ void FNiagaraSystemSimulation::Tick_GameThread_Internal(float DeltaSeconds, cons
 	if (MaxDeltaTime.IsSet() && !System->HasFixedTickDelta())
 	{
 		DeltaSeconds = FMath::Clamp(DeltaSeconds, 0.0f, MaxDeltaTime.GetValue());
+
+		TickInfo.SystemTick = DeltaSeconds;
+		TickInfo.TimeStepFraction = DeltaSeconds / TickInfo.EngineTick;
 	}
 
 	UNiagaraScript* SystemSpawnScript = System->GetSystemSpawnScript();
diff --git a/Engine/Plugins/FX/Niagara/Source/Niagara/Public/NiagaraModule.h b/Engine/Plugins/FX/Niagara/Source/Niagara/Public/NiagaraModule.h
--- a/Engine/Plugins/FX/Niagara/Source/Niagara/Public/NiagaraModule.h
+++ b/Engine/Plugins/FX/Niagara/Source/Niagara/Public/NiagaraModule.h
@@ -124,6 +124,7 @@ public:
 	static float EngineGlobalSpawnCountScale;
 	static float EngineGlobalSystemCountScale;
 
+	FORCEINLINE static const FNiagaraVariable&  GetVar_Engine_WorldDeltaTime() { return Engine_WorldDeltaTime; }
 	FORCEINLINE static const FNiagaraVariable&  GetVar_Engine_DeltaTime() { return Engine_DeltaTime; }
 	FORCEINLINE static const FNiagaraVariable&  GetVar_Engine_InvDeltaTime() { return Engine_InvDeltaTime; }
 	FORCEINLINE static const FNiagaraVariable&  GetVar_Engine_Time() { return Engine_Time; }
@@ -162,6 +163,11 @@ public:
 	FORCEINLINE static const FNiagaraVariable&  GetVar_Engine_System_NumEmittersAlive() { return Engine_System_NumEmittersAlive; }
 	FORCEINLINE static const FNiagaraVariable&  GetVar_Engine_System_SignificanceIndex() { return Engine_System_SignificanceIndex; }
 	FORCEINLINE static const FNiagaraVariable&  GetVar_Engine_System_RandomSeed() { return Engine_System_RandomSeed; }
+
+	FORCEINLINE static const FNiagaraVariable& GetVar_Engine_System_CurrentTimeStep() { return Engine_System_CurrentTimeStep; }
+	FORCEINLINE static const FNiagaraVariable& GetVar_Engine_System_NumTimeSteps() { return Engine_System_NumTimeSteps; }
+	FORCEINLINE static const FNiagaraVariable& GetVar_Engine_System_TimeStepFraction() { return Engine_System_TimeStepFraction; }
+
 	FORCEINLINE static const FNiagaraVariable&  GetVar_Engine_System_NumEmitters() { return Engine_System_NumEmitters; }
 	FORCEINLINE static const FNiagaraVariable&  GetVar_Engine_NumSystemInstances() { return Engine_NumSystemInstances; }
 
@@ -247,6 +253,7 @@ public:
 
 
 private:
+	static FNiagaraVariable Engine_WorldDeltaTime;
 	static FNiagaraVariable Engine_DeltaTime;
 	static FNiagaraVariable Engine_InvDeltaTime;
 	static FNiagaraVariable Engine_Time; 
@@ -284,6 +291,11 @@ private:
 	static FNiagaraVariable Engine_System_NumEmittersAlive;
 	static FNiagaraVariable Engine_System_SignificanceIndex;
 	static FNiagaraVariable Engine_System_RandomSeed;
+	
+	static FNiagaraVariable Engine_System_CurrentTimeStep;
+	static FNiagaraVariable Engine_System_NumTimeSteps;
+	static FNiagaraVariable Engine_System_TimeStepFraction;
+	
 	static FNiagaraVariable Engine_System_NumEmitters;
 	static FNiagaraVariable Engine_NumSystemInstances;
 
diff --git a/Engine/Plugins/FX/Niagara/Source/Niagara/Public/NiagaraSystemInstance.h b/Engine/Plugins/FX/Niagara/Source/Niagara/Public/NiagaraSystemInstance.h
--- a/Engine/Plugins/FX/Niagara/Source/Niagara/Public/NiagaraSystemInstance.h
+++ b/Engine/Plugins/FX/Niagara/Source/Niagara/Public/NiagaraSystemInstance.h
@@ -612,6 +612,9 @@ public:
 		float TimeSeconds = 0.0f;
 		float RealTimeSeconds = 0.0f;
 
+		// delta time coming from the engine (not including fixed delta time, etc)
+		float WorldDeltaSeconds = 0.0f;
+
 		int32 EmitterCount = 0;
 		int32 NumAlive = 0;
 		int32 TransformMatchCount = 0;
@@ -625,6 +628,8 @@ public:
 			TimeSeconds = 0.0f;
 			RealTimeSeconds = 0.0f;
 
+			WorldDeltaSeconds = 0.0f;
+
 			EmitterCount = 0;
 			NumAlive = 0;
 			TransformMatchCount = 0;
diff --git a/Engine/Plugins/FX/Niagara/Source/Niagara/Public/NiagaraSystemSimulation.h b/Engine/Plugins/FX/Niagara/Source/Niagara/Public/NiagaraSystemSimulation.h
--- a/Engine/Plugins/FX/Niagara/Source/Niagara/Public/NiagaraSystemSimulation.h
+++ b/Engine/Plugins/FX/Niagara/Source/Niagara/Public/NiagaraSystemSimulation.h
@@ -27,6 +27,16 @@ enum class ENiagaraGPUTickHandlingMode
 	ConcurrentBatched,/** Systems can submit their GPU ticks in batches during concurrent tick. */
 };
 
+struct FTickInfo
+{
+	bool UsesFixedTick = false;
+	float EngineTick = 0.0;
+	float SystemTick = 0.0;
+	int32 TickCount = 0;
+	int32 TickNumber = 0;
+	float TimeStepFraction = 0.0;
+};
+
 //TODO: It would be good to have the batch size be variable per system to try to keep a good work/overhead ratio.
 //Can possibly adjust in future based on average batch execution time.
 #define NiagaraSystemTickBatchSize 4
@@ -317,6 +327,8 @@ public:
 
 	FORCEINLINE UWorld* GetWorld()const{return World;}
 
+	const FTickInfo& GetTickInfo() { return TickInfo; }
+
 protected:
 	void DumpStalledInfo();
 
@@ -434,4 +446,6 @@ protected:
 	static bool bUseLegacyExecContexts;
 
 	float FixedDeltaTickAge = 0;
+
+	FTickInfo TickInfo;
 };
diff --git a/Engine/Plugins/FX/Niagara/Source/Niagara/Public/NiagaraTypes.h b/Engine/Plugins/FX/Niagara/Source/Niagara/Public/NiagaraTypes.h
--- a/Engine/Plugins/FX/Niagara/Source/Niagara/Public/NiagaraTypes.h
+++ b/Engine/Plugins/FX/Niagara/Source/Niagara/Public/NiagaraTypes.h
@@ -1969,6 +1969,7 @@ struct alignas(16) FNiagaraGlobalParameters
 	NIAGARA_API static const TArray<FNiagaraVariable>& GetVariables();
 #endif
 
+	float WorldDeltaTime = 0.0f;
 	float EngineDeltaTime =  0.0f;
 	float EngineInvDeltaTime = 0.0f;
 	float EngineTime = 0.0f;
@@ -1977,7 +1978,6 @@ struct alignas(16) FNiagaraGlobalParameters
 
 	int32 _Pad0;
 	int32 _Pad1;
-	int32 _Pad2;
 };
 
 // Any change to this structure, or it's GetVariables implementation will require a bump in the CustomNiagaraVersion so that we
@@ -2000,6 +2000,11 @@ struct alignas(16) FNiagaraSystemParameters
 	int32 SignificanceIndex = 0;
 	int32 RandomSeed = 0;
 
+	int32 CurrentTimeStep = 0;
+	int32 NumTimeSteps = 0;
+	float TimeStepFraction = 0.0f;
+
+	int32 _Pad0;
 	int32 _Pad1;
 	int32 _Pad2;
 };
@@ -2119,7 +2124,7 @@ struct NIAGARA_API FVersionedNiagaraEmitterWeakPtr
 	FGuid Version;
 
 	friend inline uint32 GetTypeHash(const FVersionedNiagaraEmitterWeakPtr& Item)
-	{
-		return GetTypeHash(Item.Emitter);
-	}
+{
+	return GetTypeHash(Item.Emitter);
+}
 };
diff --git a/Engine/Plugins/FX/Niagara/Source/NiagaraEditor/Private/NiagaraHlslTranslator.cpp b/Engine/Plugins/FX/Niagara/Source/NiagaraEditor/Private/NiagaraHlslTranslator.cpp
--- a/Engine/Plugins/FX/Niagara/Source/NiagaraEditor/Private/NiagaraHlslTranslator.cpp
+++ b/Engine/Plugins/FX/Niagara/Source/NiagaraEditor/Private/NiagaraHlslTranslator.cpp
@@ -1069,6 +1069,9 @@ static void ConvertFloatToHalf(const FNiagaraCompileOptions& InCompileOptions, T
 			SYS_PARAM_ENGINE_SYSTEM_NUM_EMITTERS_ALIVE,
 			SYS_PARAM_ENGINE_SYSTEM_NUM_EMITTERS,
 			SYS_PARAM_ENGINE_SYSTEM_RANDOM_SEED,
+			SYS_PARAM_ENGINE_SYSTEM_NUM_TIME_STEPS,
+			SYS_PARAM_ENGINE_SYSTEM_TIME_STEP_FRACTION,
+			SYS_PARAM_ENGINE_SYSTEM_NUM_EMITTERS,
 			SYS_PARAM_ENGINE_NUM_SYSTEM_INSTANCES,
 			SYS_PARAM_ENGINE_EMITTER_NUM_PARTICLES,
 			SYS_PARAM_ENGINE_EMITTER_SIMULATION_POSITION,
@@ -1722,6 +1725,7 @@ const FNiagaraTranslateResults &FHlslNiagaraTranslator::Translate(const FNiagara
 				int32 OutputIdx = 0;
 				//ensure the interpolated spawn constants are part of the parameter set.
 				ParameterMapRegisterExternalConstantNamespaceVariable(SYS_PARAM_ENGINE_TIME, nullptr, 0, OutputIdx, nullptr);
+				ParameterMapRegisterExternalConstantNamespaceVariable(SYS_PARAM_ENGINE_WORLD_DELTA_TIME, nullptr, 0, OutputIdx, nullptr);
 				ParameterMapRegisterExternalConstantNamespaceVariable(SYS_PARAM_ENGINE_DELTA_TIME, nullptr, 0, OutputIdx, nullptr);
 				ParameterMapRegisterExternalConstantNamespaceVariable(SYS_PARAM_ENGINE_INV_DELTA_TIME, nullptr, 0, OutputIdx, nullptr);
 				ParameterMapRegisterExternalConstantNamespaceVariable(SYS_PARAM_ENGINE_EXEC_COUNT, nullptr, 0, OutputIdx, nullptr);
@@ -3347,8 +3351,8 @@ void FHlslNiagaraTranslator::DefineMainGPUFunctions(
 			"	#if NIAGARA_DISPATCH_TYPE >= NIAGARA_DISPATCH_TYPE_THREE_D\n"
 			"		GLinearThreadId += GDispatchThreadId.z * DispatchThreadIdToLinear.z;\n"
 			"	#endif\n"
-			"	GEmitterTickCounter	= EmitterTickCounter;\n"
-			"	GRandomSeedOffset	= 0;\n"
+			"	GEmitterTickCounter = EmitterTickCounter;\n"
+			"	GRandomSeedOffset = 0;\n"
 			"\n"
 		);
 
@@ -4417,7 +4421,8 @@ bool FHlslNiagaraTranslator::ShouldInterpolateParameter(const FNiagaraVariable&
 	}
 
 	//Skip interpolation for some system constants.
-	if (Parameter == SYS_PARAM_ENGINE_DELTA_TIME ||
+	if (Parameter == SYS_PARAM_ENGINE_WORLD_DELTA_TIME ||
+		Parameter == SYS_PARAM_ENGINE_DELTA_TIME ||
 		Parameter == SYS_PARAM_ENGINE_INV_DELTA_TIME ||
 		Parameter == SYS_PARAM_ENGINE_EXEC_COUNT ||
 		Parameter == SYS_PARAM_EMITTER_SPAWNRATE ||
@@ -4426,6 +4431,9 @@ bool FHlslNiagaraTranslator::ShouldInterpolateParameter(const FNiagaraVariable&
 		Parameter == SYS_PARAM_ENGINE_EMITTER_TOTAL_SPAWNED_PARTICLES ||
 		Parameter == SYS_PARAM_ENGINE_EMITTER_SPAWN_COUNT_SCALE ||
 		Parameter == SYS_PARAM_EMITTER_RANDOM_SEED ||
+		Parameter == SYS_PARAM_ENGINE_SYSTEM_NUM_TIME_STEPS ||
+		Parameter == SYS_PARAM_ENGINE_SYSTEM_TIME_STEP_FRACTION ||
+		Parameter == SYS_PARAM_ENGINE_SYSTEM_NUM_EMITTERS ||
 		Parameter == SYS_PARAM_ENGINE_EMITTER_INSTANCE_SEED ||
 		Parameter == SYS_PARAM_ENGINE_SYSTEM_TICK_COUNT ||
 		Parameter == SYS_PARAM_ENGINE_SYSTEM_RANDOM_SEED)
diff --git a/Engine/Plugins/FX/Niagara/Source/NiagaraShader/Public/NiagaraShader.h b/Engine/Plugins/FX/Niagara/Source/NiagaraShader/Public/NiagaraShader.h
--- a/Engine/Plugins/FX/Niagara/Source/NiagaraShader/Public/NiagaraShader.h
+++ b/Engine/Plugins/FX/Niagara/Source/NiagaraShader/Public/NiagaraShader.h
@@ -32,6 +32,7 @@ public:
 
 	// This structure is a replication of FNiagaraGlobalParameters with interpolated parameters includes
 	BEGIN_SHADER_PARAMETER_STRUCT(FGlobalParameters, )
+		SHADER_PARAMETER(float,		Engine_WorldDeltaTime)
 		SHADER_PARAMETER(float,		Engine_DeltaTime)
 		SHADER_PARAMETER(float,		Engine_InverseDeltaTime)
 		SHADER_PARAMETER(float,		Engine_Time)
@@ -39,8 +40,8 @@ public:
 		SHADER_PARAMETER(int32,		Engine_QualityLevel)
 		SHADER_PARAMETER(int32,		Engine_Pad0)
 		SHADER_PARAMETER(int32,		Engine_Pad1)
-		SHADER_PARAMETER(int32,		Engine_Pad2)
 
+		SHADER_PARAMETER(float,		PREV_Engine_WorldDeltaTime)
 		SHADER_PARAMETER(float,		PREV_Engine_DeltaTime)
 		SHADER_PARAMETER(float,		PREV_Engine_InverseDeltaTime)
 		SHADER_PARAMETER(float,		PREV_Engine_Time)
@@ -48,7 +49,6 @@ public:
 		SHADER_PARAMETER(int32,		PREV_Engine_QualityLevel)
 		SHADER_PARAMETER(int32,		PREV_Engine_Pad0)
 		SHADER_PARAMETER(int32,		PREV_Engine_Pad1)
-		SHADER_PARAMETER(int32,		PREV_Engine_Pad2)
 	END_SHADER_PARAMETER_STRUCT()
 
 	// This structure is a replication of FNiagaraSystemParameters with interpolated parameters includes
@@ -63,8 +63,12 @@ public:
 		SHADER_PARAMETER(int32,		Engine_System_NumEmittersAlive)
 		SHADER_PARAMETER(int32,		Engine_System_SignificanceIndex)
 		SHADER_PARAMETER(int32,		Engine_System_RandomSeed)
+		SHADER_PARAMETER(int32,		Engine_System_CurrentTimeStep)
+		SHADER_PARAMETER(int32,		Engine_System_NumTimeSteps)
+		SHADER_PARAMETER(float,		Engine_System_TimeStepFraction)
 		SHADER_PARAMETER(int32,		System_Pad0)
 		SHADER_PARAMETER(int32,		System_Pad1)
+		SHADER_PARAMETER(int32,		System_Pad2)
 
 		SHADER_PARAMETER(float,		PREV_Engine_Owner_TimeSinceRendered)
 		SHADER_PARAMETER(float,		PREV_Engine_Owner_LODDistance)
@@ -76,8 +80,12 @@ public:
 		SHADER_PARAMETER(int32,		PREV_Engine_System_NumEmittersAlive)
 		SHADER_PARAMETER(int32,		PREV_Engine_System_SignificanceIndex)
 		SHADER_PARAMETER(int32,		PREV_Engine_System_RandomSeed)
+		SHADER_PARAMETER(int32,		PREV_CurrentTimeStep)
+		SHADER_PARAMETER(int32,		PREV_NumTimeSteps)
+		SHADER_PARAMETER(float,		PREV_TimeStepFraction)
 		SHADER_PARAMETER(int32,		PREV_System_Pad0)
 		SHADER_PARAMETER(int32,		PREV_System_Pad1)
+		SHADER_PARAMETER(int32,		PREV_System_Pad2)
 	END_SHADER_PARAMETER_STRUCT()
 
 	// This structure is a replication of FNiagaraOwnerParameters with interpolated parameters includes
diff --git a/Engine/Source/Runtime/Core/Private/UObject/UE5MainStreamObjectVersion.cpp b/Engine/Source/Runtime/Core/Private/UObject/UE5MainStreamObjectVersion.cpp
--- a/Engine/Source/Runtime/Core/Private/UObject/UE5MainStreamObjectVersion.cpp
+++ b/Engine/Source/Runtime/Core/Private/UObject/UE5MainStreamObjectVersion.cpp
@@ -14,7 +14,7 @@ TMap<FGuid, FGuid> FUE5MainStreamObjectVersion::GetSystemGuids()
 	SystemGuids.Add(DevGuids.MATERIALSHADERMAP_DERIVEDDATA_VER, FGuid("65C01D817C9A4EEFAE9E988D41A1F3DD"));
 	SystemGuids.Add(DevGuids.NANITE_DERIVEDDATA_VER, FGuid("7A3EA93546624597BEBE2ADFDE3A7628"));
 	SystemGuids.Add(DevGuids.NIAGARASHADERMAP_DERIVEDDATA_VER, FGuid("8A37C45D24F2423CBE5F8F371DE33575"));
-	SystemGuids.Add(DevGuids.Niagara_LatestScriptCompileVersion, FGuid("FFA132206A19A641AA7A972631B17232"));
+	SystemGuids.Add(DevGuids.Niagara_LatestScriptCompileVersion, FGuid("BD1540BCB37D4B88A2EE4FF5ACFBB803"));
 	SystemGuids.Add(DevGuids.SkeletalMeshDerivedDataVersion, FGuid("4084BE94509A8C3FA7CF410DDF8ABF04"));
 	SystemGuids.Add(DevGuids.STATICMESH_DERIVEDDATA_VER, FGuid("98F7D79A5811013825E75BBCC41ED3E9"));
 	SystemGuids.Add(DevGuids.POSESEARCHDB_DERIVEDDATA_VER, FGuid("4E595C2AC5E947D6BA9ABC874353E5BC"));
