- TickInfo exposed on NiagaraSystemSimulation to contain information about sub frame ticking (typically from fixed delta time
- Expose Engine.WorldDeltaTime, System.NumTimeSteps, System.CurrentTimeStep, System.TimeFraction to expose substeps to emitters

diff --git a/Engine/Plugins/FX/Niagara/Source/Niagara/Private/NiagaraGPUSystemTick.cpp b/Engine/Plugins/FX/Niagara/Source/Niagara/Private/NiagaraGPUSystemTick.cpp
--- a/Engine/Plugins/FX/Niagara/Source/Niagara/Private/NiagaraGPUSystemTick.cpp
+++ b/Engine/Plugins/FX/Niagara/Source/Niagara/Private/NiagaraGPUSystemTick.cpp
@@ -92,9 +92,6 @@ void FNiagaraGPUSystemTick::Init(FNiagaraSystemInstance* InSystemInstance)
 		FMemory::Memcpy(OwnerParamData + sizeof(FNiagaraOwnerParameters), &InSystemInstance->GetOwnerParameters(true), sizeof(FNiagaraOwnerParameters));
 	}
 
-	FNiagaraSystemParameters *Tmp = reinterpret_cast<FNiagaraSystemParameters*>(SystemParamData);
-	FNiagaraSystemParameters Tmp2 = *Tmp;
-
 	ParamDataBufferPtr = OwnerParamData + InterpFactor * sizeof(FNiagaraOwnerParameters);
 
 	// Now we will generate instance data for every GPU simulation we want to run on the render thread.
