- Integrate TickInfo into Rigid Mesh data interface in order to limit ticking to only first one when substepping
- Integrate substepping collisions for shallow water, and all other fluids that use rigid mesh collisions

diff --git a/Engine/Plugins/FX/Niagara/Shaders/Private/NiagaraDataInterfaceRigidMeshCollisionQuery.ush b/Engine/Plugins/FX/Niagara/Shaders/Private/NiagaraDataInterfaceRigidMeshCollisionQuery.ush
--- a/Engine/Plugins/FX/Niagara/Shaders/Private/NiagaraDataInterfaceRigidMeshCollisionQuery.ush
+++ b/Engine/Plugins/FX/Niagara/Shaders/Private/NiagaraDataInterfaceRigidMeshCollisionQuery.ush
@@ -19,6 +19,7 @@ uint			{ParameterName}_PreviousOffset;
 uint4			{ParameterName}_ElementOffsets;
 Buffer<float4>	{ParameterName}_WorldTransformBuffer;
 Buffer<float4>	{ParameterName}_InverseTransformBuffer;
+Buffer<float4>	{ParameterName}_MeshScaleBuffer;
 Buffer<float4>	{ParameterName}_ElementExtentBuffer;
 Buffer<uint>	{ParameterName}_PhysicsTypeBuffer;
 Buffer<uint>	{ParameterName}_DFIndexBuffer;
@@ -73,6 +74,10 @@ float3x4 GetPreviousInverse_{ParameterName}(in int ElementIndex)
     return float3x4({ParameterName}_InverseTransformBuffer[BufferOffset], {ParameterName}_InverseTransformBuffer[BufferOffset+1], {ParameterName}_InverseTransformBuffer[BufferOffset+2]);
 }
 
+float3 GetMeshScale_{ParameterName}(in int ElementIndex)
+{
+	return {ParameterName}_MeshScaleBuffer[ElementIndex].xyz;
+}
 
 /* -----------------------------------------------------------------
  * Collision detection utils
@@ -476,9 +481,7 @@ void GetClosestPointMeshDistanceFieldAccurate_{ParameterName}(in float3 WorldPos
 
 // Given a world space position (WorldPosition) compute the static mesh closest point (position,normal,velocity)
 void GetClosestPointMeshDistanceFieldNoNormal_{ParameterName}(in float3 WorldPosition, in float DeltaTime, in float TimeFraction, in float MaxDistance, out float OutClosestDistance, out float3 OutClosestPosition, out float3 OutClosestVelocity)
-{
-	const FLWCVector3 LWCWorldPosition = MakeLWCVector3({ParameterName}_SystemLWCTile, WorldPosition);
-
+{	
 	float3 CollisionPosition = float3(0,0,0);
 	float3 CollisionNormal = float3(0,0,0);
 	OutClosestDistance = MAX_DISTANCE;	
@@ -496,7 +499,34 @@ void GetClosestPointMeshDistanceFieldNoNormal_{ParameterName}(in float3 WorldPos
 		{
 			// #todo(dmp): this does a dynamic branch based on intersecting the bbox.  Maybe we can factor that out due to the broadphase here?
 			MaxDistance = max(.001, MaxDistance);
-			OutClosestDistance = DistanceToNearestSurfaceForObject(DFIndex, LWCWorldPosition, MaxDistance);
+
+			FDFObjectData DFObjectData = LoadDFObjectData(DFIndex);
+
+			float3 WorldExtent = {ParameterName}_ElementExtentBuffer[ElementIndex].xyz;
+
+			float3 LocalPos = GetLocalPosition_{ParameterName}(WorldPosition, ElementIndex, TimeFraction);
+			LocalPos *= GetMeshScale_{ParameterName}(ElementIndex);
+			
+			LocalPos /= DFObjectData.VolumeToWorldScale;
+							
+			float3 VolumePosition = LocalPos;
+
+			float3 ClampedVolumePosition = clamp(VolumePosition, -DFObjectData.VolumePositionExtent, DFObjectData.VolumePositionExtent);
+			float DistanceToBox = length((VolumePosition - ClampedVolumePosition) * DFObjectData.VolumeToWorldScale);
+
+			OutClosestDistance = MaxDistance;
+
+			BRANCH
+			if (DistanceToBox < MaxDistance)
+			{
+				OutClosestDistance = DistanceToMeshSurfaceStandalone(ClampedVolumePosition, DFObjectData) * DFObjectData.VolumeScale + DistanceToBox;
+
+				// Don't allow negative SDF distance to stick out of the bounding box
+				if (DistanceToBox > 0.0f)
+				{
+					OutClosestDistance = max(OutClosestDistance, DistanceToBox);
+				}
+			}		
 		}
 		
 		const float3 PreviousPosition = mul(GetPreviousTransform_{ParameterName}(ElementIndex), float4(CollisionPosition,1.0)).xyz;
diff --git a/Engine/Plugins/FX/Niagara/Source/Niagara/Private/NiagaraDataInterfaceRigidMeshCollisionQuery.cpp b/Engine/Plugins/FX/Niagara/Source/Niagara/Private/NiagaraDataInterfaceRigidMeshCollisionQuery.cpp
--- a/Engine/Plugins/FX/Niagara/Source/Niagara/Private/NiagaraDataInterfaceRigidMeshCollisionQuery.cpp
+++ b/Engine/Plugins/FX/Niagara/Source/Niagara/Private/NiagaraDataInterfaceRigidMeshCollisionQuery.cpp
@@ -64,6 +64,7 @@ BEGIN_SHADER_PARAMETER_STRUCT(FShaderParameters,)
 	SHADER_PARAMETER_SRV(Buffer<float4>,	WorldTransformBuffer)
 	SHADER_PARAMETER_SRV(Buffer<float4>,	InverseTransformBuffer)
 	SHADER_PARAMETER_SRV(Buffer<float4>,	ElementExtentBuffer)
+	SHADER_PARAMETER_SRV(Buffer<float4>,	MeshScaleBuffer)
 	SHADER_PARAMETER_SRV(Buffer<uint32>,	PhysicsTypeBuffer)
 	SHADER_PARAMETER_SRV(Buffer<uint32>,	DFIndexBuffer)
 	SHADER_PARAMETER(FVector3f,				SystemLWCTile)
@@ -358,6 +359,8 @@ void UpdateAssetArrays(TConstArrayView<TComponentType*> Components, const FVecto
 		FTransform MeshTransform = CreateElementTransform(Component, BodySetup);
 		MeshTransform.AddToTranslation(LWCTile * -FLargeWorldRenderScalar::GetTileSize());
 
+		auto CurrMeshScale = MeshTransform.GetScale3D();
+
 		const int32 ComponentIdIndex = OutAssetArrays->UniqueCompnentId.AddUnique(Component->ComponentId);
 
 		for (const FKConvexElem& ConvexElem : BodySetup->AggGeom.ConvexElems)
@@ -369,7 +372,8 @@ void UpdateAssetArrays(TConstArrayView<TComponentType*> Components, const FVecto
 				if (InitializeStatics)
 				{
 					FVector3f Extent = FVector3f(BBox.Max - BBox.Min);
-					OutAssetArrays->ElementExtent[BoxIndex] = FVector4f(Extent.X, Extent.Y, Extent.Z, 0);
+					OutAssetArrays->ElementExtent[BoxIndex] = FVector4f(Extent.X, Extent.Y, Extent.Z, 0);					
+					OutAssetArrays->MeshScale[BoxIndex] = FVector4f(CurrMeshScale.X, CurrMeshScale.Y, CurrMeshScale.Z, 0);
 					OutAssetArrays->PhysicsType[BoxIndex] = (ConvexElem.GetCollisionEnabled() == ECollisionEnabled::QueryAndPhysics);
 					OutAssetArrays->ComponentIdIndex[BoxIndex] = ComponentIdIndex;
 				}
@@ -387,6 +391,7 @@ void UpdateAssetArrays(TConstArrayView<TComponentType*> Components, const FVecto
 				if (InitializeStatics)
 				{
 					OutAssetArrays->ElementExtent[BoxIndex] = FVector4f(BoxElem.X, BoxElem.Y, BoxElem.Z, 0);
+					OutAssetArrays->MeshScale[BoxIndex] = FVector4f(CurrMeshScale.X, CurrMeshScale.Y, CurrMeshScale.Z, 0);
 					OutAssetArrays->PhysicsType[BoxIndex] = (BoxElem.GetCollisionEnabled() == ECollisionEnabled::QueryAndPhysics);
 					OutAssetArrays->ComponentIdIndex[BoxIndex] = ComponentIdIndex;
 				}
@@ -404,6 +409,7 @@ void UpdateAssetArrays(TConstArrayView<TComponentType*> Components, const FVecto
 				if (InitializeStatics)
 				{
 					OutAssetArrays->ElementExtent[SphereIndex] = FVector4f(SphereElem.Radius, 0, 0, 0);
+					OutAssetArrays->MeshScale[BoxIndex] = FVector4f(CurrMeshScale.X, CurrMeshScale.Y, CurrMeshScale.Z, 0);
 					OutAssetArrays->PhysicsType[SphereIndex] = (SphereElem.GetCollisionEnabled() == ECollisionEnabled::QueryAndPhysics);
 					OutAssetArrays->ComponentIdIndex[SphereIndex] = ComponentIdIndex;
 				}
@@ -421,6 +427,7 @@ void UpdateAssetArrays(TConstArrayView<TComponentType*> Components, const FVecto
 				if (InitializeStatics)
 				{
 					OutAssetArrays->ElementExtent[CapsuleIndex] = FVector4f(CapsuleElem.Radius, CapsuleElem.Length, 0, 0);
+					OutAssetArrays->MeshScale[BoxIndex] = FVector4f(CurrMeshScale.X, CurrMeshScale.Y, CurrMeshScale.Z, 0);
 					OutAssetArrays->PhysicsType[CapsuleIndex] = (CapsuleElem.GetCollisionEnabled() == ECollisionEnabled::QueryAndPhysics);
 					OutAssetArrays->ComponentIdIndex[CapsuleIndex] = ComponentIdIndex;
 				}
@@ -450,6 +457,7 @@ void UpdateAssetArrays(TConstArrayView<TComponentType*> Components, const FVecto
 				Extent *= 2;
 
 				OutAssetArrays->ElementExtent[BoxIndex] = FVector4f(Extent.X, Extent.Y, Extent.Z, 0);
+				OutAssetArrays->MeshScale[BoxIndex] = FVector4f(CurrMeshScale.X, CurrMeshScale.Y, CurrMeshScale.Z, 0);
 				OutAssetArrays->PhysicsType[BoxIndex] = true;
 				OutAssetArrays->ComponentIdIndex[BoxIndex] = ComponentIdIndex;
 			}
@@ -544,6 +552,7 @@ void FNDIRigidMeshCollisionBuffer::InitRHI()
 	CreateInternalBuffer<FVector4f, EPixelFormat::PF_A32B32G32R32F>(InverseTransformBuffer, 3 * MaxNumTransforms);
 
 	CreateInternalBuffer<FVector4f, EPixelFormat::PF_A32B32G32R32F>(ElementExtentBuffer, MaxNumPrimitives);
+	CreateInternalBuffer<FVector4f, EPixelFormat::PF_A32B32G32R32F>(MeshScaleBuffer, MaxNumPrimitives);
 	CreateInternalBuffer<uint32, EPixelFormat::PF_R32_UINT>(PhysicsTypeBuffer, MaxNumPrimitives);
 	CreateInternalBuffer<uint32, EPixelFormat::PF_R32_UINT>(DFIndexBuffer, MaxNumPrimitives);
 }
@@ -698,6 +707,7 @@ struct FNDIRigidMeshCollisionProxy : public FNiagaraDataInterfaceProxy
 		TArray<FVector4f> WorldTransform;
 		TArray<FVector4f> InverseTransform;
 		TArray<FVector4f> ElementExtent;
+		TArray<FVector4f> MeshScale;
 		TArray<uint32> PhysicsType;
 		TArray<int32> ComponentIdIndex;
 		uint32 MaxPrimitiveCount;
@@ -743,6 +753,7 @@ struct FNDIRigidMeshCollisionProxy : public FNiagaraDataInterfaceProxy
 				UpdateInternalBuffer<FVector4f, EPixelFormat::PF_A32B32G32R32F>(ProxyData->WorldTransform, ProxyData->AssetBuffer->WorldTransformBuffer);
 				UpdateInternalBuffer<FVector4f, EPixelFormat::PF_A32B32G32R32F>(ProxyData->InverseTransform, ProxyData->AssetBuffer->InverseTransformBuffer);
 				UpdateInternalBuffer<FVector4f, EPixelFormat::PF_A32B32G32R32F>(ProxyData->ElementExtent, ProxyData->AssetBuffer->ElementExtentBuffer);
+				UpdateInternalBuffer<FVector4f, EPixelFormat::PF_A32B32G32R32F>(ProxyData->MeshScale, ProxyData->AssetBuffer->MeshScaleBuffer);
 				UpdateInternalBuffer<uint32, EPixelFormat::PF_R32_UINT>(ProxyData->PhysicsType, ProxyData->AssetBuffer->PhysicsTypeBuffer);
 
 				// the distance field indexing needs to be generated using the scene
@@ -1026,8 +1037,14 @@ bool UNiagaraDataInterfaceRigidMeshCollisionQuery::PerInstanceTick(void* PerInst
 	FNDIRigidMeshCollisionData* InstanceData = static_cast<FNDIRigidMeshCollisionData*>(PerInstanceData);
 	if (InstanceData && SystemInstance)
 	{
-		check(InstanceData->SystemInstance == SystemInstance);
-		InstanceData->Update(this);
+		const FTickInfo &TickInfo = SystemInstance->GetSystemSimulation()->GetTickInfo();
+
+		// only update on the first tick
+		if (TickInfo.TickNumber == 0)
+		{
+			check(InstanceData->SystemInstance == SystemInstance);
+			InstanceData->Update(this);
+		}
 	}
 	return false;
 }
@@ -1491,6 +1508,7 @@ void UNiagaraDataInterfaceRigidMeshCollisionQuery::SetShaderParameters(const FNi
 		ShaderParameters->WorldTransformBuffer = FNiagaraRenderer::GetSrvOrDefaultFloat4(AssetBuffer->WorldTransformBuffer.SRV);
 		ShaderParameters->InverseTransformBuffer = FNiagaraRenderer::GetSrvOrDefaultFloat4(AssetBuffer->InverseTransformBuffer.SRV);
 		ShaderParameters->ElementExtentBuffer = FNiagaraRenderer::GetSrvOrDefaultFloat4(AssetBuffer->ElementExtentBuffer.SRV);
+		ShaderParameters->MeshScaleBuffer = FNiagaraRenderer::GetSrvOrDefaultFloat4(AssetBuffer->MeshScaleBuffer.SRV);
 		ShaderParameters->PhysicsTypeBuffer = FNiagaraRenderer::GetSrvOrDefaultUInt(AssetBuffer->PhysicsTypeBuffer.SRV);
 		ShaderParameters->DFIndexBuffer = FNiagaraRenderer::GetSrvOrDefaultUInt(AssetBuffer->DFIndexBuffer.SRV);
 
@@ -1517,6 +1535,7 @@ void UNiagaraDataInterfaceRigidMeshCollisionQuery::SetShaderParameters(const FNi
 		ShaderParameters->WorldTransformBuffer = FNiagaraRenderer::GetDummyFloat4Buffer();
 		ShaderParameters->InverseTransformBuffer = FNiagaraRenderer::GetDummyFloat4Buffer();
 		ShaderParameters->ElementExtentBuffer = FNiagaraRenderer::GetDummyFloat4Buffer();
+		ShaderParameters->MeshScaleBuffer = FNiagaraRenderer::GetDummyFloat4Buffer();
 		ShaderParameters->PhysicsTypeBuffer = FNiagaraRenderer::GetDummyUIntBuffer();
 		ShaderParameters->DFIndexBuffer = FNiagaraRenderer::GetDummyUIntBuffer();
 
@@ -1579,6 +1598,7 @@ void UNiagaraDataInterfaceRigidMeshCollisionQuery::ProvidePerInstanceDataForRend
 			CompactTransforms(GameThreadData->AssetArrays->CurrentInverse, GameThreadData->AssetArrays->PreviousInverse, RenderThreadData->InverseTransform);
 
 			RenderThreadData->ElementExtent.Append(GameThreadData->AssetArrays->ElementExtent.GetData(), ElementCount);
+			RenderThreadData->MeshScale.Append(GameThreadData->AssetArrays->MeshScale.GetData(), ElementCount);
 			RenderThreadData->PhysicsType.Append(GameThreadData->AssetArrays->PhysicsType.GetData(), ElementCount);
 			RenderThreadData->ComponentIdIndex.Append(GameThreadData->AssetArrays->ComponentIdIndex.GetData(), ElementCount);
 
diff --git a/Engine/Plugins/FX/Niagara/Source/Niagara/Public/NiagaraDataInterfaceRigidMeshCollisionQuery.h b/Engine/Plugins/FX/Niagara/Source/Niagara/Public/NiagaraDataInterfaceRigidMeshCollisionQuery.h
--- a/Engine/Plugins/FX/Niagara/Source/Niagara/Public/NiagaraDataInterfaceRigidMeshCollisionQuery.h
+++ b/Engine/Plugins/FX/Niagara/Source/Niagara/Public/NiagaraDataInterfaceRigidMeshCollisionQuery.h
@@ -36,6 +36,7 @@ struct FNDIRigidMeshCollisionArrays
 	TArray<FVector4f> PreviousTransform;
 	TArray<FVector4f> PreviousInverse;
 	TArray<FVector4f> ElementExtent;
+	TArray<FVector4f> MeshScale;
 	TArray<uint32> PhysicsType;
 	TArray<int32> ComponentIdIndex;
 	TArray<FPrimitiveComponentId> UniqueCompnentId;
@@ -55,6 +56,7 @@ struct FNDIRigidMeshCollisionArrays
 		PreviousTransform.Init(FVector4f(0, 0, 0, 0), 3 * MaxPrimitives);
 		PreviousInverse.Init(FVector4f(0, 0, 0, 0), 3 * MaxPrimitives);
 		ElementExtent.Init(FVector4f(0, 0, 0, 0), MaxPrimitives);
+		MeshScale.Init(FVector4f(0, 0, 0, 0), MaxPrimitives);
 		PhysicsType.Init(0, MaxPrimitives);
 		ComponentIdIndex.Init(INDEX_NONE, MaxPrimitives);
 		UniqueCompnentId.Reset();
@@ -84,6 +86,9 @@ struct FNDIRigidMeshCollisionBuffer : public FRenderResource
 	/** Element extent buffer */
 	FReadBuffer ElementExtentBuffer;
 
+	/** Element extent buffer */
+	FReadBuffer MeshScaleBuffer;
+
 	/** Physics type buffer */
 	FReadBuffer PhysicsTypeBuffer;
 
