optimization for get attribute index for rgba grids

diff --git a/Engine/Plugins/FX/Niagara/Source/Niagara/Private/NiagaraDataInterfaceGrid3DCollection.cpp b/Engine/Plugins/FX/Niagara/Source/Niagara/Private/NiagaraDataInterfaceGrid3DCollection.cpp
--- a/Engine/Plugins/FX/Niagara/Source/Niagara/Private/NiagaraDataInterfaceGrid3DCollection.cpp
+++ b/Engine/Plugins/FX/Niagara/Source/Niagara/Private/NiagaraDataInterfaceGrid3DCollection.cpp
@@ -2494,22 +2494,50 @@ bool UNiagaraDataInterfaceGrid3DCollection::GetFunctionHLSL(const FNiagaraDataIn
 	}
 	else if (FunctionInfo.DefinitionName == GetVector4AttributeIndexFunctionName)
 	{
-		WriteAttributeGetIndexHLSL(ParamInfo, FunctionInfo, FunctionInstanceIndex, 4, OutHLSL);
+		if (ParamHelper.UseRGBAGrid())
+		{
+			ParamHelper.WriteAttributeGetIndexHLSL(FGrid3DCollectionAttributeHelper::AttributeRetrievalMode::RGBAGrid, FunctionInfo, OutHLSL);
+		}
+		else
+		{
+			WriteAttributeGetIndexHLSL(ParamInfo, FunctionInfo, FunctionInstanceIndex, 4, OutHLSL);
+		}
 		return true;
 	}
 	else if (FunctionInfo.DefinitionName == GetVectorAttributeIndexFunctionName)
 	{
-		WriteAttributeGetIndexHLSL(ParamInfo, FunctionInfo, FunctionInstanceIndex, 3, OutHLSL);
+		if (ParamHelper.UseRGBAGrid())
+		{
+			ParamHelper.WriteAttributeGetIndexHLSL(FGrid3DCollectionAttributeHelper::AttributeRetrievalMode::RGBAGrid, FunctionInfo, OutHLSL);
+		}
+		else
+		{
+			WriteAttributeGetIndexHLSL(ParamInfo, FunctionInfo, FunctionInstanceIndex, 3, OutHLSL);
+		}
 		return true;
 	}
 	else if (FunctionInfo.DefinitionName == GetVector2DAttributeIndexFunctionName)
 	{
-		WriteAttributeGetIndexHLSL(ParamInfo, FunctionInfo, FunctionInstanceIndex, 2, OutHLSL);
+		if (ParamHelper.UseRGBAGrid())
+		{
+			ParamHelper.WriteAttributeGetIndexHLSL(FGrid3DCollectionAttributeHelper::AttributeRetrievalMode::RGBAGrid, FunctionInfo, OutHLSL);
+		}
+		else
+		{
+			WriteAttributeGetIndexHLSL(ParamInfo, FunctionInfo, FunctionInstanceIndex, 2, OutHLSL);
+		}
 		return true;
 	}
 	else if (FunctionInfo.DefinitionName == GetFloatAttributeIndexFunctionName)
 	{
-		WriteAttributeGetIndexHLSL(ParamInfo, FunctionInfo, FunctionInstanceIndex, 1, OutHLSL);
+		if (ParamHelper.UseRGBAGrid())
+		{
+			ParamHelper.WriteAttributeGetIndexHLSL(FGrid3DCollectionAttributeHelper::AttributeRetrievalMode::RGBAGrid, FunctionInfo, OutHLSL);
+		}
+		else
+		{
+			WriteAttributeGetIndexHLSL(ParamInfo, FunctionInfo, FunctionInstanceIndex, 1, OutHLSL);
+		}
 		return true;
 	}
 	else if (FunctionInfo.DefinitionName == SampleGridFunctionName || FunctionInfo.DefinitionName == SamplePreviousGridAtIndexFunctionName)
diff --git a/Engine/Plugins/FX/Niagara/Source/Niagara/Private/NiagaraDataInterfaceGrid3DCollectionUtils.cpp b/Engine/Plugins/FX/Niagara/Source/Niagara/Private/NiagaraDataInterfaceGrid3DCollectionUtils.cpp
--- a/Engine/Plugins/FX/Niagara/Source/Niagara/Private/NiagaraDataInterfaceGrid3DCollectionUtils.cpp
+++ b/Engine/Plugins/FX/Niagara/Source/Niagara/Private/NiagaraDataInterfaceGrid3DCollectionUtils.cpp
@@ -445,7 +445,7 @@ bool FGrid3DCollectionAttributeHelper::WriteSampleHLSL(AttributeRetrievalMode At
 	return true;
 }
 
-bool FGrid3DCollectionAttributeHelper::WriteAttributeGetIndexHLSL(bool bUseAttributeIndirection, const FNiagaraDataInterfaceGeneratedFunction& FunctionInfo, FString& OutHLSL)
+bool FGrid3DCollectionAttributeHelper::WriteAttributeGetIndexHLSL(AttributeRetrievalMode AttributeStorage, const FNiagaraDataInterfaceGeneratedFunction& FunctionInfo, FString& OutHLSL)
 {
 	const FName* AttributeName = FunctionInfo.FindSpecifierValue(UNiagaraDataInterfaceGrid3DCollection::NAME_Attribute);
 	if (AttributeName == nullptr)
@@ -458,12 +458,24 @@ bool FGrid3DCollectionAttributeHelper::WriteAttributeGetIndexHLSL(bool bUseAttri
 	{
 		return false;
 	}
-
-
-	OutHLSL.Appendf(TEXT("void %s(out int Value)\n"), *FunctionInfo.InstanceName);
-	OutHLSL.Appendf(TEXT("{\n"));
-	OutHLSL.Appendf(TEXT("	Value = %s;\n"), *GetAttributeIndex(bUseAttributeIndirection, AttributeInfo));
-	OutHLSL.Appendf(TEXT("}\n"));
+	
+	// #todo(dmp): for now it is ok to assume rgba grids only store 1 attribute per grid, so
+	// attributes are always stored at index 0.  In the future we might have a way to support
+	// multiple attributes per rgba grid, and this will have to change
+	if (AttributeStorage == AttributeRetrievalMode::RGBAGrid)
+	{
+		OutHLSL.Appendf(TEXT("void %s(out int Value)\n"), *FunctionInfo.InstanceName);
+		OutHLSL.Appendf(TEXT("{\n"));
+		OutHLSL.Appendf(TEXT("	Value = 0;\n"));
+		OutHLSL.Appendf(TEXT("}\n"));
+	}
+	else
+	{	
+		OutHLSL.Appendf(TEXT("void %s(out int Value)\n"), *FunctionInfo.InstanceName);
+		OutHLSL.Appendf(TEXT("{\n"));
+		OutHLSL.Appendf(TEXT("	Value = %s;\n"), *GetAttributeIndex(AttributeStorage == AttributeRetrievalMode::Indirection, AttributeInfo));
+		OutHLSL.Appendf(TEXT("}\n"));
+	}
 
 	return true;
 }
diff --git a/Engine/Plugins/FX/Niagara/Source/Niagara/Private/NiagaraDataInterfaceGrid3DCollectionUtils.h b/Engine/Plugins/FX/Niagara/Source/Niagara/Private/NiagaraDataInterfaceGrid3DCollectionUtils.h
--- a/Engine/Plugins/FX/Niagara/Source/Niagara/Private/NiagaraDataInterfaceGrid3DCollectionUtils.h
+++ b/Engine/Plugins/FX/Niagara/Source/Niagara/Private/NiagaraDataInterfaceGrid3DCollectionUtils.h
@@ -56,7 +56,7 @@ struct FGrid3DCollectionAttributeHelper
 	bool WriteSampleAtIndexHLSL(AttributeRetrievalMode AttributeStorage, const FNiagaraDataInterfaceGeneratedFunction& FunctionInfo, int NumChannels, bool IsCubic, FString& OutHLSL);
 	bool WriteSetHLSL(AttributeRetrievalMode AttributeStorage, const FNiagaraDataInterfaceGeneratedFunction& FunctionInfo, FString& OutHLSL);
 	bool WriteSampleHLSL(AttributeRetrievalMode AttributeStorage, const FNiagaraDataInterfaceGeneratedFunction& FunctionInfo, bool IsCubic, FString& OutHLSL);
-	bool WriteAttributeGetIndexHLSL(bool bUseAttributeIndirection, const FNiagaraDataInterfaceGeneratedFunction& FunctionInfo, FString& OutHLSL);
+	bool WriteAttributeGetIndexHLSL(AttributeRetrievalMode AttributeStorage, const FNiagaraDataInterfaceGeneratedFunction& FunctionInfo, FString& OutHLSL);
 #endif //WITH_EDITORONLY_DATA
 
 	const FNiagaraDataInterfaceGPUParamInfo& ParamInfo;
