- IWYU fix

#rnx

diff --git a/Engine/Plugins/FX/Niagara/Source/Niagara/Private/NiagaraDataInterfaceRigidMeshCollisionQuery.cpp b/Engine/Plugins/FX/Niagara/Source/Niagara/Private/NiagaraDataInterfaceRigidMeshCollisionQuery.cpp
--- a/Engine/Plugins/FX/Niagara/Source/Niagara/Private/NiagaraDataInterfaceRigidMeshCollisionQuery.cpp
+++ b/Engine/Plugins/FX/Niagara/Source/Niagara/Private/NiagaraDataInterfaceRigidMeshCollisionQuery.cpp
@@ -22,6 +22,7 @@
 #include "NiagaraShaderParametersBuilder.h"
 #include "NiagaraSimStageData.h"
 #include "NiagaraSystemInstance.h"
+#include "NiagaraSystemSimulation.h"
 #include "PhysicsEngine/PhysicsAsset.h"
 #include "ShaderParameterUtils.h"
 
