- Adjust RT data interfaces UI order / edit conditions, etc, to make it more logical

diff --git a/Engine/Plugins/FX/Niagara/Source/Niagara/Classes/NiagaraDataInterfaceIntRenderTarget2D.h b/Engine/Plugins/FX/Niagara/Source/Niagara/Classes/NiagaraDataInterfaceIntRenderTarget2D.h
--- a/Engine/Plugins/FX/Niagara/Source/Niagara/Classes/NiagaraDataInterfaceIntRenderTarget2D.h
+++ b/Engine/Plugins/FX/Niagara/Source/Niagara/Classes/NiagaraDataInterfaceIntRenderTarget2D.h
@@ -59,19 +59,19 @@ public:
 	void VMGetSize(FVectorVMExternalFunctionContext& Context);
 	void VMSetSize(FVectorVMExternalFunctionContext& Context);
 
-	UPROPERTY(EditAnywhere, Category = "Render Target")
+	UPROPERTY(EditAnywhere, Category = "Render Target", meta = (DisplayPriority = 1))
 	FIntPoint Size = FIntPoint::ZeroValue;
 
 #if WITH_EDITORONLY_DATA
-	UPROPERTY(EditAnywhere, Category = "Render Target")
+	UPROPERTY(EditAnywhere, Category = "Render Target", meta = (DisplayPriority = 2))
 	uint8 bPreviewRenderTarget : 1;
 
 	/* The range to normaliez the preview display to. */
-	UPROPERTY(EditAnywhere, Category = "Render Target")
+	UPROPERTY(EditAnywhere, Category = "Render Target", meta = (EditCondition = "bPreviewRenderTarget", EditConditionHides, DisplayPriority = 3))
 	FVector2D PreviewDisplayRange = FVector2D(0.0f, 255.0f);
 #endif
 
-	UPROPERTY(EditAnywhere, Category = "Render Target", meta = (ToolTip = "When valid the user parameter is used as the render target rather than creating one internal, note that the input render target will be adjusted by the Niagara simulation"))
+	UPROPERTY(EditAnywhere, Category = "Render Target", meta = (DisplayPriority = 0, ToolTip = "When valid the user parameter is used as the render target rather than creating one internal, note that the input render target will be adjusted by the Niagara simulation"))
 	FNiagaraUserParameterBinding RenderTargetUserParameter;
 
 protected:
diff --git a/Engine/Plugins/FX/Niagara/Source/Niagara/Classes/NiagaraDataInterfaceRenderTarget2D.h b/Engine/Plugins/FX/Niagara/Source/Niagara/Classes/NiagaraDataInterfaceRenderTarget2D.h
--- a/Engine/Plugins/FX/Niagara/Source/Niagara/Classes/NiagaraDataInterfaceRenderTarget2D.h
+++ b/Engine/Plugins/FX/Niagara/Source/Niagara/Classes/NiagaraDataInterfaceRenderTarget2D.h
@@ -135,22 +135,22 @@ public:
 	void VMGetSize(FVectorVMExternalFunctionContext& Context); 
 	void VMSetSize(FVectorVMExternalFunctionContext& Context);
 
-	UPROPERTY(EditAnywhere, Category = "Render Target", meta = (EditCondition = "!bInheritUserParameterSettings"))
+	UPROPERTY(EditAnywhere, Category = "Render Target", meta = (DisplayPriority = 2, EditCondition = "!bInheritUserParameterSettings", EditConditionHides))
 	FIntPoint Size;
 
 	/** Controls if and when we generate mips for the render target. */
-	UPROPERTY(EditAnywhere, Category = "Render Target", meta = (EditCondition = "!bInheritUserParameterSettings"))
+	UPROPERTY(EditAnywhere, Category = "Render Target", meta = (DisplayPriority = 20, EditCondition = "!bInheritUserParameterSettings", EditConditionHides))
 	ENiagaraMipMapGeneration MipMapGeneration = ENiagaraMipMapGeneration::Disabled;
 
-	UPROPERTY(EditAnywhere, Category = "Render Target")
+	UPROPERTY(EditAnywhere, Category = "Render Target", meta = (DisplayPriority = 21, EditCondition = "!bInheritUserParameterSettings && MipMapGeneration != ENiagaraMipMapGeneration::Disabled", EditConditionHides))
 	ENiagaraMipMapGenerationType MipMapGenerationType = ENiagaraMipMapGenerationType::Linear;
 
 	/** When enabled overrides the format of the render target, otherwise uses the project default setting. */
-	UPROPERTY(EditAnywhere, Category = "Render Target", meta = (EditCondition = "!bInheritUserParameterSettings && bOverrideFormat"))
+	UPROPERTY(EditAnywhere, Category = "Render Target", meta = (DisplayPriority = 11, EditCondition = "!bInheritUserParameterSettings && bOverrideFormat", EditConditionHides))
 	TEnumAsByte<ETextureRenderTargetFormat> OverrideRenderTargetFormat;
 
 	/** When enabled overrides the filter of the render target, otherwise uses the project default setting. */
-	UPROPERTY(EditAnywhere, Category = "Render Target", meta = (EditCondition = "!bInheritUserParameterSettings"))
+	UPROPERTY(EditAnywhere, Category = "Render Target", meta = (DisplayPriority = 12, EditCondition = "!bInheritUserParameterSettings", EditConditionHides))
 	TEnumAsByte<TextureFilter> OverrideRenderTargetFilter = TextureFilter::TF_Default;
 
 	/**
@@ -158,18 +158,18 @@ public:
 	If no valid user parameter is set the system will be invalid.
 	Note: The resource will be recreated if UAV access is not available, which will reset the contents.
 	*/
-	UPROPERTY(EditAnywhere, Category = "Render Target")
+	UPROPERTY(EditAnywhere, Category = "Render Target", meta = (DisplayPriority = 0))
 	uint8 bInheritUserParameterSettings : 1;
 
-	UPROPERTY(EditAnywhere, Category = "Render Target")
+	UPROPERTY(EditAnywhere, Category = "Render Target", meta = (DisplayPriority = 10, EditCondition = "!bInheritUserParameterSettings", EditConditionHides))
 	uint8 bOverrideFormat : 1;
 
 #if WITH_EDITORONLY_DATA
-	UPROPERTY(EditAnywhere, Category = "Render Target")
+	UPROPERTY(EditAnywhere, Category = "Render Target", meta = (DisplayPriority = 30))
 	uint8 bPreviewRenderTarget : 1;
 #endif
 
-	UPROPERTY(EditAnywhere, Category = "Render Target", meta = (ToolTip = "When valid the user parameter is used as the render target rather than creating one internal, note that the input render target will be adjusted by the Niagara simulation"))
+	UPROPERTY(EditAnywhere, Category = "Render Target", meta = (DisplayPriority = 1, ToolTip = "When valid the user parameter is used as the render target rather than creating one internal, note that the input render target will be adjusted by the Niagara simulation"))
 	FNiagaraUserParameterBinding RenderTargetUserParameter;
 
 protected:
diff --git a/Engine/Plugins/FX/Niagara/Source/Niagara/Classes/NiagaraDataInterfaceRenderTarget2DArray.h b/Engine/Plugins/FX/Niagara/Source/Niagara/Classes/NiagaraDataInterfaceRenderTarget2DArray.h
--- a/Engine/Plugins/FX/Niagara/Source/Niagara/Classes/NiagaraDataInterfaceRenderTarget2DArray.h
+++ b/Engine/Plugins/FX/Niagara/Source/Niagara/Classes/NiagaraDataInterfaceRenderTarget2DArray.h
@@ -127,15 +127,15 @@ public:
 	void VMGetSize(FVectorVMExternalFunctionContext& Context); 
 	void VMSetSize(FVectorVMExternalFunctionContext& Context);
 
-	UPROPERTY(EditAnywhere, Category = "Render Target", meta = (EditCondition = "!bInheritUserParameterSettings"))
+	UPROPERTY(EditAnywhere, Category = "Render Target", meta = (DisplayPriority = 2, EditCondition = "!bInheritUserParameterSettings", EditConditionHides))
 	FIntVector Size;
 
 	/** When enabled overrides the format of the render target, otherwise uses the project default setting. */
-	UPROPERTY(EditAnywhere, Category = "Render Target", meta = (EditCondition = "!bInheritUserParameterSettings && bOverrideFormat"))
+	UPROPERTY(EditAnywhere, Category = "Render Target", meta = (DisplayPriority = 11, EditCondition = "!bInheritUserParameterSettings && bOverrideFormat", EditConditionHides))
 	TEnumAsByte<ETextureRenderTargetFormat> OverrideRenderTargetFormat;
 
 	/** When enabled overrides the filter of the render target, otherwise uses the project default setting. */
-	UPROPERTY(EditAnywhere, Category = "Render Target", meta = (EditCondition = "!bInheritUserParameterSettings"))
+	UPROPERTY(EditAnywhere, Category = "Render Target", meta = (DisplayPriority = 12, EditCondition = "!bInheritUserParameterSettings", EditConditionHides))
 	TEnumAsByte<TextureFilter> OverrideRenderTargetFilter = TextureFilter::TF_Default;
 
 	/**
@@ -143,18 +143,18 @@ public:
 	If no valid user parameter is set the system will be invalid.
 	Note: The resource will be recreated if UAV access is not available, which will reset the contents.
 	*/
-	UPROPERTY(EditAnywhere, Category = "Render Target")
+	UPROPERTY(EditAnywhere, Category = "Render Target", meta = (DisplayPriority = 0))
 	uint8 bInheritUserParameterSettings : 1;
 
-	UPROPERTY(EditAnywhere, Category = "Render Target")
+	UPROPERTY(EditAnywhere, Category = "Render Target", meta = (DisplayPriority = 10, EditCondition = "!bInheritUserParameterSettings", EditConditionHides))
 	uint8 bOverrideFormat : 1;
 
 #if WITH_EDITORONLY_DATA
-	UPROPERTY(EditAnywhere, Category = "Render Target")
+	UPROPERTY(EditAnywhere, Category = "Render Target", meta = (DisplayPriority = 20))
 	uint8 bPreviewRenderTarget : 1;
 #endif
 
-	UPROPERTY(EditAnywhere, Category = "Render Target", meta = (ToolTip = "When valid the user parameter is used as the render target rather than creating one internal, note that the input render target will be adjusted by the Niagara simulation"))
+	UPROPERTY(EditAnywhere, Category = "Render Target", meta = (DisplayPriority = 1, ToolTip = "When valid the user parameter is used as the render target rather than creating one internal, note that the input render target will be adjusted by the Niagara simulation"))
 	FNiagaraUserParameterBinding RenderTargetUserParameter;
 
 protected:
diff --git a/Engine/Plugins/FX/Niagara/Source/Niagara/Classes/NiagaraDataInterfaceRenderTargetCube.h b/Engine/Plugins/FX/Niagara/Source/Niagara/Classes/NiagaraDataInterfaceRenderTargetCube.h
--- a/Engine/Plugins/FX/Niagara/Source/Niagara/Classes/NiagaraDataInterfaceRenderTargetCube.h
+++ b/Engine/Plugins/FX/Niagara/Source/Niagara/Classes/NiagaraDataInterfaceRenderTargetCube.h
@@ -120,15 +120,15 @@ public:
 	void VMGetSize(FVectorVMExternalFunctionContext& Context);
 	void VMSetSize(FVectorVMExternalFunctionContext& Context);
 
-	UPROPERTY(EditAnywhere, Category = "Render Target", meta = (EditCondition = "!bInheritUserParameterSettings"))
+	UPROPERTY(EditAnywhere, Category = "Render Target", meta = (DisplayPriority = 2, EditCondition = "!bInheritUserParameterSettings", EditConditionHides))
 	int Size;
 
 	/** When enabled overrides the format of the render target, otherwise uses the project default setting. */
-	UPROPERTY(EditAnywhere, Category = "Render Target", meta = (EditCondition = "!bInheritUserParameterSettings && bOverrideFormat"))
+	UPROPERTY(EditAnywhere, Category = "Render Target", meta = (DisplayPriority = 11, EditCondition = "!bInheritUserParameterSettings && bOverrideFormat", EditConditionHides))
 	TEnumAsByte<ETextureRenderTargetFormat> OverrideRenderTargetFormat;
 
 	/** When enabled overrides the filter of the render target, otherwise uses the project default setting. */
-	UPROPERTY(EditAnywhere, Category = "Render Target", meta = (EditCondition = "!bInheritUserParameterSettings"))
+	UPROPERTY(EditAnywhere, Category = "Render Target", meta = (DisplayPriority = 12, EditCondition = "!bInheritUserParameterSettings", EditConditionHides))
 	TEnumAsByte<TextureFilter> OverrideRenderTargetFilter = TextureFilter::TF_Default;
 
 	/**
@@ -136,18 +136,18 @@ public:
 	If no valid user parameter is set the system will be invalid.
 	Note: The resource will be recreated if UAV access is not available, which will reset the contents.
 	*/
-	UPROPERTY(EditAnywhere, Category = "Render Target")
+	UPROPERTY(EditAnywhere, Category = "Render Target", meta = (DisplayPriority = 0))
 	uint8 bInheritUserParameterSettings : 1;
 
-	UPROPERTY(EditAnywhere, Category = "Render Target")
+	UPROPERTY(EditAnywhere, Category = "Render Target", meta = (DisplayPriority = 10, EditCondition = "!bInheritUserParameterSettings", EditConditionHides))
 	uint8 bOverrideFormat : 1;
 
 #if WITH_EDITORONLY_DATA
-	UPROPERTY(EditAnywhere, Category = "Render Target")
+	UPROPERTY(EditAnywhere, Category = "Render Target", meta = (DisplayPriority = 20))
 	uint8 bPreviewRenderTarget : 1;
 #endif
 
-	UPROPERTY(EditAnywhere, Category = "Render Target", meta = (ToolTip = "When valid the user parameter is used as the render target rather than creating one internal, note that the input render target will be adjusted by the Niagara simulation"))
+	UPROPERTY(EditAnywhere, Category = "Render Target", meta = (DisplayPriority = 1, ToolTip = "When valid the user parameter is used as the render target rather than creating one internal, note that the input render target will be adjusted by the Niagara simulation"))
 	FNiagaraUserParameterBinding RenderTargetUserParameter;
 
 protected:
diff --git a/Engine/Plugins/FX/Niagara/Source/Niagara/Classes/NiagaraDataInterfaceRenderTargetVolume.h b/Engine/Plugins/FX/Niagara/Source/Niagara/Classes/NiagaraDataInterfaceRenderTargetVolume.h
--- a/Engine/Plugins/FX/Niagara/Source/Niagara/Classes/NiagaraDataInterfaceRenderTargetVolume.h
+++ b/Engine/Plugins/FX/Niagara/Source/Niagara/Classes/NiagaraDataInterfaceRenderTargetVolume.h
@@ -128,15 +128,15 @@ public:
 	void VMGetSize(FVectorVMExternalFunctionContext& Context); 
 	void VMSetSize(FVectorVMExternalFunctionContext& Context);
 
-	UPROPERTY(EditAnywhere, Category = "Render Target", meta = (EditCondition = "!bInheritUserParameterSettings"))
+	UPROPERTY(EditAnywhere, Category = "Render Target", meta = (DisplayPriority = 2, EditCondition = "!bInheritUserParameterSettings", EditConditionHides))
 	FIntVector Size;
 
 	/** When enabled overrides the format of the render target, otherwise uses the project default setting. */
-	UPROPERTY(EditAnywhere, Category = "Render Target", meta = (EditCondition = "!bInheritUserParameterSettings && bOverrideFormat"))
+	UPROPERTY(EditAnywhere, Category = "Render Target", meta = (DisplayPriority = 11, EditCondition = "!bInheritUserParameterSettings && bOverrideFormat", EditConditionHides))
 	TEnumAsByte<ETextureRenderTargetFormat> OverrideRenderTargetFormat;
 
 	/** When enabled overrides the filter of the render target, otherwise uses the project default setting. */
-	UPROPERTY(EditAnywhere, Category = "Render Target", meta = (EditCondition = "!bInheritUserParameterSettings"))
+	UPROPERTY(EditAnywhere, Category = "Render Target", meta = (DisplayPriority = 12, EditCondition = "!bInheritUserParameterSettings", EditConditionHides))
 	TEnumAsByte<TextureFilter> OverrideRenderTargetFilter = TextureFilter::TF_Default;
 
 	/**
@@ -144,18 +144,18 @@ public:
 	If no valid user parameter is set the system will be invalid.
 	Note: The resource will be recreated if UAV access is not available, which will reset the contents.
 	*/
-	UPROPERTY(EditAnywhere, Category = "Render Target")
+	UPROPERTY(EditAnywhere, Category = "Render Target", meta = (DisplayPriority = 0))
 	uint8 bInheritUserParameterSettings : 1;
 
-	UPROPERTY(EditAnywhere, Category = "Render Target")
+	UPROPERTY(EditAnywhere, Category = "Render Target", meta = (DisplayPriority = 10, EditCondition = "!bInheritUserParameterSettings", EditConditionHides))
 	uint8 bOverrideFormat : 1;
 
 #if WITH_EDITORONLY_DATA
-	UPROPERTY(EditAnywhere, Category = "Render Target")
+	UPROPERTY(EditAnywhere, Category = "Render Target", meta = (DisplayPriority = 20))
 	uint8 bPreviewRenderTarget : 1;
 #endif
 
-	UPROPERTY(EditAnywhere, Category = "Render Target", meta = (ToolTip = "When valid the user parameter is used as the render target rather than creating one internal, note that the input render target will be adjusted by the Niagara simulation"))
+	UPROPERTY(EditAnywhere, Category = "Render Target", meta = (DisplayPriority = 1, ToolTip = "When valid the user parameter is used as the render target rather than creating one internal, note that the input render target will be adjusted by the Niagara simulation"))
 	FNiagaraUserParameterBinding RenderTargetUserParameter;
 
 protected:
