- Added tab icons to user parameter tabs
- Changed icons for Add Section & Add Category in Hierarchy Editor to UE5 style
- Moved Add Section & Add Category UI to new places

#rnx

diff --git a/Engine/Plugins/FX/Niagara/Source/NiagaraEditor/Private/NiagaraEditorStyle.cpp b/Engine/Plugins/FX/Niagara/Source/NiagaraEditor/Private/NiagaraEditorStyle.cpp
--- a/Engine/Plugins/FX/Niagara/Source/NiagaraEditor/Private/NiagaraEditorStyle.cpp
+++ b/Engine/Plugins/FX/Niagara/Source/NiagaraEditor/Private/NiagaraEditorStyle.cpp
@@ -595,6 +595,8 @@ void FNiagaraEditorStyle::InitTabIcons()
 	Set("Tab.Spreadsheet", new IMAGE_BRUSH_SVG("Icons/Tabs/Spreadsheet", Icon20x20));
 	Set("Tab.SystemOverview", new IMAGE_BRUSH_SVG("Icons/Tabs/SystemOverview", Icon20x20));
 	Set("Tab.Timeline", new IMAGE_BRUSH_SVG("Icons/Tabs/Timeline", Icon20x20));
+	Set("Tab.UserParameters", new IMAGE_BRUSH_SVG("Icons/Tabs/UserParameters", Icon20x20));
+	Set("Tab.UserParameterHierarchy", new IMAGE_BRUSH_SVG("Icons/Tabs/UserParameterHierarchy", Icon20x20));
 	Set("Tab.Viewport", new IMAGE_BRUSH_SVG("Icons/Tabs/Viewport", Icon20x20));
 	Set("Tab.VisualEffects", new IMAGE_BRUSH_SVG("Icons/Tabs/VisualEffects", Icon20x20));
 }
@@ -706,6 +708,17 @@ void FNiagaraEditorStyle::InitHierarchyEditor()
 
 	Set("NiagaraEditor.Stack.DropTarget.BorderVertical", new IMAGE_BRUSH("Icons/StackDropTargetBorder_Vertical", FVector2D(2, 8), FLinearColor::White, ESlateBrushTileType::Vertical));
 	Set("NiagaraEditor.Stack.DropTarget.BorderHorizontal", new IMAGE_BRUSH("Icons/StackDropTargetBorder_Horizontal", FVector2D(8, 2), FLinearColor::White, ESlateBrushTileType::Horizontal));
+
+	FButtonStyle SimpleButtonStyle = FAppStyle::Get().GetWidgetStyle<FButtonStyle>("SimpleButton");
+	FButtonStyle ButtonStyle = FButtonStyle(SimpleButtonStyle)
+		.SetNormalForeground(FStyleColors::Foreground)
+		.SetPressedForeground(FStyleColors::ForegroundHover)
+		.SetHoveredForeground(FStyleColors::ForegroundHover)
+		.SetDisabledForeground(FStyleColors::Foreground)
+		.SetNormalPadding(FMargin(8.f, 2.f, 8.f, 2.f))
+		.SetPressedPadding(FMargin(8.f, 3.f, 8.f, 1.f));
+
+	Set("NiagaraEditor.HierarchyEditor.ButtonStyle", ButtonStyle);
 }
 
 FNiagaraEditorStyle::FNiagaraEditorStyle() : FSlateStyleSet("NiagaraEditorStyle")
diff --git a/Engine/Plugins/FX/Niagara/Source/NiagaraEditor/Private/Toolkits/SystemToolkitModes/NiagaraSystemToolkitModeBase.cpp b/Engine/Plugins/FX/Niagara/Source/NiagaraEditor/Private/Toolkits/SystemToolkitModes/NiagaraSystemToolkitModeBase.cpp
--- a/Engine/Plugins/FX/Niagara/Source/NiagaraEditor/Private/Toolkits/SystemToolkitModes/NiagaraSystemToolkitModeBase.cpp
+++ b/Engine/Plugins/FX/Niagara/Source/NiagaraEditor/Private/Toolkits/SystemToolkitModes/NiagaraSystemToolkitModeBase.cpp
@@ -308,11 +308,13 @@ void FNiagaraSystemToolkitModeBase::RegisterTabFactories(TSharedPtr<FTabManager>
 	
 	InTabManager->RegisterTabSpawner(UserParametersTabID, FOnSpawnTab::CreateSP(this, &FNiagaraSystemToolkitModeBase::SpawnTab_UserParameters))
 		.SetDisplayName(LOCTEXT("UserParametersTab", "User Parameters"))
-		.SetGroup(WorkspaceMenuCategory.ToSharedRef());
+		.SetGroup(WorkspaceMenuCategory.ToSharedRef())
+		.SetIcon(FSlateIcon(FNiagaraEditorStyle::Get().GetStyleSetName(), "Tab.UserParameters"));
 	
 	InTabManager->RegisterTabSpawner(UserParametersHierarchyTabID, FOnSpawnTab::CreateSP(this, &FNiagaraSystemToolkitModeBase::SpawnTab_UserParametersHierarchyEditor))
 		.SetDisplayName(LOCTEXT("UserParametersHierarchyTab", "User Parameters Hierarchy"))
-		.SetGroup(WorkspaceMenuCategory.ToSharedRef());
+		.SetGroup(WorkspaceMenuCategory.ToSharedRef())
+		.SetIcon(FSlateIcon(FNiagaraEditorStyle::Get().GetStyleSetName(), "Tab.UserParameterHierarchy"));
 }
 
 int FNiagaraSystemToolkitModeBase::GetActiveSelectionDetailsIndex() const
diff --git a/Engine/Plugins/FX/Niagara/Source/NiagaraEditor/Private/ViewModels/HierarchyEditor/NiagaraHierarchyViewModelBase.cpp b/Engine/Plugins/FX/Niagara/Source/NiagaraEditor/Private/ViewModels/HierarchyEditor/NiagaraHierarchyViewModelBase.cpp
--- a/Engine/Plugins/FX/Niagara/Source/NiagaraEditor/Private/ViewModels/HierarchyEditor/NiagaraHierarchyViewModelBase.cpp
+++ b/Engine/Plugins/FX/Niagara/Source/NiagaraEditor/Private/ViewModels/HierarchyEditor/NiagaraHierarchyViewModelBase.cpp
@@ -112,7 +112,7 @@ UNiagaraHierarchySection* UNiagaraHierarchyRoot::AddSection(FText InNewSectionNa
 	UNiagaraHierarchySection* NewSectionItem = NewObject<UNiagaraHierarchySection>(this);
 	NewSectionItem->SetSectionName(NewName);
 	NewSectionItem->SetFlags(RF_Transactional);
-	Sections.Add(NewSectionItem);
+	Sections.Insert(NewSectionItem, 0);
 	return NewSectionItem;
 }
 
@@ -726,7 +726,7 @@ TSharedPtr<FNiagaraHierarchySectionViewModel> FNiagaraHierarchyRootViewModel::Ad
 	FScopedTransaction ScopedTransaction(LOCTEXT("NewSectionAdded","Added Section"));
 	HierarchyViewModel->GetHierarchyDataRoot()->Modify();
 	
-	UNiagaraHierarchySection* SectionData = Cast<UNiagaraHierarchyRoot>(ItemBase)->AddSection(LOCTEXT("NiagaraHierarchyEditorDefaultNewSectionName", "New Section"));
+	UNiagaraHierarchySection* SectionData = Cast<UNiagaraHierarchyRoot>(ItemBase)->AddSection(LOCTEXT("NiagaraHierarchyEditorDefaultNewSectionName", "Section"));
 	SectionData->Modify();
 	TSharedPtr<FNiagaraHierarchySectionViewModel> NewSectionViewModel = MakeShared<FNiagaraHierarchySectionViewModel>(SectionData, StaticCastSharedRef<FNiagaraHierarchyRootViewModel>(AsShared()), HierarchyViewModel);
 	SectionViewModels.Add(NewSectionViewModel);
diff --git a/Engine/Plugins/FX/Niagara/Source/NiagaraEditor/Private/Widgets/SNiagaraHierarchy.cpp b/Engine/Plugins/FX/Niagara/Source/NiagaraEditor/Private/Widgets/SNiagaraHierarchy.cpp
--- a/Engine/Plugins/FX/Niagara/Source/NiagaraEditor/Private/Widgets/SNiagaraHierarchy.cpp
+++ b/Engine/Plugins/FX/Niagara/Source/NiagaraEditor/Private/Widgets/SNiagaraHierarchy.cpp
@@ -9,6 +9,8 @@
 #include "Widgets/Layout/SSeparator.h"
 #include "Widgets/Layout/SWrapBox.h"
 #include "Widgets/Input/SCheckBox.h"
+#include "Widgets/Input/SButton.h"
+#include "Widgets/Images/SImage.h"
 #include "Widgets/Text/SInlineEditableTextBlock.h"
 #include "Widgets/Text/STextBlock.h"
 #include "Framework/MultiBox/MultiBoxBuilder.h"
@@ -138,6 +140,7 @@ void SNiagaraHierarchySection::Construct(const FArguments& InArgs, TSharedPtr<FN
 						.Style(FAppStyle::Get(), "DetailsView.SectionButton")
 						.OnCheckStateChanged(this, &SNiagaraHierarchySection::OnSectionCheckChanged)
 						.IsChecked(this, &SNiagaraHierarchySection::GetSectionCheckState)
+						.Padding(FMargin(8.f, 4.f))
 						[
 							SAssignNew(InlineEditableTextBlock, SInlineEditableTextBlock)
 							.IsSelected(InArgs._IsSelected)
@@ -169,6 +172,7 @@ void SNiagaraHierarchySection::Construct(const FArguments& InArgs, TSharedPtr<FN
 				.Style(FAppStyle::Get(), "DetailsView.SectionButton")
 				.OnCheckStateChanged(this, &SNiagaraHierarchySection::OnSectionCheckChanged)
 				.IsChecked(this, &SNiagaraHierarchySection::GetSectionCheckState)
+				.Padding(FMargin(8.f, 4.f))
 				[
 					SNew(STextBlock)
 					.Text(LOCTEXT("AllSection", "All"))
@@ -389,14 +393,69 @@ void SNiagaraHierarchy::Construct(const FArguments& InArgs, TObjectPtr<UNiagaraH
 	OnGenerateCustomDetailsPanelNameWidget = InArgs._OnGenerateCustomDetailsPanelNameWidget;
 	
 	SHorizontalBox::FSlot* DetailsPanelSlot = nullptr;
+
+	LightBackgroundBrush = FSlateColorBrush(FStyleColors::Panel);
+	RecessedBackgroundBrush = FSlateColorBrush(FStyleColors::Recessed);
+
+	TSharedRef<SWidget> AddSectionButton = SNew(SHorizontalBox)
+		+ SHorizontalBox::Slot()
+		.AutoWidth()
+		[
+			SNew(SButton)
+			.OnPressed_UObject(HierarchyViewModel.Get(), &UNiagaraHierarchyViewModelBase::AddSection)
+			.ButtonStyle(FNiagaraEditorStyle::Get(), "NiagaraEditor.HierarchyEditor.ButtonStyle")
+			[
+				SNew(SHorizontalBox)
+				+ SHorizontalBox::Slot()
+				.AutoWidth()
+				.Padding(2.f)
+				[
+					SNew(SImage)
+					.Image(FAppStyle::Get().GetBrush("Icons.PlusCircle"))
+				]
+				+ SHorizontalBox::Slot()
+				.VAlign(VAlign_Center)
+				.AutoWidth()
+				[			
+					SNew(STextBlock)
+					.Text(LOCTEXT("AddSectionLabel","Add Section"))				
+				]
+			]
+		];
+
+	TSharedRef<SWidget> AddCategoryButton = SNew(SHorizontalBox)
+		+ SHorizontalBox::Slot()
+		.AutoWidth()
+		[
+			SNew(SButton)
+			.OnPressed_UObject(HierarchyViewModel.Get(), &UNiagaraHierarchyViewModelBase::AddCategory)
+			.ButtonStyle(FNiagaraEditorStyle::Get(), "NiagaraEditor.HierarchyEditor.ButtonStyle")
+			[
+				SNew(SHorizontalBox)
+				+ SHorizontalBox::Slot()
+				.AutoWidth()
+				.Padding(2.f)
+				[
+					SNew(SImage)
+					.Image(FAppStyle::Get().GetBrush("Icons.PlusCircle"))
+				]
+				+ SHorizontalBox::Slot()
+				.VAlign(VAlign_Center)
+				.AutoWidth()
+				[			
+					SNew(STextBlock)
+					.Text(LOCTEXT("AddCategoryLabel","Add Category"))				
+				]
+			]
+		];
 	
 	ChildSlot
 	[
-		SNew(SHorizontalBox)
-		+ SHorizontalBox::Slot()
+		SNew(SBorder)
+		.BorderImage(&RecessedBackgroundBrush)
 		[
-			SNew(SVerticalBox)
-			+ SVerticalBox::Slot()
+			SNew(SHorizontalBox)
+			+ SHorizontalBox::Slot()
 			[
 				SAssignNew(SourceTreeView, STreeView<TSharedPtr<FNiagaraHierarchyItemViewModelBase>>)
 				.TreeItemsSource(&InHierarchyViewModel->GetSourceItems())
@@ -405,57 +464,70 @@ void SNiagaraHierarchy::Construct(const FArguments& InArgs, TObjectPtr<UNiagaraH
 				.OnGetChildren_UObject(InHierarchyViewModel.Get(), &UNiagaraHierarchyViewModelBase::OnGetChildren)
 				.OnItemToString_Debug_UObject(InHierarchyViewModel.Get(), &UNiagaraHierarchyViewModelBase::OnItemToStringDebug)
 			]
-		]
-		+ SHorizontalBox::Slot()
-		.AutoWidth()
-		[
-			SNew(SSeparator).Orientation(EOrientation::Orient_Vertical)
-		]
-		+ SHorizontalBox::Slot()
-		[
-			SNew(SVerticalBox)
-			+ SVerticalBox::Slot()
-			.AutoHeight()
+			+ SHorizontalBox::Slot()
+			.AutoWidth()
+			[
+				SNew(SSeparator).Orientation(EOrientation::Orient_Vertical)
+			]
+			+ SHorizontalBox::Slot()
 			[
 				SNew(SVerticalBox)
 				+ SVerticalBox::Slot()
 				.AutoHeight()
 				[
-					SNew(SBox)
+					SNew(SVerticalBox)
+					+ SVerticalBox::Slot()
+					.AutoHeight()
 					[
-						CreateHierarchyButtonWidgets()
+						SNew(SBorder)
+						.Padding(0.f)
+						.BorderImage(&LightBackgroundBrush)
+						[
+							AddSectionButton
+						]
 					]
+					+ SVerticalBox::Slot()
+					.AutoHeight()
+					.Padding(1.f)
+					[
+						SAssignNew(SectionBox, SWrapBox)
+						.UseAllottedWidth(true)
+					]				
 				]
 				+ SVerticalBox::Slot()
-				.AutoHeight()
-				[
-					SAssignNew(SectionBox, SWrapBox)
-					.UseAllottedWidth(true)
-				]
-			]
-			+ SVerticalBox::Slot()
-			[
-				 SNew(SVerticalBox)
-				 + SVerticalBox::Slot()
-				 .AutoHeight()
-				 [
-					SAssignNew(HierarchyTreeView, STreeView<TSharedPtr<FNiagaraHierarchyItemViewModelBase>>)
-					.TreeItemsSource(&InHierarchyViewModel->GetHierarchyItems())
-					.OnSelectionChanged(this, &SNiagaraHierarchy::OnSelectionChanged, true)
-					.OnGenerateRow(this, &SNiagaraHierarchy::GenerateHierarchyItemRow)
-					.OnGetChildren_UObject(InHierarchyViewModel.Get(), &UNiagaraHierarchyViewModelBase::OnGetChildren)
-					.OnItemToString_Debug_UObject(InHierarchyViewModel.Get(), &UNiagaraHierarchyViewModelBase::OnItemToStringDebug)
-				]
-				+ SVerticalBox::Slot()
-				.VAlign(VAlign_Fill)
 				[
-					SNew(SDropTarget)
-					.OnDropped(this, &SNiagaraHierarchy::HandleHierarchyRootDrop)
-					.OnAllowDrop(this, &SNiagaraHierarchy::CanDropOnRoot)
+					SNew(SVerticalBox)
+					+ SVerticalBox::Slot()
+					.AutoHeight()
+					[
+						SNew(SBorder)
+						.BorderImage(&LightBackgroundBrush)
+						.Padding(0.f)
+						[
+							AddCategoryButton
+						]
+					]
+					+ SVerticalBox::Slot()
+					.AutoHeight()
+					[
+						SAssignNew(HierarchyTreeView, STreeView<TSharedPtr<FNiagaraHierarchyItemViewModelBase>>)
+						.TreeItemsSource(&InHierarchyViewModel->GetHierarchyItems())
+						.OnSelectionChanged(this, &SNiagaraHierarchy::OnSelectionChanged, true)
+						.OnGenerateRow(this, &SNiagaraHierarchy::GenerateHierarchyItemRow)
+						.OnGetChildren_UObject(InHierarchyViewModel.Get(), &UNiagaraHierarchyViewModelBase::OnGetChildren)
+						.OnItemToString_Debug_UObject(InHierarchyViewModel.Get(), &UNiagaraHierarchyViewModelBase::OnItemToStringDebug)
+					]
+					+ SVerticalBox::Slot()
+					.VAlign(VAlign_Fill)
+					[
+						SNew(SDropTarget)
+						.OnDropped(this, &SNiagaraHierarchy::HandleHierarchyRootDrop)
+						.OnAllowDrop(this, &SNiagaraHierarchy::CanDropOnRoot)
+					]
 				]
 			]
+			+ SHorizontalBox::Slot().Expose(DetailsPanelSlot)		
 		]
-		+ SHorizontalBox::Slot().Expose(DetailsPanelSlot)		
 	];
 
 	if(InHierarchyViewModel->SupportsDetailsPanel())
@@ -568,8 +640,10 @@ void SNiagaraHierarchy::RefreshSectionsWidget()
 		SectionsWidgetMap.Add(Section, SectionWidget);
 		
 		SectionBox->AddSlot()
+		.HAlign(HAlign_Center)
+		.VAlign(VAlign_Center)
 		[
-			SectionWidget.ToSharedRef()
+			SectionWidget.ToSharedRef()		
 		];
 	}
 	
@@ -582,6 +656,8 @@ void SNiagaraHierarchy::RefreshSectionsWidget()
 	SectionsWidgetMap.Add(nullptr, DefaultSectionWidget);
 
 	SectionBox->AddSlot()
+	.HAlign(HAlign_Center)
+	.VAlign(VAlign_Center)
 	[
 		DefaultSectionWidget.ToSharedRef()
 	];
@@ -600,31 +676,6 @@ void SNiagaraHierarchy::SelectObjectInDetailsPanel(UObject* Object) const
 	}
 }
 
-TSharedRef<SWidget> SNiagaraHierarchy::CreateHierarchyButtonWidgets() const
-{
-	FSlimHorizontalToolBarBuilder ToolBarBuilder(nullptr, FMultiBoxCustomization::None);
-	
-	// Add Section
-	{
-		FUIAction AddSectionAction;
-		AddSectionAction.ExecuteAction = FExecuteAction::CreateUObject(HierarchyViewModel.Get(), &UNiagaraHierarchyViewModelBase::AddSection);
-		ToolBarBuilder.AddToolBarButton(AddSectionAction, NAME_None,
-			FText::FromString("Add Section")
-			);
-	}
-	
-	// Add Category
-	{
-		FUIAction AddCategoryAction;
-		AddCategoryAction.ExecuteAction = FExecuteAction::CreateUObject(HierarchyViewModel.Get(), &UNiagaraHierarchyViewModelBase::AddCategory);
-		ToolBarBuilder.AddToolBarButton(AddCategoryAction, NAME_None,
-			FText::FromString("Add Category")
-			);
-	}
-
-	return ToolBarBuilder.MakeWidget();
-}
-
 TSharedRef<ITableRow> SNiagaraHierarchy::GenerateSourceItemRow(TSharedPtr<FNiagaraHierarchyItemViewModelBase> HierarchyItem, const TSharedRef<STableViewBase>& TableViewBase)
 {
 	return SNew(STableRow<TSharedPtr<FNiagaraHierarchyItemViewModelBase>>, TableViewBase)
diff --git a/Engine/Plugins/FX/Niagara/Source/NiagaraEditor/Public/Widgets/SNiagaraHierarchy.h b/Engine/Plugins/FX/Niagara/Source/NiagaraEditor/Public/Widgets/SNiagaraHierarchy.h
--- a/Engine/Plugins/FX/Niagara/Source/NiagaraEditor/Public/Widgets/SNiagaraHierarchy.h
+++ b/Engine/Plugins/FX/Niagara/Source/NiagaraEditor/Public/Widgets/SNiagaraHierarchy.h
@@ -40,7 +40,6 @@ public:
 	void SelectObjectInDetailsPanel(UObject* Object) const;
 	
 private:
-	TSharedRef<SWidget> CreateHierarchyButtonWidgets() const;
 	TSharedRef<ITableRow> GenerateSourceItemRow(TSharedPtr<FNiagaraHierarchyItemViewModelBase> HierarchyItem, const TSharedRef<STableViewBase>& TableViewBase);
 	TSharedRef<ITableRow> GenerateHierarchyItemRow(TSharedPtr<FNiagaraHierarchyItemViewModelBase> HierarchyItem, const TSharedRef<STableViewBase>& TableViewBase);
 private:	
@@ -68,7 +67,8 @@ private:
 	TMap<TSharedPtr<FNiagaraHierarchySectionViewModel>, TSharedPtr<class SNiagaraHierarchySection>> SectionsWidgetMap;
 	TSharedPtr<class SWrapBox> SectionBox;
 	TSharedPtr<IDetailsView> DetailsPanel;
-
+	FSlateBrush LightBackgroundBrush;
+	FSlateBrush RecessedBackgroundBrush;
 private:
 	FOnGenerateRowContentWidget OnGenerateRowContentWidget;
 	FOnGenerateCustomDetailsPanelNameWidget OnGenerateCustomDetailsPanelNameWidget;
