renamed FTickInfo to FNiagaraTickInfo

diff --git a/Engine/Plugins/FX/Niagara/Source/Niagara/Private/NiagaraDataInterfaceRigidMeshCollisionQuery.cpp b/Engine/Plugins/FX/Niagara/Source/Niagara/Private/NiagaraDataInterfaceRigidMeshCollisionQuery.cpp
--- a/Engine/Plugins/FX/Niagara/Source/Niagara/Private/NiagaraDataInterfaceRigidMeshCollisionQuery.cpp
+++ b/Engine/Plugins/FX/Niagara/Source/Niagara/Private/NiagaraDataInterfaceRigidMeshCollisionQuery.cpp
@@ -1038,7 +1038,7 @@ bool UNiagaraDataInterfaceRigidMeshCollisionQuery::PerInstanceTick(void* PerInst
 	FNDIRigidMeshCollisionData* InstanceData = static_cast<FNDIRigidMeshCollisionData*>(PerInstanceData);
 	if (InstanceData && SystemInstance)
 	{
-		const FTickInfo &TickInfo = SystemInstance->GetSystemSimulation()->GetTickInfo();
+		const FNiagaraTickInfo &TickInfo = SystemInstance->GetSystemSimulation()->GetTickInfo();
 
 		// only update on the first tick
 		if (TickInfo.TickNumber == 0)
diff --git a/Engine/Plugins/FX/Niagara/Source/Niagara/Public/NiagaraSystemSimulation.h b/Engine/Plugins/FX/Niagara/Source/Niagara/Public/NiagaraSystemSimulation.h
--- a/Engine/Plugins/FX/Niagara/Source/Niagara/Public/NiagaraSystemSimulation.h
+++ b/Engine/Plugins/FX/Niagara/Source/Niagara/Public/NiagaraSystemSimulation.h
@@ -27,7 +27,7 @@ enum class ENiagaraGPUTickHandlingMode
 	ConcurrentBatched,/** Systems can submit their GPU ticks in batches during concurrent tick. */
 };
 
-struct FTickInfo
+struct FNiagaraTickInfo
 {
 	bool UsesFixedTick = false;
 	float EngineTick = 0.0;
@@ -327,7 +327,7 @@ public:
 
 	FORCEINLINE UWorld* GetWorld()const{return World;}
 
-	const FTickInfo& GetTickInfo() { return TickInfo; }
+	const FNiagaraTickInfo& GetTickInfo() { return TickInfo; }
 
 protected:
 	void DumpStalledInfo();
@@ -447,5 +447,5 @@ protected:
 
 	float FixedDeltaTickAge = 0;
 
-	FTickInfo TickInfo;
+	FNiagaraTickInfo TickInfo;
 };
