- Change Ribbon GetDynamicDataSize to follow pattern of others to not include dynamic data information
- This fixes a 'leak' when looking at 'Primitive Memory' stat as it differs between add / remove

#rnx

diff --git a/Engine/Plugins/FX/Niagara/Source/Niagara/Private/NiagaraRendererRibbons.cpp b/Engine/Plugins/FX/Niagara/Source/Niagara/Private/NiagaraRendererRibbons.cpp
--- a/Engine/Plugins/FX/Niagara/Source/Niagara/Private/NiagaraRendererRibbons.cpp
+++ b/Engine/Plugins/FX/Niagara/Source/Niagara/Private/NiagaraRendererRibbons.cpp
@@ -920,11 +920,6 @@ int FNiagaraRendererRibbons::GetDynamicDataSize()const
 
 	Size += ShapeState.SliceVertexData.GetAllocatedSize();
 
-	if (DynamicDataRender)
-	{
-		const FNiagaraDynamicDataRibbon* RibbonDynamicData = static_cast<FNiagaraDynamicDataRibbon*>(DynamicDataRender);
-		Size += RibbonDynamicData->GetAllocatedSize();
-	}
 	return Size;
 }
 
