- Extend array debug HUD information to include some elements

diff --git a/Engine/Plugins/FX/Niagara/Source/Niagara/Classes/NiagaraDataInterfaceArrayFloat.h b/Engine/Plugins/FX/Niagara/Source/Niagara/Classes/NiagaraDataInterfaceArrayFloat.h
--- a/Engine/Plugins/FX/Niagara/Source/Niagara/Classes/NiagaraDataInterfaceArrayFloat.h
+++ b/Engine/Plugins/FX/Niagara/Source/Niagara/Classes/NiagaraDataInterfaceArrayFloat.h
@@ -22,6 +22,11 @@ struct FNDIArrayImplHelper<float> : public FNDIArrayImplHelperBase<float>
 
 	static const FNiagaraTypeDefinition& GetTypeDefinition() { return FNiagaraTypeDefinition::GetFloatDef(); }
 	static const float GetDefaultValue() { return 0.0f; }
+
+	static void AppendValueToString(const float Value, FString& OutString)
+	{
+		OutString.Appendf(TEXT("%f"), Value);
+	}
 };
 
 template<>
@@ -61,6 +66,11 @@ struct FNDIArrayImplHelper<FVector2f> : public FNDIArrayImplHelperBase<FVector2f
 			Dest[i] = FVector2D(Src[i]);
 		}
 	}
+
+	static void AppendValueToString(const FVector2f& Value, FString& OutString)
+	{
+		OutString.Appendf(TEXT("%f, %f"), Value.X, Value.Y);
+	}
 };
 
 template<>
@@ -100,6 +110,11 @@ struct FNDIArrayImplHelper<FVector3f> : public FNDIArrayImplHelperBase<FVector3f
 			Dest[i] = FVector(Src[i]);
 		}
 	}
+
+	static void AppendValueToString(const FVector3f& Value, FString& OutString)
+	{
+		OutString.Appendf(TEXT("%f, %f, %f"), Value.X, Value.Y, Value.Z);
+	}
 };
 
 template<>
@@ -117,6 +132,11 @@ struct FNDIArrayImplHelper<FNiagaraPosition> : public FNDIArrayImplHelper<FVecto
 			Dest[i] = Src[i];
 		}
 	}
+
+	static void AppendValueToString(const FVector3f& Value, FString& OutString)
+	{
+		OutString.Appendf(TEXT("%f, %f, %f"), Value.X, Value.Y, Value.Z);
+	}
 };
 
 template<>
@@ -156,6 +176,11 @@ struct FNDIArrayImplHelper<FVector4f> : public FNDIArrayImplHelperBase<FVector4f
 			Dest[i] = FVector4(Src[i]);
 		}
 	}
+
+	static void AppendValueToString(const FVector4f& Value, FString& OutString)
+	{
+		OutString.Appendf(TEXT("%f, %f, %f, %f"), Value.X, Value.Y, Value.Z, Value.W);
+	}
 };
 
 template<>
@@ -174,6 +199,11 @@ struct FNDIArrayImplHelper<FLinearColor> : public FNDIArrayImplHelperBase<FLinea
 
 	static const FNiagaraTypeDefinition& GetTypeDefinition() { return FNiagaraTypeDefinition::GetColorDef(); }
 	static const FLinearColor GetDefaultValue() { return FLinearColor::White; }
+
+	static void AppendValueToString(const FLinearColor& Value, FString& OutString)
+	{
+		OutString.Appendf(TEXT("%f, %f, %f, %f"), Value.R, Value.G, Value.B, Value.A);
+	}
 };
 
 template<>
@@ -213,6 +243,11 @@ struct FNDIArrayImplHelper<FQuat4f> : public FNDIArrayImplHelperBase<FQuat4f>
 			Dest[i] = FQuat(Src[i]);
 		}
 	}
+
+	static void AppendValueToString(const FQuat4f& Value, FString& OutString)
+	{
+		OutString.Appendf(TEXT("%f, %f, %f, %f"), Value.X, Value.Y, Value.Z, Value.W);
+	}
 };
 
 UCLASS(EditInlineNew, Category = "Array", meta = (DisplayName = "Float Array"), Blueprintable, BlueprintType)
diff --git a/Engine/Plugins/FX/Niagara/Source/Niagara/Classes/NiagaraDataInterfaceArrayImpl.h b/Engine/Plugins/FX/Niagara/Source/Niagara/Classes/NiagaraDataInterfaceArrayImpl.h
--- a/Engine/Plugins/FX/Niagara/Source/Niagara/Classes/NiagaraDataInterfaceArrayImpl.h
+++ b/Engine/Plugins/FX/Niagara/Source/Niagara/Classes/NiagaraDataInterfaceArrayImpl.h
@@ -998,8 +998,29 @@ struct FNDIArrayProxyImpl : public INDIArrayProxyBase
 		{
 			return;
 		}
+
 		FReadArrayRef ArrayData(Owner, InstanceData);
-		VariableDataString = FString::Printf(TEXT("ArrayType(%s) CpuLength(%d)"), *FNDIArrayImplHelper<TArrayType>::GetTypeDefinition().GetName(), ArrayData.GetArray().Num());
+		FString CpuValuesString;
+
+		const int32 MaxStringElements = 8;
+		const int32 NumElements = FMath::Min(MaxStringElements, ArrayData.GetArray().Num());
+		for (int32 i=0; i < NumElements; ++i)
+		{
+			CpuValuesString.Append(i > 0 ? TEXT(", [") : TEXT("["));
+			FNDIArrayImplHelper<TArrayType>::AppendValueToString(ArrayData.GetArray()[i], CpuValuesString);
+			CpuValuesString.Append(TEXT("]"));
+		}
+		if (MaxStringElements < ArrayData.GetArray().Num())
+		{
+			CpuValuesString.Append(TEXT(", ..."));
+		}
+
+		VariableDataString = FString::Printf(
+			TEXT("Type(%s) CpuLength(%d) CpuValues(%s)"),
+			*FNDIArrayImplHelper<TArrayType>::GetTypeDefinition().GetName(),
+			ArrayData.GetArray().Num(),
+			*CpuValuesString
+		);
 	}
 #endif
 
diff --git a/Engine/Plugins/FX/Niagara/Source/Niagara/Classes/NiagaraDataInterfaceArrayInt.h b/Engine/Plugins/FX/Niagara/Source/Niagara/Classes/NiagaraDataInterfaceArrayInt.h
--- a/Engine/Plugins/FX/Niagara/Source/Niagara/Classes/NiagaraDataInterfaceArrayInt.h
+++ b/Engine/Plugins/FX/Niagara/Source/Niagara/Classes/NiagaraDataInterfaceArrayInt.h
@@ -47,6 +47,11 @@ struct FNDIArrayImplHelper<int32> : public FNDIArrayImplHelperBase<int32>
 		}
 		return PrevValue;
 	}
+
+	static void AppendValueToString(const int32 Value, FString& OutString)
+	{
+		OutString.AppendInt(Value);
+	}
 };
 
 template<>
@@ -86,6 +91,11 @@ struct FNDIArrayImplHelper<uint8> : public FNDIArrayImplHelperBase<uint8>
 			Dest[i] = int32(Src[i]);
 		}
 	}
+
+	static void AppendValueToString(const uint8 Value, FString& OutString)
+	{
+		OutString.AppendInt(Value);
+	}
 };
 
 template<>
@@ -123,6 +133,11 @@ struct FNDIArrayImplHelper<bool> : public FNDIArrayImplHelperBase<bool>
 			*TypedDest++ = *TypedSrc++ == 0 ? false : true;
 		}
 	}
+
+	static void AppendValueToString(const bool Value, FString& OutString)
+	{
+		OutString.Append(Value ? TEXT("True") : TEXT("False"));
+	}
 };
 
 UCLASS(EditInlineNew, Category = "Array", meta = (DisplayName = "Int32 Array"), Blueprintable, BlueprintType)
diff --git a/Engine/Plugins/FX/Niagara/Source/Niagara/Classes/NiagaraDataInterfaceArrayNiagaraID.h b/Engine/Plugins/FX/Niagara/Source/Niagara/Classes/NiagaraDataInterfaceArrayNiagaraID.h
--- a/Engine/Plugins/FX/Niagara/Source/Niagara/Classes/NiagaraDataInterfaceArrayNiagaraID.h
+++ b/Engine/Plugins/FX/Niagara/Source/Niagara/Classes/NiagaraDataInterfaceArrayNiagaraID.h
@@ -22,6 +22,11 @@ struct FNDIArrayImplHelper<FNiagaraID> : public FNDIArrayImplHelperBase<FNiagara
 
 	static const FNiagaraTypeDefinition GetTypeDefinition() { return FNiagaraTypeDefinition(FNiagaraID::StaticStruct()); }
 	static const FNiagaraID GetDefaultValue() { return FNiagaraID(); }
+
+	static void AppendValueToString(const FNiagaraID Value, FString& OutString)
+	{
+		OutString.Appendf(TEXT("%d, %d"), Value.Index, Value.AcquireTag);
+	}
 };
 
 UCLASS(EditInlineNew, Category = "Array", meta = (DisplayName = "NiagaraID Array", Experimental), Blueprintable, BlueprintType)
