- Fix for FindClosestUnitDistanceFromPosition not working on GPU or when LUT is enabled

diff --git a/Engine/Plugins/FX/Niagara/Shaders/Private/NiagaraDataInterfaceSplineTemplate.ush b/Engine/Plugins/FX/Niagara/Shaders/Private/NiagaraDataInterfaceSplineTemplate.ush
--- a/Engine/Plugins/FX/Niagara/Shaders/Private/NiagaraDataInterfaceSplineTemplate.ush
+++ b/Engine/Plugins/FX/Niagara/Shaders/Private/NiagaraDataInterfaceSplineTemplate.ush
@@ -2,6 +2,7 @@
 	
 float4x4 {ParameterName}_SplineTransform;
 float4x4 {ParameterName}_SplineTransformRotationMat;
+float4x4 {ParameterName}_SplineTransformInverse;
 float4x4 {ParameterName}_SplineTransformInverseTranspose;
 float4	{ParameterName}_SplineTransformRotation;
 float3	{ParameterName}_DefaultUpVector;
@@ -111,25 +112,25 @@ float3 EvaluateDerivativePosition_{ParameterName}(float InDistance)
 	return {ParameterName}_SplinePositionsLUT[NextKey] - {ParameterName}_SplinePositionsLUT[PrevKey];
 }
 
-float3 EvaluateFindNearestPosition_{ParameterName}(float3 InPosition)
+float EvaluateFindNearestPosition_{ParameterName}(float3 InPosition)
 {
 	if ({ParameterName}_MaxIndex < 0)
 	{
 		return 0.0f;
 	}
 	
-	float MinDistance = dot(InPosition, {ParameterName}_SplinePositionsLUT[0].xyz);
+	float MinDistance = length2({ParameterName}_SplinePositionsLUT[0].xyz - InPosition);
 	float KeyToNearest = 0.0f;
-	for (int Index = 1; Index <= {ParameterName}_MaxIndex; Index++)
+	for (int i = 1; i <= {ParameterName}_MaxIndex; i++)
 	{
-		const float Distance = dot(InPosition, {ParameterName}_SplinePositionsLUT[Index].xyz);
+		const float Distance = length2({ParameterName}_SplinePositionsLUT[i].xyz - InPosition);
 		if (Distance < MinDistance)
 		{
 			MinDistance = Distance;
-			KeyToNearest = Index * {ParameterName}_SplineDistanceStep;
+			KeyToNearest = i;
 		}
 	}
-	return KeyToNearest;
+	return KeyToNearest > 0 ? float(KeyToNearest) / float({ParameterName}_MaxIndex) : 0.0f;
 }
 
 void SampleSplinePositionByUnitDistance_{ParameterName}(float U, out float3 Position)
@@ -214,7 +215,8 @@ void GetSplineLocalToWorldInverseTransposed_{ParameterName}(out float4x4 Transfo
 
 void FindClosestUnitDistanceFromPositionWS_{ParameterName}(float3 PositionWS, out float U)
 {
-	U = EvaluateFindNearestPosition_{ParameterName}(U * {ParameterName}_SplineLength);
+	float3 Position = mul(float4(PositionWS, 1.0f), {ParameterName}_SplineTransformInverse).xyz;
+	U = EvaluateFindNearestPosition_{ParameterName}(Position);
 }
 
 void GetSplineLength_{ParameterName}(out float Length)
diff --git a/Engine/Plugins/FX/Niagara/Source/Niagara/Classes/NiagaraDataInterfaceSpline.h b/Engine/Plugins/FX/Niagara/Source/Niagara/Classes/NiagaraDataInterfaceSpline.h
--- a/Engine/Plugins/FX/Niagara/Source/Niagara/Classes/NiagaraDataInterfaceSpline.h
+++ b/Engine/Plugins/FX/Niagara/Source/Niagara/Classes/NiagaraDataInterfaceSpline.h
@@ -21,6 +21,7 @@ struct FNDISpline_InstanceData_RenderThread
 	
 	FMatrix44f SplineTransform;
 	FMatrix44f SplineTransformRotationMat;
+	FMatrix44f SplineTransformInverse;
 	FMatrix44f SplineTransformInverseTranspose;
 	FQuat4f SplineTransformRotation;
 
@@ -45,9 +46,9 @@ struct FNDISpline_InstanceData_RenderThread
 		SplineScalesLUT.Release();
 		SplineRotationsLUT.Release();
 
-
 		SplineTransform = FMatrix44f::Identity;
 		SplineTransformRotationMat = FMatrix44f::Identity;
+		SplineTransformInverse = FMatrix44f::Identity;
 		SplineTransformInverseTranspose = FMatrix44f::Identity;
 		SplineTransformRotation = FQuat4f::Identity;
 
@@ -154,12 +155,12 @@ struct FNDISpline_InstanceData
 	FVector GetUpVectorAtDistanceAlongSpline(float Distance, ESplineCoordinateSpace::Type CoordinateSpace) const;
 	template<typename UseLUT>
 	FVector GetRightVectorAtDistanceAlongSpline(float Distance, ESplineCoordinateSpace::Type CoordinateSpace) const;
-	
-	template<typename UseLUT>
-	float FindInputKeyClosestToWorldLocation(const FVector& WorldLocation) const;
 
+	template<typename UseLUT>
+	float GetFinalKeyTime() const;
 
-	FInterpCurveVector& GetSplinePointsPosition() { return SplineCurves.Position; }
+	template<typename UseLUT>
+	float FindInputKeyClosestToWorldLocation(const FVector& WorldLocation) const;
 
 private:
 	template<typename UseLUT>
diff --git a/Engine/Plugins/FX/Niagara/Source/Niagara/Private/NiagaraDataInterfaceSpline.cpp b/Engine/Plugins/FX/Niagara/Source/Niagara/Private/NiagaraDataInterfaceSpline.cpp
--- a/Engine/Plugins/FX/Niagara/Source/Niagara/Private/NiagaraDataInterfaceSpline.cpp
+++ b/Engine/Plugins/FX/Niagara/Source/Niagara/Private/NiagaraDataInterfaceSpline.cpp
@@ -31,6 +31,7 @@ namespace NDISplineLocal
 	BEGIN_SHADER_PARAMETER_STRUCT(FShaderParameters, )
 		SHADER_PARAMETER(FMatrix44f,	SplineTransform)
 		SHADER_PARAMETER(FMatrix44f,	SplineTransformRotationMat)
+		SHADER_PARAMETER(FMatrix44f,	SplineTransformInverse)
 		SHADER_PARAMETER(FMatrix44f,	SplineTransformInverseTranspose)
 		SHADER_PARAMETER(FQuat4f,		SplineTransformRotation)
 		SHADER_PARAMETER(FVector3f,		DefaultUpVector)
@@ -451,6 +452,7 @@ void UNiagaraDataInterfaceSpline::SetShaderParameters(const FNiagaraDataInterfac
 	{
 		ShaderParameters->SplineTransform					= InstanceData_RT->SplineTransform;
 		ShaderParameters->SplineTransformRotationMat		= InstanceData_RT->SplineTransformRotationMat;
+		ShaderParameters->SplineTransformInverse			= InstanceData_RT->SplineTransformInverse;
 		ShaderParameters->SplineTransformInverseTranspose	= InstanceData_RT->SplineTransformInverseTranspose;
 		ShaderParameters->SplineTransformRotation			= InstanceData_RT->SplineTransformRotation;
 		ShaderParameters->DefaultUpVector					= InstanceData_RT->DefaultUpVector;
@@ -466,6 +468,7 @@ void UNiagaraDataInterfaceSpline::SetShaderParameters(const FNiagaraDataInterfac
 	{
 		ShaderParameters->SplineTransform					= FMatrix44f::Identity;
 		ShaderParameters->SplineTransformRotationMat		= FMatrix44f::Identity;
+		ShaderParameters->SplineTransformInverse			= FMatrix44f::Identity;
 		ShaderParameters->SplineTransformInverseTranspose	= FMatrix44f::Identity;
 		ShaderParameters->SplineTransformRotation			= FQuat4f::Identity;
 		ShaderParameters->DefaultUpVector					= FVector3f::UnitZ();
@@ -640,9 +643,10 @@ bool UNiagaraDataInterfaceSpline::PerInstanceTick(void* PerInstanceData, FNiagar
 			FNiagaraDataInterfaceProxySpline* RT_Proxy = GetProxyAs<FNiagaraDataInterfaceProxySpline>();	
 			InstData->bSyncedGPUCopy = true;
 			
+
 			// Push Updates to Proxy.
 			ENQUEUE_RENDER_COMMAND(FUpdateDIColorCurve)(
-				[RT_Proxy, InstanceId = SystemInstance->GetId(), Transform = InstData->Transform, TransformRot = InstData->TransformQuat, TransformInverseTranspose = InstData->TransformInverseTransposed, DefaultUp = InstData->DefaultUpVector, rtShaderLUT = InstData->SplineLUT](FRHICommandListImmediate& RHICmdList)
+				[RT_Proxy, InstanceId=SystemInstance->GetId(), Transform=InstData->Transform, TransformRot=InstData->TransformQuat, DefaultUp=InstData->DefaultUpVector, rtShaderLUT=InstData->SplineLUT](FRHICommandListImmediate& RHICmdList)
 			{
 
 				FNDISpline_InstanceData_RenderThread* TargetData = &RT_Proxy->SystemInstancesToProxyData_RT.FindOrAdd(InstanceId);
@@ -650,7 +654,8 @@ bool UNiagaraDataInterfaceSpline::PerInstanceTick(void* PerInstanceData, FNiagar
 				TargetData->SplineTransform = FMatrix44f(Transform);			// LWC_TODO: Precision loss
 				TargetData->SplineTransformRotationMat =  FMatrix44f(Transform.RemoveTranslation());
 				TargetData->SplineTransformRotationMat.RemoveScaling();
-				TargetData->SplineTransformInverseTranspose = FMatrix44f(TransformInverseTranspose);
+				TargetData->SplineTransformInverse = TargetData->SplineTransform.Inverse();
+				TargetData->SplineTransformInverseTranspose = TargetData->SplineTransformInverse.GetTransposed();
 				TargetData->SplineTransformRotation = FQuat4f(TransformRot);
 					
 				TargetData->DefaultUpVector = (FVector3f)DefaultUp;
@@ -994,23 +999,19 @@ float FNDISpline_InstanceData::EvaluateFindNearestPosition<TIntegralConstant<boo
 {
 	// This is a brute force search, definitely not a great idea with large tables, but also not too many ways around it without more data.
 	float MinDistance = TNumericLimits<float>::Max();
-	float KeyToNearest = 0.0f;
-	for (int32 Index = 0; Index < SplineLUT.Positions.Num(); Index++)
+	int32 KeyToNearest = 0;
+	for (int32 i=0; i < SplineLUT.Positions.Num(); i++)
 	{
-		const float Distance = FVector::DistSquared(InPosition, SplineLUT.Positions[Index]);
+		const float Distance = FVector::DistSquared(InPosition, SplineLUT.Positions[i]);
 		if (Distance < MinDistance)
 		{
 			MinDistance = Distance;
-			KeyToNearest = Index * SplineLUT.SplineDistanceStep;
+			KeyToNearest = i;
 		}
 	}
-	return KeyToNearest;
+	return KeyToNearest > 0 ? float(KeyToNearest) / float(SplineLUT.Positions.Num() - 1) : 0.0f;
 }
 
-
-
-
-
 template<typename UseLUT>
 FVector FNDISpline_InstanceData::GetTangentAtDistanceAlongSpline(float Distance, ESplineCoordinateSpace::Type CoordinateSpace) const
 {
@@ -1041,6 +1042,19 @@ FVector FNDISpline_InstanceData::GetDirectionAtDistanceAlongSpline(float Distanc
 	return Direction;
 }
 
+template <>
+float FNDISpline_InstanceData::GetFinalKeyTime<TIntegralConstant<bool, false>>() const
+{
+	const int32 NumPoints = SplineCurves.Position.Points.Num();
+	return NumPoints > 0 ? SplineCurves.Position.Points[NumPoints - 1].InVal : 1.0f;
+}
+
+template <>
+float FNDISpline_InstanceData::GetFinalKeyTime<TIntegralConstant<bool, true>>() const
+{
+	return 1.0f;
+}
+
 template<typename UseLUT>
 float FNDISpline_InstanceData::FindInputKeyClosestToWorldLocation(const FVector& WorldLocation) const
 {
@@ -1048,7 +1062,6 @@ float FNDISpline_InstanceData::FindInputKeyClosestToWorldLocation(const FVector&
 	return EvaluateFindNearestPosition<UseLUT>(LocalLocation);
 }
 
-
 template<typename UseLUT, typename TransformHandlerType, typename SplineSampleType>
 void UNiagaraDataInterfaceSpline::SampleSplinePositionByUnitDistance(FVectorVMExternalFunctionContext& Context)
 {
@@ -1374,19 +1387,17 @@ void UNiagaraDataInterfaceSpline::FindClosestUnitDistanceFromPositionWS(FVectorV
 
 	if (InstData->IsValid())
 	{
-
-		const int32 NumPoints = InstData->GetSplinePointsPosition().Points.Num();
-		const float FinalKeyTime = InstData->GetSplinePointsPosition().Points[NumPoints - 1].InVal;
+		const float FinalKeyTime = InstData->GetFinalKeyTime<UseLUT>();
 
 		for (int32 i = 0; i < Context.GetNumInstances(); ++i)
 		{
-			FNiagaraPosition SimPos = PosParam.GetAndAdvance();
-			FVector WorldPos = InstData->LwcConverter.ConvertSimulationPositionToWorld(SimPos);
+			const FNiagaraPosition SimPos = PosParam.GetAndAdvance();
+			const FVector WorldPos = InstData->LwcConverter.ConvertSimulationPositionToWorld(SimPos);
 
 			// This first call finds the key time, but this is not in 0..1 range for the spline. 
-			float KeyTime = InstData->FindInputKeyClosestToWorldLocation<UseLUT>(WorldPos);
+			const float KeyTime = InstData->FindInputKeyClosestToWorldLocation<UseLUT>(WorldPos);
 			// We need to convert into the range by dividing through by the overall duration of the spline according to the keys.
-			float UnitDistance = KeyTime / FinalKeyTime;
+			const float UnitDistance = KeyTime / FinalKeyTime;
 
 			*OutUnitDistance.GetDest() = UnitDistance;
 			OutUnitDistance.Advance();
