Static Mesh DI LOD support.

- Adding LODIndex to static mesh DI. -ve values indicate we index backwards from NumLODs. If the desirecd LOD is unavailable we fallback to the next available LOD.
- Adding optional override allowing DIs to choose whether they can sample from streaming LODs.

diff --git a/Engine/Plugins/FX/Niagara/Source/Niagara/Private/DataInterface/NiagaraDataInterfaceStaticMesh.cpp b/Engine/Plugins/FX/Niagara/Source/Niagara/Private/DataInterface/NiagaraDataInterfaceStaticMesh.cpp
--- a/Engine/Plugins/FX/Niagara/Source/Niagara/Private/DataInterface/NiagaraDataInterfaceStaticMesh.cpp
+++ b/Engine/Plugins/FX/Niagara/Source/Niagara/Private/DataInterface/NiagaraDataInterfaceStaticMesh.cpp
@@ -100,11 +100,14 @@ namespace NDIStaticMeshLocal
 		ECVF_Default
 	);
 
-	static int32 GNDIStaticMesh_UseInlineLODsOnly = 1;
+	static int32 GNDIStaticMesh_UseInlineLODsOnly = 2;
 	static FAutoConsoleVariableRef CVarNDIStaticMesh_UseInlineLODsOnly(
 		TEXT("fx.Niagara.NDIStaticMesh.UseInlineLODsOnly"),
 		GNDIStaticMesh_UseInlineLODsOnly,
-		TEXT("When enabled Niagara will never use streaming LOD levels, only inline LODs."),
+		TEXT("When enabled Niagara will never use streaming LOD levels, only inline LODs."
+		"0 = Streaming LODs can be sampled."
+		"1 = Only inlined LODs can be sampled."
+		"2 = Only inlined LODs can be sampled by default but each DI can override this if desired. "),
 		ECVF_Default
 	);
 
@@ -742,7 +745,7 @@ namespace NDIStaticMeshLocal
 			{
 				// Check if any valid LODs are found. If not, we won't use this mesh
 				MinLOD = StaticMesh->GetMinLOD().GetValue();
-				if (GNDIStaticMesh_UseInlineLODsOnly)
+				if (GNDIStaticMesh_UseInlineLODsOnly == 1 || (GNDIStaticMesh_UseInlineLODsOnly == 2 && !Interface->bAllowSamplingFromStreamingLODs))
 				{
 					MinLOD = StaticMesh->GetNumLODs() - StaticMesh->GetRenderData()->NumInlinedLODs;
 				}
@@ -750,7 +753,11 @@ namespace NDIStaticMeshLocal
 				FStaticMeshRenderData* RenderData = StaticMesh->GetRenderData();
 				if (RenderData)
 				{
-					CachedLODIdx = RenderData->GetCurrentFirstLODIdx(MinLOD);
+					int32 MinSampleLOD = RenderData->GetCurrentFirstLODIdx(MinLOD);
+					int32 MaxSampleLOD = StaticMesh->GetNumLODs() - 1;
+
+					CachedLODIdx = Interface->LODIndex >= 0 ? Interface->LODIndex : (StaticMesh->GetNumLODs() + Interface->LODIndex);
+					CachedLODIdx = FMath::Clamp(CachedLODIdx, MinSampleLOD, MaxSampleLOD);
 					if (RenderData->LODResources.IsValidIndex(CachedLODIdx))
 					{
 						LODData = &RenderData->LODResources[CachedLODIdx];
@@ -1008,39 +1015,53 @@ namespace NDIStaticMeshLocal
 			return false;
 		}
 
-		const FStaticMeshLODResources* GetCurrentFirstLOD()
+		const FStaticMeshLODResources* GetCurrentLOD()
 		{
 			//-OPT: Perhaps we could cache this during the tick function?
 			UStaticMesh* StaticMesh = StaticMeshWeakPtr.Get();
 			if ( bMeshValid && StaticMesh )
 			{
-				if ( const FStaticMeshLODResources* LODResource = StaticMesh->GetRenderData()->GetCurrentFirstLOD(MinLOD) )
+				int32 CurrentFirstLOD = StaticMesh->GetRenderData()->GetCurrentFirstLODIdx(MinLOD);
+				if(CachedLODIdx >= CurrentFirstLOD && StaticMesh->GetRenderData()->LODResources.IsValidIndex(CachedLODIdx))
 				{
-					return LODResource;
+					if (const FStaticMeshLODResources* LODResource = &StaticMesh->GetRenderData()->LODResources[CachedLODIdx])
+					{
+						if(LODResource->GetNumVertices() > 0 && LODResource->BuffersSize > 0)
+						{
+							return LODResource;
+						}
+					}
 				}
 			}
 			return nullptr;
 		}
 
-		const FStaticMeshLODResources* GetCurrentFirstLODWithVertexColorOverrides(FColorVertexBuffer*& OutVertexColorOverrides)
+		const FStaticMeshLODResources* GetCurrentLODWithVertexColorOverrides(FColorVertexBuffer*& OutVertexColorOverrides)
 		{
 			//-OPT: Perhaps we could cache this during the tick function?
 			UStaticMesh* StaticMesh = StaticMeshWeakPtr.Get();
 			OutVertexColorOverrides = nullptr;
 			if (bMeshValid && StaticMesh)
 			{
-				if(UStaticMeshComponent* SMComp = Cast<UStaticMeshComponent>(SceneComponentWeakPtr.Get()))
+				int32 CurrentFirstLOD = StaticMesh->GetRenderData()->GetCurrentFirstLODIdx(MinLOD);
+
+				if (CachedLODIdx >= CurrentFirstLOD && StaticMesh->GetRenderData()->LODResources.IsValidIndex(CachedLODIdx))
 				{
-					int32 LODIdx = StaticMesh->GetRenderData()->GetCurrentFirstLODIdx(MinLOD);
-					if(SMComp->LODData.IsValidIndex(LODIdx))
+					if(UStaticMeshComponent* SMComp = Cast<UStaticMeshComponent>(SceneComponentWeakPtr.Get()))
 					{
-						OutVertexColorOverrides = SMComp->LODData[LODIdx].OverrideVertexColors;
+						if(SMComp->LODData.IsValidIndex(CachedLODIdx))
+						{
+							OutVertexColorOverrides = SMComp->LODData[CachedLODIdx].OverrideVertexColors;
+						}
+					}
+				
+					if (const FStaticMeshLODResources* LODResource = &StaticMesh->GetRenderData()->LODResources[CachedLODIdx])
+					{
+						if (LODResource->GetNumVertices() > 0 && LODResource->BuffersSize > 0)
+						{
+							return LODResource;
+						}
 					}
-				}
-
-				if (const FStaticMeshLODResources* LODResource = StaticMesh->GetRenderData()->GetCurrentFirstLOD(MinLOD))
-				{
-					return LODResource;
 				}
 
 			}
@@ -1106,7 +1127,7 @@ namespace NDIStaticMeshLocal
 		FORCEINLINE FStaticMeshCpuHelper(FVectorVMExternalFunctionContext& Context)
 			: InstanceData(Context)
 		{
-			LODResource = InstanceData->GetCurrentFirstLODWithVertexColorOverrides(OverrideVertexColors);
+			LODResource = InstanceData->GetCurrentLODWithVertexColorOverrides(OverrideVertexColors);
 		}
 
 
@@ -1605,9 +1626,9 @@ bool UNiagaraDataInterfaceStaticMesh::InitPerInstanceData(void* PerInstanceData,
 		//This is safe to ref on the RT as it's freed on the RT in FStaticMeshComponentLODInfo::BeginReleaseOverrideVertexColors()
 		//However, it's seems unsafe to reference in Niagara's instance data this way as there looks to be a window between this RT command and it's actual use where the data could have been freed.
 #if WITH_EDITOR
-		GpuInitializeData->LODResource = InstanceData->GetCurrentFirstLOD();
+		GpuInitializeData->LODResource = InstanceData->GetCurrentLOD();
 #else
-		GpuInitializeData->LODResource = InstanceData->GetCurrentFirstLODWithVertexColorOverrides(GpuInitializeData->OverrideColorBuffer);
+		GpuInitializeData->LODResource = InstanceData->GetCurrentLODWithVertexColorOverrides(GpuInitializeData->OverrideColorBuffer);
 #endif
 
 		if ( GpuInitializeData->LODResource )
@@ -1749,7 +1770,7 @@ void UNiagaraDataInterfaceStaticMesh::ProvidePerInstanceDataForRenderThread(void
 	DataFromGT->UvMappingSet				= InstanceData->UvMapping.GetUvSetIndex();
 
 #if WITH_EDITOR
-	InstanceData->GetCurrentFirstLODWithVertexColorOverrides(DataFromGT->OverrideVertexColors);
+	InstanceData->GetCurrentLODWithVertexColorOverrides(DataFromGT->OverrideVertexColors);
 #endif
 
 	if ( UPrimitiveComponent* PrimitiveComponent = Cast<UPrimitiveComponent>(InstanceData->SceneComponentWeakPtr) )
@@ -3066,6 +3087,8 @@ bool UNiagaraDataInterfaceStaticMesh::Equals(const UNiagaraDataInterface* Other)
 		OtherTyped->SourceComponent == SourceComponent &&
 		OtherTyped->SectionFilter.AllowedMaterialSlots == SectionFilter.AllowedMaterialSlots &&
 		OtherTyped->bUsePhysicsBodyVelocity == bUsePhysicsBodyVelocity &&
+		OtherTyped->bAllowSamplingFromStreamingLODs == bAllowSamplingFromStreamingLODs &&
+		OtherTyped->LODIndex == LODIndex &&
 		OtherTyped->FilteredSockets == FilteredSockets;
 }
 
@@ -3088,6 +3111,8 @@ bool UNiagaraDataInterfaceStaticMesh::CopyToInternal(UNiagaraDataInterface* Dest
 	OtherTyped->SectionFilter = SectionFilter;
 	OtherTyped->bUsePhysicsBodyVelocity = bUsePhysicsBodyVelocity;
 	OtherTyped->FilteredSockets = FilteredSockets;
+	OtherTyped->bAllowSamplingFromStreamingLODs = bAllowSamplingFromStreamingLODs;
+	OtherTyped->LODIndex = LODIndex;
 	OtherTyped->BindSourceDelegates();
 	return true;
 }
diff --git a/Engine/Plugins/FX/Niagara/Source/Niagara/Private/DataInterface/NiagaraDataInterfaceStaticMesh.h b/Engine/Plugins/FX/Niagara/Source/Niagara/Private/DataInterface/NiagaraDataInterfaceStaticMesh.h
--- a/Engine/Plugins/FX/Niagara/Source/Niagara/Private/DataInterface/NiagaraDataInterfaceStaticMesh.h
+++ b/Engine/Plugins/FX/Niagara/Source/Niagara/Private/DataInterface/NiagaraDataInterfaceStaticMesh.h
@@ -94,6 +94,18 @@ public:
 	UPROPERTY(EditAnywhere, AdvancedDisplay, Category = "Mesh")
     bool bUsePhysicsBodyVelocity = false;
 
+	/** When true, we allow this DI to sample from streaming LODs. Selectively overriding the CVar fx.Niagara.NDIStaticMesh.UseInlineLODsOnly. */
+	UPROPERTY(EditAnywhere, Category = "LOD")
+	bool bAllowSamplingFromStreamingLODs = false;
+
+	/** 
+	Static Mesh LOD to sample.
+	When the desired LOD is not available, the next available LOD is used.
+	When LOD Index is negative, Desired LOD = Num LODs - LOD Index.
+	*/
+	UPROPERTY(EditAnywhere, Category = "LOD")
+	int32 LODIndex = 0;
+
 	/** List of filtered sockets to use. */
 	UPROPERTY(EditAnywhere, Category = "Mesh")
 	TArray<FName> FilteredSockets;
