Extended python api for module versioning

diff --git a/Engine/Plugins/FX/Niagara/Source/NiagaraEditor/Private/UpgradeNiagaraScriptResults.cpp b/Engine/Plugins/FX/Niagara/Source/NiagaraEditor/Private/UpgradeNiagaraScriptResults.cpp
--- a/Engine/Plugins/FX/Niagara/Source/NiagaraEditor/Private/UpgradeNiagaraScriptResults.cpp
+++ b/Engine/Plugins/FX/Niagara/Source/NiagaraEditor/Private/UpgradeNiagaraScriptResults.cpp
@@ -75,6 +75,11 @@ bool UNiagaraPythonScriptModuleInput::IsLocalValue() const
 	return IsSet() && Input->ValueMode == ENiagaraClipboardFunctionInputValueMode::Local;
 }
 
+bool UNiagaraPythonScriptModuleInput::IsLinkedValue() const
+{
+	return IsSet() && Input->ValueMode == ENiagaraClipboardFunctionInputValueMode::Linked;
+}
+
 float UNiagaraPythonScriptModuleInput::AsFloat() const
 {
 	if (IsSet() && Input->InputType == FNiagaraTypeDefinition::GetFloatDef())
@@ -163,6 +168,15 @@ FString UNiagaraPythonScriptModuleInput::AsEnum() const
 	return FString();
 }
 
+FString UNiagaraPythonScriptModuleInput::AsLinkedValue() const
+{
+	if (IsLinkedValue())
+	{
+		return Input->Linked.ToString();
+	}
+	return FString();
+}
+
 UUpgradeNiagaraScriptResults::UUpgradeNiagaraScriptResults()
 {
 	DummyInput = NewObject<UNiagaraPythonScriptModuleInput>();
@@ -277,6 +291,15 @@ void UUpgradeNiagaraScriptResults::SetEnumInput(const FString& InputName, FStrin
 	}
 }
 
+void UUpgradeNiagaraScriptResults::SetLinkedInput(const FString& InputName, FString Value)
+{
+	if (UNiagaraPythonScriptModuleInput* ModuleInput = GetNewInput(FName(InputName)))
+	{
+		const UNiagaraClipboardFunctionInput* CurrentInput = ModuleInput->Input;
+		ModuleInput->Input = UNiagaraClipboardFunctionInput::CreateLinkedValue(ModuleInput, CurrentInput->InputName, CurrentInput->InputType, CurrentInput->bEditConditionValue, FName(Value));
+	}
+}
+
 UNiagaraPythonScriptModuleInput* UUpgradeNiagaraScriptResults::GetNewInput(const FName& InputName) const
 {
 	for (UNiagaraPythonScriptModuleInput* ModuleInput : NewInputs)
diff --git a/Engine/Plugins/FX/Niagara/Source/NiagaraEditor/Public/UpgradeNiagaraScriptResults.h b/Engine/Plugins/FX/Niagara/Source/NiagaraEditor/Public/UpgradeNiagaraScriptResults.h
--- a/Engine/Plugins/FX/Niagara/Source/NiagaraEditor/Public/UpgradeNiagaraScriptResults.h
+++ b/Engine/Plugins/FX/Niagara/Source/NiagaraEditor/Public/UpgradeNiagaraScriptResults.h
@@ -51,6 +51,9 @@ public:
 	
 	UFUNCTION(BlueprintCallable, Category = "Scripting")
     bool IsLocalValue() const;
+    
+    UFUNCTION(BlueprintCallable, Category = "Scripting")
+    bool IsLinkedValue() const;
 
 	UFUNCTION(BlueprintCallable, Category = "Scripting")
     float AsFloat() const;
@@ -78,6 +81,9 @@ public:
 
 	UFUNCTION(BlueprintCallable, Category = "Scripting")
     FString AsEnum() const;
+    
+    UFUNCTION(BlueprintCallable, Category = "Scripting")
+    FString AsLinkedValue() const;
 };
 
 /**
@@ -133,6 +139,9 @@ public:
 
 	UFUNCTION(BlueprintCallable, Category = "Scripting")
     void SetEnumInput(const FString& InputName, FString Value);
+    
+    UFUNCTION(BlueprintCallable, Category = "Scripting")
+    void SetLinkedInput(const FString& InputName, FString Value);
 
 private:
 	UNiagaraPythonScriptModuleInput* GetNewInput(const FName& InputName) const;
