- Fix for Extents vs Half Extents (contribution by (amuTBKT))

#rnx

diff --git a/Engine/Plugins/FX/Niagara/Shaders/Private/NiagaraDataInterfaceStaticMeshTemplate.ush b/Engine/Plugins/FX/Niagara/Shaders/Private/NiagaraDataInterfaceStaticMeshTemplate.ush
--- a/Engine/Plugins/FX/Niagara/Shaders/Private/NiagaraDataInterfaceStaticMeshTemplate.ush
+++ b/Engine/Plugins/FX/Niagara/Shaders/Private/NiagaraDataInterfaceStaticMeshTemplate.ush
@@ -68,7 +68,7 @@ void GetPreSkinnedLocalBounds_{ParameterName}(out float3 OutCenter, out float3 O
 {
 	OutCenter		= {ParameterName}_PreSkinnedLocalBoundsCenter;
 	OutHalfExtents	= {ParameterName}_PreSkinnedLocalBoundsExtents;
-	OutExtents		= OutExtents * 2.0f;
+	OutExtents		= OutHalfExtents * 2.0f;
 	OutExtentsMin	= OutCenter - OutHalfExtents;
 	OutExtentsMax	= OutCenter + OutHalfExtents;
 }
