- Remove editor properties from versioned emitters

#rnx

diff --git a/Engine/Plugins/FX/Niagara/Source/Niagara/Classes/NiagaraEmitter.h b/Engine/Plugins/FX/Niagara/Source/Niagara/Classes/NiagaraEmitter.h
--- a/Engine/Plugins/FX/Niagara/Source/Niagara/Classes/NiagaraEmitter.h
+++ b/Engine/Plugins/FX/Niagara/Source/Niagara/Classes/NiagaraEmitter.h
@@ -243,6 +243,7 @@ struct NIAGARA_API FVersionedNiagaraEmitterData
 	UPROPERTY()
 	FNiagaraAssetVersion Version;
 
+#if WITH_EDITORONLY_DATA
 	/** What changed in this version compared to the last? Displayed to the user when upgrading to a new script version. */
 	UPROPERTY()
 	FText VersionChangeDescription;
@@ -258,6 +259,7 @@ struct NIAGARA_API FVersionedNiagaraEmitterData
 	/** Asset reference to a python script to run when updating to this script version. */
 	UPROPERTY()
 	FFilePath ScriptAsset;
+#endif //WITH_EDITORONLY_DATA
 
 	/* If this emitter is no longer meant to be used, this option should be set.*/
 	UPROPERTY()
