Variety of fixes for RigidBody DI
-Pushes the instance tick to be post simulate so that we're ensured to have results from FindActor in time for the GPU simulation
-Deprecates the global search path and provides feedback warnings if it's in use (along with if the FindActors is missing)
-Removes the optimization of disabling the actor search if there weren't any GPU functions being called because of the difficulty in finding bound functions for non-user parameter DIs

diff --git a/Engine/Plugins/FX/Niagara/Source/Niagara/Private/NiagaraDataInterfaceRigidMeshCollisionQuery.cpp b/Engine/Plugins/FX/Niagara/Source/Niagara/Private/NiagaraDataInterfaceRigidMeshCollisionQuery.cpp
--- a/Engine/Plugins/FX/Niagara/Source/Niagara/Private/NiagaraDataInterfaceRigidMeshCollisionQuery.cpp
+++ b/Engine/Plugins/FX/Niagara/Source/Niagara/Private/NiagaraDataInterfaceRigidMeshCollisionQuery.cpp
@@ -1,6 +1,7 @@
 // Copyright Epic Games, Inc. All Rights Reserved.
 
 #include "NiagaraDataInterfaceRigidMeshCollisionQuery.h"
+#include "Algo/Accumulate.h"
 #include "Algo/ForEach.h"
 #include "Algo/RemoveIf.h"
 #include "Animation/SkeletalMeshActor.h"
@@ -21,6 +22,7 @@
 #include "NiagaraShader.h"
 #include "NiagaraShaderParametersBuilder.h"
 #include "NiagaraSimStageData.h"
+#include "NiagaraSystemImpl.h"
 #include "NiagaraSystemInstance.h"
 #include "NiagaraSystemSimulation.h"
 #include "PhysicsEngine/PhysicsAsset.h"
@@ -105,7 +107,7 @@ static const FText OverlapRotationDescription = IF_WITH_EDITORONLY_DATA(
 );
 
 static const FText OverlapExtentDescription = IF_WITH_EDITORONLY_DATA(
-	LOCTEXT("RigidBodyOverlapExtentDescription", "The extent, in world space, of the overlap trace."),
+	LOCTEXT("RigidBodyOverlapExtentDescription", "The full extent (max-min), in world space, of the overlap trace."),
 	FText()
 );
 
@@ -266,11 +268,14 @@ void ForEachBodySetup(UStaticMeshComponent* Component, TBodySetupPredicate Predi
 /// End UStaticMeshComponent
 
 template<typename TComponentType>
-void CountCollisionPrimitives(TConstArrayView<TComponentType*> Components, uint32& BoxCount, uint32& SphereCount, uint32& CapsuleCount)
+void CountCollisionPrimitives(TConstArrayView<TComponentType*> Components, TArray<FNDIRigidMeshCollisionData::FComponentBodyCount>& PerComponentCounts, uint32& TotalBoxCount, uint32& TotalSphereCount, uint32& TotalCapsuleCount)
 {
+	PerComponentCounts.Reserve(Components.Num());
+
 	for (TComponentType* Component : Components)
 	{
-		bool HasConvexElements = false;
+		FNDIRigidMeshCollisionData::FComponentBodyCount& BodyCount = PerComponentCounts.AddDefaulted_GetRef();
+		BodyCount.ComponentHash = GetTypeHash(Component);
 
 		ForEachBodySetup(Component, [&](TComponentType* Component, const UBodySetup* BodySetup)
 		{
@@ -278,29 +283,28 @@ void CountCollisionPrimitives(TConstArrayView<TComponentType*> Components, uint3
 			{
 				if (CollisionEnabledHasPhysics(ConvexElem.GetCollisionEnabled()))
 				{
-					HasConvexElements = true;
-					++BoxCount;
+					++BodyCount.BoxCount;
 				}
 			}
 			for (const FKBoxElem& BoxElem : BodySetup->AggGeom.BoxElems)
 			{
 				if (CollisionEnabledHasPhysics(BoxElem.GetCollisionEnabled()))
 				{
-					++BoxCount;
+					++BodyCount.BoxCount;
 				}
 			}
 			for (const FKSphereElem& SphereElem : BodySetup->AggGeom.SphereElems)
 			{
 				if (CollisionEnabledHasPhysics(SphereElem.GetCollisionEnabled()))
 				{
-					++SphereCount;
+					++BodyCount.SphereCount;
 				}
 			}
 			for (const FKSphylElem& CapsuleElem : BodySetup->AggGeom.SphylElems)
 			{
 				if (CollisionEnabledHasPhysics(CapsuleElem.GetCollisionEnabled()))
 				{
-					++CapsuleCount;
+					++BodyCount.CapsuleCount;
 				}
 			}
 
@@ -310,17 +314,13 @@ void CountCollisionPrimitives(TConstArrayView<TComponentType*> Components, uint3
 				BodySetup->AggGeom.SphereElems.Num() == 0 &&
 				BodySetup->AggGeom.SphylElems.Num() == 0)
 			{
-				++BoxCount;
+				++BodyCount.BoxCount;
 			}
 		});
 
-		// suppress warning since we get it a lot
-		/*
-		if (HasConvexElements)
-		{			
-			UE_LOG(LogRigidMeshCollision, Warning, TEXT("Convex collision objects encountered and will be interpreted as a bounding box on %s"), *Component->GetOwner()->GetName());
-		}
-		*/
+		TotalBoxCount += BodyCount.BoxCount;
+		TotalSphereCount += BodyCount.SphereCount;
+		TotalCapsuleCount += BodyCount.CapsuleCount;
 	}
 }
 
@@ -475,72 +475,185 @@ void UpdateAssetArrays(TConstArrayView<TComponentType*> Components, const FVecto
 	}
 }
 
-void UpdateInternalArrays(TConstArrayView<AActor*> Actors, TConstArrayView<FName> ComponentTags, FVector LWCTile, bool bFullUpdate, FNDIRigidMeshCollisionArrays* OutAssetArrays)
+template<bool FirstStage>
+void RemapPreviousTransforms(
+	TConstArrayView<FNDIRigidMeshCollisionData::FComponentBodyCount> PreviousCounts,
+	TConstArrayView<FNDIRigidMeshCollisionData::FComponentBodyCount> CurrentCounts,
+	FNDIRigidMeshCollisionArrays* OutAssetArrays)
 {
-	if (OutAssetArrays != nullptr && OutAssetArrays->ElementOffsets.NumElements < OutAssetArrays->MaxPrimitives)
+	const FNDIRigidMeshCollisionElementOffset& CurrentOffsets = OutAssetArrays->ElementOffsets;
+
+	uint32 BoxCount = 0;
+	uint32 SphereCount = 0;
+	uint32 CapsuleCount = 0;
+
+	auto CopyComponent = [&BoxCount, &SphereCount, &CapsuleCount]
+		(const TArray<FVector4f>& Src, TArray<FVector4f>& Dst, const FNDIRigidMeshCollisionElementOffset& Offsets, const FNDIRigidMeshCollisionData::FComponentBodyCount& BodyCounts) -> void
 	{
-		TInlineComponentArray<UStaticMeshComponent*> StaticMeshes;
-		TInlineComponentArray<USkeletalMeshComponent*> SkeletalMeshes;
+		const uint32 BoxStartIndex = 3 * (Offsets.BoxOffset + BoxCount);
+		for (uint32 ElementIt = 0; ElementIt < 3 * BodyCounts.BoxCount; ++ElementIt)
+		{
+			Dst[BoxStartIndex + ElementIt] = Src[BoxStartIndex + ElementIt];
+		}
 
-		uint32 BoxCount = 0;
-		uint32 SphereCount = 0;
-		uint32 CapsuleCount = 0;
+		const uint32 SphereStartIndex = 3 * (Offsets.SphereOffset + SphereCount);
+		for (uint32 ElementIt = 0; ElementIt < 3 * BodyCounts.SphereCount; ++ElementIt)
+		{
+			Dst[SphereStartIndex + ElementIt] = Src[SphereStartIndex + ElementIt];
+		}
 
-		CollectComponents(Actors, ComponentTags, StaticMeshes);
-		CollectComponents(Actors, ComponentTags, SkeletalMeshes);
+		const uint32 CapsuleStartIndex = 3 * (Offsets.CapsuleOffset + CapsuleCount);
+		for (uint32 ElementIt = 0; ElementIt < 3 * BodyCounts.CapsuleCount; ++ElementIt)
+		{
+			Dst[CapsuleStartIndex + ElementIt] = Src[CapsuleStartIndex + ElementIt];
+		}
+	};
 
-		TConstArrayView<UStaticMeshComponent*> StaticMeshView = MakeArrayView(StaticMeshes.GetData(), StaticMeshes.Num());
-		TConstArrayView<USkeletalMeshComponent*> SkeletalMeshView = MakeArrayView(SkeletalMeshes.GetData(), SkeletalMeshes.Num());
+	int32 PreviousIndex = 0;
+	for (const FNDIRigidMeshCollisionData::FComponentBodyCount& CurrentCount : CurrentCounts)
+	{
+		bool bPreviousValuesCopied = false;
 
-		CountCollisionPrimitives(StaticMeshView, BoxCount, SphereCount, CapsuleCount);
-		CountCollisionPrimitives(SkeletalMeshView, BoxCount, SphereCount, CapsuleCount);
+		while (PreviousCounts.IsValidIndex(PreviousIndex) && !bPreviousValuesCopied)
+		{
+			const FNDIRigidMeshCollisionData::FComponentBodyCount& PreviousCount = PreviousCounts[PreviousIndex];
 
-		if ((BoxCount + SphereCount + CapsuleCount) >= OutAssetArrays->MaxPrimitives)
+			// skip through the previous components until we find a match
+			if (PreviousCount.ComponentHash < CurrentCount.ComponentHash)
+			{
+				++PreviousIndex;
+				continue;
+			}
+
+			if (PreviousCount.ComponentHash == CurrentCount.ComponentHash)
+			{
+				check(PreviousCount.BoxCount == CurrentCount.BoxCount);
+				check(PreviousCount.SphereCount == CurrentCount.SphereCount);
+				check(PreviousCount.CapsuleCount == CurrentCount.CapsuleCount);
+
+				bPreviousValuesCopied = true;
+
+				if (FirstStage)
+				{
+					CopyComponent(OutAssetArrays->CurrentTransform, OutAssetArrays->PreviousTransform, CurrentOffsets, CurrentCount);
+					CopyComponent(OutAssetArrays->CurrentInverse, OutAssetArrays->PreviousInverse, CurrentOffsets, CurrentCount);
+				}
+			}
+
+			break;
+		}
+
+		// copy the defaults over to the previous
+		if (!bPreviousValuesCopied && !FirstStage)
 		{
-			UE_LOG(LogRigidMeshCollision, Error, TEXT("Number of Collision DI primitives is higher than the %d limit.  Please increase it."), OutAssetArrays->MaxPrimitives);
+			CopyComponent(OutAssetArrays->CurrentTransform, OutAssetArrays->PreviousTransform, CurrentOffsets, CurrentCount);
+			CopyComponent(OutAssetArrays->CurrentInverse, OutAssetArrays->PreviousInverse, CurrentOffsets, CurrentCount);
+		}
+
+		BoxCount += CurrentCount.BoxCount;
+		SphereCount += CurrentCount.SphereCount;
+		CapsuleCount += CurrentCount.CapsuleCount;
+	}
+}
+
+void UpdateInternalArrays(
+	TConstArrayView<UStaticMeshComponent*> StaticMeshView,
+	TConstArrayView<USkeletalMeshComponent*> SkeletalMeshView,
+	FVector LWCTile,
+	bool bFullUpdate,
+	TArray<FNDIRigidMeshCollisionData::FComponentBodyCount>& BodyCounts,
+	FNDIRigidMeshCollisionArrays* OutAssetArrays)
+{
+	if (OutAssetArrays != nullptr && OutAssetArrays->ElementOffsets.NumElements < OutAssetArrays->MaxPrimitives)
+	{
+		// when we are in game and we don't need to worry about bodies changing for a given mesh component we can try to optimize
+		// the update by just targeting the dynamic elements (transforms) and using the original values as the previous run
+#if !WITH_EDITOR
+		if (!bFullUpdate)
+		{
+			// if we're updating, then copy over last frame's transforms before we generate new ones
+			Swap(OutAssetArrays->PreviousTransform, OutAssetArrays->CurrentTransform);
+			Swap(OutAssetArrays->PreviousInverse, OutAssetArrays->CurrentInverse);
+
+			uint32 BoxIndex = OutAssetArrays->ElementOffsets.BoxOffset;
+			uint32 SphereIndex = OutAssetArrays->ElementOffsets.SphereOffset;
+			uint32 CapsuleIndex = OutAssetArrays->ElementOffsets.CapsuleOffset;
+
+			UpdateAssetArrays<UStaticMeshComponent, false>(StaticMeshView, LWCTile, OutAssetArrays, BoxIndex, SphereIndex, CapsuleIndex);
+			UpdateAssetArrays<USkeletalMeshComponent, false>(SkeletalMeshView, LWCTile, OutAssetArrays, BoxIndex, SphereIndex, CapsuleIndex);
+
 			return;
 		}
+#endif
 
-		const bool MismatchOffsets = ((OutAssetArrays->ElementOffsets.SphereOffset - OutAssetArrays->ElementOffsets.BoxOffset) != BoxCount) ||
-			((OutAssetArrays->ElementOffsets.CapsuleOffset - OutAssetArrays->ElementOffsets.SphereOffset) != SphereCount) ||
-			((OutAssetArrays->ElementOffsets.NumElements - OutAssetArrays->ElementOffsets.CapsuleOffset) != CapsuleCount);
+		TArray<FNDIRigidMeshCollisionData::FComponentBodyCount> CurrentBodyCounts;
 
-		// if we're only running an update, then make sure that the offsets aren't mismatched
-		check(!MismatchOffsets || bFullUpdate);
+		uint32 TotalBoxCount = 0;
+		uint32 TotalSphereCount = 0;
+		uint32 TotalCapsuleCount = 0;
 
-		if (bFullUpdate)
-		{			
-			OutAssetArrays->ElementOffsets.BoxOffset = 0;
-			OutAssetArrays->ElementOffsets.SphereOffset = OutAssetArrays->ElementOffsets.BoxOffset + BoxCount;
-			OutAssetArrays->ElementOffsets.CapsuleOffset = OutAssetArrays->ElementOffsets.SphereOffset + SphereCount;
-			OutAssetArrays->ElementOffsets.NumElements = OutAssetArrays->ElementOffsets.CapsuleOffset + CapsuleCount;
+		CountCollisionPrimitives(StaticMeshView, CurrentBodyCounts, TotalBoxCount, TotalSphereCount, TotalCapsuleCount);
+		CountCollisionPrimitives(SkeletalMeshView, CurrentBodyCounts, TotalBoxCount, TotalSphereCount, TotalCapsuleCount);
+
+		if ((TotalBoxCount + TotalSphereCount + TotalCapsuleCount) >= OutAssetArrays->MaxPrimitives)
+		{
+			UE_LOG(LogRigidMeshCollision, Error, TEXT("Number of Collision DI primitives is higher than the %d limit.  Please increase it."), OutAssetArrays->MaxPrimitives);
+			return;
 		}
 
+		OutAssetArrays->ElementOffsets.BoxOffset = 0;
+		OutAssetArrays->ElementOffsets.SphereOffset = OutAssetArrays->ElementOffsets.BoxOffset + TotalBoxCount;
+		OutAssetArrays->ElementOffsets.CapsuleOffset = OutAssetArrays->ElementOffsets.SphereOffset + TotalSphereCount;
+		OutAssetArrays->ElementOffsets.NumElements = OutAssetArrays->ElementOffsets.CapsuleOffset + TotalCapsuleCount;
+
 		uint32 BoxIndex = OutAssetArrays->ElementOffsets.BoxOffset;
 		uint32 SphereIndex = OutAssetArrays->ElementOffsets.SphereOffset;
 		uint32 CapsuleIndex = OutAssetArrays->ElementOffsets.CapsuleOffset;
 
-		if (bFullUpdate)
-		{
-			UpdateAssetArrays<UStaticMeshComponent, true>(StaticMeshView, LWCTile, OutAssetArrays, BoxIndex, SphereIndex, CapsuleIndex);
-			UpdateAssetArrays<USkeletalMeshComponent, true>(SkeletalMeshView, LWCTile, OutAssetArrays, BoxIndex, SphereIndex, CapsuleIndex);
+		// where possible PreviousTransform & PreviousInverse should be pulled from the current values of CurrentTransform &
+		// CurrentInverse based on the remapped entries
+		RemapPreviousTransforms<true>(BodyCounts, CurrentBodyCounts, OutAssetArrays);
 
-			// for newly created array data we need to duplicate the transforms to our previous transforms
-			OutAssetArrays->PreviousTransform = OutAssetArrays->CurrentTransform;
-			OutAssetArrays->PreviousInverse = OutAssetArrays->CurrentInverse;
-		}
-		else
-		{
-			// if we're updating, then copy over last frame's transforms before we generate new ones
-			OutAssetArrays->PreviousTransform = OutAssetArrays->CurrentTransform;
-			OutAssetArrays->PreviousInverse = OutAssetArrays->CurrentInverse;
+		UpdateAssetArrays<UStaticMeshComponent, true>(StaticMeshView, LWCTile, OutAssetArrays, BoxIndex, SphereIndex, CapsuleIndex);
+		UpdateAssetArrays<USkeletalMeshComponent, true>(SkeletalMeshView, LWCTile, OutAssetArrays, BoxIndex, SphereIndex, CapsuleIndex);
 
-			UpdateAssetArrays<UStaticMeshComponent, false>(StaticMeshView, LWCTile, OutAssetArrays, BoxIndex, SphereIndex, CapsuleIndex);
-			UpdateAssetArrays<USkeletalMeshComponent, false>(SkeletalMeshView, LWCTile, OutAssetArrays, BoxIndex, SphereIndex, CapsuleIndex);
-		}
+		RemapPreviousTransforms<false>(BodyCounts, CurrentBodyCounts, OutAssetArrays);
+
+		BodyCounts = MoveTemp(CurrentBodyCounts);
 	}
 }
 
+static bool SystemHasFindActorsFunction(UNiagaraSystem* System)
+{
+	bool FindActorsFunctionFound = false;
+
+	if (System)
+	{
+		System->ForEachScript([&](UNiagaraScript* NiagaraScript)
+		{
+			const FNiagaraVMExecutableData& ScriptExecutableData = NiagaraScript->GetVMExecutableData();
+			if (ScriptExecutableData.IsValid())
+			{
+				for (const FVMExternalFunctionBindingInfo& FunctionBinding : ScriptExecutableData.CalledVMExternalFunctions)
+				{
+					if (FunctionBinding.Name == FindActorsName)
+					{
+						FindActorsFunctionFound = true;
+						break;
+					}
+				}
+			}
+		});
+	}
+
+	return FindActorsFunctionFound;
+}
+
+bool WeakActorPtrLess(const TWeakObjectPtr<AActor>& Lhs, const TWeakObjectPtr<AActor>& Rhs)
+{
+	return Lhs.GetWeakPtrTypeHash() < Rhs.GetWeakPtrTypeHash();
+}
+
 } // NDIRigidMeshCollisionLocal
 
 //------------------------------------------------------------------------------------------------------------
@@ -563,6 +676,7 @@ void FNDIRigidMeshCollisionBuffer::ReleaseRHI()
 	WorldTransformBuffer.Release();
 	InverseTransformBuffer.Release();
 	ElementExtentBuffer.Release();
+	MeshScaleBuffer.Release();
 	PhysicsTypeBuffer.Release();
 	DFIndexBuffer.Release();
 }
@@ -592,7 +706,9 @@ bool FNDIRigidMeshCollisionData::HasActors() const
 
 bool FNDIRigidMeshCollisionData::ShouldRunGlobalSearch(UNiagaraDataInterfaceRigidMeshCollisionQuery* Interface) const
 {
+PRAGMA_DISABLE_DEPRECATION_WARNINGS
 	return Interface->GlobalSearchAllowed && (Interface->GlobalSearchForced || (Interface->GlobalSearchFallback_Unscripted && !bHasScriptedFindActor));
+PRAGMA_ENABLE_DEPRECATION_WARNINGS
 }
 
 void FNDIRigidMeshCollisionData::MergeActors(FMergedActorArray& MergedActors) const
@@ -652,7 +768,6 @@ void FNDIRigidMeshCollisionData::Init(int32 MaxNumPrimitives)
 		ReleaseBuffers();
 	}
 
-	bFoundActorsUpdated = false;
 	bRequiresFullUpdate = true;
 }
 
@@ -660,19 +775,19 @@ void FNDIRigidMeshCollisionData::Update(UNiagaraDataInterfaceRigidMeshCollisionQ
 {
 	using namespace NDIRigidMeshCollisionLocal;
 
-	if (!Interface || !SystemInstance || !bRequiresSourceActors)
+	if (!Interface || !SystemInstance)
 	{
 		return;
 	}
 
-	const bool bExplicitActorsChanged = Interface->GetExplicitActors(*this);
+	Interface->GetExplicitActors(*this);
+
+PRAGMA_DISABLE_DEPRECATION_WARNINGS
 	if (ShouldRunGlobalSearch(Interface))
 	{
-		if (Interface->GlobalFindActors(SystemInstance->GetWorld(), *this))
-		{
-			bFoundActorsUpdated = true;
-		}
+		Interface->GlobalFindActors(SystemInstance->GetWorld(), *this);
 	}
+PRAGMA_ENABLE_DEPRECATION_WARNINGS
 
 	// see if we need to reinitialize the internals
 	const bool bAlreadyInited = AssetArrays != nullptr;
@@ -684,17 +799,64 @@ void FNDIRigidMeshCollisionData::Update(UNiagaraDataInterfaceRigidMeshCollisionQ
 
 	if (bHasActors)
 	{
-		bRequiresFullUpdate |= TrimMissingActors();
+		TrimMissingActors();
+
+		TInlineComponentArray<UStaticMeshComponent*> StaticMeshes;
+		TInlineComponentArray<USkeletalMeshComponent*> SkeletalMeshes;
+
+		{
+			FMergedActorArray MergedActors;
+			MergeActors(MergedActors);
 
-		FMergedActorArray MergedActors;
-		MergeActors(MergedActors);
+			CollectComponents(MergedActors, Interface->ComponentTags, StaticMeshes);
+			CollectComponents(MergedActors, Interface->ComponentTags, SkeletalMeshes);
+
+			// in order to better keep track of our possibly changing set of components over time we
+			// sort our collected components by their TypeHash
+			auto SortPredicateByTypeHash = [](const UMeshComponent& Lhs, const UMeshComponent& Rhs) -> bool
+			{
+				return GetTypeHash(&Lhs) < GetTypeHash(&Rhs);
+			};
+
+			StaticMeshes.Sort(SortPredicateByTypeHash);
+			SkeletalMeshes.Sort(SortPredicateByTypeHash);
+
+			auto HashComponentObject = [](const UMeshComponent* MeshComponent) -> uint32
+			{
+				return GetTypeHash(MeshComponent);
+			};
+
+			auto AccumulateHash = [](uint32 Lhs, uint32 Rhs) -> uint32
+			{
+				return HashCombine(Lhs, Rhs);
+			};
 
-		const bool bFullUpdate = bRequiresFullUpdate || bExplicitActorsChanged || bFoundActorsUpdated;
-		UpdateInternalArrays(MergedActors, Interface->ComponentTags, FVector(SystemInstance->GetLWCTile()), bFullUpdate, AssetArrays.Get());
+			// generate a hash with the collected components so that we can see if on subsequent frames we
+			// have a different collection, in which case we'll want to do a full rebuild of the array of transforms
+			uint32 NewHashValue = 0;
+			NewHashValue = Algo::TransformAccumulate(StaticMeshes, HashComponentObject, NewHashValue, AccumulateHash);
+			NewHashValue = Algo::TransformAccumulate(SkeletalMeshes, HashComponentObject, NewHashValue, AccumulateHash);
+
+			if (NewHashValue != ComponentCollectionHash)
+			{
+				ComponentCollectionHash = NewHashValue;
+				bRequiresFullUpdate = true;
+			}
+		}
+
+		TConstArrayView<UStaticMeshComponent*> StaticMeshView = MakeArrayView(StaticMeshes.GetData(), StaticMeshes.Num());
+		TConstArrayView<USkeletalMeshComponent*> SkeletalMeshView = MakeArrayView(SkeletalMeshes.GetData(), SkeletalMeshes.Num());
+
+		UpdateInternalArrays(
+			StaticMeshView,
+			SkeletalMeshView,
+			FVector(SystemInstance->GetLWCTile()),
+			bRequiresFullUpdate,
+			MeshBodyCounts,
+			AssetArrays.Get());
 	}
 
 	bRequiresFullUpdate = false;
-	bFoundActorsUpdated = false;
 }
 
 //------------------------------------------------------------------------------------------------------------
@@ -717,6 +879,13 @@ struct FNDIRigidMeshCollisionProxy : public FNiagaraDataInterfaceProxy
 		FNDIRigidMeshCollisionBuffer* AssetBuffer = nullptr;
 	};
 
+	void RemoveInstance(const FNiagaraSystemInstanceID& Instance)
+	{
+		check(IsInRenderingThread());
+
+		SystemInstancesToProxyData_RT.Remove(Instance);
+	}
+
 	/** Get the size of the data that will be passed to render*/
 	virtual int32 PerInstanceDataPassedToRenderThreadSize() const override
 	{
@@ -943,34 +1112,19 @@ bool UNiagaraDataInterfaceRigidMeshCollisionQuery::InitPerInstanceData(void* Per
 {
 	using namespace NDIRigidMeshCollisionLocal;
 
-	bool RequiresSourceActors = false;
-	bool HasScriptedFindActor = false;
+	const bool HasScriptedFindActor = SystemHasFindActorsFunction(SystemInstance->GetSystem());
 
-	FNiagaraDataInterfaceUtilities::ForEachGpuFunctionEquals(this, SystemInstance->GetSystem(), SystemInstance, [&](const FNiagaraDataInterfaceGeneratedFunction Function)
-	{
-		RequiresSourceActors = true;
-		return false;
-	});
-
-	FNiagaraDataInterfaceUtilities::ForEachVMFunctionEquals(this, SystemInstance->GetSystem(), SystemInstance, [&](const FVMExternalFunctionBindingInfo& Binding)
-	{
-		if (Binding.Name == FindActorsName)
-		{
-			HasScriptedFindActor = true;
-			return false;
-		}
-		return true;
-	});
-
-	FNDIRigidMeshCollisionData* InstanceData = new (PerInstanceData) FNDIRigidMeshCollisionData(SystemInstance, RequiresSourceActors, HasScriptedFindActor);
+	FNDIRigidMeshCollisionData* InstanceData = new (PerInstanceData) FNDIRigidMeshCollisionData(SystemInstance, HasScriptedFindActor);
 
 	GetExplicitActors(*InstanceData);
 
+PRAGMA_DISABLE_DEPRECATION_WARNINGS
 	// if we're running a global search, then run that now
 	if (InstanceData->ShouldRunGlobalSearch(this))
 	{
 		GlobalFindActors(SystemInstance->GetWorld(), *InstanceData);
 	}
+PRAGMA_ENABLE_DEPRECATION_WARNINGS
 
 	InstanceData->Init(MaxNumPrimitives);
 
@@ -1029,11 +1183,12 @@ void UNiagaraDataInterfaceRigidMeshCollisionQuery::DestroyPerInstanceData(void*
 	ENQUEUE_RENDER_COMMAND(FNiagaraDIDestroyInstanceData) (
 		[ThisProxy, InstanceID = SystemInstance->GetId()](FRHICommandListImmediate& CmdList)
 	{
-		ThisProxy->SystemInstancesToProxyData_RT.Remove(InstanceID);
+		ThisProxy->RemoveInstance(InstanceID);
 	});
 }
 
-bool UNiagaraDataInterfaceRigidMeshCollisionQuery::PerInstanceTick(void* PerInstanceData, FNiagaraSystemInstance* SystemInstance, float InDeltaSeconds)
+// we use PostSimulate so that the results from s will be ready and pushed to the render thread for handling the GPU queries
+bool UNiagaraDataInterfaceRigidMeshCollisionQuery::PerInstanceTickPostSimulate(void* PerInstanceData, FNiagaraSystemInstance* SystemInstance, float InDeltaSeconds)
 {
 	FNDIRigidMeshCollisionData* InstanceData = static_cast<FNDIRigidMeshCollisionData*>(PerInstanceData);
 	if (InstanceData && SystemInstance)
@@ -1064,9 +1219,11 @@ bool UNiagaraDataInterfaceRigidMeshCollisionQuery::CopyToInternal(UNiagaraDataIn
 	OtherTyped->SourceActors = SourceActors;
 	OtherTyped->OnlyUseMoveable = OnlyUseMoveable;
 	OtherTyped->UseComplexCollisions = UseComplexCollisions;
+PRAGMA_DISABLE_DEPRECATION_WARNINGS
 	OtherTyped->GlobalSearchAllowed = GlobalSearchAllowed;
 	OtherTyped->GlobalSearchForced = GlobalSearchForced;
 	OtherTyped->GlobalSearchFallback_Unscripted = GlobalSearchFallback_Unscripted;
+PRAGMA_ENABLE_DEPRECATION_WARNINGS
 	OtherTyped->MaxNumPrimitives = MaxNumPrimitives;
 
 	return true;
@@ -1085,9 +1242,11 @@ bool UNiagaraDataInterfaceRigidMeshCollisionQuery::Equals(const UNiagaraDataInte
 		&& (OtherTyped->SourceActors == SourceActors)
 		&& (OtherTyped->OnlyUseMoveable == OnlyUseMoveable)
 		&& (OtherTyped->UseComplexCollisions == UseComplexCollisions)
+PRAGMA_DISABLE_DEPRECATION_WARNINGS
 		&& (OtherTyped->GlobalSearchAllowed == GlobalSearchAllowed)
 		&& (OtherTyped->GlobalSearchForced == GlobalSearchForced)
 		&& (OtherTyped->GlobalSearchFallback_Unscripted == GlobalSearchFallback_Unscripted)
+PRAGMA_ENABLE_DEPRECATION_WARNINGS
 		&& (OtherTyped->MaxNumPrimitives == MaxNumPrimitives);
 }
 
@@ -1652,6 +1811,8 @@ bool UNiagaraDataInterfaceRigidMeshCollisionQuery::GlobalFindActors(UWorld* Worl
 		}
 	}
 
+	Algo::Sort(InstanceData.FoundActors, NDIRigidMeshCollisionLocal::WeakActorPtrLess);
+
 	return PreviousActors != InstanceData.FoundActors;
 }
 
@@ -1662,13 +1823,22 @@ bool UNiagaraDataInterfaceRigidMeshCollisionQuery::FindActors(UWorld* World, FND
 
 	if (ensure(World))
 	{
-		FCollisionObjectQueryParams ObjectParams;
-		ObjectParams.AddObjectTypesToQuery(Channel);
-
 		TArray<FOverlapResult> Overlaps;
 		FCollisionQueryParams Params(SCENE_QUERY_STAT(NiagaraRigidMeshCollisionQuery), UseComplexCollisions);
 
-		World->OverlapMultiByChannel(Overlaps, OverlapLocation, OverlapRotation, Channel, FCollisionShape::MakeBox(0.5f * OverlapExtent), Params);
+		// it is not clear the best strategy for filtering out results between ByObjectType and ByChannel.  For now we're going to preserve existing behavior
+		constexpr bool bFilterByObjectType = false;
+		if (bFilterByObjectType)
+		{
+			FCollisionObjectQueryParams ObjectParams;
+			ObjectParams.AddObjectTypesToQuery(Channel);
+
+			World->OverlapMultiByObjectType(Overlaps, OverlapLocation, OverlapRotation, ObjectParams, FCollisionShape::MakeBox(0.5f * OverlapExtent), Params);
+		}
+		else
+		{
+			World->OverlapMultiByChannel(Overlaps, OverlapLocation, OverlapRotation, Channel, FCollisionShape::MakeBox(0.5f * OverlapExtent), Params);
+		}
 
 		for (const FOverlapResult& OverlapResult : Overlaps)
 		{
@@ -1691,16 +1861,13 @@ bool UNiagaraDataInterfaceRigidMeshCollisionQuery::FindActors(UWorld* World, FND
 		}
 	}
 
+	Algo::Sort(InstanceData.FoundActors, NDIRigidMeshCollisionLocal::WeakActorPtrLess);
+
 	return PreviousActors != InstanceData.FoundActors;
 }
 
 bool UNiagaraDataInterfaceRigidMeshCollisionQuery::GetExplicitActors(FNDIRigidMeshCollisionData& InstanceData)
 {
-	if (!InstanceData.bRequiresSourceActors)
-	{
-		return false;
-	}
-
 	TArray<TWeakObjectPtr<AActor>> PreviousActors;
 	Swap(InstanceData.ExplicitActors, PreviousActors);
 
@@ -1740,7 +1907,7 @@ void UNiagaraDataInterfaceRigidMeshCollisionQuery::FindActorsCPU(FVectorVMExtern
 				FQuat4f OverlapRotation = OverlapRotationParam.GetAndAdvance();
 				FVector3f OverlapExtent = OverlapExtentParam.GetAndAdvance();
 				ECollisionChannel TraceChannel = TraceChannelParam.GetAndAdvance();
-				bool SkipOverlap = SkipOverlapParam.GetAndAdvance() || !InstanceData->bRequiresSourceActors;
+				bool SkipOverlap = SkipOverlapParam.GetAndAdvance();
 
 				bool ActorsChanged = false;
 
@@ -1750,7 +1917,6 @@ void UNiagaraDataInterfaceRigidMeshCollisionQuery::FindActorsCPU(FVectorVMExtern
 					if (FindActors(World, *InstanceData, TraceChannel, ConvertedOrigin, FVector(OverlapExtent), FQuat(OverlapRotation)))
 					{
 						ActorsChanged = true;
-						InstanceData->bFoundActorsUpdated = true;
 					}
 				}
 
@@ -1767,6 +1933,52 @@ void UNiagaraDataInterfaceRigidMeshCollisionQuery::FindActorsCPU(FVectorVMExtern
 	}
 }
 
+#if WITH_EDITOR
+
+void UNiagaraDataInterfaceRigidMeshCollisionQuery::GetFeedback(UNiagaraSystem* InAsset, UNiagaraComponent* InComponent, TArray<FNiagaraDataInterfaceError>& OutErrors, TArray<FNiagaraDataInterfaceFeedback>& OutWarnings, TArray<FNiagaraDataInterfaceFeedback>& OutInfo)
+{
+	using namespace NDIRigidMeshCollisionLocal;
+
+	auto IsValidActor = [](const TSoftObjectPtr<AActor>& SoftPtrActor) -> bool
+	{
+		return SoftPtrActor.IsValid();
+	};
+
+	const bool bHasScriptedFindActor = SystemHasFindActorsFunction(InAsset);
+	const bool bHasSourceActor = SourceActors.ContainsByPredicate(IsValidActor);
+
+	if (!bHasScriptedFindActor && !bHasSourceActor)
+	{
+		OutWarnings.Emplace(
+			LOCTEXT("RigidMeshDI_MissingFindActor", "RigidMeshCollisionQuery is missing call to FindActors which is needed to find collision primitives."),
+			LOCTEXT("RigidMeshDI_MissingFindActor_Summary", "Missing FIndActors."),
+			FNiagaraDataInterfaceFix()
+		);
+	}
+
+PRAGMA_DISABLE_DEPRECATION_WARNINGS
+	if (GlobalSearchAllowed)
+	{
+		OutWarnings.Emplace(
+			LOCTEXT("RigidMeshDI_GlobalSearch", "Use of Global Search options is deprecated in favor of using FindActors."),
+			LOCTEXT("RigidMeshDI_GlobalSearch_Summary", "Deprecated Global Search."),
+			FNiagaraDataInterfaceFix::CreateLambda(
+				[this]()
+				{
+					GlobalSearchAllowed = false;
+					GlobalSearchForced = false;
+					GlobalSearchFallback_Unscripted = false;
+					return true;
+				}
+			)
+		);
+	}
+PRAGMA_ENABLE_DEPRECATION_WARNINGS
+}
+
+#endif
+
+
 void UNiagaraDIRigidMeshCollisionFunctionLibrary::SetSourceActors(UNiagaraComponent* NiagaraComponent, FName OverrideName, const TArray<AActor*>& InSourceActors)
 {
 	if (UNiagaraDataInterfaceRigidMeshCollisionQuery* QueryDI = UNiagaraFunctionLibrary::GetDataInterface<UNiagaraDataInterfaceRigidMeshCollisionQuery>(NiagaraComponent, OverrideName))
diff --git a/Engine/Plugins/FX/Niagara/Source/Niagara/Public/NiagaraDataInterfaceRigidMeshCollisionQuery.h b/Engine/Plugins/FX/Niagara/Source/Niagara/Public/NiagaraDataInterfaceRigidMeshCollisionQuery.h
--- a/Engine/Plugins/FX/Niagara/Source/Niagara/Public/NiagaraDataInterfaceRigidMeshCollisionQuery.h
+++ b/Engine/Plugins/FX/Niagara/Source/Niagara/Public/NiagaraDataInterfaceRigidMeshCollisionQuery.h
@@ -111,10 +111,9 @@ struct FNDIRigidMeshCollisionBuffer : public FRenderResource
 /** Data stored per DI instance*/
 struct FNDIRigidMeshCollisionData
 {
-	FNDIRigidMeshCollisionData(const FNiagaraSystemInstance* InSystemInstance, bool InRequiresSourceActors, bool InHasScriptedFindActor)
-	: SystemInstance(InSystemInstance)
-	, bRequiresSourceActors(InRequiresSourceActors)
-	, bHasScriptedFindActor(InHasScriptedFindActor)
+	FNDIRigidMeshCollisionData(const FNiagaraSystemInstance* InSystemInstance, bool InHasScriptedFindActor)
+		: SystemInstance(InSystemInstance)
+		, bHasScriptedFindActor(InHasScriptedFindActor)
 	{}
 
 	/** Initialize the cpu datas */
@@ -143,23 +142,29 @@ struct FNDIRigidMeshCollisionData
 	/** Source actors **/
 	TArray<TWeakObjectPtr<AActor>> ExplicitActors;
 
+	struct FComponentBodyCount
+	{
+		uint32 ComponentHash = 0;
+		uint32 BoxCount = 0;
+		uint32 SphereCount = 0;
+		uint32 CapsuleCount = 0;
+	};
+
+	TArray<FComponentBodyCount> MeshBodyCounts;
+
 	/** Source actors **/
 	TArray<TWeakObjectPtr<AActor>> FoundActors;
 
 	const FNiagaraSystemInstance* SystemInstance = nullptr;
 
-	/** If false indicates that the instance is not dependent on any RigidBodies, and processing can be skipped. */
-	const bool bRequiresSourceActors;
+	/** Hash of the components that have been collected, used to see if we need to fully rebuild our array of transforms */
+	uint32 ComponentCollectionHash = 0;
 
 	/** Indicates that the instance is being used in a call to FindActors */
 	const bool bHasScriptedFindActor;
 
-	/** Indicates that something has updated the list of FoundActors */
-	bool bFoundActorsUpdated = false;
-
 	/** Indicates that a full update of the arrays is required */
 	bool bRequiresFullUpdate = false;
-
 };
 
 /** Data Interface used to collide against static meshes - whether it is the mesh distance field or a physics asset's collision primitive */
@@ -197,16 +202,19 @@ public:
 	bool UseComplexCollisions = false;
 
 	/** If enabled the global search can be executed dependeing on GlobalSearchForced and GlobalSearchFallback_Unscripted */
-	UPROPERTY(EditAnywhere, Category = "Source", meta = (DisplayName = "Global Search Allowed"))
-	bool GlobalSearchAllowed = true;
+	UE_DEPRECATED(5.2, "Global search will be deprecated in favor of using a spatial search through the FindActors function.")
+	UPROPERTY(EditAnywhere, AdvancedDisplay, Category = "Source", meta = (DisplayName = "Global Search Allowed"))
+	bool GlobalSearchAllowed = false;
 
 	/** If enabled the global search will be performed only if there are no explicit actors specified */
-	UPROPERTY(EditAnywhere, Category = "Source", meta = (DisplayName = "Global Search Forced", EditCondition = GlobalSearchAllowed))
+	UE_DEPRECATED(5.2, "Global search will be deprecated in favor of using a spatial search through the FindActors function.")
+	UPROPERTY(EditAnywhere, AdvancedDisplay, Category = "Source", meta = (DisplayName = "Global Search Forced", EditCondition = GlobalSearchAllowed))
 	bool GlobalSearchForced = false;
 
 	/** If enabled the global search will be performed only if there are no explicit actors specified */
-	UPROPERTY(EditAnywhere, Category = "Source", meta = (DisplayName = "Global Search Only If Unscripted", EditCondition = GlobalSearchAllowed))
-	bool GlobalSearchFallback_Unscripted = true;
+	UE_DEPRECATED(5.2, "Global search will be deprecated in favor of using a spatial search through the FindActors function.")
+	UPROPERTY(EditAnywhere, AdvancedDisplay, Category = "Source", meta = (DisplayName = "Global Search Only If Unscripted", EditCondition = GlobalSearchAllowed))
+	bool GlobalSearchFallback_Unscripted = false;
 
 	/** Maximum number of RigidBody represented by this DataInterface. */
 	UPROPERTY(EditAnywhere, Category = "General")
@@ -225,10 +233,10 @@ public:
 #endif
 	virtual bool InitPerInstanceData(void* PerInstanceData, FNiagaraSystemInstance* SystemInstance) override;
 	virtual void DestroyPerInstanceData(void* PerInstanceData, FNiagaraSystemInstance* SystemInstance)override;
-	virtual bool PerInstanceTick(void* PerInstanceData, FNiagaraSystemInstance* SystemInstance, float DeltaSeconds) override;
+	virtual bool PerInstanceTickPostSimulate(void* PerInstanceData, FNiagaraSystemInstance* SystemInstance, float DeltaSeconds) override;
 	virtual int32 PerInstanceDataSize()const override { return sizeof(FNDIRigidMeshCollisionData); }
 	virtual bool Equals(const UNiagaraDataInterface* Other) const override;
-	virtual bool HasPreSimulateTick() const override { return true; }
+	virtual bool HasPostSimulateTick() const override { return true; }
 	virtual bool HasTickGroupPrereqs() const override { return true; }
 	virtual ETickingGroup CalculateTickGroup(const void* PerInstanceData) const override;
 
@@ -255,10 +263,16 @@ public:
 
 	bool GetExplicitActors(FNDIRigidMeshCollisionData& InstanceData);
 	bool FindActors(UWorld* World, FNDIRigidMeshCollisionData& InstanceData, ECollisionChannel Channel, const FVector& OverlapLocation, const FVector& OverlapExtent, const FQuat& OverlapRotation) const;
+
+	UE_DEPRECATED(5.2, "Global search will be deprecated in favor of using a spatial search through the FindActors function.")
 	bool GlobalFindActors(UWorld* World, FNDIRigidMeshCollisionData& InstanceData) const;
 
 	void FindActorsCPU(FVectorVMExternalFunctionContext& Context);
 
+#if WITH_EDITOR
+	virtual void GetFeedback(UNiagaraSystem* InAsset, UNiagaraComponent* InComponent, TArray<FNiagaraDataInterfaceError>& OutErrors, TArray<FNiagaraDataInterfaceFeedback>& OutWarnings, TArray<FNiagaraDataInterfaceFeedback>& OutInfo) override;
+#endif
+
 protected:
 	/** Copy one niagara DI to this */
 	virtual bool CopyToInternal(UNiagaraDataInterface* Destination) const override;
