Fixed crash in niagara component node builder.
Seems the component becomes invalid at some point during undo/redo if a BP is involved.

diff --git a/Engine/Plugins/FX/Niagara/Source/NiagaraEditor/Private/Customizations/NiagaraComponentDetails.cpp b/Engine/Plugins/FX/Niagara/Source/NiagaraEditor/Private/Customizations/NiagaraComponentDetails.cpp
--- a/Engine/Plugins/FX/Niagara/Source/NiagaraEditor/Private/Customizations/NiagaraComponentDetails.cpp
+++ b/Engine/Plugins/FX/Niagara/Source/NiagaraEditor/Private/Customizations/NiagaraComponentDetails.cpp
@@ -871,45 +871,49 @@ void FNiagaraUserParameterNodeBuilder::GenerateUserParameterRows(IDetailChildren
 
 void FNiagaraComponentUserParametersNodeBuilder::GenerateChildContent(IDetailChildrenBuilder& ChildrenBuilder)
 {
-	if (bDelegatesInitialized == false)
+	if(Component.IsValid())
 	{
-		Component->OnSynchronizedWithAssetParameters().AddSP(this, &FNiagaraComponentUserParametersNodeBuilder::Rebuild);
-		Component->GetOverrideParameters().AddOnChangedHandler(
-			FNiagaraParameterStore::FOnChanged::FDelegate::CreateSP(this, &FNiagaraComponentUserParametersNodeBuilder::ParameterValueChanged));
-
-		RegisterRebuildOnHierarchyChanged();
-		bDelegatesInitialized = true;
-	}
+		if (bDelegatesInitialized == false)
+		{
+			Component->OnSynchronizedWithAssetParameters().AddSP(this, &FNiagaraComponentUserParametersNodeBuilder::Rebuild);
+			Component->GetOverrideParameters().AddOnChangedHandler(
+				FNiagaraParameterStore::FOnChanged::FDelegate::CreateSP(this, &FNiagaraComponentUserParametersNodeBuilder::ParameterValueChanged));
 
-	check(Component.IsValid());
+			RegisterRebuildOnHierarchyChanged();
+			bDelegatesInitialized = true;
+		}
 
-	ParameterProxies.Reset();
+		ParameterProxies.Reset();
 
-	UNiagaraSystem* SystemAsset = Component->GetAsset();
-	if (SystemAsset == nullptr)
-	{
-		return;
-	}
-		
-	TArray<FNiagaraVariable> UserParameters;
-	SystemAsset->GetExposedParameters().GetUserParameters(UserParameters);
-		
-	ParameterProxies.Reserve(UserParameters.Num());
+		UNiagaraSystem* SystemAsset = Component->GetAsset();
+		if (SystemAsset == nullptr)
+		{
+			return;
+		}
+			
+		TArray<FNiagaraVariable> UserParameters;
+		SystemAsset->GetExposedParameters().GetUserParameters(UserParameters);
+			
+		ParameterProxies.Reserve(UserParameters.Num());
 
-	ParameterToDisplayStruct.Empty();
+		ParameterToDisplayStruct.Empty();
 
-	GenerateUserParameterRows(ChildrenBuilder, *Cast<UNiagaraSystemEditorData>(SystemAsset->GetEditorData())->UserParameterHierarchy.Get(),Component.Get());
+		GenerateUserParameterRows(ChildrenBuilder, *Cast<UNiagaraSystemEditorData>(SystemAsset->GetEditorData())->UserParameterHierarchy.Get(),Component.Get());
+	}
 }
 
 void FNiagaraComponentUserParametersNodeBuilder::RegisterRebuildOnHierarchyChanged()
 {
-	if(UNiagaraSystem* Asset = Component->GetAsset())
+	if(Component.IsValid())
 	{
-		TSharedPtr<FNiagaraSystemViewModel> SystemViewModel = TNiagaraViewModelManager<UNiagaraSystem, FNiagaraSystemViewModel>::GetExistingViewModelForObject(Asset);
-		if(SystemViewModel.IsValid())
+		if(UNiagaraSystem* Asset = Component->GetAsset())
 		{
-			SystemViewModel->GetUserParametersHierarchyViewModel()->OnHierarchyChanged().RemoveAll(this);
-			SystemViewModel->GetUserParametersHierarchyViewModel()->OnHierarchyChanged().Add(UNiagaraHierarchyViewModelBase::FOnHierarchyChanged::FDelegate::CreateSP(this, &FNiagaraComponentUserParametersNodeBuilder::Rebuild));
+			TSharedPtr<FNiagaraSystemViewModel> SystemViewModel = TNiagaraViewModelManager<UNiagaraSystem, FNiagaraSystemViewModel>::GetExistingViewModelForObject(Asset);
+			if(SystemViewModel.IsValid())
+			{
+				SystemViewModel->GetUserParametersHierarchyViewModel()->OnHierarchyChanged().RemoveAll(this);
+				SystemViewModel->GetUserParametersHierarchyViewModel()->OnHierarchyChanged().Add(UNiagaraHierarchyViewModelBase::FOnHierarchyChanged::FDelegate::CreateSP(this, &FNiagaraComponentUserParametersNodeBuilder::Rebuild));
+			}
 		}
 	}
 }
