Fixed a crash when using the niagara submix data interface in a packaged game

diff --git a/Engine/Plugins/FX/Niagara/Source/Niagara/Private/NiagaraDataInterfaceAudio.cpp b/Engine/Plugins/FX/Niagara/Source/Niagara/Private/NiagaraDataInterfaceAudio.cpp
--- a/Engine/Plugins/FX/Niagara/Source/Niagara/Private/NiagaraDataInterfaceAudio.cpp
+++ b/Engine/Plugins/FX/Niagara/Source/Niagara/Private/NiagaraDataInterfaceAudio.cpp
@@ -56,14 +56,12 @@ FNiagaraSubmixListener::~FNiagaraSubmixListener()
 
 void FNiagaraSubmixListener::RegisterToSubmix()
 {
-	FAudioDevice * DeviceHandle = FAudioDeviceManager::Get()->GetAudioDeviceRaw(AudioDeviceId);
-	ensure(DeviceHandle);
-
-	if (DeviceHandle)
+	if (FAudioDevice* DeviceHandle = FAudioDeviceManager::Get()->GetAudioDeviceRaw(AudioDeviceId))
 	{
-		DeviceHandle->RegisterSubmixBufferListener(this, Submix);
 		bIsRegistered = true;
 
+		DeviceHandle->RegisterSubmixBufferListener(this, Submix);
+
 		// RegisterSubmixBufferListener lazily enqueues the registration on the audio thread,
 		// so we have to wait for the audio thread to destroy.
 		FAudioCommandFence Fence;
@@ -76,23 +74,38 @@ void FNiagaraSubmixListener::UnregisterFromSubmix()
 {
 	if (bIsRegistered)
 	{
-		FAudioDevice * DeviceHandle = FAudioDeviceManager::Get()->GetAudioDeviceRaw(AudioDeviceId);
-		ensure(DeviceHandle);
-		DeviceHandle->UnregisterSubmixBufferListener(this, Submix);
-
 		bIsRegistered = false;
-
-		// UnregisterSubmixBufferListener lazily enqueues the unregistration on the audio thread,
-		// so we have to wait for the audio thread to destroy.
-		FAudioCommandFence Fence;
-		Fence.BeginFence();
-		Fence.Wait();
+		
+		if (FAudioDevice* DeviceHandle = FAudioDeviceManager::Get()->GetAudioDeviceRaw(AudioDeviceId))
+		{
+			if (IsInGameThread())
+			{
+				DeviceHandle->UnregisterSubmixBufferListener(this, Submix);
+				
+				// UnregisterSubmixBufferListener lazily enqueues the unregistration on the audio thread,
+				// so we have to wait for the audio thread to destroy.
+				FAudioCommandFence Fence;
+				Fence.BeginFence();
+				Fence.Wait();
+			}
+			else
+			{
+				UE::Tasks::FTaskEvent CompletionEvent{ UE_SOURCE_LOCATION };
+				FAudioThread::RunCommandOnAudioThread([this, DeviceHandle, &CompletionEvent]()
+				{
+					DeviceHandle->UnregisterSubmixBufferListener(this, Submix);
+					CompletionEvent.Trigger();
+				});
+
+				CompletionEvent.Wait();
+			}
+		}
 	}
 }
 
 float FNiagaraSubmixListener::GetSampleRate() const
 {
-	return (float) SubmixSampleRate;
+	return static_cast<float>(SubmixSampleRate.Load());
 }
 
 int32 FNiagaraSubmixListener::GetNumChannels() const
@@ -140,14 +153,7 @@ void FNiagaraDataInterfaceProxySubmix::RegisterToAllAudioDevices()
 
 void FNiagaraDataInterfaceProxySubmix::UnregisterFromAllAudioDevices()
 {
-	if (FAudioDeviceManager* DeviceManager = FAudioDeviceManager::Get())
-	{
-		// Register a new submix listener for every audio device that currently exists.
-		DeviceManager->IterateOverAllDevices([&](Audio::FDeviceId DeviceId, FAudioDevice* InDevice)
-		{
-			RemoveSubmixListener(DeviceId);
-		});
-	}
+	SubmixListeners.Empty();
 }
 
 void FNiagaraDataInterfaceProxySubmix::OnUpdateSubmix(USoundSubmix* Submix)
@@ -311,11 +317,8 @@ void UNiagaraDataInterfaceAudioSubmix::PostEditChangeProperty(struct FPropertyCh
 void UNiagaraDataInterfaceAudioSubmix::PostInitProperties()
 {
 	Super::PostInitProperties();
-
-	GetProxyAs<FNiagaraDataInterfaceProxySubmix>()->OnUpdateSubmix(Submix);
 }
 
-
 void UNiagaraDataInterfaceAudioSubmix::BeginDestroy()
 {
 	GetProxyAs<FNiagaraDataInterfaceProxySubmix>()->OnBeginDestroy();
@@ -326,6 +329,8 @@ void UNiagaraDataInterfaceAudioSubmix::BeginDestroy()
 void UNiagaraDataInterfaceAudioSubmix::PostLoad()
 {
 	Super::PostLoad();
+
+	GetProxyAs<FNiagaraDataInterfaceProxySubmix>()->OnUpdateSubmix(Submix);
 }
 
 bool UNiagaraDataInterfaceAudioSubmix::CopyToInternal(UNiagaraDataInterface* Destination) const
