Parameter map & get nodes now factor in existing parameters when renaming a parameter pin

diff --git a/Engine/Plugins/FX/Niagara/Source/NiagaraEditor/Private/NiagaraNodeParameterMapGet.cpp b/Engine/Plugins/FX/Niagara/Source/NiagaraEditor/Private/NiagaraNodeParameterMapGet.cpp
--- a/Engine/Plugins/FX/Niagara/Source/NiagaraEditor/Private/NiagaraNodeParameterMapGet.cpp
+++ b/Engine/Plugins/FX/Niagara/Source/NiagaraEditor/Private/NiagaraNodeParameterMapGet.cpp
@@ -84,6 +84,13 @@ bool UNiagaraNodeParameterMapGet::VerifyEditablePinName(const FText& InName, FTe
 		OutErrorMessage = LOCTEXT("InvalidName", "Invalid pin name");
 		return false;
 	}
+	
+	if(GetNiagaraGraph()->GetScriptVariable(FName(InName.ToString())) != nullptr && !InName.EqualToCaseIgnored(FText::FromString(InGraphPinObj->GetName())))
+	{
+		OutErrorMessage = LOCTEXT("InvalidName_VariableExists", "This variable already exists.");
+		return false;
+	}
+	
 	return true;
 }
 
diff --git a/Engine/Plugins/FX/Niagara/Source/NiagaraEditor/Private/NiagaraNodeParameterMapSet.cpp b/Engine/Plugins/FX/Niagara/Source/NiagaraEditor/Private/NiagaraNodeParameterMapSet.cpp
--- a/Engine/Plugins/FX/Niagara/Source/NiagaraEditor/Private/NiagaraNodeParameterMapSet.cpp
+++ b/Engine/Plugins/FX/Niagara/Source/NiagaraEditor/Private/NiagaraNodeParameterMapSet.cpp
@@ -81,6 +81,13 @@ bool UNiagaraNodeParameterMapSet::VerifyEditablePinName(const FText& InName, FTe
 		OutErrorMessage = LOCTEXT("InvalidName", "Invalid pin name");
 		return false;
 	}
+
+	if(GetNiagaraGraph()->GetScriptVariable(FName(InName.ToString())) != nullptr && !InName.EqualToCaseIgnored(FText::FromString(InGraphPinObj->GetName())))
+	{
+		OutErrorMessage = LOCTEXT("InvalidName_VariableExists", "This variable already exists.");
+		return false;
+	}
+	
 	return true;
 }
 
