Made sure the data for editing is marked transactional

diff --git a/Engine/Plugins/FX/Niagara/Source/NiagaraEditor/Private/Widgets/SNiagaraHierarchy.cpp b/Engine/Plugins/FX/Niagara/Source/NiagaraEditor/Private/Widgets/SNiagaraHierarchy.cpp
--- a/Engine/Plugins/FX/Niagara/Source/NiagaraEditor/Private/Widgets/SNiagaraHierarchy.cpp
+++ b/Engine/Plugins/FX/Niagara/Source/NiagaraEditor/Private/Widgets/SNiagaraHierarchy.cpp
@@ -840,7 +840,11 @@ void SNiagaraHierarchy::OnSelectionChanged(TSharedPtr<FNiagaraHierarchyItemViewM
 				HierarchyTreeView->ClearSelection();
 			}
 			
-			DetailsPanel->SetObject(HierarchyItem->GetDataForEditing());	
+			UObject* DataForEditing = HierarchyItem->GetDataForEditing();
+			DataForEditing->SetFlags(RF_Transactional);
+
+			// we make sure the object we are editing is transactional
+			DetailsPanel->SetObject(DataForEditing);	
 		}
 		else
 		{
