Fix for cleaning up GetTypeHash for weak ptrs for sorting actors, avoids using a private implementation

diff --git a/Engine/Plugins/FX/Niagara/Source/Niagara/Private/NiagaraDataInterfaceRigidMeshCollisionQuery.cpp b/Engine/Plugins/FX/Niagara/Source/Niagara/Private/NiagaraDataInterfaceRigidMeshCollisionQuery.cpp
--- a/Engine/Plugins/FX/Niagara/Source/Niagara/Private/NiagaraDataInterfaceRigidMeshCollisionQuery.cpp
+++ b/Engine/Plugins/FX/Niagara/Source/Niagara/Private/NiagaraDataInterfaceRigidMeshCollisionQuery.cpp
@@ -651,7 +651,7 @@ static bool SystemHasFindActorsFunction(UNiagaraSystem* System)
 
 bool WeakActorPtrLess(const TWeakObjectPtr<AActor>& Lhs, const TWeakObjectPtr<AActor>& Rhs)
 {
-	return Lhs.GetWeakPtrTypeHash() < Rhs.GetWeakPtrTypeHash();
+	return GetTypeHash(Lhs) < GetTypeHash(Rhs);
 }
 
 } // NDIRigidMeshCollisionLocal
