- Disable this frame translucent rendering on mobile renderer as it does not call PostRenderOpaque before rendering translucent but instead afterwards

#rnx
#preflght 637d6e944004f73f62e7f275

diff --git a/Engine/Plugins/FX/Niagara/Source/Niagara/Private/NiagaraRendererMeshes.cpp b/Engine/Plugins/FX/Niagara/Source/Niagara/Private/NiagaraRendererMeshes.cpp
--- a/Engine/Plugins/FX/Niagara/Source/Niagara/Private/NiagaraRendererMeshes.cpp
+++ b/Engine/Plugins/FX/Niagara/Source/Niagara/Private/NiagaraRendererMeshes.cpp
@@ -80,7 +80,7 @@ FNiagaraRendererMeshes::FNiagaraRendererMeshes(ERHIFeatureLevel::Type FeatureLev
 	SortMode = Properties->SortMode;
 	bSortHighPrecision = UNiagaraRendererProperties::IsSortHighPrecision(Properties->SortPrecision);
 	bSortOnlyWhenTranslucent = Properties->bSortOnlyWhenTranslucent;
-	bGpuLowLatencyTranslucency = UNiagaraRendererProperties::IsGpuTranslucentThisFrame(Properties->GpuTranslucentLatency);
+	bGpuLowLatencyTranslucency = UNiagaraRendererProperties::IsGpuTranslucentThisFrame(FeatureLevel, Properties->GpuTranslucentLatency);
 	bOverrideMaterials = Properties->bOverrideMaterials;
 	SubImageSize = FVector2f(Properties->SubImageSize);	// LWC_TODO: Precision loss
 	bSubImageBlend = Properties->bSubImageBlend;
diff --git a/Engine/Plugins/FX/Niagara/Source/Niagara/Private/NiagaraRendererProperties.cpp b/Engine/Plugins/FX/Niagara/Source/Niagara/Private/NiagaraRendererProperties.cpp
--- a/Engine/Plugins/FX/Niagara/Source/Niagara/Private/NiagaraRendererProperties.cpp
+++ b/Engine/Plugins/FX/Niagara/Source/Niagara/Private/NiagaraRendererProperties.cpp
@@ -796,8 +796,14 @@ bool UNiagaraRendererProperties::IsSortHighPrecision(ENiagaraRendererSortPrecisi
 	return SortPrecision == ENiagaraRendererSortPrecision::High;
 }
 
-bool UNiagaraRendererProperties::IsGpuTranslucentThisFrame(ENiagaraRendererGpuTranslucentLatency Latency)
+bool UNiagaraRendererProperties::IsGpuTranslucentThisFrame(ERHIFeatureLevel::Type FeatureLevel, ENiagaraRendererGpuTranslucentLatency Latency)
 {
+	// We can not support low latency on the mobile renderer path as it calls PostRenderOpaque after translucent in some paths
+	if (FSceneInterface::GetShadingPath(FeatureLevel) != EShadingPath::Deferred)
+	{
+		return false;
+	}
+
 	if (Latency == ENiagaraRendererGpuTranslucentLatency::ProjectDefault)
 	{
 		return GetDefault<UNiagaraSettings>()->DefaultGpuTranslucentLatency == ENiagaraDefaultGpuTranslucentLatency::Immediate;
diff --git a/Engine/Plugins/FX/Niagara/Source/Niagara/Private/NiagaraRendererSprites.cpp b/Engine/Plugins/FX/Niagara/Source/Niagara/Private/NiagaraRendererSprites.cpp
--- a/Engine/Plugins/FX/Niagara/Source/Niagara/Private/NiagaraRendererSprites.cpp
+++ b/Engine/Plugins/FX/Niagara/Source/Niagara/Private/NiagaraRendererSprites.cpp
@@ -95,7 +95,7 @@ FNiagaraRendererSprites::FNiagaraRendererSprites(ERHIFeatureLevel::Type FeatureL
 	bRemoveHMDRollInVR = Properties->bRemoveHMDRollInVR;
 	bSortHighPrecision = UNiagaraRendererProperties::IsSortHighPrecision(Properties->SortPrecision);
 	bSortOnlyWhenTranslucent = Properties->bSortOnlyWhenTranslucent;
-	bGpuLowLatencyTranslucency = UNiagaraRendererProperties::IsGpuTranslucentThisFrame(Properties->GpuTranslucentLatency);
+	bGpuLowLatencyTranslucency = UNiagaraRendererProperties::IsGpuTranslucentThisFrame(FeatureLevel, Properties->GpuTranslucentLatency);
 	MinFacingCameraBlendDistance = Properties->MinFacingCameraBlendDistance;
 	MaxFacingCameraBlendDistance = Properties->MaxFacingCameraBlendDistance;
 	RendererVisibility = Properties->RendererVisibility;
diff --git a/Engine/Plugins/FX/Niagara/Source/Niagara/Public/NiagaraRendererProperties.h b/Engine/Plugins/FX/Niagara/Source/Niagara/Public/NiagaraRendererProperties.h
--- a/Engine/Plugins/FX/Niagara/Source/Niagara/Public/NiagaraRendererProperties.h
+++ b/Engine/Plugins/FX/Niagara/Source/Niagara/Public/NiagaraRendererProperties.h
@@ -372,7 +372,7 @@ public:
 
 	static bool IsSortHighPrecision(ENiagaraRendererSortPrecision SortPrecision);
 
-	static bool IsGpuTranslucentThisFrame(ENiagaraRendererGpuTranslucentLatency Latency);
+	static bool IsGpuTranslucentThisFrame(ERHIFeatureLevel::Type FeatureLevel, ENiagaraRendererGpuTranslucentLatency Latency);
 
 	template<typename TAction>
 	void ForEachPlatformSet(TAction Func);
