Moved "Edit Hierarchy" button to inside the user parameters header.
Also small layout tweaks & added tooltips

diff --git a/Engine/Plugins/FX/Niagara/Source/NiagaraEditor/Private/Customizations/NiagaraComponentDetails.cpp b/Engine/Plugins/FX/Niagara/Source/NiagaraEditor/Private/Customizations/NiagaraComponentDetails.cpp
--- a/Engine/Plugins/FX/Niagara/Source/NiagaraEditor/Private/Customizations/NiagaraComponentDetails.cpp
+++ b/Engine/Plugins/FX/Niagara/Source/NiagaraEditor/Private/Customizations/NiagaraComponentDetails.cpp
@@ -39,6 +39,7 @@
 #include "Widgets/Input/SCheckBox.h"
 #include "Widgets/Input/SComboButton.h"
 #include "SNiagaraSystemUserParameters.h"
+#include "SystemToolkitModes/NiagaraSystemToolkitModeBase.h"
 #include "Widgets/Input/SEditableTextBox.h"
 
 #define LOCTEXT_NAMESPACE "NiagaraComponentDetails"
@@ -600,8 +601,8 @@ void FNiagaraSystemUserParameterDetails::CustomizeDetails(IDetailLayoutBuilder&
 		ensure(SystemViewModel.IsValid());
 		TSharedRef<FNiagaraSystemUserParameterBuilder> SystemUserParameterBuilder = MakeShared<FNiagaraSystemUserParameterBuilder>(SystemViewModel, ParamCategoryName);
 		InputParamCategory.AddCustomBuilder(SystemUserParameterBuilder);
-		TSharedRef<SWidget> AddParameterButton = SystemUserParameterBuilder->GetAddParameterButton();
-		InputParamCategory.HeaderContent(AddParameterButton);
+		TSharedRef<SWidget> AdditionalHeaderWidgets = SystemUserParameterBuilder->GetAdditionalHeaderWidgets();
+		InputParamCategory.HeaderContent(AdditionalHeaderWidgets);
 	}
 }
 
@@ -1027,35 +1028,57 @@ void FNiagaraSystemUserParameterBuilder::AddCustomMenuActionsForParameter(FDetai
 	WidgetRow.AddCustomContextMenuAction(DeleteAction, LOCTEXT("DeleteParameterAction", "Delete"), LOCTEXT("DeleteParameterActionTooltip", "Delete this user parameter"));
 }
 
-TSharedRef<SWidget> FNiagaraSystemUserParameterBuilder::GetAddParameterButton()
+TSharedRef<SWidget> FNiagaraSystemUserParameterBuilder::GetAdditionalHeaderWidgets()
 {
-	if(!AddParameterButtonContainer.IsValid())
+	if(!AdditionalHeaderWidgetsContainer.IsValid())
 	{
-		AddParameterButtonContainer = SNew(SBox)
+		AdditionalHeaderWidgetsContainer = SNew(SBox)
 		.HAlign(HAlign_Right)
 		.VAlign(VAlign_Center)
+		.Padding(2.f)
 		[
-			SAssignNew(AddParameterButton, SComboButton)
+			SNew(SHorizontalBox)
+			+ SHorizontalBox::Slot()
 			.HAlign(HAlign_Center)
 			.VAlign(VAlign_Center)
-			.ButtonStyle(FAppStyle::Get(), "RoundButton")
-			.ForegroundColor(FAppStyle::GetSlateColor("DefaultForeground"))
-			.ContentPadding(FMargin(2, 0))
-			.OnGetMenuContent(this, &FNiagaraSystemUserParameterBuilder::GetAddParameterMenu)
-			.OnComboBoxOpened_Lambda([this]
-			{
-				AddParameterButton->SetMenuContentWidgetToFocus(AddParameterMenu->GetSearchBox());
-			})
-			.HasDownArrow(false)
-			.ButtonContent()
+			.AutoWidth()
+			.Padding(2.f, 1.f)
+			[
+				SNew(SButton)
+				.OnClicked(this, &FNiagaraSystemUserParameterBuilder::SummonHierarchyEditor)
+				.Text(LOCTEXT("SummonUserParametersHierarchyButtonLabel", "Edit Hierarchy"))
+				.ButtonStyle(FAppStyle::Get(), "RoundButton")
+				.ToolTipText(LOCTEXT("SummonUserParametersHierarchyButtonTooltip", "Summon the Hierarchy Editor to add categories, sections and tooltips to User Parameters."))
+			]
+			+ SHorizontalBox::Slot()
+			.HAlign(HAlign_Center)
+			.VAlign(VAlign_Center)
+			.AutoWidth()
+			.Padding(2.f)
 			[
-				SNew(SImage)
-				.Image(FAppStyle::GetBrush("Plus"))
+				SAssignNew(AddParameterButton, SComboButton)
+				.HAlign(HAlign_Center)
+				.VAlign(VAlign_Center)
+				.ButtonStyle(FAppStyle::Get(), "RoundButton")
+				.ForegroundColor(FAppStyle::GetSlateColor("DefaultForeground"))
+				.ContentPadding(FMargin(2, 3.5))
+				.ToolTipText(LOCTEXT("AddUserParameterButtonTooltip", "Add a new User Parameter to this system.\nUser Parameters exist once for the entire system and their values can be changed at runtime from Blueprints, C++ or the sequencer, among other systems."))
+				.OnGetMenuContent(this, &FNiagaraSystemUserParameterBuilder::GetAddParameterMenu)
+				.OnComboBoxOpened_Lambda([this]
+				{
+					AddParameterButton->SetMenuContentWidgetToFocus(AddParameterMenu->GetSearchBox());
+				})
+				.HasDownArrow(false)
+				.ButtonContent()
+				[
+					SNew(SImage)
+					.Image(FAppStyle::GetBrush("Plus"))
+				]
 			]
 		];
 	}
 	
-	return AddParameterButtonContainer.ToSharedRef(); 
+	return AdditionalHeaderWidgetsContainer.ToSharedRef(); 
 }
 
 TSharedRef<SWidget> FNiagaraSystemUserParameterBuilder::GetAddParameterMenu()
@@ -1114,6 +1137,17 @@ void FNiagaraSystemUserParameterBuilder::ParameterValueChanged()
 	}
 }
 
+FReply FNiagaraSystemUserParameterBuilder::SummonHierarchyEditor()
+{
+	if(TSharedPtr<FNiagaraSystemViewModel> SystemViewModelPinned = SystemViewModel.Pin())
+	{
+		SystemViewModelPinned->FocusTab(FNiagaraSystemToolkitModeBase::UserParametersHierarchyTabID, true);
+		return FReply::Handled();
+	}
+
+	return FReply::Unhandled();
+}
+
 void FNiagaraSystemUserParameterBuilder::DeleteParameter(FNiagaraVariable UserParameter) const
 {
 	FScopedTransaction Transaction(LOCTEXT("DeleteParameterTransaction", "Deleted user parameter"));
diff --git a/Engine/Plugins/FX/Niagara/Source/NiagaraEditor/Private/Customizations/NiagaraComponentDetails.h b/Engine/Plugins/FX/Niagara/Source/NiagaraEditor/Private/Customizations/NiagaraComponentDetails.h
--- a/Engine/Plugins/FX/Niagara/Source/NiagaraEditor/Private/Customizations/NiagaraComponentDetails.h
+++ b/Engine/Plugins/FX/Niagara/Source/NiagaraEditor/Private/Customizations/NiagaraComponentDetails.h
@@ -192,14 +192,15 @@ public:
 	/** We want to add rename & delete actions within a system asset */
 	virtual void AddCustomMenuActionsForParameter(FDetailWidgetRow& WidgetRow, FNiagaraVariable UserParameter) override;
 
-	TSharedRef<SWidget> GetAddParameterButton();
+	TSharedRef<SWidget> GetAdditionalHeaderWidgets();
 private:
 	TSharedRef<SWidget> GetAddParameterMenu();
 	void AddParameter(FNiagaraVariable NewParameter) const;
 	bool CanMakeNewParameterOfType(const FNiagaraTypeDefinition& InType) const;
-
+	
 	void ParameterValueChanged();
 
+	FReply SummonHierarchyEditor();
 	void DeleteParameter(FNiagaraVariable UserParameter) const;
 	void RequestRename(FNiagaraVariable UserParameter);
 	void RenameParameter(FNiagaraVariable UserParameter, FName NewName);
@@ -210,7 +211,7 @@ private:
 	TWeakPtr<FNiagaraSystemViewModel> SystemViewModel;
 	TOptional<FNiagaraVariable> SelectedParameter;
 	TMap<FNiagaraVariable, TSharedPtr<class SNiagaraParameterNameTextBlock>> UserParamToWidgetMap;
-	TSharedPtr<SWidget> AddParameterButtonContainer;
+	TSharedPtr<SWidget> AdditionalHeaderWidgetsContainer;
 	TSharedPtr<class SComboButton> AddParameterButton;
 	TSharedPtr<class SNiagaraAddParameterFromPanelMenu> AddParameterMenu;
 };
diff --git a/Engine/Plugins/FX/Niagara/Source/NiagaraEditor/Private/Widgets/SNiagaraSystemUserParameters.cpp b/Engine/Plugins/FX/Niagara/Source/NiagaraEditor/Private/Widgets/SNiagaraSystemUserParameters.cpp
--- a/Engine/Plugins/FX/Niagara/Source/NiagaraEditor/Private/Widgets/SNiagaraSystemUserParameters.cpp
+++ b/Engine/Plugins/FX/Niagara/Source/NiagaraEditor/Private/Widgets/SNiagaraSystemUserParameters.cpp
@@ -34,37 +34,12 @@ void SNiagaraSystemUserParameters::Construct(const FArguments& InArgs, TSharedPt
 	
 	ChildSlot
 	[
-		SNew(SVerticalBox)
-		+ SVerticalBox::Slot()
-		.AutoHeight()
-		.HAlign(HAlign_Fill)
-		.Padding(2.f)
+		SNew(SBox)
+		.AddMetaData<FTagMetaData>(FTagMetaData(TEXT("UserParameters")))
 		[
-			SNew(SHorizontalBox)
-			+ SHorizontalBox::Slot()
-			.HAlign(HAlign_Right)
-			[
-				SNew(SButton)
-				.HAlign(HAlign_Right)
-				.OnClicked(this, &SNiagaraSystemUserParameters::SummonHierarchyEditor)
-				.Text(LOCTEXT("SummonUserParametersHierarchyButtonLabel", "Edit Hierarchy"))
-			]
-		]
-		+ SVerticalBox::Slot()
-		[
-			SNew(SBox)
-			.AddMetaData<FTagMetaData>(FTagMetaData(TEXT("UserParameters")))
-			[
-				ObjectDetailsView
-			]
-		]
+			ObjectDetailsView
+		]		
 	];
 }
 
-FReply SNiagaraSystemUserParameters::SummonHierarchyEditor()
-{
-	SystemViewModel.Pin()->FocusTab(FNiagaraSystemToolkitModeBase::UserParametersHierarchyTabID);
-	return FReply::Handled();
-}
-
 #undef LOCTEXT_NAMESPACE
diff --git a/Engine/Plugins/FX/Niagara/Source/NiagaraEditor/Private/Widgets/SNiagaraSystemUserParameters.h b/Engine/Plugins/FX/Niagara/Source/NiagaraEditor/Private/Widgets/SNiagaraSystemUserParameters.h
--- a/Engine/Plugins/FX/Niagara/Source/NiagaraEditor/Private/Widgets/SNiagaraSystemUserParameters.h
+++ b/Engine/Plugins/FX/Niagara/Source/NiagaraEditor/Private/Widgets/SNiagaraSystemUserParameters.h
@@ -20,8 +20,6 @@ public:
 	
 	virtual void Tick( const FGeometry& AllottedGeometry, const double InCurrentTime, const float InDeltaTime ) override {}
 
-private:
-	FReply SummonHierarchyEditor();
 private:
 	TWeakPtr<FNiagaraSystemViewModel> SystemViewModel;
 };
