PR #9838: Public and exported FNiagaraGpuReadbackManager (Contributed by [REDACTED_PERSON])

diff --git a/Engine/Plugins/FX/Niagara/Source/Niagara/Private/NiagaraGpuReadbackManager.h b/Engine/Plugins/FX/Niagara/Source/Niagara/Private/NiagaraGpuReadbackManager.h
deleted file mode 100644
diff --git a/Engine/Plugins/FX/Niagara/Source/Niagara/Public/NiagaraGpuReadbackManager.h b/Engine/Plugins/FX/Niagara/Source/Niagara/Public/NiagaraGpuReadbackManager.h
new file mode 100644
--- /dev/null
+++ b/Engine/Plugins/FX/Niagara/Source/Niagara/Public/NiagaraGpuReadbackManager.h
@@ -0,0 +1,86 @@
+// Copyright Epic Games, Inc. All Rights Reserved.
+
+#pragma once
+
+#include "CoreMinimal.h"
+#include "Containers/Queue.h"
+#include "NiagaraCommon.h"
+#include "RenderGraphResources.h"
+#include "RHI.h"
+#include "RHICommandList.h"
+
+class FRDGBuilder;
+
+class NIAGARA_API FNiagaraGpuReadbackManager
+{
+protected:
+	typedef TFunction<void(TConstArrayView<TPair<void*, uint32>>)> FCompletionCallback;
+		
+	struct FPendingReadback
+	{
+		TArray<TPair<FStagingBufferRHIRef, uint32>, TInlineAllocator<1>>	StagingBuffers;
+		FGPUFenceRHIRef														Fence;
+		FCompletionCallback													Callback;
+	};
+
+public:
+	struct FRDGBufferRequest
+	{
+		FRDGBufferRequest() {}
+		FRDGBufferRequest(FRDGBufferRef InBuffer, uint32 InOffset, uint32 InSize) : Buffer(InBuffer), Offset(InOffset), Size(InSize) {}
+
+		FRDGBufferRef	Buffer = nullptr;
+		uint32			Offset = 0;
+		uint32			Size = 0;
+	};
+
+	struct FBufferRequest
+	{
+		FBufferRequest() {}
+		FBufferRequest(FRHIBuffer* InBuffer, uint32 InOffset, uint32 InSize) : Buffer(InBuffer), Offset(InOffset), Size(InSize) {}
+
+		FRHIBuffer*	Buffer = nullptr;
+		uint32		Offset = 0;
+		uint32		Size = 0;
+	};
+
+public:
+	FNiagaraGpuReadbackManager();
+	~FNiagaraGpuReadbackManager();
+
+	// Tick call which polls for completed readbacks
+	void Tick();
+	
+private:
+	// Internal tick impl
+	void TickInternal(bool bAssumeGpuIdle);
+
+public:
+	// Wait for all pending readbacks to complete
+	void WaitCompletion(FRHICommandListImmediate& RHICmdList);
+
+	// Enqueue a readback of a single buffer
+	void EnqueueReadback(FRDGBuilder& GraphBuilder, FRDGBufferRef Buffer, FCompletionCallback Callback);
+	void EnqueueReadback(FRDGBuilder& GraphBuilder, FRDGBufferRef Buffer, uint32 Offset, uint32 NumBytes, FCompletionCallback Callback);
+
+	// Enqueue a readback of multiple buffers
+	void EnqueueReadbacks(FRDGBuilder& GraphBuilder, TConstArrayView<FRDGBufferRef> Buffers, FCompletionCallback Callback);
+	void EnqueueReadbacks(FRDGBuilder& GraphBuilder, TConstArrayView<FRDGBufferRequest> BufferRequest, FCompletionCallback Callback);
+
+	//-TODO:RDG:Deprecate these calls
+	//UE_DEPRECATED(5.1, "Support for legacy shaders bindings will be removed in a future release, please upgrade your data interface.")
+	void EnqueueReadback(FRHICommandList& RHICmdList, FRHIBuffer* Buffer, FCompletionCallback Callback);
+	//UE_DEPRECATED(5.1, "Support for legacy shaders bindings will be removed in a future release, please upgrade your data interface.")
+	void EnqueueReadback(FRHICommandList& RHICmdList, FRHIBuffer* Buffer, uint32 Offset, uint32 NumBytes, FCompletionCallback Callback);
+
+	//-TODO:RDG:Deprecate these calls
+	//UE_DEPRECATED(5.1, "Support for legacy shaders bindings will be removed in a future release, please upgrade your data interface.")
+	void EnqueueReadbacks(FRHICommandList& RHICmdList, TConstArrayView<FRHIBuffer*> Buffers, FCompletionCallback Callback);
+	//UE_DEPRECATED(5.1, "Support for legacy shaders bindings will be removed in a future release, please upgrade your data interface.")
+	void EnqueueReadbacks(FRHICommandList& RHICmdList, TConstArrayView<FBufferRequest> BufferRequest, FCompletionCallback Callback);
+
+private:
+	TQueue<FPendingReadback> PendingReadbacks;
+
+	static const FName FenceName;
+};
