- PSO precaching options for Niagara compute shaders
- Options vary from follow precaching behavior to forcing on and checking the results are valid
- We have a lot of cases where the shader compiles offline fine but on some devices we fail due to driver specific rules or issues (i.e. too many temporary registers used)
- When in force on & check mode we block components from running until the PSO cache results are in for all emitters
- Emitters are enabled or disabled based upon the results, we currently only understand failure in Vulkan RHI
- A warning will be logged when we do fail, e.g 'LogNiagara: Warning: Niagara ComputePSOPrecache Failed for Emitter(Grid3D_Gas_CONTROLS_Emitter) System(NS_GasSmoke), emitter will not run.'

diff --git a/Engine/Plugins/FX/Niagara/Source/Niagara/Classes/NiagaraEmitter.h b/Engine/Plugins/FX/Niagara/Source/Niagara/Classes/NiagaraEmitter.h
--- a/Engine/Plugins/FX/Niagara/Source/Niagara/Classes/NiagaraEmitter.h
+++ b/Engine/Plugins/FX/Niagara/Source/Niagara/Classes/NiagaraEmitter.h
@@ -373,6 +373,10 @@ struct NIAGARA_API FVersionedNiagaraEmitterData
 	FORCEINLINE const TArray<FNiagaraEventScriptProperties>& GetEventHandlers() const { return EventHandlerScriptProps; }
 	void CacheFromCompiledData(const FNiagaraDataSetCompiledData* CompiledData, const UNiagaraEmitter& Emitter);
 	void CacheFromShaderCompiled();
+
+	FGraphEventRef PrecacheComputePSOs(const UNiagaraEmitter& NiagaraEmitter);
+	bool DidPSOPrecacheFail() const { return PSOPrecacheResult == EPSOPrecacheResult::NotSupported; }
+
 	bool RequiresViewUniformBuffer() const { return bRequiresViewUniformBuffer; }
 	uint32 GetMaxInstanceCount() const { return MaxInstanceCount; }
 	uint32 GetMaxAllocationCount() const { return MaxAllocationCount; }
@@ -527,6 +531,8 @@ private:
 
 	void EnsureScriptsPostLoaded();
 	void OnPostCompile(const UNiagaraEmitter& InEmitter);
+
+	EPSOPrecacheResult PSOPrecacheResult = EPSOPrecacheResult::Unknown;
 };
 
 /** 
diff --git a/Engine/Plugins/FX/Niagara/Source/Niagara/Classes/NiagaraSystem.h b/Engine/Plugins/FX/Niagara/Source/Niagara/Classes/NiagaraSystem.h
--- a/Engine/Plugins/FX/Niagara/Source/Niagara/Classes/NiagaraSystem.h
+++ b/Engine/Plugins/FX/Niagara/Source/Niagara/Classes/NiagaraSystem.h
@@ -374,7 +374,7 @@ private:
 	bool IsReadyToRunInternal() const;
 
 public:
-	bool IsReadyToRun() const { return FPlatformProperties::RequiresCookedData() ? bIsReadyToRunCached : IsReadyToRunInternal(); }
+	bool IsReadyToRun() const;
 
 	FORCEINLINE bool NeedsWarmup()const { return WarmupTickCount > 0 && WarmupTickDelta > SMALL_NUMBER; }
 	FORCEINLINE float GetWarmupTime()const { return WarmupTime; }
@@ -959,6 +959,7 @@ protected:
 	TArray<TArray<FNiagaraDataSetAccessor<FNiagaraSpawnInfo>>> EmitterSpawnInfoAccessors;
 
 	FGraphEventRef ScriptOptimizationCompletionEvent;
+	FGraphEventRef PSOPrecacheCompletionEvent;
 
 public:
 	void AsyncOptimizeAllScripts();
diff --git a/Engine/Plugins/FX/Niagara/Source/Niagara/Private/NiagaraEmitter.cpp b/Engine/Plugins/FX/Niagara/Source/Niagara/Private/NiagaraEmitter.cpp
--- a/Engine/Plugins/FX/Niagara/Source/Niagara/Private/NiagaraEmitter.cpp
+++ b/Engine/Plugins/FX/Niagara/Source/Niagara/Private/NiagaraEmitter.cpp
@@ -66,6 +66,18 @@ static FAutoConsoleVariableRef CVarEnableEmitterChangeIdMergeLogging(
 	ECVF_Default
 );
 
+static int32 GNiagaraEmitterComputePSOPrecacheMode = 0;
+static FAutoConsoleVariableRef CVarNiagaraEmitterComputePSOPrecacheMode(
+	TEXT("fx.Niagara.Emitter.ComputePSOPrecacheMode"),
+	GNiagaraEmitterComputePSOPrecacheMode,
+	TEXT("Controlls how PSO precaching should be done for Niagara compute shaders\n")
+	TEXT("0 = Disabled (Default).\n")
+	TEXT("1 = Enabled if r.PSOPrecaching is also enabled.\n")
+	TEXT("2 = Force Enabled.\n")
+	TEXT("3 = Force Enabled, emitters are not allowed to run until they complete."),
+	ECVF_Scalability
+);
+
 static int32 GNiagaraEmitterMaxGPUBufferElements = 0;
 static FAutoConsoleVariableRef CVarNiagaraEmitterMaxGPUBufferElements(
 	TEXT("fx.Niagara.Emitter.MaxGPUBufferElements"),
@@ -1620,7 +1632,6 @@ void FVersionedNiagaraEmitterData::CacheFromCompiledData(const FNiagaraDataSetCo
 		MaxInstanceCount = 0;
 	}
 	UpdateDebugName(Emitter, CompiledData);
-
 }
 
 void FVersionedNiagaraEmitterData::UpdateDebugName(const UNiagaraEmitter& Emitter, const FNiagaraDataSetCompiledData* CompiledData)
@@ -3239,6 +3250,113 @@ void UNiagaraEmitter::GenerateStatID()const
 #endif
 }
 
+FGraphEventRef FVersionedNiagaraEmitterData::PrecacheComputePSOs(const UNiagaraEmitter& NiagaraEmitter)
+{
+	FGraphEventRef PSOReadyGraphTask;
+	PSOPrecacheResult = EPSOPrecacheResult::Complete;
+	if (GNiagaraEmitterComputePSOPrecacheMode == 0 || !FApp::CanEverRender() || SimTarget != ENiagaraSimTarget::GPUComputeSim)
+	{
+		return PSOReadyGraphTask;
+	}
+	
+	FNiagaraShaderScript* ShaderScript = GPUComputeScript->GetRenderThreadScript();
+	if (ShaderScript == nullptr || GPUComputeScript->DidScriptCompilationSucceed(true) == false)
+	{
+		return PSOReadyGraphTask;
+	}
+
+	// Compute PSO can fail on some platforms even though the shader compiled successfully this path will wait for the PSO precache to complete before the emitter is allowed to run
+	if (GNiagaraEmitterComputePSOPrecacheMode == 3)
+	{
+		FGraphEventArray PSOPrecacheEvents;
+		for (int32 i = 0; i < ShaderScript->GetNumPermutations(); ++i)
+		{
+			FRHIComputeShader* ComputeShader = ShaderScript->GetShaderGameThread(i).GetComputeShader();
+			FGraphEventRef TaskRef = PipelineStateCache::PrecacheComputePipelineState(ComputeShader, true);
+			if (TaskRef.IsValid())
+			{
+				PSOPrecacheEvents.Add(TaskRef);
+				continue;
+			}
+
+			const EPSOPrecacheResult ShaderResult = PipelineStateCache::CheckPipelineStateInCache(ComputeShader);
+			if (ShaderResult != EPSOPrecacheResult::Complete)
+			{
+				PSOPrecacheResult = EPSOPrecacheResult::NotSupported;
+				return PSOReadyGraphTask;
+			}
+		}
+
+		const FNiagaraVMExecutableDataId VMid = GPUComputeScript->GetComputedVMCompilationId();	//-TODO:
+		struct FNiagaraEmitterPSOPrecacheReadyTask
+		{
+			explicit FNiagaraEmitterPSOPrecacheReadyTask(FVersionedNiagaraEmitterWeakPtr InVersionedEmitter)
+				: WeakVersionedEmitter(InVersionedEmitter)
+			{
+			}
+
+			FORCEINLINE TStatId GetStatId() const { RETURN_QUICK_DECLARE_CYCLE_STAT(FNiagaraEmitterPSOPrecacheReadyTask, STATGROUP_TaskGraphTasks); }
+			ENamedThreads::Type GetDesiredThread() { return ENamedThreads::GameThread; }
+			static ESubsequentsMode::Type GetSubsequentsMode() { return ESubsequentsMode::TrackSubsequents; }
+			void DoTask(ENamedThreads::Type CurrentThread, const FGraphEventRef& MyCompletionGraphEvent)
+			{
+				FVersionedNiagaraEmitter VersionedEmitter = WeakVersionedEmitter.ResolveWeakPtr();
+				FVersionedNiagaraEmitterData* EmitterData = VersionedEmitter.GetEmitterData();
+				if (EmitterData == nullptr)
+				{
+					return;
+				}
+
+				FNiagaraShaderScript* ShaderScript = EmitterData->GPUComputeScript->GetRenderThreadScript();
+				if (ShaderScript == nullptr)
+				{
+					return;
+				}
+
+				for (int32 i = 0; i < ShaderScript->GetNumPermutations(); ++i)
+				{
+					FRHIComputeShader* ComputeShader = ShaderScript->GetShaderGameThread(i).GetComputeShader();
+					FGraphEventRef TaskRef = PipelineStateCache::PrecacheComputePipelineState(ComputeShader, true);
+					const EPSOPrecacheResult ShaderResult = PipelineStateCache::CheckPipelineStateInCache(ComputeShader);
+					if (ShaderResult != EPSOPrecacheResult::Complete)
+					{
+						EmitterData->PSOPrecacheResult = EPSOPrecacheResult::NotSupported;
+
+						UNiagaraSystem* NiagaraSystem = VersionedEmitter.Emitter ? VersionedEmitter.Emitter->GetTypedOuter<UNiagaraSystem>() : nullptr;
+						UE_LOG(LogNiagara, Warning, TEXT("Niagara ComputePSOPrecache Failed for Emitter(%s) System(%s), emitter will not run."), *GetNameSafe(VersionedEmitter.Emitter), *GetNameSafe(NiagaraSystem));
+						return;
+					}
+				}
+				EmitterData->PSOPrecacheResult = EPSOPrecacheResult::Complete;
+			}
+
+			FVersionedNiagaraEmitterWeakPtr WeakVersionedEmitter;
+		};
+
+		if (PSOPrecacheEvents.Num() > 0)
+		{
+			PSOPrecacheResult = EPSOPrecacheResult::Active;
+
+			FVersionedNiagaraEmitterWeakPtr EmitterWeakPtr((UNiagaraEmitter*)(&NiagaraEmitter), Version.VersionGuid);
+			PSOReadyGraphTask = TGraphTask<FNiagaraEmitterPSOPrecacheReadyTask>::CreateTask(&PSOPrecacheEvents, ENamedThreads::GameThread).ConstructAndDispatchWhenReady(EmitterWeakPtr);
+		}
+		else
+		{
+			PSOPrecacheResult = EPSOPrecacheResult::Complete;
+		}
+	}
+	// In this mode we either force them to cache or respect that precaching is enabled
+	else if (GNiagaraEmitterComputePSOPrecacheMode == 2 || PipelineStateCache::IsPSOPrecachingEnabled() )
+	{
+		for (int32 i=0; i < ShaderScript->GetNumPermutations(); ++i)
+		{
+			FRHIComputeShader* ComputeShader = ShaderScript->GetShaderGameThread(i).GetComputeShader();
+			PipelineStateCache::PrecacheComputePipelineState(ComputeShader, true);
+		}
+	}
+	return PSOReadyGraphTask;
+}
+
 #if WITH_EDITORONLY_DATA
 FVersionedNiagaraEmitter FVersionedNiagaraEmitterData::GetParent() const
 {
diff --git a/Engine/Plugins/FX/Niagara/Source/Niagara/Private/NiagaraEmitterInstance.cpp b/Engine/Plugins/FX/Niagara/Source/Niagara/Private/NiagaraEmitterInstance.cpp
--- a/Engine/Plugins/FX/Niagara/Source/Niagara/Private/NiagaraEmitterInstance.cpp
+++ b/Engine/Plugins/FX/Niagara/Source/Niagara/Private/NiagaraEmitterInstance.cpp
@@ -235,6 +235,11 @@ bool FNiagaraEmitterInstance::IsAllowedToExecute() const
 			return false;
 		}
 
+		if (EmitterData->DidPSOPrecacheFail() == true)
+		{
+			return false;
+		}
+
 		if (const UNiagaraScript* GPUComputeScript = EmitterData->GetGPUComputeScript())
 		{
 			if (const FNiagaraShaderScript* ShaderScript = GPUComputeScript->GetRenderThreadScript())
diff --git a/Engine/Plugins/FX/Niagara/Source/Niagara/Private/NiagaraSystem.cpp b/Engine/Plugins/FX/Niagara/Source/Niagara/Private/NiagaraSystem.cpp
--- a/Engine/Plugins/FX/Niagara/Source/Niagara/Private/NiagaraSystem.cpp
+++ b/Engine/Plugins/FX/Niagara/Source/Niagara/Private/NiagaraSystem.cpp
@@ -95,17 +95,17 @@ static FAutoConsoleVariableRef CVarNiagaraCompileWaitLoggingTerminationCap(
 
 //////////////////////////////////////////////////////////////////////////
 
-class FNiagaraOptimizeCompleteTask
+class FNiagaraSystemNullCompletionTask
 {
 public:
-	FNiagaraOptimizeCompleteTask(UNiagaraSystem* Owner, FGraphEventRef* InRefToClear)
+	FNiagaraSystemNullCompletionTask(UNiagaraSystem* Owner, FGraphEventRef* InRefToClear)
 		: WeakOwner(Owner)
 		, RefToClear(InRefToClear)
 		, ReferenceValue(InRefToClear != nullptr ? InRefToClear->GetReference() : nullptr)
 	{
 	}
 
-	FORCEINLINE TStatId GetStatId() const { RETURN_QUICK_DECLARE_CYCLE_STAT(FNiagaraOptimizeCompleteTask, STATGROUP_TaskGraphTasks); }
+	FORCEINLINE TStatId GetStatId() const { RETURN_QUICK_DECLARE_CYCLE_STAT(FNiagaraSystemNullCompletionTask, STATGROUP_TaskGraphTasks); }
 	ENamedThreads::Type GetDesiredThread() { return ENamedThreads::GameThread; }
 	static ESubsequentsMode::Type GetSubsequentsMode() { return ESubsequentsMode::TrackSubsequents; }
 
@@ -1205,6 +1205,15 @@ bool UNiagaraSystem::IsReadyToRunInternal() const
 	return true;
 }
 
+bool UNiagaraSystem::IsReadyToRun() const
+{
+	if (PSOPrecacheCompletionEvent.IsValid() && !PSOPrecacheCompletionEvent->IsComplete())
+	{
+		return false;
+	}
+	return FPlatformProperties::RequiresCookedData() ? bIsReadyToRunCached : IsReadyToRunInternal();
+}
+
 void UNiagaraSystem::GetAssetRegistryTags(TArray<FAssetRegistryTag>& OutTags) const
 {
 	if (HasAnyFlags(RF_ClassDefaultObject) == false)
@@ -1780,6 +1789,8 @@ void UNiagaraSystem::CacheFromCompiledData()
 	// reset the MaxDeltaTime so we get the most up to date values from the emitters
 	MaxDeltaTime.Reset();
 
+	FGraphEventArray PSOPrecacheEvents;
+
 	TSet<FName> DataInterfaceGpuUsage;
 	TStringBuilder<128> ExecutionStateNameBuilder;
 	for (int32 i=0; i < EmitterHandles.Num(); ++i)
@@ -1815,6 +1826,11 @@ void UNiagaraSystem::CacheFromCompiledData()
 			Handle.GetInstance().Emitter->ConditionalPostLoad();
 			EmitterData->CacheFromCompiledData(DataSetCompiledData, *Handle.GetInstance().Emitter);
 
+			if ( FGraphEventRef PSOGraphEvent = EmitterData->PrecacheComputePSOs(*Handle.GetInstance().Emitter) )
+			{
+				PSOPrecacheEvents.Add(PSOGraphEvent);
+			}
+
 			// Allow data interfaces to cache static buffers
 			UNiagaraScript* NiagaraEmitterScripts[] =
 			{
@@ -1855,6 +1871,17 @@ void UNiagaraSystem::CacheFromCompiledData()
 		}
 	}
 
+	// If we had any precache events to wait on add a task to ensure the system does not start until they are ready
+	if (PSOPrecacheEvents.Num() > 0)
+	{
+		DECLARE_CYCLE_STAT(TEXT("FNullGraphTask.NiagaraPSOPrecacheCompletion"), STAT_FNullGraphTask_NiagaraPSOPrecacheCompletion, STATGROUP_TaskGraphTasks);
+		PSOPrecacheCompletionEvent = TGraphTask<FNullGraphTask>::CreateTask(&PSOPrecacheEvents, ENamedThreads::GameThread).ConstructAndDispatchWhenReady(GET_STATID(STAT_FNullGraphTask_NiagaraPSOPrecacheCompletion), ENamedThreads::AnyThread);
+
+		// Inject task to clear out the reference to the graph task as soon as possible
+		FGraphEventArray CompleteTaskPrereqs({ PSOPrecacheCompletionEvent });
+		TGraphTask<FNiagaraSystemNullCompletionTask>::CreateTask(&CompleteTaskPrereqs, ENamedThreads::GameThread).ConstructAndDispatchWhenReady(this, &PSOPrecacheCompletionEvent);
+	}
+
 	// Loop over system scripts these are more awkward because we need to determine the usage
 	for (UNiagaraScript* NiagaraScript : { SystemSpawnScript, SystemUpdateScript })
 	{
@@ -3470,7 +3497,7 @@ void UNiagaraSystem::AsyncOptimizeAllScripts()
 
 		// Inject task to clear out the reference to the graph task
 		FGraphEventArray CompleteTaskPrereqs({ ScriptOptimizationCompletionEvent });
-		TGraphTask<FNiagaraOptimizeCompleteTask>::CreateTask(&CompleteTaskPrereqs, ENamedThreads::GameThread).ConstructAndDispatchWhenReady(this, &ScriptOptimizationCompletionEvent);
+		TGraphTask<FNiagaraSystemNullCompletionTask>::CreateTask(&CompleteTaskPrereqs, ENamedThreads::GameThread).ConstructAndDispatchWhenReady(this, &ScriptOptimizationCompletionEvent);
 	}
 
 	bNeedsAsyncOptimize = false;
diff --git a/Engine/Source/Runtime/RHI/Private/PipelineStateCache.cpp b/Engine/Source/Runtime/RHI/Private/PipelineStateCache.cpp
--- a/Engine/Source/Runtime/RHI/Private/PipelineStateCache.cpp
+++ b/Engine/Source/Runtime/RHI/Private/PipelineStateCache.cpp
@@ -2194,10 +2194,10 @@ bool PipelineStateCache::IsPSOPrecachingEnabled()
 #endif // WITH_EDITOR
 }
 
-FGraphEventRef PipelineStateCache::PrecacheComputePipelineState(FRHIComputeShader* ComputeShader)
+FGraphEventRef PipelineStateCache::PrecacheComputePipelineState(FRHIComputeShader* ComputeShader, bool bForcePrecache)
 {
 	FGraphEventRef GraphEvent = nullptr;
-	if (!IsPSOPrecachingEnabled())
+	if (!IsPSOPrecachingEnabled() && !bForcePrecache)
 	{
 		return GraphEvent;
 	}
@@ -2311,7 +2311,7 @@ EPSOPrecacheResult PipelineStateCache::CheckPipelineStateInCache(const FGraphics
 
 EPSOPrecacheResult PipelineStateCache::CheckPipelineStateInCache(FRHIComputeShader* ComputeShader)
 {
-	if (!IsPSOPrecachingEnabled() || ComputeShader == nullptr)
+	if (ComputeShader == nullptr)
 	{
 		return EPSOPrecacheResult::Unknown;
 	}
@@ -2322,8 +2322,12 @@ EPSOPrecacheResult PipelineStateCache::CheckPipelineStateInCache(FRHIComputeShad
 		return EPSOPrecacheResult::Missed;
 	}
 
-	// TODO: add code to check if the PSO is still precaching or not
-	return EPSOPrecacheResult::Complete;
+	if (!CachedState->IsComplete())
+	{
+		return EPSOPrecacheResult::Active;
+	}
+
+	return CachedState->RHIPipeline && CachedState->RHIPipeline->IsValid() ? EPSOPrecacheResult::Complete : EPSOPrecacheResult::NotSupported;
 }
 
 bool PipelineStateCache::IsPrecaching(const FGraphicsPipelineStateInitializer& PipelineStateInitializer)
diff --git a/Engine/Source/Runtime/RHI/Public/PipelineStateCache.h b/Engine/Source/Runtime/RHI/Public/PipelineStateCache.h
--- a/Engine/Source/Runtime/RHI/Public/PipelineStateCache.h
+++ b/Engine/Source/Runtime/RHI/Public/PipelineStateCache.h
@@ -98,7 +98,7 @@ namespace PipelineStateCache
 	extern RHI_API bool						IsPSOPrecachingEnabled();
 
 	/* Precache the compute shader and return an optional graph event if precached async */
-	extern RHI_API FGraphEventRef			PrecacheComputePipelineState(FRHIComputeShader* ComputeShader);
+	extern RHI_API FGraphEventRef			PrecacheComputePipelineState(FRHIComputeShader* ComputeShader, bool bForcePrecache = false);
 
 	/* Precache the graphic PSO and return an optional graph event if precached async */
 	extern RHI_API FGraphEventRef			PrecacheGraphicsPipelineState(const FGraphicsPipelineStateInitializer& PipelineStateInitializer);
diff --git a/Engine/Source/Runtime/RHI/Public/RHIResources.h b/Engine/Source/Runtime/RHI/Public/RHIResources.h
--- a/Engine/Source/Runtime/RHI/Public/RHIResources.h
+++ b/Engine/Source/Runtime/RHI/Public/RHIResources.h
@@ -732,6 +732,12 @@ class FRHIComputePipelineState : public FRHIResource
 {
 public:
 	FRHIComputePipelineState() : FRHIResource(RRT_ComputePipelineState) {}
+
+	inline void SetValid(bool InIsValid) { bIsValid = InIsValid; }
+	inline bool IsValid() const { return bIsValid; }
+
+private:
+	bool bIsValid = true;
 };
 class FRHIRayTracingPipelineState : public FRHIResource
 {
diff --git a/Engine/Source/Runtime/VulkanRHI/Private/VulkanPipeline.cpp b/Engine/Source/Runtime/VulkanRHI/Private/VulkanPipeline.cpp
--- a/Engine/Source/Runtime/VulkanRHI/Private/VulkanPipeline.cpp
+++ b/Engine/Source/Runtime/VulkanRHI/Private/VulkanPipeline.cpp
@@ -2173,6 +2173,7 @@ FVulkanComputePipeline* FVulkanPipelineStateCacheManager::CreateComputePipelineF
 	{
 		FString ComputeHash = Shader->GetHash().ToString();
 		UE_LOG(LogVulkanRHI, Error, TEXT("Failed to create compute pipeline.\nShaders in pipeline: CS: %s"), *ComputeHash);
+		Pipeline->SetValid(false);
 	}
 
 	Pipeline->Layout = ComputeLayout;
