- Fix for cube texture debug visualization in Niagara

diff --git a/Engine/Plugins/FX/Niagara/Source/Niagara/Private/NiagaraGpuComputeDebug.cpp b/Engine/Plugins/FX/Niagara/Source/Niagara/Private/NiagaraGpuComputeDebug.cpp
--- a/Engine/Plugins/FX/Niagara/Source/Niagara/Private/NiagaraGpuComputeDebug.cpp
+++ b/Engine/Plugins/FX/Niagara/Source/Niagara/Private/NiagaraGpuComputeDebug.cpp
@@ -181,7 +181,7 @@ void FNiagaraGpuComputeDebug::AddAttributeTexture(FRDGBuilder& GraphBuilder, FNi
 
 	FRHICopyTextureInfo CopyInfo;
 	CopyInfo.NumMips = TextureDesc.NumMips;
-	CopyInfo.NumSlices = TextureDesc.ArraySize;
+	CopyInfo.NumSlices = TextureDesc.IsTextureCube() ? 6 : TextureDesc.ArraySize;
 	AddCopyTexturePass(GraphBuilder, Texture, GraphBuilder.RegisterExternalTexture(VisualizeEntry->Texture), CopyInfo);
 }
 
