- Add ExecToUnit for render target data interfaces

diff --git a/Engine/Plugins/FX/Niagara/Shaders/Private/NiagaraDataInterfaceRenderTarget2DArrayTemplate.ush b/Engine/Plugins/FX/Niagara/Shaders/Private/NiagaraDataInterfaceRenderTarget2DArrayTemplate.ush
--- a/Engine/Plugins/FX/Niagara/Shaders/Private/NiagaraDataInterfaceRenderTarget2DArrayTemplate.ush
+++ b/Engine/Plugins/FX/Niagara/Shaders/Private/NiagaraDataInterfaceRenderTarget2DArrayTemplate.ush
@@ -41,3 +41,11 @@ void ExecToIndex_{ParameterName}(out int IndexX, out int IndexY, out int IndexZ)
 {
 	LinearToIndex_{ParameterName}(ExecIndex(), IndexX, IndexY, IndexZ);
 }
+
+void ExecToUnit_{ParameterName}(out float2 Unit, out int Slice)
+{
+	int2 Texel;
+	ExecToIndex_{ParameterName}(Texel.x, Texel.y, Slice);
+	Unit.x = (float(Texel.x) + 0.5f) / float({ParameterName}_TextureSize.x);
+	Unit.y = (float(Texel.y) + 0.5f) / float({ParameterName}_TextureSize.y);
+}
diff --git a/Engine/Plugins/FX/Niagara/Shaders/Private/NiagaraDataInterfaceRenderTarget2DTemplate.ush b/Engine/Plugins/FX/Niagara/Shaders/Private/NiagaraDataInterfaceRenderTarget2DTemplate.ush
--- a/Engine/Plugins/FX/Niagara/Shaders/Private/NiagaraDataInterfaceRenderTarget2DTemplate.ush
+++ b/Engine/Plugins/FX/Niagara/Shaders/Private/NiagaraDataInterfaceRenderTarget2DTemplate.ush
@@ -39,3 +39,11 @@ void ExecToIndex_{ParameterName}(out int IndexX, out int IndexY)
 {
 	LinearToIndex_{ParameterName}(ExecIndex(), IndexX, IndexY);
 }
+
+void ExecToUnit_{ParameterName}(out float2 Unit)
+{
+	int2 Texel;
+	LinearToIndex_{ParameterName}(ExecIndex(), Texel.x, Texel.y);
+	Unit.x = (float(Texel.x) + 0.5f) / float({ParameterName}_TextureSize.x);
+	Unit.y = (float(Texel.y) + 0.5f) / float({ParameterName}_TextureSize.y);
+}
diff --git a/Engine/Plugins/FX/Niagara/Shaders/Private/NiagaraDataInterfaceRenderTargetCubeTemplate.ush b/Engine/Plugins/FX/Niagara/Shaders/Private/NiagaraDataInterfaceRenderTargetCubeTemplate.ush
--- a/Engine/Plugins/FX/Niagara/Shaders/Private/NiagaraDataInterfaceRenderTargetCubeTemplate.ush
+++ b/Engine/Plugins/FX/Niagara/Shaders/Private/NiagaraDataInterfaceRenderTargetCubeTemplate.ush
@@ -40,3 +40,11 @@ void ExecToIndex_{ParameterName}(out int IndexX, out int IndexY, out int Face)
 {
 	LinearToIndex_{ParameterName}(ExecIndex(), IndexX, IndexY, Face);
 }
+
+void ExecToUnit_{ParameterName}(out float2 Unit, out int Face)
+{
+	int2 Texel;
+	ExecToIndex_{ParameterName}(Texel.x, Texel.y, Face);
+	Unit.x = (float(Texel.x) + 0.5f) / float({ParameterName}_TextureSize);
+	Unit.y = (float(Texel.y) + 0.5f) / float({ParameterName}_TextureSize);
+}
diff --git a/Engine/Plugins/FX/Niagara/Shaders/Private/NiagaraDataInterfaceRenderTargetVolumeTemplate.ush b/Engine/Plugins/FX/Niagara/Shaders/Private/NiagaraDataInterfaceRenderTargetVolumeTemplate.ush
--- a/Engine/Plugins/FX/Niagara/Shaders/Private/NiagaraDataInterfaceRenderTargetVolumeTemplate.ush
+++ b/Engine/Plugins/FX/Niagara/Shaders/Private/NiagaraDataInterfaceRenderTargetVolumeTemplate.ush
@@ -41,3 +41,12 @@ void ExecToIndex_{ParameterName}(out int IndexX, out int IndexY, out int IndexZ)
 {
 	LinearToIndex_{ParameterName}(ExecIndex(), IndexX, IndexY, IndexZ);
 }
+
+void ExecToUnit_{ParameterName}(out float3 Unit)
+{
+	int3 Texel;
+	ExecToIndex_{ParameterName}(Texel.x, Texel.y, Texel.z);
+	Unit.x = (float(Texel.x) + 0.5f) / float({ParameterName}_TextureSize.x);
+	Unit.y = (float(Texel.y) + 0.5f) / float({ParameterName}_TextureSize.y);
+	Unit.z = (float(Texel.z) + 0.5f) / float({ParameterName}_TextureSize.z);
+}
diff --git a/Engine/Plugins/FX/Niagara/Source/Niagara/Private/NiagaraDataInterfaceRenderTarget2D.cpp b/Engine/Plugins/FX/Niagara/Source/Niagara/Private/NiagaraDataInterfaceRenderTarget2D.cpp
--- a/Engine/Plugins/FX/Niagara/Source/Niagara/Private/NiagaraDataInterfaceRenderTarget2D.cpp
+++ b/Engine/Plugins/FX/Niagara/Source/Niagara/Private/NiagaraDataInterfaceRenderTarget2D.cpp
@@ -43,6 +43,7 @@ namespace NDIRenderTarget2DLocal
 	static const FName GetSizeFunctionName("GetRenderTargetSize");
 	static const FName LinearToIndexName("LinearToIndex");
 	static const FName ExecToIndexName("ExecToIndex");
+	static const FName ExecToUnitName("ExecToUnit");
 }
 
 FNiagaraVariableBase UNiagaraDataInterfaceRenderTarget2D::ExposedRTVar;
@@ -243,6 +244,23 @@ void UNiagaraDataInterfaceRenderTarget2D::GetFunctions(TArray<FNiagaraFunctionSi
 		Sig.FunctionVersion = FNDIRenderTarget2DFunctionVersion::LatestVersion;
 	#endif
 	}
+
+	{
+		FNiagaraFunctionSignature& Sig = OutFunctions.AddDefaulted_GetRef();
+		Sig.Name = ExecToUnitName;
+		Sig.Inputs.Emplace(FNiagaraTypeDefinition(GetClass()), TEXT("RenderTarget"));
+		Sig.Outputs.Emplace(FNiagaraTypeDefinition::GetVec2Def(), TEXT("Unit"));
+
+		Sig.bExperimental = true;
+		Sig.bMemberFunction = true;
+		Sig.bRequiresContext = false;
+		Sig.bWriteFunction = true;
+		Sig.bSupportsCPU = false;
+		Sig.bSupportsGPU = true;
+#if WITH_EDITORONLY_DATA
+		Sig.FunctionVersion = FNDIRenderTarget2DFunctionVersion::LatestVersion;
+#endif
+	}
 }
 
 #if WITH_EDITORONLY_DATA
@@ -372,7 +390,8 @@ bool UNiagaraDataInterfaceRenderTarget2D::GetFunctionHLSL(const FNiagaraDataInte
 		(FunctionInfo.DefinitionName == SampleValueFunctionName) ||
 		(FunctionInfo.DefinitionName == GetSizeFunctionName) ||
 		(FunctionInfo.DefinitionName == LinearToIndexName) ||
-		(FunctionInfo.DefinitionName == ExecToIndexName))
+		(FunctionInfo.DefinitionName == ExecToIndexName) ||
+		(FunctionInfo.DefinitionName == ExecToUnitName))
 	{
 		return true;
 	}
diff --git a/Engine/Plugins/FX/Niagara/Source/Niagara/Private/NiagaraDataInterfaceRenderTarget2DArray.cpp b/Engine/Plugins/FX/Niagara/Source/Niagara/Private/NiagaraDataInterfaceRenderTarget2DArray.cpp
--- a/Engine/Plugins/FX/Niagara/Source/Niagara/Private/NiagaraDataInterfaceRenderTarget2DArray.cpp
+++ b/Engine/Plugins/FX/Niagara/Source/Niagara/Private/NiagaraDataInterfaceRenderTarget2DArray.cpp
@@ -41,6 +41,7 @@ namespace NDIRenderTarget2DArrayLocal
 	static const FName GetSizeFunctionName("GetRenderTargetSize");
 	static const FName LinearToIndexName("LinearToIndex");
 	static const FName ExecToIndexName("ExecToIndex");
+	static const FName ExecToUnitName("ExecToUnit");
 }
 
 FNiagaraVariableBase UNiagaraDataInterfaceRenderTarget2DArray::ExposedRTVar;
@@ -247,6 +248,24 @@ void UNiagaraDataInterfaceRenderTarget2DArray::GetFunctions(TArray<FNiagaraFunct
 		Sig.FunctionVersion = FNDIRenderTarget2DArrayFunctionVersion::LatestVersion;
 	#endif
 	}
+
+	{
+		FNiagaraFunctionSignature& Sig = OutFunctions.AddDefaulted_GetRef();
+		Sig.Name = ExecToUnitName;
+		Sig.Inputs.Emplace(FNiagaraTypeDefinition(GetClass()), TEXT("RenderTarget"));
+		Sig.Outputs.Emplace(FNiagaraTypeDefinition::GetVec2Def(), TEXT("Unit"));
+		Sig.Outputs.Emplace(FNiagaraTypeDefinition::GetIntDef(), TEXT("Slice"));
+
+		Sig.bExperimental = true;
+		Sig.bMemberFunction = true;
+		Sig.bRequiresContext = false;
+		Sig.bWriteFunction = true;
+		Sig.bSupportsCPU = false;
+		Sig.bSupportsGPU = true;
+#if WITH_EDITORONLY_DATA
+		Sig.FunctionVersion = FNDIRenderTarget2DArrayFunctionVersion::LatestVersion;
+#endif
+	}
 }
 
 #if WITH_EDITORONLY_DATA
@@ -370,7 +389,8 @@ bool UNiagaraDataInterfaceRenderTarget2DArray::GetFunctionHLSL(const FNiagaraDat
 		(FunctionInfo.DefinitionName == SampleValueFunctionName) ||
 		(FunctionInfo.DefinitionName == GetSizeFunctionName) ||
 		(FunctionInfo.DefinitionName == LinearToIndexName) ||
-		(FunctionInfo.DefinitionName == ExecToIndexName))
+		(FunctionInfo.DefinitionName == ExecToIndexName) ||
+		(FunctionInfo.DefinitionName == ExecToUnitName))
 	{
 		return true;
 	}
diff --git a/Engine/Plugins/FX/Niagara/Source/Niagara/Private/NiagaraDataInterfaceRenderTargetCube.cpp b/Engine/Plugins/FX/Niagara/Source/Niagara/Private/NiagaraDataInterfaceRenderTargetCube.cpp
--- a/Engine/Plugins/FX/Niagara/Source/Niagara/Private/NiagaraDataInterfaceRenderTargetCube.cpp
+++ b/Engine/Plugins/FX/Niagara/Source/Niagara/Private/NiagaraDataInterfaceRenderTargetCube.cpp
@@ -39,6 +39,7 @@ namespace NDIRenderTargetCubeLocal
 	static const FName GetSizeFunctionName("GetRenderTargetSize");
 	static const FName LinearToIndexName("LinearToIndex");
 	static const FName ExecToIndexName("ExecToIndex");
+	static const FName ExecToUnitName("ExecToUnit");
 
 	struct EFunctionVersion
 	{
@@ -240,6 +241,24 @@ void UNiagaraDataInterfaceRenderTargetCube::GetFunctions(TArray<FNiagaraFunction
 		Sig.FunctionVersion = EFunctionVersion::LatestVersion;
 	#endif
 	}
+
+	{
+		FNiagaraFunctionSignature& Sig = OutFunctions.AddDefaulted_GetRef();
+		Sig.Name = ExecToUnitName;
+		Sig.Inputs.Emplace(FNiagaraTypeDefinition(GetClass()), TEXT("RenderTarget"));
+		Sig.Outputs.Emplace(FNiagaraTypeDefinition::GetVec2Def(), TEXT("Unit"));
+		Sig.Outputs.Emplace(FNiagaraTypeDefinition::GetIntDef(), TEXT("Face"));
+
+		Sig.bExperimental = true;
+		Sig.bMemberFunction = true;
+		Sig.bRequiresContext = false;
+		Sig.bWriteFunction = true;
+		Sig.bSupportsCPU = false;
+		Sig.bSupportsGPU = true;
+#if WITH_EDITORONLY_DATA
+		Sig.FunctionVersion = EFunctionVersion::LatestVersion;
+#endif
+	}
 }
 
 #if WITH_EDITORONLY_DATA
@@ -358,7 +377,8 @@ bool UNiagaraDataInterfaceRenderTargetCube::GetFunctionHLSL(const FNiagaraDataIn
 		(FunctionInfo.DefinitionName == SampleValueFunctionName) ||
 		(FunctionInfo.DefinitionName == GetSizeFunctionName) ||
 		(FunctionInfo.DefinitionName == LinearToIndexName) ||
-		(FunctionInfo.DefinitionName == ExecToIndexName))
+		(FunctionInfo.DefinitionName == ExecToIndexName) ||
+		(FunctionInfo.DefinitionName == ExecToUnitName))
 	{
 		return true;
 	}
diff --git a/Engine/Plugins/FX/Niagara/Source/Niagara/Private/NiagaraDataInterfaceRenderTargetVolume.cpp b/Engine/Plugins/FX/Niagara/Source/Niagara/Private/NiagaraDataInterfaceRenderTargetVolume.cpp
--- a/Engine/Plugins/FX/Niagara/Source/Niagara/Private/NiagaraDataInterfaceRenderTargetVolume.cpp
+++ b/Engine/Plugins/FX/Niagara/Source/Niagara/Private/NiagaraDataInterfaceRenderTargetVolume.cpp
@@ -41,6 +41,7 @@ namespace NDIRenderTargetVolumeLocal
 	static const FName GetSizeFunctionName("GetRenderTargetSize");
 	static const FName LinearToIndexName("LinearToIndex");
 	static const FName ExecToIndexName("ExecToIndex");
+	static const FName ExecToUnitName("ExecToUnit");
 
 	struct EFunctionVersion
 	{
@@ -262,6 +263,23 @@ void UNiagaraDataInterfaceRenderTargetVolume::GetFunctions(TArray<FNiagaraFuncti
 		Sig.FunctionVersion = EFunctionVersion::LatestVersion;
 	#endif
 	}
+
+	{
+		FNiagaraFunctionSignature& Sig = OutFunctions.AddDefaulted_GetRef();
+		Sig.Name = ExecToUnitName;
+		Sig.Inputs.Emplace(FNiagaraTypeDefinition(GetClass()), TEXT("RenderTarget"));
+		Sig.Outputs.Emplace(FNiagaraTypeDefinition::GetVec3Def(), TEXT("Unit"));
+
+		Sig.bExperimental = true;
+		Sig.bMemberFunction = true;
+		Sig.bRequiresContext = false;
+		Sig.bWriteFunction = true;
+		Sig.bSupportsCPU = false;
+		Sig.bSupportsGPU = true;
+#if WITH_EDITORONLY_DATA
+		Sig.FunctionVersion = EFunctionVersion::LatestVersion;
+#endif
+	}
 }
 
 #if WITH_EDITORONLY_DATA
@@ -384,7 +402,8 @@ bool UNiagaraDataInterfaceRenderTargetVolume::GetFunctionHLSL(const FNiagaraData
 		(FunctionInfo.DefinitionName == SampleValueFunctionName) ||
 		(FunctionInfo.DefinitionName == GetSizeFunctionName) ||
 		(FunctionInfo.DefinitionName == LinearToIndexName) ||
-		(FunctionInfo.DefinitionName == ExecToIndexName))
+		(FunctionInfo.DefinitionName == ExecToIndexName) ||
+		(FunctionInfo.DefinitionName == ExecToUnitName))
 	{
 		return true;
 	}
