- Don't allow volume texture data interfaces to compile for a platform if they do not support RHIVolumeTextureRenderingSupportGuaranteed
- This is a conservative check to reduce risk of crashes on target devices, this could be loosened to a runtime check

diff --git a/Engine/Plugins/FX/Niagara/Source/Niagara/Classes/NiagaraDataInterfaceGrid3DCollection.h b/Engine/Plugins/FX/Niagara/Source/Niagara/Classes/NiagaraDataInterfaceGrid3DCollection.h
--- a/Engine/Plugins/FX/Niagara/Source/Niagara/Classes/NiagaraDataInterfaceGrid3DCollection.h
+++ b/Engine/Plugins/FX/Niagara/Source/Niagara/Classes/NiagaraDataInterfaceGrid3DCollection.h
@@ -164,6 +164,9 @@ public:
 	virtual bool Equals(const UNiagaraDataInterface* Other) const override;
 
 	// GPU sim functionality
+#if WITH_EDITOR
+	virtual bool ShouldCompile(EShaderPlatform ShaderPlatform) const override;
+#endif
 #if WITH_EDITORONLY_DATA
 	virtual bool AppendCompileHash(FNiagaraCompileHashVisitor* InVisitor) const override;
 	virtual void GetParameterDefinitionHLSL(const FNiagaraDataInterfaceGPUParamInfo& ParamInfo, FString& OutHLSL) override;
diff --git a/Engine/Plugins/FX/Niagara/Source/Niagara/Classes/NiagaraDataInterfaceRenderTargetVolume.h b/Engine/Plugins/FX/Niagara/Source/Niagara/Classes/NiagaraDataInterfaceRenderTargetVolume.h
--- a/Engine/Plugins/FX/Niagara/Source/Niagara/Classes/NiagaraDataInterfaceRenderTargetVolume.h
+++ b/Engine/Plugins/FX/Niagara/Source/Niagara/Classes/NiagaraDataInterfaceRenderTargetVolume.h
@@ -93,6 +93,9 @@ public:
 	virtual bool CopyToInternal(UNiagaraDataInterface* Destination) const override;
 
 	// GPU sim functionality
+#if WITH_EDITOR
+	virtual bool ShouldCompile(EShaderPlatform ShaderPlatform) const override;
+#endif
 #if WITH_EDITORONLY_DATA
 	virtual bool AppendCompileHash(FNiagaraCompileHashVisitor* InVisitor) const override;
 	virtual void GetParameterDefinitionHLSL(const FNiagaraDataInterfaceGPUParamInfo& ParamInfo, FString& OutHLSL) override;
diff --git a/Engine/Plugins/FX/Niagara/Source/Niagara/Private/NiagaraDataInterfaceGrid3DCollection.cpp b/Engine/Plugins/FX/Niagara/Source/Niagara/Private/NiagaraDataInterfaceGrid3DCollection.cpp
--- a/Engine/Plugins/FX/Niagara/Source/Niagara/Private/NiagaraDataInterfaceGrid3DCollection.cpp
+++ b/Engine/Plugins/FX/Niagara/Source/Niagara/Private/NiagaraDataInterfaceGrid3DCollection.cpp
@@ -336,6 +336,15 @@ void UNiagaraDataInterfaceGrid3DCollection::PostInitProperties()
 	}
 }
 
+#if WITH_EDITOR
+bool UNiagaraDataInterfaceGrid3DCollection::ShouldCompile(EShaderPlatform ShaderPlatform) const
+{
+	return
+		RHIVolumeTextureRenderingSupportGuaranteed(ShaderPlatform) &&
+		Super::ShouldCompile(ShaderPlatform);
+}
+#endif
+
 #if WITH_EDITORONLY_DATA
 bool UNiagaraDataInterfaceGrid3DCollection::AppendCompileHash(FNiagaraCompileHashVisitor* InVisitor) const
 {
diff --git a/Engine/Plugins/FX/Niagara/Source/Niagara/Private/NiagaraDataInterfaceRasterizationGrid3D.cpp b/Engine/Plugins/FX/Niagara/Source/Niagara/Private/NiagaraDataInterfaceRasterizationGrid3D.cpp
--- a/Engine/Plugins/FX/Niagara/Source/Niagara/Private/NiagaraDataInterfaceRasterizationGrid3D.cpp
+++ b/Engine/Plugins/FX/Niagara/Source/Niagara/Private/NiagaraDataInterfaceRasterizationGrid3D.cpp
@@ -500,12 +500,10 @@ bool UNiagaraDataInterfaceRasterizationGrid3D::Equals(const UNiagaraDataInterfac
 #if WITH_EDITOR
 bool UNiagaraDataInterfaceRasterizationGrid3D::ShouldCompile(EShaderPlatform ShaderPlatform) const
 {
-	if (!RHISupportsVolumeTextureAtomics(ShaderPlatform))
-	{
-		return false;
-	}
-
-	return UNiagaraDataInterface::ShouldCompile(ShaderPlatform);
+	return
+		RHISupportsVolumeTextureAtomics(ShaderPlatform) &&
+		RHIVolumeTextureRenderingSupportGuaranteed(ShaderPlatform) &&
+		Super::ShouldCompile(ShaderPlatform);
 }
 #endif
 
diff --git a/Engine/Plugins/FX/Niagara/Source/Niagara/Private/NiagaraDataInterfaceRenderTargetVolume.cpp b/Engine/Plugins/FX/Niagara/Source/Niagara/Private/NiagaraDataInterfaceRenderTargetVolume.cpp
--- a/Engine/Plugins/FX/Niagara/Source/Niagara/Private/NiagaraDataInterfaceRenderTargetVolume.cpp
+++ b/Engine/Plugins/FX/Niagara/Source/Niagara/Private/NiagaraDataInterfaceRenderTargetVolume.cpp
@@ -371,6 +371,15 @@ bool UNiagaraDataInterfaceRenderTargetVolume::CopyToInternal(UNiagaraDataInterfa
 	return true;
 }
 
+#if WITH_EDITOR
+bool UNiagaraDataInterfaceRenderTargetVolume::ShouldCompile(EShaderPlatform ShaderPlatform) const
+{
+	return
+		RHIVolumeTextureRenderingSupportGuaranteed(ShaderPlatform) &&
+		Super::ShouldCompile(ShaderPlatform);
+}
+#endif
+
 #if WITH_EDITORONLY_DATA
 bool UNiagaraDataInterfaceRenderTargetVolume::AppendCompileHash(FNiagaraCompileHashVisitor* InVisitor) const
 {
