Making static variable adjustments in the UI work again by propagating graph changed notifications that were removed in a prior checkin in certain circumstances.

diff --git a/Engine/Plugins/FX/Niagara/Source/NiagaraEditor/Private/ViewModels/Stack/NiagaraStackFunctionInput.cpp b/Engine/Plugins/FX/Niagara/Source/NiagaraEditor/Private/ViewModels/Stack/NiagaraStackFunctionInput.cpp
--- a/Engine/Plugins/FX/Niagara/Source/NiagaraEditor/Private/ViewModels/Stack/NiagaraStackFunctionInput.cpp
+++ b/Engine/Plugins/FX/Niagara/Source/NiagaraEditor/Private/ViewModels/Stack/NiagaraStackFunctionInput.cpp
@@ -1789,6 +1789,11 @@ void UNiagaraStackFunctionInput::SetLocalValue(TSharedRef<FStructOnScope> InLoca
 			if (InputType.IsStatic()) // Need to potentially trigger a recompile.
 			{
 				Script->MarkScriptAndSourceDesynchronized(TEXT("Static variable changed!"), FGuid::NewGuid());
+				UNiagaraScriptSource* Source = Cast<UNiagaraScriptSource>(Script->GetLatestSource());
+				if (Source && Source->NodeGraph)
+				{
+					Source->NodeGraph->NotifyGraphChanged();
+				}
 			}
 		}
 
