- Niagara debug draw now uses override max lines as a limit on CPU and takes this value if set on the GPU

#rnx

diff --git a/Engine/Plugins/FX/Niagara/Source/Niagara/Private/DataInterface/NiagaraDataInterfaceDebugDraw.cpp b/Engine/Plugins/FX/Niagara/Source/Niagara/Private/DataInterface/NiagaraDataInterfaceDebugDraw.cpp
--- a/Engine/Plugins/FX/Niagara/Source/Niagara/Private/DataInterface/NiagaraDataInterfaceDebugDraw.cpp
+++ b/Engine/Plugins/FX/Niagara/Source/Niagara/Private/DataInterface/NiagaraDataInterfaceDebugDraw.cpp
@@ -53,13 +53,16 @@ struct FNDIDebugDrawInstanceData_GameThread
 		//-OPT: Need to improve this
 		FScopeLock RWLock(&LineBufferLock);
 
-		auto& Line = LineBuffer.AddDefaulted_GetRef();
-		Line.Start = Start;
-		Line.End = End;
-		Line.Color = uint32(FMath::Clamp(Color.R, 0.0f, 1.0f) * 255.0f) << 24;
-		Line.Color |= uint32(FMath::Clamp(Color.G, 0.0f, 1.0f) * 255.0f) << 16;
-		Line.Color |= uint32(FMath::Clamp(Color.B, 0.0f, 1.0f) * 255.0f) << 8;
-		Line.Color |= uint32(FMath::Clamp(Color.A, 0.0f, 1.0f) * 255.0f) << 0;
+		if (OverrideMaxLineInstances == 0 || uint32(LineBuffer.Num()) < OverrideMaxLineInstances)
+		{
+			auto& Line = LineBuffer.AddDefaulted_GetRef();
+			Line.Start = Start;
+			Line.End = End;
+			Line.Color = uint32(FMath::Clamp(Color.R, 0.0f, 1.0f) * 255.0f) << 24;
+			Line.Color |= uint32(FMath::Clamp(Color.G, 0.0f, 1.0f) * 255.0f) << 16;
+			Line.Color |= uint32(FMath::Clamp(Color.B, 0.0f, 1.0f) * 255.0f) << 8;
+			Line.Color |= uint32(FMath::Clamp(Color.A, 0.0f, 1.0f) * 255.0f) << 0;
+		}
 	}
 
 	void AddSphere(const FVector3f& Location, float Radius, int32 Segments, const FLinearColor& Color)
@@ -421,7 +424,7 @@ struct FNDIDebugDrawInstanceData_GameThread
 	bool bResolvedPersistentShapes = false;
 	FCriticalSection LineBufferLock;
 	TArray<FNiagaraSimulationDebugDrawData::FGpuLine> LineBuffer;
-
+	uint32 OverrideMaxLineInstances = 0;
 
 	template<typename TYPE>
 	static inline TOptional<TYPE> GetCompileTag(FNiagaraSystemInstance* SystemInstance, const UNiagaraScript* Script, const FNiagaraTypeDefinition& VarType, const FName& ParameterName)
@@ -3039,7 +3042,8 @@ bool UNiagaraDataInterfaceDebugDraw::PerInstanceTick(void* PerInstanceData, FNia
 {
 #if NIAGARA_COMPUTEDEBUG_ENABLED
 	FNDIDebugDrawInstanceData_GameThread* InstanceData = reinterpret_cast<FNDIDebugDrawInstanceData_GameThread*>(PerInstanceData);
-	InstanceData->LineBuffer.Reset();
+	InstanceData->LineBuffer.Reset(OverrideMaxLineInstances);
+	InstanceData->OverrideMaxLineInstances = OverrideMaxLineInstances;
 #endif
 	return false;
 }
diff --git a/Engine/Plugins/FX/Niagara/Source/Niagara/Private/NiagaraGpuComputeDebug.cpp b/Engine/Plugins/FX/Niagara/Source/Niagara/Private/NiagaraGpuComputeDebug.cpp
--- a/Engine/Plugins/FX/Niagara/Source/Niagara/Private/NiagaraGpuComputeDebug.cpp
+++ b/Engine/Plugins/FX/Niagara/Source/Niagara/Private/NiagaraGpuComputeDebug.cpp
@@ -205,7 +205,7 @@ FNiagaraSimulationDebugDrawData* FNiagaraGpuComputeDebug::GetSimulationDebugDraw
 	FNiagaraSimulationDebugDrawData* DebugDrawData = GetSimulationDebugDrawData(SystemInstanceID);
 
 	const uint32 MaxPossibleLines = FMath::Min(FMath::DivideAndRoundDown(GetMaxBufferDimension(), uint64(NiagaraGpuComputeDebugLocal::NumUintsPerLine)), uint64(TNumericLimits<uint32>::Max() >> 1));
-	const uint32 MaxLineInstancesToUse = FMath::Min(FMath::Max3(DebugDrawData->GpuLineMaxInstances, uint32(GNiagaraGpuComputeDebug_MaxLineInstances), OverrideMaxDebugLines), MaxPossibleLines);
+	const uint32 MaxLineInstancesToUse = FMath::Min(FMath::Max(DebugDrawData->GpuLineMaxInstances, OverrideMaxDebugLines > 0 ? OverrideMaxDebugLines : uint32(GNiagaraGpuComputeDebug_MaxLineInstances)), MaxPossibleLines);
 	if (DebugDrawData->GpuLineMaxInstances != MaxLineInstancesToUse)
 	{
 		DebugDrawData->GpuLineBufferArgs.Release();
