- Optionally hash all data interfaces to generate the compile hash
- This is a temporary solution until generating the compile hash correctly traverses all data interfaces

#rnx

diff --git a/Engine/Plugins/FX/Niagara/Source/Niagara/Private/NiagaraScript.cpp b/Engine/Plugins/FX/Niagara/Source/Niagara/Private/NiagaraScript.cpp
--- a/Engine/Plugins/FX/Niagara/Source/Niagara/Private/NiagaraScript.cpp
+++ b/Engine/Plugins/FX/Niagara/Source/Niagara/Private/NiagaraScript.cpp
@@ -49,6 +49,9 @@
 #include "NiagaraSimulationStageBase.h"
 #include "Async/Async.h"
 
+#include "Algo/Sort.h"
+#include "Misc/CoreMiscDefines.h"
+
 #include UE_INLINE_GENERATED_CPP_BY_NAME(NiagaraScript)
 
 #if ENABLE_COOK_STATS
@@ -118,11 +121,71 @@ static FAutoConsoleVariableRef CVarNiagaraEventSpawnsUpdateInitialAttributeValue
 	ECVF_Default
 );
 
-namespace NiagaraScriptLocal
+static int32 GNiagaraCompileHashAllDataInterfaces = 1;
+static FAutoConsoleVariableRef CVarNiagaraCompileHashAllDataInterfaces(
+	TEXT("fx.Niagara.CompileHashAllDataInterfaces"),
+	GNiagaraCompileHashAllDataInterfaces,
+	TEXT("Forces Niagara to include all data interfaces as part of compile hash generation.\n")
+	TEXT("0 = Disabled.\n")
+	TEXT("1 = Enabled, but hash once, better for perf.\n")
+	TEXT("2 = Enabled, generated hash always good for iteration on Niagara work."),
+	ECVF_Default
+);
+
+namespace NiagaraScriptInternal
 {
 #if WITH_EDITORONLY_DATA
 	TOptional<ERHIFeatureLevel::Type> PreviewFeatureLevel;
+
+	static FNiagaraCompileHash& CompileHashAllDataInterfaces()
+	{
+		static TOptional<FNiagaraCompileHash> AllDataInterfacesHash;
+		if (AllDataInterfacesHash.IsSet() == false || GNiagaraCompileHashAllDataInterfaces > 1)
+		{
+			UE_CALL_ONCE(
+				[]()
+				{
+					TArray<UClass*> DataInterfaceClasses;
+					for (TObjectIterator<UClass> ClassIterator; ClassIterator; ++ClassIterator)
+					{
+						UClass* Class = *ClassIterator;
+						if (Class == nullptr ||
+							Class->HasAnyClassFlags(CLASS_Abstract | CLASS_Deprecated | CLASS_NewerVersionExists) ||
+							!ClassIterator->IsChildOf(UNiagaraDataInterface::StaticClass()))
+						{
+							continue;
+						}
+						DataInterfaceClasses.Add(Class);
+					}
+					Algo::Sort(
+						DataInterfaceClasses,
+						[](UClass* Lhs, UClass* Rhs)
+						{
+							return Lhs->GetName() < Rhs->GetName();
+						}
+					);
+
+					FSHA1 HashState;
+					FNiagaraCompileHashVisitor Visitor(HashState);
+					for (UClass* DIClass : DataInterfaceClasses)
+					{
+						UNiagaraDataInterface* CDODataInterface = DIClass->GetDefaultObject<UNiagaraDataInterface>();
+						CDODataInterface->AppendCompileHash(&Visitor);
+					}
+					HashState.Final();
+
+					TArray<uint8> DataHash;
+					DataHash.AddUninitialized(FSHA1::DigestSize);
+					HashState.GetHash(DataHash.GetData());
+
+					AllDataInterfacesHash.Emplace(DataHash);
+				}
+			);
+		}
+		return AllDataInterfacesHash.GetValue();
+	}
 #endif
+
 }
 
 FNiagaraScriptDebuggerInfo::FNiagaraScriptDebuggerInfo() : bWaitForGPU(false), Usage(ENiagaraScriptUsage::Function), FrameLastWriteId(-1), bWritten(false)
@@ -1122,6 +1185,13 @@ void UNiagaraScript::ComputeVMCompilationId(FNiagaraVMExecutableDataId& Id, FGui
 			Id.AdditionalDefines.Add(TEXT("Emitter.Determinism"));
 		}
 
+		// Append all Shader Parameter compile hashes
+		if (GNiagaraCompileHashAllDataInterfaces)
+		{
+			Id.ReferencedCompileHashes.AddUnique(NiagaraScriptInternal::CompileHashAllDataInterfaces());
+			Id.DebugReferencedObjects.Add(TEXT("AllDataInterfaceHashes"));
+		}
+
 		// Has simulation stages
 		{
 
@@ -2239,7 +2309,7 @@ bool UNiagaraScript::ShouldCacheShadersForCooking(const ITargetPlatform* TargetP
 void UNiagaraScript::SetPreviewFeatureLevel(ERHIFeatureLevel::Type InPreviewFeatureLevel)
 {
 	// Anything to do?
-	if ( NiagaraScriptLocal::PreviewFeatureLevel.Get(GMaxRHIFeatureLevel) == InPreviewFeatureLevel )
+	if ( NiagaraScriptInternal::PreviewFeatureLevel.Get(GMaxRHIFeatureLevel) == InPreviewFeatureLevel )
 	{
 		return;
 	}
@@ -2247,11 +2317,11 @@ void UNiagaraScript::SetPreviewFeatureLevel(ERHIFeatureLevel::Type InPreviewFeat
 	// Either Set of Clear the Preview Feature Level
 	if ( InPreviewFeatureLevel == GMaxRHIFeatureLevel )
 	{
-		NiagaraScriptLocal::PreviewFeatureLevel.Reset();
+		NiagaraScriptInternal::PreviewFeatureLevel.Reset();
 	}
 	else
 	{
-		NiagaraScriptLocal::PreviewFeatureLevel = InPreviewFeatureLevel;
+		NiagaraScriptInternal::PreviewFeatureLevel = InPreviewFeatureLevel;
 	}
 
 	// Force scripts to recache that can run on the GPU and already have a script resource
@@ -3337,7 +3407,7 @@ void UNiagaraScript::CacheResourceShadersForRendering(bool bRegenerateId, bool b
 		if (Source && OwnerCanBeRunOnGpu())
 		{
 		#if WITH_EDITORONLY_DATA
-			const ERHIFeatureLevel::Type CacheFeatureLevel = NiagaraScriptLocal::PreviewFeatureLevel.Get(GMaxRHIFeatureLevel);
+			const ERHIFeatureLevel::Type CacheFeatureLevel = NiagaraScriptInternal::PreviewFeatureLevel.Get(GMaxRHIFeatureLevel);
 		#else
 			const ERHIFeatureLevel::Type CacheFeatureLevel = GMaxRHIFeatureLevel;
 		#endif
