Fix for low latency translucency from being used when underwater (using single layer water) as it results in trasnlucency being rendered before the opaque, preventing LLT from working properly.

diff --git a/Engine/Plugins/FX/Niagara/Source/Niagara/Private/NiagaraRenderer.cpp b/Engine/Plugins/FX/Niagara/Source/Niagara/Private/NiagaraRenderer.cpp
--- a/Engine/Plugins/FX/Niagara/Source/Niagara/Private/NiagaraRenderer.cpp
+++ b/Engine/Plugins/FX/Niagara/Source/Niagara/Private/NiagaraRenderer.cpp
@@ -495,6 +495,17 @@ bool FNiagaraRenderer::UseLocalSpace(const FNiagaraSceneProxy* Proxy)const
 	return bLocalSpace || Proxy->GetProxyDynamicData().bUseCullProxy;
 }
 
+bool FNiagaraRenderer::ViewFamilySupportLowLatencyTranslucency(const FSceneViewFamily& ViewFamily)
+{
+	// when underwater, single layer water will render translucency before opaque
+	const bool bIncludesUnderwaterView = ViewFamily.Views.ContainsByPredicate([](const FSceneView* SceneView)
+	{
+		return SceneView->IsUnderwater();
+	});
+
+	return !bIncludesUnderwaterView;
+}
+
 void FNiagaraRenderer::ProcessMaterialParameterBindings(const FNiagaraRendererMaterialParameters& MaterialParameters, const FNiagaraEmitterInstance* InEmitter, TConstArrayView<UMaterialInterface*> InMaterials) const
 {
 	if (MaterialParameters.HasAnyBindings() == false || !InEmitter)
diff --git a/Engine/Plugins/FX/Niagara/Source/Niagara/Private/NiagaraRendererMeshes.cpp b/Engine/Plugins/FX/Niagara/Source/Niagara/Private/NiagaraRendererMeshes.cpp
--- a/Engine/Plugins/FX/Niagara/Source/Niagara/Private/NiagaraRendererMeshes.cpp
+++ b/Engine/Plugins/FX/Niagara/Source/Niagara/Private/NiagaraRendererMeshes.cpp
@@ -295,7 +295,9 @@ void FNiagaraRendererMeshes::PrepareParticleMeshRenderData(FParticleMeshRenderDa
 		!bIsWireframe &&
 		!SceneProxy->CastsVolumetricTranslucentShadow() &&
 		ParticleMeshRenderData.DynamicDataMesh->Materials.Num() > 0 &&
-		ParticleMeshRenderData.DynamicDataMesh->IsGpuLowLatencyTranslucencyEnabled();
+		ParticleMeshRenderData.DynamicDataMesh->IsGpuLowLatencyTranslucencyEnabled() &&
+		ViewFamilySupportLowLatencyTranslucency(ViewFamily);
+
 	ParticleMeshRenderData.bHasTranslucentMaterials = false;
 	for (FMaterialRenderProxy* MaterialProxy : ParticleMeshRenderData.DynamicDataMesh->Materials)
 	{
diff --git a/Engine/Plugins/FX/Niagara/Source/Niagara/Private/NiagaraRendererSprites.cpp b/Engine/Plugins/FX/Niagara/Source/Niagara/Private/NiagaraRendererSprites.cpp
--- a/Engine/Plugins/FX/Niagara/Source/Niagara/Private/NiagaraRendererSprites.cpp
+++ b/Engine/Plugins/FX/Niagara/Source/Niagara/Private/NiagaraRendererSprites.cpp
@@ -209,7 +209,7 @@ void FNiagaraRendererSprites::CreateRenderThreadResources()
 #endif
 }
 
-void FNiagaraRendererSprites::PrepareParticleSpriteRenderData(FParticleSpriteRenderData& ParticleSpriteRenderData, FNiagaraDynamicDataBase* InDynamicData, const FNiagaraSceneProxy* SceneProxy) const
+void FNiagaraRendererSprites::PrepareParticleSpriteRenderData(FParticleSpriteRenderData& ParticleSpriteRenderData, const FSceneViewFamily& ViewFamily, FNiagaraDynamicDataBase* InDynamicData, const FNiagaraSceneProxy* SceneProxy) const
 {
 	ParticleSpriteRenderData.DynamicDataSprites = static_cast<FNiagaraDynamicDataSprites*>(InDynamicData);
 	if (!ParticleSpriteRenderData.DynamicDataSprites || !SceneProxy->GetComputeDispatchInterface())
@@ -232,7 +232,13 @@ void FNiagaraRendererSprites::PrepareParticleSpriteRenderData(FParticleSpriteRen
 	const EBlendMode BlendMode = MaterialRenderProxy->GetIncompleteMaterialWithFallback(FeatureLevel).GetBlendMode();
 	ParticleSpriteRenderData.BlendMode = BlendMode;
 	ParticleSpriteRenderData.bHasTranslucentMaterials = IsTranslucentBlendMode(BlendMode);
-	ParticleSpriteRenderData.SourceParticleData = ParticleSpriteRenderData.DynamicDataSprites->GetParticleDataToRender(ParticleSpriteRenderData.bHasTranslucentMaterials && bGpuLowLatencyTranslucency && !SceneProxy->CastsVolumetricTranslucentShadow());
+
+	const bool bLowLatencyTranslucencyEnabled = ParticleSpriteRenderData.bHasTranslucentMaterials &&
+		bGpuLowLatencyTranslucency &&
+		!SceneProxy->CastsVolumetricTranslucentShadow() &&
+		ViewFamilySupportLowLatencyTranslucency(ViewFamily);
+
+	ParticleSpriteRenderData.SourceParticleData = ParticleSpriteRenderData.DynamicDataSprites->GetParticleDataToRender(bLowLatencyTranslucencyEnabled);
 	if ( !ParticleSpriteRenderData.SourceParticleData || (SourceMode == ENiagaraRendererSourceDataMode::Particles && ParticleSpriteRenderData.SourceParticleData->GetNumInstances() == 0) )
 	{
 		ParticleSpriteRenderData.SourceParticleData = nullptr;
@@ -924,7 +930,7 @@ void FNiagaraRendererSprites::GetDynamicMeshElements(const TArray<const FSceneVi
 	// Prepare our particle render data
 	// This will also determine if we have anything to render
 	FParticleSpriteRenderData ParticleSpriteRenderData;
-	PrepareParticleSpriteRenderData(ParticleSpriteRenderData, DynamicDataRender, SceneProxy);
+	PrepareParticleSpriteRenderData(ParticleSpriteRenderData, ViewFamily, DynamicDataRender, SceneProxy);
 
 	if (ParticleSpriteRenderData.SourceParticleData == nullptr)
 	{
@@ -1037,7 +1043,7 @@ void FNiagaraRendererSprites::GetDynamicRayTracingInstances(FRayTracingMaterialG
 	// Prepare our particle render data
 	// This will also determine if we have anything to render
 	FParticleSpriteRenderData ParticleSpriteRenderData;
-	PrepareParticleSpriteRenderData(ParticleSpriteRenderData, DynamicDataRender, SceneProxy);
+	PrepareParticleSpriteRenderData(ParticleSpriteRenderData, *Context.ReferenceView->Family, DynamicDataRender, SceneProxy);
 
 	if (ParticleSpriteRenderData.SourceParticleData == nullptr)
 	{
diff --git a/Engine/Plugins/FX/Niagara/Source/Niagara/Public/NiagaraRenderer.h b/Engine/Plugins/FX/Niagara/Source/Niagara/Public/NiagaraRenderer.h
--- a/Engine/Plugins/FX/Niagara/Source/Niagara/Public/NiagaraRenderer.h
+++ b/Engine/Plugins/FX/Niagara/Source/Niagara/Public/NiagaraRenderer.h
@@ -153,6 +153,8 @@ protected:
 	bool IsRendererEnabled(const UNiagaraRendererProperties* InProperties, const FNiagaraEmitterInstance* Emitter) const;
 	bool UseLocalSpace(const FNiagaraSceneProxy* Proxy) const;
 
+	static bool ViewFamilySupportLowLatencyTranslucency(const FSceneViewFamily& ViewFamily);
+
 	struct FNiagaraDynamicDataBase *DynamicDataRender;
 
 #if RHI_RAYTRACING
diff --git a/Engine/Plugins/FX/Niagara/Source/Niagara/Public/NiagaraRendererSprites.h b/Engine/Plugins/FX/Niagara/Source/Niagara/Public/NiagaraRendererSprites.h
--- a/Engine/Plugins/FX/Niagara/Source/Niagara/Public/NiagaraRendererSprites.h
+++ b/Engine/Plugins/FX/Niagara/Source/Niagara/Public/NiagaraRendererSprites.h
@@ -71,7 +71,7 @@ private:
 		FNiagaraSpriteUniformBufferRef	UniformBuffer;
 	};
 
-	void PrepareParticleSpriteRenderData(FParticleSpriteRenderData& ParticleSpriteRenderData, FNiagaraDynamicDataBase* InDynamicData, const FNiagaraSceneProxy* SceneProxy) const;
+	void PrepareParticleSpriteRenderData(FParticleSpriteRenderData& ParticleSpriteRenderData, const FSceneViewFamily& ViewFamily, FNiagaraDynamicDataBase* InDynamicData, const FNiagaraSceneProxy* SceneProxy) const;
 	void PrepareParticleRenderBuffers(FParticleSpriteRenderData& ParticleSpriteRenderData, FGlobalDynamicReadBuffer& DynamicReadBuffer) const;
 	void InitializeSortInfo(FParticleSpriteRenderData& ParticleSpriteRenderData, const FNiagaraSceneProxy& SceneProxy, const FSceneView& View, int32 ViewIndex, FNiagaraGPUSortInfo& OutSortInfo) const;
 	void SetupVertexFactory(FParticleSpriteRenderData& ParticleSpriteRenderData, FNiagaraSpriteVertexFactory& VertexFactory) const;
