- Use point sampling for cached height texture if linear is not support on the device

diff --git a/Engine/Plugins/FX/Niagara/Source/Niagara/Private/NiagaraDataInterfaceLandscape.cpp b/Engine/Plugins/FX/Niagara/Source/Niagara/Private/NiagaraDataInterfaceLandscape.cpp
--- a/Engine/Plugins/FX/Niagara/Source/Niagara/Private/NiagaraDataInterfaceLandscape.cpp
+++ b/Engine/Plugins/FX/Niagara/Source/Niagara/Private/NiagaraDataInterfaceLandscape.cpp
@@ -551,9 +551,6 @@ struct FNDILandscapeData_RenderThread
 					ShaderParameters->CachedPhysMatTextureDimension = PhysMatDimensions;
 				}
 
-				FRHISamplerState* PointClampedSampler = TStaticSamplerState<SF_Point, AM_Clamp, AM_Clamp, AM_Clamp>::GetRHI();
-				ShaderParameters->PointClampedSampler = PointClampedSampler;
-
 				return true;
 			}
 		}
@@ -573,9 +570,6 @@ struct FNDILandscapeData_RenderThread
 		ShaderParameters->CachedHeightTextureEnabled = 0;
 
 		ShaderParameters->CachedPhysMatTexture = GBlackTextureWithSRV->ShaderResourceViewRHI;
-
-		FRHISamplerState* PointClampedSampler = TStaticSamplerState<SF_Point, AM_Clamp, AM_Clamp, AM_Clamp>::GetRHI();
-		ShaderParameters->PointClampedSampler = PointClampedSampler;
 	}
 };
 
@@ -1310,11 +1304,14 @@ void UNiagaraDataInterfaceLandscape::SetShaderParameters(const FNiagaraDataInter
 
 	// Set Samplers
 	FRHISamplerState* BilinearSamplerState				= TStaticSamplerState<SF_Bilinear, AM_Clamp, AM_Clamp, AM_Clamp>::GetRHI();
+	FRHISamplerState* PointClampedSampler				= TStaticSamplerState<SF_Point, AM_Clamp, AM_Clamp, AM_Clamp>::GetRHI();
+
 	ShaderParameters->BaseColorVirtualTextureSampler	= BilinearSamplerState;
 	ShaderParameters->HeightVirtualTextureSampler		= BilinearSamplerState;
 	ShaderParameters->NormalVirtualTexture0Sampler		= BilinearSamplerState;
 	ShaderParameters->NormalVirtualTexture1Sampler		= BilinearSamplerState;
-	ShaderParameters->CachedHeightTextureSampler		= BilinearSamplerState;
+	ShaderParameters->CachedHeightTextureSampler		= RHIPixelFormatHasCapabilities(EPixelFormat::PF_R32_FLOAT, EPixelFormatCapabilities::TextureFilterable) ? BilinearSamplerState : PointClampedSampler;
+	ShaderParameters->PointClampedSampler				= PointClampedSampler;
 
 	// Set Textures
 	bool ApplyBaseColorVirtualTextureDefaults = true;
