Static Mesh DI LOD selection now available on a user paramter.

diff --git a/Engine/Plugins/FX/Niagara/Source/Niagara/Private/DataInterface/NiagaraDataInterfaceStaticMesh.cpp b/Engine/Plugins/FX/Niagara/Source/Niagara/Private/DataInterface/NiagaraDataInterfaceStaticMesh.cpp
--- a/Engine/Plugins/FX/Niagara/Source/Niagara/Private/DataInterface/NiagaraDataInterfaceStaticMesh.cpp
+++ b/Engine/Plugins/FX/Niagara/Source/Niagara/Private/DataInterface/NiagaraDataInterfaceStaticMesh.cpp
@@ -606,6 +606,8 @@ namespace NDIStaticMeshLocal
 		/** The cached LODIdx used to initialize the FNDIStaticMesh_InstanceData.*/
 		int32 CachedLODIdx = 0;
 
+		FNiagaraParameterDirectBinding<int32> LODIndexUserBinding;
+
 		/** Cached socket information, if available */
 		TArray<FTransform3f> CachedSockets;
 
@@ -620,6 +622,45 @@ namespace NDIStaticMeshLocal
 
 		FMeshUvMappingUsage UvMappingUsage;
 
+		bool CalculateLODIndex(UNiagaraDataInterfaceStaticMesh* Interface, FNiagaraSystemInstance* Instance, UStaticMesh* StaticMesh, int32& OutLOD, int32& OutMinLOD)
+		{
+			FStaticMeshRenderData* RenderData = StaticMesh->GetRenderData();
+
+			OutLOD = INDEX_NONE;
+			OutMinLOD = INDEX_NONE;
+			if(Interface == nullptr || Instance == nullptr || RenderData == nullptr)
+			{
+				return false;
+			}
+
+			OutMinLOD = StaticMesh->GetMinLOD().GetValue();
+			if (GNDIStaticMesh_UseInlineLODsOnly == 1 || (GNDIStaticMesh_UseInlineLODsOnly == 2 && !Interface->bAllowSamplingFromStreamingLODs))
+			{
+				OutMinLOD = StaticMesh->GetNumLODs() - StaticMesh->GetRenderData()->NumInlinedLODs;
+			}
+
+			int32 RequestedLOD = Interface->LODIndex;
+			if(LODIndexUserBinding.IsBound())
+			{
+				RequestedLOD = LODIndexUserBinding.GetValue();
+			}
+			else if(Interface->LODIndexUserParameter.Parameter.IsValid())
+			{
+				if (int32* UserLODIndex = LODIndexUserBinding.Init(Instance->GetInstanceParameters(), Interface->LODIndexUserParameter.Parameter))
+				{
+					RequestedLOD = *UserLODIndex;
+				}				
+			}
+
+			int32 MinSampleLOD = RenderData->GetCurrentFirstLODIdx(MinLOD);
+			int32 MaxSampleLOD = StaticMesh->GetNumLODs() - 1;
+
+			OutLOD = RequestedLOD >= 0 ? RequestedLOD : (StaticMesh->GetNumLODs() + RequestedLOD);
+			OutLOD = FMath::Clamp(CachedLODIdx, MinSampleLOD, MaxSampleLOD);
+
+			return StaticMesh->GetRenderData()->LODResources.IsValidIndex(OutLOD);
+		}
+
 #if WITH_EDITOR
 		FDelegateHandle OnMeshChanged;
 		FDelegateHandle OnPostMeshBuild;
@@ -743,22 +784,10 @@ namespace NDIStaticMeshLocal
 			TRefCountPtr<const FStaticMeshLODResources> LODData;
 			if (StaticMesh != nullptr)
 			{
-				// Check if any valid LODs are found. If not, we won't use this mesh
-				MinLOD = StaticMesh->GetMinLOD().GetValue();
-				if (GNDIStaticMesh_UseInlineLODsOnly == 1 || (GNDIStaticMesh_UseInlineLODsOnly == 2 && !Interface->bAllowSamplingFromStreamingLODs))
-				{
-					MinLOD = StaticMesh->GetNumLODs() - StaticMesh->GetRenderData()->NumInlinedLODs;
-				}
-
 				FStaticMeshRenderData* RenderData = StaticMesh->GetRenderData();
 				if (RenderData)
 				{
-					int32 MinSampleLOD = RenderData->GetCurrentFirstLODIdx(MinLOD);
-					int32 MaxSampleLOD = StaticMesh->GetNumLODs() - 1;
-
-					CachedLODIdx = Interface->LODIndex >= 0 ? Interface->LODIndex : (StaticMesh->GetNumLODs() + Interface->LODIndex);
-					CachedLODIdx = FMath::Clamp(CachedLODIdx, MinSampleLOD, MaxSampleLOD);
-					if (RenderData->LODResources.IsValidIndex(CachedLODIdx))
+					if (CalculateLODIndex(Interface, SystemInstance, StaticMesh, CachedLODIdx, MinLOD))
 					{
 						LODData = &RenderData->LODResources[CachedLODIdx];
 					}
@@ -919,7 +948,11 @@ namespace NDIStaticMeshLocal
 
 		bool Tick(UNiagaraDataInterfaceStaticMesh* Interface, FNiagaraSystemInstance* SystemInstance, float InDeltaSeconds)
 		{
-			if (ResetRequired(Interface))
+			//TODO: Recache LOD info.
+			//TODO: Currently we just grab a LOD on init and then fail/reset if that's changed or becomes invalid.
+			//TODO: If we recache here, the DI can just continue with new LODs when desired and more gracefully handle streamed out LODs.
+
+			if (ResetRequired(Interface, SystemInstance))
 			{
 				return true;
 			}
@@ -970,7 +1003,7 @@ namespace NDIStaticMeshLocal
 #endif
 		}
 
-		bool ResetRequired(UNiagaraDataInterfaceStaticMesh* Interface) const
+		bool ResetRequired(UNiagaraDataInterfaceStaticMesh* Interface, FNiagaraSystemInstance* Instance)
 		{
 			USceneComponent* Component = SceneComponentWeakPtr.Get();
 			if (bComponentValid && !Component)
@@ -1004,9 +1037,12 @@ namespace NDIStaticMeshLocal
 
 			if (Mesh != nullptr)
 			{
-				// Currently we only reset if the cached LOD was streamed out, to avoid performance hits. To revisit.
-				// We could probably just recache the data derived from the LOD instead of resetting everything.
-				if (Mesh->GetRenderData()->GetCurrentFirstLODIdx(MinLOD) > CachedLODIdx)
+				//TODO: If the LOD index selection changes then we reset the system
+				//TODO: We should just recache all the relevant data and continue on our way.
+				int32 NewLODIndex = INDEX_NONE;
+				int32 NewMinLOD = INDEX_NONE;
+				bool bSuccess = CalculateLODIndex(Interface, Instance, Mesh, NewLODIndex, NewMinLOD);
+				if (!bSuccess || NewLODIndex != CachedLODIdx)
 				{
 					return true;
 				}
@@ -1526,6 +1562,8 @@ UNiagaraDataInterfaceStaticMesh::UNiagaraDataInterfaceStaticMesh(FObjectInitiali
 	: Super(ObjectInitializer)
 {
 	Proxy.Reset(new NDIStaticMeshLocal::FRenderProxy());
+
+	LODIndexUserParameter.Parameter.SetType(FNiagaraTypeDefinition::GetIntDef());
 }
 
 void UNiagaraDataInterfaceStaticMesh::PostInitProperties()
@@ -1582,6 +1620,9 @@ void UNiagaraDataInterfaceStaticMesh::PostEditChangeProperty(struct FPropertyCha
 			DefaultMesh = nullptr;
 		}
 	}
+
+	//Temp hack to prevent the reset to default breaking this user binding.
+	LODIndexUserParameter.Parameter.SetType(FNiagaraTypeDefinition::GetIntDef());
 }
 
 bool UNiagaraDataInterfaceStaticMesh::CanEditChange(const FProperty* InProperty) const
@@ -3089,6 +3130,7 @@ bool UNiagaraDataInterfaceStaticMesh::Equals(const UNiagaraDataInterface* Other)
 		OtherTyped->bUsePhysicsBodyVelocity == bUsePhysicsBodyVelocity &&
 		OtherTyped->bAllowSamplingFromStreamingLODs == bAllowSamplingFromStreamingLODs &&
 		OtherTyped->LODIndex == LODIndex &&
+		OtherTyped->LODIndexUserParameter == LODIndexUserParameter &&
 		OtherTyped->FilteredSockets == FilteredSockets;
 }
 
@@ -3113,6 +3155,7 @@ bool UNiagaraDataInterfaceStaticMesh::CopyToInternal(UNiagaraDataInterface* Dest
 	OtherTyped->FilteredSockets = FilteredSockets;
 	OtherTyped->bAllowSamplingFromStreamingLODs = bAllowSamplingFromStreamingLODs;
 	OtherTyped->LODIndex = LODIndex;
+	OtherTyped->LODIndexUserParameter = LODIndexUserParameter;
 	OtherTyped->BindSourceDelegates();
 	return true;
 }
diff --git a/Engine/Plugins/FX/Niagara/Source/Niagara/Private/DataInterface/NiagaraDataInterfaceStaticMesh.h b/Engine/Plugins/FX/Niagara/Source/Niagara/Private/DataInterface/NiagaraDataInterfaceStaticMesh.h
--- a/Engine/Plugins/FX/Niagara/Source/Niagara/Private/DataInterface/NiagaraDataInterfaceStaticMesh.h
+++ b/Engine/Plugins/FX/Niagara/Source/Niagara/Private/DataInterface/NiagaraDataInterfaceStaticMesh.h
@@ -106,6 +106,10 @@ public:
 	UPROPERTY(EditAnywhere, Category = "LOD")
 	int32 LODIndex = 0;
 
+	/** Reference to a user parameter if we're reading one. */
+	UPROPERTY(EditAnywhere, Category = "LOD")
+	FNiagaraUserParameterBinding LODIndexUserParameter;
+
 	/** List of filtered sockets to use. */
 	UPROPERTY(EditAnywhere, Category = "Mesh")
 	TArray<FName> FilteredSockets;
diff --git a/Engine/Plugins/FX/Niagara/Source/Niagara/Public/NiagaraParameterStore.h b/Engine/Plugins/FX/Niagara/Source/Niagara/Public/NiagaraParameterStore.h
--- a/Engine/Plugins/FX/Niagara/Source/Niagara/Public/NiagaraParameterStore.h
+++ b/Engine/Plugins/FX/Niagara/Source/Niagara/Public/NiagaraParameterStore.h
@@ -903,6 +903,8 @@ struct FNiagaraParameterDirectBinding
 #endif
 	{}
 
+	FORCEINLINE bool IsBound()const { return ValuePtr != nullptr; }
+
 	T* Init(FNiagaraParameterStore& InStore, const FNiagaraVariable& DestVariable)
 	{
 #if NIAGARA_VALIDATE_DIRECT_BINDINGS
@@ -958,6 +960,8 @@ struct FNiagaraParameterDirectBinding<FMatrix44f>
 #endif
 	{}
 
+	FORCEINLINE bool IsBound()const { return ValuePtr != nullptr; }
+
 	FMatrix44f* Init(FNiagaraParameterStore& InStore, const FNiagaraVariable& DestVariable)
 	{
 #if NIAGARA_VALIDATE_DIRECT_BINDINGS
@@ -1014,6 +1018,8 @@ struct FNiagaraParameterDirectBinding<FVector4f>
 #endif
 	{}
 
+	FORCEINLINE bool IsBound()const { return ValuePtr != nullptr; }
+
 	FVector4f* Init(FNiagaraParameterStore& InStore, const FNiagaraVariable& DestVariable)
 	{
 #if NIAGARA_VALIDATE_DIRECT_BINDINGS
@@ -1070,6 +1076,8 @@ struct FNiagaraParameterDirectBinding<FQuat4f>
 #endif
 	{}
 
+	FORCEINLINE bool IsBound()const { return ValuePtr != nullptr; }
+
 	FQuat4f* Init(FNiagaraParameterStore& InStore, const FNiagaraVariable& DestVariable)
 	{
 #if NIAGARA_VALIDATE_DIRECT_BINDINGS
@@ -1126,6 +1134,8 @@ struct FNiagaraParameterDirectBinding<FNiagaraBool>
 #endif
 	{}
 
+	FORCEINLINE bool IsBound()const { return ValuePtr != nullptr; }
+
 	uint32* Init(FNiagaraParameterStore& InStore, const FNiagaraVariable& DestVariable)
 	{
 #if NIAGARA_VALIDATE_DIRECT_BINDINGS
@@ -1205,6 +1215,8 @@ struct FNiagaraParameterDirectBinding<UObject*>
 		: UObjectOffset(INDEX_NONE), BoundStore(nullptr)
 	{}
 
+	FORCEINLINE bool IsBound()const { return UObjectOffset != INDEX_NONE; }
+
 	UObject* Init(FNiagaraParameterStore& InStore, const FNiagaraVariable& DestVariable)
 	{
 		if (DestVariable.IsValid())
