- Add Load support to texture data interfaces

diff --git a/Engine/Plugins/FX/Niagara/Shaders/Private/NiagaraDataInterfaceTexture2DArrayTemplate.ush b/Engine/Plugins/FX/Niagara/Shaders/Private/NiagaraDataInterfaceTexture2DArrayTemplate.ush
--- a/Engine/Plugins/FX/Niagara/Shaders/Private/NiagaraDataInterfaceTexture2DArrayTemplate.ush
+++ b/Engine/Plugins/FX/Niagara/Shaders/Private/NiagaraDataInterfaceTexture2DArrayTemplate.ush
@@ -4,6 +4,11 @@ float3			{ParameterName}_TextureSize;
 Texture2DArray	{ParameterName}_Texture;
 SamplerState	{ParameterName}_TextureSampler;
 
+void LoadTexture_{ParameterName}(in int TexelX, in int TexelY, in int TexelZ, in int MipLevel, out float4 OutValue)
+{
+	OutValue = {ParameterName}_Texture.Load(int4(TexelX, TexelY, TexelZ, MipLevel));
+}
+
 void SampleTexture_{ParameterName}(in float3 UVW, in float MipLevel, out float4 OutValue)
 {
 	OutValue = {ParameterName}_Texture.SampleLevel({ParameterName}_TextureSampler, UVW, MipLevel);
diff --git a/Engine/Plugins/FX/Niagara/Shaders/Private/NiagaraDataInterfaceTextureTemplate.ush b/Engine/Plugins/FX/Niagara/Shaders/Private/NiagaraDataInterfaceTextureTemplate.ush
--- a/Engine/Plugins/FX/Niagara/Shaders/Private/NiagaraDataInterfaceTextureTemplate.ush
+++ b/Engine/Plugins/FX/Niagara/Shaders/Private/NiagaraDataInterfaceTextureTemplate.ush
@@ -5,6 +5,11 @@ int				{ParameterName}_MipLevels;
 Texture2D		{ParameterName}_Texture;
 SamplerState	{ParameterName}_TextureSampler;
 
+void LoadTexture2D_{ParameterName}(in int TexelX, in int TexelY, in int MipLevel, out float4 OutValue)
+{
+	OutValue = {ParameterName}_Texture.Load(int3(TexelX, TexelY, MipLevel));
+}
+
 void SampleTexture2D_{ParameterName}(in float2 UV, in float MipLevel, out float4 OutValue)
 {
 	OutValue = {ParameterName}_Texture.SampleLevel({ParameterName}_TextureSampler, UV, MipLevel);
diff --git a/Engine/Plugins/FX/Niagara/Shaders/Private/NiagaraDataInterfaceVolumeTextureTemplate.ush b/Engine/Plugins/FX/Niagara/Shaders/Private/NiagaraDataInterfaceVolumeTextureTemplate.ush
--- a/Engine/Plugins/FX/Niagara/Shaders/Private/NiagaraDataInterfaceVolumeTextureTemplate.ush
+++ b/Engine/Plugins/FX/Niagara/Shaders/Private/NiagaraDataInterfaceVolumeTextureTemplate.ush
@@ -4,6 +4,11 @@ float3			{ParameterName}_TextureSize;
 Texture3D		{ParameterName}_Texture;
 SamplerState	{ParameterName}_TextureSampler;
 
+void LoadVolumeTexture_{ParameterName}(in int TexelX, in int TexelY, in int TexelZ, in int MipLevel, out float4 OutValue)
+{
+	OutValue = {ParameterName}_Texture.Load(int4(TexelX, TexelY, TexelZ, MipLevel));
+}
+
 void SampleVolumeTexture_{ParameterName}(in float3 UV, in float MipLevel, out float4 OutValue)
 {
 	OutValue = {ParameterName}_Texture.SampleLevel({ParameterName}_TextureSampler, UV, MipLevel);
diff --git a/Engine/Plugins/FX/Niagara/Source/Niagara/Classes/NiagaraDataInterface2DArrayTexture.h b/Engine/Plugins/FX/Niagara/Source/Niagara/Classes/NiagaraDataInterface2DArrayTexture.h
--- a/Engine/Plugins/FX/Niagara/Source/Niagara/Classes/NiagaraDataInterface2DArrayTexture.h
+++ b/Engine/Plugins/FX/Niagara/Source/Niagara/Classes/NiagaraDataInterface2DArrayTexture.h
@@ -66,6 +66,7 @@ protected:
 
 protected:
 	static const TCHAR* TemplateShaderFilePath;
+	static const FName LoadTextureName;
 	static const FName SampleTextureName;
 	static const FName TextureDimsName;
 };
diff --git a/Engine/Plugins/FX/Niagara/Source/Niagara/Classes/NiagaraDataInterfaceTexture.h b/Engine/Plugins/FX/Niagara/Source/Niagara/Classes/NiagaraDataInterfaceTexture.h
--- a/Engine/Plugins/FX/Niagara/Source/Niagara/Classes/NiagaraDataInterfaceTexture.h
+++ b/Engine/Plugins/FX/Niagara/Source/Niagara/Classes/NiagaraDataInterfaceTexture.h
@@ -71,6 +71,7 @@ protected:
 
 protected:
 	static const TCHAR* TemplateShaderFilePath;
+	static const FName LoadTexture2DName;
 	static const FName SampleTexture2DName;
 	static const FName SamplePseudoVolumeTextureName;
 	static const FName GetTextureDimensionsName;
diff --git a/Engine/Plugins/FX/Niagara/Source/Niagara/Classes/NiagaraDataInterfaceVolumeTexture.h b/Engine/Plugins/FX/Niagara/Source/Niagara/Classes/NiagaraDataInterfaceVolumeTexture.h
--- a/Engine/Plugins/FX/Niagara/Source/Niagara/Classes/NiagaraDataInterfaceVolumeTexture.h
+++ b/Engine/Plugins/FX/Niagara/Source/Niagara/Classes/NiagaraDataInterfaceVolumeTexture.h
@@ -67,6 +67,7 @@ protected:
 
 protected:
 	static const TCHAR* TemplateShaderFilePath;
+	static const FName LoadVolumeTextureName;
 	static const FName SampleVolumeTextureName;
 	static const FName TextureDimsName;
 };
diff --git a/Engine/Plugins/FX/Niagara/Source/Niagara/Private/NiagaraDataInterface2DArrayTexture.cpp b/Engine/Plugins/FX/Niagara/Source/Niagara/Private/NiagaraDataInterface2DArrayTexture.cpp
--- a/Engine/Plugins/FX/Niagara/Source/Niagara/Private/NiagaraDataInterface2DArrayTexture.cpp
+++ b/Engine/Plugins/FX/Niagara/Source/Niagara/Private/NiagaraDataInterface2DArrayTexture.cpp
@@ -18,6 +18,7 @@
 #define LOCTEXT_NAMESPACE "UNiagaraDataInterface2DArrayTexture"
 
 const TCHAR* UNiagaraDataInterface2DArrayTexture::TemplateShaderFilePath = TEXT("/Plugin/FX/Niagara/Private/NiagaraDataInterfaceTexture2DArrayTemplate.ush");
+const FName UNiagaraDataInterface2DArrayTexture::LoadTextureName(TEXT("LoadTexture"));
 const FName UNiagaraDataInterface2DArrayTexture::SampleTextureName(TEXT("SampleTexture"));
 const FName UNiagaraDataInterface2DArrayTexture::TextureDimsName(TEXT("TextureDimensions"));
 
@@ -104,7 +105,24 @@ bool UNiagaraDataInterface2DArrayTexture::Equals(const UNiagaraDataInterface* Ot
 
 void UNiagaraDataInterface2DArrayTexture::GetFunctions(TArray<FNiagaraFunctionSignature>& OutFunctions)
 {
-	OutFunctions.Reserve(OutFunctions.Num() + 2);
+	OutFunctions.Reserve(OutFunctions.Num() + 3);
+
+
+	{
+		FNiagaraFunctionSignature& Sig = OutFunctions.AddDefaulted_GetRef();
+		Sig.Name = LoadTextureName;
+		Sig.bMemberFunction = true;
+		Sig.bRequiresContext = false;
+		Sig.bSupportsCPU = false;
+		Sig.bSupportsGPU = true;
+		Sig.Inputs.Emplace(FNiagaraTypeDefinition(GetClass()), TEXT("Texture"));
+		Sig.Inputs.Emplace(FNiagaraTypeDefinition::GetIntDef(), TEXT("TexelX"));
+		Sig.Inputs.Emplace(FNiagaraTypeDefinition::GetIntDef(), TEXT("TexelY"));
+		Sig.Inputs.Emplace(FNiagaraTypeDefinition::GetIntDef(), TEXT("TexelZ"));
+		Sig.Inputs.Emplace(FNiagaraTypeDefinition::GetIntDef(), TEXT("MipLevel"));
+		Sig.Outputs.Emplace(FNiagaraTypeDefinition::GetVec4Def(), TEXT("Value"));
+		Sig.SetDescription(LOCTEXT("TextureLoad2DArrayTextureDesc", "Read a texel from the provided location & mip without any filtering or sampling."));
+	}
 
 	{
 		FNiagaraFunctionSignature& Sig = OutFunctions.AddDefaulted_GetRef();
@@ -113,11 +131,11 @@ void UNiagaraDataInterface2DArrayTexture::GetFunctions(TArray<FNiagaraFunctionSi
 		Sig.bRequiresContext = false;
 		Sig.bSupportsCPU = false;
 		Sig.bSupportsGPU = true;
-		Sig.Inputs.Add(FNiagaraVariable(FNiagaraTypeDefinition(GetClass()), TEXT("Texture")));
-		Sig.Inputs.Add(FNiagaraVariable(FNiagaraTypeDefinition::GetVec3Def(), TEXT("UVW")));
-		Sig.Inputs.Add(FNiagaraVariable(FNiagaraTypeDefinition::GetFloatDef(), TEXT("MipLevel")));
+		Sig.Inputs.Emplace(FNiagaraTypeDefinition(GetClass()), TEXT("Texture"));
+		Sig.Inputs.Emplace(FNiagaraTypeDefinition::GetVec3Def(), TEXT("UVW"));
+		Sig.Inputs.Emplace(FNiagaraTypeDefinition::GetFloatDef(), TEXT("MipLevel"));
+		Sig.Outputs.Emplace(FNiagaraTypeDefinition::GetVec4Def(), TEXT("Value"));
 		Sig.SetDescription(LOCTEXT("TextureSample2DArrayTextureDesc", "Sample the specified mip level of the input texture at the specified UVW coordinates. Where W is the slice to sample (0,1,2, etc) and UV are the coordinates into the slice."));
-		Sig.Outputs.Add(FNiagaraVariable(FNiagaraTypeDefinition::GetVec4Def(), TEXT("Value")));
 	}
 
 	{
@@ -125,9 +143,9 @@ void UNiagaraDataInterface2DArrayTexture::GetFunctions(TArray<FNiagaraFunctionSi
 		Sig.Name = TextureDimsName;
 		Sig.bMemberFunction = true;
 		Sig.bRequiresContext = false;
-		Sig.Inputs.Add(FNiagaraVariable(FNiagaraTypeDefinition(GetClass()), TEXT("Texture")));
+		Sig.Inputs.Emplace(FNiagaraTypeDefinition(GetClass()), TEXT("Texture"));
+		Sig.Outputs.Emplace(FNiagaraTypeDefinition::GetVec3Def(), TEXT("Dimensions"));
 		Sig.SetDescription(LOCTEXT("TextureDimsDesc", "Get the dimensions of mip 0 of the texture."));
-		Sig.Outputs.Add(FNiagaraVariable(FNiagaraTypeDefinition::GetVec3Def(), TEXT("Dimensions")));
 	}
 }
 
@@ -240,12 +258,13 @@ void UNiagaraDataInterface2DArrayTexture::GetParameterDefinitionHLSL(const FNiag
 
 bool UNiagaraDataInterface2DArrayTexture::GetFunctionHLSL(const FNiagaraDataInterfaceGPUParamInfo& ParamInfo, const FNiagaraDataInterfaceGeneratedFunction& FunctionInfo, int FunctionInstanceIndex, FString& OutHLSL)
 {
-	if ((FunctionInfo.DefinitionName == SampleTextureName) ||
-		(FunctionInfo.DefinitionName == TextureDimsName) )
+	static const TSet<FName> ValidGpuFunctions =
 	{
-		return true;
-	}
-	return false;
+		LoadTextureName,
+		SampleTextureName,
+		TextureDimsName,
+	};
+	return ValidGpuFunctions.Contains(FunctionInfo.DefinitionName);
 }
 #endif
 
diff --git a/Engine/Plugins/FX/Niagara/Source/Niagara/Private/NiagaraDataInterfaceTexture.cpp b/Engine/Plugins/FX/Niagara/Source/Niagara/Private/NiagaraDataInterfaceTexture.cpp
--- a/Engine/Plugins/FX/Niagara/Source/Niagara/Private/NiagaraDataInterfaceTexture.cpp
+++ b/Engine/Plugins/FX/Niagara/Source/Niagara/Private/NiagaraDataInterfaceTexture.cpp
@@ -19,6 +19,7 @@
 #define LOCTEXT_NAMESPACE "UNiagaraDataInterfaceTexture"
 
 const TCHAR* UNiagaraDataInterfaceTexture::TemplateShaderFilePath = TEXT("/Plugin/FX/Niagara/Private/NiagaraDataInterfaceTextureTemplate.ush");
+const FName UNiagaraDataInterfaceTexture::LoadTexture2DName(TEXT("LoadTexture2D"));
 const FName UNiagaraDataInterfaceTexture::SampleTexture2DName(TEXT("SampleTexture2D"));
 const FName UNiagaraDataInterfaceTexture::SamplePseudoVolumeTextureName(TEXT("SamplePseudoVolumeTexture"));
 const FName UNiagaraDataInterfaceTexture::GetTextureDimensionsName(TEXT("GetTextureDimensions"));
@@ -147,28 +148,34 @@ bool UNiagaraDataInterfaceTexture::Equals(const UNiagaraDataInterface* Other) co
 
 void UNiagaraDataInterfaceTexture::GetFunctions(TArray<FNiagaraFunctionSignature>& OutFunctions)
 {
+	FNiagaraFunctionSignature DefaultGpuSig;
+	DefaultGpuSig.bMemberFunction = true;
+	DefaultGpuSig.bRequiresContext = false;
+	DefaultGpuSig.bSupportsCPU = false;
+	DefaultGpuSig.bSupportsGPU = true;
+	DefaultGpuSig.Inputs.Emplace(FNiagaraTypeDefinition(GetClass()), TEXT("Texture"));
+	DefaultGpuSig.SetFunctionVersion(FNDITextureFunctionVersion::LatestVersion);
+
 	{
-		FNiagaraFunctionSignature& Sig = OutFunctions.AddDefaulted_GetRef();
+		FNiagaraFunctionSignature& Sig = OutFunctions.Add_GetRef(DefaultGpuSig);
+		Sig.Name = LoadTexture2DName;
+		Sig.Inputs.Emplace(FNiagaraTypeDefinition::GetIntDef(), TEXT("TexelX"));
+		Sig.Inputs.Emplace(FNiagaraTypeDefinition::GetIntDef(), TEXT("TexelY"));
+		Sig.Inputs.Emplace(FNiagaraTypeDefinition::GetIntDef(), TEXT("MipLevel"));
+		Sig.Outputs.Emplace(FNiagaraTypeDefinition::GetVec4Def(), TEXT("Value"));
+		Sig.SetDescription(LOCTEXT("TextureLoadTexture2DDesc", "Read a texel from the provided location & mip without any filtering or sampling."));
+	}
+	{
+		FNiagaraFunctionSignature& Sig = OutFunctions.Add_GetRef(DefaultGpuSig);
 		Sig.Name = SampleTexture2DName;
-		Sig.bMemberFunction = true;
-		Sig.bRequiresContext = false;		
-		Sig.bSupportsCPU = false;
-		Sig.bSupportsGPU = true;
-		Sig.Inputs.Emplace(FNiagaraTypeDefinition(GetClass()), TEXT("Texture"));
 		Sig.Inputs.Emplace(FNiagaraTypeDefinition::GetVec2Def(), TEXT("UV"));
 		Sig.Inputs.Emplace(FNiagaraTypeDefinition::GetFloatDef(), TEXT("MipLevel"));
 		Sig.Outputs.Emplace(FNiagaraTypeDefinition::GetVec4Def(), TEXT("Value"));
 		Sig.SetDescription(LOCTEXT("TextureSampleTexture2DDesc", "Sample supplied mip level from input 2d texture at the specified UV coordinates. The UV origin (0,0) is in the upper left hand corner of the image."));
-		Sig.SetFunctionVersion(FNDITextureFunctionVersion::LatestVersion);
 	}
 	{
-		FNiagaraFunctionSignature& Sig = OutFunctions.AddDefaulted_GetRef();
+		FNiagaraFunctionSignature& Sig = OutFunctions.Add_GetRef(DefaultGpuSig);
 		Sig.Name = SamplePseudoVolumeTextureName;
-		Sig.bMemberFunction = true;
-		Sig.bRequiresContext = false;
-		Sig.bSupportsCPU = false;
-		Sig.bSupportsGPU = true;
-		Sig.Inputs.Emplace(FNiagaraTypeDefinition(GetClass()), TEXT("Texture"));
 		Sig.Inputs.Emplace(FNiagaraTypeDefinition::GetVec3Def(), TEXT("UVW"));
 		Sig.Inputs.Emplace(FNiagaraTypeDefinition::GetVec2Def(), TEXT("XYNumFrames"));
 		Sig.Inputs.Emplace(FNiagaraTypeDefinition::GetFloatDef(), TEXT("TotalNumFrames"));
@@ -178,8 +185,8 @@ void UNiagaraDataInterfaceTexture::GetFunctions(TArray<FNiagaraFunctionSignature
 		Sig.Inputs.Emplace(FNiagaraTypeDefinition::GetVec2Def(), TEXT("DDY"));
 		Sig.Outputs.Emplace(FNiagaraTypeDefinition::GetVec4Def(), TEXT("Value"));
 		Sig.SetDescription(LOCTEXT("TextureSamplePseudoVolumeTextureDesc", "Return a pseudovolume texture sample.\nUseful for simulating 3D texturing with a 2D texture or as a texture flipbook with lerped transitions.\nTreats 2d layout of frames as a 3d texture and performs bilinear filtering by blending with an offset Z frame.\nTexture = Input Texture Object storing Volume Data\nUVW = Input float3 for Position, 0 - 1\nXYNumFrames = Input float for num frames in x, y directions\nTotalNumFrames = Input float for num total frames\nMipMode = Sampling mode : 0 = use miplevel, 1 = use UV computed gradients, 2 = Use gradients(default = 0)\nMipLevel = MIP level to use in mipmode = 0 (default 0)\nDDX, DDY = Texture gradients in mipmode = 2\n"));
-		Sig.SetFunctionVersion(FNDITextureFunctionVersion::LatestVersion);
 	}
+
 	{
 		FNiagaraFunctionSignature& Sig = OutFunctions.AddDefaulted_GetRef();
 		Sig.Name = GetTextureDimensionsName;
@@ -386,6 +393,7 @@ bool UNiagaraDataInterfaceTexture::GetFunctionHLSL(const FNiagaraDataInterfaceGP
 {
 	static const TSet<FName> ValidGpuFunctions =
 	{
+		LoadTexture2DName,
 		SampleTexture2DName,
 		SamplePseudoVolumeTextureName,
 		GetTextureDimensionsName,
diff --git a/Engine/Plugins/FX/Niagara/Source/Niagara/Private/NiagaraDataInterfaceVolumeTexture.cpp b/Engine/Plugins/FX/Niagara/Source/Niagara/Private/NiagaraDataInterfaceVolumeTexture.cpp
--- a/Engine/Plugins/FX/Niagara/Source/Niagara/Private/NiagaraDataInterfaceVolumeTexture.cpp
+++ b/Engine/Plugins/FX/Niagara/Source/Niagara/Private/NiagaraDataInterfaceVolumeTexture.cpp
@@ -17,6 +17,7 @@
 #define LOCTEXT_NAMESPACE "UNiagaraDataInterfaceVolumeTexture"
 
 const TCHAR* UNiagaraDataInterfaceVolumeTexture::TemplateShaderFilePath = TEXT("/Plugin/FX/Niagara/Private/NiagaraDataInterfaceVolumeTextureTemplate.ush");
+const FName UNiagaraDataInterfaceVolumeTexture::LoadVolumeTextureName(TEXT("LoadVolumeTexture"));
 const FName UNiagaraDataInterfaceVolumeTexture::SampleVolumeTextureName(TEXT("SampleVolumeTexture"));
 const FName UNiagaraDataInterfaceVolumeTexture::TextureDimsName(TEXT("TextureDimensions3D"));
 
@@ -104,33 +105,43 @@ bool UNiagaraDataInterfaceVolumeTexture::Equals(const UNiagaraDataInterface* Oth
 void UNiagaraDataInterfaceVolumeTexture::GetFunctions(TArray<FNiagaraFunctionSignature>& OutFunctions)
 {
 	{
-		FNiagaraFunctionSignature Sig;
+		FNiagaraFunctionSignature & Sig = OutFunctions.AddDefaulted_GetRef();
+		Sig.Name = LoadVolumeTextureName;
+		Sig.bMemberFunction = true;
+		Sig.bRequiresContext = false;
+		Sig.bSupportsCPU = false;
+		Sig.bSupportsGPU = true;
+		Sig.Inputs.Emplace(FNiagaraTypeDefinition(GetClass()), TEXT("Texture"));
+		Sig.Inputs.Emplace(FNiagaraTypeDefinition::GetIntDef(), TEXT("TexelX"));
+		Sig.Inputs.Emplace(FNiagaraTypeDefinition::GetIntDef(), TEXT("TexelY"));
+		Sig.Inputs.Emplace(FNiagaraTypeDefinition::GetIntDef(), TEXT("TexelZ"));
+		Sig.Inputs.Emplace(FNiagaraTypeDefinition::GetIntDef(), TEXT("MipLevel"));
+		Sig.Outputs.Emplace(FNiagaraTypeDefinition::GetVec4Def(), TEXT("Value"));
+		Sig.SetDescription(LOCTEXT("TextureLoadTextureVolumeDesc", "Read a texel from the provided location & mip without any filtering or sampling."));
+	}
+
+	{
+		FNiagaraFunctionSignature& Sig = OutFunctions.AddDefaulted_GetRef();
 		Sig.Name = SampleVolumeTextureName;
 		Sig.bMemberFunction = true;
 		Sig.bRequiresContext = false;
 		Sig.bSupportsCPU = false;
 		Sig.bSupportsGPU = true;
-		Sig.Inputs.Add(FNiagaraVariable(FNiagaraTypeDefinition(GetClass()), TEXT("Texture")));
-		Sig.Inputs.Add(FNiagaraVariable(FNiagaraTypeDefinition::GetVec3Def(), TEXT("UVW")));
-		Sig.Inputs.Add(FNiagaraVariable(FNiagaraTypeDefinition::GetFloatDef(), TEXT("MipLevel")));
+		Sig.Inputs.Emplace(FNiagaraTypeDefinition(GetClass()), TEXT("Texture"));
+		Sig.Inputs.Emplace(FNiagaraTypeDefinition::GetVec3Def(), TEXT("UVW"));
+		Sig.Inputs.Emplace(FNiagaraTypeDefinition::GetFloatDef(), TEXT("MipLevel"));
+		Sig.Outputs.Emplace(FNiagaraTypeDefinition::GetVec4Def(), TEXT("Value"));
 		Sig.SetDescription(LOCTEXT("TextureSampleVolumeTextureDesc", "Sample the specified mip level of the input 3d texture at the specified UVW coordinates. The UVW origin (0, 0, 0) is in the bottom left hand corner of the volume."));
-		Sig.Outputs.Add(FNiagaraVariable(FNiagaraTypeDefinition::GetVec4Def(), TEXT("Value")));
-		//Sig.Owner = *GetName();
-
-		OutFunctions.Add(Sig);
 	}
 
 	{
-		FNiagaraFunctionSignature Sig;
+		FNiagaraFunctionSignature& Sig = OutFunctions.AddDefaulted_GetRef();
 		Sig.Name = TextureDimsName;
 		Sig.bMemberFunction = true;
 		Sig.bRequiresContext = false;
-		Sig.Inputs.Add(FNiagaraVariable(FNiagaraTypeDefinition(GetClass()), TEXT("Texture")));
+		Sig.Inputs.Emplace(FNiagaraTypeDefinition(GetClass()), TEXT("Texture"));
+		Sig.Outputs.Emplace(FNiagaraTypeDefinition::GetVec3Def(), TEXT("Dimensions3D"));
 		Sig.SetDescription(LOCTEXT("TextureDimsDesc", "Get the dimensions of mip 0 of the texture."));
-		Sig.Outputs.Add(FNiagaraVariable(FNiagaraTypeDefinition::GetVec3Def(), TEXT("Dimensions3D")));
-		//Sig.Owner = *GetName();
-
-		OutFunctions.Add(Sig);
 	}
 }
 
@@ -270,12 +281,13 @@ void UNiagaraDataInterfaceVolumeTexture::GetParameterDefinitionHLSL(const FNiaga
 
 bool UNiagaraDataInterfaceVolumeTexture::GetFunctionHLSL(const FNiagaraDataInterfaceGPUParamInfo& ParamInfo, const FNiagaraDataInterfaceGeneratedFunction& FunctionInfo, int FunctionInstanceIndex, FString& OutHLSL)
 {
-	if ((FunctionInfo.DefinitionName == SampleVolumeTextureName) ||
-		(FunctionInfo.DefinitionName == TextureDimsName) )
+	static const TSet<FName> ValidGpuFunctions =
 	{
-		return true;
-	}
-	return false;
+		LoadVolumeTextureName,
+		SampleVolumeTextureName,
+		TextureDimsName,
+	};
+	return ValidGpuFunctions.Contains(FunctionInfo.DefinitionName);
 }
 #endif
 
