Fixing issue with static variables not triggering a recompile due to prior change to notify graph changed.

diff --git a/Engine/Plugins/FX/Niagara/Source/NiagaraEditor/Private/ViewModels/Stack/NiagaraStackFunctionInput.cpp b/Engine/Plugins/FX/Niagara/Source/NiagaraEditor/Private/ViewModels/Stack/NiagaraStackFunctionInput.cpp
--- a/Engine/Plugins/FX/Niagara/Source/NiagaraEditor/Private/ViewModels/Stack/NiagaraStackFunctionInput.cpp
+++ b/Engine/Plugins/FX/Niagara/Source/NiagaraEditor/Private/ViewModels/Stack/NiagaraStackFunctionInput.cpp
@@ -1788,11 +1788,10 @@ void UNiagaraStackFunctionInput::SetLocalValue(TSharedRef<FStructOnScope> InLoca
 			}
 			if (InputType.IsStatic()) // Need to potentially trigger a recompile.
 			{
-				Script->MarkScriptAndSourceDesynchronized(TEXT("Static variable changed!"), FGuid::NewGuid());
 				UNiagaraScriptSource* Source = Cast<UNiagaraScriptSource>(Script->GetLatestSource());
 				if (Source && Source->NodeGraph)
 				{
-					Source->NodeGraph->NotifyGraphChanged();
+					Source->NodeGraph->NotifyGraphNeedsRecompile();
 				}
 			}
 		}
