SetRenderTargetFormat function - allow scripts to set the render target format
will be used in Niagara Fluids to ensure output RTs have the correct number of channels based on sim setup

diff --git a/Engine/Plugins/FX/Niagara/Source/Niagara/Classes/NiagaraDataInterfaceRenderTargetVolume.h b/Engine/Plugins/FX/Niagara/Source/Niagara/Classes/NiagaraDataInterfaceRenderTargetVolume.h
--- a/Engine/Plugins/FX/Niagara/Source/Niagara/Classes/NiagaraDataInterfaceRenderTargetVolume.h
+++ b/Engine/Plugins/FX/Niagara/Source/Niagara/Classes/NiagaraDataInterfaceRenderTargetVolume.h
@@ -130,6 +130,7 @@ public:
 	
 	void VMGetSize(FVectorVMExternalFunctionContext& Context); 
 	void VMSetSize(FVectorVMExternalFunctionContext& Context);
+	void VMSetFormat(FVectorVMExternalFunctionContext& Context);
 
 	UPROPERTY(EditAnywhere, Category = "Render Target", meta = (DisplayPriority = 2, EditCondition = "!bInheritUserParameterSettings", EditConditionHides))
 	FIntVector Size;
diff --git a/Engine/Plugins/FX/Niagara/Source/Niagara/Private/NiagaraDataInterfaceRenderTargetVolume.cpp b/Engine/Plugins/FX/Niagara/Source/Niagara/Private/NiagaraDataInterfaceRenderTargetVolume.cpp
--- a/Engine/Plugins/FX/Niagara/Source/Niagara/Private/NiagaraDataInterfaceRenderTargetVolume.cpp
+++ b/Engine/Plugins/FX/Niagara/Source/Niagara/Private/NiagaraDataInterfaceRenderTargetVolume.cpp
@@ -38,6 +38,7 @@ namespace NDIRenderTargetVolumeLocal
 	static const FName GetValueFunctionName("GetRenderTargetValue");
 	static const FName SampleValueFunctionName("SampleRenderTargetValue");
 	static const FName SetSizeFunctionName("SetRenderTargetSize");
+	static const FName SetFormatFunctionName("SetRenderTargetFormat");
 	static const FName GetSizeFunctionName("GetRenderTargetSize");
 	static const FName LinearToIndexName("LinearToIndex");
 	static const FName ExecToIndexName("ExecToIndex");
@@ -163,6 +164,25 @@ void UNiagaraDataInterfaceRenderTargetVolume::GetFunctions(TArray<FNiagaraFuncti
 	#endif
 	}
 
+	{
+		FNiagaraFunctionSignature& Sig = OutFunctions.AddDefaulted_GetRef();
+		Sig.Name = SetFormatFunctionName;
+		Sig.Inputs.Add(FNiagaraVariable(FNiagaraTypeDefinition(GetClass()), TEXT("RenderTarget")));		
+		Sig.Inputs.Add(FNiagaraVariable(FNiagaraTypeDefinition(StaticEnum<ETextureRenderTargetFormat>()), TEXT("Format")));
+		Sig.Outputs.Add(FNiagaraVariable(FNiagaraTypeDefinition::GetBoolDef(), TEXT("Success")));
+
+		Sig.ModuleUsageBitmask = EmitterSystemOnlyBitmask;
+		Sig.bExperimental = true;
+		Sig.bMemberFunction = true;
+		Sig.bRequiresExecPin = true;
+		Sig.bRequiresContext = false;
+		Sig.bSupportsCPU = true;
+		Sig.bSupportsGPU = false;
+#if WITH_EDITORONLY_DATA
+		Sig.FunctionVersion = EFunctionVersion::LatestVersion;
+#endif
+	}
+
 	{
 		FNiagaraFunctionSignature& Sig = OutFunctions.AddDefaulted_GetRef();
 		Sig.Name = SetValueFunctionName;
@@ -308,6 +328,7 @@ bool UNiagaraDataInterfaceRenderTargetVolume::UpgradeFunctionCall(FNiagaraFuncti
 
 DEFINE_NDI_DIRECT_FUNC_BINDER(UNiagaraDataInterfaceRenderTargetVolume, VMGetSize);
 DEFINE_NDI_DIRECT_FUNC_BINDER(UNiagaraDataInterfaceRenderTargetVolume, VMSetSize);
+DEFINE_NDI_DIRECT_FUNC_BINDER(UNiagaraDataInterfaceRenderTargetVolume, VMSetFormat);
 void UNiagaraDataInterfaceRenderTargetVolume::GetVMExternalFunction(const FVMExternalFunctionBindingInfo& BindingInfo, void* InstanceData, FVMExternalFunction &OutFunc)
 {
 	using namespace NDIRenderTargetVolumeLocal;
@@ -323,6 +344,11 @@ void UNiagaraDataInterfaceRenderTargetVolume::GetVMExternalFunction(const FVMExt
 		check(BindingInfo.GetNumInputs() == 4 && BindingInfo.GetNumOutputs() == 1);
 		NDI_FUNC_BINDER(UNiagaraDataInterfaceRenderTargetVolume, VMSetSize)::Bind(this, OutFunc);
 	}
+	else if (BindingInfo.Name == SetFormatFunctionName)
+	{
+		check(BindingInfo.GetNumInputs() == 2 && BindingInfo.GetNumOutputs() == 1);
+		NDI_FUNC_BINDER(UNiagaraDataInterfaceRenderTargetVolume, VMSetFormat)::Bind(this, OutFunc);
+	}
 }
 
 bool UNiagaraDataInterfaceRenderTargetVolume::Equals(const UNiagaraDataInterface* Other) const
@@ -758,6 +784,37 @@ void UNiagaraDataInterfaceRenderTargetVolume::VMSetSize(FVectorVMExternalFunctio
 			InstData->Size.X = FMath::Clamp<int>(int(float(SizeX) * NiagaraDataInterfaceRenderTargetCommon::GResolutionMultiplier), 1, GMaxVolumeTextureDimensions);
 			InstData->Size.Y = FMath::Clamp<int>(int(float(SizeY) * NiagaraDataInterfaceRenderTargetCommon::GResolutionMultiplier), 1, GMaxVolumeTextureDimensions);
 			InstData->Size.Z = FMath::Clamp<int>(int(float(SizeZ) * NiagaraDataInterfaceRenderTargetCommon::GResolutionMultiplier), 1, GMaxVolumeTextureDimensions);
+
+			if (bInheritUserParameterSettings && InstData->RTUserParamBinding.GetValue<UTextureRenderTargetVolume>())
+			{
+				UE_LOG(LogNiagara, Warning, TEXT("Overriding inherited volume render target size"));
+			}
+		}		
+	}
+}
+
+void UNiagaraDataInterfaceRenderTargetVolume::VMSetFormat(FVectorVMExternalFunctionContext& Context)
+{
+	// This should only be called from a system or emitter script due to a need for only setting up initially.
+	VectorVM::FUserPtrHandler<FRenderTargetVolumeRWInstanceData_GameThread> InstData(Context);
+	FNDIInputParam<ETextureRenderTargetFormat> InFormat(Context);
+	FNDIOutputParam<FNiagaraBool> OutSuccess(Context);
+
+	for (int32 InstanceIdx = 0; InstanceIdx < Context.GetNumInstances(); ++InstanceIdx)
+	{
+		const ETextureRenderTargetFormat Format = InFormat.GetAndAdvance();
+
+
+		const bool bSuccess = (InstData.Get() != nullptr && Context.GetNumInstances() == 1);
+		OutSuccess.SetAndAdvance(bSuccess);
+		if (bSuccess)
+		{
+			InstData->Format = GetPixelFormatFromRenderTargetFormat(Format);
+
+			if (bInheritUserParameterSettings && InstData->RTUserParamBinding.GetValue<UTextureRenderTargetVolume>())
+			{
+				UE_LOG(LogNiagara, Warning, TEXT("Overriding inherited volume render target format"));
+			}
 		}
 	}
 }
