Fix for stalled cook when multiple shader formats are being generated at a time and a portion of them are not supported (because ShouldCompile() returns false).

diff --git a/Engine/Plugins/FX/Niagara/Source/Niagara/Classes/NiagaraScript.h b/Engine/Plugins/FX/Niagara/Source/Niagara/Classes/NiagaraScript.h
--- a/Engine/Plugins/FX/Niagara/Source/Niagara/Classes/NiagaraScript.h
+++ b/Engine/Plugins/FX/Niagara/Source/Niagara/Classes/NiagaraScript.h
@@ -981,7 +981,6 @@ public:
 
 	NIAGARA_API bool CanBeRunOnGpu() const;
 	NIAGARA_API bool IsReadyToRun(ENiagaraSimTarget SimTarget) const;
-	NIAGARA_API bool ShouldCacheShadersForCooking(const ITargetPlatform* TargetPlatform) const;
 #if WITH_EDITORONLY_DATA
 	NIAGARA_API static void SetPreviewFeatureLevel(ERHIFeatureLevel::Type PreviewFeatureLevel);
 #endif
@@ -1036,6 +1035,9 @@ public:
 	virtual bool IsCachedCookedPlatformDataLoaded(const ITargetPlatform* TargetPlatform) override;
 	void CacheShadersForResources(FNiagaraShaderScript* ResourceToCache, bool bApplyCompletedShaderMapForRendering, bool bForceRecompile = false, bool bCooking=false, const ITargetPlatform* TargetPlatform = nullptr);
 	void SaveShaderStableKeys(const class ITargetPlatform* TP);
+
+	TArray<FName> FindShaderFormatsForCooking(const ITargetPlatform* TargetPlatform) const;
+
 #endif // WITH_EDITOR
 	FNiagaraShaderScript* AllocateResource();
 	FNiagaraShaderScript* GetRenderThreadScript()
diff --git a/Engine/Plugins/FX/Niagara/Source/Niagara/Private/NiagaraScript.cpp b/Engine/Plugins/FX/Niagara/Source/Niagara/Private/NiagaraScript.cpp
--- a/Engine/Plugins/FX/Niagara/Source/Niagara/Private/NiagaraScript.cpp
+++ b/Engine/Plugins/FX/Niagara/Source/Niagara/Private/NiagaraScript.cpp
@@ -2,6 +2,7 @@
 
 #include "NiagaraScript.h"
 
+#include "Algo/RemoveIf.h"
 #include "Modules/ModuleManager.h"
 #include "NiagaraScriptSourceBase.h"
 #include "NiagaraComponent.h"
@@ -2261,50 +2262,6 @@ bool UNiagaraScript::IsReadyToRun(ENiagaraSimTarget SimTarget) const
 	return false;
 }
 
-bool UNiagaraScript::ShouldCacheShadersForCooking(const ITargetPlatform* TargetPlatform) const
-{
-	if (CanBeRunOnGpu())
-	{
-		{
-			bool ShaderPlatformsSupportScript = false;
-
-			TArray<FName> DesiredShaderFormats;
-			TargetPlatform->GetAllTargetedShaderFormats(DesiredShaderFormats);
-
-			// Cache for all the shader formats that the cooking target requires
-			for (const FName& DesiredShaderFormat : DesiredShaderFormats)
-			{
-				const EShaderPlatform LegacyShaderPlatform = ShaderFormatToLegacyShaderPlatform(DesiredShaderFormat);
-				if (ShouldCompile(LegacyShaderPlatform))
-				{
-					ShaderPlatformsSupportScript = true;
-					break;
-				}
-			}
-
-			if (!ShaderPlatformsSupportScript)
-			{
-				return false;
-			}
-		}
-
-		// currently only scripts that are associated with a system are appropriate for caching.  Standalone scripts or scripts
-		// associated with standalone emitters are not required on cooked platforms.
-		if (GetTypedOuter<UNiagaraSystem>() == nullptr)
-		{
-			return false;
-		}
-
-		FVersionedNiagaraEmitter OuterEmitter = GetOuterEmitter();
-		FVersionedNiagaraEmitterData* EmitterData = OuterEmitter.GetEmitterData();
-		if (EmitterData != nullptr && EmitterData->SimTarget == ENiagaraSimTarget::GPUComputeSim && OuterEmitter.Emitter->NeedsLoadForTargetPlatform(TargetPlatform))
-		{
-			return true;
-		}
-	}
-	return false;
-}
-
 #if WITH_EDITORONLY_DATA
 void UNiagaraScript::SetPreviewFeatureLevel(ERHIFeatureLevel::Type InPreviewFeatureLevel)
 {
@@ -3229,9 +3186,43 @@ void UNiagaraScript::ModifyCompilationEnvironment(EShaderPlatform Platform, stru
 
 #if WITH_EDITOR
 
+TArray<FName> UNiagaraScript::FindShaderFormatsForCooking(const ITargetPlatform* TargetPlatform) const
+{
+	TArray<FName> DesiredShaderFormats;
+
+	if (CanBeRunOnGpu())
+	{
+		// currently only scripts that are associated with a system are appropriate for caching.  Standalone scripts or scripts
+		// associated with standalone emitters are not required on cooked platforms.
+		const bool bStandaloneScript = GetTypedOuter<UNiagaraSystem>() == nullptr;
+
+		FVersionedNiagaraEmitter OuterEmitter = GetOuterEmitter();
+		FVersionedNiagaraEmitterData* EmitterData = OuterEmitter.GetEmitterData();
+		const bool bHasValidEmitter = (EmitterData != nullptr
+			&& EmitterData->SimTarget == ENiagaraSimTarget::GPUComputeSim
+			&& OuterEmitter.Emitter->NeedsLoadForTargetPlatform(TargetPlatform));
+
+		if (!bStandaloneScript && bHasValidEmitter)
+		{
+			TargetPlatform->GetAllTargetedShaderFormats(DesiredShaderFormats);
+
+			// remove any formats that aren't supported
+			DesiredShaderFormats.SetNum(Algo::RemoveIf(DesiredShaderFormats, [this](const FName& ShaderFormat)
+			{
+				const EShaderPlatform LegacyShaderPlatform = ShaderFormatToLegacyShaderPlatform(ShaderFormat);
+				return !ShouldCompile(LegacyShaderPlatform);
+			}));
+		}
+	}
+
+	return DesiredShaderFormats;
+}
+
 void UNiagaraScript::BeginCacheForCookedPlatformData(const ITargetPlatform *TargetPlatform)
 {
-	if (ShouldCacheShadersForCooking(TargetPlatform))
+	TArray<FName> ShaderFormatsForCooking = FindShaderFormatsForCooking(TargetPlatform);
+
+	if (!ShaderFormatsForCooking.IsEmpty())
 	{
 		// Commandlets like DerivedDataCacheCommandlet call BeginCacheForCookedPlatformData directly on objects. This may mean that
 		// we have not properly gotten the HLSL script generated by the time that we get here. This does the awkward work of
@@ -3250,23 +3241,27 @@ void UNiagaraScript::BeginCacheForCookedPlatformData(const ITargetPlatform *Targ
 			return;
 		}
 
-		TArray<FName> DesiredShaderFormats;
-		TargetPlatform->GetAllTargetedShaderFormats(DesiredShaderFormats);
-
 		TArray<FNiagaraShaderScript*>& CachedScriptResourcesForPlatform = CachedScriptResourcesForCooking.FindOrAdd(TargetPlatform);
 
 		// Cache for all the shader formats that the cooking target requires
-		for (int32 FormatIndex = 0; FormatIndex < DesiredShaderFormats.Num(); FormatIndex++)
+		for (const FName& ShaderFormat : ShaderFormatsForCooking)
 		{
-			const EShaderPlatform LegacyShaderPlatform = ShaderFormatToLegacyShaderPlatform(DesiredShaderFormats[FormatIndex]);
-				CacheResourceShadersForCooking(LegacyShaderPlatform, CachedScriptResourcesForPlatform, TargetPlatform);
-			}
+			const EShaderPlatform LegacyShaderPlatform = ShaderFormatToLegacyShaderPlatform(ShaderFormat);
+			CacheResourceShadersForCooking(LegacyShaderPlatform, CachedScriptResourcesForPlatform, TargetPlatform);
 		}
 	}
+}
 
 bool UNiagaraScript::IsCachedCookedPlatformDataLoaded(const ITargetPlatform* TargetPlatform)
 {
-	if (ShouldCacheShadersForCooking(TargetPlatform) && HasIdsRequiredForShaderCaching())
+	if (!HasIdsRequiredForShaderCaching())
+	{
+		return true;
+	}
+
+	const TArray<FName> ShaderFormats = FindShaderFormatsForCooking(TargetPlatform);
+
+	if (!ShaderFormats.IsEmpty())
 	{
 		bool bHasOutstandingCompilationRequests = false;
 		if (UNiagaraSystem* SystemOwner = FindRootSystem())
@@ -3276,9 +3271,6 @@ bool UNiagaraScript::IsCachedCookedPlatformDataLoaded(const ITargetPlatform* Tar
 
 		if (!bHasOutstandingCompilationRequests)
 		{
-			TArray<FName> DesiredShaderFormats;
-			TargetPlatform->GetAllTargetedShaderFormats(DesiredShaderFormats);
-
 			const TArray<FNiagaraShaderScript*>* CachedScriptResourcesForPlatform = CachedScriptResourcesForCooking.Find(TargetPlatform);
 			if (CachedScriptResourcesForPlatform)
 			{
