Niagara - Clear out emitter thunbmails on non-assets when cooking.

diff --git a/Engine/Plugins/FX/Niagara/Source/Niagara/Private/NiagaraEmitter.cpp b/Engine/Plugins/FX/Niagara/Source/Niagara/Private/NiagaraEmitter.cpp
--- a/Engine/Plugins/FX/Niagara/Source/Niagara/Private/NiagaraEmitter.cpp
+++ b/Engine/Plugins/FX/Niagara/Source/Niagara/Private/NiagaraEmitter.cpp
@@ -520,6 +520,16 @@ void UNiagaraEmitter::Serialize(FArchive& Ar)
 			}
 		}
 	}
+
+	// When cooking an emitter that's not an asset, clear out the thumbnail image to prevent issues
+	// with cooked editor data.
+	bool bCookingNonAssetEmitter = Ar.IsCooking() && this->IsAsset() == false;
+	UTexture2D* CachedThumbnail = nullptr;
+	if (bCookingNonAssetEmitter)
+	{
+		CachedThumbnail = ThumbnailImage;
+		ThumbnailImage = nullptr;
+	}
 	
 #endif
 	Super::Serialize(Ar);
@@ -529,6 +539,12 @@ void UNiagaraEmitter::Serialize(FArchive& Ar)
 	{
 		IsCooked = Ar.IsFilterEditorOnly();
 	}
+
+	// Restore the thumbnail image that was cleared before serialize.
+	if (bCookingNonAssetEmitter)
+	{
+		ThumbnailImage = CachedThumbnail;
+	}
 #endif
 
 	Ar.UsingCustomVersion(FNiagaraCustomVersion::GUID);
diff --git a/Engine/Plugins/FX/Niagara/Source/Niagara/Private/NiagaraEmitterHandle.cpp b/Engine/Plugins/FX/Niagara/Source/Niagara/Private/NiagaraEmitterHandle.cpp
--- a/Engine/Plugins/FX/Niagara/Source/Niagara/Private/NiagaraEmitterHandle.cpp
+++ b/Engine/Plugins/FX/Niagara/Source/Niagara/Private/NiagaraEmitterHandle.cpp
@@ -227,9 +227,6 @@ void FNiagaraEmitterHandle::ConditionalPostLoad(int32 NiagaraCustomVersion)
 				SetName(Name, *OwningSystem);
 			}
 		}
-
-		// Temporary Workaround - These thumbnails aren't used by the system so remove them here to fix issues with cooked editor data.
-		InstanceEmitter->ThumbnailImage = nullptr;
 	}
 }
 
