- Change DebugHUD to display emitter name rather than instance name

diff --git a/Engine/Plugins/FX/Niagara/Source/Niagara/Private/NiagaraDebugHud.cpp b/Engine/Plugins/FX/Niagara/Source/Niagara/Private/NiagaraDebugHud.cpp
--- a/Engine/Plugins/FX/Niagara/Source/Niagara/Private/NiagaraDebugHud.cpp
+++ b/Engine/Plugins/FX/Niagara/Source/Niagara/Private/NiagaraDebugHud.cpp
@@ -2359,7 +2359,8 @@ void FNiagaraDebugHud::DrawValidation(class FNiagaraWorldManager* WorldManager,
 					[&](const FNiagaraVariable& Variable, int32 InstanceIndex, int32 ComponentIndex)
 					{
 						auto& ValidationError = GetValidationErrorInfo(NiagaraComponent);
-						ValidationError.ParticleVariablesWithErrors.FindOrAdd(EmitterInstance->GetCachedEmitter().Emitter->GetFName()).AddUnique(Variable.GetName());
+						const FName EmitterName(*EmitterInstance->GetCachedEmitter().Emitter->GetUniqueEmitterName());
+						ValidationError.ParticleVariablesWithErrors.FindOrAdd(EmitterName).AddUnique(Variable.GetName());
 					}
 				);
 			}
