Niagara - Add allowed class usage based asset filtering to the System/Emitter creation dialog and the script dialogs in the script and system editors.  This will be superceded by an asset allow list in the future, but in the short term this will prevent the user from adding assets to their scirpt or system which will immediately produce errors for using disallowed classes.

diff --git a/Engine/Plugins/FX/Niagara/Source/Niagara/Private/NiagaraEmitter.cpp b/Engine/Plugins/FX/Niagara/Source/Niagara/Private/NiagaraEmitter.cpp
--- a/Engine/Plugins/FX/Niagara/Source/Niagara/Private/NiagaraEmitter.cpp
+++ b/Engine/Plugins/FX/Niagara/Source/Niagara/Private/NiagaraEmitter.cpp
@@ -1032,6 +1032,8 @@ void UNiagaraEmitter::GetAssetRegistryTags(TArray<FAssetRegistryTag>& OutTags) c
 	FText TemplateSpecializationValueString = StaticEnum<ENiagaraScriptTemplateSpecification>()->GetDisplayNameTextByValue((int64) TemplateSpecification);
 	OutTags.Add(FAssetRegistryTag(TemplateSpecificationName, TemplateSpecializationValueString.ToString(), FAssetRegistryTag::TT_Alphabetical));
 	
+	INiagaraModule& NiagaraModule = FModuleManager::GetModuleChecked<INiagaraModule>("Niagara");
+	OutTags.Add(NiagaraModule.GetEditorOnlyDataUtilities().CreateClassUsageAssetRegistryTag(this));
 #endif
 	Super::GetAssetRegistryTags(OutTags);
 }
diff --git a/Engine/Plugins/FX/Niagara/Source/Niagara/Private/NiagaraScript.cpp b/Engine/Plugins/FX/Niagara/Source/Niagara/Private/NiagaraScript.cpp
--- a/Engine/Plugins/FX/Niagara/Source/Niagara/Private/NiagaraScript.cpp
+++ b/Engine/Plugins/FX/Niagara/Source/Niagara/Private/NiagaraScript.cpp
@@ -34,6 +34,7 @@
 	#include "NiagaraSettings.h"
 	#include "Internationalization/Regex.h"
 	#include "ShaderCodeLibrary.h"
+	#include "INiagaraEditorOnlyDataUtlities.h"
 
 	// This is a version string that mimics the old versioning scheme. In case of merge conflicts with DDC versions,
 	// you MUST generate a new GUID and set this new version. If you want to bump this version, generate a new guid
@@ -3077,6 +3078,9 @@ void UNiagaraScript::GetAssetRegistryTags(TArray<FAssetRegistryTag>& OutTags) co
 	// the whole asset.
 	const int32 NiagaraVer = GetLinkerCustomVersion(FNiagaraCustomVersion::GUID);
 	OutTags.Add(FAssetRegistryTag(NiagaraCustomVersionTagName, FString::FromInt(NiagaraVer), FAssetRegistryTag::TT_Hidden));
+
+	INiagaraModule& NiagaraModule = FModuleManager::GetModuleChecked<INiagaraModule>("Niagara");
+	OutTags.Add(NiagaraModule.GetEditorOnlyDataUtilities().CreateClassUsageAssetRegistryTag(this));
 #endif
 }
 
diff --git a/Engine/Plugins/FX/Niagara/Source/Niagara/Private/NiagaraSystem.cpp b/Engine/Plugins/FX/Niagara/Source/Niagara/Private/NiagaraSystem.cpp
--- a/Engine/Plugins/FX/Niagara/Source/Niagara/Private/NiagaraSystem.cpp
+++ b/Engine/Plugins/FX/Niagara/Source/Niagara/Private/NiagaraSystem.cpp
@@ -1407,6 +1407,8 @@ void UNiagaraSystem::GetAssetRegistryTags(TArray<FAssetRegistryTag>& OutTags) co
 	//OutTags.Add(FAssetRegistryTag("Becomes Zombie", WillBecomeZombie() ? TEXT("True") : TEXT("False"), FAssetRegistryTag::TT_Alphabetical));
 	//OutTags.Add(FAssetRegistryTag("CanBeOccluded", OcclusionBoundsMethod == EParticleSystemOcclusionBoundsMethod::EPSOBM_None ? TEXT("False") : TEXT("True"), FAssetRegistryTag::TT_Alphabetical));
 
+	INiagaraModule& NiagaraModule = FModuleManager::GetModuleChecked<INiagaraModule>("Niagara");
+	OutTags.Add(NiagaraModule.GetEditorOnlyDataUtilities().CreateClassUsageAssetRegistryTag(this));
 #endif
 	Super::GetAssetRegistryTags(OutTags);
 }
diff --git a/Engine/Plugins/FX/Niagara/Source/Niagara/Public/INiagaraEditorOnlyDataUtlities.h b/Engine/Plugins/FX/Niagara/Source/Niagara/Public/INiagaraEditorOnlyDataUtlities.h
--- a/Engine/Plugins/FX/Niagara/Source/Niagara/Public/INiagaraEditorOnlyDataUtlities.h
+++ b/Engine/Plugins/FX/Niagara/Source/Niagara/Public/INiagaraEditorOnlyDataUtlities.h
@@ -6,6 +6,7 @@ class UNiagaraScriptSourceBase;
 class UNiagaraEditorParametersAdapterBase;
 class UNiagaraEditorDataBase;
 class UObject;
+struct FAssetRegistryTag;
 
 /** Defines utility methods for creating editor only data which is stored on runtime objects. */
 class INiagaraEditorOnlyDataUtilities
@@ -16,4 +17,6 @@ public:
 	virtual UNiagaraEditorDataBase* CreateDefaultEditorData(UObject* InOuter) const = 0;
 
 	virtual UNiagaraEditorParametersAdapterBase* CreateDefaultEditorParameters(UObject* InOuter) const = 0;
+
+	virtual UObject::FAssetRegistryTag CreateClassUsageAssetRegistryTag(const UObject* SourceObject) const = 0;
 };
\ No newline at end of file
diff --git a/Engine/Plugins/FX/Niagara/Source/NiagaraEditor/Private/NiagaraEditorModule.cpp b/Engine/Plugins/FX/Niagara/Source/NiagaraEditor/Private/NiagaraEditorModule.cpp
--- a/Engine/Plugins/FX/Niagara/Source/NiagaraEditor/Private/NiagaraEditorModule.cpp
+++ b/Engine/Plugins/FX/Niagara/Source/NiagaraEditor/Private/NiagaraEditorModule.cpp
@@ -216,6 +216,12 @@ class FNiagaraEditorOnlyDataUtilities : public INiagaraEditorOnlyDataUtilities
 	{
 		return NewObject<UNiagaraEditorParametersAdapter>(InOuter);
 	}
+
+	virtual UObject::FAssetRegistryTag CreateClassUsageAssetRegistryTag(const UObject* SourceObject) const override
+	{
+		const UNiagaraEditorSettings* NiagaraEditorSettings = GetDefault<UNiagaraEditorSettings>();
+		return NiagaraEditorSettings->CreateClassUsageAssetRegistryTag(SourceObject);
+	}
 };
 
 class FNiagaraScriptGraphPanelPinFactory : public FGraphPanelPinFactory
diff --git a/Engine/Plugins/FX/Niagara/Source/NiagaraEditor/Private/NiagaraEditorSettings.cpp b/Engine/Plugins/FX/Niagara/Source/NiagaraEditor/Private/NiagaraEditorSettings.cpp
--- a/Engine/Plugins/FX/Niagara/Source/NiagaraEditor/Private/NiagaraEditorSettings.cpp
+++ b/Engine/Plugins/FX/Niagara/Source/NiagaraEditor/Private/NiagaraEditorSettings.cpp
@@ -5,6 +5,7 @@
 #include "NiagaraActions.h"
 #include "NiagaraConstants.h"
 #include "NiagaraEditorModule.h"
+#include "NiagaraNode.h"
 
 #include UE_INLINE_GENERATED_CPP_BY_NAME(NiagaraEditorSettings)
 
@@ -23,6 +24,13 @@ const FGuid FNiagaraEditorGuids::ParameterCollectionNamespaceMetaDataGuid = FGui
 const FGuid FNiagaraEditorGuids::DataInstanceNamespaceMetaDataGuid = FGuid(TEXT("17F5FC610A914BA58D7EA8FD39B9A1D0"));
 const FGuid FNiagaraEditorGuids::StaticSwitchNamespaceMetaDataGuid = FGuid(TEXT("6A44CDD2EC3D4495BDC7FE28CAE00604"));
 
+int32 GbLogFoundButNotAllowedAssets = 0;
+static FAutoConsoleVariableRef CVarLogFoundButNotAllowedAssets(
+	TEXT("fx.Niagara.LogFoundButNotAllowedAssets"),
+	GbShowNiagaraDeveloperWindows,
+	TEXT("If > 0 assets which are found, but not allowed in the current editor context will be printed to the log."),
+	ECVF_Default
+);
 
 FNiagaraNamespaceMetadata::FNiagaraNamespaceMetadata()
 	: BackgroundColor(FLinearColor::Black)
@@ -57,6 +65,7 @@ UNiagaraEditorSettings::UNiagaraEditorSettings(const FObjectInitializer& ObjectI
 	bResimulateOnChangeWhilePaused = true;
 	bResetDependentSystemsWhenEditingEmitters = false;
 	SetupNamespaceMetadata();
+	TrackedUsageBaseClasses = { UNiagaraNode::StaticClass(), UNiagaraDataInterface::StaticClass(), UNiagaraRendererProperties::StaticClass() };
 }
 
 #define LOCTEXT_NAMESPACE "NamespaceMetadata"
@@ -671,5 +680,120 @@ bool UNiagaraEditorSettings::IsAllowedTypeDefinition(const FNiagaraTypeDefinitio
 	return InTypeDefinition.GetClass() == nullptr || IsAllowedClass(InTypeDefinition.GetClass());
 }
 
+const FName ClassUsageListTagName = "ClassUsageList";
+
+bool UNiagaraEditorSettings::ShouldTrackClassUsage(const UClass* InClass) const
+{
+	for (const UClass* TrackedUsageBaseClass : TrackedUsageBaseClasses)
+	{
+		if (InClass->IsChildOf(TrackedUsageBaseClass))
+		{
+			return true;
+		}
+	}
+	return false;
+}
+
+UObject::FAssetRegistryTag UNiagaraEditorSettings::CreateClassUsageAssetRegistryTag(const UObject* SourceObject) const
+{
+	FString ClassUsageList;
+	if (SourceObject->IsAsset())
+	{
+		TSet<FString> ClassPaths;
+		TArray<UObject*> ObjectsInPackage;
+		GetObjectsWithPackage(SourceObject->GetOutermost(), ObjectsInPackage);
+		for (UObject* ObjectInPackage : ObjectsInPackage)
+		{
+			if (ShouldTrackClassUsage(ObjectInPackage->GetClass()))
+			{
+				ClassPaths.Add(ObjectInPackage->GetClass()->GetClassPathName().ToString());
+			}
+		}
+		ClassUsageList = FString::Join(ClassPaths, TEXT("\n"));
+	}
+	else
+	{
+		ClassUsageList = "";
+	}
+
+	return UObject::FAssetRegistryTag(ClassUsageListTagName, ClassUsageList, FAssetRegistryTag::TT_Hidden);
+}
+
+bool UNiagaraEditorSettings::IsAllowedAssetByClassUsage(const FAssetData& InAssetData) const
+{
+	if (OnIsClassAllowedDelegate.IsBound() == false)
+	{
+		return true;
+	}
+
+	const UObject* AssetObject = nullptr;
+	bool bClassDataFound = false;
+	bool bInvalidClassFound = false;
+	if (InAssetData.IsAssetLoaded())
+	{
+		AssetObject = InAssetData.GetAsset();
+		if (AssetObject != nullptr)
+		{
+			bClassDataFound = true;
+			TArray<UObject*> ObjectsInPackage;
+			GetObjectsWithPackage(AssetObject->GetOutermost(), ObjectsInPackage);
+			for (UObject* ObjectInPackage : ObjectsInPackage)
+			{
+				if (ShouldTrackClassUsage(ObjectInPackage->GetClass()) && IsAllowedClass(ObjectInPackage->GetClass()) == false)
+				{
+					bInvalidClassFound = true;
+					if (GbLogFoundButNotAllowedAssets)
+					{
+						UE_LOG(LogNiagaraEditor, Log, TEXT("Asset %s is not allowed due to object %s with class %s which is not allowed in this editor context."),
+							*InAssetData.GetFullName(), *ObjectInPackage->GetPathName(), *ObjectInPackage->GetClass()->GetClassPathName().ToString());
+					}
+					else
+					{
+						break;
+					}
+				}
+			}
+		}
+	}
+	else
+	{
+		FString ClassUsageList;
+		if (InAssetData.GetTagValue(ClassUsageListTagName, ClassUsageList) && ClassUsageList.Len() != 0)
+		{
+			bClassDataFound = true;
+			TArray<FString> ClassUsageListLines;
+			ClassUsageList.ParseIntoArrayLines(ClassUsageListLines);
+			for (const FString& ClassUsageListLine : ClassUsageListLines)
+			{
+				FTopLevelAssetPath ClassPath(ClassUsageListLine);
+				UClass* UsedClass = FindObject<UClass>(ClassPath);
+				if (UsedClass == nullptr || (ShouldTrackClassUsage(UsedClass) && IsAllowedClass(UsedClass) == false))
+				{
+					bInvalidClassFound = true;
+					if (GbLogFoundButNotAllowedAssets)
+					{
+						if (UsedClass == nullptr)
+						{
+							UE_LOG(LogNiagaraEditor, Log, TEXT("Asset %s is not allowed due to class %s which is was not found."),
+								*InAssetData.GetFullName(), *ClassUsageListLine);
+						}
+						else
+						{
+							UE_LOG(LogNiagaraEditor, Log, TEXT("Asset %s is not allowed due to class %s which is not allowed in this editor context."),
+								*InAssetData.GetFullName(), *UsedClass->GetPathName());
+						}
+					}
+					else
+					{
+						break;
+					}
+				}
+			}
+		}
+	}
+
+	return bClassDataFound && bInvalidClassFound == false;
+}
+
 #undef LOCTEXT_NAMESPACE
 
diff --git a/Engine/Plugins/FX/Niagara/Source/NiagaraEditor/Private/NiagaraEditorUtilities.cpp b/Engine/Plugins/FX/Niagara/Source/NiagaraEditor/Private/NiagaraEditorUtilities.cpp
--- a/Engine/Plugins/FX/Niagara/Source/NiagaraEditor/Private/NiagaraEditorUtilities.cpp
+++ b/Engine/Plugins/FX/Niagara/Source/NiagaraEditor/Private/NiagaraEditorUtilities.cpp
@@ -1403,6 +1403,7 @@ void FNiagaraEditorUtilities::GetFilteredScriptAssets(FGetFilteredScriptAssetsOp
 	TArray<FAssetData> FilteredScriptAssets;
 	AssetRegistryModule.Get().GetAssets(ScriptFilter, FilteredScriptAssets);
 
+	const UNiagaraEditorSettings* NiagaraEditorSettings = GetDefault<UNiagaraEditorSettings>();
 	for (int i = 0; i < FilteredScriptAssets.Num(); ++i)
 	{
 		// Get the custom version the asset was saved with so it can be used below.
@@ -1481,6 +1482,12 @@ void FNiagaraEditorUtilities::GetFilteredScriptAssets(FGetFilteredScriptAssetsOp
 				continue;
 			}
 		}
+
+		// Filter by allowed class usage.
+		if (NiagaraEditorSettings->IsAllowedAssetByClassUsage(FilteredScriptAssets[i]) == false)
+		{
+			continue;
+		}
 		
 		OutFilteredScriptAssets.Add(FilteredScriptAssets[i]);
 	}
diff --git a/Engine/Plugins/FX/Niagara/Source/NiagaraEditor/Private/SNiagaraAssetPickerList.cpp b/Engine/Plugins/FX/Niagara/Source/NiagaraEditor/Private/SNiagaraAssetPickerList.cpp
--- a/Engine/Plugins/FX/Niagara/Source/NiagaraEditor/Private/SNiagaraAssetPickerList.cpp
+++ b/Engine/Plugins/FX/Niagara/Source/NiagaraEditor/Private/SNiagaraAssetPickerList.cpp
@@ -114,22 +114,24 @@ TArray<FAssetData> SNiagaraAssetPickerList::GetAssetDataForSelector(UClass* Asse
 	TArray<FAssetData> EmitterAssets;
 	AssetRegistryModule.Get().GetAssetsByClass(AssetClass->GetClassPathName(), EmitterAssets);
 
+	const UNiagaraEditorSettings* NiagaraEditorSettings = GetDefault<UNiagaraEditorSettings>();
 	TArray<FAssetData> EmittersToShow;
-	if (TabOptions.GetOnlyShowTemplates())
+	for(FAssetData& EmitterAsset : EmitterAssets)
 	{
-		for(FAssetData& EmitterAsset :EmitterAssets)
-		{ 
+		bool bShowEmitter = true;
+		if (TabOptions.GetOnlyShowTemplates())
+		{
 			ENiagaraScriptTemplateSpecification TemplateSpecification;
 			bool bFoundTemplateScriptTag = FNiagaraEditorUtilities::GetTemplateSpecificationFromTag(EmitterAsset, TemplateSpecification);
-			if (bFoundTemplateScriptTag && TemplateSpecification == ENiagaraScriptTemplateSpecification::Template)
+			if (bFoundTemplateScriptTag && TemplateSpecification != ENiagaraScriptTemplateSpecification::Template)
 			{
-				EmittersToShow.Add(EmitterAsset);
+				bShowEmitter = false;
 			}
 		}
-	}
-	else
-	{
-		EmittersToShow = EmitterAssets;
+		if (bShowEmitter && NiagaraEditorSettings->IsAllowedAssetByClassUsage(EmitterAsset))
+		{
+			EmittersToShow.Add(EmitterAsset);
+		}
 	}
 
 	return EmittersToShow;
diff --git a/Engine/Plugins/FX/Niagara/Source/NiagaraEditor/Public/NiagaraEditorSettings.h b/Engine/Plugins/FX/Niagara/Source/NiagaraEditor/Public/NiagaraEditorSettings.h
--- a/Engine/Plugins/FX/Niagara/Source/NiagaraEditor/Public/NiagaraEditorSettings.h
+++ b/Engine/Plugins/FX/Niagara/Source/NiagaraEditor/Public/NiagaraEditorSettings.h
@@ -12,6 +12,7 @@
 
 class UCurveFloat;
 enum class EScriptSource : uint8;
+class UClass;
 
 namespace FNiagaraEditorGuids
 {
@@ -383,9 +384,14 @@ public:
 	/** Returns whether or not the supplied niagara type definition can be used in the current editor context. */
 	bool IsAllowedTypeDefinition(const FNiagaraTypeDefinition& InTypeDefinition) const;
 
+	FAssetRegistryTag CreateClassUsageAssetRegistryTag(const UObject* SourceObject) const;
+
+	bool IsAllowedAssetByClassUsage(const FAssetData& InAssetData) const;
+
 private:
 	void SetupNamespaceMetadata();
 	void BuildCachedPlaybackSpeeds() const;
+	bool ShouldTrackClassUsage(const UClass* InClass) const;
 protected:
 	FOnNiagaraEditorSettingsChanged SettingsChangedDelegate;
 private:
@@ -474,6 +480,8 @@ private:
 	FOnIsClassAllowed OnIsClassAllowedDelegate;
 	FOnIsClassPathAllowed OnIsClassPathAllowedDelegate;
 
+	TArray<UClass*> TrackedUsageBaseClasses;
+
 public:
 	bool IsShowInstructionsCount() const;
 	void SetShowInstructionsCount(bool bShowInstructionsCount);
