- Change select operation to always use the first input as a default
- This fixes an issue where with int / enum inputs we sometimes didn't select a value if the selector was out of range.  This leads to undefined behavior and was producing NaNs in some instances.

diff --git a/Engine/Plugins/FX/Niagara/Source/NiagaraEditor/Private/NiagaraHlslTranslator.cpp b/Engine/Plugins/FX/Niagara/Source/NiagaraEditor/Private/NiagaraHlslTranslator.cpp
--- a/Engine/Plugins/FX/Niagara/Source/NiagaraEditor/Private/NiagaraHlslTranslator.cpp
+++ b/Engine/Plugins/FX/Niagara/Source/NiagaraEditor/Private/NiagaraHlslTranslator.cpp
@@ -9003,18 +9003,22 @@ void FHlslNiagaraTranslator::Select(UNiagaraNodeSelect* SelectNode, int32 Select
 	for (int32 SelectorValueIndex = 0; SelectorValueIndex < SelectorValues.Num(); SelectorValueIndex++)
 	{
 		FString Definition;
-		if (bIsBoolSelector)
+		// We default to the first value from the select operation to ensure we always set to a valid value
+		// Failure to do this can result in bad / incorrect values being used in the VVM
+		if (SelectorValueIndex > 0)
 		{
-			Definition = SelectorValues[SelectorValueIndex] == 0 ? TEXT("if({0} == 0)\n\t{ ") : TEXT("if({0} != 0)\n\t{ ");
-		}
-		else
-		{
-			Definition = FString::Printf(TEXT("if({0} == %d)\n\t{ "), SelectorValues[SelectorValueIndex]);
+			if (bIsBoolSelector)
+			{
+				Definition = SelectorValues[SelectorValueIndex] == 0 ? TEXT("if({0} == 0)\n\t{ ") : TEXT("if({0} != 0)\n\t{ ");
+			}
+			else
+			{
+				Definition = FString::Printf(TEXT("if({0} == %d)\n\t{ "), SelectorValues[SelectorValueIndex]);
+			}
+			TArray<int32> SourceChunks = { Selector };
+			AddBodyChunk(TEXT(""), Definition, FNiagaraTypeDefinition::GetFloatDef(), SourceChunks, false, false);
 		}
 
-		TArray<int32> SourceChunks = { Selector };
-		
-		AddBodyChunk(TEXT(""), Definition, FNiagaraTypeDefinition::GetFloatDef(), SourceChunks, false, false);
 		int32 NaturalIndex = 0;
 		for(int32 CompiledPinCodeChunk : Options[SelectorValues[SelectorValueIndex]])
 		{
@@ -9023,7 +9027,10 @@ void FHlslNiagaraTranslator::Select(UNiagaraNodeSelect* SelectNode, int32 Select
 			NaturalIndex++;
 		}
 		
-		AddBodyChunk(TEXT(""), TEXT("}"), FNiagaraTypeDefinition::GetFloatDef(), false, false);
+		if (SelectorValueIndex > 0)
+		{
+			AddBodyChunk(TEXT(""), TEXT("}"), FNiagaraTypeDefinition::GetFloatDef(), false, false);
+		}
 	}
 
 	// Add an additional invalid output for the add pin which doesn't get compiled.
diff --git a/Engine/Source/Runtime/Core/Private/UObject/FortniteMainBranchObjectVersion.cpp b/Engine/Source/Runtime/Core/Private/UObject/FortniteMainBranchObjectVersion.cpp
--- a/Engine/Source/Runtime/Core/Private/UObject/FortniteMainBranchObjectVersion.cpp
+++ b/Engine/Source/Runtime/Core/Private/UObject/FortniteMainBranchObjectVersion.cpp
@@ -11,7 +11,7 @@ TMap<FGuid, FGuid> FFortniteMainBranchObjectVersion::GetSystemGuids()
 	SystemGuids.Add(DevGuids.MATERIALSHADERMAP_DERIVEDDATA_VER, FGuid("CE7776E63E9E43FBB8C103EE1AB40B5D"));
 	SystemGuids.Add(DevGuids.NANITE_DERIVEDDATA_VER, FGuid("B62CDCAF66AA45289DDB82E276E7C3E7"));
 	SystemGuids.Add(DevGuids.NIAGARASHADERMAP_DERIVEDDATA_VER, FGuid("35D7960B72164B49B9B7B3FB5857E674"));
-	SystemGuids.Add(DevGuids.Niagara_LatestScriptCompileVersion, FGuid("F74D5EA3F19B124A9A02C5EF97CDAFD1"));
+	SystemGuids.Add(DevGuids.Niagara_LatestScriptCompileVersion, FGuid("08ED74B3CB736F4AAD38953889079C6D"));
 	SystemGuids.Add(DevGuids.SkeletalMeshDerivedDataVersion, FGuid("ACF593EAAE354FCBB34CF44F5AA1BFD2"));
 	SystemGuids.Add(DevGuids.STATICMESH_DERIVEDDATA_VER, FGuid("543843B103794E0E9BA4BE80FB602F79"));
 	return SystemGuids;
