- Sanitize tangent basis for static mesh DI as some meshes have been found to contain X==Z which can result in NaNs

diff --git a/Engine/Plugins/FX/Niagara/Shaders/Private/NiagaraDataInterfaceStaticMeshTemplate.ush b/Engine/Plugins/FX/Niagara/Shaders/Private/NiagaraDataInterfaceStaticMeshTemplate.ush
--- a/Engine/Plugins/FX/Niagara/Shaders/Private/NiagaraDataInterfaceStaticMeshTemplate.ush
+++ b/Engine/Plugins/FX/Niagara/Shaders/Private/NiagaraDataInterfaceStaticMeshTemplate.ush
@@ -106,9 +106,20 @@ void GetVertex_{ParameterName}(int Vertex, out float3 Position, out float3 Veloc
 		float3 TangentY = cross(TangentZ.xyz, TangentX.xyz) * TangentZ.w;
 
 		Velocity	= float3(0, 0, 0);
-		Normal		= TangentZ.xyz;
-		Bitangent	= TangentY;
-		Tangent		= TangentX.xyz;
+
+		// Sometimes we have bogus tangent data in the buffer, i.e. X==Z so ensure we handle this and don't end up with NaNs in the buffer
+		if ( length2(TangentY) > 0.001f )
+		{
+			Normal		= TangentZ.xyz;
+			Bitangent	= TangentY;
+			Tangent		= TangentX.xyz;
+		}
+		else
+		{
+			Normal		= float3(0, 0, 1);
+			Bitangent	= float3(0, 1, 0);
+			Tangent		= float3(1, 0, 0);
+		}
 	}
 	else
 	{
diff --git a/Engine/Plugins/FX/Niagara/Source/Niagara/Private/DataInterface/NiagaraDataInterfaceStaticMesh.cpp b/Engine/Plugins/FX/Niagara/Source/Niagara/Private/DataInterface/NiagaraDataInterfaceStaticMesh.cpp
--- a/Engine/Plugins/FX/Niagara/Source/Niagara/Private/DataInterface/NiagaraDataInterfaceStaticMesh.cpp
+++ b/Engine/Plugins/FX/Niagara/Source/Niagara/Private/DataInterface/NiagaraDataInterfaceStaticMesh.cpp
@@ -257,6 +257,27 @@ namespace NDIStaticMeshLocal
 		uint32 Alias = 0;
 	};
 
+	struct FNDITangentBasis
+	{
+		FVector3f TangentX;
+		FVector3f TangentY;
+		FVector3f TangentZ;
+	
+		template<typename TTransformHandler = FNDITransformHandlerNoop>
+		void TransformAndEnsureValid(const TTransformHandler& TransformHandler, const FMatrix44f& Matrix)
+		{
+			if (TangentY.SquaredLength() <= 0.001f)
+			{
+				TangentX = FVector3f(1.0f, 0.0f, 0.0f);
+				TangentY = FVector3f(0.0f, 1.0f, 0.0f);
+				TangentZ = FVector3f(0.0f, 0.0f, 1.0f);
+			}
+			TransformHandler.TransformVector(TangentX, Matrix);
+			TransformHandler.TransformVector(TangentY, Matrix);
+			TransformHandler.TransformVector(TangentZ, Matrix);
+		}
+	};
+
 	class FNDISectionAreaWeightedSampler : public FWeightedRandomSampler
 	{
 	public:
@@ -1367,52 +1388,28 @@ namespace NDIStaticMeshLocal
 			return LODResource->VertexBuffers.PositionVertexBuffer.VertexPosition(Vertex);
 		}
 
-		FORCEINLINE FVector3f GetTangentX(int32 Vertex) const
+		FORCEINLINE FNDITangentBasis GetTangentBasis(int32 Vertex) const
 		{
-			FVector3f TangentX = LODResource->VertexBuffers.StaticMeshVertexBuffer.VertexTangentX(Vertex);
-			TransformHandler.TransformVector(TangentX, FMatrix44f(InstanceData->TransformInverseTransposed));		// LWC_TODO: Precision loss?
-			return TangentX;
+			FNDITangentBasis TangentBasis;
+			TangentBasis.TangentX = LODResource->VertexBuffers.StaticMeshVertexBuffer.VertexTangentX(Vertex);
+			TangentBasis.TangentY = LODResource->VertexBuffers.StaticMeshVertexBuffer.VertexTangentY(Vertex);
+			TangentBasis.TangentZ = LODResource->VertexBuffers.StaticMeshVertexBuffer.VertexTangentZ(Vertex);
+			TangentBasis.TransformAndEnsureValid(TransformHandler, FMatrix44f(InstanceData->TransformInverseTransposed));		// LWC_TODO: Precision loss?
+			return TangentBasis;
 		}
 
-		FORCEINLINE FVector3f GetTangentY(int32 Vertex) const
+		FORCEINLINE FNDITangentBasis GetTangentBasisInterpolated(int32 Vertex, float Interp) const
 		{
-			FVector3f TangentY = LODResource->VertexBuffers.StaticMeshVertexBuffer.VertexTangentY(Vertex);
-			TransformHandler.TransformVector(TangentY, FMatrix44f(InstanceData->TransformInverseTransposed));		// LWC_TODO: Precision loss?
-			return TangentY;
-		}
+			FNDITangentBasis TangentBasisCurr = GetTangentBasis(Vertex);
+			FNDITangentBasis TangentBasisPrev = TangentBasisCurr;
+			TangentBasisCurr.TransformAndEnsureValid(TransformHandler, FMatrix44f(InstanceData->TransformInverseTransposed));		// LWC_TODO: Precision loss?
+			TangentBasisPrev.TransformAndEnsureValid(TransformHandler, FMatrix44f(InstanceData->PrevTransformInverseTransposed));	// LWC_TODO: Precision loss?
 
-		FORCEINLINE FVector3f GetTangentZ(int32 Vertex) const
-		{
-			FVector3f TangentZ = LODResource->VertexBuffers.StaticMeshVertexBuffer.VertexTangentZ(Vertex);
-			TransformHandler.TransformVector(TangentZ, FMatrix44f(InstanceData->TransformInverseTransposed));		// LWC_TODO: Precision loss?
-			return TangentZ;
-		}
-
-		FORCEINLINE FVector3f GetTangentXInterpolated(int32 Vertex, float Interp) const
-		{
-			FVector3f TangentXPrev = LODResource->VertexBuffers.StaticMeshVertexBuffer.VertexTangentX(Vertex);
-			FVector3f TangentXCurr = TangentXPrev;
-			TransformHandler.TransformVector(TangentXPrev, FMatrix44f(InstanceData->PrevTransformInverseTransposed));		// LWC_TODO: Precision loss?
-			TransformHandler.TransformVector(TangentXCurr, FMatrix44f(InstanceData->TransformInverseTransposed));				// LWC_TODO: Precision loss?
-			return FMath::Lerp(TangentXPrev, TangentXCurr, Interp);
-		}
-
-		FORCEINLINE FVector3f GetTangentYInterpolated(int32 Vertex, float Interp) const
-		{
-			FVector3f TangentYPrev = LODResource->VertexBuffers.StaticMeshVertexBuffer.VertexTangentY(Vertex);
-			FVector3f TangentYCurr = TangentYPrev;
-			TransformHandler.TransformVector(TangentYPrev, FMatrix44f(InstanceData->PrevTransformInverseTransposed));		// LWC_TODO: Precision loss?
-			TransformHandler.TransformVector(TangentYCurr, FMatrix44f(InstanceData->TransformInverseTransposed));				// LWC_TODO: Precision loss?
-			return FMath::Lerp(TangentYPrev, TangentYCurr, Interp);
-		}
-
-		FORCEINLINE FVector3f GetTangentZInterpolated(int32 Vertex, float Interp) const
-		{
-			FVector3f TangentZPrev = LODResource->VertexBuffers.StaticMeshVertexBuffer.VertexTangentZ(Vertex);
-			FVector3f TangentZCurr = TangentZPrev;
-			TransformHandler.TransformVector(TangentZPrev, FMatrix44f(InstanceData->PrevTransformInverseTransposed));		// LWC_TODO: Precision loss?
-			TransformHandler.TransformVector(TangentZCurr, FMatrix44f(InstanceData->TransformInverseTransposed));				// LWC_TODO: Precision loss?
-			return FMath::Lerp(TangentZPrev, TangentZCurr, Interp);
+			FNDITangentBasis InterpolatedBasis;
+			InterpolatedBasis.TangentX = FMath::Lerp(TangentBasisPrev.TangentX, TangentBasisCurr.TangentX, Interp);
+			InterpolatedBasis.TangentY = FMath::Lerp(TangentBasisPrev.TangentY, TangentBasisCurr.TangentY, Interp);
+			InterpolatedBasis.TangentZ = FMath::Lerp(TangentBasisPrev.TangentZ, TangentBasisCurr.TangentZ, Interp);
+			return InterpolatedBasis;
 		}
 
 		FORCEINLINE FLinearColor GetColor(int32 Vertex) const
@@ -3566,11 +3563,12 @@ void UNiagaraDataInterfaceStaticMesh::VMGetVertex(FVectorVMExternalFunctionConte
 			const FVector3f Position = StaticMeshHelper.TransformPosition(LocalPosition);
 			const FVector3f PreviousPosition = StaticMeshHelper.PreviousTransformPosition(LocalPosition);
 			const FVector3f Velocity = (Position - PreviousPosition) * InvDt;
+			const NDIStaticMeshLocal::FNDITangentBasis TangentBasis = StaticMeshHelper.GetTangentBasis(Vertex);
 			OutPosition.SetAndAdvance(Position);
 			OutVelocity.SetAndAdvance(Velocity);
-			OutNormal.SetAndAdvance(StaticMeshHelper.GetTangentZ(Vertex));
-			OutBitangent.SetAndAdvance(StaticMeshHelper.GetTangentY(Vertex));
-			OutTangent.SetAndAdvance(StaticMeshHelper.GetTangentX(Vertex));
+			OutNormal.SetAndAdvance(TangentBasis.TangentZ);
+			OutBitangent.SetAndAdvance(TangentBasis.TangentY);
+			OutTangent.SetAndAdvance(TangentBasis.TangentX);
 		}
 	}
 	else
@@ -3616,11 +3614,12 @@ void UNiagaraDataInterfaceStaticMesh::VMGetVertexInterpolated(FVectorVMExternalF
 			const FVector3f Position = StaticMeshHelper.TransformPosition(LocalPosition);
 			const FVector3f PreviousPosition = StaticMeshHelper.PreviousTransformPosition(LocalPosition);
 			const FVector3f Velocity = (Position - PreviousPosition) * InvDt;
+			const NDIStaticMeshLocal::FNDITangentBasis TangentBasis = StaticMeshHelper.GetTangentBasisInterpolated(Vertex, Interp);
 			OutPosition.SetAndAdvance(FMath::Lerp(PreviousPosition, Position, Interp));
 			OutVelocity.SetAndAdvance(Velocity);
-			OutNormal.SetAndAdvance(StaticMeshHelper.GetTangentZInterpolated(Vertex, Interp));
-			OutBitangent.SetAndAdvance(StaticMeshHelper.GetTangentYInterpolated(Vertex, Interp));
-			OutTangent.SetAndAdvance(StaticMeshHelper.GetTangentXInterpolated(Vertex, Interp));
+			OutNormal.SetAndAdvance(TangentBasis.TangentZ);
+			OutBitangent.SetAndAdvance(TangentBasis.TangentY);
+			OutTangent.SetAndAdvance(TangentBasis.TangentX);
 		}
 	}
 	else
