Fix for Sim Cache preview scene overriding other previews scene's settings

diff --git a/Engine/Plugins/FX/Niagara/Source/NiagaraEditor/Private/Widgets/SNiagaraSimCacheViewport.cpp b/Engine/Plugins/FX/Niagara/Source/NiagaraEditor/Private/Widgets/SNiagaraSimCacheViewport.cpp
--- a/Engine/Plugins/FX/Niagara/Source/NiagaraEditor/Private/Widgets/SNiagaraSimCacheViewport.cpp
+++ b/Engine/Plugins/FX/Niagara/Source/NiagaraEditor/Private/Widgets/SNiagaraSimCacheViewport.cpp
@@ -97,8 +97,8 @@ FSceneView* FNiagaraSimCacheViewportClient::CalcSceneView(FSceneViewFamily* View
 void SNiagaraSimCacheViewport::Construct(FArguments InArgs)
 {
 	AdvancedPreviewScene = MakeShareable(new FAdvancedPreviewScene(FPreviewScene::ConstructionValues()));
-	AdvancedPreviewScene->SetFloorVisibility(false);
-	AdvancedPreviewScene->SetEnvironmentVisibility(false, false);
+	AdvancedPreviewScene->SetFloorVisibility(false, true);
+	AdvancedPreviewScene->SetEnvironmentVisibility(false, true);
 	
 	SEditorViewport::Construct(SEditorViewport::FArguments());
 
