Added Niagara System & Emitter icons, changed Niagara System content browser color to blue, removed unneeded icons, changed tooltips for emitters & systems

diff --git a/Engine/Plugins/FX/Niagara/Source/Niagara/Classes/NiagaraEmitter.h b/Engine/Plugins/FX/Niagara/Source/Niagara/Classes/NiagaraEmitter.h
--- a/Engine/Plugins/FX/Niagara/Source/Niagara/Classes/NiagaraEmitter.h
+++ b/Engine/Plugins/FX/Niagara/Source/Niagara/Classes/NiagaraEmitter.h
@@ -536,8 +536,8 @@ private:
 };
 
 /** 
- *	UNiagaraEmitter stores the attributes of an FNiagaraEmitterInstance
- *	that need to be serialized and are used for its initialization 
+ *	Niagara Emitters are particle spawners that can be reused for different effects by putting them into Niagara Systems.
+ *	Emitters render their particles using different renderers, such as Sprite Renderers or Mesh Renderers to produce different effects.
  */
 UCLASS(MinimalAPI)
 class UNiagaraEmitter : public UObject, public INiagaraParameterDefinitionsSubscriber, public FNiagaraVersionedObject
diff --git a/Engine/Plugins/FX/Niagara/Source/Niagara/Classes/NiagaraSystem.h b/Engine/Plugins/FX/Niagara/Source/Niagara/Classes/NiagaraSystem.h
--- a/Engine/Plugins/FX/Niagara/Source/Niagara/Classes/NiagaraSystem.h
+++ b/Engine/Plugins/FX/Niagara/Source/Niagara/Classes/NiagaraSystem.h
@@ -220,7 +220,8 @@ struct NIAGARA_API FNiagaraSystemScalabilityOverrides
 };
 
 
-/** Container for multiple emitters that combine together to create a particle system effect.*/
+/** A Niagara System contains multiple Niagara Emitters to create various effects.
+ * Niagara Systems can be placed in the world, unlike Emitters, and expose User Parameters to configure an effect at runtime.*/
 UCLASS(BlueprintType, meta= (LoadBehavior = "LazyOnDemand"))
 class NIAGARA_API UNiagaraSystem : public UFXSystemAsset, public INiagaraParameterDefinitionsSubscriber
 {
diff --git a/Engine/Plugins/FX/Niagara/Source/NiagaraEditor/Private/NiagaraEditorStyle.cpp b/Engine/Plugins/FX/Niagara/Source/NiagaraEditor/Private/NiagaraEditorStyle.cpp
--- a/Engine/Plugins/FX/Niagara/Source/NiagaraEditor/Private/NiagaraEditorStyle.cpp
+++ b/Engine/Plugins/FX/Niagara/Source/NiagaraEditor/Private/NiagaraEditorStyle.cpp
@@ -2,6 +2,7 @@
 
 #include "NiagaraEditorStyle.h"
 
+#include "ClassIconFinder.h"
 #include "Framework/Application/SlateApplication.h"
 #include "Styling/SlateStyleMacros.h"
 #include "Styling/AppStyle.h"
@@ -413,8 +414,8 @@ void FNiagaraEditorStyle::InitEmitterDetails()
 void FNiagaraEditorStyle::InitAssetColors()
 {
 	// Asset colors
-	Set("NiagaraEditor.AssetColors.System", FLinearColor(1.0f, 0.0f, 0.0f));
-	Set("NiagaraEditor.AssetColors.Emitter", FLinearColor(1.0f, 0.3f, 0.0f));
+	Set("NiagaraEditor.AssetColors.System", FLinearColor(FColor(1, 202, 252)));
+	Set("NiagaraEditor.AssetColors.Emitter", FLinearColor(FColor(241, 99, 6)));
 	Set("NiagaraEditor.AssetColors.Script", FLinearColor(1.0f, 1.0f, 0.0f));
 	Set("NiagaraEditor.AssetColors.ParameterCollection", FLinearColor(1.0f, 1.0f, 0.3f));
 	Set("NiagaraEditor.AssetColors.ParameterCollectionInstance", FLinearColor(1.0f, 1.0f, 0.7f));
@@ -434,12 +435,18 @@ void FNiagaraEditorStyle::InitThumbnails()
 
 void FNiagaraEditorStyle::InitClassIcon()
 {
+	Set("ClassThumbnail.NiagaraEmitter", new IMAGE_BRUSH_SVG("Icons/Classes/NiagaraEmitter_64", Icon64x64));
+	Set("ClassThumbnail.NiagaraSystem", new IMAGE_BRUSH_SVG("Icons/Classes/NiagaraSystem_64", Icon64x64));
+
+	Set("ClassIcon.NiagaraEmitter", new IMAGE_BRUSH_SVG("Icons/Classes/NiagaraEmitter_16", Icon16x16));
+	Set("ClassIcon.NiagaraSystem", new IMAGE_BRUSH_SVG("Icons/Classes/NiagaraSystem_16", Icon16x16));
+	Set("ClassIcon.NiagaraActor", new IMAGE_BRUSH_SVG("Icons/Classes/NiagaraSystem_16", Icon16x16));
 	// Renderer class icons
-	Set("ClassIcon.NiagaraSpriteRendererProperties", new IMAGE_BRUSH("Icons/Renderers/renderer_sprite", Icon16x16));
-	Set("ClassIcon.NiagaraMeshRendererProperties", new IMAGE_BRUSH("Icons/Renderers/renderer_mesh", Icon16x16));
-	Set("ClassIcon.NiagaraRibbonRendererProperties", new IMAGE_BRUSH("Icons/Renderers/renderer_ribbon", Icon16x16));
-	Set("ClassIcon.NiagaraLightRendererProperties", new IMAGE_BRUSH("Icons/Renderers/renderer_light", Icon16x16));
-	Set("ClassIcon.NiagaraRendererProperties", new IMAGE_BRUSH("Icons/Renderers/renderer_default", Icon16x16));
+	Set("ClassIcon.NiagaraSpriteRendererProperties", new IMAGE_BRUSH("Icons/Classes/renderer_sprite", Icon16x16));
+	Set("ClassIcon.NiagaraMeshRendererProperties", new IMAGE_BRUSH("Icons/Classes/renderer_mesh", Icon16x16));
+	Set("ClassIcon.NiagaraRibbonRendererProperties", new IMAGE_BRUSH("Icons/Classes/renderer_ribbon", Icon16x16));
+	Set("ClassIcon.NiagaraLightRendererProperties", new IMAGE_BRUSH("Icons/Classes/renderer_light", Icon16x16));
+	Set("ClassIcon.NiagaraRendererProperties", new IMAGE_BRUSH("Icons/Classes/renderer_default", Icon16x16));
 }
 
 void FNiagaraEditorStyle::InitStackIcons()
