- Ensure debug draw max lines can't overflow a TArray capacity
- Removed pre-warm from array also as it was allocating in the GPU case

#rnx

diff --git a/Engine/Plugins/FX/Niagara/Source/Niagara/Private/DataInterface/NiagaraDataInterfaceDebugDraw.cpp b/Engine/Plugins/FX/Niagara/Source/Niagara/Private/DataInterface/NiagaraDataInterfaceDebugDraw.cpp
--- a/Engine/Plugins/FX/Niagara/Source/Niagara/Private/DataInterface/NiagaraDataInterfaceDebugDraw.cpp
+++ b/Engine/Plugins/FX/Niagara/Source/Niagara/Private/DataInterface/NiagaraDataInterfaceDebugDraw.cpp
@@ -3042,8 +3042,8 @@ bool UNiagaraDataInterfaceDebugDraw::PerInstanceTick(void* PerInstanceData, FNia
 {
 #if NIAGARA_COMPUTEDEBUG_ENABLED
 	FNDIDebugDrawInstanceData_GameThread* InstanceData = reinterpret_cast<FNDIDebugDrawInstanceData_GameThread*>(PerInstanceData);
-	InstanceData->LineBuffer.Reset(OverrideMaxLineInstances);
-	InstanceData->OverrideMaxLineInstances = OverrideMaxLineInstances;
+	InstanceData->LineBuffer.Reset();
+	InstanceData->OverrideMaxLineInstances = FMath::Min<uint32>(OverrideMaxLineInstances, TNumericLimits<int32>::Max());
 #endif
 	return false;
 }
