Added overview tab for Niagara Sim Cache editor
- Tree view for component filtering
- Buffer selection list

Updated default layout to show the new overview tab

Updated component infos to be cached instead of regenerated every time a new buffer was selected

Updated coloring of current frame selection in Sim Cache editor timeline

diff --git a/Engine/Plugins/FX/Niagara/Source/NiagaraEditor/Private/NiagaraEditorStyle.cpp b/Engine/Plugins/FX/Niagara/Source/NiagaraEditor/Private/NiagaraEditorStyle.cpp
--- a/Engine/Plugins/FX/Niagara/Source/NiagaraEditor/Private/NiagaraEditorStyle.cpp
+++ b/Engine/Plugins/FX/Niagara/Source/NiagaraEditor/Private/NiagaraEditorStyle.cpp
@@ -728,6 +728,59 @@ void FNiagaraEditorStyle::InitHierarchyEditor()
 	Set("NiagaraEditor.HierarchyEditor.ButtonStyle", ButtonStyle);
 }
 
+void FNiagaraEditorStyle::InitSimCacheEditor()
+{
+	const FSlateColor SelectionColor = FAppStyle::GetSlateColor("SelectionColor");
+	const FSlateColor SelectionColor_Inactive = FAppStyle::GetSlateColor("SelectionColor_Inactive");
+	const FSlateColor DefaultForeground = FAppStyle::GetSlateColor("DefaultForeground");
+	const FLinearColor SystemBackground = FLinearColor(0.05f, 0.05f, 0.05f);
+
+	Set("NiagaraEditor.SimCache.SystemItem", FTableRowStyle()
+		.SetEvenRowBackgroundBrush(CORE_IMAGE_BRUSH("Common/Selection", Icon8x8, SystemBackground))
+		.SetEvenRowBackgroundHoveredBrush(CORE_IMAGE_BRUSH("Common/Selection", Icon8x8, SelectionColor_Inactive))
+		.SetOddRowBackgroundBrush(CORE_IMAGE_BRUSH("Common/Selection", Icon8x8, SystemBackground))
+		.SetOddRowBackgroundHoveredBrush(CORE_IMAGE_BRUSH("Common/Selection", Icon8x8, SelectionColor_Inactive))
+		.SetSelectorFocusedBrush(CORE_IMAGE_BRUSH("Common/Selection", Icon8x8, SystemBackground))
+		.SetActiveBrush(CORE_IMAGE_BRUSH("Common/Selection", Icon8x8, SelectionColor))
+		.SetActiveHoveredBrush(CORE_IMAGE_BRUSH("Common/Selection", Icon8x8, SelectionColor))
+		.SetInactiveBrush(CORE_IMAGE_BRUSH("Common/Selection", Icon8x8, SelectionColor_Inactive))
+		.SetInactiveHoveredBrush(CORE_IMAGE_BRUSH("Common/Selection", Icon8x8, SelectionColor_Inactive)));
+
+	Set("NiagaraEditor.SimCache.EmitterItem", FTableRowStyle()
+		.SetEvenRowBackgroundBrush(CORE_IMAGE_BRUSH("Common/Selection", Icon8x8, SystemBackground))
+		.SetEvenRowBackgroundHoveredBrush(CORE_IMAGE_BRUSH("Common/Selection", Icon8x8, SelectionColor_Inactive))
+		.SetOddRowBackgroundBrush(CORE_IMAGE_BRUSH("Common/Selection", Icon8x8, SystemBackground))
+		.SetOddRowBackgroundHoveredBrush(CORE_IMAGE_BRUSH("Common/Selection", Icon8x8, SelectionColor_Inactive))
+		.SetSelectorFocusedBrush(CORE_IMAGE_BRUSH("Common/Selection", Icon8x8, SystemBackground))
+		.SetActiveBrush(CORE_IMAGE_BRUSH("Common/Selection", Icon8x8, SelectionColor))
+		.SetActiveHoveredBrush(CORE_IMAGE_BRUSH("Common/Selection", Icon8x8, SelectionColor))
+		.SetInactiveBrush(CORE_IMAGE_BRUSH("Common/Selection", Icon8x8, SelectionColor_Inactive))
+		.SetInactiveHoveredBrush(CORE_IMAGE_BRUSH("Common/Selection", Icon8x8, SelectionColor_Inactive)));
+
+	Set("NiagaraEditor.SimCache.ComponentItem", FTableRowStyle()
+		.SetEvenRowBackgroundBrush(FSlateNoResource())
+		.SetEvenRowBackgroundHoveredBrush(CORE_IMAGE_BRUSH("Common/Selection", Icon8x8, SelectionColor_Inactive))
+		.SetOddRowBackgroundBrush(FSlateNoResource())
+		.SetOddRowBackgroundHoveredBrush(CORE_IMAGE_BRUSH("Common/Selection", Icon8x8, SelectionColor_Inactive))
+		.SetSelectorFocusedBrush(FSlateNoResource())
+		.SetActiveBrush(CORE_IMAGE_BRUSH("Common/Selection", Icon8x8, SelectionColor))
+		.SetActiveHoveredBrush(CORE_IMAGE_BRUSH("Common/Selection", Icon8x8, SelectionColor))
+		.SetInactiveBrush(CORE_IMAGE_BRUSH("Common/Selection", Icon8x8, SelectionColor_Inactive))
+		.SetInactiveHoveredBrush(CORE_IMAGE_BRUSH("Common/Selection", Icon8x8, SelectionColor_Inactive)));
+
+	const FLinearColor NormalColor(0.15, 0.15, 0.15, 1);
+
+	Set ("NiagaraEditor.SimCache.FilterToggleStyle", FCheckBoxStyle()
+		.SetCheckBoxType(ESlateCheckBoxType::ToggleButton)
+		.SetUncheckedImage( FSlateRoundedBoxBrush( NormalColor, 2.0f) )
+		.SetUncheckedHoveredImage(FSlateRoundedBoxBrush( FStyleColors::PrimaryHover, 2.0f))
+		.SetUncheckedPressedImage(FSlateRoundedBoxBrush( FStyleColors::PrimaryPress, 2.0f))
+		.SetCheckedImage(FSlateRoundedBoxBrush(FStyleColors::Select, 2.0f))
+		.SetCheckedHoveredImage(FSlateRoundedBoxBrush(FStyleColors::Select, 2.0f))
+		.SetCheckedPressedImage(FSlateRoundedBoxBrush(FStyleColors::Select, 2.0f))
+		);
+}
+
 FNiagaraEditorStyle::FNiagaraEditorStyle() : FSlateStyleSet("NiagaraEditorStyle")
 {
 	SetContentRoot(FPaths::EnginePluginsDir() / TEXT("FX/Niagara/Content/Slate"));
@@ -763,6 +816,7 @@ FNiagaraEditorStyle::FNiagaraEditorStyle() : FSlateStyleSet("NiagaraEditorStyle"
 	InitViewportStyle();
 	InitScratchStyle();
 	InitHierarchyEditor();
+	InitSimCacheEditor();
 }
 
 void FNiagaraEditorStyle::InitScratchStyle()
diff --git a/Engine/Plugins/FX/Niagara/Source/NiagaraEditor/Private/Toolkits/NiagaraSimCacheToolkit.cpp b/Engine/Plugins/FX/Niagara/Source/NiagaraEditor/Private/Toolkits/NiagaraSimCacheToolkit.cpp
--- a/Engine/Plugins/FX/Niagara/Source/NiagaraEditor/Private/Toolkits/NiagaraSimCacheToolkit.cpp
+++ b/Engine/Plugins/FX/Niagara/Source/NiagaraEditor/Private/Toolkits/NiagaraSimCacheToolkit.cpp
@@ -7,6 +7,8 @@
 #include "AdvancedPreviewSceneModule.h"
 #include "Widgets/Docking/SDockTab.h"
 #include "NiagaraSimCacheViewModel.h"
+#include "SNiagaraSimCacheTreeView.h"
+#include "SNiagaraSimCacheOverview.h"
 #include "SNiagaraSimCacheViewport.h"
 #include "SNiagaraSimCacheViewTimeline.h"
 #include "SNiagaraSimCacheViewTransportControls.h"
@@ -16,6 +18,7 @@
 const FName FNiagaraSimCacheToolkit::NiagaraSimCacheSpreadsheetTabId(TEXT("NiagaraSimCacheEditor_Spreadsheet"));
 const FName FNiagaraSimCacheToolkit::NiagaraSimCacheViewportTabId(TEXT("NiagaraSimCacheEditor_Viewport"));
 const FName FNiagaraSimCacheToolkit::NiagaraSimCachePreviewSettingsTabId(TEXT("NiagaraSimCacheEditor_PreviewSettings"));
+const FName FNiagaraSimCacheToolkit::NiagaraSimCacheOverviewTabId(TEXT("NiagaraSimCacheEditor_Overview"));
 
 FNiagaraSimCacheToolkit::FNiagaraSimCacheToolkit()
 {
@@ -44,6 +47,10 @@ void FNiagaraSimCacheToolkit::RegisterTabSpawners(const TSharedRef<class FTabMan
 	InTabManager->RegisterTabSpawner(NiagaraSimCachePreviewSettingsTabId, FOnSpawnTab::CreateSP(this, &FNiagaraSimCacheToolkit::SpawnTab_PreviewSettings))
 		.SetDisplayName(LOCTEXT("PreviewSceneSettingsTab","Preview Scene Settings"))
 		.SetGroup(WorkspaceMenuCategory.ToSharedRef());
+
+	InTabManager->RegisterTabSpawner(NiagaraSimCacheOverviewTabId, FOnSpawnTab::CreateSP(this, &FNiagaraSimCacheToolkit::SpawnTab_Overview))
+		.SetDisplayName(LOCTEXT("OverviewTab", "Overview"))
+		.SetGroup(WorkspaceMenuCategory.ToSharedRef());
 }
 
 void FNiagaraSimCacheToolkit::UnregisterTabSpawners(const TSharedRef<class FTabManager>& InTabManager)
@@ -66,14 +73,27 @@ void FNiagaraSimCacheToolkit::Initialize(const EToolkitMode::Type Mode, const TS
 		const TSharedRef<FTabManager::FLayout> StandaloneDefaultLayout = FTabManager::NewLayout("Standalone_Niagara_SimCache_Layout_V1")
 			->AddArea
 			(
-				FTabManager::NewPrimaryArea()->SetOrientation(Orient_Vertical)
+				FTabManager::NewPrimaryArea()->SetOrientation(Orient_Horizontal)
+				->Split(
+					FTabManager::NewSplitter()->SetOrientation(Orient_Vertical)
+					->SetSizeCoefficient(0.2f)
+					->Split
+					(
+						FTabManager::NewStack()
+						->SetSizeCoefficient(0.3f)
+						->AddTab(NiagaraSimCacheViewportTabId, ETabState::OpenedTab)
+					)
+					->Split(
+						FTabManager::NewStack()
+						->SetSizeCoefficient(0.7f)
+						->AddTab(NiagaraSimCacheOverviewTabId, ETabState::OpenedTab)
+					)
+				)
 				->Split
 				(
 					FTabManager::NewStack()
-					->SetSizeCoefficient(0.1f)
+					->SetSizeCoefficient(0.8f)
 					->AddTab(NiagaraSimCacheSpreadsheetTabId, ETabState::OpenedTab)
-					->AddTab(NiagaraSimCacheViewportTabId, ETabState::OpenedTab)
-					->SetForegroundTab(NiagaraSimCacheSpreadsheetTabId)
 				)
 			);
 
@@ -195,4 +215,16 @@ TSharedRef<SDockTab> FNiagaraSimCacheToolkit::SpawnTab_PreviewSettings(const FSp
 	return SpawnedTab;
 }
 
+TSharedRef<SDockTab> FNiagaraSimCacheToolkit::SpawnTab_Overview(const FSpawnTabArgs& Args)
+{
+	TSharedRef<SDockTab> SpawnedTab = SNew(SDockTab)
+		.Label(LOCTEXT("OverviewTab", "Overview"))
+		[
+			SNew(SNiagaraSimCacheOverview)
+			.SimCacheViewModel(SimCacheViewModel)
+		];
+
+	return SpawnedTab;
+}
+
 #undef LOCTEXT_NAMESPACE /** NiagaraSimCacheToolkit **/
\ No newline at end of file
diff --git a/Engine/Plugins/FX/Niagara/Source/NiagaraEditor/Private/Toolkits/NiagaraSimCacheToolkit.h b/Engine/Plugins/FX/Niagara/Source/NiagaraEditor/Private/Toolkits/NiagaraSimCacheToolkit.h
--- a/Engine/Plugins/FX/Niagara/Source/NiagaraEditor/Private/Toolkits/NiagaraSimCacheToolkit.h
+++ b/Engine/Plugins/FX/Niagara/Source/NiagaraEditor/Private/Toolkits/NiagaraSimCacheToolkit.h
@@ -38,8 +38,9 @@ protected:
 
 private:
 	TSharedRef<SDockTab> SpawnTab_SimCacheSpreadsheet(const FSpawnTabArgs& Args);
-	TSharedRef<SDockTab>  SpawnTab_SimCacheViewport(const FSpawnTabArgs& Args);
+	TSharedRef<SDockTab> SpawnTab_SimCacheViewport(const FSpawnTabArgs& Args);
 	TSharedRef<SDockTab> SpawnTab_PreviewSettings(const FSpawnTabArgs& Args);
+	TSharedRef<SDockTab> SpawnTab_Overview(const FSpawnTabArgs& Args);
 
 	// The sim cache being viewed.
 	TWeakObjectPtr <UNiagaraSimCache> SimCache;
@@ -56,4 +57,5 @@ private:
 	static const FName NiagaraSimCacheSpreadsheetTabId;
 	static const FName NiagaraSimCacheViewportTabId;
 	static const FName NiagaraSimCachePreviewSettingsTabId;
+	static const FName NiagaraSimCacheOverviewTabId;
 };
\ No newline at end of file
diff --git a/Engine/Plugins/FX/Niagara/Source/NiagaraEditor/Private/ViewModels/NiagaraSimCacheViewModel.cpp b/Engine/Plugins/FX/Niagara/Source/NiagaraEditor/Private/ViewModels/NiagaraSimCacheViewModel.cpp
--- a/Engine/Plugins/FX/Niagara/Source/NiagaraEditor/Private/ViewModels/NiagaraSimCacheViewModel.cpp
+++ b/Engine/Plugins/FX/Niagara/Source/NiagaraEditor/Private/ViewModels/NiagaraSimCacheViewModel.cpp
@@ -6,9 +6,9 @@
 #include "NiagaraEditorCommon.h"
 #include "NiagaraDebuggerCommon.h"
 #include "Serialization/MemoryReader.h"
+#include "SNiagaraSimCacheTreeView.h"
 #include "Serialization/ObjectAndNameAsStringProxyArchive.h"
 #include "UObject/Package.h"
-#include "UObject/Package.h"
 
 #define LOCTEXT_NAMESPACE "NiagaraSimCacheViewModel"
 
@@ -31,8 +31,11 @@ void FNiagaraSimCacheViewModel::Initialize(TWeakObjectPtr<UNiagaraSimCache> SimC
 	}
 
 	WeakSimCache = SimCache;
+	UpdateComponentInfos();
 	UpdateCachedFrame();
 	SetupPreviewComponentAndInstance();
+	
+	OnSimCacheChangedDelegate.Broadcast();
 	OnViewDataChangedDelegate.Broadcast(true);
 }
 
@@ -47,6 +50,9 @@ void FNiagaraSimCacheViewModel::UpdateSimCache(const FNiagaraSystemSimCacheCaptu
 		FMemoryReader ArReader(Reply.SimCacheData);
 		FObjectAndNameAsStringProxyArchive ProxyArReader(ArReader, false);
 		SimCache->Serialize(ProxyArReader);
+		bComponentFilterActive = false;
+		UpdateComponentInfos();
+		
 	}
 	else
 	{
@@ -81,7 +87,7 @@ void FNiagaraSimCacheViewModel::SetupPreviewComponentAndInstance()
 
 FText FNiagaraSimCacheViewModel::GetComponentText(const FName ComponentName, const int32 InstanceIndex) const
 {
-	const FComponentInfo* ComponentInfo = ComponentInfos.FindByPredicate([ComponentName](const FComponentInfo& FoundInfo) { return FoundInfo.Name == ComponentName; });
+	const FComponentInfo* ComponentInfo = ComponentInfos->FindByPredicate([ComponentName](const FComponentInfo& FoundInfo) { return FoundInfo.Name == ComponentName; });
 
 	if (ComponentInfo)
 	{
@@ -118,6 +124,15 @@ FText FNiagaraSimCacheViewModel::GetComponentText(const FName ComponentName, con
 	return LOCTEXT("Error", "Error");
 }
 
+void FNiagaraSimCacheViewModel::SetComponentFilterActive(bool bNewActive)
+{
+	if(bNewActive != bComponentFilterActive)
+	{
+		bComponentFilterActive = bNewActive;
+		OnViewDataChangedDelegate.Broadcast(true);
+	}
+}
+
 int32 FNiagaraSimCacheViewModel::GetNumFrames() const
 {
 	UNiagaraSimCache* SimCache = WeakSimCache.Get();
@@ -145,7 +160,10 @@ void FNiagaraSimCacheViewModel::SetEmitterIndex(const int32 InEmitterIndex)
 {
 	EmitterIndex = InEmitterIndex;
 	UpdateCachedFrame();
-	OnViewDataChangedDelegate.Broadcast(false);
+	UpdateCurrentEntries();
+	bComponentFilterActive = false;
+	OnBufferChangedDelegate.Broadcast();
+	OnViewDataChangedDelegate.Broadcast(true);
 }
 
 bool FNiagaraSimCacheViewModel::IsCacheValid()
@@ -171,13 +189,25 @@ FNiagaraSimCacheViewModel::FOnViewDataChanged& FNiagaraSimCacheViewModel::OnView
 	return OnViewDataChangedDelegate;
 }
 
+FNiagaraSimCacheViewModel::FOnSimCacheChanged& FNiagaraSimCacheViewModel::OnSimCacheChanged()
+{
+	return OnSimCacheChangedDelegate;
+}
+
+FNiagaraSimCacheViewModel::FOnBufferChanged& FNiagaraSimCacheViewModel::OnBufferChanged()
+{
+	return OnBufferChangedDelegate;
+}
+
 void FNiagaraSimCacheViewModel::OnCacheModified(UNiagaraSimCache* SimCache)
 {
 	if ( UNiagaraSimCache* ThisSimCache = WeakSimCache.Get() )
 	{
 		if ( ThisSimCache == SimCache )
 		{
+			UpdateComponentInfos();
 			UpdateCachedFrame();
+			OnSimCacheChangedDelegate.Broadcast();
 			OnViewDataChangedDelegate.Broadcast(true);
 		}
 	}
@@ -186,10 +216,10 @@ void FNiagaraSimCacheViewModel::OnCacheModified(UNiagaraSimCache* SimCache)
 void FNiagaraSimCacheViewModel::UpdateCachedFrame()
 {
 	NumInstances = 0;
-	ComponentInfos.Empty();
 	FloatComponents.Empty();
 	HalfComponents.Empty();
 	Int32Components.Empty();
+	
 
 	UNiagaraSimCache* SimCache = WeakSimCache.Get();
 	if (SimCache == nullptr)
@@ -207,21 +237,45 @@ void FNiagaraSimCacheViewModel::UpdateCachedFrame()
 		return;
 	}
 
-	FoundFloatComponents = 0;
-	FoundHalfComponents = 0;
-	FoundInt32Components = 0;
-
 	NumInstances = EmitterIndex == INDEX_NONE ? 1 : SimCache->GetEmitterNumInstances(EmitterIndex, FrameIndex);
 	const FName EmitterName = EmitterIndex == INDEX_NONE ? NAME_None : SimCache->GetEmitterName(EmitterIndex);
+	// Cached frame data needs to match the component info being viewed. Update here so they match.
+	ComponentInfos = EmitterIndex == INDEX_NONE ? &SystemComponentInfos : &EmitterComponentInfos[EmitterIndex];
 
 	SimCache->ForEachEmitterAttribute(EmitterIndex,
+		[&](const FNiagaraSimCacheVariable& Variable)
+		{
+			// Pull in data
+			SimCache->ReadAttribute(FloatComponents, HalfComponents, Int32Components, Variable.Variable.GetName(), EmitterName, FrameIndex);
+
+			return true;
+		}
+	);
+}
+
+void FNiagaraSimCacheViewModel::UpdateComponentInfos()
+{
+	SystemComponentInfos.Empty();
+	EmitterComponentInfos.Empty();
+	FoundFloatComponents = 0;
+	FoundHalfComponents = 0;
+	FoundInt32Components = 0;
+	
+	UNiagaraSimCache* SimCache = WeakSimCache.Get();
+	if (SimCache == nullptr)
+	{
+		return;
+	}
+	
+	
+	SimCache->ForEachEmitterAttribute(INDEX_NONE,
 		[&](const FNiagaraSimCacheVariable& Variable)
 		{
 			// Build component info
 			const FNiagaraTypeDefinition& TypeDef = Variable.Variable.GetType();
 			if (TypeDef.IsEnum())
 			{
-				FComponentInfo& ComponentInfo = ComponentInfos.AddDefaulted_GetRef();
+				FComponentInfo& ComponentInfo = SystemComponentInfos.AddDefaulted_GetRef();
 				ComponentInfo.Name = Variable.Variable.GetName();
 				ComponentInfo.ComponentOffset = FoundInt32Components++;
 				ComponentInfo.bIsInt32 = true;
@@ -229,18 +283,179 @@ void FNiagaraSimCacheViewModel::UpdateCachedFrame()
 			}
 			else
 			{
-				BuildComponentInfos(Variable.Variable.GetName(), TypeDef.GetScriptStruct());
+				BuildComponentInfos(Variable.Variable.GetName(), TypeDef.GetScriptStruct(), SystemComponentInfos);
 			}
 
-			// Pull in data
-			SimCache->ReadAttribute(FloatComponents, HalfComponents, Int32Components, Variable.Variable.GetName(), EmitterName, FrameIndex);
-
 			return true;
 		}
 	);
+
+	for (int32 i = 0; i < SimCache->GetNumEmitters(); ++i)
+	{
+		TArray<FComponentInfo>& CurrentComponentInfos = EmitterComponentInfos.AddDefaulted_GetRef();
+
+		FoundFloatComponents = 0;
+		FoundHalfComponents = 0;
+		FoundInt32Components = 0;
+		
+		SimCache->ForEachEmitterAttribute(i,
+		[&](const FNiagaraSimCacheVariable& Variable)
+			{
+				// Build component info
+				const FNiagaraTypeDefinition& TypeDef = Variable.Variable.GetType();
+				if (TypeDef.IsEnum())
+				{
+					FComponentInfo& ComponentInfo = CurrentComponentInfos.AddDefaulted_GetRef();
+					ComponentInfo.Name = Variable.Variable.GetName();
+					ComponentInfo.ComponentOffset = FoundInt32Components++;
+					ComponentInfo.bIsInt32 = true;
+					ComponentInfo.Enum = TypeDef.GetEnum();
+				}
+				else
+				{
+					BuildComponentInfos(Variable.Variable.GetName(), TypeDef.GetScriptStruct(), CurrentComponentInfos);
+				}
+
+				return true;
+			}
+		);
+	}
+}
+
+void FNiagaraSimCacheViewModel::BuildTreeItemChildren(TSharedPtr<FNiagaraSimCacheTreeItem> InTreeItem, TWeakPtr<SNiagaraSimCacheTreeView> OwningTreeView)
+{
+	FNiagaraSimCacheTreeItem* TreeItem = InTreeItem.Get();
+	UNiagaraSimCache* SimCache = WeakSimCache.Get();
+
+	if(TreeItem && SimCache)
+	{
+		int32 BufferIndex = TreeItem->GetBufferIndex();
+		
+		SimCache->ForEachEmitterAttribute(BufferIndex,
+	[&](const FNiagaraSimCacheVariable& Variable)
+			{
+				FNiagaraTypeDefinition TypeDef = Variable.Variable.GetType();
+			
+
+				TSharedRef<FNiagaraSimCacheComponentTreeItem> CurrentItem = MakeShared<FNiagaraSimCacheComponentTreeItem>(OwningTreeView);
+				
+				CurrentItem->SetDisplayName(FText::FromName(Variable.Variable.GetName()));
+				CurrentItem->SetFilterName(Variable.Variable.GetName().ToString());
+				CurrentItem->TypeDef = TypeDef;
+				CurrentItem->BufferIndex = TreeItem->GetBufferIndex();
+				//CurrentItem->RootItem = TreeItem;
+
+				TreeItem->AddChild(CurrentItem);
+				
+				if(!TypeDef.IsEnum() && !FNiagaraTypeDefinition::IsScalarDefinition(TypeDef))
+				{
+					RecursiveBuildTreeItemChildren(TreeItem, CurrentItem, TypeDef, OwningTreeView);
+				}
+				return true;
+			}
+		);
+		
+	}
+}
+
+void FNiagaraSimCacheViewModel::RecursiveBuildTreeItemChildren(FNiagaraSimCacheTreeItem* Root,
+	TSharedRef<FNiagaraSimCacheComponentTreeItem> Parent, FNiagaraTypeDefinition TypeDefinition, TWeakPtr<SNiagaraSimCacheTreeView> OwningTreeView)
+{
+	UScriptStruct* Struct = TypeDefinition.GetScriptStruct();
+
+	for (TFieldIterator<FProperty> PropertyIt(Struct, EFieldIteratorFlags::IncludeSuper); PropertyIt; ++PropertyIt)
+	{
+		FProperty* Property = *PropertyIt;
+		TSharedRef<FNiagaraSimCacheComponentTreeItem> CurrentItem = MakeShared<FNiagaraSimCacheComponentTreeItem> (OwningTreeView);
+		
+		FString PropertyName = Property->GetName();
+		
+		CurrentItem->SetDisplayName(FText::FromString(PropertyName));
+		CurrentItem->SetFilterName(Parent->GetFilterName().Append(".").Append(PropertyName));
+		CurrentItem->SetBufferIndex(Root->GetBufferIndex());
+		//CurrentItem->RootItem = Root;
+
+		Parent->AddChild(CurrentItem);
+
+		if(Property->IsA(FStructProperty::StaticClass()))
+		{
+			const FStructProperty* StructProperty = CastFieldChecked<FStructProperty>(Property);
+			UScriptStruct* FriendlyStruct = FNiagaraTypeHelper::FindNiagaraFriendlyTopLevelStruct(StructProperty->Struct, ENiagaraStructConversion::Simulation);
+			FNiagaraTypeDefinition StructTypeDef = FNiagaraTypeDefinition(FriendlyStruct);
+			CurrentItem->TypeDef = StructTypeDef;
+			RecursiveBuildTreeItemChildren(Root, CurrentItem, StructTypeDef, OwningTreeView);
+		}
+		else if (Property->IsA(FNumericProperty::StaticClass()))
+		{
+			if(Property->IsA(FIntProperty::StaticClass()))
+			{
+				CurrentItem->TypeDef = FNiagaraTypeDefinition::GetIntDef();
+			}
+			else if (Property->IsA(FFloatProperty::StaticClass()))
+			{
+				CurrentItem->TypeDef = FNiagaraTypeDefinition::GetFloatDef();
+			}
+		}
+		else if (Property->IsA(FBoolProperty::StaticClass()))
+		{
+			CurrentItem->TypeDef = FNiagaraTypeDefinition::GetBoolDef();
+		}
+	}
+}
+
+void FNiagaraSimCacheViewModel::BuildEntries(TWeakPtr<SNiagaraSimCacheTreeView> OwningTreeView)
+{
+	RootEntries.Empty();
+	BufferEntries.Empty();
+	
+	const TSharedRef<FNiagaraSimCacheTreeItem> SharedSystemTreeItem = MakeShared<FNiagaraSimCacheTreeItem>(OwningTreeView);
+	const TSharedRef<FNiagaraSimCacheOverviewSystemItem> SharedSystemBufferItem = MakeShared<FNiagaraSimCacheOverviewSystemItem>();
+
+	SharedSystemTreeItem->SetDisplayName(LOCTEXT("SystemInstance", "System Instance"));
+	SharedSystemBufferItem->SetDisplayName(LOCTEXT("SystemInstance", "System Instance"));
+
+	RootEntries.Add(SharedSystemTreeItem);
+	BufferEntries.Add(SharedSystemBufferItem);
+		
+	BuildTreeItemChildren(SharedSystemTreeItem, OwningTreeView);
+		
+	for(int32 i = 0; i < GetNumEmitterLayouts(); i++)
+	{
+		const TSharedRef<FNiagaraSimCacheEmitterTreeItem> CurrentEmitterItem = MakeShared<FNiagaraSimCacheEmitterTreeItem>(OwningTreeView);
+		const TSharedRef<FNiagaraSimCacheOverviewEmitterItem> CurrentEmitterBufferItem = MakeShared<FNiagaraSimCacheOverviewEmitterItem>();
+		
+		CurrentEmitterItem->SetDisplayName(FText::FromName(GetEmitterLayoutName(i)));
+		CurrentEmitterBufferItem->SetDisplayName(FText::FromName(GetEmitterLayoutName(i)));
+		
+		CurrentEmitterItem->SetBufferIndex(i);
+		CurrentEmitterBufferItem->SetBufferIndex(i);
+
+		RootEntries.Add(CurrentEmitterItem);
+		BufferEntries.Add(CurrentEmitterBufferItem);
+
+		BuildTreeItemChildren(CurrentEmitterItem, OwningTreeView);
+	}
+
+	UpdateCurrentEntries();
+}
+
+void FNiagaraSimCacheViewModel::UpdateCurrentEntries()
+{
+	CurrentRootEntries.Empty();
+	CurrentRootEntries.Add(RootEntries[EmitterIndex + 1]);
+}
+
+TArray<TSharedRef<FNiagaraSimCacheTreeItem>>* FNiagaraSimCacheViewModel::GetCurrentRootEntries()
+{
+	return &CurrentRootEntries;
+}
+
+TArray<TSharedRef<FNiagaraSimCacheOverviewItem>>* FNiagaraSimCacheViewModel::GetBufferEntries()
+{
+	return &BufferEntries;
 }
 
-void FNiagaraSimCacheViewModel::BuildComponentInfos(const FName Name, const UScriptStruct* Struct)
+void FNiagaraSimCacheViewModel::BuildComponentInfos(const FName Name, const UScriptStruct* Struct, TArray<FComponentInfo>& InComponentInfos)
 {
 	int32 NumProperties = 0;
 	for (TFieldIterator<FProperty> PropertyIt(Struct, EFieldIteratorFlags::IncludeSuper); PropertyIt; ++PropertyIt)
@@ -254,21 +469,21 @@ void FNiagaraSimCacheViewModel::BuildComponentInfos(const FName Name, const UScr
 		const FName PropertyName = NumProperties > 1 ? FName(*FString::Printf(TEXT("%s.%s"), *Name.ToString(), *Property->GetName())) : Name;
 		if (Property->IsA(FFloatProperty::StaticClass()))
 		{
-			FComponentInfo& ComponentInfo = ComponentInfos.AddDefaulted_GetRef();
+			FComponentInfo& ComponentInfo = InComponentInfos.AddDefaulted_GetRef();
 			ComponentInfo.Name = PropertyName;
 			ComponentInfo.ComponentOffset = FoundFloatComponents++;
 			ComponentInfo.bIsFloat = true;
 		}
 		else if (Property->IsA(FUInt16Property::StaticClass()))
 		{
-			FComponentInfo& ComponentInfo = ComponentInfos.AddDefaulted_GetRef();
+			FComponentInfo& ComponentInfo = InComponentInfos.AddDefaulted_GetRef();
 			ComponentInfo.Name = PropertyName;
 			ComponentInfo.ComponentOffset = FoundHalfComponents++;
 			ComponentInfo.bIsHalf = true;
 		}
 		else if (Property->IsA(FIntProperty::StaticClass()))
 		{
-			FComponentInfo& ComponentInfo = ComponentInfos.AddDefaulted_GetRef();
+			FComponentInfo& ComponentInfo = InComponentInfos.AddDefaulted_GetRef();
 			ComponentInfo.Name = PropertyName;
 			ComponentInfo.ComponentOffset = FoundInt32Components++;
 			ComponentInfo.bIsInt32 = true;
@@ -276,7 +491,7 @@ void FNiagaraSimCacheViewModel::BuildComponentInfos(const FName Name, const UScr
 		}
 		else if (Property->IsA(FBoolProperty::StaticClass()))
 		{
-			FComponentInfo& ComponentInfo = ComponentInfos.AddDefaulted_GetRef();
+			FComponentInfo& ComponentInfo = InComponentInfos.AddDefaulted_GetRef();
 			ComponentInfo.Name = PropertyName;
 			ComponentInfo.ComponentOffset = FoundInt32Components++;
 			ComponentInfo.bIsInt32 = true;
@@ -284,7 +499,7 @@ void FNiagaraSimCacheViewModel::BuildComponentInfos(const FName Name, const UScr
 		}
 		else if (Property->IsA(FEnumProperty::StaticClass()))
 		{
-			FComponentInfo& ComponentInfo = ComponentInfos.AddDefaulted_GetRef();
+			FComponentInfo& ComponentInfo = InComponentInfos.AddDefaulted_GetRef();
 			ComponentInfo.Name = PropertyName;
 			ComponentInfo.ComponentOffset = FoundInt32Components++;
 			ComponentInfo.bIsInt32 = true;
@@ -293,7 +508,7 @@ void FNiagaraSimCacheViewModel::BuildComponentInfos(const FName Name, const UScr
 		else if (Property->IsA(FStructProperty::StaticClass()))
 		{
 			const FStructProperty* StructProperty = CastFieldChecked<FStructProperty>(Property);
-			BuildComponentInfos(PropertyName, FNiagaraTypeHelper::FindNiagaraFriendlyTopLevelStruct(StructProperty->Struct, ENiagaraStructConversion::Simulation));
+			BuildComponentInfos(PropertyName, FNiagaraTypeHelper::FindNiagaraFriendlyTopLevelStruct(StructProperty->Struct, ENiagaraStructConversion::Simulation), InComponentInfos);
 		}
 		else
 		{
diff --git a/Engine/Plugins/FX/Niagara/Source/NiagaraEditor/Private/ViewModels/NiagaraSimCacheViewModel.h b/Engine/Plugins/FX/Niagara/Source/NiagaraEditor/Private/ViewModels/NiagaraSimCacheViewModel.h
--- a/Engine/Plugins/FX/Niagara/Source/NiagaraEditor/Private/ViewModels/NiagaraSimCacheViewModel.h
+++ b/Engine/Plugins/FX/Niagara/Source/NiagaraEditor/Private/ViewModels/NiagaraSimCacheViewModel.h
@@ -4,7 +4,12 @@
 
 #include "CoreMinimal.h"
 #include "NiagaraSimCache.h"
+#include "SNiagaraSimCacheOverview.h"
 
+class SNiagaraSimCacheTreeView;
+struct FNiagaraSimCacheComponentTreeItem;
+struct FNiagaraSimCacheTreeItem;
+struct FNiagaraSimCacheOverviewItem;
 struct FNiagaraSystemSimCacheCaptureReply;
 struct FNiagaraSimCacheDataBuffersLayout;
 class UNiagaraComponent;
@@ -24,6 +29,8 @@ public:
 	};
 
 	DECLARE_MULTICAST_DELEGATE_OneParam(FOnViewDataChanged, bool)
+	DECLARE_MULTICAST_DELEGATE(FOnSimCacheChanged)
+	DECLARE_MULTICAST_DELEGATE(FOnBufferChanged)
 	
 	FNiagaraSimCacheViewModel();
 	~FNiagaraSimCacheViewModel();
@@ -34,13 +41,28 @@ public:
 
 	void SetupPreviewComponentAndInstance();
 
-	TConstArrayView<FComponentInfo> GetComponentInfos() const { return ComponentInfos; }
+	TConstArrayView<FComponentInfo> GetCurrentComponentInfos() const { return *ComponentInfos; }
 
+	TConstArrayView<FComponentInfo> GetComponentInfos(int32 InEmitterIndex) const { return  InEmitterIndex == INDEX_NONE ? SystemComponentInfos : EmitterComponentInfos[InEmitterIndex]; }
+
+	TConstArrayView<FString> GetComponentFilters() const { return ComponentFilterArray; }
+
+	void SetComponentFilters(const TArray<FString>& NewComponentFilterArray)
+	{
+		ComponentFilterArray.Empty();
+		ComponentFilterArray.Append(NewComponentFilterArray);
+		OnViewDataChangedDelegate.Broadcast(true);
+	}
+
+	bool IsComponentFilterActive() const { return bComponentFilterActive; }
+
+	void SetComponentFilterActive(bool bNewActive);
+	
 	int32 GetNumInstances() const { return NumInstances; }
 
 	int32 GetNumFrames() const;
 
-	int32 GetFrameIndex() const { return FrameIndex; };
+	int32 GetFrameIndex() const { return FrameIndex; }
 
 	void SetFrameIndex(const int32 InFrameIndex);
 
@@ -58,14 +80,36 @@ public:
 
 	FOnViewDataChanged& OnViewDataChanged();
 
+	FOnSimCacheChanged& OnSimCacheChanged();
+
+	FOnBufferChanged& OnBufferChanged();
+
 	void OnCacheModified(UNiagaraSimCache* SimCache);
 
 	void UpdateCachedFrame();
 
+	// Build out component infos for each buffer in this cache
+	void UpdateComponentInfos();
+
+	void BuildTreeItemChildren(TSharedPtr<FNiagaraSimCacheTreeItem> InTreeItem, TWeakPtr<SNiagaraSimCacheTreeView> OwningTreeView);
+
+	void RecursiveBuildTreeItemChildren(FNiagaraSimCacheTreeItem* Root, TSharedRef<FNiagaraSimCacheComponentTreeItem> Parent,
+	FNiagaraTypeDefinition TypeDefinition, TWeakPtr<SNiagaraSimCacheTreeView> OwningTreeView);
+	
+	// Construct entries for the tree view.
+	void BuildEntries(TWeakPtr<SNiagaraSimCacheTreeView> OwningTreeView);
+
+	void UpdateCurrentEntries();
+
+	TArray<TSharedRef<FNiagaraSimCacheTreeItem>>* GetCurrentRootEntries();
+
+
+	TArray<TSharedRef<FNiagaraSimCacheOverviewItem>>* GetBufferEntries();
+
 	UNiagaraComponent* GetPreviewComponent() { return PreviewComponent; }
 
 private:
-	void BuildComponentInfos(const FName Name, const UScriptStruct* Struct);
+	void BuildComponentInfos(const FName Name, const UScriptStruct* Struct, TArray<FComponentInfo>& ComponentInfos);
 	
 	// The sim cache being viewed
 	TWeakObjectPtr<UNiagaraSimCache>	WeakSimCache;
@@ -82,8 +126,23 @@ private:
 	// Number of particles in the given frame
 	int32 NumInstances = 0;
 
-	// Info about the variables in the cache
-	TArray<FComponentInfo>				ComponentInfos;
+	// Info about the variables in the cache that are currently being viewed in the spreadsheet
+	TArray<FComponentInfo>* ComponentInfos = &SystemComponentInfos;
+
+	// Cached Component infos for this system
+	TArray<FComponentInfo>				SystemComponentInfos;
+	TArray<TArray<FComponentInfo>>		EmitterComponentInfos;
+	// TODO: Component infos for Data Interfaces?
+
+	TArray<TSharedRef<FNiagaraSimCacheTreeItem>> RootEntries;
+	TArray<TSharedRef<FNiagaraSimCacheTreeItem>> CurrentRootEntries;
+
+	TArray<TSharedRef<FNiagaraSimCacheOverviewItem>> BufferEntries;
+
+	// String filters gathered from the overview tree
+	TArray<FString>						ComponentFilterArray;
+	bool								bComponentFilterActive = false;
+	
 	TArray<float>						FloatComponents;
 	TArray<FFloat16>					HalfComponents;
 	TArray<int32>						Int32Components;
@@ -94,4 +153,6 @@ private:
 
 	bool								bDelegatesAdded = false;
 	FOnViewDataChanged					OnViewDataChangedDelegate;
+	FOnSimCacheChanged					OnSimCacheChangedDelegate;
+	FOnBufferChanged					OnBufferChangedDelegate;
 };
diff --git a/Engine/Plugins/FX/Niagara/Source/NiagaraEditor/Private/Widgets/SNiagaraSimCacheOverview.cpp b/Engine/Plugins/FX/Niagara/Source/NiagaraEditor/Private/Widgets/SNiagaraSimCacheOverview.cpp
new file mode 100644
--- /dev/null
+++ b/Engine/Plugins/FX/Niagara/Source/NiagaraEditor/Private/Widgets/SNiagaraSimCacheOverview.cpp
@@ -0,0 +1,141 @@
+﻿#include "SNiagaraSimCacheOverview.h"
+
+#include "SNiagaraSimCacheTreeView.h"
+
+#define LOCTEXT_NAMESPACE "NiagaraSimCacheOverview"
+
+class SNiagaraSimCacheBufferItem : public SCompoundWidget
+{
+public:
+	SLATE_BEGIN_ARGS(SNiagaraSimCacheBufferItem) {}
+		SLATE_ARGUMENT(TSharedPtr<FNiagaraSimCacheOverviewItem>, Item)
+	SLATE_END_ARGS()
+
+	void Construct(const FArguments& InArgs)
+	{
+		Item = InArgs._Item;
+
+		RefreshContent();
+	}
+
+	void RefreshContent()
+	{
+		ChildSlot
+		.Padding(2.0f)
+		[
+			Item->GetRowWidget()
+		];
+	}
+
+	TSharedPtr<FNiagaraSimCacheOverviewItem> Item;
+};
+
+
+void SNiagaraSimCacheOverview::Construct(const FArguments& InArgs)
+{
+	ViewModel = InArgs._SimCacheViewModel;
+
+	SAssignNew(TreeViewWidget, SNiagaraSimCacheTreeView)
+	.SimCacheViewModel(ViewModel);
+
+	SAssignNew(BufferListView, SListView<TSharedRef<FNiagaraSimCacheOverviewItem>>)
+	.ListItemsSource(ViewModel->GetBufferEntries())
+	.OnGenerateRow(this, &SNiagaraSimCacheOverview::OnGenerateRowForItem)
+	.OnSelectionChanged(this, &SNiagaraSimCacheOverview::OnListSelectionChanged);
+
+	const float MinSplitterSlotSize = 30.0f; 
+
+	ChildSlot
+	[
+		SNew(SSplitter)
+		.Orientation(Orient_Vertical)
+		+SSplitter::Slot()
+		.Value(0.2f)
+		.MinSize(MinSplitterSlotSize)
+		[
+
+			// Cache Buffers
+			SNew(SSplitter)
+			.Orientation(Orient_Vertical)
+			+SSplitter::Slot()
+			.Resizable(false)
+			.SizeRule(SSplitter::SizeToContent)
+			[
+				// Header
+				SNew(SBorder)
+				.BorderImage(FAppStyle::GetNoBrush())
+				.Padding(5.0f)
+				[
+					SNew(STextBlock)
+					.Text(LOCTEXT("CacheBufferSelection", "Cache Buffer Selection"))
+				]
+			]
+			+SSplitter::Slot()
+			.Resizable(false)
+			[
+				// List View
+				BufferListView.ToSharedRef()
+			]
+		]
+		+SSplitter::Slot()
+		.Value(0.8f)
+		.MinSize(MinSplitterSlotSize)
+		[
+			// Component Details
+			SNew(SSplitter)
+			.Orientation(Orient_Vertical)
+			+SSplitter::Slot()
+			.Resizable(false)
+			.SizeRule(SSplitter::SizeToContent)
+			[
+				// Header
+				SNew(SBorder)
+				.BorderImage(FAppStyle::GetNoBrush())
+				.Padding(5.0f)
+				[
+					SNew(STextBlock)
+					.Text(LOCTEXT("ComponentTreeView", "Component Details"))
+				]
+			]
+			+SSplitter::Slot()
+			.Resizable(false)
+			[
+				// Tree View
+				TreeViewWidget.ToSharedRef()
+			]
+			
+		]
+	];
+}
+
+TSharedRef<ITableRow> SNiagaraSimCacheOverview::OnGenerateRowForItem(TSharedRef<FNiagaraSimCacheOverviewItem> Item,
+	const TSharedRef<STableViewBase>& Owner)
+{
+	static const char* ItemStyles[] =
+	{
+		"NiagaraEditor.SimCache.SystemItem",
+		"NiagaraEditor.SimCache.EmitterItem",
+		"NiagaraEditor.SimCache.ComponentItem"
+	};
+
+	ENiagaraSimCacheOverviewItemType StyleType = Item->GetType();
+	
+	return SNew(STableRow<TSharedRef<FNiagaraSimCacheOverviewItem>>, Owner)
+	.Style(FNiagaraEditorStyle::Get(), ItemStyles[(int32)StyleType])
+	.Padding(1.0f)
+	[
+		SNew(SNiagaraSimCacheBufferItem)
+		.Item(Item)
+	];
+}
+
+void SNiagaraSimCacheOverview::OnListSelectionChanged(TSharedPtr<FNiagaraSimCacheOverviewItem> Item,
+	ESelectInfo::Type SelectInfo)
+{
+	if(Item.IsValid())
+	{
+		ViewModel->SetEmitterIndex(Item->GetBufferIndex());
+	}
+}
+
+#undef LOCTEXT_NAMESPACE
\ No newline at end of file
diff --git a/Engine/Plugins/FX/Niagara/Source/NiagaraEditor/Private/Widgets/SNiagaraSimCacheOverview.h b/Engine/Plugins/FX/Niagara/Source/NiagaraEditor/Private/Widgets/SNiagaraSimCacheOverview.h
new file mode 100644
--- /dev/null
+++ b/Engine/Plugins/FX/Niagara/Source/NiagaraEditor/Private/Widgets/SNiagaraSimCacheOverview.h
@@ -0,0 +1,90 @@
+﻿// // Copyright Epic Games, Inc. All Rights Reserved.
+
+#pragma once
+
+#include "Widgets/SCompoundWidget.h"
+#include "Widgets/SWidget.h"
+#include "Widgets/Text/STextBlock.h"
+#include "Widgets/Views/ITableRow.h"
+#include "Widgets/Views/SListView.h"
+#include "Templates/SharedPointer.h"
+
+class FNiagaraSimCacheViewModel;
+class SNiagaraSimCacheTreeView;
+
+enum class ENiagaraSimCacheOverviewItemType
+{
+	System,
+	Emitter,
+	Component,
+	MAX,
+};
+
+struct FNiagaraSimCacheOverviewItem : public TSharedFromThis<FNiagaraSimCacheOverviewItem>
+{
+	virtual ~FNiagaraSimCacheOverviewItem() = default;
+
+	FNiagaraSimCacheOverviewItem()
+	{
+		
+	}
+
+	FText GetDisplayNameText()
+	{
+		return DisplayName;
+	}
+
+	void SetDisplayName(FText NewName)
+	{
+		DisplayName = NewName;
+	}
+
+	virtual int32 GetBufferIndex()
+	{
+		return BufferIndex;
+	}
+
+	void SetBufferIndex(int32 NewIndex)
+	{
+		BufferIndex = NewIndex;
+	}
+
+	virtual ENiagaraSimCacheOverviewItemType GetType () { return ENiagaraSimCacheOverviewItemType::MAX; }
+
+	virtual TSharedRef<SWidget> GetRowWidget()
+	{
+		return SNew(STextBlock).Text(GetDisplayNameText());
+	}
+
+	int32 BufferIndex = INDEX_NONE;
+	FText DisplayName;
+};
+
+struct FNiagaraSimCacheOverviewSystemItem : public FNiagaraSimCacheOverviewItem
+{
+	virtual ENiagaraSimCacheOverviewItemType GetType () override { return ENiagaraSimCacheOverviewItemType::System; }
+};
+
+struct FNiagaraSimCacheOverviewEmitterItem : public FNiagaraSimCacheOverviewItem
+{
+	virtual ENiagaraSimCacheOverviewItemType GetType() override {return ENiagaraSimCacheOverviewItemType::Emitter; }
+};
+
+class SNiagaraSimCacheOverview : public SCompoundWidget
+{
+public:
+	SLATE_BEGIN_ARGS(SNiagaraSimCacheOverview) {}
+		SLATE_ARGUMENT(TSharedPtr<FNiagaraSimCacheViewModel>, SimCacheViewModel)
+	SLATE_END_ARGS()
+
+	void Construct(const FArguments& InArgs);
+
+	TSharedRef<ITableRow> OnGenerateRowForItem(TSharedRef<FNiagaraSimCacheOverviewItem> Item, const TSharedRef<STableViewBase>& Owner);
+	void OnListSelectionChanged(TSharedPtr<FNiagaraSimCacheOverviewItem> Item, ESelectInfo::Type SelectInfo);
+
+	TSharedPtr<SListView<TSharedRef<FNiagaraSimCacheOverviewItem>>> BufferListView;
+
+	TSharedPtr<SNiagaraSimCacheTreeView> TreeViewWidget;
+
+	TSharedPtr<FNiagaraSimCacheViewModel> ViewModel;
+};
diff --git a/Engine/Plugins/FX/Niagara/Source/NiagaraEditor/Private/Widgets/SNiagaraSimCacheTreeView.cpp b/Engine/Plugins/FX/Niagara/Source/NiagaraEditor/Private/Widgets/SNiagaraSimCacheTreeView.cpp
new file mode 100644
--- /dev/null
+++ b/Engine/Plugins/FX/Niagara/Source/NiagaraEditor/Private/Widgets/SNiagaraSimCacheTreeView.cpp
@@ -0,0 +1,393 @@
+﻿// Copyright Epic Games, Inc. All Rights Reserved.
+
+#include "SNiagaraSimCacheTreeView.h"
+
+#include "NiagaraSimCacheViewModel.h"
+#include "Widgets/Views/STreeView.h"
+#include "Widgets/Input/SButton.h"
+#include "Widgets/Input/SCheckBox.h"
+
+#define LOCTEXT_NAMESPACE "NiagaraSimCacheTreeView"
+
+class SNiagaraSimCacheTreeItem : public SCompoundWidget
+{
+public:
+	SLATE_BEGIN_ARGS(SNiagaraSimCacheTreeItem) {}
+		SLATE_ARGUMENT(TSharedPtr<FNiagaraSimCacheTreeItem>, Item)
+		SLATE_ARGUMENT(TWeakPtr<SNiagaraSimCacheTreeView>, Owner)
+	SLATE_END_ARGS()
+
+	void Construct(const FArguments& InArgs)
+	{
+		Item = InArgs._Item;
+		Owner = InArgs._Owner;
+	
+		RefreshContent();
+	}
+
+	void RefreshContent() 
+	{
+		ChildSlot
+		.Padding(2.0f)
+		[
+			Item->GetRowWidget()
+		];
+	}
+
+	TSharedPtr<FNiagaraSimCacheTreeItem> Item;
+	TWeakPtr<SNiagaraSimCacheTreeView> Owner;
+};
+
+//// Filter Widget /////
+
+
+void SNiagaraSimCacheTreeViewFilterWidget::Construct(const FArguments& InArgs, TWeakPtr<FNiagaraSimCacheTreeItem> InTreeItem,
+	TWeakPtr<SNiagaraSimCacheTreeView> InTreeView)
+{
+	WeakTreeItem = InTreeItem;
+	WeakTreeView = InTreeView;
+
+	ChildSlot
+	.HAlign(HAlign_Center)
+	.Padding(3.0f)
+	[
+		// Filter controls
+		SNew(SHorizontalBox)
+		+SHorizontalBox::Slot()
+		.AutoWidth()
+		.Padding(2.0f)
+		[
+			// Clear All
+			SNew(SButton)
+			.Text(LOCTEXT("ClearAll", "Clear All"))
+			.OnClicked(this, &SNiagaraSimCacheTreeViewFilterWidget::OnClearAllReleased)
+			.IsEnabled(this, &SNiagaraSimCacheTreeViewFilterWidget::IsFilterActive)
+		]
+		+SHorizontalBox::Slot()
+		.AutoWidth()
+		.Padding(2.0f)
+		[
+			// Select All
+			SNew(SButton)
+			.Text(LOCTEXT("SelectAll", "Select All"))
+			.OnClicked(this, &SNiagaraSimCacheTreeViewFilterWidget::OnSelectAllReleased)
+			.IsEnabled(this, &SNiagaraSimCacheTreeViewFilterWidget::IsFilterActive)
+		]
+		+SHorizontalBox::Slot()
+		.AutoWidth()
+		.Padding(2.0f)
+		[
+
+			// Filter Toggle
+			SNew(SCheckBox)
+			.Type(ESlateCheckBoxType::ToggleButton)
+			.HAlign(HAlign_Right)
+			//.ToolTip(LOCTEXT("ToggleFiltering", "Toggle Filtering"))
+			.Style(FNiagaraEditorStyle::Get(),"NiagaraEditor.SimCache.FilterToggleStyle")
+			.OnCheckStateChanged(this, &SNiagaraSimCacheTreeViewFilterWidget::OnCheckStateChanged)
+			.IsChecked(this, &SNiagaraSimCacheTreeViewFilterWidget::GetFilterState)
+			.Padding(FMargin(3.0f, 2.0f))
+			[
+					
+				SNew(SImage)
+				.Image(FAppStyle::GetBrush("Icons.Filter"))
+					
+			]
+		]
+	];
+}
+
+void SNiagaraSimCacheTreeViewFilterWidget::OnCheckStateChanged(ECheckBoxState InState)
+{
+	TSharedPtr<SNiagaraSimCacheTreeView> TreeView = WeakTreeView.Pin();
+
+	if(TreeView.IsValid())
+	{
+		const bool bCurrentlyActive = InState == ECheckBoxState::Checked;
+		TreeView->SetFilterActive(bCurrentlyActive);
+	}
+	
+}
+
+FReply SNiagaraSimCacheTreeViewFilterWidget::OnClearAllReleased()
+{
+	TSharedPtr<SNiagaraSimCacheTreeView> TreeView = WeakTreeView.Pin();
+
+	if(TreeView.IsValid())
+	{
+		TreeView->ClearFilterSelection();
+		return FReply::Handled();
+	}
+
+	return FReply::Unhandled();
+}
+
+FReply SNiagaraSimCacheTreeViewFilterWidget::OnSelectAllReleased()
+{
+	TSharedPtr<SNiagaraSimCacheTreeView> TreeView = WeakTreeView.Pin();
+
+	if(TreeView.IsValid())
+	{
+		TreeView->SelectAll();
+		return FReply::Handled();
+	}
+
+	return FReply::Unhandled();
+}
+
+ECheckBoxState SNiagaraSimCacheTreeViewFilterWidget::GetFilterState() const
+{
+	return IsFilterActive() ? ECheckBoxState::Checked : ECheckBoxState::Unchecked;
+}
+
+bool SNiagaraSimCacheTreeViewFilterWidget::IsFilterActive() const
+{
+	TSharedPtr<SNiagaraSimCacheTreeView> TreeView = WeakTreeView.Pin();
+
+	if(TreeView.IsValid())
+	{
+		return TreeView->IsFilterActive();
+	}
+
+	return false;
+}
+
+//// Visibility Widget //////
+
+void SSimCacheTreeViewVisibilityWidget::Construct(const FArguments& InArgs, TWeakPtr<FNiagaraSimCacheTreeItem> InTreeItem, TWeakPtr<SNiagaraSimCacheTreeView> InWeakTreeView)
+{
+	WeakTreeItem = InTreeItem;
+	WeakTreeView = InWeakTreeView;
+
+	ChildSlot
+	.HAlign(HAlign_Center)
+	.Padding(1.0f)
+	[
+		SNew(SCheckBox)
+		.OnCheckStateChanged(this, &SSimCacheTreeViewVisibilityWidget::OnCheckStateChanged)
+		.IsChecked(this, &SSimCacheTreeViewVisibilityWidget::GetCheckedState)
+		.IsEnabled(this, &SSimCacheTreeViewVisibilityWidget::IsFilterActive)
+	];
+	
+}
+	
+void SSimCacheTreeViewVisibilityWidget::OnCheckStateChanged(ECheckBoxState NewState)
+{
+	TSharedPtr<SNiagaraSimCacheTreeView> TreeView = WeakTreeView.Pin();
+	TSharedPtr<FNiagaraSimCacheTreeItem> TreeItem = WeakTreeItem.Pin();
+	
+	if(TreeView.IsValid() && TreeItem.IsValid())
+	{
+		TreeView->VisibilityButtonClicked(TreeItem.ToSharedRef());
+	}
+}
+
+ECheckBoxState SSimCacheTreeViewVisibilityWidget::GetCheckedState() const
+{
+	return IsInFilter()? ECheckBoxState::Checked : ECheckBoxState::Unchecked;
+}
+
+bool SSimCacheTreeViewVisibilityWidget::IsInFilter() const
+{
+	TSharedPtr<SNiagaraSimCacheTreeView> TreeView = WeakTreeView.Pin();
+	TSharedPtr<FNiagaraSimCacheTreeItem> TreeItem = WeakTreeItem.Pin();
+
+	if(TreeView.IsValid() && TreeItem.IsValid())
+	{
+		return TreeView->IsItemInFilter(TreeItem);
+	}
+
+	return false;
+}
+
+bool SSimCacheTreeViewVisibilityWidget::IsFilterActive() const
+{
+	const TSharedPtr<SNiagaraSimCacheTreeView> TreeView = WeakTreeView.Pin();
+
+	if(TreeView.IsValid())
+	{
+		return TreeView->IsFilterActive();
+	}
+
+	return false;
+}
+
+bool SSimCacheTreeViewVisibilityWidget::IsItemSelected() const
+{
+	const TSharedPtr<SNiagaraSimCacheTreeView> TreeView = WeakTreeView.Pin();
+	const TSharedPtr<FNiagaraSimCacheTreeItem> TreeItem = WeakTreeItem.Pin();
+
+	if(TreeView.IsValid() && TreeItem.IsValid())
+	{
+		return TreeView->IsItemSelected(TreeItem.ToSharedRef());
+	}
+
+	return false;
+}
+
+///// Tree View Widget
+
+void SNiagaraSimCacheTreeView::Construct(const FArguments& InArgs)
+{
+	constexpr float ItemHeight = 50.0f;
+
+	ViewModel = InArgs._SimCacheViewModel;
+
+	ViewModel->OnBufferChanged().AddSP(this, &SNiagaraSimCacheTreeView::OnBufferChanged);
+
+	ViewModel->BuildEntries(SharedThis(this));
+
+	TArray<TSharedRef<FNiagaraSimCacheTreeItem>>* RootEntries = ViewModel->GetCurrentRootEntries();
+
+	
+	TreeView = SNew(STreeView<TSharedRef<FNiagaraSimCacheTreeItem>>)
+	.ItemHeight(ItemHeight)
+	.SelectionMode(ESelectionMode::Multi)
+	.TreeItemsSource(ViewModel->GetCurrentRootEntries())
+	.OnGenerateRow(this, &SNiagaraSimCacheTreeView::OnGenerateRow)
+	.OnGetChildren(this, &SNiagaraSimCacheTreeView::OnGetChildren);
+
+	if(RootEntries && !RootEntries->IsEmpty())
+	{
+		TreeView->SetItemExpansion((*RootEntries)[0], true);
+	}
+
+
+	ChildSlot
+	[
+		TreeView.ToSharedRef()
+	];
+}
+
+
+TSharedRef<ITableRow> SNiagaraSimCacheTreeView::OnGenerateRow(TSharedRef<FNiagaraSimCacheTreeItem> Item,
+                                                              const TSharedRef<STableViewBase>& OwnerTable)
+{
+	static const char* ItemStyles[] =
+	{
+		"NiagaraEditor.SimCache.SystemItem",
+		"NiagaraEditor.SimCache.EmitterItem",
+		"NiagaraEditor.SimCache.ComponentItem"
+	};
+
+	ENiagaraSimCacheOverviewItemType StyleType = Item->GetType();
+	
+	return SNew(STableRow<TSharedRef<FNiagaraSimCacheTreeItem>>, OwnerTable)
+	.Style(FNiagaraEditorStyle::Get(), ItemStyles[(int32)StyleType])
+	[
+		SNew(SNiagaraSimCacheTreeItem)
+		.Item(Item)
+		.Owner(SharedThis(this))
+	];
+}
+
+void SNiagaraSimCacheTreeView::OnGetChildren(TSharedRef<FNiagaraSimCacheTreeItem> InItem,
+	TArray<TSharedRef<FNiagaraSimCacheTreeItem>>& OutChildren)
+{
+	OutChildren = InItem.Get().Children;
+}
+
+void SNiagaraSimCacheTreeView::OnBufferChanged()
+{
+	TreeView->RequestTreeRefresh();
+	SelectionForFilter.Empty();
+	ViewModel->SetComponentFilters(TArray<FString>());
+	TreeView->SetItemExpansion((*ViewModel->GetCurrentRootEntries())[0], true);
+}
+
+void SNiagaraSimCacheTreeView::RecursiveAddToSelectionFilter(TArray<TSharedRef<FNiagaraSimCacheTreeItem>>& ArrayToAdd)
+{
+	for(const TSharedRef<FNiagaraSimCacheTreeItem>& Item : ArrayToAdd)
+	{
+		if(Item->Children.Num() > 0)
+		{
+			RecursiveAddToSelectionFilter(Item->Children);
+		}
+	}
+	SelectionForFilter.Append(ArrayToAdd);
+}
+
+void SNiagaraSimCacheTreeView::RecursiveRemoveFromSelectionFilter(TArray<TSharedRef<FNiagaraSimCacheTreeItem>>& ArrayToRemove)
+{
+	for(TSharedRef<FNiagaraSimCacheTreeItem> Item : ArrayToRemove)
+	{
+		if(Item->Children.Num() > 0)
+		{
+			RecursiveRemoveFromSelectionFilter(Item->Children);
+		}
+		SelectionForFilter.Remove(Item);
+	}
+}
+
+void SNiagaraSimCacheTreeView::UpdateSelectionFilter(TSharedRef<FNiagaraSimCacheTreeItem> ClickedItem)
+{
+	TArray<TSharedRef<FNiagaraSimCacheTreeItem>> SelectedItems = TreeView->GetSelectedItems();
+	SelectedItems.AddUnique(ClickedItem);
+
+	SelectedItems.RemoveAll([](const TSharedRef<FNiagaraSimCacheTreeItem> Item){ return Item->GetType() != ENiagaraSimCacheOverviewItemType::Component; });
+
+	if(IsItemInFilter(ClickedItem))
+	{
+		RecursiveRemoveFromSelectionFilter(SelectedItems);
+	}
+	else
+	{
+		RecursiveAddToSelectionFilter(SelectedItems);
+	}
+
+	UpdateStringFilters();
+}
+
+bool SNiagaraSimCacheTreeView::IsFilterActive() const
+{
+	return ViewModel->IsComponentFilterActive();
+}
+
+bool SNiagaraSimCacheTreeView::IsItemInFilter(TSharedPtr<FNiagaraSimCacheTreeItem> InItem) const
+{
+	return ViewModel->GetComponentFilters().Contains(InItem->GetFilterName());
+}
+
+void SNiagaraSimCacheTreeView::SetFilterActive(bool bNewActive)
+{
+	ViewModel->SetComponentFilterActive(bNewActive);
+}
+
+void SNiagaraSimCacheTreeView::VisibilityButtonClicked(TSharedRef<FNiagaraSimCacheTreeItem> InItem)
+{
+	UpdateSelectionFilter(InItem);
+}
+
+bool SNiagaraSimCacheTreeView::IsItemSelected(TSharedRef<FNiagaraSimCacheTreeItem> InItem)
+{
+	return TreeView->GetSelectedItems().Contains(InItem);
+}
+
+void SNiagaraSimCacheTreeView::ClearFilterSelection()
+{
+	SelectionForFilter.Empty();
+	ViewModel->SetComponentFilters(TArray<FString>());
+}
+
+void SNiagaraSimCacheTreeView::UpdateStringFilters()
+{
+	TArray<FString> StringFilters;
+	for(TSharedRef<FNiagaraSimCacheTreeItem> FilterItem : SelectionForFilter)
+	{
+		if(FilterItem->GetType() == ENiagaraSimCacheOverviewItemType::Component)
+		{
+			StringFilters.AddUnique(FilterItem->GetFilterName());
+		}
+	}
+
+	ViewModel->SetComponentFilters(StringFilters);
+}
+
+void SNiagaraSimCacheTreeView::SelectAll()
+{
+	RecursiveAddToSelectionFilter(*ViewModel->GetCurrentRootEntries());
+
+	UpdateStringFilters();
+}
+
+#undef LOCTEXT_NAMESPACE
diff --git a/Engine/Plugins/FX/Niagara/Source/NiagaraEditor/Private/Widgets/SNiagaraSimCacheTreeView.h b/Engine/Plugins/FX/Niagara/Source/NiagaraEditor/Private/Widgets/SNiagaraSimCacheTreeView.h
new file mode 100644
--- /dev/null
+++ b/Engine/Plugins/FX/Niagara/Source/NiagaraEditor/Private/Widgets/SNiagaraSimCacheTreeView.h
@@ -0,0 +1,238 @@
+﻿// Copyright Epic Games, Inc. All Rights Reserved.
+
+#pragma once
+
+#include "NiagaraSimCacheViewModel.h"
+
+#include "NiagaraEditorStyle.h"
+#include "Widgets/SNiagaraPinTypeSelector.h"
+#include "EdGraphSchema_Niagara.h"
+#include "SNiagaraSimCacheOverview.h"
+#include "Widgets/Images/SLayeredImage.h"
+#include "Widgets/Layout/SScaleBox.h"
+
+class SNiagaraSimCacheTreeView;
+class FNiagaraSimCacheViewModel;
+struct FNiagaraSimCacheTreeItem;
+struct FNiagaraSimCacheOverviewItem;
+
+///// Widget for controlling filters on the tree view
+class SNiagaraSimCacheTreeViewFilterWidget : public SCompoundWidget
+{
+public:
+	SLATE_BEGIN_ARGS(SNiagaraSimCacheTreeViewFilterWidget) {}
+	SLATE_END_ARGS()
+
+	
+	void Construct(const FArguments& InArgs, TWeakPtr<FNiagaraSimCacheTreeItem> InTreeItem, TWeakPtr<SNiagaraSimCacheTreeView> InTreeView);
+
+protected:
+	
+	void OnCheckStateChanged(ECheckBoxState InState);
+	FReply OnClearAllReleased();
+	FReply OnSelectAllReleased();
+	ECheckBoxState GetFilterState() const;
+	bool IsFilterActive() const;
+
+private:
+	TWeakPtr<FNiagaraSimCacheTreeItem> WeakTreeItem;
+	TWeakPtr<SNiagaraSimCacheTreeView> WeakTreeView;
+
+};
+
+///// Widget for managing visibility of a component in the tree view.
+class SSimCacheTreeViewVisibilityWidget : public SCompoundWidget
+{
+public:
+	SLATE_BEGIN_ARGS(SSimCacheTreeViewVisibilityWidget) {}
+	SLATE_END_ARGS()
+
+	void Construct(const FArguments& InArgs, TWeakPtr<FNiagaraSimCacheTreeItem> InTreeItem, TWeakPtr<SNiagaraSimCacheTreeView> InTreeView);
+
+protected:
+	void OnCheckStateChanged(ECheckBoxState NewState);
+	ECheckBoxState GetCheckedState() const;
+	
+	bool IsInFilter() const;
+	bool IsFilterActive() const;
+	bool IsItemSelected() const;
+private:
+	TWeakPtr<FNiagaraSimCacheTreeItem> WeakTreeItem;
+	TWeakPtr<SNiagaraSimCacheTreeView> WeakTreeView;
+};
+
+class SNiagaraSimCacheTreeView : public SCompoundWidget
+{
+public:
+	SLATE_BEGIN_ARGS(SNiagaraSimCacheTreeView) {}
+	SLATE_ARGUMENT(TSharedPtr<FNiagaraSimCacheViewModel>, SimCacheViewModel)
+SLATE_END_ARGS()
+
+	void Construct(const FArguments& InArgs);
+	bool HasSelectionFilter() const {return !SelectionForFilter.IsEmpty();}
+	bool IsFilterActive() const;
+	bool IsItemInFilter(TSharedPtr<FNiagaraSimCacheTreeItem> InItem) const;
+	void SetFilterActive(bool bNewActive);
+	void VisibilityButtonClicked(TSharedRef<FNiagaraSimCacheTreeItem> InItem);
+	bool IsItemSelected(TSharedRef<FNiagaraSimCacheTreeItem> InItem);
+	void ClearFilterSelection();
+	void UpdateStringFilters();
+	void SelectAll();
+private:
+	TSharedRef<ITableRow> OnGenerateRow(TSharedRef<FNiagaraSimCacheTreeItem> Item, const TSharedRef<STableViewBase>& OwnerTable);
+
+	void OnGetChildren(TSharedRef<FNiagaraSimCacheTreeItem> InItem, TArray<TSharedRef<FNiagaraSimCacheTreeItem>>& OutChildren);
+
+	void OnBufferChanged();
+
+	void RecursiveAddToSelectionFilter(TArray<TSharedRef<FNiagaraSimCacheTreeItem>>& ArrayToAdd);
+	void RecursiveRemoveFromSelectionFilter(TArray<TSharedRef<FNiagaraSimCacheTreeItem>>& ArrayToRemove);
+	
+	void UpdateSelectionFilter(TSharedRef<FNiagaraSimCacheTreeItem> ClickedItem);
+	
+	
+	TSharedPtr<STreeView<TSharedRef<FNiagaraSimCacheTreeItem>>> TreeView;
+
+	TArray<TSharedRef<FNiagaraSimCacheTreeItem>> ActiveRootEntries;
+
+	TSharedPtr<FNiagaraSimCacheViewModel> ViewModel;
+
+	TArray<TSharedRef<FNiagaraSimCacheTreeItem>> SelectionForFilter;
+};
+
+struct FNiagaraSimCacheTreeItem : public FNiagaraSimCacheOverviewItem
+{
+	FNiagaraSimCacheTreeItem(TWeakPtr<SNiagaraSimCacheTreeView> InOwner)
+	{
+		Owner = InOwner;
+	}
+
+	virtual ~FNiagaraSimCacheTreeItem() override
+	{
+		Children.Empty();
+	}
+
+	FString GetFilterName()
+	{
+		return FilterName;
+	}
+
+	void SetFilterName(FString NewName)
+	{
+		FilterName = NewName;
+	}
+
+	virtual FNiagaraSimCacheTreeItem* GetRootItem()
+	{
+		return this;
+	}
+
+	void AddChild(TSharedRef<FNiagaraSimCacheTreeItem> NewChild)
+	{
+		Children.AddUnique(NewChild);
+	}
+
+	virtual ENiagaraSimCacheOverviewItemType GetType () override { return ENiagaraSimCacheOverviewItemType::System; }
+
+	virtual TSharedRef<SWidget> GetRowWidget() override
+	{
+		return SNew(SHorizontalBox)
+			+SHorizontalBox::Slot()
+			.VAlign(VAlign_Center)
+			[
+				SNew(STextBlock)
+				.Text(GetDisplayNameText())
+			]
+			+SHorizontalBox::Slot()
+			.AutoWidth()
+			.HAlign(HAlign_Center)
+			.VAlign(VAlign_Center)
+			[
+				SNew(SNiagaraSimCacheTreeViewFilterWidget, SharedThis(this), Owner)
+			];
+	}
+
+	FText DisplayName;
+	FString FilterName;
+	TSharedPtr<SSimCacheTreeViewVisibilityWidget> VisibilityWidget;
+	TArray<TSharedRef<FNiagaraSimCacheTreeItem>> Children;
+	TWeakPtr<SNiagaraSimCacheTreeView> Owner;
+	
+};
+
+struct FNiagaraSimCacheEmitterTreeItem : public FNiagaraSimCacheTreeItem
+{
+	
+	FNiagaraSimCacheEmitterTreeItem(TWeakPtr<SNiagaraSimCacheTreeView> InOwner): FNiagaraSimCacheTreeItem(InOwner)
+	{
+	}
+
+	virtual ~FNiagaraSimCacheEmitterTreeItem() override {}
+	
+	virtual ENiagaraSimCacheOverviewItemType GetType () override { return ENiagaraSimCacheOverviewItemType::Emitter; }
+};
+
+struct FNiagaraSimCacheComponentTreeItem : public FNiagaraSimCacheTreeItem
+{
+	
+	FNiagaraSimCacheComponentTreeItem(TWeakPtr<SNiagaraSimCacheTreeView> InOwner): FNiagaraSimCacheTreeItem(InOwner)
+	{
+	}
+
+	virtual ~FNiagaraSimCacheComponentTreeItem() override {}
+
+	virtual ENiagaraSimCacheOverviewItemType GetType () override { return ENiagaraSimCacheOverviewItemType::Component; }
+
+	virtual TSharedRef<SWidget> GetRowWidget() override
+	{
+		
+		TSharedRef<SHorizontalBox> Contents = SNew(SHorizontalBox);
+
+		if(TypeDef.IsSet())
+		{
+			FSlateBrush const* IconBrush = TypeDef.GetValue().IsStatic() ? FNiagaraEditorStyle::Get().GetBrush(TEXT("NiagaraEditor.StaticIcon")) : FAppStyle::GetBrush(TEXT("Kismet.AllClasses.VariableIcon"));
+			const FLinearColor TypeColor = UEdGraphSchema_Niagara::GetTypeColor(TypeDef.GetValue());
+			
+			Contents->AddSlot()
+			.HAlign(HAlign_Left)
+			.AutoWidth()
+			[
+				SNew(SScaleBox)
+				[
+					SNew(SNiagaraIconWidget)
+					.IconToolTip(TypeDef.GetValue().GetNameText())
+					.IconBrush(IconBrush)
+					.IconColor(TypeColor)
+					.SecondaryIconBrush(FAppStyle::GetBrush(TEXT("NoBrush")))
+					.SecondaryIconColor(TypeColor)
+				]
+			];
+		}
+		
+		
+		Contents->AddSlot()
+		.HAlign(HAlign_Left)
+		.VAlign(VAlign_Center)
+		.AutoWidth()
+		.Padding(5, 0)
+		[
+			SNew(STextBlock)
+			.Text(GetDisplayNameText())
+			.TextStyle(&FAppStyle::Get().GetWidgetStyle<FTextBlockStyle>("NormalText"))
+		];
+
+		Contents->AddSlot()
+		.HAlign(HAlign_Right)
+		.FillWidth(1.0f)
+		[
+			SAssignNew(VisibilityWidget, SSimCacheTreeViewVisibilityWidget, SharedThis(this), Owner)
+		];
+		
+		return Contents;
+	}
+
+	TOptional<FNiagaraTypeDefinition> TypeDef;
+};
+
+
+
diff --git a/Engine/Plugins/FX/Niagara/Source/NiagaraEditor/Private/Widgets/SNiagaraSimCacheView.cpp b/Engine/Plugins/FX/Niagara/Source/NiagaraEditor/Private/Widgets/SNiagaraSimCacheView.cpp
--- a/Engine/Plugins/FX/Niagara/Source/NiagaraEditor/Private/Widgets/SNiagaraSimCacheView.cpp
+++ b/Engine/Plugins/FX/Niagara/Source/NiagaraEditor/Private/Widgets/SNiagaraSimCacheView.cpp
@@ -16,6 +16,7 @@
 
 #include "IContentBrowserSingleton.h"
 #include "ContentBrowserModule.h"
+#include "SNiagaraSimCacheTreeView.h"
 
 #define LOCTEXT_NAMESPACE "NiagaraSimCacheView"
 
@@ -63,11 +64,18 @@ public:
 	TSharedPtr<FNiagaraSimCacheViewModel>   SimCacheViewModel;
 };
 
+bool SNiagaraSimCacheView::IsStringFilterEnabled() const
+{
+	return !SimCacheViewModel->IsComponentFilterActive();
+}
+
 void SNiagaraSimCacheView::Construct(const FArguments& InArgs)
 {
 	SimCacheViewModel = InArgs._SimCacheViewModel;
 
 	SimCacheViewModel->OnViewDataChanged().AddSP(this, &SNiagaraSimCacheView::OnViewDataChanged);
+	SimCacheViewModel->OnSimCacheChanged().AddSP(this, &SNiagaraSimCacheView::OnSimCacheChanged);
+	SimCacheViewModel->OnBufferChanged().AddSP(this, &SNiagaraSimCacheView::OnBufferChanged);
 
 	HeaderRowWidget = SNew(SHeaderRow);
 
@@ -128,6 +136,7 @@ void SNiagaraSimCacheView::Construct(const FArguments& InArgs)
 			[
 				SNew(STextBlock)
 				.Text(LOCTEXT("ComponentFilter", "Component Filter"))
+				.IsEnabled(this, &SNiagaraSimCacheView::IsStringFilterEnabled)
 			]
 			+ SHorizontalBox::Slot()
 			.AutoWidth()
@@ -137,6 +146,7 @@ void SNiagaraSimCacheView::Construct(const FArguments& InArgs)
 				SNew(SEditableTextBox)
 				.OnTextChanged(this, &SNiagaraSimCacheView::OnComponentFilterChange)
 				.MinDesiredWidth(100)
+				.IsEnabled(this, &SNiagaraSimCacheView::IsStringFilterEnabled)
 			]
 		]
 		//// Main Spreadsheet View
@@ -203,7 +213,7 @@ void SNiagaraSimCacheView::GenerateColumns()
 	);
 		
 	// Generate a column for each component
-	for (const FNiagaraSimCacheViewModel::FComponentInfo& ComponentInfo : SimCacheViewModel->GetComponentInfos())
+	for (const FNiagaraSimCacheViewModel::FComponentInfo& ComponentInfo : SimCacheViewModel->GetCurrentComponentInfos())
 	{
 		HeaderRowWidget->AddColumn(
 			SHeaderRow::Column(ComponentInfo.Name)
@@ -225,21 +235,52 @@ void SNiagaraSimCacheView::UpdateColumns(const bool bReset)
 	{
 		GenerateColumns();
 	}
+	
+	TArray<FString> FilterArray;
 
-	const bool bFilterEmpty = ComponentFilterArray.Num() == 0;
+	// There are two filtering methods, from the tree view and from the spreadsheet view.
+	// The spreadsheet view uses a filter string input by the user. The tree view uses component item selection.
+	const bool bUseComponentFilter = SimCacheViewModel->IsComponentFilterActive();
+
+	FilterArray = bUseComponentFilter ? SimCacheViewModel->GetComponentFilters() : StringFilterArray;
+
+	// If the component filter is active, always apply the filter. If using the string filter, only apply it if it's not empty.
+	const bool bApplyFilter = bUseComponentFilter || !FilterArray.IsEmpty();
 
 	// If the columns are newly generated, and there are no filters to apply they are already up to date.
-	const bool bColumnsUpToDate = bReset && bFilterEmpty;
+	const bool bColumnsUpToDate = bReset && !bApplyFilter;
+
+	FString ColumnName;
 
+	auto ComponentFilterPred = [&ColumnName](const FString& ComponentFilter)
+	{
+		return ColumnName.Equals(ComponentFilter);
+	};
+
+	auto StringFilterPred = [&ColumnName](const FString& StringFilter)
+	{
+		return ColumnName.Contains(StringFilter);
+	};
+	
 	if(!bColumnsUpToDate)
 	{
 		for (SHeaderRow::FColumn Column : HeaderRowWidget->GetColumns())
 		{
-			const bool bPassedFilter = bFilterEmpty || ComponentFilterArray.ContainsByPredicate([ColumnName = Column.DefaultText.Get().ToString()](const FString& ComponentFilter)
+			ColumnName = Column.DefaultText.Get().ToString();
+
+			// Never filter instance count column.
+			if(ColumnName.Equals(NAME_Instance.ToString()))
 			{
-				return ColumnName.Contains(ComponentFilter) || ColumnName.Equals(NAME_Instance.ToString());
-			});
+				continue;
+			}
+
+			bool bPassedFilter = true;
 
+			if(bApplyFilter)
+			{
+				bPassedFilter = bUseComponentFilter ? FilterArray.ContainsByPredicate(ComponentFilterPred) : FilterArray.ContainsByPredicate(StringFilterPred);
+			}
+			
 			HeaderRowWidget->SetShowGeneratedColumn(Column.ColumnId, bPassedFilter);
 		}
 
@@ -296,8 +337,6 @@ void SNiagaraSimCacheView::BufferSelectionChanged(TSharedPtr<FBufferSelectionInf
 	}
 	
 	SimCacheViewModel->SetEmitterIndex(NewSelection->Key);
-	UpdateColumns(true);
-	UpdateRows(true);
 }
 
 FText SNiagaraSimCacheView::GetBufferSelectionText() const
@@ -315,22 +354,15 @@ FText SNiagaraSimCacheView::GetBufferSelectionText() const
 
 void SNiagaraSimCacheView::OnComponentFilterChange(const FText& InFilter)
 {
-	ComponentFilterArray.Empty();
-	InFilter.ToString().ParseIntoArray(ComponentFilterArray, TEXT(","));
+	InFilter.ToString().ParseIntoArray(StringFilterArray, TEXT(","));
 	UpdateColumns(false);
 }
 
-void SNiagaraSimCacheView::OnSimCacheChanged(const FAssetData& InAsset)
+void SNiagaraSimCacheView::OnSimCacheChanged()
 {
-	FSlateApplication::Get().DismissAllMenus();
-
-	if (UNiagaraSimCache* SimCache = Cast<UNiagaraSimCache>(InAsset.GetAsset()))
-	{
-		SimCacheViewModel->Initialize(SimCache);
-		UpdateRows(true);
-		UpdateColumns(true);
-		UpdateBufferSelectionList();
-	}
+	UpdateRows(true);
+	UpdateColumns(true);
+	UpdateBufferSelectionList();
 }
 
 void SNiagaraSimCacheView::OnViewDataChanged(const bool bFullRefresh)
@@ -339,8 +371,13 @@ void SNiagaraSimCacheView::OnViewDataChanged(const bool bFullRefresh)
 	if (bFullRefresh)
 	{
 		UpdateColumns(false);
-		UpdateBufferSelectionList();
 	}
 }
 
+void SNiagaraSimCacheView::OnBufferChanged()
+{
+	UpdateRows(true);
+	UpdateColumns(true);
+}
+
 #undef LOCTEXT_NAMESPACE
diff --git a/Engine/Plugins/FX/Niagara/Source/NiagaraEditor/Private/Widgets/SNiagaraSimCacheView.h b/Engine/Plugins/FX/Niagara/Source/NiagaraEditor/Private/Widgets/SNiagaraSimCacheView.h
--- a/Engine/Plugins/FX/Niagara/Source/NiagaraEditor/Private/Widgets/SNiagaraSimCacheView.h
+++ b/Engine/Plugins/FX/Niagara/Source/NiagaraEditor/Private/Widgets/SNiagaraSimCacheView.h
@@ -19,6 +19,7 @@ public:
 
 	using FBufferSelectionInfo = TPair<int32, FText>;
 
+	bool IsStringFilterEnabled() const;
 	void Construct(const FArguments& InArgs);
 
 	TSharedRef<ITableRow> MakeRowWidget(const TSharedPtr<int32> RowIndexPtr, const TSharedRef<STableViewBase>& OwnerTable) const;
@@ -39,9 +40,10 @@ public:
 
 	void OnComponentFilterChange(const FText& InFilter);
 
-	void OnSimCacheChanged(const FAssetData& InAsset);
+	void OnSimCacheChanged();
 
 	void OnViewDataChanged(const bool bFullRefresh);
+	void OnBufferChanged();
 
 private:
 	void GenerateColumns();
@@ -54,5 +56,5 @@ private:
 	TSharedPtr<SHeaderRow>						HeaderRowWidget;
 	TSharedPtr<SListView<TSharedPtr<int32>>>	ListViewWidget;
 
-	TArray<FString>								ComponentFilterArray;
+	TArray<FString>								StringFilterArray;
 };
\ No newline at end of file
diff --git a/Engine/Plugins/FX/Niagara/Source/NiagaraEditor/Private/Widgets/SNiagaraSimCacheViewTimeline.cpp b/Engine/Plugins/FX/Niagara/Source/NiagaraEditor/Private/Widgets/SNiagaraSimCacheViewTimeline.cpp
--- a/Engine/Plugins/FX/Niagara/Source/NiagaraEditor/Private/Widgets/SNiagaraSimCacheViewTimeline.cpp
+++ b/Engine/Plugins/FX/Niagara/Source/NiagaraEditor/Private/Widgets/SNiagaraSimCacheViewTimeline.cpp
@@ -34,7 +34,7 @@ int32 SNiagaraSimCacheViewTimeline::OnPaint(const FPaintArgs& Args, const FGeome
 	{
 		const FSlateBrush* BoxBrush = FAppStyle::GetBrush("Sequencer.Section.BackgroundTint");
 		const FLinearColor BoxTints[2] = { FLinearColor(0.5f, 0.5f, 0.5f, 1.0f), FLinearColor(0.3f, 0.3f, 0.3f, 1.0f) };
-		const FLinearColor CurrentTint(0.5f, 0.5f, 1.0f, 1.0f);
+		const FLinearColor CurrentTint(0.2f, 0.2f, 1.0f, 1.0f);
 		if (NumFrames > 0)
 		{
 			const float UStep = AllottedGeometry.Size.X / float(NumFrames);
diff --git a/Engine/Plugins/FX/Niagara/Source/NiagaraEditor/Public/NiagaraEditorStyle.h b/Engine/Plugins/FX/Niagara/Source/NiagaraEditor/Public/NiagaraEditorStyle.h
--- a/Engine/Plugins/FX/Niagara/Source/NiagaraEditor/Public/NiagaraEditorStyle.h
+++ b/Engine/Plugins/FX/Niagara/Source/NiagaraEditor/Public/NiagaraEditorStyle.h
@@ -52,6 +52,7 @@ private:
 	void InitScalabilityIcons();
 	void InitScratchStyle();
 	void InitHierarchyEditor();
+	void InitSimCacheEditor();
 
 	static TSharedPtr<FNiagaraEditorStyle> NiagaraEditorStyle;
 };
