- support new flag ClearBeforeNonIterationStage on all grid DIs
- Grid2D will not allocate a second buffer when we have write only grids

diff --git a/Engine/Plugins/FX/Niagara/Source/Niagara/Classes/NiagaraDataInterfaceGrid2DCollection.h b/Engine/Plugins/FX/Niagara/Source/Niagara/Classes/NiagaraDataInterfaceGrid2DCollection.h
--- a/Engine/Plugins/FX/Niagara/Source/Niagara/Classes/NiagaraDataInterfaceGrid2DCollection.h
+++ b/Engine/Plugins/FX/Niagara/Source/Niagara/Classes/NiagaraDataInterfaceGrid2DCollection.h
@@ -18,6 +18,8 @@ using FGrid2DBuffer = FNiagaraPooledRWTexture;
 
 struct FGrid2DCollectionRWInstanceData_GameThread
 {
+	bool ClearBeforeNonIterationStage = true;
+
 	FIntPoint NumCells = FIntPoint(EForceInit::ForceInitToZero);
 	int32 NumAttributes = 0;
 	FVector2D CellSize = FVector2D::ZeroVector;
@@ -50,6 +52,8 @@ struct FGrid2DCollectionRWInstanceData_GameThread
 
 struct FGrid2DCollectionRWInstanceData_RenderThread
 {
+	bool ClearBeforeNonIterationStage = true;
+
 	FIntPoint NumCells = FIntPoint(EForceInit::ForceInitToZero);
 	int32 NumAttributes = 0;
 	FVector2D CellSize = FVector2D::ZeroVector;
@@ -76,7 +80,7 @@ struct FGrid2DCollectionRWInstanceData_RenderThread
 	// overrides the render thread data, which in this case is for a grid reader
 	FNiagaraDataInterfaceProxy* OtherProxy = nullptr;
 
-	void BeginSimulate(FRDGBuilder& GraphBuilder);
+	void BeginSimulate(FRDGBuilder& GraphBuilder, bool RequiresBuffering);
 	void EndSimulate();
 };
 
diff --git a/Engine/Plugins/FX/Niagara/Source/Niagara/Classes/NiagaraDataInterfaceGrid3DCollection.h b/Engine/Plugins/FX/Niagara/Source/Niagara/Classes/NiagaraDataInterfaceGrid3DCollection.h
--- a/Engine/Plugins/FX/Niagara/Source/Niagara/Classes/NiagaraDataInterfaceGrid3DCollection.h
+++ b/Engine/Plugins/FX/Niagara/Source/Niagara/Classes/NiagaraDataInterfaceGrid3DCollection.h
@@ -18,6 +18,8 @@ using FGrid3DBuffer = FNiagaraPooledRWTexture;
 
 struct FGrid3DCollectionRWInstanceData_GameThread
 {
+	bool ClearBeforeNonIterationStage = true;
+
 	FIntVector NumCells = FIntVector::ZeroValue;
 	FIntVector NumTiles = FIntVector::ZeroValue;
 	int32 TotalNumAttributes = 0;
@@ -54,6 +56,8 @@ struct FGrid3DCollectionRWInstanceData_RenderThread
 {
 	FName SourceDIName;
 
+	bool ClearBeforeNonIterationStage = true;
+
 	FIntVector NumCells = FIntVector(1, 1, 1);
 	FIntVector NumTiles = FIntVector(1, 1, 1);
 	int32 TotalNumAttributes = 0;
diff --git a/Engine/Plugins/FX/Niagara/Source/Niagara/Classes/NiagaraDataInterfaceNeighborGrid3D.h b/Engine/Plugins/FX/Niagara/Source/Niagara/Classes/NiagaraDataInterfaceNeighborGrid3D.h
--- a/Engine/Plugins/FX/Niagara/Source/Niagara/Classes/NiagaraDataInterfaceNeighborGrid3D.h
+++ b/Engine/Plugins/FX/Niagara/Source/Niagara/Classes/NiagaraDataInterfaceNeighborGrid3D.h
@@ -29,6 +29,8 @@ struct FNDINeighborGrid3DInstanceData_RT
 {
 	void ResizeBuffers(FRDGBuilder& GraphBuilder);
 
+	bool ClearBeforeNonIterationStage = true;
+
 	FIntVector	NumCells = FIntVector::ZeroValue;
 	float		CellSize = 0.0f;
 	uint32		MaxNeighborsPerCell = 0;
@@ -41,6 +43,7 @@ struct FNDINeighborGrid3DInstanceData_RT
 
 struct FNiagaraDataInterfaceProxyNeighborGrid3D : public FNiagaraDataInterfaceProxyRW
 {	
+	virtual void ResetData(const FNDIGpuComputeResetContext& Context) override;
 	virtual void PreStage(const FNDIGpuComputePreStageContext& Context) override;
 	virtual void PostSimulate(const FNDIGpuComputePostSimulateContext& Context) override;
 
diff --git a/Engine/Plugins/FX/Niagara/Source/Niagara/Classes/NiagaraDataInterfaceRW.h b/Engine/Plugins/FX/Niagara/Source/Niagara/Classes/NiagaraDataInterfaceRW.h
--- a/Engine/Plugins/FX/Niagara/Source/Niagara/Classes/NiagaraDataInterfaceRW.h
+++ b/Engine/Plugins/FX/Niagara/Source/Niagara/Classes/NiagaraDataInterfaceRW.h
@@ -114,6 +114,12 @@ class NIAGARA_API UNiagaraDataInterfaceGrid3D : public UNiagaraDataInterfaceRWBa
 	GENERATED_UCLASS_BODY()
 
 public:
+
+	// Option to clear the buffer prior to a stage where the iteration count does not match the grid resolution.  Useful for stages where one wants to do sparse writes
+	// and accumulate values.
+	UPROPERTY(EditAnywhere, AdvancedDisplay, Category = "Grid")
+	bool ClearBeforeNonIterationStage;
+
 	// Number of cells
 	UPROPERTY(EditAnywhere, Category = "Grid", meta=(EditCondition = "false", EditConditionHides))
 	FIntVector NumCells;
@@ -164,6 +170,12 @@ class NIAGARA_API UNiagaraDataInterfaceGrid2D : public UNiagaraDataInterfaceRWBa
 	GENERATED_UCLASS_BODY()
 
 public:
+
+	// Option to clear the buffer prior to a stage where the iteration count does not match the grid resolution.  Useful for stages where one wants to do sparse writes
+	// and accumulate values.
+	UPROPERTY(EditAnywhere, AdvancedDisplay, Category = "Grid")
+	bool ClearBeforeNonIterationStage;
+
 	// Number of cells in X
 	UPROPERTY(EditAnywhere, Category = "Grid", meta = (EditCondition = "false", EditConditionHides))
 	int32 NumCellsX;
diff --git a/Engine/Plugins/FX/Niagara/Source/Niagara/Classes/NiagaraDataInterfaceRasterizationGrid3D.h b/Engine/Plugins/FX/Niagara/Source/Niagara/Classes/NiagaraDataInterfaceRasterizationGrid3D.h
--- a/Engine/Plugins/FX/Niagara/Source/Niagara/Classes/NiagaraDataInterfaceRasterizationGrid3D.h
+++ b/Engine/Plugins/FX/Niagara/Source/Niagara/Classes/NiagaraDataInterfaceRasterizationGrid3D.h
@@ -20,6 +20,9 @@ class RasterizationGrid3DRWInstanceData
 {
 public:
 	void ResizeBuffers();
+	
+	bool ClearBeforeNonIterationStage = true;
+
 	int32 TotalNumAttributes = 0;
 	FIntVector NumCells = FIntVector::ZeroValue;
 	FIntVector NumTiles = FIntVector::ZeroValue;
@@ -35,6 +38,7 @@ public:
 
 struct FNiagaraDataInterfaceProxyRasterizationGrid3D : public FNiagaraDataInterfaceProxyRW
 {	
+	virtual void ResetData(const FNDIGpuComputeResetContext& Context) override;
 	virtual void PreStage(const FNDIGpuComputePreStageContext& Context) override;
 	virtual void PostSimulate(const FNDIGpuComputePostSimulateContext& Context) override;
 
diff --git a/Engine/Plugins/FX/Niagara/Source/Niagara/Private/NiagaraDataInterfaceGrid2DCollection.cpp b/Engine/Plugins/FX/Niagara/Source/Niagara/Private/NiagaraDataInterfaceGrid2DCollection.cpp
--- a/Engine/Plugins/FX/Niagara/Source/Niagara/Private/NiagaraDataInterfaceGrid2DCollection.cpp
+++ b/Engine/Plugins/FX/Niagara/Source/Niagara/Private/NiagaraDataInterfaceGrid2DCollection.cpp
@@ -1860,6 +1860,8 @@ bool UNiagaraDataInterfaceGrid2DCollection::InitPerInstanceData(void* PerInstanc
 	}
 #endif
 
+	InstanceData->ClearBeforeNonIterationStage = ClearBeforeNonIterationStage;
+
 	// Push Updates to Proxy.
 	FNiagaraDataInterfaceProxyGrid2DCollectionProxy* RT_Proxy = GetProxyAs<FNiagaraDataInterfaceProxyGrid2DCollectionProxy>();
 	ENQUEUE_RENDER_COMMAND(FUpdateData)(
@@ -1868,6 +1870,8 @@ bool UNiagaraDataInterfaceGrid2DCollection::InitPerInstanceData(void* PerInstanc
 		check(!RT_Proxy->SystemInstancesToProxyData_RT.Contains(InstanceID));
 		FGrid2DCollectionRWInstanceData_RenderThread* TargetData = &RT_Proxy->SystemInstancesToProxyData_RT.Add(InstanceID);
 
+		TargetData->ClearBeforeNonIterationStage = RT_InstanceData.ClearBeforeNonIterationStage;
+
 		TargetData->NumCells = RT_InstanceData.NumCells;
 		TargetData->NumAttributes = RT_InstanceData.NumAttributes;
 		TargetData->CellSize = RT_InstanceData.CellSize;
@@ -2533,12 +2537,12 @@ bool FGrid2DCollectionRWInstanceData_GameThread::UpdateTargetTexture(ENiagaraGpu
 	return false;
 }
 
-void FGrid2DCollectionRWInstanceData_RenderThread::BeginSimulate(FRDGBuilder& GraphBuilder)
+void FGrid2DCollectionRWInstanceData_RenderThread::BeginSimulate(FRDGBuilder& GraphBuilder, bool RequiresBuffering)
 {
 	for (TUniquePtr<FGrid2DBuffer>& Buffer : Buffers)
 	{
 		check(Buffer.IsValid());
-		if (Buffer.Get() != CurrentData)
+		if (Buffer.Get() != CurrentData || !RequiresBuffering)
 		{
 			DestinationData = Buffer.Get();
 			break;
@@ -2552,6 +2556,9 @@ void FGrid2DCollectionRWInstanceData_RenderThread::BeginSimulate(FRDGBuilder& Gr
 
 		const FRDGTextureDesc TextureDesc = FRDGTextureDesc::Create2DArray(NumCells, PixelFormat, FClearValueBinding::Black, ETextureCreateFlags::ShaderResource | ETextureCreateFlags::UAV, NumAttributes);
 		DestinationData->Initialize(GraphBuilder, TEXT("FGrid2DBuffer"), TextureDesc);
+
+		// This destination buffer will sometimes have old data in it.  Force it to clear.
+		AddClearUAVPass(GraphBuilder, DestinationData->GetOrCreateUAV(GraphBuilder), FVector4f(ForceInitToZero));
 	}
 }
 
@@ -2584,29 +2591,21 @@ void FNiagaraDataInterfaceProxyGrid2DCollectionProxy::PreStage(const FNDIGpuComp
 		FGrid2DCollectionRWInstanceData_RenderThread* ProxyData = SystemInstancesToProxyData_RT.Find(Context.GetSystemInstanceID());
 
 		FRDGBuilder& GraphBuilder = Context.GetGraphBuilder();
-		ProxyData->BeginSimulate(GraphBuilder);
+		ProxyData->BeginSimulate(GraphBuilder, Context.IsInputStage());
+		
+		// If there is an output DI and the stage dispatch count is different from the grid resolution, then we want to consider it for a pre-stage UAV clear.
+		// This is to enable the most common use case of scattered writes to temporary buffers		
+		const FNiagaraSimStageData& SimStageData = Context.GetSimStageData();
 
-		// If we don't have an iteration stage, then we should manually clear the buffer to make sure there is no residual data.  If we are doing something like rasterizing particles into a grid, we want it to be clear before
-		// we start.  If a user wants to access data from the previous stage, then they can read from the current data.
+		FIntVector ElementCount = SimStageData.ElementCountXYZ;
 
-		// #todo(dmp): we might want to expose an option where we have buffers that are write only and need a clear (ie: no buffering like the neighbor grid).  They would be considered transient perhaps?  It'd be more
-		// memory efficient since it would theoretically not require any double buffering.
-
-		// #todo(dmp): for now, if there is an output DI that is NOT an iteration DI AND if both the iteration DI and this DI have the same total number of instances, do not do the UAV clear prior to the stage.
-		// this isn't optimal, but should work to some degree to reduce overhead in cases where we have multiple grids of the same resolution being processed/written to in 1 stage
-		if (!Context.IsIterationStage())
+		if (SimStageData.AlternateIterationSource != nullptr)
 		{
-			FNiagaraDataInterfaceProxyRW* IterationInterface = Context.GetSimStageData().AlternateIterationSource;
-			if (IterationInterface != nullptr)
-			{
-				const FIntVector ElementCount = IterationInterface->GetElementCount(Context.GetSystemInstanceID());				
-
-				if (ElementCount.X == ProxyData->NumCells.X && ElementCount.Y == ProxyData->NumCells.Y)
-				{
-					return;
-				}
-			}
+			ElementCount = SimStageData.AlternateIterationSource->GetElementCount(Context.GetSystemInstanceID());
+		}	
 
+		if (ProxyData->ClearBeforeNonIterationStage && (ElementCount.X != ProxyData->NumCells.X || ElementCount.Y != ProxyData->NumCells.Y))
+		{
 			AddClearUAVPass(GraphBuilder, ProxyData->DestinationData->GetOrCreateUAV(GraphBuilder), FVector4f(ForceInitToZero));
 		}
 	}
diff --git a/Engine/Plugins/FX/Niagara/Source/Niagara/Private/NiagaraDataInterfaceGrid3DCollection.cpp b/Engine/Plugins/FX/Niagara/Source/Niagara/Private/NiagaraDataInterfaceGrid3DCollection.cpp
--- a/Engine/Plugins/FX/Niagara/Source/Niagara/Private/NiagaraDataInterfaceGrid3DCollection.cpp
+++ b/Engine/Plugins/FX/Niagara/Source/Niagara/Private/NiagaraDataInterfaceGrid3DCollection.cpp
@@ -3496,6 +3496,8 @@ bool UNiagaraDataInterfaceGrid3DCollection::InitPerInstanceData(void* PerInstanc
 	}
 #endif
 
+	InstanceData->ClearBeforeNonIterationStage = ClearBeforeNonIterationStage;
+
 	// Push Updates to Proxy.
 	FNiagaraDataInterfaceProxyGrid3DCollectionProxy* RT_Proxy = GetProxyAs<FNiagaraDataInterfaceProxyGrid3DCollectionProxy>();
 	ENQUEUE_RENDER_COMMAND(FUpdateData)(
@@ -3505,6 +3507,9 @@ bool UNiagaraDataInterfaceGrid3DCollection::InitPerInstanceData(void* PerInstanc
 			FGrid3DCollectionRWInstanceData_RenderThread* TargetData = &RT_Proxy->SystemInstancesToProxyData_RT.Add(InstanceID);
 
 			TargetData->SourceDIName = RT_Proxy->SourceDIName;
+
+			TargetData->ClearBeforeNonIterationStage = RT_InstanceData.ClearBeforeNonIterationStage;
+
 			TargetData->NumCells = RT_InstanceData.NumCells;
 			TargetData->NumTiles = RT_InstanceData.NumTiles;
 			TargetData->TotalNumAttributes = RT_InstanceData.TotalNumAttributes;
@@ -4107,6 +4112,9 @@ void FGrid3DCollectionRWInstanceData_RenderThread::BeginSimulate(FRDGBuilder& Gr
 		StringBuilder.Append("Grid3DCollection");
 
 		DestinationData->Initialize(GraphBuilder, *StringBuilder, TextureDesc);
+
+		// This destination buffer will sometimes have old data in it.  Force it to clear.
+		AddClearUAVPass(GraphBuilder, DestinationData->GetOrCreateUAV(GraphBuilder), FVector4f(ForceInitToZero));
 	}
 }
 
@@ -4140,28 +4148,20 @@ void FNiagaraDataInterfaceProxyGrid3DCollectionProxy::PreStage(const FNDIGpuComp
 		FRDGBuilder& GraphBuilder = Context.GetGraphBuilder();
 		ProxyData->BeginSimulate(GraphBuilder, Context.IsInputStage());
 
-		// If we don't have an iteration stage, then we should manually clear the buffer to make sure there is no residual data.  If we are doing something like rasterizing particles into a grid, we want it to be clear before
-		// we start.  If a user wants to access data from the previous stage, then they can read from the current data.
+		// If there is an output DI and the stage dispatch count is different from the grid resolution, then we want to consider it for a pre-stage UAV clear.
+		// This is to enable the most common use case of scattered writes to temporary buffers				
+		const FNiagaraSimStageData& SimStageData = Context.GetSimStageData();
 
-		// #todo(dmp): we might want to expose an option where we have buffers that are write only and need a clear (ie: no buffering like the neighbor grid).  They would be considered transient perhaps?  It'd be more
-		// memory efficient since it would theoretically not require any double buffering.
+		FIntVector ElementCount = SimStageData.ElementCountXYZ;
 
-		// #todo(dmp): for now, if there is an output DI that is NOT an iteration DI AND if both the iteration DI and this DI have the same total number of instances, do not do the UAV clear prior to the stage.
-		// this isn't optimal, but should work to some degree to reduce overhead in cases where we have multiple grids of the same resolution being processed/written to in 1 stage
-		if (!Context.IsIterationStage())
+		if (SimStageData.AlternateIterationSource != nullptr)
 		{
-			FNiagaraDataInterfaceProxyRW* IterationInterface = Context.GetSimStageData().AlternateIterationSource;
-			if (IterationInterface != nullptr)
-			{
-				const FIntVector ElementCount = IterationInterface->GetElementCount(Context.GetSystemInstanceID());
-				const uint64 TotalNumInstances = ElementCount.X * ElementCount.Y * ElementCount.Z;
-
-				if (ElementCount.X == ProxyData->NumCells.X && ElementCount.Y == ProxyData->NumCells.Y && ElementCount.Z == ProxyData->NumCells.Z)
-				{
-					return;
-				}
-			}
+			ElementCount = SimStageData.AlternateIterationSource->GetElementCount(Context.GetSystemInstanceID());
+		}
 
+		if (ProxyData->ClearBeforeNonIterationStage && 
+			(ElementCount.X != ProxyData->NumCells.X || ElementCount.Y != ProxyData->NumCells.Y || ElementCount.Z != ProxyData->NumCells.Z))
+		{
 			AddClearUAVPass(GraphBuilder, ProxyData->DestinationData->GetOrCreateUAV(GraphBuilder), FVector4f(ForceInitToZero));
 		}
 	}
diff --git a/Engine/Plugins/FX/Niagara/Source/Niagara/Private/NiagaraDataInterfaceNeighborGrid3D.cpp b/Engine/Plugins/FX/Niagara/Source/Niagara/Private/NiagaraDataInterfaceNeighborGrid3D.cpp
--- a/Engine/Plugins/FX/Niagara/Source/Niagara/Private/NiagaraDataInterfaceNeighborGrid3D.cpp
+++ b/Engine/Plugins/FX/Niagara/Source/Niagara/Private/NiagaraDataInterfaceNeighborGrid3D.cpp
@@ -698,6 +698,15 @@ void UNiagaraDataInterfaceNeighborGrid3D::DestroyPerInstanceData(void* PerInstan
 	);
 }
 
+void FNiagaraDataInterfaceProxyNeighborGrid3D::ResetData(const FNDIGpuComputeResetContext& Context)
+{	
+	FNDINeighborGrid3DInstanceData_RT* ProxyData = SystemInstancesToProxyData_RT.Find(Context.GetSystemInstanceID());
+
+	FRDGBuilder& GraphBuilder = Context.GetGraphBuilder();
+	AddClearUAVPass(GraphBuilder, ProxyData->NeighborhoodBuffer.GetOrCreateUAV(GraphBuilder), -1);
+	AddClearUAVPass(GraphBuilder, ProxyData->NeighborhoodCountBuffer.GetOrCreateUAV(GraphBuilder), 0);
+}
+
 void FNiagaraDataInterfaceProxyNeighborGrid3D::PreStage(const FNDIGpuComputePreStageContext& Context)
 {
 	FRDGBuilder& GraphBuilder = Context.GetGraphBuilder();
@@ -707,7 +716,7 @@ void FNiagaraDataInterfaceProxyNeighborGrid3D::PreStage(const FNDIGpuComputePreS
 		ProxyData.ResizeBuffers(GraphBuilder);
 	}
 
-	if (Context.IsOutputStage() && ProxyData.NeighborhoodBuffer.IsValid())
+	if (Context.IsOutputStage() && ProxyData.NeighborhoodBuffer.IsValid() && ProxyData.ClearBeforeNonIterationStage)
 	{
 		RDG_RHI_EVENT_SCOPE(GraphBuilder, NiagaraNeighborGrid3DClearNeighborInfo);
 		AddClearUAVPass(GraphBuilder, ProxyData.NeighborhoodBuffer.GetOrCreateUAV(GraphBuilder), -1);
diff --git a/Engine/Plugins/FX/Niagara/Source/Niagara/Private/NiagaraDataInterfaceRW.cpp b/Engine/Plugins/FX/Niagara/Source/Niagara/Private/NiagaraDataInterfaceRW.cpp
--- a/Engine/Plugins/FX/Niagara/Source/Niagara/Private/NiagaraDataInterfaceRW.cpp
+++ b/Engine/Plugins/FX/Niagara/Source/Niagara/Private/NiagaraDataInterfaceRW.cpp
@@ -47,6 +47,7 @@ UNiagaraDataInterfaceRWBase::UNiagaraDataInterfaceRWBase(FObjectInitializer cons
 
 UNiagaraDataInterfaceGrid3D::UNiagaraDataInterfaceGrid3D(FObjectInitializer const& ObjectInitializer)
 	: Super(ObjectInitializer)
+	, ClearBeforeNonIterationStage(true)
 	, NumCells(3, 3, 3)
 	, CellSize(1.)
 	, NumCellsMaxAxis(10)
@@ -251,7 +252,9 @@ bool UNiagaraDataInterfaceGrid3D::Equals(const UNiagaraDataInterface* Other) con
 	}
 	const UNiagaraDataInterfaceGrid3D* OtherTyped = CastChecked<const UNiagaraDataInterfaceGrid3D>(Other);
 
-	return OtherTyped->NumCells == NumCells &&
+	return 
+		OtherTyped->ClearBeforeNonIterationStage == ClearBeforeNonIterationStage &&
+		OtherTyped->NumCells == NumCells &&
 		FMath::IsNearlyEqual(OtherTyped->CellSize, CellSize) &&		
 		OtherTyped->WorldBBoxSize.Equals(WorldBBoxSize) && 
 		OtherTyped->SetResolutionMethod == SetResolutionMethod && 
@@ -469,7 +472,7 @@ bool UNiagaraDataInterfaceGrid3D::CopyToInternal(UNiagaraDataInterface* Destinat
 
 	UNiagaraDataInterfaceGrid3D* OtherTyped = CastChecked<UNiagaraDataInterfaceGrid3D>(Destination);
 
-
+	OtherTyped->ClearBeforeNonIterationStage = ClearBeforeNonIterationStage;
 	OtherTyped->NumCells = NumCells;
 	OtherTyped->CellSize = CellSize;
 	OtherTyped->SetResolutionMethod = SetResolutionMethod;
@@ -483,6 +486,7 @@ bool UNiagaraDataInterfaceGrid3D::CopyToInternal(UNiagaraDataInterface* Destinat
 
 UNiagaraDataInterfaceGrid2D::UNiagaraDataInterfaceGrid2D(FObjectInitializer const& ObjectInitializer)
 	: Super(ObjectInitializer)
+	, ClearBeforeNonIterationStage(true)
 	, NumCellsX(3)
 	, NumCellsY(3)
 	, NumCellsMaxAxis(3)
@@ -721,6 +725,7 @@ bool UNiagaraDataInterfaceGrid2D::Equals(const UNiagaraDataInterface* Other) con
 	const UNiagaraDataInterfaceGrid2D* OtherTyped = CastChecked<const UNiagaraDataInterfaceGrid2D>(Other);
 
 	return 
+		OtherTyped->ClearBeforeNonIterationStage == ClearBeforeNonIterationStage &&
 		OtherTyped->NumCellsX == NumCellsX &&
 		OtherTyped->NumCellsY == NumCellsY &&
 		OtherTyped->NumAttributes == NumAttributes &&
@@ -956,7 +961,7 @@ bool UNiagaraDataInterfaceGrid2D::CopyToInternal(UNiagaraDataInterface* Destinat
 
 	UNiagaraDataInterfaceGrid2D* OtherTyped = CastChecked<UNiagaraDataInterfaceGrid2D>(Destination);
 
-
+	OtherTyped->ClearBeforeNonIterationStage = ClearBeforeNonIterationStage;
 	OtherTyped->NumCellsX = NumCellsX;
 	OtherTyped->NumCellsY = NumCellsY;
 	OtherTyped->NumAttributes = NumAttributes;
diff --git a/Engine/Plugins/FX/Niagara/Source/Niagara/Private/NiagaraDataInterfaceRasterizationGrid3D.cpp b/Engine/Plugins/FX/Niagara/Source/Niagara/Private/NiagaraDataInterfaceRasterizationGrid3D.cpp
--- a/Engine/Plugins/FX/Niagara/Source/Niagara/Private/NiagaraDataInterfaceRasterizationGrid3D.cpp
+++ b/Engine/Plugins/FX/Niagara/Source/Niagara/Private/NiagaraDataInterfaceRasterizationGrid3D.cpp
@@ -938,6 +938,8 @@ bool UNiagaraDataInterfaceRasterizationGrid3D::InitPerInstanceData(void* PerInst
 //	check(InstanceData->NumTiles.Y > 0);
 //	check(InstanceData->NumTiles.Z > 0);
 
+	InstanceData->ClearBeforeNonIterationStage = ClearBeforeNonIterationStage;
+
 	// @todo-threadsafety. This would be a race but I'm taking a ref here. Not ideal in the long term.
 	// Push Updates to Proxy.
 	ENQUEUE_RENDER_COMMAND(FUpdateData)(
@@ -945,6 +947,8 @@ bool UNiagaraDataInterfaceRasterizationGrid3D::InitPerInstanceData(void* PerInst
 	{
 		check(!RT_Proxy->SystemInstancesToProxyData.Contains(InstanceID));
 		RasterizationGrid3DRWInstanceData* TargetData = &RT_Proxy->SystemInstancesToProxyData.Add(InstanceID);
+		
+		TargetData->ClearBeforeNonIterationStage = RT_InstanceData.ClearBeforeNonIterationStage;
 		TargetData->TotalNumAttributes = RT_TotalNumAttributes;
 		TargetData->NumCells = RT_NumCells;		
 		TargetData->NumTiles = RT_InstanceData.NumTiles;
@@ -1102,6 +1106,19 @@ void UNiagaraDataInterfaceRasterizationGrid3D::DestroyPerInstanceData(void* PerI
 	);
 }
 
+void FNiagaraDataInterfaceProxyRasterizationGrid3D::ResetData(const FNDIGpuComputeResetContext& Context)
+{	
+	RasterizationGrid3DRWInstanceData* ProxyData = SystemInstancesToProxyData.Find(Context.GetSystemInstanceID());
+	if (!ProxyData)
+	{
+		return;
+	}
+
+	FRDGBuilder& GraphBuilder = Context.GetGraphBuilder();
+	const FUintVector4 ResetValue(ProxyData->ResetValue, ProxyData->ResetValue, ProxyData->ResetValue, ProxyData->ResetValue);
+	AddClearUAVPass(GraphBuilder, ProxyData->RasterizationTexture.GetOrCreateUAV(GraphBuilder), ResetValue);
+}
+
 void FNiagaraDataInterfaceProxyRasterizationGrid3D::PreStage(const FNDIGpuComputePreStageContext& Context)
 {
 	using namespace NDIRasterizationGrid3DLocal;
@@ -1151,7 +1168,7 @@ void FNiagaraDataInterfaceProxyRasterizationGrid3D::PreStage(const FNDIGpuComput
 		InstanceData.PerAttributeData.Initialize(TEXT("Grid3D::PerAttributeData"), sizeof(FVector4f), PerAttributeData.Num(), EPixelFormat::PF_A32B32G32R32F, BUF_Static, &PerAttributeData);
 	}
 
-	if (Context.IsOutputStage() && NumTotalCells > 0)
+	if (Context.IsOutputStage() && NumTotalCells > 0 && InstanceData.ClearBeforeNonIterationStage)
 	{
 		FRDGBuilder& GraphBuilder = Context.GetGraphBuilder();
 		const FUintVector4 ResetValue(InstanceData.ResetValue, InstanceData.ResetValue, InstanceData.ResetValue, InstanceData.ResetValue);
