vdb sim cache support w/ optional cvars

diff --git a/Engine/Plugins/FX/Niagara/Source/Niagara/Classes/VolumeCache.h b/Engine/Plugins/FX/Niagara/Source/Niagara/Classes/VolumeCache.h
--- a/Engine/Plugins/FX/Niagara/Source/Niagara/Classes/VolumeCache.h
+++ b/Engine/Plugins/FX/Niagara/Source/Niagara/Classes/VolumeCache.h
@@ -44,7 +44,8 @@ public:
 	bool LoadFile(int frame);
 	bool UnloadFile(int frame);		
 	bool LoadRange();
-	void UnloadAll();	
+	void UnloadAll();
+	FString GetAssetPath(int frame);
 	
 	TSharedPtr<FVolumeCacheData> GetData() { return CachedVolumeFiles;  }
 
diff --git a/Engine/Plugins/FX/Niagara/Source/Niagara/Private/NiagaraDataInterfaceRenderTargetVolume.cpp b/Engine/Plugins/FX/Niagara/Source/Niagara/Private/NiagaraDataInterfaceRenderTargetVolume.cpp
--- a/Engine/Plugins/FX/Niagara/Source/Niagara/Private/NiagaraDataInterfaceRenderTargetVolume.cpp
+++ b/Engine/Plugins/FX/Niagara/Source/Niagara/Private/NiagaraDataInterfaceRenderTargetVolume.cpp
@@ -16,6 +16,7 @@
 #include "NiagaraShader.h"
 #include "NiagaraShaderParametersBuilder.h"
 #include "NiagaraGpuComputeDebugInterface.h"
+#include "VolumeCache.h"
 #include "RHIStaticStates.h"
 
 #include UE_INLINE_GENERATED_CPP_BY_NAME(NiagaraDataInterfaceRenderTargetVolume)
@@ -71,6 +72,23 @@ namespace NDIRenderTargetVolumeLocal
 		TEXT("When enabled compression is used for the sim cache data."),
 		ECVF_Default
 	);
+
+	int32 GSimCacheUseOpenVDB = true;
+	static FAutoConsoleVariableRef CVarSimCacheUseOpenVDB(
+		TEXT("fx.Niagara.RenderTargetVolume.SimCacheUseOpenVDB"),
+		GSimCacheUseOpenVDB,
+		TEXT("Use OpenVDB as the backing data for the sim cache."),
+		ECVF_Default
+	);
+
+	FString GSimCacheOpenVDBBasePath = "";
+	static FAutoConsoleVariableRef CVarSimCacheOpenVDBBasePath(
+		TEXT("fx.Niagara.RenderTargetVolume.SimCacheOpenVDBBasePath"),
+		GSimCacheOpenVDBBasePath,
+		TEXT("Base path for OpenVDB output."),
+		ECVF_Default
+	);
+
 	static const FName GetSimCacheCompressionType() { return GSimCacheCompressed ? NAME_Oodle : NAME_None; }
 }
 
@@ -299,7 +317,7 @@ void UNiagaraDataInterfaceRenderTargetVolume::GetFunctions(TArray<FNiagaraFuncti
 #if WITH_EDITORONLY_DATA
 		Sig.FunctionVersion = EFunctionVersion::LatestVersion;
 #endif
-	}
+}
 }
 
 #if WITH_EDITORONLY_DATA
@@ -593,27 +611,74 @@ bool UNiagaraDataInterfaceRenderTargetVolume::GetExposedVariableValue(const FNia
 
 UObject* UNiagaraDataInterfaceRenderTargetVolume::SimCacheBeginWrite(UObject* SimCache, FNiagaraSystemInstance* NiagaraSystemInstance, const void* OptionalPerInstanceData) const
 {
+	if (NDIRenderTargetVolumeLocal::GSimCacheUseOpenVDB)
+	{
+		const FRenderTargetVolumeRWInstanceData_GameThread* InstanceData_GT = reinterpret_cast<const FRenderTargetVolumeRWInstanceData_GameThread*>(OptionalPerInstanceData);		
+
+		UVolumeCache* OpenVDBSimCacheData = nullptr;
+
+		OpenVDBSimCacheData = NewObject<UVolumeCache>(SimCache);
+
+		FString SystemInstanceName = NiagaraSystemInstance->GetSystem()->GetName();
+
+		if (NDIRenderTargetVolumeLocal::GSimCacheOpenVDBBasePath == "")
+		{
+			UE_LOG(LogNiagara, Error, TEXT("You must set fx.Niagara.RenderTargetVolume.SimCacheOpenVDBPath"));
+			return nullptr;
+		}
+
+		FString ActorName;
+
+		AActor* Owner = NiagaraSystemInstance->GetAttachComponent()->GetOwner();
+		ActorName = Owner->GetActorLabel();
+
+		FString DIName = Proxy->SourceDIName.ToString();
+		FString FullFilePathSpec = NDIRenderTargetVolumeLocal::GSimCacheOpenVDBBasePath + "/NiagaraSimCache/" + ActorName + "/" + DIName + "_SimCache.{FrameIndex}.vdb";
+		FullFilePathSpec.ReplaceInline(TEXT("//"), TEXT("/"));		
+
+		// Create output directory		
+		const FString FileDirectory = FString(FPathViews::GetPath(FullFilePathSpec));
+
+		IPlatformFile& PlatformFile = FPlatformFileManager::Get().GetPlatformFile();
+		if (!PlatformFile.DirectoryExists(*FileDirectory))
+		{
+			if (!PlatformFile.CreateDirectoryTree(*FileDirectory))
+			{
+				UE_LOG(LogNiagara, Error, TEXT("Cannot Create Directory : %s"), *FileDirectory);
+				return nullptr;
+			}
+		}
+
+		OpenVDBSimCacheData->FilePath = FullFilePathSpec;
+		OpenVDBSimCacheData->CacheType = EVolumeCacheType::OpenVDB;
+		OpenVDBSimCacheData->Resolution = InstanceData_GT->Size;
+		OpenVDBSimCacheData->FrameRangeStart = TNumericLimits<int32>::Lowest();
+		OpenVDBSimCacheData->FrameRangeEnd = TNumericLimits<int32>::Lowest();
+		OpenVDBSimCacheData->InitData();
+
+		return OpenVDBSimCacheData;
+	}
+	else
+	{
 	UNDIRenderTargetVolumeSimCacheData* SimCacheData = nullptr;
 	if (NDIRenderTargetVolumeLocal::GSimCacheEnabled)
 	{
 		SimCacheData = NewObject<UNDIRenderTargetVolumeSimCacheData>(SimCache);
 		SimCacheData->CompressionType = NDIRenderTargetVolumeLocal::GetSimCacheCompressionType();
 	}
+
 	return SimCacheData;
 }
 
+	return nullptr;
+}
+
 bool UNiagaraDataInterfaceRenderTargetVolume::SimCacheWriteFrame(UObject* StorageObject, int FrameIndex, FNiagaraSystemInstance* SystemInstance, const void* OptionalPerInstanceData) const
 {
 	check(OptionalPerInstanceData);
 
-	UNDIRenderTargetVolumeSimCacheData* SimCacheData = CastChecked<UNDIRenderTargetVolumeSimCacheData>(StorageObject);
-	SimCacheData->Frames.SetNum(FMath::Max(SimCacheData->Frames.Num(), FrameIndex + 1));
-
-	FNDIRenderTargetVolumeSimCacheFrame* CacheFrame = &SimCacheData->Frames[FrameIndex];
-
 	const FRenderTargetVolumeRWInstanceData_GameThread* InstanceData_GT = reinterpret_cast<const FRenderTargetVolumeRWInstanceData_GameThread*>(OptionalPerInstanceData);
-	CacheFrame->Size = InstanceData_GT->Size;
-	CacheFrame->Format = EPixelFormat::PF_FloatRGBA;
+
 
 	if (InstanceData_GT->TargetTexture)
 	{
@@ -642,6 +707,27 @@ bool UNiagaraDataInterfaceRenderTargetVolume::SimCacheWriteFrame(UObject* Storag
 
 		if (TextureData.Num() > 0)
 		{
+			if (NDIRenderTargetVolumeLocal::GSimCacheUseOpenVDB)
+			{
+				UVolumeCache* OpenVDBSimCacheData = CastChecked<UVolumeCache>(StorageObject);
+				
+				FString FullPath = OpenVDBSimCacheData->GetAssetPath(FrameIndex);
+
+				OpenVDBTools::WriteImageDataToOpenVDBFile(FullPath, InstanceData_GT->Size, TextureData, false);
+
+				OpenVDBSimCacheData->FrameRangeStart = FMath::Min(FrameIndex, OpenVDBSimCacheData->FrameRangeStart);
+				OpenVDBSimCacheData->FrameRangeEnd = FMath::Max(FrameIndex, OpenVDBSimCacheData->FrameRangeEnd);
+			}
+			else
+			{
+				UNDIRenderTargetVolumeSimCacheData* SimCacheData = CastChecked<UNDIRenderTargetVolumeSimCacheData>(StorageObject);
+				SimCacheData->Frames.SetNum(FMath::Max(SimCacheData->Frames.Num(), FrameIndex + 1));
+
+				FNDIRenderTargetVolumeSimCacheFrame* CacheFrame = &SimCacheData->Frames[FrameIndex];
+				
+				CacheFrame->Size = InstanceData_GT->Size;
+				CacheFrame->Format = EPixelFormat::PF_FloatRGBA;
+
 			const FName CompressionType = SimCacheData->CompressionType;
 			const bool bUseCompression = CompressionType.IsNone() == false;
 			const int TextureSizeBytes = TextureData.Num() * TextureData.GetTypeSize();
@@ -675,19 +761,63 @@ bool UNiagaraDataInterfaceRenderTargetVolume::SimCacheWriteFrame(UObject* Storag
 			FMemory::Free(FinalPixelData);
 		}
 	}
+	}
 
 	return true;
 }
 
 bool UNiagaraDataInterfaceRenderTargetVolume::SimCacheEndWrite(UObject* StorageObject) const
 {
+	if (NDIRenderTargetVolumeLocal::GSimCacheUseOpenVDB)
+	{
+		UVolumeCache* OpenVDBSimCacheData = CastChecked<UVolumeCache>(StorageObject);
+	}
+	else
+	{
 	UNDIRenderTargetVolumeSimCacheData* SimCacheData = CastChecked<UNDIRenderTargetVolumeSimCacheData>(StorageObject);
+	}
 
 	return true;
 }
 
 bool UNiagaraDataInterfaceRenderTargetVolume::SimCacheReadFrame(UObject* StorageObject, int FrameA, int FrameB, float Interp, FNiagaraSystemInstance* SystemInstance, void* OptionalPerInstanceData)
 {
+	FRenderTargetVolumeRWInstanceData_GameThread* InstanceData_GT = reinterpret_cast<FRenderTargetVolumeRWInstanceData_GameThread*>(OptionalPerInstanceData);
+
+	if (NDIRenderTargetVolumeLocal::GSimCacheUseOpenVDB)
+	{
+		UVolumeCache *OpenVDBSimCacheData = CastChecked<UVolumeCache>(StorageObject);
+
+		TSharedPtr<FVolumeCacheData> VolumeCacheData = OpenVDBSimCacheData->GetData();
+
+		const int FrameIndex = Interp >= 0.5f ? FrameB : FrameA;
+		OpenVDBSimCacheData->LoadFile(FrameIndex);
+
+		InstanceData_GT->Size = OpenVDBSimCacheData->Resolution;
+		InstanceData_GT->Format = PF_FloatRGBA;
+
+		PerInstanceTick(InstanceData_GT, SystemInstance, 0.0f);
+		PerInstanceTickPostSimulate(InstanceData_GT, SystemInstance, 0.0f);
+
+		//  write to the volume texture
+		if (InstanceData_GT->TargetTexture)
+		{
+			FNiagaraDataInterfaceProxyRenderTargetVolumeProxy* RT_Proxy = GetProxyAs<FNiagaraDataInterfaceProxyRenderTargetVolumeProxy>();
+
+			FTextureRenderTargetResource* RT_TargetTexture = InstanceData_GT->TargetTexture->GameThread_GetRenderTargetResource();
+			ENQUEUE_RENDER_COMMAND(NDIRenderTargetVolumeUpdate)
+				(
+					[RT_Proxy, RT_VolumeCacheData = OpenVDBSimCacheData->GetData(), RT_InstanceID = SystemInstance->GetId(), RT_TargetTexture, RT_FrameIndex = FrameIndex](FRHICommandListImmediate& RHICmdList)
+					{
+						if (FRenderTargetVolumeRWInstanceData_RenderThread* InstanceData_RT = RT_Proxy->SystemInstancesToProxyData_RT.Find(RT_InstanceID))
+						{
+							RT_VolumeCacheData->Fill3DTexture(RT_FrameIndex, InstanceData_RT->RenderTarget->GetRHI());
+						}
+					});
+		}
+	}
+	else
+	{
 	UNDIRenderTargetVolumeSimCacheData* SimCacheData = CastChecked<UNDIRenderTargetVolumeSimCacheData>(StorageObject);
 
 	const int FrameIndex = Interp >= 0.5f ? FrameB : FrameA;
@@ -698,7 +828,6 @@ bool UNiagaraDataInterfaceRenderTargetVolume::SimCacheReadFrame(UObject* Storage
 
 	FNDIRenderTargetVolumeSimCacheFrame* CacheFrame = &SimCacheData->Frames[FrameIndex];
 
-	FRenderTargetVolumeRWInstanceData_GameThread* InstanceData_GT = reinterpret_cast<FRenderTargetVolumeRWInstanceData_GameThread*>(OptionalPerInstanceData);
 	InstanceData_GT->Size = CacheFrame->Size;
 	InstanceData_GT->Format = CacheFrame->Format;
 	InstanceData_GT->Filter = TextureFilter::TF_Default;
@@ -724,8 +853,8 @@ bool UNiagaraDataInterfaceRenderTargetVolume::SimCacheReadFrame(UObject* Storage
 					const int32 SrcDepthPitch = InstanceData_RT->Size.Y * SrcRowPitch;
 					if (UpdateTexture.RowPitch == SrcRowPitch && UpdateTexture.DepthPitch == SrcDepthPitch)
 					{
-						if (RT_CacheFrame->CompressedSize > 0)
-						{
+					if (RT_CacheFrame->CompressedSize > 0)
+					{
 							FCompression::UncompressMemory(RT_CompressionType, UpdateTexture.Data, UpdateTexture.DataSizeBytes, RT_CacheFrame->GetPixelData(), RT_CacheFrame->CompressedSize);
 						}
 						else
@@ -738,8 +867,8 @@ bool UNiagaraDataInterfaceRenderTargetVolume::SimCacheReadFrame(UObject* Storage
 						TArray<uint8> Decompressed;
 						if (RT_CacheFrame->CompressedSize > 0)
 						{
-							Decompressed.AddUninitialized(sizeof(FFloat16Color) * InstanceData_RT->Size.X * InstanceData_RT->Size.Y * InstanceData_RT->Size.Z);
-							FCompression::UncompressMemory(RT_CompressionType, Decompressed.GetData(), Decompressed.Num(), RT_CacheFrame->GetPixelData(), RT_CacheFrame->CompressedSize);
+						Decompressed.AddUninitialized(sizeof(FFloat16Color) * InstanceData_RT->Size.X * InstanceData_RT->Size.Y * InstanceData_RT->Size.Z);
+						FCompression::UncompressMemory(RT_CompressionType, Decompressed.GetData(), Decompressed.Num(), RT_CacheFrame->GetPixelData(), RT_CacheFrame->CompressedSize);
 						}
 
 						const uint8* SrcData = Decompressed.Num() > 0 ? Decompressed.GetData() : RT_CacheFrame->GetPixelData();
@@ -752,14 +881,14 @@ bool UNiagaraDataInterfaceRenderTargetVolume::SimCacheReadFrame(UObject* Storage
 								SrcData += SrcRowPitch;
 								DstData += UpdateTexture.RowPitch;
 							}
-						}
+					}
 					}
 					RHICmdList.EndUpdateTexture3D(UpdateTexture);
 				}
 			}
 		);
 	}
-
+	}
 	return true;
 }
 
diff --git a/Engine/Plugins/FX/Niagara/Source/Niagara/Private/VolumeCache.cpp b/Engine/Plugins/FX/Niagara/Source/Niagara/Private/VolumeCache.cpp
--- a/Engine/Plugins/FX/Niagara/Source/Niagara/Private/VolumeCache.cpp
+++ b/Engine/Plugins/FX/Niagara/Source/Niagara/Private/VolumeCache.cpp
@@ -98,6 +98,13 @@ void UVolumeCache::UnloadAll()
 	}
 }
 
+FString UVolumeCache::GetAssetPath(int frame)
+{
+	check(CachedVolumeFiles != nullptr);
+
+	return CachedVolumeFiles->GetAssetPath(FilePath, frame);
+}
+
 bool UVolumeCache::LoadFile(int frame)
 {
 	if (CachedVolumeFiles == nullptr)
