cis fix

diff --git a/Engine/Plugins/FX/Niagara/Source/Niagara/Private/NiagaraDataInterfaceRenderTargetVolume.cpp b/Engine/Plugins/FX/Niagara/Source/Niagara/Private/NiagaraDataInterfaceRenderTargetVolume.cpp
--- a/Engine/Plugins/FX/Niagara/Source/Niagara/Private/NiagaraDataInterfaceRenderTargetVolume.cpp
+++ b/Engine/Plugins/FX/Niagara/Source/Niagara/Private/NiagaraDataInterfaceRenderTargetVolume.cpp
@@ -627,10 +627,13 @@ UObject* UNiagaraDataInterfaceRenderTargetVolume::SimCacheBeginWrite(UObject* Si
 			return nullptr;
 		}
 
-		FString ActorName;
+		FString ActorName = NiagaraSystemInstance->GetAttachComponent()->GetName();
 
+		// if we are caching from the editor, use the label of the actor to name the cache files
+#if WITH_EDITOR
 		AActor* Owner = NiagaraSystemInstance->GetAttachComponent()->GetOwner();
 		ActorName = Owner->GetActorLabel();
+#endif
 
 		FString DIName = Proxy->SourceDIName.ToString();
 		FString FullFilePathSpec = NDIRenderTargetVolumeLocal::GSimCacheOpenVDBBasePath + "/NiagaraSimCache/" + ActorName + "/" + DIName + "_SimCache.{FrameIndex}.vdb";
