fix missing include and some erroneous tabbing from a merge

diff --git a/Engine/Plugins/FX/Niagara/Source/Niagara/Private/NiagaraDataInterfaceRenderTargetVolume.cpp b/Engine/Plugins/FX/Niagara/Source/Niagara/Private/NiagaraDataInterfaceRenderTargetVolume.cpp
--- a/Engine/Plugins/FX/Niagara/Source/Niagara/Private/NiagaraDataInterfaceRenderTargetVolume.cpp
+++ b/Engine/Plugins/FX/Niagara/Source/Niagara/Private/NiagaraDataInterfaceRenderTargetVolume.cpp
@@ -18,6 +18,7 @@
 #include "NiagaraGpuComputeDebugInterface.h"
 #include "VolumeCache.h"
 #include "RHIStaticStates.h"
+#include "Misc/Paths.h"
 
 #include UE_INLINE_GENERATED_CPP_BY_NAME(NiagaraDataInterfaceRenderTargetVolume)
 
@@ -663,15 +664,15 @@ UObject* UNiagaraDataInterfaceRenderTargetVolume::SimCacheBeginWrite(UObject* Si
 	}
 	else
 	{
-	UNDIRenderTargetVolumeSimCacheData* SimCacheData = nullptr;
-	if (NDIRenderTargetVolumeLocal::GSimCacheEnabled)
-	{
-		SimCacheData = NewObject<UNDIRenderTargetVolumeSimCacheData>(SimCache);
-		SimCacheData->CompressionType = NDIRenderTargetVolumeLocal::GetSimCacheCompressionType();
-	}
+		UNDIRenderTargetVolumeSimCacheData* SimCacheData = nullptr;
+		if (NDIRenderTargetVolumeLocal::GSimCacheEnabled)
+		{
+			SimCacheData = NewObject<UNDIRenderTargetVolumeSimCacheData>(SimCache);
+			SimCacheData->CompressionType = NDIRenderTargetVolumeLocal::GetSimCacheCompressionType();
+		}
 
-	return SimCacheData;
-}
+		return SimCacheData;
+	}
 
 	return nullptr;
 }
