cis fixes

diff --git a/Engine/Plugins/FX/Niagara/Source/Niagara/Private/NiagaraDataInterfaceRenderTargetVolume.cpp b/Engine/Plugins/FX/Niagara/Source/Niagara/Private/NiagaraDataInterfaceRenderTargetVolume.cpp
--- a/Engine/Plugins/FX/Niagara/Source/Niagara/Private/NiagaraDataInterfaceRenderTargetVolume.cpp
+++ b/Engine/Plugins/FX/Niagara/Source/Niagara/Private/NiagaraDataInterfaceRenderTargetVolume.cpp
@@ -713,6 +713,7 @@ bool UNiagaraDataInterfaceRenderTargetVolume::SimCacheWriteFrame(UObject* Storag
 		{
 			if (NDIRenderTargetVolumeLocal::GSimCacheUseOpenVDB)
 			{
+#if PLATFORM_WINDOWS
 				UVolumeCache* OpenVDBSimCacheData = CastChecked<UVolumeCache>(StorageObject);
 				
 				FString FullPath = OpenVDBSimCacheData->GetAssetPath(FrameIndex);
@@ -721,6 +722,7 @@ bool UNiagaraDataInterfaceRenderTargetVolume::SimCacheWriteFrame(UObject* Storag
 
 				OpenVDBSimCacheData->FrameRangeStart = FMath::Min(FrameIndex, OpenVDBSimCacheData->FrameRangeStart);
 				OpenVDBSimCacheData->FrameRangeEnd = FMath::Max(FrameIndex, OpenVDBSimCacheData->FrameRangeEnd);
+#endif
 			}
 			else
 			{
@@ -732,55 +734,46 @@ bool UNiagaraDataInterfaceRenderTargetVolume::SimCacheWriteFrame(UObject* Storag
 				CacheFrame->Size = InstanceData_GT->Size;
 				CacheFrame->Format = EPixelFormat::PF_FloatRGBA;
 
-			const FName CompressionType = SimCacheData->CompressionType;
-			const bool bUseCompression = CompressionType.IsNone() == false;
-			const int TextureSizeBytes = TextureData.Num() * TextureData.GetTypeSize();
-			int CompressedSize = bUseCompression ? FCompression::CompressMemoryBound(CompressionType, TextureSizeBytes) : TextureSizeBytes;
+				const FName CompressionType = SimCacheData->CompressionType;
+				const bool bUseCompression = CompressionType.IsNone() == false;
+				const int TextureSizeBytes = TextureData.Num() * TextureData.GetTypeSize();
+				int CompressedSize = bUseCompression ? FCompression::CompressMemoryBound(CompressionType, TextureSizeBytes) : TextureSizeBytes;
 
-			CacheFrame->UncompressedSize = TextureSizeBytes;
-			CacheFrame->CompressedSize = 0;
-			uint8* FinalPixelData = reinterpret_cast<uint8*>(FMemory::Malloc(CompressedSize));
-			if (bUseCompression)
-			{
-				if (FCompression::CompressMemory(CompressionType, FinalPixelData, CompressedSize, TextureData.GetData(), TextureSizeBytes, COMPRESS_BiasMemory))
+				CacheFrame->UncompressedSize = TextureSizeBytes;
+				CacheFrame->CompressedSize = 0;
+				uint8* FinalPixelData = reinterpret_cast<uint8*>(FMemory::Malloc(CompressedSize));
+				if (bUseCompression)
 				{
-					CacheFrame->CompressedSize = CompressedSize;
+					if (FCompression::CompressMemory(CompressionType, FinalPixelData, CompressedSize, TextureData.GetData(), TextureSizeBytes, COMPRESS_BiasMemory))
+					{
+						CacheFrame->CompressedSize = CompressedSize;
+					}
 				}
-			}
 
-			if (CacheFrame->CompressedSize == 0)
-			{
-				FMemory::Memcpy(FinalPixelData, TextureData.GetData(), TextureSizeBytes);
-			}
+				if (CacheFrame->CompressedSize == 0)
+				{
+					FMemory::Memcpy(FinalPixelData, TextureData.GetData(), TextureSizeBytes);
+				}
 
-			// Update bulk data
-			CacheFrame->BulkData.Lock(LOCK_READ_WRITE);
-			{
-				const int32 FinalByteSize = CacheFrame->CompressedSize > 0 ? CacheFrame->CompressedSize : CacheFrame->UncompressedSize;
-				uint8* BulkDataPtr = reinterpret_cast<uint8*>(CacheFrame->BulkData.Realloc(FinalByteSize));
-				FMemory::Memcpy(BulkDataPtr, FinalPixelData, FinalByteSize);
-			}
-			CacheFrame->BulkData.Unlock();
+				// Update bulk data
+				CacheFrame->BulkData.Lock(LOCK_READ_WRITE);
+				{
+					const int32 FinalByteSize = CacheFrame->CompressedSize > 0 ? CacheFrame->CompressedSize : CacheFrame->UncompressedSize;
+					uint8* BulkDataPtr = reinterpret_cast<uint8*>(CacheFrame->BulkData.Realloc(FinalByteSize));
+					FMemory::Memcpy(BulkDataPtr, FinalPixelData, FinalByteSize);
+				}
+				CacheFrame->BulkData.Unlock();
 
-			FMemory::Free(FinalPixelData);
+				FMemory::Free(FinalPixelData);
+			}
 		}
 	}
-	}
 
 	return true;
 }
 
 bool UNiagaraDataInterfaceRenderTargetVolume::SimCacheEndWrite(UObject* StorageObject) const
 {
-	if (NDIRenderTargetVolumeLocal::GSimCacheUseOpenVDB)
-	{
-		UVolumeCache* OpenVDBSimCacheData = CastChecked<UVolumeCache>(StorageObject);
-	}
-	else
-	{
-	UNDIRenderTargetVolumeSimCacheData* SimCacheData = CastChecked<UNDIRenderTargetVolumeSimCacheData>(StorageObject);
-	}
-
 	return true;
 }
 
@@ -790,6 +783,7 @@ bool UNiagaraDataInterfaceRenderTargetVolume::SimCacheReadFrame(UObject* Storage
 
 	if (NDIRenderTargetVolumeLocal::GSimCacheUseOpenVDB)
 	{
+#if PLATFORM_WINDOWS
 		UVolumeCache *OpenVDBSimCacheData = CastChecked<UVolumeCache>(StorageObject);
 
 		TSharedPtr<FVolumeCacheData> VolumeCacheData = OpenVDBSimCacheData->GetData();
@@ -819,79 +813,80 @@ bool UNiagaraDataInterfaceRenderTargetVolume::SimCacheReadFrame(UObject* Storage
 						}
 					});
 		}
+#endif
 	}
 	else
 	{
-	UNDIRenderTargetVolumeSimCacheData* SimCacheData = CastChecked<UNDIRenderTargetVolumeSimCacheData>(StorageObject);
+		UNDIRenderTargetVolumeSimCacheData* SimCacheData = CastChecked<UNDIRenderTargetVolumeSimCacheData>(StorageObject);
 
-	const int FrameIndex = Interp >= 0.5f ? FrameB : FrameA;
-	if (!SimCacheData->Frames.IsValidIndex(FrameIndex))
-	{
-		return false;
-	}
+		const int FrameIndex = Interp >= 0.5f ? FrameB : FrameA;
+		if (!SimCacheData->Frames.IsValidIndex(FrameIndex))
+		{
+			return false;
+		}
 
-	FNDIRenderTargetVolumeSimCacheFrame* CacheFrame = &SimCacheData->Frames[FrameIndex];
+		FNDIRenderTargetVolumeSimCacheFrame* CacheFrame = &SimCacheData->Frames[FrameIndex];
 
-	InstanceData_GT->Size = CacheFrame->Size;
-	InstanceData_GT->Format = CacheFrame->Format;
-	InstanceData_GT->Filter = TextureFilter::TF_Default;
+		InstanceData_GT->Size = CacheFrame->Size;
+		InstanceData_GT->Format = CacheFrame->Format;
+		InstanceData_GT->Filter = TextureFilter::TF_Default;
 
-	PerInstanceTick(InstanceData_GT, SystemInstance, 0.0f);
-	PerInstanceTickPostSimulate(InstanceData_GT, SystemInstance, 0.0f);
+		PerInstanceTick(InstanceData_GT, SystemInstance, 0.0f);
+		PerInstanceTickPostSimulate(InstanceData_GT, SystemInstance, 0.0f);
 
-	if (InstanceData_GT->TargetTexture && CacheFrame->GetPixelData() != nullptr)
-	{
-		FNiagaraDataInterfaceProxyRenderTargetVolumeProxy* RT_Proxy = GetProxyAs<FNiagaraDataInterfaceProxyRenderTargetVolumeProxy>();
-		FTextureRenderTargetResource* RT_TargetTexture = InstanceData_GT->TargetTexture->GameThread_GetRenderTargetResource();
-		ENQUEUE_RENDER_COMMAND(NDIRenderTargetVolumeUpdate)
-		(
-			[RT_Proxy, RT_InstanceID=SystemInstance->GetId(), RT_TargetTexture, RT_CacheFrame=CacheFrame, RT_CompressionType=SimCacheData->CompressionType](FRHICommandListImmediate& RHICmdList)
-			{
-				if (FRenderTargetVolumeRWInstanceData_RenderThread* InstanceData_RT = RT_Proxy->SystemInstancesToProxyData_RT.Find(RT_InstanceID))
+		if (InstanceData_GT->TargetTexture && CacheFrame->GetPixelData() != nullptr)
+		{
+			FNiagaraDataInterfaceProxyRenderTargetVolumeProxy* RT_Proxy = GetProxyAs<FNiagaraDataInterfaceProxyRenderTargetVolumeProxy>();
+			FTextureRenderTargetResource* RT_TargetTexture = InstanceData_GT->TargetTexture->GameThread_GetRenderTargetResource();
+			ENQUEUE_RENDER_COMMAND(NDIRenderTargetVolumeUpdate)
+			(
+				[RT_Proxy, RT_InstanceID=SystemInstance->GetId(), RT_TargetTexture, RT_CacheFrame=CacheFrame, RT_CompressionType=SimCacheData->CompressionType](FRHICommandListImmediate& RHICmdList)
 				{
-					FUpdateTextureRegion3D UpdateRegion(FIntVector::ZeroValue, FIntVector::ZeroValue, InstanceData_RT->Size);
+					if (FRenderTargetVolumeRWInstanceData_RenderThread* InstanceData_RT = RT_Proxy->SystemInstancesToProxyData_RT.Find(RT_InstanceID))
+					{
+						FUpdateTextureRegion3D UpdateRegion(FIntVector::ZeroValue, FIntVector::ZeroValue, InstanceData_RT->Size);
 
-					FUpdateTexture3DData UpdateTexture = RHICmdList.BeginUpdateTexture3D(InstanceData_RT->RenderTarget->GetRHI(), 0, UpdateRegion);
+						FUpdateTexture3DData UpdateTexture = RHICmdList.BeginUpdateTexture3D(InstanceData_RT->RenderTarget->GetRHI(), 0, UpdateRegion);
 
-					const int32 SrcRowPitch = InstanceData_RT->Size.X * sizeof(FFloat16Color);
-					const int32 SrcDepthPitch = InstanceData_RT->Size.Y * SrcRowPitch;
-					if (UpdateTexture.RowPitch == SrcRowPitch && UpdateTexture.DepthPitch == SrcDepthPitch)
-					{
-					if (RT_CacheFrame->CompressedSize > 0)
-					{
-							FCompression::UncompressMemory(RT_CompressionType, UpdateTexture.Data, UpdateTexture.DataSizeBytes, RT_CacheFrame->GetPixelData(), RT_CacheFrame->CompressedSize);
-						}
-						else
+						const int32 SrcRowPitch = InstanceData_RT->Size.X * sizeof(FFloat16Color);
+						const int32 SrcDepthPitch = InstanceData_RT->Size.Y * SrcRowPitch;
+						if (UpdateTexture.RowPitch == SrcRowPitch && UpdateTexture.DepthPitch == SrcDepthPitch)
 						{
-							FMemory::Memcpy(UpdateTexture.Data, RT_CacheFrame->GetPixelData(), InstanceData_RT->Size.X * InstanceData_RT->Size.Y * sizeof(FFloat16Color));
+							if (RT_CacheFrame->CompressedSize > 0)
+							{
+								FCompression::UncompressMemory(RT_CompressionType, UpdateTexture.Data, UpdateTexture.DataSizeBytes, RT_CacheFrame->GetPixelData(), RT_CacheFrame->CompressedSize);
+							}
+							else
+							{
+								FMemory::Memcpy(UpdateTexture.Data, RT_CacheFrame->GetPixelData(), InstanceData_RT->Size.X * InstanceData_RT->Size.Y * sizeof(FFloat16Color));
+							}
 						}
-					}
-					else
-					{
-						TArray<uint8> Decompressed;
-						if (RT_CacheFrame->CompressedSize > 0)
+						else
 						{
-						Decompressed.AddUninitialized(sizeof(FFloat16Color) * InstanceData_RT->Size.X * InstanceData_RT->Size.Y * InstanceData_RT->Size.Z);
-						FCompression::UncompressMemory(RT_CompressionType, Decompressed.GetData(), Decompressed.Num(), RT_CacheFrame->GetPixelData(), RT_CacheFrame->CompressedSize);
-						}
+							TArray<uint8> Decompressed;
+							if (RT_CacheFrame->CompressedSize > 0)
+							{
+								Decompressed.AddUninitialized(sizeof(FFloat16Color) * InstanceData_RT->Size.X * InstanceData_RT->Size.Y * InstanceData_RT->Size.Z);
+								FCompression::UncompressMemory(RT_CompressionType, Decompressed.GetData(), Decompressed.Num(), RT_CacheFrame->GetPixelData(), RT_CacheFrame->CompressedSize);
+							}
 
-						const uint8* SrcData = Decompressed.Num() > 0 ? Decompressed.GetData() : RT_CacheFrame->GetPixelData();
-						for (int32 z = 0; z < InstanceData_RT->Size.Z; ++z)
-						{
-							uint8* DstData = UpdateTexture.Data + (z * UpdateTexture.DepthPitch);
-							for (int32 y = 0; y < InstanceData_RT->Size.Y; ++y)
+							const uint8* SrcData = Decompressed.Num() > 0 ? Decompressed.GetData() : RT_CacheFrame->GetPixelData();
+							for (int32 z = 0; z < InstanceData_RT->Size.Z; ++z)
 							{
-								FMemory::Memcpy(DstData, SrcData, SrcRowPitch);
-								SrcData += SrcRowPitch;
-								DstData += UpdateTexture.RowPitch;
+								uint8* DstData = UpdateTexture.Data + (z * UpdateTexture.DepthPitch);
+								for (int32 y = 0; y < InstanceData_RT->Size.Y; ++y)
+								{
+									FMemory::Memcpy(DstData, SrcData, SrcRowPitch);
+									SrcData += SrcRowPitch;
+									DstData += UpdateTexture.RowPitch;
+								}
 							}
+						}
+						RHICmdList.EndUpdateTexture3D(UpdateTexture);
 					}
-					}
-					RHICmdList.EndUpdateTexture3D(UpdateTexture);
 				}
-			}
-		);
-	}
+			);
+		}
 	}
 	return true;
 }
