- UI clean up for Niagara renderer properties to improve readability
- Important categories moved to the top, other categories collapsed by default
- Properties hide rather than grey when they can't be edited

diff --git a/Engine/Plugins/FX/Niagara/Source/Niagara/Private/NiagaraRibbonRendererProperties.cpp b/Engine/Plugins/FX/Niagara/Source/Niagara/Private/NiagaraRibbonRendererProperties.cpp
--- a/Engine/Plugins/FX/Niagara/Source/Niagara/Private/NiagaraRibbonRendererProperties.cpp
+++ b/Engine/Plugins/FX/Niagara/Source/Niagara/Private/NiagaraRibbonRendererProperties.cpp
@@ -64,8 +64,8 @@ UNiagaraRibbonRendererProperties::UNiagaraRibbonRendererProperties()
 	, WidthSegmentationCount(1)
 	, MultiPlaneCount(2)
 	, TubeSubdivisions(3)
-	, CurveTension(0.f)
 	, TessellationMode(ENiagaraRibbonTessellationMode::Automatic)
+	, CurveTension(0.f)
 	, TessellationFactor(16)
 	, bUseConstantFactor(false)
 	, TessellationAngle(15)
@@ -606,30 +606,6 @@ void UNiagaraRibbonRendererProperties::GetRendererFeedback(const FVersionedNiaga
 	}
 }
 
-
-bool UNiagaraRibbonRendererProperties::CanEditChange(const FProperty* InProperty) const
-{
-
-	if (InProperty->HasMetaData(TEXT("Category")) && InProperty->GetMetaData(TEXT("Category")).Contains("Tessellation"))
-	{
-		FName PropertyName = InProperty->GetFName();
-		if (PropertyName == GET_MEMBER_NAME_CHECKED(UNiagaraRibbonRendererProperties, CurveTension))
-		{
-			return TessellationMode != ENiagaraRibbonTessellationMode::Disabled;
-		}
-		if (PropertyName == GET_MEMBER_NAME_CHECKED(UNiagaraRibbonRendererProperties, TessellationFactor))
-		{
-			return TessellationMode == ENiagaraRibbonTessellationMode::Custom;
-		}
-		if (PropertyName == GET_MEMBER_NAME_CHECKED(UNiagaraRibbonRendererProperties, TessellationMode))
-		{
-			return Super::CanEditChange(InProperty);
-		}
-		return TessellationMode == ENiagaraRibbonTessellationMode::Custom;
-	}
-	return Super::CanEditChange(InProperty);
-}
-
 void UNiagaraRibbonRendererProperties::RenameVariable(const FNiagaraVariableBase& OldVariable, const FNiagaraVariableBase& NewVariable, const FVersionedNiagaraEmitter& InEmitter)
 {
 	Super::RenameVariable(OldVariable, NewVariable, InEmitter);
diff --git a/Engine/Plugins/FX/Niagara/Source/Niagara/Public/NiagaraMeshRendererProperties.h b/Engine/Plugins/FX/Niagara/Source/Niagara/Public/NiagaraMeshRendererProperties.h
--- a/Engine/Plugins/FX/Niagara/Source/Niagara/Public/NiagaraMeshRendererProperties.h
+++ b/Engine/Plugins/FX/Niagara/Source/Niagara/Public/NiagaraMeshRendererProperties.h
@@ -29,7 +29,8 @@ enum class ENiagaraMeshFacingMode : uint8
 };
 
 UENUM()
-enum class ENiagaraMeshPivotOffsetSpace : uint8 {
+enum class ENiagaraMeshPivotOffsetSpace : uint8
+{
 	/** The pivot offset is in the mesh's local space (default) */
 	Mesh,
 	/** The pivot offset is in the emitter's local space if the emitter is marked as local-space, or in world space otherwise */
@@ -41,7 +42,8 @@ enum class ENiagaraMeshPivotOffsetSpace : uint8 {
 };
 
 UENUM()
-enum class ENiagaraMeshLockedAxisSpace : uint8 {
+enum class ENiagaraMeshLockedAxisSpace : uint8
+{
 	/** The locked axis is in the emitter's local space if the emitter is marked as local-space, or in world space otherwise */
 	Simulation,
 	/** The locked axis is in world space */
@@ -61,11 +63,11 @@ public:
 	bool SerializeFromMismatchedTag(const struct FPropertyTag& Tag, FStructuredArchive::FSlot Slot);
 
 	/** Use this UMaterialInterface if set to a valid value. This will be subordinate to UserParamBinding if it is set to a valid user variable.*/
-	UPROPERTY(EditAnywhere, Category = "Mesh Rendering", meta = (EditCondition = "bOverrideMaterials"))
+	UPROPERTY(EditAnywhere, Category = "Material", meta = (EditCondition = "bOverrideMaterials"))
 	TObjectPtr<UMaterialInterface> ExplicitMat;
 
 	/** Use the UMaterialInterface bound to this user variable if it is set to a valid value. If this is bound to a valid value and ExplicitMat is also set, UserParamBinding wins.*/
-	UPROPERTY(EditAnywhere, Category = "Mesh Rendering", meta = (EditCondition = "bOverrideMaterials"))
+	UPROPERTY(EditAnywhere, Category = "Material", meta = (EditCondition = "bOverrideMaterials"))
 	FNiagaraUserParameterBinding UserParamBinding;
 
 	bool operator==(const FNiagaraMeshMaterialOverride& Other)const
@@ -214,7 +216,7 @@ public:
 	 * - If "Override Material" is not specified, the mesh's material is used. Override materials must have the Niagara Mesh Particles flag checked.
 	 * - If "Enable Mesh Flipbook" is specified, this mesh is assumed to be the first frame of the flipbook.
 	 */
-	UPROPERTY(EditAnywhere, Category = "Mesh Rendering", meta = (EditCondition = "!bEnableMeshFlipbook"))
+	UPROPERTY(EditAnywhere, Category = "Mesh Rendering", meta = (EditCondition = "!bEnableMeshFlipbook", EditConditionHides))
 	TArray<FNiagaraMeshRendererMeshProperties> Meshes;
 
 	/** Whether or not to draw a single element for the Emitter or to draw the particles.*/
@@ -230,11 +232,11 @@ public:
 	uint32 bOverrideMaterials : 1;
 
 	/** If true, the particles are only sorted when using a translucent material. */
-	UPROPERTY(EditAnywhere, Category = "Sorting")
+	UPROPERTY(EditAnywhere, Category = "Sorting", meta = (EditCondition = "SortMode != ENiagaraSortMode::None", EditConditionHides))
 	uint32 bSortOnlyWhenTranslucent : 1;
 
 	/** Sort precision to use when sorting is active. */
-	UPROPERTY(EditAnywhere, Category = "Sorting")
+	UPROPERTY(EditAnywhere, Category = "Sorting", meta = (EditCondition = "SortMode != ENiagaraSortMode::None", EditConditionHides))
 	ENiagaraRendererSortPrecision SortPrecision = ENiagaraRendererSortPrecision::Default;
 
 	/**
@@ -242,7 +244,7 @@ public:
 	Opaque materials will run latent when these features are used.
 	Translucent materials can choose if they want to use this frames or the previous frames data to match opaque draws.
 	*/
-	UPROPERTY(EditAnywhere, AdvancedDisplay, Category = "Sprite Rendering")
+	UPROPERTY(EditAnywhere, AdvancedDisplay, Category = "Mesh Rendering")
 	ENiagaraRendererGpuTranslucentLatency GpuTranslucentLatency = ENiagaraRendererGpuTranslucentLatency::ProjectDefault;
 
 	/** If true, blends the sub-image UV lookup with its next adjacent member using the fractional part of the SubImageIndex float value as the linear interpolation factor.*/
@@ -258,12 +260,12 @@ public:
 	uint32 bEnableCameraDistanceCulling : 1;
 
 	/** When checked, will treat 'ParticleMesh' as the first frame of the flipbook, and will use the other mesh flipbook options to find the other frames */
-	UPROPERTY(EditAnywhere, Category = "Mesh Flipbook")
+	UPROPERTY(EditAnywhere, Category = "Mesh Rendering", meta = (DisplayAfter = "MeshBoundsScale"))
 	uint32 bEnableMeshFlipbook : 1;
 
 	/** The materials to be used instead of the StaticMesh's materials. Note that each material must have the Niagara Mesh Particles flag checked. If the ParticleMesh
 	requires more materials than exist in this array or any entry in this array is set to None, we will use the ParticleMesh's existing Material instead.*/
-	UPROPERTY(EditAnywhere, Category = "Mesh Rendering", meta = (EditCondition = "bOverrideMaterials"))
+	UPROPERTY(EditAnywhere, Category = "Mesh Rendering", meta = (EditCondition = "bOverrideMaterials", EditConditionHides))
 	TArray<FNiagaraMeshMaterialOverride> OverrideMaterials;
 
 #if WITH_EDITORONLY_DATA
@@ -272,7 +274,7 @@ public:
 #endif
 
 	/** When using SubImage lookups for particles, this variable contains the number of columns in X and the number of rows in Y.*/
-	UPROPERTY(EditAnywhere, Category = "SubUV")
+	UPROPERTY(EditAnywhere, Category = "SubUV", meta = (EditCondition = "bSubImageBlend", EditConditionHides))
 	FVector2D SubImageSize;
 
 	/** Determines how the mesh orients itself relative to the camera. */
@@ -284,11 +286,11 @@ public:
 	uint32 bLockedAxisEnable : 1;
 
 	/** Arbitrary axis by which to lock facing rotations */
-	UPROPERTY(EditAnywhere, Category = "Mesh Rendering", meta = (EditCondition = "bLockedAxisEnable"))
+	UPROPERTY(EditAnywhere, Category = "Mesh Rendering", meta = (EditCondition = "bLockedAxisEnable", EditConditionHides))
 	FVector LockedAxis;
 
 	/** Specifies what space the locked axis is in */
-	UPROPERTY(EditAnywhere, Category = "Mesh Rendering", meta = (EditCondition = "bLockedAxisEnable"))
+	UPROPERTY(EditAnywhere, Category = "Mesh Rendering", meta = (EditCondition = "bLockedAxisEnable", EditConditionHides))
 	ENiagaraMeshLockedAxisSpace LockedAxisSpace;
 
 	/**
@@ -300,10 +302,10 @@ public:
 	UPROPERTY(EditAnywhere, Category = "Mesh Rendering", AdvancedDisplay)
 	FVector MeshBoundsScale = FVector::OneVector;
 
-	UPROPERTY(EditAnywhere, Category = "Visibility", meta = (EditCondition = "bEnableCameraDistanceCulling", ClampMin = 0.0f))
+	UPROPERTY(EditAnywhere, Category = "Visibility", meta = (EditCondition = "bEnableCameraDistanceCulling", EditConditionHides, ClampMin = 0.0f))
 	float MinCameraDistance;
 
-	UPROPERTY(EditAnywhere, Category = "Visibility", meta = (EditCondition = "bEnableCameraDistanceCulling", ClampMin = 0.0f))
+	UPROPERTY(EditAnywhere, Category = "Visibility", meta = (EditCondition = "bEnableCameraDistanceCulling", EditConditionHides, ClampMin = 0.0f))
 	float MaxCameraDistance = 1000.0f;
 
 	/** If a render visibility tag is present, particles whose tag matches this value will be visible in this renderer. */
@@ -402,7 +404,7 @@ public:
 	 * NOTE: The subsequent frames are expected to exist in the same content directory as the first frame of the flipbook, otherwise they
 	 * will not be found or used.
 	 */
-	UPROPERTY(EditAnywhere, Category = "Mesh Flipbook", meta = (EditCondition = "bEnableMeshFlipbook"))
+	UPROPERTY(EditAnywhere, Category = "Mesh Rendering", meta = (EditCondition = "bEnableMeshFlipbook", EditConditionHides, DisplayAfter = "bEnableMeshFlipbook"))
 	TObjectPtr<UStaticMesh> FirstFlipbookFrame;
 
 	/**
@@ -411,18 +413,18 @@ public:
 	 * the starting frame index. Otherwise, it will assume "Particle Mesh" is frame number 0, and that subsequent frames follow this format,
 	 * starting with frame number 1.
 	 */
-	UPROPERTY(EditAnywhere, Category = "Mesh Flipbook", meta = (EditCondition = "bEnableMeshFlipbook && FirstFlipbookFrame != nullptr"))
+	UPROPERTY(EditAnywhere, Category = "Mesh Rendering", meta = (EditCondition = "bEnableMeshFlipbook && FirstFlipbookFrame != nullptr", EditConditionHides, DisplayAfter = "bEnableMeshFlipbook"))
 	FString FlipbookSuffixFormat;
 
 	/**
 	* The number of digits to expect in the frame number of the flipbook page. A value of 1 will expect no leading zeros in the package names,
 	* and can also be used for names with frame numbers that extend to 10 and beyond (Example: Frame_1, Frame_2, ..., Frame_10, Frame_11, etc.)
 	*/
-	UPROPERTY(EditAnywhere, Category = "Mesh Flipbook", meta = (EditCondition = "bEnableMeshFlipbook && FirstFlipbookFrame != nullptr", ClampMin = 1, ClampMax = 10, NoSpinbox = true))
+	UPROPERTY(EditAnywhere, Category = "Mesh Rendering", meta = (EditCondition = "bEnableMeshFlipbook && FirstFlipbookFrame != nullptr", EditConditionHides, ClampMin = 1, ClampMax = 10, NoSpinbox = true, DisplayAfter = "bEnableMeshFlipbook"))
 	uint32 FlipbookSuffixNumDigits;
 
 	/** The number of frames (static meshes) to be included in the flipbook. */
-	UPROPERTY(EditAnywhere, Category = "Mesh Flipbook", meta = (EditCondition = "bEnableMeshFlipbook && FirstFlipbookFrame != nullptr", ClampMin = 1, NoSpinbox = true))
+	UPROPERTY(EditAnywhere, Category = "Mesh Rendering", meta = (EditCondition = "bEnableMeshFlipbook && FirstFlipbookFrame != nullptr", EditConditionHides, ClampMin = 1, NoSpinbox = true, DisplayAfter = "bEnableMeshFlipbook"))
 	uint32 NumFlipbookFrames;
 #endif
 
@@ -454,6 +456,4 @@ private:
 
 	UPROPERTY()
 	ENiagaraMeshPivotOffsetSpace PivotOffsetSpace_DEPRECATED;
-
-
 };
diff --git a/Engine/Plugins/FX/Niagara/Source/Niagara/Public/NiagaraRendererProperties.h b/Engine/Plugins/FX/Niagara/Source/Niagara/Public/NiagaraRendererProperties.h
--- a/Engine/Plugins/FX/Niagara/Source/Niagara/Public/NiagaraRendererProperties.h
+++ b/Engine/Plugins/FX/Niagara/Source/Niagara/Public/NiagaraRendererProperties.h
@@ -386,11 +386,11 @@ public:
 
 	/** By default, emitters are drawn in the order that they are added to the system. This value will allow you to control the order in a more fine-grained manner.
 	Materials of the same type (i.e. Transparent) will draw in order from lowest to highest within the system. The default value is 0.*/
-	UPROPERTY(EditAnywhere, Category = "Sort Order")
+	UPROPERTY(EditAnywhere, Category = "Rendering")
 	int32 SortOrderHint;
 
 	/** Hint about how to generate motion (velocity) vectors for this renderer. */
-	UPROPERTY(EditAnywhere, Category = "Motion Blur")
+	UPROPERTY(EditAnywhere, Category = "Rendering")
 	ENiagaraRendererMotionVectorSetting MotionVectorSetting;
 
 	/** Optional bool binding to dynamically enable / disable the renderer. */
diff --git a/Engine/Plugins/FX/Niagara/Source/Niagara/Public/NiagaraRibbonRendererProperties.h b/Engine/Plugins/FX/Niagara/Source/Niagara/Public/NiagaraRibbonRendererProperties.h
--- a/Engine/Plugins/FX/Niagara/Source/Niagara/Public/NiagaraRibbonRendererProperties.h
+++ b/Engine/Plugins/FX/Niagara/Source/Niagara/Public/NiagaraRibbonRendererProperties.h
@@ -128,19 +128,19 @@ struct FNiagaraRibbonUVSettings
 
 
 	/** Specifies how ribbon UVs are distributed along the length of a ribbon. */
-	UPROPERTY(EditAnywhere, Category = UVs, meta = (DisplayName="UV Mode", EditCondition = "!bEnablePerParticleUOverride"))
+	UPROPERTY(EditAnywhere, Category = UVs, meta = (DisplayName="UV Mode", EditCondition = "!bEnablePerParticleUOverride", EditConditionHides))
 	ENiagaraRibbonUVDistributionMode DistributionMode;
 
 	/** Specifies how UVs transition into life at the leading edge of the ribbon. */
-	UPROPERTY(EditAnywhere, Category = UVs, meta = (DisplayName="Leading Edge Transition", EditCondition = "!bEnablePerParticleUOverride && DistributionMode != ENiagaraRibbonUVDistributionMode::TiledFromStartOverRibbonLength"))
+	UPROPERTY(EditAnywhere, Category = UVs, meta = (DisplayName="Leading Edge Transition", EditCondition = "!bEnablePerParticleUOverride && DistributionMode != ENiagaraRibbonUVDistributionMode::TiledFromStartOverRibbonLength", EditConditionHides))
 	ENiagaraRibbonUVEdgeMode LeadingEdgeMode;
 
 	/** Specifies how UVs transition out of life at the trailing edge of the ribbon. */
-	UPROPERTY(EditAnywhere, Category = UVs, meta = (DisplayName="Trailing Edge Transition", EditCondition = "!bEnablePerParticleUOverride && DistributionMode != ENiagaraRibbonUVDistributionMode::TiledOverRibbonLength && DistributionMode != ENiagaraRibbonUVDistributionMode::TiledFromStartOverRibbonLength"))
+	UPROPERTY(EditAnywhere, Category = UVs, meta = (DisplayName="Trailing Edge Transition", EditCondition = "!bEnablePerParticleUOverride && DistributionMode != ENiagaraRibbonUVDistributionMode::TiledOverRibbonLength && DistributionMode != ENiagaraRibbonUVDistributionMode::TiledFromStartOverRibbonLength", EditConditionHides))
 	ENiagaraRibbonUVEdgeMode TrailingEdgeMode;
 
 	/** Specifies the length in world units to use when tiling UVs across the ribbon when using one of the tiled distribution modes. */
-	UPROPERTY(EditAnywhere, Category = UVs, meta = (EditCondition="!bEnablePerParticleUOverride && DistributionMode == ENiagaraRibbonUVDistributionMode::TiledOverRibbonLength || DistributionMode == ENiagaraRibbonUVDistributionMode::TiledFromStartOverRibbonLength"))
+	UPROPERTY(EditAnywhere, Category = UVs, meta = (EditCondition="!bEnablePerParticleUOverride && DistributionMode == ENiagaraRibbonUVDistributionMode::TiledOverRibbonLength || DistributionMode == ENiagaraRibbonUVDistributionMode::TiledFromStartOverRibbonLength", EditConditionHides))
 	float TilingLength;
 
 	/** Specifies an additional offset which is applied to the UV range */
@@ -204,7 +204,6 @@ public:
 	virtual void Serialize(FStructuredArchive::FRecord Record) override;
 #if WITH_EDITOR
 	virtual void PostEditChangeProperty(struct FPropertyChangedEvent& PropertyChangedEvent) override;
-	virtual bool CanEditChange(const FProperty* InProperty) const override;
 	virtual void RenameVariable(const FNiagaraVariableBase& OldVariable, const FNiagaraVariableBase& NewVariable, const FVersionedNiagaraEmitter& InEmitter) override;
 	virtual void RemoveVariable(const FNiagaraVariableBase& OldVariable, const FVersionedNiagaraEmitter& InEmitter) override;
 #endif
@@ -314,57 +313,56 @@ public:
 	  * MultiPlane will double geometry count to have triangles facing both sides. With this off MultiPlane will switch normal direction to face view.
 	  * 3D Ribbons will render like normal meshes with backface culling enabled.
 	  */
-	UPROPERTY(EditAnywhere, Category = "Ribbon Shape", meta = (EditCondition = "Shape != ENiagaraRibbonShapeMode::Plane"))
+	UPROPERTY(EditAnywhere, Category = "Ribbon Shape", meta = (EditCondition = "Shape != ENiagaraRibbonShapeMode::Plane", EditConditionHides))
 	bool bEnableAccurateGeometry;
 
 	/** Tessellation factor to apply to the width of the ribbon.
 	* Ranges from 1 to 16. Greater values increase amount of tessellation.
 	*/
-	UPROPERTY(EditAnywhere, Category = "Ribbon Shape", meta = (EditCondition = "Shape == ENiagaraRibbonShapeMode::Plane || Shape == ENiagaraRibbonShapeMode::MultiPlane", ClampMin = "1", ClampMax = "16"))
+	UPROPERTY(EditAnywhere, Category = "Ribbon Shape", meta = (EditCondition = "Shape == ENiagaraRibbonShapeMode::Plane || Shape == ENiagaraRibbonShapeMode::MultiPlane", ClampMin = "1", ClampMax = "16", EditConditionHides))
 	int32 WidthSegmentationCount;
 
 	/** Number of planes in multiplane shape. Evenly distributed from 0-90 or 0-180 degrees off camera facing depending on setting */
-	UPROPERTY(EditAnywhere, Category = "Ribbon Shape", meta = (EditCondition = "Shape == ENiagaraRibbonShapeMode::MultiPlane", ClampMin = "2", ClampMax = "8"))
+	UPROPERTY(EditAnywhere, Category = "Ribbon Shape", meta = (EditCondition = "Shape == ENiagaraRibbonShapeMode::MultiPlane", ClampMin = "2", ClampMax = "8", EditConditionHides))
 	int32 MultiPlaneCount;
 
 	/** Number of vertices/faces in a tube.  */
-	UPROPERTY(EditAnywhere, Category = "Ribbon Shape", meta = (EditCondition = "Shape == ENiagaraRibbonShapeMode::Tube", ClampMin = "3", ClampMax = "16"))
+	UPROPERTY(EditAnywhere, Category = "Ribbon Shape", meta = (EditCondition = "Shape == ENiagaraRibbonShapeMode::Tube", ClampMin = "3", ClampMax = "16", EditConditionHides))
 	int32 TubeSubdivisions;
 
 	/** Vertices for a cross section of the ribbon in custom shape mode. */
-	UPROPERTY(EditAnywhere, Category = "Ribbon Shape", meta = (EditCondition = "Shape == ENiagaraRibbonShapeMode::Custom"))
+	UPROPERTY(EditAnywhere, Category = "Ribbon Shape", meta = (EditCondition = "Shape == ENiagaraRibbonShapeMode::Custom", EditConditionHides))
 	TArray<FNiagaraRibbonShapeCustomVertex> CustomVertices;
 
+	/** Defines the tessellation mode allowing custom tessellation parameters or disabling tessellation entirely. */
+	UPROPERTY(EditAnywhere, Category = "Ribbon Tessellation", meta = (DisplayName = "Mode"))
+	ENiagaraRibbonTessellationMode TessellationMode;
 
 	/** Defines the curve tension, or how long the curve's tangents are.
 	  * Ranges from 0 to 1. The higher the value, the sharper the curve becomes.
 	  */
-	UPROPERTY(EditAnywhere, Category = "Tessellation", meta = (ClampMin = "0", ClampMax = "0.99"))
+	UPROPERTY(EditAnywhere, Category = "Ribbon Tessellation", meta = (EditCondition = "TessellationMode != ENiagaraRibbonTessellationMode::Disabled", EditConditionHides, ClampMin = "0", ClampMax = "0.99"))
 	float CurveTension;
 
-	/** Defines the tessellation mode allowing custom tessellation parameters or disabling tessellation entirely. */
-	UPROPERTY(EditAnywhere, Category = "Tessellation", meta = (DisplayName = "Mode"))
-	ENiagaraRibbonTessellationMode TessellationMode;
-
 	/** Custom tessellation factor.
 	  * Ranges from 1 to 16. Greater values increase amount of tessellation.
 	  */
-	UPROPERTY(EditAnywhere, Category = "Tessellation", meta = (DisplayName = "Max Tessellation Factor", ClampMin = "1", ClampMax = "16"))
+	UPROPERTY(EditAnywhere, Category = "Ribbon Tessellation", meta = (DisplayName = "Max Tessellation Factor", EditCondition = "TessellationMode == ENiagaraRibbonTessellationMode::Custom", EditConditionHides, ClampMin = "1", ClampMax = "16"))
 	int32 TessellationFactor;
 
 	/** If checked, use the above constant factor. Otherwise, adaptively select the tessellation factor based on the below parameters. */
-	UPROPERTY(EditAnywhere, Category = "Tessellation")
+	UPROPERTY(EditAnywhere, Category = "Ribbon Tessellation", meta = (EditCondition = "TessellationMode == ENiagaraRibbonTessellationMode::Custom", EditConditionHides))
 	bool bUseConstantFactor;
 
 	/** Defines the angle in degrees at which tessellation occurs.
 	  * Ranges from 1 to 180. Smaller values increase amount of tessellation.
 	  * If set to 0, use the maximum tessellation set above.
 	  */
-	UPROPERTY(EditAnywhere, Category = "Tessellation", meta = (EditCondition = "!bUseConstantFactor", ClampMin = "0", ClampMax = "180", UIMin = "1", UIMax = "180"))
+	UPROPERTY(EditAnywhere, Category = "Ribbon Tessellation", meta = (EditCondition = "TessellationMode == ENiagaraRibbonTessellationMode::Custom && !bUseConstantFactor", ClampMin = "0", ClampMax = "180", UIMin = "1", UIMax = "180", EditConditionHides))
 	float TessellationAngle;
 
 	/** If checked, use the ribbon's screen space percentage to adaptively adjust the tessellation factor. */
-	UPROPERTY(EditAnywhere, Category = "Tessellation", meta = (DisplayName = "Screen Space", EditCondition = "!bUseConstantFactor"))
+	UPROPERTY(EditAnywhere, Category = "Ribbon Tessellation", meta = (DisplayName = "Screen Space", EditCondition = "TessellationMode == ENiagaraRibbonTessellationMode::Custom && !bUseConstantFactor", EditConditionHides))
 	bool bScreenSpaceTessellation;
 
 	/** Which attribute should we use for position when generating ribbons?*/
diff --git a/Engine/Plugins/FX/Niagara/Source/Niagara/Public/NiagaraSpriteRendererProperties.h b/Engine/Plugins/FX/Niagara/Source/Niagara/Public/NiagaraSpriteRendererProperties.h
--- a/Engine/Plugins/FX/Niagara/Source/Niagara/Public/NiagaraSpriteRendererProperties.h
+++ b/Engine/Plugins/FX/Niagara/Source/Niagara/Public/NiagaraSpriteRendererProperties.h
@@ -187,7 +187,7 @@ public:
 	ENiagaraSortMode SortMode;
 
 	/** When using SubImage lookups for particles, this variable contains the number of columns in X and the number of rows in Y.*/
-	UPROPERTY(EditAnywhere, Category = "SubUV")
+	UPROPERTY(EditAnywhere, Category = "SubUV", meta = (EditCondition = "bSubImageBlend", EditConditionHides, DisplayAfter = bSubImageBlend))
 	FVector2D SubImageSize;
 
 	/** If true, blends the sub-image UV lookup with its next adjacent member using the fractional part of the SubImageIndex float value as the linear interpolation factor.*/
@@ -199,11 +199,11 @@ public:
 	uint32 bRemoveHMDRollInVR : 1;
 
 	/** If true, the particles are only sorted when using a translucent material. */
-	UPROPERTY(EditAnywhere, Category = "Sorting")
+	UPROPERTY(EditAnywhere, Category = "Sorting", meta = (EditCondition = "SortMode != ENiagaraSortMode::None", EditConditionHides))
 	uint32 bSortOnlyWhenTranslucent : 1;
 
 	/** Sort precision to use when sorting is active. */
-	UPROPERTY(EditAnywhere, Category = "Sorting")
+	UPROPERTY(EditAnywhere, Category = "Sorting", meta = (EditCondition = "SortMode != ENiagaraSortMode::None", EditConditionHides))
 	ENiagaraRendererSortPrecision SortPrecision = ENiagaraRendererSortPrecision::Default;
 
 	/**
@@ -229,7 +229,7 @@ public:
 	When pixel coverage is enabled this allows you to control the blend of the pixel coverage color adjustment.
 	i.e. 1.0 = full, 0.5 = 50/50 blend, 0.0 = none
 	*/
-	UPROPERTY(EditAnywhere, Category = "Sprite Rendering", meta = (EditCondition = "PixelCoverageMode != ENiagaraRendererPixelCoverageMode::Disabled", UIMin=0.0f, UIMax=1.0f))
+	UPROPERTY(EditAnywhere, Category = "Sprite Rendering", meta = (EditCondition = "PixelCoverageMode != ENiagaraRendererPixelCoverageMode::Disabled", UIMin=0.0f, UIMax=1.0f, EditConditionHides))
 	float PixelCoverageBlend = 1.0f;
 
 	/** When FacingMode is FacingCameraDistanceBlend, the distance at which the sprite is fully facing the camera plane. */
@@ -244,10 +244,10 @@ public:
 	UPROPERTY(EditAnywhere, Category = "Visibility")
 	uint32 bEnableCameraDistanceCulling : 1;
 
-	UPROPERTY(EditAnywhere, Category = "Visibility", meta = (EditCondition = "bEnableCameraDistanceCulling", UIMin = 0.0f))
+	UPROPERTY(EditAnywhere, Category = "Visibility", meta = (EditCondition = "bEnableCameraDistanceCulling", UIMin = 0.0f, EditConditionHides))
 	float MinCameraDistance;
 
-	UPROPERTY(EditAnywhere, Category = "Visibility", meta = (EditCondition = "bEnableCameraDistanceCulling", UIMin = 0.0f))
+	UPROPERTY(EditAnywhere, Category = "Visibility", meta = (EditCondition = "bEnableCameraDistanceCulling", UIMin = 0.0f, EditConditionHides))
 	float MaxCameraDistance = 1000.0f;
 
 	/** If a render visibility tag is present, particles whose tag matches this value will be visible in this renderer. */
@@ -369,7 +369,7 @@ public:
 	bool bUseMaterialCutoutTexture;
 
 	/** Texture to generate bounding geometry from.	*/
-	UPROPERTY(EditAnywhere, Category="Cutout", meta = (EditCondition = "!bUseMaterialCutoutTexture"))
+	UPROPERTY(EditAnywhere, Category="Cutout", meta = (EditCondition = "!bUseMaterialCutoutTexture", EditConditionHides))
 	TObjectPtr<UTexture2D> CutoutTexture;
 
 	/**
diff --git a/Engine/Plugins/FX/Niagara/Source/NiagaraEditorWidgets/Private/DetailCustomizations/NiagaraMeshRendererDetails.cpp b/Engine/Plugins/FX/Niagara/Source/NiagaraEditorWidgets/Private/DetailCustomizations/NiagaraMeshRendererDetails.cpp
--- a/Engine/Plugins/FX/Niagara/Source/NiagaraEditorWidgets/Private/DetailCustomizations/NiagaraMeshRendererDetails.cpp
+++ b/Engine/Plugins/FX/Niagara/Source/NiagaraEditorWidgets/Private/DetailCustomizations/NiagaraMeshRendererDetails.cpp
@@ -87,10 +87,7 @@ public:
 
 	virtual void Tick(float DeltaTime) override {}
 	virtual bool RequiresTick() const override { return false; }
-	virtual bool InitiallyCollapsed() const override
-	{
-		return true;
-	}
+	virtual bool InitiallyCollapsed() const override { return true; }
 	virtual TSharedPtr<IPropertyHandle> GetPropertyHandle() const override { return ParentProperty; }
 
 	void Rebuild()
@@ -116,6 +113,38 @@ void FNiagaraMeshRendererDetails::CustomizeDetails(IDetailLayoutBuilder& DetailB
 {
 	LayoutBuilder = &DetailBuilder;
 
+	// Categories to put first
+	{
+		static const FName CategoriesToCollapse[] =
+		{
+			FName("Mesh Rendering"),
+			FName("Sorting"),
+			FName("Visibility"),
+		};
+
+		for (const FName& Category : CategoriesToCollapse)
+		{
+			IDetailCategoryBuilder& CategoryBuilder = DetailBuilder.EditCategory(Category);
+		}
+	}
+
+	// Collapse default categories
+	{
+		static const FName CategoriesToCollapse[] =
+		{
+			FName("Bindings"),
+			FName("Rendering"),
+			FName("Scalability"),
+			FName("SubUV"),
+		};
+
+		for (const FName& Category : CategoriesToCollapse)
+		{
+			IDetailCategoryBuilder& CategoryBuilder = DetailBuilder.EditCategory(Category);
+			CategoryBuilder.InitiallyCollapsed(true);
+		}
+	}
+
 	{
 		static const FName MeshesName = TEXT("Meshes");
 		static const FName EnableMeshFlipbookName = TEXT("bEnableMeshFlipbook");
@@ -160,7 +189,6 @@ void FNiagaraMeshRendererDetails::CustomizeDetails(IDetailLayoutBuilder& DetailB
 		TSharedRef<FDetailArrayBuilder> OverrideMaterialsBuilder = MakeShared<FDetailArrayBuilder>(OverrideMaterialsProperty.ToSharedRef(), true, bEnableMaterialOverrides, true);
 		OverrideMaterialsBuilder->OnGenerateArrayElementWidget(FOnGenerateArrayElementWidget::CreateSP(this, &FNiagaraMeshRendererDetails::OnGenerateMaterialOverrideWidget));
 
-
 		CategoryBuilder.AddProperty(EnableOverrideMaterialsProperty);
 		CategoryBuilder.AddCustomBuilder(OverrideMaterialsBuilder);
 	}
@@ -187,4 +215,4 @@ void FNiagaraMeshRendererDetails::OnGenerateMaterialOverrideWidget(TSharedRef<IP
 {
 	IDetailPropertyRow& Row = ChildrenBuilder.AddProperty(Property);
 	Row.IsEnabled(bEnableMaterialOverrides);
-}
\ No newline at end of file
+}
diff --git a/Engine/Plugins/FX/Niagara/Source/NiagaraEditorWidgets/Private/DetailCustomizations/NiagaraRibbonRendererDetails.cpp b/Engine/Plugins/FX/Niagara/Source/NiagaraEditorWidgets/Private/DetailCustomizations/NiagaraRibbonRendererDetails.cpp
new file mode 100644
--- /dev/null
+++ b/Engine/Plugins/FX/Niagara/Source/NiagaraEditorWidgets/Private/DetailCustomizations/NiagaraRibbonRendererDetails.cpp
@@ -0,0 +1,50 @@
+// Copyright Epic Games, Inc. All Rights Reserved.
+
+#include "NiagaraRibbonRendererDetails.h"
+#include "DetailLayoutBuilder.h"
+#include "DetailCategoryBuilder.h"
+#include "IDetailChildrenBuilder.h"
+#include "DetailWidgetRow.h"
+#include "PropertyCustomizationHelpers.h"
+#include "Widgets/Layout/SExpandableArea.h"
+
+TSharedRef<IDetailCustomization> FNiagaraRibbonRendererDetails::MakeInstance()
+{
+	return MakeShared<FNiagaraRibbonRendererDetails>();
+}
+
+void FNiagaraRibbonRendererDetails::CustomizeDetails(IDetailLayoutBuilder& DetailBuilder)
+{
+	// Categories to put first
+	{
+		static const FName CategoriesToCollapse[] =
+		{
+			FName("Ribbon Rendering"),
+			FName("Ribbon Shape"),
+			FName("Ribbon Tessellation"),
+			FName("Sorting"),
+			FName("Visibility"),
+		};
+
+		for (const FName& Category : CategoriesToCollapse)
+		{
+			IDetailCategoryBuilder& CategoryBuilder = DetailBuilder.EditCategory(Category);
+		}
+	}
+
+	// Collapse default categories
+	{
+		static const FName CategoriesToCollapse[] =
+		{
+			FName("Bindings"),
+			FName("Rendering"),
+			FName("Scalability"),
+		};
+
+		for (const FName& Category : CategoriesToCollapse)
+		{
+			IDetailCategoryBuilder& CategoryBuilder = DetailBuilder.EditCategory(Category);
+			CategoryBuilder.InitiallyCollapsed(true);
+		}
+	}
+}
diff --git a/Engine/Plugins/FX/Niagara/Source/NiagaraEditorWidgets/Private/DetailCustomizations/NiagaraRibbonRendererDetails.h b/Engine/Plugins/FX/Niagara/Source/NiagaraEditorWidgets/Private/DetailCustomizations/NiagaraRibbonRendererDetails.h
new file mode 100644
--- /dev/null
+++ b/Engine/Plugins/FX/Niagara/Source/NiagaraEditorWidgets/Private/DetailCustomizations/NiagaraRibbonRendererDetails.h
@@ -0,0 +1,15 @@
+// Copyright Epic Games, Inc. All Rights Reserved.
+
+#pragma once
+
+#include "IDetailCustomization.h"
+
+class IDetailLayoutBuilder;
+
+class FNiagaraRibbonRendererDetails : public IDetailCustomization
+{
+public:
+	static TSharedRef<IDetailCustomization> MakeInstance();
+
+	virtual void CustomizeDetails(IDetailLayoutBuilder& DetailBuilder) override;
+};
\ No newline at end of file
diff --git a/Engine/Plugins/FX/Niagara/Source/NiagaraEditorWidgets/Private/DetailCustomizations/NiagaraSpriteRendererDetails.cpp b/Engine/Plugins/FX/Niagara/Source/NiagaraEditorWidgets/Private/DetailCustomizations/NiagaraSpriteRendererDetails.cpp
new file mode 100644
--- /dev/null
+++ b/Engine/Plugins/FX/Niagara/Source/NiagaraEditorWidgets/Private/DetailCustomizations/NiagaraSpriteRendererDetails.cpp
@@ -0,0 +1,50 @@
+// Copyright Epic Games, Inc. All Rights Reserved.
+
+#include "NiagaraSpriteRendererDetails.h"
+#include "DetailLayoutBuilder.h"
+#include "DetailCategoryBuilder.h"
+#include "IDetailChildrenBuilder.h"
+#include "DetailWidgetRow.h"
+#include "PropertyCustomizationHelpers.h"
+#include "Widgets/Layout/SExpandableArea.h"
+
+TSharedRef<IDetailCustomization> FNiagaraSpriteRendererDetails::MakeInstance()
+{
+	return MakeShared<FNiagaraSpriteRendererDetails>();
+}
+
+void FNiagaraSpriteRendererDetails::CustomizeDetails(IDetailLayoutBuilder& DetailBuilder)
+{
+	// Categories to put first
+	{
+		static const FName CategoriesToCollapse[] =
+		{
+			FName("Sprite Rendering"),
+			FName("Sorting"),
+			FName("Visibility"),
+		};
+
+		for (const FName& Category : CategoriesToCollapse)
+		{
+			IDetailCategoryBuilder& CategoryBuilder = DetailBuilder.EditCategory(Category);
+		}
+	}
+
+	// Collapse default categories
+	{
+		static const FName CategoriesToCollapse[] =
+		{
+			FName("SubUV"),
+			FName("Cutout"),
+			FName("Bindings"),
+			FName("Rendering"),
+			FName("Scalability"),
+		};
+
+		for (const FName& Category : CategoriesToCollapse)
+		{
+			IDetailCategoryBuilder& CategoryBuilder = DetailBuilder.EditCategory(Category);
+			CategoryBuilder.InitiallyCollapsed(true);
+		}
+	}
+}
diff --git a/Engine/Plugins/FX/Niagara/Source/NiagaraEditorWidgets/Private/DetailCustomizations/NiagaraSpriteRendererDetails.h b/Engine/Plugins/FX/Niagara/Source/NiagaraEditorWidgets/Private/DetailCustomizations/NiagaraSpriteRendererDetails.h
new file mode 100644
--- /dev/null
+++ b/Engine/Plugins/FX/Niagara/Source/NiagaraEditorWidgets/Private/DetailCustomizations/NiagaraSpriteRendererDetails.h
@@ -0,0 +1,15 @@
+// Copyright Epic Games, Inc. All Rights Reserved.
+
+#pragma once
+
+#include "IDetailCustomization.h"
+
+class IDetailLayoutBuilder;
+
+class FNiagaraSpriteRendererDetails : public IDetailCustomization
+{
+public:	
+	static TSharedRef<IDetailCustomization> MakeInstance();
+
+	virtual void CustomizeDetails(IDetailLayoutBuilder& DetailBuilder) override;
+};
\ No newline at end of file
diff --git a/Engine/Plugins/FX/Niagara/Source/NiagaraEditorWidgets/Private/NiagaraEditorWidgetsModule.cpp b/Engine/Plugins/FX/Niagara/Source/NiagaraEditorWidgets/Private/NiagaraEditorWidgetsModule.cpp
--- a/Engine/Plugins/FX/Niagara/Source/NiagaraEditorWidgets/Private/NiagaraEditorWidgetsModule.cpp
+++ b/Engine/Plugins/FX/Niagara/Source/NiagaraEditorWidgets/Private/NiagaraEditorWidgetsModule.cpp
@@ -13,6 +13,8 @@
 #include "DetailCustomizations/NiagaraDataInterfaceMeshRendererInfoDetails.h"
 #include "DetailCustomizations/NiagaraDataInterfaceSpriteRendererInfoDetails.h"
 #include "DetailCustomizations/NiagaraMeshRendererDetails.h"
+#include "DetailCustomizations/NiagaraRibbonRendererDetails.h"
+#include "DetailCustomizations/NiagaraSpriteRendererDetails.h"
 #include "ViewModels/NiagaraSystemViewModel.h"
 #include "SNiagaraOverviewGraph.h"
 #include "NiagaraEditorWidgetsUtilities.h"
@@ -132,6 +134,8 @@ void FNiagaraEditorWidgetsModule::StartupModule()
 	PropertyModule.RegisterCustomClassLayout("NiagaraDataInterfaceMeshRendererInfo", FOnGetDetailCustomizationInstance::CreateStatic(&FNiagaraDataInterfaceMeshRendererInfoDetails::MakeInstance));
 	PropertyModule.RegisterCustomClassLayout("NiagaraDataInterfaceSpriteRendererInfo", FOnGetDetailCustomizationInstance::CreateStatic(&FNiagaraDataInterfaceSpriteRendererInfoDetails::MakeInstance));
 	PropertyModule.RegisterCustomClassLayout("NiagaraMeshRendererProperties", FOnGetDetailCustomizationInstance::CreateStatic(&FNiagaraMeshRendererDetails::MakeInstance));
+	PropertyModule.RegisterCustomClassLayout("NiagaraRibbonRendererProperties", FOnGetDetailCustomizationInstance::CreateStatic(&FNiagaraRibbonRendererDetails::MakeInstance));
+	PropertyModule.RegisterCustomClassLayout("NiagaraSpriteRendererProperties", FOnGetDetailCustomizationInstance::CreateStatic(&FNiagaraSpriteRendererDetails::MakeInstance));
 
 	ReinitializeStyleCommand = IConsoleManager::Get().RegisterConsoleCommand(
 		TEXT("fx.NiagaraEditorWidgets.ReinitializeStyle"),
