- Add Gpu Feature list to verbose debug hud
- Add Gpu latency information to debug hud, i.e. ones where this frame translucent is requested and fails or we are using opaque.

diff --git a/Engine/Plugins/FX/Niagara/Source/Niagara/Private/NiagaraDebugHud.cpp b/Engine/Plugins/FX/Niagara/Source/Niagara/Private/NiagaraDebugHud.cpp
--- a/Engine/Plugins/FX/Niagara/Source/Niagara/Private/NiagaraDebugHud.cpp
+++ b/Engine/Plugins/FX/Niagara/Source/Niagara/Private/NiagaraDebugHud.cpp
@@ -4,15 +4,17 @@
 #include "Engine/Engine.h"
 #include "GameFramework/Pawn.h"
 #include "NiagaraComponent.h"
+#include "NiagaraComputeExecutionContext.h"
+#include "NiagaraDataSetDebugAccessor.h"
 #include "NiagaraDataSetReadback.h"
 #include "NiagaraFunctionLibrary.h"
 #include "NiagaraGpuComputeDispatchInterface.h"
-#include "NiagaraComputeExecutionContext.h"
+#include "NiagaraMeshRendererProperties.h"
 #include "NiagaraScript.h"
+#include "NiagaraSpriteRendererProperties.h"
 #include "NiagaraSystem.h"
 #include "NiagaraSystemInstanceController.h"
 #include "NiagaraWorldManager.h"
-#include "NiagaraDataSetDebugAccessor.h"
 
 #include "Components/LineBatchComponent.h"
 #include "Debug/DebugDrawService.h"
@@ -659,6 +661,117 @@ namespace NiagaraDebugLocal
 			}
 		}
 	}
+
+	template<typename TOutput>
+	void BuildGpuHudInformation(TOutput& Output, UNiagaraComponent* NiagaraComponent, FNiagaraSystemInstance* SystemInstance, ERHIFeatureLevel::Type FeatureLevel)
+	{
+		static UEnum* GpuComputeTickStageEnum = StaticEnum<ENiagaraGpuComputeTickStage::Type>();
+		const FNiagaraSystemGpuComputeProxy* SystemInstanceComputeProxy = SystemInstance->GetSystemGpuComputeProxy();
+		if (GpuComputeTickStageEnum == nullptr || SystemInstanceComputeProxy == nullptr)
+		{
+			return;
+		}
+
+		const ENiagaraGpuComputeTickStage::Type GpuTickStage = SystemInstanceComputeProxy->GetComputeTickStage();
+		Output.Appendf(TEXT("GpuTickStage - %s\n"), *GpuComputeTickStageEnum->GetNameStringByValue(GpuTickStage));
+		if (Settings.SystemDebugVerbosity == ENiagaraDebugHudVerbosity::Verbose)
+		{
+			TStringBuilder<128> GpuFeaturesBuilder;
+			if (SystemInstance->RequiresDistanceFieldData())
+			{
+				GpuFeaturesBuilder.Append(TEXT(" DistanceFieldData"));
+			}
+			if (SystemInstance->RequiresDepthBuffer())
+			{
+				GpuFeaturesBuilder.Append(TEXT(" DepthBuffer"));
+			}
+			if (SystemInstance->RequiresEarlyViewData())
+			{
+				GpuFeaturesBuilder.Append(TEXT(" EarlyViewData"));
+			}
+			if (SystemInstance->RequiresViewUniformBuffer())
+			{
+				GpuFeaturesBuilder.Append(TEXT(" ViewUniformBuffer"));
+			}
+			if (SystemInstance->RequiresRayTracingScene())
+			{
+				GpuFeaturesBuilder.Append(TEXT(" RayTracingScene"));
+			}
+			if (GpuFeaturesBuilder.Len() > 0)
+			{
+				Output.Appendf(TEXT("GpuFeatures -%s\n"), *GpuFeaturesBuilder);
+			}
+		}
+
+		// Attempt to give feedback to the user if an emitter will be latent or not on the GPU
+		if (GpuTickStage == ENiagaraGpuComputeTickStage::PostOpaqueRender)
+		{
+			TStringBuilder<128> GpuLatentBuilder;
+
+			for (const TSharedRef<FNiagaraEmitterInstance, ESPMode::ThreadSafe>& EmitterInstance : SystemInstance->GetEmitters())
+			{
+				FVersionedNiagaraEmitterData* EmitterData = EmitterInstance->GetCachedEmitterData();
+				UNiagaraEmitter* NiagaraEmitter = EmitterInstance->GetCachedEmitter().Emitter;
+				if (EmitterData == nullptr || NiagaraEmitter == nullptr)
+				{
+					continue;
+				}
+
+				bool bLowLatencyFailed = false;
+				EmitterData->ForEachEnabledRenderer(
+					[&](UNiagaraRendererProperties* RenderProperties)
+					{
+						ENiagaraRendererGpuTranslucentLatency RequestedLatency = ENiagaraRendererGpuTranslucentLatency::ProjectDefault;
+						if (UNiagaraMeshRendererProperties* MeshRenderProperties = Cast<UNiagaraMeshRendererProperties>(RenderProperties))
+						{
+							RequestedLatency = MeshRenderProperties->GpuTranslucentLatency;
+						}
+						else if (UNiagaraSpriteRendererProperties* SpriteRendererProperties = Cast<UNiagaraSpriteRendererProperties>(RenderProperties))
+						{
+							RequestedLatency = SpriteRendererProperties->GpuTranslucentLatency;
+						}
+						else
+						{
+							// Renderer does not support low latency
+							return;
+						}
+
+						const bool bWantsThisFrameData = UNiagaraRendererProperties::ShouldGpuTranslucentThisFrame(RequestedLatency);
+
+						bool bSupportsThisFrameData =
+							!NiagaraComponent->bCastVolumetricTranslucentShadow &&
+							UNiagaraRendererProperties::IsGpuTranslucentThisFrame(FeatureLevel, RequestedLatency);
+
+						if (bSupportsThisFrameData)
+						{
+							TArray<UMaterialInterface*> UsedMaterials;
+							RenderProperties->GetUsedMaterials(&EmitterInstance.Get(), UsedMaterials);
+							for (UMaterialInterface* Material : UsedMaterials)
+							{
+								if (Material && !IsTranslucentBlendMode(Material->GetBlendMode()) )
+								{
+									bSupportsThisFrameData = false;
+									break;
+								}
+							}
+						}
+						bLowLatencyFailed |= bWantsThisFrameData != bSupportsThisFrameData;
+					}
+				);
+
+				if (bLowLatencyFailed)
+				{
+					GpuLatentBuilder.AppendChar(' ');
+					GpuLatentBuilder.Append(NiagaraEmitter->GetUniqueEmitterName());
+				}
+			}
+
+			if ( GpuLatentBuilder.Len() > 0 )
+			{
+				Output.Appendf(TEXT("GpuHasLatentEmitters -%s\n"), *GpuLatentBuilder);
+			}
+		}
+	}
 }
 
 FNiagaraDebugHud::FNiagaraDebugHud(UWorld* World)
@@ -2667,11 +2780,7 @@ void FNiagaraDebugHud::DrawComponents(FNiagaraWorldManager* WorldManager, UCanva
 						{
 							StringBuilder.Appendf(TEXT("TickGroup - %s\n"), *TickingGroupEnum->GetNameStringByValue(SystemInstance->CalculateTickGroup()));
 						}
-						static UEnum* GpuComputeTickStageEnum = StaticEnum<ENiagaraGpuComputeTickStage::Type>();
-						if (const FNiagaraSystemGpuComputeProxy* SystemInstanceComputeProxy = SystemInstance->GetSystemGpuComputeProxy())
-						{
-							StringBuilder.Appendf(TEXT("GpuTickStage - %s\n"), *GpuComputeTickStageEnum->GetNameStringByValue(SystemInstanceComputeProxy->GetComputeTickStage()));
-						}
+						BuildGpuHudInformation(StringBuilder, NiagaraComponent, SystemInstance, World->FeatureLevel);
 					}
 
 					int64 TotalBytes = 0;
diff --git a/Engine/Plugins/FX/Niagara/Source/Niagara/Private/NiagaraRendererMeshes.cpp b/Engine/Plugins/FX/Niagara/Source/Niagara/Private/NiagaraRendererMeshes.cpp
--- a/Engine/Plugins/FX/Niagara/Source/Niagara/Private/NiagaraRendererMeshes.cpp
+++ b/Engine/Plugins/FX/Niagara/Source/Niagara/Private/NiagaraRendererMeshes.cpp
@@ -289,6 +289,7 @@ void FNiagaraRendererMeshes::PrepareParticleMeshRenderData(FParticleMeshRenderDa
 	}
 
 	// Check if any materials are translucent and if we can pickup the low latency data
+	// If these conditions change please update the DebugHUD display also to reflect it
 	const bool bIsWireframe = AllowDebugViewmodes() && ViewFamily.EngineShowFlags.Wireframe;
 	ParticleMeshRenderData.bIsGpuLowLatencyTranslucency =
 		bGpuLowLatencyTranslucency &&
diff --git a/Engine/Plugins/FX/Niagara/Source/Niagara/Private/NiagaraRendererProperties.cpp b/Engine/Plugins/FX/Niagara/Source/Niagara/Private/NiagaraRendererProperties.cpp
--- a/Engine/Plugins/FX/Niagara/Source/Niagara/Private/NiagaraRendererProperties.cpp
+++ b/Engine/Plugins/FX/Niagara/Source/Niagara/Private/NiagaraRendererProperties.cpp
@@ -796,6 +796,15 @@ bool UNiagaraRendererProperties::IsSortHighPrecision(ENiagaraRendererSortPrecisi
 	return SortPrecision == ENiagaraRendererSortPrecision::High;
 }
 
+bool UNiagaraRendererProperties::ShouldGpuTranslucentThisFrame(ENiagaraRendererGpuTranslucentLatency Latency)
+{
+	if (Latency == ENiagaraRendererGpuTranslucentLatency::ProjectDefault)
+	{
+		return GetDefault<UNiagaraSettings>()->DefaultGpuTranslucentLatency == ENiagaraDefaultGpuTranslucentLatency::Immediate;
+	}
+	return Latency == ENiagaraRendererGpuTranslucentLatency::Immediate;
+}
+
 bool UNiagaraRendererProperties::IsGpuTranslucentThisFrame(ERHIFeatureLevel::Type FeatureLevel, ENiagaraRendererGpuTranslucentLatency Latency)
 {
 	// We can not support low latency on the mobile renderer path as it calls PostRenderOpaque after translucent in some paths
@@ -804,11 +813,7 @@ bool UNiagaraRendererProperties::IsGpuTranslucentThisFrame(ERHIFeatureLevel::Typ
 		return false;
 	}
 
-	if (Latency == ENiagaraRendererGpuTranslucentLatency::ProjectDefault)
-	{
-		return GetDefault<UNiagaraSettings>()->DefaultGpuTranslucentLatency == ENiagaraDefaultGpuTranslucentLatency::Immediate;
-	}
-	return Latency == ENiagaraRendererGpuTranslucentLatency::Immediate;
+	return ShouldGpuTranslucentThisFrame(Latency);
 }
 
 #undef LOCTEXT_NAMESPACE
diff --git a/Engine/Plugins/FX/Niagara/Source/Niagara/Private/NiagaraRendererSprites.cpp b/Engine/Plugins/FX/Niagara/Source/Niagara/Private/NiagaraRendererSprites.cpp
--- a/Engine/Plugins/FX/Niagara/Source/Niagara/Private/NiagaraRendererSprites.cpp
+++ b/Engine/Plugins/FX/Niagara/Source/Niagara/Private/NiagaraRendererSprites.cpp
@@ -233,6 +233,7 @@ void FNiagaraRendererSprites::PrepareParticleSpriteRenderData(FParticleSpriteRen
 	ParticleSpriteRenderData.BlendMode = BlendMode;
 	ParticleSpriteRenderData.bHasTranslucentMaterials = IsTranslucentBlendMode(BlendMode);
 
+	// If these conditions change please update the DebugHUD display also to reflect it
 	const bool bLowLatencyTranslucencyEnabled = ParticleSpriteRenderData.bHasTranslucentMaterials &&
 		bGpuLowLatencyTranslucency &&
 		!SceneProxy->CastsVolumetricTranslucentShadow() &&
diff --git a/Engine/Plugins/FX/Niagara/Source/Niagara/Public/NiagaraRendererProperties.h b/Engine/Plugins/FX/Niagara/Source/Niagara/Public/NiagaraRendererProperties.h
--- a/Engine/Plugins/FX/Niagara/Source/Niagara/Public/NiagaraRendererProperties.h
+++ b/Engine/Plugins/FX/Niagara/Source/Niagara/Public/NiagaraRendererProperties.h
@@ -372,6 +372,9 @@ public:
 
 	static bool IsSortHighPrecision(ENiagaraRendererSortPrecision SortPrecision);
 
+	/** Should the Gpu translucent data be this frame or not? */
+	static bool ShouldGpuTranslucentThisFrame(ENiagaraRendererGpuTranslucentLatency Latency);
+	/** Is the Gpu translucent data going to be this frame, this can be restricted by things like feature level. */
 	static bool IsGpuTranslucentThisFrame(ERHIFeatureLevel::Type FeatureLevel, ENiagaraRendererGpuTranslucentLatency Latency);
 
 	template<typename TAction>
