Adding support for static variables as parent attributes

#okforgitihub public

diff --git a/Engine/Plugins/FX/Niagara/Source/NiagaraEditor/Private/ViewModels/Stack/NiagaraStackFunctionInputCollection.cpp b/Engine/Plugins/FX/Niagara/Source/NiagaraEditor/Private/ViewModels/Stack/NiagaraStackFunctionInputCollection.cpp
--- a/Engine/Plugins/FX/Niagara/Source/NiagaraEditor/Private/ViewModels/Stack/NiagaraStackFunctionInputCollection.cpp
+++ b/Engine/Plugins/FX/Niagara/Source/NiagaraEditor/Private/ViewModels/Stack/NiagaraStackFunctionInputCollection.cpp
@@ -380,6 +380,12 @@ void UNiagaraStackFunctionInputCollectionBase::AppendInputsForFunctionCall(FFunc
 		FInputData InputData = { InputPin, InputVariable.GetType(), EditorSortPriority, DisplayName, InputCategory, false, bIsInputHidden, bShouldShowInSummary, ModuleNode, InputFunctionCallNode };
 		int32 Index = State.InputDataCollection.Add(InputData);
 
+		if (InputVariable.GetType().IsStatic())
+		{
+			FNiagaraParentData& ParentData = State.ParentMapping.FindOrAdd(InputVariable.GetName());
+			ParentData.ParentPin = InputPin;
+		}
+
 		// set up the data for the parent-child mapping
 		if (InputMetaData && !InputMetaData->ParentAttribute.IsNone())
 		{
