Making static variables not issue errors in module/dynamic input/function scripts if we can't resolve the errors

diff --git a/Engine/Plugins/FX/Niagara/Source/NiagaraEditor/Private/NiagaraHlslTranslator.cpp b/Engine/Plugins/FX/Niagara/Source/NiagaraEditor/Private/NiagaraHlslTranslator.cpp
--- a/Engine/Plugins/FX/Niagara/Source/NiagaraEditor/Private/NiagaraHlslTranslator.cpp
+++ b/Engine/Plugins/FX/Niagara/Source/NiagaraEditor/Private/NiagaraHlslTranslator.cpp
@@ -4657,6 +4657,13 @@ int32 FHlslNiagaraTranslator::GetParameter(const FNiagaraVariable& Parameter)
 		}
 	}
 
+	// We won't be able to resolve static variables that are inputs, so just let them be the defaults for support class scripts.
+	if (OutputChunkIdx == INDEX_NONE && OutputVariable.GetType().IsStatic() &&
+		(UNiagaraScript::IsModuleScript(CompileOptions.TargetUsage) || UNiagaraScript::IsFunctionScript(CompileOptions.TargetUsage) || UNiagaraScript::IsDynamicInputScript(CompileOptions.TargetUsage)))
+	{
+		OutputChunkIdx = GetConstant(OutputVariable);
+	}
+
 	if (OutputChunkIdx == INDEX_NONE)
 	{
 		Error(FText::Format(LOCTEXT("GetParameterFail", "Cannot handle type {0}! Variable: {1}"), Parameter.GetType().GetNameText(), FText::FromName(Parameter.GetName())), nullptr, nullptr);
@@ -5796,7 +5803,8 @@ void FHlslNiagaraTranslator::SetConstantByStaticVariable(int32& OutValue, const
 				OutValue = VarWithValue.GetValue<int32>();
 			}
 		}
-		else
+		else if (!UNiagaraScript::IsModuleScript(CompileOptions.TargetUsage) && !UNiagaraScript::IsFunctionScript(CompileOptions.TargetUsage) &&
+			!UNiagaraScript::IsDynamicInputScript(CompileOptions.TargetUsage)) // Can't always resolve these outside the context of an emitter or system
 		{
 			Error(LOCTEXT("CouldNotResolveStaticVarByPin", "Could not resolve static variable through pin."), Cast<UNiagaraNode>(InDefaultPin->GetOwningNode()), InDefaultPin);
 		}
