Niagara - When post loading cooked emitters, null out the thumbnail if it's not an emitter asset.

diff --git a/Engine/Plugins/FX/Niagara/Source/Niagara/Classes/NiagaraEmitter.h b/Engine/Plugins/FX/Niagara/Source/Niagara/Classes/NiagaraEmitter.h
--- a/Engine/Plugins/FX/Niagara/Source/Niagara/Classes/NiagaraEmitter.h
+++ b/Engine/Plugins/FX/Niagara/Source/Niagara/Classes/NiagaraEmitter.h
@@ -359,7 +359,7 @@ struct NIAGARA_API FVersionedNiagaraEmitterData
 	FNiagaraParameterStore RendererBindings;
 
 	void CopyFrom(const FVersionedNiagaraEmitterData& Source);
-	void PostLoad(UNiagaraEmitter& Emitter, bool bIsCooked, int32 NiagaraVer);
+	void PostLoad(UNiagaraEmitter& Emitter, int32 NiagaraVer);
 
 	void PostInitProperties(UNiagaraEmitter* Outer);
 	bool UsesCollection(const UNiagaraParameterCollection* Collection) const;
@@ -984,9 +984,6 @@ private:
 #endif
 
 #if WITH_EDITORONLY_DATA
-	/* Flag set on load based on whether the serialized data includes editor only data */
-	uint32 IsCooked : 1;
-	
 	/** Messages associated with the Emitter asset. */
 	UPROPERTY()
 	TMap<FGuid, TObjectPtr<UNiagaraMessageDataBase>> MessageKeyToMessageMap;
diff --git a/Engine/Plugins/FX/Niagara/Source/Niagara/Private/NiagaraEmitter.cpp b/Engine/Plugins/FX/Niagara/Source/Niagara/Private/NiagaraEmitter.cpp
--- a/Engine/Plugins/FX/Niagara/Source/Niagara/Private/NiagaraEmitter.cpp
+++ b/Engine/Plugins/FX/Niagara/Source/Niagara/Private/NiagaraEmitter.cpp
@@ -177,9 +177,6 @@ UNiagaraEmitter::UNiagaraEmitter(const FObjectInitializer& Initializer)
 , MaxGPUParticlesSpawnPerFrame_DEPRECATED(0)
 #endif
 {
-#if WITH_EDITORONLY_DATA
-	IsCooked = false;
-#endif
 }
 
 void UNiagaraEmitter::PostInitProperties()
@@ -535,11 +532,6 @@ void UNiagaraEmitter::Serialize(FArchive& Ar)
 	Super::Serialize(Ar);
 
 #if WITH_EDITORONLY_DATA
-	if (Ar.IsLoading())
-	{
-		IsCooked = Ar.IsFilterEditorOnly();
-	}
-
 	// Restore the thumbnail image that was cleared before serialize.
 	if (bCookingNonAssetEmitter)
 	{
@@ -662,6 +654,12 @@ void UNiagaraEmitter::PostLoad()
 	{
 		LibraryVisibility = ENiagaraScriptLibraryVisibility::Library;
 	}
+
+	if (this->IsAsset() == false && GetPackage()->bIsCookedForEditor)
+	{
+		// Remove thunbnails for non-asset emitters in cooked packages to prevent problems due to issues with cooked data being referenced by non-cooked data.
+		ThumbnailImage = nullptr;
+	}
 #endif
 
 	const int32 UE5MainVer = GetLinkerCustomVersion(FUE5MainStreamObjectVersion::GUID);
@@ -669,7 +667,7 @@ void UNiagaraEmitter::PostLoad()
 	for (FVersionedNiagaraEmitterData& Data : VersionData)
 	{
 #if WITH_EDITORONLY_DATA
-		Data.PostLoad(*this, IsCooked, NiagaraVer);
+		Data.PostLoad(*this, NiagaraVer);
 		Data.GPUComputeScript->OnGPUScriptCompiled().RemoveAll(this);
 		Data.GPUComputeScript->OnGPUScriptCompiled().AddUObject(this, &UNiagaraEmitter::RaiseOnEmitterGPUCompiled);
 		if (UE5MainVer < FUE5MainStreamObjectVersion::FixGpuAlwaysRunningUpdateScriptNoneInterpolated)
@@ -684,7 +682,7 @@ void UNiagaraEmitter::PostLoad()
 			}
 		}
 #else
-		Data.PostLoad(*this, true, NiagaraVer);
+		Data.PostLoad(*this, NiagaraVer);
 #endif
 	}
 
@@ -704,7 +702,7 @@ void UNiagaraEmitter::DeclareConstructClasses(TArray<FTopLevelAssetPath>& OutCon
 }
 #endif
 
-void FVersionedNiagaraEmitterData::PostLoad(UNiagaraEmitter& Emitter, bool bIsCooked, int32 NiagaraVer)
+void FVersionedNiagaraEmitterData::PostLoad(UNiagaraEmitter& Emitter, int32 NiagaraVer)
 {
 #if STATS
 	StatDatabase.Init();
@@ -737,6 +735,7 @@ void FVersionedNiagaraEmitterData::PostLoad(UNiagaraEmitter& Emitter, bool bIsCo
 		{
 #if WITH_EDITORONLY_DATA
 			//In cooked builds these can be cooked out and null on purpose.
+			bool bIsCooked = Emitter.GetPackage()->bIsCookedForEditor;
 			ensureMsgf(bIsCooked, TEXT("Null renderer found in %s at index %i, removing it to prevent crashes."), *Emitter.GetPathName(), RendererIndex);
 #endif
 			RendererProperties.RemoveAt(RendererIndex);
