- Customize NiagaraLightRenderer categories

diff --git a/Engine/Plugins/FX/Niagara/Source/Niagara/Public/NiagaraLightRendererProperties.h b/Engine/Plugins/FX/Niagara/Source/Niagara/Public/NiagaraLightRendererProperties.h
--- a/Engine/Plugins/FX/Niagara/Source/Niagara/Public/NiagaraLightRendererProperties.h
+++ b/Engine/Plugins/FX/Niagara/Source/Niagara/Public/NiagaraLightRendererProperties.h
@@ -65,7 +65,7 @@ public:
 	float RadiusScale;
 
 	/** The exponent to use for all lights if no exponent binding was found */
-	UPROPERTY(EditAnywhere, Category = "Light Rendering", meta = (ClampMin = "0", EditCondition = "!bUseInverseSquaredFalloff"))
+	UPROPERTY(EditAnywhere, Category = "Light Rendering", meta = (ClampMin = "0", EditCondition = "!bUseInverseSquaredFalloff", EditConditionHides))
 	float DefaultExponent;
 
 	/** A static color shift applied to each rendered light */
@@ -85,31 +85,31 @@ public:
 	int32 RendererVisibility;
 
 	/** Which attribute should we use to check if light rendering should be enabled for a particle? This can be used to control the spawn-rate on a per-particle basis. */
-	UPROPERTY(EditAnywhere, AdvancedDisplay, Category = "Bindings")
+	UPROPERTY(EditAnywhere, Category = "Bindings")
 	FNiagaraVariableAttributeBinding LightRenderingEnabledBinding;
 
 	/** Which attribute should we use for the light's exponent when inverse squared falloff is disabled? */
-	UPROPERTY(EditAnywhere, AdvancedDisplay, Category = "Bindings", meta = (EditCondition = "!bUseInverseSquaredFalloff"))
+	UPROPERTY(EditAnywhere, Category = "Bindings", meta = (EditCondition = "!bUseInverseSquaredFalloff", EditConditionHides))
 	FNiagaraVariableAttributeBinding LightExponentBinding;
 
 	/** Which attribute should we use for position when generating lights?*/
-	UPROPERTY(EditAnywhere, AdvancedDisplay, Category = "Bindings")
+	UPROPERTY(EditAnywhere, Category = "Bindings")
 	FNiagaraVariableAttributeBinding PositionBinding;
 
 	/** Which attribute should we use for light color when generating lights?*/
-	UPROPERTY(EditAnywhere, AdvancedDisplay, Category = "Bindings")
+	UPROPERTY(EditAnywhere, Category = "Bindings")
 	FNiagaraVariableAttributeBinding ColorBinding;
 
 	/** Which attribute should we use for light radius when generating lights?*/
-	UPROPERTY(EditAnywhere, AdvancedDisplay, Category = "Bindings")
+	UPROPERTY(EditAnywhere, Category = "Bindings")
 	FNiagaraVariableAttributeBinding RadiusBinding;
 
 	/** Which attribute should we use for the intensity of the volumetric scattering from this light? This scales the light's intensity and color. */
-	UPROPERTY(EditAnywhere, AdvancedDisplay, Category = "Bindings")
+	UPROPERTY(EditAnywhere, Category = "Bindings")
 	FNiagaraVariableAttributeBinding VolumetricScatteringBinding;
 
 	/** Which attribute should we use for the renderer visibility tag? */
-	UPROPERTY(EditAnywhere, AdvancedDisplay, Category = "Bindings")
+	UPROPERTY(EditAnywhere, Category = "Bindings")
 	FNiagaraVariableAttributeBinding RendererVisibilityTagBinding;
 
 	FNiagaraDataSetAccessor<FNiagaraPosition> PositionDataSetAccessor;
diff --git a/Engine/Plugins/FX/Niagara/Source/NiagaraEditorWidgets/Private/DetailCustomizations/NiagaraLightRendererDetails.cpp b/Engine/Plugins/FX/Niagara/Source/NiagaraEditorWidgets/Private/DetailCustomizations/NiagaraLightRendererDetails.cpp
new file mode 100644
--- /dev/null
+++ b/Engine/Plugins/FX/Niagara/Source/NiagaraEditorWidgets/Private/DetailCustomizations/NiagaraLightRendererDetails.cpp
@@ -0,0 +1,46 @@
+// Copyright Epic Games, Inc. All Rights Reserved.
+
+#include "NiagaraLightRendererDetails.h"
+#include "DetailLayoutBuilder.h"
+#include "DetailCategoryBuilder.h"
+#include "IDetailChildrenBuilder.h"
+#include "DetailWidgetRow.h"
+#include "PropertyCustomizationHelpers.h"
+#include "Widgets/Layout/SExpandableArea.h"
+
+TSharedRef<IDetailCustomization> FNiagaraLightRendererDetails::MakeInstance()
+{
+	return MakeShared<FNiagaraLightRendererDetails>();
+}
+
+void FNiagaraLightRendererDetails::CustomizeDetails(IDetailLayoutBuilder& DetailBuilder)
+{
+	// Categories to put first
+	{
+		static const FName CategoriesToCollapse[] =
+		{
+			FName("Light Rendering"),
+			FName("Bindings"),
+		};
+
+		for (const FName& Category : CategoriesToCollapse)
+		{
+			IDetailCategoryBuilder& CategoryBuilder = DetailBuilder.EditCategory(Category);
+		}
+	}
+
+	// Collapse default categories
+	{
+		static const FName CategoriesToCollapse[] =
+		{
+			FName("Rendering"),
+			FName("Scalability"),
+		};
+
+		for (const FName& Category : CategoriesToCollapse)
+		{
+			IDetailCategoryBuilder& CategoryBuilder = DetailBuilder.EditCategory(Category);
+			CategoryBuilder.InitiallyCollapsed(true);
+		}
+	}
+}
diff --git a/Engine/Plugins/FX/Niagara/Source/NiagaraEditorWidgets/Private/DetailCustomizations/NiagaraLightRendererDetails.h b/Engine/Plugins/FX/Niagara/Source/NiagaraEditorWidgets/Private/DetailCustomizations/NiagaraLightRendererDetails.h
new file mode 100644
--- /dev/null
+++ b/Engine/Plugins/FX/Niagara/Source/NiagaraEditorWidgets/Private/DetailCustomizations/NiagaraLightRendererDetails.h
@@ -0,0 +1,15 @@
+// Copyright Epic Games, Inc. All Rights Reserved.
+
+#pragma once
+
+#include "IDetailCustomization.h"
+
+class IDetailLayoutBuilder;
+
+class FNiagaraLightRendererDetails : public IDetailCustomization
+{
+public:
+	static TSharedRef<IDetailCustomization> MakeInstance();
+
+	virtual void CustomizeDetails(IDetailLayoutBuilder& DetailBuilder) override;
+};
\ No newline at end of file
diff --git a/Engine/Plugins/FX/Niagara/Source/NiagaraEditorWidgets/Private/NiagaraEditorWidgetsModule.cpp b/Engine/Plugins/FX/Niagara/Source/NiagaraEditorWidgets/Private/NiagaraEditorWidgetsModule.cpp
--- a/Engine/Plugins/FX/Niagara/Source/NiagaraEditorWidgets/Private/NiagaraEditorWidgetsModule.cpp
+++ b/Engine/Plugins/FX/Niagara/Source/NiagaraEditorWidgets/Private/NiagaraEditorWidgetsModule.cpp
@@ -12,6 +12,7 @@
 #include "DetailCustomizations/NiagaraDataInterfaceStaticMeshDetails.h"
 #include "DetailCustomizations/NiagaraDataInterfaceMeshRendererInfoDetails.h"
 #include "DetailCustomizations/NiagaraDataInterfaceSpriteRendererInfoDetails.h"
+#include "DetailCustomizations/NiagaraLightRendererDetails.h"
 #include "DetailCustomizations/NiagaraMeshRendererDetails.h"
 #include "DetailCustomizations/NiagaraRibbonRendererDetails.h"
 #include "DetailCustomizations/NiagaraSpriteRendererDetails.h"
@@ -133,6 +134,7 @@ void FNiagaraEditorWidgetsModule::StartupModule()
 	PropertyModule.RegisterCustomClassLayout("NiagaraDataInterfaceGrid3DCollection", FOnGetDetailCustomizationInstance::CreateStatic(&FNiagaraDataInterfaceGrid3DCollectionDetails::MakeInstance));
 	PropertyModule.RegisterCustomClassLayout("NiagaraDataInterfaceMeshRendererInfo", FOnGetDetailCustomizationInstance::CreateStatic(&FNiagaraDataInterfaceMeshRendererInfoDetails::MakeInstance));
 	PropertyModule.RegisterCustomClassLayout("NiagaraDataInterfaceSpriteRendererInfo", FOnGetDetailCustomizationInstance::CreateStatic(&FNiagaraDataInterfaceSpriteRendererInfoDetails::MakeInstance));
+	PropertyModule.RegisterCustomClassLayout("NiagaraLightRendererProperties", FOnGetDetailCustomizationInstance::CreateStatic(&FNiagaraLightRendererDetails::MakeInstance));
 	PropertyModule.RegisterCustomClassLayout("NiagaraMeshRendererProperties", FOnGetDetailCustomizationInstance::CreateStatic(&FNiagaraMeshRendererDetails::MakeInstance));
 	PropertyModule.RegisterCustomClassLayout("NiagaraRibbonRendererProperties", FOnGetDetailCustomizationInstance::CreateStatic(&FNiagaraRibbonRendererDetails::MakeInstance));
 	PropertyModule.RegisterCustomClassLayout("NiagaraSpriteRendererProperties", FOnGetDetailCustomizationInstance::CreateStatic(&FNiagaraSpriteRendererDetails::MakeInstance));
