Now recording top-level static switch values into the DDC key in parents so that we can handle if they also write static variables.

diff --git a/Engine/Plugins/FX/Niagara/Source/Niagara/Classes/NiagaraScriptSourceBase.h b/Engine/Plugins/FX/Niagara/Source/Niagara/Classes/NiagaraScriptSourceBase.h
--- a/Engine/Plugins/FX/Niagara/Source/Niagara/Classes/NiagaraScriptSourceBase.h
+++ b/Engine/Plugins/FX/Niagara/Source/Niagara/Classes/NiagaraScriptSourceBase.h
@@ -118,6 +118,7 @@ class UNiagaraScriptSourceBase : public UObject
 
 	virtual FGuid GetChangeID() { return FGuid(); }
 
+	virtual void RegisterVMCompilationIdDependencies(struct FNiagaraVMExecutableDataId& Id, ENiagaraScriptUsage InUsage, const FGuid& InUsageId) const {};
 	virtual void ComputeVMCompilationId(struct FNiagaraVMExecutableDataId& Id, ENiagaraScriptUsage InUsage, const FGuid& InUsageId) const {};
 
 #if WITH_EDITORONLY_DATA
diff --git a/Engine/Plugins/FX/Niagara/Source/Niagara/Private/NiagaraScript.cpp b/Engine/Plugins/FX/Niagara/Source/Niagara/Private/NiagaraScript.cpp
--- a/Engine/Plugins/FX/Niagara/Source/Niagara/Private/NiagaraScript.cpp
+++ b/Engine/Plugins/FX/Niagara/Source/Niagara/Private/NiagaraScript.cpp
@@ -1289,6 +1289,7 @@ void UNiagaraScript::ComputeVMCompilationId(FNiagaraVMExecutableDataId& Id, FGui
 		UNiagaraScript* Script = Scripts[ScriptIdx];
 		TArray<FNiagaraVariable> Vars;
 		Script->RapidIterationParameters.GetParameters(Vars);		
+
 		bool bIsThisEmitterUsage = UNiagaraScript::IsEmitterScript(Usage);
 		bool bIsThisParticleUsage = UNiagaraScript::IsParticleScript(Usage);
 	
@@ -1299,6 +1300,12 @@ void UNiagaraScript::ComputeVMCompilationId(FNiagaraVMExecutableDataId& Id, FGui
 		// Emitter scripts don't depend on static variables from particle scripts.
 		if (bIsThisEmitterUsage && bIsOtherParticleUsage)
 			continue;
+		
+		// We assume that latest source is ok because these are all top-level scripts in either systems or emitters.
+		if (Script->GetLatestSource())
+		{
+			Script->GetLatestSource()->RegisterVMCompilationIdDependencies(Id, Script->Usage, Script->UsageId);
+		}
 
 		for (const FNiagaraVariable& Var : Vars)
 		{
@@ -3318,6 +3325,7 @@ void UNiagaraScript::CacheResourceShadersForCooking(EShaderPlatform ShaderPlatfo
 			// see if the script has already been added before adding a new version
 			if (InOutCachedResources.ContainsByPredicate(FindExistingScriptPredicate))
 			{
+				UE_LOG(LogScript, Warning, TEXT("CacheResourceShadersForCooking Already contains script!"));
 				return;
 			}
 
@@ -3756,6 +3764,7 @@ NIAGARA_API bool UNiagaraScript::DidScriptCompilationSucceed(bool bGPUScript) co
 			if (ScriptResource->IsCompilationFinished())
 			{
 				// If we failed compilation, it would be finished and Shader would be null.
+				UE_LOG(LogScript, Warning, TEXT("DidScriptCompilationSucceed? No!"));
 				return false;
 			}
 		}
diff --git a/Engine/Plugins/FX/Niagara/Source/NiagaraEditor/Private/NiagaraGraph.cpp b/Engine/Plugins/FX/Niagara/Source/NiagaraEditor/Private/NiagaraGraph.cpp
--- a/Engine/Plugins/FX/Niagara/Source/NiagaraEditor/Private/NiagaraGraph.cpp
+++ b/Engine/Plugins/FX/Niagara/Source/NiagaraEditor/Private/NiagaraGraph.cpp
@@ -305,6 +305,25 @@ void UNiagaraGraph::PostLoad()
 		MarkGraphRequiresSynchronization(TEXT("Graph change id was invalid"));
 	}
 
+	// Since we added the ReferenceHashFromGraph we need to keep it updated but the id won't be out of sync, so we 
+	// just do a quick check to see if all are missing then force it.
+	{
+		int32 NumMissing = 0;
+		for (int32 i = 0; i < CachedUsageInfo.Num(); i++)
+		{
+			if (false == CachedUsageInfo[i].ReferenceHashFromGraph.IsValid())
+			{
+				NumMissing++;
+			}
+		}
+
+		if (NumMissing == CachedUsageInfo.Num())
+		{
+			CachedUsageInfo.Empty();
+			LastBuiltTraversalDataChangeId = FGuid();
+		}
+	}
+
 	// There was a bug where the graph traversal could have nullptr entries in it, which causes crashes later. Detect these and just clear out the cache.
 	{
 		bool bAnyInvalid = false;
@@ -1393,6 +1412,18 @@ class UNiagaraScriptSource* UNiagaraGraph::GetSource() const
 	return CastChecked<UNiagaraScriptSource>(GetOuter());
 }
 
+FNiagaraCompileHash UNiagaraGraph::GetCompileReferencedDataHash(ENiagaraScriptUsage InUsage, const FGuid& InUsageId) const
+{
+	for (int32 i = 0; i < CachedUsageInfo.Num(); i++)
+	{
+		if (UNiagaraScript::IsEquivalentUsage(CachedUsageInfo[i].UsageType, InUsage) && CachedUsageInfo[i].UsageId == InUsageId)
+		{		
+			return CachedUsageInfo[i].ReferenceHashFromGraph;
+		}
+	}
+	return FNiagaraCompileHash();
+}
+
 FNiagaraCompileHash UNiagaraGraph::GetCompileDataHash(ENiagaraScriptUsage InUsage, const FGuid& InUsageId) const
 {
 	for (int32 i = 0; i < CachedUsageInfo.Num(); i++)
@@ -3270,8 +3301,9 @@ void UNiagaraGraph::RebuildCachedCompileIds()
 
 	const FGuid CurrentScriptVersionId = FNiagaraCustomVersion::GetLatestScriptCompileVersion();
 
+
 	// If the graph hasn't changed since last rebuild, then do nothing.
-	const bool bForceRebuild = bHasValidLastBuiltScriptVersionId && (LastBuiltScriptVersionId != CurrentScriptVersionId);
+	const bool bForceRebuild = (bHasValidLastBuiltScriptVersionId && (LastBuiltScriptVersionId != CurrentScriptVersionId));
 	
 	if (!bForceRebuild && ChangeId == LastBuiltTraversalDataChangeId && LastBuiltTraversalDataChangeId.IsValid())
 	{
@@ -3332,11 +3364,20 @@ void UNiagaraGraph::RebuildCachedCompileIds()
 		GraphHashState.Final();
 
 		FSHA1 HashState;
+		FSHA1 HashStateStatics;
 		for (UNiagaraNode* Node : NewUsageCache[i].Traversal)
 		{
 			Node->UpdateCompileHashForNode(HashState);
+
+			UNiagaraNodeFunctionCall* CallNode = Cast<UNiagaraNodeFunctionCall>(Node);
+			if (CallNode)
+			{
+				CallNode->UpdateReferencedStaticsHashForNode(HashStateStatics);
+			}
+
 		}
 		HashState.Final();
+		HashStateStatics.Final();
 
 		// We can't store in a FShaHash struct directly because you can't FProperty it. Using a standin of the same size.
 		{
@@ -3346,13 +3387,18 @@ void UNiagaraGraph::RebuildCachedCompileIds()
 			NewUsageCache[i].CompileHash = FNiagaraCompileHash(DataHash.GetData(), DataHash.Num());
 		}
 
+		{
+			TStaticArray<uint8, FSHA1::DigestSize> DataHashStatics;
+			HashStateStatics.GetHash(DataHashStatics.GetData());
+			NewUsageCache[i].ReferenceHashFromGraph = FNiagaraCompileHash(DataHashStatics.GetData(), DataHashStatics.Num());
+		}
 
 		{
 			// We can't store in a FShaHash struct directly because you can't UProperty it. Using a standin of the same size.
 			TStaticArray<uint8, FSHA1::DigestSize> DataHash;
 			GraphHashState.GetHash(DataHash.GetData());
 			NewUsageCache[i].CompileHashFromGraph = FNiagaraCompileHash(DataHash.GetData(), DataHash.Num());
-			NewUsageCache[i].CompileLastObjects = Visitor.Values;
+			NewUsageCache[i].CompileLastObjects = Visitor.Values;	
 
 #if WITH_EDITORONLY_DATA
 			// Log out all the entries that differ!
diff --git a/Engine/Plugins/FX/Niagara/Source/NiagaraEditor/Private/NiagaraGraphDataCache.cpp b/Engine/Plugins/FX/Niagara/Source/NiagaraEditor/Private/NiagaraGraphDataCache.cpp
--- a/Engine/Plugins/FX/Niagara/Source/NiagaraEditor/Private/NiagaraGraphDataCache.cpp
+++ b/Engine/Plugins/FX/Niagara/Source/NiagaraEditor/Private/NiagaraGraphDataCache.cpp
@@ -145,8 +145,11 @@ FNiagaraGraphDataCache::FStackFunctionInputPinKey::FStackFunctionInputPinKey(con
 	bFilterForCompilation = bInFilterForCompilation;
 	ConstantResolverHash = InConstantResolver.BuildTypeHash();
 
+	const FString& FunctionName = FunctionCallNode.GetFunctionName();
+
 	Hash = GetTypeHash(FunctionCallGraphKey);
 	Hash = HashCombine(Hash, GetTypeHash(FunctionCallGraphChangeId));
+
 	for (const FNiagaraVariable& StaticVariable : ContextStaticVariables)
 	{
 		Hash = HashCombine(Hash, GetTypeHash(StaticVariable));
@@ -156,6 +159,23 @@ FNiagaraGraphDataCache::FStackFunctionInputPinKey::FStackFunctionInputPinKey(con
 		{
 			Hash = HashCombine(Hash, GetTypeHash(StaticVariableData[ByteIt]));
 		}
+
+	}
+
+	/* We double-hash specifically anything that matches the namespace. However we don't hash the name, just the type
+	  and the value. This way we can still re-use across different instances of the same graph.*/
+	for (const FNiagaraVariable& StaticVariable : ContextStaticVariables)
+	{
+		if (StaticVariable.IsInNameSpace(FunctionName))
+		{
+			Hash = HashCombine(Hash, GetTypeHash(StaticVariable));
+			const uint8* StaticVariableData = StaticVariable.GetData();
+			const int32 ByteCount = StaticVariable.GetAllocatedSizeInBytes();
+			for (int32 ByteIt = 0; ByteIt < ByteCount; ++ByteIt)
+			{
+				Hash = HashCombine(Hash, GetTypeHash(StaticVariableData[ByteIt]));
+			}
+		}
 	}
 	Hash = HashCombine(Hash, ConstantResolverHash);
 	Hash = HashCombine(Hash, InputPinOption);
diff --git a/Engine/Plugins/FX/Niagara/Source/NiagaraEditor/Private/NiagaraNodeFunctionCall.cpp b/Engine/Plugins/FX/Niagara/Source/NiagaraEditor/Private/NiagaraNodeFunctionCall.cpp
--- a/Engine/Plugins/FX/Niagara/Source/NiagaraEditor/Private/NiagaraNodeFunctionCall.cpp
+++ b/Engine/Plugins/FX/Niagara/Source/NiagaraEditor/Private/NiagaraNodeFunctionCall.cpp
@@ -1048,6 +1048,28 @@ void UNiagaraNodeFunctionCall::UpdateCompileHashForNode(FSHA1& HashState) const
 	HashState.UpdateWithString(*GetFunctionName(), GetFunctionName().Len());
 }
 
+void UNiagaraNodeFunctionCall::UpdateReferencedStaticsHashForNode(FSHA1& HashState) const
+{
+
+	// Assume that we only care about static switch pins and their current values
+	FPinCollectorArray InputPins;
+	GetInputPins(InputPins);
+	for (int32 PinIndex = 0; PinIndex < InputPins.Num(); ++PinIndex)
+	{
+		if (!InputPins[PinIndex])
+		{
+			continue;
+		}
+
+		if (InputPins[PinIndex]->PersistentGuid.IsValid() && InputPins[PinIndex]->LinkedTo.Num() == 0 && InputPins[PinIndex]->bNotConnectable == true)
+		{			
+			HashState.Update((const uint8*)&InputPins[PinIndex]->PersistentGuid, sizeof(InputPins[PinIndex]->PersistentGuid));
+			HashState.UpdateWithString(*InputPins[PinIndex]->DefaultValue, InputPins[PinIndex]->DefaultValue.Len());
+		}
+	}
+
+}
+
 bool UNiagaraNodeFunctionCall::HasValidScriptAndGraph() const
 {
 	return FunctionScript != nullptr && GetCalledGraph() != nullptr;
@@ -1140,7 +1162,7 @@ void UNiagaraNodeFunctionCall::BuildParameterMapHistory(FNiagaraParameterMapHist
 
 		if (bHasParamMapPin && bRecursive)
 		{
-			ParamMapIdx = OutHistory.TraceParameterMapOutputPin(UNiagaraNode::TraceOutputPin(ParamMapPin->LinkedTo[0]));
+			ParamMapIdx = OutHistory.TraceParameterMapOutputPin((ParamMapPin->LinkedTo[0]));
 		}
 
 		OutHistory.EnterFunction(GetFunctionName(), FunctionScript, FunctionGraph, this);
diff --git a/Engine/Plugins/FX/Niagara/Source/NiagaraEditor/Private/NiagaraNodeParameterMapGet.cpp b/Engine/Plugins/FX/Niagara/Source/NiagaraEditor/Private/NiagaraNodeParameterMapGet.cpp
--- a/Engine/Plugins/FX/Niagara/Source/NiagaraEditor/Private/NiagaraNodeParameterMapGet.cpp
+++ b/Engine/Plugins/FX/Niagara/Source/NiagaraEditor/Private/NiagaraNodeParameterMapGet.cpp
@@ -460,7 +460,7 @@ void UNiagaraNodeParameterMapGet::BuildParameterMapHistory(FNiagaraParameterMapH
 	int32 ParamMapIdx = INDEX_NONE;
 	if (GetInputPin(0)->LinkedTo.Num() != 0)
 	{
-		ParamMapIdx = OutHistory.TraceParameterMapOutputPin(UNiagaraNode::TraceOutputPin(GetInputPin(0)->LinkedTo[0]));
+		ParamMapIdx = OutHistory.TraceParameterMapOutputPin((GetInputPin(0)->LinkedTo[0]));
 	}
 
 	if (ParamMapIdx != INDEX_NONE)
diff --git a/Engine/Plugins/FX/Niagara/Source/NiagaraEditor/Private/NiagaraParameterMapHistory.cpp b/Engine/Plugins/FX/Niagara/Source/NiagaraEditor/Private/NiagaraParameterMapHistory.cpp
--- a/Engine/Plugins/FX/Niagara/Source/NiagaraEditor/Private/NiagaraParameterMapHistory.cpp
+++ b/Engine/Plugins/FX/Niagara/Source/NiagaraEditor/Private/NiagaraParameterMapHistory.cpp
@@ -1837,10 +1837,17 @@ int32 FNiagaraParameterMapHistoryBuilder::HandleVariableRead(int32 ParamMapIdx,
 		{					
 			FNiagaraVariable VarRapid = FNiagaraParameterMapHistory::ConvertVariableToRapidIterationConstantName(Histories[ParamMapIdx].Variables[FoundIdx], EmitterNamespace.Len() == 0 ? nullptr : *EmitterNamespace, Usage);
 			
-			int32 StaticIdx = StaticVariables.IndexOfByPredicate([&](const FNiagaraVariable& InObj) -> bool
-				{
-					return (InObj.GetName() == VarRapid.GetName());
-				});;
+			int32 StaticIdx = INDEX_NONE;
+			
+			// See the comments in FNiagaraStackGraphUtilities::GetStackFunctionInputPinsWithoutCache to see 
+			// why this check is necessary.
+			if (!bIgnoreStaticRapidIterationParameters)
+			{
+				StaticIdx = StaticVariables.IndexOfByPredicate([&](const FNiagaraVariable& InObj) -> bool
+					{
+						return (InObj.GetName() == VarRapid.GetName());
+					});;
+			}
 
 			if (StaticIdx == INDEX_NONE)
 			{
diff --git a/Engine/Plugins/FX/Niagara/Source/NiagaraEditor/Private/NiagaraScriptSource.cpp b/Engine/Plugins/FX/Niagara/Source/NiagaraEditor/Private/NiagaraScriptSource.cpp
--- a/Engine/Plugins/FX/Niagara/Source/NiagaraEditor/Private/NiagaraScriptSource.cpp
+++ b/Engine/Plugins/FX/Niagara/Source/NiagaraEditor/Private/NiagaraScriptSource.cpp
@@ -31,6 +31,39 @@ UNiagaraScriptSource::UNiagaraScriptSource(const FObjectInitializer& ObjectIniti
 {
 }
 
+void UNiagaraScriptSource::RegisterVMCompilationIdDependencies(struct FNiagaraVMExecutableDataId& Id, ENiagaraScriptUsage InUsage, const FGuid& InUsageId) const
+{
+	if (!AllowShaderCompiling())
+	{
+		return;
+	}
+
+	if (UNiagaraScript::IsGPUScript(InUsage)) // GPU particle scripts are baked in via other usages
+	{
+		return;
+	}
+
+	if (UNiagaraScript::IsParticleScript(InUsage) || UNiagaraScript::IsEmitterScript(InUsage) || UNiagaraScript::IsSystemScript(InUsage))
+	{
+		if (NodeGraph)
+		{
+			NodeGraph->RebuildCachedCompileIds();
+			FNiagaraCompileHash Hash = FNiagaraCompileHash(NodeGraph->GetCompileReferencedDataHash(InUsage, InUsageId));
+			
+			if (Hash.IsValid())
+			{
+				Id.ReferencedCompileHashes.AddUnique(Hash);
+				Id.DebugReferencedObjects.Emplace(GetPathName());
+			}
+			else
+			{
+				UE_LOG(LogNiagaraEditor, Warning, TEXT("Failed to get a valid referenced data hash back! %s"), *GetPathNameSafe(this));
+			}
+		}
+	}
+}
+
+
 void UNiagaraScriptSource::ComputeVMCompilationId(FNiagaraVMExecutableDataId& Id, ENiagaraScriptUsage InUsage, const FGuid& InUsageId) const
 {
 	if (!AllowShaderCompiling())
diff --git a/Engine/Plugins/FX/Niagara/Source/NiagaraEditor/Private/ViewModels/Stack/NiagaraStackGraphUtilities.cpp b/Engine/Plugins/FX/Niagara/Source/NiagaraEditor/Private/ViewModels/Stack/NiagaraStackGraphUtilities.cpp
--- a/Engine/Plugins/FX/Niagara/Source/NiagaraEditor/Private/ViewModels/Stack/NiagaraStackGraphUtilities.cpp
+++ b/Engine/Plugins/FX/Niagara/Source/NiagaraEditor/Private/ViewModels/Stack/NiagaraStackGraphUtilities.cpp
@@ -941,6 +941,29 @@ void FNiagaraStackGraphUtilities::GetStackFunctionInputPinsWithoutCache(
 	Builder.SetIgnoreDisabled(bIgnoreDisabled);
 	Builder.ConstantResolver = ConstantResolver;
 	Builder.RegisterExternalStaticVariables(StaticVars);
+	
+	/* The line below represents a compromise. It basically says don't look up the static variable from the 
+	static variables list set in RegisterExternalStaticVariables as a rapid iteration parameter. Why? Rapid Iteration Parameters
+	may be stale and linger around from prior changes that have been abandoned. We keep them around so that when you toggle back
+	in the UI, you haven't also lost your last set values for that version. So therefore, we can't officially remove them.
+	
+	The issue comes in when a static variable is set via some other means, linking to another variable for example. When generating
+	the UI, it *should* show anything that would match up to the linked variable, but because this is a parameter map history
+	traversal that is in isolation on the node and doesn't include upstream set nodes, it will fail and just look up the rapid iteration
+	version if left to its' own devices. 
+	
+	Trial 1 at a fix was to cache all the values at the root of the graph and look them up when we build the static variables list
+	we use above. That isn't good because that version skips all disabled modules and we want the UI to be consistent regardless 
+	of disabled module status.
+	
+	Trial 2 was to back up and include the precursor gets/sets. This fails for the same reason, ultimately something will
+	be bound to a variable that will assume that it needs to look up a bogus rapid iteration value and we end up back here.
+	
+	Trial 3 just says, we know the static variables list above includes basically all relevant intermediate static values. So if 
+	we just circumvent the logic in the parameter map history traversal that looks up by rapid iteration parameter, it will
+	find the right "cached" value. So we just do that below.
+	 */
+	Builder.SetIgnoreStaticRapidIterationParameters(true);
 
 	// if we are only dealing with the module input pins then we don't need to delve deep into the graph
 	if (Options == ENiagaraGetStackFunctionInputPinsOptions::ModuleInputsOnly)
diff --git a/Engine/Plugins/FX/Niagara/Source/NiagaraEditor/Public/NiagaraGraph.h b/Engine/Plugins/FX/Niagara/Source/NiagaraEditor/Public/NiagaraGraph.h
--- a/Engine/Plugins/FX/Niagara/Source/NiagaraEditor/Public/NiagaraGraph.h
+++ b/Engine/Plugins/FX/Niagara/Source/NiagaraEditor/Public/NiagaraGraph.h
@@ -107,6 +107,10 @@ public:
 	UPROPERTY()
 	FNiagaraCompileHash CompileHashFromGraph;
 
+	/** The hash of values that could feed into downstream scripts.*/
+	UPROPERTY()
+	FNiagaraCompileHash ReferenceHashFromGraph;
+
 	UPROPERTY(Transient)
 	TArray<FNiagaraCompileHashVisitorDebugInfo> CompileLastObjects;
 
@@ -262,6 +266,9 @@ class UNiagaraGraph : public UEdGraph
 	/** Gets the current compile data hash associated with the output node traversal specified by InUsage and InUsageId. If the usage is not found, an invalid hash is returned.*/
 	FNiagaraCompileHash GetCompileDataHash(ENiagaraScriptUsage InUsage, const FGuid& InUsageId) const;
 
+	/** Gets the current reference data hash associated with the output node traversal specified by InUsage and InUsageId (Static variables & switches). If the usage is not found, an invalid hash is returned.*/
+	FNiagaraCompileHash GetCompileReferencedDataHash(ENiagaraScriptUsage InUsage, const FGuid& InUsageId) const;
+
 	/** Gets the current base id associated with the output node traversal specified by InUsage and InUsageId. If the usage is not found, an invalid guid is returned. */
 	FGuid GetBaseId(ENiagaraScriptUsage InUsage, const FGuid& InUsageId) const;
 
diff --git a/Engine/Plugins/FX/Niagara/Source/NiagaraEditor/Public/NiagaraNodeFunctionCall.h b/Engine/Plugins/FX/Niagara/Source/NiagaraEditor/Public/NiagaraNodeFunctionCall.h
--- a/Engine/Plugins/FX/Niagara/Source/NiagaraEditor/Public/NiagaraNodeFunctionCall.h
+++ b/Engine/Plugins/FX/Niagara/Source/NiagaraEditor/Public/NiagaraNodeFunctionCall.h
@@ -115,6 +115,8 @@ public:
 	virtual void UpdateCompileHashForNode(FSHA1& HashState) const override;
 	//End UNiagaraNode interface
 
+	virtual void UpdateReferencedStaticsHashForNode(FSHA1& HashState) const;
+
 	//~ Begin EdGraphNode Interface
 	virtual void AutowireNewNode(UEdGraphPin* FromPin) override;
 	virtual void AllocateDefaultPins() override;
diff --git a/Engine/Plugins/FX/Niagara/Source/NiagaraEditor/Public/NiagaraParameterMapHistory.h b/Engine/Plugins/FX/Niagara/Source/NiagaraEditor/Public/NiagaraParameterMapHistory.h
--- a/Engine/Plugins/FX/Niagara/Source/NiagaraEditor/Public/NiagaraParameterMapHistory.h
+++ b/Engine/Plugins/FX/Niagara/Source/NiagaraEditor/Public/NiagaraParameterMapHistory.h
@@ -508,6 +508,9 @@ public:
 	bool GetIgnoreDisabled() const { return bIgnoreDisabled; }
 	void SetIgnoreDisabled(bool bInIgnore) { bIgnoreDisabled = bInIgnore; }
 
+	bool GetIgnoreStaticRapidIterationParameters() const { return bIgnoreStaticRapidIterationParameters; }
+	void SetIgnoreStaticRapidIterationParameters(bool bInIgnore) { bIgnoreStaticRapidIterationParameters = bInIgnore; }
+
 	bool IsInEncounteredFunctionNamespace(const FNiagaraVariable& InVar) const;
 	bool IsInEncounteredEmitterNamespace(const FNiagaraVariable& InVar) const;
 
@@ -591,6 +594,10 @@ protected:
 	/** Whether we want to include ParameterCollection information */
 	bool bIncludeParameterCollectionInfo;
 
+	/** Whether or not we want to look up the rapid iteration form of a variable or just the regular form. 
+	see comments in FNiagaraStackGraphUtilities::GetStackFunctionInputPinsWithoutCache for reasoning behind this.*/
+	bool bIgnoreStaticRapidIterationParameters = false;
+
 	TArray<FNiagaraVariable> EncounterableExternalVariables;
 };
 
diff --git a/Engine/Plugins/FX/Niagara/Source/NiagaraEditor/Public/NiagaraScriptSource.h b/Engine/Plugins/FX/Niagara/Source/NiagaraEditor/Public/NiagaraScriptSource.h
--- a/Engine/Plugins/FX/Niagara/Source/NiagaraEditor/Public/NiagaraScriptSource.h
+++ b/Engine/Plugins/FX/Niagara/Source/NiagaraEditor/Public/NiagaraScriptSource.h
@@ -38,6 +38,7 @@ class UNiagaraScriptSource : public UNiagaraScriptSourceBase
 	virtual FGuid GetChangeID() override;
 	FVersionedNiagaraEmitter GetOuterEmitter() const;
 
+	virtual void RegisterVMCompilationIdDependencies(struct FNiagaraVMExecutableDataId& Id, ENiagaraScriptUsage InUsage, const FGuid& InUsageId) const override;
 	virtual void ComputeVMCompilationId(struct FNiagaraVMExecutableDataId& Id, ENiagaraScriptUsage InUsage, const FGuid& InUsageId) const override;
 
 	// Will conditionally refresh the graph's CompileId
