- Move compiler tags into special purpose sub menu to reduce clutter

#rnx

diff --git a/Engine/Plugins/FX/Niagara/Source/NiagaraEditor/Private/EdGraphSchema_Niagara.cpp b/Engine/Plugins/FX/Niagara/Source/NiagaraEditor/Private/EdGraphSchema_Niagara.cpp
--- a/Engine/Plugins/FX/Niagara/Source/NiagaraEditor/Private/EdGraphSchema_Niagara.cpp
+++ b/Engine/Plugins/FX/Niagara/Source/NiagaraEditor/Private/EdGraphSchema_Niagara.cpp
@@ -884,7 +884,7 @@ PRAGMA_ENABLE_DEPRECATION_WARNINGS
 	// Handle output tag nodes
 	if (bAllowOutputTagNodes && NiagaraEditorSettings->IsAllowedClass(UNiagaraNodeOutputTag::StaticClass()))
 	{
-		FText MenuCat = FText::FromString("Compiler Tagging");
+		FText MenuCat = FText::FromString("Special Purpose Parameters");
 
 		{
 			FString Name = TEXT("Add Compiler Output Tag");
