Fix for preview component completing in sim cache view

diff --git a/Engine/Plugins/FX/Niagara/Source/Niagara/Private/NiagaraSystemInstance.cpp b/Engine/Plugins/FX/Niagara/Source/Niagara/Private/NiagaraSystemInstance.cpp
--- a/Engine/Plugins/FX/Niagara/Source/Niagara/Private/NiagaraSystemInstance.cpp
+++ b/Engine/Plugins/FX/Niagara/Source/Niagara/Private/NiagaraSystemInstance.cpp
@@ -2141,7 +2141,7 @@ void FNiagaraSystemInstance::SimCacheTick_GameThread(UNiagaraSimCache* SimCache,
 		return;
 	}
 
-	if ( DesiredAge >= SimCache->GetStartSeconds() + SimCache->GetDurationSeconds() )
+	if ( DesiredAge > SimCache->GetStartSeconds() + SimCache->GetDurationSeconds() )
 	{
 		Complete(false);
 		return;
diff --git a/Engine/Plugins/FX/Niagara/Source/NiagaraEditor/Private/ViewModels/NiagaraSimCacheViewModel.cpp b/Engine/Plugins/FX/Niagara/Source/NiagaraEditor/Private/ViewModels/NiagaraSimCacheViewModel.cpp
--- a/Engine/Plugins/FX/Niagara/Source/NiagaraEditor/Private/ViewModels/NiagaraSimCacheViewModel.cpp
+++ b/Engine/Plugins/FX/Niagara/Source/NiagaraEditor/Private/ViewModels/NiagaraSimCacheViewModel.cpp
@@ -148,10 +148,12 @@ void FNiagaraSimCacheViewModel::SetFrameIndex(const int32 InFrameIndex)
 	{
 		const float Duration = SimCache->GetDurationSeconds();
 		const int NumFrames = SimCache->GetNumFrames();
+		const float StartSeconds = SimCache->GetStartSeconds();
 
-		const float NormalizedFrame = NumFrames == 0 ? 0.0f : float(InFrameIndex) / float(NumFrames - 1);
+		const float NormalizedFrame = FMath::Clamp( NumFrames == 0 ? 0.0f : float(InFrameIndex) / float(NumFrames - 1), 0.0f, 1.0f );
+		const float DesiredAge = FMath::Clamp(StartSeconds + (Duration * NormalizedFrame), StartSeconds, StartSeconds + Duration);
 		
-		PreviewComponent->SetDesiredAge(SimCache->GetStartSeconds() + (Duration * NormalizedFrame));
+		PreviewComponent->SetDesiredAge(DesiredAge);
 	}
 	OnViewDataChangedDelegate.Broadcast(false);
 }
diff --git a/Engine/Plugins/FX/Niagara/Source/NiagaraEditor/Private/Widgets/SNiagaraSimCacheViewTimeline.cpp b/Engine/Plugins/FX/Niagara/Source/NiagaraEditor/Private/Widgets/SNiagaraSimCacheViewTimeline.cpp
--- a/Engine/Plugins/FX/Niagara/Source/NiagaraEditor/Private/Widgets/SNiagaraSimCacheViewTimeline.cpp
+++ b/Engine/Plugins/FX/Niagara/Source/NiagaraEditor/Private/Widgets/SNiagaraSimCacheViewTimeline.cpp
@@ -97,7 +97,7 @@ FReply SNiagaraSimCacheViewTimeline::OnTimelineScrubbed(const FGeometry& MyGeome
 	if ( FNiagaraSimCacheViewModel* ViewModel = WeakViewModel.Pin().Get() )
 	{
 		const FVector2D LocalLocation = MyGeometry.AbsoluteToLocal(MouseEvent.GetScreenSpacePosition());
-		const float NormalizedTime = LocalLocation.X / MyGeometry.Size.X;
+		const float NormalizedTime = FMath::Clamp( LocalLocation.X / MyGeometry.Size.X, 0.0f, 1.0f);
 
 		const int32 NewFrameIndex = FMath::RoundToInt(NormalizedTime * float(ViewModel->GetNumFrames() - 1));
 
