Bandaid fix for null access with tool tips

diff --git a/Engine/Plugins/FX/Niagara/Source/NiagaraEditor/Private/ViewModels/Stack/NiagaraStackFunctionInput.cpp b/Engine/Plugins/FX/Niagara/Source/NiagaraEditor/Private/ViewModels/Stack/NiagaraStackFunctionInput.cpp
--- a/Engine/Plugins/FX/Niagara/Source/NiagaraEditor/Private/ViewModels/Stack/NiagaraStackFunctionInput.cpp
+++ b/Engine/Plugins/FX/Niagara/Source/NiagaraEditor/Private/ViewModels/Stack/NiagaraStackFunctionInput.cpp
@@ -475,10 +475,13 @@ FText UNiagaraStackFunctionInput::GetValueToolTip() const
 		{
 		case EValueMode::Data:
 		{
-			FString DataInterfaceDescription = InputValues.DataObject->GetClass()->GetDescription();
-			if (DataInterfaceDescription.Len() > 0)
+			if (InputValues.DataObject.IsValid())
 			{
-				ValueToolTipCache = FText::FromString(DataInterfaceDescription);
+				FString DataInterfaceDescription = InputValues.DataObject->GetClass()->GetDescription();
+				if (DataInterfaceDescription.Len() > 0)
+				{
+					ValueToolTipCache = FText::FromString(DataInterfaceDescription);
+				}
 			}
 			break;
 		}
