- Add RHI & Adapter deny lists for Niagara GPU emitters
- This allows us to target specific RHIs and devices where we may encounter issues

diff --git a/Engine/Plugins/FX/Niagara/Source/Niagara/Private/NiagaraComponentSettings.cpp b/Engine/Plugins/FX/Niagara/Source/Niagara/Private/NiagaraComponentSettings.cpp
--- a/Engine/Plugins/FX/Niagara/Source/Niagara/Private/NiagaraComponentSettings.cpp
+++ b/Engine/Plugins/FX/Niagara/Source/Niagara/Private/NiagaraComponentSettings.cpp
@@ -39,6 +39,7 @@ namespace FNiagaraComponentSettings
 
 	bool									bUseSystemDenyList = false;
 	bool									bUseEmitterDenyList = false;
+	bool									bAllowGpuEmitters = true;
 
 	TSet<FName>								SystemDenyList;
 	TSet<FNiagaraEmitterNameSettingsRef>	EmitterDenyList;
@@ -49,6 +50,8 @@ namespace FNiagaraComponentSettings
 	FString									EmitterDenyListString;
 	FString									GpuEmitterDenyListString;
 	FString									GpuDataInterfaceDenyListString;
+	FString									GpuRHIDenyListString;
+	FString									GpuRHIAdapterDenyListString;
 
 	static bool ParseIntoSet(const FString& StringList, TSet<FName>& OutSet)
 	{
@@ -144,6 +147,64 @@ namespace FNiagaraComponentSettings
 		FConsoleVariableDelegate::CreateStatic(UpdateEmitterDenyList),
 		ECVF_Scalability | ECVF_Default
 	);
+	static FAutoConsoleVariableRef CVarNiagaraGpuRHIDenyList(
+		TEXT("fx.Niagara.SetGpuRHIDenyList"),
+		GpuRHIDenyListString,
+		TEXT("Set Gpu RHI deny list to use, comma separated and uses wildcards, i.e. (*MyRHI*) would exclude anything that contais MyRHI"),
+		ECVF_Scalability | ECVF_Default
+	);
+
+	static FAutoConsoleVariableRef CVarNiagaraGpuRHIAdapterDenyList(
+		TEXT("fx.Niagara.SetGpuRHIAdapterDenyList"),
+		GpuRHIAdapterDenyListString,
+		TEXT("Set Gpu RHI Adapter deny list to use, comma separated and uses wildcards, i.e. (*MyGpu*) would exclude anything that contais MyGpu"),
+		ECVF_Scalability | ECVF_Default
+	);
+
+	void OnPostEngineInit()
+	{
+		if (GDynamicRHI == nullptr)
+		{
+			return;
+		}
+
+		bool bShouldAllowGpuEmitters = true;
+		if (GpuRHIDenyListString.Len() > 0)
+		{
+			TArray<FString> BanNames;
+			GpuRHIDenyListString.ParseIntoArray(BanNames, TEXT(","));
+
+			const FString RHIName = GDynamicRHI->GetName();
+			if (BanNames.ContainsByPredicate([&RHIName](const FString& BanName) { return RHIName.MatchesWildcard(BanName); }))
+			{
+				bShouldAllowGpuEmitters = false;
+			}
+		}
+
+		if (GpuRHIAdapterDenyListString.Len() > 0)
+		{
+			TArray<FString> BanNames;
+			GpuRHIAdapterDenyListString.ParseIntoArray(BanNames, TEXT(","));
+
+			if (BanNames.ContainsByPredicate([](const FString& BanName) { return GRHIAdapterName.MatchesWildcard(BanName); }))
+			{
+				bShouldAllowGpuEmitters = false;
+			}
+		}
+
+		if (bShouldAllowGpuEmitters != bAllowGpuEmitters)
+		{
+			bAllowGpuEmitters = bShouldAllowGpuEmitters;
+			if (bAllowGpuEmitters == false)
+			{
+				UE_LOG(LogNiagara, Log,
+					TEXT("GPU emitters will not run for RHI(%s) Adapter(%s) as they match deny list RHIDeny(%s) AdapterDeny(%s)."),
+					GDynamicRHI->GetName(), *GRHIAdapterName,
+					*GpuRHIDenyListString, *GpuRHIAdapterDenyListString
+				);
+			}
+		}
+	}
 
 	bool IsSystemAllowedToRun(const UNiagaraSystem* System)
 	{
@@ -194,6 +255,12 @@ namespace FNiagaraComponentSettings
 				}
 			}
 		}
+		if (bAllowGpuEmitters == false)
+		{
+			const FVersionedNiagaraEmitter CachedEmitter = EmitterInstance->GetCachedEmitter();
+			FVersionedNiagaraEmitterData* EmitterData = CachedEmitter.GetEmitterData();
+			return EmitterData && (EmitterData->SimTarget != ENiagaraSimTarget::GPUComputeSim);
+		}
 		return true;
 	}
 }
diff --git a/Engine/Plugins/FX/Niagara/Source/Niagara/Private/NiagaraModule.cpp b/Engine/Plugins/FX/Niagara/Source/Niagara/Private/NiagaraModule.cpp
--- a/Engine/Plugins/FX/Niagara/Source/Niagara/Private/NiagaraModule.cpp
+++ b/Engine/Plugins/FX/Niagara/Source/Niagara/Private/NiagaraModule.cpp
@@ -3,6 +3,7 @@
 #include "NiagaraModule.h"
 #include "Misc/LazySingleton.h"
 #include "Modules/ModuleManager.h"
+#include "NiagaraComponentSettings.h"
 #include "NiagaraCompileHashVisitor.h"
 #include "NiagaraTypes.h"
 #include "NiagaraDebugVis.h"
@@ -422,6 +423,8 @@ void INiagaraModule::OnPostEngineInit()
 			FNiagaraPlatformSet::OnCVarUnregistered(CVar);
 		}
 	);
+
+	FNiagaraComponentSettings::OnPostEngineInit();
 }
 
 void INiagaraModule::OnPreExit()
diff --git a/Engine/Plugins/FX/Niagara/Source/Niagara/Public/NiagaraComponentSettings.h b/Engine/Plugins/FX/Niagara/Source/Niagara/Public/NiagaraComponentSettings.h
--- a/Engine/Plugins/FX/Niagara/Source/Niagara/Public/NiagaraComponentSettings.h
+++ b/Engine/Plugins/FX/Niagara/Source/Niagara/Public/NiagaraComponentSettings.h
@@ -9,6 +9,7 @@ class FNiagaraEmitterInstance;
 
 namespace FNiagaraComponentSettings
 {
+	extern void OnPostEngineInit();
 	extern NIAGARA_API bool IsSystemAllowedToRun(const UNiagaraSystem* System);
 	extern NIAGARA_API bool IsEmitterAllowedToRun(const FNiagaraEmitterInstance* EmitterInstance);
 };
