Fix check around GetNumMipLevel for texture sample DI

diff --git a/Engine/Plugins/FX/Niagara/Source/Niagara/Private/NiagaraDataInterfaceTexture.cpp b/Engine/Plugins/FX/Niagara/Source/Niagara/Private/NiagaraDataInterfaceTexture.cpp
--- a/Engine/Plugins/FX/Niagara/Source/Niagara/Private/NiagaraDataInterfaceTexture.cpp
+++ b/Engine/Plugins/FX/Niagara/Source/Niagara/Private/NiagaraDataInterfaceTexture.cpp
@@ -229,7 +229,7 @@ void UNiagaraDataInterfaceTexture::GetVMExternalFunction(const FVMExternalFuncti
 	}
 	else if (BindingInfo.Name == GetNumMipLevelsName)
 	{
-		check(BindingInfo.GetNumInputs() == 1 && BindingInfo.GetNumOutputs() == 2);
+		check(BindingInfo.GetNumInputs() == 1 && BindingInfo.GetNumOutputs() == 1);
 		OutFunc = FVMExternalFunction::CreateUObject(this, &UNiagaraDataInterfaceTexture::VMGetNumMipLevels);
 	}
 }
