- Ensure we check that the pipeline is valid after a PSO precache request
- Remove redundant call to get the task when results are complete

#rnx

diff --git a/Engine/Plugins/FX/Niagara/Source/Niagara/Private/NiagaraEmitter.cpp b/Engine/Plugins/FX/Niagara/Source/Niagara/Private/NiagaraEmitter.cpp
--- a/Engine/Plugins/FX/Niagara/Source/Niagara/Private/NiagaraEmitter.cpp
+++ b/Engine/Plugins/FX/Niagara/Source/Niagara/Private/NiagaraEmitter.cpp
@@ -3333,7 +3333,6 @@ FGraphEventRef FVersionedNiagaraEmitterData::PrecacheComputePSOs(const UNiagaraE
 				for (int32 i = 0; i < ShaderScript->GetNumPermutations(); ++i)
 				{
 					FRHIComputeShader* ComputeShader = ShaderScript->GetShaderGameThread(i).GetComputeShader();
-					FGraphEventRef TaskRef = PipelineStateCache::PrecacheComputePipelineState(ComputeShader, true);
 					const EPSOPrecacheResult ShaderResult = PipelineStateCache::CheckPipelineStateInCache(ComputeShader);
 					if (ShaderResult != EPSOPrecacheResult::Complete)
 					{
diff --git a/Engine/Source/Runtime/RHI/Private/PipelineStateCache.cpp b/Engine/Source/Runtime/RHI/Private/PipelineStateCache.cpp
--- a/Engine/Source/Runtime/RHI/Private/PipelineStateCache.cpp
+++ b/Engine/Source/Runtime/RHI/Private/PipelineStateCache.cpp
@@ -1609,8 +1609,9 @@ public:
 
 			if (Initializer.bPSOPrecache)
 			{
-				if (ComputePipeline->RHIPipeline == nullptr)
+				if (ComputePipeline->RHIPipeline == nullptr || !ComputePipeline->RHIPipeline->IsValid())
 				{
+					ComputePipeline->RHIPipeline = nullptr;
 					GPrecacheComputePipelineCache.RegisterCSFailure(ComputePipeline);
 				}
 				else
