Fix for crash where systems with warmups could be reset & activated when they need to be recompiled first

diff --git a/Engine/Plugins/FX/Niagara/Source/Niagara/Private/NiagaraEmitter.cpp b/Engine/Plugins/FX/Niagara/Source/Niagara/Private/NiagaraEmitter.cpp
--- a/Engine/Plugins/FX/Niagara/Source/Niagara/Private/NiagaraEmitter.cpp
+++ b/Engine/Plugins/FX/Niagara/Source/Niagara/Private/NiagaraEmitter.cpp
@@ -1268,7 +1268,6 @@ void UNiagaraEmitter::PostEditChangeVersionedProperty(FPropertyChangedEvent& Pro
 	ResolveScalabilitySettings();
 
 	UpdateChangeId(TEXT("PostEditChangeProperty"));
-	OnPropertiesChangedDelegate.Broadcast();
 
 #if WITH_EDITORONLY_DATA
 	if (bNeedsRecompile)
@@ -1280,6 +1279,11 @@ void UNiagaraEmitter::PostEditChangeVersionedProperty(FPropertyChangedEvent& Pro
 		UNiagaraSystem::RecomputeExecutionOrderForEmitter(FVersionedNiagaraEmitter(this, Version));
 	}
 #endif
+
+	// make sure to call this after the request for compilation is performed above (if necessary).  This broadcast
+	// will result in emitters being reset, which could involve activation (when warmup is involved) and we may be
+	// in an invalid state where a recompile is required.
+	OnPropertiesChangedDelegate.Broadcast();
 }
 
 void UNiagaraEmitter::PreSave(FObjectPreSaveContext ObjectSaveContext)
