Fix for bug in timing delay for frustum visbility culling.

diff --git a/Engine/Plugins/FX/Niagara/Source/Niagara/Private/NiagaraWorldManager.cpp b/Engine/Plugins/FX/Niagara/Source/Niagara/Private/NiagaraWorldManager.cpp
--- a/Engine/Plugins/FX/Niagara/Source/Niagara/Private/NiagaraWorldManager.cpp
+++ b/Engine/Plugins/FX/Niagara/Source/Niagara/Private/NiagaraWorldManager.cpp
@@ -1551,7 +1551,7 @@ void FNiagaraWorldManager::ViewBasedCulling(UNiagaraEffectType* EffectType, cons
 		ComponentTimeSinceRendered > ScalabilitySettings.VisibilityCulling.MaxTimeWithoutRender;
 
 	//TODO: Pull screen size out into it's own debug flag in the scalability state.
-	bool bCull = bCullByNotRendered || !bInsideAnyView;
+	bool bCull = bCullByNotRendered || !bCullByOutsideViewFrustum;
 
 	OutState.bCulled |= bCull;
 #if DEBUG_SCALABILITY_STATE
