Allow to delay emitters when using precaching + r.PSOPrecaching and r.PSOPrecache.ProxyCreationWhenPSOReady
Change fx.Niagara.Emitter.ComputePSOPrecacheMode default value to 1, to allow this delay by default

diff --git a/Engine/Plugins/FX/Niagara/Source/Niagara/Classes/NiagaraEmitter.h b/Engine/Plugins/FX/Niagara/Source/Niagara/Classes/NiagaraEmitter.h
--- a/Engine/Plugins/FX/Niagara/Source/Niagara/Classes/NiagaraEmitter.h
+++ b/Engine/Plugins/FX/Niagara/Source/Niagara/Classes/NiagaraEmitter.h
@@ -374,7 +374,7 @@ struct NIAGARA_API FVersionedNiagaraEmitterData
 	void CacheFromCompiledData(const FNiagaraDataSetCompiledData* CompiledData, const UNiagaraEmitter& Emitter);
 	void CacheFromShaderCompiled();
 
-	FGraphEventRef PrecacheComputePSOs(const UNiagaraEmitter& NiagaraEmitter);
+	FGraphEventArray PrecacheComputePSOs(const UNiagaraEmitter& NiagaraEmitter);
 	bool DidPSOPrecacheFail() const { return PSOPrecacheResult == EPSOPrecacheResult::NotSupported; }
 
 	bool RequiresViewUniformBuffer() const { return bRequiresViewUniformBuffer; }
diff --git a/Engine/Plugins/FX/Niagara/Source/Niagara/Private/NiagaraEmitter.cpp b/Engine/Plugins/FX/Niagara/Source/Niagara/Private/NiagaraEmitter.cpp
--- a/Engine/Plugins/FX/Niagara/Source/Niagara/Private/NiagaraEmitter.cpp
+++ b/Engine/Plugins/FX/Niagara/Source/Niagara/Private/NiagaraEmitter.cpp
@@ -66,13 +66,13 @@ static FAutoConsoleVariableRef CVarEnableEmitterChangeIdMergeLogging(
 	ECVF_Default
 );
 
-static int32 GNiagaraEmitterComputePSOPrecacheMode = 0;
+static int32 GNiagaraEmitterComputePSOPrecacheMode = 1;
 static FAutoConsoleVariableRef CVarNiagaraEmitterComputePSOPrecacheMode(
 	TEXT("fx.Niagara.Emitter.ComputePSOPrecacheMode"),
 	GNiagaraEmitterComputePSOPrecacheMode,
 	TEXT("Controlls how PSO precaching should be done for Niagara compute shaders\n")
 	TEXT("0 = Disabled (Default).\n")
-	TEXT("1 = Enabled if r.PSOPrecaching is also enabled.\n")
+	TEXT("1 = Enabled if r.PSOPrecaching is also enabled. Emitters are not allowed to run until they complete if r.PSOPrecache.ProxyCreationWhenPSOReady=1\n")
 	TEXT("2 = Force Enabled.\n")
 	TEXT("3 = Force Enabled, emitters are not allowed to run until they complete."),
 	ECVF_Scalability
@@ -3271,19 +3271,19 @@ void UNiagaraEmitter::GenerateStatID()const
 #endif
 }
 
-FGraphEventRef FVersionedNiagaraEmitterData::PrecacheComputePSOs(const UNiagaraEmitter& NiagaraEmitter)
+FGraphEventArray FVersionedNiagaraEmitterData::PrecacheComputePSOs(const UNiagaraEmitter& NiagaraEmitter)
 {
-	FGraphEventRef PSOReadyGraphTask;
+	FGraphEventArray PSOReadyGraphTasks;
 	PSOPrecacheResult = EPSOPrecacheResult::Complete;
 	if (GNiagaraEmitterComputePSOPrecacheMode == 0 || !FApp::CanEverRender() || SimTarget != ENiagaraSimTarget::GPUComputeSim)
 	{
-		return PSOReadyGraphTask;
+		return PSOReadyGraphTasks;
 	}
 	
 	FNiagaraShaderScript* ShaderScript = GPUComputeScript->GetRenderThreadScript();
 	if (ShaderScript == nullptr || GPUComputeScript->DidScriptCompilationSucceed(true) == false)
 	{
-		return PSOReadyGraphTask;
+		return PSOReadyGraphTasks;
 	}
 
 	// Compute PSO can fail on some platforms even though the shader compiled successfully this path will wait for the PSO precache to complete before the emitter is allowed to run
@@ -3304,7 +3304,7 @@ FGraphEventRef FVersionedNiagaraEmitterData::PrecacheComputePSOs(const UNiagaraE
 			if (ShaderResult != EPSOPrecacheResult::Complete)
 			{
 				PSOPrecacheResult = EPSOPrecacheResult::NotSupported;
-				return PSOReadyGraphTask;
+				return PSOReadyGraphTasks;
 			}
 		}
 
@@ -3358,7 +3358,7 @@ FGraphEventRef FVersionedNiagaraEmitterData::PrecacheComputePSOs(const UNiagaraE
 			PSOPrecacheResult = EPSOPrecacheResult::Active;
 
 			FVersionedNiagaraEmitterWeakPtr EmitterWeakPtr((UNiagaraEmitter*)(&NiagaraEmitter), Version.VersionGuid);
-			PSOReadyGraphTask = TGraphTask<FNiagaraEmitterPSOPrecacheReadyTask>::CreateTask(&PSOPrecacheEvents, ENamedThreads::GameThread).ConstructAndDispatchWhenReady(EmitterWeakPtr);
+			PSOReadyGraphTasks.Add( TGraphTask<FNiagaraEmitterPSOPrecacheReadyTask>::CreateTask(&PSOPrecacheEvents, ENamedThreads::GameThread).ConstructAndDispatchWhenReady(EmitterWeakPtr));
 		}
 		else
 		{
@@ -3368,13 +3368,19 @@ FGraphEventRef FVersionedNiagaraEmitterData::PrecacheComputePSOs(const UNiagaraE
 	// In this mode we either force them to cache or respect that precaching is enabled
 	else if (GNiagaraEmitterComputePSOPrecacheMode == 2 || PipelineStateCache::IsPSOPrecachingEnabled() )
 	{
+		static IConsoleVariable* CVarPSOProxyCreationWhenPSOReady = IConsoleManager::Get().FindConsoleVariable(TEXT("r.PSOPrecache.ProxyCreationWhenPSOReady"));
+		bool bAddToPSOReadyGraphTasks = GNiagaraEmitterComputePSOPrecacheMode == 1 && CVarPSOProxyCreationWhenPSOReady && CVarPSOProxyCreationWhenPSOReady->GetInt() != 0;
 		for (int32 i=0; i < ShaderScript->GetNumPermutations(); ++i)
 		{
 			FRHIComputeShader* ComputeShader = ShaderScript->GetShaderGameThread(i).GetComputeShader();
-			PipelineStateCache::PrecacheComputePipelineState(ComputeShader, true);
+			FGraphEventRef PSOPrecacheEvent = PipelineStateCache::PrecacheComputePipelineState(ComputeShader, true);
+			if (PSOPrecacheEvent != nullptr && bAddToPSOReadyGraphTasks)
+			{
+				PSOReadyGraphTasks.Add(PSOPrecacheEvent);
+			}
 		}
 	}
-	return PSOReadyGraphTask;
+	return PSOReadyGraphTasks;
 }
 
 #if WITH_EDITORONLY_DATA
diff --git a/Engine/Plugins/FX/Niagara/Source/Niagara/Private/NiagaraSystem.cpp b/Engine/Plugins/FX/Niagara/Source/Niagara/Private/NiagaraSystem.cpp
--- a/Engine/Plugins/FX/Niagara/Source/Niagara/Private/NiagaraSystem.cpp
+++ b/Engine/Plugins/FX/Niagara/Source/Niagara/Private/NiagaraSystem.cpp
@@ -1828,10 +1828,7 @@ void UNiagaraSystem::CacheFromCompiledData()
 			Handle.GetInstance().Emitter->ConditionalPostLoad();
 			EmitterData->CacheFromCompiledData(DataSetCompiledData, *Handle.GetInstance().Emitter);
 
-			if ( FGraphEventRef PSOGraphEvent = EmitterData->PrecacheComputePSOs(*Handle.GetInstance().Emitter) )
-			{
-				PSOPrecacheEvents.Add(PSOGraphEvent);
-			}
+			PSOPrecacheEvents.Append(EmitterData->PrecacheComputePSOs(*Handle.GetInstance().Emitter));
 
 			// Allow data interfaces to cache static buffers
 			UNiagaraScript* NiagaraEmitterScripts[] =
