- Add Gather support to 2d / 2d array DIs

diff --git a/Engine/Plugins/FX/Niagara/Shaders/Private/NiagaraDataInterfaceTexture2DArrayTemplate.ush b/Engine/Plugins/FX/Niagara/Shaders/Private/NiagaraDataInterfaceTexture2DArrayTemplate.ush
--- a/Engine/Plugins/FX/Niagara/Shaders/Private/NiagaraDataInterfaceTexture2DArrayTemplate.ush
+++ b/Engine/Plugins/FX/Niagara/Shaders/Private/NiagaraDataInterfaceTexture2DArrayTemplate.ush
@@ -9,6 +9,11 @@ void LoadTexture_{ParameterName}(in int TexelX, in int TexelY, in int TexelZ, in
 	OutValue = {ParameterName}_Texture.Load(int4(TexelX, TexelY, TexelZ, MipLevel));
 }
 
+void GatherRedTexture_{ParameterName}(in float3 UVW, out float4 OutValue)
+{
+	OutValue = {ParameterName}_Texture.Gather({ParameterName}_TextureSampler, UVW);
+}
+
 void SampleTexture_{ParameterName}(in float3 UVW, in float MipLevel, out float4 OutValue)
 {
 	OutValue = {ParameterName}_Texture.SampleLevel({ParameterName}_TextureSampler, UVW, MipLevel);
diff --git a/Engine/Plugins/FX/Niagara/Shaders/Private/NiagaraDataInterfaceTextureTemplate.ush b/Engine/Plugins/FX/Niagara/Shaders/Private/NiagaraDataInterfaceTextureTemplate.ush
--- a/Engine/Plugins/FX/Niagara/Shaders/Private/NiagaraDataInterfaceTextureTemplate.ush
+++ b/Engine/Plugins/FX/Niagara/Shaders/Private/NiagaraDataInterfaceTextureTemplate.ush
@@ -10,6 +10,11 @@ void LoadTexture2D_{ParameterName}(in int TexelX, in int TexelY, in int MipLevel
 	OutValue = {ParameterName}_Texture.Load(int3(TexelX, TexelY, MipLevel));
 }
 
+void GatherRedTexture2D_{ParameterName}(in float2 UV, out float4 OutValue)
+{
+	OutValue = {ParameterName}_Texture.Gather({ParameterName}_TextureSampler, UV);
+}
+
 void SampleTexture2D_{ParameterName}(in float2 UV, in float MipLevel, out float4 OutValue)
 {
 	OutValue = {ParameterName}_Texture.SampleLevel({ParameterName}_TextureSampler, UV, MipLevel);
diff --git a/Engine/Plugins/FX/Niagara/Source/Niagara/Classes/NiagaraDataInterface2DArrayTexture.h b/Engine/Plugins/FX/Niagara/Source/Niagara/Classes/NiagaraDataInterface2DArrayTexture.h
--- a/Engine/Plugins/FX/Niagara/Source/Niagara/Classes/NiagaraDataInterface2DArrayTexture.h
+++ b/Engine/Plugins/FX/Niagara/Source/Niagara/Classes/NiagaraDataInterface2DArrayTexture.h
@@ -67,6 +67,7 @@ protected:
 protected:
 	static const TCHAR* TemplateShaderFilePath;
 	static const FName LoadTextureName;
+	static const FName GatherRedTextureName;
 	static const FName SampleTextureName;
 	static const FName TextureDimsName;
 };
diff --git a/Engine/Plugins/FX/Niagara/Source/Niagara/Classes/NiagaraDataInterfaceTexture.h b/Engine/Plugins/FX/Niagara/Source/Niagara/Classes/NiagaraDataInterfaceTexture.h
--- a/Engine/Plugins/FX/Niagara/Source/Niagara/Classes/NiagaraDataInterfaceTexture.h
+++ b/Engine/Plugins/FX/Niagara/Source/Niagara/Classes/NiagaraDataInterfaceTexture.h
@@ -72,6 +72,7 @@ protected:
 protected:
 	static const TCHAR* TemplateShaderFilePath;
 	static const FName LoadTexture2DName;
+	static const FName GatherRedTexture2DName;
 	static const FName SampleTexture2DName;
 	static const FName SamplePseudoVolumeTextureName;
 	static const FName GetTextureDimensionsName;
diff --git a/Engine/Plugins/FX/Niagara/Source/Niagara/Private/NiagaraDataInterface2DArrayTexture.cpp b/Engine/Plugins/FX/Niagara/Source/Niagara/Private/NiagaraDataInterface2DArrayTexture.cpp
--- a/Engine/Plugins/FX/Niagara/Source/Niagara/Private/NiagaraDataInterface2DArrayTexture.cpp
+++ b/Engine/Plugins/FX/Niagara/Source/Niagara/Private/NiagaraDataInterface2DArrayTexture.cpp
@@ -19,6 +19,7 @@
 
 const TCHAR* UNiagaraDataInterface2DArrayTexture::TemplateShaderFilePath = TEXT("/Plugin/FX/Niagara/Private/NiagaraDataInterfaceTexture2DArrayTemplate.ush");
 const FName UNiagaraDataInterface2DArrayTexture::LoadTextureName(TEXT("LoadTexture"));
+const FName UNiagaraDataInterface2DArrayTexture::GatherRedTextureName(TEXT("GatherRedTexture"));
 const FName UNiagaraDataInterface2DArrayTexture::SampleTextureName(TEXT("SampleTexture"));
 const FName UNiagaraDataInterface2DArrayTexture::TextureDimsName(TEXT("TextureDimensions"));
 
@@ -107,15 +108,16 @@ void UNiagaraDataInterface2DArrayTexture::GetFunctions(TArray<FNiagaraFunctionSi
 {
 	OutFunctions.Reserve(OutFunctions.Num() + 3);
 
+	FNiagaraFunctionSignature DefaultSig;
+	DefaultSig.bMemberFunction = true;
+	DefaultSig.bRequiresContext = false;
+	DefaultSig.bSupportsCPU = false;
+	DefaultSig.bSupportsGPU = true;
+	DefaultSig.Inputs.Emplace(FNiagaraTypeDefinition(GetClass()), TEXT("Texture"));
 
 	{
-		FNiagaraFunctionSignature& Sig = OutFunctions.AddDefaulted_GetRef();
+		FNiagaraFunctionSignature& Sig = OutFunctions.Add_GetRef(DefaultSig);
 		Sig.Name = LoadTextureName;
-		Sig.bMemberFunction = true;
-		Sig.bRequiresContext = false;
-		Sig.bSupportsCPU = false;
-		Sig.bSupportsGPU = true;
-		Sig.Inputs.Emplace(FNiagaraTypeDefinition(GetClass()), TEXT("Texture"));
 		Sig.Inputs.Emplace(FNiagaraTypeDefinition::GetIntDef(), TEXT("TexelX"));
 		Sig.Inputs.Emplace(FNiagaraTypeDefinition::GetIntDef(), TEXT("TexelY"));
 		Sig.Inputs.Emplace(FNiagaraTypeDefinition::GetIntDef(), TEXT("TexelZ"));
@@ -125,13 +127,16 @@ void UNiagaraDataInterface2DArrayTexture::GetFunctions(TArray<FNiagaraFunctionSi
 	}
 
 	{
-		FNiagaraFunctionSignature& Sig = OutFunctions.AddDefaulted_GetRef();
+		FNiagaraFunctionSignature& Sig = OutFunctions.Add_GetRef(DefaultSig);
+		Sig.Name = GatherRedTextureName;
+		Sig.Inputs.Emplace(FNiagaraTypeDefinition::GetVec3Def(), TEXT("UVW"));
+		Sig.Outputs.Emplace(FNiagaraTypeDefinition::GetVec4Def(), TEXT("Value"));
+		Sig.SetDescription(LOCTEXT("TextureGatherTexture2ArrayDDesc", "Gather the 4 samples (Red only) that would be used in bilinear interpolation."));
+	}
+
+	{
+		FNiagaraFunctionSignature& Sig = OutFunctions.Add_GetRef(DefaultSig);
 		Sig.Name = SampleTextureName;
-		Sig.bMemberFunction = true;
-		Sig.bRequiresContext = false;
-		Sig.bSupportsCPU = false;
-		Sig.bSupportsGPU = true;
-		Sig.Inputs.Emplace(FNiagaraTypeDefinition(GetClass()), TEXT("Texture"));
 		Sig.Inputs.Emplace(FNiagaraTypeDefinition::GetVec3Def(), TEXT("UVW"));
 		Sig.Inputs.Emplace(FNiagaraTypeDefinition::GetFloatDef(), TEXT("MipLevel"));
 		Sig.Outputs.Emplace(FNiagaraTypeDefinition::GetVec4Def(), TEXT("Value"));
@@ -261,6 +266,7 @@ bool UNiagaraDataInterface2DArrayTexture::GetFunctionHLSL(const FNiagaraDataInte
 	static const TSet<FName> ValidGpuFunctions =
 	{
 		LoadTextureName,
+		GatherRedTextureName,
 		SampleTextureName,
 		TextureDimsName,
 	};
diff --git a/Engine/Plugins/FX/Niagara/Source/Niagara/Private/NiagaraDataInterfaceTexture.cpp b/Engine/Plugins/FX/Niagara/Source/Niagara/Private/NiagaraDataInterfaceTexture.cpp
--- a/Engine/Plugins/FX/Niagara/Source/Niagara/Private/NiagaraDataInterfaceTexture.cpp
+++ b/Engine/Plugins/FX/Niagara/Source/Niagara/Private/NiagaraDataInterfaceTexture.cpp
@@ -20,6 +20,7 @@
 
 const TCHAR* UNiagaraDataInterfaceTexture::TemplateShaderFilePath = TEXT("/Plugin/FX/Niagara/Private/NiagaraDataInterfaceTextureTemplate.ush");
 const FName UNiagaraDataInterfaceTexture::LoadTexture2DName(TEXT("LoadTexture2D"));
+const FName UNiagaraDataInterfaceTexture::GatherRedTexture2DName(TEXT("GatherRedTexture2D"));
 const FName UNiagaraDataInterfaceTexture::SampleTexture2DName(TEXT("SampleTexture2D"));
 const FName UNiagaraDataInterfaceTexture::SamplePseudoVolumeTextureName(TEXT("SamplePseudoVolumeTexture"));
 const FName UNiagaraDataInterfaceTexture::GetTextureDimensionsName(TEXT("GetTextureDimensions"));
@@ -165,6 +166,13 @@ void UNiagaraDataInterfaceTexture::GetFunctions(TArray<FNiagaraFunctionSignature
 		Sig.Outputs.Emplace(FNiagaraTypeDefinition::GetVec4Def(), TEXT("Value"));
 		Sig.SetDescription(LOCTEXT("TextureLoadTexture2DDesc", "Read a texel from the provided location & mip without any filtering or sampling."));
 	}
+	{
+		FNiagaraFunctionSignature& Sig = OutFunctions.Add_GetRef(DefaultGpuSig);
+		Sig.Name = GatherRedTexture2DName;
+		Sig.Inputs.Emplace(FNiagaraTypeDefinition::GetVec2Def(), TEXT("UV"));
+		Sig.Outputs.Emplace(FNiagaraTypeDefinition::GetVec4Def(), TEXT("Value"));
+		Sig.SetDescription(LOCTEXT("TextureGatherTexture2DDesc", "Gather the 4 samples (Red only) that would be used in bilinear interpolation."));
+	}
 	{
 		FNiagaraFunctionSignature& Sig = OutFunctions.Add_GetRef(DefaultGpuSig);
 		Sig.Name = SampleTexture2DName;
@@ -394,6 +402,7 @@ bool UNiagaraDataInterfaceTexture::GetFunctionHLSL(const FNiagaraDataInterfaceGP
 	static const TSet<FName> ValidGpuFunctions =
 	{
 		LoadTexture2DName,
+		GatherRedTexture2DName,
 		SampleTexture2DName,
 		SamplePseudoVolumeTextureName,
 		GetTextureDimensionsName,
