- Change max planes from 10 to 12 to be consistent with other parts of the engine
- Ensure and handle > 12 rather than hitting a check

diff --git a/Engine/Plugins/FX/Niagara/Source/Niagara/Classes/NiagaraGPUSortInfo.h b/Engine/Plugins/FX/Niagara/Source/Niagara/Classes/NiagaraGPUSortInfo.h
--- a/Engine/Plugins/FX/Niagara/Source/Niagara/Classes/NiagaraGPUSortInfo.h
+++ b/Engine/Plugins/FX/Niagara/Source/Niagara/Classes/NiagaraGPUSortInfo.h
@@ -35,7 +35,7 @@ enum class ENiagaraSortMode : uint8
 
 struct FNiagaraGPUSortInfo
 {
-	static constexpr uint32 MaxCullPlanes = 10;
+	static constexpr uint32 MaxCullPlanes = 12;
 
 	// The number of particles in the system.
 	int32 ParticleCount = 0;
diff --git a/Engine/Plugins/FX/Niagara/Source/Niagara/Private/NiagaraRendererMeshes.cpp b/Engine/Plugins/FX/Niagara/Source/Niagara/Private/NiagaraRendererMeshes.cpp
--- a/Engine/Plugins/FX/Niagara/Source/Niagara/Private/NiagaraRendererMeshes.cpp
+++ b/Engine/Plugins/FX/Niagara/Source/Niagara/Private/NiagaraRendererMeshes.cpp
@@ -540,14 +540,16 @@ void FNiagaraRendererMeshes::InitializeSortInfo(const FParticleMeshRenderData& P
 		{
 			// Ensure we don't break the maximum number of planes here
 			// (For an accurate shadow frustum, a tight hull is formed from the silhouette and back-facing planes of the view frustum)
-			check(ShadowFrustum->Planes.Num() <= FNiagaraGPUSortInfo::MaxCullPlanes);
-			OutSortInfo.CullPlanes = ShadowFrustum->Planes;
-
-			// Remove pre-shadow translation to get the planes in world space
-			const FVector PreShadowTranslation = View.GetPreShadowTranslation();
-			for (FPlane& Plane : OutSortInfo.CullPlanes)
+			if (ensure(ShadowFrustum->Planes.Num() <= FNiagaraGPUSortInfo::MaxCullPlanes))
 			{
-				Plane.W -= FVector::DotProduct(FVector(Plane), PreShadowTranslation);
+				OutSortInfo.CullPlanes = ShadowFrustum->Planes;
+
+				// Remove pre-shadow translation to get the planes in world space
+				const FVector PreShadowTranslation = View.GetPreShadowTranslation();
+				for (FPlane& Plane : OutSortInfo.CullPlanes)
+				{
+					Plane.W -= FVector::DotProduct(FVector(Plane), PreShadowTranslation);
+				}
 			}
 		}
 		else
diff --git a/Engine/Plugins/FX/Niagara/Source/NiagaraVertexFactories/Public/NiagaraSortingGPU.h b/Engine/Plugins/FX/Niagara/Source/NiagaraVertexFactories/Public/NiagaraSortingGPU.h
--- a/Engine/Plugins/FX/Niagara/Source/NiagaraVertexFactories/Public/NiagaraSortingGPU.h
+++ b/Engine/Plugins/FX/Niagara/Source/NiagaraVertexFactories/Public/NiagaraSortingGPU.h
@@ -19,7 +19,7 @@ extern NIAGARAVERTEXFACTORIES_API int32 GNiagaraGPUCullingCPUToGPUThreshold;
 #define NIAGARA_KEY_GEN_THREAD_COUNT 64
 #define NIAGARA_COPY_BUFFER_THREAD_COUNT 64
 #define NIAGARA_COPY_BUFFER_BUFFER_COUNT 3
-#define NIAGARA_KEY_GEN_MAX_CULL_PLANES 10
+#define NIAGARA_KEY_GEN_MAX_CULL_PLANES 12
 
 /**
  * Compute shader used to generate particle sort keys.
