Fix more Strata blend mode.

diff --git a/Engine/Plugins/FX/Niagara/Source/Niagara/Private/NiagaraRendererSprites.cpp b/Engine/Plugins/FX/Niagara/Source/Niagara/Private/NiagaraRendererSprites.cpp
--- a/Engine/Plugins/FX/Niagara/Source/Niagara/Private/NiagaraRendererSprites.cpp
+++ b/Engine/Plugins/FX/Niagara/Source/Niagara/Private/NiagaraRendererSprites.cpp
@@ -230,9 +230,9 @@ void FNiagaraRendererSprites::PrepareParticleSpriteRenderData(FParticleSpriteRen
 
 	// Do we have anything to render?
 	const FMaterial& Material = MaterialRenderProxy->GetIncompleteMaterialWithFallback(FeatureLevel);
-	const EBlendMode BlendMode = Material.GetBlendMode();
-	ParticleSpriteRenderData.BlendMode = BlendMode; // STRATA_TODO_BLENDMODE
-	ParticleSpriteRenderData.bHasTranslucentMaterials = IsTranslucentBlendMode(BlendMode);
+	ParticleSpriteRenderData.BlendMode = Material.GetBlendMode();
+	ParticleSpriteRenderData.StrataBlendMode = Material.GetStrataBlendMode();
+	ParticleSpriteRenderData.bHasTranslucentMaterials = IsTranslucentBlendMode(Material);
 
 	// If these conditions change please update the DebugHUD display also to reflect it
 	const bool bLowLatencyTranslucencyEnabled = ParticleSpriteRenderData.bHasTranslucentMaterials &&
@@ -265,7 +265,7 @@ void FNiagaraRendererSprites::PrepareParticleSpriteRenderData(FParticleSpriteRen
 		const EShaderPlatform ShaderPlatform = SceneProxy->GetComputeDispatchInterface()->GetShaderPlatform();
 
 		// Determine if we need sorting
-		ParticleSpriteRenderData.bNeedsSort = SortMode != ENiagaraSortMode::None && (BlendMode == BLEND_AlphaComposite || IsAlphaHoldoutBlendMode(Material) || IsTranslucentOnlyBlendMode(Material) || !bSortOnlyWhenTranslucent);
+		ParticleSpriteRenderData.bNeedsSort = SortMode != ENiagaraSortMode::None && (IsAlphaCompositeBlendMode(Material) || IsAlphaHoldoutBlendMode(Material) || IsTranslucentOnlyBlendMode(Material) || !bSortOnlyWhenTranslucent);
 		const bool bNeedCustomSort = ParticleSpriteRenderData.bNeedsSort && (SortMode == ENiagaraSortMode::CustomAscending || SortMode == ENiagaraSortMode::CustomDecending);
 		ParticleSpriteRenderData.RendererLayout = bNeedCustomSort ? RendererLayoutWithCustomSort : RendererLayoutWithoutCustomSort;
 		ParticleSpriteRenderData.SortVariable = bNeedCustomSort ? ENiagaraSpriteVFLayout::CustomSorting : ENiagaraSpriteVFLayout::Position;
@@ -575,23 +575,23 @@ FNiagaraSpriteUniformBufferRef FNiagaraRendererSprites::CreateViewUniformBuffer(
 			PerViewUniformParameters.PixelCoverageEnabled = ParticleSpriteRenderData.bHasTranslucentMaterials;
 			if (PerViewUniformParameters.PixelCoverageEnabled)
 			{
-				switch (ParticleSpriteRenderData.BlendMode)
+				if (IsTranslucentOnlyBlendMode(ParticleSpriteRenderData.BlendMode, ParticleSpriteRenderData.StrataBlendMode))
+				{
+					ParticleSpriteRenderData.bHasTranslucentMaterials = true;
+					PerViewUniformParameters.PixelCoverageColorBlend = FVector4f(PixelCoverageBlend, PixelCoverageBlend, PixelCoverageBlend, 0.0f);
+				}
+				else if (IsAdditiveBlendMode(ParticleSpriteRenderData.BlendMode, ParticleSpriteRenderData.StrataBlendMode))
+				{
+					ParticleSpriteRenderData.bHasTranslucentMaterials = true;
+					PerViewUniformParameters.PixelCoverageColorBlend = FVector4f(PixelCoverageBlend, PixelCoverageBlend, PixelCoverageBlend, PixelCoverageBlend);
+				}
+				else
 				{
-					case BLEND_Translucent:
-						ParticleSpriteRenderData.bHasTranslucentMaterials = true;
-						PerViewUniformParameters.PixelCoverageColorBlend = FVector4f(PixelCoverageBlend, PixelCoverageBlend, PixelCoverageBlend, 0.0f);
-						break;
-					case BLEND_Additive:
-						ParticleSpriteRenderData.bHasTranslucentMaterials = true;
-						PerViewUniformParameters.PixelCoverageColorBlend = FVector4f(PixelCoverageBlend, PixelCoverageBlend, PixelCoverageBlend, PixelCoverageBlend);
-						break;
 					//-TODO: Support these blend modes
 					//BLEND_Modulate
 					//BLEND_AlphaComposite
 					//BLEND_AlphaHoldout
-					default:
-						ParticleSpriteRenderData.bHasTranslucentMaterials = false;
-						break;
+					ParticleSpriteRenderData.bHasTranslucentMaterials = false;
 				}
 			}
 		}
diff --git a/Engine/Plugins/FX/Niagara/Source/Niagara/Public/NiagaraRendererSprites.h b/Engine/Plugins/FX/Niagara/Source/Niagara/Public/NiagaraRendererSprites.h
--- a/Engine/Plugins/FX/Niagara/Source/Niagara/Public/NiagaraRendererSprites.h
+++ b/Engine/Plugins/FX/Niagara/Source/Niagara/Public/NiagaraRendererSprites.h
@@ -42,7 +42,8 @@ private:
 		const FNiagaraDynamicDataSprites*	DynamicDataSprites = nullptr;
 		class FNiagaraDataBuffer*			SourceParticleData = nullptr;
 
-		EBlendMode							BlendMode = BLEND_Opaque; // STRATA_TODO_BLENDMODE
+		EBlendMode							BlendMode = BLEND_Opaque;
+		EStrataBlendMode					StrataBlendMode = SBM_Opaque;
 		bool								bHasTranslucentMaterials = false;
 		bool								bSortCullOnGpu = false;
 		bool								bNeedsSort = false;
diff --git a/Engine/Source/Developer/MaterialUtilities/Private/MaterialUtilities.cpp b/Engine/Source/Developer/MaterialUtilities/Private/MaterialUtilities.cpp
--- a/Engine/Source/Developer/MaterialUtilities/Private/MaterialUtilities.cpp
+++ b/Engine/Source/Developer/MaterialUtilities/Private/MaterialUtilities.cpp
@@ -112,7 +112,7 @@ UMaterialInterface* FMaterialUtilities::CreateProxyMaterialAndTextures(UPackage*
 	Material->BasePropertyOverrides.DitheredLODTransition = MaterialData.Material->IsDitheredLODTransition();
 	Material->BasePropertyOverrides.bOverride_DitheredLODTransition = MaterialData.Material->IsDitheredLODTransition();
 
-	if (MaterialData.Material->GetBlendMode() != BLEND_Opaque) // STRATA_TODO_BLENDMODE
+	if (!IsOpaqueBlendMode(*MaterialData.Material))
 	{
 		Material->BasePropertyOverrides.bOverride_BlendMode = true;
 		Material->BasePropertyOverrides.BlendMode = MaterialData.Material->GetBlendMode();
diff --git a/Engine/Source/Editor/UnrealEd/Private/Lightmass/LightmassRender.cpp b/Engine/Source/Editor/UnrealEd/Private/Lightmass/LightmassRender.cpp
--- a/Engine/Source/Editor/UnrealEd/Private/Lightmass/LightmassRender.cpp
+++ b/Engine/Source/Editor/UnrealEd/Private/Lightmass/LightmassRender.cpp
@@ -403,6 +403,8 @@ public:
 			const bool bIsModulate = IsModulateBlendMode(*MaterialInterface);
 			const bool bIsTranslucentOnly = IsTranslucentOnlyBlendMode(*MaterialInterface);
 			const bool bIsAlphaHoldout = IsAlphaHoldoutBlendMode(*MaterialInterface);
+			const bool bIsAdditive = IsAdditiveBlendMode(*MaterialInterface);
+			const bool bIsAlphaComposite = IsAlphaCompositeBlendMode(*MaterialInterface);
 			check(ProxyMaterial);
 			FLightmassMaterialCompiler ProxyCompiler(Compiler);
 
@@ -444,7 +446,6 @@ public:
 				{
 					return MaterialInterface->CompileProperty(&ProxyCompiler, MP_Opacity);
 				}
-				// STRATA_TODO_BLENDMODE
 				else if (bIsModulate)
 				{
 					if (bIsMaterialUnlit)
@@ -456,8 +457,7 @@ public:
 						return Compiler->Saturate(MaterialInterface->CompileProperty(Compiler, DiffuseInput, ForceCast_Exact_Replicate));
 					}
 				}
-				// STRATA_TODO_BLENDMODE
-				else if (bIsTranslucentOnly || (BlendMode == BLEND_Additive) || (BlendMode == BLEND_AlphaComposite) || bIsAlphaHoldout)
+				else if (bIsTranslucentOnly || bIsAdditive || bIsAlphaComposite || bIsAlphaHoldout)
 				{
 					int32 ColoredOpacity = INDEX_NONE;
 					if (bIsMaterialUnlit)
diff --git a/Engine/Source/Runtime/Engine/Private/DebugViewModeInterface.cpp b/Engine/Source/Runtime/Engine/Private/DebugViewModeInterface.cpp
--- a/Engine/Source/Runtime/Engine/Private/DebugViewModeInterface.cpp
+++ b/Engine/Source/Runtime/Engine/Private/DebugViewModeInterface.cpp
@@ -12,7 +12,6 @@ DebugViewModeInterface.cpp: Contains definitions for rendering debug viewmodes.
 
 FDebugViewModeInterface* FDebugViewModeInterface::Singleton = nullptr;
 
-// STRATA_TODO_BLENDMODE
 void FDebugViewModeInterface::SetDrawRenderState(EDebugViewShaderMode DebugViewMode, const FMaterial& InMaterial, FRenderState& DrawRenderState, bool bHasDepthPrepassForMaskedMaterial) const
 {
 	if (DebugViewMode == DVSM_QuadComplexity || DebugViewMode == DVSM_ShaderComplexityBleedingQuadOverhead || DebugViewMode == DVSM_ShaderComplexityContainedQuadOverhead || DebugViewMode == DVSM_ShaderComplexity)
@@ -56,7 +55,7 @@ void FDebugViewModeInterface::SetDrawRenderState(EDebugViewShaderMode DebugViewM
 			DrawRenderState.BlendState = TStaticBlendState<>::GetRHI();
 
 			// If not selected, use depth equal to make alpha test stand out (goes with EarlyZPassMode = DDM_AllOpaque) 
-			if (IsMaskedBlendMode(InMaterial) && bHasDepthPrepassForMaskedMaterial) // STRATA_TODO_BLENDMODE
+			if (IsMaskedBlendMode(InMaterial) && bHasDepthPrepassForMaskedMaterial)
 			{
 				DrawRenderState.DepthStencilState = TStaticDepthStencilState<false, CF_Equal>::GetRHI();
 			}
diff --git a/Engine/Source/Runtime/Engine/Private/Materials/Material.cpp b/Engine/Source/Runtime/Engine/Private/Materials/Material.cpp
--- a/Engine/Source/Runtime/Engine/Private/Materials/Material.cpp
+++ b/Engine/Source/Runtime/Engine/Private/Materials/Material.cpp
@@ -2933,8 +2933,8 @@ void UMaterial::BackwardsCompatibilityVirtualTextureOutputConversion()
 			{
 				OutputExpression->GetInput(5)->Connect(EditorOnly->Opacity.OutputIndex, EditorOnly->Opacity.Expression);
 			}
-			// STRATA_TODO_BLENDMODE
-			if (BlendMode != BLEND_Opaque)
+
+			if (!IsOpaqueBlendMode(BlendMode, StrataBlendMode))
 			{
 				// Full alpha blend modes were mostly/always used with MD_RuntimeVirtualTexture to allow pin connections.
 				// But we will assume the intention for any associated MD_Surface output is opaque or alpha mask and force convert here.
@@ -2944,6 +2944,7 @@ void UMaterial::BackwardsCompatibilityVirtualTextureOutputConversion()
 					EditorOnly->Opacity.Expression = nullptr;
 				}
 				BlendMode = EditorOnly->OpacityMask.IsConnected() ? BLEND_Masked : BLEND_Opaque;
+				StrataBlendMode = EditorOnly->OpacityMask.IsConnected() ? SBM_Masked : SBM_Opaque;
 				bCanMaskedBeAssumedOpaque = !EditorOnly->OpacityMask.Expression && !(EditorOnly->OpacityMask.UseConstant && EditorOnly->OpacityMask.Constant < 0.999f);
 			}
 		}
