Compute attribute indices during hlsl generation for grid2d
force grid2d reader to use indirection, give error if trying to use named+unnamed attributes together
refactor grid3d get/set/sample hlsl generation so 2d and 3d grids share attribute processing
number of unnamed attributes on 2d grids defaults 0
deprecate named and unnamed attributes for 2d grids
fixes for EngineTest

diff --git a/Engine/Plugins/FX/Niagara/Source/Niagara/Classes/NiagaraDataInterfaceGrid2DCollection.h b/Engine/Plugins/FX/Niagara/Source/Niagara/Classes/NiagaraDataInterfaceGrid2DCollection.h
--- a/Engine/Plugins/FX/Niagara/Source/Niagara/Classes/NiagaraDataInterfaceGrid2DCollection.h
+++ b/Engine/Plugins/FX/Niagara/Source/Niagara/Classes/NiagaraDataInterfaceGrid2DCollection.h
@@ -7,6 +7,7 @@
 #include "NiagaraComponent.h"
 #include "NiagaraRenderGraphUtils.h"
 #include "NiagaraStats.h"
+#include "NiagaraDataInterfaceRWUtils.h"
 
 #include "NiagaraDataInterfaceGrid2DCollection.generated.h"
 
@@ -52,6 +53,7 @@ struct FGrid2DCollectionRWInstanceData_GameThread
 
 struct FGrid2DCollectionRWInstanceData_RenderThread
 {
+	FName SourceDIName;
 	bool ClearBeforeNonIterationStage = true;
 
 	FIntPoint NumCells = FIntPoint(EForceInit::ForceInitToZero);
@@ -188,12 +190,12 @@ public:
 	//~ UNiagaraDataInterface interface END
 
 private:
-	static void CollectAttributesForScript(UNiagaraScript* Script, FName VariableName, TArray<FNiagaraVariableBase>& OutVariables, TArray<uint32>& OutVariableOffsets, int32& TotalAttributes, TArray<FText>* OutWarnings = nullptr);
+	static void CollectAttributesForScript(UNiagaraScript* Script, FName VariableName, TArray<FNiagaraVariableBase>& OutVariables, TArray<uint32>& OutVariableOffsets, int32& TotalAttributes, TArray<FText>* OutWarnings = nullptr, bool UseReader = false);
 public:
 	/** Finds all attributes by locating the variable name inside the parameter stores. */
 	void FindAttributesByName(FName DataInterfaceName, TArray<FNiagaraVariableBase>& OutVariables, TArray<uint32>& OutVariableOffsets, int32& OutNumAttribChannelsFound, TArray<FText>* OutWarnings = nullptr) const;
 	/** Finds all attributes by locating the data interface amongst the parameter stores. */
-	void FindAttributes(TArray<FNiagaraVariableBase>& OutVariables, TArray<uint32>& OutVariableOffsets, int32& OutNumAttribChannelsFound, TArray<FText>* OutWarnings = nullptr) const;
+	void FindAttributes(TArray<FNiagaraVariableBase>& OutVariables, TArray<uint32>& OutVariableOffsets, int32& OutNumAttribChannelsFound, TArray<FText>* OutWarnings = nullptr, bool UseReader = false) const;
 
 	// Fills a texture render target 2d with the current data from the simulation
 	// #todo(dmp): this will eventually go away when we formalize how data makes it out of Niagara
@@ -285,16 +287,21 @@ public:
 	TMap<FNiagaraSystemInstanceID, FGrid2DCollectionRWInstanceData_GameThread*>& GetSystemInstancesToProxyData_GT() { return SystemInstancesToProxyData_GT; }
 protected:
 #if WITH_EDITORONLY_DATA
-	void WriteSetHLSL(const FNiagaraDataInterfaceGPUParamInfo& ParamInfo, const FNiagaraDataInterfaceGeneratedFunction& FunctionInfo, int FunctionInstanceIndex, int32 InNumChannels, FString& OutHLSL);
-	void WriteGetHLSL(const FNiagaraDataInterfaceGPUParamInfo& ParamInfo, const FNiagaraDataInterfaceGeneratedFunction& FunctionInfo, int FunctionInstanceIndex, int32 InNumChannels, FString& OutHLSL);
-	void WriteSampleHLSL(const FNiagaraDataInterfaceGPUParamInfo& ParamInfo, const FNiagaraDataInterfaceGeneratedFunction& FunctionInfo, int FunctionInstanceIndex, int32 InNumChannels, FString& OutHLSL);
-	void WriteAttributeGetIndexHLSL(const FNiagaraDataInterfaceGPUParamInfo& ParamInfo, const FNiagaraDataInterfaceGeneratedFunction& FunctionInfo, int FunctionInstanceIndex, int32 InNumChannels, FString& OutHLSL);
+	void WriteSetHLSL(const FNiagaraDataInterfaceRWAttributeHelper &AttributeHelper, const FNiagaraDataInterfaceGPUParamInfo& ParamInfo, const FNiagaraDataInterfaceGeneratedFunction& FunctionInfo, int FunctionInstanceIndex, int32 InNumChannels, FString& OutHLSL);
+	void WriteGetHLSL(const FNiagaraDataInterfaceRWAttributeHelper& AttributeHelper, const FNiagaraDataInterfaceGPUParamInfo& ParamInfo, const FNiagaraDataInterfaceGeneratedFunction& FunctionInfo, int FunctionInstanceIndex, int32 InNumChannels, FString& OutHLSL);
+	void WriteSampleHLSL(const FNiagaraDataInterfaceRWAttributeHelper& AttributeHelper, const FNiagaraDataInterfaceGPUParamInfo& ParamInfo, const FNiagaraDataInterfaceGeneratedFunction& FunctionInfo, int FunctionInstanceIndex, int32 InNumChannels, FString& OutHLSL);
+	void WriteAttributeGetIndexHLSL(const FNiagaraDataInterfaceRWAttributeHelper& AttributeHelper, const FNiagaraDataInterfaceGPUParamInfo& ParamInfo, const FNiagaraDataInterfaceGeneratedFunction& FunctionInfo, int FunctionInstanceIndex, int32 InNumChannels, FString& OutHLSL);
+
+	FString GenerateAttributeIndexHLSL(const FNiagaraDataInterfaceRWAttributeHelper& AttributeHelper, const FNiagaraDataInterfaceGeneratedFunction& FunctionInfo);
 
 	const TCHAR* TypeDefinitionToHLSLTypeString(const FNiagaraTypeDefinition& InDef) const;
 	FName TypeDefinitionToGetFunctionName(const FNiagaraTypeDefinition& InDef) const;
 	FName TypeDefinitionToSetFunctionName(const FNiagaraTypeDefinition& InDef) const;
 #endif
 
+	// Only used with CDO for hlsl generation
+	bool bUseIndirection;
+
 	//~ UNiagaraDataInterface interface
 	virtual bool CopyToInternal(UNiagaraDataInterface* Destination) const override;
 
diff --git a/Engine/Plugins/FX/Niagara/Source/Niagara/Classes/NiagaraDataInterfaceGrid2DCollectionReader.h b/Engine/Plugins/FX/Niagara/Source/Niagara/Classes/NiagaraDataInterfaceGrid2DCollectionReader.h
--- a/Engine/Plugins/FX/Niagara/Source/Niagara/Classes/NiagaraDataInterfaceGrid2DCollectionReader.h
+++ b/Engine/Plugins/FX/Niagara/Source/Niagara/Classes/NiagaraDataInterfaceGrid2DCollectionReader.h
@@ -29,6 +29,8 @@ public:
 	UPROPERTY(EditAnywhere, Category = "Reader")
 	FString DIName;
 
+	virtual void Serialize(FArchive& Ar) override { UNiagaraDataInterfaceRWBase::Serialize(Ar); }
+
 	virtual bool Equals(const UNiagaraDataInterface* Other) const override;
 	virtual bool InitPerInstanceData(void* PerInstanceData, FNiagaraSystemInstance* SystemInstance) override;
 	virtual void GetEmitterDependencies(UNiagaraSystem* Asset, TArray<FVersionedNiagaraEmitter>& Dependencies) const override;
diff --git a/Engine/Plugins/FX/Niagara/Source/Niagara/Classes/NiagaraDataInterfaceGrid3DCollection.h b/Engine/Plugins/FX/Niagara/Source/Niagara/Classes/NiagaraDataInterfaceGrid3DCollection.h
--- a/Engine/Plugins/FX/Niagara/Source/Niagara/Classes/NiagaraDataInterfaceGrid3DCollection.h
+++ b/Engine/Plugins/FX/Niagara/Source/Niagara/Classes/NiagaraDataInterfaceGrid3DCollection.h
@@ -127,7 +127,7 @@ class NIAGARA_API UNiagaraDataInterfaceGrid3DCollection : public UNiagaraDataInt
 
 public:
 	// Number of attributes stored on the grid
-	UPROPERTY(EditAnywhere, Category = "Grid")
+	UPROPERTY(EditAnywhere, AdvancedDisplay, Category = "Grid")
 	int32 NumAttributes;
 
 	/** Reference to a user parameter if we're reading one. */
@@ -234,7 +234,6 @@ public:
 	static const FString UnitClampMinName;
 	static const FString UnitClampMaxName;
 
-	static const FName NAME_Attribute;
 	static const FString GridName;
 	static const FString OutputGridName;
 	static const FString SamplerName;
diff --git a/Engine/Plugins/FX/Niagara/Source/Niagara/Classes/NiagaraDataInterfaceRW.h b/Engine/Plugins/FX/Niagara/Source/Niagara/Classes/NiagaraDataInterfaceRW.h
--- a/Engine/Plugins/FX/Niagara/Source/Niagara/Classes/NiagaraDataInterfaceRW.h
+++ b/Engine/Plugins/FX/Niagara/Source/Niagara/Classes/NiagaraDataInterfaceRW.h
@@ -44,6 +44,9 @@ public:
 	static const FString CellSizeName;
 	static const FString WorldBBoxSizeName;
 
+	// Attribute names
+	static const FName NAME_Attribute;
+
 	// Global VM function names, also used by the shaders code generation methods.
 	static const FName NumCellsFunctionName;
 	static const FName CellSizeFunctionName;
@@ -189,8 +192,8 @@ public:
 	int32 NumCellsMaxAxis;
 
 	// Number of Attributes
-	UPROPERTY(EditAnywhere, Category = "Grid")
-	int32 NumAttributes;
+	UPROPERTY(EditAnywhere, Category = "Grid", AdvancedDisplay)
+	int32 NumAttributes = 0;
 
 	// Set grid resolution according to longest axis
 	UPROPERTY(EditAnywhere, Category = "Grid", meta = (EditCondition = "false", EditConditionHides))
@@ -202,6 +205,9 @@ public:
 
 
 public:
+
+	virtual void Serialize(FArchive& Ar) override;
+
 	//~ UNiagaraDataInterface interface
 	// VM functionality
 	virtual void GetFunctions(TArray<FNiagaraFunctionSignature>& OutFunctions) override;
diff --git a/Engine/Plugins/FX/Niagara/Source/Niagara/Classes/NiagaraDataInterfaceRWUtils.h b/Engine/Plugins/FX/Niagara/Source/Niagara/Classes/NiagaraDataInterfaceRWUtils.h
new file mode 100644
--- /dev/null
+++ b/Engine/Plugins/FX/Niagara/Source/Niagara/Classes/NiagaraDataInterfaceRWUtils.h
@@ -0,0 +1,45 @@
+// Copyright Epic Games, Inc. All Rights Reserved.
+#pragma once
+
+#include "NiagaraDataInterface.h"
+
+struct FNiagaraDataInterfaceRWAttributeHelper	
+{
+	static TArray<FString> Channels;
+
+	struct FAttributeInfo
+	{
+		FName					Name;
+		FNiagaraTypeDefinition	TypeDef;
+		int32					NumChannels = 0;
+		int32					ChannelOffset = 0;
+		int32					AttributeIndex = 0;
+	};
+
+	explicit FNiagaraDataInterfaceRWAttributeHelper(const FNiagaraDataInterfaceGPUParamInfo& InParamInfo);
+	
+	template<typename T> void Init(TArray<FText>* OutWarnings = nullptr);
+	const FAttributeInfo* FindAttributeInfo(FName AttributeName) const;
+
+	int GetNumAttributes() const
+	{
+		return AttributeInfos.Num();
+	}
+
+	int GetTotalChannels() const
+	{
+		return TotalChannels;
+	}
+
+#if WITH_EDITORONLY_DATA
+	// Translates named attribute into actual attribute index	
+	FString GetAttributeIndex(const FAttributeInfo* AttributeInfo, int Channel = 0) const;
+	FString GetGridChannelString()  const;
+	void GetChannelStrings(int AttributeIndex, int AttributeNumChannels, FString& NumChannelsString, FString& AttrGridChannels)  const;
+	void GetChannelStrings(const FAttributeInfo* AttributeInfo, FString& NumChannelsString, FString& AttrGridChannels)  const;
+#endif //WITH_EDITORONLY_DATA
+
+	const FNiagaraDataInterfaceGPUParamInfo& ParamInfo;
+	TArray<FAttributeInfo>						AttributeInfos;
+	int32										TotalChannels = 0;
+};
\ No newline at end of file
diff --git a/Engine/Plugins/FX/Niagara/Source/Niagara/Private/NiagaraDataInterfaceGrid2DCollection.cpp b/Engine/Plugins/FX/Niagara/Source/Niagara/Private/NiagaraDataInterfaceGrid2DCollection.cpp
--- a/Engine/Plugins/FX/Niagara/Source/Niagara/Private/NiagaraDataInterfaceGrid2DCollection.cpp
+++ b/Engine/Plugins/FX/Niagara/Source/Niagara/Private/NiagaraDataInterfaceGrid2DCollection.cpp
@@ -21,6 +21,8 @@
 #include "ShaderParameterUtils.h"
 #include "TextureResource.h"
 
+#include "NiagaraDataInterfaceGrid2DCollectionReader.h"
+
 #include UE_INLINE_GENERATED_CPP_BY_NAME(NiagaraDataInterfaceGrid2DCollection)
 
 #define LOCTEXT_NAMESPACE "NiagaraDataInterfaceGrid2DCollection"
@@ -246,6 +248,7 @@ UNiagaraDataInterfaceGrid2DCollection::UNiagaraDataInterfaceGrid2DCollection(FOb
 #if WITH_EDITORONLY_DATA
 	, bPreviewGrid(false)
 #endif
+	, bUseIndirection(false)
 {
 	Proxy.Reset(new FNiagaraDataInterfaceProxyGrid2DCollectionProxy());
 
@@ -275,6 +278,10 @@ bool UNiagaraDataInterfaceGrid2DCollection::AppendCompileHash(FNiagaraCompileHas
 
 
 	InVisitor->UpdatePOD(TEXT("UNiagaraDataInterfaceGrid2DCollectionVersion"), (int32)FNiagaraGridCollection2DDIFunctionVersion::LatestVersion);
+
+	const FGuid FNiagara2DGridCollectionGUID(0x2C8975F3, 0xC8AB4145, 0xA312186C, 0xC9BCBA2A);
+	InVisitor->UpdatePOD(TEXT("UNiagaraDataInterfaceGrid2DCollectionGUID"), FNiagara2DGridCollectionGUID);
+
 	InVisitor->UpdateShaderParameters<FNDIGrid2DShaderParameters>();
 	return true;
 }
@@ -911,8 +918,6 @@ void UNiagaraDataInterfaceGrid2DCollection::GetVMExternalFunction(const FVMExter
 {
 	Super::GetVMExternalFunction(BindingInfo, InstanceData, OutFunc);
 
-	static const FName NAME_Attribute("Attribute");
-
 	if (BindingInfo.Name == UNiagaraDataInterfaceRWBase::WorldBBoxSizeFunctionName)
 	{
 		check(BindingInfo.GetNumInputs() == 1 && BindingInfo.GetNumOutputs() == 2);
@@ -940,25 +945,25 @@ void UNiagaraDataInterfaceGrid2DCollection::GetVMExternalFunction(const FVMExter
 
 	else if (BindingInfo.Name == GetVector4AttributeIndexFunctionName)
 	{
-		FName AttributeName = BindingInfo.FindSpecifier(NAME_Attribute)->Value;
+		FName AttributeName = BindingInfo.FindSpecifier(UNiagaraDataInterfaceRWBase::NAME_Attribute)->Value;
 		check(BindingInfo.GetNumInputs() == 1 && BindingInfo.GetNumOutputs() == 1);
 		NDI_FUNC_BINDER(UNiagaraDataInterfaceGrid2DCollection, VMGetAttributeIndex)::Bind(this, OutFunc, AttributeName, 4);
 	}
 	else if (BindingInfo.Name == GetVectorAttributeIndexFunctionName)
 	{
-		FName AttributeName = BindingInfo.FindSpecifier(NAME_Attribute)->Value;
+		FName AttributeName = BindingInfo.FindSpecifier(UNiagaraDataInterfaceRWBase::NAME_Attribute)->Value;
 		check(BindingInfo.GetNumInputs() == 1 && BindingInfo.GetNumOutputs() == 1);
 		NDI_FUNC_BINDER(UNiagaraDataInterfaceGrid2DCollection, VMGetAttributeIndex)::Bind(this, OutFunc, AttributeName, 3);
 	}
 	else if (BindingInfo.Name == GetVector2DAttributeIndexFunctionName)
 	{
-		FName AttributeName = BindingInfo.FindSpecifier(NAME_Attribute)->Value;
+		FName AttributeName = BindingInfo.FindSpecifier(UNiagaraDataInterfaceRWBase::NAME_Attribute)->Value;
 		check(BindingInfo.GetNumInputs() == 1 && BindingInfo.GetNumOutputs() == 1);
 		NDI_FUNC_BINDER(UNiagaraDataInterfaceGrid2DCollection, VMGetAttributeIndex)::Bind(this, OutFunc, AttributeName, 2);
 	}
 	else if (BindingInfo.Name == GetFloatAttributeIndexFunctionName)
 	{
-		FName AttributeName = BindingInfo.FindSpecifier(NAME_Attribute)->Value;
+		FName AttributeName = BindingInfo.FindSpecifier(UNiagaraDataInterfaceRWBase::NAME_Attribute)->Value;
 		check(BindingInfo.GetNumInputs() == 1 && BindingInfo.GetNumOutputs() == 1);
 		NDI_FUNC_BINDER(UNiagaraDataInterfaceGrid2DCollection, VMGetAttributeIndex)::Bind(this, OutFunc, AttributeName, 1);
 	}
@@ -1012,59 +1017,35 @@ void UNiagaraDataInterfaceGrid2DCollection::GetParameterDefinitionHLSL(const FNi
 	OutHLSL += FString::Format(FormatDeclarations, ArgsDeclarations);
 }
 
-void UNiagaraDataInterfaceGrid2DCollection::WriteSetHLSL(const FNiagaraDataInterfaceGPUParamInfo& ParamInfo, const FNiagaraDataInterfaceGeneratedFunction& FunctionInfo, int FunctionInstanceIndex, int32 InNumChannels, FString& OutHLSL)
+void UNiagaraDataInterfaceGrid2DCollection::WriteSetHLSL(const FNiagaraDataInterfaceRWAttributeHelper& AttributeHelper, const FNiagaraDataInterfaceGPUParamInfo& ParamInfo, const FNiagaraDataInterfaceGeneratedFunction& FunctionInfo, int FunctionInstanceIndex, int32 InNumChannels, FString& OutHLSL)
 {
-
 	FString FormatBounds = TEXT(R"(
 			void {FunctionName}(int In_IndexX, int In_IndexY, float{NumChannelsVariableSuffix} In_Value)
 			{
-				int In_AttributeIndex = {AttributeIndicesName}[{AttributeIndexGroup}]{AttributeIndexComponent};
+				int In_AttributeIndex = {AttributeIndex};
 
-			    for (int i = 0; i < {NumChannels}; i++)
-				{
-					float Val;
 				)");
 	if (InNumChannels == 1)
 	{
-		FormatBounds += TEXT("					Val = In_Value;\n");
+		FormatBounds += FString("	{OutputGrid}[int3(In_IndexX, In_IndexY, In_AttributeIndex)] = In_Value;\n");
 	}
 	else if (InNumChannels > 1)
 	{
-		FormatBounds += TEXT(R"(
-					switch(i)
+		for (int i = 0; i < InNumChannels; ++i)
 					{
-						case 0:
-							Val = In_Value.x;
-							break;
-						case 1:
-							Val = In_Value.y;
-							break; )");
+			FormatBounds += FString("		{OutputGrid}[int3(In_IndexX, In_IndexY, In_AttributeIndex + ").Append(FString::FromInt(i)).Append(")] = In_Value.")
+				.Append(FNiagaraDataInterfaceRWAttributeHelper::Channels[i]).Append(";\n");
+		}
 	}
 
-	if (InNumChannels > 2)
-	{
-		FormatBounds += TEXT(R"(
-						case 2:
-							Val = In_Value.z;
-							break; )");
-	}
-	if (InNumChannels > 3)
-	{
-		FormatBounds += TEXT(R"(
-						case 3:
-							Val = In_Value.w;
-							break; )");
-	}
-	if (InNumChannels > 1)
-	{
 		FormatBounds += TEXT(R"(
-					})");
-	}
-	FormatBounds += TEXT(R"(
-					{OutputGrid}[int3(In_IndexX, In_IndexY, In_AttributeIndex + i)] = Val;
-				}
 			}
 		)");
+
+	FString AttributeIndexPlaceholder = "{AttributeIndex}";
+	FString AttributeQueryString = GenerateAttributeIndexHLSL(AttributeHelper, FunctionInfo);
+	FormatBounds.ReplaceInline(*AttributeIndexPlaceholder, *AttributeQueryString);
+
 	TMap<FString, FStringFormatArg> ArgsBounds = {
 		{TEXT("FunctionName"),				FunctionInfo.InstanceName},
 		{TEXT("OutputGrid"),				ParamInfo.DataInterfaceHLSLSymbol + OutputGridName},
@@ -1080,58 +1061,34 @@ void UNiagaraDataInterfaceGrid2DCollection::WriteSetHLSL(const FNiagaraDataInter
 	OutHLSL += FString::Format(*FormatBounds, ArgsBounds);
 }
 
-void UNiagaraDataInterfaceGrid2DCollection::WriteGetHLSL(const FNiagaraDataInterfaceGPUParamInfo& ParamInfo, const FNiagaraDataInterfaceGeneratedFunction& FunctionInfo, int FunctionInstanceIndex, int32 InNumChannels, FString& OutHLSL)
+void UNiagaraDataInterfaceGrid2DCollection::WriteGetHLSL(const FNiagaraDataInterfaceRWAttributeHelper& AttributeHelper, const FNiagaraDataInterfaceGPUParamInfo& ParamInfo, const FNiagaraDataInterfaceGeneratedFunction& FunctionInfo, int FunctionInstanceIndex, int32 InNumChannels, FString& OutHLSL)
 {
-
 	FString FormatBounds = TEXT(R"(
 			void {FunctionName}(int In_IndexX, int In_IndexY, out float{NumChannelsVariableSuffix} Out_Val)
 			{
-				int In_AttributeIndex = {AttributeIndicesName}[{AttributeIndexGroup}]{AttributeIndexComponent};
+				int In_AttributeIndex = {AttributeIndex};
 
-			    for (int i = 0; i < {NumChannels}; i++)
-				{
-					float Val = {Grid}.Load(int4(In_IndexX, In_IndexY, In_AttributeIndex + i, 0));
 					)");
 	if (InNumChannels == 1)
 	{
-		FormatBounds += TEXT("					Out_Val = Val;\n");
+		FormatBounds += FString("	Out_Val = {Grid}.Load(int4(In_IndexX, In_IndexY, In_AttributeIndex, 0));\n");
 	}
 	else if (InNumChannels > 1)
 	{
-		FormatBounds += TEXT(R"(
-					switch(i)
+		for (int i = 0; i < InNumChannels; ++i)
 					{
-						case 0:
-							Out_Val.x = Val;
-							break;
-						case 1:
-							Out_Val.y = Val;
-							break; )");
+			FormatBounds += FString("			Out_Val.").Append(FNiagaraDataInterfaceRWAttributeHelper::Channels[i]).
+				Append(" = {Grid}.Load(int4(In_IndexX, In_IndexY, In_AttributeIndex + ").Append(FString::FromInt(i)).Append(", 0)); \n");
 	}
-
-	if (InNumChannels > 2)
-	{
-		FormatBounds += TEXT(R"(
-						case 2:
-							Out_Val.z = Val;
-							break; )");
 	}
-	if (InNumChannels > 3)
-	{
-		FormatBounds += TEXT(R"(
-						case 3:
-							Out_Val.w = Val;
-							break; )");
-	}
-	if (InNumChannels > 1)
-	{
 		FormatBounds += TEXT(R"(
-					})");
 	}
-	FormatBounds += TEXT(R"(
-				}
-			}
 		)");
+
+	FString AttributeIndexPlaceholder = "{AttributeIndex}";
+	FString AttributeQueryString = GenerateAttributeIndexHLSL(AttributeHelper, FunctionInfo);
+	FormatBounds.ReplaceInline(*AttributeIndexPlaceholder, *AttributeQueryString);
+
 	TMap<FString, FStringFormatArg> ArgsBounds = {
 		{TEXT("FunctionName"),				FunctionInfo.InstanceName},
 		{TEXT("OutputGrid"),				ParamInfo.DataInterfaceHLSLSymbol + OutputGridName},
@@ -1148,58 +1105,36 @@ void UNiagaraDataInterfaceGrid2DCollection::WriteGetHLSL(const FNiagaraDataInter
 	OutHLSL += FString::Format(*FormatBounds, ArgsBounds);
 }
 
-void UNiagaraDataInterfaceGrid2DCollection::WriteSampleHLSL(const FNiagaraDataInterfaceGPUParamInfo& ParamInfo, const FNiagaraDataInterfaceGeneratedFunction& FunctionInfo, int FunctionInstanceIndex, int32 InNumChannels, FString& OutHLSL)
+void UNiagaraDataInterfaceGrid2DCollection::WriteSampleHLSL(const FNiagaraDataInterfaceRWAttributeHelper& AttributeHelper, const FNiagaraDataInterfaceGPUParamInfo& ParamInfo, const FNiagaraDataInterfaceGeneratedFunction& FunctionInfo, int FunctionInstanceIndex, int32 InNumChannels, FString& OutHLSL)
 {
+
 	FString FormatBounds = TEXT(R"(
 			void {FunctionName}(float2 In_Unit, out float{NumChannelsVariableSuffix} Out_Val)
 			{
-				int In_AttributeIndex = {AttributeIndicesName}[{AttributeIndexGroup}]{AttributeIndexComponent};
+				int In_AttributeIndex = {AttributeIndex};
 
-			    for (int i = 0; i < {NumChannels}; i++)
-				{
-					float Val = {Grid}.SampleLevel({SamplerName}, float3(In_Unit, In_AttributeIndex + i), 0);
 					)");
 	if (InNumChannels == 1)
 	{
-		FormatBounds += TEXT("					Out_Val = Val;\n");
+		FormatBounds += FString("	Out_Val = {Grid}.SampleLevel({SamplerName}, float3(In_Unit, In_AttributeIndex), 0);\n");
 	}
 	else if (InNumChannels > 1)
 	{
-		FormatBounds += TEXT(R"(
-					switch(i)
+		for (int i = 0; i < InNumChannels; ++i)
 					{
-						case 0:
-							Out_Val.x = Val;
-							break;
-						case 1:
-							Out_Val.y = Val;
-							break; )");
+			FormatBounds += FString("			Out_Val.").Append(FNiagaraDataInterfaceRWAttributeHelper::Channels[i]).
+				Append(" = {Grid}.SampleLevel({SamplerName}, float3(In_Unit, In_AttributeIndex + ").Append(FString::FromInt(i)).Append("), 0); \n");
+		}
 }
 
-	if (InNumChannels > 2)
-	{
-		FormatBounds += TEXT(R"(
-						case 2:
-							Out_Val.z = Val;
-							break; )");
-	}
-	if (InNumChannels > 3)
-	{
-		FormatBounds += TEXT(R"(
-						case 3:
-							Out_Val.w = Val;
-							break; )");
-	}
-	if (InNumChannels > 1)
-	{
-		FormatBounds += TEXT(R"(
-					})");
-	}
 	FormatBounds += TEXT(R"(
 				}
-			}
 		)");
 
+	FString AttributeIndexPlaceholder = "{AttributeIndex}";
+	FString AttributeQueryString = GenerateAttributeIndexHLSL(AttributeHelper, FunctionInfo);
+	FormatBounds.ReplaceInline(*AttributeIndexPlaceholder, *AttributeQueryString);
+
 	TMap<FString, FStringFormatArg> ArgsBounds = {
 		{TEXT("FunctionName"),				FunctionInfo.InstanceName},
 		{TEXT("Grid"),						ParamInfo.DataInterfaceHLSLSymbol + GridName},
@@ -1215,16 +1150,18 @@ void UNiagaraDataInterfaceGrid2DCollection::WriteSampleHLSL(const FNiagaraDataIn
 	OutHLSL += FString::Format(*FormatBounds, ArgsBounds);
 }
 
-void UNiagaraDataInterfaceGrid2DCollection::WriteAttributeGetIndexHLSL(const FNiagaraDataInterfaceGPUParamInfo& ParamInfo, const FNiagaraDataInterfaceGeneratedFunction& FunctionInfo, int FunctionInstanceIndex, int32 InNumChannels, FString& OutHLSL)
+void UNiagaraDataInterfaceGrid2DCollection::WriteAttributeGetIndexHLSL(const FNiagaraDataInterfaceRWAttributeHelper& AttributeHelper, const FNiagaraDataInterfaceGPUParamInfo& ParamInfo, const FNiagaraDataInterfaceGeneratedFunction& FunctionInfo, int FunctionInstanceIndex, int32 InNumChannels, FString& OutHLSL)
 {
 	FString FormatBounds = TEXT(R"(
 			void {FunctionName}(out int Out_Val)
 			{
-				int In_AttributeIndex = {AttributeIndicesName}[{AttributeIndexGroup}]{AttributeIndexComponent};
-				Out_Val = In_AttributeIndex;
+				Out_Val = {AttributeIndex};				
 			}
 	)");
 
+	FString AttributeIndexPlaceholder = "{AttributeIndex}";
+	FString AttributeQueryString = GenerateAttributeIndexHLSL(AttributeHelper, FunctionInfo);
+	FormatBounds.ReplaceInline(*AttributeIndexPlaceholder, *AttributeQueryString);
 
 	TMap<FString, FStringFormatArg> ArgsBounds = {
 		{TEXT("FunctionName"), FunctionInfo.InstanceName},
@@ -1235,6 +1172,33 @@ void UNiagaraDataInterfaceGrid2DCollection::WriteAttributeGetIndexHLSL(const FNi
 	OutHLSL += FString::Format(*FormatBounds, ArgsBounds);
 }
 
+FString UNiagaraDataInterfaceGrid2DCollection::GenerateAttributeIndexHLSL(const FNiagaraDataInterfaceRWAttributeHelper& AttributeHelper, const FNiagaraDataInterfaceGeneratedFunction& FunctionInfo)
+{
+	FString OutString = "";
+	if (bUseIndirection)
+	{
+		OutString = "{AttributeIndicesName}[{AttributeIndexGroup}]{AttributeIndexComponent}";
+	}
+	else
+	{
+		const FName* AttributeName = FunctionInfo.FindSpecifierValue(UNiagaraDataInterfaceRWBase::NAME_Attribute);
+		if (AttributeName == nullptr)
+		{
+			return OutString;
+		}
+
+		const FNiagaraDataInterfaceRWAttributeHelper::FAttributeInfo* AttributeInfo = AttributeHelper.FindAttributeInfo(*AttributeName);
+		if (AttributeInfo == nullptr)
+		{
+			return OutString;
+		}
+
+		OutString = AttributeHelper.GetAttributeIndex(AttributeInfo);
+	}
+
+	return OutString;
+}
+
 const TCHAR* UNiagaraDataInterfaceGrid2DCollection::TypeDefinitionToHLSLTypeString(const FNiagaraTypeDefinition& InDef) const
 {
 	if (InDef == FNiagaraTypeDefinition::GetFloatDef())
@@ -1274,6 +1238,9 @@ FName UNiagaraDataInterfaceGrid2DCollection::TypeDefinitionToSetFunctionName(con
 
 bool UNiagaraDataInterfaceGrid2DCollection::GetFunctionHLSL(const FNiagaraDataInterfaceGPUParamInfo& ParamInfo, const FNiagaraDataInterfaceGeneratedFunction& FunctionInfo, int FunctionInstanceIndex, FString& OutHLSL)
 {
+	FNiagaraDataInterfaceRWAttributeHelper AttributeHelper(ParamInfo);
+	AttributeHelper.Init<UNiagaraDataInterfaceGrid2DCollection>();
+
 	bool ParentRet = Super::GetFunctionHLSL(ParamInfo, FunctionInfo, FunctionInstanceIndex, OutHLSL);
 	if (ParentRet)
 	{
@@ -1357,82 +1324,82 @@ bool UNiagaraDataInterfaceGrid2DCollection::GetFunctionHLSL(const FNiagaraDataIn
 	}
 	else if (FunctionInfo.DefinitionName == SetVector4ValueFunctionName)
 	{
-		WriteSetHLSL(ParamInfo, FunctionInfo, FunctionInstanceIndex, 4, OutHLSL);
+		WriteSetHLSL(AttributeHelper, ParamInfo, FunctionInfo, FunctionInstanceIndex, 4, OutHLSL);
 		return true;
 	}
 	else if (FunctionInfo.DefinitionName == GetVector4ValueFunctionName || FunctionInfo.DefinitionName == GetPreviousVector4ValueFunctionName)
 	{
-		WriteGetHLSL(ParamInfo, FunctionInfo, FunctionInstanceIndex, 4, OutHLSL);
+		WriteGetHLSL(AttributeHelper, ParamInfo, FunctionInfo, FunctionInstanceIndex, 4, OutHLSL);
 		return true;
 	}
 	else if (FunctionInfo.DefinitionName == SampleGridVector4FunctionName || FunctionInfo.DefinitionName == SamplePreviousGridVector4FunctionName)
 	{
-		WriteSampleHLSL(ParamInfo, FunctionInfo, FunctionInstanceIndex, 4, OutHLSL);
+		WriteSampleHLSL(AttributeHelper, ParamInfo, FunctionInfo, FunctionInstanceIndex, 4, OutHLSL);
 		return true;
 	}
 	else if (FunctionInfo.DefinitionName == SetVector3ValueFunctionName || FunctionInfo.DefinitionName == SetVectorValueFunctionName)
 	{
-		WriteSetHLSL(ParamInfo, FunctionInfo, FunctionInstanceIndex, 3, OutHLSL);
+		WriteSetHLSL(AttributeHelper, ParamInfo, FunctionInfo, FunctionInstanceIndex, 3, OutHLSL);
 		return true;
 	}
 	else if (FunctionInfo.DefinitionName == GetVector3ValueFunctionName || FunctionInfo.DefinitionName == GetPreviousVectorValueFunctionName)
 	{
-		WriteGetHLSL(ParamInfo, FunctionInfo, FunctionInstanceIndex, 3, OutHLSL);
+		WriteGetHLSL(AttributeHelper, ParamInfo, FunctionInfo, FunctionInstanceIndex, 3, OutHLSL);
 		return true;
 	}
 	else if (FunctionInfo.DefinitionName == SampleGridVector3FunctionName || FunctionInfo.DefinitionName == SamplePreviousGridVectorFunctionName)
 	{
-		WriteSampleHLSL(ParamInfo, FunctionInfo, FunctionInstanceIndex, 3, OutHLSL);
+		WriteSampleHLSL(AttributeHelper, ParamInfo, FunctionInfo, FunctionInstanceIndex, 3, OutHLSL);
 		return true;
 	}
 	else if (FunctionInfo.DefinitionName == SetVector2ValueFunctionName || FunctionInfo.DefinitionName == SetVector2DValueFunctionName)
 	{
-		WriteSetHLSL(ParamInfo, FunctionInfo, FunctionInstanceIndex, 2, OutHLSL);
+		WriteSetHLSL(AttributeHelper, ParamInfo, FunctionInfo, FunctionInstanceIndex, 2, OutHLSL);
 		return true;
 	}
 	else if (FunctionInfo.DefinitionName == GetVector2ValueFunctionName || FunctionInfo.DefinitionName == GetPreviousVector2DValueFunctionName)
 	{
-		WriteGetHLSL(ParamInfo, FunctionInfo, FunctionInstanceIndex, 2, OutHLSL);
+		WriteGetHLSL(AttributeHelper, ParamInfo, FunctionInfo, FunctionInstanceIndex, 2, OutHLSL);
 		return true;
 	}
 	else if (FunctionInfo.DefinitionName == SampleGridVector2FunctionName || FunctionInfo.DefinitionName == SamplePreviousGridVector2DFunctionName)
 	{
-		WriteSampleHLSL(ParamInfo, FunctionInfo, FunctionInstanceIndex, 2, OutHLSL);
+		WriteSampleHLSL(AttributeHelper, ParamInfo, FunctionInfo, FunctionInstanceIndex, 2, OutHLSL);
 		return true;
 	}
 	else if (FunctionInfo.DefinitionName == SetFloatValueFunctionName)
 	{
-		WriteSetHLSL(ParamInfo, FunctionInfo, FunctionInstanceIndex, 1, OutHLSL);
+		WriteSetHLSL(AttributeHelper, ParamInfo, FunctionInfo, FunctionInstanceIndex, 1, OutHLSL);
 		return true;
 	}
 	else if (FunctionInfo.DefinitionName == GetFloatValueFunctionName || FunctionInfo.DefinitionName == GetPreviousFloatValueFunctionName)
 	{
-		WriteGetHLSL(ParamInfo, FunctionInfo, FunctionInstanceIndex, 1, OutHLSL);
+		WriteGetHLSL(AttributeHelper, ParamInfo, FunctionInfo, FunctionInstanceIndex, 1, OutHLSL);
 		return true;
 	}
 	else if (FunctionInfo.DefinitionName == SampleGridFloatFunctionName || FunctionInfo.DefinitionName == SamplePreviousGridFloatFunctionName)
 	{
-		WriteSampleHLSL(ParamInfo, FunctionInfo, FunctionInstanceIndex, 1, OutHLSL);
+		WriteSampleHLSL(AttributeHelper, ParamInfo, FunctionInfo, FunctionInstanceIndex, 1, OutHLSL);
 		return true;
 	}
 	else if (FunctionInfo.DefinitionName == GetVector4AttributeIndexFunctionName)
 	{
-		WriteAttributeGetIndexHLSL(ParamInfo, FunctionInfo, FunctionInstanceIndex, 4, OutHLSL);
+		WriteAttributeGetIndexHLSL(AttributeHelper, ParamInfo, FunctionInfo, FunctionInstanceIndex, 4, OutHLSL);
 		return true;
 	}
 	else if (FunctionInfo.DefinitionName == GetVectorAttributeIndexFunctionName)
 	{
-		WriteAttributeGetIndexHLSL(ParamInfo, FunctionInfo, FunctionInstanceIndex, 3, OutHLSL);
+		WriteAttributeGetIndexHLSL(AttributeHelper, ParamInfo, FunctionInfo, FunctionInstanceIndex, 3, OutHLSL);
 		return true;
 	}
 	else if (FunctionInfo.DefinitionName == GetVector2DAttributeIndexFunctionName)
 	{
-		WriteAttributeGetIndexHLSL(ParamInfo, FunctionInfo, FunctionInstanceIndex, 2, OutHLSL);
+		WriteAttributeGetIndexHLSL(AttributeHelper, ParamInfo, FunctionInfo, FunctionInstanceIndex, 2, OutHLSL);
 		return true;
 	}
 	else if (FunctionInfo.DefinitionName == GetFloatAttributeIndexFunctionName)
 	{
-		WriteAttributeGetIndexHLSL(ParamInfo, FunctionInfo, FunctionInstanceIndex, 1, OutHLSL);
+		WriteAttributeGetIndexHLSL(AttributeHelper, ParamInfo, FunctionInfo, FunctionInstanceIndex, 1, OutHLSL);
 		return true;
 	}
 	else if (FunctionInfo.DefinitionName == SampleGridFunctionName || FunctionInfo.DefinitionName == SamplePreviousGridAtIndexFunctionName)
@@ -1552,8 +1519,8 @@ FNiagaraDataInterfaceParametersCS* UNiagaraDataInterfaceGrid2DCollection::Create
 	for (int32 FuncIdx=0; FuncIdx < NumFuncs; ++FuncIdx)
 	{
 		const FNiagaraDataInterfaceGeneratedFunction& Func = ParameterInfo.GeneratedFunctions[FuncIdx];
-		static const FName NAME_Attribute("Attribute");
-		const FName* AttributeName = Func.FindSpecifierValue(NAME_Attribute);
+		
+		const FName* AttributeName = Func.FindSpecifierValue(UNiagaraDataInterfaceRWBase::NAME_Attribute);
 		if (AttributeName != nullptr)
 		{
 			int32 ComponentCount = UNiagaraDataInterfaceGrid2DCollection::GetComponentCountFromFuncName(Func.DefinitionName);
@@ -1662,6 +1629,19 @@ bool UNiagaraDataInterfaceGrid2DCollection::GenerateIterationSourceNamespaceWrit
 		OutHLSL += FString::Printf(TEXT("\tint X, Y;\n\t%s.ExecutionIndexToGridIndex(X, Y);\n"), *DIVarName);
 	}
 
+	// only copy previous to next if we haven't written to all the attributes
+	if (InAttributes.Num() > 0 && InAttributes.Num() < InAllAttributes.Num() && !bSpawnOnly && !bPartialWrites)
+	{
+		static const TCHAR* FormatBounds = TEXT(R"(
+			// We need to copy from previous to current first thing, because other functions afterwards may just set values on the local grid.			
+			{Grid}.CopyPreviousToCurrentForCell(X,Y);
+		)");
+		TMap<FString, FStringFormatArg> ArgsBounds = {
+			{TEXT("Grid"), DIVarName},
+		};
+		OutHLSL += FString::Format(FormatBounds, ArgsBounds);
+	}
+
 	TArray<FString> RootArray;
 	IterationSourceVar.GetName().ToString().ParseIntoArray(RootArray, TEXT("."));
 
@@ -1717,20 +1697,6 @@ bool UNiagaraDataInterfaceGrid2DCollection::GenerateSetupHLSL(FNiagaraDataInterf
 		}
 	}
 
-	if (!bSpawnOnly && !bPartialWrites)
-	{
-		static const TCHAR* FormatBounds = TEXT(R"(
-			// We need to copy from previous to current first thing, because other functions afterwards may just set values on the local grid.
-			int X, Y;
-			{Grid}.ExecutionIndexToGridIndex(X, Y);
-			{Grid}.CopyPreviousToCurrentForCell(X,Y);
-		)");
-		TMap<FString, FStringFormatArg> ArgsBounds = {
-			{TEXT("Grid"), DIVarName},
-		};
-		OutHLSL += FString::Format(FormatBounds, ArgsBounds);
-	}
-
 	return true;
 }
 
@@ -1772,13 +1738,21 @@ bool UNiagaraDataInterfaceGrid2DCollection::InitPerInstanceData(void* PerInstanc
 	InstanceData->NumCells.Y = NumCellsY;
 
 	/* Go through all references to this data interface and build up the attribute list from the function metadata of those referenced.*/
+
 	int32 NumAttribChannelsFound = 0;
-	FindAttributes(InstanceData->Vars, InstanceData->Offsets, NumAttribChannelsFound);
+	int32 NumNamedAttribChannelsFound = 0;
+	FindAttributes(InstanceData->Vars, InstanceData->Offsets, NumNamedAttribChannelsFound);
 
 	// Ensure we never allocate 0 attributes as that would fail to create the texture
-	NumAttribChannelsFound = FMath::Max(NumAttributes + NumAttribChannelsFound, 1);
+	NumAttribChannelsFound = FMath::Max(NumAttributes + NumNamedAttribChannelsFound, 1);
 	InstanceData->NumAttributes = NumAttribChannelsFound;
 
+	// give an error if there are any unnamed attributes
+	if (NumAttributes > 0 && NumNamedAttribChannelsFound > 0)
+	{
+		UE_LOG(LogNiagara, Warning, TEXT("%s: Unnamed attributes should not be used with named.  This is a deprecated workflow"), *Proxy->SourceDIName.ToString());
+	}
+
 	InstanceData->WorldBBoxSize = WorldBBoxSize;
 
 	ENiagaraGpuBufferFormat BufferFormat = bOverrideFormat ? OverrideBufferFormat : GetDefault<UNiagaraSettings>()->DefaultGridFormat;
@@ -1870,6 +1844,7 @@ bool UNiagaraDataInterfaceGrid2DCollection::InitPerInstanceData(void* PerInstanc
 		check(!RT_Proxy->SystemInstancesToProxyData_RT.Contains(InstanceID));
 		FGrid2DCollectionRWInstanceData_RenderThread* TargetData = &RT_Proxy->SystemInstancesToProxyData_RT.Add(InstanceID);
 
+		TargetData->SourceDIName = RT_Proxy->SourceDIName;
 		TargetData->ClearBeforeNonIterationStage = RT_InstanceData.ClearBeforeNonIterationStage;
 
 		TargetData->NumCells = RT_InstanceData.NumCells;
@@ -1973,11 +1948,20 @@ bool UNiagaraDataInterfaceGrid2DCollection::GetExposedVariableValue(const FNiaga
 	return false;
 }
 
-void UNiagaraDataInterfaceGrid2DCollection::CollectAttributesForScript(UNiagaraScript* Script, FName VariableName, TArray<FNiagaraVariableBase>& OutVariables, TArray<uint32>& OutVariableOffsets, int32& TotalAttributes, TArray<FText>* OutWarnings)
+void UNiagaraDataInterfaceGrid2DCollection::CollectAttributesForScript(UNiagaraScript* Script, FName VariableName, TArray<FNiagaraVariableBase>& OutVariables, TArray<uint32>& OutVariableOffsets, int32& TotalAttributes, TArray<FText>* OutWarnings, bool UseReader)
 {
 	if (const FNiagaraScriptExecutionParameterStore* ParameterStore = Script->GetExecutionReadyParameterStore(ENiagaraSimTarget::GPUComputeSim))
 	{
-		const FNiagaraVariableBase DataInterfaceVariable(FNiagaraTypeDefinition(UNiagaraDataInterfaceGrid2DCollection::StaticClass()), VariableName);
+		FNiagaraVariableBase DataInterfaceVariable(FNiagaraTypeDefinition(UNiagaraDataInterfaceGrid2DCollection::StaticClass()), VariableName);
+
+		if (UseReader)
+		{
+			DataInterfaceVariable = FNiagaraVariableBase(FNiagaraTypeDefinition(UNiagaraDataInterfaceGrid2DCollectionReader::StaticClass()), VariableName);
+		}
+		else
+		{
+			DataInterfaceVariable = FNiagaraVariableBase(FNiagaraTypeDefinition(UNiagaraDataInterfaceGrid2DCollection::StaticClass()), VariableName);
+		}
 
 		const int32 * IndexOfDataInterface = ParameterStore->FindParameterOffset(DataInterfaceVariable);
 		if (IndexOfDataInterface != nullptr)
@@ -1985,9 +1969,7 @@ void UNiagaraDataInterfaceGrid2DCollection::CollectAttributesForScript(UNiagaraS
 			TConstArrayView<FNiagaraDataInterfaceGPUParamInfo> ParamInfoArray = Script->GetDataInterfaceGPUParamInfos();
 			for (const FNiagaraDataInterfaceGeneratedFunction& Func : ParamInfoArray[*IndexOfDataInterface].GeneratedFunctions)
 			{
-				static const FName NAME_Attribute("Attribute");
-
-				if (const FName* AttributeName = Func.FindSpecifierValue(NAME_Attribute))
+				if (const FName* AttributeName = Func.FindSpecifierValue(UNiagaraDataInterfaceRWBase::NAME_Attribute))
 				{
 					FNiagaraVariableBase NewVar(UNiagaraDataInterfaceGrid2DCollection::GetValueTypeFromFuncName(Func.DefinitionName), *AttributeName);
 					if (UNiagaraDataInterfaceGrid2DCollection::CanCreateVarFromFuncName(Func.DefinitionName))
@@ -2040,7 +2022,7 @@ void UNiagaraDataInterfaceGrid2DCollection::FindAttributesByName(FName VariableN
 	OutNumAttribChannelsFound = TotalAttributes - NumAttributes;
 }
 
-void UNiagaraDataInterfaceGrid2DCollection::FindAttributes(TArray<FNiagaraVariableBase>& OutVariables, TArray<uint32>& OutVariableOffsets, int32& OutNumAttribChannelsFound, TArray<FText>* OutWarnings) const
+void UNiagaraDataInterfaceGrid2DCollection::FindAttributes(TArray<FNiagaraVariableBase>& OutVariables, TArray<uint32>& OutVariableOffsets, int32& OutNumAttribChannelsFound, TArray<FText>* OutWarnings, bool UseReader) const
 {
 	OutNumAttribChannelsFound = 0;
 
@@ -2083,7 +2065,7 @@ void UNiagaraDataInterfaceGrid2DCollection::FindAttributes(TArray<FNiagaraVariab
 			const FName VariableName = FindDataInterfaceVariable(this);
 			if (!VariableName.IsNone() )
 			{
-				CollectAttributesForScript(EmitterData->GetGPUComputeScript(), VariableName, OutVariables, OutVariableOffsets, TotalAttributes, OutWarnings);
+				CollectAttributesForScript(EmitterData->GetGPUComputeScript(), VariableName, OutVariables, OutVariableOffsets, TotalAttributes, OutWarnings, UseReader);
 			}
 		}
 	}
@@ -2555,8 +2537,18 @@ void FGrid2DCollectionRWInstanceData_RenderThread::BeginSimulate(FRDGBuilder& Gr
 		Buffers.Emplace(DestinationData);
 
 		const FRDGTextureDesc TextureDesc = FRDGTextureDesc::Create2DArray(NumCells, PixelFormat, FClearValueBinding::Black, ETextureCreateFlags::ShaderResource | ETextureCreateFlags::UAV, NumAttributes);
-		DestinationData->Initialize(GraphBuilder, TEXT("FGrid2DBuffer"), TextureDesc);
 
+		TStringBuilder<128> StringBuilder;
+		SourceDIName.AppendString(StringBuilder);
+		StringBuilder.Append("_");
+		for (FName CurrName : Vars)
+		{
+			CurrName.AppendString(StringBuilder);
+			StringBuilder.Append("_");
+		}
+		StringBuilder.Append("Grid2DCollection");
+
+		DestinationData->Initialize(GraphBuilder, *StringBuilder, TextureDesc);
 		// This destination buffer will sometimes have old data in it.  Force it to clear.
 		AddClearUAVPass(GraphBuilder, DestinationData->GetOrCreateUAV(GraphBuilder), FVector4f(ForceInitToZero));
 	}
@@ -2592,7 +2584,7 @@ void FNiagaraDataInterfaceProxyGrid2DCollectionProxy::PreStage(const FNDIGpuComp
 
 		FRDGBuilder& GraphBuilder = Context.GetGraphBuilder();
 		ProxyData->BeginSimulate(GraphBuilder, Context.IsInputStage());
-		
+
 		// If there is an output DI and the stage dispatch count is different from the grid resolution, then we want to consider it for a pre-stage UAV clear.
 		// This is to enable the most common use case of scattered writes to temporary buffers		
 		const FNiagaraSimStageData& SimStageData = Context.GetSimStageData();
diff --git a/Engine/Plugins/FX/Niagara/Source/Niagara/Private/NiagaraDataInterfaceGrid2DCollectionReader.cpp b/Engine/Plugins/FX/Niagara/Source/Niagara/Private/NiagaraDataInterfaceGrid2DCollectionReader.cpp
--- a/Engine/Plugins/FX/Niagara/Source/Niagara/Private/NiagaraDataInterfaceGrid2DCollectionReader.cpp
+++ b/Engine/Plugins/FX/Niagara/Source/Niagara/Private/NiagaraDataInterfaceGrid2DCollectionReader.cpp
@@ -23,6 +23,8 @@ UNiagaraDataInterfaceGrid2DCollectionReader::UNiagaraDataInterfaceGrid2DCollecti
 
 	FNiagaraTypeDefinition Def(UTextureRenderTarget::StaticClass());
 	RenderTargetUserParameter.Parameter.SetType(Def);
+
+	bUseIndirection = true;
 }
 
 bool UNiagaraDataInterfaceGrid2DCollectionReader::Equals(const UNiagaraDataInterface* Other) const
@@ -107,18 +109,29 @@ bool UNiagaraDataInterfaceGrid2DCollectionReader::InitPerInstanceData(void* PerI
 		if (GridInterface && ProxyToUse)
 		{
 
+			int32 NumAttribChannelsFound = 0;
+			int32 NumNamedAttribChannelsFound = 0;
+			TArray<FNiagaraVariableBase> Vars;
+			TArray<uint32> Offsets;
+			FindAttributes(Vars, Offsets, NumNamedAttribChannelsFound, nullptr, true);
+
 			FGrid2DCollectionRWInstanceData_GameThread* Grid2DInstanceData = GridInterface->GetSystemInstancesToProxyData_GT().FindRef(SystemInstance->GetId());
 			InstanceData->OtherDI = GridInterface;
 			InstanceData->OtherInstanceData = Grid2DInstanceData;
 
+			// error if the reader is trying to access a grid that has both named and unnamed attributes
+			if (InstanceData->OtherInstanceData != nullptr && GridInterface->NumAttributes > 0 && GridInterface->NumAttributes != InstanceData->OtherInstanceData->NumAttributes)
+			{
+				UE_LOG(LogNiagara, Error, TEXT("Grid2DCollectionReaders are not compatible with Emitter %s and DataInterface %s.  Do not use unnamed and named grids together."), *this->EmitterName, *this->DIName);
+				return false;
+			}
+
 			// Push Updates to Proxy.
 			FNiagaraDataInterfaceProxyGrid2DCollectionProxy* RT_Proxy = GetProxyAs<FNiagaraDataInterfaceProxyGrid2DCollectionProxy>();
 			ENQUEUE_RENDER_COMMAND(FUpdateData)(
 				[RT_Proxy, InstanceID = SystemInstance->GetId(), RT_ProxyToUse = ProxyToUse](FRHICommandListImmediate& RHICmdList)
 			{
-				check(!RT_Proxy->SystemInstancesToProxyData_RT.Contains(InstanceID));
-				//check(RT_ProxyToUse->SystemInstancesToProxyData_RT.Contains(InstanceID));
-				//FGrid2DCollectionRWInstanceData_RenderThread* DataToRead = RT_ProxyToUse->SystemInstancesToProxyData_RT.Find(InstanceID);
+				check(!RT_Proxy->SystemInstancesToProxyData_RT.Contains(InstanceID));				
 
 				FGrid2DCollectionRWInstanceData_RenderThread* TargetData = &RT_Proxy->SystemInstancesToProxyData_RT.Add(InstanceID);
 				TargetData->OtherProxy = RT_ProxyToUse;
diff --git a/Engine/Plugins/FX/Niagara/Source/Niagara/Private/NiagaraDataInterfaceGrid3DCollection.cpp b/Engine/Plugins/FX/Niagara/Source/Niagara/Private/NiagaraDataInterfaceGrid3DCollection.cpp
--- a/Engine/Plugins/FX/Niagara/Source/Niagara/Private/NiagaraDataInterfaceGrid3DCollection.cpp
+++ b/Engine/Plugins/FX/Niagara/Source/Niagara/Private/NiagaraDataInterfaceGrid3DCollection.cpp
@@ -67,7 +67,6 @@ const FString UNiagaraDataInterfaceGrid3DCollection::OneOverNumTilesName(TEXT("_
 const FString UNiagaraDataInterfaceGrid3DCollection::UnitClampMinName(TEXT("_UnitClampMin"));
 const FString UNiagaraDataInterfaceGrid3DCollection::UnitClampMaxName(TEXT("_UnitClampMax"));
 
-const FName UNiagaraDataInterfaceGrid3DCollection::NAME_Attribute("Attribute");
 const FString UNiagaraDataInterfaceGrid3DCollection::GridName(TEXT("_Grid"));
 const FString UNiagaraDataInterfaceGrid3DCollection::OutputGridName(TEXT("_OutputGrid"));
 const FString UNiagaraDataInterfaceGrid3DCollection::SamplerName(TEXT("_GridSampler"));
@@ -354,7 +353,7 @@ bool UNiagaraDataInterfaceGrid3DCollection::AppendCompileHash(FNiagaraCompileHas
 
 
 	InVisitor->UpdatePOD(TEXT("UNiagaraDataInterfaceGrid3DCollectionVersion"), (int32)FNiagaraGridCollection3DDIFunctionVersion::LatestVersion);
-	InVisitor->UpdatePOD(TEXT("UNiagaraDataInterfaceGrid3DCollectionSupportsRGBAGrid"), FGrid3DCollectionAttributeHelper::SupportsRGBAGrid() ? 1 : 0);
+	InVisitor->UpdatePOD(TEXT("UNiagaraDataInterfaceGrid3DCollectionSupportsRGBAGrid"), FGrid3DCollectionAttributeHlslWriter::SupportsRGBAGrid() ? 1 : 0);
 	//InVisitor->UpdateShaderFile(TemplateShaderFilePath);
 	InVisitor->UpdateShaderParameters<FNDIGrid3DShaderParameters>();
 
@@ -629,7 +628,7 @@ void UNiagaraDataInterfaceGrid3DCollection::GetFunctions(TArray<FNiagaraFunction
 		Sig.Inputs.Add(FNiagaraVariable(FNiagaraTypeDefinition::GetIntDef(), TEXT("IndexY")));
 		Sig.Inputs.Add(FNiagaraVariable(FNiagaraTypeDefinition::GetIntDef(), TEXT("IndexZ")));
 		Sig.Inputs.Add(FNiagaraVariable(FNiagaraTypeDefinition::GetVec4Def(), TEXT("Value")));
-		Sig.FunctionSpecifiers.Add(FName(NAME_Attribute));
+		Sig.FunctionSpecifiers.Add(FName(UNiagaraDataInterfaceRWBase::NAME_Attribute));
 		Sig.bMemberFunction = true;
 		Sig.bRequiresContext = false;
 		Sig.bExperimental = true;
@@ -652,7 +651,7 @@ void UNiagaraDataInterfaceGrid3DCollection::GetFunctions(TArray<FNiagaraFunction
 		Sig.Inputs.Add(FNiagaraVariable(FNiagaraTypeDefinition::GetIntDef(), TEXT("IndexY")));
 		Sig.Inputs.Add(FNiagaraVariable(FNiagaraTypeDefinition::GetIntDef(), TEXT("IndexZ")));
 		Sig.Outputs.Add(FNiagaraVariable(FNiagaraTypeDefinition::GetVec4Def(), TEXT("Value")));
-		Sig.FunctionSpecifiers.Add(FName(NAME_Attribute));
+		Sig.FunctionSpecifiers.Add(FName(UNiagaraDataInterfaceRWBase::NAME_Attribute));
 		Sig.bMemberFunction = true;
 		Sig.bRequiresContext = false;
 		Sig.bExperimental = true;
@@ -673,7 +672,7 @@ void UNiagaraDataInterfaceGrid3DCollection::GetFunctions(TArray<FNiagaraFunction
 		Sig.Inputs.Add(FNiagaraVariable(FNiagaraTypeDefinition(GetClass()), TEXT("Grid")));
 		Sig.Inputs.Add(FNiagaraVariable(FNiagaraTypeDefinition::GetVec3Def(), TEXT("Unit")));
 		Sig.Outputs.Add(FNiagaraVariable(FNiagaraTypeDefinition::GetVec4Def(), TEXT("Value")));
-		Sig.FunctionSpecifiers.Add(FName(NAME_Attribute));
+		Sig.FunctionSpecifiers.Add(FName(UNiagaraDataInterfaceRWBase::NAME_Attribute));
 		Sig.bMemberFunction = true;
 		Sig.bRequiresContext = false;
 		Sig.bExperimental = true;
@@ -694,7 +693,7 @@ void UNiagaraDataInterfaceGrid3DCollection::GetFunctions(TArray<FNiagaraFunction
 		Sig.Inputs.Add(FNiagaraVariable(FNiagaraTypeDefinition(GetClass()), TEXT("Grid")));
 		Sig.Inputs.Add(FNiagaraVariable(FNiagaraTypeDefinition::GetVec3Def(), TEXT("Unit")));
 		Sig.Outputs.Add(FNiagaraVariable(FNiagaraTypeDefinition::GetVec4Def(), TEXT("Value")));
-		Sig.FunctionSpecifiers.Add(FName(NAME_Attribute));
+		Sig.FunctionSpecifiers.Add(FName(UNiagaraDataInterfaceRWBase::NAME_Attribute));
 		Sig.bMemberFunction = true;
 		Sig.bRequiresContext = false;
 		Sig.bExperimental = true;
@@ -717,7 +716,7 @@ void UNiagaraDataInterfaceGrid3DCollection::GetFunctions(TArray<FNiagaraFunction
 		Sig.Inputs.Add(FNiagaraVariable(FNiagaraTypeDefinition::GetIntDef(), TEXT("IndexY")));
 		Sig.Inputs.Add(FNiagaraVariable(FNiagaraTypeDefinition::GetIntDef(), TEXT("IndexZ")));
 		Sig.Inputs.Add(FNiagaraVariable(FNiagaraTypeDefinition::GetVec3Def(), TEXT("Value")));
-		Sig.FunctionSpecifiers.Add(FName(NAME_Attribute));
+		Sig.FunctionSpecifiers.Add(FName(UNiagaraDataInterfaceRWBase::NAME_Attribute));
 		Sig.bMemberFunction = true;
 		Sig.bRequiresContext = false;
 		Sig.bExperimental = true;
@@ -741,7 +740,7 @@ void UNiagaraDataInterfaceGrid3DCollection::GetFunctions(TArray<FNiagaraFunction
 		Sig.Inputs.Add(FNiagaraVariable(FNiagaraTypeDefinition::GetIntDef(), TEXT("IndexY")));
 		Sig.Inputs.Add(FNiagaraVariable(FNiagaraTypeDefinition::GetIntDef(), TEXT("IndexZ")));
 		Sig.Outputs.Add(FNiagaraVariable(FNiagaraTypeDefinition::GetVec3Def(), TEXT("Value")));
-		Sig.FunctionSpecifiers.Add(FName(NAME_Attribute));
+		Sig.FunctionSpecifiers.Add(FName(UNiagaraDataInterfaceRWBase::NAME_Attribute));
 		Sig.bMemberFunction = true;
 		Sig.bRequiresContext = false;
 		Sig.bExperimental = true;
@@ -762,7 +761,7 @@ void UNiagaraDataInterfaceGrid3DCollection::GetFunctions(TArray<FNiagaraFunction
 		Sig.Inputs.Add(FNiagaraVariable(FNiagaraTypeDefinition(GetClass()), TEXT("Grid")));
 		Sig.Inputs.Add(FNiagaraVariable(FNiagaraTypeDefinition::GetVec3Def(), TEXT("Unit")));
 		Sig.Outputs.Add(FNiagaraVariable(FNiagaraTypeDefinition::GetVec3Def(), TEXT("Value")));
-		Sig.FunctionSpecifiers.Add(FName(NAME_Attribute));
+		Sig.FunctionSpecifiers.Add(FName(UNiagaraDataInterfaceRWBase::NAME_Attribute));
 		Sig.bMemberFunction = true;
 		Sig.bRequiresContext = false;
 		Sig.bExperimental = true;
@@ -783,7 +782,7 @@ void UNiagaraDataInterfaceGrid3DCollection::GetFunctions(TArray<FNiagaraFunction
 		Sig.Inputs.Add(FNiagaraVariable(FNiagaraTypeDefinition(GetClass()), TEXT("Grid")));
 		Sig.Inputs.Add(FNiagaraVariable(FNiagaraTypeDefinition::GetVec3Def(), TEXT("Unit")));
 		Sig.Outputs.Add(FNiagaraVariable(FNiagaraTypeDefinition::GetVec3Def(), TEXT("Value")));
-		Sig.FunctionSpecifiers.Add(FName(NAME_Attribute));
+		Sig.FunctionSpecifiers.Add(FName(UNiagaraDataInterfaceRWBase::NAME_Attribute));
 		Sig.bMemberFunction = true;
 		Sig.bRequiresContext = false;
 		Sig.bExperimental = true;
@@ -806,7 +805,7 @@ void UNiagaraDataInterfaceGrid3DCollection::GetFunctions(TArray<FNiagaraFunction
 		Sig.Inputs.Add(FNiagaraVariable(FNiagaraTypeDefinition::GetIntDef(), TEXT("IndexY")));
 		Sig.Inputs.Add(FNiagaraVariable(FNiagaraTypeDefinition::GetIntDef(), TEXT("IndexZ")));
 		Sig.Inputs.Add(FNiagaraVariable(FNiagaraTypeDefinition::GetVec2Def(), TEXT("Value")));
-		Sig.FunctionSpecifiers.Add(FName(NAME_Attribute));
+		Sig.FunctionSpecifiers.Add(FName(UNiagaraDataInterfaceRWBase::NAME_Attribute));
 		Sig.bMemberFunction = true;
 		Sig.bRequiresContext = false;
 		Sig.bExperimental = true;
@@ -830,7 +829,7 @@ void UNiagaraDataInterfaceGrid3DCollection::GetFunctions(TArray<FNiagaraFunction
 		Sig.Inputs.Add(FNiagaraVariable(FNiagaraTypeDefinition::GetIntDef(), TEXT("IndexY")));
 		Sig.Inputs.Add(FNiagaraVariable(FNiagaraTypeDefinition::GetIntDef(), TEXT("IndexZ")));
 		Sig.Outputs.Add(FNiagaraVariable(FNiagaraTypeDefinition::GetVec2Def(), TEXT("Value")));
-		Sig.FunctionSpecifiers.Add(FName(NAME_Attribute));
+		Sig.FunctionSpecifiers.Add(FName(UNiagaraDataInterfaceRWBase::NAME_Attribute));
 		Sig.bMemberFunction = true;
 		Sig.bRequiresContext = false;
 		Sig.bExperimental = true;
@@ -851,7 +850,7 @@ void UNiagaraDataInterfaceGrid3DCollection::GetFunctions(TArray<FNiagaraFunction
 		Sig.Inputs.Add(FNiagaraVariable(FNiagaraTypeDefinition(GetClass()), TEXT("Grid")));
 		Sig.Inputs.Add(FNiagaraVariable(FNiagaraTypeDefinition::GetVec3Def(), TEXT("Unit")));
 		Sig.Outputs.Add(FNiagaraVariable(FNiagaraTypeDefinition::GetVec2Def(), TEXT("Value")));
-		Sig.FunctionSpecifiers.Add(FName(NAME_Attribute));
+		Sig.FunctionSpecifiers.Add(FName(UNiagaraDataInterfaceRWBase::NAME_Attribute));
 		Sig.bMemberFunction = true;
 		Sig.bRequiresContext = false;
 		Sig.bExperimental = true;
@@ -872,7 +871,7 @@ void UNiagaraDataInterfaceGrid3DCollection::GetFunctions(TArray<FNiagaraFunction
 		Sig.Inputs.Add(FNiagaraVariable(FNiagaraTypeDefinition(GetClass()), TEXT("Grid")));
 		Sig.Inputs.Add(FNiagaraVariable(FNiagaraTypeDefinition::GetVec3Def(), TEXT("Unit")));
 		Sig.Outputs.Add(FNiagaraVariable(FNiagaraTypeDefinition::GetVec2Def(), TEXT("Value")));
-		Sig.FunctionSpecifiers.Add(FName(NAME_Attribute));
+		Sig.FunctionSpecifiers.Add(FName(UNiagaraDataInterfaceRWBase::NAME_Attribute));
 		Sig.bMemberFunction = true;
 		Sig.bRequiresContext = false;
 		Sig.bExperimental = true;
@@ -895,7 +894,7 @@ void UNiagaraDataInterfaceGrid3DCollection::GetFunctions(TArray<FNiagaraFunction
 		Sig.Inputs.Add(FNiagaraVariable(FNiagaraTypeDefinition::GetIntDef(), TEXT("IndexY")));
 		Sig.Inputs.Add(FNiagaraVariable(FNiagaraTypeDefinition::GetIntDef(), TEXT("IndexZ")));
 		Sig.Inputs.Add(FNiagaraVariable(FNiagaraTypeDefinition::GetFloatDef(), TEXT("Value")));
-		Sig.FunctionSpecifiers.Add(FName(NAME_Attribute));
+		Sig.FunctionSpecifiers.Add(FName(UNiagaraDataInterfaceRWBase::NAME_Attribute));
 		Sig.bMemberFunction = true;
 		Sig.bRequiresContext = false;
 		Sig.bExperimental = true;
@@ -918,7 +917,7 @@ void UNiagaraDataInterfaceGrid3DCollection::GetFunctions(TArray<FNiagaraFunction
 		Sig.Inputs.Add(FNiagaraVariable(FNiagaraTypeDefinition::GetIntDef(), TEXT("IndexY")));
 		Sig.Inputs.Add(FNiagaraVariable(FNiagaraTypeDefinition::GetIntDef(), TEXT("IndexZ")));
 		Sig.Outputs.Add(FNiagaraVariable(FNiagaraTypeDefinition::GetFloatDef(), TEXT("Value")));
-		Sig.FunctionSpecifiers.Add(FName(NAME_Attribute));
+		Sig.FunctionSpecifiers.Add(FName(UNiagaraDataInterfaceRWBase::NAME_Attribute));
 		Sig.bMemberFunction = true;
 		Sig.bRequiresContext = false;
 		Sig.bExperimental = true;
@@ -939,7 +938,7 @@ void UNiagaraDataInterfaceGrid3DCollection::GetFunctions(TArray<FNiagaraFunction
 		Sig.Inputs.Add(FNiagaraVariable(FNiagaraTypeDefinition(GetClass()), TEXT("Grid")));
 		Sig.Inputs.Add(FNiagaraVariable(FNiagaraTypeDefinition::GetVec3Def(), TEXT("Unit")));
 		Sig.Outputs.Add(FNiagaraVariable(FNiagaraTypeDefinition::GetFloatDef(), TEXT("Value")));
-		Sig.FunctionSpecifiers.Add(FName(NAME_Attribute));
+		Sig.FunctionSpecifiers.Add(FName(UNiagaraDataInterfaceRWBase::NAME_Attribute));
 		Sig.bMemberFunction = true;
 		Sig.bRequiresContext = false;
 		Sig.bExperimental = true;
@@ -960,7 +959,7 @@ void UNiagaraDataInterfaceGrid3DCollection::GetFunctions(TArray<FNiagaraFunction
 		Sig.Inputs.Add(FNiagaraVariable(FNiagaraTypeDefinition(GetClass()), TEXT("Grid")));
 		Sig.Inputs.Add(FNiagaraVariable(FNiagaraTypeDefinition::GetVec3Def(), TEXT("Unit")));
 		Sig.Outputs.Add(FNiagaraVariable(FNiagaraTypeDefinition::GetFloatDef(), TEXT("Value")));
-		Sig.FunctionSpecifiers.Add(FName(NAME_Attribute));
+		Sig.FunctionSpecifiers.Add(FName(UNiagaraDataInterfaceRWBase::NAME_Attribute));
 		Sig.bMemberFunction = true;
 		Sig.bRequiresContext = false;
 		Sig.bExperimental = true;
@@ -1100,7 +1099,7 @@ void UNiagaraDataInterfaceGrid3DCollection::GetFunctions(TArray<FNiagaraFunction
 		Sig.Name = GetVector4AttributeIndexFunctionName;
 		Sig.Inputs.Add(FNiagaraVariable(FNiagaraTypeDefinition(GetClass()), TEXT("Grid")));
 		Sig.Outputs.Add(FNiagaraVariable(FNiagaraTypeDefinition::GetIntDef(), TEXT("AttributeIndex")));
-		Sig.FunctionSpecifiers.Add(FName(NAME_Attribute));
+		Sig.FunctionSpecifiers.Add(FName(UNiagaraDataInterfaceRWBase::NAME_Attribute));
 		Sig.bMemberFunction = true;
 		Sig.bRequiresContext = false;
 		Sig.bExperimental = true;
@@ -1115,7 +1114,7 @@ void UNiagaraDataInterfaceGrid3DCollection::GetFunctions(TArray<FNiagaraFunction
 		Sig.Name = GetVectorAttributeIndexFunctionName;
 		Sig.Inputs.Add(FNiagaraVariable(FNiagaraTypeDefinition(GetClass()), TEXT("Grid")));
 		Sig.Outputs.Add(FNiagaraVariable(FNiagaraTypeDefinition::GetIntDef(), TEXT("AttributeIndex")));
-		Sig.FunctionSpecifiers.Add(FName(NAME_Attribute));
+		Sig.FunctionSpecifiers.Add(FName(UNiagaraDataInterfaceRWBase::NAME_Attribute));
 		Sig.bMemberFunction = true;
 		Sig.bRequiresContext = false;
 		Sig.bExperimental = true;
@@ -1130,7 +1129,7 @@ void UNiagaraDataInterfaceGrid3DCollection::GetFunctions(TArray<FNiagaraFunction
 		Sig.Name = GetVector2DAttributeIndexFunctionName;
 		Sig.Inputs.Add(FNiagaraVariable(FNiagaraTypeDefinition(GetClass()), TEXT("Grid")));
 		Sig.Outputs.Add(FNiagaraVariable(FNiagaraTypeDefinition::GetIntDef(), TEXT("AttributeIndex")));
-		Sig.FunctionSpecifiers.Add(FName(NAME_Attribute));
+		Sig.FunctionSpecifiers.Add(FName(UNiagaraDataInterfaceRWBase::NAME_Attribute));
 		Sig.bMemberFunction = true;
 		Sig.bRequiresContext = false;
 		Sig.bExperimental = true;
@@ -1145,7 +1144,7 @@ void UNiagaraDataInterfaceGrid3DCollection::GetFunctions(TArray<FNiagaraFunction
 		Sig.Name = GetFloatAttributeIndexFunctionName;
 		Sig.Inputs.Add(FNiagaraVariable(FNiagaraTypeDefinition(GetClass()), TEXT("Grid")));
 		Sig.Outputs.Add(FNiagaraVariable(FNiagaraTypeDefinition::GetIntDef(), TEXT("AttributeIndex")));
-		Sig.FunctionSpecifiers.Add(FName(NAME_Attribute));
+		Sig.FunctionSpecifiers.Add(FName(UNiagaraDataInterfaceRWBase::NAME_Attribute));
 		Sig.bMemberFunction = true;
 		Sig.bRequiresContext = false;
 		Sig.bExperimental = true;
@@ -1538,25 +1537,25 @@ void UNiagaraDataInterfaceGrid3DCollection::GetVMExternalFunction(const FVMExter
 	}
 	else if (BindingInfo.Name == GetVector4AttributeIndexFunctionName)
 	{
-		FName AttributeName = BindingInfo.FindSpecifier(NAME_Attribute)->Value;
+		FName AttributeName = BindingInfo.FindSpecifier(UNiagaraDataInterfaceRWBase::NAME_Attribute)->Value;
 		check(BindingInfo.GetNumInputs() == 1 && BindingInfo.GetNumOutputs() == 1);
 		NDI_FUNC_BINDER(UNiagaraDataInterfaceGrid3DCollection, VMGetAttributeIndex)::Bind(this, OutFunc, AttributeName, 4);
 	}
 	else if (BindingInfo.Name == GetVectorAttributeIndexFunctionName)
 	{
-		FName AttributeName = BindingInfo.FindSpecifier(NAME_Attribute)->Value;
+		FName AttributeName = BindingInfo.FindSpecifier(UNiagaraDataInterfaceRWBase::NAME_Attribute)->Value;
 		check(BindingInfo.GetNumInputs() == 1 && BindingInfo.GetNumOutputs() == 1);
 		NDI_FUNC_BINDER(UNiagaraDataInterfaceGrid3DCollection, VMGetAttributeIndex)::Bind(this, OutFunc, AttributeName, 3);
 	}
 	else if (BindingInfo.Name == GetVector2DAttributeIndexFunctionName)
 	{
-		FName AttributeName = BindingInfo.FindSpecifier(NAME_Attribute)->Value;
+		FName AttributeName = BindingInfo.FindSpecifier(UNiagaraDataInterfaceRWBase::NAME_Attribute)->Value;
 		check(BindingInfo.GetNumInputs() == 1 && BindingInfo.GetNumOutputs() == 1);
 		NDI_FUNC_BINDER(UNiagaraDataInterfaceGrid3DCollection, VMGetAttributeIndex)::Bind(this, OutFunc, AttributeName, 2);
 	}
 	else if (BindingInfo.Name == GetFloatAttributeIndexFunctionName)
 	{
-		FName AttributeName = BindingInfo.FindSpecifier(NAME_Attribute)->Value;
+		FName AttributeName = BindingInfo.FindSpecifier(UNiagaraDataInterfaceRWBase::NAME_Attribute)->Value;
 		check(BindingInfo.GetNumInputs() == 1 && BindingInfo.GetNumOutputs() == 1);
 		NDI_FUNC_BINDER(UNiagaraDataInterfaceGrid3DCollection, VMGetAttributeIndex)::Bind(this, OutFunc, AttributeName, 1);
 	}
@@ -1655,10 +1654,12 @@ void UNiagaraDataInterfaceGrid3DCollection::GetParameterDefinitionHLSL(const FNi
 
 	// format of the grid is defined by number of attributes
 	FString FloatBufferSuffix = "";
-	FGrid3DCollectionAttributeHelper Helper(ParamInfo);
-	if (Helper.UseRGBAGrid() && Helper.TotalChannels > 1)
+	FGrid3DCollectionAttributeHlslWriter AttributeWriter(ParamInfo);
+
+	int TotalChannels = AttributeWriter.GetAttributeHelper().GetTotalChannels();
+	if (AttributeWriter.UseRGBAGrid() && TotalChannels > 1)
 	{
-		FloatBufferSuffix = FString::FromInt(Helper.TotalChannels);
+		FloatBufferSuffix = FString::FromInt(TotalChannels);
 	}
 
 	TMap<FString, FStringFormatArg> ArgsDeclarations = {
@@ -1667,7 +1668,7 @@ void UNiagaraDataInterfaceGrid3DCollection::GetParameterDefinitionHLSL(const FNi
 		{ TEXT("SamplerName"),		ParamInfo.DataInterfaceHLSLSymbol + SamplerName},
 		{ TEXT("OutputGridName"),	ParamInfo.DataInterfaceHLSLSymbol + OutputGridName},
 		{ TEXT("FloatBufferSuffix"),FloatBufferSuffix},
-		{ TEXT("NumChannels"),		Helper.UseRGBAGrid() ? Helper.TotalChannels : 1}, // float grids only have a single channel for the underlying buffer
+		{ TEXT("NumChannels"),		AttributeWriter.UseRGBAGrid() ? TotalChannels : 1}, // float grids only have a single channel for the underlying buffer
 		{ TEXT("NumTiles"),			ParamInfo.DataInterfaceHLSLSymbol + NumTilesName},
 		{ TEXT("OneOverNumTiles"),	ParamInfo.DataInterfaceHLSLSymbol + OneOverNumTilesName},
 		{ TEXT("UnitClampMin"),		ParamInfo.DataInterfaceHLSLSymbol + UnitClampMinName},
@@ -2054,8 +2055,8 @@ bool UNiagaraDataInterfaceGrid3DCollection::GetFunctionHLSL(const FNiagaraDataIn
 		return true;
 	}
 
-	FGrid3DCollectionAttributeHelper ParamHelper(ParamInfo);
-	FString GridChannelString = ParamHelper.GetGridChannelString();
+	FGrid3DCollectionAttributeHlslWriter AttributeWriter(ParamInfo);
+	FString GridChannelString = AttributeWriter.GetAttributeHelper().GetGridChannelString();
 
 	TMap<FString, FStringFormatArg> ArgsBounds =
 	{
@@ -2081,7 +2082,7 @@ bool UNiagaraDataInterfaceGrid3DCollection::GetFunctionHLSL(const FNiagaraDataIn
 	if (FunctionInfo.DefinitionName == GetValueFunctionName || FunctionInfo.DefinitionName == GetPreviousValueAtIndexFunctionName)
 	{
 		// if we have an rgba grid, then we must query for the rgba value at the current point, and get the specific channel requested
-		if (ParamHelper.UseRGBAGrid())
+		if (AttributeWriter.UseRGBAGrid())
 		{			
 			FString RGBAHlsl;
 			RGBAHlsl.Appendf(TEXT("void {FunctionName}(int In_IndexX, int In_IndexY, int In_IndexZ, int In_AttributeIndex, out float Out_Val)\n"));
@@ -2093,11 +2094,11 @@ bool UNiagaraDataInterfaceGrid3DCollection::GetFunctionHLSL(const FNiagaraDataIn
 			RGBAHlsl.Appendf(TEXT("        Out_Val = TmpVal.r;\n"));
 			RGBAHlsl.Appendf(TEXT("    }\n"));
 
-			for (int i = 0; i < ParamHelper.TotalChannels - 1; ++i)
+			for (int i = 0; i < AttributeWriter.GetAttributeHelper().GetTotalChannels() - 1; ++i)
 			{
 				RGBAHlsl.Appendf(TEXT("    else if (In_AttributeIndex == %d)\n"), i + 1);
 				RGBAHlsl.Appendf(TEXT("    {\n"));
-				RGBAHlsl.Appendf(TEXT("        Out_Val = TmpVal.%s;\n"), *FGrid3DCollectionAttributeHelper::Channels[i + 1]);
+				RGBAHlsl.Appendf(TEXT("        Out_Val = TmpVal.%s;\n"), *FGrid3DCollectionAttributeHlslWriter::Channels[i + 1]);
 				RGBAHlsl.Appendf(TEXT("    }\n"));
 			}
 
@@ -2137,7 +2138,7 @@ bool UNiagaraDataInterfaceGrid3DCollection::GetFunctionHLSL(const FNiagaraDataIn
 
 		// if we are using an rgba grid, then we must query the float4 value from the uav, then set the channel we are requesting
 		// note that the dynamic branching seems necessary otherwise it didn't work as expected.  Might be due to code optimization
-		if (ParamHelper.UseRGBAGrid())
+		if (AttributeWriter.UseRGBAGrid())
 		{			
 			FString RGBAHlsl;
 			RGBAHlsl.Appendf(TEXT("void {FunctionName}(int In_IndexX, int In_IndexY, int In_IndexZ, int In_AttributeIndex, float In_Value, out int val)\n"));
@@ -2150,12 +2151,12 @@ bool UNiagaraDataInterfaceGrid3DCollection::GetFunctionHLSL(const FNiagaraDataIn
 			RGBAHlsl.Appendf(TEXT("        {OutputGrid}[int3(In_IndexX, In_IndexY, In_IndexZ)] = TmpVal;\n"));
 			RGBAHlsl.Appendf(TEXT("    }\n"));
 
-			for (int i = 0; i < ParamHelper.TotalChannels - 1; ++i)
+			for (int i = 0; i < AttributeWriter.GetAttributeHelper().GetTotalChannels() - 1; ++i)
 			{
 				RGBAHlsl.Appendf(TEXT("    else if (In_AttributeIndex == %d)\n"), i + 1);
 				RGBAHlsl.Appendf(TEXT("    {\n"));
 				RGBAHlsl.Appendf(TEXT("		   float{GridChannels} TmpVal = {OutputGrid}.Load(int4(In_IndexX, In_IndexY, In_IndexZ, 0));\n"));
-				RGBAHlsl.Appendf(TEXT("        TmpVal.%s = In_Value;\n"), *FGrid3DCollectionAttributeHelper::Channels[i + 1]);
+				RGBAHlsl.Appendf(TEXT("        TmpVal.%s = In_Value;\n"), *FGrid3DCollectionAttributeHlslWriter::Channels[i + 1]);
 				RGBAHlsl.Appendf(TEXT("        {OutputGrid}[int3(In_IndexX, In_IndexY, In_IndexZ)] = TmpVal;\n"));
 				RGBAHlsl.Appendf(TEXT("    }\n"));
 			}
@@ -2195,7 +2196,7 @@ bool UNiagaraDataInterfaceGrid3DCollection::GetFunctionHLSL(const FNiagaraDataIn
 	}
 	else if (FunctionInfo.DefinitionName == CopyPreviousToCurrentForCellFunctionName)
 	{
-		if (ParamHelper.UseRGBAGrid())
+		if (AttributeWriter.UseRGBAGrid())
 		{
 			static const TCHAR* FormatBounds = TEXT(R"(
 				void {FunctionName}(int In_IndexX, int In_IndexY, int In_IndexZ)
@@ -2226,7 +2227,7 @@ bool UNiagaraDataInterfaceGrid3DCollection::GetFunctionHLSL(const FNiagaraDataIn
 	}
 	else if (FunctionInfo.DefinitionName == CopyMaskedPreviousToCurrentForCellFunctionName)
 	{
-		if (ParamHelper.UseRGBAGrid())
+		if (AttributeWriter.UseRGBAGrid())
 		{			
 			static const TCHAR* FormatBounds = TEXT(R"(
 				void {FunctionName}(int In_IndexX, int In_IndexY, int In_IndexZ, int NumAttributesSet, int AttributeMask)
@@ -2306,9 +2307,9 @@ bool UNiagaraDataInterfaceGrid3DCollection::GetFunctionHLSL(const FNiagaraDataIn
 	}
 	else if (FunctionInfo.DefinitionName == SetVector4ValueFunctionName)
 	{
-		if (ParamHelper.UseRGBAGrid())
+		if (AttributeWriter.UseRGBAGrid())
 		{
-			ParamHelper.WriteSetHLSL(FGrid3DCollectionAttributeHelper::AttributeRetrievalMode::RGBAGrid, FunctionInfo, OutHLSL);
+			AttributeWriter.WriteSetHLSL(EGridAttributeRetrievalMode::RGBAGrid, FunctionInfo, OutHLSL);
 		}
 		else
 		{
@@ -2318,9 +2319,9 @@ bool UNiagaraDataInterfaceGrid3DCollection::GetFunctionHLSL(const FNiagaraDataIn
 	}
 	else if (FunctionInfo.DefinitionName == GetPreviousVector4ValueFunctionName)
 	{
-		if (ParamHelper.UseRGBAGrid())
+		if (AttributeWriter.UseRGBAGrid())
 		{
-			ParamHelper.WriteGetHLSL(FGrid3DCollectionAttributeHelper::AttributeRetrievalMode::RGBAGrid, FunctionInfo, OutHLSL);
+			AttributeWriter.WriteGetHLSL(EGridAttributeRetrievalMode::RGBAGrid, FunctionInfo, OutHLSL);
 		}
 		else
 		{
@@ -2331,9 +2332,9 @@ bool UNiagaraDataInterfaceGrid3DCollection::GetFunctionHLSL(const FNiagaraDataIn
 	}
 	else if (FunctionInfo.DefinitionName == SamplePreviousGridVector4FunctionName)
 	{
-		if (ParamHelper.UseRGBAGrid())
+		if (AttributeWriter.UseRGBAGrid())
 		{
-			ParamHelper.WriteSampleHLSL(FGrid3DCollectionAttributeHelper::AttributeRetrievalMode::RGBAGrid, FunctionInfo, false, OutHLSL);
+			AttributeWriter.WriteSampleHLSL(EGridAttributeRetrievalMode::RGBAGrid, FunctionInfo, false, OutHLSL);
 		}
 		else
 		{
@@ -2343,9 +2344,9 @@ bool UNiagaraDataInterfaceGrid3DCollection::GetFunctionHLSL(const FNiagaraDataIn
 	}
 	else if (FunctionInfo.DefinitionName == CubicSamplePreviousGridVector4FunctionName)
 	{
-		if (ParamHelper.UseRGBAGrid())
+		if (AttributeWriter.UseRGBAGrid())
 		{
-			ParamHelper.WriteSampleHLSL(FGrid3DCollectionAttributeHelper::AttributeRetrievalMode::RGBAGrid, FunctionInfo, true, OutHLSL);
+			AttributeWriter.WriteSampleHLSL(EGridAttributeRetrievalMode::RGBAGrid, FunctionInfo, true, OutHLSL);
 		}
 		else
 		{
@@ -2355,9 +2356,9 @@ bool UNiagaraDataInterfaceGrid3DCollection::GetFunctionHLSL(const FNiagaraDataIn
 	}
 	else if (FunctionInfo.DefinitionName == SetVector3ValueFunctionName || FunctionInfo.DefinitionName == SetVectorValueFunctionName)
 	{
-		if (ParamHelper.UseRGBAGrid())
+		if (AttributeWriter.UseRGBAGrid())
 		{
-			ParamHelper.WriteSetHLSL(FGrid3DCollectionAttributeHelper::AttributeRetrievalMode::RGBAGrid, FunctionInfo, OutHLSL);
+			AttributeWriter.WriteSetHLSL(EGridAttributeRetrievalMode::RGBAGrid, FunctionInfo, OutHLSL);
 		}
 		else
 		{
@@ -2367,9 +2368,9 @@ bool UNiagaraDataInterfaceGrid3DCollection::GetFunctionHLSL(const FNiagaraDataIn
 	}
 	else if (FunctionInfo.DefinitionName == GetPreviousVectorValueFunctionName)
 	{
-		if (ParamHelper.UseRGBAGrid())
+		if (AttributeWriter.UseRGBAGrid())
 		{
-			ParamHelper.WriteGetHLSL(FGrid3DCollectionAttributeHelper::AttributeRetrievalMode::RGBAGrid, FunctionInfo, OutHLSL);
+			AttributeWriter.WriteGetHLSL(EGridAttributeRetrievalMode::RGBAGrid, FunctionInfo, OutHLSL);
 		}
 		else
 		{
@@ -2380,9 +2381,9 @@ bool UNiagaraDataInterfaceGrid3DCollection::GetFunctionHLSL(const FNiagaraDataIn
 	}
 	else if (FunctionInfo.DefinitionName == SamplePreviousGridVectorFunctionName)
 	{
-		if (ParamHelper.UseRGBAGrid())
+		if (AttributeWriter.UseRGBAGrid())
 		{
-			ParamHelper.WriteSampleHLSL(FGrid3DCollectionAttributeHelper::AttributeRetrievalMode::RGBAGrid, FunctionInfo, false, OutHLSL);
+			AttributeWriter.WriteSampleHLSL(EGridAttributeRetrievalMode::RGBAGrid, FunctionInfo, false, OutHLSL);
 		}
 		else
 		{
@@ -2393,9 +2394,9 @@ bool UNiagaraDataInterfaceGrid3DCollection::GetFunctionHLSL(const FNiagaraDataIn
 	}
 	else if (FunctionInfo.DefinitionName == CubicSamplePreviousGridVectorFunctionName)
 	{
-		if (ParamHelper.UseRGBAGrid())
+		if (AttributeWriter.UseRGBAGrid())
 		{
-			ParamHelper.WriteSampleHLSL(FGrid3DCollectionAttributeHelper::AttributeRetrievalMode::RGBAGrid, FunctionInfo, true, OutHLSL);
+			AttributeWriter.WriteSampleHLSL(EGridAttributeRetrievalMode::RGBAGrid, FunctionInfo, true, OutHLSL);
 		}
 		else
 		{
@@ -2404,9 +2405,9 @@ bool UNiagaraDataInterfaceGrid3DCollection::GetFunctionHLSL(const FNiagaraDataIn
 	}
 	else if (FunctionInfo.DefinitionName == SetVector2ValueFunctionName || FunctionInfo.DefinitionName == SetVector2DValueFunctionName)
 	{
-		if (ParamHelper.UseRGBAGrid())
+		if (AttributeWriter.UseRGBAGrid())
 		{
-			ParamHelper.WriteSetHLSL(FGrid3DCollectionAttributeHelper::AttributeRetrievalMode::RGBAGrid, FunctionInfo, OutHLSL);
+			AttributeWriter.WriteSetHLSL(EGridAttributeRetrievalMode::RGBAGrid, FunctionInfo, OutHLSL);
 		}
 		else
 		{
@@ -2417,9 +2418,9 @@ bool UNiagaraDataInterfaceGrid3DCollection::GetFunctionHLSL(const FNiagaraDataIn
 	}
 	else if (FunctionInfo.DefinitionName == GetPreviousVector2DValueFunctionName)
 	{
-		if (ParamHelper.UseRGBAGrid())
+		if (AttributeWriter.UseRGBAGrid())
 		{
-			ParamHelper.WriteGetHLSL(FGrid3DCollectionAttributeHelper::AttributeRetrievalMode::RGBAGrid, FunctionInfo, OutHLSL);
+			AttributeWriter.WriteGetHLSL(EGridAttributeRetrievalMode::RGBAGrid, FunctionInfo, OutHLSL);
 		}
 		else
 		{
@@ -2430,9 +2431,9 @@ bool UNiagaraDataInterfaceGrid3DCollection::GetFunctionHLSL(const FNiagaraDataIn
 	}
 	else if (FunctionInfo.DefinitionName == SamplePreviousGridVector2DFunctionName)
 	{
-		if (ParamHelper.UseRGBAGrid())
+		if (AttributeWriter.UseRGBAGrid())
 		{
-			ParamHelper.WriteSampleHLSL(FGrid3DCollectionAttributeHelper::AttributeRetrievalMode::RGBAGrid, FunctionInfo, false, OutHLSL);
+			AttributeWriter.WriteSampleHLSL(EGridAttributeRetrievalMode::RGBAGrid, FunctionInfo, false, OutHLSL);
 		}
 		else
 		{
@@ -2442,9 +2443,9 @@ bool UNiagaraDataInterfaceGrid3DCollection::GetFunctionHLSL(const FNiagaraDataIn
 	}
 	else if (FunctionInfo.DefinitionName == CubicSamplePreviousGridVector2DFunctionName)
 	{
-		if (ParamHelper.UseRGBAGrid())
+		if (AttributeWriter.UseRGBAGrid())
 		{
-			ParamHelper.WriteSampleHLSL(FGrid3DCollectionAttributeHelper::AttributeRetrievalMode::RGBAGrid, FunctionInfo, true, OutHLSL);
+			AttributeWriter.WriteSampleHLSL(EGridAttributeRetrievalMode::RGBAGrid, FunctionInfo, true, OutHLSL);
 		}
 		else
 		{
@@ -2453,9 +2454,9 @@ bool UNiagaraDataInterfaceGrid3DCollection::GetFunctionHLSL(const FNiagaraDataIn
 	}
 	else if (FunctionInfo.DefinitionName == SetFloatValueFunctionName)
 	{
-		if (ParamHelper.UseRGBAGrid())
+		if (AttributeWriter.UseRGBAGrid())
 		{
-			ParamHelper.WriteSetHLSL(FGrid3DCollectionAttributeHelper::AttributeRetrievalMode::RGBAGrid, FunctionInfo, OutHLSL);
+			AttributeWriter.WriteSetHLSL(EGridAttributeRetrievalMode::RGBAGrid, FunctionInfo, OutHLSL);
 		}
 		else
 		{
@@ -2466,9 +2467,9 @@ bool UNiagaraDataInterfaceGrid3DCollection::GetFunctionHLSL(const FNiagaraDataIn
 	}
 	else if (FunctionInfo.DefinitionName == GetPreviousFloatValueFunctionName)
 	{
-		if (ParamHelper.UseRGBAGrid())
+		if (AttributeWriter.UseRGBAGrid())
 		{
-			ParamHelper.WriteGetHLSL(FGrid3DCollectionAttributeHelper::AttributeRetrievalMode::RGBAGrid, FunctionInfo, OutHLSL);
+			AttributeWriter.WriteGetHLSL(EGridAttributeRetrievalMode::RGBAGrid, FunctionInfo, OutHLSL);
 		}
 		else
 		{
@@ -2479,9 +2480,9 @@ bool UNiagaraDataInterfaceGrid3DCollection::GetFunctionHLSL(const FNiagaraDataIn
 	}
 	else if (FunctionInfo.DefinitionName == SamplePreviousGridFloatFunctionName)
 	{
-		if (ParamHelper.UseRGBAGrid())
+		if (AttributeWriter.UseRGBAGrid())
 		{
-			ParamHelper.WriteSampleHLSL(FGrid3DCollectionAttributeHelper::AttributeRetrievalMode::RGBAGrid, FunctionInfo, false, OutHLSL);
+			AttributeWriter.WriteSampleHLSL(EGridAttributeRetrievalMode::RGBAGrid, FunctionInfo, false, OutHLSL);
 		}
 		else
 		{
@@ -2492,9 +2493,9 @@ bool UNiagaraDataInterfaceGrid3DCollection::GetFunctionHLSL(const FNiagaraDataIn
 	}
 	else if (FunctionInfo.DefinitionName == CubicSamplePreviousGridFloatFunctionName)
 	{
-		if (ParamHelper.UseRGBAGrid())
+		if (AttributeWriter.UseRGBAGrid())
 		{
-			ParamHelper.WriteSampleHLSL(FGrid3DCollectionAttributeHelper::AttributeRetrievalMode::RGBAGrid, FunctionInfo, true, OutHLSL);
+			AttributeWriter.WriteSampleHLSL(EGridAttributeRetrievalMode::RGBAGrid, FunctionInfo, true, OutHLSL);
 		}
 		else
 		{
@@ -2504,9 +2505,9 @@ bool UNiagaraDataInterfaceGrid3DCollection::GetFunctionHLSL(const FNiagaraDataIn
 	}
 	else if (FunctionInfo.DefinitionName == GetVector4AttributeIndexFunctionName)
 	{
-		if (ParamHelper.UseRGBAGrid())
+		if (AttributeWriter.UseRGBAGrid())
 		{
-			ParamHelper.WriteAttributeGetIndexHLSL(FGrid3DCollectionAttributeHelper::AttributeRetrievalMode::RGBAGrid, FunctionInfo, OutHLSL);
+			AttributeWriter.WriteAttributeGetIndexHLSL(EGridAttributeRetrievalMode::RGBAGrid, FunctionInfo, OutHLSL);
 		}
 		else
 		{
@@ -2516,9 +2517,9 @@ bool UNiagaraDataInterfaceGrid3DCollection::GetFunctionHLSL(const FNiagaraDataIn
 	}
 	else if (FunctionInfo.DefinitionName == GetVectorAttributeIndexFunctionName)
 	{
-		if (ParamHelper.UseRGBAGrid())
+		if (AttributeWriter.UseRGBAGrid())
 		{
-			ParamHelper.WriteAttributeGetIndexHLSL(FGrid3DCollectionAttributeHelper::AttributeRetrievalMode::RGBAGrid, FunctionInfo, OutHLSL);
+			AttributeWriter.WriteAttributeGetIndexHLSL(EGridAttributeRetrievalMode::RGBAGrid, FunctionInfo, OutHLSL);
 		}
 		else
 		{
@@ -2528,9 +2529,9 @@ bool UNiagaraDataInterfaceGrid3DCollection::GetFunctionHLSL(const FNiagaraDataIn
 	}
 	else if (FunctionInfo.DefinitionName == GetVector2DAttributeIndexFunctionName)
 	{
-		if (ParamHelper.UseRGBAGrid())
+		if (AttributeWriter.UseRGBAGrid())
 		{
-			ParamHelper.WriteAttributeGetIndexHLSL(FGrid3DCollectionAttributeHelper::AttributeRetrievalMode::RGBAGrid, FunctionInfo, OutHLSL);
+			AttributeWriter.WriteAttributeGetIndexHLSL(EGridAttributeRetrievalMode::RGBAGrid, FunctionInfo, OutHLSL);
 		}
 		else
 		{
@@ -2540,9 +2541,9 @@ bool UNiagaraDataInterfaceGrid3DCollection::GetFunctionHLSL(const FNiagaraDataIn
 	}
 	else if (FunctionInfo.DefinitionName == GetFloatAttributeIndexFunctionName)
 	{
-		if (ParamHelper.UseRGBAGrid())
+		if (AttributeWriter.UseRGBAGrid())
 		{
-			ParamHelper.WriteAttributeGetIndexHLSL(FGrid3DCollectionAttributeHelper::AttributeRetrievalMode::RGBAGrid, FunctionInfo, OutHLSL);
+			AttributeWriter.WriteAttributeGetIndexHLSL(EGridAttributeRetrievalMode::RGBAGrid, FunctionInfo, OutHLSL);
 		}
 		else
 		{
@@ -2554,7 +2555,7 @@ bool UNiagaraDataInterfaceGrid3DCollection::GetFunctionHLSL(const FNiagaraDataIn
 	{
 		
 	// if we have an rgba grid, then we must sample for the rgba value at the current point, and get the specific channel requested
-		if (ParamHelper.UseRGBAGrid())
+		if (AttributeWriter.UseRGBAGrid())
 		{			
 			FString RGBAHlsl;
 			RGBAHlsl.Appendf(TEXT("void {FunctionName}(float In_UnitX, float In_UnitY, float In_UnitZ, int In_AttributeIndex, out float Out_Val)\n"));
@@ -2567,11 +2568,11 @@ bool UNiagaraDataInterfaceGrid3DCollection::GetFunctionHLSL(const FNiagaraDataIn
 			RGBAHlsl.Appendf(TEXT("        Out_Val = TmpVal.r;\n"));
 			RGBAHlsl.Appendf(TEXT("    }\n"));
 
-			for (int i = 0; i < ParamHelper.TotalChannels - 1; ++i)
+			for (int i = 0; i < AttributeWriter.GetAttributeHelper().GetTotalChannels() - 1; ++i)
 			{
 				RGBAHlsl.Appendf(TEXT("    else if (In_AttributeIndex == %d)\n"), i + 1);
 				RGBAHlsl.Appendf(TEXT("    {\n"));
-				RGBAHlsl.Appendf(TEXT("        Out_Val = TmpVal.%s;\n"), *FGrid3DCollectionAttributeHelper::Channels[i + 1]);
+				RGBAHlsl.Appendf(TEXT("        Out_Val = TmpVal.%s;\n"), *FGrid3DCollectionAttributeHlslWriter::Channels[i + 1]);
 				RGBAHlsl.Appendf(TEXT("    }\n"));
 			}
 
@@ -2601,7 +2602,7 @@ bool UNiagaraDataInterfaceGrid3DCollection::GetFunctionHLSL(const FNiagaraDataIn
 	}
 	else if (FunctionInfo.DefinitionName == CubicSampleGridFunctionName || FunctionInfo.DefinitionName == CubicSamplePreviousGridAtIndexFunctionName)
 	{
-		if (ParamHelper.UseRGBAGrid())
+		if (AttributeWriter.UseRGBAGrid())
 		{
 			//#todo(dmp): maybe fix this at some point or deprecate?
 			UE_LOG(LogNiagara, Error, TEXT("RGBA grids are not supported for this function"));
@@ -2629,7 +2630,7 @@ bool UNiagaraDataInterfaceGrid3DCollection::GetFunctionHLSL(const FNiagaraDataIn
 	}
 	else if (FunctionInfo.DefinitionName == SamplePreviousFullGridFunctionName)
 	{
-		if (ParamHelper.UseRGBAGrid())
+		if (AttributeWriter.UseRGBAGrid())
 		{
 			//#todo(dmp): maybe fix this at some point or deprecate?
 			UE_LOG(LogNiagara, Error, TEXT("RGBA grids are not supported for this function"));
@@ -2650,7 +2651,7 @@ bool UNiagaraDataInterfaceGrid3DCollection::GetFunctionHLSL(const FNiagaraDataIn
 	}
 	else if (FunctionInfo.DefinitionName == CubicSamplePreviousFullGridFunctionName)
 	{
-		if (ParamHelper.UseRGBAGrid())
+		if (AttributeWriter.UseRGBAGrid())
 		{
 			//#todo(dmp): maybe fix this at some point or deprecate?
 			UE_LOG(LogNiagara, Error, TEXT("RGBA grids are not supported for this function"));
@@ -2670,9 +2671,9 @@ bool UNiagaraDataInterfaceGrid3DCollection::GetFunctionHLSL(const FNiagaraDataIn
 	}
 	else if (FunctionInfo.DefinitionName == SetVector4ValueAtIndexFunctionName)
 	{
-		if (ParamHelper.UseRGBAGrid())
+		if (AttributeWriter.UseRGBAGrid())
 		{
-			ParamHelper.WriteSetAtIndexHLSL(FGrid3DCollectionAttributeHelper::AttributeRetrievalMode::RGBAGrid, FunctionInfo, 4, OutHLSL);
+			AttributeWriter.WriteSetAtIndexHLSL(EGridAttributeRetrievalMode::RGBAGrid, FunctionInfo, 4, OutHLSL);
 		}
 		else
 		{
@@ -2682,9 +2683,9 @@ bool UNiagaraDataInterfaceGrid3DCollection::GetFunctionHLSL(const FNiagaraDataIn
 	}
 	else if (FunctionInfo.DefinitionName == SetVector3ValueAtIndexFunctionName)
 	{
-		if (ParamHelper.UseRGBAGrid())
+		if (AttributeWriter.UseRGBAGrid())
 		{
-			ParamHelper.WriteSetAtIndexHLSL(FGrid3DCollectionAttributeHelper::AttributeRetrievalMode::RGBAGrid, FunctionInfo, 3, OutHLSL);
+			AttributeWriter.WriteSetAtIndexHLSL(EGridAttributeRetrievalMode::RGBAGrid, FunctionInfo, 3, OutHLSL);
 		}
 		else
 		{
@@ -2694,9 +2695,9 @@ bool UNiagaraDataInterfaceGrid3DCollection::GetFunctionHLSL(const FNiagaraDataIn
 	}
 	else if (FunctionInfo.DefinitionName == SetVector2ValueAtIndexFunctionName)
 	{
-		if (ParamHelper.UseRGBAGrid())
+		if (AttributeWriter.UseRGBAGrid())
 		{
-			ParamHelper.WriteSetAtIndexHLSL(FGrid3DCollectionAttributeHelper::AttributeRetrievalMode::RGBAGrid, FunctionInfo, 2, OutHLSL);
+			AttributeWriter.WriteSetAtIndexHLSL(EGridAttributeRetrievalMode::RGBAGrid, FunctionInfo, 2, OutHLSL);
 		}
 		else
 		{
@@ -2706,9 +2707,9 @@ bool UNiagaraDataInterfaceGrid3DCollection::GetFunctionHLSL(const FNiagaraDataIn
 	}
 	else if (FunctionInfo.DefinitionName == SetFloatValueAtIndexFunctionName)
 	{
-		if (ParamHelper.UseRGBAGrid())
+		if (AttributeWriter.UseRGBAGrid())
 		{
-			ParamHelper.WriteSetAtIndexHLSL(FGrid3DCollectionAttributeHelper::AttributeRetrievalMode::RGBAGrid, FunctionInfo, 1, OutHLSL);
+			AttributeWriter.WriteSetAtIndexHLSL(EGridAttributeRetrievalMode::RGBAGrid, FunctionInfo, 1, OutHLSL);
 		}
 		else
 		{
@@ -2718,9 +2719,9 @@ bool UNiagaraDataInterfaceGrid3DCollection::GetFunctionHLSL(const FNiagaraDataIn
 	}
 	else if (FunctionInfo.DefinitionName == GetPreviousVector4ValueAtIndexFunctionName)
 	{
-		if (ParamHelper.UseRGBAGrid())
+		if (AttributeWriter.UseRGBAGrid())
 		{
-			ParamHelper.WriteGetAtIndexHLSL(FGrid3DCollectionAttributeHelper::AttributeRetrievalMode::RGBAGrid, FunctionInfo, 4, OutHLSL);
+			AttributeWriter.WriteGetAtIndexHLSL(EGridAttributeRetrievalMode::RGBAGrid, FunctionInfo, 4, OutHLSL);
 		}
 		else
 		{
@@ -2730,9 +2731,9 @@ bool UNiagaraDataInterfaceGrid3DCollection::GetFunctionHLSL(const FNiagaraDataIn
 	}
 	else if (FunctionInfo.DefinitionName == GetPreviousVector3ValueAtIndexFunctionName)
 	{
-		if (ParamHelper.UseRGBAGrid())
+		if (AttributeWriter.UseRGBAGrid())
 		{
-			ParamHelper.WriteGetAtIndexHLSL(FGrid3DCollectionAttributeHelper::AttributeRetrievalMode::RGBAGrid, FunctionInfo, 3, OutHLSL);
+			AttributeWriter.WriteGetAtIndexHLSL(EGridAttributeRetrievalMode::RGBAGrid, FunctionInfo, 3, OutHLSL);
 		}
 		else
 		{
@@ -2742,9 +2743,9 @@ bool UNiagaraDataInterfaceGrid3DCollection::GetFunctionHLSL(const FNiagaraDataIn
 	}
 	else if (FunctionInfo.DefinitionName == GetPreviousVector2ValueAtIndexFunctionName)
 	{
-		if (ParamHelper.UseRGBAGrid())
+		if (AttributeWriter.UseRGBAGrid())
 		{
-			ParamHelper.WriteGetAtIndexHLSL(FGrid3DCollectionAttributeHelper::AttributeRetrievalMode::RGBAGrid, FunctionInfo, 2, OutHLSL);
+			AttributeWriter.WriteGetAtIndexHLSL(EGridAttributeRetrievalMode::RGBAGrid, FunctionInfo, 2, OutHLSL);
 		}
 		else
 		{
@@ -2754,9 +2755,9 @@ bool UNiagaraDataInterfaceGrid3DCollection::GetFunctionHLSL(const FNiagaraDataIn
 	}
 	else if (FunctionInfo.DefinitionName == GetPreviousFloatValueAtIndexFunctionName)
 	{
-		if (ParamHelper.UseRGBAGrid())
+		if (AttributeWriter.UseRGBAGrid())
 		{
-			ParamHelper.WriteGetAtIndexHLSL(FGrid3DCollectionAttributeHelper::AttributeRetrievalMode::RGBAGrid, FunctionInfo, 1, OutHLSL);
+			AttributeWriter.WriteGetAtIndexHLSL(EGridAttributeRetrievalMode::RGBAGrid, FunctionInfo, 1, OutHLSL);
 		}
 		else
 		{
@@ -2766,9 +2767,9 @@ bool UNiagaraDataInterfaceGrid3DCollection::GetFunctionHLSL(const FNiagaraDataIn
 	}
 	else if (FunctionInfo.DefinitionName == SamplePreviousGridVector4AtIndexFunctionName)
 	{
-		if (ParamHelper.UseRGBAGrid())
+		if (AttributeWriter.UseRGBAGrid())
 		{
-			ParamHelper.WriteSampleAtIndexHLSL(FGrid3DCollectionAttributeHelper::AttributeRetrievalMode::RGBAGrid, FunctionInfo, 4, false, OutHLSL);
+			AttributeWriter.WriteSampleAtIndexHLSL(EGridAttributeRetrievalMode::RGBAGrid, FunctionInfo, 4, false, OutHLSL);
 		}
 		else
 		{
@@ -2778,9 +2779,9 @@ bool UNiagaraDataInterfaceGrid3DCollection::GetFunctionHLSL(const FNiagaraDataIn
 	}
 	else if (FunctionInfo.DefinitionName == SamplePreviousGridVector3AtIndexFunctionName)
 	{
-		if (ParamHelper.UseRGBAGrid())
+		if (AttributeWriter.UseRGBAGrid())
 		{
-			ParamHelper.WriteSampleAtIndexHLSL(FGrid3DCollectionAttributeHelper::AttributeRetrievalMode::RGBAGrid, FunctionInfo, 3, false, OutHLSL);
+			AttributeWriter.WriteSampleAtIndexHLSL(EGridAttributeRetrievalMode::RGBAGrid, FunctionInfo, 3, false, OutHLSL);
 		}
 		else
 		{
@@ -2790,9 +2791,9 @@ bool UNiagaraDataInterfaceGrid3DCollection::GetFunctionHLSL(const FNiagaraDataIn
 	}
 	else if (FunctionInfo.DefinitionName == SamplePreviousGridVector2AtIndexFunctionName)
 	{
-		if (ParamHelper.UseRGBAGrid())
+		if (AttributeWriter.UseRGBAGrid())
 		{
-			ParamHelper.WriteSampleAtIndexHLSL(FGrid3DCollectionAttributeHelper::AttributeRetrievalMode::RGBAGrid, FunctionInfo, 2, false, OutHLSL);
+			AttributeWriter.WriteSampleAtIndexHLSL(EGridAttributeRetrievalMode::RGBAGrid, FunctionInfo, 2, false, OutHLSL);
 		}
 		else
 		{
@@ -2802,9 +2803,9 @@ bool UNiagaraDataInterfaceGrid3DCollection::GetFunctionHLSL(const FNiagaraDataIn
 	}
 	else if (FunctionInfo.DefinitionName == SamplePreviousGridFloatAtIndexFunctionName)
 	{
-		if (ParamHelper.UseRGBAGrid())
+		if (AttributeWriter.UseRGBAGrid())
 		{
-			ParamHelper.WriteSampleAtIndexHLSL(FGrid3DCollectionAttributeHelper::AttributeRetrievalMode::RGBAGrid, FunctionInfo, 1, false, OutHLSL);
+			AttributeWriter.WriteSampleAtIndexHLSL(EGridAttributeRetrievalMode::RGBAGrid, FunctionInfo, 1, false, OutHLSL);
 		}
 		else
 		{
@@ -2814,9 +2815,9 @@ bool UNiagaraDataInterfaceGrid3DCollection::GetFunctionHLSL(const FNiagaraDataIn
 	}
 	else if (FunctionInfo.DefinitionName == CubicSamplePreviousGridVector4AtIndexFunctionName)
 	{
-		if (ParamHelper.UseRGBAGrid())
+		if (AttributeWriter.UseRGBAGrid())
 		{
-			ParamHelper.WriteSampleAtIndexHLSL(FGrid3DCollectionAttributeHelper::AttributeRetrievalMode::RGBAGrid, FunctionInfo, 4, true, OutHLSL);
+			AttributeWriter.WriteSampleAtIndexHLSL(EGridAttributeRetrievalMode::RGBAGrid, FunctionInfo, 4, true, OutHLSL);
 		}
 		else
 		{
@@ -2826,9 +2827,9 @@ bool UNiagaraDataInterfaceGrid3DCollection::GetFunctionHLSL(const FNiagaraDataIn
 	}
 	else if (FunctionInfo.DefinitionName == CubicSamplePreviousGridVector3AtIndexFunctionName)
 	{
-		if (ParamHelper.UseRGBAGrid())
+		if (AttributeWriter.UseRGBAGrid())
 		{
-			ParamHelper.WriteSampleAtIndexHLSL(FGrid3DCollectionAttributeHelper::AttributeRetrievalMode::RGBAGrid, FunctionInfo, 3, true, OutHLSL);
+			AttributeWriter.WriteSampleAtIndexHLSL(EGridAttributeRetrievalMode::RGBAGrid, FunctionInfo, 3, true, OutHLSL);
 		}
 		else
 		{
@@ -2838,9 +2839,9 @@ bool UNiagaraDataInterfaceGrid3DCollection::GetFunctionHLSL(const FNiagaraDataIn
 	}
 	else if (FunctionInfo.DefinitionName == CubicSamplePreviousGridVector2AtIndexFunctionName)
 	{
-		if (ParamHelper.UseRGBAGrid())
+		if (AttributeWriter.UseRGBAGrid())
 		{
-			ParamHelper.WriteSampleAtIndexHLSL(FGrid3DCollectionAttributeHelper::AttributeRetrievalMode::RGBAGrid, FunctionInfo, 2, true, OutHLSL);
+			AttributeWriter.WriteSampleAtIndexHLSL(EGridAttributeRetrievalMode::RGBAGrid, FunctionInfo, 2, true, OutHLSL);
 		}
 		else
 		{
@@ -2850,9 +2851,9 @@ bool UNiagaraDataInterfaceGrid3DCollection::GetFunctionHLSL(const FNiagaraDataIn
 	}
 	else if (FunctionInfo.DefinitionName == CubicSamplePreviousGridFloatAtIndexFunctionName)
 	{
-		if (ParamHelper.UseRGBAGrid())
+		if (AttributeWriter.UseRGBAGrid())
 		{
-			ParamHelper.WriteSampleAtIndexHLSL(FGrid3DCollectionAttributeHelper::AttributeRetrievalMode::RGBAGrid, FunctionInfo, 1, true, OutHLSL);
+			AttributeWriter.WriteSampleAtIndexHLSL(EGridAttributeRetrievalMode::RGBAGrid, FunctionInfo, 1, true, OutHLSL);
 		}
 		else
 		{
@@ -3002,7 +3003,7 @@ FNiagaraDataInterfaceParametersCS* UNiagaraDataInterfaceGrid3DCollection::Create
 	{
 		const FNiagaraDataInterfaceGeneratedFunction& Func = ParameterInfo.GeneratedFunctions[FuncIdx];
 
-		const FName* AttributeName = Func.FindSpecifierValue(NAME_Attribute);
+		const FName* AttributeName = Func.FindSpecifierValue(UNiagaraDataInterfaceRWBase::NAME_Attribute);
 		if (AttributeName != nullptr)
 		{
 			int32 ComponentCount = UNiagaraDataInterfaceGrid3DCollection::GetComponentCountFromFuncName(Func.DefinitionName);
@@ -3091,7 +3092,7 @@ bool UNiagaraDataInterfaceGrid3DCollection::GenerateIterationSourceNamespaceRead
 }
 bool UNiagaraDataInterfaceGrid3DCollection::GenerateIterationSourceNamespaceWriteAttributesHLSL(FNiagaraDataInterfaceGPUParamInfo& DIInstanceInfo, const FNiagaraVariable& IterationSourceVar, TConstArrayView<FNiagaraVariable> InArguments, TConstArrayView<FNiagaraVariable> InAttributes, TConstArrayView<FString> InAttributeHLSLNames, TConstArrayView<FNiagaraVariable> InAllAttributes, bool bSpawnOnly, bool bPartialWrites, TArray<FText>& OutErrors, FString& OutHLSL) const
 {
-	FGrid3DCollectionAttributeHelper ParamHelper(DIInstanceInfo);
+	FGrid3DCollectionAttributeHlslWriter AttributeWriter(DIInstanceInfo);
 
 	FString DIVarName;
 	OutHLSL += TEXT("\t//Generated by UNiagaraDataInterfaceGrid3DCollection::GenerateIterationSourceNamespaceWriteAttributesHLSL\n");
@@ -3113,7 +3114,7 @@ bool UNiagaraDataInterfaceGrid3DCollection::GenerateIterationSourceNamespaceWrit
 		int ComponentCount = GetComponentCountFromFuncName(TypeDefinitionToSetFunctionName(InAllAttributes[i].GetType()));
 		TotalChannels += ComponentCount;
 	}
-	bool UseRGBATexture = FGrid3DCollectionAttributeHelper::ShouldUseRGBAGrid(TotalChannels, InAttributes.Num());
+	bool UseRGBATexture = FGrid3DCollectionAttributeHlslWriter::ShouldUseRGBAGrid(TotalChannels, InAttributes.Num());
 
 	TArray<FString> RootArray;
 	IterationSourceVar.GetName().ToString().ParseIntoArray(RootArray, TEXT("."));
@@ -3362,7 +3363,7 @@ bool UNiagaraDataInterfaceGrid3DCollection::InitPerInstanceData(void* PerInstanc
 	}
 
 	// if we have 4 or less attributes, pack data into rgba textures rather than a single large float texture with attributes tiled spatially
-	bool UseRGBATexture = FGrid3DCollectionAttributeHelper::ShouldUseRGBAGrid(NumAttribChannelsFound, InstanceData->Vars.Num());
+	bool UseRGBATexture = FGrid3DCollectionAttributeHlslWriter::ShouldUseRGBAGrid(NumAttribChannelsFound, InstanceData->Vars.Num());
 
 	if (UseRGBATexture)
 	{
@@ -4316,7 +4317,7 @@ void UNiagaraDataInterfaceGrid3DCollection::CollectAttributesForScript(UNiagaraS
 			TConstArrayView<FNiagaraDataInterfaceGPUParamInfo> ParamInfoArray = Script->GetDataInterfaceGPUParamInfos();
 			for (const FNiagaraDataInterfaceGeneratedFunction& Func : ParamInfoArray[*IndexOfDataInterface].GeneratedFunctions)
 			{
-				if (const FName* AttributeName = Func.FindSpecifierValue(NAME_Attribute))
+				if (const FName* AttributeName = Func.FindSpecifierValue(UNiagaraDataInterfaceRWBase::NAME_Attribute))
 				{
 					FNiagaraVariableBase NewVar(UNiagaraDataInterfaceGrid3DCollection::GetValueTypeFromFuncName(Func.DefinitionName), *AttributeName);
 					if (UNiagaraDataInterfaceGrid3DCollection::CanCreateVarFromFuncName(Func.DefinitionName))
diff --git a/Engine/Plugins/FX/Niagara/Source/Niagara/Private/NiagaraDataInterfaceGrid3DCollectionReader.cpp b/Engine/Plugins/FX/Niagara/Source/Niagara/Private/NiagaraDataInterfaceGrid3DCollectionReader.cpp
--- a/Engine/Plugins/FX/Niagara/Source/Niagara/Private/NiagaraDataInterfaceGrid3DCollectionReader.cpp
+++ b/Engine/Plugins/FX/Niagara/Source/Niagara/Private/NiagaraDataInterfaceGrid3DCollectionReader.cpp
@@ -114,7 +114,7 @@ bool UNiagaraDataInterfaceGrid3DCollectionReader::InitPerInstanceData(void* PerI
 
 			// #todo(dmp): slight hack here - these aren't always the same, but we generally can't use unnamed attrs very well
 			NumAttribChannelsFound = NumNamedAttribChannelsFound;
-			bool UseRGBATexture = FGrid3DCollectionAttributeHelper::ShouldUseRGBAGrid(NumAttribChannelsFound, Vars.Num());
+			bool UseRGBATexture = FGrid3DCollectionAttributeHlslWriter::ShouldUseRGBAGrid(NumAttribChannelsFound, Vars.Num());
 
 			FGrid3DCollectionRWInstanceData_GameThread* Grid3DInstanceData = GridInterface->GetSystemInstancesToProxyData_GT().FindRef(SystemInstance->GetId());
 			InstanceData->OtherDI = GridInterface;
@@ -132,9 +132,7 @@ bool UNiagaraDataInterfaceGrid3DCollectionReader::InitPerInstanceData(void* PerI
 			ENQUEUE_RENDER_COMMAND(FUpdateData)(
 				[RT_Proxy, InstanceID = SystemInstance->GetId(), RT_ProxyToUse = ProxyToUse](FRHICommandListImmediate& RHICmdList)
 			{
-				check(!RT_Proxy->SystemInstancesToProxyData_RT.Contains(InstanceID));
-				//check(RT_ProxyToUse->SystemInstancesToProxyData_RT.Contains(InstanceID));
-				//FGrid3DCollectionRWInstanceData_RenderThread* DataToRead = RT_ProxyToUse->SystemInstancesToProxyData_RT.Find(InstanceID);
+				check(!RT_Proxy->SystemInstancesToProxyData_RT.Contains(InstanceID));				
 
 				FGrid3DCollectionRWInstanceData_RenderThread* TargetData = &RT_Proxy->SystemInstancesToProxyData_RT.Add(InstanceID);
 				TargetData->OtherProxy = RT_ProxyToUse;
@@ -230,7 +228,7 @@ bool UNiagaraDataInterfaceGrid3DCollectionReader::PerInstanceTickPostSimulate(vo
 
 		// #todo(dmp): slight hack here - these aren't always the same, but we generally can't use unnamed attrs very well
 		NumAttribChannelsFound = NumNamedAttribChannelsFound;
-		bool UseRGBATexture = FGrid3DCollectionAttributeHelper::ShouldUseRGBAGrid(NumAttribChannelsFound, Vars.Num());
+		bool UseRGBATexture = FGrid3DCollectionAttributeHlslWriter::ShouldUseRGBAGrid(NumAttribChannelsFound, Vars.Num());
 
 		if (InstanceData->OtherInstanceData == nullptr)
 		{
diff --git a/Engine/Plugins/FX/Niagara/Source/Niagara/Private/NiagaraDataInterfaceGrid3DCollectionUtils.cpp b/Engine/Plugins/FX/Niagara/Source/Niagara/Private/NiagaraDataInterfaceGrid3DCollectionUtils.cpp
--- a/Engine/Plugins/FX/Niagara/Source/Niagara/Private/NiagaraDataInterfaceGrid3DCollectionUtils.cpp
+++ b/Engine/Plugins/FX/Niagara/Source/Niagara/Private/NiagaraDataInterfaceGrid3DCollectionUtils.cpp
@@ -13,148 +13,60 @@ static FAutoConsoleVariableRef CVarNiagaraGrid3DUseRGBAGrid(
 	ECVF_Default
 );
 
-TArray<FString> FGrid3DCollectionAttributeHelper::Channels = { "r", "g", "b", "a" };
+TArray<FString> FGrid3DCollectionAttributeHlslWriter::Channels = { "r", "g", "b", "a" };
 
 // static member function
-bool FGrid3DCollectionAttributeHelper::ShouldUseRGBAGrid(const int TotalChannels, const int TotalNumAttributes)
+bool FGrid3DCollectionAttributeHlslWriter::ShouldUseRGBAGrid(const int TotalChannels, const int TotalNumAttributes)
 {
 	return TotalNumAttributes == 1 && TotalChannels <= 4 && GNiagaraGrid3DUseRGBAGrid != 0;
 }
 
-bool FGrid3DCollectionAttributeHelper::SupportsRGBAGrid()
+bool FGrid3DCollectionAttributeHlslWriter::SupportsRGBAGrid()
 {
 	return GNiagaraGrid3DUseRGBAGrid == 1;
 }
 
-FGrid3DCollectionAttributeHelper::FGrid3DCollectionAttributeHelper(const FNiagaraDataInterfaceGPUParamInfo& InParamInfo, TArray<FText>* OutWarnings /*= nullptr*/) : ParamInfo(InParamInfo)
+FGrid3DCollectionAttributeHlslWriter::FGrid3DCollectionAttributeHlslWriter(const FNiagaraDataInterfaceGPUParamInfo& InParamInfo, TArray<FText>* OutWarnings /*= nullptr*/) :
+	ParamInfo(InParamInfo), AttributeHelper(InParamInfo)
 {
-	AttributeInfos.Reserve(ParamInfo.GeneratedFunctions.Num());
-	TotalChannels = 0;
-
-	for (const FNiagaraDataInterfaceGeneratedFunction& Function : ParamInfo.GeneratedFunctions)
-	{
-		const FName* AttributeName = Function.FindSpecifierValue(UNiagaraDataInterfaceGrid3DCollection::NAME_Attribute);
-		if (AttributeName == nullptr)
-		{
-			continue;
+	AttributeHelper.Init<UNiagaraDataInterfaceGrid3DCollection>(OutWarnings);
 		}
 
-		if (const FAttributeInfo* ExistingAttribute = FindAttributeInfo(*AttributeName))
+bool FGrid3DCollectionAttributeHlslWriter::UseRGBAGrid()
 		{
-			if (OutWarnings != nullptr)
-			{
-				FNiagaraTypeDefinition AttributeTypeDef = UNiagaraDataInterfaceGrid3DCollection::GetValueTypeFromFuncName(Function.DefinitionName);
-				if (ExistingAttribute->TypeDef != AttributeTypeDef)
-				{
-					OutWarnings->Emplace(FText::Format(LOCTEXT("BadType", "Same name, different types! {0} vs {1}, Attribute {2}"), AttributeTypeDef.GetNameText(), ExistingAttribute->TypeDef.GetNameText(), FText::FromName(ExistingAttribute->Name)));
-				}
-			}
-			continue;
-		}
-
-		FAttributeInfo& AttributeInfo = AttributeInfos.AddDefaulted_GetRef();
-		AttributeInfo.Name = *AttributeName;
-		AttributeInfo.TypeDef = UNiagaraDataInterfaceGrid3DCollection::GetValueTypeFromFuncName(Function.DefinitionName);
-		AttributeInfo.NumChannels = UNiagaraDataInterfaceGrid3DCollection::GetComponentCountFromFuncName(Function.DefinitionName);
-		AttributeInfo.ChannelOffset = TotalChannels;
-		AttributeInfo.AttributeIndex = AttributeInfos.Num() - 1;
-		TotalChannels += AttributeInfo.NumChannels;
-	}
-}
-
-bool FGrid3DCollectionAttributeHelper::UseRGBAGrid()
-{
-	return ShouldUseRGBAGrid(TotalChannels, AttributeInfos.Num());
-}
-
-const FGrid3DCollectionAttributeHelper::FGrid3DCollectionAttributeHelper::FAttributeInfo* FGrid3DCollectionAttributeHelper::FindAttributeInfo(FName AttributeName) const
-{
-	return AttributeInfos.FindByPredicate([AttributeName](const FAttributeInfo& Info) { return Info.Name == AttributeName; });
+	return ShouldUseRGBAGrid(AttributeHelper.GetTotalChannels(), AttributeHelper.GetNumAttributes());
 }
 
 #if WITH_EDITORONLY_DATA
-FString FGrid3DCollectionAttributeHelper::GetPerAttributePixelOffset(const TCHAR* DataInterfaceHLSLSymbol)
+FString FGrid3DCollectionAttributeHlslWriter::GetPerAttributePixelOffset(const TCHAR* DataInterfaceHLSLSymbol)
 {
 	return FString::Printf(TEXT("int3(%s_%s[(AttributeIndex * 2) + 0].xyz)"), DataInterfaceHLSLSymbol, *UNiagaraDataInterfaceGrid3DCollection::PerAttributeDataName);
 }
 
-FString FGrid3DCollectionAttributeHelper::GetPerAttributePixelOffset() const
+FString FGrid3DCollectionAttributeHlslWriter::GetPerAttributePixelOffset() const
 {
 	return GetPerAttributePixelOffset(*ParamInfo.DataInterfaceHLSLSymbol);
 }
 
-FString FGrid3DCollectionAttributeHelper::GetPerAttributeUVWOffset(const TCHAR* DataInterfaceHLSLSymbol)
+FString FGrid3DCollectionAttributeHlslWriter::GetPerAttributeUVWOffset(const TCHAR* DataInterfaceHLSLSymbol)
 {
 	return FString::Printf(TEXT("%s%s[(AttributeIndex * 2) + 1].xyz"), *UNiagaraDataInterfaceGrid3DCollection::PerAttributeDataName, DataInterfaceHLSLSymbol);
 }
 
-FString FGrid3DCollectionAttributeHelper::GetPerAttributeUVWOffset() const
+FString FGrid3DCollectionAttributeHlslWriter::GetPerAttributeUVWOffset() const
 {
 	return GetPerAttributeUVWOffset(*ParamInfo.DataInterfaceHLSLSymbol);
 }
 
-FString FGrid3DCollectionAttributeHelper::GetAttributeIndex(bool bUseAttributeIndirection, const FAttributeInfo* AttributeInfo, int Channel /*= 0*/) const
-{
-	check(AttributeInfo);
-	if (bUseAttributeIndirection)
-	{
-		return FString::Printf(TEXT("int(%s%s[%d].w + %d)"), *UNiagaraDataInterfaceGrid3DCollection::PerAttributeDataName, *ParamInfo.DataInterfaceHLSLSymbol, AttributeInfo->AttributeIndex, Channel);
-	}
-	else
-	{
-		return FString::Printf(TEXT("%d"), AttributeInfo->ChannelOffset + Channel);
-	}
-}
-
-FString FGrid3DCollectionAttributeHelper::GetGridChannelString()
-{
-	FString NumChannelsString = "";
-	if (TotalChannels > 1)
-	{
-		NumChannelsString = FString::FromInt(TotalChannels);
-	}
-	return NumChannelsString;
-}
-
-void FGrid3DCollectionAttributeHelper::GetChannelStrings(int AttributeIndex, int AttributeNumChannels, FString& NumChannelsString, FString& AttrGridChannels)
-{
-	NumChannelsString = "";
-	if (AttributeNumChannels > 1)
-	{
-		NumChannelsString = FString::FromInt(AttributeNumChannels);
-	}
-
-	AttrGridChannels = FGrid3DCollectionAttributeHelper::Channels[AttributeIndex];
-	for (int i = 1; i < AttributeNumChannels; ++i)
-	{
-		AttrGridChannels += FGrid3DCollectionAttributeHelper::Channels[AttributeIndex + i];
-	}
-}
-
-void FGrid3DCollectionAttributeHelper::GetChannelStrings(const FAttributeInfo* AttributeInfo, FString& NumChannelsString, FString& AttrGridChannels)
-{
-	NumChannelsString = "";
-	if (AttributeInfo->NumChannels > 1)
-	{
-		NumChannelsString = FString::FromInt(AttributeInfo->NumChannels);
-	}
-
-	AttrGridChannels = FGrid3DCollectionAttributeHelper::Channels[AttributeInfo->AttributeIndex];
-	for (int i = 1; i < AttributeInfo->NumChannels; ++i)
-	{
-		AttrGridChannels += FGrid3DCollectionAttributeHelper::Channels[AttributeInfo->AttributeIndex + i];
-	}
-}
-
-bool FGrid3DCollectionAttributeHelper::WriteGetHLSL(AttributeRetrievalMode AttributeStorage, const FNiagaraDataInterfaceGeneratedFunction& FunctionInfo, FString& OutHLSL)
+bool FGrid3DCollectionAttributeHlslWriter::WriteGetHLSL(EGridAttributeRetrievalMode AttributeStorage, const FNiagaraDataInterfaceGeneratedFunction& FunctionInfo, FString& OutHLSL)
 {
-	const FName* AttributeName = FunctionInfo.FindSpecifierValue(UNiagaraDataInterfaceGrid3DCollection::NAME_Attribute);
+	const FName* AttributeName = FunctionInfo.FindSpecifierValue(UNiagaraDataInterfaceRWBase::NAME_Attribute);
 	if (AttributeName == nullptr)
 	{
 		return false;
 	}
 
-	const FAttributeInfo* AttributeInfo = FindAttributeInfo(*AttributeName);
+	const FNiagaraDataInterfaceRWAttributeHelper::FAttributeInfo* AttributeInfo = AttributeHelper.FindAttributeInfo(*AttributeName);
 	if (AttributeInfo == nullptr)
 	{
 		return false;
@@ -162,13 +74,13 @@ bool FGrid3DCollectionAttributeHelper::WriteGetHLSL(AttributeRetrievalMode Attri
 
 	const FString GridNameHLSL = ParamInfo.DataInterfaceHLSLSymbol + UNiagaraDataInterfaceGrid3DCollection::GridName;
 
-	if (AttributeStorage == AttributeRetrievalMode::RGBAGrid)
+	if (AttributeStorage == EGridAttributeRetrievalMode::RGBAGrid)
 	{
 
 		FString NumChannelsString;
 		FString AttrGridChannels;
 
-		GetChannelStrings(AttributeInfo, NumChannelsString, AttrGridChannels);
+		AttributeHelper.GetChannelStrings(AttributeInfo, NumChannelsString, AttrGridChannels);
 
 
 		OutHLSL.Appendf(TEXT("void %s(int IndexX, int IndexY, int IndexZ, out float%s Value)\n"), *FunctionInfo.InstanceName, *NumChannelsString);
@@ -182,7 +94,7 @@ bool FGrid3DCollectionAttributeHelper::WriteGetHLSL(AttributeRetrievalMode Attri
 		{
 			OutHLSL.Appendf(TEXT("void %s(int IndexX, int IndexY, int IndexZ, out float Value)\n"), *FunctionInfo.InstanceName);
 			OutHLSL.Appendf(TEXT("{\n"));
-			OutHLSL.Appendf(TEXT("	int z %s;\n"), *GetAttributeIndex(AttributeStorage == AttributeRetrievalMode::Indirection, AttributeInfo));
+			OutHLSL.Appendf(TEXT("	int z %s;\n"), *AttributeHelper.GetAttributeIndex(AttributeInfo));
 			OutHLSL.Appendf(TEXT("	int3 PixelOffset = int3(IndexX, IndexY, IndexZ) + %s;\n"), *GetPerAttributePixelOffset());
 			OutHLSL.Appendf(TEXT("	Value = %s.Load(int4(PixelOffset, 0));\n"), *GridNameHLSL);
 			OutHLSL.Appendf(TEXT("}\n"));
@@ -194,7 +106,7 @@ bool FGrid3DCollectionAttributeHelper::WriteGetHLSL(AttributeRetrievalMode Attri
 			for (int32 i = 0; i < AttributeInfo->NumChannels; ++i)
 			{
 				OutHLSL.Appendf(TEXT("	{\n"));
-				OutHLSL.Appendf(TEXT("		int AttributeIndex = %s;\n"), *GetAttributeIndex(AttributeStorage == AttributeRetrievalMode::Indirection, AttributeInfo, i));
+				OutHLSL.Appendf(TEXT("		int AttributeIndex = %s;\n"), *AttributeHelper.GetAttributeIndex(AttributeInfo, i));
 				OutHLSL.Appendf(TEXT("		int3 PixelOffset = int3(IndexX, IndexY, IndexZ) + %s;\n"), *GetPerAttributePixelOffset());
 				OutHLSL.Appendf(TEXT("		Value[%d] = %s.Load(int4(PixelOffset, 0));\n"), i, *GridNameHLSL);
 				OutHLSL.Appendf(TEXT("	}\n"));
@@ -206,16 +118,16 @@ bool FGrid3DCollectionAttributeHelper::WriteGetHLSL(AttributeRetrievalMode Attri
 	return true;
 }
 
-bool FGrid3DCollectionAttributeHelper::WriteGetAtIndexHLSL(AttributeRetrievalMode AttributeStorage, const FNiagaraDataInterfaceGeneratedFunction& FunctionInfo, int NumChannels, FString& OutHLSL)
+bool FGrid3DCollectionAttributeHlslWriter::WriteGetAtIndexHLSL(EGridAttributeRetrievalMode AttributeStorage, const FNiagaraDataInterfaceGeneratedFunction& FunctionInfo, int NumChannels, FString& OutHLSL)
 {
 	FString NumChannelsString;
 	FString AttrGridChannels;
 
 	const FString GridNameHLSL = ParamInfo.DataInterfaceHLSLSymbol + UNiagaraDataInterfaceGrid3DCollection::GridName;
 
-	GetChannelStrings(0, NumChannels, NumChannelsString, AttrGridChannels);
+	AttributeHelper.GetChannelStrings(0, NumChannels, NumChannelsString, AttrGridChannels);
 
-	if (AttributeInfos.Num() > 1)
+	if (AttributeHelper.GetNumAttributes() > 1)
 	{
 		//#todo(dmp): fill this in
 		return false;
@@ -231,20 +143,20 @@ bool FGrid3DCollectionAttributeHelper::WriteGetAtIndexHLSL(AttributeRetrievalMod
 	return true;
 }
 
-bool FGrid3DCollectionAttributeHelper::WriteSetAtIndexHLSL(AttributeRetrievalMode AttributeStorage, const FNiagaraDataInterfaceGeneratedFunction& FunctionInfo, int NumChannels, FString& OutHLSL)
+bool FGrid3DCollectionAttributeHlslWriter::WriteSetAtIndexHLSL(EGridAttributeRetrievalMode AttributeStorage, const FNiagaraDataInterfaceGeneratedFunction& FunctionInfo, int NumChannels, FString& OutHLSL)
 {
 	FString NumChannelsString;
 	FString AttrGridChannels;
 
 	const FString OutputGridNameHLSL = ParamInfo.DataInterfaceHLSLSymbol + UNiagaraDataInterfaceGrid3DCollection::OutputGridName;
 
-	GetChannelStrings(0, NumChannels, NumChannelsString, AttrGridChannels);
+	AttributeHelper.GetChannelStrings(0, NumChannels, NumChannelsString, AttrGridChannels);
 
 	OutHLSL.Appendf(TEXT("void %s(int IndexX, int IndexY, int IndexZ, int AttributeIndex, float%s Value)\n"), *FunctionInfo.InstanceName, *NumChannelsString);
 	OutHLSL.Appendf(TEXT("{\n"));
 
 	// more than 1 attribute, we must read first
-	if (AttributeInfos.Num() > 1)
+	if (AttributeHelper.GetNumAttributes() > 1)
 	{
 		//#todo(dmp): fill this in
 		return false;
@@ -258,7 +170,7 @@ bool FGrid3DCollectionAttributeHelper::WriteSetAtIndexHLSL(AttributeRetrievalMod
 	return true;
 }
 
-bool FGrid3DCollectionAttributeHelper::WriteSampleAtIndexHLSL(AttributeRetrievalMode AttributeStorage, const FNiagaraDataInterfaceGeneratedFunction& FunctionInfo, int NumChannels, bool IsCubic, FString& OutHLSL)
+bool FGrid3DCollectionAttributeHlslWriter::WriteSampleAtIndexHLSL(EGridAttributeRetrievalMode AttributeStorage, const FNiagaraDataInterfaceGeneratedFunction& FunctionInfo, int NumChannels, bool IsCubic, FString& OutHLSL)
 {
 	FString NumChannelsString;
 	FString AttrGridChannels;
@@ -266,9 +178,9 @@ bool FGrid3DCollectionAttributeHelper::WriteSampleAtIndexHLSL(AttributeRetrieval
 	const FString GridNameHLSL = ParamInfo.DataInterfaceHLSLSymbol + UNiagaraDataInterfaceGrid3DCollection::GridName;
 	const FString SamplerNameHLSL = ParamInfo.DataInterfaceHLSLSymbol + UNiagaraDataInterfaceGrid3DCollection::SamplerName;
 
-	GetChannelStrings(0, NumChannels, NumChannelsString, AttrGridChannels);
+	AttributeHelper.GetChannelStrings(0, NumChannels, NumChannelsString, AttrGridChannels);
 
-	if (AttributeInfos.Num() > 1)
+	if (AttributeHelper.GetNumAttributes() > 1)
 	{
 		//#todo(dmp): fill this in
 		return false;
@@ -292,15 +204,15 @@ bool FGrid3DCollectionAttributeHelper::WriteSampleAtIndexHLSL(AttributeRetrieval
 	return true;
 }
 
-bool FGrid3DCollectionAttributeHelper::WriteSetHLSL(AttributeRetrievalMode AttributeStorage, const FNiagaraDataInterfaceGeneratedFunction& FunctionInfo, FString& OutHLSL)
+bool FGrid3DCollectionAttributeHlslWriter::WriteSetHLSL(EGridAttributeRetrievalMode AttributeStorage, const FNiagaraDataInterfaceGeneratedFunction& FunctionInfo, FString& OutHLSL)
 {
-	const FName* AttributeName = FunctionInfo.FindSpecifierValue(UNiagaraDataInterfaceGrid3DCollection::NAME_Attribute);
+	const FName* AttributeName = FunctionInfo.FindSpecifierValue(UNiagaraDataInterfaceRWBase::NAME_Attribute);
 	if (AttributeName == nullptr)
 	{
 		return false;
 	}
 
-	const FAttributeInfo* AttributeInfo = FindAttributeInfo(*AttributeName);
+	const FNiagaraDataInterfaceRWAttributeHelper::FAttributeInfo* AttributeInfo = AttributeHelper.FindAttributeInfo(*AttributeName);
 	if (AttributeInfo == nullptr)
 	{
 		return false;
@@ -308,20 +220,20 @@ bool FGrid3DCollectionAttributeHelper::WriteSetHLSL(AttributeRetrievalMode Attri
 
 	const FString OutputGridNameHLSL = ParamInfo.DataInterfaceHLSLSymbol + UNiagaraDataInterfaceGrid3DCollection::OutputGridName;
 
-	if (AttributeStorage == AttributeRetrievalMode::RGBAGrid)
+	if (AttributeStorage == EGridAttributeRetrievalMode::RGBAGrid)
 	{
 		FString NumChannelsString;
 		FString AttrGridChannels;
 
-		GetChannelStrings(AttributeInfo, NumChannelsString, AttrGridChannels);
+		AttributeHelper.GetChannelStrings(AttributeInfo, NumChannelsString, AttrGridChannels);
 
 		OutHLSL.Appendf(TEXT("void %s(int IndexX, int IndexY, int IndexZ, float%s Value)\n"), *FunctionInfo.InstanceName, *NumChannelsString);
 		OutHLSL.Appendf(TEXT("{\n"));
 
-		// more than 1 attrbiute, we must read first
-		if (AttributeInfos.Num() > 1)
+		// more than 1 attribute, we must read first
+		if (AttributeHelper.GetNumAttributes() > 1)
 		{
-			FString GridNumChannels = FString::FromInt(TotalChannels);
+			FString GridNumChannels = FString::FromInt(AttributeHelper.GetTotalChannels());
 
 			OutHLSL.Appendf(TEXT("	float%s TmpValue = %s.Load(int4(IndexX, IndexY, IndexZ, 0));\n"), *GridNumChannels, *OutputGridNameHLSL);
 			OutHLSL.Appendf(TEXT("	TmpValue.%s = Value;\n"), *AttrGridChannels);
@@ -339,7 +251,7 @@ bool FGrid3DCollectionAttributeHelper::WriteSetHLSL(AttributeRetrievalMode Attri
 		{
 			OutHLSL.Appendf(TEXT("void %s(int IndexX, int IndexY, int IndexZ, float Value)\n"), *FunctionInfo.InstanceName);
 			OutHLSL.Appendf(TEXT("{\n"));
-			OutHLSL.Appendf(TEXT("	int AttributeIndex = %s;\n"), *GetAttributeIndex(AttributeStorage == AttributeRetrievalMode::Indirection, AttributeInfo));
+			OutHLSL.Appendf(TEXT("	int AttributeIndex = %s;\n"), *AttributeHelper.GetAttributeIndex(AttributeInfo));
 			OutHLSL.Appendf(TEXT("	int3 PixelOffset = int3(IndexX, IndexY, IndexZ) + %s;\n"), *GetPerAttributePixelOffset());
 			OutHLSL.Appendf(TEXT("	%s[PixelOffset] = Value;\n"), *OutputGridNameHLSL);
 			OutHLSL.Appendf(TEXT("}\n"));
@@ -351,7 +263,7 @@ bool FGrid3DCollectionAttributeHelper::WriteSetHLSL(AttributeRetrievalMode Attri
 			for (int32 i = 0; i < AttributeInfo->NumChannels; ++i)
 			{
 				OutHLSL.Appendf(TEXT("	{\n"));
-				OutHLSL.Appendf(TEXT("		int AttributeIndex = %s;\n"), *GetAttributeIndex(AttributeStorage == AttributeRetrievalMode::Indirection, AttributeInfo, i));
+				OutHLSL.Appendf(TEXT("		int AttributeIndex = %s;\n"), *AttributeHelper.GetAttributeIndex(AttributeInfo, i));
 				OutHLSL.Appendf(TEXT("		int3 PixelOffset = int3(IndexX, IndexY, IndexZ) + %s;\n"), *GetPerAttributePixelOffset());
 				OutHLSL.Appendf(TEXT("		%s[PixelOffset] = Value[%d];\n"), *OutputGridNameHLSL, i);
 				OutHLSL.Appendf(TEXT("	}\n"));
@@ -362,15 +274,15 @@ bool FGrid3DCollectionAttributeHelper::WriteSetHLSL(AttributeRetrievalMode Attri
 	return true;
 }
 
-bool FGrid3DCollectionAttributeHelper::WriteSampleHLSL(AttributeRetrievalMode AttributeStorage, const FNiagaraDataInterfaceGeneratedFunction& FunctionInfo, bool IsCubic, FString& OutHLSL)
+bool FGrid3DCollectionAttributeHlslWriter::WriteSampleHLSL(EGridAttributeRetrievalMode AttributeStorage, const FNiagaraDataInterfaceGeneratedFunction& FunctionInfo, bool IsCubic, FString& OutHLSL)
 {
-	const FName* AttributeName = FunctionInfo.FindSpecifierValue(UNiagaraDataInterfaceGrid3DCollection::NAME_Attribute);
+	const FName* AttributeName = FunctionInfo.FindSpecifierValue(UNiagaraDataInterfaceRWBase::NAME_Attribute);
 	if (AttributeName == nullptr)
 	{
 		return false;
 	}
 
-	const FAttributeInfo* AttributeInfo = FindAttributeInfo(*AttributeName);
+	const FNiagaraDataInterfaceRWAttributeHelper::FAttributeInfo* AttributeInfo = AttributeHelper.FindAttributeInfo(*AttributeName);
 	if (AttributeInfo == nullptr)
 	{
 		return false;
@@ -379,11 +291,11 @@ bool FGrid3DCollectionAttributeHelper::WriteSampleHLSL(AttributeRetrievalMode At
 	const FString GridNameHLSL = ParamInfo.DataInterfaceHLSLSymbol + UNiagaraDataInterfaceGrid3DCollection::GridName;
 	const FString SamplerNameHLSL = ParamInfo.DataInterfaceHLSLSymbol + UNiagaraDataInterfaceGrid3DCollection::SamplerName;
 
-	if (AttributeStorage == AttributeRetrievalMode::RGBAGrid)
+	if (AttributeStorage == EGridAttributeRetrievalMode::RGBAGrid)
 	{
 		FString NumChannelsString;
 		FString AttrGridChannels;
-		GetChannelStrings(AttributeInfo, NumChannelsString, AttrGridChannels);
+		AttributeHelper.GetChannelStrings(AttributeInfo, NumChannelsString, AttrGridChannels);
 
 		OutHLSL.Appendf(TEXT("void %s(float3 Unit, out float%s Value)\n"), *FunctionInfo.InstanceName, *NumChannelsString);
 		OutHLSL.Appendf(TEXT("{\n"));
@@ -406,7 +318,7 @@ bool FGrid3DCollectionAttributeHelper::WriteSampleHLSL(AttributeRetrievalMode At
 			OutHLSL.Appendf(TEXT("void %s(float3 Unit, out float Value)\n"), *FunctionInfo.InstanceName);
 			OutHLSL.Appendf(TEXT("{\n"));
 			OutHLSL.Appendf(TEXT("	float3 TileUVW = clamp(Unit, %s%s, %s%s) * %s%s;\n"), *UNiagaraDataInterfaceGrid3DCollection::UnitClampMinName, *ParamInfo.DataInterfaceHLSLSymbol, *UNiagaraDataInterfaceGrid3DCollection::UnitClampMaxName, *ParamInfo.DataInterfaceHLSLSymbol, *UNiagaraDataInterfaceGrid3DCollection::OneOverNumTilesName, *ParamInfo.DataInterfaceHLSLSymbol);
-			OutHLSL.Appendf(TEXT("	int AttributeIndex = %s;\n"), *GetAttributeIndex(AttributeStorage == AttributeRetrievalMode::Indirection, AttributeInfo));
+			OutHLSL.Appendf(TEXT("	int AttributeIndex = %s;\n"), *AttributeHelper.GetAttributeIndex(AttributeInfo));
 			OutHLSL.Appendf(TEXT("	float3 UVW = TileUVW + %s;\n"), *GetPerAttributeUVWOffset());
 			if (IsCubic)
 			{
@@ -426,7 +338,7 @@ bool FGrid3DCollectionAttributeHelper::WriteSampleHLSL(AttributeRetrievalMode At
 			for (int32 i = 0; i < AttributeInfo->NumChannels; ++i)
 			{
 				OutHLSL.Appendf(TEXT("	{\n"));
-				OutHLSL.Appendf(TEXT("		int AttributeIndex = %s;\n"), *GetAttributeIndex(AttributeStorage == AttributeRetrievalMode::Indirection, AttributeInfo, i));
+				OutHLSL.Appendf(TEXT("		int AttributeIndex = %s;\n"), *AttributeHelper.GetAttributeIndex(AttributeInfo, i));
 				OutHLSL.Appendf(TEXT("		float3 UVW = TileUVW + %s;\n"), *GetPerAttributeUVWOffset());
 				if (IsCubic)
 				{
@@ -445,15 +357,15 @@ bool FGrid3DCollectionAttributeHelper::WriteSampleHLSL(AttributeRetrievalMode At
 	return true;
 }
 
-bool FGrid3DCollectionAttributeHelper::WriteAttributeGetIndexHLSL(AttributeRetrievalMode AttributeStorage, const FNiagaraDataInterfaceGeneratedFunction& FunctionInfo, FString& OutHLSL)
+bool FGrid3DCollectionAttributeHlslWriter::WriteAttributeGetIndexHLSL(EGridAttributeRetrievalMode AttributeStorage, const FNiagaraDataInterfaceGeneratedFunction& FunctionInfo, FString& OutHLSL)
 {
-	const FName* AttributeName = FunctionInfo.FindSpecifierValue(UNiagaraDataInterfaceGrid3DCollection::NAME_Attribute);
+	const FName* AttributeName = FunctionInfo.FindSpecifierValue(UNiagaraDataInterfaceRWBase::NAME_Attribute);
 	if (AttributeName == nullptr)
 	{
 		return false;
 	}
 
-	const FAttributeInfo* AttributeInfo = FindAttributeInfo(*AttributeName);
+	const FNiagaraDataInterfaceRWAttributeHelper::FAttributeInfo* AttributeInfo = AttributeHelper.FindAttributeInfo(*AttributeName);
 	if (AttributeInfo == nullptr)
 	{
 		return false;
@@ -462,7 +374,7 @@ bool FGrid3DCollectionAttributeHelper::WriteAttributeGetIndexHLSL(AttributeRetri
 	// #todo(dmp): for now it is ok to assume rgba grids only store 1 attribute per grid, so
 	// attributes are always stored at index 0.  In the future we might have a way to support
 	// multiple attributes per rgba grid, and this will have to change
-	if (AttributeStorage == AttributeRetrievalMode::RGBAGrid)
+	if (AttributeStorage == EGridAttributeRetrievalMode::RGBAGrid)
 	{
 		OutHLSL.Appendf(TEXT("void %s(out int Value)\n"), *FunctionInfo.InstanceName);
 		OutHLSL.Appendf(TEXT("{\n"));
@@ -473,7 +385,7 @@ bool FGrid3DCollectionAttributeHelper::WriteAttributeGetIndexHLSL(AttributeRetri
 	{	
 		OutHLSL.Appendf(TEXT("void %s(out int Value)\n"), *FunctionInfo.InstanceName);
 		OutHLSL.Appendf(TEXT("{\n"));
-		OutHLSL.Appendf(TEXT("	Value = %s;\n"), *GetAttributeIndex(AttributeStorage == AttributeRetrievalMode::Indirection, AttributeInfo));
+		OutHLSL.Appendf(TEXT("	Value = %s;\n"), *AttributeHelper.GetAttributeIndex(AttributeInfo));
 		OutHLSL.Appendf(TEXT("}\n"));
 	}
 
diff --git a/Engine/Plugins/FX/Niagara/Source/Niagara/Private/NiagaraDataInterfaceGrid3DCollectionUtils.h b/Engine/Plugins/FX/Niagara/Source/Niagara/Private/NiagaraDataInterfaceGrid3DCollectionUtils.h
--- a/Engine/Plugins/FX/Niagara/Source/Niagara/Private/NiagaraDataInterfaceGrid3DCollectionUtils.h
+++ b/Engine/Plugins/FX/Niagara/Source/Niagara/Private/NiagaraDataInterfaceGrid3DCollectionUtils.h
@@ -4,8 +4,17 @@
 
 #include "CoreMinimal.h"
 #include "NiagaraDataInterface.h"
+#include "NiagaraDataInterfaceRW.h"
+#include "NiagaraDataInterfaceRWUtils.h"
 
-struct FGrid3DCollectionAttributeHelper
+enum EGridAttributeRetrievalMode
+{
+	RGBAGrid = 0,
+	Indirection,
+	NoIndirection
+};
+
+struct FGrid3DCollectionAttributeHlslWriter
 {
 
 	// only support rgba textures when we have a single attribute that contains up to 4 channels
@@ -14,52 +23,30 @@ struct FGrid3DCollectionAttributeHelper
 
 	static TArray<FString> Channels;
 
-	struct FAttributeInfo
-	{
-		FName					Name;
-		FNiagaraTypeDefinition	TypeDef;
-		int32					NumChannels = 0;
-		int32					ChannelOffset = 0;
-		int32					AttributeIndex = 0;
-	};
+	explicit FGrid3DCollectionAttributeHlslWriter(const FNiagaraDataInterfaceGPUParamInfo& InParamInfo, TArray<FText>* OutWarnings = nullptr);
 
-	enum AttributeRetrievalMode
-	{
-		RGBAGrid = 0,
-		Indirection,
-		NoIndirection
-	};
-
-	explicit FGrid3DCollectionAttributeHelper(const FNiagaraDataInterfaceGPUParamInfo& InParamInfo, TArray<FText>* OutWarnings = nullptr);
+	const FNiagaraDataInterfaceRWAttributeHelper& GetAttributeHelper() { return AttributeHelper; }
 
 	bool UseRGBAGrid();
 
-	const FAttributeInfo* FindAttributeInfo(FName AttributeName) const;
-
 #if WITH_EDITORONLY_DATA
 	// Returns pixel offset for the channel
 	static FString GetPerAttributePixelOffset(const TCHAR* DataInterfaceHLSLSymbol);
 	FString GetPerAttributePixelOffset() const;
+	
 	// Returns UVW offset for the channel
 	static FString GetPerAttributeUVWOffset(const TCHAR* DataInterfaceHLSLSymbol);
 	FString GetPerAttributeUVWOffset() const;
 
-	// Translates named attribute into actual attribute index
-	FString GetAttributeIndex(bool bUseAttributeIndirection, const FAttributeInfo* AttributeInfo, int Channel = 0) const;
-	FString GetGridChannelString();
-	void GetChannelStrings(int AttributeIndex, int AttributeNumChannels, FString& NumChannelsString, FString& AttrGridChannels);
-	void GetChannelStrings(const FAttributeInfo* AttributeInfo, FString& NumChannelsString, FString& AttrGridChannels);
-
-	bool WriteGetHLSL(AttributeRetrievalMode AttributeStorage, const FNiagaraDataInterfaceGeneratedFunction& FunctionInfo, FString& OutHLSL);
-	bool WriteGetAtIndexHLSL(AttributeRetrievalMode AttributeStorage, const FNiagaraDataInterfaceGeneratedFunction& FunctionInfo, int NumChannels, FString& OutHLSL);
-	bool WriteSetAtIndexHLSL(AttributeRetrievalMode AttributeStorage, const FNiagaraDataInterfaceGeneratedFunction& FunctionInfo, int NumChannels, FString& OutHLSL);
-	bool WriteSampleAtIndexHLSL(AttributeRetrievalMode AttributeStorage, const FNiagaraDataInterfaceGeneratedFunction& FunctionInfo, int NumChannels, bool IsCubic, FString& OutHLSL);
-	bool WriteSetHLSL(AttributeRetrievalMode AttributeStorage, const FNiagaraDataInterfaceGeneratedFunction& FunctionInfo, FString& OutHLSL);
-	bool WriteSampleHLSL(AttributeRetrievalMode AttributeStorage, const FNiagaraDataInterfaceGeneratedFunction& FunctionInfo, bool IsCubic, FString& OutHLSL);
-	bool WriteAttributeGetIndexHLSL(AttributeRetrievalMode AttributeStorage, const FNiagaraDataInterfaceGeneratedFunction& FunctionInfo, FString& OutHLSL);
+	bool WriteGetHLSL(EGridAttributeRetrievalMode AttributeStorage, const FNiagaraDataInterfaceGeneratedFunction& FunctionInfo, FString& OutHLSL);
+	bool WriteGetAtIndexHLSL(EGridAttributeRetrievalMode AttributeStorage, const FNiagaraDataInterfaceGeneratedFunction& FunctionInfo, int NumChannels, FString& OutHLSL);
+	bool WriteSetAtIndexHLSL(EGridAttributeRetrievalMode AttributeStorage, const FNiagaraDataInterfaceGeneratedFunction& FunctionInfo, int NumChannels, FString& OutHLSL);
+	bool WriteSampleAtIndexHLSL(EGridAttributeRetrievalMode AttributeStorage, const FNiagaraDataInterfaceGeneratedFunction& FunctionInfo, int NumChannels, bool IsCubic, FString& OutHLSL);
+	bool WriteSetHLSL(EGridAttributeRetrievalMode AttributeStorage, const FNiagaraDataInterfaceGeneratedFunction& FunctionInfo, FString& OutHLSL);
+	bool WriteSampleHLSL(EGridAttributeRetrievalMode AttributeStorage, const FNiagaraDataInterfaceGeneratedFunction& FunctionInfo, bool IsCubic, FString& OutHLSL);
+	bool WriteAttributeGetIndexHLSL(EGridAttributeRetrievalMode AttributeStorage, const FNiagaraDataInterfaceGeneratedFunction& FunctionInfo, FString& OutHLSL);
 #endif //WITH_EDITORONLY_DATA
 
 	const FNiagaraDataInterfaceGPUParamInfo& ParamInfo;
-	TArray<FAttributeInfo>						AttributeInfos;
-	int32										TotalChannels = 0;
+	FNiagaraDataInterfaceRWAttributeHelper				AttributeHelper;	
 };
diff --git a/Engine/Plugins/FX/Niagara/Source/Niagara/Private/NiagaraDataInterfaceParticleRead.cpp b/Engine/Plugins/FX/Niagara/Source/Niagara/Private/NiagaraDataInterfaceParticleRead.cpp
--- a/Engine/Plugins/FX/Niagara/Source/Niagara/Private/NiagaraDataInterfaceParticleRead.cpp
+++ b/Engine/Plugins/FX/Niagara/Source/Niagara/Private/NiagaraDataInterfaceParticleRead.cpp
@@ -1025,9 +1025,8 @@ void UNiagaraDataInterfaceParticleRead::GetVMExternalFunction(const FVMExternalF
 
 	bool bBindSuccessful = false;
 	FNDIParticleRead_InstanceData* PIData = static_cast<FNDIParticleRead_InstanceData*>(InstanceData);
-	static const FName NAME_Attribute("Attribute");
 
-	const FVMFunctionSpecifier* FunctionSpecifier = BindingInfo.FindSpecifier(NAME_Attribute);
+	const FVMFunctionSpecifier* FunctionSpecifier = BindingInfo.FindSpecifier(UNiagaraDataInterfaceRWBase::NAME_Attribute);
 	if (FunctionSpecifier == nullptr)
 	{
 		UE_LOG(LogNiagara, Error, TEXT("VMExternalFunction '%s' does not have a function specifier 'attribute'!"), *BindingInfo.Name.ToString());
@@ -2297,8 +2296,7 @@ FNiagaraDataInterfaceParametersCS* UNiagaraDataInterfaceParticleRead::CreateShad
 	for (int32 FuncIdx = 0; FuncIdx < NumFuncs; ++FuncIdx)
 	{
 		const FNiagaraDataInterfaceGeneratedFunction& Func = ParameterInfo.GeneratedFunctions[FuncIdx];
-		static const FName NAME_Attribute("Attribute");
-		const FName* AttributeName = Func.FindSpecifierValue(NAME_Attribute);
+		const FName* AttributeName = Func.FindSpecifierValue(UNiagaraDataInterfaceRWBase::NAME_Attribute);
 		if (AttributeName != nullptr)
 		{
 			ShaderStorage->AttributeNames[FuncIdx] = *AttributeName;
@@ -2472,9 +2470,7 @@ void UNiagaraDataInterfaceParticleRead::GetFeedback(UNiagaraSystem* Asset, UNiag
 				{
 					for (const auto& Func : DIInfo.RegisteredFunctions)
 					{
-						static const FName NAME_Attribute("Attribute");
-
-						const FName* AttributeName = Func.FunctionSpecifiers.Find(NAME_Attribute);
+						const FName* AttributeName = Func.FunctionSpecifiers.Find(UNiagaraDataInterfaceRWBase::NAME_Attribute);
 						ENiagaraParticleDataValueType AttributeType = ENiagaraParticleDataValueType::Invalid;
 						if (AttributeName != nullptr)
 						{
diff --git a/Engine/Plugins/FX/Niagara/Source/Niagara/Private/NiagaraDataInterfaceRW.cpp b/Engine/Plugins/FX/Niagara/Source/Niagara/Private/NiagaraDataInterfaceRW.cpp
--- a/Engine/Plugins/FX/Niagara/Source/Niagara/Private/NiagaraDataInterfaceRW.cpp
+++ b/Engine/Plugins/FX/Niagara/Source/Niagara/Private/NiagaraDataInterfaceRW.cpp
@@ -4,6 +4,8 @@
 #include "ShaderParameterUtils.h"
 #include "ClearQuad.h"
 
+#include "UObject/UE5MainStreamObjectVersion.h"
+
 #include UE_INLINE_GENERATED_CPP_BY_NAME(NiagaraDataInterfaceRW)
 
 
@@ -17,6 +19,9 @@ const FString UNiagaraDataInterfaceRWBase::UnitToUVName(TEXT("_UnitToUV"));
 const FString UNiagaraDataInterfaceRWBase::CellSizeName(TEXT("_CellSize"));
 const FString UNiagaraDataInterfaceRWBase::WorldBBoxSizeName(TEXT("_WorldBBoxSize"));
 
+// Attribute names
+const FName UNiagaraDataInterfaceRWBase::NAME_Attribute("Attribute");
+
 // Global VM function names, also used by the shaders code generation methods.
 const FName UNiagaraDataInterfaceRWBase::NumCellsFunctionName("GetNumCells");
 const FName UNiagaraDataInterfaceRWBase::CellSizeFunctionName("GetCellSize");
@@ -490,12 +495,24 @@ UNiagaraDataInterfaceGrid2D::UNiagaraDataInterfaceGrid2D(FObjectInitializer cons
 	, NumCellsX(3)
 	, NumCellsY(3)
 	, NumCellsMaxAxis(3)
-	, NumAttributes(1)
+	, NumAttributes(0)
 	, SetGridFromMaxAxis(false)	
 	, WorldBBoxSize(100., 100.)
 {
 }
 
+void UNiagaraDataInterfaceGrid2D::Serialize(FArchive& Ar)
+{
+	Ar.UsingCustomVersion(FUE5MainStreamObjectVersion::GUID);
+
+	if (Ar.CustomVer(FUE5MainStreamObjectVersion::GUID) < FUE5MainStreamObjectVersion::NiagaraGrid2DDefaultUnnamedAttributesZero)
+	{
+		NumAttributes = 1;
+	}
+
+	Super::Serialize(Ar);
+}
+
 #if WITH_EDITOR
 void UNiagaraDataInterfaceGrid2D::ValidateFunction(const FNiagaraFunctionSignature& Function, TArray<FText>& OutValidationErrors)
 {
diff --git a/Engine/Plugins/FX/Niagara/Source/Niagara/Private/NiagaraDataInterfaceRWUtils.cpp b/Engine/Plugins/FX/Niagara/Source/Niagara/Private/NiagaraDataInterfaceRWUtils.cpp
new file mode 100644
--- /dev/null
+++ b/Engine/Plugins/FX/Niagara/Source/Niagara/Private/NiagaraDataInterfaceRWUtils.cpp
@@ -0,0 +1,115 @@
+// Copyright Epic Games, Inc. All Rights Reserved.
+#include "NiagaraDataInterfaceRWUtils.h"
+#include "NiagaraShader.h"
+#include "ShaderParameterUtils.h"
+#include "ClearQuad.h"
+
+#include "NiagaraDataInterfaceGrid3DCollection.h"
+#include "NiagaraDataInterfaceGrid2DCollection.h"
+
+#include "UObject/UE5MainStreamObjectVersion.h"
+
+#define LOCTEXT_NAMESPACE "NiagaraDataInterfaceRWUtils"
+
+TArray<FString> FNiagaraDataInterfaceRWAttributeHelper::Channels = { "r", "g", "b", "a" };
+FNiagaraDataInterfaceRWAttributeHelper::FNiagaraDataInterfaceRWAttributeHelper(const FNiagaraDataInterfaceGPUParamInfo& InParamInfo) : ParamInfo(InParamInfo)
+{
+}
+
+template<typename T>
+void FNiagaraDataInterfaceRWAttributeHelper::Init(TArray<FText>* OutWarnings /*= nullptr*/)
+{
+	AttributeInfos.Reserve(ParamInfo.GeneratedFunctions.Num());
+	TotalChannels = 0;
+
+	for (const FNiagaraDataInterfaceGeneratedFunction& Function : ParamInfo.GeneratedFunctions)
+	{
+		const FName* AttributeName = Function.FindSpecifierValue(UNiagaraDataInterfaceRWBase::NAME_Attribute);
+		if (AttributeName == nullptr)
+		{
+			continue;
+		}
+
+		if (const FAttributeInfo* ExistingAttribute = FindAttributeInfo(*AttributeName))
+		{
+			if (OutWarnings != nullptr)
+			{
+				FNiagaraTypeDefinition AttributeTypeDef = T::GetValueTypeFromFuncName(Function.DefinitionName);
+				if (ExistingAttribute->TypeDef != AttributeTypeDef)
+				{
+					OutWarnings->Emplace(FText::Format(LOCTEXT("BadType", "Same name, different types! {0} vs {1}, Attribute {2}"), AttributeTypeDef.GetNameText(), ExistingAttribute->TypeDef.GetNameText(), FText::FromName(ExistingAttribute->Name)));
+				}
+			}
+			continue;
+		}
+
+		FAttributeInfo& AttributeInfo = AttributeInfos.AddDefaulted_GetRef();
+		AttributeInfo.Name = *AttributeName;
+		AttributeInfo.TypeDef = T::GetValueTypeFromFuncName(Function.DefinitionName);
+		AttributeInfo.NumChannels = T::GetComponentCountFromFuncName(Function.DefinitionName);
+		AttributeInfo.ChannelOffset = TotalChannels;
+		AttributeInfo.AttributeIndex = AttributeInfos.Num() - 1;
+		TotalChannels += AttributeInfo.NumChannels;
+	}
+}
+
+const FNiagaraDataInterfaceRWAttributeHelper::FNiagaraDataInterfaceRWAttributeHelper::FAttributeInfo* FNiagaraDataInterfaceRWAttributeHelper::FindAttributeInfo(FName AttributeName) const
+{
+	return AttributeInfos.FindByPredicate([AttributeName](const FAttributeInfo& Info) { return Info.Name == AttributeName; });
+}
+
+#if WITH_EDITORONLY_DATA
+
+FString FNiagaraDataInterfaceRWAttributeHelper::GetAttributeIndex(const FAttributeInfo* AttributeInfo, int Channel /*= 0*/) const
+{
+	check(AttributeInfo);
+
+	return FString::Printf(TEXT("%d"), AttributeInfo->ChannelOffset + Channel);
+}
+
+FString FNiagaraDataInterfaceRWAttributeHelper::GetGridChannelString() const
+{
+	FString NumChannelsString = "";
+	if (TotalChannels > 1)
+	{
+		NumChannelsString = FString::FromInt(TotalChannels);
+	}
+	return NumChannelsString;
+}
+
+void FNiagaraDataInterfaceRWAttributeHelper::GetChannelStrings(int AttributeIndex, int AttributeNumChannels, FString& NumChannelsString, FString& AttrGridChannels)  const
+{
+	NumChannelsString = "";
+	if (AttributeNumChannels > 1)
+	{
+		NumChannelsString = FString::FromInt(AttributeNumChannels);
+	}
+
+	AttrGridChannels = FNiagaraDataInterfaceRWAttributeHelper::Channels[AttributeIndex];
+	for (int i = 1; i < AttributeNumChannels; ++i)
+	{
+		AttrGridChannels += FNiagaraDataInterfaceRWAttributeHelper::Channels[AttributeIndex + i];
+	}
+}
+
+void FNiagaraDataInterfaceRWAttributeHelper::GetChannelStrings(const FAttributeInfo* AttributeInfo, FString& NumChannelsString, FString& AttrGridChannels)  const
+{
+	NumChannelsString = "";
+	if (AttributeInfo->NumChannels > 1)
+	{
+		NumChannelsString = FString::FromInt(AttributeInfo->NumChannels);
+	}
+
+	AttrGridChannels = FNiagaraDataInterfaceRWAttributeHelper::Channels[AttributeInfo->AttributeIndex];
+	for (int i = 1; i < AttributeInfo->NumChannels; ++i)
+	{
+		AttrGridChannels += FNiagaraDataInterfaceRWAttributeHelper::Channels[AttributeInfo->AttributeIndex + i];
+	}
+}
+#endif
+
+template void FNiagaraDataInterfaceRWAttributeHelper::Init< UNiagaraDataInterfaceGrid3DCollection>(TArray<FText>* OutWarnings);
+template void FNiagaraDataInterfaceRWAttributeHelper::Init< UNiagaraDataInterfaceGrid2DCollection>(TArray<FText>* OutWarnings);
+
+#undef LOCTEXT_NAMESPACE
+
diff --git a/Engine/Source/Runtime/Core/Private/UObject/UE5MainStreamObjectVersion.cpp b/Engine/Source/Runtime/Core/Private/UObject/UE5MainStreamObjectVersion.cpp
--- a/Engine/Source/Runtime/Core/Private/UObject/UE5MainStreamObjectVersion.cpp
+++ b/Engine/Source/Runtime/Core/Private/UObject/UE5MainStreamObjectVersion.cpp
@@ -14,7 +14,7 @@ TMap<FGuid, FGuid> FUE5MainStreamObjectVersion::GetSystemGuids()
 	SystemGuids.Add(DevGuids.MATERIALSHADERMAP_DERIVEDDATA_VER, FGuid("65C01D817C9A4EEFAE9E988D41A1F3DD"));
 	SystemGuids.Add(DevGuids.NANITE_DERIVEDDATA_VER, FGuid("7A3EA93546624597BEBE2ADFDE3A7628"));
 	SystemGuids.Add(DevGuids.NIAGARASHADERMAP_DERIVEDDATA_VER, FGuid("8A37C45D24F2423CBE5F8F371DE33575"));
-	SystemGuids.Add(DevGuids.Niagara_LatestScriptCompileVersion, FGuid("BD1540BCB37D4B88A2EE4FF5ACFBB803"));
+	SystemGuids.Add(DevGuids.Niagara_LatestScriptCompileVersion, FGuid("5F19E9ECE6B844A49B23D2BD1DDB9095"));
 	SystemGuids.Add(DevGuids.SkeletalMeshDerivedDataVersion, FGuid("4084BE94509A8C3FA7CF410DDF8ABF04"));
 	SystemGuids.Add(DevGuids.STATICMESH_DERIVEDDATA_VER, FGuid("98F7D79A5811013825E75BBCC41ED3E9"));
 	SystemGuids.Add(DevGuids.POSESEARCHDB_DERIVEDDATA_VER, FGuid("4E595C2AC5E947D6BA9ABC874353E5BC"));
diff --git a/Engine/Source/Runtime/Core/Public/UObject/UE5MainStreamObjectVersion.h b/Engine/Source/Runtime/Core/Public/UObject/UE5MainStreamObjectVersion.h
--- a/Engine/Source/Runtime/Core/Public/UObject/UE5MainStreamObjectVersion.h
+++ b/Engine/Source/Runtime/Core/Public/UObject/UE5MainStreamObjectVersion.h
@@ -321,6 +321,9 @@ struct CORE_API FUE5MainStreamObjectVersion
 
 		// Animation Sequence now stores its frame-rate on a per-platform basis
 		PerPlatformAnimSequenceTargetFrameRate,
+		
+		// New default for number of attributes on 2d grids
+		NiagaraGrid2DDefaultUnnamedAttributesZero,
 
 		// -----<new versions can be added above this line>-------------------------------------------------
 		VersionPlusOne,
