- Optional on demand Niagara compilation
- Rather than kicking off compilation in PostLoad we shift it to when we first poll / wait for compilation
- Upside is way less compilation requests
- Downside is some FX might be missed initially in PIE due to compilation being required
- When Niagara compilation produces less hitches we can process this queue slowly over time

#prefliight 639919d28c5081ee9a672d77

diff --git a/Engine/Plugins/FX/Niagara/Source/Niagara/Classes/NiagaraSystem.h b/Engine/Plugins/FX/Niagara/Source/Niagara/Classes/NiagaraSystem.h
--- a/Engine/Plugins/FX/Niagara/Source/Niagara/Classes/NiagaraSystem.h
+++ b/Engine/Plugins/FX/Niagara/Source/Niagara/Classes/NiagaraSystem.h
@@ -413,7 +413,7 @@ public:
 	bool RequestCompile(bool bForce, FNiagaraSystemUpdateContext* OptionalUpdateContext = nullptr);
 
 	/** If we have a pending compile request, is it done with yet? */
-	bool PollForCompilationComplete();
+	bool PollForCompilationComplete(bool bFlushRequestCompile = true);
 
 	/** Blocks until all active compile jobs have finished */
 	void WaitForCompilationComplete(bool bIncludingGPUShaders = false, bool bShowProgress = true);
@@ -552,6 +552,9 @@ protected:
 	UPROPERTY(EditAnywhere, AdvancedDisplay, Category = "Performance", meta = (DisplayName="Disable Debug Switches"))
 	uint32 bDisableDebugSwitchesOnCook : 1;
 
+	/* When set the system needs to compile before it can be activated. */
+	uint32 bNeedsRequestCompile : 1;
+
 public:
 	/** Subscriptions to definitions of parameters. */
 	UPROPERTY()
@@ -656,6 +659,10 @@ public:
 	/** Cache data & accessors from the compiled data, allows us to avoid per instance. */
 	void CacheFromCompiledData();
 
+#if WITH_EDITORONLY_DATA
+	bool NeedsRequestCompile() const { return bNeedsRequestCompile; }
+#endif
+
 	FORCEINLINE TConstArrayView<FNiagaraEmitterExecutionIndex> GetEmitterExecutionOrder() const { return MakeArrayView(EmitterExecutionOrder); }
 	FORCEINLINE TConstArrayView<FNiagaraRendererExecutionIndex> GetRendererPostTickOrder() const { return MakeArrayView(RendererPostTickOrder); }
 	FORCEINLINE TConstArrayView<FNiagaraRendererExecutionIndex> GetRendererCompletionOrder() const { return MakeArrayView(RendererCompletionOrder); }
diff --git a/Engine/Plugins/FX/Niagara/Source/Niagara/Private/NiagaraComponentPool.cpp b/Engine/Plugins/FX/Niagara/Source/Niagara/Private/NiagaraComponentPool.cpp
--- a/Engine/Plugins/FX/Niagara/Source/Niagara/Private/NiagaraComponentPool.cpp
+++ b/Engine/Plugins/FX/Niagara/Source/Niagara/Private/NiagaraComponentPool.cpp
@@ -325,6 +325,12 @@ void UNiagaraComponentPool::PrimePool(UNiagaraSystem* Template, UWorld* World)
 		return;
 	}
 
+	const uint32 PrimeCount = FMath::Min(Template->PoolPrimeSize, Template->MaxPoolSize);
+	if (PrimeCount == 0)
+	{
+		return;
+	}
+
 	if (World->bIsTearingDown)
 	{
 		UE_LOG(LogNiagara, Warning, TEXT("Failed to prime particle pool as we are tearing the world down."));
@@ -337,6 +343,14 @@ void UNiagaraComponentPool::PrimePool(UNiagaraSystem* Template, UWorld* World)
 		return;
 	}	
 
+#if WITH_EDITORONLY_DATA
+	if (Template->NeedsRequestCompile())
+	{
+		UE_LOG(LogNiagara, Verbose, TEXT("Failed to prime particle pool as the Niagara System needs to compile."));
+		return;
+	}
+#endif
+
 	if (Template->IsReadyToRun() == false)
 	{
 		UE_LOG(LogNiagara, Verbose, TEXT("Failed to prime particle pool as the Niagara System is not ready to run."));
@@ -365,17 +379,15 @@ void UNiagaraComponentPool::PrimePool(UNiagaraSystem* Template, UWorld* World)
 
 	FNiagaraCrashReporterScope CRScope(Template);
 
-	uint32 Count = FMath::Min(Template->PoolPrimeSize, Template->MaxPoolSize);
-
-	if(Count > 0)
+	if(PrimeCount > 0)
 	{
 		FNCPool& Pool = WorldParticleSystemPools.FindOrAdd(Template);
 
 		uint32 ExistingComponents = Pool.NumComponents();
-		if (ExistingComponents < Count)
+		if (ExistingComponents < PrimeCount)
 		{
 			TArray<UNiagaraComponent*> NewComps;
-			NewComps.SetNum(Count - ExistingComponents);
+			NewComps.SetNum(PrimeCount - ExistingComponents);
 			for (auto& Comp : NewComps)
 			{
 				Comp = Pool.Acquire(World, Template, ENCPoolMethod::ManualRelease, true);//Force the pool to create a new system.
diff --git a/Engine/Plugins/FX/Niagara/Source/Niagara/Private/NiagaraModule.cpp b/Engine/Plugins/FX/Niagara/Source/Niagara/Private/NiagaraModule.cpp
--- a/Engine/Plugins/FX/Niagara/Source/Niagara/Private/NiagaraModule.cpp
+++ b/Engine/Plugins/FX/Niagara/Source/Niagara/Private/NiagaraModule.cpp
@@ -451,7 +451,7 @@ void PollSystemCompilations()
 	for (TObjectIterator<UNiagaraSystem> SysIt; SysIt; ++SysIt)
 	{
 		UNiagaraSystem* System = *SysIt;
-		System->PollForCompilationComplete();
+		System->PollForCompilationComplete(false);
 		if (FPlatformTime::Seconds() - Start > GNiagaraMaxCompilePollTimePerFrame)
 		{
 			break;
diff --git a/Engine/Plugins/FX/Niagara/Source/Niagara/Private/NiagaraSystem.cpp b/Engine/Plugins/FX/Niagara/Source/Niagara/Private/NiagaraSystem.cpp
--- a/Engine/Plugins/FX/Niagara/Source/Niagara/Private/NiagaraSystem.cpp
+++ b/Engine/Plugins/FX/Niagara/Source/Niagara/Private/NiagaraSystem.cpp
@@ -94,6 +94,16 @@ static FAutoConsoleVariableRef CVarNiagaraCompileWaitLoggingTerminationCap(
 	ECVF_Default
 );
 
+#if WITH_EDITORONLY_DATA
+static int GNiagaraOnDemandCompileEnabled = 1;
+static FAutoConsoleVariableRef CVarNiagaraOnDemandCompileEnabled(
+	TEXT("fx.Niagara.OnDemandCompileEnabled"),
+	GNiagaraOnDemandCompileEnabled,
+	TEXT("Compiles Niagara Systems on demand rather than on post load."),
+	ECVF_Default
+);
+#endif
+
 //////////////////////////////////////////////////////////////////////////
 
 class FNiagaraSystemNullCompletionTask
@@ -146,6 +156,7 @@ UNiagaraSystem::UNiagaraSystem(const FObjectInitializer& ObjectInitializer)
 , bIgnoreParticleReadsForAttributeTrim(false)
 , bDisableDebugSwitches(false)
 , bDisableDebugSwitchesOnCook(true)
+, bNeedsRequestCompile(false)
 #endif
 , bSupportLargeWorldCoordinates(true)
 , bDisableExperimentalVM(false)
@@ -580,6 +591,7 @@ void UNiagaraSystem::UpdateSystemAfterLoad()
 
 #if WITH_EDITORONLY_DATA
 	// check if the system needs to be compiled and start a compile task if necessary
+	bNeedsRequestCompile = false;
 	if (GetOutermost()->HasAnyPackageFlags(EPackageFlags::PKG_Cooked) == false)
 	{
 		bool bSystemScriptsAreSynchronized = true;
@@ -609,13 +621,17 @@ void UNiagaraSystem::UpdateSystemAfterLoad()
 
 		if (bSystemScriptsAreSynchronized == false || bEmitterScriptsAreSynchronized == false)
 		{
+			bNeedsRequestCompile = true;
 			if (IsRunningCommandlet())
 			{
 				// Call modify here so that the system will resave the compile ids and script vm when running the resave
 				// commandlet. We don't need it for normal post-loading.
 				Modify();
 			}
-			RequestCompile(false);
+			if ( !GIsEditor || GNiagaraOnDemandCompileEnabled == 0 )
+			{
+				RequestCompile(false);
+			}
 		}
 	}
 #endif
@@ -1417,6 +1433,11 @@ void UNiagaraSystem::GetAssetRegistryTags(TArray<FAssetRegistryTag>& OutTags) co
 #if WITH_EDITORONLY_DATA
 bool UNiagaraSystem::HasOutstandingCompilationRequests(bool bIncludingGPUShaders) const
 {
+	if (bNeedsRequestCompile)
+	{
+		return true;
+	}
+
 	if (ActiveCompilations.Num() > 0)
 	{
 		return true;
@@ -2224,6 +2245,11 @@ void UNiagaraSystem::WaitForCompilationComplete(bool bIncludingGPUShaders, bool
 	TRACE_CPUPROFILER_EVENT_SCOPE(WaitForNiagaraCompilation);
 	TRACE_CPUPROFILER_EVENT_SCOPE_TEXT_ON_CHANNEL(*GetPathName(), NiagaraChannel);
 
+	if (bNeedsRequestCompile)
+	{
+		RequestCompile(false);
+	}
+
 	// Calculate the slow progress for notifying via UI
 	TArray<FNiagaraShaderScript*, TInlineAllocator<16>> GPUScripts;
 	if (bIncludingGPUShaders)
@@ -2297,8 +2323,13 @@ void UNiagaraSystem::InvalidateActiveCompiles()
 	}
 }
 
-bool UNiagaraSystem::PollForCompilationComplete()
+bool UNiagaraSystem::PollForCompilationComplete(bool bFlushRequestCompile)
 {
+	if (bNeedsRequestCompile && bFlushRequestCompile)
+	{
+		RequestCompile(false);
+	}
+
 	if (ActiveCompilations.Num() > 0)
 	{
 		return QueryCompileComplete(false, true);
@@ -3031,6 +3062,8 @@ bool UNiagaraSystem::RequestCompile(bool bForce, FNiagaraSystemUpdateContext* Op
 	check(IsInGameThread());
 	TRACE_CPUPROFILER_EVENT_SCOPE(UNiagaraSystem::RequestCompile)
 
+	bNeedsRequestCompile = false;
+
 	// We remove emitters and scripts on dedicated servers, so skip further work.
 	const bool bIsDedicatedServer = !GIsClient && GIsServer;
 	if (bIsDedicatedServer)
