Fixed missing dependencies causing build error for NonUnity Android in NiagaraGenerateMips.cpp.

diff --git a/Engine/Plugins/FX/Niagara/Source/NiagaraShader/Private/NiagaraGenerateMips.cpp b/Engine/Plugins/FX/Niagara/Source/NiagaraShader/Private/NiagaraGenerateMips.cpp
--- a/Engine/Plugins/FX/Niagara/Source/NiagaraShader/Private/NiagaraGenerateMips.cpp
+++ b/Engine/Plugins/FX/Niagara/Source/NiagaraShader/Private/NiagaraGenerateMips.cpp
@@ -4,6 +4,7 @@
 #include "GlobalShader.h"
 #include "RenderGraphUtils.h"
 #include "RHIStaticStates.h"
+#include "DataDrivenShaderPlatformInfo.h"
 
 #include UE_INLINE_GENERATED_CPP_BY_NAME(NiagaraGenerateMips)
 
