- Add support for material double attribute bindings
- Position type will always be in world position regardless of emitter settings
- Filter selectable variables based upon material parameter type

diff --git a/Engine/Plugins/FX/Niagara/Source/Niagara/Private/NiagaraRenderer.cpp b/Engine/Plugins/FX/Niagara/Source/Niagara/Private/NiagaraRenderer.cpp
--- a/Engine/Plugins/FX/Niagara/Source/Niagara/Private/NiagaraRenderer.cpp
+++ b/Engine/Plugins/FX/Niagara/Source/Niagara/Private/NiagaraRenderer.cpp
@@ -527,7 +527,6 @@ void FNiagaraRenderer::ProcessMaterialParameterBindings(const FNiagaraRendererMa
 				{
 					for (const FNiagaraMaterialAttributeBinding& Binding : MaterialParameters.AttributeBindings)
 					{
-
 						if (Binding.GetParamMapBindableVariable().GetType() == FNiagaraTypeDefinition::GetVec4Def() ||
 							(Binding.GetParamMapBindableVariable().GetType().IsDataInterface() && Binding.NiagaraChildVariable.GetType() == FNiagaraTypeDefinition::GetVec4Def()))
 						{
@@ -554,10 +553,10 @@ void FNiagaraRenderer::ProcessMaterialParameterBindings(const FNiagaraRendererMa
 						{
 							FNiagaraPosition Var(ForceInitToZero);
 							InEmitter->GetBoundRendererValue_GT(Binding.GetParamMapBindableVariable(), Binding.NiagaraChildVariable, &Var);
-							FNiagaraLWCConverter LwcConverter = SystemInstance->GetLWCConverter(InEmitter->GetCachedEmitterData() ? InEmitter->GetCachedEmitterData()->bLocalSpace : false);
-							FVector WorldPos = LwcConverter.ConvertSimulationPositionToWorld(Var);
+							const FNiagaraLWCConverter LwcConverter = SystemInstance->GetLWCConverter(false);
+							const FVector WorldPos = LwcConverter.ConvertSimulationPositionToWorld(Var);
 
-							//TODO: should we use SetDoubleVectorParamValue() instead to prevent accuracy loss? 
+							MatDyn->SetDoubleVectorParameterValue(Binding.MaterialParameterName, WorldPos);
 							MatDyn->SetVectorParameterValue(Binding.MaterialParameterName, WorldPos);
 						}
 						else if (Binding.GetParamMapBindableVariable().GetType() == FNiagaraTypeDefinition::GetVec2Def() ||
diff --git a/Engine/Plugins/FX/Niagara/Source/Niagara/Private/NiagaraRendererProperties.cpp b/Engine/Plugins/FX/Niagara/Source/Niagara/Private/NiagaraRendererProperties.cpp
--- a/Engine/Plugins/FX/Niagara/Source/Niagara/Private/NiagaraRendererProperties.cpp
+++ b/Engine/Plugins/FX/Niagara/Source/Niagara/Private/NiagaraRendererProperties.cpp
@@ -91,7 +91,9 @@ bool FNiagaraRendererLayout::SetVariable(const FNiagaraDataSetCompiledData* Comp
 bool FNiagaraRendererLayout::SetVariableFromBinding(const FNiagaraDataSetCompiledData* CompiledData, const FNiagaraVariableAttributeBinding& VariableBinding, int32 VFVarOffset)
 {
 	if (VariableBinding.IsParticleBinding())
+	{
 		return SetVariable(CompiledData, VariableBinding.GetDataSetBindableVariable(), VFVarOffset);
+	}
 	return false;
 }
 
@@ -210,6 +212,14 @@ void FNiagaraRendererMaterialParameters::GetFeedback(TArrayView<UMaterialInterfa
 				VectorParametersValid[i] |= ContainsParameter(VectorParameters[i].MaterialParameterName);
 			}
 		}
+		if (AttributeBindingsValid.Num() > 0)
+		{
+			Material->GetAllDoubleVectorParameterInfo(TempParameterInfo, TempParameterIds);
+			for (int32 i = 0; i < AttributeBindings.Num(); ++i)
+			{
+				AttributeBindingsValid[i] |= ContainsParameter(AttributeBindings[i].MaterialParameterName);
+			}
+		}
 		if (AttributeBindingsValid.Num() > 0 || TextureParametersValid.Num() > 0)
 		{
 			Material->GetAllTextureParameterInfo(TempParameterInfo, TempParameterIds);
diff --git a/Engine/Plugins/FX/Niagara/Source/NiagaraEditor/Private/Customizations/NiagaraTypeCustomizations.cpp b/Engine/Plugins/FX/Niagara/Source/NiagaraEditor/Private/Customizations/NiagaraTypeCustomizations.cpp
--- a/Engine/Plugins/FX/Niagara/Source/NiagaraEditor/Private/Customizations/NiagaraTypeCustomizations.cpp
+++ b/Engine/Plugins/FX/Niagara/Source/NiagaraEditor/Private/Customizations/NiagaraTypeCustomizations.cpp
@@ -812,35 +812,92 @@ TSharedRef<SWidget> FNiagaraMaterialAttributeBindingCustomization::OnGetNiagaraM
 	return MenuContent;
 }
 
-
-bool FNiagaraMaterialAttributeBindingCustomization::IsCompatibleNiagaraVariable(const FNiagaraVariable& InVar) const
+TArray<FNiagaraTypeDefinition> FNiagaraMaterialAttributeBindingCustomization::GetAllowedVariableTypes() const
 {
+	TArray<FNiagaraTypeDefinition> AllowedTypes;
+	if (BaseSystem && TargetParameterBinding && PropertyHandle)
 	{
-
-		if (InVar.GetType() == FNiagaraTypeDefinition::GetFloatDef() ||
-			InVar.GetType() == FNiagaraTypeDefinition::GetVec4Def() ||
-			InVar.GetType() == FNiagaraTypeDefinition::GetColorDef() ||
-			InVar.GetType() == FNiagaraTypeDefinition::GetVec2Def() ||
-			InVar.GetType() == FNiagaraTypeDefinition::GetVec3Def() ||
-			InVar.GetType() == FNiagaraTypeDefinition::GetPositionDef() ||
-			InVar.GetType() == FNiagaraTypeDefinition::GetUObjectDef() ||
-			InVar.GetType() == FNiagaraTypeDefinition::GetUTextureDef() ||
-			InVar.GetType() == FNiagaraTypeDefinition::GetUTextureRenderTargetDef() )
-		{
-			return true;
-		}
-		else if (InVar.GetType().IsDataInterface())
+		TArray<UObject*> Objects;
+		PropertyHandle->GetOuterObjects(Objects);
+	
+		if (Objects.Num() == 1)
 		{
-			return true;
+			UNiagaraRendererProperties* RendererProperties = Cast<UNiagaraRendererProperties>(Objects[0]);
+			TArray<UMaterialInterface*> Materials;
+			if (RendererProperties)
+			{
+				RendererProperties->GetUsedMaterials(nullptr, Materials);
+			}
+	
+			
+			auto MatchMaterialParameter =
+				[MaterialName=TargetParameterBinding->MaterialParameterName](const FMaterialParameterInfo& InParameter)
+				{
+					return InParameter.Name == MaterialName;
+				};
+
+			TArray<FMaterialParameterInfo> TempParameterInfo;
+			TArray<FGuid> ParameterIds;
+			bool bHasTextureBinding = false;
+			bool bHasScalarBinding = false;
+			bool bHasVectorBinding = false;
+			bool bHasDoubleVectorBinding = false;
+			for (UMaterialInterface* Material : Materials)
+			{
+				if (!Material)
+				{
+					continue;
+				}
+	
+				Material->GetAllTextureParameterInfo(TempParameterInfo, ParameterIds);
+				bHasTextureBinding |= TempParameterInfo.ContainsByPredicate(MatchMaterialParameter);
+				TempParameterInfo.Reset();
+	
+				Material->GetAllScalarParameterInfo(TempParameterInfo, ParameterIds);
+				bHasScalarBinding |= TempParameterInfo.ContainsByPredicate(MatchMaterialParameter);
+				TempParameterInfo.Reset();
+
+				Material->GetAllVectorParameterInfo(TempParameterInfo, ParameterIds);
+				bHasVectorBinding |= TempParameterInfo.ContainsByPredicate(MatchMaterialParameter);
+				TempParameterInfo.Reset();
+
+				Material->GetAllDoubleVectorParameterInfo(TempParameterInfo, ParameterIds);
+				bHasDoubleVectorBinding |= TempParameterInfo.ContainsByPredicate(MatchMaterialParameter);
+				TempParameterInfo.Reset();
+
+				ParameterIds.Reset();
+			}
+
+			if (bHasTextureBinding)
+			{
+				AllowedTypes.AddUnique(FNiagaraTypeDefinition::GetUObjectDef());
+				AllowedTypes.AddUnique(FNiagaraTypeDefinition::GetUTextureDef());
+				AllowedTypes.AddUnique(FNiagaraTypeDefinition::GetUTextureRenderTargetDef());
+			}
+			if (bHasScalarBinding)
+			{
+				AllowedTypes.AddUnique(FNiagaraTypeDefinition::GetFloatDef());
+			}
+			if (bHasVectorBinding)
+			{
+				AllowedTypes.AddUnique(FNiagaraTypeDefinition::GetVec4Def());
+				AllowedTypes.AddUnique(FNiagaraTypeDefinition::GetColorDef());
+				AllowedTypes.AddUnique(FNiagaraTypeDefinition::GetVec2Def());
+				AllowedTypes.AddUnique(FNiagaraTypeDefinition::GetVec3Def());
+			}
+			if (bHasDoubleVectorBinding)
+			{
+				AllowedTypes.AddUnique(FNiagaraTypeDefinition::GetPositionDef());
+			}
 		}
 	}
-	return false;
+	return AllowedTypes;
 }
 
 TArray<TPair<FNiagaraVariableBase, FNiagaraVariableBase> > FNiagaraMaterialAttributeBindingCustomization::GetNiagaraNames() const
 {
 	TArray<TPair<FNiagaraVariableBase, FNiagaraVariableBase>> Names;
-	TArray < FNiagaraVariableBase > BaseVars;
+	TArray<FNiagaraVariableBase> BaseVars;
 
 	if (BaseSystem && BaseEmitter.Emitter && TargetParameterBinding)
 	{
@@ -884,29 +941,33 @@ TArray<TPair<FNiagaraVariableBase, FNiagaraVariableBase> > FNiagaraMaterialAttri
 				return BaseVariable.GetType().GetClass()->GetDefaultObject<UNiagaraDataInterface>();
 			};
 
-		for (const FNiagaraVariableBase& BaseVar : BaseVars)
+		if (BaseVars.Num() > 0)
 		{
-			if (BaseVar.IsDataInterface())
+			const TArray<FNiagaraTypeDefinition> AllowedVariableTypes = GetAllowedVariableTypes();
+			for (const FNiagaraVariableBase& BaseVar : BaseVars)
 			{
-				UNiagaraDataInterface* DI = FindCachedDI(BaseVar);
-				if (DI && DI->CanExposeVariables())
+				if (BaseVar.IsDataInterface())
 				{
-					TArray<FNiagaraVariableBase> ChildVars;
-					DI->GetExposedVariables(ChildVars);
-					for (const FNiagaraVariableBase& ChildVar : ChildVars)
+					UNiagaraDataInterface* DI = FindCachedDI(BaseVar);
+					if (DI && DI->CanExposeVariables())
 					{
-						if (IsCompatibleNiagaraVariable(ChildVar))
+						TArray<FNiagaraVariableBase> ChildVars;
+						DI->GetExposedVariables(ChildVars);
+						for (const FNiagaraVariableBase& ChildVar : ChildVars)
 						{
-							Names.AddUnique(TPair<FNiagaraVariableBase, FNiagaraVariableBase>(BaseVar, ChildVar));
+							if (!AllowedVariableTypes.Num() || AllowedVariableTypes.Contains(ChildVar.GetType()))
+							{
+								Names.AddUnique(TPair<FNiagaraVariableBase, FNiagaraVariableBase>(BaseVar, ChildVar));
+							}
 						}
 					}
 				}
-			}
-			else if (IsCompatibleNiagaraVariable(BaseVar))
-			{
-				if (RenderProps && TargetParameterBinding && RenderProps->IsSupportedVariableForBinding(BaseVar, TargetParameterBinding->MaterialParameterName))
+				else if (!AllowedVariableTypes.Num() || AllowedVariableTypes.Contains(BaseVar.GetType()))
 				{
-					Names.AddUnique(TPair<FNiagaraVariableBase, FNiagaraVariableBase>(BaseVar, FNiagaraVariableBase()));
+					if (RenderProps && TargetParameterBinding && RenderProps->IsSupportedVariableForBinding(BaseVar, TargetParameterBinding->MaterialParameterName))
+					{
+						Names.AddUnique(TPair<FNiagaraVariableBase, FNiagaraVariableBase>(BaseVar, FNiagaraVariableBase()));
+					}
 				}
 			}
 		}
@@ -1109,6 +1170,9 @@ void FNiagaraMaterialAttributeBindingCustomization::GetMaterialParameters(TArray
 				Material->GetAllVectorParameterInfo(TempParameterInfo, ParameterIds);
 				TransformMaterialParameterInfo(Material, TempParameterInfo);
 
+				Material->GetAllDoubleVectorParameterInfo(TempParameterInfo, ParameterIds);
+				TransformMaterialParameterInfo(Material, TempParameterInfo);
+
 				ParameterIds.Reset();
 			}
 		}
diff --git a/Engine/Plugins/FX/Niagara/Source/NiagaraEditor/Private/Customizations/NiagaraTypeCustomizations.h b/Engine/Plugins/FX/Niagara/Source/NiagaraEditor/Private/Customizations/NiagaraTypeCustomizations.h
--- a/Engine/Plugins/FX/Niagara/Source/NiagaraEditor/Private/Customizations/NiagaraTypeCustomizations.h
+++ b/Engine/Plugins/FX/Niagara/Source/NiagaraEditor/Private/Customizations/NiagaraTypeCustomizations.h
@@ -262,7 +262,7 @@ private:
 	TSharedRef<SWidget> OnCreateWidgetForMaterialAction(struct FCreateWidgetForActionData* const InCreateData);
 	void OnMaterialActionSelected(const TArray< TSharedPtr<FEdGraphSchemaAction> >& SelectedActions, ESelectInfo::Type InSelectionType);
 
-	bool IsCompatibleNiagaraVariable(const struct FNiagaraVariable& InVar) const;
+	TArray<FNiagaraTypeDefinition> GetAllowedVariableTypes() const;
 	static FText MakeCurrentText(const FNiagaraVariableBase& BaseVar, const FNiagaraVariableBase& ChildVar);
 
 	TSharedPtr<IPropertyHandle> PropertyHandle;
