- Allow static switches to set layer / blend parameters
- Remove all but global bindings from anything but static switches as it would require more complex logic to set them
- Remove duplicate parameter names from the material parameter list selection

diff --git a/Engine/Plugins/FX/Niagara/Source/Niagara/Private/NiagaraRendererProperties.cpp b/Engine/Plugins/FX/Niagara/Source/Niagara/Private/NiagaraRendererProperties.cpp
--- a/Engine/Plugins/FX/Niagara/Source/Niagara/Private/NiagaraRendererProperties.cpp
+++ b/Engine/Plugins/FX/Niagara/Source/Niagara/Private/NiagaraRendererProperties.cpp
@@ -323,8 +323,6 @@ void UNiagaraRendererProperties::ChangeToPositionBinding(FNiagaraVariableAttribu
 
 bool UNiagaraRendererProperties::UpdateMaterialStaticParameters(const FNiagaraRendererMaterialParameters& MaterialParameters, UMaterialInstanceConstant* MIC)
 {
-	FStaticParameterSet StaticParameterSet;
-
 	UNiagaraEmitter* NiagaraEmitter = GetTypedOuter<UNiagaraEmitter>();
 	UNiagaraSystem* NiagaraSystem = GetTypedOuter<UNiagaraSystem>();
 	if (NiagaraEmitter == nullptr || NiagaraSystem == nullptr)
@@ -332,6 +330,13 @@ bool UNiagaraRendererProperties::UpdateMaterialStaticParameters(const FNiagaraRe
 		return false;
 	}
 
+	TArray<FMaterialParameterInfo> AllStaticSwitchParameterInfos;
+	{
+		TArray<FGuid> ParameterGuids;
+		MIC->GetAllStaticSwitchParameterInfo(AllStaticSwitchParameterInfos, ParameterGuids);
+	}
+	FStaticParameterSet StaticParameterSet;
+
 	bool bModified = false;
 	for (const FNiagaraRendererMaterialStaticBoolParameter& ParameterBinding : MaterialParameters.StaticBoolParameters)
 	{
@@ -352,37 +357,43 @@ bool UNiagaraRendererProperties::UpdateMaterialStaticParameters(const FNiagaraRe
 						continue;
 					}
 
-					FStaticSwitchParameter* StaticParameter = StaticParameterSet.StaticSwitchParameters.FindByPredicate(
-						[ParameterBinding](const FStaticSwitchParameter& StaticParameter)
+					for (const FMaterialParameterInfo& ParameterInfo : AllStaticSwitchParameterInfos)
+					{
+						if (ParameterInfo.Name != ParameterBinding.MaterialParameterName)
 						{
-							return StaticParameter.ParameterInfo.Name == ParameterBinding.MaterialParameterName;
+							continue;
 						}
-					);
 
-					const bool bNewValue = StaticVariable.GetValue<bool>();
-					if (StaticParameter == nullptr)
-					{
-						FGuid ParameterGuid;
-						bool bDefaultValue = false;
-						if (MIC->GetStaticSwitchParameterDefaultValue(ParameterBinding.MaterialParameterName, bDefaultValue, ParameterGuid))
+						FStaticSwitchParameter* StaticParameter = StaticParameterSet.StaticSwitchParameters.FindByPredicate(
+							[ParameterInfo](const FStaticSwitchParameter& StaticParameter)
+							{
+								return StaticParameter.ParameterInfo == ParameterInfo;
+							}
+						);
+
+						const bool bNewValue = StaticVariable.GetValue<bool>();
+						if (StaticParameter == nullptr)
 						{
-							if (bDefaultValue != bNewValue)
+							FGuid ParameterGuid;
+							bool bDefaultValue = false;
+							if (MIC->GetStaticSwitchParameterDefaultValue(ParameterInfo, bDefaultValue, ParameterGuid))
 							{
-								StaticParameterSet.SetParameterValue(ParameterBinding.MaterialParameterName, FMaterialParameterValue(bNewValue), EMaterialSetParameterValueFlags::None);
-								bModified = true;
+								if (bDefaultValue != bNewValue)
+								{
+									StaticParameterSet.SetParameterValue(ParameterInfo, FMaterialParameterValue(bNewValue), EMaterialSetParameterValueFlags::None);
+									bModified = true;
+								}
 							}
 						}
-					}
-					else
-					{
-						if (StaticParameter && StaticParameter->Value != bNewValue)
+						else
 						{
-							StaticParameter->Value = bNewValue;
-							bModified = true;
+							if (StaticParameter && StaticParameter->Value != bNewValue)
+							{
+								StaticParameter->Value = bNewValue;
+								bModified = true;
+							}
 						}
 					}
-
-					break;
 				}
 			}
 		);
diff --git a/Engine/Plugins/FX/Niagara/Source/NiagaraEditor/Private/Customizations/NiagaraTypeCustomizations.cpp b/Engine/Plugins/FX/Niagara/Source/NiagaraEditor/Private/Customizations/NiagaraTypeCustomizations.cpp
--- a/Engine/Plugins/FX/Niagara/Source/NiagaraEditor/Private/Customizations/NiagaraTypeCustomizations.cpp
+++ b/Engine/Plugins/FX/Niagara/Source/NiagaraEditor/Private/Customizations/NiagaraTypeCustomizations.cpp
@@ -1126,11 +1126,22 @@ TSharedRef<SWidget> FNiagaraMaterialAttributeBindingCustomization::OnGetMaterial
 
 void FNiagaraMaterialAttributeBindingCustomization::GetMaterialParameters(TArray<TPair<FName, FString>>& OutBindingNameAndDescription) const
 {
+	TSet<FName> AddedParameterNames;
 	auto TransformMaterialParameterInfo =
-		[&OutBindingNameAndDescription](UMaterialInterface* Material, TArray<FMaterialParameterInfo>& ParameterInfos)
+		[&OutBindingNameAndDescription, &AddedParameterNames](UMaterialInterface* Material, TArray<FMaterialParameterInfo>& ParameterInfos)
 		{
 			for (const FMaterialParameterInfo& ParameterInfo : ParameterInfos)
 			{
+				if (ParameterInfo.Association != EMaterialParameterAssociation::GlobalParameter)
+				{
+					continue;
+				}
+				if (AddedParameterNames.Contains(ParameterInfo.Name))
+				{
+					continue;
+				}
+				AddedParameterNames.Add(ParameterInfo.Name);
+
 				FString ParameterDesc;
 				Material->GetParameterDesc(ParameterInfo, ParameterDesc);
 				OutBindingNameAndDescription.Emplace(ParameterInfo.Name, ParameterDesc);
@@ -2185,6 +2196,7 @@ TSharedRef<SWidget> FNiagaraRendererMaterialParameterCustomization::OnGetMateria
 	UNiagaraRendererProperties* RenderProperties = WeakRenderProperties.Get();
 	if (RenderProperties != nullptr && PropertyHandle.IsValid())
 	{
+		TSet<FName> AddedParameterNames;
 		TArray<UMaterialInterface*> UsedMaterials;
 		RenderProperties->GetUsedMaterials(nullptr, UsedMaterials);
 
@@ -2200,6 +2212,17 @@ TSharedRef<SWidget> FNiagaraRendererMaterialParameterCustomization::OnGetMateria
 
 			for (const FMaterialParameterInfo& MaterialParameterInfo : MaterialParameterInfos)
 			{
+				if (MaterialParameterInfo.Association != EMaterialParameterAssociation::GlobalParameter)
+				{
+					continue;
+				}
+
+				if (AddedParameterNames.Contains(MaterialParameterInfo.Name))
+				{
+					continue;
+				}
+				AddedParameterNames.Add(MaterialParameterInfo.Name);
+
 				FString ParameterDesc;
 				Material->GetParameterDesc(MaterialParameterInfo, ParameterDesc);
 
@@ -2274,6 +2297,10 @@ void FNiagaraRendererMaterialStaticBoolParameterCustomization::GetMaterialParame
 {
 	TArray<FGuid> Guids;
 	Material->GetAllStaticSwitchParameterInfo(OutMaterialParameterInfos, Guids);
+	for (FMaterialParameterInfo& MaterialParameterInfo : OutMaterialParameterInfos)
+	{
+		MaterialParameterInfo.Association = EMaterialParameterAssociation::GlobalParameter;
+	}
 }
 
 TSharedRef<SWidget> FNiagaraRendererMaterialStaticBoolParameterCustomization::OnGetStaticVariablelBindingNameMenuContent(TSharedPtr<IPropertyHandle> PropertyHandle) const
