Niagara - Clear out the thumbnail image on emitters that are owned by systems or other emitters to prevent issues with cook.

diff --git a/Engine/Plugins/FX/Niagara/Source/Niagara/Private/NiagaraEmitter.cpp b/Engine/Plugins/FX/Niagara/Source/Niagara/Private/NiagaraEmitter.cpp
--- a/Engine/Plugins/FX/Niagara/Source/Niagara/Private/NiagaraEmitter.cpp
+++ b/Engine/Plugins/FX/Niagara/Source/Niagara/Private/NiagaraEmitter.cpp
@@ -656,9 +656,9 @@ void UNiagaraEmitter::PostLoad()
 		LibraryVisibility = ENiagaraScriptLibraryVisibility::Library;
 	}
 
-	if (this->IsAsset() == false && GetPackage()->bIsCookedForEditor)
+	if (this->GetOuter()->IsA<UNiagaraSystem>() || this->GetOuter()->IsA<UNiagaraEmitter>())
 	{
-		// Remove thunbnails for non-asset emitters in cooked packages to prevent problems due to issues with cooked data being referenced by non-cooked data.
+		// Remove thunbnails for non-asset emitters to prevent problems with cooked for editor emitters being referenced by uncooked emitters and systems.
 		ThumbnailImage = nullptr;
 	}
 #endif
