- Add Mip Levels to Niagara texture debugging

#rnx

diff --git a/Engine/Plugins/FX/Niagara/Source/Niagara/Private/NiagaraGpuComputeDebug.cpp b/Engine/Plugins/FX/Niagara/Source/Niagara/Private/NiagaraGpuComputeDebug.cpp
--- a/Engine/Plugins/FX/Niagara/Source/Niagara/Private/NiagaraGpuComputeDebug.cpp
+++ b/Engine/Plugins/FX/Niagara/Source/Niagara/Private/NiagaraGpuComputeDebug.cpp
@@ -169,6 +169,7 @@ void FNiagaraGpuComputeDebug::AddAttributeTexture(FRDGBuilder& GraphBuilder, FNi
 		const FRHITextureCreateDesc NewTextureDesc =
 			FRHITextureCreateDesc(TEXT("FNiagaraGpuComputeDebug"), TextureDesc.Dimension)
 			.SetExtent(TextureDesc.Extent)
+			.SetNumMips(TextureDesc.NumMips)
 			.SetDepth(TextureDesc.Depth)
 			.SetArraySize(TextureDesc.ArraySize)
 			.SetFormat(TextureDesc.Format)
