- Cook out renderers with static bool bindings that evaluate to false

diff --git a/Engine/Plugins/FX/Niagara/Source/Niagara/Classes/NiagaraEmitter.h b/Engine/Plugins/FX/Niagara/Source/Niagara/Classes/NiagaraEmitter.h
--- a/Engine/Plugins/FX/Niagara/Source/Niagara/Classes/NiagaraEmitter.h
+++ b/Engine/Plugins/FX/Niagara/Source/Niagara/Classes/NiagaraEmitter.h
@@ -626,7 +626,10 @@ public:
 
 	/** Returns all available versions for this emitter. */
 	NIAGARA_API virtual TArray<FNiagaraAssetVersion> GetAllAvailableVersions() const override;
-	
+
+	template <typename TAction>
+	void ForEachVersionData(TAction Func) const;
+
 #if WITH_EDITORONLY_DATA
 	NIAGARA_API virtual TSharedPtr<FNiagaraVersionDataAccessor> GetVersionDataAccessor(const FGuid& Version) override; 
 
@@ -992,6 +995,15 @@ private:
 	void ResolveScalabilitySettings();
 };
 
+template <typename TAction>
+void UNiagaraEmitter::ForEachVersionData(TAction Func) const
+{
+	for (const FVersionedNiagaraEmitterData& Data : VersionData)
+	{
+		Func(Data);
+	}
+}
+
 template <typename TAction>
 void FVersionedNiagaraEmitterData::ForEachRenderer(TAction Func) const
 {
diff --git a/Engine/Plugins/FX/Niagara/Source/Niagara/Private/NiagaraRendererProperties.cpp b/Engine/Plugins/FX/Niagara/Source/Niagara/Private/NiagaraRendererProperties.cpp
--- a/Engine/Plugins/FX/Niagara/Source/Niagara/Private/NiagaraRendererProperties.cpp
+++ b/Engine/Plugins/FX/Niagara/Source/Niagara/Private/NiagaraRendererProperties.cpp
@@ -19,6 +19,71 @@
 
 #define LOCTEXT_NAMESPACE "UNiagaraRendererProperties"
 
+#if WITH_EDITORONLY_DATA
+int32 GNiagaraRendererCookOutStaticEnabledBinding = 1;
+static FAutoConsoleVariableRef CVarNiagaraRendererCookOutStaticEnabledBinding(
+	TEXT("fx.Niagara.Renderer.CookOutStaticEnabledBinding"),
+	GNiagaraRendererCookOutStaticEnabledBinding,
+	TEXT("If none zero renderers with static variables used for enabled binding will cook out if they are not enabled."),
+	ECVF_Scalability
+);
+
+// Attempts to resolve a static variable across all the scripts
+// If the state is undetermined (i.e. does not exist or is inconsistent across scripts) we will not return a value
+static TOptional<bool> TryResolveStaticVariableBool(const UNiagaraEmitter* NiagaraEmitter, FNiagaraVariableBase BoundVariable)
+{
+	BoundVariable.SetType(BoundVariable.GetType().ToStaticDef());
+
+	TOptional<bool> StaticValue;
+	bool bHasConflictingValues = false;
+	auto FindStaticValue =
+		[&](UNiagaraScript* NiagaraScript)
+		{
+			if (NiagaraScript == nullptr || StaticValue.IsSet())
+			{
+				return;
+			}
+
+			for (const FNiagaraVariable& StaticVariable : NiagaraScript->GetVMExecutableData().StaticVariablesWritten)
+			{
+				if (static_cast<const FNiagaraVariableBase&>(StaticVariable) != BoundVariable)
+				{
+					continue;
+				}
+				if (StaticValue.IsSet())
+				{
+					bHasConflictingValues = StaticValue.GetValue() != StaticVariable.GetValue<bool>();
+				}
+				else
+				{
+					StaticValue = StaticVariable.GetValue<bool>();
+				}
+				break;
+			}
+		};
+
+	NiagaraEmitter->ForEachVersionData(
+		[&](const FVersionedNiagaraEmitterData& EmitterData)
+		{
+			EmitterData.ForEachScript(FindStaticValue);
+		}
+	);
+
+	if (UNiagaraSystem* NiagaraSystem = NiagaraEmitter->GetTypedOuter<UNiagaraSystem>())
+	{
+		FindStaticValue(NiagaraSystem->GetSystemSpawnScript());
+		FindStaticValue(NiagaraSystem->GetSystemUpdateScript());
+	}
+
+	if (bHasConflictingValues)
+	{
+		StaticValue.Reset();
+	}
+
+	return StaticValue;
+}
+#endif
+
 void FNiagaraRendererLayout::Initialize(int32 NumVariables)
 {
 	VFVariables_GT.Reset(NumVariables);
@@ -662,7 +727,17 @@ bool UNiagaraRendererProperties::NeedsLoadForTargetPlatform(const ITargetPlatfor
 	{
 		if (OwnerEmitter->NeedsLoadForTargetPlatform(TargetPlatform))
 		{
-			return bIsEnabled && Platforms.IsEnabledForPlatform(TargetPlatform->IniPlatformName());
+			if (bIsEnabled && Platforms.IsEnabledForPlatform(TargetPlatform->IniPlatformName()))
+			{
+			#if WITH_EDITORONLY_DATA
+				if (GNiagaraRendererCookOutStaticEnabledBinding && RendererEnabledBinding.IsValid())
+				{
+					TOptional<bool> ResolvedStaticValue = TryResolveStaticVariableBool(OwnerEmitter, RendererEnabledBinding.GetParamMapBindableVariable());
+					return ResolvedStaticValue.Get(true);
+				}
+			#endif
+				return true;
+			}
 		}
 	}
 
