- Remove attribute appending to grid names as it can result in bad strings

#rnx

diff --git a/Engine/Plugins/FX/Niagara/Source/Niagara/Private/NiagaraDataInterfaceGrid2DCollection.cpp b/Engine/Plugins/FX/Niagara/Source/Niagara/Private/NiagaraDataInterfaceGrid2DCollection.cpp
--- a/Engine/Plugins/FX/Niagara/Source/Niagara/Private/NiagaraDataInterfaceGrid2DCollection.cpp
+++ b/Engine/Plugins/FX/Niagara/Source/Niagara/Private/NiagaraDataInterfaceGrid2DCollection.cpp
@@ -2538,6 +2538,8 @@ void FGrid2DCollectionRWInstanceData_RenderThread::BeginSimulate(FRDGBuilder& Gr
 
 		const FRDGTextureDesc TextureDesc = FRDGTextureDesc::Create2DArray(NumCells, PixelFormat, FClearValueBinding::Black, ETextureCreateFlags::ShaderResource | ETextureCreateFlags::UAV, NumAttributes);
 
+		const TCHAR* GridTextureName = TEXT("Grid2D::GridTexture");
+	#if 0
 		TStringBuilder<128> StringBuilder;
 		SourceDIName.AppendString(StringBuilder);
 		StringBuilder.Append("_");
@@ -2547,8 +2549,10 @@ void FGrid2DCollectionRWInstanceData_RenderThread::BeginSimulate(FRDGBuilder& Gr
 			StringBuilder.Append("_");
 		}
 		StringBuilder.Append("Grid2DCollection");
+		GridTextureName = *StringBuilder;
+	#endif
+		DestinationData->Initialize(GraphBuilder, GridTextureName, TextureDesc);
 
-		DestinationData->Initialize(GraphBuilder, *StringBuilder, TextureDesc);
 		// This destination buffer will sometimes have old data in it.  Force it to clear.
 		AddClearUAVPass(GraphBuilder, DestinationData->GetOrCreateUAV(GraphBuilder), FVector4f(ForceInitToZero));
 	}
diff --git a/Engine/Plugins/FX/Niagara/Source/Niagara/Private/NiagaraDataInterfaceGrid3DCollection.cpp b/Engine/Plugins/FX/Niagara/Source/Niagara/Private/NiagaraDataInterfaceGrid3DCollection.cpp
--- a/Engine/Plugins/FX/Niagara/Source/Niagara/Private/NiagaraDataInterfaceGrid3DCollection.cpp
+++ b/Engine/Plugins/FX/Niagara/Source/Niagara/Private/NiagaraDataInterfaceGrid3DCollection.cpp
@@ -4103,6 +4103,8 @@ void FGrid3DCollectionRWInstanceData_RenderThread::BeginSimulate(FRDGBuilder& Gr
 		const FIntVector TextureSize(NumCells.X * NumTiles.X, NumCells.Y * NumTiles.Y, NumCells.Z * NumTiles.Z);
 		const FRDGTextureDesc TextureDesc = FRDGTextureDesc::Create3D(TextureSize, PixelFormat, FClearValueBinding::Black, ETextureCreateFlags::ShaderResource | ETextureCreateFlags::UAV);
 
+		const TCHAR* GridTextureName = TEXT("Grid3D::GridTexture");
+	#if 0
 		TStringBuilder<128> StringBuilder;
 		SourceDIName.AppendString(StringBuilder);
 		StringBuilder.Append("_");
@@ -4112,8 +4114,9 @@ void FGrid3DCollectionRWInstanceData_RenderThread::BeginSimulate(FRDGBuilder& Gr
 			StringBuilder.Append("_");
 		}
 		StringBuilder.Append("Grid3DCollection");
-
-		DestinationData->Initialize(GraphBuilder, *StringBuilder, TextureDesc);
+		GridTextureName = *StringBuilder;
+	#endif
+		DestinationData->Initialize(GraphBuilder, GridTextureName, TextureDesc);
 
 		// This destination buffer will sometimes have old data in it.  Force it to clear.
 		AddClearUAVPass(GraphBuilder, DestinationData->GetOrCreateUAV(GraphBuilder), FVector4f(ForceInitToZero));
