- Only warn about CPU access if the interface is flagged as using it
- This fixes false warnings on emitters that are cooked out

#rnx

diff --git a/Engine/Plugins/FX/Niagara/Source/Niagara/Private/DataInterface/NiagaraDataInterfaceStaticMesh.cpp b/Engine/Plugins/FX/Niagara/Source/Niagara/Private/DataInterface/NiagaraDataInterfaceStaticMesh.cpp
--- a/Engine/Plugins/FX/Niagara/Source/Niagara/Private/DataInterface/NiagaraDataInterfaceStaticMesh.cpp
+++ b/Engine/Plugins/FX/Niagara/Source/Niagara/Private/DataInterface/NiagaraDataInterfaceStaticMesh.cpp
@@ -791,13 +791,16 @@ namespace NDIStaticMeshLocal
 				{
 					if (!FNiagaraUtilities::AreBufferSRVsAlwaysCreated(GMaxRHIShaderPlatform))
 					{
-						UE_LOG(LogNiagara, Log, TEXT("NiagaraStaticMeshDataInterface used by GPU emitter but does not have SRV access on this platform.  Enable CPU access to fix this issue. Interface: %s, Mesh: %s"), *GetFullNameSafe(Interface), *GetFullNameSafe(StaticMesh));
+						UE_LOG(LogNiagara, Log, TEXT("NiagaraStaticMeshDataInterface used by GPU emitter but does not have SRV access on this platform.  Enable CPU access to fix this issue. System: %s, Mesh: %s"), *GetFullNameSafe(SystemInstance->GetSystem()), *GetFullNameSafe(StaticMesh));
 						StaticMesh = nullptr;
 					}
 				}
 				else
 				{
-					UE_LOG(LogNiagara, Log, TEXT("NiagaraStaticMeshDataInterface used by CPU emitter and does not allow CPU access. Interface: %s, Mesh: %s"), *GetFullNameSafe(Interface), *GetFullNameSafe(StaticMesh));
+					if (Interface->IsUsedWithCPUEmitter())
+					{
+						UE_LOG(LogNiagara, Log, TEXT("NiagaraStaticMeshDataInterface used by CPU emitter and does not allow CPU access. System: %s, Mesh: %s"), *GetFullNameSafe(SystemInstance->GetSystem()), *GetFullNameSafe(StaticMesh));
+					}
 					StaticMesh = nullptr;
 				}
 			}
