- Deprecate Get methods from RT in favor of more explicit Load and add MipLevel support
- Allow sample / load operations to be visible by default, user beware as you can not safely read / write to a RT in the same stage
- Add SetFormat method to other RT types

diff --git a/Engine/Plugins/FX/Niagara/Shaders/Private/NiagaraDataInterfaceRenderTarget2DArrayTemplate.ush b/Engine/Plugins/FX/Niagara/Shaders/Private/NiagaraDataInterfaceRenderTarget2DArrayTemplate.ush
--- a/Engine/Plugins/FX/Niagara/Shaders/Private/NiagaraDataInterfaceRenderTarget2DArrayTemplate.ush
+++ b/Engine/Plugins/FX/Niagara/Shaders/Private/NiagaraDataInterfaceRenderTarget2DArrayTemplate.ush
@@ -1,6 +1,7 @@
 // Copyright Epic Games, Inc. All Rights Reserved.
 
 int3						{ParameterName}_TextureSize;
+int							{ParameterName}_MipLevels;
 RWTexture2DArray<float4>	{ParameterName}_RWTexture;
 Texture2DArray<float4>		{ParameterName}_Texture;
 SamplerState				{ParameterName}_TextureSampler;
@@ -18,6 +19,11 @@ void GetRenderTargetValue_{ParameterName}(int IndexX, int IndexY, int IndexZ, ou
 	Value = {ParameterName}_Texture.Load(int4(IndexX, IndexY, IndexZ, 0));
 }
 
+void LoadRenderTargetValue_{ParameterName}(int IndexX, int IndexY, int IndexZ, int MipLevel, out float4 Value)
+{
+	Value = {ParameterName}_Texture.Load(int4(IndexX, IndexY, IndexZ, MipLevel));
+}
+
 void SampleRenderTargetValue_{ParameterName}(float2 UV, int Slice, out float4 Value)
 {
 	Value = {ParameterName}_Texture.SampleLevel({ParameterName}_TextureSampler, float3(UV.x, UV.y, Slice), 0.0f);
@@ -30,6 +36,11 @@ void GetRenderTargetSize_{ParameterName}(out int Width, out int Height, out int
 	Slices = {ParameterName}_TextureSize.z;
 }
 
+void GetNumMipLevels_{ParameterName}(out int OutMipLevels)
+{
+	OutMipLevels = {ParameterName}_MipLevels;
+}
+
 void LinearToIndex_{ParameterName}(int Linear, out int IndexX, out int IndexY, out int IndexZ)
 {
 	IndexX = Linear % {ParameterName}_TextureSize.x;
diff --git a/Engine/Plugins/FX/Niagara/Shaders/Private/NiagaraDataInterfaceRenderTarget2DTemplate.ush b/Engine/Plugins/FX/Niagara/Shaders/Private/NiagaraDataInterfaceRenderTarget2DTemplate.ush
--- a/Engine/Plugins/FX/Niagara/Shaders/Private/NiagaraDataInterfaceRenderTarget2DTemplate.ush
+++ b/Engine/Plugins/FX/Niagara/Shaders/Private/NiagaraDataInterfaceRenderTarget2DTemplate.ush
@@ -1,6 +1,7 @@
 // Copyright Epic Games, Inc. All Rights Reserved.
 
 int2					{ParameterName}_TextureSize;
+int						{ParameterName}_MipLevels;
 RWTexture2D<float4>		{ParameterName}_RWTexture;
 Texture2D<float4>		{ParameterName}_Texture;
 SamplerState			{ParameterName}_TextureSampler;
@@ -18,6 +19,11 @@ void GetRenderTargetValue_{ParameterName}(int IndexX, int IndexY, out float4 Val
 	Value = {ParameterName}_Texture.Load(int3(IndexX, IndexY, 0));
 }
 
+void LoadRenderTargetValue_{ParameterName}(int IndexX, int IndexY, int MipLevel, out float4 Value)
+{
+	Value = {ParameterName}_Texture.Load(int3(IndexX, IndexY, MipLevel));
+}
+
 void SampleRenderTargetValue_{ParameterName}(float2 UV, out float4 Value)
 {
 	Value = {ParameterName}_Texture.SampleLevel({ParameterName}_TextureSampler, UV, 0.0f);
@@ -29,6 +35,11 @@ void GetRenderTargetSize_{ParameterName}(out int Width, out int Height)
 	Height = {ParameterName}_TextureSize.y;
 }
 
+void GetNumMipLevels_{ParameterName}(out int OutMipLevels)
+{
+	OutMipLevels = {ParameterName}_MipLevels;
+}
+
 void LinearToIndex_{ParameterName}(int Linear, out int IndexX, out int IndexY)
 {
 	IndexX = Linear % {ParameterName}_TextureSize.x;
diff --git a/Engine/Plugins/FX/Niagara/Shaders/Private/NiagaraDataInterfaceRenderTargetCubeTemplate.ush b/Engine/Plugins/FX/Niagara/Shaders/Private/NiagaraDataInterfaceRenderTargetCubeTemplate.ush
--- a/Engine/Plugins/FX/Niagara/Shaders/Private/NiagaraDataInterfaceRenderTargetCubeTemplate.ush
+++ b/Engine/Plugins/FX/Niagara/Shaders/Private/NiagaraDataInterfaceRenderTargetCubeTemplate.ush
@@ -1,6 +1,7 @@
 // Copyright Epic Games, Inc. All Rights Reserved.
 
 int						{ParameterName}_TextureSize;
+int						{ParameterName}_MipLevels;
 RWTextureCube<float4>	{ParameterName}_RWTexture;
 TextureCube<float4>		{ParameterName}_Texture;
 SamplerState			{ParameterName}_TextureSampler;
@@ -18,6 +19,11 @@ void SetRenderTargetValue_{ParameterName}_UEImpureCall(bool bExecute, int IndexX
 //{
 //	Value = {ParameterName}_Texture.Load(int4(IndexX, IndexY, Face, 0));
 //}
+//
+//void LoadRenderTargetValue_{ParameterName}(int IndexX, int IndexY, int Face, int MipLevel, out float4 Value)
+//{
+//	Value = {ParameterName}_Texture.Load(int4(IndexX, IndexY, Face, MipLevel));
+//}
 
 void SampleRenderTargetValue_{ParameterName}(float3 UVW, out float4 Value)
 {
@@ -29,6 +35,11 @@ void GetRenderTargetSize_{ParameterName}(out int Size)
 	Size = {ParameterName}_TextureSize;
 }
 
+void GetNumMipLevels_{ParameterName}(out int OutMipLevels)
+{
+	OutMipLevels = {ParameterName}_MipLevels;
+}
+
 void LinearToIndex_{ParameterName}(int Linear, out int IndexX, out int IndexY, out int Face)
 {
 	IndexX = Linear % {ParameterName}_TextureSize;
diff --git a/Engine/Plugins/FX/Niagara/Shaders/Private/NiagaraDataInterfaceRenderTargetVolumeTemplate.ush b/Engine/Plugins/FX/Niagara/Shaders/Private/NiagaraDataInterfaceRenderTargetVolumeTemplate.ush
--- a/Engine/Plugins/FX/Niagara/Shaders/Private/NiagaraDataInterfaceRenderTargetVolumeTemplate.ush
+++ b/Engine/Plugins/FX/Niagara/Shaders/Private/NiagaraDataInterfaceRenderTargetVolumeTemplate.ush
@@ -1,6 +1,7 @@
 // Copyright Epic Games, Inc. All Rights Reserved.
 
 int3					{ParameterName}_TextureSize;
+int						{ParameterName}_MipLevels;
 RWTexture3D<float4>		{ParameterName}_RWTexture;
 Texture3D<float4>		{ParameterName}_Texture;
 SamplerState			{ParameterName}_TextureSampler;
@@ -18,6 +19,11 @@ void GetRenderTargetValue_{ParameterName}(int IndexX, int IndexY, int IndexZ, ou
 	Value = {ParameterName}_Texture.Load(int4(IndexX, IndexY, IndexZ, 0));
 }
 
+void LoadRenderTargetValue_{ParameterName}(int IndexX, int IndexY, int IndexZ, int MipLevel, out float4 Value)
+{
+	Value = {ParameterName}_Texture.Load(int4(IndexX, IndexY, IndexZ, MipLevel));
+}
+
 void SampleRenderTargetValue_{ParameterName}(float3 UVW, out float4 Value)
 {
 	Value = {ParameterName}_Texture.SampleLevel({ParameterName}_TextureSampler, UVW, 0.0f);
@@ -30,6 +36,11 @@ void GetRenderTargetSize_{ParameterName}(out int Width, out int Height, out int
 	Slices = {ParameterName}_TextureSize.z;
 }
 
+void GetNumMipLevels_{ParameterName}(out int OutMipLevels)
+{
+	OutMipLevels = {ParameterName}_MipLevels;
+}
+
 void LinearToIndex_{ParameterName}(int Linear, out int IndexX, out int IndexY, out int IndexZ)
 {
 	IndexX = Linear % {ParameterName}_TextureSize.x;
diff --git a/Engine/Plugins/FX/Niagara/Source/Niagara/Classes/NiagaraDataInterfaceRenderTarget2D.h b/Engine/Plugins/FX/Niagara/Source/Niagara/Classes/NiagaraDataInterfaceRenderTarget2D.h
--- a/Engine/Plugins/FX/Niagara/Source/Niagara/Classes/NiagaraDataInterfaceRenderTarget2D.h
+++ b/Engine/Plugins/FX/Niagara/Source/Niagara/Classes/NiagaraDataInterfaceRenderTarget2D.h
@@ -46,6 +46,7 @@ struct FRenderTarget2DRWInstanceData_RenderThread
 	}
 
 	FIntPoint Size = FIntPoint(EForceInit::ForceInitToZero);
+	int MipLevels = 0;
 	ENiagaraMipMapGeneration MipMapGeneration = ENiagaraMipMapGeneration::Disabled;
 	ENiagaraMipMapGenerationType MipMapGenerationType = ENiagaraMipMapGenerationType::Linear;
 	bool bRebuildMips = false;
@@ -132,8 +133,10 @@ public:
 	virtual bool RenderVariableToCanvas(FNiagaraSystemInstanceID SystemInstanceID, FName VariableName, class FCanvas* Canvas, const FIntRect& DrawRect) const override;
 	//~ UNiagaraDataInterface interface END
 
-	void VMGetSize(FVectorVMExternalFunctionContext& Context); 
+	void VMGetSize(FVectorVMExternalFunctionContext& Context);
 	void VMSetSize(FVectorVMExternalFunctionContext& Context);
+	void VMGetNumMipLevels(FVectorVMExternalFunctionContext& Context);
+	void VMSetFormat(FVectorVMExternalFunctionContext& Context);
 
 	UPROPERTY(EditAnywhere, Category = "Render Target", meta = (DisplayPriority = 2, EditCondition = "!bInheritUserParameterSettings", EditConditionHides))
 	FIntPoint Size;
diff --git a/Engine/Plugins/FX/Niagara/Source/Niagara/Classes/NiagaraDataInterfaceRenderTarget2DArray.h b/Engine/Plugins/FX/Niagara/Source/Niagara/Classes/NiagaraDataInterfaceRenderTarget2DArray.h
--- a/Engine/Plugins/FX/Niagara/Source/Niagara/Classes/NiagaraDataInterfaceRenderTarget2DArray.h
+++ b/Engine/Plugins/FX/Niagara/Source/Niagara/Classes/NiagaraDataInterfaceRenderTarget2DArray.h
@@ -43,6 +43,7 @@ struct FRenderTarget2DArrayRWInstanceData_RenderThread
 	}
 
 	FIntVector Size = FIntVector(EForceInit::ForceInitToZero);
+	int MipLevels = 0;
 	bool bWroteThisFrame = false;
 	bool bReadThisFrame = false;
 
@@ -124,8 +125,10 @@ public:
 	virtual bool GetExposedVariableValue(const FNiagaraVariableBase& InVariable, void* InPerInstanceData, FNiagaraSystemInstance* InSystemInstance, void* OutData) const override;
 
 	//~ UNiagaraDataInterface interface END
-	void VMGetSize(FVectorVMExternalFunctionContext& Context); 
+	void VMGetSize(FVectorVMExternalFunctionContext& Context);
 	void VMSetSize(FVectorVMExternalFunctionContext& Context);
+	void VMGetNumMipLevels(FVectorVMExternalFunctionContext& Context);
+	void VMSetFormat(FVectorVMExternalFunctionContext& Context);
 
 	UPROPERTY(EditAnywhere, Category = "Render Target", meta = (DisplayPriority = 2, EditCondition = "!bInheritUserParameterSettings", EditConditionHides))
 	FIntVector Size;
diff --git a/Engine/Plugins/FX/Niagara/Source/Niagara/Classes/NiagaraDataInterfaceRenderTargetCube.h b/Engine/Plugins/FX/Niagara/Source/Niagara/Classes/NiagaraDataInterfaceRenderTargetCube.h
--- a/Engine/Plugins/FX/Niagara/Source/Niagara/Classes/NiagaraDataInterfaceRenderTargetCube.h
+++ b/Engine/Plugins/FX/Niagara/Source/Niagara/Classes/NiagaraDataInterfaceRenderTargetCube.h
@@ -42,6 +42,7 @@ struct FRenderTargetCubeRWInstanceData_RenderThread
 	}
 
 	int Size = 0;
+	int MipLevels = 0;
 	bool bWroteThisFrame = false;
 	bool bReadThisFrame = false;
 	bool bNeedsTransition = false;
@@ -119,6 +120,8 @@ public:
 	
 	void VMGetSize(FVectorVMExternalFunctionContext& Context);
 	void VMSetSize(FVectorVMExternalFunctionContext& Context);
+	void VMGetNumMipLevels(FVectorVMExternalFunctionContext& Context);
+	void VMSetFormat(FVectorVMExternalFunctionContext& Context);
 
 	UPROPERTY(EditAnywhere, Category = "Render Target", meta = (DisplayPriority = 2, EditCondition = "!bInheritUserParameterSettings", EditConditionHides))
 	int Size;
diff --git a/Engine/Plugins/FX/Niagara/Source/Niagara/Classes/NiagaraDataInterfaceRenderTargetVolume.h b/Engine/Plugins/FX/Niagara/Source/Niagara/Classes/NiagaraDataInterfaceRenderTargetVolume.h
--- a/Engine/Plugins/FX/Niagara/Source/Niagara/Classes/NiagaraDataInterfaceRenderTargetVolume.h
+++ b/Engine/Plugins/FX/Niagara/Source/Niagara/Classes/NiagaraDataInterfaceRenderTargetVolume.h
@@ -41,6 +41,7 @@ struct FRenderTargetVolumeRWInstanceData_RenderThread
 	}
 
 	FIntVector Size = FIntVector(EForceInit::ForceInitToZero);
+	int MipLevels = 0;
 	bool bWroteThisFrame = false;
 	bool bReadThisFrame = false;
 
@@ -128,8 +129,9 @@ public:
 	virtual bool SimCacheReadFrame(UObject* StorageObject, int FrameA, int FrameB, float Interp, FNiagaraSystemInstance* SystemInstance, void* OptionalPerInstanceData) override;
 	//~ UNiagaraDataInterface interface END
 	
-	void VMGetSize(FVectorVMExternalFunctionContext& Context); 
+	void VMGetSize(FVectorVMExternalFunctionContext& Context);
 	void VMSetSize(FVectorVMExternalFunctionContext& Context);
+	void VMGetNumMipLevels(FVectorVMExternalFunctionContext& Context);
 	void VMSetFormat(FVectorVMExternalFunctionContext& Context);
 
 	UPROPERTY(EditAnywhere, Category = "Render Target", meta = (DisplayPriority = 2, EditCondition = "!bInheritUserParameterSettings", EditConditionHides))
diff --git a/Engine/Plugins/FX/Niagara/Source/Niagara/Private/NiagaraDataInterfaceRenderTarget2D.cpp b/Engine/Plugins/FX/Niagara/Source/Niagara/Private/NiagaraDataInterfaceRenderTarget2D.cpp
--- a/Engine/Plugins/FX/Niagara/Source/Niagara/Private/NiagaraDataInterfaceRenderTarget2D.cpp
+++ b/Engine/Plugins/FX/Niagara/Source/Niagara/Private/NiagaraDataInterfaceRenderTarget2D.cpp
@@ -28,6 +28,7 @@ namespace NDIRenderTarget2DLocal
 {
 	BEGIN_SHADER_PARAMETER_STRUCT(FShaderParameters, )
 		SHADER_PARAMETER(FIntPoint,								TextureSize)
+		SHADER_PARAMETER(int,									MipLevels)
 		SHADER_PARAMETER_RDG_TEXTURE_UAV(RWTexture2D<float4>,	RWTexture)
 		SHADER_PARAMETER_RDG_TEXTURE_SRV(Texture2D<float4>,		Texture)
 		SHADER_PARAMETER_SAMPLER(SamplerState,					TextureSampler)
@@ -37,13 +38,16 @@ namespace NDIRenderTarget2DLocal
 
 	// Global VM function names, also used by the shaders code generation methods.
 	static const FName SetValueFunctionName("SetRenderTargetValue");
-	static const FName GetValueFunctionName("GetRenderTargetValue");
+	static const FName LoadValueFunctionName("LoadRenderTargetValue");
 	static const FName SampleValueFunctionName("SampleRenderTargetValue");
 	static const FName SetSizeFunctionName("SetRenderTargetSize");
 	static const FName GetSizeFunctionName("GetRenderTargetSize");
+	static const FName GetNumMipLevelsName("GetNumMipLevels");
+	static const FName SetFormatFunctionName("SetRenderTargetFormat");
 	static const FName LinearToIndexName("LinearToIndex");
 	static const FName ExecToIndexName("ExecToIndex");
 	static const FName ExecToUnitName("ExecToUnit");
+	static const FName Deprecated_GetValueFunctionName("GetRenderTargetValue");
 }
 
 FNiagaraVariableBase UNiagaraDataInterfaceRenderTarget2D::ExposedRTVar;
@@ -114,152 +118,110 @@ void UNiagaraDataInterfaceRenderTarget2D::GetFunctions(TArray<FNiagaraFunctionSi
 	const int32 EmitterSystemOnlyBitmask = ENiagaraScriptUsageMask::Emitter | ENiagaraScriptUsageMask::System;
 	OutFunctions.Reserve(OutFunctions.Num() + 4);
 
+	FNiagaraFunctionSignature DefaultSig;
+	DefaultSig.Inputs.Emplace(FNiagaraTypeDefinition(GetClass()), TEXT("RenderTarget"));
+	DefaultSig.bMemberFunction = true;
+	DefaultSig.bRequiresContext = false;
+	DefaultSig.bSupportsCPU = true;
+	DefaultSig.bSupportsGPU = true;
+	DefaultSig.SetFunctionVersion(FNDIRenderTarget2DFunctionVersion::LatestVersion);
+
 	{
-		FNiagaraFunctionSignature& Sig = OutFunctions.AddDefaulted_GetRef();
+		FNiagaraFunctionSignature& Sig = OutFunctions.Add_GetRef(DefaultSig);
 		Sig.Name = GetSizeFunctionName;
-		Sig.Inputs.Add(FNiagaraVariable(FNiagaraTypeDefinition(GetClass()), TEXT("RenderTarget")));
-		Sig.Outputs.Add(FNiagaraVariable(FNiagaraTypeDefinition::GetIntDef(), TEXT("Width")));
-		Sig.Outputs.Add(FNiagaraVariable(FNiagaraTypeDefinition::GetIntDef(), TEXT("Height")));
-		Sig.bExperimental = true;
-		Sig.bMemberFunction = true;
-		Sig.bRequiresContext = false;
-	#if WITH_EDITORONLY_DATA
-		Sig.FunctionVersion = FNDIRenderTarget2DFunctionVersion::LatestVersion;
-	#endif
+		Sig.Outputs.Emplace(FNiagaraTypeDefinition::GetIntDef(), TEXT("Width"));
+		Sig.Outputs.Emplace(FNiagaraTypeDefinition::GetIntDef(), TEXT("Height"));
 	}
 
 	{
-		FNiagaraFunctionSignature& Sig = OutFunctions.AddDefaulted_GetRef();
+		FNiagaraFunctionSignature& Sig = OutFunctions.Add_GetRef(DefaultSig);
 		Sig.Name = SetSizeFunctionName;
-		Sig.Inputs.Add(FNiagaraVariable(FNiagaraTypeDefinition(GetClass()), TEXT("RenderTarget")));
-		Sig.Inputs.Add(FNiagaraVariable(FNiagaraTypeDefinition::GetIntDef(), TEXT("Width")));
-		Sig.Inputs.Add(FNiagaraVariable(FNiagaraTypeDefinition::GetIntDef(), TEXT("Height")));
-		Sig.Outputs.Add(FNiagaraVariable(FNiagaraTypeDefinition::GetBoolDef(), TEXT("Success")));
+		Sig.Inputs.Emplace(FNiagaraTypeDefinition::GetIntDef(), TEXT("Width"));
+		Sig.Inputs.Emplace(FNiagaraTypeDefinition::GetIntDef(), TEXT("Height"));
+		Sig.Outputs.Emplace(FNiagaraTypeDefinition::GetBoolDef(), TEXT("Success"));
+		Sig.ModuleUsageBitmask = EmitterSystemOnlyBitmask;
+		Sig.bRequiresExecPin = true;
+		Sig.bSupportsGPU = false;
+	}
+
+	{
+		FNiagaraFunctionSignature& Sig = OutFunctions.Add_GetRef(DefaultSig);
+		Sig.Name = GetNumMipLevelsName;
+		Sig.Outputs.Emplace(FNiagaraTypeDefinition::GetIntDef(), TEXT("MipLevels"));
+	}
 
+	{
+		FNiagaraFunctionSignature& Sig = OutFunctions.Add_GetRef(DefaultSig);
+		Sig.Name = SetFormatFunctionName;
+		Sig.Inputs.Emplace(FNiagaraTypeDefinition(StaticEnum<ETextureRenderTargetFormat>()), TEXT("Format"));
+		Sig.Outputs.Emplace(FNiagaraTypeDefinition::GetBoolDef(), TEXT("Success"));
 		Sig.ModuleUsageBitmask = EmitterSystemOnlyBitmask;
-		Sig.bExperimental = true;
-		Sig.bMemberFunction = true;
 		Sig.bRequiresExecPin = true;
-		Sig.bRequiresContext = false;
-		Sig.bSupportsCPU = true;
 		Sig.bSupportsGPU = false;
-	#if WITH_EDITORONLY_DATA
-		Sig.FunctionVersion = FNDIRenderTarget2DFunctionVersion::LatestVersion;
-	#endif
 	}
 
 	{
-		FNiagaraFunctionSignature& Sig = OutFunctions.AddDefaulted_GetRef();
+		FNiagaraFunctionSignature& Sig = OutFunctions.Add_GetRef(DefaultSig);
 		Sig.Name = SetValueFunctionName;
-		Sig.Inputs.Add(FNiagaraVariable(FNiagaraTypeDefinition(GetClass()), TEXT("RenderTarget")));
 		Sig.Inputs.Add_GetRef(FNiagaraVariable(FNiagaraTypeDefinition::GetBoolDef(), TEXT("Enabled"))).SetValue(true);
-		Sig.Inputs.Add(FNiagaraVariable(FNiagaraTypeDefinition::GetIntDef(), TEXT("IndexX")));
-		Sig.Inputs.Add(FNiagaraVariable(FNiagaraTypeDefinition::GetIntDef(), TEXT("IndexY")));
-		Sig.Inputs.Add(FNiagaraVariable(FNiagaraTypeDefinition::GetColorDef(), TEXT("Value")));
-
-		Sig.bExperimental = true;
-		Sig.bMemberFunction = true;
+		Sig.Inputs.Emplace(FNiagaraTypeDefinition::GetIntDef(), TEXT("IndexX"));
+		Sig.Inputs.Emplace(FNiagaraTypeDefinition::GetIntDef(), TEXT("IndexY"));
+		Sig.Inputs.Emplace(FNiagaraTypeDefinition::GetColorDef(), TEXT("Value"));
 		Sig.bRequiresExecPin = true;
-		Sig.bRequiresContext = false;
 		Sig.bWriteFunction = true;
 		Sig.bSupportsCPU = false;
-		Sig.bSupportsGPU = true;
-	#if WITH_EDITORONLY_DATA
-		Sig.FunctionVersion = FNDIRenderTarget2DFunctionVersion::LatestVersion;
-	#endif
 	}
 
 	{
-		FNiagaraFunctionSignature& Sig = OutFunctions.AddDefaulted_GetRef();
-		Sig.Name = GetValueFunctionName;
-		Sig.Inputs.Add(FNiagaraVariable(FNiagaraTypeDefinition(GetClass()), TEXT("RenderTarget")));
-		Sig.Inputs.Add(FNiagaraVariable(FNiagaraTypeDefinition::GetIntDef(), TEXT("IndexX")));
-		Sig.Inputs.Add(FNiagaraVariable(FNiagaraTypeDefinition::GetIntDef(), TEXT("IndexY")));
-		Sig.Outputs.Add(FNiagaraVariable(FNiagaraTypeDefinition::GetColorDef(), TEXT("Value")));
-
-		Sig.bHidden = NiagaraDataInterfaceRenderTargetCommon::GAllowReads != 1;
-		Sig.bExperimental = true;
-		Sig.bMemberFunction = true;
-		Sig.bRequiresContext = false;
-		Sig.bWriteFunction = true;
+		FNiagaraFunctionSignature& Sig = OutFunctions.Add_GetRef(DefaultSig);
+		Sig.Name = LoadValueFunctionName;
+		Sig.Inputs.Emplace(FNiagaraTypeDefinition::GetIntDef(), TEXT("IndexX"));
+		Sig.Inputs.Emplace(FNiagaraTypeDefinition::GetIntDef(), TEXT("IndexY"));
+		Sig.Inputs.Emplace(FNiagaraTypeDefinition::GetIntDef(), TEXT("MipLevel"));
+		Sig.Outputs.Emplace(FNiagaraTypeDefinition::GetColorDef(), TEXT("Value"));
 		Sig.bSupportsCPU = false;
-		Sig.bSupportsGPU = true;
-#if WITH_EDITORONLY_DATA
-		Sig.FunctionVersion = FNDIRenderTarget2DFunctionVersion::LatestVersion;
-#endif
 	}
 
 	{
-		FNiagaraFunctionSignature& Sig = OutFunctions.AddDefaulted_GetRef();
+		FNiagaraFunctionSignature& Sig = OutFunctions.Add_GetRef(DefaultSig);
 		Sig.Name = SampleValueFunctionName;
-		Sig.Inputs.Add(FNiagaraVariable(FNiagaraTypeDefinition(GetClass()), TEXT("RenderTarget")));
-		Sig.Inputs.Add(FNiagaraVariable(FNiagaraTypeDefinition::GetVec2Def(), TEXT("UV")));
-		Sig.Outputs.Add(FNiagaraVariable(FNiagaraTypeDefinition::GetColorDef(), TEXT("Value")));
-
-		Sig.bHidden = NiagaraDataInterfaceRenderTargetCommon::GAllowReads != 1;
-		Sig.bExperimental = true;
-		Sig.bMemberFunction = true;
-		Sig.bRequiresContext = false;
-		Sig.bWriteFunction = true;
+		Sig.Inputs.Emplace(FNiagaraTypeDefinition::GetVec2Def(), TEXT("UV"));
+		Sig.Outputs.Emplace(FNiagaraTypeDefinition::GetColorDef(), TEXT("Value"));
 		Sig.bSupportsCPU = false;
-		Sig.bSupportsGPU = true;
-#if WITH_EDITORONLY_DATA
-		Sig.FunctionVersion = FNDIRenderTarget2DFunctionVersion::LatestVersion;
-#endif
 	}
 
 	{
-		FNiagaraFunctionSignature& Sig = OutFunctions.AddDefaulted_GetRef();
+		FNiagaraFunctionSignature& Sig = OutFunctions.Add_GetRef(DefaultSig);
 		Sig.Name = LinearToIndexName;
-		Sig.Inputs.Add(FNiagaraVariable(FNiagaraTypeDefinition(GetClass()), TEXT("RenderTarget")));
-		Sig.Inputs.Add(FNiagaraVariable(FNiagaraTypeDefinition::GetIntDef(), TEXT("Linear")));
-		Sig.Outputs.Add(FNiagaraVariable(FNiagaraTypeDefinition::GetIntDef(), TEXT("IndexX")));
-		Sig.Outputs.Add(FNiagaraVariable(FNiagaraTypeDefinition::GetIntDef(), TEXT("IndexY")));
-
-		Sig.bExperimental = true;
-		Sig.bMemberFunction = true;
-		Sig.bRequiresContext = false;
-		Sig.bWriteFunction = true;
+		Sig.Inputs.Emplace(FNiagaraTypeDefinition::GetIntDef(), TEXT("Linear"));
+		Sig.Outputs.Emplace(FNiagaraTypeDefinition::GetIntDef(), TEXT("IndexX"));
+		Sig.Outputs.Emplace(FNiagaraTypeDefinition::GetIntDef(), TEXT("IndexY"));
 		Sig.bSupportsCPU = false;
-		Sig.bSupportsGPU = true;
-	#if WITH_EDITORONLY_DATA
-		Sig.FunctionVersion = FNDIRenderTarget2DFunctionVersion::LatestVersion;
-	#endif
 	}
 
 	{
-		FNiagaraFunctionSignature& Sig = OutFunctions.AddDefaulted_GetRef();
+		FNiagaraFunctionSignature& Sig = OutFunctions.Add_GetRef(DefaultSig);
 		Sig.Name = ExecToIndexName;
-		Sig.Inputs.Add(FNiagaraVariable(FNiagaraTypeDefinition(GetClass()), TEXT("RenderTarget")));
-		Sig.Outputs.Add(FNiagaraVariable(FNiagaraTypeDefinition::GetIntDef(), TEXT("IndexX")));
-		Sig.Outputs.Add(FNiagaraVariable(FNiagaraTypeDefinition::GetIntDef(), TEXT("IndexY")));
-
-		Sig.bExperimental = true;
-		Sig.bMemberFunction = true;
-		Sig.bRequiresContext = false;
-		Sig.bWriteFunction = true;
+		Sig.Outputs.Emplace(FNiagaraTypeDefinition::GetIntDef(), TEXT("IndexX"));
+		Sig.Outputs.Emplace(FNiagaraTypeDefinition::GetIntDef(), TEXT("IndexY"));
 		Sig.bSupportsCPU = false;
-		Sig.bSupportsGPU = true;
-	#if WITH_EDITORONLY_DATA
-		Sig.FunctionVersion = FNDIRenderTarget2DFunctionVersion::LatestVersion;
-	#endif
 	}
 
 	{
-		FNiagaraFunctionSignature& Sig = OutFunctions.AddDefaulted_GetRef();
+		FNiagaraFunctionSignature& Sig = OutFunctions.Add_GetRef(DefaultSig);
 		Sig.Name = ExecToUnitName;
-		Sig.Inputs.Emplace(FNiagaraTypeDefinition(GetClass()), TEXT("RenderTarget"));
 		Sig.Outputs.Emplace(FNiagaraTypeDefinition::GetVec2Def(), TEXT("Unit"));
+		Sig.bSupportsCPU = false;
+	}
 
-		Sig.bExperimental = true;
-		Sig.bMemberFunction = true;
-		Sig.bRequiresContext = false;
-		Sig.bWriteFunction = true;
+	{
+		FNiagaraFunctionSignature& Sig = OutFunctions.Add_GetRef(DefaultSig);
+		Sig.Name = Deprecated_GetValueFunctionName;
+		Sig.Inputs.Emplace(FNiagaraTypeDefinition::GetIntDef(), TEXT("IndexX"));
+		Sig.Inputs.Emplace(FNiagaraTypeDefinition::GetIntDef(), TEXT("IndexY"));
+		Sig.Outputs.Emplace(FNiagaraTypeDefinition::GetColorDef(), TEXT("Value"));
 		Sig.bSupportsCPU = false;
-		Sig.bSupportsGPU = true;
-#if WITH_EDITORONLY_DATA
-		Sig.FunctionVersion = FNDIRenderTarget2DFunctionVersion::LatestVersion;
-#endif
+		Sig.bSoftDeprecatedFunction = true;
 	}
 }
 
@@ -287,24 +249,27 @@ bool UNiagaraDataInterfaceRenderTarget2D::UpgradeFunctionCall(FNiagaraFunctionSi
 }
 #endif
 
-DEFINE_NDI_DIRECT_FUNC_BINDER(UNiagaraDataInterfaceRenderTarget2D, VMGetSize);
-DEFINE_NDI_DIRECT_FUNC_BINDER(UNiagaraDataInterfaceRenderTarget2D, VMSetSize);
-void UNiagaraDataInterfaceRenderTarget2D::GetVMExternalFunction(const FVMExternalFunctionBindingInfo& BindingInfo, void* InstanceData, FVMExternalFunction &OutFunc)
+void UNiagaraDataInterfaceRenderTarget2D::GetVMExternalFunction(const FVMExternalFunctionBindingInfo& BindingInfo, void* InstanceData, FVMExternalFunction& OutFunc)
 {
 	using namespace NDIRenderTarget2DLocal;
 
 	Super::GetVMExternalFunction(BindingInfo, InstanceData, OutFunc);
 	if (BindingInfo.Name == GetSizeFunctionName)
 	{
-		check(BindingInfo.GetNumInputs() == 1 && BindingInfo.GetNumOutputs() == 2);
-		NDI_FUNC_BINDER(UNiagaraDataInterfaceRenderTarget2D, VMGetSize)::Bind(this, OutFunc);
+		OutFunc = FVMExternalFunction::CreateUObject(this, &UNiagaraDataInterfaceRenderTarget2D::VMGetSize);
 	}
 	else if (BindingInfo.Name == SetSizeFunctionName)
 	{
-		check(BindingInfo.GetNumInputs() == 3 && BindingInfo.GetNumOutputs() == 1);
-		NDI_FUNC_BINDER(UNiagaraDataInterfaceRenderTarget2D, VMSetSize)::Bind(this, OutFunc);
+		OutFunc = FVMExternalFunction::CreateUObject(this, &UNiagaraDataInterfaceRenderTarget2D::VMSetSize);
+	}
+	else if (BindingInfo.Name == GetNumMipLevelsName)
+	{
+		OutFunc = FVMExternalFunction::CreateUObject(this, &UNiagaraDataInterfaceRenderTarget2D::VMGetNumMipLevels);
+	}
+	else if (BindingInfo.Name == SetFormatFunctionName)
+	{
+		OutFunc = FVMExternalFunction::CreateUObject(this, &UNiagaraDataInterfaceRenderTarget2D::VMSetFormat);
 	}
-
 }
 
 bool UNiagaraDataInterfaceRenderTarget2D::Equals(const UNiagaraDataInterface* Other) const
@@ -386,7 +351,8 @@ bool UNiagaraDataInterfaceRenderTarget2D::GetFunctionHLSL(const FNiagaraDataInte
 	} 
 
 	if ((FunctionInfo.DefinitionName == SetValueFunctionName) ||
-		(FunctionInfo.DefinitionName == GetValueFunctionName) ||
+		(FunctionInfo.DefinitionName == Deprecated_GetValueFunctionName) ||
+		(FunctionInfo.DefinitionName == LoadValueFunctionName) ||
 		(FunctionInfo.DefinitionName == SampleValueFunctionName) ||
 		(FunctionInfo.DefinitionName == GetSizeFunctionName) ||
 		(FunctionInfo.DefinitionName == LinearToIndexName) ||
@@ -428,6 +394,7 @@ void UNiagaraDataInterfaceRenderTarget2D::SetShaderParameters(const FNiagaraData
 
 	NDIRenderTarget2DLocal::FShaderParameters* Parameters = Context.GetParameterNestedStruct<NDIRenderTarget2DLocal::FShaderParameters>();
 	Parameters->TextureSize = InstanceData_RT->Size;
+	Parameters->MipLevels = InstanceData_RT->MipLevels;
 
 	const bool bRTWrite = Context.IsResourceBound(&Parameters->RWTexture);
 	const bool bRTRead = Context.IsResourceBound(&Parameters->Texture);
@@ -613,16 +580,60 @@ void UNiagaraDataInterfaceRenderTarget2D::VMSetSize(FVectorVMExternalFunctionCon
 void UNiagaraDataInterfaceRenderTarget2D::VMGetSize(FVectorVMExternalFunctionContext& Context)
 {
 	VectorVM::FUserPtrHandler<FRenderTarget2DRWInstanceData_GameThread> InstData(Context);
-	VectorVM::FExternalFuncRegisterHandler<int> OutSizeX(Context);
-	VectorVM::FExternalFuncRegisterHandler<int> OutSizeY(Context);
+	FNDIOutputParam<int> OutSizeX(Context);
+	FNDIOutputParam<int> OutSizeY(Context);
 
 	for (int32 InstanceIdx = 0; InstanceIdx < Context.GetNumInstances(); ++InstanceIdx)
 	{
-		*OutSizeX.GetDestAndAdvance() = InstData->Size.X;
-		*OutSizeY.GetDestAndAdvance() = InstData->Size.Y;
+		OutSizeX.SetAndAdvance(InstData->Size.X);
+		OutSizeY.SetAndAdvance(InstData->Size.Y);
 	}
 }
 
+void UNiagaraDataInterfaceRenderTarget2D::VMGetNumMipLevels(FVectorVMExternalFunctionContext& Context)
+{
+	VectorVM::FUserPtrHandler<FRenderTarget2DRWInstanceData_GameThread> InstData(Context);
+	FNDIOutputParam<int> OutMipLevels(Context);
+
+	int32 NumMipLevels = 1;
+	if (InstData->TargetTexture != nullptr)
+	{
+		NumMipLevels = InstData->TargetTexture->GetNumMips();
+	}
+	else if (InstData->MipMapGeneration != ENiagaraMipMapGeneration::Disabled)
+	{
+		NumMipLevels = FMath::FloorLog2(FMath::Max(InstData->Size.X, InstData->Size.Y)) + 1;
+	}
+	for (int32 InstanceIdx = 0; InstanceIdx < Context.GetNumInstances(); ++InstanceIdx)
+	{
+		OutMipLevels.SetAndAdvance(NumMipLevels);
+	}
+}
+
+void UNiagaraDataInterfaceRenderTarget2D::VMSetFormat(FVectorVMExternalFunctionContext& Context)
+{
+	// This should only be called from a system or emitter script due to a need for only setting up initially.
+	VectorVM::FUserPtrHandler<FRenderTarget2DRWInstanceData_GameThread> InstData(Context);
+	FNDIInputParam<ETextureRenderTargetFormat> InFormat(Context);
+	FNDIOutputParam<FNiagaraBool> OutSuccess(Context);
+
+	for (int32 InstanceIdx = 0; InstanceIdx < Context.GetNumInstances(); ++InstanceIdx)
+	{
+		const ETextureRenderTargetFormat Format = InFormat.GetAndAdvance();
+
+		const bool bSuccess = (InstData.Get() != nullptr && Context.GetNumInstances() == 1);
+		OutSuccess.SetAndAdvance(bSuccess);
+		if (bSuccess)
+		{
+			InstData->Format = Format;
+
+			if (bInheritUserParameterSettings && InstData->RTUserParamBinding.GetValue<UTextureRenderTarget2D>())
+			{
+				UE_LOG(LogNiagara, Warning, TEXT("Overriding inherited volume render target format"));
+			}
+		}
+	}
+}
 
 bool UNiagaraDataInterfaceRenderTarget2D::PerInstanceTick(void* PerInstanceData, FNiagaraSystemInstance* SystemInstance, float DeltaSeconds)
 {
@@ -733,6 +744,7 @@ bool UNiagaraDataInterfaceRenderTarget2D::PerInstanceTickPostSimulate(void* PerI
 			{
 				FRenderTarget2DRWInstanceData_RenderThread* TargetData = &RT_Proxy->SystemInstancesToProxyData_RT.FindOrAdd(RT_InstanceID);
 				TargetData->Size = RT_InstanceData.Size;
+				TargetData->MipLevels = RT_TargetTexture->GetCurrentMipCount();
 				TargetData->MipMapGeneration = RT_InstanceData.MipMapGeneration;
 				TargetData->MipMapGenerationType = RT_InstanceData.MipMapGenerationType;
 			#if WITH_EDITORONLY_DATA
diff --git a/Engine/Plugins/FX/Niagara/Source/Niagara/Private/NiagaraDataInterfaceRenderTarget2DArray.cpp b/Engine/Plugins/FX/Niagara/Source/Niagara/Private/NiagaraDataInterfaceRenderTarget2DArray.cpp
--- a/Engine/Plugins/FX/Niagara/Source/Niagara/Private/NiagaraDataInterfaceRenderTarget2DArray.cpp
+++ b/Engine/Plugins/FX/Niagara/Source/Niagara/Private/NiagaraDataInterfaceRenderTarget2DArray.cpp
@@ -26,6 +26,7 @@ namespace NDIRenderTarget2DArrayLocal
 {
 	BEGIN_SHADER_PARAMETER_STRUCT(FShaderParameters, )
 		SHADER_PARAMETER(FIntVector3,								TextureSize)
+		SHADER_PARAMETER(int,										MipLevels)
 		SHADER_PARAMETER_RDG_TEXTURE_UAV(RWTexture2DArray<float4>,	RWTexture)
 		SHADER_PARAMETER_RDG_TEXTURE_SRV(Texture2DArray<float4>,	Texture)
 		SHADER_PARAMETER_SAMPLER(SamplerState,						TextureSampler)
@@ -35,13 +36,16 @@ namespace NDIRenderTarget2DArrayLocal
 
 	// Global VM function names, also used by the shaders code generation methods.
 	static const FName SetValueFunctionName("SetRenderTargetValue");
-	static const FName GetValueFunctionName("GetRenderTargetValue");
+	static const FName LoadValueFunctionName("LoadRenderTargetValue");
 	static const FName SampleValueFunctionName("SampleRenderTargetValue");
 	static const FName SetSizeFunctionName("SetRenderTargetSize");
 	static const FName GetSizeFunctionName("GetRenderTargetSize");
+	static const FName GetNumMipLevelsName("GetNumMipLevels");
+	static const FName SetFormatFunctionName("SetRenderTargetFormat");
 	static const FName LinearToIndexName("LinearToIndex");
 	static const FName ExecToIndexName("ExecToIndex");
 	static const FName ExecToUnitName("ExecToUnit");
+	static const FName Deprecated_GetValueFunctionName("GetRenderTargetValue");
 }
 
 FNiagaraVariableBase UNiagaraDataInterfaceRenderTarget2DArray::ExposedRTVar;
@@ -111,160 +115,119 @@ void UNiagaraDataInterfaceRenderTarget2DArray::GetFunctions(TArray<FNiagaraFunct
 	const int32 EmitterSystemOnlyBitmask = ENiagaraScriptUsageMask::Emitter | ENiagaraScriptUsageMask::System;
 	OutFunctions.Reserve(OutFunctions.Num() + 4);
 
+	FNiagaraFunctionSignature DefaultSig;
+	DefaultSig.Inputs.Emplace(FNiagaraTypeDefinition(GetClass()), TEXT("RenderTarget"));
+	DefaultSig.bMemberFunction = true;
+	DefaultSig.bRequiresContext = false;
+	DefaultSig.bSupportsCPU = true;
+	DefaultSig.bSupportsGPU = true;
+	DefaultSig.SetFunctionVersion(FNDIRenderTarget2DArrayFunctionVersion::LatestVersion);
+
 	{
-		FNiagaraFunctionSignature& Sig = OutFunctions.AddDefaulted_GetRef();
+		FNiagaraFunctionSignature& Sig = OutFunctions.Add_GetRef(DefaultSig);
 		Sig.Name = GetSizeFunctionName;
-		Sig.Inputs.Add(FNiagaraVariable(FNiagaraTypeDefinition(GetClass()), TEXT("RenderTarget")));
-		Sig.Outputs.Add(FNiagaraVariable(FNiagaraTypeDefinition::GetIntDef(), TEXT("Width")));
-		Sig.Outputs.Add(FNiagaraVariable(FNiagaraTypeDefinition::GetIntDef(), TEXT("Height")));
-		Sig.Outputs.Add(FNiagaraVariable(FNiagaraTypeDefinition::GetIntDef(), TEXT("Slices")));
-		Sig.bExperimental = true;
-		Sig.bMemberFunction = true;
-		Sig.bRequiresContext = false;
-	#if WITH_EDITORONLY_DATA
-		Sig.FunctionVersion = FNDIRenderTarget2DArrayFunctionVersion::LatestVersion;
-	#endif
+		Sig.Outputs.Emplace(FNiagaraTypeDefinition::GetIntDef(), TEXT("Width"));
+		Sig.Outputs.Emplace(FNiagaraTypeDefinition::GetIntDef(), TEXT("Height"));
+		Sig.Outputs.Emplace(FNiagaraTypeDefinition::GetIntDef(), TEXT("Slices"));
 	}
 
 	{
-		FNiagaraFunctionSignature& Sig = OutFunctions.AddDefaulted_GetRef();
+		FNiagaraFunctionSignature& Sig = OutFunctions.Add_GetRef(DefaultSig);
 		Sig.Name = SetSizeFunctionName;
-		Sig.Inputs.Add(FNiagaraVariable(FNiagaraTypeDefinition(GetClass()), TEXT("RenderTarget")));
-		Sig.Inputs.Add(FNiagaraVariable(FNiagaraTypeDefinition::GetIntDef(), TEXT("Width")));
-		Sig.Inputs.Add(FNiagaraVariable(FNiagaraTypeDefinition::GetIntDef(), TEXT("Height")));
-		Sig.Inputs.Add(FNiagaraVariable(FNiagaraTypeDefinition::GetIntDef(), TEXT("Slices")));
-		Sig.Outputs.Add(FNiagaraVariable(FNiagaraTypeDefinition::GetBoolDef(), TEXT("Success")));
+		Sig.Inputs.Emplace(FNiagaraTypeDefinition::GetIntDef(), TEXT("Width"));
+		Sig.Inputs.Emplace(FNiagaraTypeDefinition::GetIntDef(), TEXT("Height"));
+		Sig.Inputs.Emplace(FNiagaraTypeDefinition::GetIntDef(), TEXT("Slices"));
+		Sig.Outputs.Emplace(FNiagaraTypeDefinition::GetBoolDef(), TEXT("Success"));
+		Sig.ModuleUsageBitmask = EmitterSystemOnlyBitmask;
+		Sig.bRequiresExecPin = true;
+		Sig.bSupportsGPU = false;
+	}
+
+	{
+		FNiagaraFunctionSignature& Sig = OutFunctions.Add_GetRef(DefaultSig);
+		Sig.Name = GetNumMipLevelsName;
+		Sig.Outputs.Emplace(FNiagaraTypeDefinition::GetIntDef(), TEXT("MipLevels"));
+	}
 
+	{
+		FNiagaraFunctionSignature& Sig = OutFunctions.Add_GetRef(DefaultSig);
+		Sig.Name = SetFormatFunctionName;
+		Sig.Inputs.Emplace(FNiagaraTypeDefinition(StaticEnum<ETextureRenderTargetFormat>()), TEXT("Format"));
+		Sig.Outputs.Emplace(FNiagaraTypeDefinition::GetBoolDef(), TEXT("Success"));
 		Sig.ModuleUsageBitmask = EmitterSystemOnlyBitmask;
-		Sig.bExperimental = true;
-		Sig.bMemberFunction = true;
 		Sig.bRequiresExecPin = true;
-		Sig.bRequiresContext = false;
-		Sig.bSupportsCPU = true;
 		Sig.bSupportsGPU = false;
-	#if WITH_EDITORONLY_DATA
-		Sig.FunctionVersion = FNDIRenderTarget2DArrayFunctionVersion::LatestVersion;
-	#endif
 	}
 
 	{
-		FNiagaraFunctionSignature& Sig = OutFunctions.AddDefaulted_GetRef();
+		FNiagaraFunctionSignature& Sig = OutFunctions.Add_GetRef(DefaultSig);
 		Sig.Name = SetValueFunctionName;
-		Sig.Inputs.Add(FNiagaraVariable(FNiagaraTypeDefinition(GetClass()), TEXT("RenderTarget")));
-		Sig.Inputs.Add_GetRef(FNiagaraVariable(FNiagaraTypeDefinition::GetBoolDef(), TEXT("Enabled"))).SetValue(true);
-		Sig.Inputs.Add(FNiagaraVariable(FNiagaraTypeDefinition::GetIntDef(), TEXT("IndexX")));
-		Sig.Inputs.Add(FNiagaraVariable(FNiagaraTypeDefinition::GetIntDef(), TEXT("IndexY")));
-		Sig.Inputs.Add(FNiagaraVariable(FNiagaraTypeDefinition::GetIntDef(), TEXT("IndexZ")));
-		Sig.Inputs.Add(FNiagaraVariable(FNiagaraTypeDefinition::GetColorDef(), TEXT("Value")));
-
-		Sig.bExperimental = true;
-		Sig.bMemberFunction = true;
+		Sig.Inputs.Emplace_GetRef(FNiagaraTypeDefinition::GetBoolDef(), TEXT("Enabled")).SetValue(true);
+		Sig.Inputs.Emplace(FNiagaraTypeDefinition::GetIntDef(), TEXT("IndexX"));
+		Sig.Inputs.Emplace(FNiagaraTypeDefinition::GetIntDef(), TEXT("IndexY"));
+		Sig.Inputs.Emplace(FNiagaraTypeDefinition::GetIntDef(), TEXT("IndexZ"));
+		Sig.Inputs.Emplace(FNiagaraTypeDefinition::GetColorDef(), TEXT("Value"));
 		Sig.bRequiresExecPin = true;
-		Sig.bRequiresContext = false;
 		Sig.bWriteFunction = true;
 		Sig.bSupportsCPU = false;
-		Sig.bSupportsGPU = true;
-	#if WITH_EDITORONLY_DATA
-		Sig.FunctionVersion = FNDIRenderTarget2DArrayFunctionVersion::LatestVersion;
-	#endif
 	}
 
 	{
-		FNiagaraFunctionSignature& Sig = OutFunctions.AddDefaulted_GetRef();
-		Sig.Name = GetValueFunctionName;
-		Sig.Inputs.Add(FNiagaraVariable(FNiagaraTypeDefinition(GetClass()), TEXT("RenderTarget")));
-		Sig.Inputs.Add(FNiagaraVariable(FNiagaraTypeDefinition::GetIntDef(), TEXT("IndexX")));
-		Sig.Inputs.Add(FNiagaraVariable(FNiagaraTypeDefinition::GetIntDef(), TEXT("IndexY")));
-		Sig.Inputs.Add(FNiagaraVariable(FNiagaraTypeDefinition::GetIntDef(), TEXT("IndexZ")));
-		Sig.Outputs.Add(FNiagaraVariable(FNiagaraTypeDefinition::GetColorDef(), TEXT("Value")));
+		FNiagaraFunctionSignature& Sig = OutFunctions.Add_GetRef(DefaultSig);
+		Sig.Name = Deprecated_GetValueFunctionName;
+		Sig.Inputs.Emplace(FNiagaraTypeDefinition::GetIntDef(), TEXT("IndexX"));
+		Sig.Inputs.Emplace(FNiagaraTypeDefinition::GetIntDef(), TEXT("IndexY"));
+		Sig.Inputs.Emplace(FNiagaraTypeDefinition::GetIntDef(), TEXT("IndexZ"));
+		Sig.Outputs.Emplace(FNiagaraTypeDefinition::GetColorDef(), TEXT("Value"));
+		Sig.bSupportsCPU = false;
+	}
 
-		Sig.bHidden = NiagaraDataInterfaceRenderTargetCommon::GAllowReads != 1;
-		Sig.bExperimental = true;
-		Sig.bMemberFunction = true;
-		Sig.bRequiresContext = false;
-		Sig.bWriteFunction = true;
+	{
+		FNiagaraFunctionSignature& Sig = OutFunctions.Add_GetRef(DefaultSig);
+		Sig.Name = LoadValueFunctionName;
+		Sig.Inputs.Emplace(FNiagaraTypeDefinition::GetIntDef(), TEXT("IndexX"));
+		Sig.Inputs.Emplace(FNiagaraTypeDefinition::GetIntDef(), TEXT("IndexY"));
+		Sig.Inputs.Emplace(FNiagaraTypeDefinition::GetIntDef(), TEXT("IndexZ"));
+		Sig.Inputs.Emplace(FNiagaraTypeDefinition::GetIntDef(), TEXT("MipLevel"));
+		Sig.Outputs.Emplace(FNiagaraTypeDefinition::GetColorDef(), TEXT("Value"));
 		Sig.bSupportsCPU = false;
-		Sig.bSupportsGPU = true;
-#if WITH_EDITORONLY_DATA
-		Sig.FunctionVersion = FNDIRenderTarget2DArrayFunctionVersion::LatestVersion;
-#endif
 	}
 
 	{
-		FNiagaraFunctionSignature& Sig = OutFunctions.AddDefaulted_GetRef();
+		FNiagaraFunctionSignature& Sig = OutFunctions.Add_GetRef(DefaultSig);
 		Sig.Name = SampleValueFunctionName;
-		Sig.Inputs.Add(FNiagaraVariable(FNiagaraTypeDefinition(GetClass()), TEXT("RenderTarget")));
-		Sig.Inputs.Add(FNiagaraVariable(FNiagaraTypeDefinition::GetVec2Def(), TEXT("UV")));
-		Sig.Inputs.Add(FNiagaraVariable(FNiagaraTypeDefinition::GetIntDef(), TEXT("Slice")));
-		Sig.Outputs.Add(FNiagaraVariable(FNiagaraTypeDefinition::GetColorDef(), TEXT("Value")));
-
-		Sig.bHidden = NiagaraDataInterfaceRenderTargetCommon::GAllowReads != 1;
-		Sig.bExperimental = true;
-		Sig.bMemberFunction = true;
-		Sig.bRequiresContext = false;
-		Sig.bWriteFunction = true;
+		Sig.Inputs.Emplace(FNiagaraTypeDefinition(GetClass()), TEXT("RenderTarget"));
+		Sig.Inputs.Emplace(FNiagaraTypeDefinition::GetVec2Def(), TEXT("UV"));
+		Sig.Inputs.Emplace(FNiagaraTypeDefinition::GetIntDef(), TEXT("Slice"));
+		Sig.Outputs.Emplace(FNiagaraTypeDefinition::GetColorDef(), TEXT("Value"));
 		Sig.bSupportsCPU = false;
-		Sig.bSupportsGPU = true;
-#if WITH_EDITORONLY_DATA
-		Sig.FunctionVersion = FNDIRenderTarget2DArrayFunctionVersion::LatestVersion;
-#endif
 	}
 
 	{
-		FNiagaraFunctionSignature& Sig = OutFunctions.AddDefaulted_GetRef();
+		FNiagaraFunctionSignature& Sig = OutFunctions.Add_GetRef(DefaultSig);
 		Sig.Name = LinearToIndexName;
-		Sig.Inputs.Add(FNiagaraVariable(FNiagaraTypeDefinition(GetClass()), TEXT("RenderTarget")));
-		Sig.Inputs.Add(FNiagaraVariable(FNiagaraTypeDefinition::GetIntDef(), TEXT("Linear")));
-		Sig.Outputs.Add(FNiagaraVariable(FNiagaraTypeDefinition::GetIntDef(), TEXT("IndexX")));
-		Sig.Outputs.Add(FNiagaraVariable(FNiagaraTypeDefinition::GetIntDef(), TEXT("IndexY")));
-		Sig.Outputs.Add(FNiagaraVariable(FNiagaraTypeDefinition::GetIntDef(), TEXT("IndexZ")));
-
-		Sig.bExperimental = true;
-		Sig.bMemberFunction = true;
-		Sig.bRequiresContext = false;
-		Sig.bWriteFunction = true;
+		Sig.Inputs.Emplace(FNiagaraTypeDefinition::GetIntDef(), TEXT("Linear"));
+		Sig.Outputs.Emplace(FNiagaraTypeDefinition::GetIntDef(), TEXT("IndexX"));
+		Sig.Outputs.Emplace(FNiagaraTypeDefinition::GetIntDef(), TEXT("IndexY"));
+		Sig.Outputs.Emplace(FNiagaraTypeDefinition::GetIntDef(), TEXT("IndexZ"));
 		Sig.bSupportsCPU = false;
-		Sig.bSupportsGPU = true;
-	#if WITH_EDITORONLY_DATA
-		Sig.FunctionVersion = FNDIRenderTarget2DArrayFunctionVersion::LatestVersion;
-	#endif
 	}
 
 	{
-		FNiagaraFunctionSignature& Sig = OutFunctions.AddDefaulted_GetRef();
+		FNiagaraFunctionSignature& Sig = OutFunctions.Add_GetRef(DefaultSig);
 		Sig.Name = ExecToIndexName;
-		Sig.Inputs.Add(FNiagaraVariable(FNiagaraTypeDefinition(GetClass()), TEXT("RenderTarget")));
-		Sig.Outputs.Add(FNiagaraVariable(FNiagaraTypeDefinition::GetIntDef(), TEXT("IndexX")));
-		Sig.Outputs.Add(FNiagaraVariable(FNiagaraTypeDefinition::GetIntDef(), TEXT("IndexY")));
-		Sig.Outputs.Add(FNiagaraVariable(FNiagaraTypeDefinition::GetIntDef(), TEXT("IndexZ")));
-
-		Sig.bExperimental = true;
-		Sig.bMemberFunction = true;
-		Sig.bRequiresContext = false;
-		Sig.bWriteFunction = true;
+		Sig.Outputs.Emplace(FNiagaraTypeDefinition::GetIntDef(), TEXT("IndexX"));
+		Sig.Outputs.Emplace(FNiagaraTypeDefinition::GetIntDef(), TEXT("IndexY"));
+		Sig.Outputs.Emplace(FNiagaraTypeDefinition::GetIntDef(), TEXT("IndexZ"));
 		Sig.bSupportsCPU = false;
-		Sig.bSupportsGPU = true;
-	#if WITH_EDITORONLY_DATA
-		Sig.FunctionVersion = FNDIRenderTarget2DArrayFunctionVersion::LatestVersion;
-	#endif
 	}
 
 	{
-		FNiagaraFunctionSignature& Sig = OutFunctions.AddDefaulted_GetRef();
+		FNiagaraFunctionSignature& Sig = OutFunctions.Add_GetRef(DefaultSig);
 		Sig.Name = ExecToUnitName;
-		Sig.Inputs.Emplace(FNiagaraTypeDefinition(GetClass()), TEXT("RenderTarget"));
 		Sig.Outputs.Emplace(FNiagaraTypeDefinition::GetVec2Def(), TEXT("Unit"));
 		Sig.Outputs.Emplace(FNiagaraTypeDefinition::GetIntDef(), TEXT("Slice"));
-
-		Sig.bExperimental = true;
-		Sig.bMemberFunction = true;
-		Sig.bRequiresContext = false;
-		Sig.bWriteFunction = true;
 		Sig.bSupportsCPU = false;
-		Sig.bSupportsGPU = true;
-#if WITH_EDITORONLY_DATA
-		Sig.FunctionVersion = FNDIRenderTarget2DArrayFunctionVersion::LatestVersion;
-#endif
 	}
 }
 
@@ -290,8 +253,6 @@ bool UNiagaraDataInterfaceRenderTarget2DArray::UpgradeFunctionCall(FNiagaraFunct
 }
 #endif
 
-DEFINE_NDI_DIRECT_FUNC_BINDER(UNiagaraDataInterfaceRenderTarget2DArray, VMGetSize);
-DEFINE_NDI_DIRECT_FUNC_BINDER(UNiagaraDataInterfaceRenderTarget2DArray, VMSetSize);
 void UNiagaraDataInterfaceRenderTarget2DArray::GetVMExternalFunction(const FVMExternalFunctionBindingInfo& BindingInfo, void* InstanceData, FVMExternalFunction &OutFunc)
 {
 	using namespace NDIRenderTarget2DArrayLocal;
@@ -299,15 +260,20 @@ void UNiagaraDataInterfaceRenderTarget2DArray::GetVMExternalFunction(const FVMEx
 	Super::GetVMExternalFunction(BindingInfo, InstanceData, OutFunc);
 	if (BindingInfo.Name == GetSizeFunctionName)
 	{
-		check(BindingInfo.GetNumInputs() == 1 && BindingInfo.GetNumOutputs() == 3);
-		NDI_FUNC_BINDER(UNiagaraDataInterfaceRenderTarget2DArray, VMGetSize)::Bind(this, OutFunc);
+		OutFunc = FVMExternalFunction::CreateUObject(this, &UNiagaraDataInterfaceRenderTarget2DArray::VMGetSize);
 	}
 	else if (BindingInfo.Name == SetSizeFunctionName)
 	{
-		check(BindingInfo.GetNumInputs() == 4 && BindingInfo.GetNumOutputs() == 1);
-		NDI_FUNC_BINDER(UNiagaraDataInterfaceRenderTarget2DArray, VMSetSize)::Bind(this, OutFunc);
+		OutFunc = FVMExternalFunction::CreateUObject(this, &UNiagaraDataInterfaceRenderTarget2DArray::VMSetSize);
+	}
+	else if (BindingInfo.Name == GetNumMipLevelsName)
+	{
+		OutFunc = FVMExternalFunction::CreateUObject(this, &UNiagaraDataInterfaceRenderTarget2DArray::VMGetNumMipLevels);
+	}
+	else if (BindingInfo.Name == SetFormatFunctionName)
+	{
+		OutFunc = FVMExternalFunction::CreateUObject(this, &UNiagaraDataInterfaceRenderTarget2DArray::VMSetFormat);
 	}
-
 }
 
 bool UNiagaraDataInterfaceRenderTarget2DArray::Equals(const UNiagaraDataInterface* Other) const
@@ -385,7 +351,8 @@ bool UNiagaraDataInterfaceRenderTarget2DArray::GetFunctionHLSL(const FNiagaraDat
 	}
 
 	if ((FunctionInfo.DefinitionName == SetValueFunctionName) ||
-		(FunctionInfo.DefinitionName == GetValueFunctionName) ||
+		(FunctionInfo.DefinitionName == Deprecated_GetValueFunctionName) ||
+		(FunctionInfo.DefinitionName == LoadValueFunctionName) ||
 		(FunctionInfo.DefinitionName == SampleValueFunctionName) ||
 		(FunctionInfo.DefinitionName == GetSizeFunctionName) ||
 		(FunctionInfo.DefinitionName == LinearToIndexName) ||
@@ -423,6 +390,7 @@ void UNiagaraDataInterfaceRenderTarget2DArray::SetShaderParameters(const FNiagar
 
 	NDIRenderTarget2DArrayLocal::FShaderParameters* Parameters = Context.GetParameterNestedStruct<NDIRenderTarget2DArrayLocal::FShaderParameters>();
 	Parameters->TextureSize = InstanceData_RT->Size;
+	Parameters->MipLevels = InstanceData_RT->MipLevels;
 
 	const bool bRTWrite = Context.IsResourceBound(&Parameters->RWTexture);
 	const bool bRTRead = Context.IsResourceBound(&Parameters->Texture);
@@ -583,6 +551,52 @@ void UNiagaraDataInterfaceRenderTarget2DArray::VMGetSize(FVectorVMExternalFuncti
 	}
 }
 
+void UNiagaraDataInterfaceRenderTarget2DArray::VMGetNumMipLevels(FVectorVMExternalFunctionContext& Context)
+{
+	VectorVM::FUserPtrHandler<FRenderTarget2DArrayRWInstanceData_GameThread> InstData(Context);
+	FNDIOutputParam<int> OutMipLevels(Context);
+
+	int32 NumMipLevels = 1;
+	if (InstData->TargetTexture != nullptr)
+	{
+		NumMipLevels = InstData->TargetTexture->GetNumMips();
+	}
+	//else if (InstData->MipMapGeneration != ENiagaraMipMapGeneration::Disabled)
+	//{
+	//	NumMipLevels = FMath::FloorLog2(FMath::Max(InstData->Size.X, InstData->Size.Y)) + 1;
+	//}
+	for (int32 InstanceIdx = 0; InstanceIdx < Context.GetNumInstances(); ++InstanceIdx)
+	{
+		OutMipLevels.SetAndAdvance(NumMipLevels);
+	}
+}
+
+void UNiagaraDataInterfaceRenderTarget2DArray::VMSetFormat(FVectorVMExternalFunctionContext& Context)
+{
+	// This should only be called from a system or emitter script due to a need for only setting up initially.
+	VectorVM::FUserPtrHandler<FRenderTarget2DArrayRWInstanceData_GameThread> InstData(Context);
+	FNDIInputParam<ETextureRenderTargetFormat> InFormat(Context);
+	FNDIOutputParam<FNiagaraBool> OutSuccess(Context);
+
+	for (int32 InstanceIdx = 0; InstanceIdx < Context.GetNumInstances(); ++InstanceIdx)
+	{
+		const ETextureRenderTargetFormat Format = InFormat.GetAndAdvance();
+
+
+		const bool bSuccess = (InstData.Get() != nullptr && Context.GetNumInstances() == 1);
+		OutSuccess.SetAndAdvance(bSuccess);
+		if (bSuccess)
+		{
+			InstData->Format = GetPixelFormatFromRenderTargetFormat(Format);
+
+			if (bInheritUserParameterSettings && InstData->RTUserParamBinding.GetValue<UTextureRenderTarget2DArray>())
+			{
+				UE_LOG(LogNiagara, Warning, TEXT("Overriding inherited volume render target format"));
+			}
+		}
+	}
+}
+
 bool UNiagaraDataInterfaceRenderTarget2DArray::PerInstanceTick(void* PerInstanceData, FNiagaraSystemInstance* SystemInstance, float DeltaSeconds)
 {
 	FRenderTarget2DArrayRWInstanceData_GameThread* InstanceData = static_cast<FRenderTarget2DArrayRWInstanceData_GameThread*>(PerInstanceData);
@@ -688,6 +702,7 @@ bool UNiagaraDataInterfaceRenderTarget2DArray::PerInstanceTickPostSimulate(void*
 			{
 				FRenderTarget2DArrayRWInstanceData_RenderThread* TargetData = &RT_Proxy->SystemInstancesToProxyData_RT.FindOrAdd(RT_InstanceID);
 				TargetData->Size = RT_InstanceData.Size;
+				TargetData->MipLevels = RT_TargetTexture->GetCurrentMipCount();
 				//TargetData->MipMapGeneration = RT_InstanceData.MipMapGeneration;
 			#if WITH_EDITORONLY_DATA
 				TargetData->bPreviewTexture = RT_InstanceData.bPreviewTexture;
diff --git a/Engine/Plugins/FX/Niagara/Source/Niagara/Private/NiagaraDataInterfaceRenderTargetCommon.cpp b/Engine/Plugins/FX/Niagara/Source/Niagara/Private/NiagaraDataInterfaceRenderTargetCommon.cpp
--- a/Engine/Plugins/FX/Niagara/Source/Niagara/Private/NiagaraDataInterfaceRenderTargetCommon.cpp
+++ b/Engine/Plugins/FX/Niagara/Source/Niagara/Private/NiagaraDataInterfaceRenderTargetCommon.cpp
@@ -32,14 +32,6 @@ namespace NiagaraDataInterfaceRenderTargetCommon
 		ECVF_Default
 	);
 
-	int GAllowReads = 0;
-	static FAutoConsoleVariableRef CVarNiagaraRenderTargetAllowReads(
-		TEXT("fx.Niagara.RenderTarget.AllowReads"),
-		GAllowReads,
-		TEXT("Enables read operations to be visible in the UI, very experimental."),
-		ECVF_Default
-	);
-
 	static bool GNiagaraRenderTargetOverrideFormatEnabled = false;
 	static ETextureRenderTargetFormat GNiagaraRenderTargetOverrideFormat = ETextureRenderTargetFormat::RTF_RGBA32f;
 
diff --git a/Engine/Plugins/FX/Niagara/Source/Niagara/Private/NiagaraDataInterfaceRenderTargetCommon.h b/Engine/Plugins/FX/Niagara/Source/Niagara/Private/NiagaraDataInterfaceRenderTargetCommon.h
--- a/Engine/Plugins/FX/Niagara/Source/Niagara/Private/NiagaraDataInterfaceRenderTargetCommon.h
+++ b/Engine/Plugins/FX/Niagara/Source/Niagara/Private/NiagaraDataInterfaceRenderTargetCommon.h
@@ -9,7 +9,6 @@ namespace NiagaraDataInterfaceRenderTargetCommon
 	extern int32 GReleaseResourceOnRemove;
 	extern int32 GIgnoreCookedOut;
 	extern float GResolutionMultiplier;
-	extern int GAllowReads;
 
 	extern bool GetRenderTargetFormat(bool bOverrideFormat, ETextureRenderTargetFormat OverrideFormat, ETextureRenderTargetFormat& OutRenderTargetFormat);
 }
diff --git a/Engine/Plugins/FX/Niagara/Source/Niagara/Private/NiagaraDataInterfaceRenderTargetCube.cpp b/Engine/Plugins/FX/Niagara/Source/Niagara/Private/NiagaraDataInterfaceRenderTargetCube.cpp
--- a/Engine/Plugins/FX/Niagara/Source/Niagara/Private/NiagaraDataInterfaceRenderTargetCube.cpp
+++ b/Engine/Plugins/FX/Niagara/Source/Niagara/Private/NiagaraDataInterfaceRenderTargetCube.cpp
@@ -24,6 +24,7 @@ namespace NDIRenderTargetCubeLocal
 {
 	BEGIN_SHADER_PARAMETER_STRUCT(FShaderParameters, )
 		SHADER_PARAMETER(int, TextureSize)
+		SHADER_PARAMETER(int, MipLevels)
 		SHADER_PARAMETER_RDG_TEXTURE_UAV(RWTextureCube<float4>, RWTexture)
 		SHADER_PARAMETER_RDG_TEXTURE_SRV(TextureCube<float4>, Texture)
 		SHADER_PARAMETER_SAMPLER(SamplerState, TextureSampler)
@@ -34,9 +35,12 @@ namespace NDIRenderTargetCubeLocal
 	// Global VM function names, also used by the shaders code generation methods.
 	static const FName SetValueFunctionName("SetRenderTargetValue");
 	//static const FName GetValueFunctionName("GetRenderTargetValue");
+	//static const FName LoadValueFunctionName("LoadRenderTargetValue");
 	static const FName SampleValueFunctionName("SampleRenderTargetValue");
 	static const FName SetSizeFunctionName("SetRenderTargetSize");
 	static const FName GetSizeFunctionName("GetRenderTargetSize");
+	static const FName GetNumMipLevelsName("GetNumMipLevels");
+	static const FName SetFormatFunctionName("SetRenderTargetFormat");
 	static const FName LinearToIndexName("LinearToIndex");
 	static const FName ExecToIndexName("ExecToIndex");
 	static const FName ExecToUnitName("ExecToUnit");
@@ -107,157 +111,93 @@ void UNiagaraDataInterfaceRenderTargetCube::GetFunctions(TArray<FNiagaraFunction
 	const int32 EmitterSystemOnlyBitmask = ENiagaraScriptUsageMask::Emitter | ENiagaraScriptUsageMask::System;
 	OutFunctions.Reserve(OutFunctions.Num() + 7);
 
+	FNiagaraFunctionSignature DefaultSig;
+	DefaultSig.Inputs.Emplace(FNiagaraTypeDefinition(GetClass()), TEXT("RenderTarget"));
+	DefaultSig.bMemberFunction = true;
+	DefaultSig.bRequiresContext = false;
+	DefaultSig.bSupportsCPU = true;
+	DefaultSig.bSupportsGPU = true;
+	DefaultSig.SetFunctionVersion(EFunctionVersion::LatestVersion);
+
 	{
-		FNiagaraFunctionSignature& Sig = OutFunctions.AddDefaulted_GetRef();
+		FNiagaraFunctionSignature& Sig = OutFunctions.Add_GetRef(DefaultSig);
 		Sig.Name = GetSizeFunctionName;
-		Sig.Inputs.Add(FNiagaraVariable(FNiagaraTypeDefinition(GetClass()), TEXT("RenderTarget")));
-		Sig.Outputs.Add(FNiagaraVariable(FNiagaraTypeDefinition::GetIntDef(), TEXT("Size")));
-
-		Sig.bExperimental = true;
-		Sig.bMemberFunction = true;
-		Sig.bRequiresContext = false;
-	#if WITH_EDITORONLY_DATA
-		Sig.FunctionVersion = EFunctionVersion::LatestVersion;
-	#endif
+		Sig.Outputs.Emplace(FNiagaraTypeDefinition::GetIntDef(), TEXT("Size"));
 	}
 
 	{
-		FNiagaraFunctionSignature& Sig = OutFunctions.AddDefaulted_GetRef();
+		FNiagaraFunctionSignature& Sig = OutFunctions.Add_GetRef(DefaultSig);
 		Sig.Name = SetSizeFunctionName;
-		Sig.Inputs.Add(FNiagaraVariable(FNiagaraTypeDefinition(GetClass()), TEXT("RenderTarget")));
-		Sig.Inputs.Add(FNiagaraVariable(FNiagaraTypeDefinition::GetIntDef(), TEXT("Size")));
-		Sig.Outputs.Add(FNiagaraVariable(FNiagaraTypeDefinition::GetBoolDef(), TEXT("Success")));
+		Sig.Inputs.Emplace(FNiagaraTypeDefinition::GetIntDef(), TEXT("Size"));
+		Sig.Outputs.Emplace(FNiagaraTypeDefinition::GetBoolDef(), TEXT("Success"));
+		Sig.ModuleUsageBitmask = EmitterSystemOnlyBitmask;
+		Sig.bRequiresExecPin = true;
+		Sig.bSupportsGPU = false;
+	}
+
+	{
+		FNiagaraFunctionSignature& Sig = OutFunctions.Add_GetRef(DefaultSig);
+		Sig.Name = GetNumMipLevelsName;
+		Sig.Outputs.Emplace(FNiagaraTypeDefinition::GetIntDef(), TEXT("MipLevels"));
+	}
 
+	{
+		FNiagaraFunctionSignature& Sig = OutFunctions.Add_GetRef(DefaultSig);
+		Sig.Name = SetFormatFunctionName;
+		Sig.Inputs.Emplace(FNiagaraTypeDefinition(StaticEnum<ETextureRenderTargetFormat>()), TEXT("Format"));
+		Sig.Outputs.Emplace(FNiagaraTypeDefinition::GetBoolDef(), TEXT("Success"));
 		Sig.ModuleUsageBitmask = EmitterSystemOnlyBitmask;
-		Sig.bExperimental = true;
-		Sig.bMemberFunction = true;
 		Sig.bRequiresExecPin = true;
-		Sig.bRequiresContext = false;
-		Sig.bSupportsCPU = true;
 		Sig.bSupportsGPU = false;
-	#if WITH_EDITORONLY_DATA
-		Sig.FunctionVersion = EFunctionVersion::LatestVersion;
-	#endif
 	}
 
 	{
-		FNiagaraFunctionSignature& Sig = OutFunctions.AddDefaulted_GetRef();
+		FNiagaraFunctionSignature& Sig = OutFunctions.Add_GetRef(DefaultSig);
 		Sig.Name = SetValueFunctionName;
-		Sig.Inputs.Add(FNiagaraVariable(FNiagaraTypeDefinition(GetClass()), TEXT("RenderTarget")));
-		Sig.Inputs.Add_GetRef(FNiagaraVariable(FNiagaraTypeDefinition::GetBoolDef(), TEXT("Enabled"))).SetValue(true);
-		Sig.Inputs.Add(FNiagaraVariable(FNiagaraTypeDefinition::GetIntDef(), TEXT("IndexX")));
-		Sig.Inputs.Add(FNiagaraVariable(FNiagaraTypeDefinition::GetIntDef(), TEXT("IndexY")));
-		Sig.Inputs.Add(FNiagaraVariable(FNiagaraTypeDefinition::GetIntDef(), TEXT("Face")));
-		Sig.Inputs.Add(FNiagaraVariable(FNiagaraTypeDefinition::GetColorDef(), TEXT("Value")));
-
-		Sig.bExperimental = true;
-		Sig.bMemberFunction = true;
+		Sig.Inputs.Emplace_GetRef(FNiagaraTypeDefinition::GetBoolDef(), TEXT("Enabled")).SetValue(true);
+		Sig.Inputs.Emplace(FNiagaraTypeDefinition::GetIntDef(), TEXT("IndexX"));
+		Sig.Inputs.Emplace(FNiagaraTypeDefinition::GetIntDef(), TEXT("IndexY"));
+		Sig.Inputs.Emplace(FNiagaraTypeDefinition::GetIntDef(), TEXT("Face"));
+		Sig.Inputs.Emplace(FNiagaraTypeDefinition::GetColorDef(), TEXT("Value"));
 		Sig.bRequiresExecPin = true;
-		Sig.bRequiresContext = false;
 		Sig.bWriteFunction = true;
 		Sig.bSupportsCPU = false;
-		Sig.bSupportsGPU = true;
-	#if WITH_EDITORONLY_DATA
-		Sig.FunctionVersion = EFunctionVersion::LatestVersion;
-	#endif
-	}
-
-// CubeTexture.Load() is not available
-//	{
-//		FNiagaraFunctionSignature& Sig = OutFunctions.AddDefaulted_GetRef();
-//		Sig.Name = GetValueFunctionName;
-//		Sig.Inputs.Add(FNiagaraVariable(FNiagaraTypeDefinition(GetClass()), TEXT("RenderTarget")));
-//		Sig.Inputs.Add(FNiagaraVariable(FNiagaraTypeDefinition::GetIntDef(), TEXT("IndexX")));
-//		Sig.Inputs.Add(FNiagaraVariable(FNiagaraTypeDefinition::GetIntDef(), TEXT("IndexY")));
-//		Sig.Inputs.Add(FNiagaraVariable(FNiagaraTypeDefinition::GetIntDef(), TEXT("Face")));
-//		Sig.Outputs.Add(FNiagaraVariable(FNiagaraTypeDefinition::GetColorDef(), TEXT("Value")));
-//
-//		Sig.bHidden = NiagaraDataInterfaceRenderTargetCommon::GAllowReads != 1;
-//		Sig.bExperimental = true;
-//		Sig.bMemberFunction = true;
-//		Sig.bRequiresContext = false;
-//		Sig.bWriteFunction = true;
-//		Sig.bSupportsCPU = false;
-//		Sig.bSupportsGPU = true;
-//#if WITH_EDITORONLY_DATA
-//		Sig.FunctionVersion = EFunctionVersion::LatestVersion;
-//#endif
-//	}
-
-	{
-		FNiagaraFunctionSignature& Sig = OutFunctions.AddDefaulted_GetRef();
+	}
+
+	{
+		FNiagaraFunctionSignature& Sig = OutFunctions.Add_GetRef(DefaultSig);
 		Sig.Name = SampleValueFunctionName;
-		Sig.Inputs.Add(FNiagaraVariable(FNiagaraTypeDefinition(GetClass()), TEXT("RenderTarget")));
-		Sig.Inputs.Add(FNiagaraVariable(FNiagaraTypeDefinition::GetVec3Def(), TEXT("UVW")));
-		Sig.Outputs.Add(FNiagaraVariable(FNiagaraTypeDefinition::GetColorDef(), TEXT("Value")));
-
-		Sig.bHidden = NiagaraDataInterfaceRenderTargetCommon::GAllowReads != 1;
-		Sig.bExperimental = true;
-		Sig.bMemberFunction = true;
-		Sig.bRequiresContext = false;
-		Sig.bWriteFunction = true;
+		Sig.Inputs.Emplace(FNiagaraTypeDefinition::GetVec3Def(), TEXT("UVW"));
+		Sig.Outputs.Emplace(FNiagaraTypeDefinition::GetColorDef(), TEXT("Value"));
 		Sig.bSupportsCPU = false;
-		Sig.bSupportsGPU = true;
-#if WITH_EDITORONLY_DATA
-		Sig.FunctionVersion = EFunctionVersion::LatestVersion;
-#endif
 	}
 
 	{
-		FNiagaraFunctionSignature& Sig = OutFunctions.AddDefaulted_GetRef();
+		FNiagaraFunctionSignature& Sig = OutFunctions.Add_GetRef(DefaultSig);
 		Sig.Name = LinearToIndexName;
-		Sig.Inputs.Add(FNiagaraVariable(FNiagaraTypeDefinition(GetClass()), TEXT("RenderTarget")));
-		Sig.Inputs.Add(FNiagaraVariable(FNiagaraTypeDefinition::GetIntDef(), TEXT("Linear")));
-		Sig.Outputs.Add(FNiagaraVariable(FNiagaraTypeDefinition::GetIntDef(), TEXT("IndexX")));
-		Sig.Outputs.Add(FNiagaraVariable(FNiagaraTypeDefinition::GetIntDef(), TEXT("IndexY")));
-		Sig.Outputs.Add(FNiagaraVariable(FNiagaraTypeDefinition::GetIntDef(), TEXT("Face")));
-
-		Sig.bExperimental = true;
-		Sig.bMemberFunction = true;
-		Sig.bRequiresContext = false;
-		Sig.bWriteFunction = true;
+		Sig.Inputs.Emplace(FNiagaraTypeDefinition::GetIntDef(), TEXT("Linear"));
+		Sig.Outputs.Emplace(FNiagaraTypeDefinition::GetIntDef(), TEXT("IndexX"));
+		Sig.Outputs.Emplace(FNiagaraTypeDefinition::GetIntDef(), TEXT("IndexY"));
+		Sig.Outputs.Emplace(FNiagaraTypeDefinition::GetIntDef(), TEXT("Face"));
 		Sig.bSupportsCPU = false;
-		Sig.bSupportsGPU = true;
-	#if WITH_EDITORONLY_DATA
-		Sig.FunctionVersion = EFunctionVersion::LatestVersion;
-	#endif
 	}
 
 	{
-		FNiagaraFunctionSignature& Sig = OutFunctions.AddDefaulted_GetRef();
+		FNiagaraFunctionSignature& Sig = OutFunctions.Add_GetRef(DefaultSig);
 		Sig.Name = ExecToIndexName;
-		Sig.Inputs.Add(FNiagaraVariable(FNiagaraTypeDefinition(GetClass()), TEXT("RenderTarget")));
-		Sig.Outputs.Add(FNiagaraVariable(FNiagaraTypeDefinition::GetIntDef(), TEXT("IndexX")));
-		Sig.Outputs.Add(FNiagaraVariable(FNiagaraTypeDefinition::GetIntDef(), TEXT("IndexY")));
-		Sig.Outputs.Add(FNiagaraVariable(FNiagaraTypeDefinition::GetIntDef(), TEXT("Face")));
-
-		Sig.bExperimental = true;
-		Sig.bMemberFunction = true;
-		Sig.bRequiresContext = false;
-		Sig.bWriteFunction = true;
+		Sig.Outputs.Emplace(FNiagaraTypeDefinition::GetIntDef(), TEXT("IndexX"));
+		Sig.Outputs.Emplace(FNiagaraTypeDefinition::GetIntDef(), TEXT("IndexY"));
+		Sig.Outputs.Emplace(FNiagaraTypeDefinition::GetIntDef(), TEXT("Face"));
 		Sig.bSupportsCPU = false;
-		Sig.bSupportsGPU = true;
-	#if WITH_EDITORONLY_DATA
-		Sig.FunctionVersion = EFunctionVersion::LatestVersion;
-	#endif
 	}
 
 	{
-		FNiagaraFunctionSignature& Sig = OutFunctions.AddDefaulted_GetRef();
+		FNiagaraFunctionSignature& Sig = OutFunctions.Add_GetRef(DefaultSig);
 		Sig.Name = ExecToUnitName;
 		Sig.Inputs.Emplace(FNiagaraTypeDefinition(GetClass()), TEXT("RenderTarget"));
 		Sig.Outputs.Emplace(FNiagaraTypeDefinition::GetVec2Def(), TEXT("Unit"));
 		Sig.Outputs.Emplace(FNiagaraTypeDefinition::GetIntDef(), TEXT("Face"));
-
-		Sig.bExperimental = true;
-		Sig.bMemberFunction = true;
-		Sig.bRequiresContext = false;
-		Sig.bWriteFunction = true;
 		Sig.bSupportsCPU = false;
-		Sig.bSupportsGPU = true;
-#if WITH_EDITORONLY_DATA
-		Sig.FunctionVersion = EFunctionVersion::LatestVersion;
-#endif
 	}
 }
 
@@ -285,8 +225,6 @@ bool UNiagaraDataInterfaceRenderTargetCube::UpgradeFunctionCall(FNiagaraFunction
 }
 #endif
 
-DEFINE_NDI_DIRECT_FUNC_BINDER(UNiagaraDataInterfaceRenderTargetCube, VMGetSize);
-DEFINE_NDI_DIRECT_FUNC_BINDER(UNiagaraDataInterfaceRenderTargetCube, VMSetSize);
 void UNiagaraDataInterfaceRenderTargetCube::GetVMExternalFunction(const FVMExternalFunctionBindingInfo& BindingInfo, void* InstanceData, FVMExternalFunction &OutFunc)
 {
 	using namespace NDIRenderTargetCubeLocal;
@@ -294,13 +232,19 @@ void UNiagaraDataInterfaceRenderTargetCube::GetVMExternalFunction(const FVMExter
 	Super::GetVMExternalFunction(BindingInfo, InstanceData, OutFunc);
 	if (BindingInfo.Name == GetSizeFunctionName)
 	{
-		check(BindingInfo.GetNumInputs() == 1 && BindingInfo.GetNumOutputs() == 1);
-		NDI_FUNC_BINDER(UNiagaraDataInterfaceRenderTargetCube, VMGetSize)::Bind(this, OutFunc);
+		OutFunc = FVMExternalFunction::CreateUObject(this, &UNiagaraDataInterfaceRenderTargetCube::VMGetSize);
 	}
 	else if (BindingInfo.Name == SetSizeFunctionName)
 	{
-		check(BindingInfo.GetNumInputs() == 2 && BindingInfo.GetNumOutputs() == 1);
-		NDI_FUNC_BINDER(UNiagaraDataInterfaceRenderTargetCube, VMSetSize)::Bind(this, OutFunc);
+		OutFunc = FVMExternalFunction::CreateUObject(this, &UNiagaraDataInterfaceRenderTargetCube::VMSetSize);
+	}
+	else if (BindingInfo.Name == GetNumMipLevelsName)
+	{
+		OutFunc = FVMExternalFunction::CreateUObject(this, &UNiagaraDataInterfaceRenderTargetCube::VMGetNumMipLevels);
+	}
+	else if (BindingInfo.Name == SetFormatFunctionName)
+	{
+		OutFunc = FVMExternalFunction::CreateUObject(this, &UNiagaraDataInterfaceRenderTargetCube::VMSetFormat);
 	}
 }
 
@@ -374,6 +318,7 @@ bool UNiagaraDataInterfaceRenderTargetCube::GetFunctionHLSL(const FNiagaraDataIn
 
 	if ((FunctionInfo.DefinitionName == SetValueFunctionName) ||
 		//(FunctionInfo.DefinitionName == GetValueFunctionName) ||
+		//(FunctionInfo.DefinitionName == LoadValueFunctionName) ||
 		(FunctionInfo.DefinitionName == SampleValueFunctionName) ||
 		(FunctionInfo.DefinitionName == GetSizeFunctionName) ||
 		(FunctionInfo.DefinitionName == LinearToIndexName) ||
@@ -411,6 +356,7 @@ void UNiagaraDataInterfaceRenderTargetCube::SetShaderParameters(const FNiagaraDa
 
 	NDIRenderTargetCubeLocal::FShaderParameters* Parameters = Context.GetParameterNestedStruct<NDIRenderTargetCubeLocal::FShaderParameters>();
 	Parameters->TextureSize = InstanceData_RT->Size;
+	Parameters->MipLevels = InstanceData_RT->MipLevels;
 
 	const bool bRTWrite = Context.IsResourceBound(&Parameters->RWTexture);
 	const bool bRTRead = Context.IsResourceBound(&Parameters->Texture);
@@ -558,6 +504,51 @@ void UNiagaraDataInterfaceRenderTargetCube::VMGetSize(FVectorVMExternalFunctionC
 	}
 }
 
+void UNiagaraDataInterfaceRenderTargetCube::VMGetNumMipLevels(FVectorVMExternalFunctionContext& Context)
+{
+	VectorVM::FUserPtrHandler<FRenderTargetCubeRWInstanceData_GameThread> InstData(Context);
+	FNDIOutputParam<int> OutMipLevels(Context);
+
+	int32 NumMipLevels = 1;
+	if (InstData->TargetTexture != nullptr)
+	{
+		NumMipLevels = InstData->TargetTexture->GetNumMips();
+	}
+	//else if (InstData->MipMapGeneration != ENiagaraMipMapGeneration::Disabled)
+	//{
+	//	NumMipLevels = FMath::FloorLog2(FMath::Max(InstData->Size.X, InstData->Size.Y)) + 1;
+	//}
+	for (int32 InstanceIdx = 0; InstanceIdx < Context.GetNumInstances(); ++InstanceIdx)
+	{
+		OutMipLevels.SetAndAdvance(NumMipLevels);
+	}
+}
+
+void UNiagaraDataInterfaceRenderTargetCube::VMSetFormat(FVectorVMExternalFunctionContext& Context)
+{
+	// This should only be called from a system or emitter script due to a need for only setting up initially.
+	VectorVM::FUserPtrHandler<FRenderTargetCubeRWInstanceData_GameThread> InstData(Context);
+	FNDIInputParam<ETextureRenderTargetFormat> InFormat(Context);
+	FNDIOutputParam<FNiagaraBool> OutSuccess(Context);
+
+	for (int32 InstanceIdx = 0; InstanceIdx < Context.GetNumInstances(); ++InstanceIdx)
+	{
+		const ETextureRenderTargetFormat Format = InFormat.GetAndAdvance();
+
+
+		const bool bSuccess = (InstData.Get() != nullptr && Context.GetNumInstances() == 1);
+		OutSuccess.SetAndAdvance(bSuccess);
+		if (bSuccess)
+		{
+			InstData->Format = GetPixelFormatFromRenderTargetFormat(Format);
+
+			if (bInheritUserParameterSettings && InstData->RTUserParamBinding.GetValue<UTextureRenderTargetCube>())
+			{
+				UE_LOG(LogNiagara, Warning, TEXT("Overriding inherited volume render target format"));
+			}
+		}
+	}
+}
 
 bool UNiagaraDataInterfaceRenderTargetCube::PerInstanceTick(void* PerInstanceData, FNiagaraSystemInstance* SystemInstance, float DeltaSeconds)
 {
@@ -663,6 +654,7 @@ bool UNiagaraDataInterfaceRenderTargetCube::PerInstanceTickPostSimulate(void* Pe
 			{
 				FRenderTargetCubeRWInstanceData_RenderThread* TargetData = &RT_Proxy->SystemInstancesToProxyData_RT.FindOrAdd(RT_InstanceID);
 				TargetData->Size = RT_InstanceData.Size;
+				TargetData->MipLevels = RT_TargetTexture->GetCurrentMipCount();
 				//TargetData->MipMapGeneration = RT_InstanceData.MipMapGeneration;
 			#if WITH_EDITORONLY_DATA
 				TargetData->bPreviewTexture = RT_InstanceData.bPreviewTexture;
diff --git a/Engine/Plugins/FX/Niagara/Source/Niagara/Private/NiagaraDataInterfaceRenderTargetVolume.cpp b/Engine/Plugins/FX/Niagara/Source/Niagara/Private/NiagaraDataInterfaceRenderTargetVolume.cpp
--- a/Engine/Plugins/FX/Niagara/Source/Niagara/Private/NiagaraDataInterfaceRenderTargetVolume.cpp
+++ b/Engine/Plugins/FX/Niagara/Source/Niagara/Private/NiagaraDataInterfaceRenderTargetVolume.cpp
@@ -30,6 +30,7 @@ namespace NDIRenderTargetVolumeLocal
 {
 	BEGIN_SHADER_PARAMETER_STRUCT(FShaderParameters, )
 		SHADER_PARAMETER(FIntVector3, TextureSize)
+		SHADER_PARAMETER(int, MipLevels)
 		SHADER_PARAMETER_RDG_TEXTURE_UAV(RWTexture3D<float4>, RWTexture)
 		SHADER_PARAMETER_RDG_TEXTURE_SRV(Texture3D<float4>, Texture)
 		SHADER_PARAMETER_SAMPLER(SamplerState, TextureSampler)
@@ -39,14 +40,16 @@ namespace NDIRenderTargetVolumeLocal
 
 	// Global VM function names, also used by the shaders code generation methods.
 	static const FName SetValueFunctionName("SetRenderTargetValue");
-	static const FName GetValueFunctionName("GetRenderTargetValue");
+	static const FName LoadValueFunctionName("LoadRenderTargetValue");
 	static const FName SampleValueFunctionName("SampleRenderTargetValue");
 	static const FName SetSizeFunctionName("SetRenderTargetSize");
-	static const FName SetFormatFunctionName("SetRenderTargetFormat");
 	static const FName GetSizeFunctionName("GetRenderTargetSize");
+	static const FName GetNumMipLevelsName("GetNumMipLevels");
+	static const FName SetFormatFunctionName("SetRenderTargetFormat");
 	static const FName LinearToIndexName("LinearToIndex");
 	static const FName ExecToIndexName("ExecToIndex");
 	static const FName ExecToUnitName("ExecToUnit");
+	static const FName Deprecated_GetValueFunctionName("GetRenderTargetValue");
 
 	struct EFunctionVersion
 	{
@@ -148,179 +151,117 @@ void UNiagaraDataInterfaceRenderTargetVolume::GetFunctions(TArray<FNiagaraFuncti
 	const int32 EmitterSystemOnlyBitmask = ENiagaraScriptUsageMask::Emitter | ENiagaraScriptUsageMask::System;
 	OutFunctions.Reserve(OutFunctions.Num() + 7);
 
+	FNiagaraFunctionSignature DefaultSig;
+	DefaultSig.Inputs.Emplace(FNiagaraTypeDefinition(GetClass()), TEXT("RenderTarget"));
+	DefaultSig.bMemberFunction = true;
+	DefaultSig.bRequiresContext = false;
+	DefaultSig.bSupportsCPU = true;
+	DefaultSig.bSupportsGPU = true;
+	DefaultSig.SetFunctionVersion(EFunctionVersion::LatestVersion);
+
 	{
-		FNiagaraFunctionSignature& Sig = OutFunctions.AddDefaulted_GetRef();
+		FNiagaraFunctionSignature& Sig = OutFunctions.Add_GetRef(DefaultSig);
 		Sig.Name = GetSizeFunctionName;
-		Sig.Inputs.Add(FNiagaraVariable(FNiagaraTypeDefinition(GetClass()), TEXT("RenderTarget")));
-		Sig.Outputs.Add(FNiagaraVariable(FNiagaraTypeDefinition::GetIntDef(), TEXT("Width")));
-		Sig.Outputs.Add(FNiagaraVariable(FNiagaraTypeDefinition::GetIntDef(), TEXT("Height")));
-		Sig.Outputs.Add(FNiagaraVariable(FNiagaraTypeDefinition::GetIntDef(), TEXT("Depth")));
-
-		Sig.bExperimental = true;
-		Sig.bMemberFunction = true;
-		Sig.bRequiresContext = false;
-	#if WITH_EDITORONLY_DATA
-		Sig.FunctionVersion = EFunctionVersion::LatestVersion;
-	#endif
+		Sig.Outputs.Emplace(FNiagaraTypeDefinition::GetIntDef(), TEXT("Width"));
+		Sig.Outputs.Emplace(FNiagaraTypeDefinition::GetIntDef(), TEXT("Height"));
+		Sig.Outputs.Emplace(FNiagaraTypeDefinition::GetIntDef(), TEXT("Depth"));
 	}
 
 	{
-		FNiagaraFunctionSignature& Sig = OutFunctions.AddDefaulted_GetRef();
+		FNiagaraFunctionSignature& Sig = OutFunctions.Add_GetRef(DefaultSig);
 		Sig.Name = SetSizeFunctionName;
-		Sig.Inputs.Add(FNiagaraVariable(FNiagaraTypeDefinition(GetClass()), TEXT("RenderTarget")));
-		Sig.Inputs.Add(FNiagaraVariable(FNiagaraTypeDefinition::GetIntDef(), TEXT("Width")));
-		Sig.Inputs.Add(FNiagaraVariable(FNiagaraTypeDefinition::GetIntDef(), TEXT("Height")));
-		Sig.Inputs.Add(FNiagaraVariable(FNiagaraTypeDefinition::GetIntDef(), TEXT("Depth")));
-		Sig.Outputs.Add(FNiagaraVariable(FNiagaraTypeDefinition::GetBoolDef(), TEXT("Success")));
-
+		Sig.Inputs.Emplace(FNiagaraTypeDefinition::GetIntDef(), TEXT("Width"));
+		Sig.Inputs.Emplace(FNiagaraTypeDefinition::GetIntDef(), TEXT("Height"));
+		Sig.Inputs.Emplace(FNiagaraTypeDefinition::GetIntDef(), TEXT("Depth"));
+		Sig.Outputs.Emplace(FNiagaraTypeDefinition::GetBoolDef(), TEXT("Success"));
 		Sig.ModuleUsageBitmask = EmitterSystemOnlyBitmask;
-		Sig.bExperimental = true;
-		Sig.bMemberFunction = true;
 		Sig.bRequiresExecPin = true;
-		Sig.bRequiresContext = false;
-		Sig.bSupportsCPU = true;
 		Sig.bSupportsGPU = false;
-	#if WITH_EDITORONLY_DATA
-		Sig.FunctionVersion = EFunctionVersion::LatestVersion;
-	#endif
 	}
 
 	{
-		FNiagaraFunctionSignature& Sig = OutFunctions.AddDefaulted_GetRef();
-		Sig.Name = SetFormatFunctionName;
-		Sig.Inputs.Add(FNiagaraVariable(FNiagaraTypeDefinition(GetClass()), TEXT("RenderTarget")));		
-		Sig.Inputs.Add(FNiagaraVariable(FNiagaraTypeDefinition(StaticEnum<ETextureRenderTargetFormat>()), TEXT("Format")));
-		Sig.Outputs.Add(FNiagaraVariable(FNiagaraTypeDefinition::GetBoolDef(), TEXT("Success")));
+		FNiagaraFunctionSignature& Sig = OutFunctions.Add_GetRef(DefaultSig);
+		Sig.Name = GetNumMipLevelsName;
+		Sig.Outputs.Emplace(FNiagaraTypeDefinition::GetIntDef(), TEXT("MipLevels"));
+	}
 
+	{
+		FNiagaraFunctionSignature& Sig = OutFunctions.Add_GetRef(DefaultSig);
+		Sig.Name = SetFormatFunctionName;
+		Sig.Inputs.Emplace(FNiagaraTypeDefinition(StaticEnum<ETextureRenderTargetFormat>()), TEXT("Format"));
+		Sig.Outputs.Emplace(FNiagaraTypeDefinition::GetBoolDef(), TEXT("Success"));
 		Sig.ModuleUsageBitmask = EmitterSystemOnlyBitmask;
-		Sig.bExperimental = true;
-		Sig.bMemberFunction = true;
 		Sig.bRequiresExecPin = true;
-		Sig.bRequiresContext = false;
-		Sig.bSupportsCPU = true;
 		Sig.bSupportsGPU = false;
-#if WITH_EDITORONLY_DATA
-		Sig.FunctionVersion = EFunctionVersion::LatestVersion;
-#endif
 	}
 
 	{
-		FNiagaraFunctionSignature& Sig = OutFunctions.AddDefaulted_GetRef();
+		FNiagaraFunctionSignature& Sig = OutFunctions.Add_GetRef(DefaultSig);
 		Sig.Name = SetValueFunctionName;
-		Sig.Inputs.Add(FNiagaraVariable(FNiagaraTypeDefinition(GetClass()), TEXT("RenderTarget")));
-		Sig.Inputs.Add_GetRef(FNiagaraVariable(FNiagaraTypeDefinition::GetBoolDef(), TEXT("Enabled"))).SetValue(true);
-		Sig.Inputs.Add(FNiagaraVariable(FNiagaraTypeDefinition::GetIntDef(), TEXT("IndexX")));
-		Sig.Inputs.Add(FNiagaraVariable(FNiagaraTypeDefinition::GetIntDef(), TEXT("IndexY")));
-		Sig.Inputs.Add(FNiagaraVariable(FNiagaraTypeDefinition::GetIntDef(), TEXT("IndexZ")));
-		Sig.Inputs.Add(FNiagaraVariable(FNiagaraTypeDefinition::GetColorDef(), TEXT("Value")));
-
-		Sig.bExperimental = true;
-		Sig.bMemberFunction = true;
+		Sig.Inputs.Emplace_GetRef(FNiagaraTypeDefinition::GetBoolDef(), TEXT("Enabled")).SetValue(true);
+		Sig.Inputs.Emplace(FNiagaraTypeDefinition::GetIntDef(), TEXT("IndexX"));
+		Sig.Inputs.Emplace(FNiagaraTypeDefinition::GetIntDef(), TEXT("IndexY"));
+		Sig.Inputs.Emplace(FNiagaraTypeDefinition::GetIntDef(), TEXT("IndexZ"));
+		Sig.Inputs.Emplace(FNiagaraTypeDefinition::GetColorDef(), TEXT("Value"));
 		Sig.bRequiresExecPin = true;
-		Sig.bRequiresContext = false;
 		Sig.bWriteFunction = true;
 		Sig.bSupportsCPU = false;
-		Sig.bSupportsGPU = true;
-	#if WITH_EDITORONLY_DATA
-		Sig.FunctionVersion = EFunctionVersion::LatestVersion;
-	#endif
 	}
 
 	{
-		FNiagaraFunctionSignature& Sig = OutFunctions.AddDefaulted_GetRef();
-		Sig.Name = GetValueFunctionName;
-		Sig.Inputs.Add(FNiagaraVariable(FNiagaraTypeDefinition(GetClass()), TEXT("RenderTarget")));
-		Sig.Inputs.Add(FNiagaraVariable(FNiagaraTypeDefinition::GetIntDef(), TEXT("IndexX")));
-		Sig.Inputs.Add(FNiagaraVariable(FNiagaraTypeDefinition::GetIntDef(), TEXT("IndexY")));
-		Sig.Inputs.Add(FNiagaraVariable(FNiagaraTypeDefinition::GetIntDef(), TEXT("IndexZ")));
-		Sig.Outputs.Add(FNiagaraVariable(FNiagaraTypeDefinition::GetColorDef(), TEXT("Value")));
-
-		Sig.bHidden = NiagaraDataInterfaceRenderTargetCommon::GAllowReads != 1;
-		Sig.bExperimental = true;
-		Sig.bMemberFunction = true;
-		Sig.bRequiresContext = false;
-		Sig.bWriteFunction = true;
+		FNiagaraFunctionSignature& Sig = OutFunctions.Add_GetRef(DefaultSig);
+		Sig.Name = LoadValueFunctionName;
+		Sig.Inputs.Emplace(FNiagaraTypeDefinition::GetIntDef(), TEXT("IndexX"));
+		Sig.Inputs.Emplace(FNiagaraTypeDefinition::GetIntDef(), TEXT("IndexY"));
+		Sig.Inputs.Emplace(FNiagaraTypeDefinition::GetIntDef(), TEXT("IndexZ"));
+		Sig.Inputs.Emplace(FNiagaraTypeDefinition::GetIntDef(), TEXT("MipLevel"));
+		Sig.Outputs.Emplace(FNiagaraTypeDefinition::GetColorDef(), TEXT("Value"));
 		Sig.bSupportsCPU = false;
-		Sig.bSupportsGPU = true;
-#if WITH_EDITORONLY_DATA
-		Sig.FunctionVersion = EFunctionVersion::LatestVersion;
-#endif
 	}
 
 	{
-		FNiagaraFunctionSignature& Sig = OutFunctions.AddDefaulted_GetRef();
+		FNiagaraFunctionSignature& Sig = OutFunctions.Add_GetRef(DefaultSig);
+		Sig.Name = Deprecated_GetValueFunctionName;
+		Sig.Inputs.Emplace(FNiagaraTypeDefinition::GetIntDef(), TEXT("IndexX"));
+		Sig.Inputs.Emplace(FNiagaraTypeDefinition::GetIntDef(), TEXT("IndexY"));
+		Sig.Inputs.Emplace(FNiagaraTypeDefinition::GetIntDef(), TEXT("IndexZ"));
+		Sig.Outputs.Emplace(FNiagaraTypeDefinition::GetColorDef(), TEXT("Value"));
+		Sig.bSupportsCPU = false;
+	}
+
+	{
+		FNiagaraFunctionSignature& Sig = OutFunctions.Add_GetRef(DefaultSig);
 		Sig.Name = SampleValueFunctionName;
-		Sig.Inputs.Add(FNiagaraVariable(FNiagaraTypeDefinition(GetClass()), TEXT("RenderTarget")));
-		Sig.Inputs.Add(FNiagaraVariable(FNiagaraTypeDefinition::GetVec3Def(), TEXT("UVW")));
-		Sig.Outputs.Add(FNiagaraVariable(FNiagaraTypeDefinition::GetColorDef(), TEXT("Value")));
-
-		Sig.bHidden = NiagaraDataInterfaceRenderTargetCommon::GAllowReads != 1;
-		Sig.bExperimental = true;
-		Sig.bMemberFunction = true;
-		Sig.bRequiresContext = false;
-		Sig.bWriteFunction = true;
+		Sig.Inputs.Emplace(FNiagaraTypeDefinition::GetVec3Def(), TEXT("UVW"));
+		Sig.Outputs.Emplace(FNiagaraTypeDefinition::GetColorDef(), TEXT("Value"));
 		Sig.bSupportsCPU = false;
-		Sig.bSupportsGPU = true;
-#if WITH_EDITORONLY_DATA
-		Sig.FunctionVersion = EFunctionVersion::LatestVersion;
-#endif
 	}
 
 	{
-		FNiagaraFunctionSignature& Sig = OutFunctions.AddDefaulted_GetRef();
+		FNiagaraFunctionSignature& Sig = OutFunctions.Add_GetRef(DefaultSig);
 		Sig.Name = LinearToIndexName;
-		Sig.Inputs.Add(FNiagaraVariable(FNiagaraTypeDefinition(GetClass()), TEXT("RenderTarget")));
-		Sig.Inputs.Add(FNiagaraVariable(FNiagaraTypeDefinition::GetIntDef(), TEXT("Linear")));
-		Sig.Outputs.Add(FNiagaraVariable(FNiagaraTypeDefinition::GetIntDef(), TEXT("IndexX")));
-		Sig.Outputs.Add(FNiagaraVariable(FNiagaraTypeDefinition::GetIntDef(), TEXT("IndexY")));
-		Sig.Outputs.Add(FNiagaraVariable(FNiagaraTypeDefinition::GetIntDef(), TEXT("IndexZ")));
-
-		Sig.bExperimental = true;
-		Sig.bMemberFunction = true;
-		Sig.bRequiresContext = false;
-		Sig.bWriteFunction = true;
+		Sig.Inputs.Emplace(FNiagaraTypeDefinition::GetIntDef(), TEXT("Linear"));
+		Sig.Outputs.Emplace(FNiagaraTypeDefinition::GetIntDef(), TEXT("IndexX"));
+		Sig.Outputs.Emplace(FNiagaraTypeDefinition::GetIntDef(), TEXT("IndexY"));
+		Sig.Outputs.Emplace(FNiagaraTypeDefinition::GetIntDef(), TEXT("IndexZ"));
 		Sig.bSupportsCPU = false;
-		Sig.bSupportsGPU = true;
-	#if WITH_EDITORONLY_DATA
-		Sig.FunctionVersion = EFunctionVersion::LatestVersion;
-	#endif
 	}
 
 	{
-		FNiagaraFunctionSignature& Sig = OutFunctions.AddDefaulted_GetRef();
+		FNiagaraFunctionSignature& Sig = OutFunctions.Add_GetRef(DefaultSig);
 		Sig.Name = ExecToIndexName;
-		Sig.Inputs.Add(FNiagaraVariable(FNiagaraTypeDefinition(GetClass()), TEXT("RenderTarget")));
-		Sig.Outputs.Add(FNiagaraVariable(FNiagaraTypeDefinition::GetIntDef(), TEXT("IndexX")));
-		Sig.Outputs.Add(FNiagaraVariable(FNiagaraTypeDefinition::GetIntDef(), TEXT("IndexY")));
-		Sig.Outputs.Add(FNiagaraVariable(FNiagaraTypeDefinition::GetIntDef(), TEXT("IndexZ")));
-
-		Sig.bExperimental = true;
-		Sig.bMemberFunction = true;
-		Sig.bRequiresContext = false;
-		Sig.bWriteFunction = true;
+		Sig.Outputs.Emplace(FNiagaraTypeDefinition::GetIntDef(), TEXT("IndexX"));
+		Sig.Outputs.Emplace(FNiagaraTypeDefinition::GetIntDef(), TEXT("IndexY"));
+		Sig.Outputs.Emplace(FNiagaraTypeDefinition::GetIntDef(), TEXT("IndexZ"));
 		Sig.bSupportsCPU = false;
-		Sig.bSupportsGPU = true;
-	#if WITH_EDITORONLY_DATA
-		Sig.FunctionVersion = EFunctionVersion::LatestVersion;
-	#endif
 	}
 
 	{
-		FNiagaraFunctionSignature& Sig = OutFunctions.AddDefaulted_GetRef();
+		FNiagaraFunctionSignature& Sig = OutFunctions.Add_GetRef(DefaultSig);
 		Sig.Name = ExecToUnitName;
-		Sig.Inputs.Emplace(FNiagaraTypeDefinition(GetClass()), TEXT("RenderTarget"));
 		Sig.Outputs.Emplace(FNiagaraTypeDefinition::GetVec3Def(), TEXT("Unit"));
-
-		Sig.bExperimental = true;
-		Sig.bMemberFunction = true;
-		Sig.bRequiresContext = false;
-		Sig.bWriteFunction = true;
 		Sig.bSupportsCPU = false;
-		Sig.bSupportsGPU = true;
-#if WITH_EDITORONLY_DATA
-		Sig.FunctionVersion = EFunctionVersion::LatestVersion;
-#endif
-}
+	}
 }
 
 #if WITH_EDITORONLY_DATA
@@ -347,9 +288,6 @@ bool UNiagaraDataInterfaceRenderTargetVolume::UpgradeFunctionCall(FNiagaraFuncti
 }
 #endif
 
-DEFINE_NDI_DIRECT_FUNC_BINDER(UNiagaraDataInterfaceRenderTargetVolume, VMGetSize);
-DEFINE_NDI_DIRECT_FUNC_BINDER(UNiagaraDataInterfaceRenderTargetVolume, VMSetSize);
-DEFINE_NDI_DIRECT_FUNC_BINDER(UNiagaraDataInterfaceRenderTargetVolume, VMSetFormat);
 void UNiagaraDataInterfaceRenderTargetVolume::GetVMExternalFunction(const FVMExternalFunctionBindingInfo& BindingInfo, void* InstanceData, FVMExternalFunction &OutFunc)
 {
 	using namespace NDIRenderTargetVolumeLocal;
@@ -357,18 +295,19 @@ void UNiagaraDataInterfaceRenderTargetVolume::GetVMExternalFunction(const FVMExt
 	Super::GetVMExternalFunction(BindingInfo, InstanceData, OutFunc);
 	if (BindingInfo.Name == GetSizeFunctionName)
 	{
-		check(BindingInfo.GetNumInputs() == 1 && BindingInfo.GetNumOutputs() == 3);
-		NDI_FUNC_BINDER(UNiagaraDataInterfaceRenderTargetVolume, VMGetSize)::Bind(this, OutFunc);
+		OutFunc = FVMExternalFunction::CreateUObject(this, &UNiagaraDataInterfaceRenderTargetVolume::VMGetSize);
 	}
 	else if (BindingInfo.Name == SetSizeFunctionName)
 	{
-		check(BindingInfo.GetNumInputs() == 4 && BindingInfo.GetNumOutputs() == 1);
-		NDI_FUNC_BINDER(UNiagaraDataInterfaceRenderTargetVolume, VMSetSize)::Bind(this, OutFunc);
+		OutFunc = FVMExternalFunction::CreateUObject(this, &UNiagaraDataInterfaceRenderTargetVolume::VMSetSize);
+	}
+	else if (BindingInfo.Name == GetNumMipLevelsName)
+	{
+		OutFunc = FVMExternalFunction::CreateUObject(this, &UNiagaraDataInterfaceRenderTargetVolume::VMGetNumMipLevels);
 	}
 	else if (BindingInfo.Name == SetFormatFunctionName)
 	{
-		check(BindingInfo.GetNumInputs() == 2 && BindingInfo.GetNumOutputs() == 1);
-		NDI_FUNC_BINDER(UNiagaraDataInterfaceRenderTargetVolume, VMSetFormat)::Bind(this, OutFunc);
+		OutFunc = FVMExternalFunction::CreateUObject(this, &UNiagaraDataInterfaceRenderTargetVolume::VMSetFormat);
 	}
 }
 
@@ -454,7 +393,8 @@ bool UNiagaraDataInterfaceRenderTargetVolume::GetFunctionHLSL(const FNiagaraData
 	using namespace NDIRenderTargetVolumeLocal;
 
 	if ((FunctionInfo.DefinitionName == SetValueFunctionName) ||
-		(FunctionInfo.DefinitionName == GetValueFunctionName) ||
+		(FunctionInfo.DefinitionName == Deprecated_GetValueFunctionName) ||
+		(FunctionInfo.DefinitionName == LoadValueFunctionName) ||
 		(FunctionInfo.DefinitionName == SampleValueFunctionName) ||
 		(FunctionInfo.DefinitionName == GetSizeFunctionName) ||
 		(FunctionInfo.DefinitionName == LinearToIndexName) ||
@@ -492,6 +432,7 @@ void UNiagaraDataInterfaceRenderTargetVolume::SetShaderParameters(const FNiagara
 
 	NDIRenderTargetVolumeLocal::FShaderParameters* Parameters = Context.GetParameterNestedStruct<NDIRenderTargetVolumeLocal::FShaderParameters>();
 	Parameters->TextureSize = InstanceData_RT->Size;
+	Parameters->MipLevels = InstanceData_RT->MipLevels;
 
 	const bool bRTWrite = Context.IsResourceBound(&Parameters->RWTexture);
 	const bool bRTRead = Context.IsResourceBound(&Parameters->Texture);
@@ -923,6 +864,41 @@ void UNiagaraDataInterfaceRenderTargetVolume::VMSetSize(FVectorVMExternalFunctio
 	}
 }
 
+void UNiagaraDataInterfaceRenderTargetVolume::VMGetSize(FVectorVMExternalFunctionContext& Context)
+{
+	VectorVM::FUserPtrHandler<FRenderTargetVolumeRWInstanceData_GameThread> InstData(Context);
+	FNDIOutputParam<int> OutSizeX(Context);
+	FNDIOutputParam<int> OutSizeY(Context);
+	FNDIOutputParam<int> OutSizeZ(Context);
+
+	for (int32 InstanceIdx = 0; InstanceIdx < Context.GetNumInstances(); ++InstanceIdx)
+	{
+		OutSizeX.SetAndAdvance(InstData->Size.X);
+		OutSizeY.SetAndAdvance(InstData->Size.Y);
+		OutSizeZ.SetAndAdvance(InstData->Size.Z);
+	}
+}
+
+void UNiagaraDataInterfaceRenderTargetVolume::VMGetNumMipLevels(FVectorVMExternalFunctionContext& Context)
+{
+	VectorVM::FUserPtrHandler<FRenderTargetVolumeRWInstanceData_GameThread> InstData(Context);
+	FNDIOutputParam<int> OutMipLevels(Context);
+
+	int32 NumMipLevels = 1;
+	if (InstData->TargetTexture != nullptr)
+	{
+		NumMipLevels = InstData->TargetTexture->GetNumMips();
+	}
+	//else if (InstData->MipMapGeneration != ENiagaraMipMapGeneration::Disabled)
+	//{
+	//	NumMipLevels = FMath::FloorLog2(FMath::Max(InstData->Size.X, InstData->Size.Y)) + 1;
+	//}
+	for (int32 InstanceIdx = 0; InstanceIdx < Context.GetNumInstances(); ++InstanceIdx)
+	{
+		OutMipLevels.SetAndAdvance(NumMipLevels);
+	}
+}
+
 void UNiagaraDataInterfaceRenderTargetVolume::VMSetFormat(FVectorVMExternalFunctionContext& Context)
 {
 	// This should only be called from a system or emitter script due to a need for only setting up initially.
@@ -949,22 +925,6 @@ void UNiagaraDataInterfaceRenderTargetVolume::VMSetFormat(FVectorVMExternalFunct
 	}
 }
 
-void UNiagaraDataInterfaceRenderTargetVolume::VMGetSize(FVectorVMExternalFunctionContext& Context)
-{
-	VectorVM::FUserPtrHandler<FRenderTargetVolumeRWInstanceData_GameThread> InstData(Context);
-	FNDIOutputParam<int> OutSizeX(Context);
-	FNDIOutputParam<int> OutSizeY(Context);
-	FNDIOutputParam<int> OutSizeZ(Context);
-
-	for (int32 InstanceIdx = 0; InstanceIdx < Context.GetNumInstances(); ++InstanceIdx)
-	{
-		OutSizeX.SetAndAdvance(InstData->Size.X);
-		OutSizeY.SetAndAdvance(InstData->Size.Y);
-		OutSizeZ.SetAndAdvance(InstData->Size.Z);
-	}
-}
-
-
 bool UNiagaraDataInterfaceRenderTargetVolume::PerInstanceTick(void* PerInstanceData, FNiagaraSystemInstance* SystemInstance, float DeltaSeconds)
 {
 	FRenderTargetVolumeRWInstanceData_GameThread* InstanceData = static_cast<FRenderTargetVolumeRWInstanceData_GameThread*>(PerInstanceData);
@@ -1071,6 +1031,7 @@ bool UNiagaraDataInterfaceRenderTargetVolume::PerInstanceTickPostSimulate(void*
 			{
 				FRenderTargetVolumeRWInstanceData_RenderThread* TargetData = &RT_Proxy->SystemInstancesToProxyData_RT.FindOrAdd(RT_InstanceID);
 				TargetData->Size = RT_InstanceData.Size;
+				TargetData->MipLevels = RT_TargetTexture->GetCurrentMipCount();
 				//TargetData->MipMapGeneration = RT_InstanceData.MipMapGeneration;
 			#if WITH_EDITORONLY_DATA
 				TargetData->bPreviewTexture = RT_InstanceData.bPreviewTexture;
