- Add MipLevel to RT Sample functions

diff --git a/Engine/Plugins/FX/Niagara/Shaders/Private/NiagaraDataInterfaceRenderTarget2DArrayTemplate.ush b/Engine/Plugins/FX/Niagara/Shaders/Private/NiagaraDataInterfaceRenderTarget2DArrayTemplate.ush
--- a/Engine/Plugins/FX/Niagara/Shaders/Private/NiagaraDataInterfaceRenderTarget2DArrayTemplate.ush
+++ b/Engine/Plugins/FX/Niagara/Shaders/Private/NiagaraDataInterfaceRenderTarget2DArrayTemplate.ush
@@ -24,9 +24,9 @@ void LoadRenderTargetValue_{ParameterName}(int IndexX, int IndexY, int IndexZ, i
 	Value = {ParameterName}_Texture.Load(int4(IndexX, IndexY, IndexZ, MipLevel));
 }
 
-void SampleRenderTargetValue_{ParameterName}(float2 UV, int Slice, out float4 Value)
+void SampleRenderTargetValue_{ParameterName}(float2 UV, int Slice, float MipLevel, out float4 Value)
 {
-	Value = {ParameterName}_Texture.SampleLevel({ParameterName}_TextureSampler, float3(UV.x, UV.y, Slice), 0.0f);
+	Value = {ParameterName}_Texture.SampleLevel({ParameterName}_TextureSampler, float3(UV.x, UV.y, Slice), MipLevel);
 }
 
 void GetRenderTargetSize_{ParameterName}(out int Width, out int Height, out int Slices)
diff --git a/Engine/Plugins/FX/Niagara/Shaders/Private/NiagaraDataInterfaceRenderTarget2DTemplate.ush b/Engine/Plugins/FX/Niagara/Shaders/Private/NiagaraDataInterfaceRenderTarget2DTemplate.ush
--- a/Engine/Plugins/FX/Niagara/Shaders/Private/NiagaraDataInterfaceRenderTarget2DTemplate.ush
+++ b/Engine/Plugins/FX/Niagara/Shaders/Private/NiagaraDataInterfaceRenderTarget2DTemplate.ush
@@ -24,9 +24,9 @@ void LoadRenderTargetValue_{ParameterName}(int IndexX, int IndexY, int MipLevel,
 	Value = {ParameterName}_Texture.Load(int3(IndexX, IndexY, MipLevel));
 }
 
-void SampleRenderTargetValue_{ParameterName}(float2 UV, out float4 Value)
+void SampleRenderTargetValue_{ParameterName}(float2 UV, float MipLevel, out float4 Value)
 {
-	Value = {ParameterName}_Texture.SampleLevel({ParameterName}_TextureSampler, UV, 0.0f);
+	Value = {ParameterName}_Texture.SampleLevel({ParameterName}_TextureSampler, UV, MipLevel);
 }
 
 void GetRenderTargetSize_{ParameterName}(out int Width, out int Height)
diff --git a/Engine/Plugins/FX/Niagara/Shaders/Private/NiagaraDataInterfaceRenderTargetCubeTemplate.ush b/Engine/Plugins/FX/Niagara/Shaders/Private/NiagaraDataInterfaceRenderTargetCubeTemplate.ush
--- a/Engine/Plugins/FX/Niagara/Shaders/Private/NiagaraDataInterfaceRenderTargetCubeTemplate.ush
+++ b/Engine/Plugins/FX/Niagara/Shaders/Private/NiagaraDataInterfaceRenderTargetCubeTemplate.ush
@@ -25,9 +25,9 @@ void SetRenderTargetValue_{ParameterName}_UEImpureCall(bool bExecute, int IndexX
 //	Value = {ParameterName}_Texture.Load(int4(IndexX, IndexY, Face, MipLevel));
 //}
 
-void SampleRenderTargetValue_{ParameterName}(float3 UVW, out float4 Value)
+void SampleRenderTargetValue_{ParameterName}(float3 UVW, float MipLevel, out float4 Value)
 {
-	Value = {ParameterName}_Texture.SampleLevel({ParameterName}_TextureSampler, UVW, 0.0f);
+	Value = {ParameterName}_Texture.SampleLevel({ParameterName}_TextureSampler, UVW, MipLevel);
 }
 
 void GetRenderTargetSize_{ParameterName}(out int Size)
diff --git a/Engine/Plugins/FX/Niagara/Shaders/Private/NiagaraDataInterfaceRenderTargetVolumeTemplate.ush b/Engine/Plugins/FX/Niagara/Shaders/Private/NiagaraDataInterfaceRenderTargetVolumeTemplate.ush
--- a/Engine/Plugins/FX/Niagara/Shaders/Private/NiagaraDataInterfaceRenderTargetVolumeTemplate.ush
+++ b/Engine/Plugins/FX/Niagara/Shaders/Private/NiagaraDataInterfaceRenderTargetVolumeTemplate.ush
@@ -24,9 +24,9 @@ void LoadRenderTargetValue_{ParameterName}(int IndexX, int IndexY, int IndexZ, i
 	Value = {ParameterName}_Texture.Load(int4(IndexX, IndexY, IndexZ, MipLevel));
 }
 
-void SampleRenderTargetValue_{ParameterName}(float3 UVW, out float4 Value)
+void SampleRenderTargetValue_{ParameterName}(float3 UVW, float MipLevel, out float4 Value)
 {
-	Value = {ParameterName}_Texture.SampleLevel({ParameterName}_TextureSampler, UVW, 0.0f);
+	Value = {ParameterName}_Texture.SampleLevel({ParameterName}_TextureSampler, UVW, MipLevel);
 }
 
 void GetRenderTargetSize_{ParameterName}(out int Width, out int Height, out int Slices)
diff --git a/Engine/Plugins/FX/Niagara/Source/Niagara/Private/NiagaraDataInterfaceRenderTarget2D.cpp b/Engine/Plugins/FX/Niagara/Source/Niagara/Private/NiagaraDataInterfaceRenderTarget2D.cpp
--- a/Engine/Plugins/FX/Niagara/Source/Niagara/Private/NiagaraDataInterfaceRenderTarget2D.cpp
+++ b/Engine/Plugins/FX/Niagara/Source/Niagara/Private/NiagaraDataInterfaceRenderTarget2D.cpp
@@ -60,6 +60,7 @@ struct FNDIRenderTarget2DFunctionVersion
 	{
 		InitialVersion = 0,
 		AddedOptionalExecute = 1,
+		AddedMipLevel = 2,
 
 		VersionPlusOne,
 		LatestVersion = VersionPlusOne - 1
@@ -186,6 +187,7 @@ void UNiagaraDataInterfaceRenderTarget2D::GetFunctions(TArray<FNiagaraFunctionSi
 		FNiagaraFunctionSignature& Sig = OutFunctions.Add_GetRef(DefaultSig);
 		Sig.Name = SampleValueFunctionName;
 		Sig.Inputs.Emplace(FNiagaraTypeDefinition::GetVec2Def(), TEXT("UV"));
+		Sig.Inputs.Emplace(FNiagaraTypeDefinition::GetFloatDef(), TEXT("MipLevel"));
 		Sig.Outputs.Emplace(FNiagaraTypeDefinition::GetColorDef(), TEXT("Value"));
 		Sig.bSupportsCPU = false;
 	}
@@ -241,6 +243,14 @@ bool UNiagaraDataInterfaceRenderTarget2D::UpgradeFunctionCall(FNiagaraFunctionSi
 			bWasChanged = true;
 		}
 	}
+	if (FunctionSignature.FunctionVersion < FNDIRenderTarget2DFunctionVersion::AddedMipLevel)
+	{
+		if (FunctionSignature.Name == SampleValueFunctionName)
+		{
+			FunctionSignature.Inputs.Emplace(FNiagaraTypeDefinition::GetFloatDef(), TEXT("MipLevel"));
+			bWasChanged = true;
+		}
+	}
 
 	// Set latest version
 	FunctionSignature.FunctionVersion = FNDIRenderTarget2DFunctionVersion::LatestVersion;
diff --git a/Engine/Plugins/FX/Niagara/Source/Niagara/Private/NiagaraDataInterfaceRenderTarget2DArray.cpp b/Engine/Plugins/FX/Niagara/Source/Niagara/Private/NiagaraDataInterfaceRenderTarget2DArray.cpp
--- a/Engine/Plugins/FX/Niagara/Source/Niagara/Private/NiagaraDataInterfaceRenderTarget2DArray.cpp
+++ b/Engine/Plugins/FX/Niagara/Source/Niagara/Private/NiagaraDataInterfaceRenderTarget2DArray.cpp
@@ -58,6 +58,7 @@ struct FNDIRenderTarget2DArrayFunctionVersion
 	{
 		InitialVersion = 0,
 		AddedOptionalExecute = 1,
+		AddedMipLevel = 2,
 
 		VersionPlusOne,
 		LatestVersion = VersionPlusOne - 1
@@ -196,9 +197,9 @@ void UNiagaraDataInterfaceRenderTarget2DArray::GetFunctions(TArray<FNiagaraFunct
 	{
 		FNiagaraFunctionSignature& Sig = OutFunctions.Add_GetRef(DefaultSig);
 		Sig.Name = SampleValueFunctionName;
-		Sig.Inputs.Emplace(FNiagaraTypeDefinition(GetClass()), TEXT("RenderTarget"));
 		Sig.Inputs.Emplace(FNiagaraTypeDefinition::GetVec2Def(), TEXT("UV"));
 		Sig.Inputs.Emplace(FNiagaraTypeDefinition::GetIntDef(), TEXT("Slice"));
+		Sig.Inputs.Emplace(FNiagaraTypeDefinition::GetFloatDef(), TEXT("MipLevel"));
 		Sig.Outputs.Emplace(FNiagaraTypeDefinition::GetColorDef(), TEXT("Value"));
 		Sig.bSupportsCPU = false;
 	}
@@ -234,17 +235,27 @@ void UNiagaraDataInterfaceRenderTarget2DArray::GetFunctions(TArray<FNiagaraFunct
 #if WITH_EDITORONLY_DATA
 bool UNiagaraDataInterfaceRenderTarget2DArray::UpgradeFunctionCall(FNiagaraFunctionSignature& FunctionSignature)
 {
+	using namespace NDIRenderTarget2DArrayLocal;
+
 	bool bWasChanged = false;
 
 	if (FunctionSignature.FunctionVersion < FNDIRenderTarget2DArrayFunctionVersion::AddedOptionalExecute)
 	{
-		if (FunctionSignature.Name == NDIRenderTarget2DArrayLocal::SetValueFunctionName)
+		if (FunctionSignature.Name == SetValueFunctionName)
 		{
 			check(FunctionSignature.Inputs.Num() == 5);
 			FunctionSignature.Inputs.Insert_GetRef(FNiagaraVariable(FNiagaraTypeDefinition::GetBoolDef(), TEXT("Enabled")), 1).SetValue(true);
 			bWasChanged = true;
 		}
 	}
+	if (FunctionSignature.FunctionVersion < FNDIRenderTarget2DArrayFunctionVersion::AddedMipLevel)
+	{
+		if (FunctionSignature.Name == SampleValueFunctionName)
+		{
+			FunctionSignature.Inputs.Emplace(FNiagaraTypeDefinition::GetFloatDef(), TEXT("MipLevel"));
+			bWasChanged = true;
+		}
+	}
 
 	// Set latest version
 	FunctionSignature.FunctionVersion = FNDIRenderTarget2DArrayFunctionVersion::LatestVersion;
diff --git a/Engine/Plugins/FX/Niagara/Source/Niagara/Private/NiagaraDataInterfaceRenderTargetCube.cpp b/Engine/Plugins/FX/Niagara/Source/Niagara/Private/NiagaraDataInterfaceRenderTargetCube.cpp
--- a/Engine/Plugins/FX/Niagara/Source/Niagara/Private/NiagaraDataInterfaceRenderTargetCube.cpp
+++ b/Engine/Plugins/FX/Niagara/Source/Niagara/Private/NiagaraDataInterfaceRenderTargetCube.cpp
@@ -51,6 +51,7 @@ namespace NDIRenderTargetCubeLocal
 		{
 			InitialVersion = 0,
 			AddedOptionalExecute = 1,
+			AddedMipLevel = 2,
 
 			VersionPlusOne,
 			LatestVersion = VersionPlusOne - 1
@@ -168,6 +169,7 @@ void UNiagaraDataInterfaceRenderTargetCube::GetFunctions(TArray<FNiagaraFunction
 		FNiagaraFunctionSignature& Sig = OutFunctions.Add_GetRef(DefaultSig);
 		Sig.Name = SampleValueFunctionName;
 		Sig.Inputs.Emplace(FNiagaraTypeDefinition::GetVec3Def(), TEXT("UVW"));
+		Sig.Inputs.Emplace(FNiagaraTypeDefinition::GetFloatDef(), TEXT("MipLevel"));
 		Sig.Outputs.Emplace(FNiagaraTypeDefinition::GetColorDef(), TEXT("Value"));
 		Sig.bSupportsCPU = false;
 	}
@@ -194,7 +196,6 @@ void UNiagaraDataInterfaceRenderTargetCube::GetFunctions(TArray<FNiagaraFunction
 	{
 		FNiagaraFunctionSignature& Sig = OutFunctions.Add_GetRef(DefaultSig);
 		Sig.Name = ExecToUnitName;
-		Sig.Inputs.Emplace(FNiagaraTypeDefinition(GetClass()), TEXT("RenderTarget"));
 		Sig.Outputs.Emplace(FNiagaraTypeDefinition::GetVec2Def(), TEXT("Unit"));
 		Sig.Outputs.Emplace(FNiagaraTypeDefinition::GetIntDef(), TEXT("Face"));
 		Sig.bSupportsCPU = false;
@@ -217,6 +218,14 @@ bool UNiagaraDataInterfaceRenderTargetCube::UpgradeFunctionCall(FNiagaraFunction
 			bWasChanged = true;
 		}
 	}
+	if (FunctionSignature.FunctionVersion < EFunctionVersion::AddedMipLevel)
+	{
+		if (FunctionSignature.Name == SampleValueFunctionName)
+		{
+			FunctionSignature.Inputs.Emplace(FNiagaraTypeDefinition::GetFloatDef(), TEXT("MipLevel"));
+			bWasChanged = true;
+		}
+	}
 
 	// Set latest version
 	FunctionSignature.FunctionVersion = EFunctionVersion::LatestVersion;
diff --git a/Engine/Plugins/FX/Niagara/Source/Niagara/Private/NiagaraDataInterfaceRenderTargetVolume.cpp b/Engine/Plugins/FX/Niagara/Source/Niagara/Private/NiagaraDataInterfaceRenderTargetVolume.cpp
--- a/Engine/Plugins/FX/Niagara/Source/Niagara/Private/NiagaraDataInterfaceRenderTargetVolume.cpp
+++ b/Engine/Plugins/FX/Niagara/Source/Niagara/Private/NiagaraDataInterfaceRenderTargetVolume.cpp
@@ -57,6 +57,7 @@ namespace NDIRenderTargetVolumeLocal
 		{
 			InitialVersion = 0,
 			AddedOptionalExecute = 1,
+			AddedMipLevel = 2,
 
 			VersionPlusOne,
 			LatestVersion = VersionPlusOne - 1
@@ -233,6 +234,7 @@ void UNiagaraDataInterfaceRenderTargetVolume::GetFunctions(TArray<FNiagaraFuncti
 		FNiagaraFunctionSignature& Sig = OutFunctions.Add_GetRef(DefaultSig);
 		Sig.Name = SampleValueFunctionName;
 		Sig.Inputs.Emplace(FNiagaraTypeDefinition::GetVec3Def(), TEXT("UVW"));
+		Sig.Inputs.Emplace(FNiagaraTypeDefinition::GetFloatDef(), TEXT("MipLevel"));
 		Sig.Outputs.Emplace(FNiagaraTypeDefinition::GetColorDef(), TEXT("Value"));
 		Sig.bSupportsCPU = false;
 	}
@@ -280,6 +282,14 @@ bool UNiagaraDataInterfaceRenderTargetVolume::UpgradeFunctionCall(FNiagaraFuncti
 			bWasChanged = true;
 		}
 	}
+	if (FunctionSignature.FunctionVersion < EFunctionVersion::AddedMipLevel)
+	{
+		if (FunctionSignature.Name == SampleValueFunctionName)
+		{
+			FunctionSignature.Inputs.Emplace(FNiagaraTypeDefinition::GetFloatDef(), TEXT("MipLevel"));
+			bWasChanged = true;
+		}
+	}
 
 	// Set latest version
 	FunctionSignature.FunctionVersion = EFunctionVersion::LatestVersion;
