- Fix for GetNumMipLevels being used on GPU

#rnx

diff --git a/Engine/Plugins/FX/Niagara/Source/Niagara/Private/NiagaraDataInterfaceRenderTarget2D.cpp b/Engine/Plugins/FX/Niagara/Source/Niagara/Private/NiagaraDataInterfaceRenderTarget2D.cpp
--- a/Engine/Plugins/FX/Niagara/Source/Niagara/Private/NiagaraDataInterfaceRenderTarget2D.cpp
+++ b/Engine/Plugins/FX/Niagara/Source/Niagara/Private/NiagaraDataInterfaceRenderTarget2D.cpp
@@ -365,6 +365,7 @@ bool UNiagaraDataInterfaceRenderTarget2D::GetFunctionHLSL(const FNiagaraDataInte
 		(FunctionInfo.DefinitionName == LoadValueFunctionName) ||
 		(FunctionInfo.DefinitionName == SampleValueFunctionName) ||
 		(FunctionInfo.DefinitionName == GetSizeFunctionName) ||
+		(FunctionInfo.DefinitionName == GetNumMipLevelsName) ||
 		(FunctionInfo.DefinitionName == LinearToIndexName) ||
 		(FunctionInfo.DefinitionName == ExecToIndexName) ||
 		(FunctionInfo.DefinitionName == ExecToUnitName))
