- Static analysis fix

#rnx

diff --git a/Engine/Plugins/FX/Niagara/Source/Niagara/Private/NiagaraDataInterfaceRenderTarget2D.cpp b/Engine/Plugins/FX/Niagara/Source/Niagara/Private/NiagaraDataInterfaceRenderTarget2D.cpp
--- a/Engine/Plugins/FX/Niagara/Source/Niagara/Private/NiagaraDataInterfaceRenderTarget2D.cpp
+++ b/Engine/Plugins/FX/Niagara/Source/Niagara/Private/NiagaraDataInterfaceRenderTarget2D.cpp
@@ -755,7 +755,7 @@ bool UNiagaraDataInterfaceRenderTarget2D::PerInstanceTickPostSimulate(void* PerI
 			{
 				FRenderTarget2DRWInstanceData_RenderThread* TargetData = &RT_Proxy->SystemInstancesToProxyData_RT.FindOrAdd(RT_InstanceID);
 				TargetData->Size = RT_InstanceData.Size;
-				TargetData->MipLevels = RT_TargetTexture->GetCurrentMipCount();
+				TargetData->MipLevels = 1;
 				TargetData->MipMapGeneration = RT_InstanceData.MipMapGeneration;
 				TargetData->MipMapGenerationType = RT_InstanceData.MipMapGenerationType;
 			#if WITH_EDITORONLY_DATA
@@ -765,6 +765,7 @@ bool UNiagaraDataInterfaceRenderTarget2D::PerInstanceTickPostSimulate(void* PerI
 				TargetData->TextureRHI.SafeRelease();
 				if (RT_TargetTexture)
 				{
+					TargetData->MipLevels = RT_TargetTexture->GetCurrentMipCount();
 					if (FTextureRenderTarget2DResource* Resource2D = RT_TargetTexture->GetTextureRenderTarget2DResource())
 					{
 						TargetData->SamplerStateRHI = Resource2D->SamplerStateRHI;
diff --git a/Engine/Plugins/FX/Niagara/Source/Niagara/Private/NiagaraDataInterfaceRenderTarget2DArray.cpp b/Engine/Plugins/FX/Niagara/Source/Niagara/Private/NiagaraDataInterfaceRenderTarget2DArray.cpp
--- a/Engine/Plugins/FX/Niagara/Source/Niagara/Private/NiagaraDataInterfaceRenderTarget2DArray.cpp
+++ b/Engine/Plugins/FX/Niagara/Source/Niagara/Private/NiagaraDataInterfaceRenderTarget2DArray.cpp
@@ -713,7 +713,7 @@ bool UNiagaraDataInterfaceRenderTarget2DArray::PerInstanceTickPostSimulate(void*
 			{
 				FRenderTarget2DArrayRWInstanceData_RenderThread* TargetData = &RT_Proxy->SystemInstancesToProxyData_RT.FindOrAdd(RT_InstanceID);
 				TargetData->Size = RT_InstanceData.Size;
-				TargetData->MipLevels = RT_TargetTexture->GetCurrentMipCount();
+				TargetData->MipLevels = 1;
 				//TargetData->MipMapGeneration = RT_InstanceData.MipMapGeneration;
 			#if WITH_EDITORONLY_DATA
 				TargetData->bPreviewTexture = RT_InstanceData.bPreviewTexture;
@@ -722,6 +722,7 @@ bool UNiagaraDataInterfaceRenderTarget2DArray::PerInstanceTickPostSimulate(void*
 				TargetData->RenderTarget.SafeRelease();
 				if (RT_TargetTexture)
 				{
+					TargetData->MipLevels = RT_TargetTexture->GetCurrentMipCount();
 					if (FTextureRenderTarget2DArrayResource* Resource2DArray = RT_TargetTexture->GetTextureRenderTarget2DArrayResource())
 					{
 						TargetData->SamplerStateRHI = Resource2DArray->SamplerStateRHI;
diff --git a/Engine/Plugins/FX/Niagara/Source/Niagara/Private/NiagaraDataInterfaceRenderTargetCube.cpp b/Engine/Plugins/FX/Niagara/Source/Niagara/Private/NiagaraDataInterfaceRenderTargetCube.cpp
--- a/Engine/Plugins/FX/Niagara/Source/Niagara/Private/NiagaraDataInterfaceRenderTargetCube.cpp
+++ b/Engine/Plugins/FX/Niagara/Source/Niagara/Private/NiagaraDataInterfaceRenderTargetCube.cpp
@@ -663,7 +663,7 @@ bool UNiagaraDataInterfaceRenderTargetCube::PerInstanceTickPostSimulate(void* Pe
 			{
 				FRenderTargetCubeRWInstanceData_RenderThread* TargetData = &RT_Proxy->SystemInstancesToProxyData_RT.FindOrAdd(RT_InstanceID);
 				TargetData->Size = RT_InstanceData.Size;
-				TargetData->MipLevels = RT_TargetTexture->GetCurrentMipCount();
+				TargetData->MipLevels = 1;
 				//TargetData->MipMapGeneration = RT_InstanceData.MipMapGeneration;
 			#if WITH_EDITORONLY_DATA
 				TargetData->bPreviewTexture = RT_InstanceData.bPreviewTexture;
@@ -672,6 +672,7 @@ bool UNiagaraDataInterfaceRenderTargetCube::PerInstanceTickPostSimulate(void* Pe
 				TargetData->RenderTarget.SafeRelease();
 				if (RT_TargetTexture)
 				{
+					TargetData->MipLevels = RT_TargetTexture->GetCurrentMipCount();
 					if (FTextureRenderTargetCubeResource* ResourceCube = RT_TargetTexture->GetTextureRenderTargetCubeResource())
 					{
 						TargetData->SamplerStateRHI = ResourceCube->SamplerStateRHI;
diff --git a/Engine/Plugins/FX/Niagara/Source/Niagara/Private/NiagaraDataInterfaceRenderTargetVolume.cpp b/Engine/Plugins/FX/Niagara/Source/Niagara/Private/NiagaraDataInterfaceRenderTargetVolume.cpp
--- a/Engine/Plugins/FX/Niagara/Source/Niagara/Private/NiagaraDataInterfaceRenderTargetVolume.cpp
+++ b/Engine/Plugins/FX/Niagara/Source/Niagara/Private/NiagaraDataInterfaceRenderTargetVolume.cpp
@@ -1041,7 +1041,7 @@ bool UNiagaraDataInterfaceRenderTargetVolume::PerInstanceTickPostSimulate(void*
 			{
 				FRenderTargetVolumeRWInstanceData_RenderThread* TargetData = &RT_Proxy->SystemInstancesToProxyData_RT.FindOrAdd(RT_InstanceID);
 				TargetData->Size = RT_InstanceData.Size;
-				TargetData->MipLevels = RT_TargetTexture->GetCurrentMipCount();
+				TargetData->MipLevels = 1;
 				//TargetData->MipMapGeneration = RT_InstanceData.MipMapGeneration;
 			#if WITH_EDITORONLY_DATA
 				TargetData->bPreviewTexture = RT_InstanceData.bPreviewTexture;
@@ -1050,6 +1050,7 @@ bool UNiagaraDataInterfaceRenderTargetVolume::PerInstanceTickPostSimulate(void*
 				TargetData->RenderTarget.SafeRelease();
 				if (RT_TargetTexture)
 				{
+					TargetData->MipLevels = RT_TargetTexture->GetCurrentMipCount();
 					if (FTextureRenderTargetVolumeResource* ResourceVolume = RT_TargetTexture->GetTextureRenderTargetVolumeResource())
 					{
 						TargetData->SamplerStateRHI = ResourceVolume->SamplerStateRHI;
