Added context menu capabilities to hierarchy editor with both a "base" section for rename & delete and a "dynamic" section for adding custom actions.

diff --git a/Engine/Plugins/FX/Niagara/Source/NiagaraEditor/Private/ViewModels/HierarchyEditor/NiagaraHierarchyViewModelBase.cpp b/Engine/Plugins/FX/Niagara/Source/NiagaraEditor/Private/ViewModels/HierarchyEditor/NiagaraHierarchyViewModelBase.cpp
--- a/Engine/Plugins/FX/Niagara/Source/NiagaraEditor/Private/ViewModels/HierarchyEditor/NiagaraHierarchyViewModelBase.cpp
+++ b/Engine/Plugins/FX/Niagara/Source/NiagaraEditor/Private/ViewModels/HierarchyEditor/NiagaraHierarchyViewModelBase.cpp
@@ -16,6 +16,8 @@
 #include "Widgets/Input/SCheckBox.h"
 #include "Editor.h"
 #include "ScopedTransaction.h"
+#include "Framework/Commands/GenericCommands.h"
+#include "ToolMenus.h"
 
 #define LOCTEXT_NAMESPACE "NiagaraHierarchyEditor"
 
@@ -306,6 +308,10 @@ void UNiagaraHierarchyViewModelBase::Initialize()
 		Section->SetFlags(RF_Transactional);
 	}
 
+	UToolMenus* ToolMenus = UToolMenus::Get();
+	UToolMenu* HierarchyMenu = ToolMenus->RegisterMenu("NiagaraHierarchyMenu", NAME_None, EMultiBoxType::Menu);
+	FToolMenuSection& BaseSection = HierarchyMenu->AddSection("Base");
+	
 	SetupCommands();
 
 	HierarchyViewModelRoot = MakeShared<FNiagaraHierarchyRootViewModel>(HierarchyRoot.Get(), this);
diff --git a/Engine/Plugins/FX/Niagara/Source/NiagaraEditor/Private/Widgets/SNiagaraHierarchy.cpp b/Engine/Plugins/FX/Niagara/Source/NiagaraEditor/Private/Widgets/SNiagaraHierarchy.cpp
--- a/Engine/Plugins/FX/Niagara/Source/NiagaraEditor/Private/Widgets/SNiagaraHierarchy.cpp
+++ b/Engine/Plugins/FX/Niagara/Source/NiagaraEditor/Private/Widgets/SNiagaraHierarchy.cpp
@@ -17,9 +17,58 @@
 #include "ScopedTransaction.h"
 #include "Modules/ModuleManager.h"
 #include "PropertyEditorModule.h"
+#include "Framework/Commands/GenericCommands.h"
+#include "Framework/Application/MenuStack.h"
+#include "Framework/Application/SlateApplication.h"
+#include "ToolMenus.h"
 
 #define LOCTEXT_NAMESPACE "NiagaraHierarchy"
 
+TSharedRef<SWidget> SummonContextMenu(TArray<TSharedPtr<FNiagaraHierarchyItemViewModelBase>> Items, bool bFromHierarchy)
+{
+	UNiagaraHierarchyMenuContext* MenuContextObject = NewObject<UNiagaraHierarchyMenuContext>();
+	MenuContextObject->Items = Items;
+	MenuContextObject->bFromHierarchy = bFromHierarchy;
+
+	FToolMenuContext MenuContext(MenuContextObject);
+	MenuContext.AppendCommandList(Items[0]->GetHierarchyViewModel()->GetCommands());
+		
+	UToolMenu* Menu = UToolMenus::Get()->GenerateMenu("NiagaraHierarchyMenu", MenuContext);
+	FToolMenuSection* BaseSection = Menu->FindSection("Base");
+
+	if(ensure(BaseSection != nullptr))
+	{
+		if(Items.Num() == 1 && Items[0]->CanRename())
+		{
+			BaseSection->AddMenuEntry(FGenericCommands::Get().Rename);
+		}
+
+		bool bCanDeleteAll = true;
+		for(TSharedPtr<FNiagaraHierarchyItemViewModelBase>& Item : Items)
+		{
+			if(!Item->CanDelete())
+			{
+				bCanDeleteAll = false;
+				break;
+			}
+		}
+
+		if(bCanDeleteAll)
+		{
+			BaseSection->AddMenuEntry(FGenericCommands::Get().Delete);
+		}
+	}
+
+	if(Items.Num() == 1)
+	{
+		FToolMenuSection& DynamicSection = Menu->FindOrAddSection("Dynamic");		
+		Items[0]->PopulateDynamicContextMenuSection(DynamicSection);
+	}	
+		
+	TSharedRef<SWidget> MenuWidget = UToolMenus::Get()->GenerateWidget(Menu);
+	return MenuWidget;
+}
+
 void SNiagaraHierarchyCategory::Construct(const FArguments& InArgs,	TSharedPtr<FNiagaraHierarchyCategoryViewModel> InCategory)
 {
 	Category = InCategory;
@@ -132,22 +181,18 @@ void SNiagaraHierarchySection::Construct(const FArguments& InArgs, TSharedPtr<FN
 				.DropTargetArgs(OntoDropTargetArgs
 				.Content()
 				[
-					SAssignNew(MenuAnchor, SMenuAnchor)
-					.OnGetMenuContent(this, &SNiagaraHierarchySection::OnGetMenuContent)
+					SNew(SCheckBox)
+					.Visibility(EVisibility::HitTestInvisible)
+					.Style(FAppStyle::Get(), "DetailsView.SectionButton")
+					.OnCheckStateChanged(this, &SNiagaraHierarchySection::OnSectionCheckChanged)
+					.IsChecked(this, &SNiagaraHierarchySection::GetSectionCheckState)
+					.Padding(FMargin(8.f, 4.f))
 					[
-						SNew(SCheckBox)
-						.Visibility(EVisibility::HitTestInvisible)
-						.Style(FAppStyle::Get(), "DetailsView.SectionButton")
-						.OnCheckStateChanged(this, &SNiagaraHierarchySection::OnSectionCheckChanged)
-						.IsChecked(this, &SNiagaraHierarchySection::GetSectionCheckState)
-						.Padding(FMargin(8.f, 4.f))
-						[
-							SAssignNew(InlineEditableTextBlock, SInlineEditableTextBlock)
-							.IsSelected(InArgs._IsSelected)
-							.Text(this, &SNiagaraHierarchySection::GetText)
-							.OnTextCommitted(this, &SNiagaraHierarchySection::OnRenameSection)
-							.OnVerifyTextChanged(this, &SNiagaraHierarchySection::OnVerifySectionRename)
-						]
+						SAssignNew(InlineEditableTextBlock, SInlineEditableTextBlock)
+						.IsSelected(InArgs._IsSelected)
+						.Text(this, &SNiagaraHierarchySection::GetText)
+						.OnTextCommitted(this, &SNiagaraHierarchySection::OnRenameSection)
+						.OnVerifyTextChanged(this, &SNiagaraHierarchySection::OnVerifySectionRename)
 					]
 				])
 			]
@@ -192,6 +237,11 @@ void SNiagaraHierarchySection::EnterEditingMode() const
 	InlineEditableTextBlock->EnterEditingMode();
 }
 
+TSharedPtr<FNiagaraHierarchySectionViewModel> SNiagaraHierarchySection::GetSectionViewModel()
+{
+	return SectionViewModel;
+}
+
 bool SNiagaraHierarchySection::OnCanAcceptDrop(TSharedPtr<FDragDropOperation> DragDropOperation, EItemDropZone ItemDropZone) const
 {
 	if(SectionViewModel.IsValid())
@@ -277,8 +327,16 @@ FReply SNiagaraHierarchySection::OnMouseButtonUp(const FGeometry& MyGeometry, co
 	if(SectionViewModel.IsValid())
 	{
 		if(MouseEvent.GetEffectingButton() == EKeys::RightMouseButton)
-		{
-			MenuAnchor->SetIsOpen(true);
+		{			
+			// Show menu to choose getter vs setter
+			FSlateApplication::Get().PushMenu(
+				AsShared(),
+				FWidgetPath(),
+				SummonContextMenu({GetSectionViewModel()}, true),
+				FSlateApplication::Get().GetCursorPos(),
+				FPopupTransitionEffect(FPopupTransitionEffect::ContextMenu)
+			);
+				
 			OnSectionActivatedDelegate.ExecuteIfBound(SectionViewModel);
 			return FReply::Handled();
 		}
@@ -307,16 +365,6 @@ void SNiagaraHierarchySection::OnDragLeave(const FDragDropEvent& DragDropEvent)
 	}
 }
 
-TSharedRef<SWidget> SNiagaraHierarchySection::OnGetMenuContent() const
-{
-	FMenuBuilder MenuBuilder(true, HierarchyViewModel->GetCommands());
-	
-	MenuBuilder.AddMenuEntry(FGenericCommands::Get().Rename);
-	MenuBuilder.AddMenuEntry(FGenericCommands::Get().Delete);
-
-	return MenuBuilder.MakeWidget();
-}
-
 UNiagaraHierarchySection* SNiagaraHierarchySection::TryGetSectionData() const
 {
 	return SectionViewModel.IsValid() ? Cast<UNiagaraHierarchySection>(SectionViewModel->GetDataMutable()) : nullptr;
@@ -463,6 +511,7 @@ void SNiagaraHierarchy::Construct(const FArguments& InArgs, TObjectPtr<UNiagaraH
 				.OnGenerateRow(this, &SNiagaraHierarchy::GenerateSourceItemRow)
 				.OnGetChildren_UObject(InHierarchyViewModel.Get(), &UNiagaraHierarchyViewModelBase::OnGetChildren)
 				.OnItemToString_Debug_UObject(InHierarchyViewModel.Get(), &UNiagaraHierarchyViewModelBase::OnItemToStringDebug)
+				.OnContextMenuOpening(this, &SNiagaraHierarchy::SummonContextMenuForSelectedRows, false)
 			]
 			+ SHorizontalBox::Slot()
 			.AutoWidth()
@@ -516,6 +565,7 @@ void SNiagaraHierarchy::Construct(const FArguments& InArgs, TObjectPtr<UNiagaraH
 						.OnGenerateRow(this, &SNiagaraHierarchy::GenerateHierarchyItemRow)
 						.OnGetChildren_UObject(InHierarchyViewModel.Get(), &UNiagaraHierarchyViewModelBase::OnGetChildren)
 						.OnItemToString_Debug_UObject(InHierarchyViewModel.Get(), &UNiagaraHierarchyViewModelBase::OnItemToStringDebug)
+						.OnContextMenuOpening(this, &SNiagaraHierarchy::SummonContextMenuForSelectedRows, true)
 					]
 					+ SVerticalBox::Slot()
 					.VAlign(VAlign_Fill)
@@ -599,6 +649,21 @@ FReply SNiagaraHierarchy::OnMouseButtonUp(const FGeometry& MyGeometry, const FPo
 	return FReply::Handled().SetUserFocus(AsShared(), EFocusCause::Mouse, true);
 }
 
+TSharedPtr<SWidget> SNiagaraHierarchy::SummonContextMenuForSelectedRows(bool bFromHierarchy) const
+{
+	TArray<TSharedPtr<FNiagaraHierarchyItemViewModelBase>> ViewModels;
+	if(bFromHierarchy)
+	{
+		HierarchyTreeView->GetSelectedItems(ViewModels);
+	}
+	else
+	{
+		SourceTreeView->GetSelectedItems(ViewModels);
+	}
+
+	return SummonContextMenu(ViewModels, bFromHierarchy);
+}
+
 void SNiagaraHierarchy::RefreshSourceView(bool bFullRefresh) const
 {
 	if(bFullRefresh)
@@ -736,11 +801,40 @@ bool SNiagaraHierarchy::CanRequestRenameSelectedItem() const
 	return false;
 }
 
-void SNiagaraHierarchy::DeleteSelectedItems()
+void SNiagaraHierarchy::DeleteItems(TArray<TSharedPtr<FNiagaraHierarchyItemViewModelBase>> ItemsToDelete)
 {
 	FScopedTransaction Transaction(LOCTEXT("Transaction_DeleteSelectedItems", "Deleted selected hierarchy items"));
 	HierarchyViewModel->GetHierarchyDataRoot()->Modify();
-	
+
+	bool bAnyItemsDeleted = false;
+	for(TSharedPtr<FNiagaraHierarchyItemViewModelBase>& SelectedItem : ItemsToDelete)
+	{
+		if(SelectedItem->CanDelete())
+		{
+			SelectedItem->Finalize();
+			if(SelectedItem->GetParent().IsValid())
+			{
+				SelectedItem->GetParent().Pin()->SyncToData();
+			}
+
+			bAnyItemsDeleted = true;
+		}
+	}
+
+	if(bAnyItemsDeleted)
+	{
+		RefreshSourceView();
+		RefreshHierarchyView();
+		HierarchyViewModel->OnHierarchyChanged().Broadcast();
+	}
+	else
+	{
+		Transaction.Cancel();
+	}
+}
+
+void SNiagaraHierarchy::DeleteSelectedItems()
+{	
 	TArray<TSharedPtr<FNiagaraHierarchyItemViewModelBase>> SelectedItems = HierarchyTreeView->GetSelectedItems();
 
 	if(SelectedItems.Num() == 0)
@@ -751,19 +845,8 @@ void SNiagaraHierarchy::DeleteSelectedItems()
 			SelectedItems = { ActiveSection };
 		}
 	}
-	
-	for(TSharedPtr<FNiagaraHierarchyItemViewModelBase>& SelectedItem : SelectedItems)
-	{
-		SelectedItem->Finalize();
-		if(SelectedItem->GetParent().IsValid())
-		{
-			SelectedItem->GetParent().Pin()->SyncToData();
-		}
-	}
 
-	RefreshSourceView();
-	RefreshHierarchyView();
-	HierarchyViewModel->OnHierarchyChanged().Broadcast();
+	DeleteItems(SelectedItems);
 }
 
 bool SNiagaraHierarchy::CanDeleteSelectedItems() const
diff --git a/Engine/Plugins/FX/Niagara/Source/NiagaraEditor/Public/ViewModels/HierarchyEditor/NiagaraHierarchyViewModelBase.h b/Engine/Plugins/FX/Niagara/Source/NiagaraEditor/Public/ViewModels/HierarchyEditor/NiagaraHierarchyViewModelBase.h
--- a/Engine/Plugins/FX/Niagara/Source/NiagaraEditor/Public/ViewModels/HierarchyEditor/NiagaraHierarchyViewModelBase.h
+++ b/Engine/Plugins/FX/Niagara/Source/NiagaraEditor/Public/ViewModels/HierarchyEditor/NiagaraHierarchyViewModelBase.h
@@ -6,6 +6,7 @@
 #include "NiagaraEditorUtilities.h"
 #include "PropertyEditorDelegates.h"
 #include "ViewModels/Stack/NiagaraStackFunctionInput.h"
+#include "ToolMenuSection.h"
 #include "NiagaraHierarchyViewModelBase.generated.h"
 
 UCLASS()
@@ -238,6 +239,17 @@ protected:
 	bool bFromSourceList = false;
 };
 
+UCLASS(BlueprintType)
+class UNiagaraHierarchyMenuContext : public UObject
+{
+	GENERATED_BODY()
+
+public:
+	TArray<TSharedPtr<FNiagaraHierarchyItemViewModelBase>> Items;
+
+	bool bFromHierarchy = false;
+};
+
 UCLASS(Abstract)
 class UNiagaraHierarchyViewModelBase : public UObject, public FEditorUndoClient
 {
@@ -433,6 +445,9 @@ struct NIAGARAEDITOR_API FNiagaraHierarchyItemViewModelBase : TSharedFromThis<FN
 		}
 	}
 
+	/** Override this to register dynamic context menu entries when right clicking a hierarchy item */
+	virtual void PopulateDynamicContextMenuSection(FToolMenuSection& DynamicSection) {}
+
 	/** The UObject we display in the details panel when this item is selected. By default it's the item the view model represents. */
 	virtual UObject* GetDataForEditing() { return ItemBase; }
 	
diff --git a/Engine/Plugins/FX/Niagara/Source/NiagaraEditor/Public/Widgets/SNiagaraHierarchy.h b/Engine/Plugins/FX/Niagara/Source/NiagaraEditor/Public/Widgets/SNiagaraHierarchy.h
--- a/Engine/Plugins/FX/Niagara/Source/NiagaraEditor/Public/Widgets/SNiagaraHierarchy.h
+++ b/Engine/Plugins/FX/Niagara/Source/NiagaraEditor/Public/Widgets/SNiagaraHierarchy.h
@@ -46,6 +46,8 @@ private:
 	void RequestRenameSelectedItem();
 	bool CanRequestRenameSelectedItem() const;
 
+	void DeleteItems(TArray<TSharedPtr<FNiagaraHierarchyItemViewModelBase>> ItemsToDelete);
+	
 	void DeleteSelectedItems();
 	bool CanDeleteSelectedItems() const;
 
@@ -57,6 +59,8 @@ private:
 	virtual FReply OnMouseButtonDown(const FGeometry& MyGeometry, const FPointerEvent& MouseEvent) override;
 	virtual FReply OnMouseButtonUp(const FGeometry& MyGeometry, const FPointerEvent& MouseEvent) override;
 
+	TSharedPtr<SWidget> SummonContextMenuForSelectedRows(bool bFromHierarchy) const;
+	
 	/** This handles the drag & drop at the root level of the hierarchy, i.e. when dragging and dropping a parameter into the "empty" hierarchy */
 	FReply HandleHierarchyRootDrop(const FGeometry&, const FDragDropEvent& DragDropEvent) const;
 	bool CanDropOnRoot(TSharedPtr<FDragDropOperation> DragDropOp) const;
@@ -93,14 +97,13 @@ public:
 
 	void EnterEditingMode() const;
 
+	TSharedPtr<struct FNiagaraHierarchySectionViewModel> GetSectionViewModel();
 private:
 	virtual FReply OnMouseButtonDown(const FGeometry& MyGeometry, const FPointerEvent& MouseEvent) override;
 	virtual FReply OnMouseButtonUp(const FGeometry& MyGeometry, const FPointerEvent& MouseEvent) override;
 	
 	virtual FReply OnDragDetected(const FGeometry& MyGeometry, const FPointerEvent& MouseEvent) override;
 	virtual void OnDragLeave(const FDragDropEvent& DragDropEvent) override;
-	
-	TSharedRef<SWidget> OnGetMenuContent() const;
 
 	UNiagaraHierarchySection* TryGetSectionData() const;
 
