Fixed the niagara hierarchy menu being registered multiple times

diff --git a/Engine/Plugins/FX/Niagara/Source/NiagaraEditor/Private/ViewModels/HierarchyEditor/NiagaraHierarchyViewModelBase.cpp b/Engine/Plugins/FX/Niagara/Source/NiagaraEditor/Private/ViewModels/HierarchyEditor/NiagaraHierarchyViewModelBase.cpp
--- a/Engine/Plugins/FX/Niagara/Source/NiagaraEditor/Private/ViewModels/HierarchyEditor/NiagaraHierarchyViewModelBase.cpp
+++ b/Engine/Plugins/FX/Niagara/Source/NiagaraEditor/Private/ViewModels/HierarchyEditor/NiagaraHierarchyViewModelBase.cpp
@@ -309,8 +309,13 @@ void UNiagaraHierarchyViewModelBase::Initialize()
 	}
 
 	UToolMenus* ToolMenus = UToolMenus::Get();
-	UToolMenu* HierarchyMenu = ToolMenus->RegisterMenu("NiagaraHierarchyMenu", NAME_None, EMultiBoxType::Menu);
-	FToolMenuSection& BaseSection = HierarchyMenu->AddSection("Base");
+
+	FName MenuName("NiagaraHierarchyMenu");
+	if(!ToolMenus->IsMenuRegistered(MenuName))
+	{
+		UToolMenu* HierarchyMenu = ToolMenus->RegisterMenu(MenuName, NAME_None, EMultiBoxType::Menu);
+		HierarchyMenu->AddSection("Base");
+	}
 	
 	SetupCommands();
 
