- Niagara Mesh Renderer refactored to use an interface to get the mesh data

diff --git a/Engine/Plugins/FX/Niagara/Source/Niagara/Private/NiagaraMeshRendererProperties.cpp b/Engine/Plugins/FX/Niagara/Source/Niagara/Private/NiagaraMeshRendererProperties.cpp
--- a/Engine/Plugins/FX/Niagara/Source/Niagara/Private/NiagaraMeshRendererProperties.cpp
+++ b/Engine/Plugins/FX/Niagara/Source/Niagara/Private/NiagaraMeshRendererProperties.cpp
@@ -2,7 +2,6 @@
 
 #include "NiagaraMeshRendererProperties.h"
 #include "NiagaraRendererMeshes.h"
-#include "Engine/StaticMesh.h"
 #include "NiagaraConstants.h"
 #include "NiagaraBoundsCalculatorHelper.h"
 #include "NiagaraCustomVersion.h"
@@ -11,6 +10,7 @@
 #include "NiagaraSystem.h"
 
 #include "MaterialDomain.h"
+#include "Engine/StaticMesh.h"
 #include "Materials/MaterialInstanceConstant.h"
 #include "Modules/ModuleManager.h"
 
@@ -30,12 +30,147 @@
 #include "Subsystems/ImportSubsystem.h"
 #endif
 
-
 #define LOCTEXT_NAMESPACE "UNiagaraMeshRendererProperties"
 
-
 TArray<TWeakObjectPtr<UNiagaraMeshRendererProperties>> UNiagaraMeshRendererProperties::MeshRendererPropertiesToDeferredInit;
 
+////////////////////////////////////////////////////////////////////////////////////////////////////////////////////////////////
+
+class FNiagaraRenderableStaticMesh : public INiagaraRenderableMesh
+{
+public:
+	FNiagaraRenderableStaticMesh(const UStaticMesh* StaticMesh)
+	{
+		WeakStaticMesh	= StaticMesh;
+		RenderData		= StaticMesh->GetRenderData();
+		MinLOD			= StaticMesh->GetMinLODIdx();
+		LocalBounds		= StaticMesh->GetExtendedBounds().GetBox();
+	}
+
+	FBox GetLocalBounds() const override
+	{
+		return LocalBounds;
+	}
+
+	void GetLODModelData(FLODModelData& OutLODModelData) const override
+	{
+		// This should never happen, but for safety
+		if (!ensure(WeakStaticMesh.IsValid()))
+		{
+			OutLODModelData.LODIndex = INDEX_NONE;
+			return;
+		}
+
+		OutLODModelData.LODIndex = RenderData->GetCurrentFirstLODIdx(MinLOD);
+		if (!RenderData->LODResources.IsValidIndex(OutLODModelData.LODIndex))
+		{
+			OutLODModelData.LODIndex = INDEX_NONE;
+			return;
+		}
+
+		const FStaticMeshLODResources& LODResources = RenderData->LODResources[OutLODModelData.LODIndex];
+
+		OutLODModelData.NumVertices = LODResources.GetNumVertices();
+		OutLODModelData.NumIndices = LODResources.IndexBuffer.GetNumIndices();
+		OutLODModelData.Sections = MakeArrayView(LODResources.Sections);
+		OutLODModelData.IndexBuffer = &LODResources.IndexBuffer;
+		OutLODModelData.VertexFactoryUserData = RenderData->LODVertexFactories.IsValidIndex(OutLODModelData.LODIndex) ? RenderData->LODVertexFactories[OutLODModelData.LODIndex].VertexFactory.GetUniformBuffer() : nullptr;
+		OutLODModelData.RayTracingGeometry = &LODResources.RayTracingGeometry;
+
+		if (LODResources.AdditionalIndexBuffers != nullptr && LODResources.AdditionalIndexBuffers->WireframeIndexBuffer.IsInitialized())
+		{
+			OutLODModelData.WireframeNumIndices = LODResources.AdditionalIndexBuffers->WireframeIndexBuffer.GetNumIndices();
+			OutLODModelData.WireframeIndexBuffer = &LODResources.AdditionalIndexBuffers->WireframeIndexBuffer;
+		}
+	}
+
+	void SetupVertexFactory(FNiagaraMeshVertexFactory& InVertexFactory, const FLODModelData& LODModelData) const override
+	{
+		FStaticMeshDataType Data;
+		const FStaticMeshLODResources& LODResources = RenderData->LODResources[LODModelData.LODIndex];
+		LODResources.VertexBuffers.PositionVertexBuffer.BindPositionVertexBuffer(&InVertexFactory, Data);
+		LODResources.VertexBuffers.StaticMeshVertexBuffer.BindTangentVertexBuffer(&InVertexFactory, Data);
+		LODResources.VertexBuffers.StaticMeshVertexBuffer.BindTexCoordVertexBuffer(&InVertexFactory, Data, MAX_TEXCOORDS);
+		LODResources.VertexBuffers.ColorVertexBuffer.BindColorVertexBuffer(&InVertexFactory, Data);
+		InVertexFactory.SetData(Data);
+	}
+
+	void GetUsedMaterials(TArray<UMaterialInterface*>& OutMaterials) const override
+	{
+		const UStaticMesh* StaticMesh = WeakStaticMesh.Get();
+		if (!ensure(StaticMesh))
+		{
+			return;
+		}
+
+		// Retrieve a list of materials whose indices match up with the mesh, and only fill it in with materials that are used by any section of any LOD
+		for (const FStaticMeshLODResources& LODModel : RenderData->LODResources)
+		{
+			for (const FStaticMeshSection& Section : LODModel.Sections)
+			{
+				if (Section.MaterialIndex >= 0)
+				{
+					if (Section.MaterialIndex >= OutMaterials.Num())
+					{
+						OutMaterials.AddZeroed(Section.MaterialIndex - OutMaterials.Num() + 1);
+					}
+					else if (OutMaterials[Section.MaterialIndex])
+					{
+						continue;
+					}
+
+					UMaterialInterface* Material = StaticMesh->GetMaterial(Section.MaterialIndex);
+					if (!Material)
+					{
+						Material = UMaterial::GetDefaultMaterial(MD_Surface);
+					}
+					OutMaterials[Section.MaterialIndex] = Material;
+				}
+			}
+		}
+	}
+
+	TWeakObjectPtr<const UStaticMesh>	WeakStaticMesh;
+	const class FStaticMeshRenderData*	RenderData = nullptr;
+	int32								MinLOD = 0;
+	FBox								LocalBounds = FBox(ForceInitToZero);
+};
+
+////////////////////////////////////////////////////////////////////////////////////////////////////////////////////////////////
+
+namespace NiagaraMeshRendererPropertiesInternal
+{
+	void ResolveRenderableMeshInternal(const FNiagaraMeshRendererMeshProperties& MeshProperties, const FNiagaraEmitterInstance* EmitterInstance, INiagaraRenderableMeshInterface*& OutInterface, UStaticMesh*& OutStaticMesh)
+	{
+		OutInterface = nullptr;
+		OutStaticMesh = nullptr;
+		if (EmitterInstance)
+		{
+			if (MeshProperties.MeshParameterBinding.Parameter.IsValid())
+			{
+				if (MeshProperties.MeshParameterBinding.Parameter.IsDataInterface())
+				{
+					OutInterface = Cast<INiagaraRenderableMeshInterface>(EmitterInstance->GetRendererBoundVariables().GetDataInterface(MeshProperties.MeshParameterBinding.Parameter));
+					return;
+				}
+
+				UStaticMesh* BoundMesh = Cast<UStaticMesh>(EmitterInstance->GetRendererBoundVariables().GetUObject(MeshProperties.MeshParameterBinding.Parameter));
+				if (BoundMesh && BoundMesh->GetRenderData())
+				{
+					OutStaticMesh = BoundMesh;
+					return;
+				}
+			}
+		}
+		if (MeshProperties.Mesh && MeshProperties.Mesh->GetRenderData())
+		{
+			OutStaticMesh = MeshProperties.Mesh;
+		}
+	}
+}
+
+////////////////////////////////////////////////////////////////////////////////////////////////////////////////////////////////
+
 FNiagaraMeshMaterialOverride::FNiagaraMeshMaterialOverride()
 	: ExplicitMat(nullptr)
 	, UserParamBinding(FNiagaraTypeDefinition(UMaterialInterface::StaticClass()))
@@ -54,41 +189,49 @@ bool FNiagaraMeshMaterialOverride::SerializeFromMismatchedTag(const struct FProp
 	return false;
 }
 
+////////////////////////////////////////////////////////////////////////////////////////////////////////////////////////////////
 
 FNiagaraMeshRendererMeshProperties::FNiagaraMeshRendererMeshProperties()
 	: Mesh(nullptr)
-	, UserParamBinding(FNiagaraTypeDefinition(UStaticMesh::StaticClass()))
+#if WITH_EDITORONLY_DATA
+	, UserParamBinding_DEPRECATED(FNiagaraTypeDefinition(UStaticMesh::StaticClass()))
+#endif
 	, Scale(1.0f, 1.0f, 1.0f)
 	, Rotation(FRotator::ZeroRotator)
 	, PivotOffset(ForceInitToZero)
 	, PivotOffsetSpace(ENiagaraMeshPivotOffsetSpace::Mesh)
 {
+#if WITH_EDITORONLY_DATA
+	MeshParameterBinding.SetUsage(ENiagaraParameterBindingUsage::NotParticle);
+	MeshParameterBinding.SetAllowedInterfaces({UNiagaraRenderableMeshInterface::StaticClass()});
+	MeshParameterBinding.SetAllowedObjects({UStaticMesh::StaticClass()});
+#endif
 }
 
-UStaticMesh* FNiagaraMeshRendererMeshProperties::ResolveStaticMesh(const FNiagaraEmitterInstance* Emitter) const
+FNiagaraRenderableMeshPtr FNiagaraMeshRendererMeshProperties::ResolveRenderableMesh(const FNiagaraEmitterInstance* EmitterInstance) const
 {
-	UStaticMesh* FoundMesh = nullptr;
+	INiagaraRenderableMeshInterface* RenderableMeshInterface = nullptr;
+	UStaticMesh* StaticMesh = nullptr;
+	NiagaraMeshRendererPropertiesInternal::ResolveRenderableMeshInternal(*this, EmitterInstance, RenderableMeshInterface, StaticMesh);
 
-	if (UserParamBinding.Parameter.IsValid() && Emitter)
+	if (RenderableMeshInterface)
 	{
-		UStaticMesh* TestMesh = Cast<UStaticMesh>(Emitter->GetRendererBoundVariables().GetUObject(UserParamBinding.Parameter));
-		if (TestMesh && TestMesh->GetRenderData())
+		if ( const FNiagaraSystemInstance* SystemInstance = EmitterInstance->GetParentSystemInstance() )
 		{
-			FoundMesh = TestMesh;
+			return RenderableMeshInterface->GetRenderableMesh(SystemInstance->GetId());
 		}
 	}
 
-	if (!FoundMesh && Mesh && Mesh->GetRenderData())
+	if (StaticMesh)
 	{
-		FoundMesh = Mesh;
+		return MakeShared<FNiagaraRenderableStaticMesh>(StaticMesh);
 	}
-
-	return FoundMesh;
+	return nullptr;
 }
 
-bool FNiagaraMeshRendererMeshProperties::HasValidMeshProperties() const
+bool FNiagaraMeshRendererMeshProperties::HasValidRenderableMesh() const
 {
-	return Mesh || UserParamBinding.Parameter.IsValid();
+	return Mesh || MeshParameterBinding.Parameter.IsValid();
 }
 
 UNiagaraMeshRendererProperties::UNiagaraMeshRendererProperties()
@@ -144,7 +287,7 @@ FNiagaraRenderer* UNiagaraMeshRendererProperties::CreateEmitterRenderer(ERHIFeat
 {
 	for (const auto& MeshProperties : Meshes)
 	{
-		if (MeshProperties.ResolveStaticMesh(Emitter))
+		if (MeshProperties.HasValidRenderableMesh())
 		{
 			// There's at least one valid mesh
 			FNiagaraRenderer* NewRenderer = new FNiagaraRendererMeshes(FeatureLevel, this, Emitter);
@@ -442,73 +585,39 @@ bool UNiagaraMeshRendererProperties::IsSupportedVariableForBinding(const FNiagar
 }
 #endif
 
-void UNiagaraMeshRendererProperties::GetUsedMeshMaterials(int32 MeshIndex, const FNiagaraEmitterInstance* Emitter, TArray<UMaterialInterface*>& OutMaterials) const
+void UNiagaraMeshRendererProperties::ApplyMaterialOverrides(const FNiagaraEmitterInstance* EmitterInstance, TArray<UMaterialInterface*>& InOutMaterials) const
 {
-	check(Meshes.IsValidIndex(MeshIndex));
-
-	const UStaticMesh* Mesh = Meshes[MeshIndex].ResolveStaticMesh(Emitter);
-	check(Mesh);
-
-	const FStaticMeshRenderData* RenderData = Mesh->GetRenderData();
-	check(RenderData);
-
-	OutMaterials.SetNum(0, false);
-
-	// Retrieve a list of materials whose indices match up with the mesh, and only fill it in with materials that are used by any section of any LOD
-	for (const FStaticMeshLODResources& LODModel : RenderData->LODResources)
-	{
-		for (const FStaticMeshSection& Section : LODModel.Sections)
-		{
-			if (Section.MaterialIndex >= 0)
-			{
-				if (Section.MaterialIndex >= OutMaterials.Num())
-				{
-					OutMaterials.AddZeroed(Section.MaterialIndex - OutMaterials.Num() + 1);
-				}
-				else if (OutMaterials[Section.MaterialIndex])
-				{
-					continue;
-				}
-
-				UMaterialInterface* Material = Mesh->GetMaterial(Section.MaterialIndex);
-				if (!Material)
-				{
-					Material = UMaterial::GetDefaultMaterial(MD_Surface);
-				}
-				OutMaterials[Section.MaterialIndex] = Material;
-			}
-		}
-	}
-
 	if (bOverrideMaterials)
 	{
-		const int32 NumOverrideMaterials = FMath::Min(OverrideMaterials.Num(), OutMaterials.Num());
+		const int32 NumOverrideMaterials = FMath::Min(OverrideMaterials.Num(), InOutMaterials.Num());
 		for (int32 OverrideIndex = 0; OverrideIndex < NumOverrideMaterials; ++OverrideIndex)
 		{
-			if (OutMaterials[OverrideIndex])
+			if (!InOutMaterials[OverrideIndex])
 			{
-				UMaterialInterface* OverrideMat = nullptr;
+				continue;
+			}
 
-				// UserParamBinding, if mapped to a real value, always wins. Otherwise, use the ExplictMat if it is set. Finally, fall
-				// back to the particle mesh material. This allows the user to effectively optionally bind to a Material binding
-				// and still have good defaults if it isn't set to anything.
-				if (Emitter && OverrideMaterials[OverrideIndex].UserParamBinding.Parameter.IsValid())
-				{
-					OverrideMat = Cast<UMaterialInterface>(Emitter->FindBinding(OverrideMaterials[OverrideIndex].UserParamBinding.Parameter));
-				}
+			UMaterialInterface* OverrideMat = nullptr;
 
-				if (!OverrideMat)
-				{
-					OverrideMat = OverrideMaterials[OverrideIndex].ExplicitMat;
-				#if WITH_EDITORONLY_DATA
-					OverrideMat = MICMaterials.IsValidIndex(OverrideIndex) && MICMaterials[OverrideIndex] ? MICMaterials[OverrideIndex] : OverrideMat;
-				#endif
-				}
+			// UserParamBinding, if mapped to a real value, always wins. Otherwise, use the ExplictMat if it is set. Finally, fall
+			// back to the particle mesh material. This allows the user to effectively optionally bind to a Material binding
+			// and still have good defaults if it isn't set to anything.
+			if (EmitterInstance && OverrideMaterials[OverrideIndex].UserParamBinding.Parameter.IsValid())
+			{
+				OverrideMat = Cast<UMaterialInterface>(EmitterInstance->FindBinding(OverrideMaterials[OverrideIndex].UserParamBinding.Parameter));
+			}
 
-				if (OverrideMat)
-				{
-					OutMaterials[OverrideIndex] = OverrideMat;
-				}
+			if (!OverrideMat)
+			{
+				OverrideMat = OverrideMaterials[OverrideIndex].ExplicitMat;
+			#if WITH_EDITORONLY_DATA
+				OverrideMat = MICMaterials.IsValidIndex(OverrideIndex) && MICMaterials[OverrideIndex] ? MICMaterials[OverrideIndex] : OverrideMat;
+			#endif
+			}
+
+			if (OverrideMat)
+			{
+				InOutMaterials[OverrideIndex] = OverrideMat;
 			}
 		}
 	}
@@ -517,10 +626,10 @@ void UNiagaraMeshRendererProperties::GetUsedMeshMaterials(int32 MeshIndex, const
 	#if WITH_EDITORONLY_DATA
 		if (MICMaterials.Num() > 0)
 		{
-			const int MaterialMax = FMath::Min(OutMaterials.Num(), MICMaterials.Num());
+			const int MaterialMax = FMath::Min(InOutMaterials.Num(), MICMaterials.Num());
 			for (int i = 0; i < MaterialMax; ++i)
 			{
-				OutMaterials[i] = MICMaterials[i] ? MICMaterials[i] : OutMaterials[i];
+				InOutMaterials[i] = MICMaterials[i] ? MICMaterials[i] : InOutMaterials[i];
 			}
 		}
 	#endif
@@ -532,20 +641,38 @@ const FVertexFactoryType* UNiagaraMeshRendererProperties::GetVertexFactoryType()
 	return &FNiagaraMeshVertexFactory::StaticType;
 }
 
-void UNiagaraMeshRendererProperties::GetUsedMaterials(const FNiagaraEmitterInstance* InEmitter, TArray<UMaterialInterface*>& OutMaterials) const
+void UNiagaraMeshRendererProperties::GetUsedMaterials(const FNiagaraEmitterInstance* EmitterInstance, TArray<UMaterialInterface*>& OutMaterials) const
 {
+	FNiagaraSystemInstance* SystemInstance = EmitterInstance ? EmitterInstance->GetParentSystemInstance() : nullptr;
+
 	TArray<UMaterialInterface*> OrderedMeshMaterials;
-	for (int32 MeshIndex = 0; MeshIndex < Meshes.Num(); ++MeshIndex)
+	for (int32 MeshIndex=0; MeshIndex < Meshes.Num(); ++MeshIndex)
 	{
-		const UStaticMesh* Mesh = Meshes[MeshIndex].ResolveStaticMesh(InEmitter);
-		if (Mesh && Mesh->GetRenderData())
+		OrderedMeshMaterials.Reset(0);
+
+		INiagaraRenderableMeshInterface* RenderableMeshInterface = nullptr;
+		UStaticMesh* StaticMesh = nullptr;
+		NiagaraMeshRendererPropertiesInternal::ResolveRenderableMeshInternal(Meshes[MeshIndex], EmitterInstance, RenderableMeshInterface, StaticMesh);
+
+		if (RenderableMeshInterface && SystemInstance)
 		{
-			GetUsedMeshMaterials(MeshIndex, InEmitter, OrderedMeshMaterials);
-			for (UMaterialInterface* Material : OrderedMeshMaterials)
+			RenderableMeshInterface->GetUsedMaterials(SystemInstance->GetId(), OrderedMeshMaterials);
+		}
+		else if (StaticMesh)
+		{
+			FNiagaraRenderableStaticMesh(StaticMesh).GetUsedMaterials(OrderedMeshMaterials);
+		}
+
+		if (OrderedMeshMaterials.Num() > 0)
+		{
+			ApplyMaterialOverrides(EmitterInstance, OrderedMeshMaterials);
+
+			OutMaterials.Reserve(OutMaterials.Num() + OrderedMeshMaterials.Num());
+			for (UMaterialInterface* MaterialInterface : OrderedMeshMaterials)
 			{
-				if (Material)
+				if (MaterialInterface)
 				{
-					OutMaterials.AddUnique(Material);
+					OutMaterials.AddUnique(MaterialInterface);
 				}
 			}
 		}
@@ -556,23 +683,24 @@ void UNiagaraMeshRendererProperties::GetStreamingMeshInfo(const FBoxSphereBounds
 {
 	for (const FNiagaraMeshRendererMeshProperties& MeshProperties : Meshes)
 	{
-		UStaticMesh* Mesh = MeshProperties.ResolveStaticMesh(InEmitter);
+		INiagaraRenderableMeshInterface* RenderableMeshInterface = nullptr;
+		UStaticMesh* StaticMesh = nullptr;
+		NiagaraMeshRendererPropertiesInternal::ResolveRenderableMeshInternal(MeshProperties, InEmitter, RenderableMeshInterface, StaticMesh);
 
-		if (Mesh && Mesh->RenderResourceSupportsStreaming() && Mesh->GetRenderAssetType() == EStreamableRenderAssetType::StaticMesh)
+		if (StaticMesh && StaticMesh->RenderResourceSupportsStreaming() && StaticMesh->GetRenderAssetType() == EStreamableRenderAssetType::StaticMesh)
 		{
-			const FBoxSphereBounds MeshBounds = Mesh->GetBounds();
+			const FBoxSphereBounds MeshBounds = StaticMesh->GetBounds();
 			const FBoxSphereBounds StreamingBounds = FBoxSphereBounds(
 				OwnerBounds.Origin + MeshBounds.Origin,
 				MeshBounds.BoxExtent * MeshProperties.Scale,
 				MeshBounds.SphereRadius * MeshProperties.Scale.GetMax());
 			const float MeshTexelFactor = MeshBounds.SphereRadius * 2.0f;
 
-			new (OutStreamingRenderAssets) FStreamingRenderAssetPrimitiveInfo(Mesh, StreamingBounds, MeshTexelFactor);
+			new (OutStreamingRenderAssets) FStreamingRenderAssetPrimitiveInfo(StaticMesh, StreamingBounds, MeshTexelFactor);
 		}
 	}
 }
 
-
 #if WITH_EDITORONLY_DATA
 TArray<FNiagaraVariable> UNiagaraMeshRendererProperties::GetBoundAttributes() const 
 {
@@ -608,9 +736,9 @@ bool UNiagaraMeshRendererProperties::PopulateRequiredBindings(FNiagaraParameterS
 
 	for (FNiagaraMeshRendererMeshProperties& Binding : Meshes)
 	{
-		if (Binding.UserParamBinding.Parameter.IsValid())
+		if (Binding.MeshParameterBinding.Parameter.IsValid())
 		{
-			InParameterStore.AddParameter(Binding.UserParamBinding.Parameter, false);
+			InParameterStore.AddParameter(Binding.MeshParameterBinding.Parameter, false);
 			bAnyAdded = true;
 		}
 	}
@@ -631,7 +759,7 @@ void UNiagaraMeshRendererProperties::PostLoad()
 		Mesh.PivotOffsetSpace = PivotOffsetSpace_DEPRECATED;
 	}
 
-	for (const auto& MeshProperties : Meshes)
+	for (FNiagaraMeshRendererMeshProperties& MeshProperties : Meshes)
 	{
 		if (MeshProperties.Mesh)
 		{
@@ -644,6 +772,13 @@ void UNiagaraMeshRendererProperties::PostLoad()
 			}
 #endif
 		}
+#if WITH_EDITORONLY_DATA
+		if (MeshProperties.UserParamBinding_DEPRECATED.Parameter.GetName().IsNone() == false)
+		{
+			MeshProperties.MeshParameterBinding.Parameter = MeshProperties.UserParamBinding_DEPRECATED.Parameter;
+			MeshProperties.MeshParameterBinding.AliasedParameter = MeshProperties.UserParamBinding_DEPRECATED.Parameter;
+		}
+#endif
 	}
 
 #if WITH_EDITORONLY_DATA
@@ -750,18 +885,23 @@ void UNiagaraMeshRendererProperties::GetRendererTooltipWidgets(const FNiagaraEmi
 	{
 		const FNiagaraMeshRendererMeshProperties& MeshProperties = Meshes[MeshIndex];
 		
-		TSharedPtr<SWidget> TooltipWidget = DefaultMeshTooltip;		
+		TSharedPtr<SWidget> TooltipWidget = DefaultMeshTooltip;
+
 		// we make sure to reuse the mesh widget as a thumbnail if the mesh is valid
-		if(MeshProperties.ResolveStaticMesh(InEmitter))
+		INiagaraRenderableMeshInterface* RenderableMeshInterface = nullptr;
+		UStaticMesh* StaticMesh = nullptr;
+		NiagaraMeshRendererPropertiesInternal::ResolveRenderableMeshInternal(Meshes[MeshIndex], InEmitter, RenderableMeshInterface, StaticMesh);
+
+		if (StaticMesh)
 		{
 			TooltipWidget = RendererWidgets[MeshIndex];
 		}
 
-		// we override the previous thumbnail tooltip with a text indicating user parameter binding, if it exists
-		if(MeshProperties.UserParamBinding.Parameter.IsValid())
+		// we override the previous thumbnail tooltip with a text indicating parameter binding, if it exists
+		if(MeshProperties.MeshParameterBinding.Parameter.IsValid())
 		{
 			TooltipWidget = SNew(STextBlock)
-				.Text(FText::Format(LOCTEXT("MeshBoundTooltip", "Mesh slot is bound to user parameter {0}"), FText::FromName(MeshProperties.UserParamBinding.Parameter.GetName())));
+				.Text(FText::Format(LOCTEXT("MeshBoundTooltip", "Mesh slot is bound to parameter {0}"), FText::FromName(MeshProperties.MeshParameterBinding.Parameter.GetName())));
 		}
 		
 		OutWidgets.Add(TooltipWidget);
diff --git a/Engine/Plugins/FX/Niagara/Source/Niagara/Private/NiagaraParameterBinding.cpp b/Engine/Plugins/FX/Niagara/Source/Niagara/Private/NiagaraParameterBinding.cpp
new file mode 100644
--- /dev/null
+++ b/Engine/Plugins/FX/Niagara/Source/Niagara/Private/NiagaraParameterBinding.cpp
@@ -0,0 +1,82 @@
+// Copyright Epic Games, Inc. All Rights Reserved.
+
+#include "NiagaraParameterBinding.h"
+#include "NiagaraConstants.h"
+#include "NiagaraDataInterface.h"
+
+#if WITH_EDITORONLY_DATA
+bool FNiagaraParameterBinding::CanBindTo(FNiagaraTypeDefinition TypeDefinition) const
+{
+	if (AllowedTypeDefinitions.Contains(TypeDefinition))
+	{
+		return true;
+	}
+
+	if (UClass* TypeClass = TypeDefinition.GetClass())
+	{
+		if (TypeClass->IsChildOf(UNiagaraDataInterface::StaticClass()))
+		{
+			for (UClass* AllowedClass : AllowedDataInterfaces)
+			{
+				if (AllowedClass && TypeClass->IsChildOf(AllowedClass))
+				{
+					return true;
+				}
+			}
+
+			for (UClass* AllowedInterface : AllowedInterfaces)
+			{
+				if (AllowedInterface && TypeClass->ImplementsInterface(AllowedInterface))
+				{
+					return true;
+				}
+			}
+		}
+		else
+		{
+			for (UClass* AllowedClass : AllowedObjects)
+			{
+				if (AllowedClass && TypeClass->IsChildOf(AllowedClass))
+				{
+					return true;
+				}
+			}
+		}
+	}
+
+	return false;
+}
+
+bool FNiagaraParameterBinding::CanBindTo(FNiagaraVariableBase InVariable, FNiagaraVariableBase& OutAliasedVariable, FStringView EmitterName) const
+{
+	OutAliasedVariable = InVariable;
+	if ( !AllowStaticVariables() && InVariable.GetType().IsStatic() )
+	{
+		return false;
+	}
+
+	if ( !CanBindTo(InVariable.GetType()) )
+	{
+		return false;
+	}
+
+	if (InVariable.IsInNameSpace(FNiagaraConstants::ParticleAttributeNamespaceString))
+	{
+		return AllowParticleParameters();
+	}
+	else if (InVariable.IsInNameSpace(EmitterName))
+	{
+		if ( AllowEmitterParameters() )
+		{
+			OutAliasedVariable.ReplaceRootNamespace(EmitterName, FNiagaraConstants::EmitterNamespaceString);
+			return true;
+		}
+	}
+	else if (InVariable.IsInNameSpace(FNiagaraConstants::SystemNamespaceString))
+	{
+		return AllowSystemParameters();
+	}
+
+	return false;
+}
+#endif
diff --git a/Engine/Plugins/FX/Niagara/Source/Niagara/Private/NiagaraRendererMeshes.cpp b/Engine/Plugins/FX/Niagara/Source/Niagara/Private/NiagaraRendererMeshes.cpp
--- a/Engine/Plugins/FX/Niagara/Source/Niagara/Private/NiagaraRendererMeshes.cpp
+++ b/Engine/Plugins/FX/Niagara/Source/Niagara/Private/NiagaraRendererMeshes.cpp
@@ -12,9 +12,11 @@
 #include "NiagaraEmitterInstance.h"
 #include "NiagaraGPURayTracingTransformsShader.h"
 #include "NiagaraGpuComputeDispatchInterface.h"
+#include "NiagaraRenderableMeshInterface.h"
 #include "NiagaraSceneProxy.h"
 #include "NiagaraSortingGPU.h"
 #include "NiagaraStats.h"
+#include "NiagaraSystemInstanceController.h"
 #include "ParticleResources.h"
 #include "PipelineStateCache.h"
 #include "RayTracingDefinitions.h"
@@ -76,7 +78,7 @@ FNiagaraRendererMeshes::FNiagaraRendererMeshes(ERHIFeatureLevel::Type FeatureLev
 	SourceMode = Properties->SourceMode;
 	FacingMode = Properties->FacingMode;
 	bLockedAxisEnable = Properties->bLockedAxisEnable;
-	LockedAxis = Properties->LockedAxis;
+	LockedAxis = FVector3f(Properties->LockedAxis);
 	LockedAxisSpace = Properties->LockedAxisSpace;
 	SortMode = Properties->SortMode;
 	bSortHighPrecision = UNiagaraRendererProperties::IsSortHighPrecision(Properties->SortPrecision);
@@ -91,7 +93,6 @@ FNiagaraRendererMeshes::FNiagaraRendererMeshes(ERHIFeatureLevel::Type FeatureLev
 	DistanceCullRangeSquared = FVector2f(0, FLT_MAX);
 	RendererVisibility = Properties->RendererVisibility;
 	bAccurateMotionVectors = Properties->NeedsPreciseMotionVectors();
-	MaxSectionCount = 0;
 
 	if (Properties->bEnableCameraDistanceCulling)
 	{
@@ -177,60 +178,52 @@ void FNiagaraRendererMeshes::Initialize(const UNiagaraRendererProperties* InProp
 
 	const UNiagaraMeshRendererProperties* Properties = CastChecked<const UNiagaraMeshRendererProperties>(InProps);
 
-	MaxSectionCount = 0;
-
 	// Initialize the valid mesh slots, and prep them with the data for every mesh, LOD, and section we'll be needing over the lifetime of the renderer
 	const uint32 MaxMeshes = Properties->Meshes.Num();
 	Meshes.Empty(MaxMeshes);
 	for (uint32 SourceMeshIndex = 0; SourceMeshIndex < MaxMeshes; ++SourceMeshIndex)
 	{
 		const auto& MeshProperties = Properties->Meshes[SourceMeshIndex];
-		UStaticMesh* Mesh = MeshProperties.ResolveStaticMesh(Emitter);
-		
-		if (Mesh)
+
+		// Resolve the renderable mesh
+		FNiagaraRenderableMeshPtr RenderableMesh = MeshProperties.ResolveRenderableMesh(Emitter);
+		if (RenderableMesh == nullptr)
 		{
-			FMeshData& MeshData = Meshes.AddDefaulted_GetRef();
-			MeshData.RenderData = Mesh->GetRenderData();
-			MeshData.MinimumLOD = Mesh->GetMinLODIdx();
-			MeshData.SourceMeshIndex = SourceMeshIndex;
-			MeshData.PivotOffset = (FVector3f)MeshProperties.PivotOffset;
-			MeshData.PivotOffsetSpace = MeshProperties.PivotOffsetSpace;
-			MeshData.Scale = (FVector3f)MeshProperties.Scale;
-			MeshData.Rotation = FQuat4f(MeshProperties.Rotation.Quaternion());
-			MeshData.LocalBounds = Mesh->GetExtendedBounds().GetBox();
-			MeshData.LocalBounds.Min *= Properties->MeshBoundsScale;
-			MeshData.LocalBounds.Max *= Properties->MeshBoundsScale;
-
-			// Create an index remap from mesh material index to it's index in the base material list
-			TArray<UMaterialInterface*> MeshMaterials;
-			Properties->GetUsedMeshMaterials(SourceMeshIndex, Emitter, MeshMaterials);
-			for (UMaterialInterface* MeshMaterial : MeshMaterials)
-			{
-				MeshData.MaterialRemapTable.Add(
-					BaseMaterials_GT.IndexOfByPredicate(
-						[&](UMaterialInterface* LookMat)
+			continue;
+		}
+
+		// We have a valid mesh fill in the details
+		FMeshData& MeshData = Meshes.AddDefaulted_GetRef();
+		MeshData.RenderableMesh = RenderableMesh;
+		MeshData.SourceMeshIndex = SourceMeshIndex;
+		MeshData.PivotOffset = FVector3f(MeshProperties.PivotOffset);
+		MeshData.PivotOffsetSpace = MeshProperties.PivotOffsetSpace;
+		MeshData.Scale = FVector3f(MeshProperties.Scale);
+		MeshData.Rotation = FQuat4f(MeshProperties.Rotation.Quaternion());
+
+		// Get materials and remap them into the base material list
+		TArray<UMaterialInterface*> UsedMaterials;
+		RenderableMesh->GetUsedMaterials(UsedMaterials);
+		Properties->ApplyMaterialOverrides(Emitter, UsedMaterials);
+
+		for (UMaterialInterface* UsedMaterial : UsedMaterials)
+		{
+			MeshData.MaterialRemapTable.Add(
+				BaseMaterials_GT.IndexOfByPredicate(
+					[&](UMaterialInterface* LookMat)
+					{
+						if (LookMat == UsedMaterial)
 						{
-							if (LookMat == MeshMaterial)
-							{
-								return true;
-							}
-							if (UMaterialInstanceDynamic* MID = Cast<UMaterialInstanceDynamic>(LookMat))
-							{
-								return MeshMaterial == MID->Parent;
-							}
-							return false;
+							return true;
 						}
-					)
-				);
-			}
-
-			// Determine the max section count for all LODs of this mesh and accumulate it on the max for all meshes
-			uint32 MaxSectionCountThisMesh = 0;
-			for (const auto& LODModel : MeshData.RenderData->LODResources)
-			{
-				MaxSectionCountThisMesh = FMath::Max<uint32>(MaxSectionCountThisMesh, LODModel.Sections.Num());
-			}
-			MaxSectionCount += MaxSectionCountThisMesh;
+						if (UMaterialInstanceDynamic* MID = Cast<UMaterialInstanceDynamic>(LookMat))
+						{
+							return UsedMaterial == MID->Parent;
+						}
+						return false;
+					}
+				)
+			);
 		}
 	}
 
@@ -241,28 +234,6 @@ void FNiagaraRendererMeshes::ReleaseRenderThreadResources()
 {
 }
 
-void FNiagaraRendererMeshes::SetupVertexFactory(FNiagaraMeshVertexFactory& InVertexFactory, const FStaticMeshLODResources& LODResources) const
-{
-	FStaticMeshDataType Data;
-
-	LODResources.VertexBuffers.PositionVertexBuffer.BindPositionVertexBuffer(&InVertexFactory, Data);
-	LODResources.VertexBuffers.StaticMeshVertexBuffer.BindTangentVertexBuffer(&InVertexFactory, Data);
-	LODResources.VertexBuffers.StaticMeshVertexBuffer.BindTexCoordVertexBuffer(&InVertexFactory, Data, MAX_TEXCOORDS);
-	LODResources.VertexBuffers.ColorVertexBuffer.BindColorVertexBuffer(&InVertexFactory, Data);
-	InVertexFactory.SetData(Data);
-}
-
-int32 FNiagaraRendererMeshes::GetLODIndex(int32 MeshIndex) const
-{
-	check(IsInRenderingThread());
-	check(Meshes.IsValidIndex(MeshIndex));
-
-	const FMeshData& MeshData = Meshes[MeshIndex];
-	const int32 LODIndex = MeshData.RenderData->GetCurrentFirstLODIdx(MeshData.MinimumLOD);
-
-	return MeshData.RenderData->LODResources.IsValidIndex(LODIndex) ? LODIndex : INDEX_NONE;
-}
-
 void FNiagaraRendererMeshes::PrepareParticleMeshRenderData(FParticleMeshRenderData& ParticleMeshRenderData, const FSceneViewFamily& ViewFamily, FMeshElementCollector& Collector, FNiagaraDynamicDataBase* InDynamicData, const FNiagaraSceneProxy* SceneProxy, bool bRayTracing) const
 {
 	ParticleMeshRenderData.Collector = &Collector;
@@ -600,7 +571,7 @@ void FNiagaraRendererMeshes::InitializeSortInfo(const FParticleMeshRenderData& P
 void FNiagaraRendererMeshes::PreparePerMeshData(FParticleMeshRenderData& ParticleMeshRenderData, const FNiagaraMeshVertexFactory& VertexFactory, const FNiagaraSceneProxy& SceneProxy, const FMeshData& MeshData) const
 {
 	// Calculate pivot offset / culling sphere
-	FBox MeshLocalBounds = MeshData.LocalBounds;
+	FBox MeshLocalBounds = MeshData.RenderableMesh->GetLocalBounds();
 	MeshLocalBounds.Min *= FVector(MeshData.Scale);
 	MeshLocalBounds.Max *= FVector(MeshData.Scale);
 	ParticleMeshRenderData.CullingSphere.Center = MeshLocalBounds.GetCenter();
@@ -717,7 +688,7 @@ FNiagaraMeshCommonParameters FNiagaraRendererMeshes::CreateCommonShaderParams(co
 	}
 
 	Params.bLockedAxisEnable			= bLockedAxisEnable;
-	Params.LockedAxis 					= (FVector3f)LockedAxis;
+	Params.LockedAxis 					= LockedAxis;
 	Params.LockedAxisSpace				= (uint32)LockedAxisSpace;
 
 	if (bUseLocalSpace)
@@ -975,6 +946,10 @@ void FNiagaraRendererMeshes::SetupElementForGPUScene(
 	{
 		FMemory::Memzero(&GPUSceneRes.GPUWriteParams, sizeof(GPUSceneRes.GPUWriteParams));
 
+		FBox LocalBounds = MeshData.RenderableMesh->GetLocalBounds();
+		LocalBounds.Min *= FVector(MeshData.Scale);
+		LocalBounds.Max *= FVector(MeshData.Scale);
+
 		GPUSceneRes.GPUWriteParams.Common 					= CommonParameters;
 		GPUSceneRes.GPUWriteParams.ParticleCount 			= NumInstances;
 		GPUSceneRes.GPUWriteParams.GPUParticleCountBuffer 	= GetSrvOrDefaultUInt(ComputeDispatchInterface->GetGPUInstanceCounterManager().GetInstanceCountBuffer());
@@ -983,7 +958,7 @@ void FNiagaraRendererMeshes::SetupElementForGPUScene(
 		GPUSceneRes.GPUWriteParams.MeshIndexDataOffset 		= ParticleMeshRenderData.MeshIndexOffset;
 		GPUSceneRes.GPUWriteParams.RendererVisibility 		= RendererVisibility;
 		GPUSceneRes.GPUWriteParams.VisibilityTagDataOffset 	= ParticleMeshRenderData.RendererVisTagOffset;
-		GPUSceneRes.GPUWriteParams.LocalBoundingCenter		= (FVector3f)MeshData.LocalBounds.GetCenter();
+		GPUSceneRes.GPUWriteParams.LocalBoundingCenter		= (FVector3f)LocalBounds.GetCenter();
 		GPUSceneRes.GPUWriteParams.DistanceCullRangeSquared = DistanceCullRangeSquared;
 		GPUSceneRes.GPUWriteParams.bNeedsPrevTransform		= bNeedsPrevTransform ? 1 : 0;
 
@@ -1078,7 +1053,7 @@ void FNiagaraRendererMeshes::CreateMeshBatchForSection(
 	FMaterialRenderProxy& MaterialProxy,
 	const FNiagaraSceneProxy& SceneProxy,
 	const FMeshData& MeshData,
-	const FStaticMeshLODResources& LODModel,
+	const INiagaraRenderableMesh::FLODModelData& LODModel,
 	const FStaticMeshSection& Section,
 	const FSceneView& View,
 	int32 ViewIndex,
@@ -1098,14 +1073,18 @@ void FNiagaraRendererMeshes::CreateMeshBatchForSection(
 #endif
 	MeshBatch.DepthPriorityGroup = (ESceneDepthPriorityGroup)SceneProxy.GetDepthPriorityGroup(&View);
 
+	FBox LocalBounds = MeshData.RenderableMesh->GetLocalBounds();
+	LocalBounds.Min *= FVector(MeshData.Scale);
+	LocalBounds.Max *= FVector(MeshData.Scale);
+
 	FMeshBatchElement& BatchElement = MeshBatch.Elements[0];
 	if (ParticleMeshRenderData.bUseGPUScene)
 	{
-		BatchElement.PrimitiveUniformBufferResource = SceneProxy.GetCustomUniformBufferResource(IsMotionBlurEnabled(), MeshData.LocalBounds);
+		BatchElement.PrimitiveUniformBufferResource = SceneProxy.GetCustomUniformBufferResource(IsMotionBlurEnabled(), LocalBounds);
 	}
 	else
 	{
-		BatchElement.PrimitiveUniformBuffer = SceneProxy.GetCustomUniformBuffer(IsMotionBlurEnabled(), MeshData.LocalBounds);
+		BatchElement.PrimitiveUniformBuffer = SceneProxy.GetCustomUniformBuffer(IsMotionBlurEnabled(), LocalBounds);
 	}
 	BatchElement.FirstIndex = 0;
 	BatchElement.MinVertexIndex = 0;
@@ -1128,13 +1107,13 @@ void FNiagaraRendererMeshes::CreateMeshBatchForSection(
 
 	if (bIsWireframe)
 	{
-		if (LODModel.AdditionalIndexBuffers && LODModel.AdditionalIndexBuffers->WireframeIndexBuffer.IsInitialized())
+		if (LODModel.WireframeIndexBuffer != nullptr)
 		{
 			MeshBatch.Type = PT_LineList;
 			MeshBatch.MaterialRenderProxy = UMaterial::GetDefaultMaterial(MD_Surface)->GetRenderProxy();
 			BatchElement.FirstIndex = 0;
-			BatchElement.IndexBuffer = &LODModel.AdditionalIndexBuffers->WireframeIndexBuffer;
-			BatchElement.NumPrimitives = LODModel.AdditionalIndexBuffers->WireframeIndexBuffer.GetNumIndices() / 2;
+			BatchElement.IndexBuffer = LODModel.WireframeIndexBuffer;
+			BatchElement.NumPrimitives = LODModel.WireframeNumIndices / 2;
 		}
 		else
 		{
@@ -1142,15 +1121,15 @@ void FNiagaraRendererMeshes::CreateMeshBatchForSection(
 			MeshBatch.MaterialRenderProxy = &MaterialProxy;
 			MeshBatch.bWireframe = true;
 			BatchElement.FirstIndex = 0;
-			BatchElement.IndexBuffer = &LODModel.IndexBuffer;
-			BatchElement.NumPrimitives = LODModel.IndexBuffer.GetNumIndices() / 3;
+			BatchElement.IndexBuffer = LODModel.IndexBuffer;
+			BatchElement.NumPrimitives = LODModel.NumIndices / 3;
 		}
 	}
 	else
 	{
 		MeshBatch.Type = PT_TriangleList;
 		MeshBatch.MaterialRenderProxy = &MaterialProxy;
-		BatchElement.IndexBuffer = &LODModel.IndexBuffer;
+		BatchElement.IndexBuffer = LODModel.IndexBuffer;
 		BatchElement.FirstIndex = Section.FirstIndex;
 		BatchElement.NumPrimitives = Section.NumTriangles;
 	}
@@ -1265,16 +1244,13 @@ void FNiagaraRendererMeshes::GetDynamicMeshElements(const TArray<const FSceneVie
 
 				const FMeshData& MeshData = Meshes[MeshIndex];
 
-				// @TODO : support multiple LOD
-				const int32 LODIndex = GetLODIndex(MeshIndex);
-				if (LODIndex == INDEX_NONE)
+				INiagaraRenderableMesh::FLODModelData LODModel;
+				MeshData.RenderableMesh->GetLODModelData(LODModel);
+				if (LODModel.LODIndex == INDEX_NONE)
 				{
 					continue;
 				}
 
-				const FStaticMeshLODResources& LODModel = MeshData.RenderData->LODResources[LODIndex];
-				const int32 SectionCount = LODModel.Sections.Num();
-
 				FMeshCollectorResources* CollectorResources = &Collector.AllocateOneFrameResource<FMeshCollectorResources>();
 
 				// Get the next vertex factory to use
@@ -1282,10 +1258,10 @@ void FNiagaraRendererMeshes::GetDynamicMeshElements(const TArray<const FSceneVie
 				FNiagaraMeshVertexFactory& VertexFactory = CollectorResources->VertexFactory;
 				VertexFactory.SetParticleFactoryType(NVFT_Mesh);
 				VertexFactory.SetMeshIndex(MeshIndex);
-				VertexFactory.SetLODIndex(LODIndex);
+				VertexFactory.SetLODIndex(LODModel.LODIndex);
 				VertexFactory.EnablePrimitiveIDElement(ParticleMeshRenderData.bUseGPUScene);
 				VertexFactory.InitResource();
-				SetupVertexFactory(VertexFactory, LODModel);
+				MeshData.RenderableMesh->SetupVertexFactory(VertexFactory, LODModel);
 
 				PreparePerMeshData(ParticleMeshRenderData, VertexFactory, *SceneProxy, MeshData);
 
@@ -1294,7 +1270,7 @@ void FNiagaraRendererMeshes::GetDynamicMeshElements(const TArray<const FSceneVie
 				if ( NumInstances > 0 )
 				{
 					// Increment stats
-					INC_DWORD_STAT_BY(STAT_NiagaraNumMeshVerts, NumInstances * LODModel.GetNumVertices());
+					INC_DWORD_STAT_BY(STAT_NiagaraNumMeshVerts, NumInstances * LODModel.NumVertices);
 					INC_DWORD_STAT_BY(STAT_NiagaraNumMeshes, NumInstances);
 
 					FNiagaraMeshCommonParameters CommonParameters = CreateCommonShaderParams(ParticleMeshRenderData, *View, MeshData, *SceneProxy);
@@ -1303,7 +1279,7 @@ void FNiagaraRendererMeshes::GetDynamicMeshElements(const TArray<const FSceneVie
 					CollectorResources->UniformBuffer = VFUniformBuffer;
 
 					const bool bIsWireframe = AllowDebugViewmodes() && ViewFamily.EngineShowFlags.Wireframe;
-					for (int32 SectionIndex = 0; SectionIndex < SectionCount; SectionIndex++)
+					for (int32 SectionIndex = 0; SectionIndex < LODModel.Sections.Num(); SectionIndex++)
 					{
 						const FStaticMeshSection& Section = LODModel.Sections[SectionIndex];
 						const uint32 RemappedMaterialIndex = MeshData.MaterialRemapTable[Section.MaterialIndex];
@@ -1412,17 +1388,16 @@ void FNiagaraRendererMeshes::GetDynamicRayTracingInstances(FRayTracingMaterialGa
 		}
 
 		const FMeshData& MeshData = Meshes[MeshIndex];
-		const int32 LODIndex = GetLODIndex(MeshIndex);
-		if (LODIndex == INDEX_NONE)
+
+		INiagaraRenderableMesh::FLODModelData LODModel;
+		MeshData.RenderableMesh->GetLODModelData(LODModel);
+		if (LODModel.LODIndex == INDEX_NONE)
 		{
 			continue;
 		}
 
-		const FStaticMeshLODResources& LODModel = MeshData.RenderData->LODResources[LODIndex];
-		const FStaticMeshVertexFactories& VFs = MeshData.RenderData->LODVertexFactories[LODIndex];
-		FRayTracingGeometry& Geometry = MeshData.RenderData->LODResources[LODIndex].RayTracingGeometry;
 		FRayTracingInstance RayTracingInstance;
-		RayTracingInstance.Geometry = &Geometry;
+		RayTracingInstance.Geometry = LODModel.RayTracingGeometry;
 
 		FMeshCollectorResources* CollectorResources = &Context.RayTracingMeshResourceCollector.AllocateOneFrameResource<FMeshCollectorResources>();
 
@@ -1431,10 +1406,10 @@ void FNiagaraRendererMeshes::GetDynamicRayTracingInstances(FRayTracingMaterialGa
 		FNiagaraMeshVertexFactory& VertexFactory = CollectorResources->VertexFactory;
 		VertexFactory.SetParticleFactoryType(NVFT_Mesh);
 		VertexFactory.SetMeshIndex(MeshIndex);
-		VertexFactory.SetLODIndex(LODIndex);
+		VertexFactory.SetLODIndex(LODModel.LODIndex);
 		VertexFactory.EnablePrimitiveIDElement(ParticleMeshRenderData.bUseGPUScene);
 		VertexFactory.InitResource();
-		SetupVertexFactory(VertexFactory, LODModel);
+		MeshData.RenderableMesh->SetupVertexFactory(VertexFactory, LODModel);
 
 		PreparePerMeshData(ParticleMeshRenderData, VertexFactory, *SceneProxy, MeshData);
 
@@ -1494,15 +1469,15 @@ void FNiagaraRendererMeshes::GetDynamicRayTracingInstances(FRayTracingMaterialGa
 				false // bNeedsPrevTransform
 			);
 			MeshBatch.SegmentIndex = SectionIndex;
-			MeshBatch.LODIndex = LODIndex;
+			MeshBatch.LODIndex = LODModel.LODIndex;
 
 #if !(UE_BUILD_SHIPPING || UE_BUILD_TEST)
-			MeshBatch.VisualizeLODIndex = LODIndex;
+			MeshBatch.VisualizeLODIndex = LODModel.LODIndex;
 #endif
 
 			MeshBatch.bCanApplyViewModeOverrides = false;
 
-			MeshBatch.Elements[0].VertexFactoryUserData = VFs.VertexFactory.GetUniformBuffer();
+			MeshBatch.Elements[0].VertexFactoryUserData = LODModel.VertexFactoryUserData;
 #if !(UE_BUILD_SHIPPING || UE_BUILD_TEST)
 			MeshBatch.Elements[0].VisualizeElementIndex = SectionIndex;
 #endif
@@ -1792,7 +1767,7 @@ FNiagaraDynamicDataBase* FNiagaraRendererMeshes::GenerateDynamicData(const FNiag
 	// Bail if we have cached mesh render data for any meshes that are no longer valid
 	for (const auto& MeshData : Meshes)
 	{
-		if ( !Properties->Meshes.IsValidIndex(MeshData.SourceMeshIndex) )
+		if (MeshData.RenderableMesh == nullptr)
 		{
 			return nullptr;
 		}
diff --git a/Engine/Plugins/FX/Niagara/Source/Niagara/Public/NiagaraMeshRendererProperties.h b/Engine/Plugins/FX/Niagara/Source/Niagara/Public/NiagaraMeshRendererProperties.h
--- a/Engine/Plugins/FX/Niagara/Source/Niagara/Public/NiagaraMeshRendererProperties.h
+++ b/Engine/Plugins/FX/Niagara/Source/Niagara/Public/NiagaraMeshRendererProperties.h
@@ -3,9 +3,11 @@
 #pragma once
 
 #include "CoreMinimal.h"
-#include "NiagaraRendererProperties.h"
-#include "NiagaraGPUSortInfo.h"
 #include "NiagaraCommon.h"
+#include "NiagaraGPUSortInfo.h"
+#include "NiagaraRenderableMeshInterface.h"
+#include "NiagaraRendererProperties.h"
+#include "NiagaraParameterBinding.h"
 #include "NiagaraMeshRendererProperties.generated.h"
 
 class UMaterialInstanceConstant;
@@ -129,9 +131,14 @@ struct NIAGARA_API FNiagaraMeshRendererMeshProperties
 	UPROPERTY(EditAnywhere, Category = "Mesh")
 	TObjectPtr<UStaticMesh> Mesh;
 
-	/** Use the UStaticMesh bound to this user variable if it is set to a valid value. If this is bound to a valid value and Mesh is also set, UserParamBinding wins.*/
+#if WITH_EDITORONLY_DATA
+	UPROPERTY()
+	FNiagaraUserParameterBinding UserParamBinding_DEPRECATED;
+#endif
+
+	/** Binding to supported mesh types. */
 	UPROPERTY(EditAnywhere, Category = "Mesh")
-	FNiagaraUserParameterBinding UserParamBinding;
+	FNiagaraParameterBinding MeshParameterBinding;
 
 	/** Scale of the mesh */
 	UPROPERTY(EditAnywhere, Category = "Mesh")
@@ -149,8 +156,11 @@ struct NIAGARA_API FNiagaraMeshRendererMeshProperties
 	UPROPERTY(EditAnywhere, Category = "Mesh")
 	ENiagaraMeshPivotOffsetSpace PivotOffsetSpace;
 
-    UStaticMesh* ResolveStaticMesh(const FNiagaraEmitterInstance* Emitter) const;
-	bool HasValidMeshProperties() const;
+	/** Resolve renderable mesh. */
+	FNiagaraRenderableMeshPtr ResolveRenderableMesh(const FNiagaraEmitterInstance* EmitterInstance) const;
+
+	/** Is the renderable mesh potentially valid or not. */
+	bool HasValidRenderableMesh() const;
 };
 
 UCLASS(editinlinenew, meta = (DisplayName = "Mesh Renderer"))
@@ -207,7 +217,7 @@ public:
 	virtual ENiagaraRendererSourceDataMode GetCurrentSourceMode() const override { return SourceMode; }
 	//UNiagaraRendererProperties Interface END
 
-	void GetUsedMeshMaterials(int32 MeshIndex, const FNiagaraEmitterInstance* Emitter, TArray<UMaterialInterface*>& OutMaterials) const;	
+	void ApplyMaterialOverrides(const FNiagaraEmitterInstance* EmitterInstance, TArray<UMaterialInterface*>& InOutMaterials) const;
 
 	/**
 	 * The static mesh(es) to be instanced when rendering mesh particles.
diff --git a/Engine/Plugins/FX/Niagara/Source/Niagara/Public/NiagaraParameterBinding.h b/Engine/Plugins/FX/Niagara/Source/Niagara/Public/NiagaraParameterBinding.h
new file mode 100644
--- /dev/null
+++ b/Engine/Plugins/FX/Niagara/Source/Niagara/Public/NiagaraParameterBinding.h
@@ -0,0 +1,86 @@
+// Copyright Epic Games, Inc. All Rights Reserved.
+
+#pragma once
+
+#include "Misc/EnumClassFlags.h"
+
+#include "NiagaraCommon.h"
+#include "NiagaraParameterBinding.generated.h"
+
+// Binding usage flags
+// Note that Emitter / Particle bindings are only allowed if the binding variable is outered to a UNiagaraEmitter
+enum class ENiagaraParameterBindingUsage
+{
+	User		= 1 << 0,		// Can we bind to user parameters
+	System		= 1 << 1,		// Can we bind to system parameters
+	Emitter		= 1 << 2,		// Can we bind to emitter parameters
+	Particle	= 1 << 3,		// Can we bind to particle parameters
+	StaticVariable = 1 << 4,	// Can we bind to static variables
+
+	NotParticle	= StaticVariable | User | System | Emitter,
+	All			= StaticVariable | User | System | Emitter | Particle,
+};
+
+ENUM_CLASS_FLAGS(ENiagaraParameterBindingUsage)
+
+USTRUCT()
+struct NIAGARA_API FNiagaraParameterBinding
+{
+	GENERATED_USTRUCT_BODY()
+
+#if WITH_EDITORONLY_DATA
+	bool CanBindTo(FNiagaraTypeDefinition TypeDefinition) const;
+	bool CanBindTo(FNiagaraVariableBase InVariable, FNiagaraVariableBase& OutAliasedVariable, FStringView EmitterName) const;
+
+	void SetUsage(ENiagaraParameterBindingUsage InUsage) { BindingUsage = InUsage; }
+	void SetAllowedDataInterfaces(TArray<UClass*> InClasses) { AllowedDataInterfaces = InClasses; }
+	void SetAllowedObjects(TArray<UClass*> InClasses) { AllowedObjects = InClasses; }
+	void SetAllowedInterfaces(TArray<UClass*> InClasses) { AllowedInterfaces = InClasses; }
+	void SetAllowedTypeDefinitions(TArray<FNiagaraTypeDefinition> InTypeDefs) { AllowedTypeDefinitions = InTypeDefs; }
+
+	bool AllowUserParameters() const { return EnumHasAllFlags(BindingUsage, ENiagaraParameterBindingUsage::User); }
+	bool AllowSystemParameters() const { return EnumHasAllFlags(BindingUsage, ENiagaraParameterBindingUsage::System); }
+	bool AllowEmitterParameters() const { return EnumHasAllFlags(BindingUsage, ENiagaraParameterBindingUsage::Emitter); }
+	bool AllowParticleParameters() const { return EnumHasAllFlags(BindingUsage, ENiagaraParameterBindingUsage::Particle); }
+	bool AllowStaticVariables() const { return EnumHasAllFlags(BindingUsage, ENiagaraParameterBindingUsage::StaticVariable); }
+
+	TConstArrayView<UClass*> GetAllowedDataInterfaces() const { return MakeArrayView(AllowedDataInterfaces); }
+	TConstArrayView<UClass*> GetAllowedObjects() const { return MakeArrayView(AllowedObjects); }
+	TConstArrayView<UClass*> GetAllowedInterfaces() const { return MakeArrayView(AllowedInterfaces); }
+	TConstArrayView<FNiagaraTypeDefinition> GetAllowedTypeDefinitions() const { return MakeArrayView(AllowedTypeDefinitions); }
+#endif
+
+	/** Parameter binding used by the runtime fully resolved, contains a fallback value, i.e. NamedEmitter.Parameter */
+	UPROPERTY(EditAnywhere, Category = "Parameter Binding")
+	FNiagaraVariableBase Parameter;
+
+#if WITH_EDITORONLY_DATA
+	/** Parameter binding used in the UI, i.e. Emitter.Parameter */
+	UPROPERTY(EditAnywhere, Category = "Parameter Binding")
+	FNiagaraVariableBase AliasedParameter;
+
+protected:
+	/* Set the usage for the binding. */
+	ENiagaraParameterBindingUsage BindingUsage = ENiagaraParameterBindingUsage::NotParticle;
+
+	/** List of data interfaces we can bind to, matches with a valid Cast<>. */
+	UPROPERTY(transient)
+	TArray<TObjectPtr<UClass>> AllowedDataInterfaces;
+
+	/** List of UObject types we can bind to, matches with a valid Cast<>. */
+	UPROPERTY(transient)
+	TArray<TObjectPtr<UClass>> AllowedObjects;
+
+	/** List of Interfaces to look for on Object & DataInterfaces */
+	UPROPERTY(transient)
+	TArray<TObjectPtr<UClass>> AllowedInterfaces;
+
+	/** List of explicit type definitions allowed, must be an exact match. */
+	TArray<FNiagaraTypeDefinition>	AllowedTypeDefinitions;
+#endif
+
+	//FORCEINLINE bool operator==(const FNiagaraParameterBinding& Other)const
+	//{
+		//return Other.Parameter == Parameter;
+	//}
+};
\ No newline at end of file
diff --git a/Engine/Plugins/FX/Niagara/Source/Niagara/Public/NiagaraRenderableMeshInterface.h b/Engine/Plugins/FX/Niagara/Source/Niagara/Public/NiagaraRenderableMeshInterface.h
new file mode 100644
--- /dev/null
+++ b/Engine/Plugins/FX/Niagara/Source/Niagara/Public/NiagaraRenderableMeshInterface.h
@@ -0,0 +1,70 @@
+// Copyright Epic Games, Inc. All Rights Reserved.
+
+#pragma once
+
+#include "UObject/Interface.h"
+#include "StaticMeshResources.h"
+
+#include "NiagaraCore.h"
+#include "NiagaraRenderableMeshInterface.generated.h"
+
+class FIndexBuffer;
+class FNiagaraEmitterInstance;
+class FNiagaraMeshVertexFactory;
+class UMaterialInterface;
+class FRayTracingGeometry;
+class FRHIUniformBuffer;
+struct FStaticMeshSection;
+
+// Abstact class that for a renderable mesh
+// While Niagara holds a reference to one of these it is expected that the data will exist
+class INiagaraRenderableMesh : public TSharedFromThis<INiagaraRenderableMesh, ESPMode::ThreadSafe>
+{
+public:
+	struct FLODModelData
+	{
+		int32								LODIndex = INDEX_NONE;
+		int32								NumVertices = 0;
+		int32								NumIndices = 0;
+		TConstArrayView<FStaticMeshSection>	Sections;
+		const FIndexBuffer*					IndexBuffer = nullptr;
+		FRHIUniformBuffer*					VertexFactoryUserData = nullptr;
+		const FRayTracingGeometry*			RayTracingGeometry = nullptr;
+
+		int32								WireframeNumIndices = 0;
+		const FIndexBuffer*					WireframeIndexBuffer = nullptr;
+	};
+
+	virtual ~INiagaraRenderableMesh() {}
+
+	// Get the local bounds for the mesh
+	virtual FBox GetLocalBounds() const = 0;
+	// Gather all the relevant mesh data to render the mesh
+	virtual void GetLODModelData(FLODModelData& OutLODModelData) const = 0;
+	// Setup the vertex factory for the mesh
+	virtual void SetupVertexFactory(class FNiagaraMeshVertexFactory& InVertexFactory, const FLODModelData& LODModelData) const = 0;
+	// Gather a list of used materials
+	virtual void GetUsedMaterials(TArray<UMaterialInterface*>& OutMaterials) const = 0;
+};
+
+using FNiagaraRenderableMeshPtr = TSharedPtr<INiagaraRenderableMesh, ESPMode::ThreadSafe>;
+
+// Interface for UObjects to implement renderable mesh
+UINTERFACE()
+class UNiagaraRenderableMeshInterface : public UInterface
+{
+	GENERATED_BODY()
+};
+
+// Interface defintion for UObjects
+class INiagaraRenderableMeshInterface
+{
+	GENERATED_BODY()
+
+public:
+	// Get renderable mesh pointer
+	virtual FNiagaraRenderableMeshPtr GetRenderableMesh(FNiagaraSystemInstanceID SystemInstanceID) = 0;
+
+	// Get used materials
+	virtual void GetUsedMaterials(FNiagaraSystemInstanceID SystemInstanceID, TArray<UMaterialInterface*>& OutMaterials) const = 0;
+};
diff --git a/Engine/Plugins/FX/Niagara/Source/Niagara/Public/NiagaraRendererMeshes.h b/Engine/Plugins/FX/Niagara/Source/Niagara/Public/NiagaraRendererMeshes.h
--- a/Engine/Plugins/FX/Niagara/Source/Niagara/Public/NiagaraRendererMeshes.h
+++ b/Engine/Plugins/FX/Niagara/Source/Niagara/Public/NiagaraRendererMeshes.h
@@ -6,11 +6,11 @@ NiagaraRenderer.h: Base class for Niagara render modules
 #pragma once
 
 #include "NiagaraRenderer.h"
+#include "NiagaraRenderableMeshInterface.h"
 #include "NiagaraMeshRendererProperties.h"
 #include "NiagaraMeshVertexFactory.h"
 #include "NiagaraGPUSortInfo.h"
 #include "NiagaraSystemInstance.h"
-#include "StaticMeshResources.h"
 #include "NiagaraGPUSceneUtils.h"
 
 class FNiagaraDataSet;
@@ -38,8 +38,6 @@ public:
 #endif
 	//FNiagaraRenderer Interface END
 
-	void SetupVertexFactory(FNiagaraMeshVertexFactory& InVertexFactory, const FStaticMeshLODResources& LODResources) const;
-
 protected:
 	struct FParticleMeshRenderData
 	{
@@ -80,14 +78,12 @@ protected:
 
 	struct FMeshData
 	{
-		FStaticMeshRenderData* RenderData = nullptr;
-		int32 MinimumLOD = 0;
+		FNiagaraRenderableMeshPtr RenderableMesh;
 		uint32 SourceMeshIndex = INDEX_NONE;
 		FVector3f PivotOffset = FVector3f::ZeroVector;
 		ENiagaraMeshPivotOffsetSpace PivotOffsetSpace = ENiagaraMeshPivotOffsetSpace::Mesh;
 		FVector3f Scale = FVector3f(1.0f, 1.0f, 1.0f);
 		FQuat4f Rotation = FQuat4f::Identity;
-		FBox LocalBounds = FBox(ForceInitToZero);
 		TArray<uint32, TInlineAllocator<4>> MaterialRemapTable;
 	};
 
@@ -120,8 +116,6 @@ protected:
 		virtual ~FGPUSceneUpdateResource() {}
 	};
 
-	int32 GetLODIndex(int32 MeshIndex) const;
-
 	void PrepareParticleMeshRenderData(FParticleMeshRenderData& ParticleMeshRenderData, const FSceneViewFamily& ViewFamily, FMeshElementCollector& Collector, FNiagaraDynamicDataBase* InDynamicData, const FNiagaraSceneProxy* SceneProxy, bool bRayTracing) const;
 	void PrepareParticleRenderBuffers(FParticleMeshRenderData& ParticleMeshRenderData, FGlobalDynamicReadBuffer& DynamicReadBuffer) const;
 	void InitializeSortInfo(const FParticleMeshRenderData& ParticleMeshRenderData, const FNiagaraSceneProxy& SceneProxy, const FSceneView& View, int32 ViewIndex, bool bIsInstancedStereo, FNiagaraGPUSortInfo& OutSortInfo) const;
@@ -149,7 +143,7 @@ protected:
 		FMaterialRenderProxy& MaterialProxy,
 		const FNiagaraSceneProxy& SceneProxy,
 		const FMeshData& MeshData,
-		const FStaticMeshLODResources& LODModel,
+		const INiagaraRenderableMesh::FLODModelData& LODModel,
 		const FStaticMeshSection& Section,
 		const FSceneView& View,
 		int32 ViewIndex,
@@ -176,20 +170,19 @@ private:
 	uint32 bAccurateMotionVectors : 1;
 
 	uint32 bSubImageBlend : 1;
-	FVector2f SubImageSize;
+	FVector2f SubImageSize = FVector2f::ZeroVector;
 
-	FVector LockedAxis;
+	FVector3f LockedAxis = FVector3f::ZeroVector;
 	ENiagaraMeshLockedAxisSpace LockedAxisSpace;
 
-	FVector2f DistanceCullRange;
-	FVector2f DistanceCullRangeSquared;
+	FVector2f DistanceCullRange = FVector2f::ZeroVector;
+	FVector2f DistanceCullRangeSquared = FVector2f::ZeroVector;
 	int32 ParticleRendererVisTagOffset = INDEX_NONE;
 	int32 ParticleMeshIndexOffset = INDEX_NONE;
 	int32 RendererVisibility = 0;
 	int32 EmitterRendererVisTagOffset = INDEX_NONE;
 	int32 EmitterMeshIndexOffset = INDEX_NONE;
-	uint32 MaterialParamValidMask;
-	uint32 MaxSectionCount;
+	uint32 MaterialParamValidMask = 0;
 
 	int32 VFBoundOffsetsInParamStore[ENiagaraMeshVFLayout::Type::Num_Max];
 	uint32 bSetAnyBoundVars : 1;
diff --git a/Engine/Plugins/FX/Niagara/Source/Niagara/Public/NiagaraTypes.h b/Engine/Plugins/FX/Niagara/Source/Niagara/Public/NiagaraTypes.h
--- a/Engine/Plugins/FX/Niagara/Source/Niagara/Public/NiagaraTypes.h
+++ b/Engine/Plugins/FX/Niagara/Source/Niagara/Public/NiagaraTypes.h
@@ -1660,14 +1660,14 @@ struct FNiagaraVariableBase
 	i.e. MyNamespace.VarName would become VarName but Other.MyNamespace.VarName would not be modified.
 	For more complex namespace manipulation refer to FNiagaraUtilities::ResolveAliases
 	*/
-	bool RemoveRootNamespace(const FStringView& ExpectedNamespace);
+	NIAGARA_API bool RemoveRootNamespace(const FStringView& ExpectedNamespace);
 
 	/**
 	Replaces the root namespace from the variable if it matches the expected namespace with the NewNamespace.
 	i.e. Prev.VarName would become Next.VarName but Other.Prev.VarName would not be modified.
 	For more complex namespace manipulation refer to FNiagaraUtilities::ResolveAliases
 	*/
-	bool ReplaceRootNamespace(const FStringView& ExpectedNamespace, const FStringView& NewNamespace);
+	NIAGARA_API bool ReplaceRootNamespace(const FStringView& ExpectedNamespace, const FStringView& NewNamespace);
 
 	FORCEINLINE bool IsInNameSpace(const FStringView& Namespace) const
 	{
diff --git a/Engine/Plugins/FX/Niagara/Source/NiagaraEditor/Private/Customizations/NiagaraParameterBindingCustomization.cpp b/Engine/Plugins/FX/Niagara/Source/NiagaraEditor/Private/Customizations/NiagaraParameterBindingCustomization.cpp
new file mode 100644
--- /dev/null
+++ b/Engine/Plugins/FX/Niagara/Source/NiagaraEditor/Private/Customizations/NiagaraParameterBindingCustomization.cpp
@@ -0,0 +1,314 @@
+// Copyright Epic Games, Inc. All Rights Reserved.
+
+#include "NiagaraParameterBindingCustomization.h"
+
+#include "CoreMinimal.h"
+#include "DetailLayoutBuilder.h"
+#include "DetailWidgetRow.h"
+#include "IDetailChildrenBuilder.h"
+#include "IDetailGroup.h"
+#include "IPropertyUtilities.h"
+#include "PropertyHandle.h"
+#include "ScopedTransaction.h"
+#include "SGraphActionMenu.h"
+#include "Framework/Application/SlateApplication.h"
+#include "Widgets/Input/SComboButton.h"
+
+#include "NiagaraEmitter.h"
+#include "NiagaraNodeParameterMapBase.h"
+#include "NiagaraParameterMapHistory.h"
+#include "NiagaraRendererProperties.h"
+#include "NiagaraScriptSource.h"
+#include "NiagaraSimulationStageBase.h"
+#include "NiagaraSystem.h"
+#include "NiagaraParameterBinding.h"
+#include "Customizations/NiagaraTypeCustomizations.h"
+#include "Widgets/Input/SButton.h"
+#include "Widgets/SNiagaraParameterName.h"
+
+#include UE_INLINE_GENERATED_CPP_BY_NAME(NiagaraParameterBindingCustomization)
+
+#define LOCTEXT_NAMESPACE "NiagaraParameterBindingCustomization"
+
+FNiagaraParameterBindingAction::FNiagaraParameterBindingAction()
+	: FEdGraphSchemaAction(
+		FText(),
+		LOCTEXT("None", "None"),
+		LOCTEXT("NoneTooltip", "Not bound to any variable will use the default value"),
+		1,
+		FText()
+	)
+{
+}
+
+FNiagaraParameterBindingAction::FNiagaraParameterBindingAction(FNiagaraVariableBase InAliasedParameter, FNiagaraVariableBase InParameter)
+	: FEdGraphSchemaAction(
+		FText(),
+		FText::FromName(InAliasedParameter.GetName()),
+		FText::Format(LOCTEXT("VariableBindActionTooltip", "Bind to variable \"{0}\""), FText::FromName(InAliasedParameter.GetName())),
+		0,
+		FText()
+	)
+	, AliasedParameter(InAliasedParameter)
+	, Parameter(InParameter)
+{
+}
+
+TSharedRef<IPropertyTypeCustomization> FNiagaraParameterBindingCustomization::MakeInstance()
+{
+	return MakeShared<FNiagaraParameterBindingCustomization>();
+}
+
+void FNiagaraParameterBindingCustomization::CustomizeHeader(TSharedRef<IPropertyHandle> InPropertyHandle, FDetailWidgetRow& HeaderRow, IPropertyTypeCustomizationUtils& StructCustomizationUtils)
+{
+	PropertyHandle = InPropertyHandle;
+	EmitterWeakPtr.Reset();
+	SystemWeakPtr.Reset();
+
+	// We only handle single object
+	{
+		TArray<UObject*> OuterObjects;
+		PropertyHandle->GetOuterObjects(OuterObjects);
+		UObject* OwnerObject = OuterObjects.Num() == 1 ? OuterObjects[0] : nullptr;
+		if (OwnerObject == nullptr)
+		{
+			HeaderRow
+			.NameContent()
+			[
+				PropertyHandle->CreatePropertyNameWidget()
+			]
+			.ValueContent()
+			.MaxDesiredWidth(200.f)
+			[
+				SNew(STextBlock)
+				.Text(FText::FromString(FName::NameToDisplayString(CastField<FStructProperty>(PropertyHandle->GetProperty())->Struct->GetName(), false)))
+				.Font(IDetailLayoutBuilder::GetDetailFont())
+			];
+
+			return;
+		}
+
+		OwnerWeakPtr = OwnerObject;
+		EmitterWeakPtr = OwnerObject->GetTypedOuter<UNiagaraEmitter>();
+		SystemWeakPtr = OwnerObject->GetTypedOuter<UNiagaraSystem>();
+	}
+
+	HeaderRow
+	.NameContent()
+	[
+		PropertyHandle->CreatePropertyNameWidget()
+	]
+	.ValueContent()
+	[
+		SNew(SHorizontalBox)
+		+ SHorizontalBox::Slot()
+		.FillWidth(1.0f)
+		.Padding(0.0f, 0.0f, 4.0f, 0.0f)
+		[
+			SNew(SComboButton)					
+			.OnGetMenuContent(this, &FNiagaraParameterBindingCustomization::OnGetMenuContent)
+			.ContentPadding(1)
+			.ToolTipText(this, &FNiagaraParameterBindingCustomization::GetTooltipText)
+			.ButtonContent()
+			[
+				SNew(SNiagaraParameterName)
+				.ParameterName(this, &FNiagaraParameterBindingCustomization::GetVariableName)
+				.IsReadOnly(true)
+			]
+		]
+	];
+}
+
+bool FNiagaraParameterBindingCustomization::IsValid() const
+{
+	return OwnerWeakPtr.IsValid();
+}
+
+FNiagaraParameterBinding* FNiagaraParameterBindingCustomization::GetParameterBinding() const
+{
+	UObject* OwnerObject = OwnerWeakPtr.Get();
+	return OwnerObject ? (FNiagaraParameterBinding*)PropertyHandle->GetValueBaseAddress((uint8*)OwnerObject) : nullptr;
+}
+
+FNiagaraParameterBinding* FNiagaraParameterBindingCustomization::GetDefaultParameterBinding() const
+{
+	UObject* OwnerObject = OwnerWeakPtr.Get();
+	return OwnerObject ? (FNiagaraParameterBinding*)PropertyHandle->GetValueBaseAddress((uint8*)OwnerObject->GetClass()->GetDefaultObject()) : nullptr;
+}
+
+TSharedRef<SWidget> FNiagaraParameterBindingCustomization::OnGetMenuContent() const
+{
+	FGraphActionMenuBuilder MenuBuilder;
+
+	return SNew(SBorder)
+		.BorderImage(FAppStyle::GetBrush("Menu.Background"))
+		.Padding(5)
+		[
+			SNew(SBox)
+			[
+				SNew(SGraphActionMenu)
+				.OnActionSelected(this, &FNiagaraParameterBindingCustomization::OnActionSelected)
+				.OnCreateWidgetForAction(this, &FNiagaraParameterBindingCustomization::OnCreateWidgetForAction)
+				.OnCollectAllActions(this, &FNiagaraParameterBindingCustomization::CollectAllActions)
+				.AutoExpandActionMenu(false)
+				.ShowFilterTextBox(true)
+			]
+		];
+}
+
+void FNiagaraParameterBindingCustomization::CollectAllActions(FGraphActionListBuilderBase& OutAllActions) const
+{
+	// Add the unbound action this is where we will fallback to the value only
+	OutAllActions.AddAction(MakeShared<FNiagaraParameterBindingAction>());
+
+	// Add bindings from system
+	FNiagaraParameterBinding* ParameterBinding = GetParameterBinding();
+	UNiagaraSystem* NiagaraSystem = SystemWeakPtr.Get();
+	UNiagaraEmitter* NiagaraEmitter = EmitterWeakPtr.Get();
+	if (ParameterBinding == nullptr || NiagaraSystem == nullptr)
+	{
+		return;
+	}
+
+	TArray<FNiagaraParameterMapHistory> Histories;
+	if (UNiagaraScriptSource* Source = Cast<UNiagaraScriptSource>(NiagaraSystem->GetSystemUpdateScript()->GetLatestSource()))
+	{
+		Histories.Append(UNiagaraNodeParameterMapBase::GetParameterMaps(Source->NodeGraph));
+	}
+
+	if (NiagaraEmitter)
+	{
+		// This is lame but currently the only way to get the version data
+		if (UObject* Owner = OwnerWeakPtr.Get())
+		{
+			if (UNiagaraSimulationStageBase* SimulationStage = Cast<UNiagaraSimulationStageBase>(Owner))
+			{
+				if (UNiagaraScriptSource* Source = Cast<UNiagaraScriptSource>(SimulationStage->GetEmitterData()->GraphSource))
+				{
+					Histories.Append(UNiagaraNodeParameterMapBase::GetParameterMaps(Source->NodeGraph));
+				}
+			}
+			else if (UNiagaraRendererProperties* RendererProperties = Cast<UNiagaraRendererProperties>(Owner))
+			{
+				if (UNiagaraScriptSource* Source = Cast<UNiagaraScriptSource>(RendererProperties->GetEmitterData()->GraphSource))
+				{
+					Histories.Append(UNiagaraNodeParameterMapBase::GetParameterMaps(Source->NodeGraph));
+				}
+			}
+		}
+	}
+
+	TArray<TPair<FNiagaraVariableBase, FNiagaraVariableBase>> AllowedVariables;
+
+	for (const FNiagaraParameterMapHistory& History : Histories)
+	{
+		for (const FNiagaraVariable& Variable : History.Variables)
+		{
+			FNiagaraVariableBase AliasedParameter = Variable;
+			if (ParameterBinding->CanBindTo(Variable, AliasedParameter, NiagaraEmitter->GetUniqueEmitterName()))
+			{
+				AllowedVariables.AddUnique({ AliasedParameter, Variable });
+				continue;
+			}
+		}
+	}
+
+	if (ParameterBinding->AllowUserParameters())
+	{
+		for (const FNiagaraVariableBase& Variable : NiagaraSystem->GetExposedParameters().ReadParameterVariables())
+		{
+			if (ParameterBinding->CanBindTo(Variable.GetType()))
+			{
+				AllowedVariables.AddUnique({ Variable, Variable });
+			}
+		}
+	}
+
+	for (const TPair<FNiagaraVariableBase, FNiagaraVariableBase>& AllowedVariable : AllowedVariables)
+	{
+		OutAllActions.AddAction(MakeShared<FNiagaraParameterBindingAction>(AllowedVariable.Key, AllowedVariable.Value));
+	}
+}
+
+TSharedRef<SWidget> FNiagaraParameterBindingCustomization::OnCreateWidgetForAction(struct FCreateWidgetForActionData* const InCreateData) const
+{
+	const FNiagaraParameterBindingAction* BindingAction = static_cast<const FNiagaraParameterBindingAction*>(InCreateData->Action.Get());
+	check(BindingAction);
+
+	return
+		SNew(SVerticalBox)
+		+ SVerticalBox::Slot()
+		.AutoHeight()
+		[
+			SNew(SNiagaraParameterName)
+			.ParameterName(BindingAction->AliasedParameter.GetName())
+			.IsReadOnly(true)
+			.ToolTipText(InCreateData->Action->GetTooltipDescription())
+		];
+}
+
+void FNiagaraParameterBindingCustomization::OnActionSelected(const TArray<TSharedPtr<FEdGraphSchemaAction>>& SelectedActions, ESelectInfo::Type InSelectionType) const
+{
+	FNiagaraParameterBinding* ParameterBinding = GetParameterBinding();
+	if (ParameterBinding == nullptr || !PropertyHandle.IsValid() || (InSelectionType != ESelectInfo::OnMouseClick && InSelectionType != ESelectInfo::OnKeyPress))
+	{
+		return;
+	}
+
+	for (const TSharedPtr<FEdGraphSchemaAction>& EdAction : SelectedActions)
+	{
+		if (EdAction.IsValid())
+		{
+			FNiagaraParameterBindingAction* BindAction = static_cast<FNiagaraParameterBindingAction*>(EdAction.Get());
+			FScopedTransaction Transaction(FText::Format(LOCTEXT("ChangeParameterBinding", " Change Parameter Binding to \"{0}\" "), FText::FromName(BindAction->Parameter.GetName())));
+			FSlateApplication::Get().DismissAllMenus();
+
+			TArray<UObject*> Objects;
+			PropertyHandle->GetOuterObjects(Objects);
+			for (UObject* Obj : Objects)
+			{
+				Obj->Modify();
+			}
+			PropertyHandle->NotifyPreChange();
+
+			ParameterBinding->AliasedParameter = BindAction->AliasedParameter;
+			ParameterBinding->Parameter = BindAction->Parameter;
+
+			PropertyHandle->NotifyPostChange(EPropertyChangeType::ValueSet);
+			PropertyHandle->NotifyFinishedChangingProperties();
+		}
+	}
+}
+
+FText FNiagaraParameterBindingCustomization::GetTooltipText() const
+{
+	if ( FNiagaraParameterBinding* ParameterBinding = GetParameterBinding() )
+	{
+		const FName BoundParameterName = ParameterBinding->AliasedParameter.GetName();
+		const FText BoundText =
+			BoundParameterName.IsNone() ?
+			LOCTEXT("BindingTooltip_UnboundParameter", "Not bound to any parameter.") :
+			FText::Format(LOCTEXT("BindingTooltip_BoundParameter", "Bound to variable \"{0}\""), FText::FromName(BoundParameterName));
+
+		const FText ValueText =
+			//ParameterBinding->Parameter.IsValid() && ParameterBinding->Parameter.GetAllocatedSizeInBytes() > 0 ?
+			//FText::Format(LOCTEXT("BindingTooltip_Value", "Default value \"{0}\""), FText::FromString(ParameterBinding->Parameter.ToString())) :
+			LOCTEXT("BindingTooltip_NoValue", "No default value");
+
+		return FText::Format(LOCTEXT("BindingTooltip", "{0}\n{1}"), BoundText, ValueText);
+	}
+	return FText::FromString(TEXT("Missing"));
+}
+
+FName FNiagaraParameterBindingCustomization::GetVariableName() const
+{
+	FNiagaraParameterBinding* ParameterBinding = GetParameterBinding();
+	return ParameterBinding ? ParameterBinding->AliasedParameter.GetName() : FName();
+}
+
+void FNiagaraParameterBindingCustomization::ResetToDefault()
+{
+
+}
+
+#undef LOCTEXT_NAMESPACE
diff --git a/Engine/Plugins/FX/Niagara/Source/NiagaraEditor/Private/Customizations/NiagaraParameterBindingCustomization.h b/Engine/Plugins/FX/Niagara/Source/NiagaraEditor/Private/Customizations/NiagaraParameterBindingCustomization.h
new file mode 100644
--- /dev/null
+++ b/Engine/Plugins/FX/Niagara/Source/NiagaraEditor/Private/Customizations/NiagaraParameterBindingCustomization.h
@@ -0,0 +1,65 @@
+// Copyright Epic Games, Inc. All Rights Reserved.
+
+#pragma once
+
+#include "EdGraph/EdGraphSchema.h"
+#include "CoreMinimal.h"
+#include "IPropertyTypeCustomization.h"
+
+#include "NiagaraCommon.h"
+#include "NiagaraParameterBindingCustomization.generated.h"
+
+class FDetailWidgetRow;
+class IDetailChildrenBuilder;
+struct FEdGraphSchemaAction;
+struct FGraphActionListBuilderBase;
+class IPropetyHandle;
+class IPropertyTypeCustomizationUtils;
+class SWidget;
+
+class UNiagaraEmitter;
+class UNiagaraSystem;
+struct FNiagaraParameterBinding;
+
+USTRUCT()
+struct FNiagaraParameterBindingAction : public FEdGraphSchemaAction
+{
+	GENERATED_USTRUCT_BODY();
+
+	FNiagaraParameterBindingAction();
+	FNiagaraParameterBindingAction(FNiagaraVariableBase InAliasedParameter, FNiagaraVariableBase InParameter);
+
+	FNiagaraVariableBase AliasedParameter;
+	FNiagaraVariableBase Parameter;
+};
+
+class FNiagaraParameterBindingCustomization : public IPropertyTypeCustomization
+{
+public:
+	static TSharedRef<IPropertyTypeCustomization> MakeInstance();
+
+	/** IPropertyTypeCustomization interface begin */
+	virtual void CustomizeHeader(TSharedRef<IPropertyHandle> StructPropertyHandle, FDetailWidgetRow& HeaderRow, IPropertyTypeCustomizationUtils& StructCustomizationUtils) override;
+	virtual void CustomizeChildren(TSharedRef<IPropertyHandle> StructPropertyHandle, IDetailChildrenBuilder& ChildBuilder, IPropertyTypeCustomizationUtils& StructCustomizationUtils) override {}
+	/** IPropertyTypeCustomization interface end */
+
+	bool IsValid() const;
+	FNiagaraParameterBinding* GetParameterBinding() const;
+	FNiagaraParameterBinding* GetDefaultParameterBinding() const;
+
+	TSharedRef<SWidget> OnGetMenuContent() const;
+	void CollectAllActions(FGraphActionListBuilderBase& OutAllActions) const;
+	TSharedRef<SWidget> OnCreateWidgetForAction(struct FCreateWidgetForActionData* const InCreateData) const;
+	void OnActionSelected(const TArray<TSharedPtr<FEdGraphSchemaAction>>& SelectedActions, ESelectInfo::Type InSelectionType) const;
+
+	FText GetTooltipText() const;
+	FName GetVariableName() const;
+
+	void ResetToDefault();
+
+protected:
+	TSharedPtr<IPropertyHandle>		PropertyHandle;
+	TWeakObjectPtr<UObject>			OwnerWeakPtr;
+	TWeakObjectPtr<UNiagaraEmitter>	EmitterWeakPtr;
+	TWeakObjectPtr<UNiagaraSystem>	SystemWeakPtr;
+};
diff --git a/Engine/Plugins/FX/Niagara/Source/NiagaraEditor/Private/NiagaraEditorModule.cpp b/Engine/Plugins/FX/Niagara/Source/NiagaraEditor/Private/NiagaraEditorModule.cpp
--- a/Engine/Plugins/FX/Niagara/Source/NiagaraEditor/Private/NiagaraEditorModule.cpp
+++ b/Engine/Plugins/FX/Niagara/Source/NiagaraEditor/Private/NiagaraEditorModule.cpp
@@ -96,6 +96,7 @@
 #include "Customizations/NiagaraComponentDetails.h"
 #include "Customizations/NiagaraFunctionCallNodeDetails.h"
 #include "Customizations/NiagaraEventScriptPropertiesCustomization.h"
+#include "Customizations/NiagaraParameterBindingCustomization.h"
 #include "Customizations/NiagaraPlatformSetCustomization.h"
 #include "Customizations/NiagaraScriptVariableCustomization.h"
 #include "Customizations/NiagaraScriptDetails.h"
@@ -942,7 +943,11 @@ void FNiagaraEditorModule::StartupModule()
 	PropertyModule.RegisterCustomPropertyTypeLayout(
 		FNiagaraVariableAttributeBinding::StaticStruct()->GetFName(),
 		FOnGetPropertyTypeCustomizationInstance::CreateStatic(&FNiagaraVariableAttributeBindingCustomization::MakeInstance));
-	
+
+	PropertyModule.RegisterCustomPropertyTypeLayout(
+		FNiagaraParameterBinding::StaticStruct()->GetFName(),
+		FOnGetPropertyTypeCustomizationInstance::CreateStatic(&FNiagaraParameterBindingCustomization::MakeInstance));
+
 	PropertyModule.RegisterCustomPropertyTypeLayout(
 		FNiagaraScriptVariableBinding::StaticStruct()->GetFName(),
 		FOnGetPropertyTypeCustomizationInstance::CreateStatic(&FNiagaraScriptVariableBindingCustomization::MakeInstance));
