- Protect against DI Equals of different types

#rnx

diff --git a/Engine/Plugins/FX/Niagara/Source/Niagara/Private/DataInterface/NiagaraDataInterfaceUtilities.cpp b/Engine/Plugins/FX/Niagara/Source/Niagara/Private/DataInterface/NiagaraDataInterfaceUtilities.cpp
--- a/Engine/Plugins/FX/Niagara/Source/Niagara/Private/DataInterface/NiagaraDataInterfaceUtilities.cpp
+++ b/Engine/Plugins/FX/Niagara/Source/Niagara/Private/DataInterface/NiagaraDataInterfaceUtilities.cpp
@@ -12,6 +12,23 @@
 namespace FNiagaraDataInterfaceUtilities
 {
 
+bool DataInterfaceEqual(const UNiagaraDataInterface* Lhs, const UNiagaraDataInterface* Rhs)
+{
+	if ( !Lhs || !Rhs )
+	{
+		return false;
+	}
+	if ( Lhs == Rhs )
+	{
+		return true;
+	}
+	if ( Lhs->GetClass() != Rhs->GetClass() )
+	{
+		return false;
+	}
+	return Lhs->Equals(Rhs);
+}
+
 void ForEachVMFunctionEqualsImpl(class UNiagaraDataInterface* DataInterface, class UNiagaraSystem* NiagaraSystem, const FNiagaraParameterStore& ParameterStore, TFunction<bool(const FVMExternalFunctionBindingInfo&)> Action)
 {
 	if ( DataInterface == nullptr || NiagaraSystem == nullptr )
@@ -24,7 +41,7 @@ void ForEachVMFunctionEqualsImpl(class UNiagaraDataInterface* DataInterface, cla
 	{
 		for (const UNiagaraDataInterface* OverrideDI : ParameterStore.GetDataInterfaces())
 		{
-			if ((OverrideDI != nullptr) && ((OverrideDI == DataInterface) || OverrideDI->Equals(DataInterface)))
+			if (OverrideDI && DataInterfaceEqual(OverrideDI, DataInterface))
 			{
 				if (const FNiagaraVariableBase* Variable = ParameterStore.FindVariable(OverrideDI))
 				{
@@ -63,15 +80,9 @@ void ForEachVMFunctionEqualsImpl(class UNiagaraDataInterface* DataInterface, cla
 						continue;
 					}
 
-					const FNiagaraScriptDataInterfaceInfo& CachedDefaultDI = CachedDefaultDIs[iDataInterface];
-					if (CachedDefaultDI.DataInterface == nullptr || !DataInterfaceInfo.MatchesClass(DataInterface->GetClass()))
-					{
-						// Would be odd not to match here, but we are being safe
-						break;
-					}
-
 					// Do we have a match?
-					if (CachedDefaultDI.DataInterface->Equals(DataInterface) || OverrideParameterNames.Contains(CachedDefaultDI.Name))
+					const FNiagaraScriptDataInterfaceInfo& CachedDefaultDI = CachedDefaultDIs[iDataInterface];
+					if (CachedDefaultDI.DataInterface && (DataInterfaceEqual(CachedDefaultDI.DataInterface, DataInterface) || OverrideParameterNames.Contains(CachedDefaultDI.Name)))
 					{
 						if ( Action(FunctionBinding) == false )
 						{
@@ -98,7 +109,7 @@ void ForEachGpuFunctionEqualsImpl(class UNiagaraDataInterface* DataInterface, cl
 	{
 		for (const UNiagaraDataInterface* OverrideDI : ParameterStore.GetDataInterfaces())
 		{
-			if ((OverrideDI != nullptr) && ((OverrideDI == DataInterface) || OverrideDI->Equals(DataInterface)))
+			if (DataInterfaceEqual(OverrideDI, DataInterface))
 			{
 				if (const FNiagaraVariableBase* Variable = ParameterStore.FindVariable(OverrideDI))
 				{
@@ -127,11 +138,7 @@ void ForEachGpuFunctionEqualsImpl(class UNiagaraDataInterface* DataInterface, cl
 				for (int iDataInterface = 0; iDataInterface < NumDataInterfaces; ++iDataInterface)
 				{
 					const FNiagaraScriptDataInterfaceInfo& CachedDefaultDI = CachedDefaultDIs[iDataInterface];
-					if (CachedDefaultDI.DataInterface == nullptr)
-					{
-						continue;
-					}
-					if (CachedDefaultDI.DataInterface->Equals(DataInterface) || OverrideParameterNames.Contains(CachedDefaultDI.Name))
+					if (CachedDefaultDI.DataInterface && (DataInterfaceEqual(CachedDefaultDI.DataInterface, DataInterface) || OverrideParameterNames.Contains(CachedDefaultDI.Name)))
 					{
 						for (const FNiagaraDataInterfaceGeneratedFunction& GeneratedFunction : DataInterfaceParamInfos[iDataInterface].GeneratedFunctions)
 						{
