- Handle emitter & variable rename on mesh binding

diff --git a/Engine/Plugins/FX/Niagara/Source/Niagara/Private/NiagaraMeshRendererProperties.cpp b/Engine/Plugins/FX/Niagara/Source/Niagara/Private/NiagaraMeshRendererProperties.cpp
--- a/Engine/Plugins/FX/Niagara/Source/Niagara/Private/NiagaraMeshRendererProperties.cpp
+++ b/Engine/Plugins/FX/Niagara/Source/Niagara/Private/NiagaraMeshRendererProperties.cpp
@@ -151,7 +151,10 @@ namespace NiagaraMeshRendererPropertiesInternal
 				if (MeshProperties.MeshParameterBinding.Parameter.IsDataInterface())
 				{
 					OutInterface = Cast<INiagaraRenderableMeshInterface>(EmitterInstance->GetRendererBoundVariables().GetDataInterface(MeshProperties.MeshParameterBinding.Parameter));
-					return;
+					if (OutInterface != nullptr)
+					{
+						return;
+					}
 				}
 
 				UStaticMesh* BoundMesh = Cast<UStaticMesh>(EmitterInstance->GetRendererBoundVariables().GetUObject(MeshProperties.MeshParameterBinding.Parameter));
@@ -496,6 +499,13 @@ void UNiagaraMeshRendererProperties::UpdateSourceModeDerivates(ENiagaraRendererS
 		}
 
 		SetPreviousBindings(SrcEmitter, InSourceMode);
+
+#if WITH_EDITORONLY_DATA
+		for (FNiagaraMeshRendererMeshProperties& Mesh : Meshes)
+		{
+			Mesh.MeshParameterBinding.OnRenameEmitter(SrcEmitter.Emitter->GetUniqueEmitterName());
+		}
+#endif
 	}
 }
 
@@ -1071,6 +1081,10 @@ void UNiagaraMeshRendererProperties::RenameVariable(const FNiagaraVariableBase&
 	Super::RenameVariable(OldVariable, NewVariable, InEmitter);
 #if WITH_EDITORONLY_DATA
 	MaterialParameters.RenameVariable(OldVariable, NewVariable, InEmitter, GetCurrentSourceMode());
+	for (FNiagaraMeshRendererMeshProperties& Mesh : Meshes)
+	{
+		Mesh.MeshParameterBinding.OnRenameVariable(OldVariable, NewVariable, InEmitter.Emitter->GetUniqueEmitterName());
+	}
 #endif
 }
 
@@ -1079,6 +1093,10 @@ void UNiagaraMeshRendererProperties::RemoveVariable(const FNiagaraVariableBase&
 	Super::RemoveVariable(OldVariable, InEmitter);
 #if WITH_EDITORONLY_DATA
 	MaterialParameters.RemoveVariable(OldVariable, InEmitter, GetCurrentSourceMode());
+	for (FNiagaraMeshRendererMeshProperties& Mesh : Meshes)
+	{
+		Mesh.MeshParameterBinding.OnRemoveVariable(OldVariable, InEmitter.Emitter->GetUniqueEmitterName());
+	}
 #endif
 }
 
diff --git a/Engine/Plugins/FX/Niagara/Source/Niagara/Private/NiagaraParameterBinding.cpp b/Engine/Plugins/FX/Niagara/Source/Niagara/Private/NiagaraParameterBinding.cpp
--- a/Engine/Plugins/FX/Niagara/Source/Niagara/Private/NiagaraParameterBinding.cpp
+++ b/Engine/Plugins/FX/Niagara/Source/Niagara/Private/NiagaraParameterBinding.cpp
@@ -79,4 +79,32 @@ bool FNiagaraParameterBinding::CanBindTo(FNiagaraVariableBase InVariable, FNiaga
 
 	return false;
 }
+
+void FNiagaraParameterBinding::OnRenameEmitter(FStringView EmitterName)
+{
+	if (AliasedParameter.IsInNameSpace(FNiagaraConstants::EmitterNamespaceString))
+	{
+		Parameter = AliasedParameter;
+		Parameter.ReplaceRootNamespace(FNiagaraConstants::EmitterNamespaceString, EmitterName);
+	}
+}
+
+void FNiagaraParameterBinding::OnRenameVariable(const FNiagaraVariableBase& OldVariable, const FNiagaraVariableBase& NewVariable, FStringView EmitterName)
+{
+	if (AliasedParameter.GetName() == OldVariable.GetName() && AliasedParameter.GetType() == OldVariable.GetType() && Parameter.IsInNameSpace(EmitterName))
+	{
+		AliasedParameter = NewVariable;
+		Parameter = AliasedParameter;
+		Parameter.ReplaceRootNamespace(FNiagaraConstants::EmitterNamespaceString, EmitterName);
+	}
+}
+
+void FNiagaraParameterBinding::OnRemoveVariable(const FNiagaraVariableBase& OldVariable, FStringView EmitterName)
+{
+	if (AliasedParameter.GetName() == OldVariable.GetName() && AliasedParameter.GetType() == OldVariable.GetType() && Parameter.IsInNameSpace(EmitterName))
+	{
+		AliasedParameter = FNiagaraVariableBase();
+		Parameter = FNiagaraVariableBase();
+	}
+}
 #endif
diff --git a/Engine/Plugins/FX/Niagara/Source/Niagara/Public/NiagaraParameterBinding.h b/Engine/Plugins/FX/Niagara/Source/Niagara/Public/NiagaraParameterBinding.h
--- a/Engine/Plugins/FX/Niagara/Source/Niagara/Public/NiagaraParameterBinding.h
+++ b/Engine/Plugins/FX/Niagara/Source/Niagara/Public/NiagaraParameterBinding.h
@@ -32,6 +32,10 @@ struct NIAGARA_API FNiagaraParameterBinding
 	bool CanBindTo(FNiagaraTypeDefinition TypeDefinition) const;
 	bool CanBindTo(FNiagaraVariableBase InVariable, FNiagaraVariableBase& OutAliasedVariable, FStringView EmitterName) const;
 
+	void OnRenameEmitter(FStringView EmitterName);
+	void OnRenameVariable(const FNiagaraVariableBase& OldVariable, const FNiagaraVariableBase& NewVariable, FStringView EmitterName);
+	void OnRemoveVariable(const FNiagaraVariableBase& OldVariable, FStringView EmitterName);
+
 	void SetUsage(ENiagaraParameterBindingUsage InUsage) { BindingUsage = InUsage; }
 	void SetAllowedDataInterfaces(TArray<UClass*> InClasses) { AllowedDataInterfaces = InClasses; }
 	void SetAllowedObjects(TArray<UClass*> InClasses) { AllowedObjects = InClasses; }
