- Static analysis fix

#rnx

diff --git a/Engine/Plugins/FX/Niagara/Source/NiagaraEditor/Private/Customizations/NiagaraParameterBindingCustomization.cpp b/Engine/Plugins/FX/Niagara/Source/NiagaraEditor/Private/Customizations/NiagaraParameterBindingCustomization.cpp
--- a/Engine/Plugins/FX/Niagara/Source/NiagaraEditor/Private/Customizations/NiagaraParameterBindingCustomization.cpp
+++ b/Engine/Plugins/FX/Niagara/Source/NiagaraEditor/Private/Customizations/NiagaraParameterBindingCustomization.cpp
@@ -200,12 +200,13 @@ void FNiagaraParameterBindingCustomization::CollectAllActions(FGraphActionListBu
 
 	TArray<TPair<FNiagaraVariableBase, FNiagaraVariableBase>> AllowedVariables;
 
+	const FStringView NiagaraEmitterName = NiagaraEmitter ? NiagaraEmitter->GetUniqueEmitterName() : FStringView();
 	for (const FNiagaraParameterMapHistory& History : Histories)
 	{
 		for (const FNiagaraVariable& Variable : History.Variables)
 		{
 			FNiagaraVariableBase AliasedParameter = Variable;
-			if (ParameterBinding->CanBindTo(Variable, AliasedParameter, NiagaraEmitter->GetUniqueEmitterName()))
+			if (ParameterBinding->CanBindTo(Variable, AliasedParameter, NiagaraEmitterName))
 			{
 				AllowedVariables.AddUnique({ AliasedParameter, Variable });
 				continue;
