Properly search for a parameter collection override on the system.

diff --git a/Engine/Plugins/FX/Niagara/Source/Niagara/Private/NiagaraSystemSimulation.cpp b/Engine/Plugins/FX/Niagara/Source/Niagara/Private/NiagaraSystemSimulation.cpp
--- a/Engine/Plugins/FX/Niagara/Source/Niagara/Private/NiagaraSystemSimulation.cpp
+++ b/Engine/Plugins/FX/Niagara/Source/Niagara/Private/NiagaraSystemSimulation.cpp
@@ -789,7 +789,7 @@ UNiagaraParameterCollectionInstance* FNiagaraSystemSimulation::GetParameterColle
 
 	if (System)
 	{
-		System->GetParameterCollectionOverride(Collection);
+		Ret = System->GetParameterCollectionOverride(Collection);
 	}
 
 	//If no explicit override from the system, just get the current instance set on the world.
