Fixed the stack input cache causing a crash in case we save a script without having done any changes, since this wouldn't update the cache key (same change ID) but we'd still create a new duplicate upon saving

diff --git a/Engine/Plugins/FX/Niagara/Source/NiagaraEditor/Private/Toolkits/NiagaraScriptToolkit.cpp b/Engine/Plugins/FX/Niagara/Source/NiagaraEditor/Private/Toolkits/NiagaraScriptToolkit.cpp
--- a/Engine/Plugins/FX/Niagara/Source/NiagaraEditor/Private/Toolkits/NiagaraScriptToolkit.cpp
+++ b/Engine/Plugins/FX/Niagara/Source/NiagaraEditor/Private/Toolkits/NiagaraScriptToolkit.cpp
@@ -852,7 +852,10 @@ void FNiagaraScriptToolkit::SaveAsset_Execute()
 {
 	UE_LOG(LogNiagaraEditor, Log, TEXT("Saving and Compiling NiagaraScript %s"), *GetEditingObjects()[0]->GetName());
 
-	UpdateOriginalNiagaraScript();
+	if(IsEditScriptDifferentFromOriginalScript())
+	{
+		UpdateOriginalNiagaraScript();
+	}
 
 	FAssetEditorToolkit::SaveAsset_Execute();
 }
@@ -861,7 +864,10 @@ void FNiagaraScriptToolkit::SaveAssetAs_Execute()
 {
 	UE_LOG(LogNiagaraEditor, Log, TEXT("Saving and Compiling NiagaraScript %s"), *GetEditingObjects()[0]->GetName());
 
-	UpdateOriginalNiagaraScript();
+	if(IsEditScriptDifferentFromOriginalScript())
+	{
+		UpdateOriginalNiagaraScript();
+	}
 
 	FAssetEditorToolkit::SaveAssetAs_Execute();
 }
