- Add expected GPU data ready stage to prepare functions
- Fixes an issues with raytracing and immediate translucency rendering as raytracing build runs before PostRenderOpaque simulations have run

diff --git a/Engine/Plugins/FX/Niagara/Source/Niagara/Private/NiagaraRendererMeshes.cpp b/Engine/Plugins/FX/Niagara/Source/Niagara/Private/NiagaraRendererMeshes.cpp
--- a/Engine/Plugins/FX/Niagara/Source/Niagara/Private/NiagaraRendererMeshes.cpp
+++ b/Engine/Plugins/FX/Niagara/Source/Niagara/Private/NiagaraRendererMeshes.cpp
@@ -234,7 +234,7 @@ void FNiagaraRendererMeshes::ReleaseRenderThreadResources()
 {
 }
 
-void FNiagaraRendererMeshes::PrepareParticleMeshRenderData(FParticleMeshRenderData& ParticleMeshRenderData, const FSceneViewFamily& ViewFamily, FMeshElementCollector& Collector, FNiagaraDynamicDataBase* InDynamicData, const FNiagaraSceneProxy* SceneProxy, bool bRayTracing) const
+void FNiagaraRendererMeshes::PrepareParticleMeshRenderData(FParticleMeshRenderData& ParticleMeshRenderData, const FSceneViewFamily& ViewFamily, FMeshElementCollector& Collector, FNiagaraDynamicDataBase* InDynamicData, const FNiagaraSceneProxy* SceneProxy, bool bRayTracing, ENiagaraGpuComputeTickStage::Type GpuReadyTickStage) const
 {
 	ParticleMeshRenderData.Collector = &Collector;
 
@@ -265,6 +265,7 @@ void FNiagaraRendererMeshes::PrepareParticleMeshRenderData(FParticleMeshRenderDa
 	const bool bIsWireframe = AllowDebugViewmodes() && ViewFamily.EngineShowFlags.Wireframe;
 	ParticleMeshRenderData.bIsGpuLowLatencyTranslucency =
 		bGpuLowLatencyTranslucency &&
+		GpuReadyTickStage >= CurrentParticleData->GetGPUDataReadyStage() &&
 		!bIsWireframe &&
 		!SceneProxy->CastsVolumetricTranslucentShadow() &&
 		ParticleMeshRenderData.DynamicDataMesh->Materials.Num() > 0 &&
@@ -1170,8 +1171,9 @@ void FNiagaraRendererMeshes::GetDynamicMeshElements(const TArray<const FSceneVie
 
 	// Prepare our particle render data
 	// This will also determine if we have anything to render
+	// ENiagaraGpuComputeTickStage::Last is used as the GPU ready stage as we can support reading translucent data after PostRenderOpaque sims have run
 	FParticleMeshRenderData ParticleMeshRenderData;
-	PrepareParticleMeshRenderData(ParticleMeshRenderData, ViewFamily, Collector, DynamicDataRender, SceneProxy, false);
+	PrepareParticleMeshRenderData(ParticleMeshRenderData, ViewFamily, Collector, DynamicDataRender, SceneProxy, false, ENiagaraGpuComputeTickStage::Last);
 
 	if (ParticleMeshRenderData.SourceParticleData == nullptr )
 	{
@@ -1350,12 +1352,14 @@ void FNiagaraRendererMeshes::GetDynamicRayTracingInstances(FRayTracingMaterialGa
 	const FSceneView* View = Context.ReferenceView;
 	const bool bIsInstancedStereo = View->bIsInstancedStereoEnabled && IStereoRendering::IsStereoEyeView(*View);
 
+
 	check(View->Family);
 
 	// Prepare our particle render data
 	// This will also determine if we have anything to render
+	// ENiagaraGpuComputeTickStage::PostInitViews is used as we need the data one InitViews is complete as the HWRT BVH will be generated before other sims have run
 	FParticleMeshRenderData ParticleMeshRenderData;
-	PrepareParticleMeshRenderData(ParticleMeshRenderData, *View->Family, Context.RayTracingMeshResourceCollector, DynamicDataRender, SceneProxy, true);
+	PrepareParticleMeshRenderData(ParticleMeshRenderData, *View->Family, Context.RayTracingMeshResourceCollector, DynamicDataRender, SceneProxy, true, ENiagaraGpuComputeTickStage::PostInitViews);
 
 	if (ParticleMeshRenderData.SourceParticleData == nullptr || Meshes.Num() == 0)
 	{
diff --git a/Engine/Plugins/FX/Niagara/Source/Niagara/Private/NiagaraRendererSprites.cpp b/Engine/Plugins/FX/Niagara/Source/Niagara/Private/NiagaraRendererSprites.cpp
--- a/Engine/Plugins/FX/Niagara/Source/Niagara/Private/NiagaraRendererSprites.cpp
+++ b/Engine/Plugins/FX/Niagara/Source/Niagara/Private/NiagaraRendererSprites.cpp
@@ -210,7 +210,7 @@ void FNiagaraRendererSprites::CreateRenderThreadResources()
 #endif
 }
 
-void FNiagaraRendererSprites::PrepareParticleSpriteRenderData(FParticleSpriteRenderData& ParticleSpriteRenderData, const FSceneViewFamily& ViewFamily, FNiagaraDynamicDataBase* InDynamicData, const FNiagaraSceneProxy* SceneProxy) const
+void FNiagaraRendererSprites::PrepareParticleSpriteRenderData(FParticleSpriteRenderData& ParticleSpriteRenderData, const FSceneViewFamily& ViewFamily, FNiagaraDynamicDataBase* InDynamicData, const FNiagaraSceneProxy* SceneProxy, ENiagaraGpuComputeTickStage::Type GpuReadyTickStage) const
 {
 	ParticleSpriteRenderData.DynamicDataSprites = static_cast<FNiagaraDynamicDataSprites*>(InDynamicData);
 	if (!ParticleSpriteRenderData.DynamicDataSprites || !SceneProxy->GetComputeDispatchInterface())
@@ -236,8 +236,10 @@ void FNiagaraRendererSprites::PrepareParticleSpriteRenderData(FParticleSpriteRen
 	ParticleSpriteRenderData.bHasTranslucentMaterials = IsTranslucentBlendMode(Material);
 
 	// If these conditions change please update the DebugHUD display also to reflect it
-	const bool bLowLatencyTranslucencyEnabled = ParticleSpriteRenderData.bHasTranslucentMaterials &&
+	const bool bLowLatencyTranslucencyEnabled =
+		ParticleSpriteRenderData.bHasTranslucentMaterials &&
 		bGpuLowLatencyTranslucency &&
+		GpuReadyTickStage >= CurrentParticleData->GetGPUDataReadyStage() &&
 		!SceneProxy->CastsVolumetricTranslucentShadow() &&
 		ViewFamilySupportLowLatencyTranslucency(ViewFamily);
 
@@ -932,8 +934,9 @@ void FNiagaraRendererSprites::GetDynamicMeshElements(const TArray<const FSceneVi
 
 	// Prepare our particle render data
 	// This will also determine if we have anything to render
+	// ENiagaraGpuComputeTickStage::Last is used as the GPU ready stage as we can support reading translucent data after PostRenderOpaque sims have run
 	FParticleSpriteRenderData ParticleSpriteRenderData;
-	PrepareParticleSpriteRenderData(ParticleSpriteRenderData, ViewFamily, DynamicDataRender, SceneProxy);
+	PrepareParticleSpriteRenderData(ParticleSpriteRenderData, ViewFamily, DynamicDataRender, SceneProxy, ENiagaraGpuComputeTickStage::Last);
 
 	if (ParticleSpriteRenderData.SourceParticleData == nullptr)
 	{
@@ -1045,8 +1048,9 @@ void FNiagaraRendererSprites::GetDynamicRayTracingInstances(FRayTracingMaterialG
 
 	// Prepare our particle render data
 	// This will also determine if we have anything to render
+	// ENiagaraGpuComputeTickStage::PostInitViews is used as we need the data one InitViews is complete as the HWRT BVH will be generated before other sims have run
 	FParticleSpriteRenderData ParticleSpriteRenderData;
-	PrepareParticleSpriteRenderData(ParticleSpriteRenderData, *Context.ReferenceView->Family, DynamicDataRender, SceneProxy);
+	PrepareParticleSpriteRenderData(ParticleSpriteRenderData, *Context.ReferenceView->Family, DynamicDataRender, SceneProxy, ENiagaraGpuComputeTickStage::PostInitViews);
 
 	if (ParticleSpriteRenderData.SourceParticleData == nullptr)
 	{
diff --git a/Engine/Plugins/FX/Niagara/Source/Niagara/Public/NiagaraRendererMeshes.h b/Engine/Plugins/FX/Niagara/Source/Niagara/Public/NiagaraRendererMeshes.h
--- a/Engine/Plugins/FX/Niagara/Source/Niagara/Public/NiagaraRendererMeshes.h
+++ b/Engine/Plugins/FX/Niagara/Source/Niagara/Public/NiagaraRendererMeshes.h
@@ -116,7 +116,7 @@ protected:
 		virtual ~FGPUSceneUpdateResource() {}
 	};
 
-	void PrepareParticleMeshRenderData(FParticleMeshRenderData& ParticleMeshRenderData, const FSceneViewFamily& ViewFamily, FMeshElementCollector& Collector, FNiagaraDynamicDataBase* InDynamicData, const FNiagaraSceneProxy* SceneProxy, bool bRayTracing) const;
+	void PrepareParticleMeshRenderData(FParticleMeshRenderData& ParticleMeshRenderData, const FSceneViewFamily& ViewFamily, FMeshElementCollector& Collector, FNiagaraDynamicDataBase* InDynamicData, const FNiagaraSceneProxy* SceneProxy, bool bRayTracing, ENiagaraGpuComputeTickStage::Type GpuReadyTickStage) const;
 	void PrepareParticleRenderBuffers(FParticleMeshRenderData& ParticleMeshRenderData, FGlobalDynamicReadBuffer& DynamicReadBuffer) const;
 	void InitializeSortInfo(const FParticleMeshRenderData& ParticleMeshRenderData, const FNiagaraSceneProxy& SceneProxy, const FSceneView& View, int32 ViewIndex, bool bIsInstancedStereo, FNiagaraGPUSortInfo& OutSortInfo) const;
 	void PreparePerMeshData(FParticleMeshRenderData& ParticleMeshRenderData, const FNiagaraMeshVertexFactory& VertexFactory, const FNiagaraSceneProxy& SceneProxy, const FMeshData& MeshData) const;
diff --git a/Engine/Plugins/FX/Niagara/Source/Niagara/Public/NiagaraRendererSprites.h b/Engine/Plugins/FX/Niagara/Source/Niagara/Public/NiagaraRendererSprites.h
--- a/Engine/Plugins/FX/Niagara/Source/Niagara/Public/NiagaraRendererSprites.h
+++ b/Engine/Plugins/FX/Niagara/Source/Niagara/Public/NiagaraRendererSprites.h
@@ -72,7 +72,7 @@ private:
 		FNiagaraSpriteUniformBufferRef	UniformBuffer;
 	};
 
-	void PrepareParticleSpriteRenderData(FParticleSpriteRenderData& ParticleSpriteRenderData, const FSceneViewFamily& ViewFamily, FNiagaraDynamicDataBase* InDynamicData, const FNiagaraSceneProxy* SceneProxy) const;
+	void PrepareParticleSpriteRenderData(FParticleSpriteRenderData& ParticleSpriteRenderData, const FSceneViewFamily& ViewFamily, FNiagaraDynamicDataBase* InDynamicData, const FNiagaraSceneProxy* SceneProxy, ENiagaraGpuComputeTickStage::Type GpuReadyTickStage) const;
 	void PrepareParticleRenderBuffers(FParticleSpriteRenderData& ParticleSpriteRenderData, FGlobalDynamicReadBuffer& DynamicReadBuffer) const;
 	void InitializeSortInfo(FParticleSpriteRenderData& ParticleSpriteRenderData, const FNiagaraSceneProxy& SceneProxy, const FSceneView& View, int32 ViewIndex, FNiagaraGPUSortInfo& OutSortInfo) const;
 	void SetupVertexFactory(FParticleSpriteRenderData& ParticleSpriteRenderData, FNiagaraSpriteVertexFactory& VertexFactory) const;
