- Ensure Niagara Variant LWC arrays have the LWC type data set
- This fixes an issue where construction scripts would set the data but it was lost

#okforightub public

diff --git a/Engine/Plugins/FX/Niagara/Source/Niagara/Classes/NiagaraDataInterfaceArrayImpl.h b/Engine/Plugins/FX/Niagara/Source/Niagara/Classes/NiagaraDataInterfaceArrayImpl.h
--- a/Engine/Plugins/FX/Niagara/Source/Niagara/Classes/NiagaraDataInterfaceArrayImpl.h
+++ b/Engine/Plugins/FX/Niagara/Source/Niagara/Classes/NiagaraDataInterfaceArrayImpl.h
@@ -26,6 +26,19 @@
 		Proxy.Reset(new FProxyType(this)); \
 		Super::PostInitProperties(); \
 	} \
+	template<typename TFromArrayType> \
+	void SetVariantArrayData(TConstArrayView<TFromArrayType> InArrayData) \
+	{ \
+		MEMBERNAME = InArrayData; \
+	} \
+	template<typename TFromArrayType> \
+	void SetVariantArrayValue(int Index, const TFromArrayType& Value, bool bSizeToFit) \
+	{ \
+		const int NumRequired = Index + 1 - MEMBERNAME.Num(); \
+		if ( NumRequired > 0 && !bSizeToFit ) return; \
+		MEMBERNAME.AddDefaulted(FMath::Max(NumRequired, 0)); \
+		MEMBERNAME[Index] = Value; \
+	} \
 	TArray<TYPENAME>& GetArrayReference() { return MEMBERNAME; }
 
 #if WITH_EDITORONLY_DATA
@@ -67,6 +80,37 @@
 				TypedOther != nullptr && \
 				TypedOther->MEMBERNAME == MEMBERNAME; \
 		} \
+		template<typename TFromArrayType> \
+		typename TEnableIf<TIsSame<TFromArrayType, decltype(MEMBERNAME)::ElementType>::Value, void>::Type SetVariantArrayData(TConstArrayView<TFromArrayType> InArrayData) \
+		{ \
+			MEMBERNAME = InArrayData; \
+			GetProxyAs<FProxyType>()->template SetArrayData<decltype(MEMBERNAME)::ElementType>(InArrayData); \
+		} \
+		template<typename TFromArrayType> \
+		typename TEnableIf<!TIsSame<TFromArrayType, decltype(MEMBERNAME)::ElementType>::Value, void>::Type SetVariantArrayData(TConstArrayView<TFromArrayType> InArrayData) \
+		{ \
+			MEMBERNAME.SetNumUninitialized(InArrayData.Num()); \
+			FNDIArrayImplHelper<TYPENAME>::CopyCpuToCpuMemory(MEMBERNAME.GetData(), InArrayData.GetData(), InArrayData.Num()); \
+			GetProxyAs<FProxyType>()->template SetArrayData<decltype(Internal##MEMBERNAME)::ElementType>(InArrayData); \
+		} \
+		template<typename TFromArrayType> \
+		typename TEnableIf<TIsSame<TFromArrayType, decltype(MEMBERNAME)::ElementType>::Value, void>::Type SetVariantArrayValue(int Index, const TFromArrayType& Value, bool bSizeToFit) \
+		{ \
+			const int NumRequired = Index + 1 - MEMBERNAME.Num(); \
+			if ( NumRequired > 0 && !bSizeToFit ) return; \
+			MEMBERNAME.AddDefaulted(FMath::Max(NumRequired, 0)); \
+			MEMBERNAME[Index] = Value; \
+			GetProxyAs<FProxyType>()->template SetArrayData<decltype(MEMBERNAME)::ElementType>(MEMBERNAME); \
+		} \
+		template<typename TFromArrayType> \
+		typename TEnableIf<!TIsSame<TFromArrayType, decltype(MEMBERNAME)::ElementType>::Value, void>::Type SetVariantArrayValue(int Index, const TFromArrayType& Value, bool bSizeToFit) \
+		{ \
+			const int NumRequired = Index + 1 - MEMBERNAME.Num(); \
+			if ( NumRequired > 0 && !bSizeToFit ) return; \
+			MEMBERNAME.AddDefaulted(FMath::Max(NumRequired, 0)); \
+			MEMBERNAME[Index] = Value; \
+			GetProxyAs<FProxyType>()->template SetArrayData<decltype(MEMBERNAME)::ElementType>(MEMBERNAME); \
+		} \
 		TArray<TYPENAME>& GetArrayReference() { return Internal##MEMBERNAME; }
 #else
 	#define NDIARRAY_GENERATE_BODY_LWC(CLASSNAME, TYPENAME, MEMBERNAME) NDIARRAY_GENERATE_BODY(CLASSNAME, TYPENAME, Internal##MEMBERNAME)
diff --git a/Engine/Plugins/FX/Niagara/Source/Niagara/Private/NiagaraDataInterfaceArrayFunctionLibrary.cpp b/Engine/Plugins/FX/Niagara/Source/Niagara/Private/NiagaraDataInterfaceArrayFunctionLibrary.cpp
--- a/Engine/Plugins/FX/Niagara/Source/Niagara/Private/NiagaraDataInterfaceArrayFunctionLibrary.cpp
+++ b/Engine/Plugins/FX/Niagara/Source/Niagara/Private/NiagaraDataInterfaceArrayFunctionLibrary.cpp
@@ -13,27 +13,6 @@
 
 #include UE_INLINE_GENERATED_CPP_BY_NAME(NiagaraDataInterfaceArrayFunctionLibrary)
 
-template<typename TArrayType, typename TDataInterace>
-void SetNiagaraVariantArray(UNiagaraComponent* NiagaraComponent, FName OverrideName, TDataInterace* ExistingDataInterface)
-{
-#if WITH_EDITOR
-	// We only need to do this for editor instances of the component as we are storing instance data on them
-	// For runtime instances they already have a unique copy of the data interface that we are modifying
-	UWorld* World = NiagaraComponent->GetWorld();
-	if (World == nullptr || World->IsGameWorld())
-	{
-		return;
-	}
-
-	TDataInterace* VariantDataInterface = CastChecked<TDataInterace>(DuplicateObject(ExistingDataInterface, NiagaraComponent));
-	auto* ExistingProxy = static_cast<typename TDataInterace::FProxyType*>(ExistingDataInterface->GetProxy());
-	auto* VariantProxy = static_cast<typename TDataInterace::FProxyType*>(VariantDataInterface->GetProxy());
-	TArray<TArrayType> TempArrayCopy = ExistingProxy->template GetArrayDataCopy<TArrayType>();
-	VariantProxy->template SetArrayData<TArrayType>(MakeArrayView<TArrayType>(TempArrayCopy.GetData(), TempArrayCopy.Num()));
-	NiagaraComponent->SetParameterOverride(FNiagaraVariableBase(FNiagaraTypeDefinition(TDataInterace::StaticClass()), OverrideName), FNiagaraVariant(VariantDataInterface));
-#endif
-}
-
 // If / when we share user parameter UObjects we will need to make this per instance which introduces some tricky things about allocating before the instance is active
 template<typename TArrayType, typename TDataInterace>
 void SetNiagaraArray(UNiagaraComponent* NiagaraComponent, FName OverrideName, TConstArrayView<TArrayType> InArray)
@@ -42,7 +21,18 @@ void SetNiagaraArray(UNiagaraComponent* NiagaraComponent, FName OverrideName, TC
 	{
 		auto* ArrayProxy = static_cast<typename TDataInterace::FProxyType*>(ArrayDI->GetProxy());
 		ArrayProxy->SetArrayData(InArray);
-		SetNiagaraVariantArray<TArrayType, TDataInterace>(NiagaraComponent, OverrideName, ArrayDI);
+#if WITH_EDITOR
+		// We only need to do this for editor instances of the component as we are storing instance data on them
+		// For runtime instances they already have a unique copy of the data interface that we are modifying
+		UWorld* World = NiagaraComponent->GetWorld();
+		if (World && !World->IsGameWorld())
+		{
+			TDataInterace* VariantDI = CastChecked<TDataInterace>(DuplicateObject(ArrayDI, NiagaraComponent));
+			ArrayDI->CopyTo(VariantDI);
+			VariantDI->SetVariantArrayData(InArray);
+			NiagaraComponent->SetParameterOverride(FNiagaraVariableBase(FNiagaraTypeDefinition(TDataInterace::StaticClass()), OverrideName), FNiagaraVariant(VariantDI));
+		}
+#endif
 	}
 }
 
@@ -64,7 +54,20 @@ void SetNiagaraArrayValue(UNiagaraComponent* NiagaraComponent, FName OverrideNam
 	{
 		auto* ArrayProxy = static_cast<typename TDataInterace::FProxyType*>(ArrayDI->GetProxy());
 		ArrayProxy->SetArrayValue(Index, Value, bSizeToFit);
-		SetNiagaraVariantArray<TArrayType, TDataInterace>(NiagaraComponent, OverrideName, ArrayDI);
+#if WITH_EDITOR
+		// We only need to do this for editor instances of the component as we are storing instance data on them
+		// For runtime instances they already have a unique copy of the data interface that we are modifying
+		UWorld* World = NiagaraComponent->GetWorld();
+		if (World && !World->IsGameWorld())
+		{
+			TDataInterace* VariantDI = CastChecked<TDataInterace>(DuplicateObject(ArrayDI, NiagaraComponent));
+			ArrayDI->CopyTo(VariantDI);
+			VariantDI->SetVariantArrayValue(Index, Value, bSizeToFit);
+			NiagaraComponent->SetParameterOverride(FNiagaraVariableBase(FNiagaraTypeDefinition(TDataInterace::StaticClass()), OverrideName), FNiagaraVariant(VariantDI));
+		}
+#endif
+
+		//-TODO:FIXME:SetNiagaraVariantArray<TArrayType, TDataInterace>(NiagaraComponent, OverrideName, ArrayDI);
 	}
 }
 
