- Added missing include for NonPCH/NonUnity compile error
#okfogithub public

diff --git a/Engine/Plugins/FX/Niagara/Source/Niagara/Private/NiagaraDataInterfaceRigidMeshCollisionQuery.cpp b/Engine/Plugins/FX/Niagara/Source/Niagara/Private/NiagaraDataInterfaceRigidMeshCollisionQuery.cpp
--- a/Engine/Plugins/FX/Niagara/Source/Niagara/Private/NiagaraDataInterfaceRigidMeshCollisionQuery.cpp
+++ b/Engine/Plugins/FX/Niagara/Source/Niagara/Private/NiagaraDataInterfaceRigidMeshCollisionQuery.cpp
@@ -10,6 +10,7 @@
 #include "Components/StaticMeshComponent.h"
 #include "DistanceFieldLightingShared.h"
 #include "Engine/Canvas.h"
+#include "Engine/Engine.h"
 #include "EngineUtils.h"
 #include "SkeletalRenderPublic.h"
 
