- Fix for MICs on mesh renderer to correctly handle multiple meshes and multiple materials

#rnx

diff --git a/Engine/Plugins/FX/Niagara/Source/Niagara/Private/NiagaraMeshRendererProperties.cpp b/Engine/Plugins/FX/Niagara/Source/Niagara/Private/NiagaraMeshRendererProperties.cpp
--- a/Engine/Plugins/FX/Niagara/Source/Niagara/Private/NiagaraMeshRendererProperties.cpp
+++ b/Engine/Plugins/FX/Niagara/Source/Niagara/Private/NiagaraMeshRendererProperties.cpp
@@ -395,44 +395,7 @@ void UNiagaraMeshRendererProperties::Serialize(FArchive& Ar)
 		SortMode = ENiagaraSortMode::ViewDistance;
 	}
 
-	// MICs will replace the main material during serialize
-	// Be careful if adding code that looks at the material to make sure you get the correct one
-	{
-	#if WITH_EDITORONLY_DATA
-		// For cooked builds we put our MIC material directly into the OverrideMaterial slot to avoid wasting memory
-		TOptional<bool> OverrideMaterialsGuard;
-		TArray<FNiagaraMeshMaterialOverride> OverrideMaterialsArrayGuard;
-		if (Ar.IsSaving() && Ar.IsCooking() && MICMaterials.Num() > 0)
-		{
-			OverrideMaterialsGuard = bOverrideMaterials;
-			OverrideMaterialsArrayGuard = OverrideMaterials;
-
-			for (int i=0; i < OverrideMaterials.Num(); ++i)
-			{
-				FNiagaraMeshMaterialOverride& MaterialOverride = OverrideMaterials[i];
-				if (MICMaterials.IsValidIndex(i) && MICMaterials[i])
-				{
-					MaterialOverride.ExplicitMat = MICMaterials[i];
-				}
-				else if (bOverrideMaterials == false)
-				{
-					MaterialOverride.ExplicitMat = nullptr;
-				}
-			}
-			bOverrideMaterials = true;
-		}
-	#endif
-
-		Super::Serialize(Ar);
-
-#if WITH_EDITORONLY_DATA
-		if (OverrideMaterialsGuard.IsSet())
-		{
-			bOverrideMaterials = OverrideMaterialsGuard.GetValue();
-			OverrideMaterials = OverrideMaterialsArrayGuard;
-		}
-#endif
-	}
+	Super::Serialize(Ar);
 }
 
 /** The bindings depend on variables that are created during the NiagaraModule startup. However, the CDO's are build prior to this being initialized, so we defer setting these values until later.*/
@@ -569,15 +532,36 @@ void UNiagaraMeshRendererProperties::CacheFromCompiledData(const FNiagaraDataSet
 void UNiagaraMeshRendererProperties::UpdateMICs()
 {
 #if WITH_EDITORONLY_DATA
-	// Note: We need to swap out MICMaterials as GetUsedMaterials uses them to override the material
-	TArray<TObjectPtr<UMaterialInstanceConstant>> TempMICMaterials;
-	Swap(TempMICMaterials, MICMaterials);
+	// Grab existing MICs so we can reuse and clear them out so they aren't applied during GetUsedMaterials
+	TArray<TObjectPtr<UMaterialInstanceConstant>> MICMaterials;
+	MICMaterials.Reserve(MICOverrideMaterials.Num());
+	for (const FNiagaraMeshMICOverride& ExistingOverride : MICOverrideMaterials)
+	{
+		MICMaterials.Add(ExistingOverride.ReplacementMaterial);
+	}
+	MICOverrideMaterials.Reset(0);
 
+	// Gather materials and generate MICs
 	TArray<UMaterialInterface*> Materials;
 	GetUsedMaterials(nullptr, Materials);
-	UpdateMaterialParametersMIC(MaterialParameters, Materials, TempMICMaterials);
 
-	Swap(TempMICMaterials, MICMaterials);
+	UpdateMaterialParametersMIC(MaterialParameters, Materials, MICMaterials);
+
+	// Create Material <-> MIC remap
+	for (int i=0; i < MICMaterials.Num(); ++i)
+	{
+		const FNiagaraMeshMICOverride* ExistingOverride = MICOverrideMaterials.FindByPredicate([FindMaterial = Materials[i]](const FNiagaraMeshMICOverride& ExistingOverride) { return ExistingOverride.OriginalMaterial == FindMaterial; });
+		if (ExistingOverride)
+		{
+			ensureMsgf(ExistingOverride->ReplacementMaterial == MICMaterials[i], TEXT("MIC Material should match replacement material, static bindings will be incorrect.  Please report this issue."));
+		}
+		else
+		{
+			FNiagaraMeshMICOverride& NewOverride = MICOverrideMaterials.AddDefaulted_GetRef();
+			NewOverride.OriginalMaterial = Materials[i];
+			NewOverride.ReplacementMaterial = MICMaterials[i];
+		}
+	}
 #endif
 }
 
@@ -620,9 +604,6 @@ void UNiagaraMeshRendererProperties::ApplyMaterialOverrides(const FNiagaraEmitte
 			if (!OverrideMat)
 			{
 				OverrideMat = OverrideMaterials[OverrideIndex].ExplicitMat;
-			#if WITH_EDITORONLY_DATA
-				OverrideMat = MICMaterials.IsValidIndex(OverrideIndex) && MICMaterials[OverrideIndex] ? MICMaterials[OverrideIndex] : OverrideMat;
-			#endif
 			}
 
 			if (OverrideMat)
@@ -631,18 +612,17 @@ void UNiagaraMeshRendererProperties::ApplyMaterialOverrides(const FNiagaraEmitte
 			}
 		}
 	}
-	else
+
+	// Apply MIC override materials
+	if (MICOverrideMaterials.Num() > 0)
 	{
-	#if WITH_EDITORONLY_DATA
-		if (MICMaterials.Num() > 0)
+		for (UMaterialInterface*& Material : InOutMaterials)
 		{
-			const int MaterialMax = FMath::Min(InOutMaterials.Num(), MICMaterials.Num());
-			for (int i = 0; i < MaterialMax; ++i)
+			if (const FNiagaraMeshMICOverride* Override = MICOverrideMaterials.FindByPredicate([&Material](const FNiagaraMeshMICOverride& MICOverride) { return MICOverride.OriginalMaterial == Material; }))
 			{
-				InOutMaterials[i] = MICMaterials[i] ? MICMaterials[i] : InOutMaterials[i];
+				Material = Override->ReplacementMaterial;
 			}
 		}
-	#endif
 	}
 }
 
diff --git a/Engine/Plugins/FX/Niagara/Source/Niagara/Public/NiagaraMeshRendererProperties.h b/Engine/Plugins/FX/Niagara/Source/Niagara/Public/NiagaraMeshRendererProperties.h
--- a/Engine/Plugins/FX/Niagara/Source/Niagara/Public/NiagaraMeshRendererProperties.h
+++ b/Engine/Plugins/FX/Niagara/Source/Niagara/Public/NiagaraMeshRendererProperties.h
@@ -54,6 +54,18 @@ enum class ENiagaraMeshLockedAxisSpace : uint8
 	Local
 };
 
+USTRUCT()
+struct NIAGARA_API FNiagaraMeshMICOverride
+{
+	GENERATED_BODY()
+
+	UPROPERTY()
+	TObjectPtr<UMaterialInterface> OriginalMaterial;
+
+	UPROPERTY()
+	TObjectPtr<UMaterialInstanceConstant> ReplacementMaterial;
+};
+
 USTRUCT()
 struct NIAGARA_API FNiagaraMeshMaterialOverride
 {
@@ -278,10 +290,8 @@ public:
 	UPROPERTY(EditAnywhere, Category = "Mesh Rendering", meta = (EditCondition = "bOverrideMaterials", EditConditionHides))
 	TArray<FNiagaraMeshMaterialOverride> OverrideMaterials;
 
-#if WITH_EDITORONLY_DATA
-	UPROPERTY(transient)
-	TArray<TObjectPtr<UMaterialInstanceConstant>> MICMaterials;
-#endif
+	UPROPERTY()
+	TArray<FNiagaraMeshMICOverride> MICOverrideMaterials;
 
 	/** When using SubImage lookups for particles, this variable contains the number of columns in X and the number of rows in Y.*/
 	UPROPERTY(EditAnywhere, Category = "SubUV", meta = (EditCondition = "bSubImageBlend", EditConditionHides))
