- Niagara Decal Renderer

diff --git a/Engine/Plugins/FX/Niagara/Source/Niagara/Private/NiagaraDecalRendererProperties.cpp b/Engine/Plugins/FX/Niagara/Source/Niagara/Private/NiagaraDecalRendererProperties.cpp
new file mode 100644
--- /dev/null
+++ b/Engine/Plugins/FX/Niagara/Source/Niagara/Private/NiagaraDecalRendererProperties.cpp
@@ -0,0 +1,248 @@
+// Copyright Epic Games, Inc. All Rights Reserved.
+
+#include "NiagaraDecalRendererProperties.h"
+#include "NiagaraConstants.h"
+#include "NiagaraEmitterInstance.h"
+#include "NiagaraRenderer.h"
+#include "NiagaraRendererDecals.h"
+#include "Modules/ModuleManager.h"
+
+#include UE_INLINE_GENERATED_CPP_BY_NAME(NiagaraDecalRendererProperties)
+
+#if WITH_EDITOR
+#include "Widgets/Images/SImage.h"
+#include "Widgets/Text/STextBlock.h"
+#include "Widgets/SWidget.h"
+#include "Styling/SlateBrush.h"
+#include "Styling/SlateIconFinder.h"
+#include "AssetThumbnail.h"
+#endif
+
+#define LOCTEXT_NAMESPACE "UNiagaraDecalRendererProperties"
+
+namespace NiagaraDecalRendererPropertiesLocal
+{
+	TArray<TWeakObjectPtr<UNiagaraDecalRendererProperties>> RendererPropertiesToDeferredInit;
+
+	template<typename TValueType>
+	static FNiagaraVariable MakeNiagaraVariableWithValue(const FNiagaraTypeDefinition& TypeDef, const FName& Name, const TValueType& Value)
+	{
+		FNiagaraVariable Variable(TypeDef, Name);
+		Variable.SetValue(Value);
+		return Variable;
+	}
+
+	static FNiagaraVariable GetDecalOrientationVariable()
+	{
+		static FNiagaraVariable Variable = MakeNiagaraVariableWithValue(FNiagaraTypeDefinition::GetQuatDef(), TEXT("Particles.DecalOrientation"), FRotator3f(-90.0f, 0.0f, 90.0f).Quaternion());
+		return Variable;
+	}
+
+	static FNiagaraVariable GetDecalSizeVariable()
+	{
+		static FNiagaraVariable Variable = MakeNiagaraVariableWithValue(FNiagaraTypeDefinition::GetVec3Def(), TEXT("Particles.DecalSize"), FVector3f(50.0f, 50.0f, 50.0f));
+		return Variable;
+	}
+
+	static FNiagaraVariable GetDecalFadeVariable()
+	{
+		FNiagaraVariable Variable = MakeNiagaraVariableWithValue(FNiagaraTypeDefinition::GetFloatDef(), TEXT("Particles.DecalFade"), 1.0f);
+		return Variable;
+	}
+
+	static void SetupBindings(UNiagaraDecalRendererProperties* Props)
+	{
+		if (Props->PositionBinding.IsValid())
+		{
+			return;
+		}
+		Props->PositionBinding = FNiagaraConstants::GetAttributeDefaultBinding(SYS_PARAM_PARTICLES_POSITION);
+		Props->DecalOrientationBinding.Setup(GetDecalOrientationVariable(), GetDecalOrientationVariable());
+		Props->DecalSizeBinding.Setup(GetDecalSizeVariable(), GetDecalSizeVariable());
+		Props->DecalFadeBinding.Setup(GetDecalFadeVariable(), GetDecalFadeVariable());
+		Props->ColorBinding = FNiagaraConstants::GetAttributeDefaultBinding(SYS_PARAM_PARTICLES_COLOR);
+		Props->RendererVisibilityTagBinding = FNiagaraConstants::GetAttributeDefaultBinding(SYS_PARAM_PARTICLES_VISIBILITY_TAG);
+
+	#if WITH_EDITORONLY_DATA
+		Props->MaterialParameterBinding.SetUsage(ENiagaraParameterBindingUsage::NotParticle);
+		Props->MaterialParameterBinding.SetAllowedObjects({ UMaterialInterface::StaticClass() });
+	#endif
+	}
+}
+
+UNiagaraDecalRendererProperties::UNiagaraDecalRendererProperties()
+{
+	AttributeBindings.Reserve(5);
+	AttributeBindings.Add(&PositionBinding);
+	AttributeBindings.Add(&DecalOrientationBinding);
+	AttributeBindings.Add(&DecalSizeBinding);
+	AttributeBindings.Add(&DecalFadeBinding);
+	AttributeBindings.Add(&ColorBinding);
+	AttributeBindings.Add(&RendererVisibilityTagBinding);
+}
+
+void UNiagaraDecalRendererProperties::PostLoad()
+{
+	Super::PostLoad();
+	
+#if WITH_EDITORONLY_DATA
+	ChangeToPositionBinding(PositionBinding);
+#endif
+	PostLoadBindings(ENiagaraRendererSourceDataMode::Particles);
+}
+
+void UNiagaraDecalRendererProperties::PostInitProperties()
+{
+	using namespace NiagaraDecalRendererPropertiesLocal;
+
+	Super::PostInitProperties();
+
+	if (HasAnyFlags(RF_ClassDefaultObject) == false)
+	{
+		// We can end up hitting PostInitProperties before the Niagara Module has initialized bindings this needs, mark this object for deferred init and early out.
+		if (FModuleManager::Get().IsModuleLoaded("Niagara") == false)
+		{
+			NiagaraDecalRendererPropertiesLocal::RendererPropertiesToDeferredInit.Add(this);
+			return;
+		}
+		SetupBindings(this);
+	}
+}
+
+void UNiagaraDecalRendererProperties::InitCDOPropertiesAfterModuleStartup()
+{
+	using namespace NiagaraDecalRendererPropertiesLocal;
+
+	UNiagaraDecalRendererProperties* CDO = CastChecked<UNiagaraDecalRendererProperties>(UNiagaraDecalRendererProperties::StaticClass()->GetDefaultObject());
+	SetupBindings(CDO);
+
+	for (TWeakObjectPtr<UNiagaraDecalRendererProperties>& WeakProps : NiagaraDecalRendererPropertiesLocal::RendererPropertiesToDeferredInit)
+	{
+		if (UNiagaraDecalRendererProperties* Props = WeakProps.Get())
+		{
+			SetupBindings(Props);
+		}
+	}
+}
+
+FNiagaraRenderer* UNiagaraDecalRendererProperties::CreateEmitterRenderer(ERHIFeatureLevel::Type FeatureLevel, const FNiagaraEmitterInstance* Emitter, const FNiagaraSystemInstanceController& InController)
+{
+	FNiagaraRenderer* NewRenderer = new FNiagaraRendererDecals(FeatureLevel, this, Emitter);
+	NewRenderer->Initialize(this, Emitter, InController);
+	return NewRenderer;
+}
+
+class FNiagaraBoundsCalculatorDecals : public FNiagaraBoundsCalculator
+{
+public:
+	FNiagaraBoundsCalculatorDecals(const FNiagaraDataSetAccessor<FNiagaraPosition>& InPositionAccessor, const FNiagaraDataSetAccessor<FVector3f>& InDecalSizeAccessor)
+		: PositionAccessor(InPositionAccessor)
+		, DecalSizeAccessor(InDecalSizeAccessor)
+	{}
+
+	FNiagaraBoundsCalculatorDecals() = delete;
+	virtual ~FNiagaraBoundsCalculatorDecals() = default;
+
+	virtual void InitAccessors(const FNiagaraDataSetCompiledData* CompiledData) override {}
+	virtual FBox CalculateBounds(const FTransform& SystemTransform, const FNiagaraDataSet& DataSet, const int32 NumInstances) const override
+	{
+		if (!NumInstances || !PositionAccessor.IsValid())
+		{
+			return FBox(ForceInit);
+		}
+
+		FNiagaraPosition BoundsMin(ForceInitToZero);
+		FNiagaraPosition BoundsMax(ForceInitToZero);
+		PositionAccessor.GetReader(DataSet).GetMinMax(BoundsMin, BoundsMax);
+		FBox Bounds(BoundsMin, BoundsMax);
+
+		if (DecalSizeAccessor.IsValid())
+		{
+			Bounds = Bounds.ExpandBy(DecalSizeAccessor.GetReader(DataSet).GetMax().GetAbsMax());
+		}
+		return Bounds;
+	}
+
+protected:
+	const FNiagaraDataSetAccessor<FNiagaraPosition>& PositionAccessor;
+	const FNiagaraDataSetAccessor<FVector3f>& DecalSizeAccessor;
+};
+
+FNiagaraBoundsCalculator* UNiagaraDecalRendererProperties::CreateBoundsCalculator()
+{
+	//if (GetCurrentSourceMode() == ENiagaraRendererSourceDataMode::Emitter)
+	//{
+	//	return nullptr;
+	//}
+
+	return new FNiagaraBoundsCalculatorDecals(PositionDataSetAccessor, DecalSizeDataSetAccessor);
+}
+
+void UNiagaraDecalRendererProperties::GetUsedMaterials(const FNiagaraEmitterInstance* InEmitter, TArray<UMaterialInterface*>& OutMaterials) const
+{
+	if ( UMaterialInterface* MaterialInterface = GetMaterial(InEmitter) )
+	{
+		OutMaterials.Add(MaterialInterface);
+	}
+}
+
+void UNiagaraDecalRendererProperties::CacheFromCompiledData(const FNiagaraDataSetCompiledData* CompiledData)
+{
+	PositionDataSetAccessor.Init(CompiledData, PositionBinding.GetDataSetBindableVariable().GetName());
+	DecalOrientationDataSetAccessor.Init(CompiledData, DecalOrientationBinding.GetDataSetBindableVariable().GetName());
+	DecalSizeDataSetAccessor.Init(CompiledData, DecalSizeBinding.GetDataSetBindableVariable().GetName());
+	DecalFadeDataSetAccessor.Init(CompiledData, DecalFadeBinding.GetDataSetBindableVariable().GetName());
+	ColorDataSetAccessor.Init(CompiledData, ColorBinding.GetDataSetBindableVariable().GetName());
+	RendererVisibilityTagAccessor.Init(CompiledData, RendererVisibilityTagBinding.GetDataSetBindableVariable().GetName());
+}
+
+#if WITH_EDITORONLY_DATA
+
+const TArray<FNiagaraVariable>& UNiagaraDecalRendererProperties::GetOptionalAttributes()
+{
+	using namespace NiagaraDecalRendererPropertiesLocal;
+
+	static TArray<FNiagaraVariable> Attrs;
+	if (Attrs.Num() == 0)
+	{
+		Attrs.Add(SYS_PARAM_PARTICLES_POSITION);
+		Attrs.Add(SYS_PARAM_PARTICLES_COLOR);
+		Attrs.Add(GetDecalSizeVariable());
+		Attrs.Add(GetDecalOrientationVariable());
+		Attrs.Add(GetDecalFadeVariable());
+		Attrs.Add(SYS_PARAM_PARTICLES_VISIBILITY_TAG);
+	}
+	return Attrs;
+}
+
+void UNiagaraDecalRendererProperties::GetRendererWidgets(const FNiagaraEmitterInstance* InEmitter, TArray<TSharedPtr<SWidget>>& OutWidgets, TSharedPtr<FAssetThumbnailPool> InThumbnailPool) const
+{
+	OutWidgets.Add(SNew(SImage).Image(FSlateIconFinder::FindIconBrushForClass(GetClass())));
+}
+
+void UNiagaraDecalRendererProperties::GetRendererTooltipWidgets(const FNiagaraEmitterInstance* InEmitter, TArray<TSharedPtr<SWidget>>& OutWidgets, TSharedPtr<FAssetThumbnailPool> InThumbnailPool) const
+{
+	OutWidgets.Add(SNew(STextBlock).Text(LOCTEXT("DecalRenderer", "Decal Renderer")));
+}
+
+void UNiagaraDecalRendererProperties::GetRendererFeedback(const FVersionedNiagaraEmitter& InEmitter, TArray<FText>& OutErrors, TArray<FText>& OutWarnings, TArray<FText>& OutInfo) const
+{
+	Super::GetRendererFeedback(InEmitter, OutErrors, OutWarnings, OutInfo);
+
+	OutInfo.Emplace(LOCTEXT("DecalRenderingPreviewSceneInfo", "Decal Rendering does not show in the default preview scene, please modify the Preview Scene Settings to add ground geometry to visualize in preview."));
+}
+
+#endif // WITH_EDITORONLY_DATA
+
+UMaterialInterface* UNiagaraDecalRendererProperties::GetMaterial(const FNiagaraEmitterInstance* InEmitter) const
+{
+	UMaterialInterface* MaterialInterface = nullptr;
+	if (InEmitter != nullptr)
+	{
+		MaterialInterface = Cast<UMaterialInterface>(InEmitter->FindBinding(MaterialParameterBinding.Parameter));
+	}
+
+	return MaterialInterface ? MaterialInterface : ToRawPtr(Material);
+}
+
+#undef LOCTEXT_NAMESPACE
diff --git a/Engine/Plugins/FX/Niagara/Source/Niagara/Private/NiagaraModule.cpp b/Engine/Plugins/FX/Niagara/Source/Niagara/Private/NiagaraModule.cpp
--- a/Engine/Plugins/FX/Niagara/Source/Niagara/Private/NiagaraModule.cpp
+++ b/Engine/Plugins/FX/Niagara/Source/Niagara/Private/NiagaraModule.cpp
@@ -16,6 +16,7 @@
 #include "NiagaraWorldManager.h"
 #include "VectorVM.h"
 #include "NiagaraConstants.h"
+#include "NiagaraDecalRendererProperties.h"
 #include "NiagaraLightRendererProperties.h"
 #include "NiagaraSpriteRendererProperties.h"
 #include "NiagaraMeshRendererProperties.h"
@@ -368,6 +369,7 @@ void INiagaraModule::StartupModule()
 	Translator_CallID = FNiagaraVariable(FNiagaraTypeDefinition::GetIntDef(), TEXT("Translator.CallID"));
 
 	FNiagaraConstants::Init();
+	UNiagaraDecalRendererProperties::InitCDOPropertiesAfterModuleStartup();
 	UNiagaraLightRendererProperties::InitCDOPropertiesAfterModuleStartup();
 	UNiagaraSpriteRendererProperties::InitCDOPropertiesAfterModuleStartup();
 	UNiagaraRibbonRendererProperties::InitCDOPropertiesAfterModuleStartup();
diff --git a/Engine/Plugins/FX/Niagara/Source/Niagara/Private/NiagaraRendererDecals.cpp b/Engine/Plugins/FX/Niagara/Source/Niagara/Private/NiagaraRendererDecals.cpp
new file mode 100644
--- /dev/null
+++ b/Engine/Plugins/FX/Niagara/Source/Niagara/Private/NiagaraRendererDecals.cpp
@@ -0,0 +1,231 @@
+// Copyright Epic Games, Inc. All Rights Reserved.
+
+#include "NiagaraRendererDecals.h"
+
+#include "NiagaraCullProxyComponent.h"
+#include "NiagaraDecalRendererProperties.h"
+#include "NiagaraEmitter.h"
+#include "NiagaraSceneProxy.h"
+#include "NiagaraSystemInstance.h"
+
+#include "Async/Async.h"
+#include "Components/LineBatchComponent.h"
+#include "Engine/World.h"
+#include "SceneInterface.h"
+
+namespace NiagaraRendererDecalsLocal
+{
+	static bool GRendererEnabled = true;
+	static FAutoConsoleVariableRef CVarRendererEnabled(
+		TEXT("fx.Niagara.DecalRenderer.Enabled"),
+		GRendererEnabled,
+		TEXT("If == 0, Niagara Decal Renderers are disabled."),
+		ECVF_Default
+	);
+
+#if UE_BUILD_SHIPPING
+	void DrawDebugDecal(UWorld* World, const FDeferredDecalUpdateParams& UpdateParams, float Fade) {}
+#else
+	static bool GDrawDebug = false;
+	static FAutoConsoleVariableRef CVarDrawDebug(
+		TEXT("fx.Niagara.DecalRenderer.DrawDebug"),
+		GDrawDebug,
+		TEXT("When none zero will draw debug information."),
+		ECVF_Default
+	);
+
+	void DrawDebugDecal(UWorld* World, const FDeferredDecalUpdateParams& UpdateParams, float Fade)
+	{
+		if (GDrawDebug == false)
+		{
+			return;
+		}
+
+		// Line Batcher is not thread safe and we could be on any thread, therefore send a task to draw the information safely
+		AsyncTask(
+			ENamedThreads::GameThread,
+			[WeakWorld=MakeWeakObjectPtr(World), UpdateParams, Fade]
+			{
+				UWorld* World = WeakWorld.Get();
+				if (World == nullptr || World->LineBatcher == nullptr)
+				{
+					return;
+				}
+				ULineBatchComponent* LineBatcher = World->LineBatcher;
+				
+				const FBox BoundsBox(-FVector(1), FVector(1));
+				const uint8 IntFade = FMath::Clamp(uint8(255.0f * Fade), 0, 255);
+				LineBatcher->DrawSolidBox(BoundsBox, UpdateParams.Transform, FColor(IntFade, IntFade, IntFade, 128), 0, 0.0f);
+
+				LineBatcher->DrawCircle(UpdateParams.Bounds.Origin, FVector::XAxisVector, FVector::YAxisVector, FColor::Red, UpdateParams.Bounds.SphereRadius, 16, 0);
+				LineBatcher->DrawCircle(UpdateParams.Bounds.Origin, FVector::XAxisVector, FVector::ZAxisVector, FColor::Red, UpdateParams.Bounds.SphereRadius, 16, 0);
+				LineBatcher->DrawCircle(UpdateParams.Bounds.Origin, FVector::YAxisVector, FVector::ZAxisVector, FColor::Red, UpdateParams.Bounds.SphereRadius, 16, 0);
+		}
+		);
+	}
+#endif
+}
+
+FNiagaraRendererDecals::FNiagaraRendererDecals(ERHIFeatureLevel::Type FeatureLevel, const UNiagaraRendererProperties *InProps, const FNiagaraEmitterInstance* Emitter)
+	: FNiagaraRenderer(FeatureLevel, InProps, Emitter)
+{
+}
+
+FNiagaraRendererDecals::~FNiagaraRendererDecals()
+{
+	check(ActiveDecalProxies.Num() == 0);
+}
+
+void FNiagaraRendererDecals::ReleaseAllDecals() const
+{
+	if (ActiveDecalProxies.Num() > 0)
+	{
+		USceneComponent* OwnerComponent = WeakOwnerComponent.Get();
+		UWorld* World = OwnerComponent ? OwnerComponent->GetWorld() : nullptr;
+		if (World && World->Scene)
+		{
+			TArray<FDeferredDecalUpdateParams> DecalUpdates;
+			DecalUpdates.AddDefaulted(ActiveDecalProxies.Num());
+			for (int i = 0; i < ActiveDecalProxies.Num(); ++i)
+			{
+				DecalUpdates[i].OperationType = FDeferredDecalUpdateParams::EOperationType::RemoveFromSceneAndDelete;
+				DecalUpdates[i].DecalProxy = ActiveDecalProxies[i];
+			}
+
+			World->Scene->BatchUpdateDecals(MoveTemp(DecalUpdates));
+			ActiveDecalProxies.Empty();
+		}
+	}
+}
+
+void FNiagaraRendererDecals::DestroyRenderState_Concurrent()
+{
+	ReleaseAllDecals();
+
+	ensureMsgf(ActiveDecalProxies.Num() == 0, TEXT("ActiveDecalProxies have been leaked?"));
+	ActiveDecalProxies.Empty();
+}
+
+FPrimitiveViewRelevance FNiagaraRendererDecals::GetViewRelevance(const FSceneView* View, const FNiagaraSceneProxy *SceneProxy)const
+{
+	FPrimitiveViewRelevance Result;
+	Result.bDrawRelevance = bHasLights && SceneProxy->IsShown(View) && View->Family->EngineShowFlags.Particles && View->Family->EngineShowFlags.Niagara;
+	Result.bShadowRelevance = false;
+	Result.bDynamicRelevance = false;
+	Result.bOpaque = false;
+	Result.bHasSimpleLights = false;
+
+	return Result;
+}
+
+FNiagaraDynamicDataBase* FNiagaraRendererDecals::GenerateDynamicData(const FNiagaraSceneProxy* Proxy, const UNiagaraRendererProperties* InProperties, const FNiagaraEmitterInstance* Emitter) const
+{
+	using namespace NiagaraRendererDecalsLocal;
+
+	// Get DataToRender
+	const FNiagaraDataSet& DataSet = Emitter->GetData();
+	const FNiagaraDataBuffer* DataToRender = DataSet.GetCurrentData();
+	FNiagaraSystemInstance* SystemInstance = Emitter->GetParentSystemInstance();
+	if (DataToRender == nullptr || SystemInstance == nullptr)
+	{
+		return nullptr;
+	}
+
+	UWorld* World = SystemInstance->GetWorld();
+	check(World && World->Scene);
+
+	USceneComponent* OwnerComponent = SystemInstance->GetAttachComponent();
+	check(OwnerComponent);
+
+	WeakOwnerComponent = OwnerComponent;
+	const UNiagaraDecalRendererProperties* RendererProperties = CastChecked<const UNiagaraDecalRendererProperties>(InProperties);
+
+	const bool bUseLocalSpace = UseLocalSpace(Proxy);
+	const FTransform LocalToWorld = SystemInstance->GetWorldTransform();
+	const FNiagaraLWCConverter LwcConverter = SystemInstance->GetLWCConverter(bUseLocalSpace);
+	const FNiagaraPosition DefaultPos = bUseLocalSpace ? FVector::ZeroVector : LocalToWorld.GetLocation();
+	const FQuat4f DefaultRot = FRotator3f(-90.0f, 0.0f, 90.0f).Quaternion();
+	const FVector3f DefaultSize = FVector3f(50.0f, 50.0f, 50.0f);	//-TODO: Read defaults from properties or default variable
+
+	// Check for Material Update
+	UMaterialInterface* Material = RendererProperties->GetMaterial(Emitter);
+	if (Material != WeakMaterial.Get())
+	{
+		WeakMaterial = Material;
+		ReleaseAllDecals();
+	}
+
+	// Create all our data readers
+	const auto PositionReader	= RendererProperties->PositionDataSetAccessor.GetReader(DataSet);
+	const auto RotationReader	= RendererProperties->DecalOrientationDataSetAccessor.GetReader(DataSet);
+	const auto SizeReader		= RendererProperties->DecalSizeDataSetAccessor.GetReader(DataSet);
+	const auto ColorReader		= RendererProperties->ColorDataSetAccessor.GetReader(DataSet);
+	const auto FadeReader		= RendererProperties->DecalFadeDataSetAccessor.GetReader(DataSet);
+	const auto VisTagReader		= RendererProperties->RendererVisibilityTagAccessor.GetReader(DataSet);
+
+	TArray<FDeferredDecalUpdateParams> DecalUpdates;
+	DecalUpdates.Reserve(DataToRender->GetNumInstances());
+
+	int NumActiveDecalProxies = 0;
+	for (uint32 ParticleIndex=0; ParticleIndex < DataToRender->GetNumInstances(); ++ParticleIndex)
+	{
+		// Check visibility tag
+		const int32 VisTag = VisTagReader.GetSafe(ParticleIndex, RendererProperties->RendererVisibility);
+		if (VisTag != RendererProperties->RendererVisibility)
+		{
+			continue;
+		}
+
+		// Grab Decal Attributes
+		const FVector3f SimPos	= PositionReader.GetSafe(ParticleIndex, DefaultPos);
+		const FQuat4f SimRot	= RotationReader.GetSafe(ParticleIndex, DefaultRot);
+		const FVector Position	= bUseLocalSpace ? LocalToWorld.TransformPosition(FVector(SimPos)) : LwcConverter.ConvertSimulationPositionToWorld(SimPos);
+		const FQuat Rotation	= bUseLocalSpace ? LocalToWorld.TransformRotation(FQuat(SimRot)) : FQuat(SimRot);
+		const FVector Size		= FVector(SizeReader.GetSafe(ParticleIndex, DefaultSize) * 0.5f);
+		const float Fade		= RendererProperties->bUseColorBindingAsFade ? ColorReader.GetSafe(ParticleIndex, FLinearColor::White).A : FadeReader.GetSafe(ParticleIndex, 1.0f);
+
+		// Create Update Parameters
+		FDeferredDecalUpdateParams& UpdateParams = DecalUpdates.AddDefaulted_GetRef();
+		if (NumActiveDecalProxies < ActiveDecalProxies.Num())
+		{
+			UpdateParams.DecalProxy = ActiveDecalProxies[NumActiveDecalProxies];
+			UpdateParams.OperationType = FDeferredDecalUpdateParams::EOperationType::Update;
+		}
+		else
+		{
+			UpdateParams.DecalProxy = new FDeferredDecalProxy(OwnerComponent, Material);
+			UpdateParams.OperationType = FDeferredDecalUpdateParams::EOperationType::AddToSceneAndUpdate;
+			ActiveDecalProxies.Add(UpdateParams.DecalProxy);
+		}
+		++NumActiveDecalProxies;
+
+		UpdateParams.Transform		= FTransform(Rotation, Position, FVector(Size));
+		UpdateParams.Bounds			= FBoxSphereBounds(FSphere(Position, Size.GetAbsMax() * 2.0));
+		UpdateParams.AbsSpawnTime	= World->TimeSeconds - FMath::Clamp(1.0f - Fade, 0.0f, 1.0f);
+		UpdateParams.FadeStartDelay	= 0.0f;
+		UpdateParams.FadeDuration	= 1.0f;
+
+		DrawDebugDecal(World, UpdateParams, Fade);
+	}
+
+	// Remove any unused decals
+	// Note: We don't delete at the RemoveDecal call will automatically do this
+	const int NumToRemove = ActiveDecalProxies.Num() - NumActiveDecalProxies;
+	if (NumToRemove > 0)
+	{
+		for (int i = DataToRender->GetNumInstances(); i < ActiveDecalProxies.Num(); ++i)
+		{
+			FDeferredDecalUpdateParams& UpdateParams = DecalUpdates.AddDefaulted_GetRef();
+			UpdateParams.OperationType = FDeferredDecalUpdateParams::EOperationType::RemoveFromSceneAndDelete;
+			UpdateParams.DecalProxy = ActiveDecalProxies.Pop();
+		}
+	}
+
+	// Send updates to RT
+	if (DecalUpdates.Num() > 0)
+	{
+		World->Scene->BatchUpdateDecals(MoveTemp(DecalUpdates));
+	}
+
+	return nullptr;
+}
diff --git a/Engine/Plugins/FX/Niagara/Source/Niagara/Public/NiagaraDecalRendererProperties.h b/Engine/Plugins/FX/Niagara/Source/Niagara/Public/NiagaraDecalRendererProperties.h
new file mode 100644
--- /dev/null
+++ b/Engine/Plugins/FX/Niagara/Source/Niagara/Public/NiagaraDecalRendererProperties.h
@@ -0,0 +1,94 @@
+// Copyright Epic Games, Inc. All Rights Reserved.
+
+#pragma once
+
+#include "CoreMinimal.h"
+#include "UObject/ObjectMacros.h"
+#include "NiagaraRendererProperties.h"
+#include "NiagaraCommon.h"
+#include "NiagaraDataSetAccessor.h"
+#include "NiagaraParameterBinding.h"
+#include "NiagaraDecalRendererProperties.generated.h"
+
+class UMaterialInterface;
+class FNiagaraEmitterInstance;
+class SWidget;
+
+UCLASS(editinlinenew, MinimalAPI, meta = (DisplayName = "Decal Renderer"))
+class UNiagaraDecalRendererProperties : public UNiagaraRendererProperties
+{
+public:
+	GENERATED_BODY()
+
+	UNiagaraDecalRendererProperties();
+
+	//UObject Interface
+	virtual void PostLoad() override;
+	virtual void PostInitProperties() override;
+	//UObject Interface END
+
+	static void InitCDOPropertiesAfterModuleStartup();
+
+	//~ UNiagaraRendererProperties interface
+	virtual FNiagaraRenderer* CreateEmitterRenderer(ERHIFeatureLevel::Type FeatureLevel, const FNiagaraEmitterInstance* Emitter, const FNiagaraSystemInstanceController& InController) override;
+	virtual class FNiagaraBoundsCalculator* CreateBoundsCalculator() override;
+	virtual void GetUsedMaterials(const FNiagaraEmitterInstance* InEmitter, TArray<UMaterialInterface*>& OutMaterials) const override;
+	virtual bool IsSimTargetSupported(ENiagaraSimTarget InSimTarget) const override { return InSimTarget == ENiagaraSimTarget::CPUSim; };
+#if WITH_EDITORONLY_DATA
+	virtual const TArray<FNiagaraVariable>& GetOptionalAttributes() override;
+	virtual void GetRendererWidgets(const FNiagaraEmitterInstance* InEmitter, TArray<TSharedPtr<SWidget>>& OutWidgets, TSharedPtr<FAssetThumbnailPool> InThumbnailPool) const override;
+	virtual void GetRendererTooltipWidgets(const FNiagaraEmitterInstance* InEmitter, TArray<TSharedPtr<SWidget>>& OutWidgets, TSharedPtr<FAssetThumbnailPool> InThumbnailPool) const override;
+	virtual void GetRendererFeedback(const FVersionedNiagaraEmitter& InEmitter, TArray<FText>& OutErrors, TArray<FText>& OutWarnings, TArray<FText>& OutInfo) const override;
+#endif // WITH_EDITORONLY_DATA
+	virtual void CacheFromCompiledData(const FNiagaraDataSetCompiledData* CompiledData) override;
+	//UNiagaraRendererProperties Interface END
+
+	UMaterialInterface* GetMaterial(const FNiagaraEmitterInstance* InEmitter) const;
+
+	/** What material to use for the decal. */
+	UPROPERTY(EditAnywhere, Category = "Decal Rendering")
+	TObjectPtr<UMaterialInterface> Material;
+
+	/** Binding to material. */
+	UPROPERTY(EditAnywhere, Category = "Decal Rendering")
+	FNiagaraParameterBinding MaterialParameterBinding;
+
+	/** If a render visibility tag is present, particles whose tag matches this value will be visible in this renderer. */
+	UPROPERTY(EditAnywhere, Category = "Decal Rendering")
+	int32 RendererVisibility = 0;
+
+	/** Position binding for the decals, should be center of the decal */
+	UPROPERTY(EditAnywhere, Category = "Bindings")
+	FNiagaraVariableAttributeBinding PositionBinding;
+
+	/** Orientation binding for the decal. */
+	UPROPERTY(EditAnywhere, Category = "Bindings")
+	FNiagaraVariableAttributeBinding DecalOrientationBinding;
+
+	/** Size binding for the decal. */
+	UPROPERTY(EditAnywhere, Category = "Bindings")
+	FNiagaraVariableAttributeBinding DecalSizeBinding;
+
+	/** When enabled we use the Color Binding as the fade control, otherwise we use the DecalFade binding. */
+	UPROPERTY(EditAnywhere, Category = "Bindings")
+	bool bUseColorBindingAsFade = true;
+
+	/** Fade binding for the decal, can be queried from the decal fade material node.. */
+	UPROPERTY(EditAnywhere, Category = "Bindings", Meta = (EditCondition = "!bUseColorBindingAsFade"))
+	FNiagaraVariableAttributeBinding DecalFadeBinding;
+
+	/** Color binding for the decal, only the alpha channel is used and can be queried from the decal fade material node. */
+	UPROPERTY(EditAnywhere, Category = "Bindings", Meta = (EditCondition = "bUseColorBindingAsFade"))
+	FNiagaraVariableAttributeBinding ColorBinding;
+
+	/** Visibility tag binding, when valid the returned values is compated with RendererVisibility. */
+	UPROPERTY(EditAnywhere, Category = "Bindings")
+	FNiagaraVariableAttributeBinding RendererVisibilityTagBinding;
+
+	FNiagaraDataSetAccessor<FNiagaraPosition>	PositionDataSetAccessor;
+	FNiagaraDataSetAccessor<FQuat4f>			DecalOrientationDataSetAccessor;
+	FNiagaraDataSetAccessor<FVector3f>			DecalSizeDataSetAccessor;
+	FNiagaraDataSetAccessor<float>				DecalFadeDataSetAccessor;
+	FNiagaraDataSetAccessor<FLinearColor>		ColorDataSetAccessor;
+	FNiagaraDataSetAccessor<int32>				RendererVisibilityTagAccessor;
+};
diff --git a/Engine/Plugins/FX/Niagara/Source/Niagara/Public/NiagaraRendererDecals.h b/Engine/Plugins/FX/Niagara/Source/Niagara/Public/NiagaraRendererDecals.h
new file mode 100644
--- /dev/null
+++ b/Engine/Plugins/FX/Niagara/Source/Niagara/Public/NiagaraRendererDecals.h
@@ -0,0 +1,27 @@
+// Copyright Epic Games, Inc. All Rights Reserved.
+
+#pragma once
+
+#include "NiagaraRenderer.h"
+
+class FDeferredDecalProxy;
+class USceneComponent;
+
+class NIAGARA_API FNiagaraRendererDecals : public FNiagaraRenderer
+{
+public:
+	explicit FNiagaraRendererDecals(ERHIFeatureLevel::Type FeatureLevel, const UNiagaraRendererProperties *InProps, const FNiagaraEmitterInstance* Emitter);
+	~FNiagaraRendererDecals();
+
+	void ReleaseAllDecals() const;
+
+	//FNiagaraRenderer interface
+	virtual void DestroyRenderState_Concurrent() override;
+	virtual FPrimitiveViewRelevance GetViewRelevance(const FSceneView* View, const FNiagaraSceneProxy *SceneProxy)const override;
+	virtual FNiagaraDynamicDataBase *GenerateDynamicData(const FNiagaraSceneProxy* Proxy, const UNiagaraRendererProperties* InProperties, const FNiagaraEmitterInstance* Emitter) const override;
+	//FNiagaraRenderer interface END
+
+	mutable TWeakObjectPtr<USceneComponent>		WeakOwnerComponent;
+	mutable TWeakObjectPtr<UMaterialInterface>	WeakMaterial;
+	mutable TArray<FDeferredDecalProxy*>		ActiveDecalProxies;
+};
diff --git a/Engine/Plugins/FX/Niagara/Source/NiagaraEditor/Private/NiagaraEditorModule.cpp b/Engine/Plugins/FX/Niagara/Source/NiagaraEditor/Private/NiagaraEditorModule.cpp
--- a/Engine/Plugins/FX/Niagara/Source/NiagaraEditor/Private/NiagaraEditorModule.cpp
+++ b/Engine/Plugins/FX/Niagara/Source/NiagaraEditor/Private/NiagaraEditorModule.cpp
@@ -150,6 +150,7 @@
 #include "NiagaraDebugVis.h"
 #include "NiagaraPerfBaseline.h"
 #include "NiagaraGraphDataCache.h"
+#include "NiagaraDecalRendererProperties.h"
 #include "NiagaraLightRendererProperties.h"
 #include "NiagaraRibbonRendererProperties.h"
 #include "NiagaraSpriteRendererProperties.h"
@@ -1816,6 +1817,18 @@ void FNiagaraEditorModule::RegisterDefaultRendererFactories()
 			UNiagaraLightRendererProperties* NewRenderer = NewObject<UNiagaraLightRendererProperties>(OuterEmitter, NAME_None, RF_Transactional);
 			return NewRenderer;
 		})));
+
+	RegisterRendererCreationInfo(FNiagaraRendererCreationInfo(
+		UNiagaraDecalRendererProperties::StaticClass()->GetDisplayNameText(),
+		FText::FromString(UNiagaraDecalRendererProperties::StaticClass()->GetDescription()),
+		UNiagaraDecalRendererProperties::StaticClass()->GetClassPathName(),
+		FNiagaraRendererCreationInfo::FRendererFactory::CreateLambda([](UObject* OuterEmitter)
+		{
+			UNiagaraDecalRendererProperties* NewRenderer = NewObject<UNiagaraDecalRendererProperties>(OuterEmitter, NAME_None, RF_Transactional);
+			FSoftObjectPath DefaultMaterial(TEXT("Material'/Niagara/DefaultAssets/DefaultDecalMaterial.DefaultDecalMaterial'"));
+			NewRenderer->Material = Cast<UMaterialInterface>(DefaultMaterial.TryLoad());
+			return NewRenderer;
+		})));
 }
 
 void FNiagaraEditorModule::RegisterAssetTypeAction(IAssetTools& AssetTools, TSharedRef<IAssetTypeActions> Action)
diff --git a/Engine/Source/Runtime/Engine/Private/Components/DecalComponent.cpp b/Engine/Source/Runtime/Engine/Private/Components/DecalComponent.cpp
--- a/Engine/Source/Runtime/Engine/Private/Components/DecalComponent.cpp
+++ b/Engine/Source/Runtime/Engine/Private/Components/DecalComponent.cpp
@@ -67,6 +67,43 @@ FDeferredDecalProxy::FDeferredDecalProxy(const UDecalComponent* InComponent)
 	}
 }
 
+FDeferredDecalProxy::FDeferredDecalProxy(const USceneComponent* InComponent, UMaterialInterface* InMaterial)
+	: DrawInGame(InComponent->GetVisibleFlag() && !InComponent->bHiddenInGame)
+	, DrawInEditor(InComponent->GetVisibleFlag())
+	, InvFadeDuration(-1.0f)
+	, InvFadeInDuration(1.0f)
+	, FadeStartDelayNormalized(1.0f)
+	, FadeInStartDelayNormalized(0.0f)
+	, FadeScreenSize(0.1f)
+{
+	Component = InComponent;
+	DecalMaterial = InMaterial;
+	if (InMaterial == nullptr || (InMaterial->GetMaterial()->MaterialDomain != MD_DeferredDecal))
+	{
+		DecalMaterial = UMaterial::GetDefaultMaterial(MD_DeferredDecal);
+	}
+
+	SetTransformIncludingDecalSize(FTransform::Identity, InComponent->CalcBounds(InComponent->GetComponentTransform()));
+	bOwnerSelected = InComponent->IsOwnerSelected();
+	SortOrder = 0;
+
+#if WITH_EDITOR
+	// We don't want to fade when we're editing, only in Simulate/PIE/Game
+	if (!GIsEditor || (InComponent->GetWorld() && InComponent->GetWorld()->IsPlayInEditor()))
+#endif
+	{
+		InitializeFadingParameters(InComponent->GetWorld()->GetTimeSeconds(), 1.0f, 1.0f, 0.0f, 0.0f);
+	}
+
+	if (InComponent->GetOwner())
+	{
+		DrawInGame &= !(InComponent->GetOwner()->IsHidden());
+#if WITH_EDITOR
+		DrawInEditor &= !InComponent->GetOwner()->IsHiddenEd();
+#endif
+	}
+}
+
 void FDeferredDecalProxy::SetTransformIncludingDecalSize(const FTransform& InComponentToWorldIncludingDecalSize, const FBoxSphereBounds& InBounds)
 {
 	ComponentTrans = InComponentToWorldIncludingDecalSize;
diff --git a/Engine/Source/Runtime/Engine/Public/SceneInterface.h b/Engine/Source/Runtime/Engine/Public/SceneInterface.h
--- a/Engine/Source/Runtime/Engine/Public/SceneInterface.h
+++ b/Engine/Source/Runtime/Engine/Public/SceneInterface.h
@@ -25,6 +25,7 @@ class FSkyLightSceneProxy;
 class FTexture;
 class FVertexFactory;
 class UDecalComponent;
+struct FDeferredDecalUpdateParams;
 class ULightComponent;
 class UPlanarReflectionComponent;
 class UPrimitiveComponent;
@@ -157,6 +158,7 @@ public:
 	virtual void UpdateDecalTransform(UDecalComponent* Component) = 0;
 	virtual void UpdateDecalFadeOutTime(UDecalComponent* Component) = 0;
 	virtual void UpdateDecalFadeInTime(UDecalComponent* Component) = 0;
+	virtual void BatchUpdateDecals(TArray<FDeferredDecalUpdateParams>&& UpdateParams) = 0;
 
 	/** Adds a reflection capture to the scene. */
 	virtual void AddReflectionCapture(class UReflectionCaptureComponent* Component) {}
diff --git a/Engine/Source/Runtime/Engine/Public/SceneManagement.h b/Engine/Source/Runtime/Engine/Public/SceneManagement.h
--- a/Engine/Source/Runtime/Engine/Public/SceneManagement.h
+++ b/Engine/Source/Runtime/Engine/Public/SceneManagement.h
@@ -1866,6 +1866,7 @@ class ENGINE_API FDeferredDecalProxy
 public:
 	/** constructor */
 	FDeferredDecalProxy(const UDecalComponent* InComponent);
+	FDeferredDecalProxy(const USceneComponent* InComponent, UMaterialInterface* InMaterial);
 
 	/**
 	 * Updates the decal proxy's cached transform and bounds.
@@ -1881,8 +1882,8 @@ public:
 
 	inline const FBoxSphereBounds& GetBounds() const { return Bounds; }
 
-	/** Pointer back to the game thread decal component. */
-	const UDecalComponent* Component;
+	/** Pointer back to the game thread owner component. */
+	const USceneComponent* Component;
 
 	UMaterialInterface* DecalMaterial;
 
@@ -1922,6 +1923,28 @@ public:
 	float FadeScreenSize;
 };
 
+struct FDeferredDecalUpdateParams
+{
+	enum class EOperationType : int
+	{
+		AddToSceneAndUpdate,				// Adds the decal to the scene an updates the parameters
+		Update,								// Updates the decals parameters
+		RemoveFromSceneAndDelete,			// Remove the decal from the scene and deletes the proxy
+	};
+
+	EOperationType			OperationType = EOperationType::Update;
+	FDeferredDecalProxy*	DecalProxy = nullptr;
+	FTransform				Transform;
+	FBoxSphereBounds		Bounds;
+	float					AbsSpawnTime = 0.0f;
+	float					FadeDuration = 0.0f;
+	float					FadeStartDelay = 1.0f;
+	float					FadeInDuration = 0.0f;
+	float					FadeInStartDelay = 0.0f;
+	float					FadeScreenSize = 0.01f;
+	int32					SortOrder = 0;
+};
+
 /** Reflection capture shapes. */
 namespace EReflectionCaptureShape
 {
diff --git a/Engine/Source/Runtime/Renderer/Private/RendererScene.cpp b/Engine/Source/Runtime/Renderer/Private/RendererScene.cpp
--- a/Engine/Source/Runtime/Renderer/Private/RendererScene.cpp
+++ b/Engine/Source/Runtime/Renderer/Private/RendererScene.cpp
@@ -3005,6 +3005,33 @@ void FScene::UpdateDecalFadeInTime(UDecalComponent* Decal)
 	}
 }
 
+void FScene::BatchUpdateDecals(TArray<FDeferredDecalUpdateParams>&& UpdateParams)
+{
+	ENQUEUE_RENDER_COMMAND(FBatchUpdateDecalsCommand)(
+		[Scene=this, UpdateParams_RT=MoveTemp(UpdateParams)](FRHICommandListImmediate& RHICmdList)
+		{
+			for (const FDeferredDecalUpdateParams& DecalUpdate : UpdateParams_RT )
+			{
+				if (DecalUpdate.OperationType == FDeferredDecalUpdateParams::EOperationType::RemoveFromSceneAndDelete)
+				{
+					Scene->AddOrRemoveDecal_RenderThread(DecalUpdate.DecalProxy, false);
+					continue;
+				}
+
+				if (DecalUpdate.OperationType == FDeferredDecalUpdateParams::EOperationType::AddToSceneAndUpdate )
+				{
+					Scene->AddOrRemoveDecal_RenderThread(DecalUpdate.DecalProxy, true);
+				}
+
+				DecalUpdate.DecalProxy->SetTransformIncludingDecalSize(DecalUpdate.Transform, DecalUpdate.Bounds);
+				DecalUpdate.DecalProxy->InitializeFadingParameters(DecalUpdate.AbsSpawnTime, DecalUpdate.FadeDuration, DecalUpdate.FadeStartDelay, DecalUpdate.FadeInDuration, DecalUpdate.FadeInStartDelay);
+				DecalUpdate.DecalProxy->FadeScreenSize = DecalUpdate.FadeScreenSize;
+				DecalUpdate.DecalProxy->SortOrder = DecalUpdate.SortOrder;
+			}
+		}
+	);
+}
+
 void FScene::AddHairStrands(FHairStrandsInstance* Proxy)
 {
 	if (Proxy)
@@ -6178,6 +6205,7 @@ public:
 	virtual void UpdateDecalTransform(UDecalComponent* Decal) override {}
 	virtual void UpdateDecalFadeOutTime(UDecalComponent* Decal) override {};
 	virtual void UpdateDecalFadeInTime(UDecalComponent* Decal) override {};
+	virtual void BatchUpdateDecals(TArray<FDeferredDecalUpdateParams>&& UpdateParams) override {}
 
 	/** Updates the transform of a light which has already been added to the scene. */
 	virtual void UpdateLightTransform(ULightComponent* Light) override {}
diff --git a/Engine/Source/Runtime/Renderer/Private/ScenePrivate.h b/Engine/Source/Runtime/Renderer/Private/ScenePrivate.h
--- a/Engine/Source/Runtime/Renderer/Private/ScenePrivate.h
+++ b/Engine/Source/Runtime/Renderer/Private/ScenePrivate.h
@@ -3044,6 +3044,7 @@ public:
 	virtual void UpdateDecalTransform(UDecalComponent* Decal) override;
 	virtual void UpdateDecalFadeOutTime(UDecalComponent* Decal) override;
 	virtual void UpdateDecalFadeInTime(UDecalComponent* Decal) override;
+	virtual void BatchUpdateDecals(TArray<FDeferredDecalUpdateParams>&& UpdateParams) override;
 	virtual void AddReflectionCapture(UReflectionCaptureComponent* Component) override;
 	virtual void RemoveReflectionCapture(UReflectionCaptureComponent* Component) override;
 	virtual void GetReflectionCaptureData(UReflectionCaptureComponent* Component, class FReflectionCaptureData& OutCaptureData) override;
