Fix for linker with dependency between NiagaraShader and Niagara modules
#rnx

diff --git a/Engine/Plugins/FX/Niagara/Source/Niagara/Public/NiagaraGpuComputeDispatchInterface.h b/Engine/Plugins/FX/Niagara/Source/Niagara/Public/NiagaraGpuComputeDispatchInterface.h
--- a/Engine/Plugins/FX/Niagara/Source/Niagara/Public/NiagaraGpuComputeDispatchInterface.h
+++ b/Engine/Plugins/FX/Niagara/Source/Niagara/Public/NiagaraGpuComputeDispatchInterface.h
@@ -142,7 +142,7 @@ public:
 
 	virtual FNiagaraAsyncGpuTraceHelper& GetAsyncGpuTraceHelper() const = 0;
 
-	bool IsOutsideSceneRenderer() const { return bIsOutsideSceneRenderer; }
+	FORCEINLINE bool IsOutsideSceneRenderer() const { return bIsOutsideSceneRenderer; }
 
 #if WITH_MGPU
 	/**
