Niagara Data Channels - Supporting CL

Don't consider disabled emitters in UNiagaraSystem::ForEachScript()

diff --git a/Engine/Plugins/FX/Niagara/Source/Niagara/Classes/NiagaraSystemImpl.h b/Engine/Plugins/FX/Niagara/Source/Niagara/Classes/NiagaraSystemImpl.h
--- a/Engine/Plugins/FX/Niagara/Source/Niagara/Classes/NiagaraSystemImpl.h
+++ b/Engine/Plugins/FX/Niagara/Source/Niagara/Classes/NiagaraSystemImpl.h
@@ -6,6 +6,8 @@
 #include "NiagaraEmitter.h"
 #include "NiagaraScript.h"
 #include "NiagaraSystem.h"
+#include "NiagaraDataInterface.h"
+#include "NiagaraDataInterfaceUtilities.h"
 
 
 template<typename TAction>
@@ -16,9 +18,12 @@ void UNiagaraSystem::ForEachScript(TAction Func) const
 
 	for (const FNiagaraEmitterHandle& Handle : EmitterHandles)
 	{
-		if (FVersionedNiagaraEmitterData* EmitterData = Handle.GetEmitterData())
+		if (Handle.GetIsEnabled())
 		{
-			EmitterData->ForEachScript(Func);
+			if (FVersionedNiagaraEmitterData* EmitterData = Handle.GetEmitterData())
+			{
+				EmitterData->ForEachScript(Func);
+			}
 		}
 	}
 }
@@ -67,4 +72,4 @@ void UNiagaraSystem::ForEachPlatformSet(TAction Func)
 			Emitter->ForEachPlatformSet(Func);
 		}
 	}
-}
\ No newline at end of file
+}
