Tweak to the outliner so that it clears any stored sim caches when the outliner changes /  is updated.

diff --git a/Engine/Plugins/FX/Niagara/Source/NiagaraEditor/Private/NiagaraDebugger.cpp b/Engine/Plugins/FX/Niagara/Source/NiagaraEditor/Private/NiagaraDebugger.cpp
--- a/Engine/Plugins/FX/Niagara/Source/NiagaraEditor/Private/NiagaraDebugger.cpp
+++ b/Engine/Plugins/FX/Niagara/Source/NiagaraEditor/Private/NiagaraDebugger.cpp
@@ -297,6 +297,11 @@ void FNiagaraDebugger::HandleConnectionAcceptedMessage(const FNiagaraDebuggerAcc
 
 		UpdateDebugHUDSettings();
 
+		if (UNiagaraOutliner* Outliner = GetOutliner())
+		{
+			Outliner->Reset();
+		}
+
 		OnConnectionMadeDelegate.Broadcast(NewConnection);
 	}
 	else
diff --git a/Engine/Plugins/FX/Niagara/Source/NiagaraEditor/Private/NiagaraOutliner.cpp b/Engine/Plugins/FX/Niagara/Source/NiagaraEditor/Private/NiagaraOutliner.cpp
--- a/Engine/Plugins/FX/Niagara/Source/NiagaraEditor/Private/NiagaraOutliner.cpp
+++ b/Engine/Plugins/FX/Niagara/Source/NiagaraEditor/Private/NiagaraOutliner.cpp
@@ -38,9 +38,21 @@ void UNiagaraOutliner::OnChanged()
 	OnChangedDelegate.Broadcast();
 }
 
+void UNiagaraOutliner::Reset()
+{
+	Data.WorldData.Reset();
+	SystemSimCaches.Reset();
+	OnChanged();
+}
+
 void UNiagaraOutliner::UpdateData(const FNiagaraOutlinerData& NewData)
 {
 	FNiagaraOutlinerData::StaticStruct()->CopyScriptStruct(&Data, &NewData);
+
+	//Clear out any sim caches when we update the outliner.
+	//TODO: We should keep them all and just track info about when the were taken etc.
+	SystemSimCaches.Reset();
+
 	OnChanged();
 	//TODO: Do some kind of diff on the data and collect any recently removed components etc into their own area.
 	//Possibly keep them in the UI optionally but colour/mark them as dead until the user opts to remove them or on some timed interval.
diff --git a/Engine/Plugins/FX/Niagara/Source/NiagaraEditor/Private/NiagaraOutliner.h b/Engine/Plugins/FX/Niagara/Source/NiagaraEditor/Private/NiagaraOutliner.h
--- a/Engine/Plugins/FX/Niagara/Source/NiagaraEditor/Private/NiagaraOutliner.h
+++ b/Engine/Plugins/FX/Niagara/Source/NiagaraEditor/Private/NiagaraOutliner.h
@@ -144,6 +144,7 @@ public:
 	//FNotifyHook Interface END
 #endif
 
+	void Reset();
 	void UpdateData(const FNiagaraOutlinerData& NewData);
 	void UpdateSystemSimCache(const FNiagaraSystemSimCacheCaptureReply& Reply);
 
