Niagara Data Channels - Supporting CL

Type details customization for FNiagaraVariableBase in few cases where we want to view/edit these directly.

diff --git a/Engine/Plugins/FX/Niagara/Source/Niagara/Public/NiagaraTypes.h b/Engine/Plugins/FX/Niagara/Source/Niagara/Public/NiagaraTypes.h
--- a/Engine/Plugins/FX/Niagara/Source/Niagara/Public/NiagaraTypes.h
+++ b/Engine/Plugins/FX/Niagara/Source/Niagara/Public/NiagaraTypes.h
@@ -154,6 +154,15 @@ struct FNiagaraNumeric
 	GENERATED_USTRUCT_BODY()
 };
 
+// only used for LWC conversions, not supported by Niagara yet
+USTRUCT(meta = (DisplayName = "double"))
+struct FNiagaraDouble
+{
+	GENERATED_USTRUCT_BODY()
+
+	UPROPERTY(EditAnywhere, Category = Parameters)
+	double Value = 0;
+};
 
 USTRUCT()
 struct FNiagaraParameterMap
diff --git a/Engine/Plugins/FX/Niagara/Source/NiagaraEditor/Private/Customizations/NiagaraTypeCustomizations.cpp b/Engine/Plugins/FX/Niagara/Source/NiagaraEditor/Private/Customizations/NiagaraTypeCustomizations.cpp
--- a/Engine/Plugins/FX/Niagara/Source/NiagaraEditor/Private/Customizations/NiagaraTypeCustomizations.cpp
+++ b/Engine/Plugins/FX/Niagara/Source/NiagaraEditor/Private/Customizations/NiagaraTypeCustomizations.cpp
@@ -42,6 +42,9 @@
 #include "NiagaraSystem.h"
 #include "NiagaraTypes.h"
 
+#include "Widgets/SNiagaraParameterMenu.h"
+#include "Widgets/Input/SSuggestionTextBox.h"
+
 #include UE_INLINE_GENERATED_CPP_BY_NAME(NiagaraTypeCustomizations)
 
 #define LOCTEXT_NAMESPACE "FNiagaraVariableAttributeBindingCustomization"
@@ -2367,4 +2370,188 @@ TSharedRef<SWidget> FNiagaraRendererMaterialStaticBoolParameterCustomization::On
 	return MenuBuilder.MakeWidget();
 }
 
+
+//////////////////////////////////////////////////////////////////////////
+
+class SNiagaraRawVariableTypeSelectMenu : public SNiagaraParameterMenu
+{
+public:
+	//DECLARE_DELEGATE_RetVal_OneParam(bool, FOnAllowMakeType, const FNiagaraTypeDefinition&);
+
+	SLATE_BEGIN_ARGS(SNiagaraRawVariableTypeSelectMenu)
+		: _AutoExpandMenu(false)
+	{}
+	//~ Begin Required Args
+	SLATE_ARGUMENT(TSharedPtr<IPropertyHandle>, PropertyHandle)
+	SLATE_ARGUMENT(FNiagaraVariable*, VarToModify)
+	//~ End Required Args
+	SLATE_ARGUMENT(bool, AutoExpandMenu)
+	SLATE_END_ARGS();
+
+	void Construct(const FArguments& InArgs);
+
+protected:
+	virtual void CollectAllActions(FGraphActionListBuilderBase& OutAllActions) override;
+
+private:
+	TSharedPtr<IPropertyHandle> PropertyHandle;
+	FNiagaraVariable* VarToModify;
+};
+
+
+void SNiagaraRawVariableTypeSelectMenu::Construct(const FArguments& InArgs)
+{
+	checkf(InArgs._VarToModify != nullptr, TEXT("Tried to construct change var type menu without valid var ptr!"));
+	this->PropertyHandle = InArgs._PropertyHandle;
+	this->VarToModify = InArgs._VarToModify;
+
+	SNiagaraParameterMenu::FArguments SuperArgs;
+	SuperArgs._AutoExpandMenu = InArgs._AutoExpandMenu;
+	SNiagaraParameterMenu::Construct(SuperArgs);
+}
+
+void SNiagaraRawVariableTypeSelectMenu::CollectAllActions(FGraphActionListBuilderBase& OutAllActions)
+{
+	FNiagaraMenuActionCollector Collector;
+
+	static TArray<FNiagaraTypeDefinition> Types;
+	if(Types.IsEmpty())
+	{
+		//TODO: Build better list of possible user facing types.
+		UPackage* CorePkg = FindObjectChecked<UPackage>(nullptr, TEXT("/Script/CoreUObject"));
+		Types.Emplace(FNiagaraTypeDefinition::GetBoolDef());
+		Types.Emplace(FNiagaraTypeDefinition::GetIntDef());
+		Types.Emplace(FNiagaraTypeDefinition(FNiagaraDouble::StaticStruct(), FNiagaraTypeDefinition::EAllowUnfriendlyStruct::Allow));
+
+		Types.Emplace(FNiagaraTypeDefinition(FindObjectChecked<UScriptStruct>(CorePkg, TEXT("Vector2d")), FNiagaraTypeDefinition::EAllowUnfriendlyStruct::Allow));
+		Types.Emplace(FNiagaraTypeDefinition(FindObjectChecked<UScriptStruct>(CorePkg, TEXT("Vector")), FNiagaraTypeDefinition::EAllowUnfriendlyStruct::Allow));
+		Types.Emplace(FNiagaraTypeDefinition(FindObjectChecked<UScriptStruct>(CorePkg, TEXT("Vector4")), FNiagaraTypeDefinition::EAllowUnfriendlyStruct::Allow));
+		Types.Emplace(FNiagaraTypeDefinition(FindObjectChecked<UScriptStruct>(CorePkg, TEXT("Quat")), FNiagaraTypeDefinition::EAllowUnfriendlyStruct::Allow));
+
+		Types.Emplace(FNiagaraTypeDefinition::GetColorDef());
+		Types.Emplace(FNiagaraTypeDefinition::GetIDDef());
+		Types.Emplace(FNiagaraTypeDefinition(FNiagaraSpawnInfo::StaticStruct()));
+	};
+
+	Types.Sort([](const FNiagaraTypeDefinition& A, const FNiagaraTypeDefinition& B) { return (A.GetNameText().ToLower().ToString() < B.GetNameText().ToLower().ToString()); });
+
+	for (const FNiagaraTypeDefinition& RegisteredType : Types)
+	{
+		FNiagaraVariable DefaultVar(RegisteredType, FName(*RegisteredType.GetName()));
+		FNiagaraEditorUtilities::ResetVariableToDefaultValue(DefaultVar);
+
+		FText Category = FNiagaraEditorUtilities::GetVariableTypeCategory(DefaultVar);
+		const FText DisplayName = RegisteredType.GetNameText();
+		const FText Tooltip = RegisteredType.GetNameText();
+		TSharedPtr<FNiagaraMenuAction> Action(new FNiagaraMenuAction(
+			Category, DisplayName, Tooltip, 0, FText::GetEmpty(),
+			FNiagaraMenuAction::FOnExecuteStackAction::CreateLambda(
+			[Var=VarToModify, PropHandle=PropertyHandle, RegisteredType]()
+			{
+				FScopedTransaction Transaction(LOCTEXT("Set Raw Niagara Variable Type", "Set Variable Type"));
+				check(Var);
+				check(PropHandle.IsValid());
+				TArray<UObject*> Objects;
+				PropHandle->GetOuterObjects(Objects);
+				for (UObject* Obj : Objects)
+				{
+					Obj->Modify();
+				}
+
+				PropHandle->NotifyPreChange();
+				Var->SetType(RegisteredType);
+				PropHandle->NotifyPostChange(EPropertyChangeType::ValueSet);
+				PropHandle->NotifyFinishedChangingProperties();
+			}
+		)));
+
+		Collector.AddAction(Action, 0);
+	}
+
+	Collector.AddAllActionsTo(OutAllActions);
+}
+
+void FNiagaraVariableDetailsCustomization::CustomizeHeader(TSharedRef<IPropertyHandle> PropertyHandle, FDetailWidgetRow& HeaderRow, IPropertyTypeCustomizationUtils& CustomizationUtils)
+{
+	TSharedPtr<IPropertyHandle> NameHandle = PropertyHandle->GetChildHandle(TEXT("Name"));
+	TSharedPtr<IPropertyHandle> TypeDefHandleHandle = PropertyHandle->GetChildHandle(TEXT("TypeDefHandle"));
+	
+	TArray<UObject*> Objects;
+	PropertyHandle->GetOuterObjects(Objects);
+	if (Objects.Num() > 1)
+	{
+		HeaderRow.NameContent()
+		[
+			PropertyHandle->CreatePropertyNameWidget()
+		];
+		HeaderRow.ValueContent()
+		[
+			PropertyHandle->CreatePropertyValueWidget()
+		];
+	}
+	else
+	{		
+		void* VarPtr = nullptr;
+		if (PropertyHandle->GetValueData(VarPtr) == FPropertyAccess::Success)
+		{		
+			FNiagaraVariable* Var = (FNiagaraVariable*)VarPtr;
+
+			HeaderRow.NameContent()
+				.MaxDesiredWidth(200.f)
+				[
+					SNew(SHorizontalBox)
+					+ SHorizontalBox::Slot()
+					.FillWidth(0.9f)
+					.Padding(4.0f, 0.0f, 4.0f, 0.0f)
+					[
+						SNew(STextBlock)
+						.Text_Lambda([Var] {	return Var ? Var->GetType().GetNameText() : FText::GetEmpty(); })
+					]
+				+ SHorizontalBox::Slot()
+				.AutoWidth()
+				.Padding(4.0f, 0.0f, 4.0f, 0.0f)
+				[
+					SNew(SComboButton)
+					.ButtonStyle(FAppStyle::Get(), "RoundButton")
+					.ForegroundColor(FAppStyle::GetSlateColor("DefaultForeground"))
+					.ContentPadding(FMargin(2, 0))
+					.OnGetMenuContent(this, &FNiagaraVariableDetailsCustomization::GetTypeMenu, TypeDefHandleHandle, Var)
+					//.IsEnabled(this, &SNiagaraParameterPanel::GetCanAddParametersToCategory, Category)
+					.HAlign(HAlign_Center)
+					.VAlign(VAlign_Center)
+					.HasDownArrow(false)
+					.ButtonContent()
+					[
+						SNew(SHorizontalBox)
+						+ SHorizontalBox::Slot()
+						.AutoWidth()
+						.Padding(FMargin(0, 1))
+						[
+							SNew(SImage)
+							.Image(FAppStyle::GetBrush("Icons.Edit"))
+						]
+					]
+				]
+			];
+		
+			HeaderRow.ValueContent()
+			[
+				NameHandle->CreatePropertyValueWidget()
+			];					
+		}
+	}	
+}
+
+void FNiagaraVariableDetailsCustomization::CustomizeChildren(TSharedRef<IPropertyHandle> PropertyHandle, IDetailChildrenBuilder& ChildBuilder, IPropertyTypeCustomizationUtils& CustomizationUtils)
+{
+}
+
+TSharedRef<SWidget> FNiagaraVariableDetailsCustomization::GetTypeMenu(TSharedPtr<IPropertyHandle> InPropertyHandle, FNiagaraVariable* Var)
+{
+	return SNew(SNiagaraRawVariableTypeSelectMenu)
+	.PropertyHandle(InPropertyHandle)
+	.VarToModify(Var);
+}
+
+
 #undef LOCTEXT_NAMESPACE
diff --git a/Engine/Plugins/FX/Niagara/Source/NiagaraEditor/Private/Customizations/NiagaraTypeCustomizations.h b/Engine/Plugins/FX/Niagara/Source/NiagaraEditor/Private/Customizations/NiagaraTypeCustomizations.h
--- a/Engine/Plugins/FX/Niagara/Source/NiagaraEditor/Private/Customizations/NiagaraTypeCustomizations.h
+++ b/Engine/Plugins/FX/Niagara/Source/NiagaraEditor/Private/Customizations/NiagaraTypeCustomizations.h
@@ -438,3 +438,21 @@ public:
 	virtual void GetMaterialParameterInfos(UMaterialInterface* Material, TArray<FMaterialParameterInfo>& OutMaterialParameterInfos) const override;
 	TSharedRef<SWidget> OnGetStaticVariablelBindingNameMenuContent(TSharedPtr<IPropertyHandle> PropertyHandle) const;
 };
+
+//** Properties customization for FNiagaraVariables. */
+class FNiagaraVariableDetailsCustomization : public IPropertyTypeCustomization
+{
+public:
+	static TSharedRef<IPropertyTypeCustomization> MakeInstance()
+	{
+		return MakeShared<FNiagaraVariableDetailsCustomization>();
+	}
+
+	virtual void CustomizeHeader(TSharedRef<IPropertyHandle> PropertyHandle, FDetailWidgetRow& HeaderRow, IPropertyTypeCustomizationUtils& CustomizationUtils);
+
+	virtual void CustomizeChildren(TSharedRef<IPropertyHandle> PropertyHandle, IDetailChildrenBuilder& ChildBuilder, IPropertyTypeCustomizationUtils& CustomizationUtils);
+
+private:
+	TSharedRef<SWidget> GetTypeMenu(TSharedPtr<IPropertyHandle> InPropertyHandle, FNiagaraVariable* Var);
+};
+
diff --git a/Engine/Plugins/FX/Niagara/Source/NiagaraEditor/Private/NiagaraEditorModule.cpp b/Engine/Plugins/FX/Niagara/Source/NiagaraEditor/Private/NiagaraEditorModule.cpp
--- a/Engine/Plugins/FX/Niagara/Source/NiagaraEditor/Private/NiagaraEditorModule.cpp
+++ b/Engine/Plugins/FX/Niagara/Source/NiagaraEditor/Private/NiagaraEditorModule.cpp
@@ -992,6 +992,10 @@ void FNiagaraEditorModule::StartupModule()
 		FNiagaraDataInterfaceEmitterBinding::StaticStruct()->GetFName(),
 		FOnGetPropertyTypeCustomizationInstance::CreateStatic(&FNiagaraDataInterfaceEmitterBindingCustomization::MakeInstance));
 
+	PropertyModule.RegisterCustomPropertyTypeLayout(
+		FNiagaraVariable::StaticStruct()->GetFName(),
+		FOnGetPropertyTypeCustomizationInstance::CreateStatic(&FNiagaraVariableDetailsCustomization::MakeInstance));
+
 #if WITH_NIAGARA_DEBUGGER
 	PropertyModule.RegisterCustomPropertyTypeLayout(
 		FNiagaraDebugHUDVariable::StaticStruct()->GetFName(),
