Niagara Data Channels - Supporting CL

Fixed bug when trying to draw 0 debug lines.

diff --git a/Engine/Plugins/FX/Niagara/Source/Niagara/Private/NiagaraDebugHud.cpp b/Engine/Plugins/FX/Niagara/Source/Niagara/Private/NiagaraDebugHud.cpp
--- a/Engine/Plugins/FX/Niagara/Source/Niagara/Private/NiagaraDebugHud.cpp
+++ b/Engine/Plugins/FX/Niagara/Source/Niagara/Private/NiagaraDebugHud.cpp
@@ -1441,6 +1441,11 @@ struct FGraph
 
 	void DrawLine(FString Name, FLinearColor Color, const TArray<T>& Samples)
 	{
+		if(Samples.IsEmpty())
+		{
+			return;
+		}
+
 		FBatchedElements* BatchedElements = Canvas->GetBatchedElements(FCanvas::ET_Line);
 		BatchedElements->AddReserveLines(Samples.Num() - 1);
 		FHitProxyId HitProxyId = Canvas->GetHitProxyId();
