Removing duplicate FNigaraDouble struct coming from two directions.

diff --git a/Engine/Plugins/FX/Niagara/Source/Niagara/Public/NiagaraTypes.h b/Engine/Plugins/FX/Niagara/Source/Niagara/Public/NiagaraTypes.h
--- a/Engine/Plugins/FX/Niagara/Source/Niagara/Public/NiagaraTypes.h
+++ b/Engine/Plugins/FX/Niagara/Source/Niagara/Public/NiagaraTypes.h
@@ -154,16 +154,6 @@ struct FNiagaraNumeric
 	GENERATED_USTRUCT_BODY()
 };
 
-// only used for LWC conversions, not supported by Niagara yet
-USTRUCT(meta = (DisplayName = "double"))
-struct FNiagaraDouble
-{
-	GENERATED_USTRUCT_BODY()
-
-	UPROPERTY(EditAnywhere, Category = Parameters)
-	double Value = 0;
-};
-
 USTRUCT()
 struct FNiagaraParameterMap
 {
