- Emitter mode support for decal & light renderer

diff --git a/Engine/Plugins/FX/Niagara/Source/Niagara/Private/NiagaraDecalRendererProperties.cpp b/Engine/Plugins/FX/Niagara/Source/Niagara/Private/NiagaraDecalRendererProperties.cpp
--- a/Engine/Plugins/FX/Niagara/Source/Niagara/Private/NiagaraDecalRendererProperties.cpp
+++ b/Engine/Plugins/FX/Niagara/Source/Niagara/Private/NiagaraDecalRendererProperties.cpp
@@ -34,19 +34,19 @@ namespace NiagaraDecalRendererPropertiesLocal
 
 	static FNiagaraVariable GetDecalOrientationVariable()
 	{
-		static FNiagaraVariable Variable = MakeNiagaraVariableWithValue(FNiagaraTypeDefinition::GetQuatDef(), TEXT("Particles.DecalOrientation"), FRotator3f(-90.0f, 0.0f, 90.0f).Quaternion());
+		static FNiagaraVariable Variable = MakeNiagaraVariableWithValue(FNiagaraTypeDefinition::GetQuatDef(), TEXT("Particles.DecalOrientation"), UNiagaraDecalRendererProperties::GetDefaultOrientation());
 		return Variable;
 	}
 
 	static FNiagaraVariable GetDecalSizeVariable()
 	{
-		static FNiagaraVariable Variable = MakeNiagaraVariableWithValue(FNiagaraTypeDefinition::GetVec3Def(), TEXT("Particles.DecalSize"), FVector3f(50.0f, 50.0f, 50.0f));
+		static FNiagaraVariable Variable = MakeNiagaraVariableWithValue(FNiagaraTypeDefinition::GetVec3Def(), TEXT("Particles.DecalSize"), UNiagaraDecalRendererProperties::GetDefaultDecalSize());
 		return Variable;
 	}
 
 	static FNiagaraVariable GetDecalFadeVariable()
 	{
-		FNiagaraVariable Variable = MakeNiagaraVariableWithValue(FNiagaraTypeDefinition::GetFloatDef(), TEXT("Particles.DecalFade"), 1.0f);
+		FNiagaraVariable Variable = MakeNiagaraVariableWithValue(FNiagaraTypeDefinition::GetFloatDef(), TEXT("Particles.DecalFade"), UNiagaraDecalRendererProperties::GetDefaultDecalFade());
 		return Variable;
 	}
 
@@ -72,7 +72,7 @@ namespace NiagaraDecalRendererPropertiesLocal
 
 UNiagaraDecalRendererProperties::UNiagaraDecalRendererProperties()
 {
-	AttributeBindings.Reserve(5);
+	AttributeBindings.Reserve(6);
 	AttributeBindings.Add(&PositionBinding);
 	AttributeBindings.Add(&DecalOrientationBinding);
 	AttributeBindings.Add(&DecalSizeBinding);
@@ -109,6 +109,20 @@ void UNiagaraDecalRendererProperties::PostInitProperties()
 	}
 }
 
+#if WITH_EDITORONLY_DATA
+void UNiagaraDecalRendererProperties::PostEditChangeProperty(struct FPropertyChangedEvent& PropertyChangedEvent)
+{
+	Super::PostEditChangeProperty(PropertyChangedEvent);
+
+	const FName PropertyName = PropertyChangedEvent.GetPropertyName();
+
+	if (PropertyName == GET_MEMBER_NAME_CHECKED(UNiagaraDecalRendererProperties, SourceMode))
+	{
+		UpdateSourceModeDerivates(SourceMode, true);
+	}
+}
+#endif// WITH_EDITORONLY_DATA
+
 void UNiagaraDecalRendererProperties::InitCDOPropertiesAfterModuleStartup()
 {
 	using namespace NiagaraDecalRendererPropertiesLocal;
@@ -170,10 +184,10 @@ protected:
 
 FNiagaraBoundsCalculator* UNiagaraDecalRendererProperties::CreateBoundsCalculator()
 {
-	//if (GetCurrentSourceMode() == ENiagaraRendererSourceDataMode::Emitter)
-	//{
-	//	return nullptr;
-	//}
+	if (GetCurrentSourceMode() == ENiagaraRendererSourceDataMode::Emitter)
+	{
+		return nullptr;
+	}
 
 	return new FNiagaraBoundsCalculatorDecals(PositionDataSetAccessor, DecalSizeDataSetAccessor);
 }
@@ -188,6 +202,8 @@ void UNiagaraDecalRendererProperties::GetUsedMaterials(const FNiagaraEmitterInst
 
 void UNiagaraDecalRendererProperties::CacheFromCompiledData(const FNiagaraDataSetCompiledData* CompiledData)
 {
+	UpdateSourceModeDerivates(SourceMode);
+
 	PositionDataSetAccessor.Init(CompiledData, PositionBinding.GetDataSetBindableVariable().GetName());
 	DecalOrientationDataSetAccessor.Init(CompiledData, DecalOrientationBinding.GetDataSetBindableVariable().GetName());
 	DecalSizeDataSetAccessor.Init(CompiledData, DecalSizeBinding.GetDataSetBindableVariable().GetName());
@@ -196,7 +212,42 @@ void UNiagaraDecalRendererProperties::CacheFromCompiledData(const FNiagaraDataSe
 	RendererVisibilityTagAccessor.Init(CompiledData, RendererVisibilityTagBinding.GetDataSetBindableVariable().GetName());
 }
 
+void UNiagaraDecalRendererProperties::UpdateSourceModeDerivates(ENiagaraRendererSourceDataMode InSourceMode, bool bFromPropertyEdit)
+{
+	Super::UpdateSourceModeDerivates(InSourceMode, bFromPropertyEdit);
+}
+
+bool UNiagaraDecalRendererProperties::PopulateRequiredBindings(FNiagaraParameterStore& InParameterStore)
+{
+	bool bAnyAdded = Super::PopulateRequiredBindings(InParameterStore);
+
+	for (const FNiagaraVariableAttributeBinding* Binding : AttributeBindings)
+	{
+		if (Binding && Binding->CanBindToHostParameterMap())
+		{
+			InParameterStore.AddParameter(Binding->GetParamMapBindableVariable(), false);
+			bAnyAdded = true;
+		}
+	}
+
+	return bAnyAdded;
+}
+
 #if WITH_EDITORONLY_DATA
+bool UNiagaraDecalRendererProperties::IsSupportedVariableForBinding(const FNiagaraVariableBase& InSourceForBinding, const FName& InTargetBindingName) const
+{
+	switch (SourceMode)
+	{
+		case ENiagaraRendererSourceDataMode::Particles:
+			return InSourceForBinding.IsInNameSpace(FNiagaraConstants::ParticleAttributeNamespaceString);
+
+		default:
+			return
+				InSourceForBinding.IsInNameSpace(FNiagaraConstants::UserNamespaceString) ||
+				InSourceForBinding.IsInNameSpace(FNiagaraConstants::SystemNamespaceString) ||
+				InSourceForBinding.IsInNameSpace(FNiagaraConstants::EmitterNamespaceString);
+	}
+}
 
 const TArray<FNiagaraVariable>& UNiagaraDecalRendererProperties::GetOptionalAttributes()
 {
diff --git a/Engine/Plugins/FX/Niagara/Source/Niagara/Private/NiagaraLightRendererProperties.cpp b/Engine/Plugins/FX/Niagara/Source/Niagara/Private/NiagaraLightRendererProperties.cpp
--- a/Engine/Plugins/FX/Niagara/Source/Niagara/Private/NiagaraLightRendererProperties.cpp
+++ b/Engine/Plugins/FX/Niagara/Source/Niagara/Private/NiagaraLightRendererProperties.cpp
@@ -75,6 +75,20 @@ void UNiagaraLightRendererProperties::PostInitProperties()
 	}
 }
 
+#if WITH_EDITORONLY_DATA
+void UNiagaraLightRendererProperties::PostEditChangeProperty(struct FPropertyChangedEvent& PropertyChangedEvent)
+{
+	Super::PostEditChangeProperty(PropertyChangedEvent);
+
+	const FName PropertyName = PropertyChangedEvent.GetPropertyName();
+
+	if (PropertyName == GET_MEMBER_NAME_CHECKED(UNiagaraLightRendererProperties, SourceMode))
+	{
+		UpdateSourceModeDerivates(SourceMode, true);
+	}
+}
+#endif// WITH_EDITORONLY_DATA
+
 /** The bindings depend on variables that are created during the NiagaraModule startup. However, the CDO's are build prior to this being initialized, so we defer setting these values until later.*/
 void UNiagaraLightRendererProperties::InitCDOPropertiesAfterModuleStartup()
 {
@@ -169,6 +183,8 @@ void UNiagaraLightRendererProperties::GetUsedMaterials(const FNiagaraEmitterInst
 
 void UNiagaraLightRendererProperties::CacheFromCompiledData(const FNiagaraDataSetCompiledData* CompiledData)
 {
+	UpdateSourceModeDerivates(SourceMode);
+
 	PositionDataSetAccessor.Init(CompiledData, PositionBinding.GetDataSetBindableVariable().GetName());
 	ColorDataSetAccessor.Init(CompiledData, ColorBinding.GetDataSetBindableVariable().GetName());
 	RadiusDataSetAccessor.Init(CompiledData, RadiusBinding.GetDataSetBindableVariable().GetName());
@@ -178,7 +194,42 @@ void UNiagaraLightRendererProperties::CacheFromCompiledData(const FNiagaraDataSe
 	RendererVisibilityTagAccessor.Init(CompiledData, RendererVisibilityTagBinding.GetDataSetBindableVariable().GetName());
 }
 
+void UNiagaraLightRendererProperties::UpdateSourceModeDerivates(ENiagaraRendererSourceDataMode InSourceMode, bool bFromPropertyEdit)
+{
+	Super::UpdateSourceModeDerivates(InSourceMode, bFromPropertyEdit);
+}
+
+bool UNiagaraLightRendererProperties::PopulateRequiredBindings(FNiagaraParameterStore& InParameterStore)
+{
+	bool bAnyAdded = Super::PopulateRequiredBindings(InParameterStore);
+
+	for (const FNiagaraVariableAttributeBinding* Binding : AttributeBindings)
+	{
+		if (Binding && Binding->CanBindToHostParameterMap())
+		{
+			InParameterStore.AddParameter(Binding->GetParamMapBindableVariable(), false);
+			bAnyAdded = true;
+		}
+	}
+
+	return bAnyAdded;
+}
+
 #if WITH_EDITORONLY_DATA
+bool UNiagaraLightRendererProperties::IsSupportedVariableForBinding(const FNiagaraVariableBase& InSourceForBinding, const FName& InTargetBindingName) const
+{
+	switch (SourceMode)
+	{
+		case ENiagaraRendererSourceDataMode::Particles:
+			return InSourceForBinding.IsInNameSpace(FNiagaraConstants::ParticleAttributeNamespaceString);
+
+		default:
+			return
+				InSourceForBinding.IsInNameSpace(FNiagaraConstants::UserNamespaceString) ||
+				InSourceForBinding.IsInNameSpace(FNiagaraConstants::SystemNamespaceString) ||
+				InSourceForBinding.IsInNameSpace(FNiagaraConstants::EmitterNamespaceString);
+	}
+}
 
 const TArray<FNiagaraVariable>& UNiagaraLightRendererProperties::GetOptionalAttributes()
 {
diff --git a/Engine/Plugins/FX/Niagara/Source/Niagara/Private/NiagaraRendererDecals.cpp b/Engine/Plugins/FX/Niagara/Source/Niagara/Private/NiagaraRendererDecals.cpp
--- a/Engine/Plugins/FX/Niagara/Source/Niagara/Private/NiagaraRendererDecals.cpp
+++ b/Engine/Plugins/FX/Niagara/Source/Niagara/Private/NiagaraRendererDecals.cpp
@@ -24,8 +24,6 @@ namespace NiagaraRendererDecalsLocal
 	);
 
 #if UE_BUILD_SHIPPING
-	void DrawDebugDecal(UWorld* World, const FDeferredDecalUpdateParams& UpdateParams, float Fade) {}
-#else
 	static bool GDrawDebug = false;
 	static FAutoConsoleVariableRef CVarDrawDebug(
 		TEXT("fx.Niagara.DecalRenderer.DrawDebug"),
@@ -34,7 +32,7 @@ namespace NiagaraRendererDecalsLocal
 		ECVF_Default
 	);
 
-	void DrawDebugDecal(UWorld* World, const FDeferredDecalUpdateParams& UpdateParams, float Fade)
+	void DrawDebugDecal(UWorld* World, TConstArrayView<FDeferredDecalUpdateParams> AllUpdateParams, double WorldTime)
 	{
 		if (GDrawDebug == false)
 		{
@@ -44,7 +42,7 @@ namespace NiagaraRendererDecalsLocal
 		// Line Batcher is not thread safe and we could be on any thread, therefore send a task to draw the information safely
 		AsyncTask(
 			ENamedThreads::GameThread,
-			[WeakWorld=MakeWeakObjectPtr(World), UpdateParams, Fade]
+			[WeakWorld=MakeWeakObjectPtr(World), AllUpdateParams_GT=TArray<FDeferredDecalUpdateParams>(AllUpdateParams), WorldTime]
 			{
 				UWorld* World = WeakWorld.Get();
 				if (World == nullptr || World->LineBatcher == nullptr)
@@ -53,14 +51,23 @@ namespace NiagaraRendererDecalsLocal
 				}
 				ULineBatchComponent* LineBatcher = World->LineBatcher;
 				
-				const FBox BoundsBox(-FVector(1), FVector(1));
-				const uint8 IntFade = FMath::Clamp(uint8(255.0f * Fade), 0, 255);
-				LineBatcher->DrawSolidBox(BoundsBox, UpdateParams.Transform, FColor(IntFade, IntFade, IntFade, 128), 0, 0.0f);
-
-				LineBatcher->DrawCircle(UpdateParams.Bounds.Origin, FVector::XAxisVector, FVector::YAxisVector, FColor::Red, UpdateParams.Bounds.SphereRadius, 16, 0);
-				LineBatcher->DrawCircle(UpdateParams.Bounds.Origin, FVector::XAxisVector, FVector::ZAxisVector, FColor::Red, UpdateParams.Bounds.SphereRadius, 16, 0);
-				LineBatcher->DrawCircle(UpdateParams.Bounds.Origin, FVector::YAxisVector, FVector::ZAxisVector, FColor::Red, UpdateParams.Bounds.SphereRadius, 16, 0);
-		}
+				for (const FDeferredDecalUpdateParams& UpdateParams : AllUpdateParams_GT)
+				{
+					if ( UpdateParams.OperationType == FDeferredDecalUpdateParams::EOperationType::RemoveFromSceneAndDelete )
+					{
+						continue;
+					}
+
+					const float Fade = (WorldTime - UpdateParams.FadeStartDelay) / UpdateParams.FadeDuration;
+					const FBox BoundsBox(-FVector(1), FVector(1));
+					const uint8 IntFade = FMath::Clamp(uint8(255.0f * Fade), 0, 255);
+					LineBatcher->DrawSolidBox(BoundsBox, UpdateParams.Transform, FColor(IntFade, IntFade, IntFade, 128), 0, 0.0f);
+
+					LineBatcher->DrawCircle(UpdateParams.Bounds.Origin, FVector::XAxisVector, FVector::YAxisVector, FColor::Red, UpdateParams.Bounds.SphereRadius, 16, 0);
+					LineBatcher->DrawCircle(UpdateParams.Bounds.Origin, FVector::XAxisVector, FVector::ZAxisVector, FColor::Red, UpdateParams.Bounds.SphereRadius, 16, 0);
+					LineBatcher->DrawCircle(UpdateParams.Bounds.Origin, FVector::YAxisVector, FVector::ZAxisVector, FColor::Red, UpdateParams.Bounds.SphereRadius, 16, 0);
+				}
+			}
 		);
 	}
 #endif
@@ -109,7 +116,7 @@ void FNiagaraRendererDecals::DestroyRenderState_Concurrent()
 FPrimitiveViewRelevance FNiagaraRendererDecals::GetViewRelevance(const FSceneView* View, const FNiagaraSceneProxy *SceneProxy)const
 {
 	FPrimitiveViewRelevance Result;
-	Result.bDrawRelevance = bHasLights && SceneProxy->IsShown(View) && View->Family->EngineShowFlags.Particles && View->Family->EngineShowFlags.Niagara;
+	Result.bDrawRelevance = false;
 	Result.bShadowRelevance = false;
 	Result.bDynamicRelevance = false;
 	Result.bOpaque = false;
@@ -144,8 +151,9 @@ FNiagaraDynamicDataBase* FNiagaraRendererDecals::GenerateDynamicData(const FNiag
 	const FTransform LocalToWorld = SystemInstance->GetWorldTransform();
 	const FNiagaraLWCConverter LwcConverter = SystemInstance->GetLWCConverter(bUseLocalSpace);
 	const FNiagaraPosition DefaultPos = bUseLocalSpace ? FVector::ZeroVector : LocalToWorld.GetLocation();
-	const FQuat4f DefaultRot = FRotator3f(-90.0f, 0.0f, 90.0f).Quaternion();
-	const FVector3f DefaultSize = FVector3f(50.0f, 50.0f, 50.0f);	//-TODO: Read defaults from properties or default variable
+	const FQuat4f DefaultRot = UNiagaraDecalRendererProperties::GetDefaultOrientation();
+	const FVector3f DefaultSize = UNiagaraDecalRendererProperties::GetDefaultDecalSize();
+	const float DefaultFade = UNiagaraDecalRendererProperties::GetDefaultDecalFade();
 
 	// Check for Material Update
 	UMaterialInterface* Material = RendererProperties->GetMaterial(Emitter);
@@ -155,75 +163,117 @@ FNiagaraDynamicDataBase* FNiagaraRendererDecals::GenerateDynamicData(const FNiag
 		ReleaseAllDecals();
 	}
 
-	// Create all our data readers
-	const auto PositionReader	= RendererProperties->PositionDataSetAccessor.GetReader(DataSet);
-	const auto RotationReader	= RendererProperties->DecalOrientationDataSetAccessor.GetReader(DataSet);
-	const auto SizeReader		= RendererProperties->DecalSizeDataSetAccessor.GetReader(DataSet);
-	const auto ColorReader		= RendererProperties->ColorDataSetAccessor.GetReader(DataSet);
-	const auto FadeReader		= RendererProperties->DecalFadeDataSetAccessor.GetReader(DataSet);
-	const auto VisTagReader		= RendererProperties->RendererVisibilityTagAccessor.GetReader(DataSet);
+	// Generate list of decals to update
+	int NumActiveDecalProxies = 0;
 
 	TArray<FDeferredDecalUpdateParams> DecalUpdates;
 	DecalUpdates.Reserve(DataToRender->GetNumInstances());
 
-	int NumActiveDecalProxies = 0;
-	for (uint32 ParticleIndex=0; ParticleIndex < DataToRender->GetNumInstances(); ++ParticleIndex)
+	// Particles Source mode?
+	if (RendererProperties->SourceMode == ENiagaraRendererSourceDataMode::Particles)
 	{
-		// Check visibility tag
-		const int32 VisTag = VisTagReader.GetSafe(ParticleIndex, RendererProperties->RendererVisibility);
-		if (VisTag != RendererProperties->RendererVisibility)
+		// Create all our data readers
+		const auto PositionReader = RendererProperties->PositionDataSetAccessor.GetReader(DataSet);
+		const auto RotationReader = RendererProperties->DecalOrientationDataSetAccessor.GetReader(DataSet);
+		const auto SizeReader = RendererProperties->DecalSizeDataSetAccessor.GetReader(DataSet);
+		const auto ColorReader = RendererProperties->ColorDataSetAccessor.GetReader(DataSet);
+		const auto FadeReader = RendererProperties->DecalFadeDataSetAccessor.GetReader(DataSet);
+		const auto VisTagReader = RendererProperties->RendererVisibilityTagAccessor.GetReader(DataSet);
+
+		for (uint32 ParticleIndex = 0; ParticleIndex < DataToRender->GetNumInstances(); ++ParticleIndex)
 		{
-			continue;
-		}
+			// Check visibility tag
+			const int32 VisTag = VisTagReader.GetSafe(ParticleIndex, RendererProperties->RendererVisibility);
+			if (VisTag != RendererProperties->RendererVisibility)
+			{
+				continue;
+			}
 
-		// Grab Decal Attributes
-		const FVector3f SimPos	= PositionReader.GetSafe(ParticleIndex, DefaultPos);
-		const FQuat4f SimRot	= RotationReader.GetSafe(ParticleIndex, DefaultRot);
-		const FVector Position	= bUseLocalSpace ? LocalToWorld.TransformPosition(FVector(SimPos)) : LwcConverter.ConvertSimulationPositionToWorld(SimPos);
-		const FQuat Rotation	= bUseLocalSpace ? LocalToWorld.TransformRotation(FQuat(SimRot)) : FQuat(SimRot);
-		const FVector Size		= FVector(SizeReader.GetSafe(ParticleIndex, DefaultSize) * 0.5f);
-		const float Fade		= RendererProperties->bUseColorBindingAsFade ? ColorReader.GetSafe(ParticleIndex, FLinearColor::White).A : FadeReader.GetSafe(ParticleIndex, 1.0f);
+			// Grab Decal Attributes
+			const FVector3f SimPos = PositionReader.GetSafe(ParticleIndex, DefaultPos);
+			const FQuat4f SimRot = RotationReader.GetSafe(ParticleIndex, DefaultRot);
+			const FVector Position = bUseLocalSpace ? LocalToWorld.TransformPosition(FVector(SimPos)) : LwcConverter.ConvertSimulationPositionToWorld(SimPos);
+			const FQuat Rotation = bUseLocalSpace ? LocalToWorld.TransformRotation(FQuat(SimRot)) : FQuat(SimRot);
+			const FVector Size = FVector(SizeReader.GetSafe(ParticleIndex, DefaultSize) * 0.5f);
+			const float Fade = RendererProperties->bUseColorBindingAsFade ? ColorReader.GetSafe(ParticleIndex, FLinearColor::White).A : FadeReader.GetSafe(ParticleIndex, DefaultFade);
 
-		// Create Update Parameters
-		FDeferredDecalUpdateParams& UpdateParams = DecalUpdates.AddDefaulted_GetRef();
-		if (NumActiveDecalProxies < ActiveDecalProxies.Num())
-		{
-			UpdateParams.DecalProxy = ActiveDecalProxies[NumActiveDecalProxies];
-			UpdateParams.OperationType = FDeferredDecalUpdateParams::EOperationType::Update;
+			// Create Update Parameters
+			FDeferredDecalUpdateParams& UpdateParams = DecalUpdates.AddDefaulted_GetRef();
+			if (NumActiveDecalProxies < ActiveDecalProxies.Num())
+			{
+				UpdateParams.DecalProxy = ActiveDecalProxies[NumActiveDecalProxies];
+				UpdateParams.OperationType = FDeferredDecalUpdateParams::EOperationType::Update;
+			}
+			else
+			{
+				UpdateParams.DecalProxy = new FDeferredDecalProxy(OwnerComponent, Material);
+				UpdateParams.OperationType = FDeferredDecalUpdateParams::EOperationType::AddToSceneAndUpdate;
+				ActiveDecalProxies.Add(UpdateParams.DecalProxy);
+			}
+			++NumActiveDecalProxies;
+
+			UpdateParams.Transform = FTransform(Rotation, Position, FVector(Size));
+			UpdateParams.Bounds = FBoxSphereBounds(FSphere(Position, Size.GetAbsMax() * 2.0));
+			UpdateParams.AbsSpawnTime = World->TimeSeconds - FMath::Clamp(1.0f - Fade, 0.0f, 1.0f);
+			UpdateParams.FadeStartDelay = 0.0f;
+			UpdateParams.FadeDuration = 1.0f;
 		}
-		else
+	}
+	// Emitter source mode
+	else
+	{
+		//-OPT: Emitter mode we should be able to cache the parameter offset rather than search each frame
+		const FNiagaraParameterStore& ParameterStore = Emitter->GetRendererBoundVariables();
+		const int32 VisTag = ParameterStore.GetParameterValueOrDefault(RendererProperties->RendererVisibilityTagBinding.GetParamMapBindableVariable(), RendererProperties->RendererVisibility);
+		if (VisTag == RendererProperties->RendererVisibility)
 		{
-			UpdateParams.DecalProxy = new FDeferredDecalProxy(OwnerComponent, Material);
-			UpdateParams.OperationType = FDeferredDecalUpdateParams::EOperationType::AddToSceneAndUpdate;
-			ActiveDecalProxies.Add(UpdateParams.DecalProxy);
-		}
-		++NumActiveDecalProxies;
-
-		UpdateParams.Transform		= FTransform(Rotation, Position, FVector(Size));
-		UpdateParams.Bounds			= FBoxSphereBounds(FSphere(Position, Size.GetAbsMax() * 2.0));
-		UpdateParams.AbsSpawnTime	= World->TimeSeconds - FMath::Clamp(1.0f - Fade, 0.0f, 1.0f);
-		UpdateParams.FadeStartDelay	= 0.0f;
-		UpdateParams.FadeDuration	= 1.0f;
+			const FVector3f SimPos = ParameterStore.GetParameterValueOrDefault(RendererProperties->PositionBinding.GetParamMapBindableVariable(), DefaultPos);
+			const FQuat4f SimRot = ParameterStore.GetParameterValueOrDefault(RendererProperties->DecalOrientationBinding.GetParamMapBindableVariable(), DefaultRot);
+			const FVector Position = bUseLocalSpace ? LocalToWorld.TransformPosition(FVector(SimPos)) : LwcConverter.ConvertSimulationPositionToWorld(SimPos);
+			const FQuat Rotation = bUseLocalSpace ? LocalToWorld.TransformRotation(FQuat(SimRot)) : FQuat(SimRot);
+			const FVector Size = FVector(ParameterStore.GetParameterValueOrDefault(RendererProperties->DecalSizeBinding.GetParamMapBindableVariable(), DefaultSize) * 0.5f);
+			const float Fade =
+				RendererProperties->bUseColorBindingAsFade ?
+				ParameterStore.GetParameterValueOrDefault(RendererProperties->ColorBinding.GetParamMapBindableVariable(), FLinearColor::White).A :
+				ParameterStore.GetParameterValueOrDefault(RendererProperties->DecalFadeBinding.GetParamMapBindableVariable(), DefaultFade);
+
+			// Create Update Parameters
+			FDeferredDecalUpdateParams& UpdateParams = DecalUpdates.AddDefaulted_GetRef();
+			if (NumActiveDecalProxies < ActiveDecalProxies.Num())
+			{
+				UpdateParams.DecalProxy = ActiveDecalProxies[NumActiveDecalProxies];
+				UpdateParams.OperationType = FDeferredDecalUpdateParams::EOperationType::Update;
+			}
+			else
+			{
+				UpdateParams.DecalProxy = new FDeferredDecalProxy(OwnerComponent, Material);
+				UpdateParams.OperationType = FDeferredDecalUpdateParams::EOperationType::AddToSceneAndUpdate;
+				ActiveDecalProxies.Add(UpdateParams.DecalProxy);
+			}
+			++NumActiveDecalProxies;
 
-		DrawDebugDecal(World, UpdateParams, Fade);
+			UpdateParams.Transform = FTransform(Rotation, Position, FVector(Size));
+			UpdateParams.Bounds = FBoxSphereBounds(FSphere(Position, Size.GetAbsMax() * 2.0));
+			UpdateParams.AbsSpawnTime = World->TimeSeconds - FMath::Clamp(1.0f - Fade, 0.0f, 1.0f);
+			UpdateParams.FadeStartDelay = 0.0f;
+			UpdateParams.FadeDuration = 1.0f;
+		}
 	}
 
 	// Remove any unused decals
-	// Note: We don't delete at the RemoveDecal call will automatically do this
-	const int NumToRemove = ActiveDecalProxies.Num() - NumActiveDecalProxies;
-	if (NumToRemove > 0)
+	while (ActiveDecalProxies.Num() > NumActiveDecalProxies)
 	{
-		for (int i = DataToRender->GetNumInstances(); i < ActiveDecalProxies.Num(); ++i)
-		{
-			FDeferredDecalUpdateParams& UpdateParams = DecalUpdates.AddDefaulted_GetRef();
-			UpdateParams.OperationType = FDeferredDecalUpdateParams::EOperationType::RemoveFromSceneAndDelete;
-			UpdateParams.DecalProxy = ActiveDecalProxies.Pop();
-		}
+		FDeferredDecalUpdateParams& UpdateParams = DecalUpdates.AddDefaulted_GetRef();
+		UpdateParams.OperationType = FDeferredDecalUpdateParams::EOperationType::RemoveFromSceneAndDelete;
+		UpdateParams.DecalProxy = ActiveDecalProxies.Pop();
 	}
 
 	// Send updates to RT
 	if (DecalUpdates.Num() > 0)
 	{
+	#if UE_BUILD_SHIPPING
+		DrawDebugDecal(World, DecalUpdates, World->TimeSeconds);
+	#endif
 		World->Scene->BatchUpdateDecals(MoveTemp(DecalUpdates));
 	}
 
diff --git a/Engine/Plugins/FX/Niagara/Source/Niagara/Private/NiagaraRendererLights.cpp b/Engine/Plugins/FX/Niagara/Source/Niagara/Private/NiagaraRendererLights.cpp
--- a/Engine/Plugins/FX/Niagara/Source/Niagara/Private/NiagaraRendererLights.cpp
+++ b/Engine/Plugins/FX/Niagara/Source/Niagara/Private/NiagaraRendererLights.cpp
@@ -92,21 +92,12 @@ FNiagaraDynamicDataBase* FNiagaraRendererLights::GenerateDynamicData(const FNiag
 	FNiagaraSystemInstance* SystemInstance = Emitter->GetParentSystemInstance();
 	FNiagaraDynamicDataLights* DynamicData = new FNiagaraDynamicDataLights(Emitter);
 
-	//I'm not a great fan of pulling scalar components out to a structured vert buffer like this.
-	//TODO: Experiment with a new VF that reads the data directly from the scalar layout.
-	const auto PositionReader = Properties->PositionDataSetAccessor.GetReader(Data);
-	const auto ColorReader = Properties->ColorDataSetAccessor.GetReader(Data);
-	const auto RadiusReader = Properties->RadiusDataSetAccessor.GetReader(Data);
-	const auto ExponentReader = Properties->ExponentDataSetAccessor.GetReader(Data);
-	const auto ScatteringReader = Properties->ScatteringDataSetAccessor.GetReader(Data);
-	const auto EnabledReader = Properties->EnabledDataSetAccessor.GetReader(Data);
-	const auto VisTagReader = Properties->RendererVisibilityTagAccessor.GetReader(Data);
 
 	// This used to use Proxy->GetLocalToWorld(), but that's a bad thing to do here, because the proxy gets updated on the render thread,
 	// and this function happens during EndOfFrame updates. So instead, use the most up-to-date transform here (fixes local-space frame-behind issues)
-	FTransform LocalToWorld = SystemInstance->GetWorldTransform();
+	const FTransform LocalToWorld = SystemInstance->GetWorldTransform();
 
-	bool bUseLocalSpace = UseLocalSpace(Proxy);
+	const bool bUseLocalSpace = UseLocalSpace(Proxy);
 	FNiagaraLWCConverter LwcConverter = SystemInstance->GetLWCConverter(bUseLocalSpace);
 	const FLinearColor DefaultColor = Properties->ColorBinding.GetDefaultValue<FLinearColor>();
 	const FNiagaraPosition DefaultPos = bUseLocalSpace ? FVector::ZeroVector : LocalToWorld.GetLocation();
@@ -118,28 +109,70 @@ FNiagaraDynamicDataBase* FNiagaraRendererLights::GenerateDynamicData(const FNiag
 
 	const float InverseExposureBlend = Properties->bOverrideInverseExposureBlend ? Properties->InverseExposureBlend : GetDefault<UNiagaraSettings>()->DefaultLightInverseExposureBlend;
 
-	for (uint32 ParticleIndex = 0; ParticleIndex < DataToRender->GetNumInstances(); ParticleIndex++)
+	// Particles Source mode?
+	if (Properties->SourceMode == ENiagaraRendererSourceDataMode::Particles)
 	{
-		bool bShouldRenderParticleLight = EnabledReader.GetSafe(ParticleIndex, DefaultEnabled).GetValue();
-		if (bShouldRenderParticleLight && VisTagReader.IsValid())
+		//I'm not a great fan of pulling scalar components out to a structured vert buffer like this.
+		//TODO: Experiment with a new VF that reads the data directly from the scalar layout.
+		const auto PositionReader = Properties->PositionDataSetAccessor.GetReader(Data);
+		const auto ColorReader = Properties->ColorDataSetAccessor.GetReader(Data);
+		const auto RadiusReader = Properties->RadiusDataSetAccessor.GetReader(Data);
+		const auto ExponentReader = Properties->ExponentDataSetAccessor.GetReader(Data);
+		const auto ScatteringReader = Properties->ScatteringDataSetAccessor.GetReader(Data);
+		const auto EnabledReader = Properties->EnabledDataSetAccessor.GetReader(Data);
+		const auto VisTagReader = Properties->RendererVisibilityTagAccessor.GetReader(Data);
+
+		for (uint32 ParticleIndex = 0; ParticleIndex < DataToRender->GetNumInstances(); ParticleIndex++)
 		{
-			bShouldRenderParticleLight = VisTagReader.GetSafe(ParticleIndex, DefaultVisibilityTag) == Properties->RendererVisibility;
+			bool bShouldRenderParticleLight = EnabledReader.GetSafe(ParticleIndex, DefaultEnabled).GetValue();
+			if (bShouldRenderParticleLight && VisTagReader.IsValid())
+			{
+				bShouldRenderParticleLight = VisTagReader.GetSafe(ParticleIndex, DefaultVisibilityTag) == Properties->RendererVisibility;
+			}
+			float LightRadius = RadiusReader.GetSafe(ParticleIndex, DefaultRadius) * Properties->RadiusScale;
+			if (bShouldRenderParticleLight && LightRadius > 0)
+			{
+				SimpleLightData& LightData = DynamicData->LightArray.AddDefaulted_GetRef();
+
+				const FLinearColor Color = ColorReader.GetSafe(ParticleIndex, DefaultColor);
+				const float Brightness = Properties->bAlphaScalesBrightness ? Color.A : 1.0f;
+
+				LightData.LightEntry.Radius = LightRadius;
+				LightData.LightEntry.Color = FVector3f(Color) * Brightness + Properties->ColorAdd;
+				LightData.LightEntry.Exponent = Properties->bUseInverseSquaredFalloff ? 0 : ExponentReader.GetSafe(ParticleIndex, DefaultExponent);
+				LightData.LightEntry.InverseExposureBlend = InverseExposureBlend;
+				LightData.LightEntry.bAffectTranslucency = Properties->bAffectsTranslucency;
+				LightData.LightEntry.VolumetricScatteringIntensity = ScatteringReader.GetSafe(ParticleIndex, DefaultScattering);
+				LightData.PerViewEntry.Position = LwcConverter.ConvertSimulationPositionToWorld(PositionReader.GetSafe(ParticleIndex, DefaultPos));
+				if (bUseLocalSpace)
+				{
+					LightData.PerViewEntry.Position = LocalToWorld.TransformPosition(LightData.PerViewEntry.Position);
+				}
+			}
 		}
-		float LightRadius = RadiusReader.GetSafe(ParticleIndex, DefaultRadius) * Properties->RadiusScale;
-		if (bShouldRenderParticleLight && LightRadius > 0)
+	}
+	else
+	{
+		const FNiagaraParameterStore& ParameterStore = Emitter->GetRendererBoundVariables();
+		const bool bEnabled = ParameterStore.GetParameterValueOrDefault(Properties->LightRenderingEnabledBinding.GetParamMapBindableVariable(), FNiagaraBool(true)).GetValue();
+		const int32 VisTag = ParameterStore.GetParameterValueOrDefault(Properties->RendererVisibilityTagBinding.GetParamMapBindableVariable(), Properties->RendererVisibility);
+		const float LightRadius = ParameterStore.GetParameterValueOrDefault(Properties->RadiusBinding.GetParamMapBindableVariable(), DefaultRadius) * Properties->RadiusScale;
+		if (bEnabled && VisTag == Properties->RendererVisibility && LightRadius > 0.0f)
 		{
-			SimpleLightData& LightData = DynamicData->LightArray.AddDefaulted_GetRef();
-
-			const FLinearColor Color = ColorReader.GetSafe(ParticleIndex, DefaultColor);
-			const float Brightness = Properties->bAlphaScalesBrightness ? Color.A : 1.0f;
+			const FVector3f SimPos = ParameterStore.GetParameterValueOrDefault(Properties->PositionBinding.GetParamMapBindableVariable(), DefaultPos);
+			const FLinearColor LightColor = ParameterStore.GetParameterValueOrDefault(Properties->ColorBinding.GetParamMapBindableVariable(), DefaultColor);
+			const float LightExponent = ParameterStore.GetParameterValueOrDefault(Properties->LightExponentBinding.GetParamMapBindableVariable(), DefaultExponent);
+			const float LightScattering = ParameterStore.GetParameterValueOrDefault(Properties->VolumetricScatteringBinding.GetParamMapBindableVariable(), DefaultScattering);
+			const float Brightness = Properties->bAlphaScalesBrightness ? LightColor.A : 1.0f;
 
+			SimpleLightData& LightData = DynamicData->LightArray.AddDefaulted_GetRef();
 			LightData.LightEntry.Radius = LightRadius;
-			LightData.LightEntry.Color = FVector3f(Color) * Brightness + Properties->ColorAdd;
-			LightData.LightEntry.Exponent = Properties->bUseInverseSquaredFalloff ? 0 : ExponentReader.GetSafe(ParticleIndex, DefaultExponent);
+			LightData.LightEntry.Color = FVector3f(LightColor) * Brightness + Properties->ColorAdd;
+			LightData.LightEntry.Exponent = Properties->bUseInverseSquaredFalloff ? 0 : LightExponent;
 			LightData.LightEntry.InverseExposureBlend = InverseExposureBlend;
 			LightData.LightEntry.bAffectTranslucency = Properties->bAffectsTranslucency;
-			LightData.LightEntry.VolumetricScatteringIntensity = ScatteringReader.GetSafe(ParticleIndex, DefaultScattering);
-			LightData.PerViewEntry.Position = LwcConverter.ConvertSimulationPositionToWorld(PositionReader.GetSafe(ParticleIndex, DefaultPos));
+			LightData.LightEntry.VolumetricScatteringIntensity = LightScattering;
+			LightData.PerViewEntry.Position = LwcConverter.ConvertSimulationPositionToWorld(SimPos);
 			if (bUseLocalSpace)
 			{
 				LightData.PerViewEntry.Position = LocalToWorld.TransformPosition(LightData.PerViewEntry.Position);
diff --git a/Engine/Plugins/FX/Niagara/Source/Niagara/Public/NiagaraDecalRendererProperties.h b/Engine/Plugins/FX/Niagara/Source/Niagara/Public/NiagaraDecalRendererProperties.h
--- a/Engine/Plugins/FX/Niagara/Source/Niagara/Public/NiagaraDecalRendererProperties.h
+++ b/Engine/Plugins/FX/Niagara/Source/Niagara/Public/NiagaraDecalRendererProperties.h
@@ -25,6 +25,9 @@ public:
 	//UObject Interface
 	virtual void PostLoad() override;
 	virtual void PostInitProperties() override;
+#if WITH_EDITORONLY_DATA
+	virtual void PostEditChangeProperty(struct FPropertyChangedEvent& PropertyChangedEvent) override;
+#endif// WITH_EDITORONLY_DATA
 	//UObject Interface END
 
 	static void InitCDOPropertiesAfterModuleStartup();
@@ -35,16 +38,24 @@ public:
 	virtual void GetUsedMaterials(const FNiagaraEmitterInstance* InEmitter, TArray<UMaterialInterface*>& OutMaterials) const override;
 	virtual bool IsSimTargetSupported(ENiagaraSimTarget InSimTarget) const override { return InSimTarget == ENiagaraSimTarget::CPUSim; };
 #if WITH_EDITORONLY_DATA
+	virtual bool IsSupportedVariableForBinding(const FNiagaraVariableBase& InSourceForBinding, const FName& InTargetBindingName) const override;
 	virtual const TArray<FNiagaraVariable>& GetOptionalAttributes() override;
 	virtual void GetRendererWidgets(const FNiagaraEmitterInstance* InEmitter, TArray<TSharedPtr<SWidget>>& OutWidgets, TSharedPtr<FAssetThumbnailPool> InThumbnailPool) const override;
 	virtual void GetRendererTooltipWidgets(const FNiagaraEmitterInstance* InEmitter, TArray<TSharedPtr<SWidget>>& OutWidgets, TSharedPtr<FAssetThumbnailPool> InThumbnailPool) const override;
 	virtual void GetRendererFeedback(const FVersionedNiagaraEmitter& InEmitter, TArray<FText>& OutErrors, TArray<FText>& OutWarnings, TArray<FText>& OutInfo) const override;
 #endif // WITH_EDITORONLY_DATA
 	virtual void CacheFromCompiledData(const FNiagaraDataSetCompiledData* CompiledData) override;
+	virtual void UpdateSourceModeDerivates(ENiagaraRendererSourceDataMode InSourceMode, bool bFromPropertyEdit = false) override;
+	virtual ENiagaraRendererSourceDataMode GetCurrentSourceMode() const override { return SourceMode; }
+	virtual bool PopulateRequiredBindings(FNiagaraParameterStore& InParameterStore) override;
 	//UNiagaraRendererProperties Interface END
 
 	UMaterialInterface* GetMaterial(const FNiagaraEmitterInstance* InEmitter) const;
 
+	static const FQuat4f GetDefaultOrientation() { return FRotator3f(-90.0f, 0.0f, 90.0f).Quaternion(); }
+	static const FVector3f GetDefaultDecalSize() { return FVector3f(50.0f, 50.0f, 50.0f); }
+	static const float GetDefaultDecalFade() { return 1.0f; }
+
 	/** What material to use for the decal. */
 	UPROPERTY(EditAnywhere, Category = "Decal Rendering")
 	TObjectPtr<UMaterialInterface> Material;
@@ -53,6 +64,10 @@ public:
 	UPROPERTY(EditAnywhere, Category = "Decal Rendering")
 	FNiagaraParameterBinding MaterialParameterBinding;
 
+	/** Whether or not to draw a single element for the Emitter or to draw the particles.*/
+	UPROPERTY(EditAnywhere, Category = "Decal Rendering")
+	ENiagaraRendererSourceDataMode SourceMode = ENiagaraRendererSourceDataMode::Particles;
+
 	/** If a render visibility tag is present, particles whose tag matches this value will be visible in this renderer. */
 	UPROPERTY(EditAnywhere, Category = "Decal Rendering")
 	int32 RendererVisibility = 0;
diff --git a/Engine/Plugins/FX/Niagara/Source/Niagara/Public/NiagaraLightRendererProperties.h b/Engine/Plugins/FX/Niagara/Source/Niagara/Public/NiagaraLightRendererProperties.h
--- a/Engine/Plugins/FX/Niagara/Source/Niagara/Public/NiagaraLightRendererProperties.h
+++ b/Engine/Plugins/FX/Niagara/Source/Niagara/Public/NiagaraLightRendererProperties.h
@@ -23,6 +23,9 @@ public:
 	//UObject Interface
 	virtual void PostLoad() override;
 	virtual void PostInitProperties() override;
+#if WITH_EDITORONLY_DATA
+	virtual void PostEditChangeProperty(struct FPropertyChangedEvent& PropertyChangedEvent) override;
+#endif// WITH_EDITORONLY_DATA
 	//UObject Interface END
 
 	static void InitCDOPropertiesAfterModuleStartup();
@@ -33,14 +36,22 @@ public:
 	virtual void GetUsedMaterials(const FNiagaraEmitterInstance* InEmitter, TArray<UMaterialInterface*>& OutMaterials) const override;
 	virtual bool IsSimTargetSupported(ENiagaraSimTarget InSimTarget) const override { return (InSimTarget == ENiagaraSimTarget::CPUSim); };
 #if WITH_EDITORONLY_DATA
+	virtual bool IsSupportedVariableForBinding(const FNiagaraVariableBase& InSourceForBinding, const FName& InTargetBindingName) const override;
 	virtual const TArray<FNiagaraVariable>& GetOptionalAttributes() override;
 	virtual void GetRendererWidgets(const FNiagaraEmitterInstance* InEmitter, TArray<TSharedPtr<SWidget>>& OutWidgets, TSharedPtr<FAssetThumbnailPool> InThumbnailPool) const override;
 	virtual void GetRendererTooltipWidgets(const FNiagaraEmitterInstance* InEmitter, TArray<TSharedPtr<SWidget>>& OutWidgets, TSharedPtr<FAssetThumbnailPool> InThumbnailPool) const override;
 	virtual void GetRendererFeedback(const FVersionedNiagaraEmitter& InEmitter, TArray<FText>& OutErrors, TArray<FText>& OutWarnings, TArray<FText>& OutInfo) const override;
 #endif // WITH_EDITORONLY_DATA
+	virtual void UpdateSourceModeDerivates(ENiagaraRendererSourceDataMode InSourceMode, bool bFromPropertyEdit = false) override;
+	virtual ENiagaraRendererSourceDataMode GetCurrentSourceMode() const override { return SourceMode; }
+	virtual bool PopulateRequiredBindings(FNiagaraParameterStore& InParameterStore) override;
 	virtual void CacheFromCompiledData(const FNiagaraDataSetCompiledData* CompiledData) override;
 	//UNiagaraRendererProperties Interface END
 
+	/** Whether or not to draw a single element for the Emitter or to draw the particles.*/
+	UPROPERTY(EditAnywhere, Category = "Light Rendering")
+	ENiagaraRendererSourceDataMode SourceMode = ENiagaraRendererSourceDataMode::Particles;
+
 	/** Whether to use physically based inverse squared falloff from the light.  If unchecked, the value from the LightExponent binding will be used instead. */
 	UPROPERTY(EditAnywhere, Category = "Light Rendering")
 	uint32 bUseInverseSquaredFalloff : 1;
diff --git a/Engine/Plugins/FX/Niagara/Source/Niagara/Public/NiagaraParameterStore.h b/Engine/Plugins/FX/Niagara/Source/Niagara/Public/NiagaraParameterStore.h
--- a/Engine/Plugins/FX/Niagara/Source/Niagara/Public/NiagaraParameterStore.h
+++ b/Engine/Plugins/FX/Niagara/Source/Niagara/Public/NiagaraParameterStore.h
@@ -401,6 +401,14 @@ public:
 		return T();
 	}
 
+	template<typename T>
+	FORCEINLINE_DEBUGGABLE T GetParameterValueOrDefault(const FNiagaraVariable& Parameter, const T& DefaultValue) const
+	{
+		check(Parameter.GetSizeInBytes() == sizeof(T));
+		int32 Offset = IndexOf(Parameter);
+		return Offset != INDEX_NONE ? *(const T*)(GetParameterData(Offset)) : DefaultValue;
+	}
+
 	FORCEINLINE const uint8* GetParameterData(int32 Offset)const
 	{
 		return ParameterData.GetData() + Offset;
