- Add missing bReadFunction to various Grid2D signatures

#rnx

diff --git a/Engine/Plugins/FX/Niagara/Source/Niagara/Private/NiagaraDataInterfaceGrid2DCollection.cpp b/Engine/Plugins/FX/Niagara/Source/Niagara/Private/NiagaraDataInterfaceGrid2DCollection.cpp
--- a/Engine/Plugins/FX/Niagara/Source/Niagara/Private/NiagaraDataInterfaceGrid2DCollection.cpp
+++ b/Engine/Plugins/FX/Niagara/Source/Niagara/Private/NiagaraDataInterfaceGrid2DCollection.cpp
@@ -95,12 +95,12 @@ const FName UNiagaraDataInterfaceGrid2DCollection::SampleGridFunctionName("Sampl
 
 FNiagaraVariableBase UNiagaraDataInterfaceGrid2DCollection::ExposedRTVar;
 
-
 struct FNiagaraGridCollection2DDIFunctionVersion
 {
 	enum Type
 	{
 		InitialVersion = 0,
+		AddReadMemberToFunctions = 1,
 
 		VersionPlusOne,
 		LatestVersion = VersionPlusOne - 1
@@ -341,6 +341,22 @@ bool UNiagaraDataInterfaceGrid2DCollection::UpgradeFunctionCall(FNiagaraFunction
 		}
 	}
 
+	if ( FunctionSignature.FunctionVersion < FNiagaraGridCollection2DDIFunctionVersion::AddReadMemberToFunctions )
+	{
+		TArray<FNiagaraFunctionSignature> Sigs;
+		GetFunctions(Sigs);
+
+		for (const FNiagaraFunctionSignature& Sig : Sigs)
+		{
+			if (Sig.Name == FunctionSignature.Name)
+			{
+				FNiagaraFunctionSignature Backup = FunctionSignature;
+				FunctionSignature = Sig;
+				return true;
+			}
+		}
+	}
+
 	return false;
 }
 #endif
@@ -386,6 +402,7 @@ void UNiagaraDataInterfaceGrid2DCollection::GetFunctions(TArray<FNiagaraFunction
 		Sig.bSoftDeprecatedFunction = true;
 		Sig.bSupportsCPU = false;
 		Sig.bSupportsGPU = true;
+		Sig.bReadFunction = true;
 #if WITH_EDITORONLY_DATA
 		Sig.Description = NSLOCTEXT("Niagara", "NiagaraDataInterfaceGridColl2D_GetValueFunction", "Get the value at a specific index. Note that this is an older way of working with Grids. Consider using the SetFloat or other typed, named functions or parameter map variables with StackContext namespace instead.");
 #endif
@@ -408,6 +425,7 @@ void UNiagaraDataInterfaceGrid2DCollection::GetFunctions(TArray<FNiagaraFunction
 		Sig.ModuleUsageBitmask = ENiagaraScriptUsageMask::Particle;
 		Sig.bSupportsCPU = false;
 		Sig.bSupportsGPU = true;
+		Sig.bReadFunction = true;
 #if WITH_EDITORONLY_DATA
 		Sig.Description = NSLOCTEXT("Niagara", "NiagaraDataInterfaceGridColl2D_PreviousValueAtIndexFunction", "Get the value at a specific index.");
 #endif
@@ -500,6 +518,7 @@ void UNiagaraDataInterfaceGrid2DCollection::GetFunctions(TArray<FNiagaraFunction
 		Sig.ModuleUsageBitmask = ENiagaraScriptUsageMask::Particle;
 		Sig.bSupportsCPU = false;
 		Sig.bSupportsGPU = true;
+		Sig.bReadFunction = true;
 #if WITH_EDITORONLY_DATA
 		Sig.Description = NSLOCTEXT("Niagara", "NiagaraDataInterfaceGridColl2D_CopyPreviousToCurrentForCell", "Take the previous contents of the cell and copy to the output location for the cell.");
 #endif
@@ -542,6 +561,7 @@ void UNiagaraDataInterfaceGrid2DCollection::GetFunctions(TArray<FNiagaraFunction
 		Sig.ModuleUsageBitmask = ENiagaraScriptUsageMask::Particle;
 		Sig.bSupportsCPU = false;
 		Sig.bSupportsGPU = true;
+		Sig.bReadFunction = true;
 #if WITH_EDITORONLY_DATA
 		Sig.Description = NSLOCTEXT("Niagara", "NiagaraDataInterfaceGridColl2D_GetVector4", "Gets a Vector4 value on the Grid by Attribute name. Note that this is the value from the previous execution stage.");
 #endif
@@ -566,6 +586,7 @@ void UNiagaraDataInterfaceGrid2DCollection::GetFunctions(TArray<FNiagaraFunction
 		Sig.ModuleUsageBitmask = ENiagaraScriptUsageMask::Particle;
 		Sig.bSupportsCPU = false;
 		Sig.bSupportsGPU = true;
+		Sig.bReadFunction = true;
 #if WITH_EDITORONLY_DATA
 		Sig.Description = NSLOCTEXT("Niagara", "NiagaraDataInterfaceGridColl2D_SampleVector4", "Sample a Vector4 value on the Grid by Attribute name. Note that this is the value from the previous execution stage.");
 #endif
@@ -618,6 +639,7 @@ void UNiagaraDataInterfaceGrid2DCollection::GetFunctions(TArray<FNiagaraFunction
 		Sig.ModuleUsageBitmask = ENiagaraScriptUsageMask::Particle;
 		Sig.bSupportsCPU = false;
 		Sig.bSupportsGPU = true;
+		Sig.bReadFunction = true;
 #if WITH_EDITORONLY_DATA
 		Sig.Description = NSLOCTEXT("Niagara", "NiagaraDataInterfaceGridColl2D_GetVector3", "Gets a Vector3 value on the Grid by Attribute name. Note that this is the value from the previous execution stage.");
 #endif
@@ -642,6 +664,7 @@ void UNiagaraDataInterfaceGrid2DCollection::GetFunctions(TArray<FNiagaraFunction
 		Sig.ModuleUsageBitmask = ENiagaraScriptUsageMask::Particle;
 		Sig.bSupportsCPU = false;
 		Sig.bSupportsGPU = true;
+		Sig.bReadFunction = true;
 #if WITH_EDITORONLY_DATA
 		Sig.Description = NSLOCTEXT("Niagara", "NiagaraDataInterfaceGridColl2D_SampleVector3", "Sample a Vector3 value on the Grid by Attribute name. Note that this is the value from the previous execution stage.");
 #endif
@@ -694,6 +717,7 @@ void UNiagaraDataInterfaceGrid2DCollection::GetFunctions(TArray<FNiagaraFunction
 		Sig.ModuleUsageBitmask = ENiagaraScriptUsageMask::Particle;
 		Sig.bSupportsCPU = false;
 		Sig.bSupportsGPU = true;
+		Sig.bReadFunction = true;
 #if WITH_EDITORONLY_DATA
 		Sig.Description = NSLOCTEXT("Niagara", "NiagaraDataInterfaceGridColl2D_GetVector2", "Gets a Vector2 value on the Grid by Attribute name. Note that this is the value from the previous execution stage.");
 #endif
@@ -718,6 +742,7 @@ void UNiagaraDataInterfaceGrid2DCollection::GetFunctions(TArray<FNiagaraFunction
 		Sig.ModuleUsageBitmask = ENiagaraScriptUsageMask::Particle;
 		Sig.bSupportsCPU = false;
 		Sig.bSupportsGPU = true;
+		Sig.bReadFunction = true;
 #if WITH_EDITORONLY_DATA
 		Sig.Description = NSLOCTEXT("Niagara", "NiagaraDataInterfaceGridColl2D_SampleVector2", "Sample a Vector2 value on the Grid by Attribute name. Note that this is the value from the previous execution stage.");
 #endif
@@ -765,6 +790,7 @@ void UNiagaraDataInterfaceGrid2DCollection::GetFunctions(TArray<FNiagaraFunction
 		Sig.ModuleUsageBitmask = ENiagaraScriptUsageMask::Particle;
 		Sig.bSupportsCPU = false;
 		Sig.bSupportsGPU = true;
+		Sig.bReadFunction = true;
 #if WITH_EDITORONLY_DATA
 		Sig.Description = NSLOCTEXT("Niagara", "NiagaraDataInterfaceGridColl2D_GetFloat", "Gets a float value on the Grid by Attribute name. Note that this is the value from the previous execution stage.");
 #endif
@@ -789,6 +815,7 @@ void UNiagaraDataInterfaceGrid2DCollection::GetFunctions(TArray<FNiagaraFunction
 		Sig.ModuleUsageBitmask = ENiagaraScriptUsageMask::Particle;
 		Sig.bSupportsCPU = false;
 		Sig.bSupportsGPU = true;
+		Sig.bReadFunction = true;
 #if WITH_EDITORONLY_DATA
 		Sig.Description = NSLOCTEXT("Niagara", "NiagaraDataInterfaceGridColl2D_SampleFloat", "Sample a float value on the Grid by Attribute name.");
 #endif
@@ -815,6 +842,7 @@ void UNiagaraDataInterfaceGrid2DCollection::GetFunctions(TArray<FNiagaraFunction
 		Sig.ModuleUsageBitmask = ENiagaraScriptUsageMask::Particle;
 		Sig.bSupportsCPU = false;
 		Sig.bSupportsGPU = true;
+		Sig.bReadFunction = true;
 		Sig.bSoftDeprecatedFunction = true;
 		OutFunctions.Add(Sig);
 	}
@@ -834,6 +862,7 @@ void UNiagaraDataInterfaceGrid2DCollection::GetFunctions(TArray<FNiagaraFunction
 		Sig.ModuleUsageBitmask = ENiagaraScriptUsageMask::Particle;
 		Sig.bSupportsCPU = false;
 		Sig.bSupportsGPU = true;
+		Sig.bReadFunction = true;
 		OutFunctions.Add(Sig);
 	}
 
