Niagara Data Channels - Supporting CL

Data Set Compiled Data now generates a hash of it's layout when built.
For easy checks of equivelance of layouts between data sets.

Some utility functions for copying data between datasets.
Some utility additions for dataset accesssors

diff --git a/Engine/Plugins/FX/Niagara/Source/Niagara/Classes/NiagaraDataSet.h b/Engine/Plugins/FX/Niagara/Source/Niagara/Classes/NiagaraDataSet.h
--- a/Engine/Plugins/FX/Niagara/Source/Niagara/Classes/NiagaraDataSet.h
+++ b/Engine/Plugins/FX/Niagara/Source/Niagara/Classes/NiagaraDataSet.h
@@ -94,7 +94,9 @@ public:
 	void SwapInstances(uint32 OldIndex, uint32 NewIndex);
 	void KillInstance(uint32 InstanceIdx);
 	void CopyTo(FNiagaraDataBuffer& DestBuffer, int32 SrcStartIdx, int32 DestStartIdx, int32 NumInstances)const;
+	void CopyToUnrelated(FNiagaraDataBuffer& DestBuffer, int32 SrcStartIdx, int32 DestStartIdx, int32 NumInstances)const;
 	void GPUCopyFrom(const float* GPUReadBackFloat, const int* GPUReadBackInt, const FFloat16* GPUReadBackHalf, int32 StartIdx, int32 NumInstances, uint32 InSrcFloatStride, uint32 InSrcIntStride, uint32 InSrcHalfStride);
+	void PushCPUBuffersToGPU(const TArray<FNiagaraDataBuffer*>& SourceBuffers, bool bReleaseRef, FRHICommandList& RHICmdList, ERHIFeatureLevel::Type FeatureLevel, const TCHAR* DebugSimName);
 	void Dump(int32 StartIndex, int32 NumInstances, const FString& Label, const FName& SortParameterKey = FName())const;
 
 	FORCEINLINE TArrayView<uint8 const* RESTRICT const> GetRegisterTable() const { return TArrayView<uint8 const* RESTRICT const>(RegisterTable); }
@@ -335,8 +337,6 @@ private:
 
 	void Reset();
 
-	void BuildLayout();
-
 	void ResetBuffersInternal();
 
 	FORCEINLINE void CheckCorrectThread()const
@@ -539,7 +539,22 @@ namespace NiagaraDataSetPrivate
 		return DataBuffer->GetComponentPtrInt32(ComponentIdx);
 	}
 
-	inline uint32 GetNumInstances(FNiagaraDataBuffer* DataBuffer)
+	inline const uint8* GetComponentPtrFloat(const FNiagaraDataBuffer* DataBuffer, uint32 ComponentIdx)
+	{
+		return DataBuffer->GetComponentPtrFloat(ComponentIdx);
+	}
+
+	inline const uint8* GetComponentPtrHalf(const FNiagaraDataBuffer* DataBuffer, uint32 ComponentIdx)
+	{
+		return DataBuffer->GetComponentPtrHalf(ComponentIdx);
+	}
+
+	inline const uint8* GetComponentPtrInt32(const FNiagaraDataBuffer* DataBuffer, uint32 ComponentIdx)
+	{
+		return DataBuffer->GetComponentPtrInt32(ComponentIdx);
+	}
+
+	inline uint32 GetNumInstances(const FNiagaraDataBuffer* DataBuffer)
 	{
 		return DataBuffer->GetNumInstances();
 	}
diff --git a/Engine/Plugins/FX/Niagara/Source/Niagara/Classes/NiagaraDataSetAccessor.h b/Engine/Plugins/FX/Niagara/Source/Niagara/Classes/NiagaraDataSetAccessor.h
--- a/Engine/Plugins/FX/Niagara/Source/Niagara/Classes/NiagaraDataSetAccessor.h
+++ b/Engine/Plugins/FX/Niagara/Source/Niagara/Classes/NiagaraDataSetAccessor.h
@@ -19,7 +19,10 @@ namespace NiagaraDataSetPrivate
 	inline uint8* GetComponentPtrFloat(FNiagaraDataBuffer* DataBuffer, uint32 ComponentIdx);
 	inline uint8* GetComponentPtrHalf(FNiagaraDataBuffer* DataBuffer, uint32 ComponentIdx);
 	inline uint8* GetComponentPtrInt32(FNiagaraDataBuffer* DataBuffer, uint32 ComponentIdx);
-	inline uint32 GetNumInstances(FNiagaraDataBuffer* DataBuffer);
+	inline const uint8* GetComponentPtrFloat(const FNiagaraDataBuffer* DataBuffer, uint32 ComponentIdx);
+	inline const uint8* GetComponentPtrHalf(const FNiagaraDataBuffer* DataBuffer, uint32 ComponentIdx);
+	inline const uint8* GetComponentPtrInt32(const FNiagaraDataBuffer* DataBuffer, uint32 ComponentIdx);
+	inline uint32 GetNumInstances(const FNiagaraDataBuffer* DataBuffer);
 }
 
 template<typename TType>
@@ -531,9 +534,8 @@ struct FNiagaraDataSetReaderStruct
 {
 	FORCEINLINE FNiagaraDataSetReaderStruct() {}
 
-	explicit FNiagaraDataSetReaderStruct(FNiagaraDataBuffer* DataBuffer, int32 FloatComponentIndex, int32 Int32ComponentIndex)
+	explicit FNiagaraDataSetReaderStruct(const FNiagaraDataBuffer* DataBuffer, int32 FloatComponentIndex, int32 Int32ComponentIndex)
 	{
-		using namespace NiagaraDataSetPrivate;
 		bIsValid  = DataBuffer != nullptr;
 		bIsValid &= (FNiagaraDataSetAccessorTypeInfo<TType>::NumFloatComponents == 0) || (FloatComponentIndex != INDEX_NONE);
 		bIsValid &= (FNiagaraDataSetAccessorTypeInfo<TType>::NumInt32Components == 0) || (Int32ComponentIndex != INDEX_NONE);
@@ -542,13 +544,13 @@ struct FNiagaraDataSetReaderStruct
 			int32 ComponentIndex = 0;
 			for (int i = 0; i < FNiagaraDataSetAccessorTypeInfo<TType>::NumFloatComponents; ++i)
 			{
-				ComponentData[ComponentIndex++] = GetComponentPtrFloat(DataBuffer, FloatComponentIndex + i);
+				ComponentData[ComponentIndex++] = NiagaraDataSetPrivate::GetComponentPtrFloat(DataBuffer, FloatComponentIndex + i);
 			}
 			for (int i = 0; i < FNiagaraDataSetAccessorTypeInfo<TType>::NumInt32Components; ++i)
 			{
-				ComponentData[ComponentIndex++] = GetComponentPtrInt32(DataBuffer, Int32ComponentIndex + i);
+				ComponentData[ComponentIndex++] = NiagaraDataSetPrivate::GetComponentPtrInt32(DataBuffer, Int32ComponentIndex + i);
 			}
-			NumInstances = GetNumInstances(DataBuffer);
+			NumInstances = NiagaraDataSetPrivate::GetNumInstances(DataBuffer);
 		}
 	}
 
@@ -578,6 +580,8 @@ struct FNiagaraDataSetReaderStruct
 		return Default;
 	}
 
+	FORCEINLINE int32 GetNumInstances()const { return NumInstances; }
+
 private:
 	bool bIsValid = false;
 	int32 NumInstances = 0;
@@ -632,6 +636,7 @@ struct FNiagaraDataSetAccessorStruct
 	FORCEINLINE FNiagaraDataSetReaderStruct<TType> GetReader(const FNiagaraDataSet& DataSet) const { return FNiagaraDataSetReaderStruct<TType>(NiagaraDataSetPrivate::GetCurrentData(DataSet), FloatComponentIndex, Int32ComponentIndex); }
 	//FNiagaraDataSetReaderStruct<TType> GetWriter(const FNiagaraDataSet& DataSet) { return FNiagaraDataSetReaderStruct<TType>(NiagaraDataSetPrivate::GetDestinationData(DataSet), FloatComponentIndex, Int32ComponentIndex); }
 
+	FORCEINLINE FNiagaraDataSetReaderStruct<TType> GetReader(const FNiagaraDataBuffer* DataBuffer) const { return FNiagaraDataSetReaderStruct<TType>(DataBuffer, FloatComponentIndex, Int32ComponentIndex); }
 private:
 	bool bIsValid = false;
 	int32 FloatComponentIndex = INDEX_NONE;
diff --git a/Engine/Plugins/FX/Niagara/Source/Niagara/Classes/NiagaraDataSetCompiledData.h b/Engine/Plugins/FX/Niagara/Source/Niagara/Classes/NiagaraDataSetCompiledData.h
--- a/Engine/Plugins/FX/Niagara/Source/Niagara/Classes/NiagaraDataSetCompiledData.h
+++ b/Engine/Plugins/FX/Niagara/Source/Niagara/Classes/NiagaraDataSetCompiledData.h
@@ -72,6 +72,34 @@ struct NIAGARA_API FNiagaraDataSetCompiledData
 	void Empty();
 
 	static FNiagaraDataSetCompiledData DummyCompiledData;
+
+	uint32 GetLayoutHash()const
+	{
+		if(CachedLayoutHash == INDEX_NONE)
+		{
+			CachedLayoutHash = 0;
+			for(const FNiagaraVariableBase& Var : Variables)
+			{
+				CachedLayoutHash = HashCombine(CachedLayoutHash, GetTypeHash(Var));
+			}
+		}
+		return CachedLayoutHash;
+	}
+
+	bool CheckHashConflict(const FNiagaraDataSetCompiledData& Other)const
+	{
+		//If we have the same hash ensure we have the same data.
+		if(GetLayoutHash() == Other.GetLayoutHash())
+		{
+			return Variables != Other.Variables;			
+		}
+		return false;
+	}
+
+private:
+	/** A hash of all variables in this layout. NOT persistent. Will differ from run to run. */
+	mutable uint32 CachedLayoutHash = INDEX_NONE;
+
 };
 
 //////////////////////////////////////////////////////////////////////////
diff --git a/Engine/Plugins/FX/Niagara/Source/Niagara/Private/NiagaraDataSet.cpp b/Engine/Plugins/FX/Niagara/Source/Niagara/Private/NiagaraDataSet.cpp
--- a/Engine/Plugins/FX/Niagara/Source/Niagara/Private/NiagaraDataSet.cpp
+++ b/Engine/Plugins/FX/Niagara/Source/Niagara/Private/NiagaraDataSet.cpp
@@ -1021,6 +1021,113 @@ void FNiagaraDataBuffer::KillInstance(uint32 InstanceIdx)
 #endif
 }
 
+void FNiagaraDataBuffer::CopyToUnrelated(FNiagaraDataBuffer& DestBuffer, int32 StartIdx, int32 DestStartIdx, int32 InNumInstances)const
+{
+	CheckUsage(false);
+
+	if (StartIdx < 0 || (uint32)StartIdx >= NumInstances)
+	{
+		return;
+	}
+
+	uint32 InstancesToCopy = InNumInstances;
+	if (InstancesToCopy == INDEX_NONE)
+	{
+		InstancesToCopy = NumInstances - StartIdx;
+	}
+
+	if (InstancesToCopy != 0)
+	{
+		const uint32 NewNumInstances = DestStartIdx + InstancesToCopy;
+
+		// Only allocate if we need to increase the number of instances as the caller may have previously
+		// allocated the array and may not be expecting it to shrink inside this call.
+		if (DestStartIdx < 0 || NewNumInstances > DestBuffer.GetNumInstancesAllocated())
+		{
+			DestBuffer.Allocate(NewNumInstances, true);
+		}
+		DestBuffer.SetNumInstances(NewNumInstances);
+
+		const FNiagaraDataSetCompiledData& DestCompiledData = DestBuffer.GetOwner()->GetCompiledData();
+		const FNiagaraDataSetCompiledData& CompiledData = GetOwner()->GetCompiledData();
+
+		for(int32 VarIdx = 0; VarIdx < DestCompiledData.Variables.Num(); ++VarIdx)
+		{
+			const FNiagaraVariable& DestVar = DestCompiledData.Variables[VarIdx];
+			const FNiagaraVariableLayoutInfo& DestVarLayout = DestCompiledData.VariableLayouts[VarIdx];
+
+			const int32 SrcVarIdx = CompiledData.Variables.IndexOfByKey(DestVar);
+			if(SrcVarIdx != INDEX_NONE)
+			{
+				//Found a the variable to copy over.
+				const FNiagaraVariableLayoutInfo& SrcVarLayout = CompiledData.VariableLayouts[SrcVarIdx];
+
+				const uint32 FloatComponents = SrcVarLayout.GetNumFloatComponents();
+				check(FloatComponents == DestVarLayout.GetNumFloatComponents());
+				for (uint32 CompIdx = 0; CompIdx < FloatComponents; ++CompIdx)
+				{
+					const int32 SrcCompOffset = SrcVarLayout.FloatComponentStart + CompIdx;
+					const int32 DestCompOffest = DestVarLayout.FloatComponentStart + CompIdx;
+					const float* SrcStart = GetInstancePtrFloat(SrcCompOffset, StartIdx);
+					const float* SrcEnd = GetInstancePtrFloat(SrcCompOffset, StartIdx + InstancesToCopy);
+					float* Dst = DestBuffer.GetInstancePtrFloat(DestCompOffest, DestStartIdx);
+					size_t Count = SrcEnd - SrcStart;
+					FMemory::Memcpy(Dst, SrcStart, Count * sizeof(float));
+
+					if (Count > 0)
+					{
+						for (size_t i = 0; i < Count; i++)
+						{
+							checkSlow(SrcStart[i] == Dst[i]);
+						}
+					}
+				}
+				const uint32 IntComponents = SrcVarLayout.GetNumInt32Components();
+				check(IntComponents == DestVarLayout.GetNumInt32Components());
+				for (uint32 CompIdx = 0; CompIdx < IntComponents; ++CompIdx)
+				{
+					const int32 SrcCompOffset = SrcVarLayout.Int32ComponentStart + CompIdx;
+					const int32 DestCompOffest = DestVarLayout.Int32ComponentStart + CompIdx;
+					const int32* SrcStart = GetInstancePtrInt32(SrcCompOffset, StartIdx);
+					const int32* SrcEnd = GetInstancePtrInt32(SrcCompOffset, StartIdx + InstancesToCopy);
+					int32* Dst = DestBuffer.GetInstancePtrInt32(DestCompOffest, DestStartIdx);
+					size_t Count = SrcEnd - SrcStart;
+					FMemory::Memcpy(Dst, SrcStart, Count * sizeof(int32));
+
+					if (Count > 0)
+					{
+						for (size_t i = 0; i < Count; i++)
+						{
+							checkSlow(SrcStart[i] == Dst[i]);
+						}
+					}
+				}
+				const uint32 HalfComponents = SrcVarLayout.GetNumHalfComponents();
+				check(HalfComponents == DestVarLayout.GetNumHalfComponents());
+				for (uint32 CompIdx = 0; CompIdx < HalfComponents; ++CompIdx)
+				{
+					const int32 SrcCompOffset = SrcVarLayout.HalfComponentStart + CompIdx;
+					const int32 DestCompOffest = DestVarLayout.HalfComponentStart + CompIdx;
+					const FFloat16* SrcStart = GetInstancePtrHalf(SrcCompOffset, StartIdx);
+					const FFloat16* SrcEnd = GetInstancePtrHalf(SrcCompOffset, StartIdx + InstancesToCopy);
+					FFloat16* Dst = DestBuffer.GetInstancePtrHalf(DestCompOffest, DestStartIdx);
+					size_t Count = SrcEnd - SrcStart;
+					FMemory::Memcpy(Dst, SrcStart, Count * sizeof(FFloat16));
+
+					if (Count > 0)
+					{
+						for (size_t i = 0; i < Count; i++)
+						{
+							checkSlow(SrcStart[i] == Dst[i]);
+						}
+					}
+				}
+
+			}
+		}
+	}
+}
+
 void FNiagaraDataBuffer::CopyTo(FNiagaraDataBuffer& DestBuffer, int32 StartIdx, int32 DestStartIdx, int32 InNumInstances)const
 {
 	CheckUsage(false);
@@ -1158,6 +1265,94 @@ void FNiagaraDataBuffer::GPUCopyFrom(const float* GPUReadBackFloat, const int* G
 	}
 }
 
+void FNiagaraDataBuffer::PushCPUBuffersToGPU(const TArray<FNiagaraDataBuffer*>& SourceBuffers, bool bReleaseRef, FRHICommandList& RHICmdList, ERHIFeatureLevel::Type FeatureLevel, const TCHAR* DebugSimName)
+{
+	uint32 NewCount = 0;
+	check(GetOwner()->GetSimTarget() == ENiagaraSimTarget::GPUComputeSim);
+	for (FNiagaraDataBuffer* Buffer : SourceBuffers)
+	{
+		if (Buffer)
+		{
+			NewCount += Buffer->GetNumInstances();
+
+			checkSlow(Buffer->GetOwner()->GetVariables() == GetOwner()->GetVariables())
+		}
+	}
+
+	AllocateGPU(RHICmdList, NewCount, FeatureLevel, DebugSimName);
+	SetNumInstances(NewCount);
+
+	if (NewCount > 0)
+	{
+		uint32 FloatComponents = Owner->GetNumFloatComponents();
+		uint32 IntComponents = Owner->GetNumInt32Components();
+		uint32 HalfComponents = Owner->GetNumHalfComponents();
+
+		uint8* MappedBufferFloat = GPUBufferFloat.Buffer ? (uint8*)RHILockBuffer(GPUBufferFloat.Buffer, 0, FloatComponents * FloatStride, RLM_WriteOnly) : nullptr;
+		uint8* MappedBufferInt32 = GPUBufferInt.Buffer ? (uint8*)RHILockBuffer(GPUBufferInt.Buffer, 0, IntComponents * Int32Stride, RLM_WriteOnly) : nullptr;
+		uint8* MappedBufferHalf = GPUBufferHalf.Buffer ? (uint8*)RHILockBuffer(GPUBufferHalf.Buffer, 0, HalfComponents * HalfStride, RLM_WriteOnly) : nullptr;
+
+		for (uint32 CompIdx = 0; CompIdx < FloatComponents; ++CompIdx)
+		{
+			float* Dest = (float*)MappedBufferFloat;
+			MappedBufferFloat += FloatStride;
+			for (FNiagaraDataBuffer* Buffer : SourceBuffers)
+			{
+				const float* Src = Buffer->GetInstancePtrFloat(CompIdx, 0);
+				int32 CopyInstances = Buffer->GetNumInstances();
+				FMemory::Memcpy(Dest, Src, CopyInstances * sizeof(float));
+				Dest += CopyInstances;
+			}
+		}
+
+		for (uint32 CompIdx = 0; CompIdx < IntComponents; ++CompIdx)
+		{
+			int32* Dest = (int32*)MappedBufferInt32;
+			MappedBufferInt32 += Int32Stride;
+			for (FNiagaraDataBuffer* Buffer : SourceBuffers)
+			{
+				const int32* Src = Buffer->GetInstancePtrInt32(CompIdx, 0);
+				int32 CopyInstances = Buffer->GetNumInstances();
+				FMemory::Memcpy(Dest, Src, CopyInstances * sizeof(int32));
+				Dest += CopyInstances;
+			}
+		}
+
+		for (uint32 CompIdx = 0; CompIdx < HalfComponents; ++CompIdx)
+		{
+			FFloat16* Dest = (FFloat16*)MappedBufferHalf;
+			MappedBufferHalf += HalfStride;
+			for (FNiagaraDataBuffer* Buffer : SourceBuffers)
+			{
+				const FFloat16* Src = Buffer->GetInstancePtrHalf(CompIdx, 0);
+				int32 CopyInstances = Buffer->GetNumInstances();
+				FMemory::Memcpy(Dest, Src, CopyInstances * sizeof(FFloat16));
+				Dest += CopyInstances;
+			}
+		}
+
+		if (MappedBufferFloat)
+		{
+			RHIUnlockBuffer(GPUBufferFloat.Buffer);
+		}
+
+		if (MappedBufferInt32)
+		{
+			RHIUnlockBuffer(GPUBufferInt.Buffer);
+		}
+
+		if (MappedBufferHalf)
+		{
+			RHIUnlockBuffer(GPUBufferHalf.Buffer);
+		}
+	}
+
+	for (FNiagaraDataBuffer* Buffer : SourceBuffers)
+	{
+		Buffer->ReleaseReadRef();
+	}
+}
+
 void FNiagaraDataBuffer::Dump(int32 StartIndex, int32 InNumInstances, const FString& Label, const FName& SortParameterKey)const
 {
 	FNiagaraDataVariableIterator Itr(this, StartIndex);
@@ -1402,7 +1597,6 @@ void FNiagaraDataSetCompiledData::BuildLayout()
 	TotalInt32Components = 0;
 	TotalHalfComponents = 0;
 
-	VariableLayouts.Reserve(Variables.Num());
 	for (FNiagaraVariable& Var : Variables)
 	{
 		FNiagaraVariableLayoutInfo& VarInfo = VariableLayouts[VariableLayouts.AddDefaulted()];
@@ -1414,6 +1608,9 @@ void FNiagaraDataSetCompiledData::BuildLayout()
 		TotalInt32Components += VarInfo.GetNumInt32Components();
 		TotalHalfComponents += VarInfo.GetNumHalfComponents();
 	}
+
+	//Prime the layout hash for cases this happens at runtime.
+	GetLayoutHash();
 }
 
 FNiagaraDataSetCompiledData::FNiagaraDataSetCompiledData()
@@ -1423,6 +1620,7 @@ FNiagaraDataSetCompiledData::FNiagaraDataSetCompiledData()
 
 void FNiagaraDataSetCompiledData::Empty()
 {
+	CachedLayoutHash = INDEX_NONE;
 	bRequiresPersistentIDs = 0;
 	TotalFloatComponents = 0;
 	TotalInt32Components = 0;
