- Sort custom hlsl prototype to make them easier to find in the menu

#rnx

diff --git a/Engine/Plugins/FX/Niagara/Source/NiagaraEditor/Private/EdGraphSchema_Niagara.cpp b/Engine/Plugins/FX/Niagara/Source/NiagaraEditor/Private/EdGraphSchema_Niagara.cpp
--- a/Engine/Plugins/FX/Niagara/Source/NiagaraEditor/Private/EdGraphSchema_Niagara.cpp
+++ b/Engine/Plugins/FX/Niagara/Source/NiagaraEditor/Private/EdGraphSchema_Niagara.cpp
@@ -46,6 +46,7 @@
 #include "Modules/ModuleManager.h"
 #include "Textures/SlateIcon.h"
 #include "UObject/LinkerLoad.h"
+#include "Algo/Sort.h"
 
 #include UE_INLINE_GENERATED_CPP_BY_NAME(EdGraphSchema_Niagara)
 
@@ -2212,6 +2213,8 @@ void UEdGraphSchema_Niagara::GenerateDataInterfacePinMenu(UToolMenu* ToolMenu, c
 		return;
 	}
 
+	Algo::Sort(FunctionPrototypes, [](const TPair<FString, FString>& Lhs, const TPair<FString, FString>& Rhs) { return Lhs.Key < Rhs.Key; });
+
 	// Make the menu
 	FToolMenuSection& Section = ToolMenu->FindOrAddSection(SectionName);
 	Section.AddSubMenu(
