- Add 'Direct Set' iteration type this is not bound to particles or data interfaces
- This allows the user to specify a number of threads to launch without any bounds clipping from particles / DI sources
- Previous override element will change to this type as element counts are not allowed on other sources

diff --git a/Engine/Plugins/FX/Niagara/Source/Niagara/Private/NiagaraComputeExecutionContext.cpp b/Engine/Plugins/FX/Niagara/Source/Niagara/Private/NiagaraComputeExecutionContext.cpp
--- a/Engine/Plugins/FX/Niagara/Source/Niagara/Private/NiagaraComputeExecutionContext.cpp
+++ b/Engine/Plugins/FX/Niagara/Source/Niagara/Private/NiagaraComputeExecutionContext.cpp
@@ -103,7 +103,7 @@ bool FNiagaraComputeExecutionContext::IsIterationStage(FNiagaraDataInterfaceProx
 {
 	if (DIProxy && !DIProxy->SourceDIName.IsNone())
 	{
-		return !SimStageInfo[SimulationStageIndex].IterationSource.IsNone() && (SimStageInfo[SimulationStageIndex].IterationSource == DIProxy->SourceDIName);
+		return SimStageInfo[SimulationStageIndex].IterationSource == ENiagaraIterationSource::DataInterface && (SimStageInfo[SimulationStageIndex].IterationDataInterface == DIProxy->SourceDIName);
 	}
 	return false;
 }
@@ -111,20 +111,20 @@ bool FNiagaraComputeExecutionContext::IsIterationStage(FNiagaraDataInterfaceProx
 FNiagaraDataInterfaceProxyRW* FNiagaraComputeExecutionContext::FindIterationInterface(const TArray<FNiagaraDataInterfaceProxyRW*>& InProxies, uint32 SimulationStageIndex) const
 {
 	// Particle stage
-	if ( SimStageInfo[SimulationStageIndex].IterationSource.IsNone() )
+	if ( SimStageInfo[SimulationStageIndex].IterationSource != ENiagaraIterationSource::DataInterface )
 	{
 		return nullptr;
 	}
 
 	for (FNiagaraDataInterfaceProxyRW* Proxy : InProxies)
 	{
-		if (Proxy->SourceDIName == SimStageInfo[SimulationStageIndex].IterationSource)
+		if (Proxy->SourceDIName == SimStageInfo[SimulationStageIndex].IterationDataInterface)
 		{
 			return Proxy;
 		}
 	}
 
-	UE_LOG(LogNiagara, Verbose, TEXT("FNiagaraComputeExecutionContext::FindIterationInterface could not find IterationInterface %s"), *SimStageInfo[SimulationStageIndex].IterationSource.ToString());
+	UE_LOG(LogNiagara, Verbose, TEXT("FNiagaraComputeExecutionContext::FindIterationInterface could not find IterationInterface %s"), *SimStageInfo[SimulationStageIndex].IterationDataInterface.ToString());
 
 	return nullptr;
 }
diff --git a/Engine/Plugins/FX/Niagara/Source/Niagara/Private/NiagaraEmitter.cpp b/Engine/Plugins/FX/Niagara/Source/Niagara/Private/NiagaraEmitter.cpp
--- a/Engine/Plugins/FX/Niagara/Source/Niagara/Private/NiagaraEmitter.cpp
+++ b/Engine/Plugins/FX/Niagara/Source/Niagara/Private/NiagaraEmitter.cpp
@@ -1714,7 +1714,7 @@ bool FVersionedNiagaraEmitterData::BuildParameterStoreRendererBindings(FNiagaraP
 			{
 				bAnyBindingsAdded |= ParameterStore.AddParameter(FNiagaraVariable(FNiagaraTypeDefinition::GetBoolDef(), SimStageMetaData.EnabledBinding), false);
 			}
-			if (SimStageMetaData.bOverrideElementCount)
+			if (SimStageMetaData.IterationSource == ENiagaraIterationSource::DirectSet)
 			{
 				if (!SimStageMetaData.ElementCountXBinding.IsNone())
 				{
diff --git a/Engine/Plugins/FX/Niagara/Source/Niagara/Private/NiagaraGPUSystemTick.cpp b/Engine/Plugins/FX/Niagara/Source/Niagara/Private/NiagaraGPUSystemTick.cpp
--- a/Engine/Plugins/FX/Niagara/Source/Niagara/Private/NiagaraGPUSystemTick.cpp
+++ b/Engine/Plugins/FX/Niagara/Source/Niagara/Private/NiagaraGPUSystemTick.cpp
@@ -222,7 +222,7 @@ void FNiagaraGPUSystemTick::Init(FNiagaraSystemInstance* InSystemInstance)
 							NumIterations = StageEnabled.GetValue() ? NumIterations : 0;
 						}
 					}
-					if (SimStageMetaData.bOverrideElementCount)
+					if (SimStageMetaData.IterationSource == ENiagaraIterationSource::DirectSet)
 					{
 						if (!SimStageMetaData.ElementCountXBinding.IsNone() && (SimStageMetaData.GpuDispatchType >= ENiagaraGpuDispatchType::OneD))
 						{
diff --git a/Engine/Plugins/FX/Niagara/Source/Niagara/Private/NiagaraGpuComputeDispatch.cpp b/Engine/Plugins/FX/Niagara/Source/Niagara/Private/NiagaraGpuComputeDispatch.cpp
--- a/Engine/Plugins/FX/Niagara/Source/Niagara/Private/NiagaraGpuComputeDispatch.cpp
+++ b/Engine/Plugins/FX/Niagara/Source/Niagara/Private/NiagaraGpuComputeDispatch.cpp
@@ -599,7 +599,7 @@ void FNiagaraGpuComputeDispatch::DumpDebugFrame()
 				Builder.Appendf(TEXT("Source(%p 0x%08x %d) "), SimStageData.Source, SimStageData.SourceCountOffset, SimStageData.SourceNumInstances);
 				Builder.Appendf(TEXT("Destination(%p 0x%08x %d) "), SimStageData.Destination, SimStageData.DestinationCountOffset, SimStageData.DestinationNumInstances);
 				Builder.Appendf(TEXT("Iteration(%d | %s) "), SimStageData.IterationIndex, SimStageData.AlternateIterationSource ? *SimStageData.AlternateIterationSource->SourceDIName.ToString() : TEXT("Particles"));
-				if (SimStageData.StageMetaData->bOverrideElementCount)
+				if (SimStageData.StageMetaData->IterationSource == ENiagaraIterationSource::DirectSet)
 				{
 					Builder.Appendf(TEXT("ElementCountXYZ(%d, %d, %d) "), SimStageData.ElementCountXYZ.X, SimStageData.ElementCountXYZ.Y, SimStageData.ElementCountXYZ.Z);
 				}
@@ -1429,46 +1429,64 @@ void FNiagaraGpuComputeDispatch::DispatchStage(FRDGBuilder& GraphBuilder, const
 	ENiagaraGpuDispatchType DispatchType = ENiagaraGpuDispatchType::OneD;
 	FIntVector DispatchCount = FIntVector(0, 0, 0);
 	FIntVector DispatchNumThreads = FIntVector(0, 0, 0);
-	if (SimStageData.AlternateIterationSource)
-	{
-		DispatchType = SimStageData.StageMetaData->GpuDispatchType;
-		DispatchCount = SimStageData.AlternateIterationSource->GetElementCount(Tick.SystemInstanceID);
-		DispatchNumThreads = SimStageData.StageMetaData->GpuDispatchNumThreads;
-
-		// Verify the number of elements isn't higher that what we can handle
-		checkf(uint64(DispatchCount.X) * uint64(DispatchCount.Y) * uint64(DispatchCount.Z) < uint64(TNumericLimits<int32>::Max()), TEXT("DispatchCount(%d, %d, %d) for IterationInterface(%s) overflows an int32 this is not allowed"), DispatchCount.X, DispatchCount.Y, DispatchCount.Z, *SimStageData.AlternateIterationSource->SourceDIName.ToString());
 
-		// Data interfaces such as grids / render targets can choose to dispatch in either the correct dimensionality for the target (i.e. RT2D would choose 2D)
-		// or run in linear mode if performance is not beneficial due to increased waves.  It is also possible the we may choose to override on the simulation stage.
-		// Therefore we need to special case OneD and convert our element count back to linear.
-		if (DispatchType == ENiagaraGpuDispatchType::OneD)
-		{
-			DispatchCount.X = DispatchCount.X * DispatchCount.Y * DispatchCount.Z;
-			DispatchCount.Y = 1;
-			DispatchCount.Z = 1;
-		}
-	}
-	else
+	switch (SimStageData.StageMetaData->IterationSource)
 	{
-		DispatchType = ENiagaraGpuDispatchType::OneD;
-		DispatchCount = FIntVector(SimStageData.DestinationNumInstances, 1, 1);
-		DispatchNumThreads = FNiagaraShader::GetDefaultThreadGroupSize(ENiagaraGpuDispatchType::OneD);
-	}
-
-	if (SimStageData.StageMetaData->bOverrideElementCount)
-	{
-		if (DispatchType == ENiagaraGpuDispatchType::OneD)
+		case ENiagaraIterationSource::Particles:
 		{
-			DispatchCount = FIntVector(SimStageData.ElementCountXYZ.X, 1, 1);
+			DispatchType = ENiagaraGpuDispatchType::OneD;
+			DispatchCount = FIntVector(SimStageData.DestinationNumInstances, 1, 1);
+			DispatchNumThreads = FNiagaraShader::GetDefaultThreadGroupSize(ENiagaraGpuDispatchType::OneD);
+			break;
 		}
-		else if (DispatchType == ENiagaraGpuDispatchType::TwoD)
+
+		case ENiagaraIterationSource::DataInterface:
 		{
-			DispatchCount = FIntVector(SimStageData.ElementCountXYZ.X, SimStageData.ElementCountXYZ.Y, 1);
+			DispatchType = SimStageData.StageMetaData->GpuDispatchType;
+			DispatchCount = SimStageData.AlternateIterationSource->GetElementCount(Tick.SystemInstanceID);
+			DispatchNumThreads = SimStageData.StageMetaData->GpuDispatchNumThreads;
+
+			// Verify the number of elements isn't higher that what we can handle
+			checkf(uint64(DispatchCount.X) * uint64(DispatchCount.Y) * uint64(DispatchCount.Z) < uint64(TNumericLimits<int32>::Max()), TEXT("DispatchCount(%d, %d, %d) for IterationInterface(%s) overflows an int32 this is not allowed"), DispatchCount.X, DispatchCount.Y, DispatchCount.Z, *SimStageData.AlternateIterationSource->SourceDIName.ToString());
+
+			// Data interfaces such as grids / render targets can choose to dispatch in either the correct dimensionality for the target (i.e. RT2D would choose 2D)
+			// or run in linear mode if performance is not beneficial due to increased waves.  It is also possible the we may choose to override on the simulation stage.
+			// Therefore we need to special case OneD and convert our element count back to linear.
+			if (DispatchType == ENiagaraGpuDispatchType::OneD)
+			{
+				DispatchCount.X = DispatchCount.X * DispatchCount.Y * DispatchCount.Z;
+				DispatchCount.Y = 1;
+				DispatchCount.Z = 1;
+			}
+			break;
 		}
-		else if (DispatchType == ENiagaraGpuDispatchType::ThreeD)
+
+		case ENiagaraIterationSource::DirectSet:
 		{
-			DispatchCount = SimStageData.ElementCountXYZ;
+			DispatchType = SimStageData.StageMetaData->GpuDispatchType;
+			DispatchNumThreads = SimStageData.StageMetaData->GpuDispatchNumThreads;
+
+			switch (DispatchType)
+			{
+				case ENiagaraGpuDispatchType::OneD:
+					DispatchCount = FIntVector(SimStageData.ElementCountXYZ.X, 1, 1);
+					break;
+				case ENiagaraGpuDispatchType::TwoD:
+					DispatchCount = FIntVector(SimStageData.ElementCountXYZ.X, SimStageData.ElementCountXYZ.Y, 1);
+					break;
+				case ENiagaraGpuDispatchType::ThreeD:
+					DispatchCount = SimStageData.ElementCountXYZ;
+					break;
+				default:
+					UE_LOG(LogNiagara, Fatal, TEXT("FNiagaraGpuComputeDispatch: Unknown DispatchType(%d)"), DispatchType);
+					break;
+			}
+			break;
 		}
+
+		default:
+			UE_LOG(LogNiagara, Fatal, TEXT("FNiagaraGpuComputeDispatch: Unknown IterationSouce(%d)"), SimStageData.StageMetaData->IterationSource);
+			return;
 	}
 
 	const int32 TotalDispatchCount = DispatchCount.X * DispatchCount.Y * DispatchCount.Z;
@@ -1503,7 +1521,7 @@ void FNiagaraGpuComputeDispatch::DispatchStage(FRDGBuilder& GraphBuilder, const
 
 	// Setup instance counts
 	DispatchParameters->RWInstanceCounts			= GPUInstanceCounterManager.GetInstanceCountBuffer().UAV;
-	DispatchParameters->ReadInstanceCountOffset	= SimStageData.AlternateIterationSource ? INDEX_NONE : SimStageData.SourceCountOffset;
+	DispatchParameters->ReadInstanceCountOffset		= SimStageData.AlternateIterationSource ? INDEX_NONE : SimStageData.SourceCountOffset;
 	DispatchParameters->WriteInstanceCountOffset	= SimStageData.AlternateIterationSource ? INDEX_NONE : SimStageData.DestinationCountOffset;
 
 	// Simulation Stage Information
@@ -1514,11 +1532,23 @@ void FNiagaraGpuComputeDispatch::DispatchStage(FRDGBuilder& GraphBuilder, const
 	{
 		DispatchParameters->SimulationStageIterationInfo = FIntVector4(INDEX_NONE, -1, 0, 0);
 		DispatchParameters->SimulationStageNormalizedIterationIndex = 0.0f;
-		if (SimStageData.AlternateIterationSource != nullptr)
+		switch (SimStageData.StageMetaData->IterationSource)
 		{
-			const uint32 IterationInstanceCountOffset = SimStageData.AlternateIterationSource->GetGPUInstanceCountOffset(Tick.SystemInstanceID);
-			DispatchParameters->SimulationStageIterationInfo.X = IterationInstanceCountOffset;
-			DispatchParameters->SimulationStageIterationInfo.Y = IterationInstanceCountOffset == INDEX_NONE ? TotalDispatchCount : 0;
+			case ENiagaraIterationSource::Particles:
+				break;
+
+			case ENiagaraIterationSource::DataInterface:
+				if (SimStageData.AlternateIterationSource != nullptr)
+				{
+					const uint32 IterationInstanceCountOffset = SimStageData.AlternateIterationSource->GetGPUInstanceCountOffset(Tick.SystemInstanceID);
+					DispatchParameters->SimulationStageIterationInfo.X = IterationInstanceCountOffset;
+					DispatchParameters->SimulationStageIterationInfo.Y = IterationInstanceCountOffset == INDEX_NONE ? TotalDispatchCount : 0;
+				}
+				break;
+
+			case ENiagaraIterationSource::DirectSet:
+				DispatchParameters->SimulationStageIterationInfo.Y = TotalDispatchCount;
+				break;
 		}
 
 		const int32 NumIterations = InstanceData.PerStageInfo[SimStageData.StageIndex].NumIterations;
diff --git a/Engine/Plugins/FX/Niagara/Source/Niagara/Private/NiagaraScript.cpp b/Engine/Plugins/FX/Niagara/Source/Niagara/Private/NiagaraScript.cpp
--- a/Engine/Plugins/FX/Niagara/Source/Niagara/Private/NiagaraScript.cpp
+++ b/Engine/Plugins/FX/Niagara/Source/Niagara/Private/NiagaraScript.cpp
@@ -1133,7 +1133,7 @@ void UNiagaraScript::ComputeVMCompilationId(FNiagaraVMExecutableDataId& Id, FGui
 					if (const UNiagaraSimulationStageGeneric* SimStageGeneric = Cast<const UNiagaraSimulationStageGeneric>(SimStageBase))
 					{
 						AddAttributeToPreserve(SimStageGeneric->EnabledBinding.GetParamMapBindableVariable().GetName());
-						if (SimStageGeneric->bOverrideGpuDispatchType)
+						if (SimStageGeneric->IterationSource == ENiagaraIterationSource::DirectSet)
 						{
 							AddAttributeToPreserve(SimStageGeneric->ElementCountXBinding.GetParamMapBindableVariable().GetName());
 							AddAttributeToPreserve(SimStageGeneric->ElementCountYBinding.GetParamMapBindableVariable().GetName());
@@ -3503,10 +3503,10 @@ void UNiagaraScript::SyncAliases(const FNiagaraAliasContext& ResolveAliasesConte
 	// Sync up any simulation stage name references.
 	for (FSimulationStageMetaData& SimStageMetaData : GetVMExecutableData().SimulationStageMetaData)
 	{
-		if ( !SimStageMetaData.IterationSource.IsNone() )
+		if ( !SimStageMetaData.IterationDataInterface.IsNone() )
 		{
-			FNiagaraVariable Var(FNiagaraTypeDefinition(UNiagaraDataInterface::StaticClass()), SimStageMetaData.IterationSource);
-			SimStageMetaData.IterationSource = FNiagaraUtilities::ResolveAliases(Var, ResolveAliasesContext).GetName();
+			FNiagaraVariable Var(FNiagaraTypeDefinition(UNiagaraDataInterface::StaticClass()), SimStageMetaData.IterationDataInterface);
+			SimStageMetaData.IterationDataInterface = FNiagaraUtilities::ResolveAliases(Var, ResolveAliasesContext).GetName();
 		}
 
 		if (!SimStageMetaData.EnabledBinding.IsNone())
@@ -3515,7 +3515,7 @@ void UNiagaraScript::SyncAliases(const FNiagaraAliasContext& ResolveAliasesConte
 			SimStageMetaData.EnabledBinding = FNiagaraUtilities::ResolveAliases(Var, ResolveAliasesContext).GetName();
 		}
 
-		if (SimStageMetaData.bOverrideElementCount)
+		if (SimStageMetaData.IterationSource == ENiagaraIterationSource::DirectSet)
 		{
 			if (!SimStageMetaData.ElementCountXBinding.IsNone())
 			{
diff --git a/Engine/Plugins/FX/Niagara/Source/Niagara/Private/NiagaraSimulationStageBase.cpp b/Engine/Plugins/FX/Niagara/Source/Niagara/Private/NiagaraSimulationStageBase.cpp
--- a/Engine/Plugins/FX/Niagara/Source/Niagara/Private/NiagaraSimulationStageBase.cpp
+++ b/Engine/Plugins/FX/Niagara/Source/Niagara/Private/NiagaraSimulationStageBase.cpp
@@ -96,8 +96,7 @@ bool UNiagaraSimulationStageGeneric::AppendCompileHash(FNiagaraCompileHashVisito
 	InVisitor->UpdateString(TEXT("ParticleIterationStateBinding"), ParticleIterationStateBinding.GetDataSetBindableVariable().GetName().ToString());
 	InVisitor->UpdateString(TEXT("ParticleIterationStateRange"), FString::Printf(TEXT("%d,%d"), ParticleIterationStateRange.X, ParticleIterationStateRange.Y));
 	InVisitor->UpdatePOD(TEXT("bGpuDispatchForceLinear"), bGpuDispatchForceLinear ? 1 : 0);
-	InVisitor->UpdatePOD(TEXT("bOverrideGpuDispatchType"), bOverrideGpuDispatchType ? 1 : 0);
-	InVisitor->UpdatePOD(TEXT("OverrideGpuDispatchType"), (int32)OverrideGpuDispatchType);
+	InVisitor->UpdatePOD(TEXT("DirectDispatchType"), (int32)DirectDispatchType);
 	InVisitor->UpdatePOD(TEXT("bOverrideGpuDispatchNumThreads"), bOverrideGpuDispatchNumThreads ? 1 : 0);
 	InVisitor->UpdateString(TEXT("OverrideGpuDispatchNumThreads"), FString::Printf(TEXT("%d,%d,%d"), OverrideGpuDispatchNumThreads.X, OverrideGpuDispatchNumThreads.Y, OverrideGpuDispatchNumThreads.Z));
 
@@ -116,14 +115,14 @@ bool UNiagaraSimulationStageGeneric::FillCompilationData(TArray<FNiagaraSimulati
 	SimStageData.ElementCountZBinding = ElementCountZBinding.GetName();
 	SimStageData.NumIterations = Iterations;
 	SimStageData.NumIterationsBinding = NumIterationsBinding.GetName();
-	SimStageData.IterationSource = IterationSource == ENiagaraIterationSource::DataInterface ? DataInterface.BoundVariable.GetName() : FName();
+	SimStageData.IterationSource = IterationSource;
+	SimStageData.IterationDataInterface = IterationSource == ENiagaraIterationSource::DataInterface ? DataInterface.BoundVariable.GetName() : FName();
 	SimStageData.ExecuteBehavior = ExecuteBehavior;
 	SimStageData.PartialParticleUpdate = bDisablePartialParticleUpdate == false;
 	SimStageData.bParticleIterationStateEnabled = bParticleIterationStateEnabled;
 	SimStageData.ParticleIterationStateRange = ParticleIterationStateRange;
 	SimStageData.bGpuDispatchForceLinear = bGpuDispatchForceLinear;
-	SimStageData.bOverrideGpuDispatchType = bOverrideGpuDispatchType;
-	SimStageData.OverrideGpuDispatchType = OverrideGpuDispatchType;
+	SimStageData.DirectDispatchType = DirectDispatchType;
 	SimStageData.bOverrideGpuDispatchNumThreads = bOverrideGpuDispatchNumThreads;
 	SimStageData.OverrideGpuDispatchNumThreads = OverrideGpuDispatchNumThreads;
 
@@ -147,7 +146,9 @@ void UNiagaraSimulationStageGeneric::PostInitProperties()
 	{
 		using namespace NiagaraSimulationStageLocal;
 		EnabledBinding.Setup(GetDefaultEnabledBinding(), GetDefaultEnabledBinding(), ENiagaraRendererSourceDataMode::Emitter);
+#if WITH_EDITORONLY_DATA
 		ElementCountBinding_DEPRECATED.Setup(GetDefaultElementCountBinding(), GetDefaultElementCountBinding(), ENiagaraRendererSourceDataMode::Emitter);
+#endif
 		ElementCountXBinding.Setup(GetDefaultElementCountBinding(), GetDefaultElementCountBinding(), ENiagaraRendererSourceDataMode::Emitter);
 		ElementCountYBinding.Setup(GetDefaultElementCountBinding(), GetDefaultElementCountBinding(), ENiagaraRendererSourceDataMode::Emitter);
 		ElementCountZBinding.Setup(GetDefaultElementCountBinding(), GetDefaultElementCountBinding(), ENiagaraRendererSourceDataMode::Emitter);
@@ -188,12 +189,19 @@ void UNiagaraSimulationStageGeneric::PostLoad()
 		ParticleIterationStateBinding.Setup(GetDefaultParticleStateBinding(), GetDefaultParticleStateBinding(), ENiagaraRendererSourceDataMode::Particles);
 	}
 
+	// Handle Deprecated Properties
 	if (ElementCountBinding_DEPRECATED.IsValid())
 	{
-		bOverrideGpuDispatchType = true;
-		OverrideGpuDispatchType = ENiagaraGpuDispatchType::OneD;
+		bOverrideGpuDispatchType_DEPRECATED = true;
+		OverrideGpuDispatchType_DEPRECATED = ENiagaraGpuDispatchType::OneD;
 		ElementCountXBinding = ElementCountBinding_DEPRECATED;
 	}
+
+	if (bOverrideGpuDispatchType_DEPRECATED)
+	{
+		IterationSource = ENiagaraIterationSource::DirectSet;
+		DirectDispatchType = OverrideGpuDispatchType_DEPRECATED;
+	}
 #endif
 }
 
@@ -213,80 +221,34 @@ void UNiagaraSimulationStageGeneric::PostEditChangeProperty(struct FPropertyChan
 		MemberPropertyName = PropertyChangedEvent.MemberProperty->GetFName();
 	}
 
-	bool bNeedsRecompile = false;
-	if (PropertyName == GET_MEMBER_NAME_CHECKED(UNiagaraSimulationStageGeneric, EnabledBinding))
-	{
-		bNeedsRecompile = true;
-	}
-	else if (PropertyName == GET_MEMBER_NAME_CHECKED(UNiagaraSimulationStageGeneric, ElementCountXBinding))
-	{
-		bNeedsRecompile = true;
-	}
-	else if (PropertyName == GET_MEMBER_NAME_CHECKED(UNiagaraSimulationStageGeneric, ElementCountYBinding))
-	{
-		bNeedsRecompile = true;
-	}
-	else if (PropertyName == GET_MEMBER_NAME_CHECKED(UNiagaraSimulationStageGeneric, ElementCountZBinding))
-	{
-		bNeedsRecompile = true;
-	}
-	else if (PropertyName == GET_MEMBER_NAME_CHECKED(UNiagaraSimulationStageGeneric, Iterations))
-	{
-		bNeedsRecompile = true;
-	}
-	else if (PropertyName == GET_MEMBER_NAME_CHECKED(UNiagaraSimulationStageGeneric, NumIterationsBinding))
-	{
-		bNeedsRecompile = true;
-	}
-	else if (PropertyName == GET_MEMBER_NAME_CHECKED(UNiagaraSimulationStageGeneric, IterationSource))
-	{
-		bNeedsRecompile = true;
-	}
-	else if (PropertyName == GET_MEMBER_NAME_CHECKED(UNiagaraSimulationStageGeneric, ExecuteBehavior))
-	{
-		bNeedsRecompile = true;
-	}
-	else if (PropertyName == GET_MEMBER_NAME_CHECKED(UNiagaraSimulationStageGeneric, bDisablePartialParticleUpdate))
-	{
-		bNeedsRecompile = true;
-	}
-	else if (PropertyName == GET_MEMBER_NAME_CHECKED(UNiagaraSimulationStageGeneric, DataInterface))
-	{
-		bNeedsRecompile = true;
-	}
-	else if (PropertyName == GET_MEMBER_NAME_CHECKED(UNiagaraSimulationStageGeneric, SimulationStageName))
-	{
-		bNeedsRecompile = true;
-	}
-	else if (PropertyName == GET_MEMBER_NAME_CHECKED(UNiagaraSimulationStageGeneric, bParticleIterationStateEnabled))
-	{
-		bNeedsRecompile = true;
-	}
-	else if (PropertyName == GET_MEMBER_NAME_CHECKED(UNiagaraSimulationStageGeneric, ParticleIterationStateBinding))
-	{
-		bNeedsRecompile = true;
-	}
-	else if (MemberPropertyName == GET_MEMBER_NAME_CHECKED(UNiagaraSimulationStageGeneric, ParticleIterationStateRange))
-	{
-		bNeedsRecompile = true;
-	}
-	else if (PropertyName == GET_MEMBER_NAME_CHECKED(UNiagaraSimulationStageGeneric, bGpuDispatchForceLinear))
-	{
-		bNeedsRecompile = true;
-	}
-	else if (PropertyName == GET_MEMBER_NAME_CHECKED(UNiagaraSimulationStageGeneric, bOverrideGpuDispatchType))
-	{
-		bNeedsRecompile = true;
-	}
-	else if (PropertyName == GET_MEMBER_NAME_CHECKED(UNiagaraSimulationStageGeneric, OverrideGpuDispatchType))
-	{
-		bNeedsRecompile = true;
-	}
-	else if (PropertyName == GET_MEMBER_NAME_CHECKED(UNiagaraSimulationStageGeneric, bOverrideGpuDispatchNumThreads))
+	static const TSet<FName> ForceCompilePropertyNames =
+	{
+		GET_MEMBER_NAME_CHECKED(UNiagaraSimulationStageGeneric, EnabledBinding),
+		GET_MEMBER_NAME_CHECKED(UNiagaraSimulationStageGeneric, ElementCountXBinding),
+		GET_MEMBER_NAME_CHECKED(UNiagaraSimulationStageGeneric, ElementCountYBinding),
+		GET_MEMBER_NAME_CHECKED(UNiagaraSimulationStageGeneric, ElementCountZBinding),
+		GET_MEMBER_NAME_CHECKED(UNiagaraSimulationStageGeneric, Iterations),
+		GET_MEMBER_NAME_CHECKED(UNiagaraSimulationStageGeneric, NumIterationsBinding),
+		GET_MEMBER_NAME_CHECKED(UNiagaraSimulationStageGeneric, IterationSource),
+		GET_MEMBER_NAME_CHECKED(UNiagaraSimulationStageGeneric, ExecuteBehavior),
+		GET_MEMBER_NAME_CHECKED(UNiagaraSimulationStageGeneric, bDisablePartialParticleUpdate),
+		GET_MEMBER_NAME_CHECKED(UNiagaraSimulationStageGeneric, DataInterface),
+		GET_MEMBER_NAME_CHECKED(UNiagaraSimulationStageGeneric, SimulationStageName),
+		GET_MEMBER_NAME_CHECKED(UNiagaraSimulationStageGeneric, bParticleIterationStateEnabled),
+		GET_MEMBER_NAME_CHECKED(UNiagaraSimulationStageGeneric, ParticleIterationStateBinding),
+		GET_MEMBER_NAME_CHECKED(UNiagaraSimulationStageGeneric, bGpuDispatchForceLinear),
+		GET_MEMBER_NAME_CHECKED(UNiagaraSimulationStageGeneric, DirectDispatchType),
+		GET_MEMBER_NAME_CHECKED(UNiagaraSimulationStageGeneric, bOverrideGpuDispatchNumThreads),
+	};
+
+	static const TSet<FName> ForceCompileMemberPropertyNamess =
 	{
-		bNeedsRecompile = true;
-	}
-	else if (MemberPropertyName == GET_MEMBER_NAME_CHECKED(UNiagaraSimulationStageGeneric, OverrideGpuDispatchNumThreads))
+		GET_MEMBER_NAME_CHECKED(UNiagaraSimulationStageGeneric, ParticleIterationStateRange),
+	};
+
+	bool bNeedsRecompile = ForceCompilePropertyNames.Contains(PropertyName) || ForceCompileMemberPropertyNamess.Contains(MemberPropertyName);
+
+	if (MemberPropertyName == GET_MEMBER_NAME_CHECKED(UNiagaraSimulationStageGeneric, OverrideGpuDispatchNumThreads))
 	{
 		OverrideGpuDispatchNumThreads.X = FMath::Max(OverrideGpuDispatchNumThreads.X, 1);
 		OverrideGpuDispatchNumThreads.Y = FMath::Max(OverrideGpuDispatchNumThreads.Y, 1);
@@ -322,11 +284,19 @@ FVersionedNiagaraEmitter UNiagaraSimulationStageBase::GetOuterEmitter() const
 
 FName UNiagaraSimulationStageGeneric::GetStackContextReplacementName() const 
 {
-	if (IterationSource == ENiagaraIterationSource::Particles)
-		return NAME_None;
-	else if (IterationSource == ENiagaraIterationSource::DataInterface)
-		return DataInterface.BoundVariable.GetName();
-	ensureMsgf(false, TEXT("Should not get here! Need to handle unknown case!"));
+	switch (IterationSource)
+	{
+		case ENiagaraIterationSource::Particles:
+		case ENiagaraIterationSource::DirectSet:
+			return NAME_None;
+
+		case ENiagaraIterationSource::DataInterface:
+			return DataInterface.BoundVariable.GetName();
+
+		default:
+			ensureMsgf(false, TEXT("Should not get here! Need to handle unknown case!"));
+			break;
+	}
 	return NAME_None;
 }
 #endif
diff --git a/Engine/Plugins/FX/Niagara/Source/Niagara/Public/NiagaraCommon.h b/Engine/Plugins/FX/Niagara/Source/Niagara/Public/NiagaraCommon.h
--- a/Engine/Plugins/FX/Niagara/Source/Niagara/Public/NiagaraCommon.h
+++ b/Engine/Plugins/FX/Niagara/Source/Niagara/Public/NiagaraCommon.h
@@ -894,14 +894,6 @@ enum class ENiagaraScriptGroup : uint8
 	Max
 };
 
-
-UENUM()
-enum class ENiagaraIterationSource : uint8
-{
-	Particles = 0,
-	DataInterface
-};
-
 UENUM()
 enum ENiagaraBindingSource : int
 {
diff --git a/Engine/Plugins/FX/Niagara/Source/Niagara/Public/NiagaraSimulationStageBase.h b/Engine/Plugins/FX/Niagara/Source/Niagara/Public/NiagaraSimulationStageBase.h
--- a/Engine/Plugins/FX/Niagara/Source/Niagara/Public/NiagaraSimulationStageBase.h
+++ b/Engine/Plugins/FX/Niagara/Source/Niagara/Public/NiagaraSimulationStageBase.h
@@ -64,33 +64,6 @@ public:
 	UPROPERTY(EditAnywhere, Category = "Simulation Stage")
 	FNiagaraVariableAttributeBinding EnabledBinding;
 
-	UPROPERTY()
-	FNiagaraVariableAttributeBinding ElementCountBinding_DEPRECATED;
-
-	/**
-	Integer binding to override the number of elements the stage will execute along X.
-	For example, if you want to iterate over a custom source such as triangles on a mesh you can
-	set an int to the triangle count in an emitter script and bind that as the element count.
-	*/
-	UPROPERTY(EditAnywhere, Category = "Simulation Stage", meta = (DisplayName = "ElementCount X"))
-	FNiagaraVariableAttributeBinding ElementCountXBinding;
-
-	/**
-	Integer binding to override the number of elements the stage will execute along Y.
-	For example, if you want to iterate over a 2D texture you can set an int to the texture height
-	in an emitter script and bind that as the element count.
-	*/
-	UPROPERTY(EditAnywhere, Category = "Simulation Stage", meta = (DisplayName = "ElementCount Y"))
-	FNiagaraVariableAttributeBinding ElementCountYBinding;
-
-	/**
-	Integer binding to override the number of elements the stage will execute along Z.
-	For example, if you want to iterate over a 3D texture you can set an int to the texture depth
-	in an emitter script and bind that as the element count.
-	*/
-	UPROPERTY(EditAnywhere, Category = "Simulation Stage", meta = (DisplayName = "ElementCount Z"))
-	FNiagaraVariableAttributeBinding ElementCountZBinding;
-	
 	/**
 	Select what we should be iterating over, particles or data interfaces.
 	The source provides things such as element count (when not overriden) and stack context variables (i.e. attributes on grids)
@@ -111,9 +84,6 @@ public:
 	UPROPERTY(EditAnywhere, Category = "Simulation Stage")
 	FNiagaraVariableAttributeBinding NumIterationsBinding;
 
-	UPROPERTY()
-	uint32 bSpawnOnly_DEPRECATED : 1;
-
 	/**
 	Controls when the simulation stage should execute, only valid for data interface iteration stages
 	*/
@@ -154,17 +124,6 @@ public:
 	UPROPERTY(EditAnywhere, Category = "Simulation Stage")
 	uint32 bGpuDispatchForceLinear : 1;
 
-	/**
-	When enabled we use a custom number of element count for the dispatch.
-	NOTE: Cannot be used with GpuDispatchForceLinear.
-	*/
-	UPROPERTY(EditAnywhere, Category = "Simulation Stage")
-	uint32 bOverrideGpuDispatchType : 1;
-
-	/** Dimensions to use for dispatch. */
-	UPROPERTY(EditAnywhere, Category = "Simulation Stage", meta = (DisplayName = "Gpu Dispatch Type"))
-	ENiagaraGpuDispatchType OverrideGpuDispatchType = ENiagaraGpuDispatchType::OneD;
-
 	/** When enabled we use a custom number of threads for the dispatch. */
 	UPROPERTY(EditAnywhere, Category = "Simulation Stage")
 	uint32 bOverrideGpuDispatchNumThreads : 1;
@@ -172,6 +131,51 @@ public:
 	UPROPERTY(EditAnywhere, Category = "Simulation Stage")
 	FIntVector OverrideGpuDispatchNumThreads = FIntVector(64, 1, 1);
 
+	/** Dimensions to use for dispatch. */
+	UPROPERTY(EditAnywhere, Category = "Simulation Stage")
+	ENiagaraGpuDispatchType DirectDispatchType = ENiagaraGpuDispatchType::OneD;
+
+	/**
+	Integer binding to override the number of elements the stage will execute along X.
+	For example, if you want to iterate over a custom source such as triangles on a mesh you can
+	set an int to the triangle count in an emitter script and bind that as the element count.
+	*/
+	UPROPERTY(EditAnywhere, Category = "Simulation Stage", meta = (DisplayName = "ElementCount X"))
+	FNiagaraVariableAttributeBinding ElementCountXBinding;
+
+	/**
+	Integer binding to override the number of elements the stage will execute along Y.
+	For example, if you want to iterate over a 2D texture you can set an int to the texture height
+	in an emitter script and bind that as the element count.
+	*/
+	UPROPERTY(EditAnywhere, Category = "Simulation Stage", meta = (DisplayName = "ElementCount Y", EditCondition = "DirectDispatchType == ENiagaraGpuDispatchType::TwoD || DirectDispatchType == ENiagaraGpuDispatchType::ThreeD"))
+	FNiagaraVariableAttributeBinding ElementCountYBinding;
+
+	/**
+	Integer binding to override the number of elements the stage will execute along Z.
+	For example, if you want to iterate over a 3D texture you can set an int to the texture depth
+	in an emitter script and bind that as the element count.
+	*/
+	UPROPERTY(EditAnywhere, Category = "Simulation Stage", meta = (DisplayName = "ElementCount Z", EditCondition = "DirectDispatchType == ENiagaraGpuDispatchType::ThreeD"))
+	FNiagaraVariableAttributeBinding ElementCountZBinding;
+
+#if WITH_EDITORONLY_DATA
+	// Deprecated Properties
+	UPROPERTY()
+	FNiagaraVariableAttributeBinding ElementCountBinding_DEPRECATED;
+
+	UPROPERTY()
+	uint32 bSpawnOnly_DEPRECATED : 1;
+	
+	UPROPERTY()
+	uint32 bOverrideGpuDispatchType_DEPRECATED : 1;
+
+	UPROPERTY()
+	ENiagaraGpuDispatchType OverrideGpuDispatchType_DEPRECATED = ENiagaraGpuDispatchType::OneD;
+#endif
+	//////////////////////////////////////////////////////////////////////////////////////////////////////////////////////////////////
+
+
 	virtual void PostInitProperties() override;
 	virtual void PostLoad() override;
 	virtual bool AppendCompileHash(FNiagaraCompileHashVisitor* InVisitor) const override;
@@ -183,4 +187,4 @@ public:
 
 	virtual void PostEditChangeProperty(struct FPropertyChangedEvent& PropertyChangedEvent) override;
 #endif
-};
\ No newline at end of file
+};
diff --git a/Engine/Plugins/FX/Niagara/Source/Niagara/Public/NiagaraSimulationStageCompileData.h b/Engine/Plugins/FX/Niagara/Source/Niagara/Public/NiagaraSimulationStageCompileData.h
--- a/Engine/Plugins/FX/Niagara/Source/Niagara/Public/NiagaraSimulationStageCompileData.h
+++ b/Engine/Plugins/FX/Niagara/Source/Niagara/Public/NiagaraSimulationStageCompileData.h
@@ -14,15 +14,16 @@ struct FNiagaraSimulationStageCompilationData
 	FName                           ElementCountZBinding;
 	uint32                          NumIterations = 1;
 	FName                           NumIterationsBinding;
-	FName                           IterationSource;
+	ENiagaraIterationSource			IterationSource = ENiagaraIterationSource::Particles;
+	FName							IterationDataInterface;
+	FName							IterationDirectBinding;
 	ENiagaraSimStageExecuteBehavior ExecuteBehavior = ENiagaraSimStageExecuteBehavior::Always;
 	mutable bool                    PartialParticleUpdate = false;
 	bool                            bParticleIterationStateEnabled = false;
 	FName                           ParticleIterationStateBinding;
 	FIntPoint                       ParticleIterationStateRange = FIntPoint::ZeroValue;
 	bool                            bGpuDispatchForceLinear = false;
-	bool                            bOverrideGpuDispatchType = false;
-	ENiagaraGpuDispatchType         OverrideGpuDispatchType = ENiagaraGpuDispatchType::OneD;
+	ENiagaraGpuDispatchType         DirectDispatchType = ENiagaraGpuDispatchType::OneD;
 	bool                            bOverrideGpuDispatchNumThreads = false;
 	FIntVector                      OverrideGpuDispatchNumThreads = FIntVector(1, 1, 1);
 };
diff --git a/Engine/Plugins/FX/Niagara/Source/NiagaraCore/Public/NiagaraCore.h b/Engine/Plugins/FX/Niagara/Source/NiagaraCore/Public/NiagaraCore.h
--- a/Engine/Plugins/FX/Niagara/Source/NiagaraCore/Public/NiagaraCore.h
+++ b/Engine/Plugins/FX/Niagara/Source/NiagaraCore/Public/NiagaraCore.h
@@ -4,4 +4,17 @@
 
 #include "CoreTypes.h"
 
+#include "NiagaraCore.generated.h"
+
 typedef uint64 FNiagaraSystemInstanceID;
+
+UENUM()
+enum class ENiagaraIterationSource : uint8
+{
+	/** Iterate over all active particles. */
+	Particles = 0,
+	/** Iterate over all elements in the data interface. */
+	DataInterface,
+	/** Iterate over a user provided number of elements. */
+	DirectSet,
+};
diff --git a/Engine/Plugins/FX/Niagara/Source/NiagaraEditor/Private/Customizations/NiagaraSimulationStageCustomization.cpp b/Engine/Plugins/FX/Niagara/Source/NiagaraEditor/Private/Customizations/NiagaraSimulationStageCustomization.cpp
--- a/Engine/Plugins/FX/Niagara/Source/NiagaraEditor/Private/Customizations/NiagaraSimulationStageCustomization.cpp
+++ b/Engine/Plugins/FX/Niagara/Source/NiagaraEditor/Private/Customizations/NiagaraSimulationStageCustomization.cpp
@@ -60,10 +60,12 @@ void FNiagaraSimulationStageGenericCustomization::CustomizeDetails(class IDetail
 	static const FName NAME_SimulationStage = TEXT("Simulation Stage");
 	static const FName NAME_ParticleParameters = TEXT("Particle Parameters");
 	static const FName NAME_DataInterfaceParameters = TEXT("DataInterface Parameters");
+	static const FName NAME_DirectSetParameters = TEXT("Direct Set Parameters");
 	static const FName NAME_DispatchParameters = TEXT("Dispatch Parameters");
 	IDetailCategoryBuilder& SimStageCategory = DetailBuilder.EditCategory(NAME_SimulationStage);
 	IDetailCategoryBuilder& ParticleStageCategory = DetailBuilder.EditCategory(NAME_ParticleParameters);
 	IDetailCategoryBuilder& DataInterfaceStageCategory = DetailBuilder.EditCategory(NAME_DataInterfaceParameters);
+	IDetailCategoryBuilder& DirectSetStageCategory = DetailBuilder.EditCategory(NAME_DirectSetParameters);
 	IDetailCategoryBuilder& DispatchParametersCategory = DetailBuilder.EditCategory(NAME_DispatchParameters);
 
 	// Hide all properties by default
@@ -92,41 +94,37 @@ void FNiagaraSimulationStageGenericCustomization::CustomizeDetails(class IDetail
 	SimStageCategory.AddProperty(DetailBuilder.GetProperty(GET_MEMBER_NAME_CHECKED(UNiagaraSimulationStageGeneric, ExecuteBehavior)));
 
 	// Paterial Iteration Category
-	if (SimStage->IterationSource == ENiagaraIterationSource::Particles)
+	switch (SimStage->IterationSource)
 	{
-		ParticleStageCategory.AddProperty(DetailBuilder.GetProperty(GET_MEMBER_NAME_CHECKED(UNiagaraSimulationStageGeneric, bDisablePartialParticleUpdate)));
-
-		ParticleStageCategory.AddProperty(DetailBuilder.GetProperty(GET_MEMBER_NAME_CHECKED(UNiagaraSimulationStageGeneric, bParticleIterationStateEnabled)));
-		if (SimStage->bParticleIterationStateEnabled)
+		case ENiagaraIterationSource::Particles:
 		{
-			ParticleStageCategory.AddProperty(DetailBuilder.GetProperty(GET_MEMBER_NAME_CHECKED(UNiagaraSimulationStageGeneric, ParticleIterationStateBinding)));
-			ParticleStageCategory.AddProperty(DetailBuilder.GetProperty(GET_MEMBER_NAME_CHECKED(UNiagaraSimulationStageGeneric, ParticleIterationStateRange)));
+			ParticleStageCategory.AddProperty(DetailBuilder.GetProperty(GET_MEMBER_NAME_CHECKED(UNiagaraSimulationStageGeneric, bDisablePartialParticleUpdate)));
+
+			ParticleStageCategory.AddProperty(DetailBuilder.GetProperty(GET_MEMBER_NAME_CHECKED(UNiagaraSimulationStageGeneric, bParticleIterationStateEnabled)));
+			if (SimStage->bParticleIterationStateEnabled)
+			{
+				ParticleStageCategory.AddProperty(DetailBuilder.GetProperty(GET_MEMBER_NAME_CHECKED(UNiagaraSimulationStageGeneric, ParticleIterationStateBinding)));
+				ParticleStageCategory.AddProperty(DetailBuilder.GetProperty(GET_MEMBER_NAME_CHECKED(UNiagaraSimulationStageGeneric, ParticleIterationStateRange)));
+			}
+			break;
 		}
-	}
-	// DataInterface Iteration Category
-	else
-	{
-		DataInterfaceStageCategory.AddProperty(DetailBuilder.GetProperty(GET_MEMBER_NAME_CHECKED(UNiagaraSimulationStageGeneric, DataInterface)));
-		DataInterfaceStageCategory.AddProperty(DetailBuilder.GetProperty(GET_MEMBER_NAME_CHECKED(UNiagaraSimulationStageGeneric, bGpuDispatchForceLinear)));
-	}
-
-	// Dispatch Parameters
-	DispatchParametersCategory.AddProperty(DetailBuilder.GetProperty(GET_MEMBER_NAME_CHECKED(UNiagaraSimulationStageGeneric, bOverrideGpuDispatchType)));
-	if (SimStage->bOverrideGpuDispatchType)
-	{
-		DispatchParametersCategory.AddProperty(DetailBuilder.GetProperty(GET_MEMBER_NAME_CHECKED(UNiagaraSimulationStageGeneric, OverrideGpuDispatchType)));
-
-		DispatchParametersCategory.AddProperty(DetailBuilder.GetProperty(GET_MEMBER_NAME_CHECKED(UNiagaraSimulationStageGeneric, ElementCountXBinding)));
-		if (SimStage->OverrideGpuDispatchType >= ENiagaraGpuDispatchType::TwoD)
+		case ENiagaraIterationSource::DataInterface:
 		{
-			DispatchParametersCategory.AddProperty(DetailBuilder.GetProperty(GET_MEMBER_NAME_CHECKED(UNiagaraSimulationStageGeneric, ElementCountYBinding)));
+			DataInterfaceStageCategory.AddProperty(DetailBuilder.GetProperty(GET_MEMBER_NAME_CHECKED(UNiagaraSimulationStageGeneric, DataInterface)));
+			DataInterfaceStageCategory.AddProperty(DetailBuilder.GetProperty(GET_MEMBER_NAME_CHECKED(UNiagaraSimulationStageGeneric, bGpuDispatchForceLinear)));
+			break;
 		}
-		if (SimStage->OverrideGpuDispatchType >= ENiagaraGpuDispatchType::ThreeD)
+		case ENiagaraIterationSource::DirectSet:
 		{
-			DispatchParametersCategory.AddProperty(DetailBuilder.GetProperty(GET_MEMBER_NAME_CHECKED(UNiagaraSimulationStageGeneric, ElementCountZBinding)));
+			DirectSetStageCategory.AddProperty(DetailBuilder.GetProperty(GET_MEMBER_NAME_CHECKED(UNiagaraSimulationStageGeneric, DirectDispatchType)));
+			DirectSetStageCategory.AddProperty(DetailBuilder.GetProperty(GET_MEMBER_NAME_CHECKED(UNiagaraSimulationStageGeneric, ElementCountXBinding)));
+			DirectSetStageCategory.AddProperty(DetailBuilder.GetProperty(GET_MEMBER_NAME_CHECKED(UNiagaraSimulationStageGeneric, ElementCountYBinding)));
+			DirectSetStageCategory.AddProperty(DetailBuilder.GetProperty(GET_MEMBER_NAME_CHECKED(UNiagaraSimulationStageGeneric, ElementCountZBinding)));
+			break;
 		}
 	}
 
+	// Common Settings
 	DispatchParametersCategory.AddProperty(DetailBuilder.GetProperty(GET_MEMBER_NAME_CHECKED(UNiagaraSimulationStageGeneric, bOverrideGpuDispatchNumThreads)));
 	if (SimStage->bOverrideGpuDispatchNumThreads)
 	{
diff --git a/Engine/Plugins/FX/Niagara/Source/NiagaraEditor/Private/NiagaraHlslTranslator.cpp b/Engine/Plugins/FX/Niagara/Source/NiagaraEditor/Private/NiagaraHlslTranslator.cpp
--- a/Engine/Plugins/FX/Niagara/Source/NiagaraEditor/Private/NiagaraHlslTranslator.cpp
+++ b/Engine/Plugins/FX/Niagara/Source/NiagaraEditor/Private/NiagaraHlslTranslator.cpp
@@ -1300,13 +1300,14 @@ const FNiagaraTranslateResults &FHlslNiagaraTranslator::Translate(const FNiagara
 					TranslationStages[Index].ExecuteBehavior = CompileSimStageData.ExecuteBehavior;
 					TranslationStages[Index].bPartialParticleUpdate = CompileSimStageData.PartialParticleUpdate;
 					TranslationStages[Index].IterationSource = CompileSimStageData.IterationSource;
+					TranslationStages[Index].IterationDataInterface = CompileSimStageData.IterationDataInterface;
+					TranslationStages[Index].IterationDirectBinding = CompileSimStageData.IterationDirectBinding;
 					TranslationStages[Index].NumIterationsBinding = CompileSimStageData.NumIterationsBinding;
 					TranslationStages[Index].bParticleIterationStateEnabled = CompileSimStageData.bParticleIterationStateEnabled;
 					TranslationStages[Index].ParticleIterationStateBinding = CompileSimStageData.ParticleIterationStateBinding;
 					TranslationStages[Index].ParticleIterationStateRange = CompileSimStageData.ParticleIterationStateRange;
 					TranslationStages[Index].bGpuDispatchForceLinear = CompileSimStageData.bGpuDispatchForceLinear;
-					TranslationStages[Index].bOverrideGpuDispatchType = CompileSimStageData.bOverrideGpuDispatchType;
-					TranslationStages[Index].OverrideGpuDispatchType = CompileSimStageData.OverrideGpuDispatchType;
+					TranslationStages[Index].DirectDispatchType = CompileSimStageData.DirectDispatchType;
 					TranslationStages[Index].bOverrideGpuDispatchNumThreads = CompileSimStageData.bOverrideGpuDispatchNumThreads;
 					TranslationStages[Index].OverrideGpuDispatchNumThreads = CompileSimStageData.OverrideGpuDispatchNumThreads;
 
@@ -1372,6 +1373,8 @@ const FNiagaraTranslateResults &FHlslNiagaraTranslator::Translate(const FNiagara
 					SimulationStageMetaData.ElementCountZBinding = InCompileData->CompileSimStageData[SourceSimStageIndex].ElementCountZBinding;
 					SimulationStageMetaData.ExecuteBehavior = TranslationStages[Index].ExecuteBehavior;
 					SimulationStageMetaData.IterationSource = InCompileData->CompileSimStageData[SourceSimStageIndex].IterationSource;
+					SimulationStageMetaData.IterationDataInterface = InCompileData->CompileSimStageData[SourceSimStageIndex].IterationDataInterface;
+					SimulationStageMetaData.IterationDirectBinding = InCompileData->CompileSimStageData[SourceSimStageIndex].IterationDirectBinding;
 					SimulationStageMetaData.NumIterationsBinding = InCompileData->CompileSimStageData[SourceSimStageIndex].NumIterationsBinding;
 					SimulationStageMetaData.NumIterations = TranslationStages[Index].NumIterations;
 					SimulationStageMetaData.bWritesParticles = TranslationStages[Index].bWritesParticles;
@@ -1379,25 +1382,23 @@ const FNiagaraTranslateResults &FHlslNiagaraTranslator::Translate(const FNiagara
 					SimulationStageMetaData.bParticleIterationStateEnabled = TranslationStages[Index].bParticleIterationStateEnabled;
 					SimulationStageMetaData.ParticleIterationStateBinding = TranslationStages[Index].ParticleIterationStateBinding;
 					SimulationStageMetaData.ParticleIterationStateRange = TranslationStages[Index].ParticleIterationStateRange;
-					SimulationStageMetaData.bOverrideElementCount = TranslationStages[Index].bOverrideGpuDispatchType;
 
 					// Determine dispatch information from iteration source (if we have one)
 					SimulationStageMetaData.GpuDispatchType = ENiagaraGpuDispatchType::OneD;
 					SimulationStageMetaData.GpuDispatchNumThreads = FNiagaraShader::GetDefaultThreadGroupSize(ENiagaraGpuDispatchType::OneD);
-					if ( !SimulationStageMetaData.IterationSource.IsNone() )
+					if ( SimulationStageMetaData.IterationSource == ENiagaraIterationSource::DataInterface )
 					{
-						if (const FNiagaraVariable* IterationSourceVar = CompileData->EncounteredVariables.FindByPredicate([&](const FNiagaraVariable& VarInfo) { return VarInfo.GetName() == SimulationStageMetaData.IterationSource; }))
+						if (SimulationStageMetaData.IterationDataInterface.IsNone())
+						{
+							Error(LOCTEXT("NoIterationDataInterface", "An data interface must be bound to a data interface iteration stage"), nullptr, nullptr);
+							return TranslateResults;
+						}
+
+						if (const FNiagaraVariable* IterationSourceVar = CompileData->EncounteredVariables.FindByPredicate([&](const FNiagaraVariable& VarInfo) { return VarInfo.GetName() == SimulationStageMetaData.IterationDataInterface; }))
 						{
 							if ( UNiagaraDataInterface* IteratinoSourceCDO = CompileDuplicateData->GetDuplicatedDataInterfaceCDOForClass(IterationSourceVar->GetType().GetClass()) )
 							{
-								if (TranslationStages[Index].bGpuDispatchForceLinear)
-								{
-									SimulationStageMetaData.GpuDispatchType = ENiagaraGpuDispatchType::OneD;
-								}
-								else
-								{
-									SimulationStageMetaData.GpuDispatchType = TranslationStages[Index].bOverrideGpuDispatchType ? TranslationStages[Index].OverrideGpuDispatchType : IteratinoSourceCDO->GetGpuDispatchType();
-								}
+								SimulationStageMetaData.GpuDispatchType = TranslationStages[Index].bGpuDispatchForceLinear ? ENiagaraGpuDispatchType::OneD : IteratinoSourceCDO->GetGpuDispatchType();
 
 								if (TranslationStages[Index].bOverrideGpuDispatchNumThreads)
 								{
@@ -1410,6 +1411,10 @@ const FNiagaraTranslateResults &FHlslNiagaraTranslator::Translate(const FNiagara
 							}
 						}
 					}
+					else if (SimulationStageMetaData.IterationSource == ENiagaraIterationSource::DirectSet)
+					{
+						SimulationStageMetaData.GpuDispatchType = TranslationStages[Index].DirectDispatchType;
+					}
 
 					// Increment source stage index
 					++SourceSimStageIndex;
@@ -1438,6 +1443,8 @@ const FNiagaraTranslateResults &FHlslNiagaraTranslator::Translate(const FNiagara
 				if (StageGuid == CompileOptions.TargetUsageId && InCompileData->CompileSimStageData.IsValidIndex(StageIdx))
 				{
 					TranslationStages[0].IterationSource = InCompileData->CompileSimStageData[StageIdx].IterationSource;
+					TranslationStages[0].IterationDataInterface = InCompileData->CompileSimStageData[StageIdx].IterationDataInterface;
+					TranslationStages[0].IterationDirectBinding = InCompileData->CompileSimStageData[StageIdx].IterationDirectBinding;
 				}
 			}
 		}
@@ -1651,7 +1658,7 @@ const FNiagaraTranslateResults &FHlslNiagaraTranslator::Translate(const FNiagara
 				CurrentParamMapIndices.Add(i);
 				PinToCodeChunks.Empty();
 				PinToCodeChunks.AddDefaulted(1);				
-				FName IterSource = TranslationStages[i].IterationSource;
+				FName IterSource = TranslationStages[i].GetIterationDataInterface();
 				ActiveHistoryForFunctionCalls.BeginUsage(TranslationStages[i].ScriptUsage, IterSource);
 				TranslationStages[i].OutputNode->Compile(ThisTranslator, OutputChunks);
 				HandleSimStageSetupAndTeardown(i, StageSetupAndTeardownHLSL);
@@ -1671,7 +1678,7 @@ const FNiagaraTranslateResults &FHlslNiagaraTranslator::Translate(const FNiagara
 		if (TranslationStages[0].ShouldDoSpawnOnlyLogic())
 			bInitializedDefaults = false;
 
-		FName IterSource = TranslationStages[0].IterationSource;
+		FName IterSource = TranslationStages[0].GetIterationDataInterface();
 		ActiveHistoryForFunctionCalls.BeginUsage(TranslationStages[0].ScriptUsage, IterSource);
 		TranslationStages[0].OutputNode->Compile(ThisTranslator, OutputChunks);
 		ActiveHistoryForFunctionCalls.EndUsage();
@@ -2153,8 +2160,10 @@ const FNiagaraTranslateResults &FHlslNiagaraTranslator::Translate(const FNiagara
 void FHlslNiagaraTranslator::HandleSimStageSetupAndTeardown(int32 InWhichStage, FString& OutHlsl)
 {
 	FHlslNiagaraTranslationStage& TranslationStage = TranslationStages[InWhichStage];
+	const FName IterationDataInterface = TranslationStage.GetIterationDataInterface();
+
 	// If we're particles then do nothing different..
-	if (TranslationStage.IterationSource == NAME_None)
+	if (IterationDataInterface == NAME_None)
 		return;
 
 	FExpressionPermutationContext PermutationContext(OutHlsl);
@@ -2162,7 +2171,7 @@ void FHlslNiagaraTranslator::HandleSimStageSetupAndTeardown(int32 InWhichStage,
 
 	// Ok, we're iterating over a known iteration source. Let's find it in the parameter map history so we know type/etc.
 	FNiagaraVariable IterationSourceVar;
-	const FNiagaraVariable* FoundVar = CompileData->EncounteredVariables.FindByPredicate([&](const FNiagaraVariable& VarInfo) { return VarInfo.GetName() == TranslationStage.IterationSource; });
+	const FNiagaraVariable* FoundVar = CompileData->EncounteredVariables.FindByPredicate([&](const FNiagaraVariable& VarInfo) { return VarInfo.GetName() == IterationDataInterface; });
 	if (FoundVar != nullptr)
 	{
 		IterationSourceVar = *FoundVar;
@@ -2170,14 +2179,14 @@ void FHlslNiagaraTranslator::HandleSimStageSetupAndTeardown(int32 InWhichStage,
 	
 	if (!IterationSourceVar.IsValid())
 	{
-		Error(FText::Format(LOCTEXT("CannotFindIterationSourceInParamMap", "Variable {0} missing in graphs referenced during compile!"), FText::FromName(TranslationStage.IterationSource)), nullptr, nullptr);
+		Error(FText::Format(LOCTEXT("CannotFindIterationSourceInParamMap", "Variable {0} missing in graphs referenced during compile!"), FText::FromName(IterationDataInterface)), nullptr, nullptr);
 		return;
 	}
 
 	UNiagaraDataInterface* CDO = CompileDuplicateData->GetDuplicatedDataInterfaceCDOForClass(IterationSourceVar.GetType().GetClass());
 	if (!CDO || !IterationSourceVar.GetType().GetClass())
 	{
-		Error(FText::Format(LOCTEXT("CannotFindIterationSourceCDOInParamMap", "Variable {0}'s cached CDO for class was missing during compile!"), FText::FromName(TranslationStage.IterationSource)), nullptr, nullptr);
+		Error(FText::Format(LOCTEXT("CannotFindIterationSourceCDOInParamMap", "Variable {0}'s cached CDO for class was missing during compile!"), FText::FromName(IterationDataInterface)), nullptr, nullptr);
 		return;
 	}
 
@@ -2194,7 +2203,7 @@ void FHlslNiagaraTranslator::HandleSimStageSetupAndTeardown(int32 InWhichStage,
 		{
 			const FNiagaraVariable& Var = ParamMapHistories[ParamHistoryIdx].Variables[i];
 
-			if (Var.IsInNameSpace(TranslationStage.IterationSource))
+			if (Var.IsInNameSpace(IterationDataInterface))
 			{
 				if (ParamMapHistories[ParamHistoryIdx].PerVariableReadHistory[i].Num() > 0 && !ReadVars.Contains(Var))
 					ReadVars.Emplace(Var);
@@ -2210,7 +2219,7 @@ void FHlslNiagaraTranslator::HandleSimStageSetupAndTeardown(int32 InWhichStage,
 		{
 			const FNiagaraVariable& Var = ParamMapHistories[ParamHistoryIdx].Variables[i];
 
-			if (Var.IsInNameSpace(TranslationStage.IterationSource))
+			if (Var.IsInNameSpace(IterationDataInterface))
 			{
 				if (!AllVars.Contains(Var))
 					AllVars.Emplace(Var);
@@ -2223,7 +2232,7 @@ void FHlslNiagaraTranslator::HandleSimStageSetupAndTeardown(int32 InWhichStage,
 	for (int32 i = 0; i < CompilationOutput.ScriptData.DataInterfaceInfo.Num(); i++)
 	{
 		FNiagaraScriptDataInterfaceCompileInfo& Info = CompilationOutput.ScriptData.DataInterfaceInfo[i];
-		if (TranslationStage.IterationSource == Info.Name)
+		if (IterationDataInterface == Info.Name)
 		{
 			DataInterfaceOwnerIndex = i;
 			break;
@@ -2235,7 +2244,7 @@ void FHlslNiagaraTranslator::HandleSimStageSetupAndTeardown(int32 InWhichStage,
 	// 2) Someone called a function that was marked to write 
 	const int32 SourceSimStage = TranslationStage.SimulationStageIndex;
 	ensure(CompilationOutput.ScriptData.SimulationStageMetaData.IsValidIndex(SourceSimStage));
-	bool bWroteToIterationSource = CompilationOutput.ScriptData.SimulationStageMetaData[SourceSimStage].OutputDestinations.Contains(TranslationStage.IterationSource);
+	bool bWroteToIterationSource = CompilationOutput.ScriptData.SimulationStageMetaData[SourceSimStage].OutputDestinations.Contains(TranslationStage.GetIterationDataInterface());
 	if (WriteVars.Num() > 0)
 		bWroteToIterationSource = true;
 
@@ -2278,20 +2287,20 @@ void FHlslNiagaraTranslator::HandleSimStageSetupAndTeardown(int32 InWhichStage,
 		// If we haven't created it by now, bail out. 
 		if (DataInterfaceOwnerIndex == INDEX_NONE && bNeedsDIOwner)
 		{
-			Error(FText::Format(LOCTEXT("CannotRegisterDataInterface", "Variable {0}'s cannot register as a data interface!"), FText::FromName(TranslationStage.IterationSource)), nullptr, nullptr);
+			Error(FText::Format(LOCTEXT("CannotRegisterDataInterface", "Variable {0}'s cannot register as a data interface!"), FText::FromName(IterationDataInterface)), nullptr, nullptr);
 			return;
 		}
 
 		// It is an invalid state to use the IterationSource and StackContext namespace without implementing SupportsIterationSourceNamespaceAttributesHLSL
 		if (ReadVars.Num() > 0 && !bNeedsAttributeRead)
 		{
-			Error(FText::Format(LOCTEXT("CannotUseContextRead", "Variable {0} cannot be used in conjunction with StackContext namespace variable reads! It must implement SupportsIterationSourceNamespaceAttributesHLSL."), FText::FromName(TranslationStage.IterationSource)), nullptr, nullptr);
+			Error(FText::Format(LOCTEXT("CannotUseContextRead", "Variable {0} cannot be used in conjunction with StackContext namespace variable reads! It must implement SupportsIterationSourceNamespaceAttributesHLSL."), FText::FromName(IterationDataInterface)), nullptr, nullptr);
 			return;
 		}
 
 		if (WriteVars.Num() > 0 && !bNeedsAttributeWrite)
 		{
-			Error(FText::Format(LOCTEXT("CannotUseContextWrite", "Variable {0} cannot be used in conjunction with StackContext namespace variable writes! It must implement SupportsIterationSourceNamespaceAttributesHLSL."), FText::FromName(TranslationStage.IterationSource)), nullptr, nullptr);
+			Error(FText::Format(LOCTEXT("CannotUseContextWrite", "Variable {0} cannot be used in conjunction with StackContext namespace variable writes! It must implement SupportsIterationSourceNamespaceAttributesHLSL."), FText::FromName(IterationDataInterface)), nullptr, nullptr);
 			return;
 		}
 
@@ -3048,8 +3057,10 @@ void FHlslNiagaraTranslator::DefineMainGPUFunctions(
 		int32 StartIdx = 1;
 		for (int32 i = StartIdx; i < TranslationStages.Num(); i++)
 		{
+			const FName IterationDataInterface = TranslationStages[i].GetIterationDataInterface();
+
 			// No need to load particle data for stages with an iteration source, since those do not run one thread per particle.
-			if (TranslationStages[i].IterationSource != NAME_None)
+			if (IterationDataInterface != NAME_None)
 			{
 				if (TranslationStages[i].CustomReadFunction.IsEmpty())
 				{
@@ -3059,7 +3070,7 @@ void FHlslNiagaraTranslator::DefineMainGPUFunctions(
 
 			PermutationContext.AddBranch(*this, TranslationStages[i]);
 
-			if (TranslationStages[i].IterationSource != NAME_None)
+			if (IterationDataInterface != NAME_None)
 			{
 				if (!TranslationStages[i].CustomReadFunction.IsEmpty())
 				{
@@ -3202,8 +3213,10 @@ void FHlslNiagaraTranslator::DefineMainGPUFunctions(
 		int32 StartIdx = 1;
 		for (int32 i = StartIdx; i < TranslationStages.Num(); i++)
 		{
+			const FName IterationDataInterface = TranslationStages[i].GetIterationDataInterface();
+
 			// No need to store particle data for stages with an iteration source, since those do not run one thread per particle.
-			if (TranslationStages[i].IterationSource != NAME_None)
+			if (IterationDataInterface != NAME_None)
 			{
 				if (TranslationStages[i].CustomWriteFunction.IsEmpty())
 				{
@@ -3219,8 +3232,7 @@ void FHlslNiagaraTranslator::DefineMainGPUFunctions(
 
 			PermutationContext.AddBranch(*this, TranslationStages[i]);
 
-
-			if (TranslationStages[i].IterationSource != NAME_None)
+			if (IterationDataInterface != NAME_None)
 			{
 				if (!TranslationStages[i].CustomWriteFunction.IsEmpty())
 				{
@@ -3363,7 +3375,6 @@ void FHlslNiagaraTranslator::DefineMainGPUFunctions(
 		const FHlslNiagaraTranslationStage& TranslationStage = TranslationStages[i];
 		const bool bInterpolatedSpawning = CompileOptions.AdditionalDefines.Contains(TEXT("InterpolatedSpawn")) || i != 1;
 		const bool bAlwaysRunUpdateScript = CompileOptions.AdditionalDefines.Contains(TEXT("GpuAlwaysRunParticleUpdateScript"));
-		const bool bParticleIterationStage = TranslationStage.IterationSource.IsNone();
 		const bool bParticleSpawnStage = i == 1;
 
 		HlslOutput.Appendf(
@@ -3374,7 +3385,7 @@ void FHlslNiagaraTranslator::DefineMainGPUFunctions(
 		);
 
 		// Particle iteration stage
-		if (bParticleIterationStage)
+		if (TranslationStage.IterationSource == ENiagaraIterationSource::Particles)
 		{
 			// Do we have particle state iteration enable
 			HlslOutput.Append(TEXT("	bool bRunSpawnUpdateLogic = true;\n"));
@@ -3426,8 +3437,8 @@ void FHlslNiagaraTranslator::DefineMainGPUFunctions(
 				);
 			}
 		}
-		// Iteration interface stage
-		else
+		// Iteration Data Interface
+		else if (TranslationStage.IterationSource == ENiagaraIterationSource::DataInterface)
 		{
 			//	if (const FNiagaraVariable* IterationSourceVar = CompileData->EncounteredVariables.FindByPredicate([&](const FNiagaraVariable& VarInfo) { return VarInfo.GetName() == SimulationStageMetaData.IterationSource; }))
 			//	{
@@ -3448,6 +3459,22 @@ void FHlslNiagaraTranslator::DefineMainGPUFunctions(
 				"	const bool bRunSpawnLogic = (GLinearThreadId < GSpawnStartInstance) && (SimStart == 1);\n"
 			);
 		}
+		// Iteration Direct Value
+		else if (TranslationStage.IterationSource == ENiagaraIterationSource::DirectSet)
+		{
+			HlslOutput += TEXT(
+				"	const uint MaxInstances = SimulationStage_GetInstanceCount();\n"
+				"	const bool bValidInstance = all(GDispatchThreadId < DispatchThreadIdBounds);\n"
+				"	const bool bRunUpdateLogic = bValidInstance;\n"
+				"	const bool bRunSpawnLogic = false;\n"
+				"	GLinearThreadId = bValidInstance ? GLinearThreadId : MaxInstances;\n"
+				"	GSpawnStartInstance = MaxInstances;\n"
+			);
+		}
+		else
+		{
+			checkf(false, TEXT("Unsupported iteration source"));
+		}
 
 		HlslOutput += TEXT(
 			"	\n"
@@ -3474,7 +3501,7 @@ void FHlslNiagaraTranslator::DefineMainGPUFunctions(
 		// Add Spawn Logic
 		HlslOutput.Append(TEXT("	else if (bRunSpawnLogic)\n"));
 		HlslOutput.Append(TEXT("	{\n"));
-		if (bParticleIterationStage)
+		if (TranslationStage.IterationSource == ENiagaraIterationSource::Particles)
 		{
 			HlslOutput.Append(TEXT("		SetupExecIndexAndSpawnInfoForGPU();\n"));
 		}
@@ -4907,7 +4934,7 @@ bool FHlslNiagaraTranslationStage::IsExternalConstantNamespace(const FNiagaraVar
 {
 	if (FNiagaraParameterMapHistory::IsExternalConstantNamespace(InVar, InTargetUsage, InTargetBitmask))
 	{
-		if (IterationSource != NAME_None && InVar.IsInNameSpace(IterationSource))
+		if (IterationSource == ENiagaraIterationSource::DataInterface && InVar.IsInNameSpace(IterationDataInterface))
 			return false;
 		else
 			return true;
@@ -4924,13 +4951,12 @@ bool FHlslNiagaraTranslationStage::IsRelevantToSpawnForStage(const FNiagaraParam
 
 	if ((ScriptUsage == ENiagaraScriptUsage::ParticleSimulationStageScript) && (ExecuteBehavior == ENiagaraSimStageExecuteBehavior::OnSimulationReset))
 	{
-		if (IterationSource == NAME_None)
-		{
-			return InHistory.IsPrimaryDataSetOutput(InAliasedVar, ENiagaraScriptUsage::EmitterSpawnScript);
-		}
-		else
+		switch ( IterationSource )
 		{
-			return InVar.IsInNameSpace(IterationSource) && !InVar.IsDataInterface();
+			case ENiagaraIterationSource::Particles:		return InVar.IsInNameSpace(IterationDataInterface) && !InVar.IsDataInterface();
+			case ENiagaraIterationSource::DataInterface:	return InHistory.IsPrimaryDataSetOutput(InAliasedVar, ENiagaraScriptUsage::EmitterSpawnScript);
+			case ENiagaraIterationSource::DirectSet:		return false;
+			default:										check(false);
 		}
 	}
 	return false;
@@ -6328,9 +6354,9 @@ void FHlslNiagaraTranslator::HandleParameterRead(int32 ParamMapHistoryIdx, const
 		bIsPerInstanceAttribute = true;
 	}
 
-	if (TranslationStages[ActiveStageIdx].IterationSource != NAME_None && TranslationStages[ActiveStageIdx].ScriptUsage == ENiagaraScriptUsage::ParticleSimulationStageScript && !bIsPerInstanceAttribute)
+	if (TranslationStages[ActiveStageIdx].IterationSource == ENiagaraIterationSource::DataInterface && TranslationStages[ActiveStageIdx].ScriptUsage == ENiagaraScriptUsage::ParticleSimulationStageScript && !bIsPerInstanceAttribute)
 	{
-		bIsPerInstanceAttribute = Var.IsInNameSpace(TranslationStages[ActiveStageIdx].IterationSource);
+		bIsPerInstanceAttribute = Var.IsInNameSpace(TranslationStages[ActiveStageIdx].IterationDataInterface);
 	}
 
 	// Make sure to leave IsAlive alone if copying over previous stage params.
diff --git a/Engine/Plugins/FX/Niagara/Source/NiagaraEditor/Private/NiagaraHlslTranslator.h b/Engine/Plugins/FX/Niagara/Source/NiagaraEditor/Private/NiagaraHlslTranslator.h
--- a/Engine/Plugins/FX/Niagara/Source/NiagaraEditor/Private/NiagaraHlslTranslator.h
+++ b/Engine/Plugins/FX/Niagara/Source/NiagaraEditor/Private/NiagaraHlslTranslator.h
@@ -369,7 +369,9 @@ public:
 	bool bInterpolatePreviousParams = false;
 	bool bCopyPreviousParams = true;
 	ENiagaraCodeChunkMode ChunkModeIndex = (ENiagaraCodeChunkMode)-1;
-	FName IterationSource;
+	ENiagaraIterationSource IterationSource = ENiagaraIterationSource::Particles;
+	FName IterationDataInterface;
+	FName IterationDirectBinding;
 	int32 SimulationStageIndex = -1;
 	FName EnabledBinding;
 	FName ElementCountXBinding;
@@ -386,8 +388,7 @@ public:
 	bool bWritesParticles = false;
 	bool bPartialParticleUpdate = false;
 	bool bGpuDispatchForceLinear = false;
-	bool bOverrideGpuDispatchType = false;
-	ENiagaraGpuDispatchType OverrideGpuDispatchType = ENiagaraGpuDispatchType::OneD;
+	ENiagaraGpuDispatchType DirectDispatchType = ENiagaraGpuDispatchType::OneD;
 	bool bOverrideGpuDispatchNumThreads = false;
 	bool bShouldUpdateInitialAttributeValues = false;
 	FIntVector OverrideGpuDispatchNumThreads = FIntVector(1, 1, 1);
@@ -400,6 +401,9 @@ public:
 	bool IsRelevantToSpawnForStage(const FNiagaraParameterMapHistory& InHistory, const FNiagaraVariable& InAliasedVar, const FNiagaraVariable& InVar) const;
 
 	bool IsExternalConstantNamespace(const FNiagaraVariable& InVar, ENiagaraScriptUsage InTargetUsage, uint32 InTargetBitmask);
+
+	FName GetIterationDataInterface() const { return IterationSource == ENiagaraIterationSource::DataInterface ? IterationDataInterface : NAME_None; }
+
 	int32 CurrentCallID = 0;
 	bool bCallIDInitialized = false;
 };
diff --git a/Engine/Plugins/FX/Niagara/Source/NiagaraShader/Private/NiagaraScriptBase.cpp b/Engine/Plugins/FX/Niagara/Source/NiagaraShader/Private/NiagaraScriptBase.cpp
--- a/Engine/Plugins/FX/Niagara/Source/NiagaraShader/Private/NiagaraScriptBase.cpp
+++ b/Engine/Plugins/FX/Niagara/Source/NiagaraShader/Private/NiagaraScriptBase.cpp
@@ -13,7 +13,6 @@ FSimulationStageMetaData::FSimulationStageMetaData()
 	: bWritesParticles(0)
 	, bPartialParticleUpdate(0)
 	, bParticleIterationStateEnabled(0)
-	, bOverrideElementCount(0)
 	, GpuDispatchType(ENiagaraGpuDispatchType::OneD)
 	, GpuDispatchNumThreads(0, 0, 0)
 {
diff --git a/Engine/Plugins/FX/Niagara/Source/NiagaraShader/Public/NiagaraScriptBase.h b/Engine/Plugins/FX/Niagara/Source/NiagaraShader/Public/NiagaraScriptBase.h
--- a/Engine/Plugins/FX/Niagara/Source/NiagaraShader/Public/NiagaraScriptBase.h
+++ b/Engine/Plugins/FX/Niagara/Source/NiagaraShader/Public/NiagaraScriptBase.h
@@ -6,6 +6,7 @@
 #if UE_ENABLE_INCLUDE_ORDER_DEPRECATED_IN_5_2
 #include "RHIDefinitions.h"
 #endif
+#include "NiagaraCore.h"
 
 #include "NiagaraScriptBase.generated.h"
 
@@ -61,9 +62,17 @@ public:
 	UPROPERTY()
 	FName ElementCountZBinding;
 
-	/** The Data Interface that we iterate over for this stage. If None, then use particles.*/
+	/** The source we are iteration over. */
 	UPROPERTY()
-	FName IterationSource;
+	ENiagaraIterationSource IterationSource = ENiagaraIterationSource::Particles;
+
+	/** When IterationSource is ENiagaraIterationSource::DataInterface this is the data interface name. */
+	UPROPERTY()
+	FName IterationDataInterface;
+
+	/** When IterationSource is ENiagaraIterationSource::IterationDirectBinding this is the variable we are bound to. */
+	UPROPERTY()
+	FName IterationDirectBinding;
 
 	/** Controls when the simulation stage will execute. */
 	UPROPERTY()
@@ -80,10 +89,6 @@ public:
 	UPROPERTY()
 	uint32 bParticleIterationStateEnabled : 1;
 
-	/** When enabled the simulation stage uses element count provided by user to dispatch work. */
-	UPROPERTY()
-	uint32 bOverrideElementCount : 1;
-
 	/** When the value is not none this is the binding used for particle state iteration stages. */
 	UPROPERTY()
 	FName ParticleIterationStateBinding;
@@ -103,7 +108,6 @@ public:
 	UPROPERTY()
 	TArray<FName> InputDataInterfaces;
 
-
 	/** The number of iterations for the stage. */
 	UPROPERTY()
 	int32 NumIterations = 1;
