- Rename IterationSource -> IterationSourceType to avoid serialization issues

#rnx

diff --git a/Engine/Plugins/FX/Niagara/Source/Niagara/Private/NiagaraComputeExecutionContext.cpp b/Engine/Plugins/FX/Niagara/Source/Niagara/Private/NiagaraComputeExecutionContext.cpp
--- a/Engine/Plugins/FX/Niagara/Source/Niagara/Private/NiagaraComputeExecutionContext.cpp
+++ b/Engine/Plugins/FX/Niagara/Source/Niagara/Private/NiagaraComputeExecutionContext.cpp
@@ -103,7 +103,7 @@ bool FNiagaraComputeExecutionContext::IsIterationStage(FNiagaraDataInterfaceProx
 {
 	if (DIProxy && !DIProxy->SourceDIName.IsNone())
 	{
-		return SimStageInfo[SimulationStageIndex].IterationSource == ENiagaraIterationSource::DataInterface && (SimStageInfo[SimulationStageIndex].IterationDataInterface == DIProxy->SourceDIName);
+		return SimStageInfo[SimulationStageIndex].IterationSourceType == ENiagaraIterationSource::DataInterface && (SimStageInfo[SimulationStageIndex].IterationDataInterface == DIProxy->SourceDIName);
 	}
 	return false;
 }
@@ -111,7 +111,7 @@ bool FNiagaraComputeExecutionContext::IsIterationStage(FNiagaraDataInterfaceProx
 FNiagaraDataInterfaceProxyRW* FNiagaraComputeExecutionContext::FindIterationInterface(const TArray<FNiagaraDataInterfaceProxyRW*>& InProxies, uint32 SimulationStageIndex) const
 {
 	// Particle stage
-	if ( SimStageInfo[SimulationStageIndex].IterationSource != ENiagaraIterationSource::DataInterface )
+	if ( SimStageInfo[SimulationStageIndex].IterationSourceType != ENiagaraIterationSource::DataInterface )
 	{
 		return nullptr;
 	}
diff --git a/Engine/Plugins/FX/Niagara/Source/Niagara/Private/NiagaraEmitter.cpp b/Engine/Plugins/FX/Niagara/Source/Niagara/Private/NiagaraEmitter.cpp
--- a/Engine/Plugins/FX/Niagara/Source/Niagara/Private/NiagaraEmitter.cpp
+++ b/Engine/Plugins/FX/Niagara/Source/Niagara/Private/NiagaraEmitter.cpp
@@ -1714,7 +1714,7 @@ bool FVersionedNiagaraEmitterData::BuildParameterStoreRendererBindings(FNiagaraP
 			{
 				bAnyBindingsAdded |= ParameterStore.AddParameter(FNiagaraVariable(FNiagaraTypeDefinition::GetBoolDef(), SimStageMetaData.EnabledBinding), false);
 			}
-			if (SimStageMetaData.IterationSource == ENiagaraIterationSource::DirectSet)
+			if (SimStageMetaData.IterationSourceType == ENiagaraIterationSource::DirectSet)
 			{
 				if (!SimStageMetaData.ElementCountXBinding.IsNone())
 				{
diff --git a/Engine/Plugins/FX/Niagara/Source/Niagara/Private/NiagaraGPUSystemTick.cpp b/Engine/Plugins/FX/Niagara/Source/Niagara/Private/NiagaraGPUSystemTick.cpp
--- a/Engine/Plugins/FX/Niagara/Source/Niagara/Private/NiagaraGPUSystemTick.cpp
+++ b/Engine/Plugins/FX/Niagara/Source/Niagara/Private/NiagaraGPUSystemTick.cpp
@@ -222,7 +222,7 @@ void FNiagaraGPUSystemTick::Init(FNiagaraSystemInstance* InSystemInstance)
 							NumIterations = StageEnabled.GetValue() ? NumIterations : 0;
 						}
 					}
-					if (SimStageMetaData.IterationSource == ENiagaraIterationSource::DirectSet)
+					if (SimStageMetaData.IterationSourceType == ENiagaraIterationSource::DirectSet)
 					{
 						if (!SimStageMetaData.ElementCountXBinding.IsNone() && (SimStageMetaData.GpuDispatchType >= ENiagaraGpuDispatchType::OneD))
 						{
diff --git a/Engine/Plugins/FX/Niagara/Source/Niagara/Private/NiagaraGpuComputeDispatch.cpp b/Engine/Plugins/FX/Niagara/Source/Niagara/Private/NiagaraGpuComputeDispatch.cpp
--- a/Engine/Plugins/FX/Niagara/Source/Niagara/Private/NiagaraGpuComputeDispatch.cpp
+++ b/Engine/Plugins/FX/Niagara/Source/Niagara/Private/NiagaraGpuComputeDispatch.cpp
@@ -599,7 +599,7 @@ void FNiagaraGpuComputeDispatch::DumpDebugFrame()
 				Builder.Appendf(TEXT("Source(%p 0x%08x %d) "), SimStageData.Source, SimStageData.SourceCountOffset, SimStageData.SourceNumInstances);
 				Builder.Appendf(TEXT("Destination(%p 0x%08x %d) "), SimStageData.Destination, SimStageData.DestinationCountOffset, SimStageData.DestinationNumInstances);
 				Builder.Appendf(TEXT("Iteration(%d | %s) "), SimStageData.IterationIndex, SimStageData.AlternateIterationSource ? *SimStageData.AlternateIterationSource->SourceDIName.ToString() : TEXT("Particles"));
-				if (SimStageData.StageMetaData->IterationSource == ENiagaraIterationSource::DirectSet)
+				if (SimStageData.StageMetaData->IterationSourceType == ENiagaraIterationSource::DirectSet)
 				{
 					Builder.Appendf(TEXT("ElementCountXYZ(%d, %d, %d) "), SimStageData.ElementCountXYZ.X, SimStageData.ElementCountXYZ.Y, SimStageData.ElementCountXYZ.Z);
 				}
@@ -1430,7 +1430,7 @@ void FNiagaraGpuComputeDispatch::DispatchStage(FRDGBuilder& GraphBuilder, const
 	FIntVector DispatchCount = FIntVector(0, 0, 0);
 	FIntVector DispatchNumThreads = FIntVector(0, 0, 0);
 
-	switch (SimStageData.StageMetaData->IterationSource)
+	switch (SimStageData.StageMetaData->IterationSourceType)
 	{
 		case ENiagaraIterationSource::Particles:
 		{
@@ -1485,7 +1485,7 @@ void FNiagaraGpuComputeDispatch::DispatchStage(FRDGBuilder& GraphBuilder, const
 		}
 
 		default:
-			UE_LOG(LogNiagara, Fatal, TEXT("FNiagaraGpuComputeDispatch: Unknown IterationSouce(%d)"), SimStageData.StageMetaData->IterationSource);
+			UE_LOG(LogNiagara, Fatal, TEXT("FNiagaraGpuComputeDispatch: Unknown IterationSouce(%d)"), SimStageData.StageMetaData->IterationSourceType);
 			return;
 	}
 
@@ -1532,7 +1532,7 @@ void FNiagaraGpuComputeDispatch::DispatchStage(FRDGBuilder& GraphBuilder, const
 	{
 		DispatchParameters->SimulationStageIterationInfo = FIntVector4(INDEX_NONE, -1, 0, 0);
 		DispatchParameters->SimulationStageNormalizedIterationIndex = 0.0f;
-		switch (SimStageData.StageMetaData->IterationSource)
+		switch (SimStageData.StageMetaData->IterationSourceType)
 		{
 			case ENiagaraIterationSource::Particles:
 				break;
diff --git a/Engine/Plugins/FX/Niagara/Source/Niagara/Private/NiagaraScript.cpp b/Engine/Plugins/FX/Niagara/Source/Niagara/Private/NiagaraScript.cpp
--- a/Engine/Plugins/FX/Niagara/Source/Niagara/Private/NiagaraScript.cpp
+++ b/Engine/Plugins/FX/Niagara/Source/Niagara/Private/NiagaraScript.cpp
@@ -3515,7 +3515,7 @@ void UNiagaraScript::SyncAliases(const FNiagaraAliasContext& ResolveAliasesConte
 			SimStageMetaData.EnabledBinding = FNiagaraUtilities::ResolveAliases(Var, ResolveAliasesContext).GetName();
 		}
 
-		if (SimStageMetaData.IterationSource == ENiagaraIterationSource::DirectSet)
+		if (SimStageMetaData.IterationSourceType == ENiagaraIterationSource::DirectSet)
 		{
 			if (!SimStageMetaData.ElementCountXBinding.IsNone())
 			{
diff --git a/Engine/Plugins/FX/Niagara/Source/Niagara/Private/NiagaraSimulationStageBase.cpp b/Engine/Plugins/FX/Niagara/Source/Niagara/Private/NiagaraSimulationStageBase.cpp
--- a/Engine/Plugins/FX/Niagara/Source/Niagara/Private/NiagaraSimulationStageBase.cpp
+++ b/Engine/Plugins/FX/Niagara/Source/Niagara/Private/NiagaraSimulationStageBase.cpp
@@ -115,7 +115,7 @@ bool UNiagaraSimulationStageGeneric::FillCompilationData(TArray<FNiagaraSimulati
 	SimStageData.ElementCountZBinding = ElementCountZBinding.GetName();
 	SimStageData.NumIterations = Iterations;
 	SimStageData.NumIterationsBinding = NumIterationsBinding.GetName();
-	SimStageData.IterationSource = IterationSource;
+	SimStageData.IterationSourceType = IterationSource;
 	SimStageData.IterationDataInterface = IterationSource == ENiagaraIterationSource::DataInterface ? DataInterface.BoundVariable.GetName() : FName();
 	SimStageData.ExecuteBehavior = ExecuteBehavior;
 	SimStageData.PartialParticleUpdate = bDisablePartialParticleUpdate == false;
diff --git a/Engine/Plugins/FX/Niagara/Source/Niagara/Public/NiagaraSimulationStageCompileData.h b/Engine/Plugins/FX/Niagara/Source/Niagara/Public/NiagaraSimulationStageCompileData.h
--- a/Engine/Plugins/FX/Niagara/Source/Niagara/Public/NiagaraSimulationStageCompileData.h
+++ b/Engine/Plugins/FX/Niagara/Source/Niagara/Public/NiagaraSimulationStageCompileData.h
@@ -14,7 +14,7 @@ struct FNiagaraSimulationStageCompilationData
 	FName                           ElementCountZBinding;
 	uint32                          NumIterations = 1;
 	FName                           NumIterationsBinding;
-	ENiagaraIterationSource			IterationSource = ENiagaraIterationSource::Particles;
+	ENiagaraIterationSource			IterationSourceType = ENiagaraIterationSource::Particles;
 	FName							IterationDataInterface;
 	FName							IterationDirectBinding;
 	ENiagaraSimStageExecuteBehavior ExecuteBehavior = ENiagaraSimStageExecuteBehavior::Always;
diff --git a/Engine/Plugins/FX/Niagara/Source/NiagaraEditor/Private/NiagaraHlslTranslator.cpp b/Engine/Plugins/FX/Niagara/Source/NiagaraEditor/Private/NiagaraHlslTranslator.cpp
--- a/Engine/Plugins/FX/Niagara/Source/NiagaraEditor/Private/NiagaraHlslTranslator.cpp
+++ b/Engine/Plugins/FX/Niagara/Source/NiagaraEditor/Private/NiagaraHlslTranslator.cpp
@@ -1299,7 +1299,7 @@ const FNiagaraTranslateResults &FHlslNiagaraTranslator::Translate(const FNiagara
 					TranslationStages[Index].NumIterations = CompileSimStageData.NumIterations;
 					TranslationStages[Index].ExecuteBehavior = CompileSimStageData.ExecuteBehavior;
 					TranslationStages[Index].bPartialParticleUpdate = CompileSimStageData.PartialParticleUpdate;
-					TranslationStages[Index].IterationSource = CompileSimStageData.IterationSource;
+					TranslationStages[Index].IterationSourceType = CompileSimStageData.IterationSourceType;
 					TranslationStages[Index].IterationDataInterface = CompileSimStageData.IterationDataInterface;
 					TranslationStages[Index].IterationDirectBinding = CompileSimStageData.IterationDirectBinding;
 					TranslationStages[Index].NumIterationsBinding = CompileSimStageData.NumIterationsBinding;
@@ -1372,7 +1372,7 @@ const FNiagaraTranslateResults &FHlslNiagaraTranslator::Translate(const FNiagara
 					SimulationStageMetaData.ElementCountYBinding = InCompileData->CompileSimStageData[SourceSimStageIndex].ElementCountYBinding;
 					SimulationStageMetaData.ElementCountZBinding = InCompileData->CompileSimStageData[SourceSimStageIndex].ElementCountZBinding;
 					SimulationStageMetaData.ExecuteBehavior = TranslationStages[Index].ExecuteBehavior;
-					SimulationStageMetaData.IterationSource = InCompileData->CompileSimStageData[SourceSimStageIndex].IterationSource;
+					SimulationStageMetaData.IterationSourceType = InCompileData->CompileSimStageData[SourceSimStageIndex].IterationSourceType;
 					SimulationStageMetaData.IterationDataInterface = InCompileData->CompileSimStageData[SourceSimStageIndex].IterationDataInterface;
 					SimulationStageMetaData.IterationDirectBinding = InCompileData->CompileSimStageData[SourceSimStageIndex].IterationDirectBinding;
 					SimulationStageMetaData.NumIterationsBinding = InCompileData->CompileSimStageData[SourceSimStageIndex].NumIterationsBinding;
@@ -1386,7 +1386,7 @@ const FNiagaraTranslateResults &FHlslNiagaraTranslator::Translate(const FNiagara
 					// Determine dispatch information from iteration source (if we have one)
 					SimulationStageMetaData.GpuDispatchType = ENiagaraGpuDispatchType::OneD;
 					SimulationStageMetaData.GpuDispatchNumThreads = FNiagaraShader::GetDefaultThreadGroupSize(ENiagaraGpuDispatchType::OneD);
-					if ( SimulationStageMetaData.IterationSource == ENiagaraIterationSource::DataInterface )
+					if ( SimulationStageMetaData.IterationSourceType == ENiagaraIterationSource::DataInterface )
 					{
 						if (SimulationStageMetaData.IterationDataInterface.IsNone())
 						{
@@ -1411,7 +1411,7 @@ const FNiagaraTranslateResults &FHlslNiagaraTranslator::Translate(const FNiagara
 							}
 						}
 					}
-					else if (SimulationStageMetaData.IterationSource == ENiagaraIterationSource::DirectSet)
+					else if (SimulationStageMetaData.IterationSourceType == ENiagaraIterationSource::DirectSet)
 					{
 						SimulationStageMetaData.GpuDispatchType = TranslationStages[Index].DirectDispatchType;
 					}
@@ -1442,7 +1442,7 @@ const FNiagaraTranslateResults &FHlslNiagaraTranslator::Translate(const FNiagara
 				const FGuid& StageGuid = InCompileData->CompileSimStageData[StageIdx].StageGuid;
 				if (StageGuid == CompileOptions.TargetUsageId && InCompileData->CompileSimStageData.IsValidIndex(StageIdx))
 				{
-					TranslationStages[0].IterationSource = InCompileData->CompileSimStageData[StageIdx].IterationSource;
+					TranslationStages[0].IterationSourceType = InCompileData->CompileSimStageData[StageIdx].IterationSourceType;
 					TranslationStages[0].IterationDataInterface = InCompileData->CompileSimStageData[StageIdx].IterationDataInterface;
 					TranslationStages[0].IterationDirectBinding = InCompileData->CompileSimStageData[StageIdx].IterationDirectBinding;
 				}
@@ -3385,7 +3385,7 @@ void FHlslNiagaraTranslator::DefineMainGPUFunctions(
 		);
 
 		// Particle iteration stage
-		if (TranslationStage.IterationSource == ENiagaraIterationSource::Particles)
+		if (TranslationStage.IterationSourceType == ENiagaraIterationSource::Particles)
 		{
 			// Do we have particle state iteration enable
 			HlslOutput.Append(TEXT("	bool bRunSpawnUpdateLogic = true;\n"));
@@ -3438,7 +3438,7 @@ void FHlslNiagaraTranslator::DefineMainGPUFunctions(
 			}
 		}
 		// Iteration Data Interface
-		else if (TranslationStage.IterationSource == ENiagaraIterationSource::DataInterface)
+		else if (TranslationStage.IterationSourceType == ENiagaraIterationSource::DataInterface)
 		{
 			//	if (const FNiagaraVariable* IterationSourceVar = CompileData->EncounteredVariables.FindByPredicate([&](const FNiagaraVariable& VarInfo) { return VarInfo.GetName() == SimulationStageMetaData.IterationSource; }))
 			//	{
@@ -3460,7 +3460,7 @@ void FHlslNiagaraTranslator::DefineMainGPUFunctions(
 			);
 		}
 		// Iteration Direct Value
-		else if (TranslationStage.IterationSource == ENiagaraIterationSource::DirectSet)
+		else if (TranslationStage.IterationSourceType == ENiagaraIterationSource::DirectSet)
 		{
 			HlslOutput += TEXT(
 				"	const uint MaxInstances = SimulationStage_GetInstanceCount();\n"
@@ -3501,7 +3501,7 @@ void FHlslNiagaraTranslator::DefineMainGPUFunctions(
 		// Add Spawn Logic
 		HlslOutput.Append(TEXT("	else if (bRunSpawnLogic)\n"));
 		HlslOutput.Append(TEXT("	{\n"));
-		if (TranslationStage.IterationSource == ENiagaraIterationSource::Particles)
+		if (TranslationStage.IterationSourceType == ENiagaraIterationSource::Particles)
 		{
 			HlslOutput.Append(TEXT("		SetupExecIndexAndSpawnInfoForGPU();\n"));
 		}
@@ -4934,7 +4934,7 @@ bool FHlslNiagaraTranslationStage::IsExternalConstantNamespace(const FNiagaraVar
 {
 	if (FNiagaraParameterMapHistory::IsExternalConstantNamespace(InVar, InTargetUsage, InTargetBitmask))
 	{
-		if (IterationSource == ENiagaraIterationSource::DataInterface && InVar.IsInNameSpace(IterationDataInterface))
+		if (IterationSourceType == ENiagaraIterationSource::DataInterface && InVar.IsInNameSpace(IterationDataInterface))
 			return false;
 		else
 			return true;
@@ -4951,10 +4951,10 @@ bool FHlslNiagaraTranslationStage::IsRelevantToSpawnForStage(const FNiagaraParam
 
 	if ((ScriptUsage == ENiagaraScriptUsage::ParticleSimulationStageScript) && (ExecuteBehavior == ENiagaraSimStageExecuteBehavior::OnSimulationReset))
 	{
-		switch ( IterationSource )
+		switch (IterationSourceType)
 		{
-			case ENiagaraIterationSource::Particles:		return InVar.IsInNameSpace(IterationDataInterface) && !InVar.IsDataInterface();
-			case ENiagaraIterationSource::DataInterface:	return InHistory.IsPrimaryDataSetOutput(InAliasedVar, ENiagaraScriptUsage::EmitterSpawnScript);
+			case ENiagaraIterationSource::Particles:		return InHistory.IsPrimaryDataSetOutput(InAliasedVar, ENiagaraScriptUsage::EmitterSpawnScript); 
+			case ENiagaraIterationSource::DataInterface:	return InVar.IsInNameSpace(IterationDataInterface) && !InVar.IsDataInterface();
 			case ENiagaraIterationSource::DirectSet:		return false;
 			default:										check(false);
 		}
@@ -6354,7 +6354,7 @@ void FHlslNiagaraTranslator::HandleParameterRead(int32 ParamMapHistoryIdx, const
 		bIsPerInstanceAttribute = true;
 	}
 
-	if (TranslationStages[ActiveStageIdx].IterationSource == ENiagaraIterationSource::DataInterface && TranslationStages[ActiveStageIdx].ScriptUsage == ENiagaraScriptUsage::ParticleSimulationStageScript && !bIsPerInstanceAttribute)
+	if (TranslationStages[ActiveStageIdx].IterationSourceType == ENiagaraIterationSource::DataInterface && TranslationStages[ActiveStageIdx].ScriptUsage == ENiagaraScriptUsage::ParticleSimulationStageScript && !bIsPerInstanceAttribute)
 	{
 		bIsPerInstanceAttribute = Var.IsInNameSpace(TranslationStages[ActiveStageIdx].IterationDataInterface);
 	}
diff --git a/Engine/Plugins/FX/Niagara/Source/NiagaraEditor/Private/NiagaraHlslTranslator.h b/Engine/Plugins/FX/Niagara/Source/NiagaraEditor/Private/NiagaraHlslTranslator.h
--- a/Engine/Plugins/FX/Niagara/Source/NiagaraEditor/Private/NiagaraHlslTranslator.h
+++ b/Engine/Plugins/FX/Niagara/Source/NiagaraEditor/Private/NiagaraHlslTranslator.h
@@ -369,7 +369,7 @@ public:
 	bool bInterpolatePreviousParams = false;
 	bool bCopyPreviousParams = true;
 	ENiagaraCodeChunkMode ChunkModeIndex = (ENiagaraCodeChunkMode)-1;
-	ENiagaraIterationSource IterationSource = ENiagaraIterationSource::Particles;
+	ENiagaraIterationSource IterationSourceType = ENiagaraIterationSource::Particles;
 	FName IterationDataInterface;
 	FName IterationDirectBinding;
 	int32 SimulationStageIndex = -1;
@@ -402,7 +402,7 @@ public:
 
 	bool IsExternalConstantNamespace(const FNiagaraVariable& InVar, ENiagaraScriptUsage InTargetUsage, uint32 InTargetBitmask);
 
-	FName GetIterationDataInterface() const { return IterationSource == ENiagaraIterationSource::DataInterface ? IterationDataInterface : NAME_None; }
+	FName GetIterationDataInterface() const { return IterationSourceType == ENiagaraIterationSource::DataInterface ? IterationDataInterface : NAME_None; }
 
 	int32 CurrentCallID = 0;
 	bool bCallIDInitialized = false;
diff --git a/Engine/Plugins/FX/Niagara/Source/NiagaraShader/Public/NiagaraScriptBase.h b/Engine/Plugins/FX/Niagara/Source/NiagaraShader/Public/NiagaraScriptBase.h
--- a/Engine/Plugins/FX/Niagara/Source/NiagaraShader/Public/NiagaraScriptBase.h
+++ b/Engine/Plugins/FX/Niagara/Source/NiagaraShader/Public/NiagaraScriptBase.h
@@ -64,7 +64,7 @@ public:
 
 	/** The source we are iteration over. */
 	UPROPERTY()
-	ENiagaraIterationSource IterationSource = ENiagaraIterationSource::Particles;
+	ENiagaraIterationSource IterationSourceType = ENiagaraIterationSource::Particles;
 
 	/** When IterationSource is ENiagaraIterationSource::DataInterface this is the data interface name. */
 	UPROPERTY()
diff --git a/Engine/Source/Runtime/Core/Private/UObject/UE5MainStreamObjectVersion.cpp b/Engine/Source/Runtime/Core/Private/UObject/UE5MainStreamObjectVersion.cpp
--- a/Engine/Source/Runtime/Core/Private/UObject/UE5MainStreamObjectVersion.cpp
+++ b/Engine/Source/Runtime/Core/Private/UObject/UE5MainStreamObjectVersion.cpp
@@ -14,7 +14,7 @@ TMap<FGuid, FGuid> FUE5MainStreamObjectVersion::GetSystemGuids()
 	SystemGuids.Add(DevGuids.MATERIALSHADERMAP_DERIVEDDATA_VER, FGuid("65C01D817C9A4EEFAE9E988D41A1F3DD"));
 	SystemGuids.Add(DevGuids.NANITE_DERIVEDDATA_VER, FGuid("7A3EA93546624597BEBE2ADFDE3A7628"));
 	SystemGuids.Add(DevGuids.NIAGARASHADERMAP_DERIVEDDATA_VER, FGuid("8A37C45D24F2423CBE5F8F371DE33575"));
-	SystemGuids.Add(DevGuids.Niagara_LatestScriptCompileVersion, FGuid("5F19E9ECE6B844A49B23D2BD1DDB9095"));
+	SystemGuids.Add(DevGuids.Niagara_LatestScriptCompileVersion, FGuid("D2A4314F3211254A9531B32E5E33AD66"));
 	SystemGuids.Add(DevGuids.SkeletalMeshDerivedDataVersion, FGuid("4084BE94509A8C3FA7CF410DDF8ABF04"));
 	SystemGuids.Add(DevGuids.STATICMESH_DERIVEDDATA_VER, FGuid("98F7D79A5811013825E75BBCC41ED3E9"));
 	SystemGuids.Add(DevGuids.POSESEARCHDB_DERIVEDDATA_VER, FGuid("4E595C2AC5E947D6BA9ABC874353E5BC"));
