Niagara Data Channels - Supporting CL

This CL contains supporting code for an upcoming feature, Niagara Data Channels.
Further supporting CLs will follow and finally the main CL which contains the feature.

- Niagara DI functions can now have optionally variable numbers of input & output parameters.
DIs can mark their function signatures to allow this or not.
DIs must handle these variable parameters in their function binding and hlsl gen etc.

- DIs can now provide context menu actions on their function call nodes.
This will allows DataChannel function calls to init themselves with the right Inputs and Outputs but could be broadly useful for a number of DI functions.

- New HLSL generation API for DIs.
This passes in a context object from the translator so that we can more easily update things in future without 1000 boiler plate changes.
To avoid said changes this time, I'm having the base implementation of the new API call into the existing API with the relevant parts of the context.
No DI or it's parent chain should implement both APIs.

 - VM DI Info and GPU function info in NiagaraShadrers now track these variadic inputs.
This has a few knock on changes. One being a custom verison bump, the other being a common refernce struct for FNiagaraVariableBase in NiagaraCore.
Required so that code in NiagaraShader could serialize a refernce to an FNiagaraVariableBase. Actually moving this class to NiagaraCore is a headache of knock on dependencies.

- NiagaraShaders now serialize their base compile hash. This allows DIs to generate per-script intermediate parameter data and then look this up for setting shader parameters correctly.

- Bumping DDC and Compiler guids.

diff --git a/Engine/Plugins/FX/Niagara/Source/Niagara/Classes/NiagaraDataInterface.h b/Engine/Plugins/FX/Niagara/Source/Niagara/Classes/NiagaraDataInterface.h
--- a/Engine/Plugins/FX/Niagara/Source/Niagara/Classes/NiagaraDataInterface.h
+++ b/Engine/Plugins/FX/Niagara/Source/Niagara/Classes/NiagaraDataInterface.h
@@ -24,6 +24,7 @@ class FNiagaraGpuComputeDispatchInterface;
 class FNiagaraGPUSystemTick;
 struct FNiagaraSimStageData;
 class FNiagaraSystemInstance;
+struct FNiagaraDataInterfaceHlslGenerationContext;
 
 struct FNDITransformHandlerNoop
 {
@@ -615,6 +616,10 @@ public:
 	virtual void GetCommonHLSL(FString& OutHLSL)
 	{
 	}
+	
+	virtual void GetParameterDefinitionHLSL(FNiagaraDataInterfaceHlslGenerationContext& HlslGenContext, FString& OutHLSL);
+
+	virtual bool GetFunctionHLSL(FNiagaraDataInterfaceHlslGenerationContext& HlslGenContext, FString& OutHLSL);
 
 	virtual void GetParameterDefinitionHLSL(const FNiagaraDataInterfaceGPUParamInfo& ParamInfo, FString& OutHLSL)
 	{
diff --git a/Engine/Plugins/FX/Niagara/Source/Niagara/Private/NiagaraCommon.cpp b/Engine/Plugins/FX/Niagara/Source/Niagara/Private/NiagaraCommon.cpp
--- a/Engine/Plugins/FX/Niagara/Source/Niagara/Private/NiagaraCommon.cpp
+++ b/Engine/Plugins/FX/Niagara/Source/Niagara/Private/NiagaraCommon.cpp
@@ -1578,4 +1578,48 @@ EPSCPoolMethod ToPSCPoolMethod(ENCPoolMethod PoolingMethod)
 		return EPSCPoolMethod::ManualRelease;
 	}
 	return EPSCPoolMethod::None;
+}
+
+//////////////////////////////////////////////////////////////////////////
+
+
+void FNiagaraFunctionSignature::GetVariadicInputs(TArray<FNiagaraVariableBase>& OutVariadicInputs, bool bStripNonExecution/* =true */)const
+{
+	static const FNiagaraVariableBase InstDataVar(FNiagaraTypeDefinition::GetIntDef(), TEXT("InstanceData"));
+	OutVariadicInputs.Reset(NumOptionalInputs());
+	int32 NumInputs = 0;
+	for (const FNiagaraVariableBase& Param : Inputs)
+	{
+		if (bStripNonExecution && (Param.GetType() == FNiagaraTypeDefinition::GetParameterMapDef() || Param == InstDataVar))
+		{
+			continue;
+		}
+
+		if (NumInputs++ < NumRequiredInputs())
+		{
+			continue;
+		}
+
+		OutVariadicInputs.Emplace(Param);
+	}
+}
+
+void FNiagaraFunctionSignature::GetVariadicOutputs(TArray<FNiagaraVariableBase>& OutVariadicOutputs, bool bStripNonExecution/* =true */)const
+{
+	OutVariadicOutputs.Reset(NumOptionalOutputs());
+	int32 NumOutputs = 0;
+	for (const FNiagaraVariableBase& Param : Outputs)
+	{
+		if (bStripNonExecution && (Param.GetType() == FNiagaraTypeDefinition::GetParameterMapDef()))
+		{
+			continue;
+		}
+
+		if (NumOutputs++ < NumRequiredOutputs())
+		{
+			continue;
+		}
+
+		OutVariadicOutputs.Emplace(Param);
+	}
 }
\ No newline at end of file
diff --git a/Engine/Plugins/FX/Niagara/Source/Niagara/Public/NiagaraCommon.h b/Engine/Plugins/FX/Niagara/Source/Niagara/Public/NiagaraCommon.h
--- a/Engine/Plugins/FX/Niagara/Source/Niagara/Public/NiagaraCommon.h
+++ b/Engine/Plugins/FX/Niagara/Source/Niagara/Public/NiagaraCommon.h
@@ -67,13 +67,14 @@ inline constexpr ETickingGroup NiagaraFirstTickGroup = TG_PrePhysics;
 inline constexpr ETickingGroup NiagaraLastTickGroup = TG_LastDemotable;
 inline constexpr int NiagaraNumTickGroups = NiagaraLastTickGroup - NiagaraFirstTickGroup + 1;
 
-enum ENiagaraBaseTypes
+UENUM()
+enum class ENiagaraBaseTypes : uint8
 {
-	NBT_Half,
-	NBT_Float,
-	NBT_Int32,
-	NBT_Bool,
-	NBT_Max,
+	Half,
+	Float,
+	Int32,
+	Bool,
+	Max UMETA(Hidden),
 };
 
 /** Niagara supported buffer formats on the GPU. */
@@ -343,10 +344,10 @@ struct NIAGARA_API FNiagaraFunctionSignature
 	FName Name;
 	/** Input parameters to this function. */
 	UPROPERTY()
-	TArray<FNiagaraVariable> Inputs;
+	TArray<FNiagaraVariable> Inputs; //TODO: Can these be FNiagaraVariableBase?
 	/** Input parameters of this function. */
 	UPROPERTY()
-	TArray<FNiagaraVariable> Outputs;
+	TArray<FNiagaraVariable> Outputs; //TODO: Can these be FNiagaraVariableBase?
 	/** Id of the owner is this is a member function. */
 	UPROPERTY()
 	FName OwnerName;
@@ -407,6 +408,14 @@ struct NIAGARA_API FNiagaraFunctionSignature
 	UPROPERTY()
 	int32 ContextStageIndex;
 
+	/** Number of required inputs. Can be less than the actual number of inputs when using bVariadicInput. INDEX_NONE when not using bVariadicInput, denoting that all inputs are required. */
+	UPROPERTY()
+	int16 RequiredInputs;
+
+	/** Number of required outputs. Can be less than the actual number of outputs when using bVariadicOutput. INDEX_NONE when not using bVariadicOutput, denoting that all outputs are required. */
+	UPROPERTY()
+	int16 RequiredOutputs;
+
 	/** Function specifiers verified at bind time. */
 	UPROPERTY()
 	TMap<FName, FName> FunctionSpecifiers;
@@ -437,6 +446,8 @@ struct NIAGARA_API FNiagaraFunctionSignature
 		, bHidden(false)
 		, ModuleUsageBitmask(0)
 		, ContextStageIndex(INDEX_NONE)
+		, RequiredInputs(INDEX_NONE)
+		, RequiredOutputs(INDEX_NONE)
 	{
 	}
 
@@ -455,6 +466,8 @@ struct NIAGARA_API FNiagaraFunctionSignature
 		, bHidden(false)
 		, ModuleUsageBitmask(0)
 		, ContextStageIndex(INDEX_NONE)
+		, RequiredInputs(INDEX_NONE)
+		, RequiredOutputs(INDEX_NONE)
 	{
 		Inputs.Reserve(InInputs.Num());
 		for (FNiagaraVariable& Var : InInputs)
@@ -497,8 +510,45 @@ struct NIAGARA_API FNiagaraFunctionSignature
 	bool EqualsIgnoringSpecifiers(const FNiagaraFunctionSignature& Other) const
 	{
 		bool bMatches = Name.ToString().Equals(Other.Name.ToString());
-		bMatches &= Inputs == Other.Inputs;
-		bMatches &= Outputs == Other.Outputs;
+		
+		//Allow divergent inputs for variadic input functions.
+		if (VariadicInput() && bMatches)
+		{
+			if (Other.Inputs.Num() >= RequiredInputs)
+			{
+				for (int32 i = 0; i < RequiredInputs; ++i)
+				{
+					bMatches &= Inputs[i] == Other.Inputs[i];
+				}
+			}
+			else
+			{
+				bMatches = false;
+			}
+		}
+		else
+		{
+			bMatches &= Inputs == Other.Inputs;
+		}
+		//Allow divergent outputs for variadic output functions.
+		if (VariadicOutput() && bMatches)
+		{
+			if (Other.Inputs.Num() >= RequiredInputs)
+			{
+				for (int32 i = 0; i < RequiredOutputs; ++i)
+				{
+					bMatches &= Outputs[i] == Other.Outputs[i];
+				}
+			}
+			else
+			{
+				bMatches = false;
+			}
+		}
+		else
+		{
+			bMatches &= Outputs == Other.Outputs;
+		}
 		bMatches &= bRequiresContext == Other.bRequiresContext;
 		bMatches &= bRequiresExecPin == Other.bRequiresExecPin;
 		bMatches &= bMemberFunction == Other.bMemberFunction;
@@ -570,6 +620,18 @@ struct NIAGARA_API FNiagaraFunctionSignature
 	}
 	
 	bool IsValid()const { return Name != NAME_None && (Inputs.Num() > 0 || Outputs.Num() > 0); }
+
+	bool VariadicInput()const { return RequiredInputs != INDEX_NONE; }
+	bool VariadicOutput()const { return RequiredOutputs != INDEX_NONE; }
+
+	int32 NumRequiredInputs()const { return RequiredInputs == INDEX_NONE ? Inputs.Num() : RequiredInputs; }
+	int32 NumOptionalInputs()const { return Inputs.Num() - NumRequiredInputs(); }
+
+	int32 NumRequiredOutputs()const { return RequiredOutputs == INDEX_NONE ? Outputs.Num() : RequiredOutputs; }
+	int32 NumOptionalOutputs()const { return Outputs.Num() - NumRequiredOutputs(); }
+
+	void GetVariadicInputs(TArray<FNiagaraVariableBase>& OutVariadicInputs, bool bStripNonExecution = true)const;
+	void GetVariadicOutputs(TArray<FNiagaraVariableBase>& OutVariadicOutputs, bool bStripNonExecution = true)const;
 };
 
 USTRUCT()
@@ -708,6 +770,12 @@ struct FVMExternalFunctionBindingInfo
 	UPROPERTY()
 	TArray<FVMFunctionSpecifier> FunctionSpecifiers;
 
+	UPROPERTY()
+	TArray<FNiagaraVariableBase> VariadicInputs;
+
+	UPROPERTY()
+	TArray<FNiagaraVariableBase> VariadicOutputs;
+
 	FORCEINLINE int32 GetNumInputs() const { return InputParamLocations.Num(); }
 	FORCEINLINE int32 GetNumOutputs() const { return NumOutputs; }
 
diff --git a/Engine/Plugins/FX/Niagara/Source/Niagara/Public/NiagaraDataInterfaceUtilities.h b/Engine/Plugins/FX/Niagara/Source/Niagara/Public/NiagaraDataInterfaceUtilities.h
--- a/Engine/Plugins/FX/Niagara/Source/Niagara/Public/NiagaraDataInterfaceUtilities.h
+++ b/Engine/Plugins/FX/Niagara/Source/Niagara/Public/NiagaraDataInterfaceUtilities.h
@@ -51,3 +51,34 @@ namespace FNiagaraDataInterfaceUtilities
 	// The action should return True to continue iteration or False to stop
 	NIAGARA_API void ForEachDataInterface(UNiagaraSystem* NiagaraSystem, TFunction<bool(const FDataInterfaceUsageContext&)> Action);
 }
+
+#if WITH_EDITOR
+
+/** Helper context object helping to facilitate data interfaces building their hlsl shader code for GPU simulations. */
+struct FNiagaraDataInterfaceHlslGenerationContext
+{
+	DECLARE_DELEGATE_RetVal_OneParam(FString, FGetStructHlslTypeName, const FNiagaraTypeDefinition& /*Type*/)
+	DECLARE_DELEGATE_RetVal_OneParam(FString, FGetPropertyHlslTypeName, const FProperty* /*Property*/)
+	DECLARE_DELEGATE_RetVal_TwoParams(FString, FGetSanitizedSymbolName, FStringView /*SymbolName*/, bool /*bCollapsNamespaces*/)
+
+	FNiagaraDataInterfaceHlslGenerationContext(const FNiagaraDataInterfaceGPUParamInfo& InParameterInfo, TArrayView<const FNiagaraFunctionSignature> InSignatures)
+		: ParameterInfo(InParameterInfo), Signatures(InSignatures)
+	{
+	}
+
+	const FNiagaraDataInterfaceGPUParamInfo& ParameterInfo;
+	TArrayView<const FNiagaraFunctionSignature> Signatures;
+	int32 FunctionInstanceIndex = INDEX_NONE;
+
+	FGetStructHlslTypeName GetStructHlslTypeNameDelegate;
+	FGetPropertyHlslTypeName GetPropertyHlslTypeNameDelegate;
+	FGetSanitizedSymbolName GetSanitizedSymbolNameDelegate;
+
+	const FNiagaraDataInterfaceGeneratedFunction& GetFunctionInfo() const { return ParameterInfo.GeneratedFunctions[FunctionInstanceIndex]; }
+
+	FString GetStructHlslTypeName(const FNiagaraTypeDefinition& Type) { return GetStructHlslTypeNameDelegate.Execute(Type); }
+	FString GetPropertyHlslTypeName(const FProperty* Property) { return GetPropertyHlslTypeNameDelegate.Execute(Property); }
+	FString GetSanitizedSymbolName(FStringView SymbolName, bool bCollapsNamespaces = false) { return GetSanitizedSymbolNameDelegate.Execute(SymbolName, bCollapsNamespaces); }
+};
+
+#endif
\ No newline at end of file
diff --git a/Engine/Plugins/FX/Niagara/Source/NiagaraCore/Private/NiagaraCoreModule.cpp b/Engine/Plugins/FX/Niagara/Source/NiagaraCore/Private/NiagaraCoreModule.cpp
--- a/Engine/Plugins/FX/Niagara/Source/NiagaraCore/Private/NiagaraCoreModule.cpp
+++ b/Engine/Plugins/FX/Niagara/Source/NiagaraCore/Private/NiagaraCoreModule.cpp
@@ -13,4 +13,16 @@ IMPLEMENT_EXPORTED_INTRINSIC_TYPE_LAYOUT(TIndexedPtr<UNiagaraDataInterfaceBase>)
 UNiagaraDataInterfaceBase::UNiagaraDataInterfaceBase(class FObjectInitializer const & Initializer)
 {
 
-}
\ No newline at end of file
+}
+
+//////////////////////////////////////////////////////////////////////////
+
+//TODO: This does not need a custom serialize. It's only here as it's called by another custom serialize that itself is pointless
+//TODO: Rip these out to get back to tagged property serialize.
+
+bool FNiagaraVariableCommonReference::Serialize(FArchive& Ar)
+{
+	Ar << Name;
+	Ar << UnderlyingType;
+	return true;
+}
diff --git a/Engine/Plugins/FX/Niagara/Source/NiagaraCore/Public/NiagaraCore.h b/Engine/Plugins/FX/Niagara/Source/NiagaraCore/Public/NiagaraCore.h
--- a/Engine/Plugins/FX/Niagara/Source/NiagaraCore/Public/NiagaraCore.h
+++ b/Engine/Plugins/FX/Niagara/Source/NiagaraCore/Public/NiagaraCore.h
@@ -18,3 +18,37 @@ enum class ENiagaraIterationSource : uint8
 	/** Iterate over a user provided number of elements. */
 	DirectSet,
 };
+
+/** A utility class allowing for references to FNiagaraVariableBase outside of the Niagara module. */
+USTRUCT()
+struct NIAGARACORE_API FNiagaraVariableCommonReference
+{
+	GENERATED_USTRUCT_BODY()
+
+	UPROPERTY()
+	FName Name;
+
+	UPROPERTY()
+	TObjectPtr<UObject> UnderlyingType;
+
+	bool Serialize(FArchive& Ar);
+	friend bool operator<<(FArchive& Ar, FNiagaraVariableCommonReference& VariableReference);
+
+	bool operator==(const FNiagaraVariableCommonReference& Other)const
+	{
+		return Name == Other.Name && UnderlyingType == Other.UnderlyingType;
+	}	
+};
+
+template<> struct TStructOpsTypeTraits<FNiagaraVariableCommonReference> : public TStructOpsTypeTraitsBase2<FNiagaraVariableCommonReference>
+{
+	enum
+	{
+		WithSerializer = true,
+	};
+};
+
+inline bool operator<<(FArchive& Ar, FNiagaraVariableCommonReference& VariableReference)
+{
+	return VariableReference.Serialize(Ar);
+}
diff --git a/Engine/Plugins/FX/Niagara/Source/NiagaraCore/Public/NiagaraCustomVersion.h b/Engine/Plugins/FX/Niagara/Source/NiagaraCore/Public/NiagaraCustomVersion.h
--- a/Engine/Plugins/FX/Niagara/Source/NiagaraCore/Public/NiagaraCustomVersion.h
+++ b/Engine/Plugins/FX/Niagara/Source/NiagaraCore/Public/NiagaraCustomVersion.h
@@ -182,6 +182,8 @@ struct FNiagaraCustomVersion
 
 		EventSpawnsUpdateInitialAttributeValues, // Event spawns now optionally update Initial. attribute values. New default is true but old data is kept false to maintain existing behavior.
 
+		AddVariadicParametersToGPUFunctionInfo, // Adds list of variadic parameters to the information about GPU functions.
+
 		// DO NOT ADD A NEW VERSION UNLESS YOU HAVE TALKED TO THE NIAGARA LEAD. Mismanagement of these versions can lead to data loss if it is adjusted in multiple streams simultaneously.
 		// -----<new versions can be added above this line>  -------------------------------------------------
 		VersionPlusOne,
diff --git a/Engine/Plugins/FX/Niagara/Source/NiagaraEditor/Private/NiagaraActions.cpp b/Engine/Plugins/FX/Niagara/Source/NiagaraEditor/Private/NiagaraActions.cpp
--- a/Engine/Plugins/FX/Niagara/Source/NiagaraEditor/Private/NiagaraActions.cpp
+++ b/Engine/Plugins/FX/Niagara/Source/NiagaraEditor/Private/NiagaraActions.cpp
@@ -6,6 +6,7 @@
 #include "NiagaraNodeParameterMapSet.h"
 #include "EdGraphSchema_Niagara.h"
 #include "NiagaraEditorUtilities.h"
+#include "NiagaraNodeFunctionCall.h"
 #include "NiagaraNodeStaticSwitch.h"
 #include "NiagaraScriptVariable.h"
 #include "Widgets/SWidget.h"
@@ -658,5 +659,23 @@ EVisibility FNiagaraParameterDragOperation::IsTextVisible() const
 	return CurrentHoverText.IsEmpty() ? EVisibility::Collapsed : EVisibility::Visible;
 }
 
-#undef LOCTEXT_NAMESPACE
+//////////////////////////////////////////////////////////////////////////
 
+TMap<FName, TUniquePtr<INiagaraDataInterfaceNodeActionProvider>> INiagaraDataInterfaceNodeActionProvider::RegisteredActionProviders;
+
+void INiagaraDataInterfaceNodeActionProvider::GetNodeContextMenuActions(UClass* DIClass, UToolMenu* Menu, UGraphNodeContextMenuContext* Context, FNiagaraFunctionSignature Signature)
+{
+	if(ensure(DIClass))
+	{
+		while(DIClass && DIClass != UObject::StaticClass())
+		{
+			if (TUniquePtr<INiagaraDataInterfaceNodeActionProvider>* Provider = RegisteredActionProviders.Find(DIClass->GetFName()))
+			{
+				(*Provider)->GetNodeContextMenuActionsImpl(Menu, Context, Signature);
+			}
+			DIClass = DIClass->GetSuperClass();
+		}
+	}
+}
+
+#undef LOCTEXT_NAMESPACE
\ No newline at end of file
diff --git a/Engine/Plugins/FX/Niagara/Source/NiagaraEditor/Private/NiagaraCompiler.cpp b/Engine/Plugins/FX/Niagara/Source/NiagaraEditor/Private/NiagaraCompiler.cpp
--- a/Engine/Plugins/FX/Niagara/Source/NiagaraEditor/Private/NiagaraCompiler.cpp
+++ b/Engine/Plugins/FX/Niagara/Source/NiagaraEditor/Private/NiagaraCompiler.cpp
@@ -2067,16 +2067,20 @@ TOptional<FNiagaraCompileResults> FHlslNiagaraCompiler::GetCompileResult(int32 J
 
 			if (Sig)
 			{
-				int32 NewBindingIdx = Results.Data->CalledVMExternalFunctions.AddDefaulted();
-				Results.Data->CalledVMExternalFunctions[NewBindingIdx].Name = Sig->Name;
-				Results.Data->CalledVMExternalFunctions[NewBindingIdx].OwnerName = Sig->OwnerName;
-				Results.Data->CalledVMExternalFunctions[NewBindingIdx].InputParamLocations = FuncInfo.InputParamLocations;
-				Results.Data->CalledVMExternalFunctions[NewBindingIdx].NumOutputs = FuncInfo.NumOutputs;
+				FVMExternalFunctionBindingInfo& NewBinding = Results.Data->CalledVMExternalFunctions.AddDefaulted_GetRef();
+				NewBinding.Name = Sig->Name;
+				NewBinding.OwnerName = Sig->OwnerName;
+				NewBinding.InputParamLocations = FuncInfo.InputParamLocations;
+				NewBinding.NumOutputs = FuncInfo.NumOutputs;
 				for (auto it = Sig->FunctionSpecifiers.CreateConstIterator(); it; ++it)
 				{
 					// we convert the map into an array here to save runtime memory
-					Results.Data->CalledVMExternalFunctions[NewBindingIdx].FunctionSpecifiers.Emplace(it->Key, it->Value);
+					NewBinding.FunctionSpecifiers.Emplace(it->Key, it->Value);
 				}
+
+				//Write out our variadic parameters to allow proper binding for VM external functions.
+				Sig->GetVariadicInputs(NewBinding.VariadicInputs);
+				Sig->GetVariadicOutputs(NewBinding.VariadicOutputs);
 			}
 			else
 			{
diff --git a/Engine/Plugins/FX/Niagara/Source/NiagaraEditor/Private/NiagaraHlslTranslator.cpp b/Engine/Plugins/FX/Niagara/Source/NiagaraEditor/Private/NiagaraHlslTranslator.cpp
--- a/Engine/Plugins/FX/Niagara/Source/NiagaraEditor/Private/NiagaraHlslTranslator.cpp
+++ b/Engine/Plugins/FX/Niagara/Source/NiagaraEditor/Private/NiagaraHlslTranslator.cpp
@@ -41,6 +41,8 @@
 #include "ShaderCore.h"
 #include "Misc/FileHelper.h"
 
+#include "NiagaraDataInterfaceUtilities.h"
+
 #define LOCTEXT_NAMESPACE "NiagaraCompiler"
 
 DECLARE_CYCLE_STAT(TEXT("Niagara - HlslTranslator - Translate"), STAT_NiagaraEditor_HlslTranslator_Translate, STATGROUP_NiagaraEditor);
@@ -2310,8 +2312,9 @@ void FHlslNiagaraTranslator::HandleSimStageSetupAndTeardown(int32 InWhichStage,
 		}
 
 		// Convert to a FNiagaraDataInterfaceGPUParamInfo to keep the API simple and consistent
+		TArray<FNiagaraFunctionSignature> GeneratedFunctionSignatures;
 		FNiagaraDataInterfaceGPUParamInfo DIInstanceInfo;
-		ConvertCompileInfoToParamInfo(CompilationOutput.ScriptData.DataInterfaceInfo[DataInterfaceOwnerIndex], DIInstanceInfo);
+		ConvertCompileInfoToParamInfo(CompilationOutput.ScriptData.DataInterfaceInfo[DataInterfaceOwnerIndex], DIInstanceInfo, GeneratedFunctionSignatures);
 
 		// This next part might be a big confusing, but because DataInterfaces are in non-editor code, it makes it impossible for them to return graphs or other
 		// structures. We want them to feel free to invoke their own fucntions and not have to do a lot of extra wranging, so we treat them like a custom hlsl node.
@@ -2455,19 +2458,19 @@ void FHlslNiagaraTranslator::GatherVariableForDataSetAccess(const FNiagaraVariab
 	check(Components.Num() == Types.Num());
 	for (int32 CompIdx = 0; CompIdx < Components.Num(); ++CompIdx)
 	{
-		if (Types[CompIdx] == NBT_Float)
+		if (Types[CompIdx] == ENiagaraBaseTypes::Float)
 		{
 			FormatArgs[1] = TEXT("Float");
 			FormatArgs[DefaultIdx] = TEXT("0.0f");
 			FormatArgs[RegIdx] = FloatCounter++;
 		}
-		else if (Types[CompIdx] == NBT_Half)
+		else if (Types[CompIdx] == ENiagaraBaseTypes::Half)
 		{
 			FormatArgs[1] = TEXT("Half");
 			FormatArgs[DefaultIdx] = TEXT("0.0f");
 			FormatArgs[RegIdx] = HalfCounter++;
 		}
-		else if (Types[CompIdx] == NBT_Int32)
+		else if (Types[CompIdx] == ENiagaraBaseTypes::Int32)
 		{
 			FormatArgs[1] = TEXT("Int");
 			FormatArgs[DefaultIdx] = TEXT("0");
@@ -2475,7 +2478,7 @@ void FHlslNiagaraTranslator::GatherVariableForDataSetAccess(const FNiagaraVariab
 		}
 		else
 		{
-			check(Types[CompIdx] == NBT_Bool);
+			check(Types[CompIdx] == ENiagaraBaseTypes::Bool);
 			FormatArgs[1] = TEXT("Bool");
 			FormatArgs[DefaultIdx] = TEXT("false");
 			FormatArgs[RegIdx] = IntCounter++;
@@ -2501,13 +2504,13 @@ void FHlslNiagaraTranslator::GatherComponentsForDataSetAccess(UScriptStruct* Str
 	//Bools are an awkward special case. TODO: make neater.
 	if (FNiagaraTypeDefinition(Struct) == FNiagaraTypeDefinition::GetBoolDef())
 	{
-		Types.Add(ENiagaraBaseTypes::NBT_Bool);
+		Types.Add(ENiagaraBaseTypes::Bool);
 		Components.Add(VariableSymbol);
 		return;
 	}
 	else if (FNiagaraTypeDefinition(Struct) == FNiagaraTypeDefinition::GetHalfDef())
 	{
-		Types.Add(ENiagaraBaseTypes::NBT_Half);
+		Types.Add(ENiagaraBaseTypes::Half);
 		Components.Add(VariableSymbol);
 		return;
 	}
@@ -2550,22 +2553,22 @@ void FHlslNiagaraTranslator::GatherComponentsForDataSetAccess(UScriptStruct* Str
 
 			if (Property->IsA(FFloatProperty::StaticClass()))
 			{
-				Types.Add(ENiagaraBaseTypes::NBT_Float);
+				Types.Add(ENiagaraBaseTypes::Float);
 				Components.Add(VarName);
 			}
 			else if (Property->IsA(FIntProperty::StaticClass()))
 			{
-				Types.Add(ENiagaraBaseTypes::NBT_Int32);
+				Types.Add(ENiagaraBaseTypes::Int32);
 				Components.Add(VarName);
 			}
 			else if (Property->IsA(FBoolProperty::StaticClass()))
 			{
-				Types.Add(ENiagaraBaseTypes::NBT_Bool);
+				Types.Add(ENiagaraBaseTypes::Bool);
 				Components.Add(VarName);
 			}
 			else if (Property->IsA(FUInt16Property::StaticClass()))
 			{
-				Types.Add(ENiagaraBaseTypes::NBT_Half);
+				Types.Add(ENiagaraBaseTypes::Half);
 				Components.Add(VarName);
 			}
 		}
@@ -2821,7 +2824,7 @@ void FHlslNiagaraTranslator::DefineDataSetWriteFunction(FString &HlslOutputStrin
 	HlslOutput += TEXT("}\n\n");
 }
 
-void FHlslNiagaraTranslator::ConvertCompileInfoToParamInfo(const FNiagaraScriptDataInterfaceCompileInfo& Info, FNiagaraDataInterfaceGPUParamInfo& DIInstanceInfo)
+void FHlslNiagaraTranslator::ConvertCompileInfoToParamInfo(const FNiagaraScriptDataInterfaceCompileInfo& Info, FNiagaraDataInterfaceGPUParamInfo& DIInstanceInfo, TArray<FNiagaraFunctionSignature>& GeneratedFunctionSignatures)
 {
 	FString OwnerIDString = Info.Name.ToString();
 	FString SanitizedOwnerIDString = GetSanitizedSymbolName(OwnerIDString.Replace(TEXT("."), TEXT("_")));
@@ -2833,6 +2836,7 @@ void FHlslNiagaraTranslator::ConvertCompileInfoToParamInfo(const FNiagaraScriptD
 	bool bHasWriteFunctions = false;
 	TSet<FNiagaraFunctionSignature> SeenFunctions;
 	DIInstanceInfo.GeneratedFunctions.Reserve(Info.RegisteredFunctions.Num());
+	GeneratedFunctionSignatures.Reserve(Info.RegisteredFunctions.Num());
 	for (const FNiagaraFunctionSignature& OriginalSig : Info.RegisteredFunctions)
 	{
 		if (SeenFunctions.Contains(OriginalSig))
@@ -2863,6 +2867,27 @@ void FHlslNiagaraTranslator::ConvertCompileInfoToParamInfo(const FNiagaraScriptD
 		{
 			DIFunc.Specifiers.Add(Specifier);
 		}
+
+		auto AddVarsToGPUDIFuncInfo = [](const TArray<FNiagaraVariableBase>& InVars, TArray<FNiagaraVariableCommonReference>& OutVarRefs)
+		{
+			OutVarRefs.Reset(InVars.Num());
+			for(const FNiagaraVariableBase& Var : InVars)
+			{
+				auto& NewVarRef = OutVarRefs.AddDefaulted_GetRef();
+				NewVarRef.Name = Var.GetName();
+				NewVarRef.UnderlyingType = Var.GetType().ClassStructOrEnum;
+			}
+		};
+
+		TArray<FNiagaraVariableBase> TempVars;
+		//Write out our variadic parameters to allow proper binding for VM external functions.
+		Sig.GetVariadicInputs(TempVars);
+		AddVarsToGPUDIFuncInfo(TempVars, DIFunc.VariadicInputs);
+		Sig.GetVariadicOutputs(TempVars);
+		AddVarsToGPUDIFuncInfo(TempVars, DIFunc.VariadicOutputs);
+
+		//Also output the actual signature to help with some hlsl generation.
+		GeneratedFunctionSignatures.Add(Sig);
 	}
 }
 
@@ -2887,10 +2912,16 @@ void FHlslNiagaraTranslator::DefineDataInterfaceHLSL(FString& InHlslOutput)
 				InterfaceClasses.Add(Info.Type.GetFName());
 			}
 
+			TArray<FNiagaraFunctionSignature> GeneratedFunctionSignatures;
 			FNiagaraDataInterfaceGPUParamInfo& DIInstanceInfo = ShaderScriptParametersMetadata.DataInterfaceParamInfo.AddDefaulted_GetRef();
-			ConvertCompileInfoToParamInfo(Info, DIInstanceInfo);
+			ConvertCompileInfoToParamInfo(Info, DIInstanceInfo,GeneratedFunctionSignatures);
+
+			FNiagaraDataInterfaceHlslGenerationContext DIHlslGenContext(DIInstanceInfo, GeneratedFunctionSignatures);
 
-			CDO->GetParameterDefinitionHLSL(DIInstanceInfo, InterfaceUniformHLSL);
+			DIHlslGenContext.GetStructHlslTypeNameDelegate.BindStatic(&FHlslNiagaraTranslator::GetStructHlslTypeName);
+			DIHlslGenContext.GetPropertyHlslTypeNameDelegate.BindStatic(&FHlslNiagaraTranslator::GetPropertyHlslTypeName);
+			DIHlslGenContext.GetSanitizedSymbolNameDelegate.BindStatic(&FHlslNiagaraTranslator::GetSanitizedSymbolName);
+			CDO->GetParameterDefinitionHLSL(DIHlslGenContext, InterfaceUniformHLSL);
 
 			// Ask the DI to generate HLSL.
 			TArray<FNiagaraDataInterfaceGeneratedFunction> PreviousHits;
@@ -2898,7 +2929,9 @@ void FHlslNiagaraTranslator::DefineDataInterfaceHLSL(FString& InHlslOutput)
 			{
 				const FNiagaraDataInterfaceGeneratedFunction& DIFunc = DIInstanceInfo.GeneratedFunctions[FunctionInstanceIndex];
 				ensure(!PreviousHits.Contains(DIFunc));
-				const bool HlslOK = CDO->GetFunctionHLSL(DIInstanceInfo, DIFunc, FunctionInstanceIndex, InterfaceFunctionHLSL);
+
+				DIHlslGenContext.FunctionInstanceIndex = FunctionInstanceIndex;
+				const bool HlslOK = CDO->GetFunctionHLSL(DIHlslGenContext, InterfaceFunctionHLSL);
 				if (!HlslOK)
 				{
 					Error(FText::Format(LOCTEXT("GPUDataInterfaceFunctionNotImplemented", "DataInterface {0} function {1} is not implemented for GPU."), FText::FromName(Info.Type.GetFName()), FText::FromName(DIFunc.DefinitionName)), nullptr, nullptr);
@@ -7134,20 +7167,24 @@ void FHlslNiagaraTranslator::FunctionCall(UNiagaraNodeFunctionCall* FunctionNode
 	}
 
 	// Remove input add pin if it exists
+	bool bHasOutputAdd = false;
 	for (int32 i = 0; i < CallOutputs.Num(); i++)
 	{
 		if (FunctionNode->IsAddPin(CallOutputs[i]))
 		{
+			bHasOutputAdd = true;
 			CallOutputs.RemoveAt(i);
 			break;
 		}
 	}
 
 	// Remove output add pin if it exists
+	bool bHasInputAdd = false;
 	for (int32 i = 0; i < CallInputs.Num(); i++)
 	{
 		if (FunctionNode->IsAddPin(CallInputs[i]))
 		{
+			bHasInputAdd = true;
 			CallInputs.RemoveAt(i);
 			break;
 		}
@@ -7256,10 +7293,13 @@ void FHlslNiagaraTranslator::FunctionCall(UNiagaraNodeFunctionCall* FunctionNode
 
 	GenerateFunctionCall(ScriptUsage, OutputSignature, Inputs, Outputs);
 
-	if (bCustomHlsl)
+	// Re-add the add pins.
+	if (bHasInputAdd)
 	{
-		// Re-add the add pins.
 		Inputs.Add(INDEX_NONE);
+	}
+	if(bHasOutputAdd)
+	{
 		Outputs.Add(INDEX_NONE);
 	}
 	ActiveHistoryForFunctionCalls.ExitFunction(FunctionNode->GetFunctionName(), FunctionNode->FunctionScript, FunctionNode);
diff --git a/Engine/Plugins/FX/Niagara/Source/NiagaraEditor/Private/NiagaraHlslTranslator.h b/Engine/Plugins/FX/Niagara/Source/NiagaraEditor/Private/NiagaraHlslTranslator.h
--- a/Engine/Plugins/FX/Niagara/Source/NiagaraEditor/Private/NiagaraHlslTranslator.h
+++ b/Engine/Plugins/FX/Niagara/Source/NiagaraEditor/Private/NiagaraHlslTranslator.h
@@ -555,7 +555,7 @@ protected:
 	FString GetUniqueEmitterName() const;
 
 	void HandleDataInterfaceCall(FNiagaraScriptDataInterfaceCompileInfo& Info, const FNiagaraFunctionSignature& InMatchingSignature);
-	void ConvertCompileInfoToParamInfo(const FNiagaraScriptDataInterfaceCompileInfo& InCompileInfo, FNiagaraDataInterfaceGPUParamInfo& OutGPUParamInfo);
+	void ConvertCompileInfoToParamInfo(const FNiagaraScriptDataInterfaceCompileInfo& InCompileInfo, FNiagaraDataInterfaceGPUParamInfo& OutGPUParamInfo, TArray<FNiagaraFunctionSignature>& GeneratedFunctionSignatures);
 public:
 
 	FHlslNiagaraTranslator();
diff --git a/Engine/Plugins/FX/Niagara/Source/NiagaraEditor/Private/NiagaraNodeCustomHlsl.cpp b/Engine/Plugins/FX/Niagara/Source/NiagaraEditor/Private/NiagaraNodeCustomHlsl.cpp
--- a/Engine/Plugins/FX/Niagara/Source/NiagaraEditor/Private/NiagaraNodeCustomHlsl.cpp
+++ b/Engine/Plugins/FX/Niagara/Source/NiagaraEditor/Private/NiagaraNodeCustomHlsl.cpp
@@ -632,6 +632,10 @@ void UNiagaraNodeCustomHlsl::RebuildSignatureFromPins()
 		Sig.Outputs.Add(Schema->PinToNiagaraVariable(Pin, false));
 	}
 
+	//Allow signature to have variadic params so we still get the add pins.
+	Sig.RequiredInputs = Sig.Inputs.Num();
+	Sig.RequiredOutputs = Sig.Outputs.Num();
+
 	Signature = Sig;
 
 }
diff --git a/Engine/Plugins/FX/Niagara/Source/NiagaraEditor/Private/NiagaraNodeFunctionCall.cpp b/Engine/Plugins/FX/Niagara/Source/NiagaraEditor/Private/NiagaraNodeFunctionCall.cpp
--- a/Engine/Plugins/FX/Niagara/Source/NiagaraEditor/Private/NiagaraNodeFunctionCall.cpp
+++ b/Engine/Plugins/FX/Niagara/Source/NiagaraEditor/Private/NiagaraNodeFunctionCall.cpp
@@ -391,8 +391,22 @@ PRAGMA_ENABLE_DEPRECATION_WARNINGS;
 
 		if (AllowDynamicPins())
 		{
-			CreateAddPin(EGPD_Input);
-			CreateAddPin(EGPD_Output);
+			if(Signature.IsValid())
+			{
+				if(Signature.VariadicInput())
+				{
+					CreateAddPin(EGPD_Input);
+				}
+				if(Signature.VariadicOutput())
+				{
+					CreateAddPin(EGPD_Output);
+				}
+			}
+			else
+			{
+				CreateAddPin(EGPD_Input);
+				CreateAddPin(EGPD_Output);
+			}
 		}
 
 		// We don't reference an external function, so set an invalid id.
@@ -1070,6 +1084,19 @@ void UNiagaraNodeFunctionCall::UpdateReferencedStaticsHashForNode(FSHA1& HashSta
 
 }
 
+void UNiagaraNodeFunctionCall::GetNodeContextMenuActions(class UToolMenu* Menu, class UGraphNodeContextMenuContext* Context) const
+{
+	Super::GetNodeContextMenuActions(Menu, Context);
+	if(FunctionScript == nullptr)
+	{
+		UClass* DIClass = CastChecked<UClass>(Signature.Inputs[0].GetType().GetClass());
+		if(DIClass)
+		{
+			INiagaraDataInterfaceNodeActionProvider::GetNodeContextMenuActions(DIClass, Menu, Context, Signature);
+		}
+	}
+}
+
 bool UNiagaraNodeFunctionCall::HasValidScriptAndGraph() const
 {
 	return FunctionScript != nullptr && GetCalledGraph() != nullptr;
@@ -1492,6 +1519,52 @@ void UNiagaraNodeFunctionCall::AutowireNewNode(UEdGraphPin* FromPin)
 	ComputeNodeName();
 }
 
+void UNiagaraNodeFunctionCall::RefreshSignature()
+{
+	if (Signature.IsValid() && Signature.bMemberFunction)
+	{
+		if ((Signature.Inputs.Num() > 0) && Signature.Inputs[0].GetType().IsDataInterface())
+		{
+			UNiagaraDataInterface* CDO = CastChecked<UNiagaraDataInterface>(Signature.Inputs[0].GetType().GetClass()->GetDefaultObject());
+			TArray<FNiagaraFunctionSignature> BaseDIFuncs;
+			CDO->GetFunctions(BaseDIFuncs);
+			if (FNiagaraFunctionSignature* BaseSig = BaseDIFuncs.FindByPredicate([&](const FNiagaraFunctionSignature& CheckSig) { return Signature.Name == CheckSig.Name; }))
+			{
+				//Revert signature to base and re-add dynamic pins as new inputs and outputs.
+				Signature = *BaseSig;				
+
+				FPinCollectorArray FoundPins;
+				GetInputPins(FoundPins);
+
+				for (int32 i = 0; i < FoundPins.Num(); ++i)
+				{
+					UEdGraphPin* Pin = FoundPins[i];
+
+					if (DynamicPins.Contains(Pin->PinId))
+					{
+						Signature.AddInput(UEdGraphSchema_Niagara::PinToNiagaraVariable(Pin), FText::FromString(Pin->PinToolTip));
+					}
+				}
+
+				GetOutputPins(FoundPins);
+
+				for (int32 i = 0; i < FoundPins.Num(); ++i)
+				{
+					UEdGraphPin* Pin = FoundPins[i];
+
+					if (DynamicPins.Contains(Pin->PinId))
+					{
+						Signature.AddOutput(UEdGraphSchema_Niagara::PinToNiagaraVariable(Pin), FText::FromString(Pin->PinToolTip));
+					}
+				}
+
+				return;
+			}
+		}
+	}
+	Signature = FNiagaraFunctionSignature();
+}
+
 void UNiagaraNodeFunctionCall::ComputeNodeName(FString SuggestedName, bool bForceSuggestion)
 {
 	FString FunctionName = FunctionScript ? FunctionScript->GetName() : Signature.GetNameString();
diff --git a/Engine/Plugins/FX/Niagara/Source/NiagaraEditor/Private/NiagaraNodeWithDynamicPins.cpp b/Engine/Plugins/FX/Niagara/Source/NiagaraEditor/Private/NiagaraNodeWithDynamicPins.cpp
--- a/Engine/Plugins/FX/Niagara/Source/NiagaraEditor/Private/NiagaraNodeWithDynamicPins.cpp
+++ b/Engine/Plugins/FX/Niagara/Source/NiagaraEditor/Private/NiagaraNodeWithDynamicPins.cpp
@@ -58,6 +58,8 @@ void UNiagaraNodeWithDynamicPins::PinConnectionListChanged(UEdGraphPin* Pin)
 		Pin->PinName = NewPinName;
 		Pin->PinFriendlyName = NewPinFriendlyName;
 
+		DynamicPins.Add(Pin->PinId);
+
 		CreateAddPin(Pin->Direction);
 		OnNewTypedPinAdded(Pin);
 		MarkNodeRequiresSynchronization(__FUNCTION__, true);
@@ -115,6 +117,8 @@ UEdGraphPin* UNiagaraNodeWithDynamicPins::RequestNewTypedPin(EEdGraphPinDirectio
 	AddPin->Modify();
 	AddPin->PinType = Schema->TypeDefinitionToPinType(Type);
 	AddPin->PinName = InName;
+	
+	DynamicPins.Add(AddPin->PinId);
 
 	CreateAddPin(Direction);
 	// we pass the pointer in as reference in case we want to reallocate so the overriding node has a chance to restore the pointer
@@ -143,19 +147,33 @@ bool UNiagaraNodeWithDynamicPins::IsAddPin(const UEdGraphPin* Pin) const
 		Pin->PinType.PinSubCategory == AddPinSubCategory;
 }
 
+bool UNiagaraNodeWithDynamicPins::IsExecPin(const UEdGraphPin* Pin) const
+{
+	const UEdGraphSchema_Niagara* Schema = GetDefault<UEdGraphSchema_Niagara>();
+	return Pin->PinType == Schema->TypeDefinitionToPinType(FNiagaraTypeDefinition::GetParameterMapDef()) && Pin->PinName == TEXT("");
+}
+
 bool UNiagaraNodeWithDynamicPins::CanRenamePin(const UEdGraphPin* Pin) const
 {
-	return IsAddPin(Pin) == false;
+	if (CanModifyExistingPins() == false )
+	{
+		return DynamicPins.Contains(Pin->PinId);
+	}
+	return IsAddPin(Pin) == false || IsExecPin(Pin);
 }
 
 bool UNiagaraNodeWithDynamicPins::CanRemovePin(const UEdGraphPin* Pin) const
 {
-	return IsAddPin(Pin) == false;
+	if (CanModifyExistingPins() == false)
+	{
+		return DynamicPins.Contains(Pin->PinId);
+	}
+	return IsAddPin(Pin) == false || IsExecPin(Pin);
 }
 
 bool UNiagaraNodeWithDynamicPins::CanMovePin(const UEdGraphPin* Pin, int32 DirectionToMove) const
 {
-	if(IsAddPin(Pin) || IsParameterMapPin(Pin) || Pin->bOrphanedPin)
+	if(IsAddPin(Pin) || IsExecPin(Pin) || IsParameterMapPin(Pin) || Pin->bOrphanedPin || (CanModifyExistingPins() == false && DynamicPins.Contains(Pin->PinId) == false))
 	{
 		return false;
 	}
@@ -174,7 +192,7 @@ bool UNiagaraNodeWithDynamicPins::CanMovePin(const UEdGraphPin* Pin, int32 Direc
 	if(PinArray.IsValidIndex(Index + DirectionToMove))
 	{
 		const UEdGraphPin* PinToMoveTo = PinArray[Index + DirectionToMove];
-		if(IsAddPin(PinToMoveTo) || IsParameterMapPin(PinToMoveTo) || PinToMoveTo->bOrphanedPin)
+		if(IsAddPin(PinToMoveTo) || IsParameterMapPin(PinToMoveTo) || PinToMoveTo->bOrphanedPin || (CanModifyExistingPins() == false && DynamicPins.Contains(PinToMoveTo->PinId) == false))
 		{
 			return false;
 		}
@@ -452,6 +470,7 @@ void UNiagaraNodeWithDynamicPins::AddExistingParameter(FNiagaraVariable Paramete
 
 void UNiagaraNodeWithDynamicPins::RemoveDynamicPin(UEdGraphPin* Pin)
 {
+	DynamicPins.Remove(Pin->PinId);
 	RemovePin(Pin);
 	MarkNodeRequiresSynchronization(__FUNCTION__, true);
 }
@@ -492,5 +511,23 @@ void UNiagaraNodeWithDynamicPins::MoveDynamicPinFromMenu(UEdGraphPin* Pin, int32
 	MoveDynamicPin(Pin, DirectionToMove);
 }
 
+void UNiagaraNodeWithDynamicPins::RemoveAllDynamicPins()
+{
+	FScopedTransaction RemovePinTransaction(LOCTEXT("RemoveAllDynamicPinsTransaction", "Remove all dynamic pins"));
+	while(DynamicPins.Num() > 0)
+	{
+		FGuid& PinToRemoveGuid = DynamicPins.Last();
+		if(UEdGraphPin** Pin = Pins.FindByPredicate([&](UEdGraphPin* Pin){ return Pin->PinId == PinToRemoveGuid;}))
+		{
+			RemoveDynamicPin(*Pin);
+		}
+		else
+		{
+			checkf(0, TEXT("Dynamic pin could not be found."));
+			break;
+		}
+	}
+}
+
 #undef LOCTEXT_NAMESPACE
 
diff --git a/Engine/Plugins/FX/Niagara/Source/NiagaraEditor/Public/NiagaraActions.h b/Engine/Plugins/FX/Niagara/Source/NiagaraEditor/Public/NiagaraActions.h
--- a/Engine/Plugins/FX/Niagara/Source/NiagaraEditor/Public/NiagaraActions.h
+++ b/Engine/Plugins/FX/Niagara/Source/NiagaraEditor/Public/NiagaraActions.h
@@ -11,6 +11,11 @@
 #include "ViewModels/Stack/NiagaraParameterHandle.h"
 #include "NiagaraActions.generated.h"
 
+
+class UToolMenu;
+class UGraphNodeContextMenuContext;
+
+
 UENUM()
 enum class ENiagaraMenuSections : uint8
 {
@@ -377,4 +382,57 @@ public:
 	TWeakObjectPtr<UNiagaraScript> Script;
 	FGuid Version;
 	FText FriendlyName;
-};
\ No newline at end of file
+};
+
+
+class NIAGARAEDITOR_API INiagaraDataInterfaceNodeActionProvider
+{
+public:
+
+	INiagaraDataInterfaceNodeActionProvider() = default;
+	INiagaraDataInterfaceNodeActionProvider(INiagaraDataInterfaceNodeActionProvider&) = delete;
+	INiagaraDataInterfaceNodeActionProvider(INiagaraDataInterfaceNodeActionProvider&&) = delete;
+	INiagaraDataInterfaceNodeActionProvider& operator=(INiagaraDataInterfaceNodeActionProvider&) = delete;
+	INiagaraDataInterfaceNodeActionProvider& operator=(INiagaraDataInterfaceNodeActionProvider&&) = delete;
+	virtual ~INiagaraDataInterfaceNodeActionProvider() = default;
+
+	/** Pure virtual that sub classes should implement to provide menu actions specific to DIs and their functions. */
+	virtual void GetNodeContextMenuActionsImpl(UToolMenu* Menu, UGraphNodeContextMenuContext* Context, FNiagaraFunctionSignature Signature) const = 0;
+
+	template<typename DIClass, typename ActionProviderClass>
+	static void Register();
+
+	template<typename DIClass>
+	static void Unregister();
+
+	template<typename DIClass>
+	static void GetNodeContextMenuActions(UToolMenu* Menu, UGraphNodeContextMenuContext* Context, FNiagaraFunctionSignature Signature);
+
+	/**  */
+	static void GetNodeContextMenuActions(UClass* DIClass, UToolMenu* Menu, UGraphNodeContextMenuContext* Context, FNiagaraFunctionSignature Signature);
+
+	/** All currently registered action providers. */
+	static TMap<FName, TUniquePtr<INiagaraDataInterfaceNodeActionProvider>> RegisteredActionProviders;
+};
+
+template<typename DIClass, typename ActionProviderClass>
+void INiagaraDataInterfaceNodeActionProvider::Register()
+{
+	TUniquePtr<INiagaraDataInterfaceNodeActionProvider>& Provider = RegisteredActionProviders.FindOrAdd(DIClass::StaticClass()->GetFName());
+	Provider = MakeUnique<ActionProviderClass>();
+}
+
+template<typename DIClass>
+void INiagaraDataInterfaceNodeActionProvider::Unregister()
+{
+	RegisteredActionProviders.Remove(DIClass::StaticClass()->GetFName());
+}
+
+template<typename DIClass>
+void INiagaraDataInterfaceNodeActionProvider::GetNodeContextMenuActions(UToolMenu* Menu, UGraphNodeContextMenuContext* Context, FNiagaraFunctionSignature Signature)
+{
+	if(TUniquePtr<INiagaraDataInterfaceNodeActionProvider>* Provider = RegisteredActionProviders.Find(DIClass::StaticClass()->GetFName()))
+	{
+		(*Provider)->GetNodeContextMenuActionsImpl(Menu, Context, Signature);
+	}
+}
diff --git a/Engine/Plugins/FX/Niagara/Source/NiagaraEditor/Public/NiagaraNodeAssignment.h b/Engine/Plugins/FX/Niagara/Source/NiagaraEditor/Public/NiagaraNodeAssignment.h
--- a/Engine/Plugins/FX/Niagara/Source/NiagaraEditor/Public/NiagaraNodeAssignment.h
+++ b/Engine/Plugins/FX/Niagara/Source/NiagaraEditor/Public/NiagaraNodeAssignment.h
@@ -45,6 +45,7 @@ public:
 	virtual void AllocateDefaultPins() override;
 
 	//~ UNiagaraNodeFunctionCall interface
+	virtual bool AllowDynamicPins() const override { return true; }
 	virtual bool RefreshFromExternalChanges() override;
 	virtual void BuildParameterMapHistory(FNiagaraParameterMapHistoryBuilder& OutHistory, bool bRecursive = true, bool bFilterForCompilation = true) const override;
 	virtual void GatherExternalDependencyData(ENiagaraScriptUsage InUsage, const FGuid& InUsageId, TArray<FNiagaraCompileHash>& InReferencedCompileHashes, TArray<FString>& InReferencedObjs) const override;
diff --git a/Engine/Plugins/FX/Niagara/Source/NiagaraEditor/Public/NiagaraNodeFunctionCall.h b/Engine/Plugins/FX/Niagara/Source/NiagaraEditor/Public/NiagaraNodeFunctionCall.h
--- a/Engine/Plugins/FX/Niagara/Source/NiagaraEditor/Public/NiagaraNodeFunctionCall.h
+++ b/Engine/Plugins/FX/Niagara/Source/NiagaraEditor/Public/NiagaraNodeFunctionCall.h
@@ -113,6 +113,7 @@ public:
 	virtual bool CanAddToGraph(UNiagaraGraph* TargetGraph, FString& OutErrorMsg) const override;
 	virtual void GatherExternalDependencyData(ENiagaraScriptUsage InUsage, const FGuid& InUsageId, TArray<FNiagaraCompileHash>& InReferencedCompileHashes, TArray<FString>& InReferencedObjs) const override;
 	virtual void UpdateCompileHashForNode(FSHA1& HashState) const override;
+	virtual void GetNodeContextMenuActions(class UToolMenu* Menu, class UGraphNodeContextMenuContext* Context) const override;
 	//End UNiagaraNode interface
 
 	virtual void UpdateReferencedStaticsHashForNode(FSHA1& HashState) const;
@@ -191,10 +192,18 @@ protected:
 
 	virtual bool GetValidateDataInterfaces() const { return true; };
 
-	virtual bool AllowDynamicPins() const override { return false; }
-	virtual bool CanRenamePin(const UEdGraphPin* Pin) const override { return false; }
-	virtual bool CanRemovePin(const UEdGraphPin* Pin) const override { return false; }
-	virtual bool CanMovePin(const UEdGraphPin* Pin, int32 DirectionToMove) const override { return false; }
+	virtual bool AllowDynamicPins() const override { return Signature.VariadicInput() || Signature.VariadicOutput(); }
+	virtual bool CanRenamePin(const UEdGraphPin* Pin) const override { return AllowDynamicPins() && Super::CanRenamePin(Pin); }
+	virtual bool CanRemovePin(const UEdGraphPin* Pin) const override { return AllowDynamicPins() && Super::CanRemovePin(Pin); }
+	virtual bool CanMovePin(const UEdGraphPin* Pin, int32 DirectionToMove) const override { return AllowDynamicPins() && Super::CanMovePin(Pin, DirectionToMove); }
+	virtual bool CanModifyExistingPins() const { return false; }
+
+	virtual void OnNewTypedPinAdded(UEdGraphPin*& NewPin) { Super::OnNewTypedPinAdded(NewPin); RefreshSignature(); }
+	virtual void OnPinRenamed(UEdGraphPin* RenamedPin, const FString& OldPinName) override { Super::OnPinRenamed(RenamedPin, OldPinName); RefreshSignature(); }
+	virtual void RemoveDynamicPin(UEdGraphPin* Pin) override { Super::RemoveDynamicPin(Pin); RefreshSignature(); }
+	virtual void MoveDynamicPin(UEdGraphPin* Pin, int32 MoveAmount)override { Super::MoveDynamicPin(Pin, MoveAmount); RefreshSignature(); }
+
+	void RefreshSignature();
 
 	/** Resets the node name based on the referenced script or signature. Guaranteed unique within a given graph instance.*/
 	void ComputeNodeName(FString SuggestedName = FString(), bool bForceSuggestion = false);
diff --git a/Engine/Plugins/FX/Niagara/Source/NiagaraEditor/Public/NiagaraNodeWithDynamicPins.h b/Engine/Plugins/FX/Niagara/Source/NiagaraEditor/Public/NiagaraNodeWithDynamicPins.h
--- a/Engine/Plugins/FX/Niagara/Source/NiagaraEditor/Public/NiagaraNodeWithDynamicPins.h
+++ b/Engine/Plugins/FX/Niagara/Source/NiagaraEditor/Public/NiagaraNodeWithDynamicPins.h
@@ -34,6 +34,9 @@ public:
 	/** Helper to identify if a pin is an Add pin.*/
 	bool IsAddPin(const UEdGraphPin* Pin) const;
 
+	/** Helper to identify if a pin is an Add pin.*/
+	bool IsExecPin(const UEdGraphPin* Pin) const;
+
 	/** Determine whether or not a Niagara type is supported for an Add Pin possibility.*/
 	virtual bool AllowNiagaraTypeForAddPin(const FNiagaraTypeDefinition& InType) const;
 	virtual bool AllowNiagaraTypeForAddPin(const FNiagaraTypeDefinition& InType, EEdGraphPinDirection InDirection) const { return AllowNiagaraTypeForAddPin(InType); };
@@ -46,6 +49,9 @@ public:
 	/** Convenience method to determine whether this Node is a Map Get or Map Set when adding a parameter through the parameter panel. */
 	virtual EEdGraphPinDirection GetPinDirectionForNewParameters() { return EEdGraphPinDirection::EGPD_MAX; };
 
+	/** Will remove any dynamic pins. */
+	void RemoveAllDynamicPins();
+
 protected:
 	virtual bool AllowDynamicPins() const { return true; }
 
@@ -70,6 +76,9 @@ protected:
 	/** Called to determine if a pin can be removed by the user. */
 	virtual bool CanRemovePin(const UEdGraphPin* Pin) const;
 
+	/** Called to determine if this node can modify pins that existed on node creation and were not added dynamically. If true, the only pin that cannot be modified is the Add pin. */
+	virtual bool CanModifyExistingPins() const { return true; }
+
 	/** Called to determine if a pin can be moved by the user. Negative values for up, positive for down. */
 	virtual bool CanMovePin(const UEdGraphPin* Pin, int32 DirectionToMove) const;
 
@@ -100,4 +109,8 @@ private:
 public:
 	/** The sub category for add pins. */
 	static const FName AddPinSubCategory;
+
+	/** Track dynamic pins so that we only allow removing of dynamic pins. */
+	UPROPERTY()
+	TArray<FGuid> DynamicPins;
 };
diff --git a/Engine/Plugins/FX/Niagara/Source/NiagaraShader/Private/NiagaraShader.cpp b/Engine/Plugins/FX/Niagara/Source/NiagaraShader/Private/NiagaraShader.cpp
--- a/Engine/Plugins/FX/Niagara/Source/NiagaraShader/Private/NiagaraShader.cpp
+++ b/Engine/Plugins/FX/Niagara/Source/NiagaraShader/Private/NiagaraShader.cpp
@@ -1091,9 +1091,18 @@ FNiagaraShader::FNiagaraShader(const FNiagaraShaderType::CompiledShaderInitializ
 //////////////////////////////////////////////////////////////////////////
 bool FNiagaraDataInterfaceGeneratedFunction::Serialize(FArchive& Ar)
 {
+	Ar.UsingCustomVersion(FNiagaraCustomVersion::GUID);
+	const int32 NiagaraVer = Ar.CustomVer(FNiagaraCustomVersion::GUID);
+
 	Ar << DefinitionName;
 	Ar << InstanceName;
 	Ar << Specifiers;
+
+	if (NiagaraVer >= FNiagaraCustomVersion::AddVariadicParametersToGPUFunctionInfo)
+	{
+		Ar << VariadicInputs;
+		Ar << VariadicOutputs;
+	}
 	return true;
 }
 
diff --git a/Engine/Plugins/FX/Niagara/Source/NiagaraShader/Private/NiagaraShared.cpp b/Engine/Plugins/FX/Niagara/Source/NiagaraShader/Private/NiagaraShared.cpp
--- a/Engine/Plugins/FX/Niagara/Source/NiagaraShader/Private/NiagaraShared.cpp
+++ b/Engine/Plugins/FX/Niagara/Source/NiagaraShader/Private/NiagaraShared.cpp
@@ -326,6 +326,7 @@ void FNiagaraShaderScript::SerializeShaderMap(FArchive& Ar)
 	bool bCooked = Ar.IsCooking();
 	Ar << bCooked;
 	Ar << NumPermutations;
+	Ar << BaseCompileHash;
 
 	if (Ar.IsLoading())
 	{
diff --git a/Engine/Plugins/FX/Niagara/Source/NiagaraShader/Public/NiagaraShared.h b/Engine/Plugins/FX/Niagara/Source/NiagaraShader/Public/NiagaraShared.h
--- a/Engine/Plugins/FX/Niagara/Source/NiagaraShader/Public/NiagaraShared.h
+++ b/Engine/Plugins/FX/Niagara/Source/NiagaraShader/Public/NiagaraShared.h
@@ -27,6 +27,7 @@
 #include "Misc/Optional.h"
 #include "NiagaraCompileHash.h"
 #include "NiagaraDataInterfaceBase.h"
+#include "NiagaraCore.h"
 #include "NiagaraShared.generated.h"
 
 class FNiagaraShaderScript;
@@ -155,6 +156,11 @@ struct NIAGARASHADER_API FNiagaraDataInterfaceGeneratedFunction
 	using FunctionSpecifier = TTuple<FName, FName>;
 	TArray<FunctionSpecifier> Specifiers;
 
+	UPROPERTY()
+	TArray<FNiagaraVariableCommonReference> VariadicInputs;
+	UPROPERTY()
+	TArray<FNiagaraVariableCommonReference> VariadicOutputs;
+
 	const FName* FindSpecifierValue(const FName& Name) const
 	{
 		for (const TTuple<FName, FName>& Specifier : Specifiers)
@@ -181,6 +187,12 @@ struct NIAGARASHADER_API FNiagaraDataInterfaceGeneratedFunction
 			if (Specifiers[i] != Other.Specifiers[i])
 				return false;
 		}
+
+		if(VariadicInputs != Other.VariadicInputs)
+			return false;
+		if(VariadicOutputs != Other.VariadicOutputs)
+			return false;
+
 		return true;
 	}
 
@@ -898,6 +910,8 @@ public:
 	FORCEINLINE bool IsExternalConstantBufferUsed_RenderThread(int32 Index) const { return (CachedData_RenderThread.bExternalConstantBufferUsed & (1 << Index)) != 0; }
 	FORCEINLINE bool IsViewUniformBufferUsed_RenderThread() const { return CachedData_RenderThread.bViewUniformBufferUsed != 0; }
 
+	FNiagaraCompileHash GetBaseCompileHash()const { return BaseCompileHash; }
+
 protected:
 #if WITH_EDITOR
 	/**
diff --git a/Engine/Source/Runtime/Core/Private/UObject/FortniteMainBranchObjectVersion.cpp b/Engine/Source/Runtime/Core/Private/UObject/FortniteMainBranchObjectVersion.cpp
--- a/Engine/Source/Runtime/Core/Private/UObject/FortniteMainBranchObjectVersion.cpp
+++ b/Engine/Source/Runtime/Core/Private/UObject/FortniteMainBranchObjectVersion.cpp
@@ -10,8 +10,8 @@ TMap<FGuid, FGuid> FFortniteMainBranchObjectVersion::GetSystemGuids()
 	SystemGuids.Add(DevGuids.LANDSCAPE_MOBILE_COOK_VERSION, FGuid("32D02EF867C74B71A0D4E0FA41392732"));
 	SystemGuids.Add(DevGuids.MATERIALSHADERMAP_DERIVEDDATA_VER, FGuid("CE7776E63E9E43FBB8C103EE1AB40B5D"));
 	SystemGuids.Add(DevGuids.NANITE_DERIVEDDATA_VER, FGuid("B62CDCAF66AA45289DDB82E276E7C3E7"));
-	SystemGuids.Add(DevGuids.NIAGARASHADERMAP_DERIVEDDATA_VER, FGuid("35D7960B72164B49B9B7B3FB5857E674"));
-	SystemGuids.Add(DevGuids.Niagara_LatestScriptCompileVersion, FGuid("08ED74B3CB736F4AAD38953889079C6D"));
+	SystemGuids.Add(DevGuids.NIAGARASHADERMAP_DERIVEDDATA_VER, FGuid("07FCB513221642B795CAC18A0D71FA01"));
+	SystemGuids.Add(DevGuids.Niagara_LatestScriptCompileVersion, FGuid("35DC6B707D5E4883AB13750DCD18F021"));
 	SystemGuids.Add(DevGuids.SkeletalMeshDerivedDataVersion, FGuid("ACF593EAAE354FCBB34CF44F5AA1BFD2"));
 	SystemGuids.Add(DevGuids.STATICMESH_DERIVEDDATA_VER, FGuid("543843B103794E0E9BA4BE80FB602F79"));
 	return SystemGuids;
