Missing file from previous

diff --git a/Engine/Plugins/FX/Niagara/Source/Niagara/Private/NiagaraDataInterface.cpp b/Engine/Plugins/FX/Niagara/Source/Niagara/Private/NiagaraDataInterface.cpp
--- a/Engine/Plugins/FX/Niagara/Source/Niagara/Private/NiagaraDataInterface.cpp
+++ b/Engine/Plugins/FX/Niagara/Source/Niagara/Private/NiagaraDataInterface.cpp
@@ -12,6 +12,7 @@
 #include "NiagaraSystem.h"
 #include "NiagaraComponent.h"
 #include "ShaderCompilerCore.h"
+#include "NiagaraDataInterfaceUtilities.h"
 
 #include UE_INLINE_GENERATED_CPP_BY_NAME(NiagaraDataInterface)
 
@@ -164,6 +165,21 @@ bool UNiagaraDataInterface::CopyToInternal(UNiagaraDataInterface* Destination) c
 }
 
 #if WITH_EDITOR
+
+void UNiagaraDataInterface::GetParameterDefinitionHLSL(FNiagaraDataInterfaceHlslGenerationContext& HlslGenContext, FString& OutHLSL)
+{
+	//Base implementation will call into the legacy functions so that all DIs need not be updated immediately.
+	//Do not implement both current and legacy functions in any DI or it's parent classes.
+	GetParameterDefinitionHLSL(HlslGenContext.ParameterInfo, OutHLSL);
+}
+
+bool UNiagaraDataInterface::GetFunctionHLSL(FNiagaraDataInterfaceHlslGenerationContext& HlslGenContext, FString& OutHLSL)
+{
+	//Base implementation will call into the legacy functions so that all DIs need not be updated immediately.
+	//Do not implement both current and legacy functions in any DI or it's parent classes.
+	return GetFunctionHLSL(HlslGenContext.ParameterInfo, HlslGenContext.ParameterInfo.GeneratedFunctions[HlslGenContext.FunctionInstanceIndex], HlslGenContext.FunctionInstanceIndex, OutHLSL);
+}
+
 void UNiagaraDataInterface::GetFeedback(UNiagaraSystem* InAsset, UNiagaraComponent* InComponent, TArray<FNiagaraDataInterfaceError>& OutErrors, TArray<FNiagaraDataInterfaceFeedback>& OutWarnings, TArray<FNiagaraDataInterfaceFeedback>& OutInfo)
 {
 	OutErrors = GetErrors();
