Niagara Data Channels -  Supporting CL.

Adding Script paramter to DI utility functors as needed by data channels.
New filter option for DI search that can optionally include "internal" DIs.

diff --git a/Engine/Plugins/FX/Niagara/Source/Niagara/Private/DataInterface/NiagaraDataInterfaceStaticMesh.cpp b/Engine/Plugins/FX/Niagara/Source/Niagara/Private/DataInterface/NiagaraDataInterfaceStaticMesh.cpp
--- a/Engine/Plugins/FX/Niagara/Source/Niagara/Private/DataInterface/NiagaraDataInterfaceStaticMesh.cpp
+++ b/Engine/Plugins/FX/Niagara/Source/Niagara/Private/DataInterface/NiagaraDataInterfaceStaticMesh.cpp
@@ -3203,7 +3203,7 @@ void UNiagaraDataInterfaceStaticMesh::GetFeedback(UNiagaraSystem* Asset, UNiagar
 			FNiagaraDataInterfaceUtilities::ForEachVMFunctionEquals(
 				this,
 				Asset,
-				[&](const FVMExternalFunctionBindingInfo& VMFunction) -> bool
+				[&](const UNiagaraScript* Script, const FVMExternalFunctionBindingInfo& VMFunction) -> bool
 				{
 					bCpuAccessWarning |= CpuAccessFunctions.ContainsByPredicate([&](const FNiagaraFunctionSignature& CheckFunction) { return CheckFunction.Name == VMFunction.Name; });
 					return bCpuAccessWarning == false;
@@ -3213,7 +3213,7 @@ void UNiagaraDataInterfaceStaticMesh::GetFeedback(UNiagaraSystem* Asset, UNiagar
 			FNiagaraDataInterfaceUtilities::ForEachGpuFunctionEquals(
 				this,
 				Asset,
-				[&](const FNiagaraDataInterfaceGeneratedFunction& GpuFunction) -> bool
+				[&](const UNiagaraScript* Script,const FNiagaraDataInterfaceGeneratedFunction& GpuFunction) -> bool
 				{
 					bCpuAccessWarning |= CpuAccessFunctions.ContainsByPredicate([&](const FNiagaraFunctionSignature& CheckFunction) { return CheckFunction.Name == GpuFunction.DefinitionName; });
 					return bCpuAccessWarning == false;
@@ -3251,7 +3251,7 @@ void UNiagaraDataInterfaceStaticMesh::GetFeedback(UNiagaraSystem* Asset, UNiagar
 	{
 		FNiagaraDataInterfaceUtilities::ForEachGpuFunctionEquals(
 			this, Asset, Component,
-			[&](const FNiagaraDataInterfaceGeneratedFunction& FunctionBinding)
+			[&](const UNiagaraScript* Script, const FNiagaraDataInterfaceGeneratedFunction& FunctionBinding)
 			{
 				if (FunctionBinding.DefinitionName == NDIStaticMeshLocal::QueryDistanceFieldName )
 				{
diff --git a/Engine/Plugins/FX/Niagara/Source/Niagara/Private/DataInterface/NiagaraDataInterfaceUtilities.cpp b/Engine/Plugins/FX/Niagara/Source/Niagara/Private/DataInterface/NiagaraDataInterfaceUtilities.cpp
--- a/Engine/Plugins/FX/Niagara/Source/Niagara/Private/DataInterface/NiagaraDataInterfaceUtilities.cpp
+++ b/Engine/Plugins/FX/Niagara/Source/Niagara/Private/DataInterface/NiagaraDataInterfaceUtilities.cpp
@@ -12,23 +12,6 @@
 namespace FNiagaraDataInterfaceUtilities
 {
 
-bool DataInterfaceEqual(const UNiagaraDataInterface* Lhs, const UNiagaraDataInterface* Rhs)
-{
-	if ( !Lhs || !Rhs )
-	{
-		return false;
-	}
-	if ( Lhs == Rhs )
-	{
-		return true;
-	}
-	if ( Lhs->GetClass() != Rhs->GetClass() )
-	{
-		return false;
-	}
-	return Lhs->Equals(Rhs);
-}
-
 void ForEachVMFunctionEqualsImpl(class UNiagaraDataInterface* DataInterface, class UNiagaraSystem* NiagaraSystem, const FNiagaraParameterStore& ParameterStore, TFunction<bool(const FVMExternalFunctionBindingInfo&)> Action)
 {
 	if ( DataInterface == nullptr || NiagaraSystem == nullptr )
@@ -84,7 +67,7 @@ void ForEachVMFunctionEqualsImpl(class UNiagaraDataInterface* DataInterface, cla
 					const FNiagaraScriptDataInterfaceInfo& CachedDefaultDI = CachedDefaultDIs[iDataInterface];
 					if (CachedDefaultDI.DataInterface && (DataInterfaceEqual(CachedDefaultDI.DataInterface, DataInterface) || OverrideParameterNames.Contains(CachedDefaultDI.Name)))
 					{
-						if ( Action(FunctionBinding) == false )
+						if ( Action(Script, FunctionBinding) == false )
 						{
 							bContinueSearching = false;
 							return;
@@ -97,7 +80,7 @@ void ForEachVMFunctionEqualsImpl(class UNiagaraDataInterface* DataInterface, cla
 	);
 }
 
-void ForEachGpuFunctionEqualsImpl(class UNiagaraDataInterface* DataInterface, class UNiagaraSystem* NiagaraSystem, const FNiagaraParameterStore& ParameterStore, TFunction<bool(const FNiagaraDataInterfaceGeneratedFunction&)> Action)
+void ForEachGpuFunctionEqualsImpl(class UNiagaraDataInterface* DataInterface, class UNiagaraSystem* NiagaraSystem, const FNiagaraParameterStore& ParameterStore, TFunction<bool(const UNiagaraScript*, const FNiagaraDataInterfaceGeneratedFunction&)> Action)
 {
 	if (DataInterface == nullptr || NiagaraSystem == nullptr)
 	{
@@ -142,7 +125,7 @@ void ForEachGpuFunctionEqualsImpl(class UNiagaraDataInterface* DataInterface, cl
 					{
 						for (const FNiagaraDataInterfaceGeneratedFunction& GeneratedFunction : DataInterfaceParamInfos[iDataInterface].GeneratedFunctions)
 						{
-							if (Action(GeneratedFunction) == false)
+							if (Action(Script, GeneratedFunction) == false)
 							{
 								bContinueSearching = false;
 								return;
@@ -203,41 +186,41 @@ void ForEachDataInterface(FDataInterfaceUsageContext& UsageContext, const FNiaga
 
 } //namespace FNiagaraDataInterfaceUtilities
 
-void FNiagaraDataInterfaceUtilities::ForEachVMFunctionEquals(UNiagaraDataInterface* DataInterface, UNiagaraSystem* NiagaraSystem, TFunction<bool(const FVMExternalFunctionBindingInfo&)> Action)
+void FNiagaraDataInterfaceUtilities::ForEachVMFunctionEquals(UNiagaraDataInterface* DataInterface, UNiagaraSystem* NiagaraSystem, TFunction<bool(const UNiagaraScript*, const FVMExternalFunctionBindingInfo&)> Action)
 {
 	ForEachVMFunctionEqualsImpl(DataInterface, NiagaraSystem, NiagaraSystem->GetExposedParameters(), Action);
 }
 
-void FNiagaraDataInterfaceUtilities::ForEachVMFunctionEquals(UNiagaraDataInterface* DataInterface, UNiagaraSystem* NiagaraSystem, UNiagaraComponent* Component, TFunction<bool(const FVMExternalFunctionBindingInfo&)> Action)
+void FNiagaraDataInterfaceUtilities::ForEachVMFunctionEquals(UNiagaraDataInterface* DataInterface, UNiagaraSystem* NiagaraSystem, UNiagaraComponent* Component, TFunction<bool(const UNiagaraScript*, const FVMExternalFunctionBindingInfo&)> Action)
 {
 	const FNiagaraParameterStore& Parameters = Component ? Component->GetOverrideParameters() : NiagaraSystem->GetExposedParameters();
 	ForEachVMFunctionEqualsImpl(DataInterface, NiagaraSystem, Parameters, Action);
 }
 
-void FNiagaraDataInterfaceUtilities::ForEachVMFunctionEquals(UNiagaraDataInterface* DataInterface, UNiagaraSystem* NiagaraSystem, FNiagaraSystemInstance* SystemInstance, TFunction<bool(const FVMExternalFunctionBindingInfo&)> Action)
+void FNiagaraDataInterfaceUtilities::ForEachVMFunctionEquals(UNiagaraDataInterface* DataInterface, UNiagaraSystem* NiagaraSystem, FNiagaraSystemInstance* SystemInstance, TFunction<bool(const UNiagaraScript*, const FVMExternalFunctionBindingInfo&)> Action)
 {
 	const FNiagaraParameterStore& Parameters = (SystemInstance && SystemInstance->GetOverrideParameters()) ? *SystemInstance->GetOverrideParameters() : NiagaraSystem->GetExposedParameters();
 	ForEachVMFunctionEqualsImpl(DataInterface, NiagaraSystem, Parameters, Action);
 }
 
-void FNiagaraDataInterfaceUtilities::ForEachGpuFunctionEquals(class UNiagaraDataInterface* DataInterface, class UNiagaraSystem* NiagaraSystem, TFunction<bool(const FNiagaraDataInterfaceGeneratedFunction&)> Action)
+void FNiagaraDataInterfaceUtilities::ForEachGpuFunctionEquals(class UNiagaraDataInterface* DataInterface, class UNiagaraSystem* NiagaraSystem, TFunction<bool(const UNiagaraScript*, const FNiagaraDataInterfaceGeneratedFunction&)> Action)
 {
 	ForEachGpuFunctionEqualsImpl(DataInterface, NiagaraSystem, NiagaraSystem->GetExposedParameters(), Action);
 }
 
-void FNiagaraDataInterfaceUtilities::ForEachGpuFunctionEquals(class UNiagaraDataInterface* DataInterface, class UNiagaraSystem* NiagaraSystem, UNiagaraComponent* Component, TFunction<bool(const FNiagaraDataInterfaceGeneratedFunction&)> Action)
+void FNiagaraDataInterfaceUtilities::ForEachGpuFunctionEquals(class UNiagaraDataInterface* DataInterface, class UNiagaraSystem* NiagaraSystem, UNiagaraComponent* Component, TFunction<bool(const UNiagaraScript*, const FNiagaraDataInterfaceGeneratedFunction&)> Action)
 {
 	const FNiagaraParameterStore& Parameters = Component ? Component->GetOverrideParameters() : NiagaraSystem->GetExposedParameters();
 	ForEachGpuFunctionEqualsImpl(DataInterface, NiagaraSystem, Parameters, Action);
 }
 
-void FNiagaraDataInterfaceUtilities::ForEachGpuFunctionEquals(class UNiagaraDataInterface* DataInterface, class UNiagaraSystem* NiagaraSystem, FNiagaraSystemInstance* SystemInstance, TFunction<bool(const FNiagaraDataInterfaceGeneratedFunction&)> Action)
+void FNiagaraDataInterfaceUtilities::ForEachGpuFunctionEquals(class UNiagaraDataInterface* DataInterface, class UNiagaraSystem* NiagaraSystem, FNiagaraSystemInstance* SystemInstance, TFunction<bool(const UNiagaraScript*, const FNiagaraDataInterfaceGeneratedFunction&)> Action)
 {
 	const FNiagaraParameterStore& Parameters = (SystemInstance && SystemInstance->GetOverrideParameters()) ? *SystemInstance->GetOverrideParameters() : NiagaraSystem->GetExposedParameters();
 	ForEachGpuFunctionEqualsImpl(DataInterface, NiagaraSystem, Parameters, Action);
 }
 
-void FNiagaraDataInterfaceUtilities::ForEachVMFunction(class UNiagaraDataInterface* DataInterface, class FNiagaraSystemInstance* SystemInstance, TFunction<bool(const FVMExternalFunctionBindingInfo&)> Action)
+void FNiagaraDataInterfaceUtilities::ForEachVMFunction(class UNiagaraDataInterface* DataInterface, class FNiagaraSystemInstance* SystemInstance, TFunction<bool(const UNiagaraScript*, const FVMExternalFunctionBindingInfo&)> Action)
 {
 	if ( !DataInterface || !SystemInstance )
 	{
@@ -245,7 +228,7 @@ void FNiagaraDataInterfaceUtilities::ForEachVMFunction(class UNiagaraDataInterfa
 	}
 
 	auto SearchVMContext =
-		[](const FNiagaraScriptExecutionContextBase* ExecContext, UNiagaraDataInterface* DataInterface, const TFunction<bool(const FVMExternalFunctionBindingInfo&)>& Action) -> bool
+		[](const FNiagaraScriptExecutionContextBase* ExecContext, UNiagaraDataInterface* DataInterface, const TFunction<bool(const UNiagaraScript* Script, const FVMExternalFunctionBindingInfo&)>& Action) -> bool
 		{
 			if (ExecContext == nullptr || ExecContext->Script == nullptr)
 			{
@@ -274,7 +257,7 @@ void FNiagaraDataInterfaceUtilities::ForEachVMFunction(class UNiagaraDataInterfa
 						continue;
 					}
 
-					if (Action(FunctionBinding) == false)
+					if (Action(ExecContext->Script, FunctionBinding) == false)
 					{
 						return false;
 					}
diff --git a/Engine/Plugins/FX/Niagara/Source/Niagara/Public/NiagaraDataInterfaceUtilities.h b/Engine/Plugins/FX/Niagara/Source/Niagara/Public/NiagaraDataInterfaceUtilities.h
--- a/Engine/Plugins/FX/Niagara/Source/Niagara/Public/NiagaraDataInterfaceUtilities.h
+++ b/Engine/Plugins/FX/Niagara/Source/Niagara/Public/NiagaraDataInterfaceUtilities.h
@@ -21,21 +21,21 @@ namespace FNiagaraDataInterfaceUtilities
 	// Finds all VM function calls made using the data interface that equals this one (i.e. A->Equals(B))
 	// The same function call made be made multiple times across different scripts so you may see the same function multiple times
 	// The action should return True to continue iteration or False to stop
-	NIAGARA_API void ForEachVMFunctionEquals(UNiagaraDataInterface* DataInterface, UNiagaraSystem* NiagaraSystem, TFunction<bool(const FVMExternalFunctionBindingInfo&)> Action);
-	NIAGARA_API void ForEachVMFunctionEquals(UNiagaraDataInterface* DataInterface, UNiagaraSystem* NiagaraSystem, UNiagaraComponent* Component, TFunction<bool(const FVMExternalFunctionBindingInfo&)> Action);
-	NIAGARA_API void ForEachVMFunctionEquals(UNiagaraDataInterface* DataInterface, UNiagaraSystem* NiagaraSystem, FNiagaraSystemInstance* SystemInstance, TFunction<bool(const FVMExternalFunctionBindingInfo&)> Action);
+	NIAGARA_API void ForEachVMFunctionEquals(UNiagaraDataInterface* DataInterface, UNiagaraSystem* NiagaraSystem, TFunction<bool(const UNiagaraScript*, const FVMExternalFunctionBindingInfo&)> Action);
+	NIAGARA_API void ForEachVMFunctionEquals(UNiagaraDataInterface* DataInterface, UNiagaraSystem* NiagaraSystem, UNiagaraComponent* Component, TFunction<bool(const UNiagaraScript*, const FVMExternalFunctionBindingInfo&)> Action);
+	NIAGARA_API void ForEachVMFunctionEquals(UNiagaraDataInterface* DataInterface, UNiagaraSystem* NiagaraSystem, FNiagaraSystemInstance* SystemInstance, TFunction<bool(const UNiagaraScript*, const FVMExternalFunctionBindingInfo&)> Action);
 
 	// Finds all Gpu function calls made using the data interface that equals this one (i.e. A->Equals(B))
 	// The same function call made be made multiple times across different scripts so you may see the same function multiple times
 	// The action should return True to continue iteration or False to stop
-	NIAGARA_API void ForEachGpuFunctionEquals(UNiagaraDataInterface* DataInterface, UNiagaraSystem* NiagaraSystem, TFunction<bool(const FNiagaraDataInterfaceGeneratedFunction&)> Action);
-	NIAGARA_API void ForEachGpuFunctionEquals(UNiagaraDataInterface* DataInterface, UNiagaraSystem* NiagaraSystem, UNiagaraComponent* Component, TFunction<bool(const FNiagaraDataInterfaceGeneratedFunction&)> Action);
-	NIAGARA_API void ForEachGpuFunctionEquals(UNiagaraDataInterface* DataInterface, UNiagaraSystem* NiagaraSystem, FNiagaraSystemInstance* SystemInstance, TFunction<bool(const FNiagaraDataInterfaceGeneratedFunction&)> Action);
+	NIAGARA_API void ForEachGpuFunctionEquals(UNiagaraDataInterface* DataInterface, UNiagaraSystem* NiagaraSystem, TFunction<bool(const UNiagaraScript*, const FNiagaraDataInterfaceGeneratedFunction&)> Action);
+	NIAGARA_API void ForEachGpuFunctionEquals(UNiagaraDataInterface* DataInterface, UNiagaraSystem* NiagaraSystem, UNiagaraComponent* Component, TFunction<bool(const UNiagaraScript*, const FNiagaraDataInterfaceGeneratedFunction&)> Action);
+	NIAGARA_API void ForEachGpuFunctionEquals(UNiagaraDataInterface* DataInterface, UNiagaraSystem* NiagaraSystem, FNiagaraSystemInstance* SystemInstance, TFunction<bool(const UNiagaraScript*, const FNiagaraDataInterfaceGeneratedFunction&)> Action);
 
 	// Finds all VM function calls made using the data interface (i.e. pointer comparison A == B)
 	// The same function call made be made multiple times across different scripts so you may see the same function multiple times
 	// The action should return True to continue iteration or False to stop
-	NIAGARA_API void ForEachVMFunction(UNiagaraDataInterface* DataInterface, FNiagaraSystemInstance* SystemInstance, TFunction<bool(const FVMExternalFunctionBindingInfo&)> Action);
+	NIAGARA_API void ForEachVMFunction(UNiagaraDataInterface* DataInterface, FNiagaraSystemInstance* SystemInstance, TFunction<bool(const UNiagaraScript*, const FVMExternalFunctionBindingInfo&)> Action);
 	// Finds all Gpu function calls made using the data interface (i.e. pointer comparison A == B)
 	// The same function call made be made multiple times across different scripts so you may see the same function multiple times
 	// The action should return True to continue iteration or False to stop
