Fixed merge error.

diff --git a/Engine/Plugins/FX/Niagara/Source/Niagara/Private/DataInterface/NiagaraDataInterfaceUtilities.cpp b/Engine/Plugins/FX/Niagara/Source/Niagara/Private/DataInterface/NiagaraDataInterfaceUtilities.cpp
--- a/Engine/Plugins/FX/Niagara/Source/Niagara/Private/DataInterface/NiagaraDataInterfaceUtilities.cpp
+++ b/Engine/Plugins/FX/Niagara/Source/Niagara/Private/DataInterface/NiagaraDataInterfaceUtilities.cpp
@@ -12,7 +12,24 @@
 namespace FNiagaraDataInterfaceUtilities
 {
 
-void ForEachVMFunctionEqualsImpl(class UNiagaraDataInterface* DataInterface, class UNiagaraSystem* NiagaraSystem, const FNiagaraParameterStore& ParameterStore, TFunction<bool(const FVMExternalFunctionBindingInfo&)> Action)
+bool DataInterfaceEqual(const UNiagaraDataInterface* Lhs, const UNiagaraDataInterface* Rhs)
+{
+	if ( !Lhs || !Rhs )
+	{
+		return false;
+	}
+	if ( Lhs == Rhs )
+	{
+		return true;
+	}
+	if ( Lhs->GetClass() != Rhs->GetClass() )
+	{
+		return false;
+	}
+	return Lhs->Equals(Rhs);
+}
+
+void ForEachVMFunctionEqualsImpl(class UNiagaraDataInterface* DataInterface, class UNiagaraSystem* NiagaraSystem, const FNiagaraParameterStore& ParameterStore, TFunction<bool(const UNiagaraScript* Script, const FVMExternalFunctionBindingInfo&)> Action)
 {
 	if ( DataInterface == nullptr || NiagaraSystem == nullptr )
 	{
@@ -474,4 +491,4 @@ void FNiagaraDataInterfaceUtilities::ForEachDataInterface(UNiagaraSystem* Niagar
 			}
 		);
 	}
-}
+}
\ No newline at end of file
