Update CitySampleSensorGrid to use both RDG based shader parameter bindings and also RDG for it's internal shader passes

diff --git a/Engine/Plugins/FX/Niagara/Source/Niagara/Classes/NiagaraDataInterface.h b/Engine/Plugins/FX/Niagara/Source/Niagara/Classes/NiagaraDataInterface.h
--- a/Engine/Plugins/FX/Niagara/Source/Niagara/Classes/NiagaraDataInterface.h
+++ b/Engine/Plugins/FX/Niagara/Source/Niagara/Classes/NiagaraDataInterface.h
@@ -340,7 +340,7 @@ struct FNiagaraDataInterfaceProxy
 
 	// New data interface path
 	virtual void ResetData(const FNDIGpuComputeResetContext& Context) { }
-	virtual void PreStage(const FNDIGpuComputePostStageContext& Context) {}
+	virtual void PreStage(const FNDIGpuComputePreStageContext& Context) {}
 	virtual void PostStage(const FNDIGpuComputePostStageContext& Context) {}
 	virtual void PostSimulate(const FNDIGpuComputePostSimulateContext& Context) {}
 
