- Add guards for overflow or underflowing thread group counts

#rnx

diff --git a/Engine/Plugins/FX/Niagara/Source/Niagara/Private/NiagaraGpuComputeDispatch.cpp b/Engine/Plugins/FX/Niagara/Source/Niagara/Private/NiagaraGpuComputeDispatch.cpp
--- a/Engine/Plugins/FX/Niagara/Source/Niagara/Private/NiagaraGpuComputeDispatch.cpp
+++ b/Engine/Plugins/FX/Niagara/Source/Niagara/Private/NiagaraGpuComputeDispatch.cpp
@@ -1490,18 +1490,50 @@ void FNiagaraGpuComputeDispatch::DispatchStage(FRDGBuilder& GraphBuilder, const
 	}
 
 	const int32 TotalDispatchCount = DispatchCount.X * DispatchCount.Y * DispatchCount.Z;
-	if (TotalDispatchCount == 0)
+	if (TotalDispatchCount <= 0)
 	{
 		return;
 	}
 
 	checkf(DispatchNumThreads.X * DispatchNumThreads.Y * DispatchNumThreads.Z > 0, TEXT("DispatchNumThreads(%d, %d, %d) is invalid"), DispatchNumThreads.X, DispatchNumThreads.Y, DispatchNumThreads.Z);
 
+	// Calculate thread groups
+	// Note: In the OneD case we can use the Y dimension to get higher particle counts
+	if (DispatchType == ENiagaraGpuDispatchType::OneD)
+	{
+		const int32 GroupCount = FMath::DivideAndRoundUp(DispatchCount.X, DispatchNumThreads.X);
+		if (GroupCount > GRHIMaxDispatchThreadGroupsPerDimension.X)
+		{
+			DispatchCount.Y = FMath::DivideAndRoundUp(GroupCount, GRHIMaxDispatchThreadGroupsPerDimension.X);
+			DispatchCount.X = FMath::DivideAndRoundUp(GroupCount, DispatchCount.Y) * DispatchNumThreads.X;
+			ensure(DispatchCount.Y <= GRHIMaxDispatchThreadGroupsPerDimension.Y);
+		}
+	}
+
+	const FIntVector ThreadGroupCount(
+		FMath::DivideAndRoundUp(DispatchCount.X, DispatchNumThreads.X),
+		FMath::DivideAndRoundUp(DispatchCount.Y, DispatchNumThreads.Y),
+		FMath::DivideAndRoundUp(DispatchCount.Z, DispatchNumThreads.Z)
+	);
+
+	// Validate we don't overflow ThreadGroupCounts
+	if ((ThreadGroupCount.X < 0) || (ThreadGroupCount.X > GRHIMaxDispatchThreadGroupsPerDimension.X) ||
+		(ThreadGroupCount.Y < 0) || (ThreadGroupCount.Y > GRHIMaxDispatchThreadGroupsPerDimension.Y) ||
+		(ThreadGroupCount.Z < 0) || (ThreadGroupCount.Z > GRHIMaxDispatchThreadGroupsPerDimension.Z) )
+	{
+		UE_LOG(LogNiagara, Warning, TEXT("FNiagaraGpuComputeDispatch: Invalid ThreadGroupdCount(%d, %d, %d) for DispatchCount(%d, %d, %d) Stage (%s)"),
+			ThreadGroupCount.X, ThreadGroupCount.Y, ThreadGroupCount.Z,
+			DispatchCount.X, DispatchCount.Y, DispatchCount.Z,
+			*SimStageData.StageMetaData->SimulationStageName.ToString()
+		);
+		return;
+	}
+
+	// Set Parameters
 	const FNiagaraShaderScriptParametersMetadata& NiagaraShaderParametersMetadata = InstanceData.Context->GPUScript_RT->GetScriptParametersMetadata_RT();
 	const FShaderParametersMetadata* ShaderParametersMetadata = NiagaraShaderParametersMetadata.ShaderParametersMetadata.Get();
 	FNiagaraShader::FParameters* DispatchParameters = GraphBuilder.AllocParameters<FNiagaraShader::FParameters>(ShaderParametersMetadata);
 
-	// Set Parameters
 	const bool bRequiresPersistentIDs = InstanceData.Context->MainDataSet->RequiresPersistentIDs();
 
 	DispatchParameters->SimStart = InstanceData.bResetData ? 1U : 0U;
@@ -1623,23 +1655,6 @@ void FNiagaraGpuComputeDispatch::DispatchStage(FRDGBuilder& GraphBuilder, const
 
 	// Execute the dispatch
 	{
-		// In the OneD case we can use the Y dimension to get higher particle counts
-		if (DispatchType == ENiagaraGpuDispatchType::OneD)
-		{
-			const int32 GroupCount = FMath::DivideAndRoundUp(DispatchCount.X, DispatchNumThreads.X);
-			if (GroupCount > GRHIMaxDispatchThreadGroupsPerDimension.X)
-			{
-				DispatchCount.Y = FMath::DivideAndRoundUp(GroupCount, GRHIMaxDispatchThreadGroupsPerDimension.X);
-				DispatchCount.X = FMath::DivideAndRoundUp(GroupCount, DispatchCount.Y) * DispatchNumThreads.X;
-				ensure(DispatchCount.Y <= GRHIMaxDispatchThreadGroupsPerDimension.Y);
-			}
-		}
-
-		FIntVector ThreadGroupCount;
-		ThreadGroupCount.X = FMath::DivideAndRoundUp(DispatchCount.X, DispatchNumThreads.X);
-		ThreadGroupCount.Y = FMath::DivideAndRoundUp(DispatchCount.Y, DispatchNumThreads.Y);
-		ThreadGroupCount.Z = FMath::DivideAndRoundUp(DispatchCount.Z, DispatchNumThreads.Z);
-
 		DispatchParameters->DispatchThreadIdToLinear	= FUintVector3(1, DispatchCount.X, DispatchCount.X * DispatchCount.Y);
 		DispatchParameters->DispatchThreadIdBounds		= FUintVector3(DispatchCount.X, DispatchCount.Y, DispatchCount.Z);
 
