Add insight trace when niagara is waiting on compilation results

#rnx

diff --git a/Engine/Plugins/FX/Niagara/Source/NiagaraEditor/Private/NiagaraCompiler.cpp b/Engine/Plugins/FX/Niagara/Source/NiagaraEditor/Private/NiagaraCompiler.cpp
--- a/Engine/Plugins/FX/Niagara/Source/NiagaraEditor/Private/NiagaraCompiler.cpp
+++ b/Engine/Plugins/FX/Niagara/Source/NiagaraEditor/Private/NiagaraCompiler.cpp
@@ -1598,6 +1598,8 @@ TSharedPtr<FNiagaraGraphCachedDataBase, ESPMode::ThreadSafe> FNiagaraEditorModul
 
 TSharedPtr<FNiagaraVMExecutableData> FNiagaraEditorModule::GetCompilationResult(int32 JobID, bool bWait)
 {
+	TRACE_CPUPROFILER_EVENT_SCOPE(FNiagaraEditorModule::GetCompilationResult);
+
 	TSharedPtr<FHlslNiagaraCompiler>* MapEntry = ActiveCompilations.Find(JobID);
 	check(MapEntry && *MapEntry);
 
