AssetDefinition - Upgrading the asset type actions for Niagara so that we no longer need to load any of those asset types on right click.  Needed to add VersioningEnabled to the NiagaraEmitter asset tag data.  The current way the code is setup it will offer the option if there's no tag data yet, and if it turns out that version wasn't enabled for the emitter it will just skip the action and the next time you right click, since the asset will have been loaded the option will be grayed out.  Future savings of the files will make the option grayed out from the start.

diff --git a/Engine/Plugins/Editor/EngineAssetDefinitions/Source/Private/AssetDefinition_StaticMesh.cpp b/Engine/Plugins/Editor/EngineAssetDefinitions/Source/Private/AssetDefinition_StaticMesh.cpp
--- a/Engine/Plugins/Editor/EngineAssetDefinitions/Source/Private/AssetDefinition_StaticMesh.cpp
+++ b/Engine/Plugins/Editor/EngineAssetDefinitions/Source/Private/AssetDefinition_StaticMesh.cpp
@@ -521,14 +521,18 @@ namespace MenuExtension_StaticMesh
 						const TAttribute<FText> Label = LOCTEXT("ObjectContext_SaveGeneratedLODsInPackage", "Save Generated LODs");
 						const TAttribute<FText> ToolTip = LOCTEXT("ObjectContext_SaveGeneratedLODsInPackageTooltip", "Run the mesh reduce and save the generated LODs as part of the package.");
 						const FSlateIcon Icon = FSlateIcon();
-						const FToolMenuExecuteAction UIAction = FToolMenuExecuteAction::CreateStatic(&ExecuteSaveGeneratedLODsInPackage);
+
+						FToolUIAction UIAction;
+						UIAction.ExecuteAction = FToolMenuExecuteAction::CreateStatic(&ExecuteSaveGeneratedLODsInPackage);
 						InSection.AddMenuEntry("ObjectContext_SaveGeneratedLODsInPackage", Label, ToolTip, Icon, UIAction);
 					}
 					{
 						const TAttribute<FText> Label = LOCTEXT("ObjectContext_ClearVertexColors", "Remove Vertex Colors");
 						const TAttribute<FText> ToolTip = LOCTEXT("ObjectContext_ClearVertexColorsTooltip", "Removes vertex colors from all LODS in all selected meshes.");
 						const FSlateIcon Icon = FSlateIcon(FAppStyle::GetAppStyleSetName(), "ContentBrowser.AssetActions.RemoveVertexColors");
-						const FToolMenuExecuteAction UIAction = FToolMenuExecuteAction::CreateStatic(&ExecuteRemoveVertexColors);
+
+						FToolUIAction UIAction;
+						UIAction.ExecuteAction = FToolMenuExecuteAction::CreateStatic(&ExecuteRemoveVertexColors);
 						InSection.AddMenuEntry("ObjectContext_ClearVertexColors", Label, ToolTip, Icon, UIAction);
 					}
 					{
diff --git a/Engine/Plugins/Editor/EngineAssetDefinitions/Source/Private/AssetDefinition_Texture.cpp b/Engine/Plugins/Editor/EngineAssetDefinitions/Source/Private/AssetDefinition_Texture.cpp
--- a/Engine/Plugins/Editor/EngineAssetDefinitions/Source/Private/AssetDefinition_Texture.cpp
+++ b/Engine/Plugins/Editor/EngineAssetDefinitions/Source/Private/AssetDefinition_Texture.cpp
@@ -45,52 +45,14 @@ namespace MenuExtension_Texture
 	{
 		const UContentBrowserAssetContextMenuContext* CBContext = UContentBrowserAssetContextMenuContext::FindContextWithAssets(InContext);
 
-		const FString DefaultSuffix = TEXT("_Mat");
-
-		if ( const FAssetData* TextureAsset = CBContext->GetSingleSelectedAssetOfType(UTexture::StaticClass()) )
-        {
-			if (UTexture* Texture = CastChecked<UTexture>(TextureAsset->GetAsset()))
+		IAssetTools::Get().CreateAssetsFrom<UTexture>(
+			CBContext->LoadSelectedObjects<UTexture>(), UMaterial::StaticClass(), TEXT("_Mat"), [](UTexture* SourceObject)
 			{
-				// Determine an appropriate name
-				FString Name;
-				FString PackagePath;
-				IAssetTools::Get().CreateUniqueAssetName(Texture->GetOutermost()->GetName(), DefaultSuffix, PackagePath, Name);
-
-				// Create the factory used to generate the asset
 				UMaterialFactoryNew* Factory = NewObject<UMaterialFactoryNew>();
-				Factory->InitialTexture = Texture;
-
-				FContentBrowserModule& ContentBrowserModule = FModuleManager::LoadModuleChecked<FContentBrowserModule>("ContentBrowser");
-				ContentBrowserModule.Get().CreateNewAsset(Name, FPackageName::GetLongPackagePath(PackagePath), UMaterial::StaticClass(), Factory);
-			}
-		}
-		else
-		{
-			TArray<UObject*> ObjectsToSync;
-
-			for (UTexture* Texture : CBContext->LoadSelectedObjects<UTexture>())
-			{
-				FString Name;
-				FString PackageName;
-				IAssetTools::Get().CreateUniqueAssetName(Texture->GetOutermost()->GetName(), DefaultSuffix, PackageName, Name);
-
-				// Create the factory used to generate the asset
-				UMaterialFactoryNew* Factory = NewObject<UMaterialFactoryNew>();
-				Factory->InitialTexture = Texture;
-
-				UObject* NewAsset = IAssetTools::Get().CreateAsset(Name, FPackageName::GetLongPackagePath(PackageName), UMaterial::StaticClass(), Factory);
-
-				if ( NewAsset )
-				{
-					ObjectsToSync.Add(NewAsset);
-				}
-			}
-
-			if ( ObjectsToSync.Num() > 0 )
-			{
-				IAssetTools::Get().SyncBrowserToAssets(ObjectsToSync);
+				Factory->InitialTexture = SourceObject;
+				return Factory;
 			}
-		}
+		);
 	}
 
 	static void ExecuteConvertToVirtualTexture(const FToolMenuContext& InContext)
@@ -111,63 +73,6 @@ namespace MenuExtension_Texture
 		}
 	}
 
-	static void ExecuteCreateSubUVAnimation(TArray<TWeakObjectPtr<UTexture>> Objects)
-	{
-		const FString DefaultSuffix = TEXT("_SubUV");
-
-		if ( Objects.Num() == 1 )
-		{
-			UTexture2D* Object = Cast<UTexture2D>(Objects[0].Get());
-
-			if ( Object )
-			{
-				// Determine an appropriate name
-				FString Name;
-				FString PackagePath;
-				IAssetTools::Get().CreateUniqueAssetName(Object->GetOutermost()->GetName(), DefaultSuffix, PackagePath, Name);
-
-				// Create the factory used to generate the asset
-				USubUVAnimationFactory* Factory = NewObject<USubUVAnimationFactory>();
-				Factory->InitialTexture = Object;
-
-				FContentBrowserModule& ContentBrowserModule = FModuleManager::LoadModuleChecked<FContentBrowserModule>("ContentBrowser");
-				ContentBrowserModule.Get().CreateNewAsset(Name, FPackageName::GetLongPackagePath(PackagePath), USubUVAnimation::StaticClass(), Factory);
-			}
-		}
-		else
-		{
-			TArray<UObject*> ObjectsToSync;
-
-			for (auto ObjIt = Objects.CreateConstIterator(); ObjIt; ++ObjIt)
-			{
-				UTexture2D* Object = Cast<UTexture2D>((*ObjIt).Get());
-
-				if ( Object )
-				{
-					FString Name;
-					FString PackageName;
-					IAssetTools::Get().CreateUniqueAssetName(Object->GetOutermost()->GetName(), DefaultSuffix, PackageName, Name);
-
-					// Create the factory used to generate the asset
-					USubUVAnimationFactory* Factory = NewObject<USubUVAnimationFactory>();
-					Factory->InitialTexture = Object;
-					
-					UObject* NewAsset = IAssetTools::Get().CreateAsset(Name, FPackageName::GetLongPackagePath(PackageName), USubUVAnimation::StaticClass(), Factory);
-
-					if ( NewAsset )
-					{
-						ObjectsToSync.Add(NewAsset);
-					}
-				}
-			}
-
-			if ( ObjectsToSync.Num() > 0 )
-			{
-				IAssetTools::Get().SyncBrowserToAssets(ObjectsToSync);
-			}
-		}
-	}
-
 	static void ExecuteFindMaterials(const FToolMenuContext& InContext)
 	{
 		if (const UContentBrowserAssetContextMenuContext* CBContext = UContentBrowserAssetContextMenuContext::FindContextWithAssets(InContext))
diff --git a/Engine/Plugins/Experimental/Mutable/Source/CustomizableObjectPopulationEditor/Private/MuCOPE/AssetTypeActions_CustomizableObjectPopulation.cpp b/Engine/Plugins/Experimental/Mutable/Source/CustomizableObjectPopulationEditor/Private/MuCOPE/AssetTypeActions_CustomizableObjectPopulation.cpp
--- a/Engine/Plugins/Experimental/Mutable/Source/CustomizableObjectPopulationEditor/Private/MuCOPE/AssetTypeActions_CustomizableObjectPopulation.cpp
+++ b/Engine/Plugins/Experimental/Mutable/Source/CustomizableObjectPopulationEditor/Private/MuCOPE/AssetTypeActions_CustomizableObjectPopulation.cpp
@@ -18,26 +18,22 @@ FText FAssetTypeActions_CustomizableObjectPopulation::GetName() const
 	return LOCTEXT("AssetTypeActions_CustomizableObjectPopulation", "Customizable Population");
 }
 
-
 FColor FAssetTypeActions_CustomizableObjectPopulation::GetTypeColor() const
 { 
 	return FColor(0, 0, 0); 
 }
 
-
 UClass* FAssetTypeActions_CustomizableObjectPopulation::GetSupportedClass() const
 { 
 	return UCustomizableObjectPopulation::StaticClass(); 
 }
 
-
 uint32 FAssetTypeActions_CustomizableObjectPopulation::GetCategories()
 { 
 	const ICustomizableObjectEditorModule* CustomizableObjectEditorModule = &FModuleManager::LoadModuleChecked<ICustomizableObjectEditorModule>("CustomizableObjectEditor");
 	return CustomizableObjectEditorModule->GetAssetCategory();
 }
 
-
 void FAssetTypeActions_CustomizableObjectPopulation::GetActions(const TArray<UObject*>& InObjects, FMenuBuilder& MenuBuilder)
 {
 	TArray<TWeakObjectPtr<UCustomizableObjectPopulation>> Objects = GetTypedWeakObjectPtrs<UCustomizableObjectPopulation>(InObjects);
@@ -51,88 +47,8 @@ void FAssetTypeActions_CustomizableObjectPopulation::GetActions(const TArray<UOb
 			FCanExecuteAction()
 		)
 	);
-
-	//MenuBuilder.AddMenuEntry(
-	//	NSLOCTEXT("AssetTypeActions_CustomizableObjectPopulation", "CustomizableObjectPopulation_CreatePopulationGenerator", "Create Population Generator"),
-	//	NSLOCTEXT("AssetTypeActions_CustomizableObjectPopulation", "CustomizableObjectPopulation_CreatePopulationGeneratorToolTip", "Compiles the Population and Creates a Population Generator Asset"),
-	//	FSlateIcon(),
-	//	FUIAction(
-	//		FExecuteAction::CreateSP(this, &FAssetTypeActions_CustomizableObjectPopulation::CreatePopulationGeneratorAsset, Objects),
-	//		FCanExecuteAction()
-	//	)
-	//);
-
-	//MenuBuilder.AddMenuEntry(
-	//	NSLOCTEXT("AssetTypeActions_CustomizableObjectPopulation", "CustomizableObjectPopulation_RecompilePopulationGenerator", "Recompile Population Generators"),
-	//	NSLOCTEXT("AssetTypeActions_CustomizableObjectPopulation", "CustomizableObjectPopulation_RecompilePopulationGeneratorToolTip", "Recompiles all the Population Generator Assets"),
-	//	FSlateIcon(),
-	//	FUIAction(
-	//		FExecuteAction::CreateSP(this, &FAssetTypeActions_CustomizableObjectPopulation::RecompilePopulationGenerators, Objects),
-	//		FCanExecuteAction()
-	//	)
-	//);
 }
 
-//void FAssetTypeActions_CustomizableObjectPopulation::CreatePopulationGeneratorAsset(TArray<TWeakObjectPtr<UCustomizableObjectPopulation>> Objects)
-//{
-//	const FString DefaultSuffix = TEXT("_Generator");
-//
-//	TArray<UObject*> ObjectsToSync;
-//
-//	for (auto ObjIt = Objects.CreateConstIterator(); ObjIt; ++ObjIt)
-//	{
-//		auto Object = (*ObjIt).Get();
-//		if (Object)
-//		{
-//			if (!Object->IsValidPopulation())
-//			{
-//				UE_LOG(LogMutable, Warning, TEXT("There are one or more unassigned Population Classes. Please review your %s Population."), *(Object->Name));
-//				continue;
-//			}
-//
-//			// Determine an appropriate name
-//			FString Name;
-//			FString PackageName;
-//			CreateUniqueAssetName(Object->GetOutermost()->GetName(), DefaultSuffix, PackageName, Name);
-//			UPackage* Pkg = CreatePackage(*PackageName);
-//
-//			// Create avoid asset
-//			UCustomizableObjectPopulation* CustomizableObject = Cast<UCustomizableObjectPopulation>(Object);
-//			UCustomizableObjectPopulationGenerator* Generator = NewObject<UCustomizableObjectPopulationGenerator>(Pkg, FName(*Name), RF_Public | RF_Standalone);
-//
-//			// Fill the Generator information
-//			FCustomizableObjectPopulationGeneratorCompiler GeneratorCompiler;
-//			GeneratorCompiler.CompilePopulation(Object, true, Generator);
-//
-//			if (Generator)
-//			{
-//				// Mark the package dirty...
-//				Pkg->MarkPackageDirty();
-//
-//				ObjectsToSync.Add(Generator);
-//			}
-//		}
-//	}
-//
-//	if (ObjectsToSync.Num() > 0)
-//	{
-//		FContentBrowserModule& ContentBrowserModule = FModuleManager::Get().LoadModuleChecked<FContentBrowserModule>("ContentBrowser");
-//		ContentBrowserModule.Get().SyncBrowserToAssets(ObjectsToSync, /*bAllowLockedBrowsers=*/true);
-//	}
-//}
-//
-//void FAssetTypeActions_CustomizableObjectPopulation::RecompilePopulationGenerators(TArray<TWeakObjectPtr<UCustomizableObjectPopulation>> Populations)
-//{
-//	FCustomizableObjectPopulationGeneratorCompiler GeneratorCompiler;
-//	for (int32 i = 0; i < Populations.Num(); ++i)
-//	{
-//		if (Populations[i].IsValid() && Populations[i]->IsValidPopulation())
-//		{
-//			GeneratorCompiler.CompilePopulation(Populations[i].Get(), false, Populations[i]->Generator);
-//		}
-//	}
-//}
-
 void FAssetTypeActions_CustomizableObjectPopulation::OpenAssetEditor(const TArray<UObject*>& InObjects, TSharedPtr<IToolkitHost> EditWithinLevelEditor)
 {
 	EToolkitMode::Type Mode = EditWithinLevelEditor.IsValid() ? EToolkitMode::WorldCentric : EToolkitMode::Standalone;
diff --git a/Engine/Plugins/FX/Niagara/Source/Niagara/Private/NiagaraEmitter.cpp b/Engine/Plugins/FX/Niagara/Source/Niagara/Private/NiagaraEmitter.cpp
--- a/Engine/Plugins/FX/Niagara/Source/Niagara/Private/NiagaraEmitter.cpp
+++ b/Engine/Plugins/FX/Niagara/Source/Niagara/Private/NiagaraEmitter.cpp
@@ -909,6 +909,8 @@ bool UNiagaraEmitter::IsEditorOnly() const
 void UNiagaraEmitter::GetAssetRegistryTags(TArray<FAssetRegistryTag>& OutTags) const
 {
 #if WITH_EDITOR
+	OutTags.Add(FAssetRegistryTag("VersioningEnabled", bVersioningEnabled ? TEXT("True") : TEXT("False"), FAssetRegistryTag::TT_Alphabetical));
+
 	const FVersionedNiagaraEmitterData* AssetData = GetLatestEmitterData();
 	FVersionedNiagaraEmitterData DefaultData;
 	const FVersionedNiagaraEmitterData& EmitterData = AssetData ? *AssetData : DefaultData; // the CDO does not have any version data, so just use the default struct values in that case 
diff --git a/Engine/Plugins/FX/Niagara/Source/NiagaraEditor/NiagaraEditor.Build.cs b/Engine/Plugins/FX/Niagara/Source/NiagaraEditor/NiagaraEditor.Build.cs
--- a/Engine/Plugins/FX/Niagara/Source/NiagaraEditor/NiagaraEditor.Build.cs
+++ b/Engine/Plugins/FX/Niagara/Source/NiagaraEditor/NiagaraEditor.Build.cs
@@ -62,7 +62,8 @@ public class NiagaraEditor : ModuleRules
 				"PythonScriptPlugin",
 				"ImageWrapper",
 				"AssetDefinition",
-				"ContentBrowser"
+				"ContentBrowser",
+				"AssetTools"
 			}
 		);
 
@@ -81,7 +82,6 @@ public class NiagaraEditor : ModuleRules
 				"Engine",
 				"MessageLog",
 				"Messaging",
-				"AssetTools",
 				"ContentBrowser",
                 "DerivedDataCache",
                 "LevelEditor",
diff --git a/Engine/Plugins/FX/Niagara/Source/NiagaraEditor/Private/AssetTypeActions/AssetDefinition_NiagaraEffectType.cpp b/Engine/Plugins/FX/Niagara/Source/NiagaraEditor/Private/AssetTypeActions/AssetDefinition_NiagaraEffectType.cpp
new file mode 100644
--- /dev/null
+++ b/Engine/Plugins/FX/Niagara/Source/NiagaraEditor/Private/AssetTypeActions/AssetDefinition_NiagaraEffectType.cpp
@@ -0,0 +1,12 @@
+// Copyright Epic Games, Inc. All Rights Reserved.
+#include "AssetDefinition_NiagaraEffectType.h"
+#include "NiagaraEditorStyle.h"
+
+#define LOCTEXT_NAMESPACE "AssetTypeActions_NiagaraEffectType"
+
+FLinearColor UAssetDefinition_NiagaraEffectType::GetAssetColor() const
+{
+	return FNiagaraEditorStyle::Get().GetColor("NiagaraEditor.AssetColors.EffectType");
+}
+
+#undef LOCTEXT_NAMESPACE
diff --git a/Engine/Plugins/FX/Niagara/Source/NiagaraEditor/Private/AssetTypeActions/AssetDefinition_NiagaraEffectType.h b/Engine/Plugins/FX/Niagara/Source/NiagaraEditor/Private/AssetTypeActions/AssetDefinition_NiagaraEffectType.h
new file mode 100644
--- /dev/null
+++ b/Engine/Plugins/FX/Niagara/Source/NiagaraEditor/Private/AssetTypeActions/AssetDefinition_NiagaraEffectType.h
@@ -0,0 +1,30 @@
+// Copyright Epic Games, Inc. All Rights Reserved.
+
+#pragma once
+
+#include "CoreMinimal.h"
+#include "NiagaraEffectType.h"
+#include "AssetDefinitionDefault.h"
+
+#include "AssetDefinition_NiagaraEffectType.generated.h"
+
+UCLASS()
+class UAssetDefinition_NiagaraEffectType : public UAssetDefinitionDefault
+{
+	GENERATED_BODY()
+
+public:
+	// UAssetDefinition Begin
+	virtual FText GetAssetDisplayName() const override { return NSLOCTEXT("AssetTypeActions", "AssetTypeActions_NiagaraEffectType", "Niagara Effect Type"); }
+	virtual FLinearColor GetAssetColor() const override;
+	virtual TSoftClassPtr<UObject> GetAssetClass() const override { return UNiagaraEffectType::StaticClass(); }
+	virtual TConstArrayView<FAssetCategoryPath> GetAssetCategories() const override
+	{
+		static const auto Categories = {
+			EAssetCategoryPaths::FX / NSLOCTEXT("Niagara", "NiagaraAssetSubMenu_Advanced", "Advanced")
+		};
+		
+		return Categories;
+	}
+	// UAssetDefinition End
+};
diff --git a/Engine/Plugins/FX/Niagara/Source/NiagaraEditor/Private/AssetTypeActions/AssetDefinition_NiagaraEmitter.cpp b/Engine/Plugins/FX/Niagara/Source/NiagaraEditor/Private/AssetTypeActions/AssetDefinition_NiagaraEmitter.cpp
new file mode 100644
--- /dev/null
+++ b/Engine/Plugins/FX/Niagara/Source/NiagaraEditor/Private/AssetTypeActions/AssetDefinition_NiagaraEmitter.cpp
@@ -0,0 +1,199 @@
+// Copyright Epic Games, Inc. All Rights Reserved.
+
+#include "AssetDefinition_NiagaraEmitter.h"
+
+#include "AssetToolsModule.h"
+#include "ContentBrowserMenuContexts.h"
+#include "NiagaraEmitter.h"
+#include "NiagaraScriptSourceBase.h"
+#include "NiagaraSystem.h"
+#include "NiagaraSystemToolkit.h"
+#include "NiagaraEditorStyle.h"
+#include "NiagaraSystemFactoryNew.h"
+#include "ViewModels/NiagaraSystemViewModel.h"
+#include "ToolMenus.h"
+#include "Styling/AppStyle.h"
+#include "Toolkits/AssetEditorToolkit.h"
+#include "IContentBrowserSingleton.h"
+#include "ContentBrowserModule.h"
+#include "NiagaraEditorUtilities.h"
+#include "Modules/ModuleManager.h"
+#include "Subsystems/AssetEditorSubsystem.h"
+
+#define LOCTEXT_NAMESPACE "UAssetDefinition_NiagaraEmitter"
+
+FLinearColor UAssetDefinition_NiagaraEmitter::GetAssetColor() const
+{
+	return FNiagaraEditorStyle::Get().GetColor("NiagaraEditor.AssetColors.Emitter");
+}
+
+EAssetCommandResult UAssetDefinition_NiagaraEmitter::OpenAssets(const FAssetOpenArgs& OpenArgs) const
+{
+	for (UNiagaraEmitter* Emitter : OpenArgs.LoadObjects<UNiagaraEmitter>())
+	{
+		if (!Emitter->VersionToOpenInEditor.IsValid())
+		{
+			Emitter->VersionToOpenInEditor = Emitter->GetExposedVersion().VersionGuid;
+		}
+		
+		TSharedRef<FNiagaraSystemToolkit> SystemToolkit(new FNiagaraSystemToolkit());
+		SystemToolkit->InitializeWithEmitter(OpenArgs.GetToolkitMode(), OpenArgs.ToolkitHost, *Emitter);
+	}
+	
+	return EAssetCommandResult::Handled;
+}
+
+// Menu Extensions
+//--------------------------------------------------------------------
+
+namespace MenuExtension_NiagaraEmitter
+{
+	void ExecuteNewNiagaraSystem(const FToolMenuContext& InContext)
+	{
+		const FString DefaultSuffix = TEXT("_System");
+
+		FNiagaraSystemViewModelOptions SystemOptions;
+		SystemOptions.bCanModifyEmittersFromTimeline = true;
+		SystemOptions.EditMode = ENiagaraSystemViewModelEditMode::SystemAsset;
+		
+		FAssetToolsModule& AssetToolsModule = FModuleManager::GetModuleChecked<FAssetToolsModule>("AssetTools");
+
+		TArray<UObject*> ObjectsToSync;
+		
+		const UContentBrowserAssetContextMenuContext* CBContext = UContentBrowserAssetContextMenuContext::FindContextWithAssets(InContext);
+		for (UNiagaraEmitter* Emitter : CBContext->LoadSelectedObjects<UNiagaraEmitter>())
+		{
+			// Determine an appropriate names
+			FString Name;
+			FString PackageName;
+			AssetToolsModule.Get().CreateUniqueAssetName(Emitter->GetOutermost()->GetName(), DefaultSuffix, PackageName, Name);
+
+			// Create the factory used to generate the asset
+			UNiagaraSystemFactoryNew* Factory = NewObject<UNiagaraSystemFactoryNew>();
+			Factory->EmittersToAddToNewSystem.Add(FVersionedNiagaraEmitter(Emitter, Emitter->GetExposedVersion().VersionGuid));
+			UObject* NewAsset = AssetToolsModule.Get().CreateAsset(Name, FPackageName::GetLongPackagePath(PackageName), UNiagaraSystem::StaticClass(), Factory);
+				
+			UNiagaraSystem* System = Cast<UNiagaraSystem>(NewAsset);
+			if (System != nullptr)
+			{
+				ObjectsToSync.Add(NewAsset);
+			}
+		}
+
+		if (ObjectsToSync.Num() > 0)
+		{
+			FContentBrowserModule& ContentBrowserModule = FModuleManager::LoadModuleChecked<FContentBrowserModule>("ContentBrowser");
+			ContentBrowserModule.Get().SyncBrowserToAssets(ObjectsToSync);
+		}
+	}
+
+	void ExecuteCreateDuplicateParent(const FToolMenuContext& InContext)
+	{
+		const FString DefaultSuffix = TEXT("_Parent");
+
+		TArray<UObject*> ObjectsToSync;
+		FAssetToolsModule& AssetToolsModule = FModuleManager::GetModuleChecked<FAssetToolsModule>("AssetTools");
+
+		const UContentBrowserAssetContextMenuContext* CBContext = UContentBrowserAssetContextMenuContext::FindContextWithAssets(InContext);
+		for (UNiagaraEmitter* Emitter : CBContext->LoadSelectedObjects<UNiagaraEmitter>())
+		{
+			if (Emitter->IsVersioningEnabled())
+			{
+				FString Name;
+				FString PackageName;
+				AssetToolsModule.Get().CreateUniqueAssetName(Emitter->GetOutermost()->GetName(), DefaultSuffix, PackageName, Name);
+
+				// Create the factory used to generate the asset
+				UNiagaraEmitter* NewEmitter = Cast<UNiagaraEmitter>(AssetToolsModule.Get().DuplicateAssetWithDialog(Name, FPackageName::GetLongPackagePath(PackageName), Emitter));
+
+				if (NewEmitter != nullptr)
+				{
+					GEditor->GetEditorSubsystem<UAssetEditorSubsystem>()->CloseAllEditorsForAsset(Emitter);
+
+					NewEmitter->Modify();
+					NewEmitter->GetLatestEmitterData()->GraphSource->MarkNotSynchronized(TEXT("Emitter created"));
+
+					Emitter->Modify();
+					Emitter->SetParent(FVersionedNiagaraEmitter(NewEmitter, NewEmitter->GetExposedVersion().VersionGuid));
+					Emitter->GetLatestEmitterData()->GraphSource->MarkNotSynchronized(TEXT("Emitter parent changed"));
+
+					ObjectsToSync.Add(NewEmitter);
+				}
+			}
+		}
+
+		if (ObjectsToSync.Num() > 0)
+		{
+			AssetToolsModule.Get().OpenEditorForAssets(ObjectsToSync);
+
+			FContentBrowserModule& ContentBrowserModule = FModuleManager::LoadModuleChecked<FContentBrowserModule>("ContentBrowser");
+			ContentBrowserModule.Get().SyncBrowserToAssets(ObjectsToSync);
+		}
+	}
+	
+	bool CanExecuteCreateDuplicateParent(const FToolMenuContext& InContext)
+	{
+		const static FName NAME_VersioningEnabled("VersioningEnabled");
+	
+		const UContentBrowserAssetContextMenuContext* CBContext = UContentBrowserAssetContextMenuContext::FindContextWithAssets(InContext);
+        for (const FAssetData& SelectedAsset : CBContext->SelectedAssets)
+        {
+			bool VersioningEnabled = false;
+        	if (SelectedAsset.GetTagValue<bool>(NAME_VersioningEnabled, VersioningEnabled) && !VersioningEnabled)
+        	{
+        		return false;
+        	}
+        }
+ 
+		return true;
+	}
+
+	static void ExecuteMarkDependentCompilableAssetsDirty(const FToolMenuContext& InContext)
+	{
+		const UContentBrowserAssetContextMenuContext* CBContext = UContentBrowserAssetContextMenuContext::FindContextWithAssets(InContext);
+		FNiagaraEditorUtilities::MarkDependentCompilableAssetsDirty(CBContext->LoadSelectedObjects<UObject>());
+	}
+	
+	static FDelayedAutoRegisterHelper DelayedAutoRegister(EDelayedRegisterRunPhase::EndOfEngineInit, []{ 
+		UToolMenus::RegisterStartupCallback(FSimpleMulticastDelegate::FDelegate::CreateLambda([]()
+		{
+			FToolMenuOwnerScoped OwnerScoped(UE_MODULE_NAME);
+			UToolMenu* Menu = UE::ContentBrowser::ExtendToolMenu_AssetContextMenu(UNiagaraEmitter::StaticClass());
+		
+			FToolMenuSection& Section = Menu->FindOrAddSection("GetAssetActions");
+			Section.AddDynamicEntry(NAME_None, FNewToolMenuSectionDelegate::CreateLambda([](FToolMenuSection& InSection)
+			{
+				{
+					const TAttribute<FText> Label = LOCTEXT("Emitter_NewNiagaraSystem", "Create Niagara System");
+					const TAttribute<FText> ToolTip = LOCTEXT("Emitter_NewNiagaraSystemTooltip", "Creates a niagara system using this emitter as a base.");
+					const FSlateIcon Icon = FSlateIcon(FAppStyle::GetAppStyleSetName(), "ClassIcon.ParticleSystem");
+
+					FToolUIAction UIAction;
+					UIAction.ExecuteAction = FToolMenuExecuteAction::CreateStatic(&ExecuteNewNiagaraSystem);
+					InSection.AddMenuEntry("Emitter_NewNiagaraSystem", Label, ToolTip, Icon, UIAction);
+				}
+				{
+					const TAttribute<FText> Label = LOCTEXT("Emitter_CreateDuplicateParent", "Create Duplicate Parent");
+					const TAttribute<FText> ToolTip = LOCTEXT("Emitter_CreateDuplicateParentTooltip", "Duplicate this emitter and set this emitter's parent to the new emitter.");
+					const FSlateIcon Icon = FSlateIcon();
+
+					FToolUIAction UIAction;
+					UIAction.ExecuteAction = FToolMenuExecuteAction::CreateStatic(&ExecuteCreateDuplicateParent);
+					UIAction.CanExecuteAction = FToolMenuCanExecuteAction::CreateStatic(&CanExecuteCreateDuplicateParent);
+					InSection.AddMenuEntry("Emitter_CreateDuplicateParent", Label, ToolTip, Icon, UIAction);
+				}
+				{
+					const TAttribute<FText> Label = LOCTEXT("MarkDependentCompilableAssetsDirtyLabel", "Mark Dependent Compilable Assets Dirty");
+					const TAttribute<FText> ToolTip = LOCTEXT("MarkDependentCompilableAssetsDirtyToolTip", "Finds all niagara assets which depend on this asset either directly or indirectly, and marks them dirty so they can be saved with the latest version.");
+					const FSlateIcon Icon = FSlateIcon();
+					
+					FToolUIAction UIAction;
+					UIAction.ExecuteAction = FToolMenuExecuteAction::CreateStatic(&ExecuteMarkDependentCompilableAssetsDirty);
+					InSection.AddMenuEntry("MarkDependentCompilableAssetsDirty", Label, ToolTip, Icon, UIAction);
+				}
+			}));
+		}));
+	});
+}
+
+#undef LOCTEXT_NAMESPACE
diff --git a/Engine/Plugins/FX/Niagara/Source/NiagaraEditor/Private/AssetTypeActions/AssetDefinition_NiagaraEmitter.h b/Engine/Plugins/FX/Niagara/Source/NiagaraEditor/Private/AssetTypeActions/AssetDefinition_NiagaraEmitter.h
new file mode 100644
--- /dev/null
+++ b/Engine/Plugins/FX/Niagara/Source/NiagaraEditor/Private/AssetTypeActions/AssetDefinition_NiagaraEmitter.h
@@ -0,0 +1,29 @@
+// Copyright Epic Games, Inc. All Rights Reserved.
+
+#pragma once
+
+#include "CoreMinimal.h"
+#include "NiagaraEmitter.h"
+#include "AssetDefinitionDefault.h"
+
+#include "AssetDefinition_NiagaraEmitter.generated.h"
+
+UCLASS()
+class UAssetDefinition_NiagaraEmitter : public UAssetDefinitionDefault
+{
+	GENERATED_BODY()
+
+public:
+	// UAssetDefinition Begin
+	virtual FText GetAssetDisplayName() const override { return NSLOCTEXT("AssetTypeActions", "AssetTypeActions_NiagaraEmitter", "Niagara Emitter"); }
+	virtual FLinearColor GetAssetColor() const override;
+	virtual TSoftClassPtr<UObject> GetAssetClass() const override { return UNiagaraEmitter::StaticClass(); }
+	virtual TConstArrayView<FAssetCategoryPath> GetAssetCategories() const override
+	{
+		static const auto Categories = { EAssetCategoryPaths::FX };
+		return Categories;
+	}
+
+	virtual EAssetCommandResult OpenAssets(const FAssetOpenArgs& OpenArgs) const override;
+	// UAssetDefinition End
+};
diff --git a/Engine/Plugins/FX/Niagara/Source/NiagaraEditor/Private/AssetTypeActions/AssetDefinition_NiagaraParameterCollection.cpp b/Engine/Plugins/FX/Niagara/Source/NiagaraEditor/Private/AssetTypeActions/AssetDefinition_NiagaraParameterCollection.cpp
new file mode 100644
--- /dev/null
+++ b/Engine/Plugins/FX/Niagara/Source/NiagaraEditor/Private/AssetTypeActions/AssetDefinition_NiagaraParameterCollection.cpp
@@ -0,0 +1,78 @@
+// Copyright Epic Games, Inc. All Rights Reserved.
+
+#include "AssetDefinition_NiagaraParameterCollection.h"
+
+#include "ContentBrowserMenuContexts.h"
+#include "NiagaraParameterCollection.h"
+#include "NiagaraParameterCollectionToolkit.h"
+#include "NiagaraParameterCollectionFactoryNew.h"
+#include "NiagaraEditorStyle.h"
+
+#include "IContentBrowserSingleton.h"
+#include "ContentBrowserModule.h"
+#include "IAssetTools.h"
+
+#include "ToolMenus.h"
+#include "Styling/AppStyle.h"
+
+#define LOCTEXT_NAMESPACE "AssetTypeActions_NiagaraParameterCollection"
+
+FLinearColor UAssetDefinition_NiagaraParameterCollection::GetAssetColor() const
+{
+	return FNiagaraEditorStyle::Get().GetColor("NiagaraEditor.AssetColors.ParameterCollection");
+}
+
+EAssetCommandResult UAssetDefinition_NiagaraParameterCollection::OpenAssets(const FAssetOpenArgs& OpenArgs) const
+{
+	for (UNiagaraParameterCollection* NPC : OpenArgs.LoadObjects<UNiagaraParameterCollection>())
+	{
+		TSharedRef< FNiagaraParameterCollectionToolkit > NewNiagaraNPCToolkit(new FNiagaraParameterCollectionToolkit());
+		NewNiagaraNPCToolkit->Initialize(OpenArgs.GetToolkitMode(), OpenArgs.ToolkitHost, NPC);
+	}
+
+	return EAssetCommandResult::Handled;
+}
+
+// Menu Extensions
+//--------------------------------------------------------------------
+
+namespace MenuExtension_NiagaraParameterCollection
+{
+	static void ExecuteCreateNewNPCI(const FToolMenuContext& InContext)
+	{
+		const UContentBrowserAssetContextMenuContext* CBContext = UContentBrowserAssetContextMenuContext::FindContextWithAssets(InContext);
+
+		IAssetTools::Get().CreateAssetsFrom<UNiagaraParameterCollection>(
+			CBContext->LoadSelectedObjects<UNiagaraParameterCollection>(), UNiagaraParameterCollectionInstance::StaticClass(), TEXT("_Inst"), [](UNiagaraParameterCollection* SourceObject)
+			{
+				UNiagaraParameterCollectionInstanceFactoryNew* Factory = NewObject<UNiagaraParameterCollectionInstanceFactoryNew>();
+				Factory->InitialParent = SourceObject;
+				return Factory;
+			}
+		);
+	}
+	
+	static FDelayedAutoRegisterHelper DelayedAutoRegister(EDelayedRegisterRunPhase::EndOfEngineInit, []{ 
+		UToolMenus::RegisterStartupCallback(FSimpleMulticastDelegate::FDelegate::CreateLambda([]()
+		{
+			FToolMenuOwnerScoped OwnerScoped(UE_MODULE_NAME);
+			UToolMenu* Menu = UE::ContentBrowser::ExtendToolMenu_AssetContextMenu(UNiagaraParameterCollection::StaticClass());
+		
+			FToolMenuSection& Section = Menu->FindOrAddSection("GetAssetActions");
+			Section.AddDynamicEntry(NAME_None, FNewToolMenuSectionDelegate::CreateLambda([](FToolMenuSection& InSection)
+			{
+				{
+					const TAttribute<FText> Label = LOCTEXT("NewNPC", "Create Niagara Parameter Collection Instance");
+					const TAttribute<FText> ToolTip = LOCTEXT("NewNPCTooltip", "Creates an instance of this Niagara Parameter Collection.");
+					const FSlateIcon Icon = FSlateIcon(FAppStyle::GetAppStyleSetName(), "ClassIcon.MaterialInstanceActor");
+					
+					FToolUIAction UIAction;
+					UIAction.ExecuteAction = FToolMenuExecuteAction::CreateStatic(&ExecuteCreateNewNPCI);
+					InSection.AddMenuEntry("Niagara_NewNPC", Label, ToolTip, Icon, UIAction);
+				}
+			}));
+		}));
+	});
+}
+
+#undef LOCTEXT_NAMESPACE
diff --git a/Engine/Plugins/FX/Niagara/Source/NiagaraEditor/Private/AssetTypeActions/AssetDefinition_NiagaraParameterCollection.h b/Engine/Plugins/FX/Niagara/Source/NiagaraEditor/Private/AssetTypeActions/AssetDefinition_NiagaraParameterCollection.h
new file mode 100644
--- /dev/null
+++ b/Engine/Plugins/FX/Niagara/Source/NiagaraEditor/Private/AssetTypeActions/AssetDefinition_NiagaraParameterCollection.h
@@ -0,0 +1,29 @@
+// Copyright Epic Games, Inc. All Rights Reserved.
+
+#pragma once
+
+#include "CoreMinimal.h"
+#include "NiagaraParameterCollection.h"
+#include "AssetDefinitionDefault.h"
+
+#include "AssetDefinition_NiagaraParameterCollection.generated.h"
+
+UCLASS()
+class UAssetDefinition_NiagaraParameterCollection : public UAssetDefinitionDefault
+{
+	GENERATED_BODY()
+
+public:
+	// UAssetDefinition Begin
+	virtual FText GetAssetDisplayName() const override { return NSLOCTEXT("AssetTypeActions", "AssetTypeActions_NiagaraParameterCollection", "Niagara Parameter Collection"); }
+	virtual FLinearColor GetAssetColor() const override;
+	virtual TSoftClassPtr<UObject> GetAssetClass() const override { return UNiagaraParameterCollection::StaticClass(); }
+	virtual TConstArrayView<FAssetCategoryPath> GetAssetCategories() const override
+	{
+		static const auto Categories = { EAssetCategoryPaths::FX / NSLOCTEXT("Niagara", "NiagaraAssetSubMenu_Advanced", "Advanced") };
+		return Categories;
+	}
+
+	virtual EAssetCommandResult OpenAssets(const FAssetOpenArgs& OpenArgs) const override;
+	// UAssetDefinition End
+};
diff --git a/Engine/Plugins/FX/Niagara/Source/NiagaraEditor/Private/AssetTypeActions/AssetDefinition_NiagaraParameterCollectionInstance.cpp b/Engine/Plugins/FX/Niagara/Source/NiagaraEditor/Private/AssetTypeActions/AssetDefinition_NiagaraParameterCollectionInstance.cpp
new file mode 100644
--- /dev/null
+++ b/Engine/Plugins/FX/Niagara/Source/NiagaraEditor/Private/AssetTypeActions/AssetDefinition_NiagaraParameterCollectionInstance.cpp
@@ -0,0 +1,26 @@
+// Copyright Epic Games, Inc. All Rights Reserved.
+
+#include "AssetDefinition_NiagaraParameterCollectionInstance.h"
+#include "NiagaraParameterCollection.h"
+#include "NiagaraParameterCollectionToolkit.h"
+#include "NiagaraEditorStyle.h"
+
+#define LOCTEXT_NAMESPACE "UAssetDefinition_NiagaraParameterCollectionInstance"
+
+FLinearColor UAssetDefinition_NiagaraParameterCollectionInstance::GetAssetColor() const
+{ 
+	return FNiagaraEditorStyle::Get().GetColor("NiagaraEditor.AssetColors.ParameterCollectionInstance"); 
+}
+
+EAssetCommandResult UAssetDefinition_NiagaraParameterCollectionInstance::OpenAssets(const FAssetOpenArgs& OpenArgs) const
+{
+	for (UNiagaraParameterCollectionInstance* Instance : OpenArgs.LoadObjects<UNiagaraParameterCollectionInstance>())
+	{
+		TSharedRef< FNiagaraParameterCollectionToolkit > NewNiagaraNPCToolkit(new FNiagaraParameterCollectionToolkit());
+		NewNiagaraNPCToolkit->Initialize(OpenArgs.GetToolkitMode(), OpenArgs.ToolkitHost, Instance);
+	}
+
+	return EAssetCommandResult::Handled;
+}
+
+#undef LOCTEXT_NAMESPACE
diff --git a/Engine/Plugins/FX/Niagara/Source/NiagaraEditor/Private/AssetTypeActions/AssetDefinition_NiagaraParameterCollectionInstance.h b/Engine/Plugins/FX/Niagara/Source/NiagaraEditor/Private/AssetTypeActions/AssetDefinition_NiagaraParameterCollectionInstance.h
new file mode 100644
--- /dev/null
+++ b/Engine/Plugins/FX/Niagara/Source/NiagaraEditor/Private/AssetTypeActions/AssetDefinition_NiagaraParameterCollectionInstance.h
@@ -0,0 +1,29 @@
+// Copyright Epic Games, Inc. All Rights Reserved.
+
+#pragma once
+
+#include "CoreMinimal.h"
+#include "NiagaraParameterCollection.h"
+#include "AssetDefinitionDefault.h"
+
+#include "AssetDefinition_NiagaraParameterCollectionInstance.generated.h"
+
+UCLASS()
+class UAssetDefinition_NiagaraParameterCollectionInstance : public UAssetDefinitionDefault
+{
+	GENERATED_BODY()
+
+public:
+	// UAssetDefinition Begin
+	virtual FText GetAssetDisplayName() const override { return NSLOCTEXT("AssetTypeActions", "AssetTypeActions_NiagaraParameterCollectionInstance", "Niagara Parameter Collection Instance"); }
+	virtual FLinearColor GetAssetColor() const override;
+	virtual TSoftClassPtr<UObject> GetAssetClass() const override { return UNiagaraParameterCollectionInstance::StaticClass(); }
+	virtual TConstArrayView<FAssetCategoryPath> GetAssetCategories() const override
+	{
+		static const auto Categories = { EAssetCategoryPaths::FX / NSLOCTEXT("Niagara", "NiagaraAssetSubMenu_Advanced", "Advanced") };
+		return Categories;
+	}
+
+	virtual EAssetCommandResult OpenAssets(const FAssetOpenArgs& OpenArgs) const override;
+	// UAssetDefinition End
+};
diff --git a/Engine/Plugins/FX/Niagara/Source/NiagaraEditor/Private/AssetTypeActions/AssetDefinition_NiagaraParameterDefinitions.cpp b/Engine/Plugins/FX/Niagara/Source/NiagaraEditor/Private/AssetTypeActions/AssetDefinition_NiagaraParameterDefinitions.cpp
new file mode 100644
--- /dev/null
+++ b/Engine/Plugins/FX/Niagara/Source/NiagaraEditor/Private/AssetTypeActions/AssetDefinition_NiagaraParameterDefinitions.cpp
@@ -0,0 +1,21 @@
+// Copyright Epic Games, Inc. All Rights Reserved.
+
+#include "AssetDefinition_NiagaraParameterDefinitions.h"
+#include "NiagaraEditorStyle.h"
+#include "Toolkits/NiagaraParameterDefinitionsToolkit.h"
+
+FLinearColor UAssetDefinition_NiagaraParameterDefinitions::GetAssetColor() const
+{
+	return FNiagaraEditorStyle::Get().GetColor("NiagaraEditor.AssetColors.ParameterDefinitions");
+}
+
+EAssetCommandResult UAssetDefinition_NiagaraParameterDefinitions::OpenAssets(const FAssetOpenArgs& OpenArgs) const
+{
+	for (UNiagaraParameterDefinitions* ParameterDefinitions : OpenArgs.LoadObjects<UNiagaraParameterDefinitions>())
+   	{
+		TSharedRef<FNiagaraParameterDefinitionsToolkit> NewNiagaraParameterDefinitionsToolkit(new FNiagaraParameterDefinitionsToolkit());
+		NewNiagaraParameterDefinitionsToolkit->Initialize(OpenArgs.GetToolkitMode(), OpenArgs.ToolkitHost, ParameterDefinitions);
+	}
+
+	return EAssetCommandResult::Handled;
+}
diff --git a/Engine/Plugins/FX/Niagara/Source/NiagaraEditor/Private/AssetTypeActions/AssetDefinition_NiagaraParameterDefinitions.h b/Engine/Plugins/FX/Niagara/Source/NiagaraEditor/Private/AssetTypeActions/AssetDefinition_NiagaraParameterDefinitions.h
new file mode 100644
--- /dev/null
+++ b/Engine/Plugins/FX/Niagara/Source/NiagaraEditor/Private/AssetTypeActions/AssetDefinition_NiagaraParameterDefinitions.h
@@ -0,0 +1,29 @@
+// Copyright Epic Games, Inc. All Rights Reserved.
+
+#pragma once
+
+#include "CoreMinimal.h"
+#include "NiagaraParameterDefinitions.h"
+#include "AssetDefinitionDefault.h"
+
+#include "AssetDefinition_NiagaraParameterDefinitions.generated.h"
+
+UCLASS()
+class UAssetDefinition_NiagaraParameterDefinitions : public UAssetDefinitionDefault
+{
+	GENERATED_BODY()
+
+public:
+	// UAssetDefinition Begin
+	virtual FText GetAssetDisplayName() const override { return NSLOCTEXT("AssetTypeActions", "AssetTypeActions_NiagaraParameterDefinitions", "Niagara Parameter Definitions"); }
+	virtual FLinearColor GetAssetColor() const override;
+	virtual TSoftClassPtr<UObject> GetAssetClass() const override { return UNiagaraParameterDefinitions::StaticClass(); }
+	virtual TConstArrayView<FAssetCategoryPath> GetAssetCategories() const override
+	{
+		static const auto Categories = { EAssetCategoryPaths::FX / NSLOCTEXT("Niagara", "NiagaraAssetSubMenu_Advanced", "Advanced") };
+		return Categories;
+	}
+
+	virtual EAssetCommandResult OpenAssets(const FAssetOpenArgs& OpenArgs) const override;
+	// UAssetDefinition End
+};
diff --git a/Engine/Plugins/FX/Niagara/Source/NiagaraEditor/Private/AssetTypeActions/AssetDefinition_NiagaraScript.cpp b/Engine/Plugins/FX/Niagara/Source/NiagaraEditor/Private/AssetTypeActions/AssetDefinition_NiagaraScript.cpp
new file mode 100644
--- /dev/null
+++ b/Engine/Plugins/FX/Niagara/Source/NiagaraEditor/Private/AssetTypeActions/AssetDefinition_NiagaraScript.cpp
@@ -0,0 +1,98 @@
+// Copyright Epic Games, Inc. All Rights Reserved.
+
+#include "AssetTypeActions/AssetDefinition_NiagaraScript.h"
+#include "NiagaraScriptToolkit.h"
+#include "NiagaraEditorStyle.h"
+#include "AssetRegistry/AssetData.h"
+#include "NiagaraEditorUtilities.h"
+#include "ContentBrowserMenuContexts.h"
+#include "ToolMenus.h"
+
+#define LOCTEXT_NAMESPACE "NiagaraScriptAssetTypeActions"
+
+const FName UAssetDefinition_NiagaraScript::FunctionScriptName = "Niagara Function Script";
+const FText UAssetDefinition_NiagaraScript::FunctionScriptNameText = LOCTEXT("NiagaraFunctionScript", "Niagara Function Script");
+const FName UAssetDefinition_NiagaraScript::ModuleScriptName = "Niagara Module Script";
+const FText UAssetDefinition_NiagaraScript::ModuleScriptNameText = LOCTEXT("NiagaraModuleScript", "Niagara Module Script");
+const FName UAssetDefinition_NiagaraScript::DynamicInputScriptName = "Niagara Dynamic Input Script";
+const FText UAssetDefinition_NiagaraScript::DynamicInputScriptNameText = LOCTEXT("NiagaraDynamicInputScript", "Niagara Dynamic Input Script");
+
+FLinearColor UAssetDefinition_NiagaraScript::GetAssetColor() const
+{ 
+	return FNiagaraEditorStyle::Get().GetColor("NiagaraEditor.AssetColors.Script"); 
+}
+
+FText UAssetDefinition_NiagaraScript::GetAssetDisplayName(const FAssetData& AssetData) const
+{
+	static const FName NAME_Usage(TEXT("Usage"));
+	const FAssetDataTagMapSharedView::FFindTagResult Usage = AssetData.TagsAndValues.FindTag(NAME_Usage);
+
+	if (Usage.IsSet())
+	{
+		static const FString FunctionString(TEXT("Function"));
+		static const FString ModuleString(TEXT("Module"));
+		static const FString DynamicInputString(TEXT("DynamicInput"));
+		
+		if (Usage.GetValue() == FunctionString)
+		{
+			return FunctionScriptNameText;
+		}
+		else if (Usage.GetValue() == ModuleString)
+		{
+			return ModuleScriptNameText;
+		}
+		else if (Usage.GetValue() == DynamicInputString)
+		{
+			return DynamicInputScriptNameText;
+		}
+	}
+
+	return Super::GetAssetDisplayName(AssetData);
+}
+
+EAssetCommandResult UAssetDefinition_NiagaraScript::OpenAssets(const FAssetOpenArgs& OpenArgs) const
+{
+	for (UNiagaraScript* Script : OpenArgs.LoadObjects<UNiagaraScript>())
+	{
+		TSharedRef<FNiagaraScriptToolkit> NewNiagaraScriptToolkit(new FNiagaraScriptToolkit());
+		NewNiagaraScriptToolkit->Initialize(OpenArgs.GetToolkitMode(), OpenArgs.ToolkitHost, Script);
+	}
+	
+	return EAssetCommandResult::Handled;
+}
+
+// Menu Extensions
+//--------------------------------------------------------------------
+
+namespace MenuExtension_NiagaraScript
+{
+	static void ExecuteMarkDependentCompilableAssetsDirty(const FToolMenuContext& InContext)
+	{
+		const UContentBrowserAssetContextMenuContext* CBContext = UContentBrowserAssetContextMenuContext::FindContextWithAssets(InContext);
+		FNiagaraEditorUtilities::MarkDependentCompilableAssetsDirty(CBContext->LoadSelectedObjects<UObject>());
+	}
+	
+	static FDelayedAutoRegisterHelper DelayedAutoRegister(EDelayedRegisterRunPhase::EndOfEngineInit, []{ 
+		UToolMenus::RegisterStartupCallback(FSimpleMulticastDelegate::FDelegate::CreateLambda([]()
+		{
+			FToolMenuOwnerScoped OwnerScoped(UE_MODULE_NAME);
+			UToolMenu* Menu = UE::ContentBrowser::ExtendToolMenu_AssetContextMenu(UNiagaraScript::StaticClass());
+		
+			FToolMenuSection& Section = Menu->FindOrAddSection("GetAssetActions");
+			Section.AddDynamicEntry(NAME_None, FNewToolMenuSectionDelegate::CreateLambda([](FToolMenuSection& InSection)
+			{
+				{
+					const TAttribute<FText> Label = LOCTEXT("MarkDependentCompilableAssetsDirtyLabel", "Mark Dependent Compilable Assets Dirty");
+					const TAttribute<FText> ToolTip = LOCTEXT("MarkDependentCompilableAssetsDirtyToolTip", "Finds all niagara assets which depend on this asset either directly or indirectly, and marks them dirty so they can be saved with the latest version.");
+					const FSlateIcon Icon = FSlateIcon();
+					
+					FToolUIAction UIAction;
+					UIAction.ExecuteAction = FToolMenuExecuteAction::CreateStatic(&ExecuteMarkDependentCompilableAssetsDirty);
+					InSection.AddMenuEntry("MarkDependentCompilableAssetsDirty", Label, ToolTip, Icon, UIAction);
+				}
+			}));
+		}));
+	});
+}
+
+#undef LOCTEXT_NAMESPACE
diff --git a/Engine/Plugins/FX/Niagara/Source/NiagaraEditor/Private/AssetTypeActions/AssetDefinition_NiagaraScript.h b/Engine/Plugins/FX/Niagara/Source/NiagaraEditor/Private/AssetTypeActions/AssetDefinition_NiagaraScript.h
new file mode 100644
--- /dev/null
+++ b/Engine/Plugins/FX/Niagara/Source/NiagaraEditor/Private/AssetTypeActions/AssetDefinition_NiagaraScript.h
@@ -0,0 +1,40 @@
+// Copyright Epic Games, Inc. All Rights Reserved.
+
+#pragma once
+
+#include "CoreMinimal.h"
+#include "NiagaraScript.h"
+#include "AssetDefinitionDefault.h"
+
+#include "AssetDefinition_NiagaraScript.generated.h"
+
+UCLASS()
+class UAssetDefinition_NiagaraScript : public UAssetDefinitionDefault
+{
+	GENERATED_BODY()
+
+public:
+	static const FName FunctionScriptName;
+	static const FText FunctionScriptNameText;
+
+	static const FName ModuleScriptName;
+	static const FText ModuleScriptNameText;
+	
+	static const FName DynamicInputScriptName;
+	static const FText DynamicInputScriptNameText;
+
+public:
+	// UAssetDefinition Begin
+	virtual FText GetAssetDisplayName() const override { return NSLOCTEXT("AssetTypeActions", "AssetTypeActions_NiagaraScript", "Niagara Script"); }
+	virtual FText GetAssetDisplayName(const FAssetData& AssetData) const override;
+	virtual FLinearColor GetAssetColor() const override;
+	virtual TSoftClassPtr<UObject> GetAssetClass() const override { return UNiagaraScript::StaticClass(); }
+	virtual TConstArrayView<FAssetCategoryPath> GetAssetCategories() const override
+	{
+		static const auto Categories = { EAssetCategoryPaths::FX / NSLOCTEXT("Niagara", "NiagaraAssetSubMenu_Script", "Script") };
+		return Categories;
+	}
+
+	virtual EAssetCommandResult OpenAssets(const FAssetOpenArgs& OpenArgs) const override;
+	// UAssetDefinition End
+};
diff --git a/Engine/Plugins/FX/Niagara/Source/NiagaraEditor/Private/AssetTypeActions/AssetDefinition_NiagaraSimCache.cpp b/Engine/Plugins/FX/Niagara/Source/NiagaraEditor/Private/AssetTypeActions/AssetDefinition_NiagaraSimCache.cpp
new file mode 100644
--- /dev/null
+++ b/Engine/Plugins/FX/Niagara/Source/NiagaraEditor/Private/AssetTypeActions/AssetDefinition_NiagaraSimCache.cpp
@@ -0,0 +1,22 @@
+// Copyright Epic Games, Inc. All Rights Reserved.
+
+#include "AssetDefinition_NiagaraSimCache.h"
+
+#include "NiagaraEditorStyle.h"
+#include "NiagaraSimCacheToolkit.h"
+
+FLinearColor UAssetDefinition_NiagaraSimCache::GetAssetColor() const
+{
+	return FNiagaraEditorStyle::Get().GetColor("NiagaraEditor.AssetColors.SimCache").ToFColor(true);
+}
+
+EAssetCommandResult UAssetDefinition_NiagaraSimCache::OpenAssets(const FAssetOpenArgs& OpenArgs) const
+{
+	for (UNiagaraSimCache* SimCache : OpenArgs.LoadObjects<UNiagaraSimCache>())
+	{
+		const TSharedRef< FNiagaraSimCacheToolkit > NewNiagaraSimCacheToolkit(new FNiagaraSimCacheToolkit());
+		NewNiagaraSimCacheToolkit->Initialize(OpenArgs.GetToolkitMode(), OpenArgs.ToolkitHost, SimCache);
+	}
+
+	return EAssetCommandResult::Handled;
+}
diff --git a/Engine/Plugins/FX/Niagara/Source/NiagaraEditor/Private/AssetTypeActions/AssetDefinition_NiagaraSimCache.h b/Engine/Plugins/FX/Niagara/Source/NiagaraEditor/Private/AssetTypeActions/AssetDefinition_NiagaraSimCache.h
new file mode 100644
--- /dev/null
+++ b/Engine/Plugins/FX/Niagara/Source/NiagaraEditor/Private/AssetTypeActions/AssetDefinition_NiagaraSimCache.h
@@ -0,0 +1,29 @@
+// Copyright Epic Games, Inc. All Rights Reserved.
+
+#pragma once
+
+#include "CoreMinimal.h"
+#include "NiagaraSimCache.h"
+#include "AssetDefinitionDefault.h"
+
+#include "AssetDefinition_NiagaraSimCache.generated.h"
+
+UCLASS()
+class UAssetDefinition_NiagaraSimCache : public UAssetDefinitionDefault
+{
+	GENERATED_BODY()
+
+public:
+	// UAssetDefinition Begin
+	virtual FText GetAssetDisplayName() const override { return NSLOCTEXT("AssetTypeActions", "AssetTypeActions_NiagaraSimCache", "Niagara Simulation Cache"); }
+	virtual FLinearColor GetAssetColor() const override;
+	virtual TSoftClassPtr<UObject> GetAssetClass() const override { return UNiagaraSimCache::StaticClass(); }
+	virtual TConstArrayView<FAssetCategoryPath> GetAssetCategories() const override
+	{
+		static const auto Categories = { EAssetCategoryPaths::FX / NSLOCTEXT("Niagara", "NiagaraAssetSubMenu_Advanced", "Advanced") };
+		return Categories;
+	}
+
+	virtual EAssetCommandResult OpenAssets(const FAssetOpenArgs& OpenArgs) const override;
+	// UAssetDefinition End
+};
diff --git a/Engine/Plugins/FX/Niagara/Source/NiagaraEditor/Private/AssetTypeActions/AssetDefinition_NiagaraSystem.cpp b/Engine/Plugins/FX/Niagara/Source/NiagaraEditor/Private/AssetTypeActions/AssetDefinition_NiagaraSystem.cpp
new file mode 100644
--- /dev/null
+++ b/Engine/Plugins/FX/Niagara/Source/NiagaraEditor/Private/AssetTypeActions/AssetDefinition_NiagaraSystem.cpp
@@ -0,0 +1,22 @@
+// Copyright Epic Games, Inc. All Rights Reserved.
+
+#include "AssetTypeActions/AssetDefinition_NiagaraSystem.h"
+#include "NiagaraSystem.h"
+#include "NiagaraSystemToolkit.h"
+#include "NiagaraEditorStyle.h"
+
+FLinearColor UAssetDefinition_NiagaraSystem::GetAssetColor() const
+{
+	return FNiagaraEditorStyle::Get().GetColor("NiagaraEditor.AssetColors.System");
+}
+
+EAssetCommandResult UAssetDefinition_NiagaraSystem::OpenAssets(const FAssetOpenArgs& OpenArgs) const
+{
+	for (UNiagaraSystem* System : OpenArgs.LoadObjects<UNiagaraSystem>())
+	{
+		TSharedRef<FNiagaraSystemToolkit> NiagaraSystemToolkit = MakeShared<FNiagaraSystemToolkit>();
+		NiagaraSystemToolkit->InitializeWithSystem(OpenArgs.GetToolkitMode(), OpenArgs.ToolkitHost, *System);
+	}
+
+	return EAssetCommandResult::Handled;
+}
diff --git a/Engine/Plugins/FX/Niagara/Source/NiagaraEditor/Private/AssetTypeActions/AssetDefinition_NiagaraSystem.h b/Engine/Plugins/FX/Niagara/Source/NiagaraEditor/Private/AssetTypeActions/AssetDefinition_NiagaraSystem.h
new file mode 100644
--- /dev/null
+++ b/Engine/Plugins/FX/Niagara/Source/NiagaraEditor/Private/AssetTypeActions/AssetDefinition_NiagaraSystem.h
@@ -0,0 +1,29 @@
+// Copyright Epic Games, Inc. All Rights Reserved.
+
+#pragma once
+
+#include "CoreMinimal.h"
+#include "NiagaraSystem.h"
+#include "AssetDefinitionDefault.h"
+
+#include "AssetDefinition_NiagaraSystem.generated.h"
+
+UCLASS()
+class UAssetDefinition_NiagaraSystem : public UAssetDefinitionDefault
+{
+	GENERATED_BODY()
+
+public:
+	// UAssetDefinition Begin
+	virtual FText GetAssetDisplayName() const override { return NSLOCTEXT("AssetTypeActions", "AssetTypeActions_NiagaraSystem", "Niagara System"); }
+	virtual FLinearColor GetAssetColor() const override;
+	virtual TSoftClassPtr<UObject> GetAssetClass() const override { return UNiagaraSystem::StaticClass(); }
+	virtual TConstArrayView<FAssetCategoryPath> GetAssetCategories() const override
+	{
+		static const auto Categories = { EAssetCategoryPaths::Basic, EAssetCategoryPaths::FX };
+		return Categories;
+	}
+
+	virtual EAssetCommandResult OpenAssets(const FAssetOpenArgs& OpenArgs) const override;
+	// UAssetDefinition End
+};
diff --git a/Engine/Plugins/FX/Niagara/Source/NiagaraEditor/Private/AssetTypeActions/AssetDefinition_NiagaraValidationRuleSet.cpp b/Engine/Plugins/FX/Niagara/Source/NiagaraEditor/Private/AssetTypeActions/AssetDefinition_NiagaraValidationRuleSet.cpp
new file mode 100644
--- /dev/null
+++ b/Engine/Plugins/FX/Niagara/Source/NiagaraEditor/Private/AssetTypeActions/AssetDefinition_NiagaraValidationRuleSet.cpp
@@ -0,0 +1,9 @@
+// Copyright Epic Games, Inc. All Rights Reserved.
+
+#include "AssetDefinition_NiagaraValidationRuleSet.h"
+#include "NiagaraEditorStyle.h"
+
+FLinearColor UAssetDefinition_NiagaraValidationRuleSet::GetAssetColor() const
+{
+	return FNiagaraEditorStyle::Get().GetColor("NiagaraEditor.AssetColors.ValidationRuleSet");
+}
diff --git a/Engine/Plugins/FX/Niagara/Source/NiagaraEditor/Private/AssetTypeActions/AssetDefinition_NiagaraValidationRuleSet.h b/Engine/Plugins/FX/Niagara/Source/NiagaraEditor/Private/AssetTypeActions/AssetDefinition_NiagaraValidationRuleSet.h
new file mode 100644
--- /dev/null
+++ b/Engine/Plugins/FX/Niagara/Source/NiagaraEditor/Private/AssetTypeActions/AssetDefinition_NiagaraValidationRuleSet.h
@@ -0,0 +1,30 @@
+// Copyright Epic Games, Inc. All Rights Reserved.
+
+#pragma once
+
+#include "CoreMinimal.h"
+#include "NiagaraValidationRuleSet.h"
+#include "AssetDefinitionDefault.h"
+
+#include "AssetDefinition_NiagaraValidationRuleSet.generated.h"
+
+UCLASS()
+class UAssetDefinition_NiagaraValidationRuleSet : public UAssetDefinitionDefault
+{
+	GENERATED_BODY()
+
+public:
+	// UAssetDefinition Begin
+	virtual FText GetAssetDisplayName() const override { return NSLOCTEXT("AssetTypeActions", "AssetTypeActions_NiagaraValidationRuleSet", "Niagara Validation Rule Set"); }
+	virtual FLinearColor GetAssetColor() const override;
+	virtual TSoftClassPtr<UObject> GetAssetClass() const override { return UNiagaraValidationRuleSet::StaticClass(); }
+	virtual TConstArrayView<FAssetCategoryPath> GetAssetCategories() const override
+	{
+		static const auto Categories = {
+			EAssetCategoryPaths::FX / NSLOCTEXT("Niagara", "NiagaraAssetSubMenu_Advanced", "Advanced")
+		};
+		
+		return Categories;
+	}
+	// UAssetDefinition End
+};
diff --git a/Engine/Plugins/FX/Niagara/Source/NiagaraEditor/Private/AssetTypeActions/AssetTypeActions_NiagaraParameterCollection.cpp b/Engine/Plugins/FX/Niagara/Source/NiagaraEditor/Private/AssetTypeActions/AssetTypeActions_NiagaraParameterCollection.cpp
deleted file mode 100644
diff --git a/Engine/Plugins/FX/Niagara/Source/NiagaraEditor/Private/AssetTypeActions/AssetTypeActions_NiagaraParameterCollection.h b/Engine/Plugins/FX/Niagara/Source/NiagaraEditor/Private/AssetTypeActions/AssetTypeActions_NiagaraParameterCollection.h
deleted file mode 100644
diff --git a/Engine/Plugins/FX/Niagara/Source/NiagaraEditor/Private/AssetTypeActions/AssetTypeActions_NiagaraScript.cpp b/Engine/Plugins/FX/Niagara/Source/NiagaraEditor/Private/AssetTypeActions/AssetTypeActions_NiagaraScript.cpp
deleted file mode 100644
diff --git a/Engine/Plugins/FX/Niagara/Source/NiagaraEditor/Private/AssetTypeActions/AssetTypeActions_NiagaraScript.h b/Engine/Plugins/FX/Niagara/Source/NiagaraEditor/Private/AssetTypeActions/AssetTypeActions_NiagaraScript.h
deleted file mode 100644
diff --git a/Engine/Plugins/FX/Niagara/Source/NiagaraEditor/Private/NiagaraEditorModule.cpp b/Engine/Plugins/FX/Niagara/Source/NiagaraEditor/Private/NiagaraEditorModule.cpp
--- a/Engine/Plugins/FX/Niagara/Source/NiagaraEditor/Private/NiagaraEditorModule.cpp
+++ b/Engine/Plugins/FX/Niagara/Source/NiagaraEditor/Private/NiagaraEditorModule.cpp
@@ -15,9 +15,6 @@
 #include "Subsystems/ImportSubsystem.h"
 #include "UObject/UObjectThreadContext.h"
 
-#include "AssetTypeActions/AssetTypeActions_NiagaraScript.h"
-#include "AssetTypeActions/AssetTypeActions_NiagaraParameterCollection.h"
-
 #include "EdGraphSchema_Niagara.h"
 #include "EdGraphUtilities.h"
 #include "KismetPins/SGraphPinBool.h"
@@ -165,8 +162,6 @@ IMPLEMENT_MODULE( FNiagaraEditorModule, NiagaraEditor );
 const FName FNiagaraEditorModule::NiagaraEditorAppIdentifier( TEXT( "NiagaraEditorApp" ) );
 const FLinearColor FNiagaraEditorModule::WorldCentricTabColorScale(0.0f, 0.0f, 0.2f, 0.5f);
 
-EAssetTypeCategories::Type FNiagaraEditorModule::NiagaraAssetCategory;
-
 int32 GbShowNiagaraDeveloperWindows = 0;
 static FAutoConsoleVariableRef CVarShowNiagaraDeveloperWindows(
 	TEXT("fx.ShowNiagaraDeveloperWindows"),
@@ -859,12 +854,6 @@ void FNiagaraEditorModule::StartupModule()
 
 
 	IAssetTools& AssetTools = FModuleManager::LoadModuleChecked<FAssetToolsModule>("AssetTools").Get();
-	NiagaraAssetCategory = AssetTools.RegisterAdvancedAssetCategory(FName(TEXT("FX")), LOCTEXT("NiagaraAssetsCategory", "FX"));
-	RegisterAssetTypeAction(AssetTools, MakeShareable(new FAssetTypeActions_NiagaraScriptFunctions()));
-	RegisterAssetTypeAction(AssetTools, MakeShareable(new FAssetTypeActions_NiagaraScriptModules()));
-	RegisterAssetTypeAction(AssetTools, MakeShareable(new FAssetTypeActions_NiagaraScriptDynamicInputs()));
-	RegisterAssetTypeAction(AssetTools, MakeShareable(new FAssetTypeActions_NiagaraParameterCollection())); 
-	RegisterAssetTypeAction(AssetTools, MakeShareable(new FAssetTypeActions_NiagaraParameterCollectionInstance()));
 
 	// Preload all parameter definition & collection assets so that they will be postloaded before postload calls to scripts/emitters/systems that rely on them.
 	{
@@ -1264,16 +1253,6 @@ void FNiagaraEditorModule::ShutdownModule()
 {
 	MenuExtensibilityManager.Reset();
 	ToolBarExtensibilityManager.Reset();
-	
-	if (FModuleManager::Get().IsModuleLoaded("AssetTools"))
-	{
-		IAssetTools& AssetTools = FModuleManager::GetModuleChecked<FAssetToolsModule>("AssetTools").Get();
-		for (auto CreatedAssetTypeAction : CreatedAssetTypeActions)
-		{
-			AssetTools.UnregisterAssetTypeActions(CreatedAssetTypeAction.ToSharedRef());
-		}
-	}
-	CreatedAssetTypeActions.Empty();
 
 	// Clean up asset registry callbacks
 	if (FModuleManager::Get().IsModuleLoaded("AssetRegistry"))
@@ -1835,12 +1814,6 @@ void FNiagaraEditorModule::RegisterDefaultRendererFactories()
 		})));
 }
 
-void FNiagaraEditorModule::RegisterAssetTypeAction(IAssetTools& AssetTools, TSharedRef<IAssetTypeActions> Action)
-{
-	AssetTools.RegisterAssetTypeActions(Action);
-	CreatedAssetTypeActions.Add(Action);
-}
-
 void FNiagaraEditorModule::RegisterSettings()
 {
 	ISettingsModule* SettingsModule = FModuleManager::GetModulePtr<ISettingsModule>("Settings");
diff --git a/Engine/Plugins/FX/Niagara/Source/NiagaraEditor/Private/NiagaraScriptFactory.cpp b/Engine/Plugins/FX/Niagara/Source/NiagaraEditor/Private/NiagaraScriptFactory.cpp
--- a/Engine/Plugins/FX/Niagara/Source/NiagaraEditor/Private/NiagaraScriptFactory.cpp
+++ b/Engine/Plugins/FX/Niagara/Source/NiagaraEditor/Private/NiagaraScriptFactory.cpp
@@ -16,7 +16,7 @@
 #include "NiagaraScript.h"
 #include "NiagaraScriptFactoryNew.h"
 #include "NiagaraScriptSource.h"
-#include "AssetTypeActions/AssetTypeActions_NiagaraScript.h"
+#include "AssetTypeActions/AssetDefinition_NiagaraScript.h"
 #include "EdGraph/EdGraph.h"
 #include "Modules/ModuleManager.h"
 #include "ViewModels/Stack/NiagaraStackGraphUtilities.h"
@@ -89,34 +89,12 @@ const FSoftObjectPath& UNiagaraScriptFactoryNew::GetDefaultScriptFromSettings(co
 
 FText UNiagaraScriptFactoryNew::GetDisplayName() const
 {
-	// Get the displayname for this factory from the action type.
-	FAssetToolsModule& AssetToolsModule = FModuleManager::LoadModuleChecked<FAssetToolsModule>("AssetTools");
-
-	UClass* LocalSupportedClass = GetSupportedClass();
-	if (LocalSupportedClass)
-	{
-		TArray<TWeakPtr<IAssetTypeActions>> AssetTypeActions = AssetToolsModule.Get().GetAssetTypeActionsListForClass(LocalSupportedClass);
-		for (auto& CurrentAssetTypeAction : AssetTypeActions)
-		{
-			TSharedPtr<FAssetTypeActions_NiagaraScript> AssetTypeAction = StaticCastSharedPtr<FAssetTypeActions_NiagaraScript>(CurrentAssetTypeAction.Pin());
-			if (AssetTypeAction.IsValid() && AssetTypeAction->GetTypeName() == GetAssetTypeActionName())
-			{
-				FText Name = AssetTypeAction->GetName();
-				if (!Name.IsEmpty())
-				{
-					return Name;
-				}
-			}
-		}
-	}
-
-	// Factories that have no supported class have no display name.
-	return FText();
+	return FText::FromName(GetAssetTypeActionName());
 }
 
 void UNiagaraScriptFactoryNew::InitializeScript(UNiagaraScript* NewScript)
 {
-	if(NewScript != NULL)
+	if(NewScript != nullptr)
 	{
 		UNiagaraScriptSource* Source = NewObject<UNiagaraScriptSource>(NewScript, NAME_None, RF_Transactional);
 		if(Source)
@@ -240,7 +218,7 @@ FName UNiagaraModuleScriptFactory::GetNewAssetThumbnailOverride() const
 
 FName UNiagaraModuleScriptFactory::GetAssetTypeActionName() const
 {
-	return FAssetTypeActions_NiagaraScriptModules::NiagaraModuleScriptName;
+	return UAssetDefinition_NiagaraScript::ModuleScriptName;
 }
 
 ENiagaraScriptUsage UNiagaraModuleScriptFactory::GetScriptUsage() const
@@ -263,7 +241,7 @@ FName UNiagaraFunctionScriptFactory::GetNewAssetThumbnailOverride() const
 
 FName UNiagaraFunctionScriptFactory::GetAssetTypeActionName() const
 {
-	return FAssetTypeActions_NiagaraScriptFunctions::NiagaraFunctionScriptName;
+	return UAssetDefinition_NiagaraScript::FunctionScriptName;
 }
 
 ENiagaraScriptUsage UNiagaraFunctionScriptFactory::GetScriptUsage() const
@@ -286,7 +264,7 @@ FName UNiagaraDynamicInputScriptFactory::GetNewAssetThumbnailOverride() const
 
 FName UNiagaraDynamicInputScriptFactory::GetAssetTypeActionName() const
 {
-	return FAssetTypeActions_NiagaraScriptDynamicInputs::NiagaraDynamicInputScriptName;
+	return UAssetDefinition_NiagaraScript::DynamicInputScriptName;
 }
 
 ENiagaraScriptUsage UNiagaraDynamicInputScriptFactory::GetScriptUsage() const
diff --git a/Engine/Plugins/FX/Niagara/Source/NiagaraEditor/Private/Toolkits/NiagaraScriptToolkit.cpp b/Engine/Plugins/FX/Niagara/Source/NiagaraEditor/Private/Toolkits/NiagaraScriptToolkit.cpp
--- a/Engine/Plugins/FX/Niagara/Source/NiagaraEditor/Private/Toolkits/NiagaraScriptToolkit.cpp
+++ b/Engine/Plugins/FX/Niagara/Source/NiagaraEditor/Private/Toolkits/NiagaraScriptToolkit.cpp
@@ -2,7 +2,7 @@
 
 #include "NiagaraScriptToolkit.h"
 
-#include "AssetTypeActions/AssetTypeActions_NiagaraScript.h"
+#include "AssetTypeActions/AssetDefinition_NiagaraScript.h"
 #include "BusyCursor.h"
 #include "DetailLayoutBuilder.h"
 #include "Editor.h"
@@ -152,15 +152,15 @@ FText FNiagaraScriptToolkit::GetGraphEditorDisplayName() const
 	FText DisplayName = LOCTEXT("NiagaraScriptDisplayName", "Niagara Script");
 	if (EditedNiagaraScript.Script->GetUsage() == ENiagaraScriptUsage::Function)
 	{
-		DisplayName = FAssetTypeActions_NiagaraScriptFunctions::GetFormattedName();
+		DisplayName = UAssetDefinition_NiagaraScript::FunctionScriptNameText;
 	}
 	else if (EditedNiagaraScript.Script->GetUsage() == ENiagaraScriptUsage::Module)
 	{
-		DisplayName = FAssetTypeActions_NiagaraScriptModules::GetFormattedName();
+		DisplayName = UAssetDefinition_NiagaraScript::ModuleScriptNameText;
 	}
 	else if (EditedNiagaraScript.Script->GetUsage() == ENiagaraScriptUsage::DynamicInput)
 	{
-		DisplayName = FAssetTypeActions_NiagaraScriptDynamicInputs::GetFormattedName();
+		DisplayName = UAssetDefinition_NiagaraScript::DynamicInputScriptNameText;
 	}
 
 	FVersionedNiagaraScriptData* ScriptData = EditedNiagaraScript.Script->GetScriptData(EditedNiagaraScript.Version);
diff --git a/Engine/Plugins/FX/Niagara/Source/NiagaraEditor/Public/NiagaraEditorModule.h b/Engine/Plugins/FX/Niagara/Source/NiagaraEditor/Public/NiagaraEditorModule.h
--- a/Engine/Plugins/FX/Niagara/Source/NiagaraEditor/Public/NiagaraEditorModule.h
+++ b/Engine/Plugins/FX/Niagara/Source/NiagaraEditor/Public/NiagaraEditorModule.h
@@ -171,8 +171,6 @@ public:
 	/** Gets Niagara editor type utilities for a specific type if there are any registered. */
 	TSharedPtr<INiagaraEditorTypeUtilities, ESPMode::ThreadSafe> NIAGARAEDITOR_API GetTypeUtilities(const FNiagaraTypeDefinition& Type);
 
-	static EAssetTypeCategories::Type GetAssetCategory() { return NiagaraAssetCategory; }
-
 	NIAGARAEDITOR_API void RegisterWidgetProvider(TSharedRef<INiagaraEditorWidgetProvider> InWidgetProvider);
 	NIAGARAEDITOR_API void UnregisterWidgetProvider(TSharedRef<INiagaraEditorWidgetProvider> InWidgetProvider);
 
@@ -304,7 +302,6 @@ private:
 	};
 
 	void RegisterDefaultRendererFactories();
-	void RegisterAssetTypeAction(IAssetTools& AssetTools, TSharedRef<IAssetTypeActions> Action);
 	void OnNiagaraSettingsChangedEvent(const FName& PropertyName, const UNiagaraSettings* Settings);
 	void OnPreGarbageCollection();
 	void OnExecParticleInvoked(const TCHAR* InStr);
@@ -344,8 +341,6 @@ private:
 	TMap<FNiagaraTypeDefinition, TSharedRef<INiagaraEditorTypeUtilities, ESPMode::ThreadSafe>> TypeToEditorUtilitiesMap;
 	TSharedPtr<INiagaraEditorTypeUtilities, ESPMode::ThreadSafe> EnumTypeUtilities;
 
-	static EAssetTypeCategories::Type NiagaraAssetCategory;
-
 	FDelegateHandle CreateEmitterTrackEditorHandle;
 	FDelegateHandle CreateSystemTrackEditorHandle;
 
diff --git a/Engine/Source/Editor/AssetDefinition/Private/AssetDefinition.cpp b/Engine/Source/Editor/AssetDefinition/Private/AssetDefinition.cpp
--- a/Engine/Source/Editor/AssetDefinition/Private/AssetDefinition.cpp
+++ b/Engine/Source/Editor/AssetDefinition/Private/AssetDefinition.cpp
@@ -18,6 +18,7 @@ FAssetCategoryPath EAssetCategoryPaths::Blueprint(LOCTEXT("Blueprint", "Blueprin
 FAssetCategoryPath EAssetCategoryPaths::Texture(LOCTEXT("Texture", "Texture"));
 FAssetCategoryPath EAssetCategoryPaths::Foliage(LOCTEXT("Foliage", "Foliage"));
 FAssetCategoryPath EAssetCategoryPaths::Input(LOCTEXT("Input", "Input"));
+FAssetCategoryPath EAssetCategoryPaths::FX(LOCTEXT("FX", "FX"));
 FAssetCategoryPath EAssetCategoryPaths::Cinematics(LOCTEXT("Cinematics", "Cinematics"));
 
 FAssetCategoryPath::FAssetCategoryPath(const FText& InCategory)
diff --git a/Engine/Source/Editor/AssetDefinition/Public/AssetDefinition.h b/Engine/Source/Editor/AssetDefinition/Public/AssetDefinition.h
--- a/Engine/Source/Editor/AssetDefinition/Public/AssetDefinition.h
+++ b/Engine/Source/Editor/AssetDefinition/Public/AssetDefinition.h
@@ -245,7 +245,9 @@ struct ASSETDEFINITION_API EAssetCategoryPaths
 	static FAssetCategoryPath Animation;
 	static FAssetCategoryPath Audio;
 	static FAssetCategoryPath Blueprint;
+	static FAssetCategoryPath Cinematics;
 	static FAssetCategoryPath Foliage;
+	static FAssetCategoryPath FX;
 	static FAssetCategoryPath Gameplay;
 	static FAssetCategoryPath Input;
 	static FAssetCategoryPath Material;
@@ -253,7 +255,6 @@ struct ASSETDEFINITION_API EAssetCategoryPaths
 	static FAssetCategoryPath Physics;
 	static FAssetCategoryPath Texture;
 	static FAssetCategoryPath UI;
-	static FAssetCategoryPath Cinematics;
 };
 
 struct FAssetOpenSupport
