AssetDefinition - Merging differences between the existing AD work in niagara and the work I did.

diff --git a/Engine/Plugins/FX/Niagara/Source/NiagaraEditor/Private/AssetDefinitions/AssetDefinition_NiagaraEffectType.cpp b/Engine/Plugins/FX/Niagara/Source/NiagaraEditor/Private/AssetDefinitions/AssetDefinition_NiagaraEffectType.cpp
--- a/Engine/Plugins/FX/Niagara/Source/NiagaraEditor/Private/AssetDefinitions/AssetDefinition_NiagaraEffectType.cpp
+++ b/Engine/Plugins/FX/Niagara/Source/NiagaraEditor/Private/AssetDefinitions/AssetDefinition_NiagaraEffectType.cpp
@@ -1,26 +1,13 @@
 // Copyright Epic Games, Inc. All Rights Reserved.
 
-#include "AssetDefinitions/AssetDefinition_NiagaraEffectType.h"
-
+#include "AssetDefinition_NiagaraEffectType.h"
 #include "NiagaraEditorStyle.h"
-#include "NiagaraEffectType.h"
 
-#define LOCTEXT_NAMESPACE "UAssetDefinition_NiagaraEffectType"
+#define LOCTEXT_NAMESPACE "AssetTypeActions_NiagaraEffectType"
 
 FLinearColor UAssetDefinition_NiagaraEffectType::GetAssetColor() const
 {
 	return FNiagaraEditorStyle::Get().GetColor("NiagaraEditor.AssetColors.EffectType");
 }
 
-TSoftClassPtr<> UAssetDefinition_NiagaraEffectType::GetAssetClass() const
-{
-	return UNiagaraEffectType::StaticClass();
-}
-
-TConstArrayView<FAssetCategoryPath> UAssetDefinition_NiagaraEffectType::GetAssetCategories() const
-{
-	static FAssetCategoryPath AssetPaths[] = { FAssetCategoryPath(LOCTEXT("NiagaraAssetsCategory", "FX"), LOCTEXT("NiagaraEffectType_SubCategory", "Advanced")) };
-	return AssetPaths;
-}
-
 #undef LOCTEXT_NAMESPACE
diff --git a/Engine/Plugins/FX/Niagara/Source/NiagaraEditor/Private/AssetDefinitions/AssetDefinition_NiagaraEffectType.h b/Engine/Plugins/FX/Niagara/Source/NiagaraEditor/Private/AssetDefinitions/AssetDefinition_NiagaraEffectType.h
--- a/Engine/Plugins/FX/Niagara/Source/NiagaraEditor/Private/AssetDefinitions/AssetDefinition_NiagaraEffectType.h
+++ b/Engine/Plugins/FX/Niagara/Source/NiagaraEditor/Private/AssetDefinitions/AssetDefinition_NiagaraEffectType.h
@@ -3,6 +3,7 @@
 #pragma once
 
 #include "CoreMinimal.h"
+#include "NiagaraEffectType.h"
 #include "AssetDefinitionDefault.h"
 
 #include "AssetDefinition_NiagaraEffectType.generated.h"
@@ -13,9 +14,17 @@ class UAssetDefinition_NiagaraEffectType : public UAssetDefinitionDefault
 	GENERATED_BODY()
 
 public:
-	// UAssetDefinition Implementation
-	virtual FText GetAssetDisplayName() const override { return NSLOCTEXT("AssetDefinition", "NiagaraEffectType", "Niagara Effect Type"); }
+	// UAssetDefinition Begin
+	virtual FText GetAssetDisplayName() const override { return NSLOCTEXT("AssetTypeActions", "AssetTypeActions_NiagaraEffectType", "Niagara Effect Type"); }
 	virtual FLinearColor GetAssetColor() const override;
-	virtual TSoftClassPtr<UObject> GetAssetClass() const override;
-	virtual TConstArrayView<FAssetCategoryPath> GetAssetCategories() const override;
+	virtual TSoftClassPtr<UObject> GetAssetClass() const override { return UNiagaraEffectType::StaticClass(); }
+	virtual TConstArrayView<FAssetCategoryPath> GetAssetCategories() const override
+	{
+		static const auto Categories = {
+			EAssetCategoryPaths::FX / NSLOCTEXT("Niagara", "NiagaraAssetSubMenu_Advanced", "Advanced")
+		};
+		
+		return Categories;
+	}
+	// UAssetDefinition End
 };
diff --git a/Engine/Plugins/FX/Niagara/Source/NiagaraEditor/Private/AssetDefinitions/AssetDefinition_NiagaraEmitter.cpp b/Engine/Plugins/FX/Niagara/Source/NiagaraEditor/Private/AssetDefinitions/AssetDefinition_NiagaraEmitter.cpp
--- a/Engine/Plugins/FX/Niagara/Source/NiagaraEditor/Private/AssetDefinitions/AssetDefinition_NiagaraEmitter.cpp
+++ b/Engine/Plugins/FX/Niagara/Source/NiagaraEditor/Private/AssetDefinitions/AssetDefinition_NiagaraEmitter.cpp
@@ -1,20 +1,24 @@
 // Copyright Epic Games, Inc. All Rights Reserved.
 
-#include "AssetDefinitions/AssetDefinition_NiagaraEmitter.h"
+#include "AssetDefinition_NiagaraEmitter.h"
 
 #include "AssetToolsModule.h"
 #include "ContentBrowserMenuContexts.h"
-#include "ContentBrowserModule.h"
-#include "IContentBrowserSingleton.h"
-#include "ToolMenus.h"
-#include "ToolMenuDelegates.h"
-#include "NiagaraEditorStyle.h"
-#include "NiagaraEditorUtilities.h"
 #include "NiagaraEmitter.h"
 #include "NiagaraScriptSourceBase.h"
-#include "NiagaraSystemFactoryNew.h"
+#include "NiagaraSystem.h"
 #include "NiagaraSystemToolkit.h"
+#include "NiagaraEditorStyle.h"
+#include "NiagaraSystemFactoryNew.h"
 #include "ViewModels/NiagaraSystemViewModel.h"
+#include "ToolMenus.h"
+#include "Styling/AppStyle.h"
+#include "Toolkits/AssetEditorToolkit.h"
+#include "IContentBrowserSingleton.h"
+#include "ContentBrowserModule.h"
+#include "NiagaraEditorUtilities.h"
+#include "Modules/ModuleManager.h"
+#include "Subsystems/AssetEditorSubsystem.h"
 
 #define LOCTEXT_NAMESPACE "UAssetDefinition_NiagaraEmitter"
 
@@ -23,115 +27,86 @@ FLinearColor UAssetDefinition_NiagaraEmitter::GetAssetColor() const
 	return FNiagaraEditorStyle::Get().GetColor("NiagaraEditor.AssetColors.Emitter");
 }
 
-TSoftClassPtr<> UAssetDefinition_NiagaraEmitter::GetAssetClass() const
-{
-	return UNiagaraEmitter::StaticClass();
-}
-
-TConstArrayView<FAssetCategoryPath> UAssetDefinition_NiagaraEmitter::GetAssetCategories() const
-{
-	static FAssetCategoryPath AssetPaths[] = { LOCTEXT("NiagaraAssetsCategory", "FX") };
-	return AssetPaths;
-}
-
 EAssetCommandResult UAssetDefinition_NiagaraEmitter::OpenAssets(const FAssetOpenArgs& OpenArgs) const
 {
-	if (OpenArgs.OpenMethod == EAssetOpenMethod::Edit)
+	for (UNiagaraEmitter* Emitter : OpenArgs.LoadObjects<UNiagaraEmitter>())
 	{
-		EToolkitMode::Type Mode = OpenArgs.ToolkitHost.IsValid() ? EToolkitMode::WorldCentric : EToolkitMode::Standalone;
-		TArray<UObject*> Objects = OpenArgs.LoadObjects<UObject>();
-		for (UObject* OpenObj : Objects)
+		if (!Emitter->VersionToOpenInEditor.IsValid())
 		{
-			if (UNiagaraEmitter* Emitter = Cast<UNiagaraEmitter>(OpenObj))
-			{
-				if (!Emitter->VersionToOpenInEditor.IsValid())
-				{
-					Emitter->VersionToOpenInEditor = Emitter->GetExposedVersion().VersionGuid;
-				}
-				TSharedRef<FNiagaraSystemToolkit> SystemToolkit(new FNiagaraSystemToolkit());
-				SystemToolkit->InitializeWithEmitter(Mode, OpenArgs.ToolkitHost, *Emitter);
-			}
+			Emitter->VersionToOpenInEditor = Emitter->GetExposedVersion().VersionGuid;
 		}
-		return EAssetCommandResult::Handled;
+		
+		TSharedRef<FNiagaraSystemToolkit> SystemToolkit(new FNiagaraSystemToolkit());
+		SystemToolkit->InitializeWithEmitter(OpenArgs.GetToolkitMode(), OpenArgs.ToolkitHost, *Emitter);
 	}
-
-	return EAssetCommandResult::Unhandled;
+	
+	return EAssetCommandResult::Handled;
 }
 
-namespace MenuExtension_Emitter
+// Menu Extensions
+//--------------------------------------------------------------------
+
+namespace MenuExtension_NiagaraEmitter
 {
-	void CreateUniqueAssetName(const FString& InBasePackageName, const FString& InSuffix, FString& OutPackageName, FString& OutAssetName)
+	void ExecuteNewNiagaraSystem(const FToolMenuContext& InContext)
 	{
-		FAssetToolsModule& AssetToolsModule = FModuleManager::Get().LoadModuleChecked<FAssetToolsModule>("AssetTools");
-		AssetToolsModule.Get().CreateUniqueAssetName(InBasePackageName, InSuffix, OutPackageName, OutAssetName);
-	}
-	
-	static void ExecuteNewNiagaraSystem(const FToolMenuContext& MenuContext)
-	{
-		if (const UContentBrowserAssetContextMenuContext* Context = UContentBrowserAssetContextMenuContext::FindContextWithAssets(MenuContext))
+		const FString DefaultSuffix = TEXT("_System");
+
+		FNiagaraSystemViewModelOptions SystemOptions;
+		SystemOptions.bCanModifyEmittersFromTimeline = true;
+		SystemOptions.EditMode = ENiagaraSystemViewModelEditMode::SystemAsset;
+		
+		FAssetToolsModule& AssetToolsModule = FModuleManager::GetModuleChecked<FAssetToolsModule>("AssetTools");
+
+		TArray<UObject*> ObjectsToSync;
+		
+		const UContentBrowserAssetContextMenuContext* CBContext = UContentBrowserAssetContextMenuContext::FindContextWithAssets(InContext);
+		for (UNiagaraEmitter* Emitter : CBContext->LoadSelectedObjects<UNiagaraEmitter>())
 		{
-			TArray<UNiagaraEmitter*> Emitters = Context->LoadSelectedObjects<UNiagaraEmitter>();
-			const FString DefaultSuffix = TEXT("_System");
-
-			FNiagaraSystemViewModelOptions SystemOptions;
-			SystemOptions.bCanModifyEmittersFromTimeline = true;
-			SystemOptions.EditMode = ENiagaraSystemViewModelEditMode::SystemAsset;
-
-			TArray<UObject*> ObjectsToSync;
-			for (UNiagaraEmitter* Emitter : Emitters)
+			// Determine an appropriate names
+			FString Name;
+			FString PackageName;
+			AssetToolsModule.Get().CreateUniqueAssetName(Emitter->GetOutermost()->GetName(), DefaultSuffix, PackageName, Name);
+
+			// Create the factory used to generate the asset
+			UNiagaraSystemFactoryNew* Factory = NewObject<UNiagaraSystemFactoryNew>();
+			Factory->EmittersToAddToNewSystem.Add(FVersionedNiagaraEmitter(Emitter, Emitter->GetExposedVersion().VersionGuid));
+			UObject* NewAsset = AssetToolsModule.Get().CreateAsset(Name, FPackageName::GetLongPackagePath(PackageName), UNiagaraSystem::StaticClass(), Factory);
+				
+			UNiagaraSystem* System = Cast<UNiagaraSystem>(NewAsset);
+			if (System != nullptr)
 			{
-				if (Emitter)
-				{
-					// Determine an appropriate names
-					FString Name;
-					FString PackageName;
-					CreateUniqueAssetName(Emitter->GetOutermost()->GetName(), DefaultSuffix, PackageName, Name);
-
-					// Create the factory used to generate the asset
-					UNiagaraSystemFactoryNew* Factory = NewObject<UNiagaraSystemFactoryNew>();
-					Factory->EmittersToAddToNewSystem.Add(FVersionedNiagaraEmitter(Emitter, Emitter->GetExposedVersion().VersionGuid));
-					FAssetToolsModule& AssetToolsModule = FModuleManager::GetModuleChecked<FAssetToolsModule>("AssetTools");
-					UObject* NewAsset = AssetToolsModule.Get().CreateAsset(Name, FPackageName::GetLongPackagePath(PackageName), UNiagaraSystem::StaticClass(), Factory);
-			
-					UNiagaraSystem* System = Cast<UNiagaraSystem>(NewAsset);
-					if (System != nullptr)
-					{
-						ObjectsToSync.Add(NewAsset);
-					}
-				}
+				ObjectsToSync.Add(NewAsset);
 			}
+		}
 
-			if (ObjectsToSync.Num() > 0)
-			{
-				FContentBrowserModule& ContentBrowserModule = FModuleManager::LoadModuleChecked<FContentBrowserModule>("ContentBrowser");
-				ContentBrowserModule.Get().SyncBrowserToAssets(ObjectsToSync);
-			}
+		if (ObjectsToSync.Num() > 0)
+		{
+			FContentBrowserModule& ContentBrowserModule = FModuleManager::LoadModuleChecked<FContentBrowserModule>("ContentBrowser");
+			ContentBrowserModule.Get().SyncBrowserToAssets(ObjectsToSync);
 		}
 	}
 
-	static void ExecuteCreateDuplicateParent(const FToolMenuContext& MenuContext)
+	void ExecuteCreateDuplicateParent(const FToolMenuContext& InContext)
 	{
-		if (const UContentBrowserAssetContextMenuContext* Context = UContentBrowserAssetContextMenuContext::FindContextWithAssets(MenuContext))
-		{
-			TArray<UNiagaraEmitter*> Emitters = Context->LoadSelectedObjects<UNiagaraEmitter>();
-			const FString DefaultSuffix = TEXT("_Parent");
+		const FString DefaultSuffix = TEXT("_Parent");
 
-			TArray<UObject*> ObjectsToSync;
-			FAssetToolsModule& AssetToolsModule = FModuleManager::GetModuleChecked<FAssetToolsModule>("AssetTools");
+		TArray<UObject*> ObjectsToSync;
+		FAssetToolsModule& AssetToolsModule = FModuleManager::GetModuleChecked<FAssetToolsModule>("AssetTools");
 
-			for (UNiagaraEmitter* Emitter : Emitters)
+		const UContentBrowserAssetContextMenuContext* CBContext = UContentBrowserAssetContextMenuContext::FindContextWithAssets(InContext);
+		for (UNiagaraEmitter* Emitter : CBContext->LoadSelectedObjects<UNiagaraEmitter>())
+		{
+			if (Emitter->IsVersioningEnabled())
 			{
-				if (Emitter == nullptr)
-				{
-					continue;
-				}
 				FString Name;
 				FString PackageName;
-				CreateUniqueAssetName(Emitter->GetOutermost()->GetName(), DefaultSuffix, PackageName, Name);
+				AssetToolsModule.Get().CreateUniqueAssetName(Emitter->GetOutermost()->GetName(), DefaultSuffix, PackageName, Name);
 
 				// Create the factory used to generate the asset
+				UNiagaraEmitter* NewEmitter = Cast<UNiagaraEmitter>(AssetToolsModule.Get().DuplicateAssetWithDialog(Name, FPackageName::GetLongPackagePath(PackageName), Emitter));
 
-				if (UNiagaraEmitter* NewEmitter = Cast<UNiagaraEmitter>(AssetToolsModule.Get().DuplicateAssetWithDialog(Name, FPackageName::GetLongPackagePath(PackageName), Emitter)))
+				if (NewEmitter != nullptr)
 				{
 					GEditor->GetEditorSubsystem<UAssetEditorSubsystem>()->CloseAllEditorsForAsset(Emitter);
 
@@ -145,74 +120,75 @@ namespace MenuExtension_Emitter
 					ObjectsToSync.Add(NewEmitter);
 				}
 			}
+		}
 
-			if (ObjectsToSync.Num() > 0)
-			{
-				AssetToolsModule.Get().OpenEditorForAssets(ObjectsToSync);
+		if (ObjectsToSync.Num() > 0)
+		{
+			AssetToolsModule.Get().OpenEditorForAssets(ObjectsToSync);
 
-				FContentBrowserModule& ContentBrowserModule = FModuleManager::LoadModuleChecked<FContentBrowserModule>("ContentBrowser");
-				ContentBrowserModule.Get().SyncBrowserToAssets(ObjectsToSync);
-			}
+			FContentBrowserModule& ContentBrowserModule = FModuleManager::LoadModuleChecked<FContentBrowserModule>("ContentBrowser");
+			ContentBrowserModule.Get().SyncBrowserToAssets(ObjectsToSync);
 		}
 	}
-
-	static bool CanCreateDuplicateParent(const FToolMenuContext& MenuContext)
+	
+	bool CanExecuteCreateDuplicateParent(const FToolMenuContext& InContext)
 	{
-		if (const UContentBrowserAssetContextMenuContext* Context = UContentBrowserAssetContextMenuContext::FindContextWithAssets(MenuContext))
-		{
-			TArray<UNiagaraEmitter*> Emitters = Context->LoadSelectedObjects<UNiagaraEmitter>();
-			for (UNiagaraEmitter* EmitterPtr : Emitters)
-			{
-				if (EmitterPtr && EmitterPtr->IsVersioningEnabled())
-				{
-					return false;
-				}
-			}
-		}
+		const static FName NAME_VersioningEnabled("VersioningEnabled");
+	
+		const UContentBrowserAssetContextMenuContext* CBContext = UContentBrowserAssetContextMenuContext::FindContextWithAssets(InContext);
+        for (const FAssetData& SelectedAsset : CBContext->SelectedAssets)
+        {
+			bool VersioningEnabled = false;
+        	if (SelectedAsset.GetTagValue<bool>(NAME_VersioningEnabled, VersioningEnabled) && !VersioningEnabled)
+        	{
+        		return false;
+        	}
+        }
+ 
 		return true;
 	}
 
-	static void MarkDependentCompilableAssetsDirty(const FToolMenuContext& MenuContext)
+	static void ExecuteMarkDependentCompilableAssetsDirty(const FToolMenuContext& InContext)
 	{
-		if (const UContentBrowserAssetContextMenuContext* Context = UContentBrowserAssetContextMenuContext::FindContextWithAssets(MenuContext))
-		{
-			FNiagaraEditorUtilities::MarkDependentCompilableAssetsDirty(Context->SelectedAssets);
-		}
+		const UContentBrowserAssetContextMenuContext* CBContext = UContentBrowserAssetContextMenuContext::FindContextWithAssets(InContext);
+		FNiagaraEditorUtilities::MarkDependentCompilableAssetsDirty(CBContext->LoadSelectedObjects<UObject>());
 	}
 	
-	static FDelayedAutoRegisterHelper DelayedAutoRegister(EDelayedRegisterRunPhase::EndOfEngineInit, []
-	{
+	static FDelayedAutoRegisterHelper DelayedAutoRegister(EDelayedRegisterRunPhase::EndOfEngineInit, []{ 
 		UToolMenus::RegisterStartupCallback(FSimpleMulticastDelegate::FDelegate::CreateLambda([]()
 		{
 			FToolMenuOwnerScoped OwnerScoped(UE_MODULE_NAME);
 			UToolMenu* Menu = UE::ContentBrowser::ExtendToolMenu_AssetContextMenu(UNiagaraEmitter::StaticClass());
-	        
+		
 			FToolMenuSection& Section = Menu->FindOrAddSection("GetAssetActions");
-			Section.AddDynamicEntry("GetAssetActions_Emitter", FNewToolMenuSectionDelegate::CreateLambda([](FToolMenuSection& InSection)
+			Section.AddDynamicEntry(NAME_None, FNewToolMenuSectionDelegate::CreateLambda([](FToolMenuSection& InSection)
 			{
 				{
-					const TAttribute Label = LOCTEXT("Emitter_NewNiagaraSystem", "Create Niagara System");
-					const TAttribute ToolTip = LOCTEXT("Emitter_NewNiagaraSystemTooltip", "Creates a Niagara system using this emitter as a base.");
+					const TAttribute<FText> Label = LOCTEXT("Emitter_NewNiagaraSystem", "Create Niagara System");
+					const TAttribute<FText> ToolTip = LOCTEXT("Emitter_NewNiagaraSystemTooltip", "Creates a niagara system using this emitter as a base.");
 					const FSlateIcon Icon = FSlateIcon(FAppStyle::GetAppStyleSetName(), "ClassIcon.ParticleSystem");
-					const FToolMenuExecuteAction UIAction = FToolMenuExecuteAction::CreateStatic(&ExecuteNewNiagaraSystem);
 
+					FToolUIAction UIAction;
+					UIAction.ExecuteAction = FToolMenuExecuteAction::CreateStatic(&ExecuteNewNiagaraSystem);
 					InSection.AddMenuEntry("Emitter_NewNiagaraSystem", Label, ToolTip, Icon, UIAction);
 				}
 				{
-					const TAttribute Label = LOCTEXT("Emitter_CreateDuplicateParentLabel", "Create Duplicate Parent");
-					const TAttribute ToolTip = LOCTEXT("Emitter_CreateDuplicateParentTooltip", "Duplicate this emitter and set this emitter's parent to the new emitter.");
+					const TAttribute<FText> Label = LOCTEXT("Emitter_CreateDuplicateParent", "Create Duplicate Parent");
+					const TAttribute<FText> ToolTip = LOCTEXT("Emitter_CreateDuplicateParentTooltip", "Duplicate this emitter and set this emitter's parent to the new emitter.");
 					const FSlateIcon Icon = FSlateIcon();
-					FToolUIAction UIAction = FToolMenuExecuteAction::CreateStatic(&ExecuteCreateDuplicateParent);
-					UIAction.CanExecuteAction = FToolMenuCanExecuteAction::CreateStatic(&CanCreateDuplicateParent);
 
+					FToolUIAction UIAction;
+					UIAction.ExecuteAction = FToolMenuExecuteAction::CreateStatic(&ExecuteCreateDuplicateParent);
+					UIAction.CanExecuteAction = FToolMenuCanExecuteAction::CreateStatic(&CanExecuteCreateDuplicateParent);
 					InSection.AddMenuEntry("Emitter_CreateDuplicateParent", Label, ToolTip, Icon, UIAction);
 				}
 				{
-					const TAttribute Label = LOCTEXT("MarkDependentCompilableAssetsDirtyLabel", "Mark Dependent Compilable Assets Dirty");
-					const TAttribute ToolTip = LOCTEXT("MarkDependentCompilableAssetsDirtyToolTip", "Finds all niagara assets which depend on this asset either directly or indirectly,\n and marks them dirty so they can be saved with the latest version.");
+					const TAttribute<FText> Label = LOCTEXT("MarkDependentCompilableAssetsDirtyLabel", "Mark Dependent Compilable Assets Dirty");
+					const TAttribute<FText> ToolTip = LOCTEXT("MarkDependentCompilableAssetsDirtyToolTip", "Finds all niagara assets which depend on this asset either directly or indirectly, and marks them dirty so they can be saved with the latest version.");
 					const FSlateIcon Icon = FSlateIcon();
-					const FToolMenuExecuteAction UIAction = FToolMenuExecuteAction::CreateStatic(&MarkDependentCompilableAssetsDirty);
-
+					
+					FToolUIAction UIAction;
+					UIAction.ExecuteAction = FToolMenuExecuteAction::CreateStatic(&ExecuteMarkDependentCompilableAssetsDirty);
 					InSection.AddMenuEntry("MarkDependentCompilableAssetsDirty", Label, ToolTip, Icon, UIAction);
 				}
 			}));
diff --git a/Engine/Plugins/FX/Niagara/Source/NiagaraEditor/Private/AssetDefinitions/AssetDefinition_NiagaraEmitter.h b/Engine/Plugins/FX/Niagara/Source/NiagaraEditor/Private/AssetDefinitions/AssetDefinition_NiagaraEmitter.h
--- a/Engine/Plugins/FX/Niagara/Source/NiagaraEditor/Private/AssetDefinitions/AssetDefinition_NiagaraEmitter.h
+++ b/Engine/Plugins/FX/Niagara/Source/NiagaraEditor/Private/AssetDefinitions/AssetDefinition_NiagaraEmitter.h
@@ -3,6 +3,7 @@
 #pragma once
 
 #include "CoreMinimal.h"
+#include "NiagaraEmitter.h"
 #include "AssetDefinitionDefault.h"
 
 #include "AssetDefinition_NiagaraEmitter.generated.h"
@@ -13,10 +14,16 @@ class UAssetDefinition_NiagaraEmitter : public UAssetDefinitionDefault
 	GENERATED_BODY()
 
 public:
-	// UAssetDefinition Implementation
-	virtual FText GetAssetDisplayName() const override { return NSLOCTEXT("AssetDefinition", "NiagaraEmitter", "Niagara Emitter"); }
+	// UAssetDefinition Begin
+	virtual FText GetAssetDisplayName() const override { return NSLOCTEXT("AssetTypeActions", "AssetTypeActions_NiagaraEmitter", "Niagara Emitter"); }
 	virtual FLinearColor GetAssetColor() const override;
-	virtual TSoftClassPtr<UObject> GetAssetClass() const override;
-	virtual TConstArrayView<FAssetCategoryPath> GetAssetCategories() const override;
+	virtual TSoftClassPtr<UObject> GetAssetClass() const override { return UNiagaraEmitter::StaticClass(); }
+	virtual TConstArrayView<FAssetCategoryPath> GetAssetCategories() const override
+	{
+		static const auto Categories = { EAssetCategoryPaths::FX };
+		return Categories;
+	}
+
 	virtual EAssetCommandResult OpenAssets(const FAssetOpenArgs& OpenArgs) const override;
+	// UAssetDefinition End
 };
diff --git a/Engine/Plugins/FX/Niagara/Source/NiagaraEditor/Private/AssetDefinitions/AssetDefinition_NiagaraParameterCollection.cpp b/Engine/Plugins/FX/Niagara/Source/NiagaraEditor/Private/AssetDefinitions/AssetDefinition_NiagaraParameterCollection.cpp
new file mode 100644
--- /dev/null
+++ b/Engine/Plugins/FX/Niagara/Source/NiagaraEditor/Private/AssetDefinitions/AssetDefinition_NiagaraParameterCollection.cpp
@@ -0,0 +1,78 @@
+// Copyright Epic Games, Inc. All Rights Reserved.
+
+#include "AssetDefinition_NiagaraParameterCollection.h"
+
+#include "ContentBrowserMenuContexts.h"
+#include "NiagaraParameterCollection.h"
+#include "NiagaraParameterCollectionToolkit.h"
+#include "NiagaraParameterCollectionFactoryNew.h"
+#include "NiagaraEditorStyle.h"
+
+#include "IContentBrowserSingleton.h"
+#include "ContentBrowserModule.h"
+#include "IAssetTools.h"
+
+#include "ToolMenus.h"
+#include "Styling/AppStyle.h"
+
+#define LOCTEXT_NAMESPACE "AssetTypeActions_NiagaraParameterCollection"
+
+FLinearColor UAssetDefinition_NiagaraParameterCollection::GetAssetColor() const
+{
+	return FNiagaraEditorStyle::Get().GetColor("NiagaraEditor.AssetColors.ParameterCollection");
+}
+
+EAssetCommandResult UAssetDefinition_NiagaraParameterCollection::OpenAssets(const FAssetOpenArgs& OpenArgs) const
+{
+	for (UNiagaraParameterCollection* NPC : OpenArgs.LoadObjects<UNiagaraParameterCollection>())
+	{
+		TSharedRef< FNiagaraParameterCollectionToolkit > NewNiagaraNPCToolkit(new FNiagaraParameterCollectionToolkit());
+		NewNiagaraNPCToolkit->Initialize(OpenArgs.GetToolkitMode(), OpenArgs.ToolkitHost, NPC);
+	}
+
+	return EAssetCommandResult::Handled;
+}
+
+// Menu Extensions
+//--------------------------------------------------------------------
+
+namespace MenuExtension_NiagaraParameterCollection
+{
+	static void ExecuteCreateNewNPCI(const FToolMenuContext& InContext)
+	{
+		const UContentBrowserAssetContextMenuContext* CBContext = UContentBrowserAssetContextMenuContext::FindContextWithAssets(InContext);
+
+		IAssetTools::Get().CreateAssetsFrom<UNiagaraParameterCollection>(
+			CBContext->LoadSelectedObjects<UNiagaraParameterCollection>(), UNiagaraParameterCollectionInstance::StaticClass(), TEXT("_Inst"), [](UNiagaraParameterCollection* SourceObject)
+			{
+				UNiagaraParameterCollectionInstanceFactoryNew* Factory = NewObject<UNiagaraParameterCollectionInstanceFactoryNew>();
+				Factory->InitialParent = SourceObject;
+				return Factory;
+			}
+		);
+	}
+	
+	static FDelayedAutoRegisterHelper DelayedAutoRegister(EDelayedRegisterRunPhase::EndOfEngineInit, []{ 
+		UToolMenus::RegisterStartupCallback(FSimpleMulticastDelegate::FDelegate::CreateLambda([]()
+		{
+			FToolMenuOwnerScoped OwnerScoped(UE_MODULE_NAME);
+			UToolMenu* Menu = UE::ContentBrowser::ExtendToolMenu_AssetContextMenu(UNiagaraParameterCollection::StaticClass());
+		
+			FToolMenuSection& Section = Menu->FindOrAddSection("GetAssetActions");
+			Section.AddDynamicEntry(NAME_None, FNewToolMenuSectionDelegate::CreateLambda([](FToolMenuSection& InSection)
+			{
+				{
+					const TAttribute<FText> Label = LOCTEXT("NewNPC", "Create Niagara Parameter Collection Instance");
+					const TAttribute<FText> ToolTip = LOCTEXT("NewNPCTooltip", "Creates an instance of this Niagara Parameter Collection.");
+					const FSlateIcon Icon = FSlateIcon(FAppStyle::GetAppStyleSetName(), "ClassIcon.MaterialInstanceActor");
+					
+					FToolUIAction UIAction;
+					UIAction.ExecuteAction = FToolMenuExecuteAction::CreateStatic(&ExecuteCreateNewNPCI);
+					InSection.AddMenuEntry("Niagara_NewNPC", Label, ToolTip, Icon, UIAction);
+				}
+			}));
+		}));
+	});
+}
+
+#undef LOCTEXT_NAMESPACE
diff --git a/Engine/Plugins/FX/Niagara/Source/NiagaraEditor/Private/AssetDefinitions/AssetDefinition_NiagaraParameterCollection.h b/Engine/Plugins/FX/Niagara/Source/NiagaraEditor/Private/AssetDefinitions/AssetDefinition_NiagaraParameterCollection.h
new file mode 100644
--- /dev/null
+++ b/Engine/Plugins/FX/Niagara/Source/NiagaraEditor/Private/AssetDefinitions/AssetDefinition_NiagaraParameterCollection.h
@@ -0,0 +1,29 @@
+// Copyright Epic Games, Inc. All Rights Reserved.
+
+#pragma once
+
+#include "CoreMinimal.h"
+#include "NiagaraParameterCollection.h"
+#include "AssetDefinitionDefault.h"
+
+#include "AssetDefinition_NiagaraParameterCollection.generated.h"
+
+UCLASS()
+class UAssetDefinition_NiagaraParameterCollection : public UAssetDefinitionDefault
+{
+	GENERATED_BODY()
+
+public:
+	// UAssetDefinition Begin
+	virtual FText GetAssetDisplayName() const override { return NSLOCTEXT("AssetTypeActions", "AssetTypeActions_NiagaraParameterCollection", "Niagara Parameter Collection"); }
+	virtual FLinearColor GetAssetColor() const override;
+	virtual TSoftClassPtr<UObject> GetAssetClass() const override { return UNiagaraParameterCollection::StaticClass(); }
+	virtual TConstArrayView<FAssetCategoryPath> GetAssetCategories() const override
+	{
+		static const auto Categories = { EAssetCategoryPaths::FX / NSLOCTEXT("Niagara", "NiagaraAssetSubMenu_Advanced", "Advanced") };
+		return Categories;
+	}
+
+	virtual EAssetCommandResult OpenAssets(const FAssetOpenArgs& OpenArgs) const override;
+	// UAssetDefinition End
+};
diff --git a/Engine/Plugins/FX/Niagara/Source/NiagaraEditor/Private/AssetDefinitions/AssetDefinition_NiagaraParameterCollectionInstance.cpp b/Engine/Plugins/FX/Niagara/Source/NiagaraEditor/Private/AssetDefinitions/AssetDefinition_NiagaraParameterCollectionInstance.cpp
new file mode 100644
--- /dev/null
+++ b/Engine/Plugins/FX/Niagara/Source/NiagaraEditor/Private/AssetDefinitions/AssetDefinition_NiagaraParameterCollectionInstance.cpp
@@ -0,0 +1,26 @@
+// Copyright Epic Games, Inc. All Rights Reserved.
+
+#include "AssetDefinition_NiagaraParameterCollectionInstance.h"
+#include "NiagaraParameterCollection.h"
+#include "NiagaraParameterCollectionToolkit.h"
+#include "NiagaraEditorStyle.h"
+
+#define LOCTEXT_NAMESPACE "UAssetDefinition_NiagaraParameterCollectionInstance"
+
+FLinearColor UAssetDefinition_NiagaraParameterCollectionInstance::GetAssetColor() const
+{ 
+	return FNiagaraEditorStyle::Get().GetColor("NiagaraEditor.AssetColors.ParameterCollectionInstance"); 
+}
+
+EAssetCommandResult UAssetDefinition_NiagaraParameterCollectionInstance::OpenAssets(const FAssetOpenArgs& OpenArgs) const
+{
+	for (UNiagaraParameterCollectionInstance* Instance : OpenArgs.LoadObjects<UNiagaraParameterCollectionInstance>())
+	{
+		TSharedRef< FNiagaraParameterCollectionToolkit > NewNiagaraNPCToolkit(new FNiagaraParameterCollectionToolkit());
+		NewNiagaraNPCToolkit->Initialize(OpenArgs.GetToolkitMode(), OpenArgs.ToolkitHost, Instance);
+	}
+
+	return EAssetCommandResult::Handled;
+}
+
+#undef LOCTEXT_NAMESPACE
diff --git a/Engine/Plugins/FX/Niagara/Source/NiagaraEditor/Private/AssetDefinitions/AssetDefinition_NiagaraParameterCollectionInstance.h b/Engine/Plugins/FX/Niagara/Source/NiagaraEditor/Private/AssetDefinitions/AssetDefinition_NiagaraParameterCollectionInstance.h
new file mode 100644
--- /dev/null
+++ b/Engine/Plugins/FX/Niagara/Source/NiagaraEditor/Private/AssetDefinitions/AssetDefinition_NiagaraParameterCollectionInstance.h
@@ -0,0 +1,29 @@
+// Copyright Epic Games, Inc. All Rights Reserved.
+
+#pragma once
+
+#include "CoreMinimal.h"
+#include "NiagaraParameterCollection.h"
+#include "AssetDefinitionDefault.h"
+
+#include "AssetDefinition_NiagaraParameterCollectionInstance.generated.h"
+
+UCLASS()
+class UAssetDefinition_NiagaraParameterCollectionInstance : public UAssetDefinitionDefault
+{
+	GENERATED_BODY()
+
+public:
+	// UAssetDefinition Begin
+	virtual FText GetAssetDisplayName() const override { return NSLOCTEXT("AssetTypeActions", "AssetTypeActions_NiagaraParameterCollectionInstance", "Niagara Parameter Collection Instance"); }
+	virtual FLinearColor GetAssetColor() const override;
+	virtual TSoftClassPtr<UObject> GetAssetClass() const override { return UNiagaraParameterCollectionInstance::StaticClass(); }
+	virtual TConstArrayView<FAssetCategoryPath> GetAssetCategories() const override
+	{
+		static const auto Categories = { EAssetCategoryPaths::FX / NSLOCTEXT("Niagara", "NiagaraAssetSubMenu_Advanced", "Advanced") };
+		return Categories;
+	}
+
+	virtual EAssetCommandResult OpenAssets(const FAssetOpenArgs& OpenArgs) const override;
+	// UAssetDefinition End
+};
diff --git a/Engine/Plugins/FX/Niagara/Source/NiagaraEditor/Private/AssetDefinitions/AssetDefinition_NiagaraParameterDefinitions.cpp b/Engine/Plugins/FX/Niagara/Source/NiagaraEditor/Private/AssetDefinitions/AssetDefinition_NiagaraParameterDefinitions.cpp
--- a/Engine/Plugins/FX/Niagara/Source/NiagaraEditor/Private/AssetDefinitions/AssetDefinition_NiagaraParameterDefinitions.cpp
+++ b/Engine/Plugins/FX/Niagara/Source/NiagaraEditor/Private/AssetDefinitions/AssetDefinition_NiagaraParameterDefinitions.cpp
@@ -1,47 +1,21 @@
 // Copyright Epic Games, Inc. All Rights Reserved.
 
-#include "AssetDefinitions/AssetDefinition_NiagaraParameterDefinitions.h"
-
+#include "AssetDefinition_NiagaraParameterDefinitions.h"
 #include "NiagaraEditorStyle.h"
-#include "NiagaraParameterDefinitions.h"
-#include "NiagaraParameterDefinitionsToolkit.h"
-
-#define LOCTEXT_NAMESPACE "UAssetDefinition_NiagaraParameterDefinitions"
+#include "Toolkits/NiagaraParameterDefinitionsToolkit.h"
 
 FLinearColor UAssetDefinition_NiagaraParameterDefinitions::GetAssetColor() const
 {
 	return FNiagaraEditorStyle::Get().GetColor("NiagaraEditor.AssetColors.ParameterDefinitions");
 }
 
-TSoftClassPtr<> UAssetDefinition_NiagaraParameterDefinitions::GetAssetClass() const
-{
-	return UNiagaraParameterDefinitions::StaticClass();
-}
-
-TConstArrayView<FAssetCategoryPath> UAssetDefinition_NiagaraParameterDefinitions::GetAssetCategories() const
-{
-	static FAssetCategoryPath AssetPaths[] = { FAssetCategoryPath(LOCTEXT("NiagaraAssetsCategory", "FX"), LOCTEXT("NiagaraParameterDefinitions_SubCategory", "Advanced")) };
-	return AssetPaths;
-}
-
 EAssetCommandResult UAssetDefinition_NiagaraParameterDefinitions::OpenAssets(const FAssetOpenArgs& OpenArgs) const
 {
-	if (OpenArgs.OpenMethod == EAssetOpenMethod::Edit)
-	{
-		EToolkitMode::Type Mode = OpenArgs.ToolkitHost.IsValid() ? EToolkitMode::WorldCentric : EToolkitMode::Standalone;
-		TArray<UObject*> Objects = OpenArgs.LoadObjects<UObject>();
-		for (UObject* OpenObj : Objects)
-		{
-			if (UNiagaraParameterDefinitions* ParameterDefinitions = Cast<UNiagaraParameterDefinitions>(OpenObj))
-			{
-				TSharedRef< FNiagaraParameterDefinitionsToolkit > NewNiagaraParameterDefinitionsToolkit(new FNiagaraParameterDefinitionsToolkit());
-				NewNiagaraParameterDefinitionsToolkit->Initialize(Mode, OpenArgs.ToolkitHost, ParameterDefinitions);
-			}
-		}
-		return EAssetCommandResult::Handled;
+	for (UNiagaraParameterDefinitions* ParameterDefinitions : OpenArgs.LoadObjects<UNiagaraParameterDefinitions>())
+   	{
+		TSharedRef<FNiagaraParameterDefinitionsToolkit> NewNiagaraParameterDefinitionsToolkit(new FNiagaraParameterDefinitionsToolkit());
+		NewNiagaraParameterDefinitionsToolkit->Initialize(OpenArgs.GetToolkitMode(), OpenArgs.ToolkitHost, ParameterDefinitions);
 	}
 
-	return EAssetCommandResult::Unhandled;
+	return EAssetCommandResult::Handled;
 }
-
-#undef LOCTEXT_NAMESPACE
diff --git a/Engine/Plugins/FX/Niagara/Source/NiagaraEditor/Private/AssetDefinitions/AssetDefinition_NiagaraParameterDefinitions.h b/Engine/Plugins/FX/Niagara/Source/NiagaraEditor/Private/AssetDefinitions/AssetDefinition_NiagaraParameterDefinitions.h
--- a/Engine/Plugins/FX/Niagara/Source/NiagaraEditor/Private/AssetDefinitions/AssetDefinition_NiagaraParameterDefinitions.h
+++ b/Engine/Plugins/FX/Niagara/Source/NiagaraEditor/Private/AssetDefinitions/AssetDefinition_NiagaraParameterDefinitions.h
@@ -3,6 +3,7 @@
 #pragma once
 
 #include "CoreMinimal.h"
+#include "NiagaraParameterDefinitions.h"
 #include "AssetDefinitionDefault.h"
 
 #include "AssetDefinition_NiagaraParameterDefinitions.generated.h"
@@ -13,10 +14,16 @@ class UAssetDefinition_NiagaraParameterDefinitions : public UAssetDefinitionDefa
 	GENERATED_BODY()
 
 public:
-	// UAssetDefinition Implementation
-	virtual FText GetAssetDisplayName() const override { return NSLOCTEXT("AssetDefinition", "NiagaraParameterDefinitions", "Niagara Parameter Definitions"); }
+	// UAssetDefinition Begin
+	virtual FText GetAssetDisplayName() const override { return NSLOCTEXT("AssetTypeActions", "AssetTypeActions_NiagaraParameterDefinitions", "Niagara Parameter Definitions"); }
 	virtual FLinearColor GetAssetColor() const override;
-	virtual TSoftClassPtr<UObject> GetAssetClass() const override;
-	virtual TConstArrayView<FAssetCategoryPath> GetAssetCategories() const override;
+	virtual TSoftClassPtr<UObject> GetAssetClass() const override { return UNiagaraParameterDefinitions::StaticClass(); }
+	virtual TConstArrayView<FAssetCategoryPath> GetAssetCategories() const override
+	{
+		static const auto Categories = { EAssetCategoryPaths::FX / NSLOCTEXT("Niagara", "NiagaraAssetSubMenu_Advanced", "Advanced") };
+		return Categories;
+	}
+
 	virtual EAssetCommandResult OpenAssets(const FAssetOpenArgs& OpenArgs) const override;
+	// UAssetDefinition End
 };
diff --git a/Engine/Plugins/FX/Niagara/Source/NiagaraEditor/Private/AssetDefinitions/AssetDefinition_NiagaraScript.cpp b/Engine/Plugins/FX/Niagara/Source/NiagaraEditor/Private/AssetDefinitions/AssetDefinition_NiagaraScript.cpp
new file mode 100644
--- /dev/null
+++ b/Engine/Plugins/FX/Niagara/Source/NiagaraEditor/Private/AssetDefinitions/AssetDefinition_NiagaraScript.cpp
@@ -0,0 +1,98 @@
+// Copyright Epic Games, Inc. All Rights Reserved.
+
+#include "AssetDefinition_NiagaraScript.h"
+#include "NiagaraScriptToolkit.h"
+#include "NiagaraEditorStyle.h"
+#include "AssetRegistry/AssetData.h"
+#include "NiagaraEditorUtilities.h"
+#include "ContentBrowserMenuContexts.h"
+#include "ToolMenus.h"
+
+#define LOCTEXT_NAMESPACE "NiagaraScriptAssetTypeActions"
+
+const FName UAssetDefinition_NiagaraScript::FunctionScriptName = "Niagara Function Script";
+const FText UAssetDefinition_NiagaraScript::FunctionScriptNameText = LOCTEXT("NiagaraFunctionScript", "Niagara Function Script");
+const FName UAssetDefinition_NiagaraScript::ModuleScriptName = "Niagara Module Script";
+const FText UAssetDefinition_NiagaraScript::ModuleScriptNameText = LOCTEXT("NiagaraModuleScript", "Niagara Module Script");
+const FName UAssetDefinition_NiagaraScript::DynamicInputScriptName = "Niagara Dynamic Input Script";
+const FText UAssetDefinition_NiagaraScript::DynamicInputScriptNameText = LOCTEXT("NiagaraDynamicInputScript", "Niagara Dynamic Input Script");
+
+FLinearColor UAssetDefinition_NiagaraScript::GetAssetColor() const
+{ 
+	return FNiagaraEditorStyle::Get().GetColor("NiagaraEditor.AssetColors.Script"); 
+}
+
+FText UAssetDefinition_NiagaraScript::GetAssetDisplayName(const FAssetData& AssetData) const
+{
+	static const FName NAME_Usage(TEXT("Usage"));
+	const FAssetDataTagMapSharedView::FFindTagResult Usage = AssetData.TagsAndValues.FindTag(NAME_Usage);
+
+	if (Usage.IsSet())
+	{
+		static const FString FunctionString(TEXT("Function"));
+		static const FString ModuleString(TEXT("Module"));
+		static const FString DynamicInputString(TEXT("DynamicInput"));
+		
+		if (Usage.GetValue() == FunctionString)
+		{
+			return FunctionScriptNameText;
+		}
+		else if (Usage.GetValue() == ModuleString)
+		{
+			return ModuleScriptNameText;
+		}
+		else if (Usage.GetValue() == DynamicInputString)
+		{
+			return DynamicInputScriptNameText;
+		}
+	}
+
+	return Super::GetAssetDisplayName(AssetData);
+}
+
+EAssetCommandResult UAssetDefinition_NiagaraScript::OpenAssets(const FAssetOpenArgs& OpenArgs) const
+{
+	for (UNiagaraScript* Script : OpenArgs.LoadObjects<UNiagaraScript>())
+	{
+		TSharedRef<FNiagaraScriptToolkit> NewNiagaraScriptToolkit(new FNiagaraScriptToolkit());
+		NewNiagaraScriptToolkit->Initialize(OpenArgs.GetToolkitMode(), OpenArgs.ToolkitHost, Script);
+	}
+	
+	return EAssetCommandResult::Handled;
+}
+
+// Menu Extensions
+//--------------------------------------------------------------------
+
+namespace MenuExtension_NiagaraScript
+{
+	static void ExecuteMarkDependentCompilableAssetsDirty(const FToolMenuContext& InContext)
+	{
+		const UContentBrowserAssetContextMenuContext* CBContext = UContentBrowserAssetContextMenuContext::FindContextWithAssets(InContext);
+		FNiagaraEditorUtilities::MarkDependentCompilableAssetsDirty(CBContext->LoadSelectedObjects<UObject>());
+	}
+	
+	static FDelayedAutoRegisterHelper DelayedAutoRegister(EDelayedRegisterRunPhase::EndOfEngineInit, []{ 
+		UToolMenus::RegisterStartupCallback(FSimpleMulticastDelegate::FDelegate::CreateLambda([]()
+		{
+			FToolMenuOwnerScoped OwnerScoped(UE_MODULE_NAME);
+			UToolMenu* Menu = UE::ContentBrowser::ExtendToolMenu_AssetContextMenu(UNiagaraScript::StaticClass());
+		
+			FToolMenuSection& Section = Menu->FindOrAddSection("GetAssetActions");
+			Section.AddDynamicEntry(NAME_None, FNewToolMenuSectionDelegate::CreateLambda([](FToolMenuSection& InSection)
+			{
+				{
+					const TAttribute<FText> Label = LOCTEXT("MarkDependentCompilableAssetsDirtyLabel", "Mark Dependent Compilable Assets Dirty");
+					const TAttribute<FText> ToolTip = LOCTEXT("MarkDependentCompilableAssetsDirtyToolTip", "Finds all niagara assets which depend on this asset either directly or indirectly, and marks them dirty so they can be saved with the latest version.");
+					const FSlateIcon Icon = FSlateIcon();
+					
+					FToolUIAction UIAction;
+					UIAction.ExecuteAction = FToolMenuExecuteAction::CreateStatic(&ExecuteMarkDependentCompilableAssetsDirty);
+					InSection.AddMenuEntry("MarkDependentCompilableAssetsDirty", Label, ToolTip, Icon, UIAction);
+				}
+			}));
+		}));
+	});
+}
+
+#undef LOCTEXT_NAMESPACE
diff --git a/Engine/Plugins/FX/Niagara/Source/NiagaraEditor/Private/AssetDefinitions/AssetDefinition_NiagaraScript.h b/Engine/Plugins/FX/Niagara/Source/NiagaraEditor/Private/AssetDefinitions/AssetDefinition_NiagaraScript.h
new file mode 100644
--- /dev/null
+++ b/Engine/Plugins/FX/Niagara/Source/NiagaraEditor/Private/AssetDefinitions/AssetDefinition_NiagaraScript.h
@@ -0,0 +1,40 @@
+// Copyright Epic Games, Inc. All Rights Reserved.
+
+#pragma once
+
+#include "CoreMinimal.h"
+#include "NiagaraScript.h"
+#include "AssetDefinitionDefault.h"
+
+#include "AssetDefinition_NiagaraScript.generated.h"
+
+UCLASS()
+class UAssetDefinition_NiagaraScript : public UAssetDefinitionDefault
+{
+	GENERATED_BODY()
+
+public:
+	static const FName FunctionScriptName;
+	static const FText FunctionScriptNameText;
+
+	static const FName ModuleScriptName;
+	static const FText ModuleScriptNameText;
+	
+	static const FName DynamicInputScriptName;
+	static const FText DynamicInputScriptNameText;
+
+public:
+	// UAssetDefinition Begin
+	virtual FText GetAssetDisplayName() const override { return NSLOCTEXT("AssetTypeActions", "AssetTypeActions_NiagaraScript", "Niagara Script"); }
+	virtual FText GetAssetDisplayName(const FAssetData& AssetData) const override;
+	virtual FLinearColor GetAssetColor() const override;
+	virtual TSoftClassPtr<UObject> GetAssetClass() const override { return UNiagaraScript::StaticClass(); }
+	virtual TConstArrayView<FAssetCategoryPath> GetAssetCategories() const override
+	{
+		static const auto Categories = { EAssetCategoryPaths::FX / NSLOCTEXT("Niagara", "NiagaraAssetSubMenu_Script", "Script") };
+		return Categories;
+	}
+
+	virtual EAssetCommandResult OpenAssets(const FAssetOpenArgs& OpenArgs) const override;
+	// UAssetDefinition End
+};
diff --git a/Engine/Plugins/FX/Niagara/Source/NiagaraEditor/Private/AssetDefinitions/AssetDefinition_NiagaraSimCache.cpp b/Engine/Plugins/FX/Niagara/Source/NiagaraEditor/Private/AssetDefinitions/AssetDefinition_NiagaraSimCache.cpp
--- a/Engine/Plugins/FX/Niagara/Source/NiagaraEditor/Private/AssetDefinitions/AssetDefinition_NiagaraSimCache.cpp
+++ b/Engine/Plugins/FX/Niagara/Source/NiagaraEditor/Private/AssetDefinitions/AssetDefinition_NiagaraSimCache.cpp
@@ -1,47 +1,22 @@
 // Copyright Epic Games, Inc. All Rights Reserved.
 
-#include "AssetDefinitions/AssetDefinition_NiagaraSimCache.h"
+#include "AssetDefinition_NiagaraSimCache.h"
 
 #include "NiagaraEditorStyle.h"
-#include "NiagaraSimCache.h"
 #include "NiagaraSimCacheToolkit.h"
 
-#define LOCTEXT_NAMESPACE "UAssetDefinition_NiagaraSimCache"
-
 FLinearColor UAssetDefinition_NiagaraSimCache::GetAssetColor() const
 {
-	return FNiagaraEditorStyle::Get().GetColor("NiagaraEditor.AssetColors.SimCache");
-}
-
-TSoftClassPtr<> UAssetDefinition_NiagaraSimCache::GetAssetClass() const
-{
-	return UNiagaraSimCache::StaticClass();
-}
-
-TConstArrayView<FAssetCategoryPath> UAssetDefinition_NiagaraSimCache::GetAssetCategories() const
-{
-	static FAssetCategoryPath AssetPaths[] = { FAssetCategoryPath(LOCTEXT("NiagaraAssetsCategory", "FX"), LOCTEXT("NiagaraSimCache_SubCategory", "Advanced")) };
-	return AssetPaths;
+	return FNiagaraEditorStyle::Get().GetColor("NiagaraEditor.AssetColors.SimCache").ToFColor(true);
 }
 
 EAssetCommandResult UAssetDefinition_NiagaraSimCache::OpenAssets(const FAssetOpenArgs& OpenArgs) const
 {
-	if (OpenArgs.OpenMethod == EAssetOpenMethod::Edit)
+	for (UNiagaraSimCache* SimCache : OpenArgs.LoadObjects<UNiagaraSimCache>())
 	{
-		EToolkitMode::Type Mode = OpenArgs.ToolkitHost.IsValid() ? EToolkitMode::WorldCentric : EToolkitMode::Standalone;
-		TArray<UObject*> Objects = OpenArgs.LoadObjects<UObject>();
-		for (UObject* OpenObj : Objects)
-		{
-			if (UNiagaraSimCache* Cache = Cast<UNiagaraSimCache>(OpenObj))
-			{
-				TSharedRef<FNiagaraSimCacheToolkit> NewNiagaraSystemToolkit(new FNiagaraSimCacheToolkit());
-				NewNiagaraSystemToolkit->Initialize(Mode, OpenArgs.ToolkitHost, Cache);
-			}
-		}
-		return EAssetCommandResult::Handled;
+		const TSharedRef< FNiagaraSimCacheToolkit > NewNiagaraSimCacheToolkit(new FNiagaraSimCacheToolkit());
+		NewNiagaraSimCacheToolkit->Initialize(OpenArgs.GetToolkitMode(), OpenArgs.ToolkitHost, SimCache);
 	}
 
-	return EAssetCommandResult::Unhandled;
+	return EAssetCommandResult::Handled;
 }
-
-#undef LOCTEXT_NAMESPACE
diff --git a/Engine/Plugins/FX/Niagara/Source/NiagaraEditor/Private/AssetDefinitions/AssetDefinition_NiagaraSimCache.h b/Engine/Plugins/FX/Niagara/Source/NiagaraEditor/Private/AssetDefinitions/AssetDefinition_NiagaraSimCache.h
--- a/Engine/Plugins/FX/Niagara/Source/NiagaraEditor/Private/AssetDefinitions/AssetDefinition_NiagaraSimCache.h
+++ b/Engine/Plugins/FX/Niagara/Source/NiagaraEditor/Private/AssetDefinitions/AssetDefinition_NiagaraSimCache.h
@@ -3,6 +3,7 @@
 #pragma once
 
 #include "CoreMinimal.h"
+#include "NiagaraSimCache.h"
 #include "AssetDefinitionDefault.h"
 
 #include "AssetDefinition_NiagaraSimCache.generated.h"
@@ -13,10 +14,16 @@ class UAssetDefinition_NiagaraSimCache : public UAssetDefinitionDefault
 	GENERATED_BODY()
 
 public:
-	// UAssetDefinition Implementation
-	virtual FText GetAssetDisplayName() const override { return NSLOCTEXT("AssetDefinition", "NiagaraSimCache", "Niagara Simulation Cache"); }
+	// UAssetDefinition Begin
+	virtual FText GetAssetDisplayName() const override { return NSLOCTEXT("AssetTypeActions", "AssetTypeActions_NiagaraSimCache", "Niagara Simulation Cache"); }
 	virtual FLinearColor GetAssetColor() const override;
-	virtual TSoftClassPtr<UObject> GetAssetClass() const override;
-	virtual TConstArrayView<FAssetCategoryPath> GetAssetCategories() const override;
+	virtual TSoftClassPtr<UObject> GetAssetClass() const override { return UNiagaraSimCache::StaticClass(); }
+	virtual TConstArrayView<FAssetCategoryPath> GetAssetCategories() const override
+	{
+		static const auto Categories = { EAssetCategoryPaths::FX / NSLOCTEXT("Niagara", "NiagaraAssetSubMenu_Advanced", "Advanced") };
+		return Categories;
+	}
+
 	virtual EAssetCommandResult OpenAssets(const FAssetOpenArgs& OpenArgs) const override;
+	// UAssetDefinition End
 };
diff --git a/Engine/Plugins/FX/Niagara/Source/NiagaraEditor/Private/AssetDefinitions/AssetDefinition_NiagaraSystem.cpp b/Engine/Plugins/FX/Niagara/Source/NiagaraEditor/Private/AssetDefinitions/AssetDefinition_NiagaraSystem.cpp
--- a/Engine/Plugins/FX/Niagara/Source/NiagaraEditor/Private/AssetDefinitions/AssetDefinition_NiagaraSystem.cpp
+++ b/Engine/Plugins/FX/Niagara/Source/NiagaraEditor/Private/AssetDefinitions/AssetDefinition_NiagaraSystem.cpp
@@ -1,47 +1,22 @@
 // Copyright Epic Games, Inc. All Rights Reserved.
 
-#include "AssetDefinitions/AssetDefinition_NiagaraSystem.h"
-
-#include "NiagaraEditorStyle.h"
+#include "AssetDefinition_NiagaraSystem.h"
 #include "NiagaraSystem.h"
 #include "NiagaraSystemToolkit.h"
-
-#define LOCTEXT_NAMESPACE "UAssetDefinition_NiagaraSystem"
+#include "NiagaraEditorStyle.h"
 
 FLinearColor UAssetDefinition_NiagaraSystem::GetAssetColor() const
 {
 	return FNiagaraEditorStyle::Get().GetColor("NiagaraEditor.AssetColors.System");
 }
 
-TSoftClassPtr<> UAssetDefinition_NiagaraSystem::GetAssetClass() const
-{
-	return UNiagaraSystem::StaticClass();
-}
-
-TConstArrayView<FAssetCategoryPath> UAssetDefinition_NiagaraSystem::GetAssetCategories() const
-{
-	static FAssetCategoryPath AssetPaths[] = { LOCTEXT("NiagaraAssetsCategory", "FX"), EAssetCategoryPaths::Basic };
-	return AssetPaths;
-}
-
 EAssetCommandResult UAssetDefinition_NiagaraSystem::OpenAssets(const FAssetOpenArgs& OpenArgs) const
 {
-	if (OpenArgs.OpenMethod == EAssetOpenMethod::Edit)
+	for (UNiagaraSystem* System : OpenArgs.LoadObjects<UNiagaraSystem>())
 	{
-		EToolkitMode::Type Mode = OpenArgs.ToolkitHost.IsValid() ? EToolkitMode::WorldCentric : EToolkitMode::Standalone;
-		TArray<UObject*> Objects = OpenArgs.LoadObjects<UObject>();
-		for (UObject* OpenObj : Objects)
-		{
-			if (UNiagaraSystem* System = Cast<UNiagaraSystem>(OpenObj))
-			{
-				TSharedRef<FNiagaraSystemToolkit> NewNiagaraSystemToolkit(new FNiagaraSystemToolkit());
-				NewNiagaraSystemToolkit->InitializeWithSystem(Mode, OpenArgs.ToolkitHost, *System);
-			}
-		}
-		return EAssetCommandResult::Handled;
+		TSharedRef<FNiagaraSystemToolkit> NiagaraSystemToolkit = MakeShared<FNiagaraSystemToolkit>();
+		NiagaraSystemToolkit->InitializeWithSystem(OpenArgs.GetToolkitMode(), OpenArgs.ToolkitHost, *System);
 	}
 
-	return EAssetCommandResult::Unhandled;
+	return EAssetCommandResult::Handled;
 }
-
-#undef LOCTEXT_NAMESPACE
diff --git a/Engine/Plugins/FX/Niagara/Source/NiagaraEditor/Private/AssetDefinitions/AssetDefinition_NiagaraSystem.h b/Engine/Plugins/FX/Niagara/Source/NiagaraEditor/Private/AssetDefinitions/AssetDefinition_NiagaraSystem.h
--- a/Engine/Plugins/FX/Niagara/Source/NiagaraEditor/Private/AssetDefinitions/AssetDefinition_NiagaraSystem.h
+++ b/Engine/Plugins/FX/Niagara/Source/NiagaraEditor/Private/AssetDefinitions/AssetDefinition_NiagaraSystem.h
@@ -3,6 +3,7 @@
 #pragma once
 
 #include "CoreMinimal.h"
+#include "NiagaraSystem.h"
 #include "AssetDefinitionDefault.h"
 
 #include "AssetDefinition_NiagaraSystem.generated.h"
@@ -13,10 +14,16 @@ class UAssetDefinition_NiagaraSystem : public UAssetDefinitionDefault
 	GENERATED_BODY()
 
 public:
-	// UAssetDefinition Implementation
-	virtual FText GetAssetDisplayName() const override { return NSLOCTEXT("AssetDefinition", "NiagaraSystem", "Niagara System"); }
+	// UAssetDefinition Begin
+	virtual FText GetAssetDisplayName() const override { return NSLOCTEXT("AssetTypeActions", "AssetTypeActions_NiagaraSystem", "Niagara System"); }
 	virtual FLinearColor GetAssetColor() const override;
-	virtual TSoftClassPtr<UObject> GetAssetClass() const override;
-	virtual TConstArrayView<FAssetCategoryPath> GetAssetCategories() const override;
+	virtual TSoftClassPtr<UObject> GetAssetClass() const override { return UNiagaraSystem::StaticClass(); }
+	virtual TConstArrayView<FAssetCategoryPath> GetAssetCategories() const override
+	{
+		static const auto Categories = { EAssetCategoryPaths::Basic, EAssetCategoryPaths::FX };
+		return Categories;
+	}
+
 	virtual EAssetCommandResult OpenAssets(const FAssetOpenArgs& OpenArgs) const override;
+	// UAssetDefinition End
 };
diff --git a/Engine/Plugins/FX/Niagara/Source/NiagaraEditor/Private/AssetDefinitions/AssetDefinition_NiagaraValidationRuleSet.cpp b/Engine/Plugins/FX/Niagara/Source/NiagaraEditor/Private/AssetDefinitions/AssetDefinition_NiagaraValidationRuleSet.cpp
--- a/Engine/Plugins/FX/Niagara/Source/NiagaraEditor/Private/AssetDefinitions/AssetDefinition_NiagaraValidationRuleSet.cpp
+++ b/Engine/Plugins/FX/Niagara/Source/NiagaraEditor/Private/AssetDefinitions/AssetDefinition_NiagaraValidationRuleSet.cpp
@@ -1,26 +1,9 @@
 // Copyright Epic Games, Inc. All Rights Reserved.
 
-#include "AssetDefinitions/AssetDefinition_NiagaraValidationRuleSet.h"
-
+#include "AssetDefinition_NiagaraValidationRuleSet.h"
 #include "NiagaraEditorStyle.h"
-#include "NiagaraValidationRuleSet.h"
-
-#define LOCTEXT_NAMESPACE "UAssetDefinition_NiagaraValidationRuleSet"
 
 FLinearColor UAssetDefinition_NiagaraValidationRuleSet::GetAssetColor() const
 {
 	return FNiagaraEditorStyle::Get().GetColor("NiagaraEditor.AssetColors.ValidationRuleSet");
 }
-
-TSoftClassPtr<> UAssetDefinition_NiagaraValidationRuleSet::GetAssetClass() const
-{
-	return UNiagaraValidationRuleSet::StaticClass();
-}
-
-TConstArrayView<FAssetCategoryPath> UAssetDefinition_NiagaraValidationRuleSet::GetAssetCategories() const
-{
-	static FAssetCategoryPath AssetPaths[] = { FAssetCategoryPath(LOCTEXT("NiagaraAssetsCategory", "FX"), LOCTEXT("NiagaraValidationRuleSet_SubCategory", "Advanced")) };
-	return AssetPaths;
-}
-
-#undef LOCTEXT_NAMESPACE
diff --git a/Engine/Plugins/FX/Niagara/Source/NiagaraEditor/Private/AssetDefinitions/AssetDefinition_NiagaraValidationRuleSet.h b/Engine/Plugins/FX/Niagara/Source/NiagaraEditor/Private/AssetDefinitions/AssetDefinition_NiagaraValidationRuleSet.h
--- a/Engine/Plugins/FX/Niagara/Source/NiagaraEditor/Private/AssetDefinitions/AssetDefinition_NiagaraValidationRuleSet.h
+++ b/Engine/Plugins/FX/Niagara/Source/NiagaraEditor/Private/AssetDefinitions/AssetDefinition_NiagaraValidationRuleSet.h
@@ -3,6 +3,7 @@
 #pragma once
 
 #include "CoreMinimal.h"
+#include "NiagaraValidationRuleSet.h"
 #include "AssetDefinitionDefault.h"
 
 #include "AssetDefinition_NiagaraValidationRuleSet.generated.h"
@@ -13,9 +14,17 @@ class UAssetDefinition_NiagaraValidationRuleSet : public UAssetDefinitionDefault
 	GENERATED_BODY()
 
 public:
-	// UAssetDefinition Implementation
-	virtual FText GetAssetDisplayName() const override { return NSLOCTEXT("AssetDefinition", "NiagaraValidationRuleSet", "Niagara Validation Rule Set"); }
+	// UAssetDefinition Begin
+	virtual FText GetAssetDisplayName() const override { return NSLOCTEXT("AssetTypeActions", "AssetTypeActions_NiagaraValidationRuleSet", "Niagara Validation Rule Set"); }
 	virtual FLinearColor GetAssetColor() const override;
-	virtual TSoftClassPtr<UObject> GetAssetClass() const override;
-	virtual TConstArrayView<FAssetCategoryPath> GetAssetCategories() const override;
+	virtual TSoftClassPtr<UObject> GetAssetClass() const override { return UNiagaraValidationRuleSet::StaticClass(); }
+	virtual TConstArrayView<FAssetCategoryPath> GetAssetCategories() const override
+	{
+		static const auto Categories = {
+			EAssetCategoryPaths::FX / NSLOCTEXT("Niagara", "NiagaraAssetSubMenu_Advanced", "Advanced")
+		};
+		
+		return Categories;
+	}
+	// UAssetDefinition End
 };
diff --git a/Engine/Plugins/FX/Niagara/Source/NiagaraEditor/Private/AssetTypeActions/AssetDefinition_NiagaraEffectType.cpp b/Engine/Plugins/FX/Niagara/Source/NiagaraEditor/Private/AssetTypeActions/AssetDefinition_NiagaraEffectType.cpp
deleted file mode 100644
diff --git a/Engine/Plugins/FX/Niagara/Source/NiagaraEditor/Private/AssetTypeActions/AssetDefinition_NiagaraEffectType.h b/Engine/Plugins/FX/Niagara/Source/NiagaraEditor/Private/AssetTypeActions/AssetDefinition_NiagaraEffectType.h
deleted file mode 100644
diff --git a/Engine/Plugins/FX/Niagara/Source/NiagaraEditor/Private/AssetTypeActions/AssetDefinition_NiagaraEmitter.cpp b/Engine/Plugins/FX/Niagara/Source/NiagaraEditor/Private/AssetTypeActions/AssetDefinition_NiagaraEmitter.cpp
deleted file mode 100644
diff --git a/Engine/Plugins/FX/Niagara/Source/NiagaraEditor/Private/AssetTypeActions/AssetDefinition_NiagaraEmitter.h b/Engine/Plugins/FX/Niagara/Source/NiagaraEditor/Private/AssetTypeActions/AssetDefinition_NiagaraEmitter.h
deleted file mode 100644
diff --git a/Engine/Plugins/FX/Niagara/Source/NiagaraEditor/Private/AssetTypeActions/AssetDefinition_NiagaraParameterCollection.cpp b/Engine/Plugins/FX/Niagara/Source/NiagaraEditor/Private/AssetTypeActions/AssetDefinition_NiagaraParameterCollection.cpp
deleted file mode 100644
diff --git a/Engine/Plugins/FX/Niagara/Source/NiagaraEditor/Private/AssetTypeActions/AssetDefinition_NiagaraParameterCollection.h b/Engine/Plugins/FX/Niagara/Source/NiagaraEditor/Private/AssetTypeActions/AssetDefinition_NiagaraParameterCollection.h
deleted file mode 100644
diff --git a/Engine/Plugins/FX/Niagara/Source/NiagaraEditor/Private/AssetTypeActions/AssetDefinition_NiagaraParameterCollectionInstance.cpp b/Engine/Plugins/FX/Niagara/Source/NiagaraEditor/Private/AssetTypeActions/AssetDefinition_NiagaraParameterCollectionInstance.cpp
deleted file mode 100644
diff --git a/Engine/Plugins/FX/Niagara/Source/NiagaraEditor/Private/AssetTypeActions/AssetDefinition_NiagaraParameterCollectionInstance.h b/Engine/Plugins/FX/Niagara/Source/NiagaraEditor/Private/AssetTypeActions/AssetDefinition_NiagaraParameterCollectionInstance.h
deleted file mode 100644
diff --git a/Engine/Plugins/FX/Niagara/Source/NiagaraEditor/Private/AssetTypeActions/AssetDefinition_NiagaraParameterDefinitions.cpp b/Engine/Plugins/FX/Niagara/Source/NiagaraEditor/Private/AssetTypeActions/AssetDefinition_NiagaraParameterDefinitions.cpp
deleted file mode 100644
diff --git a/Engine/Plugins/FX/Niagara/Source/NiagaraEditor/Private/AssetTypeActions/AssetDefinition_NiagaraParameterDefinitions.h b/Engine/Plugins/FX/Niagara/Source/NiagaraEditor/Private/AssetTypeActions/AssetDefinition_NiagaraParameterDefinitions.h
deleted file mode 100644
diff --git a/Engine/Plugins/FX/Niagara/Source/NiagaraEditor/Private/AssetTypeActions/AssetDefinition_NiagaraScript.cpp b/Engine/Plugins/FX/Niagara/Source/NiagaraEditor/Private/AssetTypeActions/AssetDefinition_NiagaraScript.cpp
deleted file mode 100644
diff --git a/Engine/Plugins/FX/Niagara/Source/NiagaraEditor/Private/AssetTypeActions/AssetDefinition_NiagaraScript.h b/Engine/Plugins/FX/Niagara/Source/NiagaraEditor/Private/AssetTypeActions/AssetDefinition_NiagaraScript.h
deleted file mode 100644
diff --git a/Engine/Plugins/FX/Niagara/Source/NiagaraEditor/Private/AssetTypeActions/AssetDefinition_NiagaraSimCache.cpp b/Engine/Plugins/FX/Niagara/Source/NiagaraEditor/Private/AssetTypeActions/AssetDefinition_NiagaraSimCache.cpp
deleted file mode 100644
diff --git a/Engine/Plugins/FX/Niagara/Source/NiagaraEditor/Private/AssetTypeActions/AssetDefinition_NiagaraSimCache.h b/Engine/Plugins/FX/Niagara/Source/NiagaraEditor/Private/AssetTypeActions/AssetDefinition_NiagaraSimCache.h
deleted file mode 100644
diff --git a/Engine/Plugins/FX/Niagara/Source/NiagaraEditor/Private/AssetTypeActions/AssetDefinition_NiagaraSystem.cpp b/Engine/Plugins/FX/Niagara/Source/NiagaraEditor/Private/AssetTypeActions/AssetDefinition_NiagaraSystem.cpp
deleted file mode 100644
diff --git a/Engine/Plugins/FX/Niagara/Source/NiagaraEditor/Private/AssetTypeActions/AssetDefinition_NiagaraSystem.h b/Engine/Plugins/FX/Niagara/Source/NiagaraEditor/Private/AssetTypeActions/AssetDefinition_NiagaraSystem.h
deleted file mode 100644
diff --git a/Engine/Plugins/FX/Niagara/Source/NiagaraEditor/Private/AssetTypeActions/AssetDefinition_NiagaraValidationRuleSet.cpp b/Engine/Plugins/FX/Niagara/Source/NiagaraEditor/Private/AssetTypeActions/AssetDefinition_NiagaraValidationRuleSet.cpp
deleted file mode 100644
diff --git a/Engine/Plugins/FX/Niagara/Source/NiagaraEditor/Private/AssetTypeActions/AssetDefinition_NiagaraValidationRuleSet.h b/Engine/Plugins/FX/Niagara/Source/NiagaraEditor/Private/AssetTypeActions/AssetDefinition_NiagaraValidationRuleSet.h
deleted file mode 100644
