AssetDefinition - More merging differences between the existing AD work in niagara and the work I did.

diff --git a/Engine/Plugins/FX/Niagara/Source/NiagaraEditor/Private/NiagaraScriptFactory.cpp b/Engine/Plugins/FX/Niagara/Source/NiagaraEditor/Private/NiagaraScriptFactory.cpp
--- a/Engine/Plugins/FX/Niagara/Source/NiagaraEditor/Private/NiagaraScriptFactory.cpp
+++ b/Engine/Plugins/FX/Niagara/Source/NiagaraEditor/Private/NiagaraScriptFactory.cpp
@@ -16,7 +16,7 @@
 #include "NiagaraScript.h"
 #include "NiagaraScriptFactoryNew.h"
 #include "NiagaraScriptSource.h"
-#include "AssetTypeActions/AssetDefinition_NiagaraScript.h"
+#include "AssetDefinitions/AssetDefinition_NiagaraScript.h"
 #include "EdGraph/EdGraph.h"
 #include "Modules/ModuleManager.h"
 #include "ViewModels/Stack/NiagaraStackGraphUtilities.h"
diff --git a/Engine/Plugins/FX/Niagara/Source/NiagaraEditor/Private/Toolkits/NiagaraScriptToolkit.cpp b/Engine/Plugins/FX/Niagara/Source/NiagaraEditor/Private/Toolkits/NiagaraScriptToolkit.cpp
--- a/Engine/Plugins/FX/Niagara/Source/NiagaraEditor/Private/Toolkits/NiagaraScriptToolkit.cpp
+++ b/Engine/Plugins/FX/Niagara/Source/NiagaraEditor/Private/Toolkits/NiagaraScriptToolkit.cpp
@@ -2,7 +2,7 @@
 
 #include "NiagaraScriptToolkit.h"
 
-#include "AssetTypeActions/AssetDefinition_NiagaraScript.h"
+#include "AssetDefinitions/AssetDefinition_NiagaraScript.h"
 #include "BusyCursor.h"
 #include "DetailLayoutBuilder.h"
 #include "Editor.h"
