- Add support for each decal to have a color associated with it
- Add Decal Color material node which can be used to read the color
- Niagara decal renderer will pipe particles.color through automatically to decals now

diff --git a/Engine/Plugins/FX/Niagara/Source/Niagara/Private/NiagaraRendererDecals.cpp b/Engine/Plugins/FX/Niagara/Source/Niagara/Private/NiagaraRendererDecals.cpp
--- a/Engine/Plugins/FX/Niagara/Source/Niagara/Private/NiagaraRendererDecals.cpp
+++ b/Engine/Plugins/FX/Niagara/Source/Niagara/Private/NiagaraRendererDecals.cpp
@@ -195,7 +195,8 @@ FNiagaraDynamicDataBase* FNiagaraRendererDecals::GenerateDynamicData(const FNiag
 			const FVector Position = bUseLocalSpace ? LocalToWorld.TransformPosition(FVector(SimPos)) : LwcConverter.ConvertSimulationPositionToWorld(SimPos);
 			const FQuat Rotation = bUseLocalSpace ? LocalToWorld.TransformRotation(FQuat(SimRot)) : FQuat(SimRot);
 			const FVector Size = FVector(SizeReader.GetSafe(ParticleIndex, DefaultSize) * 0.5f);
-			const float Fade = RendererProperties->bUseColorBindingAsFade ? ColorReader.GetSafe(ParticleIndex, FLinearColor::White).A : FadeReader.GetSafe(ParticleIndex, DefaultFade);
+			const FLinearColor DecalColor = ColorReader.GetSafe(ParticleIndex, FLinearColor::White);
+			const float Fade = FadeReader.GetSafe(ParticleIndex, DefaultFade);
 
 			// Create Update Parameters
 			FDeferredDecalUpdateParams& UpdateParams = DecalUpdates.AddDefaulted_GetRef();
@@ -217,6 +218,8 @@ FNiagaraDynamicDataBase* FNiagaraRendererDecals::GenerateDynamicData(const FNiag
 			UpdateParams.AbsSpawnTime = World->TimeSeconds - FMath::Clamp(1.0f - Fade, 0.0f, 1.0f);
 			UpdateParams.FadeStartDelay = 0.0f;
 			UpdateParams.FadeDuration = 1.0f;
+			UpdateParams.FadeScreenSize = RendererProperties->DecalScreenSizeFade;
+			UpdateParams.DecalColor = DecalColor;
 		}
 	}
 	// Emitter source mode
@@ -232,10 +235,8 @@ FNiagaraDynamicDataBase* FNiagaraRendererDecals::GenerateDynamicData(const FNiag
 			const FVector Position = bUseLocalSpace ? LocalToWorld.TransformPosition(FVector(SimPos)) : LwcConverter.ConvertSimulationPositionToWorld(SimPos);
 			const FQuat Rotation = bUseLocalSpace ? LocalToWorld.TransformRotation(FQuat(SimRot)) : FQuat(SimRot);
 			const FVector Size = FVector(ParameterStore.GetParameterValueOrDefault(RendererProperties->DecalSizeBinding.GetParamMapBindableVariable(), DefaultSize) * 0.5f);
-			const float Fade =
-				RendererProperties->bUseColorBindingAsFade ?
-				ParameterStore.GetParameterValueOrDefault(RendererProperties->ColorBinding.GetParamMapBindableVariable(), FLinearColor::White).A :
-				ParameterStore.GetParameterValueOrDefault(RendererProperties->DecalFadeBinding.GetParamMapBindableVariable(), DefaultFade);
+			const FLinearColor DecalColor = ParameterStore.GetParameterValueOrDefault(RendererProperties->ColorBinding.GetParamMapBindableVariable(), FLinearColor::White);
+			const float Fade = ParameterStore.GetParameterValueOrDefault(RendererProperties->DecalFadeBinding.GetParamMapBindableVariable(), DefaultFade);
 
 			// Create Update Parameters
 			FDeferredDecalUpdateParams& UpdateParams = DecalUpdates.AddDefaulted_GetRef();
@@ -257,6 +258,8 @@ FNiagaraDynamicDataBase* FNiagaraRendererDecals::GenerateDynamicData(const FNiag
 			UpdateParams.AbsSpawnTime = World->TimeSeconds - FMath::Clamp(1.0f - Fade, 0.0f, 1.0f);
 			UpdateParams.FadeStartDelay = 0.0f;
 			UpdateParams.FadeDuration = 1.0f;
+			UpdateParams.FadeScreenSize = RendererProperties->DecalScreenSizeFade;
+			UpdateParams.DecalColor = DecalColor;
 		}
 	}
 
diff --git a/Engine/Plugins/FX/Niagara/Source/Niagara/Public/NiagaraDecalRendererProperties.h b/Engine/Plugins/FX/Niagara/Source/Niagara/Public/NiagaraDecalRendererProperties.h
--- a/Engine/Plugins/FX/Niagara/Source/Niagara/Public/NiagaraDecalRendererProperties.h
+++ b/Engine/Plugins/FX/Niagara/Source/Niagara/Public/NiagaraDecalRendererProperties.h
@@ -72,6 +72,10 @@ public:
 	UPROPERTY(EditAnywhere, Category = "Decal Rendering")
 	int32 RendererVisibility = 0;
 
+	/** When the decal is smaller than this screen size fade out the decal, can be used to reduce the amount of small decals drawn. */
+	UPROPERTY(EditAnywhere, Category = "Decal Rendering")
+	float DecalScreenSizeFade = 0.f;
+
 	/** Position binding for the decals, should be center of the decal */
 	UPROPERTY(EditAnywhere, Category = "Bindings")
 	FNiagaraVariableAttributeBinding PositionBinding;
@@ -84,16 +88,12 @@ public:
 	UPROPERTY(EditAnywhere, Category = "Bindings")
 	FNiagaraVariableAttributeBinding DecalSizeBinding;
 
-	/** When enabled we use the Color Binding as the fade control, otherwise we use the DecalFade binding. */
+	/** Fade binding for the decal, value can be queried using the Decal Lifetime Opacity material node. */
 	UPROPERTY(EditAnywhere, Category = "Bindings")
-	bool bUseColorBindingAsFade = true;
-
-	/** Fade binding for the decal, can be queried from the decal fade material node.. */
-	UPROPERTY(EditAnywhere, Category = "Bindings", Meta = (EditCondition = "!bUseColorBindingAsFade"))
 	FNiagaraVariableAttributeBinding DecalFadeBinding;
 
-	/** Color binding for the decal, only the alpha channel is used and can be queried from the decal fade material node. */
-	UPROPERTY(EditAnywhere, Category = "Bindings", Meta = (EditCondition = "bUseColorBindingAsFade"))
+	/** Color binding for the decal, value can be queried using the Decal Color material node. */
+	UPROPERTY(EditAnywhere, Category = "Bindings")
 	FNiagaraVariableAttributeBinding ColorBinding;
 
 	/** Visibility tag binding, when valid the returned values is compated with RendererVisibility. */
diff --git a/Engine/Shaders/Private/DeferredDecal.usf b/Engine/Shaders/Private/DeferredDecal.usf
--- a/Engine/Shaders/Private/DeferredDecal.usf
+++ b/Engine/Shaders/Private/DeferredDecal.usf
@@ -39,6 +39,8 @@ float3 DecalOrientation;
 float2 DecalParams;
 // from component to clip space (for decal frustum)
 float4x4 FrustumComponentToClip;
+// Decal color can be accessed using the material node
+float4 DecalColorParam;
 
 // Global decal vertex shader
 void MainVS(
diff --git a/Engine/Shaders/Private/MaterialTemplate.ush b/Engine/Shaders/Private/MaterialTemplate.ush
--- a/Engine/Shaders/Private/MaterialTemplate.ush
+++ b/Engine/Shaders/Private/MaterialTemplate.ush
@@ -1093,17 +1093,23 @@ MaterialFloat2 GetDefaultSceneTextureUV(FMaterialPixelParameters Parameters, con
 	#endif
 	}
 
-#if DECAL_PRIMITIVE
-	float DecalLifetimeOpacity()
+	float4 DecalColor()
 	{
-		return DecalParams.y;
+	#if DECAL_PRIMITIVE
+		return DecalColorParam;
+	#else
+		return float4(1.0f, 1.0f, 1.0f, 1.0f);
+	#endif
 	}
-#else
+
 	float DecalLifetimeOpacity()
 	{
+	#if DECAL_PRIMITIVE
+		return DecalParams.y;
+	#else
 		return 0.0f;
+	#endif
 	}
-#endif // DECAL_PRIMITIVE
 
 /** Per Instance Custom Data Getter (Pixel Shader Variant - Visibility Buffer) */
 float GetPerInstanceCustomData(FMaterialPixelParameters Parameters, int Index, float DefaultValue)
diff --git a/Engine/Shaders/Private/RayTracing/RayTracingDecalMaterialShader.usf b/Engine/Shaders/Private/RayTracing/RayTracingDecalMaterialShader.usf
--- a/Engine/Shaders/Private/RayTracing/RayTracingDecalMaterialShader.usf
+++ b/Engine/Shaders/Private/RayTracing/RayTracingDecalMaterialShader.usf
@@ -13,6 +13,7 @@
 #define DecalToWorldInvScale RayTracingDecalParameters.DecalToWorldInvScale
 #define DecalOrientation RayTracingDecalParameters.DecalOrientation
 #define DecalParams RayTracingDecalParameters.DecalParams
+#define DecalColorParam RayTracingDecalParameters.DecalColorParam
 #define DecalWriteFlags RayTracingDecalParameters.DecalWriteFlags
 
 #include "/Engine/Generated/Material.ush"
diff --git a/Engine/Source/Editor/UnrealEd/Private/Lightmass/LightmassRender.cpp b/Engine/Source/Editor/UnrealEd/Private/Lightmass/LightmassRender.cpp
--- a/Engine/Source/Editor/UnrealEd/Private/Lightmass/LightmassRender.cpp
+++ b/Engine/Source/Editor/UnrealEd/Private/Lightmass/LightmassRender.cpp
@@ -217,6 +217,11 @@ struct FLightmassMaterialCompiler : public FProxyMaterialCompiler
 		return Compiler->Constant(0.0f);
 	}
 
+	virtual int32 DecalColor() override
+	{
+		return Compiler->Constant4(1.0f, 1.0f, 1.0f, 1.0f);
+	}
+
 	virtual int32 DecalLifetimeOpacity() override
 	{
 		return Compiler->Constant(0.0f);
diff --git a/Engine/Source/Runtime/Engine/Classes/Components/DecalComponent.h b/Engine/Source/Runtime/Engine/Classes/Components/DecalComponent.h
--- a/Engine/Source/Runtime/Engine/Classes/Components/DecalComponent.h
+++ b/Engine/Source/Runtime/Engine/Classes/Components/DecalComponent.h
@@ -75,6 +75,14 @@ public:
 	UFUNCTION(BlueprintCallable, Category = "Rendering|Components|Decal")
 	float GetFadeInDuration() const;
 
+	/** Decal size in local space (does not include the component scale), technically redundant but there for convenience */
+	UPROPERTY(EditAnywhere, BlueprintReadOnly, Category = Decal, meta=(AllowPreserveRatio = "true"))
+	FVector DecalSize;
+
+	/** Decal color, can be accessed using the material Decal Color node. */
+	UPROPERTY(EditAnywhere, BlueprintReadOnly, Category = Decal)
+	FLinearColor DecalColor = FLinearColor::White;
+
 	/**
 	* Sets the decal's fade start time, duration and if the owning actor should be destroyed after the decal is fully faded out.
 	* The default value of 0 for FadeStartDelay and FadeDuration makes the decal persistent. See DecalLifetimeOpacity material 
@@ -94,14 +102,14 @@ public:
 	UFUNCTION(BlueprintCallable, Category = "Rendering|Components|Decal")
 	void SetFadeScreenSize(float NewFadeScreenSize);
 
-	/** Decal size in local space (does not include the component scale), technically redundant but there for convenience */
-	UPROPERTY(EditAnywhere, BlueprintReadOnly, Category = Decal, meta=(AllowPreserveRatio = "true"))
-	FVector DecalSize;
-
 	/** Sets the sort order for the decal component. Higher values draw later (on top). This will force the decal to reattach */
 	UFUNCTION(BlueprintCallable, Category = "Rendering|Components|Decal")
 	void SetSortOrder(int32 Value);
 
+	/** Sets the decal color. */
+	UFUNCTION(BlueprintCallable, Category = "Rendering|Components|Decal")
+	void SetDecalColor(const FLinearColor& Color);
+
 	/** setting decal material on decal component. This will force the decal to reattach */
 	UFUNCTION(BlueprintCallable, Category="Rendering|Components|Decal")
 	void SetDecalMaterial(class UMaterialInterface* NewDecalMaterial);
diff --git a/Engine/Source/Runtime/Engine/Classes/Materials/MaterialExpressionDecalColor.h b/Engine/Source/Runtime/Engine/Classes/Materials/MaterialExpressionDecalColor.h
new file mode 100644
--- /dev/null
+++ b/Engine/Source/Runtime/Engine/Classes/Materials/MaterialExpressionDecalColor.h
@@ -0,0 +1,22 @@
+// Copyright Epic Games, Inc. All Rights Reserved.
+
+
+#pragma once
+
+#include "CoreMinimal.h"
+#include "UObject/ObjectMacros.h"
+#include "Materials/MaterialExpression.h"
+#include "MaterialExpressionDecalColor.generated.h"
+
+UCLASS(collapsecategories, hidecategories = Object)
+class UMaterialExpressionDecalColor : public UMaterialExpression
+{
+	GENERATED_UCLASS_BODY()
+
+	//~ Begin UMaterialExpression Interface
+#if WITH_EDITOR
+	virtual int32 Compile(class FMaterialCompiler* Compiler, int32 OutputIndex) override;
+	virtual void GetCaption(TArray<FString>& OutCaptions) const override;
+#endif
+	//~ End UMaterialExpression Interface
+};
diff --git a/Engine/Source/Runtime/Engine/Private/Components/DecalComponent.cpp b/Engine/Source/Runtime/Engine/Private/Components/DecalComponent.cpp
--- a/Engine/Source/Runtime/Engine/Private/Components/DecalComponent.cpp
+++ b/Engine/Source/Runtime/Engine/Private/Components/DecalComponent.cpp
@@ -29,7 +29,8 @@ FDeferredDecalProxy::FDeferredDecalProxy(const UDecalComponent* InComponent)
 	, InvFadeInDuration(1.0f)
 	, FadeStartDelayNormalized(1.0f)
 	, FadeInStartDelayNormalized(0.0f)
-	, FadeScreenSize( InComponent->FadeScreenSize )
+	, FadeScreenSize(InComponent->FadeScreenSize)
+	, DecalColor(InComponent->DecalColor)
 {
 	UMaterialInterface* EffectiveMaterial = UMaterial::GetDefaultMaterial(MD_DeferredDecal);
 	UMaterialInterface* ComponentMaterial = InComponent->GetDecalMaterial();
@@ -286,6 +287,13 @@ void UDecalComponent::SetSortOrder(int32 Value)
 	MarkRenderStateDirty();
 }
 
+void UDecalComponent::SetDecalColor(const FLinearColor& InColor)
+{
+	DecalColor = InColor;
+
+	MarkRenderStateDirty();
+}
+
 void UDecalComponent::SetDecalMaterial(class UMaterialInterface* NewDecalMaterial)
 {
 	DecalMaterial = NewDecalMaterial;
diff --git a/Engine/Source/Runtime/Engine/Private/Materials/HLSLMaterialTranslator.cpp b/Engine/Source/Runtime/Engine/Private/Materials/HLSLMaterialTranslator.cpp
--- a/Engine/Source/Runtime/Engine/Private/Materials/HLSLMaterialTranslator.cpp
+++ b/Engine/Source/Runtime/Engine/Private/Materials/HLSLMaterialTranslator.cpp
@@ -6565,6 +6565,24 @@ int32 FHLSLMaterialTranslator::TextureDecalDerivative(bool bDDY)
 		);
 }
 
+int32 FHLSLMaterialTranslator::DecalColor()
+{
+	if (Material->GetMaterialDomain() != MD_DeferredDecal)
+	{
+		return Errorf(TEXT("Decal color is only available in the decal material domain."));
+	}
+
+	if (ShaderFrequency != SF_Pixel)
+	{
+		return Errorf(TEXT("Decal color is only available in the pixel shader."));
+	}
+
+	return AddCodeChunkZeroDeriv(
+		MCT_Float4,
+		TEXT("DecalColor()")
+	);
+}
+
 int32 FHLSLMaterialTranslator::DecalLifetimeOpacity()
 {
 	if (Material->GetMaterialDomain() != MD_DeferredDecal)
diff --git a/Engine/Source/Runtime/Engine/Private/Materials/HLSLMaterialTranslator.h b/Engine/Source/Runtime/Engine/Private/Materials/HLSLMaterialTranslator.h
--- a/Engine/Source/Runtime/Engine/Private/Materials/HLSLMaterialTranslator.h
+++ b/Engine/Source/Runtime/Engine/Private/Materials/HLSLMaterialTranslator.h
@@ -875,6 +875,7 @@ protected:
 	virtual int32 TextureDecalMipmapLevel(int32 TextureSizeInput) override;
 	virtual int32 TextureDecalDerivative(bool bDDY) override;
 
+	virtual int32 DecalColor() override;
 	virtual int32 DecalLifetimeOpacity() override;
 
 	virtual int32 PixelDepth() override;
diff --git a/Engine/Source/Runtime/Engine/Private/Materials/MaterialExpressions.cpp b/Engine/Source/Runtime/Engine/Private/Materials/MaterialExpressions.cpp
--- a/Engine/Source/Runtime/Engine/Private/Materials/MaterialExpressions.cpp
+++ b/Engine/Source/Runtime/Engine/Private/Materials/MaterialExpressions.cpp
@@ -79,6 +79,7 @@
 #include "Materials/MaterialExpressionDBufferTexture.h"
 #include "Materials/MaterialExpressionDDX.h"
 #include "Materials/MaterialExpressionDDY.h"
+#include "Materials/MaterialExpressionDecalColor.h"
 #include "Materials/MaterialExpressionDecalDerivative.h"
 #include "Materials/MaterialExpressionDecalLifetimeOpacity.h"
 #include "Materials/MaterialExpressionDecalMipmapLevel.h"
@@ -19361,6 +19362,56 @@ void UMaterialExpressionDecalDerivative::GetCaption(TArray<FString>& OutCaptions
 }
 #endif // WITH_EDITOR
 
+//
+//	UMaterialExpressionDecalColor
+//
+UMaterialExpressionDecalColor::UMaterialExpressionDecalColor(const FObjectInitializer& ObjectInitializer)
+	: Super(ObjectInitializer)
+{
+#if WITH_EDITORONLY_DATA
+	// Structure to hold one-time initialization
+	struct FConstructorStatics
+	{
+		FText NAME_Vectors;
+		FConstructorStatics()
+			: NAME_Vectors(LOCTEXT("Utils", "Utils"))
+		{
+		}
+	};
+	static FConstructorStatics ConstructorStatics;
+
+	MenuCategories.Add(ConstructorStatics.NAME_Vectors);
+
+	bShowOutputNameOnPin = true;
+	bShaderInputData = true;
+
+	Outputs.Reset();
+	Outputs.Add(FExpressionOutput(TEXT("RGB"), 1, 1, 1, 1, 0));
+	Outputs.Add(FExpressionOutput(TEXT("R"), 1, 1, 0, 0, 0));
+	Outputs.Add(FExpressionOutput(TEXT("G"), 1, 0, 1, 0, 0));
+	Outputs.Add(FExpressionOutput(TEXT("B"), 1, 0, 0, 1, 0));
+	Outputs.Add(FExpressionOutput(TEXT("A"), 1, 0, 0, 0, 1));
+	Outputs.Add(FExpressionOutput(TEXT("RGBA"), 1, 1, 1, 1, 1));
+#endif
+}
+
+#if WITH_EDITOR
+int32 UMaterialExpressionDecalColor::Compile(class FMaterialCompiler* Compiler, int32 OutputIndex)
+{
+	if (Material && Material->MaterialDomain != MD_DeferredDecal)
+	{
+		return CompilerError(Compiler, TEXT("Node only works for the deferred decal material domain."));
+	}
+
+	return Compiler->DecalColor();
+}
+
+void UMaterialExpressionDecalColor::GetCaption(TArray<FString>& OutCaptions) const
+{
+	OutCaptions.Add(TEXT("Decal Color"));
+}
+#endif // WITH_EDITOR
+
 //
 //	UMaterialExpressionDecalLifetimeOpacity
 //
diff --git a/Engine/Source/Runtime/Engine/Public/MaterialCompiler.h b/Engine/Source/Runtime/Engine/Public/MaterialCompiler.h
--- a/Engine/Source/Runtime/Engine/Public/MaterialCompiler.h
+++ b/Engine/Source/Runtime/Engine/Public/MaterialCompiler.h
@@ -359,6 +359,7 @@ public:
 
 	virtual int32 TextureDecalMipmapLevel(int32 TextureSizeInput) = 0;
 	virtual int32 TextureDecalDerivative(bool bDDY) = 0;
+	virtual int32 DecalColor() = 0;
 	virtual int32 DecalLifetimeOpacity() = 0;
 
 	virtual int32 Texture(UTexture* Texture,int32& TextureReferenceIndex, EMaterialSamplerType SamplerType, ESamplerSourceMode SamplerSource=SSM_FromTextureAsset,ETextureMipValueMode MipValueMode=TMVM_None) = 0;
@@ -881,6 +882,7 @@ public:
 
 	virtual int32 TextureDecalMipmapLevel(int32 TextureSizeInput) override { return Compiler->TextureDecalMipmapLevel(TextureSizeInput); }
 	virtual int32 TextureDecalDerivative(bool bDDY) override { return Compiler->TextureDecalDerivative(bDDY); }
+	virtual int32 DecalColor() override { return Compiler->DecalColor(); }
 	virtual int32 DecalLifetimeOpacity() override { return Compiler->DecalLifetimeOpacity(); }
 
 	virtual int32 Texture(UTexture* InTexture, int32& TextureReferenceIndex, EMaterialSamplerType SamplerType, ESamplerSourceMode SamplerSource = SSM_FromTextureAsset, ETextureMipValueMode MipValueMode = TMVM_None) override { return Compiler->Texture(InTexture, TextureReferenceIndex, SamplerType, SamplerSource, MipValueMode); }
diff --git a/Engine/Source/Runtime/Engine/Public/SceneManagement.h b/Engine/Source/Runtime/Engine/Public/SceneManagement.h
--- a/Engine/Source/Runtime/Engine/Public/SceneManagement.h
+++ b/Engine/Source/Runtime/Engine/Public/SceneManagement.h
@@ -1452,6 +1452,8 @@ public:
 	float FadeInStartDelayNormalized;
 
 	float FadeScreenSize;
+
+	FLinearColor DecalColor = FLinearColor::White;
 };
 
 struct FDeferredDecalUpdateParams
@@ -1474,6 +1476,7 @@ struct FDeferredDecalUpdateParams
 	float					FadeInStartDelay = 0.0f;
 	float					FadeScreenSize = 0.01f;
 	int32					SortOrder = 0;
+	FLinearColor			DecalColor = FLinearColor::White;
 };
 
 /** Reflection capture shapes. */
diff --git a/Engine/Source/Runtime/Renderer/Private/DecalRenderingShared.cpp b/Engine/Source/Runtime/Renderer/Private/DecalRenderingShared.cpp
--- a/Engine/Source/Runtime/Renderer/Private/DecalRenderingShared.cpp
+++ b/Engine/Source/Runtime/Renderer/Private/DecalRenderingShared.cpp
@@ -28,6 +28,7 @@ FTransientDecalRenderData::FTransientDecalRenderData(const FScene& InScene, cons
 	, MaterialProxy(InDecalProxy.DecalMaterial->GetRenderProxy())
 	, ConservativeRadius(InConservativeRadius)
 	, FadeAlpha(1.0f)
+	, DecalColor(InDecalProxy.DecalColor)
 {
 	// Build BlendDesc from a potentially incomplete material.
 	// If our shader isn't compiled yet then we will potentially render later with a different fallback material.
@@ -81,6 +82,7 @@ public:
 		DecalToWorldInvScale.Bind(Initializer.ParameterMap, TEXT("DecalToWorldInvScale"));
 		DecalOrientation.Bind(Initializer.ParameterMap,TEXT("DecalOrientation"));
 		DecalParams.Bind(Initializer.ParameterMap, TEXT("DecalParams"));
+		DecalColorParam.Bind(Initializer.ParameterMap, TEXT("DecalColorParam"));
 		MobileBasePassUniformBuffer.Bind(Initializer.ParameterMap, FMobileBasePassUniformParameters::StaticStructMetadata.GetShaderVariableName());
 	}
 
@@ -147,6 +149,7 @@ public:
 		}
  
 		SetShaderValue(RHICmdList, ShaderRHI, DecalParams, FVector2f(FadeAlphaValue, LifetimeAlpha));
+		SetShaderValue(RHICmdList, ShaderRHI, DecalColorParam, DecalProxy.DecalColor);
 	}
 
 private:
@@ -156,6 +159,7 @@ private:
 	LAYOUT_FIELD(FShaderParameter, DecalToWorldInvScale);
 	LAYOUT_FIELD(FShaderParameter, DecalOrientation);
 	LAYOUT_FIELD(FShaderParameter, DecalParams);
+	LAYOUT_FIELD(FShaderParameter, DecalColorParam);
 	LAYOUT_FIELD(FShaderUniformBufferParameter, MobileBasePassUniformBuffer);
 };
 
diff --git a/Engine/Source/Runtime/Renderer/Private/DecalRenderingShared.h b/Engine/Source/Runtime/Renderer/Private/DecalRenderingShared.h
--- a/Engine/Source/Runtime/Renderer/Private/DecalRenderingShared.h
+++ b/Engine/Source/Runtime/Renderer/Private/DecalRenderingShared.h
@@ -22,6 +22,7 @@ struct FTransientDecalRenderData
 	FDecalBlendDesc BlendDesc;
 	float ConservativeRadius;
 	float FadeAlpha;
+	FLinearColor DecalColor;
 
 	FTransientDecalRenderData(const FScene& InScene, const FDeferredDecalProxy& InDecalProxy, float InConservativeRadius);
 };
diff --git a/Engine/Source/Runtime/Renderer/Private/RayTracing/RayTracingDecals.cpp b/Engine/Source/Runtime/Renderer/Private/RayTracing/RayTracingDecals.cpp
--- a/Engine/Source/Runtime/Renderer/Private/RayTracing/RayTracingDecals.cpp
+++ b/Engine/Source/Runtime/Renderer/Private/RayTracing/RayTracingDecals.cpp
@@ -60,6 +60,7 @@ BEGIN_GLOBAL_SHADER_PARAMETER_STRUCT(FDecalParametersRayTracing, )
 	SHADER_PARAMETER(FVector3f, DecalToWorldInvScale)
 	SHADER_PARAMETER(FVector3f, DecalOrientation)
 	SHADER_PARAMETER(FVector2f, DecalParams)
+	SHADER_PARAMETER(FVector4f, DecalColorParam)
 	SHADER_PARAMETER(uint32, DecalWriteFlags)
 END_GLOBAL_SHADER_PARAMETER_STRUCT()
 
@@ -85,6 +86,7 @@ static TUniformBufferRef<FDecalParametersRayTracing> CreateDecalParametersBuffer
 	Parameters.DecalToWorld = RelativeDecalToWorldMatrix;
 	Parameters.DecalToWorldInvScale = RelativeDecalToWorldMatrix.GetScaleVector().Reciprocal();
 	Parameters.DecalOrientation = OrientationVector;
+	Parameters.DecalColorParam = DecalData.DecalColor;
 
 	float LifetimeAlpha = 1.0f;
 
@@ -95,6 +97,7 @@ static TUniformBufferRef<FDecalParametersRayTracing> CreateDecalParametersBuffer
 	}
 
 	Parameters.DecalParams = FVector2f(DecalData.FadeAlpha, LifetimeAlpha);
+	Parameters.DecalColorParam = DecalData.DecalColor;
 
 	Parameters.DecalWriteFlags = CalculateDecalWriteFlags(DecalData.BlendDesc);
 
diff --git a/Engine/Source/Runtime/Renderer/Private/RendererScene.cpp b/Engine/Source/Runtime/Renderer/Private/RendererScene.cpp
--- a/Engine/Source/Runtime/Renderer/Private/RendererScene.cpp
+++ b/Engine/Source/Runtime/Renderer/Private/RendererScene.cpp
@@ -3027,6 +3027,7 @@ void FScene::BatchUpdateDecals(TArray<FDeferredDecalUpdateParams>&& UpdateParams
 				DecalUpdate.DecalProxy->InitializeFadingParameters(DecalUpdate.AbsSpawnTime, DecalUpdate.FadeDuration, DecalUpdate.FadeStartDelay, DecalUpdate.FadeInDuration, DecalUpdate.FadeInStartDelay);
 				DecalUpdate.DecalProxy->FadeScreenSize = DecalUpdate.FadeScreenSize;
 				DecalUpdate.DecalProxy->SortOrder = DecalUpdate.SortOrder;
+				DecalUpdate.DecalProxy->DecalColor = DecalUpdate.DecalColor;
 			}
 		}
 	);
