Fixed an instance of graph data cache being invalid due to graph change ID not updating when a graph's metadata was modified

diff --git a/Engine/Plugins/FX/Niagara/Source/Niagara/Private/NiagaraScript.cpp b/Engine/Plugins/FX/Niagara/Source/Niagara/Private/NiagaraScript.cpp
--- a/Engine/Plugins/FX/Niagara/Source/Niagara/Private/NiagaraScript.cpp
+++ b/Engine/Plugins/FX/Niagara/Source/Niagara/Private/NiagaraScript.cpp
@@ -2368,23 +2368,26 @@ void UNiagaraScript::PostEditChangeVersionedProperty(FPropertyChangedEvent& Prop
 
 	if (UNiagaraScriptSourceBase* Source = GetSource(Version))
 	{
+		FString UnsynchronizedMessage = TEXT("Script property changed");
 		if (PropertyName == GET_MEMBER_NAME_CHECKED(FVersionedNiagaraScriptData, bDeprecated) || 
          PropertyName == GET_MEMBER_NAME_CHECKED(FVersionedNiagaraScriptData, DeprecationMessage) ||
          PropertyName == GET_MEMBER_NAME_CHECKED(FVersionedNiagaraScriptData, DeprecationRecommendation))
 		{
-			Source->MarkNotSynchronized(TEXT("Deprecation changed."));
+			UnsynchronizedMessage = TEXT("Deprecation changed.");
 		}
 
 		if (PropertyName == GET_MEMBER_NAME_CHECKED(FVersionedNiagaraScriptData, bExperimental) || 
             PropertyName == GET_MEMBER_NAME_CHECKED(FVersionedNiagaraScriptData, ExperimentalMessage))
 		{
-			Source->MarkNotSynchronized(TEXT("Experimental changed."));
+			UnsynchronizedMessage = TEXT("Experimental changed.");
 		}
 	
 		if (PropertyName == GET_MEMBER_NAME_CHECKED(FVersionedNiagaraScriptData, NoteMessage))
 		{
-			Source->MarkNotSynchronized(TEXT("Note changed."));
+			UnsynchronizedMessage = TEXT("Note changed.");
 		}
+
+		Source->MarkNotSynchronized(UnsynchronizedMessage);
 	}
 
 	PostEditChangeProperty(PropertyChangedEvent);
