Changed script toolkit to compare graph change IDs instead of a manually marked bool to determine if we can apply the changes

diff --git a/Engine/Plugins/FX/Niagara/Source/NiagaraEditor/Private/Toolkits/NiagaraScriptToolkit.cpp b/Engine/Plugins/FX/Niagara/Source/NiagaraEditor/Private/Toolkits/NiagaraScriptToolkit.cpp
--- a/Engine/Plugins/FX/Niagara/Source/NiagaraEditor/Private/Toolkits/NiagaraScriptToolkit.cpp
+++ b/Engine/Plugins/FX/Niagara/Source/NiagaraEditor/Private/Toolkits/NiagaraScriptToolkit.cpp
@@ -181,7 +181,6 @@ void FNiagaraScriptToolkit::Initialize( const EToolkitMode::Type Mode, const TSh
 	EditedNiagaraScript.Script = (UNiagaraScript*)StaticDuplicateObject(InputScript, GetTransientPackage(), NAME_None, ~RF_Standalone, UNiagaraScript::StaticClass());
 	EditedNiagaraScript.Script->OnVMScriptCompiled().AddSP(this, &FNiagaraScriptToolkit::OnVMScriptCompiled);
 	EditedNiagaraScript.Version = ScriptVersion;
-	bEditedScriptHasPendingChanges = false;
 
 	const FGuid MessageLogGuidKey = FGuid::NewGuid();
 	NiagaraMessageLogViewModel = MakeShared<FNiagaraMessageLogViewModel>(GetNiagaraScriptMessageLogName(EditedNiagaraScript), MessageLogGuidKey, NiagaraMessageLog);
@@ -222,7 +221,6 @@ void FNiagaraScriptToolkit::Initialize( const EToolkitMode::Type Mode, const TSh
 		.OnChangeToVersion(this, &FNiagaraScriptToolkit::SwitchToVersion)
 		.OnVersionDataChanged_Lambda([this]()
 		{
-			bEditedScriptHasPendingChanges = true;
 	        if (UNiagaraScriptSource* ScriptSource = Cast<UNiagaraScriptSource>(EditedNiagaraScript.Script->GetSource(EditedNiagaraScript.Version)))
 	        {
 	            ScriptSource->NodeGraph->NotifyGraphChanged();
@@ -426,7 +424,7 @@ void FNiagaraScriptToolkit::OnEditedScriptPropertyFinishedChanging(const FProper
 		}
 	}
 
-	MarkDirtyWithPendingChanges();
+	PromptVersioningWarning();
 }
 
 void FNiagaraScriptToolkit::OnVMScriptCompiled(UNiagaraScript*, const FGuid&)
@@ -769,7 +767,7 @@ FText FNiagaraScriptToolkit::GetVersionMenuLabel(FNiagaraAssetVersion Version) c
 
 bool FNiagaraScriptToolkit::IsEditScriptDifferentFromOriginalScript() const
 {
-	return bEditedScriptHasPendingChanges;
+	return OriginalNiagaraScript.Script->GetBaseChangeID() != EditedNiagaraScript.Script->GetBaseChangeID();
 }
 
 void FNiagaraScriptToolkit::OnApply()
@@ -839,7 +837,7 @@ void FNiagaraScriptToolkit::PostUndo(bool bSuccess)
 			EditedNiagaraScript.Version = FGuid();
 		}
 	}
-	MarkDirtyWithPendingChanges();
+	PromptVersioningWarning();
 }
 
 void FNiagaraScriptToolkit::GetSaveableObjects(TArray<UObject*>& OutObjects) const
@@ -905,7 +903,6 @@ void FNiagaraScriptToolkit::UpdateOriginalNiagaraScript()
 	FNiagaraEditorUtilities::RefreshAllScriptsFromExternalChanges(Args);
 
 	GWarn->EndSlowTask();
-	bEditedScriptHasPendingChanges = false;
 	FNiagaraEditorModule::Get().InvalidateCachedScriptAssetData();
 }
 
@@ -944,11 +941,11 @@ bool FNiagaraScriptToolkit::OnRequestClose()
 
 void FNiagaraScriptToolkit::OnEditedScriptGraphChanged(const FEdGraphEditAction& InAction)
 {
-	MarkDirtyWithPendingChanges();
+	PromptVersioningWarning();
 	bRefreshSelected = true;
 }
 
-void FNiagaraScriptToolkit::MarkDirtyWithPendingChanges()
+void FNiagaraScriptToolkit::PromptVersioningWarning()
 {
 	if (!bShowedEditingVersionWarning && EditedNiagaraScript.Script->IsVersioningEnabled())
 	{
@@ -960,7 +957,6 @@ void FNiagaraScriptToolkit::MarkDirtyWithPendingChanges()
 			FMessageDialog::Open(EAppMsgType::Ok, LOCTEXT("EditingExposedAssetWarning", "Warning: you are editing an already exposed asset version. Saving these changes will force-push them out to existing usages!\nConsider creating a new version instead to make those changes."));	
 		}
 	}
-	bEditedScriptHasPendingChanges = true;
 }
 
 void FNiagaraScriptToolkit::FocusGraphElementIfSameScriptID(const FNiagaraScriptIDAndGraphFocusInfo* FocusInfo)
diff --git a/Engine/Plugins/FX/Niagara/Source/NiagaraEditor/Private/Toolkits/NiagaraScriptToolkit.h b/Engine/Plugins/FX/Niagara/Source/NiagaraEditor/Private/Toolkits/NiagaraScriptToolkit.h
--- a/Engine/Plugins/FX/Niagara/Source/NiagaraEditor/Private/Toolkits/NiagaraScriptToolkit.h
+++ b/Engine/Plugins/FX/Niagara/Source/NiagaraEditor/Private/Toolkits/NiagaraScriptToolkit.h
@@ -152,7 +152,7 @@ private:
 
 	void OnEditedScriptGraphChanged(const FEdGraphEditAction& InAction);
 
-	void MarkDirtyWithPendingChanges();
+	void PromptVersioningWarning();
 
 	/** Navigates to element in graph (node, pin, etc.) 
 	* @Param ElementToFocus Defines the graph element to navigate to and select.
@@ -196,8 +196,7 @@ private:
 	TSharedPtr<class SWidget> VersionsWidget;
 
 	FDelegateHandle OnEditedScriptGraphChangedHandle;
-
-	bool bEditedScriptHasPendingChanges = false;
+	
 	bool bChangesDiscarded = false;
 	bool bRefreshSelected = false;
 	bool bShowedEditingVersionWarning = false;
