Making sure to invalidate cached traversals when emitters are added or removed from a system

diff --git a/Engine/Plugins/FX/Niagara/Source/Niagara/Private/NiagaraSystem.cpp b/Engine/Plugins/FX/Niagara/Source/Niagara/Private/NiagaraSystem.cpp
--- a/Engine/Plugins/FX/Niagara/Source/Niagara/Private/NiagaraSystem.cpp
+++ b/Engine/Plugins/FX/Niagara/Source/Niagara/Private/NiagaraSystem.cpp
@@ -2126,6 +2126,7 @@ FNiagaraEmitterHandle UNiagaraSystem::AddEmitterHandle(UNiagaraEmitter& InEmitte
 	}
 	EmitterHandles.Add(EmitterHandle);
 	RefreshSystemParametersFromEmitter(EmitterHandle);
+	InvalidateCachedData();
 	return EmitterHandle;
 }
 
@@ -2133,6 +2134,7 @@ void UNiagaraSystem::AddEmitterHandleDirect(FNiagaraEmitterHandle& EmitterHandle
 {
 	EmitterHandles.Add(EmitterHandle);
 	RefreshSystemParametersFromEmitter(EmitterHandle);
+	InvalidateCachedData();
 }
 
 FNiagaraEmitterHandle UNiagaraSystem::DuplicateEmitterHandle(const FNiagaraEmitterHandle& EmitterHandleToDuplicate, FName EmitterName)
@@ -2142,6 +2144,7 @@ FNiagaraEmitterHandle UNiagaraSystem::DuplicateEmitterHandle(const FNiagaraEmitt
 	EmitterHandle.SetIsEnabled(EmitterHandleToDuplicate.GetIsEnabled(), *this, false);
 	EmitterHandles.Add(EmitterHandle);
 	RefreshSystemParametersFromEmitter(EmitterHandle);
+	InvalidateCachedData();
 	return EmitterHandle;
 }
 
@@ -2150,6 +2153,7 @@ void UNiagaraSystem::RemoveEmitterHandle(const FNiagaraEmitterHandle& EmitterHan
 	RemoveSystemParametersForEmitter(EmitterHandleToDelete);
 	auto RemovePredicate = [&](const FNiagaraEmitterHandle& EmitterHandle) { return EmitterHandle.GetId() == EmitterHandleToDelete.GetId(); };
 	EmitterHandles.RemoveAll(RemovePredicate);
+	InvalidateCachedData();
 }
 
 void UNiagaraSystem::RemoveEmitterHandlesById(const TSet<FGuid>& HandlesToRemove)
@@ -2161,6 +2165,7 @@ void UNiagaraSystem::RemoveEmitterHandlesById(const TSet<FGuid>& HandlesToRemove
 	EmitterHandles.RemoveAll(RemovePredicate);
 
 	InitEmitterCompiledData();
+	InvalidateCachedData();
 }
 #endif
 
