- Fix for potential nullptr access in GetNodeContextMenuActions

#rnx

diff --git a/Engine/Plugins/FX/Niagara/Source/NiagaraEditor/Private/NiagaraNodeFunctionCall.cpp b/Engine/Plugins/FX/Niagara/Source/NiagaraEditor/Private/NiagaraNodeFunctionCall.cpp
--- a/Engine/Plugins/FX/Niagara/Source/NiagaraEditor/Private/NiagaraNodeFunctionCall.cpp
+++ b/Engine/Plugins/FX/Niagara/Source/NiagaraEditor/Private/NiagaraNodeFunctionCall.cpp
@@ -1089,7 +1089,7 @@ void UNiagaraNodeFunctionCall::GetNodeContextMenuActions(class UToolMenu* Menu,
 	Super::GetNodeContextMenuActions(Menu, Context);
 	if(FunctionScript == nullptr)
 	{
-		UClass* DIClass = CastChecked<UClass>(Signature.Inputs[0].GetType().GetClass());
+		UClass* DIClass = Signature.Inputs.Num() > 0 ? CastChecked<UClass>(Signature.Inputs[0].GetType().GetClass()) : nullptr;
 		if(DIClass)
 		{
 			INiagaraDataInterfaceNodeActionProvider::GetNodeContextMenuActions(DIClass, Menu, Context, Signature);
