fix error reporting for max size
increase max size of neighborgrid

diff --git a/Engine/Plugins/FX/Niagara/Source/Niagara/Private/NiagaraDataInterfaceNeighborGrid3D.cpp b/Engine/Plugins/FX/Niagara/Source/Niagara/Private/NiagaraDataInterfaceNeighborGrid3D.cpp
--- a/Engine/Plugins/FX/Niagara/Source/Niagara/Private/NiagaraDataInterfaceNeighborGrid3D.cpp
+++ b/Engine/Plugins/FX/Niagara/Source/Niagara/Private/NiagaraDataInterfaceNeighborGrid3D.cpp
@@ -31,7 +31,7 @@ static const FName GetParticleNeighborCountFunctionName("GetParticleNeighborCoun
 static const FName SetParticleNeighborCountFunctionName("SetParticleNeighborCount");
 static const FName AddParticleFunctionName("AddParticle");
 
-static int32 GMaxNiagaraNeighborGridCells = (100*100*100);
+static int32 GMaxNiagaraNeighborGridCells = (128*128*128*64);
 static FAutoConsoleVariableRef CVarMaxNiagaraNeighborGridCells(
 	TEXT("fx.MaxNiagaraNeighborGridCells"),
 	GMaxNiagaraNeighborGridCells,
@@ -585,10 +585,15 @@ bool UNiagaraDataInterfaceNeighborGrid3D::InitPerInstanceData(void* PerInstanceD
 	InstanceData->MaxNeighborsPerCell = RT_MaxNeighborsPerCell;	
 	InstanceData->NumCells = RT_NumCells;
 
-	const uint32 NumTotalCells = RT_NumCells.X * RT_NumCells.Y * RT_NumCells.Z;
+	const uint32 NumTotalCells = RT_NumCells.X * RT_NumCells.Y * RT_NumCells.Z * RT_MaxNeighborsPerCell;
 	if (NumTotalCells == 0 || NumTotalCells > uint64(GMaxNiagaraNeighborGridCells))
 	{
-		UE_LOG(LogNiagara, Warning, TEXT("NiagaraDataInterfaceNeighborGrid3D - Invalid NumCells(%dx%dx%d) total (%u / %d) for '%s'"), RT_NumCells.X, RT_NumCells.Y, RT_NumCells.Z, NumTotalCells, GMaxNiagaraNeighborGridCells, *GetFullNameSafe(this));
+		UE_LOG(LogNiagara, Error, TEXT("NiagaraDataInterfaceNeighborGrid3D - Invalid NumCells(%dx%dx%dx%d) with total = %u exceeds the maximum of %d for '%s'.  Please adjust your NeighborGrid3D or increase fx.MaxNiagaraNeighborGridCells"), RT_NumCells.X, RT_NumCells.Y, RT_NumCells.Z, RT_MaxNeighborsPerCell, NumTotalCells, GMaxNiagaraNeighborGridCells, *GetFullNameSafe(this));
+		return false;
+	}
+	else if (NumTotalCells > GetMaxBufferDimension())
+	{
+		UE_LOG(LogNiagara, Error, TEXT("NiagaraDataInterfaceNeighborGrid3D - Invalid NumCells(%dx%dx%dx%d) with total = %u exceeds the platform maximum of %d for '%s'.  Please reduce the memory usage."), RT_NumCells.X, RT_NumCells.Y, RT_NumCells.Z, RT_MaxNeighborsPerCell, NumTotalCells, GetMaxBufferDimension(), *GetFullNameSafe(this));
 		return false;
 	}
 
@@ -628,10 +633,15 @@ void UNiagaraDataInterfaceNeighborGrid3D::SetNumCells(FVectorVMExternalFunctionC
 		const int NewMaxNeighborsPerCell = InMaxNeighborsPerCell.GetAndAdvance();
 		bool bSuccess = (InstData.Get() != nullptr && Context.GetNumInstances() == 1 && NumCells.X >= 0 && NumCells.Y >= 0 && NumCells.Z >= 0 && MaxNeighborsPerCell >= 0);
 
-		const uint32 NumTotalCells = NewNumCellsX * NewNumCellsY * NewNumCellsZ;
+		const uint32 NumTotalCells = NewNumCellsX * NewNumCellsY * NewNumCellsZ * NewMaxNeighborsPerCell;
 		if (NumTotalCells == 0 || NumTotalCells > uint64(GMaxNiagaraNeighborGridCells))
 		{
-			UE_LOG(LogNiagara, Warning, TEXT("NiagaraDataInterfaceNeighborGrid3D - Invalid NumCells(%dx%dx%d) total (%u / %d) for '%s'"), NewNumCellsX, NewNumCellsY, NewNumCellsZ, NumTotalCells, GMaxNiagaraNeighborGridCells, *GetFullNameSafe(this));
+			UE_LOG(LogNiagara, Error, TEXT("NiagaraDataInterfaceNeighborGrid3D - Invalid NumCells(%dx%dx%dx%d) with total = %u exceeds the maximum of %d for '%s'.  Please adjust your NeighborGrid3D or increase fx.MaxNiagaraNeighborGridCells"), NewNumCellsX, NewNumCellsY, NewNumCellsZ, NewMaxNeighborsPerCell, NumTotalCells, GMaxNiagaraNeighborGridCells, *GetFullNameSafe(this));
+			bSuccess = false;
+		}
+		else if (NumTotalCells > GetMaxBufferDimension())
+		{
+			UE_LOG(LogNiagara, Error, TEXT("NiagaraDataInterfaceNeighborGrid3D - Invalid NumCells(%dx%dx%dx%d) with total = %u exceeds the platform maximum of %d for '%s'.  Please reduce the memory usage."), NewNumCellsX, NewNumCellsY, NewNumCellsZ, NewMaxNeighborsPerCell, NumTotalCells, GetMaxBufferDimension(), *GetFullNameSafe(this));
 			bSuccess = false;
 		}
 
