- Add decal visible binding to decal renderer

diff --git a/Engine/Plugins/FX/Niagara/Source/Niagara/Private/NiagaraDecalRendererProperties.cpp b/Engine/Plugins/FX/Niagara/Source/Niagara/Private/NiagaraDecalRendererProperties.cpp
--- a/Engine/Plugins/FX/Niagara/Source/Niagara/Private/NiagaraDecalRendererProperties.cpp
+++ b/Engine/Plugins/FX/Niagara/Source/Niagara/Private/NiagaraDecalRendererProperties.cpp
@@ -46,7 +46,13 @@ namespace NiagaraDecalRendererPropertiesLocal
 
 	static FNiagaraVariable GetDecalFadeVariable()
 	{
-		FNiagaraVariable Variable = MakeNiagaraVariableWithValue(FNiagaraTypeDefinition::GetFloatDef(), TEXT("Particles.DecalFade"), UNiagaraDecalRendererProperties::GetDefaultDecalFade());
+		static FNiagaraVariable Variable = MakeNiagaraVariableWithValue(FNiagaraTypeDefinition::GetFloatDef(), TEXT("Particles.DecalFade"), UNiagaraDecalRendererProperties::GetDefaultDecalFade());
+		return Variable;
+	}
+
+	static FNiagaraVariable GetDecalVisibleVariable()
+	{
+		static FNiagaraVariable Variable = MakeNiagaraVariableWithValue(FNiagaraTypeDefinition::GetBoolDef(), TEXT("Particles.DecalVisible"), UNiagaraDecalRendererProperties::GetDefaultDecalVisible());
 		return Variable;
 	}
 
@@ -60,7 +66,8 @@ namespace NiagaraDecalRendererPropertiesLocal
 		Props->DecalOrientationBinding.Setup(GetDecalOrientationVariable(), GetDecalOrientationVariable());
 		Props->DecalSizeBinding.Setup(GetDecalSizeVariable(), GetDecalSizeVariable());
 		Props->DecalFadeBinding.Setup(GetDecalFadeVariable(), GetDecalFadeVariable());
-		Props->ColorBinding = FNiagaraConstants::GetAttributeDefaultBinding(SYS_PARAM_PARTICLES_COLOR);
+		Props->DecalColorBinding = FNiagaraConstants::GetAttributeDefaultBinding(SYS_PARAM_PARTICLES_COLOR);
+		Props->DecalVisibleBinding.Setup(GetDecalVisibleVariable(), GetDecalVisibleVariable());
 		Props->RendererVisibilityTagBinding = FNiagaraConstants::GetAttributeDefaultBinding(SYS_PARAM_PARTICLES_VISIBILITY_TAG);
 
 	#if WITH_EDITORONLY_DATA
@@ -77,7 +84,8 @@ UNiagaraDecalRendererProperties::UNiagaraDecalRendererProperties()
 	AttributeBindings.Add(&DecalOrientationBinding);
 	AttributeBindings.Add(&DecalSizeBinding);
 	AttributeBindings.Add(&DecalFadeBinding);
-	AttributeBindings.Add(&ColorBinding);
+	AttributeBindings.Add(&DecalColorBinding);
+	AttributeBindings.Add(&DecalVisibleBinding);
 	AttributeBindings.Add(&RendererVisibilityTagBinding);
 }
 
@@ -208,7 +216,8 @@ void UNiagaraDecalRendererProperties::CacheFromCompiledData(const FNiagaraDataSe
 	DecalOrientationDataSetAccessor.Init(CompiledData, DecalOrientationBinding.GetDataSetBindableVariable().GetName());
 	DecalSizeDataSetAccessor.Init(CompiledData, DecalSizeBinding.GetDataSetBindableVariable().GetName());
 	DecalFadeDataSetAccessor.Init(CompiledData, DecalFadeBinding.GetDataSetBindableVariable().GetName());
-	ColorDataSetAccessor.Init(CompiledData, ColorBinding.GetDataSetBindableVariable().GetName());
+	DecalColorDataSetAccessor.Init(CompiledData, DecalColorBinding.GetDataSetBindableVariable().GetName());
+	DecalVisibleAccessor.Init(CompiledData, DecalVisibleBinding.GetDataSetBindableVariable().GetName());
 	RendererVisibilityTagAccessor.Init(CompiledData, RendererVisibilityTagBinding.GetDataSetBindableVariable().GetName());
 }
 
diff --git a/Engine/Plugins/FX/Niagara/Source/Niagara/Private/NiagaraRendererDecals.cpp b/Engine/Plugins/FX/Niagara/Source/Niagara/Private/NiagaraRendererDecals.cpp
--- a/Engine/Plugins/FX/Niagara/Source/Niagara/Private/NiagaraRendererDecals.cpp
+++ b/Engine/Plugins/FX/Niagara/Source/Niagara/Private/NiagaraRendererDecals.cpp
@@ -153,6 +153,7 @@ FNiagaraDynamicDataBase* FNiagaraRendererDecals::GenerateDynamicData(const FNiag
 	const FNiagaraPosition DefaultPos = bUseLocalSpace ? FVector::ZeroVector : LocalToWorld.GetLocation();
 	const FQuat4f DefaultRot = UNiagaraDecalRendererProperties::GetDefaultOrientation();
 	const FVector3f DefaultSize = UNiagaraDecalRendererProperties::GetDefaultDecalSize();
+	const FNiagaraBool DefaultVisible = UNiagaraDecalRendererProperties::GetDefaultDecalVisible();
 	const float DefaultFade = UNiagaraDecalRendererProperties::GetDefaultDecalFade();
 
 	// Check for Material Update
@@ -176,15 +177,17 @@ FNiagaraDynamicDataBase* FNiagaraRendererDecals::GenerateDynamicData(const FNiag
 		const auto PositionReader = RendererProperties->PositionDataSetAccessor.GetReader(DataSet);
 		const auto RotationReader = RendererProperties->DecalOrientationDataSetAccessor.GetReader(DataSet);
 		const auto SizeReader = RendererProperties->DecalSizeDataSetAccessor.GetReader(DataSet);
-		const auto ColorReader = RendererProperties->ColorDataSetAccessor.GetReader(DataSet);
+		const auto ColorReader = RendererProperties->DecalColorDataSetAccessor.GetReader(DataSet);
 		const auto FadeReader = RendererProperties->DecalFadeDataSetAccessor.GetReader(DataSet);
+		const auto VisibleReader = RendererProperties->DecalVisibleAccessor.GetReader(DataSet);
 		const auto VisTagReader = RendererProperties->RendererVisibilityTagAccessor.GetReader(DataSet);
 
 		for (uint32 ParticleIndex = 0; ParticleIndex < DataToRender->GetNumInstances(); ++ParticleIndex)
 		{
-			// Check visibility tag
+			// Check if the decal is visible
+			const bool bVisible = VisibleReader.GetSafe(ParticleIndex, DefaultVisible).GetValue();
 			const int32 VisTag = VisTagReader.GetSafe(ParticleIndex, RendererProperties->RendererVisibility);
-			if (VisTag != RendererProperties->RendererVisibility)
+			if (!bVisible || VisTag != RendererProperties->RendererVisibility)
 			{
 				continue;
 			}
@@ -227,15 +230,16 @@ FNiagaraDynamicDataBase* FNiagaraRendererDecals::GenerateDynamicData(const FNiag
 	{
 		//-OPT: Emitter mode we should be able to cache the parameter offset rather than search each frame
 		const FNiagaraParameterStore& ParameterStore = Emitter->GetRendererBoundVariables();
+		const bool bVisible = ParameterStore.GetParameterValueOrDefault(RendererProperties->DecalVisibleBinding.GetParamMapBindableVariable(), DefaultVisible).GetValue();
 		const int32 VisTag = ParameterStore.GetParameterValueOrDefault(RendererProperties->RendererVisibilityTagBinding.GetParamMapBindableVariable(), RendererProperties->RendererVisibility);
-		if (VisTag == RendererProperties->RendererVisibility)
+		if (bVisible && VisTag == RendererProperties->RendererVisibility)
 		{
 			const FVector3f SimPos = ParameterStore.GetParameterValueOrDefault(RendererProperties->PositionBinding.GetParamMapBindableVariable(), DefaultPos);
 			const FQuat4f SimRot = ParameterStore.GetParameterValueOrDefault(RendererProperties->DecalOrientationBinding.GetParamMapBindableVariable(), DefaultRot);
 			const FVector Position = bUseLocalSpace ? LocalToWorld.TransformPosition(FVector(SimPos)) : LwcConverter.ConvertSimulationPositionToWorld(SimPos);
 			const FQuat Rotation = bUseLocalSpace ? LocalToWorld.TransformRotation(FQuat(SimRot)) : FQuat(SimRot);
 			const FVector Size = FVector(ParameterStore.GetParameterValueOrDefault(RendererProperties->DecalSizeBinding.GetParamMapBindableVariable(), DefaultSize) * 0.5f);
-			const FLinearColor DecalColor = ParameterStore.GetParameterValueOrDefault(RendererProperties->ColorBinding.GetParamMapBindableVariable(), FLinearColor::White);
+			const FLinearColor DecalColor = ParameterStore.GetParameterValueOrDefault(RendererProperties->DecalColorBinding.GetParamMapBindableVariable(), FLinearColor::White);
 			const float Fade = ParameterStore.GetParameterValueOrDefault(RendererProperties->DecalFadeBinding.GetParamMapBindableVariable(), DefaultFade);
 
 			// Create Update Parameters
diff --git a/Engine/Plugins/FX/Niagara/Source/Niagara/Public/NiagaraDecalRendererProperties.h b/Engine/Plugins/FX/Niagara/Source/Niagara/Public/NiagaraDecalRendererProperties.h
--- a/Engine/Plugins/FX/Niagara/Source/Niagara/Public/NiagaraDecalRendererProperties.h
+++ b/Engine/Plugins/FX/Niagara/Source/Niagara/Public/NiagaraDecalRendererProperties.h
@@ -55,6 +55,7 @@ public:
 	static const FQuat4f GetDefaultOrientation() { return FRotator3f(-90.0f, 0.0f, 90.0f).Quaternion(); }
 	static const FVector3f GetDefaultDecalSize() { return FVector3f(50.0f, 50.0f, 50.0f); }
 	static const float GetDefaultDecalFade() { return 1.0f; }
+	static const FNiagaraBool GetDefaultDecalVisible() { return FNiagaraBool(true); }
 
 	/** What material to use for the decal. */
 	UPROPERTY(EditAnywhere, Category = "Decal Rendering")
@@ -94,7 +95,11 @@ public:
 
 	/** Color binding for the decal, value can be queried using the Decal Color material node. */
 	UPROPERTY(EditAnywhere, Category = "Bindings")
-	FNiagaraVariableAttributeBinding ColorBinding;
+	FNiagaraVariableAttributeBinding DecalColorBinding;
+
+	/** Should the decal be visibile or not, works in conjunction with RendererVisibilityTagBinding to determine visibility. */
+	UPROPERTY(EditAnywhere, Category = "Bindings")
+	FNiagaraVariableAttributeBinding DecalVisibleBinding;
 
 	/** Visibility tag binding, when valid the returned values is compated with RendererVisibility. */
 	UPROPERTY(EditAnywhere, Category = "Bindings")
@@ -104,6 +109,7 @@ public:
 	FNiagaraDataSetAccessor<FQuat4f>			DecalOrientationDataSetAccessor;
 	FNiagaraDataSetAccessor<FVector3f>			DecalSizeDataSetAccessor;
 	FNiagaraDataSetAccessor<float>				DecalFadeDataSetAccessor;
-	FNiagaraDataSetAccessor<FLinearColor>		ColorDataSetAccessor;
+	FNiagaraDataSetAccessor<FLinearColor>		DecalColorDataSetAccessor;
+	FNiagaraDataSetAccessor<FNiagaraBool>		DecalVisibleAccessor;
 	FNiagaraDataSetAccessor<int32>				RendererVisibilityTagAccessor;
 };
