Remove editcondition for SubImageSize from sprite & mesh renderer as they are used independently of bSubImageBlend
#rnx

diff --git a/Engine/Plugins/FX/Niagara/Source/Niagara/Public/NiagaraMeshRendererProperties.h b/Engine/Plugins/FX/Niagara/Source/Niagara/Public/NiagaraMeshRendererProperties.h
--- a/Engine/Plugins/FX/Niagara/Source/Niagara/Public/NiagaraMeshRendererProperties.h
+++ b/Engine/Plugins/FX/Niagara/Source/Niagara/Public/NiagaraMeshRendererProperties.h
@@ -294,7 +294,7 @@ public:
 	TArray<FNiagaraMeshMICOverride> MICOverrideMaterials;
 
 	/** When using SubImage lookups for particles, this variable contains the number of columns in X and the number of rows in Y.*/
-	UPROPERTY(EditAnywhere, Category = "SubUV", meta = (EditCondition = "bSubImageBlend", EditConditionHides))
+	UPROPERTY(EditAnywhere, Category = "SubUV")
 	FVector2D SubImageSize;
 
 	/** Determines how the mesh orients itself relative to the camera. */
diff --git a/Engine/Plugins/FX/Niagara/Source/Niagara/Public/NiagaraSpriteRendererProperties.h b/Engine/Plugins/FX/Niagara/Source/Niagara/Public/NiagaraSpriteRendererProperties.h
--- a/Engine/Plugins/FX/Niagara/Source/Niagara/Public/NiagaraSpriteRendererProperties.h
+++ b/Engine/Plugins/FX/Niagara/Source/Niagara/Public/NiagaraSpriteRendererProperties.h
@@ -187,7 +187,7 @@ public:
 	ENiagaraSortMode SortMode;
 
 	/** When using SubImage lookups for particles, this variable contains the number of columns in X and the number of rows in Y.*/
-	UPROPERTY(EditAnywhere, Category = "SubUV", meta = (EditCondition = "bSubImageBlend", EditConditionHides, DisplayAfter = bSubImageBlend))
+	UPROPERTY(EditAnywhere, Category = "SubUV", meta = (DisplayAfter = bSubImageBlend))
 	FVector2D SubImageSize;
 
 	/** If true, blends the sub-image UV lookup with its next adjacent member using the fractional part of the SubImageIndex float value as the linear interpolation factor.*/
