Changed hierarchy data to use a custom identity struct instead of just guids

diff --git a/Engine/Plugins/FX/Niagara/Source/NiagaraEditor/Private/ViewModels/HierarchyEditor/NiagaraHierarchyViewModelBase.cpp b/Engine/Plugins/FX/Niagara/Source/NiagaraEditor/Private/ViewModels/HierarchyEditor/NiagaraHierarchyViewModelBase.cpp
--- a/Engine/Plugins/FX/Niagara/Source/NiagaraEditor/Private/ViewModels/HierarchyEditor/NiagaraHierarchyViewModelBase.cpp
+++ b/Engine/Plugins/FX/Niagara/Source/NiagaraEditor/Private/ViewModels/HierarchyEditor/NiagaraHierarchyViewModelBase.cpp
@@ -83,6 +83,16 @@ bool UNiagaraHierarchyItemBase::Modify(bool bAlwaysMarkDirty)
 	return bSavedToTransactionBuffer;
 }
 
+void UNiagaraHierarchyItemBase::PostLoad()
+{
+	if(Guid_DEPRECATED.IsValid())
+	{
+		SetIdentity(FNiagaraHierarchyIdentity({Guid_DEPRECATED}, {}));
+	}
+
+	Super::PostLoad();
+}
+
 void UNiagaraHierarchyRoot::RefreshDataInternal()
 {
 	// in addition, we update our sections
diff --git a/Engine/Plugins/FX/Niagara/Source/NiagaraEditor/Private/ViewModels/HierarchyEditor/NiagaraUserParametersHierarchyViewModel.cpp b/Engine/Plugins/FX/Niagara/Source/NiagaraEditor/Private/ViewModels/HierarchyEditor/NiagaraUserParametersHierarchyViewModel.cpp
--- a/Engine/Plugins/FX/Niagara/Source/NiagaraEditor/Private/ViewModels/HierarchyEditor/NiagaraUserParametersHierarchyViewModel.cpp
+++ b/Engine/Plugins/FX/Niagara/Source/NiagaraEditor/Private/ViewModels/HierarchyEditor/NiagaraUserParametersHierarchyViewModel.cpp
@@ -21,13 +21,13 @@
 void UNiagaraHierarchyUserParameter::Initialize(UNiagaraScriptVariable& InUserParameterScriptVariable, UNiagaraSystem& InSystem)
 {
 	UserParameterScriptVariable = &InUserParameterScriptVariable;
-	SetGuid(InUserParameterScriptVariable.Metadata.GetVariableGuid());
+	SetIdentity(FNiagaraHierarchyIdentity({InUserParameterScriptVariable.Metadata.GetVariableGuid()}, {}));
 	System = &InSystem;
 }
 
 void UNiagaraHierarchyUserParameter::RefreshDataInternal()
 {	
-	if(System == nullptr || FNiagaraEditorUtilities::FindScriptVariableForUserParameter(GetPersistentIdentity(), *System) == nullptr)
+	if(System == nullptr || FNiagaraEditorUtilities::FindScriptVariableForUserParameter(GetPersistentIdentity().Guids[0], *System) == nullptr)
 	{
 		Finalize();
 	}
@@ -115,7 +115,7 @@ void UNiagaraUserParametersHierarchyViewModel::PrepareSourceItems()
 		
 		if(UNiagaraHierarchyItemBase** ExistingChild = OldChildren.FindByPredicate([ScriptVariable](UNiagaraHierarchyItemBase* ItemBase)
 		{
-			return ItemBase->GetPersistentIdentity() == ScriptVariable->Metadata.GetVariableGuid();
+			return ItemBase->GetPersistentIdentity().Guids[0] == ScriptVariable->Metadata.GetVariableGuid();
 		}))
 		{
 			SourceRoot->GetChildrenMutable().Add(*ExistingChild);
diff --git a/Engine/Plugins/FX/Niagara/Source/NiagaraEditor/Public/ViewModels/HierarchyEditor/NiagaraHierarchyViewModelBase.h b/Engine/Plugins/FX/Niagara/Source/NiagaraEditor/Public/ViewModels/HierarchyEditor/NiagaraHierarchyViewModelBase.h
--- a/Engine/Plugins/FX/Niagara/Source/NiagaraEditor/Public/ViewModels/HierarchyEditor/NiagaraHierarchyViewModelBase.h
+++ b/Engine/Plugins/FX/Niagara/Source/NiagaraEditor/Public/ViewModels/HierarchyEditor/NiagaraHierarchyViewModelBase.h
@@ -9,18 +9,64 @@
 #include "ToolMenuSection.h"
 #include "NiagaraHierarchyViewModelBase.generated.h"
 
+USTRUCT()
+struct FNiagaraHierarchyIdentity
+{
+	GENERATED_BODY()
+
+	FNiagaraHierarchyIdentity() {}
+	FNiagaraHierarchyIdentity(TArray<FGuid> InGuids, TArray<FName> InNames) : Guids(InGuids), Names(InNames) {}
+	
+	/** An array of guids that have to be satisfied in order to match. */
+	UPROPERTY()
+	TArray<FGuid> Guids;
+
+	/** Optionally, an array of names can be specified in place of guids. If guids & names are present, guids have to be satisfied first, then names. */
+	UPROPERTY()
+	TArray<FName> Names;
+
+	bool operator==(const FNiagaraHierarchyIdentity& OtherIdentity)
+	{
+		if(Guids.Num() != OtherIdentity.Guids.Num() || Names.Num() != OtherIdentity.Names.Num())
+		{
+			return false;
+		}
+
+		for(int32 GuidIndex = 0; GuidIndex < Guids.Num(); GuidIndex++)
+		{
+			if(Guids[GuidIndex] != OtherIdentity.Guids[GuidIndex])
+			{
+				return false;
+			}
+		}
+
+		for(int32 NameIndex = 0; NameIndex < Names.Num(); NameIndex++)
+		{
+			if(!Names[NameIndex].IsEqual(OtherIdentity.Names[NameIndex]))
+			{
+				return false;
+			}
+		}
+
+		return true;
+	}
+};
+
 UCLASS()
 class UNiagaraHierarchyItemBase : public UObject
 {
 	GENERATED_BODY()
 
 public:
-	UNiagaraHierarchyItemBase() { Guid = FGuid::NewGuid(); }
+	UNiagaraHierarchyItemBase() { Identity.Guids.Add(FGuid::NewGuid()); }
 	virtual ~UNiagaraHierarchyItemBase() override {}
 
 	TArray<UNiagaraHierarchyItemBase*>& GetChildrenMutable() { return Children; }
 	const TArray<UNiagaraHierarchyItemBase*>& GetChildren() const;
 
+	template<class ChildClass>
+	ChildClass* AddChild();
+	
 	template<class ChildClass>
 	bool DoesOneChildExist(bool bRecursive = false) const;
 
@@ -36,9 +82,9 @@ public:
 	void Finalize();
 	bool IsFinalized() const { return bFinalized; }
 
-	/** A guid can be optionally set to create a mapping from previously existing guids to hierarchy items that represent them. */
-	void SetGuid(FGuid InGuid) { Guid = InGuid; }
-	virtual FGuid GetPersistentIdentity() const { return Guid; }
+	/** An identity can be optionally set to create a mapping from previously existing guids or names to hierarchy items that represent them. */
+	void SetIdentity(FNiagaraHierarchyIdentity InIdentity) { Identity = InIdentity; }
+	virtual FNiagaraHierarchyIdentity GetPersistentIdentity() const { return Identity; }
 	
 	virtual bool CanHaveChildren() const { return false; }
 
@@ -49,17 +95,32 @@ protected:
 	/** Called by the public Refresh function. Can be overridden to further customize the refresh process, i.e. refreshing section data in the root. */
 	virtual void RefreshDataInternal() { }
 
+protected:
+	virtual void PostLoad() override;
+
 	UPROPERTY()
 	TArray<TObjectPtr<UNiagaraHierarchyItemBase>> Children;
+
+	UPROPERTY()
+	FNiagaraHierarchyIdentity Identity;
 	
 	/** An optional guid; can be used if hierarchy items represent outside items */
 	UPROPERTY()
-	FGuid Guid;
+	FGuid Guid_DEPRECATED;
 
 	UPROPERTY(Transient)
 	bool bFinalized = false;
 };
 
+template <class ChildClass>
+ChildClass* UNiagaraHierarchyItemBase::AddChild()
+{
+	ChildClass* NewChild = NewObject<ChildClass>(this);
+	GetChildrenMutable().Add(NewChild);
+
+	return NewChild;
+}
+
 template <class ChildClass>
 bool UNiagaraHierarchyItemBase::DoesOneChildExist(bool bRecursive) const
 {
@@ -383,7 +444,7 @@ struct NIAGARAEDITOR_API FNiagaraHierarchyItemViewModelBase : TSharedFromThis<FN
 	virtual void Tick(float DeltaTime) override;
 	virtual TStatId GetStatId() const override;
 	
-	FString ToString() const { return ItemBase->ToString(); }
+	virtual FString ToString() const { return ItemBase->ToString(); }
 	
 	void SyncToData();
 	const TArray<TSharedPtr<FNiagaraHierarchyItemViewModelBase>>& GetChildren() const { return Children; }
