- Fix for crash with function call context actions

#rnx

diff --git a/Engine/Plugins/FX/Niagara/Source/NiagaraEditor/Private/NiagaraNodeFunctionCall.cpp b/Engine/Plugins/FX/Niagara/Source/NiagaraEditor/Private/NiagaraNodeFunctionCall.cpp
--- a/Engine/Plugins/FX/Niagara/Source/NiagaraEditor/Private/NiagaraNodeFunctionCall.cpp
+++ b/Engine/Plugins/FX/Niagara/Source/NiagaraEditor/Private/NiagaraNodeFunctionCall.cpp
@@ -1087,13 +1087,19 @@ void UNiagaraNodeFunctionCall::UpdateReferencedStaticsHashForNode(FSHA1& HashSta
 void UNiagaraNodeFunctionCall::GetNodeContextMenuActions(class UToolMenu* Menu, class UGraphNodeContextMenuContext* Context) const
 {
 	Super::GetNodeContextMenuActions(Menu, Context);
-	if(FunctionScript == nullptr)
+	if ( FunctionScript != nullptr || Signature.Inputs.Num() == 0 )
 	{
-		UClass* DIClass = Signature.Inputs.Num() > 0 ? CastChecked<UClass>(Signature.Inputs[0].GetType().GetClass()) : nullptr;
-		if(DIClass)
-		{
-			INiagaraDataInterfaceNodeActionProvider::GetNodeContextMenuActions(DIClass, Menu, Context, Signature);
-		}
+		return;
+	}
+
+	if ( Signature.Inputs[0].GetType().IsDataInterface() == false )
+	{
+		return;
+	}
+
+	if ( UClass* DIClass = Cast<UClass>(Signature.Inputs[0].GetType().GetClass()) )
+	{
+		INiagaraDataInterfaceNodeActionProvider::GetNodeContextMenuActions(DIClass, Menu, Context, Signature);
 	}
 }
 
