- Check for null iteration interface
- We fail to compile if we don't have a valid DI, however we can still get a null DI in the stage

#rnx

diff --git a/Engine/Plugins/FX/Niagara/Source/Niagara/Private/NiagaraGpuComputeDispatch.cpp b/Engine/Plugins/FX/Niagara/Source/Niagara/Private/NiagaraGpuComputeDispatch.cpp
--- a/Engine/Plugins/FX/Niagara/Source/Niagara/Private/NiagaraGpuComputeDispatch.cpp
+++ b/Engine/Plugins/FX/Niagara/Source/Niagara/Private/NiagaraGpuComputeDispatch.cpp
@@ -1442,6 +1442,11 @@ void FNiagaraGpuComputeDispatch::DispatchStage(FRDGBuilder& GraphBuilder, const
 
 		case ENiagaraIterationSource::DataInterface:
 		{
+			if (SimStageData.AlternateIterationSource == nullptr)
+			{
+				return;
+			}
+
 			DispatchType = SimStageData.StageMetaData->GpuDispatchType;
 			DispatchCount = SimStageData.AlternateIterationSource->GetElementCount(Tick.SystemInstanceID);
 			DispatchNumThreads = SimStageData.StageMetaData->GpuDispatchNumThreads;
