- Fix various truncation warnings in Niagara

#rnx

diff --git a/Engine/Plugins/FX/Niagara/Source/Niagara/Classes/NiagaraDataInterface.h b/Engine/Plugins/FX/Niagara/Source/Niagara/Classes/NiagaraDataInterface.h
--- a/Engine/Plugins/FX/Niagara/Source/Niagara/Classes/NiagaraDataInterface.h
+++ b/Engine/Plugins/FX/Niagara/Source/Niagara/Classes/NiagaraDataInterface.h
@@ -972,7 +972,7 @@ struct TNDIRandomHelper : public TRandomPolicy
 		//Can probably handle lower quality randoms / uniformity for a decent speed win.
 		FVector2f r = FVector2f(TRandomPolicy::Rand2(InstanceIndex));
 		float sqrt0 = FMath::Sqrt(r.X);
-		return FVector3f(1.0f - sqrt0, sqrt0 * (1.0 - r.Y), r.Y * sqrt0);
+		return FVector3f(1.0f - sqrt0, sqrt0 * (1.0f - r.Y), r.Y * sqrt0);
 	}
 };
 
diff --git a/Engine/Plugins/FX/Niagara/Source/Niagara/Classes/NiagaraDataInterfaceMeshCommon.h b/Engine/Plugins/FX/Niagara/Source/Niagara/Classes/NiagaraDataInterfaceMeshCommon.h
--- a/Engine/Plugins/FX/Niagara/Source/Niagara/Classes/NiagaraDataInterfaceMeshCommon.h
+++ b/Engine/Plugins/FX/Niagara/Source/Niagara/Classes/NiagaraDataInterfaceMeshCommon.h
@@ -37,7 +37,7 @@ FORCEINLINE FVector3f RandomBarycentricCoord(FRandomStream& RandStream)
 	float r0 = RandStream.GetFraction();
 	float r1 = RandStream.GetFraction();
 	float sqrt0 = FMath::Sqrt(r0);
-	return FVector3f(1.0f - sqrt0, sqrt0 * (1.0 - r1), r1 * sqrt0);
+	return FVector3f(1.0f - sqrt0, sqrt0 * (1.0f - r1), r1 * sqrt0);
 }
 
 template<typename T>
diff --git a/Engine/Plugins/FX/Niagara/Source/Niagara/Classes/NiagaraDataInterfaceVector2DCurve.h b/Engine/Plugins/FX/Niagara/Source/Niagara/Classes/NiagaraDataInterfaceVector2DCurve.h
--- a/Engine/Plugins/FX/Niagara/Source/Niagara/Classes/NiagaraDataInterfaceVector2DCurve.h
+++ b/Engine/Plugins/FX/Niagara/Source/Niagara/Classes/NiagaraDataInterfaceVector2DCurve.h
@@ -72,7 +72,7 @@ protected:
 
 private:
 	template<typename UseLUT>
-	FORCEINLINE_DEBUGGABLE FVector2D SampleCurveInternal(float X);
+	FORCEINLINE_DEBUGGABLE FVector2f SampleCurveInternal(float X);
 
 	static const FName SampleCurveName;
 };
diff --git a/Engine/Plugins/FX/Niagara/Source/Niagara/Classes/NiagaraDataInterfaceVectorCurve.h b/Engine/Plugins/FX/Niagara/Source/Niagara/Classes/NiagaraDataInterfaceVectorCurve.h
--- a/Engine/Plugins/FX/Niagara/Source/Niagara/Classes/NiagaraDataInterfaceVectorCurve.h
+++ b/Engine/Plugins/FX/Niagara/Source/Niagara/Classes/NiagaraDataInterfaceVectorCurve.h
@@ -80,7 +80,7 @@ protected:
 
 private:
 	template<typename UseLUT>
-	FORCEINLINE_DEBUGGABLE FVector SampleCurveInternal(float X);
+	FORCEINLINE_DEBUGGABLE FVector3f SampleCurveInternal(float X);
 
 	static const FName SampleCurveName;
 };
diff --git a/Engine/Plugins/FX/Niagara/Source/Niagara/Classes/NiagaraGPUSystemTick.h b/Engine/Plugins/FX/Niagara/Source/Niagara/Classes/NiagaraGPUSystemTick.h
--- a/Engine/Plugins/FX/Niagara/Source/Niagara/Classes/NiagaraGPUSystemTick.h
+++ b/Engine/Plugins/FX/Niagara/Source/Niagara/Classes/NiagaraGPUSystemTick.h
@@ -112,7 +112,7 @@ public:
 
 	FORCEINLINE FRHIUniformBuffer* GetExternalUniformBuffer(const FNiagaraComputeInstanceData& InstanceData, bool bPrevious) const
 	{
-		const int32 InstanceIndex = &InstanceData - GetInstances().GetData();
+		const int32 InstanceIndex = UE_PTRDIFF_TO_INT32(&InstanceData - GetInstances().GetData());
 		const int32 BufferIndex = InstanceIndex + (bPrevious ? InstanceCount : 0);
 		return ExternalUnformBuffers_RT[BufferIndex];
 	}
diff --git a/Engine/Plugins/FX/Niagara/Source/Niagara/Classes/NiagaraPerfBaseline.h b/Engine/Plugins/FX/Niagara/Source/Niagara/Classes/NiagaraPerfBaseline.h
--- a/Engine/Plugins/FX/Niagara/Source/Niagara/Classes/NiagaraPerfBaseline.h
+++ b/Engine/Plugins/FX/Niagara/Source/Niagara/Classes/NiagaraPerfBaseline.h
@@ -234,7 +234,7 @@ public:
 	uint32 TotalTests = 0;
 
 	/** We store the current target world time so we can give an indication of the time range for poor tests in the report. */
-	float CurrentWorldTime = 0.0f;
+	double CurrentWorldTime = 0.0;
 	uint32 CurrentFrameNumber = 0;
 	FString TestNameString;
 
@@ -242,8 +242,8 @@ public:
 	struct FStatTestInfo
 	{
 		uint32 TestIndex = 0;
-		float StartTime = 0.0f;
-		float EndTime = 0.0f;
+		double StartTime = 0.0;
+		double EndTime = 0.0;
 		FAccumulatedParticlePerfStats AccumulatedStats;
 	};
 
@@ -296,7 +296,7 @@ private:
 	};
 	TArray<TWeakObjectPtr<UNiagaraEffectType>> BaselinesToGenerate;
 	EBaselineGenState BaselineGenerationState = EBaselineGenState::Complete;
-	float WorldTimeToGenerate = 0.0f;
+	double WorldTimeToGenerate = 0.0;
 	TWeakObjectPtr<UWorld> BaselineGenWorld;
 };
 
diff --git a/Engine/Plugins/FX/Niagara/Source/Niagara/Private/DataInterface/NiagaraDataInterfaceDebugDraw.cpp b/Engine/Plugins/FX/Niagara/Source/Niagara/Private/DataInterface/NiagaraDataInterfaceDebugDraw.cpp
--- a/Engine/Plugins/FX/Niagara/Source/Niagara/Private/DataInterface/NiagaraDataInterfaceDebugDraw.cpp
+++ b/Engine/Plugins/FX/Niagara/Source/Niagara/Private/DataInterface/NiagaraDataInterfaceDebugDraw.cpp
@@ -274,9 +274,9 @@ struct FNDIDebugDrawInstanceData_GameThread
 	{
 		const FVector3f HeightVector = FVector3f::ZAxisVector * Height;
 		const FVector3f HalfHeightOffset = bCenterVertically ? HeightVector * 0.5f : FVector3f::ZeroVector;
-		const FVector3f AxisStep = FVector3f::ZAxisVector * Height / NumHeightSegments;
+		const FVector3f AxisStep = FVector3f::ZAxisVector * Height / float(NumHeightSegments);
 
-		const float HeightSegmentStepAlpha = 1.0 / NumHeightSegments;
+		const float HeightSegmentStepAlpha = 1.0f / float(NumHeightSegments);
 
 		float CurrentRotation = 0.0f;
 		FVector3f LastPointNormal = (FVector3f::YAxisVector * FMath::Cos(CurrentRotation)) + (FVector3f::XAxisVector * FMath::Sin(CurrentRotation));
@@ -1012,8 +1012,8 @@ struct FNDIDebugDrawInstanceData_GameThread
 			TOptional<FQuat4f> Rotation = GetTag<FQuat4f>(SystemInstance, FNiagaraTypeDefinition::GetQuatDef(), 4);
 			TOptional<ENiagaraCoordinateSpace> RotationWorldSpace = GetTag<ENiagaraCoordinateSpace>(SystemInstance, CoordTypeDef, 5);
 			TOptional<FVector2f> Extents = GetTag<FVector2f>(SystemInstance, FNiagaraTypeDefinition::GetVec2Def(), 6);
-			TOptional<float> NumCellsX = GetTag<float>(SystemInstance, FNiagaraTypeDefinition::GetFloatDef(), 7);
-			TOptional<float> NumCellsY = GetTag<float>(SystemInstance, FNiagaraTypeDefinition::GetFloatDef(), 8);
+			TOptional<int32> NumCellsX = GetTag<int32>(SystemInstance, FNiagaraTypeDefinition::GetIntDef(), 7);
+			TOptional<int32> NumCellsY = GetTag<int32>(SystemInstance, FNiagaraTypeDefinition::GetIntDef(), 8);
 			TOptional<FLinearColor> Color = GetTag<FLinearColor>(SystemInstance, FNiagaraTypeDefinition::GetColorDef(), 9);
 
 
@@ -1049,9 +1049,9 @@ struct FNDIDebugDrawInstanceData_GameThread
 			TOptional<FQuat4f> Rotation = GetTag<FQuat4f>(SystemInstance, FNiagaraTypeDefinition::GetQuatDef(), 4);
 			TOptional<ENiagaraCoordinateSpace> RotationWorldSpace = GetTag<ENiagaraCoordinateSpace>(SystemInstance, CoordTypeDef, 5);
 			TOptional<FVector3f> Extents = GetTag<FVector3f>(SystemInstance, FNiagaraTypeDefinition::GetVec3Def(), 6);
-			TOptional<float> NumCellsX = GetTag<float>(SystemInstance, FNiagaraTypeDefinition::GetFloatDef(), 7);
-			TOptional<float> NumCellsY = GetTag<float>(SystemInstance, FNiagaraTypeDefinition::GetFloatDef(), 8);
-			TOptional<float> NumCellsZ = GetTag<float>(SystemInstance, FNiagaraTypeDefinition::GetFloatDef(), 9);
+			TOptional<int32> NumCellsX = GetTag<int32>(SystemInstance, FNiagaraTypeDefinition::GetIntDef(), 7);
+			TOptional<int32> NumCellsY = GetTag<int32>(SystemInstance, FNiagaraTypeDefinition::GetIntDef(), 8);
+			TOptional<int32> NumCellsZ = GetTag<int32>(SystemInstance, FNiagaraTypeDefinition::GetIntDef(), 9);
 			TOptional<FLinearColor> Color = GetTag<FLinearColor>(SystemInstance, FNiagaraTypeDefinition::GetColorDef(), 10);
 
 
diff --git a/Engine/Plugins/FX/Niagara/Source/Niagara/Private/DataInterface/NiagaraDataInterfaceDynamicMesh.cpp b/Engine/Plugins/FX/Niagara/Source/Niagara/Private/DataInterface/NiagaraDataInterfaceDynamicMesh.cpp
--- a/Engine/Plugins/FX/Niagara/Source/Niagara/Private/DataInterface/NiagaraDataInterfaceDynamicMesh.cpp
+++ b/Engine/Plugins/FX/Niagara/Source/Niagara/Private/DataInterface/NiagaraDataInterfaceDynamicMesh.cpp
@@ -191,8 +191,8 @@ namespace NDIDynamicMeshLocal
 			check(Vertex >= 0 && Vertex < NumVertices && TangentBasisOffset != INDEX_NONE);
 			const uint32* TangentData = reinterpret_cast<const uint32*>(VertexData.GetData() + (Vertex * sizeof(FPackedNormal) * 2) + TangentBasisOffset);
 
-			OutTangentX = FVector3f((TangentData[0] >> 0) & 0xff, (TangentData[0] >> 8) & 0xff, (TangentData[0] >> 16) & 0xff);
-			OutTangentZ = FVector3f((TangentData[1] >> 0) & 0xff, (TangentData[1] >> 8) & 0xff, (TangentData[1] >> 16) & 0xff);
+			OutTangentX = FVector3f(float((TangentData[0] >> 0) & 0xff), float((TangentData[0] >> 8) & 0xff), float((TangentData[0] >> 16) & 0xff));
+			OutTangentZ = FVector3f(float((TangentData[1] >> 0) & 0xff), float((TangentData[1] >> 8) & 0xff), float((TangentData[1] >> 16) & 0xff));
 			OutTangentX = (OutTangentX / 127.5f) - 1.0f;
 			OutTangentZ = (OutTangentX / 127.5f) - 1.0f;
 
diff --git a/Engine/Plugins/FX/Niagara/Source/Niagara/Private/DataInterface/NiagaraDataInterfaceStaticMesh.cpp b/Engine/Plugins/FX/Niagara/Source/Niagara/Private/DataInterface/NiagaraDataInterfaceStaticMesh.cpp
--- a/Engine/Plugins/FX/Niagara/Source/Niagara/Private/DataInterface/NiagaraDataInterfaceStaticMesh.cpp
+++ b/Engine/Plugins/FX/Niagara/Source/Niagara/Private/DataInterface/NiagaraDataInterfaceStaticMesh.cpp
@@ -1453,7 +1453,7 @@ namespace NDIStaticMeshLocal
 		FORCEINLINE int32 RandomUniformTriangle(const TRandomHelper& RandHelper, int32 InstanceIndex, const TConstArrayView<FNDISectionInfo>& SectionInfos, TConstArrayView<int32>& SectionRemap) const
 		{
 			int32 Triangle = 0;
-			int32 SectionIndex = RandHelper.Rand(InstanceIndex) * SectionInfos.Num();
+			int32 SectionIndex = int32(RandHelper.Rand(InstanceIndex) * float(SectionInfos.Num()));
 			SectionIndex = RandHelper.Rand(InstanceIndex) < SectionInfos[SectionIndex].Prob ? SectionIndex : SectionInfos[SectionIndex].Alias;
 			const int32 Section = SectionRemap[SectionIndex];
 			if ( LODResource->AreaWeightedSectionSamplers.IsValidIndex(Section) && LODResource->AreaWeightedSectionSamplers[Section].GetNumEntries() )
@@ -1467,8 +1467,8 @@ namespace NDIStaticMeshLocal
 		template<typename TRandomHelper>
 		FORCEINLINE int32 RandomTriangle(const TRandomHelper& RandHelper, int32 InstanceIndex, const TConstArrayView<FNDISectionInfo>& SectionInfos) const
 		{
-			const int32 Section = RandHelper.Rand(InstanceIndex) * SectionInfos.Num();
-			const int32 Triangle = SectionInfos[Section].FirstTriangle + (RandHelper.Rand(InstanceIndex) * SectionInfos[Section].NumTriangles);
+			const int32 Section = int32(RandHelper.Rand(InstanceIndex) * float(SectionInfos.Num()));
+			const int32 Triangle = SectionInfos[Section].FirstTriangle + int32(RandHelper.Rand(InstanceIndex) * float(SectionInfos[Section].NumTriangles));
 			return Triangle;
 		}
 
@@ -1501,7 +1501,7 @@ namespace NDIStaticMeshLocal
 		template<typename TRandomHelper>
 		FORCEINLINE int32 RandomSectionTriangle(const TRandomHelper& RandHelper, int32 InstanceIndex, int32 Section) const
 		{
-			const int32 Triangle = (LODResource->Sections[Section].FirstIndex / 3) + (RandHelper.Rand(InstanceIndex) * LODResource->Sections[Section].NumTriangles);
+			const int32 Triangle = (LODResource->Sections[Section].FirstIndex / 3) + int32(RandHelper.Rand(InstanceIndex) * float(LODResource->Sections[Section].NumTriangles));
 			return Triangle;
 		}
 
@@ -1709,7 +1709,7 @@ bool UNiagaraDataInterfaceStaticMesh::InitPerInstanceData(void* PerInstanceData,
 				}
 				for (int Section : InstanceData->GetFilteredSections())
 				{
-					GpuInitializeData->FilteredAndUnfilteredSections.Add(Section);
+					GpuInitializeData->FilteredAndUnfilteredSections.Add(uint16(Section));
 				}
 
 				GpuInitializeData->NumUnfilteredSections = InstanceData->NumUnfilteredSections;
@@ -1719,7 +1719,7 @@ bool UNiagaraDataInterfaceStaticMesh::InitPerInstanceData(void* PerInstanceData,
 				}
 				for (int Section : InstanceData->GetUnfilteredSections())
 				{
-					GpuInitializeData->FilteredAndUnfilteredSections.Add(Section);
+					GpuInitializeData->FilteredAndUnfilteredSections.Add(uint16(Section));
 				}
 			}
 		}
diff --git a/Engine/Plugins/FX/Niagara/Source/Niagara/Private/DataInterface/NiagaraDataInterfaceVirtualTexture.cpp b/Engine/Plugins/FX/Niagara/Source/Niagara/Private/DataInterface/NiagaraDataInterfaceVirtualTexture.cpp
--- a/Engine/Plugins/FX/Niagara/Source/Niagara/Private/DataInterface/NiagaraDataInterfaceVirtualTexture.cpp
+++ b/Engine/Plugins/FX/Niagara/Source/Niagara/Private/DataInterface/NiagaraDataInterfaceVirtualTexture.cpp
@@ -266,7 +266,7 @@ bool UNiagaraDataInterfaceVirtualTexture::PerInstanceTick(void* PerInstanceData,
 					InstanceData.UVUniforms[2] = RT_Texture->GetUniformParameter(ERuntimeVirtualTextureShaderUniform_WorldToUVTransform2);
 
 					const FVector4 WorldHeightUnpack = RT_Texture->GetUniformParameter(ERuntimeVirtualTextureShaderUniform_WorldHeightUnpack);
-					InstanceData.WorldHeightUnpack = FVector2f(WorldHeightUnpack.X, WorldHeightUnpack.Y);
+					InstanceData.WorldHeightUnpack = FVector2f(float(WorldHeightUnpack.X), float(WorldHeightUnpack.Y));		//LWC Precision Loss
 
 					for ( int32 iLayer=0; iLayer < MaxRVTLayers; ++iLayer)
 					{
diff --git a/Engine/Plugins/FX/Niagara/Source/Niagara/Private/MovieScene/MovieSceneNiagaraSystemTrackTemplate.cpp b/Engine/Plugins/FX/Niagara/Source/Niagara/Private/MovieScene/MovieSceneNiagaraSystemTrackTemplate.cpp
--- a/Engine/Plugins/FX/Niagara/Source/Niagara/Private/MovieScene/MovieSceneNiagaraSystemTrackTemplate.cpp
+++ b/Engine/Plugins/FX/Niagara/Source/Niagara/Private/MovieScene/MovieSceneNiagaraSystemTrackTemplate.cpp
@@ -230,7 +230,7 @@ struct FNiagaraSystemUpdateDesiredAgeExecutionToken : IMovieSceneExecutionToken
 
 			if (SystemInstanceController.IsValid() && SystemInstanceController->IsValid() && SystemInstanceController->IsComplete() == false)
 			{
-				float DesiredAge = Context.GetFrameRate().AsSeconds(Context.GetTime() - SpawnSectionStartFrame);
+				const float DesiredAge = float(Context.GetFrameRate().AsSeconds(Context.GetTime() - SpawnSectionStartFrame));
 				if (DesiredAge >= 0)
 				{
 					// Add a quarter of a frame offset here to push the desired age into the middle of the frame since it will be automatically rounded
diff --git a/Engine/Plugins/FX/Niagara/Source/Niagara/Private/NiagaraAsyncCompile.cpp b/Engine/Plugins/FX/Niagara/Source/Niagara/Private/NiagaraAsyncCompile.cpp
--- a/Engine/Plugins/FX/Niagara/Source/Niagara/Private/NiagaraAsyncCompile.cpp
+++ b/Engine/Plugins/FX/Niagara/Source/Niagara/Private/NiagaraAsyncCompile.cpp
@@ -372,7 +372,7 @@ void FNiagaraAsyncCompileTask::ProcessResult()
 	// convert results to be saved in the ddc
 	ScriptPair.CompileResults = ExeData;
 	ScriptPair.bResultsReady = true;
-	ScriptPair.CompileTime = FPlatformTime::Seconds() - StartCompileTime;
+	ScriptPair.CompileTime = float(FPlatformTime::Seconds() - StartCompileTime);
 	ScriptPair.CompiledScript->ExecToBinaryData(ScriptPair.CompiledScript, DDCOutData, *ExeData);
 	UE_LOG(LogNiagara, Verbose, TEXT("Got %i bytes in ddc data for %s"), DDCOutData.Num(), *AssetPath);
 }
diff --git a/Engine/Plugins/FX/Niagara/Source/Niagara/Private/NiagaraAsyncCompile.h b/Engine/Plugins/FX/Niagara/Source/Niagara/Private/NiagaraAsyncCompile.h
--- a/Engine/Plugins/FX/Niagara/Source/Niagara/Private/NiagaraAsyncCompile.h
+++ b/Engine/Plugins/FX/Niagara/Source/Niagara/Private/NiagaraAsyncCompile.h
@@ -22,7 +22,7 @@ struct FEmitterCompiledScriptPair
 	uint32 PendingJobID = INDEX_NONE; // this is the ID for any active shader compiler worker job
 	FNiagaraVMExecutableDataId CompileId;
 	TSharedPtr<FNiagaraVMExecutableData> CompileResults;
-	double CompileTime = 0;
+	float CompileTime = 0.0f;
 };
 
 UENUM()
diff --git a/Engine/Plugins/FX/Niagara/Source/Niagara/Private/NiagaraBoundsCalculatorHelper.h b/Engine/Plugins/FX/Niagara/Source/Niagara/Private/NiagaraBoundsCalculatorHelper.h
--- a/Engine/Plugins/FX/Niagara/Source/Niagara/Private/NiagaraBoundsCalculatorHelper.h
+++ b/Engine/Plugins/FX/Niagara/Source/Niagara/Private/NiagaraBoundsCalculatorHelper.h
@@ -59,16 +59,16 @@ public:
 		float MaxSize = KINDA_SMALL_NUMBER;
 		if (bUsedWithMeshes)
 		{
-			FVector MaxScale(1.0f, 1.0f, 1.0f);
+			FVector3f MaxScale(1.0f, 1.0f, 1.0f);
 			if (ScaleAccessor.IsValid())
 			{
-				MaxScale = (FVector)ScaleAccessor.GetReader(DataSet).GetMax();
+				MaxScale = ScaleAccessor.GetReader(DataSet).GetMax();
 			}
 
 			// NOTE: Since we're not taking particle rotation into account we have to treat the extents like a sphere,
 			// which is a little bit more conservative, but saves us having to rotate the extents per particle
-			const FVector ScaledExtents = MeshExtents * MaxScale;
-			MaxSize = FMath::Max(MaxSize, ScaledExtents.Size());
+			const FVector ScaledExtents = MeshExtents * FVector(MaxScale);
+			MaxSize = FMath::Max(MaxSize, float(ScaledExtents.Size()));
 			
 			FVector MaxTransformedOffset;
 			if (bLocalSpace)
diff --git a/Engine/Plugins/FX/Niagara/Source/Niagara/Private/NiagaraComponent.cpp b/Engine/Plugins/FX/Niagara/Source/Niagara/Private/NiagaraComponent.cpp
--- a/Engine/Plugins/FX/Niagara/Source/Niagara/Private/NiagaraComponent.cpp
+++ b/Engine/Plugins/FX/Niagara/Source/Niagara/Private/NiagaraComponent.cpp
@@ -299,7 +299,7 @@ const FVector3f& FNiagaraSceneProxy::GetLWCRenderTile() const
 TUniformBuffer<FPrimitiveUniformShaderParameters>* FNiagaraSceneProxy::GetCustomUniformBufferResource(bool bHasVelocity, const FBox& InstanceBounds) const
 {
 	// Use a hash to determine if we can re-use any uniform buffer
-	uint64 KeyHash = HashCombine(bHasVelocity, InstanceBounds.IsValid);
+	uint32 KeyHash = HashCombine(bHasVelocity, InstanceBounds.IsValid);
 
 	bool bHasPrecomputedVolumetricLightmap;
 	FMatrix PreviousLocalToWorld;
@@ -384,7 +384,7 @@ uint32 FNiagaraSceneProxy::GetMemoryFootprint() const
 
 uint32 FNiagaraSceneProxy::GetAllocatedSize() const
 { 
-	uint32 Size = FPrimitiveSceneProxy::GetAllocatedSize();
+	uint32 Size = uint32(FPrimitiveSceneProxy::GetAllocatedSize());
 	if (RenderData)
 	{
 		Size += RenderData->GetDynamicDataSize();
@@ -899,7 +899,7 @@ void UNiagaraComponent::AdvanceSimulationByTime(float SimulateTime, float TickDe
 {
 	if (SystemInstanceController.IsValid() && TickDeltaSeconds > SMALL_NUMBER)
 	{
-		int32 TickCount = SimulateTime / TickDeltaSeconds;
+		const int32 TickCount = FMath::FloorToInt(SimulateTime / TickDeltaSeconds);
 		SystemInstanceController->AdvanceSimulation(TickCount, TickDeltaSeconds);
 	}
 }
diff --git a/Engine/Plugins/FX/Niagara/Source/Niagara/Private/NiagaraDataInterface.cpp b/Engine/Plugins/FX/Niagara/Source/Niagara/Private/NiagaraDataInterface.cpp
--- a/Engine/Plugins/FX/Niagara/Source/Niagara/Private/NiagaraDataInterface.cpp
+++ b/Engine/Plugins/FX/Niagara/Source/Niagara/Private/NiagaraDataInterface.cpp
@@ -373,7 +373,8 @@ uint16 FNiagaraDataInterfaceSetShaderParametersContext::GetParameterIncludedStru
 	{
 		if (StructIncludeInfo.StructMetadata == StructMetadata)
 		{
-			return StructIncludeInfo.ParamterOffset;
+			ensure(StructIncludeInfo.ParamterOffset <= TNumericLimits<uint16>::Max());
+			return uint16(StructIncludeInfo.ParamterOffset);
 		}
 	}
 
diff --git a/Engine/Plugins/FX/Niagara/Source/Niagara/Private/NiagaraDataInterfaceRigidMeshCollisionQuery.cpp b/Engine/Plugins/FX/Niagara/Source/Niagara/Private/NiagaraDataInterfaceRigidMeshCollisionQuery.cpp
--- a/Engine/Plugins/FX/Niagara/Source/Niagara/Private/NiagaraDataInterfaceRigidMeshCollisionQuery.cpp
+++ b/Engine/Plugins/FX/Niagara/Source/Niagara/Private/NiagaraDataInterfaceRigidMeshCollisionQuery.cpp
@@ -1092,7 +1092,7 @@ void UNiagaraDataInterfaceRigidMeshCollisionQuery::DrawDebugHud(UCanvas* Canvas,
 						}
 
 						const FVector ScreenLoc = Canvas->Project(ActorOrigin);
-						Canvas->DrawText(Font, FString::Printf(TEXT("RigidMeshDI[%s Actor] - %s"), ActorSourceString, *ActorLabel), ScreenLoc.X, ScreenLoc.Y);
+						Canvas->DrawText(Font, FString::Printf(TEXT("RigidMeshDI[%s Actor] - %s"), ActorSourceString, *ActorLabel), float(ScreenLoc.X), float(ScreenLoc.Y));
 					}
 				}
 			};
diff --git a/Engine/Plugins/FX/Niagara/Source/Niagara/Private/NiagaraDataInterfaceSkeletalMesh.cpp b/Engine/Plugins/FX/Niagara/Source/Niagara/Private/NiagaraDataInterfaceSkeletalMesh.cpp
--- a/Engine/Plugins/FX/Niagara/Source/Niagara/Private/NiagaraDataInterfaceSkeletalMesh.cpp
+++ b/Engine/Plugins/FX/Niagara/Source/Niagara/Private/NiagaraDataInterfaceSkeletalMesh.cpp
@@ -1871,7 +1871,7 @@ bool FNDISkeletalMesh_InstanceData::Init(UNiagaraDataInterfaceSkeletalMesh* Inte
 				else
 				{
 					ensure(Bone <= TNumericLimits<uint16>::Max());
-					FilteredAndUnfilteredBones.Add(Bone);
+					FilteredAndUnfilteredBones.Add(uint16(Bone));
 					++NumFilteredBones;
 				}
 			}
@@ -1896,7 +1896,7 @@ bool FNDISkeletalMesh_InstanceData::Init(UNiagaraDataInterfaceSkeletalMesh* Inte
 				}
 				if (!bExists)
 				{
-					FilteredAndUnfilteredBones.Add(i);
+					FilteredAndUnfilteredBones.Add(uint16(i));
 					++NumUnfilteredBones;
 				}
 			}
diff --git a/Engine/Plugins/FX/Niagara/Source/Niagara/Private/NiagaraDataInterfaceSkeletalMeshConnectivity.cpp b/Engine/Plugins/FX/Niagara/Source/Niagara/Private/NiagaraDataInterfaceSkeletalMeshConnectivity.cpp
--- a/Engine/Plugins/FX/Niagara/Source/Niagara/Private/NiagaraDataInterfaceSkeletalMeshConnectivity.cpp
+++ b/Engine/Plugins/FX/Niagara/Source/Niagara/Private/NiagaraDataInterfaceSkeletalMeshConnectivity.cpp
@@ -272,7 +272,7 @@ static bool BuildAdjacencyBuffer(const FSkeletalMeshLODRenderData& LodRenderData
 	FMemory::Memset(Buffer.GetData(), 0xFF, Buffer.Num());
 	TriangleIndexType* AdjacencyBuffer = reinterpret_cast<TriangleIndexType*>(Buffer.GetData());
 
-	TMap<FVector, int32 /*UniqueVertexIndex*/> UniqueIndexMap;
+	TMap<FVector3f, int32 /*UniqueVertexIndex*/> UniqueIndexMap;
 
 	TArray<uint32> RedirectionArray;
 	TArray<float> TriangleSizes;
@@ -292,11 +292,11 @@ static bool BuildAdjacencyBuffer(const FSkeletalMeshLODRenderData& LodRenderData
 			IndexBuffer->Get(TriangleIt * 3 + 2)
 		};
 
-		const FVector P[3] =
+		const FVector3f P[3] =
 		{
-			(FVector)VertexBuffer.VertexPosition(V[0]),
-			(FVector)VertexBuffer.VertexPosition(V[1]),
-			(FVector)VertexBuffer.VertexPosition(V[2])
+			VertexBuffer.VertexPosition(V[0]),
+			VertexBuffer.VertexPosition(V[1]),
+			VertexBuffer.VertexPosition(V[2])
 		};
 
 		if (SortBySize)
diff --git a/Engine/Plugins/FX/Niagara/Source/Niagara/Private/NiagaraDataInterfaceTexture.cpp b/Engine/Plugins/FX/Niagara/Source/Niagara/Private/NiagaraDataInterfaceTexture.cpp
--- a/Engine/Plugins/FX/Niagara/Source/Niagara/Private/NiagaraDataInterfaceTexture.cpp
+++ b/Engine/Plugins/FX/Niagara/Source/Niagara/Private/NiagaraDataInterfaceTexture.cpp
@@ -277,12 +277,14 @@ bool UNiagaraDataInterfaceTexture::PerInstanceTick(void* PerInstanceData, FNiaga
 	int32 CurrentTextureMipLevels = 0;
 	if (UTexture2D* CurrentTexture2D = Cast<UTexture2D>(CurrentTexture) )
 	{
-		CurrentTextureSize = FIntPoint(CurrentTexture->GetSurfaceWidth(), CurrentTexture->GetSurfaceHeight());
+		CurrentTextureSize.X = int32(CurrentTexture->GetSurfaceWidth());
+		CurrentTextureSize.Y = int32(CurrentTexture->GetSurfaceHeight());
 		CurrentTextureMipLevels = CurrentTexture2D->GetNumMips();
 	}
 	else if (UTextureRenderTarget2D* CurrentTexture2DRT = Cast<UTextureRenderTarget2D>(CurrentTexture))
 	{
-		CurrentTextureSize = FIntPoint(CurrentTexture->GetSurfaceWidth(), CurrentTexture->GetSurfaceHeight());
+		CurrentTextureSize.X = int32(CurrentTexture->GetSurfaceWidth());
+		CurrentTextureSize.Y = int32(CurrentTexture->GetSurfaceHeight());
 		CurrentTextureMipLevels = CurrentTexture2DRT->GetNumMips();
 	}
 
@@ -361,8 +363,8 @@ void UNiagaraDataInterfaceTexture::VMGetTextureDimensions(FVectorVMExternalFunct
 	{
 		const int32 MipLevel = InMipLevel.GetAndAdvance();
 		const FVector2f TextureSize(
-			FMath::Max(InstData->CurrentTextureSize.X >> MipLevel, 1),
-			FMath::Max(InstData->CurrentTextureSize.Y >> MipLevel, 1)
+			float(FMath::Max(InstData->CurrentTextureSize.X >> MipLevel, 1)),
+			float(FMath::Max(InstData->CurrentTextureSize.Y >> MipLevel, 1))
 		);
 		OutSize.SetAndAdvance(TextureSize);
 	}
diff --git a/Engine/Plugins/FX/Niagara/Source/Niagara/Private/NiagaraDataInterfaceVector2DCurve.cpp b/Engine/Plugins/FX/Niagara/Source/Niagara/Private/NiagaraDataInterfaceVector2DCurve.cpp
--- a/Engine/Plugins/FX/Niagara/Source/Niagara/Private/NiagaraDataInterfaceVector2DCurve.cpp
+++ b/Engine/Plugins/FX/Niagara/Source/Niagara/Private/NiagaraDataInterfaceVector2DCurve.cpp
@@ -99,10 +99,9 @@ TArray<float> UNiagaraDataInterfaceVector2DCurve::BuildLUT(int32 NumEntries) con
 	OutputLUT.Reserve(NumEntries * 2);
 	for (int32 i = 0; i < NumEntries; i++)
 	{
-		float X = UnnormalizeTime(i * InvEntryCountFactor);
-		FVector2D C(XCurve.Eval(X), YCurve.Eval(X));
-		OutputLUT.Add(C.X);
-		OutputLUT.Add(C.Y);
+		const float X = UnnormalizeTime(i * InvEntryCountFactor);
+		OutputLUT.Add(XCurve.Eval(X));
+		OutputLUT.Add(YCurve.Eval(X));
 	}
 	return OutputLUT;
 }
@@ -186,24 +185,24 @@ void UNiagaraDataInterfaceVector2DCurve::GetVMExternalFunction(const FVMExternal
 }
 
 template<>
-FORCEINLINE_DEBUGGABLE FVector2D UNiagaraDataInterfaceVector2DCurve::SampleCurveInternal<TIntegralConstant<bool, true>>(float X)
+FORCEINLINE_DEBUGGABLE FVector2f UNiagaraDataInterfaceVector2DCurve::SampleCurveInternal<TIntegralConstant<bool, true>>(float X)
 {
-	float RemappedX = FMath::Clamp(NormalizeTime(X) * LUTNumSamplesMinusOne, 0.0f, LUTNumSamplesMinusOne);
-	float PrevEntry = FMath::TruncToFloat(RemappedX);
-	float NextEntry = PrevEntry < LUTNumSamplesMinusOne ? PrevEntry + 1.0f : PrevEntry;
-	float Interp = RemappedX - PrevEntry;
-
-	int32 AIndex = (int32)(PrevEntry * (float)CurveLUTNumElems);
-	int32 BIndex = (int32)(NextEntry * (float)CurveLUTNumElems);
-	FVector2D A = FVector2D(ShaderLUT[AIndex], ShaderLUT[AIndex + 1]);
-	FVector2D B = FVector2D(ShaderLUT[BIndex], ShaderLUT[BIndex + 1]);
+	const float RemappedX = FMath::Clamp(NormalizeTime(X) * LUTNumSamplesMinusOne, 0.0f, LUTNumSamplesMinusOne);
+	const float PrevEntry = FMath::TruncToFloat(RemappedX);
+	const float NextEntry = PrevEntry < LUTNumSamplesMinusOne ? PrevEntry + 1.0f : PrevEntry;
+	const float Interp = RemappedX - PrevEntry;
+
+	const int32 AIndex = (int32)(PrevEntry * (float)CurveLUTNumElems);
+	const int32 BIndex = (int32)(NextEntry * (float)CurveLUTNumElems);
+	const FVector2f A = FVector2f(ShaderLUT[AIndex], ShaderLUT[AIndex + 1]);
+	const FVector2f B = FVector2f(ShaderLUT[BIndex], ShaderLUT[BIndex + 1]);
 	return FMath::Lerp(A, B, Interp);
 }
 
 template<>
-FORCEINLINE_DEBUGGABLE FVector2D UNiagaraDataInterfaceVector2DCurve::SampleCurveInternal<TIntegralConstant<bool, false>>(float X)
+FORCEINLINE_DEBUGGABLE FVector2f UNiagaraDataInterfaceVector2DCurve::SampleCurveInternal<TIntegralConstant<bool, false>>(float X)
 {
-	return FVector2D(XCurve.Eval(X), YCurve.Eval(X));
+	return FVector2f(XCurve.Eval(X), YCurve.Eval(X));
 }
 
 template<typename UseLUT>
@@ -216,8 +215,8 @@ void UNiagaraDataInterfaceVector2DCurve::SampleCurve(FVectorVMExternalFunctionCo
 
 	for (int32 i = 0; i < Context.GetNumInstances(); ++i)
 	{
-		float X = XParam.GetAndAdvance();
-		FVector2D V = SampleCurveInternal<UseLUT>(X);
+		const float X = XParam.GetAndAdvance();
+		const FVector2f V = SampleCurveInternal<UseLUT>(X);
 		*OutSampleX.GetDestAndAdvance() = V.X;
 		*OutSampleY.GetDestAndAdvance() = V.Y;
 	}
diff --git a/Engine/Plugins/FX/Niagara/Source/Niagara/Private/NiagaraDataInterfaceVectorCurve.cpp b/Engine/Plugins/FX/Niagara/Source/Niagara/Private/NiagaraDataInterfaceVectorCurve.cpp
--- a/Engine/Plugins/FX/Niagara/Source/Niagara/Private/NiagaraDataInterfaceVectorCurve.cpp
+++ b/Engine/Plugins/FX/Niagara/Source/Niagara/Private/NiagaraDataInterfaceVectorCurve.cpp
@@ -105,8 +105,8 @@ TArray<float> UNiagaraDataInterfaceVectorCurve::BuildLUT(int32 NumEntries) const
 	OutputLUT.Reserve(NumEntries * 3);
 	for (int32 i=0; i < NumEntries; i++)
 	{
-		float X = UnnormalizeTime(i * InvEntryCountFactor);
-		FVector C(XCurve.Eval(X), YCurve.Eval(X), ZCurve.Eval(X));
+		const float X = UnnormalizeTime(i * InvEntryCountFactor);
+		const FVector3f C(XCurve.Eval(X), YCurve.Eval(X), ZCurve.Eval(X));
 		OutputLUT.Add(C.X);
 		OutputLUT.Add(C.Y);
 		OutputLUT.Add(C.Z);
@@ -198,24 +198,24 @@ void UNiagaraDataInterfaceVectorCurve::GetVMExternalFunction(const FVMExternalFu
 }
 
 template<>
-FORCEINLINE_DEBUGGABLE FVector UNiagaraDataInterfaceVectorCurve::SampleCurveInternal<TIntegralConstant<bool, true>>(float X)
+FORCEINLINE_DEBUGGABLE FVector3f UNiagaraDataInterfaceVectorCurve::SampleCurveInternal<TIntegralConstant<bool, true>>(float X)
 {
-	float RemappedX = FMath::Clamp(NormalizeTime(X) * LUTNumSamplesMinusOne, 0.0f, LUTNumSamplesMinusOne);
-	float PrevEntry = FMath::TruncToFloat(RemappedX);
-	float NextEntry = PrevEntry < LUTNumSamplesMinusOne ? PrevEntry + 1.0f : PrevEntry;
-	float Interp = RemappedX - PrevEntry;
-
-	int32 AIndex = (int32)(PrevEntry * (float)CurveLUTNumElems);
-	int32 BIndex = (int32)(NextEntry * (float)CurveLUTNumElems);
-	FVector A = FVector(ShaderLUT[AIndex], ShaderLUT[AIndex + 1], ShaderLUT[AIndex + 2]);
-	FVector B = FVector(ShaderLUT[BIndex], ShaderLUT[BIndex + 1], ShaderLUT[BIndex + 2]);
+	const float RemappedX = FMath::Clamp(NormalizeTime(X) * LUTNumSamplesMinusOne, 0.0f, LUTNumSamplesMinusOne);
+	const float PrevEntry = FMath::TruncToFloat(RemappedX);
+	const float NextEntry = PrevEntry < LUTNumSamplesMinusOne ? PrevEntry + 1.0f : PrevEntry;
+	const float Interp = RemappedX - PrevEntry;
+
+	const int32 AIndex = (int32)(PrevEntry * (float)CurveLUTNumElems);
+	const int32 BIndex = (int32)(NextEntry * (float)CurveLUTNumElems);
+	const FVector3f A = FVector3f(ShaderLUT[AIndex], ShaderLUT[AIndex + 1], ShaderLUT[AIndex + 2]);
+	const FVector3f B = FVector3f(ShaderLUT[BIndex], ShaderLUT[BIndex + 1], ShaderLUT[BIndex + 2]);
 	return FMath::Lerp(A, B, Interp);
 }
 
 template<>
-FORCEINLINE_DEBUGGABLE FVector UNiagaraDataInterfaceVectorCurve::SampleCurveInternal<TIntegralConstant<bool, false>>(float X)
+FORCEINLINE_DEBUGGABLE FVector3f UNiagaraDataInterfaceVectorCurve::SampleCurveInternal<TIntegralConstant<bool, false>>(float X)
 {
-	return FVector(XCurve.Eval(X), YCurve.Eval(X), ZCurve.Eval(X));
+	return FVector3f(XCurve.Eval(X), YCurve.Eval(X), ZCurve.Eval(X));
 }
 
 template<typename UseLUT>
@@ -229,8 +229,8 @@ void UNiagaraDataInterfaceVectorCurve::SampleCurve(FVectorVMExternalFunctionCont
 
 	for (int32 i = 0; i < Context.GetNumInstances(); ++i)
 	{
-		float X = XParam.GetAndAdvance();
-		FVector V = SampleCurveInternal<UseLUT>(X);
+		const float X = XParam.GetAndAdvance();
+		const FVector3f V = SampleCurveInternal<UseLUT>(X);
 		*OutSampleX.GetDestAndAdvance() = V.X;
 		*OutSampleY.GetDestAndAdvance() = V.Y;
 		*OutSampleZ.GetDestAndAdvance() = V.Z;
diff --git a/Engine/Plugins/FX/Niagara/Source/Niagara/Private/NiagaraDataInterfaceVolumeCache.cpp b/Engine/Plugins/FX/Niagara/Source/Niagara/Private/NiagaraDataInterfaceVolumeCache.cpp
--- a/Engine/Plugins/FX/Niagara/Source/Niagara/Private/NiagaraDataInterfaceVolumeCache.cpp
+++ b/Engine/Plugins/FX/Niagara/Source/Niagara/Private/NiagaraDataInterfaceVolumeCache.cpp
@@ -65,7 +65,9 @@ struct FNiagaraDataInterfaceVolumeCacheProxy : public FNiagaraDataInterfaceProxy
 		{
 			InstanceData->CurrFrame = FromGameThread->CurrFrame;
 			InstanceData->ReadFile = FromGameThread->ReadFile;
-			InstanceData->TextureSize = FVector3f(FromGameThread->NumCells.X, FromGameThread->NumCells.Y, FromGameThread->NumCells.Z);		
+			InstanceData->TextureSize.X = float(FromGameThread->NumCells.X);
+			InstanceData->TextureSize.Y = float(FromGameThread->NumCells.Y);
+			InstanceData->TextureSize.Z = float(FromGameThread->NumCells.Z);
 
 			FromGameThread->~FVolumeCacheInstanceData_GameThread();
 		}
@@ -176,21 +178,18 @@ void FNiagaraDataInterfaceVolumeCacheProxy::PreStage(const FNDIGpuComputePreStag
 	
 	if (Context.IsInputStage() && InstanceData && InstanceData->ReadFile && InstanceData->VolumeCacheData != nullptr)
 	{
-	
-		FIntVector Size = FIntVector(InstanceData->TextureSize.X, InstanceData->TextureSize.Y, InstanceData->TextureSize.Z);
-
 		if (!InstanceData->ResolvedTextureRHI.IsValid())
 		{
+			const FIntVector IntTextureSize = FIntVector(int32(InstanceData->TextureSize.X), int32(InstanceData->TextureSize.Y), int32(InstanceData->TextureSize.Z));
 
 			// #todo(dmp): this should be in instance data
 			EPixelFormat Format = PF_FloatRGBA;
 
 			const FRHITextureCreateDesc Desc =
-				FRHITextureCreateDesc::Create3D(TEXT("stuff"), Size.X, Size.Y, Size.Z, Format)
+				FRHITextureCreateDesc::Create3D(TEXT("NiagaraVolumeCache"), IntTextureSize.X, IntTextureSize.Y, IntTextureSize.Z, Format)
 				.SetFlags(ETextureCreateFlags::ShaderResource | ETextureCreateFlags::NoTiling);
 
 			InstanceData->ResolvedTextureRHI = RHICreateTexture(Desc);
-			InstanceData->TextureSize = FVector3f(Size.X, Size.Y, Size.Z);
 			InstanceData->SamplerStateRHI = TStaticSamplerState<SF_Bilinear, AM_Clamp, AM_Clamp, AM_Clamp>::GetRHI();
 		}
 		// must make sure readfile, curr frame, size are in the instance data - must sync in provide instance data for render thread
diff --git a/Engine/Plugins/FX/Niagara/Source/Niagara/Private/NiagaraDataInterfaceVolumeTexture.cpp b/Engine/Plugins/FX/Niagara/Source/Niagara/Private/NiagaraDataInterfaceVolumeTexture.cpp
--- a/Engine/Plugins/FX/Niagara/Source/Niagara/Private/NiagaraDataInterfaceVolumeTexture.cpp
+++ b/Engine/Plugins/FX/Niagara/Source/Niagara/Private/NiagaraDataInterfaceVolumeTexture.cpp
@@ -213,7 +213,9 @@ bool UNiagaraDataInterfaceVolumeTexture::PerInstanceTick(void* PerInstanceData,
 					{
 						InstanceData.TextureReferenceRHI = RT_Texture->TextureReference.TextureReferenceRHI;
 						InstanceData.SamplerStateRHI = RT_Texture->GetResource() ? RT_Texture->GetResource()->SamplerStateRHI.GetReference() : TStaticSamplerState<SF_Point>::GetRHI();
-						InstanceData.TextureSize = FVector3f(RT_TextureSize.X, RT_TextureSize.Y, RT_TextureSize.Z);
+						InstanceData.TextureSize.X = float(RT_TextureSize.X);
+						InstanceData.TextureSize.Y = float(RT_TextureSize.Y);
+						InstanceData.TextureSize.Z = float(RT_TextureSize.Z);
 					}
 					else
 					{
@@ -235,7 +237,7 @@ void UNiagaraDataInterfaceVolumeTexture::GetTextureDimensions(FVectorVMExternalF
 	FNDIOutputParam<float> OutHeight(Context);
 	FNDIOutputParam<float> OutDepth(Context);
 
-	FVector3f FloatTextureSize(InstData->CurrentTextureSize.X, InstData->CurrentTextureSize.Y, InstData->CurrentTextureSize.Z);
+	const FVector3f FloatTextureSize(float(InstData->CurrentTextureSize.X), float(InstData->CurrentTextureSize.Y), float(InstData->CurrentTextureSize.Z));
 	for (int32 i = 0; i < Context.GetNumInstances(); ++i)
 	{
 		OutWidth.SetAndAdvance(FloatTextureSize.X);
diff --git a/Engine/Plugins/FX/Niagara/Source/Niagara/Private/NiagaraDebugHud.cpp b/Engine/Plugins/FX/Niagara/Source/Niagara/Private/NiagaraDebugHud.cpp
--- a/Engine/Plugins/FX/Niagara/Source/Niagara/Private/NiagaraDebugHud.cpp
+++ b/Engine/Plugins/FX/Niagara/Source/Niagara/Private/NiagaraDebugHud.cpp
@@ -106,11 +106,11 @@ namespace NiagaraDebugLocal
 		return OutValue;
 	}
 
-	static FVector2D StringToVector2D(const FString& Arg, const FVector2D& DefaultValue)
+	static FVector2f StringToVector2f(const FString& Arg, const FVector2f& DefaultValue)
 	{
 		TArray<FString> Values;
 		Arg.ParseIntoArray(Values, TEXT(","));
-		FVector2D OutValue;
+		FVector2f OutValue;
 		OutValue.X = Values.Num() > 0 ? FCString::Atof(*Values[0]) : DefaultValue.X;
 		OutValue.Y = Values.Num() > 1 ? FCString::Atof(*Values[1]) : DefaultValue.Y;
 		return OutValue;
@@ -134,7 +134,7 @@ namespace NiagaraDebugLocal
 		MakeTuple(TEXT("OverviewEnabled="), TEXT("Enable or disable the main overview display"), [](FString Arg) {Settings.bOverviewEnabled = FCString::Atoi(*Arg) != 0; }),
 		MakeTuple(TEXT("OverviewMode="), TEXT("Change the mode of the debug overivew"), [](FString Arg) {Settings.OverviewMode = (ENiagaraDebugHUDOverviewMode)FCString::Atoi(*Arg); }),
 
-		MakeTuple(TEXT("OverviewLocation="), TEXT("Set the overview location"), [](FString Arg) { Settings.OverviewLocation = StringToVector2D(Arg, Settings.OverviewLocation); }),
+		MakeTuple(TEXT("OverviewLocation="), TEXT("Set the overview location"), [](FString Arg) { Settings.OverviewLocation = FVector2D(StringToVector2f(Arg, FVector2f(Settings.OverviewLocation))); }),
 		MakeTuple(TEXT("SystemFilter="), TEXT("Set the system filter"), [](FString Arg) {Settings.SystemFilter = Arg; Settings.bSystemFilterEnabled = !Arg.IsEmpty(); }),
 		MakeTuple(TEXT("EmitterFilter="), TEXT("Set the emitter filter"), [](FString Arg) {Settings.EmitterFilter = Arg; Settings.bEmitterFilterEnabled = !Arg.IsEmpty(); GCachedSystemVariables.Empty(); }),
 		MakeTuple(TEXT("ActorFilter="), TEXT("Set the actor filter"), [](FString Arg) {Settings.ActorFilter = Arg; Settings.bActorFilterEnabled = !Arg.IsEmpty(); }),
@@ -516,10 +516,10 @@ namespace NiagaraDebugLocal
 		}
 	};
 
-	FVector2D GetStringSize(UFont* Font, const TCHAR* Text)
+	FVector2f GetStringSize(UFont* Font, const TCHAR* Text)
 	{
-		FVector2D MaxSize = FVector2D::ZeroVector;
-		FVector2D CurrSize = FVector2D::ZeroVector;
+		FVector2f MaxSize = FVector2f::ZeroVector;
+		FVector2f CurrSize = FVector2f::ZeroVector;
 
 		const float fAdvanceHeight = Font->GetMaxCharHeight();
 		const TCHAR* PrevChar = nullptr;
@@ -553,10 +553,10 @@ namespace NiagaraDebugLocal
 		return MaxSize;
 	}
 
-	TPair<FVector2D, FVector2D> GetTextLocation(UFont* Font, const TCHAR* Text, const FNiagaraDebugHudTextOptions& TextOptions, const FVector2D ScreenLocation)
+	TPair<FVector2f, FVector2f> GetTextLocation(UFont* Font, const TCHAR* Text, const FNiagaraDebugHudTextOptions& TextOptions, const FVector2f ScreenLocation)
 	{
-		FVector2D StringSize = GetStringSize(Font, Text);
-		FVector2D OutLocation = ScreenLocation + TextOptions.ScreenOffset;
+		FVector2f StringSize = GetStringSize(Font, Text);
+		FVector2f OutLocation = ScreenLocation + FVector2f(TextOptions.ScreenOffset);
 		if (TextOptions.HorizontalAlignment == ENiagaraDebugHudHAlign::Center )
 		{
 			OutLocation.X -= StringSize.X * 0.5f;
@@ -573,7 +573,7 @@ namespace NiagaraDebugLocal
 		{
 			OutLocation.Y -= StringSize.Y;
 		}
-		return TPair<FVector2D, FVector2D>(StringSize, OutLocation);
+		return TPair<FVector2f, FVector2f>(StringSize, OutLocation);
 	}
 
 	void DrawBox(UWorld* World, const FVector& Location, const FVector& Extents, const FLinearColor& Color, float SolidAlpha = 0.0f, float Thickness = 3.0f)
@@ -584,7 +584,7 @@ namespace NiagaraDebugLocal
 			{
 				const FBox BoundsBox(-Extents, Extents);
 				FColor BoxColor = Color.ToFColor(false);
-				BoxColor.A = FMath::Clamp(int(SolidAlpha * 255.0f), 0, 255);
+				BoxColor.A = uint8(FMath::Clamp(int32(SolidAlpha * 255.0f), 0, 255));
 				LineBatcher->DrawSolidBox(BoundsBox, FTransform(FQuat::Identity, Location), BoxColor, 0, 0.0f);
 			}
 			else
@@ -1154,9 +1154,9 @@ void FNiagaraDebugHud::GatherSystemInfo()
 
 		//Generate a unique-ish random color for use in graphs and the world to help visually ID this system.
 		FRandomStream Rands(GetTypeHash(NiagaraComponent->GetAsset()->GetName()) + Settings.SystemColorSeed);
-		uint8 RandomHue = (uint8)Rands.RandRange(Settings.SystemColorHSVMin.X, Settings.SystemColorHSVMax.X);
-		uint8 RandomSat = (uint8)Rands.RandRange(Settings.SystemColorHSVMin.Y, Settings.SystemColorHSVMax.Y);
-		uint8 RandomValue = (uint8)Rands.RandRange(Settings.SystemColorHSVMin.Z, Settings.SystemColorHSVMax.Z);
+		uint8 RandomHue = (uint8)Rands.RandRange(int32(Settings.SystemColorHSVMin.X), int32(Settings.SystemColorHSVMax.X));
+		uint8 RandomSat = (uint8)Rands.RandRange(int32(Settings.SystemColorHSVMin.Y), int32(Settings.SystemColorHSVMax.Y));
+		uint8 RandomValue = (uint8)Rands.RandRange(int32(Settings.SystemColorHSVMin.Z), int32(Settings.SystemColorHSVMax.Z));
 		SystemDebugInfo.UniqueColor = FLinearColor::MakeFromHSV8(RandomHue, RandomSat, RandomValue);
 
 
@@ -1275,8 +1275,8 @@ void FNiagaraDebugHud::Draw(FNiagaraWorldManager* WorldManager, UCanvas* Canvas,
 {
 	using namespace NiagaraDebugLocal;
 
-	float CurrTime = WorldManager->GetWorld()->GetRealTimeSeconds();
-	DeltaSeconds = CurrTime - LastDrawTime;
+	const double CurrTime = WorldManager->GetWorld()->GetRealTimeSeconds();
+	DeltaSeconds = float(CurrTime - LastDrawTime);
 
 	// We may want to use the HUD to pause / step but not actually display anything on screen
 	if (Settings.bHudRenderingEnabled)
@@ -1317,7 +1317,7 @@ void FNiagaraDebugHud::Draw(FNiagaraWorldManager* WorldManager, UCanvas* Canvas,
 	LastDrawTime = CurrTime;
 }
 
-void FNiagaraDebugHud::AddLine2D(FVector2D Start, FVector2D End, FLinearColor Color, float Thickness, float Lifetime)
+void FNiagaraDebugHud::AddLine2D(FVector2f Start, FVector2f End, FLinearColor Color, float Thickness, float Lifetime)
 {
 	if (NiagaraDebugLocal::Settings.IsEnabled())
 	{
@@ -1331,7 +1331,7 @@ void FNiagaraDebugHud::AddLine2D(FVector2D Start, FVector2D End, FLinearColor Co
 	}
 }
 
-void FNiagaraDebugHud::AddCircle2D(FVector2D Pos, float Rad, float Segments, FLinearColor Color, float Thickness, float Lifetime)
+void FNiagaraDebugHud::AddCircle2D(FVector2f Pos, float Rad, int32 Segments, FLinearColor Color, float Thickness, float Lifetime)
 {
 	if (NiagaraDebugLocal::Settings.IsEnabled())
 	{
@@ -1346,7 +1346,7 @@ void FNiagaraDebugHud::AddCircle2D(FVector2D Pos, float Rad, float Segments, FLi
 	}
 }
 
-void FNiagaraDebugHud::AddBox2D(FVector2D Pos, FVector2D Extents, FLinearColor Color, float Thickness, float Lifetime)
+void FNiagaraDebugHud::AddBox2D(FVector2f Pos, FVector2f Extents, FLinearColor Color, float Thickness, float Lifetime)
 {
 	if (NiagaraDebugLocal::Settings.IsEnabled())
 	{
@@ -1371,12 +1371,12 @@ void FNiagaraDebugHud::DrawDebugGeomerty(class FNiagaraWorldManager* WorldManage
 
 	FCanvas* DrawCanvas = Canvas->Canvas;
 
-	FVector2D Size(Canvas->ClipX, Canvas->ClipY);
+	FVector2f Size(Canvas->ClipX, Canvas->ClipY);
 	
 	for (auto it = Lines2D.CreateIterator(); it; ++it)
 	{
 		FDebugLine2D& Line = *it;
-		DrawDebugCanvas2DLine(Canvas, Line.Start * Size, Line.End* Size, Line.Color, Line.Thickness);
+		DrawDebugCanvas2DLine(Canvas, FVector2D(Line.Start * Size), FVector2D(Line.End * Size), Line.Color, Line.Thickness);
 
 		Line.Lifetime -= DeltaSeconds;
 		if (Line.Lifetime <= 0.0f)
@@ -1387,7 +1387,7 @@ void FNiagaraDebugHud::DrawDebugGeomerty(class FNiagaraWorldManager* WorldManage
 	for (auto it = Circles2D.CreateIterator(); it; ++it)
 	{
 		FDebugCircle2D& Circle = *it;
-		DrawDebugCanvas2DCircle(Canvas, Circle.Pos * Size, Circle.Rad * Size.X, Circle.Segments, Circle.Color, Circle.Thickness);
+		DrawDebugCanvas2DCircle(Canvas, FVector2D(Circle.Pos * Size), Circle.Rad * Size.X, Circle.Segments, Circle.Color, Circle.Thickness);
 
 		Circle.Lifetime -= DeltaSeconds;
 		if (Circle.Lifetime <= 0.0f)
@@ -1400,8 +1400,8 @@ void FNiagaraDebugHud::DrawDebugGeomerty(class FNiagaraWorldManager* WorldManage
 		FDebugBox2D& Box = *it;
 
 		FBox2D Box2D;
-		Box2D.Min = (Box.Pos - Box.Extents) * Size;
-		Box2D.Max = (Box.Pos + Box.Extents) * Size;
+		Box2D.Min = FVector2D((Box.Pos - Box.Extents) * Size);
+		Box2D.Max = FVector2D((Box.Pos + Box.Extents) * Size);
 		DrawDebugCanvas2DBox(Canvas, Box2D, Box.Color, Box.Thickness);
 
 		Box.Lifetime -= DeltaSeconds;
@@ -1418,17 +1418,17 @@ struct FGraph
 	FCanvas* Canvas;
 	UFont* Font;
 
-	FVector2D Location;
-	FVector2D Size;
+	FVector2f Location;
+	FVector2f Size;
 
-	FVector2D MaxValues;
+	FVector2f MaxValues;
 
 	FString XLabel;
 	FString YLabel;
 
 	float EstYAxisLabelWidth = 35.0f;
 
-	FGraph(FCanvas* InCanvas, UFont* InFont, FVector2D InLocation, FVector2D InSize, FVector2D InMaxValues, FString InXLabel, FString InYLabel)
+	FGraph(FCanvas* InCanvas, UFont* InFont, FVector2f InLocation, FVector2f InSize, FVector2f InMaxValues, FString InXLabel, FString InYLabel)
 		: Canvas(InCanvas)
 		, Font(InFont)
 		, Location(InLocation)
@@ -1497,7 +1497,7 @@ struct FGraph
 		{
 			float LineTime = (MaxValues.Y / NumGridLines) * i;
 			FString LineTimeString = LexToSanitizedString(LineTime);
-			FVector2D LabelSize = FVector2D(Font->GetStringSize(*LineTimeString), Font->GetMaxCharHeight());
+			FVector2f LabelSize = FVector2f(float(Font->GetStringSize(*LineTimeString)), Font->GetMaxCharHeight());
 			float LineX = Left;
 			float LineY = FMath::Lerp(Bottom, Top, (float)i / NumGridLines);
 			BatchedElements->AddLine(FVector(Left, LineY, 0.0f), FVector(LineX + Size.X, LineY, 0.0f), GraphAxesColor * 0.5f, HitProxyId, 1.0f);
@@ -1521,7 +1521,7 @@ void FNiagaraDebugHud::DrawOverview(class FNiagaraWorldManager* WorldManager, FC
 	const FLinearColor DetailHighlightColor = Settings.OverviewDetailHighlightColor;
 	const FLinearColor BackgroundColor = Settings.DefaultBackgroundColor;
 
-	FVector2D TextLocation = Settings.OverviewLocation;
+	FVector2f TextLocation = FVector2f(Settings.OverviewLocation);
 
 	struct FOverviewColumn
 	{
@@ -1540,16 +1540,16 @@ void FNiagaraDebugHud::DrawOverview(class FNiagaraWorldManager* WorldManager, FC
 		typedef TFunction<void(FCanvas*, UFont*, float, float, FOverviewColumn&, const FSystemDebugInfo&)>  FSystemDrawFunc;
 		FSystemDrawFunc SystemDrawFunc;
 
-		FOverviewColumn(FString InGlobalHeader, FString InGlobalData, FString InSystemHeader, int32& InOffset, UFont* Font, FString ExampleSystemString, FSystemDrawFunc InSystemDrawFunc)
+		FOverviewColumn(FString InGlobalHeader, FString InGlobalData, FString InSystemHeader, float& InOffset, UFont* Font, FString ExampleSystemString, FSystemDrawFunc InSystemDrawFunc)
 			: GlobalHeader(InGlobalHeader)
 			, GlobalData(InGlobalData)
 			, SystemHeader(InSystemHeader)
 			, Offset(InOffset)
 			, SystemDrawFunc(InSystemDrawFunc)
 		{
-			HeaderWidth = GetStringSize(Font, *GlobalHeader).X;
-			SystemWidth = FMath::Max(GetStringSize(Font, *SystemHeader).X, GetStringSize(Font, *ExampleSystemString).X);
-			MaxWidth = FMath::Max(HeaderWidth + GlobalDataSeparator + GetStringSize(Font, *GlobalData).X, SystemWidth) + ColumnSeparator;
+			HeaderWidth = float(GetStringSize(Font, *GlobalHeader).X);
+			SystemWidth = float(FMath::Max(GetStringSize(Font, *SystemHeader).X, GetStringSize(Font, *ExampleSystemString).X));
+			MaxWidth = float(FMath::Max(HeaderWidth + GlobalDataSeparator + GetStringSize(Font, *GlobalData).X, SystemWidth)) + ColumnSeparator;
 
 			InOffset += MaxWidth;
 		}
@@ -1572,7 +1572,7 @@ void FNiagaraDebugHud::DrawOverview(class FNiagaraWorldManager* WorldManager, FC
 	};
 
 	float TotalOverviewWidth = 750.0f;
-	int32 ColumnOffset = 0;
+	float ColumnOffset = 0.0f;
 	TArray<FOverviewColumn, TInlineAllocator<16>> OverviewColumns;
 
 	if (Settings.bOverviewEnabled)
@@ -1808,7 +1808,7 @@ void FNiagaraDebugHud::DrawOverview(class FNiagaraWorldManager* WorldManager, FC
 					{
 						Canvas->DrawTile(X, Y, Col.MaxWidth - 6.0f, fAdvanceHeight - 2.0f, 0.0f, 0.0f, 0.0f, 0.0f, SystemInfo.UniqueColor);
 						FString SysCountString = LexToSanitizedString(SystemInfo.TotalActive);
-						float StringWidth = Font->GetStringSize(*SysCountString);
+						float StringWidth = float(Font->GetStringSize(*SysCountString));
 						Canvas->DrawShadowedString((X + (Col.MaxWidth * 0.5f)) - StringWidth * 0.5f, Y, *SysCountString, Font, FLinearColor::Black);
 					});
 			}
@@ -1890,8 +1890,8 @@ void FNiagaraDebugHud::DrawOverview(class FNiagaraWorldManager* WorldManager, FC
 		if (Settings.bOverviewEnabled || OverviewString.Len() > 0)
 		{
 			const int32 NumLines = 2 + (Settings.bOverviewEnabled ? 1 : 0);
-			const FVector2D OverviewStringSize = GetStringSize(Font, OverviewString.ToString());
-			const FVector2D ActualSize(FMath::Max(OverviewStringSize.X, TotalOverviewWidth), (NumLines*fAdvanceHeight) + OverviewStringSize.Y);
+			const FVector2f OverviewStringSize = GetStringSize(Font, OverviewString.ToString());
+			const FVector2f ActualSize(FMath::Max(OverviewStringSize.X, TotalOverviewWidth), (NumLines*fAdvanceHeight) + OverviewStringSize.Y);
 
 			// Draw background
 			DrawCanvas->DrawTile(TextLocation.X - 1.0f, TextLocation.Y - 1.0f, ActualSize.X + 2.0f, ActualSize.Y + 2.0f, 0.0f, 0.0f, 0.0f, 0.0f, BackgroundColor);
@@ -1985,14 +1985,14 @@ void FNiagaraDebugHud::DrawOverview(class FNiagaraWorldManager* WorldManager, FC
 
 				TextLocation.Y += fAdvanceHeight;
 
-				FVector2D GraphLocation(TextLocation.X, SystemDataY + fAdvanceHeight);
+				FVector2f GraphLocation(TextLocation.X, SystemDataY + fAdvanceHeight);
 
 				if (OverviewColumns.Num())
 				{
 					GraphLocation.X += OverviewColumns.Last().Offset + OverviewColumns.Last().MaxWidth + 50.0f;
 				}
 
-				FGraph<double> Graph(DrawCanvas, Font, GraphLocation, Settings.PerfGraphSize, FVector2D(Settings.PerfHistoryFrames, Settings.PerfGraphTimeRange), TEXT("Frame"), TEXT("Time(us)"));
+				FGraph<double> Graph(DrawCanvas, Font, GraphLocation, FVector2f(Settings.PerfGraphSize), FVector2f(float(Settings.PerfHistoryFrames), Settings.PerfGraphTimeRange), TEXT("Frame"), TEXT("Time(us)"));
 
 				Graph.Draw(Settings.PerfGraphAxisColor, BackgroundColor);
 
@@ -2054,7 +2054,7 @@ void FNiagaraDebugHud::DrawOverview(class FNiagaraWorldManager* WorldManager, FC
 	DrawMessages(WorldManager, DrawCanvas, TextLocation);
 }
 
-void FNiagaraDebugHud::DrawGpuComputeOverriew(class FNiagaraWorldManager* WorldManager, class FCanvas* DrawCanvas, FVector2D& TextLocation)
+void FNiagaraDebugHud::DrawGpuComputeOverriew(class FNiagaraWorldManager* WorldManager, class FCanvas* DrawCanvas, FVector2f& TextLocation)
 {
 #if WITH_NIAGARA_GPU_PROFILER
 	using namespace NiagaraDebugLocal;
@@ -2074,7 +2074,7 @@ void FNiagaraDebugHud::DrawGpuComputeOverriew(class FNiagaraWorldManager* WorldM
 		static const FString EnableCVarWarning(TEXT("GPU Profiling is disabled, enable 'fx.Niagara.GpuProfiling.Enabled'"));
 		static const FString NoDataWarning(TEXT("No GPU data is ready"));
 
-		const FVector2D StringSize = GetStringSize(Font, *EnableCVarWarning);
+		const FVector2f StringSize = GetStringSize(Font, *EnableCVarWarning);
 		DrawCanvas->DrawTile(TextLocation.X - 1.0f, TextLocation.Y - 1.0f, StringSize.X + 1.0f, 2.0f + StringSize.Y, 0.0f, 0.0f, 0.0f, 0.0f, BackgroundColor);
 		DrawCanvas->DrawShadowedString(TextLocation.X, TextLocation.Y, *(ProfilingEnabledCVar && ProfilingEnabledCVar->GetBool() ? NoDataWarning : EnableCVarWarning), Font, HeadingColor);
 		TextLocation.Y += StringSize.Y;
@@ -2087,7 +2087,7 @@ void FNiagaraDebugHud::DrawGpuComputeOverriew(class FNiagaraWorldManager* WorldM
 		struct FColumn
 		{
 			float				DefaultWidth = 150.0f;
-			FVector2D			MeasuredSize = FVector2D::ZeroVector;
+			FVector2f			MeasuredSize = FVector2f::ZeroVector;
 			float				DrawOffset = 0.0f;
 			TStringBuilder<128>	AllRowsText;
 		};
@@ -2114,9 +2114,9 @@ void FNiagaraDebugHud::DrawGpuComputeOverriew(class FNiagaraWorldManager* WorldM
 			}
 		}
 
-		void Draw(UFont* Font, FCanvas* DrawCanvas, FVector2D& Location, FLinearColor TextColor, FLinearColor BackgroundColor)
+		void Draw(UFont* Font, FCanvas* DrawCanvas, FVector2f& Location, FLinearColor TextColor, FLinearColor BackgroundColor)
 		{
-			FVector2D TotalSize = FVector2D::ZeroVector;
+			FVector2f TotalSize = FVector2f::ZeroVector;
 			for ( FColumn& Column : Columns)
 			{
 				Column.MeasuredSize	= GetStringSize(Font, Column.AllRowsText.ToString());
@@ -2306,7 +2306,7 @@ void FNiagaraDebugHud::DrawGpuComputeOverriew(class FNiagaraWorldManager* WorldM
 #endif //WITH_NIAGARA_GPU_PROFILER
 }
 
-void FNiagaraDebugHud::DrawGlobalBudgetInfo(class FNiagaraWorldManager* WorldManager, class FCanvas* DrawCanvas, FVector2D& TextLocation)
+void FNiagaraDebugHud::DrawGlobalBudgetInfo(class FNiagaraWorldManager* WorldManager, class FCanvas* DrawCanvas, FVector2f& TextLocation)
 {
 	using namespace NiagaraDebugLocal;
 
@@ -2400,7 +2400,7 @@ void FNiagaraDebugHud::DrawGlobalBudgetInfo(class FNiagaraWorldManager* WorldMan
 	}
 }
 
-void FNiagaraDebugHud::DrawValidation(class FNiagaraWorldManager* WorldManager, class FCanvas* DrawCanvas, FVector2D& TextLocation)
+void FNiagaraDebugHud::DrawValidation(class FNiagaraWorldManager* WorldManager, class FCanvas* DrawCanvas, FVector2f& TextLocation)
 {
 	using namespace NiagaraDebugLocal;
 
@@ -2561,7 +2561,7 @@ void FNiagaraDebugHud::DrawValidation(class FNiagaraWorldManager* WorldManager,
 			UFont* Font = GetFont(Settings.OverviewFont);
 			const float fAdvanceHeight = Font->GetMaxCharHeight() + 1.0f;
 
-			const FVector2D ErrorStringSize = GetStringSize(Font, ErrorString.ToString());
+			const FVector2f ErrorStringSize = GetStringSize(Font, ErrorString.ToString());
 
 			TextLocation.Y += fAdvanceHeight;
 			DrawCanvas->DrawTile(TextLocation.X - 1.0f, TextLocation.Y - 1.0f, ErrorStringSize.X + 2.0f, ErrorStringSize.Y + fAdvanceHeight + 2.0f, 0.0f, 0.0f, 0.0f, 0.0f, Settings.DefaultBackgroundColor);
@@ -2663,8 +2663,8 @@ void FNiagaraDebugHud::DrawComponents(FNiagaraWorldManager* WorldManager, UCanva
 				const bool bParticlesLocalSpace = EmitterInstance->GetCachedEmitterData()->bLocalSpace;
 				//const float ClipRadius = Settings.bUseParticleDisplayRadius ? 1.0f : 0.0f;
 				const float ParticleDisplayCenterRadiusSq = Settings.bUseParticleDisplayCenterRadius ? (Settings.ParticleDisplayCenterRadius * Settings.ParticleDisplayCenterRadius) : 0.0f;
-				const float ParticleDisplayClipNearPlane = Settings.bUseParticleDisplayClip ? FMath::Max(Settings.ParticleDisplayClip.X, 0.0f) : 0.0f;
-				const float ParticleDisplayClipFarPlane = Settings.bUseParticleDisplayClip ? FMath::Max(Settings.ParticleDisplayClip.Y, ParticleDisplayClipNearPlane) : FLT_MAX;
+				const float ParticleDisplayClipNearPlane = Settings.bUseParticleDisplayClip ? float(FMath::Max(Settings.ParticleDisplayClip.X, 0.0)) : 0.0f;
+				const float ParticleDisplayClipFarPlane = Settings.bUseParticleDisplayClip ? float(FMath::Max(Settings.ParticleDisplayClip.Y, ParticleDisplayClipNearPlane)) : FLT_MAX;
 				const uint32 MaxDisplayParticles = Settings.bUseMaxParticlesToDisplay ? Settings.MaxParticlesToDisplay : 0xFFFFFFFF;
 				uint32 NumDisplayedParticles = 0;
 
@@ -2702,7 +2702,7 @@ void FNiagaraDebugHud::DrawComponents(FNiagaraWorldManager* WorldManager, UCanva
 					}
 
 					const TCHAR* FinalString = StringBuilder.ToString();
-					const TPair<FVector2D, FVector2D> SizeAndLocation = GetTextLocation(ParticleFont, FinalString, Settings.ParticleTextOptions, FVector2D(ParticleScreenLocation));
+					const TPair<FVector2f, FVector2f> SizeAndLocation = GetTextLocation(ParticleFont, FinalString, Settings.ParticleTextOptions, FVector2f(float(ParticleScreenLocation.X), float(ParticleScreenLocation.Y)));
 					DrawCanvas->DrawTile(SizeAndLocation.Value.X - 1.0f, SizeAndLocation.Value.Y - 1.0f, SizeAndLocation.Key.X + 2.0f, SizeAndLocation.Key.Y + 2.0f, 0.0f, 0.0f, 0.0f, 0.0f, Settings.DefaultBackgroundColor);
 					DrawCanvas->DrawShadowedString(SizeAndLocation.Value.X, SizeAndLocation.Value.Y, FinalString, ParticleFont, TextColor);
 
@@ -3021,7 +3021,7 @@ void FNiagaraDebugHud::DrawComponents(FNiagaraWorldManager* WorldManager, UCanva
 				}
 
 				const TCHAR* FinalString = StringBuilder.ToString();
-				const TPair<FVector2D, FVector2D> SizeAndLocation = GetTextLocation(SystemFont, FinalString, Settings.SystemTextOptions, FVector2D(ScreenLocation));
+				const TPair<FVector2f, FVector2f> SizeAndLocation = GetTextLocation(SystemFont, FinalString, Settings.SystemTextOptions, FVector2f(float(ScreenLocation.X), float(ScreenLocation.Y)));
 				DrawCanvas->DrawTile(SizeAndLocation.Value.X - 1.0f, SizeAndLocation.Value.Y - 1.0f, SizeAndLocation.Key.X + 2.0f, SizeAndLocation.Key.Y + 2.0f, 0.0f, 0.0f, 0.0f, 0.0f, Settings.DefaultBackgroundColor);
 				DrawCanvas->DrawShadowedString(SizeAndLocation.Value.X, SizeAndLocation.Value.Y, FinalString, SystemFont, TextColor);
 			}
@@ -3029,7 +3029,7 @@ void FNiagaraDebugHud::DrawComponents(FNiagaraWorldManager* WorldManager, UCanva
 	}
 }
 
-void FNiagaraDebugHud::DrawMessages(class FNiagaraWorldManager* WorldManager, class FCanvas* DrawCanvas, FVector2D& TextLocation)
+void FNiagaraDebugHud::DrawMessages(class FNiagaraWorldManager* WorldManager, class FCanvas* DrawCanvas, FVector2f& TextLocation)
 {
 	using namespace NiagaraDebugLocal;
 
@@ -3038,7 +3038,7 @@ void FNiagaraDebugHud::DrawMessages(class FNiagaraWorldManager* WorldManager, cl
 	UFont* Font = GetFont(Settings.OverviewFont);
 	const float fAdvanceHeight = Font->GetMaxCharHeight() + 1.0f;
 
-	FVector2D BackgroundSize(MinWidth, 0.0f);
+	FVector2f BackgroundSize(MinWidth, 0.0f);
 	TArray<FName, TInlineAllocator<8>> ToRemove;
 	for (TPair<FName, FNiagaraDebugMessage>& Pair : Messages)
 	{
@@ -3048,7 +3048,7 @@ void FNiagaraDebugHud::DrawMessages(class FNiagaraWorldManager* WorldManager, cl
 		Message.Lifetime -= DeltaSeconds;
 		if (Message.Lifetime > 0.0f)
 		{
-			BackgroundSize = FVector2D::Max(BackgroundSize, GetStringSize(Font, *Message.Message));
+			BackgroundSize = FVector2f::Max(BackgroundSize, GetStringSize(Font, *Message.Message));
 		}
 		else
 		{
@@ -3347,23 +3347,23 @@ void FNiagaraDebugHUDStatsListener::TickRT()
 
 		if (FAccumulatedParticlePerfStats* Stats = GetStats(System))
 		{
-			double RTSysAvg;
-			double RTSysMax;
-			double GPUSysAvg;
-			double GPUSysMax;
+			float RTSysAvg;
+			float RTSysMax;
+			float GPUSysAvg;
+			float GPUSysMax;
 			if (Settings.PerfSampleMode == ENiagaraDebugHUDPerfSampleMode::FrameTotal)
 			{
 				RTSysAvg = Stats->GetRenderThreadStats().GetPerFrameAvg();
 				RTSysMax = Stats->GetRenderThreadStats().GetPerFrameAvg();
-				GPUSysAvg = Stats->GetGPUStats().GetPerFrameAvgMicroseconds();
-				GPUSysMax = Stats->GetGPUStats().GetPerFrameAvgMicroseconds();
+				GPUSysAvg = float(Stats->GetGPUStats().GetPerFrameAvgMicroseconds());
+				GPUSysMax = float(Stats->GetGPUStats().GetPerFrameAvgMicroseconds());
 			}
 			else
 			{
 				RTSysAvg = Stats->GetRenderThreadStats().GetPerInstanceAvg();
 				RTSysMax = Stats->GetRenderThreadStats().GetPerInstanceMax();
-				GPUSysAvg = Stats->GetGPUStats().GetPerInstanceAvgMicroseconds();
-				GPUSysMax = Stats->GetGPUStats().GetPerInstanceMaxMicroseconds();
+				GPUSysAvg = float(Stats->GetGPUStats().GetPerInstanceAvgMicroseconds());
+				GPUSysMax = float(Stats->GetGPUStats().GetPerInstanceMaxMicroseconds());
 			}
 
 			HUDStats->History.AddFrame_RT(RTSysAvg);
diff --git a/Engine/Plugins/FX/Niagara/Source/Niagara/Private/NiagaraDebugHud.h b/Engine/Plugins/FX/Niagara/Source/Niagara/Private/NiagaraDebugHud.h
--- a/Engine/Plugins/FX/Niagara/Source/Niagara/Private/NiagaraDebugHud.h
+++ b/Engine/Plugins/FX/Niagara/Source/Niagara/Private/NiagaraDebugHud.h
@@ -238,11 +238,11 @@ private:
 
 	void Draw(class FNiagaraWorldManager* WorldManager, class UCanvas* Canvas, class APlayerController* PC);
 	void DrawOverview(class FNiagaraWorldManager* WorldManager, class FCanvas* DrawCanvas);
-	void DrawGpuComputeOverriew(class FNiagaraWorldManager* WorldManager, class FCanvas* DrawCanvas, FVector2D& TextLocation);
-	void DrawGlobalBudgetInfo(class FNiagaraWorldManager* WorldManager, class FCanvas* DrawCanvas, FVector2D& TextLocation);
-	void DrawValidation(class FNiagaraWorldManager* WorldManager, class FCanvas* DrawCanvas, FVector2D& TextLocation);
+	void DrawGpuComputeOverriew(class FNiagaraWorldManager* WorldManager, class FCanvas* DrawCanvas, FVector2f& TextLocation);
+	void DrawGlobalBudgetInfo(class FNiagaraWorldManager* WorldManager, class FCanvas* DrawCanvas, FVector2f& TextLocation);
+	void DrawValidation(class FNiagaraWorldManager* WorldManager, class FCanvas* DrawCanvas, FVector2f& TextLocation);
 	void DrawComponents(class FNiagaraWorldManager* WorldManager, class UCanvas* Canvas);
-	void DrawMessages(class FNiagaraWorldManager* WorldManager, class FCanvas* DrawCanvas, FVector2D& TextLocation);
+	void DrawMessages(class FNiagaraWorldManager* WorldManager, class FCanvas* DrawCanvas, FVector2f& TextLocation);
 	void DrawDebugGeomerty(class FNiagaraWorldManager* WorldManager, class UCanvas* DrawCanvas);
 
 private:
@@ -314,7 +314,7 @@ private:
 	TMap<FName, FGpuUsagePerEvent> GpuUsagePerEvent;
 #endif
 
-	float LastDrawTime = 0.0f;
+	double LastDrawTime = 0.0f;
 	float DeltaSeconds = 0.0f;
 
 	/** Generic messages that the debugger or other parts of Niagara can post to the HUD. */
@@ -323,8 +323,8 @@ private:
 	//Additional debug geometry helpers
 	struct FDebugLine2D
 	{
-		FVector2D Start;
-		FVector2D End;
+		FVector2f Start;
+		FVector2f End;
 		FLinearColor Color;
 		float Thickness;
 		float Lifetime;
@@ -332,9 +332,9 @@ private:
 	TArray<FDebugLine2D> Lines2D;
 	struct FDebugCircle2D
 	{
-		FVector2D Pos;
+		FVector2f Pos;
 		float Rad;
-		float Segments;
+		int32 Segments;
 		FLinearColor Color;
 		float Thickness;
 		float Lifetime;
@@ -342,8 +342,8 @@ private:
 	TArray<FDebugCircle2D> Circles2D;
 	struct FDebugBox2D
 	{
-		FVector2D Pos;
-		FVector2D Extents;
+		FVector2f Pos;
+		FVector2f Extents;
 		FLinearColor Color;
 		float Thickness;
 		float Lifetime;
@@ -352,9 +352,9 @@ private:
 
 public:
 	/** Add a 2D line to the debug renering. Positions are in normalized screen space. (0,0) in top left, (1,1) bottom right.*/
-	void AddLine2D(FVector2D Start, FVector2D End, FLinearColor Color, float Thickness, float Lifetime);
-	void AddCircle2D(FVector2D Pos, float Rad, float Segments, FLinearColor Color, float Thickness, float Lifetime);
-	void AddBox2D(FVector2D Pos, FVector2D Extents, FLinearColor Color, float Thickness, float Lifetime);
+	void AddLine2D(FVector2f Start, FVector2f End, FLinearColor Color, float Thickness, float Lifetime);
+	void AddCircle2D(FVector2f Pos, float Rad, int32 Segments, FLinearColor Color, float Thickness, float Lifetime);
+	void AddBox2D(FVector2f Pos, FVector2f Extents, FLinearColor Color, float Thickness, float Lifetime);
 	//Additional debug geometry helpers END
 };
 
diff --git a/Engine/Plugins/FX/Niagara/Source/Niagara/Private/NiagaraDebuggerClient.cpp b/Engine/Plugins/FX/Niagara/Source/Niagara/Private/NiagaraDebuggerClient.cpp
--- a/Engine/Plugins/FX/Niagara/Source/Niagara/Private/NiagaraDebuggerClient.cpp
+++ b/Engine/Plugins/FX/Niagara/Source/Niagara/Private/NiagaraDebuggerClient.cpp
@@ -356,8 +356,8 @@ bool FNiagaraDebuggerClient::UpdateOutliner(float DeltaSeconds)
 				if (World)
 				{
 					WorldData.bHasBegunPlay = World->HasBegunPlay();
-					WorldData.WorldType = World->WorldType;
-					WorldData.NetMode = World->GetNetMode();
+					WorldData.WorldType = uint8(World->WorldType);
+					WorldData.NetMode = uint8(World->GetNetMode());
 
 					#if WITH_PARTICLE_PERF_STATS
 					if(StatsListener)
diff --git a/Engine/Plugins/FX/Niagara/Source/Niagara/Private/NiagaraEmitter.cpp b/Engine/Plugins/FX/Niagara/Source/Niagara/Private/NiagaraEmitter.cpp
--- a/Engine/Plugins/FX/Niagara/Source/Niagara/Private/NiagaraEmitter.cpp
+++ b/Engine/Plugins/FX/Niagara/Source/Niagara/Private/NiagaraEmitter.cpp
@@ -916,7 +916,7 @@ void UNiagaraEmitter::GetAssetRegistryTags(TArray<FAssetRegistryTag>& OutTags) c
 	const FVersionedNiagaraEmitterData& EmitterData = AssetData ? *AssetData : DefaultData; // the CDO does not have any version data, so just use the default struct values in that case 
 	OutTags.Add(FAssetRegistryTag("HasGPUEmitter", ( EmitterData.SimTarget == ENiagaraSimTarget::GPUComputeSim) ? TEXT("True") : TEXT("False"), FAssetRegistryTag::TT_Alphabetical));
 
-	const float BoundsSize =  EmitterData.FixedBounds.GetSize().GetMax();
+	const float BoundsSize =  float(EmitterData.FixedBounds.GetSize().GetMax());
 	OutTags.Add(FAssetRegistryTag("FixedBoundsSize",  EmitterData.CalculateBoundsMode == ENiagaraEmitterCalculateBoundMode::Fixed ? FString::Printf(TEXT("%.2f"), BoundsSize) : FString(TEXT("None")), FAssetRegistryTag::TT_Numerical));
 
 
diff --git a/Engine/Plugins/FX/Niagara/Source/Niagara/Private/NiagaraEmitterInstance.cpp b/Engine/Plugins/FX/Niagara/Source/Niagara/Private/NiagaraEmitterInstance.cpp
--- a/Engine/Plugins/FX/Niagara/Source/Niagara/Private/NiagaraEmitterInstance.cpp
+++ b/Engine/Plugins/FX/Niagara/Source/Niagara/Private/NiagaraEmitterInstance.cpp
@@ -920,7 +920,7 @@ float FNiagaraEmitterInstance::GetTotalCPUTimeMS()
 int64 FNiagaraEmitterInstance::GetTotalBytesUsed()
 {
 	check(ParticleDataSet);
-	int32 BytesUsed = ParticleDataSet->GetSizeBytes();
+	int64 BytesUsed = ParticleDataSet->GetSizeBytes();
 	/*
 	for (FNiagaraDataSet& Set : DataSets)
 	{
diff --git a/Engine/Plugins/FX/Niagara/Source/Niagara/Private/NiagaraFunctionLibrary.cpp b/Engine/Plugins/FX/Niagara/Source/Niagara/Private/NiagaraFunctionLibrary.cpp
--- a/Engine/Plugins/FX/Niagara/Source/Niagara/Private/NiagaraFunctionLibrary.cpp
+++ b/Engine/Plugins/FX/Niagara/Source/Niagara/Private/NiagaraFunctionLibrary.cpp
@@ -1011,14 +1011,14 @@ struct FVectorKernelFastMatrixToQuaternion
 		for ( int32 iInstance=0; iInstance < Context.GetNumInstances(); ++iInstance)
 		{
 			// Read Matrix
-			FMatrix Mat;
+			FMatrix44f Mat;
 			for ( int32 j=0; j < 16; ++j )
 			{
 				Mat.M[j >> 2][j & 3] = InMatrix[j].GetAndAdvance();
 			}
 
 			// Generate Quat (without checking for consistency with the GPU)
-			FQuat Quat;
+			FQuat4f Quat;
 			{
 				// Check diagonal (trace)
 				const float tr = Mat.M[0][0] + Mat.M[1][1] + Mat.M[2][2];
@@ -1484,12 +1484,12 @@ struct FVectorKernel_SolveVelocitiesAndForces
 		for (int i = 0; i < Context.GetNumInstances(); ++i)
 		{
 			// Gather values
-			FVector PhysicsForce(InPhysicsForceX.GetAndAdvance(), InPhysicsForceY.GetAndAdvance(), InPhysicsForceZ.GetAndAdvance());
+			FVector3f PhysicsForce(InPhysicsForceX.GetAndAdvance(), InPhysicsForceY.GetAndAdvance(), InPhysicsForceZ.GetAndAdvance());
 			float PhysicsDrag = InPhysicsDrag.GetAndAdvance();
 
 			float ParticleMass = InParticlesMass.GetAndAdvance();
-			FVector ParticlePosition(InParticlesPositionX.GetAndAdvance(), InParticlesPositionY.GetAndAdvance(), InParticlesPositionZ.GetAndAdvance());
-			FVector ParticleVelocity(InParticlesVelocityX.GetAndAdvance(), InParticlesVelocityY.GetAndAdvance(), InParticlesVelocityZ.GetAndAdvance());
+			FVector3f ParticlePosition(InParticlesPositionX.GetAndAdvance(), InParticlesPositionY.GetAndAdvance(), InParticlesPositionZ.GetAndAdvance());
+			FVector3f ParticleVelocity(InParticlesVelocityX.GetAndAdvance(), InParticlesVelocityY.GetAndAdvance(), InParticlesVelocityZ.GetAndAdvance());
 
 			*OutParticlesPreviousVelocityX.GetDestAndAdvance() = ParticleVelocity.X;
 			*OutParticlesPreviousVelocityY.GetDestAndAdvance() = ParticleVelocity.Y;
diff --git a/Engine/Plugins/FX/Niagara/Source/Niagara/Private/NiagaraPerfBaseline.cpp b/Engine/Plugins/FX/Niagara/Source/Niagara/Private/NiagaraPerfBaseline.cpp
--- a/Engine/Plugins/FX/Niagara/Source/Niagara/Private/NiagaraPerfBaseline.cpp
+++ b/Engine/Plugins/FX/Niagara/Source/Niagara/Private/NiagaraPerfBaseline.cpp
@@ -95,11 +95,11 @@ FNiagaraPerfBaselineStats::FNiagaraPerfBaselineStats(FAccumulatedParticlePerfSta
 		RTMaxIndex = Stats.GetRenderThreadStats().NumFrames / 2;
 	}
 
-	PerInstanceAvg_GT = FPlatformTime::ToMilliseconds64(Stats.GetGameThreadStats().GetPerInstanceAvgCycles()) * 1000.0;
-	PerInstanceMax_GT = FPlatformTime::ToMilliseconds64(Stats.GetGameThreadStats().GetPerInstanceMaxCycles(GTMaxIndex)) * 1000.0;
+	PerInstanceAvg_GT = float(FPlatformTime::ToMilliseconds64(Stats.GetGameThreadStats().GetPerInstanceAvgCycles()) * 1000.0);
+	PerInstanceMax_GT = float(FPlatformTime::ToMilliseconds64(Stats.GetGameThreadStats().GetPerInstanceMaxCycles(GTMaxIndex)) * 1000.0);
 
-	PerInstanceAvg_RT = FPlatformTime::ToMilliseconds64(Stats.GetRenderThreadStats().GetPerInstanceAvgCycles()) * 1000.0;
-	PerInstanceMax_RT = FPlatformTime::ToMilliseconds64(Stats.GetRenderThreadStats().GetPerInstanceMaxCycles(RTMaxIndex)) * 1000.0;
+	PerInstanceAvg_RT = float(FPlatformTime::ToMilliseconds64(Stats.GetRenderThreadStats().GetPerInstanceAvgCycles()) * 1000.0);
+	PerInstanceMax_RT = float(FPlatformTime::ToMilliseconds64(Stats.GetRenderThreadStats().GetPerInstanceMaxCycles(RTMaxIndex)) * 1000.0);
 }
 
 FNiagaraPerfBaselineStats::EComparisonResult FNiagaraPerfBaselineStats::Compare(float SystemStats, float Baseline, float& OutRatio)
@@ -477,10 +477,10 @@ void FParticlePerfStatsListener_NiagaraPerformanceReporter::ReportToLog()
 					GTMaxIndex = GT.NumFrames / 2;
 				}
 
-				uint32 TickGameThreadTime = FPlatformTime::ToMilliseconds(GT.AccumulatedStats.TickGameThreadCycles) * 1000.0;
-				uint32 TickConcurrentTime = FPlatformTime::ToMilliseconds(GT.AccumulatedStats.TickConcurrentCycles) * 1000.0;
-				uint32 FinalizeTime = FPlatformTime::ToMilliseconds(GT.AccumulatedStats.FinalizeCycles) * 1000.0;
-				uint32 EndOfFrameTime = FPlatformTime::ToMilliseconds(GT.AccumulatedStats.EndOfFrameCycles) * 1000.0;
+				const uint32 TickGameThreadTime = uint32(FPlatformTime::ToMilliseconds64(GT.AccumulatedStats.TickGameThreadCycles) * 1000.0);
+				const uint32 TickConcurrentTime = uint32(FPlatformTime::ToMilliseconds64(GT.AccumulatedStats.TickConcurrentCycles) * 1000.0);
+				const uint32 FinalizeTime = uint32(FPlatformTime::ToMilliseconds64(GT.AccumulatedStats.FinalizeCycles) * 1000.0);
+				const uint32 EndOfFrameTime = uint32(FPlatformTime::ToMilliseconds64(GT.AccumulatedStats.EndOfFrameCycles) * 1000.0);
 				
 				UE_LOG(LogNiagara, Log, TEXT("| Game Thread | Num: %4u | Avg: %6u | Max: %6u | AvgRatio: %g | Max Ratio: %g |"), GT.AccumulatedStats.NumInstances, (uint32)GT.GetPerInstanceAvg(), (uint32)GT.GetPerInstanceMax(GTMaxIndex), Ratio.PerInstanceAvg_GT, Ratio.PerInstanceMax_GT);
 				UE_LOG(LogNiagara, Log, TEXT("| | Tick            | Total: %5u | Avg: %5u |"), TickGameThreadTime, TickGameThreadTime / GT.AccumulatedStats.NumInstances);
@@ -499,10 +499,10 @@ void FParticlePerfStatsListener_NiagaraPerformanceReporter::ReportToLog()
 				{
 					RTMaxIndex = RT.NumFrames / 2;
 				}
-				uint32 RTAvg = RT.GetPerInstanceAvg();
-				uint32 RTMax = RT.GetPerInstanceMax(RTMaxIndex);
-				uint32 RenderUpdateTime = FPlatformTime::ToMilliseconds(RT.AccumulatedStats.RenderUpdateCycles) * 1000.0;
-				uint32 GDMETime = FPlatformTime::ToMilliseconds(RT.AccumulatedStats.GetDynamicMeshElementsCycles) * 1000.0;
+				const uint32 RTAvg = uint32(RT.GetPerInstanceAvg());
+				const uint32 RTMax = uint32(RT.GetPerInstanceMax(RTMaxIndex));
+				const uint32 RenderUpdateTime = uint32(FPlatformTime::ToMilliseconds64(RT.AccumulatedStats.RenderUpdateCycles) * 1000.0);
+				const uint32 GDMETime = uint32(FPlatformTime::ToMilliseconds64(RT.AccumulatedStats.GetDynamicMeshElementsCycles) * 1000.0);
 				UE_LOG(LogNiagara, Log, TEXT("| Render Thread | Num: %4u | Avg: %6u | Max: %6u | AvgRatio: %g | Max Ratio: %g |"), RT.AccumulatedStats.NumInstances, RTAvg, RTMax, Ratio.PerInstanceAvg_RT, Ratio.PerInstanceMax_RT);
 				UE_LOG(LogNiagara, Log, TEXT("| | Update | Total: %5u | Avg: %5u |"), RenderUpdateTime, RenderUpdateTime / RT.AccumulatedStats.NumInstances);
 				UE_LOG(LogNiagara, Log, TEXT("| | GDME   | Total: %5u | Avg: %5u |"), GDMETime, GDMETime / RT.AccumulatedStats.NumInstances);
@@ -637,8 +637,8 @@ void FParticlePerfStatsListener_NiagaraPerformanceReporter::HandleTestResults()
 	};
 	TArray<FTestReport> TestReports;
 
-	float TestStartTime = CurrentWorldTime;
-	int32 TestStartFrame = CurrentFrameNumber;
+	const double TestStartTime = CurrentWorldTime;
+	const int32 TestStartFrame = CurrentFrameNumber;
 	CurrentWorldTime = World->GetTimeSeconds();
 	CurrentFrameNumber = GFrameNumber;
 
@@ -732,10 +732,10 @@ void FParticlePerfStatsListener_NiagaraPerformanceReporter::HandleTestResults()
 						GTMaxIndex = GT.NumFrames / 2;
 					}
 
-					uint32 TickGameThreadTime = FPlatformTime::ToMilliseconds(GT.AccumulatedStats.TickGameThreadCycles) * 1000.0;
-					uint32 TickConcurrentTime = FPlatformTime::ToMilliseconds(GT.AccumulatedStats.TickConcurrentCycles) * 1000.0;
-					uint32 FinalizeTime = FPlatformTime::ToMilliseconds(GT.AccumulatedStats.FinalizeCycles) * 1000.0;
-					uint32 EndOfFrameTime = FPlatformTime::ToMilliseconds(GT.AccumulatedStats.EndOfFrameCycles) * 1000.0;
+					const uint32 TickGameThreadTime = uint32(FPlatformTime::ToMilliseconds64(GT.AccumulatedStats.TickGameThreadCycles) * 1000.0);
+					const uint32 TickConcurrentTime = uint32(FPlatformTime::ToMilliseconds64(GT.AccumulatedStats.TickConcurrentCycles) * 1000.0);
+					const uint32 FinalizeTime = uint32(FPlatformTime::ToMilliseconds64(GT.AccumulatedStats.FinalizeCycles) * 1000.0);
+					const uint32 EndOfFrameTime = uint32(FPlatformTime::ToMilliseconds64(GT.AccumulatedStats.EndOfFrameCycles) * 1000.0);
 
 					Ar.Logf(TEXT("| Game Thread | Num: %4u | Avg: %6u | Max: %6u | AvgRatio: %g | Max Ratio: %g |"), GT.AccumulatedStats.NumInstances, (uint32)GT.GetPerInstanceAvg(), (uint32)GT.GetPerInstanceMax(GTMaxIndex), Report.Ratio.PerInstanceAvg_GT, Report.Ratio.PerInstanceMax_GT);
 					Ar.Logf(TEXT("| | Tick            | Total: %5u | Avg: %5u |"), TickGameThreadTime, TickGameThreadTime / GT.AccumulatedStats.NumInstances);
@@ -754,10 +754,10 @@ void FParticlePerfStatsListener_NiagaraPerformanceReporter::HandleTestResults()
 					{
 						RTMaxIndex = RT.NumFrames / 2;
 					}
-					uint32 RTAvg = RT.GetPerInstanceAvg();
-					uint32 RTMax = RT.GetPerInstanceMax(RTMaxIndex);
-					uint32 RenderUpdateTime = FPlatformTime::ToMilliseconds(RT.AccumulatedStats.RenderUpdateCycles) * 1000.0;
-					uint32 GDMETime = FPlatformTime::ToMilliseconds(RT.AccumulatedStats.GetDynamicMeshElementsCycles) * 1000.0;
+					const uint32 RTAvg = uint32(RT.GetPerInstanceAvg());
+					const uint32 RTMax = uint32(RT.GetPerInstanceMax(RTMaxIndex));
+					const uint32 RenderUpdateTime = uint32(FPlatformTime::ToMilliseconds64(RT.AccumulatedStats.RenderUpdateCycles) * 1000.0);
+					const uint32 GDMETime = uint32(FPlatformTime::ToMilliseconds64(RT.AccumulatedStats.GetDynamicMeshElementsCycles) * 1000.0);
 					Ar.Logf(TEXT("| Render Thread | Num: %4u | Avg: %6u | Max: %6u | AvgRatio: %g | Max Ratio: %g |"), RT.AccumulatedStats.NumInstances, RTAvg, RTMax, Report.Ratio.PerInstanceAvg_RT, Report.Ratio.PerInstanceMax_RT);
 					Ar.Logf(TEXT("| | Update | Total: %5u | Avg: %5u |"), RenderUpdateTime, RenderUpdateTime / RT.AccumulatedStats.NumInstances);
 					Ar.Logf(TEXT("| | GDME   | Total: %5u | Avg: %5u |"), GDMETime, GDMETime / RT.AccumulatedStats.NumInstances);
@@ -1008,11 +1008,11 @@ int32 FParticlePerfStatsListener_NiagaraBaselineComparisonRender::RenderStats(cl
 	float CharWidth = 0.0f;
 	float CharHeight = 0.0f;
 	Font->GetCharSize('W', CharWidth, CharHeight);
-	const float NameWidth = 32 * CharWidth;
-	const float ColumnWidth = 6 * CharWidth;
-	const int32 FontHeight = Font->GetMaxCharHeight() + 2.0f;
+	const int32 NameWidth = FMath::RoundToInt(32.0f * CharWidth);
+	const int32 ColumnWidth = FMath::RoundToInt(6.0f * CharWidth);
+	const int32 FontHeight = FMath::TruncToInt(Font->GetMaxCharHeight() + 2.0f);
 
-	const float BaseX = 50;
+	const int32 BaseX = 50;
 	int32 X = BaseX + NameWidth;
 
 	// Draw background
diff --git a/Engine/Plugins/FX/Niagara/Source/Niagara/Private/NiagaraPreviewGrid.cpp b/Engine/Plugins/FX/Niagara/Source/Niagara/Private/NiagaraPreviewGrid.cpp
--- a/Engine/Plugins/FX/Niagara/Source/Niagara/Private/NiagaraPreviewGrid.cpp
+++ b/Engine/Plugins/FX/Niagara/Source/Niagara/Private/NiagaraPreviewGrid.cpp
@@ -151,11 +151,11 @@ void ANiagaraPreviewGrid::ActivatePreviews(bool bReset)
 
 	for (int32 X = 0; X < NumX; ++X)
 	{
-		float XLocation = X / NumX - 1;
+		const float XLocation = X / float(NumX - 1);
 
 		for (int32 Y = 0; Y < NumY; ++Y)
 		{
-			float YLocation = Y / NumY - 1;
+			const float YLocation = Y / float(NumY - 1);
 
 			int32 PreviewIdx = PreviewIndex(X, Y);
 			UChildActorComponent* PreviewComp = PreviewComponents[PreviewIdx];
@@ -178,11 +178,11 @@ void ANiagaraPreviewGrid::DeactivatePreviews()
 {
 	for (int32 X = 0; X < NumX; ++X)
 	{
-		float XLocation = X / NumX - 1;
+		const float XLocation = X / float(NumX - 1);
 
 		for (int32 Y = 0; Y < NumY; ++Y)
 		{
-			float YLocation = Y / NumY - 1;
+			const float YLocation = Y / float(NumY - 1);
 
 			int32 PreviewIdx = PreviewIndex(X, Y);
 			UChildActorComponent* PreviewComp = PreviewComponents[PreviewIdx];
@@ -202,11 +202,11 @@ void ANiagaraPreviewGrid::SetPaused(bool bPaused)
 {
 	for (int32 X = 0; X < NumX; ++X)
 	{
-		float XLocation = X / NumX - 1;
+		const float XLocation = X / float(NumX - 1);
 
 		for (int32 Y = 0; Y < NumY; ++Y)
 		{
-			float YLocation = Y / NumY - 1;
+			const float YLocation = Y / float(NumY - 1);
 
 			int32 PreviewIdx = PreviewIndex(X, Y);
 			UChildActorComponent* PreviewComp = PreviewComponents[PreviewIdx];
@@ -228,11 +228,11 @@ void ANiagaraPreviewGrid::GetPreviews(TArray<UNiagaraComponent*>& OutPreviews)
 {
 	for (int32 X = 0; X < NumX; ++X)
 	{
-		float XLocation = X / NumX - 1;
+		const float XLocation = X / float(NumX - 1);
 
 		for (int32 Y = 0; Y < NumY; ++Y)
 		{
-			float YLocation = Y / NumY - 1;
+			const float YLocation = Y / float(NumY - 1);
 
 			int32 PreviewIdx = PreviewIndex(X, Y);
 			UChildActorComponent* PreviewComp = PreviewComponents[PreviewIdx];
@@ -292,7 +292,7 @@ void ANiagaraPreviewGrid::GeneratePreviews()
 		{
 			for (int32 X = 0; X < NumX; ++X)
 			{
-				float XLocation = X / NumX - 1;
+				const float XLocation = X / float(NumX - 1);
 
 				for (int32 Y = 0; Y < NumY; ++Y)
 				{
@@ -332,11 +332,11 @@ void ANiagaraPreviewGrid::TickPreviews()
 		bool bAllInactive = true;
 		for (int32 X = 0; X < NumX; ++X)
 		{
-			float XLocation = X / NumX - 1;
+			const float XLocation = X / float(NumX - 1);
 
 			for (int32 Y = 0; Y < NumY; ++Y)
 			{
-				float YLocation = Y / NumY - 1;
+				const float YLocation = Y / float(NumY - 1);
 
 				int32 PreviewIdx = PreviewIndex(X, Y);
 
@@ -393,7 +393,7 @@ void UNiagaraPreviewAxis_InterpParamInt32::ApplyToPreview_Implementation(UNiagar
 	check(PreviewComponent);
 	float Interp = Count > 1 ? (float)PreviewIndex / (Count - 1) : 1.0f;
 	int32 Val = FMath::Lerp(Min, Max, Interp);
-	PreviewComponent->SetVariableFloat(Param, Val);
+	PreviewComponent->SetVariableInt(Param, Val);
 	OutLabelText = FString::Printf(TEXT("%s = %d"), *Param.ToString(), Val);
 }
 void UNiagaraPreviewAxis_InterpParamFloat::ApplyToPreview_Implementation(UNiagaraComponent* PreviewComponent, int32 PreviewIndex, bool bIsXAxis, FString& OutLabelText)
diff --git a/Engine/Plugins/FX/Niagara/Source/Niagara/Private/NiagaraRenderer.cpp b/Engine/Plugins/FX/Niagara/Source/Niagara/Private/NiagaraRenderer.cpp
--- a/Engine/Plugins/FX/Niagara/Source/Niagara/Private/NiagaraRenderer.cpp
+++ b/Engine/Plugins/FX/Niagara/Source/Niagara/Private/NiagaraRenderer.cpp
@@ -651,14 +651,14 @@ void FNiagaraRenderer::SortIndices(const FNiagaraGPUSortInfo& SortInfo, const FN
 			{
 				for (uint32 i = 0; i < NumInstances; ++i)
 				{
-					ParticleOrder[i].SetAsUint<true, false>(i, FVector::DotProduct(GetPos(i) - SortInfo.ViewOrigin, SortInfo.ViewDirection));
+					ParticleOrder[i].SetAsUint<true, false>(i, float(FVector::DotProduct(GetPos(i) - SortInfo.ViewOrigin, SortInfo.ViewDirection)));
 				}
 			}
 			else
 			{
 				for (uint32 i = 0; i < NumInstances; ++i)
 				{
-					ParticleOrder[i].SetAsUint<true, false>(i, (GetPos(i) - SortInfo.ViewOrigin).SizeSquared());
+					ParticleOrder[i].SetAsUint<true, false>(i, float((GetPos(i) - SortInfo.ViewOrigin).SizeSquared()));
 				}
 			}
 		}
@@ -677,14 +677,14 @@ void FNiagaraRenderer::SortIndices(const FNiagaraGPUSortInfo& SortInfo, const FN
 			{
 				for (uint32 i = 0; i < NumInstances; ++i)
 				{
-					ParticleOrder[i].SetAsUint<true, false>(i, FVector::DotProduct(GetPos(i) - SortInfo.ViewOrigin, SortInfo.ViewDirection));
+					ParticleOrder[i].SetAsUint<true, false>(i, float(FVector::DotProduct(GetPos(i) - SortInfo.ViewOrigin, SortInfo.ViewDirection)));
 				}
 			}
 			else
 			{
 				for (uint32 i = 0; i < NumInstances; ++i)
 				{
-					ParticleOrder[i].SetAsUint<true, false>(i, (GetPos(i) - SortInfo.ViewOrigin).SizeSquared());
+					ParticleOrder[i].SetAsUint<true, false>(i, float((GetPos(i) - SortInfo.ViewOrigin).SizeSquared()));
 				}
 			}
 		}
@@ -792,7 +792,7 @@ private:
 		{
 			if (IsVisibile(i))
 			{
-				ParticleOrder[OutInstances++].SetAsUint<bStrictlyPositive, bAscending>(i, GetSortKey(i));
+				ParticleOrder[OutInstances++].SetAsUint<bStrictlyPositive, bAscending>(i, float(GetSortKey(i)));
 			}
 		}
 		return OutInstances;
diff --git a/Engine/Plugins/FX/Niagara/Source/Niagara/Private/NiagaraRendererComponents.cpp b/Engine/Plugins/FX/Niagara/Source/Niagara/Private/NiagaraRendererComponents.cpp
--- a/Engine/Plugins/FX/Niagara/Source/Niagara/Private/NiagaraRendererComponents.cpp
+++ b/Engine/Plugins/FX/Niagara/Source/Niagara/Private/NiagaraRendererComponents.cpp
@@ -391,7 +391,7 @@ void FNiagaraRendererComponents::PostSystemTick_GameThread(const UNiagaraRendere
 		return;
 	}
 
-	const float CurrentTime = AttachComponent->GetWorld()->GetRealTimeSeconds();
+	const double CurrentTime = AttachComponent->GetWorld()->GetRealTimeSeconds();
 	FNiagaraDataSet& Data = Emitter->GetData();
 	FNiagaraDataBuffer& ParticleData = Data.GetCurrentDataChecked();
 	FNiagaraDataSetReaderInt32<FNiagaraBool> EnabledAccessor = FNiagaraDataSetAccessor<FNiagaraBool>::CreateReader(Data, Properties->EnabledBinding.GetDataSetBindableVariable().GetName());
diff --git a/Engine/Plugins/FX/Niagara/Source/Niagara/Private/NiagaraRendererMeshes.cpp b/Engine/Plugins/FX/Niagara/Source/Niagara/Private/NiagaraRendererMeshes.cpp
--- a/Engine/Plugins/FX/Niagara/Source/Niagara/Private/NiagaraRendererMeshes.cpp
+++ b/Engine/Plugins/FX/Niagara/Source/Niagara/Private/NiagaraRendererMeshes.cpp
@@ -1473,11 +1473,11 @@ void FNiagaraRendererMeshes::GetDynamicRayTracingInstances(FRayTracingMaterialGa
 				bIsInstancedStereo,
 				false // bNeedsPrevTransform
 			);
-			MeshBatch.SegmentIndex = SectionIndex;
-			MeshBatch.LODIndex = LODModel.LODIndex;
+			MeshBatch.SegmentIndex = uint8(SectionIndex);
+			MeshBatch.LODIndex = uint8(LODModel.LODIndex);
 
 #if !(UE_BUILD_SHIPPING || UE_BUILD_TEST)
-			MeshBatch.VisualizeLODIndex = LODModel.LODIndex;
+			MeshBatch.VisualizeLODIndex = uint8(LODModel.LODIndex);
 #endif
 
 			MeshBatch.bCanApplyViewModeOverrides = false;
diff --git a/Engine/Plugins/FX/Niagara/Source/Niagara/Private/NiagaraRendererRibbons.cpp b/Engine/Plugins/FX/Niagara/Source/Niagara/Private/NiagaraRendererRibbons.cpp
--- a/Engine/Plugins/FX/Niagara/Source/Niagara/Private/NiagaraRendererRibbons.cpp
+++ b/Engine/Plugins/FX/Niagara/Source/Niagara/Private/NiagaraRendererRibbons.cpp
@@ -166,8 +166,8 @@ static TAutoConsoleVariable<int32> CVarRayTracingNiagaraRibbons(
 FORCEINLINE float AcosFast(float InX)
 {
 	float X = FMath::Abs(InX);
-	float Res = -0.156583f * X + (0.5 * PI);
-	Res *= sqrt(FMath::Max(0.f, 1.0f - X));
+	float Res = -0.156583f * X + (0.5f * PI);
+	Res *= FMath::Sqrt(FMath::Max(0.f, 1.0f - X));
 	return (InX >= 0) ? Res : PI - Res;
 }
 
@@ -921,7 +921,7 @@ int FNiagaraRendererRibbons::GetDynamicDataSize()const
 {
 	uint32 Size = sizeof(FNiagaraDynamicDataRibbon);
 
-	Size += ShapeState.SliceVertexData.GetAllocatedSize();
+	Size += uint32(ShapeState.SliceVertexData.GetAllocatedSize());
 
 	return Size;
 }
@@ -1343,8 +1343,8 @@ void FNiagaraRendererRibbons::CalculateUVScaleAndOffsets(const FNiagaraRibbonUVS
 		checkf(false, TEXT("Unsupported ribbon distribution mode"));
 	}
 
-	OutUScale = CalculatedUScale * UVSettings.Scale.X;
-	OutUOffset = (CalculatedUOffset * UVSettings.Scale.X) + UVSettings.Offset.X;
+	OutUScale = float(CalculatedUScale * UVSettings.Scale.X);
+	OutUOffset = float((CalculatedUOffset * UVSettings.Scale.X) + UVSettings.Offset.X);
 }
 
 
@@ -1360,20 +1360,20 @@ void FNiagaraRendererRibbons::GenerateVertexBufferForRibbonPart(const FNiagaraGe
 	
 	const int32 StartIndex = OutputData.SortedIndices.Num();
 
-	const FVector FirstPos = static_cast<FVector>(PosData[RibbonIndices[0]]);
-	FVector CurrPos = FirstPos;
-	FVector LastToCurrVec = FVector::ZeroVector;
+	const FVector3f FirstPos = PosData[RibbonIndices[0]];
+	FVector3f CurrPos = FirstPos;
+	FVector3f LastToCurrVec = FVector3f::ZeroVector;
 	float LastToCurrSize = 0;	
 	float LastTwist = 0;
 	float LastWidth = 0;
-	double TotalDistance = 0.0f;
+	float TotalDistance = 0.0f;
 
 	// Find the first position with enough distance.
 	int32 CurrentIndex = 1;
 	while (CurrentIndex < RibbonIndices.Num())
 	{
 		const int32 CurrentDataIndex = RibbonIndices[CurrentIndex];
-		CurrPos = static_cast<FVector>(PosData[CurrentDataIndex]);
+		CurrPos = PosData[CurrentDataIndex];
 		LastToCurrVec = CurrPos - FirstPos;
 		LastToCurrSize = LastToCurrVec.Size();
 		if constexpr (bHasTwist)
@@ -1409,8 +1409,8 @@ void FNiagaraRendererRibbons::GenerateVertexBufferForRibbonPart(const FNiagaraGe
 	while (NextIndex < RibbonIndices.Num())
 	{
 		const int32 NextDataIndex = RibbonIndices[NextIndex];
-		const FVector NextPos = static_cast<FVector>(PosData[NextDataIndex]);
-		FVector CurrToNextVec = NextPos - CurrPos;
+		const FVector3f NextPos = PosData[NextDataIndex];
+		FVector3f CurrToNextVec = NextPos - CurrPos;
 		const float CurrToNextSize = CurrToNextVec.Size();
 
 		float NextTwist = 0;
@@ -1426,7 +1426,7 @@ void FNiagaraRendererRibbons::GenerateVertexBufferForRibbonPart(const FNiagaraGe
 		{
 			// Normalize CurrToNextVec
 			CurrToNextVec *= 1.f / FMath::Max(GNiagaraRibbonMinSegmentLength, CurrToNextSize);
-			const FVector Tangent = (1.f - GenerationConfig.GetCurveTension()) * (LastToCurrVec + CurrToNextVec).GetSafeNormal();
+			const FVector3f Tangent = (1.f - GenerationConfig.GetCurveTension()) * (LastToCurrVec + CurrToNextVec).GetSafeNormal();
 
 			// Update the distance for CurrentIndex.
 			TotalDistance += LastToCurrSize;
@@ -1446,7 +1446,7 @@ void FNiagaraRendererRibbons::GenerateVertexBufferForRibbonPart(const FNiagaraGe
 			if constexpr (bWantsTessellation)
 			{
 				OutputData.AverageSegmentLength += CurrToNextSize * CurrToNextSize;
-				OutputData.AverageSegmentAngle += CurrToNextSize * AcosFast(FVector::DotProduct(LastToCurrVec, CurrToNextVec));
+				OutputData.AverageSegmentAngle += CurrToNextSize * AcosFast(FVector3f::DotProduct(LastToCurrVec, CurrToNextVec));
 				if constexpr (bHasTwist)
 				{
 					OutputData.AverageTwistAngle += CurrToNextSize * FMath::Abs(NextTwist - LastTwist);
@@ -1671,13 +1671,13 @@ void FNiagaraRendererRibbons::GenerateVertexBufferCPU(const FNiagaraGenerationIn
 	
 	if (OutputData.TotalSegmentLength > 0.0)
 	{
-		const double& TotalSegmentLength = OutputData.TotalSegmentLength;
+		const float& TotalSegmentLength = OutputData.TotalSegmentLength;
 	
 		// weighted sum based on the segment length :
-		double& AverageSegmentLength = OutputData.AverageSegmentLength;
-		double& AverageSegmentAngle = OutputData.AverageSegmentAngle;
-		double& AverageTwistAngle = OutputData.AverageTwistAngle;
-		double& AverageWidth = OutputData.AverageWidth;
+		float& AverageSegmentLength = OutputData.AverageSegmentLength;
+		float& AverageSegmentAngle = OutputData.AverageSegmentAngle;
+		float& AverageTwistAngle = OutputData.AverageTwistAngle;
+		float& AverageWidth = OutputData.AverageWidth;
 		
 		// Blend the result between the last frame tessellation factors and the current frame base on the total length of all segments.
 		// This is only used to increase the tessellation value of the current frame data to prevent glitches where tessellation is significantly changin between frames.
@@ -1748,7 +1748,7 @@ int32 FNiagaraRendererRibbons::CalculateTessellationFactor(const FNiagaraScenePr
 			                             FMath::Max<float>(1.f, FMath::Max(TessellationSmoothingData.TessellationTwistAngle, TessellationSmoothingData.TessellationAngle) / FMath::Max<float>(SMALL_NUMBER, TessellationMinAngle));
 
 		constexpr float MAX_CURVATURE_FACTOR = 0.002f; // This will clamp the curvature to around 2.5 km and avoid numerical issues.
-		const float ViewDistance = SceneProxy->GetProxyDynamicData().LODDistanceOverride >= 0.0f ? SceneProxy->GetProxyDynamicData().LODDistanceOverride : SceneProxy->GetBounds().ComputeSquaredDistanceFromBoxToPoint(ViewOriginForDistanceCulling);
+		const float ViewDistance = SceneProxy->GetProxyDynamicData().LODDistanceOverride >= 0.0f ? SceneProxy->GetProxyDynamicData().LODDistanceOverride : float(SceneProxy->GetBounds().ComputeSquaredDistanceFromBoxToPoint(ViewOriginForDistanceCulling));
 		const float MaxDisplacementError = FMath::Max(GNiagaraRibbonTessellationMinDisplacementError, ScreenPercentage * FMath::Sqrt(ViewDistance) / View->LODDistanceFactor);
 		float Tess = TessellationSmoothingData.TessellationAngle / FMath::Max(MAX_CURVATURE_FACTOR, AcosFast(TessellationSmoothingData.TessellationCurvature / (TessellationSmoothingData.TessellationCurvature + MaxDisplacementError)));
 		// FMath::RoundUpToPowerOfTwo ? This could avoid vertices moving around as tesselation increases
@@ -1770,7 +1770,7 @@ FNiagaraIndexGenerationInput FNiagaraRendererRibbons::CalculateIndexBufferConfig
 {
 	FNiagaraIndexGenerationInput IndexGenInput;
 
-	IndexGenInput.ViewDistance = SceneProxy->GetProxyDynamicData().LODDistanceOverride >= 0.0f ? SceneProxy->GetProxyDynamicData().LODDistanceOverride : SceneProxy->GetBounds().ComputeSquaredDistanceFromBoxToPoint(ViewOriginForDistanceCulling);
+	IndexGenInput.ViewDistance = SceneProxy->GetProxyDynamicData().LODDistanceOverride >= 0.0f ? SceneProxy->GetProxyDynamicData().LODDistanceOverride : float(SceneProxy->GetBounds().ComputeSquaredDistanceFromBoxToPoint(ViewOriginForDistanceCulling));
 	IndexGenInput.LODDistanceFactor = View->LODDistanceFactor;
 
 	if (bShouldUseGPUInitIndices)
@@ -1876,7 +1876,7 @@ void FNiagaraRendererRibbons::GenerateIndexBufferCPU(FNiagaraIndexGenerationInpu
 		const TArrayView<uint32> CurrentSegmentData(GeneratedGeometryData->SegmentData.GetData(), GeneratedGeometryData->SegmentData.Num());
 		CurrentIndexBuffer = AppendToIndexBufferCPU<TValue>(CurrentIndexBuffer, GeneratedData, ShapeState, CurrentSegmentData, false);
 	}
-	GeneratedData.CPUTriangleCount = (CurrentIndexBuffer - StartIndexBuffer) / 3;
+	GeneratedData.CPUTriangleCount = uint32(CurrentIndexBuffer - StartIndexBuffer) / 3;
 	check(CurrentIndexBuffer <= StartIndexBuffer + GeneratedData.TotalNumIndices);
 }
 
@@ -1916,11 +1916,11 @@ TValue* FNiagaraRendererRibbons::AppendToIndexBufferCPU(TValue* OutIndices, cons
 				const int32 FirstIndex = ShapeState.SliceTriangleToVertexIds[TriangleId];
 				const int32 SecondIndex = ShapeState.SliceTriangleToVertexIds[TriangleId + 1];
 				
-				OutIndices[0] = CurrSegment | FirstIndex;
-				OutIndices[1] = CurrSegment | SecondIndex;
-				OutIndices[2] = NextSegment | FirstIndex;
+				OutIndices[0] = TValue(CurrSegment | FirstIndex);
+				OutIndices[1] = TValue(CurrSegment | SecondIndex);
+				OutIndices[2] = TValue(NextSegment | FirstIndex);
 				OutIndices[3] = OutIndices[1];
-				OutIndices[4] = NextSegment | SecondIndex;
+				OutIndices[4] = TValue(NextSegment | SecondIndex);
 				OutIndices[5] = OutIndices[2];
 
 				OutIndices += 6;
@@ -1930,12 +1930,12 @@ TValue* FNiagaraRendererRibbons::AppendToIndexBufferCPU(TValue* OutIndices, cons
 				const uint32 FirstIndex = ShapeState.SliceTriangleToVertexIds[TriangleId];
 				const uint32 SecondIndex = ShapeState.SliceTriangleToVertexIds[TriangleId + 1];
 
-				OutIndices[0] = CurrSegment | FirstIndex;
-				OutIndices[1] = CurrSegment | SecondIndex;
-				OutIndices[2] = NextSegment | SecondIndex;
+				OutIndices[0] = TValue(CurrSegment | FirstIndex);
+				OutIndices[1] = TValue(CurrSegment | SecondIndex);
+				OutIndices[2] = TValue(NextSegment | SecondIndex);
 				OutIndices[3] = OutIndices[0];
 				OutIndices[4] = OutIndices[2];
-				OutIndices[5] = NextSegment | FirstIndex;
+				OutIndices[5] = TValue(NextSegment | FirstIndex);
 
 				OutIndices += 6;
 			}			
@@ -2002,7 +2002,10 @@ void FNiagaraRendererRibbons::SetupPerViewUniformBuffer(FNiagaraIndexGenerationI
 
 	PerViewUniformParameters.U0DistributionMode = static_cast<int32>(UV0Settings.DistributionMode);
 	PerViewUniformParameters.U1DistributionMode = static_cast<int32>(UV1Settings.DistributionMode);
-	PerViewUniformParameters.PackedVData = FVector4f(UV0Settings.Scale.Y, UV0Settings.Offset.Y, UV1Settings.Scale.Y, UV1Settings.Offset.Y);
+	PerViewUniformParameters.PackedVData.X = float(UV0Settings.Scale.Y);
+	PerViewUniformParameters.PackedVData.Y = float(UV0Settings.Offset.Y);
+	PerViewUniformParameters.PackedVData.Z = float(UV1Settings.Scale.Y);
+	PerViewUniformParameters.PackedVData.W = float(UV1Settings.Offset.Y);
 
 	OutUniformBuffer = FNiagaraRibbonUniformBufferRef::CreateUniformBufferImmediate(PerViewUniformParameters, UniformBuffer_SingleFrame);
 }
diff --git a/Engine/Plugins/FX/Niagara/Source/Niagara/Private/NiagaraRendererSprites.cpp b/Engine/Plugins/FX/Niagara/Source/Niagara/Private/NiagaraRendererSprites.cpp
--- a/Engine/Plugins/FX/Niagara/Source/Niagara/Private/NiagaraRendererSprites.cpp
+++ b/Engine/Plugins/FX/Niagara/Source/Niagara/Private/NiagaraRendererSprites.cpp
@@ -520,7 +520,7 @@ FNiagaraSpriteUniformBufferRef FNiagaraRendererSprites::CreateViewUniformBuffer(
 	PerViewUniformParameters.NormalsCylinderUnitDirection = FVector4f(0.0f, 0.0f, 1.0f, 0.0f);
 	PerViewUniformParameters.MacroUVParameters = CalcMacroUVParameters(View, SceneProxy.GetActorPosition(), MacroUVRadius);
 	PerViewUniformParameters.CameraFacingBlend = FVector4f(0.0f, 0.0f, 0.0f, 1.0f);
-	PerViewUniformParameters.RemoveHMDRoll = bRemoveHMDRollInVR;
+	PerViewUniformParameters.RemoveHMDRoll = bRemoveHMDRollInVR ? 0.0f : 1.0f;
 	PerViewUniformParameters.SubImageSize = FVector4f(SubImageSize.X, SubImageSize.Y, 1.0f / SubImageSize.X, 1.0f / SubImageSize.Y);
 
 	if (bUseLocalSpace)
@@ -961,7 +961,7 @@ void FNiagaraRendererSprites::GetDynamicMeshElements(const TArray<const FSceneVi
 
 			if (SourceMode == ENiagaraRendererSourceDataMode::Emitter && bEnableDistanceCulling)
 			{
-				FVector ViewOrigin = View->ViewMatrices.GetViewOrigin();
+				const FVector ViewOrigin = View->ViewMatrices.GetViewOrigin();
 				FVector RefPosition = SceneProxy->GetLocalToWorld().GetOrigin();
 				const int32 BoundPosOffset = VFBoundOffsetsInParamStore[ENiagaraSpriteVFLayout::Type::Position];
 				if (BoundPosOffset != INDEX_NONE && ParticleSpriteRenderData.DynamicDataSprites->ParameterDataBound.IsValidIndex(BoundPosOffset))
@@ -975,7 +975,7 @@ void FNiagaraRendererSprites::GetDynamicMeshElements(const TArray<const FSceneVi
 					}
 				}
 
-				float DistSquared = SceneProxy->GetProxyDynamicData().LODDistanceOverride >= 0.0f ? FMath::Square(SceneProxy->GetProxyDynamicData().LODDistanceOverride) : FVector::DistSquared(RefPosition, ViewOrigin);
+				const float DistSquared = SceneProxy->GetProxyDynamicData().LODDistanceOverride >= 0.0f ? FMath::Square(SceneProxy->GetProxyDynamicData().LODDistanceOverride) : float(FVector::DistSquared(RefPosition, ViewOrigin));
 				if (DistSquared < DistanceCullRange.X * DistanceCullRange.X || DistSquared > DistanceCullRange.Y * DistanceCullRange.Y)
 				{
 					// Distance cull the whole emitter
diff --git a/Engine/Plugins/FX/Niagara/Source/Niagara/Private/NiagaraScalabilityManager.cpp b/Engine/Plugins/FX/Niagara/Source/Niagara/Private/NiagaraScalabilityManager.cpp
--- a/Engine/Plugins/FX/Niagara/Source/Niagara/Private/NiagaraScalabilityManager.cpp
+++ b/Engine/Plugins/FX/Niagara/Source/Niagara/Private/NiagaraScalabilityManager.cpp
@@ -75,7 +75,7 @@ static int32 GetMaxUpdatesPerFrame(const UNiagaraEffectType* EffectType, int32 I
 
 FNiagaraScalabilityManager::FNiagaraScalabilityManager()
 	: EffectType(nullptr)
-	, LastUpdateTime(0.0f)
+	, LastUpdateTime(0.0)
 	, bRefreshCachedSystemData(false)
 {
 
@@ -476,7 +476,7 @@ void FNiagaraScalabilityManager::Update(FNiagaraWorldManager* WorldMan, float De
 
 		check(ManagedComponents.Num() == 0);
 		check(State.Num() == 0);
-		LastUpdateTime = 0.0f;
+		LastUpdateTime = 0.0;
 		return;
 	}
 
@@ -504,8 +504,8 @@ void FNiagaraScalabilityManager::Update(FNiagaraWorldManager* WorldMan, float De
 		return;
 	}
 
-	const float CurrentTime = WorldMan->GetWorld()->GetTimeSeconds();
-	const float TimeSinceUpdate = CurrentTime - LastUpdateTime;
+	const double CurrentTime = WorldMan->GetWorld()->GetTimeSeconds();
+	const float TimeSinceUpdate = float(CurrentTime - LastUpdateTime);
 	const float UpdatePeriod = GetScalabilityUpdatePeriod(EffectType->UpdateFrequency);
 
 	const bool bResetUpdate = bRefreshOwnerAllowsScalability || EffectType->UpdateFrequency == ENiagaraScalabilityUpdateFrequency::Continuous
@@ -573,7 +573,8 @@ FNiagaraScalabilitySystemData& FNiagaraScalabilityManager::GetSystemData(int32 C
 		int32* SysDataIndex = SystemDataIndexMap.Find(System);
 		if (SysDataIndex == nullptr)
 		{
-			CompState.SystemDataIndex = SystemData.Num();
+			ensure(SystemData.Num() <= TNumericLimits<decltype(CompState.SystemDataIndex)>::Max());
+			CompState.SystemDataIndex = int16(SystemData.Num());
 			SystemDataIndexMap.Add(System) = CompState.SystemDataIndex;			
 
 			FNiagaraScalabilitySystemData& NewSystemData = SystemData.AddDefaulted_GetRef();
@@ -582,7 +583,8 @@ FNiagaraScalabilitySystemData& FNiagaraScalabilityManager::GetSystemData(int32 C
 		}
 		else
 		{
-			CompState.SystemDataIndex = *SysDataIndex;
+			ensure(*SysDataIndex <= TNumericLimits<decltype(CompState.SystemDataIndex)>::Max());
+			CompState.SystemDataIndex = int16(*SysDataIndex);
 		}
 	}
 
diff --git a/Engine/Plugins/FX/Niagara/Source/Niagara/Private/NiagaraScript.cpp b/Engine/Plugins/FX/Niagara/Source/Niagara/Private/NiagaraScript.cpp
--- a/Engine/Plugins/FX/Niagara/Source/Niagara/Private/NiagaraScript.cpp
+++ b/Engine/Plugins/FX/Niagara/Source/Niagara/Private/NiagaraScript.cpp
@@ -1663,7 +1663,8 @@ FGraphEventRef UNiagaraScript::HandleByteCodeOptimization(bool bShouldForceNow)
 	OptimizedByteCodeResults->ExternalFunctionRegisterCounts.Reserve(CachedScriptVM.CalledVMExternalFunctions.Num());
 	for (const FVMExternalFunctionBindingInfo& FunctionBindingInfo : CachedScriptVM.CalledVMExternalFunctions)
 	{
-		const uint8 RegisterCount = FunctionBindingInfo.GetNumInputs() + FunctionBindingInfo.GetNumOutputs();
+		ensure(FunctionBindingInfo.GetNumInputs() + FunctionBindingInfo.GetNumOutputs() <= TNumericLimits<uint8>::Max());
+		const uint8 RegisterCount = uint8(FunctionBindingInfo.GetNumInputs() + FunctionBindingInfo.GetNumOutputs());
 		OptimizedByteCodeResults->ExternalFunctionRegisterCounts.Add(RegisterCount);
 	}
 	
diff --git a/Engine/Plugins/FX/Niagara/Source/Niagara/Private/NiagaraSimCache.cpp b/Engine/Plugins/FX/Niagara/Source/Niagara/Private/NiagaraSimCache.cpp
--- a/Engine/Plugins/FX/Niagara/Source/Niagara/Private/NiagaraSimCache.cpp
+++ b/Engine/Plugins/FX/Niagara/Source/Niagara/Private/NiagaraSimCache.cpp
@@ -400,7 +400,7 @@ bool UNiagaraSimCache::Read(float TimeSeconds, FNiagaraSystemInstance* SystemIns
 	}
 
 	const float FrameTime		= DurationSeconds > 0.0f ? (RelativeTime / DurationSeconds) * float(CacheFrames.Num() - 1) : 0.0f;
-	const float FrameIndex		= FMath::Floor(FrameTime);
+	const int32 FrameIndex		= FMath::FloorToInt(FrameTime);
 	const float FrameFraction	= FrameTime - float(FrameIndex);
 
 	return ReadFrame(FrameIndex, FrameFraction, SystemInstance);
diff --git a/Engine/Plugins/FX/Niagara/Source/Niagara/Private/NiagaraSimCacheAttributeReaderHelper.h b/Engine/Plugins/FX/Niagara/Source/Niagara/Private/NiagaraSimCacheAttributeReaderHelper.h
--- a/Engine/Plugins/FX/Niagara/Source/Niagara/Private/NiagaraSimCacheAttributeReaderHelper.h
+++ b/Engine/Plugins/FX/Niagara/Source/Niagara/Private/NiagaraSimCacheAttributeReaderHelper.h
@@ -155,7 +155,7 @@ struct FNiagaraSimCacheAttributeReaderHelper
 		FMemory::Memcpy(OutBuffer, &DataBuffers->Int32Data[(Variable->Int32Offset + Component) * DataBuffers->NumInstances * sizeof(int32)], sizeof(int32) * DataBuffers->NumInstances);
 	}
 
-	float ReadInt(int32 Instance) const
+	int32 ReadInt(int32 Instance) const
 	{
 		check(IsValid());
 		check(Variable->Int32Offset != INDEX_NONE && Variable->Int32Count == 1);
diff --git a/Engine/Plugins/FX/Niagara/Source/Niagara/Private/NiagaraSimCacheHelper.h b/Engine/Plugins/FX/Niagara/Source/Niagara/Private/NiagaraSimCacheHelper.h
--- a/Engine/Plugins/FX/Niagara/Source/Niagara/Private/NiagaraSimCacheHelper.h
+++ b/Engine/Plugins/FX/Niagara/Source/Niagara/Private/NiagaraSimCacheHelper.h
@@ -101,15 +101,15 @@ struct FNiagaraSimCacheHelper
 
 			CacheVariable.Variable = CompiledData.Variables[iDataSetVariable];
 			CacheVariable.FloatOffset = DataSetVariableLayout.GetNumFloatComponents() > 0 ? CacheLayout.FloatCount : INDEX_NONE;
-			CacheVariable.FloatCount = DataSetVariableLayout.GetNumFloatComponents();
+			CacheVariable.FloatCount = uint16(DataSetVariableLayout.GetNumFloatComponents());
 			CacheVariable.HalfOffset = DataSetVariableLayout.GetNumHalfComponents() > 0 ? CacheLayout.HalfCount : INDEX_NONE;
-			CacheVariable.HalfCount = DataSetVariableLayout.GetNumHalfComponents();
+			CacheVariable.HalfCount = uint16(DataSetVariableLayout.GetNumHalfComponents());
 			CacheVariable.Int32Offset = DataSetVariableLayout.GetNumInt32Components() > 0 ? CacheLayout.Int32Count : INDEX_NONE;
-			CacheVariable.Int32Count = DataSetVariableLayout.GetNumInt32Components();
+			CacheVariable.Int32Count = uint16(DataSetVariableLayout.GetNumInt32Components());
 
-			CacheLayout.FloatCount += DataSetVariableLayout.GetNumFloatComponents();
-			CacheLayout.HalfCount += DataSetVariableLayout.GetNumHalfComponents();
-			CacheLayout.Int32Count += DataSetVariableLayout.GetNumInt32Components();
+			CacheLayout.FloatCount += uint16(DataSetVariableLayout.GetNumFloatComponents());
+			CacheLayout.HalfCount += uint16(DataSetVariableLayout.GetNumHalfComponents());
+			CacheLayout.Int32Count += uint16(DataSetVariableLayout.GetNumInt32Components());
 		}
 
 		// Build write mappings we will build read mappings in a separate path
@@ -124,19 +124,19 @@ struct FNiagaraSimCacheHelper
 
 			for (int32 iComponent=0; iComponent < CacheVariable.FloatCount; ++iComponent)
 			{
-				CacheLayout.ComponentMappingsFromDataBuffer[FloatOffset] = DataSetVariableLayout.FloatComponentStart + iComponent;
+				CacheLayout.ComponentMappingsFromDataBuffer[FloatOffset] = uint16(DataSetVariableLayout.FloatComponentStart + iComponent);
 				++FloatOffset;
 			}
 
 			for (int32 iComponent=0; iComponent < CacheVariable.HalfCount; ++iComponent)
 			{
-				CacheLayout.ComponentMappingsFromDataBuffer[HalfOffset] = DataSetVariableLayout.HalfComponentStart + iComponent;
+				CacheLayout.ComponentMappingsFromDataBuffer[HalfOffset] = uint16(DataSetVariableLayout.HalfComponentStart + iComponent);
 				++HalfOffset;
 			}
 
 			for (int32 iComponent=0; iComponent < CacheVariable.Int32Count; ++iComponent)
 			{
-				CacheLayout.ComponentMappingsFromDataBuffer[Int32Offset] = DataSetVariableLayout.Int32ComponentStart + iComponent;
+				CacheLayout.ComponentMappingsFromDataBuffer[Int32Offset] = uint16(DataSetVariableLayout.Int32ComponentStart + iComponent);
 				++Int32Offset;
 			}
 		}
@@ -353,17 +353,17 @@ struct FNiagaraSimCacheHelper
 
 			for (int32 i = 0; i < SourceVariable.FloatCount; ++i)
 			{
-				CacheLayout.ComponentMappingsToDataBuffer[FloatOffset++] = DestVariableLayout ? DestVariableLayout->FloatComponentStart + i : INDEX_NONE;
+				CacheLayout.ComponentMappingsToDataBuffer[FloatOffset++] = uint16(DestVariableLayout ? DestVariableLayout->FloatComponentStart + i : INDEX_NONE);
 			}
 
 			for (int32 i = 0; i < SourceVariable.HalfCount; ++i)
 			{
-				CacheLayout.ComponentMappingsToDataBuffer[HalfOffset++] = DestVariableLayout ? DestVariableLayout->HalfComponentStart + i : INDEX_NONE;
+				CacheLayout.ComponentMappingsToDataBuffer[HalfOffset++] = uint16(DestVariableLayout ? DestVariableLayout->HalfComponentStart + i : INDEX_NONE);
 			}
 
 			for (int32 i = 0; i < SourceVariable.Int32Count; ++i)
 			{
-				CacheLayout.ComponentMappingsToDataBuffer[Int32Offset++] = DestVariableLayout ? DestVariableLayout->Int32ComponentStart + i : INDEX_NONE;
+				CacheLayout.ComponentMappingsToDataBuffer[Int32Offset++] = uint16(DestVariableLayout ? DestVariableLayout->Int32ComponentStart + i : INDEX_NONE);
 			}
 		}
 
@@ -661,9 +661,9 @@ struct FNiagaraSimCacheHelper
 				SrcFloats[i + (SrcStride * 2)]
 			);
 			const FVector RebasedPosition = RebaseTransform.TransformPosition(CachePosition);
-			DstFloats[i + (DstStride * 0)] = RebasedPosition.X;
-			DstFloats[i + (DstStride * 1)] = RebasedPosition.Y;
-			DstFloats[i + (DstStride * 2)] = RebasedPosition.Z;
+			DstFloats[i + (DstStride * 0)] = float(RebasedPosition.X);
+			DstFloats[i + (DstStride * 1)] = float(RebasedPosition.Y);
+			DstFloats[i + (DstStride * 2)] = float(RebasedPosition.Z);
 		}
 	}
 
diff --git a/Engine/Plugins/FX/Niagara/Source/Niagara/Private/NiagaraSpriteRendererProperties.cpp b/Engine/Plugins/FX/Niagara/Source/Niagara/Private/NiagaraSpriteRendererProperties.cpp
--- a/Engine/Plugins/FX/Niagara/Source/Niagara/Private/NiagaraSpriteRendererProperties.cpp
+++ b/Engine/Plugins/FX/Niagara/Source/Niagara/Private/NiagaraSpriteRendererProperties.cpp
@@ -451,8 +451,8 @@ void UNiagaraSpriteRendererProperties::PostEditChangeProperty(struct FPropertyCh
 	const FName PropertyName = PropertyChangedEvent.GetPropertyName();
 	const FName MemberPropertyName = PropertyChangedEvent.GetMemberPropertyName();
 
-	SubImageSize.X = FMath::Max<float>(SubImageSize.X, 1.f);
-	SubImageSize.Y = FMath::Max<float>(SubImageSize.Y, 1.f);
+	SubImageSize.X = FMath::Max(SubImageSize.X, 1.0);
+	SubImageSize.Y = FMath::Max(SubImageSize.Y, 1.0);
 
 	// DerivedData.BoundingGeometry in case we cleared the CutoutTexture
 	if (bUseMaterialCutoutTexture || CutoutTexture || DerivedData.BoundingGeometry.Num())
diff --git a/Engine/Plugins/FX/Niagara/Source/Niagara/Private/NiagaraSystem.cpp b/Engine/Plugins/FX/Niagara/Source/Niagara/Private/NiagaraSystem.cpp
--- a/Engine/Plugins/FX/Niagara/Source/Niagara/Private/NiagaraSystem.cpp
+++ b/Engine/Plugins/FX/Niagara/Source/Niagara/Private/NiagaraSystem.cpp
@@ -790,7 +790,7 @@ void UNiagaraSystem::ResolveWarmupTickCount()
 	}
 	else
 	{
-		WarmupTickCount = WarmupTime / WarmupTickDelta;
+		WarmupTickCount = FMath::FloorToInt(WarmupTime / WarmupTickDelta);
 		WarmupTime = WarmupTickDelta * WarmupTickCount;
 	}
 }
@@ -1241,7 +1241,7 @@ void UNiagaraSystem::GetAssetRegistryTags(TArray<FAssetRegistryTag>& OutTags) co
 #if WITH_EDITOR
 	OutTags.Add(FAssetRegistryTag("HasGPUEmitter", HasAnyGPUEmitters() ? TEXT("True") : TEXT("False"), FAssetRegistryTag::TT_Alphabetical));
 
-	const float BoundsSize = FixedBounds.GetSize().GetMax();
+	const float BoundsSize = float(FixedBounds.GetSize().GetMax());
 	OutTags.Add(FAssetRegistryTag("FixedBoundsSize", bFixedBounds ? FString::Printf(TEXT("%.2f"), BoundsSize) : FString(TEXT("None")), FAssetRegistryTag::TT_Numerical));
 
 	OutTags.Add(FAssetRegistryTag("NumEmitters", LexToString(EmitterHandles.Num()), FAssetRegistryTag::TT_Numerical));
@@ -2277,7 +2277,7 @@ void UNiagaraSystem::WaitForCompilationComplete(bool bIncludingGPUShaders, bool
 	}
 	
 	const int32 TotalCompiles = ActiveCompilations.Num() + GPUScripts.Num();
-	FScopedSlowTask Progress(TotalCompiles, LOCTEXT("WaitingForCompile", "Waiting for compilation to complete"));
+	FScopedSlowTask Progress(float(TotalCompiles), LOCTEXT("WaitingForCompile", "Waiting for compilation to complete"));
 	if (bShowProgress && TotalCompiles > 0)
 	{
 		Progress.MakeDialog();
diff --git a/Engine/Plugins/FX/Niagara/Source/Niagara/Private/NiagaraSystemInstance.cpp b/Engine/Plugins/FX/Niagara/Source/Niagara/Private/NiagaraSystemInstance.cpp
--- a/Engine/Plugins/FX/Niagara/Source/Niagara/Private/NiagaraSystemInstance.cpp
+++ b/Engine/Plugins/FX/Niagara/Source/Niagara/Private/NiagaraSystemInstance.cpp
@@ -1712,7 +1712,7 @@ float FNiagaraSystemInstance::GetLODDistance()
 			const FVector::FReal DistanceToEffectSqr = FVector(ViewInfo.ViewToWorld.GetOrigin() - EffectLocation).SizeSquared();
 			LODDistanceSqr = FMath::Min(LODDistanceSqr, DistanceToEffectSqr);
 		}
-		LODDistance = FMath::Sqrt(LODDistanceSqr);
+		LODDistance = float(FMath::Sqrt(LODDistanceSqr));
 	}
 	else
 	{
@@ -1746,7 +1746,7 @@ float FNiagaraSystemInstance::GetLODDistance()
 				const FVector::FReal DistanceToEffectSqr = FVector(ViewLocation - EffectLocation).SizeSquared();
 				LODDistanceSqr = FMath::Min(LODDistanceSqr, DistanceToEffectSqr);
 			}
-			LODDistance = FMath::Sqrt(LODDistanceSqr);
+			LODDistance = float(FMath::Sqrt(LODDistanceSqr));
 		}
 	}
 
@@ -1872,8 +1872,8 @@ void FNiagaraSystemInstance::TickInstanceParameters_GameThread(float DeltaSecond
 
 	//Bias the LastRenderTime slightly to account for any delay as it's written by the RT.
 	check(World);
-	GatheredInstanceParameters.TimeSeconds = World->TimeSeconds;
-	GatheredInstanceParameters.RealTimeSeconds = World->RealTimeSeconds;
+	GatheredInstanceParameters.TimeSeconds = float(World->TimeSeconds);				//LWC: Precision Loss
+	GatheredInstanceParameters.RealTimeSeconds = float(World->RealTimeSeconds);		//LWC: Precision Loss
 
 	// flip our buffered parameters
 	FlipParameterBuffers();
@@ -1950,7 +1950,7 @@ void FNiagaraSystemInstance::TickInstanceParameters_Concurrent()
 		CurrentOwnerParameters.EngineZAxis = CurrentOwnerParameters.EngineRotation.GetAxisZ();
 		CurrentOwnerParameters.EngineScale = (FVector3f)GatheredInstanceParameters.ComponentTrans.GetScale3D();
 		CurrentOwnerParameters.EngineLWCTile = LWCTile;
-		CurrentOwnerParameters.EngineLWCTile.W = FLargeWorldRenderScalar::GetTileSize();
+		CurrentOwnerParameters.EngineLWCTile.W = float(FLargeWorldRenderScalar::GetTileSize());
 	}
 
 	CurrentSystemParameters.EngineEmitterCount = GatheredInstanceParameters.EmitterCount;
diff --git a/Engine/Plugins/FX/Niagara/Source/Niagara/Private/NiagaraTypes.cpp b/Engine/Plugins/FX/Niagara/Source/Niagara/Private/NiagaraTypes.cpp
--- a/Engine/Plugins/FX/Niagara/Source/Niagara/Private/NiagaraTypes.cpp
+++ b/Engine/Plugins/FX/Niagara/Source/Niagara/Private/NiagaraTypes.cpp
@@ -165,7 +165,7 @@ void FNiagaraStructConversionStep::CopyToSim(uint8* DestinationData, const uint8
 	if (ConversionType == ENiagaraStructConversionType::DoubleToFloat)
 	{
 		checkf(SourceBytes == 8 && SimulationBytes == 4, TEXT("Wrong bytesizes for double->float conversion: Source %i, Simulation %i"), SourceBytes, SimulationBytes);
-		*static_cast<float*>(Dest) = *static_cast<const double*>(Src);
+		*static_cast<float*>(Dest) = float(*static_cast<const double*>(Src));
 	}
 	else if (ConversionType == ENiagaraStructConversionType::Vector2)
 	{
@@ -189,7 +189,7 @@ void FNiagaraStructConversionStep::CopyToSim(uint8* DestinationData, const uint8
 	{
 		checkf(SourceBytes == 32 && SimulationBytes == 16, TEXT("Wrong bytesizes for Quat d->f conversion: Source %i, Simulation %i"), SourceBytes, SimulationBytes);
 		FQuat4d SrcQuat = *static_cast<const FQuat4d*>(Src);
-		FQuat4f TargetQuat(SrcQuat.X, SrcQuat.Y, SrcQuat.Z, SrcQuat.W);
+		FQuat4f TargetQuat(SrcQuat);
 		FMemory::Memcpy(Dest, &TargetQuat, SimulationBytes);
 	}
 	else
diff --git a/Engine/Plugins/FX/Niagara/Source/Niagara/Private/NiagaraWorldManager.cpp b/Engine/Plugins/FX/Niagara/Source/Niagara/Private/NiagaraWorldManager.cpp
--- a/Engine/Plugins/FX/Niagara/Source/Niagara/Private/NiagaraWorldManager.cpp
+++ b/Engine/Plugins/FX/Niagara/Source/Niagara/Private/NiagaraWorldManager.cpp
@@ -938,8 +938,8 @@ void FNiagaraWorldManager::PostActorTick(float DeltaSeconds)
 			//Record custom events marking split times at set intervals. Allows us to generate summary tables for averages over shorter bursts.
 			if (GNiagaraCSVSplitTime > 0.0f)
 			{
-				float WorldTime = World->GetTimeSeconds();
-				float PrevWorldTime = WorldTime - DeltaSeconds;
+				const double WorldTime = World->GetTimeSeconds();
+				const double PrevWorldTime = WorldTime - DeltaSeconds;
 
 				int32 CurrentSplitIdx = (int32)(WorldTime / GNiagaraCSVSplitTime);
 				int32 PrevSplitIdx = (int32)(PrevWorldTime / GNiagaraCSVSplitTime);
@@ -1391,7 +1391,7 @@ void FNiagaraWorldManager::CalculateScalabilityState(UNiagaraSystem* System, con
 	{
 		FBoxSphereBounds Bounds = Component->CalcBounds(Component->GetComponentToWorld());		
 
-		float TimeSinceRendered = FMath::Max(0.0f, GetWorld()->LastRenderTime - Component->GetLastRenderTime() - World->GetDeltaSeconds() - GLastRenderTimeSafetyBias);		
+		const float TimeSinceRendered = float(FMath::Max(0.0f, GetWorld()->LastRenderTime - Component->GetLastRenderTime() - World->GetDeltaSeconds() - GLastRenderTimeSafetyBias));
 
 		if(bIsPreCull)
 		{
@@ -1457,14 +1457,14 @@ void FNiagaraWorldManager::SortedSignificanceCull(UNiagaraEffectType* EffectType
 
 			if (ScalabilitySettings.bCullMaxInstanceCount && ScalabilitySettings.BudgetScaling.bScaleMaxInstanceCountByGlobalBudgetUse)
 			{
-				float Scale = ScalabilitySettings.BudgetScaling.MaxInstanceCountScaleByGlobalBudgetUse.Evaluate(Usage);
-				EffectTypeInstanceMax *= Scale;
+				const float Scale = ScalabilitySettings.BudgetScaling.MaxInstanceCountScaleByGlobalBudgetUse.Evaluate(Usage);
+				EffectTypeInstanceMax = int32(float(EffectTypeInstanceMax) * Scale);
 				bCull = EffectTypeInstCount >= EffectTypeInstanceMax;
 			}
 			if (ScalabilitySettings.bCullPerSystemMaxInstanceCount && ScalabilitySettings.BudgetScaling.bScaleSystemInstanceCountByGlobalBudgetUse)
 			{
-				float Scale = ScalabilitySettings.BudgetScaling.MaxSystemInstanceCountScaleByGlobalBudgetUse.Evaluate(Usage);
-				SystemInstanceMax *= Scale;
+				const float Scale = ScalabilitySettings.BudgetScaling.MaxSystemInstanceCountScaleByGlobalBudgetUse.Evaluate(Usage);
+				SystemInstanceMax = int32(float(SystemInstanceMax) * Scale);
 				bCull |= SystemInstCount >= SystemInstanceMax;
 			}
 
@@ -1588,13 +1588,13 @@ void FNiagaraWorldManager::InstanceCountCull(UNiagaraEffectType* EffectType, UNi
 
 		if (ScalabilitySettings.BudgetScaling.bScaleMaxInstanceCountByGlobalBudgetUse)
 		{
-			float Scale = ScalabilitySettings.BudgetScaling.MaxInstanceCountScaleByGlobalBudgetUse.Evaluate(Usage);
-			EffectTypeInstanceMax *= Scale;
+			const float Scale = ScalabilitySettings.BudgetScaling.MaxInstanceCountScaleByGlobalBudgetUse.Evaluate(Usage);
+			EffectTypeInstanceMax = int32(float(EffectTypeInstanceMax) * Scale);
 		}
 		if (ScalabilitySettings.BudgetScaling.bScaleSystemInstanceCountByGlobalBudgetUse)
 		{
-			float Scale = ScalabilitySettings.BudgetScaling.MaxSystemInstanceCountScaleByGlobalBudgetUse.Evaluate(Usage);
-			SystemInstanceMax *= Scale;
+			const float Scale = ScalabilitySettings.BudgetScaling.MaxSystemInstanceCountScaleByGlobalBudgetUse.Evaluate(Usage);
+			SystemInstanceMax = int32(float(SystemInstanceMax) * Scale);
 		}
 		bCull = ScalabilitySettings.bCullMaxInstanceCount && EffectType->NumInstances >= EffectTypeInstanceMax;
 		bCull |= ScalabilitySettings.bCullPerSystemMaxInstanceCount && System->GetActiveInstancesCount() >= SystemInstanceMax;
diff --git a/Engine/Plugins/FX/Niagara/Source/Niagara/Public/NiagaraDebuggerCommon.h b/Engine/Plugins/FX/Niagara/Source/Niagara/Public/NiagaraDebuggerCommon.h
--- a/Engine/Plugins/FX/Niagara/Source/Niagara/Public/NiagaraDebuggerCommon.h
+++ b/Engine/Plugins/FX/Niagara/Source/Niagara/Public/NiagaraDebuggerCommon.h
@@ -465,7 +465,7 @@ struct NIAGARA_API FNiagaraDebugHUDSettingsData
 
 	/** Overview display location. */
 	UPROPERTY(EditAnywhere, Category = "Debug Overview", meta = (DisplayName = "Debug Overview Text Location", EditCondition = "bOverviewEnabled"))
-	FVector2D OverviewLocation = FIntPoint(30.0f, 150.0f);
+	FVector2D OverviewLocation = FVector2D(30.0f, 150.0f);
 
 	/** Overview display font to use. */
 	UPROPERTY(EditAnywhere, Category = "Debug Overview", meta = (EditCondition = "bOverviewEnabled && OverviewMode == ENiagaraDebugHUDOverviewMode::Overview"))
diff --git a/Engine/Plugins/FX/Niagara/Source/Niagara/Public/NiagaraRendererComponents.h b/Engine/Plugins/FX/Niagara/Source/Niagara/Public/NiagaraRendererComponents.h
--- a/Engine/Plugins/FX/Niagara/Source/Niagara/Public/NiagaraRendererComponents.h
+++ b/Engine/Plugins/FX/Niagara/Source/Niagara/Public/NiagaraRendererComponents.h
@@ -52,7 +52,7 @@ private:
 	struct FComponentPoolEntry
 	{
 		TWeakObjectPtr<USceneComponent> Component;
-		float LastActiveTime = 0.0f;
+		double LastActiveTime = 0.0;
 		TMap<FName, FComponentPropertyAddress> PropertyAddressMapping;
 		int32 LastAssignedToParticleID = -1;
 	};
diff --git a/Engine/Plugins/FX/Niagara/Source/Niagara/Public/NiagaraRendererRibbons.h b/Engine/Plugins/FX/Niagara/Source/Niagara/Public/NiagaraRendererRibbons.h
--- a/Engine/Plugins/FX/Niagara/Source/Niagara/Public/NiagaraRendererRibbons.h
+++ b/Engine/Plugins/FX/Niagara/Source/Niagara/Public/NiagaraRendererRibbons.h
@@ -59,7 +59,7 @@ struct FNiagaraGenerationInputDataCPUAccessors
 struct FNiagaraIndexGenerationInput
 {
 	float ViewDistance = 0.0f;
-	int32 LODDistanceFactor = 0;
+	float LODDistanceFactor = 0.0f;
 	
 	uint32 MaxSegmentCount = 0;
 	uint32 SubSegmentCount = 0;
@@ -161,8 +161,8 @@ struct FRibbonMultiRibbonInfo
 
 	FORCEINLINE bool UseInvertOrder(const FVector& ViewDirection, const FVector& ViewOriginForDistanceCulling, ENiagaraRibbonDrawDirection DrawDirection) const
 	{
-		const float StartDist = FVector::DotProduct(ViewDirection, StartPos - ViewOriginForDistanceCulling);
-		const float EndDist = FVector::DotProduct(ViewDirection, EndPos - ViewOriginForDistanceCulling);
+		const double StartDist = FVector::DotProduct(ViewDirection, StartPos - ViewOriginForDistanceCulling);
+		const double EndDist = FVector::DotProduct(ViewDirection, EndPos - ViewOriginForDistanceCulling);
 		return ((StartDist >= EndDist) && DrawDirection == ENiagaraRibbonDrawDirection::BackToFront)
 			|| ((StartDist < EndDist) && DrawDirection == ENiagaraRibbonDrawDirection::FrontToBack);
 	}
@@ -194,28 +194,20 @@ struct FNiagaraRibbonCPUGeneratedVertexData
 	/** Ribbon perperties required for sorting. */
 	TArray<FRibbonMultiRibbonInfo, TInlineAllocator<1>> RibbonInfoLookup;
 
-	double TotalSegmentLength;
-	double AverageSegmentLength;
-	double AverageSegmentAngle;
-	double AverageTwistAngle;
-	double AverageWidth;
-
-	FNiagaraRibbonCPUGeneratedVertexData()
-		: TotalSegmentLength(0.0)
-		, AverageSegmentLength(0.0)
-		, AverageSegmentAngle(0.0)
-		, AverageTwistAngle(0.0)
-		, AverageWidth(0.0)
-	{ }
+	float TotalSegmentLength = 0.0f;
+	float AverageSegmentLength = 0.0f;
+	float AverageSegmentAngle = 0.0f;
+	float AverageTwistAngle = 0.0f;
+	float AverageWidth = 0.0f;
 
 	int32 GetAllocatedSize()const
 	{
 		int32 Size = 0;
-		Size += SegmentData.GetAllocatedSize();
-		Size += SortedIndices.GetAllocatedSize();
-		Size += TangentAndDistances.GetAllocatedSize();
-		Size += MultiRibbonIndices.GetAllocatedSize();
-		Size += RibbonInfoLookup.GetAllocatedSize();
+		Size += int32(SegmentData.GetAllocatedSize());
+		Size += int32(SortedIndices.GetAllocatedSize());
+		Size += int32(TangentAndDistances.GetAllocatedSize());
+		Size += int32(MultiRibbonIndices.GetAllocatedSize());
+		Size += int32(RibbonInfoLookup.GetAllocatedSize());
 
 		return Size;
 	}
diff --git a/Engine/Plugins/FX/Niagara/Source/Niagara/Public/NiagaraScalabilityManager.h b/Engine/Plugins/FX/Niagara/Source/Niagara/Public/NiagaraScalabilityManager.h
--- a/Engine/Plugins/FX/Niagara/Source/Niagara/Public/NiagaraScalabilityManager.h
+++ b/Engine/Plugins/FX/Niagara/Source/Niagara/Public/NiagaraScalabilityManager.h
@@ -60,7 +60,7 @@ struct FNiagaraScalabilityManager
 	TMap<UNiagaraSystem*, int32> SystemDataIndexMap;
 	TArray<FNiagaraScalabilitySystemData> SystemData;
 
-	float LastUpdateTime;
+	double LastUpdateTime;
 
 	bool bRefreshOwnerAllowsScalability = false;
 
diff --git a/Engine/Plugins/FX/Niagara/Source/Niagara/Public/NiagaraSceneProxy.h b/Engine/Plugins/FX/Niagara/Source/Niagara/Public/NiagaraSceneProxy.h
--- a/Engine/Plugins/FX/Niagara/Source/Niagara/Public/NiagaraSceneProxy.h
+++ b/Engine/Plugins/FX/Niagara/Source/Niagara/Public/NiagaraSceneProxy.h
@@ -91,7 +91,7 @@ private:
 
 private:
 	/** Custom Uniform Buffers, allows us to have renderer specific data packed inside such as pre-skinned bounds. */
-	mutable TMap<uint64, TUniformBuffer<FPrimitiveUniformShaderParameters>*> CustomUniformBuffers;
+	mutable TMap<uint32, TUniformBuffer<FPrimitiveUniformShaderParameters>*> CustomUniformBuffers;
 
 	/** The data required to render a single instance of a NiagaraSystem */
 	FNiagaraSystemRenderData* RenderData = nullptr;
diff --git a/Engine/Plugins/FX/Niagara/Source/Niagara/Public/NiagaraScriptExecutionParameterStore.h b/Engine/Plugins/FX/Niagara/Source/Niagara/Public/NiagaraScriptExecutionParameterStore.h
--- a/Engine/Plugins/FX/Niagara/Source/Niagara/Public/NiagaraScriptExecutionParameterStore.h
+++ b/Engine/Plugins/FX/Niagara/Source/Niagara/Public/NiagaraScriptExecutionParameterStore.h
@@ -12,7 +12,14 @@ struct FNiagaraScriptExecutionPaddingInfo
 	GENERATED_USTRUCT_BODY()
 public:
 	FNiagaraScriptExecutionPaddingInfo() : SrcOffset(0), DestOffset(0), SrcSize(0), DestSize(0) {}
-	FNiagaraScriptExecutionPaddingInfo(uint32 InSrcOffset, uint32 InDestOffset, uint32 InSrcSize, uint32 InDestSize) : SrcOffset(InSrcOffset), DestOffset(InDestOffset), SrcSize(InSrcSize), DestSize(InDestSize) {}
+	FNiagaraScriptExecutionPaddingInfo(uint32 InSrcOffset, uint32 InDestOffset, uint32 InSrcSize, uint32 InDestSize)
+		: SrcOffset(uint16(InSrcOffset)), DestOffset(uint16(InDestOffset)), SrcSize(uint16(InSrcSize)), DestSize(uint16(InDestSize))
+	{
+		ensure(InSrcOffset <= TNumericLimits<uint16>::Max());
+		ensure(InDestOffset <= TNumericLimits<uint16>::Max());
+		ensure(InSrcSize <= TNumericLimits<uint16>::Max());
+		ensure(InDestSize <= TNumericLimits<uint16>::Max());
+	}
 
 	UPROPERTY()
 	uint16 SrcOffset;
diff --git a/Engine/Plugins/FX/Niagara/Source/Niagara/Public/NiagaraTypes.h b/Engine/Plugins/FX/Niagara/Source/Niagara/Public/NiagaraTypes.h
--- a/Engine/Plugins/FX/Niagara/Source/Niagara/Public/NiagaraTypes.h
+++ b/Engine/Plugins/FX/Niagara/Source/Niagara/Public/NiagaraTypes.h
@@ -766,14 +766,14 @@ struct NIAGARA_API FNiagaraTypeDefinition
 		UT_Enum
 	};
 
-	enum FTypeFlags
+	enum FTypeFlags : uint8
 	{
-		TF_None			= 0x0000,
-		TF_Static		= 0x0001,
+		TF_None			= 0x00,
+		TF_Static		= 0x01,
 
 		/// indicates that the ClassStructOrEnum property has been serialized as the LWC struct (see FNiagaraTypeHelper)
 		/// instead of the Transient SWC version of the struct
-		TF_SerializedAsLWC	= 0x0002,
+		TF_SerializedAsLWC	= 0x02,
 	};
 
 public:
@@ -970,7 +970,9 @@ public:
 			}
 			else
 			{
-				Size = CastChecked<UScriptStruct>(GetStruct())->GetStructureSize();
+				const int32 StructSize = CastChecked<UScriptStruct>(GetStruct())->GetStructureSize();
+				ensure(StructSize <= TNumericLimits<int16>::Max());
+				Size = int16(StructSize);
 			}
 		}
 		return Size;
@@ -987,7 +989,9 @@ public:
 			}
 			else
 			{
-				Alignment = CastChecked<UScriptStruct>(GetStruct())->GetMinAlignment();
+				const int32 StructAlignment = CastChecked<UScriptStruct>(GetStruct())->GetMinAlignment();
+				ensure(StructAlignment <= TNumericLimits<int16>::Max());
+				Alignment = int16(StructAlignment);
 			}
 		}
 		return Alignment;
diff --git a/Engine/Plugins/FX/Niagara/Source/NiagaraVertexFactories/Public/NiagaraVertexFactory.h b/Engine/Plugins/FX/Niagara/Source/NiagaraVertexFactories/Public/NiagaraVertexFactory.h
--- a/Engine/Plugins/FX/Niagara/Source/NiagaraVertexFactories/Public/NiagaraVertexFactory.h
+++ b/Engine/Plugins/FX/Niagara/Source/NiagaraVertexFactories/Public/NiagaraVertexFactory.h
@@ -67,7 +67,7 @@ public:
 		, LastFrameSetup(MAX_uint32)
 		, LastViewFamily(nullptr)
 		, LastView(nullptr)
-		, LastFrameRealTime(-1.0f)
+		, LastFrameRealTime(-1.0)
 		, ParticleFactoryType(Type)
 		, bInUse(false)
 	{
@@ -123,7 +123,7 @@ private:
 	mutable uint32 LastFrameSetup;
 	mutable const FSceneViewFamily *LastViewFamily;
 	mutable const FSceneView *LastView;
-	mutable float LastFrameRealTime;
+	mutable double LastFrameRealTime;
 
 	/** The type of the vertex factory. */
 	ENiagaraVertexFactoryType ParticleFactoryType;
diff --git a/Engine/Source/Runtime/Engine/Classes/Particles/ParticlePerfStatsManager.h b/Engine/Source/Runtime/Engine/Classes/Particles/ParticlePerfStatsManager.h
--- a/Engine/Source/Runtime/Engine/Classes/Particles/ParticlePerfStatsManager.h
+++ b/Engine/Source/Runtime/Engine/Classes/Particles/ParticlePerfStatsManager.h
@@ -28,27 +28,27 @@ struct ENGINE_API FAccumulatedParticlePerfStats_GT
 	void Tick(FParticlePerfStats& Stats);
 
 	/** Returns the total cycles used by all GameThread stats. */
-	FORCEINLINE uint64 GetTotalCycles()const { return AccumulatedStats.GetTotalCycles(); }
+	FORCEINLINE uint64 GetTotalCycles() const { return AccumulatedStats.GetTotalCycles(); }
 
 	/** Returns the average cycles per frame by all GameThread stats. */
-	FORCEINLINE uint64 GetPerFrameAvgCycles()const { return NumFrames > 0 ? AccumulatedStats.GetTotalCycles() / NumFrames : 0; }
+	FORCEINLINE uint64 GetPerFrameAvgCycles() const { return NumFrames > 0 ? AccumulatedStats.GetTotalCycles() / NumFrames : 0; }
 	/** Returns the max cycles per frame by all GameThread stats. */
-	FORCEINLINE uint64 GetPerFrameMaxCycles(int32 Index = 0)const { return MaxPerFrameTotalCycles[Index]; }
+	FORCEINLINE uint64 GetPerFrameMaxCycles(int32 Index = 0) const { return MaxPerFrameTotalCycles[Index]; }
 
 	/** Returns the average time in µs per frame by all GameThread stats. */
-	FORCEINLINE double GetPerFrameAvg()const { return FPlatformTime::ToMilliseconds64(GetPerFrameAvgCycles()) * 1000.0; }
+	FORCEINLINE float GetPerFrameAvg() const { return float(FPlatformTime::ToMilliseconds64(GetPerFrameAvgCycles()) * 1000.0); }
 	/** Returns the max time in µs per frame by all GameThread stats. */
-	FORCEINLINE double GetPerFrameMax(int32 Index = 0)const { return FPlatformTime::ToMilliseconds64(GetPerFrameMaxCycles(Index)) * 1000.0; }
+	FORCEINLINE float GetPerFrameMax(int32 Index = 0) const { return float(FPlatformTime::ToMilliseconds64(GetPerFrameMaxCycles(Index)) * 1000.0); }
 
 	/** Returns the average cycles per instance by all GameThread stats. */
-	FORCEINLINE uint64 GetPerInstanceAvgCycles()const { return AccumulatedStats.GetPerInstanceAvgCycles(); }
+	FORCEINLINE uint64 GetPerInstanceAvgCycles() const { return AccumulatedStats.GetPerInstanceAvgCycles(); }
 	/** Returns the max cycles per instance by all GameThread stats. */
-	FORCEINLINE uint64 GetPerInstanceMaxCycles(int32 Index = 0)const { return MaxPerInstanceCycles[Index]; }
+	FORCEINLINE uint64 GetPerInstanceMaxCycles(int32 Index = 0) const { return MaxPerInstanceCycles[Index]; }
 
 	/** Returns the average time in µs per instance by all GameThread stats. */
-	FORCEINLINE double GetPerInstanceAvg()const { return FPlatformTime::ToMilliseconds64(GetPerInstanceAvgCycles()) * 1000.0; }
+	FORCEINLINE float GetPerInstanceAvg() const { return float(FPlatformTime::ToMilliseconds64(GetPerInstanceAvgCycles()) * 1000.0); }
 	/** Returns the max time in µs per instance by all GameThread stats. */
-	FORCEINLINE double GetPerInstanceMax(int32 Index = 0)const { return FPlatformTime::ToMilliseconds64(GetPerInstanceMaxCycles(Index)) * 1000.0; }
+	FORCEINLINE float GetPerInstanceMax(int32 Index = 0) const { return float(FPlatformTime::ToMilliseconds64(GetPerInstanceMaxCycles(Index)) * 1000.0); }
 };
 
 struct ENGINE_API FAccumulatedParticlePerfStats_RT
@@ -64,27 +64,27 @@ struct ENGINE_API FAccumulatedParticlePerfStats_RT
 	void Tick(FParticlePerfStats& Stats);
 
 	/** Returns the total cycles used by all RenderThread stats. */
-	FORCEINLINE uint64 GetTotalCycles()const { return AccumulatedStats.GetTotalCycles(); }
+	FORCEINLINE uint64 GetTotalCycles() const { return AccumulatedStats.GetTotalCycles(); }
 
 	/** Returns the average cycles per frame by all RenderThread stats. */
-	FORCEINLINE uint64 GetPerFrameAvgCycles()const { return NumFrames > 0 ? AccumulatedStats.GetTotalCycles() / NumFrames : 0; }
+	FORCEINLINE uint64 GetPerFrameAvgCycles() const { return NumFrames > 0 ? AccumulatedStats.GetTotalCycles() / NumFrames : 0; }
 	/** Returns the max cycles per frame by all RenderThread stats. */
-	FORCEINLINE uint64 GetPerFrameMaxCycles(int32 Index = 0)const { return MaxPerFrameTotalCycles[Index]; }
+	FORCEINLINE uint64 GetPerFrameMaxCycles(int32 Index = 0) const { return MaxPerFrameTotalCycles[Index]; }
 
 	/** Returns the average time in µs per frame by all RenderThread stats. */
-	FORCEINLINE double GetPerFrameAvg()const { return FPlatformTime::ToMilliseconds64(GetPerFrameAvgCycles()) * 1000.0; }
+	FORCEINLINE float GetPerFrameAvg() const { return float(FPlatformTime::ToMilliseconds64(GetPerFrameAvgCycles()) * 1000.0); }
 	/** Returns the max time in µs per frame by all RenderThread stats. */
-	FORCEINLINE double GetPerFrameMax(int32 Index = 0)const { return FPlatformTime::ToMilliseconds64(GetPerFrameMaxCycles(Index)) * 1000.0; }
+	FORCEINLINE float GetPerFrameMax(int32 Index = 0) const { return float(FPlatformTime::ToMilliseconds64(GetPerFrameMaxCycles(Index)) * 1000.0); }
 
 	/** Returns the average cycles per instance by all RenderThread stats. */
-	FORCEINLINE uint64 GetPerInstanceAvgCycles()const { return AccumulatedStats.GetPerInstanceAvgCycles(); }
+	FORCEINLINE uint64 GetPerInstanceAvgCycles() const { return AccumulatedStats.GetPerInstanceAvgCycles(); }
 	/** Returns the max cycles per instance by all RenderThread stats. */
-	FORCEINLINE uint64 GetPerInstanceMaxCycles(int32 Index = 0)const { return MaxPerInstanceCycles[Index];  }
+	FORCEINLINE uint64 GetPerInstanceMaxCycles(int32 Index = 0) const { return MaxPerInstanceCycles[Index];  }
 
 	/** Returns the average time in µs per instance by all RenderThread stats. */
-	FORCEINLINE double GetPerInstanceAvg()const { return FPlatformTime::ToMilliseconds64(GetPerInstanceAvgCycles()) * 1000.0; }
+	FORCEINLINE float GetPerInstanceAvg() const { return float(FPlatformTime::ToMilliseconds64(GetPerInstanceAvgCycles()) * 1000.0); }
 	/** Returns the max time in µs per instance by all RenderThread stats. */
-	FORCEINLINE double GetPerInstanceMax(int32 Index = 0)const { return FPlatformTime::ToMilliseconds64(GetPerInstanceMaxCycles(Index)) * 1000.0; }
+	FORCEINLINE float GetPerInstanceMax(int32 Index = 0) const { return float(FPlatformTime::ToMilliseconds64(GetPerInstanceMaxCycles(Index)) * 1000.0); }
 };
 
 struct ENGINE_API FAccumulatedParticlePerfStats_GPU
@@ -105,7 +105,7 @@ struct ENGINE_API FAccumulatedParticlePerfStats_GPU
 	FORCEINLINE uint64 GetPerFrameAvgMicroseconds() const { return NumFrames > 0 ? AccumulatedStats.GetTotalMicroseconds() / NumFrames : 0; }
 	FORCEINLINE uint64 GetPerFrameMaxMicroseconds(int32 Index = 0) const { return MaxPerFrameTotalMicroseconds[Index]; }
 
-	FORCEINLINE uint64 GetPerInstanceAvgMicroseconds()const { return AccumulatedStats.GetPerInstanceAvgMicroseconds(); }
+	FORCEINLINE uint64 GetPerInstanceAvgMicroseconds() const { return AccumulatedStats.GetPerInstanceAvgMicroseconds(); }
 	FORCEINLINE uint64 GetPerInstanceMaxMicroseconds(int32 Index = 0) const { return MaxPerInstanceMicroseconds[Index]; }
 };
 
@@ -196,16 +196,16 @@ public:
 	virtual void OnRemoveComponent(const TWeakObjectPtr<const UFXSystemComponent>& Component) {}
 
 
-	virtual bool NeedsWorldStats()const = 0;
-	virtual bool NeedsSystemStats()const = 0;
-	virtual bool NeedsComponentStats()const = 0;
+	virtual bool NeedsWorldStats() const = 0;
+	virtual bool NeedsSystemStats() const = 0;
+	virtual bool NeedsComponentStats() const = 0;
 
 	/** 
 	Controls whether this listener should be stooped and cleaned up when it's orphaned, i.e. the manager is the only one with a reference.
 	In some cases we want to signal that a listener should stop by clearing an external reference.
 	In other cases we want to have fire and forget listeners that can signal their own termination via their tick function.
 	*/
-	virtual bool AllowOrphaned()const { return false; }
+	virtual bool AllowOrphaned() const { return false; }
 };
 
 typedef TSharedPtr<FParticlePerfStatsListener, ESPMode::ThreadSafe> FParticlePerfStatsListenerPtr;
@@ -329,9 +329,9 @@ public:
 	virtual void OnRemoveComponent(const TWeakObjectPtr<const UFXSystemComponent>& Component)override;
 	#endif
 
-	virtual bool NeedsWorldStats()const override { return bGatherWorldStats; }
-	virtual bool NeedsSystemStats()const override { return bGatherSystemStats; }
-	virtual bool NeedsComponentStats()const override { return bGatherComponentStats; }
+	virtual bool NeedsWorldStats() const override { return bGatherWorldStats; }
+	virtual bool NeedsSystemStats() const override { return bGatherSystemStats; }
+	virtual bool NeedsComponentStats() const override { return bGatherComponentStats; }
 
 	void DumpStatsToDevice(FOutputDevice& Ar);
 	void DumpStatsToFile();
@@ -382,7 +382,7 @@ public:
 	virtual void End()override;
 	virtual bool Tick()override;
 
-	virtual bool AllowOrphaned()const { return true; }
+	virtual bool AllowOrphaned() const { return true; }
 
 private:
 	int32 FramesRemaining;
