- Remove check for WeakStaticMesh ptr
- This isn't a valid check as the SM could have being flagged as unreachable but the rendering side is still safe to use for a frame

#rnx

diff --git a/Engine/Plugins/FX/Niagara/Source/Niagara/Private/NiagaraMeshRendererProperties.cpp b/Engine/Plugins/FX/Niagara/Source/Niagara/Private/NiagaraMeshRendererProperties.cpp
--- a/Engine/Plugins/FX/Niagara/Source/Niagara/Private/NiagaraMeshRendererProperties.cpp
+++ b/Engine/Plugins/FX/Niagara/Source/Niagara/Private/NiagaraMeshRendererProperties.cpp
@@ -54,13 +54,6 @@ public:
 
 	void GetLODModelData(FLODModelData& OutLODModelData) const override
 	{
-		// This should never happen, but for safety
-		if (!ensure(WeakStaticMesh.IsValid()))
-		{
-			OutLODModelData.LODIndex = INDEX_NONE;
-			return;
-		}
-
 		OutLODModelData.LODIndex = RenderData->GetCurrentFirstLODIdx(MinLOD);
 		if (!RenderData->LODResources.IsValidIndex(OutLODModelData.LODIndex))
 		{
