- Add AppendTriangle as a quick method to add triangles rather than allocating and setting each element
- Fix for crashing if you specify no materials

diff --git a/Engine/Plugins/FX/Niagara/Shaders/Private/NiagaraDataInterfaceDynamicMeshTemplate.ush b/Engine/Plugins/FX/Niagara/Shaders/Private/NiagaraDataInterfaceDynamicMeshTemplate.ush
--- a/Engine/Plugins/FX/Niagara/Shaders/Private/NiagaraDataInterfaceDynamicMeshTemplate.ush
+++ b/Engine/Plugins/FX/Niagara/Shaders/Private/NiagaraDataInterfaceDynamicMeshTemplate.ush
@@ -16,25 +16,29 @@ RWBuffer<uint>	{ParameterName}_VertexBuffer;
 
 ////////////////////////////////////////////////////////
 // Helper functions
+uint {ParameterName}_GetPositionOffsetUnsafe(int Vertex)				{ return (Vertex * 3) + {ParameterName}_PositionOffset; }
+uint {ParameterName}_GetTangentBasisOffsetUnsafe(int Vertex)			{ return (Vertex * 2) + {ParameterName}_TangentBasisOffset; }
+uint {ParameterName}_GetTexCoordOffsetUnsafe(int Vertex, int TexCoord)	{ return (Vertex * 2) + (TexCoord * 2) + {ParameterName}_TexCoordOffset; }
+uint {ParameterName}_GetColorOffsetUnsafe(int Vertex)					{ return Vertex + {ParameterName}_ColorOffset;}
 
 uint {ParameterName}_GetPositionOffset(int Vertex)
 {
-	return ({ParameterName}_PositionOffset != -1) && (Vertex >= 0) && (Vertex < {ParameterName}_NumVertices) ? (Vertex * 3) + {ParameterName}_PositionOffset : -1;
+	return ({ParameterName}_PositionOffset != -1) && (Vertex >= 0) && (Vertex < {ParameterName}_NumVertices) ? {ParameterName}_GetPositionOffsetUnsafe(Vertex) : -1;
 }
 
 uint {ParameterName}_GetTangentBasisOffset(int Vertex)
 {
-	return ({ParameterName}_TangentBasisOffset != -1) && (Vertex >= 0) && (Vertex < {ParameterName}_NumVertices) ? (Vertex * 2) + {ParameterName}_TangentBasisOffset : -1;
+	return ({ParameterName}_TangentBasisOffset != -1) && (Vertex >= 0) && (Vertex < {ParameterName}_NumVertices) ? {ParameterName}_GetTangentBasisOffsetUnsafe(Vertex) : -1;
 }
 
 uint {ParameterName}_GetTexCoordOffset(int Vertex, int TexCoord)
 {
-	return ({ParameterName}_TexCoordOffset != -1) && (Vertex >= 0) && (Vertex < {ParameterName}_NumVertices) && (TexCoord >= 0) && (TexCoord <= {ParameterName}_NumTexCoords) ? (Vertex * 2) + (TexCoord * 2) + {ParameterName}_TexCoordOffset : -1;
+	return ({ParameterName}_TexCoordOffset != -1) && (Vertex >= 0) && (Vertex < {ParameterName}_NumVertices) && (TexCoord >= 0) && (TexCoord <= {ParameterName}_NumTexCoords) ? {ParameterName}_GetTexCoordOffsetUnsafe(Vertex, TexCoord) : -1;
 }
 
 uint {ParameterName}_GetColorOffset(int Vertex)
 {
-	return ({ParameterName}_ColorOffset != -1) && (Vertex >= 0) && (Vertex < {ParameterName}_NumVertices) ? Vertex + {ParameterName}_ColorOffset : -1;
+	return ({ParameterName}_ColorOffset != -1) && (Vertex >= 0) && (Vertex < {ParameterName}_NumVertices) ? {ParameterName}_GetColorOffsetUnsafe(Vertex) : -1;
 }
 
 float4 {ParameterName}_UnpackTangent(uint Packed)
@@ -79,6 +83,35 @@ uint {ParameterName}_PackColor(float4 Color)
 	return Packed;
 }
 
+void {ParameterName}_SetVertexPositionUnsafe(int Vertex, float3 Position)
+{
+	uint Offset = {ParameterName}_GetPositionOffsetUnsafe(Vertex);
+	{ParameterName}_VertexBuffer[Offset + 0] = asuint(Position.x);
+	{ParameterName}_VertexBuffer[Offset + 1] = asuint(Position.y);
+	{ParameterName}_VertexBuffer[Offset + 2] = asuint(Position.z);
+}
+
+void {ParameterName}_SetVertexTangentBasisUnsafe(int Vertex, float3 TangentX, float3 TangentY, float3 TangentZ)
+{
+	uint Offset = {ParameterName}_GetTangentBasisOffsetUnsafe(Vertex);
+	uint TangentSign = dot(cross(TangentX, TangentZ), TangentY) < 0 ? 0x80 : 0x7f;
+	{ParameterName}_VertexBuffer[Offset + 0] = {ParameterName}_PackTangent(TangentX, TangentSign);
+	{ParameterName}_VertexBuffer[Offset + 1] = {ParameterName}_PackTangent(TangentZ, TangentSign);
+}
+
+void {ParameterName}_SetVertexTexCoordUnsafe(int Vertex, float2 TexCoord, int TexCoordIndex)
+{
+	uint Offset = {ParameterName}_GetTexCoordOffsetUnsafe(Vertex, TexCoordIndex);
+	{ParameterName}_VertexBuffer[Offset + 0] = asuint(TexCoord.x);
+	{ParameterName}_VertexBuffer[Offset + 1] = asuint(TexCoord.y);
+}
+
+void {ParameterName}_SetVertexColorUnsafe(int Vertex, float4 Color)
+{
+	uint Offset = {ParameterName}_GetColorOffsetUnsafe(Vertex);
+	{ParameterName}_VertexBuffer[Offset] = {ParameterName}_PackColor(Color);
+}
+
 ////////////////////////////////////////////////////////
 // Mutable Functions
 void AllocateSectionTriangles_{ParameterName}_UEImpureCall(bool bExecute, int SectionIndex, int NumTriangles, out int TriangleIndex, out int NumAllocated)
@@ -198,6 +231,58 @@ void SetVertexData_{ParameterName}_UEImpureCall(in bool bExecute, in int Vertex,
 	}
 }
 
+void AppendTriangle_{ParameterName}_UEImpureCall(
+	in bool bExecute, in int SectionIndex,
+	in float3 Position0, in float3 TangentX0, in float3 TangentY0, in float3 TangentZ0, in float2 TexCoord0, in float4 Color0,
+	in float3 Position1, in float3 TangentX1, in float3 TangentY1, in float3 TangentZ1, in float2 TexCoord1, in float4 Color1,
+	in float3 Position2, in float3 TangentX2, in float3 TangentY2, in float3 TangentZ2, in float2 TexCoord2, in float4 Color2,
+	out int TriangleIndex
+)
+{
+	TriangleIndex = -1;
+	if (bExecute && (SectionIndex >= 0 && SectionIndex < {ParameterName}_NumSections))
+	{
+		int MaxTriangles = {ParameterName}_SectionBuffer[SectionIndex * 2 + 0];
+		InterlockedAdd({ParameterName}_SectionBuffer[SectionIndex * 2 + 1], 1, TriangleIndex);
+		if (TriangleIndex >= MaxTriangles)
+		{
+			InterlockedAdd({ParameterName}_SectionBuffer[SectionIndex * 2 + 1], -1);
+			TriangleIndex = -1;
+		}
+		else
+		{
+			{ParameterName}_IndexBuffer[TriangleIndex * 3 + 0] = (uint)TriangleIndex * 3 + 0;
+			{ParameterName}_IndexBuffer[TriangleIndex * 3 + 1] = (uint)TriangleIndex * 3 + 1;
+			{ParameterName}_IndexBuffer[TriangleIndex * 3 + 2] = (uint)TriangleIndex * 3 + 2;
+
+			if ( {ParameterName}_PositionOffset != -1 )
+			{
+				{ParameterName}_SetVertexPositionUnsafe(TriangleIndex * 3 + 0, Position0);
+				{ParameterName}_SetVertexPositionUnsafe(TriangleIndex * 3 + 1, Position1);
+				{ParameterName}_SetVertexPositionUnsafe(TriangleIndex * 3 + 2, Position2);
+			}
+			if ( {ParameterName}_TangentBasisOffset != -1 )
+			{
+				{ParameterName}_SetVertexTangentBasisUnsafe(TriangleIndex * 3 + 0, TangentX0, TangentY0, TangentZ0);
+				{ParameterName}_SetVertexTangentBasisUnsafe(TriangleIndex * 3 + 1, TangentX1, TangentY1, TangentZ1);
+				{ParameterName}_SetVertexTangentBasisUnsafe(TriangleIndex * 3 + 2, TangentX2, TangentY2, TangentZ2);
+			}
+			if ( {ParameterName}_TexCoordOffset != -1 )
+			{
+				{ParameterName}_SetVertexTexCoordUnsafe(TriangleIndex * 3 + 0, TexCoord0, 0);
+				{ParameterName}_SetVertexTexCoordUnsafe(TriangleIndex * 3 + 1, TexCoord1, 0);
+				{ParameterName}_SetVertexTexCoordUnsafe(TriangleIndex * 3 + 2, TexCoord2, 0);
+			}
+			if ( {ParameterName}_ColorOffset != -1 )
+			{
+				{ParameterName}_SetVertexColorUnsafe(TriangleIndex * 3 + 0, Color0);
+				{ParameterName}_SetVertexColorUnsafe(TriangleIndex * 3 + 1, Color1);
+				{ParameterName}_SetVertexColorUnsafe(TriangleIndex * 3 + 2, Color2);
+			}
+		}
+	}
+}
+
 ////////////////////////////////////////////////////////
 // Imutable Functions
 
diff --git a/Engine/Plugins/FX/Niagara/Source/Niagara/Private/DataInterface/NiagaraDataInterfaceDynamicMesh.cpp b/Engine/Plugins/FX/Niagara/Source/Niagara/Private/DataInterface/NiagaraDataInterfaceDynamicMesh.cpp
--- a/Engine/Plugins/FX/Niagara/Source/Niagara/Private/DataInterface/NiagaraDataInterfaceDynamicMesh.cpp
+++ b/Engine/Plugins/FX/Niagara/Source/Niagara/Private/DataInterface/NiagaraDataInterfaceDynamicMesh.cpp
@@ -16,6 +16,8 @@
 
 #define LOCTEXT_NAMESPACE "UNiagaraDataInterfaceDynamicMesh"
 
+//-OPT: Allocating triangles doesn't need to lock / unlock, use atomics / batches to improve performance on both CPU & GPU
+//-OPT: We need to be able to pass an GPU generated count for the number of triangles to mesh renderers to improve draw performance
 namespace NDIDynamicMeshLocal
 {
 	BEGIN_SHADER_PARAMETER_STRUCT(FShaderParameters, )
@@ -63,11 +65,55 @@ namespace NDIDynamicMeshLocal
 	static const FName SetVertexColorName("SetVertexColor");
 	static const FName SetVertexDataName("SetVertexData");
 
+	static const FName AppendTriangleName("AppendTriangle");
+
 	static constexpr EPixelFormat PixelFormatPosition = EPixelFormat::PF_R32_FLOAT;
 	static constexpr EPixelFormat PixelFormatTangentBasis = EPixelFormat::PF_R8G8B8A8_SNORM;
 	static constexpr EPixelFormat PixelFormatTexCoord = EPixelFormat::PF_G32R32F;
 	static constexpr EPixelFormat PixelFormatColor = EPixelFormat::PF_R8G8B8A8;
 
+	struct FVertexData
+	{
+		FVector3f		Position = FVector3f::ZeroVector;
+		FVector3f		TangentX = FVector3f::ZeroVector;
+		FVector3f		TangentY = FVector3f::ZeroVector;
+		FVector3f		TangentZ = FVector3f::ZeroVector;
+		FVector2f		TexCoord = FVector2f::ZeroVector;
+		FLinearColor	Color = FLinearColor::White;
+	};
+
+	struct FVMVertexInput
+	{
+		explicit FVMVertexInput(FVectorVMExternalFunctionContext& Context)
+			: Position(Context)
+			, TangentX(Context)
+			, TangentY(Context)
+			, TangentZ(Context)
+			, TexCoord(Context)
+			, Color(Context)
+		{
+		}
+
+		[[nodiscard]] FVertexData GetAndAdvance()
+		{
+			FVertexData Data;
+			Data.Position = Position.GetAndAdvance();
+			Data.TangentX = TangentX.GetAndAdvance();
+			Data.TangentY = TangentY.GetAndAdvance();
+			Data.TangentZ = TangentZ.GetAndAdvance();
+			Data.TexCoord = TexCoord.GetAndAdvance();
+			Data.Color = Color.GetAndAdvance();
+			return Data;
+		}
+
+		FNDIInputParam<FVector3f>		Position;
+		FNDIInputParam<FVector3f>		TangentX;
+		FNDIInputParam<FVector3f>		TangentY;
+		FNDIInputParam<FVector3f>		TangentZ;
+		FNDIInputParam<FVector2f>		TexCoord;
+		FNDIInputParam<FLinearColor>	Color;
+	};
+
 	struct FNDIInstanceMaterial
 	{
 		TObjectPtr<UMaterialInterface> Material = nullptr;
@@ -275,6 +321,13 @@ namespace NDIDynamicMeshLocal
 			FMemory::Memcpy(VertexData.GetData(), InstanceData.VertexData.GetData(), VertexData.Num());
 
 			MeshSections = InstanceData.Sections;
+
+			// Clamp material indices to avoid user error, we always have at least 1 material even if none are specified
+			const int32 NumMaterials = FMath::Max(InstanceData.Materials.Num() - 1, 0);
+			for (FNDIInstanceSection& Section : MeshSections)
+			{
+				Section.MaterialIndex = FMath::Clamp(Section.MaterialIndex, 0, NumMaterials);
+			}
 		}
 
 		bool			bGpuUsesDynamicAllocation = false;
@@ -607,13 +660,16 @@ namespace NDIDynamicMeshLocal
 					OutMaterials.Add(InstanceMaterial.Material);
 				}
 			}
+			if (OutMaterials.Num() == 0)
+			{
+				OutMaterials.Add(nullptr);
+			}
 		}
 		// INiagaraRenderableMesh Impl End
 
 		const FNDIProxy*				OwnerProxy = nullptr;
 		const FNiagaraSystemInstanceID	SystemInstanceID;
 	};
-
 }
 
 ////////////////////////////////////////////////////////////////////////////////////////////////////////////////////////////////////
@@ -863,6 +919,37 @@ void UNiagaraDataInterfaceDynamicMesh::GetFunctions(TArray<FNiagaraFunctionSigna
 		Signature.Inputs.Emplace(FNiagaraTypeDefinition::GetVec2Def(), TEXT("TexCoord"));
 		Signature.Inputs.Emplace(FNiagaraTypeDefinition::GetColorDef(), TEXT("Color"));
 	}
+
+	{
+		FNiagaraFunctionSignature& Signature = OutFunctions.Add_GetRef(MutableSignature);
+		Signature.Name = AppendTriangleName;
+		Signature.Inputs.Emplace_GetRef(FNiagaraTypeDefinition::GetBoolDef(), TEXT("Execute")).SetValue(true);
+		Signature.Inputs.Emplace(FNiagaraTypeDefinition::GetIntDef(), TEXT("SectionIndex"));
+		Signature.Inputs.Emplace(FNiagaraTypeDefinition::GetVec3Def(), TEXT("Position0"));
+		Signature.Inputs.Emplace(FNiagaraTypeDefinition::GetVec3Def(), TEXT("TangentX0"));
+		Signature.Inputs.Emplace(FNiagaraTypeDefinition::GetVec3Def(), TEXT("TangentY0"));
+		Signature.Inputs.Emplace(FNiagaraTypeDefinition::GetVec3Def(), TEXT("TangentZ0"));
+		Signature.Inputs.Emplace(FNiagaraTypeDefinition::GetVec2Def(), TEXT("TexCoord0"));
+		Signature.Inputs.Emplace(FNiagaraTypeDefinition::GetColorDef(), TEXT("Color0"));
+
+		Signature.Inputs.Emplace(FNiagaraTypeDefinition::GetVec3Def(), TEXT("Position1"));
+		Signature.Inputs.Emplace(FNiagaraTypeDefinition::GetVec3Def(), TEXT("TangentX1"));
+		Signature.Inputs.Emplace(FNiagaraTypeDefinition::GetVec3Def(), TEXT("TangentY1"));
+		Signature.Inputs.Emplace(FNiagaraTypeDefinition::GetVec3Def(), TEXT("TangentZ1"));
+		Signature.Inputs.Emplace(FNiagaraTypeDefinition::GetVec2Def(), TEXT("TexCoord1"));
+		Signature.Inputs.Emplace(FNiagaraTypeDefinition::GetColorDef(), TEXT("Color1"));
+
+		Signature.Inputs.Emplace(FNiagaraTypeDefinition::GetVec3Def(), TEXT("Position2"));
+		Signature.Inputs.Emplace(FNiagaraTypeDefinition::GetVec3Def(), TEXT("TangentX2"));
+		Signature.Inputs.Emplace(FNiagaraTypeDefinition::GetVec3Def(), TEXT("TangentY2"));
+		Signature.Inputs.Emplace(FNiagaraTypeDefinition::GetVec3Def(), TEXT("TangentZ2"));
+		Signature.Inputs.Emplace(FNiagaraTypeDefinition::GetVec2Def(), TEXT("TexCoord2"));
+		Signature.Inputs.Emplace(FNiagaraTypeDefinition::GetColorDef(), TEXT("Color2"));
+
+		Signature.Outputs.Emplace(FNiagaraTypeDefinition::GetIntDef(), TEXT("TriangleIndex"));
+
+		Signature.SetDescription(LOCTEXT("AppendTriangleDesc", "Appends a triangle to the section.  This assumes that triangles are stored sequentially in the section, no vertex sharing, etc."));
+	}
 }
 
 void UNiagaraDataInterfaceDynamicMesh::GetVMExternalFunction(const FVMExternalFunctionBindingInfo& BindingInfo, void* InstanceData, FVMExternalFunction& OutFunc)
@@ -895,6 +982,8 @@ void UNiagaraDataInterfaceDynamicMesh::GetVMExternalFunction(const FVMExternalFu
 		{SetVertexTexCoordName, FVMExternalFunction::CreateStatic(UNiagaraDataInterfaceDynamicMesh::VMSetVertexTexCoord)},
 		{SetVertexColorName, FVMExternalFunction::CreateStatic(UNiagaraDataInterfaceDynamicMesh::VMSetVertexColor)},
 		{SetVertexDataName, FVMExternalFunction::CreateStatic(UNiagaraDataInterfaceDynamicMesh::VMSetVertexData)},
+
+		{AppendTriangleName, FVMExternalFunction::CreateStatic(UNiagaraDataInterfaceDynamicMesh::VMAppendTriangle)},
 	};
 
 	for (const auto& StaticBinding : StaticBindings)
@@ -1099,6 +1188,7 @@ bool UNiagaraDataInterfaceDynamicMesh::GetFunctionHLSL(const FNiagaraDataInterfa
 		SetVertexTexCoordName,
 		SetVertexColorName,
 		SetVertexDataName,
+		AppendTriangleName,
 	};
 
 	return ValidGpuFunctions.Contains(FunctionInfo.DefinitionName);
@@ -1787,4 +1877,75 @@ void UNiagaraDataInterfaceDynamicMesh::VMSetVertexData(FVectorVMExternalFunction
 	}
 }
 
+void UNiagaraDataInterfaceDynamicMesh::VMAppendTriangle(FVectorVMExternalFunctionContext& Context)
+{
+	using namespace NDIDynamicMeshLocal;
+
+	VectorVM::FUserPtrHandler<FNDIInstanceData_GameThread> InstanceData(Context);
+	FNDIInputParam<bool>	InExecute(Context);
+	FNDIInputParam<int32>	InSectionIndex(Context);
+	FVMVertexInput			InVertex0(Context);
+	FVMVertexInput			InVertex1(Context);
+	FVMVertexInput			InVertex2(Context);
+	FNDIOutputParam<int32>	OutTriangleIndex(Context);
+
+	InstanceData->ModifyCpuData();
+
+	uint32* IndexBuffer = reinterpret_cast<uint32*>(InstanceData->IndexData.GetData());
+	const int32 NumVertices = InstanceData->NumVertices;
+	for (int32 i = 0; i < Context.GetNumInstances(); ++i)
+	{
+		const bool bExecute = InExecute.GetAndAdvance();
+		const int32 SectionIndex = InSectionIndex.GetAndAdvance();
+		const FVertexData Vertex0 = InVertex0.GetAndAdvance();
+		const FVertexData Vertex1 = InVertex1.GetAndAdvance();
+		const FVertexData Vertex2 = InVertex2.GetAndAdvance();
+		if (!bExecute || !InstanceData->Sections.IsValidIndex(SectionIndex))
+		{
+			continue;
+		}
+		int32 TriangleIndex = INDEX_NONE;
+		{
+			FScopeLock DataGuard(&InstanceData->CpuDataGuard);
+			FNDIInstanceSection& Section = InstanceData->Sections[SectionIndex];
+			if (Section.AllocatedTriangles >= Section.MaxTriangles)
+			{
+				OutTriangleIndex.SetAndAdvance(TriangleIndex);
+				continue;
+			}
+			TriangleIndex = Section.AllocatedTriangles++;
+			OutTriangleIndex.SetAndAdvance(TriangleIndex);
+		}
+
+		IndexBuffer[TriangleIndex * 3 + 0] = TriangleIndex * 3 + 0;
+		IndexBuffer[TriangleIndex * 3 + 1] = TriangleIndex * 3 + 1;
+		IndexBuffer[TriangleIndex * 3 + 2] = TriangleIndex * 3 + 2;
+
+		if (InstanceData->PositionOffset != INDEX_NONE)
+		{
+			InstanceData->SetVertexPosition(TriangleIndex * 3 + 0, Vertex0.Position);
+			InstanceData->SetVertexPosition(TriangleIndex * 3 + 1, Vertex1.Position);
+			InstanceData->SetVertexPosition(TriangleIndex * 3 + 2, Vertex2.Position);
+		}
+		if (InstanceData->TangentBasisOffset != INDEX_NONE)
+		{
+			InstanceData->SetVertexTangentBasis(TriangleIndex * 3 + 0, Vertex0.TangentX, Vertex0.TangentY, Vertex0.TangentZ);
+			InstanceData->SetVertexTangentBasis(TriangleIndex * 3 + 1, Vertex1.TangentX, Vertex1.TangentY, Vertex1.TangentZ);
+			InstanceData->SetVertexTangentBasis(TriangleIndex * 3 + 2, Vertex2.TangentX, Vertex2.TangentY, Vertex2.TangentZ);
+		}
+		if (InstanceData->TexCoordOffset != INDEX_NONE)
+		{
+			InstanceData->SetVertexTexCoord(TriangleIndex * 3 + 0, 0, Vertex0.TexCoord);
+			InstanceData->SetVertexTexCoord(TriangleIndex * 3 + 1, 0, Vertex1.TexCoord);
+			InstanceData->SetVertexTexCoord(TriangleIndex * 3 + 2, 0, Vertex2.TexCoord);
+		}
+		if (InstanceData->ColorOffset != INDEX_NONE)
+		{
+			InstanceData->SetVertexColor(TriangleIndex * 3 + 0, Vertex0.Color);
+			InstanceData->SetVertexColor(TriangleIndex * 3 + 1, Vertex1.Color);
+			InstanceData->SetVertexColor(TriangleIndex * 3 + 2, Vertex2.Color);
+		}
+	}
+}
+
 #undef LOCTEXT_NAMESPACE
diff --git a/Engine/Plugins/FX/Niagara/Source/Niagara/Private/DataInterface/NiagaraDataInterfaceDynamicMesh.h b/Engine/Plugins/FX/Niagara/Source/Niagara/Private/DataInterface/NiagaraDataInterfaceDynamicMesh.h
--- a/Engine/Plugins/FX/Niagara/Source/Niagara/Private/DataInterface/NiagaraDataInterfaceDynamicMesh.h
+++ b/Engine/Plugins/FX/Niagara/Source/Niagara/Private/DataInterface/NiagaraDataInterfaceDynamicMesh.h
@@ -155,4 +155,6 @@ public:
 	static void VMSetVertexTexCoord(FVectorVMExternalFunctionContext& Context);
 	static void VMSetVertexColor(FVectorVMExternalFunctionContext& Context);
 	static void VMSetVertexData(FVectorVMExternalFunctionContext& Context);
+
+	static void VMAppendTriangle(FVectorVMExternalFunctionContext& Context);
 };
