Removed TIsSame from NiagaraDataInterfaceArrayImpl.

diff --git a/Engine/Plugins/FX/Niagara/Source/Niagara/Classes/NiagaraDataInterfaceArrayImpl.h b/Engine/Plugins/FX/Niagara/Source/Niagara/Classes/NiagaraDataInterfaceArrayImpl.h
--- a/Engine/Plugins/FX/Niagara/Source/Niagara/Classes/NiagaraDataInterfaceArrayImpl.h
+++ b/Engine/Plugins/FX/Niagara/Source/Niagara/Classes/NiagaraDataInterfaceArrayImpl.h
@@ -81,29 +81,22 @@
 				TypedOther->MEMBERNAME == MEMBERNAME; \
 		} \
 		template<typename TFromArrayType> \
-		typename TEnableIf<TIsSame<TFromArrayType, decltype(MEMBERNAME)::ElementType>::Value, void>::Type SetVariantArrayData(TConstArrayView<TFromArrayType> InArrayData) \
+		void SetVariantArrayData(TConstArrayView<TFromArrayType> InArrayData) \
 		{ \
-			MEMBERNAME = InArrayData; \
-			GetProxyAs<FProxyType>()->template SetArrayData<decltype(MEMBERNAME)::ElementType>(InArrayData); \
-		} \
-		template<typename TFromArrayType> \
-		typename TEnableIf<!TIsSame<TFromArrayType, decltype(MEMBERNAME)::ElementType>::Value, void>::Type SetVariantArrayData(TConstArrayView<TFromArrayType> InArrayData) \
-		{ \
-			MEMBERNAME.SetNumUninitialized(InArrayData.Num()); \
-			FNDIArrayImplHelper<TYPENAME>::CopyCpuToCpuMemory(MEMBERNAME.GetData(), InArrayData.GetData(), InArrayData.Num()); \
-			GetProxyAs<FProxyType>()->template SetArrayData<decltype(Internal##MEMBERNAME)::ElementType>(InArrayData); \
-		} \
-		template<typename TFromArrayType> \
-		typename TEnableIf<TIsSame<TFromArrayType, decltype(MEMBERNAME)::ElementType>::Value, void>::Type SetVariantArrayValue(int Index, const TFromArrayType& Value, bool bSizeToFit) \
-		{ \
-			const int NumRequired = Index + 1 - MEMBERNAME.Num(); \
-			if ( NumRequired > 0 && !bSizeToFit ) return; \
-			MEMBERNAME.AddDefaulted(FMath::Max(NumRequired, 0)); \
-			MEMBERNAME[Index] = Value; \
-			GetProxyAs<FProxyType>()->template SetArrayData<decltype(MEMBERNAME)::ElementType>(MEMBERNAME); \
+			if constexpr (std::is_same_v<TFromArrayType, decltype(MEMBERNAME)::ElementType>) \
+			{ \
+				MEMBERNAME = InArrayData; \
+				GetProxyAs<FProxyType>()->template SetArrayData<decltype(MEMBERNAME)::ElementType>(InArrayData); \
+			} \
+			else \
+			{ \
+				MEMBERNAME.SetNumUninitialized(InArrayData.Num()); \
+				FNDIArrayImplHelper<TYPENAME>::CopyCpuToCpuMemory(MEMBERNAME.GetData(), InArrayData.GetData(), InArrayData.Num()); \
+				GetProxyAs<FProxyType>()->template SetArrayData<decltype(Internal##MEMBERNAME)::ElementType>(InArrayData); \
+			} \
 		} \
 		template<typename TFromArrayType> \
-		typename TEnableIf<!TIsSame<TFromArrayType, decltype(MEMBERNAME)::ElementType>::Value, void>::Type SetVariantArrayValue(int Index, const TFromArrayType& Value, bool bSizeToFit) \
+		void SetVariantArrayValue(int Index, const TFromArrayType& Value, bool bSizeToFit) \
 		{ \
 			const int NumRequired = Index + 1 - MEMBERNAME.Num(); \
 			if ( NumRequired > 0 && !bSizeToFit ) return; \
