- Move LWC reset from instance tick to system tick
- The instance tick can be called from various locations where it may be unsafe to reset, calling from a known safe location in the system simulation is better

diff --git a/Engine/Plugins/FX/Niagara/Source/Niagara/Private/NiagaraSystemInstance.cpp b/Engine/Plugins/FX/Niagara/Source/Niagara/Private/NiagaraSystemInstance.cpp
--- a/Engine/Plugins/FX/Niagara/Source/Niagara/Private/NiagaraSystemInstance.cpp
+++ b/Engine/Plugins/FX/Niagara/Source/Niagara/Private/NiagaraSystemInstance.cpp
@@ -2335,21 +2335,6 @@ void FNiagaraSystemInstance::Tick_GameThread(float DeltaSeconds)
 		return;
 	}
 
-	// Has the actor position changed to the point where we need to reset the LWC tile
-	if (GetSystem()->SupportsLargeWorldCoordinates())
-	{
-		if (USceneComponent* SceneComponent = AttachComponent.Get())
-		{
-			if (UFXSystemComponent::RequiresLWCTileRecache(LWCTile, SceneComponent->GetComponentLocation()))
-			{
-				//-OPT: For safety we reset everything, but if everything is local space we may not need to, or we could rebase.
-				UE_LOG(LogNiagara, Warning, TEXT("NiagaraComponent(%s - %s) required LWC tile recache and was reset."), *GetFullNameSafe(SceneComponent), *GetFullNameSafe(Asset.Get()));
-				Reset(EResetMode::ResetAll);
-				return;
-			}
-		}
-	}
-
 	CachedDeltaSeconds = DeltaSeconds;
 	FixedBounds_CNC = FixedBounds_GT;
 
diff --git a/Engine/Plugins/FX/Niagara/Source/Niagara/Private/NiagaraSystemSimulation.cpp b/Engine/Plugins/FX/Niagara/Source/Niagara/Private/NiagaraSystemSimulation.cpp
--- a/Engine/Plugins/FX/Niagara/Source/Niagara/Private/NiagaraSystemSimulation.cpp
+++ b/Engine/Plugins/FX/Niagara/Source/Niagara/Private/NiagaraSystemSimulation.cpp
@@ -1221,6 +1221,8 @@ void FNiagaraSystemSimulation::Tick_GameThread_Internal(float DeltaSeconds, cons
 	const bool bUpdateTickGroups = !bIsSolo;
 
 	// Update instances
+	const bool bSupportsLargeWorldCoordinates = System->SupportsLargeWorldCoordinates();
+
 	int32 SystemIndex = 0;
 	FNiagaraWorldManager* WorldManager = FNiagaraWorldManager::Get(World);
 	check(WorldManager != nullptr);
@@ -1254,10 +1256,28 @@ void FNiagaraSystemSimulation::Tick_GameThread_Internal(float DeltaSeconds, cons
 
 		// Perform instance tick
 		Instance->Tick_GameThread(DeltaSeconds);
+
 #if NIAGARA_SYSTEMSIMULATION_DEBUGGING
 		NiagaraSystemSimulationLocal::DebugKillInstanceOnTick(Instance);
 #endif
 
+		// Has the actor position changed to the point where we need to reset the LWC tile
+		if (bSupportsLargeWorldCoordinates)
+		{
+			if ( SystemInstances.IsValidIndex(SystemIndex) && (SystemInstances[SystemIndex] == Instance) )
+			{
+				if (USceneComponent* SceneComponent = Instance->GetAttachComponent())
+				{
+					if (UFXSystemComponent::RequiresLWCTileRecache(Instance->GetLWCTile(), SceneComponent->GetComponentLocation()))
+					{
+						//-OPT: For safety we reset everything, but if everything is local space we may not need to, or we could rebase.
+						UE_LOG(LogNiagara, Warning, TEXT("NiagaraComponent(%s - %s) required LWC tile recache and was reset."), *GetFullNameSafe(SceneComponent), *GetFullNameSafe(System));
+						Instance->Reset(FNiagaraSystemInstance::EResetMode::ResetAll);
+					}
+				}
+			}
+		}
+
 		// Ticking the instance can result in it being removed, completing + reactivating or transferring
 		if (SystemInstances.IsValidIndex(SystemIndex) && (SystemInstances[SystemIndex] == Instance))
 		{
