- Add element type to direct dispatch
- This allows the user to opt out of clipping or launch group sized dispatches rather than having to specify number of threads

diff --git a/Engine/Plugins/FX/Niagara/Source/Niagara/Private/NiagaraGpuComputeDispatch.cpp b/Engine/Plugins/FX/Niagara/Source/Niagara/Private/NiagaraGpuComputeDispatch.cpp
--- a/Engine/Plugins/FX/Niagara/Source/Niagara/Private/NiagaraGpuComputeDispatch.cpp
+++ b/Engine/Plugins/FX/Niagara/Source/Niagara/Private/NiagaraGpuComputeDispatch.cpp
@@ -1449,6 +1449,12 @@ void FNiagaraGpuComputeDispatch::DispatchStage(FRDGBuilder& GraphBuilder, const
 					UE_LOG(LogNiagara, Fatal, TEXT("FNiagaraGpuComputeDispatch: Unknown DispatchType(%d)"), DispatchType);
 					break;
 			}
+
+			// If we are dispatch groups, not threads, then we need to * by the thread group size
+			if (SimStageData.StageMetaData->GpuDirectDispatchElementType == ENiagaraDirectDispatchElementType::NumGroups)
+			{
+				DispatchCount *= DispatchNumThreads;
+			}
 			break;
 		}
 
diff --git a/Engine/Plugins/FX/Niagara/Source/Niagara/Private/NiagaraSimulationStageBase.cpp b/Engine/Plugins/FX/Niagara/Source/Niagara/Private/NiagaraSimulationStageBase.cpp
--- a/Engine/Plugins/FX/Niagara/Source/Niagara/Private/NiagaraSimulationStageBase.cpp
+++ b/Engine/Plugins/FX/Niagara/Source/Niagara/Private/NiagaraSimulationStageBase.cpp
@@ -97,6 +97,7 @@ bool UNiagaraSimulationStageGeneric::AppendCompileHash(FNiagaraCompileHashVisito
 	InVisitor->UpdateString(TEXT("ParticleIterationStateRange"), FString::Printf(TEXT("%d,%d"), ParticleIterationStateRange.X, ParticleIterationStateRange.Y));
 	InVisitor->UpdatePOD(TEXT("bGpuDispatchForceLinear"), bGpuDispatchForceLinear ? 1 : 0);
 	InVisitor->UpdatePOD(TEXT("DirectDispatchType"), (int32)DirectDispatchType);
+	InVisitor->UpdatePOD(TEXT("DirectDispatchElementType"), int32(DirectDispatchElementType));
 	InVisitor->UpdatePOD(TEXT("bOverrideGpuDispatchNumThreads"), bOverrideGpuDispatchNumThreads ? 1 : 0);
 	InVisitor->UpdateString(TEXT("OverrideGpuDispatchNumThreads"), FString::Printf(TEXT("%d,%d,%d"), OverrideGpuDispatchNumThreads.X, OverrideGpuDispatchNumThreads.Y, OverrideGpuDispatchNumThreads.Z));
 
@@ -123,6 +124,7 @@ bool UNiagaraSimulationStageGeneric::FillCompilationData(TArray<FNiagaraSimulati
 	SimStageData.ParticleIterationStateRange = ParticleIterationStateRange;
 	SimStageData.bGpuDispatchForceLinear = bGpuDispatchForceLinear;
 	SimStageData.DirectDispatchType = DirectDispatchType;
+	SimStageData.DirectDispatchElementType = DirectDispatchElementType;
 	SimStageData.bOverrideGpuDispatchNumThreads = bOverrideGpuDispatchNumThreads;
 	SimStageData.OverrideGpuDispatchNumThreads = OverrideGpuDispatchNumThreads;
 
@@ -238,6 +240,7 @@ void UNiagaraSimulationStageGeneric::PostEditChangeProperty(struct FPropertyChan
 		GET_MEMBER_NAME_CHECKED(UNiagaraSimulationStageGeneric, ParticleIterationStateBinding),
 		GET_MEMBER_NAME_CHECKED(UNiagaraSimulationStageGeneric, bGpuDispatchForceLinear),
 		GET_MEMBER_NAME_CHECKED(UNiagaraSimulationStageGeneric, DirectDispatchType),
+		GET_MEMBER_NAME_CHECKED(UNiagaraSimulationStageGeneric, DirectDispatchElementType),
 		GET_MEMBER_NAME_CHECKED(UNiagaraSimulationStageGeneric, bOverrideGpuDispatchNumThreads),
 	};
 
diff --git a/Engine/Plugins/FX/Niagara/Source/Niagara/Public/NiagaraSimulationStageBase.h b/Engine/Plugins/FX/Niagara/Source/Niagara/Public/NiagaraSimulationStageBase.h
--- a/Engine/Plugins/FX/Niagara/Source/Niagara/Public/NiagaraSimulationStageBase.h
+++ b/Engine/Plugins/FX/Niagara/Source/Niagara/Public/NiagaraSimulationStageBase.h
@@ -135,6 +135,10 @@ public:
 	UPROPERTY(EditAnywhere, Category = "Simulation Stage")
 	ENiagaraGpuDispatchType DirectDispatchType = ENiagaraGpuDispatchType::OneD;
 
+	/** Customizes what the element count means for a direct dispatch. */
+	UPROPERTY(EditAnywhere, Category = "Simulation Stage")
+	ENiagaraDirectDispatchElementType DirectDispatchElementType = ENiagaraDirectDispatchElementType::NumThreads;
+
 	/**
 	Integer binding to override the number of elements the stage will execute along X.
 	For example, if you want to iterate over a custom source such as triangles on a mesh you can
diff --git a/Engine/Plugins/FX/Niagara/Source/Niagara/Public/NiagaraSimulationStageCompileData.h b/Engine/Plugins/FX/Niagara/Source/Niagara/Public/NiagaraSimulationStageCompileData.h
--- a/Engine/Plugins/FX/Niagara/Source/Niagara/Public/NiagaraSimulationStageCompileData.h
+++ b/Engine/Plugins/FX/Niagara/Source/Niagara/Public/NiagaraSimulationStageCompileData.h
@@ -24,6 +24,7 @@ struct FNiagaraSimulationStageCompilationData
 	FIntPoint                       ParticleIterationStateRange = FIntPoint::ZeroValue;
 	bool                            bGpuDispatchForceLinear = false;
 	ENiagaraGpuDispatchType         DirectDispatchType = ENiagaraGpuDispatchType::OneD;
+	ENiagaraDirectDispatchElementType DirectDispatchElementType = ENiagaraDirectDispatchElementType::NumThreads;
 	bool                            bOverrideGpuDispatchNumThreads = false;
 	FIntVector                      OverrideGpuDispatchNumThreads = FIntVector(1, 1, 1);
 };
diff --git a/Engine/Plugins/FX/Niagara/Source/NiagaraEditor/Private/Customizations/NiagaraSimulationStageCustomization.cpp b/Engine/Plugins/FX/Niagara/Source/NiagaraEditor/Private/Customizations/NiagaraSimulationStageCustomization.cpp
--- a/Engine/Plugins/FX/Niagara/Source/NiagaraEditor/Private/Customizations/NiagaraSimulationStageCustomization.cpp
+++ b/Engine/Plugins/FX/Niagara/Source/NiagaraEditor/Private/Customizations/NiagaraSimulationStageCustomization.cpp
@@ -117,6 +117,7 @@ void FNiagaraSimulationStageGenericCustomization::CustomizeDetails(class IDetail
 		case ENiagaraIterationSource::DirectSet:
 		{
 			DirectSetStageCategory.AddProperty(DetailBuilder.GetProperty(GET_MEMBER_NAME_CHECKED(UNiagaraSimulationStageGeneric, DirectDispatchType)));
+			DirectSetStageCategory.AddProperty(DetailBuilder.GetProperty(GET_MEMBER_NAME_CHECKED(UNiagaraSimulationStageGeneric, DirectDispatchElementType)));
 			DirectSetStageCategory.AddProperty(DetailBuilder.GetProperty(GET_MEMBER_NAME_CHECKED(UNiagaraSimulationStageGeneric, ElementCountXBinding)));
 			DirectSetStageCategory.AddProperty(DetailBuilder.GetProperty(GET_MEMBER_NAME_CHECKED(UNiagaraSimulationStageGeneric, ElementCountYBinding)));
 			DirectSetStageCategory.AddProperty(DetailBuilder.GetProperty(GET_MEMBER_NAME_CHECKED(UNiagaraSimulationStageGeneric, ElementCountZBinding)));
diff --git a/Engine/Plugins/FX/Niagara/Source/NiagaraEditor/Private/NiagaraHlslTranslator.cpp b/Engine/Plugins/FX/Niagara/Source/NiagaraEditor/Private/NiagaraHlslTranslator.cpp
--- a/Engine/Plugins/FX/Niagara/Source/NiagaraEditor/Private/NiagaraHlslTranslator.cpp
+++ b/Engine/Plugins/FX/Niagara/Source/NiagaraEditor/Private/NiagaraHlslTranslator.cpp
@@ -1240,6 +1240,7 @@ const FNiagaraTranslateResults &FHlslNiagaraTranslator::Translate(const FNiagara
 			SimulationStageMetaData.bWritesParticles = true;
 			SimulationStageMetaData.bPartialParticleUpdate = false;
 			SimulationStageMetaData.GpuDispatchType = ENiagaraGpuDispatchType::OneD;
+			SimulationStageMetaData.GpuDirectDispatchElementType = ENiagaraDirectDispatchElementType::NumThreads;
 			SimulationStageMetaData.GpuDispatchNumThreads = FNiagaraShader::GetDefaultThreadGroupSize(ENiagaraGpuDispatchType::OneD);
 		}
 
@@ -1310,6 +1311,7 @@ const FNiagaraTranslateResults &FHlslNiagaraTranslator::Translate(const FNiagara
 					TranslationStages[Index].ParticleIterationStateRange = CompileSimStageData.ParticleIterationStateRange;
 					TranslationStages[Index].bGpuDispatchForceLinear = CompileSimStageData.bGpuDispatchForceLinear;
 					TranslationStages[Index].DirectDispatchType = CompileSimStageData.DirectDispatchType;
+					TranslationStages[Index].DirectDispatchElementType = CompileSimStageData.DirectDispatchElementType;
 					TranslationStages[Index].bOverrideGpuDispatchNumThreads = CompileSimStageData.bOverrideGpuDispatchNumThreads;
 					TranslationStages[Index].OverrideGpuDispatchNumThreads = CompileSimStageData.OverrideGpuDispatchNumThreads;
 
@@ -1416,6 +1418,8 @@ const FNiagaraTranslateResults &FHlslNiagaraTranslator::Translate(const FNiagara
 					else if (SimulationStageMetaData.IterationSourceType == ENiagaraIterationSource::DirectSet)
 					{
 						SimulationStageMetaData.GpuDispatchType = TranslationStages[Index].DirectDispatchType;
+						SimulationStageMetaData.GpuDirectDispatchElementType = TranslationStages[Index].DirectDispatchElementType;
+						SimulationStageMetaData.GpuDispatchNumThreads = FNiagaraShader::GetDefaultThreadGroupSize(TranslationStages[Index].DirectDispatchType);
 					}
 
 					// Increment source stage index
@@ -3495,14 +3499,28 @@ void FHlslNiagaraTranslator::DefineMainGPUFunctions(
 		// Iteration Direct Value
 		else if (TranslationStage.IterationSourceType == ENiagaraIterationSource::DirectSet)
 		{
-			HlslOutput += TEXT(
-				"	const uint MaxInstances = SimulationStage_GetInstanceCount();\n"
-				"	const bool bValidInstance = all(GDispatchThreadId < DispatchThreadIdBounds);\n"
-				"	const bool bRunUpdateLogic = bValidInstance;\n"
-				"	const bool bRunSpawnLogic = false;\n"
-				"	GLinearThreadId = bValidInstance ? GLinearThreadId : MaxInstances;\n"
-				"	GSpawnStartInstance = MaxInstances;\n"
-			);
+			if (TranslationStage.DirectDispatchElementType == ENiagaraDirectDispatchElementType::NumThreads)
+			{
+				HlslOutput += TEXT(
+					"	const uint MaxInstances = SimulationStage_GetInstanceCount();\n"
+					"	const bool bValidInstance = all(GDispatchThreadId < DispatchThreadIdBounds);\n"
+					"	const bool bRunUpdateLogic = bValidInstance;\n"
+					"	const bool bRunSpawnLogic = false;\n"
+					"	GLinearThreadId = bValidInstance ? GLinearThreadId : MaxInstances;\n"
+					"	GSpawnStartInstance = MaxInstances;\n"
+				);
+			}
+			else
+			{
+				HlslOutput += TEXT(
+					"	const uint MaxInstances = SimulationStage_GetInstanceCount();\n"
+					"	const bool bValidInstance = all(GDispatchThreadId < DispatchThreadIdBounds);\n"
+					"	const bool bRunUpdateLogic = true;\n"
+					"	const bool bRunSpawnLogic = false;\n"
+					"	GLinearThreadId = bValidInstance ? GLinearThreadId : MaxInstances;\n"
+					"	GSpawnStartInstance = MaxInstances;\n"
+				);
+			}
 		}
 		else
 		{
diff --git a/Engine/Plugins/FX/Niagara/Source/NiagaraEditor/Private/NiagaraHlslTranslator.h b/Engine/Plugins/FX/Niagara/Source/NiagaraEditor/Private/NiagaraHlslTranslator.h
--- a/Engine/Plugins/FX/Niagara/Source/NiagaraEditor/Private/NiagaraHlslTranslator.h
+++ b/Engine/Plugins/FX/Niagara/Source/NiagaraEditor/Private/NiagaraHlslTranslator.h
@@ -389,6 +389,7 @@ public:
 	bool bPartialParticleUpdate = false;
 	bool bGpuDispatchForceLinear = false;
 	ENiagaraGpuDispatchType DirectDispatchType = ENiagaraGpuDispatchType::OneD;
+	ENiagaraDirectDispatchElementType DirectDispatchElementType = ENiagaraDirectDispatchElementType::NumThreads;
 	bool bOverrideGpuDispatchNumThreads = false;
 	bool bShouldUpdateInitialAttributeValues = false;
 	FIntVector OverrideGpuDispatchNumThreads = FIntVector(1, 1, 1);
diff --git a/Engine/Plugins/FX/Niagara/Source/NiagaraShader/Public/NiagaraScriptBase.h b/Engine/Plugins/FX/Niagara/Source/NiagaraShader/Public/NiagaraScriptBase.h
--- a/Engine/Plugins/FX/Niagara/Source/NiagaraShader/Public/NiagaraScriptBase.h
+++ b/Engine/Plugins/FX/Niagara/Source/NiagaraShader/Public/NiagaraScriptBase.h
@@ -25,6 +25,30 @@ enum class ENiagaraGpuDispatchType : uint8
 	Custom			UMETA(Hidden),
 };
 
+UENUM()
+enum class ENiagaraDirectDispatchElementType : uint8
+{
+	/**
+	Number of elements is the number of threads launched in that dimension.
+	Threads that are out of bounds due to shader thread group size will be automatically clipped (i.e. code will not run).
+	For example, if GroupSize = 64,1,1 and NumElements = 32,1,1 only the first 32 threads will run the code.
+	*/
+	NumThreads,
+	/**
+	Number of elements is the number of threads launched in that dimension.
+	Threads that are out of bounds due to shader thread group size will not be clipped and your code will execute.
+	You are responsible for ensuring you do not make invalid access from these OOB threads.
+	For example, if GroupSize = 64,1,1 and NumElements = 32,1,1 you code will execute 64 times.
+	Use this path if you need to add group sync's within you graph code.
+	*/
+	NumThreadsNoClipping,
+	/**
+	Number of elements refers to the number of groups to launch.
+	For example, if you defined NumElements as 3,1,1 and GroupSize was 64,1,1 you are effectively launching 192 threads.
+	*/
+	NumGroups,
+};
+
 UENUM()
 enum class ENiagaraSimStageExecuteBehavior : uint8
 {
@@ -127,6 +151,9 @@ public:
 	UPROPERTY()
 	ENiagaraGpuDispatchType GpuDispatchType;
 
+	UPROPERTY()
+	ENiagaraDirectDispatchElementType GpuDirectDispatchElementType = ENiagaraDirectDispatchElementType::NumThreads;
+
 	/** When in custom mode this is the num threads. */
 	UPROPERTY()
 	FIntVector GpuDispatchNumThreads;
