- Fix missing final transition for dynammic mesh when running on the GPU
- Fix for vertex colors swizzle order PR #9979: [Niagara] Vertex color fixes for DynamicMeshDI (Contributed by [REDACTED_PERSON])

#rnx

diff --git a/Engine/Plugins/FX/Niagara/Shaders/Private/NiagaraDataInterfaceDynamicMeshTemplate.ush b/Engine/Plugins/FX/Niagara/Shaders/Private/NiagaraDataInterfaceDynamicMeshTemplate.ush
--- a/Engine/Plugins/FX/Niagara/Shaders/Private/NiagaraDataInterfaceDynamicMeshTemplate.ush
+++ b/Engine/Plugins/FX/Niagara/Shaders/Private/NiagaraDataInterfaceDynamicMeshTemplate.ush
@@ -68,18 +68,20 @@ float4 {ParameterName}_UnpackColor(uint Packed)
 	Unpacked.y = (Packed >> 8) & 0xff;
 	Unpacked.z = (Packed >> 16) & 0xff;
 	Unpacked.w = (Packed >> 24) & 0xff;
-	return Unpacked * 255.0f;
+	Unpacked = Unpacked FMANUALFETCH_COLOR_COMPONENT_SWIZZLE;
+	return Unpacked / 255.0f;
 }
 
 uint {ParameterName}_PackColor(float4 Color)
 {
+	Color = Color FMANUALFETCH_COLOR_COMPONENT_SWIZZLE;
 	Color = clamp((Color * 255.0f) + 0.5f, 0.0f, 255.0f);
 
 	uint Packed;
-	Packed  = uint(Color.r) << 0;
-	Packed |= uint(Color.g) << 8;
-	Packed |= uint(Color.b) << 16;
-	Packed |= uint(Color.a) << 24;
+	Packed  = uint(Color.x) << 0;
+	Packed |= uint(Color.y) << 8;
+	Packed |= uint(Color.z) << 16;
+	Packed |= uint(Color.w) << 24;
 	return Packed;
 }
 
diff --git a/Engine/Plugins/FX/Niagara/Source/Niagara/Private/DataInterface/NiagaraDataInterfaceDynamicMesh.cpp b/Engine/Plugins/FX/Niagara/Source/Niagara/Private/DataInterface/NiagaraDataInterfaceDynamicMesh.cpp
--- a/Engine/Plugins/FX/Niagara/Source/Niagara/Private/DataInterface/NiagaraDataInterfaceDynamicMesh.cpp
+++ b/Engine/Plugins/FX/Niagara/Source/Niagara/Private/DataInterface/NiagaraDataInterfaceDynamicMesh.cpp
@@ -377,6 +377,13 @@ namespace NDIDynamicMeshLocal
 		FShaderResourceViewRHIRef	VertexBufferTexCoordSRV;
 		FShaderResourceViewRHIRef	VertexBufferColorSRV;
 
+		FRDGBufferRef				RDGTransientSectionBuffer = nullptr;
+		FRDGBufferUAVRef			RDGTransientSectionBufferUAV = nullptr;
+		FRDGBufferRef				RDGTransientIndexBuffer = nullptr;
+		FRDGBufferUAVRef			RDGTransientIndexBufferUAV = nullptr;
+		FRDGBufferRef				RDGTransientVertexBuffer = nullptr;
+		FRDGBufferUAVRef			RDGTransientVertexBufferUAV = nullptr;
+
 		bool						bGpuUsesDynamicAllocation = true;
 		bool						bNeedsClearIndexBuffer = true;
 		bool						bNeedsGpuSectionUpdate = false;
@@ -514,24 +521,40 @@ namespace NDIDynamicMeshLocal
 		virtual void ConsumePerInstanceDataFromGameThread(void* PerInstanceData, const FNiagaraSystemInstanceID& Instance) override { checkNoEntry(); }
 		virtual int32 PerInstanceDataPassedToRenderThreadSize() const override { return 0; }
 
-		virtual void PreStage(const FNDIGpuComputePreStageContext& Context)
+		virtual void PreStage(const FNDIGpuComputePreStageContext& Context) override
 		{
 			FNDIInstanceData_RenderThread& InstanceData = InstanceData_RT.FindChecked(Context.GetSystemInstanceID());
 			const FNiagaraSimStageData& SimStageData = Context.GetSimStageData();
 			if (SimStageData.bFirstStage)
 			{
 				FRDGBuilder& GraphBuilder = Context.GetGraphBuilder();
+
+				// Do we need to register our graph resources
+				// Note: We are assuming they will be used we could possibly defer until later on
+				if (InstanceData.RDGTransientSectionBuffer == nullptr)
+				{
+					InstanceData.RDGTransientSectionBuffer		= GraphBuilder.RegisterExternalBuffer(InstanceData.SectionPooledBuffer);
+					InstanceData.RDGTransientSectionBufferUAV	= GraphBuilder.CreateUAV(InstanceData.RDGTransientSectionBuffer, EPixelFormat::PF_R32_UINT);
+					InstanceData.RDGTransientIndexBuffer		= GraphBuilder.RegisterExternalBuffer(InstanceData.IndexPooledBuffer);
+					InstanceData.RDGTransientIndexBufferUAV		= GraphBuilder.CreateUAV(InstanceData.RDGTransientIndexBuffer, EPixelFormat::PF_R32_UINT);
+					InstanceData.RDGTransientVertexBuffer		= GraphBuilder.RegisterExternalBuffer(InstanceData.VertexPooledBuffer);
+					InstanceData.RDGTransientVertexBufferUAV	= GraphBuilder.CreateUAV(InstanceData.RDGTransientVertexBuffer, EPixelFormat::PF_R32_UINT);
+
+					Context.GetRDGExternalAccessQueue().Add(InstanceData.RDGTransientIndexBuffer, ERHIAccess::VertexOrIndexBuffer | ERHIAccess::SRVMask);
+					Context.GetRDGExternalAccessQueue().Add(InstanceData.RDGTransientVertexBuffer, ERHIAccess::VertexOrIndexBuffer | ERHIAccess::SRVMask);
+				}
+
 				if (InstanceData.bNeedsClearIndexBuffer)
 				{
 					InstanceData.bNeedsClearIndexBuffer = false;
-					AddClearUAVPass(Context.GetGraphBuilder(), GraphBuilder.CreateUAV(GraphBuilder.RegisterExternalBuffer(InstanceData.IndexPooledBuffer), EPixelFormat::PF_R32_UINT), 0);
+					AddClearUAVPass(Context.GetGraphBuilder(), InstanceData.RDGTransientIndexBufferUAV, 0);
 				}
 
 				if (InstanceData.bNeedsGpuSectionUpdate && InstanceData.NumSections > 0)
 				{
 					AddCopyBufferPass(
 						GraphBuilder,
-						GraphBuilder.RegisterExternalBuffer(InstanceData.SectionPooledBuffer),
+						InstanceData.RDGTransientSectionBuffer,
 						GraphBuilder.RegisterExternalBuffer(InstanceData.SectionDefaultPooledBuffer)
 					);
 				}
@@ -543,6 +566,20 @@ namespace NDIDynamicMeshLocal
 			}
 		}
 
+		virtual void PostSimulate(const FNDIGpuComputePostSimulateContext& Context) override
+		{
+			if (Context.IsFinalPostSimulate())
+			{
+				FNDIInstanceData_RenderThread& InstanceData = InstanceData_RT.FindChecked(Context.GetSystemInstanceID());
+				InstanceData.RDGTransientSectionBuffer = nullptr;
+				InstanceData.RDGTransientSectionBufferUAV = nullptr;
+				InstanceData.RDGTransientIndexBuffer = nullptr;
+				InstanceData.RDGTransientIndexBufferUAV = nullptr;
+				InstanceData.RDGTransientVertexBuffer = nullptr;
+				InstanceData.RDGTransientVertexBufferUAV = nullptr;
+			}
+		}
+
 		TMap<FNiagaraSystemInstanceID, FNDIInstanceData_GameThread*> InstanceData_GT;
 		TMap<FNiagaraSystemInstanceID, FNDIInstanceData_RenderThread> InstanceData_RT;
 	};
@@ -1219,9 +1256,9 @@ void UNiagaraDataInterfaceDynamicMesh::SetShaderParameters(const FNiagaraDataInt
 	ShaderParameters->TangentBasisOffset	= InstanceData.TangentBasisOffset >> 2;
 	ShaderParameters->TexCoordOffset		= InstanceData.TexCoordOffset >> 2;
 	ShaderParameters->ColorOffset			= InstanceData.ColorOffset >> 2;
-	ShaderParameters->SectionBuffer			= GraphBuilder.CreateUAV(GraphBuilder.RegisterExternalBuffer(InstanceData.SectionPooledBuffer), EPixelFormat::PF_R32_UINT);
-	ShaderParameters->IndexBuffer			= GraphBuilder.CreateUAV(GraphBuilder.RegisterExternalBuffer(InstanceData.IndexPooledBuffer), EPixelFormat::PF_R32_UINT);
-	ShaderParameters->VertexBuffer			= GraphBuilder.CreateUAV(GraphBuilder.RegisterExternalBuffer(InstanceData.VertexPooledBuffer), EPixelFormat::PF_R32_UINT);
+	ShaderParameters->SectionBuffer			= InstanceData.RDGTransientSectionBufferUAV;
+	ShaderParameters->IndexBuffer			= InstanceData.RDGTransientIndexBufferUAV;
+	ShaderParameters->VertexBuffer			= InstanceData.RDGTransientVertexBufferUAV;
 }
 
 FNiagaraRenderableMeshPtr UNiagaraDataInterfaceDynamicMesh::GetRenderableMesh(FNiagaraSystemInstanceID SystemInstanceID)
