- Fix for incorrect static analysis warning

#rnx

diff --git a/Engine/Plugins/FX/Niagara/Source/Niagara/Private/NiagaraScalabilityManager.cpp b/Engine/Plugins/FX/Niagara/Source/Niagara/Private/NiagaraScalabilityManager.cpp
--- a/Engine/Plugins/FX/Niagara/Source/Niagara/Private/NiagaraScalabilityManager.cpp
+++ b/Engine/Plugins/FX/Niagara/Source/Niagara/Private/NiagaraScalabilityManager.cpp
@@ -583,8 +583,9 @@ FNiagaraScalabilitySystemData& FNiagaraScalabilityManager::GetSystemData(int32 C
 		}
 		else
 		{
-			ensure(*SysDataIndex <= TNumericLimits<decltype(CompState.SystemDataIndex)>::Max());
-			CompState.SystemDataIndex = int16(*SysDataIndex);
+			const int32 SysDataIndexValue = *SysDataIndex;
+			ensure(SysDataIndexValue <= TNumericLimits<int16>::Max());
+			CompState.SystemDataIndex = int16(SysDataIndexValue);
 		}
 	}
 
