fix vdb caching

diff --git a/Engine/Plugins/FX/Niagara/Source/Niagara/Private/NiagaraDataInterfaceRenderTargetVolume.cpp b/Engine/Plugins/FX/Niagara/Source/Niagara/Private/NiagaraDataInterfaceRenderTargetVolume.cpp
--- a/Engine/Plugins/FX/Niagara/Source/Niagara/Private/NiagaraDataInterfaceRenderTargetVolume.cpp
+++ b/Engine/Plugins/FX/Niagara/Source/Niagara/Private/NiagaraDataInterfaceRenderTargetVolume.cpp
@@ -675,6 +675,7 @@ bool UNiagaraDataInterfaceRenderTargetVolume::SimCacheWriteFrame(UObject* Storag
 
 				OpenVDBSimCacheData->FrameRangeStart = FMath::Min(FrameIndex, OpenVDBSimCacheData->FrameRangeStart);
 				OpenVDBSimCacheData->FrameRangeEnd = FMath::Max(FrameIndex, OpenVDBSimCacheData->FrameRangeEnd);
+				OpenVDBSimCacheData->Resolution = InstanceData_GT->Size;
 #endif
 			}
 			else
@@ -734,22 +735,23 @@ bool UNiagaraDataInterfaceRenderTargetVolume::SimCacheReadFrame(UObject* Storage
 {
 	FRenderTargetVolumeRWInstanceData_GameThread* InstanceData_GT = reinterpret_cast<FRenderTargetVolumeRWInstanceData_GameThread*>(OptionalPerInstanceData);
 
-	if (NDIRenderTargetVolumeLocal::GSimCacheUseOpenVDB)
+	if (Cast<UVolumeCache>(StorageObject))
 	{
 #if PLATFORM_WINDOWS
 		UVolumeCache *OpenVDBSimCacheData = CastChecked<UVolumeCache>(StorageObject);
 
 		TSharedPtr<FVolumeCacheData> VolumeCacheData = OpenVDBSimCacheData->GetData();
 
-		const int FrameIndex = Interp >= 0.5f ? FrameB : FrameA;
-		OpenVDBSimCacheData->LoadFile(FrameIndex);
-
 		InstanceData_GT->Size = OpenVDBSimCacheData->Resolution;
 		InstanceData_GT->Format = PF_FloatRGBA;
+		InstanceData_GT->Filter = TextureFilter::TF_Default;
 
 		PerInstanceTick(InstanceData_GT, SystemInstance, 0.0f);
 		PerInstanceTickPostSimulate(InstanceData_GT, SystemInstance, 0.0f);
 
+		const int FrameIndex = Interp >= 0.5f ? FrameB : FrameA;
+		OpenVDBSimCacheData->LoadFile(FrameIndex);
+
 		//  write to the volume texture
 		if (InstanceData_GT->TargetTexture)
 		{
diff --git a/Engine/Plugins/FX/Niagara/Source/Niagara/Private/VolumeCache.cpp b/Engine/Plugins/FX/Niagara/Source/Niagara/Private/VolumeCache.cpp
--- a/Engine/Plugins/FX/Niagara/Source/Niagara/Private/VolumeCache.cpp
+++ b/Engine/Plugins/FX/Niagara/Source/Niagara/Private/VolumeCache.cpp
@@ -192,16 +192,22 @@ bool FOpenVDBCacheData::LoadFile(FString Path, int frame)
 			if (DenseResolution == FIntVector(-1, -1, -1))
 			{
 				DenseResolution = Size;
-			} 
+			}
 			else if (DenseResolution != Size)
 			{
 				UE_LOG(LogVolumeCache, Warning, TEXT("Grids are not the same size: %s"), *FullPath);
 				return false;
-			}			
+			}
 
 			OpenVDBGrids.Add(frame, ColorGrid);
 
-			// load to dense buffer			
+			// if dense vdb buffer doesn't match current resolution, re initialize it	
+			if (DenseGridPtr->bbox().dim() != openvdb::Coord(DenseResolution.X - 1, DenseResolution.Y - 1, DenseResolution.Z - 1))
+			{
+				Init(DenseResolution);
+			}
+
+			// load to dense buffer	
 			openvdb::tools::CopyToDense<Vec4Tree, Vec4Dense> Copier(ColorGrid->tree(), *DenseGridPtr);
 			Copier.copy();
 		}
@@ -287,19 +293,17 @@ bool FOpenVDBCacheData::Fill3DTexture(int frame, FTextureRHIRef TextureToFill)
 
 		if (TextureToFill->GetSizeXYZ() != DenseResolution)
 		{
-			UE_LOG(LogVolumeCache, Warning, TEXT("Target texture resolution %s doesn't match OpenVDB grid resolution %s.  Regenerate your cache"), *TextureToFill->GetSizeXYZ().ToString(), *DenseResolution.ToString());
+			UE_LOG(LogVolumeCache, Warning, TEXT("Target texture resolution %s doesn't match OpenVDB grid resolution %s."), *TextureToFill->GetSizeXYZ().ToString(), *DenseResolution.ToString());
 		}
-		else
-		{
-			const FUpdateTextureRegion3D UpdateRegion(0, 0, 0, 0, 0, 0, DenseResolution.X, DenseResolution.Y, DenseResolution.Z);
-			const SIZE_T MemorySize = static_cast<SIZE_T>(UpdateRegion.Width) * static_cast<SIZE_T>(UpdateRegion.Height) * static_cast<SIZE_T>(UpdateRegion.Depth) * static_cast<SIZE_T>(FormatSize);
 
-			FUpdateTexture3DData TheData = RHIBeginUpdateTexture3D(TextureToFill, 0, UpdateRegion);
+		const FUpdateTextureRegion3D UpdateRegion(0, 0, 0, 0, 0, 0, DenseResolution.X, DenseResolution.Y, DenseResolution.Z);
+		const SIZE_T MemorySize = static_cast<SIZE_T>(UpdateRegion.Width) * static_cast<SIZE_T>(UpdateRegion.Height) * static_cast<SIZE_T>(UpdateRegion.Depth) * static_cast<SIZE_T>(FormatSize);
 
-			FMemory::Memcpy(TheData.Data, DataPtr, MemorySize);
+		FUpdateTexture3DData TheData = RHIBeginUpdateTexture3D(TextureToFill, 0, UpdateRegion);
 
-			RHIEndUpdateTexture3D(TheData);
-		}
+		FMemory::Memcpy(TheData.Data, DataPtr, MemorySize);
+
+		RHIEndUpdateTexture3D(TheData);		
 
 		return true;
 	}
