Added new utility functions to access graph data that are used in the summary view rework and made function parameters const

diff --git a/Engine/Plugins/FX/Niagara/Source/NiagaraEditor/Private/NiagaraGraphDataCache.cpp b/Engine/Plugins/FX/Niagara/Source/NiagaraEditor/Private/NiagaraGraphDataCache.cpp
--- a/Engine/Plugins/FX/Niagara/Source/NiagaraEditor/Private/NiagaraGraphDataCache.cpp
+++ b/Engine/Plugins/FX/Niagara/Source/NiagaraEditor/Private/NiagaraGraphDataCache.cpp
@@ -26,7 +26,7 @@ FNiagaraGraphDataCache::FNiagaraGraphDataCache()
 }
 
 void FNiagaraGraphDataCache::GetStackFunctionInputPinsInternal(
-	UNiagaraNodeFunctionCall& FunctionCallNode,
+	const UNiagaraNodeFunctionCall& FunctionCallNode,
 	const UNiagaraGraph* CalledGraph,
 	TConstArrayView<FNiagaraVariable> StaticVars,
 	TArray<const UEdGraphPin*>& OutInputPins,
@@ -96,7 +96,7 @@ void FNiagaraGraphDataCache::GetStackFunctionInputPinsInternal(
 	}
 }
 
-void FNiagaraGraphDataCache::GetStackFunctionInputPins(UNiagaraNodeFunctionCall& FunctionCallNode, TConstArrayView<FNiagaraVariable> StaticVars, TArray<const UEdGraphPin*>& OutInputPins, FCompileConstantResolver ConstantResolver, FNiagaraStackGraphUtilities::ENiagaraGetStackFunctionInputPinsOptions Options, bool bIgnoreDisabled)
+void FNiagaraGraphDataCache::GetStackFunctionInputPins(const UNiagaraNodeFunctionCall& FunctionCallNode, TConstArrayView<FNiagaraVariable> StaticVars, TArray<const UEdGraphPin*>& OutInputPins, FCompileConstantResolver ConstantResolver, FNiagaraStackGraphUtilities::ENiagaraGetStackFunctionInputPinsOptions Options, bool bIgnoreDisabled)
 {
 	if (const UNiagaraGraph* CalledGraph = FunctionCallNode.GetCalledGraph())
 	{
@@ -104,7 +104,7 @@ void FNiagaraGraphDataCache::GetStackFunctionInputPins(UNiagaraNodeFunctionCall&
 	}
 }
 
-void FNiagaraGraphDataCache::GetStackFunctionInputPins(UNiagaraNodeFunctionCall& FunctionCallNode, TConstArrayView<FNiagaraVariable> StaticVars, TArray<const UEdGraphPin*>& OutInputPins, TSet<const UEdGraphPin*>& OutHiddenPins, FCompileConstantResolver ConstantResolver, FNiagaraStackGraphUtilities::ENiagaraGetStackFunctionInputPinsOptions Options, bool bIgnoreDisabled)
+void FNiagaraGraphDataCache::GetStackFunctionInputPins(const UNiagaraNodeFunctionCall& FunctionCallNode, TConstArrayView<FNiagaraVariable> StaticVars, TArray<const UEdGraphPin*>& OutInputPins, TSet<const UEdGraphPin*>& OutHiddenPins, FCompileConstantResolver ConstantResolver, FNiagaraStackGraphUtilities::ENiagaraGetStackFunctionInputPinsOptions Options, bool bIgnoreDisabled)
 {
 	if (const UNiagaraGraph* CalledGraph = FunctionCallNode.GetCalledGraph())
 	{
diff --git a/Engine/Plugins/FX/Niagara/Source/NiagaraEditor/Private/ViewModels/Stack/NiagaraStackGraphUtilities.cpp b/Engine/Plugins/FX/Niagara/Source/NiagaraEditor/Private/ViewModels/Stack/NiagaraStackGraphUtilities.cpp
--- a/Engine/Plugins/FX/Niagara/Source/NiagaraEditor/Private/ViewModels/Stack/NiagaraStackGraphUtilities.cpp
+++ b/Engine/Plugins/FX/Niagara/Source/NiagaraEditor/Private/ViewModels/Stack/NiagaraStackGraphUtilities.cpp
@@ -42,8 +42,9 @@
 #include "INiagaraEditorTypeUtilities.h"
 #include "NiagaraEditorData.h"
 #include "NiagaraGraphDataCache.h"
+#include "NiagaraScriptGraphViewModel.h"
 #include "NiagaraSettings.h"
-
+#include "ViewModels/HierarchyEditor/NiagaraHierarchyViewModelBase.h"
 
 DECLARE_CYCLE_STAT(TEXT("Niagara - StackGraphUtilities - RelayoutGraph"), STAT_NiagaraEditor_StackGraphUtilities_RelayoutGraph, STATGROUP_NiagaraEditor);
 
@@ -52,12 +53,12 @@ DECLARE_CYCLE_STAT(TEXT("Niagara - StackGraphUtilities - RelayoutGraph"), STAT_N
 namespace FNiagaraStackGraphUtilitiesImpl
 {
 
-static FVersionedNiagaraEmitter GetOuterEmitter(UNiagaraNode* Node)
+static FVersionedNiagaraEmitter GetOuterEmitter(const UNiagaraNode* Node)
 {
 	FVersionedNiagaraEmitter Result;
 	if (Node)
 	{
-		if (UNiagaraGraph* NiagaraGraph = Node->GetNiagaraGraph())
+		if (const UNiagaraGraph* NiagaraGraph = Node->GetNiagaraGraph())
 		{
 			Result = NiagaraGraph->GetOwningEmitter();
 		}
@@ -71,7 +72,7 @@ static FVersionedNiagaraEmitter GetOuterEmitter(UNiagaraNode* Node)
 	return Result;
 }
 
-static void GetFunctionStaticVariables(UNiagaraNodeFunctionCall& FunctionCallNode, TArray<FNiagaraVariable>& StaticVars)
+static void GetFunctionStaticVariables(const UNiagaraNodeFunctionCall& FunctionCallNode, TArray<FNiagaraVariable>& StaticVars)
 {
 	FVersionedNiagaraEmitter Emitter = GetOuterEmitter(&FunctionCallNode);
 	UNiagaraSystem* System = FunctionCallNode.GetTypedOuter<UNiagaraSystem>();
@@ -274,6 +275,46 @@ void FNiagaraStackGraphUtilities::GetOrderedModuleNodes(UNiagaraNodeOutput& Outp
 	}
 }
 
+TArray<UNiagaraNodeFunctionCall*> FNiagaraStackGraphUtilities::GetAllModuleNodes(TSharedRef<FNiagaraEmitterViewModel> EmitterViewModel)
+{
+	TArray<UNiagaraNodeFunctionCall*> ModuleNodes;
+	TArray<UNiagaraNodeOutput*> OutputNodes = GetAllOutputNodes(EmitterViewModel);
+
+	for(UNiagaraNodeOutput* OutputNode : OutputNodes)
+	{
+		TArray<UNiagaraNodeFunctionCall*> FunctionCalls;
+		FNiagaraStackGraphUtilities::GetOrderedModuleNodes(*OutputNode, FunctionCalls);
+		
+		for(UNiagaraNodeFunctionCall* FunctionCall : FunctionCalls)
+		{
+			ModuleNodes.Add(FunctionCall);
+		}
+	}
+
+	return ModuleNodes;
+}
+
+TArray<UNiagaraNodeFunctionCall*> FNiagaraStackGraphUtilities::GetAllModuleNodes(UNiagaraGraph* Graph)
+{
+	TArray<UNiagaraNodeFunctionCall*> ModuleNodes;
+
+	TArray<UNiagaraNodeOutput*> OutputNodes;
+	Graph->GetNodesOfClass<UNiagaraNodeOutput>(OutputNodes);
+
+	for(UNiagaraNodeOutput* OutputNode : OutputNodes)
+	{
+		TArray<UNiagaraNodeFunctionCall*> FunctionCalls;
+		FNiagaraStackGraphUtilities::GetOrderedModuleNodes(*OutputNode, FunctionCalls);
+		
+		for(UNiagaraNodeFunctionCall* FunctionCall : FunctionCalls)
+		{
+			ModuleNodes.Add(FunctionCall);
+		}
+	}
+
+	return ModuleNodes;
+}
+
 UNiagaraNodeFunctionCall* FNiagaraStackGraphUtilities::GetPreviousModuleNode(UNiagaraNodeFunctionCall& CurrentNode)
 {
 	UNiagaraNodeOutput* OutputNode = GetEmitterOutputNodeForStackNode(CurrentNode);
@@ -345,9 +386,9 @@ UNiagaraNodeOutput* FNiagaraStackGraphUtilities::GetEmitterOutputNodeForStackNod
 	return GetEmitterOutputNodeForStackNodeInternal<UNiagaraNodeOutput, UNiagaraNode>(StackNode);
 }
 
-ENiagaraScriptUsage FNiagaraStackGraphUtilities::GetOutputNodeUsage(UNiagaraNode& StackNode)
+ENiagaraScriptUsage FNiagaraStackGraphUtilities::GetOutputNodeUsage(const UNiagaraNode& StackNode)
 {
-	UNiagaraNodeOutput* OutputNode = GetEmitterOutputNodeForStackNode(StackNode);
+	const UNiagaraNodeOutput* OutputNode = GetEmitterOutputNodeForStackNode(StackNode);
 	return OutputNode ? OutputNode->GetUsage() : ENiagaraScriptUsage::Function;
 }
 
@@ -906,7 +947,7 @@ void ExtractInputPinsFromHistory(FNiagaraParameterMapHistory& History, UEdGraph*
 	}
 }
 
-void FNiagaraStackGraphUtilities::GetStackFunctionInputPins(UNiagaraNodeFunctionCall& FunctionCallNode, TArray<const UEdGraphPin*>& OutInputPins, TSet<const UEdGraphPin*>& OutHiddenPins, FCompileConstantResolver ConstantResolver, ENiagaraGetStackFunctionInputPinsOptions Options, bool bIgnoreDisabled)
+void FNiagaraStackGraphUtilities::GetStackFunctionInputPins(const UNiagaraNodeFunctionCall& FunctionCallNode, TArray<const UEdGraphPin*>& OutInputPins, TSet<const UEdGraphPin*>& OutHiddenPins, FCompileConstantResolver ConstantResolver, ENiagaraGetStackFunctionInputPinsOptions Options, bool bIgnoreDisabled)
 {
 	TArray<FNiagaraVariable> StaticVars;
 	FNiagaraStackGraphUtilitiesImpl::GetFunctionStaticVariables(FunctionCallNode, StaticVars);
@@ -914,7 +955,7 @@ void FNiagaraStackGraphUtilities::GetStackFunctionInputPins(UNiagaraNodeFunction
 	FNiagaraEditorModule::Get().GetGraphDataCache().GetStackFunctionInputPins(FunctionCallNode, StaticVars, OutInputPins, OutHiddenPins, ConstantResolver, Options, bIgnoreDisabled);
 }
 
-void FNiagaraStackGraphUtilities::GetStackFunctionInputPins(UNiagaraNodeFunctionCall& FunctionCallNode, TArray<const UEdGraphPin*>& OutInputPins, FCompileConstantResolver ConstantResolver, ENiagaraGetStackFunctionInputPinsOptions Options, bool bIgnoreDisabled)
+void FNiagaraStackGraphUtilities::GetStackFunctionInputPins(const UNiagaraNodeFunctionCall& FunctionCallNode, TArray<const UEdGraphPin*>& OutInputPins, FCompileConstantResolver ConstantResolver, ENiagaraGetStackFunctionInputPinsOptions Options, bool bIgnoreDisabled)
 {
 	TArray<FNiagaraVariable> StaticVars;
 	FNiagaraStackGraphUtilitiesImpl::GetFunctionStaticVariables(FunctionCallNode, StaticVars);
@@ -922,14 +963,14 @@ void FNiagaraStackGraphUtilities::GetStackFunctionInputPins(UNiagaraNodeFunction
 	FNiagaraEditorModule::Get().GetGraphDataCache().GetStackFunctionInputPins(FunctionCallNode, StaticVars, OutInputPins, ConstantResolver, Options, bIgnoreDisabled);
 }
 
-void FNiagaraStackGraphUtilities::GetStackFunctionInputPins(UNiagaraNodeFunctionCall& FunctionCallNode, TArray<const UEdGraphPin*>& OutInputPins, ENiagaraGetStackFunctionInputPinsOptions Options /*= ENiagaraGetStackFunctionInputPinsOptions::AllInputs*/, bool bIgnoreDisabled /*= false*/)
+void FNiagaraStackGraphUtilities::GetStackFunctionInputPins(const UNiagaraNodeFunctionCall& FunctionCallNode, TArray<const UEdGraphPin*>& OutInputPins, ENiagaraGetStackFunctionInputPinsOptions Options /*= ENiagaraGetStackFunctionInputPinsOptions::AllInputs*/, bool bIgnoreDisabled /*= false*/)
 {
 	FCompileConstantResolver EmptyResolver;
 	GetStackFunctionInputPins(FunctionCallNode, OutInputPins, EmptyResolver, Options, bIgnoreDisabled);
 }
 
 void FNiagaraStackGraphUtilities::GetStackFunctionInputPinsWithoutCache(
-	UNiagaraNodeFunctionCall& FunctionCallNode,
+	const UNiagaraNodeFunctionCall& FunctionCallNode,
 	TConstArrayView<FNiagaraVariable> StaticVars,
 	TArray<const UEdGraphPin*>& OutInputPins,
 	const FCompileConstantResolver& ConstantResolver,
@@ -981,6 +1022,172 @@ void FNiagaraStackGraphUtilities::GetStackFunctionInputPinsWithoutCache(
 	}
 }
 
+TArray<UNiagaraNodeOutput*> FNiagaraStackGraphUtilities::GetAllOutputNodes(TSharedRef<FNiagaraEmitterViewModel> EmitterViewModel)
+{
+	TSharedRef<FNiagaraScriptGraphViewModel> GraphViewModel = EmitterViewModel->GetSharedScriptViewModel()->GetGraphViewModel();
+
+	TArray<UNiagaraNodeOutput*> OutputNodes;
+	GraphViewModel->GetGraph()->GetNodesOfClass<UNiagaraNodeOutput>(OutputNodes);
+	
+	return OutputNodes;
+}
+
+UNiagaraNodeFunctionCall* FNiagaraStackGraphUtilities::FindModuleNode(FGuid ModuleNodeGuid, TSharedRef<FNiagaraEmitterViewModel> EmitterViewModel)
+{
+	TArray<UNiagaraNodeOutput*> OutputNodes = GetAllOutputNodes(EmitterViewModel);
+	
+	for(UNiagaraNodeOutput* OutputNode : OutputNodes)
+	{
+		TArray<UNiagaraNodeFunctionCall*> FunctionCalls;
+		FNiagaraStackGraphUtilities::GetOrderedModuleNodes(*OutputNode, FunctionCalls);
+		
+		for(UNiagaraNodeFunctionCall* FunctionCall : FunctionCalls)
+		{
+			if(FunctionCall->NodeGuid == ModuleNodeGuid)
+			{
+				return FunctionCall;
+			}
+		}
+	}
+
+	return nullptr;
+}
+
+TOptional<FNiagaraStackGraphUtilities::FMatchingFunctionInputData> FNiagaraStackGraphUtilities::FindInputData(FNiagaraHierarchyIdentity ModuleInputIdentity, TSharedRef<FNiagaraEmitterViewModel> EmitterViewModel)
+{
+	FMatchingFunctionInputData Result;
+	
+	TArray<UNiagaraNodeOutput*> OutputNodes = GetAllOutputNodes(EmitterViewModel);
+	
+	UNiagaraNodeFunctionCall* MatchingFunctionCall = FindModuleNode(ModuleInputIdentity.Guids[0], EmitterViewModel);
+	if(MatchingFunctionCall)
+	{
+		Result.FunctionCallNode = MatchingFunctionCall;
+	}
+					
+	if(MatchingFunctionCall)
+	{
+		TOptional<FMatchingFunctionInputData> InputData = FindInputData(*MatchingFunctionCall, ModuleInputIdentity.Guids[1], EmitterViewModel);
+
+		if(InputData.IsSet())
+		{
+			Result.InputName = InputData->InputName;
+			Result.Type = InputData->Type;
+			Result.MetaData = InputData->MetaData;
+			Result.ChildrenInputGuids = InputData->ChildrenInputGuids;
+
+			return Result;
+		}
+	}
+	
+	return TOptional<FMatchingFunctionInputData>();
+}
+
+TOptional<FNiagaraStackGraphUtilities::FMatchingFunctionInputData> FNiagaraStackGraphUtilities::FindInputData(const UNiagaraNodeFunctionCall& FunctionCallNode, FGuid InputGuid, TSharedRef<FNiagaraEmitterViewModel> EmitterViewModel, bool bIncludeChildrenInputs)
+{
+	FMatchingFunctionInputData Result;
+	Result.FunctionCallNode = const_cast<UNiagaraNodeFunctionCall*>(&FunctionCallNode);
+	
+	FCompileConstantResolver ConstantResolver(EmitterViewModel->GetEmitter(), GetOutputNodeUsage(FunctionCallNode));
+
+	// our goal is to find the correct variable's metadata from the graph
+	TOptional<FNiagaraVariableMetaData> MetaData;
+				
+	// first, we check the normal module inputs
+	TArray<const UEdGraphPin*> InputPins;
+	FNiagaraStackGraphUtilities::ENiagaraGetStackFunctionInputPinsOptions Options = FNiagaraStackGraphUtilities::ENiagaraGetStackFunctionInputPinsOptions::AllInputs;
+	FNiagaraStackGraphUtilities::GetStackFunctionInputPins(FunctionCallNode, InputPins, ConstantResolver, Options, true);
+				
+	for(const UEdGraphPin* InputPin : InputPins)
+	{
+		FNiagaraVariable InputVariable = UEdGraphSchema_Niagara::PinToNiagaraVariable(InputPin);
+		if (InputVariable.GetType().IsValid() == false)
+		{
+			continue;
+		}
+
+		MetaData = FunctionCallNode.GetCalledGraph()->GetMetaData(InputVariable);
+		if(MetaData.IsSet() && MetaData.GetValue().GetVariableGuid() == InputGuid)
+		{
+			Result.InputName = InputPin->PinName;
+			Result.Type = InputVariable.GetType();
+			Result.MetaData = MetaData.GetValue();
+
+			if(bIncludeChildrenInputs)
+			{
+				Result.ChildrenInputGuids = GetChildrenInputGuids(FunctionCallNode, Result.InputName);				
+			}
+			
+			return Result;
+		}
+	}
+
+	// then if we didn't find the input under module inputs, we check the static switch inputs
+	TSet<UEdGraphPin*> HiddenSwitchPins;
+	TArray<UEdGraphPin*> SwitchPins;
+				
+	FNiagaraStackGraphUtilities::GetStackFunctionStaticSwitchPins(FunctionCallNode, SwitchPins, HiddenSwitchPins, ConstantResolver);
+	for (UEdGraphPin* InputPin : SwitchPins)
+	{
+		// The static switch pin names to not contain the module namespace, as they are not part of the parameter maps.
+		// We add it here only to check for name clashes with actual module parameters.
+		FString ModuleName = PARAM_MAP_MODULE_STR;
+		InputPin->PinName.AppendString(ModuleName);
+		FName SwitchPinName(*ModuleName);
+
+		FNiagaraVariable InputVariable = UEdGraphSchema_Niagara::PinToNiagaraVariable(InputPin);
+		if (InputVariable.GetType().IsValid() == false)
+		{
+			continue;
+		}
+					
+		if (FunctionCallNode.GetCalledGraph() != nullptr)
+		{
+			MetaData = FunctionCallNode.GetCalledGraph()->GetMetaData(InputVariable);
+			// if metadata is set and the 2nd guid matches the variable, we found the input that our HierarchyModuleInput represents
+			if(MetaData.IsSet() && MetaData.GetValue().GetVariableGuid() == InputGuid)
+			{
+				Result.InputName = SwitchPinName;
+				Result.Type = InputVariable.GetType();
+				Result.MetaData = MetaData.GetValue();
+
+				if(bIncludeChildrenInputs)
+				{
+					Result.ChildrenInputGuids = GetChildrenInputGuids(FunctionCallNode, Result.InputName);				
+				}
+				
+				return Result;
+			}
+		}
+	}
+	
+	return TOptional<FMatchingFunctionInputData>();
+}
+
+TArray<FGuid> FNiagaraStackGraphUtilities::GetChildrenInputGuids(const UNiagaraNodeFunctionCall& FunctionCallNode, FName ParentInputName)
+{
+	TArray<FGuid> Result;
+	
+	TArray<FNiagaraVariable> Variables;
+	FunctionCallNode.GetCalledGraph()->GetAllVariables(Variables);
+	
+	for(FNiagaraVariable& Variable : Variables)
+	{
+		TOptional<FNiagaraVariableMetaData> MetaData = FunctionCallNode.GetCalledGraph()->GetMetaData(Variable);
+
+		if(MetaData.IsSet())
+		{
+			FNiagaraParameterHandle ParentInputHandle(ParentInputName);
+			if(MetaData->ParentAttribute == ParentInputHandle.GetName())
+			{
+				Result.Add(MetaData->GetVariableGuid());
+			}
+		}
+	}
+
+	return Result;
+}
+
 TArray<UEdGraphPin*> FNiagaraStackGraphUtilities::GetUnusedFunctionInputPins(UNiagaraNodeFunctionCall& FunctionCallNode, FCompileConstantResolver ConstantResolver)
 {
 	UNiagaraGraph* FunctionGraph = FunctionCallNode.GetCalledGraph();
@@ -1063,7 +1270,7 @@ TArray<UEdGraphPin*> FNiagaraStackGraphUtilities::GetUnusedFunctionInputPins(UNi
 	return ResultPins;
 }
 
-void FNiagaraStackGraphUtilities::GetStackFunctionStaticSwitchPins(UNiagaraNodeFunctionCall& FunctionCallNode, TArray<UEdGraphPin*>& OutInputPins, TSet<UEdGraphPin*>& OutHiddenPins,
+void FNiagaraStackGraphUtilities::GetStackFunctionStaticSwitchPins(const UNiagaraNodeFunctionCall& FunctionCallNode, TArray<UEdGraphPin*>& OutInputPins, TSet<UEdGraphPin*>& OutHiddenPins,
 	FCompileConstantResolver& ConstantResolver)
 {
 	const UEdGraphSchema_Niagara* Schema = CastChecked<UEdGraphSchema_Niagara>(FunctionCallNode.GetSchema());
diff --git a/Engine/Plugins/FX/Niagara/Source/NiagaraEditor/Public/NiagaraGraphDataCache.h b/Engine/Plugins/FX/Niagara/Source/NiagaraEditor/Public/NiagaraGraphDataCache.h
--- a/Engine/Plugins/FX/Niagara/Source/NiagaraEditor/Public/NiagaraGraphDataCache.h
+++ b/Engine/Plugins/FX/Niagara/Source/NiagaraEditor/Public/NiagaraGraphDataCache.h
@@ -19,12 +19,12 @@ class FNiagaraGraphDataCache
 public:
 	FNiagaraGraphDataCache();
 
-	void GetStackFunctionInputPins(UNiagaraNodeFunctionCall& FunctionCallNode, TConstArrayView<FNiagaraVariable> StaticVars, TArray<const UEdGraphPin*>& OutInputPins, FCompileConstantResolver ConstantResolver, FNiagaraStackGraphUtilities::ENiagaraGetStackFunctionInputPinsOptions Options, bool bIgnoreDisabled);
-	void GetStackFunctionInputPins(UNiagaraNodeFunctionCall& FunctionCallNode, TConstArrayView<FNiagaraVariable> StaticVars, TArray<const UEdGraphPin*>& OutInputPins, TSet<const UEdGraphPin*>& OutHiddenPins, FCompileConstantResolver ConstantResolver, FNiagaraStackGraphUtilities::ENiagaraGetStackFunctionInputPinsOptions Options, bool bIgnoreDisabled);
+	void GetStackFunctionInputPins(const UNiagaraNodeFunctionCall& FunctionCallNode, TConstArrayView<FNiagaraVariable> StaticVars, TArray<const UEdGraphPin*>& OutInputPins, FCompileConstantResolver ConstantResolver, FNiagaraStackGraphUtilities::ENiagaraGetStackFunctionInputPinsOptions Options, bool bIgnoreDisabled);
+	void GetStackFunctionInputPins(const UNiagaraNodeFunctionCall& FunctionCallNode, TConstArrayView<FNiagaraVariable> StaticVars, TArray<const UEdGraphPin*>& OutInputPins, TSet<const UEdGraphPin*>& OutHiddenPins, FCompileConstantResolver ConstantResolver, FNiagaraStackGraphUtilities::ENiagaraGetStackFunctionInputPinsOptions Options, bool bIgnoreDisabled);
 
 private:
 	void GetStackFunctionInputPinsInternal(
-		UNiagaraNodeFunctionCall& FunctionCallNode,
+		const UNiagaraNodeFunctionCall& FunctionCallNode,
 		const UNiagaraGraph* CalledGraph,
 		TConstArrayView<FNiagaraVariable> StaticVars,
 		TArray<const UEdGraphPin*>& OutInputPins,
diff --git a/Engine/Plugins/FX/Niagara/Source/NiagaraEditor/Public/ViewModels/Stack/NiagaraStackGraphUtilities.h b/Engine/Plugins/FX/Niagara/Source/NiagaraEditor/Public/ViewModels/Stack/NiagaraStackGraphUtilities.h
--- a/Engine/Plugins/FX/Niagara/Source/NiagaraEditor/Public/ViewModels/Stack/NiagaraStackGraphUtilities.h
+++ b/Engine/Plugins/FX/Niagara/Source/NiagaraEditor/Public/ViewModels/Stack/NiagaraStackGraphUtilities.h
@@ -9,6 +9,7 @@
 #include "AssetRegistry/AssetData.h"
 #include "NiagaraStackEditorData.h"
 
+struct FNiagaraHierarchyIdentity;
 class UEdGraph;
 class UEdGraphPin;
 class UNiagaraGraph;
@@ -46,6 +47,8 @@ namespace FNiagaraStackGraphUtilities
 	UEdGraphPin* GetParameterMapOutputPin(UNiagaraNode& Node);
 
 	void GetOrderedModuleNodes(UNiagaraNodeOutput& OutputNode, TArray<UNiagaraNodeFunctionCall*>& ModuleNodes);
+	TArray<UNiagaraNodeFunctionCall*> GetAllModuleNodes(TSharedRef<FNiagaraEmitterViewModel> EmitterViewModel);
+	TArray<UNiagaraNodeFunctionCall*> GetAllModuleNodes(UNiagaraGraph* Graph);
 
 	UNiagaraNodeFunctionCall* GetPreviousModuleNode(UNiagaraNodeFunctionCall& CurrentNode);
 
@@ -53,7 +56,7 @@ namespace FNiagaraStackGraphUtilities
 
 	UNiagaraNodeOutput* GetEmitterOutputNodeForStackNode(UNiagaraNode& StackNode);
 
-	NIAGARAEDITOR_API ENiagaraScriptUsage GetOutputNodeUsage(UNiagaraNode& StackNode);
+	NIAGARAEDITOR_API ENiagaraScriptUsage GetOutputNodeUsage(const UNiagaraNode& StackNode);
 
 	const UNiagaraNodeOutput* GetEmitterOutputNodeForStackNode(const UNiagaraNode& StackNode);
 
@@ -93,15 +96,15 @@ namespace FNiagaraStackGraphUtilities
 	};
 
 	/* Returns the input pins for the given function call node. Try not to use this method if possible and use the version that accepts a FCompileConstantResolver parameter. This method should only be used if the current context is completely outside of any system or emitter (e.g. inside the graph editor itself) and constants cannot be resolved at all. */
-	void GetStackFunctionInputPins(UNiagaraNodeFunctionCall& FunctionCallNode, TArray<const UEdGraphPin*>& OutInputPins, ENiagaraGetStackFunctionInputPinsOptions Options, bool bIgnoreDisabled = false);
+	void GetStackFunctionInputPins(const UNiagaraNodeFunctionCall& FunctionCallNode, TArray<const UEdGraphPin*>& OutInputPins, ENiagaraGetStackFunctionInputPinsOptions Options, bool bIgnoreDisabled = false);
 
 	/* Returns the input pins for the given function call node. */
-	void GetStackFunctionInputPins(UNiagaraNodeFunctionCall& FunctionCallNode, TArray<const UEdGraphPin*>& OutInputPins, TSet<const UEdGraphPin*>& OutHiddenPins, FCompileConstantResolver ConstantResolver, ENiagaraGetStackFunctionInputPinsOptions Options, bool bIgnoreDisabled = false);
-	void GetStackFunctionInputPins(UNiagaraNodeFunctionCall& FunctionCallNode, TArray<const UEdGraphPin*>& OutInputPins, FCompileConstantResolver ConstantResolver, ENiagaraGetStackFunctionInputPinsOptions Options, bool bIgnoreDisabled = false);
+	void GetStackFunctionInputPins(const UNiagaraNodeFunctionCall& FunctionCallNode, TArray<const UEdGraphPin*>& OutInputPins, TSet<const UEdGraphPin*>& OutHiddenPins, FCompileConstantResolver ConstantResolver, ENiagaraGetStackFunctionInputPinsOptions Options, bool bIgnoreDisabled = false);
+	void GetStackFunctionInputPins(const UNiagaraNodeFunctionCall& FunctionCallNode, TArray<const UEdGraphPin*>& OutInputPins, FCompileConstantResolver ConstantResolver, ENiagaraGetStackFunctionInputPinsOptions Options, bool bIgnoreDisabled = false);
 
 	/* Returns the input pins for the given function call node.  Bypasses the module level caching of this data and generates it each call. */
 	void GetStackFunctionInputPinsWithoutCache(
-		UNiagaraNodeFunctionCall& FunctionCallNode,
+		const UNiagaraNodeFunctionCall& FunctionCallNode,
 		TConstArrayView<FNiagaraVariable> StaticVars,
 		TArray<const UEdGraphPin*>& OutInputPins,
 		const FCompileConstantResolver& ConstantResolver,
@@ -109,10 +112,26 @@ namespace FNiagaraStackGraphUtilities
 		bool bIgnoreDisabled,
 		bool bFilterForCompilation);
 
+	TArray<UNiagaraNodeOutput*> GetAllOutputNodes(TSharedRef<FNiagaraEmitterViewModel> EmitterViewModel);
+	
+	struct FMatchingFunctionInputData
+	{
+		FName InputName;
+		FNiagaraTypeDefinition Type;
+		FNiagaraVariableMetaData MetaData;
+		UNiagaraNodeFunctionCall* FunctionCallNode = nullptr;
+		TArray<FGuid> ChildrenInputGuids;
+	};
+
+	UNiagaraNodeFunctionCall* FindModuleNode(FGuid ModuleNodeGuid, TSharedRef<FNiagaraEmitterViewModel> EmitterViewModel);
+	TOptional<FMatchingFunctionInputData> FindInputData(FNiagaraHierarchyIdentity ModuleInputIdentity, TSharedRef<FNiagaraEmitterViewModel> EmitterViewModel);
+	TOptional<FMatchingFunctionInputData> FindInputData(const UNiagaraNodeFunctionCall& FunctionCallNode, FGuid InputGuid, TSharedRef<FNiagaraEmitterViewModel> EmitterViewModel, bool bIncludeChildrenInputs = true);
+	TArray<FGuid> GetChildrenInputGuids(const UNiagaraNodeFunctionCall& FunctionCallNode, FName ParentInputName);
+	
 	/* Module script calls do not have direct inputs, but rely on the parameter map being initialized correctly. This utility function resolves which of the module's parameters are reachable during compilation and returns a list of pins on the parameter map node that do not have to be compiled. */
 	TArray<UEdGraphPin*> GetUnusedFunctionInputPins(UNiagaraNodeFunctionCall& FunctionCallNode, FCompileConstantResolver ConstantResolver);
 
-	void GetStackFunctionStaticSwitchPins(UNiagaraNodeFunctionCall& FunctionCallNode, TArray<UEdGraphPin*>& OutInputPins, TSet<UEdGraphPin*>& OutHiddenPins, FCompileConstantResolver& ConstantResolver);
+	void GetStackFunctionStaticSwitchPins(const UNiagaraNodeFunctionCall& FunctionCallNode, TArray<UEdGraphPin*>& OutInputPins, TSet<UEdGraphPin*>& OutHiddenPins, FCompileConstantResolver& ConstantResolver);
 
 	void GetStackFunctionOutputVariables(UNiagaraNodeFunctionCall& FunctionCallNode, FCompileConstantResolver ConstantResolver, TArray<FNiagaraVariable>& OutOutputVariables, TArray<FNiagaraVariable>& OutOutputVariablesWithOriginalAliasesIntact);
 
