Niagara Data Channels v0.1

Niagara Data Channels are new system for communication between Niagara Systems and with game code/BP.

The system has a few peices.

DataChannel - Data Definition for a named channel of data. New types can be added to deal with the channel differently.
DataChannelHander - Runtime counterpart for DataChannels each Channel Type will have a handler that uses the data definition to manager the channel data accordingly.
DataChannelDefinitions - Asset class containing one or more Data Channel.
DataChannelManager - Manger class contianing and providing access to all currently loaded Data Channels and their Handlers.
DataChannelAccessor - Temporary initial BP API allowing BPs to read and write to Data Channels.
Data Channel Global - Initial simple Data Channel Type. Stores all data globally.
Data Interface Data Channel Read - New DI that allows Niagara Systems to read data from a Data Channel.
Data Interface Data Channel Write - New DI that allows Niagara Systems to write data toi a Data Channel.
Data Interface Data Channel Spawn - Specialization of the Read DI which allows for easy spwaning based on the contents of a Data Channel.

Data Channels can be enabled/disabled via FX.Niagara.DataChannels.Enabled = 0.
Initially this defaults to off.

Misc:

- New DI node action providers to enable actions on the Data Channel DI nodes.
- Details customized reference class for data channels that avoids any typing the available channel names etc.
- DIs now have a post compile step that can be used to cache of certain data at compile time.
- NiagaraDataInterfaceUtilities now has an optional SearchOptions parameter that allows us to include "internal" DIs. This is required to support the above PostCompile step.
- Niagara Type Definitions now optionally allow LWC types, previously any LWC types would be auto converted to SWC on load.
- Asset Registry now tracks SWC Sim types as we would get duplicated entires in several places of the UI due to confusion between LWC and SWC types.

diff --git a/Engine/Plugins/FX/Niagara/Source/Niagara/Classes/NiagaraDataInterface.h b/Engine/Plugins/FX/Niagara/Source/Niagara/Classes/NiagaraDataInterface.h
--- a/Engine/Plugins/FX/Niagara/Source/Niagara/Classes/NiagaraDataInterface.h
+++ b/Engine/Plugins/FX/Niagara/Source/Niagara/Classes/NiagaraDataInterface.h
@@ -576,6 +576,9 @@ public:
 	virtual bool HasPreSimulateTick() const { return false; }
 	virtual bool HasPostSimulateTick() const { return false; }
 
+	/** Called after system compilation completes. Useful for caching any constant, compile dependent data. */
+	virtual void PostCompile() { }
+
 	/**
 	When set to true the simulation may not complete in the same frame it started, allowing maximum overlap with the GameThread.
 	You must override and set to false if you require the data interface to complete before PostActorTick completes.
@@ -631,8 +634,6 @@ public:
 	void AppendTemplateHLSL(FString& OutHLSL, const TCHAR* TemplateShaderFile, const TMap<FString, FStringFormatArg>& TemplateArgs) const;
 #endif
 
-	virtual void PostExecute() {}
-
 #if WITH_NIAGARA_DEBUGGER
 	/**
 	Override this function to provide additional context to the debug HUD.
@@ -975,6 +976,19 @@ struct FNDIInputParam
 	VectorVM::FExternalFuncInputHandler<T> Data;
 	FORCEINLINE FNDIInputParam(FVectorVMExternalFunctionContext& Context) : Data(Context) {}
 	FORCEINLINE T GetAndAdvance() { return Data.GetAndAdvance(); }
+	FORCEINLINE T Get() { return Data.Get(); }
+	FORCEINLINE void Advance() { return Data.Advance(); }
+	FORCEINLINE bool IsConstant() const { return Data.IsConstant(); }
+};
+
+template<>
+struct FNDIInputParam<FFloat16>
+{
+	VectorVM::FExternalFuncInputHandler<FFloat16> Data;
+	FORCEINLINE FNDIInputParam(FVectorVMExternalFunctionContext& Context) : Data(Context) {}
+	FORCEINLINE FFloat16 GetAndAdvance() { return Data.GetAndAdvance(); }
+	FORCEINLINE FFloat16 Get() { return Data.Get(); }
+	FORCEINLINE void Advance() { return Data.Advance(); }
 	FORCEINLINE bool IsConstant() const { return Data.IsConstant(); }
 };
 
@@ -984,6 +998,8 @@ struct FNDIInputParam<FNiagaraBool>
 	VectorVM::FExternalFuncInputHandler<FNiagaraBool> Data;
 	FORCEINLINE FNDIInputParam(FVectorVMExternalFunctionContext& Context) : Data(Context) {}
 	FORCEINLINE bool GetAndAdvance() { return Data.GetAndAdvance().GetValue(); }
+	FORCEINLINE bool Get() { return Data.Get().GetValue(); }
+	FORCEINLINE void Advance() { return Data.Advance(); }
 	FORCEINLINE bool IsConstant() const { return Data.IsConstant(); }
 };
 
@@ -993,6 +1009,8 @@ struct FNDIInputParam<bool>
 	VectorVM::FExternalFuncInputHandler<FNiagaraBool> Data;
 	FORCEINLINE FNDIInputParam(FVectorVMExternalFunctionContext& Context) : Data(Context) {}
 	FORCEINLINE bool GetAndAdvance() { return Data.GetAndAdvance().GetValue(); }
+	FORCEINLINE bool Get() { return Data.Get().GetValue(); }
+	FORCEINLINE void Advance() { return Data.Advance(); }
 	FORCEINLINE bool IsConstant() const { return Data.IsConstant(); }
 };
 
@@ -1003,6 +1021,8 @@ struct FNDIInputParam<FVector2f>
 	VectorVM::FExternalFuncInputHandler<float> Y;
 	FORCEINLINE FNDIInputParam(FVectorVMExternalFunctionContext& Context) : X(Context), Y(Context) {}
 	FORCEINLINE FVector2f GetAndAdvance() { return FVector2f(X.GetAndAdvance(), Y.GetAndAdvance()); }
+	FORCEINLINE FVector2f Get() { return FVector2f(X.Get(), Y.Get()); }
+	FORCEINLINE void Advance() { X.Advance(), Y.Advance(); }
 	FORCEINLINE bool IsConstant() const { return X.IsConstant() && Y.IsConstant(); }
 };
 
@@ -1014,6 +1034,8 @@ struct FNDIInputParam<FVector3f>
 	VectorVM::FExternalFuncInputHandler<float> Z;
 	FNDIInputParam(FVectorVMExternalFunctionContext& Context) : X(Context), Y(Context), Z(Context) {}
 	FORCEINLINE FVector3f GetAndAdvance() { return FVector3f(X.GetAndAdvance(), Y.GetAndAdvance(), Z.GetAndAdvance()); }
+	FORCEINLINE FVector3f Get() { return FVector3f(X.Get(), Y.Get(), Z.Get()); }
+	FORCEINLINE void Advance() { X.Advance(); Y.Advance(); Z.Advance(); }
 	FORCEINLINE bool IsConstant() const { return X.IsConstant() && Y.IsConstant() && Z.IsConstant(); }
 };
 
@@ -1025,6 +1047,8 @@ struct FNDIInputParam<FNiagaraPosition>
 	VectorVM::FExternalFuncInputHandler<float> Z;
 	FNDIInputParam(FVectorVMExternalFunctionContext& Context) : X(Context), Y(Context), Z(Context) {}
 	FORCEINLINE FNiagaraPosition GetAndAdvance() { return FNiagaraPosition(X.GetAndAdvance(), Y.GetAndAdvance(), Z.GetAndAdvance()); }
+	FORCEINLINE FNiagaraPosition Get() { return FNiagaraPosition(X.Get(), Y.Get(), Z.Get()); }
+	FORCEINLINE void Advance() { X.Advance(); Y.Advance(); Z.Advance(); }
 };
 template<>
 struct FNDIInputParam<FVector4f>
@@ -1035,6 +1059,8 @@ struct FNDIInputParam<FVector4f>
 	VectorVM::FExternalFuncInputHandler<float> W;
 	FORCEINLINE FNDIInputParam(FVectorVMExternalFunctionContext& Context) : X(Context), Y(Context), Z(Context), W(Context) {}
 	FORCEINLINE FVector4f GetAndAdvance() { return FVector4f(X.GetAndAdvance(), Y.GetAndAdvance(), Z.GetAndAdvance(), W.GetAndAdvance()); }
+	FORCEINLINE FVector4f Get() { return FVector4f(X.Get(), Y.Get(), Z.Get(), W.Get()); }
+	FORCEINLINE void Advance() { return X.Advance(); Y.Advance(); Z.Advance(); W.Advance(); }
 	FORCEINLINE bool IsConstant() const { return X.IsConstant() && Y.IsConstant() && Z.IsConstant() && W.IsConstant(); }
 };
 
@@ -1047,6 +1073,8 @@ struct FNDIInputParam<FQuat4f>
 	VectorVM::FExternalFuncInputHandler<float> W;
 	FORCEINLINE FNDIInputParam(FVectorVMExternalFunctionContext& Context) : X(Context), Y(Context), Z(Context), W(Context) {}
 	FORCEINLINE FQuat4f GetAndAdvance() { return FQuat4f(X.GetAndAdvance(), Y.GetAndAdvance(), Z.GetAndAdvance(), W.GetAndAdvance()); }
+	FORCEINLINE FQuat4f Get() { return FQuat4f(X.Get(), Y.Get(), Z.Get(), W.Get()); }
+	FORCEINLINE void Advance() { X.Advance(); Y.Advance(); Z.Advance(); W.Advance(); }
 	FORCEINLINE bool IsConstant() const { return X.IsConstant() && Y.IsConstant() && Z.IsConstant() && W.IsConstant(); }
 };
 
@@ -1059,6 +1087,8 @@ struct FNDIInputParam<FLinearColor>
 	VectorVM::FExternalFuncInputHandler<float> A;
 	FORCEINLINE FNDIInputParam(FVectorVMExternalFunctionContext& Context) : R(Context), G(Context), B(Context), A(Context) {}
 	FORCEINLINE FLinearColor GetAndAdvance() { return FLinearColor(R.GetAndAdvance(), G.GetAndAdvance(), B.GetAndAdvance(), A.GetAndAdvance()); }
+	FORCEINLINE FLinearColor Get() { return FLinearColor(R.Get(), G.Get(), B.Get(), A.Get()); }
+	FORCEINLINE void Advance() { R.Advance(); G.Advance(); B.Advance(); A.Advance(); }
 	FORCEINLINE bool IsConstant() const { return R.IsConstant() && G.IsConstant() && B.IsConstant() && A.IsConstant(); }
 };
 
@@ -1069,6 +1099,8 @@ struct FNDIInputParam<FNiagaraID>
 	VectorVM::FExternalFuncInputHandler<int32> AcquireTag;
 	FORCEINLINE FNDIInputParam(FVectorVMExternalFunctionContext& Context) : Index(Context), AcquireTag(Context) {}
 	FORCEINLINE FNiagaraID GetAndAdvance() { return FNiagaraID(Index.GetAndAdvance(), AcquireTag.GetAndAdvance()); }
+	FORCEINLINE FNiagaraID Get() { return FNiagaraID(Index.Get(), AcquireTag.Get()); }
+	FORCEINLINE void Advance() { return Index.Advance(); AcquireTag.Advance(); }
 	FORCEINLINE bool IsConstant() const { return Index.IsConstant() && AcquireTag.IsConstant(); }
 };
 
diff --git a/Engine/Plugins/FX/Niagara/Source/Niagara/Classes/NiagaraSystem.h b/Engine/Plugins/FX/Niagara/Source/Niagara/Classes/NiagaraSystem.h
--- a/Engine/Plugins/FX/Niagara/Source/Niagara/Classes/NiagaraSystem.h
+++ b/Engine/Plugins/FX/Niagara/Source/Niagara/Classes/NiagaraSystem.h
@@ -46,6 +46,9 @@ enum class ENiagaraScriptTemplateSpecification : uint8;
 struct FNiagaraGraphCachedDataBase;
 #endif
 
+class UNiagaraDataChannel;
+class UNiagaraDataChannelDefinitions;
+
 USTRUCT()
 struct FNiagaraEmitterCompiledData
 {
@@ -762,6 +765,12 @@ public:
 
 	/** Resets internal data leaving it in a state which would have minimal cost to exist in headless builds (servers) */
 	void ResetToEmptySystem();
+
+	/** Registers that this System uses the passed DI. No symmetrical unregister as this is cleared and rebuilt on compile. */
+	void RegisterDataChannelUse(const UNiagaraDataChannel* DataChannel);
+
+	/** Updates any post compile data based upon data interfaces. */
+	void OnCompiledDataInterfaceChanged();
 #endif
 
 private:
@@ -943,6 +952,9 @@ protected:
 	UPROPERTY()
 	TArray<FName> UserDINamesReadInSystemScripts;
 
+	UPROPERTY()
+	TArray<TObjectPtr<UNiagaraDataChannelDefinitions>> ReferencedDataChannelDefinitions;
+
 	/** Array of emitter indices sorted by execution priority. The emitters will be ticked in this order. Please note that some indices may have the top bit set (kStartNewOverlapGroupBit)
 	* to indicate synchronization points in parallel execution, so mask it out before using the values as indices in the emitters array.
 	*/
diff --git a/Engine/Plugins/FX/Niagara/Source/Niagara/Private/DataInterface/NiagaraDataInterfaceDataChannelRead.cpp b/Engine/Plugins/FX/Niagara/Source/Niagara/Private/DataInterface/NiagaraDataInterfaceDataChannelRead.cpp
new file mode 100644
--- /dev/null
+++ b/Engine/Plugins/FX/Niagara/Source/Niagara/Private/DataInterface/NiagaraDataInterfaceDataChannelRead.cpp
@@ -0,0 +1,1363 @@
+// Copyright Epic Games, Inc. All Rights Reserved.
+
+#include "NiagaraDataInterfaceDataChannelRead.h"
+
+#include "ShaderCompilerCore.h"
+#include "ShaderParameterUtils.h"
+#include "NiagaraShaderParametersBuilder.h"
+
+#include "CoreMinimal.h"
+#include "NiagaraCommon.h"
+#include "NiagaraWorldManager.h"
+
+#include "NiagaraDataChannel.h"
+#include "NiagaraDataChannelHandler.h"
+#include "NiagaraDataChannelDefinitions.h"
+
+#include "NiagaraRenderer.h"
+#include "NiagaraGPUSystemTick.h"
+
+#include "NiagaraDataInterfaceUtilities.h"
+
+//////////////////////////////////////////////////////////////////////////
+
+#define LOCTEXT_NAMESPACE "NiagaraDataInterfaceDataChannelRead"
+
+DECLARE_CYCLE_STAT(TEXT("NDIDataChannelRead Read"), STAT_NDIDataChannelRead_Read, STATGROUP_NiagaraDataChannels);
+DECLARE_CYCLE_STAT(TEXT("NDIDataChannelRead Consume"), STAT_NDIDataChannelRead_Consume, STATGROUP_NiagaraDataChannels);
+DECLARE_CYCLE_STAT(TEXT("NDIDataChannelRead Tick"), STAT_NDIDataChannelRead_Tick, STATGROUP_NiagaraDataChannels);
+DECLARE_CYCLE_STAT(TEXT("NDIDataChannelRead PostTick"), STAT_NDIDataChannelRead_PostTick, STATGROUP_NiagaraDataChannels);
+
+namespace NDIDataChannelReadLocal
+{
+	static const TCHAR* CommonShaderFile = TEXT("/Plugin/FX/Niagara/Private/DataChannel/NiagaraDataInterfaceDataChannelCommon.ush");
+	static const TCHAR* TemplateShaderFile_Common = TEXT("/Plugin/FX/Niagara/Private/DataChannel/NiagaraDataInterfaceDataChannelTemplateCommon.ush");
+	static const TCHAR* TemplateShaderFile_Read = TEXT("/Plugin/FX/Niagara/Private/DataChannel/NiagaraDataInterfaceDataChannelTemplate_Read.ush");
+	static const TCHAR* TemplateShaderFile_Consume = TEXT("/Plugin/FX/Niagara/Private/DataChannel/NiagaraDataInterfaceDataChannelTemplate_Consume.ush");
+
+	static const FName NumName(TEXT("Num"));
+	static const FName ReadName(TEXT("Read"));
+	static const FName ConsumeName(TEXT("Consume"));
+
+	const TCHAR* GetFunctionTemplate(FName FunctionName)
+	{
+		if (FunctionName == ReadName) return TemplateShaderFile_Read;
+		if (FunctionName == ConsumeName) return TemplateShaderFile_Consume;
+
+		return nullptr;
+	}
+
+	BEGIN_SHADER_PARAMETER_STRUCT(FShaderParameters, )
+		SHADER_PARAMETER_SRV(Buffer<uint>, ParamOffsetTable)
+		SHADER_PARAMETER_SRV(Buffer<float>, DataFloat)
+		SHADER_PARAMETER_SRV(Buffer<int>, DataInt32)
+		SHADER_PARAMETER_SRV(Buffer<float>, DataHalf)
+		SHADER_PARAMETER(int32, ParameterOffsetTableIndex)
+		SHADER_PARAMETER(int32, Num)
+		SHADER_PARAMETER(int32, FloatStride)
+		SHADER_PARAMETER(int32, Int32Stride)
+		SHADER_PARAMETER(int32, HalfStride)
+	END_SHADER_PARAMETER_STRUCT()
+}
+
+/** Render thread copy of current instance data. */
+struct FNDIDataChannelReadInstanceData_RT
+{
+	//GPU Dataset from the channel handler. We'll grab the current buffer from this on the RT.
+	//This must be grabbed fresh from the handler each frame as it's lifetime cannot be ensured.
+	//TODO: Improve this.
+	//TODO: Likely better to have the handler provide some RT proxy class for the correct data from which we grab the correct buffer.
+	FNiagaraDataSet* ChannelDataSet = nullptr;
+
+	/**
+	Table of all parameter offsets used by each GPU script using this DI.
+	Each script has to have it's own section of this table as the offsets into this table are embedded in the hlsl.
+	At hlsl gen time we only have the context of each script individually to generate these indexes.
+	TODO: Can possible elevate this up to the LayoutManager and have a single layout buffer for all scripts
+	*/
+	TResourceArray<uint32> GPUScriptParameterOffsetTable;
+
+	/**
+	Offsets into the parameter table are embedded in the gpu script hlsl.
+	At hlsl gen time we can only know which parameters are accessed by each script individually so each script must have it's own parameter binding table.
+	We provide the offset into the above table via a shader param.
+	TODO: Can just as easily be an offset into a global buffer in the Layout manager.
+	*/
+	TMap<FNiagaraCompileHash, uint32> GPUScriptParameterTableOffsets;
+
+	/** Signal we have updated data from the GT this frame we should process. */
+	bool bHasUpdate = false;
+};
+
+//////////////////////////////////////////////////////////////////////////
+
+FNDIDataChannelReadInstanceData::~FNDIDataChannelReadInstanceData()
+{
+	if (ExternalCPU)
+	{
+		ExternalCPU->ReleaseReadRef();
+		ExternalCPU = nullptr;
+	}
+
+	ExternalGPU = nullptr;
+}
+
+FNiagaraDataBuffer* FNDIDataChannelReadInstanceData::GetReadBufferCPU()
+{
+	if (ExternalCPU)
+	{
+		return ExternalCPU;
+	}
+
+	//TODO: Local reads.
+// 	if (SourceInstData)
+// 	{
+// 		if (SourceInstData->Data && SourceInstData->Data->GetSimTarget() == ENiagaraSimTarget::CPUSim)
+// 		{
+// 			return SourceInstData->Data->GetDestinationData();
+// 		}
+// 	}
+
+	return nullptr;
+}
+
+bool FNDIDataChannelReadInstanceData::Init(UNiagaraDataInterfaceDataChannelRead* Interface, FNiagaraSystemInstance* Instance)
+{
+	bool bSuccess = Tick(Interface, Instance);
+	bSuccess &= PostTick(Interface, Instance);
+	return bSuccess;
+}
+
+bool FNDIDataChannelReadInstanceData::Tick(UNiagaraDataInterfaceDataChannelRead* Interface, FNiagaraSystemInstance* Instance)
+{
+	ConsumeIndex = 0;
+
+	const FNDIDataChannelCompiledData& CompiledData = Interface->GetCompiledData();
+
+	//Interface is never used so do not create any instance data.
+	if (CompiledData.UsedByCPU() == false && CompiledData.UsedByGPU() == false)
+	{
+		UE_LOG(LogNiagara, Warning, TEXT("Data Channel Interface is being initialized but it is never used.\nSystem: %s\nInterface: %s"), *Instance->GetSystem()->GetFullName(), *Interface->GetFullName());
+		return true;
+	}
+
+	//TODO: Reads directly from a local writer DI.
+	//For local readers, we find the source DI inside the same system and bind our functions to it's data layout.
+// 		if (Interface->Scope == ENiagaraDataChannelScope::Local)
+// 		{
+// 			//Find our source DI and it's inst data if needed.
+// 			//We have generated our compiled data already in PostCompile but we have to get the actual DI and inst data at runtime as they may not be the same as compile time.
+// 			//Though the layout of the data should be the same etc.
+// 			if (SourceDI.IsValid() == false && Interface->Source != NAME_None)
+// 			{
+// 				//Should not be able to bind to ourselves!
+// 				check(Interface->GetFName() != Interface->Source);
+// 
+// 				UNiagaraDataInterface* FoundDI = nullptr;
+// 				auto FindSourceDI = [&](const FNiagaraDataInterfaceUtilities::FDataInterfaceUsageContext& Context)
+// 				{
+// 					if (Context.Variable.GetName() == Interface->Source)
+// 					{
+// 						check(FoundDI == nullptr);
+// 						FoundDI = Context.DataInterface;
+// 						return false;
+// 					}
+// 					return true;
+// 				};
+// 
+// 				if (UNiagaraSystem* System = Interface->GetTypedOuter<UNiagaraSystem>())
+// 				{
+// 					FNiagaraDataInterfaceUtilities::FDataInterfaceSearchFilter Filter;
+// 					Filter.bIncludeInternal = true;
+// 					FNiagaraDataInterfaceUtilities::ForEachDataInterface(System, Filter, FindSourceDI);
+// 
+// 
+// 					if (UNiagaraDataInterfaceDataChannelWrite* SrcWriter = Cast<UNiagaraDataInterfaceDataChannelWrite>(FoundDI))
+// 					{
+// 						SourceDI = SrcWriter;
+// 						SourceInstData = static_cast<FNDIDataChannelWriteInstanceData*>(Instance->FindDataInterfaceInstanceData(FoundDI));
+// 
+// 						uint32 NumFuncs = CompiledData.GetFunctionInfo().Num();
+// 						FuncToDataSetBindingInfo.Reset(NumFuncs);
+// 						for(const FNDIDataChannelFunctionInfo& FuncInfo : CompiledData.GetFunctionInfo())
+// 						{
+// 							FuncToDataSetBindingInfo.Add(FNDIDataChannelLayoutManager::Get().GetLayoutInfo(FuncInfo, SourceDI->GetCompiledData().DataLayout));
+// 						}
+// 					}
+// 					else
+// 					{
+// 						UE_LOG(LogNiagara, Warning, TEXT("Failed to find the Source DataChannel Writer %s for DataChannel Reader %s"), *Interface->Source.ToString(), *Interface->GetFName().ToString());
+// 						return false;
+// 					}
+// 				}				
+// 			}
+// 		}
+// 		else
+	{
+		//For external readers, we find the DataChannel channel in the current world and bind our functions to it's data layout so we can read directly from the channel data.
+		//check(Interface->Scope == ENiagaraDataChannelScope::World);
+
+		//Grab our DataChannel channel and init the compiled data if needed.
+		UNiagaraDataChannelHandler* DataChannelPtr = DataChannel.Get();
+		if (DataChannelPtr == nullptr)
+		{
+			//DataChannel channel we had could have been destroy so clear these.
+			ExternalCPU = nullptr;
+			ExternalGPU = nullptr;
+			ChachedDataSetLayoutHash = INDEX_NONE;
+			if (UNiagaraSystem* System = Interface->GetTypedOuter<UNiagaraSystem>())
+			{
+				UWorld* World = Instance->GetWorld();
+				if (FNiagaraWorldManager* WorldMan = FNiagaraWorldManager::Get(World))
+				{
+					if (UNiagaraDataChannelHandler* NewChannelHandler = WorldMan->GetDataChannelManager().FindDataChannelHandler(Interface->Channel.ChannelName))
+					{
+						DataChannelPtr = NewChannelHandler;
+						DataChannel = NewChannelHandler;
+					}
+					else
+					{
+						UE_LOG(LogNiagara, Warning, TEXT("Failed to find or add Naigara DataChannel Channel: %s"), *Interface->Channel.ChannelName.ToString());
+						return false;
+					}
+				}
+			}
+		}
+
+		check(ExternalCPU == nullptr);
+
+		//Grab the world DataChannel data if we're reading from there.	
+		if (DataChannelPtr)
+		{
+			const FNiagaraDataSetCompiledData* CPUSourceDataCompiledData = nullptr;
+			const FNiagaraDataSetCompiledData* GPUSourceDataCompiledData = nullptr;
+			if (CompiledData.UsedByCPU())
+			{
+				//TODO: Automatically modify tick group if we have DIs that require current frame info?
+				DataChannelPtr->GetData(Instance, ExternalCPU, Interface->bReadCurrentFrame == false);
+				CPUSourceDataCompiledData = &ExternalCPU->GetOwner()->GetCompiledData();
+			}
+
+			if (CompiledData.UsedByGPU())
+			{
+				//TODO: Detect in post compile if we need GPU/CPU or both.
+				DataChannelPtr->GetDataGPU(Instance, ExternalGPU);
+				GPUSourceDataCompiledData = &ExternalGPU->GetCompiledData();
+			}
+
+			const FNiagaraDataSetCompiledData* SourceDataCompiledData = CPUSourceDataCompiledData ? CPUSourceDataCompiledData : GPUSourceDataCompiledData;
+			check(SourceDataCompiledData);
+			check(CPUSourceDataCompiledData == nullptr || GPUSourceDataCompiledData == nullptr || CPUSourceDataCompiledData->GetLayoutHash() == GPUSourceDataCompiledData->GetLayoutHash());
+
+			uint64 SourceDataLayoutHash = SourceDataCompiledData->GetLayoutHash();
+			bool bChanged = SourceDataLayoutHash != ChachedDataSetLayoutHash;
+
+			//If our CPU or GPU source data has changed then regenerate our binding info.
+			//TODO: Multi-source buffer support. 
+			//TODO: Variable input layout support. i.e. allow source systems to publish their particle buffers without the need for a separate write.
+			if (bChanged)
+			{
+				ChachedDataSetLayoutHash = SourceDataLayoutHash;
+
+				//We can likely be more targeted here.
+				//Could probably only update the RT when the GPU data changes and only update the bindings if the function hashes change etc.
+				bUpdateRTData = CompiledData.UsedByGPU();
+				int32 NumFuncs = CompiledData.GetFunctionInfo().Num();
+				FuncToDataSetBindingInfo.SetNum(NumFuncs);
+				//FuncToDataSetLayoutKeys.SetNumZeroed(NumFuncs);
+				for (int32 BindingIdx = 0; BindingIdx < NumFuncs; ++BindingIdx)
+				{
+					const FNDIDataChannelFunctionInfo& FuncInfo = CompiledData.GetFunctionInfo()[BindingIdx];
+
+					FNDIDataChannel_FuncToDataSetBindingPtr& BindingPtr = FuncToDataSetBindingInfo[BindingIdx];
+					BindingPtr = FNDIDataChannelLayoutManager::Get().GetLayoutInfo(FuncInfo, *SourceDataCompiledData/*, FuncToDataSetLayoutKeys[BindingIdx]*/);
+				}
+			}
+		}
+	}
+
+
+	if (!DataChannel.IsValid() /*&& !SourceDI.IsValid()*/)//TODO: Local reads
+	{
+		//No valid source for data so bail
+		return false;
+	}
+
+	//Verify we have valid binding info. If not, we have to bail as we cannot properly parse the vm bytecode.
+	if (FuncToDataSetBindingInfo.Num() != Interface->GetCompiledData().GetFunctionInfo().Num())
+	{
+		return false;
+	}
+
+	for (const FNDIDataChannel_FuncToDataSetBindingPtr& FuncBinding : FuncToDataSetBindingInfo)
+	{
+		if (FuncBinding.IsValid() == false || FuncBinding->IsValid() == false)
+		{
+			return false;
+		}
+	}
+
+	return true;
+}
+
+bool FNDIDataChannelReadInstanceData::PostTick(UNiagaraDataInterfaceDataChannelRead* Interface, FNiagaraSystemInstance* Instance)
+{
+	if (ExternalCPU)
+	{
+		//TODO: Create ref object to do this for us and make Add/Release ref private.
+		ExternalCPU->ReleaseReadRef();
+		ExternalCPU = nullptr;
+	}
+
+	return true;
+}
+
+//////////////////////////////////////////////////////////////////////////
+
+UNiagaraDataInterfaceDataChannelRead::UNiagaraDataInterfaceDataChannelRead(FObjectInitializer const& ObjectInitializer)
+	: Super(ObjectInitializer)
+{
+	Proxy.Reset(new FNiagaraDataInterfaceProxy_DataChannelRead());
+}
+
+void UNiagaraDataInterfaceDataChannelRead::PostInitProperties()
+{
+	Super::PostInitProperties();
+
+	if (HasAnyFlags(RF_ClassDefaultObject) && INiagaraModule::DataChannelsEnabled())
+	{
+		ENiagaraTypeRegistryFlags Flags = ENiagaraTypeRegistryFlags::AllowAnyVariable | ENiagaraTypeRegistryFlags::AllowParameter;
+		FNiagaraTypeRegistry::Register(FNiagaraTypeDefinition(GetClass()), Flags);
+	}
+}
+
+void UNiagaraDataInterfaceDataChannelRead::BeginDestroy()
+{
+	Super::BeginDestroy();
+}
+
+
+bool UNiagaraDataInterfaceDataChannelRead::InitPerInstanceData(void* PerInstanceData, FNiagaraSystemInstance* SystemInstance)
+{
+	FNDIDataChannelReadInstanceData* InstanceData = new (PerInstanceData) FNDIDataChannelReadInstanceData;
+
+	if (INiagaraModule::DataChannelsEnabled() == false)
+	{
+		return false;
+	}
+
+	if (InstanceData->Init(this, SystemInstance) == false)
+	{
+		return false;
+	}
+
+	return true;
+}
+
+void UNiagaraDataInterfaceDataChannelRead::DestroyPerInstanceData(void* PerInstanceData, FNiagaraSystemInstance* SystemInstance)
+{
+	FNDIDataChannelReadInstanceData* InstanceData = static_cast<FNDIDataChannelReadInstanceData*>(PerInstanceData);
+	InstanceData->~FNDIDataChannelReadInstanceData();
+
+	ENQUEUE_RENDER_COMMAND(RemoveProxy)
+	(
+		[RT_Proxy = GetProxyAs<FNiagaraDataInterfaceProxy_DataChannelRead>(), InstanceID = SystemInstance->GetId()](FRHICommandListImmediate& CmdList)
+		{
+			RT_Proxy->SystemInstancesToProxyData_RT.Remove(InstanceID);
+		}
+	);
+}
+
+int32 UNiagaraDataInterfaceDataChannelRead::PerInstanceDataSize() const
+{
+	return sizeof(FNDIDataChannelReadInstanceData);
+}
+
+bool UNiagaraDataInterfaceDataChannelRead::PerInstanceTick(void* PerInstanceData, FNiagaraSystemInstance* SystemInstance, float DeltaSeconds)
+{
+	SCOPE_CYCLE_COUNTER(STAT_NDIDataChannelRead_Tick);
+	check(SystemInstance);
+
+	if (INiagaraModule::DataChannelsEnabled() == false)
+	{
+		return true;
+	}
+	
+	FNDIDataChannelReadInstanceData* InstanceData = static_cast<FNDIDataChannelReadInstanceData*>(PerInstanceData);
+	if (!InstanceData)
+	{
+		return true;
+	}
+
+	if (InstanceData->Tick(this, SystemInstance) == false)
+	{
+		return true;
+	}
+
+	return false;
+}
+
+bool UNiagaraDataInterfaceDataChannelRead::PerInstanceTickPostSimulate(void* PerInstanceData, FNiagaraSystemInstance* SystemInstance, float DeltaSeconds)
+{
+	SCOPE_CYCLE_COUNTER(STAT_NDIDataChannelRead_PostTick);
+	check(SystemInstance);
+
+	if (INiagaraModule::DataChannelsEnabled() == false)
+	{
+		return true;
+	}
+
+	FNDIDataChannelReadInstanceData* InstanceData = static_cast<FNDIDataChannelReadInstanceData*>(PerInstanceData);
+	if (!InstanceData)
+	{
+		return true;
+	}
+
+	if (InstanceData->PostTick(this, SystemInstance) == false)
+	{
+		return true;
+	}
+
+	return false;
+}
+
+void UNiagaraDataInterfaceDataChannelRead::ProvidePerInstanceDataForRenderThread(void* DataForRenderThread, void* PerInstanceData, const FNiagaraSystemInstanceID& SystemInstance)
+{
+	const FNDIDataChannelReadInstanceData& SourceData = *reinterpret_cast<const FNDIDataChannelReadInstanceData*>(PerInstanceData);
+	FNDIDataChannelReadInstanceData_RT* TargetData = new(DataForRenderThread) FNDIDataChannelReadInstanceData_RT();
+
+	//Always update the dataset, this may change without triggering a full update if it's layout is the same.
+	TargetData->ChannelDataSet = SourceData.ExternalGPU;
+
+	if (SourceData.bUpdateRTData && INiagaraModule::DataChannelsEnabled())
+	{
+		SourceData.bUpdateRTData = false;
+
+		TargetData->bHasUpdate = true;
+
+		//For every GPU script, we append it's parameter access info to the table.
+		TargetData->GPUScriptParameterTableOffsets.Reset();
+		constexpr int32 ElemsPerParam = 3;
+		TargetData->GPUScriptParameterOffsetTable.Reset(CompiledData.GetTotalParams() * ElemsPerParam);
+		for (auto& GPUParameterAccessInfoPair : CompiledData.GetGPUScriptParameterInfos())
+		{
+			const FNDIDataChannel_GPUScriptParameterAccessInfo& ParamAccessInfo = GPUParameterAccessInfoPair.Value;
+
+			//First get the offset for this script in the table.
+			TargetData->GPUScriptParameterTableOffsets.FindOrAdd(GPUParameterAccessInfoPair.Key) = TargetData->GPUScriptParameterOffsetTable.Num();
+
+			//Now fill the table for this script
+			for (const FNiagaraVariableBase& Param : ParamAccessInfo.SortedParameters)
+			{
+				if (const FNiagaraVariableLayoutInfo* LayoutInfo = SourceData.ExternalGPU->GetVariableLayout(Param))
+				{
+					TargetData->GPUScriptParameterOffsetTable.Add(LayoutInfo->GetNumFloatComponents() > 0 ? LayoutInfo->FloatComponentStart : INDEX_NONE);
+					TargetData->GPUScriptParameterOffsetTable.Add(LayoutInfo->GetNumInt32Components() > 0 ? LayoutInfo->Int32ComponentStart : INDEX_NONE);
+					TargetData->GPUScriptParameterOffsetTable.Add(LayoutInfo->GetNumHalfComponents() > 0 ? LayoutInfo->HalfComponentStart : INDEX_NONE);
+				}
+				else
+				{
+					TargetData->GPUScriptParameterOffsetTable.Add(INDEX_NONE);
+					TargetData->GPUScriptParameterOffsetTable.Add(INDEX_NONE);
+					TargetData->GPUScriptParameterOffsetTable.Add(INDEX_NONE);
+				}
+			}
+		}
+	}
+}
+
+void UNiagaraDataInterfaceDataChannelRead::GetEmitterDependencies(UNiagaraSystem* Asset, TArray<FVersionedNiagaraEmitter>& Dependencies)const
+{
+	//TODO: Local support.
+	// When reading directly from a local writer DI we modify the tick order so the readers come after writers.
+	//Find our source DI and add a dependency for any emitter that writes to it.
+// 	if (UNiagaraDataInterfaceDataChannelWrite* SourceDI = FindSourceDI())
+// 	{
+// 		auto HandleVMFunc = [&](const UNiagaraScript* Script, const FVMExternalFunctionBindingInfo& BindingInfo)
+// 		{
+// 			if (UNiagaraEmitter* Emitter = Script->GetTypedOuter<UNiagaraEmitter>())
+// 			{
+// 				if (Script->IsParticleScript())
+// 				{
+// 					for (const FNiagaraEmitterHandle& EmitterHandle : Asset->GetEmitterHandles())
+// 					{
+// 						FVersionedNiagaraEmitter EmitterInstance = EmitterHandle.GetInstance();
+// 						if (EmitterInstance.Emitter.Get() == Emitter)
+// 						{
+// 							Dependencies.AddUnique(EmitterInstance);
+// 						}
+// 					}
+// 				}
+// 			}
+// 			return true;
+// 		};
+// 		if (SourceDI)
+// 		{
+// 			FNiagaraDataInterfaceUtilities::ForEachVMFunctionEquals(SourceDI, Asset, HandleVMFunc);
+// 		}
+// 	}
+}
+
+#if WITH_EDITORONLY_DATA
+
+void UNiagaraDataInterfaceDataChannelRead::PostCompile()
+{
+	UNiagaraSystem* OwnerSystem = GetTypedOuter<UNiagaraSystem>();
+	CompiledData.Init(OwnerSystem, this);
+
+	if (const UNiagaraDataChannel* DataChannel = UNiagaraDataChannelDefinitions::FindDataChannel(Channel.ChannelName))
+	{
+		OwnerSystem->RegisterDataChannelUse(DataChannel);
+	}
+}
+
+#endif
+
+
+#if WITH_EDITOR	
+
+void UNiagaraDataInterfaceDataChannelRead::GetFeedback(UNiagaraSystem* InAsset, UNiagaraComponent* InComponent, TArray<FNiagaraDataInterfaceError>& OutErrors, TArray<FNiagaraDataInterfaceFeedback>& OutWarnings, TArray<FNiagaraDataInterfaceFeedback>& OutInfo)
+{
+	//We have to post compile here manually as these graph node DIs don't get done elsewhere but we rely on the post compile data for reporting errors.
+	PostCompile();
+
+	Super::GetFeedback(InAsset, InComponent, OutErrors, OutWarnings, OutInfo);
+
+	//if(Scope == ENiagaraDataChannelScope::World)
+	{
+		if (const UNiagaraDataChannel* DataChannel = UNiagaraDataChannelDefinitions::FindDataChannel(Channel.ChannelName))
+		{
+			//Ensure the data channel contains all the parameters this function is requesting.
+			TConstArrayView<FNiagaraVariable> ChannelVars = DataChannel->GetVariables();
+			for (const FNDIDataChannelFunctionInfo& FuncInfo : CompiledData.GetFunctionInfo())
+			{
+				auto VerifyChannelContainsParams = [&](const TArray<FNiagaraVariableBase>& Parameters)
+				{
+					for (const FNiagaraVariableBase& FuncParam : Parameters)
+					{
+						bool bParamFound = false;
+						for (const FNiagaraVariable& ChannelVar : ChannelVars)
+						{
+							//We have to convert each channel var to SWC for comparrison with the function variables as there is no reliable way to go back from the SWC function var to the originating LWC var.
+							UScriptStruct* ChannelSWCStruct = FNiagaraTypeHelper::GetSWCStruct(ChannelVar.GetType().GetScriptStruct());
+							if (ChannelSWCStruct)
+							{
+								FNiagaraTypeDefinition SWCType(ChannelSWCStruct, FNiagaraTypeDefinition::EAllowUnfriendlyStruct::Deny);
+								FNiagaraVariable SWCVar(SWCType, ChannelVar.GetName());
+								if (SWCVar == FuncParam)
+								{
+									bParamFound = true;
+									break;
+								}
+							}
+						}
+
+						if (bParamFound == false)
+						{
+							OutErrors.Emplace(FText::Format(LOCTEXT("FuncParamMissingFromDataChannelErrorFmt", "Accessing variable {0} {1} which does not exist in Data Channel {2}."), FuncParam.GetType().GetNameText(), FText::FromName(FuncParam.GetName()), FText::FromName(Channel.ChannelName)),
+								LOCTEXT("FuncParamMissingFromDataChannelErrorSummaryFmt", "Data Channel DI function is accessing and invalid parameter."),
+								FNiagaraDataInterfaceFix());
+						}
+					}
+				};
+				VerifyChannelContainsParams(FuncInfo.Inputs);
+				VerifyChannelContainsParams(FuncInfo.Outputs);
+			}
+		}
+		else
+		{
+			OutErrors.Emplace(FText::Format(LOCTEXT("DataChannelDoesNotExistErrorFmt", "Data Channel {0} does not exist. It may have been deleted."), FText::FromName(Channel.ChannelName)),
+				LOCTEXT("DataChannelDoesNotExistErrorSummaryFmt", "Data Channel DI is accesssinga a Data Channel that doesn't exist."),
+				FNiagaraDataInterfaceFix());
+		}
+	}
+	//TODO: Local support.
+// 	else
+// 	{
+// 		check(Scope == ENiagaraDataChannelScope::Local);
+// 		UNiagaraDataInterfaceDataChannelWrite* SourceDI = FindSourceDI();
+// 		if (SourceDI == nullptr)
+// 		{
+// 			OutErrors.Emplace(FText::Format(LOCTEXT("DataChannelCouldNotFindSourceErrorFmt", "Could not find local Data Channel Writer DI with name {0}."), FText::FromName(Source)),
+// 				LOCTEXT("DataChannelCouldNotFindSourceErrorSummaryFmt", "Could not find source Data Channel Writer DI for local Data Channel Reader DI."),
+// 				FNiagaraDataInterfaceFix());
+// 		}
+// 	}
+}
+
+void UNiagaraDataInterfaceDataChannelRead::ValidateFunction(const FNiagaraFunctionSignature& Function, TArray<FText>& OutValidationErrors)
+{
+	Super::ValidateFunction(Function, OutValidationErrors);
+
+	//It would be great to be able to validate the parameters on the function calls here but this is only called on the DI CDO. We don't have the context of which data channel we'll be accessing.
+	//The translator should have all the required data to use the actual DIs when validating functions. We just need to do some wrangling to pull it from the pre compiled data correctly.
+	//This would probably also allow us to actually call hlsl generation functions on the actual DIs rather than their CDOs. Which would allow for a bunch of better optimized code gen for things like fluids.
+	//TODO!!!
+}
+
+#endif
+
+
+UNiagaraDataInterfaceDataChannelWrite* UNiagaraDataInterfaceDataChannelRead::FindSourceDI()const
+{
+	//TODO: Local Read/Write?
+// 	if (Scope == ENiagaraDataChannelScope::Local)
+// 	{
+// 		UNiagaraDataInterface* FoundDI = nullptr;
+// 		auto FindSourceDI = [Interface=this, &FoundDI](const FNiagaraDataInterfaceUtilities::FDataInterfaceUsageContext& Context)
+// 		{
+// 			if (Context.Variable.GetName() == Interface->Source)
+// 			{
+// 				check(FoundDI == nullptr);				
+// 				FoundDI = Context.DataInterface;
+// 				return false;
+// 			}
+// 			return true;
+// 		};
+// 
+// 		if (UNiagaraSystem* System = GetTypedOuter<UNiagaraSystem>())
+// 		{
+// 			FNiagaraDataInterfaceUtilities::FDataInterfaceSearchFilter Filter;
+// 			Filter.bIncludeInternal = true;
+// 			FNiagaraDataInterfaceUtilities::ForEachDataInterface(System, Filter, FindSourceDI);
+// 
+// 			if (UNiagaraDataInterfaceDataChannelWrite* SrcWriter = Cast<UNiagaraDataInterfaceDataChannelWrite>(FoundDI))
+// 			{
+// 				return SrcWriter;
+// 			}
+// 		}
+// 	}
+	return nullptr;
+}
+
+bool UNiagaraDataInterfaceDataChannelRead::Equals(const UNiagaraDataInterface* Other)const
+{
+	//UE_LOG(LogNiagara, Warning, TEXT("Checking Equality DCRead DI %s == %s"), *GetPathName(), *Other->GetPathName());
+	if (const UNiagaraDataInterfaceDataChannelRead* OtherTyped = CastChecked<UNiagaraDataInterfaceDataChannelRead>(Other))
+	{
+		if (Super::Equals(Other) &&
+			//Scope == OtherTyped->Scope &&
+			//Source == OtherTyped->Source &&
+			Channel.ChannelName == OtherTyped->Channel.ChannelName &&
+			bReadCurrentFrame == OtherTyped->bReadCurrentFrame)
+		{
+			return true;
+		}
+	}
+
+	return false;
+}
+
+bool UNiagaraDataInterfaceDataChannelRead::CopyToInternal(UNiagaraDataInterface* Destination)const
+{
+	//UE_LOG(LogNiagara, Warning, TEXT("Coping DCRead DI %s --> %s"), *GetPathName(), *Destination->GetPathName());
+	if (!Super::CopyToInternal(Destination))
+	{
+		return false;
+	}
+
+	if (UNiagaraDataInterfaceDataChannelRead* DestTyped = CastChecked<UNiagaraDataInterfaceDataChannelRead>(Destination))
+	{
+		//DestTyped->Scope = Scope;
+		//DestTyped->Source = Source;
+		DestTyped->Channel = Channel;
+		DestTyped->CompiledData = CompiledData;
+		DestTyped->bReadCurrentFrame = bReadCurrentFrame;
+		return true;
+	}
+
+	return false;
+}
+
+void UNiagaraDataInterfaceDataChannelRead::GetFunctions(TArray<FNiagaraFunctionSignature>& OutFunctions)
+{
+	{
+		FNiagaraFunctionSignature Sig;
+		Sig.Name = NDIDataChannelReadLocal::NumName;
+#if WITH_EDITORONLY_DATA
+		Sig.Description = LOCTEXT("NumFunctionDescription", "Returns the current number of DataChannel accessible by this interface.");
+#endif
+		Sig.bMemberFunction = true;
+		Sig.bExperimental = true;
+		Sig.AddInput(FNiagaraVariable(FNiagaraTypeDefinition(GetClass()), TEXT("DataChannel interface")));
+		Sig.AddOutput(FNiagaraVariable(FNiagaraTypeDefinition::GetIntDef(), TEXT("Num")));
+		OutFunctions.Add(Sig);
+	}
+	{
+		FNiagaraFunctionSignature Sig;
+		Sig.Name = NDIDataChannelReadLocal::ReadName;
+#if WITH_EDITORONLY_DATA
+		Sig.Description = LOCTEXT("ReadFunctionDescription", "Reads DataChannel data at a specific index. Any values we read that are not in the DataChannel data are set to their default values. Returns success if there was a valid DataChannel to read from.");
+#endif
+		Sig.bMemberFunction = true;
+		Sig.bExperimental = true;
+		Sig.AddInput(FNiagaraVariable(FNiagaraTypeDefinition(GetClass()), TEXT("DataChannel interface")));
+		Sig.AddInput(FNiagaraVariable(FNiagaraTypeDefinition::GetIntDef(), TEXT("Index")), LOCTEXT("ConsumeIndexInputDesc", "The index to read."));
+		Sig.AddOutput(FNiagaraVariable(FNiagaraTypeDefinition::GetBoolDef(), TEXT("Success")), LOCTEXT("ConsumeSuccessOutputDesc", "True if all reads succeeded."));
+		Sig.RequiredOutputs = Sig.Outputs.Num();//The user defines what we read in the graph.
+		OutFunctions.Add(Sig);
+	}
+	{
+		FNiagaraFunctionSignature Sig;
+		Sig.Name = NDIDataChannelReadLocal::ConsumeName;
+#if WITH_EDITORONLY_DATA
+		Sig.Description = LOCTEXT("ConsumeFunctionDescription", "Consumes an DataChannel from the end of the DataChannel array and reads the specified values. Any values we read that are not in the DataChannel data are set to their default values. Returns success if an DataChannel was available to pop.");
+#endif
+		Sig.bMemberFunction = true;
+		Sig.bExperimental = true;
+		Sig.AddInput(FNiagaraVariable(FNiagaraTypeDefinition(GetClass()), TEXT("DataChannel interface")));
+		Sig.AddInput(FNiagaraVariable(FNiagaraTypeDefinition::GetBoolDef(), TEXT("Consume")), LOCTEXT("ConsumeInputDesc", "True if this instance (particle/emitter etc) should consume data from the data channel in this call."));
+		Sig.AddOutput(FNiagaraVariable(FNiagaraTypeDefinition::GetBoolDef(), TEXT("Success")), LOCTEXT("ConsumeSuccessOutputDesc", "True if all reads succeeded."));
+		Sig.AddOutput(FNiagaraVariable(FNiagaraTypeDefinition::GetIntDef(), TEXT("Index")), LOCTEXT("ConsumeIndexOutputDesc", "The index we actually read from. If reading failed this can be -1. This allows subsequent reads of the data channel at this index."));
+		Sig.RequiredOutputs = Sig.Outputs.Num();//The user defines what we read in the graph.
+		OutFunctions.Add(Sig);
+	}
+}
+
+void UNiagaraDataInterfaceDataChannelRead::GetVMExternalFunction(const FVMExternalFunctionBindingInfo& BindingInfo, void* InstanceData, FVMExternalFunction& OutFunc)
+{
+	if (BindingInfo.Name == NDIDataChannelReadLocal::NumName)
+	{
+		OutFunc = FVMExternalFunction::CreateLambda([this](FVectorVMExternalFunctionContext& Context) { this->Num(Context); });
+	}
+	else
+	{
+		int32 FuncIndex = CompiledData.FindFunctionInfoIndex(BindingInfo.Name, BindingInfo.VariadicInputs, BindingInfo.VariadicOutputs);
+		if (FuncIndex != INDEX_NONE)
+		{
+			if (BindingInfo.Name == NDIDataChannelReadLocal::ReadName)
+			{
+				OutFunc = FVMExternalFunction::CreateLambda([this, FuncIndex](FVectorVMExternalFunctionContext& Context) { this->Read(Context, FuncIndex); });
+			}
+			else if (BindingInfo.Name == NDIDataChannelReadLocal::ConsumeName)
+			{
+				OutFunc = FVMExternalFunction::CreateLambda([this, FuncIndex](FVectorVMExternalFunctionContext& Context) { this->Consume(Context, FuncIndex); });
+			}
+			else
+			{
+				UE_LOG(LogTemp, Display, TEXT("Could not find data interface external function in %s. Received Name: %s"), *GetPathNameSafe(this), *BindingInfo.Name.ToString());
+			}
+		}
+		else
+		{
+			UE_LOG(LogTemp, Display, TEXT("Could not find data interface external function in %s. Received Name: %s"), *GetPathNameSafe(this), *BindingInfo.Name.ToString());
+		}
+	}
+}
+
+template<int32 EXPECTED_NUM_INPUTS>
+struct FNDIVariadicOutputHandler
+{
+	TArray<VectorVM::FExternalFuncRegisterHandler<float>, TInlineAllocator<EXPECTED_NUM_INPUTS>> FloatOutputs;
+	TArray<VectorVM::FExternalFuncRegisterHandler<int32>, TInlineAllocator<EXPECTED_NUM_INPUTS>> IntOutputs;
+	TArray<VectorVM::FExternalFuncRegisterHandler<FFloat16>, TInlineAllocator<EXPECTED_NUM_INPUTS>> HalfOutputs;
+
+	FNDIVariadicOutputHandler(FVectorVMExternalFunctionContext& Context, const FNDIDataChannel_FuncToDataSetBindingPtr& BindingPtr)
+	{
+		//Parse the VM bytecode inputs in order, mapping them to the correct DataChannel data 
+		FloatOutputs.Reserve(BindingPtr->NumFloatComponents);
+		IntOutputs.Reserve(BindingPtr->NumInt32Components);
+		HalfOutputs.Reserve(BindingPtr->NumHalfComponents);
+		for (FNDIDataChannelRegisterBinding& VMBinding : BindingPtr->VMRegisterBindings)
+		{
+			if (VMBinding.DataType == (int32)ENiagaraBaseTypes::Float)
+			{
+				FloatOutputs.Emplace(Context);
+			}
+			else if (VMBinding.DataType == (int32)ENiagaraBaseTypes::Int32 || VMBinding.DataType == (int32)ENiagaraBaseTypes::Bool)
+			{
+				IntOutputs.Emplace(Context);
+			}
+			else if (VMBinding.DataType == (int32)ENiagaraBaseTypes::Half)
+			{
+				HalfOutputs.Emplace(Context);
+			}
+			else
+			{
+				checkf(false, TEXT("Didn't find a binding for this function input. Likely an error in the building of the binding data and the Float/Int/Half Bindings arrays."));
+			}
+		}
+	}
+
+	bool Process(FNiagaraDataBuffer* Data, uint32 Index, const FNDIDataChannel_FuncToDataSetBindingPtr& BindingPtr)
+	{
+		checkSlow(Data);
+
+		//Process this index.
+		//TODO: Can be optimized to pull N values out at a time, reducing the binding/lookup overhead.
+		if ((uint32)Index < Data->GetNumInstances())
+		{
+			for (FNDIDataChannelRegisterBinding& VMBinding : BindingPtr->VMRegisterBindings)
+			{
+				if (VMBinding.DataType == (int32)ENiagaraBaseTypes::Float)
+				{
+					if (VMBinding.DataSetRegisterIndex == INDEX_NONE)
+					{
+						*FloatOutputs[VMBinding.FunctionRegisterIndex].GetDestAndAdvance() = 0.0f;
+					}
+					else
+					{
+						*FloatOutputs[VMBinding.FunctionRegisterIndex].GetDestAndAdvance() = *Data->GetInstancePtrFloat(VMBinding.DataSetRegisterIndex, (uint32)Index);
+					}
+				}
+				else if (VMBinding.DataType == (int32)ENiagaraBaseTypes::Int32 || VMBinding.DataType == (int32)ENiagaraBaseTypes::Bool)
+				{
+					if (VMBinding.DataSetRegisterIndex == INDEX_NONE)
+					{
+						*IntOutputs[VMBinding.FunctionRegisterIndex].GetDestAndAdvance() = 0;
+					}
+					else
+					{
+						*IntOutputs[VMBinding.FunctionRegisterIndex].GetDestAndAdvance() = *Data->GetInstancePtrInt32(VMBinding.DataSetRegisterIndex, (uint32)Index);
+					}
+				}
+				else if (VMBinding.DataType == (int32)ENiagaraBaseTypes::Half)
+				{
+					if (VMBinding.DataSetRegisterIndex == INDEX_NONE)
+					{
+						*HalfOutputs[VMBinding.FunctionRegisterIndex].GetDestAndAdvance() = FFloat16(0.0f);
+					}
+					else
+					{
+						*HalfOutputs[VMBinding.FunctionRegisterIndex].GetDestAndAdvance() = *Data->GetInstancePtrHalf(VMBinding.DataSetRegisterIndex, (uint32)Index);
+					}
+				}
+			}
+			return true;
+		}
+
+		Fallback(1);
+
+		return false;
+	}
+
+	void Fallback(int32 Count)
+	{
+		for (int32 OutIdx = 0; OutIdx < FloatOutputs.Num(); ++OutIdx)
+		{
+			FMemory::Memzero(FloatOutputs[OutIdx].GetDest(), sizeof(float) * Count);
+			FloatOutputs[OutIdx].Advance(Count);
+		}
+
+		for (int32 OutIdx = 0; OutIdx < IntOutputs.Num(); ++OutIdx)
+		{
+			FMemory::Memzero(IntOutputs[OutIdx].GetDest(), sizeof(int32) * Count);
+			IntOutputs[OutIdx].Advance(Count);
+		}
+
+		for (int32 OutIdx = 0; OutIdx < HalfOutputs.Num(); ++OutIdx)
+		{
+			FMemory::Memzero(HalfOutputs[OutIdx].GetDest(), sizeof(FFloat16) * Count);
+			HalfOutputs[OutIdx].Advance(Count);
+		}
+	}
+};
+
+void UNiagaraDataInterfaceDataChannelRead::Num(FVectorVMExternalFunctionContext& Context)
+{
+	VectorVM::FUserPtrHandler<FNDIDataChannelReadInstanceData> InstData(Context);
+
+	FNDIOutputParam<int32> OutNum(Context);
+
+	FNiagaraDataBuffer* Buffer = InstData->GetReadBufferCPU();
+
+	int32 Num = 0;
+	if (Buffer && INiagaraModule::DataChannelsEnabled())
+	{
+		Num = (int32)Buffer->GetNumInstances();
+	}
+
+	for (int32 i = 0; i < Context.GetNumInstances(); ++i)
+	{
+		OutNum.SetAndAdvance(Num);
+	}
+}
+
+void UNiagaraDataInterfaceDataChannelRead::Read(FVectorVMExternalFunctionContext& Context, int32 FuncIdx)
+{
+	SCOPE_CYCLE_COUNTER(STAT_NDIDataChannelRead_Read);
+	VectorVM::FUserPtrHandler<FNDIDataChannelReadInstanceData> InstData(Context);
+	FNDIInputParam<int32> InIndex(Context);
+
+	FNDIOutputParam<bool> OutSuccess(Context);
+
+	const FNDIDataChannelFunctionInfo& FuncInfo = CompiledData.GetFunctionInfo()[FuncIdx];
+	const FNDIDataChannel_FuncToDataSetBindingPtr& BindingInfo = InstData->FuncToDataSetBindingInfo.IsValidIndex(FuncIdx) ? InstData->FuncToDataSetBindingInfo[FuncIdx] : nullptr;
+	FNDIVariadicOutputHandler<16> VariadicOutputs(Context, BindingInfo);//TODO: Make static / avoid allocation
+
+	FNiagaraDataBuffer* Buffer = InstData->GetReadBufferCPU();
+	if (Buffer && BindingInfo && INiagaraModule::DataChannelsEnabled())
+	{
+		for (int32 i = 0; i < Context.GetNumInstances(); ++i)
+		{
+			int32 Index = InIndex.GetAndAdvance();
+
+			//TODO: Optimize for long runs of reads.
+			bool bSuccess = VariadicOutputs.Process(Buffer, Index, BindingInfo);
+
+			if (OutSuccess.IsValid())
+			{
+				OutSuccess.SetAndAdvance(bSuccess);
+			}
+		}
+	}
+	else
+	{
+		VariadicOutputs.Fallback(Context.GetNumInstances());
+		if (OutSuccess.IsValid())
+		{
+			FMemory::Memzero(OutSuccess.Data.GetDest(), sizeof(int32) * Context.GetNumInstances());
+		}
+	}
+}
+
+void UNiagaraDataInterfaceDataChannelRead::Consume(FVectorVMExternalFunctionContext& Context, int32 FuncIdx)
+{
+	SCOPE_CYCLE_COUNTER(STAT_NDIDataChannelRead_Consume);
+	VectorVM::FUserPtrHandler<FNDIDataChannelReadInstanceData> InstData(Context);
+	FNDIInputParam<bool> InConsume(Context);
+
+	FNDIOutputParam<bool> OutSuccess(Context);
+	FNDIOutputParam<int32> OutIndex(Context);
+
+	const FNDIDataChannelFunctionInfo& FuncInfo = CompiledData.GetFunctionInfo()[FuncIdx];
+	const FNDIDataChannel_FuncToDataSetBindingPtr& BindingInfo = InstData->FuncToDataSetBindingInfo.IsValidIndex(FuncIdx) ? InstData->FuncToDataSetBindingInfo[FuncIdx] : nullptr;
+	FNDIVariadicOutputHandler<16> VariadicOutputs(Context, BindingInfo);//TODO: Make static / avoid allocation
+
+	//TODO: Optimize for constant bConsume.
+	//TODO: Optimize for long runs of bConsume==true;
+	FNiagaraDataBuffer* Buffer = InstData->GetReadBufferCPU();
+	if (Buffer && BindingInfo && INiagaraModule::DataChannelsEnabled())
+	{
+		for (int32 i = 0; i < Context.GetNumInstances(); ++i)
+		{
+			bool bConsume = InConsume.GetAndAdvance();
+
+			bool bSuccess = false;
+			int32 ConsumeIdx = InstData->ConsumeIndex;
+			if (bConsume)
+			{
+				while (ConsumeIdx < (int32)Buffer->GetNumInstances() && InstData->ConsumeIndex.compare_exchange_strong(ConsumeIdx, ConsumeIdx + 1) == false)
+				{
+					ConsumeIdx = InstData->ConsumeIndex;
+				}
+
+				bSuccess = VariadicOutputs.Process(Buffer, ConsumeIdx, BindingInfo);
+			}
+			else
+			{
+				VariadicOutputs.Fallback(1);
+			}
+
+			if (OutSuccess.IsValid())
+			{
+				OutSuccess.SetAndAdvance(bSuccess);
+			}
+			if (OutIndex.IsValid())
+			{
+				OutIndex.SetAndAdvance(bSuccess ? ConsumeIdx : INDEX_NONE);
+			}
+		}
+	}
+	else
+	{
+		VariadicOutputs.Fallback(Context.GetNumInstances());
+
+		if (OutSuccess.IsValid())
+		{
+			FMemory::Memzero(OutSuccess.Data.GetDest(), sizeof(int32) * Context.GetNumInstances());
+		}
+
+		if (OutIndex.IsValid())
+		{
+			FMemory::Memset(OutSuccess.Data.GetDest(), 0xFF, sizeof(int32) * Context.GetNumInstances());
+		}
+	}
+}
+
+#if WITH_EDITORONLY_DATA
+bool UNiagaraDataInterfaceDataChannelRead::AppendCompileHash(FNiagaraCompileHashVisitor* InVisitor) const
+{
+	bool bSuccess = Super::AppendCompileHash(InVisitor);
+	bSuccess &= InVisitor->UpdateString(TEXT("UNiagaraDataInterfaceDataChannelCommon"), GetShaderFileHash(NDIDataChannelReadLocal::CommonShaderFile, EShaderPlatform::SP_PCD3D_SM5).ToString());
+	bSuccess &= InVisitor->UpdateString(TEXT("UNiagaraDataInterfaceDataChannelRead_Common"), GetShaderFileHash(NDIDataChannelReadLocal::TemplateShaderFile_Common, EShaderPlatform::SP_PCD3D_SM5).ToString());
+	bSuccess &= InVisitor->UpdateString(TEXT("UNiagaraDataInterfaceDataChannelRead_Read"), GetShaderFileHash(NDIDataChannelReadLocal::TemplateShaderFile_Read, EShaderPlatform::SP_PCD3D_SM5).ToString());
+	bSuccess &= InVisitor->UpdateString(TEXT("UNiagaraDataInterfaceDataChannelRead_Consume"), GetShaderFileHash(NDIDataChannelReadLocal::TemplateShaderFile_Consume, EShaderPlatform::SP_PCD3D_SM5).ToString());
+
+	bSuccess &= InVisitor->UpdateShaderParameters<NDIDataChannelReadLocal::FShaderParameters>();
+	return bSuccess;
+}
+
+void UNiagaraDataInterfaceDataChannelRead::GetCommonHLSL(FString& OutHLSL)
+{
+	Super::GetCommonHLSL(OutHLSL);
+	OutHLSL.Appendf(TEXT("#include \"%s\"\n"), NDIDataChannelReadLocal::CommonShaderFile);
+}
+
+bool UNiagaraDataInterfaceDataChannelRead::GetFunctionHLSL(FNiagaraDataInterfaceHlslGenerationContext& HlslGenContext, FString& OutHLSL)
+{
+	return	HlslGenContext.GetFunctionInfo().DefinitionName == GET_FUNCTION_NAME_CHECKED(UNiagaraDataInterfaceDataChannelRead, Num) ||
+		HlslGenContext.GetFunctionInfo().DefinitionName == GET_FUNCTION_NAME_CHECKED(UNiagaraDataInterfaceDataChannelRead, Read) ||
+		HlslGenContext.GetFunctionInfo().DefinitionName == GET_FUNCTION_NAME_CHECKED(UNiagaraDataInterfaceDataChannelRead, Consume);
+}
+
+void UNiagaraDataInterfaceDataChannelRead::GetParameterDefinitionHLSL(FNiagaraDataInterfaceHlslGenerationContext& HlslGenContext, FString& OutHLSL)
+{
+	//Desc:
+	// This function iterates over all functions called for this DI in each script and builds the correct hlsl.
+	// The main part of the complexity here is dealing with Variadic function parameters.
+	// We must interrogate the script to see what parameters are actually accessed and generate hlsl accordingly.
+	// Ideally at some future point we can refactor most of this out to a utility helper that will do most of the heavy lifting.
+	// Allowing users to simply provide some details and lambdas etc to specify what exactly they'd like to do in the function body with each variadic param.
+
+	//Initially we'll do some preamble, setting up various template strings and args etc.
+
+	//Map of all arguments for various pieces of template code.
+	//We have some common code that is shared by all functions.
+	//Some code is unique for each function.
+	//Some is unique per parameter to each function.
+	//Finally there is some that is unique for each sub component of each parameter until we've hit a base type. float/2/3/4 etc.
+	TMap<FString, FStringFormatArg> HlslTemplateArgs =
+	{
+		//Common args for all functions
+		{TEXT("ParameterName"),							HlslGenContext.ParameterInfo.DataInterfaceHLSLSymbol},
+		//Per function args. These will be changed with each function written
+		{TEXT("FunctionParameters"),					FString(TEXT("FunctionParameters"))},//Function parameters written into the function signature.
+		{TEXT("PerFunctionParameterShaderCode"),		FString(TEXT("PerFunctionParameterShaderCode"))},//Combined string of all code dealing with each of the function parameters.
+		//Per function parameter args. These will be changed with each parameter written
+		{TEXT("FunctionParameterIndex"),				TEXT("FunctionParameterIndex")},//Function parameter index allowing us to look up the layout information for the correct parameter to the function.
+		{TEXT("FunctionParameterName"),					FString(TEXT("FunctionParameterName"))},//Name of this function parameter.
+		{TEXT("FunctionParameterType"),					FString(TEXT("FunctionParameterType"))},//Type of this function parameter.
+		{TEXT("PerParamRWCode"),						FString(TEXT("PerParamRWCode"))},//Code that does the actual reading or writing to the data channel buffers.
+		//Per component/base type args. These will change with every base type we I/O from the Data Channel.
+		{TEXT("FunctionParameterComponentBufferType"),	FString(TEXT("FunctionParameterComponentBufferType"))},//The actual base data buffer type being accessed by a particular DataChannel access code line. Float, Int32, Half etc
+		{TEXT("FunctionParameterComponentName"),		FString(TEXT("FunctionParameterComponentName"))},//The name/symbol of the actual member of a parameter that we can I/O from the DataChannel via a standard getter/setter.
+		{TEXT("FunctionParameterComponentType"),		FString(TEXT("FunctionParameterComponentType"))},//The type of the actual member of a parameter that we can I/O from the DataChannel via a standard getter/setter.
+	};
+
+	//Grab refs to per function args we'll change with each function written.
+	FStringFormatArg& FunctionParameters = HlslTemplateArgs.FindChecked(TEXT("FunctionParameters"));
+	FStringFormatArg& PerFunctionParameterShaderCode = HlslTemplateArgs.FindChecked(TEXT("PerFunctionParameterShaderCode"));
+
+	//Grab refs to per function args we'll change with each parameter written.
+	FStringFormatArg& FunctionParameterIndex = HlslTemplateArgs.FindChecked(TEXT("FunctionParameterIndex"));
+	FStringFormatArg& FuncParamName = HlslTemplateArgs.FindChecked(TEXT("FunctionParameterName"));
+	FStringFormatArg& FuncParamType = HlslTemplateArgs.FindChecked(TEXT("FunctionParameterType"));
+	FStringFormatArg& PerParamRWCode = HlslTemplateArgs.FindChecked(TEXT("PerParamRWCode"));
+
+	//Grab refs to per component/base type args we'll change for every parameter component we access in the Data Channel.
+	FStringFormatArg& FunctionParameterComponentBufferType = HlslTemplateArgs.FindChecked(TEXT("FunctionParameterComponentBufferType"));
+	FStringFormatArg& FunctionParameterComponentName = HlslTemplateArgs.FindChecked(TEXT("FunctionParameterComponentName"));
+	FStringFormatArg& FunctionParameterComponentType = HlslTemplateArgs.FindChecked(TEXT("FunctionParameterComponentType"));
+
+	//Template code for writing handling code for each of the function parameters.
+	//Some preamble and an replacement arg into which we write all the actual I/O with the DataChannel buffers.
+	FString PerFunctionParameterTemplate = TEXT("\n\
+		//Handling for function parameter {FunctionParameterName}\n\
+		int FloatRegisterIndex_{FunctionParameterName};\n\
+		int Int32RegisterIndex_{FunctionParameterName};\n\
+		int HalfRegisterIndex_{FunctionParameterName};\n\
+		if(GetParameterLayoutInfo_{ParameterName}({FunctionParameterIndex}, FloatRegisterIndex_{FunctionParameterName}, Int32RegisterIndex_{FunctionParameterName}, HalfRegisterIndex_{FunctionParameterName}))\n\
+		{\n\
+{PerParamRWCode}\n\
+		}\n\
+		else\n\
+		{\n\
+			bOutSuccess = false;\n\
+		}\n");
+
+	//Template code for accessing data from the Data Channel's buffers.
+	static const FString ReadDataTemplate = TEXT("			NiagaraDataChannelRead_{FunctionParameterComponentType}({ParameterName}_Data{FunctionParameterComponentBufferType}, {ParameterName}_{FunctionParameterComponentBufferType}Stride, {FunctionParameterComponentBufferType}RegisterIndex_{FunctionParameterName}, ElementIndex, {FunctionParameterComponentName});\n\t\t");
+	static const FString WriteDataTemplate = TEXT("			NiagaraDataChannelWrite_{FunctionParameterComponentType}({ParameterName}_Data{FunctionParameterComponentBufferType}, {ParameterName}_{FunctionParameterComponentBufferType}Stride, {FunctionParameterComponentBufferType}RegisterIndex_{FunctionParameterName}, ElementIndex, {FunctionParameterComponentName});\n\t\t");
+
+	//Map to store template shader code for each DI function as each is in it's own ush file. We'll load those as needed and place in here for ease of access.
+	TMap<FName, FString> FunctionTemplateMap;
+
+	//We may call the same function multiple times so avoid duplicating the same function impl.
+	TMap<uint32, FString> FunctionHlslMap;
+
+	//Now lets build some hlsl!
+	// 
+	//First build the shader code common to all functions.
+	FString CommonTemplateFile;
+	LoadShaderSourceFile(NDIDataChannelReadLocal::TemplateShaderFile_Common, EShaderPlatform::SP_PCD3D_SM5, &CommonTemplateFile, nullptr);
+	OutHLSL += FString::Format(*CommonTemplateFile, HlslTemplateArgs);
+
+	auto GetSignatureHash = [](const FNiagaraFunctionSignature& Sig)
+	{
+		uint32 Ret = GetTypeHash(Sig.Name);
+		for (const FNiagaraVariable& Input : Sig.Inputs)
+		{
+			Ret = HashCombine(Ret, GetTypeHash(Input));
+		}
+		for (const FNiagaraVariable& Output : Sig.Outputs)
+		{
+			Ret = HashCombine(Ret, GetTypeHash(Output));
+		}
+		return Ret;
+	};
+
+	TArray<FNiagaraVariableBase> ParametersAccessed;
+	//First iterate over the generated functions to gather all used parameters so we can generate the sorted parameter list for all functions called by this script.
+	for (int32 FuncIdx = 0; FuncIdx < HlslGenContext.ParameterInfo.GeneratedFunctions.Num(); ++FuncIdx)
+	{
+		const FNiagaraDataInterfaceGeneratedFunction& Func = HlslGenContext.ParameterInfo.GeneratedFunctions[FuncIdx];
+		const FNiagaraFunctionSignature& Signature = HlslGenContext.Signatures[FuncIdx];
+
+		for (int32 InputIdx = Signature.NumRequiredInputs(); InputIdx < Signature.Inputs.Num(); ++InputIdx)
+		{
+			const FNiagaraVariable& InputParam = Signature.Inputs[InputIdx];
+			ParametersAccessed.AddUnique(InputParam);
+		}
+
+		for (int32 OutputIdx = Signature.NumRequiredOutputs(); OutputIdx < Signature.Outputs.Num(); ++OutputIdx)
+		{
+			const FNiagaraVariable& OutputParam = Signature.Outputs[OutputIdx];
+			ParametersAccessed.AddUnique(OutputParam);
+		}
+	}
+	//Sort the parameters so that that generated hlsl can exactly match the runtime ordering of parameters.
+	NDIDataChannelUtilities::SortParameters(ParametersAccessed);
+
+	//Utility map for easy access of the right param offset during hlsl gen.
+	TMap<FNiagaraVariableBase, uint32> ParamOffsets;
+	for (int32 ParamIdx = 0; ParamIdx < ParametersAccessed.Num(); ++ParamIdx)
+	{
+		ParamOffsets.FindOrAdd(ParametersAccessed[ParamIdx]) = ParamIdx;
+	}
+
+	//Now to iterate on the functions and build the hlsl for each as needed.
+	for (int32 FuncIdx = 0; FuncIdx < HlslGenContext.ParameterInfo.GeneratedFunctions.Num(); ++FuncIdx)
+	{
+		const FNiagaraDataInterfaceGeneratedFunction& Func = HlslGenContext.ParameterInfo.GeneratedFunctions[FuncIdx];
+		const FNiagaraFunctionSignature& Signature = HlslGenContext.Signatures[FuncIdx];
+
+		//Init/Reset our per function hlsl template args.
+		FunctionParameters.StringValue.Reset();//Reset function parameters ready to rebuild when we iterate over the parameters.
+		PerFunctionParameterShaderCode.StringValue.Reset(); //Reset function parameters ready to rebuild when we iterate over the parameters.
+
+		//Generate the function hlsl if we've not done so already for this signature.
+		uint32 FuncHash = GetSignatureHash(Signature);
+		FString& FunctionHlsl = FunctionHlslMap.FindOrAdd(FuncHash);
+		if (FunctionHlsl.IsEmpty())
+		{
+			//Function that will recurse down a parameter's type and generate the appropriate IO code for all of it's members.
+			TFunction<void(bool, UScriptStruct*, FString&)> GenerateRWParameterCode = [&](bool bRead, UScriptStruct* Struct, FString& OutCode)
+			{
+				if (Struct == FNiagaraTypeDefinition::GetFloatStruct() || Struct == FNiagaraTypeDefinition::GetVec2Struct() || Struct == FNiagaraTypeDefinition::GetVec3Struct()
+					|| Struct == FNiagaraTypeDefinition::GetVec4Struct() || Struct == FNiagaraTypeDefinition::GetColorStruct() || Struct == FNiagaraTypeDefinition::GetQuatStruct())
+				{
+					FunctionParameterComponentBufferType = TEXT("Float");
+					FunctionParameterComponentType = HlslGenContext.GetStructHlslTypeName(FNiagaraTypeDefinition(Struct));
+					OutCode += FString::Format(bRead ? *ReadDataTemplate : *WriteDataTemplate, HlslTemplateArgs);
+				}
+				else if (Struct == FNiagaraTypeDefinition::GetIntStruct())
+				{
+					FunctionParameterComponentBufferType = TEXT("Int32");
+					FunctionParameterComponentType = HlslGenContext.GetStructHlslTypeName(FNiagaraTypeDefinition(Struct));
+					OutCode += FString::Format(bRead ? *ReadDataTemplate : *WriteDataTemplate, HlslTemplateArgs);
+				}
+				else
+				{
+					FString PropertyBaseName = FunctionParameterComponentName.StringValue;
+
+					for (TFieldIterator<FProperty> PropertyIt(Struct, EFieldIteratorFlags::IncludeSuper); PropertyIt; ++PropertyIt)
+					{
+						FunctionParameterComponentName.StringValue = PropertyBaseName;
+						const FProperty* Property = *PropertyIt;
+
+						if (Property->IsA(FFloatProperty::StaticClass()))
+						{
+							FunctionParameterComponentBufferType = TEXT("Float");
+							FunctionParameterComponentType = HlslGenContext.GetStructHlslTypeName(FNiagaraTypeDefinition::GetFloatDef());
+							FunctionParameterComponentName.StringValue += FString::Printf(TEXT(".%s"), *HlslGenContext.GetSanitizedSymbolName(Property->GetName()));
+							OutCode += FString::Format(bRead ? *ReadDataTemplate : *WriteDataTemplate, HlslTemplateArgs);
+						}
+						else if (Property->IsA(FIntProperty::StaticClass()))
+						{
+							FunctionParameterComponentBufferType = TEXT("Int32");
+							FunctionParameterComponentType = HlslGenContext.GetStructHlslTypeName(FNiagaraTypeDefinition::GetIntDef());
+							FunctionParameterComponentName.StringValue += FString::Printf(TEXT(".%s"), *HlslGenContext.GetSanitizedSymbolName(Property->GetName()));
+							OutCode += FString::Format(bRead ? *ReadDataTemplate : *WriteDataTemplate, HlslTemplateArgs);
+						}
+						else if (const FStructProperty* StructProp = CastFieldChecked<const FStructProperty>(Property))
+						{
+							FunctionParameterComponentName.StringValue += FString::Printf(TEXT(".%s"), *HlslGenContext.GetSanitizedSymbolName(Property->GetName()));
+							FunctionParameterComponentType = HlslGenContext.GetPropertyHlslTypeName(Property);
+							GenerateRWParameterCode(bRead, StructProp->Struct, OutCode);
+						}
+						else
+						{
+							check(false);
+							OutCode += FString::Printf(TEXT("Error! - DataChannel Interface encountered an invalid type when generating it's function hlsl. %s"), *Property->GetName());
+						}
+					}
+				}
+			};
+
+			int32 CurrParamIdx = 0;
+			int32 CurrVariadicParamIdx = 0;//Track current variadic param index so the shader code can tell which param's layout info to lookup.
+
+			//Iterate through the variadic inputs and output parameters building the function signature and the per parameter function body as we go.
+			static FString FirstParamPrefix = TEXT("");
+			static FString ParamPrefix = TEXT(", ");
+			static FString OutputParamPrefix = TEXT("out");
+			static FString InputParamPrefix = TEXT("in");
+			auto HandleParameters = [&](TConstArrayView<FNiagaraVariable> Parameters, int32 StartParamIdx, bool bOutput)
+			{
+				for (int32 ParamIdx = StartParamIdx; ParamIdx < Parameters.Num(); ++ParamIdx)
+				{
+					const FNiagaraVariable& Param = Parameters[ParamIdx];
+
+					//Reset or init per parameter hlsl template args.
+
+					//Grab the offset for this parameter and embed in the hlsl directly.
+					//This will correspond to an entry in the ParameterLayoutBuffer passed to the shader.
+					int32 ParamOffset = ParamOffsets.FindChecked(Param);
+					FunctionParameterIndex = LexToString(ParamOffset);
+					FuncParamName = Param.GetName().ToString();
+					FuncParamType = HlslGenContext.GetStructHlslTypeName(Param.GetType());
+					PerParamRWCode.StringValue.Reset();
+
+					bool bFirstParam = CurrParamIdx++ == 0;
+					FunctionParameters.StringValue += FString::Printf(TEXT("%s%s %s %s"),
+						bFirstParam ? *FirstParamPrefix : *ParamPrefix,
+						bOutput ? *OutputParamPrefix : *InputParamPrefix,
+						*FuncParamType.StringValue, *FuncParamName.StringValue);
+
+					//Reset per component args and build the code for per component/base type Data Channel access.
+					FunctionParameterComponentName.StringValue.Reset();
+					FunctionParameterComponentType.StringValue.Reset();
+					FunctionParameterComponentBufferType.StringValue.Reset();
+
+					//Generate the per component/base type IO code that will be used in the subsequent per paramters shader code template.
+					FunctionParameterComponentName = HlslGenContext.GetSanitizedSymbolName(Param.GetName().ToString());
+					GenerateRWParameterCode(bOutput, Param.GetType().GetScriptStruct(), PerParamRWCode.StringValue);
+
+					//Generate the whole code for this parameter and append it to the function's per parameter code.
+					//This will be used at the end in the per function template.
+					PerFunctionParameterShaderCode.StringValue += FString::Format(*PerFunctionParameterTemplate, HlslTemplateArgs);
+
+					++CurrVariadicParamIdx;
+				}
+			};
+
+			HandleParameters(Signature.Inputs, Signature.NumRequiredInputs(), false);
+			HandleParameters(Signature.Outputs, Signature.NumRequiredOutputs(), true);
+
+			//Grab the template file for the function.
+			FString& FunctionTemplateFile = FunctionTemplateMap.FindOrAdd(Signature.Name);
+			if (FunctionTemplateFile.IsEmpty())
+			{
+				//Load it if we have not previously.
+				LoadShaderSourceFile(NDIDataChannelReadLocal::GetFunctionTemplate(Signature.Name), EShaderPlatform::SP_PCD3D_SM5, &FunctionTemplateFile, nullptr);
+			}
+
+			//Finally generate the final code for this function and add it to the final hlsl.
+			OutHLSL += FString::Format(*FunctionTemplateFile, HlslTemplateArgs);
+		}
+	}
+}
+#endif
+
+void UNiagaraDataInterfaceDataChannelRead::BuildShaderParameters(FNiagaraShaderParametersBuilder& ShaderParametersBuilder) const
+{
+	ShaderParametersBuilder.AddNestedStruct<NDIDataChannelReadLocal::FShaderParameters>();
+}
+
+void UNiagaraDataInterfaceDataChannelRead::SetShaderParameters(const FNiagaraDataInterfaceSetShaderParametersContext& Context) const
+{
+	FNiagaraDataInterfaceProxy_DataChannelRead& DataInterfaceProxy = Context.GetProxy<FNiagaraDataInterfaceProxy_DataChannelRead>();
+	FNiagaraDataInterfaceProxy_DataChannelRead::FInstanceData* InstanceData = DataInterfaceProxy.SystemInstancesToProxyData_RT.Find(Context.GetSystemInstanceID());
+
+	NDIDataChannelReadLocal::FShaderParameters* InstParameters = Context.GetParameterNestedStruct<NDIDataChannelReadLocal::FShaderParameters>();
+
+	bool bSuccess = false;
+	if (InstanceData)
+	{
+		//Find the start offset in the parameter table for this script.
+		const FNiagaraCompileHash& ScriptCompileHash = Context.GetComputeInstanceData().Context->GPUScript_RT->GetBaseCompileHash();
+		uint32 ParameterOffsetTableIndex = INDEX_NONE;
+		if (uint32* ParameterOffsetTableIndexPtr = InstanceData->GPUScriptParameterTableOffsets.Find(ScriptCompileHash))
+		{
+			ParameterOffsetTableIndex = *ParameterOffsetTableIndexPtr;
+		}
+
+		if (InstanceData->ChannelDataSet && ParameterOffsetTableIndex != INDEX_NONE)
+		{
+			if (FNiagaraDataBuffer* Data = InstanceData->ChannelDataSet->GetCurrentData())
+			{
+				const FReadBuffer& ParameterLayoutBuffer = InstanceData->ParameterLayoutBuffer;
+
+				if (ParameterLayoutBuffer.SRV.IsValid() && ParameterLayoutBuffer.NumBytes > 0)
+				{
+					InstParameters->ParamOffsetTable = ParameterLayoutBuffer.SRV.IsValid() ? ParameterLayoutBuffer.SRV.GetReference() : FNiagaraRenderer::GetDummyUIntBuffer();
+					InstParameters->DataFloat = Data->GetGPUBufferFloat().SRV.IsValid() ? Data->GetGPUBufferFloat().SRV.GetReference() : FNiagaraRenderer::GetDummyFloatBuffer();
+					InstParameters->DataInt32 = Data->GetGPUBufferInt().SRV.IsValid() ? Data->GetGPUBufferInt().SRV.GetReference() : FNiagaraRenderer::GetDummyIntBuffer();
+					InstParameters->DataHalf = Data->GetGPUBufferHalf().SRV.IsValid() ? Data->GetGPUBufferHalf().SRV.GetReference() : FNiagaraRenderer::GetDummyHalfBuffer();
+					InstParameters->ParameterOffsetTableIndex = ParameterOffsetTableIndex;
+					InstParameters->Num = Data->GetNumInstances();
+					InstParameters->FloatStride = Data->GetFloatStride() / sizeof(float);
+					InstParameters->Int32Stride = Data->GetInt32Stride() / sizeof(int32);
+					InstParameters->HalfStride = Data->GetHalfStride() / sizeof(FFloat16);
+
+					bSuccess = true;
+				}
+			}
+		}
+	}
+
+	if (bSuccess == false)
+	{
+		InstParameters->ParamOffsetTable = FNiagaraRenderer::GetDummyUIntBuffer();
+		InstParameters->DataFloat = FNiagaraRenderer::GetDummyFloatBuffer();
+		InstParameters->DataInt32 = FNiagaraRenderer::GetDummyIntBuffer();
+		InstParameters->DataHalf = FNiagaraRenderer::GetDummyHalfBuffer();
+		InstParameters->ParameterOffsetTableIndex = INDEX_NONE;
+		InstParameters->Num = 0;
+		InstParameters->FloatStride = 0;
+		InstParameters->Int32Stride = 0;
+		InstParameters->HalfStride = 0;
+	}
+}
+
+void FNiagaraDataInterfaceProxy_DataChannelRead::ConsumePerInstanceDataFromGameThread(void* PerInstanceData, const FNiagaraSystemInstanceID& Instance)
+{
+	FNDIDataChannelReadInstanceData_RT& SourceData = *reinterpret_cast<FNDIDataChannelReadInstanceData_RT*>(PerInstanceData);
+	FInstanceData& InstData = SystemInstancesToProxyData_RT.FindOrAdd(Instance);
+
+	InstData.ChannelDataSet = SourceData.ChannelDataSet;
+
+	if (SourceData.bHasUpdate)
+	{
+		//Take the offset map from the source data.
+		//This maps from GPU script to that scripts offset into the ParameterLayoutBuffer.
+		//Allows us to look up and pass in at SetShaderParameters time.
+		InstData.GPUScriptParameterTableOffsets = MoveTemp(SourceData.GPUScriptParameterTableOffsets);
+
+		//Now generate the ParameterLayoutBuffer
+		//This contains a table of all parameters used by each GPU script that uses this DI.
+		//TODO: This buffer can likely be shared among many instances and stored in the layout manager or in the DI proxy.
+		{
+			if (InstData.ParameterLayoutBuffer.NumBytes > 0)
+			{
+				InstData.ParameterLayoutBuffer.Release();
+			}
+
+			if(SourceData.GPUScriptParameterOffsetTable.Num() > 0)
+			{
+				InstData.ParameterLayoutData = SourceData.GPUScriptParameterOffsetTable;
+				InstData.ParameterLayoutBuffer.Initialize(TEXT("NDIDataChannel_ParameterLayoutBuffer"), sizeof(uint32), SourceData.GPUScriptParameterOffsetTable.Num(), EPixelFormat::PF_R32_UINT, BUF_Static, &InstData.ParameterLayoutData);
+			}
+		}
+	}
+
+	SourceData.~FNDIDataChannelReadInstanceData_RT();
+}
+
+int32 FNiagaraDataInterfaceProxy_DataChannelRead::PerInstanceDataPassedToRenderThreadSize() const
+{
+	return sizeof(FNDIDataChannelReadInstanceData_RT);
+}
+
+#undef LOCTEXT_NAMESPACE
\ No newline at end of file
diff --git a/Engine/Plugins/FX/Niagara/Source/Niagara/Private/DataInterface/NiagaraDataInterfaceDataChannelSpawn.cpp b/Engine/Plugins/FX/Niagara/Source/Niagara/Private/DataInterface/NiagaraDataInterfaceDataChannelSpawn.cpp
new file mode 100644
--- /dev/null
+++ b/Engine/Plugins/FX/Niagara/Source/Niagara/Private/DataInterface/NiagaraDataInterfaceDataChannelSpawn.cpp
@@ -0,0 +1,285 @@
+// Copyright Epic Games, Inc. All Rights Reserved.
+
+/*
+UNiagaraDataInterfaceDataChannelSpawn
+A specialized version of the Read DI that allows for direct spawning into emitters from a SpawnInfo variable in the Data Channel Payload.
+*/
+
+#include "NiagaraDataInterfaceDataChannelSpawn.h"
+
+#include "CoreMinimal.h"
+
+#include "NiagaraCommon.h"
+#include "NiagaraDataChannel.h"
+#include "NiagaraDataChannelHandler.h"
+#include "NiagaraDataChannelDefinitions.h"
+
+#include "NiagaraDataSetAccessor.h"
+#include "NiagaraSystemInstance.h"
+
+
+#define LOCTEXT_NAMESPACE "NiagaraDataInterfaceDataChannelSpawn"
+
+DECLARE_CYCLE_STAT(TEXT("NDIDataChannelSpawn Spawn"), STAT_NDIDataChannelSpawn_Spawn, STATGROUP_NiagaraDataChannels);
+DECLARE_CYCLE_STAT(TEXT("NDIDataChannelSpawn Tick"), STAT_NDIDataChannelSpawn_Tick, STATGROUP_NiagaraDataChannels);
+DECLARE_CYCLE_STAT(TEXT("NDIDataChannelSpawn PostTick"), STAT_NDIDataChannelSpawn_PostTick, STATGROUP_NiagaraDataChannels);
+
+namespace NDIDataChannelLocal
+{
+	static const FName SpawnName(TEXT("Spawn"));
+}
+
+struct FNDIDataChannelSpawnInstanceData : public FNDIDataChannelReadInstanceData
+{
+	virtual ~FNDIDataChannelSpawnInstanceData() {}
+
+	bool Init(UNiagaraDataInterfaceDataChannelSpawn* Interface, FNiagaraSystemInstance* Instance)
+	{
+		if (FNDIDataChannelReadInstanceData::Init(Interface, Instance))
+		{
+			EmitterInstance = Interface->EmitterBinding.Resolve(Instance, Interface);
+			if (EmitterInstance)
+			{
+				return true;
+			}
+		}
+
+		return false;
+	}
+
+	bool Tick(UNiagaraDataInterfaceDataChannelSpawn* Interface, FNiagaraSystemInstance* Instance)
+	{
+		//Cache the current CPU dataset to ensure we don't need to regen the bindings.
+		FNiagaraDataBuffer* CachedData = ExternalCPU;
+
+		bool bSuccess = FNDIDataChannelReadInstanceData::Tick(Interface, Instance);
+
+		if (bSuccess)
+		{
+			//Do we need to regen the read offsets.
+			if (CachedData != ExternalCPU)
+			{
+				const FNiagaraDataSetCompiledData& DataSetCompiledData = ExternalCPU->GetOwner()->GetCompiledData();
+				SpawnInfoAccessor = FNiagaraDataSetAccessor<FNiagaraSpawnInfo>(DataSetCompiledData, SpawnInfoName);
+			}
+			return true;
+		}
+
+		return false;
+	}
+
+	bool PostTick(UNiagaraDataInterfaceDataChannelSpawn* Interface, FNiagaraSystemInstance* Instance)
+	{
+		return FNDIDataChannelReadInstanceData::PostTick(Interface, Instance);
+	}
+
+	FNiagaraEmitterInstance* EmitterInstance = nullptr;
+
+	/*
+	We read a spawn info for future proofing here but currently we only use the spawn count.
+	NOTE: This is 100% dependent on SpawnInfo class not changing. If it does we'll need to update this.
+	TODO: To support continuous interpolated spawning we'd need;
+	- Support persistent IDs when adding data to the data channel buffers.
+	- Support getting previous and current frame data by ID to do interpolation.
+	- Likely more issues to resolve.
+	*/
+	FNiagaraDataSetAccessor<FNiagaraSpawnInfo> SpawnInfoAccessor;
+
+	static const FName SpawnInfoName;
+};
+
+const FName FNDIDataChannelSpawnInstanceData::SpawnInfoName(TEXT("SpawnInfo"));
+
+UNiagaraDataInterfaceDataChannelSpawn::UNiagaraDataInterfaceDataChannelSpawn(FObjectInitializer const& ObjectInitializer)
+	: Super(ObjectInitializer)
+{
+}
+
+bool UNiagaraDataInterfaceDataChannelSpawn::Equals(const UNiagaraDataInterface* Other) const
+{
+	if (const UNiagaraDataInterfaceDataChannelSpawn* OtherTyped = CastChecked<UNiagaraDataInterfaceDataChannelSpawn>(Other))
+	{
+		if (Super::Equals(Other) &&
+			EmitterBinding == OtherTyped->EmitterBinding &&
+			bOverrideSpawnGroupToDataChannelIndex == OtherTyped->bOverrideSpawnGroupToDataChannelIndex)
+		{
+			return true;
+		}
+	}
+	return false;
+}
+
+void UNiagaraDataInterfaceDataChannelSpawn::GetFunctions(TArray<FNiagaraFunctionSignature>& OutFunctions)
+{
+	Super::GetFunctions(OutFunctions);
+	{
+		FNiagaraFunctionSignature Sig;
+		Sig.Name = NDIDataChannelLocal::SpawnName;
+#if WITH_EDITORONLY_DATA
+		Sig.Description = LOCTEXT("SpawnFunctionDescription", "Spawns particles in the given emitter for every entry in a Data Channel.");
+#endif
+		Sig.bMemberFunction = true;
+		Sig.bRequiresExecPin = true;
+		Sig.bExperimental = true;
+		Sig.ModuleUsageBitmask = ENiagaraScriptUsageMask::Emitter | ENiagaraScriptUsageMask::System;
+		Sig.AddInput(FNiagaraVariable(FNiagaraTypeDefinition(GetClass()), TEXT("Interface")));
+		Sig.AddInput(FNiagaraVariable(FNiagaraTypeDefinition::GetIntDef(), TEXT("Max Spawns")), LOCTEXT("Spawn_MaxSpawnsDesc", "Maximum number of elements to process into individual spawns. Number of particles spawned per element is determined by the Spawn Count property in each element."));
+		OutFunctions.Add(Sig);
+	}
+}
+void UNiagaraDataInterfaceDataChannelSpawn::GetVMExternalFunction(const FVMExternalFunctionBindingInfo& BindingInfo, void* InstanceData, FVMExternalFunction& OutFunc)
+{
+	if (BindingInfo.Name == NDIDataChannelLocal::SpawnName)
+	{
+		OutFunc = FVMExternalFunction::CreateLambda([this](FVectorVMExternalFunctionContext& Context) { this->Spawn(Context); });
+	}
+	else
+	{
+		Super::GetVMExternalFunction(BindingInfo, InstanceData, OutFunc);
+	}
+
+	if (!OutFunc.IsBound())
+	{
+		UE_LOG(LogTemp, Display, TEXT("Could not find data interface external function in %s. Received Name: %s"), *GetPathNameSafe(this), *BindingInfo.Name.ToString());
+	}
+}
+
+
+bool UNiagaraDataInterfaceDataChannelSpawn::InitPerInstanceData(void* PerInstanceData, FNiagaraSystemInstance* SystemInstance)
+{
+	SCOPE_CYCLE_COUNTER(STAT_NDIDataChannelSpawn_Spawn);
+	FNDIDataChannelSpawnInstanceData* InstanceData = new (PerInstanceData) FNDIDataChannelSpawnInstanceData;
+
+	if (INiagaraModule::DataChannelsEnabled() == false)
+	{
+		return false;
+	}
+
+	if (InstanceData->Init(this, SystemInstance) == false)
+	{
+		return false;
+	}
+
+	return true;
+}
+
+void UNiagaraDataInterfaceDataChannelSpawn::DestroyPerInstanceData(void* PerInstanceData, FNiagaraSystemInstance* SystemInstance)
+{
+	FNDIDataChannelSpawnInstanceData* InstanceData = static_cast<FNDIDataChannelSpawnInstanceData*>(PerInstanceData);
+	InstanceData->~FNDIDataChannelSpawnInstanceData();
+
+	ENQUEUE_RENDER_COMMAND(RemoveProxy)
+	(
+		[RT_Proxy = GetProxyAs<FNiagaraDataInterfaceProxy_DataChannelRead>(), InstanceID = SystemInstance->GetId()](FRHICommandListImmediate& CmdList)
+		{
+			RT_Proxy->SystemInstancesToProxyData_RT.Remove(InstanceID);
+		}
+	);
+}
+
+bool UNiagaraDataInterfaceDataChannelSpawn::PerInstanceTick(void* PerInstanceData, FNiagaraSystemInstance* SystemInstance, float DeltaSeconds)
+{
+	SCOPE_CYCLE_COUNTER(STAT_NDIDataChannelSpawn_Tick);
+	check(SystemInstance);
+	if (INiagaraModule::DataChannelsEnabled() == false)
+	{
+		return true;
+	}
+
+	FNDIDataChannelSpawnInstanceData* InstanceData = static_cast<FNDIDataChannelSpawnInstanceData*>(PerInstanceData);
+	if (!InstanceData)
+	{
+		return true;
+	}
+
+	if (InstanceData->Tick(this, SystemInstance) == false)
+	{
+		return true;
+	}
+	return false;
+}
+
+bool UNiagaraDataInterfaceDataChannelSpawn::PerInstanceTickPostSimulate(void* PerInstanceData, FNiagaraSystemInstance* SystemInstance, float DeltaSeconds)
+{
+	SCOPE_CYCLE_COUNTER(STAT_NDIDataChannelSpawn_PostTick);
+	check(SystemInstance);
+	if(INiagaraModule::DataChannelsEnabled() == false)
+	{
+		return true;
+	}
+	
+	FNDIDataChannelSpawnInstanceData* InstanceData = static_cast<FNDIDataChannelSpawnInstanceData*>(PerInstanceData);
+	if (!InstanceData)
+	{
+		return true;
+	}
+
+	if (InstanceData->PostTick(this, SystemInstance) == false)
+	{
+		return true;
+	}
+
+	return false;
+}
+
+int32 UNiagaraDataInterfaceDataChannelSpawn::PerInstanceDataSize() const
+{
+	return sizeof(FNDIDataChannelSpawnInstanceData);
+}
+
+void UNiagaraDataInterfaceDataChannelSpawn::Spawn(FVectorVMExternalFunctionContext& Context)
+{
+	SCOPE_CYCLE_COUNTER(STAT_NDIDataChannelSpawn_Spawn);
+
+	//This should only be called from emitter scripts and since it has per instance data then we process them individually.
+	check(Context.GetNumInstances() == 1);
+
+	VectorVM::FUserPtrHandler<FNDIDataChannelSpawnInstanceData> InstData(Context);
+	FNDIInputParam<int32> InMax(Context);
+
+	FNiagaraEmitterInstance* EmittterInst = InstData->EmitterInstance;
+	if (InstData->ExternalCPU && EmittterInst && InstData->SpawnInfoAccessor.IsValid() && INiagaraModule::DataChannelsEnabled())
+	{
+		TArray<FNiagaraSpawnInfo>& EmitterSpawnInfos = EmittterInst->GetSpawnInfo();
+		FNiagaraDataSetReaderStruct<FNiagaraSpawnInfo> Reader = InstData->SpawnInfoAccessor.GetReader(InstData->ExternalCPU);
+		int32 TotalSpawns = 0;
+		int32 MaxSpawns = InMax.GetAndAdvance();
+		if (TotalSpawns < MaxSpawns)
+		{
+			for (int32 DataChannelIdx = 0; DataChannelIdx < Reader.GetNumInstances(); ++DataChannelIdx)
+			{
+				FNiagaraSpawnInfo SpawnInfo = Reader.Get(DataChannelIdx);
+				if (SpawnInfo.Count > 0)
+				{
+					++TotalSpawns;
+					if (bOverrideSpawnGroupToDataChannelIndex)
+					{
+						SpawnInfo.SpawnGroup = DataChannelIdx;
+					}
+
+					EmitterSpawnInfos.Emplace(SpawnInfo);
+				}
+			}
+		}
+	}
+}
+
+bool UNiagaraDataInterfaceDataChannelSpawn::CopyToInternal(UNiagaraDataInterface* Destination) const
+{
+	if (!Super::CopyToInternal(Destination))
+	{
+		return false;
+	}
+
+	if (UNiagaraDataInterfaceDataChannelSpawn* DestTyped = CastChecked<UNiagaraDataInterfaceDataChannelSpawn>(Destination))
+	{
+		DestTyped->EmitterBinding = EmitterBinding;
+		DestTyped->bOverrideSpawnGroupToDataChannelIndex = bOverrideSpawnGroupToDataChannelIndex;
+		return true;
+	}
+
+	return false;
+}
+
+
+#undef LOCTEXT_NAMESPACE
\ No newline at end of file
diff --git a/Engine/Plugins/FX/Niagara/Source/Niagara/Private/DataInterface/NiagaraDataInterfaceDataChannelSpawn.h b/Engine/Plugins/FX/Niagara/Source/Niagara/Private/DataInterface/NiagaraDataInterfaceDataChannelSpawn.h
new file mode 100644
--- /dev/null
+++ b/Engine/Plugins/FX/Niagara/Source/Niagara/Private/DataInterface/NiagaraDataInterfaceDataChannelSpawn.h
@@ -0,0 +1,47 @@
+// Copyright Epic Games, Inc. All Rights Reserved.
+#pragma once
+
+/**
+Data Channel Spawn Interface.
+Specialization of the Read DI that allows Emitters to Spawn directly from Data Channel contents.
+*/
+
+#include "NiagaraDataInterfaceDataChannelRead.h"
+#include "NiagaraDataInterfaceEmitterBinding.h"
+#include "NiagaraDataInterfaceDataChannelSpawn.generated.h"
+
+
+UCLASS(Experimental, EditInlineNew, Category = "Data Channels", meta = (DisplayName = "Spawn From Data Channel"))
+class NIAGARA_API UNiagaraDataInterfaceDataChannelSpawn : public UNiagaraDataInterfaceDataChannelRead
+{
+	GENERATED_UCLASS_BODY()
+
+	/** The emitter into which we will spawn. */
+	UPROPERTY(EditAnywhere, Category = "Spawning")
+	FNiagaraDataInterfaceEmitterBinding EmitterBinding;
+	
+	/** 
+	If true, we override the SpawnGroup of any generated spawns with the element index of the data channel that generated the spawn. 
+	This allows the particle spawn script to subsequently access further data from the Data Channel at this index.
+	*/
+	UPROPERTY(EditAnywhere, Category = "Spawning")
+	bool bOverrideSpawnGroupToDataChannelIndex = true;
+
+	virtual bool Equals(const UNiagaraDataInterface* Other) const override;
+
+	virtual void GetFunctions(TArray<FNiagaraFunctionSignature>& OutFunctions) override;
+	virtual void GetVMExternalFunction(const FVMExternalFunctionBindingInfo& BindingInfo, void* InstanceData, FVMExternalFunction& OutFunc) override;
+	virtual bool CanExecuteOnTarget(ENiagaraSimTarget Target) const override { return true; }
+	virtual bool InitPerInstanceData(void* PerInstanceData, FNiagaraSystemInstance* SystemInstance) override;
+	virtual void DestroyPerInstanceData(void* PerInstanceData, FNiagaraSystemInstance* SystemInstance) override;
+	virtual int32 PerInstanceDataSize() const override;
+
+	virtual bool PerInstanceTick(void* PerInstanceData, FNiagaraSystemInstance* SystemInstance, float DeltaSeconds) override;
+	virtual bool PerInstanceTickPostSimulate(void* PerInstanceData, FNiagaraSystemInstance* SystemInstance, float DeltaSeconds) override;
+	
+	void Spawn(FVectorVMExternalFunctionContext& Context);
+
+protected:
+	virtual bool CopyToInternal(UNiagaraDataInterface* Destination) const override;
+
+};
\ No newline at end of file
