Niagara Data Channels v0.1

diff --git a/Engine/Plugins/FX/Niagara/Source/Niagara/Private/DataInterface/NiagaraDataInterfaceDataChannelWrite.cpp b/Engine/Plugins/FX/Niagara/Source/Niagara/Private/DataInterface/NiagaraDataInterfaceDataChannelWrite.cpp
new file mode 100644
--- /dev/null
+++ b/Engine/Plugins/FX/Niagara/Source/Niagara/Private/DataInterface/NiagaraDataInterfaceDataChannelWrite.cpp
@@ -0,0 +1,751 @@
+// Copyright Epic Games, Inc. All Rights Reserved.
+
+#include "NiagaraDataInterfaceDataChannelWrite.h"
+
+#include "CoreMinimal.h"
+#include "Stats/Stats.h"
+#include "NiagaraCommon.h"
+#include "NiagaraShared.h"
+
+#include "NiagaraWorldManager.h"
+#include "NiagaraSystemInstance.h"
+
+#include "NiagaraSystemImpl.h"
+
+#include "NiagaraDataChannel.h"
+#include "NiagaraDataChannelHandler.h"
+#include "NiagaraDataChannelDefinitions.h"
+
+
+#define LOCTEXT_NAMESPACE "NiagaraDataInterfaceDataChannelWrite"
+
+DECLARE_CYCLE_STAT(TEXT("NDIDataChannelWrite Write"), STAT_NDIDataChannelWrite_Write, STATGROUP_NiagaraDataChannels);
+DECLARE_CYCLE_STAT(TEXT("NDIDataChannelWrite Append"), STAT_NDIDataChannelWrite_Append, STATGROUP_NiagaraDataChannels);
+DECLARE_CYCLE_STAT(TEXT("NDIDataChannelWrite Tick"), STAT_NDIDataChannelWrite_Tick, STATGROUP_NiagaraDataChannels);
+DECLARE_CYCLE_STAT(TEXT("NDIDataChannelWrite PostTick"), STAT_NDIDataChannelWrite_PostTick, STATGROUP_NiagaraDataChannels);
+
+int32 GbDebugDumpWriter = 0;
+static FAutoConsoleVariableRef CVarDebugDumpWriterDI(
+	TEXT("fx.Niagara.DataChannels.DebugDumpWriterDI"),
+	GbDebugDumpWriter,
+	TEXT(" \n"),
+	ECVF_Default
+);
+
+namespace NDIDataChannelWriteLocal
+{
+	static const TCHAR* CommonShaderFile = TEXT("/Plugin/FX/Niagara/Private/DataChannel/NiagaraDataInterfaceDataChannelCommon.ush");
+	static const TCHAR* TemplateShaderFile_Common = TEXT("/Plugin/FX/Niagara/Private/DataChannel/NiagaraDataInterfaceDataChannelTemplateCommon.ush");
+	static const TCHAR* TemplateShaderFile_Write = TEXT("/Plugin/FX/Niagara/Private/DataChannel/NiagaraDataInterfaceDataChannelTemplate_Write.ush");
+	static const TCHAR* TemplateShaderFile_Append = TEXT("/Plugin/FX/Niagara/Private/DataChannel/NiagaraDataInterfaceDataChannelTemplate_Append.ush");
+
+	static const FName NumName(TEXT("Num"));
+	static const FName WriteName(TEXT("Write"));
+	static const FName AppendName(TEXT("Append"));
+
+	static const FName SpawnName(TEXT("Spawn"));
+
+	const TCHAR* GetFunctionTemplate(FName FunctionName)
+	{
+		if (FunctionName == WriteName) return TemplateShaderFile_Write;
+		if (FunctionName == AppendName) return TemplateShaderFile_Append;
+
+		return nullptr;
+	}
+}
+
+//////////////////////////////////////////////////////////////////////////
+//FNDIDataChannelWriteCompiledData
+
+bool FNDIDataChannelWriteCompiledData::Init(UNiagaraSystem* System, UNiagaraDataInterfaceDataChannelWrite* OwnerDI)
+{
+	FunctionInfo.Reset();
+
+	DataLayout.Empty();
+
+	GatherAccessInfo(System, OwnerDI);
+
+	for (FNDIDataChannelFunctionInfo& FuncInfo : FunctionInfo)
+	{
+		for (FNiagaraVariableBase& Param : FuncInfo.Inputs)
+		{
+			DataLayout.Variables.AddUnique(Param);
+		}
+	}
+
+	DataLayout.BuildLayout();
+
+	return true;
+}
+
+//FNDIDataChannelWriteCompiledData END
+//////////////////////////////////////////////////////////////////////////
+
+
+
+
+/**
+ The data channel write interface allows one Niagara System to write out arbitrary data to be later read by some other Niagara System or Game code/BP.
+
+ Currently this is done by writing the data to a local buffer and then copying into a global buffer when the data channel next ticks.
+ In the future we may add alternatives to this that allow for less copying etc.
+ Though for now this method allows the system to work without any synchronization headaches for the Read/Write or data races accessing a shared buffer concurrently etc.
+
+ Write DIs can also write in "Local" mode, which means their data is defined by whatever they write rather than any predefined
+
+*/
+struct FNDIDataChannelWriteInstanceData
+{
+	/** Pointer to the world DataChannel Channel we'll push our DataChannel into. Can be null if DI is not set to publish it's DataChannel. */
+	TWeakObjectPtr<UNiagaraDataChannelHandler> DataChannel;
+
+	FNiagaraDataSet* Data = nullptr;
+
+	bool bNeedsLastFrameData = true;
+
+	TArray<FNDIDataChannel_FuncToDataSetBindingPtr, TInlineAllocator<8>> FunctionToDatSetBindingInfo;
+
+	~FNDIDataChannelWriteInstanceData()
+	{
+		if (Data && DataChannel.IsValid())
+		{
+			DataChannel->RemovePublishRequests(Data);
+		}
+
+		/** We defer the deletion of the datase to the RT to be sure all in-flight RT commands have finished using it.*/
+		ENQUEUE_RENDER_COMMAND(FDeleteContextCommand)(
+			[DataChannelDataSet = Data](FRHICommandListImmediate& RHICmdList)
+			{
+				if (DataChannelDataSet != nullptr)
+				{
+					delete DataChannelDataSet;
+				}
+			}
+		);
+		Data = nullptr;
+	}
+
+	bool Init(UNiagaraDataInterfaceDataChannelWrite* Interface, FNiagaraSystemInstance* Instance)
+	{
+		Data = new FNiagaraDataSet();
+		Data->Init(&Interface->GetCompiledData().DataLayout);
+
+		//Grab the correct function binding infos for this DI.
+		const FNDIDataChannelCompiledData& CompiledData = Interface->GetCompiledData();
+		FunctionToDatSetBindingInfo.Reset(CompiledData.GetFunctionInfo().Num());
+		for (const FNDIDataChannelFunctionInfo& FuncInfo : CompiledData.GetFunctionInfo())
+		{
+			FunctionToDatSetBindingInfo.Add(FNDIDataChannelLayoutManager::Get().GetLayoutInfo(FuncInfo, Data->GetCompiledData()));
+		}
+
+		return true;
+	}
+
+	bool Tick(UNiagaraDataInterfaceDataChannelWrite* Interface, FNiagaraSystemInstance* Instance)
+	{
+		if (Interface->ShouldPublish())
+		{
+			UNiagaraDataChannelHandler* DataChannelPtr = DataChannel.Get();
+			if (DataChannelPtr == nullptr)
+			{
+				UWorld* World = Instance->GetWorld();
+				if (FNiagaraWorldManager* WorldMan = FNiagaraWorldManager::Get(World))
+				{
+					if (UNiagaraDataChannelHandler* NewChannelHandler = WorldMan->GetDataChannelManager().FindDataChannelHandler(Interface->Channel.ChannelName))
+					{
+						DataChannelPtr = NewChannelHandler;
+						DataChannel = NewChannelHandler;
+					}
+					else
+					{
+						UE_LOG(LogNiagara, Warning, TEXT("Failed to find or add Naigara DataChannel Channel: %s"), *Interface->Channel.ChannelName.ToString());
+						return false;
+					}
+				}
+			}
+		}
+
+		FNiagaraDataBuffer* CurrBuff = Data->GetCurrentData();
+		if (bNeedsLastFrameData)
+		{
+			CurrBuff->AddReadRef();//Ensure we keep the previous data when grabbing a new dest buffer.
+		}
+
+		FNiagaraDataBuffer& DestBuff = Data->BeginSimulate(true);
+
+		if (bNeedsLastFrameData && CurrBuff)
+		{
+			CurrBuff->ReleaseReadRef();
+		}
+
+		if (Interface->AllocationMode == ENiagaraDataChannelAllocationMode::Static)
+		{
+			Data->GetDestinationData()->Allocate(Interface->AllocationCount);
+		}
+		//else if (Interface->AllocationMode == ENiagaraDataChannelAllocationMode::PerInstance)
+		//{
+			//TODO: 
+			//Have to count up the current size of all users of this DI.
+			//Can this be done in DI tick? Or do we need a pre stage API once we do sim stages on CPU?
+			//Also adds extra wrinkle if we have multiple stages or emitters using the same writer DI. 
+			//Maybe need a pre pass to gather all users and their allocation sizes this frame.
+		//}
+		else
+		{
+			check(0);
+		}
+
+		//Verify our function info.
+		if(!ensure(Interface->GetCompiledData().GetFunctionInfo().Num() == FunctionToDatSetBindingInfo.Num()))
+		{
+			UE_LOG(LogNiagara, Warning, TEXT("Invalid Bindings for Niagara Data Interface Data Channel Write: %s"), *Interface->Channel.ChannelName.ToString());
+			return false;			
+		}
+
+		for(auto& Binding : FunctionToDatSetBindingInfo)
+		{
+			if(Binding.IsValid() == false)
+			{
+				UE_LOG(LogNiagara, Warning, TEXT("Invalid Bindings for Niagara Data Interface Data Channel Write: %s"), *Interface->Channel.ChannelName.ToString());
+				return false;			
+			}
+		}
+
+		return true;
+	}
+
+	bool PostTick(UNiagaraDataInterfaceDataChannelWrite* Interface, FNiagaraSystemInstance* Instance)
+	{
+		if (Data && Data->GetDestinationData())
+		{
+			Data->EndSimulate();
+
+			if (GbDebugDumpWriter)
+			{
+				FNiagaraDataBuffer& Buffer = Data->GetCurrentDataChecked();
+				Buffer.Dump(0, Buffer.GetNumInstances(), FString::Printf(TEXT("=== Data Channle Write: %d Elements --> %s ==="), Buffer.GetNumInstances(), *Interface->Channel.ChannelName.ToString()));
+			}
+
+			if (Interface->ShouldPublish() && Data->GetCurrentData()->GetNumInstances() > 0)
+			{
+				if (UNiagaraDataChannelHandler* Channel = DataChannel.Get())
+				{
+					FNiagaraDataChannelPublishRequest PublishRequest(Data->GetCurrentData());
+					PublishRequest.bVisibleToGame = Interface->bPublishToGame;
+					PublishRequest.bVisibleToCPUSims = Interface->bPublishToCPU;
+					PublishRequest.bVisibleToGPUSims = Interface->bPublishToGPU;
+					PublishRequest.Data = Data->GetCurrentData();
+					PublishRequest.LwcTile = Instance->GetLWCTile();
+					Channel->Publish(PublishRequest);
+				}
+			}
+		}
+		return true;
+	}
+};
+
+bool UNiagaraDataInterfaceDataChannelWrite::InitPerInstanceData(void* PerInstanceData, FNiagaraSystemInstance* SystemInstance)
+{
+	FNDIDataChannelWriteInstanceData* InstanceData = new (PerInstanceData) FNDIDataChannelWriteInstanceData;
+
+	//If data channels are disabled we just skip and return ok so that systems can continue to function.
+	if (INiagaraModule::DataChannelsEnabled() == false)
+	{
+		return false;
+	}
+
+	if (InstanceData->Init(this, SystemInstance) == false)
+	{
+		return false;
+	}
+
+	return true;
+}
+
+void UNiagaraDataInterfaceDataChannelWrite::DestroyPerInstanceData(void* PerInstanceData, FNiagaraSystemInstance* SystemInstance)
+{
+	FNDIDataChannelWriteInstanceData* InstanceData = static_cast<FNDIDataChannelWriteInstanceData*>(PerInstanceData);
+	InstanceData->~FNDIDataChannelWriteInstanceData();
+
+	// 	ENQUEUE_RENDER_COMMAND(RemoveProxy)
+	// 	(
+	// 		[RT_Proxy = GetProxyAs<FNDIDataChannelProxy>(), InstanceID = SystemInstance->GetId()](FRHICommandListImmediate& CmdList)
+	// 		{
+	// 			RT_Proxy->SystemInstancesToInstanceData_RT.Remove(InstanceID);
+	// 		}
+	// 	);
+}
+
+
+UNiagaraDataInterfaceDataChannelWrite::UNiagaraDataInterfaceDataChannelWrite(FObjectInitializer const& ObjectInitializer)
+	: Super(ObjectInitializer)
+{
+	//TODO:
+	//Proxy.Reset(new FNDIDataChannelProxy());
+}
+
+void UNiagaraDataInterfaceDataChannelWrite::PostInitProperties()
+{
+	Super::PostInitProperties();
+
+	if (HasAnyFlags(RF_ClassDefaultObject) && INiagaraModule::DataChannelsEnabled())
+	{
+		ENiagaraTypeRegistryFlags Flags = ENiagaraTypeRegistryFlags::AllowAnyVariable | ENiagaraTypeRegistryFlags::AllowParameter;
+		FNiagaraTypeRegistry::Register(FNiagaraTypeDefinition(GetClass()), Flags);
+	}
+}
+
+
+int32 UNiagaraDataInterfaceDataChannelWrite::PerInstanceDataSize() const
+{
+	return sizeof(FNDIDataChannelWriteInstanceData);
+}
+
+bool UNiagaraDataInterfaceDataChannelWrite::PerInstanceTick(void* PerInstanceData, FNiagaraSystemInstance* SystemInstance, float DeltaSeconds)
+{
+	if (INiagaraModule::DataChannelsEnabled() == false)
+	{
+		return true;
+	}
+	
+	SCOPE_CYCLE_COUNTER(STAT_NDIDataChannelWrite_Tick);
+	check(SystemInstance);
+	FNDIDataChannelWriteInstanceData* InstanceData = static_cast<FNDIDataChannelWriteInstanceData*>(PerInstanceData);
+	if (!InstanceData)
+	{
+		return true;
+	}
+
+	if (InstanceData->Tick(this, SystemInstance) == false)
+	{
+		return true;
+	}
+
+	return false;
+}
+
+bool UNiagaraDataInterfaceDataChannelWrite::PerInstanceTickPostSimulate(void* PerInstanceData, FNiagaraSystemInstance* SystemInstance, float DeltaSeconds)
+{
+	if (INiagaraModule::DataChannelsEnabled() == false)
+	{
+		return true;
+	}
+
+	SCOPE_CYCLE_COUNTER(STAT_NDIDataChannelWrite_PostTick);
+	check(SystemInstance);
+	FNDIDataChannelWriteInstanceData* InstanceData = static_cast<FNDIDataChannelWriteInstanceData*>(PerInstanceData);
+	if (!InstanceData)
+	{
+		return true;
+	}
+
+	if (InstanceData->PostTick(this, SystemInstance) == false)
+	{
+		return true;
+	}
+	return false;
+}
+
+void UNiagaraDataInterfaceDataChannelWrite::ProvidePerInstanceDataForRenderThread(void* DataForRenderThread, void* PerInstanceData, const FNiagaraSystemInstanceID& SystemInstance)
+{
+	//TODO:
+	//FNDIDataChannelProxy::ProvidePerInstanceDataForRenderThread(DataForRenderThread, PerInstanceData, SystemInstance);
+}
+
+#if WITH_EDITORONLY_DATA
+
+void UNiagaraDataInterfaceDataChannelWrite::PostCompile()
+{
+	UNiagaraSystem* OwnerSystem = GetTypedOuter<UNiagaraSystem>();
+	CompiledData.Init(OwnerSystem, this);
+
+	if (const UNiagaraDataChannel* DataChannel = UNiagaraDataChannelDefinitions::FindDataChannel(Channel.ChannelName))
+	{
+		OwnerSystem->RegisterDataChannelUse(DataChannel);
+	}
+}
+
+#endif
+
+bool UNiagaraDataInterfaceDataChannelWrite::Equals(const UNiagaraDataInterface* Other)const
+{
+	if (const UNiagaraDataInterfaceDataChannelWrite* OtherTyped = CastChecked<UNiagaraDataInterfaceDataChannelWrite>(Other))
+	{
+		if (Super::Equals(Other) &&
+			AllocationMode == OtherTyped->AllocationMode &&
+			AllocationCount == OtherTyped->AllocationCount &&
+			bPublishToGame == OtherTyped->bPublishToGame &&
+			bPublishToCPU == OtherTyped->bPublishToCPU &&
+			bPublishToGPU == OtherTyped->bPublishToGPU &&
+			Channel.ChannelName == OtherTyped->Channel.ChannelName)
+		{
+			return true;
+		}
+	}
+	return false;
+}
+
+bool UNiagaraDataInterfaceDataChannelWrite::CopyToInternal(UNiagaraDataInterface* Destination)const
+{
+	if (!Super::CopyToInternal(Destination))
+	{
+		return false;
+	}
+
+	if (UNiagaraDataInterfaceDataChannelWrite* DestTyped = CastChecked<UNiagaraDataInterfaceDataChannelWrite>(Destination))
+	{
+		DestTyped->AllocationMode = AllocationMode;
+		DestTyped->AllocationCount = AllocationCount;
+		DestTyped->bPublishToGame = bPublishToGame;
+		DestTyped->bPublishToCPU = bPublishToCPU;
+		DestTyped->bPublishToGPU = bPublishToGPU;
+		DestTyped->Channel.ChannelName = Channel.ChannelName;
+		DestTyped->CompiledData = CompiledData;
+		return true;
+	}
+
+	return false;
+}
+
+void UNiagaraDataInterfaceDataChannelWrite::GetFunctions(TArray<FNiagaraFunctionSignature>& OutFunctions)
+{
+	{
+		FNiagaraFunctionSignature Sig;
+		Sig.Name = NDIDataChannelWriteLocal::NumName;
+#if WITH_EDITORONLY_DATA
+		Sig.Description = LOCTEXT("NumFunctionDescription", "Returns the current number of DataChannel accessible by this interface.");
+#endif
+		Sig.bMemberFunction = true;
+		Sig.bExperimental = true;
+		Sig.AddInput(FNiagaraVariable(FNiagaraTypeDefinition(GetClass()), TEXT("DataChannel interface")));
+		Sig.AddOutput(FNiagaraVariable(FNiagaraTypeDefinition::GetIntDef(), TEXT("Num")));
+		OutFunctions.Add(Sig);
+	}
+	{
+		FNiagaraFunctionSignature Sig;
+		Sig.Name = NDIDataChannelWriteLocal::WriteName;
+#if WITH_EDITORONLY_DATA
+		Sig.Description = LOCTEXT("WriteFunctionDescription", "Writes DataChannel data at a specific index.  Values in the DataChannel that are not written here are set to their defaults. Returns success if an DataChannel was written to.");
+#endif
+		Sig.bMemberFunction = true;
+		Sig.bRequiresExecPin = true;
+		Sig.bExperimental = true;
+		Sig.AddInput(FNiagaraVariable(FNiagaraTypeDefinition(GetClass()), TEXT("DataChannel interface")));
+		Sig.AddInput(FNiagaraVariable(FNiagaraTypeDefinition::GetBoolDef(), TEXT("Emit")));
+		Sig.AddInput(FNiagaraVariable(FNiagaraTypeDefinition::GetIntDef(), TEXT("Index")));
+		Sig.AddOutput(FNiagaraVariable(FNiagaraTypeDefinition::GetBoolDef(), TEXT("Success")));
+		Sig.RequiredInputs = Sig.Inputs.Num();//The user defines what we write in the graph.
+		OutFunctions.Add(Sig);
+	}
+	{
+		FNiagaraFunctionSignature Sig;
+		Sig.Name = NDIDataChannelWriteLocal::AppendName;
+#if WITH_EDITORONLY_DATA
+		Sig.Description = LOCTEXT("AppendFunctionDescription", "Appends a new DataChannel to the end of the DataChannel array and writes the specified values. Values in the DataChannel that are not written here are set to their defaults. Returns success if an DataChannel was successfully pushed.");
+#endif
+		Sig.bMemberFunction = true;
+		Sig.bRequiresExecPin = true;
+		Sig.bExperimental = true;
+		Sig.AddInput(FNiagaraVariable(FNiagaraTypeDefinition(GetClass()), TEXT("DataChannel interface")));
+		Sig.AddInput(FNiagaraVariable(FNiagaraTypeDefinition::GetBoolDef(), TEXT("Emit")));
+		Sig.AddOutput(FNiagaraVariable(FNiagaraTypeDefinition::GetBoolDef(), TEXT("Success")));
+		Sig.RequiredInputs = Sig.Inputs.Num();//The user defines what we write in the graph.
+		OutFunctions.Add(Sig);
+	}
+}
+
+void UNiagaraDataInterfaceDataChannelWrite::GetVMExternalFunction(const FVMExternalFunctionBindingInfo& BindingInfo, void* InstanceData, FVMExternalFunction& OutFunc)
+{
+	if (BindingInfo.Name == NDIDataChannelWriteLocal::NumName)
+	{
+		OutFunc = FVMExternalFunction::CreateLambda([this](FVectorVMExternalFunctionContext& Context) { this->Num(Context); });
+	}
+	else
+	{
+		int32 FuncIndex = CompiledData.FindFunctionInfoIndex(BindingInfo.Name, BindingInfo.VariadicInputs, BindingInfo.VariadicOutputs);
+
+		if (FuncIndex != INDEX_NONE)
+		{
+			if (BindingInfo.Name == NDIDataChannelWriteLocal::WriteName)
+			{
+				OutFunc = FVMExternalFunction::CreateLambda([this, FuncIndex](FVectorVMExternalFunctionContext& Context) { this->Write(Context, FuncIndex); });
+			}
+			else if (BindingInfo.Name == NDIDataChannelWriteLocal::AppendName)
+			{
+				OutFunc = FVMExternalFunction::CreateLambda([this, FuncIndex](FVectorVMExternalFunctionContext& Context) { this->Append(Context, FuncIndex); });
+			}
+			else
+			{
+				UE_LOG(LogTemp, Display, TEXT("Could not find data interface external function in %s. Received Name: %s"), *GetPathNameSafe(this), *BindingInfo.Name.ToString());
+			}
+		}
+		else
+		{
+			UE_LOG(LogTemp, Display, TEXT("Could not find data interface external function in %s. Received Name: %s"), *GetPathNameSafe(this), *BindingInfo.Name.ToString());
+		}
+	}
+}
+
+/** Handles any number of variadic parameter inputs. */
+template<int32 EXPECTED_NUM_INPUTS>
+struct FNDIVariadicInputHandler
+{
+	TArray<FNDIInputParam<float>, TInlineAllocator<EXPECTED_NUM_INPUTS>> FloatInputs;
+	TArray<FNDIInputParam<int32>, TInlineAllocator<EXPECTED_NUM_INPUTS>> IntInputs;
+	TArray<FNDIInputParam<FFloat16>, TInlineAllocator<EXPECTED_NUM_INPUTS>> HalfInputs;
+
+	FNDIVariadicInputHandler(FVectorVMExternalFunctionContext& Context, const FNDIDataChannel_FuncToDataSetBindingPtr& BindingPtr)
+	{
+		FloatInputs.Reserve(BindingPtr->NumFloatComponents);
+		IntInputs.Reserve(BindingPtr->NumInt32Components);
+		HalfInputs.Reserve(BindingPtr->NumHalfComponents);
+		for (FNDIDataChannelRegisterBinding& VMBinding : BindingPtr->VMRegisterBindings)
+		{
+			if (VMBinding.DataType == (int32)ENiagaraBaseTypes::Float)
+			{
+				FloatInputs.Emplace(Context);
+			}
+			else if (VMBinding.DataType == (int32)ENiagaraBaseTypes::Int32 || VMBinding.DataType == (int32)ENiagaraBaseTypes::Bool)
+			{
+				IntInputs.Emplace(Context);
+			}
+			else if (VMBinding.DataType == (int32)ENiagaraBaseTypes::Half)
+			{
+				HalfInputs.Emplace(Context);
+			}
+			else
+			{
+				checkf(false, TEXT("Didn't find a binding for this function input. Likely an error in the building of the binding data and the Float/Int/Half Bindings arrays."));
+			}
+		}
+	}
+
+	void Advance()
+	{
+		for (FNDIInputParam<float>& Input : FloatInputs) { Input.Advance(); }
+		for (FNDIInputParam<int32>& Input : IntInputs) { Input.Advance(); }
+		for (FNDIInputParam<FFloat16>& Input : HalfInputs) { Input.Advance(); }
+	}
+
+	bool Process(FNiagaraDataBuffer* Data, uint32 Index, const FNDIDataChannel_FuncToDataSetBindingPtr& BindingInfo)
+	{
+		if (Data && BindingInfo)
+		{
+			//TODO: Optimize for long runs of writes to reduce binding/lookup overhead.
+			if (Index < Data->GetNumInstances())
+			{
+				for (FNDIDataChannelRegisterBinding& VMBinding : BindingInfo->VMRegisterBindings)
+				{
+					if (VMBinding.DataType == (int32)ENiagaraBaseTypes::Float)
+					{
+						if (VMBinding.DataSetRegisterIndex != INDEX_NONE)
+						{
+							*Data->GetInstancePtrFloat(VMBinding.DataSetRegisterIndex, (uint32)Index) = FloatInputs[VMBinding.FunctionRegisterIndex].Get();
+						}
+					}
+					else if (VMBinding.DataType == (int32)ENiagaraBaseTypes::Int32 || VMBinding.DataType == (int32)ENiagaraBaseTypes::Bool)
+					{
+						if (VMBinding.DataSetRegisterIndex != INDEX_NONE)
+						{
+							*Data->GetInstancePtrInt32(VMBinding.DataSetRegisterIndex, (uint32)Index) = IntInputs[VMBinding.FunctionRegisterIndex].Get();
+						}
+					}
+					else if (VMBinding.DataType == (int32)ENiagaraBaseTypes::Half)
+					{
+						if (VMBinding.DataSetRegisterIndex != INDEX_NONE)
+						{
+							*Data->GetInstancePtrHalf(VMBinding.DataSetRegisterIndex, (uint32)Index) = HalfInputs[VMBinding.FunctionRegisterIndex].Get();
+						}
+					}
+					else
+					{
+						checkf(false, TEXT("Didn't find a binding for this function input. Likely an error in the building of the binding data and the Float/Int/Half Bindings arrays."));
+					}
+				}
+
+				return true;
+			}
+		}
+		return false;
+	}
+};
+
+void UNiagaraDataInterfaceDataChannelWrite::Num(FVectorVMExternalFunctionContext& Context)
+{
+	VectorVM::FUserPtrHandler<FNDIDataChannelWriteInstanceData> InstData(Context);
+
+	FNDIOutputParam<int32> OutNum(Context);
+
+	FNiagaraDataBuffer* Buffer = InstData->Data->GetDestinationData();	
+	int32 Num = 0;
+	if (Buffer && INiagaraModule::DataChannelsEnabled())
+	{
+		Num = (int32)Buffer->GetNumInstances();
+	}
+
+	for (int32 i = 0; i < Context.GetNumInstances(); ++i)
+	{
+		OutNum.SetAndAdvance(Num);
+	}
+}
+
+void UNiagaraDataInterfaceDataChannelWrite::Write(FVectorVMExternalFunctionContext& Context, int32 FuncIdx)
+{
+	SCOPE_CYCLE_COUNTER(STAT_NDIDataChannelWrite_Write);
+	VectorVM::FUserPtrHandler<FNDIDataChannelWriteInstanceData> InstData(Context);
+	FNDIInputParam<bool> InEmit(Context);
+	FNDIInputParam<int32> InIndex(Context);
+
+	const FNDIDataChannelFunctionInfo& FuncInfo = CompiledData.GetFunctionInfo()[FuncIdx];
+	const FNDIDataChannel_FuncToDataSetBindingPtr BindingInfo = InstData->FunctionToDatSetBindingInfo.IsValidIndex(FuncIdx) ? InstData->FunctionToDatSetBindingInfo[FuncIdx] : nullptr;
+	FNDIVariadicInputHandler<16> VariadicInputs(Context, BindingInfo);//TODO: Make static / avoid allocation
+
+	FNDIOutputParam<bool> OutSuccess(Context);
+
+	if (InstData->Data && INiagaraModule::DataChannelsEnabled())
+	{
+		FNiagaraDataBuffer* Buffer = InstData->Data->GetDestinationData();
+
+		for (int32 i = 0; i < Context.GetNumInstances(); ++i)
+		{
+			int32 Index = InIndex.GetAndAdvance();
+
+			bool bEmit = InEmit.GetAndAdvance();
+
+			bool bSuccess = false;
+
+			//TODO: Optimize case where emit is constant
+			//TODO: Optimize for runs of sequential true emits.
+			if (bEmit)
+			{
+				bSuccess = VariadicInputs.Process(Buffer, Index, BindingInfo);
+			}
+
+			if (OutSuccess.IsValid())
+			{
+				OutSuccess.SetAndAdvance(bSuccess);
+			}
+
+			VariadicInputs.Advance();
+		}
+	}
+	else
+	{
+		for (int32 i = 0; i < Context.GetNumInstances(); ++i)
+		{
+			if (OutSuccess.IsValid())
+			{
+				OutSuccess.SetAndAdvance(false);
+			}
+		}
+	}
+}
+
+void UNiagaraDataInterfaceDataChannelWrite::Append(FVectorVMExternalFunctionContext& Context, int32 FuncIdx)
+{
+	SCOPE_CYCLE_COUNTER(STAT_NDIDataChannelWrite_Append);
+	VectorVM::FUserPtrHandler<FNDIDataChannelWriteInstanceData> InstData(Context);
+	FNDIInputParam<bool> InEmit(Context);
+
+	const FNDIDataChannelFunctionInfo& FuncInfo = CompiledData.GetFunctionInfo()[FuncIdx];
+	const FNDIDataChannel_FuncToDataSetBindingPtr BindingInfo = InstData->FunctionToDatSetBindingInfo.IsValidIndex(FuncIdx) ? InstData->FunctionToDatSetBindingInfo[FuncIdx] : nullptr;
+	FNDIVariadicInputHandler<16> VariadicInputs(Context, BindingInfo);//TODO: Make static / avoid allocation
+
+	FNDIOutputParam<bool> OutSuccess(Context);
+
+	if (InstData->Data && INiagaraModule::DataChannelsEnabled())
+	{
+		FNiagaraDataBuffer* Buffer = InstData->Data->GetDestinationData();
+		for (int32 i = 0; i < Context.GetNumInstances(); ++i)
+		{
+			bool bEmit = InEmit.GetAndAdvance();
+
+			bool bSuccess = false;
+
+			//TODO: Optimize case where emit is constant
+			//TODO: Optimize for runs of sequential true emits.
+			if (bEmit && Buffer)
+			{
+				int32 Num = Buffer->GetNumInstances();
+				if ((int32)Buffer->GetNumInstancesAllocated() > Num)
+				{
+					Buffer->SetNumInstances(Num + 1);
+					bSuccess = VariadicInputs.Process(Buffer, Num, BindingInfo);
+				}
+			}
+
+			if (OutSuccess.IsValid())
+			{
+				OutSuccess.SetAndAdvance(bSuccess);
+			}
+
+			VariadicInputs.Advance();
+		}
+	}
+	else
+	{
+		for (int32 i = 0; i < Context.GetNumInstances(); ++i)
+		{
+			if (OutSuccess.IsValid())
+			{
+				OutSuccess.SetAndAdvance(false);
+			}
+		}
+	}
+}
+
+#if WITH_EDITORONLY_DATA
+bool UNiagaraDataInterfaceDataChannelWrite::AppendCompileHash(FNiagaraCompileHashVisitor* InVisitor) const
+{
+	//TODO: GPU Writes.
+	// 
+// 	bool bSuccess = Super::AppendCompileHash(InVisitor);
+// 	bSuccess &= InVisitor->UpdateString(TEXT("UNiagaraDataInterfaceDataChannelCommon"), GetShaderFileHash(NDIDataChannelWriteLocal::CommonShaderFile, EShaderPlatform::SP_PCD3D_SM5).ToString());
+// 	bSuccess &= InVisitor->UpdateString(TEXT("UNiagaraDataInterfaceDataChannelWrite_Common"), GetShaderFileHash(NDIDataChannelWriteLocal::TemplateShaderFile_Common, EShaderPlatform::SP_PCD3D_SM5).ToString());
+// 	bSuccess &= InVisitor->UpdateString(TEXT("UNiagaraDataInterfaceDataChannelWrite_Write"), GetShaderFileHash(NDIDataChannelWriteLocal::TemplateShaderFile_Write, EShaderPlatform::SP_PCD3D_SM5).ToString());
+// 	bSuccess &= InVisitor->UpdateString(TEXT("UNiagaraDataInterfaceDataChannelWrite_Append"), GetShaderFileHash(NDIDataChannelWriteLocal::TemplateShaderFile_Append, EShaderPlatform::SP_PCD3D_SM5).ToString());
+// 
+// 	bSuccess &= InVisitor->UpdateShaderParameters<NDIDataChannelWriteLocal::FShaderParameters>();
+// 	return bSuccess;
+
+	return false;
+}
+
+void UNiagaraDataInterfaceDataChannelWrite::GetCommonHLSL(FString& OutHLSL)
+{
+	//TODO: GPU Writes.
+	// 
+// 	Super::GetCommonHLSL(OutHLSL);
+// 	OutHLSL.Append(TEXT("\n//Niagara Data Channel Write Interface Common Code.\n"));
+// 	OutHLSL.Appendf(TEXT("#include \"%s\"\n"), NDIDataChannelWriteLocal::CommonShaderFile);
+}
+
+bool UNiagaraDataInterfaceDataChannelWrite::GetFunctionHLSL(FNiagaraDataInterfaceHlslGenerationContext& HlslGenContext, FString& OutHLSL)
+{
+	//TODO: GPU Writes.
+	// 
+// 	return	HlslGenContext.GetFunctionInfo().DefinitionName == GET_FUNCTION_NAME_CHECKED(UNiagaraDataInterfaceDataChannelWrite, Num) ||
+// 		HlslGenContext.GetFunctionInfo().DefinitionName == GET_FUNCTION_NAME_CHECKED(UNiagaraDataInterfaceDataChannelWrite, Write) ||
+// 		HlslGenContext.GetFunctionInfo().DefinitionName == GET_FUNCTION_NAME_CHECKED(UNiagaraDataInterfaceDataChannelWrite, Append);
+
+	return false;
+}
+
+void UNiagaraDataInterfaceDataChannelWrite::GetParameterDefinitionHLSL(FNiagaraDataInterfaceHlslGenerationContext& HlslGenContext, FString& OutHLSL)
+{
+	Super::GetParameterDefinitionHLSL(HlslGenContext, OutHLSL);
+
+	//TODO: GPU Writes. 
+	//TODO: ADD VARIADIC PARAM HANDLING SIMILAR TO READ DI
+}
+
+#endif
+void UNiagaraDataInterfaceDataChannelWrite::BuildShaderParameters(FNiagaraShaderParametersBuilder& ShaderParametersBuilder) const
+{
+	//ShaderParametersBuilder.AddNestedStruct<FShaderParameters>();
+}
+void UNiagaraDataInterfaceDataChannelWrite::SetShaderParameters(const FNiagaraDataInterfaceSetShaderParametersContext& Context) const
+{
+}
+
+#undef LOCTEXT_NAMESPACE
\ No newline at end of file
diff --git a/Engine/Plugins/FX/Niagara/Source/Niagara/Private/DataInterface/NiagaraDataInterfaceUtilities.cpp b/Engine/Plugins/FX/Niagara/Source/Niagara/Private/DataInterface/NiagaraDataInterfaceUtilities.cpp
--- a/Engine/Plugins/FX/Niagara/Source/Niagara/Private/DataInterface/NiagaraDataInterfaceUtilities.cpp
+++ b/Engine/Plugins/FX/Niagara/Source/Niagara/Private/DataInterface/NiagaraDataInterfaceUtilities.cpp
@@ -155,7 +155,7 @@ void ForEachGpuFunctionEqualsImpl(class UNiagaraDataInterface* DataInterface, cl
 	);
 }
 
-void ForEachDataInterface(const FNiagaraParameterStore& ParameterStore, TFunction<bool(const FNiagaraVariableBase Variable, UNiagaraDataInterface* DataInterface)> Action)
+void ForEachDataInterface(const FNiagaraParameterStore& ParameterStore, TFunction<bool(const FNiagaraVariableBase Variable, UNiagaraDataInterface* DataInterface)> Action, FDataInterfaceSearchOptions SearchOptions)
 {
 	const TArray<UNiagaraDataInterface*>& DataInterfaces = ParameterStore.GetDataInterfaces();
 	for (const FNiagaraVariableWithOffset& Variable : ParameterStore.ReadParameterVariables())
@@ -165,7 +165,7 @@ void ForEachDataInterface(const FNiagaraParameterStore& ParameterStore, TFunctio
 			continue;
 		}
 
-		if (Variable.IsInNameSpace(FNiagaraConstants::InternalNamespaceString))
+		if (SearchOptions.bIncludeInternal == false && Variable.IsInNameSpace(FNiagaraConstants::InternalNamespaceString))
 		{
 			continue;
 		}
@@ -177,7 +177,7 @@ void ForEachDataInterface(const FNiagaraParameterStore& ParameterStore, TFunctio
 	}
 }
 
-void ForEachDataInterface(FDataInterfaceUsageContext& UsageContext, const FNiagaraParameterStore& ParameterStore, TFunction<bool(const FDataInterfaceUsageContext&)> Action)
+void ForEachDataInterface(FDataInterfaceUsageContext& UsageContext, const FNiagaraParameterStore& ParameterStore, TFunction<bool(const FDataInterfaceUsageContext&)> Action, FDataInterfaceSearchOptions SearchOptions)
 {
 	const TArray<UNiagaraDataInterface*>& DataInterfaces = ParameterStore.GetDataInterfaces();
 	for (const FNiagaraVariableWithOffset& Variable : ParameterStore.ReadParameterVariables())
@@ -187,7 +187,7 @@ void ForEachDataInterface(FDataInterfaceUsageContext& UsageContext, const FNiaga
 			continue;
 		}
 
-		if (Variable.IsInNameSpace(FNiagaraConstants::InternalNamespaceString))
+		if (SearchOptions.bIncludeInternal == false && Variable.IsInNameSpace(FNiagaraConstants::InternalNamespaceString))
 		{
 			continue;
 		}
@@ -362,7 +362,7 @@ void FNiagaraDataInterfaceUtilities::ForEachGpuFunction(class UNiagaraDataInterf
 	}
 }
 
-void FNiagaraDataInterfaceUtilities::ForEachDataInterface(FNiagaraSystemInstance* SystemInstance, TFunction<bool(const FNiagaraVariableBase Variable, UNiagaraDataInterface* DataInterface)> Action)
+void FNiagaraDataInterfaceUtilities::ForEachDataInterface(FNiagaraSystemInstance* SystemInstance, TFunction<bool(const FNiagaraVariableBase Variable, UNiagaraDataInterface* DataInterface)> Action, FDataInterfaceSearchOptions SearchOptions)
 {
 	if ( SystemInstance == nullptr )
 	{
@@ -375,10 +375,10 @@ void FNiagaraDataInterfaceUtilities::ForEachDataInterface(FNiagaraSystemInstance
 		return;
 	}
 
-	ForEachDataInterface(*SystemInstance->GetOverrideParameters(), Action);
+	ForEachDataInterface(*SystemInstance->GetOverrideParameters(), Action, SearchOptions);
 
-	ForEachDataInterface(SystemSimulation->GetSpawnExecutionContext()->Parameters, Action);
-	ForEachDataInterface(SystemSimulation->GetUpdateExecutionContext()->Parameters, Action);
+	ForEachDataInterface(SystemSimulation->GetSpawnExecutionContext()->Parameters, Action, SearchOptions);
+	ForEachDataInterface(SystemSimulation->GetUpdateExecutionContext()->Parameters, Action, SearchOptions);
 
 	for (const TSharedRef<FNiagaraEmitterInstance, ESPMode::ThreadSafe>& Emitter : SystemInstance->GetEmitters())
 	{
@@ -387,22 +387,22 @@ void FNiagaraDataInterfaceUtilities::ForEachDataInterface(FNiagaraSystemInstance
 			continue;
 		}
 
-		ForEachDataInterface(Emitter->GetSpawnExecutionContext().Parameters, Action);
-		ForEachDataInterface(Emitter->GetUpdateExecutionContext().Parameters, Action);
+		ForEachDataInterface(Emitter->GetSpawnExecutionContext().Parameters, Action, SearchOptions);
+		ForEachDataInterface(Emitter->GetUpdateExecutionContext().Parameters, Action, SearchOptions);
 		for (int32 i=0; i < Emitter->GetEventExecutionContexts().Num(); i++)
 		{
-			ForEachDataInterface(Emitter->GetEventExecutionContexts()[i].Parameters, Action);
+			ForEachDataInterface(Emitter->GetEventExecutionContexts()[i].Parameters, Action, SearchOptions);
 		}
 
 		FVersionedNiagaraEmitterData* EmitterData = Emitter->GetCachedEmitterData();
 		if (EmitterData && EmitterData->SimTarget == ENiagaraSimTarget::GPUComputeSim && Emitter->GetGPUContext())
 		{
-			ForEachDataInterface(Emitter->GetGPUContext()->CombinedParamStore, Action);
+			ForEachDataInterface(Emitter->GetGPUContext()->CombinedParamStore, Action, SearchOptions);
 		}
 	}
 }
 
-void FNiagaraDataInterfaceUtilities::ForEachDataInterface(FNiagaraSystemInstance* SystemInstance, TFunction<bool(const FDataInterfaceUsageContext&)> Action)
+void FNiagaraDataInterfaceUtilities::ForEachDataInterface(FNiagaraSystemInstance* SystemInstance, TFunction<bool(const FDataInterfaceUsageContext&)> Action, FDataInterfaceSearchOptions SearchOptions)
 {
 	if (SystemInstance == nullptr)
 	{
@@ -417,10 +417,10 @@ void FNiagaraDataInterfaceUtilities::ForEachDataInterface(FNiagaraSystemInstance
 
 	FDataInterfaceUsageContext UsageContext;
 	UsageContext.OwnerObject = SystemSimulation->GetSystem();
-	ForEachDataInterface(UsageContext, *SystemInstance->GetOverrideParameters(), Action);
+	ForEachDataInterface(UsageContext, *SystemInstance->GetOverrideParameters(), Action, SearchOptions);
 
-	ForEachDataInterface(UsageContext, SystemSimulation->GetSpawnExecutionContext()->Parameters, Action);
-	ForEachDataInterface(UsageContext, SystemSimulation->GetUpdateExecutionContext()->Parameters, Action);
+	ForEachDataInterface(UsageContext, SystemSimulation->GetSpawnExecutionContext()->Parameters, Action, SearchOptions);
+	ForEachDataInterface(UsageContext, SystemSimulation->GetUpdateExecutionContext()->Parameters, Action, SearchOptions);
 
 	for (const TSharedRef<FNiagaraEmitterInstance, ESPMode::ThreadSafe>& Emitter : SystemInstance->GetEmitters())
 	{
@@ -430,22 +430,22 @@ void FNiagaraDataInterfaceUtilities::ForEachDataInterface(FNiagaraSystemInstance
 		}
 		UsageContext.OwnerObject = Emitter->GetCachedEmitter().Emitter;
 
-		ForEachDataInterface(UsageContext, Emitter->GetSpawnExecutionContext().Parameters, Action);
-		ForEachDataInterface(UsageContext, Emitter->GetUpdateExecutionContext().Parameters, Action);
+		ForEachDataInterface(UsageContext, Emitter->GetSpawnExecutionContext().Parameters, Action, SearchOptions);
+		ForEachDataInterface(UsageContext, Emitter->GetUpdateExecutionContext().Parameters, Action, SearchOptions);
 		for (int32 i = 0; i < Emitter->GetEventExecutionContexts().Num(); i++)
 		{
-			ForEachDataInterface(UsageContext, Emitter->GetEventExecutionContexts()[i].Parameters, Action);
+			ForEachDataInterface(UsageContext, Emitter->GetEventExecutionContexts()[i].Parameters, Action, SearchOptions);
 		}
 
 		FVersionedNiagaraEmitterData* EmitterData = Emitter->GetCachedEmitterData();
 		if (EmitterData && EmitterData->SimTarget == ENiagaraSimTarget::GPUComputeSim && Emitter->GetGPUContext())
 		{
-			ForEachDataInterface(UsageContext, Emitter->GetGPUContext()->CombinedParamStore, Action);
+			ForEachDataInterface(UsageContext, Emitter->GetGPUContext()->CombinedParamStore, Action, SearchOptions);
 		}
 	}
 }
 
-void FNiagaraDataInterfaceUtilities::ForEachDataInterface(UNiagaraSystem* NiagaraSystem, TFunction<bool(const FDataInterfaceUsageContext&)> Action)
+void FNiagaraDataInterfaceUtilities::ForEachDataInterface(UNiagaraSystem* NiagaraSystem, TFunction<bool(const FDataInterfaceUsageContext&)> Action, FDataInterfaceSearchOptions SearchOptions)
 {
 	if (NiagaraSystem == nullptr)
 	{
@@ -464,7 +464,7 @@ void FNiagaraDataInterfaceUtilities::ForEachDataInterface(UNiagaraSystem* Niagar
 
 		if (const FNiagaraScriptExecutionParameterStore* ParameterStore = SystemScript->GetExecutionReadyParameterStore(ENiagaraSimTarget::CPUSim))
 		{
-			ForEachDataInterface(UsageContext, *ParameterStore, Action);
+			ForEachDataInterface(UsageContext, *ParameterStore, Action, SearchOptions);
 		}
 	}
 
@@ -487,7 +487,7 @@ void FNiagaraDataInterfaceUtilities::ForEachDataInterface(UNiagaraSystem* Niagar
 					return;
 				}
 
-				ForEachDataInterface(UsageContext, *ParameterStore, Action);
+				ForEachDataInterface(UsageContext, *ParameterStore, Action, SearchOptions);
 			}
 		);
 	}
diff --git a/Engine/Plugins/FX/Niagara/Source/Niagara/Private/NiagaraDataChannel.cpp b/Engine/Plugins/FX/Niagara/Source/Niagara/Private/NiagaraDataChannel.cpp
new file mode 100644
--- /dev/null
+++ b/Engine/Plugins/FX/Niagara/Source/Niagara/Private/NiagaraDataChannel.cpp
@@ -0,0 +1,541 @@
+// Copyright Epic Games, Inc. All Rights Reserved.
+
+#include "NiagaraDataChannel.h"
+
+#include "CoreMinimal.h"
+#include "NiagaraWorldManager.h"
+#include "NiagaraDataChannelHandler.h"
+#include "Engine/Engine.h"
+
+#include UE_INLINE_GENERATED_CPP_BY_NAME(NiagaraDataChannel)
+
+#define LOCTEXT_NAMESPACE "NiagaraDataChannels"
+
+//////////////////////////////////////////////////////////////////////////
+
+void FNiagaraDataChannelGameDataLayout::Init(const TArray<FNiagaraVariable>& Variables)
+{
+	VariableIndices.Reset();
+	LwcConverters.Reserve(Variables.Num());
+	for (const FNiagaraVariable& Var : Variables)
+	{
+		int32& VarIdx = VariableIndices.Add(Var);
+		VarIdx = VariableIndices.Num() - 1;
+
+		FNiagaraLwcStructConverter& Converter = LwcConverters.AddDefaulted_GetRef();
+		Converter = FNiagaraTypeRegistry::GetStructConverter(Var.GetType());
+	}
+}
+
+//////////////////////////////////////////////////////////////////////////
+
+void FNiagaraDataChannelGameData::Init(const UNiagaraDataChannel* InDataChannel)
+{
+	check(InDataChannel);
+	DataChannel = InDataChannel;
+
+	const FNiagaraDataChannelGameDataLayout& Layout = DataChannel->GetGameDataLayout();
+	VariableData.Empty();
+	VariableData.SetNum(Layout.VariableIndices.Num());
+	for(const auto& Pair : Layout.VariableIndices)
+	{
+		int32 Index = Pair.Value;
+		VariableData[Index].Init(Pair.Key);
+	}
+}
+
+void FNiagaraDataChannelGameData::Empty()
+{
+	NumElements = 0;
+	for (FNiagaraDataChannelVariableBuffer& VarData : VariableData)
+	{
+		VarData.Empty();
+	}
+}
+
+void FNiagaraDataChannelGameData::Reset()
+{
+	NumElements = 0;
+	for (FNiagaraDataChannelVariableBuffer& VarData : VariableData)
+	{
+		VarData.Reset();
+	}
+}
+
+void FNiagaraDataChannelGameData::SetNum(int32 NewNum)
+{
+	NumElements = NewNum;
+	for (auto& Buffer : VariableData)
+	{
+		Buffer.SetNum(NewNum);
+	}
+}
+
+FNiagaraDataChannelVariableBuffer* FNiagaraDataChannelGameData::FindVariableBuffer(const FNiagaraVariableBase& Var)
+{
+	const FNiagaraDataChannelGameDataLayout& Layout = DataChannel->GetGameDataLayout();
+	if(const int32* Idx = Layout.VariableIndices.Find(Var))
+	{
+		return &VariableData[*Idx];
+	}
+	return nullptr;
+}
+
+void FNiagaraDataChannelGameData::WriteToDataSet(FNiagaraDataBuffer* DestBuffer, int32 DestStartIdx)
+{
+	const FNiagaraDataSetCompiledData& CompiledData = DestBuffer->GetOwner()->GetCompiledData();
+
+	static TArray<uint8> LWCConversionBuffer;
+
+	int32 NumInstances = NumElements;
+
+	if (NumInstances == 0)
+	{
+		return;
+	}
+	
+	DestBuffer->SetNumInstances(NumInstances);
+
+	const FNiagaraDataChannelGameDataLayout& Layout = DataChannel->GetGameDataLayout();
+
+	for (const TPair<FNiagaraVariableBase, int32>& VarIndexPair : Layout.VariableIndices)
+	{
+		const FNiagaraVariableBase& Var = VarIndexPair.Key;
+		int32 VarIndex = VarIndexPair.Value;
+		FNiagaraDataChannelVariableBuffer& VarBuffer = VariableData[VarIndex];
+		uint8* SrcData = VarBuffer.Data.GetData();
+		
+		FNiagaraVariableBase SimVar = Var;
+
+		//Convert from LWC types to Niagara Simulation Types where required.
+		if (FNiagaraTypeHelper::IsLWCType(Var.GetType()))
+		{
+			SimVar.SetType(FNiagaraTypeDefinition(FNiagaraTypeHelper::FindNiagaraFriendlyTopLevelStruct(CastChecked<UScriptStruct>(Var.GetType().GetStruct()), ENiagaraStructConversion::Simulation)));
+		}
+
+		int32 SimVarIndex = CompiledData.Variables.IndexOfByKey(SimVar);
+		if (SimVarIndex == INDEX_NONE)
+		{
+			continue; //Did not find this variable in the dataset. Warn?
+		}
+
+		//Convert from LWC types to Niagara Simulation Types where required.
+// 		if (FNiagaraTypeHelper::IsLWCType(Var.GetType()))
+// 		{
+// 			LWCConversionBuffer.SetNumUninitialized(SimVar.GetSizeInBytes() * NumInstances, false);
+// 			
+// 			//TODO: Can we just avoid all this and convert any doubles direct to floats when we transfer the date into the dest buffer?
+// 			FNiagaraLwcStructConverter& Converter = LwcConverters[VarIndex];
+// 			Converter.ConvertDataToSimulation(LWCConversionBuffer.GetData(), SrcData, NumInstances);
+// 			SrcData = LWCConversionBuffer.GetData();
+// 		}
+
+		int32 SrcVarSize = Var.GetSizeInBytes();
+		int32 DestVarSize = SimVar.GetSizeInBytes();
+		const FNiagaraVariableLayoutInfo& SimLayout = CompiledData.VariableLayouts[SimVarIndex];
+
+		int32 FloatCompIdx = SimLayout.FloatComponentStart;
+		int32 IntCompIdx = SimLayout.Int32ComponentStart;
+		int32 HalfCompIdx = SimLayout.HalfComponentStart;
+
+		TFunction<void(UScriptStruct*, UScriptStruct*)> WriteData;
+		WriteData = [&](UScriptStruct* SrcStruct, UScriptStruct* DestStruct)
+		{
+			TFieldIterator<FProperty> SrcPropertyIt(SrcStruct, EFieldIteratorFlags::IncludeSuper);
+			TFieldIterator<FProperty> DestPropertyIt(DestStruct, EFieldIteratorFlags::IncludeSuper);
+			for (; SrcPropertyIt; ++SrcPropertyIt, ++DestPropertyIt)
+			{
+				FProperty* SrcProperty = *SrcPropertyIt;
+				FProperty* DestProperty = *DestPropertyIt;
+				int32 SrcOffset = SrcProperty->GetOffset_ForInternal();
+
+				//Convert any LWC doubles to floats. //TODO: Insert LWCTile... probably need to explicitly check for vectors etc.
+				if (SrcProperty->IsA(FDoubleProperty::StaticClass()))
+				{
+					check(DestProperty->IsA(FFloatProperty::StaticClass()));
+					float* Dest = (float*)DestBuffer->GetComponentPtrFloat(FloatCompIdx++);
+					//Write all instances for this component.
+					for (int32 i = 0; i < NumInstances; ++i)
+					{
+						double* Src = (double*)((SrcData + i * SrcVarSize) + SrcOffset);
+						*Dest++ = *Src;
+					}
+				}
+				else if (SrcProperty->IsA(FFloatProperty::StaticClass()))
+				{
+					float* Dest = (float*)DestBuffer->GetComponentPtrFloat(FloatCompIdx++);
+
+					//Write all instances for this component.
+					for (int32 i = 0; i < NumInstances; ++i)
+					{
+						float* Src = (float*)((SrcData + i * SrcVarSize) + SrcOffset);
+						*Dest++ = *Src;
+					}
+				}
+				else if (SrcProperty->IsA(FUInt16Property::StaticClass()))
+				{
+					FFloat16* Dest = (FFloat16*)DestBuffer->GetComponentPtrHalf(HalfCompIdx++);
+
+					//Write all instances for this component.
+					for (int32 i = 0; i < NumInstances; ++i)
+					{
+						FFloat16* Src = (FFloat16*)((SrcData + i * SrcVarSize) + SrcOffset);
+						*Dest++ = *Src;						
+					}
+				}
+				else if (SrcProperty->IsA(FIntProperty::StaticClass()) || SrcProperty->IsA(FBoolProperty::StaticClass()))
+				{
+					int32* Dest = (int32*)DestBuffer->GetComponentPtrInt32(IntCompIdx++);
+
+					//Write all instances for this component.
+					for (int32 i = 0; i < NumInstances; ++i)
+					{
+						int32* Src = (int32*)((SrcData + i * SrcVarSize) + SrcOffset);
+						*Dest++ = *Src;						
+					}
+				}
+				//Should be able to support double easily enough
+				else if (FStructProperty* StructProp = CastField<FStructProperty>(SrcProperty))
+				{
+					FStructProperty* DestStructProp = CastField<FStructProperty>(DestProperty);
+					WriteData(StructProp->Struct, DestStructProp->Struct);
+				}
+				else
+				{
+					checkf(false, TEXT("Property(%s) Class(%s) is not a supported type"), *SrcProperty->GetName(), *SrcProperty->GetClass()->GetName());
+				}
+			}
+		};
+		WriteData(Var.GetType().GetScriptStruct(), SimVar.GetType().GetScriptStruct());
+	}
+}
+void FNiagaraDataChannelGameData::AppendFromGameData(const FNiagaraDataChannelGameData& GameData)
+{
+	check(GetDataChannel() == GameData.GetDataChannel());
+
+	NumElements += GameData.Num();
+	TConstArrayView<FNiagaraDataChannelVariableBuffer> SrcBuffers = GameData.GetVariableBuffers();
+	for(int32 i=0; i < SrcBuffers.Num(); ++i)
+	{
+		const FNiagaraDataChannelVariableBuffer& SrcBuffer = SrcBuffers[i];
+		FNiagaraDataChannelVariableBuffer& DestBuffer = VariableData[i];
+		DestBuffer.Data.Append(SrcBuffer.Data);
+	}
+}
+
+void FNiagaraDataChannelGameData::AppendFromDataSet(const FNiagaraDataBuffer* SrcBuffer, FVector3f LwcTile)
+{
+	const FNiagaraDataSetCompiledData& CompiledData = SrcBuffer->GetOwner()->GetCompiledData();
+
+	static TArray<uint8> LWCConversionBuffer;
+
+	int32 NumInstances = SrcBuffer->GetNumInstances();
+	int32 OrigNumDataChannel = NumElements;
+	NumElements += NumInstances;
+
+	const FNiagaraDataChannelGameDataLayout& Layout = DataChannel->GetGameDataLayout();
+	for (const TPair<FNiagaraVariableBase, int32>& VarIndexPair : Layout.VariableIndices)
+	{
+		const FNiagaraVariableBase& Var = VarIndexPair.Key;
+		int32 VarIndex = VarIndexPair.Value;
+		FNiagaraDataChannelVariableBuffer& VarBuffer = VariableData[VarIndex];
+
+		int32 VarSize = Var.GetSizeInBytes();
+
+		VarBuffer.SetNum(NumElements);
+
+		uint8* DestDataBase = VarBuffer.Data.GetData();
+		uint8* DestData = DestDataBase;
+		
+		FNiagaraVariableBase SimVar = Var;
+		if (FNiagaraTypeHelper::IsLWCType(Var.GetType()))
+		{
+			SimVar.SetType(FNiagaraTypeDefinition(FNiagaraTypeHelper::FindNiagaraFriendlyTopLevelStruct(Var.GetType().GetScriptStruct(), ENiagaraStructConversion::Simulation)));
+		}
+		
+		int32 SimVarIndex = CompiledData.Variables.IndexOfByKey(SimVar);
+		if (SimVarIndex == INDEX_NONE)
+		{
+			continue; //Did not find this variable in the dataset. Warn?
+		}
+
+		//Convert from Niagara Simulation to LWC types where required.
+// 		if (FNiagaraTypeHelper::IsLWCType(Var.GetType()))
+// 		{
+// 			LWCConversionBuffer.SetNumUninitialized(Var.GetSizeInBytes() * NumInstances, false);
+// 
+// 			//Write instead to our intermediate conversion buffer and do the LWC convert after.
+// 			DestData = LWCConversionBuffer.GetData();
+// 		}
+
+		const FNiagaraVariableLayoutInfo& SimLayout = CompiledData.VariableLayouts[SimVarIndex];
+		
+		int32 FloatCompIdx = SimLayout.FloatComponentStart;
+		int32 IntCompIdx = SimLayout.Int32ComponentStart;
+		int32 HalfCompIdx = SimLayout.HalfComponentStart;
+		
+		TFunction<void(UScriptStruct*, UScriptStruct*)> ReadData;
+		ReadData = [&](UScriptStruct* SrcStruct, UScriptStruct* DestStruct)
+		{
+			TFieldIterator<FProperty> SrcPropertyIt(SrcStruct, EFieldIteratorFlags::IncludeSuper);
+			TFieldIterator<FProperty> DestPropertyIt(DestStruct, EFieldIteratorFlags::IncludeSuper);			
+			for (; SrcPropertyIt; ++SrcPropertyIt, ++DestPropertyIt)
+			{
+				FProperty* SrcProperty = *SrcPropertyIt;
+				FProperty* DestProperty = *DestPropertyIt;
+				DestData = DestDataBase + DestProperty->GetOffset_ForInternal();
+				if (DestPropertyIt->IsA(FDoubleProperty::StaticClass()))//TODO: Insert LWCTile... probably need to explicitly check for vectors etc.
+				{
+					double* Dest = (double*)(DestData);
+					float* Src = (float*)SrcBuffer->GetComponentPtrFloat(FloatCompIdx++);
+
+					//Write all instances for this component.
+					for (int32 i = 0; i < NumInstances; ++i)
+					{
+						Dest = (double*)(DestData + VarSize * i);
+						*Dest = *Src++;
+					}
+				}
+				else if (DestPropertyIt->IsA(FFloatProperty::StaticClass()))
+				{
+					float* Dest = (float*)(DestData);
+					float* Src = (float*)SrcBuffer->GetComponentPtrFloat(FloatCompIdx++);
+
+					//Write all instances for this component.
+					for (int32 i = 0; i < NumInstances; ++i)
+					{
+						Dest = (float*)(DestData + VarSize * i);
+						*Dest = *Src++;
+					}
+				}
+				else if (DestPropertyIt->IsA(FUInt16Property::StaticClass()))
+				{
+					FFloat16* Dest = (FFloat16*)(DestData);
+					FFloat16* Src = (FFloat16*)SrcBuffer->GetComponentPtrHalf(HalfCompIdx++);
+
+					//Write all instances for this component.
+					for (int32 i = 0; i < NumInstances; ++i)
+					{
+						Dest = (FFloat16*)(DestData + VarSize * i);
+						*Dest = *Src++;
+					}
+				}
+				else if (DestPropertyIt->IsA(FIntProperty::StaticClass()) || DestPropertyIt->IsA(FBoolProperty::StaticClass()))
+				{
+					int32* Dest = (int32*)(DestData);
+					int32* Src = (int32*)SrcBuffer->GetComponentPtrInt32(IntCompIdx++);
+
+					//Write all instances for this component.
+					for (int32 i = 0; i < NumInstances; ++i)
+					{
+						Dest = (int32*)(DestData + VarSize * i);
+						*Dest = *Src++;
+					}
+				}
+				//Should be able to support double easily enough
+				else if (FStructProperty* SrcStructProp = CastField<FStructProperty>(SrcProperty))
+				{
+					FStructProperty* DestStructProp = CastField<FStructProperty>(DestProperty);
+					ReadData(SrcStructProp->Struct, DestStructProp->Struct);
+				}
+				else
+				{
+					checkf(false, TEXT("Property(%s) Class(%s) is not a supported type"), *SrcProperty->GetName(), *SrcProperty->GetClass()->GetName());
+				}
+			}
+		};
+		ReadData(SimVar.GetType().GetScriptStruct(), Var.GetType().GetScriptStruct());
+
+		//For LWC types we're written into an intermediate structure for the sim var. Now convert that into LWC proper.
+		if (FNiagaraTypeHelper::IsLWCType(Var.GetType()))
+		{
+			const FNiagaraLwcStructConverter& Converter = Layout.LwcConverters[VarIndex];
+			VarBuffer.SetNum(NumElements);
+			//TODO: Incorporate the LWC Tile
+			Converter.ConvertDataFromSimulation(VarBuffer.Data.GetData() + Var.GetSizeInBytes() * OrigNumDataChannel, LWCConversionBuffer.GetData(), NumInstances);
+		}
+	}
+}
+
+//////////////////////////////////////////////////////////////////////////
+
+FNiagaraWorldDataChannelStore::~FNiagaraWorldDataChannelStore()
+{
+	Empty();
+}
+
+void FNiagaraWorldDataChannelStore::Empty()
+{
+	GameData.Reset();
+
+	/** We defer the deletion of the dataset to the RT to be sure all in-flight RT commands have finished using it.*/
+	ENQUEUE_RENDER_COMMAND(FDeleteContextCommand)(
+		[CPUDataChannelDataSet = CPUSimData, GPUDataChannelDataSet = GPUSimData](FRHICommandListImmediate& RHICmdList)
+		{
+			if (CPUDataChannelDataSet != nullptr)
+			{
+				delete CPUDataChannelDataSet;
+			}
+			if (GPUDataChannelDataSet != nullptr)
+			{
+				delete GPUDataChannelDataSet;
+			}
+		}
+	);
+	CPUSimData = nullptr;
+	GPUSimData = nullptr;
+}
+
+void FNiagaraWorldDataChannelStore::Init(UNiagaraDataChannelHandler* Owner)
+{
+	GameData = Owner->GetDataChannel()->CreateGameData();
+
+	CPUSimData = new FNiagaraDataSet();
+	GPUSimData = new FNiagaraDataSet();
+
+	const FNiagaraDataSetCompiledData& CompiledData = Owner->GetDataChannel()->GetCompiledData(ENiagaraSimTarget::CPUSim);
+	const FNiagaraDataSetCompiledData& CompiledDataGPU = Owner->GetDataChannel()->GetCompiledData(ENiagaraSimTarget::GPUComputeSim);
+	CPUSimData->Init(&CompiledData);
+	GPUSimData->Init(&CompiledDataGPU);
+}
+
+void FNiagaraWorldDataChannelStore::Tick(UNiagaraDataChannelHandler* Owner)
+{
+
+}
+
+
+//////////////////////////////////////////////////////////////////////////
+
+void UNiagaraDataChannel::PostLoad()
+{
+	Super::PostLoad();
+
+	//TODO: Can serialize?
+	GameDataLayout.Init(Variables);
+
+	FNiagaraWorldManager::ForAllWorldManagers(
+		[DataChannel = this](FNiagaraWorldManager& WorldMan)
+		{
+			WorldMan.GetDataChannelManager().InitDataChannel(DataChannel, true);
+		});
+}
+
+void UNiagaraDataChannel::BeginDestroy()
+{
+	Super::BeginDestroy();
+
+	FNiagaraWorldManager::ForAllWorldManagers(
+		[DataChannel = this](FNiagaraWorldManager& WorldMan)
+		{
+			WorldMan.WaitForAsyncWork();
+			WorldMan.GetDataChannelManager().RemoveDataChannel(DataChannel->GetChannelName());
+		});
+}
+
+#if WITH_EDITOR
+
+void UNiagaraDataChannel::PreEditChange(FProperty* PropertyAboutToChange)
+{
+	//We'll do this for all channels inside the definitions pre edit change.
+	// 
+// 	FNiagaraWorldManager::ForAllWorldManagers(
+// 		[DataChannel = this](FNiagaraWorldManager& WorldMan)
+// 		{
+// 			WorldMan.WaitForAsyncWork();
+// 			WorldMan.GetDataChannelManager().RemoveDataChannel(DataChannel->GetChannelName());
+// 		});
+}
+
+void UNiagaraDataChannel::PostEditChangeProperty(FPropertyChangedEvent& PropertyChangedEvent)
+{
+	Super::PostEditChangeProperty(PropertyChangedEvent);
+	
+	check(IsInGameThread());
+
+	//Refresh compiled data
+	CompiledData.Empty();
+	CompiledDataGPU.Empty();
+	GetCompiledData(ENiagaraSimTarget::CPUSim);
+	GetCompiledData(ENiagaraSimTarget::GPUComputeSim);
+
+	GameDataLayout.Init(Variables);
+
+	OnChangedDelegate.Broadcast(this);
+
+	//We'll reinit this channel with the world managers etc in the PostEditChange of the data channel definitions.
+}
+
+#endif//WITH_EDITOR
+
+const FNiagaraDataSetCompiledData& UNiagaraDataChannel::GetCompiledData(ENiagaraSimTarget SimTarget)const
+{
+	if(SimTarget == ENiagaraSimTarget::CPUSim)
+	{
+		if (CompiledData.Variables.Num() != Variables.Num())
+		{
+			//Build the compiled data from the current variables but convert to Simulation types.
+			CompiledData.Empty();
+			CompiledData.SimTarget = ENiagaraSimTarget::CPUSim;
+			for (FNiagaraVariableBase Var : Variables)
+			{
+				Var.SetType(FNiagaraTypeDefinition(FNiagaraTypeHelper::FindNiagaraFriendlyTopLevelStruct(Var.GetType().GetScriptStruct(), ENiagaraStructConversion::Simulation)));
+				CompiledData.Variables.Add(Var);
+			}
+			CompiledData.BuildLayout();
+		}
+		return CompiledData;
+	}
+	else
+	{
+		if (CompiledDataGPU.Variables.Num() != Variables.Num())
+		{
+			check(SimTarget == ENiagaraSimTarget::GPUComputeSim);
+			//Build the compiled data from the current variables but convert to Simulation types.
+			CompiledDataGPU.Empty();
+			CompiledDataGPU.SimTarget = ENiagaraSimTarget::GPUComputeSim;
+			for (FNiagaraVariableBase Var : Variables)
+			{
+				Var.SetType(FNiagaraTypeDefinition(FNiagaraTypeHelper::FindNiagaraFriendlyTopLevelStruct(Var.GetType().GetScriptStruct(), ENiagaraStructConversion::Simulation)));
+				CompiledDataGPU.Variables.Add(Var);
+			}
+			CompiledDataGPU.BuildLayout();
+		}
+		return CompiledDataGPU;
+	}
+}
+
+FNiagaraDataChannelGameDataPtr UNiagaraDataChannel::CreateGameData()const
+{
+	FNiagaraDataChannelGameDataPtr NewGameData = MakeShared<FNiagaraDataChannelGameData>();
+	NewGameData->Init(this);
+	return NewGameData;
+}
+
+bool UNiagaraDataChannel::IsValid()const
+{
+	return ChannelName != NAME_None && Variables.Num() > 0 && CompiledData.Variables.Num() == Variables.Num() && CompiledDataGPU.Variables.Num() == Variables.Num();
+}
+
+//////////////////////////////////////////////////////////////////////////
+
+UNiagaraDataChannelLibrary::UNiagaraDataChannelLibrary(const FObjectInitializer& ObjectInitializer)
+	: Super(ObjectInitializer)
+{
+}
+
+UNiagaraDataChannelHandler* UNiagaraDataChannelLibrary::GetNiagaraDataChannel(const UObject* WorldContextObject, FName Channel)
+{
+	UWorld* World = (WorldContextObject != nullptr) ? GEngine->GetWorldFromContextObject(WorldContextObject, EGetWorldErrorMode::LogAndReturnNull) : nullptr;
+	if (World)
+	{
+		if (FNiagaraWorldManager* WorldMan = FNiagaraWorldManager::Get(World))
+		{
+			return WorldMan->GetDataChannelManager().FindDataChannelHandler(Channel);
+		}
+	}
+	return nullptr;
+}
+
+
+#undef LOCTEXT_NAMESPACE
\ No newline at end of file
diff --git a/Engine/Plugins/FX/Niagara/Source/Niagara/Private/NiagaraDataChannelAccessor.cpp b/Engine/Plugins/FX/Niagara/Source/Niagara/Private/NiagaraDataChannelAccessor.cpp
new file mode 100644
--- /dev/null
+++ b/Engine/Plugins/FX/Niagara/Source/Niagara/Private/NiagaraDataChannelAccessor.cpp
@@ -0,0 +1,246 @@
+// Copyright Epic Games, Inc. All Rights Reserved.
+
+#include "NiagaraDataChannelAccessor.h"
+
+#include "CoreMinimal.h"
+#include "UObject/UObjectGlobals.h"
+#include "UObject/Package.h"
+#include "NiagaraDataChannel.h"
+#include "NiagaraDataChannelHandler.h"
+
+//////////////////////////////////////////////////////////////////////////
+
+namespace NiagaraDataChannelAccesorLocal
+{
+	//TODO: Move into FNiagaraTypeDefinition.
+	static FNiagaraTypeDefinition DoubleDef;
+	static FNiagaraTypeDefinition Vector2DDef;
+	static FNiagaraTypeDefinition VectorDef;
+	static FNiagaraTypeDefinition Vector4Def;
+	static FNiagaraTypeDefinition QuatDef;
+
+	void InitDefs()
+	{
+		static UPackage* CoreUObjectPkg = FindObjectChecked<UPackage>(nullptr, TEXT("/Script/CoreUObject"));
+		static UScriptStruct* Vector2Struct = FindObjectChecked<UScriptStruct>(CoreUObjectPkg, TEXT("Vector2D"));
+		static UScriptStruct* Vector3Struct = FindObjectChecked<UScriptStruct>(CoreUObjectPkg, TEXT("Vector"));
+		static UScriptStruct* Vector4Struct = FindObjectChecked<UScriptStruct>(CoreUObjectPkg, TEXT("Vector4"));
+		static UScriptStruct* Quat4Struct = FindObjectChecked<UScriptStruct>(CoreUObjectPkg, TEXT("Quat"));
+		DoubleDef = FNiagaraTypeDefinition(FNiagaraDouble::StaticStruct(), FNiagaraTypeDefinition::EAllowUnfriendlyStruct::Allow);
+		Vector2DDef = FNiagaraTypeDefinition(Vector2Struct, FNiagaraTypeDefinition::EAllowUnfriendlyStruct::Allow);
+		VectorDef = FNiagaraTypeDefinition(Vector3Struct, FNiagaraTypeDefinition::EAllowUnfriendlyStruct::Allow);
+		Vector4Def = FNiagaraTypeDefinition(Vector4Struct, FNiagaraTypeDefinition::EAllowUnfriendlyStruct::Allow);
+		QuatDef = FNiagaraTypeDefinition(Quat4Struct, FNiagaraTypeDefinition::EAllowUnfriendlyStruct::Allow);
+	}
+
+	const FNiagaraTypeDefinition& GetDoubleDef()
+	{
+		if(DoubleDef.IsValid() == false){ InitDefs(); }
+		return DoubleDef;
+	}
+
+	const FNiagaraTypeDefinition& GetVector2DDef()
+	{
+		if (Vector2DDef.IsValid() == false) { InitDefs(); }
+		return Vector2DDef;
+	}
+
+	const FNiagaraTypeDefinition& GetVectorDef()
+	{
+		if (VectorDef.IsValid() == false) { InitDefs(); }
+		return VectorDef;
+	}
+
+	const FNiagaraTypeDefinition& GetVector4Def()
+	{
+		if (Vector4Def.IsValid() == false) { InitDefs(); }
+		return Vector4Def;
+	}
+
+	const FNiagaraTypeDefinition& GetQuatDef()
+	{
+		if (QuatDef.IsValid() == false) { InitDefs(); }
+		return QuatDef;
+	}
+}
+
+//////////////////////////////////////////////////////////////////////////
+
+bool UNiagaraDataChannelReader::InitAccess(UActorComponent* OwningComponent)
+{
+	Data = nullptr;
+	check(Owner);
+
+	if (OwningComponent == nullptr)
+	{
+		UE_LOG(LogNiagara, Warning, TEXT("Call to UNiagaraDataChannelReader::InitAccess did not provide a valid Owning Componet."));
+		return false;
+	}
+
+	Data = Owner->GetGameData(OwningComponent);
+	return true;
+}
+
+int32 UNiagaraDataChannelReader::Num()const
+{
+	if (Data.IsValid())
+	{
+		return Data->Num();
+	}
+	return 0;
+}
+
+template<typename T>
+T UNiagaraDataChannelReader::ReadData(const FNiagaraVariableBase& Var, int32 Index)const
+{
+	if (ensure(Data.IsValid()))
+	{
+		if (FNiagaraDataChannelVariableBuffer* VarBuffer = Data->FindVariableBuffer(Var))
+		{
+			T OutData;
+			VarBuffer->Read<T>(Index, OutData);
+			return OutData;
+		}
+	}
+	return T();
+}
+
+double UNiagaraDataChannelReader::ReadFloat(FName VarName, int32 Index)const
+{
+	return ReadData<double>(FNiagaraVariableBase(NiagaraDataChannelAccesorLocal::GetDoubleDef(), VarName), Index);
+}
+
+FVector2D UNiagaraDataChannelReader::ReadVector2D(FName VarName, int32 Index)const
+{
+	return ReadData<FVector2D>(FNiagaraVariableBase(NiagaraDataChannelAccesorLocal::GetVector2DDef(), VarName), Index);
+}
+
+FVector UNiagaraDataChannelReader::ReadVector(FName VarName, int32 Index)const
+{
+	return ReadData<FVector>(FNiagaraVariableBase(NiagaraDataChannelAccesorLocal::GetVectorDef(), VarName), Index);
+}
+
+FVector4 UNiagaraDataChannelReader::ReadVector4(FName VarName, int32 Index)const
+{
+	return ReadData<FVector4>(FNiagaraVariableBase(NiagaraDataChannelAccesorLocal::GetVector4Def(), VarName), Index);
+}
+
+FQuat UNiagaraDataChannelReader::ReadQuat(FName VarName, int32 Index)const
+{
+	return ReadData<FQuat>(FNiagaraVariableBase(NiagaraDataChannelAccesorLocal::GetQuatDef(), VarName), Index);
+}
+
+FLinearColor UNiagaraDataChannelReader::ReadLinearColor(FName VarName, int32 Index)const
+{
+	return ReadData<FLinearColor>(FNiagaraVariableBase(FNiagaraTypeDefinition::GetColorDef(), VarName), Index);
+}
+
+int32 UNiagaraDataChannelReader::ReadInt(FName VarName, int32 Index)const
+{
+	return ReadData<int32>(FNiagaraVariableBase(FNiagaraTypeDefinition::GetIntDef(), VarName), Index);
+}
+
+bool UNiagaraDataChannelReader::ReadBool(FName VarName, int32 Index)const
+{
+	return ReadData<bool>(FNiagaraVariableBase(FNiagaraTypeDefinition::GetBoolDef(), VarName), Index);
+}
+
+//////////////////////////////////////////////////////////////////////////
+
+
+template<typename T>
+void UNiagaraDataChannelWriter::WriteData(const FNiagaraVariableBase& Var, int32 Index, const T& InData)
+{
+	if (ensure(Data.IsValid()))
+	{
+		if (FNiagaraDataChannelVariableBuffer* VarBuffer = Data->FindVariableBuffer(Var))
+		{
+			VarBuffer->Write<T>(Index, InData);
+		}
+	}
+}
+
+bool UNiagaraDataChannelWriter::InitWrite(UActorComponent* OwningComponent, int32 Count, bool bVisibleToGame, bool bVisibleToCPU, bool bVisibleToGPU)
+{
+	if( OwningComponent == nullptr )
+	{
+		UE_LOG(LogNiagara, Warning, TEXT("Call to UNiagaraDataChannelWriter::InitWrite did not provide a valid Owning Componet."));
+		return false;
+	}
+
+	if (Count == 0)
+	{
+		UE_LOG(LogNiagara, Warning, TEXT("Call to UNiagaraDataChannelWriter::InitWrite with Count == 0. Ignored."));
+		return false;
+	}
+
+	check(Owner);
+
+	Data = Owner->GetDataChannel()->CreateGameData();
+
+	FNiagaraDataChannelPublishRequest PublishRequest;
+	PublishRequest.bVisibleToGame = bVisibleToGame;
+	PublishRequest.bVisibleToCPUSims = bVisibleToCPU;
+	PublishRequest.bVisibleToGPUSims = bVisibleToGPU;
+	PublishRequest.GameData = Data;
+	Owner->Publish(PublishRequest);
+
+	Data->SetNum(Count);
+
+	return true;
+}
+
+int32 UNiagaraDataChannelWriter::Num()const
+{
+	if (Data.IsValid())
+	{
+		return Data->Num();
+	}
+	return 0;
+}
+
+void UNiagaraDataChannelWriter::WriteFloat(FName VarName, int32 Index, double InData)
+{
+	WriteData(FNiagaraVariableBase(NiagaraDataChannelAccesorLocal::GetDoubleDef(), VarName), Index, InData);
+}
+
+void UNiagaraDataChannelWriter::WriteVector2D(FName VarName, int32 Index, FVector2D InData)
+{
+	WriteData(FNiagaraVariableBase(NiagaraDataChannelAccesorLocal::GetVector2DDef(), VarName), Index, InData);
+}
+
+void UNiagaraDataChannelWriter::WriteVector(FName VarName, int32 Index, FVector InData)
+{
+	WriteData(FNiagaraVariableBase(NiagaraDataChannelAccesorLocal::GetVectorDef(), VarName), Index, InData);
+}
+
+void UNiagaraDataChannelWriter::WriteVector4(FName VarName, int32 Index, FVector4 InData)
+{
+	WriteData(FNiagaraVariableBase(NiagaraDataChannelAccesorLocal::GetVector4Def(), VarName), Index, InData);
+}
+
+void UNiagaraDataChannelWriter::WriteQuat(FName VarName, int32 Index, FQuat InData)
+{
+	WriteData(FNiagaraVariableBase(NiagaraDataChannelAccesorLocal::GetQuatDef(), VarName), Index, InData);
+}
+
+void UNiagaraDataChannelWriter::WriteLinearColor(FName VarName, int32 Index, FLinearColor InData)
+{
+	WriteData(FNiagaraVariableBase(FNiagaraTypeDefinition::GetColorDef(), VarName), Index, InData);
+}
+
+void UNiagaraDataChannelWriter::WriteInt(FName VarName, int32 Index, int32 InData)
+{
+	WriteData(FNiagaraVariableBase(FNiagaraTypeDefinition::GetIntDef(), VarName), Index, InData);
+}
+
+void UNiagaraDataChannelWriter::WriteBool(FName VarName, int32 Index, bool InData)
+{
+	WriteData(FNiagaraVariableBase(FNiagaraTypeDefinition::GetBoolDef(), VarName), Index, InData);
+}
+
+void UNiagaraDataChannelWriter::WriteSpawnInfo(FName VarName, int32 Index, FNiagaraSpawnInfo InData)
+{
+	WriteData(FNiagaraVariableBase(FNiagaraTypeDefinition(FNiagaraSpawnInfo::StaticStruct()), VarName), Index, InData);
+}
+
diff --git a/Engine/Plugins/FX/Niagara/Source/Niagara/Private/NiagaraDataChannelAccessor.h b/Engine/Plugins/FX/Niagara/Source/Niagara/Private/NiagaraDataChannelAccessor.h
new file mode 100644
--- /dev/null
+++ b/Engine/Plugins/FX/Niagara/Source/Niagara/Private/NiagaraDataChannelAccessor.h
@@ -0,0 +1,112 @@
+// Copyright Epic Games, Inc. All Rights Reserved.
+
+#pragma once
+
+#include "NiagaraDataChannelCommon.h"
+#include "NiagaraDataChannelAccessor.generated.h"
+
+
+/** 
+Initial simple API for reading and writing data in a data channel from game code / BP. 
+Likely to be replaced in the near future with a custom BP node and a helper struct.
+*/
+
+
+UCLASS(Experimental, BlueprintType)
+class UNiagaraDataChannelReader : public UObject
+{
+	GENERATED_BODY()
+private:
+
+	FNiagaraDataChannelGameDataPtr Data = nullptr;
+
+	template<typename T>
+	T ReadData(const FNiagaraVariableBase& Var, int32 Index)const;
+
+public:
+
+	UPROPERTY()
+	TObjectPtr<UNiagaraDataChannelHandler> Owner;
+	
+	/** Call before each access to the data channel to grab the correct data to read. */
+	UFUNCTION(BlueprintCallable, Category = NiagaraDataChannel, meta = (Keywords = "Niagara DataChannel"))
+	bool InitAccess(UActorComponent* OwningComponent);
+
+	UFUNCTION(BlueprintCallable, Category = NiagaraDataChannel, meta = (Keywords = "Niagara DataChannel"))
+	int32 Num()const;
+
+	UFUNCTION(BlueprintCallable, Category = NiagaraDataChannel, meta = (Keywords = "Niagara DataChannel"))
+	double ReadFloat(FName VarName, int32 Index)const;
+
+	UFUNCTION(BlueprintCallable, Category = NiagaraDataChannel, meta = (Keywords = "Niagara DataChannel"))
+	FVector2D ReadVector2D(FName VarName, int32 Index)const;
+
+	UFUNCTION(BlueprintCallable, Category = NiagaraDataChannel, meta = (Keywords = "Niagara DataChannel"))
+	FVector ReadVector(FName VarName, int32 Index)const;
+
+	UFUNCTION(BlueprintCallable, Category = NiagaraDataChannel, meta = (Keywords = "Niagara DataChannel"))
+	FVector4 ReadVector4(FName VarName, int32 Index)const;
+
+	UFUNCTION(BlueprintCallable, Category = NiagaraDataChannel, meta = (Keywords = "Niagara DataChannel"))
+	FQuat ReadQuat(FName VarName, int32 Index)const;
+
+	UFUNCTION(BlueprintCallable, Category = NiagaraDataChannel, meta = (Keywords = "Niagara DataChannel"))
+	FLinearColor ReadLinearColor(FName VarName, int32 Index)const;
+
+	UFUNCTION(BlueprintCallable, Category = NiagaraDataChannel, meta = (Keywords = "Niagara DataChannel"))
+	int32 ReadInt(FName VarName, int32 Index)const;
+
+	UFUNCTION(BlueprintCallable, Category = NiagaraDataChannel, meta = (Keywords = "Niagara DataChannel"))
+	bool ReadBool(FName VarName, int32 Index)const;
+};
+
+UCLASS(Experimental, BlueprintType)
+class UNiagaraDataChannelWriter : public UObject
+{
+	GENERATED_BODY()
+private:
+
+	FNiagaraDataChannelGameDataPtr Data = nullptr;
+
+	template<typename T>
+	void WriteData(const FNiagaraVariableBase& Var, int32 Index, const T& InData);
+
+public:
+
+	UPROPERTY()
+	TObjectPtr<UNiagaraDataChannelHandler> Owner;
+	
+	/** Call before each batch of writes to allocate the data we'll be writing to. */
+	UFUNCTION(BlueprintCallable, Category = NiagaraDataChannel, meta = (Keywords = "niagara DataChannel"))
+	bool InitWrite(UActorComponent* OwningComponent, int32 Count, bool bVisibleToGame=true, bool bVisibleToCPU=true, bool bVisibleToGPU=true);
+
+	UFUNCTION(BlueprintCallable, Category = NiagaraDataChannel, meta = (Keywords = "niagara DataChannel"))
+	int32 Num()const;
+
+	UFUNCTION(BlueprintCallable, Category = NiagaraDataChannel, meta = (Keywords = "niagara DataChannel"))
+	void WriteFloat(FName VarName, int32 Index, double InData);
+	
+	UFUNCTION(BlueprintCallable, Category = NiagaraDataChannel, meta = (Keywords = "niagara DataChannel"))
+	void WriteVector2D(FName VarName, int32 Index, FVector2D InData);
+
+	UFUNCTION(BlueprintCallable, Category = NiagaraDataChannel, meta = (Keywords = "niagara DataChannel"))
+	void WriteVector(FName VarName, int32 Index, FVector InData);
+
+	UFUNCTION(BlueprintCallable, Category = NiagaraDataChannel, meta = (Keywords = "niagara DataChannel"))
+	void WriteVector4(FName VarName, int32 Index, FVector4 InData);
+
+	UFUNCTION(BlueprintCallable, Category = NiagaraDataChannel, meta = (Keywords = "niagara DataChannel"))
+	void WriteQuat(FName VarName, int32 Index, FQuat InData);
+
+	UFUNCTION(BlueprintCallable, Category = NiagaraDataChannel, meta = (Keywords = "niagara DataChannel"))
+	void WriteLinearColor(FName VarName, int32 Index, FLinearColor InData);
+	
+	UFUNCTION(BlueprintCallable, Category = NiagaraDataChannel, meta = (Keywords = "niagara DataChannel"))
+	void WriteInt(FName VarName, int32 Index, int32 InData);
+
+	UFUNCTION(BlueprintCallable, Category = NiagaraDataChannel, meta = (Keywords = "niagara DataChannel"))
+	void WriteBool(FName VarName, int32 Index, bool InData);
+
+	UFUNCTION(BlueprintCallable, Category = NiagaraDataChannel, meta = (Keywords = "niagara DataChannel"))
+	void WriteSpawnInfo(FName VarName, int32 Index, FNiagaraSpawnInfo InData);
+};
\ No newline at end of file
diff --git a/Engine/Plugins/FX/Niagara/Source/Niagara/Private/NiagaraDataChannelDefinitions.cpp b/Engine/Plugins/FX/Niagara/Source/Niagara/Private/NiagaraDataChannelDefinitions.cpp
new file mode 100644
--- /dev/null
+++ b/Engine/Plugins/FX/Niagara/Source/Niagara/Private/NiagaraDataChannelDefinitions.cpp
@@ -0,0 +1,255 @@
+// Copyright Epic Games, Inc. All Rights Reserved.
+
+#include "NiagaraDataChannelDefinitions.h"
+#include "NiagaraWorldManager.h"
+#include "NiagaraDataChannel.h"
+#include "DataInterface/NiagaraDataInterfaceDataChannelWrite.h"
+#include "DataInterface/NiagaraDataInterfaceDataChannelRead.h"
+#include "DataInterface/NiagaraDataInterfaceDataChannelSpawn.h"
+
+#include "AssetRegistry/AssetRegistryModule.h"
+#include "AssetRegistry/IAssetRegistry.h"
+#include "AssetRegistry/AssetData.h"
+
+#include "Misc/MessageDialog.h"
+
+#define LOCTEXT_NAMESPACE "NiagaraDataChannelDefinitions"
+
+TArray<UNiagaraDataChannelDefinitions*> UNiagaraDataChannelDefinitions::Definitions;
+
+
+void UNiagaraDataChannelDefinitions::PostInitProperties()
+{
+	Super::PostInitProperties();
+}
+
+void UNiagaraDataChannelDefinitions::PostLoad()
+{
+	Super::PostLoad();
+
+	if(Definitions.AddUnique(this) != INDEX_NONE)
+	{
+		FNiagaraWorldManager::ForAllWorldManagers(
+			[&](FNiagaraWorldManager& WorldMan)
+			{
+				for(UNiagaraDataChannel* DataChannel : DataChannels)
+				{
+					WorldMan.GetDataChannelManager().InitDataChannel(DataChannel, false);
+				}
+			});
+	}
+}
+
+void UNiagaraDataChannelDefinitions::BeginDestroy()
+{
+	Super::BeginDestroy();
+
+	if(Definitions.Remove(this) != INDEX_NONE)
+	{
+		FNiagaraWorldManager::ForAllWorldManagers(
+			[&](FNiagaraWorldManager& WorldMan)
+			{
+				for (UNiagaraDataChannel* DataChannel : DataChannels)
+				{
+					WorldMan.GetDataChannelManager().RemoveDataChannel(DataChannel->GetChannelName());
+				}
+			});
+	}
+}
+
+#if WITH_EDITOR
+
+void UNiagaraDataChannelDefinitions::PreEditChange( FProperty* PropertyAboutToChange )
+{
+	Super::PreEditChange(PropertyAboutToChange);
+
+	//Belt and braces refreshing for data channels on edit. Can probably be a bit more targeted than this.	
+	SysUpdateContext.SetDestroySystemSim(true);
+	SysUpdateContext.SetDestroyOnAdd(true);
+
+	FNiagaraWorldManager::ForAllWorldManagers(
+		[&](FNiagaraWorldManager& WorldMan)
+		{
+			WorldMan.WaitForAsyncWork();
+			for (const UNiagaraDataChannel* DataChannel : DataChannels)
+			{
+				if (DataChannel)
+				{
+					WorldMan.GetDataChannelManager().RemoveDataChannel(DataChannel->GetChannelName());
+
+					//Re-initialize any UNiagaraSystems that are using this channel.
+					auto AddSystemsForDIUsingChannel = [&](auto* DI)
+					{
+						UNiagaraSystem* System = DI ? DI->GetTypedOuter<UNiagaraSystem>() : nullptr;
+						if (DI && System && DI->Channel.ChannelName == DataChannel->GetChannelName())
+						{
+							SysUpdateContext.Add(System, true);
+						}
+					};
+
+					//TODO: Originally I iterated over systems and used ForEachDataInterface to check for DataChannel Usage. However ForEachDataInterface skips some important DIs so that needs to be updated before that approach will work.
+					for (TObjectIterator<UNiagaraDataInterfaceDataChannelSpawn> It; It; ++It) { AddSystemsForDIUsingChannel(*It); }
+					for (TObjectIterator<UNiagaraDataInterfaceDataChannelRead> It; It; ++It) { AddSystemsForDIUsingChannel(*It); }
+					for (TObjectIterator<UNiagaraDataInterfaceDataChannelWrite> It; It; ++It) { AddSystemsForDIUsingChannel(*It); }
+				}
+			}
+		});
+
+	//Ensure we don't have any inflight commands that may refernce compiled data in the channel.
+	FlushRenderingCommands();
+}
+
+void UNiagaraDataChannelDefinitions::PostEditChangeProperty(FPropertyChangedEvent& PropertyChangedEvent)
+{
+	Super::PostEditChangeProperty(PropertyChangedEvent);
+
+	FNiagaraWorldManager::ForAllWorldManagers(
+		[&](FNiagaraWorldManager& WorldMan)
+		{			
+			for (const UNiagaraDataChannel* DataChannel : DataChannels)
+			{
+				if (DataChannel)
+				{
+					WorldMan.GetDataChannelManager().InitDataChannel(DataChannel, true);
+
+					//Re-initialize any UNiagaraSystems that are using this channel.
+					auto AddSystemsForDIUsingChannel = [&](auto* DI)
+					{
+						UNiagaraSystem* System = DI ? DI->GetTypedOuter<UNiagaraSystem>() : nullptr;
+						if (DI && System && DI->Channel.ChannelName == DataChannel->GetChannelName())
+						{
+							SysUpdateContext.Add(System, true);
+						}
+					};
+
+					//TODO: Originally I iterated over systems and used ForEachDataInterface to check for DataChannel Usage. However ForEachDataInterface skips some important DIs so that needs to be updated before that approach will work.
+					for (TObjectIterator<UNiagaraDataInterfaceDataChannelSpawn> It; It; ++It) { AddSystemsForDIUsingChannel(*It); }
+					for (TObjectIterator<UNiagaraDataInterfaceDataChannelRead> It; It; ++It) { AddSystemsForDIUsingChannel(*It); }
+					for (TObjectIterator<UNiagaraDataInterfaceDataChannelWrite> It; It; ++It) { AddSystemsForDIUsingChannel(*It); }
+				}
+			}
+		});
+
+	SysUpdateContext.CommitUpdate();
+}
+
+#endif
+
+#if WITH_EDITORONLY_DATA
+
+FDelegateHandle UNiagaraDataChannelDefinitions::AssetRegistryOnLoadCompleteHandle;
+
+void UNiagaraDataChannelDefinitions::OnAssetCreated(UNiagaraDataChannelDefinitions* CreatedDef)
+{
+	check(Definitions.Contains(CreatedDef) == false);
+	CreatedDef->AddToRoot();
+	Definitions.Add(CreatedDef);
+	FNiagaraWorldManager::ForAllWorldManagers(
+		[&](FNiagaraWorldManager& WorldMan)
+		{
+			WorldMan.WaitForAsyncWork();
+			for (UNiagaraDataChannel* DataChannel : CreatedDef->DataChannels)
+			{
+				WorldMan.GetDataChannelManager().InitDataChannel(DataChannel, true);
+			}
+		});
+}
+
+void UNiagaraDataChannelDefinitions::OnAssetDeleted(UNiagaraDataChannelDefinitions* DeletedDef)
+{
+	check(Definitions.Contains(DeletedDef));
+	DeletedDef->RemoveFromRoot();
+	FNiagaraWorldManager::ForAllWorldManagers(
+		[&](FNiagaraWorldManager& WorldMan)
+		{
+			WorldMan.WaitForAsyncWork();
+			for (UNiagaraDataChannel* DataChannel : DeletedDef->DataChannels)
+			{
+				WorldMan.GetDataChannelManager().RemoveDataChannel(DataChannel->GetChannelName());
+			}
+		});
+}
+
+void UNiagaraDataChannelDefinitions::OnAssetRegistryLoadComplete()
+{
+	check(IsInGameThread());
+
+	bAssetRegistryScanComplete = true;
+
+	FAssetRegistryModule& AssetRegistryModule = FModuleManager::LoadModuleChecked<FAssetRegistryModule>("AssetRegistry");
+	IAssetRegistry& AssetRegistry = AssetRegistryModule.Get();
+
+	AssetRegistry.OnFilesLoaded().Remove(AssetRegistryOnLoadCompleteHandle);
+
+	check(AssetRegistry.IsLoadingAssets() == false);
+
+	TArray<FAssetData> AllDataChannelDefinitions;
+	AssetRegistry.GetAssetsByClass(UNiagaraDataChannelDefinitions::StaticClass()->GetClassPathName(), AllDataChannelDefinitions);
+
+	for (FAssetData& DataChannelDef : AllDataChannelDefinitions)
+	{
+		UNiagaraDataChannelDefinitions* NewDefs = Cast<UNiagaraDataChannelDefinitions>(DataChannelDef.GetAsset());
+		NewDefs->AddToRoot();
+		Definitions.AddUnique(NewDefs);
+
+		FNiagaraWorldManager::ForAllWorldManagers(
+			[&](FNiagaraWorldManager& WorldMan)
+			{
+				for (UNiagaraDataChannel* DataChannel : NewDefs->DataChannels)
+				{
+					WorldMan.GetDataChannelManager().InitDataChannel(DataChannel, false);
+				}
+			});
+	}
+}
+
+#endif
+
+const TArray<UNiagaraDataChannelDefinitions*>& UNiagaraDataChannelDefinitions::GetDataChannelDefinitions(bool bRequired, bool bInformUser)
+{
+	check(IsInGameThread());
+
+	//In the editor we'll discovery all definitions from the asset registry for use in the UI.
+	//At runtime, definitions are added as they are loaded by refs in UNiagaraSystems.
+#if WITH_EDITORONLY_DATA
+	if (bAssetRegistryScanBegun == false)
+	{
+		bAssetRegistryScanBegun = true;
+		FAssetRegistryModule& AssetRegistryModule = FModuleManager::LoadModuleChecked<FAssetRegistryModule>("AssetRegistry");
+		IAssetRegistry& AssetRegistry = AssetRegistryModule.Get();
+		AssetRegistryOnLoadCompleteHandle = AssetRegistry.OnFilesLoaded().AddStatic(&UNiagaraDataChannelDefinitions::OnAssetRegistryLoadComplete);
+		AssetRegistry.SearchAllAssets(true);
+		if(bRequired)
+		{
+			AssetRegistry.WaitForCompletion();
+		}
+	}
+
+	if(bAssetRegistryScanComplete == false && bInformUser)
+	{
+		// Open a dialog asking the user to wait while assets are being discovered.
+		// Todo. A better approach would be to customize all locations waiting for this scan to complete and inform the user inline. 
+		// Then have the niagara UI refresh on scan completion.
+		FMessageDialog::Open(EAppMsgType::Ok, LOCTEXT("AssetsStillScanningError", "Asset Registry is still scanning files. Some lists such as those for Data Channels may be incommplete until the Asset Registry scan is complete."));
+	}
+#endif
+
+	return Definitions;
+}
+
+const UNiagaraDataChannel* UNiagaraDataChannelDefinitions::FindDataChannel(FName ChannelName)
+{
+	for (const UNiagaraDataChannelDefinitions* DataChannelDef : GetDataChannelDefinitions(true, false))
+	{
+		for (UNiagaraDataChannel* DataChannel : DataChannelDef->DataChannels)
+		{
+			if (DataChannel->GetChannelName() == ChannelName)
+			{
+				return DataChannel;
+			}
+		}
+	}
+	return nullptr;
+}
+
+#undef LOCTEXT_NAMESPACE
\ No newline at end of file
diff --git a/Engine/Plugins/FX/Niagara/Source/Niagara/Private/NiagaraDataChannelHandler.cpp b/Engine/Plugins/FX/Niagara/Source/Niagara/Private/NiagaraDataChannelHandler.cpp
new file mode 100644
--- /dev/null
+++ b/Engine/Plugins/FX/Niagara/Source/Niagara/Private/NiagaraDataChannelHandler.cpp
@@ -0,0 +1,76 @@
+// Copyright Epic Games, Inc. All Rights Reserved.
+
+#include "NiagaraDataChannelHandler.h"
+#include "NiagaraDataChannelAccessor.h"
+
+UNiagaraDataChannelHandler::~UNiagaraDataChannelHandler()
+{
+	if (FNiagaraDataChannelHandlerRTProxyBase* ReleasedProxy = RTProxy.Release())
+	{
+		ENQUEUE_RENDER_COMMAND(FDeleteRTProxyCmd) (
+			[RT_Proxy = ReleasedProxy](FRHICommandListImmediate& CmdList) mutable
+			{
+				delete RT_Proxy;
+			}
+		);
+		check(RTProxy.IsValid() == false);
+	}
+}
+
+void UNiagaraDataChannelHandler::Init(const UNiagaraDataChannel* InChannel)
+{
+	DataChannel = InChannel;
+}
+
+void UNiagaraDataChannelHandler::Tick(float DeltaTime, ETickingGroup TickGroup, FNiagaraWorldManager* OwningWorld)
+{
+}
+
+void UNiagaraDataChannelHandler::Publish(const FNiagaraDataChannelPublishRequest& Request)
+{
+	PublishRequests.Add(Request);
+}
+
+void UNiagaraDataChannelHandler::RemovePublishRequests(const FNiagaraDataSet* DataSet)
+{
+	for (auto It = PublishRequests.CreateIterator(); It; ++It)
+	{
+		const FNiagaraDataChannelPublishRequest& PublishRequest = *It;
+
+		//First remove any buffers we've queued up to push to the GPU
+		for (auto BufferIt = BuffersForGPU.CreateIterator(); BufferIt; ++BufferIt)
+		{
+			FNiagaraDataBuffer* Buffer = *BufferIt;
+			if (Buffer == nullptr || Buffer->GetOwner() == DataSet)
+			{
+				It.RemoveCurrentSwap();
+			}
+		}
+
+		if (PublishRequest.Data && PublishRequest.Data->GetOwner() == DataSet)
+		{
+			It.RemoveCurrentSwap();
+		}
+	}
+}
+
+
+UNiagaraDataChannelWriter* UNiagaraDataChannelHandler::GetDataChannelWriter()
+{
+	if(Writer == nullptr)
+	{
+		Writer =  NewObject<UNiagaraDataChannelWriter>();
+		Writer->Owner = this;
+	}
+	return Writer;
+}
+
+UNiagaraDataChannelReader* UNiagaraDataChannelHandler::GetDataChannelReader()
+{
+	if (Reader == nullptr)
+	{
+		Reader = NewObject<UNiagaraDataChannelReader>();
+		Reader->Owner = this;
+	}
+	return Reader;
+}
\ No newline at end of file
diff --git a/Engine/Plugins/FX/Niagara/Source/Niagara/Private/NiagaraDataChannelManager.cpp b/Engine/Plugins/FX/Niagara/Source/Niagara/Private/NiagaraDataChannelManager.cpp
new file mode 100644
--- /dev/null
+++ b/Engine/Plugins/FX/Niagara/Source/Niagara/Private/NiagaraDataChannelManager.cpp
@@ -0,0 +1,84 @@
+// Copyright Epic Games, Inc. All Rights Reserved.
+
+#include "NiagaraDataChannelManager.h"
+
+#include "NiagaraWorldManager.h"
+#include "NiagaraDataChannel.h"
+#include "NiagaraDataChannelHandler.h"
+#include "NiagaraDataChannelDefinitions.h"
+
+
+FNiagaraDataChannelManager::FNiagaraDataChannelManager(FNiagaraWorldManager* InWorldMan)
+	: WorldMan(InWorldMan)
+{
+}
+
+void FNiagaraDataChannelManager::AddReferencedObjects(FReferenceCollector& Collector)
+{
+	Collector.AddReferencedObjects(Channels);
+}
+
+void FNiagaraDataChannelManager::Init()
+{
+	for (const UNiagaraDataChannelDefinitions* DataChannelDefs : UNiagaraDataChannelDefinitions::GetDataChannelDefinitions(false, false))
+	{
+		for (const UNiagaraDataChannel* DataChannel : DataChannelDefs->DataChannels)
+		{
+			InitDataChannel(DataChannel, true);
+		}
+	}
+}
+
+void FNiagaraDataChannelManager::Cleanup()
+{
+	Channels.Empty();
+}
+
+void FNiagaraDataChannelManager::Tick(float DeltaSeconds, ETickingGroup TickGroup)
+{
+	if(INiagaraModule::DataChannelsEnabled())
+	{
+		//Tick all DataChannel channel handlers.
+		for (auto& ChannelPair : Channels)
+		{
+			ChannelPair.Value->Tick(DeltaSeconds, TickGroup, WorldMan);
+		}
+	}
+	else
+	{
+		Channels.Empty();
+	}
+}
+
+UNiagaraDataChannelHandler* FNiagaraDataChannelManager::FindDataChannelHandler(FName ChannelName)
+{
+	if (TObjectPtr<UNiagaraDataChannelHandler>* Found = Channels.Find(ChannelName))
+	{
+		return (*Found).Get();
+	}
+	return nullptr;
+}
+
+void FNiagaraDataChannelManager::InitDataChannel(const UNiagaraDataChannel* InChannel, bool bForce)
+{
+	if(INiagaraModule::DataChannelsEnabled() && InChannel->IsValid())
+	{
+		FName Channel = InChannel->GetChannelName();
+		TObjectPtr<UNiagaraDataChannelHandler>& Handler = Channels.FindOrAdd(Channel);
+
+		if (bForce || Handler == nullptr)
+		{
+			Handler = InChannel->CreateHandler();
+		}
+	}
+}
+
+void FNiagaraDataChannelManager::RemoveDataChannel(FName ChannelName)
+{
+	Channels.Remove(ChannelName);
+}
+
+UWorld* FNiagaraDataChannelManager::GetWorld()const
+{
+	return WorldMan->GetWorld();
+}
diff --git a/Engine/Plugins/FX/Niagara/Source/Niagara/Private/NiagaraDataChannelManager.h b/Engine/Plugins/FX/Niagara/Source/Niagara/Private/NiagaraDataChannelManager.h
new file mode 100644
--- /dev/null
+++ b/Engine/Plugins/FX/Niagara/Source/Niagara/Private/NiagaraDataChannelManager.h
@@ -0,0 +1,39 @@
+// Copyright Epic Games, Inc. All Rights Reserved.
+
+#pragma once
+
+#include "NiagaraDataChannelCommon.h"
+
+
+/**
+Manager class for all Niagara DataChannels at the System/World level.
+*/
+class FNiagaraDataChannelManager
+{
+public:
+
+	explicit FNiagaraDataChannelManager(FNiagaraWorldManager* InWorldMan);
+
+	void AddReferencedObjects(FReferenceCollector& Collector);
+
+	void Init();
+	void Cleanup();
+
+	void InitDataChannel(const UNiagaraDataChannel* Channel, bool bForce);
+	void RemoveDataChannel(FName ChannelName);
+
+	void Tick(float DeltaSeconds, ETickingGroup TickGroup);
+	
+	/**
+	Return the DataChannel handler for the given channel.
+	*/
+	UNiagaraDataChannelHandler* FindDataChannelHandler(FName ChannelName);
+
+	UWorld* GetWorld()const;
+
+private:
+	FNiagaraWorldManager* WorldMan = nullptr;
+
+	/** Runtime handlers for each DataChannel channel. */
+	TMap<FName, TObjectPtr<UNiagaraDataChannelHandler>> Channels;
+};
\ No newline at end of file
diff --git a/Engine/Plugins/FX/Niagara/Source/Niagara/Private/NiagaraDataChannel_Global.cpp b/Engine/Plugins/FX/Niagara/Source/Niagara/Private/NiagaraDataChannel_Global.cpp
new file mode 100644
--- /dev/null
+++ b/Engine/Plugins/FX/Niagara/Source/Niagara/Private/NiagaraDataChannel_Global.cpp
@@ -0,0 +1,320 @@
+// Copyright Epic Games, Inc. All Rights Reserved.
+
+#include "NiagaraDataChannel_Global.h"
+
+#include "NiagaraSystemInstance.h"
+#include "NiagaraDataChannelManager.h"
+#include "DataInterface/NiagaraDataInterfaceDataChannelWrite.h"
+#include "NiagaraGPUComputeDispatch.h"
+#include "NiagaraWorldManager.h"
+
+
+int32 GbDebugDumpDumpHandlerTick = 0;
+static FAutoConsoleVariableRef CVarDumpHandlerTick(
+	TEXT("fx.Niagara.DataChannels.DumpHandlerTick"),
+	GbDebugDumpDumpHandlerTick,
+	TEXT(" \n"),
+	ECVF_Default
+);
+
+DECLARE_CYCLE_STAT(TEXT("UNiagaraDataChannelHandler_Global::Tick"), STAT_DataChannelHandler_Global_Tick, STATGROUP_NiagaraDataChannels);
+
+//////////////////////////////////////////////////////////////////////////
+
+UNiagaraDataChannelHandler* UNiagaraDataChannel_Global::CreateHandler()const
+{
+	UNiagaraDataChannelHandler* NewHandler = NewObject<UNiagaraDataChannelHandler_Global>();
+	NewHandler->Init(this);
+	return NewHandler;
+}
+
+//TODO: Other channel types.
+// Octree etc.
+
+
+//////////////////////////////////////////////////////////////////////////
+
+class FNiagaraDataChannelHandlerRTProxy_Global : public FNiagaraDataChannelHandlerRTProxyBase
+{	
+public:
+	FNiagaraDataSet* GPUDataSet = nullptr;
+	FString DebugName;
+
+	void BeginFrame(FNiagaraGpuComputeDispatchInterface* DispathInterface, FRHICommandListImmediate& CmdList)
+	{
+	}
+
+	void EndFrame(FNiagaraGpuComputeDispatchInterface* DispathInterface, FRHICommandListImmediate& CmdList, const TArray<FNiagaraDataBuffer*>& BuffersForGPU)
+	{
+		uint32 NumInstance = 0;
+		for(FNiagaraDataBuffer* Buffer : BuffersForGPU)
+		{
+			NumInstance += Buffer->GetNumInstances();
+		}		
+
+		GPUDataSet->BeginSimulate();
+		FNiagaraDataBuffer* DestBuffer = GPUDataSet->GetDestinationData();
+		DestBuffer->PushCPUBuffersToGPU(BuffersForGPU, true, CmdList, DispathInterface->GetFeatureLevel(), *DebugName);
+		GPUDataSet->EndSimulate();
+		
+		//For now we need not deal with the instance count manager but when we do GPU->GPU writes we will
+		//DestBuffer->SetGPUInstanceCountBufferOffset()
+	}
+};
+
+UNiagaraDataChannelHandler_Global::UNiagaraDataChannelHandler_Global(FObjectInitializer const& ObjectInitializer)
+	: Super(ObjectInitializer)
+{
+	RTProxy.Reset(new FNiagaraDataChannelHandlerRTProxy_Global());
+}
+
+void UNiagaraDataChannelHandler_Global::BeginDestroy()
+{
+	Super::BeginDestroy();
+	Data.Empty();
+	
+	ENQUEUE_RENDER_COMMAND(FDeleteContextCommand)(
+		[RT_GameDataStaging = GameDataStaging](FRHICommandListImmediate& RHICmdList)
+		{
+			if (RT_GameDataStaging != nullptr)
+			{
+				delete RT_GameDataStaging;
+			}
+		}
+	);
+	GameDataStaging = nullptr;
+}
+
+void UNiagaraDataChannelHandler_Global::Init(const UNiagaraDataChannel* InChannel)
+{
+	Super::Init(InChannel);
+	Data.Init(this);
+
+	const FNiagaraDataSetCompiledData& CompiledData = GetDataChannel()->GetCompiledData(ENiagaraSimTarget::CPUSim);
+	GameDataStaging = new FNiagaraDataSet();
+	GameDataStaging->Init(&CompiledData);
+
+	FString GPUDebugName;
+#if !UE_BUILD_SHIPPING
+	GPUDebugName = FString::Printf(TEXT("%s__GPUData"), *InChannel->GetChannelName().ToString());
+#endif
+
+	ENQUEUE_RENDER_COMMAND(FInitProxy) (
+		[RT_Proxy = GetRTProxyAs<FNiagaraDataChannelHandlerRTProxy_Global>(), GPUDataSet = Data.GPUSimData, DebugName=MoveTemp(GPUDebugName)](FRHICommandListImmediate& CmdList)mutable
+	{
+		RT_Proxy->GPUDataSet = GPUDataSet;
+		RT_Proxy->DebugName = MoveTemp(DebugName);
+	}
+	);
+}
+
+void UNiagaraDataChannelHandler_Global::Tick(float DeltaSeconds, ETickingGroup TickGroup, FNiagaraWorldManager* OwningWorld)
+{
+	SCOPE_CYCLE_COUNTER(STAT_DataChannelHandler_Global_Tick);
+	Super::Tick(DeltaSeconds, TickGroup, OwningWorld);
+
+	if (Data.IsInitialized() == false)
+	{
+		Data.Init(this);
+	}
+
+	if(GameDataStaging == nullptr)
+	{
+		const FNiagaraDataSetCompiledData CompiledData = GetDataChannel()->GetCompiledData(ENiagaraSimTarget::CPUSim);
+		GameDataStaging = new FNiagaraDataSet();
+		GameDataStaging->Init(&CompiledData);
+	}
+
+	bool bIsFirstTick = TickGroup == NiagaraFirstTickGroup;
+	bool bIsLastTick = TickGroup == NiagaraLastTickGroup;
+
+	if (bIsFirstTick)
+	{
+		//TODO: Add reference struct to handle these ref counts and make add/release private.
+		if(DataChannel->KeepPreviousFrameData())
+		{
+			if(Data.PrevCPUSimData)
+			{
+				Data.PrevCPUSimData->ReleaseReadRef();
+			}
+			Data.PrevCPUSimData = Data.CPUSimData->GetCurrentData();
+			Data.PrevCPUSimData->AddReadRef();
+		}
+
+		FNiagaraGpuComputeDispatchInterface* DispathInterface = FNiagaraGpuComputeDispatchInterface::Get(OwningWorld->GetWorld());
+		ENQUEUE_RENDER_COMMAND(FDataChannelProxyBeginFrame) (
+			[DispathInterface, RT_Proxy = GetRTProxyAs<FNiagaraDataChannelHandlerRTProxy_Global>()](FRHICommandListImmediate& CmdList)
+		{
+			if(RT_Proxy)
+			{
+				RT_Proxy->BeginFrame(DispathInterface, CmdList);
+			}
+		});
+	}
+
+	//Gather various sources of DataChannel and ship them around as heeded.
+
+	int32 GameDataOrigSize = bIsFirstTick ? 0 : Data.GameData->Num();
+	int32 CPUDataOrigSize = bIsFirstTick ? 0 : Data.CPUSimData->GetCurrentData()->GetNumInstances();
+	int32 NewGameDataChannel = GameDataOrigSize;
+	int32 NewCPUDataChannel = CPUDataOrigSize;
+
+	//Do a pass to gather the new total size for our DataChannel data.
+
+	//Each DI that generates DataChannel can control whether it's pushed to Game/CPU/GPU.
+	BuffersForGPU.Reserve(BuffersForGPU.Num() + PublishRequests.Num());
+	for (FNiagaraDataChannelPublishRequest& PublishRequest : PublishRequests)
+	{
+		uint32 NumInsts = 0;
+
+		if(FNiagaraDataChannelGameData* RequestGameData = PublishRequest.GameData.Get())
+		{
+			NumInsts = RequestGameData->Num();
+			if (RequestGameData->Num() > 0 && (PublishRequest.bVisibleToCPUSims || PublishRequest.bVisibleToGPUSims))
+			{
+				//For now we copy game data into a staging buffer for pushing to CPU & GPU buffers.
+				//We can avoid this in future but will add complexity to the transfer code.
+				GameDataStaging->BeginSimulate();
+				GameDataStaging->Allocate(RequestGameData->Num());
+				RequestGameData->WriteToDataSet(GameDataStaging->GetDestinationData(), 0);
+				GameDataStaging->EndSimulate();
+
+				if (GbDebugDumpDumpHandlerTick)
+				{
+					FNiagaraDataBuffer& Buffer = GameDataStaging->GetCurrentDataChecked();
+					FString DebugString = FString::Printf(TEXT("%s Game Data Staging"), *DataChannel->GetChannelName().ToString());
+					Buffer.Dump(0, Buffer.GetNumInstances(), FString::Printf(TEXT("=== Data Channle Staging Game Data: %d Elements --> %s ==="), Buffer.GetNumInstances(), *DebugString));
+				}
+
+				PublishRequest.Data = GameDataStaging->GetCurrentData();
+				PublishRequest.Data->AddReadRef();
+			}
+		}
+		else if( ensure(PublishRequest.Data) )
+		{
+			NumInsts = PublishRequest.Data->GetNumInstances();
+		}
+	
+		if (PublishRequest.bVisibleToGame)
+		{
+			NewGameDataChannel += NumInsts;
+		}
+		if (PublishRequest.bVisibleToCPUSims)
+		{
+			NewCPUDataChannel += NumInsts;
+		}
+		if(PublishRequest.bVisibleToGPUSims)
+		{
+			PublishRequest.Data->AddReadRef();
+			BuffersForGPU.Add(PublishRequest.Data);
+		}
+	}
+
+	//Allocate Sim buffers ready for gather.
+	Data.CPUSimData->BeginSimulate();
+	Data.CPUSimData->Allocate(NewCPUDataChannel, bIsFirstTick == false);
+
+	//Now do the actual data collection.
+	
+	//Allocate game buffers ready for gather.
+	if(bIsFirstTick)
+	{
+		Data.GameData->Reset();
+	}
+	
+	Data.GameData->SetNum(NewGameDataChannel);
+
+	for (FNiagaraDataChannelPublishRequest& PublishRequest : PublishRequests)
+	{
+		if (FNiagaraDataChannelGameData* RequestGameData = PublishRequest.GameData.Get())
+		{
+			if(RequestGameData->Num() && PublishRequest.bVisibleToGame)
+			{
+				Data.GameData->AppendFromGameData(*RequestGameData);
+			}
+		}
+
+		if (PublishRequest.Data)
+		{
+			const ENiagaraSimTarget SimTarget = PublishRequest.Data->GetOwner()->GetSimTarget();
+			if (SimTarget == ENiagaraSimTarget::CPUSim)
+			{
+				if (PublishRequest.bVisibleToGame)
+				{
+					Data.GameData->AppendFromDataSet(PublishRequest.Data, PublishRequest.LwcTile);
+				}
+				if (PublishRequest.bVisibleToCPUSims)
+				{
+					PublishRequest.Data->CopyToUnrelated(Data.CPUSimData->GetDestinationDataChecked(), 0, Data.CPUSimData->GetDestinationDataChecked().GetNumInstances(), PublishRequest.Data->GetNumInstances());
+				}
+			}
+			else if (SimTarget == ENiagaraSimTarget::GPUComputeSim)
+			{
+				//TODO: GPU->GPU handling will be done all RT side. GPU->CPU may be done here in future but we'll have to pull the data from the GPU on the RT and pass beck to GT here.
+			}
+			else
+			{
+				check(0);
+			}
+		}
+	}
+
+	Data.CPUSimData->EndSimulate();
+
+	if (GbDebugDumpDumpHandlerTick)
+	{
+		FNiagaraDataBuffer& Buffer = Data.CPUSimData->GetCurrentDataChecked();
+		FString DebugString = FString::Printf(TEXT("%s CPU Sim Data"), *DataChannel->GetChannelName().ToString());
+		Buffer.Dump(0, Buffer.GetNumInstances(), FString::Printf(TEXT("=== Updated CPU Sim Data: %d Elements --> %s ==="), CPUDataOrigSize, *DebugString));
+	}
+
+	PublishRequests.Reset();
+
+	//Now we push the CPU sim data to the RT for upload to the GPU.
+	//We may need to do a final tick from PostActorTick in the world manager?
+	if(bIsLastTick)
+	{
+		FNiagaraDataChannelHandlerRTProxy_Global* Proxy = GetRTProxyAs<FNiagaraDataChannelHandlerRTProxy_Global>();
+
+		check(Proxy);
+		check(Proxy->GPUDataSet);
+		check(Proxy->GPUDataSet->GetSimTarget() == ENiagaraSimTarget::GPUComputeSim);
+		FNiagaraGpuComputeDispatchInterface* DispathInterface = FNiagaraGpuComputeDispatchInterface::Get(OwningWorld->GetWorld());
+		ENQUEUE_RENDER_COMMAND(FDataChannelProxyEndFrame) (
+			[DispathInterface, RT_Proxy = Proxy, RT_BuffersForGPU=MoveTemp(BuffersForGPU)](FRHICommandListImmediate& CmdList) 
+			{
+				RT_Proxy->EndFrame(DispathInterface, CmdList, RT_BuffersForGPU);				
+			});
+		BuffersForGPU.Reset();
+	}
+}
+
+void UNiagaraDataChannelHandler_Global::GetData(FNiagaraSystemInstance* SystemInstance, FNiagaraDataBuffer*& OutCPUData, bool bGetPrevFrameData)
+{
+	if(bGetPrevFrameData && Data.PrevCPUSimData)
+	{
+		OutCPUData = Data.PrevCPUSimData;
+	}
+	else
+	{
+		OutCPUData = Data.CPUSimData->GetCurrentData();	
+	}
+
+	if(OutCPUData)
+	{
+		OutCPUData->AddReadRef();
+	}
+	//For more complicated channels we could check the location + bounds of the system instance etc to return some spatially localized data.
+}
+
+void UNiagaraDataChannelHandler_Global::GetDataGPU(FNiagaraSystemInstance* SystemInstance, FNiagaraDataSet*& OutCPUData)
+{
+	OutCPUData = Data.GPUSimData;	
+	//For more complicated channels we could check the location + bounds of the system instance etc to return some spatially localized data.
+}
+
+
+FNiagaraDataChannelGameDataPtr UNiagaraDataChannelHandler_Global::GetGameData(UActorComponent* OwningComponent)
+{
+	return Data.GameData;
+}
diff --git a/Engine/Plugins/FX/Niagara/Source/Niagara/Private/NiagaraDataChannel_Global.h b/Engine/Plugins/FX/Niagara/Source/Niagara/Private/NiagaraDataChannel_Global.h
new file mode 100644
--- /dev/null
+++ b/Engine/Plugins/FX/Niagara/Source/Niagara/Private/NiagaraDataChannel_Global.h
@@ -0,0 +1,63 @@
+// Copyright Epic Games, Inc. All Rights Reserved.
+
+#pragma once
+
+#include "NiagaraDataChannel.h"
+#include "NiagaraDataChannelHandler.h"
+#include "NiagaraDataChannel_Global.generated.h"
+
+
+UCLASS(Experimental)
+class UNiagaraDataChannel_Global : public UNiagaraDataChannel
+{
+	GENERATED_BODY()
+
+	virtual UNiagaraDataChannelHandler* CreateHandler()const override;
+};
+
+/**
+Basic DataChannel handler that makes all data visible globally.
+*/
+UCLASS(Experimental, BlueprintType)
+class UNiagaraDataChannelHandler_Global : public UNiagaraDataChannelHandler
+{
+	GENERATED_UCLASS_BODY()
+
+	FNiagaraWorldDataChannelStore Data;
+
+	//UObject Interface
+	virtual void BeginDestroy()override;
+	//UObject Interface End
+
+	virtual void Init(const UNiagaraDataChannel* InChannel) override;
+	virtual void Tick(float DeltaTime, ETickingGroup TickGroup, FNiagaraWorldManager* OwningWorld) override;
+	virtual void GetData(FNiagaraSystemInstance* SystemInstance, FNiagaraDataBuffer*& OutCPUData, bool bGetLastFrameData) override;
+	virtual void GetDataGPU(FNiagaraSystemInstance* SystemInstance, FNiagaraDataSet*& OutCPUData) override;
+
+	/** Returns the game level DataChannel data store. */
+	virtual FNiagaraDataChannelGameDataPtr GetGameData(UActorComponent* OwningComponent);
+
+private:
+
+	/** Staging buffer for Game level data so we can push it around as a publish request like other niagara data. Simplifying the code. */
+	FNiagaraDataSet* GameDataStaging;
+
+	/** Last Frame's CPU data. */
+	FNiagaraDataBuffer* PrevFrameCPU;
+};
+
+/**
+DataChannel handler that separates DataChannel out by the hash of their position into spatially localized grid cells.
+Useful for DataChannel that should only be visible to nearby consumers.
+TODO: Likely unsuitable for LWC games as we need a defined min/max for the hash and storage. A sparse octree etc would be needed
+*/
+// UCLASS(EditInlineNew)
+// class UNiagaraDataChannelChannelHandler_WorldGrid : public UObject
+// {
+// };
+
+// UCLASS(EditInlineNew, abstract)
+// class UNiagaraDataChannelChannelHandler_Octree : UObject
+// {
+// 
+// };
\ No newline at end of file
diff --git a/Engine/Plugins/FX/Niagara/Source/Niagara/Private/NiagaraEmitterInstance.cpp b/Engine/Plugins/FX/Niagara/Source/Niagara/Private/NiagaraEmitterInstance.cpp
--- a/Engine/Plugins/FX/Niagara/Source/Niagara/Private/NiagaraEmitterInstance.cpp
+++ b/Engine/Plugins/FX/Niagara/Source/Niagara/Private/NiagaraEmitterInstance.cpp
@@ -1015,6 +1015,10 @@ void FNiagaraEmitterInstance::PostTick()
 {
 	SCOPE_CYCLE_COUNTER(STAT_NiagaraEmitterPostTick);
 
+	//We reset the spawn infos to it's minimum size for script spawning.
+	//It can exceed this size as external DIs etc can inject additional spawn infos.
+	SpawnInfos.SetNum(CachedEmitterCompiledData->SpawnAttributes.Num());
+
 	if (EventInstanceData.IsValid())
 	{
 		//Clear refs to event data buffers.
