Niagara Data Channels v0.1

diff --git a/Engine/Plugins/FX/Niagara/Source/Niagara/Private/NiagaraModule.cpp b/Engine/Plugins/FX/Niagara/Source/Niagara/Private/NiagaraModule.cpp
--- a/Engine/Plugins/FX/Niagara/Source/Niagara/Private/NiagaraModule.cpp
+++ b/Engine/Plugins/FX/Niagara/Source/Niagara/Private/NiagaraModule.cpp
@@ -33,6 +33,11 @@
 #include "Engine/StaticMesh.h"
 #include "UObject/UObjectGlobals.h"
 
+#include "DataInterface/NiagaraDataInterfaceDataChannelCommon.h"
+#include "DataInterface/NiagaraDataInterfaceDataChannelWrite.h"
+#include "DataInterface/NiagaraDataInterfaceDataChannelRead.h"
+#include "DataInterface/NiagaraDataInterfaceDataChannelSpawn.h"
+
 #if PLATFORM_WINDOWS
 #include "NiagaraOpenVDB.h"
 #endif
@@ -95,6 +100,44 @@ void INiagaraModule::OnUseGlobalFXBudgetChanged(IConsoleVariable* Variable)
 	}
 }
 
+bool INiagaraModule::bDataChannelsEnabled = false;
+
+static FAutoConsoleVariableRef CVarDataChannelEnabled(
+	TEXT("fx.Niagara.DataChannels.Enabled"),
+	INiagaraModule::bDataChannelsEnabled,
+	TEXT("If true, Niagara Data Channels will be enabled. \n"),
+	FConsoleVariableDelegate::CreateStatic(&INiagaraModule::OnDataChannelsEnabledChanged),
+	ECVF_Default
+);
+
+void INiagaraModule::OnDataChannelsEnabledChanged(IConsoleVariable* Variable)
+{
+	if(DataChannelsEnabled())
+	{
+		ENiagaraTypeRegistryFlags Flags = ENiagaraTypeRegistryFlags::AllowAnyVariable | ENiagaraTypeRegistryFlags::AllowParameter;
+		FNiagaraTypeRegistry::Register(UNiagaraDataInterfaceDataChannelWrite::StaticClass(), Flags);
+		FNiagaraTypeRegistry::Register(UNiagaraDataInterfaceDataChannelRead::StaticClass(), Flags);
+		FNiagaraTypeRegistry::Register(UNiagaraDataInterfaceDataChannelSpawn::StaticClass(), Flags);
+
+		FNiagaraWorldManager::ForAllWorldManagers(
+			[](FNiagaraWorldManager& WorldMan)
+			{
+				WorldMan.GetDataChannelManager().Init();
+			});
+	}
+// 	else TODO: for completeness we could add the ability to unregister types 
+// 	{
+// 		FNiagaraTypeRegistry::Unregister(UNiagaraDataInterfaceDataChannelWrite::StaticClass());
+// 		FNiagaraTypeRegistry::Register(UNiagaraDataInterfaceDataChannelRead::StaticClass());
+// 		FNiagaraTypeRegistry::Register(UNiagaraDataInterfaceDataChannelSpawn::StaticClass());
+// 	}
+
+	//Re-init everything for now.
+	//Can be more targeted if needed.
+	FNiagaraSystemUpdateContext Context;
+	Context.AddAll(true);
+}
+
 // these two globals are intended to help ensure that a global variable is accessible to natvis (as defined in Niagara.natvis)
 // while debugging.  GCoreTypeRegistrySingletonPtr is the pointer to the actual data stored within the TLazySingleton<FNiagaraTypeRegistry>
 // while GTypeRegistrySingletonPtr can be declared in each module to ensure that it can be accessed while debugging any Niagara
@@ -517,6 +560,8 @@ void INiagaraModule::ShutdownModule()
 
 	FNiagaraTypeRegistry::TearDown();
 
+	FNDIDataChannelLayoutManager::TearDown();
+
 #if WITH_EDITOR
 	FCoreUObjectDelegates::OnAssetLoaded.RemoveAll(this);
 #endif
@@ -1364,20 +1409,24 @@ bool FNiagaraTypeDefinition::Serialize(FArchive& Ar)
 		}
 #endif
 
-		if (UScriptStruct* StructDef = GetScriptStruct())
+		if(!IsEnum())
 		{
-			if (Flags & TF_SerializedAsLWC)
+			if (UScriptStruct* StructDef = GetScriptStruct())
 			{
-				ClassStructOrEnum = FNiagaraTypeHelper::GetSWCStruct(StructDef);
+				if (Flags & TF_SerializedAsLWC)
+				{
+					ClassStructOrEnum = FNiagaraTypeHelper::GetSWCStruct(StructDef);
 
-				Flags &= ~TF_SerializedAsLWC;
-			}
-#if WITH_EDITORONLY_DATA
-			else
-			{
-				ClassStructOrEnum = FNiagaraTypeHelper::FindNiagaraFriendlyTopLevelStruct(static_cast<UScriptStruct*>(ClassStructOrEnum), ENiagaraStructConversion::UserFacing);
+					Flags &= ~TF_SerializedAsLWC;
+				}
+	#if WITH_EDITORONLY_DATA
+				else if((Flags & TF_AllowLWC) == 0)
+				{	
+					//If the type doesn't allow LWC we should convert down to SWC.
+					ClassStructOrEnum = FNiagaraTypeHelper::FindNiagaraFriendlyTopLevelStruct(static_cast<UScriptStruct*>(ClassStructOrEnum), ENiagaraStructConversion::Simulation);
+				}
+	#endif
 			}
-#endif
 		}
 	}
 
diff --git a/Engine/Plugins/FX/Niagara/Source/Niagara/Private/NiagaraSystem.cpp b/Engine/Plugins/FX/Niagara/Source/Niagara/Private/NiagaraSystem.cpp
--- a/Engine/Plugins/FX/Niagara/Source/Niagara/Private/NiagaraSystem.cpp
+++ b/Engine/Plugins/FX/Niagara/Source/Niagara/Private/NiagaraSystem.cpp
@@ -39,6 +39,8 @@
 #include "UObject/UObjectThreadContext.h"
 #include "ShaderCompiler.h"
 #include "PipelineStateCache.h"
+#include "NiagaraDataChannel.h"
+#include "NiagaraDataChannelDefinitions.h"
 
 #include UE_INLINE_GENERATED_CPP_BY_NAME(NiagaraSystem)
 
@@ -1971,6 +1973,21 @@ void UNiagaraSystem::UpdatePostCompileDIInfo()
 			}
 		}
 	}
+	
+	ReferencedDataChannelDefinitions.Reset();
+
+	//We must include internal DIs in this search so they get correctly PostCompiled.
+	FNiagaraDataInterfaceUtilities::FDataInterfaceSearchOptions SearchOptions;
+	SearchOptions.bIncludeInternal = true;
+	FNiagaraDataInterfaceUtilities::ForEachDataInterface(this, [](const FNiagaraDataInterfaceUtilities::FDataInterfaceUsageContext& Context)
+	{
+		if (Context.DataInterface)
+		{
+			Context.DataInterface->PostCompile();
+		}
+		return true;
+	}
+	, SearchOptions);
 }
 
 void UNiagaraSystem::UpdateDITickFlags()
@@ -2027,6 +2044,24 @@ void UNiagaraSystem::UpdateHasGPUEmitters()
 	}
 }
 
+#if WITH_EDITORONLY_DATA
+
+void UNiagaraSystem::RegisterDataChannelUse(const UNiagaraDataChannel* DataChannel)
+{
+	if(UNiagaraDataChannelDefinitions* DefinitionsAsset = DataChannel->GetTypedOuter<UNiagaraDataChannelDefinitions>())
+	{
+		ReferencedDataChannelDefinitions.AddUnique(DefinitionsAsset);
+	}
+}
+
+void UNiagaraSystem::OnCompiledDataInterfaceChanged()
+{
+	UpdatePostCompileDIInfo();
+	UpdateDITickFlags();
+}
+
+#endif
+
 bool UNiagaraSystem::IsValidInternal() const
 {
 	if (!SystemSpawnScript || !SystemUpdateScript)
diff --git a/Engine/Plugins/FX/Niagara/Source/Niagara/Private/NiagaraTypes.cpp b/Engine/Plugins/FX/Niagara/Source/Niagara/Private/NiagaraTypes.cpp
--- a/Engine/Plugins/FX/Niagara/Source/Niagara/Private/NiagaraTypes.cpp
+++ b/Engine/Plugins/FX/Niagara/Source/Niagara/Private/NiagaraTypes.cpp
@@ -12,6 +12,7 @@
 
 #include UE_INLINE_GENERATED_CPP_BY_NAME(NiagaraTypes)
 
+static FName NAME_NiagaraDouble(TEXT("NiagaraDouble"));
 static FName NAME_NiagaraPosition(TEXT("NiagaraPosition"));
 
 void FNiagaraVariableBase::SetNamespacedName(const FString& InNamespace, FName InVariableName)
@@ -153,106 +154,112 @@ FNiagaraStructConversionStep::FNiagaraStructConversionStep()
 {
 }
 
-FNiagaraStructConversionStep::FNiagaraStructConversionStep(int32 InSourceBytes, int32 InSourceOffset, int32 InSimulationBytes, int32 InSimulationOffset, ENiagaraStructConversionType InConversionType)
-	: SourceBytes(InSourceBytes), SourceOffset(InSourceOffset), SimulationBytes(InSimulationBytes), SimulationOffset(InSimulationOffset), ConversionType(InConversionType)
+FNiagaraStructConversionStep::FNiagaraStructConversionStep(int32 InLWCBytes, int32 InLWCOffset, int32 InSimulationBytes, int32 InSimulationOffset, ENiagaraStructConversionType InConversionType)
+	: LWCBytes(InLWCBytes), LWCOffset(InLWCOffset), SimulationBytes(InSimulationBytes), SimulationOffset(InSimulationOffset), ConversionType(InConversionType)
 {
 }
 
-void FNiagaraStructConversionStep::CopyToSim(uint8* DestinationData, const uint8* SourceData) const
+void FNiagaraStructConversionStep::CopyToSim(const uint8* LWCData, uint8* SimulationData, int32 Count, int32 LWCStride, int32 SimulationStride)const
 {
-	void* Dest = DestinationData + SimulationOffset;
-	const void* Src = SourceData + SourceOffset;
 	if (ConversionType == ENiagaraStructConversionType::DoubleToFloat)
 	{
-		checkf(SourceBytes == 8 && SimulationBytes == 4, TEXT("Wrong bytesizes for double->float conversion: Source %i, Simulation %i"), SourceBytes, SimulationBytes);
-		*static_cast<float*>(Dest) = float(*static_cast<const double*>(Src));
+		checkf(LWCBytes == 8 && SimulationBytes == 4, TEXT("Wrong bytesizes for double->float conversion: Source %i, Simulation %i"), LWCBytes, SimulationBytes);
+		CopyInternal<float, double>(SimulationData, SimulationStride, SimulationOffset, LWCData, LWCStride, LWCOffset, Count);
 	}
 	else if (ConversionType == ENiagaraStructConversionType::Vector2)
 	{
-		checkf(SourceBytes == 16 && SimulationBytes == 8, TEXT("Wrong bytesizes for Vector2 d->f conversion: Source %i, Simulation %i"), SourceBytes, SimulationBytes);
-		FVector2f Vector2(*static_cast<const FVector2d*>(Src));
-		FMemory::Memcpy(Dest, &Vector2, SimulationBytes);
+		checkf(LWCBytes == 16 && SimulationBytes == 8, TEXT("Wrong bytesizes for Vector2 d->f conversion: Source %i, Simulation %i"), LWCBytes, SimulationBytes);
+		CopyInternal<FVector2f, FVector2d>(SimulationData, SimulationStride, SimulationOffset, LWCData, LWCStride, LWCOffset, Count);
 	}
 	else if (ConversionType == ENiagaraStructConversionType::Vector3)
 	{
-		checkf(SourceBytes == 24 && SimulationBytes == 12, TEXT("Wrong bytesizes for Vector3 d->f conversion: Source %i, Simulation %i"), SourceBytes, SimulationBytes);
-		FVector3f Vector(*static_cast<const FVector3d*>(Src));
-		FMemory::Memcpy(Dest, &Vector, SimulationBytes);
+		checkf(LWCBytes == 24 && SimulationBytes == 12, TEXT("Wrong bytesizes for Vector3 d->f conversion: Source %i, Simulation %i"), LWCBytes, SimulationBytes);
+		CopyInternal<FVector3f, FVector>(SimulationData, SimulationStride, SimulationOffset, LWCData, LWCStride, LWCOffset, Count);
 	}
 	else if (ConversionType == ENiagaraStructConversionType::Vector4)
 	{
-		checkf(SourceBytes == 32 && SimulationBytes == 16, TEXT("Wrong bytesizes for Vector4 d->f conversion: Source %i, Simulation %i"), SourceBytes, SimulationBytes);
-		FVector4f Vector4(*static_cast<const FVector4d*>(Src));
-		FMemory::Memcpy(Dest, &Vector4, SimulationBytes);
+		checkf(LWCBytes == 32 && SimulationBytes == 16, TEXT("Wrong bytesizes for Vector4 d->f conversion: Source %i, Simulation %i"), LWCBytes, SimulationBytes);
+		CopyInternal<FVector4f, FVector4>(SimulationData, SimulationStride, SimulationOffset, LWCData, LWCStride, LWCOffset, Count);
 	}
 	else if (ConversionType == ENiagaraStructConversionType::Quat)
 	{
-		checkf(SourceBytes == 32 && SimulationBytes == 16, TEXT("Wrong bytesizes for Quat d->f conversion: Source %i, Simulation %i"), SourceBytes, SimulationBytes);
-		FQuat4d SrcQuat = *static_cast<const FQuat4d*>(Src);
-		FQuat4f TargetQuat(SrcQuat);
-		FMemory::Memcpy(Dest, &TargetQuat, SimulationBytes);
+		checkf(LWCBytes == 32 && SimulationBytes == 16, TEXT("Wrong bytesizes for Quat d->f conversion: Source %i, Simulation %i"), LWCBytes, SimulationBytes);
+		CopyInternal<FQuat4f, FQuat4d>(SimulationData, SimulationStride, SimulationOffset, LWCData, LWCStride, LWCOffset, Count);
 	}
 	else
 	{
-		checkf(SourceBytes == SimulationBytes, TEXT("SourceBytes != TargetBytes: %i != %i"), SourceBytes, SimulationBytes);
-		FMemory::Memcpy(Dest, Src, SourceBytes);
+		checkf(LWCBytes == SimulationBytes, TEXT("LWCBytes != TargetBytes: %i != %i"), LWCBytes, SimulationBytes);
+		for (int32 i = 0; i < Count; ++i)
+		{
+			uint8* Dst = (SimulationData + i * SimulationStride) + SimulationOffset;
+			const uint8* Src = (LWCData + i * LWCStride) + LWCOffset;
+			FMemory::Memcpy(Dst, Src, LWCBytes);
+		}
 	}
 }
 
-void FNiagaraStructConversionStep::CopyFromSim(uint8* DestinationData, const uint8* SourceData) const
+void FNiagaraStructConversionStep::CopyFromSim(uint8* LWCData, const uint8* SimulationData, int32 Count, int32 LWCStride, int32 SimulationStride) const
 {
-	void* Dest = DestinationData + SourceOffset;
-	const void* Src = SourceData + SimulationOffset;
 	if (ConversionType == ENiagaraStructConversionType::DoubleToFloat)
 	{
-		checkf(SourceBytes == 8 && SimulationBytes == 4, TEXT("Wrong bytesizes for float->double conversion: Source %i, Simulation %i"), SourceBytes, SimulationBytes);
-		*static_cast<double*>(Dest) = *static_cast<const float*>(Src);
+		checkf(LWCBytes == 8 && SimulationBytes == 4, TEXT("Wrong bytesizes for float->double conversion: Source %i, Simulation %i"), LWCBytes, SimulationBytes);
+		CopyInternal<double, float>(LWCData, LWCStride, LWCOffset, SimulationData, SimulationStride, SimulationOffset, Count);
 	}
 	else if (ConversionType == ENiagaraStructConversionType::Vector2)
 	{
-		checkf(SourceBytes == 16 && SimulationBytes == 8, TEXT("Wrong bytesizes for Vector2 f->d conversion: Source %i, Simulation %i"), SourceBytes, SimulationBytes);
-		FVector2d Vector2(*static_cast<const FVector2f*>(Src));
-		FMemory::Memcpy(Dest, &Vector2, SourceBytes);
+		checkf(LWCBytes == 16 && SimulationBytes == 8, TEXT("Wrong bytesizes for Vector2 f->d conversion: Source %i, Simulation %i"), LWCBytes, SimulationBytes);
+		CopyInternal<FVector2D, FVector2f>(LWCData, LWCStride, LWCOffset, SimulationData, SimulationStride, SimulationOffset, Count);
 	}
 	else if (ConversionType == ENiagaraStructConversionType::Vector3)
 	{
-		checkf(SourceBytes == 24 && SimulationBytes == 12, TEXT("Wrong bytesizes for Vector3 f->d conversion: Source %i, Simulation %i"), SourceBytes, SimulationBytes);
-		FVector3d Vector3(*static_cast<const FVector3f*>(Src));
-        FMemory::Memcpy(Dest, &Vector3, SourceBytes);
+		checkf(LWCBytes == 24 && SimulationBytes == 12, TEXT("Wrong bytesizes for Vector3 f->d conversion: Source %i, Simulation %i"), LWCBytes, SimulationBytes);
+		CopyInternal<FVector, FVector3f>(LWCData, LWCStride, LWCOffset, SimulationData, SimulationStride, SimulationOffset, Count);
 	}
 	else if (ConversionType == ENiagaraStructConversionType::Vector4)
 	{
-		checkf(SourceBytes == 32 && SimulationBytes == 16, TEXT("Wrong bytesizes for Vector4 f->d conversion: Source %i, Simulation %i"), SourceBytes, SimulationBytes);
-		FVector4d Vector4(*static_cast<const FVector4f*>(Src));
-		FMemory::Memcpy(Dest, &Vector4, SourceBytes);
+		checkf(LWCBytes == 32 && SimulationBytes == 16, TEXT("Wrong bytesizes for Vector4 f->d conversion: Source %i, Simulation %i"), LWCBytes, SimulationBytes);
+		CopyInternal<FVector4d, FVector4f>(LWCData, LWCStride, LWCOffset, SimulationData, SimulationStride, SimulationOffset, Count);
 	}
 	else if (ConversionType == ENiagaraStructConversionType::Quat)
 	{
-		checkf(SourceBytes == 32 && SimulationBytes == 16, TEXT("Wrong bytesizes for Quat f->d conversion: Source %i, Simulation %i"), SourceBytes, SimulationBytes);
-		FQuat4f SrcQuat = *static_cast<const FQuat4f*>(Src);
-		FQuat4d TargetQuat(SrcQuat.X, SrcQuat.Y, SrcQuat.Z, SrcQuat.W);
-		FMemory::Memcpy(Dest, &TargetQuat, SourceBytes);
+		checkf(LWCBytes == 32 && SimulationBytes == 16, TEXT("Wrong bytesizes for Quat f->d conversion: Source %i, Simulation %i"), LWCBytes, SimulationBytes);
+		CopyInternal<FQuat4d, FQuat4f>(LWCData, LWCStride, LWCOffset, SimulationData, SimulationStride, SimulationOffset, Count);
 	}
 	else
 	{
-		checkf(SourceBytes == SimulationBytes, TEXT("SourceBytes != TargetBytes: %i != %i"), SourceBytes, SimulationBytes);
-		FMemory::Memcpy(DestinationData, SourceData, SourceBytes);
+		checkf(LWCBytes == SimulationBytes, TEXT("LWCBytes != TargetBytes: %i != %i"), LWCBytes, SimulationBytes);
+		for (int32 i = 0; i < Count; ++i)
+		{
+			uint8* Dst = (LWCData + i * LWCStride) + LWCOffset; 
+			const uint8* Src = (SimulationData + i * SimulationStride) + SimulationOffset;
+			FMemory::Memcpy(Dst, Src, LWCBytes);
+		}
 	}
 }
 
-void FNiagaraLwcStructConverter::ConvertDataToSimulation(uint8* DestinationData, const uint8* SourceData) const
+void FNiagaraLwcStructConverter::ConvertDataToSimulation(uint8* DestinationData, const uint8* SourceData, int32 Count) const
 {
+	if (IsValid() == false || Count == 0)
+	{
+		return;
+	}
+
 	for (const FNiagaraStructConversionStep& ConversionStep : ConversionSteps)
 	{
-		ConversionStep.CopyToSim(DestinationData, SourceData);
+		ConversionStep.CopyToSim(SourceData, DestinationData, Count, LWCSize, SWCSize);
 	}
 }
 
-void FNiagaraLwcStructConverter::ConvertDataFromSimulation(uint8* DestinationData, const uint8* SourceData) const
+void FNiagaraLwcStructConverter::ConvertDataFromSimulation(uint8* DestinationData, const uint8* SourceData, int32 Count) const
 {
+	if (IsValid() == false || Count == 0)
+	{
+		return;
+	}
+
 	for (const FNiagaraStructConversionStep& ConversionStep : ConversionSteps)
 	{
-		ConversionStep.CopyFromSim(DestinationData, SourceData);
+		ConversionStep.CopyFromSim(DestinationData, SourceData, Count, LWCSize, SWCSize);
 	}
 }
 
@@ -524,32 +531,66 @@ UScriptStruct* FNiagaraTypeHelper::GetSWCStruct(UScriptStruct* LWCStruct)
 			// If this is a LWC structure we need to built a new structure now
 			if (IsLWCStructure(LWCStruct))
 			{
-				FName SWCName = FName(LWCStruct->GetName() + TEXT("_SWC"));
-
-				//check if the remap table contains a previous entry for that struct. This might happen when the lwc struct was changed at runtime.
-				// In that case we want to reuse the existing swc struct object, since it might already be referenced by compile results.
-				for (auto Iter = RemapTable.CreateIterator(); Iter; ++Iter)
+				//Return common SWC variant of common LWC Types.
+				static UPackage* CoreUObjectPkg = FindObjectChecked<UPackage>(nullptr, TEXT("/Script/CoreUObject"));
+				static UScriptStruct* Vector2Struct = FindObjectChecked<UScriptStruct>(CoreUObjectPkg, TEXT("Vector2D"));
+				static UScriptStruct* Vector3Struct = FindObjectChecked<UScriptStruct>(CoreUObjectPkg, TEXT("Vector"));
+				static UScriptStruct* Vector4Struct = FindObjectChecked<UScriptStruct>(CoreUObjectPkg, TEXT("Vector4"));
+				static UScriptStruct* QuatStruct = FindObjectChecked<UScriptStruct>(CoreUObjectPkg, TEXT("Quat"));
+				if(LWCStruct == FNiagaraDouble::StaticStruct())
+				{
+					SWCStruct = FNiagaraFloat::StaticStruct();
+				}
+				else if(LWCStruct == Vector2Struct) 
+				{
+					static UScriptStruct* Vector2fStruct = FindObjectChecked<UScriptStruct>(CoreUObjectPkg, TEXT("Vector2f"));
+					SWCStruct = Vector2fStruct;
+				}
+				else if (LWCStruct == Vector3Struct)
 				{
-					TWeakObjectPtr<UScriptStruct> ExistingStruct = Iter.Value().Struct;
-					if (ExistingStruct.IsValid() && ExistingStruct->GetFName() == SWCName)
+					static UScriptStruct* Vector3fStruct = FindObjectChecked<UScriptStruct>(CoreUObjectPkg, TEXT("Vector3f"));
+					SWCStruct = Vector3fStruct;
+				}
+				else if (LWCStruct == Vector4Struct)
+				{
+					static UScriptStruct* Vector4fStruct = FindObjectChecked<UScriptStruct>(CoreUObjectPkg, TEXT("Vector4f"));
+					SWCStruct = Vector4fStruct;
+				}
+				else if (LWCStruct == QuatStruct)
+				{
+					static UScriptStruct* Quat4fStruct = FindObjectChecked<UScriptStruct>(CoreUObjectPkg, TEXT("Quat4f"));
+					SWCStruct = Quat4fStruct;
+				}
+				//More?
+				else
+				{
+					FName SWCName = FName(LWCStruct->GetName() + TEXT("_SWC"));
+
+					//check if the remap table contains a previous entry for that struct. This might happen when the lwc struct was changed at runtime.
+					// In that case we want to reuse the existing swc struct object, since it might already be referenced by compile results.
+					for (auto Iter = RemapTable.CreateIterator(); Iter; ++Iter)
 					{
-						SWCStruct = ExistingStruct.Get();
-						Iter.RemoveCurrent();
-						RemapTableDirty = true;
-						break;
+						TWeakObjectPtr<UScriptStruct> ExistingStruct = Iter.Value().Struct;
+						if (ExistingStruct.IsValid() && ExistingStruct->GetFName() == SWCName)
+						{
+							SWCStruct = ExistingStruct.Get();
+							Iter.RemoveCurrent();
+							RemapTableDirty = true;
+							break;
+						}
 					}
-				}
 
-				if (SWCStruct == nullptr)
-				{
-					SWCStruct = NewObject<UScriptStruct>(GetTransientPackage(), SWCName, RF_NoFlags);
-					SWCStruct->AddToRoot();
+					if (SWCStruct == nullptr)
+					{
+						SWCStruct = NewObject<UScriptStruct>(GetTransientPackage(), SWCName, RF_NoFlags);
+						SWCStruct->AddToRoot();
+					}
+					FNiagaraLwcStructConverter StructConverter = BuildSWCStructure(SWCStruct, LWCStruct);
+					FNiagaraTypeRegistry::RegisterStructConverter(FNiagaraTypeDefinition(LWCStruct, FNiagaraTypeDefinition::EAllowUnfriendlyStruct::Allow), StructConverter);
+					SWCStruct->Bind();
+					SWCStruct->PrepareCppStructOps();
+					SWCStruct->StaticLink(true);
 				}
-				FNiagaraLwcStructConverter StructConverter = BuildSWCStructure(SWCStruct, LWCStruct);
-				FNiagaraTypeRegistry::RegisterStructConverter(FNiagaraTypeDefinition(LWCStruct, FNiagaraTypeDefinition::EAllowUnfriendlyStruct::Allow), StructConverter);
-				SWCStruct->Bind();
-				SWCStruct->PrepareCppStructOps();
-				SWCStruct->StaticLink(true);
 			}
 			else
 			{
@@ -594,6 +635,11 @@ UScriptStruct* FNiagaraTypeHelper::FindNiagaraFriendlyTopLevelStruct(UScriptStru
 	}
 
 	// Note: UE core types are converted to the float variant as Niagara only works with float types.
+	if (InStruct->GetFName() == NAME_NiagaraDouble)
+	{
+		return FNiagaraTypeDefinition::GetFloatStruct();
+	}
+
 	if (InStruct->GetFName() == NAME_Vector2D || InStruct->GetFName() == NAME_Vector2d) // LWC support
 	{
 		return FNiagaraTypeDefinition::GetVec2Struct();
diff --git a/Engine/Plugins/FX/Niagara/Source/Niagara/Private/NiagaraWorldManager.cpp b/Engine/Plugins/FX/Niagara/Source/Niagara/Private/NiagaraWorldManager.cpp
--- a/Engine/Plugins/FX/Niagara/Source/Niagara/Private/NiagaraWorldManager.cpp
+++ b/Engine/Plugins/FX/Niagara/Source/Niagara/Private/NiagaraWorldManager.cpp
@@ -5,7 +5,7 @@
 #include "NiagaraModule.h"
 #include "Modules/ModuleManager.h"
 #include "NiagaraTypes.h"
-#include "NiagaraEvents.h"
+#include "NiagaraDataChannel.h"
 #include "NiagaraSettings.h"
 #include "NiagaraParameterCollection.h"
 #include "NiagaraSystemInstance.h"
@@ -302,6 +302,7 @@ FNiagaraWorldManager::FNiagaraWorldManager()
 	: World(nullptr)
 	, ActiveNiagaraTickGroup(-1)
 	, bAppHasFocus(true)
+	, DataChannelManager(this)
 {
 }
 
@@ -338,6 +339,8 @@ void FNiagaraWorldManager::Init(UWorld* InWorld)
 	//Possibly a later hook we can use.
 	//PrimePoolForAllSystems();
 
+	DataChannelManager.Init();
+
 #if WITH_NIAGARA_DEBUGGER
 	NiagaraDebugHud.Reset(new FNiagaraDebugHud(World));
 #endif
@@ -440,6 +443,8 @@ void FNiagaraWorldManager::AddReferencedObjects(FReferenceCollector& Collector)
 	}
 
 	Collector.AddReferencedObjects(CullProxyMap);
+
+	DataChannelManager.AddReferencedObjects(Collector);
 }
 
 FString FNiagaraWorldManager::GetReferencerName() const
@@ -604,11 +609,20 @@ void FNiagaraWorldManager::OnComputeDispatchInterfaceDestroyed_Internal(FNiagara
 	}
 }
 
+void FNiagaraWorldManager::OnWorldBeginTearDown()
+{
+	bIsTearingDown = true;
+	ComponentPool->Cleanup(World);
+	DataChannelManager.Cleanup();
+}
+
 void FNiagaraWorldManager::OnWorldCleanup(bool bSessionEnded, bool bCleanupResources)
 {
 	DeferredMethods.ExecuteAndClear();
 	ComponentPool->Cleanup(World);
 
+	DataChannelManager.Cleanup();
+
 	for (int TG = 0; TG < NiagaraNumTickGroups; ++TG)
 	{
 		for (TPair<UNiagaraSystem*, FNiagaraSystemSimulationRef>& SimPair : SystemSimulations[TG])
@@ -703,6 +717,17 @@ void FNiagaraWorldManager::RefreshOwnerAllowsScalability()
 	}
 }
 
+void FNiagaraWorldManager::WaitForAsyncWork()
+{
+	for (int TG = 0; TG < NiagaraNumTickGroups; ++TG)
+	{
+		for (TPair<UNiagaraSystem*, FNiagaraSystemSimulationRef>& SimPair : SystemSimulations[TG])
+		{
+			SimPair.Value->WaitForInstancesTickComplete();
+		}
+	}
+}
+
 void FNiagaraWorldManager::OnWorldInit(UWorld* World, const UWorld::InitializationValues IVS)
 {
 	FNiagaraWorldManager*& NewManager = WorldManagers.FindOrAdd(World);
@@ -748,7 +773,7 @@ void FNiagaraWorldManager::OnWorldBeginTearDown(UWorld* World)
 	FNiagaraWorldManager** Manager = WorldManagers.Find(World);
 	if (Manager)
 	{
-		(*Manager)->bIsTearingDown = true;
+		(*Manager)->OnWorldBeginTearDown();
 	}
 // 	FNiagaraWorldManager** Manager = WorldManagers.Find(World);
 // 	if (Manager)
@@ -1047,9 +1072,12 @@ void FNiagaraWorldManager::Tick(ETickingGroup TickGroup, float DeltaSeconds, ELe
 
 	DeltaSeconds *= DebugPlaybackRate;
 
+	//Tick DataChannel Manager.
+	DataChannelManager.Tick(DeltaSeconds, TickGroup);
+
 	// We do book keeping in the first tick group
 	if ( TickGroup == NiagaraFirstTickGroup )
-	{		
+	{
 		// Update playback mode
 		DebugPlaybackMode = RequestedDebugPlaybackMode;
 
diff --git a/Engine/Plugins/FX/Niagara/Source/Niagara/Public/NiagaraDataChannel.h b/Engine/Plugins/FX/Niagara/Source/Niagara/Public/NiagaraDataChannel.h
new file mode 100644
--- /dev/null
+++ b/Engine/Plugins/FX/Niagara/Source/Niagara/Public/NiagaraDataChannel.h
@@ -0,0 +1,267 @@
+// Copyright Epic Games, Inc. All Rights Reserved.
+
+/*==============================================================================
+NiagaraDataChannel.h: Code dealing with Niagara Data Channels and their management.
+
+Niagara Data Channels are a system for communication between Niagara Systems and with Game code/BP.
+
+Niagara Data Channels define a common payload and other settings for a particular named Data Channel.
+Niagara Data Channel Handlers are the runtime handler class that will provide access to the data channel to it's users and manage it's internal data.
+
+Niagara Systems can read from and write to Data Channels via data interfaces.
+Blueprint and game code can also read from and write to Data Channels.
+Each of these writes optionally being made visible to Game, CPU and/or GPU Systems.
+
+At the "Game" level, all data is held in LWC compatible types in AoS format.
+When making this data available to Niagara Systems it is converted to SWC, SoA layout that is compatible with Niagara simulation.
+
+EXPERIMENTAL:
+Data Channels are currently experimental and undergoing heavy development.
+Anything and everything can change, including content breaking changes.
+
+Some Current limitations:
+
+Tick Ordering:
+Niagara Systems can chose to read the current frame's data or the previous frame.
+Reading from the current frame allows zero latency but introduces a frame dependency, i.e. you must ensure that the reader ticks after the writer.
+This frame dependency needs work to be more robust and less error prone.
+Reading the previous frames data introduces a frame of latency but removes the need to tick later than the writer. Also means you're sure to get a complete frame worth of data.
+
+GPU Support:
+Currently GPU support is very limited.
+Only Game->GPU and CPUSim->GPU are supported.
+Only the Read function of the DI is supported.
+GPU simulations always use the current frame's CPU data as this is all pushed to the at the end of the frame.
+When GPU->GPU is supported, the frame dependency issue have more meaning for GPU systems.
+
+==============================================================================*/
+
+#pragma once
+
+#include "NiagaraDataChannelCommon.h"
+#include "Kismet/BlueprintFunctionLibrary.h"
+#include "NiagaraDataChannel.generated.h"
+
+DECLARE_STATS_GROUP(TEXT("Niagara Data Channels"), STATGROUP_NiagaraDataChannels, STATCAT_Advanced);
+
+class UNiagaraDataChannelHandler;
+
+/** Buffer containing a single FNiagaraVariable's data at the game level. AoS layout. LWC Types. */
+struct FNiagaraDataChannelVariableBuffer
+{
+	TArray<uint8> Data;
+	int32 Size = 0;
+	
+	void Init(const FNiagaraVariableBase& Var)
+	{
+		Size = Var.GetSizeInBytes();
+	}
+
+	void Empty(){ Data.Empty(); }
+	void Reset(){ Data.Reset(); }
+	
+	template<typename T>
+	bool Write(int32 Index, const T& InData)
+	{
+		check(sizeof(T) == Size);
+		if (Index >= 0 && Index < Num())
+		{
+			T* Dest = ((T*)Data.GetData()) + Index;
+			*Dest = InData;
+
+			return true;
+		}
+		return false;
+	}
+
+	template<typename T>
+	bool Read(int32 Index, T& OutData)const
+	{
+		check(sizeof(T) == Size);
+		if (Index >= 0 && Index < Num())
+		{
+			T* Src = ((T*)Data.GetData()) + Index;
+			OutData = *Src;
+
+			return true;
+		}
+		return false;
+	}
+
+	void SetNum(int32 Num)
+	{
+		Data.SetNumZeroed(Size * Num);
+	}
+
+	int32 Num()const { return Data.Num() / Size; }
+};
+
+USTRUCT()
+struct FNiagaraDataChannelGameDataLayout
+{
+	GENERATED_BODY();
+
+	/** Map of all variables contained in this DataChannel data and the indices into data arrays for game data. */
+	UPROPERTY()
+	TMap<FNiagaraVariableBase, int32> VariableIndices;
+
+	/** Helpers for converting LWC types into Niagara simulation SWC types. */
+	UPROPERTY()
+	TArray<FNiagaraLwcStructConverter> LwcConverters;
+
+	void Init(const TArray<FNiagaraVariable>& Variables);
+};
+
+/** Storage for game level DataChannels generated by BP / C++ */
+struct FNiagaraDataChannelGameData : public TSharedFromThis<FNiagaraDataChannelGameData>
+{
+private:
+	
+	/** Per variable storage buffers. */
+	TArray<FNiagaraDataChannelVariableBuffer> VariableData;
+
+	int32 NumElements = 0;
+
+	TWeakObjectPtr<const UNiagaraDataChannel> DataChannel;
+
+public:
+
+	void Init(const UNiagaraDataChannel* InDataChannel);
+	void Empty();
+	void Reset();
+	void SetNum(int32 NewNum);
+	int32 Num() const { return NumElements; }
+
+	FNiagaraDataChannelVariableBuffer* FindVariableBuffer(const FNiagaraVariableBase& Var);
+
+	void WriteToDataSet(FNiagaraDataBuffer* DestBuffer, int32 DestStartIdx);
+
+	void AppendFromGameData(const FNiagaraDataChannelGameData& GameData);
+
+	void AppendFromDataSet(const FNiagaraDataBuffer* SrcBuffer, FVector3f LwcTile);
+
+	const UNiagaraDataChannel* GetDataChannel()const { return DataChannel.Get(); }
+
+	const TConstArrayView<FNiagaraDataChannelVariableBuffer> GetVariableBuffers()const { return VariableData; }
+};
+
+//////////////////////////////////////////////////////////////////////////
+
+/** Storage for DataChannels the world level. */
+struct NIAGARA_API FNiagaraWorldDataChannelStore
+{
+
+	/** DataChannel data accessible from Game/BP. AoS Layout. LWC types. */
+	FNiagaraDataChannelGameDataPtr GameData;
+
+	//		▲	CPU Sim Data can optionally be made visible to the Game Data.
+	//		|
+	//		|
+	//		|
+	//		▼	Game/BP Data can optionally be made visible to CPU sims.
+
+	/** DataChannel data accessible to Niagara CPU sims. SoA layout. Non LWC types. */
+	FNiagaraDataSet* CPUSimData = nullptr;
+
+	// Cached off buffer with the previous frame's CPU Sim accessible data.
+	// Some systems can choose to read this to avoid any current frame tick ordering issues.
+	FNiagaraDataBuffer* PrevCPUSimData = nullptr;
+	
+	//		▲	GPU Sim Data can optionally be made visible to CPU Sims and Game Data.
+	//		|
+	//		|
+	//		|
+	//		▼	CPU Sim and Game Data can optionally be made visible to the GPU.
+
+	/** DataChannel data accessible to Niagara GPU sims. SoA layout. Non LWC types. */
+	FNiagaraDataSet* GPUSimData = nullptr;
+	
+	~FNiagaraWorldDataChannelStore();
+
+	bool IsInitialized()const {return CPUSimData != nullptr; }
+	void Empty();
+	void Init(UNiagaraDataChannelHandler* Owner);
+	void Tick(UNiagaraDataChannelHandler* Owner);
+};
+
+
+DECLARE_MULTICAST_DELEGATE_OneParam(FOnDataChannelChanged, const UNiagaraDataChannel*);
+
+UCLASS(Experimental, abstract, EditInlineNew)
+class NIAGARA_API UNiagaraDataChannel : public UObject
+{
+public:
+	GENERATED_BODY()
+
+	//UObject Interface
+	virtual void PostLoad();
+	virtual void BeginDestroy();
+#if WITH_EDITOR
+	virtual void PreEditChange(FProperty* PropertyAboutToChange)override;
+	virtual void PostEditChangeProperty(FPropertyChangedEvent& PropertyChangedDataChannel)override;
+#endif
+	//UObject Interface End.
+
+	FName GetChannelName()const { return ChannelName; }
+	TConstArrayView<FNiagaraVariable> GetVariables()const { return Variables; }
+
+	/** If true, we keep our previous frame's data. Some users will prefer a frame of latency to tick dependency. */
+	bool KeepPreviousFrameData()const { return bKeepPreviousFrameData; }
+
+	/** Returns the compiled data describing the data layout for DataChannels in this channel. */
+	const FNiagaraDataSetCompiledData& GetCompiledData(ENiagaraSimTarget SimTarget)const;
+
+	/** Create the appropriate handler object for this data channel. */
+	virtual UNiagaraDataChannelHandler* CreateHandler() const PURE_VIRTUAL(UNiagaraDataChannel::CreateHandler, {return nullptr;} );
+
+#if WITH_EDITOR
+	FOnDataChannelChanged& OnChanged() { return OnChangedDelegate; }
+#endif
+	
+	const FNiagaraDataChannelGameDataLayout& GetGameDataLayout()const { return GameDataLayout; }
+
+	FNiagaraDataChannelGameDataPtr CreateGameData()const;
+
+	bool IsValid()const;
+
+private:
+	UPROPERTY(EditAnywhere, Category = "Data Channel")
+	FName ChannelName;
+
+	/** The variables that define the data contained in this Data Channel. */
+	UPROPERTY(EditAnywhere, Category = "Data Channel")
+	TArray<FNiagaraVariable> Variables;
+
+	/** If true, we keep our previous frame's data. Some users will prefer a frame of latency to tick dependency. */
+	UPROPERTY(EditAnywhere, Category = "Data Channel")
+	bool bKeepPreviousFrameData = false;
+		
+	/**
+	Data layout for payloads in Niagara datasets.
+	*/
+	UPROPERTY()
+	mutable FNiagaraDataSetCompiledData CompiledData;
+
+	UPROPERTY()
+	mutable FNiagaraDataSetCompiledData CompiledDataGPU;
+
+#if WITH_EDITOR
+	FOnDataChannelChanged OnChangedDelegate;
+#endif
+
+	/** Layout information for any data stored at the "Game" level. i.e. From game code/BP. AoS layout and LWC types. */
+	UPROPERTY()
+	FNiagaraDataChannelGameDataLayout GameDataLayout;
+};
+
+/**
+* A C++ and Blueprint accessible library of utility functions for accessing Niagara DataChannel
+*/
+UCLASS(Experimental)
+class NIAGARA_API UNiagaraDataChannelLibrary : public UBlueprintFunctionLibrary
+{
+	GENERATED_UCLASS_BODY()
+
+	UFUNCTION(BlueprintCallable, Category = NiagaraDataChannel, meta = (Keywords = "niagara DataChannel", WorldContext = "WorldContextObject", UnsafeDuringActorConstruction = "true"))
+	static UNiagaraDataChannelHandler* GetNiagaraDataChannel(const UObject* WorldContextObject, FName Channel);
+};
diff --git a/Engine/Plugins/FX/Niagara/Source/Niagara/Public/NiagaraDataChannelHandler.h b/Engine/Plugins/FX/Niagara/Source/Niagara/Public/NiagaraDataChannelHandler.h
new file mode 100644
--- /dev/null
+++ b/Engine/Plugins/FX/Niagara/Source/Niagara/Public/NiagaraDataChannelHandler.h
@@ -0,0 +1,95 @@
+// Copyright Epic Games, Inc. All Rights Reserved.
+
+#pragma once
+
+/**
+Base class for Niagara DataChannel Handlers.
+Data Channel handlers are the runtime counterpart to Data Channels.
+They control how data being written to the Data Channel is stored and how to expose data being read from the Data Channel.
+For example, the simplest handler is the UNiagaraDataChannelHandler_Global which just keeps all data in a single large set that is used by all systems.
+Some more complex handlers may want to divide up the scene in various different ways to better match particular use cases.
+
+*/
+
+
+#include "NiagaraDataChannelCommon.h"
+#include "NiagaraDataChannelHandler.generated.h"
+
+/** Base class for all UNiagaraDataChannelHandler's Render Thread Proxies. */
+struct FNiagaraDataChannelHandlerRTProxyBase
+{
+};
+
+UCLASS(Experimental, abstract, BlueprintType)
+class NIAGARA_API UNiagaraDataChannelHandler : public UObject
+{
+public:
+
+	GENERATED_BODY()
+
+	virtual ~UNiagaraDataChannelHandler();
+
+	virtual void Init(const UNiagaraDataChannel* InChannel);
+
+	/** Returns the DataChannel data for the given DataChannel name and location. */
+	virtual void Tick(float DeltaTime, ETickingGroup TickGroup, FNiagaraWorldManager* OwningWorld);
+
+	/** Returns the data for this data channel given the system instance. Possibly some handlers will subdivide spacially etc so may not return the same data for the different systems. */
+	virtual void GetData(FNiagaraSystemInstance* SystemInstance, FNiagaraDataBuffer*& OutCPUData, bool bGetLastFrameData)  PURE_VIRTUAL(UNiagaraDataChannelHandler::GetData, );	
+
+	/** Gets the GPU data set with data for the given system. This API will likely be reworked as GPU support is fleshed out and more complex handlers are created. */
+	virtual void GetDataGPU(FNiagaraSystemInstance* SystemInstance, FNiagaraDataSet*& OutCPUData) PURE_VIRTUAL(UNiagaraDataChannelHandler::GetDataGPU, );	
+
+	/** Adds a request to publish some data into the channel on the next tick. */
+	virtual void Publish(const FNiagaraDataChannelPublishRequest& Request);
+		
+	/**
+	 *Removes all publish requests involving the given dataset.
+	 *TODO: REMOVE
+	 *This is a hack to get around lifetime issues wrt data buffers/datasets and their compiled data.
+	 *We should rework things such that data buffers can exist beyond their owning dataset, including the compiled data detailing their layout.
+	 **/
+	virtual void RemovePublishRequests(const FNiagaraDataSet* DataSet);
+
+	/** Returns the game level DataChannel data store. */
+	virtual FNiagaraDataChannelGameDataPtr GetGameData(UActorComponent* OwningComponent) PURE_VIRTUAL(UNiagaraDataChannelHandler::GetGameData, return nullptr; );
+
+	/** Utility allowing easy casting of proxy to it's proper derived type. */
+	template<typename T>
+	T* GetRTProxyAs() 
+	{
+		T* TypedProxy = static_cast<T*>(RTProxy.Get());
+		check(TypedProxy != nullptr);
+		return TypedProxy;
+	}
+
+	const UNiagaraDataChannel* GetDataChannel()const { return DataChannel; }
+
+	UFUNCTION(BlueprintCallable, Category="Data Channel")
+	UNiagaraDataChannelWriter* GetDataChannelWriter();
+
+	UFUNCTION(BlueprintCallable, Category = "Data Channel")
+	UNiagaraDataChannelReader* GetDataChannelReader();
+
+protected:
+
+	UPROPERTY()
+	TObjectPtr<const UNiagaraDataChannel> DataChannel;
+
+	/** Per frame requests to publish data from Niagara Systems into the main world DataChannels buffers. */
+	TArray<FNiagaraDataChannelPublishRequest> PublishRequests;
+
+	/** Data buffers we'll be passing to the RT for uploading to the GPU */
+	TArray<FNiagaraDataBuffer*> BuffersForGPU;
+
+	/** Render Thread Proxy for this handler. Each handler type must create it's proxy inside it's constructor. */
+	TUniquePtr<FNiagaraDataChannelHandlerRTProxyBase> RTProxy;
+
+	/** Helper object allowing BP to write data in this channel. */
+	UPROPERTY()
+	TObjectPtr<UNiagaraDataChannelWriter> Writer;
+
+	/** Helper object allowing BP to read data in this channel. */
+	UPROPERTY()
+	TObjectPtr<UNiagaraDataChannelReader> Reader;
+};
\ No newline at end of file
diff --git a/Engine/Plugins/FX/Niagara/Source/Niagara/Public/NiagaraDataInterfaceUtilities.h b/Engine/Plugins/FX/Niagara/Source/Niagara/Public/NiagaraDataInterfaceUtilities.h
--- a/Engine/Plugins/FX/Niagara/Source/Niagara/Public/NiagaraDataInterfaceUtilities.h
+++ b/Engine/Plugins/FX/Niagara/Source/Niagara/Public/NiagaraDataInterfaceUtilities.h
@@ -18,6 +18,11 @@ namespace FNiagaraDataInterfaceUtilities
 		UNiagaraDataInterface*	DataInterface = nullptr;
 	};
 
+	struct FDataInterfaceSearchOptions
+	{
+		bool bIncludeInternal = false;
+	};
+
 	// Finds all VM function calls made using the data interface that equals this one (i.e. A->Equals(B))
 	// The same function call made be made multiple times across different scripts so you may see the same function multiple times
 	// The action should return True to continue iteration or False to stop
@@ -43,13 +48,13 @@ namespace FNiagaraDataInterfaceUtilities
 
 	// Loops over all data interfaces inside the SystemInstance
 	// The action should return True to continue iteration or False to stop
-	NIAGARA_API void ForEachDataInterface(FNiagaraSystemInstance* SystemInstance, TFunction<bool(const FNiagaraVariableBase Variable, UNiagaraDataInterface* DataInterface)> Action);
+	NIAGARA_API void ForEachDataInterface(FNiagaraSystemInstance* SystemInstance, TFunction<bool(const FNiagaraVariableBase Variable, UNiagaraDataInterface* DataInterface)> Action, FDataInterfaceSearchOptions SearchOptions = FDataInterfaceSearchOptions());
 	// Loops over all data interfaces inside the SystemInstance
 	// The action should return True to continue iteration or False to stop
-	NIAGARA_API void ForEachDataInterface(FNiagaraSystemInstance* SystemInstance, TFunction<bool(const FDataInterfaceUsageContext&)> Action);
+	NIAGARA_API void ForEachDataInterface(FNiagaraSystemInstance* SystemInstance, TFunction<bool(const FDataInterfaceUsageContext&)> Action, FDataInterfaceSearchOptions SearchOptions = FDataInterfaceSearchOptions());
 	// Loops over all data interfaces inside the NiagaraSystem
 	// The action should return True to continue iteration or False to stop
-	NIAGARA_API void ForEachDataInterface(UNiagaraSystem* NiagaraSystem, TFunction<bool(const FDataInterfaceUsageContext&)> Action);
+	NIAGARA_API void ForEachDataInterface(UNiagaraSystem* NiagaraSystem, TFunction<bool(const FDataInterfaceUsageContext&)> Action, FDataInterfaceSearchOptions SearchOptions = FDataInterfaceSearchOptions());
 }
 
 #if WITH_EDITOR
diff --git a/Engine/Plugins/FX/Niagara/Source/Niagara/Public/NiagaraModule.h b/Engine/Plugins/FX/Niagara/Source/Niagara/Public/NiagaraModule.h
--- a/Engine/Plugins/FX/Niagara/Source/Niagara/Public/NiagaraModule.h
+++ b/Engine/Plugins/FX/Niagara/Source/Niagara/Public/NiagaraModule.h
@@ -118,6 +118,9 @@ public:
 	FORCEINLINE static bool UseGlobalFXBudget() { return bUseGlobalFXBudget; }
 	static void OnUseGlobalFXBudgetChanged(IConsoleVariable* Variable);
 
+	FORCEINLINE static bool DataChannelsEnabled() { return bDataChannelsEnabled; }
+	static void OnDataChannelsEnabledChanged(IConsoleVariable* Variable);
+
 	FORCEINLINE static float GetGlobalSpawnCountScale() { return EngineGlobalSpawnCountScale; }
 	FORCEINLINE static float GetGlobalSystemCountScale() { return EngineGlobalSystemCountScale; }
 
@@ -250,7 +253,7 @@ public:
 	static int32 EngineEffectsQuality;
 
 	static bool bUseGlobalFXBudget;
-
+	static bool bDataChannelsEnabled;
 
 private:
 	static FNiagaraVariable Engine_WorldDeltaTime;
diff --git a/Engine/Plugins/FX/Niagara/Source/Niagara/Public/NiagaraTypes.h b/Engine/Plugins/FX/Niagara/Source/Niagara/Public/NiagaraTypes.h
--- a/Engine/Plugins/FX/Niagara/Source/Niagara/Public/NiagaraTypes.h
+++ b/Engine/Plugins/FX/Niagara/Source/Niagara/Public/NiagaraTypes.h
@@ -10,6 +10,7 @@
 #if UE_ENABLE_INCLUDE_ORDER_DEPRECATED_IN_5_2
 #include "Engine/Texture2D.h"
 #endif
+#include "NiagaraCore.h"
 
 #include "NiagaraTypes.generated.h"
 
@@ -254,10 +255,10 @@ struct FNiagaraStructConversionStep
 	GENERATED_USTRUCT_BODY();
 	
 	UPROPERTY()
-	int32 SourceBytes = 0;
+	int32 LWCBytes = 0;
 
 	UPROPERTY()
-	int32 SourceOffset = 0;
+	int32 LWCOffset = 0;
 	
 	UPROPERTY()
 	int32 SimulationBytes = 0;
@@ -269,10 +270,21 @@ struct FNiagaraStructConversionStep
 	ENiagaraStructConversionType ConversionType = ENiagaraStructConversionType::CopyOnly;
 
 	FNiagaraStructConversionStep();
-	FNiagaraStructConversionStep(int32 InSourceBytes, int32 InSourceOffset, int32 InSimulationBytes, int32 InSimulationOffset, ENiagaraStructConversionType InConversionType);
+	FNiagaraStructConversionStep(int32 InLWCBytes, int32 InLWCOffset, int32 InSimulationBytes, int32 InSimulationOffset, ENiagaraStructConversionType InConversionType);
 
-	void CopyToSim(uint8* DestinationData, const uint8* SourceData) const;
-	void CopyFromSim(uint8* DestinationData, const uint8* SourceData) const;
+	void CopyToSim(const uint8* LWCData, uint8* SimulationData, int32 Count, int32 LWCStride, int32 SimulationStride) const;
+	void CopyFromSim(uint8* LWCData, const uint8* SimulationData, int32 Count, int32 LWCStride, int32 SimulationStride) const;
+
+	template<typename DstType, typename SrcType>
+	void CopyInternal(uint8* DestinationData, int32 DestStride, int32 DestOffset, const uint8* SourceData, int32 SourceStride, int32 SourceOffset, int32 Count)const
+	{
+		for (int32 i = 0; i < Count; ++i)
+		{
+			DstType* Dst = (DstType*)((DestinationData + i * DestStride) + DestOffset);
+			const SrcType* Src = (const SrcType*)((SourceData + i * SourceStride) + SourceOffset);			
+			*Dst = DstType(*Src);
+		}
+	}
 };
 
 /** Can convert struct data from custom structs containing LWC data such as FVector3d into struct data suitable for Niagara simulations and vice versa. */
@@ -280,14 +292,24 @@ USTRUCT()
 struct NIAGARA_API FNiagaraLwcStructConverter
 {
 	GENERATED_USTRUCT_BODY();
-	
-	void ConvertDataToSimulation(uint8* DestinationData, const uint8* SourceData) const;
-	void ConvertDataFromSimulation(uint8* DestinationData, const uint8* SourceData) const;
+
+	FNiagaraLwcStructConverter(){}
+	FNiagaraLwcStructConverter(UScriptStruct* LWCStruct, UScriptStruct* SWCStruct)
+		: LWCSize(LWCStruct ? LWCStruct->GetStructureSize() : 0)
+		, SWCSize(SWCStruct ? SWCStruct->GetStructureSize() : 0)
+	{
+	}
+	void ConvertDataToSimulation(uint8* DestinationData, const uint8* SourceData, int32 Count = 1) const;
+	void ConvertDataFromSimulation(uint8* DestinationData, const uint8* SourceData, int32 Count = 1) const;
 	
 	void AddConversionStep(int32 InSourceBytes, int32 InSourceOffset, int32 InSimulationBytes, int32 InSimulationOffset, ENiagaraStructConversionType ConversionType);
-	bool IsValid() const { return ConversionSteps.Num() > 0; }
+	bool IsValid() const { return ConversionSteps.Num() > 0 && LWCSize > 0 && SWCSize > 0; }
 	
 private:
+	UPROPERTY()
+	int32 LWCSize = 0;
+	UPROPERTY()
+	int32 SWCSize = 0;
 	UPROPERTY()
 	TArray<FNiagaraStructConversionStep> ConversionSteps;
 };
@@ -774,11 +796,38 @@ struct NIAGARA_API FNiagaraTypeDefinition
 		/// indicates that the ClassStructOrEnum property has been serialized as the LWC struct (see FNiagaraTypeHelper)
 		/// instead of the Transient SWC version of the struct
 		TF_SerializedAsLWC	= 0x02,
+
+		/// indicates this type allows LWC types and should not attempt to convert to SWC.
+		TF_AllowLWC		= 0x04,
 	};
 
 public:
 
 	// Construct blank raw type definition 
+	FORCEINLINE FNiagaraTypeDefinition(UObject* Object)
+		: ClassStructOrEnum(Object), Flags(TF_None), Size(INDEX_NONE), Alignment(INDEX_NONE)
+#if WITH_EDITORONLY_DATA
+		, Struct_DEPRECATED(nullptr), Enum_DEPRECATED(nullptr)
+#endif
+	{
+		if(UClass* Class = Cast<UClass>(Object))
+		{
+			*this = FNiagaraTypeDefinition(Class);
+		}
+		else if (UScriptStruct* Struct = Cast<UScriptStruct>(Object))
+		{
+			*this = FNiagaraTypeDefinition(Struct, EAllowUnfriendlyStruct::Deny);
+		}
+		else if(UEnum* Enum = Cast<UEnum>(Object))
+		{
+			*this = FNiagaraTypeDefinition(Enum);
+		}
+		else
+		{
+			check(0);
+		}
+	}
+
 	FORCEINLINE FNiagaraTypeDefinition(UClass *ClassDef)
 		: ClassStructOrEnum(ClassDef), UnderlyingType(UT_Class), Flags(TF_None), Size(INDEX_NONE), Alignment(INDEX_NONE)
 #if WITH_EDITORONLY_DATA
@@ -806,6 +855,10 @@ public:
 	{
 		checkSlow(ClassStructOrEnum != nullptr);
 		ensureAlwaysMsgf(AllowUnfriendlyStruct == EAllowUnfriendlyStruct::Allow || FNiagaraTypeHelper::IsNiagaraFriendlyTopLevelStruct(StructDef, ENiagaraStructConversion::UserFacing), TEXT("Struct(%s) is not supported."), *StructDef->GetName());
+		if(AllowUnfriendlyStruct == EAllowUnfriendlyStruct::Allow)
+		{
+			Flags |= TF_AllowLWC;
+		}
 	}
 	FORCEINLINE FNiagaraTypeDefinition(UScriptStruct* StructDef) : FNiagaraTypeDefinition(StructDef, EAllowUnfriendlyStruct::Deny) {}
 
@@ -817,6 +870,19 @@ public:
 	{
 	}
 
+	FORCEINLINE FNiagaraTypeDefinition(FProperty* Property, EAllowUnfriendlyStruct AllowUnfriendlyStruct)
+	{
+		if (Property->IsA(FFloatProperty::StaticClass())) { *this = FNiagaraTypeDefinition::GetFloatDef(); }
+		else if (Property->IsA(FUInt16Property::StaticClass())) { *this = FNiagaraTypeDefinition::GetHalfDef(); }
+		else if (Property->IsA(FIntProperty::StaticClass())) { *this = FNiagaraTypeDefinition::GetIntDef(); }
+		else if (Property->IsA(FBoolProperty::StaticClass())) { *this = FNiagaraTypeDefinition::GetBoolDef(); }
+		else if (FStructProperty* StructProp = CastField<FStructProperty>(Property)) { *this = FNiagaraTypeDefinition(StructProp->Struct, AllowUnfriendlyStruct); }
+		else
+		{
+			checkf(0, TEXT("Invalid Type"))
+		}
+	}
+
 	// Construct a blank raw type definition
 	FORCEINLINE FNiagaraTypeDefinition()
 		: ClassStructOrEnum(nullptr), UnderlyingType(UT_None), Flags(TF_None), Size(INDEX_NONE), Alignment(INDEX_NONE)
@@ -885,7 +951,7 @@ public:
 			{
 				NameText = FText::Format(NSLOCTEXT("NiagaraTypeDefinition", "FlagFormatNameText", "Static {0}"), NameText);
 			}
-			else
+			else if (!AllowsLWC())
 			{
 				NameText = FText::Format(NSLOCTEXT("NiagaraTypeDefinition", "UnknownNameText", "Unknown Flag Type {0}"), NameText);
 			}
@@ -943,9 +1009,9 @@ public:
 
 	bool IsEnum() const { return UnderlyingType == UT_Enum; }
 
-	bool IsStatic() const {
-		return (GetFlags() & TF_Static) != 0;
-	}
+	bool IsStatic() const { return (GetFlags() & TF_Static) != 0; }
+
+	bool AllowsLWC() const { return (GetFlags() & TF_AllowLWC) != 0; }
 
 	void SetFlags(FTypeFlags InFlags)
 	{
@@ -1339,6 +1405,11 @@ public:
 		return Get().RegisteredIndexTypes;
 	}
 
+	static const TArray<FNiagaraTypeDefinition>& GetSimulationTypes()
+	{
+		return Get().RegisteredSimulationTypes;
+	}
+
 	static UNiagaraDataInterfaceBase* GetDefaultDataInterfaceByName(const FString& DIClassName);
 
 	static void ClearUserDefinedRegistry()
@@ -1389,11 +1460,19 @@ public:
 	{
 		FNiagaraTypeRegistry& Registry = Get();
 
+		FNiagaraTypeDefinition SimType = NewType;
+		if (UScriptStruct* Struct = NewType.GetScriptStruct())
+		{
+			SimType = FNiagaraTypeDefinition(FNiagaraTypeHelper::GetSWCStruct(Struct), FNiagaraTypeDefinition::EAllowUnfriendlyStruct::Deny);			
+		}
+
 		FRWScopeLock Lock(Registry.RegisteredTypesLock, SLT_Write);
 
 		//TODO: Make this a map of type to a more verbose set of metadata? Such as the hlsl defs, offset table for conversions etc.
 		Registry.RegisteredTypeIndexMap.Add(GetTypeHash(NewType), Registry.RegisteredTypes.AddUnique(NewType));
 
+		Registry.RegisteredSimulationTypes.AddUnique(SimType);
+
 		if (EnumHasAnyFlags(Flags, ENiagaraTypeRegistryFlags::AllowUserVariable))
 		{
 			Registry.RegisteredUserVariableTypes.AddUnique(NewType);
@@ -1524,6 +1603,7 @@ private:
 	TArray<FNiagaraTypeDefinition> RegisteredUserDefinedTypes;
 	TArray<FNiagaraTypeDefinition> RegisteredNumericTypes;
 	TArray<FNiagaraTypeDefinition> RegisteredIndexTypes;
+	TArray<FNiagaraTypeDefinition> RegisteredSimulationTypes;
 
 
 	TMap<uint32, int32> RegisteredTypeIndexMap;
@@ -1603,6 +1683,11 @@ struct FNiagaraVariableBase
 #endif
 		{}
 
+	FORCEINLINE FNiagaraVariableBase(const FNiagaraVariableCommonReference& VariableRef)
+		: Name(VariableRef.Name)
+		, TypeDefHandle(FNiagaraTypeDefinition(VariableRef.UnderlyingType))
+	{}
+
 	/** Check if Name and Type definition are the same. The actual stored value is not checked here.*/
 	bool operator==(const FNiagaraVariableBase& Other)const
 	{
@@ -1736,7 +1821,7 @@ struct TStructOpsTypeTraits<FNiagaraVariableBase> : public TStructOpsTypeTraitsB
 	};
 };
 
-USTRUCT()
+USTRUCT(BlueprintType)
 struct FNiagaraVariable : public FNiagaraVariableBase
 {
 	GENERATED_USTRUCT_BODY()
diff --git a/Engine/Plugins/FX/Niagara/Source/Niagara/Public/NiagaraWorldManager.h b/Engine/Plugins/FX/Niagara/Source/Niagara/Public/NiagaraWorldManager.h
--- a/Engine/Plugins/FX/Niagara/Source/Niagara/Public/NiagaraWorldManager.h
+++ b/Engine/Plugins/FX/Niagara/Source/Niagara/Public/NiagaraWorldManager.h
@@ -21,6 +21,7 @@
 #include "NiagaraScalabilityManager.h"
 #include "NiagaraDebuggerCommon.h"
 #include "NiagaraDeferredMethodQueue.h"
+#include "NiagaraDataChannelManager.h"
 
 #include "NiagaraWorldManager.generated.h"
 
@@ -123,6 +124,7 @@ public:
 	/** Called before we run end of frame updates, allows us to wait on async work. */
 	void PreSendAllEndOfFrameUpdates();
 
+	void OnWorldBeginTearDown();
 	void OnWorldCleanup(bool bSessionEnded, bool bCleanupResources);
 	void OnPostWorldCleanup(bool bSessionEnded, bool bCleanupResources);
 
@@ -216,6 +218,11 @@ public:
 	NIAGARA_API static void SetScalabilityCullingMode(ENiagaraScalabilityCullingMode NewMode);
 	NIAGARA_API static ENiagaraScalabilityCullingMode GetScalabilityCullingMode() { return ScalabilityCullingMode; }
 
+	FNiagaraDataChannelManager& GetDataChannelManager(){ return DataChannelManager; }
+
+	/** Waits for all currently in flight async work to be completed. */
+	void WaitForAsyncWork();
+
 private:
 	// Callback function registered with global world delegates to instantiate world manager when a game world is created
 	static void OnWorldInit(UWorld* World, const UWorld::InitializationValues IVS);
@@ -327,6 +334,8 @@ private:
 
 	/** A global flag for all scalability culling */
 	static ENiagaraScalabilityCullingMode ScalabilityCullingMode;
+
+	FNiagaraDataChannelManager DataChannelManager;
 };
 
 
@@ -352,4 +361,4 @@ void FNiagaraWorldManager::ForAllWorldManagers(TAction Func)
 			Func(*Pair.Value);
 		}
 	}
-}
\ No newline at end of file
+}
diff --git a/Engine/Plugins/FX/Niagara/Source/NiagaraCore/Public/NiagaraCompileHash.h b/Engine/Plugins/FX/Niagara/Source/NiagaraCore/Public/NiagaraCompileHash.h
--- a/Engine/Plugins/FX/Niagara/Source/NiagaraCore/Public/NiagaraCompileHash.h
+++ b/Engine/Plugins/FX/Niagara/Source/NiagaraCore/Public/NiagaraCompileHash.h
@@ -66,6 +66,11 @@ private:
 	TArray<uint8> DataHash;
 };
 
+FORCEINLINE uint32 GetTypeHash(const FNiagaraCompileHash& Hash)
+{
+	return Hash.GetTypeHash();
+}
+
 inline bool operator==(const FSHAHash& Lhs, const FNiagaraCompileHash& Rhs)
 {
 	return Rhs.operator==(Lhs);
diff --git a/Engine/Plugins/FX/Niagara/Source/NiagaraEditor/Private/AssetDefinitions/AssetDefinition_NiagaraDataChannelDefinitions.cpp b/Engine/Plugins/FX/Niagara/Source/NiagaraEditor/Private/AssetDefinitions/AssetDefinition_NiagaraDataChannelDefinitions.cpp
new file mode 100644
--- /dev/null
+++ b/Engine/Plugins/FX/Niagara/Source/NiagaraEditor/Private/AssetDefinitions/AssetDefinition_NiagaraDataChannelDefinitions.cpp
@@ -0,0 +1,26 @@
+// Copyright Epic Games, Inc. All Rights Reserved.
+
+#include "AssetDefinitions/AssetDefinition_NiagaraDataChannelDefinitions.h"
+
+#include "NiagaraEditorStyle.h"
+#include "NiagaraDataChannelDefinitions.h"
+
+#define LOCTEXT_NAMESPACE "UAssetDefinition_NiagaraDataChannelDefinitions"
+
+FLinearColor UAssetDefinition_NiagaraDataChannelDefinitions::GetAssetColor() const
+{
+	return FNiagaraEditorStyle::Get().GetColor("NiagaraEditor.AssetColors.DataChannelDefinitions");
+}
+
+TSoftClassPtr<> UAssetDefinition_NiagaraDataChannelDefinitions::GetAssetClass() const
+{
+	return UNiagaraDataChannelDefinitions::StaticClass();
+}
+
+TConstArrayView<FAssetCategoryPath> UAssetDefinition_NiagaraDataChannelDefinitions::GetAssetCategories() const
+{
+	static FAssetCategoryPath AssetPaths[] = { FAssetCategoryPath(LOCTEXT("NiagaraAssetsCategory", "FX"), LOCTEXT("NiagaraDataChannelDefinitions_SubCategory", "Advanced")) };
+	return AssetPaths;
+}
+
+#undef LOCTEXT_NAMESPACE
diff --git a/Engine/Plugins/FX/Niagara/Source/NiagaraEditor/Private/AssetDefinitions/AssetDefinition_NiagaraDataChannelDefinitions.h b/Engine/Plugins/FX/Niagara/Source/NiagaraEditor/Private/AssetDefinitions/AssetDefinition_NiagaraDataChannelDefinitions.h
new file mode 100644
--- /dev/null
+++ b/Engine/Plugins/FX/Niagara/Source/NiagaraEditor/Private/AssetDefinitions/AssetDefinition_NiagaraDataChannelDefinitions.h
@@ -0,0 +1,21 @@
+// Copyright Epic Games, Inc. All Rights Reserved.
+
+#pragma once
+
+#include "CoreMinimal.h"
+#include "AssetDefinitionDefault.h"
+
+#include "AssetDefinition_NiagaraDataChannelDefinitions.generated.h"
+
+UCLASS()
+class UAssetDefinition_NiagaraDataChannelDefinitions : public UAssetDefinitionDefault
+{
+	GENERATED_BODY()
+
+public:
+	// UAssetDefinition Implementation
+	virtual FText GetAssetDisplayName() const override { return NSLOCTEXT("AssetDefinition", "NiagaraDataChannelDefinitions", "Niagara Data Channel Definitions"); }
+	virtual FLinearColor GetAssetColor() const override;
+	virtual TSoftClassPtr<UObject> GetAssetClass() const override;
+	virtual TConstArrayView<FAssetCategoryPath> GetAssetCategories() const override;
+};
diff --git a/Engine/Plugins/FX/Niagara/Source/NiagaraEditor/Private/NiagaraActions.cpp b/Engine/Plugins/FX/Niagara/Source/NiagaraEditor/Private/NiagaraActions.cpp
--- a/Engine/Plugins/FX/Niagara/Source/NiagaraEditor/Private/NiagaraActions.cpp
+++ b/Engine/Plugins/FX/Niagara/Source/NiagaraEditor/Private/NiagaraActions.cpp
@@ -22,6 +22,11 @@
 #include "Widgets/SNiagaraParameterName.h"
 #include "Widgets/SToolTip.h"
 
+#include "NiagaraDataChannel.h"
+#include "NiagaraDataChannelDefinitions.h"
+#include "DataInterface/NiagaraDataInterfaceDataChannelRead.h"
+#include "DataInterface/NiagaraDataInterfaceDataChannelWrite.h"
+
 #include UE_INLINE_GENERATED_CPP_BY_NAME(NiagaraActions)
 
 #define LOCTEXT_NAMESPACE "NiagaraActions"
@@ -678,4 +683,117 @@ void INiagaraDataInterfaceNodeActionProvider::GetNodeContextMenuActions(UClass*
 	}
 }
 
+////////////////////////
+
+void FNiagaraDataInterfaceNodeActionProvider_DataChannelWrite::GetNodeContextMenuActionsImpl(UToolMenu* Menu, UGraphNodeContextMenuContext* Context, FNiagaraFunctionSignature Signature) const
+{
+	//For all functions except "Num", add a context menu to initialized to a specific data channel.
+	if (Signature.Name == TEXT("Num"))
+	{
+		return;
+	}
+
+	const TArray<UNiagaraDataChannelDefinitions*>& ChannelDefs = UNiagaraDataChannelDefinitions::GetDataChannelDefinitions(false, true);
+	if (ChannelDefs.Num() > 0)
+	{
+		return;
+	}
+
+	FToolMenuSection& Section = Menu->AddSection("DataChannelWrite", LOCTEXT("DataChannelWriteHeader", "Data Channel Write"));
+
+	TWeakObjectPtr<UNiagaraNodeFunctionCall> WeakFuncNode = CastChecked<UNiagaraNodeFunctionCall>(Context->Node);
+
+	auto CreateNodeContextMenu = [WeakFuncNode](UToolMenu* InNewToolMenu)
+	{
+		FToolMenuSection& SubSection = InNewToolMenu->AddSection("InitForDataChannelSection", LOCTEXT("DataChannels", "Data Channels"));
+
+		const TArray<UNiagaraDataChannelDefinitions*>& DataChannelDefs = UNiagaraDataChannelDefinitions::GetDataChannelDefinitions(false, true);
+		for (const UNiagaraDataChannelDefinitions* Def : DataChannelDefs)
+		{
+			for (TObjectPtr<const UNiagaraDataChannel> Channel : Def->DataChannels)
+			{
+				TWeakObjectPtr<const UNiagaraDataChannel> WeakChannel = Channel;
+				auto CreateDataChannelActionEntry = [WeakFuncNode, WeakChannel]()
+				{
+					if (UNiagaraNodeFunctionCall* Node = WeakFuncNode.Get())
+					{
+						if (const UNiagaraDataChannel* ChannelPtr = WeakChannel.Get())
+						{
+							TConstArrayView<FNiagaraVariable> ChannelVars = ChannelPtr->GetVariables();
+							for (const FNiagaraVariable& Var : ChannelVars)
+							{
+								FNiagaraVariable SWCVar(FNiagaraTypeHelper::GetSWCStruct(Var.GetType().GetScriptStruct()), Var.GetName());
+								Node->AddParameter(SWCVar, Node->Pins.Last());
+							}
+						}
+					}
+				};
+				SubSection.AddMenuEntry(NAME_None, FText::FromName(Channel->GetChannelName()), FText::FromName(Channel->GetChannelName()), FSlateIcon(FAppStyle::GetAppStyleSetName(), "Icons.Import"),
+					FUIAction(FExecuteAction::CreateLambda(CreateDataChannelActionEntry), FCanExecuteAction()));
+			}
+		}
+	};
+
+	Section.AddSubMenu("InitForDataChannelMenu", LOCTEXT("InitForDataChannel", "Init For Data Channel..."), LOCTEXT("InitForDataChannelTooltip", "Initializes this node to write to all members of a given data channel."),
+		FNewToolMenuDelegate::CreateLambda(CreateNodeContextMenu));
+}
+
+
+void FNiagaraDataInterfaceNodeActionProvider_DataChannelRead::GetNodeContextMenuActionsImpl(UToolMenu* Menu, UGraphNodeContextMenuContext* Context, FNiagaraFunctionSignature Signature) const
+{
+	//For all functions except "Num", add a context menu to initialized to a specific data channel.
+	if (Signature.Name == TEXT("Num"))
+	{
+		return;
+	}
+
+	const TArray<UNiagaraDataChannelDefinitions*>& ChannelDefs = UNiagaraDataChannelDefinitions::GetDataChannelDefinitions(false, true);
+	if (ChannelDefs.Num() > 0)
+	{
+		return;
+	}
+
+	FToolMenuSection& Section = Menu->AddSection("DataChannelRead", LOCTEXT("DataChannelReadHeader", "Data Channel Read"));
+
+	TWeakObjectPtr<UNiagaraNodeFunctionCall> WeakFuncNode = CastChecked<UNiagaraNodeFunctionCall>(Context->Node);
+
+	auto CreateNodeContextMenu = [WeakFuncNode](UToolMenu* InNewToolMenu)
+	{
+		FToolMenuSection& SubSection = InNewToolMenu->AddSection("InitForDataChannelSection", LOCTEXT("DataChannels", "Data Channels"));
+
+		const TArray<UNiagaraDataChannelDefinitions*>& DataChannelDefs = UNiagaraDataChannelDefinitions::GetDataChannelDefinitions(false, true);
+		for (const UNiagaraDataChannelDefinitions* Def : DataChannelDefs)
+		{
+			for (TObjectPtr<const UNiagaraDataChannel> Channel : Def->DataChannels)
+			{
+				TWeakObjectPtr<const UNiagaraDataChannel> WeakChannel = Channel;
+				auto CreateDataChannelActionEntry = [WeakFuncNode, WeakChannel]()
+				{
+					if (UNiagaraNodeFunctionCall* Node = WeakFuncNode.Get())
+					{
+						Node->RemoveAllDynamicPins();
+						if (const UNiagaraDataChannel* ChannelPtr = WeakChannel.Get())
+						{
+							TConstArrayView<FNiagaraVariable> ChannelVars = ChannelPtr->GetVariables();
+							for (const FNiagaraVariable& Var : ChannelVars)
+							{
+								FNiagaraVariable SWCVar(FNiagaraTypeHelper::GetSWCStruct(Var.GetType().GetScriptStruct()), Var.GetName());
+								Node->AddParameter(SWCVar, Node->Pins.Last());
+							}
+						}
+					}
+				};
+
+				SubSection.AddMenuEntry(NAME_None, FText::FromName(Channel->GetChannelName()), FText::FromName(Channel->GetChannelName()), FSlateIcon(FAppStyle::GetAppStyleSetName(), "Icons.Import"), 
+					FUIAction(FExecuteAction::CreateLambda(CreateDataChannelActionEntry), FCanExecuteAction()));
+			}
+		}
+	};
+
+	Section.AddSubMenu("InitForDataChannelMenu", LOCTEXT("InitForDataChannel", "Init For Data Channel..."),	LOCTEXT("InitForDataChannelTooltip", "Initializes this node to read to all members of a given data channel."),
+		FNewToolMenuDelegate::CreateLambda(CreateNodeContextMenu));
+}
+
+
+
 #undef LOCTEXT_NAMESPACE
\ No newline at end of file
diff --git a/Engine/Plugins/FX/Niagara/Source/NiagaraEditor/Private/NiagaraDataChannelDefinitionsFactoryNew.cpp b/Engine/Plugins/FX/Niagara/Source/NiagaraEditor/Private/NiagaraDataChannelDefinitionsFactoryNew.cpp
new file mode 100644
--- /dev/null
+++ b/Engine/Plugins/FX/Niagara/Source/NiagaraEditor/Private/NiagaraDataChannelDefinitionsFactoryNew.cpp
@@ -0,0 +1,24 @@
+// Copyright Epic Games, Inc. All Rights Reserved.
+
+#include "NiagaraDataChannelDefinitionsFactoryNew.h"
+#include "CoreMinimal.h"
+#include "Misc/ConfigCacheIni.h"
+#include "NiagaraSettings.h"
+#include "NiagaraDataChannelDefinitions.h"
+
+UNiagaraDataChannelDefinitionsFactoryNew::UNiagaraDataChannelDefinitionsFactoryNew(const FObjectInitializer& ObjectInitializer)
+	: Super(ObjectInitializer)
+{
+	SupportedClass = UNiagaraDataChannelDefinitions::StaticClass();
+	bEditAfterNew = true;
+	bCreateNew = true;
+}
+
+UObject* UNiagaraDataChannelDefinitionsFactoryNew::FactoryCreateNew(UClass* Class, UObject* InParent, FName Name, EObjectFlags Flags, UObject* Context, FFeedbackContext* Warn)
+{
+	check(Class->IsChildOf(UNiagaraDataChannelDefinitions::StaticClass()));
+
+	UNiagaraDataChannelDefinitions* NewDefinitions = NewObject<UNiagaraDataChannelDefinitions>(InParent, Class, Name, Flags | RF_Transactional);
+
+	return NewDefinitions;
+}
diff --git a/Engine/Plugins/FX/Niagara/Source/NiagaraEditor/Private/NiagaraEditorModule.cpp b/Engine/Plugins/FX/Niagara/Source/NiagaraEditor/Private/NiagaraEditorModule.cpp
--- a/Engine/Plugins/FX/Niagara/Source/NiagaraEditor/Private/NiagaraEditorModule.cpp
+++ b/Engine/Plugins/FX/Niagara/Source/NiagaraEditor/Private/NiagaraEditorModule.cpp
@@ -45,11 +45,15 @@
 #include "PropertyEditorModule.h"
 #include "NiagaraSettings.h"
 #include "NiagaraShaderModule.h"
+
 #include "NiagaraDataInterfaceCurve.h"
 #include "NiagaraDataInterfaceVector2DCurve.h"
 #include "NiagaraDataInterfaceVectorCurve.h"
 #include "NiagaraDataInterfaceVector4Curve.h"
 #include "NiagaraDataInterfaceColorCurve.h"
+#include "DataInterface/NiagaraDataInterfaceDataChannelRead.h"
+#include "DataInterface/NiagaraDataInterfaceDataChannelWrite.h"
+
 #include "ViewModels/NiagaraScriptViewModel.h"
 #include "ViewModels/NiagaraSystemViewModel.h"
 #include "ViewModels/NiagaraEmitterHandleViewModel.h"
@@ -153,6 +157,11 @@
 #include "NiagaraSpriteRendererProperties.h"
 #include "Misc/ScopedSlowTask.h"
 
+#include "NiagaraActions.h"
+
+#include "NiagaraDataChannelDefinitions.h"
+#include "NiagaraDataChannel.h"
+
 #include UE_INLINE_GENERATED_CPP_BY_NAME(NiagaraEditorModule)
 
 IMPLEMENT_MODULE( FNiagaraEditorModule, NiagaraEditor );
@@ -832,6 +841,14 @@ void FNiagaraEditorModule::OnPreExit()
 		}
 	}
 	ClearObjectPool();
+	
+	INiagaraDataInterfaceNodeActionProvider::Unregister<UNiagaraDataInterfaceDataChannelWrite>();
+	INiagaraDataInterfaceNodeActionProvider::Unregister<UNiagaraDataInterfaceDataChannelRead>();
+
+	const FAssetRegistryModule& AssetRegistryModule = FModuleManager::LoadModuleChecked<FAssetRegistryModule>(TEXT("AssetRegistry"));
+	IAssetRegistry& AssetRegistry = AssetRegistryModule.Get();
+	AssetRegistry.OnInMemoryAssetCreated().Remove(OnAssetCreatedHandle);
+	AssetRegistry.OnInMemoryAssetDeleted().Remove(OnAssetDeletedHandle);
 }
 
 void FNiagaraEditorModule::PostGarbageCollect()
@@ -1247,6 +1264,10 @@ void FNiagaraEditorModule::StartupModule()
 #endif
 
 	GraphDataCache = MakeUnique<FNiagaraGraphDataCache>();
+	
+	//Register node action providers for data interface functions.
+	INiagaraDataInterfaceNodeActionProvider::Register<UNiagaraDataInterfaceDataChannelWrite, FNiagaraDataInterfaceNodeActionProvider_DataChannelWrite>();
+	INiagaraDataInterfaceNodeActionProvider::Register<UNiagaraDataInterfaceDataChannelRead, FNiagaraDataInterfaceNodeActionProvider_DataChannelRead>();
 }
 
 void FNiagaraEditorModule::ShutdownModule()
@@ -1429,6 +1450,11 @@ void FNiagaraEditorModule::OnPostEngineInit()
 	Debugger->Init();
 	SNiagaraDebugger::RegisterTabSpawner();
 #endif
+
+	const FAssetRegistryModule& AssetRegistryModule = FModuleManager::LoadModuleChecked<FAssetRegistryModule>(TEXT("AssetRegistry"));
+	IAssetRegistry& AssetRegistry = AssetRegistryModule.Get();
+	OnAssetCreatedHandle = AssetRegistry.OnInMemoryAssetCreated().AddRaw(this, &FNiagaraEditorModule::OnAssetCreated);
+	OnAssetDeletedHandle = AssetRegistry.OnInMemoryAssetDeleted().AddRaw(this, &FNiagaraEditorModule::OnAssetDeleted);
 }
 
 void FNiagaraEditorModule::OnDeviceProfileManagerUpdated()
@@ -1933,6 +1959,22 @@ bool FNiagaraEditorModule::DeferredDestructObjects(float InDeltaTime)
 	return false;
 }
 
+void FNiagaraEditorModule::OnAssetCreated(UObject* CreatedObject)
+{
+	if (UNiagaraDataChannelDefinitions* DataChannelDef = Cast<UNiagaraDataChannelDefinitions>(CreatedObject))
+	{
+		UNiagaraDataChannelDefinitions::OnAssetCreated(DataChannelDef);
+	}
+}
+
+void FNiagaraEditorModule::OnAssetDeleted(UObject* DeletedObject)
+{
+	if (UNiagaraDataChannelDefinitions* DataChannelDef = Cast<UNiagaraDataChannelDefinitions>(DeletedObject))
+	{
+		UNiagaraDataChannelDefinitions::OnAssetDeleted(DataChannelDef);
+	}
+}
+
 UNiagaraParameterCollection* FNiagaraEditorModule::FindCollectionForVariable(const FString& VariableName)
 {
 	auto FindCachedCollectionByPrefix = [this](const FString& Prefix)
diff --git a/Engine/Plugins/FX/Niagara/Source/NiagaraEditor/Private/NiagaraEditorStyle.cpp b/Engine/Plugins/FX/Niagara/Source/NiagaraEditor/Private/NiagaraEditorStyle.cpp
--- a/Engine/Plugins/FX/Niagara/Source/NiagaraEditor/Private/NiagaraEditorStyle.cpp
+++ b/Engine/Plugins/FX/Niagara/Source/NiagaraEditor/Private/NiagaraEditorStyle.cpp
@@ -423,6 +423,7 @@ void FNiagaraEditorStyle::InitAssetColors()
 	Set("NiagaraEditor.AssetColors.EffectType", FLinearColor(1.f, 1.f, 1.f));
 	Set("NiagaraEditor.AssetColors.SimCache", FLinearColor(0.9f, 0.3f, 0.25f));
 	Set("NiagaraEditor.AssetColors.ValidationRuleSet", FLinearColor(0.25f, 0.25f, 0.25f));
+	Set("NiagaraEditor.AssetColors.DataChannelDefinitions", FLinearColor(0.5f, 0.0f, 1.0f));
 }
 
 void FNiagaraEditorStyle::InitThumbnails()
diff --git a/Engine/Plugins/FX/Niagara/Source/NiagaraEditor/Private/NiagaraEditorUtilities.cpp b/Engine/Plugins/FX/Niagara/Source/NiagaraEditor/Private/NiagaraEditorUtilities.cpp
--- a/Engine/Plugins/FX/Niagara/Source/NiagaraEditor/Private/NiagaraEditorUtilities.cpp
+++ b/Engine/Plugins/FX/Niagara/Source/NiagaraEditor/Private/NiagaraEditorUtilities.cpp
@@ -3947,7 +3947,8 @@ int FNiagaraEditorUtilities::GetReferencedAssetCount(const FAssetData& SourceAss
 void FNiagaraEditorUtilities::GetAllowedTypes(TArray<FNiagaraTypeDefinition>& OutAllowedTypes)
 {
 	const UNiagaraEditorSettings* NiagaraEditorSettings = GetDefault<UNiagaraEditorSettings>();
-	for (const FNiagaraTypeDefinition& RegisteredType : FNiagaraTypeRegistry::GetRegisteredTypes())
+	//Only use simulation types here so as to not polute lists with LWC types in the registry. 
+	for (const FNiagaraTypeDefinition& RegisteredType : FNiagaraTypeRegistry::GetSimulationTypes())
 	{
 		if (NiagaraEditorSettings->IsAllowedTypeDefinition(RegisteredType))
 		{
diff --git a/Engine/Plugins/FX/Niagara/Source/NiagaraEditor/Private/ViewModels/NiagaraSystemViewModel.cpp b/Engine/Plugins/FX/Niagara/Source/NiagaraEditor/Private/ViewModels/NiagaraSystemViewModel.cpp
--- a/Engine/Plugins/FX/Niagara/Source/NiagaraEditor/Private/ViewModels/NiagaraSystemViewModel.cpp
+++ b/Engine/Plugins/FX/Niagara/Source/NiagaraEditor/Private/ViewModels/NiagaraSystemViewModel.cpp
@@ -1635,6 +1635,7 @@ void FNiagaraSystemViewModel::NotifyDataObjectChanged(TArray<UObject*> ChangedOb
 	}
 
 	bool bRefreshCurveSelectionViewModel = false;
+	bool bRefreshSystemDICompiledData = false;
 	if(ChangedObjects.Num() != 0)
 	{
 		for(UObject* ChangedObject : ChangedObjects)
@@ -1643,6 +1644,7 @@ void FNiagaraSystemViewModel::NotifyDataObjectChanged(TArray<UObject*> ChangedOb
 			if (ChangedDataInterface)
 			{
 				UpdateCompiledDataInterfaces(ChangedDataInterface);
+				bRefreshSystemDICompiledData = true;
 			}
 
 
@@ -1672,6 +1674,11 @@ void FNiagaraSystemViewModel::NotifyDataObjectChanged(TArray<UObject*> ChangedOb
 		bRefreshCurveSelectionViewModel = true;
 	}
 
+	if(bRefreshSystemDICompiledData)
+	{
+		System->OnCompiledDataInterfaceChanged();		
+	}
+
 	if (bRefreshCurveSelectionViewModel)
 	{
 		CurveSelectionViewModel->Refresh();
diff --git a/Engine/Plugins/FX/Niagara/Source/NiagaraEditor/Public/NagaraDataInterfaceNodeActions.h b/Engine/Plugins/FX/Niagara/Source/NiagaraEditor/Public/NagaraDataInterfaceNodeActions.h
new file mode 100644
--- /dev/null
+++ b/Engine/Plugins/FX/Niagara/Source/NiagaraEditor/Public/NagaraDataInterfaceNodeActions.h
@@ -0,0 +1 @@
+﻿#pragma once
diff --git a/Engine/Plugins/FX/Niagara/Source/NiagaraEditor/Public/NiagaraActions.h b/Engine/Plugins/FX/Niagara/Source/NiagaraEditor/Public/NiagaraActions.h
--- a/Engine/Plugins/FX/Niagara/Source/NiagaraEditor/Public/NiagaraActions.h
+++ b/Engine/Plugins/FX/Niagara/Source/NiagaraEditor/Public/NiagaraActions.h
@@ -436,3 +436,21 @@ void INiagaraDataInterfaceNodeActionProvider::GetNodeContextMenuActions(UToolMen
 		(*Provider)->GetNodeContextMenuActionsImpl(Menu, Context, Signature);
 	}
 }
+
+////////////////////////////////
+/// Actions for engine data interfaces.
+/////////////////////////////////
+
+class FNiagaraDataInterfaceNodeActionProvider_DataChannelWrite : public INiagaraDataInterfaceNodeActionProvider
+{
+public:
+
+	virtual void GetNodeContextMenuActionsImpl(UToolMenu* Menu, UGraphNodeContextMenuContext* Context, FNiagaraFunctionSignature Signature) const override;
+};
+
+class FNiagaraDataInterfaceNodeActionProvider_DataChannelRead : public INiagaraDataInterfaceNodeActionProvider
+{
+public:
+
+	virtual void GetNodeContextMenuActionsImpl(UToolMenu* Menu, UGraphNodeContextMenuContext* Context, FNiagaraFunctionSignature Signature) const override;
+};
\ No newline at end of file
diff --git a/Engine/Plugins/FX/Niagara/Source/NiagaraEditor/Public/NiagaraDataChannelDefinitionsFactoryNew.h b/Engine/Plugins/FX/Niagara/Source/NiagaraEditor/Public/NiagaraDataChannelDefinitionsFactoryNew.h
new file mode 100644
--- /dev/null
+++ b/Engine/Plugins/FX/Niagara/Source/NiagaraEditor/Public/NiagaraDataChannelDefinitionsFactoryNew.h
@@ -0,0 +1,18 @@
+// Copyright Epic Games, Inc. All Rights Reserved.
+
+#pragma once
+
+#include "CoreMinimal.h"
+#include "UObject/ObjectMacros.h"
+#include "Factories/Factory.h"
+#include "NiagaraDataChannelDefinitionsFactoryNew.generated.h"
+
+UCLASS(hidecategories = Object)
+class UNiagaraDataChannelDefinitionsFactoryNew : public UFactory
+{
+	GENERATED_UCLASS_BODY()
+
+	//~ Begin UFactory Interface
+	virtual UObject* FactoryCreateNew(UClass* Class, UObject* InParent, FName Name, EObjectFlags Flags, UObject* Context, FFeedbackContext* Warn) override;
+	//~ Begin UFactory Interface
+};
diff --git a/Engine/Plugins/FX/Niagara/Source/NiagaraEditor/Public/NiagaraEditorModule.h b/Engine/Plugins/FX/Niagara/Source/NiagaraEditor/Public/NiagaraEditorModule.h
--- a/Engine/Plugins/FX/Niagara/Source/NiagaraEditor/Public/NiagaraEditorModule.h
+++ b/Engine/Plugins/FX/Niagara/Source/NiagaraEditor/Public/NiagaraEditorModule.h
@@ -330,6 +330,9 @@ private:
 		StackIssueGenerators.Add(StructName) = Generator;
 	}
 
+	void OnAssetCreated(UObject* DeletedObject);
+	void OnAssetDeleted(UObject* DeletedObject);
+
 private:
 	TSharedPtr<FExtensibilityManager> MenuExtensibilityManager;
 	TSharedPtr<FExtensibilityManager> ToolBarExtensibilityManager;
@@ -361,6 +364,9 @@ private:
 	FDelegateHandle PreviewPlatformChangedHandle;
 	FDelegateHandle PreviewFeatureLevelChangedHandle;
 
+	FDelegateHandle OnAssetCreatedHandle;
+	FDelegateHandle OnAssetDeletedHandle;
+
 	USequencerSettings* SequencerSettings;
 
 	TSharedPtr<INiagaraEditorWidgetProvider> WidgetProvider;
diff --git a/Engine/Plugins/FX/Niagara/Source/NiagaraEditorWidgets/Private/DetailCustomizations/NiagaraDataChannelDetails.cpp b/Engine/Plugins/FX/Niagara/Source/NiagaraEditorWidgets/Private/DetailCustomizations/NiagaraDataChannelDetails.cpp
new file mode 100644
--- /dev/null
+++ b/Engine/Plugins/FX/Niagara/Source/NiagaraEditorWidgets/Private/DetailCustomizations/NiagaraDataChannelDetails.cpp
@@ -0,0 +1,40 @@
+// Copyright Epic Games, Inc. All Rights Reserved.
+
+#include "NiagaraDataChannelDetails.h"
+#include "SNiagaraNamePropertySelector.h"
+#include "NiagaraDataChannel.h"
+#include "NiagaraDataChannelDefinitions.h"
+
+void FNiagaraDataChannelReferenceDetailsCustomization::CustomizeHeader(TSharedRef<IPropertyHandle> PropertyHandle, FDetailWidgetRow& HeaderRow, IPropertyTypeCustomizationUtils& CustomizationUtils)
+{
+	ChannelHandle = PropertyHandle->GetChildHandle(TEXT("ChannelName"));
+
+	TArray<TSharedPtr<FName>> DataChannelList;
+	const TArray<UNiagaraDataChannelDefinitions*>& DataChannelDefs = UNiagaraDataChannelDefinitions::GetDataChannelDefinitions(false, true);
+	for (const TObjectPtr<UNiagaraDataChannelDefinitions>& Def : DataChannelDefs)
+	{
+		for (TObjectPtr<UNiagaraDataChannel> Channel : Def->DataChannels)
+		{
+			if(Channel)
+			{
+				DataChannelList.Add(MakeShared<FName>(Channel->GetChannelName()));
+			}
+		}
+	}
+
+	// Build Widget
+	HeaderRow
+		.NameContent()
+		[
+			ChannelHandle->CreatePropertyNameWidget()
+		]
+		.ValueContent()
+		.MinDesiredWidth(200.0f)
+		[
+			SNew(SNiagaraNamePropertySelector, ChannelHandle.ToSharedRef(), DataChannelList)
+		];	
+}
+
+void FNiagaraDataChannelReferenceDetailsCustomization::CustomizeChildren(TSharedRef<IPropertyHandle> PropertyHandle, IDetailChildrenBuilder& ChildBuilder, IPropertyTypeCustomizationUtils& CustomizationUtils)
+{
+}
diff --git a/Engine/Plugins/FX/Niagara/Source/NiagaraEditorWidgets/Private/DetailCustomizations/NiagaraDataChannelDetails.h b/Engine/Plugins/FX/Niagara/Source/NiagaraEditorWidgets/Private/DetailCustomizations/NiagaraDataChannelDetails.h
new file mode 100644
--- /dev/null
+++ b/Engine/Plugins/FX/Niagara/Source/NiagaraEditorWidgets/Private/DetailCustomizations/NiagaraDataChannelDetails.h
@@ -0,0 +1,25 @@
+// Copyright Epic Games, Inc. All Rights Reserved.
+
+#pragma once
+
+#include "CoreMinimal.h"
+#include "IPropertyTypeCustomization.h"
+#include "IDetailCustomization.h"
+
+//** Properties customization for FNiagaraDataChannelReference. */
+class FNiagaraDataChannelReferenceDetailsCustomization : public IPropertyTypeCustomization
+{
+public:
+	static TSharedRef<IPropertyTypeCustomization> MakeInstance()
+	{
+		return MakeShared<FNiagaraDataChannelReferenceDetailsCustomization>();
+	}
+
+	virtual void CustomizeHeader(TSharedRef<IPropertyHandle> PropertyHandle, FDetailWidgetRow& HeaderRow, IPropertyTypeCustomizationUtils& CustomizationUtils);
+
+	virtual void CustomizeChildren(TSharedRef<IPropertyHandle> PropertyHandle, IDetailChildrenBuilder& ChildBuilder, IPropertyTypeCustomizationUtils& CustomizationUtils);
+
+private:
+	
+	TSharedPtr<IPropertyHandle> ChannelHandle;
+};
diff --git a/Engine/Plugins/FX/Niagara/Source/NiagaraEditorWidgets/Private/NiagaraEditorWidgetsModule.cpp b/Engine/Plugins/FX/Niagara/Source/NiagaraEditorWidgets/Private/NiagaraEditorWidgetsModule.cpp
--- a/Engine/Plugins/FX/Niagara/Source/NiagaraEditorWidgets/Private/NiagaraEditorWidgetsModule.cpp
+++ b/Engine/Plugins/FX/Niagara/Source/NiagaraEditorWidgets/Private/NiagaraEditorWidgetsModule.cpp
@@ -16,6 +16,7 @@
 #include "DetailCustomizations/NiagaraMeshRendererDetails.h"
 #include "DetailCustomizations/NiagaraRibbonRendererDetails.h"
 #include "DetailCustomizations/NiagaraSpriteRendererDetails.h"
+#include "DetailCustomizations/NiagaraDataChannelDetails.h"
 #include "ViewModels/NiagaraSystemViewModel.h"
 #include "SNiagaraOverviewGraph.h"
 #include "NiagaraEditorWidgetsUtilities.h"
@@ -28,6 +29,8 @@
 #include "GraphEditorActions.h"
 #include "NiagaraEditorCommands.h"
 
+#include "NiagaraDataChannelCommon.h"
+
 IMPLEMENT_MODULE(FNiagaraEditorWidgetsModule, NiagaraEditorWidgets);
 
 FNiagaraStackCurveEditorOptions::FNiagaraStackCurveEditorOptions()
@@ -139,6 +142,10 @@ void FNiagaraEditorWidgetsModule::StartupModule()
 	PropertyModule.RegisterCustomClassLayout("NiagaraRibbonRendererProperties", FOnGetDetailCustomizationInstance::CreateStatic(&FNiagaraRibbonRendererDetails::MakeInstance));
 	PropertyModule.RegisterCustomClassLayout("NiagaraSpriteRendererProperties", FOnGetDetailCustomizationInstance::CreateStatic(&FNiagaraSpriteRendererDetails::MakeInstance));
 
+	PropertyModule.RegisterCustomPropertyTypeLayout(
+		FNiagaraDataChannelReference::StaticStruct()->GetFName(),
+		FOnGetPropertyTypeCustomizationInstance::CreateStatic(&FNiagaraDataChannelReferenceDetailsCustomization::MakeInstance));
+
 	ReinitializeStyleCommand = IConsoleManager::Get().RegisterConsoleCommand(
 		TEXT("fx.NiagaraEditorWidgets.ReinitializeStyle"),
 		TEXT("Reinitializes the style for the niagara editor widgets module.  Used in conjuction with live coding for UI tweaks.  May crash the editor if style objects are in use."),
diff --git a/Engine/Source/Runtime/VectorVM/Public/VectorVM.h b/Engine/Source/Runtime/VectorVM/Public/VectorVM.h
--- a/Engine/Source/Runtime/VectorVM/Public/VectorVM.h
+++ b/Engine/Source/Runtime/VectorVM/Public/VectorVM.h
@@ -346,6 +346,7 @@ namespace VectorVM
 		FORCEINLINE const T Get() { return *Register; }
 		FORCEINLINE T* GetDest() { return Register; }
 		FORCEINLINE void Advance() { Register += AdvanceOffset; }
+		FORCEINLINE void Advance(int32 Count) { Register += AdvanceOffset * Count; }
 		FORCEINLINE const T GetAndAdvance()
 		{
 			T* Ret = Register;
