Fixed include.

#rnx

diff --git a/Engine/Plugins/FX/Niagara/Source/Niagara/Private/NiagaraDataChannel_Global.cpp b/Engine/Plugins/FX/Niagara/Source/Niagara/Private/NiagaraDataChannel_Global.cpp
--- a/Engine/Plugins/FX/Niagara/Source/Niagara/Private/NiagaraDataChannel_Global.cpp
+++ b/Engine/Plugins/FX/Niagara/Source/Niagara/Private/NiagaraDataChannel_Global.cpp
@@ -5,7 +5,7 @@
 #include "NiagaraSystemInstance.h"
 #include "NiagaraDataChannelManager.h"
 #include "DataInterface/NiagaraDataInterfaceDataChannelWrite.h"
-#include "NiagaraGPUComputeDispatch.h"
+#include "NiagaraGPUComputeDispatchInterface.h"
 #include "NiagaraWorldManager.h"
 
 
