Missing shader files.

diff --git a/Engine/Plugins/FX/Niagara/Shaders/Private/DataChannel/NiagaraDataInterfaceDataChannelCommon.ush b/Engine/Plugins/FX/Niagara/Shaders/Private/DataChannel/NiagaraDataInterfaceDataChannelCommon.ush
new file mode 100644
--- /dev/null
+++ b/Engine/Plugins/FX/Niagara/Shaders/Private/DataChannel/NiagaraDataInterfaceDataChannelCommon.ush
@@ -0,0 +1,99 @@
+// Copyright Epic Games, Inc. All Rights Reserved.
+
+#pragma once
+
+/**
+
+Shader code common to all instances and functions of Niagara Data Channel DI.
+
+*/
+
+//TODO: Bindless buffers would allow us to read from any number of separate buffers without having to do a bunch of copying around.
+#define NIAGARA_DATACHANNEL_BINDLESSBUFFERS 0
+
+//TOOD: Add support for N multiple buffers
+//TODO: Add support for bindless buffers on platforms that support it.
+
+void NiagaraDataChannelRead_float(Buffer<float> InBuffer, uint Stride, in out uint RegisterIdx, uint ParticleIdx, out float OutValue)
+{
+	OutValue = InBuffer[RegisterIdx * Stride + ParticleIdx];
+	++RegisterIdx;
+}
+
+void NiagaraDataChannelRead_float2(Buffer<float> InBuffer, uint Stride, in out uint RegisterIdx, uint ParticleIdx, out float2 OutValue)
+{
+	NiagaraDataChannelRead_float(InBuffer, Stride, RegisterIdx, ParticleIdx, OutValue.x);
+	NiagaraDataChannelRead_float(InBuffer, Stride, RegisterIdx, ParticleIdx, OutValue.y);
+}
+
+void NiagaraDataChannelRead_float3(Buffer<float> InBuffer, uint Stride, in out uint RegisterIdx, uint ParticleIdx, out float3 OutValue)
+{
+	NiagaraDataChannelRead_float(InBuffer, Stride, RegisterIdx, ParticleIdx, OutValue.x);
+	NiagaraDataChannelRead_float(InBuffer, Stride, RegisterIdx, ParticleIdx, OutValue.y);
+	NiagaraDataChannelRead_float(InBuffer, Stride, RegisterIdx, ParticleIdx, OutValue.z);
+}
+
+void NiagaraDataChannelRead_float4(Buffer<float> InBuffer, uint Stride, in out uint RegisterIdx, uint ParticleIdx, out float4 OutValue)
+{
+	NiagaraDataChannelRead_float(InBuffer, Stride, RegisterIdx, ParticleIdx, OutValue.x);
+	NiagaraDataChannelRead_float(InBuffer, Stride, RegisterIdx, ParticleIdx, OutValue.y);
+	NiagaraDataChannelRead_float(InBuffer, Stride, RegisterIdx, ParticleIdx, OutValue.z);
+	NiagaraDataChannelRead_float(InBuffer, Stride, RegisterIdx, ParticleIdx, OutValue.w);
+}
+
+void NiagaraDataChannelRead_int(Buffer<int> InBuffer, uint Stride, in out uint RegisterIdx, uint ParticleIdx, out int OutValue)
+{
+	OutValue = InBuffer[RegisterIdx * Stride + ParticleIdx];
+	++RegisterIdx;
+}
+
+void NiagaraDataChannelRead_int2(Buffer<int> InBuffer, uint Stride, in out uint RegisterIdx, uint ParticleIdx, out int2 OutValue)
+{
+	NiagaraDataChannelRead_int(InBuffer, Stride, RegisterIdx, ParticleIdx, OutValue.x);
+	NiagaraDataChannelRead_int(InBuffer, Stride, RegisterIdx, ParticleIdx, OutValue.y);
+}
+
+void NiagaraDataChannelRead_int3(Buffer<int> InBuffer, uint Stride, in out uint RegisterIdx, uint ParticleIdx, out int3 OutValue)
+{
+	NiagaraDataChannelRead_int(InBuffer, Stride, RegisterIdx, ParticleIdx, OutValue.x);
+	NiagaraDataChannelRead_int(InBuffer, Stride, RegisterIdx, ParticleIdx, OutValue.y);
+	NiagaraDataChannelRead_int(InBuffer, Stride, RegisterIdx, ParticleIdx, OutValue.z);
+}
+
+void NiagaraDataChannelRead_int4(Buffer<int> InBuffer, uint Stride, in out uint RegisterIdx, uint ParticleIdx, out int4 OutValue)
+{
+	NiagaraDataChannelRead_int(InBuffer, Stride, RegisterIdx, ParticleIdx, OutValue.x);
+	NiagaraDataChannelRead_int(InBuffer, Stride, RegisterIdx, ParticleIdx, OutValue.y);
+	NiagaraDataChannelRead_int(InBuffer, Stride, RegisterIdx, ParticleIdx, OutValue.z);
+	NiagaraDataChannelRead_int(InBuffer, Stride, RegisterIdx, ParticleIdx, OutValue.w);
+}
+
+void NiagaraDataChannelRead_half(Buffer<half> InBuffer, uint Stride, in out uint RegisterIdx, uint ParticleIdx, out half OutValue)
+{
+	OutValue = InBuffer[RegisterIdx * Stride + ParticleIdx];
+	++RegisterIdx;
+}
+
+void NiagaraDataChannelRead_half2(Buffer<half> InBuffer, uint Stride, in out uint RegisterIdx, uint ParticleIdx, out half2 OutValue)
+{
+	NiagaraDataChannelRead_half(InBuffer, Stride, RegisterIdx, ParticleIdx, OutValue.x);
+	NiagaraDataChannelRead_half(InBuffer, Stride, RegisterIdx, ParticleIdx, OutValue.y);
+}
+
+void NiagaraDataChannelRead_half3(Buffer<half> InBuffer, uint Stride, in out uint RegisterIdx, uint ParticleIdx, out half3 OutValue)
+{
+	NiagaraDataChannelRead_half(InBuffer, Stride, RegisterIdx, ParticleIdx, OutValue.x);
+	NiagaraDataChannelRead_half(InBuffer, Stride, RegisterIdx, ParticleIdx, OutValue.y);
+	NiagaraDataChannelRead_half(InBuffer, Stride, RegisterIdx, ParticleIdx, OutValue.z);
+}
+
+void NiagaraDataChannelRead_half4(Buffer<half> InBuffer, uint Stride, in out uint RegisterIdx, uint ParticleIdx, out half4 OutValue)
+{
+	NiagaraDataChannelRead_half(InBuffer, Stride, RegisterIdx, ParticleIdx, OutValue.x);
+	NiagaraDataChannelRead_half(InBuffer, Stride, RegisterIdx, ParticleIdx, OutValue.y);
+	NiagaraDataChannelRead_half(InBuffer, Stride, RegisterIdx, ParticleIdx, OutValue.z);
+	NiagaraDataChannelRead_half(InBuffer, Stride, RegisterIdx, ParticleIdx, OutValue.w);
+}
+
+//TODO: More base type accessors.
+
diff --git a/Engine/Plugins/FX/Niagara/Shaders/Private/DataChannel/NiagaraDataInterfaceDataChannelTemplateCommon.ush b/Engine/Plugins/FX/Niagara/Shaders/Private/DataChannel/NiagaraDataInterfaceDataChannelTemplateCommon.ush
new file mode 100644
--- /dev/null
+++ b/Engine/Plugins/FX/Niagara/Shaders/Private/DataChannel/NiagaraDataInterfaceDataChannelTemplateCommon.ush
@@ -0,0 +1,45 @@
+// Copyright Epic Games, Inc. All Rights Reserved.
+
+/**
+
+Template hlsl for Data Channel DI functions to apply to their own that is shared accross all functions for the DI.
+
+*/
+
+//Offset of each parameter accessed by this script.
+//Indices for accessing this table are embeded directly into the HLSL and we generate the table on the CPU accordingly.
+//TODO: Can probably pack this into a single global buffer rather than one per DI.
+Buffer<uint> {ParameterName}_ParamOffsetTable;
+
+Buffer<float> {ParameterName}_DataFloat;
+Buffer<int> {ParameterName}_DataInt32;
+Buffer<float> {ParameterName}_DataHalf;
+
+int {ParameterName}_ParameterOffsetTableIndex;
+int {ParameterName}_Num;
+int {ParameterName}_FloatStride;
+int {ParameterName}_Int32Stride;
+int {ParameterName}_HalfStride;
+
+// Outputs the starting offset of each component in the channel data for a partiuclar parameter.
+bool GetParameterLayoutInfo_{ParameterName}(uint ParameterIndex, out uint OutFloatRegister, out uint OutInt32Register, out uint OutHalfRegister)
+{
+	int ParameterTableStart = {ParameterName}_ParameterOffsetTableIndex;
+	if(ParameterTableStart != -1)
+	{
+		const int NumComponents = 3;
+		uint ParamLayoutStart = ParameterTableStart + NumComponents * ParameterIndex;
+		OutFloatRegister = {ParameterName}_ParamOffsetTable[ParamLayoutStart];
+		OutInt32Register = {ParameterName}_ParamOffsetTable[ParamLayoutStart + 1];
+		OutHalfRegister = {ParameterName}_ParamOffsetTable[ParamLayoutStart + 2];
+		return true;
+	}
+	else
+	{
+		OutFloatRegister = -1;
+		OutInt32Register = -1;
+		OutHalfRegister = -1;
+		return false;
+	}
+}
+
diff --git a/Engine/Plugins/FX/Niagara/Shaders/Private/DataChannel/NiagaraDataInterfaceDataChannelTemplate_Append.ush b/Engine/Plugins/FX/Niagara/Shaders/Private/DataChannel/NiagaraDataInterfaceDataChannelTemplate_Append.ush
new file mode 100644
--- /dev/null
+++ b/Engine/Plugins/FX/Niagara/Shaders/Private/DataChannel/NiagaraDataInterfaceDataChannelTemplate_Append.ush
@@ -0,0 +1,14 @@
+// Copyright Epic Games, Inc. All Rights Reserved.
+
+#pragma once
+
+/**
+ * Template shader code for the Append function for the NiagaraDataInterfaceDataChannel.
+ * Contained in it's own shader file to allow reuse between different invocations with different params.
+ */
+
+void Append_{ParameterName}_UEImpureCall(bool bEmit, out bool bOutSuccess, {FunctionParameters})
+{
+	//TODO;
+	bOutSuccess = false;
+}
\ No newline at end of file
diff --git a/Engine/Plugins/FX/Niagara/Shaders/Private/DataChannel/NiagaraDataInterfaceDataChannelTemplate_Consume.ush b/Engine/Plugins/FX/Niagara/Shaders/Private/DataChannel/NiagaraDataInterfaceDataChannelTemplate_Consume.ush
new file mode 100644
--- /dev/null
+++ b/Engine/Plugins/FX/Niagara/Shaders/Private/DataChannel/NiagaraDataInterfaceDataChannelTemplate_Consume.ush
@@ -0,0 +1,14 @@
+// Copyright Epic Games, Inc. All Rights Reserved.
+
+#pragma once
+
+/**
+ * Template shader code for the Consume function for the NiagaraDataInterfaceDataChannel.
+ * Contained in it's own shader file to allow reuse between different invocations with different params.
+ */
+
+void Consume_{ParameterName}(int ElementIndex, out bool bOutSuccess, {FunctionParameters})
+{
+	//TODO;
+	bOutSuccess = false;
+}
\ No newline at end of file
diff --git a/Engine/Plugins/FX/Niagara/Shaders/Private/DataChannel/NiagaraDataInterfaceDataChannelTemplate_Read.ush b/Engine/Plugins/FX/Niagara/Shaders/Private/DataChannel/NiagaraDataInterfaceDataChannelTemplate_Read.ush
new file mode 100644
--- /dev/null
+++ b/Engine/Plugins/FX/Niagara/Shaders/Private/DataChannel/NiagaraDataInterfaceDataChannelTemplate_Read.ush
@@ -0,0 +1,30 @@
+// Copyright Epic Games, Inc. All Rights Reserved.
+
+/**
+ * Template shader code for the Read function for the NiagaraDataInterfaceDataChannel.
+ * Contained in it's own shader file to allow reuse between different invocations with different params.
+ */
+
+void Read_{ParameterName}(int ElementIndex, out bool bOutSuccess, {FunctionParameters})
+{
+	if(ElementIndex >= 0 && ElementIndex < {ParameterName}_Num)
+	{
+		bOutSuccess = true;
+
+#if NIAGARA_DATACHANNEL_MULTIBUFFERS
+		//TODO: Allow reads from N fixed number of buffers.
+#elif NIAGARA_DATACHANNEL_BINDLESSBUFFERS
+		//TODO: Allow reading from any number of buffers when bindless buffers are supported.
+#else
+		//Reads from only a single buffer supported.
+		
+		//This is generated in the DI Cpp depending on the function parameters.
+{PerFunctionParameterShaderCode}
+#endif
+	}
+	//else
+	//{
+	//	//TODO: generate defaults hlsl.
+			//	bOutSuccess = true;
+	//}
+}
\ No newline at end of file
diff --git a/Engine/Plugins/FX/Niagara/Shaders/Private/DataChannel/NiagaraDataInterfaceDataChannelTemplate_Write.ush b/Engine/Plugins/FX/Niagara/Shaders/Private/DataChannel/NiagaraDataInterfaceDataChannelTemplate_Write.ush
new file mode 100644
--- /dev/null
+++ b/Engine/Plugins/FX/Niagara/Shaders/Private/DataChannel/NiagaraDataInterfaceDataChannelTemplate_Write.ush
@@ -0,0 +1,14 @@
+// Copyright Epic Games, Inc. All Rights Reserved.
+
+#pragma once
+
+/**
+ * Template shader code for the Write function for the NiagaraDataInterfaceDataChannel.
+ * Contained in it's own shader file to allow reuse between different invocations with different params.
+ */
+
+void Write_{ParameterName}_UEImpureCall(bool bEmit, int ElementIndex, out bool bOutSuccess, {FunctionParameters})
+{
+	//TODO;
+	bOutSuccess = false;
+}
\ No newline at end of file
