Linux template call compile fix.

diff --git a/Engine/Plugins/FX/Niagara/Source/Niagara/Private/NiagaraDataChannelDefinitions.cpp b/Engine/Plugins/FX/Niagara/Source/Niagara/Private/NiagaraDataChannelDefinitions.cpp
--- a/Engine/Plugins/FX/Niagara/Source/Niagara/Private/NiagaraDataChannelDefinitions.cpp
+++ b/Engine/Plugins/FX/Niagara/Source/Niagara/Private/NiagaraDataChannelDefinitions.cpp
@@ -80,7 +80,7 @@ void UNiagaraDataChannelDefinitions::PreEditChange( FProperty* PropertyAboutToCh
 					//Re-initialize any UNiagaraSystems that are using this channel.
 					auto AddSystemsForDIUsingChannel = [&](auto* DI)
 					{
-						UNiagaraSystem* System = DI ? DI->GetTypedOuter<UNiagaraSystem>() : nullptr;
+						UNiagaraSystem* System = DI ? DI->template GetTypedOuter<UNiagaraSystem>() : nullptr;
 						if (DI && System && DI->Channel.ChannelName == DataChannel->GetChannelName())
 						{
 							SysUpdateContext.Add(System, true);
