linux fix

diff --git a/Engine/Plugins/FX/Niagara/Source/NiagaraEditorWidgets/Private/DetailCustomizations/NiagaraDataChannelDetails.cpp b/Engine/Plugins/FX/Niagara/Source/NiagaraEditorWidgets/Private/DetailCustomizations/NiagaraDataChannelDetails.cpp
--- a/Engine/Plugins/FX/Niagara/Source/NiagaraEditorWidgets/Private/DetailCustomizations/NiagaraDataChannelDetails.cpp
+++ b/Engine/Plugins/FX/Niagara/Source/NiagaraEditorWidgets/Private/DetailCustomizations/NiagaraDataChannelDetails.cpp
@@ -11,7 +11,7 @@ void FNiagaraDataChannelReferenceDetailsCustomization::CustomizeHeader(TSharedRe
 
 	TArray<TSharedPtr<FName>> DataChannelList;
 	const TArray<UNiagaraDataChannelDefinitions*>& DataChannelDefs = UNiagaraDataChannelDefinitions::GetDataChannelDefinitions(false, true);
-	for (const TObjectPtr<UNiagaraDataChannelDefinitions>& Def : DataChannelDefs)
+	for (const UNiagaraDataChannelDefinitions* Def : DataChannelDefs)
 	{
 		for (TObjectPtr<UNiagaraDataChannel> Channel : Def->DataChannels)
 		{
