linux fix

diff --git a/Engine/Plugins/FX/Niagara/Source/Niagara/Private/NiagaraDataChannelDefinitions.cpp b/Engine/Plugins/FX/Niagara/Source/Niagara/Private/NiagaraDataChannelDefinitions.cpp
--- a/Engine/Plugins/FX/Niagara/Source/Niagara/Private/NiagaraDataChannelDefinitions.cpp
+++ b/Engine/Plugins/FX/Niagara/Source/Niagara/Private/NiagaraDataChannelDefinitions.cpp
@@ -115,7 +115,7 @@ void UNiagaraDataChannelDefinitions::PostEditChangeProperty(FPropertyChangedEven
 					//Re-initialize any UNiagaraSystems that are using this channel.
 					auto AddSystemsForDIUsingChannel = [&](auto* DI)
 					{
-						UNiagaraSystem* System = DI ? DI->GetTypedOuter<UNiagaraSystem>() : nullptr;
+						UNiagaraSystem* System = DI ? DI->template GetTypedOuter<UNiagaraSystem>() : nullptr;
 						if (DI && System && DI->Channel.ChannelName == DataChannel->GetChannelName())
 						{
 							SysUpdateContext.Add(System, true);
diff --git a/Engine/Plugins/FX/Niagara/Source/Niagara/Private/NiagaraDataChannel_Global.cpp b/Engine/Plugins/FX/Niagara/Source/Niagara/Private/NiagaraDataChannel_Global.cpp
--- a/Engine/Plugins/FX/Niagara/Source/Niagara/Private/NiagaraDataChannel_Global.cpp
+++ b/Engine/Plugins/FX/Niagara/Source/Niagara/Private/NiagaraDataChannel_Global.cpp
@@ -5,7 +5,7 @@
 #include "NiagaraSystemInstance.h"
 #include "NiagaraDataChannelManager.h"
 #include "DataInterface/NiagaraDataInterfaceDataChannelWrite.h"
-#include "NiagaraGPUComputeDispatchInterface.h"
+#include "NiagaraGpuComputeDispatchInterface.h"
 #include "NiagaraWorldManager.h"
 
 
