- Update RendererEnabledBinding tooltip to make it more clear

#rnx

diff --git a/Engine/Plugins/FX/Niagara/Source/Niagara/Public/NiagaraRendererProperties.h b/Engine/Plugins/FX/Niagara/Source/Niagara/Public/NiagaraRendererProperties.h
--- a/Engine/Plugins/FX/Niagara/Source/Niagara/Public/NiagaraRendererProperties.h
+++ b/Engine/Plugins/FX/Niagara/Source/Niagara/Public/NiagaraRendererProperties.h
@@ -396,7 +396,11 @@ public:
 	UPROPERTY(EditAnywhere, Category = "Rendering")
 	ENiagaraRendererMotionVectorSetting MotionVectorSetting;
 
-	/** Optional bool binding to dynamically enable / disable the renderer. */
+	/**
+	Binding to control if the renderer is enabled or disabled.
+	When disabled the renderer does not generate or render any particle data.
+	When disabled via a static bool the renderer will be removed in cooked content.
+	*/
 	UPROPERTY(EditAnywhere, Category = "Bindings")
 	FNiagaraVariableAttributeBinding RendererEnabledBinding;
 
