- Fix CPU version of DrawSphere
- Was broken when hemisphere code was added as test was inverted

#rnx

diff --git a/Engine/Plugins/FX/Niagara/Source/Niagara/Private/DataInterface/NiagaraDataInterfaceDebugDraw.cpp b/Engine/Plugins/FX/Niagara/Source/Niagara/Private/DataInterface/NiagaraDataInterfaceDebugDraw.cpp
--- a/Engine/Plugins/FX/Niagara/Source/Niagara/Private/DataInterface/NiagaraDataInterfaceDebugDraw.cpp
+++ b/Engine/Plugins/FX/Niagara/Source/Niagara/Private/DataInterface/NiagaraDataInterfaceDebugDraw.cpp
@@ -77,7 +77,7 @@ struct FNDIDebugDrawInstanceData_GameThread
 		float ux = 0.0f;
 		float SinX0 = FMath::Sin(ux);
 		float CosX0 = FMath::Cos(ux);
-		FVector3f LerpVector = FVector3f(bHemiX ? 1.0f : 0.0f, bHemiY ? 1.0f : 0.0f, bHemiZ ? 1.0f : 0.0f);
+		FVector3f LerpVector = FVector3f(bHemiX ? 0.0f : 1.0f, bHemiY ? 0.0f : 1.0f, bHemiZ ? 0.0f : 1.0f);
 		FTransform3f LocalToWorld;
 		LocalToWorld.SetComponents(WorldRotate, FVector3f::ZeroVector, WorldScale);
 
