Fixed some LOCTEXT warnings.

diff --git a/Engine/Plugins/FX/Niagara/Source/NiagaraEditor/Private/NiagaraActions.cpp b/Engine/Plugins/FX/Niagara/Source/NiagaraEditor/Private/NiagaraActions.cpp
--- a/Engine/Plugins/FX/Niagara/Source/NiagaraEditor/Private/NiagaraActions.cpp
+++ b/Engine/Plugins/FX/Niagara/Source/NiagaraEditor/Private/NiagaraActions.cpp
@@ -685,8 +685,18 @@ void INiagaraDataInterfaceNodeActionProvider::GetNodeContextMenuActions(UClass*
 
 ////////////////////////
 
+namespace NiagaraActionsLocal
+{
+	FText InitForDataChannelHeaderText = LOCTEXT("DataChannelsHeader", "Data Channels");
+	FText InitForDataChannelSectionText = LOCTEXT("DataChannelsSection","Data Channels");
+	FText InitForDataChannelMenuText = LOCTEXT("InitForDataChannelMenu", "Init For Data Channel...");
+	FText InitForDataChannelMenuTooltipText = LOCTEXT("InitForDataChannelTooltip", "Initializes this node to write to all members of a given data channel.");
+}
+
 void FNiagaraDataInterfaceNodeActionProvider_DataChannelWrite::GetNodeContextMenuActionsImpl(UToolMenu* Menu, UGraphNodeContextMenuContext* Context, FNiagaraFunctionSignature Signature) const
 {
+	using namespace NiagaraActionsLocal;
+
 	//For all functions except "Num", add a context menu to initialized to a specific data channel.
 	if (Signature.Name == TEXT("Num"))
 	{
@@ -699,13 +709,13 @@ void FNiagaraDataInterfaceNodeActionProvider_DataChannelWrite::GetNodeContextMen
 		return;
 	}
 
-	FToolMenuSection& Section = Menu->AddSection("DataChannelWrite", LOCTEXT("DataChannelWriteHeader", "Data Channel Write"));
+	FToolMenuSection& Section = Menu->AddSection("DataChannelWrite", InitForDataChannelHeaderText);
 
 	TWeakObjectPtr<UNiagaraNodeFunctionCall> WeakFuncNode = CastChecked<UNiagaraNodeFunctionCall>(Context->Node);
 
 	auto CreateNodeContextMenu = [WeakFuncNode](UToolMenu* InNewToolMenu)
 	{
-		FToolMenuSection& SubSection = InNewToolMenu->AddSection("InitForDataChannelSection", LOCTEXT("DataChannels", "Data Channels"));
+		FToolMenuSection& SubSection = InNewToolMenu->AddSection("InitForDataChannelSection", InitForDataChannelSectionText);
 
 		const TArray<UNiagaraDataChannelDefinitions*>& DataChannelDefs = UNiagaraDataChannelDefinitions::GetDataChannelDefinitions(false, true);
 		for (const UNiagaraDataChannelDefinitions* Def : DataChannelDefs)
@@ -734,13 +744,14 @@ void FNiagaraDataInterfaceNodeActionProvider_DataChannelWrite::GetNodeContextMen
 		}
 	};
 
-	Section.AddSubMenu("InitForDataChannelMenu", LOCTEXT("InitForDataChannel", "Init For Data Channel..."), LOCTEXT("InitForDataChannelTooltip", "Initializes this node to write to all members of a given data channel."),
-		FNewToolMenuDelegate::CreateLambda(CreateNodeContextMenu));
+	Section.AddSubMenu("InitForDataChannelMenu", InitForDataChannelMenuText, InitForDataChannelMenuTooltipText,	FNewToolMenuDelegate::CreateLambda(CreateNodeContextMenu));
 }
 
 
 void FNiagaraDataInterfaceNodeActionProvider_DataChannelRead::GetNodeContextMenuActionsImpl(UToolMenu* Menu, UGraphNodeContextMenuContext* Context, FNiagaraFunctionSignature Signature) const
 {
+	using namespace NiagaraActionsLocal;
+
 	//For all functions except "Num", add a context menu to initialized to a specific data channel.
 	if (Signature.Name == TEXT("Num"))
 	{
@@ -753,13 +764,13 @@ void FNiagaraDataInterfaceNodeActionProvider_DataChannelRead::GetNodeContextMenu
 		return;
 	}
 
-	FToolMenuSection& Section = Menu->AddSection("DataChannelRead", LOCTEXT("DataChannelReadHeader", "Data Channel Read"));
+	FToolMenuSection& Section = Menu->AddSection("DataChannelRead", InitForDataChannelHeaderText);
 
 	TWeakObjectPtr<UNiagaraNodeFunctionCall> WeakFuncNode = CastChecked<UNiagaraNodeFunctionCall>(Context->Node);
 
 	auto CreateNodeContextMenu = [WeakFuncNode](UToolMenu* InNewToolMenu)
 	{
-		FToolMenuSection& SubSection = InNewToolMenu->AddSection("InitForDataChannelSection", LOCTEXT("DataChannels", "Data Channels"));
+		FToolMenuSection& SubSection = InNewToolMenu->AddSection("InitForDataChannelSection", InitForDataChannelSectionText);
 
 		const TArray<UNiagaraDataChannelDefinitions*>& DataChannelDefs = UNiagaraDataChannelDefinitions::GetDataChannelDefinitions(false, true);
 		for (const UNiagaraDataChannelDefinitions* Def : DataChannelDefs)
@@ -790,8 +801,7 @@ void FNiagaraDataInterfaceNodeActionProvider_DataChannelRead::GetNodeContextMenu
 		}
 	};
 
-	Section.AddSubMenu("InitForDataChannelMenu", LOCTEXT("InitForDataChannel", "Init For Data Channel..."),	LOCTEXT("InitForDataChannelTooltip", "Initializes this node to read to all members of a given data channel."),
-		FNewToolMenuDelegate::CreateLambda(CreateNodeContextMenu));
+	Section.AddSubMenu("InitForDataChannelMenu", InitForDataChannelMenuText, InitForDataChannelMenuTooltipText,	FNewToolMenuDelegate::CreateLambda(CreateNodeContextMenu));
 }
 
 
