Static analysis fix

diff --git a/Engine/Plugins/FX/Niagara/Source/Niagara/Private/DataInterface/NiagaraDataInterfaceDataChannelWrite.cpp b/Engine/Plugins/FX/Niagara/Source/Niagara/Private/DataInterface/NiagaraDataInterfaceDataChannelWrite.cpp
--- a/Engine/Plugins/FX/Niagara/Source/Niagara/Private/DataInterface/NiagaraDataInterfaceDataChannelWrite.cpp
+++ b/Engine/Plugins/FX/Niagara/Source/Niagara/Private/DataInterface/NiagaraDataInterfaceDataChannelWrite.cpp
@@ -165,17 +165,19 @@ struct FNDIDataChannelWriteInstanceData
 			}
 		}
 
-		FNiagaraDataBuffer* CurrBuff = Data->GetCurrentData();
-		if (bNeedsLastFrameData)
-		{
-			CurrBuff->AddReadRef();//Ensure we keep the previous data when grabbing a new dest buffer.
-		}
+		if(FNiagaraDataBuffer* CurrBuff = Data->GetCurrentData())
+		{	
+			if (bNeedsLastFrameData)
+			{
+				CurrBuff->AddReadRef();//Ensure we keep the previous data when grabbing a new dest buffer.
+			}
 
-		FNiagaraDataBuffer& DestBuff = Data->BeginSimulate(true);
+			FNiagaraDataBuffer& DestBuff = Data->BeginSimulate(true);
 
-		if (bNeedsLastFrameData && CurrBuff)
-		{
-			CurrBuff->ReleaseReadRef();
+			if (bNeedsLastFrameData && CurrBuff)
+			{
+				CurrBuff->ReleaseReadRef();
+			}
 		}
 
 		if (Interface->AllocationMode == ENiagaraDataChannelAllocationMode::Static)
