Fixing issue with LWC->SWC down conversion on load for some niagara structs

diff --git a/Engine/Plugins/FX/Niagara/Source/Niagara/Private/NiagaraModule.cpp b/Engine/Plugins/FX/Niagara/Source/Niagara/Private/NiagaraModule.cpp
--- a/Engine/Plugins/FX/Niagara/Source/Niagara/Private/NiagaraModule.cpp
+++ b/Engine/Plugins/FX/Niagara/Source/Niagara/Private/NiagaraModule.cpp
@@ -1423,7 +1423,7 @@ bool FNiagaraTypeDefinition::Serialize(FArchive& Ar)
 				else if((Flags & TF_AllowLWC) == 0)
 				{	
 					//If the type doesn't allow LWC we should convert down to SWC.
-					ClassStructOrEnum = FNiagaraTypeHelper::FindNiagaraFriendlyTopLevelStruct(static_cast<UScriptStruct*>(ClassStructOrEnum), ENiagaraStructConversion::Simulation);
+					ClassStructOrEnum = FNiagaraTypeHelper::FindNiagaraFriendlyTopLevelStruct(static_cast<UScriptStruct*>(ClassStructOrEnum), ENiagaraStructConversion::UserFacing);
 				}
 	#endif
 			}
