Readded support for copy pasting DIs in the User Parameter panel. Stopped working going from stack copy paste logic to default details panel logic

diff --git a/Engine/Plugins/FX/Niagara/Source/NiagaraEditor/Private/Customizations/NiagaraComponentDetails.cpp b/Engine/Plugins/FX/Niagara/Source/NiagaraEditor/Private/Customizations/NiagaraComponentDetails.cpp
--- a/Engine/Plugins/FX/Niagara/Source/NiagaraEditor/Private/Customizations/NiagaraComponentDetails.cpp
+++ b/Engine/Plugins/FX/Niagara/Source/NiagaraEditor/Private/Customizations/NiagaraComponentDetails.cpp
@@ -28,6 +28,7 @@
 #include "GameDelegates.h"
 #include "NiagaraEditorStyle.h"
 #include "IDetailChildrenBuilder.h"
+#include "NiagaraClipboard.h"
 #include "NiagaraConstants.h"
 #include "NiagaraScriptVariable.h"
 #include "NiagaraSystemEditorData.h"
@@ -400,7 +401,7 @@ void FNiagaraUserParameterNodeBuilder::GenerateRowForUserParameter(IDetailChildr
 			.CreateCategoryNodes(false);
 	
 		Row = ChildrenBuilder.AddExternalObjectProperty(Objects, NAME_None, Params);
-
+		
 		CustomValueWidget =
 			SNew(STextBlock)
 			.TextStyle(FNiagaraEditorStyle::Get(), "NiagaraEditor.ParameterText")
@@ -478,6 +479,36 @@ void FNiagaraUserParameterNodeBuilder::GenerateRowForUserParameter(IDetailChildr
 	Row->DisplayName(FText::FromName(ChoppedUserParameter.GetName()));
 	
 	FDetailWidgetRow& CustomWidget = Row->CustomWidget(true);
+
+	if(UserParameter.GetType().IsDataInterface())
+	{
+		FUIAction CopyAction;
+		CopyAction.ExecuteAction = FExecuteAction::CreateLambda([=]()
+		{
+			UNiagaraClipboardContent* ClipboardContent = UNiagaraClipboardContent::Create();
+			const UNiagaraClipboardFunctionInput* Input = UNiagaraClipboardFunctionInput::CreateDataValue(ClipboardContent, UserParameter.GetName(), UserParameter.GetType().GetClass(), TOptional<bool>(), ParameterProxy->Value().GetDataInterface());
+			ClipboardContent->FunctionInputs.Add(Input);
+			
+			FNiagaraEditorModule::Get().GetClipboard().SetClipboardContent(ClipboardContent);
+		});
+		CustomWidget.CopyAction(CopyAction);
+
+		FUIAction PasteAction;
+		PasteAction.ExecuteAction = FExecuteAction::CreateLambda([=]()
+		{
+			const UNiagaraClipboardContent* ClipboardContent = FNiagaraEditorModule::Get().GetClipboard().GetClipboardContent();
+
+			if(ClipboardContent->FunctionInputs.Num() == 1 && ClipboardContent->FunctionInputs[0]->InputType.IsDataInterface())
+			{
+				UNiagaraDataInterface* CopiedDI = ClipboardContent->FunctionInputs[0]->Data;
+				if(CopiedDI->IsA(ParameterProxy->Value().GetDataInterface()->GetClass()))
+				{
+					CopiedDI->CopyTo(ParameterProxy->Value().GetDataInterface());
+				}
+			}			
+		});
+		CustomWidget.PasteAction(PasteAction);
+	}
 	
 	AddCustomMenuActionsForParameter(CustomWidget, UserParameter);
 	
