- Move common IsSupportedVariableForBinding into base renderer
- Ensure we handle IsRendererEnabled consistently

diff --git a/Engine/Plugins/FX/Niagara/Source/Niagara/Private/NiagaraComponentRendererProperties.cpp b/Engine/Plugins/FX/Niagara/Source/Niagara/Private/NiagaraComponentRendererProperties.cpp
--- a/Engine/Plugins/FX/Niagara/Source/Niagara/Private/NiagaraComponentRendererProperties.cpp
+++ b/Engine/Plugins/FX/Niagara/Source/Niagara/Private/NiagaraComponentRendererProperties.cpp
@@ -352,16 +352,14 @@ void UNiagaraComponentRendererProperties::CacheFromCompiledData(const FNiagaraDa
 #if WITH_EDITORONLY_DATA
 bool UNiagaraComponentRendererProperties::IsSupportedVariableForBinding(const FNiagaraVariableBase& InSourceForBinding, const FName& InTargetBindingName) const
 {
-	if ( InTargetBindingName == GET_MEMBER_NAME_CHECKED(UNiagaraComponentRendererProperties, RendererEnabledBinding) )
+	if (InTargetBindingName == GET_MEMBER_NAME_CHECKED(UNiagaraRendererProperties, RendererEnabledBinding))
 	{
-		if (
+		return
 			InSourceForBinding.IsInNameSpace(FNiagaraConstants::UserNamespace) ||
 			InSourceForBinding.IsInNameSpace(FNiagaraConstants::SystemNamespace) ||
-			InSourceForBinding.IsInNameSpace(FNiagaraConstants::EmitterNamespace))
-		{
-			return true;
-		}
+			InSourceForBinding.IsInNameSpace(FNiagaraConstants::EmitterNamespace);
 	}
+
 	return false;
 }
 #endif
diff --git a/Engine/Plugins/FX/Niagara/Source/Niagara/Private/NiagaraDecalRendererProperties.cpp b/Engine/Plugins/FX/Niagara/Source/Niagara/Private/NiagaraDecalRendererProperties.cpp
--- a/Engine/Plugins/FX/Niagara/Source/Niagara/Private/NiagaraDecalRendererProperties.cpp
+++ b/Engine/Plugins/FX/Niagara/Source/Niagara/Private/NiagaraDecalRendererProperties.cpp
@@ -243,21 +243,6 @@ bool UNiagaraDecalRendererProperties::PopulateRequiredBindings(FNiagaraParameter
 }
 
 #if WITH_EDITORONLY_DATA
-bool UNiagaraDecalRendererProperties::IsSupportedVariableForBinding(const FNiagaraVariableBase& InSourceForBinding, const FName& InTargetBindingName) const
-{
-	switch (SourceMode)
-	{
-		case ENiagaraRendererSourceDataMode::Particles:
-			return InSourceForBinding.IsInNameSpace(FNiagaraConstants::ParticleAttributeNamespaceString);
-
-		default:
-			return
-				InSourceForBinding.IsInNameSpace(FNiagaraConstants::UserNamespaceString) ||
-				InSourceForBinding.IsInNameSpace(FNiagaraConstants::SystemNamespaceString) ||
-				InSourceForBinding.IsInNameSpace(FNiagaraConstants::EmitterNamespaceString);
-	}
-}
-
 const TArray<FNiagaraVariable>& UNiagaraDecalRendererProperties::GetOptionalAttributes()
 {
 	using namespace NiagaraDecalRendererPropertiesLocal;
diff --git a/Engine/Plugins/FX/Niagara/Source/Niagara/Private/NiagaraLightRendererProperties.cpp b/Engine/Plugins/FX/Niagara/Source/Niagara/Private/NiagaraLightRendererProperties.cpp
--- a/Engine/Plugins/FX/Niagara/Source/Niagara/Private/NiagaraLightRendererProperties.cpp
+++ b/Engine/Plugins/FX/Niagara/Source/Niagara/Private/NiagaraLightRendererProperties.cpp
@@ -216,21 +216,6 @@ bool UNiagaraLightRendererProperties::PopulateRequiredBindings(FNiagaraParameter
 }
 
 #if WITH_EDITORONLY_DATA
-bool UNiagaraLightRendererProperties::IsSupportedVariableForBinding(const FNiagaraVariableBase& InSourceForBinding, const FName& InTargetBindingName) const
-{
-	switch (SourceMode)
-	{
-		case ENiagaraRendererSourceDataMode::Particles:
-			return InSourceForBinding.IsInNameSpace(FNiagaraConstants::ParticleAttributeNamespaceString);
-
-		default:
-			return
-				InSourceForBinding.IsInNameSpace(FNiagaraConstants::UserNamespaceString) ||
-				InSourceForBinding.IsInNameSpace(FNiagaraConstants::SystemNamespaceString) ||
-				InSourceForBinding.IsInNameSpace(FNiagaraConstants::EmitterNamespaceString);
-	}
-}
-
 const TArray<FNiagaraVariable>& UNiagaraLightRendererProperties::GetOptionalAttributes()
 {
 	static TArray<FNiagaraVariable> Attrs;
diff --git a/Engine/Plugins/FX/Niagara/Source/Niagara/Private/NiagaraMeshRendererProperties.cpp b/Engine/Plugins/FX/Niagara/Source/Niagara/Private/NiagaraMeshRendererProperties.cpp
--- a/Engine/Plugins/FX/Niagara/Source/Niagara/Private/NiagaraMeshRendererProperties.cpp
+++ b/Engine/Plugins/FX/Niagara/Source/Niagara/Private/NiagaraMeshRendererProperties.cpp
@@ -558,20 +558,6 @@ void UNiagaraMeshRendererProperties::UpdateMICs()
 #endif
 }
 
-#if WITH_EDITORONLY_DATA
-bool UNiagaraMeshRendererProperties::IsSupportedVariableForBinding(const FNiagaraVariableBase& InSourceForBinding, const FName& InTargetBindingName) const
-{
-	if ((SourceMode == ENiagaraRendererSourceDataMode::Particles && InSourceForBinding.IsInNameSpace(FNiagaraConstants::ParticleAttributeNamespaceString)) ||
-		InSourceForBinding.IsInNameSpace(FNiagaraConstants::UserNamespaceString) ||
-		InSourceForBinding.IsInNameSpace(FNiagaraConstants::SystemNamespaceString) ||
-		InSourceForBinding.IsInNameSpace(FNiagaraConstants::EmitterNamespaceString))
-	{
-		return true;
-	}
-	return false;
-}
-#endif
-
 void UNiagaraMeshRendererProperties::ApplyMaterialOverrides(const FNiagaraEmitterInstance* EmitterInstance, TArray<UMaterialInterface*>& InOutMaterials) const
 {
 	if (bOverrideMaterials)
@@ -677,7 +663,7 @@ void UNiagaraMeshRendererProperties::GetStreamingMeshInfo(const FBoxSphereBounds
 				OwnerBounds.Origin + MeshBounds.Origin,
 				MeshBounds.BoxExtent * MeshProperties.Scale,
 				MeshBounds.SphereRadius * MeshProperties.Scale.GetMax());
-			const float MeshTexelFactor = MeshBounds.SphereRadius * 2.0f;
+			const float MeshTexelFactor = float(MeshBounds.SphereRadius * 2.0);
 
 			new (OutStreamingRenderAssets) FStreamingRenderAssetPrimitiveInfo(StaticMesh, StreamingBounds, MeshTexelFactor);
 		}
@@ -946,8 +932,8 @@ void UNiagaraMeshRendererProperties::PreEditChange(class FProperty* PropertyThat
 
 void UNiagaraMeshRendererProperties::PostEditChangeProperty(FPropertyChangedEvent& PropertyChangedEvent)
 {
-	SubImageSize.X = FMath::Max<float>(SubImageSize.X, 1.f);
-	SubImageSize.Y = FMath::Max<float>(SubImageSize.Y, 1.f);
+	SubImageSize.X = FMath::Max(SubImageSize.X, 1.0);
+	SubImageSize.Y = FMath::Max(SubImageSize.Y, 1.0);
 
 	const bool bIsRedirect = PropertyChangedEvent.ChangeType == EPropertyChangeType::Redirected;
 	const bool bRebuildMeshList = ChangeRequiresMeshListRebuild(PropertyChangedEvent.Property);
diff --git a/Engine/Plugins/FX/Niagara/Source/Niagara/Private/NiagaraRendererDecals.cpp b/Engine/Plugins/FX/Niagara/Source/Niagara/Private/NiagaraRendererDecals.cpp
--- a/Engine/Plugins/FX/Niagara/Source/Niagara/Private/NiagaraRendererDecals.cpp
+++ b/Engine/Plugins/FX/Niagara/Source/Niagara/Private/NiagaraRendererDecals.cpp
@@ -133,8 +133,9 @@ FNiagaraDynamicDataBase* FNiagaraRendererDecals::GenerateDynamicData(const FNiag
 	const FNiagaraDataSet& DataSet = Emitter->GetData();
 	const FNiagaraDataBuffer* DataToRender = DataSet.GetCurrentData();
 	FNiagaraSystemInstance* SystemInstance = Emitter->GetParentSystemInstance();
-	if (DataToRender == nullptr || SystemInstance == nullptr)
+	if (DataToRender == nullptr || SystemInstance == nullptr || !IsRendererEnabled(InProperties, Emitter))
 	{
+		ReleaseAllDecals();
 		return nullptr;
 	}
 
diff --git a/Engine/Plugins/FX/Niagara/Source/Niagara/Private/NiagaraRendererLights.cpp b/Engine/Plugins/FX/Niagara/Source/Niagara/Private/NiagaraRendererLights.cpp
--- a/Engine/Plugins/FX/Niagara/Source/Niagara/Private/NiagaraRendererLights.cpp
+++ b/Engine/Plugins/FX/Niagara/Source/Niagara/Private/NiagaraRendererLights.cpp
@@ -63,6 +63,11 @@ FNiagaraDynamicDataBase* FNiagaraRendererLights::GenerateDynamicData(const FNiag
 		return nullptr;
 	}
 
+	if (!IsRendererEnabled(InProperties, Emitter))
+	{
+		return nullptr;
+	}
+
 	SCOPE_CYCLE_COUNTER(STAT_NiagaraGenLights);
 
 	//Bail if we don't have the required attributes to render this emitter.
diff --git a/Engine/Plugins/FX/Niagara/Source/Niagara/Private/NiagaraRendererProperties.cpp b/Engine/Plugins/FX/Niagara/Source/Niagara/Private/NiagaraRendererProperties.cpp
--- a/Engine/Plugins/FX/Niagara/Source/Niagara/Private/NiagaraRendererProperties.cpp
+++ b/Engine/Plugins/FX/Niagara/Source/Niagara/Private/NiagaraRendererProperties.cpp
@@ -348,11 +348,26 @@ void FNiagaraRendererMaterialParameters::GetFeedback(TArrayView<UMaterialInterfa
 #if WITH_EDITORONLY_DATA
 bool UNiagaraRendererProperties::IsSupportedVariableForBinding(const FNiagaraVariableBase& InSourceForBinding, const FName& InTargetBindingName) const
 {
-	if (InSourceForBinding.IsInNameSpace(FNiagaraConstants::ParticleAttributeNamespaceString))
+	if (InTargetBindingName == GET_MEMBER_NAME_CHECKED(UNiagaraRendererProperties, RendererEnabledBinding))
 	{
-		return true;
+		return
+			InSourceForBinding.IsInNameSpace(FNiagaraConstants::UserNamespace) ||
+			InSourceForBinding.IsInNameSpace(FNiagaraConstants::SystemNamespace) ||
+			InSourceForBinding.IsInNameSpace(FNiagaraConstants::EmitterNamespace);
 	}
-	return false;
+
+	switch (GetCurrentSourceMode())
+	{
+		case ENiagaraRendererSourceDataMode::Particles:
+			return InSourceForBinding.IsInNameSpace(FNiagaraConstants::ParticleAttributeNamespaceString);
+
+		default:
+			return
+				InSourceForBinding.IsInNameSpace(FNiagaraConstants::UserNamespaceString) ||
+				InSourceForBinding.IsInNameSpace(FNiagaraConstants::SystemNamespaceString) ||
+				InSourceForBinding.IsInNameSpace(FNiagaraConstants::EmitterNamespaceString);
+	}
+
 }
 
 void UNiagaraRendererProperties::RenameEmitter(const FName& InOldName, const UNiagaraEmitter* InRenamedEmitter)
diff --git a/Engine/Plugins/FX/Niagara/Source/Niagara/Private/NiagaraRibbonRendererProperties.cpp b/Engine/Plugins/FX/Niagara/Source/Niagara/Private/NiagaraRibbonRendererProperties.cpp
--- a/Engine/Plugins/FX/Niagara/Source/Niagara/Private/NiagaraRibbonRendererProperties.cpp
+++ b/Engine/Plugins/FX/Niagara/Source/Niagara/Private/NiagaraRibbonRendererProperties.cpp
@@ -430,18 +430,6 @@ void UNiagaraRibbonRendererProperties::UpdateMICs()
 
 #if WITH_EDITORONLY_DATA
 
-bool UNiagaraRibbonRendererProperties::IsSupportedVariableForBinding(const FNiagaraVariableBase& InSourceForBinding, const FName& InTargetBindingName) const
-{
-	if (InSourceForBinding.IsInNameSpace(FNiagaraConstants::ParticleAttributeNamespaceString) ||
-		InSourceForBinding.IsInNameSpace(FNiagaraConstants::UserNamespaceString) ||
-		InSourceForBinding.IsInNameSpace(FNiagaraConstants::SystemNamespaceString) ||
-		InSourceForBinding.IsInNameSpace(FNiagaraConstants::EmitterNamespaceString))
-	{
-		return true;
-	}
-	return false;
-}
-
 void UNiagaraRibbonRendererProperties::PostEditChangeProperty(struct FPropertyChangedEvent& PropertyChangedEvent)
 {
 	Super::PostEditChangeProperty(PropertyChangedEvent);
diff --git a/Engine/Plugins/FX/Niagara/Source/Niagara/Private/NiagaraSpriteRendererProperties.cpp b/Engine/Plugins/FX/Niagara/Source/Niagara/Private/NiagaraSpriteRendererProperties.cpp
--- a/Engine/Plugins/FX/Niagara/Source/Niagara/Private/NiagaraSpriteRendererProperties.cpp
+++ b/Engine/Plugins/FX/Niagara/Source/Niagara/Private/NiagaraSpriteRendererProperties.cpp
@@ -418,34 +418,6 @@ void UNiagaraSpriteRendererProperties::UpdateMICs()
 }
 
 #if WITH_EDITORONLY_DATA
-
-bool UNiagaraSpriteRendererProperties::IsSupportedVariableForBinding(const FNiagaraVariableBase& InSourceForBinding, const FName& InTargetBindingName) const
-{
-	if ((SourceMode == ENiagaraRendererSourceDataMode::Particles && InSourceForBinding.IsInNameSpace(FNiagaraConstants::ParticleAttributeNamespaceString)) ||
-		InSourceForBinding.IsInNameSpace(FNiagaraConstants::UserNamespaceString) ||
-		InSourceForBinding.IsInNameSpace(FNiagaraConstants::SystemNamespaceString) ||
-		InSourceForBinding.IsInNameSpace(FNiagaraConstants::EmitterNamespaceString))
-	{
-		return true;
-	}
-	return false;
-}
-
-
-//#if WITH_EDITOR
-//void UNiagaraSpriteRendererProperties::PostEditUndo()
-//{
-//	Super::PostEditUndo();
-//	UpdateSourceModeDerivates(SourceMode);
-//}
-//
-//void UNiagaraSpriteRendererProperties::PostEditUndo(TSharedPtr<ITransactionObjectAnnotation> TransactionAnnotation)
-//{
-//	Super::PostEditUndo(TransactionAnnotation);
-//	UpdateSourceModeDerivates(SourceMode);
-//}
-//#endif
-
 void UNiagaraSpriteRendererProperties::PostEditChangeProperty(struct FPropertyChangedEvent& PropertyChangedEvent)
 {
 	const FName PropertyName = PropertyChangedEvent.GetPropertyName();
diff --git a/Engine/Plugins/FX/Niagara/Source/Niagara/Public/NiagaraDecalRendererProperties.h b/Engine/Plugins/FX/Niagara/Source/Niagara/Public/NiagaraDecalRendererProperties.h
--- a/Engine/Plugins/FX/Niagara/Source/Niagara/Public/NiagaraDecalRendererProperties.h
+++ b/Engine/Plugins/FX/Niagara/Source/Niagara/Public/NiagaraDecalRendererProperties.h
@@ -38,7 +38,6 @@ public:
 	virtual void GetUsedMaterials(const FNiagaraEmitterInstance* InEmitter, TArray<UMaterialInterface*>& OutMaterials) const override;
 	virtual bool IsSimTargetSupported(ENiagaraSimTarget InSimTarget) const override { return InSimTarget == ENiagaraSimTarget::CPUSim; };
 #if WITH_EDITORONLY_DATA
-	virtual bool IsSupportedVariableForBinding(const FNiagaraVariableBase& InSourceForBinding, const FName& InTargetBindingName) const override;
 	virtual const TArray<FNiagaraVariable>& GetOptionalAttributes() override;
 	virtual void GetRendererWidgets(const FNiagaraEmitterInstance* InEmitter, TArray<TSharedPtr<SWidget>>& OutWidgets, TSharedPtr<FAssetThumbnailPool> InThumbnailPool) const override;
 	virtual void GetRendererTooltipWidgets(const FNiagaraEmitterInstance* InEmitter, TArray<TSharedPtr<SWidget>>& OutWidgets, TSharedPtr<FAssetThumbnailPool> InThumbnailPool) const override;
diff --git a/Engine/Plugins/FX/Niagara/Source/Niagara/Public/NiagaraLightRendererProperties.h b/Engine/Plugins/FX/Niagara/Source/Niagara/Public/NiagaraLightRendererProperties.h
--- a/Engine/Plugins/FX/Niagara/Source/Niagara/Public/NiagaraLightRendererProperties.h
+++ b/Engine/Plugins/FX/Niagara/Source/Niagara/Public/NiagaraLightRendererProperties.h
@@ -36,7 +36,6 @@ public:
 	virtual void GetUsedMaterials(const FNiagaraEmitterInstance* InEmitter, TArray<UMaterialInterface*>& OutMaterials) const override;
 	virtual bool IsSimTargetSupported(ENiagaraSimTarget InSimTarget) const override { return (InSimTarget == ENiagaraSimTarget::CPUSim); };
 #if WITH_EDITORONLY_DATA
-	virtual bool IsSupportedVariableForBinding(const FNiagaraVariableBase& InSourceForBinding, const FName& InTargetBindingName) const override;
 	virtual const TArray<FNiagaraVariable>& GetOptionalAttributes() override;
 	virtual void GetRendererWidgets(const FNiagaraEmitterInstance* InEmitter, TArray<TSharedPtr<SWidget>>& OutWidgets, TSharedPtr<FAssetThumbnailPool> InThumbnailPool) const override;
 	virtual void GetRendererTooltipWidgets(const FNiagaraEmitterInstance* InEmitter, TArray<TSharedPtr<SWidget>>& OutWidgets, TSharedPtr<FAssetThumbnailPool> InThumbnailPool) const override;
diff --git a/Engine/Plugins/FX/Niagara/Source/Niagara/Public/NiagaraMeshRendererProperties.h b/Engine/Plugins/FX/Niagara/Source/Niagara/Public/NiagaraMeshRendererProperties.h
--- a/Engine/Plugins/FX/Niagara/Source/Niagara/Public/NiagaraMeshRendererProperties.h
+++ b/Engine/Plugins/FX/Niagara/Source/Niagara/Public/NiagaraMeshRendererProperties.h
@@ -222,10 +222,6 @@ public:
 
 	void UpdateMICs();
 
-#if WITH_EDITORONLY_DATA
-	bool IsSupportedVariableForBinding(const FNiagaraVariableBase& InSourceForBinding, const FName& InTargetBindingName) const;
-#endif
-
 	virtual ENiagaraRendererSourceDataMode GetCurrentSourceMode() const override { return SourceMode; }
 	//UNiagaraRendererProperties Interface END
 
diff --git a/Engine/Plugins/FX/Niagara/Source/Niagara/Public/NiagaraRibbonRendererProperties.h b/Engine/Plugins/FX/Niagara/Source/Niagara/Public/NiagaraRibbonRendererProperties.h
--- a/Engine/Plugins/FX/Niagara/Source/Niagara/Public/NiagaraRibbonRendererProperties.h
+++ b/Engine/Plugins/FX/Niagara/Source/Niagara/Public/NiagaraRibbonRendererProperties.h
@@ -231,10 +231,6 @@ public:
 	virtual TArray<FNiagaraVariable> GetBoundAttributes() const override; 
 #endif
 	virtual void CacheFromCompiledData(const FNiagaraDataSetCompiledData* CompiledData) override;
-	
-#if WITH_EDITORONLY_DATA
-	virtual bool IsSupportedVariableForBinding(const FNiagaraVariableBase& InSourceForBinding, const FName& InTargetBindingName) const override;
-#endif
 	//UNiagaraRendererProperties Interface END
 
 	UPROPERTY(EditAnywhere, Category = "Ribbon Rendering")
diff --git a/Engine/Plugins/FX/Niagara/Source/Niagara/Public/NiagaraSpriteRendererProperties.h b/Engine/Plugins/FX/Niagara/Source/Niagara/Public/NiagaraSpriteRendererProperties.h
--- a/Engine/Plugins/FX/Niagara/Source/Niagara/Public/NiagaraSpriteRendererProperties.h
+++ b/Engine/Plugins/FX/Niagara/Source/Niagara/Public/NiagaraSpriteRendererProperties.h
@@ -362,8 +362,6 @@ public:
 
 
 #if WITH_EDITORONLY_DATA
-	virtual bool IsSupportedVariableForBinding(const FNiagaraVariableBase& InSourceForBinding, const FName& InTargetBindingName) const override;
-
 	/** Use the cutout texture from the material opacity mask, or if none exist, from the material opacity.	*/
 	UPROPERTY(EditAnywhere, BlueprintReadOnly, Category = "Cutout")
 	bool bUseMaterialCutoutTexture;
diff --git a/Engine/Plugins/Runtime/GeometryCache/Source/GeometryCache/Private/NiagaraGeometryCacheRendererProperties.cpp b/Engine/Plugins/Runtime/GeometryCache/Source/GeometryCache/Private/NiagaraGeometryCacheRendererProperties.cpp
--- a/Engine/Plugins/Runtime/GeometryCache/Source/GeometryCache/Private/NiagaraGeometryCacheRendererProperties.cpp
+++ b/Engine/Plugins/Runtime/GeometryCache/Source/GeometryCache/Private/NiagaraGeometryCacheRendererProperties.cpp
@@ -225,7 +225,7 @@ void UNiagaraGeometryCacheRendererProperties::GetUsedMaterials(const FNiagaraEmi
 
 bool UNiagaraGeometryCacheRendererProperties::PopulateRequiredBindings(FNiagaraParameterStore& InParameterStore)
 {
-	bool bAnyAdded = false;
+	bool bAnyAdded = Super::PopulateRequiredBindings(InParameterStore);
 
 	for (const FNiagaraVariableAttributeBinding* Binding : AttributeBindings)
 	{
diff --git a/Engine/Plugins/Runtime/GeometryCache/Source/GeometryCache/Private/NiagaraRendererGeometryCache.cpp b/Engine/Plugins/Runtime/GeometryCache/Source/GeometryCache/Private/NiagaraRendererGeometryCache.cpp
--- a/Engine/Plugins/Runtime/GeometryCache/Source/GeometryCache/Private/NiagaraRendererGeometryCache.cpp
+++ b/Engine/Plugins/Runtime/GeometryCache/Source/GeometryCache/Private/NiagaraRendererGeometryCache.cpp
@@ -127,7 +127,8 @@ void FNiagaraRendererGeometryCache::PostSystemTick_GameThread(const UNiagaraRend
 	FNiagaraDataSetReaderInt32<int32> VisTagAccessor = FNiagaraDataSetAccessor<int32>::CreateReader(Data, Properties->RendererVisibilityTagBinding.GetDataSetBindableVariable().GetName());
 	FNiagaraDataSetReaderInt32<int32> ArrayIndexAccessor = FNiagaraDataSetAccessor<int32>::CreateReader(Data, Properties->ArrayIndexBinding.GetDataSetBindableVariable().GetName());
 	FNiagaraDataSetReaderInt32<int32> UniqueIDAccessor = FNiagaraDataSetAccessor<int32>::CreateReader(Data, FName("UniqueID"));
-	float CurrentTime = AttachComponent->GetWorld()->GetRealTimeSeconds();
+	const float CurrentTime = AttachComponent->GetWorld()->GetRealTimeSeconds();
+	const bool bIsRendererEnabled = IsRendererEnabled(InProperties, Emitter);
 	
 	auto IsParticleEnabled = [&EnabledAccessor, &VisTagAccessor, Properties](int32 ParticleIndex)
 	{
@@ -172,7 +173,7 @@ void FNiagaraRendererGeometryCache::PostSystemTick_GameThread(const UNiagaraRend
 			int32 PoolIndex;
 			if (UsedSlots.RemoveAndCopyValue(ParticleID, PoolIndex))
 			{
-				if (IsParticleEnabled(ParticleIndex))
+				if (bIsRendererEnabled && IsParticleEnabled(ParticleIndex))
 				{
 					ParticlesWithComponents.Emplace(ParticleID, PoolIndex);
 				}
@@ -210,7 +211,7 @@ void FNiagaraRendererGeometryCache::PostSystemTick_GameThread(const UNiagaraRend
 	int32 ComponentCount = 0;
 	for (uint32 ParticleIndex = 0; ParticleIndex < ParticleData.GetNumInstances(); ParticleIndex++)
 	{
-		if (!IsParticleEnabled(ParticleIndex))
+		if (!bIsRendererEnabled || !IsParticleEnabled(ParticleIndex))
 		{
 			// Skip particles that don't want a component
 			continue;
