- Truncation warning fixes

diff --git a/Engine/Plugins/FX/Niagara/Source/Niagara/Classes/NiagaraDataInterface.h b/Engine/Plugins/FX/Niagara/Source/Niagara/Classes/NiagaraDataInterface.h
--- a/Engine/Plugins/FX/Niagara/Source/Niagara/Classes/NiagaraDataInterface.h
+++ b/Engine/Plugins/FX/Niagara/Source/Niagara/Classes/NiagaraDataInterface.h
@@ -852,7 +852,12 @@ struct FNDIRandomInfoPolicy
 			//           We use the upper 24 bits as they tend to be higher quality.
 
 			// NOTE(mv): The divide can often be folded with the range scale in the rand functions
-			return FVector4f((v >> 8) & 0x00ffffff) / 16777216.0; // 0x01000000 == 16777216
+			FVector4f OutValue;
+			OutValue.X = float((v.X >> 8) & 0x00ffffff) / 16777216.0f; // 0x01000000 == 16777216
+			OutValue.Y = float((v.Y >> 8) & 0x00ffffff) / 16777216.0f;
+			OutValue.Z = float((v.Z >> 8) & 0x00ffffff) / 16777216.0f;
+			OutValue.W = float((v.W >> 8) & 0x00ffffff) / 16777216.0f;
+			return OutValue;
 			// return float4((v >> 8) & 0x00ffffff) * (1.0/16777216.0); // bugged, see UE-67738
 		}
 		else
@@ -876,7 +881,11 @@ struct FNDIRandomInfoPolicy
 			v.Y += v.Z * v.X;
 			v.Z += v.X * v.Y;
 
-			return FVector3f((v >> 8) & 0x00ffffff) / 16777216.0; // 0x01000000 == 16777216
+			FVector3f OutValue;
+			OutValue.X = float((v.X >> 8) & 0x00ffffff) / 16777216.0f; // 0x01000000 == 16777216
+			OutValue.Y = float((v.Y >> 8) & 0x00ffffff) / 16777216.0f;
+			OutValue.Z = float((v.Z >> 8) & 0x00ffffff) / 16777216.0f;
+			return OutValue;
 		}
 		else
 		{
diff --git a/Engine/Plugins/FX/Niagara/Source/Niagara/Classes/NiagaraDataInterfaceCurlNoise.h b/Engine/Plugins/FX/Niagara/Source/Niagara/Classes/NiagaraDataInterfaceCurlNoise.h
--- a/Engine/Plugins/FX/Niagara/Source/Niagara/Classes/NiagaraDataInterfaceCurlNoise.h
+++ b/Engine/Plugins/FX/Niagara/Source/Niagara/Classes/NiagaraDataInterfaceCurlNoise.h
@@ -22,7 +22,7 @@ public:
 	uint32 Seed;
 
 	// Precalculated when Seed changes. 
-	FVector OffsetFromSeed;
+	FVector3f OffsetFromSeed;
 
 	//UObject Interface
 	virtual void PostInitProperties()override;
@@ -59,7 +59,7 @@ protected:
 
 struct FNiagaraDataInterfaceProxyCurlNoise : public FNiagaraDataInterfaceProxy
 {
-	FNiagaraDataInterfaceProxyCurlNoise(const FVector& InOffset)
+	FNiagaraDataInterfaceProxyCurlNoise(const FVector3f& InOffset)
 	{
 		OffsetFromSeed = InOffset;
 	}
@@ -70,5 +70,5 @@ struct FNiagaraDataInterfaceProxyCurlNoise : public FNiagaraDataInterfaceProxy
 		return 0;
 	}
 
-	FVector OffsetFromSeed;
+	FVector3f OffsetFromSeed;
 };
diff --git a/Engine/Plugins/FX/Niagara/Source/Niagara/Classes/NiagaraDataInterfaceSpline.h b/Engine/Plugins/FX/Niagara/Source/Niagara/Classes/NiagaraDataInterfaceSpline.h
--- a/Engine/Plugins/FX/Niagara/Source/Niagara/Classes/NiagaraDataInterfaceSpline.h
+++ b/Engine/Plugins/FX/Niagara/Source/Niagara/Classes/NiagaraDataInterfaceSpline.h
@@ -262,7 +262,7 @@ protected:
 
 private:
 	
-	void WriteTransform(const FMatrix& ToWrite, FVectorVMExternalFunctionContext& Context);
+	void WriteTransform(const FMatrix44f& ToWrite, FVectorVMExternalFunctionContext& Context);
 
 	TMap<FNiagaraSystemInstanceID, FNDISpline_InstanceData*> SystemInstancesToProxyData_GT;
 };
diff --git a/Engine/Plugins/FX/Niagara/Source/Niagara/Classes/NiagaraDataInterfaceVector4Curve.h b/Engine/Plugins/FX/Niagara/Source/Niagara/Classes/NiagaraDataInterfaceVector4Curve.h
--- a/Engine/Plugins/FX/Niagara/Source/Niagara/Classes/NiagaraDataInterfaceVector4Curve.h
+++ b/Engine/Plugins/FX/Niagara/Source/Niagara/Classes/NiagaraDataInterfaceVector4Curve.h
@@ -84,7 +84,7 @@ protected:
 
 private: 
 	template<typename UseLUT>
-	FORCEINLINE_DEBUGGABLE FVector4 SampleCurveInternal(float X);
+	FORCEINLINE_DEBUGGABLE FVector4f SampleCurveInternal(float X);
 
 	static const FName SampleCurveName;
 };
diff --git a/Engine/Plugins/FX/Niagara/Source/Niagara/Private/DataInterface/NiagaraDataInterfaceDynamicMesh.cpp b/Engine/Plugins/FX/Niagara/Source/Niagara/Private/DataInterface/NiagaraDataInterfaceDynamicMesh.cpp
--- a/Engine/Plugins/FX/Niagara/Source/Niagara/Private/DataInterface/NiagaraDataInterfaceDynamicMesh.cpp
+++ b/Engine/Plugins/FX/Niagara/Source/Niagara/Private/DataInterface/NiagaraDataInterfaceDynamicMesh.cpp
@@ -659,7 +659,7 @@ namespace NDIDynamicMeshLocal
 				StaticMeshData.NumTexCoords = InstanceData->NumTexCoords;
 				for (uint32 i=0; i < InstanceData->NumTexCoords; ++i)
 				{
-					StaticMeshData.TextureCoordinates.Emplace(&InstanceData->VertexBuffer, /*InstanceData->TexCoordOffset + (i * sizeof(FVector2f))*/0, sizeof(FVector2f) * InstanceData->NumTexCoords, VET_Float2, EVertexStreamUsage::ManualFetch);
+					StaticMeshData.TextureCoordinates.Emplace(&InstanceData->VertexBuffer, /*InstanceData->TexCoordOffset + (i * sizeof(FVector2f))*/0, int32(sizeof(FVector2f) * InstanceData->NumTexCoords), VET_Float2, EVertexStreamUsage::ManualFetch);
 				}
 				StaticMeshData.TextureCoordinatesSRV = InstanceData->VertexBufferTexCoordSRV;
 			}
diff --git a/Engine/Plugins/FX/Niagara/Source/Niagara/Private/DataInterface/NiagaraDataInterfaceStaticMesh.cpp b/Engine/Plugins/FX/Niagara/Source/Niagara/Private/DataInterface/NiagaraDataInterfaceStaticMesh.cpp
--- a/Engine/Plugins/FX/Niagara/Source/Niagara/Private/DataInterface/NiagaraDataInterfaceStaticMesh.cpp
+++ b/Engine/Plugins/FX/Niagara/Source/Niagara/Private/DataInterface/NiagaraDataInterfaceStaticMesh.cpp
@@ -1481,7 +1481,7 @@ namespace NDIStaticMeshLocal
 		template<typename TRandomHelper>
 		FORCEINLINE int32 RandomUniformSection(const TRandomHelper& RandHelper, int32 InstanceIndex, const TConstArrayView<FNDISectionInfo>& SectionInfos, TConstArrayView<int32>& SectionRemap) const
 		{
-			int32 SectionIndex = RandHelper.Rand(InstanceIndex) * SectionInfos.Num();
+			int32 SectionIndex = int32(RandHelper.Rand(InstanceIndex) * float(SectionInfos.Num()));
 			SectionIndex = RandHelper.Rand(InstanceIndex) < SectionInfos[SectionIndex].Prob ? SectionIndex : SectionInfos[SectionIndex].Alias;
 			return SectionRemap[SectionIndex];
 		}
diff --git a/Engine/Plugins/FX/Niagara/Source/Niagara/Private/DataInterface/NiagaraDataInterfaceUObjectPropertyReader.cpp b/Engine/Plugins/FX/Niagara/Source/Niagara/Private/DataInterface/NiagaraDataInterfaceUObjectPropertyReader.cpp
--- a/Engine/Plugins/FX/Niagara/Source/Niagara/Private/DataInterface/NiagaraDataInterfaceUObjectPropertyReader.cpp
+++ b/Engine/Plugins/FX/Niagara/Source/Niagara/Private/DataInterface/NiagaraDataInterfaceUObjectPropertyReader.cpp
@@ -188,7 +188,7 @@ namespace NDIUObjectPropertyReaderLocal
 			}
 			else if (const FDoubleProperty* DoubleProperty = CastField<const FDoubleProperty>(InProperty))
 			{
-				return [PropertyAddress](const FNiagaraLWCConverter& LwcConverter, void* DestAddress) { *reinterpret_cast<float*>(DestAddress) = *reinterpret_cast<const double*>(PropertyAddress); };
+				return [PropertyAddress](const FNiagaraLWCConverter& LwcConverter, void* DestAddress) { *reinterpret_cast<float*>(DestAddress) = float(*reinterpret_cast<const double*>(PropertyAddress)); };
 			}
 			return nullptr;
 		}
@@ -310,7 +310,7 @@ namespace NDIUObjectPropertyReaderLocal
 	{
 		static FName GetFunctionName() { return FName("GetIntProperty"); }
 		static FNiagaraTypeDefinition GetTypeDef() { return FNiagaraTypeDefinition::GetIntDef(); }
-		static void VMFunction(FVectorVMExternalFunctionContext& Context, const uint32 GetterIndex) { VMReadData<int32, int32>(Context, GetterIndex, 0.0f); }
+		static void VMFunction(FVectorVMExternalFunctionContext& Context, const uint32 GetterIndex) { VMReadData<int32, int32>(Context, GetterIndex, 0); }
 		static TFunction<void(const FNiagaraLWCConverter&, void*)> GetCopyFunction(const FProperty* InProperty, const void* PropertyAddress)
 		{
 			if (const FEnumProperty* EnumProperty = CastField<FEnumProperty>(InProperty))
diff --git a/Engine/Plugins/FX/Niagara/Source/Niagara/Private/Experimental/NiagaraMeshUvMapping.h b/Engine/Plugins/FX/Niagara/Source/Niagara/Private/Experimental/NiagaraMeshUvMapping.h
--- a/Engine/Plugins/FX/Niagara/Source/Niagara/Private/Experimental/NiagaraMeshUvMapping.h
+++ b/Engine/Plugins/FX/Niagara/Source/Niagara/Private/Experimental/NiagaraMeshUvMapping.h
@@ -133,10 +133,10 @@ struct FQuadTreeQueryHelper
 		for (int32 i = 0; i < UE_ARRAY_COUNT(TriangleEdges); ++i)
 		{
 			const FVector2D SeparatingAxis(-TriangleEdges[i].Y, TriangleEdges[i].X);
-			float AabbSegmentMin = FMath::Min(0.0f, FMath::Min3(SeparatingAxis.X, SeparatingAxis.Y, SeparatingAxis.X + SeparatingAxis.Y));
-			float AabbSegmentMax = FMath::Max(0.0f, FMath::Max3(SeparatingAxis.X, SeparatingAxis.Y, SeparatingAxis.X + SeparatingAxis.Y));
-			float TriangleSegmentMin = FMath::Min3(FVector2D::DotProduct(A, SeparatingAxis), FVector2D::DotProduct(B, SeparatingAxis), FVector2D::DotProduct(C, SeparatingAxis));
-			float TriangleSegmentMax = FMath::Max3(FVector2D::DotProduct(A, SeparatingAxis), FVector2D::DotProduct(B, SeparatingAxis), FVector2D::DotProduct(C, SeparatingAxis));
+			const double AabbSegmentMin = FMath::Min(0.0, FMath::Min3(SeparatingAxis.X, SeparatingAxis.Y, SeparatingAxis.X + SeparatingAxis.Y));
+			const double AabbSegmentMax = FMath::Max(0.0, FMath::Max3(SeparatingAxis.X, SeparatingAxis.Y, SeparatingAxis.X + SeparatingAxis.Y));
+			const double TriangleSegmentMin = FMath::Min3(FVector2D::DotProduct(A, SeparatingAxis), FVector2D::DotProduct(B, SeparatingAxis), FVector2D::DotProduct(C, SeparatingAxis));
+			const double TriangleSegmentMax = FMath::Max3(FVector2D::DotProduct(A, SeparatingAxis), FVector2D::DotProduct(B, SeparatingAxis), FVector2D::DotProduct(C, SeparatingAxis));
 
 			if (AabbSegmentMin > TriangleSegmentMax || AabbSegmentMax < TriangleSegmentMin)
 			{
diff --git a/Engine/Plugins/FX/Niagara/Source/Niagara/Private/NiagaraComponentPool.cpp b/Engine/Plugins/FX/Niagara/Source/Niagara/Private/NiagaraComponentPool.cpp
--- a/Engine/Plugins/FX/Niagara/Source/Niagara/Private/NiagaraComponentPool.cpp
+++ b/Engine/Plugins/FX/Niagara/Source/Niagara/Private/NiagaraComponentPool.cpp
@@ -138,7 +138,7 @@ UNiagaraComponent* FNCPool::Acquire(UWorld* World, UNiagaraSystem* Template, ENC
 	return RetElem.Component;
 }
 
-void FNCPool::Reclaim(UNiagaraComponent* Component, const float CurrentTimeSeconds)
+void FNCPool::Reclaim(UNiagaraComponent* Component, const double CurrentTimeSeconds)
 {
 	check(Component);
 	check(Component->GetAsset());
@@ -219,7 +219,7 @@ bool FNCPool::RemoveComponent(UNiagaraComponent* Component)
 }
 
 extern int32 GNigaraAllowPrimedPools;
-void FNCPool::KillUnusedComponents(float KillTime, UNiagaraSystem* Template)
+void FNCPool::KillUnusedComponents(double KillTime, UNiagaraSystem* Template)
 {
 	int32 i = 0;
 	check(Template);
@@ -268,7 +268,7 @@ bool UNiagaraComponentPool::Enabled()
 UNiagaraComponentPool::UNiagaraComponentPool(const FObjectInitializer& ObjectInitializer)
 	: Super(ObjectInitializer)
 {
-	LastParticleSytemPoolCleanTime = 0.0f;
+	LastParticleSytemPoolCleanTime = 0.0;
 }
 
 UNiagaraComponentPool::~UNiagaraComponentPool()
@@ -475,7 +475,7 @@ void UNiagaraComponentPool::ReclaimWorldParticleSystem(UNiagaraComponent* Compon
 	// WorldParticleSystemPools is empty after world cleanup, so destroy components coming in to be reclaimed instead
 	if (GbEnableNiagaraSystemPooling && Asset != nullptr && WorldParticleSystemPools.Num() > 0)
 	{
-		float CurrentTime = Component->GetWorld()->GetTimeSeconds();
+		const double CurrentTime = Component->GetWorld()->GetTimeSeconds();
 
 		//Periodically clear up the pools.
 		if (CurrentTime - LastParticleSytemPoolCleanTime > GNiagaraSystemPoolingCleanTime)
diff --git a/Engine/Plugins/FX/Niagara/Source/Niagara/Private/NiagaraDataInterface2DArrayTexture.cpp b/Engine/Plugins/FX/Niagara/Source/Niagara/Private/NiagaraDataInterface2DArrayTexture.cpp
--- a/Engine/Plugins/FX/Niagara/Source/Niagara/Private/NiagaraDataInterface2DArrayTexture.cpp
+++ b/Engine/Plugins/FX/Niagara/Source/Niagara/Private/NiagaraDataInterface2DArrayTexture.cpp
@@ -216,7 +216,9 @@ bool UNiagaraDataInterface2DArrayTexture::PerInstanceTick(void* PerInstanceData,
 					{
 						InstanceData.TextureReferenceRHI = RT_Texture->TextureReference.TextureReferenceRHI;
 						InstanceData.SamplerStateRHI = RT_Texture->GetResource() ? RT_Texture->GetResource()->SamplerStateRHI.GetReference() : TStaticSamplerState<SF_Point>::GetRHI();
-						InstanceData.TextureSize = FVector3f(RT_TextureSize.X, RT_TextureSize.Y, RT_TextureSize.Z);
+						InstanceData.TextureSize.X = float(RT_TextureSize.X);
+						InstanceData.TextureSize.Y = float(RT_TextureSize.Y);
+						InstanceData.TextureSize.Z = float(RT_TextureSize.Z);
 					}
 					else
 					{
diff --git a/Engine/Plugins/FX/Niagara/Source/Niagara/Private/NiagaraDataInterfaceCamera.cpp b/Engine/Plugins/FX/Niagara/Source/Niagara/Private/NiagaraDataInterfaceCamera.cpp
--- a/Engine/Plugins/FX/Niagara/Source/Niagara/Private/NiagaraDataInterfaceCamera.cpp
+++ b/Engine/Plugins/FX/Niagara/Source/Niagara/Private/NiagaraDataInterfaceCamera.cpp
@@ -469,14 +469,14 @@ void UNiagaraDataInterfaceCamera::GetCameraProperties(FVectorVMExternalFunctionC
 	VectorVM::FExternalFuncRegisterHandler<float> CamRightZ(Context);
 
 	FCameraDataInterface_InstanceData* CamData = InstData.Get();
-	float XPos = CamData->CameraLocation.X;
-	float YPos = CamData->CameraLocation.Y;
-	float ZPos = CamData->CameraLocation.Z;
+	const float XPos = CamData->CameraLocation.X;
+	const float YPos = CamData->CameraLocation.Y;
+	const float ZPos = CamData->CameraLocation.Z;
 
 	FRotationMatrix RotationMatrix(CamData->CameraRotation);
-	const FVector Forward = RotationMatrix.GetScaledAxis(EAxis::X);
-	const FVector Up = RotationMatrix.GetScaledAxis(EAxis::Z);
-	const FVector Right = RotationMatrix.GetScaledAxis(EAxis::Y);
+	const FVector3f Forward = FVector3f(RotationMatrix.GetScaledAxis(EAxis::X));
+	const FVector3f Up = FVector3f(RotationMatrix.GetScaledAxis(EAxis::Z));
+	const FVector3f Right = FVector3f(RotationMatrix.GetScaledAxis(EAxis::Y));
 
 	for (int32 i = 0; i < Context.GetNumInstances(); ++i)
 	{
diff --git a/Engine/Plugins/FX/Niagara/Source/Niagara/Private/NiagaraDataInterfaceCurlNoise.cpp b/Engine/Plugins/FX/Niagara/Source/Niagara/Private/NiagaraDataInterfaceCurlNoise.cpp
--- a/Engine/Plugins/FX/Niagara/Source/Niagara/Private/NiagaraDataInterfaceCurlNoise.cpp
+++ b/Engine/Plugins/FX/Niagara/Source/Niagara/Private/NiagaraDataInterfaceCurlNoise.cpp
@@ -82,28 +82,28 @@ FNiagaraUIntVector Rand3DPCG16(FIntVector p)
 	return v >> 16u;
 }
 
-FVector NiagaraVectorFrac(FVector v)
+FVector3f NiagaraVectorFrac(FVector3f v)
 {
-	return FVector(FMath::Frac(v.X), FMath::Frac(v.Y), FMath::Frac(v.Z));
+	return FVector3f(FMath::Frac(v.X), FMath::Frac(v.Y), FMath::Frac(v.Z));
 }
 
-FVector NoiseTileWrap(FVector v, bool bTiling, float RepeatSize)
+FVector3f NoiseTileWrap(FVector3f v, bool bTiling, float RepeatSize)
 {
-	FVector vv = bTiling ? (NiagaraVectorFrac(v / RepeatSize) * RepeatSize) : v;
+	FVector3f vv = bTiling ? (NiagaraVectorFrac(v / RepeatSize) * RepeatSize) : v;
 	return vv;
 }
 
 struct FNiagaraMatrix4x3
 {
-	FVector Row0;
-	FVector Row1;
-	FVector Row2;
-	FVector Row3;
+	FVector3f Row0;
+	FVector3f Row1;
+	FVector3f Row2;
+	FVector3f Row3;
 
-	FNiagaraMatrix4x3() : Row0(FVector()), Row1(FVector()), Row2(FVector()), Row3(FVector()) {}
-	FNiagaraMatrix4x3(FVector Row0, FVector Row1, FVector Row2, FVector Row3) : Row0(Row0), Row1(Row1), Row2(Row2), Row3(Row3) {}
+	FNiagaraMatrix4x3() : Row0(FVector3f::Zero()), Row1(FVector3f::Zero()), Row2(FVector3f::Zero()), Row3(FVector3f::Zero()) {}
+	FNiagaraMatrix4x3(FVector3f Row0, FVector3f Row1, FVector3f Row2, FVector3f Row3) : Row0(Row0), Row1(Row1), Row2(Row2), Row3(Row3) {}
 
-	FVector& operator[](int Row)
+	FVector3f& operator[](int Row)
 	{
 		if (Row == 0) return Row0;
 		if (Row == 1) return Row1;
@@ -112,7 +112,7 @@ struct FNiagaraMatrix4x3
 		check(false);
 		return Row0;
 	}
-	const FVector& operator[](int Row) const
+	const FVector3f& operator[](int Row) const
 	{
 		if (Row == 0) return Row0;
 		if (Row == 1) return Row1;
@@ -123,72 +123,80 @@ struct FNiagaraMatrix4x3
 	}
 };
 
-FVector NiagaraVectorFloor(FVector v) {
-	return FVector(FGenericPlatformMath::FloorToFloat(v.X),
+FVector3f NiagaraVectorFloor(FVector3f v) {
+	return FVector3f(FGenericPlatformMath::FloorToFloat(v.X),
 		           FGenericPlatformMath::FloorToFloat(v.Y),
 		           FGenericPlatformMath::FloorToFloat(v.Z));
 };
 
-FVector NiagaraVectorStep(FVector v, FVector u)
+FIntVector NiagaraVectorFloorToInt(FVector3f v) {
+	return FIntVector(
+		FGenericPlatformMath::FloorToInt(v.X),
+		FGenericPlatformMath::FloorToInt(v.Y),
+		FGenericPlatformMath::FloorToInt(v.Z)
+	);
+};
+
+FVector3f NiagaraVectorStep(FVector3f v, FVector3f u)
 {
-	return FVector(u.X >= v.X ? 1 : 0, 
-		           u.Y >= v.Y ? 1 : 0, 
-		           u.Z >= v.Z ? 1 : 0);
+	return FVector3f(u.X >= v.X ? 1.0f : 0.0f,
+		           u.Y >= v.Y ? 1.0f : 0.0f,
+		           u.Z >= v.Z ? 1.0f : 0.0f);
 }
 
-FVector NiagaraVectorSwizzle(FVector v, uint32 x, uint32 y, uint32 z)
+FVector3f NiagaraVectorSwizzle(FVector3f v, uint32 x, uint32 y, uint32 z)
 {
-	return FVector(v[x], v[y], v[z]);
+	return FVector3f(v[x], v[y], v[z]);
 }
 
-FVector NiagaraVectorMin(FVector u, FVector v)
+FVector3f NiagaraVectorMin(FVector3f u, FVector3f v)
 {
-	return FVector(u.X < v.X ? u.X : v.X, 
+	return FVector3f(u.X < v.X ? u.X : v.X,
 		           u.Y < v.Y ? u.Y : v.Y, 
 		           u.Z < v.Z ? u.Z : v.Z);
 }
 
-FVector NiagaraVectorMax(FVector u, FVector v)
+FVector3f NiagaraVectorMax(FVector3f u, FVector3f v)
 {
-	return FVector(u.X > v.X ? u.X : v.X,
+	return FVector3f(u.X > v.X ? u.X : v.X,
 		           u.Y > v.Y ? u.Y : v.Y,
 	           	   u.Z > v.Z ? u.Z : v.Z);
 } 
 
-FNiagaraMatrix4x3 SimplexCorners(FVector v)
+FNiagaraMatrix4x3 SimplexCorners(FVector3f v)
 {
-	FVector tet = NiagaraVectorFloor(v + v.X / 3 + v.Y / 3 + v.Z / 3);
-	FVector base = tet - tet.X / 6 - tet.Y / 6 - tet.Z / 6;
-	FVector f = v - base;
+	FVector3f tet = NiagaraVectorFloor(v + v.X / 3 + v.Y / 3 + v.Z / 3);
+	FVector3f base = tet - tet.X / 6 - tet.Y / 6 - tet.Z / 6;
+	FVector3f f = v - base;
 
-	FVector g = NiagaraVectorStep(NiagaraVectorSwizzle(f, 1, 2, 0), f), h = FVector(1.0f) - NiagaraVectorSwizzle(g, 2, 0, 1);
-	FVector a1 = NiagaraVectorMin(g, h) - 1. / 6., a2 = NiagaraVectorMax(g, h) - 1. / 3.;
+	FVector3f g = NiagaraVectorStep(NiagaraVectorSwizzle(f, 1, 2, 0), f), h = FVector3f(1.0f) - NiagaraVectorSwizzle(g, 2, 0, 1);
+	FVector3f a1 = NiagaraVectorMin(g, h) - 1. / 6., a2 = NiagaraVectorMax(g, h) - 1. / 3.;
 
 	return FNiagaraMatrix4x3(base, base + a1, base + a2, base + 0.5);
 }
 
-FVector4 NiagaraVector4Saturate(const FVector4& v) {
-	return FVector4(FMath::Clamp(v.X, 0.0f, 1.0f), FMath::Clamp(v.Y, 0.0f, 1.0f), FMath::Clamp(v.Z, 0.0f, 1.0f), FMath::Clamp(v.W, 0.0f, 1.0f));
+FVector4f NiagaraVector4Saturate(const FVector4f& v) {
+	return FVector4f(FMath::Clamp(v.X, 0.0f, 1.0f), FMath::Clamp(v.Y, 0.0f, 1.0f), FMath::Clamp(v.Z, 0.0f, 1.0f), FMath::Clamp(v.W, 0.0f, 1.0f));
 }
 
-FVector4 SimplexSmooth(FNiagaraMatrix4x3 f)
+FVector4f SimplexSmooth(FNiagaraMatrix4x3 f)
 {
-	const float scale = 1024. / 375.;
-	FVector4 d = FVector4(FVector::DotProduct(f[0], f[0]), FVector::DotProduct(f[1], f[1]), FVector::DotProduct(f[2], f[2]), FVector::DotProduct(f[3], f[3]));
-	FVector4 s = NiagaraVector4Saturate(2.0f * d);
-	return scale * (FVector4(1.0f, 1.0f, 1.0f, 1.0f) + s * (FVector4(-3.0f, -3.0f, -3.0f, -3.0f) + s * (FVector4(3.0f, 3.0f, 3.0f, 3.0f) - s)));
+	const float scale = 1024.f / 375.f;
+	FVector4f d = FVector4f(FVector3f::DotProduct(f[0], f[0]), FVector3f::DotProduct(f[1], f[1]), FVector3f::DotProduct(f[2], f[2]), FVector3f::DotProduct(f[3], f[3]));
+	FVector4f s = NiagaraVector4Saturate(2.0f * d);
+	return scale * (FVector4f(1.0f, 1.0f, 1.0f, 1.0f) + s * (FVector4f(-3.0f, -3.0f, -3.0f, -3.0f) + s * (FVector4f(3.0f, 3.0f, 3.0f, 3.0f) - s)));
 }
 
 struct FNiagaraMatrix3x4
 {
-	FVector4 row0;
-	FVector4 row1;
-	FVector4 row2;
+	FVector4f row0;
+	FVector4f row1;
+	FVector4f row2;
 
-	FNiagaraMatrix3x4() : row0(FVector4(0.0, 0.0, 0.0, 0.0)), row1(FVector4(0.0, 0.0, 0.0, 0.0)), row2(FVector4(0.0, 0.0, 0.0, 0.0)) {}
-	FNiagaraMatrix3x4(const FVector4& row0, const FVector4& row1, const FVector4& row2) : row0(row0), row1(row1), row2(row2) {}
+	FNiagaraMatrix3x4() : row0(FVector4f::Zero()), row1(FVector4f::Zero()), row2(FVector4f::Zero()) {}
+	FNiagaraMatrix3x4(const FVector4f& row0, const FVector4f& row1, const FVector4f& row2) : row0(row0), row1(row1), row2(row2) {}
 
-	FVector4& operator[](int row)
+	FVector4f& operator[](int row)
 	{
 		if (row == 0) return row0;
 		if (row == 1) return row1;
@@ -196,7 +204,7 @@ struct FNiagaraMatrix3x4
 		check(false);
 		return row0;
 	}
-	const FVector4& operator[](int row) const
+	const FVector4f& operator[](int row) const
 	{
 		if (row == 0) return row0;
 		if (row == 1) return row1;
@@ -209,14 +217,14 @@ struct FNiagaraMatrix3x4
 FNiagaraMatrix3x4 SimplexDSmooth(FNiagaraMatrix4x3 f)
 {
 	const float scale = 1024. / 375.;
-	FVector4 d = FVector4(FVector::DotProduct(f[0], f[0]), FVector::DotProduct(f[1], f[1]), FVector::DotProduct(f[2], f[2]), FVector::DotProduct(f[3], f[3]));
-	FVector4 s = NiagaraVector4Saturate(2 * d);
-	s = -12 * FVector4(scale, scale, scale, scale) + s * (24 * FVector4(scale, scale, scale, scale) - s * 12 * scale);
+	FVector4f d = FVector4f(FVector3f::DotProduct(f[0], f[0]), FVector3f::DotProduct(f[1], f[1]), FVector3f::DotProduct(f[2], f[2]), FVector3f::DotProduct(f[3], f[3]));
+	FVector4f s = NiagaraVector4Saturate(2 * d);
+	s = -12 * FVector4f(scale, scale, scale, scale) + s * (24 * FVector4f(scale, scale, scale, scale) - s * 12 * scale);
 
 	return FNiagaraMatrix3x4(
-		s * FVector4(f[0][0], f[1][0], f[2][0], f[3][0]),
-		s * FVector4(f[0][1], f[1][1], f[2][1], f[3][1]),
-		s * FVector4(f[0][2], f[1][2], f[2][2], f[3][2]));
+		s * FVector4f(f[0][0], f[1][0], f[2][0], f[3][0]),
+		s * FVector4f(f[0][1], f[1][1], f[2][1], f[3][1]),
+		s * FVector4f(f[0][2], f[1][2], f[2][2], f[3][2]));
 }
 
 FNiagaraUIntVector FNiagaraUIntVectorSwizzle(FNiagaraUIntVector v, int x, int y, int z)
@@ -224,29 +232,29 @@ FNiagaraUIntVector FNiagaraUIntVectorSwizzle(FNiagaraUIntVector v, int x, int y,
 	return FNiagaraUIntVector(v[x], v[y], v[z]);
 }
 
-FVector FNiagaraUIntVectorToFVector(FNiagaraUIntVector v)
+FVector3f FNiagaraUIntVectorToFVector(FNiagaraUIntVector v)
 {
-	return FVector(v.X, v.Y, v.Z);
+	return FVector3f(float(v.X), float(v.Y), float(v.Z));
 }
 
-FVector MulFVector4AndFNiagaraMatrix4x3(const FVector4& lhs, FNiagaraMatrix4x3 rhs)
+FVector3f MulFVector4AndFNiagaraMatrix4x3(const FVector4f& lhs, FNiagaraMatrix4x3 rhs)
 {
-	return FVector(lhs[0] * rhs[0][0] + lhs[1] * rhs[1][0] + lhs[2] * rhs[2][0] + lhs[3] * rhs[3][0],
+	return FVector3f(lhs[0] * rhs[0][0] + lhs[1] * rhs[1][0] + lhs[2] * rhs[2][0] + lhs[3] * rhs[3][0],
 		           lhs[0] * rhs[0][1] + lhs[1] * rhs[1][1] + lhs[2] * rhs[2][1] + lhs[3] * rhs[3][1],
 	           	   lhs[0] * rhs[0][2] + lhs[1] * rhs[1][2] + lhs[2] * rhs[2][2] + lhs[3] * rhs[3][2]);
 }
 
-FVector MulFNiagaraMatrix3x4FAndVector4(const FNiagaraMatrix3x4& lhs, const FVector4& rhs)
+FVector3f MulFNiagaraMatrix3x4FAndVector4(const FNiagaraMatrix3x4& lhs, const FVector4f& rhs)
 {
-	return FVector(lhs[0][0] * rhs[0] + lhs[0][1] * rhs[1] + lhs[0][2] * rhs[2] + lhs[0][3] * rhs[3],
+	return FVector3f(lhs[0][0] * rhs[0] + lhs[0][1] * rhs[1] + lhs[0][2] * rhs[2] + lhs[0][3] * rhs[3],
 		           lhs[1][0] * rhs[0] + lhs[1][1] * rhs[1] + lhs[1][2] * rhs[2] + lhs[1][3] * rhs[3],
 		           lhs[2][0] * rhs[0] + lhs[2][1] * rhs[1] + lhs[2][2] * rhs[2] + lhs[2][3] * rhs[3]);
 }
 
 #define MGradientMask FNiagaraUIntVector(0x8000, 0x4000, 0x2000)
-#define MGradientScale FVector(1. / 0x4000, 1. / 0x2000, 1. / 0x1000)
+#define MGradientScale FVector3f(1.f / 0x4000, 1.f / 0x2000, 1.f / 0x1000)
 
-FNiagaraMatrix3x4 JacobianSimplex_ALU(FVector v)
+FNiagaraMatrix3x4 JacobianSimplex_ALU(FVector3f v)
 {
 	FNiagaraMatrix4x3 T = SimplexCorners(v);
 	FNiagaraUIntVector rand;
@@ -254,48 +262,48 @@ FNiagaraMatrix3x4 JacobianSimplex_ALU(FVector v)
 	FNiagaraMatrix3x4 grad;
 
 	fv[0] = v - T[0];
-	rand = Rand3DPCG16(FIntVector(NiagaraVectorFloor(6 * T[0] + 0.5)));
-	gvec[0][0] = FVector(FNiagaraUIntVectorToFVector(FNiagaraUIntVectorSwizzle(rand, 0, 0, 0) & MGradientMask)) * MGradientScale - 1;
-	gvec[1][0] = FVector(FNiagaraUIntVectorToFVector(FNiagaraUIntVectorSwizzle(rand, 1, 1, 1) & MGradientMask)) * MGradientScale - 1;
-	gvec[2][0] = FVector(FNiagaraUIntVectorToFVector(FNiagaraUIntVectorSwizzle(rand, 2, 2, 2) & MGradientMask)) * MGradientScale - 1;
-	grad[0][0] = FVector::DotProduct(gvec[0][0], fv[0]);
-	grad[1][0] = FVector::DotProduct(gvec[1][0], fv[0]);
-	grad[2][0] = FVector::DotProduct(gvec[2][0], fv[0]);
+	rand = Rand3DPCG16(NiagaraVectorFloorToInt(6 * T[0] + 0.5));
+	gvec[0][0] = FVector3f(FNiagaraUIntVectorToFVector(FNiagaraUIntVectorSwizzle(rand, 0, 0, 0) & MGradientMask)) * MGradientScale - 1;
+	gvec[1][0] = FVector3f(FNiagaraUIntVectorToFVector(FNiagaraUIntVectorSwizzle(rand, 1, 1, 1) & MGradientMask)) * MGradientScale - 1;
+	gvec[2][0] = FVector3f(FNiagaraUIntVectorToFVector(FNiagaraUIntVectorSwizzle(rand, 2, 2, 2) & MGradientMask)) * MGradientScale - 1;
+	grad[0][0] = FVector3f::DotProduct(gvec[0][0], fv[0]);
+	grad[1][0] = FVector3f::DotProduct(gvec[1][0], fv[0]);
+	grad[2][0] = FVector3f::DotProduct(gvec[2][0], fv[0]);
 
 	fv[1] = v - T[1];
-	rand = Rand3DPCG16(FIntVector(NiagaraVectorFloor(6 * T[1] + 0.5)));
-	gvec[0][1] = FVector(FNiagaraUIntVectorToFVector(FNiagaraUIntVectorSwizzle(rand, 0, 0, 0) & MGradientMask)) * MGradientScale - 1;
-	gvec[1][1] = FVector(FNiagaraUIntVectorToFVector(FNiagaraUIntVectorSwizzle(rand, 1, 1, 1) & MGradientMask)) * MGradientScale - 1;
-	gvec[2][1] = FVector(FNiagaraUIntVectorToFVector(FNiagaraUIntVectorSwizzle(rand, 2, 2, 2) & MGradientMask)) * MGradientScale - 1;
-	grad[0][1] = FVector::DotProduct(gvec[0][1], fv[1]);
-	grad[1][1] = FVector::DotProduct(gvec[1][1], fv[1]);
-	grad[2][1] = FVector::DotProduct(gvec[2][1], fv[1]);
+	rand = Rand3DPCG16(NiagaraVectorFloorToInt(6 * T[1] + 0.5));
+	gvec[0][1] = FVector3f(FNiagaraUIntVectorToFVector(FNiagaraUIntVectorSwizzle(rand, 0, 0, 0) & MGradientMask)) * MGradientScale - 1;
+	gvec[1][1] = FVector3f(FNiagaraUIntVectorToFVector(FNiagaraUIntVectorSwizzle(rand, 1, 1, 1) & MGradientMask)) * MGradientScale - 1;
+	gvec[2][1] = FVector3f(FNiagaraUIntVectorToFVector(FNiagaraUIntVectorSwizzle(rand, 2, 2, 2) & MGradientMask)) * MGradientScale - 1;
+	grad[0][1] = FVector3f::DotProduct(gvec[0][1], fv[1]);
+	grad[1][1] = FVector3f::DotProduct(gvec[1][1], fv[1]);
+	grad[2][1] = FVector3f::DotProduct(gvec[2][1], fv[1]);
 
 	fv[2] = v - T[2];
-	rand = Rand3DPCG16(FIntVector(NiagaraVectorFloor(6 * T[2] + 0.5)));
-	gvec[0][2] = FVector(FNiagaraUIntVectorToFVector(FNiagaraUIntVectorSwizzle(rand, 0, 0, 0) & MGradientMask)) * MGradientScale - 1;
-	gvec[1][2] = FVector(FNiagaraUIntVectorToFVector(FNiagaraUIntVectorSwizzle(rand, 1, 1, 1) & MGradientMask)) * MGradientScale - 1;
-	gvec[2][2] = FVector(FNiagaraUIntVectorToFVector(FNiagaraUIntVectorSwizzle(rand, 2, 2, 2) & MGradientMask)) * MGradientScale - 1;
-	grad[0][2] = FVector::DotProduct(gvec[0][2], fv[2]);
-	grad[1][2] = FVector::DotProduct(gvec[1][2], fv[2]);
-	grad[2][2] = FVector::DotProduct(gvec[2][2], fv[2]);
+	rand = Rand3DPCG16(NiagaraVectorFloorToInt(6 * T[2] + 0.5));
+	gvec[0][2] = FVector3f(FNiagaraUIntVectorToFVector(FNiagaraUIntVectorSwizzle(rand, 0, 0, 0) & MGradientMask)) * MGradientScale - 1;
+	gvec[1][2] = FVector3f(FNiagaraUIntVectorToFVector(FNiagaraUIntVectorSwizzle(rand, 1, 1, 1) & MGradientMask)) * MGradientScale - 1;
+	gvec[2][2] = FVector3f(FNiagaraUIntVectorToFVector(FNiagaraUIntVectorSwizzle(rand, 2, 2, 2) & MGradientMask)) * MGradientScale - 1;
+	grad[0][2] = FVector3f::DotProduct(gvec[0][2], fv[2]);
+	grad[1][2] = FVector3f::DotProduct(gvec[1][2], fv[2]);
+	grad[2][2] = FVector3f::DotProduct(gvec[2][2], fv[2]);
 
 	fv[3] = v - T[3];
-	rand = Rand3DPCG16(FIntVector(NiagaraVectorFloor(6 * T[3] + 0.5)));
-	gvec[0][3] = FVector(FNiagaraUIntVectorToFVector(FNiagaraUIntVectorSwizzle(rand, 0, 0, 0) & MGradientMask)) * MGradientScale - 1;
-	gvec[1][3] = FVector(FNiagaraUIntVectorToFVector(FNiagaraUIntVectorSwizzle(rand, 1, 1, 1) & MGradientMask)) * MGradientScale - 1;
-	gvec[2][3] = FVector(FNiagaraUIntVectorToFVector(FNiagaraUIntVectorSwizzle(rand, 2, 2, 2) & MGradientMask)) * MGradientScale - 1;
-	grad[0][3] = FVector::DotProduct(gvec[0][3], fv[3]);
-	grad[1][3] = FVector::DotProduct(gvec[1][3], fv[3]);
-	grad[2][3] = FVector::DotProduct(gvec[2][3], fv[3]);
-
-	FVector4 sv = SimplexSmooth(fv);
+	rand = Rand3DPCG16(NiagaraVectorFloorToInt(6 * T[3] + 0.5));
+	gvec[0][3] = FVector3f(FNiagaraUIntVectorToFVector(FNiagaraUIntVectorSwizzle(rand, 0, 0, 0) & MGradientMask)) * MGradientScale - 1;
+	gvec[1][3] = FVector3f(FNiagaraUIntVectorToFVector(FNiagaraUIntVectorSwizzle(rand, 1, 1, 1) & MGradientMask)) * MGradientScale - 1;
+	gvec[2][3] = FVector3f(FNiagaraUIntVectorToFVector(FNiagaraUIntVectorSwizzle(rand, 2, 2, 2) & MGradientMask)) * MGradientScale - 1;
+	grad[0][3] = FVector3f::DotProduct(gvec[0][3], fv[3]);
+	grad[1][3] = FVector3f::DotProduct(gvec[1][3], fv[3]);
+	grad[2][3] = FVector3f::DotProduct(gvec[2][3], fv[3]);
+
+	FVector4f sv = SimplexSmooth(fv);
 	FNiagaraMatrix3x4 ds = SimplexDSmooth(fv);
 
 	FNiagaraMatrix3x4 jacobian;
-	jacobian[0] = FVector4(MulFVector4AndFNiagaraMatrix4x3(sv, gvec[0]) + MulFNiagaraMatrix3x4FAndVector4(ds, grad[0]), Dot4(sv, grad[0]));
-	jacobian[1] = FVector4(MulFVector4AndFNiagaraMatrix4x3(sv, gvec[1]) + MulFNiagaraMatrix3x4FAndVector4(ds, grad[1]), Dot4(sv, grad[1]));
-	jacobian[2] = FVector4(MulFVector4AndFNiagaraMatrix4x3(sv, gvec[2]) + MulFNiagaraMatrix3x4FAndVector4(ds, grad[2]), Dot4(sv, grad[2]));
+	jacobian[0] = FVector4f(MulFVector4AndFNiagaraMatrix4x3(sv, gvec[0]) + MulFNiagaraMatrix3x4FAndVector4(ds, grad[0]), Dot4(sv, grad[0]));
+	jacobian[1] = FVector4f(MulFVector4AndFNiagaraMatrix4x3(sv, gvec[1]) + MulFNiagaraMatrix3x4FAndVector4(ds, grad[1]), Dot4(sv, grad[1]));
+	jacobian[2] = FVector4f(MulFVector4AndFNiagaraMatrix4x3(sv, gvec[2]) + MulFNiagaraMatrix3x4FAndVector4(ds, grad[2]), Dot4(sv, grad[2]));
 
 	return jacobian;
 }
@@ -310,7 +318,7 @@ UNiagaraDataInterfaceCurlNoise::UNiagaraDataInterfaceCurlNoise(FObjectInitialize
 	: Super(ObjectInitializer)
 	, Seed(0)
 {
-	OffsetFromSeed = FNiagaraUIntVectorToFVector(Rand3DPCG16(FIntVector(Seed, Seed, Seed))) / 100.0;
+	OffsetFromSeed = FNiagaraUIntVectorToFVector(Rand3DPCG16(FIntVector(Seed, Seed, Seed))) / 100.0f;
 
 	Proxy.Reset(new FNiagaraDataInterfaceProxyCurlNoise(OffsetFromSeed));
 }
@@ -329,7 +337,7 @@ void UNiagaraDataInterfaceCurlNoise::PostInitProperties()
 void UNiagaraDataInterfaceCurlNoise::PostLoad()
 {
 	Super::PostLoad();
-	OffsetFromSeed = FNiagaraUIntVectorToFVector(Rand3DPCG16(FIntVector(Seed, Seed, Seed))) / 100.0;
+	OffsetFromSeed = FNiagaraUIntVectorToFVector(Rand3DPCG16(FIntVector(Seed, Seed, Seed))) / 100.0f;
 
 	MarkRenderDataDirty();
 }
@@ -352,7 +360,7 @@ void UNiagaraDataInterfaceCurlNoise::PostEditChangeProperty(struct FPropertyChan
 	if (PropertyChangedEvent.Property && PropertyChangedEvent.Property->GetFName() == GET_MEMBER_NAME_CHECKED(UNiagaraDataInterfaceCurlNoise, Seed))
 	{
 		// NOTE: Calculate the offset based on the seed on-change instead of on every invocation for every particle...
-		OffsetFromSeed = FNiagaraUIntVectorToFVector(Rand3DPCG16(FIntVector(Seed, Seed, Seed))) / 100.0;
+		OffsetFromSeed = FNiagaraUIntVectorToFVector(Rand3DPCG16(FIntVector(Seed, Seed, Seed))) / 100.0f;
 	}
 
 	MarkRenderDataDirty();
@@ -417,7 +425,7 @@ void UNiagaraDataInterfaceCurlNoise::SampleNoiseField(FVectorVMExternalFunctionC
 
 	for (int32 InstanceIdx = 0; InstanceIdx < Context.GetNumInstances(); ++InstanceIdx)
 	{
-		FVector InCoords = FVector(XParam.GetAndAdvance(), YParam.GetAndAdvance(), ZParam.GetAndAdvance());
+		const FVector3f InCoords = FVector3f(XParam.GetAndAdvance(), YParam.GetAndAdvance(), ZParam.GetAndAdvance());
 
 		// See comments to JacobianSimplex_ALU in Random.ush
 		FNiagaraMatrix3x4 J = JacobianSimplex_ALU(InCoords + OffsetFromSeed);
@@ -475,11 +483,9 @@ void UNiagaraDataInterfaceCurlNoise::PushToRenderThreadImpl()
 {
 	FNiagaraDataInterfaceProxyCurlNoise* RT_Proxy = GetProxyAs<FNiagaraDataInterfaceProxyCurlNoise>();
 
-	FVector RT_Offset = OffsetFromSeed;
-
 	// Push Updates to Proxy.
 	ENQUEUE_RENDER_COMMAND(FUpdateDIColorCurve)(
-		[RT_Proxy, RT_Offset](FRHICommandListImmediate& RHICmdList)
+		[RT_Proxy, RT_Offset=OffsetFromSeed](FRHICommandListImmediate& RHICmdList)
 	{
 		RT_Proxy->OffsetFromSeed = RT_Offset;
 	});
diff --git a/Engine/Plugins/FX/Niagara/Source/Niagara/Private/NiagaraDataInterfaceCurveBase.cpp b/Engine/Plugins/FX/Niagara/Source/Niagara/Private/NiagaraDataInterfaceCurveBase.cpp
--- a/Engine/Plugins/FX/Niagara/Source/Niagara/Private/NiagaraDataInterfaceCurveBase.cpp
+++ b/Engine/Plugins/FX/Niagara/Source/Niagara/Private/NiagaraDataInterfaceCurveBase.cpp
@@ -156,7 +156,7 @@ void UNiagaraDataInterfaceCurveBase::UpdateLUT(bool bFromSerialize)
 	{
 		ShaderLUT = BuildLUT(CurveLUTDefaultWidth);
 		OptimizeLUT();
-		LUTNumSamplesMinusOne = (ShaderLUT.Num() / GetCurveNumElems()) - 1;
+		LUTNumSamplesMinusOne = float((ShaderLUT.Num() / GetCurveNumElems()) - 1);
 	}
 
 	if (!bUseLUT || (ShaderLUT.Num() == 0))
@@ -401,7 +401,7 @@ void UNiagaraDataInterfaceCurveBase::PushToRenderThreadImpl()
 			RT_Proxy->LUTMinTime = RT_LUTMinTime;
 			RT_Proxy->LUTMaxTime = RT_LUTMaxTime;
 			RT_Proxy->LUTInvTimeRange = RT_LUTInvTimeRange;
-			RT_Proxy->CurveLUTNumMinusOne = RT_CurveLUTNumMinusOne;
+			RT_Proxy->CurveLUTNumMinusOne = uint32(RT_CurveLUTNumMinusOne);
 			RT_Proxy->LUTOffset = RT_LUTOffset;
 
 			DEC_MEMORY_STAT_BY(STAT_NiagaraGPUDataInterfaceMemory, RT_Proxy->CurveLUT.NumBytes);
diff --git a/Engine/Plugins/FX/Niagara/Source/Niagara/Private/NiagaraDataInterfaceGrid2DCollection.cpp b/Engine/Plugins/FX/Niagara/Source/Niagara/Private/NiagaraDataInterfaceGrid2DCollection.cpp
--- a/Engine/Plugins/FX/Niagara/Source/Niagara/Private/NiagaraDataInterfaceGrid2DCollection.cpp
+++ b/Engine/Plugins/FX/Niagara/Source/Niagara/Private/NiagaraDataInterfaceGrid2DCollection.cpp
@@ -1803,8 +1803,8 @@ bool UNiagaraDataInterfaceGrid2DCollection::InitPerInstanceData(void* PerInstanc
 	{
 		FVector2D::FReal CellSize = FMath::Max(WorldBBoxSize.X, WorldBBoxSize.Y) / NumCellsMaxAxis;
 
-		InstanceData->NumCells.X = WorldBBoxSize.X / CellSize;
-		InstanceData->NumCells.Y = WorldBBoxSize.Y / CellSize;
+		InstanceData->NumCells.X = int32(FMath::FloorToInt(WorldBBoxSize.X / CellSize));
+		InstanceData->NumCells.Y = int32(FMath::FloorToInt(WorldBBoxSize.Y / CellSize));
 
 		// Pad grid by 1 voxel if our computed bounding box is too small
 		if (WorldBBoxSize.X > WorldBBoxSize.Y && !FMath::IsNearlyEqual(CellSize * InstanceData->NumCells.Y, WorldBBoxSize.Y))
@@ -1886,7 +1886,7 @@ bool UNiagaraDataInterfaceGrid2DCollection::InitPerInstanceData(void* PerInstanc
 		for (int32 i = 0; i < RT_InstanceData.Vars.Num(); i++)
 		{
 			TargetData->Vars.Emplace(RT_InstanceData.Vars[i].GetName());
-			TargetData->VarComponents.Emplace(RT_InstanceData.Vars[i].GetType().GetSize() / sizeof(float));
+			TargetData->VarComponents.Add(RT_InstanceData.Vars[i].GetType().GetSize() / sizeof(float));
 		}
 #if WITH_EDITORONLY_DATA
 		TargetData->bPreviewGrid = RT_InstanceData.bPreviewGrid;
@@ -2565,7 +2565,7 @@ void FGrid2DCollectionRWInstanceData_RenderThread::BeginSimulate(FRDGBuilder& Gr
 		DestinationData = new FGrid2DBuffer();
 		Buffers.Emplace(DestinationData);
 
-		const FRDGTextureDesc TextureDesc = FRDGTextureDesc::Create2DArray(NumCells, PixelFormat, FClearValueBinding::Black, ETextureCreateFlags::ShaderResource | ETextureCreateFlags::UAV, NumAttributes);
+		const FRDGTextureDesc TextureDesc = FRDGTextureDesc::Create2DArray(NumCells, PixelFormat, FClearValueBinding::Black, ETextureCreateFlags::ShaderResource | ETextureCreateFlags::UAV, uint16(NumAttributes));
 
 		const TCHAR* GridTextureName = TEXT("Grid2D::GridTexture");
 	#if 0
diff --git a/Engine/Plugins/FX/Niagara/Source/Niagara/Private/NiagaraDataInterfaceGrid3DCollection.cpp b/Engine/Plugins/FX/Niagara/Source/Niagara/Private/NiagaraDataInterfaceGrid3DCollection.cpp
--- a/Engine/Plugins/FX/Niagara/Source/Niagara/Private/NiagaraDataInterfaceGrid3DCollection.cpp
+++ b/Engine/Plugins/FX/Niagara/Source/Niagara/Private/NiagaraDataInterfaceGrid3DCollection.cpp
@@ -2921,15 +2921,15 @@ void UNiagaraDataInterfaceGrid3DCollection::SetShaderParameters(const FNiagaraDa
 		{
 			const FIntVector AttributeTileIndex(iAttribute % ProxyData->NumTiles.X, (iAttribute / ProxyData->NumTiles.X) % ProxyData->NumTiles.Y, iAttribute / (ProxyData->NumTiles.X * ProxyData->NumTiles.Y));
 			PerAttributeData[iAttribute] = FVector4f(
-				AttributeTileIndex.X * ProxyData->NumCells.X,
-				AttributeTileIndex.Y * ProxyData->NumCells.Y,
-				AttributeTileIndex.Z * ProxyData->NumCells.Z,
+				float(AttributeTileIndex.X * ProxyData->NumCells.X),
+				float(AttributeTileIndex.Y * ProxyData->NumCells.Y),
+				float(AttributeTileIndex.Z * ProxyData->NumCells.Z),
 				0
 			);
 			PerAttributeData[iAttribute + ProxyData->TotalNumAttributes] = FVector4f(
-				(1.0f / ProxyData->NumTiles.X) * float(AttributeTileIndex.X),
-				(1.0f / ProxyData->NumTiles.Y) * float(AttributeTileIndex.Y),
-				(1.0f / ProxyData->NumTiles.Z) * float(AttributeTileIndex.Z),
+				(1.0f / float(ProxyData->NumTiles.X)) * float(AttributeTileIndex.X),
+				(1.0f / float(ProxyData->NumTiles.Y)) * float(AttributeTileIndex.Y),
+				(1.0f / float(ProxyData->NumTiles.Z)) * float(AttributeTileIndex.Z),
 				0.0f
 			);
 		}
@@ -3268,7 +3268,7 @@ bool UNiagaraDataInterfaceGrid3DCollection::InitPerInstanceData(void* PerInstanc
 	}
 	else if (SetResolutionMethod == ESetResolutionMethod::MaxAxis)
 	{
-		InstanceData->CellSize = FVector(FMath::Max<float>(FMath::Max(WorldBBoxSize.X, WorldBBoxSize.Y), WorldBBoxSize.Z) / ((float)NumCellsMaxAxis));
+		InstanceData->CellSize = FVector(FMath::Max(FMath::Max(WorldBBoxSize.X, WorldBBoxSize.Y), WorldBBoxSize.Z) / double(NumCellsMaxAxis));
 	}
 	else if (SetResolutionMethod == ESetResolutionMethod::CellSize)
 	{
@@ -3278,9 +3278,9 @@ bool UNiagaraDataInterfaceGrid3DCollection::InitPerInstanceData(void* PerInstanc
 	// compute world bounds and padding based on cell size
 	if (SetResolutionMethod == ESetResolutionMethod::MaxAxis || SetResolutionMethod == ESetResolutionMethod::CellSize)
 	{
-		InstanceData->NumCells.X = WorldBBoxSize.X / InstanceData->CellSize[0];
-		InstanceData->NumCells.Y = WorldBBoxSize.Y / InstanceData->CellSize[0];
-		InstanceData->NumCells.Z = WorldBBoxSize.Z / InstanceData->CellSize[0];
+		InstanceData->NumCells.X = int32(FMath::FloorToInt(WorldBBoxSize.X / InstanceData->CellSize[0]));
+		InstanceData->NumCells.Y = int32(FMath::FloorToInt(WorldBBoxSize.Y / InstanceData->CellSize[0]));
+		InstanceData->NumCells.Z = int32(FMath::FloorToInt(WorldBBoxSize.Z / InstanceData->CellSize[0]));
 
 		// Pad grid by 1 cell if our computed bounding box is too small
 		if (WorldBBoxSize.X > WorldBBoxSize.Y && WorldBBoxSize.X > WorldBBoxSize.Z)
@@ -3424,11 +3424,11 @@ bool UNiagaraDataInterfaceGrid3DCollection::InitPerInstanceData(void* PerInstanc
 		// Compute number of tiles based on resolution of individual attributes
 
 		// #todo(dmp): refactor
-		int32 MaxDim = 2048;
-		int32 MaxTilesX = floor(MaxDim / InstanceData->NumCells.X);
-		int32 MaxTilesY = floor(MaxDim / InstanceData->NumCells.Y);
-		int32 MaxTilesZ = floor(MaxDim / InstanceData->NumCells.Z);
-		int32 MaxAttributes = MaxTilesX * MaxTilesY * MaxTilesZ;
+		const float MaxDim = 2048;
+		const int32 MaxTilesX = FMath::FloorToInt(MaxDim / float(InstanceData->NumCells.X));
+		const int32 MaxTilesY = FMath::FloorToInt(MaxDim / float(InstanceData->NumCells.Y));
+		const int32 MaxTilesZ = FMath::FloorToInt(MaxDim / float(InstanceData->NumCells.Z));
+		const int32 MaxAttributes = MaxTilesX * MaxTilesY * MaxTilesZ;
 
 		if (InstanceData->NumCells.X > MaxDim || InstanceData->NumCells.Y > MaxDim || InstanceData->NumCells.Z > MaxDim)
 		{
@@ -3443,9 +3443,9 @@ bool UNiagaraDataInterfaceGrid3DCollection::InitPerInstanceData(void* PerInstanc
 		}
 
 		// need to determine number of tiles in x and y based on number of attributes and max dimension size
-		int32 NumTilesX = FMath::Min<int32>(MaxTilesX, NumAttribChannelsFound);
-		int32 NumTilesY = FMath::Min<int32>(MaxTilesY, ceil(1.0 * NumAttribChannelsFound / NumTilesX));
-		int32 NumTilesZ = FMath::Min<int32>(MaxTilesZ, ceil(1.0 * NumAttribChannelsFound / (NumTilesX * NumTilesY)));
+		const int32 NumTilesX = FMath::Min<int32>(MaxTilesX, NumAttribChannelsFound);
+		const int32 NumTilesY = FMath::Min<int32>(MaxTilesY, FMath::CeilToInt(1.0f * float(NumAttribChannelsFound) / float(NumTilesX)));
+		const int32 NumTilesZ = FMath::Min<int32>(MaxTilesZ, FMath::CeilToInt(1.0f * float(NumAttribChannelsFound) / float(NumTilesX * NumTilesY)));
 
 		InstanceData->NumTiles.X = NumTilesX;
 		InstanceData->NumTiles.Y = NumTilesY;
@@ -3525,7 +3525,7 @@ bool UNiagaraDataInterfaceGrid3DCollection::InitPerInstanceData(void* PerInstanc
 			for (int32 i = 0; i < RT_InstanceData.Vars.Num(); i++)
 			{
 				TargetData->Vars.Emplace(RT_InstanceData.Vars[i].GetName());
-				TargetData->VarComponents.Emplace(RT_InstanceData.Vars[i].GetType().GetSize() / sizeof(float));
+				TargetData->VarComponents.Add(RT_InstanceData.Vars[i].GetType().GetSize() / sizeof(float));
 			}
 #if WITH_EDITOR
 			TargetData->bPreviewGrid = RT_InstanceData.bPreviewGrid;
@@ -3842,9 +3842,9 @@ void UNiagaraDataInterfaceGrid3DCollection::VMGetWorldBBoxSize(FVectorVMExternal
 
 	for (int32 InstanceIdx = 0; InstanceIdx < Context.GetNumInstances(); ++InstanceIdx)
 	{
-		*OutWorldBoundsX.GetDestAndAdvance() = InstData->WorldBBoxSize.X;
-		*OutWorldBoundsY.GetDestAndAdvance() = InstData->WorldBBoxSize.Y;
-		*OutWorldBoundsZ.GetDestAndAdvance() = InstData->WorldBBoxSize.Z;
+		*OutWorldBoundsX.GetDestAndAdvance() = float(InstData->WorldBBoxSize.X);
+		*OutWorldBoundsY.GetDestAndAdvance() = float(InstData->WorldBBoxSize.Y);
+		*OutWorldBoundsZ.GetDestAndAdvance() = float(InstData->WorldBBoxSize.Z);
 	}
 }
 
@@ -3935,11 +3935,11 @@ bool UNiagaraDataInterfaceGrid3DCollection::PerInstanceTickPostSimulate(void* Pe
 		else
 		{
 			// #todo(dmp): refactor
-			int32 MaxDim = 2048;
-			int32 MaxTilesX = floor(MaxDim / InstanceData->NumCells.X);
-			int32 MaxTilesY = floor(MaxDim / InstanceData->NumCells.Y);
-			int32 MaxTilesZ = floor(MaxDim / InstanceData->NumCells.Z);
-			int32 MaxAttributes = MaxTilesX * MaxTilesY * MaxTilesZ;
+			const float MaxDim = 2048.0f;
+			const int32 MaxTilesX = FMath::FloorToInt(MaxDim / float(InstanceData->NumCells.X));
+			const int32 MaxTilesY = FMath::FloorToInt(MaxDim / float(InstanceData->NumCells.Y));
+			const int32 MaxTilesZ = FMath::FloorToInt(MaxDim / float(InstanceData->NumCells.Z));
+			const int32 MaxAttributes = MaxTilesX * MaxTilesY * MaxTilesZ;
 
 			if (InstanceData->NumCells.X > MaxDim || InstanceData->NumCells.Y > MaxDim || InstanceData->NumCells.Z > MaxDim)
 			{
@@ -3960,9 +3960,9 @@ bool UNiagaraDataInterfaceGrid3DCollection::PerInstanceTickPostSimulate(void* Pe
 			}
 
 			// need to determine number of tiles in x and y based on number of attributes and max dimension size
-			int32 NumTilesX = FMath::Min<int32>(MaxTilesX, InstanceData->TotalNumAttributes);
-			int32 NumTilesY = FMath::Min<int32>(MaxTilesY, ceil(1.0 * InstanceData->TotalNumAttributes / NumTilesX));
-			int32 NumTilesZ = FMath::Min<int32>(MaxTilesZ, ceil(1.0 * InstanceData->TotalNumAttributes / (NumTilesX * NumTilesY)));
+			const int32 NumTilesX = FMath::Min<int32>(MaxTilesX, InstanceData->TotalNumAttributes);
+			const int32 NumTilesY = FMath::Min<int32>(MaxTilesY, FMath::CeilToInt(1.0f * float(InstanceData->TotalNumAttributes) / float(NumTilesX)));
+			const int32 NumTilesZ = FMath::Min<int32>(MaxTilesZ, FMath::CeilToInt(1.0f * float(InstanceData->TotalNumAttributes) / float(NumTilesX * NumTilesY)));
 
 			InstanceData->NumTiles.X = NumTilesX;
 			InstanceData->NumTiles.Y = NumTilesY;
@@ -4008,9 +4008,9 @@ void UNiagaraDataInterfaceGrid3DCollection::VMGetCellSize(FVectorVMExternalFunct
 
 	for (int32 InstanceIdx = 0; InstanceIdx < Context.GetNumInstances(); ++InstanceIdx)
 	{
-		*OutCellSizeX.GetDestAndAdvance() = InstData->CellSize.X;
-		*OutCellSizeY.GetDestAndAdvance() = InstData->CellSize.Y;
-		*OutCellSizeZ.GetDestAndAdvance() = InstData->CellSize.Z;
+		*OutCellSizeX.GetDestAndAdvance() = float(InstData->CellSize.X);
+		*OutCellSizeY.GetDestAndAdvance() = float(InstData->CellSize.Y);
+		*OutCellSizeZ.GetDestAndAdvance() = float(InstData->CellSize.Z);
 	}
 }
 
diff --git a/Engine/Plugins/FX/Niagara/Source/Niagara/Private/NiagaraDataInterfaceIntRenderTarget2D.cpp b/Engine/Plugins/FX/Niagara/Source/Niagara/Private/NiagaraDataInterfaceIntRenderTarget2D.cpp
--- a/Engine/Plugins/FX/Niagara/Source/Niagara/Private/NiagaraDataInterfaceIntRenderTarget2D.cpp
+++ b/Engine/Plugins/FX/Niagara/Source/Niagara/Private/NiagaraDataInterfaceIntRenderTarget2D.cpp
@@ -572,8 +572,8 @@ void UNiagaraDataInterfaceIntRenderTarget2D::SetShaderParameters(const FNiagaraD
 
 	// Set Parameters
 	FShaderParameters* Parameters = Context.GetParameterNestedStruct<FShaderParameters>();
-	Parameters->TextureSizeAndInvSize.X = InstanceData.Size.X;
-	Parameters->TextureSizeAndInvSize.Y = InstanceData.Size.Y;
+	Parameters->TextureSizeAndInvSize.X = float(InstanceData.Size.X);
+	Parameters->TextureSizeAndInvSize.Y = float(InstanceData.Size.Y);
 	Parameters->TextureSizeAndInvSize.Z = 1.0f / float(InstanceData.Size.X);
 	Parameters->TextureSizeAndInvSize.W = 1.0f / float(InstanceData.Size.Y);
 
diff --git a/Engine/Plugins/FX/Niagara/Source/Niagara/Private/NiagaraDataInterfaceLandscape.cpp b/Engine/Plugins/FX/Niagara/Source/Niagara/Private/NiagaraDataInterfaceLandscape.cpp
--- a/Engine/Plugins/FX/Niagara/Source/Niagara/Private/NiagaraDataInterfaceLandscape.cpp
+++ b/Engine/Plugins/FX/Niagara/Source/Niagara/Private/NiagaraDataInterfaceLandscape.cpp
@@ -811,7 +811,9 @@ void FNDI_Landscape_SharedResource::Initialize()
 						TArray<uint8> PhysMatRemap;
 						for (const UPhysicalMaterial* ComponentMaterial : CollisionComponent->CookedPhysicalMaterials)
 						{
-							PhysMatRemap.Emplace(ResourceKey.PhysicalMaterials.IndexOfByKey(ComponentMaterial));
+							const int32 RemapIndex = ResourceKey.PhysicalMaterials.IndexOfByKey(ComponentMaterial);
+							ensure(RemapIndex <= TNumericLimits<uint8>::Max());
+							PhysMatRemap.Emplace(uint8(RemapIndex));
 						}
 
 						for (int32 Y = 0; Y < ComponentQuadCount; ++Y)
@@ -975,8 +977,8 @@ FNDI_Landscape_SharedResourceHandle FNDI_Landscape_GeneratedData::GetLandscapeDa
 
 	// at runtime we don't know the potential number of components in the landscape, so we can only clamp the values to start at 0.
 	// Ideally, we'd be able to specify an upper bound as well to further refine the region of overlap.
-	const FIntPoint MinBoundRegion = FIntPoint(SystemBoundsInLandscape.Min.X, SystemBoundsInLandscape.Min.Y).ComponentMax(FIntPoint(0, 0));
-	const FIntPoint MaxBoundRegion = FIntPoint(SystemBoundsInLandscape.Max.X, SystemBoundsInLandscape.Max.Y).ComponentMax(FIntPoint(0, 0));
+	const FIntPoint MinBoundRegion = FIntPoint(int32(SystemBoundsInLandscape.Min.X), int32(SystemBoundsInLandscape.Min.Y)).ComponentMax(FIntPoint(0, 0));	// LWC Precision Loss
+	const FIntPoint MaxBoundRegion = FIntPoint(int32(SystemBoundsInLandscape.Max.X), int32(SystemBoundsInLandscape.Max.Y)).ComponentMax(FIntPoint(0, 0));
 
 	const FIntRect SystemRect = FIntRect(MinBoundRegion, MaxBoundRegion).Scale(1.0f / Landscape->ComponentSizeQuads);
 
diff --git a/Engine/Plugins/FX/Niagara/Source/Niagara/Private/NiagaraDataInterfaceNeighborGrid3D.cpp b/Engine/Plugins/FX/Niagara/Source/Niagara/Private/NiagaraDataInterfaceNeighborGrid3D.cpp
--- a/Engine/Plugins/FX/Niagara/Source/Niagara/Private/NiagaraDataInterfaceNeighborGrid3D.cpp
+++ b/Engine/Plugins/FX/Niagara/Source/Niagara/Private/NiagaraDataInterfaceNeighborGrid3D.cpp
@@ -472,9 +472,9 @@ void UNiagaraDataInterfaceNeighborGrid3D::SetShaderParameters(const FNiagaraData
 	{
 		ShaderParameters->NumCells		= ProxyData->NumCells;
 		ShaderParameters->UnitToUV		= FVector3f(1.0f) / FVector3f(ProxyData->NumCells);
-		ShaderParameters->CellSize.X	= ProxyData->WorldBBoxSize.X / float(ProxyData->NumCells.X);
-		ShaderParameters->CellSize.Y	= ProxyData->WorldBBoxSize.Y / float(ProxyData->NumCells.Y);
-		ShaderParameters->CellSize.Z	= ProxyData->WorldBBoxSize.Z / float(ProxyData->NumCells.Z);
+		ShaderParameters->CellSize.X	= float(ProxyData->WorldBBoxSize.X / double(ProxyData->NumCells.X));
+		ShaderParameters->CellSize.Y	= float(ProxyData->WorldBBoxSize.Y / double(ProxyData->NumCells.Y));
+		ShaderParameters->CellSize.Z	= float(ProxyData->WorldBBoxSize.Z / double(ProxyData->NumCells.Z));
 		ShaderParameters->WorldBBoxSize	= FVector3f(ProxyData->WorldBBoxSize);
 		ShaderParameters->MaxNeighborsPerCellValue = ProxyData->MaxNeighborsPerCell;
 
@@ -532,9 +532,9 @@ bool UNiagaraDataInterfaceNeighborGrid3D::InitPerInstanceData(void* PerInstanceD
 	// compute world bounds and padding based on cell size
 	if (SetResolutionMethod == ESetResolutionMethod::MaxAxis || SetResolutionMethod == ESetResolutionMethod::CellSize)
 	{
-		RT_NumCells.X = WorldBBoxSize.X / TmpCellSize;
-		RT_NumCells.Y = WorldBBoxSize.Y / TmpCellSize;
-		RT_NumCells.Z = WorldBBoxSize.Z / TmpCellSize;
+		RT_NumCells.X = int32(FMath::FloorToInt(WorldBBoxSize.X / TmpCellSize));
+		RT_NumCells.Y = int32(FMath::FloorToInt(WorldBBoxSize.Y / TmpCellSize));
+		RT_NumCells.Z = int32(FMath::FloorToInt(WorldBBoxSize.Z / TmpCellSize));
 
 		// Pad grid by 1 cell if our computed bounding box is too small
 		if (WorldBBoxSize.X > WorldBBoxSize.Y && WorldBBoxSize.X > WorldBBoxSize.Z)
@@ -580,7 +580,7 @@ bool UNiagaraDataInterfaceNeighborGrid3D::InitPerInstanceData(void* PerInstanceD
 	RT_NumCells.Y = FMath::Max(RT_NumCells.Y, 1);
 	RT_NumCells.Z = FMath::Max(RT_NumCells.Z, 1);
 
-	InstanceData->CellSize = TmpCellSize;
+	InstanceData->CellSize = float(TmpCellSize);
 	InstanceData->WorldBBoxSize = RT_WorldBBoxSize;
 	InstanceData->MaxNeighborsPerCell = RT_MaxNeighborsPerCell;	
 	InstanceData->NumCells = RT_NumCells;
@@ -671,7 +671,7 @@ bool UNiagaraDataInterfaceNeighborGrid3D::PerInstanceTickPostSimulate(void* PerI
 	{
 		InstanceData->bNeedsRealloc = false;
 
-		InstanceData->CellSize = (InstanceData->WorldBBoxSize / FVector(InstanceData->NumCells.X, InstanceData->NumCells.Y, InstanceData->NumCells.Z))[0];
+		InstanceData->CellSize = float((InstanceData->WorldBBoxSize / FVector(InstanceData->NumCells.X, InstanceData->NumCells.Y, InstanceData->NumCells.Z))[0]);
 
 		FNiagaraDataInterfaceProxyNeighborGrid3D* RT_Proxy = GetProxyAs<FNiagaraDataInterfaceProxyNeighborGrid3D>();
 		ENQUEUE_RENDER_COMMAND(FUpdateData)(
diff --git a/Engine/Plugins/FX/Niagara/Source/Niagara/Private/NiagaraDataInterfaceRasterizationGrid3D.cpp b/Engine/Plugins/FX/Niagara/Source/Niagara/Private/NiagaraDataInterfaceRasterizationGrid3D.cpp
--- a/Engine/Plugins/FX/Niagara/Source/Niagara/Private/NiagaraDataInterfaceRasterizationGrid3D.cpp
+++ b/Engine/Plugins/FX/Niagara/Source/Niagara/Private/NiagaraDataInterfaceRasterizationGrid3D.cpp
@@ -915,16 +915,16 @@ bool UNiagaraDataInterfaceRasterizationGrid3D::InitPerInstanceData(void* PerInst
 
 	// Compute number of tiles based on resolution of individual attributes
 	// #todo(dmp): refactor
-	int32 MaxDim = 2048;
-	int32 MaxTilesX = floor(MaxDim / InstanceData->NumCells.X);
-	int32 MaxTilesY = floor(MaxDim / InstanceData->NumCells.Y);
-	int32 MaxTilesZ = floor(MaxDim / InstanceData->NumCells.Z);
-	int32 MaxAttributes = MaxTilesX * MaxTilesY * MaxTilesZ;
+	const float MaxDim = 2048;
+	const int32 MaxTilesX = FMath::FloorToInt(MaxDim / float(InstanceData->NumCells.X));
+	const int32 MaxTilesY = FMath::FloorToInt(MaxDim / float(InstanceData->NumCells.Y));
+	const int32 MaxTilesZ = FMath::FloorToInt(MaxDim / float(InstanceData->NumCells.Z));
+	const int32 MaxAttributes = MaxTilesX * MaxTilesY * MaxTilesZ;
 
 	// need to determine number of tiles in x and y based on number of attributes and max dimension size
-	int32 NumTilesX = FMath::Min<int32>(MaxTilesX, NumAttribChannelsFound);
-	int32 NumTilesY = FMath::Min<int32>(MaxTilesY, ceil(1.0 * NumAttribChannelsFound / NumTilesX));
-	int32 NumTilesZ = FMath::Min<int32>(MaxTilesZ, ceil(1.0 * NumAttribChannelsFound / (NumTilesX * NumTilesY)));
+	const int32 NumTilesX = FMath::Min<int32>(MaxTilesX, NumAttribChannelsFound);
+	const int32 NumTilesY = FMath::Min<int32>(MaxTilesY, FMath::CeilToInt(1.0f * float(NumAttribChannelsFound) / float(NumTilesX)));
+	const int32 NumTilesZ = FMath::Min<int32>(MaxTilesZ, FMath::CeilToInt(1.0f * float(NumAttribChannelsFound) / float(NumTilesX * NumTilesY)));
 
 	InstanceData->NumTiles.X = NumTilesX;
 	InstanceData->NumTiles.Y = NumTilesY;
@@ -1027,7 +1027,7 @@ void UNiagaraDataInterfaceRasterizationGrid3D::VMSetFloatResetValue(FVectorVMExt
 		if (bSuccess)
 		{
 			// Quantize value
-			InstData->ResetValue = NewResetValue * InstData->Precision;
+			InstData->ResetValue = FMath::FloorToInt(NewResetValue * InstData->Precision);
 		}
 	}
 }
@@ -1043,16 +1043,16 @@ bool UNiagaraDataInterfaceRasterizationGrid3D::PerInstanceTickPostSimulate(void*
 		
 		// Compute number of tiles based on resolution of individual attributes
 		// #todo(dmp): refactor
-		int32 MaxDim = 2048;
-		int32 MaxTilesX = floor(MaxDim / InstanceData->NumCells.X);
-		int32 MaxTilesY = floor(MaxDim / InstanceData->NumCells.Y);
-		int32 MaxTilesZ = floor(MaxDim / InstanceData->NumCells.Z);
-		int32 MaxAttributes = MaxTilesX * MaxTilesY * MaxTilesZ;
+		const float MaxDim = 2048;
+		const int32 MaxTilesX = FMath::FloorToInt(MaxDim / float(InstanceData->NumCells.X));
+		const int32 MaxTilesY = FMath::FloorToInt(MaxDim / float(InstanceData->NumCells.Y));
+		const int32 MaxTilesZ = FMath::FloorToInt(MaxDim / float(InstanceData->NumCells.Z));
+		const int32 MaxAttributes = MaxTilesX * MaxTilesY * MaxTilesZ;
 
 		// need to determine number of tiles in x and y based on number of attributes and max dimension size
-		int32 NumTilesX = FMath::Min<int32>(MaxTilesX, InstanceData->TotalNumAttributes);
-		int32 NumTilesY = FMath::Min<int32>(MaxTilesY, ceil(1.0 * InstanceData->TotalNumAttributes / NumTilesX));
-		int32 NumTilesZ = FMath::Min<int32>(MaxTilesZ, ceil(1.0 * InstanceData->TotalNumAttributes / (NumTilesX * NumTilesY)));
+		const int32 NumTilesX = FMath::Min<int32>(MaxTilesX, InstanceData->TotalNumAttributes);
+		const int32 NumTilesY = FMath::Min<int32>(MaxTilesY, FMath::CeilToInt(1.0f * float(InstanceData->TotalNumAttributes) / float(NumTilesX)));
+		const int32 NumTilesZ = FMath::Min<int32>(MaxTilesZ, FMath::CeilToInt(1.0f * float(InstanceData->TotalNumAttributes) / float(NumTilesX * NumTilesY)));
 
 		InstanceData->NumTiles.X = NumTilesX;
 		InstanceData->NumTiles.Y = NumTilesY;
@@ -1153,15 +1153,15 @@ void FNiagaraDataInterfaceProxyRasterizationGrid3D::PreStage(const FNDIGpuComput
 		{
 			const FIntVector AttributeTileIndex(iAttribute % InstanceData.NumTiles.X, (iAttribute / InstanceData.NumTiles.X) % InstanceData.NumTiles.Y, iAttribute / (InstanceData.NumTiles.X * InstanceData.NumTiles.Y));
 			PerAttributeData[iAttribute] = FVector4f(
-				AttributeTileIndex.X * InstanceData.NumCells.X,
-				AttributeTileIndex.Y * InstanceData.NumCells.Y,
-				AttributeTileIndex.Z * InstanceData.NumCells.Z,
-				0
+				float(AttributeTileIndex.X * InstanceData.NumCells.X),
+				float(AttributeTileIndex.Y * InstanceData.NumCells.Y),
+				float(AttributeTileIndex.Z * InstanceData.NumCells.Z),
+				0.0f
 			);
 			PerAttributeData[iAttribute + InstanceData.TotalNumAttributes] = FVector4f(
-				(1.0f / InstanceData.NumTiles.X) * float(AttributeTileIndex.X),
-				(1.0f / InstanceData.NumTiles.Y) * float(AttributeTileIndex.Y),
-				(1.0f / InstanceData.NumTiles.Z) * float(AttributeTileIndex.Z),
+				(1.0f / float(InstanceData.NumTiles.X)) * float(AttributeTileIndex.X),
+				(1.0f / float(InstanceData.NumTiles.Y)) * float(AttributeTileIndex.Y),
+				(1.0f / float(InstanceData.NumTiles.Z)) * float(AttributeTileIndex.Z),
 				0.0f
 			);
 		}
diff --git a/Engine/Plugins/FX/Niagara/Source/Niagara/Private/NiagaraDataInterfaceRigidMeshCollisionQuery.cpp b/Engine/Plugins/FX/Niagara/Source/Niagara/Private/NiagaraDataInterfaceRigidMeshCollisionQuery.cpp
--- a/Engine/Plugins/FX/Niagara/Source/Niagara/Private/NiagaraDataInterfaceRigidMeshCollisionQuery.cpp
+++ b/Engine/Plugins/FX/Niagara/Source/Niagara/Private/NiagaraDataInterfaceRigidMeshCollisionQuery.cpp
@@ -363,7 +363,7 @@ void UpdateAssetArrays(TConstArrayView<TComponentType*> Components, const FVecto
 		FTransform MeshTransform = CreateElementTransform(Component, BodySetup);
 		MeshTransform.AddToTranslation(LWCTile * -FLargeWorldRenderScalar::GetTileSize());
 
-		auto CurrMeshScale = MeshTransform.GetScale3D();
+		const FVector3f CurrMeshScale(MeshTransform.GetScale3D());
 
 		const int32 ComponentIdIndex = OutAssetArrays->UniqueCompnentId.AddUnique(Component->ComponentId);
 
@@ -458,9 +458,9 @@ void UpdateAssetArrays(TConstArrayView<TComponentType*> Components, const FVecto
 			if (InitializeStatics)
 			{			
 				// local bounds extent is half the world extents of the bounding box in local space
-				Extent *= 2;
+				Extent *= 2.0;
 
-				OutAssetArrays->ElementExtent[BoxIndex] = FVector4f(Extent.X, Extent.Y, Extent.Z, 0);
+				OutAssetArrays->ElementExtent[BoxIndex] = FVector4f(float(Extent.X), float(Extent.Y), float(Extent.Z), 0);
 				OutAssetArrays->MeshScale[BoxIndex] = FVector4f(CurrMeshScale.X, CurrMeshScale.Y, CurrMeshScale.Z, 0);
 				OutAssetArrays->PhysicsType[BoxIndex] = true;
 				OutAssetArrays->ComponentIdIndex[BoxIndex] = ComponentIdIndex;
diff --git a/Engine/Plugins/FX/Niagara/Source/Niagara/Private/NiagaraDataInterfaceSkeletalMesh.cpp b/Engine/Plugins/FX/Niagara/Source/Niagara/Private/NiagaraDataInterfaceSkeletalMesh.cpp
--- a/Engine/Plugins/FX/Niagara/Source/Niagara/Private/NiagaraDataInterfaceSkeletalMesh.cpp
+++ b/Engine/Plugins/FX/Niagara/Source/Niagara/Private/NiagaraDataInterfaceSkeletalMesh.cpp
@@ -1228,7 +1228,7 @@ void FSkeletalMeshGpuDynamicBufferProxy::NewFrame(const FNDISkeletalMesh_Instanc
 			for (int i=0; i < BoneTransforms.Num(); ++i )
 			{
 				const FTransform& BoneTransform = BoneTransforms[i];
-				const FQuat Rotation = BoneTransform.GetRotation();
+				const FQuat4f Rotation = FQuat4f(BoneTransform.GetRotation());
 				const int32 ParentIndex = ReferenceSkeleton ? ReferenceSkeleton->GetParentIndex(i) : -1;
 				BoneSamplingData.Emplace((FVector3f)BoneTransform.GetLocation()); // LWC_TODO: precision loss
 				BoneSamplingData.Emplace(Rotation.X, Rotation.Y, Rotation.Z, Rotation.W);
diff --git a/Engine/Plugins/FX/Niagara/Source/Niagara/Private/NiagaraDataInterfaceSpline.cpp b/Engine/Plugins/FX/Niagara/Source/Niagara/Private/NiagaraDataInterfaceSpline.cpp
--- a/Engine/Plugins/FX/Niagara/Source/Niagara/Private/NiagaraDataInterfaceSpline.cpp
+++ b/Engine/Plugins/FX/Niagara/Source/Niagara/Private/NiagaraDataInterfaceSpline.cpp
@@ -682,7 +682,10 @@ bool UNiagaraDataInterfaceSpline::PerInstanceTick(void* PerInstanceData, FNiagar
 				FVector4f* PositionBufferData = static_cast<FVector4f*>(RHILockBuffer(TargetData->SplinePositionsLUT.Buffer, 0, BufferSize, EResourceLockMode::RLM_WriteOnly));
 				for (int32 Index = 0; Index < rtShaderLUT.Positions.Num(); Index++)
 				{
-					PositionBufferData[Index] = FVector4f(rtShaderLUT.Positions[Index].X, rtShaderLUT.Positions[Index].Y, rtShaderLUT.Positions[Index].Z);
+					PositionBufferData[Index].X = float(rtShaderLUT.Positions[Index].X);	// LWC Precision Loss
+					PositionBufferData[Index].Y = float(rtShaderLUT.Positions[Index].Y);
+					PositionBufferData[Index].Z = float(rtShaderLUT.Positions[Index].Z);
+					PositionBufferData[Index].W = 1.0f;
 				}
 				RHIUnlockBuffer(TargetData->SplinePositionsLUT.Buffer);
 		
@@ -692,7 +695,10 @@ bool UNiagaraDataInterfaceSpline::PerInstanceTick(void* PerInstanceData, FNiagar
 				FVector4f* ScaleBufferData = static_cast<FVector4f*>(RHILockBuffer(TargetData->SplineScalesLUT.Buffer, 0, BufferSize, EResourceLockMode::RLM_WriteOnly));
 				for (int32 Index = 0; Index < rtShaderLUT.Scales.Num(); Index++)
 				{
-					ScaleBufferData[Index] = FVector4f(rtShaderLUT.Scales[Index].X, rtShaderLUT.Scales[Index].Y, rtShaderLUT.Scales[Index].Z);
+					ScaleBufferData[Index].X = float(rtShaderLUT.Scales[Index].X);
+					ScaleBufferData[Index].Y = float(rtShaderLUT.Scales[Index].Y);
+					ScaleBufferData[Index].Z = float(rtShaderLUT.Scales[Index].Z);
+					ScaleBufferData[Index].W = 1.0f;
 				}
 				RHIUnlockBuffer(TargetData->SplineScalesLUT.Buffer);
 				
@@ -1002,7 +1008,7 @@ float FNDISpline_InstanceData::EvaluateFindNearestPosition<TIntegralConstant<boo
 	int32 KeyToNearest = 0;
 	for (int32 i=0; i < SplineLUT.Positions.Num(); i++)
 	{
-		const float Distance = FVector::DistSquared(InPosition, SplineLUT.Positions[i]);
+		const float Distance = float(FVector::DistSquared(InPosition, SplineLUT.Positions[i]));
 		if (Distance < MinDistance)
 		{
 			MinDistance = Distance;
@@ -1082,9 +1088,9 @@ void UNiagaraDataInterfaceSpline::SampleSplinePositionByUnitDistance(FVectorVMEx
 			FVector Pos = InstData->GetLocationAtDistanceAlongSpline<UseLUT>(DistanceUnitDistance * SplineLength, ESplineCoordinateSpace::Local);
 			TransformHandler.TransformPosition(Pos, InstData->Transform);
 
-			*OutPosX.GetDest() = Pos.X;
-			*OutPosY.GetDest() = Pos.Y;
-			*OutPosZ.GetDest() = Pos.Z;
+			*OutPosX.GetDest() = float(Pos.X);	// LWC Precision Loss
+			*OutPosY.GetDest() = float(Pos.Y);
+			*OutPosZ.GetDest() = float(Pos.Z);
 			SplineSampleParam.Advance();
 			OutPosX.Advance();
 			OutPosY.Advance();
@@ -1098,9 +1104,9 @@ void UNiagaraDataInterfaceSpline::SampleSplinePositionByUnitDistance(FVectorVMEx
 			FVector Pos = FVector(EForceInit::ForceInitToZero);
 			TransformHandler.TransformPosition(Pos, InstData->Transform);
 
-			*OutPosX.GetDest() = Pos.X;
-			*OutPosY.GetDest() = Pos.Y;
-			*OutPosZ.GetDest() = Pos.Z;
+			*OutPosX.GetDest() = float(Pos.X);	// LWC Precision Loss
+			*OutPosY.GetDest() = float(Pos.Y);
+			*OutPosZ.GetDest() = float(Pos.Z);
 			SplineSampleParam.Advance();
 			OutPosX.Advance();
 			OutPosY.Advance();
@@ -1131,20 +1137,20 @@ void UNiagaraDataInterfaceSpline::SampleSplineRotationByUnitDistance(FVectorVMEx
 			FQuat Quat = InstData->GetQuaternionAtDistanceAlongSpline<UseLUT>(DistanceUnitDistance * SplineLength, ESplineCoordinateSpace::Local);
 			TransformHandler.TransformRotation(Quat, TransformQuat);
 
-			*OutQuatX.GetDestAndAdvance() = Quat.X;
-			*OutQuatY.GetDestAndAdvance() = Quat.Y;
-			*OutQuatZ.GetDestAndAdvance() = Quat.Z;
-			*OutQuatW.GetDestAndAdvance() = Quat.W;
+			*OutQuatX.GetDestAndAdvance() = float(Quat.X);	// LWC Precison Loss
+			*OutQuatY.GetDestAndAdvance() = float(Quat.Y);
+			*OutQuatZ.GetDestAndAdvance() = float(Quat.Z);
+			*OutQuatW.GetDestAndAdvance() = float(Quat.W);
 		}
 	}
 	else
 	{
 		for (int32 i = 0; i < Context.GetNumInstances(); ++i)
 		{
-			*OutQuatX.GetDestAndAdvance() = FQuat::Identity.X;
-			*OutQuatY.GetDestAndAdvance() = FQuat::Identity.Y;
-			*OutQuatZ.GetDestAndAdvance() = FQuat::Identity.Z;
-			*OutQuatW.GetDestAndAdvance() = FQuat::Identity.W;
+			*OutQuatX.GetDestAndAdvance() = FQuat4f::Identity.X;
+			*OutQuatY.GetDestAndAdvance() = FQuat4f::Identity.Y;
+			*OutQuatZ.GetDestAndAdvance() = FQuat4f::Identity.Z;
+			*OutQuatW.GetDestAndAdvance() = FQuat4f::Identity.W;
 		}
 	}
 }
@@ -1169,9 +1175,9 @@ void UNiagaraDataInterfaceSpline::SampleSplineUpVectorByUnitDistance(FVectorVMEx
 			FVector Pos = InstData->GetUpVectorAtDistanceAlongSpline<UseLUT>(DistanceUnitDistance * SplineLength, ESplineCoordinateSpace::Local);
 			TransformHandler.TransformVector(Pos, InstData->Transform);
 
-			*OutPosX.GetDest() = Pos.X;
-			*OutPosY.GetDest() = Pos.Y;
-			*OutPosZ.GetDest() = Pos.Z;
+			*OutPosX.GetDest() = float(Pos.X);		// LWC Precision Loss
+			*OutPosY.GetDest() = float(Pos.Y);
+			*OutPosZ.GetDest() = float(Pos.Z);
 			SplineSampleParam.Advance();
 			OutPosX.Advance();
 			OutPosY.Advance();
@@ -1185,9 +1191,9 @@ void UNiagaraDataInterfaceSpline::SampleSplineUpVectorByUnitDistance(FVectorVMEx
 			FVector Pos = FVector(0.0f, 0.0f, 1.0f); 
 			TransformHandler.TransformVector(Pos, InstData->Transform);
 
-			*OutPosX.GetDest() = Pos.X;
-			*OutPosY.GetDest() = Pos.Y;
-			*OutPosZ.GetDest() = Pos.Z;
+			*OutPosX.GetDest() = float(Pos.X);		// LWC Precision Loss
+			*OutPosY.GetDest() = float(Pos.Y);
+			*OutPosZ.GetDest() = float(Pos.Z);
 			SplineSampleParam.Advance();
 			OutPosX.Advance();
 			OutPosY.Advance();
@@ -1217,9 +1223,9 @@ void UNiagaraDataInterfaceSpline::SampleSplineRightVectorByUnitDistance(FVectorV
 			FVector Pos = InstData->GetRightVectorAtDistanceAlongSpline<UseLUT>(DistanceUnitDistance * SplineLength, ESplineCoordinateSpace::Local);
 			TransformHandler.TransformVector(Pos, InstData->Transform);
 
-			*OutPosX.GetDest() = Pos.X;
-			*OutPosY.GetDest() = Pos.Y;
-			*OutPosZ.GetDest() = Pos.Z;
+			*OutPosX.GetDest() = float(Pos.X);	// LWC Precision Loss
+			*OutPosY.GetDest() = float(Pos.Y);
+			*OutPosZ.GetDest() = float(Pos.Z);
 			SplineSampleParam.Advance();
 			OutPosX.Advance();
 			OutPosY.Advance();
@@ -1233,9 +1239,9 @@ void UNiagaraDataInterfaceSpline::SampleSplineRightVectorByUnitDistance(FVectorV
 			FVector Pos = FVector(-1.0f, 0.0f, 0.0f); 
 			TransformHandler.TransformVector(Pos, InstData->Transform);
 
-			*OutPosX.GetDest() = Pos.X;
-			*OutPosY.GetDest() = Pos.Y;
-			*OutPosZ.GetDest() = Pos.Z;
+			*OutPosX.GetDest() = float(Pos.X);	// LWC Precision Loss
+			*OutPosY.GetDest() = float(Pos.Y);
+			*OutPosZ.GetDest() = float(Pos.Z);
 			SplineSampleParam.Advance();
 			OutPosX.Advance();
 			OutPosY.Advance();
@@ -1264,9 +1270,9 @@ void UNiagaraDataInterfaceSpline::SampleSplineTangentByUnitDistance(FVectorVMExt
 			FVector Pos = InstData->GetTangentAtDistanceAlongSpline<UseLUT>(DistanceUnitDistance * SplineLength, ESplineCoordinateSpace::Local);
 			TransformHandler.TransformVector(Pos, InstData->Transform);
 
-			*OutPosX.GetDest() = Pos.X;
-			*OutPosY.GetDest() = Pos.Y;
-			*OutPosZ.GetDest() = Pos.Z;
+			*OutPosX.GetDest() = float(Pos.X);	// LWC Precision Loss
+			*OutPosY.GetDest() = float(Pos.Y);
+			*OutPosZ.GetDest() = float(Pos.Z);
 			SplineSampleParam.Advance();
 			OutPosX.Advance();
 			OutPosY.Advance();
@@ -1280,9 +1286,9 @@ void UNiagaraDataInterfaceSpline::SampleSplineTangentByUnitDistance(FVectorVMExt
 			FVector Pos = FVector(EForceInit::ForceInitToZero); 
 			TransformHandler.TransformVector(Pos, InstData->Transform);
 
-			*OutPosX.GetDest() = Pos.X;
-			*OutPosY.GetDest() = Pos.Y;
-			*OutPosZ.GetDest() = Pos.Z;
+			*OutPosX.GetDest() = float(Pos.X);	// LWC Precision Loss
+			*OutPosY.GetDest() = float(Pos.Y);
+			*OutPosZ.GetDest() = float(Pos.Z);
 			SplineSampleParam.Advance();
 			OutPosX.Advance();
 			OutPosY.Advance();
@@ -1311,9 +1317,9 @@ void UNiagaraDataInterfaceSpline::SampleSplineDirectionByUnitDistance(FVectorVME
 			FVector Pos = InstData->GetDirectionAtDistanceAlongSpline<UseLUT>(DistanceUnitDistance * SplineLength, ESplineCoordinateSpace::Local);
 			TransformHandler.TransformVector(Pos, InstData->Transform);
 
-			*OutPosX.GetDest() = Pos.X;
-			*OutPosY.GetDest() = Pos.Y;
-			*OutPosZ.GetDest() = Pos.Z;
+			*OutPosX.GetDest() = float(Pos.X);		// LWC Precision Loss
+			*OutPosY.GetDest() = float(Pos.Y);
+			*OutPosZ.GetDest() = float(Pos.Z);
 			SplineSampleParam.Advance();
 			OutPosX.Advance();
 			OutPosY.Advance();
@@ -1327,9 +1333,9 @@ void UNiagaraDataInterfaceSpline::SampleSplineDirectionByUnitDistance(FVectorVME
 			FVector Pos = FVector(0.0f, 1.0f, 0.0f); 
 			TransformHandler.TransformVector(Pos, InstData->Transform);
 
-			*OutPosX.GetDest() = Pos.X;
-			*OutPosY.GetDest() = Pos.Y;
-			*OutPosZ.GetDest() = Pos.Z;
+			*OutPosX.GetDest() = float(Pos.X);	// LWC Precision Loss
+			*OutPosY.GetDest() = float(Pos.Y);
+			*OutPosZ.GetDest() = float(Pos.Z);
 			SplineSampleParam.Advance();
 			OutPosX.Advance();
 			OutPosY.Advance();
@@ -1338,7 +1344,7 @@ void UNiagaraDataInterfaceSpline::SampleSplineDirectionByUnitDistance(FVectorVME
 	}
 }
 
-void UNiagaraDataInterfaceSpline::WriteTransform(const FMatrix& ToWrite, FVectorVMExternalFunctionContext& Context)
+void UNiagaraDataInterfaceSpline::WriteTransform(const FMatrix44f& ToWrite, FVectorVMExternalFunctionContext& Context)
 {
 	VectorVM::FExternalFuncRegisterHandler<float> Out00(Context);
 	VectorVM::FExternalFuncRegisterHandler<float> Out01(Context);
@@ -1418,13 +1424,13 @@ void UNiagaraDataInterfaceSpline::FindClosestUnitDistanceFromPositionWS(FVectorV
 void UNiagaraDataInterfaceSpline::GetLocalToWorld(FVectorVMExternalFunctionContext& Context)
 {
 	VectorVM::FUserPtrHandler<FNDISpline_InstanceData> InstData(Context);
-	WriteTransform(InstData->Transform, Context);
+	WriteTransform(FMatrix44f(InstData->Transform), Context);	//LWC Precision Loss
 }
 
 void UNiagaraDataInterfaceSpline::GetLocalToWorldInverseTransposed(FVectorVMExternalFunctionContext& Context)
 {
 	VectorVM::FUserPtrHandler<FNDISpline_InstanceData> InstData(Context);
-	WriteTransform(InstData->TransformInverseTransposed, Context);
+	WriteTransform(FMatrix44f(InstData->TransformInverseTransposed), Context);	//LWC Precision Loss
 }
 
 void UNiagaraDataInterfaceSpline::VMGetSplineLength(FVectorVMExternalFunctionContext& Context)
diff --git a/Engine/Plugins/FX/Niagara/Source/Niagara/Private/NiagaraDataInterfaceVector4Curve.cpp b/Engine/Plugins/FX/Niagara/Source/Niagara/Private/NiagaraDataInterfaceVector4Curve.cpp
--- a/Engine/Plugins/FX/Niagara/Source/Niagara/Private/NiagaraDataInterfaceVector4Curve.cpp
+++ b/Engine/Plugins/FX/Niagara/Source/Niagara/Private/NiagaraDataInterfaceVector4Curve.cpp
@@ -207,24 +207,24 @@ void UNiagaraDataInterfaceVector4Curve::GetVMExternalFunction(const FVMExternalF
 }
 
 template<>
-FORCEINLINE_DEBUGGABLE FVector4 UNiagaraDataInterfaceVector4Curve::SampleCurveInternal<TIntegralConstant<bool, true>>(float X)
+FORCEINLINE_DEBUGGABLE FVector4f UNiagaraDataInterfaceVector4Curve::SampleCurveInternal<TIntegralConstant<bool, true>>(float X)
 {
-	float RemappedX = FMath::Clamp(NormalizeTime(X) * LUTNumSamplesMinusOne, 0.0f, LUTNumSamplesMinusOne);
-	float PrevEntry = FMath::TruncToFloat(RemappedX);
-	float NextEntry = PrevEntry < LUTNumSamplesMinusOne ? PrevEntry + 1.0f : PrevEntry;
-	float Interp = RemappedX - PrevEntry;
-
-	int32 AIndex = (int32)(PrevEntry * (float)CurveLUTNumElems);
-	int32 BIndex = (int32)(NextEntry * (float)CurveLUTNumElems);
-	FVector4 A = FVector4(ShaderLUT[AIndex], ShaderLUT[AIndex + 1], ShaderLUT[AIndex + 2], ShaderLUT[AIndex + 3]);
-	FVector4 B = FVector4(ShaderLUT[BIndex], ShaderLUT[BIndex + 1], ShaderLUT[BIndex + 2], ShaderLUT[BIndex + 3]);
+	const float RemappedX = FMath::Clamp(NormalizeTime(X) * LUTNumSamplesMinusOne, 0.0f, LUTNumSamplesMinusOne);
+	const float PrevEntry = FMath::TruncToFloat(RemappedX);
+	const float NextEntry = PrevEntry < LUTNumSamplesMinusOne ? PrevEntry + 1.0f : PrevEntry;
+	const float Interp = RemappedX - PrevEntry;
+
+	const int32 AIndex = (int32)(PrevEntry * (float)CurveLUTNumElems);
+	const int32 BIndex = (int32)(NextEntry * (float)CurveLUTNumElems);
+	const FVector4f A = FVector4f(ShaderLUT[AIndex], ShaderLUT[AIndex + 1], ShaderLUT[AIndex + 2], ShaderLUT[AIndex + 3]);
+	const FVector4f B = FVector4f(ShaderLUT[BIndex], ShaderLUT[BIndex + 1], ShaderLUT[BIndex + 2], ShaderLUT[BIndex + 3]);
 	return FMath::Lerp(A, B, Interp);
 }
 
 template<>
-FORCEINLINE_DEBUGGABLE FVector4 UNiagaraDataInterfaceVector4Curve::SampleCurveInternal<TIntegralConstant<bool, false>>(float X)
+FORCEINLINE_DEBUGGABLE FVector4f UNiagaraDataInterfaceVector4Curve::SampleCurveInternal<TIntegralConstant<bool, false>>(float X)
 {
-	return FVector4(XCurve.Eval(X), YCurve.Eval(X), ZCurve.Eval(X), WCurve.Eval(X));
+	return FVector4f(XCurve.Eval(X), YCurve.Eval(X), ZCurve.Eval(X), WCurve.Eval(X));
 }
 
 
@@ -240,8 +240,8 @@ void UNiagaraDataInterfaceVector4Curve::SampleCurve(FVectorVMExternalFunctionCon
 
 	for (int32 i = 0; i < Context.GetNumInstances(); ++i)
 	{
-		float X = XParam.GetAndAdvance();
-		FVector4 Sample = SampleCurveInternal<UseLUT>(X);
+		const float X = XParam.GetAndAdvance();
+		const FVector4f Sample = SampleCurveInternal<UseLUT>(X);
 		*SamplePtrR.GetDestAndAdvance() = Sample.X;
 		*SamplePtrG.GetDestAndAdvance() = Sample.Y;
 		*SamplePtrB.GetDestAndAdvance() = Sample.Z;
diff --git a/Engine/Plugins/FX/Niagara/Source/Niagara/Private/NiagaraDataInterfaceVectorField.cpp b/Engine/Plugins/FX/Niagara/Source/Niagara/Private/NiagaraDataInterfaceVectorField.cpp
--- a/Engine/Plugins/FX/Niagara/Source/Niagara/Private/NiagaraDataInterfaceVectorField.cpp
+++ b/Engine/Plugins/FX/Niagara/Source/Niagara/Private/NiagaraDataInterfaceVectorField.cpp
@@ -493,7 +493,7 @@ void UNiagaraDataInterfaceVectorField::GetFieldTilingAxes(FVectorVMExternalFunct
 	VectorVM::FExternalFuncRegisterHandler<float> OutSizeY(Context);
 	VectorVM::FExternalFuncRegisterHandler<float> OutSizeZ(Context);
 
-	FVector Tilings = GetTilingAxes();
+	const FVector3f Tilings(GetTilingAxes());
 	for (int32 i = 0; i < Context.GetNumInstances(); ++i)
 	{
 		*OutSizeX.GetDest() = Tilings.X;
@@ -512,7 +512,7 @@ void UNiagaraDataInterfaceVectorField::GetFieldDimensions(FVectorVMExternalFunct
 	VectorVM::FExternalFuncRegisterHandler<float> OutSizeY(Context);
 	VectorVM::FExternalFuncRegisterHandler<float> OutSizeZ(Context);
 
-	FVector Dim = GetDimensions();
+	const FVector3f Dim(GetDimensions());
 	for (int32 i = 0; i < Context.GetNumInstances(); ++i)
 	{
 		*OutSizeX.GetDest() = Dim.X;
@@ -534,8 +534,8 @@ void UNiagaraDataInterfaceVectorField::GetFieldBounds(FVectorVMExternalFunctionC
 	VectorVM::FExternalFuncRegisterHandler<float> OutMaxY(Context);
 	VectorVM::FExternalFuncRegisterHandler<float> OutMaxZ(Context);
 
-	FVector MinBounds = GetMinBounds();
-	FVector MaxBounds = GetMaxBounds();
+	const FVector3f MinBounds(GetMinBounds());
+	const FVector3f MaxBounds(GetMaxBounds());
 	for (int32 i = 0; i < Context.GetNumInstances(); ++i)
 	{
 		*OutMinX.GetDest() = MinBounds.X;
@@ -615,7 +615,7 @@ void UNiagaraDataInterfaceVectorField::SampleVectorField(FVectorVMExternalFuncti
 					// Offset by half a cell size due to sample being in the center of its cell
 					Pos = Pos - FVector(0.5f, 0.5f, 0.5f);
 
-					const FVector V = StaticVectorField->FilteredSample(Pos, TilingAxes);
+					const FVector3f V(StaticVectorField->FilteredSample(Pos, TilingAxes));
 
 					// Write final output...
 					*OutSampleX.GetDest() = V.X;
diff --git a/Engine/Plugins/FX/Niagara/Source/Niagara/Private/NiagaraDataSetReadback.cpp b/Engine/Plugins/FX/Niagara/Source/Niagara/Private/NiagaraDataSetReadback.cpp
--- a/Engine/Plugins/FX/Niagara/Source/Niagara/Private/NiagaraDataSetReadback.cpp
+++ b/Engine/Plugins/FX/Niagara/Source/Niagara/Private/NiagaraDataSetReadback.cpp
@@ -119,7 +119,7 @@ void FNiagaraDataSetReadback::GPUReadbackInternal(FRHICommandListImmediate& RHIC
 
 	constexpr int32 NumReadbackBuffers = 5;
 	TArray<FNiagaraGpuReadbackManager::FBufferRequest, TInlineAllocator<NumReadbackBuffers>> ReadbackBuffers;
-	ReadbackBuffers.Emplace(CountManager.GetInstanceCountBuffer().Buffer, CountOffset * sizeof(uint32), sizeof(uint32));
+	ReadbackBuffers.Emplace(CountManager.GetInstanceCountBuffer().Buffer, uint32(CountOffset * sizeof(uint32)), uint32(sizeof(uint32)));
 	const int32 FloatBufferIndex = FloatBuffer.NumBytes == 0 ? INDEX_NONE : ReadbackBuffers.Emplace(FloatBuffer.Buffer, 0, FloatBuffer.NumBytes);
 	const int32 HalfBufferIndex = HalfBuffer.NumBytes == 0 ? INDEX_NONE : ReadbackBuffers.Emplace(HalfBuffer.Buffer, 0, HalfBuffer.NumBytes);
 	const int32 IntBufferIndex = IntBuffer.NumBytes == 0 ? INDEX_NONE : ReadbackBuffers.Emplace(IntBuffer.Buffer, 0, IntBuffer.NumBytes);
diff --git a/Engine/Plugins/FX/Niagara/Source/Niagara/Private/NiagaraGPUInstanceCountManager.cpp b/Engine/Plugins/FX/Niagara/Source/Niagara/Private/NiagaraGPUInstanceCountManager.cpp
--- a/Engine/Plugins/FX/Niagara/Source/Niagara/Private/NiagaraGPUInstanceCountManager.cpp
+++ b/Engine/Plugins/FX/Niagara/Source/Niagara/Private/NiagaraGPUInstanceCountManager.cpp
@@ -318,7 +318,7 @@ void FNiagaraGPUInstanceCountManager::FlushIndirectArgsPool()
 	if (GNiagaraIndirectArgsPoolAllowShrinking && DrawIndirectPool.Num() > 0 && DrawIndirectLowWaterFrames >= uint32(GNiagaraIndirectArgsPoolLowWaterFrames))
 	{
 		FIndirectArgsPoolEntryPtr& PoolEntry = DrawIndirectPool[0];
-		const uint32 NewSize = FMath::Max<uint32>(GNiagaraIndirectArgsPoolMinSize, PoolEntry->AllocatedEntries / GNiagaraIndirectArgsPoolBlockSizeFactor);
+		const uint32 NewSize = FMath::Max<uint32>(GNiagaraIndirectArgsPoolMinSize, FMath::FloorToInt(float(PoolEntry->AllocatedEntries) / GNiagaraIndirectArgsPoolBlockSizeFactor));
 
 		INDIRECT_ARG_POOL_LOG("Shrinking pool from size %d to %d", PoolEntry->AllocatedEntries, NewSize);
 
@@ -566,7 +566,7 @@ void FNiagaraGPUInstanceCountManager::UpdateDrawIndirectBuffers(FNiagaraGpuCompu
 			{
 				// See if this was a low water mark frame
 				FIndirectArgsPoolEntryPtr& PoolEntry = DrawIndirectPool[0];
-				const uint32 LowWaterCount = FMath::Max<uint32>(GNiagaraIndirectArgsPoolMinSize, PoolEntry->AllocatedEntries * GNiagaraIndirectArgsPoolLowWaterAmount);
+				const uint32 LowWaterCount = FMath::Max<uint32>(GNiagaraIndirectArgsPoolMinSize, FMath::CeilToInt(float(PoolEntry->AllocatedEntries) * GNiagaraIndirectArgsPoolLowWaterAmount));
 				if (PoolEntry->UsedEntriesTotal < LowWaterCount)
 				{
 					++DrawIndirectLowWaterFrames;
diff --git a/Engine/Plugins/FX/Niagara/Source/Niagara/Private/NiagaraGpuComputeDebug.cpp b/Engine/Plugins/FX/Niagara/Source/Niagara/Private/NiagaraGpuComputeDebug.cpp
--- a/Engine/Plugins/FX/Niagara/Source/Niagara/Private/NiagaraGpuComputeDebug.cpp
+++ b/Engine/Plugins/FX/Niagara/Source/Niagara/Private/NiagaraGpuComputeDebug.cpp
@@ -205,7 +205,7 @@ FNiagaraSimulationDebugDrawData* FNiagaraGpuComputeDebug::GetSimulationDebugDraw
 
 	FNiagaraSimulationDebugDrawData* DebugDrawData = GetSimulationDebugDrawData(SystemInstanceID);
 
-	const uint32 MaxPossibleLines = FMath::Min(FMath::DivideAndRoundDown(GetMaxBufferDimension(), uint64(NiagaraGpuComputeDebugLocal::NumUintsPerLine)), uint64(TNumericLimits<uint32>::Max() >> 1));
+	const uint32 MaxPossibleLines = uint32(FMath::Min(FMath::DivideAndRoundDown(GetMaxBufferDimension(), uint64(NiagaraGpuComputeDebugLocal::NumUintsPerLine)), uint64(TNumericLimits<uint32>::Max() >> 1)));
 	const uint32 MaxLineInstancesToUse = FMath::Min(FMath::Max(DebugDrawData->GpuLineMaxInstances, OverrideMaxDebugLines > 0 ? OverrideMaxDebugLines : uint32(GNiagaraGpuComputeDebug_MaxLineInstances)), MaxPossibleLines);
 	if (DebugDrawData->GpuLineMaxInstances != MaxLineInstancesToUse)
 	{
@@ -246,9 +246,9 @@ void FNiagaraGpuComputeDebug::DrawDebug(FRDGBuilder& GraphBuilder, const FViewIn
 	++TickCounter;
 
 	const UFont* Font = GEngine->GetTinyFont();
-	const float FontHeight = Font->GetMaxCharHeight();
+	const int32 FontHeight = FMath::CeilToInt(Font->GetMaxCharHeight());
 
-	FIntPoint Location(10.0f, Output.ViewRect.Height() - 10.0f);
+	FIntPoint Location(10, Output.ViewRect.Height() - 10);
 
 	const int32 DisplayMinHeight = GNiagaraGpuComputeDebug_MinTextureHeight > 0 ? GNiagaraGpuComputeDebug_MinTextureHeight : 0;
 	const int32 DisplayMaxHeight = GNiagaraGpuComputeDebug_MaxTextureHeight > 0 ? GNiagaraGpuComputeDebug_MaxTextureHeight : TNumericLimits<int32>::Max();
@@ -281,7 +281,7 @@ void FNiagaraGpuComputeDebug::DrawDebug(FRDGBuilder& GraphBuilder, const FViewIn
 			}
 		);
 
-		Location.Y -= 1.0f;
+		Location.Y -= 1;
 	}
 }
 
diff --git a/Engine/Plugins/FX/Niagara/Source/Niagara/Private/NiagaraGpuComputeDispatch.cpp b/Engine/Plugins/FX/Niagara/Source/Niagara/Private/NiagaraGpuComputeDispatch.cpp
--- a/Engine/Plugins/FX/Niagara/Source/Niagara/Private/NiagaraGpuComputeDispatch.cpp
+++ b/Engine/Plugins/FX/Niagara/Source/Niagara/Private/NiagaraGpuComputeDispatch.cpp
@@ -2085,7 +2085,11 @@ void FNiagaraGpuComputeDispatch::GenerateSortKeys(FRHICommandListImmediate& RHIC
 			Params.NumCullPlanes = 0;
 			for (const FPlane& Plane : SortInfo.CullPlanes)
 			{
-				Params.CullPlanes[Params.NumCullPlanes++] = FVector4f(Plane.X, Plane.Y, Plane.Z, Plane.W);
+				Params.CullPlanes[Params.NumCullPlanes].X = float(Plane.X);
+				Params.CullPlanes[Params.NumCullPlanes].Y = float(Plane.Y);
+				Params.CullPlanes[Params.NumCullPlanes].Z = float(Plane.Z);
+				Params.CullPlanes[Params.NumCullPlanes].W = float(Plane.W);	// LWC Precision loss
+				++Params.NumCullPlanes;
 			}
 
 			// Choose the shader to bind
diff --git a/Engine/Plugins/FX/Niagara/Source/Niagara/Private/NiagaraSimCacheHelper.h b/Engine/Plugins/FX/Niagara/Source/Niagara/Private/NiagaraSimCacheHelper.h
--- a/Engine/Plugins/FX/Niagara/Source/Niagara/Private/NiagaraSimCacheHelper.h
+++ b/Engine/Plugins/FX/Niagara/Source/Niagara/Private/NiagaraSimCacheHelper.h
@@ -333,19 +333,19 @@ struct FNiagaraSimCacheHelper
 					if (SourceVariable.Variable.GetType() == FNiagaraTypeDefinition::GetPositionDef())
 					{
 						check(SourceVariable.FloatCount == 3);
-						CacheLayout.VariableMappingsToDataBuffer.Emplace(FloatOffset, DestVariableLayout->FloatComponentStart, &FNiagaraSimCacheHelper::CopyPositions);
+						CacheLayout.VariableMappingsToDataBuffer.Emplace(uint16(FloatOffset), uint16(DestVariableLayout->FloatComponentStart), &FNiagaraSimCacheHelper::CopyPositions);
 						DestVariableLayout = nullptr;
 					}
 					else if (SourceVariable.Variable.GetType() == FNiagaraTypeDefinition::GetQuatDef())
 					{
 						check(SourceVariable.FloatCount == 4);
-						CacheLayout.VariableMappingsToDataBuffer.Emplace(FloatOffset, DestVariableLayout->FloatComponentStart, &FNiagaraSimCacheHelper::CopyQuaternions);
+						CacheLayout.VariableMappingsToDataBuffer.Emplace(uint16(FloatOffset), uint16(DestVariableLayout->FloatComponentStart), &FNiagaraSimCacheHelper::CopyQuaternions);
 						DestVariableLayout = nullptr;
 					}
 					else if (SourceVariable.Variable.GetType() == FNiagaraTypeDefinition::GetMatrix4Def())
 					{
 						check(SourceVariable.FloatCount == 16);
-						CacheLayout.VariableMappingsToDataBuffer.Emplace(FloatOffset, DestVariableLayout->FloatComponentStart, &FNiagaraSimCacheHelper::CopyMatrices);
+						CacheLayout.VariableMappingsToDataBuffer.Emplace(uint16(FloatOffset), uint16(DestVariableLayout->FloatComponentStart), &FNiagaraSimCacheHelper::CopyMatrices);
 						DestVariableLayout = nullptr;
 					}
 				}
diff --git a/Engine/Plugins/FX/Niagara/Source/Niagara/Private/NiagaraSystemSimulation.cpp b/Engine/Plugins/FX/Niagara/Source/Niagara/Private/NiagaraSystemSimulation.cpp
--- a/Engine/Plugins/FX/Niagara/Source/Niagara/Private/NiagaraSystemSimulation.cpp
+++ b/Engine/Plugins/FX/Niagara/Source/Niagara/Private/NiagaraSystemSimulation.cpp
@@ -1119,7 +1119,7 @@ void FNiagaraSystemSimulation::Tick_GameThread(float DeltaSeconds, const FGraphE
 	{
 		float FixedDelta = System->GetFixedTickDeltaTime();
 		float Budget = FixedDelta > 0 ? FMath::Fmod(FixedDeltaTickAge, FixedDelta) + DeltaSeconds : 0;
-		int32 Ticks = FixedDelta > 0 ? FMath::Min(Budget / FixedDelta, GNiagaraSystemSimulationMaxTickSubsteps) : 0;
+		int32 Ticks = FixedDelta > 0 ? FMath::Min(FMath::FloorToInt(Budget / FixedDelta), GNiagaraSystemSimulationMaxTickSubsteps) : 0;
 
 		TickInfo.UsesFixedTick = true;
 		TickInfo.EngineTick = DeltaSeconds;
diff --git a/Engine/Plugins/FX/Niagara/Source/Niagara/Private/NiagaraUvQuadTree.cpp b/Engine/Plugins/FX/Niagara/Source/Niagara/Private/NiagaraUvQuadTree.cpp
--- a/Engine/Plugins/FX/Niagara/Source/Niagara/Private/NiagaraUvQuadTree.cpp
+++ b/Engine/Plugins/FX/Niagara/Source/Niagara/Private/NiagaraUvQuadTree.cpp
@@ -85,7 +85,7 @@ int32 FNiagaraUvQuadTree::FSubTree::GetQuads(const FBox2D& Box, FChildArray& Qua
 
 int32 FNiagaraUvQuadTree::FSubTree::Freeze(const FNiagaraUvQuadTree & QuadTree, FArchive& Ar) const
 {
-	int32 StartOffset = Ar.Tell();
+	int64 StartOffset = Ar.Tell();
 	int32 PlaceholderIndex = INDEX_NONE;
 
 	for (int32 OffsetIt = 0; OffsetIt < SubTreeIndices.Num(); ++OffsetIt)
@@ -109,7 +109,7 @@ int32 FNiagaraUvQuadTree::FSubTree::Freeze(const FNiagaraUvQuadTree & QuadTree,
 			ChildOffsets[ChildIt] = QuadTree.ChildTrees[SubTreeIndices[ChildIt]].Freeze(QuadTree, Ar);
 		}
 
-		const int32 EndOffset = Ar.Tell();
+		const int64 EndOffset = Ar.Tell();
 		Ar.Seek(StartOffset);
 
 		for (int32 ChildIt = 0; ChildIt < ChildOffsets.Num(); ++ChildIt)
@@ -183,7 +183,7 @@ void FNiagaraUvQuadTree::Split(int32 SubTreeIndex)
 			const FVector2D OverlapMin = FVector2D::Max(ContentNode.Coverage.Min, Parent.Coverage.Min);
 			const FVector2D OverlapMax = FVector2D::Min(ContentNode.Coverage.Max, Parent.Coverage.Max);
 			FBox2D OverlapRegion(OverlapMin, FVector2D::Max(OverlapMin, OverlapMax));
-			const float AreaRatio = OverlapRegion.GetArea() / Parent.Coverage.GetArea();
+			const double AreaRatio = OverlapRegion.GetArea() / Parent.Coverage.GetArea();
 
 			if (AreaRatio > GNiagaraUvQuadTreeDuplicateThreshold)
 			{
diff --git a/Engine/Plugins/FX/Niagara/Source/Niagara/Private/NiagaraWorldManager.cpp b/Engine/Plugins/FX/Niagara/Source/Niagara/Private/NiagaraWorldManager.cpp
--- a/Engine/Plugins/FX/Niagara/Source/Niagara/Private/NiagaraWorldManager.cpp
+++ b/Engine/Plugins/FX/Niagara/Source/Niagara/Private/NiagaraWorldManager.cpp
@@ -1652,7 +1652,7 @@ void FNiagaraWorldManager::DistanceCull(UNiagaraEffectType* EffectType, const FN
 		FVector Location = Component->GetComponentLocation();
 		for (const FNiagaraCachedViewInfo& ViewInfo : GetCachedViewInfo())
 		{
-			ClosestDistSq = FMath::Min(ClosestDistSq, FVector::DistSquared(ViewInfo.ViewToWorld.GetOrigin(), Location));
+			ClosestDistSq = FMath::Min(ClosestDistSq, float(FVector::DistSquared(ViewInfo.ViewToWorld.GetOrigin(), Location)));
 		}
 
 		LODDistance = FMath::Sqrt(ClosestDistSq);
@@ -1703,7 +1703,7 @@ void FNiagaraWorldManager::DistanceCull(UNiagaraEffectType* EffectType, const FN
 		float ClosestDistSq = FLT_MAX;
 		for (const FNiagaraCachedViewInfo& ViewInfo : GetCachedViewInfo())
 		{
-			ClosestDistSq = FMath::Min(ClosestDistSq, FVector::DistSquared(ViewInfo.ViewToWorld.GetOrigin(), Location));
+			ClosestDistSq = FMath::Min(ClosestDistSq, float(FVector::DistSquared(ViewInfo.ViewToWorld.GetOrigin(), Location)));
 		}
 
 		if (GEnableNiagaraDistanceCulling && ScalabilitySettings.bCullByDistance)
diff --git a/Engine/Plugins/FX/Niagara/Source/Niagara/Public/NiagaraComponentPool.h b/Engine/Plugins/FX/Niagara/Source/Niagara/Public/NiagaraComponentPool.h
--- a/Engine/Plugins/FX/Niagara/Source/Niagara/Public/NiagaraComponentPool.h
+++ b/Engine/Plugins/FX/Niagara/Source/Niagara/Public/NiagaraComponentPool.h
@@ -21,14 +21,14 @@ struct FNCPoolElement
 	UPROPERTY(transient)
 	TObjectPtr<UNiagaraComponent> Component;
 
-	float LastUsedTime;
+	double LastUsedTime;
 
 	FNCPoolElement()
 		: Component(nullptr), LastUsedTime(0.0f)
 	{
 
 	}
-	FNCPoolElement(UNiagaraComponent* InNC, float InLastUsedTime)
+	FNCPoolElement(UNiagaraComponent* InNC, double InLastUsedTime)
 		: Component(InNC), LastUsedTime(InLastUsedTime)
 	{
 
@@ -65,13 +65,13 @@ public:
 	UNiagaraComponent* Acquire(UWorld* World, UNiagaraSystem* Template, ENCPoolMethod PoolingMethod, bool bForceNew=false);
 
 	/** Returns a component to the pool. */
-	void Reclaim(UNiagaraComponent* NC, const float CurrentTimeSeconds);
+	void Reclaim(UNiagaraComponent* NC, const double CurrentTimeSeconds);
 
 	/** Forces us to remove a component, returns true if the component exists inside the pool. */
 	bool RemoveComponent(UNiagaraComponent* Component);
 
 	/** Kills any components that have not been used since the passed KillTime. */
-	void KillUnusedComponents(float KillTime, UNiagaraSystem* Template);
+	void KillUnusedComponents(double KillTime, UNiagaraSystem* Template);
 
 	int32 NumComponents() { return FreeElements.Num(); }
 };
@@ -85,9 +85,9 @@ private:
 	UPROPERTY()
 	TMap<TObjectPtr<UNiagaraSystem>, FNCPool> WorldParticleSystemPools;
 
-	float LastParticleSytemPoolCleanTime;
-public:
+	double LastParticleSytemPoolCleanTime;
 
+public:
 	static bool Enabled();
 
 	~UNiagaraComponentPool();
diff --git a/Engine/Plugins/FX/Niagara/Source/Niagara/Public/NiagaraTypes.h b/Engine/Plugins/FX/Niagara/Source/Niagara/Public/NiagaraTypes.h
--- a/Engine/Plugins/FX/Niagara/Source/Niagara/Public/NiagaraTypes.h
+++ b/Engine/Plugins/FX/Niagara/Source/Niagara/Public/NiagaraTypes.h
@@ -780,7 +780,7 @@ struct NIAGARA_API FNiagaraTypeDefinition
 {
 	GENERATED_USTRUCT_BODY()
 
-	enum FUnderlyingType
+	enum FUnderlyingType : uint16
 	{
 		UT_None = 0,
 		UT_Class,
