#UE Remove SearchAllAssets call from UNiagaraDataChannelDefinitions::GetDataChannelDefinitions, and instead just wait for the existing global scan to finish, if it isnt already. Commandlets are responsible for deciding how much to scan themselves

diff --git a/Engine/Plugins/FX/Niagara/Source/Niagara/Private/NiagaraDataChannelDefinitions.cpp b/Engine/Plugins/FX/Niagara/Source/Niagara/Private/NiagaraDataChannelDefinitions.cpp
--- a/Engine/Plugins/FX/Niagara/Source/Niagara/Private/NiagaraDataChannelDefinitions.cpp
+++ b/Engine/Plugins/FX/Niagara/Source/Niagara/Private/NiagaraDataChannelDefinitions.cpp
@@ -217,12 +217,18 @@ const TArray<UNiagaraDataChannelDefinitions*>& UNiagaraDataChannelDefinitions::G
 		bAssetRegistryScanBegun = true;
 		FAssetRegistryModule& AssetRegistryModule = FModuleManager::LoadModuleChecked<FAssetRegistryModule>("AssetRegistry");
 		IAssetRegistry& AssetRegistry = AssetRegistryModule.Get();
-		AssetRegistryOnLoadCompleteHandle = AssetRegistry.OnFilesLoaded().AddStatic(&UNiagaraDataChannelDefinitions::OnAssetRegistryLoadComplete);
-		AssetRegistry.SearchAllAssets(true);
-		if(bRequired)
+		if (AssetRegistry.IsLoadingAssets())
 		{
-			AssetRegistry.WaitForCompletion();
+			AssetRegistryOnLoadCompleteHandle = AssetRegistry.OnFilesLoaded().AddStatic(&UNiagaraDataChannelDefinitions::OnAssetRegistryLoadComplete);
+			if (bRequired)
+			{
+				AssetRegistry.WaitForCompletion();
+			}
 		}
+		else
+		{
+			OnAssetRegistryLoadComplete();
+		}		
 	}
 
 	if(bAssetRegistryScanComplete == false && bInformUser)
