fixed crash with vdb caching

diff --git a/Engine/Plugins/FX/Niagara/Source/Niagara/Private/NiagaraDataInterfaceRenderTargetVolume.cpp b/Engine/Plugins/FX/Niagara/Source/Niagara/Private/NiagaraDataInterfaceRenderTargetVolume.cpp
--- a/Engine/Plugins/FX/Niagara/Source/Niagara/Private/NiagaraDataInterfaceRenderTargetVolume.cpp
+++ b/Engine/Plugins/FX/Niagara/Source/Niagara/Private/NiagaraDataInterfaceRenderTargetVolume.cpp
@@ -566,9 +566,7 @@ bool UNiagaraDataInterfaceRenderTargetVolume::GetExposedVariableValue(const FNia
 UObject* UNiagaraDataInterfaceRenderTargetVolume::SimCacheBeginWrite(UObject* SimCache, FNiagaraSystemInstance* NiagaraSystemInstance, const void* OptionalPerInstanceData) const
 {
 	if (NDIRenderTargetVolumeLocal::GSimCacheUseOpenVDB)
-	{
-		const FRenderTargetVolumeRWInstanceData_GameThread* InstanceData_GT = reinterpret_cast<const FRenderTargetVolumeRWInstanceData_GameThread*>(OptionalPerInstanceData);		
-
+	{		
 		UVolumeCache* OpenVDBSimCacheData = nullptr;
 
 		OpenVDBSimCacheData = NewObject<UVolumeCache>(SimCache);
@@ -608,7 +606,7 @@ UObject* UNiagaraDataInterfaceRenderTargetVolume::SimCacheBeginWrite(UObject* Si
 
 		OpenVDBSimCacheData->FilePath = FullFilePathSpec;
 		OpenVDBSimCacheData->CacheType = EVolumeCacheType::OpenVDB;
-		OpenVDBSimCacheData->Resolution = InstanceData_GT->Size;
+		OpenVDBSimCacheData->Resolution = FIntVector3(1,1,1);
 		OpenVDBSimCacheData->FrameRangeStart = TNumericLimits<int32>::Lowest();
 		OpenVDBSimCacheData->FrameRangeEnd = TNumericLimits<int32>::Lowest();
 		OpenVDBSimCacheData->InitData();
diff --git a/Engine/Plugins/FX/Niagara/Source/Niagara/Private/VolumeCache.cpp b/Engine/Plugins/FX/Niagara/Source/Niagara/Private/VolumeCache.cpp
--- a/Engine/Plugins/FX/Niagara/Source/Niagara/Private/VolumeCache.cpp
+++ b/Engine/Plugins/FX/Niagara/Source/Niagara/Private/VolumeCache.cpp
@@ -149,6 +149,12 @@ bool FOpenVDBCacheData::LoadFile(FString Path, int frame)
 	// if the file is loaded at the current frame, return true
 	if (OpenVDBGrids.Contains(frame) && OpenVDBGrids[frame] != nullptr)
 	{
+		// if dense vdb buffer doesn't match current resolution, re initialize it				
+		if (DenseGridPtr == nullptr || (DenseGridPtr->bbox().dim() != openvdb::Coord(DenseResolution.X - 1, DenseResolution.Y - 1, DenseResolution.Z - 1)))
+		{
+			Init(DenseResolution);
+		}
+
 		openvdb::tools::CopyToDense<Vec4Tree, Vec4Dense> Copier(OpenVDBGrids[frame]->tree(), *DenseGridPtr);
 		Copier.copy();
 
@@ -201,8 +207,8 @@ bool FOpenVDBCacheData::LoadFile(FString Path, int frame)
 
 			OpenVDBGrids.Add(frame, ColorGrid);
 
-			// if dense vdb buffer doesn't match current resolution, re initialize it	
-			if (DenseGridPtr->bbox().dim() != openvdb::Coord(DenseResolution.X - 1, DenseResolution.Y - 1, DenseResolution.Z - 1))
+			// if dense vdb buffer doesn't match current resolution, re initialize it				
+			if (DenseGridPtr == nullptr || (DenseGridPtr->bbox().dim() != openvdb::Coord(DenseResolution.X - 1, DenseResolution.Y - 1, DenseResolution.Z - 1)))
 			{
 				Init(DenseResolution);
 			}
