Niagara - Fix a bug with dynamic input nodes where output pins could become invalid when refreshing from external changes.

diff --git a/Engine/Plugins/FX/Niagara/Source/NiagaraEditor/Private/NiagaraNodeFunctionCall.cpp b/Engine/Plugins/FX/Niagara/Source/NiagaraEditor/Private/NiagaraNodeFunctionCall.cpp
--- a/Engine/Plugins/FX/Niagara/Source/NiagaraEditor/Private/NiagaraNodeFunctionCall.cpp
+++ b/Engine/Plugins/FX/Niagara/Source/NiagaraEditor/Private/NiagaraNodeFunctionCall.cpp
@@ -1027,6 +1027,7 @@ bool UNiagaraNodeFunctionCall::RefreshFromExternalChanges()
 		// TODO - Leverage code in reallocate pins to determine if any pins have changed...
 		ReallocatePins(false);
 		FNiagaraStackGraphUtilities::SynchronizeReferencingMapPinsWithFunctionCall(*this);
+		FNiagaraStackGraphUtilities::FixDynamicInputNodeOutputPinsFromExternalChanges(*this);
 		return true;
 	}
 	else
diff --git a/Engine/Plugins/FX/Niagara/Source/NiagaraEditor/Private/NiagaraScriptMergeManager.cpp b/Engine/Plugins/FX/Niagara/Source/NiagaraEditor/Private/NiagaraScriptMergeManager.cpp
--- a/Engine/Plugins/FX/Niagara/Source/NiagaraEditor/Private/NiagaraScriptMergeManager.cpp
+++ b/Engine/Plugins/FX/Niagara/Source/NiagaraEditor/Private/NiagaraScriptMergeManager.cpp
@@ -1483,6 +1483,7 @@ void FNiagaraScriptMergeManager::UpdateModuleVersions(const FVersionedNiagaraEmi
 				UpgradeContext.ApplyClipboardCallback = [ModuleItem](UNiagaraClipboardContent* ClipboardContent, FText& OutWarning) { ModuleItem->Paste(ClipboardContent, OutWarning); };
 			}
 			ChangedModule->GetFunctionCallNode()->ChangeScriptVersion(NewScriptVersion, UpgradeContext, true);
+			ChangedModule->GetFunctionCallNode()->RefreshFromExternalChanges();
 		}
 	}
 
diff --git a/Engine/Plugins/FX/Niagara/Source/NiagaraEditor/Private/ViewModels/Stack/NiagaraStackFunctionInput.cpp b/Engine/Plugins/FX/Niagara/Source/NiagaraEditor/Private/ViewModels/Stack/NiagaraStackFunctionInput.cpp
--- a/Engine/Plugins/FX/Niagara/Source/NiagaraEditor/Private/ViewModels/Stack/NiagaraStackFunctionInput.cpp
+++ b/Engine/Plugins/FX/Niagara/Source/NiagaraEditor/Private/ViewModels/Stack/NiagaraStackFunctionInput.cpp
@@ -2519,6 +2519,7 @@ void UNiagaraStackFunctionInput::ChangeScriptVersion(FGuid NewScriptVersion)
       FCompileConstantResolver(GetEmitterViewModel()->GetEmitter(), FNiagaraStackGraphUtilities::GetOutputNodeUsage(*GetDynamicInputNode())) :
       FCompileConstantResolver(&GetSystemViewModel()->GetSystem(), FNiagaraStackGraphUtilities::GetOutputNodeUsage(*GetDynamicInputNode()));
 	GetDynamicInputNode()->ChangeScriptVersion(NewScriptVersion, UpgradeContext, true);
+	GetDynamicInputNode()->RefreshFromExternalChanges();
 	TSharedRef<FNiagaraSystemViewModel> CachedSysViewModel = GetSystemViewModel();
 	if (CachedSysViewModel->GetSystemStackViewModel())
 		CachedSysViewModel->GetSystemStackViewModel()->InvalidateCachedParameterUsage();
diff --git a/Engine/Plugins/FX/Niagara/Source/NiagaraEditor/Private/ViewModels/Stack/NiagaraStackGraphUtilities.cpp b/Engine/Plugins/FX/Niagara/Source/NiagaraEditor/Private/ViewModels/Stack/NiagaraStackGraphUtilities.cpp
--- a/Engine/Plugins/FX/Niagara/Source/NiagaraEditor/Private/ViewModels/Stack/NiagaraStackGraphUtilities.cpp
+++ b/Engine/Plugins/FX/Niagara/Source/NiagaraEditor/Private/ViewModels/Stack/NiagaraStackGraphUtilities.cpp
@@ -3916,6 +3916,69 @@ void FNiagaraStackGraphUtilities::DependencyUtilities::GetModuleScriptAssetsByDe
 	}
 }
 
+void FNiagaraStackGraphUtilities::FixDynamicInputNodeOutputPinsFromExternalChanges(UNiagaraNodeFunctionCall& InFunctionCallNode)
+{
+	if (InFunctionCallNode.FunctionScript == nullptr || InFunctionCallNode.FunctionScript->GetUsage() != ENiagaraScriptUsage::DynamicInput)
+	{
+		return;
+	}
+
+	FPinCollectorArray OutputPins;
+	InFunctionCallNode.GetOutputPins(OutputPins);
+	if (OutputPins.Num() > 1)
+	{
+		// Try to find an orphaned pin which is still linked, and any non-orphaned unlinked pins.
+		UEdGraphPin* LinkedOrphanedPin = nullptr;
+		TArray<UEdGraphPin*> UnlinkedValidPins;
+		for (UEdGraphPin* OutputPin : OutputPins)
+		{
+			if (OutputPin->bOrphanedPin)
+			{
+				if (OutputPin->LinkedTo.Num() == 1 && ensureMsgf(LinkedOrphanedPin == nullptr, TEXT("Dynamic inputs should only ever have one linked pin.")))
+				{
+					LinkedOrphanedPin = OutputPin;
+				}
+			}
+			else
+			{
+				if (OutputPin->LinkedTo.Num() == 0)
+				{
+					UnlinkedValidPins.Add(OutputPin);
+				}
+			}
+		}
+
+		UEdGraphPin* LinkedValidPin = nullptr;
+		if (LinkedOrphanedPin != nullptr && UnlinkedValidPins.Num() > 0)
+		{
+			// Try to match the first available valid pin to the orphaned pin by type and then fix up the pin links.
+			for (UEdGraphPin* UnlinkedValidPin : UnlinkedValidPins)
+			{
+				if (UnlinkedValidPin->PinType == LinkedOrphanedPin->PinType)
+				{
+					// A valid unlinked pin was found, so move the link from the orphaned pin to the valid pin.
+					UnlinkedValidPin->MakeLinkTo(LinkedOrphanedPin->LinkedTo[0]);
+					LinkedOrphanedPin->BreakAllPinLinks(false);
+					LinkedValidPin = UnlinkedValidPin;
+					break;
+				}
+			}
+		}
+
+		if (LinkedValidPin != nullptr)
+		{
+			// If a valid pin was linked, remove all of the other output pins since they're not valid for a dynamic input node.
+			for (UEdGraphPin* OutputPin : OutputPins)
+			{
+				if (OutputPin != LinkedValidPin)
+				{
+					InFunctionCallNode.RemovePin(OutputPin);
+				}
+			}
+		}
+	}
+}
+
 int32 FNiagaraStackGraphUtilities::DependencyUtilities::FindBestIndexForModuleInStack(UNiagaraNodeFunctionCall& ModuleNode, UNiagaraGraph& EmitterScriptGraph)
 {
 	// Check if the new module node has any dependencies to begin with. If not, early exit.
diff --git a/Engine/Plugins/FX/Niagara/Source/NiagaraEditor/Public/ViewModels/Stack/NiagaraStackGraphUtilities.h b/Engine/Plugins/FX/Niagara/Source/NiagaraEditor/Public/ViewModels/Stack/NiagaraStackGraphUtilities.h
--- a/Engine/Plugins/FX/Niagara/Source/NiagaraEditor/Public/ViewModels/Stack/NiagaraStackGraphUtilities.h
+++ b/Engine/Plugins/FX/Niagara/Source/NiagaraEditor/Public/ViewModels/Stack/NiagaraStackGraphUtilities.h
@@ -299,6 +299,9 @@ namespace FNiagaraStackGraphUtilities
 
 	FGuid GetScriptVariableIdForLinkedModuleParameterHandle(const FNiagaraParameterHandle& LinkedOutputHandle, FNiagaraTypeDefinition LinkedType, UNiagaraGraph& TargetGraph);
 
+	/** Fixes orphaned and disconnected output pins on dynamic input nodes resulting from the underlying script changing and the pins being reallocated. */
+	void FixDynamicInputNodeOutputPinsFromExternalChanges(UNiagaraNodeFunctionCall& InFunctionCallNode);
+
 	namespace DependencyUtilities
 	{
 		bool DoesStackModuleProvideDependency(const FNiagaraStackModuleData& StackModuleData, const FNiagaraModuleDependency& SourceModuleRequiredDependency, const UNiagaraNodeOutput& SourceOutputNode);
