Fixed being unable to rename a parameter on a map get/set to another parameter provided it's the same type.
Also fixed a resulting metadata merge after renaming since we really only want a pin to reference another parameter.

diff --git a/Engine/Plugins/FX/Niagara/Source/NiagaraEditor/Private/NiagaraGraph.cpp b/Engine/Plugins/FX/Niagara/Source/NiagaraEditor/Private/NiagaraGraph.cpp
--- a/Engine/Plugins/FX/Niagara/Source/NiagaraEditor/Private/NiagaraGraph.cpp
+++ b/Engine/Plugins/FX/Niagara/Source/NiagaraEditor/Private/NiagaraGraph.cpp
@@ -2348,7 +2348,9 @@ bool UNiagaraGraph::RenameParameterFromPin(const FNiagaraVariable& Parameter, FN
 		FNiagaraParameterDefinitionsUtilities::TrySubscribeScriptVarToDefinitionByName(TargetScriptVar, ScriptViewModel.Get());
 	};
 
-	if (Parameter.GetName() == NewName)
+	FNiagaraVariable NewVariable(Parameter.GetType(), NewName);
+	UNiagaraScriptVariable* ExistingVariable = GetScriptVariable(NewVariable);
+	if (Parameter.GetName() == NewName || (ExistingVariable != nullptr && ExistingVariable->Variable.GetType() == NewVariable.GetType()))
 	{
 		return true;
 	}
diff --git a/Engine/Plugins/FX/Niagara/Source/NiagaraEditor/Private/NiagaraNodeParameterMapGet.cpp b/Engine/Plugins/FX/Niagara/Source/NiagaraEditor/Private/NiagaraNodeParameterMapGet.cpp
--- a/Engine/Plugins/FX/Niagara/Source/NiagaraEditor/Private/NiagaraNodeParameterMapGet.cpp
+++ b/Engine/Plugins/FX/Niagara/Source/NiagaraEditor/Private/NiagaraNodeParameterMapGet.cpp
@@ -85,10 +85,19 @@ bool UNiagaraNodeParameterMapGet::VerifyEditablePinName(const FText& InName, FTe
 		return false;
 	}
 	
-	if(GetNiagaraGraph()->GetScriptVariable(FName(InName.ToString())) != nullptr && !InName.EqualToCaseIgnored(FText::FromString(InGraphPinObj->GetName())))
+	UNiagaraScriptVariable* ExistingScriptVariable = GetNiagaraGraph()->GetScriptVariable(FName(InName.ToString()));
+	if(ExistingScriptVariable != nullptr)
 	{
-		OutErrorMessage = LOCTEXT("InvalidName_VariableExists", "This variable already exists.");
-		return false;
+		// if the variable already exists and the type matches, we are trying to reference an existing parameter which is allowed.
+		if(ExistingScriptVariable->Variable.GetType() == UEdGraphSchema_Niagara::PinToTypeDefinition(InGraphPinObj))
+		{
+			return true;
+		}
+		else
+		{
+			OutErrorMessage = LOCTEXT("InvalidName_VariableExistsDifferentType", "This variable already exists with a different type.\nChoose another name.");
+			return false;
+		}
 	}
 	
 	return true;
diff --git a/Engine/Plugins/FX/Niagara/Source/NiagaraEditor/Private/NiagaraNodeParameterMapSet.cpp b/Engine/Plugins/FX/Niagara/Source/NiagaraEditor/Private/NiagaraNodeParameterMapSet.cpp
--- a/Engine/Plugins/FX/Niagara/Source/NiagaraEditor/Private/NiagaraNodeParameterMapSet.cpp
+++ b/Engine/Plugins/FX/Niagara/Source/NiagaraEditor/Private/NiagaraNodeParameterMapSet.cpp
@@ -82,10 +82,19 @@ bool UNiagaraNodeParameterMapSet::VerifyEditablePinName(const FText& InName, FTe
 		return false;
 	}
 
-	if(GetNiagaraGraph()->GetScriptVariable(FName(InName.ToString())) != nullptr && !InName.EqualToCaseIgnored(FText::FromString(InGraphPinObj->GetName())))
+	UNiagaraScriptVariable* ExistingScriptVariable = GetNiagaraGraph()->GetScriptVariable(FName(InName.ToString()));
+	if(ExistingScriptVariable != nullptr)
 	{
-		OutErrorMessage = LOCTEXT("InvalidName_VariableExists", "This variable already exists.");
-		return false;
+		// if the variable already exists and the type matches, we are trying to reference an existing parameter which is allowed.
+		if(ExistingScriptVariable->Variable.GetType() == UEdGraphSchema_Niagara::PinToTypeDefinition(InGraphPinObj))
+		{
+			return true;
+		}
+		else
+		{
+			OutErrorMessage = LOCTEXT("InvalidName_VariableExistsDifferentType", "This variable already exists with a different type.\nChoose another name.");
+			return false;
+		}
 	}
 	
 	return true;
