Reduce friction for capturing sim caches in the editor

diff --git a/Engine/Plugins/FX/Niagara/Source/Niagara/Classes/NiagaraSimCacheCapture.cpp b/Engine/Plugins/FX/Niagara/Source/Niagara/Classes/NiagaraSimCacheCapture.cpp
new file mode 100644
--- /dev/null
+++ b/Engine/Plugins/FX/Niagara/Source/Niagara/Classes/NiagaraSimCacheCapture.cpp
@@ -0,0 +1,132 @@
+﻿// Copyright Epic Games, Inc. All Rights Reserved.
+
+#include "NiagaraSimCacheCapture.h"
+
+#include "Containers/Ticker.h"
+
+#include "NiagaraComponent.h"
+#include "NiagaraSimCache.h"
+
+void FNiagaraSimCacheCapture::FinishCapture()
+{
+	if (TickerHandle.IsValid())
+	{
+		FTSTicker::GetCoreTicker().RemoveTicker(TickerHandle);
+	}
+
+	UNiagaraSimCache* CaptureSimCache = WeakCaptureSimCache.Get();
+
+	if ( CaptureSimCache != nullptr )
+	{
+		CaptureSimCache->EndWrite();
+	}
+	CaptureComplete.Broadcast(CaptureSimCache);
+
+	WeakCaptureComponent = nullptr;
+}
+
+void FNiagaraSimCacheCapture::CaptureNiagaraSimCache(UNiagaraSimCache* SimCache,
+                                                     FNiagaraSimCacheCreateParameters CreateParameters, UNiagaraComponent* NiagaraComponent, FNiagaraSimCacheCaptureParameters InCaptureParameters)
+{
+	WeakCaptureSimCache = SimCache;
+	WeakCaptureComponent = NiagaraComponent;
+	CaptureParameters = InCaptureParameters;
+
+	if(SimCache == nullptr || NiagaraComponent == nullptr)
+	{
+		FinishCapture();
+		return;
+	}
+	
+	SimCache->BeginWrite(CreateParameters, NiagaraComponent);
+
+	// In this path we are going to manually advance the simulation until we complete
+	if (CaptureParameters.bManuallyAdvanceSimulation)
+	{
+		while (WeakCaptureComponent != nullptr && NiagaraComponent != nullptr)
+		{
+			NiagaraComponent->AdvanceSimulation(1, CaptureParameters.AdvanceDeltaTime);
+			OnFrameTick(CaptureParameters.AdvanceDeltaTime);
+		}
+	}
+	// In this path we are going to monitor the simulation and capture each frame
+	else
+	{
+		if (NiagaraComponent && NiagaraComponent->GetWorld())
+		{
+			TickerHandle = FTSTicker::GetCoreTicker().AddTicker(FTickerDelegate::CreateRaw(this, &FNiagaraSimCacheCapture::OnFrameTick));
+		}
+
+		if (!TickerHandle.IsValid())
+		{
+			FinishCapture();
+		}
+	}
+}
+
+bool FNiagaraSimCacheCapture::CaptureCurrentFrameImmediate(UNiagaraSimCache* SimCache, FNiagaraSimCacheCreateParameters CreateParameters, UNiagaraComponent* NiagaraComponent, UNiagaraSimCache*& OutSimCache, bool bAdvanceSimulation, float AdvanceDeltaTime)
+{
+	bool bSuccess = false;
+	if (SimCache != nullptr && NiagaraComponent != nullptr)
+	{
+		if ( bAdvanceSimulation )
+		{
+			NiagaraComponent->AdvanceSimulation(1, AdvanceDeltaTime);
+		}
+
+		SimCache->BeginWrite(CreateParameters, NiagaraComponent);
+		SimCache->WriteFrame(NiagaraComponent);
+		SimCache->EndWrite();
+		bSuccess = SimCache->IsCacheValid();
+	}
+	OutSimCache = SimCache;
+	return bSuccess;
+}
+
+bool FNiagaraSimCacheCapture::OnFrameTick(float DeltaTime)
+{
+	// Component invalid or not active?  If so complete the cache recording
+	UNiagaraComponent* CaptureComponent = WeakCaptureComponent.Get();
+	UNiagaraSimCache* CaptureSimCache = WeakCaptureSimCache.Get();
+	
+	if ( !CaptureComponent || !CaptureComponent->IsActive() || !CaptureSimCache )
+	{
+		FinishCapture();
+		return true;
+	}
+
+	// Should we record this frame?
+	if ( (CaptureFrameCounter % CaptureParameters.CaptureRate) == 0 )
+	{
+		// If we fail to capture the frame it might be because things became invalid
+		// Or it might be because the simulation was not ticked since the last capture in which case don't advance the counter
+		if ( CaptureSimCache->WriteFrame(CaptureComponent) == false )
+		{
+			if ( CaptureSimCache->IsCacheValid() == false )
+			{
+				FinishCapture();
+			}
+
+			// Make sure we don't keep this alive forever, if we didn't managed to capture anything in 10 ticks something has probably gone wrong so bail
+			if (TimeOutCounter++ > 10)
+			{
+				UE_LOG(LogNiagara, Warning, TEXT("SimCache Write has failed too many times, abandoning capturing for (%s)"), *GetFullNameSafe(CaptureSimCache));
+				FinishCapture();
+			}
+			return true;
+		}
+	}
+
+	TimeOutCounter = 0;
+	++CaptureFrameCounter;
+
+	// Have we recorded all the frames we need?
+	// Note: the -1 is because T0 was the initial frame
+	if ( (CaptureParameters.NumFrames > 0) && (CaptureFrameCounter > (CaptureParameters.CaptureRate * (CaptureParameters.NumFrames - 1))) )
+	{
+		FinishCapture();
+		return true;
+	}
+
+	return true;
+}
diff --git a/Engine/Plugins/FX/Niagara/Source/Niagara/Classes/NiagaraSimCacheCapture.h b/Engine/Plugins/FX/Niagara/Source/Niagara/Classes/NiagaraSimCacheCapture.h
new file mode 100644
--- /dev/null
+++ b/Engine/Plugins/FX/Niagara/Source/Niagara/Classes/NiagaraSimCacheCapture.h
@@ -0,0 +1,82 @@
+﻿// Copyright Epic Games, Inc. All Rights Reserved.
+
+#pragma once
+
+#include "NiagaraSimCache.h"
+#include "Containers/Ticker.h"
+#include "NiagaraSimCacheCapture.generated.h"
+
+class UNiagaraComponent;
+class UNiagaraSimCache;
+struct FNiagaraSimCacheCreateParameters;
+
+USTRUCT(BlueprintType)
+struct FNiagaraSimCacheCaptureParameters
+{
+	GENERATED_BODY()
+
+	/**
+	 The max number of frames to capture. The capture will stop when the simulation completes or we have rendered this many frames, whichever happens first.
+	 Set to 0 to capture until simulation completes.
+	 **/
+	UPROPERTY(EditAnywhere, BlueprintReadWrite, Category = "SimCache")
+	int32 NumFrames = 0;
+
+	/**
+	 Allows for reducing the frequency of captured frames in relation to the simulation's frames. The ratio is 1 / CaptureRate, so a CaptureRate of 2 would captures frames 0, 2, 4, etc.
+	 */	
+	UPROPERTY(EditAnywhere, BlueprintReadWrite, Category = "SimCache")
+	int32 CaptureRate = 1;
+
+	/**
+	 Controls manually advancing the simulation in a loop vs reading from it every frame. If bAdvanceSimulation is true, this capture will manually advance the simulation
+	 */
+	UPROPERTY(EditAnywhere, BlueprintReadWrite, Category = "SimCache")
+	bool bManuallyAdvanceSimulation = false;
+
+	/*
+	 The delta time in seconds to use when manually advancing the simulation. Defaults to 1/60th of a second (0.01666).
+	 */
+	UPROPERTY(EditAnywhere, BlueprintReadWrite, Category = "SimCache", meta = (EditCondition = "bManuallyAdvanceSimulation"))
+	float AdvanceDeltaTime = 0.01666f;
+};
+
+class NIAGARA_API FNiagaraSimCacheCapture
+{
+public:
+	DECLARE_MULTICAST_DELEGATE_OneParam(FOnCaptureComplete, UNiagaraSimCache*);
+
+	FOnCaptureComplete& OnCaptureComplete() {return CaptureComplete;}
+	
+	
+	// Captures a niagara sim cache. The caller is responsible for managing the lifetime of the provided component and sim cache. 
+	void CaptureNiagaraSimCache(UNiagaraSimCache* SimCache, FNiagaraSimCacheCreateParameters CreateParameters, UNiagaraComponent* NiagaraComponent, FNiagaraSimCacheCaptureParameters CaptureParameters);
+
+	/**
+	Captures the simulations current frame data into the SimCache.
+	This happens immediately so you may need to be careful with tick order of the component you are capturing from.
+	The return can be invalid if the component can not be captured for some reason (i.e. not active).
+	When AdvanceSimulation is true we will manually advance the simulation one frame using the provided AdvanceDeltaTime before capturing.
+	*/
+	bool CaptureCurrentFrameImmediate(UNiagaraSimCache* SimCache, FNiagaraSimCacheCreateParameters CreateParameters, UNiagaraComponent* NiagaraComponent, UNiagaraSimCache*& OutSimCache, bool bAdvanceSimulation=false, float AdvanceDeltaTime=0.01666f);
+
+private:
+
+	TWeakObjectPtr<UNiagaraSimCache> WeakCaptureSimCache;
+
+	TWeakObjectPtr<UNiagaraComponent> WeakCaptureComponent;
+	
+	FOnCaptureComplete CaptureComplete;
+
+	FTSTicker::FDelegateHandle TickerHandle;
+
+	FNiagaraSimCacheCaptureParameters CaptureParameters;
+
+	int32 CaptureFrameCounter = 0;
+
+	int32 TimeOutCounter = 0;
+	
+	bool OnFrameTick(float DeltaTime);
+
+	void FinishCapture();
+};
diff --git a/Engine/Plugins/FX/Niagara/Source/Niagara/Classes/NiagaraSimCacheFunctionLibrary.h b/Engine/Plugins/FX/Niagara/Source/Niagara/Classes/NiagaraSimCacheFunctionLibrary.h
--- a/Engine/Plugins/FX/Niagara/Source/Niagara/Classes/NiagaraSimCacheFunctionLibrary.h
+++ b/Engine/Plugins/FX/Niagara/Source/Niagara/Classes/NiagaraSimCacheFunctionLibrary.h
@@ -3,8 +3,8 @@
 #pragma once
 
 #include "NiagaraSimCache.h"
+#include "NiagaraSimCacheCapture.h"
 
-#include "Containers/Ticker.h"
 #include "Kismet/BlueprintFunctionLibrary.h"
 #include "Kismet/BlueprintAsyncActionBase.h"
 
@@ -23,27 +23,18 @@ public:
 
 	UPROPERTY()
 	TObjectPtr<UNiagaraComponent> CaptureComponent;
-
-	UPROPERTY()
-	int32 CaptureNumFrames = 0;
-
-	UPROPERTY()
-	int32 CaptureFrameRate = 0;
-
-	UPROPERTY()
-	int32 CaptureFrameCounter = 0;
-
-	UPROPERTY()
-	int32 TimeOutCounter = 0;
-
+	
 	UPROPERTY(BlueprintAssignable)
 	FOnCaptureComplete CaptureComplete;
 
-	FTSTicker::FDelegateHandle TickerHandle;
+	FNiagaraSimCacheCapture SimCacheCapture;
+
+	FNiagaraSimCacheCaptureParameters CaptureParameters;
+
+	FNiagaraSimCacheCreateParameters CreateParameters;
 
 	virtual void Activate() override;
 	virtual void SetReadyToDestroy() override;
-	bool OnFrameTick(float DeltaTime);
 
 	/**
 	Capture multiple frames from the provided simulation into a SimCache until the simulation becomes inactive, completes or we hit the NumFrames limit.
@@ -51,23 +42,22 @@ public:
 	CaptureRate allows you to reduce the rate of capture, i.e. a rate of 2 would capture frames 0, 2, 4, etc.
 	When AdvanceSimulation is true we will manually advance the simulation in a loop until we have captured the number of frames request, rather than reading from the component each frame.
 	*/
-	UFUNCTION(BlueprintCallable, meta=(BlueprintInternalUseOnly="true", Category=NiagaraSimCache))
+	UFUNCTION(BlueprintCallable, meta=(BlueprintInternalUseOnly="true", Category=NiagaraSimCache, AdvancedDisplay = "CaptureRate, bAdvanceSimulation, AdvanceDeltaTime"))
 	static UAsyncNiagaraCaptureSimCache* CaptureNiagaraSimCacheMultiFrame(UNiagaraSimCache* SimCache, FNiagaraSimCacheCreateParameters CreateParameters, UNiagaraComponent* NiagaraComponent, UNiagaraSimCache*& OutSimCache, int32 NumFrames = 16, int32 CaptureRate = 1, bool bAdvanceSimulation=false, float AdvanceDeltaTime=0.01666f);
-
+	
 	/**
 	Capture frames from the provided simulation into a SimCache until the simulation becomes inactive or completes.
 	Capture occurs at the end of each frame with the first frame being this frame.
 	CaptureRate allows you to reduce the rate of capture, i.e. a rate of 2 would capture frames 0, 2, 4, etc.
 	When AdvanceSimulation is true we will manually advance the simulation until the capture is complete inside a loop, rather than reading from the component each frame.
 	*/
-	UFUNCTION(BlueprintCallable, meta=(BlueprintInternalUseOnly="true", Category=NiagaraSimCache))
+	UFUNCTION(BlueprintCallable, meta=(BlueprintInternalUseOnly="true", Category=NiagaraSimCache, AdvancedDisplay = "CaptureRate, bAdvanceSimulation, AdvanceDeltaTime"))
 	static UAsyncNiagaraCaptureSimCache* CaptureNiagaraSimCacheUntilComplete(UNiagaraSimCache* SimCache, FNiagaraSimCacheCreateParameters CreateParameters, UNiagaraComponent* NiagaraComponent, UNiagaraSimCache*& OutSimCache, int32 CaptureRate = 1, bool bAdvanceSimulation=false, float AdvanceDeltaTime=0.01666f);
 
-	UPROPERTY()
-	bool bAdvanceSimulation = false;
+private:
+	static UAsyncNiagaraCaptureSimCache* CaptureNiagaraSimCacheImpl(UNiagaraSimCache* SimCache, FNiagaraSimCacheCreateParameters CreateParameters, UNiagaraComponent* NiagaraComponent, UNiagaraSimCache*& OutSimCache, FNiagaraSimCacheCaptureParameters CaptureParameters);
 
-	UPROPERTY()
-	float AdvanceDeltaTime = 1.0f / 60.0f;
+	void CaptureFinished(UNiagaraSimCache* CapturedSimCache);
 };
 
 UCLASS()
diff --git a/Engine/Plugins/FX/Niagara/Source/Niagara/Private/NiagaraSimCacheFunctionLibrary.cpp b/Engine/Plugins/FX/Niagara/Source/Niagara/Private/NiagaraSimCacheFunctionLibrary.cpp
--- a/Engine/Plugins/FX/Niagara/Source/Niagara/Private/NiagaraSimCacheFunctionLibrary.cpp
+++ b/Engine/Plugins/FX/Niagara/Source/Niagara/Private/NiagaraSimCacheFunctionLibrary.cpp
@@ -3,6 +3,7 @@
 #include "NiagaraSimCacheFunctionLibrary.h"
 #include "Engine/Engine.h"
 #include "NiagaraComponent.h"
+#include "NiagaraSimCacheCapture.h"
 
 #include UE_INLINE_GENERATED_CPP_BY_NAME(NiagaraSimCacheFunctionLibrary)
 
@@ -10,42 +11,16 @@ void UAsyncNiagaraCaptureSimCache::Activate()
 {
 	Super::Activate();
 
-	// In this path we are going to manually advance the simulation until we complete
-	if (bAdvanceSimulation)
-	{
-		while (CaptureComponent != nullptr)
-		{
-			CaptureComponent->AdvanceSimulation(1, AdvanceDeltaTime);
-			OnFrameTick(AdvanceDeltaTime);
-		}
-	}
-	// In this path we are going to monitor the simulation and capture each frame
-	else
-	{
-		if (CaptureComponent != nullptr)
-		{
-			if (UWorld* OwnerWorld = CaptureComponent->GetWorld())
-			{
-				TickerHandle = FTSTicker::GetCoreTicker().AddTicker(FTickerDelegate::CreateUObject(this, &UAsyncNiagaraCaptureSimCache::OnFrameTick), 0);
-			}
-		}
-
-		if (CaptureSimCache == nullptr || CaptureComponent == nullptr || TickerHandle.IsValid() == false)
-		{
-			SetReadyToDestroy();
-		}
-	}
+	SimCacheCapture.OnCaptureComplete().AddUObject(this, &UAsyncNiagaraCaptureSimCache::CaptureFinished);
+	
+	SimCacheCapture.CaptureNiagaraSimCache(CaptureSimCache, CreateParameters, CaptureComponent, CaptureParameters);
+	
 }
 
 void UAsyncNiagaraCaptureSimCache::SetReadyToDestroy()
 {
 	Super::SetReadyToDestroy();
 
-	if (TickerHandle.IsValid())
-	{
-		FTSTicker::GetCoreTicker().RemoveTicker(TickerHandle);
-	}
-
 	if ( CaptureSimCache != nullptr )
 	{
 		CaptureSimCache->EndWrite();
@@ -55,66 +30,28 @@ void UAsyncNiagaraCaptureSimCache::SetReadyToDestroy()
 	CaptureComponent = nullptr;
 }
 
-bool UAsyncNiagaraCaptureSimCache::OnFrameTick(float DeltaTime)
+UAsyncNiagaraCaptureSimCache* UAsyncNiagaraCaptureSimCache::CaptureNiagaraSimCacheMultiFrame(UNiagaraSimCache* SimCache, FNiagaraSimCacheCreateParameters CreateParameters, UNiagaraComponent* NiagaraComponent, UNiagaraSimCache*& OutSimCache, int32 NumFrames, int32 CaptureRate, bool bAdvanceSimulation, float AdvanceDeltaTime)
 {
-	// Component invalid or not active?  If so complete the cache recording
-	if ( CaptureComponent == nullptr || !CaptureComponent->IsActive() || CaptureSimCache == nullptr )
-	{
-		SetReadyToDestroy();
-		return true;
-	}
-
-	// Should we record this frame?
-	if ( (CaptureFrameCounter % CaptureFrameRate) == 0 )
-	{
-		// If we fail to capture the frame it might be because things became invalid
-		// Or it might be because the simulation was not ticked since the last capture in which case don't advance the counter
-		if ( CaptureSimCache->WriteFrame(CaptureComponent) == false )
-		{
-			if ( CaptureSimCache->IsCacheValid() == false )
-			{
-				SetReadyToDestroy();
-			}
-
-			// Make sure we don't keep this alive forever, if we didn't managed to capture anything in 10 ticks something has probably gone wrong so bail
-			if (TimeOutCounter++ > 10)
-			{
-				UE_LOG(LogNiagara, Warning, TEXT("SimCache Write has failed too many times, abandoning capturing for (%s)"), *GetFullNameSafe(CaptureSimCache));
-				SetReadyToDestroy();
-			}
-			return true;
-		}
-	}
-
-	TimeOutCounter = 0;
-	++CaptureFrameCounter;
-
-	// Have we recorded all the frames we need?
-	// Note: the -1 is because T0 was the initial frame
-	if ( (CaptureNumFrames > 0) && (CaptureFrameCounter > (CaptureFrameRate * (CaptureNumFrames - 1))) )
-	{
-		SetReadyToDestroy();
-		return true;
-	}
-
-	return true;
+	FNiagaraSimCacheCaptureParameters CaptureParameters;
+	
+	CaptureParameters.NumFrames = FMath::Max(1, NumFrames);
+	CaptureParameters.CaptureRate = FMath::Max(1, CaptureRate);
+	CaptureParameters.bManuallyAdvanceSimulation = bAdvanceSimulation;
+	CaptureParameters.AdvanceDeltaTime = AdvanceDeltaTime;
+
+	return CaptureNiagaraSimCacheImpl(SimCache, CreateParameters, NiagaraComponent, OutSimCache, CaptureParameters);
 }
 
-UAsyncNiagaraCaptureSimCache* UAsyncNiagaraCaptureSimCache::CaptureNiagaraSimCacheMultiFrame(UNiagaraSimCache* SimCache, FNiagaraSimCacheCreateParameters CreateParameters, UNiagaraComponent* NiagaraComponent, UNiagaraSimCache*& OutSimCache, int32 NumFrames, int32 CaptureRate, bool bAdvanceSimulation, float AdvanceDeltaTime)
+UAsyncNiagaraCaptureSimCache* UAsyncNiagaraCaptureSimCache::CaptureNiagaraSimCacheImpl(UNiagaraSimCache* SimCache,
+	FNiagaraSimCacheCreateParameters CreateParameters, UNiagaraComponent* NiagaraComponent, UNiagaraSimCache*& OutSimCache,
+	FNiagaraSimCacheCaptureParameters InCaptureParameters)
 {
 	UAsyncNiagaraCaptureSimCache* CaptureAction = NewObject<UAsyncNiagaraCaptureSimCache>();
+
+	CaptureAction->CaptureParameters = InCaptureParameters;
+	CaptureAction->CreateParameters = CreateParameters;
 	CaptureAction->CaptureSimCache = SimCache;
 	CaptureAction->CaptureComponent = NiagaraComponent;
-	CaptureAction->CaptureNumFrames = FMath::Max(1, NumFrames);
-	CaptureAction->CaptureFrameRate = FMath::Max(1, CaptureRate);
-	CaptureAction->CaptureFrameCounter = 0;
-	CaptureAction->bAdvanceSimulation = bAdvanceSimulation;
-	CaptureAction->AdvanceDeltaTime = AdvanceDeltaTime;
-
-	if (SimCache != nullptr)
-	{
-		SimCache->BeginWrite(CreateParameters, NiagaraComponent);
-	}
 	OutSimCache = SimCache;
 
 	return CaptureAction;
@@ -122,22 +59,21 @@ UAsyncNiagaraCaptureSimCache* UAsyncNiagaraCaptureSimCache::CaptureNiagaraSimCac
 
 UAsyncNiagaraCaptureSimCache* UAsyncNiagaraCaptureSimCache::CaptureNiagaraSimCacheUntilComplete(UNiagaraSimCache* SimCache, FNiagaraSimCacheCreateParameters CreateParameters, UNiagaraComponent* NiagaraComponent, UNiagaraSimCache*& OutSimCache, int32 CaptureRate, bool bAdvanceSimulation, float AdvanceDeltaTime)
 {
-	UAsyncNiagaraCaptureSimCache* CaptureAction = NewObject<UAsyncNiagaraCaptureSimCache>();
-	CaptureAction->CaptureSimCache = SimCache;
-	CaptureAction->CaptureComponent = NiagaraComponent;
-	CaptureAction->CaptureNumFrames = 0;
-	CaptureAction->CaptureFrameRate = FMath::Max(1, CaptureRate);
-	CaptureAction->CaptureFrameCounter = 0;
-	CaptureAction->bAdvanceSimulation = bAdvanceSimulation;
-	CaptureAction->AdvanceDeltaTime = AdvanceDeltaTime;
-
-	if (SimCache != nullptr)
-	{
-		SimCache->BeginWrite(CreateParameters, NiagaraComponent);
-	}
-	OutSimCache = SimCache;
+	
+	FNiagaraSimCacheCaptureParameters CaptureParameters;
+	
+	CaptureParameters.NumFrames = 0;
+	CaptureParameters.CaptureRate = FMath::Max(1, CaptureRate);
+	CaptureParameters.bManuallyAdvanceSimulation = bAdvanceSimulation;
+	CaptureParameters.AdvanceDeltaTime = AdvanceDeltaTime;
+
+	return CaptureNiagaraSimCacheImpl(SimCache, CreateParameters, NiagaraComponent, OutSimCache, CaptureParameters);
+}
 
-	return CaptureAction;
+void UAsyncNiagaraCaptureSimCache::CaptureFinished(UNiagaraSimCache* InCapturedSimCache)
+{
+	CaptureSimCache = InCapturedSimCache;
+	SetReadyToDestroy();
 }
 
 UNiagaraSimCacheFunctionLibrary::UNiagaraSimCacheFunctionLibrary(const FObjectInitializer& ObjectInitializer)
@@ -147,21 +83,9 @@ UNiagaraSimCacheFunctionLibrary::UNiagaraSimCacheFunctionLibrary(const FObjectIn
 
 bool UNiagaraSimCacheFunctionLibrary::CaptureNiagaraSimCacheImmediate(UNiagaraSimCache* SimCache, FNiagaraSimCacheCreateParameters CreateParameters, UNiagaraComponent* NiagaraComponent, UNiagaraSimCache*& OutSimCache, bool bAdvanceSimulation, float AdvanceDeltaTime)
 {
-	bool bSuccess = false;
-	if (SimCache != nullptr && NiagaraComponent != nullptr)
-	{
-		if ( bAdvanceSimulation )
-		{
-			NiagaraComponent->AdvanceSimulation(1, AdvanceDeltaTime);
-		}
-
-		SimCache->BeginWrite(CreateParameters, NiagaraComponent);
-		SimCache->WriteFrame(NiagaraComponent);
-		SimCache->EndWrite();
-		bSuccess = SimCache->IsCacheValid();
-	}
-	OutSimCache = SimCache;
-	return bSuccess;
+	FNiagaraSimCacheCapture SimCacheCapture;
+
+	return SimCacheCapture.CaptureCurrentFrameImmediate(SimCache, CreateParameters, NiagaraComponent, OutSimCache, bAdvanceSimulation, AdvanceDeltaTime);
 }
 
 UNiagaraSimCache* UNiagaraSimCacheFunctionLibrary::CreateNiagaraSimCache(UObject* WorldContextObject)
diff --git a/Engine/Plugins/FX/Niagara/Source/Niagara/Public/NiagaraSettings.h b/Engine/Plugins/FX/Niagara/Source/Niagara/Public/NiagaraSettings.h
--- a/Engine/Plugins/FX/Niagara/Source/Niagara/Public/NiagaraSettings.h
+++ b/Engine/Plugins/FX/Niagara/Source/Niagara/Public/NiagaraSettings.h
@@ -115,6 +115,10 @@ class NIAGARA_API UNiagaraSettings : public UDeveloperSettings
 	/** If true then the "link input" menu will also show variables of different types, as long as there is a conversion script for them. */
 	UPROPERTY(config, EditAnywhere, Category = Niagara)
 	bool bShowConvertibleInputsInStack = false;
+
+	/** The number of frames to capture when capturing a Sim Cache from the Niagara Component Details Panel. **/
+	UPROPERTY(Config, EditAnywhere, Category = SimulationCaching, meta = (UIMin = 1))
+	int32 QuickSimCacheCaptureFrameCount = 5;
 #endif // WITH_EDITORONLY_DATA
 
 	/** If true then active effects rebase the simulation positions to not lose precision. Can be turned off if not needed to skip unnecessary rebasing calculations. */
diff --git a/Engine/Plugins/FX/Niagara/Source/NiagaraEditor/Private/Customizations/NiagaraComponentDetails.cpp b/Engine/Plugins/FX/Niagara/Source/NiagaraEditor/Private/Customizations/NiagaraComponentDetails.cpp
--- a/Engine/Plugins/FX/Niagara/Source/NiagaraEditor/Private/Customizations/NiagaraComponentDetails.cpp
+++ b/Engine/Plugins/FX/Niagara/Source/NiagaraEditor/Private/Customizations/NiagaraComponentDetails.cpp
@@ -8,6 +8,7 @@
 #include "NiagaraDataInterface.h"
 #include "IDetailPropertyRow.h"
 #include "DetailWidgetRow.h"
+#include "NiagaraSimCacheCapture.h"
 #include "Widgets/Text/SInlineEditableTextBlock.h"
 #include "PropertyCustomizationHelpers.h"
 #include "Widgets/Text/STextBlock.h"
@@ -26,11 +27,13 @@
 #include "NiagaraScriptSource.h"
 #include "NiagaraNodeInput.h"
 #include "GameDelegates.h"
+#include "IAssetTools.h"
 #include "NiagaraEditorStyle.h"
 #include "IDetailChildrenBuilder.h"
 #include "NiagaraClipboard.h"
 #include "NiagaraConstants.h"
 #include "NiagaraScriptVariable.h"
+#include "NiagaraSettings.h"
 #include "NiagaraSystemEditorData.h"
 #include "NiagaraUserRedirectionParameterStore.h"
 #include "Widgets/Layout/SUniformGridPanel.h"
@@ -1254,6 +1257,11 @@ FNiagaraComponentDetails::~FNiagaraComponentDetails()
 		GEngine->OnWorldDestroyed().RemoveAll(this);
 	}
 
+	for(UNiagaraSimCache*& SimCache : CapturedCaches)
+	{
+		SimCache->ClearFlags(RF_Standalone);
+	}
+
 	FGameDelegates::Get().GetEndPlayMapDelegate().RemoveAll(this);
 }
 
@@ -1399,6 +1407,17 @@ void FNiagaraComponentDetails::CustomizeDetails(IDetailLayoutBuilder& DetailBuil
 						.Text(LOCTEXT("ResetEmitterButton", "Reset"))
 					]
 				]
+				+ SUniformGridPanel::Slot(2,0)
+				[
+					SNew(SButton)
+					.OnClicked(this, &FNiagaraComponentDetails::OnCaptureSelectedSystem)
+					.ToolTipText(LOCTEXT("CaptureSimCacheButtonToolTip", "Capture a temporary sim cache from the first selected particle system."))
+					.HAlign(HAlign_Center)
+					[
+						SNew(STextBlock)
+						.Text(LOCTEXT("CaptureSimCacheButton", "Capture"))
+					]
+				]
 			]
 		];
 }
@@ -1481,4 +1500,82 @@ FReply FNiagaraComponentDetails::OnDebugSelectedSystem()
 	return FReply::Handled();
 }
 
+FReply FNiagaraComponentDetails::OnCaptureSelectedSystem()
+{
+	if (!Builder)
+		return FReply::Handled();
+
+	const TArray< TWeakObjectPtr<UObject> >& SelectedObjects = Builder->GetSelectedObjects();
+
+	UNiagaraComponent* NiagaraComponentToUse = nullptr;
+	for (int32 Idx = 0; Idx < SelectedObjects.Num(); ++Idx)
+	{
+		if (SelectedObjects[Idx].IsValid())
+		{
+			if (AActor* Actor = Cast<AActor>(SelectedObjects[Idx].Get()))
+			{
+				for (UActorComponent* AC : Actor->GetComponents())
+				{
+					UNiagaraComponent* NiagaraComponent = Cast<UNiagaraComponent>(AC);
+					if (NiagaraComponent)
+					{
+						NiagaraComponentToUse = NiagaraComponent;
+						break;
+					}
+				}
+			}
+			else if (UNiagaraComponent* NiagaraComponent = Cast<UNiagaraComponent>(SelectedObjects[Idx].Get()))
+			{
+				NiagaraComponentToUse = NiagaraComponent;
+				break;
+			}
+		}
+	}
+
+	if (NiagaraComponentToUse)
+	{
+		
+		int32 CaptureIndex = ComponentCaptures.AddDefaulted();
+		UNiagaraSimCache* SimCache = NewObject<UNiagaraSimCache>(GetTransientPackage());
+		SimCache->SetFlags(RF_Standalone);
+		CapturedCaches.Add(SimCache);
+		FNiagaraSimCacheCapture& SimCacheCapture = ComponentCaptures[CaptureIndex];
+
+		FNiagaraSimCacheCreateParameters CreateParameters;
+		
+		FNiagaraSimCacheCaptureParameters CaptureParameters;
+
+		CaptureParameters.CaptureRate = 1;
+		// A value of 0 would capture until completion. Since this can apply broadly, looping systems would not finish capturing, so don't go below 1.
+		CaptureParameters.NumFrames = FMath::Max(1, GetDefault<UNiagaraSettings>()->QuickSimCacheCaptureFrameCount);
+		CaptureParameters.bManuallyAdvanceSimulation = false;
+		
+		if(NiagaraComponentToUse->GetWorld())
+		{
+			
+			SimCacheCapture.CaptureNiagaraSimCache(SimCache, CreateParameters, NiagaraComponentToUse, CaptureParameters);
+
+			SimCacheCapture.OnCaptureComplete().AddLambda(
+			[CaptureIndex, this](UNiagaraSimCache* CapturedSimCache)
+			{
+				if(CapturedSimCache)
+				{
+					GEditor->GetEditorSubsystem<UAssetEditorSubsystem>()->OpenEditorForAsset(CapturedSimCache, EToolkitMode::Standalone);
+				}
+
+				if(ComponentCaptures.IsValidIndex(CaptureIndex))
+				{
+					ComponentCaptures.RemoveAt(CaptureIndex);
+				}
+				
+			});
+		}
+
+		
+		return FReply::Handled();
+	}
+
+	return FReply::Handled();
+}
+
 #undef LOCTEXT_NAMESPACE
diff --git a/Engine/Plugins/FX/Niagara/Source/NiagaraEditor/Private/Customizations/NiagaraComponentDetails.h b/Engine/Plugins/FX/Niagara/Source/NiagaraEditor/Private/Customizations/NiagaraComponentDetails.h
--- a/Engine/Plugins/FX/Niagara/Source/NiagaraEditor/Private/Customizations/NiagaraComponentDetails.h
+++ b/Engine/Plugins/FX/Niagara/Source/NiagaraEditor/Private/Customizations/NiagaraComponentDetails.h
@@ -7,6 +7,7 @@
 #include "IDetailCustomNodeBuilder.h"
 #include "Input/Reply.h"
 #include "NiagaraComponent.h"
+#include "NiagaraSimCacheCapture.h"
 #include "ViewModels/TNiagaraViewModelManager.h"
 #include "ViewModels/HierarchyEditor/NiagaraHierarchyViewModelBase.h"
 #include "ViewModels/HierarchyEditor/NiagaraUserParametersHierarchyViewModel.h"
@@ -18,6 +19,7 @@ class FNiagaraParameterViewModelCustomDetails;
 class INiagaraParameterViewModel;
 class SNiagaraParameterEditor;
 class FNiagaraParameterViewModelCustomDetails;
+class FNiagaraSimCacheCapture;
 class UNiagaraSystem;
 
 USTRUCT()
@@ -46,9 +48,13 @@ protected:
 
 	FReply OnResetSelectedSystem();
 	FReply OnDebugSelectedSystem();
+	FReply OnCaptureSelectedSystem();
 private:
 	TWeakObjectPtr<UNiagaraComponent> Component;
 	IDetailLayoutBuilder* Builder = nullptr;
+
+	TArray<FNiagaraSimCacheCapture> ComponentCaptures;
+	TArray<UNiagaraSimCache*> CapturedCaches;
 };
 
 class FNiagaraSystemUserParameterDetails : public IDetailCustomization
