Fix scrubbing in Sim Cache editor frames being off.

diff --git a/Engine/Plugins/FX/Niagara/Source/NiagaraEditor/Private/Widgets/SNiagaraSimCacheViewTimeline.cpp b/Engine/Plugins/FX/Niagara/Source/NiagaraEditor/Private/Widgets/SNiagaraSimCacheViewTimeline.cpp
--- a/Engine/Plugins/FX/Niagara/Source/NiagaraEditor/Private/Widgets/SNiagaraSimCacheViewTimeline.cpp
+++ b/Engine/Plugins/FX/Niagara/Source/NiagaraEditor/Private/Widgets/SNiagaraSimCacheViewTimeline.cpp
@@ -99,7 +99,7 @@ FReply SNiagaraSimCacheViewTimeline::OnTimelineScrubbed(const FGeometry& MyGeome
 		const FVector2D LocalLocation = MyGeometry.AbsoluteToLocal(MouseEvent.GetScreenSpacePosition());
 		const float NormalizedTime = FMath::Clamp( LocalLocation.X / MyGeometry.Size.X, 0.0f, 1.0f);
 
-		const int32 NewFrameIndex = FMath::RoundToInt(NormalizedTime * float(ViewModel->GetNumFrames() - 1));
+		const int32 NewFrameIndex = FMath::Floor(NormalizedTime * float(ViewModel->GetNumFrames()));
 
 		ViewModel->SetFrameIndex(NewFrameIndex);
 
