- Add Material Binding support to Decal Renderer

diff --git a/Engine/Plugins/FX/Niagara/Source/Niagara/Private/NiagaraDecalRendererProperties.cpp b/Engine/Plugins/FX/Niagara/Source/Niagara/Private/NiagaraDecalRendererProperties.cpp
--- a/Engine/Plugins/FX/Niagara/Source/Niagara/Private/NiagaraDecalRendererProperties.cpp
+++ b/Engine/Plugins/FX/Niagara/Source/Niagara/Private/NiagaraDecalRendererProperties.cpp
@@ -5,6 +5,7 @@
 #include "NiagaraEmitterInstance.h"
 #include "NiagaraRenderer.h"
 #include "NiagaraRendererDecals.h"
+#include "Materials/MaterialInstanceConstant.h"
 #include "Modules/ModuleManager.h"
 
 #include UE_INLINE_GENERATED_CPP_BY_NAME(NiagaraDecalRendererProperties)
@@ -117,17 +118,43 @@ void UNiagaraDecalRendererProperties::PostInitProperties()
 	}
 }
 
+void UNiagaraDecalRendererProperties::Serialize(FArchive& Ar)
+{
+	// MIC will replace the main material during serialize
+	// Be careful if adding code that looks at the material to make sure you get the correct one
+	{
+#if WITH_EDITORONLY_DATA
+		TOptional<TGuardValue<TObjectPtr<UMaterialInterface>>> MICGuard;
+		if (Ar.IsSaving() && Ar.IsCooking() && MICMaterial)
+		{
+			MICGuard.Emplace(Material, MICMaterial);
+		}
+#endif
+
+		Super::Serialize(Ar);
+	}
+}
+
 #if WITH_EDITORONLY_DATA
 void UNiagaraDecalRendererProperties::PostEditChangeProperty(struct FPropertyChangedEvent& PropertyChangedEvent)
 {
 	Super::PostEditChangeProperty(PropertyChangedEvent);
 
 	const FName PropertyName = PropertyChangedEvent.GetPropertyName();
+	const FName MemberPropertyName = PropertyChangedEvent.GetMemberPropertyName();
 
 	if (PropertyName == GET_MEMBER_NAME_CHECKED(UNiagaraDecalRendererProperties, SourceMode))
 	{
 		UpdateSourceModeDerivates(SourceMode, true);
 	}
+
+	// Update our MICs if we change material / material bindings
+	//-OPT: Could narrow down further to only static materials
+	if ((PropertyName == GET_MEMBER_NAME_CHECKED(UNiagaraDecalRendererProperties, Material)) ||
+		(MemberPropertyName == GET_MEMBER_NAME_CHECKED(UNiagaraDecalRendererProperties, MaterialParameters)))
+	{
+		UpdateMICs();
+	}
 }
 #endif// WITH_EDITORONLY_DATA
 
@@ -211,6 +238,7 @@ void UNiagaraDecalRendererProperties::GetUsedMaterials(const FNiagaraEmitterInst
 void UNiagaraDecalRendererProperties::CacheFromCompiledData(const FNiagaraDataSetCompiledData* CompiledData)
 {
 	UpdateSourceModeDerivates(SourceMode);
+	UpdateMICs();
 
 	PositionDataSetAccessor.Init(CompiledData, PositionBinding.GetDataSetBindableVariable().GetName());
 	DecalOrientationDataSetAccessor.Init(CompiledData, DecalOrientationBinding.GetDataSetBindableVariable().GetName());
@@ -224,6 +252,14 @@ void UNiagaraDecalRendererProperties::CacheFromCompiledData(const FNiagaraDataSe
 void UNiagaraDecalRendererProperties::UpdateSourceModeDerivates(ENiagaraRendererSourceDataMode InSourceMode, bool bFromPropertyEdit)
 {
 	Super::UpdateSourceModeDerivates(InSourceMode, bFromPropertyEdit);
+
+	if ( UNiagaraEmitter* SrcEmitter = GetTypedOuter<UNiagaraEmitter>() )
+	{
+		for (FNiagaraMaterialAttributeBinding& MaterialParamBinding : MaterialParameters.AttributeBindings)
+		{
+			MaterialParamBinding.CacheValues(SrcEmitter);
+		}
+	}
 }
 
 bool UNiagaraDecalRendererProperties::PopulateRequiredBindings(FNiagaraParameterStore& InParameterStore)
@@ -239,6 +275,12 @@ bool UNiagaraDecalRendererProperties::PopulateRequiredBindings(FNiagaraParameter
 		}
 	}
 
+	for (const FNiagaraMaterialAttributeBinding& MaterialParamBinding : MaterialParameters.AttributeBindings)
+	{
+		InParameterStore.AddParameter(MaterialParamBinding.GetParamMapBindableVariable(), false);
+		bAnyAdded = true;
+	}
+
 	return bAnyAdded;
 }
 
@@ -270,15 +312,56 @@ void UNiagaraDecalRendererProperties::GetRendererTooltipWidgets(const FNiagaraEm
 	OutWidgets.Add(SNew(STextBlock).Text(LOCTEXT("DecalRenderer", "Decal Renderer")));
 }
 
-void UNiagaraDecalRendererProperties::GetRendererFeedback(const FVersionedNiagaraEmitter& InEmitter, TArray<FText>& OutErrors, TArray<FText>& OutWarnings, TArray<FText>& OutInfo) const
+void UNiagaraDecalRendererProperties::GetRendererFeedback(const FVersionedNiagaraEmitter& InEmitter, TArray<FNiagaraRendererFeedback>& OutErrors, TArray<FNiagaraRendererFeedback>& OutWarnings, TArray<FNiagaraRendererFeedback>& OutInfo) const
 {
 	Super::GetRendererFeedback(InEmitter, OutErrors, OutWarnings, OutInfo);
 
 	OutInfo.Emplace(LOCTEXT("DecalRenderingPreviewSceneInfo", "Decal Rendering does not show in the default preview scene, please modify the Preview Scene Settings to add ground geometry to visualize in preview."));
+
+	if (MaterialParameters.HasAnyBindings())
+	{
+		TArray<UMaterialInterface*> Materials;
+		GetUsedMaterials(nullptr, Materials);
+		MaterialParameters.GetFeedback(Materials, OutWarnings);
+	}
 }
 
+TArray<FNiagaraVariable> UNiagaraDecalRendererProperties::GetBoundAttributes() const
+{
+	TArray<FNiagaraVariable> BoundAttributes = Super::GetBoundAttributes();
+	BoundAttributes.Reserve(BoundAttributes.Num() + MaterialParameters.AttributeBindings.Num());
+
+	for (const FNiagaraMaterialAttributeBinding& MaterialParamBinding : MaterialParameters.AttributeBindings)
+	{
+		BoundAttributes.AddUnique(MaterialParamBinding.GetParamMapBindableVariable());
+	}
+	return BoundAttributes;
+}
+
+void UNiagaraDecalRendererProperties::RenameVariable(const FNiagaraVariableBase& OldVariable, const FNiagaraVariableBase& NewVariable, const FVersionedNiagaraEmitter& InEmitter)
+{
+	Super::RenameVariable(OldVariable, NewVariable, InEmitter);
+#if WITH_EDITORONLY_DATA
+	MaterialParameters.RenameVariable(OldVariable, NewVariable, InEmitter, GetCurrentSourceMode());
+#endif
+}
+
+void UNiagaraDecalRendererProperties::RemoveVariable(const FNiagaraVariableBase& OldVariable, const FVersionedNiagaraEmitter& InEmitter)
+{
+	Super::RemoveVariable(OldVariable, InEmitter);
+#if WITH_EDITORONLY_DATA
+	MaterialParameters.RemoveVariable(OldVariable, InEmitter, GetCurrentSourceMode());
+#endif
+}
 #endif // WITH_EDITORONLY_DATA
 
+void UNiagaraDecalRendererProperties::UpdateMICs()
+{
+#if WITH_EDITORONLY_DATA
+	UpdateMaterialParametersMIC(MaterialParameters, Material, MICMaterial);
+#endif
+}
+
 UMaterialInterface* UNiagaraDecalRendererProperties::GetMaterial(const FNiagaraEmitterInstance* InEmitter) const
 {
 	UMaterialInterface* MaterialInterface = nullptr;
@@ -287,6 +370,10 @@ UMaterialInterface* UNiagaraDecalRendererProperties::GetMaterial(const FNiagaraE
 		MaterialInterface = Cast<UMaterialInterface>(InEmitter->FindBinding(MaterialParameterBinding.Parameter));
 	}
 
+#if WITH_EDITORONLY_DATA
+	MaterialInterface = MaterialInterface ? MaterialInterface : ToRawPtr(MICMaterial);
+#endif
+
 	return MaterialInterface ? MaterialInterface : ToRawPtr(Material);
 }
 
diff --git a/Engine/Plugins/FX/Niagara/Source/Niagara/Private/NiagaraRendererDecals.cpp b/Engine/Plugins/FX/Niagara/Source/Niagara/Private/NiagaraRendererDecals.cpp
--- a/Engine/Plugins/FX/Niagara/Source/Niagara/Private/NiagaraRendererDecals.cpp
+++ b/Engine/Plugins/FX/Niagara/Source/Niagara/Private/NiagaraRendererDecals.cpp
@@ -158,13 +158,18 @@ FNiagaraDynamicDataBase* FNiagaraRendererDecals::GenerateDynamicData(const FNiag
 	const float DefaultFade = UNiagaraDecalRendererProperties::GetDefaultDecalFade();
 
 	// Check for Material Update
-	UMaterialInterface* Material = RendererProperties->GetMaterial(Emitter);
+	UMaterialInterface* Material = BaseMaterials_GT.Num() ? BaseMaterials_GT[0] : nullptr;
 	if (Material != WeakMaterial.Get())
 	{
 		WeakMaterial = Material;
 		ReleaseAllDecals();
 	}
 
+	if (RendererProperties->MaterialParameters.HasAnyBindings())
+	{
+		ProcessMaterialParameterBindings(RendererProperties->MaterialParameters, Emitter, MakeArrayView(BaseMaterials_GT));
+	}
+
 	// Generate list of decals to update
 	int NumActiveDecalProxies = 0;
 
diff --git a/Engine/Plugins/FX/Niagara/Source/Niagara/Public/NiagaraDecalRendererProperties.h b/Engine/Plugins/FX/Niagara/Source/Niagara/Public/NiagaraDecalRendererProperties.h
--- a/Engine/Plugins/FX/Niagara/Source/Niagara/Public/NiagaraDecalRendererProperties.h
+++ b/Engine/Plugins/FX/Niagara/Source/Niagara/Public/NiagaraDecalRendererProperties.h
@@ -10,6 +10,7 @@
 #include "NiagaraParameterBinding.h"
 #include "NiagaraDecalRendererProperties.generated.h"
 
+class UMaterialInstanceConstant;
 class UMaterialInterface;
 class FNiagaraEmitterInstance;
 class SWidget;
@@ -25,6 +26,7 @@ public:
 	//UObject Interface
 	virtual void PostLoad() override;
 	virtual void PostInitProperties() override;
+	virtual void Serialize(FArchive& Ar) override;
 #if WITH_EDITORONLY_DATA
 	virtual void PostEditChangeProperty(struct FPropertyChangedEvent& PropertyChangedEvent) override;
 #endif// WITH_EDITORONLY_DATA
@@ -41,14 +43,20 @@ public:
 	virtual const TArray<FNiagaraVariable>& GetOptionalAttributes() override;
 	virtual void GetRendererWidgets(const FNiagaraEmitterInstance* InEmitter, TArray<TSharedPtr<SWidget>>& OutWidgets, TSharedPtr<FAssetThumbnailPool> InThumbnailPool) const override;
 	virtual void GetRendererTooltipWidgets(const FNiagaraEmitterInstance* InEmitter, TArray<TSharedPtr<SWidget>>& OutWidgets, TSharedPtr<FAssetThumbnailPool> InThumbnailPool) const override;
-	virtual void GetRendererFeedback(const FVersionedNiagaraEmitter& InEmitter, TArray<FText>& OutErrors, TArray<FText>& OutWarnings, TArray<FText>& OutInfo) const override;
+	virtual void GetRendererFeedback(const FVersionedNiagaraEmitter& InEmitter, TArray<FNiagaraRendererFeedback>& OutErrors, TArray<FNiagaraRendererFeedback>& OutWarnings, TArray<FNiagaraRendererFeedback>& OutInfo) const override;
+	TArray<FNiagaraVariable> GetBoundAttributes() const override;
+	virtual void RenameVariable(const FNiagaraVariableBase& OldVariable, const FNiagaraVariableBase& NewVariable, const FVersionedNiagaraEmitter& InEmitter) override;
+	virtual void RemoveVariable(const FNiagaraVariableBase& OldVariable, const FVersionedNiagaraEmitter& InEmitter) override;
 #endif // WITH_EDITORONLY_DATA
 	virtual void CacheFromCompiledData(const FNiagaraDataSetCompiledData* CompiledData) override;
 	virtual void UpdateSourceModeDerivates(ENiagaraRendererSourceDataMode InSourceMode, bool bFromPropertyEdit = false) override;
 	virtual ENiagaraRendererSourceDataMode GetCurrentSourceMode() const override { return SourceMode; }
 	virtual bool PopulateRequiredBindings(FNiagaraParameterStore& InParameterStore) override;
+	virtual bool NeedsMIDsForMaterials() const { return MaterialParameters.HasAnyBindings(); }
 	//UNiagaraRendererProperties Interface END
 
+	void UpdateMICs();
+
 	UMaterialInterface* GetMaterial(const FNiagaraEmitterInstance* InEmitter) const;
 
 	static const FQuat4f GetDefaultOrientation() { return FRotator3f(-90.0f, 0.0f, 90.0f).Quaternion(); }
@@ -60,6 +68,11 @@ public:
 	UPROPERTY(EditAnywhere, Category = "Decal Rendering")
 	TObjectPtr<UMaterialInterface> Material;
 
+#if WITH_EDITORONLY_DATA
+	UPROPERTY(transient)
+	TObjectPtr<UMaterialInstanceConstant> MICMaterial;
+#endif
+
 	/** Binding to material. */
 	UPROPERTY(EditAnywhere, Category = "Decal Rendering")
 	FNiagaraParameterBinding MaterialParameterBinding;
@@ -104,6 +117,10 @@ public:
 	UPROPERTY(EditAnywhere, Category = "Bindings")
 	FNiagaraVariableAttributeBinding RendererVisibilityTagBinding;
 
+	/** If this array has entries, we will create a MaterialInstanceDynamic per Emitter instance from Material and set the Material parameters using the Niagara simulation variables listed.*/
+	UPROPERTY(EditAnywhere, Category = "Bindings")
+	FNiagaraRendererMaterialParameters MaterialParameters;
+
 	FNiagaraDataSetAccessor<FNiagaraPosition>	PositionDataSetAccessor;
 	FNiagaraDataSetAccessor<FQuat4f>			DecalOrientationDataSetAccessor;
 	FNiagaraDataSetAccessor<FVector3f>			DecalSizeDataSetAccessor;
