Fixes for the graph pin cache.
Disabled cache for assignment nodes for now

diff --git a/Engine/Plugins/FX/Niagara/Source/NiagaraEditor/Private/NiagaraGraphDataCache.cpp b/Engine/Plugins/FX/Niagara/Source/NiagaraEditor/Private/NiagaraGraphDataCache.cpp
--- a/Engine/Plugins/FX/Niagara/Source/NiagaraEditor/Private/NiagaraGraphDataCache.cpp
+++ b/Engine/Plugins/FX/Niagara/Source/NiagaraEditor/Private/NiagaraGraphDataCache.cpp
@@ -3,6 +3,7 @@
 #include "NiagaraGraphDataCache.h"
 
 #include "NiagaraNodeFunctionCall.h"
+#include "NiagaraNodeAssignment.h"
 
 static int32 GNiagaraGraphDataCacheSize = 16384;
 static FAutoConsoleVariableRef CVarNiagaraGraphDataCacheSize(
@@ -51,16 +52,25 @@ void FNiagaraGraphDataCache::GetStackFunctionInputPinsInternal(
 		return;
 	}
 
-	FStackFunctionInputPinKey CacheKey(FunctionCallNode, CalledGraph, StaticVars, ConstantResolver, Options, bIgnoreDisabled, bFilterForCompilation);
-
-	bool FoundInCache = false;
+	bool bUseCache = true;
+	if(FunctionCallNode.IsA<UNiagaraNodeAssignment>())
+	{
+		bUseCache = false;
+	}
 
-	if (const FStackFunctionInputPinValue* CachedValue = StackFunctionInputPinCache.FindAndTouch(CacheKey))
+	FStackFunctionInputPinKey CacheKey(FunctionCallNode, CalledGraph, StaticVars, ConstantResolver, Options, bIgnoreDisabled, bFilterForCompilation);
+	
+	bool bFoundInCache = false;
+	if(bUseCache)
 	{
-		FoundInCache = true;
-		OutInputPins = CachedValue->InputPins;
+		if (const FStackFunctionInputPinValue* CachedValue = StackFunctionInputPinCache.FindAndTouch(CacheKey))
+		{
+			bFoundInCache = true;
+			OutInputPins = CachedValue->InputPins;
+		}
 	}
-	else
+
+	if(bUseCache == false || bFoundInCache == false)
 	{
 		FStackFunctionInputPinValue NewValue;
 		FNiagaraStackGraphUtilities::GetStackFunctionInputPinsWithoutCache(
@@ -74,8 +84,11 @@ void FNiagaraGraphDataCache::GetStackFunctionInputPinsInternal(
 
 		OutInputPins = NewValue.InputPins;
 
-		StackFunctionInputPinCache.Add(CacheKey, NewValue);
-	}
+		if(bUseCache)
+		{
+			StackFunctionInputPinCache.Add(CacheKey, NewValue);
+		}
+	}	
 
 	if (GNiagaraGraphDataCacheValidation)
 	{
