- Move DI SimCache functionality into an interface to reduce methods in base data interface class
- Fix a crash when baking simulations with VDB enabled from the baker

diff --git a/Engine/Plugins/FX/Niagara/Source/Niagara/Classes/NiagaraDataInterface.h b/Engine/Plugins/FX/Niagara/Source/Niagara/Classes/NiagaraDataInterface.h
--- a/Engine/Plugins/FX/Niagara/Source/Niagara/Classes/NiagaraDataInterface.h
+++ b/Engine/Plugins/FX/Niagara/Source/Niagara/Classes/NiagaraDataInterface.h
@@ -489,27 +489,6 @@ public:
 	virtual bool PerInstanceTick(void* PerInstanceData, FNiagaraSystemInstance* SystemInstance, float DeltaSeconds) { return false; }
 	virtual bool PerInstanceTickPostSimulate(void* PerInstanceData, FNiagaraSystemInstance* SystemInstance, float DeltaSeconds) { return false; }
 	
-	/** Begin writing data for a simulation cache, returning a nullptr means the data interface does not store data into the simulation cache. */
-	virtual UObject* SimCacheBeginWrite(UObject* SimCache, FNiagaraSystemInstance* NiagaraSystemInstance, const void* OptionalPerInstanceData) const { return nullptr; }
-	/** Write a new frame of data for the simulation cache.  This is always in sequence, i.e. 0, 1, 2, etc, we will never jump around frames. */
-	virtual bool SimCacheWriteFrame(UObject* StorageObject, int FrameIndex, FNiagaraSystemInstance* SystemInstance, const void* OptionalPerInstanceData) const { return true; }
-	/** End writing data for a simulation cache.  Note this is called on the CDO not the instance the object was created from. */
-	virtual bool SimCacheEndWrite(UObject* StorageObject) const { return true; }
-	/** Read a frame of data from the simulation cache. */
-	virtual bool SimCacheReadFrame(UObject* StorageObject, int FrameA, int FrameB, float Interp, FNiagaraSystemInstance* SystemInstance, void* OptionalPerInstanceData) { return true; }
-	/**
-	Called when the simulation cache has finished reading a frame.
-	Only DataInterfaces with PerInstanceData are currently supported.
-	*/
-	virtual void SimCachePostReadFrame(void* OptionalPerInstanceData, FNiagaraSystemInstance* SystemInstance) {}
-	/**
-	This function allows you to preserve a list of attributes when building a renderer only cache.
-	The UsageContext will be either a UNiagaraSystem or a UNiagaraEmitter and can be used to scope your variables accordingly.
-	For example, if you were to require 'Particles.MyAttribute' in order to process the cache results you would need to convert
-	this into 'MyEmitter.Particles.MyAttribute' by checking the UsageContext is a UNiagaraEmitter and then creating the variable from the unique name.
-	*/
-	virtual TArray<FNiagaraVariableBase> GetSimCacheRendererAttributes(UObject* UsageContext) const { return TArray<FNiagaraVariableBase>(); }
-
 #if WITH_EDITORONLY_DATA
 	/** Allows the generic class defaults version of this class to specify any dependencies/version/etc that might invalidate the compile. It should never depend on the value of specific properties.*/
 	virtual bool AppendCompileHash(FNiagaraCompileHashVisitor* InVisitor) const;
diff --git a/Engine/Plugins/FX/Niagara/Source/Niagara/Classes/NiagaraDataInterfaceRenderTargetVolume.h b/Engine/Plugins/FX/Niagara/Source/Niagara/Classes/NiagaraDataInterfaceRenderTargetVolume.h
--- a/Engine/Plugins/FX/Niagara/Source/Niagara/Classes/NiagaraDataInterfaceRenderTargetVolume.h
+++ b/Engine/Plugins/FX/Niagara/Source/Niagara/Classes/NiagaraDataInterfaceRenderTargetVolume.h
@@ -6,6 +6,7 @@
 
 #include "NiagaraDataInterfaceRW.h"
 #include "NiagaraComponent.h"
+#include "NiagaraSimCacheCustomStorageInterface.h"
 #include "NiagaraDataInterfaceRenderTargetVolume.generated.h"
 
 class FNiagaraSystemInstance;
@@ -74,7 +75,7 @@ struct FNiagaraDataInterfaceProxyRenderTargetVolumeProxy : public FNiagaraDataIn
 };
 
 UCLASS(EditInlineNew, Category = "Rendering", meta = (DisplayName = "Render Target Volume"), Blueprintable, BlueprintType)
-class NIAGARA_API UNiagaraDataInterfaceRenderTargetVolume : public UNiagaraDataInterfaceRWBase
+class NIAGARA_API UNiagaraDataInterfaceRenderTargetVolume : public UNiagaraDataInterfaceRWBase, public INiagaraSimCacheCustomStorageInterface
 {
 	GENERATED_UCLASS_BODY()
 
@@ -121,13 +122,15 @@ public:
 	virtual bool CanExposeVariables() const override { return true;}
 	virtual void GetExposedVariables(TArray<FNiagaraVariableBase>& OutVariables) const override;
 	virtual bool GetExposedVariableValue(const FNiagaraVariableBase& InVariable, void* InPerInstanceData, FNiagaraSystemInstance* InSystemInstance, void* OutData) const override;
+	//~ UNiagaraDataInterface interface END
 
+	//~ INiagaraSimCacheCustomStorageInterface interface BEGIN
 	virtual UObject* SimCacheBeginWrite(UObject* SimCache, FNiagaraSystemInstance* NiagaraSystemInstance, const void* OptionalPerInstanceData) const override;
 	virtual bool SimCacheWriteFrame(UObject* StorageObject, int FrameIndex, FNiagaraSystemInstance* SystemInstance, const void* OptionalPerInstanceData) const override;
 	virtual bool SimCacheEndWrite(UObject* StorageObject) const override;
 	virtual bool SimCacheReadFrame(UObject* StorageObject, int FrameA, int FrameB, float Interp, FNiagaraSystemInstance* SystemInstance, void* OptionalPerInstanceData) override;
 	//~ UNiagaraDataInterface interface END
-	
+
 	void VMGetSize(FVectorVMExternalFunctionContext& Context);
 	void VMSetSize(FVectorVMExternalFunctionContext& Context);
 	void VMGetNumMipLevels(FVectorVMExternalFunctionContext& Context);
diff --git a/Engine/Plugins/FX/Niagara/Source/Niagara/Classes/NiagaraSimCacheCustomStorageInterface.h b/Engine/Plugins/FX/Niagara/Source/Niagara/Classes/NiagaraSimCacheCustomStorageInterface.h
new file mode 100644
--- /dev/null
+++ b/Engine/Plugins/FX/Niagara/Source/Niagara/Classes/NiagaraSimCacheCustomStorageInterface.h
@@ -0,0 +1,63 @@
+// Copyright Epic Games, Inc. All Rights Reserved.
+
+#pragma once
+
+#include "NiagaraCommon.h"
+#include "NiagaraSimCacheCustomStorageInterface.generated.h"
+
+class FNiagaraSystemInstance;
+
+// Interface for UObjects to implement renderable mesh
+UINTERFACE()
+class UNiagaraSimCacheCustomStorageInterface : public UInterface
+{
+	GENERATED_BODY()
+};
+
+/**
+The current API for storing data inside a simulation cache.
+This is highly experimental and the API will change as we split editor / runtime data storage.
+*/
+class INiagaraSimCacheCustomStorageInterface
+{
+	GENERATED_BODY()
+
+public:
+	/**
+	Called when we begin to write data into a simulation cache.
+	Returning nullptr means you are not going to cache any data for the simulation.
+	The object returned will be stored directly into the cache file, so you are expected to manage the size of the object and store data appropriately.
+	*/
+	virtual UObject* SimCacheBeginWrite(UObject* SimCache, FNiagaraSystemInstance* NiagaraSystemInstance, const void* OptionalPerInstanceData) const { return nullptr; }
+
+	/**
+	Called when we are ready to write data into the simulation cache.
+	This is always called in sequence, i.e. 0, 1, 2, etc, we will never jump around frames.
+	*/
+	virtual bool SimCacheWriteFrame(UObject* StorageObject, int FrameIndex, FNiagaraSystemInstance* SystemInstance, const void* OptionalPerInstanceData) const { return true; }
+
+	/**
+	Called when we complete writing data into the simulation cache.
+	Note: This is called using the Class Default Object not the instance the object was created from
+	*/
+	virtual bool SimCacheEndWrite(UObject* StorageObject) const { return true; }
+
+	/**
+	Read a frame of data from the simulation cache.
+	*/
+	virtual bool SimCacheReadFrame(UObject* StorageObject, int FrameA, int FrameB, float Interp, FNiagaraSystemInstance* SystemInstance, void* OptionalPerInstanceData) { return true; }
+
+	/**
+	Called when the simulation cache has finished reading a frame.
+	Only DataInterfaces with PerInstanceData will have this method called on them.
+	*/
+	virtual void SimCachePostReadFrame(void* OptionalPerInstanceData, FNiagaraSystemInstance* SystemInstance) {}
+
+	/**
+	This function allows you to preserve a list of attributes when building a renderer only cache.
+	The UsageContext will be either a UNiagaraSystem or a UNiagaraEmitter and can be used to scope your variables accordingly.
+	For example, if you were to require 'Particles.MyAttribute' in order to process the cache results you would need to convert
+	this into 'MyEmitter.Particles.MyAttribute' by checking the UsageContext is a UNiagaraEmitter and then creating the variable from the unique name.
+	*/
+	virtual TArray<FNiagaraVariableBase> GetSimCacheRendererAttributes(UObject* UsageContext) const { return TArray<FNiagaraVariableBase>(); }
+};
diff --git a/Engine/Plugins/FX/Niagara/Source/Niagara/Private/NiagaraDataInterfaceRenderTargetVolume.cpp b/Engine/Plugins/FX/Niagara/Source/Niagara/Private/NiagaraDataInterfaceRenderTargetVolume.cpp
--- a/Engine/Plugins/FX/Niagara/Source/Niagara/Private/NiagaraDataInterfaceRenderTargetVolume.cpp
+++ b/Engine/Plugins/FX/Niagara/Source/Niagara/Private/NiagaraDataInterfaceRenderTargetVolume.cpp
@@ -581,10 +581,12 @@ UObject* UNiagaraDataInterfaceRenderTargetVolume::SimCacheBeginWrite(UObject* Si
 
 		FString ActorName = NiagaraSystemInstance->GetAttachComponent()->GetName();
 
-		// if we are caching from the editor, use the label of the actor to name the cache files
 #if WITH_EDITOR
-		AActor* Owner = NiagaraSystemInstance->GetAttachComponent()->GetOwner();
-		ActorName = Owner->GetActorLabel();
+		// if we are caching from the editor, use the label of the actor to name the cache files
+		if (AActor* Owner = NiagaraSystemInstance->GetAttachComponent()->GetOwner())
+		{
+			ActorName = Owner->GetActorLabel();
+		}
 #endif
 
 		FString DIName = Proxy->SourceDIName.ToString();
diff --git a/Engine/Plugins/FX/Niagara/Source/Niagara/Private/NiagaraSimCache.cpp b/Engine/Plugins/FX/Niagara/Source/Niagara/Private/NiagaraSimCache.cpp
--- a/Engine/Plugins/FX/Niagara/Source/Niagara/Private/NiagaraSimCache.cpp
+++ b/Engine/Plugins/FX/Niagara/Source/Niagara/Private/NiagaraSimCache.cpp
@@ -12,6 +12,7 @@
 #include "NiagaraGPUInstanceCountManager.h"
 #include "NiagaraGpuComputeDispatchInterface.h"
 #include "NiagaraSimCacheAttributeReaderHelper.h"
+#include "NiagaraSimCacheCustomStorageInterface.h"
 #include "NiagaraSimCacheHelper.h"
 #include "NiagaraSystemImpl.h"
 #include "NiagaraSystemInstance.h"
@@ -112,9 +113,12 @@ bool UNiagaraSimCache::BeginWrite(FNiagaraSimCacheCreateParameters InCreateParam
 			Helper.SystemInstance,
 			[&](const FNiagaraDataInterfaceUtilities::FDataInterfaceUsageContext& UsageContext)
 			{
-				for ( const FNiagaraVariableBase& PreserveAttribute : UsageContext.DataInterface->GetSimCacheRendererAttributes(UsageContext.OwnerObject) )
+				if (INiagaraSimCacheCustomStorageInterface* SimCacheCustomStorageInterface = Cast<INiagaraSimCacheCustomStorageInterface>(UsageContext.DataInterface))
 				{
-					CreateParameters.ExplicitCaptureAttributes.AddUnique(PreserveAttribute.GetName());
+					for (const FNiagaraVariableBase& PreserveAttribute : SimCacheCustomStorageInterface->GetSimCacheRendererAttributes(UsageContext.OwnerObject))
+					{
+						CreateParameters.ExplicitCaptureAttributes.AddUnique(PreserveAttribute.GetName());
+					}
 				}
 				return true;
 			}
@@ -138,10 +142,16 @@ bool UNiagaraSimCache::BeginWrite(FNiagaraSimCacheCreateParameters InCreateParam
 			Helper.SystemInstance,
 			[&](const FNiagaraVariableBase& Variable, UNiagaraDataInterface* DataInterface)
 			{
-				if ( (CreateParameters.ExplicitCaptureAttributes.Num() == 0) || CreateParameters.ExplicitCaptureAttributes.Contains(Variable.GetName()) )
+				// Are we capturing this data interface?
+				if ((CreateParameters.ExplicitCaptureAttributes.Num() != 0) && !CreateParameters.ExplicitCaptureAttributes.Contains(Variable.GetName()))
+				{
+					return true;
+				}
+
+				if (INiagaraSimCacheCustomStorageInterface* SimCacheCustomStorageInterface = Cast<INiagaraSimCacheCustomStorageInterface>(DataInterface))
 				{
 					const void* PerInstanceData = Helper.SystemInstance->FindDataInterfaceInstanceData(DataInterface);
-					if (UObject* DICacheStorage = DataInterface->SimCacheBeginWrite(this, Helper.SystemInstance, PerInstanceData))
+					if (UObject* DICacheStorage = SimCacheCustomStorageInterface->SimCacheBeginWrite(this, Helper.SystemInstance, PerInstanceData))
 					{
 						DataInterfaceStorage.FindOrAdd(Variable) = DICacheStorage;
 					}
@@ -267,8 +277,9 @@ bool UNiagaraSimCache::WriteFrame(UNiagaraComponent* NiagaraComponent)
 			{
 				if ( UObject* StorageObject = DataInterfaceStorage.FindRef(Variable) )
 				{
+					INiagaraSimCacheCustomStorageInterface* SimCacheCustomStorageInterface = CastChecked<INiagaraSimCacheCustomStorageInterface>(DataInterface);
 					const void* PerInstanceData = Helper.SystemInstance->FindDataInterfaceInstanceData(DataInterface);
-					bDataInterfacesSucess &= DataInterface->SimCacheWriteFrame(StorageObject, FrameIndex, Helper.SystemInstance, PerInstanceData);
+					bDataInterfacesSucess &= SimCacheCustomStorageInterface->SimCacheWriteFrame(StorageObject, FrameIndex, Helper.SystemInstance, PerInstanceData);
 				}
 				return true;
 			}
@@ -299,8 +310,8 @@ bool UNiagaraSimCache::EndWrite()
 		{
 			UClass* DataInterfaceClass = it.Key().GetType().GetClass();
 			check(DataInterfaceClass != nullptr);
-			UNiagaraDataInterface* DataInterface = CastChecked<UNiagaraDataInterface>(DataInterfaceClass->GetDefaultObject());
-			bDataInterfacesSucess &= DataInterface->SimCacheEndWrite(it.Value());
+			INiagaraSimCacheCustomStorageInterface* SimCacheCustomStorageInterface = CastChecked<INiagaraSimCacheCustomStorageInterface>(DataInterfaceClass->GetDefaultObject());
+			bDataInterfacesSucess &= SimCacheCustomStorageInterface->SimCacheEndWrite(it.Value());
 		}
 
 		if (bDataInterfacesSucess == false)
@@ -458,8 +469,11 @@ bool UNiagaraSimCache::ReadFrame(int32 FrameIndex, float FrameFraction, FNiagara
 			{
 				if (UObject* StorageObject = DataInterfaceStorage.FindRef(Variable))
 				{
-					void* PerInstanceData = Helper.SystemInstance->FindDataInterfaceInstanceData(DataInterface);
-					bDataInterfacesSucess &= DataInterface->SimCacheReadFrame(StorageObject, FrameIndex, NextFrameIndex, FrameFraction, Helper.SystemInstance, PerInstanceData);
+					if (INiagaraSimCacheCustomStorageInterface* SimCacheCustomStorageInterface = Cast<INiagaraSimCacheCustomStorageInterface>(DataInterface))
+					{
+						void* PerInstanceData = Helper.SystemInstance->FindDataInterfaceInstanceData(DataInterface);
+						bDataInterfacesSucess &= SimCacheCustomStorageInterface->SimCacheReadFrame(StorageObject, FrameIndex, NextFrameIndex, FrameFraction, Helper.SystemInstance, PerInstanceData);
+					}
 				}
 				return true;
 			}
@@ -475,9 +489,9 @@ bool UNiagaraSimCache::ReadFrame(int32 FrameIndex, float FrameFraction, FNiagara
 	//-TODO: This should loop over all DataInterfaces that register not just ones with instance data
 	for (TPair<TWeakObjectPtr<UNiagaraDataInterface>, int32>& DataInterfacePair : SystemInstance->DataInterfaceInstanceDataOffsets)
 	{
-		if (UNiagaraDataInterface* Interface = DataInterfacePair.Key.Get())
+		if (INiagaraSimCacheCustomStorageInterface* SimCacheCustomStorageInterface = Cast<INiagaraSimCacheCustomStorageInterface>(DataInterfacePair.Key.Get()))
 		{
-			Interface->SimCachePostReadFrame(&SystemInstance->DataInterfaceInstanceData[DataInterfacePair.Value], SystemInstance);
+			SimCacheCustomStorageInterface->SimCachePostReadFrame(&SystemInstance->DataInterfaceInstanceData[DataInterfacePair.Value], SystemInstance);
 		}
 	}
 	return true;
diff --git a/Engine/Plugins/Runtime/HairStrands/Source/HairStrandsCore/Public/Niagara/NiagaraDataInterfaceHairStrands.h b/Engine/Plugins/Runtime/HairStrands/Source/HairStrandsCore/Public/Niagara/NiagaraDataInterfaceHairStrands.h
--- a/Engine/Plugins/Runtime/HairStrands/Source/HairStrandsCore/Public/Niagara/NiagaraDataInterfaceHairStrands.h
+++ b/Engine/Plugins/Runtime/HairStrands/Source/HairStrandsCore/Public/Niagara/NiagaraDataInterfaceHairStrands.h
@@ -6,6 +6,7 @@
 #include "NiagaraDataInterfaceRW.h"
 #include "NiagaraCommon.h"
 #include "NiagaraRenderGraphUtils.h"
+#include "NiagaraSimCacheCustomStorageInterface.h"
 #include "VectorVM.h"
 #include "GroomAsset.h"
 #include "GroomActor.h"
@@ -378,7 +379,7 @@ struct FNDIHairStrandsData
 
 /** Data Interface for the strand base */
 UCLASS(EditInlineNew, Category = "Strands", meta = (DisplayName = "Hair Strands"))
-class HAIRSTRANDSCORE_API UNiagaraDataInterfaceHairStrands : public UNiagaraDataInterface
+class HAIRSTRANDSCORE_API UNiagaraDataInterfaceHairStrands : public UNiagaraDataInterface, public INiagaraSimCacheCustomStorageInterface
 {
 	GENERATED_UCLASS_BODY()
 
@@ -397,7 +398,7 @@ public:
 	/** UObject Interface */
 	virtual void PostInitProperties() override;
 
-	/** UNiagaraDataInterface Interface */
+	/** Begin UNiagaraDataInterface Interface */
 	virtual void GetFunctions(TArray<FNiagaraFunctionSignature>& OutFunctions) override;
 	virtual void GetVMExternalFunction(const FVMExternalFunctionBindingInfo& BindingInfo, void* InstanceData, FVMExternalFunction &OutFunc) override;
 	virtual bool CanExecuteOnTarget(ENiagaraSimTarget Target) const override { return Target == ENiagaraSimTarget::GPUComputeSim; }
@@ -409,8 +410,6 @@ public:
 	virtual bool HasPreSimulateTick() const override { return true; }
 	virtual bool HasTickGroupPrereqs() const override { return true; }
 	virtual ETickingGroup CalculateTickGroup(const void* PerInstanceData) const override;
-	virtual void SimCachePostReadFrame(void* OptionalPerInstanceData, FNiagaraSystemInstance* SystemInstance) override;
-	virtual TArray<FNiagaraVariableBase> GetSimCacheRendererAttributes(UObject* UsageContext) const override;
 
 	/** GPU simulation  functionality */
 #if WITH_EDITORONLY_DATA
@@ -423,6 +422,12 @@ public:
 	virtual void SetShaderParameters(const FNiagaraDataInterfaceSetShaderParametersContext& Context) const override;
 
 	virtual void ProvidePerInstanceDataForRenderThread(void* DataForRenderThread, void* PerInstanceData, const FNiagaraSystemInstanceID& SystemInstance) override;
+	/** End UNiagaraDataInterface Interface */
+
+	/** Begin INiagaraSimCacheCustomStorageInterface Interface */
+	virtual void SimCachePostReadFrame(void* OptionalPerInstanceData, FNiagaraSystemInstance* SystemInstance) override;
+	virtual TArray<FNiagaraVariableBase> GetSimCacheRendererAttributes(UObject* UsageContext) const override;
+	/** End INiagaraSimCacheCustomStorageInterface Interface */
 
 	/** Update the source component */
 	void ExtractSourceComponent(FNiagaraSystemInstance* SystemInstance);
