project setting for aux files outputted from the Niagara Sim Cache (ie: vdb)

diff --git a/Engine/Plugins/FX/Niagara/Source/Niagara/Private/NiagaraDataInterfaceRenderTargetVolume.cpp b/Engine/Plugins/FX/Niagara/Source/Niagara/Private/NiagaraDataInterfaceRenderTargetVolume.cpp
--- a/Engine/Plugins/FX/Niagara/Source/Niagara/Private/NiagaraDataInterfaceRenderTargetVolume.cpp
+++ b/Engine/Plugins/FX/Niagara/Source/Niagara/Private/NiagaraDataInterfaceRenderTargetVolume.cpp
@@ -20,6 +20,7 @@
 #include "VolumeCache.h"
 #include "RHIStaticStates.h"
 #include "Misc/PathViews.h"
+#include "Misc/Paths.h"
 #include "DataDrivenShaderPlatformInfo.h"
 
 #include UE_INLINE_GENERATED_CPP_BY_NAME(NiagaraDataInterfaceRenderTargetVolume)
@@ -88,14 +89,6 @@ namespace NDIRenderTargetVolumeLocal
 		ECVF_Default
 	);
 
-	FString GSimCacheOpenVDBBasePath = "";
-	static FAutoConsoleVariableRef CVarSimCacheOpenVDBBasePath(
-		TEXT("fx.Niagara.RenderTargetVolume.SimCacheOpenVDBBasePath"),
-		GSimCacheOpenVDBBasePath,
-		TEXT("Base path for OpenVDB output."),
-		ECVF_Default
-	);
-
 	static const FName GetSimCacheCompressionType() { return GSimCacheCompressed ? NAME_Oodle : NAME_None; }
 }
 
@@ -573,9 +566,9 @@ UObject* UNiagaraDataInterfaceRenderTargetVolume::SimCacheBeginWrite(UObject* Si
 
 		FString SystemInstanceName = NiagaraSystemInstance->GetSystem()->GetName();
 
-		if (NDIRenderTargetVolumeLocal::GSimCacheOpenVDBBasePath == "")
+		if (GetDefault<UNiagaraSettings>()->SimCacheAuxiliaryFileBasePath == "")
 		{
-			UE_LOG(LogNiagara, Error, TEXT("You must set fx.Niagara.RenderTargetVolume.SimCacheOpenVDBPath"));
+			UE_LOG(LogNiagara, Error, TEXT("You must set SimCacheAuxiliaryFileBasePath in project settings"));
 			return nullptr;
 		}
 
@@ -590,8 +583,12 @@ UObject* UNiagaraDataInterfaceRenderTargetVolume::SimCacheBeginWrite(UObject* Si
 #endif
 
 		FString DIName = Proxy->SourceDIName.ToString();
-		FString FullFilePathSpec = NDIRenderTargetVolumeLocal::GSimCacheOpenVDBBasePath + "/NiagaraSimCache/" + ActorName + "/" + DIName + "_SimCache.{FrameIndex}.vdb";
+		FString FullFilePathSpec = GetDefault<UNiagaraSettings>()->SimCacheAuxiliaryFileBasePath + "/" + ActorName + "/" + DIName + "_SimCache.{FrameIndex}.vdb";
 		FullFilePathSpec.ReplaceInline(TEXT("//"), TEXT("/"));		
+		
+		FullFilePathSpec.ReplaceInline(TEXT("{project_dir}"), *FPaths::ProjectDir());
+
+		FullFilePathSpec = FPaths::ConvertRelativePathToFull(FullFilePathSpec);
 
 		// Create output directory		
 		const FString FileDirectory = FString(FPathViews::GetPath(FullFilePathSpec));
diff --git a/Engine/Plugins/FX/Niagara/Source/Niagara/Private/NiagaraSettings.cpp b/Engine/Plugins/FX/Niagara/Source/Niagara/Private/NiagaraSettings.cpp
--- a/Engine/Plugins/FX/Niagara/Source/Niagara/Private/NiagaraSettings.cpp
+++ b/Engine/Plugins/FX/Niagara/Source/Niagara/Private/NiagaraSettings.cpp
@@ -12,6 +12,8 @@ UNiagaraSettings::UNiagaraSettings(const FObjectInitializer& ObjectInitlaizer)
 
 	NDICollisionQuery_AsyncGpuTraceProviderOrder.Add(ENDICollisionQuery_AsyncGpuTraceProvider::Type::HWRT);
 	NDICollisionQuery_AsyncGpuTraceProviderOrder.Add(ENDICollisionQuery_AsyncGpuTraceProvider::Type::GSDF);
+
+	SimCacheAuxiliaryFileBasePath = "{project_dir}/Saved/NiagaraSimCache";
 }
 
 FName UNiagaraSettings::GetCategoryName() const
diff --git a/Engine/Plugins/FX/Niagara/Source/Niagara/Public/NiagaraSettings.h b/Engine/Plugins/FX/Niagara/Source/Niagara/Public/NiagaraSettings.h
--- a/Engine/Plugins/FX/Niagara/Source/Niagara/Public/NiagaraSettings.h
+++ b/Engine/Plugins/FX/Niagara/Source/Niagara/Public/NiagaraSettings.h
@@ -208,6 +208,12 @@ class NIAGARA_API UNiagaraSettings : public UDeveloperSettings
 	UPROPERTY(config, EditAnywhere, Category = AsyncGpuTraceDI, meta = (DisplayName = "Trace Provider Priorities (Experimental)", ConfigRestartRequired = true))
 	TArray<TEnumAsByte<ENDICollisionQuery_AsyncGpuTraceProvider::Type>> NDICollisionQuery_AsyncGpuTraceProviderOrder;
 
+	/**
+	Base path for auxiliary files written out during the generation of a Niagara Sim Cache (ie: volume files).
+	*/
+	UPROPERTY(config, EditAnywhere, Category = SimCache, meta = (DisplayName = "Sim Cache Auxiliary File Base Path", ConfigRestartRequired = false))
+	FString SimCacheAuxiliaryFileBasePath;
+
 	UPROPERTY(config, EditAnywhere, Category = Scalability)
 	TArray<FNiagaraPlatformSetRedirect> PlatformSetRedirects;
 
