Fix Niagara assets destroying thumbnail data when loaded from cooked package.
A texture in a cooked package will not contain the raw texture data so UpdateResource should not be called on them to rebuild texture.

diff --git a/Engine/Plugins/FX/Niagara/Source/NiagaraEditor/Private/NiagaraThumbnailRenderer.cpp b/Engine/Plugins/FX/Niagara/Source/NiagaraEditor/Private/NiagaraThumbnailRenderer.cpp
--- a/Engine/Plugins/FX/Niagara/Source/NiagaraEditor/Private/NiagaraThumbnailRenderer.cpp
+++ b/Engine/Plugins/FX/Niagara/Source/NiagaraEditor/Private/NiagaraThumbnailRenderer.cpp
@@ -5,6 +5,7 @@
 #include "Engine/Texture2D.h"
 #include "NiagaraEmitter.h"
 #include "NiagaraSystem.h"
+#include "UObject/Package.h"
 
 #include UE_INLINE_GENERATED_CPP_BY_NAME(NiagaraThumbnailRenderer)
 
@@ -28,8 +29,12 @@ UTexture2D* UNiagaraEmitterThumbnailRenderer::GetThumbnailTextureFromObject(UObj
 	UNiagaraEmitter* Emitter = Cast<UNiagaraEmitter>(Object);
 	if (Emitter && Emitter->ThumbnailImage)
 	{
-		Emitter->ThumbnailImage->FinishCachePlatformData();
-		Emitter->ThumbnailImage->UpdateResource();
+		// Avoid calling UpdateResource on cooked texture as doing so will destroy the texture's data
+		if (!Emitter->GetPackage()->HasAnyPackageFlags(PKG_Cooked | PKG_FilterEditorOnly))
+		{
+			Emitter->ThumbnailImage->FinishCachePlatformData();
+			Emitter->ThumbnailImage->UpdateResource();
+		}
 		return Emitter->ThumbnailImage;
 	}
 	return nullptr;
@@ -40,8 +45,12 @@ UTexture2D* UNiagaraSystemThumbnailRenderer::GetThumbnailTextureFromObject(UObje
 	UNiagaraSystem* System = Cast<UNiagaraSystem>(Object);
 	if (System && System->ThumbnailImage)
 	{
-		System->ThumbnailImage->FinishCachePlatformData();
-		System->ThumbnailImage->UpdateResource();
+		// Avoid calling UpdateResource on cooked texture as doing so will destroy the texture's data
+		if (!System->GetPackage()->HasAnyPackageFlags(PKG_Cooked | PKG_FilterEditorOnly))
+		{
+			System->ThumbnailImage->FinishCachePlatformData();
+			System->ThumbnailImage->UpdateResource();
+		}
 		return System->ThumbnailImage;
 	}
 	return nullptr;
