- Fix for GetNumCells not returning the instance data cell count on rasterization grid

#rnx

diff --git a/Engine/Plugins/FX/Niagara/Source/Niagara/Private/NiagaraDataInterfaceRasterizationGrid3D.cpp b/Engine/Plugins/FX/Niagara/Source/Niagara/Private/NiagaraDataInterfaceRasterizationGrid3D.cpp
--- a/Engine/Plugins/FX/Niagara/Source/Niagara/Private/NiagaraDataInterfaceRasterizationGrid3D.cpp
+++ b/Engine/Plugins/FX/Niagara/Source/Niagara/Private/NiagaraDataInterfaceRasterizationGrid3D.cpp
@@ -480,9 +480,9 @@ void UNiagaraDataInterfaceRasterizationGrid3D::VMGetNumCells(FVectorVMExternalFu
 
 	for (int32 InstanceIdx = 0; InstanceIdx < Context.GetNumInstances(); ++InstanceIdx)
 	{
-		NumCellsX.SetAndAdvance(NumCells.X);
-		NumCellsY.SetAndAdvance(NumCells.Y);
-		NumCellsZ.SetAndAdvance(NumCells.Z);
+		NumCellsX.SetAndAdvance(InstData->NumCells.X);
+		NumCellsY.SetAndAdvance(InstData->NumCells.Y);
+		NumCellsZ.SetAndAdvance(InstData->NumCells.Z);
 	}
 }
 
