- Remove duplicate get num cells functions

#rnx

diff --git a/Engine/Plugins/FX/Niagara/Source/Niagara/Private/NiagaraDataInterfaceRW.cpp b/Engine/Plugins/FX/Niagara/Source/Niagara/Private/NiagaraDataInterfaceRW.cpp
--- a/Engine/Plugins/FX/Niagara/Source/Niagara/Private/NiagaraDataInterfaceRW.cpp
+++ b/Engine/Plugins/FX/Niagara/Source/Niagara/Private/NiagaraDataInterfaceRW.cpp
@@ -207,19 +207,6 @@ void UNiagaraDataInterfaceGrid3D::GetFunctions(TArray<FNiagaraFunctionSignature>
 		OutFunctions.Add(Sig);
 	}
 
-	{
-		FNiagaraFunctionSignature Sig;
-		Sig.Name = UNiagaraDataInterfaceRWBase::NumCellsFunctionName;
-		Sig.Inputs.Add(FNiagaraVariable(FNiagaraTypeDefinition(GetClass()), TEXT("Grid")));		
-		Sig.Outputs.Add(FNiagaraVariable(FNiagaraTypeDefinition::GetIntDef(), TEXT("NumCellsX")));
-		Sig.Outputs.Add(FNiagaraVariable(FNiagaraTypeDefinition::GetIntDef(), TEXT("NumCellsY")));
-		Sig.Outputs.Add(FNiagaraVariable(FNiagaraTypeDefinition::GetIntDef(), TEXT("NumCellsZ")));
-
-		Sig.bMemberFunction = true;
-		Sig.bRequiresContext = false;
-		OutFunctions.Add(Sig);
-	}
-
 	{
 		FNiagaraFunctionSignature Sig;
 		Sig.Name = UNiagaraDataInterfaceRWBase::CellSizeFunctionName;
