- Solo components read playback rate from debugger

diff --git a/Engine/Plugins/FX/Niagara/Source/Niagara/Private/NiagaraComponent.cpp b/Engine/Plugins/FX/Niagara/Source/Niagara/Private/NiagaraComponent.cpp
--- a/Engine/Plugins/FX/Niagara/Source/Niagara/Private/NiagaraComponent.cpp
+++ b/Engine/Plugins/FX/Niagara/Source/Niagara/Private/NiagaraComponent.cpp
@@ -734,6 +734,17 @@ void UNiagaraComponent::TickComponent(float DeltaSeconds, enum ELevelTick TickTy
 		INC_DWORD_STAT_BY(STAT_TotalNiagaraSystemInstances, 1);
 		INC_DWORD_STAT_BY(STAT_TotalNiagaraSystemInstancesSolo, 1);
 
+		// Support for pausing / playback rate from the debugger
+		if (FNiagaraWorldManager* WorldManager = FNiagaraWorldManager::Get(GetWorld()))
+		{
+			if (WorldManager->GetDebugPlaybackMode() == ENiagaraDebugPlaybackMode::Paused)
+			{
+				return;
+			}
+
+			DeltaSeconds *= WorldManager->GetDebugPlaybackRate();
+		}
+
 		// If we have a sim cache attached then use that
 		if ( SimCache != nullptr )
 		{
