- Add guard around DenseGridPtr access as the RT can be reading while the GT could be writing

#rnx

diff --git a/Engine/Plugins/FX/Niagara/Source/Niagara/Private/VolumeCache.cpp b/Engine/Plugins/FX/Niagara/Source/Niagara/Private/VolumeCache.cpp
--- a/Engine/Plugins/FX/Niagara/Source/Niagara/Private/VolumeCache.cpp
+++ b/Engine/Plugins/FX/Niagara/Source/Niagara/Private/VolumeCache.cpp
@@ -13,7 +13,6 @@
 #include "NiagaraOpenVDB.h"
 #endif
 
-
 DEFINE_LOG_CATEGORY_STATIC(LogVolumeCache, Log, All);
 
 
@@ -40,6 +39,8 @@ public:
 	
 private:
 	TMap<int32, Vec4Grid::Ptr> OpenVDBGrids;
+
+	FCriticalSection DenseGridGuard;
 	Vec4Dense::Ptr DenseGridPtr;
 };
 #endif
@@ -140,6 +141,7 @@ FString FVolumeCacheData::GetAssetPath(FString PathFormat, int32 FrameIndex) con
 #if PLATFORM_WINDOWS
 void FOpenVDBCacheData::Init(FIntVector Resolution)
 {
+	FScopeLock ScopeLock(&DenseGridGuard);
 	DenseGridPtr.reset(new Vec4Dense(openvdb::CoordBBox(0, 0, 0, Resolution.X - 1, Resolution.Y - 1, Resolution.Z - 1), Vec4(0.0, 0.0, 0.0, 0.0)));
 }
 
@@ -150,6 +152,7 @@ bool FOpenVDBCacheData::LoadFile(FString Path, int frame)
 	if (OpenVDBGrids.Contains(frame) && OpenVDBGrids[frame] != nullptr)
 	{
 		// if dense vdb buffer doesn't match current resolution, re initialize it				
+		FScopeLock ScopeLock(&DenseGridGuard);
 		if (DenseGridPtr == nullptr || (DenseGridPtr->bbox().dim() != openvdb::Coord(DenseResolution.X - 1, DenseResolution.Y - 1, DenseResolution.Z - 1)))
 		{
 			Init(DenseResolution);
@@ -208,6 +211,7 @@ bool FOpenVDBCacheData::LoadFile(FString Path, int frame)
 			OpenVDBGrids.Add(frame, ColorGrid);
 
 			// if dense vdb buffer doesn't match current resolution, re initialize it				
+			FScopeLock ScopeLock(&DenseGridGuard);
 			if (DenseGridPtr == nullptr || (DenseGridPtr->bbox().dim() != openvdb::Coord(DenseResolution.X - 1, DenseResolution.Y - 1, DenseResolution.Z - 1)))
 			{
 				Init(DenseResolution);
@@ -263,6 +267,7 @@ bool FOpenVDBCacheData::Fill3DTexture_RenderThread(int frame, FTextureRHIRef Tex
 {
 	if (OpenVDBGrids.Contains(frame) && OpenVDBGrids[frame] != nullptr)
 	{
+		FScopeLock ScopeLock(&DenseGridGuard);
 		uint8* DataPtr = (uint8*)DenseGridPtr->data();
 
 		const int32 FormatSize = GPixelFormats[TextureToFill->GetFormat()].BlockBytes;
@@ -293,6 +298,7 @@ bool FOpenVDBCacheData::Fill3DTexture(int frame, FTextureRHIRef TextureToFill)
 {
 	if (OpenVDBGrids.Contains(frame) && OpenVDBGrids[frame] != nullptr)
 	{
+		FScopeLock ScopeLock(&DenseGridGuard);
 		uint8* DataPtr = (uint8*)DenseGridPtr->data();
 
 		const int32 FormatSize = GPixelFormats[TextureToFill->GetFormat()].BlockBytes;
