Take scene's feature level into account when creating the niagara mesh vertex factory. This will help for PSO precaching

#rnx

diff --git a/Engine/Plugins/FX/Niagara/Source/Niagara/Private/NiagaraRendererMeshes.cpp b/Engine/Plugins/FX/Niagara/Source/Niagara/Private/NiagaraRendererMeshes.cpp
--- a/Engine/Plugins/FX/Niagara/Source/Niagara/Private/NiagaraRendererMeshes.cpp
+++ b/Engine/Plugins/FX/Niagara/Source/Niagara/Private/NiagaraRendererMeshes.cpp
@@ -1260,6 +1260,10 @@ void FNiagaraRendererMeshes::GetDynamicMeshElements(const TArray<const FSceneVie
 				// TODO: Find a way to safely pool these such that they won't be concurrently accessed by multiple views
 				FNiagaraMeshVertexFactory& VertexFactory = CollectorResources->VertexFactory;
 				VertexFactory.SetParticleFactoryType(NVFT_Mesh);
+				if (ViewFamily.Scene != nullptr)
+				{
+					VertexFactory.SetNiagaraMeshFeatureLevel(ViewFamily.Scene->GetFeatureLevel());
+				}
 				VertexFactory.SetMeshIndex(MeshIndex);
 				VertexFactory.SetLODIndex(LODModel.LODIndex);
 				VertexFactory.EnablePrimitiveIDElement(ParticleMeshRenderData.bUseGPUScene);
@@ -1410,6 +1414,10 @@ void FNiagaraRendererMeshes::GetDynamicRayTracingInstances(FRayTracingMaterialGa
 		// TODO: Find a way to safely pool these such that they won't be concurrently accessed by multiple views
 		FNiagaraMeshVertexFactory& VertexFactory = CollectorResources->VertexFactory;
 		VertexFactory.SetParticleFactoryType(NVFT_Mesh);
+		if (View->Family->Scene != nullptr)
+		{
+			VertexFactory.SetNiagaraMeshFeatureLevel(View->Family->Scene->GetFeatureLevel());
+		}
 		VertexFactory.SetMeshIndex(MeshIndex);
 		VertexFactory.SetLODIndex(LODModel.LODIndex);
 		VertexFactory.EnablePrimitiveIDElement(ParticleMeshRenderData.bUseGPUScene);
diff --git a/Engine/Plugins/FX/Niagara/Source/NiagaraVertexFactories/Private/NiagaraMeshVertexFactory.cpp b/Engine/Plugins/FX/Niagara/Source/NiagaraVertexFactories/Private/NiagaraMeshVertexFactory.cpp
--- a/Engine/Plugins/FX/Niagara/Source/NiagaraVertexFactories/Private/NiagaraMeshVertexFactory.cpp
+++ b/Engine/Plugins/FX/Niagara/Source/NiagaraVertexFactories/Private/NiagaraMeshVertexFactory.cpp
@@ -79,11 +79,12 @@ void FNiagaraMeshVertexFactory::InitRHI()
 		Elements.Add(AccessStreamComponent(Data.PositionComponent, 0));
 	}
 
-	// Can't use GetFeatureLevel() on FNiagaraMeshVertexFactory because it's never set during creation
-	// This needs to be fixed and then the unused stream components don't have to be added when
-	// manual vertex fetch is used - can't use GMaxRHIFeatureLevel to check for support because then ES3.1 preview won't work anymore
-	//const bool bUseManualVertexFetch = SupportsManualVertexFetch(GetFeatureLevel());
-	//if (!bUseManualVertexFetch)
+	// Should also work in editor (PIE with ES 3.1 preview seems to work), but we don't precache PSOs for that platform currently, so keep it safe in case there's a flow we missed
+#if !WITH_EDITOR
+	const bool bHasValidFeatureLevel = ensure(HasValidFeatureLevel());
+	const bool bRemoveUnusedVertexElements = bHasValidFeatureLevel ? SupportsManualVertexFetch(GetFeatureLevel()) : false;
+	if (!bRemoveUnusedVertexElements)
+#endif
 	{
 		// only tangent,normal are used by the stream. the binormal is derived in the shader
 		uint8 TangentBasisAttributes[2] = { 1, 2 };
diff --git a/Engine/Plugins/FX/Niagara/Source/NiagaraVertexFactories/Public/NiagaraMeshVertexFactory.h b/Engine/Plugins/FX/Niagara/Source/NiagaraVertexFactories/Public/NiagaraMeshVertexFactory.h
--- a/Engine/Plugins/FX/Niagara/Source/NiagaraVertexFactories/Public/NiagaraMeshVertexFactory.h
+++ b/Engine/Plugins/FX/Niagara/Source/NiagaraVertexFactories/Public/NiagaraMeshVertexFactory.h
@@ -216,6 +216,11 @@ public:
 	bool IsPrimitiveIDElementEnabled() const { return bAddPrimitiveIDElement; }
 	void EnablePrimitiveIDElement(bool bEnable) { bAddPrimitiveIDElement = bEnable; }
 	
+	inline void SetNiagaraMeshFeatureLevel(ERHIFeatureLevel::Type InFeatureLevel)
+	{
+		FNiagaraVertexFactoryBase::SetFeatureLevel(InFeatureLevel);
+	}
+
 protected:
 	FStaticMeshDataType Data;
 	int32 MeshIndex;	
