Fix for errors caused by ambiguous operator== calls between FNiagaraVariable and FNiagaraVariableBase.
Added specific operator== to handle it which should pacify the compiler.

diff --git a/Engine/Plugins/FX/Niagara/Source/Niagara/Public/NiagaraTypes.h b/Engine/Plugins/FX/Niagara/Source/Niagara/Public/NiagaraTypes.h
--- a/Engine/Plugins/FX/Niagara/Source/Niagara/Public/NiagaraTypes.h
+++ b/Engine/Plugins/FX/Niagara/Source/Niagara/Public/NiagaraTypes.h
@@ -1856,6 +1856,12 @@ struct FNiagaraVariable : public FNiagaraVariableBase
 		return Name == Other.Name && TypeDefHandle == Other.TypeDefHandle;
 	}
 
+	/** Check if Name and Type definition are the same. The actual stored value is not checked here.*/
+	bool operator==(const FNiagaraVariableBase& Other)const
+	{
+		return FNiagaraVariableBase(*this) == Other;
+	}
+
 	/** Check if Name and Type definition are not the same. The actual stored value is not checked here.*/
 	bool operator!=(const FNiagaraVariable& Other)const
 	{
