Fix crash when changing selected emitter when the overview panel was never opened.

diff --git a/Engine/Plugins/FX/Niagara/Source/NiagaraEditor/Private/ViewModels/NiagaraSimCacheViewModel.cpp b/Engine/Plugins/FX/Niagara/Source/NiagaraEditor/Private/ViewModels/NiagaraSimCacheViewModel.cpp
--- a/Engine/Plugins/FX/Niagara/Source/NiagaraEditor/Private/ViewModels/NiagaraSimCacheViewModel.cpp
+++ b/Engine/Plugins/FX/Niagara/Source/NiagaraEditor/Private/ViewModels/NiagaraSimCacheViewModel.cpp
@@ -443,6 +443,12 @@ void FNiagaraSimCacheViewModel::BuildEntries(TWeakPtr<SNiagaraSimCacheTreeView>
 
 void FNiagaraSimCacheViewModel::UpdateCurrentEntries()
 {
+	// The overview panel is in charge of building these entries.
+	// Early out if they haven't been built yet.
+	if(!RootEntries.IsValidIndex(EmitterIndex +1))
+	{
+		return;
+	}
 	CurrentRootEntries.Empty();
 	CurrentRootEntries.Add(RootEntries[EmitterIndex + 1]);
 }
