- Add Guid to SimCache
- This is created the first time the cache is written to

#rnx

diff --git a/Engine/Plugins/FX/Niagara/Source/Niagara/Classes/NiagaraSimCache.h b/Engine/Plugins/FX/Niagara/Source/Niagara/Classes/NiagaraSimCache.h
--- a/Engine/Plugins/FX/Niagara/Source/Niagara/Classes/NiagaraSimCache.h
+++ b/Engine/Plugins/FX/Niagara/Source/Niagara/Classes/NiagaraSimCache.h
@@ -267,6 +267,11 @@ public:
 	virtual bool IsReadyForFinishDestroy() override;
 	// UObject Interface
 
+	/** Get the caches assigned GUID, this can be set from a user or will be auto created on first write of the cache. */
+	const FGuid& GetCacheGuid() { return CacheGuid; }
+	/** Set the caches GUID to use. */
+	void SetCacheGuid(const FGuid& InGuid) { CacheGuid = InGuid; }
+
 	/** A valid cache is one that contains at least 1 frames worth of data. */
 	UFUNCTION(BlueprintCallable, Category=NiagaraSimCache, meta=(DisplayName="IsValid"))
 	bool IsCacheValid() const { return SoftNiagaraSystem.IsNull() == false; }
@@ -416,6 +421,9 @@ public:
 	void ReadQuatAttributeWithRebase(TArray<FQuat>& OutValues, FQuat Quat, FName AttributeName = FName("MeshOrientation"), FName EmitterName = NAME_None, int FrameIndex = 0) const;
 
 private:
+	UPROPERTY(VisibleAnywhere, Category=SimCache)
+	FGuid CacheGuid;
+
 	UPROPERTY(VisibleAnywhere, Category=SimCache, meta=(DisplayName="Niagara System"))
 	TSoftObjectPtr<UNiagaraSystem> SoftNiagaraSystem;
 
diff --git a/Engine/Plugins/FX/Niagara/Source/Niagara/Private/NiagaraDataInterfaceRenderTargetVolume.cpp b/Engine/Plugins/FX/Niagara/Source/Niagara/Private/NiagaraDataInterfaceRenderTargetVolume.cpp
--- a/Engine/Plugins/FX/Niagara/Source/Niagara/Private/NiagaraDataInterfaceRenderTargetVolume.cpp
+++ b/Engine/Plugins/FX/Niagara/Source/Niagara/Private/NiagaraDataInterfaceRenderTargetVolume.cpp
@@ -556,12 +556,13 @@ bool UNiagaraDataInterfaceRenderTargetVolume::GetExposedVariableValue(const FNia
 	return false;
 }
 
-UObject* UNiagaraDataInterfaceRenderTargetVolume::SimCacheBeginWrite(UObject* SimCache, FNiagaraSystemInstance* NiagaraSystemInstance, const void* OptionalPerInstanceData) const
+UObject* UNiagaraDataInterfaceRenderTargetVolume::SimCacheBeginWrite(UObject* InSimCache, FNiagaraSystemInstance* NiagaraSystemInstance, const void* OptionalPerInstanceData) const
 {
 	if (NDIRenderTargetVolumeLocal::GSimCacheUseOpenVDB)
 	{		
 		UVolumeCache* OpenVDBSimCacheData = nullptr;
 
+		UNiagaraSimCache* SimCache = CastChecked<UNiagaraSimCache>(InSimCache);
 		OpenVDBSimCacheData = NewObject<UVolumeCache>(SimCache);
 
 		FString SystemInstanceName = NiagaraSystemInstance->GetSystem()->GetName();
@@ -572,18 +573,9 @@ UObject* UNiagaraDataInterfaceRenderTargetVolume::SimCacheBeginWrite(UObject* Si
 			return nullptr;
 		}
 
-		FString ActorName = NiagaraSystemInstance->GetAttachComponent()->GetName();
-
-#if WITH_EDITOR
-		// if we are caching from the editor, use the label of the actor to name the cache files
-		if (AActor* Owner = NiagaraSystemInstance->GetAttachComponent()->GetOwner())
-		{
-			ActorName = Owner->GetActorLabel();
-		}
-#endif
-
-		FString DIName = Proxy->SourceDIName.ToString();
-		FString FullFilePathSpec = GetDefault<UNiagaraSettings>()->SimCacheAuxiliaryFileBasePath + "/" + ActorName + "/" + DIName + "_SimCache.{FrameIndex}.vdb";
+		const FGuid& CacheGuid = SimCache->GetCacheGuid();
+		const FString DIName = Proxy->SourceDIName.ToString();
+		FString FullFilePathSpec = GetDefault<UNiagaraSettings>()->SimCacheAuxiliaryFileBasePath + "/" + CacheGuid.ToString() + "/" + DIName + "_SimCache.{FrameIndex}.vdb";
 		FullFilePathSpec.ReplaceInline(TEXT("//"), TEXT("/"));		
 		
 		FullFilePathSpec.ReplaceInline(TEXT("{project_dir}"), *FPaths::ProjectDir());
@@ -617,7 +609,7 @@ UObject* UNiagaraDataInterfaceRenderTargetVolume::SimCacheBeginWrite(UObject* Si
 		UNDIRenderTargetVolumeSimCacheData* SimCacheData = nullptr;
 		if (NDIRenderTargetVolumeLocal::GSimCacheEnabled)
 		{
-			SimCacheData = NewObject<UNDIRenderTargetVolumeSimCacheData>(SimCache);
+			SimCacheData = NewObject<UNDIRenderTargetVolumeSimCacheData>(InSimCache);
 			SimCacheData->CompressionType = NDIRenderTargetVolumeLocal::GetSimCacheCompressionType();
 		}
 
diff --git a/Engine/Plugins/FX/Niagara/Source/Niagara/Private/NiagaraSimCache.cpp b/Engine/Plugins/FX/Niagara/Source/Niagara/Private/NiagaraSimCache.cpp
--- a/Engine/Plugins/FX/Niagara/Source/Niagara/Private/NiagaraSimCache.cpp
+++ b/Engine/Plugins/FX/Niagara/Source/Niagara/Private/NiagaraSimCache.cpp
@@ -43,6 +43,11 @@ bool UNiagaraSimCache::BeginWrite(FNiagaraSimCacheCreateParameters InCreateParam
 		return false;
 	}
 
+	if (CacheGuid.IsValid() == false)
+	{
+		CacheGuid = FGuid::NewGuid();
+	}
+
 	Modify();
 #if WITH_EDITOR
 	PostEditChange();
