- Add playback rate to the baker

diff --git a/Engine/Plugins/FX/Niagara/Source/NiagaraEditor/Private/ViewModels/NiagaraBakerViewModel.cpp b/Engine/Plugins/FX/Niagara/Source/NiagaraEditor/Private/ViewModels/NiagaraBakerViewModel.cpp
--- a/Engine/Plugins/FX/Niagara/Source/NiagaraEditor/Private/ViewModels/NiagaraBakerViewModel.cpp
+++ b/Engine/Plugins/FX/Niagara/Source/NiagaraEditor/Private/ViewModels/NiagaraBakerViewModel.cpp
@@ -45,11 +45,6 @@ void FNiagaraBakerViewModel::Initialize(TWeakPtr<FNiagaraSystemViewModel> InWeak
 		.WeakViewModel(this->AsShared());
 }
 
-void FNiagaraBakerViewModel::RefreshView()
-{
-	//-TOOD:
-}
-
 void FNiagaraBakerViewModel::SetDisplayTimeFromNormalized(float NormalizeTime)
 {
 	if ( Widget )
diff --git a/Engine/Plugins/FX/Niagara/Source/NiagaraEditor/Private/ViewModels/NiagaraBakerViewModel.h b/Engine/Plugins/FX/Niagara/Source/NiagaraEditor/Private/ViewModels/NiagaraBakerViewModel.h
--- a/Engine/Plugins/FX/Niagara/Source/NiagaraEditor/Private/ViewModels/NiagaraBakerViewModel.h
+++ b/Engine/Plugins/FX/Niagara/Source/NiagaraEditor/Private/ViewModels/NiagaraBakerViewModel.h
@@ -36,16 +36,12 @@ public:
 
 	void RenderBaker();
 
-	void RefreshView();
-
 	void SetDisplayTimeFromNormalized(float NormalizeTime);
 
 	FNiagaraBakerRenderer& GetBakerRenderer() const { check(BakerRenderer.IsValid()); return *BakerRenderer.Get(); }
 	UNiagaraBakerSettings* GetBakerSettings() const { return BakerRenderer->GetBakerSettings(); }
 	const UNiagaraBakerSettings* GetBakerGeneratedSettings() { return BakerRenderer->GetBakerGeneratedSettings(); }
 
-	bool RenderView(const FRenderTarget* RenderTarget, FCanvas* Canvas, float WorldTime, int32 iOutputTextureIndex, bool bFillCanvas = false) const;
-
 	bool IsChannelEnabled(ENiagaraBakerColorChannel Channel) const { return bColorChannelEnabled[int(Channel)]; }
 	void ToggleChannelEnabled(ENiagaraBakerColorChannel Channel) { bColorChannelEnabled[int(Channel)] = !bColorChannelEnabled[int(Channel)]; }
 	void SetChannelEnabled(ENiagaraBakerColorChannel Channel, bool bEnabled) { bColorChannelEnabled[int(Channel)] = bEnabled; }
@@ -136,6 +132,9 @@ public:
 	int32 GetFramesOnY() const;
 	void SetFramesOnY(int32 Value);
 
+	float GetPlaybackRate() const { return PlaybackRate; }
+	void SetPlaybackRate(float Value) { PlaybackRate = FMath::Clamp(Value, 0.0f, 10.0f); }
+
 	FOnCurrentOutputChanged OnCurrentOutputChanged;
 
 private:
@@ -150,4 +149,6 @@ private:
 	bool bCheckerboardEnabled = true;		//-TODO: Move to Baker Settings?
 	bool bShowInfoText = true;				//-TODO: Remove later
 	bool bColorChannelEnabled[int(ENiagaraBakerColorChannel::Num)] = {true, true, true, false};
+
+	float PlaybackRate = 1.0f;
 };
diff --git a/Engine/Plugins/FX/Niagara/Source/NiagaraEditor/Private/Widgets/SNiagaraBakerWidget.cpp b/Engine/Plugins/FX/Niagara/Source/NiagaraEditor/Private/Widgets/SNiagaraBakerWidget.cpp
--- a/Engine/Plugins/FX/Niagara/Source/NiagaraEditor/Private/Widgets/SNiagaraBakerWidget.cpp
+++ b/Engine/Plugins/FX/Niagara/Source/NiagaraEditor/Private/Widgets/SNiagaraBakerWidget.cpp
@@ -486,6 +486,30 @@ TSharedRef<SWidget> SNiagaraBakerWidget::MakeSettingsWidget()
 			.OnValueChanged(ViewModel, &FNiagaraBakerViewModel::SetFramesOnY)
 		]
 
+		+SHorizontalBox::Slot()
+		.AutoWidth()
+		.Padding(10.0f)
+		.VAlign(VAlign_Center)
+		.HAlign(HAlign_Right)
+		[
+			SNew(STextBlock)
+			.Text(LOCTEXT("PlaybackRate", "Playback Rate"))
+			.Font(IDetailLayoutBuilder::GetDetailFont())
+		]
+		+SHorizontalBox::Slot()
+		.AutoWidth()
+		[
+			SNew(SSpinBox<float>)
+			.MinDesiredWidth(50.0f)
+			.MinValue(0.0f)
+			.MinSliderValue(0.0f)
+			.MaxValue(10.0f)
+			.MaxSliderValue(10.0f)
+			.Font(IDetailLayoutBuilder::GetDetailFont())
+			.Value(ViewModel, &FNiagaraBakerViewModel::GetPlaybackRate)
+			.OnValueChanged(ViewModel, &FNiagaraBakerViewModel::SetPlaybackRate)
+		]
+
 		+SHorizontalBox::Slot()
 		.AutoWidth()
 		.Padding(10.0f)
@@ -935,7 +959,7 @@ void SNiagaraBakerWidget::RefreshWidget()
 	BakerSettingsDetails->SetObject(OutputObject, true);
 }
 
-void SNiagaraBakerWidget::Tick(const FGeometry& AllottedGeometry, const double CurrentTime, const float DeltaTime)
+void SNiagaraBakerWidget::Tick(const FGeometry& AllottedGeometry, const double _UnusedCurrentTime, float DeltaTime)
 {
 	UNiagaraBakerSettings* BakerSettings = GetBakerSettings();
 	if (BakerSettings == nullptr)
@@ -946,6 +970,11 @@ void SNiagaraBakerWidget::Tick(const FGeometry& AllottedGeometry, const double C
 	const float DurationSeconds = BakerSettings->DurationSeconds;
 	if ( DurationSeconds > 0.0f )
 	{
+		if (TSharedPtr<FNiagaraBakerViewModel> ViewModel = WeakViewModel.Pin())
+		{
+			DeltaTime *= ViewModel->GetPlaybackRate();
+		}
+
 		if (PlaybackMode != EPlaybackMode::Stopped)
 		{
 			PreviewRelativeTime += PlaybackMode == EPlaybackMode::PlayingForward ? DeltaTime : -DeltaTime;
