- Indirect dispatch support for Niagara

diff --git a/Engine/Plugins/FX/Niagara/Shaders/Private/NiagaraDispatchIndirectArgsGen.usf b/Engine/Plugins/FX/Niagara/Shaders/Private/NiagaraDispatchIndirectArgsGen.usf
new file mode 100644
--- /dev/null
+++ b/Engine/Plugins/FX/Niagara/Shaders/Private/NiagaraDispatchIndirectArgsGen.usf
@@ -0,0 +1,61 @@
+// Copyright Epic Games, Inc. All Rights Reserved.
+
+/*==============================================================================
+	NiagaraDrawIndirectArgsGen.usf: Shader to copy sorted index buffers.
+==============================================================================*/
+
+#include "/Engine/Private/Common.ush"
+
+#define INDEX_NONE uint(-1)
+
+// Input Format (FNiagaraDispatchIndirectInfoCS)
+//	uint3 - Counter Offsets
+//	uint3 - Thread Group Size
+//	uint  - Indirect Args Offset
+Buffer<uint>	DispatchInfos;
+uint			NumDispatchInfos;
+uint3			MaxGroupsPerDimension;
+
+Buffer<uint>	InstanceCounts;
+
+// Output Format (FNiagaraDispatchIndirectParametersCS)
+//	uint4 - Indirect Group Count
+//	uint4 - Indirect Thread Count
+RWBuffer<uint>	RWDispatchIndirectArgs;
+
+[numthreads(THREAD_COUNT,1,1)]
+void IndirectArgsGenCS(uint DispatchIndex : SV_DispatchThreadID)
+{
+	// Clip dead threads
+	if (DispatchIndex >= NumDispatchInfos)
+		return;
+
+	uint DispatchOffset = DispatchIndex * 7;
+
+	uint3 NumThreads = uint3(DispatchInfos[DispatchOffset + 0], DispatchInfos[DispatchOffset + 1], DispatchInfos[DispatchOffset + 2]);
+	NumThreads.x = NumThreads.x != INDEX_NONE ? InstanceCounts[NumThreads.x] : 0;
+	NumThreads.y = NumThreads.y != INDEX_NONE ? InstanceCounts[NumThreads.y] : 1;
+	NumThreads.z = NumThreads.z != INDEX_NONE ? InstanceCounts[NumThreads.z] : 1;
+
+	uint3 ThreadGroupSize = uint3(DispatchInfos[DispatchOffset + 3], DispatchInfos[DispatchOffset + 4], DispatchInfos[DispatchOffset + 5]);
+
+	uint3 NumGroups = 0;
+	if ( all(NumThreads > 0) && all(ThreadGroupSize > 0) )
+	{
+		NumGroups = (NumThreads + ThreadGroupSize - 1) / ThreadGroupSize;
+		if ( any(NumGroups > MaxGroupsPerDimension) )
+		{
+			NumGroups = 0;
+		}
+	}
+
+	uint IndirectArgsOffset = DispatchInfos[DispatchOffset + 6] * 8;
+	RWDispatchIndirectArgs[IndirectArgsOffset + 0] = NumGroups.x;
+	RWDispatchIndirectArgs[IndirectArgsOffset + 1] = NumGroups.y;
+	RWDispatchIndirectArgs[IndirectArgsOffset + 2] = NumGroups.z;
+	RWDispatchIndirectArgs[IndirectArgsOffset + 3] = 0;
+	RWDispatchIndirectArgs[IndirectArgsOffset + 4] = NumThreads.x;
+	RWDispatchIndirectArgs[IndirectArgsOffset + 5] = NumThreads.y;
+	RWDispatchIndirectArgs[IndirectArgsOffset + 6] = NumThreads.z;
+	RWDispatchIndirectArgs[IndirectArgsOffset + 7] = 0;
+}
diff --git a/Engine/Plugins/FX/Niagara/Shaders/Private/NiagaraEmitterInstanceShader.usf b/Engine/Plugins/FX/Niagara/Shaders/Private/NiagaraEmitterInstanceShader.usf
--- a/Engine/Plugins/FX/Niagara/Shaders/Private/NiagaraEmitterInstanceShader.usf
+++ b/Engine/Plugins/FX/Niagara/Shaders/Private/NiagaraEmitterInstanceShader.usf
@@ -79,8 +79,14 @@ bool NiagaraAny(bool4 b) { return b.x || b.y || b.z || b.w; }
 		#define NIAGARA_DISPATCH_TYPE NIAGARA_DISPATCH_TYPE_ONE_D
 	#endif
 
-	uint3 DispatchThreadIdBounds;		// Used to track thread bounds
-	uint3 DispatchThreadIdToLinear;		// Used to convert from DispatchThreadId to a linear thread id
+	#if NIAGARA_DISPATCH_INDIRECT
+		static uint3	DispatchThreadIdBounds;		// Used to track thread bounds
+		Buffer<uint4>	IndirectDispatchArgs;		// Indirect count buffer we lookup into to get actual thread count
+		uint			IndirectDispatchArgsOffset;	// Offset into the indirect count buffer
+	#else
+		uint3 DispatchThreadIdBounds;			// Used to track thread bounds
+		uint3 DispatchThreadIdToLinear;			// Used to convert from DispatchThreadId to a linear thread id
+	#endif
 
 	static uint3 GDispatchThreadId;		// SV_DispatchThreadId
 	static uint3 GGroupId;				// SV_GroupdId
@@ -534,7 +540,11 @@ int rand_int(int x, int Seed1, int Seed2, int Seed3)
 	float SimulationStageNormalizedIterationIndex;
 
 	// Note: These are referenced from an asset that passes back the data to the user (see SimulationStageIterationInfo)
-	int SimulationStage_GetInstanceCount() { return SimulationStageIterationInfo.x == -1 ? SimulationStageIterationInfo.y : RWInstanceCounts[SimulationStageIterationInfo.x]; }
+	#if NIAGARA_DISPATCH_INDIRECT
+		int SimulationStage_GetInstanceCount() { return DispatchThreadIdBounds.x * DispatchThreadIdBounds.y * DispatchThreadIdBounds.z; }
+	#else
+		int SimulationStage_GetInstanceCount() { return SimulationStageIterationInfo.x == -1 ? SimulationStageIterationInfo.y : RWInstanceCounts[SimulationStageIterationInfo.x]; }
+	#endif
 	int SimulationStage_GetIterationIndex() { return SimulationStageIterationInfo.z; }
 	int SimulationStage_GetNumIterations() { return SimulationStageIterationInfo.w; }
 	float SimulationStage_GetNormalizedIterationIndex() { return SimulationStageNormalizedIterationIndex; }
diff --git a/Engine/Plugins/FX/Niagara/Source/Niagara/Classes/NiagaraDataInterface.h b/Engine/Plugins/FX/Niagara/Source/Niagara/Classes/NiagaraDataInterface.h
--- a/Engine/Plugins/FX/Niagara/Source/Niagara/Classes/NiagaraDataInterface.h
+++ b/Engine/Plugins/FX/Niagara/Source/Niagara/Classes/NiagaraDataInterface.h
@@ -258,7 +258,7 @@ struct NIAGARA_API FNDIGpuComputeContext
 
 	const FNiagaraGpuComputeDispatchInterface& GetComputeDispatchInterface() const { return ComputeDispatchInterface; }
 
-private:
+protected:
 	FRDGBuilder& GraphBuilder;
 	const FNiagaraGpuComputeDispatchInterface& ComputeDispatchInterface;
 };
@@ -273,7 +273,7 @@ struct NIAGARA_API FNDIGpuComputeResetContext : public FNDIGpuComputeContext
 
 	FNiagaraSystemInstanceID GetSystemInstanceID() const { return SystemInstanceID; }
 
-private:
+protected:
 	FNiagaraSystemInstanceID SystemInstanceID = FNiagaraSystemInstanceID();
 };
 
@@ -298,7 +298,7 @@ struct NIAGARA_API FNDIGpuComputePrePostStageContext : public FNDIGpuComputeCont
 
 	void SetDataInterfaceProxy(FNiagaraDataInterfaceProxy* InDataInterfaceProxy) { DataInterfaceProxy = InDataInterfaceProxy; }
 
-private:
+protected:
 	const FNiagaraGPUSystemTick& SystemTick;
 	const FNiagaraComputeInstanceData& ComputeInstanceData;
 	const FNiagaraSimStageData& SimStageData;
@@ -320,7 +320,7 @@ struct NIAGARA_API FNDIGpuComputePostSimulateContext : public FNDIGpuComputeCont
 	FNiagaraSystemInstanceID GetSystemInstanceID() const { return SystemInstanceID; }
 	bool IsFinalPostSimulate() const { return bFinalPostSimulate; }
 
-private:
+protected:
 	FNiagaraSystemInstanceID SystemInstanceID = FNiagaraSystemInstanceID();
 	bool bFinalPostSimulate = false;
 };
@@ -471,6 +471,7 @@ public:
 
 	virtual ENiagaraGpuDispatchType GetGpuDispatchType() const { return ENiagaraGpuDispatchType::OneD; }
 	virtual FIntVector GetGpuDispatchNumThreads() const { return FIntVector(64, 1, 1); }
+	virtual bool GetGpuUseIndirectDispatch() const { return false; }
 #endif
 
 	/** Allows data interfaces to cache any static data that may be shared between instances */
diff --git a/Engine/Plugins/FX/Niagara/Source/Niagara/Classes/NiagaraDataInterfaceRW.h b/Engine/Plugins/FX/Niagara/Source/Niagara/Classes/NiagaraDataInterfaceRW.h
--- a/Engine/Plugins/FX/Niagara/Source/Niagara/Classes/NiagaraDataInterfaceRW.h
+++ b/Engine/Plugins/FX/Niagara/Source/Niagara/Classes/NiagaraDataInterfaceRW.h
@@ -2,10 +2,77 @@
 #pragma once
 
 #include "NiagaraDataInterface.h"
+#include "NiagaraDispatchIndirect.h"
 #include "NiagaraDataInterfaceRW.generated.h"
 
+struct FNiagaraSimStageData;
+struct FNiagaraSimStageDispatchArgs;
+struct FNiagaraDispatchIndirectInfoCS;
 class FNiagaraSystemInstance;
 
+BEGIN_SHADER_PARAMETER_STRUCT(FNDIGpuComputeDispatchArgsGenParameters, )
+	SHADER_PARAMETER_STRUCT_INCLUDE(FNiagaraDispatchIndirectArgsGenCS::FParameters,	IndirectArgsGenParameters)
+	
+	RDG_BUFFER_ACCESS_ARRAY(BufferAccessArray)
+	RDG_TEXTURE_ACCESS_ARRAY(TextureAccessArray)
+END_SHADER_PARAMETER_STRUCT()
+
+struct NIAGARA_API FNDIGpuComputeDispatchArgsGenContext : public FNDIGpuComputeContext
+{
+	using FIndirectArgs = TPair<FRDGBuffer*, uint32>;
+	using FCreateIndirectCallback = TFunction<void(FRHICommandList&, FRDGBuffer*, uint32)>;
+
+	explicit FNDIGpuComputeDispatchArgsGenContext(FRDGBuilder& InGraphBuilder, const FNiagaraGpuComputeDispatchInterface& InComputeDispatchInterface)
+		: FNDIGpuComputeContext(InGraphBuilder, InComputeDispatchInterface)
+	{
+	}
+
+	~FNDIGpuComputeDispatchArgsGenContext() { FlushPass(); }
+
+	void SetInstanceData(FNiagaraSystemInstanceID InSystemInstanceID, FNiagaraSimStageData& InSimStageData)
+	{
+		SystemInstanceID = InSystemInstanceID;
+		SimStageData = &InSimStageData;
+	}
+
+	const FNiagaraSimStageData& GetSimStageData() const { return *SimStageData; }
+	FNiagaraSystemInstanceID GetSystemInstanceID() const { return SystemInstanceID; }
+
+	// Set Direct Dispatch Element Count / Gpu Element Count Offset
+	void SetDirect(uint32 InElementCount, uint32 GpuCountOffset = INDEX_NONE) const { SetDirect(FIntVector3(InElementCount, 1, 1), GpuCountOffset); }
+	void SetDirect(const FIntPoint& InElementCount, uint32 GpuCountOffset = INDEX_NONE) const { SetDirect(FIntVector3(InElementCount.X, InElementCount.Y, 1), GpuCountOffset); }
+	void SetDirect(const FIntVector2& InElementCountXY, uint32 GpuCountOffset = INDEX_NONE) const { SetDirect(FIntVector3(InElementCountXY.X, InElementCountXY.Y, 1), GpuCountOffset); }
+	void SetDirect(const FIntVector3& InElementCountXYZ, uint32 GpuCountOffset = INDEX_NONE) const;
+
+	// Set Indirect buffer to use
+	void SetIndirect(FRDGBuffer* InBuffer, uint32 BufferByteOffset) const;
+	// Set indirect buffer to use, the callback will be inside the RDG pass
+	void SetIndirect(FRDGBuffer* InBuffer, uint32 BufferByteOffset, FCreateIndirectCallback&& Callback) const;
+	// Create indirect buffer from count values that are held in the GPU count buffer
+	// The returned indirect argument information can be kept around for the current frame
+	FIndirectArgs CreateIndirect(uint32 InCounterOffset) const { return CreateIndirect(FUintVector3(InCounterOffset, INDEX_NONE, INDEX_NONE)); }
+	FIndirectArgs CreateIndirect(const FUintVector2& InCounterOffsets) const { return CreateIndirect(FUintVector3(InCounterOffsets.X, InCounterOffsets.Y, INDEX_NONE)); }
+	FIndirectArgs CreateIndirect(const FUintVector3& InCounterOffsets) const;
+
+	// Add buffer / texture access used for batching transitions together
+	void AddBufferAccess(FRDGBuffer* InBuffer, ERHIAccess InAccess) const;
+	void AddTextureAccess(FRDGTexture* InTexture, ERHIAccess InAccess) const;
+
+protected:
+	FNDIGpuComputeDispatchArgsGenParameters* GetPassParameters() const;
+	void FlushPass() const;
+
+protected:
+	FNiagaraSystemInstanceID SystemInstanceID;
+	FNiagaraSimStageData* SimStageData = nullptr;
+
+	mutable FRDGBuffer* IndirectBuffer = nullptr;
+	mutable FNiagaraDispatchIndirectInfoCS* IndirectCounterGenArgs = nullptr;
+	mutable int32 NumIndirectCounterGenArgs = 0;
+	mutable TArray<TPair<FCreateIndirectCallback, FIndirectArgs>> IndirectCallbacks;
+	mutable FNDIGpuComputeDispatchArgsGenParameters* PassParameters = nullptr;
+};
+
 UENUM()
 enum class ESetResolutionMethod
 {
@@ -21,11 +88,18 @@ public:
 	virtual void ConsumePerInstanceDataFromGameThread(void* PerInstanceData, const FNiagaraSystemInstanceID& Instance) override { check(false); }
 	virtual int32 PerInstanceDataPassedToRenderThreadSize() const override { return 0; }	
 
+	// Called before PreStage to get the dispatch arguments
+	// The default implementation is temporary and will call GetElementCount & GetGPUInstanceCountOffset to set direct arguments.
+	// This method will become a pure virtual in 5.3 to enforce its usage over the older methods.
+	NIAGARA_API virtual void GetDispatchArgs(const FNDIGpuComputeDispatchArgsGenContext& Context);
+
 	// Get the element count for this instance
-	virtual FIntVector GetElementCount(FNiagaraSystemInstanceID SystemInstanceID) const = 0;
+	UE_DEPRECATED(5.2, "This function will be removed in 5.3.  You must implement GetDispatchArgs instead.")
+	virtual FIntVector GetElementCount(FNiagaraSystemInstanceID SystemInstanceID) const { return FIntVector::ZeroValue; }
 
 	// For data interfaces that support iteration on the GPU we need to be able to get the 'real' element count as known only by the GPU
 	// The dispatch will use the CPU count, which is potentially an over-estimation, and the value inside the buffer will be used to clip instances that are not valid
+	UE_DEPRECATED(5.2, "This function will be removed in 5.3.  You must implement GetDispatchArgs instead.")
 	virtual uint32 GetGPUInstanceCountOffset(FNiagaraSystemInstanceID SystemInstanceID) const { return INDEX_NONE; }
 
 	virtual FNiagaraDataInterfaceProxyRW* AsIterationProxy() override { return this; }
diff --git a/Engine/Plugins/FX/Niagara/Source/Niagara/Private/NiagaraDataInterface.cpp b/Engine/Plugins/FX/Niagara/Source/Niagara/Private/NiagaraDataInterface.cpp
--- a/Engine/Plugins/FX/Niagara/Source/Niagara/Private/NiagaraDataInterface.cpp
+++ b/Engine/Plugins/FX/Niagara/Source/Niagara/Private/NiagaraDataInterface.cpp
@@ -44,13 +44,14 @@ bool UNiagaraDataInterface::AppendCompileHash(FNiagaraCompileHashVisitor* InVisi
 {
 	//-TODO: Currently applied to all, but we only need to hash this in for the iteration source
 	const UNiagaraDataInterface* BaseDataInterface = GetDefault<UNiagaraDataInterface>();
+	const FString DataInterfaceName = GetClass()->GetName();
 	if (BaseDataInterface->GetGpuDispatchType() != GetGpuDispatchType())
 	{
-		const FString DataInterfaceName = GetClass()->GetName();
 		InVisitor->UpdatePOD(*FString::Printf(TEXT("%s_GpuDispatchType"), *DataInterfaceName), (int32)GetGpuDispatchType());
 		InVisitor->UpdateString(*FString::Printf(TEXT("%s_GpuDispatchNumThreads"), *DataInterfaceName), *FString::Printf(TEXT("%dx%dx%d"), GetGpuDispatchNumThreads().X, GetGpuDispatchNumThreads().Y, GetGpuDispatchNumThreads().Z));
 	}
-
+	InVisitor->UpdatePOD(*FString::Printf(TEXT("%s_GpuUseIndirectDispatch"), *DataInterfaceName), (int32)GetGpuUseIndirectDispatch());
+	
 	return true;
 }
 #endif
diff --git a/Engine/Plugins/FX/Niagara/Source/Niagara/Private/NiagaraDataInterfaceGrid2DCollection.cpp b/Engine/Plugins/FX/Niagara/Source/Niagara/Private/NiagaraDataInterfaceGrid2DCollection.cpp
--- a/Engine/Plugins/FX/Niagara/Source/Niagara/Private/NiagaraDataInterfaceGrid2DCollection.cpp
+++ b/Engine/Plugins/FX/Niagara/Source/Niagara/Private/NiagaraDataInterfaceGrid2DCollection.cpp
@@ -2621,13 +2621,7 @@ void FNiagaraDataInterfaceProxyGrid2DCollectionProxy::PreStage(const FNDIGpuComp
 		// If there is an output DI and the stage dispatch count is different from the grid resolution, then we want to consider it for a pre-stage UAV clear.
 		// This is to enable the most common use case of scattered writes to temporary buffers		
 		const FNiagaraSimStageData& SimStageData = Context.GetSimStageData();
-
-		FIntVector ElementCount = SimStageData.ElementCountXYZ;
-
-		if (SimStageData.AlternateIterationSource != nullptr)
-		{
-			ElementCount = SimStageData.AlternateIterationSource->GetElementCount(Context.GetSystemInstanceID());
-		}	
+		const FIntVector3& ElementCount = SimStageData.DispatchArgs.ElementCount;
 
 		if (ProxyData->ClearBeforeNonIterationStage && (ElementCount.X != ProxyData->NumCells.X || ElementCount.Y != ProxyData->NumCells.Y))
 		{
diff --git a/Engine/Plugins/FX/Niagara/Source/Niagara/Private/NiagaraDataInterfaceGrid3DCollection.cpp b/Engine/Plugins/FX/Niagara/Source/Niagara/Private/NiagaraDataInterfaceGrid3DCollection.cpp
--- a/Engine/Plugins/FX/Niagara/Source/Niagara/Private/NiagaraDataInterfaceGrid3DCollection.cpp
+++ b/Engine/Plugins/FX/Niagara/Source/Niagara/Private/NiagaraDataInterfaceGrid3DCollection.cpp
@@ -4156,13 +4156,7 @@ void FNiagaraDataInterfaceProxyGrid3DCollectionProxy::PreStage(const FNDIGpuComp
 		// If there is an output DI and the stage dispatch count is different from the grid resolution, then we want to consider it for a pre-stage UAV clear.
 		// This is to enable the most common use case of scattered writes to temporary buffers				
 		const FNiagaraSimStageData& SimStageData = Context.GetSimStageData();
-
-		FIntVector ElementCount = SimStageData.ElementCountXYZ;
-
-		if (SimStageData.AlternateIterationSource != nullptr)
-		{
-			ElementCount = SimStageData.AlternateIterationSource->GetElementCount(Context.GetSystemInstanceID());
-		}
+		const FIntVector3& ElementCount = SimStageData.DispatchArgs.ElementCount;
 
 		if (ProxyData->ClearBeforeNonIterationStage && 
 			(ElementCount.X != ProxyData->NumCells.X || ElementCount.Y != ProxyData->NumCells.Y || ElementCount.Z != ProxyData->NumCells.Z))
diff --git a/Engine/Plugins/FX/Niagara/Source/Niagara/Private/NiagaraDataInterfaceRW.cpp b/Engine/Plugins/FX/Niagara/Source/Niagara/Private/NiagaraDataInterfaceRW.cpp
--- a/Engine/Plugins/FX/Niagara/Source/Niagara/Private/NiagaraDataInterfaceRW.cpp
+++ b/Engine/Plugins/FX/Niagara/Source/Niagara/Private/NiagaraDataInterfaceRW.cpp
@@ -1,6 +1,9 @@
 // Copyright Epic Games, Inc. All Rights Reserved.
 #include "NiagaraDataInterfaceRW.h"
+#include "NiagaraGpuComputeDispatchInterface.h"
 #include "NiagaraShader.h"
+#include "NiagaraSimStageData.h"
+#include "RenderGraph.h"
 #include "ShaderParameterUtils.h"
 #include "ClearQuad.h"
 
@@ -42,6 +45,138 @@ const FName UNiagaraDataInterfaceRWBase::LinearToIndexFunctionName("LinearToInde
 const FName UNiagaraDataInterfaceRWBase::ExecutionIndexToUnitFunctionName("ExecutionIndexToUnit");
 const FName UNiagaraDataInterfaceRWBase::ExecutionIndexToGridIndexFunctionName("ExecutionIndexToGridIndex");
 
+/*--------------------------------------------------------------------------------------------------------------------------*/
+
+void FNDIGpuComputeDispatchArgsGenContext::SetDirect(const FIntVector3& InElementCount, uint32 GpuCountOffset) const
+{
+	SimStageData->DispatchArgs.ElementCount = InElementCount;
+	SimStageData->DispatchArgs.GpuElementCountOffset = GpuCountOffset;
+}
+
+void FNDIGpuComputeDispatchArgsGenContext::SetIndirect(FRDGBuffer* InBuffer, uint32 BufferByteOffset) const
+{
+	SimStageData->DispatchArgs.IndirectBuffer = InBuffer;
+	SimStageData->DispatchArgs.IndirectOffset = BufferByteOffset;
+}
+
+void FNDIGpuComputeDispatchArgsGenContext::SetIndirect(FRDGBuffer* InBuffer, uint32 BufferByteOffset, FCreateIndirectCallback&& Callback) const
+{
+	SimStageData->DispatchArgs.IndirectBuffer = InBuffer;
+	SimStageData->DispatchArgs.IndirectOffset = BufferByteOffset;
+
+	auto& IndirectCallback = IndirectCallbacks.AddDefaulted_GetRef();
+	IndirectCallback.Key = Callback;
+	IndirectCallback.Value.Key = InBuffer;
+	IndirectCallback.Value.Value = BufferByteOffset;
+}
+
+FNDIGpuComputeDispatchArgsGenContext::FIndirectArgs FNDIGpuComputeDispatchArgsGenContext::CreateIndirect(const FUintVector3& InCounterOffsets) const
+{
+	const int32 BatchSize = FNiagaraDispatchIndirectArgsGenCS::ThreadCount;
+
+	// Do we need to allocate new indirect buffer?
+	if (IndirectBuffer == nullptr)
+	{
+		IndirectCounterGenArgs = GraphBuilder.AllocPODArray<FNiagaraDispatchIndirectInfoCS>(BatchSize);
+		NumIndirectCounterGenArgs = 0;
+
+		IndirectBuffer = GraphBuilder.CreateBuffer(FRDGBufferDesc::CreateIndirectDesc<FNiagaraDispatchIndirectParametersCS>(BatchSize), TEXT("NiagaraIndirectBuffer"));
+	}
+
+	SimStageData->DispatchArgs.IndirectBuffer = IndirectBuffer;
+	SimStageData->DispatchArgs.IndirectOffset = NumIndirectCounterGenArgs * sizeof(FNiagaraDispatchIndirectParametersCS);
+
+	FIndirectArgs IndirectArgs;
+	IndirectArgs.Key = IndirectBuffer;
+	IndirectArgs.Value = NumIndirectCounterGenArgs;
+	IndirectCounterGenArgs[NumIndirectCounterGenArgs++] = FNiagaraDispatchIndirectInfoCS(InCounterOffsets, SimStageData->StageMetaData->GpuDispatchNumThreads, IndirectArgs.Value);
+	if (NumIndirectCounterGenArgs == BatchSize)
+	{
+		FlushPass();
+	}
+	return IndirectArgs;
+}
+
+void FNDIGpuComputeDispatchArgsGenContext::AddBufferAccess(FRDGBuffer* InBuffer, ERHIAccess InAccess) const
+{
+	GetPassParameters()->BufferAccessArray.Emplace(InBuffer, InAccess);
+}
+
+void FNDIGpuComputeDispatchArgsGenContext::AddTextureAccess(FRDGTexture* InTexture, ERHIAccess InAccess) const
+{
+	GetPassParameters()->TextureAccessArray.Emplace(InTexture, InAccess);
+}
+
+FNDIGpuComputeDispatchArgsGenParameters* FNDIGpuComputeDispatchArgsGenContext::GetPassParameters() const
+{
+	if (PassParameters == nullptr)
+	{
+		PassParameters = GraphBuilder.AllocParameters<FNDIGpuComputeDispatchArgsGenParameters>();
+	}
+	return PassParameters;
+}
+
+void FNDIGpuComputeDispatchArgsGenContext::FlushPass() const
+{
+	if (NumIndirectCounterGenArgs == 0 && IndirectCallbacks.Num() == 0)
+	{
+		return;
+	}
+
+	// Create Args Buffer and Upload into it
+	FNiagaraDispatchIndirectArgsGenCS::FParameters* IndirectArgsGenParameters = &GetPassParameters()->IndirectArgsGenParameters;
+	if (NumIndirectCounterGenArgs > 0)
+	{
+		const uint32 DispatchInfosBufferSize = sizeof(FNiagaraDispatchIndirectInfoCS) * NumIndirectCounterGenArgs;
+		FRDGBufferRef DispatchInfosBuffer = GraphBuilder.CreateBuffer(FRDGBufferDesc::CreateBufferDesc(sizeof(uint32), DispatchInfosBufferSize / sizeof(uint32)), TEXT("NiagaraIndirectGenArgs"));
+		GraphBuilder.QueueBufferUpload(DispatchInfosBuffer, IndirectCounterGenArgs, DispatchInfosBufferSize, ERDGInitialDataFlags::NoCopy);
+
+		IndirectArgsGenParameters->DispatchInfos = GraphBuilder.CreateSRV(DispatchInfosBuffer, PF_R32_UINT);
+		IndirectArgsGenParameters->NumDispatchInfos = NumIndirectCounterGenArgs;
+		IndirectArgsGenParameters->MaxGroupsPerDimension = FUintVector3(GRHIMaxDispatchThreadGroupsPerDimension);
+
+		IndirectArgsGenParameters->InstanceCounts = ComputeDispatchInterface.GetGPUInstanceCounterManager().GetInstanceCountBuffer().SRV;
+		IndirectArgsGenParameters->RWDispatchIndirectArgs = GraphBuilder.CreateUAV(IndirectBuffer, PF_R32_UINT);
+	}
+
+	// Kick pass to genreate args
+	GraphBuilder.AddPass(
+		RDG_EVENT_NAME("Niagara::ExecuteTicks::DispatchGroupPre"),
+		GetPassParameters(),
+		ERDGPassFlags::Compute | ERDGPassFlags::NeverCull,
+		[PassParameters=GetPassParameters(), IndirectArgsGenParameters, CreateCallbacks=MoveTemp(IndirectCallbacks)](FRHICommandListImmediate& RHICmdList)
+		{
+			if (IndirectArgsGenParameters != nullptr)
+			{
+				TShaderMapRef<FNiagaraDispatchIndirectArgsGenCS> DispatchIndirectArgsGenCS(GetGlobalShaderMap(GMaxRHIFeatureLevel));
+				FComputeShaderUtils::Dispatch(RHICmdList, DispatchIndirectArgsGenCS, *IndirectArgsGenParameters, FIntVector(1, 1, 1));
+			}
+			for (const auto& CreateCallback : CreateCallbacks)
+			{
+				CreateCallback.Key(RHICmdList, CreateCallback.Value.Key, CreateCallback.Value.Value);
+			}
+		}
+	);
+
+	// Clear our data
+	IndirectBuffer = nullptr;
+	IndirectCounterGenArgs = nullptr;
+	NumIndirectCounterGenArgs = 0;
+	IndirectCallbacks.Empty();
+	PassParameters = nullptr;
+}
+
+/*--------------------------------------------------------------------------------------------------------------------------*/
+
+void FNiagaraDataInterfaceProxyRW::GetDispatchArgs(const FNDIGpuComputeDispatchArgsGenContext& Context)
+{
+	PRAGMA_DISABLE_DEPRECATION_WARNINGS
+	Context.SetDirect(GetElementCount(Context.GetSystemInstanceID()), GetGPUInstanceCountOffset(Context.GetSystemInstanceID()));
+	PRAGMA_ENABLE_DEPRECATION_WARNINGS
+}
+
+/*--------------------------------------------------------------------------------------------------------------------------*/
+
 UNiagaraDataInterfaceRWBase::UNiagaraDataInterfaceRWBase(FObjectInitializer const& ObjectInitializer)
 	: Super(ObjectInitializer)
 {
diff --git a/Engine/Plugins/FX/Niagara/Source/Niagara/Private/NiagaraGpuComputeDispatch.cpp b/Engine/Plugins/FX/Niagara/Source/Niagara/Private/NiagaraGpuComputeDispatch.cpp
--- a/Engine/Plugins/FX/Niagara/Source/Niagara/Private/NiagaraGpuComputeDispatch.cpp
+++ b/Engine/Plugins/FX/Niagara/Source/Niagara/Private/NiagaraGpuComputeDispatch.cpp
@@ -397,7 +397,6 @@ void FNiagaraGpuComputeDispatch::ProcessPendingTicksFlush(FRHICommandListImmedia
 
 			// Make a temporary ViewInfo
 			//-TODO: We could gather some more information here perhaps?
-
 			FSceneViewFamily ViewFamily(FSceneViewFamily::ConstructionValues(nullptr, nullptr, FEngineShowFlags(ESFIM_Game))
 				.SetTime(FGameTime())
 				.SetGammaCorrection(1.0f));
@@ -594,10 +593,7 @@ void FNiagaraGpuComputeDispatch::DumpDebugFrame()
 				Builder.Appendf(TEXT("Source(%p 0x%08x %d) "), SimStageData.Source, SimStageData.SourceCountOffset, SimStageData.SourceNumInstances);
 				Builder.Appendf(TEXT("Destination(%p 0x%08x %d) "), SimStageData.Destination, SimStageData.DestinationCountOffset, SimStageData.DestinationNumInstances);
 				Builder.Appendf(TEXT("Iteration(%d | %s) "), SimStageData.IterationIndex, SimStageData.AlternateIterationSource ? *SimStageData.AlternateIterationSource->SourceDIName.ToString() : TEXT("Particles"));
-				if (SimStageData.StageMetaData->IterationSourceType == ENiagaraIterationSource::DirectSet)
-				{
-					Builder.Appendf(TEXT("ElementCountXYZ(%d, %d, %d) "), SimStageData.ElementCountXYZ.X, SimStageData.ElementCountXYZ.Y, SimStageData.ElementCountXYZ.Z);
-				}
+				Builder.Appendf(TEXT("DispatchElementCount(%d, %d, %d) "), SimStageData.DispatchArgs.ElementCount.X, SimStageData.DispatchArgs.ElementCount.Y, SimStageData.DispatchArgs.ElementCount.Z);
 				UE_LOG(LogNiagara, Warning, TEXT("%s"), Builder.ToString());
 			}
 
@@ -851,7 +847,7 @@ void FNiagaraGpuComputeDispatch::PrepareTicksForProxy(FRHICommandListImmediate&
 					SimStageData.bFirstStage = bFirstStage;
 					SimStageData.StageIndex = SimStageIndex;
 					SimStageData.IterationIndex = IterationIndex;
-					SimStageData.ElementCountXYZ = InstanceData.PerStageInfo[SimStageIndex].ElementCountXYZ;
+					SimStageData.DispatchArgs.ElementCount = InstanceData.PerStageInfo[SimStageIndex].ElementCountXYZ;
 					SimStageData.StageMetaData = &SimStageMetaData;
 					SimStageData.AlternateIterationSource = IterationInterface;
 
@@ -1139,9 +1135,17 @@ void FNiagaraGpuComputeDispatch::ExecuteTicks(FRDGBuilder& GraphBuilder, TConstA
 
 		// Pre Stage
 		{
+			FNDIGpuComputeDispatchArgsGenContext DispatchArgsGenContext(GraphBuilder, *this);
+
 			TSet<FNiagaraDataInterfaceProxy*> ProxiesToFinalize;
 			for (const FNiagaraGpuDispatchInstance& DispatchInstance : DispatchGroup.DispatchInstances)
 			{
+				if (DispatchInstance.SimStageData.AlternateIterationSource)
+				{
+					DispatchArgsGenContext.SetInstanceData(DispatchInstance.Tick.SystemInstanceID, const_cast<FNiagaraSimStageData&>(DispatchInstance.SimStageData));
+					DispatchInstance.SimStageData.AlternateIterationSource->GetDispatchArgs(DispatchArgsGenContext);
+				}
+
 				PreStageInterface(GraphBuilder, DispatchInstance.Tick, DispatchInstance.InstanceData, DispatchInstance.SimStageData, ProxiesToFinalize);
 			}
 			for (FNiagaraDataInterfaceProxy* ProxyToFinalize : ProxiesToFinalize)
@@ -1150,6 +1154,7 @@ void FNiagaraGpuComputeDispatch::ExecuteTicks(FRDGBuilder& GraphBuilder, TConstA
 			}
 		}
 
+		// Execute Transitions
  		GraphBuilder.AddPass(
 			{},//RDG_EVENT_NAME("Niagara::ExecuteTicks::DispatchGroupPre"),
 			ERDGPassFlags::None,
@@ -1396,20 +1401,31 @@ void FNiagaraGpuComputeDispatch::DispatchStage(FRDGBuilder& GraphBuilder, const
 			}
 
 			DispatchType = SimStageData.StageMetaData->GpuDispatchType;
-			DispatchCount = SimStageData.AlternateIterationSource->GetElementCount(Tick.SystemInstanceID);
 			DispatchNumThreads = SimStageData.StageMetaData->GpuDispatchNumThreads;
 
-			// Verify the number of elements isn't higher that what we can handle
-			checkf(uint64(DispatchCount.X) * uint64(DispatchCount.Y) * uint64(DispatchCount.Z) < uint64(TNumericLimits<int32>::Max()), TEXT("DispatchCount(%d, %d, %d) for IterationInterface(%s) overflows an int32 this is not allowed"), DispatchCount.X, DispatchCount.Y, DispatchCount.Z, *SimStageData.AlternateIterationSource->SourceDIName.ToString());
-
-			// Data interfaces such as grids / render targets can choose to dispatch in either the correct dimensionality for the target (i.e. RT2D would choose 2D)
-			// or run in linear mode if performance is not beneficial due to increased waves.  It is also possible the we may choose to override on the simulation stage.
-			// Therefore we need to special case OneD and convert our element count back to linear.
-			if (DispatchType == ENiagaraGpuDispatchType::OneD)
+			if (SimStageData.StageMetaData->bGpuIndirectDispatch)
 			{
-				DispatchCount.X = DispatchCount.X * DispatchCount.Y * DispatchCount.Z;
-				DispatchCount.Y = 1;
-				DispatchCount.Z = 1;
+				if (SimStageData.DispatchArgs.IndirectBuffer == nullptr)
+				{
+					return;
+				}
+			}
+			else
+			{
+				DispatchCount = SimStageData.DispatchArgs.ElementCount;
+			
+				// Verify the number of elements isn't higher that what we can handle
+				checkf(uint64(DispatchCount.X) * uint64(DispatchCount.Y) * uint64(DispatchCount.Z) < uint64(TNumericLimits<int32>::Max()), TEXT("DispatchCount(%d, %d, %d) for IterationInterface(%s) overflows an int32 this is not allowed"), DispatchCount.X, DispatchCount.Y, DispatchCount.Z, *SimStageData.AlternateIterationSource->SourceDIName.ToString());
+
+				// Data interfaces such as grids / render targets can choose to dispatch in either the correct dimensionality for the target (i.e. RT2D would choose 2D)
+				// or run in linear mode if performance is not beneficial due to increased waves.  It is also possible the we may choose to override on the simulation stage.
+				// Therefore we need to special case OneD and convert our element count back to linear.
+				if (DispatchType == ENiagaraGpuDispatchType::OneD)
+				{
+					DispatchCount.X = DispatchCount.X * DispatchCount.Y * DispatchCount.Z;
+					DispatchCount.Y = 1;
+					DispatchCount.Z = 1;
+				}
 			}
 			break;
 		}
@@ -1422,13 +1438,13 @@ void FNiagaraGpuComputeDispatch::DispatchStage(FRDGBuilder& GraphBuilder, const
 			switch (DispatchType)
 			{
 				case ENiagaraGpuDispatchType::OneD:
-					DispatchCount = FIntVector(SimStageData.ElementCountXYZ.X, 1, 1);
+					DispatchCount = FIntVector(SimStageData.DispatchArgs.ElementCount.X, 1, 1);
 					break;
 				case ENiagaraGpuDispatchType::TwoD:
-					DispatchCount = FIntVector(SimStageData.ElementCountXYZ.X, SimStageData.ElementCountXYZ.Y, 1);
+					DispatchCount = FIntVector(SimStageData.DispatchArgs.ElementCount.X, SimStageData.DispatchArgs.ElementCount.Y, 1);
 					break;
 				case ENiagaraGpuDispatchType::ThreeD:
-					DispatchCount = SimStageData.ElementCountXYZ;
+					DispatchCount = SimStageData.DispatchArgs.ElementCount;
 					break;
 				default:
 					UE_LOG(LogNiagara, Fatal, TEXT("FNiagaraGpuComputeDispatch: Unknown DispatchType(%d)"), DispatchType);
@@ -1449,45 +1465,13 @@ void FNiagaraGpuComputeDispatch::DispatchStage(FRDGBuilder& GraphBuilder, const
 	}
 
 	const int32 TotalDispatchCount = DispatchCount.X * DispatchCount.Y * DispatchCount.Z;
-	if (TotalDispatchCount <= 0)
+	if (TotalDispatchCount <= 0 && SimStageData.DispatchArgs.IndirectBuffer == nullptr)
 	{
 		return;
 	}
 
 	checkf(DispatchNumThreads.X * DispatchNumThreads.Y * DispatchNumThreads.Z > 0, TEXT("DispatchNumThreads(%d, %d, %d) is invalid"), DispatchNumThreads.X, DispatchNumThreads.Y, DispatchNumThreads.Z);
 
-	// Calculate thread groups
-	// Note: In the OneD case we can use the Y dimension to get higher particle counts
-	if (DispatchType == ENiagaraGpuDispatchType::OneD)
-	{
-		const int32 GroupCount = FMath::DivideAndRoundUp(DispatchCount.X, DispatchNumThreads.X);
-		if (GroupCount > GRHIMaxDispatchThreadGroupsPerDimension.X)
-		{
-			DispatchCount.Y = FMath::DivideAndRoundUp(GroupCount, GRHIMaxDispatchThreadGroupsPerDimension.X);
-			DispatchCount.X = FMath::DivideAndRoundUp(GroupCount, DispatchCount.Y) * DispatchNumThreads.X;
-			ensure(DispatchCount.Y <= GRHIMaxDispatchThreadGroupsPerDimension.Y);
-		}
-	}
-
-	const FIntVector ThreadGroupCount(
-		FMath::DivideAndRoundUp(DispatchCount.X, DispatchNumThreads.X),
-		FMath::DivideAndRoundUp(DispatchCount.Y, DispatchNumThreads.Y),
-		FMath::DivideAndRoundUp(DispatchCount.Z, DispatchNumThreads.Z)
-	);
-
-	// Validate we don't overflow ThreadGroupCounts
-	if ((ThreadGroupCount.X < 0) || (ThreadGroupCount.X > GRHIMaxDispatchThreadGroupsPerDimension.X) ||
-		(ThreadGroupCount.Y < 0) || (ThreadGroupCount.Y > GRHIMaxDispatchThreadGroupsPerDimension.Y) ||
-		(ThreadGroupCount.Z < 0) || (ThreadGroupCount.Z > GRHIMaxDispatchThreadGroupsPerDimension.Z) )
-	{
-		UE_LOG(LogNiagara, Warning, TEXT("FNiagaraGpuComputeDispatch: Invalid ThreadGroupdCount(%d, %d, %d) for DispatchCount(%d, %d, %d) Stage (%s)"),
-			ThreadGroupCount.X, ThreadGroupCount.Y, ThreadGroupCount.Z,
-			DispatchCount.X, DispatchCount.Y, DispatchCount.Z,
-			*SimStageData.StageMetaData->SimulationStageName.ToString()
-		);
-		return;
-	}
-
 	// Set Parameters
 	const FNiagaraShaderScriptParametersMetadata& NiagaraShaderParametersMetadata = InstanceData.Context->GPUScript_RT->GetScriptParametersMetadata_RT();
 	const FShaderParametersMetadata* ShaderParametersMetadata = NiagaraShaderParametersMetadata.ShaderParametersMetadata.Get();
@@ -1512,8 +1496,16 @@ void FNiagaraGpuComputeDispatch::DispatchStage(FRDGBuilder& GraphBuilder, const
 
 	// Setup instance counts
 	DispatchParameters->RWInstanceCounts			= GPUInstanceCounterManager.GetInstanceCountBuffer().UAV;
-	DispatchParameters->ReadInstanceCountOffset		= SimStageData.AlternateIterationSource ? INDEX_NONE : SimStageData.SourceCountOffset;
-	DispatchParameters->WriteInstanceCountOffset	= SimStageData.AlternateIterationSource ? INDEX_NONE : SimStageData.DestinationCountOffset;
+	if (SimStageData.StageMetaData->IterationSourceType == ENiagaraIterationSource::Particles)
+	{
+		DispatchParameters->ReadInstanceCountOffset = SimStageData.SourceCountOffset;
+		DispatchParameters->WriteInstanceCountOffset = SimStageData.DestinationCountOffset;
+	}
+	else
+	{
+		DispatchParameters->ReadInstanceCountOffset = INDEX_NONE;
+		DispatchParameters->WriteInstanceCountOffset = INDEX_NONE;
+	}
 
 	// Simulation Stage Information
 	// X = Count Buffer Instance Count Offset (INDEX_NONE == Use Instance Count)
@@ -1529,9 +1521,9 @@ void FNiagaraGpuComputeDispatch::DispatchStage(FRDGBuilder& GraphBuilder, const
 				break;
 
 			case ENiagaraIterationSource::DataInterface:
-				if (SimStageData.AlternateIterationSource != nullptr)
+				if (SimStageData.AlternateIterationSource != nullptr)	
 				{
-					const uint32 IterationInstanceCountOffset = SimStageData.AlternateIterationSource->GetGPUInstanceCountOffset(Tick.SystemInstanceID);
+					const uint32 IterationInstanceCountOffset = SimStageData.DispatchArgs.GpuElementCountOffset;
 					DispatchParameters->SimulationStageIterationInfo.X = IterationInstanceCountOffset;
 					DispatchParameters->SimulationStageIterationInfo.Y = IterationInstanceCountOffset == INDEX_NONE ? TotalDispatchCount : 0;
 				}
@@ -1612,8 +1604,86 @@ void FNiagaraGpuComputeDispatch::DispatchStage(FRDGBuilder& GraphBuilder, const
 	DispatchParameters->SceneTextures.SceneTextures			= SceneTexturesUniformParams;
 	DispatchParameters->SceneTextures.MobileSceneTextures	= MobileSceneTexturesUniformParams;
 
-	// Execute the dispatch
+	// Indirect Gpu Dispatch
+	if (SimStageData.StageMetaData->bGpuIndirectDispatch)
 	{
+		DispatchParameters->IndirectDispatchArgsBuffer = SimStageData.DispatchArgs.IndirectBuffer;
+		DispatchParameters->IndirectDispatchArgs = GraphBuilder.CreateSRV(SimStageData.DispatchArgs.IndirectBuffer, EPixelFormat::PF_R32G32B32A32_UINT);
+		DispatchParameters->IndirectDispatchArgsOffset = (SimStageData.DispatchArgs.IndirectOffset / sizeof(FUintVector4)) + 1;
+
+		GraphBuilder.AddPass(
+			RDG_EVENT_NAME(
+				"NiagaraGpuSim(%s) Indirect Stage(%s %u) Iteration(%u) NumThreads(%dx%dx%d)",
+				InstanceData.Context->GetDebugSimName(),
+				*SimStageData.StageMetaData->SimulationStageName.ToString(),
+				SimStageData.StageIndex,
+				SimStageData.IterationIndex,
+				DispatchNumThreads.X, DispatchNumThreads.Y, DispatchNumThreads.Z
+			),
+			ShaderParametersMetadata,
+			DispatchParameters,
+			ERDGPassFlags::Compute | ERDGPassFlags::NeverCull,
+			[this, ShaderParametersMetadata, DispatchParameters, ComputeShader, TickPtr=&Tick, InstanceDataPtr=&InstanceData, SimStageDataPtr=&SimStageData](FRHICommandListImmediate& RHICmdList)		//-TODO:RDG: When legacy is removed this can be FRHIComputeCommandList
+			{
+				FRHIComputeShader* RHIComputeShader = ComputeShader.GetComputeShader();
+				SetComputePipelineState(RHICmdList, RHIComputeShader);
+
+				if (ComputeShader->ExternalConstantBufferParam[0].IsBound())
+				{
+					RHICmdList.SetShaderUniformBuffer(RHIComputeShader, ComputeShader->ExternalConstantBufferParam[0].GetBaseIndex(), TickPtr->GetExternalUniformBuffer(*InstanceDataPtr, false));
+				}
+				if (ComputeShader->ExternalConstantBufferParam[1].IsBound())
+				{
+					check(InstanceDataPtr->Context->HasInterpolationParameters);
+					RHICmdList.SetShaderUniformBuffer(RHIComputeShader, ComputeShader->ExternalConstantBufferParam[1].GetBaseIndex(), TickPtr->GetExternalUniformBuffer(*InstanceDataPtr, true));
+				}
+
+				FNiagaraEmptyUAVPoolScopedAccess UAVPoolAccessScope(GetEmptyUAVPool());
+
+				FNiagaraGpuProfileScope GpuProfileDispatchScope(RHICmdList, this, FNiagaraGpuProfileEvent(*InstanceDataPtr, *SimStageDataPtr, InstanceDataPtr == &TickPtr->GetInstances()[0]));
+
+				FComputeShaderUtils::ValidateIndirectArgsBuffer(SimStageDataPtr->DispatchArgs.IndirectBuffer->GetSize(), SimStageDataPtr->DispatchArgs.IndirectOffset, SimStageDataPtr->DispatchArgs.IndirectBuffer->GetStride());
+				SetShaderParameters<FRHICommandList, FNiagaraShader>(RHICmdList, ComputeShader, RHIComputeShader, ShaderParametersMetadata, *DispatchParameters);
+				RHICmdList.DispatchIndirectComputeShader(SimStageDataPtr->DispatchArgs.IndirectBuffer->GetIndirectRHICallBuffer(), SimStageDataPtr->DispatchArgs.IndirectOffset);
+				UnsetShaderUAVs<FRHICommandList, FNiagaraShader>(RHICmdList, ComputeShader, RHIComputeShader);
+			}
+		);
+	}
+	// Direct Gpu Dispatch
+	else
+	{
+		// Calculate thread groups
+		// Note: In the OneD case we can use the Y dimension to get higher particle counts
+		if (DispatchType == ENiagaraGpuDispatchType::OneD)
+		{
+			const int32 GroupCount = FMath::DivideAndRoundUp(DispatchCount.X, DispatchNumThreads.X);
+			if (GroupCount > GRHIMaxDispatchThreadGroupsPerDimension.X)
+			{
+				DispatchCount.Y = FMath::DivideAndRoundUp(GroupCount, GRHIMaxDispatchThreadGroupsPerDimension.X);
+				DispatchCount.X = FMath::DivideAndRoundUp(GroupCount, DispatchCount.Y) * DispatchNumThreads.X;
+				ensure(DispatchCount.Y <= GRHIMaxDispatchThreadGroupsPerDimension.Y);
+			}
+		}
+
+		const FIntVector ThreadGroupCount(
+			FMath::DivideAndRoundUp(DispatchCount.X, DispatchNumThreads.X),
+			FMath::DivideAndRoundUp(DispatchCount.Y, DispatchNumThreads.Y),
+			FMath::DivideAndRoundUp(DispatchCount.Z, DispatchNumThreads.Z)
+		);
+
+		// Validate we don't overflow ThreadGroupCounts
+		if ((ThreadGroupCount.X < 0) || (ThreadGroupCount.X > GRHIMaxDispatchThreadGroupsPerDimension.X) ||
+			(ThreadGroupCount.Y < 0) || (ThreadGroupCount.Y > GRHIMaxDispatchThreadGroupsPerDimension.Y) ||
+			(ThreadGroupCount.Z < 0) || (ThreadGroupCount.Z > GRHIMaxDispatchThreadGroupsPerDimension.Z) )
+		{
+			UE_LOG(LogNiagara, Warning, TEXT("FNiagaraGpuComputeDispatch: Invalid ThreadGroupdCount(%d, %d, %d) for ElementCount(%d, %d, %d) Stage (%s)"),
+				ThreadGroupCount.X, ThreadGroupCount.Y, ThreadGroupCount.Z,
+				DispatchCount.X, DispatchCount.Y, DispatchCount.Z,
+				*SimStageData.StageMetaData->SimulationStageName.ToString()
+			);
+			return;
+		}
+
 		DispatchParameters->DispatchThreadIdToLinear	= FUintVector3(1, DispatchCount.X, DispatchCount.X * DispatchCount.Y);
 		DispatchParameters->DispatchThreadIdBounds		= FUintVector3(DispatchCount.X, DispatchCount.Y, DispatchCount.Z);
 
@@ -1665,18 +1735,18 @@ void FNiagaraGpuComputeDispatch::DispatchStage(FRDGBuilder& GraphBuilder, const
 				UnsetShaderUAVs<FRHICommandList, FNiagaraShader>(RHICmdList, ComputeShader, RHIComputeShader);
 			}
 		);
+	}
 
-		INC_DWORD_STAT(STAT_NiagaraGPUDispatches);
-	#if CSV_PROFILER && WITH_PER_SYSTEM_PARTICLE_PERF_STATS && WITH_NIAGARA_DEBUG_EMITTER_NAME
-		if (FNiagaraGpuComputeDispatchLocal::CsvStatsEnabled())
+	INC_DWORD_STAT(STAT_NiagaraGPUDispatches);
+#if CSV_PROFILER && WITH_PER_SYSTEM_PARTICLE_PERF_STATS && WITH_NIAGARA_DEBUG_EMITTER_NAME
+	if (FNiagaraGpuComputeDispatchLocal::CsvStatsEnabled())
+	{
+		if (FCsvProfiler* CSVProfiler = FCsvProfiler::Get())
 		{
-			if (FCsvProfiler* CSVProfiler = FCsvProfiler::Get())
-			{
-				CSVProfiler->RecordCustomStat(InstanceData.Context->GetDebugSimFName(), CSV_CATEGORY_INDEX(NiagaraGpuCompute), 1, ECsvCustomStatOp::Accumulate);
-			}
+			CSVProfiler->RecordCustomStat(InstanceData.Context->GetDebugSimFName(), CSV_CATEGORY_INDEX(NiagaraGpuCompute), 1, ECsvCustomStatOp::Accumulate);
 		}
-	#endif
 	}
+#endif
 
 	// Optionally submit commands to the GPU
 	// This can be used to avoid accidental TDR detection in the editor especially when issuing multiple ticks in the same frame
diff --git a/Engine/Plugins/FX/Niagara/Source/Niagara/Public/NiagaraSimStageData.h b/Engine/Plugins/FX/Niagara/Source/Niagara/Public/NiagaraSimStageData.h
--- a/Engine/Plugins/FX/Niagara/Source/Niagara/Public/NiagaraSimStageData.h
+++ b/Engine/Plugins/FX/Niagara/Source/Niagara/Public/NiagaraSimStageData.h
@@ -6,6 +6,19 @@
 #include "Math/UnrealMath.h"
 #include "Misc/CoreMiscDefines.h"
 
+class FRDGBuffer;
+
+struct FNiagaraSimStageDispatchArgs
+{
+	// Dispatch Element Count and optional Gpu Element Count to clip OOB threads
+	FIntVector3	ElementCount = FIntVector3::ZeroValue;
+	uint32		GpuElementCountOffset = INDEX_NONE;
+
+	// Indirect dispatch buffer and the byte offset into the buffer
+	FRDGBuffer*	IndirectBuffer = nullptr;
+	uint32		IndirectOffset = 0;
+};
+
 struct FNiagaraSimStageData
 {
 	FNiagaraSimStageData()
@@ -21,7 +34,7 @@ struct FNiagaraSimStageData
 
 	uint32 StageIndex = INDEX_NONE;
 	uint32 IterationIndex = 0;
-	FIntVector ElementCountXYZ = FIntVector::NoneValue;
+	FNiagaraSimStageDispatchArgs DispatchArgs;
 
 	class FNiagaraDataBuffer* Source = nullptr;
 	uint32 SourceCountOffset = INDEX_NONE;
diff --git a/Engine/Plugins/FX/Niagara/Source/NiagaraEditor/Private/NiagaraHlslTranslator.cpp b/Engine/Plugins/FX/Niagara/Source/NiagaraEditor/Private/NiagaraHlslTranslator.cpp
--- a/Engine/Plugins/FX/Niagara/Source/NiagaraEditor/Private/NiagaraHlslTranslator.cpp
+++ b/Engine/Plugins/FX/Niagara/Source/NiagaraEditor/Private/NiagaraHlslTranslator.cpp
@@ -1390,6 +1390,7 @@ const FNiagaraTranslateResults &FHlslNiagaraTranslator::Translate(const FNiagara
 					// Determine dispatch information from iteration source (if we have one)
 					SimulationStageMetaData.GpuDispatchType = ENiagaraGpuDispatchType::OneD;
 					SimulationStageMetaData.GpuDispatchNumThreads = FNiagaraShader::GetDefaultThreadGroupSize(ENiagaraGpuDispatchType::OneD);
+					SimulationStageMetaData.bGpuIndirectDispatch = false;
 					if ( SimulationStageMetaData.IterationSourceType == ENiagaraIterationSource::DataInterface )
 					{
 						if (SimulationStageMetaData.IterationDataInterface.IsNone())
@@ -1403,6 +1404,7 @@ const FNiagaraTranslateResults &FHlslNiagaraTranslator::Translate(const FNiagara
 							if ( UNiagaraDataInterface* IteratinoSourceCDO = CompileDuplicateData->GetDuplicatedDataInterfaceCDOForClass(IterationSourceVar->GetType().GetClass()) )
 							{
 								SimulationStageMetaData.GpuDispatchType = TranslationStages[Index].bGpuDispatchForceLinear ? ENiagaraGpuDispatchType::OneD : IteratinoSourceCDO->GetGpuDispatchType();
+								SimulationStageMetaData.bGpuIndirectDispatch = IteratinoSourceCDO->GetGpuUseIndirectDispatch();
 
 								if (TranslationStages[Index].bOverrideGpuDispatchNumThreads)
 								{
@@ -3396,13 +3398,8 @@ void FHlslNiagaraTranslator::DefineMainGPUFunctions(
 			"	GGroupId			= InGroupId;\n"
 			"	GGroupThreadId		= InGroupThreadId;\n"
 			"	GGroupIndex			= InGroupIndex;\n"
-			"	GLinearThreadId		= GDispatchThreadId.x + (GDispatchThreadId.y * DispatchThreadIdToLinear.y);\n"
-			"	#if NIAGARA_DISPATCH_TYPE >= NIAGARA_DISPATCH_TYPE_THREE_D\n"
-			"		GLinearThreadId += GDispatchThreadId.z * DispatchThreadIdToLinear.z;\n"
-			"	#endif\n"
 			"	GEmitterTickCounter = EmitterTickCounter;\n"
 			"	GRandomSeedOffset = 0;\n"
-			"\n"
 		);
 
 	//////////////////////////////////////////////////////////////////////////
@@ -3421,6 +3418,26 @@ void FHlslNiagaraTranslator::DefineMainGPUFunctions(
 			*TranslationStage.PassNamespace
 		);
 
+		// Indirect dispatches need to setup the bounds and translation into linear
+		HlslOutput += TEXT(
+			"#if NIAGARA_DISPATCH_INDIRECT\n"
+			"	DispatchThreadIdBounds = IndirectDispatchArgs[IndirectDispatchArgsOffset].xyz;\n"
+			"	GLinearThreadId = GDispatchThreadId.x;\n"
+			"	#if NIAGARA_DISPATCH_TYPE >= NIAGARA_DISPATCH_TYPE_TWO_D\n"
+			"		GLinearThreadId += GDispatchThreadId.y * DispatchThreadIdBounds.x;\n"
+			"	#endif\n"
+			"	#if NIAGARA_DISPATCH_TYPE >= NIAGARA_DISPATCH_TYPE_THREE_D\n"
+			"		GLinearThreadId += GDispatchThreadId.z * DispatchThreadIdBounds.x * DispatchThreadIdBounds.y;\n"
+			"	#endif\n"
+			"#else //NIAGARA_DISPATCH_INDIRECT\n"
+			"	GLinearThreadId		= GDispatchThreadId.x + (GDispatchThreadId.y * DispatchThreadIdToLinear.y);\n"
+			"	#if NIAGARA_DISPATCH_TYPE >= NIAGARA_DISPATCH_TYPE_THREE_D\n"
+			"		GLinearThreadId += GDispatchThreadId.z * DispatchThreadIdToLinear.z;\n"
+			"	#endif\n"
+			"#endif //NIAGARA_DISPATCH_INDIRECT\n"
+			"\n"
+		);
+
 		// Particle iteration stage
 		if (TranslationStage.IterationSourceType == ENiagaraIterationSource::Particles)
 		{
@@ -3477,15 +3494,6 @@ void FHlslNiagaraTranslator::DefineMainGPUFunctions(
 		// Iteration Data Interface
 		else if (TranslationStage.IterationSourceType == ENiagaraIterationSource::DataInterface)
 		{
-			//	if (const FNiagaraVariable* IterationSourceVar = CompileData->EncounteredVariables.FindByPredicate([&](const FNiagaraVariable& VarInfo) { return VarInfo.GetName() == SimulationStageMetaData.IterationSource; }))
-			//	{
-			//		if (UNiagaraDataInterface* IteratinoSourceCDO = CompileDuplicateData->GetDuplicatedDataInterfaceCDOForClass(IterationSourceVar->GetType().GetClass()))
-			//		{
-			//			SimulationStageMetaData.GpuDispatchType = IteratinoSourceCDO->GetGpuDispatchType();
-			//			SimulationStageMetaData.GpuDispatchNumThreads = IteratinoSourceCDO->GetGpuDispatchNumThreads();
-			//		}
-			//	}
-
 			//-TODO: We can simplify the logic here with SimulationStage_GetInstanceCount() as only things that can provide an instance count offset
 			//		 can really be variable, everything else is driven from CPU code.
 			HlslOutput += TEXT(
diff --git a/Engine/Plugins/FX/Niagara/Source/NiagaraShader/Private/NiagaraScriptBase.cpp b/Engine/Plugins/FX/Niagara/Source/NiagaraShader/Private/NiagaraScriptBase.cpp
--- a/Engine/Plugins/FX/Niagara/Source/NiagaraShader/Private/NiagaraScriptBase.cpp
+++ b/Engine/Plugins/FX/Niagara/Source/NiagaraShader/Private/NiagaraScriptBase.cpp
@@ -13,6 +13,7 @@ FSimulationStageMetaData::FSimulationStageMetaData()
 	: bWritesParticles(0)
 	, bPartialParticleUpdate(0)
 	, bParticleIterationStateEnabled(0)
+	, bGpuIndirectDispatch(0)
 	, GpuDispatchType(ENiagaraGpuDispatchType::OneD)
 	, GpuDispatchNumThreads(0, 0, 0)
 {
diff --git a/Engine/Plugins/FX/Niagara/Source/NiagaraShader/Private/NiagaraShader.cpp b/Engine/Plugins/FX/Niagara/Source/NiagaraShader/Private/NiagaraShader.cpp
--- a/Engine/Plugins/FX/Niagara/Source/NiagaraShader/Private/NiagaraShader.cpp
+++ b/Engine/Plugins/FX/Niagara/Source/NiagaraShader/Private/NiagaraShader.cpp
@@ -354,6 +354,7 @@ void FNiagaraShaderType::BeginCompileShader(
 
 	const FIntVector ThreadGroupSize = SimStageMetaData.GpuDispatchNumThreads;
 	NewJob->Input.Environment.SetDefine(TEXT("NIAGARA_DISPATCH_TYPE"), int(SimStageMetaData.GpuDispatchType));
+	NewJob->Input.Environment.SetDefine(TEXT("NIAGARA_DISPATCH_INDIRECT"), SimStageMetaData.bGpuIndirectDispatch ? 1 : 0);
 	NewJob->Input.Environment.SetDefine(TEXT("THREADGROUP_SIZE"), ThreadGroupSize.X * ThreadGroupSize.Y * ThreadGroupSize.Z);
 	NewJob->Input.Environment.SetDefine(TEXT("THREADGROUP_SIZE_X"), ThreadGroupSize.X);
 	NewJob->Input.Environment.SetDefine(TEXT("THREADGROUP_SIZE_Y"), ThreadGroupSize.Y);
diff --git a/Engine/Plugins/FX/Niagara/Source/NiagaraShader/Public/NiagaraScriptBase.h b/Engine/Plugins/FX/Niagara/Source/NiagaraShader/Public/NiagaraScriptBase.h
--- a/Engine/Plugins/FX/Niagara/Source/NiagaraShader/Public/NiagaraScriptBase.h
+++ b/Engine/Plugins/FX/Niagara/Source/NiagaraShader/Public/NiagaraScriptBase.h
@@ -113,6 +113,9 @@ public:
 	UPROPERTY()
 	uint32 bParticleIterationStateEnabled : 1;
 
+	UPROPERTY()
+	uint32 bGpuIndirectDispatch : 1;
+
 	/** When the value is not none this is the binding used for particle state iteration stages. */
 	UPROPERTY()
 	FName ParticleIterationStateBinding;
diff --git a/Engine/Plugins/FX/Niagara/Source/NiagaraShader/Public/NiagaraShader.h b/Engine/Plugins/FX/Niagara/Source/NiagaraShader/Public/NiagaraShader.h
--- a/Engine/Plugins/FX/Niagara/Source/NiagaraShader/Public/NiagaraShader.h
+++ b/Engine/Plugins/FX/Niagara/Source/NiagaraShader/Public/NiagaraShader.h
@@ -188,6 +188,10 @@ public:
 		SHADER_PARAMETER(FUintVector3,								DispatchThreadIdBounds)
 		SHADER_PARAMETER(FUintVector3,								DispatchThreadIdToLinear)
 
+		SHADER_PARAMETER_RDG_BUFFER_SRV(Buffer<uint4>,				IndirectDispatchArgs)
+		SHADER_PARAMETER(uint32,									IndirectDispatchArgsOffset)
+		RDG_BUFFER_ACCESS(IndirectDispatchArgsBuffer,				ERHIAccess::IndirectArgs)
+
 		SHADER_PARAMETER_STRUCT_INCLUDE(FGlobalParameters,			GlobalParameters)
 		SHADER_PARAMETER_STRUCT_INCLUDE(FSystemParameters,			SystemParameters)
 		SHADER_PARAMETER_STRUCT_INCLUDE(FOwnerParameters,			OwnerParameters)
diff --git a/Engine/Plugins/FX/Niagara/Source/NiagaraVertexFactories/Private/NiagaraDispatchIndirect.cpp b/Engine/Plugins/FX/Niagara/Source/NiagaraVertexFactories/Private/NiagaraDispatchIndirect.cpp
new file mode 100644
--- /dev/null
+++ b/Engine/Plugins/FX/Niagara/Source/NiagaraVertexFactories/Private/NiagaraDispatchIndirect.cpp
@@ -0,0 +1,11 @@
+// Copyright Epic Games, Inc. All Rights Reserved.
+
+#include "NiagaraDispatchIndirect.h"
+
+IMPLEMENT_GLOBAL_SHADER(FNiagaraDispatchIndirectArgsGenCS, "/Plugin/FX/Niagara/Private/NiagaraDispatchIndirectArgsGen.usf", "IndirectArgsGenCS", SF_Compute);
+
+void FNiagaraDispatchIndirectArgsGenCS::ModifyCompilationEnvironment(const FGlobalShaderPermutationParameters& Parameters, FShaderCompilerEnvironment& OutEnvironment)
+{
+	FGlobalShader::ModifyCompilationEnvironment(Parameters, OutEnvironment);
+	OutEnvironment.SetDefine(TEXT("THREAD_COUNT"), ThreadCount);
+}
diff --git a/Engine/Plugins/FX/Niagara/Source/NiagaraVertexFactories/Public/NiagaraDispatchIndirect.h b/Engine/Plugins/FX/Niagara/Source/NiagaraVertexFactories/Public/NiagaraDispatchIndirect.h
new file mode 100644
--- /dev/null
+++ b/Engine/Plugins/FX/Niagara/Source/NiagaraVertexFactories/Public/NiagaraDispatchIndirect.h
@@ -0,0 +1,50 @@
+// Copyright Epic Games, Inc. All Rights Reserved.
+
+#pragma once
+
+#include "CoreMinimal.h"
+#include "GlobalShader.h"
+#include "ShaderPermutation.h"
+#include "ShaderParameterStruct.h"
+#include "ShaderParameterUtils.h"
+
+struct FNiagaraDispatchIndirectInfoCS
+{
+	FNiagaraDispatchIndirectInfoCS() = default;
+	explicit FNiagaraDispatchIndirectInfoCS(const FUintVector& InCounterOffsets, const FIntVector& InThreadGroupSize, uint32 InIndirectArgsOffset)
+		: CounterOffsets(InCounterOffsets)
+		, ThreadGroupSize(InThreadGroupSize)
+		, IndirectArgsOffset(InIndirectArgsOffset)
+	{
+	}
+
+	FUintVector		CounterOffsets = FUintVector(INDEX_NONE, INDEX_NONE, INDEX_NONE);
+	FIntVector		ThreadGroupSize = FIntVector::ZeroValue;
+	uint32			IndirectArgsOffset = 0;
+};
+
+struct FNiagaraDispatchIndirectParametersCS
+{
+	FUintVector4	GroupCount = FUintVector4(0, 0, 0, 0);
+	FUintVector4	ThreadCount = FUintVector4(0, 0, 0, 0);
+};
+
+class NIAGARAVERTEXFACTORIES_API FNiagaraDispatchIndirectArgsGenCS : public FGlobalShader
+{
+	DECLARE_GLOBAL_SHADER(FNiagaraDispatchIndirectArgsGenCS);
+	SHADER_USE_PARAMETER_STRUCT(FNiagaraDispatchIndirectArgsGenCS, FGlobalShader);
+
+public:
+	static constexpr int ThreadCount = 32;
+
+	BEGIN_SHADER_PARAMETER_STRUCT(FParameters, NIAGARAVERTEXFACTORIES_API)
+		SHADER_PARAMETER_RDG_BUFFER_SRV(Buffer<uint>,		DispatchInfos)
+		SHADER_PARAMETER(uint32,							NumDispatchInfos)
+		SHADER_PARAMETER(FUintVector3,						MaxGroupsPerDimension)
+
+		SHADER_PARAMETER_SRV(Buffer<uint>,					InstanceCounts)
+		SHADER_PARAMETER_RDG_BUFFER_UAV(RWBuffer<uint4>,	RWDispatchIndirectArgs)
+	END_SHADER_PARAMETER_STRUCT()
+
+	static void ModifyCompilationEnvironment(const FGlobalShaderPermutationParameters& Parameters, FShaderCompilerEnvironment& OutEnvironment);
+};
