- Iteration support for simple counter
- Requires indirect dispatch since we can't know the counter value when used with GPU

diff --git a/Engine/Plugins/FX/Niagara/Source/Niagara/Private/DataInterface/NiagaraDataInterfaceSimpleCounter.cpp b/Engine/Plugins/FX/Niagara/Source/Niagara/Private/DataInterface/NiagaraDataInterfaceSimpleCounter.cpp
--- a/Engine/Plugins/FX/Niagara/Source/Niagara/Private/DataInterface/NiagaraDataInterfaceSimpleCounter.cpp
+++ b/Engine/Plugins/FX/Niagara/Source/Niagara/Private/DataInterface/NiagaraDataInterfaceSimpleCounter.cpp
@@ -41,7 +41,7 @@ struct FNDISimpleCounterInstanceData_GameThread
 	std::atomic<int32>	Counter;
 };
 
-struct FNDISimpleCounterProxy : public FNiagaraDataInterfaceProxy
+struct FNDISimpleCounterProxy : public FNiagaraDataInterfaceProxyRW
 {
 	FNDISimpleCounterProxy(UNiagaraDataInterfaceSimpleCounter* Owner)
 		: WeakOwner(Owner)
@@ -52,6 +52,14 @@ struct FNDISimpleCounterProxy : public FNiagaraDataInterfaceProxy
 	virtual void ConsumePerInstanceDataFromGameThread(void* PerInstanceData, const FNiagaraSystemInstanceID& Instance) override {}
 	virtual int32 PerInstanceDataPassedToRenderThreadSize() const override { return 0; }
 
+	virtual void GetDispatchArgs(const FNDIGpuComputeDispatchArgsGenContext& Context) override
+	{
+		if (const FNDISimpleCounterInstanceData_RenderThread* InstanceData = PerInstanceData_RenderThread.Find(Context.GetSystemInstanceID()))
+		{
+			Context.CreateIndirect(InstanceData->CountOffset);
+		}
+	}
+
 	virtual void PreStage(const FNDIGpuComputePreStageContext& Context) override
 	{
 		if (FNDISimpleCounterInstanceData_RenderThread* InstanceData = PerInstanceData_RenderThread.Find(Context.GetSystemInstanceID()) )
diff --git a/Engine/Plugins/FX/Niagara/Source/Niagara/Private/DataInterface/NiagaraDataInterfaceSimpleCounter.h b/Engine/Plugins/FX/Niagara/Source/Niagara/Private/DataInterface/NiagaraDataInterfaceSimpleCounter.h
--- a/Engine/Plugins/FX/Niagara/Source/Niagara/Private/DataInterface/NiagaraDataInterfaceSimpleCounter.h
+++ b/Engine/Plugins/FX/Niagara/Source/Niagara/Private/DataInterface/NiagaraDataInterfaceSimpleCounter.h
@@ -4,7 +4,7 @@
 #include "NiagaraCommon.h"
 #include "NiagaraShared.h"
 #include "VectorVM.h"
-#include "NiagaraDataInterface.h"
+#include "NiagaraDataInterfaceRW.h"
 #include "HAL/PlatformAtomics.h"
 #include "NiagaraDataInterfaceSimpleCounter.generated.h"
 
@@ -13,7 +13,7 @@ Thread safe counter starts at the initial value on start / reset.
 When operating between CPU & GPU ensure you set the appropriate sync mode.
 */
 UCLASS(EditInlineNew, Category = "Counting", meta = (DisplayName = "Simple Counter"))
-class NIAGARA_API UNiagaraDataInterfaceSimpleCounter : public UNiagaraDataInterface
+class NIAGARA_API UNiagaraDataInterfaceSimpleCounter : public UNiagaraDataInterfaceRWBase
 {
 	GENERATED_UCLASS_BODY()
 
@@ -45,6 +45,9 @@ public:
 	virtual void GetParameterDefinitionHLSL(const FNiagaraDataInterfaceGPUParamInfo& ParamInfo, FString& OutHLSL) override;
 	virtual bool GetFunctionHLSL(const FNiagaraDataInterfaceGPUParamInfo& ParamInfo, const FNiagaraDataInterfaceGeneratedFunction& FunctionInfo, int FunctionInstanceIndex, FString& OutHLSL) override;
 	virtual bool AppendCompileHash(FNiagaraCompileHashVisitor* InVisitor) const override;
+#endif
+#if WITH_EDITOR
+	virtual bool GetGpuUseIndirectDispatch() const override { return true; }
 #endif
 	virtual void BuildShaderParameters(FNiagaraShaderParametersBuilder& ShaderParametersBuilder) const override;
 	virtual void SetShaderParameters(const FNiagaraDataInterfaceSetShaderParametersContext& Context) const override;
