- Replace GetElementCount with GetDispatchArgs to match newer functionality

diff --git a/Engine/Plugins/FX/Niagara/Source/Niagara/Classes/NiagaraDataInterfaceArrayImpl.h b/Engine/Plugins/FX/Niagara/Source/Niagara/Classes/NiagaraDataInterfaceArrayImpl.h
--- a/Engine/Plugins/FX/Niagara/Source/Niagara/Classes/NiagaraDataInterfaceArrayImpl.h
+++ b/Engine/Plugins/FX/Niagara/Source/Niagara/Classes/NiagaraDataInterfaceArrayImpl.h
@@ -472,22 +472,12 @@ struct FNDIArrayProxyImpl : public INDIArrayProxyBase
 		GameToRenderInstanceData->~FGameToRenderInstanceData();
 	}
 
-	virtual FIntVector GetElementCount(FNiagaraSystemInstanceID SystemInstanceID) const override
+	virtual void GetDispatchArgs(const FNDIGpuComputeDispatchArgsGenContext& Context) override
 	{
-		if (const FNDIArrayInstanceData_RenderThread<TArrayType>* InstanceData_RT = PerInstanceData_RenderThread.Find(SystemInstanceID))
+		if (const FNDIArrayInstanceData_RenderThread<TArrayType>* InstanceData_RT = PerInstanceData_RenderThread.Find(Context.GetSystemInstanceID()))
 		{
-			return FIntVector(InstanceData_RT->NumElements, 1, 1);
+			Context.SetDirect(InstanceData_RT->NumElements, InstanceData_RT->CountOffset);
 		}
-		return FIntVector::ZeroValue;
-	}
-
-	virtual uint32 GetGPUInstanceCountOffset(FNiagaraSystemInstanceID SystemInstanceID) const override
-	{
-		if (const FNDIArrayInstanceData_RenderThread<TArrayType>* InstanceData_RT = PerInstanceData_RenderThread.Find(SystemInstanceID))
-		{
-			return InstanceData_RT->CountOffset;
-		}
-		return INDEX_NONE;
 	}
 
 	virtual void PostSimulate(const FNDIGpuComputePostSimulateContext& Context) override
diff --git a/Engine/Plugins/FX/Niagara/Source/Niagara/Classes/NiagaraDataInterfaceGrid2DCollection.h b/Engine/Plugins/FX/Niagara/Source/Niagara/Classes/NiagaraDataInterfaceGrid2DCollection.h
--- a/Engine/Plugins/FX/Niagara/Source/Niagara/Classes/NiagaraDataInterfaceGrid2DCollection.h
+++ b/Engine/Plugins/FX/Niagara/Source/Niagara/Classes/NiagaraDataInterfaceGrid2DCollection.h
@@ -95,7 +95,7 @@ struct FNiagaraDataInterfaceProxyGrid2DCollectionProxy : public FNiagaraDataInte
 	virtual void PostStage(const FNDIGpuComputePostStageContext& Context) override;
 	virtual void PostSimulate(const FNDIGpuComputePostSimulateContext& Context) override;
 
-	virtual FIntVector GetElementCount(FNiagaraSystemInstanceID SystemInstanceID) const override;
+	virtual void GetDispatchArgs(const FNDIGpuComputeDispatchArgsGenContext& Context) override;
 
 	/* List of proxy data for each system instances*/
 	// #todo(dmp): this should all be refactored to avoid duplicate code
diff --git a/Engine/Plugins/FX/Niagara/Source/Niagara/Classes/NiagaraDataInterfaceGrid3DCollection.h b/Engine/Plugins/FX/Niagara/Source/Niagara/Classes/NiagaraDataInterfaceGrid3DCollection.h
--- a/Engine/Plugins/FX/Niagara/Source/Niagara/Classes/NiagaraDataInterfaceGrid3DCollection.h
+++ b/Engine/Plugins/FX/Niagara/Source/Niagara/Classes/NiagaraDataInterfaceGrid3DCollection.h
@@ -105,7 +105,7 @@ struct FNiagaraDataInterfaceProxyGrid3DCollectionProxy : public FNiagaraDataInte
 	virtual void PostStage(const FNDIGpuComputePostStageContext& Context) override;
 	virtual void PostSimulate(const FNDIGpuComputePostSimulateContext& Context) override;
 
-	virtual FIntVector GetElementCount(FNiagaraSystemInstanceID SystemInstanceID) const override;
+	virtual void GetDispatchArgs(const FNDIGpuComputeDispatchArgsGenContext& Context) override;
 
 	/* List of proxy data for each system instances*/
 	// #todo(dmp): this should all be refactored to avoid duplicate code
diff --git a/Engine/Plugins/FX/Niagara/Source/Niagara/Classes/NiagaraDataInterfaceNeighborGrid3D.h b/Engine/Plugins/FX/Niagara/Source/Niagara/Classes/NiagaraDataInterfaceNeighborGrid3D.h
--- a/Engine/Plugins/FX/Niagara/Source/Niagara/Classes/NiagaraDataInterfaceNeighborGrid3D.h
+++ b/Engine/Plugins/FX/Niagara/Source/Niagara/Classes/NiagaraDataInterfaceNeighborGrid3D.h
@@ -50,7 +50,7 @@ struct FNiagaraDataInterfaceProxyNeighborGrid3D : public FNiagaraDataInterfacePr
 	virtual void ConsumePerInstanceDataFromGameThread(void* PerInstanceData, const FNiagaraSystemInstanceID& Instance) override {}
 	virtual int32 PerInstanceDataPassedToRenderThreadSize() const override { return 0; }	
 
-	virtual FIntVector GetElementCount(FNiagaraSystemInstanceID SystemInstanceID) const override;
+	virtual void GetDispatchArgs(const FNDIGpuComputeDispatchArgsGenContext& Context) override;
 
 	/* List of proxy data for each system instances*/
 	// #todo(dmp): this should all be refactored to avoid duplicate code
diff --git a/Engine/Plugins/FX/Niagara/Source/Niagara/Classes/NiagaraDataInterfaceRasterizationGrid3D.h b/Engine/Plugins/FX/Niagara/Source/Niagara/Classes/NiagaraDataInterfaceRasterizationGrid3D.h
--- a/Engine/Plugins/FX/Niagara/Source/Niagara/Classes/NiagaraDataInterfaceRasterizationGrid3D.h
+++ b/Engine/Plugins/FX/Niagara/Source/Niagara/Classes/NiagaraDataInterfaceRasterizationGrid3D.h
@@ -45,7 +45,7 @@ struct FNiagaraDataInterfaceProxyRasterizationGrid3D : public FNiagaraDataInterf
 	virtual void ConsumePerInstanceDataFromGameThread(void* PerInstanceData, const FNiagaraSystemInstanceID& Instance) override {}
 	virtual int32 PerInstanceDataPassedToRenderThreadSize() const override { return sizeof(RasterizationGrid3DRWInstanceData); }	
 
-	virtual FIntVector GetElementCount(FNiagaraSystemInstanceID SystemInstanceID) const override;
+	virtual void GetDispatchArgs(const FNDIGpuComputeDispatchArgsGenContext& Context) override;
 
 	/* List of proxy data for each system instances*/
 	// #todo(dmp): this should all be refactored to avoid duplicate code
diff --git a/Engine/Plugins/FX/Niagara/Source/Niagara/Classes/NiagaraDataInterfaceRenderTarget2D.h b/Engine/Plugins/FX/Niagara/Source/Niagara/Classes/NiagaraDataInterfaceRenderTarget2D.h
--- a/Engine/Plugins/FX/Niagara/Source/Niagara/Classes/NiagaraDataInterfaceRenderTarget2D.h
+++ b/Engine/Plugins/FX/Niagara/Source/Niagara/Classes/NiagaraDataInterfaceRenderTarget2D.h
@@ -78,7 +78,7 @@ struct FNiagaraDataInterfaceProxyRenderTarget2DProxy : public FNiagaraDataInterf
 	virtual void PostStage(const FNDIGpuComputePostStageContext& Context) override;
 	virtual void PostSimulate(const FNDIGpuComputePostSimulateContext& Context) override;
 
-	virtual FIntVector GetElementCount(FNiagaraSystemInstanceID SystemInstanceID) const override;
+	virtual void GetDispatchArgs(const FNDIGpuComputeDispatchArgsGenContext& Context) override;
 
 	/* List of proxy data for each system instances*/
 	// #todo(dmp): this should all be refactored to avoid duplicate code
diff --git a/Engine/Plugins/FX/Niagara/Source/Niagara/Classes/NiagaraDataInterfaceRenderTarget2DArray.h b/Engine/Plugins/FX/Niagara/Source/Niagara/Classes/NiagaraDataInterfaceRenderTarget2DArray.h
--- a/Engine/Plugins/FX/Niagara/Source/Niagara/Classes/NiagaraDataInterfaceRenderTarget2DArray.h
+++ b/Engine/Plugins/FX/Niagara/Source/Niagara/Classes/NiagaraDataInterfaceRenderTarget2DArray.h
@@ -74,7 +74,7 @@ struct FNiagaraDataInterfaceProxyRenderTarget2DArrayProxy : public FNiagaraDataI
 
 	virtual void PostSimulate(const FNDIGpuComputePostSimulateContext& Context) override;
 
-	virtual FIntVector GetElementCount(FNiagaraSystemInstanceID SystemInstanceID) const override;
+	virtual void GetDispatchArgs(const FNDIGpuComputeDispatchArgsGenContext& Context) override;
 
 	/* List of proxy data for each system instances*/
 	// #todo(dmp): this should all be refactored to avoid duplicate code
diff --git a/Engine/Plugins/FX/Niagara/Source/Niagara/Classes/NiagaraDataInterfaceRenderTargetCube.h b/Engine/Plugins/FX/Niagara/Source/Niagara/Classes/NiagaraDataInterfaceRenderTargetCube.h
--- a/Engine/Plugins/FX/Niagara/Source/Niagara/Classes/NiagaraDataInterfaceRenderTargetCube.h
+++ b/Engine/Plugins/FX/Niagara/Source/Niagara/Classes/NiagaraDataInterfaceRenderTargetCube.h
@@ -67,7 +67,7 @@ struct FNiagaraDataInterfaceProxyRenderTargetCubeProxy : public FNiagaraDataInte
 	virtual int32 PerInstanceDataPassedToRenderThreadSize() const override { return 0; }
 	virtual void PostSimulate(const FNDIGpuComputePostSimulateContext& Context) override;
 
-	virtual FIntVector GetElementCount(FNiagaraSystemInstanceID SystemInstanceID) const override;
+	virtual void GetDispatchArgs(const FNDIGpuComputeDispatchArgsGenContext& Context) override;
 
 	/* List of proxy data for each system instances*/
 	// #todo(dmp): this should all be refactored to avoid duplicate code
diff --git a/Engine/Plugins/FX/Niagara/Source/Niagara/Classes/NiagaraDataInterfaceRenderTargetVolume.h b/Engine/Plugins/FX/Niagara/Source/Niagara/Classes/NiagaraDataInterfaceRenderTargetVolume.h
--- a/Engine/Plugins/FX/Niagara/Source/Niagara/Classes/NiagaraDataInterfaceRenderTargetVolume.h
+++ b/Engine/Plugins/FX/Niagara/Source/Niagara/Classes/NiagaraDataInterfaceRenderTargetVolume.h
@@ -67,7 +67,7 @@ struct FNiagaraDataInterfaceProxyRenderTargetVolumeProxy : public FNiagaraDataIn
 	virtual int32 PerInstanceDataPassedToRenderThreadSize() const override { return 0; }
 	virtual void PostSimulate(const FNDIGpuComputePostSimulateContext& Context) override;
 
-	virtual FIntVector GetElementCount(FNiagaraSystemInstanceID SystemInstanceID) const override;
+	virtual void GetDispatchArgs(const FNDIGpuComputeDispatchArgsGenContext& Context) override;
 
 	/* List of proxy data for each system instances*/
 	// #todo(dmp): this should all be refactored to avoid duplicate code
diff --git a/Engine/Plugins/FX/Niagara/Source/Niagara/Private/NiagaraDataInterfaceGrid2DCollection.cpp b/Engine/Plugins/FX/Niagara/Source/Niagara/Private/NiagaraDataInterfaceGrid2DCollection.cpp
--- a/Engine/Plugins/FX/Niagara/Source/Niagara/Private/NiagaraDataInterfaceGrid2DCollection.cpp
+++ b/Engine/Plugins/FX/Niagara/Source/Niagara/Private/NiagaraDataInterfaceGrid2DCollection.cpp
@@ -2694,23 +2694,22 @@ void FNiagaraDataInterfaceProxyGrid2DCollectionProxy::PostSimulate(const FNDIGpu
 	}
 }
 
-FIntVector FNiagaraDataInterfaceProxyGrid2DCollectionProxy::GetElementCount(FNiagaraSystemInstanceID SystemInstanceID) const
+void FNiagaraDataInterfaceProxyGrid2DCollectionProxy::GetDispatchArgs(const FNDIGpuComputeDispatchArgsGenContext& Context)
 {
-	if ( const FGrid2DCollectionRWInstanceData_RenderThread* ProxyData = SystemInstancesToProxyData_RT.Find(SystemInstanceID) )
+	if ( const FGrid2DCollectionRWInstanceData_RenderThread* ProxyData = SystemInstancesToProxyData_RT.Find(Context.GetSystemInstanceID()) )
 	{
 		// support a grid reader acting as an iteration source
 		if (ProxyData->OtherProxy != nullptr)
 		{
 			FNiagaraDataInterfaceProxyGrid2DCollectionProxy* OtherGrid2DProxy = static_cast<FNiagaraDataInterfaceProxyGrid2DCollectionProxy*>(ProxyData->OtherProxy);
-			const FGrid2DCollectionRWInstanceData_RenderThread* OtherProxyData = OtherGrid2DProxy->SystemInstancesToProxyData_RT.Find(SystemInstanceID);
-			return  FIntVector(OtherProxyData->NumCells.X, OtherProxyData->NumCells.Y, 1);
+			const FGrid2DCollectionRWInstanceData_RenderThread* OtherProxyData = OtherGrid2DProxy->SystemInstancesToProxyData_RT.Find(Context.GetSystemInstanceID());
+			Context.SetDirect(OtherProxyData->NumCells);
 		}
 		else
 		{
-			return FIntVector(ProxyData->NumCells.X, ProxyData->NumCells.Y, 1);
+			Context.SetDirect(ProxyData->NumCells);
 		}		
 	}
-	return FIntVector::ZeroValue;
 }
 
 #undef LOCTEXT_NAMESPACE
diff --git a/Engine/Plugins/FX/Niagara/Source/Niagara/Private/NiagaraDataInterfaceGrid3DCollection.cpp b/Engine/Plugins/FX/Niagara/Source/Niagara/Private/NiagaraDataInterfaceGrid3DCollection.cpp
--- a/Engine/Plugins/FX/Niagara/Source/Niagara/Private/NiagaraDataInterfaceGrid3DCollection.cpp
+++ b/Engine/Plugins/FX/Niagara/Source/Niagara/Private/NiagaraDataInterfaceGrid3DCollection.cpp
@@ -4222,23 +4222,22 @@ void FNiagaraDataInterfaceProxyGrid3DCollectionProxy::PostSimulate(const FNDIGpu
 	}
 }
 
-FIntVector FNiagaraDataInterfaceProxyGrid3DCollectionProxy::GetElementCount(FNiagaraSystemInstanceID SystemInstanceID) const
+void FNiagaraDataInterfaceProxyGrid3DCollectionProxy::GetDispatchArgs(const FNDIGpuComputeDispatchArgsGenContext& Context)
 {
-	if (const FGrid3DCollectionRWInstanceData_RenderThread* ProxyData = SystemInstancesToProxyData_RT.Find(SystemInstanceID))
+	if (const FGrid3DCollectionRWInstanceData_RenderThread* ProxyData = SystemInstancesToProxyData_RT.Find(Context.GetSystemInstanceID()))
 	{
 		// support a grid reader acting as an iteration source
 		if (ProxyData->OtherProxy != nullptr)
 		{
 			FNiagaraDataInterfaceProxyGrid3DCollectionProxy* OtherGrid3DProxy = static_cast<FNiagaraDataInterfaceProxyGrid3DCollectionProxy*>(ProxyData->OtherProxy);
-			const FGrid3DCollectionRWInstanceData_RenderThread* OtherProxyData = OtherGrid3DProxy->SystemInstancesToProxyData_RT.Find(SystemInstanceID);
-			return OtherProxyData->NumCells;
+			const FGrid3DCollectionRWInstanceData_RenderThread* OtherProxyData = OtherGrid3DProxy->SystemInstancesToProxyData_RT.Find(Context.GetSystemInstanceID());
+			Context.SetDirect(OtherProxyData->NumCells);
 		}
 		else
 		{
-			return ProxyData->NumCells;
+			Context.SetDirect(ProxyData->NumCells);
 		}
 	}
-	return FIntVector::ZeroValue;
 }
 
 void UNiagaraDataInterfaceGrid3DCollection::FindAttributes(TArray<FNiagaraVariableBase>& OutVariables, TArray<uint32>& OutVariableOffsets, int32& OutNumAttribChannelsFound, TArray<FText>* OutWarnings, bool UseReader) const
diff --git a/Engine/Plugins/FX/Niagara/Source/Niagara/Private/NiagaraDataInterfaceIntRenderTarget2D.cpp b/Engine/Plugins/FX/Niagara/Source/Niagara/Private/NiagaraDataInterfaceIntRenderTarget2D.cpp
--- a/Engine/Plugins/FX/Niagara/Source/Niagara/Private/NiagaraDataInterfaceIntRenderTarget2D.cpp
+++ b/Engine/Plugins/FX/Niagara/Source/Niagara/Private/NiagaraDataInterfaceIntRenderTarget2D.cpp
@@ -72,7 +72,7 @@ namespace NDIIntRenderTarget2DLocal
 
 struct FNDIIntRenderTarget2DInstanceData_GameThread
 {
-	FIntVector Size = FIntVector(EForceInit::ForceInitToZero);
+	FIntVector2 Size = FIntVector2(EForceInit::ForceInitToZero);
 	EPixelFormat Format = EPixelFormat::PF_R32_SINT;
 #if WITH_EDITORONLY_DATA
 	bool bPreviewRenderTarget = false;
@@ -87,7 +87,7 @@ struct FNDIIntRenderTarget2DInstanceData_GameThread
 
 struct FNDIIntRenderTarget2DInstanceData_RenderThread
 {
-	FIntVector Size = FIntVector(EForceInit::ForceInitToZero);
+	FIntVector2 Size = FIntVector2(EForceInit::ForceInitToZero);
 #if WITH_EDITORONLY_DATA
 	bool bPreviewRenderTarget = false;
 	FVector2D PreviewDisplayRange = FVector2D(0.0f, 255.0f);
@@ -154,13 +154,12 @@ struct FNDIIntRenderTarget2DProxy : public FNiagaraDataInterfaceProxyRW
 		}
 	}
 
-	virtual FIntVector GetElementCount(FNiagaraSystemInstanceID SystemInstanceID) const override
+	virtual void GetDispatchArgs(const FNDIGpuComputeDispatchArgsGenContext& Context) override
 	{
-		if (const FNDIIntRenderTarget2DInstanceData_RenderThread* InstanceData = SystemInstancesToProxyData_RT.Find(SystemInstanceID))
+		if (const FNDIIntRenderTarget2DInstanceData_RenderThread* InstanceData = SystemInstancesToProxyData_RT.Find(Context.GetSystemInstanceID()))
 		{
-			return FIntVector(InstanceData->Size.X, InstanceData->Size.Y, 1);
+			Context.SetDirect(InstanceData->Size);
 		}
-		return FIntVector::ZeroValue;
 	}
 
 	TMap<FNiagaraSystemInstanceID, FNDIIntRenderTarget2DInstanceData_RenderThread> SystemInstancesToProxyData_RT;
diff --git a/Engine/Plugins/FX/Niagara/Source/Niagara/Private/NiagaraDataInterfaceNeighborGrid3D.cpp b/Engine/Plugins/FX/Niagara/Source/Niagara/Private/NiagaraDataInterfaceNeighborGrid3D.cpp
--- a/Engine/Plugins/FX/Niagara/Source/Niagara/Private/NiagaraDataInterfaceNeighborGrid3D.cpp
+++ b/Engine/Plugins/FX/Niagara/Source/Niagara/Private/NiagaraDataInterfaceNeighborGrid3D.cpp
@@ -745,13 +745,12 @@ void FNiagaraDataInterfaceProxyNeighborGrid3D::PostSimulate(const FNDIGpuCompute
 	}
 }
 
-FIntVector FNiagaraDataInterfaceProxyNeighborGrid3D::GetElementCount(FNiagaraSystemInstanceID SystemInstanceID) const
+void FNiagaraDataInterfaceProxyNeighborGrid3D::GetDispatchArgs(const FNDIGpuComputeDispatchArgsGenContext& Context)
 {
-	if ( const FNDINeighborGrid3DInstanceData_RT* TargetData = SystemInstancesToProxyData_RT.Find(SystemInstanceID) )
+	if ( const FNDINeighborGrid3DInstanceData_RT* TargetData = SystemInstancesToProxyData_RT.Find(Context.GetSystemInstanceID()) )
 	{
-		return TargetData->NumCells;
+		Context.SetDirect(TargetData->NumCells);
 	}
-	return FIntVector::ZeroValue;
 }
 
 bool UNiagaraDataInterfaceNeighborGrid3D::CopyToInternal(UNiagaraDataInterface* Destination) const
diff --git a/Engine/Plugins/FX/Niagara/Source/Niagara/Private/NiagaraDataInterfaceParticleRead.cpp b/Engine/Plugins/FX/Niagara/Source/Niagara/Private/NiagaraDataInterfaceParticleRead.cpp
--- a/Engine/Plugins/FX/Niagara/Source/Niagara/Private/NiagaraDataInterfaceParticleRead.cpp
+++ b/Engine/Plugins/FX/Niagara/Source/Niagara/Private/NiagaraDataInterfaceParticleRead.cpp
@@ -257,9 +257,9 @@ struct FNiagaraDataInterfaceProxyParticleRead : public FNiagaraDataInterfaceProx
 		return SystemsRenderData.Find(InstanceID);
 	}
 
-	virtual FIntVector GetElementCount(FNiagaraSystemInstanceID SystemInstanceID) const override
+	virtual void GetDispatchArgs(const FNDIGpuComputeDispatchArgsGenContext& Context) override
 	{
-		const FNDIParticleRead_RenderInstanceData* InstanceData = SystemsRenderData.Find(SystemInstanceID);
+		const FNDIParticleRead_RenderInstanceData* InstanceData = SystemsRenderData.Find(Context.GetSystemInstanceID());
 		if (ensure(InstanceData))
 		{
 			if ( InstanceData->SourceEmitterGPUContext != nullptr )
@@ -268,32 +268,11 @@ struct FNiagaraDataInterfaceProxyParticleRead : public FNiagaraDataInterfaceProx
 				{
 					if ( const FNiagaraDataBuffer* CurrentData = SourceDataSet->GetCurrentData() )
 					{
-						return FIntVector(CurrentData->GetNumInstances(), 1, 1);
+						Context.SetDirect(CurrentData->GetNumInstances(), CurrentData->GetGPUInstanceCountBufferOffset());
 					}
 				}
 			}
 		}
-		return FIntVector::ZeroValue;
-	}
-
-	virtual uint32 GetGPUInstanceCountOffset(FNiagaraSystemInstanceID SystemInstanceID) const override
-	{
-		const FNDIParticleRead_RenderInstanceData* InstanceData = SystemsRenderData.Find(SystemInstanceID);
-		if (ensure(InstanceData))
-		{
-			if ( InstanceData->SourceEmitterGPUContext != nullptr )
-			{
-				if ( FNiagaraDataSet* SourceDataSet = InstanceData->SourceEmitterGPUContext->MainDataSet )
-				{
-					if ( const FNiagaraDataBuffer* CurrentData = SourceDataSet->GetCurrentData() )
-					{
-						return CurrentData->GetGPUInstanceCountBufferOffset();
-					}
-				}
-			}
-		}
-
-		return INDEX_NONE;
 	}
 
 private:
diff --git a/Engine/Plugins/FX/Niagara/Source/Niagara/Private/NiagaraDataInterfaceRasterizationGrid3D.cpp b/Engine/Plugins/FX/Niagara/Source/Niagara/Private/NiagaraDataInterfaceRasterizationGrid3D.cpp
--- a/Engine/Plugins/FX/Niagara/Source/Niagara/Private/NiagaraDataInterfaceRasterizationGrid3D.cpp
+++ b/Engine/Plugins/FX/Niagara/Source/Niagara/Private/NiagaraDataInterfaceRasterizationGrid3D.cpp
@@ -1189,13 +1189,12 @@ void FNiagaraDataInterfaceProxyRasterizationGrid3D::PostSimulate(const FNDIGpuCo
 	}
 }
 
-FIntVector FNiagaraDataInterfaceProxyRasterizationGrid3D::GetElementCount(FNiagaraSystemInstanceID SystemInstanceID) const
+void FNiagaraDataInterfaceProxyRasterizationGrid3D::GetDispatchArgs(const FNDIGpuComputeDispatchArgsGenContext& Context)
 {
-	if ( const RasterizationGrid3DRWInstanceData* TargetData = SystemInstancesToProxyData.Find(SystemInstanceID) )
+	if ( const RasterizationGrid3DRWInstanceData* TargetData = SystemInstancesToProxyData.Find(Context.GetSystemInstanceID()) )
 	{
-		return TargetData->NumCells;
+		Context.SetDirect(TargetData->NumCells);
 	}
-	return FIntVector::ZeroValue;
 }
 
 bool UNiagaraDataInterfaceRasterizationGrid3D::CopyToInternal(UNiagaraDataInterface* Destination) const
diff --git a/Engine/Plugins/FX/Niagara/Source/Niagara/Private/NiagaraDataInterfaceRenderTarget2D.cpp b/Engine/Plugins/FX/Niagara/Source/Niagara/Private/NiagaraDataInterfaceRenderTarget2D.cpp
--- a/Engine/Plugins/FX/Niagara/Source/Niagara/Private/NiagaraDataInterfaceRenderTarget2D.cpp
+++ b/Engine/Plugins/FX/Niagara/Source/Niagara/Private/NiagaraDataInterfaceRenderTarget2D.cpp
@@ -835,13 +835,12 @@ void FNiagaraDataInterfaceProxyRenderTarget2DProxy::PostSimulate(const FNDIGpuCo
 	ProxyData->TransientRDGUAV = nullptr;
 }
 
-FIntVector FNiagaraDataInterfaceProxyRenderTarget2DProxy::GetElementCount(FNiagaraSystemInstanceID SystemInstanceID) const
+void FNiagaraDataInterfaceProxyRenderTarget2DProxy::GetDispatchArgs(const FNDIGpuComputeDispatchArgsGenContext& Context)
 {
-	if ( const FRenderTarget2DRWInstanceData_RenderThread* TargetData = SystemInstancesToProxyData_RT.Find(SystemInstanceID) )
+	if ( const FRenderTarget2DRWInstanceData_RenderThread* TargetData = SystemInstancesToProxyData_RT.Find(Context.GetSystemInstanceID()) )
 	{
-		return FIntVector(TargetData->Size.X, TargetData->Size.Y, 1);
+		Context.SetDirect(TargetData->Size);
 	}
-	return FIntVector::ZeroValue;
 }
 
 #undef LOCTEXT_NAMESPACE
diff --git a/Engine/Plugins/FX/Niagara/Source/Niagara/Private/NiagaraDataInterfaceRenderTarget2DArray.cpp b/Engine/Plugins/FX/Niagara/Source/Niagara/Private/NiagaraDataInterfaceRenderTarget2DArray.cpp
--- a/Engine/Plugins/FX/Niagara/Source/Niagara/Private/NiagaraDataInterfaceRenderTarget2DArray.cpp
+++ b/Engine/Plugins/FX/Niagara/Source/Niagara/Private/NiagaraDataInterfaceRenderTarget2DArray.cpp
@@ -772,13 +772,12 @@ void FNiagaraDataInterfaceProxyRenderTarget2DArrayProxy::PostSimulate(const FNDI
 	ProxyData->TransientRDGUAV = nullptr;
 }
 
-FIntVector FNiagaraDataInterfaceProxyRenderTarget2DArrayProxy::GetElementCount(FNiagaraSystemInstanceID SystemInstanceID) const
+void FNiagaraDataInterfaceProxyRenderTarget2DArrayProxy::GetDispatchArgs(const FNDIGpuComputeDispatchArgsGenContext& Context) 
 {
-	if ( const FRenderTarget2DArrayRWInstanceData_RenderThread* TargetData = SystemInstancesToProxyData_RT.Find(SystemInstanceID) )
+	if ( const FRenderTarget2DArrayRWInstanceData_RenderThread* TargetData = SystemInstancesToProxyData_RT.Find(Context.GetSystemInstanceID()) )
 	{
-		return TargetData->Size;
+		Context.SetDirect(TargetData->Size);
 	}
-	return FIntVector::ZeroValue;
 }
 
 #undef LOCTEXT_NAMESPACE
diff --git a/Engine/Plugins/FX/Niagara/Source/Niagara/Private/NiagaraDataInterfaceRenderTargetCube.cpp b/Engine/Plugins/FX/Niagara/Source/Niagara/Private/NiagaraDataInterfaceRenderTargetCube.cpp
--- a/Engine/Plugins/FX/Niagara/Source/Niagara/Private/NiagaraDataInterfaceRenderTargetCube.cpp
+++ b/Engine/Plugins/FX/Niagara/Source/Niagara/Private/NiagaraDataInterfaceRenderTargetCube.cpp
@@ -723,13 +723,12 @@ void FNiagaraDataInterfaceProxyRenderTargetCubeProxy::PostSimulate(const FNDIGpu
 	ProxyData->TransientRDGUAV = nullptr;
 }
 
-FIntVector FNiagaraDataInterfaceProxyRenderTargetCubeProxy::GetElementCount(FNiagaraSystemInstanceID SystemInstanceID) const
+void FNiagaraDataInterfaceProxyRenderTargetCubeProxy::GetDispatchArgs(const FNDIGpuComputeDispatchArgsGenContext& Context)
 {
-	if ( const FRenderTargetCubeRWInstanceData_RenderThread* TargetData = SystemInstancesToProxyData_RT.Find(SystemInstanceID) )
+	if ( const FRenderTargetCubeRWInstanceData_RenderThread* TargetData = SystemInstancesToProxyData_RT.Find(Context.GetSystemInstanceID()) )
 	{
-		return FIntVector(TargetData->Size, TargetData->Size, 6);
+		Context.SetDirect(FIntVector(TargetData->Size, TargetData->Size, 6));
 	}
-	return FIntVector::ZeroValue;
 }
 
 #undef LOCTEXT_NAMESPACE
diff --git a/Engine/Plugins/FX/Niagara/Source/Niagara/Private/NiagaraDataInterfaceRenderTargetVolume.cpp b/Engine/Plugins/FX/Niagara/Source/Niagara/Private/NiagaraDataInterfaceRenderTargetVolume.cpp
--- a/Engine/Plugins/FX/Niagara/Source/Niagara/Private/NiagaraDataInterfaceRenderTargetVolume.cpp
+++ b/Engine/Plugins/FX/Niagara/Source/Niagara/Private/NiagaraDataInterfaceRenderTargetVolume.cpp
@@ -1091,13 +1091,12 @@ void FNiagaraDataInterfaceProxyRenderTargetVolumeProxy::PostSimulate(const FNDIG
 	ProxyData->TransientRDGUAV = nullptr;
 }
 
-FIntVector FNiagaraDataInterfaceProxyRenderTargetVolumeProxy::GetElementCount(FNiagaraSystemInstanceID SystemInstanceID) const
+void FNiagaraDataInterfaceProxyRenderTargetVolumeProxy::GetDispatchArgs(const FNDIGpuComputeDispatchArgsGenContext& Context)
 {
-	if ( const FRenderTargetVolumeRWInstanceData_RenderThread* TargetData = SystemInstancesToProxyData_RT.Find(SystemInstanceID) )
+	if ( const FRenderTargetVolumeRWInstanceData_RenderThread* TargetData = SystemInstancesToProxyData_RT.Find(Context.GetSystemInstanceID()) )
 	{
-		return TargetData->Size;
+		Context.SetDirect(TargetData->Size);
 	}
-	return FIntVector::ZeroValue;
 }
 
 #undef LOCTEXT_NAMESPACE
diff --git a/Engine/Plugins/Runtime/HairStrands/Source/HairStrandsCore/Private/Niagara/NiagaraDataInterfaceVelocityGrid.cpp b/Engine/Plugins/Runtime/HairStrands/Source/HairStrandsCore/Private/Niagara/NiagaraDataInterfaceVelocityGrid.cpp
--- a/Engine/Plugins/Runtime/HairStrands/Source/HairStrandsCore/Private/Niagara/NiagaraDataInterfaceVelocityGrid.cpp
+++ b/Engine/Plugins/Runtime/HairStrands/Source/HairStrandsCore/Private/Niagara/NiagaraDataInterfaceVelocityGrid.cpp
@@ -739,14 +739,12 @@ void FNDIVelocityGridProxy::ResetData(const FNDIGpuComputeResetContext& Context)
 	}
 }
 
-// Get the element count for this instance
-FIntVector FNDIVelocityGridProxy::GetElementCount(FNiagaraSystemInstanceID SystemInstanceID) const
+void FNDIVelocityGridProxy::GetDispatchArgs(const FNDIGpuComputeDispatchArgsGenContext& Context)
 {
-	if  ( const FNDIVelocityGridData* ProxyData = SystemInstancesToProxyData.Find(SystemInstanceID) )
+	if  ( const FNDIVelocityGridData* ProxyData = SystemInstancesToProxyData.Find(Context.GetSystemInstanceID()) )
 	{
-		return FIntVector(ProxyData->GridSize.X + 1, ProxyData->GridSize.Y + 1, ProxyData->GridSize.Z + 1);
+		Context.SetDirect(FIntVector(ProxyData->GridSize.X + 1, ProxyData->GridSize.Y + 1, ProxyData->GridSize.Z + 1));
 	}
-	return FIntVector::ZeroValue;
 }
 
 #undef LOCTEXT_NAMESPACE
diff --git a/Engine/Plugins/Runtime/HairStrands/Source/HairStrandsCore/Public/Niagara/NiagaraDataInterfaceVelocityGrid.h b/Engine/Plugins/Runtime/HairStrands/Source/HairStrandsCore/Public/Niagara/NiagaraDataInterfaceVelocityGrid.h
--- a/Engine/Plugins/Runtime/HairStrands/Source/HairStrandsCore/Public/Niagara/NiagaraDataInterfaceVelocityGrid.h
+++ b/Engine/Plugins/Runtime/HairStrands/Source/HairStrandsCore/Public/Niagara/NiagaraDataInterfaceVelocityGrid.h
@@ -152,8 +152,8 @@ struct FNDIVelocityGridProxy : public FNiagaraDataInterfaceProxyRW
 	/** Reset the buffers  */
 	virtual void ResetData(const FNDIGpuComputeResetContext& Context) override;
 
-	// Get the element count for this instance
-	virtual FIntVector GetElementCount(FNiagaraSystemInstanceID SystemInstanceID) const override;
+	// Get the dispatch arguments for this instance
+	virtual void GetDispatchArgs(const FNDIGpuComputeDispatchArgsGenContext& Context) override;
 
 	/** List of proxy data for each system instances*/
 	TMap<FNiagaraSystemInstanceID, FNDIVelocityGridData> SystemInstancesToProxyData;
