Updates NiagaraGraph to remove many unecessary references to UNiagaraScriptVaraible (most significantly those that are declared by child graphs).  Makes a dramatic improvement to the number of objects that get serialized with a system and improves load times.

diff --git a/Engine/Plugins/FX/Niagara/Source/NiagaraEditor/Private/NiagaraGraph.cpp b/Engine/Plugins/FX/Niagara/Source/NiagaraEditor/Private/NiagaraGraph.cpp
--- a/Engine/Plugins/FX/Niagara/Source/NiagaraEditor/Private/NiagaraGraph.cpp
+++ b/Engine/Plugins/FX/Niagara/Source/NiagaraEditor/Private/NiagaraGraph.cpp
@@ -65,8 +65,107 @@ static FAutoConsoleVariableRef CVarNiagaraUseGraphHash(
 	ECVF_Default
 );
 
-FNiagaraGraphParameterReferenceCollection::FNiagaraGraphParameterReferenceCollection(const bool bInCreated)
-	: Graph(nullptr), bCreatedByUser(bInCreated)
+namespace NiagaraGraphImpl
+{
+
+void FindReferencedVariables(const UNiagaraGraph* Graph, TConstArrayView<UNiagaraNode*> Nodes, TSet<FNiagaraVariable>& ReferencedVariables)
+{
+	// Check all pins on all nodes in the graph to find parameter pins which may have been missed in the parameter map traversal.  This
+	// can happen for nodes which are not fully connected
+	for (const UNiagaraNode* Node : Nodes)
+	{
+		if (const UNiagaraNodeStaticSwitch* SwitchNode = Cast<const UNiagaraNodeStaticSwitch>(Node))
+		{
+			if (!SwitchNode->IsSetByCompiler() && !SwitchNode->IsSetByPin())
+			{
+				ReferencedVariables.Add(FNiagaraVariable(SwitchNode->GetInputType(), SwitchNode->InputParameterName));
+			}
+		}
+		else if (const UNiagaraNodeFunctionCall* FunctionNode = Cast<const UNiagaraNodeFunctionCall>(Node))
+		{
+			for (const FNiagaraPropagatedVariable& Propagated : FunctionNode->PropagatedStaticSwitchParameters)
+			{
+				ReferencedVariables.Add(Propagated.ToVariable());
+			}
+		}
+		else if (const UNiagaraNodeParameterMapBase* ParameterMapNode = Cast<const UNiagaraNodeParameterMapBase>(Node))
+		{
+			const bool IsMapGet = Node->IsA<UNiagaraNodeParameterMapGet>();
+			const bool IsMapSet = Node->IsA<UNiagaraNodeParameterMapSet>();
+
+			for (UEdGraphPin* Pin : Node->Pins)
+			{
+				// we only consider parameter pins, inputs for MapSets and outputs for MapGets
+				if (Pin->PinType.PinSubCategory == UNiagaraNodeParameterMapBase::ParameterPinSubCategory)
+				{
+					const bool bRelevantPin = ((Pin->Direction == EEdGraphPinDirection::EGPD_Input) && IsMapSet)
+						|| ((Pin->Direction == EEdGraphPinDirection::EGPD_Output) && IsMapGet);
+
+					if (!bRelevantPin)
+					{
+						continue;
+					}
+
+					ReferencedVariables.Add(UEdGraphSchema_Niagara::PinToNiagaraVariable(Pin, false));
+				}
+			}
+		}
+	}
+
+	auto GetBindingScriptVariable = [&](const FNiagaraVariable& Variable) -> UNiagaraScriptVariable*
+	{
+		if (UNiagaraScriptVariable* ScriptVariable = Graph->GetScriptVariable(Variable))
+		{
+			if (ScriptVariable->DefaultMode == ENiagaraDefaultMode::Binding && ScriptVariable->DefaultBinding.IsValid())
+			{
+				return ScriptVariable;
+			}
+		}
+
+		return nullptr;
+	};
+
+	TArray<UNiagaraScriptVariable*> VariablesToProcess;
+	for (const FNiagaraVariable& ReferencedVariable : ReferencedVariables)
+	{
+		if (UNiagaraScriptVariable* ScriptVariable = GetBindingScriptVariable(ReferencedVariable))
+		{
+			VariablesToProcess.Add(ScriptVariable);
+		}
+	}
+
+	while (!VariablesToProcess.IsEmpty())
+	{
+		UNiagaraScriptVariable* Variable = VariablesToProcess.Pop(false);
+
+		FNiagaraTypeDefinition LinkedType = Variable->Variable.GetType();
+		FName BindingName = Variable->DefaultBinding.GetName();
+		if (FNiagaraConstants::GetOldPositionTypeVariables().Contains(FNiagaraVariable(LinkedType, BindingName)))
+		{
+			// it is not uncommon that old assets have vector inputs that default bind to what is now a position type. If we detect that, we change the type to prevent a compiler error.
+			LinkedType = FNiagaraTypeDefinition::GetPositionDef();
+		}
+
+		const FNiagaraVariable BoundVariable(LinkedType, BindingName);
+
+		bool bAlreadyInSet = false;
+		ReferencedVariables.Add(BoundVariable, &bAlreadyInSet);
+		if (!bAlreadyInSet)
+		{
+			if (UNiagaraScriptVariable* ScriptVariable = GetBindingScriptVariable(BoundVariable))
+			{
+				VariablesToProcess.Add(ScriptVariable);
+			}
+		}
+	}
+}
+
+}
+
+
+
+FNiagaraGraphParameterReferenceCollection::FNiagaraGraphParameterReferenceCollection(bool bInCreated)
+	: bCreatedByUser(bInCreated)
 {
 }
 
@@ -281,6 +380,68 @@ void UNiagaraGraph::PostLoad_LWCFixup(int32 NiagaraVersion)
 	}
 }
 
+// while it would be great to only run this fixup once, it's not clear how we are missing 'local' niagara script variables from our map
+// and so we'll keep this running all the time.  It just running through nodes rather than a parmaeter map traversal is still a huge
+// improvement, but we need to get to the bottom of where the missing script variables are coming from.
+void UNiagaraGraph::PostLoad_ManageScriptVariables(int32 NiagaraVersion)
+{
+	// we only want to manage script variables for modules/functions.  Top level scripts system/emitter/particle scripts don't
+	// need to manage their parameters via script variables
+	bool bHasStandaloneScriptUsage = false;
+
+	{
+		TArray<UNiagaraNodeOutput*> OutputNodes;
+		FindOutputNodes(OutputNodes);
+
+		auto ScriptUsagePredicate = [](const UNiagaraNodeOutput* OutputNode) -> bool
+		{
+			const ENiagaraScriptUsage Usage = OutputNode->GetUsage();
+
+			return (Usage == ENiagaraScriptUsage::Module ||
+				Usage == ENiagaraScriptUsage::DynamicInput ||
+				Usage == ENiagaraScriptUsage::Function);
+		};
+
+		bHasStandaloneScriptUsage = OutputNodes.ContainsByPredicate(ScriptUsagePredicate);
+	}
+
+	TSet<FNiagaraVariable> ReferencedVariables;
+	if (bHasStandaloneScriptUsage)
+	{
+		TArray<UNiagaraNode*> NiagaraNodes;
+		NiagaraNodes.Reserve(Nodes.Num());
+		for (TObjectPtr<UEdGraphNode>& Node : Nodes)
+		{
+			if (Node)
+			{
+				if (UNiagaraNode* NiagaraNode = Cast<UNiagaraNode>(Node.Get()))
+				{
+					NiagaraNodes.Add(NiagaraNode);
+				}
+			}
+		}
+		NiagaraGraphImpl::FindReferencedVariables(this, NiagaraNodes, ReferencedVariables);
+	}
+
+	// clear out any script variables that are not in the set of referenced variables above
+	for (FScriptVariableMap::TIterator ScriptVariableIt = VariableToScriptVariable.CreateIterator(); ScriptVariableIt; ++ScriptVariableIt)
+	{
+		if (!ReferencedVariables.Contains(ScriptVariableIt.Key()))
+		{
+			ScriptVariableIt.RemoveCurrent();
+		}
+	}
+
+	// and also ensure that we're not actually missing any script variables
+	for (const FNiagaraVariable& Variable : ReferencedVariables)
+	{
+		if (!VariableToScriptVariable.Contains(Variable))
+		{
+			CreateScriptVariableInternal(Variable, FNiagaraVariableMetaData(), false);
+		}
+	}
+}
+
 void UNiagaraGraph::PostLoad()
 {
 	LLM_SCOPE(ELLMTag::Niagara);
@@ -574,37 +735,11 @@ void UNiagaraGraph::PostLoad()
 		}
 	}
 
-	// Fix LWC position type mapping; before LWC, all position were FVector3f, but now that they are a separate type, we need to upgrade the pins and metadata to match
-	// NOTE - this function has the capacity to change the graph and so needs be done before we RefreshParameterReferences below (GetParameterReferenceMap)
-	// otherwise we'll be getting out of date results
+	// Fix LWC position type mapping; before LWC, all position were FVector3f, but now that they are a separate type, we
+	// need to upgrade the pins and metadata to match
 	PostLoad_LWCFixup(NiagaraVer);
 
-	check(bParameterReferenceRefreshPending);
-
-	{
-		// Create a UNiagaraScriptVariable instance for every entry in the parameter map for which there is no existing script variable. 
-		// Note that the call to GetParameterReferenceMap() will:
-		//	-generate a parameter map traversal (slow)
-		//	-remove/modify UNiagaraScriptVariables from VariableToScriptVariable
-		//
-		// Once GetParameterReferenceMap is called we need to make sure that we don't invalidate the ParameterReferenceMap.  Any changes
-		// to the graph through the postload needs to also fixup the parameter map to avoid us having to regenerate the map from scratch
-		// costing us a parameter map traversal
-		TArray<FNiagaraVariable> VarsToAdd;
-		for (auto& ParameterToReferences : GetParameterReferenceMap())
-		{
-			UNiagaraScriptVariable* Variable = GetScriptVariable(ParameterToReferences.Key.GetName());
-			if (Variable == nullptr)
-			{
-				VarsToAdd.Add(ParameterToReferences.Key);
-			}
-		}
-
-		for (FNiagaraVariable& Var : VarsToAdd)
-		{
-			AddParameter(Var, FNiagaraVariableMetaData(), false, false);
-		}
-	}
+	PostLoad_ManageScriptVariables(NiagaraVer);
 
 	// Fix inconsistencies in the default value declaration between graph and metadata
 	ValidateDefaultPins();
@@ -614,6 +749,12 @@ void UNiagaraGraph::PostLoad()
 	// generating the compile ids but there's still code out there that is expecting to have fully generated compile ids
 	// available on load.
 	RebuildCachedCompileIds();
+
+	// we want to make sure that nothing above resulted in the reference map from being initialized as it shouldn't be necessary to 
+	// do such heavy lifting during PostLoad()
+	ensureMsgf(bParameterReferenceRefreshPending,
+		TEXT("ParameterReferenceMap was built during PostLoad of %s - for load time reasons this should be cleaned up."),
+		*GetFullName());
 }
 
 #if WITH_EDITORONLY_DATA
@@ -957,7 +1098,7 @@ void UNiagaraGraph::ValidateDefaultPins()
 			}
 		}
 
-		for (auto& MetaData : GetAllMetaData())
+		for (const FScriptVariableMap::ElementType& MetaData : VariableToScriptVariable)
 		{
 			const FNiagaraVariable& Variable = MetaData.Key;
 			UNiagaraScriptVariable* ScriptVariable = MetaData.Value;
@@ -1335,7 +1476,7 @@ void UNiagaraGraph::StandardizeParameterNames()
 
 	// Since we'll be modifying the keys, make a copy of the map and then clear the original so it cal be 
 	// repopulated.
-	TMap<FNiagaraVariable, TObjectPtr<UNiagaraScriptVariable>> OldVariableToScriptVariable = VariableToScriptVariable;
+	FScriptVariableMap OldVariableToScriptVariable = VariableToScriptVariable;
 	VariableToScriptVariable.Empty();
 	for (TPair<FNiagaraVariable, TObjectPtr<UNiagaraScriptVariable>> VariableScriptVariablePair : OldVariableToScriptVariable)
 	{
@@ -1779,8 +1920,11 @@ UNiagaraGraph* UNiagaraGraph::CreateCompilationCopy(const TArray<ENiagaraScriptU
 
 	// probably not necessary for compilation, but remove references to the original graph
 	Result->ParameterToReferencesMap.Empty();
-	Result->RefreshParameterReferences();
-	
+
+	// by setting this flag to false we make sure that if something attempts to refresh the reference map that we'll
+	// hit the ensure at the top of RefreshParameterReferences()
+	Result->bParameterReferenceRefreshPending = true;
+
 	return Result;
 }
 
@@ -2080,67 +2224,42 @@ void UNiagaraGraph::GetParameters(TArray<FNiagaraVariable>& Inputs, TArray<FNiag
 // 	Outputs.Sort(SortVars);
 }
 
-const TMap<FNiagaraVariable, TObjectPtr<UNiagaraScriptVariable>>& UNiagaraGraph::GetAllMetaData() const
+const UNiagaraGraph::FScriptVariableMap& UNiagaraGraph::GetAllMetaData() const
 {
 	check(!bIsForCompilationOnly);
 
 	return VariableToScriptVariable;
 }
 
-TMap<FNiagaraVariable, TObjectPtr<UNiagaraScriptVariable>>& UNiagaraGraph::GetAllMetaData()
+UNiagaraGraph::FScriptVariableMap& UNiagaraGraph::GetAllMetaData()
 {
 	check(!bIsForCompilationOnly);
 
 	return VariableToScriptVariable;
 }
 
-const TMap<FNiagaraVariable, FNiagaraGraphParameterReferenceCollection>& UNiagaraGraph::GetParameterReferenceMap() const
+void UNiagaraGraph::ConditionalRefreshParameterReferences()
 {
 	if (bParameterReferenceRefreshPending)
 	{
 		RefreshParameterReferences();
 	}
-	return ParameterToReferencesMap;
 }
 
-UNiagaraScriptVariable* UNiagaraGraph::GetScriptVariable(FNiagaraVariable Parameter, bool bUpdateIfPending)
-{
-	check(!bIsForCompilationOnly);
-
-	if (bUpdateIfPending && bParameterReferenceRefreshPending)
-	{
-		RefreshParameterReferences();
-	}
-	if (TObjectPtr<UNiagaraScriptVariable>* FoundScriptVariable = VariableToScriptVariable.Find(Parameter))
-	{
-		return *FoundScriptVariable;
-	}
-	return nullptr;
-}
-
-UNiagaraScriptVariable* UNiagaraGraph::GetScriptVariable(FName ParameterName, bool bUpdateIfPending)
+const TMap<FNiagaraVariable, FNiagaraGraphParameterReferenceCollection>& UNiagaraGraph::GetParameterReferenceMap() const
 {
-	check(!bIsForCompilationOnly);
-
-	if (bUpdateIfPending && bParameterReferenceRefreshPending)
+	if (bParameterReferenceRefreshPending)
 	{
 		RefreshParameterReferences();
 	}
-	for (auto& VariableToScriptVariableItem : VariableToScriptVariable)
-	{
-		if (VariableToScriptVariableItem.Key.GetName() == ParameterName)
-		{
-			return VariableToScriptVariableItem.Value;
-		}
-	}
-	return nullptr;
+	return ParameterToReferencesMap;
 }
 
 UNiagaraScriptVariable* UNiagaraGraph::GetScriptVariable(FNiagaraVariable Parameter) const
 {
 	check(!bIsForCompilationOnly);
 
-	if (TObjectPtr<UNiagaraScriptVariable>* FoundScriptVariable = VariableToScriptVariable.Find(Parameter))
+	if (const TObjectPtr<UNiagaraScriptVariable>* FoundScriptVariable = VariableToScriptVariable.Find(Parameter))
 	{
 		return *FoundScriptVariable;
 	}
@@ -2180,32 +2299,16 @@ UNiagaraScriptVariable* UNiagaraGraph::AddParameter(const FNiagaraVariable& Para
 	if (!FoundParameterReferenceCollection)
 	{
 		const bool bCreatedByUser = !bIsStaticSwitch;
-		FNiagaraGraphParameterReferenceCollection NewReferenceCollection = FNiagaraGraphParameterReferenceCollection(bCreatedByUser);
-		NewReferenceCollection.Graph = this;
-		ParameterToReferencesMap.Add(Parameter, NewReferenceCollection);
+		ParameterToReferencesMap.Emplace(Parameter, {bCreatedByUser});
 	}
 
 	TObjectPtr<UNiagaraScriptVariable>* FoundScriptVariable = VariableToScriptVariable.Find(Parameter);
 	if (!FoundScriptVariable)
 	{
 		Modify();
-		UNiagaraScriptVariable* NewScriptVariable = NewObject<UNiagaraScriptVariable>(this, FName(), RF_Transactional);
-		NewScriptVariable->Init(Parameter, ParameterMetaData);
-		NewScriptVariable->SetIsStaticSwitch(bIsStaticSwitch);
 
-		FNiagaraEditorUtilities::ResetVariableToDefaultValue(NewScriptVariable->Variable);
-		
-		// Inputs in graphs (module namespace) are intrinsically written to, so set the default mode to value instead of fail if not set.
-		if (NewScriptVariable->Variable.IsInNameSpace(FNiagaraConstants::ModuleNamespace))
-		{
-			NewScriptVariable->DefaultMode = ENiagaraDefaultMode::Value;
-		}
-		else
-		{ 
-			NewScriptVariable->DefaultMode = ENiagaraDefaultMode::FailIfPreviouslyNotSet;
-		}
+		UNiagaraScriptVariable* NewScriptVariable = CreateScriptVariableInternal(Parameter, ParameterMetaData, bIsStaticSwitch);
 
-		VariableToScriptVariable.Add(Parameter, NewScriptVariable);
 		if (bNotifyChanged)
 		{
 			NotifyGraphChanged();
@@ -2233,9 +2336,7 @@ UNiagaraScriptVariable* UNiagaraGraph::AddParameter(const UNiagaraScriptVariable
 		{
 			NewScriptVariable->Metadata.CreateNewGuid();
 		}
-		FNiagaraGraphParameterReferenceCollection NewReferenceCollection = FNiagaraGraphParameterReferenceCollection(true /*bCreated*/);
-		NewReferenceCollection.Graph = this;
-		ParameterToReferencesMap.Add(NewScriptVariable->Variable, NewReferenceCollection);
+ 		ParameterToReferencesMap.Emplace(NewScriptVariable->Variable, {true /*bInCreated*/});
 		VariableToScriptVariable.Add(NewScriptVariable->Variable, NewScriptVariable);
 		NotifyGraphChanged();
 		return NewScriptVariable;
@@ -2244,6 +2345,29 @@ UNiagaraScriptVariable* UNiagaraGraph::AddParameter(const UNiagaraScriptVariable
 	return *FoundScriptVariable;
 }
 
+UNiagaraScriptVariable* UNiagaraGraph::CreateScriptVariableInternal(const FNiagaraVariable& Parameter, const FNiagaraVariableMetaData& ParameterMetaData, bool bIsStaticSwitch)
+{
+	UNiagaraScriptVariable* NewScriptVariable = NewObject<UNiagaraScriptVariable>(this, FName(), RF_Transactional);
+	NewScriptVariable->Init(Parameter, ParameterMetaData);
+	NewScriptVariable->SetIsStaticSwitch(bIsStaticSwitch);
+
+	FNiagaraEditorUtilities::ResetVariableToDefaultValue(NewScriptVariable->Variable);
+
+	// Inputs in graphs (module namespace) are intrinsically written to, so set the default mode to value instead of fail if not set.
+	if (NewScriptVariable->Variable.IsInNameSpace(FNiagaraConstants::ModuleNamespace))
+	{
+		NewScriptVariable->DefaultMode = ENiagaraDefaultMode::Value;
+	}
+	else
+	{
+		NewScriptVariable->DefaultMode = ENiagaraDefaultMode::FailIfPreviouslyNotSet;
+	}
+
+	VariableToScriptVariable.Add(Parameter, NewScriptVariable);
+
+	return NewScriptVariable;
+}
+
 FName UNiagaraGraph::MakeUniqueParameterName(const FName& InName)
 {
 	TArray<TWeakObjectPtr<UNiagaraGraph>> Graphs;
@@ -2258,7 +2382,7 @@ FName UNiagaraGraph::MakeUniqueParameterName(const FName& InName)
 	 {
 		 if (Graph.IsValid())
 		 {
-			 for (const auto& ParameterElement : Graph->ParameterToReferencesMap)
+			 for (const auto& ParameterElement : Graph->GetParameterReferenceMap())
 			 {
 				 Names.Add(ParameterElement.Key.GetName());
 			 }
@@ -2281,6 +2405,9 @@ void UNiagaraGraph::RemoveParameter(const FNiagaraVariable& Parameter, bool bAll
 {
 	check(!bIsForCompilationOnly);
 
+	// make sure that we have collected parameter references already
+	ConditionalRefreshParameterReferences();
+
 	FNiagaraGraphParameterReferenceCollection* ReferenceCollection = ParameterToReferencesMap.Find(Parameter);
 	if (ReferenceCollection)
 	{
@@ -2356,6 +2483,10 @@ bool UNiagaraGraph::RenameParameterFromPin(const FNiagaraVariable& Parameter, FN
 	}
 
 	Modify();
+
+	// make sure that we have collected parameter references already
+	ConditionalRefreshParameterReferences();
+
 	if (FNiagaraGraphParameterReferenceCollection* ReferenceCollection = ParameterToReferencesMap.Find(Parameter))
 	{
 		FNiagaraGraphParameterReferenceCollection NewReferences = *ReferenceCollection;
@@ -2577,6 +2708,9 @@ bool UNiagaraGraph::RenameParameter(const FNiagaraVariable& Parameter, FName New
 			FNiagaraParameterUtilities::FormatParameterNameForTextDisplay(NewName)));
 	}
 
+	// make sure that we have collected parameter references already
+	ConditionalRefreshParameterReferences();
+
 	FixupReferenceCollectionsPostRename(Parameter, NewParameter, bParameterMerged, bSuppressEvents);
 
 	bIsRenamingParameter = false;
@@ -2665,8 +2799,8 @@ void UNiagaraGraph::ScriptVariableChanged(FNiagaraVariable Variable)
 {
 	check(!bIsForCompilationOnly);
 
-	TObjectPtr<UNiagaraScriptVariable>* ScriptVariable = GetAllMetaData().Find(Variable);
-	if (!ScriptVariable || !*ScriptVariable || (*ScriptVariable)->GetIsStaticSwitch())
+	UNiagaraScriptVariable* ScriptVariable = GetScriptVariable(Variable);
+	if (!ScriptVariable || ScriptVariable->GetIsStaticSwitch())
 	{
 		return;
 	}
@@ -2680,7 +2814,7 @@ void UNiagaraGraph::ScriptVariableChanged(FNiagaraVariable Variable)
 			MapGetNodes.Add(MapGetNode);
 		}
 		
-		if ((*ScriptVariable)->DefaultMode == ENiagaraDefaultMode::Value && !Variable.GetType().IsDataInterface())
+		if (ScriptVariable->DefaultMode == ENiagaraDefaultMode::Value && !Variable.GetType().IsDataInterface())
 		{
 			FNiagaraEditorModule& EditorModule = FNiagaraEditorModule::Get();
 			auto TypeUtilityValue = EditorModule.GetTypeUtilities(Variable.GetType());
@@ -3012,27 +3146,20 @@ bool UNiagaraGraph::GetPropertyMetadata(FName PropertyName, FString& OutValue) c
 		}
 	}
 
-	const TMap<FNiagaraVariable, TObjectPtr<UNiagaraScriptVariable>>& MetaDataMap = GetAllMetaData();
-	auto Iter = MetaDataMap.CreateConstIterator();
-	while (Iter)
+	for (const FScriptVariableMap::ElementType& Iter : VariableToScriptVariable)
 	{
 		// TODO: This should never be null, but somehow it is in some assets so guard this to prevent crashes
 		// until we have better repro steps.
-		if (Iter.Value() != nullptr)
+		if (const UNiagaraScriptVariable* ScriptVariable = Iter.Value)
 		{
-			auto PropertyIter = Iter.Value()->Metadata.PropertyMetaData.CreateConstIterator();
-			while (PropertyIter)
+			if (const FString* PropertyValue = ScriptVariable->Metadata.PropertyMetaData.Find(PropertyName))
 			{
-				if (PropertyIter.Key() == PropertyName)
-				{
-					OutValue = PropertyIter.Value();
-					return true;
-				}
-				++PropertyIter;
+				OutValue = *PropertyValue;
+				return true;
 			}
 		}
-		++Iter;
 	}
+
 	return false;
 }
 
@@ -3105,91 +3232,33 @@ bool UNiagaraGraph::AppendCompileHash(FNiagaraCompileHashVisitor* InVisitor, con
 #endif
 	InVisitor->UpdateString(TEXT("ForceRebuildId"), ForceRebuildId.ToString());
 
-	ENiagaraScriptUsage TraversalUsage = InTraversal.Num() > 0 && InTraversal.Last() && Cast<UNiagaraNodeOutput>(InTraversal.Last()) ? Cast<UNiagaraNodeOutput>(InTraversal.Last())->GetUsage() : ENiagaraScriptUsage::Module;
+	TSet<FNiagaraVariable> ReferencedVariables;
 
-	// Since we are using the parameter references below, make sure that they are up to date.
-	if (bParameterReferenceRefreshPending)
-	{
-		RefreshParameterReferences();
-	}
+	NiagaraGraphImpl::FindReferencedVariables(this, InTraversal, ReferencedVariables);
 
-	// We need to sort the variables in a stable manner.
-	TArray<TObjectPtr<UNiagaraScriptVariable>> Values;
-	VariableToScriptVariable.GenerateValueArray(Values);
-	Values.Remove(nullptr);
-	Values.Sort([&](const UNiagaraScriptVariable& A, const UNiagaraScriptVariable& B)
+	TArray<UNiagaraScriptVariable*> SortedScriptVariables;
+	SortedScriptVariables.Reserve(ReferencedVariables.Num());
+	for (const FNiagaraVariable& ReferencedVariable : ReferencedVariables)
 	{
-		const FName& AName = A.Variable.GetName();
-		const FName& BName = B.Variable.GetName();
+		if (UNiagaraScriptVariable* ScriptVariable = VariableToScriptVariable.FindRef(ReferencedVariable))
+		{
+			SortedScriptVariables.Add(ScriptVariable);
+		}
+	}
 
-		return AName.LexicalLess(BName);
-	});
+	Algo::StableSortBy(SortedScriptVariables,
+		[](UNiagaraScriptVariable* ScriptVar) { return ScriptVar->Variable.GetName(); },
+		FNameLexicalLess());
 
-	// Write all the values of the local variables to the visitor as they could potentially influence compilation.
-	for (const UNiagaraScriptVariable* Var : Values)
+	for (UNiagaraScriptVariable* ScriptVariable : SortedScriptVariables)
 	{
-		FNiagaraGraphParameterReferenceCollection* Collection = ParameterToReferencesMap.Find(Var->Variable);
-		if (Collection && Collection->ParameterReferences.Num() > 0)
+#if WITH_EDITORONLY_DATA
+		if (FNiagaraCompileHashVisitorDebugInfo* DebugInfo = InVisitor->AddDebugInfo())
 		{
-			// Only add parameters to this hash that could potentially affect this compile.
-			bool bFoundInTraversal = false;
-			for (const FNiagaraGraphParameterReference& Ref : Collection->ParameterReferences)
-			{
-				if (InTraversal.Contains(Ref.Value.Get()))
-				{
-					bFoundInTraversal = true;
-					break;
-				}
-			}
-
-			// Sometimes variables exist outside the traversal that will still impact the compile. Include those here.
-			FString FoundDefaultValue;
-			if (!bFoundInTraversal)
-			{
-				bool bRelevantToTraversal = FNiagaraParameterMapHistory::IsWrittenToScriptUsage(Var->Variable, TraversalUsage, false);
-
-				if (!bRelevantToTraversal)
-					continue;
-
-				// Because we are outside the traversal, the default value won't be properly serialized into the key.
-				// Record the actual default value so that we can do that below..
-				if (Var->DefaultMode == ENiagaraDefaultMode::Value)
-				{
-					for (const FNiagaraGraphParameterReference& Ref : Collection->ParameterReferences)
-					{
-						UNiagaraNodeParameterMapGet* Node = Cast<UNiagaraNodeParameterMapGet>(Ref.Value.Get());
-						if (Node)
-						{
-							UEdGraphPin* Pin = Node->GetPinByPersistentGuid(Ref.Key);
-							if (Pin)
-							{
-								UEdGraphPin* DefaultPin = Node->GetDefaultPin(Pin);
-								if (DefaultPin && DefaultPin->DefaultValue.Len())
-								{
-									FoundDefaultValue = DefaultPin->DefaultValue;
-								}
-								break;
-							}
-						}
-					}
-				}
-			}
-
-		#if WITH_EDITORONLY_DATA
-			if (FNiagaraCompileHashVisitorDebugInfo* DebugInfo = InVisitor->AddDebugInfo())
-			{
-				DebugInfo->Object = FString::Printf(TEXT("Class: \"%s\"  Name: \"%s\""), *Var->GetClass()->GetName(), *Var->Variable.GetName().ToString());
-			}
-		#endif
-			verify(Var->AppendCompileHash(InVisitor));
-			
-			// If we are not in the traversal, make sure to also captue the default value as it isn't 
-			// currently embedded in the UNiagaraScriptVariable.
-			if (FoundDefaultValue.Len())
-			{
-				InVisitor->UpdateString(TEXT("DefaultValue"), FoundDefaultValue);
-			}
+			DebugInfo->Object = FString::Printf(TEXT("Class: \"%s\"  Name: \"%s\""), *ScriptVariable->GetClass()->GetName(), *ScriptVariable->Variable.GetName().ToString());
 		}
+#endif
+		verify(ScriptVariable->AppendCompileHash(InVisitor));
 	}
 
 	// Write all the values of the nodes to the visitor as they could influence compilation.
@@ -3502,6 +3571,9 @@ void UNiagaraGraph::RebuildCachedCompileIds()
 
 void UNiagaraGraph::CopyCachedReferencesMap(UNiagaraGraph* TargetGraph)
 {
+	// make sure that we have collected parameter references already
+	ConditionalRefreshParameterReferences();
+
 	TargetGraph->ParameterToReferencesMap = ParameterToReferencesMap;
 }
 
@@ -3679,12 +3751,9 @@ TOptional<FNiagaraVariableMetaData> UNiagaraGraph::GetMetaData(const FNiagaraVar
 		return VariableData->Metadata;
 	}
 
-	if (TObjectPtr<UNiagaraScriptVariable>* MetaData = VariableToScriptVariable.Find(InVar))
+	if (const UNiagaraScriptVariable* ScriptVariable = VariableToScriptVariable.FindRef(InVar))
 	{
-		if (*MetaData)
-		{
-			return (*MetaData)->Metadata;
-		}
+		return ScriptVariable->Metadata;
 	}
 	return TOptional<FNiagaraVariableMetaData>();
 }
@@ -3728,6 +3797,8 @@ UNiagaraGraph::FOnSubObjectSelectionChanged& UNiagaraGraph::OnSubObjectSelection
 
 void UNiagaraGraph::RefreshParameterReferences() const
 {
+	ensureMsgf(!bIsForCompilationOnly, TEXT("RefreshParameterReferences() shouldn't be called on a graph duplicated for compilation - %s"), *GetFullName());
+
 	// A set of variables to track which parameters are used so that unused parameters can be removed after the reference tracking.
 	TSet<FNiagaraVariable> CandidateUnreferencedParametersToRemove;
 
@@ -3747,14 +3818,10 @@ void UNiagaraGraph::RefreshParameterReferences() const
 
 	auto AddParameterReference = [&](const FNiagaraVariable& Parameter, const UEdGraphPin* Pin)
 	{
-		FNiagaraGraphParameterReferenceCollection* ReferenceCollection = ParameterToReferencesMap.Find(Parameter);
-		if (ReferenceCollection == nullptr)
-		{
-			FNiagaraGraphParameterReferenceCollection& NewReferenceCollection = ParameterToReferencesMap.Add(Parameter);
-			NewReferenceCollection.Graph = this;
-			ReferenceCollection = &NewReferenceCollection;
-		}
-		ReferenceCollection->ParameterReferences.AddUnique(FNiagaraGraphParameterReference(Pin->PersistentGuid, Cast<UNiagaraNode>(Pin->GetOwningNode())));
+		FNiagaraGraphParameterReferenceCollection& ReferenceCollection =
+			ParameterToReferencesMap.FindOrAdd(Parameter, FNiagaraGraphParameterReferenceCollection(false /*bInCreated*/));
+
+		ReferenceCollection.ParameterReferences.AddUnique(FNiagaraGraphParameterReference(Pin->PersistentGuid, Cast<UNiagaraNode>(Pin->GetOwningNode())));
 
 		// If we're adding a parameter reference then it needs to be removed from the list of candidate variables to remove since it's been referenced.
 		CandidateUnreferencedParametersToRemove.Remove(Parameter);
@@ -3762,29 +3829,20 @@ void UNiagaraGraph::RefreshParameterReferences() const
 
 	auto AddStaticParameterReference = [&](const FNiagaraVariable& Variable, UNiagaraNode* Node)
 	{
-		FNiagaraGraphParameterReferenceCollection* ReferenceCollection = ParameterToReferencesMap.Find(Variable);
-		if (ReferenceCollection == nullptr)
-		{
-			FNiagaraGraphParameterReferenceCollection NewReferenceCollection(false);
-			NewReferenceCollection.Graph = this;
-			ReferenceCollection = &ParameterToReferencesMap.Add(Variable, NewReferenceCollection);
-		}
-		ReferenceCollection->ParameterReferences.AddUnique(FNiagaraGraphParameterReference(Node->NodeGuid, Node));
+		FNiagaraGraphParameterReferenceCollection& ReferenceCollection =
+			ParameterToReferencesMap.FindOrAdd(Variable, FNiagaraGraphParameterReferenceCollection(false /*bInCreated*/));
+
+		ReferenceCollection.ParameterReferences.AddUnique(FNiagaraGraphParameterReference(Node->NodeGuid, Node));
 		CandidateUnreferencedParametersToRemove.Remove(Variable);
 	};
 
 	auto AddBindingParameterReference = [&](const FNiagaraVariable& Variable)
 	{
-		FNiagaraGraphParameterReferenceCollection* ReferenceCollection = ParameterToReferencesMap.Find(Variable);
-		if (ReferenceCollection == nullptr)
-		{
-			// We add an empty reference collection only when no other references exist.
-			// We cannot add an actual reference since those require both a guid and a node,
-			// but neither of these exist for direct bindings.
-			FNiagaraGraphParameterReferenceCollection NewReferenceCollection(true);
-			NewReferenceCollection.Graph = this;
-			ReferenceCollection = &ParameterToReferencesMap.Add(Variable, NewReferenceCollection);
-		}
+		// We add an empty reference collection only when no other references exist.
+		// We cannot add an actual reference since those require both a guid and a node,
+		// but neither of these exist for direct bindings.
+		FNiagaraGraphParameterReferenceCollection& ReferenceCollection =
+			ParameterToReferencesMap.FindOrAdd(Variable, FNiagaraGraphParameterReferenceCollection(true /*bInCreated*/));
 	};
 
 	// Add parameter references from parameter map traversals.
@@ -3813,7 +3871,7 @@ void UNiagaraGraph::RefreshParameterReferences() const
 			}
 		}
 	}
-	
+
 	// Check all pins on all nodes in the graph to find parameter pins which may have been missed in the parameter map traversal.  This
 	// can happen for nodes which are not fully connected and therefore don't show up in the traversal.
 	const UEdGraphSchema_Niagara* NiagaraSchema = GetNiagaraSchema();
@@ -3843,7 +3901,7 @@ void UNiagaraGraph::RefreshParameterReferences() const
 				{
 					const bool bRelevantPin = ((Pin->Direction == EEdGraphPinDirection::EGPD_Input) && Node->IsA<UNiagaraNodeParameterMapSet>())
 						|| ((Pin->Direction == EEdGraphPinDirection::EGPD_Output) && Node->IsA<UNiagaraNodeParameterMapGet>());
-					
+
 					if (!bRelevantPin)
 					{
 						continue;
@@ -3862,77 +3920,29 @@ void UNiagaraGraph::RefreshParameterReferences() const
 		}
 	}
 
-	if (bIsForCompilationOnly)
-	{
-		// Add reference to all variables that are default bound to
-		for (const FNiagaraScriptVariableData& Variable : CompilationScriptVariables)
-		{
-			FNiagaraVariable LinkedVariable(Variable.Variable.GetType(), Variable.DefaultBinding.GetName());
-
-			if ((Variable.DefaultMode != ENiagaraDefaultMode::Binding || !Variable.DefaultBinding.IsValid()))
-			{
-				continue;
-			}
-
-			if (FNiagaraConstants::GetOldPositionTypeVariables().Contains(LinkedVariable))
-			{
-				// it is not uncommon that old assets have vector inputs that default bind to what is now a position type. If we detect that, we change the type to prevent a compiler error.
-				LinkedVariable.SetType(FNiagaraTypeDefinition::GetPositionDef());
-			}
-
-			AddBindingParameterReference(LinkedVariable);
-		}
-
-		check(CandidateUnreferencedParametersToRemove.IsEmpty());
-
-		CompilationScriptVariables.SetNum(Algo::StableRemoveIf(CompilationScriptVariables, [&](const FNiagaraScriptVariableData& VariableData)
-		{
-			return !ParameterToReferencesMap.Contains(VariableData.Variable);
-		}));
-	}
-	else
+	// Add reference to all variables that are default bound to
+	for (auto It = VariableToScriptVariable.CreateConstIterator(); It; ++It)
 	{
-		// Add reference to all variables that are default bound to
-		for (auto It = VariableToScriptVariable.CreateConstIterator(); It; ++It)
+		UNiagaraScriptVariable* Variable = It.Value();
+		if (!Variable || (Variable->DefaultMode != ENiagaraDefaultMode::Binding || !Variable->DefaultBinding.IsValid()))
 		{
-			UNiagaraScriptVariable* Variable = It.Value();
-			if (!Variable || (Variable->DefaultMode != ENiagaraDefaultMode::Binding || !Variable->DefaultBinding.IsValid()))
-			{
-				continue;
-			}
-
-			FNiagaraTypeDefinition LinkedType = Variable->Variable.GetType();
-			FName BindingName = Variable->DefaultBinding.GetName();
-			if (FNiagaraConstants::GetOldPositionTypeVariables().Contains(FNiagaraVariable(LinkedType, BindingName)))
-			{
-				// it is not uncommon that old assets have vector inputs that default bind to what is now a position type. If we detect that, we change the type to prevent a compiler error.
-				LinkedType = FNiagaraTypeDefinition::GetPositionDef();
-			}
-			AddBindingParameterReference(FNiagaraVariable(LinkedType, BindingName));
-		}
-
-		// If there were any previous parameters which didn't have any references added, remove them here.
-		for (const FNiagaraVariable& UnreferencedParameterToRemove : CandidateUnreferencedParametersToRemove)
-		{
-			ParameterToReferencesMap.Remove(UnreferencedParameterToRemove);
-			VariableToScriptVariable.Remove(UnreferencedParameterToRemove);
+			continue;
 		}
 
-		// Remove any script variables 
-		TArray<FNiagaraVariable> UnreferencedScriptVariables;
-		for (auto It : VariableToScriptVariable)
+		FNiagaraTypeDefinition LinkedType = Variable->Variable.GetType();
+		FName BindingName = Variable->DefaultBinding.GetName();
+		if (FNiagaraConstants::GetOldPositionTypeVariables().Contains(FNiagaraVariable(LinkedType, BindingName)))
 		{
-			FNiagaraGraphParameterReferenceCollection* ReferenceCollection = ParameterToReferencesMap.Find(It.Key);
-			if (ReferenceCollection == nullptr)
-			{
-				UnreferencedScriptVariables.Add(It.Key);
-			}
+			// it is not uncommon that old assets have vector inputs that default bind to what is now a position type. If we detect that, we change the type to prevent a compiler error.
+			LinkedType = FNiagaraTypeDefinition::GetPositionDef();
 		}
+		AddBindingParameterReference(FNiagaraVariable(LinkedType, BindingName));
+	}
 
-		for (auto Variable : UnreferencedScriptVariables)
-		{
-			VariableToScriptVariable.Remove(Variable);
-		}
+	// If there were any previous parameters which didn't have any references added, remove them here.
+	for (const FNiagaraVariable& UnreferencedParameterToRemove : CandidateUnreferencedParametersToRemove)
+	{
+		ParameterToReferencesMap.Remove(UnreferencedParameterToRemove);
 	}
 
 	bParameterReferenceRefreshPending = false;
@@ -4036,9 +4046,9 @@ TOptional<bool> UNiagaraGraph::IsStaticSwitch(const FNiagaraVariable& Variable)
 		return VariableData->GetIsStaticSwitch();
 	}
 
-	if (TObjectPtr<UNiagaraScriptVariable>* FoundScriptVariable = VariableToScriptVariable.Find(Variable))
+	if (const UNiagaraScriptVariable* FoundScriptVariable = VariableToScriptVariable.FindRef(Variable))
 	{
-		return (*FoundScriptVariable)->GetIsStaticSwitch();
+		return FoundScriptVariable->GetIsStaticSwitch();
 	}
 
 	return TOptional<bool>();
@@ -4051,9 +4061,9 @@ TOptional<int> UNiagaraGraph::GetStaticSwitchDefaultValue(const FNiagaraVariable
 		return VariableData->GetStaticSwitchDefaultValue();
 	}
 
-	if (TObjectPtr<UNiagaraScriptVariable>* FoundScriptVariable = VariableToScriptVariable.Find(Variable))
+	if (const UNiagaraScriptVariable* FoundScriptVariable = VariableToScriptVariable.FindRef(Variable))
 	{
-		return (*FoundScriptVariable)->GetStaticSwitchDefaultValue();
+		return FoundScriptVariable->GetStaticSwitchDefaultValue();
 	}
 
 	return TOptional<int32>();
@@ -4077,20 +4087,19 @@ TOptional<ENiagaraDefaultMode> UNiagaraGraph::GetDefaultMode(const FNiagaraVaria
 		}
 	}
 
-	if (TObjectPtr<UNiagaraScriptVariable>* FoundScriptVariable = VariableToScriptVariable.Find(Variable))
+	if (const UNiagaraScriptVariable* FoundScriptVariable = VariableToScriptVariable.FindRef(Variable))
 	{
-		const UNiagaraScriptVariable* ScriptVariable = *FoundScriptVariable;
-		if (ScriptVariable->DefaultMode == ENiagaraDefaultMode::Binding && ScriptVariable->DefaultBinding.IsValid())
+		if (FoundScriptVariable->DefaultMode == ENiagaraDefaultMode::Binding && FoundScriptVariable->DefaultBinding.IsValid())
 		{
 			if (Binding)
 			{
-				*Binding = ScriptVariable->DefaultBinding;
+				*Binding = FoundScriptVariable->DefaultBinding;
 			}
 			return {ENiagaraDefaultMode::Binding};
 		}
 		else
 		{
-			return ScriptVariable->DefaultMode;
+			return FoundScriptVariable->DefaultMode;
 		}
 	}
 
@@ -4104,9 +4113,9 @@ TOptional<FGuid> UNiagaraGraph::GetScriptVariableGuid(const FNiagaraVariable& Va
 		return VariableData->Metadata.GetVariableGuid();
 	}
 
-	if (TObjectPtr<UNiagaraScriptVariable>* FoundScriptVariable = VariableToScriptVariable.Find(Variable))
+	if (const UNiagaraScriptVariable* FoundScriptVariable = VariableToScriptVariable.FindRef(Variable))
 	{
-		return (*FoundScriptVariable)->Metadata.GetVariableGuid();
+		return FoundScriptVariable->Metadata.GetVariableGuid();
 	}
 
 	return TOptional<FGuid>();
@@ -4119,9 +4128,9 @@ TOptional<FNiagaraVariable> UNiagaraGraph::GetVariable(const FNiagaraVariable& V
 		return VariableData->Variable;
 	}
 
-	if (TObjectPtr<UNiagaraScriptVariable>* FoundScriptVariable = VariableToScriptVariable.Find(Variable))
+	if (const UNiagaraScriptVariable* FoundScriptVariable = VariableToScriptVariable.FindRef(Variable))
 	{
-		return (*FoundScriptVariable)->Variable;
+		return FoundScriptVariable->Variable;
 	}
 
 	return TOptional<FNiagaraVariable>();
diff --git a/Engine/Plugins/FX/Niagara/Source/NiagaraEditor/Private/ViewModels/NiagaraParameterPanelViewModel.cpp b/Engine/Plugins/FX/Niagara/Source/NiagaraEditor/Private/ViewModels/NiagaraParameterPanelViewModel.cpp
--- a/Engine/Plugins/FX/Niagara/Source/NiagaraEditor/Private/ViewModels/NiagaraParameterPanelViewModel.cpp
+++ b/Engine/Plugins/FX/Niagara/Source/NiagaraEditor/Private/ViewModels/NiagaraParameterPanelViewModel.cpp
@@ -74,7 +74,6 @@ FReply FNiagaraSystemToolkitParameterPanelUtilities::CreateDragEventForParameter
 
 	//@todo(ng) refactor drag action to not carry around the reference collection; graph ptr goes unused.
 	FNiagaraGraphParameterReferenceCollection ReferenceCollection = FNiagaraGraphParameterReferenceCollection(true);
-	ReferenceCollection.Graph = nullptr;
 	ReferenceCollection.ParameterReferences = GraphParameterReferencesForItem;
 	const TArray<FNiagaraGraphParameterReferenceCollection> ReferenceCollectionArray = { ReferenceCollection };
 	const FText Name = FNiagaraParameterUtilities::FormatParameterNameForTextDisplay(DraggedItem.GetVariable().GetName());
@@ -97,7 +96,6 @@ FReply FNiagaraScriptToolkitParameterPanelUtilities::CreateDragEventForParameter
 
 	//@todo(ng) refactor drag action to not carry around the reference collection; graph ptr goes unused.
 	FNiagaraGraphParameterReferenceCollection ReferenceCollection = FNiagaraGraphParameterReferenceCollection(true);
-	ReferenceCollection.Graph = nullptr;
 	ReferenceCollection.ParameterReferences = GraphParameterReferencesForItem;
 	const TArray<FNiagaraGraphParameterReferenceCollection> ReferenceCollectionArray = { ReferenceCollection };
 	const FText Name = FNiagaraParameterUtilities::FormatParameterNameForTextDisplay(DraggedItem.GetVariable().GetName());
diff --git a/Engine/Plugins/FX/Niagara/Source/NiagaraEditor/Public/NiagaraGraph.h b/Engine/Plugins/FX/Niagara/Source/NiagaraEditor/Public/NiagaraGraph.h
--- a/Engine/Plugins/FX/Niagara/Source/NiagaraEditor/Public/NiagaraGraph.h
+++ b/Engine/Plugins/FX/Niagara/Source/NiagaraEditor/Public/NiagaraGraph.h
@@ -59,22 +59,20 @@ struct FNiagaraGraphParameterReferenceCollection
 
 	GENERATED_USTRUCT_BODY()
 public:
-	FNiagaraGraphParameterReferenceCollection(const bool bInCreated = false);
+	FNiagaraGraphParameterReferenceCollection() = default;
+	FNiagaraGraphParameterReferenceCollection(bool bInCreated);
 
 	/** All the references in the graph. */
 	UPROPERTY()
 	TArray<FNiagaraGraphParameterReference> ParameterReferences;
 
-	UPROPERTY()
-	TObjectPtr<const UNiagaraGraph> Graph;
-
 	/** Returns true if this parameter was initially created by the user. */
 	bool WasCreatedByUser() const;
 
 private:
 	/** Whether this parameter was initially created by the user. */
 	UPROPERTY()
-	bool bCreatedByUser;
+	bool bCreatedByUser = false;
 };
 
 
@@ -282,20 +280,23 @@ class UNiagaraGraph : public UEdGraph
 	/** Walk through the graph for an ParameterMapGet nodes and find all matching default pins for VariableName, irrespective of usage. */
 	TArray<UEdGraphPin*> FindParameterMapDefaultValuePins(const FName VariableName) const;
 
+	/** Rebuilds the internally stored ParameterRefrenceMap if it has been marked dirty.  Used to explicitly trigger the rebuild rather than
+	impliticly rebuilding it through a call to GetParameterReferenceMap(). */
+	NIAGARAEDITOR_API void ConditionalRefreshParameterReferences();
+
+	const TMap<FNiagaraVariable, FNiagaraGraphParameterReferenceCollection>& GetParameterReferenceMap() const; // NOTE: The const is a lie! (This indirectly calls RefreshParameterReferences, which can recreate the entire map)
+
 	/** Gets the meta-data associated with this variable, if it exists.*/
 	TOptional<FNiagaraVariableMetaData> GetMetaData(const FNiagaraVariable& InVar) const;
 
 	/** Sets the meta-data associated with this variable. Creates a new UNiagaraScriptVariable if the target variable cannot be found. Illegal to call on FNiagaraVariables that are Niagara Constants. */
 	void SetMetaData(const FNiagaraVariable& InVar, const FNiagaraVariableMetaData& MetaData);
 
-	const TMap<FNiagaraVariable, FNiagaraGraphParameterReferenceCollection>& GetParameterReferenceMap() const; // NOTE: The const is a lie! (This indirectly calls RefreshParameterReferences, which can recreate the entire map)
-
 	// These functions are not supported for compilation copies
-	NIAGARAEDITOR_API const TMap<FNiagaraVariable, TObjectPtr<UNiagaraScriptVariable>>& GetAllMetaData() const;
-	NIAGARAEDITOR_API TMap<FNiagaraVariable, TObjectPtr<UNiagaraScriptVariable>>& GetAllMetaData();
+	using FScriptVariableMap = TMap<FNiagaraVariable, TObjectPtr<UNiagaraScriptVariable>>;
 
-	UNiagaraScriptVariable* GetScriptVariable(FNiagaraVariable Parameter, bool bUpdateIfPending = false);
-	NIAGARAEDITOR_API UNiagaraScriptVariable* GetScriptVariable(FName ParameterName, bool bUpdateIfPending = false);
+	NIAGARAEDITOR_API const FScriptVariableMap& GetAllMetaData() const;
+	NIAGARAEDITOR_API FScriptVariableMap& GetAllMetaData();
 
 	UNiagaraScriptVariable* GetScriptVariable(FNiagaraVariable Parameter) const;
 	NIAGARAEDITOR_API UNiagaraScriptVariable* GetScriptVariable(FName ParameterName) const;
@@ -435,6 +436,7 @@ protected:
 	TMap<TPair<FGuid, UEdGraphNode*>, FNiagaraTypeDefinition> CachedNumericConversions;
 	void ResolveNumerics(TMap<UNiagaraNode*, bool>& VisitedNodes, UEdGraphNode* Node);
 	bool AppendCompileHash(FNiagaraCompileHashVisitor* InVisitor, const TArray<UNiagaraNode*>& InTraversal) const;
+	UNiagaraScriptVariable* CreateScriptVariableInternal(const FNiagaraVariable& Parameter, const FNiagaraVariableMetaData& ParameterMetaData, bool bIsStaticSwitch);
 
 private:
 	virtual void NotifyGraphChanged(const FEdGraphEditAction& InAction) override;
@@ -462,6 +464,7 @@ private:
 	void FixupReferenceCollectionsPostRename(const FNiagaraVariable& OrigVariable, const FNiagaraVariable& DestVariable, bool bParameterMerged, bool bSuppressEvents);
 
 	void PostLoad_LWCFixup(int32 NiagaraVersion);
+	void PostLoad_ManageScriptVariables(int32 NiagaraVersion);
 
 private:
 	/** The current change identifier for this graph overall. Used to sync status with UNiagaraScripts.*/
@@ -491,7 +494,7 @@ private:
 
 	/** Storage of variables defined for use with this graph.*/
 	UPROPERTY()
-	mutable TMap<FNiagaraVariable, TObjectPtr<UNiagaraScriptVariable>> VariableToScriptVariable;
+	TMap<FNiagaraVariable, TObjectPtr<UNiagaraScriptVariable>> VariableToScriptVariable;
 
 	/** A map of parameters in the graph to their referencers. */
 	UPROPERTY(Transient)
