- Clamp EffectsQuality between Min and Max level to avoid invalid data which can lead to no FX showing up

#rnx

diff --git a/Engine/Plugins/FX/Niagara/Source/Niagara/Private/NiagaraPlatformSet.cpp b/Engine/Plugins/FX/Niagara/Source/Niagara/Private/NiagaraPlatformSet.cpp
--- a/Engine/Plugins/FX/Niagara/Source/Niagara/Private/NiagaraPlatformSet.cpp
+++ b/Engine/Plugins/FX/Niagara/Source/Niagara/Private/NiagaraPlatformSet.cpp
@@ -1302,6 +1302,7 @@ FNiagaraPlatformSet::FPlatformIniSettings& FNiagaraPlatformSet::GetPlatformIniSe
 	{
 		FindCVarValue(TEXT("SystemSettings"), TEXT("sg.EffectsQuality"), EffectsQuality);
 	}
+	EffectsQuality = FMath::Clamp(EffectsQuality, MinQualityLevel, MaxQualityLevel);
 
 	//Get the platforms default quality level setting.
 	//This can be overridden directly in a device profile or indirectly by overriding effects quality.
