Ensure CVar sinks are called when setting preview platform.
Loosen how restrictive Niagara is about what quality level it expects from a device profile.
When checking for currently active sets we'll allow any available levels for a profile instead of only the default.

diff --git a/Engine/Plugins/FX/Niagara/Source/Niagara/Private/NiagaraPlatformSet.cpp b/Engine/Plugins/FX/Niagara/Source/Niagara/Private/NiagaraPlatformSet.cpp
--- a/Engine/Plugins/FX/Niagara/Source/Niagara/Private/NiagaraPlatformSet.cpp
+++ b/Engine/Plugins/FX/Niagara/Source/Niagara/Private/NiagaraPlatformSet.cpp
@@ -579,7 +579,7 @@ FNiagaraPlatformSetEnabledState FNiagaraPlatformSet::IsEnabled(const UDeviceProf
 	{
 		if (bQLIsEnabled)//Is the whole quality level disabled?
 		{
-			if (bProfileActiveQL)//The quality level is enabled but this profile is not active for this QL
+			if (bProfileActiveQL || bCheckCurrentStateOnly)//The quality level is enabled but this profile is not active for this QL
 			{
 				Ret.bIsActive = true;
 				Ret.bCanBeActive = true;
diff --git a/Engine/Source/Runtime/Engine/Private/DeviceProfiles/DeviceProfileManager.cpp b/Engine/Source/Runtime/Engine/Private/DeviceProfiles/DeviceProfileManager.cpp
--- a/Engine/Source/Runtime/Engine/Private/DeviceProfiles/DeviceProfileManager.cpp
+++ b/Engine/Source/Runtime/Engine/Private/DeviceProfiles/DeviceProfileManager.cpp
@@ -1011,6 +1011,9 @@ void UDeviceProfileManager::SetPreviewDeviceProfile(UDeviceProfile* DeviceProfil
 			}
 		}
 	}
+
+	// broadcast cvar sinks now that we are done
+	IConsoleManager::Get().CallAllConsoleVariableSinks();
 }
 
 
