Fix for vis culling issue

diff --git a/Engine/Plugins/FX/Niagara/Source/Niagara/Private/NiagaraWorldManager.cpp b/Engine/Plugins/FX/Niagara/Source/Niagara/Private/NiagaraWorldManager.cpp
--- a/Engine/Plugins/FX/Niagara/Source/Niagara/Private/NiagaraWorldManager.cpp
+++ b/Engine/Plugins/FX/Niagara/Source/Niagara/Private/NiagaraWorldManager.cpp
@@ -1571,7 +1571,7 @@ void FNiagaraWorldManager::ViewBasedCulling(UNiagaraEffectType* EffectType, cons
 
 	bool bCullByOutsideViewFrustum = ScalabilitySettings.VisibilityCulling.bCullByViewFrustum &&
 		(!bIsPrecull || ScalabilitySettings.VisibilityCulling.bAllowPreCullingByViewFrustum) &&
-		OutState.LastVisibleTime > ScalabilitySettings.VisibilityCulling.MaxTimeOutsideViewFrustum;
+		TimeSinceInsideView > ScalabilitySettings.VisibilityCulling.MaxTimeOutsideViewFrustum;
 
 	//Check for the component having been rendered recently. If the app doesn't have focus we skip this to avoid issues when alt-tabbing away from the game/editor.
 	float TimeSinceWorldRendered = World->GetTimeSeconds() - World->LastRenderTime;	
