Fixed compile issues when compiling for Windows ARM64.

diff --git a/Engine/Plugins/FX/Niagara/Source/Niagara/Classes/NiagaraOpenVDB.h b/Engine/Plugins/FX/Niagara/Source/Niagara/Classes/NiagaraOpenVDB.h
--- a/Engine/Plugins/FX/Niagara/Source/Niagara/Classes/NiagaraOpenVDB.h
+++ b/Engine/Plugins/FX/Niagara/Source/Niagara/Classes/NiagaraOpenVDB.h
@@ -4,7 +4,7 @@
 
 #include "HAL/Platform.h"
 
-#if PLATFORM_WINDOWS
+#if NIAGARA_USE_OPENVDB
 
 THIRD_PARTY_INCLUDES_START
 
diff --git a/Engine/Plugins/FX/Niagara/Source/Niagara/Niagara.Build.cs b/Engine/Plugins/FX/Niagara/Source/Niagara/Niagara.Build.cs
--- a/Engine/Plugins/FX/Niagara/Source/Niagara/Niagara.Build.cs
+++ b/Engine/Plugins/FX/Niagara/Source/Niagara/Niagara.Build.cs
@@ -78,7 +78,7 @@ public class Niagara : ModuleRules
             });
         }
 
-		if (Target.IsInPlatformGroup(UnrealPlatformGroup.Windows))
+		if (Target.IsInPlatformGroup(UnrealPlatformGroup.Windows) && Target.Architecture.bIsX64)
 		{ 
 			AddEngineThirdPartyPrivateStaticDependencies(Target,
 				"IntelTBB",
@@ -87,6 +87,20 @@ public class Niagara : ModuleRules
 				"Boost",				
 				"OpenVDB"
 			);
+
+			PublicDefinitions.AddRange(
+				new string[]
+				{
+					"NIAGARA_USE_OPENVDB=1"
+				});
+		}
+		else
+		{
+			PublicDefinitions.AddRange(
+				new string[]
+				{
+					"NIAGARA_USE_OPENVDB=0"
+				});
 		}
 
 		PublicDefinitions.AddRange(
diff --git a/Engine/Plugins/FX/Niagara/Source/Niagara/Private/NiagaraModule.cpp b/Engine/Plugins/FX/Niagara/Source/Niagara/Private/NiagaraModule.cpp
--- a/Engine/Plugins/FX/Niagara/Source/Niagara/Private/NiagaraModule.cpp
+++ b/Engine/Plugins/FX/Niagara/Source/Niagara/Private/NiagaraModule.cpp
@@ -38,7 +38,7 @@
 #include "DataInterface/NiagaraDataInterfaceDataChannelRead.h"
 #include "DataInterface/NiagaraDataInterfaceDataChannelSpawn.h"
 
-#if PLATFORM_WINDOWS
+#if NIAGARA_USE_OPENVDB
 #include "NiagaraOpenVDB.h"
 #endif
 
@@ -267,7 +267,7 @@ void INiagaraModule::StartupModule()
 	LLM_SCOPE(ELLMTag::Niagara);
 	FNiagaraTypeDefinition::Init();
 
-#if PLATFORM_WINDOWS
+#if NIAGARA_USE_OPENVDB
 	// Global registration of  the vdb types.
 	openvdb::initialize();
 	if (!Vec4SGrid::isRegistered())
diff --git a/Engine/Plugins/FX/Niagara/Source/Niagara/Private/VolumeCache.cpp b/Engine/Plugins/FX/Niagara/Source/Niagara/Private/VolumeCache.cpp
--- a/Engine/Plugins/FX/Niagara/Source/Niagara/Private/VolumeCache.cpp
+++ b/Engine/Plugins/FX/Niagara/Source/Niagara/Private/VolumeCache.cpp
@@ -9,7 +9,7 @@
 #include UE_INLINE_GENERATED_CPP_BY_NAME(VolumeCache)
 
 // @todo: we need builds for OpenVDB for platforms other than windows
-#if PLATFORM_WINDOWS
+#if NIAGARA_USE_OPENVDB
 #include "NiagaraOpenVDB.h"
 #endif
 
@@ -17,7 +17,7 @@ DEFINE_LOG_CATEGORY_STATIC(LogVolumeCache, Log, All);
 
 
 // @todo: we need builds for OpenVDB for platforms other than windows
-#if PLATFORM_WINDOWS
+#if NIAGARA_USE_OPENVDB
 class NIAGARA_API FOpenVDBCacheData : public FVolumeCacheData
 {
 public:
@@ -66,7 +66,7 @@ void UVolumeCache::InitData()
 	{
 		case EVolumeCacheType::OpenVDB :
 		{			
-#if PLATFORM_WINDOWS
+#if NIAGARA_USE_OPENVDB
 			CachedVolumeFiles = TSharedPtr<FVolumeCacheData>(new FOpenVDBCacheData());
 #else
 			UE_LOG(LogVolumeCache, Warning, TEXT("OpenVDB not supported on the platform"));
@@ -138,7 +138,7 @@ FString FVolumeCacheData::GetAssetPath(FString PathFormat, int32 FrameIndex) con
 }
 
 // @todo: we need builds for OpenVDB for platforms other than windows
-#if PLATFORM_WINDOWS
+#if NIAGARA_USE_OPENVDB
 void FOpenVDBCacheData::Init(FIntVector Resolution)
 {
 	FScopeLock ScopeLock(&DenseGridGuard);
