- Allow custom thread counts on direct iteration stages

diff --git a/Engine/Plugins/FX/Niagara/Source/NiagaraEditor/Private/NiagaraHlslTranslator.cpp b/Engine/Plugins/FX/Niagara/Source/NiagaraEditor/Private/NiagaraHlslTranslator.cpp
--- a/Engine/Plugins/FX/Niagara/Source/NiagaraEditor/Private/NiagaraHlslTranslator.cpp
+++ b/Engine/Plugins/FX/Niagara/Source/NiagaraEditor/Private/NiagaraHlslTranslator.cpp
@@ -1421,7 +1421,14 @@ const FNiagaraTranslateResults &FHlslNiagaraTranslator::Translate(const FNiagara
 					{
 						SimulationStageMetaData.GpuDispatchType = TranslationStages[Index].DirectDispatchType;
 						SimulationStageMetaData.GpuDirectDispatchElementType = TranslationStages[Index].DirectDispatchElementType;
-						SimulationStageMetaData.GpuDispatchNumThreads = FNiagaraShader::GetDefaultThreadGroupSize(TranslationStages[Index].DirectDispatchType);
+						if (TranslationStages[Index].bOverrideGpuDispatchNumThreads)
+						{
+							SimulationStageMetaData.GpuDispatchNumThreads = TranslationStages[Index].OverrideGpuDispatchNumThreads;
+						}
+						else
+						{
+							SimulationStageMetaData.GpuDispatchNumThreads = FNiagaraShader::GetDefaultThreadGroupSize(TranslationStages[Index].DirectDispatchType);
+						}
 					}
 
 					// Increment source stage index
