Fix for niagra system scripts having experimental VM byte code being generated even if it's disabled via project settings
#rnx

diff --git a/Engine/Plugins/FX/Niagara/Source/Niagara/Private/NiagaraScript.cpp b/Engine/Plugins/FX/Niagara/Source/Niagara/Private/NiagaraScript.cpp
--- a/Engine/Plugins/FX/Niagara/Source/Niagara/Private/NiagaraScript.cpp
+++ b/Engine/Plugins/FX/Niagara/Source/Niagara/Private/NiagaraScript.cpp
@@ -1249,6 +1249,11 @@ void UNiagaraScript::ComputeVMCompilationId(FNiagaraVMExecutableDataId& Id, FGui
 			Id.AdditionalDefines.Add(TEXT("CompressAttributes"));
 		}
 
+		if (System->ShouldDisableExperimentalVM() || !GetDefault<UNiagaraSettings>()->bExperimentalVMEnabled)
+		{
+			Id.AdditionalDefines.Add(FNiagaraCompileOptions::ExperimentalVMDisabled);
+		}
+
 		for (const FNiagaraEmitterHandle& EmitterHandle: System->GetEmitterHandles())
 		{
 			FVersionedNiagaraEmitterData* HandleData = EmitterHandle.GetEmitterData();
