Fixed a linker errors when compiling arm64ec.

diff --git a/Engine/Plugins/FX/Niagara/Source/Niagara/Private/NiagaraDataInterfaceRenderTargetVolume.cpp b/Engine/Plugins/FX/Niagara/Source/Niagara/Private/NiagaraDataInterfaceRenderTargetVolume.cpp
--- a/Engine/Plugins/FX/Niagara/Source/Niagara/Private/NiagaraDataInterfaceRenderTargetVolume.cpp
+++ b/Engine/Plugins/FX/Niagara/Source/Niagara/Private/NiagaraDataInterfaceRenderTargetVolume.cpp
@@ -655,7 +655,7 @@ bool UNiagaraDataInterfaceRenderTargetVolume::SimCacheWriteFrame(UObject* Storag
 		{
 			if (NDIRenderTargetVolumeLocal::GSimCacheUseOpenVDB)
 			{
-#if PLATFORM_WINDOWS
+#if NIAGARA_USE_OPENVDB
 				UVolumeCache* OpenVDBSimCacheData = CastChecked<UVolumeCache>(StorageObject);
 				
 				FString FullPath = OpenVDBSimCacheData->GetAssetPath(FrameIndex);
diff --git a/Engine/Plugins/FX/Niagara/Source/NiagaraEditor/Private/NiagaraBakerRenderer.cpp b/Engine/Plugins/FX/Niagara/Source/NiagaraEditor/Private/NiagaraBakerRenderer.cpp
--- a/Engine/Plugins/FX/Niagara/Source/NiagaraEditor/Private/NiagaraBakerRenderer.cpp
+++ b/Engine/Plugins/FX/Niagara/Source/NiagaraEditor/Private/NiagaraBakerRenderer.cpp
@@ -742,7 +742,7 @@ bool FNiagaraBakerRenderer::ExportVolume(FStringView FilePath, FIntVector ImageS
 	const FString FileExtension = FPaths::GetExtension(FilePath.GetData(), true);
 	if (FileExtension == TEXT(".vdb"))
 	{
-#if PLATFORM_WINDOWS
+#if NIAGARA_USE_OPENVDB
 		return OpenVDBTools::WriteImageDataToOpenVDBFile(FilePath, ImageSize, ImageData, false);
 #else
 		return false;
