- Fix for systems being marked as garbage ahead of GC

#preflght 63d422a23656ea96dc0b941b

diff --git a/Engine/Plugins/FX/Niagara/Source/Niagara/Private/NiagaraComponent.cpp b/Engine/Plugins/FX/Niagara/Source/Niagara/Private/NiagaraComponent.cpp
--- a/Engine/Plugins/FX/Niagara/Source/Niagara/Private/NiagaraComponent.cpp
+++ b/Engine/Plugins/FX/Niagara/Source/Niagara/Private/NiagaraComponent.cpp
@@ -703,7 +703,13 @@ void UNiagaraComponent::TickComponent(float DeltaSeconds, enum ELevelTick TickTy
 	CSV_SCOPED_TIMING_STAT_EXCLUSIVE(Effects);
 	SCOPE_CYCLE_COUNTER(STAT_NiagaraComponentTick);
 
-	FScopeCycleCounter SystemStatCounter(Asset ? Asset->GetStatID(true, false) : TStatId());
+	if (!::IsValid(Asset))
+	{
+		//-TODO: If our Asset is null or has become Garbage should we force deactive the component / ticking?
+		return;
+	}
+
+	FScopeCycleCounter SystemStatCounter(Asset->GetStatID(true, false));
 
 	if (bAwaitingActivationDueToNotReady)
 	{
diff --git a/Engine/Plugins/FX/Niagara/Source/Niagara/Private/NiagaraSystemInstance.cpp b/Engine/Plugins/FX/Niagara/Source/Niagara/Private/NiagaraSystemInstance.cpp
--- a/Engine/Plugins/FX/Niagara/Source/Niagara/Private/NiagaraSystemInstance.cpp
+++ b/Engine/Plugins/FX/Niagara/Source/Niagara/Private/NiagaraSystemInstance.cpp
@@ -100,12 +100,12 @@ static FAutoConsoleVariableRef CVarNiagaraAllowDeferredReset(
 	ECVF_Default
 );
 
-FNiagaraSystemInstance::FNiagaraSystemInstance(UWorld& InWorld, UNiagaraSystem& InAsset, FNiagaraUserRedirectionParameterStore* InOverrideParameters,
+FNiagaraSystemInstance::FNiagaraSystemInstance(UWorld& InWorld, UNiagaraSystem& InSystem, FNiagaraUserRedirectionParameterStore* InOverrideParameters,
                                                USceneComponent* InAttachComponent, ENiagaraTickBehavior InTickBehavior, bool bInPooled)
 	: SystemInstanceIndex(INDEX_NONE)
 	  , SignificanceIndex(INDEX_NONE)
 	  , World(&InWorld)
-	  , Asset(&InAsset)
+	  , System(&InSystem)
 	  , OverrideParameters(InOverrideParameters)
 	  , AttachComponent(InAttachComponent)
 	  , TickBehavior(InTickBehavior)
@@ -156,7 +156,7 @@ FNiagaraSystemInstance::FNiagaraSystemInstance(UWorld& InWorld, UNiagaraSystem&
 	FeatureLevel = World->FeatureLevel;
 
 	// In some cases the system may have already stated that you should ignore dependencies and tick as early as possible.
-	if (!InAsset.bRequireCurrentFrameData)
+	if (!InSystem.bRequireCurrentFrameData)
 	{
 		TickBehavior = ENiagaraTickBehavior::ForceTickFirst;
 	}
@@ -173,7 +173,7 @@ void FNiagaraSystemInstance::SetEmitterEnable(FName EmitterName, bool bNewEnable
 	//	return;
 	//}
 
-	UE_LOG(LogNiagara, Warning, TEXT("SetEmitterEnable: Is not implemented in Niagara. Emitter(%s) System(%s) Component(%s)"), *EmitterName.ToString(), *GetNameSafe(Asset.Get()), *GetFullNameSafe(AttachComponent.Get()));
+	UE_LOG(LogNiagara, Warning, TEXT("SetEmitterEnable: Is not implemented in Niagara. Emitter(%s) System(%s) Component(%s)"), *EmitterName.ToString(), *GetNameSafe(System), *GetFullNameSafe(AttachComponent.Get()));
 	return;
 
 	/*
@@ -492,7 +492,6 @@ void FNiagaraSystemInstance::SetSolo(bool bInSolo)
 	// We only need to transfer the instance into a new simulation if this one is within a simulation
 	if ( SystemSimulation.IsValid() )
 	{
-		UNiagaraSystem* System = GetSystem();
 		if (bInSolo)
 		{
 			TSharedPtr<FNiagaraSystemSimulation, ESPMode::ThreadSafe> NewSoloSim = MakeShared<FNiagaraSystemSimulation, ESPMode::ThreadSafe>();
@@ -521,8 +520,7 @@ void FNiagaraSystemInstance::SetSolo(bool bInSolo)
 void FNiagaraSystemInstance::SetGpuComputeDebug(bool bEnableDebug)
 {
 #if NIAGARA_COMPUTEDEBUG_ENABLED
-	UNiagaraSystem* System  = GetSystem();
-	if (ComputeDispatchInterface == nullptr || System == nullptr)
+	if (ComputeDispatchInterface == nullptr || !::IsValid(System))
 	{
 		return;
 	}
@@ -594,8 +592,7 @@ void FNiagaraSystemInstance::Activate(EResetMode InResetMode)
 {
 	SCOPE_CYCLE_COUNTER(STAT_NiagaraSystemActivate);
 
-	UNiagaraSystem* System = GetSystem();
-	if (System && System->IsValid() && IsReadyToRun())
+	if (::IsValid(System) && System->IsValid() && IsReadyToRun())
 	{
 		if (GNiagaraAllowDeferredReset && SystemInstanceIndex != INDEX_NONE && FinalizeRef.IsPending() )
 		{
@@ -639,10 +636,10 @@ void FNiagaraSystemInstance::Deactivate(bool bImmediate)
 	}
 }
 
-bool FNiagaraSystemInstance::AllocateSystemInstance(FNiagaraSystemInstancePtr& OutSystemInstanceAllocation, UWorld& InWorld, UNiagaraSystem& InAsset,
+bool FNiagaraSystemInstance::AllocateSystemInstance(FNiagaraSystemInstancePtr& OutSystemInstanceAllocation, UWorld& InWorld, UNiagaraSystem& InSystem,
 	FNiagaraUserRedirectionParameterStore* InOverrideParameters, USceneComponent* InAttachComponent, ENiagaraTickBehavior InTickBehavior, bool bInPooled)
 {
-	OutSystemInstanceAllocation = MakeShared<FNiagaraSystemInstance, ESPMode::ThreadSafe>(InWorld, InAsset, InOverrideParameters, InAttachComponent, InTickBehavior, bInPooled);
+	OutSystemInstanceAllocation = MakeShared<FNiagaraSystemInstance, ESPMode::ThreadSafe>(InWorld, InSystem, InOverrideParameters, InAttachComponent, InTickBehavior, bInPooled);
 	return true;
 }
 
@@ -898,7 +895,6 @@ void FNiagaraSystemInstance::Reset(EResetMode Mode)
 			float WarmupDt = WarmupTickDelta;
 			if (WarmupTickCount == -1)
 			{
-				UNiagaraSystem* System = GetSystem();
 				WarmupTicks = System->GetWarmupTickCount();
 				WarmupDt = System->GetWarmupTickDelta();
 			}
@@ -939,8 +935,7 @@ void FNiagaraSystemInstance::ResetInternal(bool bResetSimulations)
 	GPUParamIncludeInterpolation = false;
 	// Note: We do not need to update our bounds here as they are still valid
 
-	UNiagaraSystem* System = GetSystem();
-	if (System == nullptr || IsDisabled())
+	if (!::IsValid(System) || IsDisabled())
 	{
 		return;
 	}
@@ -1016,8 +1011,7 @@ bool FNiagaraSystemInstance::IsReadyToRun() const
 	}
 
 	// check system
-	UNiagaraSystem* System = GetSystem();
-	if (!System || !System->IsReadyToRun())
+	if (!::IsValid(System) || !System->IsReadyToRun())
 	{
 		return false;
 	}
@@ -1052,8 +1046,7 @@ void FNiagaraSystemInstance::ReInitInternal()
 	bAlreadyBound = false;
 	bSolo = bForceSolo;
 
-	UNiagaraSystem* System = GetSystem();
-	if (System == nullptr)
+	if (!::IsValid(System))
 	{
 		return;
 	}
@@ -1119,8 +1112,7 @@ void FNiagaraSystemInstance::ReInitInternal()
 
 void FNiagaraSystemInstance::ResetParameters()
 {
-	UNiagaraSystem* System = GetSystem();
-	if (!System)
+	if (!::IsValid(System))
 	{
 		return;
 	}
@@ -1551,7 +1543,7 @@ void FNiagaraSystemInstance::InitDataInterfaces()
 			bDataInterfacesInitialized &= bResult;
 			if (!bResult)
 			{
-				UE_LOG(LogNiagara, Error, TEXT("Error initializing data interface \"%s\" for system. %s"), *Interface->GetPathName(), Asset.IsValid() ? *Asset->GetName() : TEXT("nullptr"));
+				UE_LOG(LogNiagara, Error, TEXT("Error initializing data interface \"%s\" for system. %s"), *Interface->GetPathName(), *GetNameSafe(System));
 			}
 		}
 		else
@@ -1564,7 +1556,7 @@ void FNiagaraSystemInstance::InitDataInterfaces()
 	if (!bDataInterfacesInitialized && (!IsComplete() && !IsPendingSpawn()))
 	{
 		//Some error initializing the data interfaces so disable until we're explicitly reinitialized.
-		UE_LOG(LogNiagara, Error, TEXT("Error initializing data interfaces. Completing system. %s"), Asset.IsValid() ? *Asset->GetName() : TEXT("nullptr"));
+		UE_LOG(LogNiagara, Error, TEXT("Error initializing data interfaces. Completing system. %s"), *GetNameSafe(System));
 		Complete(true);
 		return;
 	}
@@ -1580,7 +1572,7 @@ void FNiagaraSystemInstance::InitDataInterfaces()
 		if (!bSuccess)
 		{
 			//Some error initializing the per instance function tables.
-			UE_LOG(LogNiagara, Error, TEXT("Error initializing data interfaces. Completing system. %s"), Asset.IsValid() ? *Asset->GetName() : TEXT("nullptr"));
+			UE_LOG(LogNiagara, Error, TEXT("Error initializing data interfaces. Completing system. %s"), *GetNameSafe(System));
 			Complete(true);
 			return;
 		}
@@ -1675,7 +1667,7 @@ void FNiagaraSystemInstance::TickDataInterfaces(float DeltaSeconds, bool bPostSi
 			if (!bSuccess)
 			{
 				// Some error initializing the per instance function tables.
-				UE_LOG(LogNiagara, Error, TEXT("Error rebinding VM functions after re-initializing data interface(s). Completing system. %s"), *GetNameSafe(Asset.Get()));
+				UE_LOG(LogNiagara, Error, TEXT("Error rebinding VM functions after re-initializing data interface(s). Completing system. %s"), *GetNameSafe(System));
 				Complete(true);
 			}
 		}
@@ -1830,8 +1822,7 @@ ETickingGroup FNiagaraSystemInstance::CalculateTickGroup() const
 
 void FNiagaraSystemInstance::SetTickBehavior(ENiagaraTickBehavior NewTickBehavior)
 {
-	UNiagaraSystem* System = GetSystem();
-	if (!System || System->bRequireCurrentFrameData)
+	if (!::IsValid(System) || System->bRequireCurrentFrameData)
 	{
 		TickBehavior = NewTickBehavior;
 	}
@@ -2011,7 +2002,7 @@ FNiagaraSystemInstance::GetEventDataSet(FName EmitterName, FName EventName) cons
 
 bool FNiagaraSystemInstance::UsesCollection(const UNiagaraParameterCollection* Collection)const
 {
-	if (UNiagaraSystem* System = GetSystem())
+	if (::IsValid(System))
 	{
 		if (System->UsesCollection(Collection))
 		{
@@ -2032,8 +2023,7 @@ void FNiagaraSystemInstance::InitEmitters()
 	LocalBounds = FBox(FVector::ZeroVector, FVector::ZeroVector);
 
 	Emitters.Empty(false);
-	UNiagaraSystem* System = GetSystem();
-	if (System != nullptr)
+	if (::IsValid(System))
 	{
 		const TArray<FNiagaraEmitterHandle>& EmitterHandles = System->GetEmitterHandles();
 
@@ -2126,7 +2116,6 @@ void FNiagaraSystemInstance::SimCacheTick_GameThread(UNiagaraSimCache* SimCache,
 
 	FNiagaraCrashReporterScope CRScope(this);
 
-	UNiagaraSystem* System = GetSystem();
 	FScopeCycleCounter SystemStat(System->GetStatID(true, false));
 
 	if (IsComplete() || IsDisabled())
@@ -2310,7 +2299,6 @@ void FNiagaraSystemInstance::Tick_GameThread(float DeltaSeconds)
 
 	FNiagaraCrashReporterScope CRScope(this);
 
-	UNiagaraSystem* System = GetSystem();
 	FScopeCycleCounter SystemStat(System->GetStatID(true, false));
 
 	// We should have no pending async operations, but wait to be safe
@@ -2371,9 +2359,8 @@ void FNiagaraSystemInstance::Tick_Concurrent(bool bEnqueueGPUTickIfNeeded)
 	TotalGPUParamSize = 0;
 	ActiveGPUEmitterCount = 0;
 	GPUParamIncludeInterpolation = false;
-	UNiagaraSystem* System = GetSystem();
 
-	if (IsComplete() || System == nullptr || CachedDeltaSeconds < SMALL_NUMBER)
+	if (IsComplete() || !::IsValid(System) || CachedDeltaSeconds < SMALL_NUMBER)
 	{
 		return;
 	}
@@ -2565,8 +2552,8 @@ void FNiagaraSystemInstance::FinalizeTick_GameThread(bool bEnqueueGPUTickIfNeede
 		//If we're batching our tick passing we may still need to enqueue here if not called from the regular finalize task. The caller will tell us with bEnqueueGPUTickIfNeeded.
 		FNiagaraSystemSimulation* Sim = SystemSimulation.Get();
 		checkf(Sim != nullptr,
-			TEXT("SystemSimulation is nullptr during Finalize and the Asset(%s) AttachComponent(%s) ActualExecutionState(%d) SystemInstanceIndex(%d) SystemInstanceState(%s)"),
-			*GetFullNameSafe(Asset.Get()), *GetFullNameSafe(AttachComponent.Get()), ActualExecutionState, SystemInstanceIndex,
+			TEXT("SystemSimulation is nullptr during Finalize and the System(%s) AttachComponent(%s) ActualExecutionState(%d) SystemInstanceIndex(%d) SystemInstanceState(%s)"),
+			*GetFullNameSafe(System), *GetFullNameSafe(AttachComponent.Get()), ActualExecutionState, SystemInstanceIndex,
 			*StaticEnum<ENiagaraSystemInstanceState>()->GetValueAsString(SystemInstanceState)
 		);
 		
@@ -2649,8 +2636,7 @@ bool FNiagaraSystemInstance::HandleNeedsUIResync()
 #if WITH_EDITORONLY_DATA
 bool FNiagaraSystemInstance::GetIsolateEnabled() const
 {
-	UNiagaraSystem* System = GetSystem();
-	if (System)
+	if (::IsValid(System))
 	{
 		return System->GetIsolateEnabled();
 	}
@@ -2708,10 +2694,9 @@ FBox FNiagaraSystemInstance::GetSystemFixedBounds() const
 	}
 	else
 	{
-		const UNiagaraSystem* NiagaraSystem = GetSystem();
-		if (NiagaraSystem && NiagaraSystem->bFixedBounds)
+		if (::IsValid(System) && System->bFixedBounds)
 		{
-			return NiagaraSystem->GetFixedBounds();
+			return System->GetFixedBounds();
 		}
 	}
 	return FBox(ForceInit);
@@ -2729,7 +2714,7 @@ void FNiagaraSystemInstance::SetEmitterFixedBounds(FName EmitterName, const FBox
 	}
 
 	// Failed to find emitter
-	UE_LOG(LogNiagara, Warning, TEXT("SetEmitterFixedBounds: Failed to find Emitter(%s) System(%s) Component(%s)"), *EmitterName.ToString(), *GetNameSafe(Asset.Get()), *GetFullNameSafe(AttachComponent.Get()));
+	UE_LOG(LogNiagara, Warning, TEXT("SetEmitterFixedBounds: Failed to find Emitter(%s) System(%s) Component(%s)"), *EmitterName.ToString(), *GetNameSafe(System), *GetFullNameSafe(AttachComponent.Get()));
 }
 
 FBox FNiagaraSystemInstance::GetEmitterFixedBounds(FName EmitterName) const
@@ -2743,7 +2728,7 @@ FBox FNiagaraSystemInstance::GetEmitterFixedBounds(FName EmitterName) const
 	}
 
 	// Failed to find emitter
-	UE_LOG(LogNiagara, Warning, TEXT("GetEmitterFixedBounds: Failed to find Emitter(%s) System(%s) Component(%s)"), *EmitterName.ToString(), *GetNameSafe(Asset.Get()), *GetFullNameSafe(AttachComponent.Get()));
+	UE_LOG(LogNiagara, Warning, TEXT("GetEmitterFixedBounds: Failed to find Emitter(%s) System(%s) Component(%s)"), *EmitterName.ToString(), *GetNameSafe(System), *GetFullNameSafe(AttachComponent.Get()));
 	return FBox(ForceInit);
 }
 
@@ -2815,7 +2800,7 @@ void FNiagaraSystemInstance::EvaluateBoundFunction(FName FunctionName, bool& Use
 		}
 	};
 
-	if (const UNiagaraSystem* System = GetSystem())
+	if (::IsValid(System))
 	{
 		System->ForEachScript(ScriptUsesFunction);
 	}
@@ -2847,10 +2832,10 @@ const FString& FNiagaraSystemInstance::GetCrashReporterTag()const
 		USceneComponent* AttachParent = Component ? Component->GetAttachParent() : AttachComponent.Get();
 
 		const FString& CompName = GetFullNameSafe(Component);
-		const FString& AssetName = GetFullNameSafe(Sys);
+		const FString& SystemName = GetFullNameSafe(Sys);
 		const FString& AttachName = GetFullNameSafe(AttachParent);
 
-		CrashReporterTag = FString::Printf(TEXT("SystemInstance | System: %s | bSolo: %s | Component: %s | AttachedTo: %s |"), *AssetName, IsSolo() ? TEXT("true") : TEXT("false"), *CompName, *AttachName);
+		CrashReporterTag = FString::Printf(TEXT("SystemInstance | System: %s | bSolo: %s | Component: %s | AttachedTo: %s |"), *SystemName, IsSolo() ? TEXT("true") : TEXT("false"), *CompName, *AttachName);
 	}
 	return CrashReporterTag;
 }
diff --git a/Engine/Plugins/FX/Niagara/Source/Niagara/Private/NiagaraWorldManager.cpp b/Engine/Plugins/FX/Niagara/Source/Niagara/Private/NiagaraWorldManager.cpp
--- a/Engine/Plugins/FX/Niagara/Source/Niagara/Private/NiagaraWorldManager.cpp
+++ b/Engine/Plugins/FX/Niagara/Source/Niagara/Private/NiagaraWorldManager.cpp
@@ -1191,10 +1191,9 @@ void FNiagaraWorldManager::Tick(ETickingGroup TickGroup, float DeltaSeconds, ELe
 
 	ActiveNiagaraTickGroup = ActualTickGroup;
 
-	TArray<UNiagaraSystem*, TInlineAllocator<4>> DeadSystems;
-	for (TPair<UNiagaraSystem*, FNiagaraSystemSimulationRef>& SystemSim : SystemSimulations[ActualTickGroup])
+	for (auto SimIt=SystemSimulations[ActualTickGroup].CreateIterator(); SimIt; ++SimIt)
 	{
-		FNiagaraSystemSimulation*  Sim = &SystemSim.Value.Get();
+		FNiagaraSystemSimulation* Sim = &SimIt.Value().Get();
 
 		if (Sim->IsValid())
 		{
@@ -1202,17 +1201,13 @@ void FNiagaraWorldManager::Tick(ETickingGroup TickGroup, float DeltaSeconds, ELe
 		}
 		else
 		{
-			DeadSystems.Add(SystemSim.Key);
+			Sim->Destroy();
+			SimIt.RemoveCurrent();
 		}
 	}
 
 	ActiveNiagaraTickGroup = -1;
 
-	for (UNiagaraSystem* DeadSystem : DeadSystems)
-	{
-		SystemSimulations[ActualTickGroup].Remove(DeadSystem);
-	}
-
 	// Loop over all simulations that have been marked for post actor (i.e. ones whos TG is changing or have pending spawn systems)
 	if (GNiagaraSpawnPerTickGroup && (SimulationsWithPostActorWork.Num() > 0))
 	{
diff --git a/Engine/Plugins/FX/Niagara/Source/Niagara/Public/NiagaraSystemInstance.h b/Engine/Plugins/FX/Niagara/Source/Niagara/Public/NiagaraSystemInstance.h
--- a/Engine/Plugins/FX/Niagara/Source/Niagara/Public/NiagaraSystemInstance.h
+++ b/Engine/Plugins/FX/Niagara/Source/Niagara/Public/NiagaraSystemInstance.h
@@ -233,7 +233,7 @@ public:
 	TSharedPtr<FNiagaraEmitterInstance, ESPMode::ThreadSafe> GetSimulationForHandle(const FNiagaraEmitterHandle& EmitterHandle);
 
 	FORCEINLINE UWorld* GetWorld() const { return World; }
-	FORCEINLINE UNiagaraSystem* GetSystem() const { return Asset.Get(); }
+	FORCEINLINE UNiagaraSystem* GetSystem() const { return System; }
 	FORCEINLINE USceneComponent* GetAttachComponent() { return AttachComponent.Get(); }
 	FORCEINLINE FNiagaraUserRedirectionParameterStore* GetOverrideParameters() { return OverrideParameters; }
 	FORCEINLINE const FNiagaraUserRedirectionParameterStore* GetOverrideParameters() const { return OverrideParameters; }
@@ -447,7 +447,7 @@ private:
 	TSharedPtr<class FNiagaraSystemSimulation, ESPMode::ThreadSafe> SystemSimulation;
 
 	UWorld* World;
-	TWeakObjectPtr<UNiagaraSystem> Asset;
+	UNiagaraSystem* System = nullptr;
 	FNiagaraUserRedirectionParameterStore* OverrideParameters;
 	TWeakObjectPtr<USceneComponent> AttachComponent;
 
diff --git a/Engine/Plugins/FX/Niagara/Source/Niagara/Public/NiagaraSystemSimulation.h b/Engine/Plugins/FX/Niagara/Source/Niagara/Public/NiagaraSystemSimulation.h
--- a/Engine/Plugins/FX/Niagara/Source/Niagara/Public/NiagaraSystemSimulation.h
+++ b/Engine/Plugins/FX/Niagara/Source/Niagara/Public/NiagaraSystemSimulation.h
@@ -264,7 +264,7 @@ public:
 	bool Init(UNiagaraSystem* InSystem, UWorld* InWorld, bool bInIsSolo, ETickingGroup TickGroup);
 	void Destroy();
 
-	bool IsValid() const { return bCanExecute && World != nullptr; }
+	bool IsValid() const { return bCanExecute && World != nullptr && ::IsValid(System); }
 
 	/** First phase of system sim tick. Must run on GameThread. */
 	void Tick_GameThread(float DeltaSeconds, const FGraphEventRef& MyCompletionGraphEvent);
