Track OpenVDB cache memory uasge
expose max memory usage project setting

diff --git a/Engine/Plugins/FX/Niagara/Source/Niagara/Private/NiagaraSettings.cpp b/Engine/Plugins/FX/Niagara/Source/Niagara/Private/NiagaraSettings.cpp
--- a/Engine/Plugins/FX/Niagara/Source/Niagara/Private/NiagaraSettings.cpp
+++ b/Engine/Plugins/FX/Niagara/Source/Niagara/Private/NiagaraSettings.cpp
@@ -14,6 +14,7 @@ UNiagaraSettings::UNiagaraSettings(const FObjectInitializer& ObjectInitlaizer)
 	NDICollisionQuery_AsyncGpuTraceProviderOrder.Add(ENDICollisionQuery_AsyncGpuTraceProvider::Type::GSDF);
 
 	SimCacheAuxiliaryFileBasePath = "{project_dir}/Saved/NiagaraSimCache";
+	SimCacheMaxCPUMemoryVolumetrics = 5000;
 }
 
 FName UNiagaraSettings::GetCategoryName() const
diff --git a/Engine/Plugins/FX/Niagara/Source/Niagara/Private/VolumeCache.cpp b/Engine/Plugins/FX/Niagara/Source/Niagara/Private/VolumeCache.cpp
--- a/Engine/Plugins/FX/Niagara/Source/Niagara/Private/VolumeCache.cpp
+++ b/Engine/Plugins/FX/Niagara/Source/Niagara/Private/VolumeCache.cpp
@@ -5,6 +5,7 @@
 =============================================================================*/
 
 #include "VolumeCache.h"
+#include "NiagaraSettings.h"
 #include "RHICommandList.h"
 
 #include UE_INLINE_GENERATED_CPP_BY_NAME(VolumeCache)
@@ -22,7 +23,7 @@ DEFINE_LOG_CATEGORY_STATIC(LogVolumeCache, Log, All);
 class NIAGARA_API FOpenVDBCacheData : public FVolumeCacheData
 {
 public:
-	FOpenVDBCacheData() {}
+	FOpenVDBCacheData() : TotalMemoryUsage(0) {}
 
 	virtual ~FOpenVDBCacheData()
 	{
@@ -40,6 +41,9 @@ public:
 	
 private:
 	TMap<int32, Vec4Grid::Ptr> OpenVDBGrids;
+	TQueue<int32> CacheQueue;
+
+	int64 TotalMemoryUsage;
 
 	FCriticalSection DenseGridGuard;
 	Vec4Dense::Ptr DenseGridPtr;
@@ -143,7 +147,13 @@ FString FVolumeCacheData::GetAssetPath(FString PathFormat, int32 FrameIndex) con
 void FOpenVDBCacheData::Init(FIntVector Resolution)
 {
 	FScopeLock ScopeLock(&DenseGridGuard);
+	
 	DenseGridPtr.reset(new Vec4Dense(openvdb::CoordBBox(0, 0, 0, Resolution.X - 1, Resolution.Y - 1, Resolution.Z - 1), Vec4(0.0, 0.0, 0.0, 0.0)));
+
+	if (OpenVDBGrids.Num() == 0)
+	{
+		TotalMemoryUsage = 0;
+	}
 }
 
 bool FOpenVDBCacheData::LoadFile(FString Path, int frame)
@@ -211,6 +221,20 @@ bool FOpenVDBCacheData::LoadFile(FString Path, int frame)
 
 			OpenVDBGrids.Add(frame, ColorGrid);
 
+			// add total memory usage reported in bytes
+			TotalMemoryUsage += ColorGrid->memUsage() / 1000000;
+
+			// enqueue 
+			CacheQueue.Enqueue(frame);
+
+			// if total memory usage is over the limit, dequeue oldest one
+			int64 MaxMemory = GetDefault<UNiagaraSettings>()->SimCacheMaxCPUMemoryVolumetrics;
+			if (TotalMemoryUsage > MaxMemory && !CacheQueue.IsEmpty())
+			{
+				int FrameToRemove = CacheQueue.Pop();
+				UnloadFile(FrameToRemove);
+			}
+
 			// if dense vdb buffer doesn't match current resolution, re initialize it				
 			FScopeLock ScopeLock(&DenseGridGuard);
 			if (DenseGridPtr == nullptr || (DenseGridPtr->bbox().dim() != openvdb::Coord(DenseResolution.X - 1, DenseResolution.Y - 1, DenseResolution.Z - 1)))
@@ -237,6 +261,7 @@ bool FOpenVDBCacheData::UnloadFile(int frame)
 {
 	if (OpenVDBGrids.Contains(frame) && OpenVDBGrids[frame] != nullptr)
 	{
+		TotalMemoryUsage -= OpenVDBGrids[frame]->memUsage() / 100000;
 		OpenVDBGrids[frame] = nullptr;
 		OpenVDBGrids.Remove(frame);
 
@@ -262,6 +287,9 @@ bool FOpenVDBCacheData::LoadRange(FString Path, int Start, int End)
 void FOpenVDBCacheData::UnloadAll()
 {
 	OpenVDBGrids.Reset();
+	CacheQueue.Empty();
+
+	TotalMemoryUsage = 0;
 }
 
 bool FOpenVDBCacheData::Fill3DTexture_RenderThread(int frame, FTextureRHIRef TextureToFill, FRHICommandListImmediate& RHICmdList)
diff --git a/Engine/Plugins/FX/Niagara/Source/Niagara/Public/NiagaraSettings.h b/Engine/Plugins/FX/Niagara/Source/Niagara/Public/NiagaraSettings.h
--- a/Engine/Plugins/FX/Niagara/Source/Niagara/Public/NiagaraSettings.h
+++ b/Engine/Plugins/FX/Niagara/Source/Niagara/Public/NiagaraSettings.h
@@ -214,6 +214,12 @@ class NIAGARA_API UNiagaraSettings : public UDeveloperSettings
 	UPROPERTY(config, EditAnywhere, Category = SimCache, meta = (DisplayName = "Sim Cache Auxiliary File Base Path", ConfigRestartRequired = false))
 	FString SimCacheAuxiliaryFileBasePath;
 
+	/**
+	Max memory in megabytes for total CPU memory for cached volumetric data
+	*/
+	UPROPERTY(config, EditAnywhere, Category = SimCache, meta = (DisplayName = "Sim Cache Max CPU Memory For Volumetrics", ConfigRestartRequired = false))
+	int64 SimCacheMaxCPUMemoryVolumetrics;
+
 	UPROPERTY(config, EditAnywhere, Category = Scalability)
 	TArray<FNiagaraPlatformSetRedirect> PlatformSetRedirects;
 
