Removing NiagaraDataInterfaceArrayImpl.h includes from public headers. This required moving all the FNDIArrayImplHelper types into a single file for private use inside Niagara.

diff --git a/Engine/Plugins/FX/Niagara/Source/Niagara/Classes/NiagaraDataInterfaceArray.h b/Engine/Plugins/FX/Niagara/Source/Niagara/Classes/NiagaraDataInterfaceArray.h
--- a/Engine/Plugins/FX/Niagara/Source/Niagara/Classes/NiagaraDataInterfaceArray.h
+++ b/Engine/Plugins/FX/Niagara/Source/Niagara/Classes/NiagaraDataInterfaceArray.h
@@ -11,6 +11,64 @@
 
 #include "NiagaraDataInterfaceArray.generated.h"
 
+template<typename TArrayType, class TOwnerType>
+struct FNDIArrayProxyImpl;
+
+#define NDIARRAY_GENERATE_BODY(CLASSNAME, TYPENAME, MEMBERNAME) \
+	using FProxyType = FNDIArrayProxyImpl<TYPENAME, CLASSNAME>; \
+	virtual void PostInitProperties() override; \
+	template<typename TFromArrayType> \
+	void SetVariantArrayData(TConstArrayView<TFromArrayType> InArrayData) \
+	{ \
+		MEMBERNAME = InArrayData; \
+	} \
+	template<typename TFromArrayType> \
+	void SetVariantArrayValue(int Index, const TFromArrayType& Value, bool bSizeToFit) \
+	{ \
+		const int NumRequired = Index + 1 - MEMBERNAME.Num(); \
+		if ( NumRequired > 0 && !bSizeToFit ) return; \
+		MEMBERNAME.AddDefaulted(FMath::Max(NumRequired, 0)); \
+		MEMBERNAME[Index] = Value; \
+	} \
+	TArray<TYPENAME>& GetArrayReference() { return MEMBERNAME; }
+
+#if WITH_EDITORONLY_DATA
+#define NDIARRAY_GENERATE_BODY_LWC(CLASSNAME, TYPENAME, MEMBERNAME) \
+		using FProxyType = FNDIArrayProxyImpl<TYPENAME, CLASSNAME>; \
+		virtual void PostInitProperties() override; \
+		virtual void PostLoad() override; \
+		virtual void PostEditChangeProperty(struct FPropertyChangedEvent& PropertyChangedEvent) override; \
+		virtual bool CopyToInternal(UNiagaraDataInterface* Destination) const override; \
+		virtual bool Equals(const UNiagaraDataInterface* Other) const override; \
+		template<typename TFromArrayType> \
+		void SetVariantArrayData(TConstArrayView<TFromArrayType> InArrayData) \
+		{ \
+			if constexpr (std::is_same_v<TFromArrayType, decltype(MEMBERNAME)::ElementType>) \
+			{ \
+				MEMBERNAME = InArrayData; \
+				GetProxyAs<FProxyType>()->template SetArrayData<decltype(MEMBERNAME)::ElementType>(InArrayData); \
+			} \
+			else \
+			{ \
+				MEMBERNAME.SetNumUninitialized(InArrayData.Num()); \
+				FNDIArrayImplHelper<TYPENAME>::CopyCpuToCpuMemory(MEMBERNAME.GetData(), InArrayData.GetData(), InArrayData.Num()); \
+				GetProxyAs<FProxyType>()->template SetArrayData<decltype(Internal##MEMBERNAME)::ElementType>(InArrayData); \
+			} \
+		} \
+		template<typename TFromArrayType> \
+		void SetVariantArrayValue(int Index, const TFromArrayType& Value, bool bSizeToFit) \
+		{ \
+			const int NumRequired = Index + 1 - MEMBERNAME.Num(); \
+			if ( NumRequired > 0 && !bSizeToFit ) return; \
+			MEMBERNAME.AddDefaulted(FMath::Max(NumRequired, 0)); \
+			MEMBERNAME[Index] = Value; \
+			GetProxyAs<FProxyType>()->template SetArrayData<decltype(MEMBERNAME)::ElementType>(MEMBERNAME); \
+		} \
+		TArray<TYPENAME>& GetArrayReference() { return Internal##MEMBERNAME; }
+#else
+	#define NDIARRAY_GENERATE_BODY_LWC(CLASSNAME, TYPENAME, MEMBERNAME) NDIARRAY_GENERATE_BODY(CLASSNAME, TYPENAME, Internal##MEMBERNAME)
+#endif
+
 struct INDIArrayProxyBase : public FNiagaraDataInterfaceProxyRW
 {
 	BEGIN_SHADER_PARAMETER_STRUCT(FShaderParameters,)
diff --git a/Engine/Plugins/FX/Niagara/Source/Niagara/Classes/NiagaraDataInterfaceArrayFloat.h b/Engine/Plugins/FX/Niagara/Source/Niagara/Classes/NiagaraDataInterfaceArrayFloat.h
--- a/Engine/Plugins/FX/Niagara/Source/Niagara/Classes/NiagaraDataInterfaceArrayFloat.h
+++ b/Engine/Plugins/FX/Niagara/Source/Niagara/Classes/NiagaraDataInterfaceArrayFloat.h
@@ -3,253 +3,8 @@
 #pragma once
 
 #include "NiagaraDataInterfaceArray.h"
-#include "NiagaraDataInterfaceArrayImpl.h"
 #include "NiagaraDataInterfaceArrayFloat.generated.h"
 
-template<>
-struct FNDIArrayImplHelper<float> : public FNDIArrayImplHelperBase<float>
-{
-	typedef float TVMArrayType;
-
-	static constexpr TCHAR const* HLSLVariableType = TEXT("float");
-	static constexpr EPixelFormat ReadPixelFormat = PF_R32_FLOAT;
-	static constexpr TCHAR const* ReadHLSLBufferType = TEXT("float");
-	static constexpr TCHAR const* ReadHLSLBufferRead = TEXT("Value = BUFFER_NAME[Index];");
-	static constexpr EPixelFormat RWPixelFormat = PF_R32_FLOAT;
-	static constexpr TCHAR const* RWHLSLBufferType = TEXT("float");
-	static constexpr TCHAR const* RWHLSLBufferRead = TEXT("Value = BUFFER_NAME[Index];");
-	static constexpr TCHAR const* RWHLSLBufferWrite = TEXT("BUFFER_NAME[Index] = Value;");
-
-	static const FNiagaraTypeDefinition& GetTypeDefinition() { return FNiagaraTypeDefinition::GetFloatDef(); }
-	static const float GetDefaultValue() { return 0.0f; }
-
-	static void AppendValueToString(const float Value, FString& OutString)
-	{
-		OutString.Appendf(TEXT("%f"), Value);
-	}
-};
-
-template<>
-struct FNDIArrayImplHelper<FVector2f> : public FNDIArrayImplHelperBase<FVector2f>
-{
-	typedef FVector2f TVMArrayType;
-
-	static constexpr TCHAR const* HLSLVariableType = TEXT("float2");
-	static constexpr EPixelFormat ReadPixelFormat = PF_G32R32F;
-	static constexpr TCHAR const* ReadHLSLBufferType = TEXT("float2");
-	static constexpr TCHAR const* ReadHLSLBufferRead = TEXT("Value = float2(BUFFER_NAME[Index]);");
-	static constexpr EPixelFormat RWPixelFormat = PF_R32_FLOAT;
-	static constexpr TCHAR const* RWHLSLBufferType = TEXT("float");
-	static constexpr TCHAR const* RWHLSLBufferRead = TEXT("Value = float2(BUFFER_NAME[Index * 2 + 0], BUFFER_NAME[Index * 2 + 1]);");
-	static constexpr TCHAR const* RWHLSLBufferWrite = TEXT("BUFFER_NAME[Index * 2 + 0] = Value.x, BUFFER_NAME[Index * 2 + 1] = Value.y;");
-
-	static const FNiagaraTypeDefinition& GetTypeDefinition() { return FNiagaraTypeDefinition::GetVec2Def(); }
-	static const FVector2f GetDefaultValue() { return FVector2f::ZeroVector; }
-
-	static void CopyCpuToCpuMemory(FVector2f* Dest, const FVector2f* Src, int32 NumElements)
-	{
-		FMemory::Memcpy(Dest, Src, NumElements * sizeof(FVector2f));
-	}
-
-	static void CopyCpuToCpuMemory(FVector2f* Dest, const FVector2D* Src, int32 NumElements)
-	{
-		for ( int32 i=0; i < NumElements; ++i )
-		{
-			Dest[i] = FVector2f(Src[i]);
-		}
-	}
-
-	static void CopyCpuToCpuMemory(FVector2D* Dest, const FVector2f* Src, int32 NumElements)
-	{
-		for (int32 i = 0; i < NumElements; ++i)
-		{
-			Dest[i] = FVector2D(Src[i]);
-		}
-	}
-
-	static void AppendValueToString(const FVector2f& Value, FString& OutString)
-	{
-		OutString.Appendf(TEXT("%f, %f"), Value.X, Value.Y);
-	}
-};
-
-template<>
-struct FNDIArrayImplHelper<FVector3f> : public FNDIArrayImplHelperBase<FVector3f>
-{
-	typedef FVector3f TVMArrayType;
-
-	static constexpr TCHAR const* HLSLVariableType = TEXT("float3");
-	static constexpr EPixelFormat ReadPixelFormat = PF_R32_FLOAT;		// Lack of float3 format
-	static constexpr TCHAR const* ReadHLSLBufferType = TEXT("float");
-	static constexpr TCHAR const* ReadHLSLBufferRead = TEXT("Value = float3(BUFFER_NAME[Index * 3 + 0], BUFFER_NAME[Index * 3 + 1], BUFFER_NAME[Index * 3 + 2]);");
-	static constexpr EPixelFormat RWPixelFormat = PF_R32_FLOAT;
-	static constexpr TCHAR const* RWHLSLBufferType = TEXT("float");
-	static constexpr TCHAR const* RWHLSLBufferRead = TEXT("Value = float3(BUFFER_NAME[Index * 3 + 0], BUFFER_NAME[Index * 3 + 1], BUFFER_NAME[Index * 3 + 2]);");
-	static constexpr TCHAR const* RWHLSLBufferWrite = TEXT("BUFFER_NAME[Index * 3 + 0] = Value.x, BUFFER_NAME[Index * 3 + 1] = Value.y, BUFFER_NAME[Index * 3 + 2] = Value.z;");
-
-	static const FNiagaraTypeDefinition& GetTypeDefinition() { return FNiagaraTypeDefinition::GetVec3Def(); }
-	static const FVector3f GetDefaultValue() { return FVector3f::ZeroVector; }
-
-	static void CopyCpuToCpuMemory(FVector3f* Dest, const FVector3f* Src, int32 NumElements)
-	{
-		FMemory::Memcpy(Dest, Src, NumElements * sizeof(FVector3f));
-	}
-
-	static void CopyCpuToCpuMemory(FVector3f* Dest, const FVector* Src, int32 NumElements)
-	{
-		for (int32 i = 0; i < NumElements; ++i)
-		{
-			Dest[i] = FVector3f(Src[i]);
-		}
-	}
-
-	static void CopyCpuToCpuMemory(FVector* Dest, const FVector3f* Src, int32 NumElements)
-	{
-		for (int32 i = 0; i < NumElements; ++i)
-		{
-			Dest[i] = FVector(Src[i]);
-		}
-	}
-
-	static void AppendValueToString(const FVector3f& Value, FString& OutString)
-	{
-		OutString.Appendf(TEXT("%f, %f, %f"), Value.X, Value.Y, Value.Z);
-	}
-};
-
-template<>
-struct FNDIArrayImplHelper<FNiagaraPosition> : public FNDIArrayImplHelper<FVector3f>
-{
-	typedef FNiagaraPosition TVMArrayType;
-
-	static const FNiagaraTypeDefinition& GetTypeDefinition() { return FNiagaraTypeDefinition::GetPositionDef(); }
-	static const FNiagaraPosition GetDefaultValue() { return FVector3f::ZeroVector; }
-
-	static void CopyCpuToCpuMemory(FVector3f* Dest, const FNiagaraPosition* Src, int32 NumElements)
-	{
-		for (int32 i = 0; i < NumElements; ++i)
-		{
-			Dest[i] = Src[i];
-		}
-	}
-
-	static void AppendValueToString(const FVector3f& Value, FString& OutString)
-	{
-		OutString.Appendf(TEXT("%f, %f, %f"), Value.X, Value.Y, Value.Z);
-	}
-};
-
-template<>
-struct FNDIArrayImplHelper<FVector4f> : public FNDIArrayImplHelperBase<FVector4f>
-{
-	typedef FVector4f TVMArrayType;
-
-	static constexpr TCHAR const* HLSLVariableType = TEXT("float4");
-	static constexpr EPixelFormat ReadPixelFormat = PF_A32B32G32R32F;
-	static constexpr TCHAR const* ReadHLSLBufferType = TEXT("float4");
-	static constexpr TCHAR const* ReadHLSLBufferRead = TEXT("Value = BUFFER_NAME[Index];");
-	static constexpr EPixelFormat RWPixelFormat = PF_R32_FLOAT;
-	static constexpr TCHAR const* RWHLSLBufferType = TEXT("float");
-	static constexpr TCHAR const* RWHLSLBufferRead = TEXT("Value = float4(BUFFER_NAME[Index * 4 + 0], BUFFER_NAME[Index * 4 + 1], BUFFER_NAME[Index * 4 + 2], BUFFER_NAME[Index * 4 + 3]);");
-	static constexpr TCHAR const* RWHLSLBufferWrite = TEXT("BUFFER_NAME[Index * 4 + 0] = Value.x, BUFFER_NAME[Index * 4 + 1] = Value.y, BUFFER_NAME[Index * 4 + 2] = Value.z, BUFFER_NAME[Index * 4 + 3] = Value.w;");
-
-	static const FNiagaraTypeDefinition& GetTypeDefinition() { return FNiagaraTypeDefinition::GetVec4Def(); }
-	static const FVector4f GetDefaultValue() { return FVector4f(ForceInitToZero); }
-
-	static void CopyCpuToCpuMemory(FVector4f* Dest, const FVector4f* Src, int32 NumElements)
-	{
-		FMemory::Memcpy(Dest, Src, NumElements * sizeof(FVector4f));
-	}
-
-	static void CopyCpuToCpuMemory(FVector4f* Dest, const FVector4* Src, int32 NumElements)
-	{
-		for (int32 i = 0; i < NumElements; ++i)
-		{
-			Dest[i] = FVector4f(Src[i]);
-		}
-	}
-
-	static void CopyCpuToCpuMemory(FVector4* Dest, const FVector4f* Src, int32 NumElements)
-	{
-		for (int32 i = 0; i < NumElements; ++i)
-		{
-			Dest[i] = FVector4(Src[i]);
-		}
-	}
-
-	static void AppendValueToString(const FVector4f& Value, FString& OutString)
-	{
-		OutString.Appendf(TEXT("%f, %f, %f, %f"), Value.X, Value.Y, Value.Z, Value.W);
-	}
-};
-
-template<>
-struct FNDIArrayImplHelper<FLinearColor> : public FNDIArrayImplHelperBase<FLinearColor>
-{
-	typedef FLinearColor TVMArrayType;
-
-	static constexpr TCHAR const* HLSLVariableType = TEXT("float4");
-	static constexpr EPixelFormat ReadPixelFormat = PF_A32B32G32R32F;
-	static constexpr TCHAR const* ReadHLSLBufferType = TEXT("float4");
-	static constexpr TCHAR const* ReadHLSLBufferRead = TEXT("Value = BUFFER_NAME[Index];");
-	static constexpr EPixelFormat RWPixelFormat = PF_R32_FLOAT;
-	static constexpr TCHAR const* RWHLSLBufferType = TEXT("float");
-	static constexpr TCHAR const* RWHLSLBufferRead = TEXT("Value = float4(BUFFER_NAME[Index * 4 + 0], BUFFER_NAME[Index * 4 + 1], BUFFER_NAME[Index * 4 + 2], BUFFER_NAME[Index * 4 + 3]);");
-	static constexpr TCHAR const* RWHLSLBufferWrite = TEXT("BUFFER_NAME[Index * 4 + 0] = Value.x, BUFFER_NAME[Index * 4 + 1] = Value.y, BUFFER_NAME[Index * 4 + 2] = Value.z, BUFFER_NAME[Index * 4 + 3] = Value.w;");
-
-	static const FNiagaraTypeDefinition& GetTypeDefinition() { return FNiagaraTypeDefinition::GetColorDef(); }
-	static const FLinearColor GetDefaultValue() { return FLinearColor::White; }
-
-	static void AppendValueToString(const FLinearColor& Value, FString& OutString)
-	{
-		OutString.Appendf(TEXT("%f, %f, %f, %f"), Value.R, Value.G, Value.B, Value.A);
-	}
-};
-
-template<>
-struct FNDIArrayImplHelper<FQuat4f> : public FNDIArrayImplHelperBase<FQuat4f>
-{
-	typedef FQuat4f TVMArrayType;
-
-	static constexpr TCHAR const* HLSLVariableType = TEXT("float4");
-	static constexpr EPixelFormat ReadPixelFormat = PF_A32B32G32R32F;
-	static constexpr TCHAR const* ReadHLSLBufferType = TEXT("float4");
-	static constexpr TCHAR const* ReadHLSLBufferRead = TEXT("Value = BUFFER_NAME[Index];");
-	static constexpr EPixelFormat RWPixelFormat = PF_R32_FLOAT;
-	static constexpr TCHAR const* RWHLSLBufferType = TEXT("float");
-	static constexpr TCHAR const* RWHLSLBufferRead = TEXT("Value = float4(BUFFER_NAME[Index * 4 + 0], BUFFER_NAME[Index * 4 + 1], BUFFER_NAME[Index * 4 + 2], BUFFER_NAME[Index * 4 + 3]);");
-	static constexpr TCHAR const* RWHLSLBufferWrite = TEXT("BUFFER_NAME[Index * 4 + 0] = Value.x, BUFFER_NAME[Index * 4 + 1] = Value.y, BUFFER_NAME[Index * 4 + 2] = Value.z, BUFFER_NAME[Index * 4 + 3] = Value.w;");
-
-	static const FNiagaraTypeDefinition& GetTypeDefinition() { return FNiagaraTypeDefinition::GetQuatDef(); }
-	static const FQuat4f GetDefaultValue() { return FQuat4f::Identity; }
-
-	static void CopyCpuToCpuMemory(FQuat4f* Dest, const FQuat4f* Src, int32 NumElements)
-	{
-		FMemory::Memcpy(Dest, Src, NumElements * sizeof(FQuat4f));
-	}
-
-	static void CopyCpuToCpuMemory(FQuat4f* Dest, const FQuat* Src, int32 NumElements)
-	{
-		for (int32 i = 0; i < NumElements; ++i)
-		{
-			Dest[i] = FQuat4f(Src[i]);
-		}
-	}
-
-	static void CopyCpuToCpuMemory(FQuat* Dest, const FQuat4f* Src, int32 NumElements)
-	{
-		for (int32 i = 0; i < NumElements; ++i)
-		{
-			Dest[i] = FQuat(Src[i]);
-		}
-	}
-
-	static void AppendValueToString(const FQuat4f& Value, FString& OutString)
-	{
-		OutString.Appendf(TEXT("%f, %f, %f, %f"), Value.X, Value.Y, Value.Z, Value.W);
-	}
-};
-
 UCLASS(EditInlineNew, Category = "Array", meta = (DisplayName = "Float Array"), Blueprintable, BlueprintType)
 class NIAGARA_API UNiagaraDataInterfaceArrayFloat : public UNiagaraDataInterfaceArray
 {
diff --git a/Engine/Plugins/FX/Niagara/Source/Niagara/Classes/NiagaraDataInterfaceArrayImpl.h b/Engine/Plugins/FX/Niagara/Source/Niagara/Classes/NiagaraDataInterfaceArrayImpl.h
--- a/Engine/Plugins/FX/Niagara/Source/Niagara/Classes/NiagaraDataInterfaceArrayImpl.h
+++ b/Engine/Plugins/FX/Niagara/Source/Niagara/Classes/NiagaraDataInterfaceArrayImpl.h
@@ -19,47 +19,31 @@
 //////////////////////////////////////////////////////////////////////////
 // Helpers
 
-#define NDIARRAY_GENERATE_BODY(CLASSNAME, TYPENAME, MEMBERNAME) \
-	using FProxyType = FNDIArrayProxyImpl<TYPENAME, CLASSNAME>; \
-	virtual void PostInitProperties() override \
+#define NDIARRAY_GENERATE_IMPL(CLASSNAME, TYPENAME, MEMBERNAME) \
+	void CLASSNAME::PostInitProperties() \
 	{ \
 		Proxy.Reset(new FProxyType(this)); \
 		Super::PostInitProperties(); \
-	} \
-	template<typename TFromArrayType> \
-	void SetVariantArrayData(TConstArrayView<TFromArrayType> InArrayData) \
-	{ \
-		MEMBERNAME = InArrayData; \
-	} \
-	template<typename TFromArrayType> \
-	void SetVariantArrayValue(int Index, const TFromArrayType& Value, bool bSizeToFit) \
-	{ \
-		const int NumRequired = Index + 1 - MEMBERNAME.Num(); \
-		if ( NumRequired > 0 && !bSizeToFit ) return; \
-		MEMBERNAME.AddDefaulted(FMath::Max(NumRequired, 0)); \
-		MEMBERNAME[Index] = Value; \
-	} \
-	TArray<TYPENAME>& GetArrayReference() { return MEMBERNAME; }
+	}
 
 #if WITH_EDITORONLY_DATA
-	#define NDIARRAY_GENERATE_BODY_LWC(CLASSNAME, TYPENAME, MEMBERNAME) \
-		using FProxyType = FNDIArrayProxyImpl<TYPENAME, CLASSNAME>; \
-		virtual void PostInitProperties() override \
+	#define NDIARRAY_GENERATE_IMPL_LWC(CLASSNAME, TYPENAME, MEMBERNAME) \
+		void CLASSNAME::PostInitProperties() \
 		{ \
 			Super::PostInitProperties(); \
 			Proxy.Reset(new FProxyType(this)); \
 		} \
-		virtual void PostLoad() override \
+		void CLASSNAME::PostLoad() \
 		{ \
 			Super::PostLoad(); \
 			GetProxyAs<FProxyType>()->template SetArrayData<decltype(MEMBERNAME)::ElementType>(MakeArrayView(MEMBERNAME)); \
 		} \
-		virtual void PostEditChangeProperty(struct FPropertyChangedEvent& PropertyChangedEvent) override \
+		void CLASSNAME::PostEditChangeProperty(struct FPropertyChangedEvent& PropertyChangedEvent) \
 		{ \
 			Super::PostEditChangeProperty(PropertyChangedEvent); \
 			GetProxyAs<FProxyType>()->template SetArrayData<decltype(MEMBERNAME)::ElementType>(MakeArrayView(MEMBERNAME)); \
 		} \
-		virtual bool CopyToInternal(UNiagaraDataInterface* Destination) const override \
+		bool CLASSNAME::CopyToInternal(UNiagaraDataInterface* Destination) const \
 		{ \
 			if ( Super::CopyToInternal(Destination) == false ) \
 			{ \
@@ -72,41 +56,16 @@
 			} \
 			return TypedDestination != nullptr;  \
 		} \
-		virtual bool Equals(const UNiagaraDataInterface* Other) const override \
+		bool CLASSNAME::Equals(const UNiagaraDataInterface* Other) const \
 		{ \
 			const CLASSNAME* TypedOther = Cast<const CLASSNAME>(Other); \
 			return \
 				Super::Equals(Other) && \
 				TypedOther != nullptr && \
 				TypedOther->MEMBERNAME == MEMBERNAME; \
-		} \
-		template<typename TFromArrayType> \
-		void SetVariantArrayData(TConstArrayView<TFromArrayType> InArrayData) \
-		{ \
-			if constexpr (std::is_same_v<TFromArrayType, decltype(MEMBERNAME)::ElementType>) \
-			{ \
-				MEMBERNAME = InArrayData; \
-				GetProxyAs<FProxyType>()->template SetArrayData<decltype(MEMBERNAME)::ElementType>(InArrayData); \
-			} \
-			else \
-			{ \
-				MEMBERNAME.SetNumUninitialized(InArrayData.Num()); \
-				FNDIArrayImplHelper<TYPENAME>::CopyCpuToCpuMemory(MEMBERNAME.GetData(), InArrayData.GetData(), InArrayData.Num()); \
-				GetProxyAs<FProxyType>()->template SetArrayData<decltype(Internal##MEMBERNAME)::ElementType>(InArrayData); \
-			} \
-		} \
-		template<typename TFromArrayType> \
-		void SetVariantArrayValue(int Index, const TFromArrayType& Value, bool bSizeToFit) \
-		{ \
-			const int NumRequired = Index + 1 - MEMBERNAME.Num(); \
-			if ( NumRequired > 0 && !bSizeToFit ) return; \
-			MEMBERNAME.AddDefaulted(FMath::Max(NumRequired, 0)); \
-			MEMBERNAME[Index] = Value; \
-			GetProxyAs<FProxyType>()->template SetArrayData<decltype(MEMBERNAME)::ElementType>(MEMBERNAME); \
-		} \
-		TArray<TYPENAME>& GetArrayReference() { return Internal##MEMBERNAME; }
+		}
 #else
-	#define NDIARRAY_GENERATE_BODY_LWC(CLASSNAME, TYPENAME, MEMBERNAME) NDIARRAY_GENERATE_BODY(CLASSNAME, TYPENAME, Internal##MEMBERNAME)
+	#define NDIARRAY_GENERATE_IMPL_LWC(CLASSNAME, TYPENAME, MEMBERNAME) NDIARRAY_GENERATE_IMPL(CLASSNAME, TYPENAME, Internal##MEMBERNAME)
 #endif
 
 template<typename TArrayType>
diff --git a/Engine/Plugins/FX/Niagara/Source/Niagara/Classes/NiagaraDataInterfaceArrayInt.h b/Engine/Plugins/FX/Niagara/Source/Niagara/Classes/NiagaraDataInterfaceArrayInt.h
--- a/Engine/Plugins/FX/Niagara/Source/Niagara/Classes/NiagaraDataInterfaceArrayInt.h
+++ b/Engine/Plugins/FX/Niagara/Source/Niagara/Classes/NiagaraDataInterfaceArrayInt.h
@@ -3,143 +3,8 @@
 #pragma once
 
 #include "NiagaraDataInterfaceArray.h"
-#include "NiagaraDataInterfaceArrayImpl.h"
 #include "NiagaraDataInterfaceArrayInt.generated.h"
 
-template<>
-struct FNDIArrayImplHelper<int32> : public FNDIArrayImplHelperBase<int32>
-{
-	typedef int32 TVMArrayType;
-
-	static constexpr bool bSupportsAtomicOps = true;
-
-	static constexpr TCHAR const* HLSLVariableType = TEXT("int");
-	static constexpr EPixelFormat ReadPixelFormat = PF_R32_SINT;
-	static constexpr TCHAR const* ReadHLSLBufferType = TEXT("int");
-	static constexpr TCHAR const* ReadHLSLBufferRead = TEXT("Value = BUFFER_NAME[Index];");
-	static constexpr EPixelFormat RWPixelFormat = PF_R32_SINT;
-	static constexpr TCHAR const* RWHLSLBufferType = TEXT("int");
-	static constexpr TCHAR const* RWHLSLBufferRead = TEXT("Value = BUFFER_NAME[Index];");
-	static constexpr TCHAR const* RWHLSLBufferWrite = TEXT("BUFFER_NAME[Index] = Value;");
-
-	static const FNiagaraTypeDefinition& GetTypeDefinition() { return FNiagaraTypeDefinition::GetIntDef(); }
-	static const int32 GetDefaultValue() { return 0; }
-
-	static int32 AtomicAdd(int32* Dest, int32 Value)
-	{
-		return FPlatformAtomics::InterlockedAdd(Dest, Value);
-	}
-	static int32 AtomicMin(int32* Dest, int32 Value)
-	{
-		int32 PrevValue = *Dest;
-		while (PrevValue > Value)
-		{
-			PrevValue = FPlatformAtomics::InterlockedCompareExchange(Dest, Value, PrevValue);
-		}
-		return PrevValue;
-	}
-	static int32 AtomicMax(int32* Dest, int32 Value)
-	{
-		int32 PrevValue = *Dest;
-		while (PrevValue < Value)
-		{
-			PrevValue = FPlatformAtomics::InterlockedCompareExchange(Dest, Value, PrevValue);
-		}
-		return PrevValue;
-	}
-
-	static void AppendValueToString(const int32 Value, FString& OutString)
-	{
-		OutString.AppendInt(Value);
-	}
-};
-
-template<>
-struct FNDIArrayImplHelper<uint8> : public FNDIArrayImplHelperBase<uint8>
-{
-	typedef int32 TVMArrayType;
-
-	static constexpr TCHAR const* HLSLVariableType = TEXT("int");
-	static constexpr EPixelFormat ReadPixelFormat = PF_R8_UINT;
-	static constexpr TCHAR const* ReadHLSLBufferType = TEXT("int");
-	static constexpr TCHAR const* ReadHLSLBufferRead = TEXT("Value = BUFFER_NAME[Index];");
-	static constexpr EPixelFormat RWPixelFormat = PF_R8_UINT;
-	static constexpr TCHAR const* RWHLSLBufferType = TEXT("int");
-	static constexpr TCHAR const* RWHLSLBufferRead = TEXT("Value = BUFFER_NAME[Index];");
-	static constexpr TCHAR const* RWHLSLBufferWrite = TEXT("BUFFER_NAME[Index] = Value;");
-
-	static const FNiagaraTypeDefinition& GetTypeDefinition() { return FNiagaraTypeDefinition::GetIntDef(); }
-	static const int32 GetDefaultValue() { return 0; }
-
-	static void CopyCpuToCpuMemory(uint8* Dest, const uint8* Src, int32 NumElements)
-	{
-		FMemory::Memcpy(Dest, Src, NumElements * sizeof(uint8));
-	}
-
-	static void CopyCpuToCpuMemory(uint8* Dest, const int32* Src, int32 NumElements)
-	{
-		for (int32 i = 0; i < NumElements; ++i)
-		{
-			Dest[i] = uint8(Src[i]);
-		}
-	}
-
-	static void CopyCpuToCpuMemory(int32* Dest, const uint8* Src, int32 NumElements)
-	{
-		for (int32 i = 0; i < NumElements; ++i)
-		{
-			Dest[i] = int32(Src[i]);
-		}
-	}
-
-	static void AppendValueToString(const uint8 Value, FString& OutString)
-	{
-		OutString.AppendInt(Value);
-	}
-};
-
-template<>
-struct FNDIArrayImplHelper<bool> : public FNDIArrayImplHelperBase<bool>
-{
-	typedef FNiagaraBool TVMArrayType;
-
-	static constexpr TCHAR const* HLSLVariableType = TEXT("bool");
-	static constexpr EPixelFormat ReadPixelFormat = PF_R8_UINT;
-	static constexpr TCHAR const* ReadHLSLBufferType = TEXT("uint");
-	static constexpr TCHAR const* ReadHLSLBufferRead = TEXT("Value = BUFFER_NAME[Index];");
-	static constexpr EPixelFormat RWPixelFormat = PF_R8_UINT;
-	static constexpr TCHAR const* RWHLSLBufferType = TEXT("uint");
-	static constexpr TCHAR const* RWHLSLBufferRead = TEXT("Value = BUFFER_NAME[Index];");
-	static constexpr TCHAR const* RWHLSLBufferWrite = TEXT("BUFFER_NAME[Index] = Value;");
-
-	static const FNiagaraTypeDefinition& GetTypeDefinition() { return FNiagaraTypeDefinition::GetBoolDef(); }
-	static const FNiagaraBool GetDefaultValue() { return FNiagaraBool::False; }
-
-	static void CopyCpuToGpuMemory(void* Dest, const bool* Src, int32 NumElements)
-	{
-		uint8* TypedDest = reinterpret_cast<uint8*>(Dest);
-		while (NumElements--)
-		{
-			*TypedDest++ = *Src++ == false ? 0 : 0xff;
-		}
-	}
-
-	static void CopyGpuToCpuMemory(void* Dest, const void* Src, int32 NumElements)
-	{
-		FNiagaraBool* TypedDest = reinterpret_cast<FNiagaraBool*>(Dest);
-		const uint8* TypedSrc = reinterpret_cast<const uint8*>(Src);
-		while (NumElements--)
-		{
-			*TypedDest++ = *TypedSrc++ == 0 ? false : true;
-		}
-	}
-
-	static void AppendValueToString(const bool Value, FString& OutString)
-	{
-		OutString.Append(Value ? TEXT("True") : TEXT("False"));
-	}
-};
-
 UCLASS(EditInlineNew, Category = "Array", meta = (DisplayName = "Int32 Array"), Blueprintable, BlueprintType)
 class NIAGARA_API UNiagaraDataInterfaceArrayInt32 : public UNiagaraDataInterfaceArray
 {
diff --git a/Engine/Plugins/FX/Niagara/Source/Niagara/Classes/NiagaraDataInterfaceArrayNiagaraID.h b/Engine/Plugins/FX/Niagara/Source/Niagara/Classes/NiagaraDataInterfaceArrayNiagaraID.h
--- a/Engine/Plugins/FX/Niagara/Source/Niagara/Classes/NiagaraDataInterfaceArrayNiagaraID.h
+++ b/Engine/Plugins/FX/Niagara/Source/Niagara/Classes/NiagaraDataInterfaceArrayNiagaraID.h
@@ -3,32 +3,8 @@
 #pragma once
 
 #include "NiagaraDataInterfaceArray.h"
-#include "NiagaraDataInterfaceArrayImpl.h"
 #include "NiagaraDataInterfaceArrayNiagaraID.generated.h"
 
-template<>
-struct FNDIArrayImplHelper<FNiagaraID> : public FNDIArrayImplHelperBase<FNiagaraID>
-{
-	typedef FNiagaraID TVMArrayType;
-
-	static constexpr TCHAR const* HLSLVariableType = TEXT("NiagaraID");
-	static constexpr EPixelFormat ReadPixelFormat = PF_R32_SINT;		// Should be int2 but we don't have it
-	static constexpr TCHAR const* ReadHLSLBufferType = TEXT("int");
-	static constexpr TCHAR const* ReadHLSLBufferRead = TEXT("Value.Index = BUFFER_NAME[Index * 2 + 0]; Value.AcquireTag = BUFFER_NAME[Index * 2 + 1];");
-	static constexpr EPixelFormat RWPixelFormat = PF_R32_SINT;
-	static constexpr TCHAR const* RWHLSLBufferType = TEXT("int");
-	static constexpr TCHAR const* RWHLSLBufferRead = TEXT("Value.Index = BUFFER_NAME[Index * 2 + 0]; Value.AcquireTag = BUFFER_NAME[Index * 2 + 1];");
-	static constexpr TCHAR const* RWHLSLBufferWrite = TEXT("BUFFER_NAME[Index * 2 + 0] = Value.Index; BUFFER_NAME[Index * 2 + 1] = Value.AcquireTag;");
-
-	static const FNiagaraTypeDefinition GetTypeDefinition() { return FNiagaraTypeDefinition(FNiagaraID::StaticStruct()); }
-	static const FNiagaraID GetDefaultValue() { return FNiagaraID(); }
-
-	static void AppendValueToString(const FNiagaraID Value, FString& OutString)
-	{
-		OutString.Appendf(TEXT("%d, %d"), Value.Index, Value.AcquireTag);
-	}
-};
-
 UCLASS(EditInlineNew, Category = "Array", meta = (DisplayName = "NiagaraID Array", Experimental), Blueprintable, BlueprintType)
 class NIAGARA_API UNiagaraDataInterfaceArrayNiagaraID : public UNiagaraDataInterfaceArray
 {
diff --git a/Engine/Plugins/FX/Niagara/Source/Niagara/Private/NiagaraDataInterfaceArrayFloat.cpp b/Engine/Plugins/FX/Niagara/Source/Niagara/Private/NiagaraDataInterfaceArrayFloat.cpp
--- a/Engine/Plugins/FX/Niagara/Source/Niagara/Private/NiagaraDataInterfaceArrayFloat.cpp
+++ b/Engine/Plugins/FX/Niagara/Source/Niagara/Private/NiagaraDataInterfaceArrayFloat.cpp
@@ -1,6 +1,16 @@
 // Copyright Epic Games, Inc. All Rights Reserved.
 
 #include "NiagaraDataInterfaceArrayFloat.h"
+#include "NiagaraDataInterfaceArrayImpl.h"
+#include "NiagaraDataInterfaceArrayImplHelpers.h"
 #include "NiagaraRenderer.h"
 
 #include UE_INLINE_GENERATED_CPP_BY_NAME(NiagaraDataInterfaceArrayFloat)
+
+NDIARRAY_GENERATE_IMPL(UNiagaraDataInterfaceArrayFloat, float, FloatData)
+NDIARRAY_GENERATE_IMPL_LWC(UNiagaraDataInterfaceArrayFloat2, FVector2f, FloatData)
+NDIARRAY_GENERATE_IMPL_LWC(UNiagaraDataInterfaceArrayFloat3, FVector3f, FloatData)
+NDIARRAY_GENERATE_IMPL(UNiagaraDataInterfaceArrayPosition, FNiagaraPosition, PositionData)
+NDIARRAY_GENERATE_IMPL_LWC(UNiagaraDataInterfaceArrayFloat4, FVector4f, FloatData)
+NDIARRAY_GENERATE_IMPL(UNiagaraDataInterfaceArrayColor, FLinearColor, ColorData)
+NDIARRAY_GENERATE_IMPL_LWC(UNiagaraDataInterfaceArrayQuat, FQuat4f, QuatData)
diff --git a/Engine/Plugins/FX/Niagara/Source/Niagara/Private/NiagaraDataInterfaceArrayFunctionLibrary.cpp b/Engine/Plugins/FX/Niagara/Source/Niagara/Private/NiagaraDataInterfaceArrayFunctionLibrary.cpp
--- a/Engine/Plugins/FX/Niagara/Source/Niagara/Private/NiagaraDataInterfaceArrayFunctionLibrary.cpp
+++ b/Engine/Plugins/FX/Niagara/Source/Niagara/Private/NiagaraDataInterfaceArrayFunctionLibrary.cpp
@@ -6,6 +6,7 @@
 #include "NiagaraDataInterfaceArrayFloat.h"
 #include "NiagaraDataInterfaceArrayInt.h"
 #include "NiagaraDataInterfaceArrayImpl.h"
+#include "NiagaraDataInterfaceArrayImplHelpers.h"
 #include "NiagaraComponent.h"
 #include "NiagaraSystemInstance.h"
 #include "NiagaraSystemInstanceController.h"
diff --git a/Engine/Plugins/FX/Niagara/Source/Niagara/Private/NiagaraDataInterfaceArrayImplHelpers.h b/Engine/Plugins/FX/Niagara/Source/Niagara/Private/NiagaraDataInterfaceArrayImplHelpers.h
new file mode 100644
--- /dev/null
+++ b/Engine/Plugins/FX/Niagara/Source/Niagara/Private/NiagaraDataInterfaceArrayImplHelpers.h
@@ -0,0 +1,351 @@
+// Copyright Epic Games, Inc. All Rights Reserved.
+
+#pragma once
+
+#include "NiagaraDataInterfaceArrayImpl.h"
+
+template<>
+struct FNDIArrayImplHelper<float> : public FNDIArrayImplHelperBase<float>
+{
+	typedef float TVMArrayType;
+
+	static constexpr TCHAR const* HLSLVariableType = TEXT("float");
+	static constexpr EPixelFormat ReadPixelFormat = PF_R32_FLOAT;
+	static constexpr TCHAR const* ReadHLSLBufferType = TEXT("float");
+	static constexpr TCHAR const* ReadHLSLBufferRead = TEXT("Value = BUFFER_NAME[Index];");
+	static constexpr EPixelFormat RWPixelFormat = PF_R32_FLOAT;
+	static constexpr TCHAR const* RWHLSLBufferType = TEXT("float");
+	static constexpr TCHAR const* RWHLSLBufferRead = TEXT("Value = BUFFER_NAME[Index];");
+	static constexpr TCHAR const* RWHLSLBufferWrite = TEXT("BUFFER_NAME[Index] = Value;");
+
+	static const FNiagaraTypeDefinition& GetTypeDefinition() { return FNiagaraTypeDefinition::GetFloatDef(); }
+	static const float GetDefaultValue() { return 0.0f; }
+};
+
+template<>
+struct FNDIArrayImplHelper<FVector2f> : public FNDIArrayImplHelperBase<FVector2f>
+{
+	typedef FVector2f TVMArrayType;
+
+	static constexpr TCHAR const* HLSLVariableType = TEXT("float2");
+	static constexpr EPixelFormat ReadPixelFormat = PF_G32R32F;
+	static constexpr TCHAR const* ReadHLSLBufferType = TEXT("float2");
+	static constexpr TCHAR const* ReadHLSLBufferRead = TEXT("Value = float2(BUFFER_NAME[Index]);");
+	static constexpr EPixelFormat RWPixelFormat = PF_R32_FLOAT;
+	static constexpr TCHAR const* RWHLSLBufferType = TEXT("float");
+	static constexpr TCHAR const* RWHLSLBufferRead = TEXT("Value = float2(BUFFER_NAME[Index * 2 + 0], BUFFER_NAME[Index * 2 + 1]);");
+	static constexpr TCHAR const* RWHLSLBufferWrite = TEXT("BUFFER_NAME[Index * 2 + 0] = Value.x, BUFFER_NAME[Index * 2 + 1] = Value.y;");
+
+	static const FNiagaraTypeDefinition& GetTypeDefinition() { return FNiagaraTypeDefinition::GetVec2Def(); }
+	static const FVector2f GetDefaultValue() { return FVector2f::ZeroVector; }
+
+	static void CopyCpuToCpuMemory(FVector2f* Dest, const FVector2f* Src, int32 NumElements)
+	{
+		FMemory::Memcpy(Dest, Src, NumElements * sizeof(FVector2f));
+	}
+
+	static void CopyCpuToCpuMemory(FVector2f* Dest, const FVector2D* Src, int32 NumElements)
+	{
+		for (int32 i = 0; i < NumElements; ++i)
+		{
+			Dest[i] = FVector2f(Src[i]);
+		}
+	}
+
+	static void CopyCpuToCpuMemory(FVector2D* Dest, const FVector2f* Src, int32 NumElements)
+	{
+		for (int32 i = 0; i < NumElements; ++i)
+		{
+			Dest[i] = FVector2D(Src[i]);
+		}
+	}
+};
+
+template<>
+struct FNDIArrayImplHelper<FVector3f> : public FNDIArrayImplHelperBase<FVector3f>
+{
+	typedef FVector3f TVMArrayType;
+
+	static constexpr TCHAR const* HLSLVariableType = TEXT("float3");
+	static constexpr EPixelFormat ReadPixelFormat = PF_R32_FLOAT;		// Lack of float3 format
+	static constexpr TCHAR const* ReadHLSLBufferType = TEXT("float");
+	static constexpr TCHAR const* ReadHLSLBufferRead = TEXT("Value = float3(BUFFER_NAME[Index * 3 + 0], BUFFER_NAME[Index * 3 + 1], BUFFER_NAME[Index * 3 + 2]);");
+	static constexpr EPixelFormat RWPixelFormat = PF_R32_FLOAT;
+	static constexpr TCHAR const* RWHLSLBufferType = TEXT("float");
+	static constexpr TCHAR const* RWHLSLBufferRead = TEXT("Value = float3(BUFFER_NAME[Index * 3 + 0], BUFFER_NAME[Index * 3 + 1], BUFFER_NAME[Index * 3 + 2]);");
+	static constexpr TCHAR const* RWHLSLBufferWrite = TEXT("BUFFER_NAME[Index * 3 + 0] = Value.x, BUFFER_NAME[Index * 3 + 1] = Value.y, BUFFER_NAME[Index * 3 + 2] = Value.z;");
+
+	static const FNiagaraTypeDefinition& GetTypeDefinition() { return FNiagaraTypeDefinition::GetVec3Def(); }
+	static const FVector3f GetDefaultValue() { return FVector3f::ZeroVector; }
+
+	static void CopyCpuToCpuMemory(FVector3f* Dest, const FVector3f* Src, int32 NumElements)
+	{
+		FMemory::Memcpy(Dest, Src, NumElements * sizeof(FVector3f));
+	}
+
+	static void CopyCpuToCpuMemory(FVector3f* Dest, const FVector* Src, int32 NumElements)
+	{
+		for (int32 i = 0; i < NumElements; ++i)
+		{
+			Dest[i] = FVector3f(Src[i]);
+		}
+	}
+
+	static void CopyCpuToCpuMemory(FVector* Dest, const FVector3f* Src, int32 NumElements)
+	{
+		for (int32 i = 0; i < NumElements; ++i)
+		{
+			Dest[i] = FVector(Src[i]);
+		}
+	}
+};
+
+template<>
+struct FNDIArrayImplHelper<FNiagaraPosition> : public FNDIArrayImplHelper<FVector3f>
+{
+	typedef FNiagaraPosition TVMArrayType;
+
+	static const FNiagaraTypeDefinition& GetTypeDefinition() { return FNiagaraTypeDefinition::GetPositionDef(); }
+	static const FNiagaraPosition GetDefaultValue() { return FVector3f::ZeroVector; }
+
+	static void CopyCpuToCpuMemory(FVector3f* Dest, const FNiagaraPosition* Src, int32 NumElements)
+	{
+		for (int32 i = 0; i < NumElements; ++i)
+		{
+			Dest[i] = Src[i];
+		}
+	}
+};
+
+template<>
+struct FNDIArrayImplHelper<FVector4f> : public FNDIArrayImplHelperBase<FVector4f>
+{
+	typedef FVector4f TVMArrayType;
+
+	static constexpr TCHAR const* HLSLVariableType = TEXT("float4");
+	static constexpr EPixelFormat ReadPixelFormat = PF_A32B32G32R32F;
+	static constexpr TCHAR const* ReadHLSLBufferType = TEXT("float4");
+	static constexpr TCHAR const* ReadHLSLBufferRead = TEXT("Value = BUFFER_NAME[Index];");
+	static constexpr EPixelFormat RWPixelFormat = PF_R32_FLOAT;
+	static constexpr TCHAR const* RWHLSLBufferType = TEXT("float");
+	static constexpr TCHAR const* RWHLSLBufferRead = TEXT("Value = float4(BUFFER_NAME[Index * 4 + 0], BUFFER_NAME[Index * 4 + 1], BUFFER_NAME[Index * 4 + 2], BUFFER_NAME[Index * 4 + 3]);");
+	static constexpr TCHAR const* RWHLSLBufferWrite = TEXT("BUFFER_NAME[Index * 4 + 0] = Value.x, BUFFER_NAME[Index * 4 + 1] = Value.y, BUFFER_NAME[Index * 4 + 2] = Value.z, BUFFER_NAME[Index * 4 + 3] = Value.w;");
+
+	static const FNiagaraTypeDefinition& GetTypeDefinition() { return FNiagaraTypeDefinition::GetVec4Def(); }
+	static const FVector4f GetDefaultValue() { return FVector4f(ForceInitToZero); }
+
+	static void CopyCpuToCpuMemory(FVector4f* Dest, const FVector4f* Src, int32 NumElements)
+	{
+		FMemory::Memcpy(Dest, Src, NumElements * sizeof(FVector4f));
+	}
+
+	static void CopyCpuToCpuMemory(FVector4f* Dest, const FVector4* Src, int32 NumElements)
+	{
+		for (int32 i = 0; i < NumElements; ++i)
+		{
+			Dest[i] = FVector4f(Src[i]);
+		}
+	}
+
+	static void CopyCpuToCpuMemory(FVector4* Dest, const FVector4f* Src, int32 NumElements)
+	{
+		for (int32 i = 0; i < NumElements; ++i)
+		{
+			Dest[i] = FVector4(Src[i]);
+		}
+	}
+};
+
+template<>
+struct FNDIArrayImplHelper<FLinearColor> : public FNDIArrayImplHelperBase<FLinearColor>
+{
+	typedef FLinearColor TVMArrayType;
+
+	static constexpr TCHAR const* HLSLVariableType = TEXT("float4");
+	static constexpr EPixelFormat ReadPixelFormat = PF_A32B32G32R32F;
+	static constexpr TCHAR const* ReadHLSLBufferType = TEXT("float4");
+	static constexpr TCHAR const* ReadHLSLBufferRead = TEXT("Value = BUFFER_NAME[Index];");
+	static constexpr EPixelFormat RWPixelFormat = PF_R32_FLOAT;
+	static constexpr TCHAR const* RWHLSLBufferType = TEXT("float");
+	static constexpr TCHAR const* RWHLSLBufferRead = TEXT("Value = float4(BUFFER_NAME[Index * 4 + 0], BUFFER_NAME[Index * 4 + 1], BUFFER_NAME[Index * 4 + 2], BUFFER_NAME[Index * 4 + 3]);");
+	static constexpr TCHAR const* RWHLSLBufferWrite = TEXT("BUFFER_NAME[Index * 4 + 0] = Value.x, BUFFER_NAME[Index * 4 + 1] = Value.y, BUFFER_NAME[Index * 4 + 2] = Value.z, BUFFER_NAME[Index * 4 + 3] = Value.w;");
+
+	static const FNiagaraTypeDefinition& GetTypeDefinition() { return FNiagaraTypeDefinition::GetColorDef(); }
+	static const FLinearColor GetDefaultValue() { return FLinearColor::White; }
+};
+
+template<>
+struct FNDIArrayImplHelper<FQuat4f> : public FNDIArrayImplHelperBase<FQuat4f>
+{
+	typedef FQuat4f TVMArrayType;
+
+	static constexpr TCHAR const* HLSLVariableType = TEXT("float4");
+	static constexpr EPixelFormat ReadPixelFormat = PF_A32B32G32R32F;
+	static constexpr TCHAR const* ReadHLSLBufferType = TEXT("float4");
+	static constexpr TCHAR const* ReadHLSLBufferRead = TEXT("Value = BUFFER_NAME[Index];");
+	static constexpr EPixelFormat RWPixelFormat = PF_R32_FLOAT;
+	static constexpr TCHAR const* RWHLSLBufferType = TEXT("float");
+	static constexpr TCHAR const* RWHLSLBufferRead = TEXT("Value = float4(BUFFER_NAME[Index * 4 + 0], BUFFER_NAME[Index * 4 + 1], BUFFER_NAME[Index * 4 + 2], BUFFER_NAME[Index * 4 + 3]);");
+	static constexpr TCHAR const* RWHLSLBufferWrite = TEXT("BUFFER_NAME[Index * 4 + 0] = Value.x, BUFFER_NAME[Index * 4 + 1] = Value.y, BUFFER_NAME[Index * 4 + 2] = Value.z, BUFFER_NAME[Index * 4 + 3] = Value.w;");
+
+	static const FNiagaraTypeDefinition& GetTypeDefinition() { return FNiagaraTypeDefinition::GetQuatDef(); }
+	static const FQuat4f GetDefaultValue() { return FQuat4f::Identity; }
+
+	static void CopyCpuToCpuMemory(FQuat4f* Dest, const FQuat4f* Src, int32 NumElements)
+	{
+		FMemory::Memcpy(Dest, Src, NumElements * sizeof(FQuat4f));
+	}
+
+	static void CopyCpuToCpuMemory(FQuat4f* Dest, const FQuat* Src, int32 NumElements)
+	{
+		for (int32 i = 0; i < NumElements; ++i)
+		{
+			Dest[i] = FQuat4f(Src[i]);
+		}
+	}
+
+	static void CopyCpuToCpuMemory(FQuat* Dest, const FQuat4f* Src, int32 NumElements)
+	{
+		for (int32 i = 0; i < NumElements; ++i)
+		{
+			Dest[i] = FQuat(Src[i]);
+		}
+	}
+};
+
+template<>
+struct FNDIArrayImplHelper<int32> : public FNDIArrayImplHelperBase<int32>
+{
+	typedef int32 TVMArrayType;
+
+	static constexpr bool bSupportsAtomicOps = true;
+
+	static constexpr TCHAR const* HLSLVariableType = TEXT("int");
+	static constexpr EPixelFormat ReadPixelFormat = PF_R32_SINT;
+	static constexpr TCHAR const* ReadHLSLBufferType = TEXT("int");
+	static constexpr TCHAR const* ReadHLSLBufferRead = TEXT("Value = BUFFER_NAME[Index];");
+	static constexpr EPixelFormat RWPixelFormat = PF_R32_SINT;
+	static constexpr TCHAR const* RWHLSLBufferType = TEXT("int");
+	static constexpr TCHAR const* RWHLSLBufferRead = TEXT("Value = BUFFER_NAME[Index];");
+	static constexpr TCHAR const* RWHLSLBufferWrite = TEXT("BUFFER_NAME[Index] = Value;");
+
+	static const FNiagaraTypeDefinition& GetTypeDefinition() { return FNiagaraTypeDefinition::GetIntDef(); }
+	static const int32 GetDefaultValue() { return 0; }
+
+	static int32 AtomicAdd(int32* Dest, int32 Value)
+	{
+		return FPlatformAtomics::InterlockedAdd(Dest, Value);
+	}
+	static int32 AtomicMin(int32* Dest, int32 Value)
+	{
+		int32 PrevValue = *Dest;
+		while (PrevValue > Value)
+		{
+			PrevValue = FPlatformAtomics::InterlockedCompareExchange(Dest, Value, PrevValue);
+		}
+		return PrevValue;
+	}
+	static int32 AtomicMax(int32* Dest, int32 Value)
+	{
+		int32 PrevValue = *Dest;
+		while (PrevValue < Value)
+		{
+			PrevValue = FPlatformAtomics::InterlockedCompareExchange(Dest, Value, PrevValue);
+		}
+		return PrevValue;
+	}
+};
+
+template<>
+struct FNDIArrayImplHelper<uint8> : public FNDIArrayImplHelperBase<uint8>
+{
+	typedef int32 TVMArrayType;
+
+	static constexpr TCHAR const* HLSLVariableType = TEXT("int");
+	static constexpr EPixelFormat ReadPixelFormat = PF_R8_UINT;
+	static constexpr TCHAR const* ReadHLSLBufferType = TEXT("int");
+	static constexpr TCHAR const* ReadHLSLBufferRead = TEXT("Value = BUFFER_NAME[Index];");
+	static constexpr EPixelFormat RWPixelFormat = PF_R8_UINT;
+	static constexpr TCHAR const* RWHLSLBufferType = TEXT("int");
+	static constexpr TCHAR const* RWHLSLBufferRead = TEXT("Value = BUFFER_NAME[Index];");
+	static constexpr TCHAR const* RWHLSLBufferWrite = TEXT("BUFFER_NAME[Index] = Value;");
+
+	static const FNiagaraTypeDefinition& GetTypeDefinition() { return FNiagaraTypeDefinition::GetIntDef(); }
+	static const int32 GetDefaultValue() { return 0; }
+
+	static void CopyCpuToCpuMemory(uint8* Dest, const uint8* Src, int32 NumElements)
+	{
+		FMemory::Memcpy(Dest, Src, NumElements * sizeof(uint8));
+	}
+
+	static void CopyCpuToCpuMemory(uint8* Dest, const int32* Src, int32 NumElements)
+	{
+		for (int32 i = 0; i < NumElements; ++i)
+		{
+			Dest[i] = uint8(Src[i]);
+		}
+	}
+
+	static void CopyCpuToCpuMemory(int32* Dest, const uint8* Src, int32 NumElements)
+	{
+		for (int32 i = 0; i < NumElements; ++i)
+		{
+			Dest[i] = int32(Src[i]);
+		}
+	}
+};
+
+template<>
+struct FNDIArrayImplHelper<bool> : public FNDIArrayImplHelperBase<bool>
+{
+	typedef FNiagaraBool TVMArrayType;
+
+	static constexpr TCHAR const* HLSLVariableType = TEXT("bool");
+	static constexpr EPixelFormat ReadPixelFormat = PF_R8_UINT;
+	static constexpr TCHAR const* ReadHLSLBufferType = TEXT("uint");
+	static constexpr TCHAR const* ReadHLSLBufferRead = TEXT("Value = BUFFER_NAME[Index];");
+	static constexpr EPixelFormat RWPixelFormat = PF_R8_UINT;
+	static constexpr TCHAR const* RWHLSLBufferType = TEXT("uint");
+	static constexpr TCHAR const* RWHLSLBufferRead = TEXT("Value = BUFFER_NAME[Index];");
+	static constexpr TCHAR const* RWHLSLBufferWrite = TEXT("BUFFER_NAME[Index] = Value;");
+
+	static const FNiagaraTypeDefinition& GetTypeDefinition() { return FNiagaraTypeDefinition::GetBoolDef(); }
+	static const FNiagaraBool GetDefaultValue() { return FNiagaraBool::False; }
+
+	static void CopyCpuToGpuMemory(void* Dest, const bool* Src, int32 NumElements)
+	{
+		uint8* TypedDest = reinterpret_cast<uint8*>(Dest);
+		while (NumElements--)
+		{
+			*TypedDest++ = *Src++ == false ? 0 : 0xff;
+		}
+	}
+
+	static void CopyGpuToCpuMemory(void* Dest, const void* Src, int32 NumElements)
+	{
+		FNiagaraBool* TypedDest = reinterpret_cast<FNiagaraBool*>(Dest);
+		const uint8* TypedSrc = reinterpret_cast<const uint8*>(Src);
+		while (NumElements--)
+		{
+			*TypedDest++ = *TypedSrc++ == 0 ? false : true;
+		}
+	}
+};
+
+template<>
+struct FNDIArrayImplHelper<FNiagaraID> : public FNDIArrayImplHelperBase<FNiagaraID>
+{
+	typedef FNiagaraID TVMArrayType;
+
+	static constexpr TCHAR const* HLSLVariableType = TEXT("NiagaraID");
+	static constexpr EPixelFormat ReadPixelFormat = PF_R32_SINT;		// Should be int2 but we don't have it
+	static constexpr TCHAR const* ReadHLSLBufferType = TEXT("int");
+	static constexpr TCHAR const* ReadHLSLBufferRead = TEXT("Value.Index = BUFFER_NAME[Index * 2 + 0]; Value.AcquireTag = BUFFER_NAME[Index * 2 + 1];");
+	static constexpr EPixelFormat RWPixelFormat = PF_R32_SINT;
+	static constexpr TCHAR const* RWHLSLBufferType = TEXT("int");
+	static constexpr TCHAR const* RWHLSLBufferRead = TEXT("Value.Index = BUFFER_NAME[Index * 2 + 0]; Value.AcquireTag = BUFFER_NAME[Index * 2 + 1];");
+	static constexpr TCHAR const* RWHLSLBufferWrite = TEXT("BUFFER_NAME[Index * 2 + 0] = Value.Index; BUFFER_NAME[Index * 2 + 1] = Value.AcquireTag;");
+
+	static const FNiagaraTypeDefinition GetTypeDefinition() { return FNiagaraTypeDefinition(FNiagaraID::StaticStruct()); }
+	static const FNiagaraID GetDefaultValue() { return FNiagaraID(); }
+};
diff --git a/Engine/Plugins/FX/Niagara/Source/Niagara/Private/NiagaraDataInterfaceArrayInt.cpp b/Engine/Plugins/FX/Niagara/Source/Niagara/Private/NiagaraDataInterfaceArrayInt.cpp
--- a/Engine/Plugins/FX/Niagara/Source/Niagara/Private/NiagaraDataInterfaceArrayInt.cpp
+++ b/Engine/Plugins/FX/Niagara/Source/Niagara/Private/NiagaraDataInterfaceArrayInt.cpp
@@ -1,8 +1,14 @@
 // Copyright Epic Games, Inc. All Rights Reserved.
 
 #include "NiagaraDataInterfaceArrayInt.h"
+#include "NiagaraDataInterfaceArrayImpl.h"
+#include "NiagaraDataInterfaceArrayImplHelpers.h"
 #include "NiagaraRenderer.h"
 
 #include UE_INLINE_GENERATED_CPP_BY_NAME(NiagaraDataInterfaceArrayInt)
 
 static_assert(sizeof(bool) == sizeof(uint8), "Bool != 1 byte this will mean the GPU array does not match in size");
+
+NDIARRAY_GENERATE_IMPL(UNiagaraDataInterfaceArrayInt32, int32, IntData)
+NDIARRAY_GENERATE_IMPL_LWC(UNiagaraDataInterfaceArrayUInt8, uint8, IntData)
+NDIARRAY_GENERATE_IMPL(UNiagaraDataInterfaceArrayBool, bool, BoolData)
diff --git a/Engine/Plugins/FX/Niagara/Source/Niagara/Private/NiagaraDataInterfaceArrayNiagaraID.cpp b/Engine/Plugins/FX/Niagara/Source/Niagara/Private/NiagaraDataInterfaceArrayNiagaraID.cpp
--- a/Engine/Plugins/FX/Niagara/Source/Niagara/Private/NiagaraDataInterfaceArrayNiagaraID.cpp
+++ b/Engine/Plugins/FX/Niagara/Source/Niagara/Private/NiagaraDataInterfaceArrayNiagaraID.cpp
@@ -1,7 +1,10 @@
 // Copyright Epic Games, Inc. All Rights Reserved.
 
 #include "NiagaraDataInterfaceArrayNiagaraID.h"
+#include "NiagaraDataInterfaceArrayImpl.h"
+#include "NiagaraDataInterfaceArrayImplHelpers.h"
 #include "NiagaraRenderer.h"
 
 #include UE_INLINE_GENERATED_CPP_BY_NAME(NiagaraDataInterfaceArrayNiagaraID)
 
+NDIARRAY_GENERATE_IMPL(UNiagaraDataInterfaceArrayNiagaraID, FNiagaraID, IntData)
