Fixed transform function descriptions in camera DI

diff --git a/Engine/Plugins/FX/Niagara/Source/Niagara/Private/NiagaraDataInterfaceCamera.cpp b/Engine/Plugins/FX/Niagara/Source/Niagara/Private/NiagaraDataInterfaceCamera.cpp
--- a/Engine/Plugins/FX/Niagara/Source/Niagara/Private/NiagaraDataInterfaceCamera.cpp
+++ b/Engine/Plugins/FX/Niagara/Source/Niagara/Private/NiagaraDataInterfaceCamera.cpp
@@ -180,12 +180,12 @@ void UNiagaraDataInterfaceCamera::GetFunctions(TArray<FNiagaraFunctionSignature>
 	Sig.bSupportsCPU = false;
 	Sig.AddInput(FNiagaraVariable(FNiagaraTypeDefinition(GetClass()), TEXT("Camera interface")));
 	Sig.AddOutput(FNiagaraVariable(FNiagaraTypeDefinition::GetMatrix4Def(), TEXT("World To Clip Transform")), LOCTEXT("WorldToClipTransformDescription", "Transforms a world space position to clip space"));
-	Sig.AddOutput(FNiagaraVariable(FNiagaraTypeDefinition::GetMatrix4Def(), TEXT("Translated World To Clip Transform")), LOCTEXT("TranslatedWorldToClipTransformDescription", "Same as world to clip transform, but the camera position is at the origin when camera translation is enabled.\nThis allows for more precision of the transform when dealing with big coordinates.\nTo use this transform, you first need to subtract PreViewTranslation from your transform target."));
+	Sig.AddOutput(FNiagaraVariable(FNiagaraTypeDefinition::GetMatrix4Def(), TEXT("Translated World To Clip Transform")), LOCTEXT("TranslatedWorldToClipTransformDescription", "Same as world to clip transform, but the camera position is at the origin when camera translation is enabled.\nThis allows for more precision of the transform when dealing with big coordinates.\nTo use this transform, you first need to add PreViewTranslation to your transform target."));
 	Sig.AddOutput(FNiagaraVariable(FNiagaraTypeDefinition::GetMatrix4Def(), TEXT("Clip To World Transform")), LOCTEXT("ClipToWorldTransformDescription", "Transforms a clip space position to world space"));
 	Sig.AddOutput(FNiagaraVariable(FNiagaraTypeDefinition::GetMatrix4Def(), TEXT("Clip To View Transform")), LOCTEXT("ClipToViewTransformDescription", "Transforms a clip space position to view space"));
-	Sig.AddOutput(FNiagaraVariable(FNiagaraTypeDefinition::GetMatrix4Def(), TEXT("Clip To Translated World Transform")), LOCTEXT("ClipToTranslatedWorldTransformDescription", "The inverse of the translated world to clip transform. When camera translation is enabled, the transform is such that the camera position is at the origin.\nThis allows for more precision of the transform when dealing with big coordinates.\nTo get the actual world space position, you need to add PreViewTranslation to the transform result."));
+	Sig.AddOutput(FNiagaraVariable(FNiagaraTypeDefinition::GetMatrix4Def(), TEXT("Clip To Translated World Transform")), LOCTEXT("ClipToTranslatedWorldTransformDescription", "The inverse of the translated world to clip transform. When camera translation is enabled, the transform is such that the camera position is at the origin.\nThis allows for more precision of the transform when dealing with big coordinates.\nTo get the actual world space position, you need to subtract PreViewTranslation from the transform result."));
 	Sig.AddOutput(FNiagaraVariable(FNiagaraTypeDefinition::GetMatrix4Def(), TEXT("Screen To World Transform")), LOCTEXT("ScreenToWorldTransformDescription", "Converts a screen space position a world position.\nExample combining the transform with scene depth:\nfloat3 WorldPosition = mul(float4(ScreenPosition * SceneDepth, SceneDepth, 1), ScreenToWorld).xyz;"));
-	Sig.AddOutput(FNiagaraVariable(FNiagaraTypeDefinition::GetMatrix4Def(), TEXT("Screen To Translated World Transform")), LOCTEXT("ScreenToTranslatedWorldTransformDescription", "Same as the screen to world transform, but with the assumption that the camera is at the origin.\nThis allows for more precision of the transform when dealing with big coordinates.\nTo get the actual world space position, you need to add PreViewTranslation to the transform result."));
+	Sig.AddOutput(FNiagaraVariable(FNiagaraTypeDefinition::GetMatrix4Def(), TEXT("Screen To Translated World Transform")), LOCTEXT("ScreenToTranslatedWorldTransformDescription", "Same as the screen to world transform, but with the assumption that the camera is at the origin.\nThis allows for more precision of the transform when dealing with big coordinates.\nTo get the actual world space position, you need to subtract PreViewTranslation from the transform result."));
 	Sig.AddOutput(FNiagaraVariable(FNiagaraTypeDefinition::GetMatrix4Def(), TEXT("Clip To Previous Clip Transform")), LOCTEXT("ClipToPreviousClipTransformDescription", "Transforms from a current clip space position to last frame's clip space position.\nThis can be used to calculate for example post process velocity."));
 	OutFunctions.Add(Sig);
 
@@ -200,8 +200,8 @@ void UNiagaraDataInterfaceCamera::GetFunctions(TArray<FNiagaraFunctionSignature>
 	Sig.bRequiresContext = false;
 	Sig.bSupportsCPU = false;
 	Sig.AddInput(FNiagaraVariable(FNiagaraTypeDefinition(GetClass()), TEXT("Camera interface")));
-	Sig.AddOutput(FNiagaraVariable(FNiagaraTypeDefinition::GetMatrix4Def(), TEXT("Translated World To View Transform")), LOCTEXT("TranslatedWorldToViewTransformDescription", "Transforms a position from world to view, but the view position is assumed to be at the origin.\nThis allows for more precision of the transform when dealing with big coordinates.\nTo use this transform, you first need to subtract PreViewTranslation from your transform target."));
-	Sig.AddOutput(FNiagaraVariable(FNiagaraTypeDefinition::GetMatrix4Def(), TEXT("View To Translated World Transform")), LOCTEXT("ViewToTranslatedWorldTransformDescription", "The inverse of the translated world to view transform. When camera translation is enabled, the transform is such that the camera position is at the origin.\nThis allows for more precision of the transform when dealing with big coordinates.\nTo get the actual world space position, you need to add PreViewTranslation to the transform result."));
+	Sig.AddOutput(FNiagaraVariable(FNiagaraTypeDefinition::GetMatrix4Def(), TEXT("Translated World To View Transform")), LOCTEXT("TranslatedWorldToViewTransformDescription", "Transforms a position from world to view, but the view position is assumed to be at the origin.\nThis allows for more precision of the transform when dealing with big coordinates.\nTo use this transform, you first need to add PreViewTranslation to your transform target."));
+	Sig.AddOutput(FNiagaraVariable(FNiagaraTypeDefinition::GetMatrix4Def(), TEXT("View To Translated World Transform")), LOCTEXT("ViewToTranslatedWorldTransformDescription", "The inverse of the translated world to view transform. When camera translation is enabled, the transform is such that the camera position is at the origin.\nThis allows for more precision of the transform when dealing with big coordinates.\nTo get the actual world space position, you need to subtract PreViewTranslation from the transform result."));
 	Sig.AddOutput(FNiagaraVariable(FNiagaraTypeDefinition::GetMatrix4Def(), TEXT("Translated World To Camera View Transform")), LOCTEXT("TranslatedWorldToCameraTransformDescription", "Same as 'Translated World To View', but transforms to the camera position instead of the view position."));
 	Sig.AddOutput(FNiagaraVariable(FNiagaraTypeDefinition::GetMatrix4Def(), TEXT("Camera View To Translated World Transform")), LOCTEXT("CameraToTranslatedWorldTransformDescription", "Same as 'View To Translated World', but transforms from the camera position instead of the view position."));
 	Sig.AddOutput(FNiagaraVariable(FNiagaraTypeDefinition::GetMatrix4Def(), TEXT("View To Clip Transform")), LOCTEXT("ViewToClipTransformDescription", "Transforms a view space position to clip space"));
