Fix for the Apply button only checking for the exposed version guid instead of the currently edited version

diff --git a/Engine/Plugins/FX/Niagara/Source/NiagaraEditor/Private/Toolkits/NiagaraScriptToolkit.cpp b/Engine/Plugins/FX/Niagara/Source/NiagaraEditor/Private/Toolkits/NiagaraScriptToolkit.cpp
--- a/Engine/Plugins/FX/Niagara/Source/NiagaraEditor/Private/Toolkits/NiagaraScriptToolkit.cpp
+++ b/Engine/Plugins/FX/Niagara/Source/NiagaraEditor/Private/Toolkits/NiagaraScriptToolkit.cpp
@@ -767,7 +767,7 @@ FText FNiagaraScriptToolkit::GetVersionMenuLabel(FNiagaraAssetVersion Version) c
 
 bool FNiagaraScriptToolkit::IsEditScriptDifferentFromOriginalScript() const
 {
-	return OriginalNiagaraScript.Script->GetBaseChangeID() != EditedNiagaraScript.Script->GetBaseChangeID();
+	return OriginalNiagaraScript.Script->GetBaseChangeID(EditedNiagaraScript.Version) != EditedNiagaraScript.Script->GetBaseChangeID(EditedNiagaraScript.Version);
 }
 
 void FNiagaraScriptToolkit::OnApply()
