Pause/Resume scalability state

diff --git a/Engine/Plugins/FX/Niagara/Source/Niagara/Classes/NiagaraEffectType.h b/Engine/Plugins/FX/Niagara/Source/Niagara/Classes/NiagaraEffectType.h
--- a/Engine/Plugins/FX/Niagara/Source/Niagara/Classes/NiagaraEffectType.h
+++ b/Engine/Plugins/FX/Niagara/Source/Niagara/Classes/NiagaraEffectType.h
@@ -24,8 +24,8 @@ enum class ENiagaraCullReaction
 	DeactivateResume UMETA(DisplayName = "Asleep"),
 	/** The system instance will be deactivated and particles killed immediately. Will reactivate when it passes cull tests again. */
 	DeactivateImmediateResume UMETA(DisplayName = "Asleep and Clear"),
-	/** The system instance will be paused but will resume ticking when it passes cull tests again. */
-	//PauseResume,
+	/** The system instance will be paused, maintaining it's current state, but will resume ticking when it passes cull tests again. */
+	PauseResume UMETA(DisplayName = "Pause"),
 };
 
 /** Controls how often should we update scalability states for these effects. */
diff --git a/Engine/Plugins/FX/Niagara/Source/Niagara/Private/NiagaraComponent.cpp b/Engine/Plugins/FX/Niagara/Source/Niagara/Private/NiagaraComponent.cpp
--- a/Engine/Plugins/FX/Niagara/Source/Niagara/Private/NiagaraComponent.cpp
+++ b/Engine/Plugins/FX/Niagara/Source/Niagara/Private/NiagaraComponent.cpp
@@ -486,6 +486,7 @@ UNiagaraComponent::UNiagaraComponent(const FObjectInitializer& ObjectInitializer
 	, bForceLocalPlayerEffect(false)
 	, bIsCulledByScalability(false)
 	, bDuringUpdateContextReset(false)
+	, bDesiredPauseState(false)
 	//, bIsChangingAutoAttachment(false)
 	, ScalabilityManagerHandle(INDEX_NONE)
 	, ForceUpdateTransformTime(0.0f)
@@ -923,19 +924,58 @@ void UNiagaraComponent::AdvanceSimulationByTime(float SimulateTime, float TickDe
 
 void UNiagaraComponent::SetPaused(bool bInPaused)
 {
-	if (SystemInstanceController.IsValid())
+	bDesiredPauseState = bInPaused;
+
+	bool bIsActuallyActive = SystemInstanceController.IsValid() &&
+		SystemInstanceController->GetRequestedExecutionState() == ENiagaraExecutionState::Active &&
+		SystemInstanceController->GetActualExecutionState() == ENiagaraExecutionState::Active;
+
+	if(!bIsActuallyActive)
+	{
+		//Don't allow pausing when the system is not active.
+		//Activate/Deactivate should clear the pause state.
+		return;
+	}
+
+	if(bIsActuallyActive && !bInPaused)
 	{
-		SystemInstanceController->SetPaused(bInPaused);
+		//If we're active and we unpause, try to re-register with the scalability manager
+		RegisterWithScalabilityManager();
 	}
+	else 
+	{
+		//If we're being paused externally, unregister from scalability.
+		UnregisterWithScalabilityManager();
+	}
+	
+	SetPausedInternal(bDesiredPauseState, false);
 }
 
 bool UNiagaraComponent::IsPaused()const
 {
+	//check the system instance is actually in the right state. 
+	//In cases where we don't have a system instance or we're culled by scalability, the instance can have a differnt internal state to the user's desired state.
+	check(SystemInstanceController.IsValid() == false || bIsCulledByScalability || SystemInstanceController->IsPaused() == bDesiredPauseState);
+	return bDesiredPauseState;
+}
+
+void UNiagaraComponent::SetPausedInternal(bool bInPaused, bool bIsScalabilityCull)
+{
+	if(bIsScalabilityCull)
+	{
+		bIsCulledByScalability = bInPaused;
+	}
+	
 	if (SystemInstanceController.IsValid())
 	{
-		return SystemInstanceController->IsPaused();
+		bool bShouldPauseFromScalability = false;
+		check(GetAsset());
+		if(UNiagaraEffectType* FXType = GetAsset()->GetEffectType())
+		{
+			bShouldPauseFromScalability = FXType->CullReaction == ENiagaraCullReaction::PauseResume && bIsCulledByScalability;
+		}
+		SystemInstanceController->SetPaused(bShouldPauseFromScalability || bDesiredPauseState);
 	}
-	return false;
 }
 
 UNiagaraDataInterface* UNiagaraComponent::GetDataInterface(const FString& Name)
@@ -1137,6 +1177,13 @@ void UNiagaraComponent::ActivateInternal(bool bReset /* = false */, bool bIsScal
 
 	bIsCulledByScalability = false;
 	
+	//Ensure we're not paused. Non-scalability activate calls should override pause state.
+	if(bIsScalabilityCull == false)
+	{
+		bDesiredPauseState = false;
+		SetPausedInternal(false, false);
+	}
+
 	//Update our owner scalability state but don't register with the manager.
 	ResolveOwnerAllowsScalability(false);
 	
@@ -1174,6 +1221,10 @@ void UNiagaraComponent::ActivateInternal(bool bReset /* = false */, bool bIsScal
 	{
 		return;
 	}
+	
+	//Unpause on activate if needed.
+	checkf(!(bIsCulledByScalability && IsPaused()), TEXT("We should never be paused when being re-activated by scalability. We unregister from the manager when paused so we should never get here."));
+	SetPausedInternal(false, false);
 
 	DestroyCullProxy();
 
@@ -1268,6 +1319,12 @@ void UNiagaraComponent::ActivateInternal(bool bReset /* = false */, bool bIsScal
 			ScopedComponentTickEnabled.bTickEnabled = true;
 		}
 	}
+
+	if(bIsScalabilityCull)
+	{
+		//If this is an activate call from a scalability cull we may need to begin in a paused state if someone paused us while we were culled from scalability.
+		SetPausedInternal(bDesiredPauseState, false);
+	}
 }
 
 void UNiagaraComponent::Deactivate()
@@ -1294,6 +1351,10 @@ void UNiagaraComponent::DeactivateInternal(bool bIsScalabilityCull /* = false */
 		// Unregister with the scalability manager if this is a genuine deactivation from outside.
 		// The scalability manager itself can call this function when culling systems.
 		UnregisterWithScalabilityManager();
+	
+		//Ensure we're not paused.
+		bDesiredPauseState = false;
+		SetPausedInternal(false, false);
 	}
 
 	if (IsActive() && SystemInstanceController)
diff --git a/Engine/Plugins/FX/Niagara/Source/Niagara/Private/NiagaraScalabilityManager.cpp b/Engine/Plugins/FX/Niagara/Source/Niagara/Private/NiagaraScalabilityManager.cpp
--- a/Engine/Plugins/FX/Niagara/Source/Niagara/Private/NiagaraScalabilityManager.cpp
+++ b/Engine/Plugins/FX/Niagara/Source/Niagara/Private/NiagaraScalabilityManager.cpp
@@ -378,6 +378,7 @@ bool FNiagaraScalabilityManager::ApplyScalabilityState(int32 ComponentIndex, ENi
 			case ENiagaraCullReaction::DeactivateImmediate:			Component->DeactivateImmediateInternal(false); bContinueIteration = false;  break; //We don't increment CompIdx here as this call will remove an entry from ManagedObjects;
 			case ENiagaraCullReaction::DeactivateResume:			Component->DeactivateInternal(true); break;
 			case ENiagaraCullReaction::DeactivateImmediateResume:	Component->DeactivateImmediateInternal(true); break;
+			case ENiagaraCullReaction::PauseResume:					Component->SetPausedInternal(true,true); break;
 			};
 		}
 		else
@@ -386,7 +387,14 @@ bool FNiagaraScalabilityManager::ApplyScalabilityState(int32 ComponentIndex, ENi
 			{
 				UE_LOG(LogNiagara, Error, TEXT("Niagara Component is incorrectly still registered with the scalability manager. %d - %s "), (int32)CullReaction, *Component->GetAsset()->GetFullName());
 			}
-			Component->ActivateInternal(false, true);
+			else if(CullReaction == ENiagaraCullReaction::PauseResume)
+			{
+				Component->SetPausedInternal(false, true);
+			}
+			else
+			{
+				Component->ActivateInternal(false, true);
+			}
 		}
 
 		//TODO: Beyond culling by hard limits here we could progressively scale down fx by biasing detail levels they use. Could also introduce some budgeting here like N at lvl 0, M at lvl 1 etc.
diff --git a/Engine/Plugins/FX/Niagara/Source/Niagara/Public/NiagaraComponent.h b/Engine/Plugins/FX/Niagara/Source/Niagara/Public/NiagaraComponent.h
--- a/Engine/Plugins/FX/Niagara/Source/Niagara/Public/NiagaraComponent.h
+++ b/Engine/Plugins/FX/Niagara/Source/Niagara/Public/NiagaraComponent.h
@@ -218,6 +218,8 @@ private:
 	void ActivateInternal(bool bReset, bool bIsScalabilityCull);
 	void DeactivateInternal(bool bIsScalabilityCull);
 	void DeactivateImmediateInternal(bool bIsScalabilityCull);
+	void SetPausedInternal(bool bInPaused, bool bIsScalabilityCull);	
+
 
 	bool ShouldPreCull();
 	void RegisterWithScalabilityManager();
@@ -814,6 +816,9 @@ private:
 	/** True if we're currently inside an update context reset. This will prevent us from doing some completion work such as releaseing to the pool or auto destroy etc during a reset. */
 	uint32 bDuringUpdateContextReset : 1;
 
+	/** Stores the current state for pause/unpause desired by the use. Allows us to pause/unpause correctly while also pausing/unpausing via scalability. */
+	uint32 bDesiredPauseState : 1;
+
 	/** Restore relative transform from auto attachment and optionally detach from parent (regardless of whether it was an auto attachment). */
 	void CancelAutoAttachment(bool bDetachFromParent);
 
