Fix issue with some function call nodes and dynamic pins.
- Updating existing function call nodes with dynamic pins to have correct dynamic pins.
- Creating add pins where needed, some custom hlsl nodes had lost their add pins.
- Misc crash fix I hit during.

diff --git a/Engine/Plugins/FX/Niagara/Source/NiagaraCore/Public/NiagaraCustomVersion.h b/Engine/Plugins/FX/Niagara/Source/NiagaraCore/Public/NiagaraCustomVersion.h
--- a/Engine/Plugins/FX/Niagara/Source/NiagaraCore/Public/NiagaraCustomVersion.h
+++ b/Engine/Plugins/FX/Niagara/Source/NiagaraCore/Public/NiagaraCustomVersion.h
@@ -184,6 +184,8 @@ struct FNiagaraCustomVersion
 
 		AddVariadicParametersToGPUFunctionInfo, // Adds list of variadic parameters to the information about GPU functions.
 
+		DynamicPinNodeFixup, // Some data fixup for NiagaraNodeWithDynamicPins.
+
 		// DO NOT ADD A NEW VERSION UNLESS YOU HAVE TALKED TO THE NIAGARA LEAD. Mismanagement of these versions can lead to data loss if it is adjusted in multiple streams simultaneously.
 		// -----<new versions can be added above this line>  -------------------------------------------------
 		VersionPlusOne,
diff --git a/Engine/Plugins/FX/Niagara/Source/NiagaraEditor/Private/NiagaraNodeCustomHlsl.cpp b/Engine/Plugins/FX/Niagara/Source/NiagaraEditor/Private/NiagaraNodeCustomHlsl.cpp
--- a/Engine/Plugins/FX/Niagara/Source/NiagaraEditor/Private/NiagaraNodeCustomHlsl.cpp
+++ b/Engine/Plugins/FX/Niagara/Source/NiagaraEditor/Private/NiagaraNodeCustomHlsl.cpp
@@ -632,12 +632,7 @@ void UNiagaraNodeCustomHlsl::RebuildSignatureFromPins()
 		Sig.Outputs.Add(Schema->PinToNiagaraVariable(Pin, false));
 	}
 
-	//Allow signature to have variadic params so we still get the add pins.
-	Sig.RequiredInputs = Sig.Inputs.Num();
-	Sig.RequiredOutputs = Sig.Outputs.Num();
-
 	Signature = Sig;
-
 }
 
 #undef LOCTEXT_NAMESPACE
diff --git a/Engine/Plugins/FX/Niagara/Source/NiagaraEditor/Private/NiagaraNodeFunctionCall.cpp b/Engine/Plugins/FX/Niagara/Source/NiagaraEditor/Private/NiagaraNodeFunctionCall.cpp
--- a/Engine/Plugins/FX/Niagara/Source/NiagaraEditor/Private/NiagaraNodeFunctionCall.cpp
+++ b/Engine/Plugins/FX/Niagara/Source/NiagaraEditor/Private/NiagaraNodeFunctionCall.cpp
@@ -101,6 +101,47 @@ void UNiagaraNodeFunctionCall::PostLoad()
 			FNiagaraStackGraphUtilities::PopulateFunctionCallNameBindings(*this);
 		}
 	}
+	else
+	{
+		//Perform some fix up due to a bug with dynamic pins and add pins.
+		//Rebuild the dynamic pins and ensure we have add pins where needed.		
+		const int32 NiagaraVer = GetLinkerCustomVersion(FNiagaraCustomVersion::GUID);
+		if (NiagaraVer < FNiagaraCustomVersion::DynamicPinNodeFixup && AllowDynamicPins())
+		{
+			DynamicPins.Empty();
+			bool bFoundInputAdd = false;
+			bool bFoundOutputAdd = false;
+			for (const UEdGraphPin* Pin : Pins)
+			{
+				if (IsAddPin(Pin))
+				{
+					if (Pin->Direction == EGPD_Input)
+					{
+						bFoundInputAdd = true;
+					}
+					else
+					{
+						bFoundOutputAdd = true;
+					}
+				}
+				else
+				{
+					DynamicPins.Add(Pin->PinId);
+				}
+			}
+
+			if (!bFoundInputAdd)
+			{
+				CreateAddPin(EGPD_Input);
+			}
+
+			if (!bFoundOutputAdd)
+			{
+				CreateAddPin(EGPD_Output);
+			}
+		}
+	}
+
 
 	// Allow data interfaces an opportunity to intercept changes
 	if (Signature.IsValid() && Signature.bMemberFunction)
@@ -391,7 +432,7 @@ PRAGMA_ENABLE_DEPRECATION_WARNINGS;
 
 		if (AllowDynamicPins())
 		{
-			if(Signature.IsValid())
+			if(Signature.IsValid() && (Signature.VariadicInput() || Signature.VariadicOutput()))
 			{
 				if(Signature.VariadicInput())
 				{
diff --git a/Engine/Plugins/FX/Niagara/Source/NiagaraEditor/Public/NiagaraNodeCustomHlsl.h b/Engine/Plugins/FX/Niagara/Source/NiagaraEditor/Public/NiagaraNodeCustomHlsl.h
--- a/Engine/Plugins/FX/Niagara/Source/NiagaraEditor/Public/NiagaraNodeCustomHlsl.h
+++ b/Engine/Plugins/FX/Niagara/Source/NiagaraEditor/Public/NiagaraNodeCustomHlsl.h
@@ -56,6 +56,7 @@ public:
 
 protected:
 	virtual bool AllowDynamicPins() const override { return true; }
+	virtual bool CanModifyExistingPins() const { return true; }
 	virtual bool GetValidateDataInterfaces() const override { return false; }
 	virtual bool VerifyEditablePinName(const FText& InName, FText& OutErrorMessage, const UEdGraphPin* InGraphPinObj) const override;
 	virtual bool IsPinNameEditableUponCreation(const UEdGraphPin* Pin) const override;
