Removing private includes from public Niagara files.
Primarily achieved by moving code from headers into cpp files.

diff --git a/Engine/Plugins/FX/Niagara/Source/Niagara/Classes/NiagaraCollision.h b/Engine/Plugins/FX/Niagara/Source/Niagara/Classes/NiagaraCollision.h
--- a/Engine/Plugins/FX/Niagara/Source/Niagara/Classes/NiagaraCollision.h
+++ b/Engine/Plugins/FX/Niagara/Source/Niagara/Classes/NiagaraCollision.h
@@ -3,13 +3,10 @@
 #pragma once
 
 #include "NiagaraEvents.h"
-#include "NiagaraStats.h"
 #include "WorldCollision.h"
 
 #include "NiagaraCollision.generated.h"
 
-DECLARE_CYCLE_STAT(TEXT("Collision"), STAT_NiagaraCollision, STATGROUP_Niagara);
-
 class FNiagaraDataSet;
 
 UENUM()
diff --git a/Engine/Plugins/FX/Niagara/Source/Niagara/Classes/NiagaraDataInterfaceCurveBase.h b/Engine/Plugins/FX/Niagara/Source/Niagara/Classes/NiagaraDataInterfaceCurveBase.h
--- a/Engine/Plugins/FX/Niagara/Source/Niagara/Classes/NiagaraDataInterfaceCurveBase.h
+++ b/Engine/Plugins/FX/Niagara/Source/Niagara/Classes/NiagaraDataInterfaceCurveBase.h
@@ -7,35 +7,8 @@
 #include "StaticMeshResources.h"
 #include "Curves/RichCurve.h"
 #include "NiagaraDataInterface.h"
-#include "NiagaraStats.h"
 #include "NiagaraDataInterfaceCurveBase.generated.h"
 
-/** Base class for curve data proxy data. */
-struct FNiagaraDataInterfaceProxyCurveBase : public FNiagaraDataInterfaceProxy
-{
-	virtual ~FNiagaraDataInterfaceProxyCurveBase()
-	{
-		check(IsInRenderingThread());
-		DEC_MEMORY_STAT_BY(STAT_NiagaraGPUDataInterfaceMemory, CurveLUT.NumBytes);
-		CurveLUT.Release();
-	}
-
-	float LUTMinTime = 0.0f;
-	float LUTMaxTime = 0.0f;
-	float LUTInvTimeRange = 0.0f;
-	uint32 CurveLUTNumMinusOne = 0;
-	uint32 LUTOffset = INDEX_NONE;
-	FReadBuffer CurveLUT;
-
-	virtual int32 PerInstanceDataPassedToRenderThreadSize() const override
-	{
-		return 0;
-	}
-
-	// @todo REMOVEME
-	virtual void ConsumePerInstanceDataFromGameThread(void* PerInstanceData, const FNiagaraSystemInstanceID& Instance) override { check(false); }
-};
-
 /** Base class for curve data interfaces which facilitates handling the curve data in a standardized way. */
 UCLASS(EditInlineNew, abstract)
 class NIAGARA_API UNiagaraDataInterfaceCurveBase : public UNiagaraDataInterface
@@ -113,39 +86,8 @@ public:
 #endif
 
 public:
-	UNiagaraDataInterfaceCurveBase()
-		: LUTMinTime(0.0f)
-		, LUTMaxTime(1.0f)
-		, LUTInvTimeRange(1.0f)
-		, bUseLUT(true)
-		, bExposeCurve(false)
-#if WITH_EDITORONLY_DATA
-		, bOptimizeLUT(true)
-		, bOverrideOptimizeThreshold(false)
-		, HasEditorData(true)
-		, OptimizeThreshold(DefaultOptimizeThreshold)
-#endif
-		, ExposedName(TEXT("Curve"))
-	{
-		Proxy.Reset(new FNiagaraDataInterfaceProxyCurveBase());
-	}
-
-	UNiagaraDataInterfaceCurveBase(FObjectInitializer const& ObjectInitializer)
-		: LUTMinTime(0.0f)
-		, LUTMaxTime(1.0f)
-		, LUTInvTimeRange(1.0f)
-		, bUseLUT(true)
-		, bExposeCurve(false)
-#if WITH_EDITORONLY_DATA
-		, bOptimizeLUT(true)
-		, bOverrideOptimizeThreshold(false)
-		, HasEditorData(true)
-		, OptimizeThreshold(DefaultOptimizeThreshold)
-#endif
-		, ExposedName(TEXT("Curve"))
-	{
-		Proxy.Reset(new FNiagaraDataInterfaceProxyCurveBase());
-	}
+	UNiagaraDataInterfaceCurveBase();
+	UNiagaraDataInterfaceCurveBase(FObjectInitializer const& ObjectInitializer);
 
 	enum
 	{
diff --git a/Engine/Plugins/FX/Niagara/Source/Niagara/Classes/NiagaraDataInterfaceGrid2DCollection.h b/Engine/Plugins/FX/Niagara/Source/Niagara/Classes/NiagaraDataInterfaceGrid2DCollection.h
--- a/Engine/Plugins/FX/Niagara/Source/Niagara/Classes/NiagaraDataInterfaceGrid2DCollection.h
+++ b/Engine/Plugins/FX/Niagara/Source/Niagara/Classes/NiagaraDataInterfaceGrid2DCollection.h
@@ -6,7 +6,6 @@
 #include "ClearQuad.h"
 #include "NiagaraComponent.h"
 #include "NiagaraRenderGraphUtils.h"
-#include "NiagaraStats.h"
 #include "NiagaraDataInterfaceRWUtils.h"
 
 #include "NiagaraDataInterfaceGrid2DCollection.generated.h"
diff --git a/Engine/Plugins/FX/Niagara/Source/Niagara/Classes/NiagaraDataInterfaceGrid2DCollectionReader.h b/Engine/Plugins/FX/Niagara/Source/Niagara/Classes/NiagaraDataInterfaceGrid2DCollectionReader.h
--- a/Engine/Plugins/FX/Niagara/Source/Niagara/Classes/NiagaraDataInterfaceGrid2DCollectionReader.h
+++ b/Engine/Plugins/FX/Niagara/Source/Niagara/Classes/NiagaraDataInterfaceGrid2DCollectionReader.h
@@ -6,7 +6,6 @@
 #include "ClearQuad.h"
 #include "NiagaraComponent.h"
 #include "NiagaraDataInterfaceGrid2DCollection.h"
-#include "NiagaraStats.h"
 
 #include "NiagaraDataInterfaceGrid2DCollectionReader.generated.h"
 
diff --git a/Engine/Plugins/FX/Niagara/Source/Niagara/Classes/NiagaraDataInterfaceGrid3DCollection.h b/Engine/Plugins/FX/Niagara/Source/Niagara/Classes/NiagaraDataInterfaceGrid3DCollection.h
--- a/Engine/Plugins/FX/Niagara/Source/Niagara/Classes/NiagaraDataInterfaceGrid3DCollection.h
+++ b/Engine/Plugins/FX/Niagara/Source/Niagara/Classes/NiagaraDataInterfaceGrid3DCollection.h
@@ -6,7 +6,7 @@
 #include "ClearQuad.h"
 #include "NiagaraComponent.h"
 #include "NiagaraRenderGraphUtils.h"
-#include "NiagaraStats.h"
+#include "RHIUtilities.h"
 
 #include "NiagaraDataInterfaceGrid3DCollection.generated.h"
 
diff --git a/Engine/Plugins/FX/Niagara/Source/Niagara/Classes/NiagaraDataInterfaceGrid3DCollectionReader.h b/Engine/Plugins/FX/Niagara/Source/Niagara/Classes/NiagaraDataInterfaceGrid3DCollectionReader.h
--- a/Engine/Plugins/FX/Niagara/Source/Niagara/Classes/NiagaraDataInterfaceGrid3DCollectionReader.h
+++ b/Engine/Plugins/FX/Niagara/Source/Niagara/Classes/NiagaraDataInterfaceGrid3DCollectionReader.h
@@ -6,7 +6,6 @@
 #include "ClearQuad.h"
 #include "NiagaraComponent.h"
 #include "NiagaraDataInterfaceGrid3DCollection.h"
-#include "NiagaraStats.h"
 
 #include "NiagaraDataInterfaceGrid3DCollectionReader.generated.h"
 
diff --git a/Engine/Plugins/FX/Niagara/Source/Niagara/Classes/NiagaraDataInterfaceNeighborGrid3D.h b/Engine/Plugins/FX/Niagara/Source/Niagara/Classes/NiagaraDataInterfaceNeighborGrid3D.h
--- a/Engine/Plugins/FX/Niagara/Source/Niagara/Classes/NiagaraDataInterfaceNeighborGrid3D.h
+++ b/Engine/Plugins/FX/Niagara/Source/Niagara/Classes/NiagaraDataInterfaceNeighborGrid3D.h
@@ -6,7 +6,6 @@
 #include "NiagaraCommon.h"
 #include "NiagaraDataInterfaceRW.h"
 #include "NiagaraRenderGraphUtils.h"
-#include "NiagaraStats.h"
 
 #include "NiagaraDataInterfaceNeighborGrid3D.generated.h"
 
diff --git a/Engine/Plugins/FX/Niagara/Source/Niagara/Classes/NiagaraDataInterfaceRasterizationGrid3D.h b/Engine/Plugins/FX/Niagara/Source/Niagara/Classes/NiagaraDataInterfaceRasterizationGrid3D.h
--- a/Engine/Plugins/FX/Niagara/Source/Niagara/Classes/NiagaraDataInterfaceRasterizationGrid3D.h
+++ b/Engine/Plugins/FX/Niagara/Source/Niagara/Classes/NiagaraDataInterfaceRasterizationGrid3D.h
@@ -6,7 +6,7 @@
 #include "NiagaraDataInterfaceRW.h"
 #include "NiagaraRenderGraphUtils.h"
 #include "ClearQuad.h"
-#include "NiagaraStats.h"
+#include "RHIUtilities.h"
 
 #include "NiagaraDataInterfaceRasterizationGrid3D.generated.h"
 
diff --git a/Engine/Plugins/FX/Niagara/Source/Niagara/Classes/NiagaraDataInterfaceSpline.h b/Engine/Plugins/FX/Niagara/Source/Niagara/Classes/NiagaraDataInterfaceSpline.h
--- a/Engine/Plugins/FX/Niagara/Source/Niagara/Classes/NiagaraDataInterfaceSpline.h
+++ b/Engine/Plugins/FX/Niagara/Source/Niagara/Classes/NiagaraDataInterfaceSpline.h
@@ -5,7 +5,7 @@
 #include "NiagaraCommon.h"
 #include "NiagaraDataInterface.h"
 #include "NiagaraParameterStore.h"
-#include "NiagaraStats.h"
+#include "RHIUtilities.h"
 #include "VectorVM.h"
 #include "NiagaraDataInterfaceSpline.generated.h"
 
@@ -36,29 +36,7 @@ struct FNDISpline_InstanceData_RenderThread
 	FReadBuffer SplineScalesLUT;
 	FReadBuffer SplineRotationsLUT;
 
-	void Reset()
-	{
-		check(IsInRenderingThread());
-		DEC_MEMORY_STAT_BY(STAT_NiagaraGPUDataInterfaceMemory, SplinePositionsLUT.NumBytes);
-		DEC_MEMORY_STAT_BY(STAT_NiagaraGPUDataInterfaceMemory, SplineScalesLUT.NumBytes);
-		DEC_MEMORY_STAT_BY(STAT_NiagaraGPUDataInterfaceMemory, SplineRotationsLUT.NumBytes);
-		SplinePositionsLUT.Release();
-		SplineScalesLUT.Release();
-		SplineRotationsLUT.Release();
-
-		SplineTransform = FMatrix44f::Identity;
-		SplineTransformRotationMat = FMatrix44f::Identity;
-		SplineTransformInverse = FMatrix44f::Identity;
-		SplineTransformInverseTranspose = FMatrix44f::Identity;
-		SplineTransformRotation = FQuat4f::Identity;
-
-		DefaultUpVector = FVector3f::ZAxisVector;
-
-		SplineLength = 0.0f;
-		SplineDistanceStep = 0.0f;
-		InvSplineDistanceStep = 0.0f;
-		MaxIndex = INDEX_NONE;
-	}
+	void Reset();
 };
 
 
diff --git a/Engine/Plugins/FX/Niagara/Source/Niagara/Private/DataInterface/NiagaraDataInterfaceDataChannelRead.cpp b/Engine/Plugins/FX/Niagara/Source/Niagara/Private/DataInterface/NiagaraDataInterfaceDataChannelRead.cpp
--- a/Engine/Plugins/FX/Niagara/Source/Niagara/Private/DataInterface/NiagaraDataInterfaceDataChannelRead.cpp
+++ b/Engine/Plugins/FX/Niagara/Source/Niagara/Private/DataInterface/NiagaraDataInterfaceDataChannelRead.cpp
@@ -13,6 +13,7 @@
 #include "NiagaraDataChannel.h"
 #include "NiagaraDataChannelHandler.h"
 #include "NiagaraDataChannelDefinitions.h"
+#include "NiagaraDataChannelManager.h"
 
 #include "NiagaraRenderer.h"
 #include "NiagaraGPUSystemTick.h"
diff --git a/Engine/Plugins/FX/Niagara/Source/Niagara/Private/DataInterface/NiagaraDataInterfaceDataChannelWrite.cpp b/Engine/Plugins/FX/Niagara/Source/Niagara/Private/DataInterface/NiagaraDataInterfaceDataChannelWrite.cpp
--- a/Engine/Plugins/FX/Niagara/Source/Niagara/Private/DataInterface/NiagaraDataInterfaceDataChannelWrite.cpp
+++ b/Engine/Plugins/FX/Niagara/Source/Niagara/Private/DataInterface/NiagaraDataInterfaceDataChannelWrite.cpp
@@ -15,6 +15,7 @@
 #include "NiagaraDataChannel.h"
 #include "NiagaraDataChannelHandler.h"
 #include "NiagaraDataChannelDefinitions.h"
+#include "NiagaraDataChannelManager.h"
 
 
 #define LOCTEXT_NAMESPACE "NiagaraDataInterfaceDataChannelWrite"
diff --git a/Engine/Plugins/FX/Niagara/Source/Niagara/Private/NiagaraCollision.cpp b/Engine/Plugins/FX/Niagara/Source/Niagara/Private/NiagaraCollision.cpp
--- a/Engine/Plugins/FX/Niagara/Source/Niagara/Private/NiagaraCollision.cpp
+++ b/Engine/Plugins/FX/Niagara/Source/Niagara/Private/NiagaraCollision.cpp
@@ -5,6 +5,7 @@
 #include "NiagaraEmitterInstance.h"
 #include "NiagaraSystemInstance.h"
 #include "NiagaraComponent.h"
+#include "NiagaraStats.h"
 #include "PhysicalMaterials/PhysicalMaterial.h"
 
 DECLARE_CYCLE_STAT(TEXT("Event Emission"), STAT_NiagaraEventWrite, STATGROUP_Niagara);
diff --git a/Engine/Plugins/FX/Niagara/Source/Niagara/Private/NiagaraDataChannel.cpp b/Engine/Plugins/FX/Niagara/Source/Niagara/Private/NiagaraDataChannel.cpp
--- a/Engine/Plugins/FX/Niagara/Source/Niagara/Private/NiagaraDataChannel.cpp
+++ b/Engine/Plugins/FX/Niagara/Source/Niagara/Private/NiagaraDataChannel.cpp
@@ -5,6 +5,7 @@
 #include "CoreMinimal.h"
 #include "NiagaraWorldManager.h"
 #include "NiagaraDataChannelHandler.h"
+#include "NiagaraDataChannelManager.h"
 #include "Engine/Engine.h"
 
 #include UE_INLINE_GENERATED_CPP_BY_NAME(NiagaraDataChannel)
diff --git a/Engine/Plugins/FX/Niagara/Source/Niagara/Private/NiagaraDataChannelDefinitions.cpp b/Engine/Plugins/FX/Niagara/Source/Niagara/Private/NiagaraDataChannelDefinitions.cpp
--- a/Engine/Plugins/FX/Niagara/Source/Niagara/Private/NiagaraDataChannelDefinitions.cpp
+++ b/Engine/Plugins/FX/Niagara/Source/Niagara/Private/NiagaraDataChannelDefinitions.cpp
@@ -3,6 +3,7 @@
 #include "NiagaraDataChannelDefinitions.h"
 #include "NiagaraWorldManager.h"
 #include "NiagaraDataChannel.h"
+#include "NiagaraDataChannelManager.h"
 #include "DataInterface/NiagaraDataInterfaceDataChannelWrite.h"
 #include "DataInterface/NiagaraDataInterfaceDataChannelRead.h"
 #include "DataInterface/NiagaraDataInterfaceDataChannelSpawn.h"
diff --git a/Engine/Plugins/FX/Niagara/Source/Niagara/Private/NiagaraDataInterfaceCurveBase.cpp b/Engine/Plugins/FX/Niagara/Source/Niagara/Private/NiagaraDataInterfaceCurveBase.cpp
--- a/Engine/Plugins/FX/Niagara/Source/Niagara/Private/NiagaraDataInterfaceCurveBase.cpp
+++ b/Engine/Plugins/FX/Niagara/Source/Niagara/Private/NiagaraDataInterfaceCurveBase.cpp
@@ -11,6 +11,7 @@
 #include "NiagaraRenderer.h"
 #include "NiagaraShader.h"
 #include "NiagaraShaderParametersBuilder.h"
+#include "NiagaraStats.h"
 #include "NiagaraSystemStaticBuffers.h"
 #include "NiagaraTypes.h"
 
@@ -45,6 +46,66 @@ END_SHADER_PARAMETER_STRUCT()
 
 static const TCHAR* NDICurve_TemplateShaderFile = TEXT("/Plugin/FX/Niagara/Private/NiagaraDataInterfaceCurveTemplate.ush");
 
+/** Base class for curve data proxy data. */
+struct FNiagaraDataInterfaceProxyCurveBase : public FNiagaraDataInterfaceProxy
+{
+	virtual ~FNiagaraDataInterfaceProxyCurveBase()
+	{
+		check(IsInRenderingThread());
+		DEC_MEMORY_STAT_BY(STAT_NiagaraGPUDataInterfaceMemory, CurveLUT.NumBytes);
+		CurveLUT.Release();
+	}
+
+	float LUTMinTime = 0.0f;
+	float LUTMaxTime = 0.0f;
+	float LUTInvTimeRange = 0.0f;
+	uint32 CurveLUTNumMinusOne = 0;
+	uint32 LUTOffset = INDEX_NONE;
+	FReadBuffer CurveLUT;
+
+	virtual int32 PerInstanceDataPassedToRenderThreadSize() const override
+	{
+		return 0;
+	}
+
+	// @todo REMOVEME
+	virtual void ConsumePerInstanceDataFromGameThread(void* PerInstanceData, const FNiagaraSystemInstanceID& Instance) override { check(false); }
+};
+
+UNiagaraDataInterfaceCurveBase::UNiagaraDataInterfaceCurveBase()
+	: LUTMinTime(0.0f)
+	, LUTMaxTime(1.0f)
+	, LUTInvTimeRange(1.0f)
+	, bUseLUT(true)
+	, bExposeCurve(false)
+#if WITH_EDITORONLY_DATA
+	, bOptimizeLUT(true)
+	, bOverrideOptimizeThreshold(false)
+	, HasEditorData(true)
+	, OptimizeThreshold(DefaultOptimizeThreshold)
+#endif
+	, ExposedName(TEXT("Curve"))
+{
+	Proxy.Reset(new FNiagaraDataInterfaceProxyCurveBase());
+}
+
+UNiagaraDataInterfaceCurveBase::UNiagaraDataInterfaceCurveBase(FObjectInitializer const& ObjectInitializer)
+	: LUTMinTime(0.0f)
+	, LUTMaxTime(1.0f)
+	, LUTInvTimeRange(1.0f)
+	, bUseLUT(true)
+	, bExposeCurve(false)
+#if WITH_EDITORONLY_DATA
+	, bOptimizeLUT(true)
+	, bOverrideOptimizeThreshold(false)
+	, HasEditorData(true)
+	, OptimizeThreshold(DefaultOptimizeThreshold)
+#endif
+	, ExposedName(TEXT("Curve"))
+{
+	Proxy.Reset(new FNiagaraDataInterfaceProxyCurveBase());
+}
+
 void UNiagaraDataInterfaceCurveBase::PostLoad()
 {
 	Super::PostLoad();
diff --git a/Engine/Plugins/FX/Niagara/Source/Niagara/Private/NiagaraDataInterfaceSpline.cpp b/Engine/Plugins/FX/Niagara/Source/Niagara/Private/NiagaraDataInterfaceSpline.cpp
--- a/Engine/Plugins/FX/Niagara/Source/Niagara/Private/NiagaraDataInterfaceSpline.cpp
+++ b/Engine/Plugins/FX/Niagara/Source/Niagara/Private/NiagaraDataInterfaceSpline.cpp
@@ -5,6 +5,7 @@
 #include "NiagaraComponent.h"
 #include "NiagaraRenderer.h"
 #include "NiagaraShaderParametersBuilder.h"
+#include "NiagaraStats.h"
 #include "NiagaraSystemInstance.h"
 #include "Internationalization/Internationalization.h"
 #include "ShaderParameterUtils.h"
@@ -26,6 +27,30 @@ struct FNiagaraSplineDIFunctionVersion
 	};
 };
 
+void FNDISpline_InstanceData_RenderThread::Reset()
+{
+	check(IsInRenderingThread());
+	DEC_MEMORY_STAT_BY(STAT_NiagaraGPUDataInterfaceMemory, SplinePositionsLUT.NumBytes);
+	DEC_MEMORY_STAT_BY(STAT_NiagaraGPUDataInterfaceMemory, SplineScalesLUT.NumBytes);
+	DEC_MEMORY_STAT_BY(STAT_NiagaraGPUDataInterfaceMemory, SplineRotationsLUT.NumBytes);
+	SplinePositionsLUT.Release();
+	SplineScalesLUT.Release();
+	SplineRotationsLUT.Release();
+
+	SplineTransform = FMatrix44f::Identity;
+	SplineTransformRotationMat = FMatrix44f::Identity;
+	SplineTransformInverse = FMatrix44f::Identity;
+	SplineTransformInverseTranspose = FMatrix44f::Identity;
+	SplineTransformRotation = FQuat4f::Identity;
+
+	DefaultUpVector = FVector3f::ZAxisVector;
+
+	SplineLength = 0.0f;
+	SplineDistanceStep = 0.0f;
+	InvSplineDistanceStep = 0.0f;
+	MaxIndex = INDEX_NONE;
+}
+
 namespace NDISplineLocal
 {
 	BEGIN_SHADER_PARAMETER_STRUCT(FShaderParameters, )
diff --git a/Engine/Plugins/FX/Niagara/Source/Niagara/Private/NiagaraModule.cpp b/Engine/Plugins/FX/Niagara/Source/Niagara/Private/NiagaraModule.cpp
--- a/Engine/Plugins/FX/Niagara/Source/Niagara/Private/NiagaraModule.cpp
+++ b/Engine/Plugins/FX/Niagara/Source/Niagara/Private/NiagaraModule.cpp
@@ -9,6 +9,7 @@
 #include "NiagaraDebugVis.h"
 #include "NiagaraEvents.h"
 #include "NiagaraSettings.h"
+#include "NiagaraDataChannelManager.h"
 #include "NiagaraDataInterfaceCurlNoise.h"
 #include "UObject/Class.h"
 #include "UObject/ObjectRedirector.h"
diff --git a/Engine/Plugins/FX/Niagara/Source/Niagara/Private/NiagaraStats.h b/Engine/Plugins/FX/Niagara/Source/Niagara/Private/NiagaraStats.h
--- a/Engine/Plugins/FX/Niagara/Source/Niagara/Private/NiagaraStats.h
+++ b/Engine/Plugins/FX/Niagara/Source/Niagara/Private/NiagaraStats.h
@@ -31,3 +31,5 @@ DECLARE_STATS_GROUP(TEXT("Niagara System Counts"), STATGROUP_NiagaraSystemCounts
 
 DECLARE_DWORD_COUNTER_STAT(TEXT("Total System Instances"), STAT_TotalNiagaraSystemInstances, STATGROUP_NiagaraSystemCounts);
 DECLARE_DWORD_COUNTER_STAT(TEXT("Total Solo System Instances"), STAT_TotalNiagaraSystemInstancesSolo, STATGROUP_NiagaraSystemCounts);
+
+DECLARE_CYCLE_STAT(TEXT("Collision"), STAT_NiagaraCollision, STATGROUP_Niagara);
diff --git a/Engine/Plugins/FX/Niagara/Source/Niagara/Private/NiagaraWorldManager.cpp b/Engine/Plugins/FX/Niagara/Source/Niagara/Private/NiagaraWorldManager.cpp
--- a/Engine/Plugins/FX/Niagara/Source/Niagara/Private/NiagaraWorldManager.cpp
+++ b/Engine/Plugins/FX/Niagara/Source/Niagara/Private/NiagaraWorldManager.cpp
@@ -19,6 +19,7 @@
 #include "NiagaraStats.h"
 #include "NiagaraComponentPool.h"
 #include "NiagaraComponent.h"
+#include "NiagaraDataChannelManager.h"
 #include "NiagaraEffectType.h"
 #include "NiagaraDebugHud.h"
 #include "NiagaraGpuComputeDispatchInterface.h"
@@ -302,7 +303,7 @@ FNiagaraWorldManager::FNiagaraWorldManager()
 	: World(nullptr)
 	, ActiveNiagaraTickGroup(-1)
 	, bAppHasFocus(true)
-	, DataChannelManager(this)
+	, DataChannelManager(new FNiagaraDataChannelManager(this))
 {
 }
 
@@ -339,7 +340,7 @@ void FNiagaraWorldManager::Init(UWorld* InWorld)
 	//Possibly a later hook we can use.
 	//PrimePoolForAllSystems();
 
-	DataChannelManager.Init();
+	DataChannelManager->Init();
 
 #if WITH_NIAGARA_DEBUGGER
 	NiagaraDebugHud.Reset(new FNiagaraDebugHud(World));
@@ -444,7 +445,7 @@ void FNiagaraWorldManager::AddReferencedObjects(FReferenceCollector& Collector)
 
 	Collector.AddReferencedObjects(CullProxyMap);
 
-	DataChannelManager.AddReferencedObjects(Collector);
+	DataChannelManager->AddReferencedObjects(Collector);
 }
 
 FString FNiagaraWorldManager::GetReferencerName() const
@@ -613,7 +614,7 @@ void FNiagaraWorldManager::OnWorldBeginTearDown()
 {
 	bIsTearingDown = true;
 	ComponentPool->Cleanup(World);
-	DataChannelManager.Cleanup();
+	DataChannelManager->Cleanup();
 }
 
 void FNiagaraWorldManager::OnWorldCleanup(bool bSessionEnded, bool bCleanupResources)
@@ -621,7 +622,7 @@ void FNiagaraWorldManager::OnWorldCleanup(bool bSessionEnded, bool bCleanupResou
 	DeferredMethods.ExecuteAndClear();
 	ComponentPool->Cleanup(World);
 
-	DataChannelManager.Cleanup();
+	DataChannelManager->Cleanup();
 
 	for (int TG = 0; TG < NiagaraNumTickGroups; ++TG)
 	{
@@ -1073,7 +1074,7 @@ void FNiagaraWorldManager::Tick(ETickingGroup TickGroup, float DeltaSeconds, ELe
 	DeltaSeconds *= DebugPlaybackRate;
 
 	//Tick DataChannel Manager.
-	DataChannelManager.Tick(DeltaSeconds, TickGroup);
+	DataChannelManager->Tick(DeltaSeconds, TickGroup);
 
 	// We do book keeping in the first tick group
 	if ( TickGroup == NiagaraFirstTickGroup )
diff --git a/Engine/Plugins/FX/Niagara/Source/Niagara/Public/NiagaraWorldManager.h b/Engine/Plugins/FX/Niagara/Source/Niagara/Public/NiagaraWorldManager.h
--- a/Engine/Plugins/FX/Niagara/Source/Niagara/Public/NiagaraWorldManager.h
+++ b/Engine/Plugins/FX/Niagara/Source/Niagara/Public/NiagaraWorldManager.h
@@ -21,10 +21,10 @@
 #include "NiagaraScalabilityManager.h"
 #include "NiagaraDebuggerCommon.h"
 #include "NiagaraDeferredMethodQueue.h"
-#include "NiagaraDataChannelManager.h"
 
 #include "NiagaraWorldManager.generated.h"
 
+class FNiagaraDataChannelManager;
 class UWorld;
 class UNiagaraParameterCollection;
 class UNiagaraParameterCollectionInstance;
@@ -218,7 +218,7 @@ public:
 	NIAGARA_API static void SetScalabilityCullingMode(ENiagaraScalabilityCullingMode NewMode);
 	NIAGARA_API static ENiagaraScalabilityCullingMode GetScalabilityCullingMode() { return ScalabilityCullingMode; }
 
-	FNiagaraDataChannelManager& GetDataChannelManager(){ return DataChannelManager; }
+	FNiagaraDataChannelManager& GetDataChannelManager(){ return *DataChannelManager; }
 
 	/** Waits for all currently in flight async work to be completed. */
 	void WaitForAsyncWork();
@@ -335,7 +335,7 @@ private:
 	/** A global flag for all scalability culling */
 	static ENiagaraScalabilityCullingMode ScalabilityCullingMode;
 
-	FNiagaraDataChannelManager DataChannelManager;
+	TUniquePtr<FNiagaraDataChannelManager> DataChannelManager;
 };
 
 
