Removing Niagara/Private from the Niagara public include paths.

diff --git a/Engine/Plugins/FX/Niagara/Source/Niagara/Niagara.Build.cs b/Engine/Plugins/FX/Niagara/Source/Niagara/Niagara.Build.cs
--- a/Engine/Plugins/FX/Niagara/Source/Niagara/Niagara.Build.cs
+++ b/Engine/Plugins/FX/Niagara/Source/Niagara/Niagara.Build.cs
@@ -109,8 +109,5 @@ public class Niagara : ModuleRules
                 "VECTORVM_SUPPORTS_EXPERIMENTAL=1",
                 "VECTORVM_SUPPORTS_LEGACY=1"
             });
-
-		// TODO: Should not be including private headers in public code
-		PublicIncludePaths.Add(Path.Combine(ModuleDirectory, "Private")); // For NiagaraStats.h & NiagaraGpuReadbackManager.h
 	}
 }
