- Add feedback context to both SimCache & Baker
- Display issues if a bake failed or has warnings that could lead to invalid results

diff --git a/Engine/Plugins/FX/Niagara/Source/Niagara/Classes/NiagaraDataInterfaceRenderTargetVolume.h b/Engine/Plugins/FX/Niagara/Source/Niagara/Classes/NiagaraDataInterfaceRenderTargetVolume.h
--- a/Engine/Plugins/FX/Niagara/Source/Niagara/Classes/NiagaraDataInterfaceRenderTargetVolume.h
+++ b/Engine/Plugins/FX/Niagara/Source/Niagara/Classes/NiagaraDataInterfaceRenderTargetVolume.h
@@ -125,8 +125,8 @@ public:
 	//~ UNiagaraDataInterface interface END
 
 	//~ INiagaraSimCacheCustomStorageInterface interface BEGIN
-	virtual UObject* SimCacheBeginWrite(UObject* SimCache, FNiagaraSystemInstance* NiagaraSystemInstance, const void* OptionalPerInstanceData) const override;
-	virtual bool SimCacheWriteFrame(UObject* StorageObject, int FrameIndex, FNiagaraSystemInstance* SystemInstance, const void* OptionalPerInstanceData) const override;
+	virtual UObject* SimCacheBeginWrite(UObject* SimCache, FNiagaraSystemInstance* NiagaraSystemInstance, const void* OptionalPerInstanceData, FNiagaraSimCacheFeedbackContext& FeedbackContext) const override;
+	virtual bool SimCacheWriteFrame(UObject* StorageObject, int FrameIndex, FNiagaraSystemInstance* SystemInstance, const void* OptionalPerInstanceData, FNiagaraSimCacheFeedbackContext& FeedbackContext) const override;
 	virtual bool SimCacheEndWrite(UObject* StorageObject) const override;
 	virtual bool SimCacheReadFrame(UObject* StorageObject, int FrameA, int FrameB, float Interp, FNiagaraSystemInstance* SystemInstance, void* OptionalPerInstanceData) override;
 	//~ UNiagaraDataInterface interface END
diff --git a/Engine/Plugins/FX/Niagara/Source/Niagara/Classes/NiagaraSimCache.h b/Engine/Plugins/FX/Niagara/Source/Niagara/Classes/NiagaraSimCache.h
--- a/Engine/Plugins/FX/Niagara/Source/Niagara/Classes/NiagaraSimCache.h
+++ b/Engine/Plugins/FX/Niagara/Source/Niagara/Classes/NiagaraSimCache.h
@@ -250,10 +250,15 @@ struct FNiagaraSimCacheLayout
 	TArray<FNiagaraSimCacheDataBuffersLayout> EmitterLayouts;
 };
 
-struct FNiagaraSimCacheWriteResult
+struct FNiagaraSimCacheFeedbackContext
 {
-	bool bSuccess = false;
-	FString ErrorMsg;
+	FNiagaraSimCacheFeedbackContext() = default;
+	explicit FNiagaraSimCacheFeedbackContext(bool bInAutoLogIssues) : bAutoLogIssues(bInAutoLogIssues) {}
+	NIAGARA_API ~FNiagaraSimCacheFeedbackContext();
+
+	bool bAutoLogIssues = true;
+	TArray<FString> Errors;
+	TArray<FString> Warnings;
 };
 
 UCLASS(Experimental, BlueprintType)
@@ -290,8 +295,10 @@ public:
 	UFUNCTION(BlueprintCallable, Category=NiagaraSimCache)
 	ENiagaraSimCacheAttributeCaptureMode GetAttributeCaptureMode() const { return CreateParameters.AttributeCaptureMode; }
 
-	bool BeginWrite(FNiagaraSimCacheCreateParameters InCreateParameters, UNiagaraComponent* NiagaraComponent);
-	FNiagaraSimCacheWriteResult WriteFrame(UNiagaraComponent* NiagaraComponent);
+	bool BeginWrite(FNiagaraSimCacheCreateParameters InCreateParameters, UNiagaraComponent * NiagaraComponent);
+	bool BeginWrite(FNiagaraSimCacheCreateParameters InCreateParameters, UNiagaraComponent* NiagaraComponent, FNiagaraSimCacheFeedbackContext& FeedbackContext);
+	bool WriteFrame(UNiagaraComponent* NiagaraComponent);
+	bool WriteFrame(UNiagaraComponent* NiagaraComponent, FNiagaraSimCacheFeedbackContext& FeedbackContext);
 	bool EndWrite();
 
 	bool CanRead(UNiagaraSystem* NiagaraSystem);
diff --git a/Engine/Plugins/FX/Niagara/Source/Niagara/Classes/NiagaraSimCacheCapture.cpp b/Engine/Plugins/FX/Niagara/Source/Niagara/Classes/NiagaraSimCacheCapture.cpp
--- a/Engine/Plugins/FX/Niagara/Source/Niagara/Classes/NiagaraSimCacheCapture.cpp
+++ b/Engine/Plugins/FX/Niagara/Source/Niagara/Classes/NiagaraSimCacheCapture.cpp
@@ -100,7 +100,7 @@ bool FNiagaraSimCacheCapture::OnFrameTick(float DeltaTime)
 	{
 		// If we fail to capture the frame it might be because things became invalid
 		// Or it might be because the simulation was not ticked since the last capture in which case don't advance the counter
-		if ( CaptureSimCache->WriteFrame(CaptureComponent).bSuccess == false )
+		if ( CaptureSimCache->WriteFrame(CaptureComponent) == false )
 		{
 			if ( CaptureSimCache->IsCacheValid() == false )
 			{
diff --git a/Engine/Plugins/FX/Niagara/Source/Niagara/Classes/NiagaraSimCacheCustomStorageInterface.h b/Engine/Plugins/FX/Niagara/Source/Niagara/Classes/NiagaraSimCacheCustomStorageInterface.h
--- a/Engine/Plugins/FX/Niagara/Source/Niagara/Classes/NiagaraSimCacheCustomStorageInterface.h
+++ b/Engine/Plugins/FX/Niagara/Source/Niagara/Classes/NiagaraSimCacheCustomStorageInterface.h
@@ -6,6 +6,7 @@
 #include "UObject/Interface.h"
 #include "NiagaraSimCacheCustomStorageInterface.generated.h"
 
+struct FNiagaraSimCacheFeedbackContext;
 class FNiagaraSystemInstance;
 
 // Interface for UObjects to implement renderable mesh
@@ -29,13 +30,13 @@ public:
 	Returning nullptr means you are not going to cache any data for the simulation.
 	The object returned will be stored directly into the cache file, so you are expected to manage the size of the object and store data appropriately.
 	*/
-	virtual UObject* SimCacheBeginWrite(UObject* SimCache, FNiagaraSystemInstance* NiagaraSystemInstance, const void* OptionalPerInstanceData) const { return nullptr; }
+	virtual UObject* SimCacheBeginWrite(UObject* SimCache, FNiagaraSystemInstance* NiagaraSystemInstance, const void* OptionalPerInstanceData, FNiagaraSimCacheFeedbackContext& FeedbackContext) const { return nullptr; }
 
 	/**
 	Called when we are ready to write data into the simulation cache.
 	This is always called in sequence, i.e. 0, 1, 2, etc, we will never jump around frames.
 	*/
-	virtual bool SimCacheWriteFrame(UObject* StorageObject, int FrameIndex, FNiagaraSystemInstance* SystemInstance, const void* OptionalPerInstanceData) const { return true; }
+	virtual bool SimCacheWriteFrame(UObject* StorageObject, int FrameIndex, FNiagaraSystemInstance* SystemInstance, const void* OptionalPerInstanceData, FNiagaraSimCacheFeedbackContext& FeedbackContext) const { return true; }
 
 	/**
 	Called when we complete writing data into the simulation cache.
diff --git a/Engine/Plugins/FX/Niagara/Source/Niagara/Private/NiagaraDataInterfaceRenderTargetVolume.cpp b/Engine/Plugins/FX/Niagara/Source/Niagara/Private/NiagaraDataInterfaceRenderTargetVolume.cpp
--- a/Engine/Plugins/FX/Niagara/Source/Niagara/Private/NiagaraDataInterfaceRenderTargetVolume.cpp
+++ b/Engine/Plugins/FX/Niagara/Source/Niagara/Private/NiagaraDataInterfaceRenderTargetVolume.cpp
@@ -556,7 +556,7 @@ bool UNiagaraDataInterfaceRenderTargetVolume::GetExposedVariableValue(const FNia
 	return false;
 }
 
-UObject* UNiagaraDataInterfaceRenderTargetVolume::SimCacheBeginWrite(UObject* InSimCache, FNiagaraSystemInstance* NiagaraSystemInstance, const void* OptionalPerInstanceData) const
+UObject* UNiagaraDataInterfaceRenderTargetVolume::SimCacheBeginWrite(UObject* InSimCache, FNiagaraSystemInstance* NiagaraSystemInstance, const void* OptionalPerInstanceData, FNiagaraSimCacheFeedbackContext& FeedbackContext) const
 {
 	if (NDIRenderTargetVolumeLocal::GSimCacheUseOpenVDB)
 	{		
@@ -569,7 +569,7 @@ UObject* UNiagaraDataInterfaceRenderTargetVolume::SimCacheBeginWrite(UObject* In
 
 		if (GetDefault<UNiagaraSettings>()->SimCacheAuxiliaryFileBasePath == "")
 		{
-			UE_LOG(LogNiagara, Error, TEXT("You must set SimCacheAuxiliaryFileBasePath in project settings"));
+			FeedbackContext.Errors.Emplace(TEXT("UNiagaraDataInterfaceRenderTargetVolume - You must set SimCacheAuxiliaryFileBasePath in project settings"));
 			return nullptr;
 		}
 
@@ -590,7 +590,7 @@ UObject* UNiagaraDataInterfaceRenderTargetVolume::SimCacheBeginWrite(UObject* In
 		{
 			if (!PlatformFile.CreateDirectoryTree(*FileDirectory))
 			{
-				UE_LOG(LogNiagara, Error, TEXT("Cannot Create Directory : %s"), *FileDirectory);
+				FeedbackContext.Errors.Emplace(FString::Printf(TEXT("Cannot Create Directory : %s"), *FileDirectory));
 				return nullptr;
 			}
 		}
@@ -619,13 +619,12 @@ UObject* UNiagaraDataInterfaceRenderTargetVolume::SimCacheBeginWrite(UObject* In
 	return nullptr;
 }
 
-bool UNiagaraDataInterfaceRenderTargetVolume::SimCacheWriteFrame(UObject* StorageObject, int FrameIndex, FNiagaraSystemInstance* SystemInstance, const void* OptionalPerInstanceData) const
+bool UNiagaraDataInterfaceRenderTargetVolume::SimCacheWriteFrame(UObject* StorageObject, int FrameIndex, FNiagaraSystemInstance* SystemInstance, const void* OptionalPerInstanceData, FNiagaraSimCacheFeedbackContext& FeedbackContext) const
 {
 	check(OptionalPerInstanceData);
 
 	const FRenderTargetVolumeRWInstanceData_GameThread* InstanceData_GT = reinterpret_cast<const FRenderTargetVolumeRWInstanceData_GameThread*>(OptionalPerInstanceData);
 
-
 	if (InstanceData_GT->TargetTexture)
 	{
 		const FNiagaraDataInterfaceProxyRenderTargetVolumeProxy* RT_Proxy = GetProxyAs<FNiagaraDataInterfaceProxyRenderTargetVolumeProxy>();
diff --git a/Engine/Plugins/FX/Niagara/Source/Niagara/Private/NiagaraSimCache.cpp b/Engine/Plugins/FX/Niagara/Source/Niagara/Private/NiagaraSimCache.cpp
--- a/Engine/Plugins/FX/Niagara/Source/Niagara/Private/NiagaraSimCache.cpp
+++ b/Engine/Plugins/FX/Niagara/Source/Niagara/Private/NiagaraSimCache.cpp
@@ -19,6 +19,22 @@
 UNiagaraSimCache::FOnCacheBeginWrite	UNiagaraSimCache::OnCacheBeginWrite;
 UNiagaraSimCache::FOnCacheEndWrite		UNiagaraSimCache::OnCacheEndWrite;
 
+FNiagaraSimCacheFeedbackContext::~FNiagaraSimCacheFeedbackContext()
+{
+	if (bAutoLogIssues)
+	{
+		for (const FString& Warning : Warnings)
+		{
+			UE_LOG(LogNiagara, Warning, TEXT("SimCache Warning: %s"), *Warning);
+		}
+
+		for (const FString& Error : Errors)
+		{
+			UE_LOG(LogNiagara, Warning, TEXT("SimCache Error: %s"), *Error);
+		}
+	}
+}
+
 UNiagaraSimCache::UNiagaraSimCache(const FObjectInitializer& ObjectInitializer)
 {
 }
@@ -29,6 +45,12 @@ bool UNiagaraSimCache::IsReadyForFinishDestroy()
 }
 
 bool UNiagaraSimCache::BeginWrite(FNiagaraSimCacheCreateParameters InCreateParameters, UNiagaraComponent* NiagaraComponent)
+{
+	FNiagaraSimCacheFeedbackContext FeedbackContext;
+	return BeginWrite(InCreateParameters, NiagaraComponent, FeedbackContext);
+}
+
+bool UNiagaraSimCache::BeginWrite(FNiagaraSimCacheCreateParameters InCreateParameters, UNiagaraComponent* NiagaraComponent, FNiagaraSimCacheFeedbackContext& FeedbackContext)
 {
 	check(PendingCommandsInFlight == 0);
 
@@ -37,6 +59,7 @@ bool UNiagaraSimCache::BeginWrite(FNiagaraSimCacheCreateParameters InCreateParam
 	FNiagaraSimCacheHelper Helper(NiagaraComponent);
 	if (Helper.HasValidSimulation() == false)
 	{
+		FeedbackContext.Errors.Emplace(FString::Printf(TEXT("No valid simulation data for %s"), *GetPathNameSafe(Helper.NiagaraSystem)));
 		return false;
 	}
 
@@ -71,6 +94,7 @@ bool UNiagaraSimCache::BeginWrite(FNiagaraSimCacheCreateParameters InCreateParam
 		// Invalid as the user specified nothing to actual capture
 		if (CreateParameters.ExplicitCaptureAttributes.IsEmpty())
 		{
+			FeedbackContext.Errors.Emplace(FString::Printf(TEXT("Missing attributes to capture when in explicit capture mode %s"), *GetPathNameSafe(Helper.NiagaraSystem)));
 			SoftNiagaraSystem.Reset();
 			return false;
 		}
@@ -153,7 +177,7 @@ bool UNiagaraSimCache::BeginWrite(FNiagaraSimCacheCreateParameters InCreateParam
 				if (INiagaraSimCacheCustomStorageInterface* SimCacheCustomStorageInterface = Cast<INiagaraSimCacheCustomStorageInterface>(DataInterface))
 				{
 					const void* PerInstanceData = Helper.SystemInstance->FindDataInterfaceInstanceData(DataInterface);
-					if (UObject* DICacheStorage = SimCacheCustomStorageInterface->SimCacheBeginWrite(this, Helper.SystemInstance, PerInstanceData))
+					if (UObject* DICacheStorage = SimCacheCustomStorageInterface->SimCacheBeginWrite(this, Helper.SystemInstance, PerInstanceData, FeedbackContext))
 					{
 						DataInterfaceStorage.FindOrAdd(Variable) = DICacheStorage;
 					}
@@ -162,39 +186,76 @@ bool UNiagaraSimCache::BeginWrite(FNiagaraSimCacheCreateParameters InCreateParam
 			}
 		);
 	}
+
+#if WITH_EDITORONLY_DATA
+	// Check all our bindings are ok
+	for ( const FNiagaraEmitterHandle& EmitterHandle : Helper.NiagaraSystem->GetEmitterHandles() )
+	{
+		EmitterHandle.GetInstance().GetEmitterData()->ForEachEnabledRenderer(
+			[&](UNiagaraRendererProperties* RenderProperties)
+			{
+				for (FNiagaraVariableBase BoundAttribute : RenderProperties->GetBoundAttributes())
+				{
+					if (!BoundAttribute.IsDataInterface())
+					{
+						continue;
+					}
+
+					UNiagaraDataInterface* DataInterfaceCDO = CastChecked<UNiagaraDataInterface>(BoundAttribute.GetType().GetClass()->GetDefaultObject());
+					if (DataInterfaceCDO->PerInstanceDataSize() == 0)
+					{
+						continue;
+					}
+
+					if (!DataInterfaceStorage.Contains(BoundAttribute))
+					{
+						FeedbackContext.Warnings.Emplace(FString::Printf(TEXT("Material Data Interface Binding '%s : %s' did not store any data or does not support sim caching, baked result may be incorrect."), *DataInterfaceCDO->GetClass()->GetName(), *BoundAttribute.GetName().ToString()));
+					}
+				}
+			}
+		);
+	}
+#endif
+
 	return true;
 }
 
-FNiagaraSimCacheWriteResult UNiagaraSimCache::WriteFrame(UNiagaraComponent* NiagaraComponent)
+bool UNiagaraSimCache::WriteFrame(UNiagaraComponent* NiagaraComponent)
+{
+	FNiagaraSimCacheFeedbackContext FeedbackContext;
+	return WriteFrame(NiagaraComponent, FeedbackContext);
+}
+
+bool UNiagaraSimCache::WriteFrame(UNiagaraComponent* NiagaraComponent, FNiagaraSimCacheFeedbackContext& FeedbackContext)
 {
-	FNiagaraSimCacheWriteResult Result;
 	FNiagaraSimCacheHelper Helper(NiagaraComponent);
 	if (Helper.HasValidSimulationData() == false)
 	{
 		SoftNiagaraSystem.Reset();
-		Result.ErrorMsg = FString::Printf(TEXT("No valid simulation data for %s"), *GetPathNameSafe(Helper.NiagaraSystem));
-		return Result;
+		FeedbackContext.Errors.Emplace(FString::Printf(TEXT("No valid simulation data for %s"), *GetPathNameSafe(Helper.NiagaraSystem)));
+		return false;
 	}
 
 	if ( SoftNiagaraSystem.Get() != Helper.NiagaraSystem )
 	{
-		Result.ErrorMsg = FString::Printf(TEXT("System %s != System %s"), *GetPathNameSafe(SoftNiagaraSystem.Get()), *GetPathNameSafe(Helper.NiagaraSystem));
+		FeedbackContext.Errors.Emplace(FString::Printf(TEXT("System %s != System %s"), *GetPathNameSafe(SoftNiagaraSystem.Get()), *GetPathNameSafe(Helper.NiagaraSystem)));
 		SoftNiagaraSystem.Reset();
-		return Result;
+		return false;
 	}
 
 	// Simulation is complete nothing to cache
 	if ( Helper.SystemInstance->IsComplete() )
 	{
-		Result.ErrorMsg = FString::Printf(TEXT("System already completed its simulation. %s"), *Helper.SystemInstance->GetCrashReporterTag());
-		return Result;
+		FeedbackContext.Errors.Emplace(FString::Printf(TEXT("System already completed its simulation. %s"), *Helper.SystemInstance->GetCrashReporterTag()));
+		return false;
 	}
 
 	// Is the simulation running?  If not nothing to cache yet
 	if ( Helper.SystemInstance->SystemInstanceState != ENiagaraSystemInstanceState::Running )
 	{
-		Result.ErrorMsg = FString::Printf(TEXT("System is not running, so there is no data to cache. %s"), *Helper.SystemInstance->GetCrashReporterTag());
-		return Result;
+		// This warning will trigger way to often and isn't that useful
+		//	FeedbackContext.Errors.Emplace(FString::Printf(TEXT("System is not running, so there is no data to cache. %s"), *Helper.SystemInstance->GetCrashReporterTag()));
+		return false;
 	}
 
 	// First frame we are about to cache?
@@ -206,16 +267,16 @@ FNiagaraSimCacheWriteResult UNiagaraSimCache::WriteFrame(UNiagaraComponent* Niag
 	// If our tick counter hasn't moved then we won't capture a frame as there's nothing new to process
 	else if (CaptureTickCount == Helper.SystemInstance->GetTickCount())
 	{
-		Result.ErrorMsg = FString::Printf(TEXT("System was not ticked since the last capture. %s"), *Helper.SystemInstance->GetCrashReporterTag());
-		return Result;
+		FeedbackContext.Errors.Emplace(FString::Printf(TEXT("System was not ticked since the last capture. %s"), *Helper.SystemInstance->GetCrashReporterTag()));
+		return false;
 	}
 
 	// If the tick counter is lower than the previous value then the system was reset so we won't capture
 	if ( Helper.SystemInstance->GetTickCount() < CaptureTickCount )
 	{
-		Result.ErrorMsg = FString::Printf(TEXT("System was was reset since the last capture, skipping further cache writes. %s"), *Helper.SystemInstance->GetCrashReporterTag());
+		FeedbackContext.Errors.Emplace(FString::Printf(TEXT("System was was reset since the last capture, skipping further cache writes. %s"), *Helper.SystemInstance->GetCrashReporterTag()));
 		SoftNiagaraSystem.Reset();
-		return Result;
+		return false;
 	}
 
 	DurationSeconds = Helper.SystemInstance->GetAge() - StartSeconds;
@@ -289,7 +350,7 @@ FNiagaraSimCacheWriteResult UNiagaraSimCache::WriteFrame(UNiagaraComponent* Niag
 				{
 					INiagaraSimCacheCustomStorageInterface* SimCacheCustomStorageInterface = CastChecked<INiagaraSimCacheCustomStorageInterface>(DataInterface);
 					const void* PerInstanceData = Helper.SystemInstance->FindDataInterfaceInstanceData(DataInterface);
-					if (!SimCacheCustomStorageInterface->SimCacheWriteFrame(StorageObject, FrameIndex, Helper.SystemInstance, PerInstanceData))
+					if (!SimCacheCustomStorageInterface->SimCacheWriteFrame(StorageObject, FrameIndex, Helper.SystemInstance, PerInstanceData, FeedbackContext))
 					{
 						bDataInterfacesSucess = false;
 						DataInterfaceName += FString::Printf(TEXT("%s "), *DataInterface->GetName());
@@ -302,13 +363,12 @@ FNiagaraSimCacheWriteResult UNiagaraSimCache::WriteFrame(UNiagaraComponent* Niag
 		// A data interface failed to write information
 		if (bDataInterfacesSucess == false)
 		{
-			Result.ErrorMsg = FString::Printf(TEXT("Data interface(s) %s failed to write information for system %s"), *DataInterfaceName, *Helper.SystemInstance->GetCrashReporterTag());
+			FeedbackContext.Errors.Emplace(FString::Printf(TEXT("Data interface(s) %s failed to write information for system %s"), *DataInterfaceName, *Helper.SystemInstance->GetCrashReporterTag()));
 			SoftNiagaraSystem.Reset();
-			return Result;
+			return false;
 		}
 	}
-	Result.bSuccess = true;
-	return Result;
+	return true;
 }
 
 bool UNiagaraSimCache::EndWrite()
diff --git a/Engine/Plugins/FX/Niagara/Source/NiagaraEditor/NiagaraEditor.Build.cs b/Engine/Plugins/FX/Niagara/Source/NiagaraEditor/NiagaraEditor.Build.cs
--- a/Engine/Plugins/FX/Niagara/Source/NiagaraEditor/NiagaraEditor.Build.cs
+++ b/Engine/Plugins/FX/Niagara/Source/NiagaraEditor/NiagaraEditor.Build.cs
@@ -63,6 +63,7 @@ public class NiagaraEditor : ModuleRules
 				"ImageWrapper",
 				"AssetDefinition",
 				"ContentBrowser",
+				"ToolWidgets",
 				"AssetTools"
 			}
 		);
diff --git a/Engine/Plugins/FX/Niagara/Source/NiagaraEditor/Private/NiagaraBakerRenderer.h b/Engine/Plugins/FX/Niagara/Source/NiagaraEditor/Private/NiagaraBakerRenderer.h
--- a/Engine/Plugins/FX/Niagara/Source/NiagaraEditor/Private/NiagaraBakerRenderer.h
+++ b/Engine/Plugins/FX/Niagara/Source/NiagaraEditor/Private/NiagaraBakerRenderer.h
@@ -17,6 +17,14 @@ class USceneCaptureComponent2D;
 class FAdvancedPreviewScene;
 class FCanvas;
 
+struct FNiagaraBakerFeedbackContext
+{
+	bool HasIssues() const { return Errors.Num() + Warnings.Num() > 0; }
+
+	TArray<FString>	Errors;
+	TArray<FString>	Warnings;
+};
+
 struct FNiagaraBakerOutputBinding
 {
 	FName	BindingName;
@@ -79,9 +87,9 @@ public:
 	virtual FIntPoint GetGeneratedSize(UNiagaraBakerOutput* BakerOutput, FIntPoint InAvailableSize) const { return FIntPoint::ZeroValue; }
 	virtual void RenderGenerated(UNiagaraBakerOutput* BakerOutput, const FNiagaraBakerRenderer& BakerRenderer, UTextureRenderTarget2D* RenderTarget, TOptional<FString>& OutErrorString) const = 0;
 
-	virtual bool BeginBake(UNiagaraBakerOutput* InBakerOutput) = 0;
-	virtual void BakeFrame(UNiagaraBakerOutput* InBakerOutput, int FrameIndex, const FNiagaraBakerRenderer& BakerRenderer) = 0;
-	virtual void EndBake(UNiagaraBakerOutput* InBakerOutput) = 0;
+	virtual bool BeginBake(FNiagaraBakerFeedbackContext& FeedbackContext, UNiagaraBakerOutput* InBakerOutput) = 0;
+	virtual void BakeFrame(FNiagaraBakerFeedbackContext& FeedbackContext, UNiagaraBakerOutput* InBakerOutput, int FrameIndex, const FNiagaraBakerRenderer& BakerRenderer) = 0;
+	virtual void EndBake(FNiagaraBakerFeedbackContext& FeedbackContext, UNiagaraBakerOutput* InBakerOutput) = 0;
 };
 
 class FNiagaraBakerRenderer : FGCObject
diff --git a/Engine/Plugins/FX/Niagara/Source/NiagaraEditor/Private/NiagaraBakerRendererOutputSimCache.cpp b/Engine/Plugins/FX/Niagara/Source/NiagaraEditor/Private/NiagaraBakerRendererOutputSimCache.cpp
--- a/Engine/Plugins/FX/Niagara/Source/NiagaraEditor/Private/NiagaraBakerRendererOutputSimCache.cpp
+++ b/Engine/Plugins/FX/Niagara/Source/NiagaraEditor/Private/NiagaraBakerRendererOutputSimCache.cpp
@@ -46,7 +46,7 @@ void FNiagaraBakerRendererOutputSimCache::RenderGenerated(UNiagaraBakerOutput* I
 	BakerRenderer.RenderSimCache(RenderTarget, SimCache);
 }
 
-bool FNiagaraBakerRendererOutputSimCache::BeginBake(UNiagaraBakerOutput* InBakerOutput)
+bool FNiagaraBakerRendererOutputSimCache::BeginBake(FNiagaraBakerFeedbackContext& FeedbackContext, UNiagaraBakerOutput* InBakerOutput)
 {
 	UNiagaraBakerOutputSimCache* BakerOutput = CastChecked<UNiagaraBakerOutputSimCache>(InBakerOutput);
 	UNiagaraBakerSettings* BakerSettings = BakerOutput->GetTypedOuter<UNiagaraBakerSettings>();
@@ -64,7 +64,7 @@ bool FNiagaraBakerRendererOutputSimCache::BeginBake(UNiagaraBakerOutput* InBaker
 	return true;
 }
 
-void FNiagaraBakerRendererOutputSimCache::BakeFrame(UNiagaraBakerOutput* InBakerOutput, int FrameIndex, const FNiagaraBakerRenderer& BakerRenderer)
+void FNiagaraBakerRendererOutputSimCache::BakeFrame(FNiagaraBakerFeedbackContext& FeedbackContext, UNiagaraBakerOutput* InBakerOutput, int FrameIndex, const FNiagaraBakerRenderer& BakerRenderer)
 {
 	UNiagaraBakerOutputSimCache* BakerOutput = CastChecked<UNiagaraBakerOutputSimCache>(InBakerOutput);
 	if (BakeSimCache == nullptr)
@@ -72,14 +72,18 @@ void FNiagaraBakerRendererOutputSimCache::BakeFrame(UNiagaraBakerOutput* InBaker
 		return;
 	}
 
+	FNiagaraSimCacheFeedbackContext SimCacheFeedbackContext(false);
 	if ( FrameIndex == 0 )
 	{
-		BakeSimCache->BeginWrite(BakerOutput->CreateParameters, BakerRenderer.GetPreviewComponent());
+		BakeSimCache->BeginWrite(BakerOutput->CreateParameters, BakerRenderer.GetPreviewComponent(), SimCacheFeedbackContext);
 	}
-	BakeSimCache->WriteFrame(BakerRenderer.GetPreviewComponent());
+	BakeSimCache->WriteFrame(BakerRenderer.GetPreviewComponent(), SimCacheFeedbackContext);
+
+	FeedbackContext.Errors.Append(SimCacheFeedbackContext.Errors);
+	FeedbackContext.Warnings.Append(SimCacheFeedbackContext.Warnings);
 }
 
-void FNiagaraBakerRendererOutputSimCache::EndBake(UNiagaraBakerOutput* InBakerOutput)
+void FNiagaraBakerRendererOutputSimCache::EndBake(FNiagaraBakerFeedbackContext& FeedbackContext, UNiagaraBakerOutput* InBakerOutput)
 {
 	UNiagaraBakerOutputSimCache* BakerOutput = CastChecked<UNiagaraBakerOutputSimCache>(InBakerOutput);
 	if (BakeSimCache == nullptr)
diff --git a/Engine/Plugins/FX/Niagara/Source/NiagaraEditor/Private/NiagaraBakerRendererOutputSimCache.h b/Engine/Plugins/FX/Niagara/Source/NiagaraEditor/Private/NiagaraBakerRendererOutputSimCache.h
--- a/Engine/Plugins/FX/Niagara/Source/NiagaraEditor/Private/NiagaraBakerRendererOutputSimCache.h
+++ b/Engine/Plugins/FX/Niagara/Source/NiagaraEditor/Private/NiagaraBakerRendererOutputSimCache.h
@@ -17,9 +17,9 @@ public:
 	virtual FIntPoint GetGeneratedSize(UNiagaraBakerOutput* InBakerOutput, FIntPoint InAvailableSize) const override;
 	virtual void RenderGenerated(UNiagaraBakerOutput* InBakerOutput, const FNiagaraBakerRenderer& BakerRenderer, UTextureRenderTarget2D* RenderTarget, TOptional<FString>& OutErrorString) const override;
 
-	virtual bool BeginBake(UNiagaraBakerOutput* InBakerOutput) override;
-	virtual void BakeFrame(UNiagaraBakerOutput* InBakerOutput, int FrameIndex, const FNiagaraBakerRenderer& BakerRenderer) override;
-	virtual void EndBake(UNiagaraBakerOutput* InBakerOutput) override;
+	virtual bool BeginBake(FNiagaraBakerFeedbackContext& FeedbackContext, UNiagaraBakerOutput* InBakerOutput) override;
+	virtual void BakeFrame(FNiagaraBakerFeedbackContext& FeedbackContext, UNiagaraBakerOutput* InBakerOutput, int FrameIndex, const FNiagaraBakerRenderer& BakerRenderer) override;
+	virtual void EndBake(FNiagaraBakerFeedbackContext& FeedbackContext, UNiagaraBakerOutput* InBakerOutput) override;
 
 protected:
 	UNiagaraSimCache* BakeSimCache = nullptr;
diff --git a/Engine/Plugins/FX/Niagara/Source/NiagaraEditor/Private/NiagaraBakerRendererOutputTexture2D.cpp b/Engine/Plugins/FX/Niagara/Source/NiagaraEditor/Private/NiagaraBakerRendererOutputTexture2D.cpp
--- a/Engine/Plugins/FX/Niagara/Source/NiagaraEditor/Private/NiagaraBakerRendererOutputTexture2D.cpp
+++ b/Engine/Plugins/FX/Niagara/Source/NiagaraEditor/Private/NiagaraBakerRendererOutputTexture2D.cpp
@@ -118,7 +118,7 @@ void FNiagaraBakerRendererOutputTexture2D::RenderGenerated(UNiagaraBakerOutput*
 	Canvas.Flush_GameThread();
 }
 
-bool FNiagaraBakerRendererOutputTexture2D::BeginBake(UNiagaraBakerOutput* InBakerOutput)
+bool FNiagaraBakerRendererOutputTexture2D::BeginBake(FNiagaraBakerFeedbackContext& FeedbackContext, UNiagaraBakerOutput* InBakerOutput)
 {
 	UNiagaraBakerOutputTexture2D* BakerOutput = CastChecked<UNiagaraBakerOutputTexture2D>(InBakerOutput);
 	if ( !BakerOutput->bGenerateAtlas && !BakerOutput->bGenerateFrames && !BakerOutput->bExportFrames )
@@ -141,7 +141,7 @@ bool FNiagaraBakerRendererOutputTexture2D::BeginBake(UNiagaraBakerOutput* InBake
 	return true;
 }
 
-void FNiagaraBakerRendererOutputTexture2D::BakeFrame(UNiagaraBakerOutput* InBakerOutput, int FrameIndex, const FNiagaraBakerRenderer& BakerRenderer)
+void FNiagaraBakerRendererOutputTexture2D::BakeFrame(FNiagaraBakerFeedbackContext& FeedbackContext, UNiagaraBakerOutput* InBakerOutput, int FrameIndex, const FNiagaraBakerRenderer& BakerRenderer)
 {
 	UNiagaraBakerOutputTexture2D* BakerOutput = CastChecked<UNiagaraBakerOutputTexture2D>(InBakerOutput);
 	UNiagaraBakerSettings* BakerSettings = BakerRenderer.GetBakerSettings();
@@ -195,7 +195,7 @@ void FNiagaraBakerRendererOutputTexture2D::BakeFrame(UNiagaraBakerOutput* InBake
 	}
 }
 
-void FNiagaraBakerRendererOutputTexture2D::EndBake(UNiagaraBakerOutput* InBakerOutput)
+void FNiagaraBakerRendererOutputTexture2D::EndBake(FNiagaraBakerFeedbackContext& FeedbackContext, UNiagaraBakerOutput* InBakerOutput)
 {
 	UNiagaraBakerOutputTexture2D* BakerOutput = CastChecked<UNiagaraBakerOutputTexture2D>(InBakerOutput);
 
diff --git a/Engine/Plugins/FX/Niagara/Source/NiagaraEditor/Private/NiagaraBakerRendererOutputTexture2D.h b/Engine/Plugins/FX/Niagara/Source/NiagaraEditor/Private/NiagaraBakerRendererOutputTexture2D.h
--- a/Engine/Plugins/FX/Niagara/Source/NiagaraEditor/Private/NiagaraBakerRendererOutputTexture2D.h
+++ b/Engine/Plugins/FX/Niagara/Source/NiagaraEditor/Private/NiagaraBakerRendererOutputTexture2D.h
@@ -17,9 +17,9 @@ public:
 	virtual FIntPoint GetGeneratedSize(UNiagaraBakerOutput* InBakerOutput, FIntPoint InAvailableSize) const override;
 	virtual void RenderGenerated(UNiagaraBakerOutput* InBakerOutput, const FNiagaraBakerRenderer& BakerRenderer, UTextureRenderTarget2D* RenderTarget, TOptional<FString>& OutErrorString) const override;
 
-	virtual bool BeginBake(UNiagaraBakerOutput* InBakerOutput) override;
-	virtual void BakeFrame(UNiagaraBakerOutput* InBakerOutput, int FrameIndex, const FNiagaraBakerRenderer& BakerRenderer) override;
-	virtual void EndBake(UNiagaraBakerOutput* InBakerOutput) override;
+	virtual bool BeginBake(FNiagaraBakerFeedbackContext& FeedbackContext, UNiagaraBakerOutput* InBakerOutput) override;
+	virtual void BakeFrame(FNiagaraBakerFeedbackContext& FeedbackContext, UNiagaraBakerOutput* InBakerOutput, int FrameIndex, const FNiagaraBakerRenderer& BakerRenderer) override;
+	virtual void EndBake(FNiagaraBakerFeedbackContext& FeedbackContext, UNiagaraBakerOutput* InBakerOutput) override;
 
 protected:
 	UTextureRenderTarget2D* BakeRenderTarget = nullptr;
diff --git a/Engine/Plugins/FX/Niagara/Source/NiagaraEditor/Private/NiagaraBakerRendererOutputVolumeTexture.cpp b/Engine/Plugins/FX/Niagara/Source/NiagaraEditor/Private/NiagaraBakerRendererOutputVolumeTexture.cpp
--- a/Engine/Plugins/FX/Niagara/Source/NiagaraEditor/Private/NiagaraBakerRendererOutputVolumeTexture.cpp
+++ b/Engine/Plugins/FX/Niagara/Source/NiagaraEditor/Private/NiagaraBakerRendererOutputVolumeTexture.cpp
@@ -391,7 +391,7 @@ void FNiagaraBakerRendererOutputVolumeTexture::RenderGenerated(UNiagaraBakerOutp
 	Canvas.Flush_GameThread();
 }
 
-bool FNiagaraBakerRendererOutputVolumeTexture::BeginBake(UNiagaraBakerOutput* InBakerOutput)
+bool FNiagaraBakerRendererOutputVolumeTexture::BeginBake(FNiagaraBakerFeedbackContext& FeedbackContext, UNiagaraBakerOutput* InBakerOutput)
 {
 	UNiagaraBakerOutputVolumeTexture* OutputVolumeTexture = CastChecked<UNiagaraBakerOutputVolumeTexture>(InBakerOutput);
 
@@ -401,7 +401,7 @@ bool FNiagaraBakerRendererOutputVolumeTexture::BeginBake(UNiagaraBakerOutput* In
 	return OutputVolumeTexture->bGenerateAtlas || OutputVolumeTexture->bGenerateFrames || OutputVolumeTexture->bExportFrames;
 }
 
-void FNiagaraBakerRendererOutputVolumeTexture::BakeFrame(UNiagaraBakerOutput* InBakerOutput, int FrameIndex, const FNiagaraBakerRenderer& BakerRenderer)
+void FNiagaraBakerRendererOutputVolumeTexture::BakeFrame(FNiagaraBakerFeedbackContext& FeedbackContext, UNiagaraBakerOutput* InBakerOutput, int FrameIndex, const FNiagaraBakerRenderer& BakerRenderer)
 {
 	UNiagaraBakerOutputVolumeTexture* OutputVolumeTexture = CastChecked<UNiagaraBakerOutputVolumeTexture>(InBakerOutput);
 
@@ -605,7 +605,7 @@ void FNiagaraBakerRendererOutputVolumeTexture::BakeFrame(UNiagaraBakerOutput* In
 	}
 }
 
-void FNiagaraBakerRendererOutputVolumeTexture::EndBake(UNiagaraBakerOutput* InBakerOutput)
+void FNiagaraBakerRendererOutputVolumeTexture::EndBake(FNiagaraBakerFeedbackContext& FeedbackContext, UNiagaraBakerOutput* InBakerOutput)
 {
 	UNiagaraBakerOutputVolumeTexture* OutputVolumeTexture = CastChecked<UNiagaraBakerOutputVolumeTexture>(InBakerOutput);
 	if (BakeAtlasTextureData.Num() > 0)
diff --git a/Engine/Plugins/FX/Niagara/Source/NiagaraEditor/Private/NiagaraBakerRendererOutputVolumeTexture.h b/Engine/Plugins/FX/Niagara/Source/NiagaraEditor/Private/NiagaraBakerRendererOutputVolumeTexture.h
--- a/Engine/Plugins/FX/Niagara/Source/NiagaraEditor/Private/NiagaraBakerRendererOutputVolumeTexture.h
+++ b/Engine/Plugins/FX/Niagara/Source/NiagaraEditor/Private/NiagaraBakerRendererOutputVolumeTexture.h
@@ -17,9 +17,9 @@ public:
 	virtual FIntPoint GetGeneratedSize(UNiagaraBakerOutput* InBakerOutput, FIntPoint InAvailableSize) const override;
 	virtual void RenderGenerated(UNiagaraBakerOutput* InBakerOutput, const FNiagaraBakerRenderer& BakerRenderer, UTextureRenderTarget2D* RenderTarget, TOptional<FString>& OutErrorString) const override;
 
-	virtual bool BeginBake(UNiagaraBakerOutput* InBakerOutput) override;
-	virtual void BakeFrame(UNiagaraBakerOutput* InBakerOutput, int FrameIndex, const FNiagaraBakerRenderer& BakerRenderer) override;
-	virtual void EndBake(UNiagaraBakerOutput* InBakerOutput) override;
+	virtual bool BeginBake(FNiagaraBakerFeedbackContext& FeedbackContext, UNiagaraBakerOutput* InBakerOutput) override;
+	virtual void BakeFrame(FNiagaraBakerFeedbackContext& FeedbackContext, UNiagaraBakerOutput* InBakerOutput, int FrameIndex, const FNiagaraBakerRenderer& BakerRenderer) override;
+	virtual void EndBake(FNiagaraBakerFeedbackContext& FeedbackContext, UNiagaraBakerOutput* InBakerOutput) override;
 
 protected:
 	FIntVector				BakeAtlasFrameSize = FIntVector::ZeroValue;
diff --git a/Engine/Plugins/FX/Niagara/Source/NiagaraEditor/Private/ViewModels/NiagaraBakerViewModel.cpp b/Engine/Plugins/FX/Niagara/Source/NiagaraEditor/Private/ViewModels/NiagaraBakerViewModel.cpp
--- a/Engine/Plugins/FX/Niagara/Source/NiagaraEditor/Private/ViewModels/NiagaraBakerViewModel.cpp
+++ b/Engine/Plugins/FX/Niagara/Source/NiagaraEditor/Private/ViewModels/NiagaraBakerViewModel.cpp
@@ -579,20 +579,23 @@ void FNiagaraBakerViewModel::SetFramesOnY(int32 Value)
 	}
 }
 
-void FNiagaraBakerViewModel::RenderBaker()
+FNiagaraBakerFeedbackContext FNiagaraBakerViewModel::RenderBaker()
 {
+	FNiagaraBakerFeedbackContext FeedbackContext;
+
 	UNiagaraBakerSettings* BakerSettings = GetBakerSettings();
 	UNiagaraSystem* NiagaraSystem = BakerRenderer->GetNiagaraSystem();
 	if ( BakerSettings == nullptr || NiagaraSystem == nullptr )
 	{
-		return;
+		FeedbackContext.Errors.Emplace(TEXT("Nothing to bake as we don't have valid settings or system to bake with."));
+		return FeedbackContext;
 	}
 
 	if ( BakerSettings->Outputs.Num() == 0 )
 	{
-		//-TODO: Message no outputs
 		NiagaraSystem->SetBakerGeneratedSettings(nullptr);
-		return;
+		FeedbackContext.Errors.Emplace(TEXT("Nothing to bake as we have no outputs."));
+		return FeedbackContext;
 	}
 
 	// Create output renderers
@@ -606,16 +609,16 @@ void FNiagaraBakerViewModel::RenderBaker()
 		OutputRenderers.Emplace(FNiagaraBakerRenderer::GetOutputRenderer(Output->GetClass()));
 		if ( FNiagaraBakerOutputRenderer* OutputRenderer = OutputRenderers.Last().Get() )
 		{
-			bHasValidRenderer |= OutputRenderer->BeginBake(Output);
+			bHasValidRenderer |= OutputRenderer->BeginBake(FeedbackContext, Output);
 		}
 	}
 
 	// Early out if we have no valid renderers
 	if ( !bHasValidRenderer )
 	{
-		//-TODO: No valid renderers
 		NiagaraSystem->SetBakerGeneratedSettings(nullptr);
-		return;
+		FeedbackContext.Errors.Emplace(TEXT("Nothing to bake as all outputs will not genreate any data."));
+		return FeedbackContext;
 	}
 
 	// Render frames
@@ -652,7 +655,7 @@ void FNiagaraBakerViewModel::RenderBaker()
 			{
 				if (OutputRenderers[i].IsValid())
 				{
-					OutputRenderers[i]->BakeFrame(BakerSettings->Outputs[i], iFrame, *BakerRenderer.Get());
+					OutputRenderers[i]->BakeFrame(FeedbackContext, BakerSettings->Outputs[i], iFrame, *BakerRenderer.Get());
 				}
 			}
 		}
@@ -683,7 +686,7 @@ void FNiagaraBakerViewModel::RenderBaker()
 			{
 				SlowTask.EnterProgressFrame(1, FText::Format(LOCTEXT("FinishBakeFormat", "Finish Bake for Output '{0}'"), FText::FromString(BakerSettings->Outputs[i]->OutputName)));
 
-				OutputRenderers[i]->EndBake(BakerSettings->Outputs[i]);
+				OutputRenderers[i]->EndBake(FeedbackContext, BakerSettings->Outputs[i]);
 			}
 		}
 	}
@@ -693,6 +696,8 @@ void FNiagaraBakerViewModel::RenderBaker()
 
 	// GC to cleanup any transient objects we created or ones we have ditched
 	CollectGarbage(GARBAGE_COLLECTION_KEEPFLAGS);
+
+	return FeedbackContext;
 }
 
 #undef LOCTEXT_NAMESPACE
diff --git a/Engine/Plugins/FX/Niagara/Source/NiagaraEditor/Private/ViewModels/NiagaraBakerViewModel.h b/Engine/Plugins/FX/Niagara/Source/NiagaraEditor/Private/ViewModels/NiagaraBakerViewModel.h
--- a/Engine/Plugins/FX/Niagara/Source/NiagaraEditor/Private/ViewModels/NiagaraBakerViewModel.h
+++ b/Engine/Plugins/FX/Niagara/Source/NiagaraEditor/Private/ViewModels/NiagaraBakerViewModel.h
@@ -34,7 +34,7 @@ public:
 
 	TSharedPtr<class SWidget> GetWidget();
 
-	void RenderBaker();
+	FNiagaraBakerFeedbackContext RenderBaker();
 
 	void SetDisplayTimeFromNormalized(float NormalizeTime);
 
diff --git a/Engine/Plugins/FX/Niagara/Source/NiagaraEditor/Private/Widgets/SNiagaraBakerWidget.cpp b/Engine/Plugins/FX/Niagara/Source/NiagaraEditor/Private/Widgets/SNiagaraBakerWidget.cpp
--- a/Engine/Plugins/FX/Niagara/Source/NiagaraEditor/Private/Widgets/SNiagaraBakerWidget.cpp
+++ b/Engine/Plugins/FX/Niagara/Source/NiagaraEditor/Private/Widgets/SNiagaraBakerWidget.cpp
@@ -11,6 +11,7 @@
 #include "NiagaraSystem.h"
 #include "ViewModels/NiagaraBakerViewModel.h"
 
+#include "Dialog/SMessageDialog.h"
 #include "Modules/ModuleManager.h"
 #include "Widgets/Input/SVectorInputBox.h"
 #include "Widgets/Input/SCheckBox.h"
@@ -1011,7 +1012,26 @@ FReply SNiagaraBakerWidget::OnCapture()
 {
 	if (auto ViewModel = WeakViewModel.Pin())
 	{
-		ViewModel->RenderBaker();
+		const FNiagaraBakerFeedbackContext FeedbackContext = ViewModel->RenderBaker();
+		if (FeedbackContext.HasIssues())
+		{
+			FString IssuesString;
+			for (const FString& Error : FeedbackContext.Errors)
+			{
+				IssuesString.Appendf(TEXT("Error: %s\n\n"), *Error);
+			}
+			for (const FString& Warning : FeedbackContext.Warnings)
+			{
+				IssuesString.Appendf(TEXT("Warning: %s\n\n"), *Warning);
+			}
+
+			TSharedRef<SMessageDialog> IssuesDialog = SNew(SMessageDialog)
+				.Title(FText(LOCTEXT("BakerIssuesDialogTitle", "Baker Issues")))
+				.Message(FText::FromString(IssuesString))
+				.Buttons({SCustomDialog::FButton(LOCTEXT("Ok", "Ok"))});
+
+			IssuesDialog->Show();
+		}
 	}
 	return FReply::Handled();
 }
diff --git a/Engine/Plugins/FX/NiagaraSimCaching/Source/NiagaraSimCachingEditor/Private/Sequencer/MovieSceneNiagaraTrackRecorder.cpp b/Engine/Plugins/FX/NiagaraSimCaching/Source/NiagaraSimCachingEditor/Private/Sequencer/MovieSceneNiagaraTrackRecorder.cpp
--- a/Engine/Plugins/FX/NiagaraSimCaching/Source/NiagaraSimCachingEditor/Private/Sequencer/MovieSceneNiagaraTrackRecorder.cpp
+++ b/Engine/Plugins/FX/NiagaraSimCaching/Source/NiagaraSimCachingEditor/Private/Sequencer/MovieSceneNiagaraTrackRecorder.cpp
@@ -231,12 +231,7 @@ void UMovieSceneNiagaraTrackRecorder::RecordSampleImpl(const FQualifiedFrameTime
 
 	if (NiagaraCacheSection.IsValid() && SystemToRecord.IsValid())
 	{
-		FNiagaraSimCacheWriteResult Result = NiagaraCacheSection->Params.SimCache->WriteFrame(SystemToRecord.Get());
-		if (!Result.bSuccess)
-		{
-			UE_LOG(LogNiagaraSimCachingEditor, Warning, TEXT("Error trying to write sim cache: %s"), *Result.ErrorMsg);
-		}
-		else
+		if ( NiagaraCacheSection->Params.SimCache->WriteFrame(SystemToRecord.Get()) )
 		{
 			// Expand the section to the new length
 			NiagaraCacheSection->SetEndFrame(CurrentFrame);
