Fix Landscape DI when dealing with huge system bounds exceeding limits of int32 for number of potential landscape components

diff --git a/Engine/Plugins/FX/Niagara/Source/Niagara/Private/NiagaraDataInterfaceLandscape.cpp b/Engine/Plugins/FX/Niagara/Source/Niagara/Private/NiagaraDataInterfaceLandscape.cpp
--- a/Engine/Plugins/FX/Niagara/Source/Niagara/Private/NiagaraDataInterfaceLandscape.cpp
+++ b/Engine/Plugins/FX/Niagara/Source/Niagara/Private/NiagaraDataInterfaceLandscape.cpp
@@ -966,30 +966,48 @@ FNDI_Landscape_SharedResourceHandle FNDI_Landscape_GeneratedData::GetLandscapeDa
 	}
 
 	// we want to use the bounds of the system to figure out which cells of the landscape that we need to handle
+	const int32 MaxLandscapeRegionCount = LandscapeInfo->XYtoCollisionComponentMap.Num();
+
 	const FBox SystemWorldBounds = SystemInstance.GetLocalBounds().TransformBy(SystemInstance.GetWorldTransform());
 	const FTransform& LandscapeActorToWorld = Landscape->LandscapeActorToWorld();
 	const FVector SystemMinInLandscape = LandscapeActorToWorld.InverseTransformPosition(SystemWorldBounds.Min);
 	const FVector SystemMaxInLandscape = LandscapeActorToWorld.InverseTransformPosition(SystemWorldBounds.Max);
 
-	const FBox SystemBoundsInLandscape(
+	FBox SystemBoundsInLandscape(
 		SystemMinInLandscape.ComponentMin(SystemMaxInLandscape),
 		SystemMinInLandscape.ComponentMax(SystemMaxInLandscape));
 
-	// at runtime we don't know the potential number of components in the landscape, so we can only clamp the values to start at 0.
-	// Ideally, we'd be able to specify an upper bound as well to further refine the region of overlap.
-	const FIntPoint MinBoundRegion = FIntPoint(int32(SystemBoundsInLandscape.Min.X), int32(SystemBoundsInLandscape.Min.Y)).ComponentMax(FIntPoint(0, 0));	// LWC Precision Loss
-	const FIntPoint MaxBoundRegion = FIntPoint(int32(SystemBoundsInLandscape.Max.X), int32(SystemBoundsInLandscape.Max.Y)).ComponentMax(FIntPoint(0, 0));
+	// transform the above box into a range of integers covering the cells of the landscape
+	// first clamp it at 0
+	SystemBoundsInLandscape.Min = SystemBoundsInLandscape.Min.ComponentMax(FVector::ZeroVector);
+	SystemBoundsInLandscape.Max = SystemBoundsInLandscape.Max.ComponentMax(FVector::ZeroVector);
+
+	// truncate to integers
+	const FVector MaxIntValue = FVector(double(TNumericLimits<int32>::Max()));
+	SystemBoundsInLandscape.Min = SystemBoundsInLandscape.Min.ComponentMin(MaxIntValue); // LWC Precision Loss 
+	SystemBoundsInLandscape.Max = SystemBoundsInLandscape.Max.ComponentMin(MaxIntValue);
+
+	// scale it by the quad size
+	const double QuadSizeScaleFactor = 1.0 / ((double)Landscape->ComponentSizeQuads);
+
+	const FIntRect SystemRect = FIntRect(
+		FIntPoint(int32(SystemBoundsInLandscape.Min.X), int32(SystemBoundsInLandscape.Min.Y)),
+		FIntPoint(int32(SystemBoundsInLandscape.Max.X), int32(SystemBoundsInLandscape.Max.Y))).Scale(QuadSizeScaleFactor);
+
+	// for obnoxiously large system bounds we need to guard against potential overflow on the number of cells
+	const int32 MaxSystemWidth = FMath::Clamp<int32>(SystemRect.Max.X - SystemRect.Min.X, 0, MaxLandscapeRegionCount);
+	const int32 MaxSystemHeight = FMath::Clamp<int32>(SystemRect.Max.Y - SystemRect.Min.Y, 0, MaxLandscapeRegionCount);
 
-	const FIntRect SystemRect = FIntRect(MinBoundRegion, MaxBoundRegion).Scale(1.0f / Landscape->ComponentSizeQuads);
+	const int64 MaxSystemRegionCount64 = int64(MaxSystemWidth) * int64(MaxSystemHeight);
+	const int32 MaxSystemRegionCount = int32(FMath::Min<int64>(int64(TNumericLimits<int32>::Max()), MaxSystemRegionCount64));
+
+	const int32 MaxRegionCount = FMath::Min(MaxSystemRegionCount, MaxLandscapeRegionCount);
 
 	FNDI_Landscape_SharedResource::FResourceKey Key;
 	Key.Source = Landscape;
+	Key.CapturedRegions.Reserve(MaxRegionCount);
 	Key.MinCaptureRegion = FIntPoint(TNumericLimits<int32>::Max(), TNumericLimits<int32>::Max());
 	Key.MaxCaptureRegion = FIntPoint(TNumericLimits<int32>::Min(), TNumericLimits<int32>::Min());
-
-	const int32 MaxSystemRegionCount = (SystemRect.Max.X - SystemRect.Min.X) * (SystemRect.Max.Y - SystemRect.Min.Y);
-	const int32 MaxLandscapeRegionCount = LandscapeInfo->XYtoCollisionComponentMap.Num();
-
 	Key.IncludesCachedHeight = InstanceData.RequiresCollisionCacheGpu;
 	Key.IncludesCachedPhysMat = InstanceData.RequiresPhysMatCacheGpu;
 	Key.PhysicalMaterials.Reserve(LandscapeDI.PhysicalMaterials.Num());
@@ -1007,8 +1025,6 @@ FNDI_Landscape_SharedResourceHandle FNDI_Landscape_GeneratedData::GetLandscapeDa
 
 	if (MaxSystemRegionCount > MaxLandscapeRegionCount)
 	{
-		Key.CapturedRegions.Reserve(MaxLandscapeRegionCount);
-
 		for (const auto& LandscapeComponent : LandscapeInfo->XYtoCollisionComponentMap)
 		{
 			if (SystemRect.Contains(LandscapeComponent.Key))
@@ -1019,8 +1035,6 @@ FNDI_Landscape_SharedResourceHandle FNDI_Landscape_GeneratedData::GetLandscapeDa
 	}
 	else
 	{
-		Key.CapturedRegions.Reserve(MaxSystemRegionCount);
-
 		for (int32 GridY = SystemRect.Min.Y; GridY < SystemRect.Max.Y; ++GridY)
 		{
 			for (int32 GridX = SystemRect.Min.X; GridX < SystemRect.Max.X; ++GridX)
