- Only adjust desired age by a small time step when GetLockDesiredAgeDeltaTimeToSeekDelta is enabled
- This fixes and issue with sim cache tracks not matching none sim cache tracks

diff --git a/Engine/Plugins/FX/Niagara/Source/Niagara/Private/MovieScene/MovieSceneNiagaraSystemTrackTemplate.cpp b/Engine/Plugins/FX/Niagara/Source/Niagara/Private/MovieScene/MovieSceneNiagaraSystemTrackTemplate.cpp
--- a/Engine/Plugins/FX/Niagara/Source/Niagara/Private/MovieScene/MovieSceneNiagaraSystemTrackTemplate.cpp
+++ b/Engine/Plugins/FX/Niagara/Source/Niagara/Private/MovieScene/MovieSceneNiagaraSystemTrackTemplate.cpp
@@ -238,7 +238,7 @@ struct FNiagaraSystemUpdateDesiredAgeExecutionToken : IMovieSceneExecutionToken
 					// Add a quarter of a frame offset here to push the desired age into the middle of the frame since it will be automatically rounded
 					// down to the nearest seek delta.  This prevents a situation where float rounding results in a value which is just slightly less than
 					// the frame boundary, which results in a skipped simulation frame.
-					float FrameOffset = NiagaraComponent->GetSeekDelta() / 4;
+					const float FrameOffset = NiagaraComponent->GetLockDesiredAgeDeltaTimeToSeekDelta() ? (NiagaraComponent->GetSeekDelta() / 4.0f) : 0.0f;
 					NiagaraComponent->SetDesiredAge(DesiredAge + FrameOffset);
 				}
 			}
