- Added interpolate / extrapolate options to simulation cache
- Experimental and will need further changes to make more performant / correct

diff --git a/Engine/Plugins/FX/Niagara/Source/Niagara/Classes/NiagaraSimCache.h b/Engine/Plugins/FX/Niagara/Source/Niagara/Classes/NiagaraSimCache.h
--- a/Engine/Plugins/FX/Niagara/Source/Niagara/Classes/NiagaraSimCache.h
+++ b/Engine/Plugins/FX/Niagara/Source/Niagara/Classes/NiagaraSimCache.h
@@ -41,6 +41,8 @@ struct FNiagaraSimCacheCreateParameters
 	FNiagaraSimCacheCreateParameters()
 		: bAllowRebasing(true)
 		, bAllowDataInterfaceCaching(true)
+		, bAllowInterpolation(false)
+		, bAllowVelocityExtrapolation(false)
 	{
 	}
 
@@ -65,6 +67,21 @@ struct FNiagaraSimCacheCreateParameters
 	UPROPERTY(EditAnywhere, BlueprintReadWrite, Category="SimCache")
 	uint32 bAllowDataInterfaceCaching : 1;
 
+	/**
+	When enabled we allow the cache to be generated for interpolation.
+	This will increase the memory usage for the cache slightly but can allow you to reduce the capture rate.
+	By default we will capture and interpolate all Position & Quat types, you can adjust this using the include / exclude list.
+	*/
+	UPROPERTY(EditAnywhere, BlueprintReadWrite, Category="SimCache")
+	uint32 bAllowInterpolation : 1;
+
+	/**
+	When enabled we allow the cache to be generated for extrapolation.
+	This will force the velocity attribute to be maintained.
+	*/
+	UPROPERTY(EditAnywhere, BlueprintReadWrite, Category="SimCache")
+	uint32 bAllowVelocityExtrapolation : 1;
+
 	/**
 	List of Attributes to force include in the SimCache rebase, they should be the full path to the attribute
 	For example, MyEmitter.Particles.MyQuat would force the particle attribute MyQuat to be included for MyEmitter
@@ -79,6 +96,20 @@ struct FNiagaraSimCacheCreateParameters
 	UPROPERTY(EditAnywhere, BlueprintReadWrite, Category = "SimCache")
 	TArray<FName> RebaseExcludeAttributes;
 
+	/**
+	List of specific Attributes to include when using interpolation.  They must be types that are supported for interpolation.
+	For example, MyEmitter.Particles.MyPosition would force MyPosition to be interpolated.
+	*/
+	UPROPERTY(EditAnywhere, BlueprintReadWrite, Category = "SimCache")
+	TArray<FName> InterpolationIncludeAttributes;
+
+	/**
+	List of specific Attributes to exclude interpolation for.  They must be types that are supported for interpolation.
+	For example, MyEmitter.Particles.MyPosition would force MyPosition to be interpolated.
+	*/
+	UPROPERTY(EditAnywhere, BlueprintReadWrite, Category = "SimCache")
+	TArray<FName> InterpolationExcludeAttributes;
+
 	/**
 	List of attributes to capture when the capture attribute capture mode is set to explicit.
 	For example, adding MyEmitter.Particles.Position will only gather that attribute inside the cache.
@@ -109,6 +140,9 @@ struct FNiagaraSimCacheDataBuffers
 
 	UPROPERTY()
 	uint32 IDAcquireTag = 0;
+
+	UPROPERTY()
+	TArray<uint32> InterpMapping;
 };
 
 USTRUCT()
@@ -150,6 +184,9 @@ struct FNiagaraSimCacheFrame
 	UPROPERTY()
 	FVector3f LWCTile = FVector3f::ZeroVector;
 
+	UPROPERTY()
+	float SimulationAge = 0.0f;
+
 	UPROPERTY()
 	FNiagaraSimCacheSystemFrame SystemData;
 
@@ -189,14 +226,28 @@ struct FNiagaraSimCacheDataBuffersLayout
 {
 	GENERATED_BODY()
 
-	// Copy Function parameters are
-	// Dest, DestStride, Source, SourceStride, NumInstances, RebasedTransform
-	typedef void (*FVariableCopyFunction)(uint8*, uint32, const uint8*, uint32, uint32, const FTransform&);
+	struct FVariableCopyContext
+	{
+		float					FrameFraction		= 0.0f;
+		float					RecipDt				= 0.0f;
+		uint32					NumInstances		= 0;
+		uint8*					Dest				= nullptr;
+		uint32					DestStride			= 0;
+		const uint8*			SourceAComponent	= nullptr;
+		uint32					SourceAStride		= 0;
+		const uint8*			SourceBComponent	= nullptr;
+		uint32					SourceBStride		= 0;
+		const uint8*			VelocityComponent	= nullptr;
+		FTransform				RebaseTransform;
+		TConstArrayView<uint32>	InterpMappings;
+	};
+
+	typedef void (*FVariableCopyFunction)(const FVariableCopyContext& CopyData);
 
-	struct FVariableCopyInfo
+	struct FVariableCopyMapping
 	{
-		FVariableCopyInfo() = default;
-		explicit FVariableCopyInfo(uint16 InComponentFrom, uint16 InComponentTo, FVariableCopyFunction InCopyFunc)
+		FVariableCopyMapping() = default;
+		explicit FVariableCopyMapping(uint16 InComponentFrom, uint16 InComponentTo, FVariableCopyFunction InCopyFunc)
 			: ComponentFrom(InComponentFrom)
 			, ComponentTo(InComponentTo)
 			, CopyFunc(InCopyFunc)
@@ -229,13 +280,24 @@ struct FNiagaraSimCacheDataBuffersLayout
 	UPROPERTY()
 	bool bLocalSpace = false;
 
+	UPROPERTY()
+	bool bAllowInterpolation = false;
+
+	UPROPERTY()
+	bool bAllowVelocityExtrapolation = false;
+
 	UPROPERTY()
 	TArray<FName> RebaseVariableNames;
 
+	UPROPERTY()
+	TArray<FName> InterpVariableNames;
+
 	TArray<uint16> ComponentMappingsToDataBuffer;
-	TArray<FVariableCopyInfo> VariableMappingsToDataBuffer;
+	TArray<FVariableCopyMapping> VariableCopyMappingsToDataBuffer;
+	uint16 ComponentVelocity = INDEX_NONE;
 
-	TArray<uint16> ComponentMappingsFromDataBuffer;
+	TArray<uint16>	ComponentMappingsFromDataBuffer;
+	uint16			ComponentUniqueID = INDEX_NONE;
 };
 
 USTRUCT()
@@ -274,6 +336,9 @@ public:
 	DECLARE_MULTICAST_DELEGATE_OneParam(FOnCacheBeginWrite, UNiagaraSimCache*)
 	DECLARE_MULTICAST_DELEGATE_OneParam(FOnCacheEndWrite, UNiagaraSimCache*)
 
+	// Used to reduce rounding issues with age
+	static constexpr float CacheAgeResolution = 10000.0f;
+
 	// UObject Interface
 	virtual bool IsReadyForFinishDestroy() override;
 	// UObject Interface
diff --git a/Engine/Plugins/FX/Niagara/Source/Niagara/Private/NiagaraSimCache.cpp b/Engine/Plugins/FX/Niagara/Source/Niagara/Private/NiagaraSimCache.cpp
--- a/Engine/Plugins/FX/Niagara/Source/Niagara/Private/NiagaraSimCache.cpp
+++ b/Engine/Plugins/FX/Niagara/Source/Niagara/Private/NiagaraSimCache.cpp
@@ -285,6 +285,7 @@ bool UNiagaraSimCache::WriteFrame(UNiagaraComponent* NiagaraComponent, FNiagaraS
 	FNiagaraSimCacheFrame& CacheFrame = CacheFrames.AddDefaulted_GetRef();
 	CacheFrame.LocalToWorld = Helper.SystemInstance->GatheredInstanceParameters.ComponentTrans;
 	CacheFrame.LWCTile = Helper.SystemInstance->GetLWCTile();
+	CacheFrame.SimulationAge = FMath::TruncToFloat(Helper.SystemInstance->GetAge() * CacheAgeResolution) / CacheAgeResolution;
 
 	CacheFrame.SystemData.LocalBounds = Helper.SystemInstance->GetLocalBounds();
 
@@ -379,6 +380,7 @@ bool UNiagaraSimCache::EndWrite()
 		SoftNiagaraSystem.Reset();
 	}
 
+
 	if (DataInterfaceStorage.IsEmpty() == false)
 	{
 		bool bDataInterfacesSucess = true;
@@ -396,6 +398,58 @@ bool UNiagaraSimCache::EndWrite()
 		}
 	}
 
+	// If we allow interpolation we need to build our mapping table from Current -> Next
+	if (CreateParameters.bAllowInterpolation)
+	{
+		const auto& CreateInterpMapping =
+			[this](const FNiagaraSimCacheDataBuffersLayout& BufferLayout, FNiagaraSimCacheDataBuffers& BufferA, const FNiagaraSimCacheDataBuffers& BufferB)
+			{
+				if (BufferLayout.bAllowInterpolation == false)
+				{
+					return;
+				}
+
+				for (int32 i = 0; i < BufferA.InterpMapping.Num(); ++i)
+				{
+					BufferA.InterpMapping[i] = BufferB.InterpMapping.IndexOfByKey(BufferA.InterpMapping[i]);
+				}
+			};
+
+		const auto& ClearInterpMapping =
+			[this](const FNiagaraSimCacheDataBuffersLayout& BufferLayout, FNiagaraSimCacheDataBuffers& BufferA)
+			{
+				if (BufferLayout.bAllowInterpolation == false)
+				{
+					return;
+				}
+
+				for (int32 i = 0; i < BufferA.InterpMapping.Num(); ++i)
+				{
+					BufferA.InterpMapping[i] = INDEX_NONE;
+				}
+			};
+
+		// Do all frames that can Interp first
+		for (int32 iFrame=0; iFrame < CacheFrames.Num() - 1; ++iFrame)
+		{
+			FNiagaraSimCacheFrame& CacheFrameA = CacheFrames[iFrame];
+			const FNiagaraSimCacheFrame& CacheFrameB = CacheFrames[iFrame + 1];
+
+			CreateInterpMapping(CacheLayout.SystemLayout, CacheFrameA.SystemData.SystemDataBuffers, CacheFrameB.SystemData.SystemDataBuffers);
+			for (int32 iEmitter=0; iEmitter < CacheLayout.EmitterLayouts.Num(); ++iEmitter)
+			{
+				CreateInterpMapping(CacheLayout.EmitterLayouts[iEmitter], CacheFrameA.EmitterData[iEmitter].ParticleDataBuffers, CacheFrameB.EmitterData[iEmitter].ParticleDataBuffers);
+			}
+		}
+
+		// Final frame is special as we won't Interp
+		ClearInterpMapping(CacheLayout.SystemLayout, CacheFrames.Last().SystemData.SystemDataBuffers);
+		for (int32 iEmitter = 0; iEmitter < CacheLayout.EmitterLayouts.Num(); ++iEmitter)
+		{
+			ClearInterpMapping(CacheLayout.EmitterLayouts[iEmitter], CacheFrames.Last().EmitterData[iEmitter].ParticleDataBuffers);
+		}
+	}
+
 	OnCacheEndWrite.Broadcast(this);
 	return IsCacheValid();
 }
@@ -479,6 +533,9 @@ bool UNiagaraSimCache::CanRead(UNiagaraSystem* NiagaraSystem)
 
 bool UNiagaraSimCache::Read(float TimeSeconds, FNiagaraSystemInstance* SystemInstance) const
 {
+	// Adjust time to match our allow resolution
+	TimeSeconds = FMath::TruncToFloat(TimeSeconds * CacheAgeResolution) / CacheAgeResolution;
+
 	const float RelativeTime = FMath::Max(TimeSeconds - StartSeconds, 0.0f);
 	if ( RelativeTime > DurationSeconds )
 	{
@@ -486,9 +543,43 @@ bool UNiagaraSimCache::Read(float TimeSeconds, FNiagaraSystemInstance* SystemIns
 		return false;
 	}
 
-	const float FrameTime		= DurationSeconds > 0.0f ? (RelativeTime / DurationSeconds) * float(CacheFrames.Num() - 1) : 0.0f;
-	const int32 FrameIndex		= FMath::FloorToInt(FrameTime);
-	const float FrameFraction	= FrameTime - float(FrameIndex);
+	//const float FrameTime		= DurationSeconds > 0.0f ? (RelativeTime / DurationSeconds) * float(CacheFrames.Num() - 1) : 0.0f;
+	//const int32 FrameIndex		= FMath::FloorToInt(FrameTime);
+	//const float FrameFraction	= FrameTime - float(FrameIndex);
+
+	// Take a guess at the frame then iterate to find the actual frame
+	// We have to do this because not all sources are recording the cache at a consistent rate
+	// This could also be replaced with a binary search but in most cases our guess will be very close to the actual
+	const int32 NumFramesMinusOne = CacheFrames.Num() - 1;
+
+	int32 FrameIndex = 0;
+	float FrameFraction = 0.0f;
+	if (DurationSeconds > 0.0f)
+	{
+		FrameIndex = FMath::FloorToInt((RelativeTime / DurationSeconds) * float(NumFramesMinusOne));
+	}
+
+	if ( CacheFrames[FrameIndex].SimulationAge > TimeSeconds)
+	{
+		while (FrameIndex > 0 && CacheFrames[FrameIndex].SimulationAge > TimeSeconds )
+		{
+			--FrameIndex;
+		}
+	}
+	else
+	{
+		while (FrameIndex < NumFramesMinusOne && CacheFrames[FrameIndex + 1].SimulationAge <= TimeSeconds)
+		{
+			++FrameIndex;
+		}
+	}
+
+	if (FrameIndex < NumFramesMinusOne)
+	{
+		const float SimulationAgeA = CacheFrames[FrameIndex].SimulationAge;
+		const float SimulationAgeB = CacheFrames[FMath::Min(FrameIndex + 1, NumFramesMinusOne)].SimulationAge;
+		FrameFraction = (TimeSeconds - SimulationAgeA) / (SimulationAgeB - SimulationAgeA);
+	}
 
 	return ReadFrame(FrameIndex, FrameFraction, SystemInstance);
 }
@@ -511,8 +602,12 @@ bool UNiagaraSimCache::ReadFrame(int32 FrameIndex, float FrameFraction, FNiagara
 		RebaseTransform = RebaseTransform * CacheFrame.LocalToWorld.Inverse();
 	}
 
+
+	const int32 NextFrameIndex = FMath::Min(FrameIndex + 1, CacheFrames.Num() - 1);
+	const float FrameDeltaSeconds = CacheFrames[NextFrameIndex].SimulationAge - CacheFrame.SimulationAge;
+
 	Helper.SystemInstance->LocalBounds = CacheFrame.SystemData.LocalBounds;
-	Helper.ReadDataBuffer(RebaseTransform, CacheLayout.SystemLayout, CacheFrame.SystemData.SystemDataBuffers, Helper.GetSystemSimulationDataSet());
+	Helper.ReadDataBuffer(FrameFraction, FrameDeltaSeconds, RebaseTransform, CacheLayout.SystemLayout, CacheFrame.SystemData.SystemDataBuffers, CacheFrame.SystemData.SystemDataBuffers, Helper.GetSystemSimulationDataSet());
 
 	const int32 NumEmitters = CacheLayout.EmitterLayouts.Num();
 	for (int32 i=0; i < NumEmitters; ++i)
@@ -522,13 +617,14 @@ bool UNiagaraSimCache::ReadFrame(int32 FrameIndex, float FrameFraction, FNiagara
 		EmitterInstance.CachedBounds = CacheEmitterFrame.LocalBounds;
 		EmitterInstance.TotalSpawnedParticles = CacheEmitterFrame.TotalSpawnedParticles;
 
+		const FNiagaraSimCacheEmitterFrame& CacheEmitterFrameB = CacheFrames[NextFrameIndex].EmitterData[i];
 		if (CacheLayout.EmitterLayouts[i].SimTarget == ENiagaraSimTarget::GPUComputeSim)
 		{
-			Helper.ReadDataBufferGPU(RebaseTransform, EmitterInstance, CacheLayout.EmitterLayouts[i], CacheEmitterFrame.ParticleDataBuffers, EmitterInstance.GetData(), PendingCommandsInFlight);
+			Helper.ReadDataBufferGPU(FrameFraction, FrameDeltaSeconds, RebaseTransform, EmitterInstance, CacheLayout.EmitterLayouts[i], CacheEmitterFrame.ParticleDataBuffers, CacheEmitterFrameB.ParticleDataBuffers, EmitterInstance.GetData(), PendingCommandsInFlight);
 		}
 		else
 		{
-			Helper.ReadDataBuffer(RebaseTransform, CacheLayout.EmitterLayouts[i], CacheEmitterFrame.ParticleDataBuffers, EmitterInstance.GetData());
+			Helper.ReadDataBuffer(FrameFraction, FrameDeltaSeconds, RebaseTransform, CacheLayout.EmitterLayouts[i], CacheEmitterFrame.ParticleDataBuffers, CacheEmitterFrameB.ParticleDataBuffers, EmitterInstance.GetData());
 		}
 	}
 
@@ -536,7 +632,6 @@ bool UNiagaraSimCache::ReadFrame(int32 FrameIndex, float FrameFraction, FNiagara
 	//-OPT: We shouldn't need to search all the time here
 	if (DataInterfaceStorage.IsEmpty() == false)
 	{
-		const int NextFrameIndex = FMath::Min(FrameIndex + 1, CacheFrames.Num() - 1);
 		bool bDataInterfacesSucess = true;
 
 		FNiagaraDataInterfaceUtilities::ForEachDataInterface(
diff --git a/Engine/Plugins/FX/Niagara/Source/Niagara/Private/NiagaraSimCacheHelper.h b/Engine/Plugins/FX/Niagara/Source/Niagara/Private/NiagaraSimCacheHelper.h
--- a/Engine/Plugins/FX/Niagara/Source/Niagara/Private/NiagaraSimCacheHelper.h
+++ b/Engine/Plugins/FX/Niagara/Source/Niagara/Private/NiagaraSimCacheHelper.h
@@ -65,11 +65,12 @@ struct FNiagaraSimCacheHelper
 	bool HasValidSimulation() const { return SystemSimulation != nullptr; }
 	bool HasValidSimulationData() const { return SystemSimulationDataBuffer != nullptr; }
 
-	void BuildCacheLayout(FNiagaraSimCacheDataBuffersLayout& CacheLayout, const FNiagaraDataSetCompiledData& CompiledData, FName LayoutName, TArray<FName> InRebaseVariableNames, TConstArrayView<FName> ExplicitCaptureAttributes) const
+	void BuildCacheLayout(const FNiagaraSimCacheCreateParameters& CreateParameters, FNiagaraSimCacheDataBuffersLayout& CacheLayout, const FNiagaraDataSetCompiledData& CompiledData, FName LayoutName, TArray<FName> InRebaseVariableNames, TArray<FName> InInterpVariableNames, TConstArrayView<FName> ExplicitCaptureAttributes) const
 	{
 		CacheLayout.LayoutName = LayoutName;
 		CacheLayout.SimTarget = CompiledData.SimTarget;
 
+		// Determine the components to cache
 		int32 TotalCacheComponents = 0;
 		TArray<int32> CacheToDataSetVariables;
 
@@ -85,6 +86,23 @@ struct FNiagaraSimCacheHelper
 			}
 		}
 
+		// We need to preserve the velocity attribute if we want to use velocity based extrapolation of positions
+		if (CreateParameters.bAllowVelocityExtrapolation)
+		{
+			const FNiagaraVariableBase VelocityVariable(FNiagaraTypeDefinition::GetVec3Def(), "Velocity");
+			const int32 VelocityAttributeIndex = CompiledData.Variables.IndexOfByKey(VelocityVariable);
+			if (VelocityAttributeIndex != INDEX_NONE)
+			{
+				CacheLayout.bAllowVelocityExtrapolation = true;
+				if (!CacheToDataSetVariables.Contains(VelocityAttributeIndex))
+				{
+					CacheToDataSetVariables.AddUnique(VelocityAttributeIndex);
+					const FNiagaraVariableLayoutInfo& DataSetVariableLayout = CompiledData.VariableLayouts[VelocityAttributeIndex];
+					TotalCacheComponents += DataSetVariableLayout.GetNumFloatComponents() + DataSetVariableLayout.GetNumHalfComponents() + DataSetVariableLayout.GetNumInt32Components();
+				}
+			}
+		}
+
 		const int32 NumCacheVariables = CacheToDataSetVariables.Num();
 
 		CacheLayout.Variables.AddDefaulted(NumCacheVariables);
@@ -92,6 +110,7 @@ struct FNiagaraSimCacheHelper
 		CacheLayout.ComponentMappingsFromDataBuffer.Empty(TotalCacheComponents);
 		CacheLayout.ComponentMappingsFromDataBuffer.AddDefaulted(TotalCacheComponents);
 		CacheLayout.RebaseVariableNames = MoveTemp(InRebaseVariableNames);
+		CacheLayout.InterpVariableNames = MoveTemp(InInterpVariableNames);
 
 		for ( int32 iCacheVariable=0; iCacheVariable < NumCacheVariables; ++iCacheVariable)
 		{
@@ -112,6 +131,19 @@ struct FNiagaraSimCacheHelper
 			CacheLayout.Int32Count += uint16(DataSetVariableLayout.GetNumInt32Components());
 		}
 
+		if (CreateParameters.bAllowInterpolation)
+		{
+			const FNiagaraVariableBase UniqueIDVariable(FNiagaraTypeDefinition::GetIntDef(), "UniqueID");
+			CacheLayout.ComponentUniqueID = CompiledData.Variables.IndexOfByKey(UniqueIDVariable);
+			if (CacheLayout.ComponentUniqueID != uint16(INDEX_NONE))
+			{
+				const FNiagaraVariableLayoutInfo& DataSetVariableLayout = CompiledData.VariableLayouts[CacheLayout.ComponentUniqueID];
+				check(DataSetVariableLayout.GetNumInt32Components() == 1);
+				CacheLayout.ComponentUniqueID = DataSetVariableLayout.Int32ComponentStart;
+				CacheLayout.bAllowInterpolation = true;
+			}
+		}
+
 		// Build write mappings we will build read mappings in a separate path
 		int32 FloatOffset = 0;
 		int32 HalfOffset = CacheLayout.FloatCount;
@@ -166,7 +198,7 @@ struct FNiagaraSimCacheHelper
 				}
 			}
 
-			for (const FNiagaraVariable& Variable : SystemCompileData.Variables)
+			for (const FNiagaraVariableBase& Variable : SystemCompileData.Variables)
 			{
 				if (Variable.GetType() == FNiagaraTypeDefinition::GetPositionDef())
 				{
@@ -193,7 +225,24 @@ struct FNiagaraSimCacheHelper
 			}
 		}
 
-		BuildCacheLayout(CacheLayout, SystemCompileData, NiagaraSystem->GetFName(), MoveTemp(RebaseVariableNames), CreateParameters.ExplicitCaptureAttributes);
+		TArray<FName> InterpVariableNames;
+		if (CreateParameters.bAllowInterpolation)
+		{
+			for (const FNiagaraVariableBase& Variable : SystemCompileData.Variables)
+			{
+				if (!CanInterpolateVariable(Variable) || CreateParameters.RebaseExcludeAttributes.Contains(Variable.GetName()))
+				{
+					continue;
+				}
+
+				if (CreateParameters.InterpolationIncludeAttributes.Num() == 0 || CreateParameters.InterpolationIncludeAttributes.Contains(Variable.GetName()))
+				{
+					InterpVariableNames.Add(Variable.GetName());
+				}
+			}
+		}
+
+		BuildCacheLayout(CreateParameters, CacheLayout, SystemCompileData, NiagaraSystem->GetFName(), MoveTemp(RebaseVariableNames), MoveTemp(InterpVariableNames), CreateParameters.ExplicitCaptureAttributes);
 	}
 
 	void BuildCacheLayoutForEmitter(const FNiagaraSimCacheCreateParameters& CreateParameters, FNiagaraSimCacheDataBuffersLayout& CacheLayout, int EmitterIndex)
@@ -262,7 +311,7 @@ struct FNiagaraSimCacheHelper
 		#endif
 
 			// Look for regular attributes that we are forcing to rebase or can rebase like positions
-			for (const FNiagaraVariable& Variable : EmitterCompiledData.DataSetCompiledData.Variables)
+			for (const FNiagaraVariableBase& Variable : EmitterCompiledData.DataSetCompiledData.Variables)
 			{
 				if ( Variable.GetType() == FNiagaraTypeDefinition::GetPositionDef() )
 				{
@@ -277,6 +326,24 @@ struct FNiagaraSimCacheHelper
 				}
 			}
 		}
+
+		TArray<FName> InterpVariableNames;
+		if (CreateParameters.bAllowInterpolation)
+		{
+			for (const FNiagaraVariableBase& Variable : EmitterCompiledData.DataSetCompiledData.Variables)
+			{
+				if (!CanInterpolateVariable(Variable) || CreateParameters.RebaseExcludeAttributes.Contains(Variable.GetName()))
+				{
+					continue;
+				}
+
+				if (CreateParameters.InterpolationIncludeAttributes.Num() == 0 || CreateParameters.InterpolationIncludeAttributes.Contains(Variable.GetName()))
+				{
+					InterpVariableNames.Add(Variable.GetName());
+				}
+			}
+		}
+
 		TArray<FName> ExplicitCaptureAttributes;
 		if ( CreateParameters.ExplicitCaptureAttributes.Num() > 0 )
 		{
@@ -294,7 +361,7 @@ struct FNiagaraSimCacheHelper
 			}
 		}
 
-		BuildCacheLayout(CacheLayout, EmitterCompiledData.DataSetCompiledData, EmitterHandle.GetName(), MoveTemp(RebaseVariableNames), ExplicitCaptureAttributes);
+		BuildCacheLayout(CreateParameters, CacheLayout, EmitterCompiledData.DataSetCompiledData, EmitterHandle.GetName(), MoveTemp(RebaseVariableNames), MoveTemp(InterpVariableNames), ExplicitCaptureAttributes);
 	}
 
 	static bool BuildCacheReadMappings(FNiagaraSimCacheDataBuffersLayout& CacheLayout, const FNiagaraDataSetCompiledData& CompiledData)
@@ -302,7 +369,8 @@ struct FNiagaraSimCacheHelper
 		const int32 CacheTotalComponents = CacheLayout.FloatCount + CacheLayout.HalfCount + CacheLayout.Int32Count;
 		CacheLayout.ComponentMappingsToDataBuffer.Empty(CacheTotalComponents);
 		CacheLayout.ComponentMappingsToDataBuffer.AddDefaulted(CacheTotalComponents);
-		CacheLayout.VariableMappingsToDataBuffer.Empty(0);
+		CacheLayout.VariableCopyMappingsToDataBuffer.Empty(0);
+		CacheLayout.ComponentVelocity = INDEX_NONE;
 
 		int32 FloatOffset = 0;
 		int32 HalfOffset = CacheLayout.FloatCount;
@@ -325,27 +393,60 @@ struct FNiagaraSimCacheHelper
 				}
 			}
 
+			// If this is our velocity component track it
+			//-OPT: Don't need to do this each time
+			if (CacheLayout.bAllowVelocityExtrapolation)
+			{
+				const FNiagaraVariableBase VelocityVariable(FNiagaraTypeDefinition::GetVec3Def(), "Velocity");
+
+				if (SourceVariable.Variable == VelocityVariable)
+				{
+					CacheLayout.ComponentVelocity = uint16(FloatOffset);
+				}
+			}
+
 			// Is this a type that requires conversion / re-basing?
 			if (DestVariableLayout != nullptr)
 			{
-				if ( CacheLayout.RebaseVariableNames.Contains(SourceVariable.Variable.GetName()) )
+				const bool bInterpVariable = CacheLayout.InterpVariableNames.Contains(SourceVariable.Variable.GetName());
+				const bool bRebaseVariable = CacheLayout.RebaseVariableNames.Contains(SourceVariable.Variable.GetName());
+
+				if (bInterpVariable || bRebaseVariable)
 				{
 					if (SourceVariable.Variable.GetType() == FNiagaraTypeDefinition::GetPositionDef())
 					{
 						check(SourceVariable.FloatCount == 3);
-						CacheLayout.VariableMappingsToDataBuffer.Emplace(uint16(FloatOffset), uint16(DestVariableLayout->FloatComponentStart), &FNiagaraSimCacheHelper::CopyPositions);
+						if (bInterpVariable)
+						{
+							CacheLayout.VariableCopyMappingsToDataBuffer.Emplace(uint16(FloatOffset), uint16(DestVariableLayout->FloatComponentStart), bRebaseVariable ? &FNiagaraSimCacheHelper::InterpPositions<true> : &FNiagaraSimCacheHelper::InterpPositions<false>);
+						}
+						else if ( CacheLayout.bAllowVelocityExtrapolation )
+						{
+							CacheLayout.VariableCopyMappingsToDataBuffer.Emplace(uint16(FloatOffset), uint16(DestVariableLayout->FloatComponentStart), bRebaseVariable ? &FNiagaraSimCacheHelper::ExtrapolatePositions<true> : &FNiagaraSimCacheHelper::ExtrapolatePositions<false>);
+						}
+						else
+						{
+							CacheLayout.VariableCopyMappingsToDataBuffer.Emplace(uint16(FloatOffset), uint16(DestVariableLayout->FloatComponentStart), &FNiagaraSimCacheHelper::CopyPositions);
+						}
 						DestVariableLayout = nullptr;
 					}
 					else if (SourceVariable.Variable.GetType() == FNiagaraTypeDefinition::GetQuatDef())
 					{
 						check(SourceVariable.FloatCount == 4);
-						CacheLayout.VariableMappingsToDataBuffer.Emplace(uint16(FloatOffset), uint16(DestVariableLayout->FloatComponentStart), &FNiagaraSimCacheHelper::CopyQuaternions);
+						if (bInterpVariable)
+						{
+							CacheLayout.VariableCopyMappingsToDataBuffer.Emplace(uint16(FloatOffset), uint16(DestVariableLayout->FloatComponentStart), bInterpVariable  ? &FNiagaraSimCacheHelper::InterpQuaternions<true> : &FNiagaraSimCacheHelper::InterpQuaternions<false>);
+						}
+						else
+						{
+							CacheLayout.VariableCopyMappingsToDataBuffer.Emplace(uint16(FloatOffset), uint16(DestVariableLayout->FloatComponentStart), &FNiagaraSimCacheHelper::CopyQuaternions);
+						}
 						DestVariableLayout = nullptr;
 					}
 					else if (SourceVariable.Variable.GetType() == FNiagaraTypeDefinition::GetMatrix4Def())
 					{
 						check(SourceVariable.FloatCount == 16);
-						CacheLayout.VariableMappingsToDataBuffer.Emplace(uint16(FloatOffset), uint16(DestVariableLayout->FloatComponentStart), &FNiagaraSimCacheHelper::CopyMatrices);
+						CacheLayout.VariableCopyMappingsToDataBuffer.Emplace(uint16(FloatOffset), uint16(DestVariableLayout->FloatComponentStart), &FNiagaraSimCacheHelper::CopyMatrices);
 						DestVariableLayout = nullptr;
 					}
 				}
@@ -424,6 +525,15 @@ struct FNiagaraSimCacheHelper
 		// Copy ID to Index Table
 		CacheBuffer.IDToIndexTable = DataBuffer.GetIDTable();
 		CacheBuffer.IDAcquireTag = DataBuffer.GetIDAcquireTag();
+
+		// Generate a interp mapping (if we have enabled it)
+		if (CacheLayout.bAllowInterpolation)
+		{
+			//-TODO: Persistent ID mapping
+			CacheBuffer.InterpMapping.SetNumUninitialized(DataBuffer.GetNumInstances());
+			const uint8* UniqueIDs = DataBuffer.GetComponentPtrInt32(CacheLayout.ComponentUniqueID);
+			FMemory::Memcpy(CacheBuffer.InterpMapping.GetData(), UniqueIDs, DataBuffer.GetNumInstances() * sizeof(int32));
+		}
 	}
 
 	void WriteDataBufferGPU(FNiagaraEmitterInstance& EmitterInstance, const FNiagaraDataBuffer& DataBuffer, const FNiagaraSimCacheDataBuffersLayout& CacheLayout, FNiagaraSimCacheDataBuffers& CacheBuffer)
@@ -437,75 +547,109 @@ struct FNiagaraSimCacheHelper
 		}
 	}
 
-	void ReadDataBuffer(const FTransform& RebaseTransform, const FNiagaraSimCacheDataBuffersLayout& CacheLayout, const FNiagaraSimCacheDataBuffers& CacheBuffer, FNiagaraDataSet& DataSet)
+	static void ReadFloatBuffers(int32& iComponent, const FNiagaraSimCacheDataBuffersLayout& CacheLayout, const FNiagaraSimCacheDataBuffers& CacheBuffer, TArrayView<uint8> DestBuffer, uint32 DestStride)
 	{
-		FNiagaraDataBuffer& DataBuffer = DataSet.BeginSimulate();
-		DataBuffer.Allocate(CacheBuffer.NumInstances);
-		DataBuffer.SetNumInstances(CacheBuffer.NumInstances);
-		if ( CacheBuffer.NumInstances > 0 )
+		const int32 NumInstances = CacheBuffer.NumInstances;
+		for (uint32 i = 0; i < CacheLayout.FloatCount; ++i)
 		{
-			int32 iComponent = 0;
-			const int32 NumInstances = CacheBuffer.NumInstances;
-
-			// Copy Float
-			for (uint32 i=0; i < CacheLayout.FloatCount; ++i)
+			const uint32 Component = CacheLayout.ComponentMappingsToDataBuffer[iComponent++];
+			if (Component == InvalidComponent)
 			{
-				const uint32 Component = CacheLayout.ComponentMappingsToDataBuffer[iComponent++];
-				if (Component == InvalidComponent)
-				{
-					continue;
-				}
-				const uint8* Source = CacheBuffer.FloatData.GetData() + (i * NumInstances * sizeof(float));
-				uint8* Dest = DataBuffer.GetComponentPtrFloat(Component);
-				CheckedMemcpy(DataBuffer.GetFloatBuffer(), Dest, CacheBuffer.FloatData, Source, sizeof(float) * NumInstances);
-				//FMemory::Memcpy(Dest, Source, sizeof(float) * NumInstances);
+				continue;
 			}
+			const uint8* SourceFloats = CacheBuffer.FloatData.GetData() + (i * NumInstances * sizeof(float));
+			uint8* DestFloats = DestBuffer.GetData() + (DestStride * Component);
+			CheckedMemcpy(DestBuffer, DestFloats, CacheBuffer.FloatData, SourceFloats, sizeof(float) * NumInstances);
+		}
+	}
 
-			// Copy Half
-			for (uint32 i = 0; i < CacheLayout.HalfCount; ++i)
+	static void ReadHalfBuffers(int32& iComponent, const FNiagaraSimCacheDataBuffersLayout& CacheLayout, const FNiagaraSimCacheDataBuffers& CacheBuffer, TArrayView<uint8> DestBuffer, uint32 DestStride)
+	{
+		const int32 NumInstances = CacheBuffer.NumInstances;
+		for (uint32 i = 0; i < CacheLayout.HalfCount; ++i)
+		{
+			const uint32 Component = CacheLayout.ComponentMappingsToDataBuffer[iComponent++];
+			if (Component == InvalidComponent)
 			{
-				const uint32 Component = CacheLayout.ComponentMappingsToDataBuffer[iComponent++];
-				if (Component == InvalidComponent)
-				{
-					continue;
-				}
-				const uint8* Source = CacheBuffer.HalfData.GetData() + (i * NumInstances * sizeof(FFloat16));
-				uint8* Dest = DataBuffer.GetComponentPtrHalf(Component);
-				CheckedMemcpy(DataBuffer.GetHalfBuffer(), Dest, CacheBuffer.HalfData, Source, sizeof(FFloat16) * NumInstances);
-				//FMemory::Memcpy(Dest, Source, sizeof(FFloat16) * NumInstances);
+				continue;
 			}
+			const uint8* SourceHalfs = CacheBuffer.HalfData.GetData() + (i * NumInstances * sizeof(FFloat16));
+			uint8* DestHalfs = DestBuffer.GetData() + (DestStride * Component);
+			CheckedMemcpy(DestBuffer, DestHalfs, CacheBuffer.HalfData, SourceHalfs, sizeof(FFloat16) * NumInstances);
+		}
+	}
 
-			// Copy Int32
-			for (uint32 i = 0; i < CacheLayout.Int32Count; ++i)
+	static void ReadInt32Buffers(int32& iComponent, const FNiagaraSimCacheDataBuffersLayout& CacheLayout, const FNiagaraSimCacheDataBuffers& CacheBuffer, TArrayView<uint8> DestBuffer, uint32 DestStride)
+	{
+		const int32 NumInstances = CacheBuffer.NumInstances;
+		for (uint32 i = 0; i < CacheLayout.Int32Count; ++i)
+		{
+			const uint32 Component = CacheLayout.ComponentMappingsToDataBuffer[iComponent++];
+			if (Component == InvalidComponent)
 			{
-				const uint32 Component = CacheLayout.ComponentMappingsToDataBuffer[iComponent++];
-				if (Component == InvalidComponent)
-				{
-					continue;
-				}
-				const uint8* Source = CacheBuffer.Int32Data.GetData() + (i * NumInstances * sizeof(int32));
-				uint8* Dest = DataBuffer.GetComponentPtrInt32(Component);
-				CheckedMemcpy(DataBuffer.GetInt32Buffer(), Dest, CacheBuffer.Int32Data, Source, sizeof(int32) * NumInstances);
-				//FMemory::Memcpy(Dest, Source, sizeof(int32) * NumInstances);
+				continue;
 			}
+			const uint8* SourceInt32s = CacheBuffer.Int32Data.GetData() + (i * NumInstances * sizeof(int32));
+			uint8* DestInt32s = DestBuffer.GetData() + (DestStride * Component);
+			CheckedMemcpy(DestBuffer, DestInt32s, CacheBuffer.Int32Data, SourceInt32s, sizeof(int32) * NumInstances);
+		}
+	}
 
-			// Copy variables that require processing
-			for ( const FNiagaraSimCacheDataBuffersLayout::FVariableCopyInfo& VariableCopyInfo : CacheLayout.VariableMappingsToDataBuffer )
-			{
-				const uint32 SrcStride = uint32(NumInstances) * sizeof(float);
-				const uint8* Src = CacheBuffer.FloatData.GetData() + (VariableCopyInfo.ComponentFrom * SrcStride);
-				uint8* Dst = DataBuffer.GetComponentPtrFloat(VariableCopyInfo.ComponentTo);
-				VariableCopyInfo.CopyFunc(Dst, DataBuffer.GetFloatStride(), Src, SrcStride, uint32(NumInstances), RebaseTransform);
-			}
+	static void ReadCustomBuffers(float FrameFraction, float FrameDeltaSeconds, const FTransform& RebaseTransform, const FNiagaraSimCacheDataBuffersLayout& CacheLayout, const FNiagaraSimCacheDataBuffers& CacheBufferA, const FNiagaraSimCacheDataBuffers& CacheBufferB, uint8* DestBuffer, uint32 DestStride)
+	{
+		if (CacheLayout.VariableCopyMappingsToDataBuffer.Num() == 0)
+		{
+			return;
 		}
 
-		//-TODO:DestinationDataBuffer.SetIDTable(CacheBuffer.IDToIndexTable);
-		DataBuffer.SetIDAcquireTag(CacheBuffer.IDAcquireTag);
+		FNiagaraSimCacheDataBuffersLayout::FVariableCopyContext VariableCopyDataContext;
+		VariableCopyDataContext.FrameFraction	= FrameFraction;
+		VariableCopyDataContext.RecipDt			= FrameDeltaSeconds;
+		VariableCopyDataContext.NumInstances	= CacheBufferA.NumInstances;
+		VariableCopyDataContext.RebaseTransform	= RebaseTransform;
+		VariableCopyDataContext.InterpMappings	= CacheBufferA.InterpMapping;
+		VariableCopyDataContext.DestStride		= DestStride;
+		VariableCopyDataContext.SourceAStride	= CacheBufferA.NumInstances * sizeof(float);
+		VariableCopyDataContext.SourceBStride	= CacheBufferB.NumInstances * sizeof(float);
+		if (CacheLayout.ComponentVelocity != uint16(INDEX_NONE))
+		{
+			VariableCopyDataContext.VelocityComponent = CacheBufferA.FloatData.GetData() + (CacheLayout.ComponentVelocity * VariableCopyDataContext.SourceAStride);
+		}
+		else
+		{
+			VariableCopyDataContext.VelocityComponent = nullptr;
+		}
+
+		for (const FNiagaraSimCacheDataBuffersLayout::FVariableCopyMapping& VariableCopyMapping : CacheLayout.VariableCopyMappingsToDataBuffer)
+		{
+			VariableCopyDataContext.Dest				= DestBuffer + (uint32(VariableCopyMapping.ComponentTo) * DestStride);
+			VariableCopyDataContext.SourceAComponent	= CacheBufferA.FloatData.GetData() + (VariableCopyMapping.ComponentFrom * VariableCopyDataContext.SourceAStride);
+			VariableCopyDataContext.SourceBComponent	= CacheBufferB.FloatData.GetData() + (VariableCopyMapping.ComponentFrom * VariableCopyDataContext.SourceBStride);
+			VariableCopyMapping.CopyFunc(VariableCopyDataContext);
+		}
+	}
+
+	void ReadDataBuffer(float FrameFraction, float FrameDeltaSeconds, const FTransform& RebaseTransform, const FNiagaraSimCacheDataBuffersLayout& CacheLayout, const FNiagaraSimCacheDataBuffers& CacheBufferA, const FNiagaraSimCacheDataBuffers& CacheBufferB, FNiagaraDataSet& DataSet)
+	{
+		FNiagaraDataBuffer& DataBuffer = DataSet.BeginSimulate();
+		DataBuffer.Allocate(CacheBufferA.NumInstances);
+		DataBuffer.SetNumInstances(CacheBufferA.NumInstances);
+		if (CacheBufferA.NumInstances > 0 )
+		{
+			int32 iComponent = 0;
+			ReadFloatBuffers(iComponent, CacheLayout, CacheBufferA, MakeArrayView(DataBuffer.GetComponentPtrFloat(0), DataBuffer.GetFloatBuffer().Num()), DataBuffer.GetFloatStride());
+			ReadHalfBuffers(iComponent, CacheLayout, CacheBufferA, MakeArrayView(DataBuffer.GetComponentPtrHalf(0), DataBuffer.GetHalfBuffer().Num()), DataBuffer.GetHalfStride());
+			ReadInt32Buffers(iComponent, CacheLayout, CacheBufferA, MakeArrayView(DataBuffer.GetComponentPtrInt32(0), DataBuffer.GetInt32Buffer().Num()), DataBuffer.GetInt32Stride());
+			ReadCustomBuffers(FrameFraction, FrameDeltaSeconds, RebaseTransform, CacheLayout, CacheBufferA, CacheBufferB, DataBuffer.GetComponentPtrFloat(0), DataBuffer.GetFloatStride());
+		}
+
+		//-TODO:DestinationDataBuffer.SetIDTable(CacheBufferA.IDToIndexTable);
+		DataBuffer.SetIDAcquireTag(CacheBufferA.IDAcquireTag);
 
 		DataSet.EndSimulate();
 	}
 
-	void ReadDataBufferGPU(const FTransform& InRebaseTransform, FNiagaraEmitterInstance& EmitterInstance, const FNiagaraSimCacheDataBuffersLayout& InCacheLayout, const FNiagaraSimCacheDataBuffers& InCacheBuffer, FNiagaraDataSet& InDataSet, std::atomic<int32>& InPendingCommandsCounter)
+	void ReadDataBufferGPU(float InFrameFraction, float InFrameDeltaSeconds, const FTransform& InRebaseTransform, FNiagaraEmitterInstance& EmitterInstance, const FNiagaraSimCacheDataBuffersLayout& InCacheLayout, const FNiagaraSimCacheDataBuffers& InCacheBufferA, const FNiagaraSimCacheDataBuffers& InCacheBufferB, FNiagaraDataSet& InDataSet, std::atomic<int32>& InPendingCommandsCounter)
 	{
 		if (EmitterInstance.IsDisabled())
 		{
@@ -518,9 +662,9 @@ struct FNiagaraSimCacheHelper
 
 		FNiagaraGpuComputeDispatchInterface* DispathInterface = EmitterInstance.GetParentSystemInstance()->GetComputeDispatchInterface();
 		ENQUEUE_RENDER_COMMAND(NiagaraSimCacheGpuReadFrame)(
-			[DispathInterface, GPUExecContext=EmitterInstance.GetGPUContext(), RebaseTransform=InRebaseTransform, CacheLayout=&InCacheLayout, CacheBuffer=&InCacheBuffer, DataSet=&InDataSet, PendingCommandsCounter=&InPendingCommandsCounter](FRHICommandListImmediate& RHICmdList)
+			[DispathInterface, GPUExecContext=EmitterInstance.GetGPUContext(), FrameFraction=InFrameFraction, FrameDeltaSeconds=InFrameDeltaSeconds, RebaseTransform=InRebaseTransform, CacheLayout=&InCacheLayout, CacheBufferA=&InCacheBufferA, CacheBufferB=&InCacheBufferB, DataSet=&InDataSet, PendingCommandsCounter=&InPendingCommandsCounter](FRHICommandListImmediate& RHICmdList)
 			{
-				const int32 NumInstances = CacheBuffer->NumInstances;
+				const int32 NumInstances = CacheBufferA->NumInstances;
 
 				// Set Instance Count
 				{
@@ -544,7 +688,7 @@ struct FNiagaraSimCacheHelper
 				DataBuffer.SetGPUDataReadyStage(ENiagaraGpuComputeTickStage::PreInitViews);
 				GPUExecContext->SetDataToRender(&DataBuffer);
 
-				if (CacheBuffer->NumInstances > 0 )
+				if (NumInstances > 0)
 				{
 					int32 iComponent = 0;
 
@@ -552,31 +696,9 @@ struct FNiagaraSimCacheHelper
 					if ( CacheLayout->FloatCount > 0 )
 					{
 						FRWBuffer& RWBuffer = DataBuffer.GetGPUBufferFloat();
-						const int32 RWComponentStride = DataBuffer.GetFloatStride();
 						uint8* RWBufferMemory = reinterpret_cast<uint8*>(RHILockBuffer(RWBuffer.Buffer, 0, RWBuffer.NumBytes, RLM_WriteOnly));
-
-						for (uint32 i=0; i < CacheLayout->FloatCount; ++i)
-						{
-							const uint32 Component = CacheLayout->ComponentMappingsToDataBuffer[iComponent++];
-							if (Component == InvalidComponent)
-							{
-								continue;
-							}
-							const uint8* Source = CacheBuffer->FloatData.GetData() + (i * NumInstances * sizeof(float));
-							uint8* Dest = RWBufferMemory + (uint32(Component) * RWComponentStride);
-							CheckedMemcpy(MakeArrayView(RWBufferMemory, RWBuffer.NumBytes), Dest, CacheBuffer->FloatData, Source, sizeof(float) * NumInstances);
-							//FMemory::Memcpy(Dest, Source, sizeof(float) * NumInstances);
-						}
-
-						// Copy variables that require processing
-						for (const FNiagaraSimCacheDataBuffersLayout::FVariableCopyInfo& VariableCopyInfo : CacheLayout->VariableMappingsToDataBuffer)
-						{
-							const uint32 SrcStride = uint32(NumInstances) * sizeof(float);
-							const uint8* Src = CacheBuffer->FloatData.GetData() + (VariableCopyInfo.ComponentFrom * SrcStride);
-							uint8* Dst = RWBufferMemory + (uint32(VariableCopyInfo.ComponentTo) * RWComponentStride);
-							VariableCopyInfo.CopyFunc(Dst, DataBuffer.GetFloatStride(), Src, SrcStride, uint32(NumInstances), RebaseTransform);
-						}
-
+						ReadFloatBuffers(iComponent, *CacheLayout, *CacheBufferA, MakeArrayView(RWBufferMemory, RWBuffer.NumBytes), DataBuffer.GetFloatStride());
+						ReadCustomBuffers(FrameFraction, FrameDeltaSeconds, RebaseTransform, *CacheLayout, *CacheBufferA, *CacheBufferB, RWBufferMemory, DataBuffer.GetFloatStride());
 						RHIUnlockBuffer(RWBuffer.Buffer);
 					}
 
@@ -584,52 +706,23 @@ struct FNiagaraSimCacheHelper
 					if (CacheLayout->HalfCount > 0)
 					{
 						FRWBuffer& RWBuffer = DataBuffer.GetGPUBufferHalf();
-						const int32 RWComponentStride = DataBuffer.GetHalfStride();
 						uint8* RWBufferMemory = reinterpret_cast<uint8*>(RHILockBuffer(RWBuffer.Buffer, 0, RWBuffer.NumBytes, RLM_WriteOnly));
-
-						for (uint32 i = 0; i < CacheLayout->HalfCount; ++i)
-						{
-							const uint32 Component = CacheLayout->ComponentMappingsToDataBuffer[iComponent++];
-							if (Component == InvalidComponent)
-							{
-								continue;
-							}
-							const uint8* Source = CacheBuffer->HalfData.GetData() + (i * NumInstances * sizeof(FFloat16));
-							uint8* Dest = RWBufferMemory + (uint32(Component) * RWComponentStride);
-							CheckedMemcpy(MakeArrayView(RWBufferMemory, RWBuffer.NumBytes), Dest, CacheBuffer->HalfData, Source, sizeof(FFloat16) * NumInstances);
-							//FMemory::Memcpy(Dest, Source, sizeof(FFloat16) * NumInstances);
-						}
-
+						ReadHalfBuffers(iComponent, *CacheLayout, *CacheBufferA, MakeArrayView(RWBufferMemory, RWBuffer.NumBytes), DataBuffer.GetHalfStride());
 						RHIUnlockBuffer(RWBuffer.Buffer);
 					}
 
-
 					// Copy Int32
 					if (CacheLayout->Int32Count > 0)
 					{
 						FRWBuffer& RWBuffer = DataBuffer.GetGPUBufferInt();
-						const int32 RWComponentStride = DataBuffer.GetInt32Stride();
 						uint8* RWBufferMemory = reinterpret_cast<uint8*>(RHILockBuffer(RWBuffer.Buffer, 0, RWBuffer.NumBytes, RLM_WriteOnly));
-
-						for (uint32 i = 0; i < CacheLayout->Int32Count; ++i)
-						{
-							const uint32 Component = CacheLayout->ComponentMappingsToDataBuffer[iComponent++];
-							if (Component == InvalidComponent)
-							{
-								continue;
-							}
-							const uint8* Source = CacheBuffer->Int32Data.GetData() + (i * NumInstances * sizeof(int32));
-							uint8* Dest = RWBufferMemory + (uint32(Component) * RWComponentStride);
-							CheckedMemcpy(MakeArrayView(RWBufferMemory, RWBuffer.NumBytes), Dest, CacheBuffer->Int32Data, Source, sizeof(int32) * NumInstances);
-							//FMemory::Memcpy(Dest, Source, sizeof(int32) * NumInstances);
-						}
-
+						ReadInt32Buffers(iComponent, *CacheLayout, *CacheBufferA, MakeArrayView(RWBufferMemory, RWBuffer.NumBytes), DataBuffer.GetInt32Stride());
 						RHIUnlockBuffer(RWBuffer.Buffer);
 					}
 				}
 
 				//-TODO:DestinationDataBuffer.SetIDTable(CacheBuffer.IDToIndexTable);
-				DataBuffer.SetIDAcquireTag(CacheBuffer->IDAcquireTag);
+				DataBuffer.SetIDAcquireTag(CacheBufferA->IDAcquireTag);
 
 				// Ensure we decrement our counter so the GameThread knows the state of things
 				--(*PendingCommandsCounter);
@@ -637,6 +730,13 @@ struct FNiagaraSimCacheHelper
 		);
 	}
 
+	static bool CanInterpolateVariable(const FNiagaraVariableBase& Variable)
+	{
+		return
+			(Variable.GetType() == FNiagaraTypeDefinition::GetQuatDef()) ||
+			(Variable.GetType() == FNiagaraTypeDefinition::GetPositionDef());
+	}
+
 	static bool CanRebaseVariable(const FNiagaraVariableBase& Variable)
 	{
 		return	
@@ -645,62 +745,109 @@ struct FNiagaraSimCacheHelper
 			(Variable.GetType() == FNiagaraTypeDefinition::GetPositionDef());
 	}
 
-	static void CopyPositions(uint8* Dst, uint32 DstStride, const uint8* Src, uint32 SrcStride, uint32 NumInstances, const FTransform& RebaseTransform)
+	template<bool bWithRebase>
+	static void ExtrapolatePositions(const FNiagaraSimCacheDataBuffersLayout::FVariableCopyContext& CopyDataContext)
+	{
+		float* DstPositions[3] = { reinterpret_cast<float*>(CopyDataContext.Dest + (CopyDataContext.DestStride * 0)), reinterpret_cast<float*>(CopyDataContext.Dest + (CopyDataContext.DestStride * 1)), reinterpret_cast<float*>(CopyDataContext.Dest + (CopyDataContext.DestStride * 2))};
+		const float* SrcPositions[3] = { reinterpret_cast<const float*>(CopyDataContext.SourceAComponent + (CopyDataContext.SourceAStride * 0)), reinterpret_cast<const float*>(CopyDataContext.SourceAComponent + (CopyDataContext.SourceAStride * 1)), reinterpret_cast<const float*>(CopyDataContext.SourceAComponent + (CopyDataContext.SourceAStride * 2)) };
+		const float* SrcVelocities[3] = { reinterpret_cast<const float*>(CopyDataContext.VelocityComponent + (CopyDataContext.SourceAStride * 0)), reinterpret_cast<const float*>(CopyDataContext.VelocityComponent + (CopyDataContext.SourceAStride * 1)), reinterpret_cast<const float*>(CopyDataContext.VelocityComponent + (CopyDataContext.SourceAStride * 2)) };
+
+		for (uint32 i = 0; i < CopyDataContext.NumInstances; ++i)
+		{
+			const FVector3f CachePosition(SrcPositions[0][i], SrcPositions[1][i], SrcPositions[2][i]);
+			const FVector3f CacheVelocity(SrcVelocities[0][i], SrcVelocities[1][i], SrcVelocities[2][i]);
+			const FVector3f Position = CachePosition + (CacheVelocity * CopyDataContext.FrameFraction * CopyDataContext.RecipDt);
+			const FVector3f RebasedPosition = bWithRebase ? FVector3f(CopyDataContext.RebaseTransform.TransformPosition(FVector(Position))) : Position;
+			DstPositions[0][i] = RebasedPosition.X;
+			DstPositions[1][i] = RebasedPosition.Y;
+			DstPositions[2][i] = RebasedPosition.Z;
+		}
+	}
+
+	template<bool bWithRebase>
+	static void InterpPositions(const FNiagaraSimCacheDataBuffersLayout::FVariableCopyContext& CopyDataContext)
 	{
-		float* DstFloats = reinterpret_cast<float*>(Dst);
-		DstStride = DstStride >> 2;
+		float* DstPositions[3] = { reinterpret_cast<float*>(CopyDataContext.Dest + (CopyDataContext.DestStride * 0)), reinterpret_cast<float*>(CopyDataContext.Dest + (CopyDataContext.DestStride * 1)), reinterpret_cast<float*>(CopyDataContext.Dest + (CopyDataContext.DestStride * 2)) };
+		const float* SrcAPositions[3] = { reinterpret_cast<const float*>(CopyDataContext.SourceAComponent + (CopyDataContext.SourceAStride * 0)), reinterpret_cast<const float*>(CopyDataContext.SourceAComponent + (CopyDataContext.SourceAStride * 1)), reinterpret_cast<const float*>(CopyDataContext.SourceAComponent + (CopyDataContext.SourceAStride * 2)) };
+		const float* SrcBPositions[3] = { reinterpret_cast<const float*>(CopyDataContext.SourceBComponent + (CopyDataContext.SourceBStride * 0)), reinterpret_cast<const float*>(CopyDataContext.SourceBComponent + (CopyDataContext.SourceBStride * 1)), reinterpret_cast<const float*>(CopyDataContext.SourceBComponent + (CopyDataContext.SourceBStride * 2)) };
+
+		for (uint32 iInstanceA=0; iInstanceA < CopyDataContext.NumInstances; ++iInstanceA)
+		{
+			const uint32 iInstanceB = CopyDataContext.InterpMappings[iInstanceA];
+
+			const FVector3f CachePositionA(SrcAPositions[0][iInstanceA], SrcAPositions[1][iInstanceA], SrcAPositions[2][iInstanceA]);
+			const FVector3f CachePositionB = iInstanceB == INDEX_NONE ? CachePositionA : FVector3f(SrcBPositions[0][iInstanceB], SrcBPositions[1][iInstanceB], SrcBPositions[2][iInstanceB]);
+			const FVector3f Position(FMath::Lerp(CachePositionA, CachePositionB, CopyDataContext.FrameFraction));
+
+			const FVector3f RebasedPosition = bWithRebase ? FVector3f(CopyDataContext.RebaseTransform.TransformPosition(FVector(Position))) : Position;
+			DstPositions[0][iInstanceA] = float(RebasedPosition.X);
+			DstPositions[1][iInstanceA] = float(RebasedPosition.Y);
+			DstPositions[2][iInstanceA] = float(RebasedPosition.Z);
+		}
+	}
 
-		const float* SrcFloats = reinterpret_cast<const float*>(Src);
-		SrcStride = SrcStride >> 2;
+	static void CopyPositions(const FNiagaraSimCacheDataBuffersLayout::FVariableCopyContext& CopyDataContext)
+	{
+		float* DstPositions[3] = { reinterpret_cast<float*>(CopyDataContext.Dest + (CopyDataContext.DestStride * 0)), reinterpret_cast<float*>(CopyDataContext.Dest + (CopyDataContext.DestStride * 1)), reinterpret_cast<float*>(CopyDataContext.Dest + (CopyDataContext.DestStride * 2)) };
+		const float* SrcPositions[3] = { reinterpret_cast<const float*>(CopyDataContext.SourceAComponent + (CopyDataContext.SourceAStride * 0)), reinterpret_cast<const float*>(CopyDataContext.SourceAComponent + (CopyDataContext.SourceAStride * 1)), reinterpret_cast<const float*>(CopyDataContext.SourceAComponent + (CopyDataContext.SourceAStride * 2)) };
 
-		for (uint32 i = 0; i < NumInstances; ++i)
+		for (uint32 i = 0; i < CopyDataContext.NumInstances; ++i)
 		{
-			const FVector CachePosition(
-				SrcFloats[i + (SrcStride * 0)],
-				SrcFloats[i + (SrcStride * 1)],
-				SrcFloats[i + (SrcStride * 2)]
-			);
-			const FVector RebasedPosition = RebaseTransform.TransformPosition(CachePosition);
-			DstFloats[i + (DstStride * 0)] = float(RebasedPosition.X);
-			DstFloats[i + (DstStride * 1)] = float(RebasedPosition.Y);
-			DstFloats[i + (DstStride * 2)] = float(RebasedPosition.Z);
+			const FVector3f CachePosition(SrcPositions[0][i], SrcPositions[1][i], SrcPositions[2][i]);
+			const FVector3f RebasedPosition = FVector3f(CopyDataContext.RebaseTransform.TransformPosition(FVector(CachePosition)));
+			DstPositions[0][i] = RebasedPosition.X;
+			DstPositions[1][i] = RebasedPosition.Y;
+			DstPositions[2][i] = RebasedPosition.Z;
 		}
 	}
 
-	static void CopyQuaternions(uint8* Dst, uint32 DstStride, const uint8* Src, uint32 SrcStride, uint32 NumInstances, const FTransform& RebaseTransform)
+	template<bool bWithRebase>
+	static void InterpQuaternions(const FNiagaraSimCacheDataBuffersLayout::FVariableCopyContext& CopyDataContext)
 	{
-		float* DstFloats = reinterpret_cast<float*>(Dst);
-		DstStride = DstStride >> 2;
+		float* DstQuats[4] = { reinterpret_cast<float*>(CopyDataContext.Dest + (CopyDataContext.DestStride * 0)), reinterpret_cast<float*>(CopyDataContext.Dest + (CopyDataContext.DestStride * 1)), reinterpret_cast<float*>(CopyDataContext.Dest + (CopyDataContext.DestStride * 2)), reinterpret_cast<float*>(CopyDataContext.Dest + (CopyDataContext.DestStride * 3)) };
+		const float* SrcAQuats[4] = { reinterpret_cast<const float*>(CopyDataContext.SourceAComponent + (CopyDataContext.SourceAStride * 0)), reinterpret_cast<const float*>(CopyDataContext.SourceAComponent + (CopyDataContext.SourceAStride * 1)), reinterpret_cast<const float*>(CopyDataContext.SourceAComponent + (CopyDataContext.SourceAStride * 2)), reinterpret_cast<const float*>(CopyDataContext.SourceAComponent + (CopyDataContext.SourceAStride * 3)) };
+		const float* SrcBQuats[4] = { reinterpret_cast<const float*>(CopyDataContext.SourceBComponent + (CopyDataContext.SourceBStride * 0)), reinterpret_cast<const float*>(CopyDataContext.SourceBComponent + (CopyDataContext.SourceBStride * 1)), reinterpret_cast<const float*>(CopyDataContext.SourceBComponent + (CopyDataContext.SourceBStride * 2)), reinterpret_cast<const float*>(CopyDataContext.SourceBComponent + (CopyDataContext.SourceBStride * 3)) };
 
-		const float* SrcFloats = reinterpret_cast<const float*>(Src);
-		SrcStride = SrcStride >> 2;
+		for (uint32 iInstanceA=0; iInstanceA < CopyDataContext.NumInstances; ++iInstanceA)
+		{
+			const uint32 iInstanceB = CopyDataContext.InterpMappings[iInstanceA];
+			const FQuat4f CacheRotationA(SrcAQuats[0][iInstanceA], SrcAQuats[1][iInstanceA], SrcAQuats[2][iInstanceA], SrcAQuats[3][iInstanceA]);
+			const FQuat4f CacheRotationB = iInstanceB == INDEX_NONE ? CacheRotationA : FQuat4f(SrcBQuats[0][iInstanceB], SrcBQuats[1][iInstanceB], SrcBQuats[2][iInstanceB], SrcBQuats[3][iInstanceB]);
+			const FQuat4f CacheRotation(FQuat4f::Slerp(CacheRotationA, CacheRotationB, CopyDataContext.FrameFraction));
+			const FQuat4f RebasedQuat = bWithRebase ? CacheRotation * FQuat4f(CopyDataContext.RebaseTransform.GetRotation()) : CacheRotation;
+			DstQuats[0][iInstanceA] = RebasedQuat.X;
+			DstQuats[1][iInstanceA] = RebasedQuat.Y;
+			DstQuats[2][iInstanceA] = RebasedQuat.Z;
+			DstQuats[3][iInstanceA] = RebasedQuat.W;
+		}
+	}
+
+	static void CopyQuaternions(const FNiagaraSimCacheDataBuffersLayout::FVariableCopyContext& CopyDataContext)
+	{
+		float* DstQuats[4] = { reinterpret_cast<float*>(CopyDataContext.Dest + (CopyDataContext.DestStride * 0)), reinterpret_cast<float*>(CopyDataContext.Dest + (CopyDataContext.DestStride * 1)), reinterpret_cast<float*>(CopyDataContext.Dest + (CopyDataContext.DestStride * 2)), reinterpret_cast<float*>(CopyDataContext.Dest + (CopyDataContext.DestStride * 3)) };
+		const float* SrcQuats[4] = { reinterpret_cast<const float*>(CopyDataContext.SourceAComponent + (CopyDataContext.SourceAStride * 0)), reinterpret_cast<const float*>(CopyDataContext.SourceAComponent + (CopyDataContext.SourceAStride * 1)), reinterpret_cast<const float*>(CopyDataContext.SourceAComponent + (CopyDataContext.SourceAStride * 2)), reinterpret_cast<const float*>(CopyDataContext.SourceAComponent + (CopyDataContext.SourceAStride * 3)) };
 
-		for (uint32 i = 0; i < NumInstances; ++i)
+		for (uint32 i = 0; i < CopyDataContext.NumInstances; ++i)
 		{
-			const FQuat4f CacheRotation(
-				SrcFloats[i + (SrcStride * 0)],
-				SrcFloats[i + (SrcStride * 1)],
-				SrcFloats[i + (SrcStride * 2)],
-				SrcFloats[i + (SrcStride * 3)]
-			);
-			const FQuat4f RebasedQuat = CacheRotation * FQuat4f(RebaseTransform.GetRotation());
-			DstFloats[i + (DstStride * 0)] = RebasedQuat.X;
-			DstFloats[i + (DstStride * 1)] = RebasedQuat.Y;
-			DstFloats[i + (DstStride * 2)] = RebasedQuat.Z;
-			DstFloats[i + (DstStride * 3)] = RebasedQuat.W;
+			const FQuat4f CacheRotation(SrcQuats[0][i], SrcQuats[1][i], SrcQuats[2][i], SrcQuats[3][i]);
+			const FQuat4f RebasedQuat = CacheRotation * FQuat4f(CopyDataContext.RebaseTransform.GetRotation());
+			DstQuats[0][i] = RebasedQuat.X;
+			DstQuats[1][i] = RebasedQuat.Y;
+			DstQuats[2][i] = RebasedQuat.Z;
+			DstQuats[3][i] = RebasedQuat.W;
 		}
 	}
 
-	static void CopyMatrices(uint8* Dst, uint32 DstStride, const uint8* Src, uint32 SrcStride, uint32 NumInstances, const FTransform& RebaseTransform)
+	static void CopyMatrices(const FNiagaraSimCacheDataBuffersLayout::FVariableCopyContext& CopyDataContext)
 	{
-		float* DstFloats = reinterpret_cast<float*>(Dst);
-		DstStride = DstStride >> 2;
+		float* DstFloats = reinterpret_cast<float*>(CopyDataContext.Dest);
+		const uint32 DstStride = CopyDataContext.DestStride >> 2;
 
-		const float* SrcFloats = reinterpret_cast<const float*>(Src);
-		SrcStride = SrcStride >> 2;
+		const float* SrcFloats = reinterpret_cast<const float*>(CopyDataContext.SourceAComponent);
+		const uint32 SrcStride = CopyDataContext.SourceAStride >> 2;
 
-		const FMatrix44d RebaseMatrix = RebaseTransform.ToMatrixWithScale();
-		for (uint32 i = 0; i < NumInstances; ++i)
+		const FMatrix44d RebaseMatrix = CopyDataContext.RebaseTransform.ToMatrixWithScale();
+		for (uint32 i = 0; i < CopyDataContext.NumInstances; ++i)
 		{
 			FMatrix44d CacheMatrix;
 			for (int32 j = 0; j < 16; ++j)
