- Fix for material attribute bindings not always working
- Don't use RenderProps->IsSupportedVariableForBinding as the value is always support

diff --git a/Engine/Plugins/FX/Niagara/Source/NiagaraEditor/Private/Customizations/NiagaraTypeCustomizations.cpp b/Engine/Plugins/FX/Niagara/Source/NiagaraEditor/Private/Customizations/NiagaraTypeCustomizations.cpp
--- a/Engine/Plugins/FX/Niagara/Source/NiagaraEditor/Private/Customizations/NiagaraTypeCustomizations.cpp
+++ b/Engine/Plugins/FX/Niagara/Source/NiagaraEditor/Private/Customizations/NiagaraTypeCustomizations.cpp
@@ -967,7 +967,7 @@ TArray<TPair<FNiagaraVariableBase, FNiagaraVariableBase> > FNiagaraMaterialAttri
 				}
 				else if (!AllowedVariableTypes.Num() || AllowedVariableTypes.Contains(BaseVar.GetType()))
 				{
-					if (RenderProps && TargetParameterBinding && RenderProps->IsSupportedVariableForBinding(BaseVar, TargetParameterBinding->MaterialParameterName))
+					if (RenderProps && TargetParameterBinding)
 					{
 						Names.AddUnique(TPair<FNiagaraVariableBase, FNiagaraVariableBase>(BaseVar, FNiagaraVariableBase()));
 					}
