Niagara - Prevent a crash on load when generating hashes and a niagara variable is missing its type.  Also fix the niagara examples plugin to specify dependencies which is what triggered this crash to begin with.

diff --git a/Engine/Plugins/FX/Niagara/Source/Niagara/Private/NiagaraScript.cpp b/Engine/Plugins/FX/Niagara/Source/Niagara/Private/NiagaraScript.cpp
--- a/Engine/Plugins/FX/Niagara/Source/Niagara/Private/NiagaraScript.cpp
+++ b/Engine/Plugins/FX/Niagara/Source/Niagara/Private/NiagaraScript.cpp
@@ -1354,9 +1354,9 @@ void UNiagaraScript::ComputeVMCompilationId(FNiagaraVMExecutableDataId& Id, FGui
 			//UE_LOG(LogNiagara, Display, TEXT("AreScriptAndSourceSynchronized %s ======================== "), *GetFullName());
 			for (int32 i = 0; i < Vars.Num(); i++)
 			{
-				if (Vars[i].IsDataInterface() || Vars[i].IsUObject())
+				if (Vars[i].IsValid() == false || Vars[i].IsDataInterface() || Vars[i].IsUObject())
 				{
-					// Skip these types as they don't bake out, just normal parameters get baked.
+					// Skip these types as they're invalid or don't bake out, just normal parameters get baked.
 				}
 				else
 				{
