Changing default camera view and directional light angle in Niagara Editor's preview window.

- Assigning the preview component flags the camera view for refocusing
- Camera refocus happens once the age is non-zero to capture programmable bounds modifications
- Directional light angle is flipped 180 degrees

diff --git a/Engine/Plugins/FX/Niagara/Source/NiagaraEditor/Private/Widgets/SNiagaraSystemViewport.cpp b/Engine/Plugins/FX/Niagara/Source/NiagaraEditor/Private/Widgets/SNiagaraSystemViewport.cpp
--- a/Engine/Plugins/FX/Niagara/Source/NiagaraEditor/Private/Widgets/SNiagaraSystemViewport.cpp
+++ b/Engine/Plugins/FX/Niagara/Source/NiagaraEditor/Private/Widgets/SNiagaraSystemViewport.cpp
@@ -61,6 +61,7 @@ public:
 	void DrawParticleCounts(UNiagaraComponent* Component, FCanvas* Canvas, float& CurrentX, float& CurrentY, UFont* Font, const float FontHeight);
 	void DrawEmitterExecutionOrder(UNiagaraComponent* Component, FCanvas* Canvas, float& CurrentX, float& CurrentY, UFont* Font, const float FontHeight);
 	void DrawGpuTickInformation(UNiagaraComponent* Component, FCanvas* Canvas, float& CurrentX, float& CurrentY, UFont* Font, const float FontHeight);
+	void SetUpdateViewportFocus(bool bUpdate) { bUpdateViewportFocus = bUpdate; }
 
 	TWeakPtr<SNiagaraSystemViewport> NiagaraViewportPtr;
 	bool bCaptureScreenShot;
@@ -69,6 +70,9 @@ public:
 	FAdvancedPreviewScene* AdvancedPreviewScene = nullptr;
 
 	FOnScreenShotCaptured OnScreenShotCaptured;
+
+private:
+	bool bUpdateViewportFocus = false;
 };
 
 FNiagaraSystemViewportClient::FNiagaraSystemViewportClient(FAdvancedPreviewScene& InPreviewScene, const TSharedRef<SNiagaraSystemViewport>& InNiagaraEditorViewport, FOnScreenShotCaptured InOnScreenShotCaptured)
@@ -118,6 +122,17 @@ void FNiagaraSystemViewportClient::Draw(FViewport* InViewport, FCanvas* Canvas)
 	UNiagaraSystem* ParticleSystem = NiagaraViewport.IsValid() ? NiagaraViewport->GetPreviewComponent()->GetAsset() : nullptr;
 	UNiagaraComponent* Component = NiagaraViewport.IsValid() ? NiagaraViewport->GetPreviewComponent() : nullptr;
 
+	FNiagaraSystemInstanceControllerConstPtr SystemInstanceController = NiagaraViewport.IsValid() ? NiagaraViewport->GetPreviewComponent()->GetSystemInstanceController() : nullptr;
+	if (Component && SystemInstanceController.IsValid() && SystemInstanceController->GetAge() > 0.0)
+	{
+		FBox Bounds = Component->GetLocalBounds().GetBox();
+		if (bUpdateViewportFocus && Bounds.IsValid)
+		{
+			FocusViewportOnBox(Bounds);
+			SetUpdateViewportFocus(false);
+		}
+	}
+
 	if (NiagaraViewport.IsValid() && NiagaraViewport->GetDrawElement(SNiagaraSystemViewport::EDrawElements::Bounds))
 	{
 		EngineShowFlags.SetBounds(true);
@@ -470,6 +485,12 @@ void SNiagaraSystemViewport::Construct(const FArguments& InArgs)
 	PreviewComponent = nullptr;
 	AdvancedPreviewScene = MakeShareable(new FAdvancedPreviewScene(FPreviewScene::ConstructionValues()));
 	AdvancedPreviewScene->SetFloorVisibility(false);
+
+	float Pitch = -40.0;
+	float Yaw = 128.0;
+	float Roll = 0.0;
+	AdvancedPreviewScene->SetLightDirection(FRotator(Pitch, Yaw, Roll));
+
 	OnThumbnailCaptured = InArgs._OnThumbnailCaptured;
 	Sequencer = InArgs._Sequencer;
 	
@@ -573,6 +594,8 @@ void SNiagaraSystemViewport::SetPreviewComponent(UNiagaraComponent* NiagaraCompo
 	{
 		PreviewComponent->SetGpuComputeDebug(true);
 		AdvancedPreviewScene->AddComponent(PreviewComponent, PreviewComponent->GetRelativeTransform());
+
+		SystemViewportClient->SetUpdateViewportFocus(true);
 	}
 
 	SystemViewportClient->SetOrbitModeFromSettings();
