Fixed copy paste of modules not properly pasting values for static switches in static switch chains.
i.e. static switch 1 would be set successfully, but static switch 2 that appears due to static switch 1 being set would not be set.

diff --git a/Engine/Plugins/FX/Niagara/Source/NiagaraEditor/Private/ViewModels/Stack/NiagaraStackFunctionInputCollection.cpp b/Engine/Plugins/FX/Niagara/Source/NiagaraEditor/Private/ViewModels/Stack/NiagaraStackFunctionInputCollection.cpp
--- a/Engine/Plugins/FX/Niagara/Source/NiagaraEditor/Private/ViewModels/Stack/NiagaraStackFunctionInputCollection.cpp
+++ b/Engine/Plugins/FX/Niagara/Source/NiagaraEditor/Private/ViewModels/Stack/NiagaraStackFunctionInputCollection.cpp
@@ -137,7 +137,7 @@ void UNiagaraStackFunctionInputCollection::SetValuesFromClipboardFunctionInputs(
 	// Set static switches first so that other inputs will be available to set.
 	for (UNiagaraStackInputCategory* ChildCategory : ChildCategories)
 	{
-		ChildCategory->SetStaticSwitchValuesFromClipboardFunctionInputs(ClipboardFunctionInputs);
+		ChildCategory->SetStaticSwitchValuesFromClipboardFunctionInputs(ClipboardFunctionInputs, *this);
 	}
 
 	RefreshChildren();
diff --git a/Engine/Plugins/FX/Niagara/Source/NiagaraEditor/Private/ViewModels/Stack/NiagaraStackInputCategory.cpp b/Engine/Plugins/FX/Niagara/Source/NiagaraEditor/Private/ViewModels/Stack/NiagaraStackInputCategory.cpp
--- a/Engine/Plugins/FX/Niagara/Source/NiagaraEditor/Private/ViewModels/Stack/NiagaraStackInputCategory.cpp
+++ b/Engine/Plugins/FX/Niagara/Source/NiagaraEditor/Private/ViewModels/Stack/NiagaraStackInputCategory.cpp
@@ -4,6 +4,7 @@
 #include "ViewModels/Stack/NiagaraStackFunctionInput.h"
 #include "NiagaraNodeFunctionCall.h"
 #include "NiagaraClipboard.h"
+#include "ViewModels/Stack/NiagaraStackFunctionInputCollection.h"
 
 #include UE_INLINE_GENERATED_CPP_BY_NAME(NiagaraStackInputCategory)
 
@@ -142,30 +143,36 @@ void UNiagaraStackInputCategory::ToClipboardFunctionInputs(UObject* InOuter, TAr
 }
 
 template<typename Predicate>
-void SetValuesFromFunctionInputsInternal(UNiagaraStackInputCategory* Category, const TArray<const UNiagaraClipboardFunctionInput*>& ClipboardFunctionInputs, Predicate InputMatchesFilter)
+void SetValuesFromFunctionInputsInternal(UNiagaraStackInputCategory* Category, const TArray<const UNiagaraClipboardFunctionInput*>& ClipboardFunctionInputs, Predicate InputMatchesFilter, UNiagaraStackFunctionInputCollection* OwningFunctionCollection)
 {
-	TArray<UNiagaraStackFunctionInput*> ChildInputs;
-	Category->GetUnfilteredChildrenOfType(ChildInputs);
-	for (UNiagaraStackFunctionInput* ChildInput : ChildInputs)
+	for (const UNiagaraClipboardFunctionInput* ClipboardFunctionInput : ClipboardFunctionInputs)
 	{
-		for (const UNiagaraClipboardFunctionInput* ClipboardFunctionInput : ClipboardFunctionInputs)
+		TArray<UNiagaraStackFunctionInput*> ChildInputs;
+		Category->GetUnfilteredChildrenOfType(ChildInputs);
+		
+		for (UNiagaraStackFunctionInput* ChildInput : ChildInputs)
 		{
 			if (InputMatchesFilter(ChildInput) && ChildInput->GetInputParameterHandle().GetName() == ClipboardFunctionInput->InputName && ChildInput->GetInputType() == ClipboardFunctionInput->InputType)
 			{
 				ChildInput->SetValueFromClipboardFunctionInput(*ClipboardFunctionInput);
+				if(OwningFunctionCollection)
+				{
+					OwningFunctionCollection->RefreshChildren();
+				}
+				break;
 			}
 		}
 	}
 }
 
-void  UNiagaraStackInputCategory::SetStaticSwitchValuesFromClipboardFunctionInputs(const TArray<const UNiagaraClipboardFunctionInput*>& ClipboardFunctionInputs)
+void  UNiagaraStackInputCategory::SetStaticSwitchValuesFromClipboardFunctionInputs(const TArray<const UNiagaraClipboardFunctionInput*>& ClipboardFunctionInputs, UNiagaraStackFunctionInputCollection& OwningFunctionCollection)
 {
-	SetValuesFromFunctionInputsInternal(this, ClipboardFunctionInputs, [](UNiagaraStackFunctionInput* ChildInput) { return ChildInput->IsStaticParameter(); });
+	SetValuesFromFunctionInputsInternal(this, ClipboardFunctionInputs, [](UNiagaraStackFunctionInput* ChildInput) { return ChildInput->IsStaticParameter(); }, &OwningFunctionCollection);
 }
 
 void  UNiagaraStackInputCategory::SetStandardValuesFromClipboardFunctionInputs(const TArray<const UNiagaraClipboardFunctionInput*>& ClipboardFunctionInputs)
 {
-	SetValuesFromFunctionInputsInternal(this, ClipboardFunctionInputs, [](UNiagaraStackFunctionInput* ChildInput) { return ChildInput->IsStaticParameter() == false; });
+	SetValuesFromFunctionInputsInternal(this, ClipboardFunctionInputs, [](UNiagaraStackFunctionInput* ChildInput) { return ChildInput->IsStaticParameter() == false; }, nullptr);
 }
 
 bool UNiagaraStackInputCategory::FilterForVisibleCondition(const UNiagaraStackEntry& Child) const
diff --git a/Engine/Plugins/FX/Niagara/Source/NiagaraEditor/Public/ViewModels/Stack/NiagaraStackInputCategory.h b/Engine/Plugins/FX/Niagara/Source/NiagaraEditor/Public/ViewModels/Stack/NiagaraStackInputCategory.h
--- a/Engine/Plugins/FX/Niagara/Source/NiagaraEditor/Public/ViewModels/Stack/NiagaraStackInputCategory.h
+++ b/Engine/Plugins/FX/Niagara/Source/NiagaraEditor/Public/ViewModels/Stack/NiagaraStackInputCategory.h
@@ -46,7 +46,9 @@ public:
 
 	void ToClipboardFunctionInputs(UObject* InOuter, TArray<const UNiagaraClipboardFunctionInput*>& OutClipboardFunctionInputs) const;
 
-	void SetStaticSwitchValuesFromClipboardFunctionInputs(const TArray<const UNiagaraClipboardFunctionInput*>& ClipboardFunctionInputs);
+	// We pass in the owning input collection to refresh the inputs of a category after setting a static switch value.
+	// This is done to support copy paste for chained static switches, as otherwise the stack entry for the 2nd static switch won't exist for pasting data
+	void SetStaticSwitchValuesFromClipboardFunctionInputs(const TArray<const UNiagaraClipboardFunctionInput*>& ClipboardFunctionInputs, class UNiagaraStackFunctionInputCollection& OwningFunctionCollection);
 
 	void SetStandardValuesFromClipboardFunctionInputs(const TArray<const UNiagaraClipboardFunctionInput*>& ClipboardFunctionInputs);
 
