crash when RT res doesn't match vdb res

diff --git a/Engine/Plugins/FX/Niagara/Source/Niagara/Private/VolumeCache.cpp b/Engine/Plugins/FX/Niagara/Source/Niagara/Private/VolumeCache.cpp
--- a/Engine/Plugins/FX/Niagara/Source/Niagara/Private/VolumeCache.cpp
+++ b/Engine/Plugins/FX/Niagara/Source/Niagara/Private/VolumeCache.cpp
@@ -335,6 +335,7 @@ bool FOpenVDBCacheData::Fill3DTexture(int frame, FTextureRHIRef TextureToFill)
 		if (TextureToFill->GetSizeXYZ() != DenseResolution)
 		{
 			UE_LOG(LogVolumeCache, Warning, TEXT("Target texture resolution %s doesn't match OpenVDB grid resolution %s."), *TextureToFill->GetSizeXYZ().ToString(), *DenseResolution.ToString());
+			return false;
 		}
 
 		const FUpdateTextureRegion3D UpdateRegion(0, 0, 0, 0, 0, 0, DenseResolution.X, DenseResolution.Y, DenseResolution.Z);
