- Cache more information for when we have to initialize a dummy view
- This ensure the camera position / orientation / FOV / etc will all match on flush frames

diff --git a/Engine/Plugins/FX/Niagara/Source/Niagara/Private/NiagaraGpuComputeDispatch.cpp b/Engine/Plugins/FX/Niagara/Source/Niagara/Private/NiagaraGpuComputeDispatch.cpp
--- a/Engine/Plugins/FX/Niagara/Source/Niagara/Private/NiagaraGpuComputeDispatch.cpp
+++ b/Engine/Plugins/FX/Niagara/Source/Niagara/Private/NiagaraGpuComputeDispatch.cpp
@@ -168,7 +168,6 @@ FFXSystemInterface* FNiagaraGpuComputeDispatch::GetInterface(const FName& InName
 FNiagaraGpuComputeDispatch::FNiagaraGpuComputeDispatch(ERHIFeatureLevel::Type InFeatureLevel, EShaderPlatform InShaderPlatform, FGPUSortManager* InGPUSortManager)
 	: FNiagaraGpuComputeDispatchInterface(InShaderPlatform, InFeatureLevel)
 	, GPUSortManager(InGPUSortManager)
-	, CachedViewRect(0, 0, 64, 64)
 {
 	// Register the batcher callback in the GPUSortManager.
 	// The callback is used to generate the initial keys and values for the GPU sort tasks,
@@ -397,16 +396,18 @@ void FNiagaraGpuComputeDispatch::ProcessPendingTicksFlush(FRHICommandListImmedia
 
 			// Make a temporary ViewInfo
 			//-TODO: We could gather some more information here perhaps?
-			FSceneViewFamily ViewFamily(FSceneViewFamily::ConstructionValues(nullptr, nullptr, FEngineShowFlags(ESFIM_Game))
-				.SetTime(FGameTime())
-				.SetGammaCorrection(1.0f));
+			FSceneViewFamily ViewFamily(
+				FSceneViewFamily::ConstructionValues(nullptr, nullptr, FEngineShowFlags(ESFIM_Game))
+				.SetTime(CachedViewInitOptions.GameTime)
+				.SetGammaCorrection(CachedViewInitOptions.GammaCorrection)
+			);
 
 			FSceneViewInitOptions ViewInitOptions;
 			ViewInitOptions.ViewFamily = &ViewFamily;
-			ViewInitOptions.SetViewRectangle(CachedViewRect);
-			ViewInitOptions.ViewOrigin = FVector::ZeroVector;
-			ViewInitOptions.ViewRotationMatrix = FMatrix::Identity;
-			ViewInitOptions.ProjectionMatrix = FMatrix::Identity;
+			ViewInitOptions.SetViewRectangle(CachedViewInitOptions.ViewRect);
+			ViewInitOptions.ViewOrigin = CachedViewInitOptions.ViewOrigin;
+			ViewInitOptions.ViewRotationMatrix = CachedViewInitOptions.ViewRotationMatrix;
+			ViewInitOptions.ProjectionMatrix = CachedViewInitOptions.ProjectionMatrix;
 
 			FViewInfo DummyView(ViewInitOptions);
 
@@ -1887,9 +1888,19 @@ void FNiagaraGpuComputeDispatch::PostRenderOpaque(FRDGBuilder& GraphBuilder, TCo
 
 	bAllowGPUParticleUpdate = bAllowGPUParticleUpdate && GetReferenceAllowGPUUpdate(Views);
 
+	// Cache view information which will be used if we have to flush simulation commands in the future
 	if ( bAllowGPUParticleUpdate && Views.IsValidIndex(0) )
 	{
-		CachedViewRect = Views[0].ViewRect;
+		if (const FSceneViewFamily* ViewFamily = Views[0].Family)
+		{
+			CachedViewInitOptions.GameTime			= ViewFamily->Time;
+			CachedViewInitOptions.GammaCorrection	= ViewFamily->GammaCorrection;
+		}
+
+		CachedViewInitOptions.ViewRect				= Views[0].ViewRect;
+		CachedViewInitOptions.ViewOrigin			= Views[0].SceneViewInitOptions.ViewOrigin;
+		CachedViewInitOptions.ViewRotationMatrix	= Views[0].SceneViewInitOptions.ViewRotationMatrix;
+		CachedViewInitOptions.ProjectionMatrix		= Views[0].SceneViewInitOptions.ProjectionMatrix;
 	}
 
 	if (bAllowGPUParticleUpdate && FNiagaraUtilities::AllowGPUParticles(GetShaderPlatform()))
diff --git a/Engine/Plugins/FX/Niagara/Source/Niagara/Private/NiagaraGpuComputeDispatch.h b/Engine/Plugins/FX/Niagara/Source/Niagara/Private/NiagaraGpuComputeDispatch.h
--- a/Engine/Plugins/FX/Niagara/Source/Niagara/Private/NiagaraGpuComputeDispatch.h
+++ b/Engine/Plugins/FX/Niagara/Source/Niagara/Private/NiagaraGpuComputeDispatch.h
@@ -267,7 +267,17 @@ private:
 #endif // WITH_MGPU
 
 	// Cached information to build a dummy view info if necessary
-	FIntRect CachedViewRect;
+	struct FCachedViewInitOptions
+	{
+		FGameTime	GameTime;
+		float		GammaCorrection		= 1.0f;
+		FIntRect	ViewRect			= FIntRect(0, 0, 64, 64);
+		FVector		ViewOrigin			= FVector::ZeroVector;
+		FMatrix		ViewRotationMatrix	= FMatrix::Identity;
+		FMatrix		ProjectionMatrix	= FMatrix::Identity;
+	};
+
+	FCachedViewInitOptions CachedViewInitOptions;
 
 #if WITH_EDITOR
 	bool bRaisedWarningThisFrame = false;
