Fixed emitters opening up the parameter panel to "Active Module" instead of "Active Overview"

diff --git a/Engine/Plugins/FX/Niagara/Source/NiagaraEditor/Private/ViewModels/NiagaraParameterPanelViewModel.cpp b/Engine/Plugins/FX/Niagara/Source/NiagaraEditor/Private/ViewModels/NiagaraParameterPanelViewModel.cpp
--- a/Engine/Plugins/FX/Niagara/Source/NiagaraEditor/Private/ViewModels/NiagaraParameterPanelViewModel.cpp
+++ b/Engine/Plugins/FX/Niagara/Source/NiagaraEditor/Private/ViewModels/NiagaraParameterPanelViewModel.cpp
@@ -867,19 +867,10 @@ void FNiagaraSystemToolkitParameterPanelViewModel::Init(const FSystemToolkitUICo
 	UNiagaraSystem& System = SystemViewModel->GetSystem();
 
 	VariableObjectSelection = MakeShared< FNiagaraObjectSelection>();
-
-	// Standalone emitter assets have different default section setups than systems
-	if (SystemViewModel->GetEditMode() == ENiagaraSystemViewModelEditMode::EmitterAsset)
-	{
-		ActiveSectionIndex = 1;
-		LastActiveSystemSectionIdx = ActiveSectionIndex;
-	}
-	else
-	{
-		ActiveSectionIndex = 0;
-		LastActiveSystemSectionIdx = ActiveSectionIndex;
-	}
-
+	
+	ActiveSectionIndex = 0;
+	LastActiveSystemSectionIdx = ActiveSectionIndex;
+	
 	// Define the sections
 	ActiveSystemIdx = Sections.Emplace(FGuid(0x7B4AFB34, 0xD0DF4618, 0x9A05349E, 0x361CE735), LOCTEXT("SectionSystemActive", "Active Overview"), LOCTEXT("SectionActiveSystemDesc", "Displays parameters that are context-sensitive to the selected items in the System Overview."));
 	ActiveScriptIdx = Sections.Emplace(FGuid(0x6E11F3EA, 0x5E314E82, 0xA64FEA41, 0x4F1DB891), LOCTEXT("SectionScriptActive", "Active Module"), LOCTEXT("SectionActiveScriptDesc", "Displays parameters that are context-sensitive to the active Scratch module document. Use this for renaming/editing local parameters to the module."));
