Allowed OpenVDB to be enabled via a macro in the .Build.cs rather than hardcoded at a platform level.

diff --git a/Engine/Plugins/FX/Niagara/Source/Niagara/Classes/NiagaraOpenVDB.h b/Engine/Plugins/FX/Niagara/Source/Niagara/Classes/NiagaraOpenVDB.h
--- a/Engine/Plugins/FX/Niagara/Source/Niagara/Classes/NiagaraOpenVDB.h
+++ b/Engine/Plugins/FX/Niagara/Source/Niagara/Classes/NiagaraOpenVDB.h
@@ -3,8 +3,9 @@
 #pragma once
 
 #include "HAL/Platform.h"
+#include "NiagaraUseOpenVDB.h"
 
-#if NIAGARA_USE_OPENVDB
+#if UE_USE_OPENVDB
 
 THIRD_PARTY_INCLUDES_START
 
diff --git a/Engine/Plugins/FX/Niagara/Source/Niagara/Classes/VolumeCache.h b/Engine/Plugins/FX/Niagara/Source/Niagara/Classes/VolumeCache.h
--- a/Engine/Plugins/FX/Niagara/Source/Niagara/Classes/VolumeCache.h
+++ b/Engine/Plugins/FX/Niagara/Source/Niagara/Classes/VolumeCache.h
@@ -6,6 +6,7 @@
 #include "CoreMinimal.h"
 #include "UObject/ObjectMacros.h"
 #include "RHI.h"
+#include "NiagaraUseOpenVDB.h"
 
 #include "VolumeCache.generated.h"
 
@@ -78,7 +79,9 @@ protected:
 	FIntVector DenseResolution;
 };
 
+#if UE_USE_OPENVDB
 namespace OpenVDBTools
 {
 	NIAGARA_API bool WriteImageDataToOpenVDBFile(FStringView FilePath, FIntVector ImageSize, TArrayView<FFloat16Color> ImageData, bool UseFloatGrids = false);
-}
\ No newline at end of file
+}
+#endif
diff --git a/Engine/Plugins/FX/Niagara/Source/Niagara/Niagara.Build.cs b/Engine/Plugins/FX/Niagara/Source/Niagara/Niagara.Build.cs
--- a/Engine/Plugins/FX/Niagara/Source/Niagara/Niagara.Build.cs
+++ b/Engine/Plugins/FX/Niagara/Source/Niagara/Niagara.Build.cs
@@ -11,11 +11,16 @@ public class Niagara : ModuleRules
 
 		PrivateIncludePaths.Add("../../../../Shaders/Shared");
 
-		// Specific to OpenVDB support
+		// Enable OpenVDB under Windows
+		// @todo: we need builds for OpenVDB for platforms other than windows
 		if (Target.IsInPlatformGroup(UnrealPlatformGroup.Windows))
 		{
+			// Specific to OpenVDB support
 			bUseRTTI = true;
 			bEnableExceptions = true;
+
+			AddEngineThirdPartyPrivateStaticDependencies(Target, "OpenVDB");
+			PublicDefinitions.Add("UE_USE_OPENVDB=1");
 		}
 		
 		PrivateDependencyModuleNames.AddRange(
@@ -84,8 +89,7 @@ public class Niagara : ModuleRules
 				"IntelTBB",
 				"Blosc",
 				"zlib",
-				"Boost",				
-				"OpenVDB"
+				"Boost"
 			);
 
 			PublicDefinitions.AddRange(
diff --git a/Engine/Plugins/FX/Niagara/Source/Niagara/Private/NiagaraModule.cpp b/Engine/Plugins/FX/Niagara/Source/Niagara/Private/NiagaraModule.cpp
--- a/Engine/Plugins/FX/Niagara/Source/Niagara/Private/NiagaraModule.cpp
+++ b/Engine/Plugins/FX/Niagara/Source/Niagara/Private/NiagaraModule.cpp
@@ -11,6 +11,7 @@
 #include "NiagaraSettings.h"
 #include "NiagaraDataChannelManager.h"
 #include "NiagaraDataInterfaceCurlNoise.h"
+#include "NiagaraUseOpenVDB.h"
 #include "UObject/Class.h"
 #include "UObject/ObjectRedirector.h"
 #include "UObject/Package.h"
@@ -39,7 +40,7 @@
 #include "DataInterface/NiagaraDataInterfaceDataChannelRead.h"
 #include "DataInterface/NiagaraDataInterfaceDataChannelSpawn.h"
 
-#if NIAGARA_USE_OPENVDB
+#if UE_USE_OPENVDB
 #include "NiagaraOpenVDB.h"
 #endif
 
@@ -268,7 +269,7 @@ void INiagaraModule::StartupModule()
 	LLM_SCOPE(ELLMTag::Niagara);
 	FNiagaraTypeDefinition::Init();
 
-#if NIAGARA_USE_OPENVDB
+#if UE_USE_OPENVDB
 	// Global registration of  the vdb types.
 	openvdb::initialize();
 	if (!Vec4SGrid::isRegistered())
diff --git a/Engine/Plugins/FX/Niagara/Source/Niagara/Private/VolumeCache.cpp b/Engine/Plugins/FX/Niagara/Source/Niagara/Private/VolumeCache.cpp
--- a/Engine/Plugins/FX/Niagara/Source/Niagara/Private/VolumeCache.cpp
+++ b/Engine/Plugins/FX/Niagara/Source/Niagara/Private/VolumeCache.cpp
@@ -7,19 +7,18 @@
 #include "VolumeCache.h"
 #include "NiagaraSettings.h"
 #include "RHICommandList.h"
+#include "NiagaraUseOpenVDB.h"
 
 #include UE_INLINE_GENERATED_CPP_BY_NAME(VolumeCache)
 
-// @todo: we need builds for OpenVDB for platforms other than windows
-#if NIAGARA_USE_OPENVDB
+#if UE_USE_OPENVDB
 #include "NiagaraOpenVDB.h"
 #endif
 
 DEFINE_LOG_CATEGORY_STATIC(LogVolumeCache, Log, All);
 
 
-// @todo: we need builds for OpenVDB for platforms other than windows
-#if NIAGARA_USE_OPENVDB
+#if UE_USE_OPENVDB
 class NIAGARA_API FOpenVDBCacheData : public FVolumeCacheData
 {
 public:
@@ -71,7 +70,7 @@ void UVolumeCache::InitData()
 	{
 		case EVolumeCacheType::OpenVDB :
 		{			
-#if NIAGARA_USE_OPENVDB
+#if UE_USE_OPENVDB
 			CachedVolumeFiles = TSharedPtr<FVolumeCacheData>(new FOpenVDBCacheData());
 #else
 			UE_LOG(LogVolumeCache, Warning, TEXT("OpenVDB not supported on the platform"));
@@ -142,8 +141,7 @@ FString FVolumeCacheData::GetAssetPath(FString PathFormat, int32 FrameIndex) con
 	return AssetPath;
 }
 
-// @todo: we need builds for OpenVDB for platforms other than windows
-#if NIAGARA_USE_OPENVDB
+#if UE_USE_OPENVDB
 void FOpenVDBCacheData::Init(FIntVector Resolution)
 {
 	FScopeLock ScopeLock(&DenseGridGuard);
@@ -153,7 +151,7 @@ void FOpenVDBCacheData::Init(FIntVector Resolution)
 	if (OpenVDBGrids.Num() == 0)
 	{
 		TotalMemoryUsage = 0;
-	}
+}
 }
 
 bool FOpenVDBCacheData::LoadFile(FString Path, int frame)
diff --git a/Engine/Plugins/FX/Niagara/Source/NiagaraEditor/Private/NiagaraBakerRenderer.cpp b/Engine/Plugins/FX/Niagara/Source/NiagaraEditor/Private/NiagaraBakerRenderer.cpp
--- a/Engine/Plugins/FX/Niagara/Source/NiagaraEditor/Private/NiagaraBakerRenderer.cpp
+++ b/Engine/Plugins/FX/Niagara/Source/NiagaraEditor/Private/NiagaraBakerRenderer.cpp
@@ -9,6 +9,7 @@
 #include "NiagaraComputeExecutionContext.h"
 #include "NiagaraGpuComputeDispatchInterface.h"
 #include "NiagaraBatchedElements.h"
+#include "NiagaraUseOpenVDB.h"
 
 #include "NiagaraDataInterfaceGrid2DCollection.h"
 #include "NiagaraDataInterfaceRenderTarget2D.h"
@@ -742,7 +743,7 @@ bool FNiagaraBakerRenderer::ExportVolume(FStringView FilePath, FIntVector ImageS
 	const FString FileExtension = FPaths::GetExtension(FilePath.GetData(), true);
 	if (FileExtension == TEXT(".vdb"))
 	{
-#if NIAGARA_USE_OPENVDB
+#if UE_USE_OPENVDB
 		return OpenVDBTools::WriteImageDataToOpenVDBFile(FilePath, ImageSize, ImageData, false);
 #else
 		return false;
