Missing file.

diff --git a/Engine/Plugins/FX/Niagara/Source/Niagara/Classes/NiagaraUseOpenVDB.h b/Engine/Plugins/FX/Niagara/Source/Niagara/Classes/NiagaraUseOpenVDB.h
new file mode 100644
--- /dev/null
+++ b/Engine/Plugins/FX/Niagara/Source/Niagara/Classes/NiagaraUseOpenVDB.h
@@ -0,0 +1,7 @@
+// Copyright Epic Games, Inc. All Rights Reserved.
+
+#pragma once
+
+#ifndef UE_USE_OPENVDB
+	#define UE_USE_OPENVDB 0
+#endif
