Removing dependencies on Niagara Private files.

diff --git a/Engine/Plugins/FX/Niagara/Source/Niagara/Internal/DataInterface/NiagaraDataInterfaceActorComponent.h b/Engine/Plugins/FX/Niagara/Source/Niagara/Internal/DataInterface/NiagaraDataInterfaceActorComponent.h
new file mode 100644
--- /dev/null
+++ b/Engine/Plugins/FX/Niagara/Source/Niagara/Internal/DataInterface/NiagaraDataInterfaceActorComponent.h
@@ -0,0 +1,108 @@
+// Copyright Epic Games, Inc. All Rights Reserved.
+#pragma once
+
+#include "NiagaraCommon.h"
+#include "NiagaraShared.h"
+#include "NiagaraDataInterface.h"
+#include "NiagaraParameterStore.h"
+#include "NiagaraDataInterfaceActorComponent.generated.h"
+
+UENUM()
+enum class ENDIActorComponentSourceMode : uint8
+{
+	/**
+	The default binding mode first we look for a valid binding on the ActorOrComponentParameter.
+	If this it no valid we then look at the SourceActor.
+	If these both fail we are bound to nothing.
+	*/
+	Default,
+	/**
+	We will first look at the attach parent.
+	If this is not valid we fallback to the Default binding mode.
+	*/
+	AttachParent,
+	/**
+	We will look for a ULocalPlayer with the provided index.
+	If this is not valid we fallback to the Default binding mode.
+	*/
+	LocalPlayer,
+};
+
+/**
+Data interface that gives you access to actor & component information.
+*/
+UCLASS(EditInlineNew, Category = "Actor", meta = (DisplayName = "Actor Component Interface"))
+class NIAGARA_API UNiagaraDataInterfaceActorComponent : public UNiagaraDataInterface
+{
+	GENERATED_UCLASS_BODY()
+
+	BEGIN_SHADER_PARAMETER_STRUCT(FShaderParameters, )
+		SHADER_PARAMETER(uint32,		Valid)
+		SHADER_PARAMETER(FMatrix44f,	Matrix)
+		SHADER_PARAMETER(FQuat4f,		Rotation)
+		SHADER_PARAMETER(FVector3f,		Scale)
+		SHADER_PARAMETER(FVector3f,		Velocity)
+	END_SHADER_PARAMETER_STRUCT();
+
+public:
+	/** When this option is disabled, we use the previous frame's data for the skeletal mesh and can often issue the simulation early. This greatly
+	reduces overhead and allows the game thread to run faster, but comes at a tradeoff if the dependencies might leave gaps or other visual artifacts.*/
+	UPROPERTY(EditAnywhere, Category = "Performance")
+	bool bRequireCurrentFrameData = true;
+
+	/** Controls how we find the actor / component we want to bind to. */
+	UPROPERTY(EditAnywhere, Category = "ActorComponent")
+	ENDIActorComponentSourceMode SourceMode = ENDIActorComponentSourceMode::Default;
+
+	UPROPERTY(EditAnywhere, Category = "ActorComponent", meta=(EditCondition="SourceMode==ENDIActorComponentSourceMode::LocalPlayer"))//, EditConditionHides))
+	int32 LocalPlayerIndex = 0;
+
+	/** Optional source actor to use, if the user parameter binding is valid this will be ignored. */
+	UPROPERTY(EditAnywhere, Category = "ActorComponent")
+	TLazyObjectPtr<AActor> SourceActor;
+
+	/** User parameter binding to use, overrides SourceActor.  Can be set by Blueprint, etc. */
+	UPROPERTY(EditAnywhere, Category = "ActorComponent")
+	FNiagaraUserParameterBinding ActorOrComponentParameter;
+
+	//UObject Interface
+	virtual void PostInitProperties() override;
+	//UObject Interface End
+
+	class UActorComponent* ResolveComponent(FNiagaraSystemInstance* SystemInstance, const void* PerInstanceData) const;
+
+	//UNiagaraDataInterface Interface
+	virtual bool CanExecuteOnTarget(ENiagaraSimTarget Target) const override { return true; }
+
+	virtual void GetFunctions(TArray<FNiagaraFunctionSignature>& OutFunctions) override;
+	virtual void GetVMExternalFunction(const FVMExternalFunctionBindingInfo& BindingInfo, void* InstanceData, FVMExternalFunction &OutFunc) override;
+#if WITH_EDITORONLY_DATA
+	virtual bool AppendCompileHash(FNiagaraCompileHashVisitor* InVisitor) const override;
+	virtual void GetParameterDefinitionHLSL(const FNiagaraDataInterfaceGPUParamInfo& ParamInfo, FString& OutHLSL) override;
+	virtual bool GetFunctionHLSL(const FNiagaraDataInterfaceGPUParamInfo& ParamInfo, const FNiagaraDataInterfaceGeneratedFunction& FunctionInfo, int FunctionInstanceIndex, FString& OutHLSL) override;
+	virtual bool UpgradeFunctionCall(FNiagaraFunctionSignature& FunctionSignature) override;
+#endif
+	virtual void BuildShaderParameters(FNiagaraShaderParametersBuilder& ShaderParametersBuilder) const override;
+	virtual void SetShaderParameters(const FNiagaraDataInterfaceSetShaderParametersContext& Context) const override;
+
+	virtual bool InitPerInstanceData(void* PerInstanceData, FNiagaraSystemInstance* SystemInstance) override;
+	virtual void DestroyPerInstanceData(void* PerInstanceData, FNiagaraSystemInstance* SystemInstance) override;
+	virtual int32 PerInstanceDataSize() const override;
+	virtual bool PerInstanceTick(void* PerInstanceData, FNiagaraSystemInstance* SystemInstance, float DeltaSeconds) override;
+
+	virtual void ProvidePerInstanceDataForRenderThread(void* DataForRenderThread, void* PerInstanceData, const FNiagaraSystemInstanceID& SystemInstance) override;
+
+	virtual bool HasTickGroupPrereqs() const override { return true; }
+	virtual ETickingGroup CalculateTickGroup(const void* PerInstanceData) const override;
+
+	virtual bool Equals(const UNiagaraDataInterface* Other) const override;
+	virtual bool CopyToInternal(UNiagaraDataInterface* Destination) const override;
+
+	virtual bool HasPreSimulateTick() const override { return true; }
+	//UNiagaraDataInterface Interface
+
+private:
+	void VMGetMatrix(FVectorVMExternalFunctionContext& Context);
+	void VMGetTransform(FVectorVMExternalFunctionContext& Context);
+	void VMGetPhysicsVelocity(FVectorVMExternalFunctionContext& Context);
+};
diff --git a/Engine/Plugins/FX/Niagara/Source/Niagara/Internal/DataInterface/NiagaraDataInterfaceDataChannelCommon.h b/Engine/Plugins/FX/Niagara/Source/Niagara/Internal/DataInterface/NiagaraDataInterfaceDataChannelCommon.h
new file mode 100644
--- /dev/null
+++ b/Engine/Plugins/FX/Niagara/Source/Niagara/Internal/DataInterface/NiagaraDataInterfaceDataChannelCommon.h
@@ -0,0 +1,206 @@
+// Copyright Epic Games, Inc. All Rights Reserved.
+#pragma once
+
+#include "NiagaraCommon.h"
+#include "NiagaraDataSetCompiledData.h"
+#include "NiagaraCompileHash.h"
+#include "Misc/LazySingleton.h"
+#include "NiagaraDataChannelCommon.h"
+#include "NiagaraDataInterfaceDataChannelCommon.generated.h"
+
+UENUM()
+enum class ENiagaraDataChannelAllocationMode : uint8
+{
+	/** Fixed number of elements available to write per frame. */
+	Static,
+
+	/** Allow N elements per instance, per frame. Per instance is context dependent meaning data written from particles scripts will allocate per particle. From emitter script will allocate per writing emitter etc. */
+	//TODO: For this we need a Pre Stage (+on CPU) from which we can allocate and a Post Stage from which we can publish the results.
+	//PerInstance,
+	//Dynamic?
+};
+
+//TODO: Possible we may want to do reads and writes using data channels in a single system in future, avoiding the need to push data out to any manager class etc.
+// UENUM()
+// enum class ENiagaraDataChannelScope
+// {
+// 	/** Data is read or written internally to other DIs in the same Niagara System. */
+// 	Local,
+// 	/** Data is read or written externally to system, into the outside world. */
+// 	World,
+// };
+
+//Enable various invasive debugging features that will bloat memory and incur overhead.
+#define DEBUG_NDI_DATACHANNEL (!(UE_BUILD_TEST||UE_BUILD_SHIPPING))
+
+/** 
+Stores info for a function called on a DataChannel DI.
+Describes a function call which is used when generating binding information between the data and the VM & GPU scripts.
+*/
+USTRUCT()
+struct FNDIDataChannelFunctionInfo
+{
+	GENERATED_BODY()
+
+	UPROPERTY()
+	FName FunctionName;
+
+	UPROPERTY()
+	TArray<FNiagaraVariableBase> Inputs;
+
+	UPROPERTY()
+	TArray<FNiagaraVariableBase> Outputs;
+
+	bool operator==(const FNDIDataChannelFunctionInfo& Other)const;
+	bool operator!=(const FNDIDataChannelFunctionInfo& Other)const{ return !operator==(Other);}
+	bool CheckHashConflict(const FNDIDataChannelFunctionInfo& Other)const;
+};
+
+uint32 GetTypeHash(const FNDIDataChannelFunctionInfo& FuncInfo);
+
+/** Binding between registers accessed in the data channel DI function calls and the relevant data in a dataset. */
+struct FNDIDataChannelRegisterBinding
+{
+	static const uint32 RegisterBits = 30;
+	static const uint32 DataTypeBitst = 2;
+	FNDIDataChannelRegisterBinding(uint32 InFunctionRegisterIndex, uint32 InDataSetRegisterIndex, ENiagaraBaseTypes InDataType)
+	: FunctionRegisterIndex(InFunctionRegisterIndex)
+	, DataSetRegisterIndex(InDataSetRegisterIndex)
+	, DataType((uint32)InDataType)
+	{
+		check(InDataSetRegisterIndex <= (1u << RegisterBits) - 1);
+		check((uint32)InDataType <= (1u << DataTypeBitst) - 1);
+	}
+	uint32 FunctionRegisterIndex;
+	uint32 DataSetRegisterIndex : RegisterBits;
+	uint32 DataType : DataTypeBitst;
+};
+
+
+/** Layout info mapping from a function called by a data channel DI to the actual data set register. */
+struct FNDIDataChannel_FunctionToDataSetBinding
+{
+	/** Bindings used by the VM calls to map dataset registers to the relevant function call registers. */
+	TArray<FNDIDataChannelRegisterBinding> VMRegisterBindings;
+
+	uint32 NumFloatComponents = 0;
+	uint32 NumInt32Components = 0;
+	uint32 NumHalfComponents = 0;
+	uint32 FunctionLayoutHash = 0;
+	uint32 DataSetLayoutHash = 0;
+
+	#if DEBUG_NDI_DATACHANNEL
+	FNDIDataChannelFunctionInfo DebugFunctionInfo;
+	FNiagaraDataSetCompiledData DebugCompiledData;
+	#endif
+
+	FNDIDataChannel_FunctionToDataSetBinding(const FNDIDataChannelFunctionInfo& FunctionInfo, const FNiagaraDataSetCompiledData& DataSetLayout);
+
+	bool IsValid()const { return DataSetLayoutHash != 0; }
+	void GenVMBindings(const FNiagaraVariableBase& Var, const UStruct* Struct, uint32& FuncFloatRegister, uint32& FuncIntRegister, uint32& FuncHalfRegister, uint32& DataSetFloatRegister, uint32& DataSetIntRegister, uint32& DataSetHalfRegister);
+};
+
+typedef TSharedPtr<FNDIDataChannel_FunctionToDataSetBinding,ESPMode::ThreadSafe> FNDIDataChannel_FuncToDataSetBindingPtr;
+
+/** 
+Manager class that generates and allows access to layout information used by the Data Channel DIs. 
+These layout buffers will map from a DI's function calls to the register offsets of the relevant data inside the DataSet buffers.
+Each combination of dataset layout and function info will need a unique mapping but these will be used by many instances.
+This manager class allows the de-duplication and sharing of such binding data that would otherwise have to be generated and stored per DI instance.
+*/
+struct FNDIDataChannelLayoutManager
+{
+	/** Map containing binding information for each function info/dataset layout pair. */
+	TMap<uint32, FNDIDataChannel_FuncToDataSetBindingPtr> FunctionToDataSetLayoutMap;
+
+	/** 
+	Typically this map will be accessed from the game thread and then the shared ptrs of actual layout information passed off to various threads. 
+	Though for additional safety we'll use a lock. It should be very low contention.
+	*/
+	FRWLock FunctionToDataSetMapLock;
+
+public:
+
+	//TLazySingleton interface
+	static FNDIDataChannelLayoutManager& Get();
+	static void TearDown();
+
+	void Reset();
+
+	/** Generates a key that can be used to retrieve layout information on both the GT and RT. */
+	uint32 GetLayoutKey(const FNDIDataChannelFunctionInfo& FunctionInfo, const FNiagaraDataSetCompiledData& DataSetLayout) const
+	{
+		return HashCombine(GetTypeHash(FunctionInfo), DataSetLayout.GetLayoutHash());
+	}
+	
+	/** 
+	Retrieves, or generates, the layout information that maps from the given function to the data in the given dataset.
+	*/
+	FNDIDataChannel_FuncToDataSetBindingPtr GetLayoutInfo(const FNDIDataChannelFunctionInfo& FunctionInfo, const FNiagaraDataSetCompiledData& DataSetLayout);
+};
+
+/** A sorted table of parameters accessed by each GPU script */
+USTRUCT()
+struct FNDIDataChannel_GPUScriptParameterAccessInfo 
+{
+	GENERATED_BODY()
+	
+	UPROPERTY()
+	TArray<FNiagaraVariableBase> SortedParameters;
+};
+
+/** 
+Compile time data used by Data Channel interfaces.
+*/
+USTRUCT()
+struct FNDIDataChannelCompiledData
+{
+	GENERATED_BODY()
+
+	/** Initializes function bindings and binds if there is a valid DataSetCompiledData. */
+	bool Init(UNiagaraSystem* System, UNiagaraDataInterface* OwnerDI);
+
+	int32 FindFunctionInfoIndex(FName Name, const TArray<FNiagaraVariableBase>& VariadicInputs, const TArray<FNiagaraVariableBase>& VariadicOutputs)const;
+
+	const TArray<FNDIDataChannelFunctionInfo>& GetFunctionInfo()const { return FunctionInfo; }
+	const TMap<FNiagaraCompileHash, FNDIDataChannel_GPUScriptParameterAccessInfo>& GetGPUScriptParameterInfos()const{ return GPUScriptParameterInfos; }
+
+	bool UsedByCPU()const{ return bUsedByCPU; }
+	bool UsedByGPU()const{ return bUsedByGPU; }
+	int32 GetTotalParams()const{ return TotalParams; }
+
+protected:
+
+	/**
+	Data describing every function call for this DI in VM scripts. 
+	VM Access to data channels uses a binding from script to DataSet per function call (de-duped by layout).
+	*/
+	UPROPERTY()
+	TArray<FNDIDataChannelFunctionInfo> FunctionInfo;
+
+	/** 
+	Info about which parameters are accessed for each GPU script. 
+	GPU access to data channels uses a binding from script to DataSet per script via a mapping of param<-->data set offsets.
+	*/
+	UPROPERTY()
+	TMap<FNiagaraCompileHash, FNDIDataChannel_GPUScriptParameterAccessInfo> GPUScriptParameterInfos;
+
+	/** Total param count across all scripts. Allows easy pre-allocation for the buffers at runtime. */
+	UPROPERTY()
+	uint32 TotalParams = 0;
+
+	UPROPERTY()
+	bool bUsedByCPU = false;
+
+	UPROPERTY()
+	bool bUsedByGPU = false;
+
+	/** Iterates over all scripts for the owning system and gathers all functions and parameters accessing this DI. Building the FunctionInfoTable and GPUScriptParameterInfos map.  */
+	void GatherAccessInfo(UNiagaraSystem* System, UNiagaraDataInterface* Owner);
+};
+
+
+namespace NDIDataChannelUtilities
+{
+	void SortParameters(TArray<FNiagaraVariableBase>& Parameters);	
+}
\ No newline at end of file
diff --git a/Engine/Plugins/FX/Niagara/Source/Niagara/Internal/DataInterface/NiagaraDataInterfaceDataChannelRead.h b/Engine/Plugins/FX/Niagara/Source/Niagara/Internal/DataInterface/NiagaraDataInterfaceDataChannelRead.h
new file mode 100644
--- /dev/null
+++ b/Engine/Plugins/FX/Niagara/Source/Niagara/Internal/DataInterface/NiagaraDataInterfaceDataChannelRead.h
@@ -0,0 +1,191 @@
+// Copyright Epic Games, Inc. All Rights Reserved.
+#pragma once
+
+
+/**
+
+Data Channel Read Interface.
+Enabled Niagara Systems to read data from a NiagaraDataChannel that has been previously generated by another System or Game code/BP.
+
+Read DIs will grab a buffer from the handler for their bound Data Channel.
+This could be a global buffer or some localized buffer for their region of the map
+Or any other buffer sub division scheme that the handler chooses.
+The main point being that this buffer can change from frame to frame and instance to instance.
+The handler is free to store and distribute it's Channel Data as it pleases.
+
+Accessor functions on the Data Channel Read and Write DIs can have any number of parameters, allowing a single function call to access arbitrary data from the Channel.
+This avoids cumbersome work in the graph to access data but requires special handling inside the DI.
+
+TODO: Local support - We can read directly from a writer DI in the same system to avoid having to publish that data in some cases. 
+There is working code for this already but the concept/API needs more fleshing out before we include it.
+
+*/
+
+
+#include "NiagaraDataInterfaceDataChannelCommon.h"
+#include "NiagaraDataInterface.h"
+#include "NiagaraDataInterfaceDataChannelRead.generated.h"
+
+class UNiagaraDataInterfaceDataChannelWrite;
+struct FNDIDataChannelWriteInstanceData;
+class UNiagaraDataChannelHandler;
+class FNiagaraDataBuffer;
+
+UCLASS(Experimental, EditInlineNew, Category = "Data Channels", meta = (DisplayName = "Data Channel Reader"))
+class NIAGARA_API UNiagaraDataInterfaceDataChannelRead : public UNiagaraDataInterface
+{
+	GENERATED_UCLASS_BODY()
+protected:
+
+public:
+
+	//TODO: In future we may allow reads and writes that stay local to a single system.	
+// 	/** The scope at which to read data. When reading locally we'll read directly from a Data Channel Write Interface in this or another emitter. When reading from World scope, we'll read from a named Data Channel. */
+// 	UPROPERTY(EditAnywhere, Category = "Data Channel")
+// 	ENiagaraDataChannelScope Scope;
+// 
+// 	/** Name of the source */
+// 	UPROPERTY(EditAnywhere, Category="Data Channel", meta=(EditCondition = "Scope == ENiagaraDataChannelScope::Local"))
+// 	FName Source;
+
+	/** When reading from external, the channel to consume. */
+	UPROPERTY(EditAnywhere, Category="Data Channel")
+	FNiagaraDataChannelReference Channel;
+	
+	/** True if this reader will read the current frame's data. If false, we read the previous frame.
+	* Reading the current frame introduces a tick order dependency but allows for zero latency reads. Any data channel elements that are generated after this reader is used are missed.
+	* Reading the previous frame's data introduces a frame of latency but ensures we never miss any data as we have access to the whole frame.
+	*/
+	UPROPERTY(EditAnywhere, Category = "Data Channel")
+	bool bReadCurrentFrame = false;
+
+	//UObject Interface
+	virtual void PostInitProperties() override;
+	virtual void BeginDestroy() override;
+	//UObject Interface End
+
+	//UNiagaraDataInterface Interface
+	virtual void GetFunctions(TArray<FNiagaraFunctionSignature>& OutFunctions) override;
+	virtual void GetVMExternalFunction(const FVMExternalFunctionBindingInfo& BindingInfo, void* InstanceData, FVMExternalFunction &OutFunc) override;
+	virtual bool CanExecuteOnTarget(ENiagaraSimTarget Target) const override { return true; }
+#if WITH_EDITORONLY_DATA
+	virtual bool AppendCompileHash(FNiagaraCompileHashVisitor* InVisitor) const override;
+	virtual void GetCommonHLSL(FString& OutHLSL)override;
+	virtual bool GetFunctionHLSL(FNiagaraDataInterfaceHlslGenerationContext& HlslGenContext, FString& OutHLSL) override;
+	virtual void GetParameterDefinitionHLSL(FNiagaraDataInterfaceHlslGenerationContext& HlslGenContext, FString& OutHLSL) override;
+
+	virtual void PostCompile()override;
+#endif	
+
+#if WITH_EDITOR	
+	virtual void GetFeedback(UNiagaraSystem* InAsset, UNiagaraComponent* InComponent, TArray<FNiagaraDataInterfaceError>& OutErrors, TArray<FNiagaraDataInterfaceFeedback>& OutWarnings, TArray<FNiagaraDataInterfaceFeedback>& OutInfo) override;
+	virtual void ValidateFunction(const FNiagaraFunctionSignature& Function, TArray<FText>& OutValidationErrors) override;
+#endif
+
+	virtual bool Equals(const UNiagaraDataInterface* Other) const override;
+
+	virtual void BuildShaderParameters(FNiagaraShaderParametersBuilder& ShaderParametersBuilder) const override;
+	virtual void SetShaderParameters(const FNiagaraDataInterfaceSetShaderParametersContext& Context) const override;
+
+	virtual bool InitPerInstanceData(void* PerInstanceData, FNiagaraSystemInstance* SystemInstance) override;
+	virtual void DestroyPerInstanceData(void* PerInstanceData, FNiagaraSystemInstance* SystemInstance) override;
+	virtual int32 PerInstanceDataSize() const override;
+	virtual bool HasPreSimulateTick() const override { return true; }
+	virtual bool HasPostSimulateTick() const override { return true; }
+	virtual bool PerInstanceTick(void* PerInstanceData, FNiagaraSystemInstance* SystemInstance, float DeltaSeconds) override;
+	virtual bool PerInstanceTickPostSimulate(void* PerInstanceData, FNiagaraSystemInstance* SystemInstance, float DeltaSeconds) override;
+	virtual void ProvidePerInstanceDataForRenderThread(void* DataForRenderThread, void* PerInstanceData, const FNiagaraSystemInstanceID& SystemInstance) override;
+	virtual void GetEmitterDependencies(UNiagaraSystem* Asset, TArray<FVersionedNiagaraEmitter>& Dependencies) const;
+	//UNiagaraDataInterface Interface
+
+	//Functions usable anywhere.
+	void Num(FVectorVMExternalFunctionContext& Context);
+	void Read(FVectorVMExternalFunctionContext& Context, int32 FuncIdx);
+	void Consume(FVectorVMExternalFunctionContext& Context, int32 FuncIdx);
+	
+	FNDIDataChannelCompiledData& GetCompiledData() { return CompiledData; }
+
+protected:
+	virtual bool CopyToInternal(UNiagaraDataInterface* Destination) const override;
+	UNiagaraDataInterfaceDataChannelWrite* FindSourceDI()const;
+
+ 	UPROPERTY()
+ 	FNDIDataChannelCompiledData CompiledData;
+};
+
+struct FNDIDataChannelReadInstanceData
+{
+	//TODO: Local reads.
+// 	/** The local DataChannel writer we're bounds to, if any. */
+// 	TWeakObjectPtr<UNiagaraDataInterfaceDataChannelWrite> SourceDI;
+// 
+// 	/** The instance data of the local DataChannel writer we're bound to, if any. */
+// 	FNDIDataChannelWriteInstanceData* SourceInstData = nullptr;
+	
+	/** Pointer to the world DataChannel Channel we're reading from, if any. */
+	TWeakObjectPtr<UNiagaraDataChannelHandler> DataChannel;
+
+	/** Pointer to the cpu simulation buffer from the DataChannel if we're reading from there. */
+	FNiagaraDataBuffer* ExternalCPU = nullptr;
+	
+	/** The GPU data set from the DataChannel. We must grab the whole set rather than a buffer here as we don't know which buffer will be contain this frame's channel data. 
+	*	TODO: We'll likely want to improve this GPU sim flow soon. Possibly grab an RT proxy object from the data channel handler and pass that over to the GPU instead of a direct dataset like this.
+	*/
+	FNiagaraDataSet* ExternalGPU = nullptr;
+
+	/** Cached hash to check if the layout of our source data has changed. */
+	uint64 ChachedDataSetLayoutHash = INDEX_NONE;
+
+	/** When true we should update our RT data next tick. */
+	mutable bool bUpdateRTData = false;
+
+	/** Binding info from the Data Channel to each of our function's parameters. */
+	//FNDIDataChannelBindingInfo BindingInfo;
+
+	TArray<FNDIDataChannel_FuncToDataSetBindingPtr, TInlineAllocator<8>> FuncToDataSetBindingInfo;
+
+	/** Keys for each of the above function infos. */
+	//TArray<uint32> FuncToDataSetLayoutKeys;
+	
+	/** Index we use to consume DataChannel data.
+		Reset at each tick.
+		All consumers of the DI use the same index.
+		TODO: Allow ConsumeUnique that has a unique counter per call.
+	*/
+	std::atomic<int32> ConsumeIndex = 0;
+
+	virtual ~FNDIDataChannelReadInstanceData();
+	FNiagaraDataBuffer* GetReadBufferCPU();
+	bool Init(UNiagaraDataInterfaceDataChannelRead* Interface, FNiagaraSystemInstance* Instance);
+	bool Tick(UNiagaraDataInterfaceDataChannelRead* Interface, FNiagaraSystemInstance* Instance);
+	bool PostTick(UNiagaraDataInterfaceDataChannelRead* Interface, FNiagaraSystemInstance* Instance);
+};
+
+struct FNiagaraDataInterfaceProxy_DataChannelRead : public FNiagaraDataInterfaceProxy
+{
+	virtual void ConsumePerInstanceDataFromGameThread(void* PerInstanceData, const FNiagaraSystemInstanceID& Instance) override;
+	virtual int32 PerInstanceDataPassedToRenderThreadSize() const override;
+
+	/** Persistent per instance data on the RT. Constructed when consuming data passed from GT->RT. */
+	struct FInstanceData
+	{
+		//GPU Dataset from the channel handler. We'll grab the current buffer from this on the RT.
+		//This must be grabbed fresh from the handler each frame as it's lifetime cannot be ensured.
+		FNiagaraDataSet* ChannelDataSet = nullptr;
+
+		/**
+		A buffer containing layout information needed to access parameters for each script using this DI.
+		*/
+		FReadBuffer ParameterLayoutBuffer;
+
+		TResourceArray<uint32> ParameterLayoutData;
+
+		/**
+		Offsets into the parameter table are embedded in the gpu script hlsl.
+		At hlsl gen time we can only know which parameter are accessed by each script individually so each script must have it's own parameter binding table.
+		*/
+		TMap<FNiagaraCompileHash, uint32> GPUScriptParameterTableOffsets;
+	};
+
+	TMap<FNiagaraSystemInstanceID, FInstanceData> SystemInstancesToProxyData_RT;
+};
\ No newline at end of file
diff --git a/Engine/Plugins/FX/Niagara/Source/Niagara/Internal/DataInterface/NiagaraDataInterfaceDataChannelWrite.h b/Engine/Plugins/FX/Niagara/Source/Niagara/Internal/DataInterface/NiagaraDataInterfaceDataChannelWrite.h
new file mode 100644
--- /dev/null
+++ b/Engine/Plugins/FX/Niagara/Source/Niagara/Internal/DataInterface/NiagaraDataInterfaceDataChannelWrite.h
@@ -0,0 +1,112 @@
+// Copyright Epic Games, Inc. All Rights Reserved.
+#pragma once
+
+/**
+
+Data Channel Write Interface.
+Enables Niagara Systems to push data into a NiagaraDataChannel for access later by another System or Game code/BP.
+
+We write into an intermediate buffer on the DI which is then published to the Data Channel on post tick.
+This is simple and avoids race conditions and synchronization headaches but does introduce additional copying work that in many cases may be avoided.
+In the future we may allow for direct writes into the data channel buffers and/or avoiding separate writes entirely by publishing the owning emitter particle buffers.
+Though these other options have their own downsides.
+
+Accessor functions on the Data Channel Read and Write DIs can have any number of parameters, allowing a single function call to access arbitrary data from the Channel.
+This avoids cumbersome work in the graph to access data but requires special handling inside the DI.
+
+*/
+
+
+#include "NiagaraDataInterfaceDataChannelCommon.h"
+#include "NiagaraDataInterface.h"
+#include "NiagaraDataInterfaceDataChannelWrite.generated.h"
+
+/** Additional compile time information used by the Write DI. */
+USTRUCT()
+struct FNDIDataChannelWriteCompiledData : public FNDIDataChannelCompiledData
+{
+	GENERATED_BODY()
+
+	/** Internal buffer layout. Contains only the data actually written by this DI. */
+	UPROPERTY()
+	FNiagaraDataSetCompiledData DataLayout;
+
+	bool Init(UNiagaraSystem* System, UNiagaraDataInterfaceDataChannelWrite* OwnerDI);
+};
+
+UCLASS(Experimental, EditInlineNew, Category = "Data Channels", meta = (DisplayName = "Data Channel Writer"))
+class NIAGARA_API UNiagaraDataInterfaceDataChannelWrite : public UNiagaraDataInterface
+{
+	GENERATED_UCLASS_BODY()
+
+public:
+
+	/** How should we allocate the buffer into which we write data. */
+	UPROPERTY(EditAnywhere, Category = "Data Channel")
+	ENiagaraDataChannelAllocationMode AllocationMode = ENiagaraDataChannelAllocationMode::Static;
+
+	/** How many elements to allocate for writing per frame? Usage is defendant on AllocationMode. TODO: Allow allocation count to be controlled dynamically from script? */
+	UPROPERTY(EditAnywhere, Category = "Data Channel")
+	uint32 AllocationCount = 0;
+
+	/** Whether the data generated by this DI should be published to the world game data channel. This is require to allow game BP and C++ to read this data. */
+	UPROPERTY(EditAnywhere, Category = "Data Channel")
+	bool bPublishToGame = false;
+
+	/** Whether the data generated by this DI should be published to the world CPUSim data channel. This is required for CPU emitters in other Niagara Systems to read this data. */
+	UPROPERTY(EditAnywhere, Category = "Data Channel")
+	bool bPublishToCPU = false;
+
+	/** Whether the data generated by this DI should be published to a world data channel. This is required to allow GPU emitters in other Niagara Systems to read this data.  */
+	UPROPERTY(EditAnywhere, Category = "Data Channel")
+	bool bPublishToGPU = false;
+
+	/** When writing externally, the channel to use. */
+	UPROPERTY(EditAnywhere, Category = "Data Channel", meta = (EditCondition = "bPublishToWorld"))
+	FNiagaraDataChannelReference Channel;
+
+	//UObject Interface
+	virtual void PostInitProperties() override;
+	//UObject Interface End
+
+	//UNiagaraDataInterface Interface
+	virtual void GetFunctions(TArray<FNiagaraFunctionSignature>& OutFunctions) override;
+	virtual void GetVMExternalFunction(const FVMExternalFunctionBindingInfo& BindingInfo, void* InstanceData, FVMExternalFunction& OutFunc) override;
+	virtual bool CanExecuteOnTarget(ENiagaraSimTarget Target) const override { return true; }
+#if WITH_EDITORONLY_DATA
+	virtual bool AppendCompileHash(FNiagaraCompileHashVisitor* InVisitor) const override;
+	virtual void GetCommonHLSL(FString& OutHLSL)override;
+	virtual bool GetFunctionHLSL(FNiagaraDataInterfaceHlslGenerationContext& HlslGenContext, FString& OutHLSL) override;
+	virtual void GetParameterDefinitionHLSL(FNiagaraDataInterfaceHlslGenerationContext& HlslGenContext, FString& OutHLSL) override;
+
+	virtual void PostCompile()override;
+#endif
+	virtual void BuildShaderParameters(FNiagaraShaderParametersBuilder& ShaderParametersBuilder) const override;
+	virtual void SetShaderParameters(const FNiagaraDataInterfaceSetShaderParametersContext& Context) const override;
+
+	virtual bool Equals(const UNiagaraDataInterface* Other) const override;
+
+	virtual bool InitPerInstanceData(void* PerInstanceData, FNiagaraSystemInstance* SystemInstance) override;
+	virtual void DestroyPerInstanceData(void* PerInstanceData, FNiagaraSystemInstance* SystemInstance) override;
+	virtual int32 PerInstanceDataSize() const override;
+	virtual bool HasPreSimulateTick() const override { return true; }
+	virtual bool HasPostSimulateTick() const override { return true; }
+	virtual bool PerInstanceTick(void* PerInstanceData, FNiagaraSystemInstance* SystemInstance, float DeltaSeconds) override;
+	virtual bool PerInstanceTickPostSimulate(void* PerInstanceData, FNiagaraSystemInstance* SystemInstance, float DeltaSeconds) override;
+	virtual void ProvidePerInstanceDataForRenderThread(void* DataForRenderThread, void* PerInstanceData, const FNiagaraSystemInstanceID& SystemInstance) override;
+	//UNiagaraDataInterface Interface
+
+	void Num(FVectorVMExternalFunctionContext& Context);
+	void Write(FVectorVMExternalFunctionContext& Context, int32 FuncIdx);
+	void Append(FVectorVMExternalFunctionContext& Context, int32 FuncIdx);
+
+	const FNDIDataChannelWriteCompiledData& GetCompiledData()const { return CompiledData; }
+
+	bool ShouldPublish()const { return bPublishToGame || bPublishToCPU || bPublishToGPU; }
+
+protected:
+	virtual bool CopyToInternal(UNiagaraDataInterface* Destination) const override;
+
+	UPROPERTY()
+	FNDIDataChannelWriteCompiledData CompiledData;
+};
\ No newline at end of file
diff --git a/Engine/Plugins/FX/Niagara/Source/Niagara/Internal/DataInterface/NiagaraDataInterfaceStaticMesh.h b/Engine/Plugins/FX/Niagara/Source/Niagara/Internal/DataInterface/NiagaraDataInterfaceStaticMesh.h
new file mode 100644
--- /dev/null
+++ b/Engine/Plugins/FX/Niagara/Source/Niagara/Internal/DataInterface/NiagaraDataInterfaceStaticMesh.h
@@ -0,0 +1,314 @@
+// Copyright Epic Games, Inc. All Rights Reserved.
+#pragma once
+
+#include "NiagaraDataInterfaceMeshCommon.h"
+#include "StaticMeshResources.h"
+#include "Engine/StaticMesh.h"
+#include "Experimental/NiagaraMeshUvMappingHandle.h"
+#include "NiagaraDataInterfaceStaticMesh.generated.h"
+
+UENUM()
+enum class ENDIStaticMesh_SourceMode : uint8
+{
+	/**
+	Default behavior.
+	- Use "Source" when specified (either set explicitly or via blueprint with Set Niagara Static Mesh Component).
+	- When no source is specified, attempt to find a Static Mesh Component on an attached actor or component.
+	- If no source actor/component specified and no attached component found, fall back to the "Default Mesh" specified.
+	*/
+	Default,
+
+	/**
+	Only use "Source" (either set explicitly or via blueprint with Set Niagara Static Mesh Component).
+	*/
+	Source,
+
+	/**
+	Only use the parent actor or component the system is attached to.
+	*/
+	AttachParent,
+
+	/**
+	Only use the "Default Mesh" specified.
+	*/
+	DefaultMeshOnly,
+};
+
+USTRUCT()
+struct FNDIStaticMeshSectionFilter
+{
+	GENERATED_USTRUCT_BODY();
+
+	/** Only allow sections these material slots. */
+	UPROPERTY(EditAnywhere, Category="Section Filter")
+	TArray<int32> AllowedMaterialSlots;
+
+	//Others?
+	//Banned material slots
+	
+	void Init(class UNiagaraDataInterfaceStaticMesh* Owner, bool bAreaWeighted);
+	FORCEINLINE bool CanEverReject()const { return AllowedMaterialSlots.Num() > 0; }
+};
+
+/** Data Interface allowing sampling of static meshes. */
+UCLASS(EditInlineNew, Category = "Meshes", meta = (DisplayName = "Static Mesh"))
+class NIAGARA_API UNiagaraDataInterfaceStaticMesh : public UNiagaraDataInterface
+{
+	GENERATED_UCLASS_BODY()
+
+public:
+	/** Controls how to retrieve the Static Mesh Component to attach to. */
+	UPROPERTY(EditAnywhere, Category = "Mesh")
+	ENDIStaticMesh_SourceMode SourceMode = ENDIStaticMesh_SourceMode::Default;
+	
+#if WITH_EDITORONLY_DATA
+	/** Mesh used to sample from when not overridden by a source actor from the scene. Only available in editor for previewing. This is removed in cooked builds. */
+	UPROPERTY(EditAnywhere, Category = "Mesh")
+	TSoftObjectPtr<UStaticMesh> PreviewMesh;
+#endif
+
+	/** Mesh used to sample from when not overridden by a source actor from the scene. This mesh is NOT removed from cooked builds. */
+	UPROPERTY(EditAnywhere, Category = "Mesh")
+	TObjectPtr<UStaticMesh> DefaultMesh;
+
+protected:
+	/** The source actor from which to sample. Takes precedence over the direct mesh. Note that this can only be set when used as a user variable on a component in the world. */
+	UPROPERTY(EditAnywhere, Category = "Mesh", meta = (DisplayName = "Source Actor"))
+	TSoftObjectPtr<AActor> SoftSourceActor;
+
+#if WITH_EDITORONLY_DATA
+	UPROPERTY()
+	TObjectPtr<AActor> Source_DEPRECATED;
+#endif
+
+	/** The source component from which to sample. Takes precedence over the direct mesh. Not exposed to the user, only indirectly accessible from blueprints. */
+	UPROPERTY(Transient)
+	TObjectPtr<UStaticMeshComponent> SourceComponent;
+public:
+
+	/** Array of filters the can be used to limit sampling to certain sections of the mesh. */
+	UPROPERTY(EditAnywhere, Category = "Mesh")
+	FNDIStaticMeshSectionFilter SectionFilter;
+
+	/** If true then the mesh velocity is taken from the mesh component's physics data. Otherwise it will be calculated by diffing the component transforms between ticks, which is more reliable but won't work on the first frame. */
+	UPROPERTY(EditAnywhere, AdvancedDisplay, Category = "Mesh")
+    bool bUsePhysicsBodyVelocity = false;
+
+	/** When true, we allow this DI to sample from streaming LODs. Selectively overriding the CVar fx.Niagara.NDIStaticMesh.UseInlineLODsOnly. */
+	UPROPERTY(EditAnywhere, Category = "LOD")
+	bool bAllowSamplingFromStreamingLODs = false;
+
+	/** 
+	Static Mesh LOD to sample.
+	When the desired LOD is not available, the next available LOD is used.
+	When LOD Index is negative, Desired LOD = Num LODs - LOD Index.
+	*/
+	UPROPERTY(EditAnywhere, Category = "LOD")
+	int32 LODIndex = 0;
+
+	/** Reference to a user parameter if we're reading one. */
+	UPROPERTY(EditAnywhere, Category = "LOD")
+	FNiagaraUserParameterBinding LODIndexUserParameter;
+
+	/** List of filtered sockets to use. */
+	UPROPERTY(EditAnywhere, Category = "Mesh")
+	TArray<FName> FilteredSockets;
+
+    /** Changed within the editor on PostEditChangeProperty. Should be changed whenever a refresh is desired.*/
+	uint32 ChangeId = 0;
+
+	//~ UObject interface
+	virtual void PostInitProperties() override;
+	virtual void PostLoad() override;
+#if WITH_EDITOR
+	virtual void PostEditChangeProperty(struct FPropertyChangedEvent& PropertyChangedEvent) override;
+	virtual bool CanEditChange(const FProperty* InProperty) const override;
+#endif
+	//~ UObject interface
+
+	//~ UNiagaraDataInterface interface
+	virtual bool InitPerInstanceData(void* PerInstanceData, FNiagaraSystemInstance* SystemInstance) override;
+	virtual void DestroyPerInstanceData(void* PerInstanceData, FNiagaraSystemInstance* SystemInstance) override;
+	virtual bool PerInstanceTick(void* PerInstanceData, FNiagaraSystemInstance* SystemInstance, float DeltaSeconds) override;
+	virtual int32 PerInstanceDataSize() const override;
+
+	virtual void ProvidePerInstanceDataForRenderThread(void* DataForRenderThread, void* PerInstanceData, const FNiagaraSystemInstanceID& SystemInstance) override;
+
+	virtual void GetFunctions(TArray<FNiagaraFunctionSignature>& OutFunctions) override;
+	void GetVertexSamplingFunctions(TArray<FNiagaraFunctionSignature>& OutFunctions, const FNiagaraFunctionSignature& BaseSignature) const;
+	void GetTriangleSamplingFunctions(TArray<FNiagaraFunctionSignature>& OutFunctions, const FNiagaraFunctionSignature& BaseSignature) const;
+	void GetSocketSamplingFunctions(TArray<FNiagaraFunctionSignature>& OutFunctions, const FNiagaraFunctionSignature& BaseSignature) const;
+	void GetSectionFunctions(TArray<FNiagaraFunctionSignature>& OutFunctions, const FNiagaraFunctionSignature& BaseSignature) const;
+	void GetMiscFunctions(TArray<FNiagaraFunctionSignature>& OutFunctions, const FNiagaraFunctionSignature& BaseSignature) const;
+	void GetUVMappingFunctions(TArray<FNiagaraFunctionSignature>& OutFunctions, const FNiagaraFunctionSignature& BaseSignature) const;
+	void GetDistanceFieldFunctions(TArray<FNiagaraFunctionSignature>& OutFunctions, const FNiagaraFunctionSignature& BaseSignature) const;
+	void GetDeprecatedFunctions(TArray<FNiagaraFunctionSignature>& OutFunctions, const FNiagaraFunctionSignature& BaseSignature) const;
+	void GetCpuAccessFunctions(TArray<FNiagaraFunctionSignature>& OutFunctions);
+
+#if WITH_EDITORONLY_DATA
+	virtual void GetCommonHLSL(FString& OutHLSL) override;
+	virtual bool UpgradeFunctionCall(FNiagaraFunctionSignature& FunctionSignature) override;
+#endif
+	virtual void GetVMExternalFunction(const FVMExternalFunctionBindingInfo& BindingInfo, void* InstanceData, FVMExternalFunction &OutFunc) override;
+	virtual bool CanExecuteOnTarget(ENiagaraSimTarget Target) const override { return true; }
+	virtual bool RequiresDistanceFieldData() const override;
+	virtual bool HasPreSimulateTick() const override { return true; }
+
+#if WITH_EDITORONLY_DATA
+	virtual bool AppendCompileHash(FNiagaraCompileHashVisitor* InVisitor) const override;
+	virtual void ModifyCompilationEnvironment(EShaderPlatform ShaderPlatform, struct FShaderCompilerEnvironment& OutEnvironment) const override;
+	virtual void GetParameterDefinitionHLSL(const FNiagaraDataInterfaceGPUParamInfo& ParamInfo, FString& OutHLSL) override;
+	virtual bool GetFunctionHLSL(const FNiagaraDataInterfaceGPUParamInfo& ParamInfo, const FNiagaraDataInterfaceGeneratedFunction& FunctionInfo, int FunctionInstanceIndex, FString& OutHLSL) override;
+#endif
+	virtual void BuildShaderParameters(FNiagaraShaderParametersBuilder& ShaderParametersBuilder) const override;
+	virtual void SetShaderParameters(const FNiagaraDataInterfaceSetShaderParametersContext& Context) const override;
+
+	virtual bool Equals(const UNiagaraDataInterface* Other) const override;
+protected:
+	virtual bool CopyToInternal(UNiagaraDataInterface* Destination) const override;
+public:
+#if WITH_NIAGARA_DEBUGGER
+	virtual void DrawDebugHud(UCanvas* Canvas, FNiagaraSystemInstance* SystemInstance, FString& VariableDataString, bool bVerbose) const override;
+#endif
+#if WITH_EDITOR
+	virtual void GetFeedback(UNiagaraSystem* Asset, UNiagaraComponent* Component, TArray<FNiagaraDataInterfaceError>& OutErrors,
+		TArray<FNiagaraDataInterfaceFeedback>& OutWarnings, TArray<FNiagaraDataInterfaceFeedback>& OutInfo) override;
+#endif
+	//~ UNiagaraDataInterface interface
+
+	UStaticMesh* GetStaticMesh(USceneComponent*& OutComponent, class FNiagaraSystemInstance* SystemInstance = nullptr);
+	void SetSourceComponentFromBlueprints(UStaticMeshComponent* ComponentToUse);
+	void SetDefaultMeshFromBlueprints(UStaticMesh* MeshToUse);
+
+protected:
+	// Bind/unbind delegates to release references to the source actor & component.
+	void UnbindSourceDelegates();
+	void BindSourceDelegates();
+
+	UFUNCTION()
+	void OnSourceEndPlay(AActor* InSource, EEndPlayReason::Type Reason);
+
+	// VM Vertex Sampling
+	void VMIsValidVertex(FVectorVMExternalFunctionContext& Context);
+	void VMRandomVertex(FVectorVMExternalFunctionContext& Context);
+	void VMGetVertexCount(FVectorVMExternalFunctionContext& Context);
+
+	template<typename TTransformHandler>
+	void VMGetVertex(FVectorVMExternalFunctionContext& Context);
+	template<typename TTransformHandler>
+	void VMGetVertexInterpolated(FVectorVMExternalFunctionContext& Context);
+	void VMGetVertexColor(FVectorVMExternalFunctionContext& Context);
+	void VMGetVertexUV(FVectorVMExternalFunctionContext& Context);
+
+	// VM Triangle Sampling
+	void VMIsValidTriangle(FVectorVMExternalFunctionContext& Context);
+	template<typename TRandomHelper>
+	void VMRandomTriangle(FVectorVMExternalFunctionContext& Context);
+	void VMGetTriangleCount(FVectorVMExternalFunctionContext& Context);
+
+	void VMIsValidFilteredTriangle(FVectorVMExternalFunctionContext& Context);
+	void VMRandomFilteredTriangle(FVectorVMExternalFunctionContext& Context);
+	void VMGetFilteredTriangleCount(FVectorVMExternalFunctionContext& Context);
+	void VMGetFilteredTriangleAt(FVectorVMExternalFunctionContext& Context);
+
+	void VMIsValidUnfilteredTriangle(FVectorVMExternalFunctionContext& Context);
+	void VMRandomUnfilteredTriangle(FVectorVMExternalFunctionContext& Context);
+	void VMGetUnfilteredTriangleCount(FVectorVMExternalFunctionContext& Context);
+	void VMGetUnfilteredTriangleAt(FVectorVMExternalFunctionContext& Context);
+
+	template<typename TTransformHandler>
+	void VMGetTriangle(FVectorVMExternalFunctionContext& Context);
+	template<typename TTransformHandler>
+	void VMGetTriangleInterpolated(FVectorVMExternalFunctionContext& Context);
+	void VMGetTriangleColor(FVectorVMExternalFunctionContext& Context);
+	void VMGetTriangleUV(FVectorVMExternalFunctionContext& Context);
+
+	void VMGetTriangleIndices(FVectorVMExternalFunctionContext& Context);
+
+	// Socket Functions
+	void VMGetSocketCount(FVectorVMExternalFunctionContext& Context);
+	void VMGetFilteredSocketCount(FVectorVMExternalFunctionContext& Context);
+	void VMGetUnfilteredSocketCount(FVectorVMExternalFunctionContext& Context);
+	void VMRandomSocket(FVectorVMExternalFunctionContext& Context);
+	void VMRandomFilteredSocket(FVectorVMExternalFunctionContext& Context);
+	void VMRandomUnfilteredSocket(FVectorVMExternalFunctionContext& Context);
+	template<typename TTransformHandler>
+	void VMGetSocketTransform(FVectorVMExternalFunctionContext& Context);
+	template<typename TTransformHandler>
+	void VMGetSocketTransformInterpolated(FVectorVMExternalFunctionContext& Context);
+	template<typename TTransformHandler>
+	void VMGetFilteredSocketTransform(FVectorVMExternalFunctionContext& Context);
+	template<typename TTransformHandler>
+	void VMGetFilteredSocketTransformInterpolated(FVectorVMExternalFunctionContext& Context);
+	template<typename TTransformHandler>
+	void VMGetUnfilteredSocketTransform(FVectorVMExternalFunctionContext& Context);
+	template<typename TTransformHandler>
+	void VMGetUnfilteredSocketTransformInterpolated(FVectorVMExternalFunctionContext& Context);
+	void VMGetFilteredSocket(FVectorVMExternalFunctionContext& Context);
+	void VMGetUnfilteredSocket(FVectorVMExternalFunctionContext& Context);
+
+	// Section functions
+	void VMIsValidSection(FVectorVMExternalFunctionContext& Context);
+	void VMGetSectionTriangleCount(FVectorVMExternalFunctionContext& Context);
+	template<typename TRandomHandler>
+	void VMRandomSectionTriangle(FVectorVMExternalFunctionContext& Context);
+	void VMGetSectionTriangleAt(FVectorVMExternalFunctionContext& Context);
+
+	void VMGetFilteredSectionAt(FVectorVMExternalFunctionContext& Context);
+	void VMGetUnfilteredSectionAt(FVectorVMExternalFunctionContext& Context);
+
+	void VMGetSectionCount(FVectorVMExternalFunctionContext& Context);
+	void VMGetFilteredSectionCount(FVectorVMExternalFunctionContext& Context);
+	void VMGetUnfilteredSectionCount(FVectorVMExternalFunctionContext& Context);
+
+	template<typename TRandomHandler>
+	void VMRandomSection(FVectorVMExternalFunctionContext& Context);
+	template<typename TRandomHandler>
+	void VMRandomFilteredSection(FVectorVMExternalFunctionContext& Context);
+	void VMRandomUnfilteredSection(FVectorVMExternalFunctionContext& Context);
+
+	// VM Misc Functions
+	void VMIsValid(FVectorVMExternalFunctionContext& Context);
+
+	void VMGetPreSkinnedLocalBounds(FVectorVMExternalFunctionContext& Context);
+	template<bool bLocalSpace>
+	void VMGetMeshBounds(FVectorVMExternalFunctionContext& Context);
+
+	void VMGetLocalToWorld(FVectorVMExternalFunctionContext& Context);
+	void VMGetLocalToWorldInverseTransposed(FVectorVMExternalFunctionContext& Context);
+	void VMGetWorldVelocity(FVectorVMExternalFunctionContext& Context);
+
+	// VM UV mapping functions
+	void VMGetTriangleCoordAtUV(FVectorVMExternalFunctionContext& Context);
+	void VMGetTriangleCoordInAabb(FVectorVMExternalFunctionContext& Context);
+	void VMBuildUvMapping(FVectorVMExternalFunctionContext& Context);;
+
+	// Deprecated VM Functions
+	template<typename TTransformHandler>
+	void VMGetVertexPosition_Deprecated(FVectorVMExternalFunctionContext& Context);
+	template<typename TTransformHandler>
+	void VMGetTriPosition_Deprecated(FVectorVMExternalFunctionContext& Context);
+	template<typename TTransformHandler>
+	void VMGetTriPositionAndVelocity_Deprecated(FVectorVMExternalFunctionContext& Context);
+	template<typename TTransformHandler>
+	void VMGetTriangleTangentBasis_Deprecated(FVectorVMExternalFunctionContext& Context);
+	template<typename TTransformHandler>
+	void VMGetTriangleNormal_Deprecated(FVectorVMExternalFunctionContext& Context);
+};
+
+class NIAGARA_API FNDI_StaticMesh_GeneratedData : public FNDI_GeneratedData
+{
+	FRWLock CachedUvMappingGuard;
+	TArray<TSharedPtr<FStaticMeshUvMapping>> CachedUvMapping;
+
+public:
+	FStaticMeshUvMappingHandle GetCachedUvMapping(TWeakObjectPtr<UStaticMesh>& InMeshObject, int32 InLodIndex, int32 InUvSetIndex, FMeshUvMappingUsage Usage, bool bNeedsDataImmediately);
+
+	virtual void Tick(ETickingGroup TickGroup, float DeltaSeconds) override;
+
+	static TypeHash GetTypeHash()
+	{
+		static const TypeHash Hash = ::GetTypeHash(TEXT("FNDI_StaticMesh_GeneratedData"));
+		return Hash;
+	}
+};
diff --git a/Engine/Plugins/FX/Niagara/Source/Niagara/Internal/NiagaraDataChannelCommon.h b/Engine/Plugins/FX/Niagara/Source/Niagara/Internal/NiagaraDataChannelCommon.h
new file mode 100644
--- /dev/null
+++ b/Engine/Plugins/FX/Niagara/Source/Niagara/Internal/NiagaraDataChannelCommon.h
@@ -0,0 +1,98 @@
+// Copyright Epic Games, Inc. All Rights Reserved.
+
+#pragma once
+
+#include "NiagaraDataSet.h"
+#include "NiagaraDataChannelCommon.generated.h"
+
+class UWorld;
+class UActorComponentc;
+class FNiagaraWorldManager;
+class FNiagaraDataBuffer;
+class FNiagaraSystemInstance;
+class UNiagaraDataChannel;
+class UNiagaraDataChannelHandler;
+class UNiagaraDataChannelDefinitions;
+class UNiagaraDataInterfaceDataChannelWrite;
+class UNiagaraDataInterfaceDataChannelRead;
+class UNiagaraDataChannelWriter;
+class UNiagaraDataChannelReader;
+struct FNiagaraDataChannelGameData;
+
+typedef TSharedPtr<FNiagaraDataChannelGameData> FNiagaraDataChannelGameDataPtr;
+
+/** Wrapper struct for FNames referencing Niagara Data Channels allowing a type customization. */
+USTRUCT(BlueprintType)
+struct FNiagaraDataChannelReference
+{
+	GENERATED_BODY()
+
+	UPROPERTY(EditAnywhere, Category = "Data Channel")
+	FName ChannelName = NAME_None;
+};
+
+/** A request to publish data into a Niagara Data Channel.  */
+struct FNiagaraDataChannelPublishRequest
+{
+	/** The buffer containing the data to be published. This can come from a data channel DI or can be the direct contents on a Niagara simulation. */
+	FNiagaraDataBuffer* Data = nullptr;
+
+	/** Game level data if this request comes from the game code. */
+	FNiagaraDataChannelGameDataPtr GameData;
+
+	/** If true, data in this request will be made visible to BP and C++ game code.*/
+	bool bVisibleToGame = false;
+
+	/** If true, data in this request will be made visible to Niagara CPU simulations. */
+	bool bVisibleToCPUSims = false;
+
+	/** If true, data in this request will be made visible to Niagara GPU simulations. */
+	bool bVisibleToGPUSims = false;
+
+	/** 
+	LWC Tile for the originator system of this request.
+	Allows us to convert from the Niagara Simulation space into LWC coordinates.
+	*/
+	FVector3f LwcTile = FVector3f::ZeroVector;
+
+	FNiagaraDataChannelPublishRequest() = default;
+	explicit FNiagaraDataChannelPublishRequest(FNiagaraDataBuffer* InData)
+		: Data(InData)
+	{
+		Data->AddReadRef();
+	}
+
+	explicit FNiagaraDataChannelPublishRequest(FNiagaraDataBuffer* InData, bool bInVisibleToGame, bool bInVisibleToCPUSims, bool bInVisibleToGPUSims, FVector3f InLwcTile)
+	: Data(InData), bVisibleToGame(bInVisibleToGame), bVisibleToCPUSims(bInVisibleToCPUSims), bVisibleToGPUSims(bInVisibleToGPUSims), LwcTile(InLwcTile)
+	{
+		Data->AddReadRef();
+	}
+
+	explicit FNiagaraDataChannelPublishRequest(const FNiagaraDataChannelPublishRequest& Other)
+	{
+		*this = Other;
+	}
+
+	FNiagaraDataChannelPublishRequest& operator=(const FNiagaraDataChannelPublishRequest& Other)
+	{
+		Data = Other.Data;
+		if(Data)
+		{
+			Data->AddReadRef();
+		}
+		GameData = Other.GameData;
+		bVisibleToGame = Other.bVisibleToGame;
+		bVisibleToCPUSims = Other.bVisibleToCPUSims;
+		bVisibleToGPUSims = Other.bVisibleToGPUSims;
+		LwcTile = Other.LwcTile;
+		return *this;
+	}
+
+	~FNiagaraDataChannelPublishRequest()
+	{
+		if(Data)
+		{
+			Data->ReleaseReadRef();
+		}
+	}
+};
diff --git a/Engine/Plugins/FX/Niagara/Source/Niagara/Internal/NiagaraDataChannelDefinitions.h b/Engine/Plugins/FX/Niagara/Source/Niagara/Internal/NiagaraDataChannelDefinitions.h
new file mode 100644
--- /dev/null
+++ b/Engine/Plugins/FX/Niagara/Source/Niagara/Internal/NiagaraDataChannelDefinitions.h
@@ -0,0 +1,51 @@
+// Copyright Epic Games, Inc. All Rights Reserved.
+
+#pragma once
+
+#include "NiagaraCommon.h"
+#include "NiagaraDataChannelCommon.h"
+#include "NiagaraDataChannelDefinitions.generated.h"
+
+class UNiagaraDataChannel;
+class UNiagaraDataChannelDefinitions;
+
+/** Asset class defining a set of data channels that can be used for communications between Niagara Emitters and Systems. */
+UCLASS()
+class NIAGARA_API UNiagaraDataChannelDefinitions : public UObject
+{
+	GENERATED_BODY()
+
+public:
+
+	//UObject Interface
+	virtual void PostInitProperties()override;
+	virtual void PostLoad()override;
+	virtual void BeginDestroy()override;
+#if WITH_EDITOR
+	virtual void PreEditChange(class FProperty* PropertyAboutToChange) override;
+	virtual void PostEditChangeProperty(FPropertyChangedEvent& PropertyChangedEvent) override;
+#endif
+	//UObject Interface END
+
+	UPROPERTY(EditAnywhere, Category = DataChannel, Instanced)
+	TArray<TObjectPtr<UNiagaraDataChannel>> DataChannels;
+
+	static const TArray<UNiagaraDataChannelDefinitions*>& GetDataChannelDefinitions(bool bRequired, bool bInformUser);
+	static const UNiagaraDataChannel* FindDataChannel(FName ChannelName);
+#if WITH_EDITORONLY_DATA
+	static void OnAssetCreated(UNiagaraDataChannelDefinitions* CreatedDef);
+	static void OnAssetDeleted(UNiagaraDataChannelDefinitions* DeletedDef);
+	static void OnAssetRegistryLoadComplete();
+#endif
+private:
+
+	static TArray<UNiagaraDataChannelDefinitions*> Definitions;
+
+#if WITH_EDITORONLY_DATA
+	FNiagaraSystemUpdateContext SysUpdateContext;
+
+	static inline bool bAssetRegistryScanBegun = false;
+	static inline bool bAssetRegistryScanComplete = false;
+	static FDelegateHandle AssetRegistryOnLoadCompleteHandle;
+#endif
+};
diff --git a/Engine/Plugins/FX/Niagara/Source/Niagara/Internal/NiagaraPrecompileContainer.h b/Engine/Plugins/FX/Niagara/Source/Niagara/Internal/NiagaraPrecompileContainer.h
new file mode 100644
--- /dev/null
+++ b/Engine/Plugins/FX/Niagara/Source/Niagara/Internal/NiagaraPrecompileContainer.h
@@ -0,0 +1,20 @@
+// Copyright Epic Games, Inc. All Rights Reserved.
+#pragma once
+
+#include "NiagaraSystem.h"
+#include "NiagaraScript.h"
+#include "NiagaraPrecompileContainer.generated.h"
+
+
+UCLASS()
+class NIAGARA_API UNiagaraPrecompileContainer : public UObject
+{
+	GENERATED_UCLASS_BODY()
+
+public:
+	UPROPERTY()
+	TArray<TObjectPtr<UNiagaraScript>> Scripts;
+
+	UPROPERTY()
+	TObjectPtr<UNiagaraSystem> System;
+};
\ No newline at end of file
diff --git a/Engine/Plugins/FX/Niagara/Source/Niagara/Private/DataInterface/NiagaraDataInterfaceActorComponent.cpp b/Engine/Plugins/FX/Niagara/Source/Niagara/Private/DataInterface/NiagaraDataInterfaceActorComponent.cpp
--- a/Engine/Plugins/FX/Niagara/Source/Niagara/Private/DataInterface/NiagaraDataInterfaceActorComponent.cpp
+++ b/Engine/Plugins/FX/Niagara/Source/Niagara/Private/DataInterface/NiagaraDataInterfaceActorComponent.cpp
@@ -1,6 +1,6 @@
 // Copyright Epic Games, Inc. All Rights Reserved.
 
-#include "NiagaraDataInterfaceActorComponent.h"
+#include "DataInterface/NiagaraDataInterfaceActorComponent.h"
 #include "GameFramework/Pawn.h"
 #include "NiagaraTypes.h"
 #include "NiagaraCustomVersion.h"
diff --git a/Engine/Plugins/FX/Niagara/Source/Niagara/Private/DataInterface/NiagaraDataInterfaceActorComponent.h b/Engine/Plugins/FX/Niagara/Source/Niagara/Private/DataInterface/NiagaraDataInterfaceActorComponent.h
deleted file mode 100644
diff --git a/Engine/Plugins/FX/Niagara/Source/Niagara/Private/DataInterface/NiagaraDataInterfaceDataChannelCommon.cpp b/Engine/Plugins/FX/Niagara/Source/Niagara/Private/DataInterface/NiagaraDataInterfaceDataChannelCommon.cpp
--- a/Engine/Plugins/FX/Niagara/Source/Niagara/Private/DataInterface/NiagaraDataInterfaceDataChannelCommon.cpp
+++ b/Engine/Plugins/FX/Niagara/Source/Niagara/Private/DataInterface/NiagaraDataInterfaceDataChannelCommon.cpp
@@ -1,6 +1,6 @@
 // Copyright Epic Games, Inc. All Rights Reserved.
 
-#include "NiagaraDataInterfaceDataChannelCommon.h"
+#include "DataInterface/NiagaraDataInterfaceDataChannelCommon.h"
 #include "NiagaraDataInterfaceUtilities.h"
 #include "NiagaraScript.h"
 
diff --git a/Engine/Plugins/FX/Niagara/Source/Niagara/Private/DataInterface/NiagaraDataInterfaceDataChannelCommon.h b/Engine/Plugins/FX/Niagara/Source/Niagara/Private/DataInterface/NiagaraDataInterfaceDataChannelCommon.h
deleted file mode 100644
diff --git a/Engine/Plugins/FX/Niagara/Source/Niagara/Private/DataInterface/NiagaraDataInterfaceDataChannelRead.cpp b/Engine/Plugins/FX/Niagara/Source/Niagara/Private/DataInterface/NiagaraDataInterfaceDataChannelRead.cpp
--- a/Engine/Plugins/FX/Niagara/Source/Niagara/Private/DataInterface/NiagaraDataInterfaceDataChannelRead.cpp
+++ b/Engine/Plugins/FX/Niagara/Source/Niagara/Private/DataInterface/NiagaraDataInterfaceDataChannelRead.cpp
@@ -1,6 +1,6 @@
 // Copyright Epic Games, Inc. All Rights Reserved.
 
-#include "NiagaraDataInterfaceDataChannelRead.h"
+#include "DataInterface/NiagaraDataInterfaceDataChannelRead.h"
 
 #include "ShaderCompilerCore.h"
 #include "ShaderParameterUtils.h"
diff --git a/Engine/Plugins/FX/Niagara/Source/Niagara/Private/DataInterface/NiagaraDataInterfaceDataChannelRead.h b/Engine/Plugins/FX/Niagara/Source/Niagara/Private/DataInterface/NiagaraDataInterfaceDataChannelRead.h
deleted file mode 100644
diff --git a/Engine/Plugins/FX/Niagara/Source/Niagara/Private/DataInterface/NiagaraDataInterfaceDataChannelSpawn.h b/Engine/Plugins/FX/Niagara/Source/Niagara/Private/DataInterface/NiagaraDataInterfaceDataChannelSpawn.h
--- a/Engine/Plugins/FX/Niagara/Source/Niagara/Private/DataInterface/NiagaraDataInterfaceDataChannelSpawn.h
+++ b/Engine/Plugins/FX/Niagara/Source/Niagara/Private/DataInterface/NiagaraDataInterfaceDataChannelSpawn.h
@@ -6,7 +6,7 @@ Data Channel Spawn Interface.
 Specialization of the Read DI that allows Emitters to Spawn directly from Data Channel contents.
 */
 
-#include "NiagaraDataInterfaceDataChannelRead.h"
+#include "DataInterface/NiagaraDataInterfaceDataChannelRead.h"
 #include "NiagaraDataInterfaceEmitterBinding.h"
 #include "NiagaraDataInterfaceDataChannelSpawn.generated.h"
 
diff --git a/Engine/Plugins/FX/Niagara/Source/Niagara/Private/DataInterface/NiagaraDataInterfaceDataChannelWrite.cpp b/Engine/Plugins/FX/Niagara/Source/Niagara/Private/DataInterface/NiagaraDataInterfaceDataChannelWrite.cpp
--- a/Engine/Plugins/FX/Niagara/Source/Niagara/Private/DataInterface/NiagaraDataInterfaceDataChannelWrite.cpp
+++ b/Engine/Plugins/FX/Niagara/Source/Niagara/Private/DataInterface/NiagaraDataInterfaceDataChannelWrite.cpp
@@ -1,6 +1,6 @@
 // Copyright Epic Games, Inc. All Rights Reserved.
 
-#include "NiagaraDataInterfaceDataChannelWrite.h"
+#include "DataInterface/NiagaraDataInterfaceDataChannelWrite.h"
 
 #include "CoreMinimal.h"
 #include "Stats/Stats.h"
diff --git a/Engine/Plugins/FX/Niagara/Source/Niagara/Private/DataInterface/NiagaraDataInterfaceDataChannelWrite.h b/Engine/Plugins/FX/Niagara/Source/Niagara/Private/DataInterface/NiagaraDataInterfaceDataChannelWrite.h
deleted file mode 100644
diff --git a/Engine/Plugins/FX/Niagara/Source/Niagara/Private/DataInterface/NiagaraDataInterfaceStaticMesh.cpp b/Engine/Plugins/FX/Niagara/Source/Niagara/Private/DataInterface/NiagaraDataInterfaceStaticMesh.cpp
--- a/Engine/Plugins/FX/Niagara/Source/Niagara/Private/DataInterface/NiagaraDataInterfaceStaticMesh.cpp
+++ b/Engine/Plugins/FX/Niagara/Source/Niagara/Private/DataInterface/NiagaraDataInterfaceStaticMesh.cpp
@@ -1,6 +1,6 @@
 // Copyright Epic Games, Inc. All Rights Reserved.
 
-#include "NiagaraDataInterfaceStaticMesh.h"
+#include "DataInterface/NiagaraDataInterfaceStaticMesh.h"
 #include "NiagaraEmitterInstance.h"
 #include "NiagaraComponent.h"
 #include "NiagaraDistanceFieldHelper.h"
diff --git a/Engine/Plugins/FX/Niagara/Source/Niagara/Private/DataInterface/NiagaraDataInterfaceStaticMesh.h b/Engine/Plugins/FX/Niagara/Source/Niagara/Private/DataInterface/NiagaraDataInterfaceStaticMesh.h
deleted file mode 100644
diff --git a/Engine/Plugins/FX/Niagara/Source/Niagara/Private/NiagaraDataChannelCommon.h b/Engine/Plugins/FX/Niagara/Source/Niagara/Private/NiagaraDataChannelCommon.h
deleted file mode 100644
diff --git a/Engine/Plugins/FX/Niagara/Source/Niagara/Private/NiagaraDataChannelDefinitions.h b/Engine/Plugins/FX/Niagara/Source/Niagara/Private/NiagaraDataChannelDefinitions.h
deleted file mode 100644
diff --git a/Engine/Plugins/FX/Niagara/Source/Niagara/Private/NiagaraPrecompileContainer.h b/Engine/Plugins/FX/Niagara/Source/Niagara/Private/NiagaraPrecompileContainer.h
deleted file mode 100644
diff --git a/Engine/Plugins/FX/Niagara/Source/NiagaraEditor/NiagaraEditor.Build.cs b/Engine/Plugins/FX/Niagara/Source/NiagaraEditor/NiagaraEditor.Build.cs
--- a/Engine/Plugins/FX/Niagara/Source/NiagaraEditor/NiagaraEditor.Build.cs
+++ b/Engine/Plugins/FX/Niagara/Source/NiagaraEditor/NiagaraEditor.Build.cs
@@ -15,7 +15,6 @@ public class NiagaraEditor : ModuleRules
 			"NiagaraEditor/Private/ViewModels",
 			"NiagaraEditor/Private/TypeEditorUtilities",
 			Path.Combine(GetModuleDirectory("GraphEditor"), "Private"),
-			Path.Combine(GetModuleDirectory("Niagara"), "Private"),
 			Path.Combine(GetModuleDirectory("PropertyEditor"), "Private"),
 			Path.Combine(GetModuleDirectory("Renderer"), "Private"),
 		});
diff --git a/Engine/Plugins/FX/Niagara/Source/NiagaraEditor/Private/NiagaraActions.cpp b/Engine/Plugins/FX/Niagara/Source/NiagaraEditor/Private/NiagaraActions.cpp
--- a/Engine/Plugins/FX/Niagara/Source/NiagaraEditor/Private/NiagaraActions.cpp
+++ b/Engine/Plugins/FX/Niagara/Source/NiagaraEditor/Private/NiagaraActions.cpp
@@ -24,8 +24,6 @@
 
 #include "NiagaraDataChannel.h"
 #include "NiagaraDataChannelDefinitions.h"
-#include "DataInterface/NiagaraDataInterfaceDataChannelRead.h"
-#include "DataInterface/NiagaraDataInterfaceDataChannelWrite.h"
 
 #include UE_INLINE_GENERATED_CPP_BY_NAME(NiagaraActions)
 
diff --git a/Engine/Plugins/FX/Niagara/Source/NiagaraEditor/Private/NiagaraDebugger.h b/Engine/Plugins/FX/Niagara/Source/NiagaraEditor/Private/NiagaraDebugger.h
--- a/Engine/Plugins/FX/Niagara/Source/NiagaraEditor/Private/NiagaraDebugger.h
+++ b/Engine/Plugins/FX/Niagara/Source/NiagaraEditor/Private/NiagaraDebugger.h
@@ -3,9 +3,6 @@
 #pragma once
 
 #include "IMessageContext.h"
-#include "NiagaraDebugHud.h"
-#include "NiagaraDebugHud.h"
-#include "NiagaraDebuggerCommon.h"
 #include "NiagaraDebuggerCommon.h"
 #include "NiagaraOutliner.h"
 #include "NiagaraWorldManager.h"
diff --git a/Engine/Plugins/FX/Niagara/Source/NiagaraEditorWidgets/NiagaraEditorWidgets.Build.cs b/Engine/Plugins/FX/Niagara/Source/NiagaraEditorWidgets/NiagaraEditorWidgets.Build.cs
--- a/Engine/Plugins/FX/Niagara/Source/NiagaraEditorWidgets/NiagaraEditorWidgets.Build.cs
+++ b/Engine/Plugins/FX/Niagara/Source/NiagaraEditorWidgets/NiagaraEditorWidgets.Build.cs
@@ -8,7 +8,6 @@ public class NiagaraEditorWidgets : ModuleRules
 	public NiagaraEditorWidgets(ReadOnlyTargetRules Target) : base(Target)
 	{
 		PrivateIncludePaths.AddRange(new string[] {
-			System.IO.Path.Combine(GetModuleDirectory("Niagara"), "Private"),
 			System.IO.Path.Combine(GetModuleDirectory("NiagaraEditor"), "Private"),
 			System.IO.Path.Combine(GetModuleDirectory("GraphEditor"), "Private"),
 		});
