Final uses of NIAGARA_USE_OPENVDB removed after integration.

diff --git a/Engine/Plugins/FX/Niagara/Source/Niagara/Niagara.Build.cs b/Engine/Plugins/FX/Niagara/Source/Niagara/Niagara.Build.cs
--- a/Engine/Plugins/FX/Niagara/Source/Niagara/Niagara.Build.cs
+++ b/Engine/Plugins/FX/Niagara/Source/Niagara/Niagara.Build.cs
@@ -91,20 +91,6 @@ public class Niagara : ModuleRules
 				"zlib",
 				"Boost"
 			);
-
-			PublicDefinitions.AddRange(
-				new string[]
-				{
-					"NIAGARA_USE_OPENVDB=1"
-				});
-		}
-		else
-		{
-			PublicDefinitions.AddRange(
-				new string[]
-				{
-					"NIAGARA_USE_OPENVDB=0"
-				});
 		}
 
 		PublicDefinitions.AddRange(
diff --git a/Engine/Plugins/FX/Niagara/Source/Niagara/Private/NiagaraDataInterfaceRenderTargetVolume.cpp b/Engine/Plugins/FX/Niagara/Source/Niagara/Private/NiagaraDataInterfaceRenderTargetVolume.cpp
--- a/Engine/Plugins/FX/Niagara/Source/Niagara/Private/NiagaraDataInterfaceRenderTargetVolume.cpp
+++ b/Engine/Plugins/FX/Niagara/Source/Niagara/Private/NiagaraDataInterfaceRenderTargetVolume.cpp
@@ -656,7 +656,7 @@ bool UNiagaraDataInterfaceRenderTargetVolume::SimCacheWriteFrame(UObject* Storag
 		{
 			if (NDIRenderTargetVolumeLocal::GSimCacheUseOpenVDB)
 			{
-#if NIAGARA_USE_OPENVDB
+#if UE_USE_OPENVDB
 				UVolumeCache* OpenVDBSimCacheData = CastChecked<UVolumeCache>(StorageObject);
 				
 				FString FullPath = OpenVDBSimCacheData->GetAssetPath(FrameIndex);
