Removing dependencies on NiagaraEditor Private files.

diff --git a/Engine/Plugins/FX/Niagara/Source/NiagaraEditor/Internal/NiagaraDeviceProfileViewModel.h b/Engine/Plugins/FX/Niagara/Source/NiagaraEditor/Internal/NiagaraDeviceProfileViewModel.h
new file mode 100644
--- /dev/null
+++ b/Engine/Plugins/FX/Niagara/Source/NiagaraEditor/Internal/NiagaraDeviceProfileViewModel.h
@@ -0,0 +1,14 @@
+// Copyright Epic Games, Inc. All Rights Reserved.
+
+#pragma once
+
+#include "Containers/Array.h"
+#include "Templates/SharedPointer.h"
+
+class UDeviceProfile;
+
+struct FNiagaraDeviceProfileViewModel
+{
+	UDeviceProfile* Profile;
+	TArray<TSharedPtr<FNiagaraDeviceProfileViewModel>> Children;
+};
diff --git a/Engine/Plugins/FX/Niagara/Source/NiagaraEditor/Internal/SNiagaraAssetPickerList.h b/Engine/Plugins/FX/Niagara/Source/NiagaraEditor/Internal/SNiagaraAssetPickerList.h
new file mode 100644
--- /dev/null
+++ b/Engine/Plugins/FX/Niagara/Source/NiagaraEditor/Internal/SNiagaraAssetPickerList.h
@@ -0,0 +1,154 @@
+// Copyright Epic Games, Inc. All Rights Reserved.
+
+#pragma once
+
+#include "Widgets/SItemSelector.h"
+#include "NiagaraActions.h"
+#include "AssetRegistry/AssetData.h"
+#include "Widgets/SNiagaraFilterBox.h"
+
+typedef SItemSelector<FText, FAssetData> SNiagaraAssetItemSelector;
+
+class FAssetThumbnailPool;
+
+struct FNiagaraAssetPickerListViewOptions
+{
+public:
+	FNiagaraAssetPickerListViewOptions()
+		: bCategorizeAssetsByAssetPath(false)
+		, bExpandTemplateAndLibraryAssets(false)
+		, bCategorizeLibraryAssets(false)
+		, bCategorizeUserDefinedCategory(false)
+		, bAddLibraryOnlyCheckbox(true)
+	{};
+
+	// Only showing template assets also implies categorizing by asset path; enforce via public setter.
+	void SetOnlyShowTemplatesAndCategorizeByAssetPath(bool bOnlyShowTemplatesAndCategorizeByAssetPath) {
+		bCategorizeAssetsByAssetPath = bOnlyShowTemplatesAndCategorizeByAssetPath;
+	};
+
+	void SetExpandTemplateAndLibraryAssets(bool bInExpandTemplateAndLibraryAssets) {bExpandTemplateAndLibraryAssets = bInExpandTemplateAndLibraryAssets; };
+	void SetCategorizeLibraryAssets(bool bInCategorizeLibraryAssets) {bCategorizeLibraryAssets = bInCategorizeLibraryAssets; };
+	void SetCategorizeUserDefinedCategory(bool bInCategorizeUserDefinedCategory) {bCategorizeUserDefinedCategory = bInCategorizeUserDefinedCategory; };
+	void SetAddLibraryOnlyCheckbox(bool bInAddLibraryOnlyCheckbox) { bAddLibraryOnlyCheckbox = bInAddLibraryOnlyCheckbox; };
+
+	bool GetCategorizeAssetsByAssetPath() const {return bCategorizeAssetsByAssetPath; };
+	bool GetExpandTemplateAndLibraryAssets() const {return bExpandTemplateAndLibraryAssets; };
+	bool GetCategorizeLibraryAssets() const { return bCategorizeLibraryAssets; };
+	bool GetCategorizeUserDefinedCategory() const {return bCategorizeUserDefinedCategory; };
+	bool GetAddLibraryOnlyCheckbox() const {return bAddLibraryOnlyCheckbox; };
+
+private:
+
+	// If true, categorize assets by their asset path, relative to the Niagara plugin directory, the project directory, and any other directory.
+	bool bCategorizeAssetsByAssetPath;
+
+	// If true, expand the thumbnails for assets marked as templates or library only.
+	bool bExpandTemplateAndLibraryAssets; 
+
+	// If true, categorize assets marked as library only.
+	bool bCategorizeLibraryAssets;
+
+	// If true, categorize assets that have a user defined category.
+	bool bCategorizeUserDefinedCategory;
+
+	// If true, we add a "Library only" checkbox at the top of the list
+	bool bAddLibraryOnlyCheckbox;
+};
+
+class NIAGARAEDITOR_API SNiagaraAssetPickerList : public SCompoundWidget
+{
+public:
+	DECLARE_DELEGATE_OneParam(FOnTemplateAssetActivated, const FAssetData&);
+	DECLARE_DELEGATE_RetVal_OneParam(bool, FOnDoesAssetPassCustomFilter, const FAssetData&)
+	
+public:
+	SLATE_BEGIN_ARGS(SNiagaraAssetPickerList) 
+		: _bAllowMultiSelect(false)
+		, _TabOptions(SNiagaraTemplateTabBox::FNiagaraTemplateTabOptions())
+		, _bLibraryOnly(true)
+		, _ClickActivateMode(EItemSelectorClickActivateMode::DoubleClick)
+	{}
+		// Callback for when an asset is activated.
+		SLATE_EVENT(FOnTemplateAssetActivated, OnTemplateAssetActivated);
+
+		// Whether to allow multi-selecting assets.
+		SLATE_ARGUMENT(bool, bAllowMultiSelect);
+
+		/** An optional delegate which is called to check if an asset should be filtered out by external code. Return false to exclude the asset from the view. */
+		SLATE_EVENT(FOnDoesAssetPassCustomFilter, OnDoesAssetPassCustomFilter);
+
+		/** An optional array of delegates to refresh the item selector view when executed. */
+		SLATE_ARGUMENT(TArray<FRefreshItemSelectorDelegate*>, RefreshItemSelectorDelegates);
+
+		// Arguments describing how to display, filter and categorize items of the asset picker.
+		SLATE_ARGUMENT(FNiagaraAssetPickerListViewOptions, ViewOptions);
+
+		/** Tab options that indicate which tabs are available. If a single tab is specified, that tab will serve as a filter and not display a tab widget */
+		SLATE_ARGUMENT(SNiagaraTemplateTabBox::FNiagaraTemplateTabOptions, TabOptions);
+
+		/** WWhether the Library Only checkbox is ticked initially */
+		SLATE_ARGUMENT(bool, bLibraryOnly)
+
+		/** Whether or not a single click activates an item. */
+		SLATE_ARGUMENT(EItemSelectorClickActivateMode, ClickActivateMode);
+	SLATE_END_ARGS();
+
+	void Construct(const FArguments& InArgs, UClass* AssetClass);
+
+	virtual void Tick(const FGeometry& AllottedGeometry, const double InCurrentTime, const float InDeltaTime) override;
+
+	TArray<FAssetData> GetSelectedAssets() const;
+
+	void RefreshAll() const;
+
+	void ExpandTree();
+
+	TSharedRef<SWidget> GetSearchBox() const;
+
+private:
+	TArray<FAssetData> GetAssetDataForSelector(UClass* AssetClass);
+
+	TArray<FText> OnGetCategoriesForItem(const FAssetData& Item);
+
+	bool OnCompareCategoriesForEquality(const FText& CategoryA, const FText& CategoryB) const;
+
+	bool OnCompareCategoriesForSorting(const FText& CategoryA, const FText& CategoryB) const;
+
+	bool OnCompareItemsForSorting(const FAssetData& ItemA, const FAssetData& ItemB) const;
+
+	bool OnDoesItemMatchFilterText(const FText& FilterText, const FAssetData& Item);
+	
+	TSharedRef<SWidget> OnGenerateWidgetForCategory(const FText& Category);
+
+	TSharedRef<SWidget> OnGenerateWidgetForItem(const FAssetData& Item);
+
+	void OnItemActivated(const FAssetData& Item);
+
+	bool DoesItemPassCustomFilter(const FAssetData& Item);
+	
+	FText GetFilterText() const;
+
+	void TriggerRefresh(const TMap<EScriptSource, bool>& SourceState);
+	void TriggerRefreshFromTabs(ENiagaraScriptTemplateSpecification Tab);
+
+	void LibraryCheckBoxStateChanged(bool bInLibraryOnly);
+	bool GetLibraryCheckBoxState() const;
+
+private:
+	TSharedPtr<SNiagaraFilterBox> FilterBox;
+	TSharedPtr<SNiagaraAssetItemSelector> ItemSelector;
+	
+	static FText NiagaraPluginCategory;
+	static FText ProjectCategory;
+	static FText LibraryCategory;
+	static FText NonLibraryCategory;
+	static FText UncategorizedCategory;
+	bool bLibraryOnly = true;
+	SNiagaraTemplateTabBox::FNiagaraTemplateTabOptions TabOptions;
+	TSharedPtr<FAssetThumbnailPool> AssetThumbnailPool;
+	FOnTemplateAssetActivated OnTemplateAssetActivated;
+	FOnDoesAssetPassCustomFilter OnDoesAssetPassCustomFilter;
+	FNiagaraAssetPickerListViewOptions ViewOptions;
+	FOnDoesAssetPassCustomFilter CustomFilter;
+};
diff --git a/Engine/Plugins/FX/Niagara/Source/NiagaraEditor/Private/Customizations/NiagaraPlatformSetCustomization.h b/Engine/Plugins/FX/Niagara/Source/NiagaraEditor/Private/Customizations/NiagaraPlatformSetCustomization.h
--- a/Engine/Plugins/FX/Niagara/Source/NiagaraEditor/Private/Customizations/NiagaraPlatformSetCustomization.h
+++ b/Engine/Plugins/FX/Niagara/Source/NiagaraEditor/Private/Customizations/NiagaraPlatformSetCustomization.h
@@ -4,6 +4,7 @@
 
 #include "CoreMinimal.h"
 #include "IPropertyTypeCustomization.h"
+#include "NiagaraDeviceProfileViewModel.h"
 #include "NiagaraPlatformSet.h"
 #include "Layout/Visibility.h"
 #include "Widgets/Views/STreeView.h"
@@ -17,12 +18,6 @@ class SWrapBox;
 
 enum class ECheckBoxState : uint8;
 
-struct FNiagaraDeviceProfileViewModel
-{
-	class UDeviceProfile* Profile;
-	TArray<TSharedPtr<FNiagaraDeviceProfileViewModel>> Children;
-};
-
 class FNiagaraPlatformSetCustomization : public IPropertyTypeCustomization
 {
 public:
diff --git a/Engine/Plugins/FX/Niagara/Source/NiagaraEditor/Private/SNiagaraAssetPickerList.h b/Engine/Plugins/FX/Niagara/Source/NiagaraEditor/Private/SNiagaraAssetPickerList.h
deleted file mode 100644
diff --git a/Engine/Plugins/FX/Niagara/Source/NiagaraEditorWidgets/NiagaraEditorWidgets.Build.cs b/Engine/Plugins/FX/Niagara/Source/NiagaraEditorWidgets/NiagaraEditorWidgets.Build.cs
--- a/Engine/Plugins/FX/Niagara/Source/NiagaraEditorWidgets/NiagaraEditorWidgets.Build.cs
+++ b/Engine/Plugins/FX/Niagara/Source/NiagaraEditorWidgets/NiagaraEditorWidgets.Build.cs
@@ -8,7 +8,6 @@ public class NiagaraEditorWidgets : ModuleRules
 	public NiagaraEditorWidgets(ReadOnlyTargetRules Target) : base(Target)
 	{
 		PrivateIncludePaths.AddRange(new string[] {
-			System.IO.Path.Combine(GetModuleDirectory("NiagaraEditor"), "Private"),
 			System.IO.Path.Combine(GetModuleDirectory("GraphEditor"), "Private"),
 		});
 
diff --git a/Engine/Plugins/FX/Niagara/Source/NiagaraEditorWidgets/Private/SNiagaraScalabilityPreviewSettings.cpp b/Engine/Plugins/FX/Niagara/Source/NiagaraEditorWidgets/Private/SNiagaraScalabilityPreviewSettings.cpp
--- a/Engine/Plugins/FX/Niagara/Source/NiagaraEditorWidgets/Private/SNiagaraScalabilityPreviewSettings.cpp
+++ b/Engine/Plugins/FX/Niagara/Source/NiagaraEditorWidgets/Private/SNiagaraScalabilityPreviewSettings.cpp
@@ -1,12 +1,12 @@
 // Copyright Epic Games, Inc. All Rights Reserved.
 
 #include "SNiagaraScalabilityPreviewSettings.h"
-#include "Customizations/NiagaraPlatformSetCustomization.h"
 #include "DeviceProfiles/DeviceProfile.h"
 #include "DeviceProfiles/DeviceProfileManager.h"
 #include "IDetailTreeNode.h"
 #include "IPropertyRowGenerator.h"
 #include "ISinglePropertyView.h"
+#include "NiagaraDeviceProfileViewModel.h"
 #include "NiagaraEditorStyle.h"
 #include "NiagaraSettings.h"
 #include "PlatformInfo.h"
