Niagara - When opening the emitter dialog, preload the cooked emitter assets so that filtering works correctly.

diff --git a/Engine/Plugins/FX/Niagara/Source/NiagaraEditor/Private/SNewEmitterDialog.cpp b/Engine/Plugins/FX/Niagara/Source/NiagaraEditor/Private/SNewEmitterDialog.cpp
--- a/Engine/Plugins/FX/Niagara/Source/NiagaraEditor/Private/SNewEmitterDialog.cpp
+++ b/Engine/Plugins/FX/Niagara/Source/NiagaraEditor/Private/SNewEmitterDialog.cpp
@@ -1,6 +1,7 @@
 // Copyright Epic Games, Inc. All Rights Reserved.
 
 #include "SNewEmitterDialog.h"
+#include "NiagaraEditorModule.h"
 #include "NiagaraEditorStyle.h"
 #include "NiagaraEmitter.h"
 #include "SNiagaraAssetPickerList.h"
@@ -28,6 +29,7 @@ void SNewEmitterDialog::Construct(const FArguments& InArgs)
 	TabOptions.ChangeTabState(ENiagaraScriptTemplateSpecification::None, true);
 	TabOptions.ChangeTabState(ENiagaraScriptTemplateSpecification::Behavior, true);
 
+	FNiagaraEditorModule::Get().PreloadSelectablePluginAssetsByClass(UNiagaraEmitter::StaticClass());
 	SAssignNew(NewAssetPicker, SNiagaraAssetPickerList, UNiagaraEmitter::StaticClass())
 		.OnTemplateAssetActivated(this, &SNewEmitterDialog::ConfirmSelection)
 		.ViewOptions(DisplayAllViewOptions)
