Preventing new NiagaraDataChannelDefinitions assets from being created if the the data channels CVar is disabled.

diff --git a/Engine/Plugins/FX/Niagara/Source/NiagaraEditor/Private/NiagaraDataChannelDefinitionsFactoryNew.cpp b/Engine/Plugins/FX/Niagara/Source/NiagaraEditor/Private/NiagaraDataChannelDefinitionsFactoryNew.cpp
--- a/Engine/Plugins/FX/Niagara/Source/NiagaraEditor/Private/NiagaraDataChannelDefinitionsFactoryNew.cpp
+++ b/Engine/Plugins/FX/Niagara/Source/NiagaraEditor/Private/NiagaraDataChannelDefinitionsFactoryNew.cpp
@@ -22,3 +22,8 @@ UObject* UNiagaraDataChannelDefinitionsFactoryNew::FactoryCreateNew(UClass* Clas
 
 	return NewDefinitions;
 }
+
+bool UNiagaraDataChannelDefinitionsFactoryNew::CanCreateNew()const
+{
+	return INiagaraModule::DataChannelsEnabled();
+}
diff --git a/Engine/Plugins/FX/Niagara/Source/NiagaraEditor/Public/NiagaraDataChannelDefinitionsFactoryNew.h b/Engine/Plugins/FX/Niagara/Source/NiagaraEditor/Public/NiagaraDataChannelDefinitionsFactoryNew.h
--- a/Engine/Plugins/FX/Niagara/Source/NiagaraEditor/Public/NiagaraDataChannelDefinitionsFactoryNew.h
+++ b/Engine/Plugins/FX/Niagara/Source/NiagaraEditor/Public/NiagaraDataChannelDefinitionsFactoryNew.h
@@ -14,5 +14,7 @@ class UNiagaraDataChannelDefinitionsFactoryNew : public UFactory
 
 	//~ Begin UFactory Interface
 	virtual UObject* FactoryCreateNew(UClass* Class, UObject* InParent, FName Name, EObjectFlags Flags, UObject* Context, FFeedbackContext* Warn) override;
+
+	virtual bool CanCreateNew()const override;
 	//~ Begin UFactory Interface
 };
