Replacing manual Add/Release of FNiagaraDataBuffers* we pass around with a ref counted ptr.

diff --git a/Engine/Plugins/FX/Niagara/Source/Niagara/Classes/NiagaraComputeExecutionContext.h b/Engine/Plugins/FX/Niagara/Source/Niagara/Classes/NiagaraComputeExecutionContext.h
--- a/Engine/Plugins/FX/Niagara/Source/Niagara/Classes/NiagaraComputeExecutionContext.h
+++ b/Engine/Plugins/FX/Niagara/Source/Niagara/Classes/NiagaraComputeExecutionContext.h
@@ -135,10 +135,10 @@ public:
 	TArray<FNiagaraDataInterfaceProxy*> DataInterfaceProxies;
 
 	// Most current buffer that can be used for rendering.
-	FNiagaraDataBuffer* DataToRender = nullptr;
+	FNiagaraDataBufferRef DataToRender = nullptr;
 
 	// Optional buffer which can be used to render translucent data with no latency (i.e. this frames data)
-	FNiagaraDataBuffer* TranslucentDataToRender = nullptr;
+	FNiagaraDataBufferRef TranslucentDataToRender = nullptr;
 
 	// Game thread spawn info will be sent to the render thread inside FNiagaraComputeInstanceData
 	FNiagaraGpuSpawnInfo GpuSpawnInfo_GT;
diff --git a/Engine/Plugins/FX/Niagara/Source/Niagara/Classes/NiagaraDataSet.h b/Engine/Plugins/FX/Niagara/Source/Niagara/Classes/NiagaraDataSet.h
--- a/Engine/Plugins/FX/Niagara/Source/Niagara/Classes/NiagaraDataSet.h
+++ b/Engine/Plugins/FX/Niagara/Source/Niagara/Classes/NiagaraDataSet.h
@@ -34,13 +34,13 @@ public:
 	FORCEINLINE bool IsBeingRead()const { return ReadRefCount.Load() > 0; }
 	FORCEINLINE bool IsBeingWritten()const { return ReadRefCount.Load() == INDEX_NONE; }
 
-	FORCEINLINE void AddReadRef()
+	FORCEINLINE void AddRef()
 	{
 		check(!IsBeingWritten());
 		ReadRefCount++;
 	}
 
-	FORCEINLINE void ReleaseReadRef()
+	FORCEINLINE void Release()
 	{
 		check(IsBeingRead());
 		ReadRefCount--;
@@ -97,7 +97,7 @@ public:
 	void CopyTo(FNiagaraDataBuffer& DestBuffer, int32 SrcStartIdx, int32 DestStartIdx, int32 NumInstances)const;
 	void CopyToUnrelated(FNiagaraDataBuffer& DestBuffer, int32 SrcStartIdx, int32 DestStartIdx, int32 NumInstances)const;
 	void GPUCopyFrom(const float* GPUReadBackFloat, const int* GPUReadBackInt, const FFloat16* GPUReadBackHalf, int32 StartIdx, int32 NumInstances, uint32 InSrcFloatStride, uint32 InSrcIntStride, uint32 InSrcHalfStride);
-	void PushCPUBuffersToGPU(const TArray<FNiagaraDataBuffer*>& SourceBuffers, bool bReleaseRef, FRHICommandList& RHICmdList, ERHIFeatureLevel::Type FeatureLevel, const TCHAR* DebugSimName);
+	void PushCPUBuffersToGPU(const TArray<FNiagaraDataBufferRef>& SourceBuffers, bool bReleaseRef, FRHICommandList& RHICmdList, ERHIFeatureLevel::Type FeatureLevel, const TCHAR* DebugSimName);
 	void Dump(int32 StartIndex, int32 NumInstances, const FString& Label, const FName& SortParameterKey = FName())const;
 
 	FORCEINLINE TArrayView<uint8 const* RESTRICT const> GetRegisterTable() const { return TArrayView<uint8 const* RESTRICT const>(RegisterTable); }
diff --git a/Engine/Plugins/FX/Niagara/Source/Niagara/Classes/NiagaraScriptExecutionContext.h b/Engine/Plugins/FX/Niagara/Source/Niagara/Classes/NiagaraScriptExecutionContext.h
--- a/Engine/Plugins/FX/Niagara/Source/Niagara/Classes/NiagaraScriptExecutionContext.h
+++ b/Engine/Plugins/FX/Niagara/Source/Niagara/Classes/NiagaraScriptExecutionContext.h
@@ -41,20 +41,12 @@ struct FNiagaraEventHandlingInfo
 
 	void SetEventData(FNiagaraDataBuffer* InEventData)
 	{
-		if (EventData)
-		{
-			EventData->ReleaseReadRef();
-		}
 		EventData = InEventData;
-		if (EventData)
-		{
-			EventData->AddReadRef();
-		}
 	}
 
 	TArray<int32, TInlineAllocator<16>> SpawnCounts;
 	int32 TotalSpawnCount;
-	FNiagaraDataBuffer* EventData;
+	FNiagaraDataBufferRef EventData;
 	FGuid SourceEmitterGuid;
 	FName SourceEmitterName;
 };
@@ -65,7 +57,6 @@ struct FNiagaraDataSetExecutionInfo
 	FNiagaraDataSetExecutionInfo()
 		: DataSet(nullptr)
 		, Input(nullptr)
-		, Output(nullptr)
 		, StartInstance(0)
 		, bUpdateInstanceCount(false)
 	{
@@ -73,28 +64,15 @@ struct FNiagaraDataSetExecutionInfo
 	}
 
 
-	FORCEINLINE void Init(FNiagaraDataSet* InDataSet, FNiagaraDataBuffer* InInput, FNiagaraDataBuffer* InOutput, int32 InStartInstance, bool bInUpdateInstanceCount)
+	FORCEINLINE void Init(FNiagaraDataSet* InDataSet, FNiagaraDataBuffer* InInput, int32 InStartInstance, bool bInUpdateInstanceCount)
 	{
-		if (Input)
-		{
-			Input->ReleaseReadRef();
-		}
-
 		DataSet = InDataSet;
 		Input = InInput;
-		Output = InOutput;
 		StartInstance = InStartInstance;
 		bUpdateInstanceCount = bInUpdateInstanceCount;
 
 		check(DataSet);
 		check(Input == nullptr || DataSet == Input->GetOwner());
-		check(Output == nullptr || DataSet == Output->GetOwner());
-
-		if (Input)
-		{
-			Input->AddReadRef();
-		}
-		check(Output == nullptr || Output->IsBeingWritten());
 	}
 	
 	~FNiagaraDataSetExecutionInfo()
@@ -104,21 +82,16 @@ struct FNiagaraDataSetExecutionInfo
 
 	FORCEINLINE void Reset()
 	{
-		if (Input)
-		{
-			Input->ReleaseReadRef();
-		}
-
 		DataSet = nullptr;
 		Input = nullptr;
-		Output = nullptr;
 		StartInstance = INDEX_NONE;
 		bUpdateInstanceCount = false;
 	}
 
+	FNiagaraDataBuffer* GetOutput()const { return DataSet->GetDestinationData(); }
+
 	FNiagaraDataSet* DataSet;
-	FNiagaraDataBuffer* Input;
-	FNiagaraDataBuffer* Output;
+	FNiagaraDataBufferRef Input;
 	int32 StartInstance;
 	bool bUpdateInstanceCount;
 };
diff --git a/Engine/Plugins/FX/Niagara/Source/Niagara/Internal/DataInterface/NiagaraDataInterfaceDataChannelRead.h b/Engine/Plugins/FX/Niagara/Source/Niagara/Internal/DataInterface/NiagaraDataInterfaceDataChannelRead.h
--- a/Engine/Plugins/FX/Niagara/Source/Niagara/Internal/DataInterface/NiagaraDataInterfaceDataChannelRead.h
+++ b/Engine/Plugins/FX/Niagara/Source/Niagara/Internal/DataInterface/NiagaraDataInterfaceDataChannelRead.h
@@ -126,7 +126,7 @@ struct FNDIDataChannelReadInstanceData
 	TWeakObjectPtr<UNiagaraDataChannelHandler> DataChannel;
 
 	/** Pointer to the cpu simulation buffer from the DataChannel if we're reading from there. */
-	FNiagaraDataBuffer* ExternalCPU = nullptr;
+	FNiagaraDataBufferRef ExternalCPU = nullptr;
 	
 	/** The GPU data set from the DataChannel. We must grab the whole set rather than a buffer here as we don't know which buffer will be contain this frame's channel data. 
 	*	TODO: We'll likely want to improve this GPU sim flow soon. Possibly grab an RT proxy object from the data channel handler and pass that over to the GPU instead of a direct dataset like this.
diff --git a/Engine/Plugins/FX/Niagara/Source/Niagara/Internal/NiagaraDataChannelCommon.h b/Engine/Plugins/FX/Niagara/Source/Niagara/Internal/NiagaraDataChannelCommon.h
--- a/Engine/Plugins/FX/Niagara/Source/Niagara/Internal/NiagaraDataChannelCommon.h
+++ b/Engine/Plugins/FX/Niagara/Source/Niagara/Internal/NiagaraDataChannelCommon.h
@@ -35,7 +35,7 @@ struct FNiagaraDataChannelReference
 struct FNiagaraDataChannelPublishRequest
 {
 	/** The buffer containing the data to be published. This can come from a data channel DI or can be the direct contents on a Niagara simulation. */
-	FNiagaraDataBuffer* Data = nullptr;
+	FNiagaraDataBufferRef Data;
 
 	/** Game level data if this request comes from the game code. */
 	FNiagaraDataChannelGameDataPtr GameData;
@@ -59,40 +59,10 @@ struct FNiagaraDataChannelPublishRequest
 	explicit FNiagaraDataChannelPublishRequest(FNiagaraDataBuffer* InData)
 		: Data(InData)
 	{
-		Data->AddReadRef();
 	}
 
 	explicit FNiagaraDataChannelPublishRequest(FNiagaraDataBuffer* InData, bool bInVisibleToGame, bool bInVisibleToCPUSims, bool bInVisibleToGPUSims, FVector3f InLwcTile)
 	: Data(InData), bVisibleToGame(bInVisibleToGame), bVisibleToCPUSims(bInVisibleToCPUSims), bVisibleToGPUSims(bInVisibleToGPUSims), LwcTile(InLwcTile)
 	{
-		Data->AddReadRef();
-	}
-
-	explicit FNiagaraDataChannelPublishRequest(const FNiagaraDataChannelPublishRequest& Other)
-	{
-		*this = Other;
-	}
-
-	FNiagaraDataChannelPublishRequest& operator=(const FNiagaraDataChannelPublishRequest& Other)
-	{
-		Data = Other.Data;
-		if(Data)
-		{
-			Data->AddReadRef();
-		}
-		GameData = Other.GameData;
-		bVisibleToGame = Other.bVisibleToGame;
-		bVisibleToCPUSims = Other.bVisibleToCPUSims;
-		bVisibleToGPUSims = Other.bVisibleToGPUSims;
-		LwcTile = Other.LwcTile;
-		return *this;
-	}
-
-	~FNiagaraDataChannelPublishRequest()
-	{
-		if(Data)
-		{
-			Data->ReleaseReadRef();
-		}
 	}
 };
diff --git a/Engine/Plugins/FX/Niagara/Source/Niagara/Private/DataInterface/NiagaraDataInterfaceDataChannelRead.cpp b/Engine/Plugins/FX/Niagara/Source/Niagara/Private/DataInterface/NiagaraDataInterfaceDataChannelRead.cpp
--- a/Engine/Plugins/FX/Niagara/Source/Niagara/Private/DataInterface/NiagaraDataInterfaceDataChannelRead.cpp
+++ b/Engine/Plugins/FX/Niagara/Source/Niagara/Private/DataInterface/NiagaraDataInterfaceDataChannelRead.cpp
@@ -94,13 +94,7 @@ struct FNDIDataChannelReadInstanceData_RT
 
 FNDIDataChannelReadInstanceData::~FNDIDataChannelReadInstanceData()
 {
-	if (ExternalCPU)
-	{
-		ExternalCPU->ReleaseReadRef();
-		ExternalCPU = nullptr;
-	}
 
-	ExternalGPU = nullptr;
 }
 
 FNiagaraDataBuffer* FNDIDataChannelReadInstanceData::GetReadBufferCPU()
@@ -235,14 +229,14 @@ bool FNDIDataChannelReadInstanceData::Tick(UNiagaraDataInterfaceDataChannelRead*
 			if (CompiledData.UsedByCPU())
 			{
 				//TODO: Automatically modify tick group if we have DIs that require current frame info?
-				DataChannelPtr->GetData(Instance, ExternalCPU, Interface->bReadCurrentFrame == false);
+				ExternalCPU = DataChannelPtr->GetData(Instance, Interface->bReadCurrentFrame == false);
 				CPUSourceDataCompiledData = &ExternalCPU->GetOwner()->GetCompiledData();
 			}
 
 			if (CompiledData.UsedByGPU())
 			{
 				//TODO: Detect in post compile if we need GPU/CPU or both.
-				DataChannelPtr->GetDataGPU(Instance, ExternalGPU);
+				ExternalGPU = DataChannelPtr->GetDataGPU(Instance);
 				GPUSourceDataCompiledData = &ExternalGPU->GetCompiledData();
 			}
 
@@ -305,8 +299,6 @@ bool FNDIDataChannelReadInstanceData::PostTick(UNiagaraDataInterfaceDataChannelR
 {
 	if (ExternalCPU)
 	{
-		//TODO: Create ref object to do this for us and make Add/Release ref private.
-		ExternalCPU->ReleaseReadRef();
 		ExternalCPU = nullptr;
 	}
 
diff --git a/Engine/Plugins/FX/Niagara/Source/Niagara/Private/DataInterface/NiagaraDataInterfaceDataChannelWrite.cpp b/Engine/Plugins/FX/Niagara/Source/Niagara/Private/DataInterface/NiagaraDataInterfaceDataChannelWrite.cpp
--- a/Engine/Plugins/FX/Niagara/Source/Niagara/Private/DataInterface/NiagaraDataInterfaceDataChannelWrite.cpp
+++ b/Engine/Plugins/FX/Niagara/Source/Niagara/Private/DataInterface/NiagaraDataInterfaceDataChannelWrite.cpp
@@ -168,17 +168,7 @@ struct FNDIDataChannelWriteInstanceData
 
 		if(FNiagaraDataBuffer* CurrBuff = Data->GetCurrentData())
 		{	
-			if (bNeedsLastFrameData)
-			{
-				CurrBuff->AddReadRef();//Ensure we keep the previous data when grabbing a new dest buffer.
-			}
-
 			FNiagaraDataBuffer& DestBuff = Data->BeginSimulate(true);
-
-			if (bNeedsLastFrameData && CurrBuff)
-			{
-				CurrBuff->ReleaseReadRef();
-			}
 		}
 
 		if (Interface->AllocationMode == ENiagaraDataChannelAllocationMode::Static)
diff --git a/Engine/Plugins/FX/Niagara/Source/Niagara/Private/NiagaraComputeExecutionContext.cpp b/Engine/Plugins/FX/Niagara/Source/Niagara/Private/NiagaraComputeExecutionContext.cpp
--- a/Engine/Plugins/FX/Niagara/Source/Niagara/Private/NiagaraComputeExecutionContext.cpp
+++ b/Engine/Plugins/FX/Niagara/Source/Niagara/Private/NiagaraComputeExecutionContext.cpp
@@ -225,40 +225,19 @@ void FNiagaraComputeExecutionContext::ResetInternal(FNiagaraGpuComputeDispatchIn
 
 void FNiagaraComputeExecutionContext::SetDataToRender(FNiagaraDataBuffer* InDataToRender)
 {
-	if (DataToRender)
-	{
-		DataToRender->ReleaseReadRef();
-	}
-
 	DataToRender = InDataToRender;
 
-	if (DataToRender)
-	{
-		DataToRender->AddReadRef();
-	}
-
 	// This call the DataToRender should be equal to the TranslucentDataToRender so we can release the read ref
 	if (TranslucentDataToRender)
 	{
 		ensure((DataToRender == nullptr) || (DataToRender == TranslucentDataToRender));
-		TranslucentDataToRender->ReleaseReadRef();
 		TranslucentDataToRender = nullptr;
 	}
 }
 
 void FNiagaraComputeExecutionContext::SetTranslucentDataToRender(FNiagaraDataBuffer* InTranslucentDataToRender)
 {
-	if (TranslucentDataToRender)
-	{
-		TranslucentDataToRender->ReleaseReadRef();
-	}
-
 	TranslucentDataToRender = InTranslucentDataToRender;
-
-	if (TranslucentDataToRender)
-	{
-		TranslucentDataToRender->AddReadRef();
-	}
 }
 
 bool FNiagaraComputeInstanceData::IsOutputStage(FNiagaraDataInterfaceProxy* DIProxy, uint32 SimulationStageIndex) const
diff --git a/Engine/Plugins/FX/Niagara/Source/Niagara/Private/NiagaraDataChannel_Global.cpp b/Engine/Plugins/FX/Niagara/Source/Niagara/Private/NiagaraDataChannel_Global.cpp
--- a/Engine/Plugins/FX/Niagara/Source/Niagara/Private/NiagaraDataChannel_Global.cpp
+++ b/Engine/Plugins/FX/Niagara/Source/Niagara/Private/NiagaraDataChannel_Global.cpp
@@ -44,10 +44,10 @@ public:
 	{
 	}
 
-	void EndFrame(FNiagaraGpuComputeDispatchInterface* DispathInterface, FRHICommandListImmediate& CmdList, const TArray<FNiagaraDataBuffer*>& BuffersForGPU)
+	void EndFrame(FNiagaraGpuComputeDispatchInterface* DispathInterface, FRHICommandListImmediate& CmdList, const TArray<FNiagaraDataBufferRef>& BuffersForGPU)
 	{
 		uint32 NumInstance = 0;
-		for(FNiagaraDataBuffer* Buffer : BuffersForGPU)
+		for(auto& Buffer : BuffersForGPU)
 		{
 			NumInstance += Buffer->GetNumInstances();
 		}		
@@ -133,12 +133,7 @@ void UNiagaraDataChannelHandler_Global::Tick(float DeltaSeconds, ETickingGroup T
 		//TODO: Add reference struct to handle these ref counts and make add/release private.
 		if(DataChannel->KeepPreviousFrameData())
 		{
-			if(Data.PrevCPUSimData)
-			{
-				Data.PrevCPUSimData->ReleaseReadRef();
-			}
 			Data.PrevCPUSimData = Data.CPUSimData->GetCurrentData();
-			Data.PrevCPUSimData->AddReadRef();
 		}
 
 		FNiagaraGpuComputeDispatchInterface* DispathInterface = FNiagaraGpuComputeDispatchInterface::Get(OwningWorld->GetWorld());
@@ -187,7 +182,6 @@ void UNiagaraDataChannelHandler_Global::Tick(float DeltaSeconds, ETickingGroup T
 				}
 
 				PublishRequest.Data = GameDataStaging->GetCurrentData();
-				PublishRequest.Data->AddReadRef();
 			}
 		}
 		else if( ensure(PublishRequest.Data) )
@@ -205,7 +199,6 @@ void UNiagaraDataChannelHandler_Global::Tick(float DeltaSeconds, ETickingGroup T
 		}
 		if(PublishRequest.bVisibleToGPUSims)
 		{
-			PublishRequest.Data->AddReadRef();
 			BuffersForGPU.Add(PublishRequest.Data);
 		}
 	}
@@ -289,27 +282,22 @@ void UNiagaraDataChannelHandler_Global::Tick(float DeltaSeconds, ETickingGroup T
 	}
 }
 
-void UNiagaraDataChannelHandler_Global::GetData(FNiagaraSystemInstance* SystemInstance, FNiagaraDataBuffer*& OutCPUData, bool bGetPrevFrameData)
+FNiagaraDataBufferRef UNiagaraDataChannelHandler_Global::GetData(FNiagaraSystemInstance* SystemInstance, bool bGetPrevFrameData)
 {
 	if(bGetPrevFrameData && Data.PrevCPUSimData)
 	{
-		OutCPUData = Data.PrevCPUSimData;
+		return Data.PrevCPUSimData;
 	}
 	else
 	{
-		OutCPUData = Data.CPUSimData->GetCurrentData();	
-	}
-
-	if(OutCPUData)
-	{
-		OutCPUData->AddReadRef();
+		return Data.CPUSimData->GetCurrentData();	
 	}
 	//For more complicated channels we could check the location + bounds of the system instance etc to return some spatially localized data.
 }
 
-void UNiagaraDataChannelHandler_Global::GetDataGPU(FNiagaraSystemInstance* SystemInstance, FNiagaraDataSet*& OutCPUData)
+FNiagaraDataSet* UNiagaraDataChannelHandler_Global::GetDataGPU(FNiagaraSystemInstance* SystemInstance)
 {
-	OutCPUData = Data.GPUSimData;	
+	return Data.GPUSimData;	
 	//For more complicated channels we could check the location + bounds of the system instance etc to return some spatially localized data.
 }
 
diff --git a/Engine/Plugins/FX/Niagara/Source/Niagara/Private/NiagaraDataChannel_Global.h b/Engine/Plugins/FX/Niagara/Source/Niagara/Private/NiagaraDataChannel_Global.h
--- a/Engine/Plugins/FX/Niagara/Source/Niagara/Private/NiagaraDataChannel_Global.h
+++ b/Engine/Plugins/FX/Niagara/Source/Niagara/Private/NiagaraDataChannel_Global.h
@@ -31,8 +31,8 @@ class UNiagaraDataChannelHandler_Global : public UNiagaraDataChannelHandler
 
 	virtual void Init(const UNiagaraDataChannel* InChannel) override;
 	virtual void Tick(float DeltaTime, ETickingGroup TickGroup, FNiagaraWorldManager* OwningWorld) override;
-	virtual void GetData(FNiagaraSystemInstance* SystemInstance, FNiagaraDataBuffer*& OutCPUData, bool bGetLastFrameData) override;
-	virtual void GetDataGPU(FNiagaraSystemInstance* SystemInstance, FNiagaraDataSet*& OutCPUData) override;
+	virtual FNiagaraDataBufferRef GetData(FNiagaraSystemInstance* SystemInstance, bool bGetLastFrameData) override;
+	virtual FNiagaraDataSet* GetDataGPU(FNiagaraSystemInstance* SystemInstance) override;
 
 	/** Returns the game level DataChannel data store. */
 	virtual FNiagaraDataChannelGameDataPtr GetGameData(UActorComponent* OwningComponent);
@@ -43,7 +43,7 @@ private:
 	FNiagaraDataSet* GameDataStaging;
 
 	/** Last Frame's CPU data. */
-	FNiagaraDataBuffer* PrevFrameCPU;
+	FNiagaraDataBufferRef PrevFrameCPU;
 };
 
 /**
@@ -60,4 +60,4 @@ TODO: Likely unsuitable for LWC games as we need a defined min/max for the hash
 // class UNiagaraDataChannelChannelHandler_Octree : UObject
 // {
 // 
-// };
\ No newline at end of file
+// };
diff --git a/Engine/Plugins/FX/Niagara/Source/Niagara/Private/NiagaraDataSet.cpp b/Engine/Plugins/FX/Niagara/Source/Niagara/Private/NiagaraDataSet.cpp
--- a/Engine/Plugins/FX/Niagara/Source/Niagara/Private/NiagaraDataSet.cpp
+++ b/Engine/Plugins/FX/Niagara/Source/Niagara/Private/NiagaraDataSet.cpp
@@ -1265,7 +1265,7 @@ void FNiagaraDataBuffer::GPUCopyFrom(const float* GPUReadBackFloat, const int* G
 	}
 }
 
-void FNiagaraDataBuffer::PushCPUBuffersToGPU(const TArray<FNiagaraDataBuffer*>& SourceBuffers, bool bReleaseRef, FRHICommandList& RHICmdList, ERHIFeatureLevel::Type FeatureLevel, const TCHAR* DebugSimName)
+void FNiagaraDataBuffer::PushCPUBuffersToGPU(const TArray<FNiagaraDataBufferRef>& SourceBuffers, bool bReleaseRef, FRHICommandList& RHICmdList, ERHIFeatureLevel::Type FeatureLevel, const TCHAR* DebugSimName)
 {
 	uint32 NewCount = 0;
 	check(GetOwner()->GetSimTarget() == ENiagaraSimTarget::GPUComputeSim);
@@ -1346,11 +1346,6 @@ void FNiagaraDataBuffer::PushCPUBuffersToGPU(const TArray<FNiagaraDataBuffer*>&
 			RHIUnlockBuffer(GPUBufferHalf.Buffer);
 		}
 	}
-
-	for (FNiagaraDataBuffer* Buffer : SourceBuffers)
-	{
-		Buffer->ReleaseReadRef();
-	}
 }
 
 void FNiagaraDataBuffer::Dump(int32 StartIndex, int32 InNumInstances, const FString& Label, const FName& SortParameterKey)const
diff --git a/Engine/Plugins/FX/Niagara/Source/Niagara/Private/NiagaraRenderer.cpp b/Engine/Plugins/FX/Niagara/Source/Niagara/Private/NiagaraRenderer.cpp
--- a/Engine/Plugins/FX/Niagara/Source/Niagara/Private/NiagaraRenderer.cpp
+++ b/Engine/Plugins/FX/Niagara/Source/Niagara/Private/NiagaraRenderer.cpp
@@ -307,28 +307,19 @@ FNiagaraDynamicDataBase::FNiagaraDynamicDataBase(const FNiagaraEmitterInstance*
 	if (SimTarget == ENiagaraSimTarget::CPUSim)
 	{
 		//On CPU we pass through direct ptr to the most recent data buffer.
-		Data.CPUParticleData = &DataSet.GetCurrentDataChecked();
-
-		//Mark this buffer as in use by this renderer. Prevents this buffer being reused to write new simulation data while it's inuse by the renderer.
-		Data.CPUParticleData->AddReadRef();
+		CPUParticleData = &DataSet.GetCurrentDataChecked();
 	}
 	else
 	{
 		//On GPU we must access the correct buffer via the GPUExecContext. Probably a way to route this data better outside the dynamic data in future.
 		//During simulation, the correct data buffer for rendering will be placed in the GPUContext and AddReadRef called.
 		check(SimTarget == ENiagaraSimTarget::GPUComputeSim);
-		Data.GPUExecContext = InEmitter->GetGPUContext();
+		GPUExecContext = InEmitter->GetGPUContext();
 	}
 }
 
 FNiagaraDynamicDataBase::~FNiagaraDynamicDataBase()
 {
-	if (SimTarget == ENiagaraSimTarget::CPUSim)
-	{
-		check(Data.CPUParticleData);
-		//Release our ref on the buffer so it can be reused as a destination for a new simulation tick.
-		Data.CPUParticleData->ReleaseReadRef();
-	}
 }
 
 bool FNiagaraDynamicDataBase::IsGpuLowLatencyTranslucencyEnabled() const
@@ -339,7 +330,7 @@ bool FNiagaraDynamicDataBase::IsGpuLowLatencyTranslucencyEnabled() const
 	}
 	else
 	{
-		return Data.GPUExecContext ? Data.GPUExecContext->HasTranslucentDataToRender() : false;
+		return GPUExecContext ? GPUExecContext->HasTranslucentDataToRender() : false;
 	}
 }
 
@@ -349,11 +340,11 @@ FNiagaraDataBuffer* FNiagaraDynamicDataBase::GetParticleDataToRender(bool bIsLow
 
 	if (SimTarget == ENiagaraSimTarget::CPUSim)
 	{
-		Ret = Data.CPUParticleData;
+		Ret = CPUParticleData;
 	}
 	else
 	{
-		Ret = Data.GPUExecContext->GetDataToRender(bIsLowLatencyTranslucent);
+		Ret = GPUExecContext->GetDataToRender(bIsLowLatencyTranslucent);
 	}
 
 	checkSlow(Ret == nullptr || Ret->IsBeingRead());
diff --git a/Engine/Plugins/FX/Niagara/Source/Niagara/Private/NiagaraScriptExecutionContext.cpp b/Engine/Plugins/FX/Niagara/Source/Niagara/Private/NiagaraScriptExecutionContext.cpp
--- a/Engine/Plugins/FX/Niagara/Source/Niagara/Private/NiagaraScriptExecutionContext.cpp
+++ b/Engine/Plugins/FX/Niagara/Source/Niagara/Private/NiagaraScriptExecutionContext.cpp
@@ -96,7 +96,7 @@ void FNiagaraScriptExecutionContextBase::BindData(int32 Index, FNiagaraDataSet&
 	FNiagaraDataBuffer* Output = DataSet.GetDestinationData();
 
 	DataSetInfo.SetNum(FMath::Max(DataSetInfo.Num(), Index + 1));
-	DataSetInfo[Index].Init(&DataSet, Input, Output, StartInstance, bUpdateInstanceCounts);
+	DataSetInfo[Index].Init(&DataSet, Input, StartInstance, bUpdateInstanceCounts);
 
 	//Would be nice to roll this and DataSetInfo into one but currently the VM being in it's own Engine module prevents this. Possibly should move the VM into Niagara itself.
 	TArrayView<uint8 const* RESTRICT const> InputRegisters = Input ? Input->GetRegisterTable() : TArrayView<uint8 const* RESTRICT const>();
@@ -124,7 +124,7 @@ void FNiagaraScriptExecutionContextBase::BindData(int32 Index, FNiagaraDataBuffe
 	check(Input && Input->GetOwner());
 	DataSetInfo.SetNum(FMath::Max(DataSetInfo.Num(), Index + 1));
 	FNiagaraDataSet* DataSet = Input->GetOwner();
-	DataSetInfo[Index].Init(DataSet, Input, nullptr, StartInstance, bUpdateInstanceCounts);
+	DataSetInfo[Index].Init(DataSet, Input, StartInstance, bUpdateInstanceCounts);
 
 	TArrayView<uint8 const* RESTRICT const> InputRegisters = Input->GetRegisterTable();
 
@@ -216,9 +216,9 @@ bool FNiagaraScriptExecutionContextBase::Execute(uint32 NumInstances, const FScr
 #if NIAGARA_NAN_CHECKING
 			Info.DataSet->CheckForNaNs();
 #endif
-			if (Info.bUpdateInstanceCount)
+			if (Info.bUpdateInstanceCount && Info.GetOutput())
 			{
-				Info.Output->SetNumInstances(Info.StartInstance + DataSetMetaTable[Idx].DataSetAccessIndex + 1);
+				Info.GetOutput()->SetNumInstances(Info.StartInstance + DataSetMetaTable[Idx].DataSetAccessIndex + 1);
 			}
 		}
 
diff --git a/Engine/Plugins/FX/Niagara/Source/Niagara/Public/NiagaraCommon.h b/Engine/Plugins/FX/Niagara/Source/Niagara/Public/NiagaraCommon.h
--- a/Engine/Plugins/FX/Niagara/Source/Niagara/Public/NiagaraCommon.h
+++ b/Engine/Plugins/FX/Niagara/Source/Niagara/Public/NiagaraCommon.h
@@ -23,6 +23,8 @@
 #include "RHIDefinitions.h"
 #endif
 
+#include "Templates/RefCounting.h"
+
 #include "NiagaraCommon.generated.h"
 
 struct FVersionedNiagaraEmitter;
@@ -42,6 +44,10 @@ namespace ERHIFeatureLevel { enum Type : int; }
 enum EShaderPlatform : uint16;
 enum EPixelFormat : uint8;
 
+class FNiagaraDataBuffer;
+
+using FNiagaraDataBufferRef = TRefCountPtr<FNiagaraDataBuffer>;
+
 #define NIAGARA_MEMORY_TRACKING	!UE_BUILD_SHIPPING
 
 #ifndef NIAGARA_COMPUTEDEBUG_ENABLED
diff --git a/Engine/Plugins/FX/Niagara/Source/Niagara/Public/NiagaraDataChannel.h b/Engine/Plugins/FX/Niagara/Source/Niagara/Public/NiagaraDataChannel.h
--- a/Engine/Plugins/FX/Niagara/Source/Niagara/Public/NiagaraDataChannel.h
+++ b/Engine/Plugins/FX/Niagara/Source/Niagara/Public/NiagaraDataChannel.h
@@ -165,7 +165,7 @@ struct NIAGARA_API FNiagaraWorldDataChannelStore
 
 	// Cached off buffer with the previous frame's CPU Sim accessible data.
 	// Some systems can choose to read this to avoid any current frame tick ordering issues.
-	FNiagaraDataBuffer* PrevCPUSimData = nullptr;
+	FNiagaraDataBufferRef PrevCPUSimData = nullptr;
 	
 	//		▲	GPU Sim Data can optionally be made visible to CPU Sims and Game Data.
 	//		|
diff --git a/Engine/Plugins/FX/Niagara/Source/Niagara/Public/NiagaraDataChannelHandler.h b/Engine/Plugins/FX/Niagara/Source/Niagara/Public/NiagaraDataChannelHandler.h
--- a/Engine/Plugins/FX/Niagara/Source/Niagara/Public/NiagaraDataChannelHandler.h
+++ b/Engine/Plugins/FX/Niagara/Source/Niagara/Public/NiagaraDataChannelHandler.h
@@ -35,10 +35,10 @@ public:
 	virtual void Tick(float DeltaTime, ETickingGroup TickGroup, FNiagaraWorldManager* OwningWorld);
 
 	/** Returns the data for this data channel given the system instance. Possibly some handlers will subdivide spacially etc so may not return the same data for the different systems. */
-	virtual void GetData(FNiagaraSystemInstance* SystemInstance, FNiagaraDataBuffer*& OutCPUData, bool bGetLastFrameData)  PURE_VIRTUAL(UNiagaraDataChannelHandler::GetData, );	
+	virtual FNiagaraDataBufferRef GetData(FNiagaraSystemInstance* SystemInstance, bool bGetLastFrameData)  PURE_VIRTUAL(UNiagaraDataChannelHandler::GetData, return nullptr;);	
 
 	/** Gets the GPU data set with data for the given system. This API will likely be reworked as GPU support is fleshed out and more complex handlers are created. */
-	virtual void GetDataGPU(FNiagaraSystemInstance* SystemInstance, FNiagaraDataSet*& OutCPUData) PURE_VIRTUAL(UNiagaraDataChannelHandler::GetDataGPU, );	
+	virtual FNiagaraDataSet* GetDataGPU(FNiagaraSystemInstance* SystemInstance) PURE_VIRTUAL(UNiagaraDataChannelHandler::GetDataGPU, return nullptr;);	
 
 	/** Adds a request to publish some data into the channel on the next tick. */
 	virtual void Publish(const FNiagaraDataChannelPublishRequest& Request);
@@ -80,7 +80,7 @@ protected:
 	TArray<FNiagaraDataChannelPublishRequest> PublishRequests;
 
 	/** Data buffers we'll be passing to the RT for uploading to the GPU */
-	TArray<FNiagaraDataBuffer*> BuffersForGPU;
+	TArray<FNiagaraDataBufferRef> BuffersForGPU;
 
 	/** Render Thread Proxy for this handler. Each handler type must create it's proxy inside it's constructor. */
 	TUniquePtr<FNiagaraDataChannelHandlerRTProxyBase> RTProxy;
diff --git a/Engine/Plugins/FX/Niagara/Source/Niagara/Public/NiagaraRenderer.h b/Engine/Plugins/FX/Niagara/Source/Niagara/Public/NiagaraRenderer.h
--- a/Engine/Plugins/FX/Niagara/Source/Niagara/Public/NiagaraRenderer.h
+++ b/Engine/Plugins/FX/Niagara/Source/Niagara/Public/NiagaraRenderer.h
@@ -51,11 +51,8 @@ protected:
 	ENiagaraSimTarget SimTarget;
 	FNiagaraSystemInstanceID SystemInstanceID;
 
-	union
-	{
-		FNiagaraDataBuffer* CPUParticleData;
-		FNiagaraComputeExecutionContext* GPUExecContext;
-	}Data;
+	FNiagaraDataBufferRef CPUParticleData;
+	FNiagaraComputeExecutionContext* GPUExecContext = nullptr;
 };
 
 //////////////////////////////////////////////////////////////////////////
