- Call InitializeSystemTextures when flushing pending ticks as they may not have been initialized

diff --git a/Engine/Plugins/FX/Niagara/Source/Niagara/Private/NiagaraGpuComputeDispatch.cpp b/Engine/Plugins/FX/Niagara/Source/Niagara/Private/NiagaraGpuComputeDispatch.cpp
--- a/Engine/Plugins/FX/Niagara/Source/Niagara/Private/NiagaraGpuComputeDispatch.cpp
+++ b/Engine/Plugins/FX/Niagara/Source/Niagara/Private/NiagaraGpuComputeDispatch.cpp
@@ -5,6 +5,7 @@
 #include "Async/Async.h"
 #include "CanvasTypes.h"
 #include "ClearQuad.h"
+#include "EngineModule.h"
 #include "Engine/Canvas.h"
 #include "Engine/Engine.h"
 #include "Engine/GameViewportClient.h"
@@ -426,6 +427,9 @@ void FNiagaraGpuComputeDispatch::ProcessPendingTicksFlush(FRHICommandListImmedia
 			// in the event of submitting commands, i.e. when we write a fence, or indeed perform a manual flush.
 			RHICmdList.BeginScene();
 
+			// Ensure system textures are initialized
+			GetRendererModule().InitializeSystemTextures(RHICmdList);
+
 			// Allow downstream logic to detect we are running pending ticks outside the scene renderer
 			bIsOutsideSceneRenderer = true;
 
