- Handle invalid rotation from emitter mode mesh renderer

diff --git a/Engine/Plugins/FX/Niagara/Source/Niagara/Private/NiagaraRendererMeshes.cpp b/Engine/Plugins/FX/Niagara/Source/Niagara/Private/NiagaraRendererMeshes.cpp
--- a/Engine/Plugins/FX/Niagara/Source/Niagara/Private/NiagaraRendererMeshes.cpp
+++ b/Engine/Plugins/FX/Niagara/Source/Niagara/Private/NiagaraRendererMeshes.cpp
@@ -1005,8 +1005,8 @@ void FNiagaraRendererMeshes::SetupElementForGPUScene(
 		FEmitterSourceInstanceData& EmitterSourceData = GPUSceneRes.EmitterSourceData;
 		EmitterSourceData.CustomData = 0.0f;		
 
-		const FVector4 Rot = FVector4(CommonParameters.DefaultRotation);
-		FMatrix LocalToPrimitive = FTransform(FQuat(Rot.X, Rot.Y, Rot.Z, Rot.W), FVector(CommonParameters.DefaultPosition), FVector(CommonParameters.DefaultScale)).ToMatrixWithScale();
+		const FQuat Rot = FQuat(CommonParameters.DefaultRotation.X, CommonParameters.DefaultRotation.Y, CommonParameters.DefaultRotation.Z, CommonParameters.DefaultRotation.W).GetNormalized();
+		FMatrix LocalToPrimitive = FTransform(Rot, FVector(CommonParameters.DefaultPosition), FVector(CommonParameters.DefaultScale)).ToMatrixWithScale();
 
 		if (!bLocalSpace)
 		{
@@ -1022,8 +1022,8 @@ void FNiagaraRendererMeshes::SetupElementForGPUScene(
 		
 		if (bNeedsPrevTransform)
 		{
-			const FVector4 PrevRot = FVector4(CommonParameters.DefaultPrevRotation);
-			FMatrix PrevLocalToPrimitive = FTransform(FQuat(PrevRot.X, PrevRot.Y, PrevRot.Z, PrevRot.W), FVector(CommonParameters.DefaultPrevPosition), FVector(CommonParameters.DefaultPrevScale)).ToMatrixWithScale();
+			const FQuat PrevRot = FQuat(CommonParameters.DefaultPrevRotation.X, CommonParameters.DefaultPrevRotation.Y, CommonParameters.DefaultPrevRotation.Z, CommonParameters.DefaultPrevRotation.W).GetNormalized();
+			FMatrix PrevLocalToPrimitive = FTransform(PrevRot, FVector(CommonParameters.DefaultPrevPosition), FVector(CommonParameters.DefaultPrevScale)).ToMatrixWithScale();
 
 			if (!bLocalSpace)
 			{
