Changing some Niagara code to use Renderer interfaces instead of implementations.

diff --git a/Engine/Plugins/FX/Niagara/Source/Niagara/Private/DataInterface/NiagaraDataInterfaceStaticMesh.cpp b/Engine/Plugins/FX/Niagara/Source/Niagara/Private/DataInterface/NiagaraDataInterfaceStaticMesh.cpp
--- a/Engine/Plugins/FX/Niagara/Source/Niagara/Private/DataInterface/NiagaraDataInterfaceStaticMesh.cpp
+++ b/Engine/Plugins/FX/Niagara/Source/Niagara/Private/DataInterface/NiagaraDataInterfaceStaticMesh.cpp
@@ -3062,18 +3062,13 @@ void UNiagaraDataInterfaceStaticMesh::SetShaderParameters(const FNiagaraDataInte
 
 	if (Context.IsParameterBound(&ShaderParameters->InstanceDistanceFieldIndex))
 	{
-		int32 DistanceFieldIndex = -1;
+		int32 DistanceFieldIndex = INDEX_NONE;
 		if (DistanceFieldSceneData != nullptr && InstanceData.DistanceFieldPrimitiveId.IsValid())
 		{
-			if (FScene* Scene = Context.GetComputeDispatchInterface().GetScene())
+			const FSceneInterface* Scene = Context.GetComputeDispatchInterface().GetScene();
+			if (const FPrimitiveSceneInfo* PrimitiveSceneInfo = Scene ? Scene->GetPrimitiveSceneInfo(InstanceData.DistanceFieldPrimitiveId) : nullptr)
 			{
-				// Kind of gross, but currently no way to reference other primitive scene infos
-				const int32 PrimitiveSceneIndex = Scene->PrimitiveComponentIds.Find(InstanceData.DistanceFieldPrimitiveId);
-				if (PrimitiveSceneIndex != INDEX_NONE)
-				{
-					const TArray<int32, TInlineAllocator<1>>& DFIndices = Scene->Primitives[PrimitiveSceneIndex]->DistanceFieldInstanceIndices;
-					DistanceFieldIndex = DFIndices.Num() > 0 ? DFIndices[0] : -1;
-				}
+				DistanceFieldIndex = PrimitiveSceneInfo->DistanceFieldInstanceIndices.Num() > 0 ? PrimitiveSceneInfo->DistanceFieldInstanceIndices[0] : INDEX_NONE;
 			}
 		}
 		ShaderParameters->InstanceDistanceFieldIndex = DistanceFieldIndex;
diff --git a/Engine/Plugins/FX/Niagara/Source/Niagara/Private/NiagaraDataInterfaceOcclusion.cpp b/Engine/Plugins/FX/Niagara/Source/Niagara/Private/NiagaraDataInterfaceOcclusion.cpp
--- a/Engine/Plugins/FX/Niagara/Source/Niagara/Private/NiagaraDataInterfaceOcclusion.cpp
+++ b/Engine/Plugins/FX/Niagara/Source/Niagara/Private/NiagaraDataInterfaceOcclusion.cpp
@@ -9,7 +9,6 @@
 #include "NiagaraSystemGpuComputeProxy.h"
 #include "NiagaraSystemInstance.h"
 
-#include "SceneRendering.h"
 #include "SceneView.h"
 #include "ShaderParameterUtils.h"
 #include "Internationalization/Internationalization.h"
@@ -186,8 +185,9 @@ void UNiagaraDataInterfaceOcclusion::SetShaderParameters(const FNiagaraDataInter
 
 	if (Context.IsResourceBound(&ShaderParameters->CloudVolumetricTexture))
 	{
-		TConstArrayView<FViewInfo> SimulationViewInfos = Context.GetComputeDispatchInterface().GetSimulationViewInfos();
-		FRDGTextureRef CloudTexture = SimulationViewInfos.Num() > 0 ? SimulationViewInfos[0].GetVolumetricCloudTexture(Context.GetGraphBuilder()) : nullptr;
+		TConstStridedView<FSceneView> SimulationSceneViews = Context.GetComputeDispatchInterface().GetSimulationSceneViews();
+		const FSceneView* View = SimulationSceneViews.Num() > 0 ? &SimulationSceneViews[0] : nullptr;
+		FRDGTextureRef CloudTexture = View && View->State ? View->State->GetVolumetricCloudTexture(Context.GetGraphBuilder()) : nullptr;
 		if (CloudTexture == nullptr)
 		{
 			CloudTexture = Context.GetComputeDispatchInterface().GetBlackTexture(Context.GetGraphBuilder(), ETextureDimension::Texture2D);
diff --git a/Engine/Plugins/FX/Niagara/Source/Niagara/Private/NiagaraDataInterfaceRigidMeshCollisionQuery.cpp b/Engine/Plugins/FX/Niagara/Source/Niagara/Private/NiagaraDataInterfaceRigidMeshCollisionQuery.cpp
--- a/Engine/Plugins/FX/Niagara/Source/Niagara/Private/NiagaraDataInterfaceRigidMeshCollisionQuery.cpp
+++ b/Engine/Plugins/FX/Niagara/Source/Niagara/Private/NiagaraDataInterfaceRigidMeshCollisionQuery.cpp
@@ -936,21 +936,17 @@ struct FNDIRigidMeshCollisionProxy : public FNiagaraDataInterfaceProxy
 					const uint32 BufferBytes = sizeof(uint32) * ElementCount;
 					void* BufferData = RHILockBuffer(ProxyData->AssetBuffer->DFIndexBuffer.Buffer, 0, BufferBytes, RLM_WriteOnly);
 
-					const FScene* Scene = Context.GetComputeDispatchInterface().GetScene();
+					const FSceneInterface* Scene = Context.GetComputeDispatchInterface().GetScene();
 					if (Scene && !ProxyData->UniqueComponentIds.IsEmpty())
 					{
 						TArray<uint32> UniqueDistanceFieldIndices;
 						UniqueDistanceFieldIndices.Reserve(ProxyData->UniqueComponentIds.Num());
 
-						for (const FPrimitiveComponentId& ComponentId : ProxyData->UniqueComponentIds)
+						for (const FPrimitiveComponentId ComponentId : ProxyData->UniqueComponentIds)
 						{
-							uint32& DistanceFieldIndex = UniqueDistanceFieldIndices.Add_GetRef(INDEX_NONE);
-							const int32 PrimitiveSceneIndex = Scene->PrimitiveComponentIds.Find(ComponentId);
-							if (PrimitiveSceneIndex != INDEX_NONE)
-							{
-								const TArray<int32, TInlineAllocator<1>>& DFIndices = Scene->Primitives[PrimitiveSceneIndex]->DistanceFieldInstanceIndices;
-								DistanceFieldIndex = DFIndices.IsEmpty() ? INDEX_NONE : DFIndices[0];
-							}
+							const FPrimitiveSceneInfo* PrimitiveSceneInfo = Scene->GetPrimitiveSceneInfo(ComponentId);
+							const uint32 DistanceFieldIndex = PrimitiveSceneInfo && PrimitiveSceneInfo->DistanceFieldInstanceIndices.Num() > 0 ? PrimitiveSceneInfo->DistanceFieldInstanceIndices[0] : INDEX_NONE;
+							UniqueDistanceFieldIndices.Emplace(DistanceFieldIndex);
 						}
 
 						TArrayView<uint32> BufferView(reinterpret_cast<uint32*>(BufferData), ElementCount);
diff --git a/Engine/Plugins/FX/Niagara/Source/Niagara/Private/NiagaraGpuComputeDispatch.cpp b/Engine/Plugins/FX/Niagara/Source/Niagara/Private/NiagaraGpuComputeDispatch.cpp
--- a/Engine/Plugins/FX/Niagara/Source/Niagara/Private/NiagaraGpuComputeDispatch.cpp
+++ b/Engine/Plugins/FX/Niagara/Source/Niagara/Private/NiagaraGpuComputeDispatch.cpp
@@ -10,6 +10,7 @@
 #include "Engine/Engine.h"
 #include "Engine/GameViewportClient.h"
 #include "Engine/World.h"
+#include "EngineModule.h"
 #include "Materials/MaterialRenderProxy.h"
 #include "Misc/ScopeExit.h"
 #include "ProfilingDebugging/CsvProfiler.h"
@@ -410,17 +411,10 @@ void FNiagaraGpuComputeDispatch::ProcessPendingTicksFlush(FRHICommandListImmedia
 			ViewInitOptions.ViewRotationMatrix = CachedViewInitOptions.ViewRotationMatrix;
 			ViewInitOptions.ProjectionMatrix = CachedViewInitOptions.ProjectionMatrix;
 
-			FViewInfo DummyView(ViewInitOptions);
-
-			DummyView.ViewRect = DummyView.UnscaledViewRect;
-			DummyView.CachedViewUniformShaderParameters = MakeUnique<FViewUniformShaderParameters>();
-
-			FBox UnusedVolumeBounds[TVC_MAX];
-			DummyView.SetupUniformBufferParameters(UnusedVolumeBounds, TVC_MAX, *DummyView.CachedViewUniformShaderParameters);
-
-			DummyView.ViewUniformBuffer = TUniformBufferRef<FViewUniformShaderParameters>::CreateUniformBufferImmediate(*DummyView.CachedViewUniformShaderParameters, UniformBuffer_SingleFrame);
+			GetRendererModule().CreateAndInitSingleView(RHICmdList, &ViewFamily, &ViewInitOptions);
+			const FViewInfo* DummyView = (const FViewInfo*)ViewFamily.Views[0];
 
-			TConstArrayView<FViewInfo> DummyViews = MakeArrayView(&DummyView, 1);
+			TConstArrayView<FViewInfo> DummyViews = MakeArrayView(DummyView, 1);
 			const bool bAllowGPUParticleUpdate = true;
 
 			// Notify that we are about to begin rendering the 'scene' this is required because some RHIs will ClearState
@@ -1050,6 +1044,7 @@ void FNiagaraGpuComputeDispatch::ExecuteTicks(FRDGBuilder& GraphBuilder, TConstA
 	RDG_RHI_EVENT_SCOPE(GraphBuilder, NiagaraGpuComputeDispatch);	//-TODO:RDG: Show TickStage
 
 	// Setup Parameters that can be read from data interfaces
+	SimulationSceneViews = MakeStridedViewOfBase<const FSceneView>(Views);
 	SimulationViewInfos = Views;
 
 	{
@@ -1346,6 +1341,7 @@ void FNiagaraGpuComputeDispatch::ExecuteTicks(FRDGBuilder& GraphBuilder, TConstA
 	);
 
 	// Tear down for tick pass
+	SimulationSceneViews = TConstStridedView<FSceneView>();
 	SimulationViewInfos = TConstArrayView<FViewInfo>();
 	SceneTexturesUniformParams = nullptr;
 	MobileSceneTexturesUniformParams = nullptr;
@@ -1604,7 +1600,7 @@ void FNiagaraGpuComputeDispatch::DispatchStage(FRDGBuilder& GraphBuilder, const
 	// Set ViewUniformBuffer if it's required
 	if (ComputeShader->bNeedsViewUniformBuffer)
 	{
-		DispatchParameters->View = SimulationViewInfos[0].ViewUniformBuffer;
+		DispatchParameters->View = SimulationSceneViews[0].ViewUniformBuffer;
 	}
 	DispatchParameters->SceneTextures.SceneTextures			= SceneTexturesUniformParams;
 	DispatchParameters->SceneTextures.MobileSceneTextures	= MobileSceneTexturesUniformParams;
diff --git a/Engine/Plugins/FX/Niagara/Source/Niagara/Public/NiagaraGpuComputeDispatchInterface.h b/Engine/Plugins/FX/Niagara/Source/Niagara/Public/NiagaraGpuComputeDispatchInterface.h
--- a/Engine/Plugins/FX/Niagara/Source/Niagara/Public/NiagaraGpuComputeDispatchInterface.h
+++ b/Engine/Plugins/FX/Niagara/Source/Niagara/Public/NiagaraGpuComputeDispatchInterface.h
@@ -2,6 +2,7 @@
 
 #pragma once
 
+#include "Containers/StridedView.h"
 #include "NiagaraCommon.h"
 #include "NiagaraEmptyUAVPool.h"
 #include "NiagaraGpuComputeDataManager.h"
@@ -77,6 +78,7 @@ public:
 	Get access to the Views the simulation is being rendered with.
 	List is only valid during graph building (i.e. during ExecuteTicks) and for simulations in PostInitViews / PostRenderOpaque.
 	*/
+	TConstStridedView<FSceneView> GetSimulationSceneViews() const { return SimulationSceneViews; }
 	TConstArrayView<FViewInfo> GetSimulationViewInfos() const { return SimulationViewInfos; }
 
 	/** Get access to the instance count manager. */
@@ -181,6 +183,7 @@ protected:
 
 	TArray<TPair<FName, TUniquePtr<FNiagaraGpuComputeDataManager>>> GpuDataManagers;
 
+	TConstStridedView<FSceneView>			SimulationSceneViews;
 	TConstArrayView<FViewInfo>				SimulationViewInfos;
 
 	bool									bIsOutsideSceneRenderer = false;
