Niagara Editor Preview: Preserve camera orbit behavior when refocusing.

diff --git a/Engine/Plugins/FX/Niagara/Source/NiagaraEditor/Private/Widgets/SNiagaraSystemViewport.cpp b/Engine/Plugins/FX/Niagara/Source/NiagaraEditor/Private/Widgets/SNiagaraSystemViewport.cpp
--- a/Engine/Plugins/FX/Niagara/Source/NiagaraEditor/Private/Widgets/SNiagaraSystemViewport.cpp
+++ b/Engine/Plugins/FX/Niagara/Source/NiagaraEditor/Private/Widgets/SNiagaraSystemViewport.cpp
@@ -128,7 +128,9 @@ void FNiagaraSystemViewportClient::Draw(FViewport* InViewport, FCanvas* Canvas)
 		FBox Bounds = Component->GetLocalBounds().GetBox();
 		if (bUpdateViewportFocus && Bounds.IsValid)
 		{
+			bool bCachedShouldOrbitCamera = ShouldOrbitCamera();
 			FocusViewportOnBox(Bounds);
+			ToggleOrbitCamera(bCachedShouldOrbitCamera);
 			SetUpdateViewportFocus(false);
 		}
 	}
