Fix for NiagaraVaraible serialization in object replacement archives

diff --git a/Engine/Plugins/FX/Niagara/Source/Niagara/Private/NiagaraModule.cpp b/Engine/Plugins/FX/Niagara/Source/Niagara/Private/NiagaraModule.cpp
--- a/Engine/Plugins/FX/Niagara/Source/Niagara/Private/NiagaraModule.cpp
+++ b/Engine/Plugins/FX/Niagara/Source/Niagara/Private/NiagaraModule.cpp
@@ -1380,7 +1380,7 @@ bool FNiagaraTypeDefinition::Serialize(FArchive& Ar)
 		}
 	}
 
-	if (Ar.IsLoading() || Ar.IsSaving())
+	if (Ar.IsLoading() || Ar.IsSaving() || Ar.IsModifyingWeakAndStrongReferences())
 	{
 		UScriptStruct* Struct = FNiagaraTypeDefinition::StaticStruct();
 		Struct->SerializeTaggedProperties(Ar, (uint8*)this, Struct, nullptr);
