Moving more custom Rendering code from Niagara to FXRenderingUtils.

diff --git a/Engine/Plugins/FX/Niagara/Source/NiagaraShader/Private/NiagaraAsyncGpuTraceProviderHwrt.cpp b/Engine/Plugins/FX/Niagara/Source/NiagaraShader/Private/NiagaraAsyncGpuTraceProviderHwrt.cpp
--- a/Engine/Plugins/FX/Niagara/Source/NiagaraShader/Private/NiagaraAsyncGpuTraceProviderHwrt.cpp
+++ b/Engine/Plugins/FX/Niagara/Source/NiagaraShader/Private/NiagaraAsyncGpuTraceProviderHwrt.cpp
@@ -5,9 +5,16 @@
 
 #if RHI_RAYTRACING
 
+#include "FXRenderingUtils.h"
 #include "GlobalShader.h"
+#include "MeshPassProcessor.h"
 #include "NiagaraRenderer.h"
 #include "NiagaraSettings.h"
+#include "SceneInterface.h"
+#include "ShaderParameterMacros.h"
+#include "ShaderParameterStruct.h"
+#include "Templates/Function.h"
+
 #include "ScenePrivate.h"
 #include "RayTracingPayloadType.h"
 
@@ -301,9 +308,7 @@ bool FNiagaraAsyncGpuTraceProviderHwrt::IsAvailable() const
 		return false;
 	}
 
-	FScene* Scene = Dispatcher->GetScene();
-
-	return Scene ? Scene->RayTracingScene.IsCreated() : false;
+	return UE::FXRenderingUtils::RayTracing::HasRayTracingScene(Dispatcher->GetSceneInterface());
 }
 
 void FNiagaraAsyncGpuTraceProviderHwrt::PostRenderOpaque(FRHICommandList& RHICmdList, TConstArrayView<FViewInfo> Views, FCollisionGroupHashMap* CollisionGroupHash)
@@ -311,13 +316,13 @@ void FNiagaraAsyncGpuTraceProviderHwrt::PostRenderOpaque(FRHICommandList& RHICmd
 	check(IsAvailable());
 	check(Views.Num() > 0);
 
-	const FViewInfo& ReferenceView = Views[0];
-	FScene* Scene = Dispatcher->GetScene();
+	const FSceneView& ReferenceView = Views[0];
+	FSceneInterface* Scene = Dispatcher->GetSceneInterface();
 
-	if (Scene && Scene->RayTracingScene.IsCreated())
+	if (UE::FXRenderingUtils::RayTracing::HasRayTracingScene(Scene))
 	{
-		RayTracingScene = Scene->RayTracingScene.GetRHIRayTracingSceneChecked();
-		RayTracingSceneView = Scene->RayTracingScene.CreateLayerViewRHI(ERayTracingSceneLayer::Base);
+		RayTracingScene = UE::FXRenderingUtils::RayTracing::GetRayTracingScene(Scene);
+		RayTracingSceneView = UE::FXRenderingUtils::RayTracing::GetRayTracingSceneView(Scene);
 		ViewUniformBuffer = ReferenceView.ViewUniformBuffer;
 
 		FGlobalShaderMap* ShaderMap = GetGlobalShaderMap(ShaderPlatform);
@@ -341,7 +346,7 @@ void FNiagaraAsyncGpuTraceProviderHwrt::PostRenderOpaque(FRHICommandList& RHICmd
 			RHICmdList,
 			RayTracingScene,
 			ViewUniformBuffer,
-			ReferenceView.VisibleRayTracingMeshCommands,
+			UE::FXRenderingUtils::RayTracing::GetVisibleRayTracingMeshCommands(ReferenceView),
 			RayTracingPipelineState,
 			BakeDefault);
 	}
@@ -365,7 +370,6 @@ void FNiagaraAsyncGpuTraceProviderHwrt::IssueTraces(FRHICommandList& RHICmdList,
 
 	SCOPED_DRAW_EVENT(RHICmdList, NiagaraIssueTracesHwrt);
 
-	FScene* Scene = Dispatcher->GetScene();
 	FGlobalShaderMap* ShaderMap = GetGlobalShaderMap(ShaderPlatform);
 
 	TShaderRef<FNiagaraCollisionRayTraceRG> RGShader = FNiagaraCollisionRayTraceRG::GetShader(ShaderMap, CollisionGroupHash != nullptr);
@@ -373,10 +377,13 @@ void FNiagaraAsyncGpuTraceProviderHwrt::IssueTraces(FRHICommandList& RHICmdList,
 	FNiagaraCollisionRayTraceRG::FParameters Params;
 
 	Params.View = GetShaderBinding(ViewUniformBuffer);
-	Params.GPUSceneParameters.GPUSceneInstanceSceneData = Scene->GPUScene.InstanceSceneDataBuffer->GetSRV();
-	Params.GPUSceneParameters.GPUSceneInstancePayloadData = Scene->GPUScene.InstancePayloadDataBuffer->GetSRV();
-	Params.GPUSceneParameters.GPUScenePrimitiveSceneData = Scene->GPUScene.PrimitiveBuffer->GetSRV();
-	Params.GPUSceneParameters.GPUSceneFrameNumber = Scene->GPUScene.GetSceneFrameNumber();
+
+	UE::FXRenderingUtils::GPUScene::FBuffers GPUSceneBuffers = UE::FXRenderingUtils::GPUScene::GetBuffers(Dispatcher->GetSceneInterface());
+
+	Params.GPUSceneParameters.GPUSceneInstanceSceneData = GPUSceneBuffers.InstanceSceneDataBuffer;
+	Params.GPUSceneParameters.GPUSceneInstancePayloadData = GPUSceneBuffers.InstancePayloadDataBuffer;
+	Params.GPUSceneParameters.GPUScenePrimitiveSceneData = GPUSceneBuffers.PrimitiveBuffer;
+	Params.GPUSceneParameters.GPUSceneFrameNumber = GPUSceneBuffers.SceneFrameNumber;
 
 	if (CollisionGroupHash)
 	{
diff --git a/Engine/Source/Runtime/Renderer/Private/FXRenderingUtils.cpp b/Engine/Source/Runtime/Renderer/Private/FXRenderingUtils.cpp
--- a/Engine/Source/Runtime/Renderer/Private/FXRenderingUtils.cpp
+++ b/Engine/Source/Runtime/Renderer/Private/FXRenderingUtils.cpp
@@ -146,3 +146,58 @@ void UE::FXRenderingUtils::DistanceFields::SetupAtlasParameters(FRDGBuilder& Gra
 		}
 	}
 }
+
+UE::FXRenderingUtils::GPUScene::FBuffers UE::FXRenderingUtils::GPUScene::GetBuffers(const FSceneInterface* InScene)
+{
+	UE::FXRenderingUtils::GPUScene::FBuffers Buffers{};
+
+	if (const FScene* Scene = InScene->GetRenderScene())
+	{
+		Buffers.InstanceSceneDataBuffer = Scene->GPUScene.InstanceSceneDataBuffer->GetSRV();
+		Buffers.InstancePayloadDataBuffer = Scene->GPUScene.InstancePayloadDataBuffer->GetSRV();
+		Buffers.PrimitiveBuffer = Scene->GPUScene.PrimitiveBuffer->GetSRV();
+		Buffers.SceneFrameNumber = Scene->GPUScene.GetSceneFrameNumber();
+	}
+
+	return Buffers;
+}
+
+
+#if RHI_RAYTRACING
+
+bool UE::FXRenderingUtils::RayTracing::HasRayTracingScene(const FSceneInterface* InScene)
+{
+	if (const FScene* Scene = InScene->GetRenderScene())
+	{
+		return Scene->RayTracingScene.IsCreated();
+	}
+
+	return false;
+}
+
+FRHIRayTracingScene* UE::FXRenderingUtils::RayTracing::GetRayTracingScene(const FSceneInterface* InScene)
+{
+	if (const FScene* Scene = InScene->GetRenderScene())
+	{
+		return Scene->RayTracingScene.GetRHIRayTracingSceneChecked();
+	}
+
+	return nullptr;
+}
+
+FRHIShaderResourceView* UE::FXRenderingUtils::RayTracing::GetRayTracingSceneView(const FSceneInterface* InScene)
+{
+	if (const FScene* Scene = InScene->GetRenderScene())
+	{
+		return Scene->RayTracingScene.CreateLayerViewRHI(ERayTracingSceneLayer::Base);
+	}
+
+	return nullptr;
+}
+
+TConstArrayView<FVisibleRayTracingMeshCommand> UE::FXRenderingUtils::RayTracing::GetVisibleRayTracingMeshCommands(const FSceneView& View)
+{
+	return static_cast<const FViewInfo&>(View).VisibleRayTracingMeshCommands;
+}
+
+#endif // RHI_RAYTRACING
diff --git a/Engine/Source/Runtime/Renderer/Public/FXRenderingUtils.h b/Engine/Source/Runtime/Renderer/Public/FXRenderingUtils.h
--- a/Engine/Source/Runtime/Renderer/Public/FXRenderingUtils.h
+++ b/Engine/Source/Runtime/Renderer/Public/FXRenderingUtils.h
@@ -3,6 +3,7 @@
 #pragma once
 
 #include "HAL/Platform.h"
+#include "RHIDefinitions.h"
 #include "UObject/ObjectMacros.h"
 
 #if UE_ENABLE_INCLUDE_ORDER_DEPRECATED_IN_5_3
@@ -12,10 +13,17 @@
 class FMaterial;
 class FPrimitiveSceneProxy;
 class FRDGBuilder;
+class FRHIShaderResourceView;
+class FSceneInterface;
 class FSceneView;
 class FSceneViewFamily;
 class FShaderParametersMetadata;
 
+#if RHI_RAYTRACING
+class FRHIRayTracingScene;
+class FVisibleRayTracingMeshCommand;
+#endif
+
 namespace UE::FXRenderingUtils
 {
 	RENDERER_API bool CanMaterialRenderBeforeFXPostOpaque(const FSceneViewFamily& ViewFamily, const FPrimitiveSceneProxy& SceneProxy, const FMaterial& Material);
@@ -30,6 +38,30 @@ namespace UE::FXRenderingUtils
 		RENDERER_API void SetupObjectBufferParameters(FRDGBuilder& GraphBuilder, uint8* DestinationData, const FSceneView* View);
 		RENDERER_API void SetupAtlasParameters(FRDGBuilder& GraphBuilder, uint8* DestinationData, const FSceneView* View);
 	}
+
+	namespace GPUScene
+	{
+		struct FBuffers
+		{
+			FRHIShaderResourceView* InstanceSceneDataBuffer;
+			FRHIShaderResourceView* InstancePayloadDataBuffer;
+			FRHIShaderResourceView* PrimitiveBuffer;
+			uint32                  SceneFrameNumber;
+		};
+
+		RENDERER_API FBuffers GetBuffers(const FSceneInterface* Scene);
+	}
+
+#if RHI_RAYTRACING
+	namespace RayTracing
+	{
+		RENDERER_API bool HasRayTracingScene(const FSceneInterface* Scene);
+		RENDERER_API FRHIRayTracingScene* GetRayTracingScene(const FSceneInterface* Scene);
+		RENDERER_API FRHIShaderResourceView* GetRayTracingSceneView(const FSceneInterface* Scene);
+
+		RENDERER_API TConstArrayView<FVisibleRayTracingMeshCommand> GetVisibleRayTracingMeshCommands(const FSceneView& View);
+	}
+#endif
 }
 
 /**
