Removing some more dependencies on Renderer Private headers.

diff --git a/Engine/Plugins/FX/Niagara/Source/Niagara/Private/DataInterface/NiagaraDataInterfaceStaticMesh.cpp b/Engine/Plugins/FX/Niagara/Source/Niagara/Private/DataInterface/NiagaraDataInterfaceStaticMesh.cpp
--- a/Engine/Plugins/FX/Niagara/Source/Niagara/Private/DataInterface/NiagaraDataInterfaceStaticMesh.cpp
+++ b/Engine/Plugins/FX/Niagara/Source/Niagara/Private/DataInterface/NiagaraDataInterfaceStaticMesh.cpp
@@ -12,11 +12,11 @@
 #include "NiagaraShaderParametersBuilder.h"
 #include "NiagaraWorldManager.h"
 #include "NiagaraDataInterfaceUtilities.h"
-#include "ScenePrivate.h"
-#include "SceneRendering.h"
+#include "SceneInterface.h"
 #include "ShaderParameterUtils.h"
 #include "StaticMeshComponentLODInfo.h"
 #include "NiagaraStats.h"
+#include "PrimitiveSceneInfo.h"
 
 #if WITH_EDITOR
 #include "Editor.h"
@@ -3063,7 +3063,7 @@ void UNiagaraDataInterfaceStaticMesh::SetShaderParameters(const FNiagaraDataInte
 		int32 DistanceFieldIndex = INDEX_NONE;
 		if (bDistanceFieldParametersExist && InstanceData.DistanceFieldPrimitiveId.IsValid())
 		{
-			const FSceneInterface* Scene = Context.GetComputeDispatchInterface().GetScene();
+			const FSceneInterface* Scene = Context.GetComputeDispatchInterface().GetSceneInterface();
 			if (const FPrimitiveSceneInfo* PrimitiveSceneInfo = Scene ? Scene->GetPrimitiveSceneInfo(InstanceData.DistanceFieldPrimitiveId) : nullptr)
 			{
 				DistanceFieldIndex = PrimitiveSceneInfo->DistanceFieldInstanceIndices.Num() > 0 ? PrimitiveSceneInfo->DistanceFieldInstanceIndices[0] : INDEX_NONE;
diff --git a/Engine/Plugins/FX/Niagara/Source/Niagara/Private/NiagaraAsyncGpuTraceHelper.cpp b/Engine/Plugins/FX/Niagara/Source/Niagara/Private/NiagaraAsyncGpuTraceHelper.cpp
--- a/Engine/Plugins/FX/Niagara/Source/Niagara/Private/NiagaraAsyncGpuTraceHelper.cpp
+++ b/Engine/Plugins/FX/Niagara/Source/Niagara/Private/NiagaraAsyncGpuTraceHelper.cpp
@@ -121,7 +121,7 @@ void FNiagaraAsyncGpuTraceHelper::PostRenderOpaque(FRHICommandList& RHICmdList,
 #if NIAGARA_ASYNC_GPU_TRACE_COLLISION_GROUPS
 	if (bCollisionGroupMapDirty)
 	{
-		FNiagaraAsyncGpuTraceProvider::BuildCollisionGroupHashMap(RHICmdList, FeatureLevel, Dispatcher->GetScene(), CollisionGroupMap, CollisionGroupHashMapBuffer);
+		FNiagaraAsyncGpuTraceProvider::BuildCollisionGroupHashMap(RHICmdList, FeatureLevel, Dispatcher->GetSceneInterface(), CollisionGroupMap, CollisionGroupHashMapBuffer);
 		bCollisionGroupMapDirty = false;
 	}
 #endif
diff --git a/Engine/Plugins/FX/Niagara/Source/Niagara/Private/NiagaraGpuComputeDispatch.cpp b/Engine/Plugins/FX/Niagara/Source/Niagara/Private/NiagaraGpuComputeDispatch.cpp
--- a/Engine/Plugins/FX/Niagara/Source/Niagara/Private/NiagaraGpuComputeDispatch.cpp
+++ b/Engine/Plugins/FX/Niagara/Source/Niagara/Private/NiagaraGpuComputeDispatch.cpp
@@ -32,8 +32,8 @@
 #include "PipelineStateCache.h"
 #include "RHI.h"
 #include "SceneInterface.h"
-#include "ScenePrivate.h"
 #include "SceneRendering.h"
+#include "SceneRenderTargetParameters.h"
 #include "ShaderParameterUtils.h"
 #include "TextureResource.h"
 
diff --git a/Engine/Plugins/FX/Niagara/Source/Niagara/Private/NiagaraRendererMeshes.cpp b/Engine/Plugins/FX/Niagara/Source/Niagara/Private/NiagaraRendererMeshes.cpp
--- a/Engine/Plugins/FX/Niagara/Source/Niagara/Private/NiagaraRendererMeshes.cpp
+++ b/Engine/Plugins/FX/Niagara/Source/Niagara/Private/NiagaraRendererMeshes.cpp
@@ -6,6 +6,8 @@
 #include "Engine/StaticMesh.h"
 #include "FXRenderingUtils.h"
 #include "IXRTrackingSystem.h"
+#include "MaterialDomain.h"
+#include "Materials/MaterialInstanceDynamic.h"
 #include "Materials/MaterialRenderProxy.h"
 #include "NiagaraComponent.h"
 #include "NiagaraCullProxyComponent.h"
@@ -23,7 +25,9 @@
 #include "RayTracingDefinitions.h"
 #include "RayTracingDynamicGeometryCollection.h"
 #include "RayTracingInstance.h"
-#include "ScenePrivate.h"
+#include "RenderGraphBuilder.h"
+#include "RenderGraphUtils.h"
+#include "SceneInterface.h"
 
 DECLARE_CYCLE_STAT(TEXT("Generate Mesh Vertex Data [GT]"), STAT_NiagaraGenMeshVertexData, STATGROUP_Niagara);
 
@@ -480,17 +484,16 @@ void FNiagaraRendererMeshes::InitializeSortInfo(const FParticleMeshRenderData& P
 		OutViewOrigin = View.ViewMatrices.GetViewOrigin();
 
 #if !(UE_BUILD_SHIPPING || UE_BUILD_TEST)
-		const FSceneViewState* ViewState = View.State != nullptr ? View.State->GetConcreteViewState() : nullptr;
-		if (ViewState && ViewState->bIsFrozen && ViewState->bIsFrozenViewMatricesCached)
+		if (const FViewMatrices* CachedViewMatrices = View.State ? View.State->GetFrozenViewMatrices() : nullptr)
 		{
 			// Use the frozen view for culling so we can test that it's working
-			OutViewOrigin = ViewState->CachedViewMatrices.GetViewOrigin();
+			OutViewOrigin = CachedViewMatrices->GetViewOrigin();
 
 			// Don't retrieve the cached matrices for shadow views
 			bool bIsShadowView = View.GetDynamicMeshElementsShadowCullFrustum() != nullptr;
 			if (!bIsShadowView)
 			{
-				return ViewState->CachedViewMatrices;
+				return *CachedViewMatrices;
 			}
 		}
 #endif
diff --git a/Engine/Plugins/FX/Niagara/Source/Niagara/Private/NiagaraRendererSprites.cpp b/Engine/Plugins/FX/Niagara/Source/Niagara/Private/NiagaraRendererSprites.cpp
--- a/Engine/Plugins/FX/Niagara/Source/Niagara/Private/NiagaraRendererSprites.cpp
+++ b/Engine/Plugins/FX/Niagara/Source/Niagara/Private/NiagaraRendererSprites.cpp
@@ -1,6 +1,7 @@
 // Copyright Epic Games, Inc. All Rights Reserved.
 
 #include "NiagaraRendererSprites.h"
+#include "MaterialDomain.h"
 #include "Materials/MaterialInstanceDynamic.h"
 #include "Materials/MaterialRenderProxy.h"
 #include "NiagaraComponent.h"
@@ -17,7 +18,6 @@
 #include "RayTracingDefinitions.h"
 #include "RayTracingDynamicGeometryCollection.h"
 #include "RayTracingInstance.h"
-#include "ScenePrivate.h"
 
 DECLARE_DWORD_COUNTER_STAT(TEXT("NumSprites"), STAT_NiagaraNumSprites, STATGROUP_Niagara);
 
@@ -407,14 +407,13 @@ void FNiagaraRendererSprites::InitializeSortInfo(FParticleSpriteRenderData& Part
 		[](const FSceneView& View) -> const FViewMatrices&
 		{
 #if !(UE_BUILD_SHIPPING || UE_BUILD_TEST)
-			const FSceneViewState* ViewState = View.State != nullptr ? View.State->GetConcreteViewState() : nullptr;
-			if (ViewState && ViewState->bIsFrozen && ViewState->bIsFrozenViewMatricesCached)
+			if (const FViewMatrices* ViewMatrices = View.State ? View.State->GetFrozenViewMatrices() : nullptr)
 			{
 				// Don't retrieve the cached matrices for shadow views
 				bool bIsShadow = View.GetDynamicMeshElementsShadowCullFrustum() != nullptr;
 				if (!bIsShadow)
 				{
-					return ViewState->CachedViewMatrices;
+					return *ViewMatrices;
 				}
 			}
 #endif
diff --git a/Engine/Plugins/FX/Niagara/Source/NiagaraShader/Private/NiagaraAsyncGpuTraceProvider.cpp b/Engine/Plugins/FX/Niagara/Source/NiagaraShader/Private/NiagaraAsyncGpuTraceProvider.cpp
--- a/Engine/Plugins/FX/Niagara/Source/NiagaraShader/Private/NiagaraAsyncGpuTraceProvider.cpp
+++ b/Engine/Plugins/FX/Niagara/Source/NiagaraShader/Private/NiagaraAsyncGpuTraceProvider.cpp
@@ -9,6 +9,10 @@
 #include "NiagaraShaderParticleID.h"
 #include "ScenePrivate.h"
 
+#include "PrimitiveSceneInfo.h"
+#include "RenderGraphUtils.h"
+#include "ShaderParameterStruct.h"
+
 ////////////////////////////////////////////////////////////////////////////////
 ////////////////////////////////////////////////////////////////////////////////
 ////////////////////////////////////////////////////////////////////////////////
@@ -194,7 +198,7 @@ void FNiagaraAsyncGpuTraceProvider::ClearResults(FRHICommandList& RHICmdList, ES
 	UnsetShaderUAVs(RHICmdList, TraceShader, TraceShader.GetComputeShader());
 }
 
-void FNiagaraAsyncGpuTraceProvider::BuildCollisionGroupHashMap(FRHICommandList& RHICmdList, ERHIFeatureLevel::Type FeatureLevel, FScene* Scene, const TMap<FPrimitiveComponentId, uint32>& CollisionGroupMap, FCollisionGroupHashMap& Result)
+void FNiagaraAsyncGpuTraceProvider::BuildCollisionGroupHashMap(FRHICommandList& RHICmdList, ERHIFeatureLevel::Type FeatureLevel, FSceneInterface* Scene, const TMap<FPrimitiveComponentId, uint32>& CollisionGroupMap, FCollisionGroupHashMap& Result)
 {
 	SCOPED_DRAW_EVENT(RHICmdList, NiagaraUpdateCollisionGroupsMap);
 
@@ -225,13 +229,8 @@ void FNiagaraAsyncGpuTraceProvider::BuildCollisionGroupHashMap(FRHICommandList&
 		FPrimitiveComponentId PrimId = Entry.Key;
 		uint32 CollisionGroup = Entry.Value;
 
-		uint32 GPUSceneInstanceIndex = INDEX_NONE;
-		//Ugh this is a bit pants. Maybe try to rework things so I can use the direct prim index.
-		int32 PrimIndex = Scene->PrimitiveComponentIds.Find(PrimId);
-		if (PrimIndex != INDEX_NONE)
-		{
-			GPUSceneInstanceIndex = Scene->Primitives[PrimIndex]->GetInstanceSceneDataOffset();
-		}
+		const FPrimitiveSceneInfo* const PrimitiveSceneInfo = Scene->GetPrimitiveSceneInfo(PrimId);
+		const uint32 GPUSceneInstanceIndex = PrimitiveSceneInfo ? PrimitiveSceneInfo->GetInstanceSceneDataOffset() : INDEX_NONE;
 
 		PrimIdCollisionGroupPairPtr[0] = GPUSceneInstanceIndex;
 		PrimIdCollisionGroupPairPtr[1] = CollisionGroup;
diff --git a/Engine/Plugins/FX/Niagara/Source/NiagaraShader/Private/NiagaraAsyncGpuTraceProviderGsdf.cpp b/Engine/Plugins/FX/Niagara/Source/NiagaraShader/Private/NiagaraAsyncGpuTraceProviderGsdf.cpp
--- a/Engine/Plugins/FX/Niagara/Source/NiagaraShader/Private/NiagaraAsyncGpuTraceProviderGsdf.cpp
+++ b/Engine/Plugins/FX/Niagara/Source/NiagaraShader/Private/NiagaraAsyncGpuTraceProviderGsdf.cpp
@@ -7,9 +7,12 @@
 #include "NiagaraDistanceFieldHelper.h"
 #include "NiagaraGpuComputeDispatchInterface.h"
 #include "NiagaraSettings.h"
+#include "RenderGraphUtils.h"
+#include "SceneInterface.h"
 #include "SceneManagement.h"
+#include "ShaderParameterStruct.h"
+
 #include "ScenePrivate.h"
-#include "SceneRendering.h"
 
 static int GNiagaraAsyncGpuTraceGsdfEnabled = 1;
 static FAutoConsoleVariableRef CVarNiagaraAsyncGpuTraceGsdfEnabled(
@@ -96,7 +99,7 @@ bool FNiagaraAsyncGpuTraceProviderGsdf::IsAvailable() const
 	}
 
 	// check the feature level of the scene.  Gsdf sampling requires SM5+
-	if (Dispatcher->GetScene()->GetFeatureLevel() < ERHIFeatureLevel::SM5)
+	if (Dispatcher->GetSceneInterface()->GetFeatureLevel() < ERHIFeatureLevel::SM5)
 	{
 		return false;
 	}
diff --git a/Engine/Plugins/FX/Niagara/Source/NiagaraShader/Public/NiagaraAsyncGpuTraceProvider.h b/Engine/Plugins/FX/Niagara/Source/NiagaraShader/Public/NiagaraAsyncGpuTraceProvider.h
--- a/Engine/Plugins/FX/Niagara/Source/NiagaraShader/Public/NiagaraAsyncGpuTraceProvider.h
+++ b/Engine/Plugins/FX/Niagara/Source/NiagaraShader/Public/NiagaraAsyncGpuTraceProvider.h
@@ -69,7 +69,7 @@ public:
 		uint32 HashTableSize = 0;
 	};
 
-	static void BuildCollisionGroupHashMap(FRHICommandList& RHICmdList, ERHIFeatureLevel::Type FeatureLevel, FScene* Scene, const TMap<FPrimitiveComponentId, uint32>& CollisionGroupMap, FCollisionGroupHashMap& Result);
+	static void BuildCollisionGroupHashMap(FRHICommandList& RHICmdList, ERHIFeatureLevel::Type FeatureLevel, FSceneInterface* Scene, const TMap<FPrimitiveComponentId, uint32>& CollisionGroupMap, FCollisionGroupHashMap& Result);
 
 	static TArray<TUniquePtr<FNiagaraAsyncGpuTraceProvider>> CreateSupportedProviders(EShaderPlatform ShaderPlatform, FNiagaraGpuComputeDispatchInterface* Dispatcher, const FProviderPriorityArray& Priorities);
 	static void ClearResults(FRHICommandList& RHICmdList, EShaderPlatform ShaderPlatform, const FDispatchRequest& Request);
