Adding missing include.

diff --git a/Engine/Plugins/FX/Niagara/Source/Niagara/Private/NiagaraRendererMeshes.cpp b/Engine/Plugins/FX/Niagara/Source/Niagara/Private/NiagaraRendererMeshes.cpp
--- a/Engine/Plugins/FX/Niagara/Source/Niagara/Private/NiagaraRendererMeshes.cpp
+++ b/Engine/Plugins/FX/Niagara/Source/Niagara/Private/NiagaraRendererMeshes.cpp
@@ -28,6 +28,7 @@
 #include "RenderGraphBuilder.h"
 #include "RenderGraphUtils.h"
 #include "SceneInterface.h"
+#include "SceneUniformBuffer.h"
 
 DECLARE_CYCLE_STAT(TEXT("Generate Mesh Vertex Data [GT]"), STAT_NiagaraGenMeshVertexData, STATGROUP_Niagara);
 
