* Fixed a crash when using LWC structs
* Fixed a bug where SWC structs would appear in the UI
* Fixed a bug where using LWC structs as module inputs would not use the correct struct converter

diff --git a/Engine/Plugins/FX/Niagara/Source/Niagara/Private/NiagaraDataChannel.cpp b/Engine/Plugins/FX/Niagara/Source/Niagara/Private/NiagaraDataChannel.cpp
--- a/Engine/Plugins/FX/Niagara/Source/Niagara/Private/NiagaraDataChannel.cpp
+++ b/Engine/Plugins/FX/Niagara/Source/Niagara/Private/NiagaraDataChannel.cpp
@@ -7,6 +7,8 @@
 #include "NiagaraDataChannelHandler.h"
 #include "NiagaraDataChannelManager.h"
 #include "Engine/Engine.h"
+#include "UObject/UObjectGlobals.h"
+#include "UObject/Package.h"
 
 #include UE_INLINE_GENERATED_CPP_BY_NAME(NiagaraDataChannel)
 
@@ -408,10 +410,75 @@ void FNiagaraWorldDataChannelStore::Tick(UNiagaraDataChannelHandler* Owner)
 
 //////////////////////////////////////////////////////////////////////////
 
+namespace NiagaraDataChannel
+{
+	FNiagaraTypeDefinition GetFVectorDef()
+	{
+		static UPackage* CoreUObjectPkg = FindObjectChecked<UPackage>(nullptr, TEXT("/Script/CoreUObject"));
+		static UScriptStruct* VectorStruct = FindObjectChecked<UScriptStruct>(CoreUObjectPkg, TEXT("Vector"));
+		return FNiagaraTypeDefinition(VectorStruct, FNiagaraTypeDefinition::EAllowUnfriendlyStruct::Allow);
+	}
+	
+	FNiagaraTypeDefinition GetDoubleDef()
+	{
+		static UPackage* NiagaraPkg = FindObjectChecked<UPackage>(nullptr, TEXT("/Script/Niagara"));
+		static UScriptStruct* DoubleStruct = FindObjectChecked<UScriptStruct>(NiagaraPkg, TEXT("NiagaraDouble"));
+		return FNiagaraTypeDefinition(DoubleStruct, FNiagaraTypeDefinition::EAllowUnfriendlyStruct::Allow);
+	}
+
+	FNiagaraTypeDefinition GetFQuatDef()
+	{
+		static UPackage* CoreUObjectPkg = FindObjectChecked<UPackage>(nullptr, TEXT("/Script/CoreUObject"));
+		static UScriptStruct* Struct = FindObjectChecked<UScriptStruct>(CoreUObjectPkg, TEXT("Quat"));
+		return FNiagaraTypeDefinition(Struct, FNiagaraTypeDefinition::EAllowUnfriendlyStruct::Allow);
+	}
+	
+	FNiagaraTypeDefinition GetFVector2DDef()
+	{
+		static UPackage* CoreUObjectPkg = FindObjectChecked<UPackage>(nullptr, TEXT("/Script/CoreUObject"));
+		static UScriptStruct* VectorStruct = FindObjectChecked<UScriptStruct>(CoreUObjectPkg, TEXT("Vector2D"));
+		return FNiagaraTypeDefinition(VectorStruct, FNiagaraTypeDefinition::EAllowUnfriendlyStruct::Allow);
+	}
+
+	FNiagaraTypeDefinition GetFVector4Def()
+	{
+		static UPackage* CoreUObjectPkg = FindObjectChecked<UPackage>(nullptr, TEXT("/Script/CoreUObject"));
+		static UScriptStruct* Vector4Struct = FindObjectChecked<UScriptStruct>(CoreUObjectPkg, TEXT("Vector4"));
+		return FNiagaraTypeDefinition(Vector4Struct, FNiagaraTypeDefinition::EAllowUnfriendlyStruct::Allow);
+	}
+}
+
 void UNiagaraDataChannel::PostLoad()
 {
 	Super::PostLoad();
 
+	for (int i = 0; i < Variables.Num(); i++)
+	{
+		// fix up variables deserialized with wrong type
+		// TODO (mga) find a better solution than this
+		FNiagaraVariable& Var = Variables[i];
+		if (Var.GetType() == FNiagaraTypeDefinition::GetVec3Def())
+		{
+			Var.SetType(NiagaraDataChannel::GetFVectorDef());
+		}
+		else if (Var.GetType() == FNiagaraTypeDefinition::GetFloatDef())
+		{
+			Var.SetType(NiagaraDataChannel::GetDoubleDef());
+		}
+		else if (Var.GetType() == FNiagaraTypeDefinition::GetQuatDef())
+		{
+			Var.SetType(NiagaraDataChannel::GetFQuatDef());
+		}
+		else if (Var.GetType() == FNiagaraTypeDefinition::GetVec2Def())
+		{
+			Var.SetType(NiagaraDataChannel::GetFVector2DDef());
+		}
+		else if (Var.GetType() == FNiagaraTypeDefinition::GetVec4Def())
+		{
+			Var.SetType(NiagaraDataChannel::GetFVector4Def());
+		}
+	}
+
 	//TODO: Can serialize?
 	GameDataLayout.Init(Variables);
 
diff --git a/Engine/Plugins/FX/Niagara/Source/Niagara/Private/NiagaraModule.cpp b/Engine/Plugins/FX/Niagara/Source/Niagara/Private/NiagaraModule.cpp
--- a/Engine/Plugins/FX/Niagara/Source/Niagara/Private/NiagaraModule.cpp
+++ b/Engine/Plugins/FX/Niagara/Source/Niagara/Private/NiagaraModule.cpp
@@ -1144,7 +1144,7 @@ void FNiagaraTypeDefinition::RecreateUserDefinedTypeRegistry()
 				UScriptStruct* ScriptStruct = Cast<UScriptStruct>(Obj);
 				if (ScriptStruct != nullptr)
 				{
-					ENiagaraTypeRegistryFlags Flags = ENiagaraTypeRegistryFlags::AllowAnyVariable;
+					ENiagaraTypeRegistryFlags Flags = ENiagaraTypeRegistryFlags::AllowAnyVariable | ENiagaraTypeRegistryFlags::IsUserDefined;
 					if (ParamRefFound)
 					{
 						Flags |= ENiagaraTypeRegistryFlags::AllowParameter;
@@ -1191,7 +1191,7 @@ void FNiagaraTypeDefinition::RecreateUserDefinedTypeRegistry()
 				UEnum* Enum = Cast<UEnum>(Obj);
 				if (Enum != nullptr)
 				{
-					ENiagaraTypeRegistryFlags Flags = ENiagaraTypeRegistryFlags::AllowAnyVariable | ENiagaraTypeRegistryFlags::AllowParameter;
+					ENiagaraTypeRegistryFlags Flags = ENiagaraTypeRegistryFlags::AllowAnyVariable | ENiagaraTypeRegistryFlags::AllowParameter | ENiagaraTypeRegistryFlags::IsUserDefined;
 					
 					FNiagaraTypeDefinition Def(Enum);
 					FNiagaraTypeRegistry::Register(Def, Flags);
@@ -1421,13 +1421,13 @@ bool FNiagaraTypeDefinition::Serialize(FArchive& Ar)
 
 					Flags &= ~TF_SerializedAsLWC;
 				}
-	#if WITH_EDITORONLY_DATA
-				else if((Flags & TF_AllowLWC) == 0)
+#if WITH_EDITORONLY_DATA
+				else if((Flags & TF_AllowLWC_DEPRECATED) == 0)
 				{	
-					//If the type doesn't allow LWC we should convert down to SWC.
+					// If a basic type was serialized with the LWC flag, restore the correct struct
 					ClassStructOrEnum = FNiagaraTypeHelper::FindNiagaraFriendlyTopLevelStruct(static_cast<UScriptStruct*>(ClassStructOrEnum), ENiagaraStructConversion::UserFacing);
 				}
-	#endif
+#endif
 			}
 		}
 	}
diff --git a/Engine/Plugins/FX/Niagara/Source/Niagara/Private/NiagaraTypes.cpp b/Engine/Plugins/FX/Niagara/Source/Niagara/Private/NiagaraTypes.cpp
--- a/Engine/Plugins/FX/Niagara/Source/Niagara/Private/NiagaraTypes.cpp
+++ b/Engine/Plugins/FX/Niagara/Source/Niagara/Private/NiagaraTypes.cpp
@@ -271,6 +271,11 @@ void FNiagaraLwcStructConverter::AddConversionStep(int32 InSourceBytes, int32 In
 	check(InSimulationOffset >= 0);
 	// TODO (mg) as an optimization, if the last step was a copy step and this one is as well, they can be merged into one
 	ConversionSteps.Emplace(InSourceBytes, InSourceOffset, InSimulationBytes, InSimulationOffset, ConversionType);
+
+	// SWCSize and LWCSize are the sizes of the full structs.
+	// We're updating them with each added step, as the full struct size should be (last property offset + size).
+	SWCSize = InSimulationOffset + InSimulationBytes;
+	LWCSize = InSourceOffset + InSourceBytes;
 }
 
 //////////////////////////////////////////////////////////////////////////
@@ -473,7 +478,7 @@ bool FNiagaraTypeHelper::IsLWCStructure(UStruct* InStruct)
 
 bool FNiagaraTypeHelper::IsLWCType(const FNiagaraTypeDefinition& InType)
 {
-	return InType.IsValid() && !InType.IsUObject() && IsLWCStructure(InType.GetStruct());
+	return InType.IsValid() && !InType.IsEnum() && !InType.IsUObject() && IsLWCStructure(InType.GetStruct());
 }
 
 void FNiagaraTypeHelper::TickTypeRemap()
@@ -586,7 +591,8 @@ UScriptStruct* FNiagaraTypeHelper::GetSWCStruct(UScriptStruct* LWCStruct)
 						SWCStruct->AddToRoot();
 					}
 					FNiagaraLwcStructConverter StructConverter = BuildSWCStructure(SWCStruct, LWCStruct);
-					FNiagaraTypeRegistry::RegisterStructConverter(FNiagaraTypeDefinition(LWCStruct, FNiagaraTypeDefinition::EAllowUnfriendlyStruct::Allow), StructConverter);
+					FNiagaraTypeDefinition LwcTypeDefinition = FNiagaraTypeRegistry::GetTypeForStruct(LWCStruct);
+					FNiagaraTypeRegistry::RegisterStructConverter(LwcTypeDefinition, StructConverter);
 					SWCStruct->Bind();
 					SWCStruct->PrepareCppStructOps();
 					SWCStruct->StaticLink(true);
@@ -627,6 +633,15 @@ UScriptStruct* FNiagaraTypeHelper::GetLWCStruct(UScriptStruct* SWCStruct)
 	return nullptr;
 }
 
+FNiagaraTypeDefinition FNiagaraTypeHelper::GetSWCType(const FNiagaraTypeDefinition& InType)
+{
+	if (IsLWCType(InType))
+	{
+		return FNiagaraTypeDefinition(FindNiagaraFriendlyTopLevelStruct(CastChecked<UScriptStruct>(InType.GetStruct()), ENiagaraStructConversion::Simulation));
+	}
+	return InType;
+}
+
 UScriptStruct* FNiagaraTypeHelper::FindNiagaraFriendlyTopLevelStruct(UScriptStruct* InStruct, ENiagaraStructConversion StructConversion)
 {
 	if (!InStruct)
diff --git a/Engine/Plugins/FX/Niagara/Source/Niagara/Public/NiagaraTypes.h b/Engine/Plugins/FX/Niagara/Source/Niagara/Public/NiagaraTypes.h
--- a/Engine/Plugins/FX/Niagara/Source/Niagara/Public/NiagaraTypes.h
+++ b/Engine/Plugins/FX/Niagara/Source/Niagara/Public/NiagaraTypes.h
@@ -293,12 +293,6 @@ struct NIAGARA_API FNiagaraLwcStructConverter
 {
 	GENERATED_USTRUCT_BODY();
 
-	FNiagaraLwcStructConverter(){}
-	FNiagaraLwcStructConverter(UScriptStruct* LWCStruct, UScriptStruct* SWCStruct)
-		: LWCSize(LWCStruct ? LWCStruct->GetStructureSize() : 0)
-		, SWCSize(SWCStruct ? SWCStruct->GetStructureSize() : 0)
-	{
-	}
 	void ConvertDataToSimulation(uint8* DestinationData, const uint8* SourceData, int32 Count = 1) const;
 	void ConvertDataFromSimulation(uint8* DestinationData, const uint8* SourceData, int32 Count = 1) const;
 	
@@ -410,6 +404,7 @@ public:
 	static bool IsNiagaraFriendlyTopLevelStruct(UScriptStruct* InStruct, ENiagaraStructConversion StructConversion);
 	static UScriptStruct* GetSWCStruct(UScriptStruct* LWCStruct);
 	static UScriptStruct* GetLWCStruct(UScriptStruct* LWCStruct);
+	static FNiagaraTypeDefinition GetSWCType(const FNiagaraTypeDefinition& InType);
 	static void TickTypeRemap();
 
 private:
@@ -797,8 +792,8 @@ struct NIAGARA_API FNiagaraTypeDefinition
 		/// instead of the Transient SWC version of the struct
 		TF_SerializedAsLWC	= 0x02,
 
-		/// indicates this type allows LWC types and should not attempt to convert to SWC.
-		TF_AllowLWC		= 0x04,
+		/// DEPRECATED, DO NOT USE
+		TF_AllowLWC_DEPRECATED		= 0x04,
 	};
 
 public:
@@ -855,11 +850,8 @@ public:
 	{
 		checkSlow(ClassStructOrEnum != nullptr);
 		ensureAlwaysMsgf(AllowUnfriendlyStruct == EAllowUnfriendlyStruct::Allow || FNiagaraTypeHelper::IsNiagaraFriendlyTopLevelStruct(StructDef, ENiagaraStructConversion::UserFacing), TEXT("Struct(%s) is not supported."), *StructDef->GetName());
-		if(AllowUnfriendlyStruct == EAllowUnfriendlyStruct::Allow)
-		{
-			Flags |= TF_AllowLWC;
-		}
 	}
+	
 	FORCEINLINE FNiagaraTypeDefinition(UScriptStruct* StructDef) : FNiagaraTypeDefinition(StructDef, EAllowUnfriendlyStruct::Deny) {}
 
 	FORCEINLINE FNiagaraTypeDefinition(const FNiagaraTypeDefinition &Other)
@@ -951,7 +943,7 @@ public:
 			{
 				NameText = FText::Format(NSLOCTEXT("NiagaraTypeDefinition", "FlagFormatNameText", "Static {0}"), NameText);
 			}
-			else if (!AllowsLWC())
+			else
 			{
 				NameText = FText::Format(NSLOCTEXT("NiagaraTypeDefinition", "UnknownNameText", "Unknown Flag Type {0}"), NameText);
 			}
@@ -1011,8 +1003,6 @@ public:
 
 	bool IsStatic() const { return (GetFlags() & TF_Static) != 0; }
 
-	bool AllowsLWC() const { return (GetFlags() & TF_AllowLWC) != 0; }
-
 	void SetFlags(FTypeFlags InFlags)
 	{
 		Flags = InFlags;
@@ -1405,11 +1395,6 @@ public:
 		return Get().RegisteredIndexTypes;
 	}
 
-	static const TArray<FNiagaraTypeDefinition>& GetSimulationTypes()
-	{
-		return Get().RegisteredSimulationTypes;
-	}
-
 	static UNiagaraDataInterfaceBase* GetDefaultDataInterfaceByName(const FString& DIClassName);
 
 	static void ClearUserDefinedRegistry()
@@ -1460,19 +1445,16 @@ public:
 	{
 		FNiagaraTypeRegistry& Registry = Get();
 
-		FNiagaraTypeDefinition SimType = NewType;
-		if (UScriptStruct* Struct = NewType.GetScriptStruct())
+		if (FNiagaraTypeHelper::IsLWCType(NewType))
 		{
-			SimType = FNiagaraTypeDefinition(FNiagaraTypeHelper::GetSWCStruct(Struct), FNiagaraTypeDefinition::EAllowUnfriendlyStruct::Deny);			
+			// register the swc type as well if necessary
+			FNiagaraTypeDefinition(FNiagaraTypeHelper::GetSWCStruct(NewType.GetScriptStruct()), FNiagaraTypeDefinition::EAllowUnfriendlyStruct::Deny);
 		}
 
 		FRWScopeLock Lock(Registry.RegisteredTypesLock, SLT_Write);
-
 		//TODO: Make this a map of type to a more verbose set of metadata? Such as the hlsl defs, offset table for conversions etc.
 		Registry.RegisteredTypeIndexMap.Add(GetTypeHash(NewType), Registry.RegisteredTypes.AddUnique(NewType));
 
-		Registry.RegisteredSimulationTypes.AddUnique(SimType);
-
 		if (EnumHasAnyFlags(Flags, ENiagaraTypeRegistryFlags::AllowUserVariable))
 		{
 			Registry.RegisteredUserVariableTypes.AddUnique(NewType);
@@ -1578,6 +1560,20 @@ public:
 		return FNiagaraLwcStructConverter();
 	}
 
+	static FNiagaraTypeDefinition GetTypeForStruct(UScriptStruct* InStruct)
+	{
+		FNiagaraTypeRegistry& Registry = Get();
+		FReadScopeLock Lock(Registry.RegisteredTypesLock);
+		for (const FNiagaraTypeDefinition& Type : Registry.RegisteredTypes)
+		{
+			if (Type.GetStruct() && InStruct == Type.GetStruct())
+			{
+				return Type;
+			}
+		}
+		return FNiagaraTypeDefinition(InStruct);
+	}
+
 	/** LazySingleton interface */
 	static FNiagaraTypeRegistry& Get();
 	static void TearDown();
@@ -1603,8 +1599,6 @@ private:
 	TArray<FNiagaraTypeDefinition> RegisteredUserDefinedTypes;
 	TArray<FNiagaraTypeDefinition> RegisteredNumericTypes;
 	TArray<FNiagaraTypeDefinition> RegisteredIndexTypes;
-	TArray<FNiagaraTypeDefinition> RegisteredSimulationTypes;
-
 
 	TMap<uint32, int32> RegisteredTypeIndexMap;
 	TMap<TWeakObjectPtr<UScriptStruct>, TWeakObjectPtr<UScriptStruct>> LWCRegisteredStructRemapping;
diff --git a/Engine/Plugins/FX/Niagara/Source/NiagaraEditor/Private/EdGraphSchema_Niagara.cpp b/Engine/Plugins/FX/Niagara/Source/NiagaraEditor/Private/EdGraphSchema_Niagara.cpp
--- a/Engine/Plugins/FX/Niagara/Source/NiagaraEditor/Private/EdGraphSchema_Niagara.cpp
+++ b/Engine/Plugins/FX/Niagara/Source/NiagaraEditor/Private/EdGraphSchema_Niagara.cpp
@@ -690,7 +690,7 @@ PRAGMA_ENABLE_DEPRECATION_WARNINGS
 			
 			// Dynamic Makes from type registry. Should include the pin type as well.
 			TArray<FNiagaraTypeDefinition> CandidateTypes;
-			FNiagaraEditorUtilities::GetAllowedTypes(CandidateTypes);
+			FNiagaraEditorUtilities::GetAllowedParameterTypes(CandidateTypes);
 			for (const FNiagaraTypeDefinition& CandidateType : FNiagaraTypeRegistry::GetRegisteredTypes())
 			{
 				bool bAddMake = false;
diff --git a/Engine/Plugins/FX/Niagara/Source/NiagaraEditor/Private/NiagaraEditorUtilities.cpp b/Engine/Plugins/FX/Niagara/Source/NiagaraEditor/Private/NiagaraEditorUtilities.cpp
--- a/Engine/Plugins/FX/Niagara/Source/NiagaraEditor/Private/NiagaraEditorUtilities.cpp
+++ b/Engine/Plugins/FX/Niagara/Source/NiagaraEditor/Private/NiagaraEditorUtilities.cpp
@@ -3944,19 +3944,6 @@ int FNiagaraEditorUtilities::GetReferencedAssetCount(const FAssetData& SourceAss
 	return FMath::Max(Count, 0);
 }
 
-void FNiagaraEditorUtilities::GetAllowedTypes(TArray<FNiagaraTypeDefinition>& OutAllowedTypes)
-{
-	const UNiagaraEditorSettings* NiagaraEditorSettings = GetDefault<UNiagaraEditorSettings>();
-	//Only use simulation types here so as to not polute lists with LWC types in the registry. 
-	for (const FNiagaraTypeDefinition& RegisteredType : FNiagaraTypeRegistry::GetSimulationTypes())
-	{
-		if (NiagaraEditorSettings->IsAllowedTypeDefinition(RegisteredType))
-		{
-			OutAllowedTypes.Add(RegisteredType);
-		}
-	}
-}
-
 void FNiagaraEditorUtilities::GetAllowedUserVariableTypes(TArray<FNiagaraTypeDefinition>& OutAllowedTypes)
 {
 	const UNiagaraEditorSettings* NiagaraEditorSettings = GetDefault<UNiagaraEditorSettings>();
diff --git a/Engine/Plugins/FX/Niagara/Source/NiagaraEditor/Private/NiagaraHlslTranslator.cpp b/Engine/Plugins/FX/Niagara/Source/NiagaraEditor/Private/NiagaraHlslTranslator.cpp
--- a/Engine/Plugins/FX/Niagara/Source/NiagaraEditor/Private/NiagaraHlslTranslator.cpp
+++ b/Engine/Plugins/FX/Niagara/Source/NiagaraEditor/Private/NiagaraHlslTranslator.cpp
@@ -9088,13 +9088,14 @@ void FHlslNiagaraTranslator::Select(UNiagaraNodeSelect* SelectNode, int32 Select
 	{
 		for (const FNiagaraVariable& Variable : OutputVariables)
 		{
-			const FNiagaraVariable DefaultVar(Variable.GetType(), Variable.GetName());
+			FNiagaraTypeDefinition OutputType = FNiagaraTypeHelper::GetSWCType(Variable.GetType());
+			const FNiagaraVariable DefaultVar(OutputType, Variable.GetName());
 			const int32 DefaultConstant = GetConstant(DefaultVar);
 
 			FString SymbolName = Variable.GetName().ToString() + SymbolNameSuffix;
 			SymbolNames.Add(SymbolName);
 			
-			const int32 SymbolIndex = AddBodyChunk(SymbolName, TEXT("{0}"), Variable.GetType(), true);
+			const int32 SymbolIndex = AddBodyChunk(SymbolName, TEXT("{0}"), OutputType, true);
 			FNiagaraCodeChunk& OutputSymbolChunk = CodeChunks[SymbolIndex];
 			OutputSymbolChunk.AddSourceChunk(DefaultConstant);
 			
@@ -9129,7 +9130,7 @@ void FHlslNiagaraTranslator::Select(UNiagaraNodeSelect* SelectNode, int32 Select
 		int32 NaturalIndex = 0;
 		for(int32 CompiledPinCodeChunk : Options[SelectorValues[SelectorValueIndex]])
 		{
-			FNiagaraCodeChunk& BranchChunk = CodeChunks[AddBodyChunk(SymbolNames[NaturalIndex], TEXT("{0}"), OutputVariables[NaturalIndex].GetType(), false)];
+			FNiagaraCodeChunk& BranchChunk = CodeChunks[AddBodyChunk(SymbolNames[NaturalIndex], TEXT("{0}"), FNiagaraTypeHelper::GetSWCType(OutputVariables[NaturalIndex].GetType()), false)];
 			BranchChunk.AddSourceChunk(CompiledPinCodeChunk);
 			NaturalIndex++;
 		}
diff --git a/Engine/Plugins/FX/Niagara/Source/NiagaraEditor/Private/ViewModels/NiagaraParameterPanelViewModel.cpp b/Engine/Plugins/FX/Niagara/Source/NiagaraEditor/Private/ViewModels/NiagaraParameterPanelViewModel.cpp
--- a/Engine/Plugins/FX/Niagara/Source/NiagaraEditor/Private/ViewModels/NiagaraParameterPanelViewModel.cpp
+++ b/Engine/Plugins/FX/Niagara/Source/NiagaraEditor/Private/ViewModels/NiagaraParameterPanelViewModel.cpp
@@ -576,7 +576,7 @@ void INiagaraParameterPanelViewModel::GetChangeTypeSubMenu(FMenuBuilder& MenuBui
 {
 	TArray<FNiagaraTypeDefinition> FilteredTypes;
 	TArray<FNiagaraTypeDefinition> AllowedTypes;
-	FNiagaraEditorUtilities::GetAllowedTypes(AllowedTypes);
+	FNiagaraEditorUtilities::GetAllowedParameterTypes(AllowedTypes);
 	for (const FNiagaraTypeDefinition& RegisteredType : AllowedTypes)
 	{
 		// only allow basic types for now
diff --git a/Engine/Plugins/FX/Niagara/Source/NiagaraEditor/Public/NiagaraEditorUtilities.h b/Engine/Plugins/FX/Niagara/Source/NiagaraEditor/Public/NiagaraEditorUtilities.h
--- a/Engine/Plugins/FX/Niagara/Source/NiagaraEditor/Public/NiagaraEditorUtilities.h
+++ b/Engine/Plugins/FX/Niagara/Source/NiagaraEditor/Public/NiagaraEditorUtilities.h
@@ -369,10 +369,6 @@ namespace FNiagaraEditorUtilities
 	};
 
 	NIAGARAEDITOR_API int GetReferencedAssetCount(const FAssetData& SourceAsset, TFunction<ETrackAssetResult(const FAssetData&)> Predicate);
-	
-	/** Gets a list of the registered types which are allowed in the current editor context.  This API should be 
-		called when providing a list types to the user instead of getting the type list directly from the type registry. */
-	void GetAllowedTypes(TArray<FNiagaraTypeDefinition>& OutAllowedTypes);
 
 	/** Gets a list of the registered user variable types which are allowed in the current editor context.  This API should be
 		called when providing a list types to the user instead of getting the type list directly from the type registry. */
diff --git a/Engine/Plugins/FX/Niagara/Source/NiagaraEditor/Public/Widgets/SNiagaraParameterMenu.cpp b/Engine/Plugins/FX/Niagara/Source/NiagaraEditor/Public/Widgets/SNiagaraParameterMenu.cpp
--- a/Engine/Plugins/FX/Niagara/Source/NiagaraEditor/Public/Widgets/SNiagaraParameterMenu.cpp
+++ b/Engine/Plugins/FX/Niagara/Source/NiagaraEditor/Public/Widgets/SNiagaraParameterMenu.cpp
@@ -175,7 +175,7 @@ void SNiagaraAddParameterFromPanelMenu::CollectMakeNew(FNiagaraMenuActionCollect
 	}
 	else
 	{
-		FNiagaraEditorUtilities::GetAllowedTypes(SectionTypes);
+		FNiagaraEditorUtilities::GetAllowedParameterTypes(SectionTypes);
 	}
 
 	TArray<FNiagaraTypeDefinition> TypeDefinitions;
@@ -829,7 +829,7 @@ void SNiagaraAddParameterFromPinMenu::CollectAllActions(FGraphActionListBuilderB
 {
 	FNiagaraMenuActionCollector Collector;
 	TArray<FNiagaraTypeDefinition> Types;
-	FNiagaraEditorUtilities::GetAllowedTypes(Types);
+	FNiagaraEditorUtilities::GetAllowedParameterTypes(Types);
 	Types.Sort([](const FNiagaraTypeDefinition& A, const FNiagaraTypeDefinition& B) { return (A.GetNameText().ToLower().ToString() < B.GetNameText().ToLower().ToString()); });
 
 	for (const FNiagaraTypeDefinition& RegisteredType : Types)
@@ -878,7 +878,7 @@ void SNiagaraChangePinTypeMenu::CollectAllActions(FGraphActionListBuilderBase& O
 	checkf(Node, TEXT("Niagara node pin did not have a valid outer node!"));
 
 	TArray<FNiagaraTypeDefinition> Types;
-	FNiagaraEditorUtilities::GetAllowedTypes(Types);
+	FNiagaraEditorUtilities::GetAllowedParameterTypes(Types);
 	Types.Sort([](const FNiagaraTypeDefinition& A, const FNiagaraTypeDefinition& B) { return (A.GetNameText().ToLower().ToString() < B.GetNameText().ToLower().ToString()); });
 
 	for (const FNiagaraTypeDefinition& RegisteredType : Types)
