Restoring missing include.

diff --git a/Engine/Plugins/FX/Niagara/Source/NiagaraShader/Private/NiagaraAsyncGpuTraceProviderHwrt.cpp b/Engine/Plugins/FX/Niagara/Source/NiagaraShader/Private/NiagaraAsyncGpuTraceProviderHwrt.cpp
--- a/Engine/Plugins/FX/Niagara/Source/NiagaraShader/Private/NiagaraAsyncGpuTraceProviderHwrt.cpp
+++ b/Engine/Plugins/FX/Niagara/Source/NiagaraShader/Private/NiagaraAsyncGpuTraceProviderHwrt.cpp
@@ -16,6 +16,8 @@
 #include "Templates/Function.h"
 #include "RayTracingPayloadType.h"
 
+#include "ScenePrivate.h"
+
 static int GNiagaraAsyncGpuTraceHwrtEnabled = 1;
 static FAutoConsoleVariableRef CVarNiagaraAsyncGpuTraceHwrtEnabled(
 	TEXT("fx.Niagara.AsyncGpuTrace.HWRayTraceEnabled"),
