- Gpu usage validation rule
- Allows us to inform the user if they are using Gpu but it may not work as expected.  I.e. we disable gpu on those platforms, etc.

diff --git a/Engine/Plugins/FX/Niagara/Source/Niagara/Classes/NiagaraPlatformSet.h b/Engine/Plugins/FX/Niagara/Source/Niagara/Classes/NiagaraPlatformSet.h
--- a/Engine/Plugins/FX/Niagara/Source/Niagara/Classes/NiagaraPlatformSet.h
+++ b/Engine/Plugins/FX/Niagara/Source/Niagara/Classes/NiagaraPlatformSet.h
@@ -394,7 +394,7 @@ public:
 	void OnChanged();
 
 	/** Inspects the passed sets and generates an array of all conflicts between these sets. Used to keep arrays of platform sets orthogonal. */
-	static bool GatherConflicts(const TArray<const FNiagaraPlatformSet*>& PlatformSets, TArray<FNiagaraPlatformSetConflictInfo>& OutConflicts);
+	static bool GatherConflicts(TConstArrayView<const FNiagaraPlatformSet*> PlatformSets, TArray<FNiagaraPlatformSetConflictInfo>& OutConflicts);
 	
 	DECLARE_DELEGATE_RetVal(int32, FOnOverrideQualityLevel);
 	FOnOverrideQualityLevel OnOverrideQualityLevelDelegate;
diff --git a/Engine/Plugins/FX/Niagara/Source/Niagara/Private/NiagaraPlatformSet.cpp b/Engine/Plugins/FX/Niagara/Source/Niagara/Private/NiagaraPlatformSet.cpp
--- a/Engine/Plugins/FX/Niagara/Source/Niagara/Private/NiagaraPlatformSet.cpp
+++ b/Engine/Plugins/FX/Niagara/Source/Niagara/Private/NiagaraPlatformSet.cpp
@@ -1201,7 +1201,7 @@ void FNiagaraPlatformSet::OnChanged()
 	InvalidateCachedData();
 }
 
-bool FNiagaraPlatformSet::GatherConflicts(const TArray<const FNiagaraPlatformSet*>& PlatformSets, TArray<FNiagaraPlatformSetConflictInfo>& OutConflicts)
+bool FNiagaraPlatformSet::GatherConflicts(TConstArrayView<const FNiagaraPlatformSet*> PlatformSets, TArray<FNiagaraPlatformSetConflictInfo>& OutConflicts)
 {
 	const UNiagaraSettings* Settings = GetDefault<UNiagaraSettings>();
 	check(Settings);
diff --git a/Engine/Plugins/FX/Niagara/Source/NiagaraEditor/Private/NiagaraValidationRules.cpp b/Engine/Plugins/FX/Niagara/Source/NiagaraEditor/Private/NiagaraValidationRules.cpp
--- a/Engine/Plugins/FX/Niagara/Source/NiagaraEditor/Private/NiagaraValidationRules.cpp
+++ b/Engine/Plugins/FX/Niagara/Source/NiagaraEditor/Private/NiagaraValidationRules.cpp
@@ -28,6 +28,9 @@
 #include "AssetToolsModule.h"
 #include "Materials/MaterialInterface.h"
 
+#include "DeviceProfiles/DeviceProfile.h"
+#include "DeviceProfiles/DeviceProfileManager.h"
+
 #include UE_INLINE_GENERATED_CPP_BY_NAME(NiagaraValidationRules)
 
 #define LOCTEXT_NAMESPACE "NiagaraValidationRules"
@@ -530,6 +533,81 @@ void UNiagaraValidationRule_BannedDataInterfaces::CheckValidity(const FNiagaraVa
 	);
 }
 
+void UNiagaraValidationRule_GpuUsage::CheckValidity(const FNiagaraValidationContext& Context, TArray<FNiagaraValidationResult>& OutResults) const
+{
+	UNiagaraSystem& System = Context.ViewModel->GetSystem();
+	for (TSharedRef<FNiagaraEmitterHandleViewModel> EmitterHandleModel : Context.ViewModel->GetEmitterHandleViewModels())
+	{
+		FNiagaraEmitterHandle* EmitterHandle = EmitterHandleModel.Get().GetEmitterHandle();
+		FVersionedNiagaraEmitterData* EmitterData = EmitterHandle->GetEmitterData();
+		if (EmitterData->SimTarget != ENiagaraSimTarget::GPUComputeSim)
+		{
+			continue;
+		}
+
+		const FString PlatformConflictsString = NiagaraValidation::GetPlatformConflictsString(Platforms, EmitterData->Platforms);
+		if (PlatformConflictsString.IsEmpty())
+		{
+			continue;
+		}
+
+		FNiagaraValidationResult& ValidationResult = OutResults.Emplace_GetRef(
+			Severity,
+			LOCTEXT("GpuUsageInfo", "GPU usage may not function as expected"),
+			FText::Format(LOCTEXT("GpuUsageInfoDetails", "GPU usage may not function as expected on '{0}'."), FText::FromString(PlatformConflictsString)),
+			NiagaraValidation::GetStackEntry<UNiagaraStackEmitterPropertiesItem>(EmitterHandleModel->GetEmitterStackViewModel())
+		);
+		
+		ValidationResult.Fixes.Emplace(
+			FText::Format(LOCTEXT("GpuUsageInfoFix_DisablePlatforms", "Disable emitter on '{0}'."), FText::FromString(PlatformConflictsString)),
+			FNiagaraValidationFixDelegate::CreateLambda(
+				[WeakEmitterPtr=EmitterHandle->GetInstance().ToWeakPtr(), PlatformsToDisable=Platforms]()
+				{
+					FVersionedNiagaraEmitter VersionedEmitter = WeakEmitterPtr.ResolveWeakPtr();
+					if (FVersionedNiagaraEmitterData* VersionedEmitterData = VersionedEmitter.GetEmitterData())
+					{
+						TArray<FNiagaraPlatformSetConflictInfo> ConflictInfos;
+						FNiagaraPlatformSet::GatherConflicts({&VersionedEmitterData->Platforms, &PlatformsToDisable}, ConflictInfos);
+
+						for (const FNiagaraPlatformSetConflictInfo& ConflictInfo : ConflictInfos)
+						{
+							for (const FNiagaraPlatformSetConflictEntry& ConflictEntry : ConflictInfo.Conflicts)
+							{
+								UDeviceProfile* DeviceProfile = UDeviceProfileManager::Get().FindProfile(ConflictEntry.ProfileName.ToString());
+								for (int32 iQualityLevel=0; iQualityLevel < 32; ++iQualityLevel)
+								{
+									if ( (ConflictEntry.QualityLevelMask & (1 << iQualityLevel)) != 0 )
+									{
+										VersionedEmitterData->Platforms.SetDeviceProfileState(DeviceProfile, iQualityLevel, ENiagaraPlatformSelectionState::Disabled);
+									}
+								}
+							}
+						}
+					}
+				}
+			)
+		);
+			
+		ValidationResult.Fixes.Emplace(
+			LOCTEXT("GpuUsageInfoFix_SwitchToCput", "Set emitter to CPU"),
+			FNiagaraValidationFixDelegate::CreateLambda(
+				[WeakEmitterPtr=EmitterHandle->GetInstance().ToWeakPtr()]()
+				{
+					FVersionedNiagaraEmitter VersionedEmitter = WeakEmitterPtr.ResolveWeakPtr();
+					if (FVersionedNiagaraEmitterData* VersionedEmitterData = VersionedEmitter.GetEmitterData())
+					{
+						VersionedEmitterData->SimTarget = ENiagaraSimTarget::CPUSim;
+
+						FProperty* SimTargetProperty = FindFProperty<FProperty>(FVersionedNiagaraEmitterData::StaticStruct(), GET_MEMBER_NAME_CHECKED(FVersionedNiagaraEmitterData, SimTarget));
+						FPropertyChangedEvent PropertyChangedEvent(SimTargetProperty);
+						VersionedEmitter.Emitter->PostEditChangeVersionedProperty(PropertyChangedEvent, VersionedEmitter.Version);
+					}
+				}
+			)
+		);
+	}
+}
+
 void UNiagaraValidationRule_InvalidEffectType::CheckValidity(const FNiagaraValidationContext& Context, TArray<FNiagaraValidationResult>& Results)  const
 {
 	UNiagaraStackSystemPropertiesItem* SystemProperties = NiagaraValidation::GetStackEntry<UNiagaraStackSystemPropertiesItem>(Context.ViewModel->GetSystemStackViewModel());
diff --git a/Engine/Plugins/FX/Niagara/Source/NiagaraEditor/Public/NiagaraValidationRules.h b/Engine/Plugins/FX/Niagara/Source/NiagaraEditor/Public/NiagaraValidationRules.h
--- a/Engine/Plugins/FX/Niagara/Source/NiagaraEditor/Public/NiagaraValidationRules.h
+++ b/Engine/Plugins/FX/Niagara/Source/NiagaraEditor/Public/NiagaraValidationRules.h
@@ -95,6 +95,22 @@ public:
 	virtual void CheckValidity(const FNiagaraValidationContext& Context, TArray<FNiagaraValidationResult>& OutResults) const override;
 };
 
+/** This validation rule that can be used to inform or ban GPU usage on the provided platforms. */
+UCLASS(Category = "Validation", DisplayName = "Gpu Usage")
+class UNiagaraValidationRule_GpuUsage : public UNiagaraValidationRule
+{
+	GENERATED_BODY()
+
+public:
+	UPROPERTY(EditAnywhere, Category = Validation)
+	ENiagaraValidationSeverity Severity = ENiagaraValidationSeverity::Warning;
+
+	UPROPERTY(EditAnywhere, Category = Validation)
+	FNiagaraPlatformSet Platforms;
+
+	virtual void CheckValidity(const FNiagaraValidationContext& Context, TArray<FNiagaraValidationResult>& OutResults) const override;
+};
+
 /** This validation rule can marks this effect type as invalid and so must be changed. Forces a choice of correct Effect Type for an System rather than. Leaving as the default. */
 UCLASS(Category = "Validation", DisplayName = "Invalid Effect Type")
 class UNiagaraValidationRule_InvalidEffectType : public UNiagaraValidationRule
