- Show desired age on HUD when in desired age mode

#rnx

diff --git a/Engine/Plugins/FX/Niagara/Source/Niagara/Private/NiagaraDebugHud.cpp b/Engine/Plugins/FX/Niagara/Source/Niagara/Private/NiagaraDebugHud.cpp
--- a/Engine/Plugins/FX/Niagara/Source/Niagara/Private/NiagaraDebugHud.cpp
+++ b/Engine/Plugins/FX/Niagara/Source/Niagara/Private/NiagaraDebugHud.cpp
@@ -2750,11 +2750,21 @@ void FNiagaraDebugHud::DrawComponents(FNiagaraWorldManager* WorldManager, UCanva
 
 					if (SystemInstanceController->IsSolo())
 					{
+						StringBuilder.Append(TEXT("Solo(true)"));
 						if ( NiagaraComponent->GetForceSolo() )
 						{
-							StringBuilder.Append(TEXT("ForceSolo "));
+							StringBuilder.Append(TEXT(" ForceSolo"));
 						}
-						StringBuilder.Append(TEXT("IsSolo\n"));
+						
+						switch (NiagaraComponent->GetAgeUpdateMode())
+						{
+							case ENiagaraAgeUpdateMode::DesiredAge:
+							case ENiagaraAgeUpdateMode::DesiredAgeNoSeek:
+								StringBuilder.Appendf(TEXT(" DesiredAge(%6.3f)"), NiagaraComponent->GetDesiredAge());
+								break;
+						}
+
+						StringBuilder.Append(TEXT("\n"));
 					}
 
 					if ( NiagaraComponent->bHiddenInGame || !NiagaraComponent->GetVisibleFlag() || (OwnerActor && OwnerActor->IsHidden()) )
